Separate CVar for HZB for VSMs vs Atlases.
Add light name to non-nanite VSM raster passes.

diff --git a/Engine/Source/Runtime/Renderer/Private/ShadowDepthRendering.cpp b/Engine/Source/Runtime/Renderer/Private/ShadowDepthRendering.cpp
--- a/Engine/Source/Runtime/Renderer/Private/ShadowDepthRendering.cpp
+++ b/Engine/Source/Runtime/Renderer/Private/ShadowDepthRendering.cpp
@@ -99,6 +99,12 @@ static TAutoConsoleVariable<int32> CVarNaniteShadowsUseHZB(
 	TEXT("Enables HZB for Nanite shadows."),
 	ECVF_RenderThreadSafe);
 
+static TAutoConsoleVariable<int32> CVarShadowsVirtualUseHZB(
+	TEXT("r.Shadow.Virtual.UseHZB"),
+	0,
+	TEXT("Enables HZB for Virtual Shadow Maps."),
+	ECVF_RenderThreadSafe);
+
 static TAutoConsoleVariable<float> CVarNaniteShadowsLODBias(
 	TEXT("r.Shadow.NaniteLODBias"),
 	1.0f,
@@ -1594,13 +1600,14 @@ void FSceneRenderer::RenderShadowDepthMaps(FRDGBuilder& GraphBuilder, FInstanceC
 	const bool bHasVSMShadows = SortedShadowsForShadowDepthPass.VirtualShadowMapShadows.Num() > 0;
 	const bool bHasVSMClipMaps = SortedShadowsForShadowDepthPass.VirtualShadowMapClipmaps.Num() > 0;
 	const bool bNaniteEnabled = UseNanite(ShaderPlatform) && ViewFamily.EngineShowFlags.NaniteMeshes;
-	const bool bUseHZB = (CVarNaniteShadowsUseHZB.GetValueOnRenderThread() != 0);
 	const bool bAllocatePageRectAtlas = CVarAllocatePagesUsingRects.GetValueOnRenderThread() != 0;
 
 	if (bNaniteEnabled && (bHasVSMShadows || bHasVSMClipMaps))
-	{				
+	{
+		const bool bVSMUseHZB = (CVarShadowsVirtualUseHZB.GetValueOnRenderThread() != 0);
+
 		FVirtualShadowMapArrayCacheManager *CacheManager = Scene->VirtualShadowMapArrayCacheManager;
-		const TRefCountPtr<IPooledRenderTarget> PrevHZBPhysical = bUseHZB ? CacheManager->PrevBuffers.HZBPhysical : nullptr;
+		const TRefCountPtr<IPooledRenderTarget> PrevHZBPhysical = bVSMUseHZB ? CacheManager->PrevBuffers.HZBPhysical : nullptr;
 		
 		{
 			RDG_EVENT_SCOPE(GraphBuilder, "RenderVirtualShadowMaps");
@@ -1635,7 +1642,7 @@ void FSceneRenderer::RenderShadowDepthMaps(FRDGBuilder& GraphBuilder, FInstanceC
 				&VirtualShadowMapArray = VirtualShadowMapArray,
 				&GraphBuilder,
 				bUpdateStreaming,
-				bUseHZB,
+				bVSMUseHZB,
 				Scene = Scene,
 				CacheManager = CacheManager,
 				PrevHZBPhysical
@@ -1691,7 +1698,7 @@ void FSceneRenderer::RenderShadowDepthMaps(FRDGBuilder& GraphBuilder, FInstanceC
 							}
 
 							// If we're going to generate a new HZB this frame, save the associated metadata
-							if (bUseHZB)
+							if (bVSMUseHZB)
 							{
 								FVirtualShadowMapHZBMetadata& HZBMeta = VirtualShadowMapArray.HZBMetadata.FindOrAdd(HZBKey);
 								HZBMeta.TargetLayerIndex = Params.TargetLayerIndex;
@@ -1742,7 +1749,7 @@ void FSceneRenderer::RenderShadowDepthMaps(FRDGBuilder& GraphBuilder, FInstanceC
 							}
 
 							// If we're going to generate a new HZB this frame, save the associated metadata
-							if (bUseHZB)
+							if (bVSMUseHZB)
 							{
 								FVirtualShadowMapHZBMetadata& HZBMeta = VirtualShadowMapArray.HZBMetadata.FindOrAdd(HZBKey);
 								HZBMeta.TargetLayerIndex = Params.TargetLayerIndex;
@@ -1814,7 +1821,7 @@ void FSceneRenderer::RenderShadowDepthMaps(FRDGBuilder& GraphBuilder, FInstanceC
 				FilterAndRenderVirtualShadowMaps(false, VirtualFilterName);
 			}
 
-			if (bUseHZB)
+			if (bVSMUseHZB)
 			{
 				FRDGTextureRef SceneDepth = GraphBuilder.RegisterExternalTexture( GSystemTextures.BlackDummy );
 
@@ -1879,6 +1886,8 @@ void FSceneRenderer::RenderShadowDepthMaps(FRDGBuilder& GraphBuilder, FInstanceC
 			ProjectedShadowInfo->CacheMode != SDCM_MovablePrimitivesOnly		// See note in RenderShadowDepthMapAtlases
 			)
 		{
+			const bool bUseHZB = (CVarNaniteShadowsUseHZB.GetValueOnRenderThread() != 0);
+
 			FString LightName;
 			GetLightNameForDrawEvent(ProjectedShadowInfo->GetLightSceneInfo().Proxy, LightName);
 
diff --git a/Engine/Source/Runtime/Renderer/Private/VirtualShadowMaps/VirtualShadowMapArray.cpp b/Engine/Source/Runtime/Renderer/Private/VirtualShadowMaps/VirtualShadowMapArray.cpp
--- a/Engine/Source/Runtime/Renderer/Private/VirtualShadowMaps/VirtualShadowMapArray.cpp
+++ b/Engine/Source/Runtime/Renderer/Private/VirtualShadowMaps/VirtualShadowMapArray.cpp
@@ -1658,6 +1658,10 @@ void FVirtualShadowMapArray::RenderVirtualShadowMapsHw(FRDGBuilder& GraphBuilder
 		const int32 NumDrawCommands = MeshCommandPass.GetInstanceCullingContext()->CullingCommands.Num();
 		if (NumDrawCommands > 0)
 		{
+			FString LightNameWithLevel;
+			FSceneRenderer::GetLightNameForDrawEvent(ProjectedShadowInfo->GetLightSceneInfo().Proxy, LightNameWithLevel);
+			RDG_EVENT_SCOPE(GraphBuilder, "%s", *LightNameWithLevel);
+
 			TArray<int32, SceneRenderingAllocator> DrawCommandInstanceCountTmp;
 			DrawCommandInstanceCountTmp.AddZeroed(NumDrawCommands);
 			FRDGBufferRef DrawCommandInstanceCountRDG = CreateStructuredBuffer(GraphBuilder, TEXT("Shadow.Virtual.DrawCommandInstanceCount"), DrawCommandInstanceCountTmp);
