Fix crash with hair rect allocation when screen percentage is enabled.

diff --git a/Engine/Source/Runtime/Renderer/Private/HairStrands/HairStrandsUtils.cpp b/Engine/Source/Runtime/Renderer/Private/HairStrands/HairStrandsUtils.cpp
--- a/Engine/Source/Runtime/Renderer/Private/HairStrands/HairStrandsUtils.cpp
+++ b/Engine/Source/Runtime/Renderer/Private/HairStrands/HairStrandsUtils.cpp
@@ -214,10 +214,12 @@ FIntRect ComputeProjectedScreenRect(const FBox& B, const FViewInfo& View)
 		FVector(B.Max)
 	};
 
+	// Compute the MinP/MaxP in pixel coord, relative to View.ViewRect.Min
+	const FMatrix& ViewProj = View.ViewMatrices.GetViewProjectionMatrix();
 	for (uint32 i = 0; i < 8; ++i)
 	{
 		FVector2D P;
-		if (View.WorldToPixel(Vertices[i], P))
+		if (FSceneView::ProjectWorldToScreen(Vertices[i], View.ViewRect, ViewProj, P))
 		{
 			MinP.X = FMath::Min(MinP.X, P.X);
 			MinP.Y = FMath::Min(MinP.Y, P.Y);
@@ -226,15 +228,16 @@ FIntRect ComputeProjectedScreenRect(const FBox& B, const FViewInfo& View)
 		}
 	}
 
+	// Clamp to pixel border
 	FIntRect OutRect;
 	OutRect.Min = FIntPoint(FMath::FloorToInt(MinP.X), FMath::FloorToInt(MinP.Y));
 	OutRect.Max = FIntPoint(FMath::CeilToInt(MaxP.X),  FMath::CeilToInt(MaxP.Y));
 
 	// Clamp to screen rect
-	OutRect.Min.X = FMath::Max(View.ViewRect.Min.X, OutRect.Min.X);
-	OutRect.Min.Y = FMath::Max(View.ViewRect.Min.Y, OutRect.Min.Y);
-	OutRect.Max.X = FMath::Min(View.ViewRect.Max.X, OutRect.Max.X);
-	OutRect.Max.Y = FMath::Min(View.ViewRect.Max.Y, OutRect.Max.Y);
+	OutRect.Min.X = FMath::Clamp(OutRect.Min.X, View.ViewRect.Min.X, View.ViewRect.Max.X);
+	OutRect.Max.X = FMath::Clamp(OutRect.Max.X, View.ViewRect.Min.X, View.ViewRect.Max.X);
+	OutRect.Min.Y = FMath::Clamp(OutRect.Min.Y, View.ViewRect.Min.Y, View.ViewRect.Max.Y);
+	OutRect.Max.Y = FMath::Clamp(OutRect.Max.Y, View.ViewRect.Min.Y, View.ViewRect.Max.Y);
 
 	return OutRect;
 }
diff --git a/Engine/Source/Runtime/Renderer/Private/HairStrands/HairStrandsVisibility.cpp b/Engine/Source/Runtime/Renderer/Private/HairStrands/HairStrandsVisibility.cpp
--- a/Engine/Source/Runtime/Renderer/Private/HairStrands/HairStrandsVisibility.cpp
+++ b/Engine/Source/Runtime/Renderer/Private/HairStrands/HairStrandsVisibility.cpp
@@ -3146,6 +3146,14 @@ void RenderHairStrandsVisibilityBuffer(
 
 	const bool bGenerateTile = GHairStrandsTile > 0;
 
+	const FIntRect HairRect = ComputeVisibleHairStrandsMacroGroupsRect(View.ViewRect, MacroGroupDatas);
+	const int32 HairPixelCount = HairRect.Width() * HairRect.Height();
+	if (HairPixelCount <= 0)
+	{
+		View.HairStrandsViewData.VisibilityData = FHairStrandsVisibilityData();
+		return;
+	}
+
 	{
 		
 		{
