Remove the InlineAllocator from the  FPassDrawCommandArray and FPassDrawCommandBuildRequestArray typedefs.  The size of these arrays was causing the the Virtual Allocator to get used when allocating using FMemStack::Get(), resulting in significant performance isssues in the render thread.

diff --git a/Engine/Source/Runtime/Renderer/Private/SceneVisibility.cpp b/Engine/Source/Runtime/Renderer/Private/SceneVisibility.cpp
--- a/Engine/Source/Runtime/Renderer/Private/SceneVisibility.cpp
+++ b/Engine/Source/Runtime/Renderer/Private/SceneVisibility.cpp
@@ -1841,8 +1841,8 @@ namespace EMarkMaskBits
 	};
 }
 
-typedef TArray<FVisibleMeshDrawCommand, TInlineAllocator<AverageMeshBatchNumPerRelevancePacket>> FPassDrawCommandArray;
-typedef TArray<const FStaticMeshBatch*, TInlineAllocator<AverageMeshBatchNumPerRelevancePacket>> FPassDrawCommandBuildRequestArray;
+typedef TArray<FVisibleMeshDrawCommand> FPassDrawCommandArray;
+typedef TArray<const FStaticMeshBatch*> FPassDrawCommandBuildRequestArray;
 
 struct FDrawCommandRelevancePacket
 {
