Added DistanceFields LLM stat

diff --git a/Engine/Source/Runtime/Renderer/Private/DistanceFieldObjectManagement.cpp b/Engine/Source/Runtime/Renderer/Private/DistanceFieldObjectManagement.cpp
--- a/Engine/Source/Runtime/Renderer/Private/DistanceFieldObjectManagement.cpp
+++ b/Engine/Source/Runtime/Renderer/Private/DistanceFieldObjectManagement.cpp
@@ -21,6 +21,7 @@
 #include "DistanceFieldAmbientOcclusion.h"
 #include "ComponentRecreateRenderStateContext.h"
 #include "GlobalDistanceField.h"
+#include "HAL/LowLevelMemStats.h"
 #include "ProfilingDebugging/CpuProfilerTrace.h"
 
 extern void LumenUpdateDFObjectIndex(FScene* Scene, int32 MappingIndex);
@@ -67,6 +68,12 @@ FAutoConsoleVariableRef CVarTwoSidedSurfaceBiasExpand(
 
 int32 GDFPreviousReverseAtlasAllocationOrder = 0;
 
+#if ENABLE_LOW_LEVEL_MEM_TRACKER
+DECLARE_LLM_MEMORY_STAT(TEXT("DistanceFields"), STAT_DistanceFieldsLLM, STATGROUP_LLMFULL);
+DECLARE_LLM_MEMORY_STAT(TEXT("DistanceFields"), STAT_DistanceFieldsSummaryLLM, STATGROUP_LLM);
+LLM_DEFINE_TAG(DistanceFields, NAME_None, NAME_None, GET_STATFNAME(STAT_DistanceFieldsLLM), GET_STATFNAME(STAT_DistanceFieldsSummaryLLM));
+#endif // ENABLE_LOW_LEVEL_MEM_TRACKER
+
 // Must match equivalent shader defines
 template<> int32 TDistanceFieldObjectBuffers<DFPT_SignedDistanceField>::ObjectDataStride = 9;
 template<> int32 TDistanceFieldObjectBuffers<DFPT_SignedDistanceField>::ObjectBoundsStride = 2;
@@ -940,6 +947,7 @@ void FSceneRenderer::PrepareDistanceFieldScene(FRDGBuilder& GraphBuilder, bool b
 	CSV_SCOPED_TIMING_STAT_EXCLUSIVE(RenderDFAO);
 	TRACE_CPUPROFILER_EVENT_SCOPE(FSceneRenderer::PrepareDistanceFieldScene);
 	QUICK_SCOPE_CYCLE_COUNTER(STAT_PrepareDistanceFieldScene);
+	LLM_SCOPE_BYTAG(DistanceFields);
 
 	const bool bShouldPrepareHeightFieldScene = ShouldPrepareHeightFieldScene();
 	const bool bShouldPrepareDistanceFieldScene = ShouldPrepareDistanceFieldScene();
