Remove some debug code that was failing non-unity compile

diff --git a/Engine/Source/Runtime/Renderer/Private/VT/TexturePagePool.cpp b/Engine/Source/Runtime/Renderer/Private/VT/TexturePagePool.cpp
--- a/Engine/Source/Runtime/Renderer/Private/VT/TexturePagePool.cpp
+++ b/Engine/Source/Runtime/Renderer/Private/VT/TexturePagePool.cpp
@@ -376,10 +376,8 @@ void FTexturePagePool::UnmapPageMapping(FVirtualTextureSystem* System, uint32 Ma
 void FTexturePagePool::UnmapAllPages(FVirtualTextureSystem* System, uint16 pAddress, bool bMapAncestorPages)
 {
 	FPageEntry& PageEntry = Pages[pAddress];
-	FVirtualTextureProducer* Producer = nullptr;
 	if (PageEntry.PackedProducerHandle != 0u)
 	{
-		Producer = System->FindProducer(PageEntry.GetProducerHandle());
 		check(NumPagesAllocated > 0u);
 		--NumPagesAllocated;
 		PageHash.Remove(GetPageHash(PageEntry), pAddress);
diff --git a/Engine/Source/Runtime/Renderer/Private/VT/TexturePagePool.h b/Engine/Source/Runtime/Renderer/Private/VT/TexturePagePool.h
--- a/Engine/Source/Runtime/Renderer/Private/VT/TexturePagePool.h
+++ b/Engine/Source/Runtime/Renderer/Private/VT/TexturePagePool.h
@@ -163,8 +163,6 @@ private:
 			uint32 Local_vLevel : 4;
 			uint32 GroupIndex : 4;
 		};
-
-		FVirtualTextureProducerHandle GetProducerHandle() const { return FVirtualTextureProducerHandle(PackedProducerHandle); }
 	};
 
 	static uint16 GetPageHash(const FPageEntry& Entry);
