Fix a bug of separate pass for each view.

diff --git a/Engine/Source/Runtime/Renderer/Private/MobileShadingRenderer.cpp b/Engine/Source/Runtime/Renderer/Private/MobileShadingRenderer.cpp
--- a/Engine/Source/Runtime/Renderer/Private/MobileShadingRenderer.cpp
+++ b/Engine/Source/Runtime/Renderer/Private/MobileShadingRenderer.cpp
@@ -424,7 +424,7 @@ void FMobileSceneRenderer::InitViews(FRDGBuilder& GraphBuilder, FSceneTexturesCo
 	const bool bForceDepthResolve = (CVarMobileForceDepthResolve.GetValueOnRenderThread() == 1);
 	const bool bSeparateTranslucencyActive = IsMobileSeparateTranslucencyActive(Views.GetData(), Views.Num()); 
 	const bool bPostProcessUsesSceneDepth = PostProcessUsesSceneDepth(Views[0]) || IsMobileDistortionActive(Views[0]);
-	const bool bRequireSeparateViewPass = !Views[0].bIsMobileMultiViewEnabled && Views.Num() > 0;
+	const bool bRequireSeparateViewPass = Views.Num() > 1 && !Views[0].bIsMobileMultiViewEnabled;
 	bRequiresMultiPass = RequiresMultiPass(RHICmdList, Views[0]);
 	bKeepDepthContent =
 		bRequiresMultiPass ||
