Removed BUF_Transient from HaltonSequenceIteration buffer

diff --git a/Engine/Source/Runtime/Renderer/Private/Halton.cpp b/Engine/Source/Runtime/Renderer/Private/Halton.cpp
--- a/Engine/Source/Runtime/Renderer/Private/Halton.cpp
+++ b/Engine/Source/Runtime/Renderer/Private/Halton.cpp
@@ -76,7 +76,7 @@ void FHaltonSequenceIteration::InitRHI()
 	FRHIResourceCreateInfo CreateInfo(TEXT("HaltonSequenceIteration"));
 	{
 		CreateInfo.ResourceArray = &RandomSamples;
-		SequenceIteration = RHICreateStructuredBuffer(sizeof(FSequenceIterationData), RandomSamples.Num() * sizeof(FSequenceIterationData), BUF_Transient | BUF_FastVRAM | BUF_ShaderResource | BUF_UnorderedAccess, CreateInfo);
+		SequenceIteration = RHICreateStructuredBuffer(sizeof(FSequenceIterationData), RandomSamples.Num() * sizeof(FSequenceIterationData), BUF_FastVRAM | BUF_ShaderResource | BUF_UnorderedAccess, CreateInfo);
 	}
 }
 
