Add hair strands composition point after separate translucent pass.

diff --git a/Engine/Source/Runtime/Renderer/Private/DeferredShadingRenderer.cpp b/Engine/Source/Runtime/Renderer/Private/DeferredShadingRenderer.cpp
--- a/Engine/Source/Runtime/Renderer/Private/DeferredShadingRenderer.cpp
+++ b/Engine/Source/Runtime/Renderer/Private/DeferredShadingRenderer.cpp
@@ -2719,7 +2719,7 @@ void FDeferredShadingSceneRenderer::Render(FRDGBuilder& GraphBuilder)
 		RenderDebugSkyAtmosphere(GraphBuilder, SceneTextures.Color.Target, SceneTextures.Depth.Target);
 	}
 
-	if (!IsHairStrandsComposeAfterTranslucency())
+	if (GetHairStrandsComposition() == EHairStrandsCompositionType::BeforeTranslucent)
 	{
 		RDG_GPU_STAT_SCOPE(GraphBuilder, HairRendering);
 		RenderHairComposition(GraphBuilder, Views, SceneTextures.Color.Target, SceneTextures.Depth.Target);
@@ -2752,7 +2752,7 @@ void FDeferredShadingSceneRenderer::Render(FRDGBuilder& GraphBuilder)
 
 		// Compose hair before velocity/distortion pass since these pass write depth value, 
 		// and this would make the hair composition fails in this cases.
-		if (IsHairStrandsComposeAfterTranslucency())
+		if (GetHairStrandsComposition() == EHairStrandsCompositionType::AfterTranslucent)
 		{
 			RDG_GPU_STAT_SCOPE(GraphBuilder, HairRendering);
 			RenderHairComposition(GraphBuilder, Views, SceneTextures.Color.Target, SceneTextures.Depth.Target);
@@ -2782,7 +2782,7 @@ void FDeferredShadingSceneRenderer::Render(FRDGBuilder& GraphBuilder)
 
 		GraphBuilder.SetCommandListStat(GET_STATID(STAT_CLM_AfterTranslucency));
 	}
-	else if (IsHairStrandsComposeAfterTranslucency())
+	else if (GetHairStrandsComposition() == EHairStrandsCompositionType::AfterTranslucent)
 	{
 		RDG_GPU_STAT_SCOPE(GraphBuilder, HairRendering);
 		RenderHairComposition(GraphBuilder, Views, SceneTextures.Color.Target, SceneTextures.Depth.Target);
diff --git a/Engine/Source/Runtime/Renderer/Private/HairStrands/HairStrandsComposition.cpp b/Engine/Source/Runtime/Renderer/Private/HairStrands/HairStrandsComposition.cpp
--- a/Engine/Source/Runtime/Renderer/Private/HairStrands/HairStrandsComposition.cpp
+++ b/Engine/Source/Runtime/Renderer/Private/HairStrands/HairStrandsComposition.cpp
@@ -689,3 +689,19 @@ void RenderHairComposition(
 		}
 	}
 }
+
+void RenderHairComposition(
+	FRDGBuilder& GraphBuilder,
+	const FViewInfo& View,
+	FRDGTextureRef SceneColorTexture,
+	FRDGTextureRef SceneDepthTexture)
+{
+	if (View.Family && HairStrands::HasViewHairStrandsData(View))
+	{
+		InternalRenderHairComposition(
+			GraphBuilder,
+			View,
+			SceneColorTexture,
+			SceneDepthTexture);
+	}
+}
\ No newline at end of file
diff --git a/Engine/Source/Runtime/Renderer/Private/HairStrands/HairStrandsComposition.h b/Engine/Source/Runtime/Renderer/Private/HairStrands/HairStrandsComposition.h
--- a/Engine/Source/Runtime/Renderer/Private/HairStrands/HairStrandsComposition.h
+++ b/Engine/Source/Runtime/Renderer/Private/HairStrands/HairStrandsComposition.h
@@ -9,6 +9,12 @@
 #include "RenderGraph.h"
 #include "Renderer/Private/SceneRendering.h"
 
+void RenderHairComposition(
+	FRDGBuilder& GraphBuilder,
+	const FViewInfo& View,
+	FRDGTextureRef SceneColorTexture,
+	FRDGTextureRef SceneDepthTexture);
+
 void RenderHairComposition(
 	FRDGBuilder& GraphBuilder, 
 	const TArray<FViewInfo>& Views,
diff --git a/Engine/Source/Runtime/Renderer/Private/HairStrands/HairStrandsUtils.cpp b/Engine/Source/Runtime/Renderer/Private/HairStrands/HairStrandsUtils.cpp
--- a/Engine/Source/Runtime/Renderer/Private/HairStrands/HairStrandsUtils.cpp
+++ b/Engine/Source/Runtime/Renderer/Private/HairStrands/HairStrandsUtils.cpp
@@ -38,7 +38,7 @@ static int32 GHairVisibilityRectOptimEnable = 0;
 static FAutoConsoleVariableRef CVarHairVisibilityRectOptimEnable(TEXT("r.HairStrands.RectLightingOptim"), GHairVisibilityRectOptimEnable, TEXT("Hair Visibility use projected view rect to light only relevant pixels"));
 
 static int32 GHairStrandsComposeAfterTranslucency = 1;
-static FAutoConsoleVariableRef CVarHairStrandsComposeAfterTranslucency(TEXT("r.HairStrands.ComposeAfterTranslucency"), GHairStrandsComposeAfterTranslucency, TEXT("Compose hair rendering with scene color after the translucency pass if true. Otherwise compose hair defor the translucent objects are rendered."));
+static FAutoConsoleVariableRef CVarHairStrandsComposeAfterTranslucency(TEXT("r.HairStrands.ComposeAfterTranslucency"), GHairStrandsComposeAfterTranslucency, TEXT("0: Compose hair before translucent objects. 1: Compose hair after translucent objects, but before separate translucent objects. 2: Compose hair after seperate translucent objects"));
 
 static float GHairDualScatteringRoughnessOverride = 0;
 static FAutoConsoleVariableRef CVarHairDualScatteringRoughnessOverride(TEXT("r.HairStrands.DualScatteringRoughness"), GHairDualScatteringRoughnessOverride, TEXT("Override all roughness for the dual scattering evaluation. 0 means no override. Default:0"));
@@ -56,9 +56,15 @@ float GetHairDualScatteringRoughnessOverride()
 	return GHairDualScatteringRoughnessOverride;
 }
 
-bool IsHairStrandsComposeAfterTranslucency()
+EHairStrandsCompositionType GetHairStrandsComposition()
 {
-	return GHairStrandsComposeAfterTranslucency > 0;
+	switch (GHairStrandsComposeAfterTranslucency)
+	{
+	case 0	: return EHairStrandsCompositionType::BeforeTranslucent;
+	case 1	: return EHairStrandsCompositionType::AfterTranslucent;
+	case 2	: return EHairStrandsCompositionType::AfterSeparateTranslucent;
+	default	: return EHairStrandsCompositionType::BeforeTranslucent;
+	}
 }
 
 float GetDeepShadowAABBScale()
diff --git a/Engine/Source/Runtime/Renderer/Private/HairStrands/HairStrandsUtils.h b/Engine/Source/Runtime/Renderer/Private/HairStrands/HairStrandsUtils.h
--- a/Engine/Source/Runtime/Renderer/Private/HairStrands/HairStrandsUtils.h
+++ b/Engine/Source/Runtime/Renderer/Private/HairStrands/HairStrandsUtils.h
@@ -62,7 +62,14 @@ EHairVisibilityVendor GetVendor();
 uint32 GetVendorOptimalGroupSize1D();
 FIntPoint GetVendorOptimalGroupSize2D();
 
-bool IsHairStrandsComposeAfterTranslucency();
+enum class  EHairStrandsCompositionType : uint8
+{
+	BeforeTranslucent,
+	AfterTranslucent,
+	AfterSeparateTranslucent,
+};
+
+EHairStrandsCompositionType GetHairStrandsComposition();
 
 FVector4 PackHairRenderInfo(
 	float PrimaryRadiusAtDepth1,
diff --git a/Engine/Source/Runtime/Renderer/Private/PostProcess/PostProcessing.cpp b/Engine/Source/Runtime/Renderer/Private/PostProcess/PostProcessing.cpp
--- a/Engine/Source/Runtime/Renderer/Private/PostProcess/PostProcessing.cpp
+++ b/Engine/Source/Runtime/Renderer/Private/PostProcess/PostProcessing.cpp
@@ -521,6 +521,11 @@ void AddPostProcessingPasses(FRDGBuilder& GraphBuilder, const FViewInfo& View, c
 			}
 
 			SceneColor.Texture = LocalSceneColorTexture;
+
+			if (GetHairStrandsComposition() == EHairStrandsCompositionType::AfterSeparateTranslucent)
+			{
+				RenderHairComposition(GraphBuilder, View, SceneColor.Texture, SceneDepth.Texture);
+			}
 		}
 
 		// Post Process Material Chain - Before Tonemapping
