Always destroy ray tracing scene from the render thread. Previously it could be destroyed from worker thread, which would trigger a validation error when running without RHI thread.

diff --git a/Engine/Source/Runtime/Renderer/Private/RayTracing/RayTracingScene.cpp b/Engine/Source/Runtime/Renderer/Private/RayTracing/RayTracingScene.cpp
--- a/Engine/Source/Runtime/Renderer/Private/RayTracing/RayTracingScene.cpp
+++ b/Engine/Source/Runtime/Renderer/Private/RayTracing/RayTracingScene.cpp
@@ -50,6 +50,8 @@ FGraphEventRef FRayTracingScene::BeginCreate(FRDGBuilder& GraphBuilder)
 		RayTracingSceneSRV = RHICreateShaderResourceView(RayTracingSceneBuffer);
 	}
 
+	RayTracingSceneRHI.SafeRelease();
+
 	CreateRayTracingSceneTask = FFunctionGraphTask::CreateAndDispatchWhenReady([
 			&ResultScene = RayTracingSceneRHI,
 			Instances = MakeArrayView(Instances)]()
