Fixes D3D11 -d3ddebug error in Gen5TAA

diff --git a/Engine/Source/Runtime/Renderer/Private/PostProcess/TemporalAA.cpp b/Engine/Source/Runtime/Renderer/Private/PostProcess/TemporalAA.cpp
--- a/Engine/Source/Runtime/Renderer/Private/PostProcess/TemporalAA.cpp
+++ b/Engine/Source/Runtime/Renderer/Private/PostProcess/TemporalAA.cpp
@@ -1206,6 +1206,7 @@ static void AddGen5MainTemporalAAPasses(
 	RDG_EVENT_SCOPE(GraphBuilder, "TAAU Gen5 %dx%d -> %dx%d", InputRect.Width(), InputRect.Height(), OutputRect.Width(), OutputRect.Height());
 	RDG_GPU_STAT_SCOPE(GraphBuilder, TAA);
 
+	FRDGTextureRef BlackUintDummy = GSystemTextures.GetZeroUIntDummy(GraphBuilder);
 	FRDGTextureRef BlackDummy = GraphBuilder.RegisterExternalTexture(GSystemTextures.BlackDummy);
 	FRDGTextureRef WhiteDummy = GraphBuilder.RegisterExternalTexture(GSystemTextures.WhiteDummy);
 
@@ -1319,35 +1320,22 @@ static void AddGen5MainTemporalAAPasses(
 	// Setup the previous frame history
 	FScreenPassTextureViewportParameters PrevHistoryInfo;
 	FTAAHistoryTextures PrevHistory;
-	if (bCameraCut)
 	{
-		PrevHistoryInfo = GetScreenPassTextureViewportParameters(FScreenPassTextureViewport(
-			FIntPoint(1, 1), FIntRect(FIntPoint(0, 0), FIntPoint(1, 1))));
-
-		for (int32 i = 0; i < PrevHistory.LowResTextures.Num(); i++)
-		{
-			PrevHistory.LowResTextures[i] = BlackDummy;
-		}
-
-		for (int32 i = 0; i < PrevHistory.Textures.Num(); i++)
+		if (bCameraCut)
 		{
-			PrevHistory.Textures[i] = BlackDummy;
+			PrevHistoryInfo = GetScreenPassTextureViewportParameters(FScreenPassTextureViewport(
+				FIntPoint(1, 1), FIntRect(FIntPoint(0, 0), FIntPoint(1, 1))));
 		}
-
-		for (int32 i = 0; i < PrevHistory.SuperResTextures.Num(); i++)
+		else
 		{
-			PrevHistory.SuperResTextures[i] = BlackDummy;
+			PrevHistoryInfo = GetScreenPassTextureViewportParameters(FScreenPassTextureViewport(
+				InputHistory.ReferenceBufferSize,
+				InputHistory.ViewportRect));
 		}
-	}
-	else
-	{
-		PrevHistoryInfo = GetScreenPassTextureViewportParameters(FScreenPassTextureViewport(
-			InputHistory.ReferenceBufferSize,
-			InputHistory.ViewportRect));
 
 		for (int32 i = 0; i < PrevHistory.LowResTextures.Num(); i++)
 		{
-			if (InputHistory.LowResRT[i].IsValid())
+			if (InputHistory.LowResRT[i].IsValid() && !bCameraCut)
 			{
 				PrevHistory.LowResTextures[i] = GraphBuilder.RegisterExternalTexture(InputHistory.LowResRT[i]);
 			}
@@ -1359,10 +1347,14 @@ static void AddGen5MainTemporalAAPasses(
 
 		for (int32 i = 0; i < PrevHistory.Textures.Num(); i++)
 		{
-			if (InputHistory.RT[i].IsValid())
+			if (InputHistory.RT[i].IsValid() && !bCameraCut)
 			{
 				PrevHistory.Textures[i] = GraphBuilder.RegisterExternalTexture(InputHistory.RT[i]);
 			}
+			else if (i == 3)
+			{
+				PrevHistory.Textures[i] = BlackUintDummy;
+			}
 			else
 			{
 				PrevHistory.Textures[i] = BlackDummy;
@@ -1371,17 +1363,15 @@ static void AddGen5MainTemporalAAPasses(
 
 		for (int32 i = 0; i < PrevHistory.SuperResTextures.Num(); i++)
 		{
-			if (InputHistory.SuperResRT[i].IsValid())
+			if (InputHistory.SuperResRT[i].IsValid() && !bCameraCut)
 			{
 				PrevHistory.SuperResTextures[i] = GraphBuilder.RegisterExternalTexture(InputHistory.SuperResRT[i]);
 			}
 			else
 			{
-				PrevHistory.SuperResTextures[i] = BlackDummy;
+				PrevHistory.SuperResTextures[i] = BlackUintDummy;
 			}
 		}
-
-		// InputHistory.SafeRelease(); TODO
 	}
 
 
