Add light culling for path tracer

Split the light list info finite (point, spot, rect) and infinite (skydome, directional) sets. Build a 2d grid of which finite lights overlap particular regions. The grid is aligned to the bounding box of all lights and collapsed along the shortext axis so that memory usage scales quadratically with resolution instead of cubically. The light list per cell is currently pre-allocated up to a maximum count per cell (overflow is simply cut and lights are dropped as they were prior to this change). In most scenes I have tried so far this maximum is never reached with the current defaults of 256^2 resolution and 128 max lights per cluster.

diff --git a/Engine/Plugins/Experimental/GPULightmass/Shaders/Private/LightmapPathTracing.usf b/Engine/Plugins/Experimental/GPULightmass/Shaders/Private/LightmapPathTracing.usf
--- a/Engine/Plugins/Experimental/GPULightmass/Shaders/Private/LightmapPathTracing.usf
+++ b/Engine/Plugins/Experimental/GPULightmass/Shaders/Private/LightmapPathTracing.usf
@@ -391,7 +391,7 @@ void PathTracingKernel(
 				int LightId;
 				float LightPickPdf = 0;
 
-				SelectLight(RandSample.x * LightPickingCdfSum, LightPickingCdf, LightId, LightPickPdf);
+				SelectLight(RandSample.x * LightPickingCdfSum, SceneLightCount, LightPickingCdf, LightId, LightPickPdf);
 
 				FLightSample LightSample = SampleLight(LightId, RandSample, WorldPos, WorldNormal);
 
diff --git a/Engine/Plugins/Experimental/GPULightmass/Source/GPULightmass/Private/LightmapRayTracing.h b/Engine/Plugins/Experimental/GPULightmass/Source/GPULightmass/Private/LightmapRayTracing.h
--- a/Engine/Plugins/Experimental/GPULightmass/Source/GPULightmass/Private/LightmapRayTracing.h
+++ b/Engine/Plugins/Experimental/GPULightmass/Source/GPULightmass/Private/LightmapRayTracing.h
@@ -5,7 +5,6 @@
 #include "GPULightmassCommon.h"
 #include "ShaderParameterStruct.h"
 #include "GlobalShader.h"
-#include "PathTracingUniformBuffers.h"
 #include "RayTracing/RayTracingSkyLight.h"
 #include "SceneView.h"
 #include "RayTracing/RayTracingMaterialHitShaders.h"
diff --git a/Engine/Plugins/Experimental/GPULightmass/Source/GPULightmass/Private/LightmapRenderer.cpp b/Engine/Plugins/Experimental/GPULightmass/Source/GPULightmass/Private/LightmapRenderer.cpp
--- a/Engine/Plugins/Experimental/GPULightmass/Source/GPULightmass/Private/LightmapRenderer.cpp
+++ b/Engine/Plugins/Experimental/GPULightmass/Source/GPULightmass/Private/LightmapRenderer.cpp
@@ -25,6 +25,7 @@
 #include "PostProcess/PostProcessing.h"
 #include "RayTracingGeometryManager.h"
 #include "RayTracingDefinitions.h"
+#include "PathTracingDefinitions.h"
 #include "ScreenPass.h"
 
 class FCopyConvergedLightmapTilesCS : public FGlobalShader
diff --git a/Engine/Plugins/Experimental/GPULightmass/Source/GPULightmass/Private/PathTracingLightParameters.inl b/Engine/Plugins/Experimental/GPULightmass/Source/GPULightmass/Private/PathTracingLightParameters.inl
--- a/Engine/Plugins/Experimental/GPULightmass/Source/GPULightmass/Private/PathTracingLightParameters.inl
+++ b/Engine/Plugins/Experimental/GPULightmass/Source/GPULightmass/Private/PathTracingLightParameters.inl
@@ -2,6 +2,9 @@
 
 #pragma once
 
+#include "RayTracingDefinitions.h"
+#include "PathTracingDefinitions.h"
+
 template<typename PassParameterType>
 void SetupPathTracingLightParameters(
 	const GPULightmass::FLightSceneRenderState& LightScene,
diff --git a/Engine/Shaders/Private/PathTracing/Light/PathTracingLightCommon.ush b/Engine/Shaders/Private/PathTracing/Light/PathTracingLightCommon.ush
--- a/Engine/Shaders/Private/PathTracing/Light/PathTracingLightCommon.ush
+++ b/Engine/Shaders/Private/PathTracing/Light/PathTracingLightCommon.ush
@@ -7,6 +7,7 @@ LightCommon.usf: Common utility functions for light sampling
 #pragma once
 
 #include "/Engine/Shared/RayTracingTypes.h"
+#include "/Engine/Shared/PathTracingDefinitions.h"
 
 uint SceneLightCount;
 StructuredBuffer<FPathTracingLight> SceneLights;
@@ -56,12 +57,6 @@ FLightSample CreateLightSample(float3 RadianceOverPdf, float Pdf, float3 Directi
 	return Sample;
 }
 
-
-int GetLightId(FMaterialClosestHitPayload RayHitInfo)
-{
-	return -1;
-}
-
 bool IsEnvironmentLight(int LightId)
 {
 	return (SceneLights[LightId].Flags & PATHTRACER_FLAG_TYPE_MASK) == PATHTRACING_LIGHT_SKY;
@@ -93,16 +88,6 @@ bool IsPhysicalLight(int LightId)
 	return IsEnvironmentLight(LightId);
 }
 
-float3 GetLightUV(int LightId, RayDesc Ray, FMaterialClosestHitPayload RayHitInfo)
-{
-	float3 LightUV = 0.0;
-	if (IsEnvironmentLight(LightId))
-	{
-		LightUV = Ray.Direction;
-	}
-	return LightUV;
-}
-
 float3 GetPosition(int LightId)
 {
 	return SceneLights[LightId].Position;
diff --git a/Engine/Shaders/Private/PathTracing/Light/PathTracingLightGrid.ush b/Engine/Shaders/Private/PathTracing/Light/PathTracingLightGrid.ush
new file mode 100644
--- /dev/null
+++ b/Engine/Shaders/Private/PathTracing/Light/PathTracingLightGrid.ush
@@ -0,0 +1,153 @@
+// Copyright Epic Games, Inc. All Rights Reserved.
+
+#ifndef USE_LIGHT_GRID
+#define USE_LIGHT_GRID 0
+#endif
+
+#include "/Engine/Shared/PathTracingDefinitions.h"
+
+
+#if USE_LIGHT_GRID
+// extra shader parameters that must be defined when using the light grid
+uint SceneInfiniteLightCount;
+float3 SceneLightsBoundMin;
+float3 SceneLightsBoundMax;
+Texture2D<uint> LightGrid;
+Buffer<uint> LightGridData;
+uint LightGridResolution;
+uint LightGridMaxCount;
+int LightGridAxis;
+#endif
+
+
+// TODO: unify with similar function in other visualizations
+float3 HashToColor(uint H)
+{
+	return 0.5 + 0.3 * cos(2 * PI * ((H & 0xFFFFFF) * 5.96046447754e-08 + float3(0.0, 0.333, 0.6666)));
+}
+
+
+
+struct FLightLoopCount {
+	uint NumLights;
+	uint NumMISLights;
+#if USE_LIGHT_GRID
+	uint LightGridOffset;
+#endif
+};
+
+#if USE_LIGHT_GRID
+
+FLightLoopCount LightGridLookup(float3 WorldPos)
+{
+	FLightLoopCount Result;
+	Result.NumLights = SceneInfiniteLightCount;
+	Result.NumMISLights = SceneInfiniteLightCount;
+	Result.LightGridOffset = ~0u;
+
+	if (all(SceneLightsBoundMin <= WorldPos) && all(WorldPos <= SceneLightsBoundMax))
+	{
+		float3 P = WorldPos - SceneLightsBoundMin;
+		float3 D = SceneLightsBoundMax - SceneLightsBoundMin;
+		int2 UV = 0;
+		switch (LightGridAxis)
+		{
+			case 0: UV = int2(floor(LightGridResolution * P.yz / D.yz)); break;
+			case 1: UV = int2(floor(LightGridResolution * P.xz / D.xz)); break;
+			case 2: UV = int2(floor(LightGridResolution * P.xy / D.xy)); break;
+		}
+		uint LightGridN = LightGrid.Load(int3(UV, 0));
+		Result.LightGridOffset = LightGridMaxCount * (UV.x + UV.y * LightGridResolution);
+		Result.NumLights += LightGridN & PATHTRACER_LIGHT_GRID_LIGHT_COUNT_MASK;
+		if ((LightGridN & PATHTRACER_LIGHT_GRID_SINGULAR_MASK) == 0)
+		{
+			// at least some area lights, enable MIS for this cell
+			Result.NumMISLights = Result.NumLights;
+		}
+	}
+	return Result;
+}
+
+float3 LightGridVisualize(FLightLoopCount LoopCount, int VisualizeMode)
+{
+	const float3 OutsideColor = 0.18;
+	const float3 EmptyColor = 0.36;
+
+	if (LoopCount.LightGridOffset != ~0u)
+	{
+		switch (VisualizeMode)
+		{
+			case 2:
+			{
+				// color by unique light list
+				uint H = 0;
+				for (int Index = 0, Num = LoopCount.NumLights - SceneInfiniteLightCount; Index < Num; Index++)
+				{
+					H = StrongIntegerHash(H + LightGridData[LoopCount.LightGridOffset + Index]);
+				}
+				return HashToColor(H);
+			}
+			case 3:
+			{
+				// color by presence of singular lights
+				if (LoopCount.NumLights == SceneInfiniteLightCount)
+				{
+					return EmptyColor;
+				}
+				else if (LoopCount.NumLights == LoopCount.NumMISLights)
+				{
+					return float3(0.2, 0.2, 1); // blue - some area lights
+				}
+				else
+				{
+					return float3(0.2, 1, 0.2); // green - only point lights
+				}
+			}
+			default: 
+			{
+				// default mode - color by light count
+				uint N = LoopCount.NumLights - SceneInfiniteLightCount;
+				if (N < 1)
+				{
+					return EmptyColor;
+				}
+				float Max = LightGridMaxCount;
+				float t = saturate(float(N) / Max);
+				return t * saturate(2.0 - abs(float3(4, 2, 0) - 4.0 * t));
+			}
+		}
+	}
+	// outside light grid bounds
+	return OutsideColor;
+}
+
+uint GetLightId(uint Index, FLightLoopCount LoopCount)
+{
+	if (Index >= SceneInfiniteLightCount)
+	{
+		return LightGridData[LoopCount.LightGridOffset + Index - SceneInfiniteLightCount];
+	}
+	return Index;
+}
+
+#else
+
+FLightLoopCount LightGridLookup(float3 WorldPos)
+{
+	uint NumLights = min(SceneLightCount, RAY_TRACING_LIGHT_COUNT_MAXIMUM);
+	FLightLoopCount Result;
+	Result.NumLights = Result.NumMISLights = NumLights;
+	return Result;
+}
+
+float3 LightGridVisualize(FLightLoopCount LoopCount, int VisualizeMode)
+{
+	return 0.5;
+}
+
+uint GetLightId(uint Index, FLightLoopCount LoopCount)
+{
+	return Index;
+}
+
+#endif
diff --git a/Engine/Shaders/Private/PathTracing/Light/PathTracingLightSampling.ush b/Engine/Shaders/Private/PathTracing/Light/PathTracingLightSampling.ush
--- a/Engine/Shaders/Private/PathTracing/Light/PathTracingLightSampling.ush
+++ b/Engine/Shaders/Private/PathTracing/Light/PathTracingLightSampling.ush
@@ -102,17 +102,17 @@ float InitLightPickingCdf(
 }
 
 // #dxr_todo: passing LightPickingCdf as anything but inout causes performance to drop significantly -- why ?
-void SelectLight(float RandSample, inout float LightPickingCdf[RAY_TRACING_LIGHT_COUNT_MAXIMUM], out uint LightId, out float LightPickPdf)
+void SelectLight(float RandSample, int NumLights, inout float LightPickingCdf[RAY_TRACING_LIGHT_COUNT_MAXIMUM], out uint LightId, out float LightPickPdf)
 {
 	// This code assumes SceneLightCount > 0
 #ifdef UNIFORM_LIGHT_SELECTION
-	LightId = min(int(floor(RandSample * SceneLightCount)), SceneLightCount - 1);
+	LightId = min(int(floor(RandSample)), NumLights - 1);
 	LightPickPdf = 1.0;
 #else
 #if 0
 	// linear search
 	float PreviousCdfValue = 0;
-	for (LightId = 0; LightId < SceneLightCount - 1; ++LightId)
+	for (LightId = 0; LightId < NumLights - 1; ++LightId)
 	{
 		if (RandSample < LightPickingCdf[LightId])
 		{
@@ -125,7 +125,7 @@ void SelectLight(float RandSample, inout float LightPickingCdf[RAY_TRACING_LIGHT
 #else
 	// binary search
 	LightId = 0;
-	for (int Count = SceneLightCount; Count > 0;)
+	for (int Count = NumLights; Count > 0;)
 	{
 		int Step = Count / 2;
 		int Iterator = LightId + Step;
@@ -182,7 +182,7 @@ bool SampleLight(
 
 	float LightPickPdf = 0;
 
-	SelectLight(RandSample.x * LightPickingCdfSum, LightPickingCdf, OutLightId, LightPickPdf);
+	SelectLight(RandSample.x * LightPickingCdfSum, SceneLightCount, LightPickingCdf, OutLightId, LightPickPdf);
 
 	LightPickPdf /= LightPickingCdfSum;
 
diff --git a/Engine/Shaders/Private/PathTracing/Light/PathTracingRectLight.ush b/Engine/Shaders/Private/PathTracing/Light/PathTracingRectLight.ush
--- a/Engine/Shaders/Private/PathTracing/Light/PathTracingRectLight.ush
+++ b/Engine/Shaders/Private/PathTracing/Light/PathTracingRectLight.ush
@@ -28,14 +28,14 @@ float3 EvaluateTexture(int LightId, float2 UV)
 	UV.y = 1 - UV.y; // match orientation with raster path
 	switch (SceneLights[LightId].RectLightTextureIndex)
 	{
-		case 0:  return RectLightTexture_0.SampleLevel(RectLightSampler, UV, 0);
-		case 1:  return RectLightTexture_1.SampleLevel(RectLightSampler, UV, 0);
-		case 2:  return RectLightTexture_2.SampleLevel(RectLightSampler, UV, 0);
-		case 3:  return RectLightTexture_3.SampleLevel(RectLightSampler, UV, 0);
-		case 4:  return RectLightTexture_4.SampleLevel(RectLightSampler, UV, 0);
-		case 5:  return RectLightTexture_5.SampleLevel(RectLightSampler, UV, 0);
-		case 6:  return RectLightTexture_6.SampleLevel(RectLightSampler, UV, 0);
-		case 7:  return RectLightTexture_7.SampleLevel(RectLightSampler, UV, 0);
+		case 0:  return RectLightTexture_0.SampleLevel(RectLightSampler, UV, 0).xyz;
+		case 1:  return RectLightTexture_1.SampleLevel(RectLightSampler, UV, 0).xyz;
+		case 2:  return RectLightTexture_2.SampleLevel(RectLightSampler, UV, 0).xyz;
+		case 3:  return RectLightTexture_3.SampleLevel(RectLightSampler, UV, 0).xyz;
+		case 4:  return RectLightTexture_4.SampleLevel(RectLightSampler, UV, 0).xyz;
+		case 5:  return RectLightTexture_5.SampleLevel(RectLightSampler, UV, 0).xyz;
+		case 6:  return RectLightTexture_6.SampleLevel(RectLightSampler, UV, 0).xyz;
+		case 7:  return RectLightTexture_7.SampleLevel(RectLightSampler, UV, 0).xyz;
 		default: return 1;
 	}
 }
@@ -160,6 +160,10 @@ FLightSample RectLight_SampleLight(
 		GetRectLightBarnCosAngle(LightId),
 		GetRectLightBarnLength(LightId),
 		true /* bComputeVisibleRect */);
+	if (!IsRectVisible(Rect) || dot(Rect.Axis[2], Rect.Origin) < 0)
+	{
+		return NullLightSample();
+	}
 	FSphericalRect SphericalRect = BuildSphericalRect(Rect);
 
 #if USE_RECT_LIGHT_TEXTURES
diff --git a/Engine/Shaders/Private/PathTracing/Material/PathTracingMaterialSampling.ush b/Engine/Shaders/Private/PathTracing/Material/PathTracingMaterialSampling.ush
--- a/Engine/Shaders/Private/PathTracing/Material/PathTracingMaterialSampling.ush
+++ b/Engine/Shaders/Private/PathTracing/Material/PathTracingMaterialSampling.ush
@@ -132,7 +132,7 @@ float3 EstimateMaterialAlbedo(FMaterialClosestHitPayload Payload)
 			break;
 	}
 
-	// FIXME: could make this more accurate by estimating the directional albedo
+	// TODO: make this more accurate by estimating the directional albedo
 	float3 Albedo = 0;
 	switch (Payload.ShadingModelID)
 	{
diff --git a/Engine/Shaders/Private/PathTracing/PathTracing.usf b/Engine/Shaders/Private/PathTracing/PathTracing.usf
--- a/Engine/Shaders/Private/PathTracing/PathTracing.usf
+++ b/Engine/Shaders/Private/PathTracing/PathTracing.usf
@@ -10,6 +10,8 @@ PathTracingRayGenShader.usf: Reference path tracing
 
 #define RANDSEQ_UNROLL_SOBOL 0 // unrolling sobol sampler doesn't seem to pay off in the path tracer
 
+#define USE_LIGHT_GRID	1
+
 #include "../Common.ush"
 #include "../PostProcessCommon.ush"
 #include "../RectLight.ush"
@@ -21,22 +23,45 @@ PathTracingRayGenShader.usf: Reference path tracing
 #include "../ShadingModels.ush"
 #include "./Utilities/PathTracingRandomSequence.ush"
 #include "./Light/PathTracingLightSampling.ush"
+#include "./Light/PathTracingLightGrid.ush"
 #include "./Material/PathTracingMaterialSampling.ush"
 #include "PathTracingCamera.ush"
 
+uint Iteration;
+uint TemporalSeed;
+uint MaxSamples;
+uint UseErrorDiffusion;
+uint MaxBounces;
+uint MaxSSSBounces;
+
+// 0: only Material sampling
+// 1: only Light sampling
+// 2: both Material and Light
+uint MISMode;
+
+uint ApproximateCaustics;
+uint EnableCameraBackfaceCulling;
+uint SkipDirectLighting;
+uint SkipEmissive;
+uint CoherentSampling;
+uint VisualizeLightGrid;
+float MaxPathIntensity;
+float MaxNormalBias;
+float FilterWidth;
 
 RWTexture2D<float4> RadianceTexture;
 RaytracingAccelerationStructure TLAS;
 int3 TileOffset;
 uint SceneVisibleLightCount;
 
+
 void AccumulateRadiance(inout float3 TotalRadiance, float3 PathRadiance)
 {
-	if (PathTracingData.MaxPathIntensity > 0)
+	if (MaxPathIntensity > 0)
 	{
 		// User asked for path contributions to be clamped to reduce fireflies.
 		// Depending on how aggressive this value is, the image could be quite biased
-		TotalRadiance += min(PathRadiance, PathTracingData.MaxPathIntensity);
+		TotalRadiance += min(PathRadiance, MaxPathIntensity);
 	}
 	else
 	{
@@ -47,11 +72,17 @@ void AccumulateRadiance(inout float3 TotalRadiance, float3 PathRadiance)
 
 FMaterialClosestHitPayload TraceTransparentRay(RayDesc Ray, FRayCone RayCone, bool IsCameraRay, bool LastBounce, bool IncludeEmission, uint2 LaunchIndex, uint NumLights, inout RandomSequence RandSequence, inout float3 PathThroughput, inout float3 Radiance)
 {
-	const uint RayFlags = IsCameraRay && PathTracingData.EnableCameraBackfaceCulling ? RAY_FLAG_CULL_BACK_FACING_TRIANGLES : 0;
+	const uint RayFlags = IsCameraRay && EnableCameraBackfaceCulling ? RAY_FLAG_CULL_BACK_FACING_TRIANGLES : 0;
 	const uint MissShaderIndex = 0;
 	float SelectionWeightSum = 0;
 	float3 PayloadThroughput = PathThroughput;
 	FMaterialClosestHitPayload Payload;
+	if (!IncludeEmission && LastBounce)
+	{
+		Payload.SetMiss();
+		PathThroughput = 0;
+		return Payload;
+	}
 	for (;;)
 	{
 		FPackedMaterialClosestHitPayload PackedPayload = (FPackedMaterialClosestHitPayload)0;
@@ -260,7 +291,7 @@ bool ProcessSubsurfaceRandomWalk(inout FMaterialClosestHitPayload Payload, inout
 	float3 SSSColor = GetSubsurfaceColor(Payload);
 	float3 SSSRadius = GetSubsurfaceRadius(Payload);
 
-	if (SimplifySSS || all(SSSRadius == 0) || all(SSSColor == 0) || PathTracingData.MaxSSSBounces == 0)
+	if (SimplifySSS || all(SSSRadius == 0) || all(SSSColor == 0) || MaxSSSBounces == 0)
 	{
 		// if we are running in a context that won't see the SSS directly -- just turn it into diffuse and skip the random walk
 		// we can also skip random walk if the radius or color is 0
@@ -293,7 +324,7 @@ bool ProcessSubsurfaceRandomWalk(inout FMaterialClosestHitPayload Payload, inout
 	Ray.Origin = Payload.WorldPos;
 	Ray.Direction = TangentToWorld(-CosineSampleHemisphere(RandSample.xy).xyz, Payload.WorldNormal);
 	Ray.TMin = 0;
-	ApplyPositionBias(Ray, -Payload.WorldNormal, PathTracingData.MaxNormalBias);
+	ApplyPositionBias(Ray, -Payload.WorldNormal, MaxNormalBias);
 
 	// Use the technique detailed in the Hyperion paper (Section 4.4.2)
 	// https://media.disneyanimation.com/uploads/production/publication_asset/177/asset/a.pdf
@@ -309,7 +340,7 @@ bool ProcessSubsurfaceRandomWalk(inout FMaterialClosestHitPayload Payload, inout
 	const float3 Albedo = 1 - Pow2(4.09712 + 4.20863 * SSSColor - sqrt(9.59217 + SSSColor * (41.6808 + 17.7126 * SSSColor)));
 	SSSRadius *= 2.0; // roughly match parameterization above
 #endif
-	const int MAX_SSS_BOUNCES = PathTracingData.MaxSSSBounces;
+	const int MAX_SSS_BOUNCES = MaxSSSBounces;
 	const float3 SigmaT = 1.0 / SSSRadius;
 	const float3 SigmaS = Albedo * SigmaT;
 	for (int i = 0; i < MAX_SSS_BOUNCES; i++)
@@ -383,7 +414,7 @@ bool ProcessSubsurfaceRandomWalk(inout FMaterialClosestHitPayload Payload, inout
 				// NOTE: weight and probability cancel out, so no throughput adjustment is needed
 				Ray.Origin += ProbeResult.HitT * Ray.Direction;
 				Ray.Direction = reflect(Ray.Direction, WorldNormal);
-				ApplyPositionBias(Ray, WorldNormal, PathTracingData.MaxNormalBias);
+				ApplyPositionBias(Ray, WorldNormal, MaxNormalBias);
 				continue;
 			}
 #endif
@@ -459,24 +490,24 @@ RAY_TRACING_ENTRY_RAYGEN(PathTracingMainRG)
 	uint2 DispatchDim = DispatchRaysDimensions().xy;
 
 	// shuffle the pixels around so the ones that share a seed land in the same warp
-	uint2 LaunchIndex = (PathTracingData.CoherentSampling ? Reorder(DispatchIdx, DispatchDim) : DispatchIdx) + View.ViewRectMin.xy + TileOffset.xy;
+	uint2 LaunchIndex = (CoherentSampling ? Reorder(DispatchIdx, DispatchDim) : DispatchIdx) + View.ViewRectMin.xy + TileOffset.xy;
 
 	// Initialize random sequence 
 	RandomSequence RandSequence;
 
-	uint2 RandPixel = PathTracingData.CoherentSampling
+	uint2 RandPixel = CoherentSampling
 		? (DispatchIdx / PATHTRACER_COHERENT_TILE_SIZE)	// unique seed within small squares of the shuffled sub-images
 		//? (DispatchIdx * PATHTRACER_COHERENT_TILE_SIZE) / DispatchDim // unique seed per sub-image
 		: LaunchIndex; // unique seed per pixel
-	if (PathTracingData.UseErrorDiffusion)
+	if (UseErrorDiffusion)
 	{
 		// z-sampler init
-		RandomSequence_Initialize(RandSequence, RandPixel, PathTracingData.Iteration, PathTracingData.TemporalSeed - PathTracingData.Iteration, PathTracingData.MaxSamples);
+		RandomSequence_Initialize(RandSequence, RandPixel, Iteration, TemporalSeed - Iteration, MaxSamples);
 	}
 	else
 	{
 		// random sobol init
-		RandomSequence_Initialize(RandSequence, RandPixel.x + RandPixel.y * 65536, PathTracingData.TemporalSeed);
+		RandomSequence_Initialize(RandSequence, RandPixel.x + RandPixel.y * 65536, TemporalSeed);
 	}
 
 
@@ -501,25 +532,20 @@ RAY_TRACING_ENTRY_RAYGEN(PathTracingMainRG)
 	float3 PathThroughput = 1.0;
 	float PathRoughness = 0;
 
-	// 0: only Material sampling
-	// 1: only Light sampling
-	// 2: both Material and Light
-	const int MISMode = PathTracingData.MISMode;
-
 	// number of directly visible lights for the first bounce
 	uint NumVisibleLights = SceneVisibleLightCount;
 
-	for (int Bounce = 0; Bounce <= PathTracingData.MaxBounces; Bounce++)
+	for (int Bounce = 0; Bounce <= MaxBounces; Bounce++)
 	{
 #if 0
 		// visualize pixel seeds
-		Radiance = 0.5 + 0.3 * cos(2 * PI * ((StrongIntegerHash(RandPixel.x + StrongIntegerHash(RandPixel.y)) & 0xFFFFFF) * 5.96046447754e-08 + float3(0.0, 0.333, 0.6666)));
+		Radiance = HashToColor(StrongIntegerHash(RandPixel.x + StrongIntegerHash(RandPixel.y)));
 		break;
 #endif
 		const bool bIsCameraRay = Bounce == 0;
-		const bool bIsLastBounce = Bounce == PathTracingData.MaxBounces;
-		const bool bIncludeEmissive = (PathTracingData.SkipDirectLighting == 0 || Bounce > 1) &&
-									  (PathTracingData.SkipEmissive       == 0 || bIsCameraRay);
+		const bool bIsLastBounce = Bounce == MaxBounces;
+		const bool bIncludeEmissive = (SkipDirectLighting == 0 || Bounce > 1) &&
+									  (SkipEmissive       == 0 || bIsCameraRay);
 		FMaterialClosestHitPayload Payload = TraceTransparentRay(Ray, RayCone, bIsCameraRay, bIsLastBounce, bIncludeEmissive, LaunchIndex, NumVisibleLights, RandSequence, PathThroughput, Radiance);
 
 		if (Payload.IsMiss())
@@ -531,7 +557,7 @@ RAY_TRACING_ENTRY_RAYGEN(PathTracingMainRG)
 		// As soon as the path is blurry enough, we can get away with diffuse sampling only
 		const bool bSimplifySSS = PathRoughness >= 0.15;
 
-		if (PathTracingData.ApproximateCaustics)
+		if (ApproximateCaustics)
 		{
 			// modify the payload roughness to minimize difficult caustics
 			// This is inspired by a trick used in the Arnold renderer:
@@ -547,13 +573,24 @@ RAY_TRACING_ENTRY_RAYGEN(PathTracingMainRG)
 			break;
 		}
 
+		FLightLoopCount LightLoopCount = LightGridLookup(Payload.WorldPos);
+
+		// visualize the # of lights
+		if (VisualizeLightGrid)
+		{
+			Radiance = LightGridVisualize(LightLoopCount, VisualizeLightGrid);
+			Radiance *= abs(dot(Payload.WorldNormal, Ray.Direction));
+			Radiance /= View.PreExposure; // cancel out exposure
+			break;
+		}
+
 		// Choose a random number for both Light sampling and BxDF sampling
 		float4 RandSample = RandomSequence_GenerateSample4D(RandSequence);
 
 		// Does this material require NEE? (will be false if MaterialPdf is always +inf)
 		const bool bIsNeeValid = IsNeeValidMaterial(Payload);
 
-		const bool bDoLightLoop = PathTracingData.SkipDirectLighting == 0 || Bounce > 0;
+		const bool bDoLightLoop = SkipDirectLighting == 0 || Bounce > 0;
 
 		float LightPickingCdfSum = 0;
 
@@ -567,7 +604,12 @@ RAY_TRACING_ENTRY_RAYGEN(PathTracingMainRG)
 
 			bool IsTransmissiveMaterial = ENABLE_TRANSMISSION && IsMaterialTransmissive(Payload);
 
-			LightPickingCdfSum = InitLightPickingCdf(WorldPos, WorldNormal, PrimitiveLightingChannelMask, IsTransmissiveMaterial, LightPickingCdf);
+			for (uint Index = 0, Num = LightLoopCount.NumLights; Index < Num; ++Index)
+			{
+				uint LightIndex = GetLightId(Index, LightLoopCount);
+				LightPickingCdfSum += EstimateLight(LightIndex, WorldPos, WorldNormal, PrimitiveLightingChannelMask, IsTransmissiveMaterial);
+				LightPickingCdf[Index] = LightPickingCdfSum;
+			}
 
 			if (LightPickingCdfSum > 0)
 			{
@@ -575,9 +617,10 @@ RAY_TRACING_ENTRY_RAYGEN(PathTracingMainRG)
 				int LightId;
 				float LightPickPdf = 0;
 
-				SelectLight(RandSample.x * LightPickingCdfSum, LightPickingCdf, LightId, LightPickPdf);
+				SelectLight(RandSample.x * LightPickingCdfSum, LightLoopCount.NumLights, LightPickingCdf, LightId, LightPickPdf);
 				LightPickPdf /= LightPickingCdfSum;
 
+				LightId = GetLightId(LightId, LightLoopCount);
 				FLightSample LightSample = SampleLight(LightId, RandSample, WorldPos, WorldNormal);
 
 				LightSample.RadianceOverPdf /= LightPickPdf;
@@ -594,9 +637,9 @@ RAY_TRACING_ENTRY_RAYGEN(PathTracingMainRG)
 						LightRay.TMin = 0;
 						LightRay.Direction = LightSample.Direction;
 						LightRay.TMax = LightSample.Distance;
-						ApplyPositionBias(LightRay, SignedPositionBias * Payload.WorldNormal, PathTracingData.MaxNormalBias);
+						ApplyPositionBias(LightRay, SignedPositionBias * Payload.WorldNormal, MaxNormalBias);
 
-						float AvgRoughness = PathTracingData.ApproximateCaustics ? GetAverageRoughness(Payload) : 0.0;
+						float AvgRoughness = ApproximateCaustics ? GetAverageRoughness(Payload) : 0.0;
 
 						LightSample.RadianceOverPdf *= TraceTransparentVisibilityRay(LightRay, LaunchIndex, AvgRoughness);
 					}
@@ -673,9 +716,9 @@ RAY_TRACING_ENTRY_RAYGEN(PathTracingMainRG)
 		Ray.Direction = MaterialSample.Direction;
 		Ray.TMin = 0;
 		Ray.TMax = POSITIVE_INFINITY;
-		ApplyPositionBias(Ray, MaterialSample.PositionBiasSign * Payload.WorldNormal, PathTracingData.MaxNormalBias);
+		ApplyPositionBias(Ray, MaterialSample.PositionBiasSign * Payload.WorldNormal, MaxNormalBias);
 
-		if (PathTracingData.ApproximateCaustics)
+		if (ApproximateCaustics)
 		{
 			// enlarge roughness based on the chosen lobe roughness
 			PathRoughness = max(PathRoughness, MaterialSample.Roughness);
@@ -688,8 +731,9 @@ RAY_TRACING_ENTRY_RAYGEN(PathTracingMainRG)
 			// While it would be possible to just loop around and use the indirect ray to do this, it would prevent the application
 			// of shadow ray specific logic for transparent shadows or various per light tricks like shadow casting
 			const bool bUseMIS = MISMode == 2 && LightPickingCdfSum > 0;
-			for (uint LightId = 0; LightId < SceneLightCount; ++LightId)
+			for (uint Index = 0, Num = LightLoopCount.NumMISLights; Index < Num; ++Index)
 			{
+				uint LightId = GetLightId(Index, LightLoopCount);
 				if ((Payload.PrimitiveLightingChannelMask & GetLightingChannelMask(LightId)) == 0)
 				{
 					// light does not affect the current ray
@@ -706,8 +750,8 @@ RAY_TRACING_ENTRY_RAYGEN(PathTracingMainRG)
 				float3 LightContrib = PathThroughput * LightResult.Radiance;
 				if (bUseMIS)
 				{
-					float PreviousCdfValue = LightId > 0 ? LightPickingCdf[LightId - 1] : 0.0;
-					float LightPickPdf =  (LightPickingCdf[LightId] - PreviousCdfValue) / LightPickingCdfSum;
+					float PreviousCdfValue = Index > 0 ? LightPickingCdf[Index - 1] : 0.0;
+					float LightPickPdf = (LightPickingCdf[Index] - PreviousCdfValue) / LightPickingCdfSum;
 
 					LightContrib *= MISWeightRobust(MaterialSample.Pdf, LightResult.Pdf * LightPickPdf);
 				}
@@ -737,7 +781,7 @@ RAY_TRACING_ENTRY_RAYGEN(PathTracingMainRG)
 
 
 	// Accumulate radiance and update pixel variance
-	uint NumSamples = PathTracingData.Iteration + 1;
+	uint NumSamples = Iteration + 1;
 	float BlendFactor = 1.0 / float(NumSamples);
 	// Avoid reading the old pixel on the first sample on the off-chance there is a NaN/Inf pixel ...
 	float4 OldPixel = NumSamples > 1 ? RadianceTexture[LaunchIndex] : 0;
diff --git a/Engine/Shaders/Private/PathTracing/PathTracingBuildLightGrid.usf b/Engine/Shaders/Private/PathTracing/PathTracingBuildLightGrid.usf
new file mode 100644
--- /dev/null
+++ b/Engine/Shaders/Private/PathTracing/PathTracingBuildLightGrid.usf
@@ -0,0 +1,96 @@
+// Copyright Epic Games, Inc. All Rights Reserved.
+
+#include "../Common.ush"
+#include "/Engine/Shared/RayTracingDefinitions.h"
+#include "/Engine/Shared/RayTracingTypes.h"
+#include "/Engine/Shared/PathTracingDefinitions.h"
+
+#define USE_ATTENUATION_TERM 0
+#define USE_IES_TERM 0
+#include "Light/PathTracingLightCommon.ush"
+
+uint SceneInfiniteLightCount;
+float3 SceneLightsBoundMin;
+float3 SceneLightsBoundMax;
+RWTexture2D<uint> RWLightGrid;
+RWBuffer<uint> RWLightGridData;
+uint LightGridResolution;
+uint LightGridMaxCount;
+int LightGridAxis;
+
+
+[numthreads(THREADGROUPSIZE_X, THREADGROUPSIZE_Y, 1)]
+void PathTracingBuildLightGridCS(uint2 DispatchThreadId : SV_DispatchThreadID)
+{
+	if (any(DispatchThreadId >= LightGridResolution))
+	{
+		return;
+	}
+	// figure out dimension of the 3d grid and the current ID to be filled
+	uint3 VoxelId = 0, VoxelRes = 1;
+	switch (LightGridAxis)
+	{
+		case 0: VoxelId.yz = DispatchThreadId; VoxelRes.yz = LightGridResolution; break;
+		case 1: VoxelId.xz = DispatchThreadId; VoxelRes.xz = LightGridResolution; break;
+		case 2: VoxelId.xy = DispatchThreadId; VoxelRes.xy = LightGridResolution; break;
+	}
+	// get bounding box of current voxel
+	float3 Lo = lerp(SceneLightsBoundMin, SceneLightsBoundMax, float3(VoxelId + 0) / float3(VoxelRes));
+	float3 Hi = lerp(SceneLightsBoundMin, SceneLightsBoundMax, float3(VoxelId + 1) / float3(VoxelRes));
+
+	uint NumVisible = 0;
+	uint MaxVisible = min(LightGridMaxCount, RAY_TRACING_LIGHT_COUNT_MAXIMUM - SceneInfiniteLightCount);
+	uint Offset = LightGridMaxCount * (DispatchThreadId.x + LightGridResolution * DispatchThreadId.y);
+	bool AllSingular = true;
+	for (uint LightIndex = SceneInfiniteLightCount; LightIndex < SceneLightCount; LightIndex++)
+	{
+		uint LightType = SceneLights[LightIndex].Flags & PATHTRACER_FLAG_TYPE_MASK;
+		float3 Center = GetPosition(LightIndex);
+		float Radius = 1.0 / GetAttenuation(LightIndex);
+
+		float3 E = max(Lo - Center, 0) + max(Center - Hi, 0);
+		float D2 = dot(E, E);
+
+		// reject against normal
+		if (LightType != PATHTRACING_LIGHT_POINT)
+		{
+			float3 N = GetNormal(LightIndex);
+			float x = N.x > 0 ? Hi.x : Lo.x;
+			float y = N.y > 0 ? Hi.y : Lo.y;
+			float z = N.z > 0 ? Hi.z : Lo.z;
+
+			if (dot(float3(x, y, z) - Center, N) < 0)
+			{
+				continue;
+			}
+		}
+
+		if (LightType == PATHTRACING_LIGHT_SPOT)
+		{
+			float3 LightBoundLo = SceneLights[LightIndex].BoundMin;
+			float3 LightBoundHi = SceneLights[LightIndex].BoundMax;
+
+			if (any(Hi < LightBoundLo) || any(Lo > LightBoundHi))
+			{
+				continue;
+			}
+		}
+
+		if (D2 < Square(Radius))
+		{
+			AllSingular = AllSingular && all(SceneLights[LightIndex].Dimensions == 0);
+			RWLightGridData[Offset + NumVisible] = LightIndex;
+			NumVisible++;
+			// TODO: handle overflow better?
+			if (NumVisible >= MaxVisible)
+			{
+				break;
+			}
+		}
+	}
+	if (AllSingular)
+	{
+		NumVisible |= PATHTRACER_LIGHT_GRID_SINGULAR_MASK;
+	}
+	RWLightGrid[DispatchThreadId] = NumVisible;
+}
diff --git a/Engine/Shaders/Private/PathTracing/Sensor/PathTracingPinholeCamera.ush b/Engine/Shaders/Private/PathTracing/Sensor/PathTracingPinholeCamera.ush
--- a/Engine/Shaders/Private/PathTracing/Sensor/PathTracingPinholeCamera.ush
+++ b/Engine/Shaders/Private/PathTracing/Sensor/PathTracingPinholeCamera.ush
@@ -20,7 +20,7 @@ void PinholeCamera_SampleEmitter(
 
 	// importance sample a gaussian kernel with variable sigma
 	float3 Disk = ConcentricDiskSamplingHelper(AAJitter);
-	float Sigma = PathTracingData.FilterWidth / 6.0; // user-provided width covers +/-3*Sigma
+	float Sigma = FilterWidth / 6.0; // user-provided width covers +/-3*Sigma
 	float R = min(Disk.z, 0.99999994); // avoid log(0) when R=1
 	AAJitter = 0.5 + Sigma * Disk.xy * sqrt(-2.0 * log(1.0 - R * R));
 
diff --git a/Engine/Shaders/Shared/PathTracingDefinitions.h b/Engine/Shaders/Shared/PathTracingDefinitions.h
new file mode 100644
--- /dev/null
+++ b/Engine/Shaders/Shared/PathTracingDefinitions.h
@@ -0,0 +1,26 @@
+// Copyright Epic Games, Inc. All Rights Reserved.
+
+#pragma once
+
+
+#define PATHTRACER_MAX_RECT_TEXTURES		8
+#define PATHTRACER_COHERENT_TILE_SHIFT		3
+#define PATHTRACER_COHERENT_TILE_SIZE		(1 << PATHTRACER_COHERENT_TILE_SHIFT)
+
+
+// Constants for the 'Flags' field of FPathTracingLight
+#define PATHTRACER_FLAG_LIGHTING_CHANNEL_MASK				(7 << 0)	// Which lighting channel is this light assigned to?
+#define PATHTRACER_FLAG_TRANSMISSION_MASK					(1 << 3)	// Does the light affect the transmission side?
+#define PATHTRACER_FLAG_NON_INVERSE_SQUARE_FALLOFF_MASK		(1 << 4)	// Does the light have a non-inverse square decay?
+#define PATHTRACER_FLAG_STATIONARY_MASK						(1 << 5)	// Only used by GPULightmass
+#define PATHTRACER_FLAG_TYPE_MASK							(7 << 6)
+#define PATHTRACING_LIGHT_SKY								(0 << 6)
+#define PATHTRACING_LIGHT_DIRECTIONAL						(1 << 6)
+#define PATHTRACING_LIGHT_POINT								(2 << 6)
+#define PATHTRACING_LIGHT_SPOT								(3 << 6)
+#define PATHTRACING_LIGHT_RECT								(4 << 6)
+#define PATHTRACER_FLAG_CAST_SHADOW_MASK 					(1 << 9)
+
+// Constants for the path tracer light grid
+#define PATHTRACER_LIGHT_GRID_SINGULAR_MASK					0x80000000u
+#define PATHTRACER_LIGHT_GRID_LIGHT_COUNT_MASK				0x7FFFFFFFu
diff --git a/Engine/Shaders/Shared/RayTracingDefinitions.h b/Engine/Shaders/Shared/RayTracingDefinitions.h
--- a/Engine/Shaders/Shared/RayTracingDefinitions.h
+++ b/Engine/Shaders/Shared/RayTracingDefinitions.h
@@ -35,21 +35,3 @@
 #define RAY_TRACING_MAX_ALLOWED_RECURSION_DEPTH 1   // Only allow ray tracing from RayGen shader
 #define RAY_TRACING_MAX_ALLOWED_ATTRIBUTE_SIZE  8   // Sizeof 2 floats (barycentrics)
 #define RAY_TRACING_MAX_ALLOWED_PAYLOAD_SIZE    64  // Our maximum allowed payload size (sizeof FPackedMaterialClosestHitPayload)
-
-#define PATHTRACER_MAX_RECT_TEXTURES		8
-#define PATHTRACER_COHERENT_TILE_SHIFT		3
-#define PATHTRACER_COHERENT_TILE_SIZE		(1 << PATHTRACER_COHERENT_TILE_SHIFT)
-
-
-// Constants for the 'Flags' field of FPathTracingLight
-#define PATHTRACER_FLAG_LIGHTING_CHANNEL_MASK				(7 << 0)	// Which lighting channel is this light assigned to?
-#define PATHTRACER_FLAG_TRANSMISSION_MASK					(1 << 3)	// Does the light affect the transmission side?
-#define PATHTRACER_FLAG_NON_INVERSE_SQUARE_FALLOFF_MASK		(1 << 4)	// Does the light have a non-inverse square decay?
-#define PATHTRACER_FLAG_STATIONARY_MASK						(1 << 5)	// Only used by GPULightmass
-#define PATHTRACER_FLAG_TYPE_MASK							(7 << 6)
-#define PATHTRACING_LIGHT_SKY								(0 << 6)
-#define PATHTRACING_LIGHT_DIRECTIONAL						(1 << 6)
-#define PATHTRACING_LIGHT_POINT								(2 << 6)
-#define PATHTRACING_LIGHT_SPOT								(3 << 6)
-#define PATHTRACING_LIGHT_RECT								(4 << 6)
-#define PATHTRACER_FLAG_CAST_SHADOW_MASK 					(1 << 9)
\ No newline at end of file
diff --git a/Engine/Shaders/Shared/RayTracingTypes.h b/Engine/Shaders/Shared/RayTracingTypes.h
--- a/Engine/Shaders/Shared/RayTracingTypes.h
+++ b/Engine/Shaders/Shared/RayTracingTypes.h
@@ -22,10 +22,11 @@ struct FPathTracingLight {
 	float   Attenuation;
 	float   FalloffExponent; // for non-inverse square decay lights only
 	int32   IESTextureSlice;
-	uint32  Flags; // see defines above
+	uint32  Flags; // see defines PATHTRACER_FLAG_*
 	int32   RectLightTextureIndex;
-
-	float Padding[7]; // keep structure aligned - take new fields from here as needed
+	FVector BoundMin;
+	FVector BoundMax;
+	float Padding; // keep structure aligned
 };
 
 static_assert(sizeof(FPathTracingLight) == 128, "Path tracing light structure should be aligned to 128 bytes for optimal access on the GPU");
@@ -47,8 +48,9 @@ struct FPathTracingLight {
 	int     IESTextureSlice;
 	uint    Flags;
 	int	    RectLightTextureIndex;
-
-	float Padding[7];
+	float3  BoundMin;
+	float3  BoundMax;
+	float Padding;
 };
 
 #endif
diff --git a/Engine/Source/Runtime/Renderer/Private/PathTracing.cpp b/Engine/Source/Runtime/Renderer/Private/PathTracing.cpp
--- a/Engine/Source/Runtime/Renderer/Private/PathTracing.cpp
+++ b/Engine/Source/Runtime/Renderer/Private/PathTracing.cpp
@@ -7,9 +7,10 @@
 #include "RendererPrivate.h"
 #include "GlobalShader.h"
 #include "DeferredShadingRenderer.h"
-#include "PathTracingUniformBuffers.h"
 #include "HAL/PlatformApplicationMisc.h"
 #include "RayTracingTypes.h"
+#include "RayTracingDefinitions.h"
+#include "PathTracingDefinitions.h"
 #include "RenderCore/Public/GenerateMips.h"
 
 TAutoConsoleVariable<int32> CVarPathTracingMaxBounces(
@@ -153,7 +154,50 @@ TAutoConsoleVariable<int32> CVarPathTracingProgressDisplay(
 	ECVF_RenderThreadSafe
 );
 
-IMPLEMENT_GLOBAL_SHADER_PARAMETER_STRUCT(FPathTracingData, "PathTracingData");
+TAutoConsoleVariable<int32> CVarPathTracingLightGridResolution(
+	TEXT("r.PathTracing.LightGridResolution"),
+	256,
+	TEXT("Controls the resolution of the 2D light grid used to cull irrelevant lights from lighting calculations (default = 256)\n"),
+	ECVF_RenderThreadSafe
+);
+
+TAutoConsoleVariable<int32> CVarPathTracingLightGridMaxCount(
+	TEXT("r.PathTracing.LightGridMaxCount"),
+	128,
+	TEXT("Controls the maximum number of lights per cell in the 2D light grid (default = 128)\n"),
+	ECVF_RenderThreadSafe
+);
+
+TAutoConsoleVariable<int32> CVarPathTracingLightGridVisualize(
+	TEXT("r.PathTracing.LightGridVisualize"),
+	0,
+	TEXT("Enables a visualization mode of the light grid density where red indicates the maximum light count has been reached (default = 0)\n")
+	TEXT("0: off (default)\n")
+	TEXT("1: light count heatmap (red - close to overflow, increase r.PathTracing.LightGridMaxCount)")
+	TEXT("2: unique light lists (colors are a function of which lights occupy each cell)")
+	TEXT("3: area light visualization (green: point light sources only, blue: some area light sources)"),
+	ECVF_RenderThreadSafe
+);
+
+BEGIN_SHADER_PARAMETER_STRUCT(FPathTracingData, )
+	SHADER_PARAMETER(uint32, Iteration)
+	SHADER_PARAMETER(uint32, TemporalSeed)
+	SHADER_PARAMETER(uint32, MaxSamples)
+	SHADER_PARAMETER(uint32, UseErrorDiffusion)
+	SHADER_PARAMETER(uint32, MaxBounces)
+	SHADER_PARAMETER(uint32, MaxSSSBounces)
+	SHADER_PARAMETER(uint32, MISMode)
+	SHADER_PARAMETER(uint32, ApproximateCaustics)
+	SHADER_PARAMETER(uint32, EnableCameraBackfaceCulling)
+	SHADER_PARAMETER(uint32, SkipDirectLighting)
+	SHADER_PARAMETER(uint32, SkipEmissive)
+	SHADER_PARAMETER(uint32, CoherentSampling)
+	SHADER_PARAMETER(uint32, VisualizeLightGrid)
+	SHADER_PARAMETER(float, MaxPathIntensity)
+	SHADER_PARAMETER(float, MaxNormalBias)
+	SHADER_PARAMETER(float, FilterWidth)
+END_SHADER_PARAMETER_STRUCT()
+
 
 // This function prepares the portion of shader arguments that may involve invalidating the path traced state
 static bool PrepareShaderArgs(const FViewInfo& View, FPathTracingData& PathTracingData) {
@@ -196,6 +240,7 @@ static bool PrepareShaderArgs(const FViewInfo& View, FPathTracingData& PathTraci
 	PathTracingData.EnableCameraBackfaceCulling = CVarPathTracingEnableCameraBackfaceCulling.GetValueOnRenderThread();
 	PathTracingData.CoherentSampling = CVarPathTracingCoherentSampling.GetValueOnRenderThread();
 	PathTracingData.SkipEmissive = CVarPathTracingSkipEmissive.GetValueOnRenderThread();
+	PathTracingData.VisualizeLightGrid = CVarPathTracingLightGridVisualize.GetValueOnRenderThread();
 	float FilterWidth = CVarPathTracingFilterWidth.GetValueOnRenderThread();
 	if (FilterWidth < 0)
 	{
@@ -206,7 +251,7 @@ static bool PrepareShaderArgs(const FViewInfo& View, FPathTracingData& PathTraci
 	bool NeedInvalidation = false;
 
 	// If any of the parameters above changed since last time -- reset the accumulation
-	// FIXME: find something cleaner than just using static variables here. Should all
+	// TODO: find something cleaner than just using static variables here. Should all
 	// the state used for comparison go into ViewState?
 	static uint32 PrevMaxBounces = PathTracingData.MaxBounces;
 	if (PathTracingData.MaxBounces != PrevMaxBounces)
@@ -303,6 +348,14 @@ static bool PrepareShaderArgs(const FViewInfo& View, FPathTracingData& PathTraci
 		PreviousCoherentSampling = PathTracingData.CoherentSampling;
 	}
 
+	// Changing visualization mode requires starting over
+	static int32 PreviousVisualizeLightGrid = PathTracingData.VisualizeLightGrid;
+	if (PreviousVisualizeLightGrid != PathTracingData.VisualizeLightGrid)
+	{
+		NeedInvalidation = true;
+		PreviousVisualizeLightGrid = PathTracingData.VisualizeLightGrid;
+	}
+
 	// the rest of PathTracingData and AdaptiveSamplingData is filled in by SetParameters below
 	return NeedInvalidation;
 }
@@ -361,10 +414,50 @@ class FPathTracingSkylightMISCompensationCS : public FGlobalShader
 		SHADER_PARAMETER_RDG_TEXTURE_UAV(RWTexture2D, SkylightTextureOutput)
 		SHADER_PARAMETER_RDG_TEXTURE_UAV(RWTexture2D, SkylightTexturePdf)
 		SHADER_PARAMETER(FVector, SkyColor)
-		END_SHADER_PARAMETER_STRUCT()
+	END_SHADER_PARAMETER_STRUCT()
 };
 IMPLEMENT_SHADER_TYPE(, FPathTracingSkylightMISCompensationCS, TEXT("/Engine/Private/PathTracing/PathTracingSkylightMISCompensation.usf"), TEXT("PathTracingSkylightMISCompensationCS"), SF_Compute);
 
+// this struct holds a light grid for both building or rendering
+BEGIN_SHADER_PARAMETER_STRUCT(FPathTracingLightGrid,)
+	SHADER_PARAMETER(uint32, SceneInfiniteLightCount)
+	SHADER_PARAMETER(FVector, SceneLightsBoundMin)
+	SHADER_PARAMETER(FVector, SceneLightsBoundMax)
+	SHADER_PARAMETER_RDG_TEXTURE(Texture2D, LightGrid)
+	SHADER_PARAMETER_RDG_BUFFER_SRV(Buffer<uint>, LightGridData)
+	SHADER_PARAMETER(unsigned, LightGridResolution)
+	SHADER_PARAMETER(unsigned, LightGridMaxCount)
+	SHADER_PARAMETER(int, LightGridAxis)
+END_SHADER_PARAMETER_STRUCT()
+
+class FPathTracingBuildLightGridCS : public FGlobalShader
+{
+	DECLARE_GLOBAL_SHADER(FPathTracingBuildLightGridCS)
+	SHADER_USE_PARAMETER_STRUCT(FPathTracingBuildLightGridCS, FGlobalShader)
+
+	static bool ShouldCompilePermutation(const FGlobalShaderPermutationParameters& Parameters)
+	{
+		return ShouldCompileRayTracingShadersForProject(Parameters.Platform);
+	}
+
+	static void ModifyCompilationEnvironment(const FGlobalShaderPermutationParameters& Parameters, FShaderCompilerEnvironment& OutEnvironment)
+	{
+		//OutEnvironment.CompilerFlags.Add(CFLAG_WarningsAsErrors);
+		OutEnvironment.CompilerFlags.Add(CFLAG_AllowTypedUAVLoads);
+		OutEnvironment.SetDefine(TEXT("THREADGROUPSIZE_X"), FComputeShaderUtils::kGolden2DGroupSize);
+		OutEnvironment.SetDefine(TEXT("THREADGROUPSIZE_Y"), FComputeShaderUtils::kGolden2DGroupSize);
+	}
+
+	BEGIN_SHADER_PARAMETER_STRUCT(FParameters, )
+		SHADER_PARAMETER_RDG_BUFFER_SRV(StructuredBuffer<FPathTracingLight>, SceneLights)
+		SHADER_PARAMETER(uint32, SceneLightCount)
+		SHADER_PARAMETER_STRUCT_INCLUDE(FPathTracingLightGrid, LightGridParameters)
+		SHADER_PARAMETER_RDG_TEXTURE_UAV(RWTexture2D, RWLightGrid)
+		SHADER_PARAMETER_RDG_BUFFER_UAV(RWBuffer<uint>, RWLightGridData)
+	END_SHADER_PARAMETER_STRUCT()
+};
+IMPLEMENT_SHADER_TYPE(, FPathTracingBuildLightGridCS, TEXT("/Engine/Private/PathTracing/PathTracingBuildLightGrid.usf"), TEXT("PathTracingBuildLightGridCS"), SF_Compute);
+
 
 class FPathTracingRG : public FGlobalShader
 {
@@ -390,10 +483,14 @@ class FPathTracingRG : public FGlobalShader
 		SHADER_PARAMETER_SRV(RaytracingAccelerationStructure, TLAS)
 
 		SHADER_PARAMETER_STRUCT_REF(FViewUniformShaderParameters, ViewUniformBuffer)
-		SHADER_PARAMETER_STRUCT_REF(FPathTracingData, PathTracingData)
+		SHADER_PARAMETER_STRUCT_INCLUDE(FPathTracingData, PathTracingData)
+
+		// scene lights
 		SHADER_PARAMETER_RDG_BUFFER_SRV(StructuredBuffer<FPathTracingLight>, SceneLights)
 		SHADER_PARAMETER(uint32, SceneLightCount)
 		SHADER_PARAMETER(uint32, SceneVisibleLightCount)
+		SHADER_PARAMETER_STRUCT_INCLUDE(FPathTracingLightGrid, LightGridParameters)
+
 		// Skylight
 		SHADER_PARAMETER_RDG_TEXTURE(Texture2D, SkylightTexture)
 		SHADER_PARAMETER_RDG_TEXTURE(Texture2D, SkylightPdf)
@@ -575,45 +672,200 @@ bool PrepareSkyTexture(FRDGBuilder& GraphBuilder, FScene* Scene, const FViewInfo
 	return true;
 }
 
+void PrepareLightGrid(FRDGBuilder& GraphBuilder, FPathTracingLightGrid* LightGridParameters, const FPathTracingLight* Lights, uint32 NumLights, uint32 NumInfiniteLights, FRDGBufferSRV* LightsSRV, const FViewInfo& View)
+{
+	const float Inf = std::numeric_limits<float>::infinity();
+	LightGridParameters->SceneInfiniteLightCount = NumInfiniteLights;
+	LightGridParameters->SceneLightsBoundMin = FVector(+Inf, +Inf, +Inf);
+	LightGridParameters->SceneLightsBoundMax = FVector(-Inf, -Inf, -Inf);
+	LightGridParameters->LightGrid = nullptr;
+	LightGridParameters->LightGridData = nullptr;
+
+	int NumFiniteLights = NumLights - NumInfiniteLights;
+	// if we have some finite lights -- build a light grid
+	if (NumFiniteLights > 0)
+	{
+		// get bounding box of all finite lights
+		const FPathTracingLight* FiniteLights = Lights + NumInfiniteLights;
+		for (int Index = 0; Index < NumFiniteLights; Index++)
+		{
+			const FPathTracingLight& Light = FiniteLights[Index];
+			LightGridParameters->SceneLightsBoundMin = FVector::Min(LightGridParameters->SceneLightsBoundMin, Light.BoundMin);
+			LightGridParameters->SceneLightsBoundMax = FVector::Max(LightGridParameters->SceneLightsBoundMax, Light.BoundMax);
+		}
+
+		const uint32 Resolution = FMath::RoundUpToPowerOfTwo(CVarPathTracingLightGridResolution.GetValueOnRenderThread());
+		const uint32 MaxCount = FMath::Clamp(
+			CVarPathTracingLightGridMaxCount.GetValueOnRenderThread(),
+			1,
+			FMath::Min(NumFiniteLights, RAY_TRACING_LIGHT_COUNT_MAXIMUM)
+		);
+		LightGridParameters->LightGridResolution = Resolution;
+		LightGridParameters->LightGridMaxCount = MaxCount;
+
+		// pick the shortest axis
+		FVector Diag = LightGridParameters->SceneLightsBoundMax - LightGridParameters->SceneLightsBoundMin;
+		if (Diag.X < Diag.Y && Diag.X < Diag.Z)
+		{
+			LightGridParameters->LightGridAxis = 0;
+		}
+		else if (Diag.Y < Diag.Z)
+		{
+			LightGridParameters->LightGridAxis = 1;
+		}
+		else
+		{
+			LightGridParameters->LightGridAxis = 2;
+		}
+
+		FPathTracingBuildLightGridCS::FParameters* LightGridPassParameters = GraphBuilder.AllocParameters< FPathTracingBuildLightGridCS::FParameters>();
+
+		FRDGTextureDesc LightGridDesc = FRDGTextureDesc::Create2D(
+			FIntPoint(Resolution, Resolution),
+			PF_R32_UINT,
+			FClearValueBinding::None,
+			TexCreate_ShaderResource | TexCreate_UAV);
+		FRDGTexture* LightGridTexture = GraphBuilder.CreateTexture(LightGridDesc, TEXT("PathTracer.LightGrid"), ERDGTextureFlags::None);
+		LightGridPassParameters->RWLightGrid = GraphBuilder.CreateUAV(LightGridTexture);
+
+		EPixelFormat LightGridDataFormat = PF_R32_UINT;
+		size_t LightGridDataNumBytes = sizeof(uint32);
+		if (NumLights <= (MAX_uint8 + 1))
+		{
+			LightGridDataFormat = PF_R8_UINT;
+			LightGridDataNumBytes = sizeof(uint8);
+		}
+		else if (NumLights <= (MAX_uint16 + 1))
+		{
+			LightGridDataFormat = PF_R16_UINT;
+			LightGridDataNumBytes = sizeof(uint16);
+		}
+		FRDGBufferDesc LightGridDataDesc = FRDGBufferDesc::CreateBufferDesc(LightGridDataNumBytes, MaxCount * Resolution * Resolution);
+		FRDGBuffer* LightGridData = GraphBuilder.CreateBuffer(LightGridDataDesc, TEXT("PathTracer.LightGridData"));
+		LightGridPassParameters->RWLightGridData = GraphBuilder.CreateUAV(LightGridData, LightGridDataFormat);
+		LightGridPassParameters->LightGridParameters = *LightGridParameters;
+		LightGridPassParameters->SceneLights = LightsSRV;
+		LightGridPassParameters->SceneLightCount = NumLights;
+
+		TShaderMapRef<FPathTracingBuildLightGridCS> ComputeShader(View.ShaderMap);
+		FComputeShaderUtils::AddPass(
+			GraphBuilder,
+			RDG_EVENT_NAME("Light Grid Create (%u lights)", NumFiniteLights),
+			ComputeShader,
+			LightGridPassParameters,
+			FComputeShaderUtils::GetGroupCount(FIntPoint(Resolution, Resolution), FComputeShaderUtils::kGolden2DGroupSize));
+
+		// hookup to the actual rendering pass
+		LightGridParameters->LightGrid = LightGridTexture;
+		LightGridParameters->LightGridData = GraphBuilder.CreateSRV(LightGridData, LightGridDataFormat);
+
+
+	}
+	else
+	{
+		// light grid is not needed - just hookup dummy data
+		LightGridParameters->LightGridResolution = 0;
+		LightGridParameters->LightGridMaxCount = 0;
+		LightGridParameters->LightGridAxis = 0;
+		LightGridParameters->LightGrid = GraphBuilder.RegisterExternalTexture(GSystemTextures.BlackDummy);
+		FRDGBufferDesc LightGridDataDesc = FRDGBufferDesc::CreateBufferDesc(sizeof(uint32), 1);
+		FRDGBuffer* LightGridData = GraphBuilder.CreateBuffer(LightGridDataDesc, TEXT("PathTracer.LightGridData"));
+		uint32 Dummy = 0;
+		AddBufferUploadPass(GraphBuilder, LightGridData, &Dummy, sizeof(uint32));
+		LightGridParameters->LightGridData = GraphBuilder.CreateSRV(LightGridData, PF_R32_UINT);
+	}
+}
+
 void SetLightParameters(FRDGBuilder& GraphBuilder, FPathTracingRG::FParameters* PassParameters, FScene* Scene, const FViewInfo& View, bool UseMISCompensation)
 {
 	PassParameters->SceneVisibleLightCount = 0;
 
 	// Lights
-	FPathTracingLight Lights[RAY_TRACING_LIGHT_COUNT_MAXIMUM]; // Keep this on the stack for now -- eventually will need to make this dynamic to lift size limit (and also avoid uploading per frame ...)
-	unsigned LightCount = 0;
+	TArray<FPathTracingLight> Lights;
+	Lights.Reserve(1 + Scene->Lights.Num());
 
 	// Prepend SkyLight to light buffer since it is not part of the regular light list
+	const float Inf = std::numeric_limits<float>::infinity();
 	if (PrepareSkyTexture(GraphBuilder, Scene, View, UseMISCompensation, PassParameters))
 	{
 		check(Scene->SkyLight != nullptr);
-		FPathTracingLight& DestLight = Lights[LightCount];
+		Lights.AddDefaulted();
+		FPathTracingLight& DestLight = Lights.Last();
 		DestLight.Color = FVector(1, 1, 1); // not used (it is folded into the importance table directly)
 		DestLight.Flags = Scene->SkyLight->bTransmission ? PATHTRACER_FLAG_TRANSMISSION_MASK : 0;
 		DestLight.Flags |= PATHTRACER_FLAG_LIGHTING_CHANNEL_MASK;
 		DestLight.Flags |= PATHTRACING_LIGHT_SKY;
 		DestLight.Flags |= Scene->SkyLight->bCastShadows ? PATHTRACER_FLAG_CAST_SHADOW_MASK : 0;
 		DestLight.IESTextureSlice = -1;
+		DestLight.BoundMin = FVector(-Inf, -Inf, -Inf);
+		DestLight.BoundMax = FVector( Inf,  Inf,  Inf);
 		if (Scene->SkyLight->bRealTimeCaptureEnabled)
 		{
 			// When using the realtime capture system, always make the skylight visible
 			// because this is our only way of "seeing" the atmo/clouds at the moment
 			PassParameters->SceneVisibleLightCount = 1;
 		}
+	}
+
+	// Add directional lights next (all lights with infinite bounds should come first)
+	if (View.Family->EngineShowFlags.DirectionalLights)
+	{
+		for (auto Light : Scene->Lights)
+		{
+			ELightComponentType LightComponentType = (ELightComponentType)Light.LightSceneInfo->Proxy->GetLightType();
+
+			if (LightComponentType != LightType_Directional)
+			{
+				continue;
+			}
 
-		LightCount++;
+			FLightShaderParameters LightParameters;
+			Light.LightSceneInfo->Proxy->GetLightShaderParameters(LightParameters);
+
+			if (LightParameters.Color.IsZero())
+			{
+				continue;
+			}
+
+			Lights.AddDefaulted();
+			FPathTracingLight& DestLight = Lights.Last();
+			uint32 Transmission = Light.LightSceneInfo->Proxy->Transmission();
+			uint8 LightingChannelMask = Light.LightSceneInfo->Proxy->GetLightingChannelMask();
+
+			DestLight.Flags = Transmission ? PATHTRACER_FLAG_TRANSMISSION_MASK : 0;
+			DestLight.Flags |= LightingChannelMask & PATHTRACER_FLAG_LIGHTING_CHANNEL_MASK;
+			DestLight.Flags |= Light.LightSceneInfo->Proxy->CastsDynamicShadow() ? PATHTRACER_FLAG_CAST_SHADOW_MASK : 0;
+			DestLight.IESTextureSlice = -1;
+			DestLight.RectLightTextureIndex = -1;
+
+			// these mean roughly the same thing across all light types
+			DestLight.Color = LightParameters.Color;
+			DestLight.Position = LightParameters.Position;
+			DestLight.Normal = -LightParameters.Direction;
+			DestLight.dPdu = FVector::CrossProduct(LightParameters.Tangent, LightParameters.Direction);
+			DestLight.dPdv = LightParameters.Tangent;
+			DestLight.Attenuation = LightParameters.InvRadius;
+			DestLight.FalloffExponent = 0;
+
+			DestLight.Normal = LightParameters.Direction;
+			DestLight.Dimensions = FVector(LightParameters.SourceRadius, LightParameters.SoftSourceRadius, 0.0f);
+			DestLight.Flags |= PATHTRACING_LIGHT_DIRECTIONAL;
+
+			DestLight.BoundMin = FVector(-Inf, -Inf, -Inf);
+			DestLight.BoundMax = FVector( Inf,  Inf,  Inf);
+		}
 	}
 
+	uint32 NumInfiniteLights = Lights.Num();
+
 	int32 NextRectTextureIndex = 0;
 
 	TMap<FTexture*, int> IESLightProfilesMap;
 	for (auto Light : Scene->Lights)
 	{
-		if (LightCount >= RAY_TRACING_LIGHT_COUNT_MAXIMUM) break;
-
 		ELightComponentType LightComponentType = (ELightComponentType)Light.LightSceneInfo->Proxy->GetLightType();
 
-		if (((LightComponentType == LightType_Directional) && !View.Family->EngineShowFlags.DirectionalLights) ||
+		if ( (LightComponentType == LightType_Directional) /* already handled by the loop above */  ||
 			((LightComponentType == LightType_Rect       ) && !View.Family->EngineShowFlags.RectLights       ) ||
 			((LightComponentType == LightType_Spot       ) && !View.Family->EngineShowFlags.SpotLights       ) ||
 			((LightComponentType == LightType_Point      ) && !View.Family->EngineShowFlags.PointLights      ))
@@ -622,10 +874,17 @@ void SetLightParameters(FRDGBuilder& GraphBuilder, FPathTracingRG::FParameters*
 			continue;
 		}
 
-		FPathTracingLight& DestLight = Lights[LightCount];
-
 		FLightShaderParameters LightParameters;
 		Light.LightSceneInfo->Proxy->GetLightShaderParameters(LightParameters);
+
+		if (LightParameters.Color.IsZero())
+		{
+			continue;
+		}
+
+		Lights.AddDefaulted();
+		FPathTracingLight& DestLight = Lights.Last();
+
 		uint32 Transmission = Light.LightSceneInfo->Proxy->Transmission();
 		uint8 LightingChannelMask = Light.LightSceneInfo->Proxy->GetLightingChannelMask();
 
@@ -656,13 +915,6 @@ void SetLightParameters(FRDGBuilder& GraphBuilder, FPathTracingRG::FParameters*
 
 		switch (LightComponentType)
 		{
-			case LightType_Directional:
-			{
-				DestLight.Normal = LightParameters.Direction;
-				DestLight.Dimensions = FVector(LightParameters.SourceRadius, LightParameters.SoftSourceRadius, 0.0f);
-				DestLight.Flags |= PATHTRACING_LIGHT_DIRECTIONAL;
-				break;
-			}
 			case LightType_Rect:
 			{
 				DestLight.Dimensions = FVector(2.0f * LightParameters.SourceRadius, 2.0f * LightParameters.SourceLength, 0.0f);
@@ -698,6 +950,20 @@ void SetLightParameters(FRDGBuilder& GraphBuilder, FPathTracingRG::FParameters*
 						}
 					}
 				}
+
+				float Radius = 1.0f / LightParameters.InvRadius;
+				FVector Center = DestLight.Position;
+				FVector Normal = DestLight.Normal;
+				FVector Disc = FVector(
+					FMath::Sqrt(FMath::Clamp(1 - Normal.X * Normal.X, 0.0f, 1.0f)),
+					FMath::Sqrt(FMath::Clamp(1 - Normal.Y * Normal.Y, 0.0f, 1.0f)),
+					FMath::Sqrt(FMath::Clamp(1 - Normal.Z * Normal.Z, 0.0f, 1.0f))
+				);
+				// quad bbox is the bbox of the disc +  the tip of the hemisphere
+				// TODO: is it worth trying to account for barndoors? seems unlikely to cut much empty space since the volume _inside_ the barndoor receives light
+				FVector Tip = Center + Normal * Radius;
+				DestLight.BoundMin = FVector::Min(Tip, Center - Radius * Disc);
+				DestLight.BoundMax = FVector::Max(Tip, Center + Radius * Disc);
 				break;
 			}
 			case LightType_Spot:
@@ -707,6 +973,26 @@ void SetLightParameters(FRDGBuilder& GraphBuilder, FPathTracingRG::FParameters*
 				DestLight.FalloffExponent = LightParameters.FalloffExponent;
 				DestLight.Flags |= Light.LightSceneInfo->Proxy->IsInverseSquared() ? 0 : PATHTRACER_FLAG_NON_INVERSE_SQUARE_FALLOFF_MASK;
 				DestLight.Flags |= PATHTRACING_LIGHT_SPOT;
+
+				float Radius = 1.0f / LightParameters.InvRadius;
+				FVector Center = DestLight.Position;
+				FVector Normal = DestLight.Normal;
+				FVector Disc = FVector(
+					FMath::Sqrt(FMath::Clamp(1 - Normal.X * Normal.X, 0.0f, 1.0f)),
+					FMath::Sqrt(FMath::Clamp(1 - Normal.Y * Normal.Y, 0.0f, 1.0f)),
+					FMath::Sqrt(FMath::Clamp(1 - Normal.Z * Normal.Z, 0.0f, 1.0f))
+				);
+				// box around ray from light center to tip of the cone
+				FVector Tip = Center + Normal * Radius;
+				DestLight.BoundMin = FVector::Min(Center, Tip);
+				DestLight.BoundMax = FVector::Max(Center, Tip);
+				// expand by disc around the farthest part of the cone
+
+				float CosOuter = LightParameters.SpotAngles.X;
+				float SinOuter = FMath::Sqrt(1.0f - CosOuter * CosOuter);
+
+				DestLight.BoundMin = FVector::Min(DestLight.BoundMin, Center + Radius * (Normal * CosOuter - Disc * SinOuter));
+				DestLight.BoundMax = FVector::Max(DestLight.BoundMax, Center + Radius * (Normal * CosOuter + Disc * SinOuter));
 				break;
 			}
 			case LightType_Point:
@@ -715,6 +1001,11 @@ void SetLightParameters(FRDGBuilder& GraphBuilder, FPathTracingRG::FParameters*
 				DestLight.FalloffExponent = LightParameters.FalloffExponent;
 				DestLight.Flags |= Light.LightSceneInfo->Proxy->IsInverseSquared() ? 0 : PATHTRACER_FLAG_NON_INVERSE_SQUARE_FALLOFF_MASK;
 				DestLight.Flags |= PATHTRACING_LIGHT_POINT;
+				float Radius = 1.0f / LightParameters.InvRadius;
+				FVector Center = DestLight.Position;
+				// simple sphere of influence
+				DestLight.BoundMin = Center - FVector(Radius, Radius, Radius);
+				DestLight.BoundMax = Center + FVector(Radius, Radius, Radius);
 				break;
 			}
 			default:
@@ -724,8 +1015,6 @@ void SetLightParameters(FRDGBuilder& GraphBuilder, FPathTracingRG::FParameters*
 				break;
 			}
 		}
-
-		LightCount++;
 	}
 
 	{
@@ -737,31 +1026,35 @@ void SetLightParameters(FRDGBuilder& GraphBuilder, FPathTracingRG::FParameters*
 		PassParameters->RectLightSampler = TStaticSamplerState<SF_Bilinear, AM_Clamp, AM_Clamp, AM_Clamp>::GetRHI();
 	}
 
+	PassParameters->SceneLightCount = Lights.Num();
 	{
+		// Buffer creation needs at least one element
+		if (Lights.IsEmpty())
+		{
+			Lights.AddDefaulted();
+		}
 		// Upload the buffer of lights to the GPU
-		size_t DataSize = sizeof(FPathTracingLight) * FMath::Max(LightCount, 1u);
-		PassParameters->SceneLights = GraphBuilder.CreateSRV(FRDGBufferSRVDesc(CreateStructuredBuffer(GraphBuilder, TEXT("PathTracer.LightsBuffer"), sizeof(FPathTracingLight), FMath::Max(LightCount, 1u), Lights, DataSize)));
-		PassParameters->SceneLightCount = LightCount;
+		size_t DataSize = sizeof(FPathTracingLight) * Lights.Num();
+		PassParameters->SceneLights = GraphBuilder.CreateSRV(FRDGBufferSRVDesc(CreateStructuredBuffer(GraphBuilder, TEXT("PathTracer.LightsBuffer"), sizeof(FPathTracingLight), Lights.Num(), Lights.GetData(), DataSize)));
 	}
 
 	if (CVarPathTracingVisibleLights.GetValueOnRenderThread() != 0)
 	{
-		PassParameters->SceneVisibleLightCount = LightCount;
+		PassParameters->SceneVisibleLightCount = PassParameters->SceneLightCount;
 	}
 
 	if (!IESLightProfilesMap.IsEmpty())
 	{
 		// We found some IES profiles to use -- upload them into a single atlas so we can access them easily in HLSL
 
-		// FIXME: This is redundant because all the IES textures are already on the GPU, we just don't have the ability to use
-		// an array of texture handles on the HLSL side.
+		// TODO: This is redundant because all the IES textures are already on the GPU. Handling IES profiles via Miss shaders
+		// would be cleaner.
 
-		// FIXME: This is also redundant with the logic in RayTracingLighting.cpp, but the latter is limitted to 1D profiles and 
+		// TODO: This is also redundant with the logic in RayTracingLighting.cpp, but the latter is limitted to 1D profiles and 
 		// does not consider the same set of lights as the path tracer. Longer term we should aim to unify the representation of lights
 		// across both passes
 
-		// FIXME: This process is repeated every frame! would be nicer to cache the data somehow. Perhaps just do this step for
-		// Iteration == 0 since we can assume that any changes in IES profiles will invalidate the path tracer anyway?
+		// TODO: This process is repeated every frame! More motivation to move to a Miss shader based implementation
 
 		// This size matches the import resolution of light profiles (see FIESLoader::GetWidth)
 		const int kIESAtlasSize = 256;
@@ -797,6 +1090,9 @@ void SetLightParameters(FRDGBuilder& GraphBuilder, FPathTracingRG::FParameters*
 	{
 		PassParameters->IESTexture = GraphBuilder.RegisterExternalTexture(GSystemTextures.WhiteDummy);
 	}
+
+
+	PrepareLightGrid(GraphBuilder, &PassParameters->LightGridParameters, Lights.GetData(), Lights.Num(), NumInfiniteLights, PassParameters->SceneLights, View);
 }
 
 class FPathTracingCompositorPS : public FGlobalShader
@@ -903,6 +1199,23 @@ void FDeferredShadingSceneRenderer::RenderPathTracing(
 		Scene->PathTracingSkylightPdf.SafeRelease();
 	}
 
+	const uint32 LightGridResolution = FMath::RoundUpToPowerOfTwo(CVarPathTracingLightGridResolution.GetValueOnRenderThread());
+	const uint32 LightGridMaxCount = FMath::Clamp(CVarPathTracingLightGridMaxCount.GetValueOnRenderThread(), 1, RAY_TRACING_LIGHT_COUNT_MAXIMUM);
+
+	static uint32 PreviousLightGridResolution = LightGridResolution;
+	if (PreviousLightGridResolution != LightGridResolution)
+	{
+		PreviousLightGridResolution = LightGridResolution;
+		bArgsChanged = true;
+	}
+
+	static uint32_t PreviousLightGridMaxCount = LightGridMaxCount;
+	if (PreviousLightGridMaxCount != LightGridMaxCount)
+	{
+		PreviousLightGridMaxCount = LightGridMaxCount;
+		bArgsChanged = true;
+	}
+
 	// Get other basic path tracing settings and see if we need to invalidate the current state
 	FPathTracingData PathTracingData;
 	bArgsChanged |= PrepareShaderArgs(View, PathTracingData);
@@ -953,7 +1266,7 @@ void FDeferredShadingSceneRenderer::RenderPathTracing(
 		FPathTracingRG::FParameters* PassParameters = GraphBuilder.AllocParameters<FPathTracingRG::FParameters>();
 		PassParameters->TLAS = View.GetRayTracingSceneViewChecked();
 		PassParameters->ViewUniformBuffer = View.ViewUniformBuffer;
-		PassParameters->PathTracingData = CreateUniformBufferImmediate(PathTracingData, EUniformBufferUsage::UniformBuffer_SingleFrame);
+		PassParameters->PathTracingData = PathTracingData;
 		// upload sky/lights data
 		SetLightParameters(GraphBuilder, PassParameters, Scene, View, UseMISCompensation);
 		if (PathTracingData.SkipDirectLighting)
diff --git a/Engine/Source/Runtime/Renderer/Private/PathTracingUniformBuffers.h b/Engine/Source/Runtime/Renderer/Private/PathTracingUniformBuffers.h
deleted file mode 100644
diff --git a/Engine/Source/Runtime/Renderer/Private/RayTracing/RayTracingGlobalIllumination.cpp b/Engine/Source/Runtime/Renderer/Private/RayTracing/RayTracingGlobalIllumination.cpp
--- a/Engine/Source/Runtime/Renderer/Private/RayTracing/RayTracingGlobalIllumination.cpp
+++ b/Engine/Source/Runtime/Renderer/Private/RayTracing/RayTracingGlobalIllumination.cpp
@@ -12,7 +12,6 @@
 #include "RHIResources.h"
 #include "UniformBuffer.h"
 #include "RayGenShaderUtils.h"
-#include "PathTracingUniformBuffers.h"
 #include "SceneTextureParameters.h"
 #include "ScreenSpaceDenoise.h"
 #include "ClearQuad.h"
@@ -24,6 +23,7 @@
 #include "RayTracingDefinitions.h"
 #include "RayTracingDeferredMaterials.h"
 #include "RayTracingTypes.h"
+#include "PathTracingDefinitions.h"
 
 static TAutoConsoleVariable<int32> CVarRayTracingGlobalIllumination(
 	TEXT("r.RayTracing.GlobalIllumination"),
diff --git a/Engine/Source/Runtime/Renderer/Private/RayTracing/RayTracingShadows.cpp b/Engine/Source/Runtime/Renderer/Private/RayTracing/RayTracingShadows.cpp
--- a/Engine/Source/Runtime/Renderer/Private/RayTracing/RayTracingShadows.cpp
+++ b/Engine/Source/Runtime/Renderer/Private/RayTracing/RayTracingShadows.cpp
@@ -6,7 +6,6 @@
 #if RHI_RAYTRACING
 
 #include "ClearQuad.h"
-#include "PathTracingUniformBuffers.h"
 #include "SceneRendering.h"
 #include "RenderGraphBuilder.h"
 #include "RenderTargetPool.h"
