Remove QuasiRandom and Halton code in favor of the newer random sequence

The last remaining usage was inside denoising code. This code assumes the sampling pattern is similar to what was used during rendering which was no longer true since we switched to RandomSequence there a while ago.

diff --git a/Engine/Shaders/Private/Halton.ush b/Engine/Shaders/Private/Halton.ush
deleted file mode 100644
diff --git a/Engine/Shaders/Private/QuasiRandom.ush b/Engine/Shaders/Private/QuasiRandom.ush
deleted file mode 100644
diff --git a/Engine/Shaders/Private/ScreenSpaceDenoise/SSDSignalFramework.ush b/Engine/Shaders/Private/ScreenSpaceDenoise/SSDSignalFramework.ush
--- a/Engine/Shaders/Private/ScreenSpaceDenoise/SSDSignalFramework.ush
+++ b/Engine/Shaders/Private/ScreenSpaceDenoise/SSDSignalFramework.ush
@@ -27,11 +27,6 @@
 #include "../EyeAdaptationCommon.ush"
 #include "../LightShaderParameters.ush"
 
-#if COMPILE_SIGNAL_COLOR_SH
-	#include "../QuasiRandom.ush"
-#endif
-
-
 #if !defined(CONFIG_SIGNAL_PROCESSING)
 	#error Missing CONFIG_SIGNAL_PROCESSING
 
@@ -78,6 +73,8 @@ float4 LightDirectionAndLength[MAX_SIGNAL_BATCH_SIZE];
 float HitDistanceToWorldBluringRadius[MAX_SIGNAL_BATCH_SIZE];
 uint LightType[MAX_SIGNAL_BATCH_SIZE];
 
+uint FrameIndex;
+
 #if SHADER_PARAMETERS_HAS_SOURCETEXTURE
 Texture2D UnusedLightSourceTexture;
 #endif
@@ -681,13 +678,12 @@ float GetSamplePenumbra(FSSDSignalSample Sample)
 FSSDSphericalHarmonicRGB ComputeSampleColorSH(FSSDSampleSceneInfos SampleSceneMetadata, FSSDSignalSample Sample, uint2 NeighborPixelCoord)
 {
 	const uint SamplesPerPixel = 1;
-	FQuasiRandom QuasiRandom;
-	QuasiRandom_Initialize(QuasiRandom, NeighborPixelCoord, 0, SamplesPerPixel);
-	QuasiRandom_Update(QuasiRandom, /* SampleIndex = */ 0, SamplesPerPixel);
+	RandomSequence RandSequence;
+	RandomSequence_Initialize(RandSequence, NeighborPixelCoord, /* SampleIndex */ 0,  FrameIndex, SamplesPerPixel);
 
-	float3 RandSample = QuasiRandom_Sample3(QuasiRandom);
+	float2 RandSample = RandomSequence_GenerateSample2D(RandSequence);
 
-	float4 TangentDirection = CosineSampleHemisphere(RandSample.yz);
+	float4 TangentDirection = CosineSampleHemisphere(RandSample);
 	float3 WorldDirection = TangentToWorld(TangentDirection.xyz, GetWorldNormal(SampleSceneMetadata));
 
 	#if SPHERICAL_HARMONIC_ORDER == 2
diff --git a/Engine/Source/Runtime/Renderer/Private/Halton.cpp b/Engine/Source/Runtime/Renderer/Private/Halton.cpp
deleted file mode 100644
diff --git a/Engine/Source/Runtime/Renderer/Private/Halton.h b/Engine/Source/Runtime/Renderer/Private/Halton.h
deleted file mode 100644
diff --git a/Engine/Source/Runtime/Renderer/Private/RendererScene.cpp b/Engine/Source/Runtime/Renderer/Private/RendererScene.cpp
--- a/Engine/Source/Runtime/Renderer/Private/RendererScene.cpp
+++ b/Engine/Source/Runtime/Renderer/Private/RendererScene.cpp
@@ -1086,8 +1086,6 @@ FScene::FScene(UWorld* InWorld, bool bInRequiresHitProxies, bool bInIsEditorScen
 	{
 		RayTracingDynamicGeometryCollection = new FRayTracingDynamicGeometryCollection();
 	}
-
-	BeginInitResource(&HaltonPrimesResource);
 #endif
 
 	World->UpdateParameterCollectionInstances(false, false);
@@ -1155,9 +1153,6 @@ FScene::~FScene()
 
 	delete RayTracingDynamicGeometryCollection;
 	RayTracingDynamicGeometryCollection = nullptr;
-
-	BeginReleaseResource(&HaltonPrimesResource);
-
 #endif // RHI_RAYTRACING
 
 	checkf(RemovedPrimitiveSceneInfos.Num() == 0, TEXT("Leaking %d FPrimitiveSceneInfo instances."), RemovedPrimitiveSceneInfos.Num()); // Ensure UpdateAllPrimitiveSceneInfos() is called before destruction.
diff --git a/Engine/Source/Runtime/Renderer/Private/ScenePrivate.h b/Engine/Source/Runtime/Renderer/Private/ScenePrivate.h
--- a/Engine/Source/Runtime/Renderer/Private/ScenePrivate.h
+++ b/Engine/Source/Runtime/Renderer/Private/ScenePrivate.h
@@ -52,7 +52,6 @@
 #if RHI_RAYTRACING
 #include "RayTracing/RayTracingIESLightProfiles.h"
 #include "RayTracing/RayTracingScene.h"
-#include "Halton.h"
 #endif
 #include "Nanite/NaniteRender.h"
 #include "Lumen/LumenViewState.h"
@@ -2878,8 +2877,6 @@ public:
 
 #if RHI_RAYTRACING
 	class FRayTracingDynamicGeometryCollection* RayTracingDynamicGeometryCollection;
-	FHaltonSequence HaltonSequence;
-	FHaltonPrimesResource HaltonPrimesResource;
 #endif
 
 	/** Initialization constructor. */
diff --git a/Engine/Source/Runtime/Renderer/Private/ScreenSpaceDenoise.cpp b/Engine/Source/Runtime/Renderer/Private/ScreenSpaceDenoise.cpp
--- a/Engine/Source/Runtime/Renderer/Private/ScreenSpaceDenoise.cpp
+++ b/Engine/Source/Runtime/Renderer/Private/ScreenSpaceDenoise.cpp
@@ -13,8 +13,6 @@
 #include "ClearQuad.h"
 #include "PipelineStateCache.h"
 #include "SceneTextureParameters.h"
-#include "BlueNoise.h"
-#include "Halton.h"
 #include "Lumen/LumenSceneRendering.h"
 
 
@@ -876,10 +874,8 @@ BEGIN_SHADER_PARAMETER_STRUCT(FSSDCommonParameters, )
 	SHADER_PARAMETER_RDG_TEXTURE(Texture2D, EyeAdaptationTexture)
 	SHADER_PARAMETER_RDG_TEXTURE(Texture2D<uint>, TileClassificationTexture)
 	SHADER_PARAMETER_STRUCT_REF(FViewUniformShaderParameters, ViewUniformBuffer)
-	
-	SHADER_PARAMETER_STRUCT_REF(FHaltonIteration, HaltonIteration)
-	SHADER_PARAMETER_STRUCT_REF(FHaltonPrimes, HaltonPrimes)
-	SHADER_PARAMETER_STRUCT_REF(FBlueNoise, BlueNoise)
+
+	SHADER_PARAMETER(uint32, FrameIndex)
 END_SHADER_PARAMETER_STRUCT()
 
 BEGIN_SHADER_PARAMETER_STRUCT(FSSDSignalSRVs, )
@@ -1653,30 +1649,7 @@ static void DenoiseSignalAtConstantPixelDensity(
 		CommonParameters.BufferUVBilinearCorrection.Y = (0.5f * PixelPositionToFullResPixel - FullResPixelOffset.Y) / float(FullResBufferExtent.Y);
 	}
 
-	#if RHI_RAYTRACING
-	if (Settings.SignalProcessing == ESignalProcessing::DiffuseAndAmbientOcclusion)
-	{
-		uint32 IterationCount = Settings.MaxInputSPP;
-		uint32 SequenceCount = 1;
-		uint32 DimensionCount = 24;
-
-		FScene* Scene = static_cast<FScene*>(View.Family->Scene);
-
-		FHaltonSequenceIteration HaltonSequenceIteration(Scene->HaltonSequence, IterationCount, SequenceCount, DimensionCount, View.ViewState ? (View.ViewState->FrameIndex % 1024) : 0);
-		FHaltonIteration HaltonIteration;
-		InitializeHaltonSequenceIteration(HaltonSequenceIteration, HaltonIteration);
-
-		FHaltonPrimes HaltonPrimes;
-		InitializeHaltonPrimes(Scene->HaltonPrimesResource, HaltonPrimes);
-
-		FBlueNoise BlueNoise;
-		InitializeBlueNoise(BlueNoise);
-		
-		CommonParameters.HaltonIteration = CreateUniformBufferImmediate(HaltonIteration, EUniformBufferUsage::UniformBuffer_SingleFrame);
-		CommonParameters.HaltonPrimes = CreateUniformBufferImmediate(HaltonPrimes, EUniformBufferUsage::UniformBuffer_SingleFrame);
-		CommonParameters.BlueNoise = CreateUniformBufferImmediate(BlueNoise, EUniformBufferUsage::UniformBuffer_SingleFrame);
-	}
-	#endif // RHI_RAYTRACING
+	CommonParameters.FrameIndex = View.ViewState ? View.ViewState->FrameIndex : 0;
 
 	// Setup all the metadata to do spatial convolution.
 	FSSDConvolutionMetaData ConvolutionMetaData;
