Moved dynamic ray tracing geometry shared vertex buffer validation to GatherRayTracingWorldInstancesForView(). Made it fault-tolerant by skipping the geometry instance after reporting a validation error instead of crashing.

diff --git a/Engine/Source/Runtime/Engine/Private/SkeletalMesh.cpp b/Engine/Source/Runtime/Engine/Private/SkeletalMesh.cpp
--- a/Engine/Source/Runtime/Engine/Private/SkeletalMesh.cpp
+++ b/Engine/Source/Runtime/Engine/Private/SkeletalMesh.cpp
@@ -6206,9 +6206,8 @@ void FSkeletalMeshSceneProxy::GetDynamicRayTracingInstances(FRayTracingMaterialG
 				const FLODSectionElements& LODSection = LODSections[LODIndex];
 				check(LODSection.SectionElements.Num() == LODData.RenderSections.Num());
 				
-				//#dxr_todo: verify why this condition is not fulfilled sometimes
-				verify(LODSection.SectionElements.Num() == MeshObject->GetRayTracingGeometry()->Initializer.Segments.Num());
-				if(LODSection.SectionElements.Num() != MeshObject->GetRayTracingGeometry()->Initializer.Segments.Num())
+				//#dxr_todo (UE-113617): verify why this condition is not fulfilled sometimes
+				if(!ensure(LODSection.SectionElements.Num() == MeshObject->GetRayTracingGeometry()->Initializer.Segments.Num()))
 				{
 					return;
 				}
diff --git a/Engine/Source/Runtime/Renderer/Private/DeferredShadingRenderer.cpp b/Engine/Source/Runtime/Renderer/Private/DeferredShadingRenderer.cpp
--- a/Engine/Source/Runtime/Renderer/Private/DeferredShadingRenderer.cpp
+++ b/Engine/Source/Runtime/Renderer/Private/DeferredShadingRenderer.cpp
@@ -791,7 +791,7 @@ bool FDeferredShadingSceneRenderer::GatherRayTracingWorldInstancesForView(FRHICo
 
 		const bool bParallelMeshBatchSetup = GRayTracingParallelMeshBatchSetup && FApp::ShouldUseThreadingForPerformance();
 
-		Scene->GetRayTracingDynamicGeometryCollection()->BeginUpdate();
+		const int64 SharedBufferGenerationID = Scene->GetRayTracingDynamicGeometryCollection()->BeginUpdate();
 
 		struct FRayTracingMeshBatchWorkItem
 		{
@@ -891,6 +891,8 @@ bool FDeferredShadingSceneRenderer::GatherRayTracingWorldInstancesForView(FRHICo
 
 			FPrimitiveSceneProxy* SceneProxy = Scene->PrimitiveSceneProxies[PrimitiveIndex];
 			TempRayTracingInstances.Reset();
+			MaterialGatheringContext.DynamicRayTracingGeometriesToUpdate.Reset();
+
 			SceneProxy->GetDynamicRayTracingInstances(MaterialGatheringContext, TempRayTracingInstances);
 
 			for (auto DynamicRayTracingGeometryUpdate : MaterialGatheringContext.DynamicRayTracingGeometriesToUpdate)
@@ -904,26 +906,37 @@ bool FDeferredShadingSceneRenderer::GatherRayTracingWorldInstancesForView(FRHICo
 				);
 			}
 
-			MaterialGatheringContext.DynamicRayTracingGeometriesToUpdate.Reset();
-
 			if (TempRayTracingInstances.Num() > 0)
 			{
 				for (FRayTracingInstance& Instance : TempRayTracingInstances)
 				{
+					const FRayTracingGeometry* Geometry = Instance.Geometry;
+
+					if (!ensureMsgf(Geometry->DynamicGeometrySharedBufferGenerationID == FRayTracingGeometry::NonSharedVertexBuffers
+						|| Geometry->DynamicGeometrySharedBufferGenerationID == SharedBufferGenerationID,
+						TEXT("GenerationID %lld, but expected to be %lld or %lld. Geometry debug name: '%s'. ")
+						TEXT("When shared vertex buffers are used, the contents is expected to be written every frame. ")
+						TEXT("Possibly AddDynamicMeshBatchForGeometryUpdate() was not called for this geometry."),
+						Geometry->DynamicGeometrySharedBufferGenerationID, SharedBufferGenerationID, FRayTracingGeometry::NonSharedVertexBuffers,
+						*Geometry->Initializer.DebugName.ToString()))
+					{
+						continue;
+					}
+
 					// If geometry still has pending build request then add to list which requires a force build
-					if (Instance.Geometry->HasPendingBuildRequest())
+					if (Geometry->HasPendingBuildRequest())
 					{
-						RayTracingScene.GeometriesToBuild.Add(Instance.Geometry);
+						RayTracingScene.GeometriesToBuild.Add(Geometry);
 					}
 
 					// Thin geometries like hair don't have material, as they only support shadow at the moment.
-					if (!ensureMsgf(Instance.GetMaterials().Num() == Instance.Geometry->Initializer.Segments.Num() ||
-						(Instance.Geometry->Initializer.Segments.Num() == 0 && Instance.GetMaterials().Num() == 1) ||
+					if (!ensureMsgf(Instance.GetMaterials().Num() == Geometry->Initializer.Segments.Num() ||
+						(Geometry->Initializer.Segments.Num() == 0 && Instance.GetMaterials().Num() == 1) ||
 						(Instance.GetMaterials().Num() == 0 && (Instance.Mask & RAY_TRACING_MASK_THIN_SHADOW) > 0),
 						TEXT("Ray tracing material assignment validation failed for geometry '%s'. "
-							"Instance.GetMaterials().Num() = %d, Instance.Geometry->Initializer.Segments.Num() = %d, Instance.Mask = 0x%X."),
-						*Instance.Geometry->Initializer.DebugName.ToString(), Instance.GetMaterials().Num(),
-						Instance.Geometry->Initializer.Segments.Num(), Instance.Mask))
+							"Instance.GetMaterials().Num() = %d, Geometry->Initializer.Segments.Num() = %d, Instance.Mask = 0x%X."),
+						*Geometry->Initializer.DebugName.ToString(), Instance.GetMaterials().Num(),
+						Geometry->Initializer.Segments.Num(), Instance.Mask))
 					{
 						continue;
 					}
@@ -931,7 +944,7 @@ bool FDeferredShadingSceneRenderer::GatherRayTracingWorldInstancesForView(FRHICo
 					const uint32 InstanceIndex = RayTracingScene.Instances.Num();
 
 					FRayTracingGeometryInstance& RayTracingInstance = RayTracingScene.Instances.AddDefaulted_GetRef();
-					RayTracingInstance.GeometryRHI = Instance.Geometry->RayTracingGeometryRHI;
+					RayTracingInstance.GeometryRHI = Geometry->RayTracingGeometryRHI;
 					RayTracingInstance.UserData.Add((uint32)PrimitiveIndex);
 					RayTracingInstance.Mask = Instance.Mask;
 					RayTracingInstance.bForceOpaque = Instance.bForceOpaque;
@@ -962,13 +975,6 @@ bool FDeferredShadingSceneRenderer::GatherRayTracingWorldInstancesForView(FRHICo
 						}
 					}
 
-				#if DO_CHECK
-					if (Instance.Geometry->DynamicGeometrySharedBufferGenerationID != FRayTracingGeometry::NonSharedVertexBuffers)
-					{
-						RayTracingScene.GeometriesToValidate.Add(Instance.Geometry);
-					}
-				#endif // DO_CHECK
-
 					if (bParallelMeshBatchSetup)
 					{
 						if (NumPendingMeshBatches >= RayTracingParallelMeshBatchSize)
@@ -1203,18 +1209,6 @@ bool FDeferredShadingSceneRenderer::DispatchRayTracingWorldUpdates(FRDGBuilder&
 
 	RDG_GPU_MASK_SCOPE(GraphBuilder, FRHIGPUMask::All());
 
-#ifdef DO_CHECK
-	// Validate all the ray tracing geometries lifetimes
-	int64 SharedBufferGenerationID = (int64)Scene->GetRayTracingDynamicGeometryCollection()->GetSharedBufferGenerationID();
-	for (const FRayTracingGeometry* Geometry : RayTracingScene.GeometriesToValidate)
-	{
-		checkf(Geometry->DynamicGeometrySharedBufferGenerationID == FRayTracingGeometry::NonSharedVertexBuffers 
-			|| Geometry->DynamicGeometrySharedBufferGenerationID == SharedBufferGenerationID,
-			TEXT("GenerationID %lld, but expected to be %lld or %lld"), 
-			Geometry->DynamicGeometrySharedBufferGenerationID, SharedBufferGenerationID, FRayTracingGeometry::NonSharedVertexBuffers);
-	}
-#endif // DO_CHECK
-
 	// #dxr_todo: UE-72565: refactor ray tracing effects to not be member functions of DeferredShadingRenderer. register each effect at startup and just loop over them automatically to gather all required shaders
 	TArray<FRHIRayTracingShader*> RayGenShaders;
 
diff --git a/Engine/Source/Runtime/Renderer/Private/RayTracing/RayTracingDynamicGeometry.cpp b/Engine/Source/Runtime/Renderer/Private/RayTracing/RayTracingDynamicGeometry.cpp
--- a/Engine/Source/Runtime/Renderer/Private/RayTracing/RayTracingDynamicGeometry.cpp
+++ b/Engine/Source/Runtime/Renderer/Private/RayTracing/RayTracingDynamicGeometry.cpp
@@ -101,7 +101,7 @@ FRayTracingDynamicGeometryCollection::~FRayTracingDynamicGeometryCollection()
 	VertexPositionBuffers.Empty();
 }
 
-void FRayTracingDynamicGeometryCollection::BeginUpdate()
+int64 FRayTracingDynamicGeometryCollection::BeginUpdate()
 {
 	// Clear working arrays - keep max size allocated
 	DispatchCommands.Empty(DispatchCommands.Max());
@@ -116,6 +116,8 @@ void FRayTracingDynamicGeometryCollection::BeginUpdate()
 
 	// Increment generation ID used for validation
 	SharedBufferGenerationID++;
+
+	return SharedBufferGenerationID;
 }
 
 void FRayTracingDynamicGeometryCollection::AddDynamicMeshBatchForGeometryUpdate(
diff --git a/Engine/Source/Runtime/Renderer/Private/RayTracing/RayTracingScene.cpp b/Engine/Source/Runtime/Renderer/Private/RayTracing/RayTracingScene.cpp
--- a/Engine/Source/Runtime/Renderer/Private/RayTracing/RayTracingScene.cpp
+++ b/Engine/Source/Runtime/Renderer/Private/RayTracing/RayTracingScene.cpp
@@ -132,9 +132,6 @@ void FRayTracingScene::Reset()
 
 	Instances.Reset();
 	GeometriesToBuild.Reset();
-#ifdef DO_CHECK
-	GeometriesToValidate.Reset();
-#endif
 
 	BuildScratchBuffer = nullptr;
 }
diff --git a/Engine/Source/Runtime/Renderer/Private/RayTracing/RayTracingScene.h b/Engine/Source/Runtime/Renderer/Private/RayTracing/RayTracingScene.h
--- a/Engine/Source/Runtime/Renderer/Private/RayTracing/RayTracingScene.h
+++ b/Engine/Source/Runtime/Renderer/Private/RayTracing/RayTracingScene.h
@@ -67,11 +67,6 @@ public:
 	// Geometries which still have a pending build request but are used this frame and require a force build.
 	TArray<const FRayTracingGeometry*> GeometriesToBuild;
 
-#ifdef DO_CHECK
-	// Keep track of all used RT Geometries which are used to validate the vertex buffer data (see FRayTracingGeometry::DynamicGeometrySharedBufferGenerationID)
-	TSet<const FRayTracingGeometry*> GeometriesToValidate;
-#endif
-
 	FRayTracingAccelerationStructureSize SizeInfo = {};
 
 	FRDGBufferRef BuildScratchBuffer;
diff --git a/Engine/Source/Runtime/Renderer/Public/RayTracingDynamicGeometryCollection.h b/Engine/Source/Runtime/Renderer/Public/RayTracingDynamicGeometryCollection.h
--- a/Engine/Source/Runtime/Renderer/Public/RayTracingDynamicGeometryCollection.h
+++ b/Engine/Source/Runtime/Renderer/Public/RayTracingDynamicGeometryCollection.h
@@ -17,9 +17,9 @@ public:
 		FRayTracingDynamicGeometryUpdateParams Params,
 		uint32 PrimitiveId 
 	);
-	uint64 GetSharedBufferGenerationID() const { return SharedBufferGenerationID; }
 
-	void BeginUpdate();
+	// Starts an update batch and returns the current shared buffer generation ID which is used for validation.
+	int64 BeginUpdate();
 	void DispatchUpdates(FRHIComputeCommandList& ParentCmdList);
 	void EndUpdate(FRHICommandListImmediate& RHICmdList);
 
@@ -37,7 +37,7 @@ private:
 
 	// Generation ID when the shared vertex buffers have been reset. The current generation ID is stored in the FRayTracingGeometry to keep track
 	// if the vertex buffer data is still valid for that frame - validated before generation the TLAS
-	uint64 SharedBufferGenerationID = 0;
+	int64 SharedBufferGenerationID = 0;
 };
 
 #endif // RHI_RAYTRACING
