Fixed sky to use the directional light color*intensity as outer space illuminance.
Previously, the illuminance on ground was use as input to the sky system. Sky is now using the outerspace illuminance as input.

diff --git a/Engine/Plugins/Experimental/GPULightmass/Source/GPULightmass/Private/Scene/Lights.cpp b/Engine/Plugins/Experimental/GPULightmass/Source/GPULightmass/Private/Scene/Lights.cpp
--- a/Engine/Plugins/Experimental/GPULightmass/Source/GPULightmass/Private/Scene/Lights.cpp
+++ b/Engine/Plugins/Experimental/GPULightmass/Source/GPULightmass/Private/Scene/Lights.cpp
@@ -219,7 +219,7 @@ FLightShaderParameters FDirectionalLightRenderState::GetLightShaderParameters()
 	LightParameters.Position = FVector::ZeroVector;
 	LightParameters.InvRadius = 0.0f;
 	// TODO: support SkyAtmosphere
-	// LightParameters.Color = FVector(GetColor() * AtmosphereTransmittanceFactor);
+	// LightParameters.Color = FVector(GetColor() * SkyAtmosphereTransmittanceToLight);
 	LightParameters.Color = FVector(Color);
 	LightParameters.FalloffExponent = 0.0f;
 
diff --git a/Engine/Shaders/Private/AtmosphereCommon.ush b/Engine/Shaders/Private/AtmosphereCommon.ush
--- a/Engine/Shaders/Private/AtmosphereCommon.ush
+++ b/Engine/Shaders/Private/AtmosphereCommon.ush
@@ -612,7 +612,7 @@ float4 GetAtmosphericFog(float3 ViewPosition, float3 ViewVector, float SceneDept
 
 	// Decay color
 	float3 OriginalColor = Sun + GroundColor + InscatterColor;
-	OriginalColor = View.AtmosphericFogSunPower * ShadowFactor * View.AtmosphereLightColor[0].rgb * OriginalColor;
+	OriginalColor = View.AtmosphericFogSunPower * ShadowFactor * View.AtmosphereLightIlluminanceOnGround[0].rgb * OriginalColor;
 	float4 OutColor = float4(OriginalColor, lerp(saturate(Attenuation.r * View.AtmosphericFogDensityScale - View.AtmosphericFogDensityOffset), 1.f, (1.f - DistanceRatio)) );
 	return OutColor;
 }
diff --git a/Engine/Shaders/Private/HeightFogCommon.ush b/Engine/Shaders/Private/HeightFogCommon.ush
--- a/Engine/Shaders/Private/HeightFogCommon.ush
+++ b/Engine/Shaders/Private/HeightFogCommon.ush
@@ -150,13 +150,13 @@ half4 GetExponentialHeightFog(float3 WorldPositionRelativeToCamera, float Exclud
 		half3 DirectionalInscatteringColor;
 		half3 DirectionalLightInscattering;
 
-		// No need to test View.AtmosphereLightColorGlobalPostTransmittance[0].a because InscatteringLightDirection.w above is doing the same test already.
-		DirectionalInscatteringColor = FogStruct.DirectionalInscatteringColor.xyz + View.AtmosphereLightColorGlobalPostTransmittance[0].rgb * UniformPhaseFunction;
+		// No need to test View.AtmosphereLightIlluminanceOnGroundPostTransmittance[0].a because InscatteringLightDirection.w above is doing the same test already.
+		DirectionalInscatteringColor = FogStruct.DirectionalInscatteringColor.xyz + View.AtmosphereLightIlluminanceOnGroundPostTransmittance[0].rgb * UniformPhaseFunction;
 		DirectionalLightInscattering = DirectionalInscatteringColor * pow(saturate(dot(CameraToReceiverNormalized, View.AtmosphereLightDirection[0].xyz)), FogStruct.DirectionalInscatteringColor.w);
 
-		if (View.AtmosphereLightColorGlobalPostTransmittance[1].a > 0.0f) // Skip DirectionalInscatteringColor on the second light when disabled.
+		if (View.AtmosphereLightIlluminanceOnGroundPostTransmittance[1].a > 0.0f) // Skip DirectionalInscatteringColor on the second light when disabled.
 		{
-			DirectionalInscatteringColor = FogStruct.DirectionalInscatteringColor.xyz + View.AtmosphereLightColorGlobalPostTransmittance[1].rgb * UniformPhaseFunction;
+			DirectionalInscatteringColor = FogStruct.DirectionalInscatteringColor.xyz + View.AtmosphereLightIlluminanceOnGroundPostTransmittance[1].rgb * UniformPhaseFunction;
 			DirectionalLightInscattering += DirectionalInscatteringColor * pow(saturate(dot(CameraToReceiverNormalized, View.AtmosphereLightDirection[1].xyz)), FogStruct.DirectionalInscatteringColor.w);
 		}
 		DirectionalLightInscattering *= View.SkyAtmosphereHeightFogContribution;
diff --git a/Engine/Shaders/Private/MaterialTemplate.ush b/Engine/Shaders/Private/MaterialTemplate.ush
--- a/Engine/Shaders/Private/MaterialTemplate.ush
+++ b/Engine/Shaders/Private/MaterialTemplate.ush
@@ -1256,7 +1256,7 @@ float3 MaterialExpressionAtmosphericLightVector(FMaterialPixelParameters Paramet
 float3 MaterialExpressionAtmosphericLightColor(FMaterialPixelParameters Parameters)
 {
 #if MATERIAL_ATMOSPHERIC_FOG
-	return ResolvedView.AtmosphereLightColor[0].rgb;
+	return ResolvedView.AtmosphereLightIlluminanceOnGroundPostTransmittance[0].rgb;
 #else
 	return float3(0.f, 0.f, 0.f);
 #endif
@@ -1272,7 +1272,7 @@ float3 MaterialExpressionSkyAtmosphereLightIlluminance(FMaterialPixelParameters
 		PlanetCenterToWorldPos, ResolvedView.AtmosphereLightDirection[LightIndex].xyz, ResolvedView.SkyAtmosphereBottomRadiusKm, ResolvedView.SkyAtmosphereTopRadiusKm,
 		View.TransmittanceLutTexture, View.TransmittanceLutTextureSampler);
 
-	return ResolvedView.AtmosphereLightColor[LightIndex].rgb * TransmittanceToLight;
+	return ResolvedView.AtmosphereLightIlluminanceOuterSpace[LightIndex].rgb * TransmittanceToLight;
 #else
 	return float3(0.0f, 0.0f, 0.0f);
 #endif
diff --git a/Engine/Shaders/Private/SkyAtmosphere.usf b/Engine/Shaders/Private/SkyAtmosphere.usf
--- a/Engine/Shaders/Private/SkyAtmosphere.usf
+++ b/Engine/Shaders/Private/SkyAtmosphere.usf
@@ -994,7 +994,8 @@ void RenderSkyAtmosphereRayMarchingPS(
 	SingleScatteringResult ss = IntegrateSingleScatteredLuminance(
 		SVPos, WorldPos, WorldDir,
 		Ground, Sampling, DeviceZ, MieRayPhase,
-		View.AtmosphereLightDirection[0].xyz, View.AtmosphereLightDirection[1].xyz, View.AtmosphereLightColor[0].rgb, View.AtmosphereLightColor[1].rgb,
+		View.AtmosphereLightDirection[0].xyz, View.AtmosphereLightDirection[1].xyz, 
+		View.AtmosphereLightIlluminanceOuterSpace[0].rgb, View.AtmosphereLightIlluminanceOuterSpace[1].rgb,
 		AerialPespectiveViewDistanceScale);
 
 	PreExposedL += ss.L * LuminanceScale;
@@ -1250,7 +1251,8 @@ void RenderSkyViewLutCS(uint3 ThreadId : SV_DispatchThreadID)
 	SingleScatteringResult ss = IntegrateSingleScatteredLuminance(
 		float4(PixPos, 0.0f, 1.0f), WorldPos, WorldDir,
 		Ground, Sampling, DeviceZ, MieRayPhase,
-		AtmosphereLightDirection0, AtmosphereLightDirection1, View.AtmosphereLightColor[0].rgb, View.AtmosphereLightColor[1].rgb,
+		AtmosphereLightDirection0, AtmosphereLightDirection1,
+		View.AtmosphereLightIlluminanceOuterSpace[0].rgb, View.AtmosphereLightIlluminanceOuterSpace[1].rgb,
 		AerialPespectiveViewDistanceScale);
 
 	SkyViewLutUAV[int2(PixPos)] = ss.L;
@@ -1266,8 +1268,8 @@ RWTexture2D<float3> DistantSkyLightLutUAV;
 //Buffer<float4> UniformSphereSamplesBuffer;
 float4 AtmosphereLightDirection0;
 float4 AtmosphereLightDirection1;
-float4 AtmosphereLightColor0;
-float4 AtmosphereLightColor1;
+float4 AtmosphereLightIlluminanceOuterSpace0;
+float4 AtmosphereLightIlluminanceOuterSpace1;
 float DistantSkyLightSampleAltitude;
 
 groupshared float3 GroupSkyLuminanceSamples[THREADGROUP_SIZE*THREADGROUP_SIZE];
@@ -1305,7 +1307,8 @@ void RenderDistantSkyLightLutCS(uint3 ThreadId : SV_DispatchThreadID)
 	SingleScatteringResult ss = IntegrateSingleScatteredLuminance(
 		float4(PixPos, 0.0f, 1.0f), SamplePos, SampleDir,
 		Ground, Sampling, DeviceZ, MieRayPhase,
-		AtmosphereLightDirection0.xyz, AtmosphereLightDirection1.xyz, AtmosphereLightColor0.rgb, AtmosphereLightColor1.rgb,
+		AtmosphereLightDirection0.xyz, AtmosphereLightDirection1.xyz, 
+		AtmosphereLightIlluminanceOuterSpace0.rgb, AtmosphereLightIlluminanceOuterSpace1.rgb,
 		AerialPespectiveViewDistanceScale);
 
 	GroupSkyLuminanceSamples[LinearIndex] = ss.L * SkyAtmosphere.SkyLuminanceFactor;
@@ -1494,7 +1497,8 @@ void RenderCameraAerialPerspectiveVolumeCS(uint3 ThreadId : SV_DispatchThreadID)
 	SingleScatteringResult ss = IntegrateSingleScatteredLuminance( 
 		float4(PixPos, 0.0f, 1.0f), RayStartWorldPos, WorldDir,
 		Ground, Sampling, DeviceZ, MieRayPhase,
-		View.AtmosphereLightDirection[0].xyz, View.AtmosphereLightDirection[1].xyz, View.AtmosphereLightColor[0].rgb, View.AtmosphereLightColor[1].rgb,
+		View.AtmosphereLightDirection[0].xyz, View.AtmosphereLightDirection[1].xyz, 
+		View.AtmosphereLightIlluminanceOuterSpace[0].rgb, View.AtmosphereLightIlluminanceOuterSpace[1].rgb,
 		AerialPespectiveViewDistanceScale,
 		tMaxMax);
 
@@ -1563,7 +1567,8 @@ void RenderSkyAtmosphereDebugPS(
 			SingleScatteringResult SS = IntegrateSingleScatteredLuminance(
 				SVPos, WorldPos, WorldDir,
 				Ground, Sampling, DeviceZ, MieRayPhase,
-				View.AtmosphereLightDirection[0].xyz, NullLightDirection, View.AtmosphereLightColor[0].rgb, NullLightIlluminance,
+				View.AtmosphereLightDirection[0].xyz, NullLightDirection, 
+				View.AtmosphereLightIlluminanceOuterSpace[0].rgb, NullLightIlluminance,
 				AerialPespectiveViewDistanceScale);
 			OutLuminance = float4(SS.L, 0.0f);
 			return;
@@ -1603,7 +1608,8 @@ void RenderSkyAtmosphereDebugPS(
 		SingleScatteringResult SS = IntegrateSingleScatteredLuminance(
 			SVPos, WorldPos, WorldDir,
 			Ground, Sampling, DeviceZ, MieRayPhase,
-			LightDir, NullLightDirection, View.AtmosphereLightColor[0].rgb, NullLightIlluminance,
+			LightDir, NullLightDirection, 
+			View.AtmosphereLightIlluminanceOuterSpace[0].rgb, NullLightIlluminance,
 			AerialPespectiveViewDistanceScale);
 
 		OutLuminance = float4(SS.L, 0.0f);
diff --git a/Engine/Shaders/Private/VolumetricCloud.usf b/Engine/Shaders/Private/VolumetricCloud.usf
--- a/Engine/Shaders/Private/VolumetricCloud.usf
+++ b/Engine/Shaders/Private/VolumetricCloud.usf
@@ -528,18 +528,18 @@ void MainCommon(in FMaterialPixelParameters MaterialParameters, in float4 SvPosi
 		//
 
 #if CLOUD_PER_SAMPLE_ATMOSPHERE_TRANSMITTANCE 
-		const float3 Light0Illuminance = ResolvedView.AtmosphereLightColor[0].rgb;
+		const float3 Light0Illuminance = ResolvedView.AtmosphereLightIlluminanceOuterSpace[0].rgb;
 #else
-		const float3 Light0Illuminance = ResolvedView.AtmosphereLightColorGlobalPostTransmittance[0].rgb;
+		const float3 Light0Illuminance = ResolvedView.AtmosphereLightIlluminanceOnGroundPostTransmittance[0].rgb;
 #endif
 		const float3 Light0IlluminanceFinal = Light0Illuminance * (RenderVolumetricCloudParameters.EnableAtmosphericLightsSampling ? RenderVolumetricCloudParameters.AtmosphericLightCloudScatteredLuminanceScale[0].rgb : float3(0.0f, 0.0f, 0.0f));
 		const float3 Light0Direction = ResolvedView.AtmosphereLightDirection[0].xyz;
 
 #if CLOUD_SAMPLE_SECOND_LIGHT
 #if CLOUD_PER_SAMPLE_ATMOSPHERE_TRANSMITTANCE 
-		const float3 Light1Illuminance = ResolvedView.AtmosphereLightColor[1].rgb;
+		const float3 Light1Illuminance = ResolvedView.AtmosphereLightIlluminanceOuterSpace[1].rgb;
 #else
-		const float3 Light1Illuminance = ResolvedView.AtmosphereLightColorGlobalPostTransmittance[1].rgb;
+		const float3 Light1Illuminance = ResolvedView.AtmosphereLightIlluminanceOnGroundPostTransmittance[1].rgb;
 #endif
 		const float3 Light1IlluminanceFinal = Light1Illuminance * (RenderVolumetricCloudParameters.EnableAtmosphericLightsSampling ? RenderVolumetricCloudParameters.AtmosphericLightCloudScatteredLuminanceScale[1].rgb : float3(0.0f, 0.0f, 0.0f));
 		const float3 Light1Direction = ResolvedView.AtmosphereLightDirection[1].xyz;
diff --git a/Engine/Source/Runtime/Engine/Private/Components/DirectionalLightComponent.cpp b/Engine/Source/Runtime/Engine/Private/Components/DirectionalLightComponent.cpp
--- a/Engine/Source/Runtime/Engine/Private/Components/DirectionalLightComponent.cpp
+++ b/Engine/Source/Runtime/Engine/Private/Components/DirectionalLightComponent.cpp
@@ -213,20 +213,26 @@ public:
 	FVector LightShaftOverrideDirection;
 
 	/**
-	 * The atmosphere transmittance to apply on the illuminance
+	 * Whether or not the light uses per pixel transmittance, driving the fact that should apply the simple transmittance computed on CPU during lighting pass. 
+	 * If per pixel transmittance is enabled, it should not be done to avoid double transmittance contribution (transmittance will be applied on GPU).
 	 */
-	FLinearColor AtmosphereTransmittanceFactor;
-	/**
-	 * Whether or not the light should apply the simple transmittance computed on CPU during lighting pass. 
-	 * If per pixel transmittance is enabled, it should not be done to avoid double transmittance contribution.
-	 */
-	bool bApplyAtmosphereTransmittanceToLightShaderParam;
+	bool bPerPixelTransmittanceEnabled;
 
 	/**
 	 * The luminance of the sun disk in space (function of the sun illuminance and solid angle)
 	 */
 	FLinearColor SunDiscOuterSpaceLuminance;
 
+	/**
+	 * The light outer space illuminance used to feed the SkyAtmosphere simulation
+	 */
+	FLinearColor SunOuterSpaceIlluminance;
+
+	/**
+	 * The atmosphere transmittance from ground to current light direction.
+	 */
+	FLinearColor TransmittanceTowardSun;
+
 	/** 
 	 * Radius of the whole scene dynamic shadow centered on the viewer, which replaces the precomputed shadows based on distance from the camera.  
 	 * A Radius of 0 disables the dynamic shadow.
@@ -289,9 +295,10 @@ public:
 		OcclusionMaskDarkness(Component->OcclusionMaskDarkness),
 		OcclusionDepthRange(Component->OcclusionDepthRange),
 		LightShaftOverrideDirection(Component->LightShaftOverrideDirection),
-		AtmosphereTransmittanceFactor(FLinearColor::White),
-		bApplyAtmosphereTransmittanceToLightShaderParam(true),
+		bPerPixelTransmittanceEnabled(false),
 		SunDiscOuterSpaceLuminance(FLinearColor::White),
+		SunOuterSpaceIlluminance(FLinearColor::White),
+		TransmittanceTowardSun(FLinearColor::White),
 		DynamicShadowCascades(Component->DynamicShadowCascades > 0 ? Component->DynamicShadowCascades : 0),
 		CascadeDistributionExponent(Component->CascadeDistributionExponent),
 		CascadeTransitionFraction(Component->CascadeTransitionFraction),
@@ -372,8 +379,7 @@ public:
 		LightParameters.InvRadius = 0.0f;
 		LightParameters.FalloffExponent = 0.0f;
 
-		// We only apply transmittance in some cases. For instance if transmittance is evaluated per pixel, no apply it to the light illuminance.
-		LightParameters.Color = FVector(GetColor() * (bApplyAtmosphereTransmittanceToLightShaderParam ? AtmosphereTransmittanceFactor : FLinearColor::White));
+		LightParameters.Color = FVector(GetSunIlluminanceAccountingForSkyAtmospherePerPixelTransmittance());
 
 		LightParameters.Direction = -GetDirection();
 		LightParameters.Tangent = -GetDirection();
@@ -530,11 +536,12 @@ public:
 		return DoesPlatformSupportDistanceFieldShadowing(GShaderPlatformForFeatureLevel[InFeatureLevel]) && (bCreateWithCSM || bCreateWithoutCSM);
 	}
 
-	virtual void SetAtmosphereRelatedProperties(FLinearColor TransmittanceFactor, FLinearColor SunOuterSpaceLuminance, bool bApplyAtmosphereTransmittanceToLightShaderParamIn) override
+	virtual void SetAtmosphereRelatedProperties(FLinearColor TransmittanceTowardSunIn, FLinearColor SunOuterSpaceIlluminanceIn, FLinearColor SunOuterSpaceLuminance, bool bPerPixelTransmittanceEnabledIn) override
 	{
-		AtmosphereTransmittanceFactor = TransmittanceFactor;
+		TransmittanceTowardSun = TransmittanceTowardSunIn;
 		SunDiscOuterSpaceLuminance = SunOuterSpaceLuminance;
-		bApplyAtmosphereTransmittanceToLightShaderParam = bApplyAtmosphereTransmittanceToLightShaderParamIn;
+		SunOuterSpaceIlluminance = SunOuterSpaceIlluminanceIn;
+		bPerPixelTransmittanceEnabled = bPerPixelTransmittanceEnabledIn;
 	}
 
 	virtual FLinearColor GetOuterSpaceLuminance() const override
@@ -542,9 +549,41 @@ public:
 		return SunDiscOuterSpaceLuminance;
 	}
 
-	virtual FLinearColor GetTransmittanceFactor() const override
+	virtual FLinearColor GetOuterSpaceIlluminance() const override
+	{
+		return GetColor();
+	}
+
+	virtual FLinearColor GetAtmosphereTransmittanceTowardSun() const override
+	{
+		return TransmittanceTowardSun;
+	}
+
+	virtual FLinearColor GetSunIlluminanceOnGroundPostTransmittance() const override
 	{
-		return AtmosphereTransmittanceFactor;
+		return GetOuterSpaceIlluminance() * GetAtmosphereTransmittanceTowardSun();
+	}
+
+	virtual bool GetPerPixelTransmittanceEnabled() const override
+	{
+		return bPerPixelTransmittanceEnabled;
+	}
+
+	virtual FLinearColor GetSunIlluminanceAccountingForSkyAtmospherePerPixelTransmittance() const override
+	{
+		if (IsUsedAsAtmosphereSunLight())
+		{
+			// We only apply transmittance in some cases. For instance if transmittance is evaluated per pixel, we do not apply it to the light illuminance.
+			if (bPerPixelTransmittanceEnabled)
+			{
+				return FVector(GetOuterSpaceIlluminance());
+			}
+			else
+			{
+				return FVector(GetOuterSpaceIlluminance() * GetAtmosphereTransmittanceTowardSun());
+			}
+		}
+		return GetColor();
 	}
 
 	virtual float GetSunLightHalfApexAngleRadian() const override
diff --git a/Engine/Source/Runtime/Engine/Public/Rendering/SkyAtmosphereCommonData.cpp b/Engine/Source/Runtime/Engine/Public/Rendering/SkyAtmosphereCommonData.cpp
--- a/Engine/Source/Runtime/Engine/Public/Rendering/SkyAtmosphereCommonData.cpp
+++ b/Engine/Source/Runtime/Engine/Public/Rendering/SkyAtmosphereCommonData.cpp
@@ -142,7 +142,7 @@ FLinearColor FAtmosphereSetup::GetTransmittanceAtGroundLevel(const FVector& SunD
 		if (TMax > 0.0f)
 		{
 			const float SampleCount = 15.0f;
-			const float SampleStep = 1.0f / 15.0f;
+			const float SampleStep = 1.0f / SampleCount;
 			const float SampleLength = SampleStep * TMax;
 			for (float SampleT = 0.0f; SampleT < 1.0f; SampleT += SampleStep)
 			{
diff --git a/Engine/Source/Runtime/Engine/Public/SceneManagement.h b/Engine/Source/Runtime/Engine/Public/SceneManagement.h
--- a/Engine/Source/Runtime/Engine/Public/SceneManagement.h
+++ b/Engine/Source/Runtime/Engine/Public/SceneManagement.h
@@ -1521,9 +1521,13 @@ public:
 	inline bool IsUsedAsAtmosphereSunLight() const { return bUsedAsAtmosphereSunLight; }
 	inline uint8 GetAtmosphereSunLightIndex() const { return AtmosphereSunLightIndex; }
 	inline FLinearColor GetAtmosphereSunDiskColorScale() const { return AtmosphereSunDiskColorScale; }
-	virtual void SetAtmosphereRelatedProperties(FLinearColor TransmittanceFactor, FLinearColor SunOuterSpaceLuminance, bool bApplyAtmosphereTransmittanceToLightShaderParamIn) {}
+	virtual void SetAtmosphereRelatedProperties(FLinearColor TransmittanceTowardSun, FLinearColor SunOuterSpaceIlluminance, FLinearColor SunOuterSpaceLuminance, bool bPerPixelTransmittanceEnabled) {}
 	virtual FLinearColor GetOuterSpaceLuminance() const { return FLinearColor::White; }
-	virtual FLinearColor GetTransmittanceFactor() const { return FLinearColor::White; }
+	virtual FLinearColor GetOuterSpaceIlluminance() const { return GetColor(); }
+	virtual FLinearColor GetAtmosphereTransmittanceTowardSun() const { return FLinearColor::White; }
+	virtual FLinearColor GetSunIlluminanceOnGroundPostTransmittance() const { return GetColor(); }
+	virtual FLinearColor GetSunIlluminanceAccountingForSkyAtmospherePerPixelTransmittance() const { return GetColor(); }
+	virtual bool GetPerPixelTransmittanceEnabled() const { return false; }
 	static float GetSunOnEarthHalfApexAngleRadian() 
 	{ 
 		const float SunOnEarthApexAngleDegree = 0.545f;	// Apex angle == angular diameter
diff --git a/Engine/Source/Runtime/Engine/Public/SceneView.h b/Engine/Source/Runtime/Engine/Public/SceneView.h
--- a/Engine/Source/Runtime/Engine/Public/SceneView.h
+++ b/Engine/Source/Runtime/Engine/Public/SceneView.h
@@ -695,8 +695,9 @@ enum ETranslucencyVolumeCascade
 	VIEW_UNIFORM_BUFFER_MEMBER_EX(float, AtmosphericFogDistanceOffset, EShaderPrecisionModifier::Half) \
 	VIEW_UNIFORM_BUFFER_MEMBER_EX(float, AtmosphericFogSunDiscScale, EShaderPrecisionModifier::Half) \
 	VIEW_UNIFORM_BUFFER_MEMBER_ARRAY(FVector4, AtmosphereLightDirection, [NUM_ATMOSPHERE_LIGHTS]) \
-	VIEW_UNIFORM_BUFFER_MEMBER_ARRAY(FLinearColor, AtmosphereLightColor, [NUM_ATMOSPHERE_LIGHTS]) \
-	VIEW_UNIFORM_BUFFER_MEMBER_ARRAY(FLinearColor, AtmosphereLightColorGlobalPostTransmittance, [NUM_ATMOSPHERE_LIGHTS]) \
+	VIEW_UNIFORM_BUFFER_MEMBER_ARRAY(FLinearColor, AtmosphereLightIlluminanceOnGround, [NUM_ATMOSPHERE_LIGHTS]) \
+	VIEW_UNIFORM_BUFFER_MEMBER_ARRAY(FLinearColor, AtmosphereLightIlluminanceOnGroundPostTransmittance, [NUM_ATMOSPHERE_LIGHTS]) \
+	VIEW_UNIFORM_BUFFER_MEMBER_ARRAY(FLinearColor, AtmosphereLightIlluminanceOuterSpace, [NUM_ATMOSPHERE_LIGHTS]) \
 	VIEW_UNIFORM_BUFFER_MEMBER_ARRAY(FLinearColor, AtmosphereLightDiscLuminance, [NUM_ATMOSPHERE_LIGHTS]) \
 	VIEW_UNIFORM_BUFFER_MEMBER_ARRAY(FVector4, AtmosphereLightDiscCosHalfApexAngle, [NUM_ATMOSPHERE_LIGHTS]) \
 	VIEW_UNIFORM_BUFFER_MEMBER(FVector4, SkyViewLutSizeAndInvSize) \
diff --git a/Engine/Source/Runtime/Renderer/Private/AtmosphereRendering.cpp b/Engine/Source/Runtime/Renderer/Private/AtmosphereRendering.cpp
--- a/Engine/Source/Runtime/Renderer/Private/AtmosphereRendering.cpp
+++ b/Engine/Source/Runtime/Renderer/Private/AtmosphereRendering.cpp
@@ -1771,8 +1771,8 @@ void FAtmosphericFogSceneInfo::PrepareSunLightProxy(FLightSceneInfo& SunLight) c
 	const float SunSolidAngle = 2.0f * PI * (1.0f - FMath::Cos(SunLight.Proxy->GetSunLightHalfApexAngleRadian())); // Solid angle from aperture https://en.wikipedia.org/wiki/Solid_angle 
 	FLinearColor SunDiskOuterSpaceLuminance = SunOuterSpaceIlluminance / SunSolidAngle; // approximation  
 
-	const bool bApplyAtmosphereTransmittanceToLightShaderParam = true;
-	SunLight.Proxy->SetAtmosphereRelatedProperties(TransmittanceTowardSun / TransmittanceAtZenithFinal, SunDiskOuterSpaceLuminance, bApplyAtmosphereTransmittanceToLightShaderParam);
+	const bool bPerPixelTransmittanceEnabled = false;
+	SunLight.Proxy->SetAtmosphereRelatedProperties(TransmittanceTowardSun, SunOuterSpaceIlluminance, SunDiskOuterSpaceLuminance, bPerPixelTransmittanceEnabled);
 }
 
 /** Initialization constructor. */
diff --git a/Engine/Source/Runtime/Renderer/Private/MobileBasePassRendering.cpp b/Engine/Source/Runtime/Renderer/Private/MobileBasePassRendering.cpp
--- a/Engine/Source/Runtime/Renderer/Private/MobileBasePassRendering.cpp
+++ b/Engine/Source/Runtime/Renderer/Private/MobileBasePassRendering.cpp
@@ -233,11 +233,7 @@ void SetupMobileDirectionalLightUniformParameters(
 	FLightSceneInfo* Light = Scene.MobileDirectionalLights[ChannelIdx];
 	if (Light)
 	{
-		Params.DirectionalLightColor = Light->Proxy->GetColor() / PI;
-		if (Light->Proxy->IsUsedAsAtmosphereSunLight())
-		{
-			Params.DirectionalLightColor *= Light->Proxy->GetTransmittanceFactor();
-		}
+		Params.DirectionalLightColor = Light->Proxy->GetSunIlluminanceAccountingForSkyAtmospherePerPixelTransmittance() / PI;
 		Params.DirectionalLightDirectionAndShadowTransition = FVector4(-Light->Proxy->GetDirection(), 0.f);
 
 		const FVector2D FadeParams = Light->Proxy->GetDirectionalLightDistanceFadeParameters(FeatureLevel, Light->IsPrecomputedLightingValid(), SceneView.MaxShadowCascades);
diff --git a/Engine/Source/Runtime/Renderer/Private/SceneRendering.cpp b/Engine/Source/Runtime/Renderer/Private/SceneRendering.cpp
--- a/Engine/Source/Runtime/Renderer/Private/SceneRendering.cpp
+++ b/Engine/Source/Runtime/Renderer/Private/SceneRendering.cpp
@@ -1284,10 +1284,11 @@ void FViewInfo::SetupUniformBufferParameters(
 		check(Index < NUM_ATMOSPHERE_LIGHTS);
 		ViewUniformShaderParameters.AtmosphereLightDiscCosHalfApexAngle[Index] = FVector4(1.0f);
 		ViewUniformShaderParameters.AtmosphereLightDiscLuminance[Index] = FLinearColor::Black;
-		ViewUniformShaderParameters.AtmosphereLightColor[Index] = FLinearColor::Black;
-		ViewUniformShaderParameters.AtmosphereLightColor[Index].A = 0.0f;
-		ViewUniformShaderParameters.AtmosphereLightColorGlobalPostTransmittance[Index] = FLinearColor::Black;
-		ViewUniformShaderParameters.AtmosphereLightColorGlobalPostTransmittance[Index].A = 0.0f;
+		ViewUniformShaderParameters.AtmosphereLightIlluminanceOnGround[Index] = FLinearColor::Black;
+		ViewUniformShaderParameters.AtmosphereLightIlluminanceOnGround[Index].A = 0.0f;
+		ViewUniformShaderParameters.AtmosphereLightIlluminanceOnGroundPostTransmittance[Index] = FLinearColor::Black;
+		ViewUniformShaderParameters.AtmosphereLightIlluminanceOnGroundPostTransmittance[Index].A = 0.0f;
+		ViewUniformShaderParameters.AtmosphereLightIlluminanceOuterSpace[Index] = FLinearColor::Black;
 		ViewUniformShaderParameters.AtmosphereLightDirection[Index] = DefaultSunDirection;
 	};
 
@@ -1297,7 +1298,7 @@ void FViewInfo::SetupUniformBufferParameters(
 		bShouldRenderAtmosphericFog = ShouldRenderAtmosphere(*Family) && Scene->AtmosphericFog;
 		if (Scene->SimpleDirectionalLight)
 		{
-			ViewUniformShaderParameters.DirectionalLightColor = Scene->SimpleDirectionalLight->Proxy->GetTransmittanceFactor() * Scene->SimpleDirectionalLight->Proxy->GetColor() / PI;
+			ViewUniformShaderParameters.DirectionalLightColor = Scene->SimpleDirectionalLight->Proxy->GetAtmosphereTransmittanceTowardSun() * Scene->SimpleDirectionalLight->Proxy->GetColor() / PI;
 			ViewUniformShaderParameters.DirectionalLightDirection = -Scene->SimpleDirectionalLight->Proxy->GetDirection();
 		}
 		else
@@ -1326,10 +1327,12 @@ void FViewInfo::SetupUniformBufferParameters(
 			ViewUniformShaderParameters.AtmosphericFogInscatterAltitudeSampleNum = Scene->AtmosphericFog->InscatterAltitudeSampleNum;
 			ViewUniformShaderParameters.AtmosphereLightDiscCosHalfApexAngle[0] = FVector4(FMath::Cos(Scene->AtmosphericFog->SunDiscScale * SunLightDiskHalfApexAngleRadian));
 			ViewUniformShaderParameters.AtmosphereLightDiscLuminance[0] = SunLight ? SunLight->Proxy->GetOuterSpaceLuminance() : FLinearColor::White;
-			ViewUniformShaderParameters.AtmosphereLightColor[0] = SunLight ? SunLight->Proxy->GetColor() : Scene->AtmosphericFog->DefaultSunColor; // Sun light color unaffected by atmosphere transmittance
-			ViewUniformShaderParameters.AtmosphereLightColor[0].A = 1.0f;
-			ViewUniformShaderParameters.AtmosphereLightColorGlobalPostTransmittance[0] = ViewUniformShaderParameters.AtmosphereLightColor[0];
-			ViewUniformShaderParameters.AtmosphereLightColorGlobalPostTransmittance[0].A = 0.0f;
+			ViewUniformShaderParameters.AtmosphereLightIlluminanceOnGround[0] = SunLight ? SunLight->Proxy->GetColor() : Scene->AtmosphericFog->DefaultSunColor; // Sun light color unaffected by atmosphere transmittance
+			ViewUniformShaderParameters.AtmosphereLightIlluminanceOnGround[0].A = 1.0f;
+			ViewUniformShaderParameters.AtmosphereLightIlluminanceOnGroundPostTransmittance[0] = ViewUniformShaderParameters.AtmosphereLightIlluminanceOnGround[0];
+			ViewUniformShaderParameters.AtmosphereLightIlluminanceOnGroundPostTransmittance[0].A = 0.0f;
+			ViewUniformShaderParameters.AtmosphereLightIlluminanceOuterSpace[0] = SunLight ? SunLight->Proxy->GetOuterSpaceIlluminance() : FLinearColor::Black;
+			ViewUniformShaderParameters.AtmosphereLightIlluminanceOuterSpace[0].A = 1.0f;
 			ViewUniformShaderParameters.AtmosphereLightDirection[0] = SunLight ? -SunLight->Proxy->GetDirection() : -Scene->AtmosphericFog->DefaultSunDirection;
 		}
 		else
@@ -1350,11 +1353,11 @@ void FViewInfo::SetupUniformBufferParameters(
 			ViewUniformShaderParameters.AtmosphereLightDiscCosHalfApexAngle[0] = FVector4(FMath::Cos(SunLightDiskHalfApexAngleRadian));
 			//Added check so atmospheric light color and vector can use a directional light without needing an atmospheric fog actor in the scene
 			ViewUniformShaderParameters.AtmosphereLightDiscLuminance[0] = SunLight ? SunLight->Proxy->GetOuterSpaceLuminance() : FLinearColor::Black;
-			ViewUniformShaderParameters.AtmosphereLightColor[0] = SunLight ? SunLight->Proxy->GetColor() : FLinearColor::Black;
-			ViewUniformShaderParameters.AtmosphereLightColor[0].A = 1.0f;
-			ViewUniformShaderParameters.AtmosphereLightColorGlobalPostTransmittance[0] = ViewUniformShaderParameters.AtmosphereLightColor[0];
-			ViewUniformShaderParameters.AtmosphereLightColorGlobalPostTransmittance[0].A = 0.0f;
-			ViewUniformShaderParameters.AtmosphereLightDirection[0] = SunLight ? -SunLight->Proxy->GetDirection() : DefaultSunDirection;
+			ViewUniformShaderParameters.AtmosphereLightIlluminanceOnGround[0] = SunLight ? SunLight->Proxy->GetColor() : FLinearColor::Black;
+			ViewUniformShaderParameters.AtmosphereLightIlluminanceOnGround[0].A = 1.0f;
+			ViewUniformShaderParameters.AtmosphereLightIlluminanceOnGroundPostTransmittance[0] = ViewUniformShaderParameters.AtmosphereLightIlluminanceOnGround[0];
+			ViewUniformShaderParameters.AtmosphereLightIlluminanceOnGroundPostTransmittance[0].A = 0.0f;
+			ViewUniformShaderParameters.AtmosphereLightIlluminanceOuterSpace[0] = FLinearColor::Black;
 		}
 
 		// Do not clear the first AtmosphereLight data, it has been setup above
@@ -1447,10 +1450,12 @@ void FViewInfo::SetupUniformBufferParameters(
 			{
 				ViewUniformShaderParameters.AtmosphereLightDiscCosHalfApexAngle[Index] = FVector4(FMath::Cos(Light->Proxy->GetSunLightHalfApexAngleRadian()));
 				ViewUniformShaderParameters.AtmosphereLightDiscLuminance[Index] = Light->Proxy->GetOuterSpaceLuminance();
-				ViewUniformShaderParameters.AtmosphereLightColor[Index] = Light->Proxy->GetColor();
-				ViewUniformShaderParameters.AtmosphereLightColor[Index].A = 1.0f;
-				ViewUniformShaderParameters.AtmosphereLightColorGlobalPostTransmittance[Index] = Light->Proxy->GetColor() * Light->Proxy->GetTransmittanceFactor();
-				ViewUniformShaderParameters.AtmosphereLightColorGlobalPostTransmittance[Index].A = 1.0f;
+				ViewUniformShaderParameters.AtmosphereLightIlluminanceOnGround[Index] = Light->Proxy->GetSunIlluminanceAccountingForSkyAtmospherePerPixelTransmittance();
+				ViewUniformShaderParameters.AtmosphereLightIlluminanceOnGround[Index].A = 1.0f;
+				ViewUniformShaderParameters.AtmosphereLightIlluminanceOnGroundPostTransmittance[Index] = Light->Proxy->GetSunIlluminanceOnGroundPostTransmittance();
+				ViewUniformShaderParameters.AtmosphereLightIlluminanceOnGroundPostTransmittance[Index].A = 1.0f; // interactions with HeightFogComponent
+				ViewUniformShaderParameters.AtmosphereLightIlluminanceOuterSpace[Index] = Light->Proxy->GetOuterSpaceIlluminance();
+				ViewUniformShaderParameters.AtmosphereLightIlluminanceOuterSpace[Index].A = 1.0f;
 				ViewUniformShaderParameters.AtmosphereLightDirection[Index] = SkyAtmosphereSceneProxy.GetAtmosphereLightDirection(Index, -Light->Proxy->GetDirection());
 			}
 			else
@@ -1493,10 +1498,11 @@ void FViewInfo::SetupUniformBufferParameters(
 				{
 					ViewUniformShaderParameters.AtmosphereLightDiscCosHalfApexAngle[Index] = FVector4(1.0f);
 					ViewUniformShaderParameters.AtmosphereLightDiscLuminance[Index] = FLinearColor::Black;
-					ViewUniformShaderParameters.AtmosphereLightColor[Index] = Light->Proxy->GetColor();
-					ViewUniformShaderParameters.AtmosphereLightColor[Index].A = 1.0f;
-					ViewUniformShaderParameters.AtmosphereLightColorGlobalPostTransmittance[Index] = Light->Proxy->GetColor();
-					ViewUniformShaderParameters.AtmosphereLightColorGlobalPostTransmittance[Index].A = 0.0f; // no interactions with HeightFogComponent
+					ViewUniformShaderParameters.AtmosphereLightIlluminanceOnGround[Index] = Light->Proxy->GetColor();
+					ViewUniformShaderParameters.AtmosphereLightIlluminanceOnGround[Index].A = 1.0f;
+					ViewUniformShaderParameters.AtmosphereLightIlluminanceOnGroundPostTransmittance[Index] = Light->Proxy->GetColor();
+					ViewUniformShaderParameters.AtmosphereLightIlluminanceOnGroundPostTransmittance[Index].A = 0.0f; // no interactions with HeightFogComponent
+					ViewUniformShaderParameters.AtmosphereLightIlluminanceOuterSpace[Index] = Light->Proxy->GetColor();
 					ViewUniformShaderParameters.AtmosphereLightDirection[Index] = -Light->Proxy->GetDirection();
 				}
 				else
diff --git a/Engine/Source/Runtime/Renderer/Private/SkyAtmosphereRendering.cpp b/Engine/Source/Runtime/Renderer/Private/SkyAtmosphereRendering.cpp
--- a/Engine/Source/Runtime/Renderer/Private/SkyAtmosphereRendering.cpp
+++ b/Engine/Source/Runtime/Renderer/Private/SkyAtmosphereRendering.cpp
@@ -419,17 +419,20 @@ void PrepareSunLightProxy(const FSkyAtmosphereRenderSceneInfo& SkyAtmosphere, ui
 	const bool bAtmosphereAffectsSunIlluminance = true;
 	const FSkyAtmosphereSceneProxy& SkyAtmosphereProxy = SkyAtmosphere.GetSkyAtmosphereSceneProxy();
 	const FVector AtmosphereLightDirection = SkyAtmosphereProxy.GetAtmosphereLightDirection(AtmosphereLightIndex, -AtmosphereLight.Proxy->GetDirection());
-	FLinearColor TransmittanceTowardSun = bAtmosphereAffectsSunIlluminance ? SkyAtmosphereProxy.GetAtmosphereSetup().GetTransmittanceAtGroundLevel(AtmosphereLightDirection) : FLinearColor(FLinearColor::White);
-	FLinearColor TransmittanceAtZenithFinal = bAtmosphereAffectsSunIlluminance ? SkyAtmosphereProxy.GetTransmittanceAtZenith() : FLinearColor(FLinearColor::White);
+	const FLinearColor TransmittanceTowardSun = bAtmosphereAffectsSunIlluminance ? SkyAtmosphereProxy.GetAtmosphereSetup().GetTransmittanceAtGroundLevel(AtmosphereLightDirection) : FLinearColor(FLinearColor::White);
 
-	FLinearColor SunZenithIlluminance = AtmosphereLight.Proxy->GetColor();
-	FLinearColor SunOuterSpaceIlluminance = SunZenithIlluminance / TransmittanceAtZenithFinal;
-	FLinearColor SunDiskOuterSpaceLuminance = GetLightDiskLuminance(AtmosphereLight, SunOuterSpaceIlluminance);
+	// TODO: when the AtmosphericFog is gone, we hsould be able to remove this, as well as TransmittanceAtZenith on the SkyAtmosphereProxy.
+	const FLinearColor TransmittanceAtZenithFinal = bAtmosphereAffectsSunIlluminance ? SkyAtmosphereProxy.GetTransmittanceAtZenith() : FLinearColor(FLinearColor::White);
+
+	const FLinearColor SunZenithIlluminanceOnGround = AtmosphereLight.Proxy->GetColor();
+	const FLinearColor SunOuterSpaceIlluminance = SunZenithIlluminanceOnGround / TransmittanceAtZenithFinal;
+	const FLinearColor SunDiskOuterSpaceLuminance = GetLightDiskLuminance(AtmosphereLight, SunOuterSpaceIlluminance);
 
 	// We always set the transmittance on the proxy. Shaders using atmospheric light color then have to decide which sun illuminance to use (without or with transmittance).
 	// We also set whether or not the light should apply the simple transmittance computed on CPU during lighting pass. If per pixel transmittance is enabled, it should not be done.
-	const bool bApplyAtmosphereTransmittanceToLightShaderParam = !AtmosphereLight.Proxy->GetUsePerPixelAtmosphereTransmittance();
-	AtmosphereLight.Proxy->SetAtmosphereRelatedProperties(TransmittanceTowardSun / TransmittanceAtZenithFinal, SunDiskOuterSpaceLuminance, bApplyAtmosphereTransmittanceToLightShaderParam);
+	const bool bPerPixelTransmittanceEnabled = AtmosphereLight.Proxy->GetUsePerPixelAtmosphereTransmittance();
+
+	AtmosphereLight.Proxy->SetAtmosphereRelatedProperties(TransmittanceTowardSun, SunOuterSpaceIlluminance, SunDiskOuterSpaceLuminance, bPerPixelTransmittanceEnabled);
 }
 
 
@@ -533,8 +536,8 @@ void FScene::ResetAtmosphereLightsProperties()
 		if (Light)
 		{
 			FLinearColor LightZenithIlluminance = Light->Proxy->GetColor();
-			const bool bApplyAtmosphereTransmittanceToLightShaderParam = true;
-			Light->Proxy->SetAtmosphereRelatedProperties(FLinearColor::White, GetLightDiskLuminance(*Light, LightZenithIlluminance), bApplyAtmosphereTransmittanceToLightShaderParam);
+			const bool bPerPixelTransmittanceEnabled = false;
+			Light->Proxy->SetAtmosphereRelatedProperties(FLinearColor::White, LightZenithIlluminance, GetLightDiskLuminance(*Light, LightZenithIlluminance), bPerPixelTransmittanceEnabled);
 		}
 	}
 }
@@ -786,8 +789,8 @@ public:
 		SHADER_PARAMETER_SRV(Buffer<float4>, UniformSphereSamplesBuffer)
 		SHADER_PARAMETER(FVector4, AtmosphereLightDirection0)
 		SHADER_PARAMETER(FVector4, AtmosphereLightDirection1)
-		SHADER_PARAMETER(FLinearColor, AtmosphereLightColor0)
-		SHADER_PARAMETER(FLinearColor, AtmosphereLightColor1)
+		SHADER_PARAMETER(FLinearColor, AtmosphereLightIlluminanceOuterSpace0)
+		SHADER_PARAMETER(FLinearColor, AtmosphereLightIlluminanceOuterSpace1)
 		SHADER_PARAMETER(float, DistantSkyLightSampleAltitude)
 	END_SHADER_PARAMETER_STRUCT()
 
@@ -1328,7 +1331,6 @@ void FSceneRenderer::RenderSkyAtmosphereLookUpTables(FRDGBuilder& GraphBuilder)
 
 		FRenderDistantSkyLightLutCS::FPermutationDomain PermutationVector;
 		PermutationVector.Set<FSecondAtmosphereLight>(bSecondAtmosphereLightEnabled);
-		PermutationVector.Set<FSecondAtmosphereLight>(bSecondAtmosphereLightEnabled);
 		TShaderMapRef<FRenderDistantSkyLightLutCS> ComputeShader(GlobalShaderMap, PermutationVector);
 
 		FRenderDistantSkyLightLutCS::FParameters * PassParameters = GraphBuilder.AllocParameters<FRenderDistantSkyLightLutCS::FParameters>();
@@ -1346,22 +1348,22 @@ void FSceneRenderer::RenderSkyAtmosphereLookUpTables(FRDGBuilder& GraphBuilder)
 		if (Light0)
 		{
 			PassParameters->AtmosphereLightDirection0 = -Light0->Proxy->GetDirection();
-			PassParameters->AtmosphereLightColor0 = Light0->Proxy->GetColor();
+			PassParameters->AtmosphereLightIlluminanceOuterSpace0 = Light0->Proxy->GetOuterSpaceIlluminance();
 		}
 		else
 		{
 			PassParameters->AtmosphereLightDirection0 = FVector4(0.0f, 0.0f, 1.0f, 1.0f);
-			PassParameters->AtmosphereLightColor0 = FLinearColor::Black;
+			PassParameters->AtmosphereLightIlluminanceOuterSpace0 = FLinearColor::Black;
 		}
 		if (Light1)
 		{
 			PassParameters->AtmosphereLightDirection1 = -Light1->Proxy->GetDirection();
-			PassParameters->AtmosphereLightColor1 = Light1->Proxy->GetColor();
+			PassParameters->AtmosphereLightIlluminanceOuterSpace1 = Light1->Proxy->GetOuterSpaceIlluminance();
 		}
 		else
 		{
 			PassParameters->AtmosphereLightDirection1 = FVector4(0.0f, 0.0f, 1.0f, 1.0f);
-			PassParameters->AtmosphereLightColor1 = FLinearColor::Black;
+			PassParameters->AtmosphereLightIlluminanceOuterSpace1 = FLinearColor::Black;
 		}
 		PassParameters->DistantSkyLightSampleAltitude = CVarSkyAtmosphereDistantSkyLightLUTAltitude.GetValueOnAnyThread();
 
