Fixed issue with GBufferClear overrunning array when 8 render targets are used.

diff --git a/Engine/Source/Runtime/Renderer/Private/BasePassRendering.cpp b/Engine/Source/Runtime/Renderer/Private/BasePassRendering.cpp
--- a/Engine/Source/Runtime/Renderer/Private/BasePassRendering.cpp
+++ b/Engine/Source/Runtime/Renderer/Private/BasePassRendering.cpp
@@ -806,14 +806,10 @@ void FDeferredShadingSceneRenderer::RenderBasePass(
 				FRHITexture* Textures[MaxSimultaneousRenderTargets];
 				int32 TextureIndex = 0;
 
-				ClearColors[TextureIndex] = SceneColorClearValue;
-				Textures[TextureIndex] = RenderTargets[0].GetTexture()->GetRHI();
-				++TextureIndex;
-
 				RenderTargets.Enumerate([&](const FRenderTargetBinding& RenderTarget)
 				{
 					FRHITexture* TextureRHI = RenderTarget.GetTexture()->GetRHI();
-					ClearColors[TextureIndex] = TextureRHI->GetClearColor();
+					ClearColors[TextureIndex] = TextureIndex == 0 ? SceneColorClearValue : TextureRHI->GetClearColor();
 					Textures[TextureIndex] = TextureRHI;
 					++TextureIndex;
 				});
