Path Tracing setting cleanups

Add PPV flag to enable path traced emissive materials (on by default). These paths are not always needed if there are lots of tiny emitters, or if the scene was modelled with both emissive geometry and lights.

Change Skip* variables to Enable* to be more intuitive

Turn the sampler options into a proper enum since they were intended to be mutually exclusive

Tweak MaxPathIntensity logic to match Lumen and avoid color shifts

diff --git a/Engine/Shaders/Private/PathTracing/PathTracing.usf b/Engine/Shaders/Private/PathTracing/PathTracing.usf
--- a/Engine/Shaders/Private/PathTracing/PathTracing.usf
+++ b/Engine/Shaders/Private/PathTracing/PathTracing.usf
@@ -30,7 +30,6 @@ PathTracingRayGenShader.usf: Reference path tracing
 uint Iteration;
 uint TemporalSeed;
 uint MaxSamples;
-uint UseErrorDiffusion;
 uint MaxBounces;
 uint MaxSSSBounces;
 
@@ -41,9 +40,9 @@ uint MISMode;
 
 uint ApproximateCaustics;
 uint EnableCameraBackfaceCulling;
-uint SkipDirectLighting;
-uint SkipEmissive;
-uint CoherentSampling;
+uint EnableDirectLighting;
+uint EnableEmissive;
+uint SamplerType;
 uint VisualizeLightGrid;
 float MaxPathIntensity;
 float MaxNormalBias;
@@ -57,17 +56,17 @@ uint SceneVisibleLightCount;
 
 void AccumulateRadiance(inout float3 TotalRadiance, float3 PathRadiance)
 {
-	if (MaxPathIntensity > 0)
+	// User asked for path contributions to be clamped to reduce fireflies.
+	// This puts an upper bound on variance within the pixel at the expense of bias
+	float MaxPathRadiance = max3(PathRadiance.x, PathRadiance.y, PathRadiance.z);
+	if (MaxPathRadiance > MaxPathIntensity)
 	{
-		// User asked for path contributions to be clamped to reduce fireflies.
-		// Depending on how aggressive this value is, the image could be quite biased
-		TotalRadiance += min(PathRadiance, MaxPathIntensity);
-	}
-	else
-	{
-		// Just average values directly
-		TotalRadiance += PathRadiance;
+		// adjust sample while keeping color
+		PathRadiance *= MaxPathIntensity / MaxPathRadiance;
 	}
+
+	// accumulate into the total
+	TotalRadiance += PathRadiance;
 }
 
 FMaterialClosestHitPayload TraceTransparentRay(RayDesc Ray, FRayCone RayCone, bool IsCameraRay, bool LastBounce, bool IncludeEmission, uint2 LaunchIndex, uint NumLights, inout RandomSequence RandSequence, inout float3 PathThroughput, inout float3 Radiance)
@@ -490,16 +489,16 @@ RAY_TRACING_ENTRY_RAYGEN(PathTracingMainRG)
 	uint2 DispatchDim = DispatchRaysDimensions().xy;
 
 	// shuffle the pixels around so the ones that share a seed land in the same warp
-	uint2 LaunchIndex = (CoherentSampling ? Reorder(DispatchIdx, DispatchDim) : DispatchIdx) + View.ViewRectMin.xy + TileOffset.xy;
+	uint2 LaunchIndex = (SamplerType == PATHTRACER_SAMPLER_COHERENT ? Reorder(DispatchIdx, DispatchDim) : DispatchIdx) + View.ViewRectMin.xy + TileOffset.xy;
 
 	// Initialize random sequence 
 	RandomSequence RandSequence;
 
-	uint2 RandPixel = CoherentSampling
+	uint2 RandPixel = SamplerType == PATHTRACER_SAMPLER_COHERENT
 		? (DispatchIdx / PATHTRACER_COHERENT_TILE_SIZE)	// unique seed within small squares of the shuffled sub-images
 		//? (DispatchIdx * PATHTRACER_COHERENT_TILE_SIZE) / DispatchDim // unique seed per sub-image
 		: LaunchIndex; // unique seed per pixel
-	if (UseErrorDiffusion)
+	if (SamplerType == PATHTRACER_SAMPLER_ERROR_DIFFUSION)
 	{
 		// z-sampler init
 		RandomSequence_Initialize(RandSequence, RandPixel, Iteration, TemporalSeed - Iteration, MaxSamples);
@@ -544,8 +543,8 @@ RAY_TRACING_ENTRY_RAYGEN(PathTracingMainRG)
 #endif
 		const bool bIsCameraRay = Bounce == 0;
 		const bool bIsLastBounce = Bounce == MaxBounces;
-		const bool bIncludeEmissive = (SkipDirectLighting == 0 || Bounce > 1) &&
-									  (SkipEmissive       == 0 || bIsCameraRay);
+		const bool bIncludeEmissive = (EnableDirectLighting != 0 || Bounce > 1) &&
+									  (EnableEmissive       != 0 || bIsCameraRay);
 		FMaterialClosestHitPayload Payload = TraceTransparentRay(Ray, RayCone, bIsCameraRay, bIsLastBounce, bIncludeEmissive, LaunchIndex, NumVisibleLights, RandSequence, PathThroughput, Radiance);
 
 		if (Payload.IsMiss())
@@ -590,7 +589,7 @@ RAY_TRACING_ENTRY_RAYGEN(PathTracingMainRG)
 		// Does this material require NEE? (will be false if MaterialPdf is always +inf)
 		const bool bIsNeeValid = IsNeeValidMaterial(Payload);
 
-		const bool bDoLightLoop = SkipDirectLighting == 0 || Bounce > 0;
+		const bool bDoLightLoop = EnableDirectLighting != 0 || Bounce > 0;
 
 		float LightPickingCdfSum = 0;
 
diff --git a/Engine/Shaders/Shared/PathTracingDefinitions.h b/Engine/Shaders/Shared/PathTracingDefinitions.h
--- a/Engine/Shaders/Shared/PathTracingDefinitions.h
+++ b/Engine/Shaders/Shared/PathTracingDefinitions.h
@@ -7,6 +7,10 @@
 #define PATHTRACER_COHERENT_TILE_SHIFT		3
 #define PATHTRACER_COHERENT_TILE_SIZE		(1 << PATHTRACER_COHERENT_TILE_SHIFT)
 
+// Constants for 'SamplerType'
+#define PATHTRACER_SAMPLER_DEFAULT			0
+#define PATHTRACER_SAMPLER_ERROR_DIFFUSION	1
+#define PATHTRACER_SAMPLER_COHERENT			2
 
 // Constants for the 'Flags' field of FPathTracingLight
 #define PATHTRACER_FLAG_LIGHTING_CHANNEL_MASK				(7 << 0)	// Which lighting channel is this light assigned to?
diff --git a/Engine/Source/Runtime/Engine/Classes/Engine/Scene.h b/Engine/Source/Runtime/Engine/Classes/Engine/Scene.h
--- a/Engine/Source/Runtime/Engine/Classes/Engine/Scene.h
+++ b/Engine/Source/Runtime/Engine/Classes/Engine/Scene.h
@@ -1192,6 +1192,9 @@ struct FPostProcessSettings
 	UPROPERTY(EditAnywhere, BlueprintReadWrite, Category = Overrides, meta = (PinHiddenByDefault, InlineEditConditionToggle))
 	uint32 bOverride_PathTracingFilterWidth : 1;
 
+	UPROPERTY(EditAnywhere, BlueprintReadWrite, Category = Overrides, meta = (PinHiddenByDefault, InlineEditConditionToggle))
+	uint32 bOverride_PathTracingEnableEmissive : 1;
+
 	// -----------------------------------------------------------------------
 
 	/** Enable HQ Gaussian on high end mobile platforms. (ES3_1) */
@@ -1925,6 +1928,11 @@ struct FPostProcessSettings
 	UPROPERTY(interp, EditAnywhere, BlueprintReadWrite, Category = "Rendering Features|PathTracing", meta = (ClampMin = "1.0", ClampMax = "6.0", editcondition = "bOverride_PathTracingFilterWidth", DisplayName = "Filter Width"))
 	float PathTracingFilterWidth;
 
+	/** Enables emissive materials for the path tracer. */
+	UPROPERTY(EditAnywhere, BlueprintReadWrite, Category = "Rendering Features|PathTracing", meta = (editcondition = "bOverride_PathTracingEnableEmissive", DisplayName = "Emissive Materials"))
+	uint32 PathTracingEnableEmissive : 1;
+
+
 
 	UPROPERTY()
 	float LPVFadeRange_DEPRECATED;
diff --git a/Engine/Source/Runtime/Engine/Private/Scene.cpp b/Engine/Source/Runtime/Engine/Private/Scene.cpp
--- a/Engine/Source/Runtime/Engine/Private/Scene.cpp
+++ b/Engine/Source/Runtime/Engine/Private/Scene.cpp
@@ -578,6 +578,7 @@ FPostProcessSettings::FPostProcessSettings()
 	PathTracingMaxBounces = 32;
 	PathTracingSamplesPerPixel = 16384;
 	PathTracingFilterWidth = 3.0f;
+	PathTracingEnableEmissive = 1;
 	
 	bMobileHQGaussian = false;
 
@@ -757,6 +758,7 @@ FPostProcessSettings::FPostProcessSettings(const FPostProcessSettings& Settings)
 	, bOverride_PathTracingMaxBounces(Settings.bOverride_PathTracingMaxBounces)
 	, bOverride_PathTracingSamplesPerPixel(Settings.bOverride_PathTracingSamplesPerPixel)
 	, bOverride_PathTracingFilterWidth(Settings.bOverride_PathTracingFilterWidth)
+	, bOverride_PathTracingEnableEmissive(Settings.bOverride_PathTracingEnableEmissive)
 
 	, bMobileHQGaussian(Settings.bMobileHQGaussian)
 	, BloomMethod(Settings.BloomMethod)
@@ -927,6 +929,7 @@ FPostProcessSettings::FPostProcessSettings(const FPostProcessSettings& Settings)
 	, PathTracingMaxBounces(Settings.PathTracingMaxBounces)
 	, PathTracingSamplesPerPixel(Settings.PathTracingSamplesPerPixel)
 	, PathTracingFilterWidth(Settings.PathTracingFilterWidth)
+	, PathTracingEnableEmissive(Settings.PathTracingEnableEmissive)
 
 	, ScreenPercentage(Settings.ScreenPercentage)
 
diff --git a/Engine/Source/Runtime/Engine/Private/SceneView.cpp b/Engine/Source/Runtime/Engine/Private/SceneView.cpp
--- a/Engine/Source/Runtime/Engine/Private/SceneView.cpp
+++ b/Engine/Source/Runtime/Engine/Private/SceneView.cpp
@@ -1564,6 +1564,11 @@ void FSceneView::OverridePostProcessSettings(const FPostProcessSettings& Src, fl
 			Dest.PathTracingFilterWidth = Src.PathTracingFilterWidth;
 		}
 
+		if (Src.bOverride_PathTracingEnableEmissive)
+		{
+			Dest.PathTracingEnableEmissive = Src.PathTracingEnableEmissive;
+		}
+
 
 		if (Src.bOverride_DepthOfFieldBladeCount)
 		{
diff --git a/Engine/Source/Runtime/Renderer/Private/PathTracing.cpp b/Engine/Source/Runtime/Renderer/Private/PathTracing.cpp
--- a/Engine/Source/Runtime/Renderer/Private/PathTracing.cpp
+++ b/Engine/Source/Runtime/Renderer/Private/PathTracing.cpp
@@ -34,14 +34,6 @@ TAutoConsoleVariable<float> CVarPathTracingFilterWidth(
 	ECVF_RenderThreadSafe
 );
 
-
-TAutoConsoleVariable<int32> CVarPathTracingUseErrorDiffusion(
-	TEXT("r.PathTracing.UseErrorDiffusion"),
-	0,
-	TEXT("Enables an experimental sampler that diffuses visible error in screen space. This generally produces better results when the target sample count can be reached. (default = 0 (disabled))"),
-	ECVF_RenderThreadSafe
-);
-
 TAutoConsoleVariable<int32> CVarPathTracingMISMode(
 	TEXT("r.PathTracing.MISMode"),
 	2,
@@ -98,10 +90,10 @@ TAutoConsoleVariable<int32> CVarPathTracingApproximateCaustics(
 	ECVF_RenderThreadSafe
 );
 
-TAutoConsoleVariable<int32> CVarPathTracingSkipEmissive(
-	TEXT("r.PathTracing.SkipEmissive"),
-	0,
-	TEXT("When non-zero, the path tracer will skip emissive results after the first bounce. (default = 0 (off))"),
+TAutoConsoleVariable<int32> CVarPathTracingEnableEmissive(
+	TEXT("r.PathTracing.EnableEmissive"),
+	-1,
+	TEXT("Indicates if emissive materials should contribute to scene lighting (default = -1 (inherit from PostProcessVolume)"),
 	ECVF_RenderThreadSafe
 );
 
@@ -121,12 +113,14 @@ TAutoConsoleVariable<int32> CVarPathTracingFrameIndependentTemporalSeed(
 	ECVF_RenderThreadSafe
 );
 
-TAutoConsoleVariable<int32> CVarPathTracingCoherentSampling(
-	TEXT("r.PathTracing.CoherentSampling"),
-	0,
-	TEXT("When non-zero, share pixel seeds to improve coherence of execution on the GPU. This trades some correlation across the image in exchange for better performance.\n")
-	TEXT("0: off (default)\n")
-	TEXT("1: on\n"),
+// See PATHTRACER_SAMPLER_* defines
+TAutoConsoleVariable<int32> CVarPathTracingSamplerType(
+	TEXT("r.PathTracing.SamplerType"),
+	PATHTRACER_SAMPLER_DEFAULT,
+	TEXT("Controls the way the path tracer generates its random numbers\n")
+	TEXT("0: use a different high quality random sequence per pixel\n")
+	TEXT("1: optimize the random sequence across pixels to lower error at the target sample count\n")
+	TEXT("2: share random seeds across pixels to improve coherence of execution on the GPU. This trades some correlation across the image in exchange for better performance.\n"),
 	ECVF_RenderThreadSafe
 );
 
@@ -183,15 +177,14 @@ BEGIN_SHADER_PARAMETER_STRUCT(FPathTracingData, )
 	SHADER_PARAMETER(uint32, Iteration)
 	SHADER_PARAMETER(uint32, TemporalSeed)
 	SHADER_PARAMETER(uint32, MaxSamples)
-	SHADER_PARAMETER(uint32, UseErrorDiffusion)
 	SHADER_PARAMETER(uint32, MaxBounces)
 	SHADER_PARAMETER(uint32, MaxSSSBounces)
 	SHADER_PARAMETER(uint32, MISMode)
 	SHADER_PARAMETER(uint32, ApproximateCaustics)
 	SHADER_PARAMETER(uint32, EnableCameraBackfaceCulling)
-	SHADER_PARAMETER(uint32, SkipDirectLighting)
-	SHADER_PARAMETER(uint32, SkipEmissive)
-	SHADER_PARAMETER(uint32, CoherentSampling)
+	SHADER_PARAMETER(uint32, EnableDirectLighting)
+	SHADER_PARAMETER(uint32, EnableEmissive)
+	SHADER_PARAMETER(uint32, SamplerType)
 	SHADER_PARAMETER(uint32, VisualizeLightGrid)
 	SHADER_PARAMETER(float, MaxPathIntensity)
 	SHADER_PARAMETER(float, MaxNormalBias)
@@ -201,7 +194,7 @@ END_SHADER_PARAMETER_STRUCT()
 
 // This function prepares the portion of shader arguments that may involve invalidating the path traced state
 static bool PrepareShaderArgs(const FViewInfo& View, FPathTracingData& PathTracingData) {
-	PathTracingData.SkipDirectLighting = false;
+	PathTracingData.EnableDirectLighting = true;
 	int32 MaxBounces = CVarPathTracingMaxBounces.GetValueOnRenderThread();
 	if (MaxBounces < 0)
 	{
@@ -217,7 +210,7 @@ static bool PrepareShaderArgs(const FViewInfo& View, FPathTracingData& PathTraci
 	}
 	else
 	{
-		PathTracingData.SkipDirectLighting = true;
+		PathTracingData.EnableDirectLighting = false;
 		if (View.Family->EngineShowFlags.GlobalIllumination)
 		{
 			// skip direct lighting, but still do the full bounces
@@ -235,11 +228,16 @@ static bool PrepareShaderArgs(const FViewInfo& View, FPathTracingData& PathTraci
 	PathTracingData.MISMode = CVarPathTracingMISMode.GetValueOnRenderThread();
 	uint32 VisibleLights = CVarPathTracingVisibleLights.GetValueOnRenderThread();
 	PathTracingData.MaxPathIntensity = CVarPathTracingMaxPathIntensity.GetValueOnRenderThread();
-	PathTracingData.UseErrorDiffusion = CVarPathTracingUseErrorDiffusion.GetValueOnRenderThread();
+	if (PathTracingData.MaxPathIntensity <= 0)
+	{
+		// no clamping required - set to an impossible value
+		PathTracingData.MaxPathIntensity = std::numeric_limits<float>::infinity();
+	}
 	PathTracingData.ApproximateCaustics = CVarPathTracingApproximateCaustics.GetValueOnRenderThread();
 	PathTracingData.EnableCameraBackfaceCulling = CVarPathTracingEnableCameraBackfaceCulling.GetValueOnRenderThread();
-	PathTracingData.CoherentSampling = CVarPathTracingCoherentSampling.GetValueOnRenderThread();
-	PathTracingData.SkipEmissive = CVarPathTracingSkipEmissive.GetValueOnRenderThread();
+	PathTracingData.SamplerType = CVarPathTracingSamplerType.GetValueOnRenderThread();
+	int32 EnableEmissive = CVarPathTracingEnableEmissive.GetValueOnRenderThread();
+	PathTracingData.EnableEmissive = EnableEmissive < 0 ? View.FinalPostProcessSettings.PathTracingEnableEmissive : EnableEmissive;
 	PathTracingData.VisualizeLightGrid = CVarPathTracingLightGridVisualize.GetValueOnRenderThread();
 	float FilterWidth = CVarPathTracingFilterWidth.GetValueOnRenderThread();
 	if (FilterWidth < 0)
@@ -292,14 +290,6 @@ static bool PrepareShaderArgs(const FViewInfo& View, FPathTracingData& PathTraci
 		PreviousMaxPathIntensity = PathTracingData.MaxPathIntensity;
 	}
 
-	// Changing sampler requires starting over
-	static uint32 PreviousUseErrorDiffusion = PathTracingData.UseErrorDiffusion;
-	if (PreviousUseErrorDiffusion != PathTracingData.UseErrorDiffusion)
-	{
-		NeedInvalidation = true;
-		PreviousUseErrorDiffusion = PathTracingData.UseErrorDiffusion;
-	}
-
 	// Changing approximate caustics requires starting over
 	static uint32 PreviousApproximateCaustics = PathTracingData.ApproximateCaustics;
 	if (PreviousApproximateCaustics != PathTracingData.ApproximateCaustics)
@@ -324,28 +314,28 @@ static bool PrepareShaderArgs(const FViewInfo& View, FPathTracingData& PathTraci
 		PreviousBackfaceCulling = PathTracingData.EnableCameraBackfaceCulling;
 	}
 
-	// Changing direct lighting skipping requires starting over
-	static uint32 PreviousSkipDirectLighting = PathTracingData.SkipDirectLighting;
-	if (PreviousSkipDirectLighting != PathTracingData.SkipDirectLighting)
+	// Changing direct lighting requires starting over
+	static uint32 PreviousEnableDirectLighting = PathTracingData.EnableDirectLighting;
+	if (PreviousEnableDirectLighting != PathTracingData.EnableDirectLighting)
 	{
 		NeedInvalidation = true;
-		PreviousSkipDirectLighting = PathTracingData.SkipDirectLighting;
+		PreviousEnableDirectLighting = PathTracingData.EnableDirectLighting;
 	}
 
-	// Changing skip emissive requires starting over
-	static uint32 PreviousSkipEmissive = PathTracingData.SkipEmissive;
-	if (PreviousSkipEmissive != PathTracingData.SkipEmissive)
+	// Changing enable emissive requires starting over
+	static uint32 PreviousEnableEmissive = PathTracingData.EnableEmissive;
+	if (PreviousEnableEmissive != PathTracingData.EnableEmissive)
 	{
 		NeedInvalidation = true;
-		PreviousSkipEmissive = PathTracingData.SkipEmissive;
+		PreviousEnableEmissive = PathTracingData.EnableEmissive;
 	}
 
-	// Changing coherent sampling requires starting over
-	static uint32 PreviousCoherentSampling = PathTracingData.CoherentSampling;
-	if (PreviousCoherentSampling != PathTracingData.CoherentSampling)
+	// Changing sampler type requires starting over
+	static uint32 PreviousSamplerType = PathTracingData.SamplerType;
+	if (PreviousSamplerType != PathTracingData.SamplerType)
 	{
 		NeedInvalidation = true;
-		PreviousCoherentSampling = PathTracingData.CoherentSampling;
+		PreviousSamplerType = PathTracingData.SamplerType;
 	}
 
 	// Changing visualization mode requires starting over
@@ -1269,7 +1259,7 @@ void FDeferredShadingSceneRenderer::RenderPathTracing(
 		PassParameters->PathTracingData = PathTracingData;
 		// upload sky/lights data
 		SetLightParameters(GraphBuilder, PassParameters, Scene, View, UseMISCompensation);
-		if (PathTracingData.SkipDirectLighting)
+		if (PathTracingData.EnableDirectLighting == 0)
 		{
 			PassParameters->SceneVisibleLightCount = 0;
 		}
