Fixed crash when editor primitives are rendered with the default material because the correct materials are still compiling.

The default material uses the ReflectionCapture UB, which wasn't bound in the editor primitives pass because the real material doesn't need it.

#rnx

diff --git a/Engine/Source/Runtime/Renderer/Private/PostProcess/PostProcessCompositeEditorPrimitives.cpp b/Engine/Source/Runtime/Renderer/Private/PostProcess/PostProcessCompositeEditorPrimitives.cpp
--- a/Engine/Source/Runtime/Renderer/Private/PostProcess/PostProcessCompositeEditorPrimitives.cpp
+++ b/Engine/Source/Runtime/Renderer/Private/PostProcess/PostProcessCompositeEditorPrimitives.cpp
@@ -274,6 +274,7 @@ const FViewInfo* CreateEditorPrimitiveView(const FViewInfo& ParentView, FIntRect
 
 BEGIN_SHADER_PARAMETER_STRUCT(FEditorPrimitivesPassParameters, )
 	SHADER_PARAMETER_STRUCT_REF(FViewUniformShaderParameters, View)
+	SHADER_PARAMETER_STRUCT_REF(FReflectionCaptureShaderData, ReflectionCapture)
 	SHADER_PARAMETER_RDG_UNIFORM_BUFFER(FOpaqueBasePassUniformParameters, BasePass)
 	SHADER_PARAMETER_RDG_UNIFORM_BUFFER(FMobileBasePassUniformParameters, MobileBasePass)
 	RENDER_TARGET_BINDING_SLOTS()
@@ -470,6 +471,7 @@ FScreenPassTexture AddEditorPrimitivePass(
 		{
 			FEditorPrimitivesPassParameters* PassParameters = GraphBuilder.AllocParameters<FEditorPrimitivesPassParameters>();
 			PassParameters->View = EditorView->ViewUniformBuffer;
+			PassParameters->ReflectionCapture = View.ReflectionCaptureUniformBuffer;
 			PassParameters->RenderTargets[0] = FRenderTargetBinding(EditorPrimitiveColor, ERenderTargetLoadAction::ELoad);
 			PassParameters->RenderTargets.DepthStencil = FDepthStencilBinding(EditorPrimitiveDepth, ERenderTargetLoadAction::ELoad, ERenderTargetLoadAction::ELoad, FExclusiveDepthStencil::DepthWrite_StencilWrite);
 
