Delete a duplicate code.

diff --git a/Engine/Source/Runtime/Renderer/Private/PostProcess/PostProcessing.cpp b/Engine/Source/Runtime/Renderer/Private/PostProcess/PostProcessing.cpp
--- a/Engine/Source/Runtime/Renderer/Private/PostProcess/PostProcessing.cpp
+++ b/Engine/Source/Runtime/Renderer/Private/PostProcess/PostProcessing.cpp
@@ -1921,7 +1921,6 @@ void AddMobilePostProcessingPasses(FRDGBuilder& GraphBuilder, FScene* Scene, con
 		PassSequence.SetEnabled(EPass::SunMerge, false);
 		PassSequence.SetEnabled(EPass::SeparateTranslucency, false);
 		PassSequence.SetEnabled(EPass::TAA, false);
-		PassSequence.SetEnabled(EPass::FXAA, false);
 		PassSequence.SetEnabled(EPass::PostProcessMaterialAfterTonemapping, false);
 		PassSequence.SetEnabled(EPass::FXAA, false);
 		PassSequence.Finalize();
