Fix a crash when rendering shadow map depth and velocity on mobile opengl platform.

diff --git a/Engine/Source/Runtime/Renderer/Private/PostProcess/PostProcessMotionBlur.cpp b/Engine/Source/Runtime/Renderer/Private/PostProcess/PostProcessMotionBlur.cpp
--- a/Engine/Source/Runtime/Renderer/Private/PostProcess/PostProcessMotionBlur.cpp
+++ b/Engine/Source/Runtime/Renderer/Private/PostProcess/PostProcessMotionBlur.cpp
@@ -784,9 +784,9 @@ FScreenPassTexture AddVisualizeMotionBlurPass(FRDGBuilder& GraphBuilder, const F
 		static const auto VelocityTestVar = IConsoleManager::Get().FindTConsoleVariableDataInt(TEXT("r.VelocityTest"));
 		const int32 VelocityTest = VelocityTestVar ? VelocityTestVar->GetValueOnRenderThread() : 0;
 
-		extern bool IsParallelVelocity();
+		extern bool IsParallelVelocity(EShaderPlatform ShaderPlatform);
 
-		Line = FString::Printf(TEXT("%d, %d, %d"), View.Family->bWorldIsPaused, VelocityTest, IsParallelVelocity());
+		Line = FString::Printf(TEXT("%d, %d, %d"), View.Family->bWorldIsPaused, VelocityTest, IsParallelVelocity(View.GetShaderPlatform()));
 		Canvas.DrawShadowedString(X, Y += YStep, TEXT("Paused, r.VelocityTest, Parallel:"), GetStatsFont(), FLinearColor(1, 1, 0));
 		Canvas.DrawShadowedString(X + ColumnWidth, Y, *Line, GetStatsFont(), FLinearColor(1, 1, 0));
 
diff --git a/Engine/Source/Runtime/Renderer/Private/ShadowDepthRendering.cpp b/Engine/Source/Runtime/Renderer/Private/ShadowDepthRendering.cpp
--- a/Engine/Source/Runtime/Renderer/Private/ShadowDepthRendering.cpp
+++ b/Engine/Source/Runtime/Renderer/Private/ShadowDepthRendering.cpp
@@ -1457,10 +1457,12 @@ static void RenderShadowDepthAtlasNanite(
 	}
 }
 
-bool IsParallelDispatchEnabled(const FProjectedShadowInfo* ProjectedShadowInfo)
+bool IsParallelDispatchEnabled(const FProjectedShadowInfo* ProjectedShadowInfo, EShaderPlatform ShaderPlatform)
 {
 	return GRHICommandList.UseParallelAlgorithms() && CVarParallelShadows.GetValueOnRenderThread()
-		&& (ProjectedShadowInfo->IsWholeSceneDirectionalShadow() || CVarParallelShadowsNonWholeScene.GetValueOnRenderThread());
+		&& (ProjectedShadowInfo->IsWholeSceneDirectionalShadow() || CVarParallelShadowsNonWholeScene.GetValueOnRenderThread())
+		// Parallel dispatch is not supported on mobile platform
+		&& !IsMobilePlatform(ShaderPlatform);
 }
 
 void FSceneRenderer::RenderShadowDepthMapAtlases(FRDGBuilder& GraphBuilder)
@@ -1492,7 +1494,7 @@ void FSceneRenderer::RenderShadowDepthMapAtlases(FRDGBuilder& GraphBuilder)
 		// Gather our passes here to minimize switching render passes
 		for (FProjectedShadowInfo* ProjectedShadowInfo : ShadowMapAtlas.Shadows)
 		{
-			if (IsParallelDispatchEnabled(ProjectedShadowInfo))
+			if (IsParallelDispatchEnabled(ProjectedShadowInfo, ShaderPlatform))
 			{
 				ParallelShadowPasses.Add(ProjectedShadowInfo);
 			}
@@ -1876,7 +1878,7 @@ void FSceneRenderer::RenderShadowDepthMaps(FRDGBuilder& GraphBuilder, FInstanceC
 		}
 
 		{
-			const bool bDoParallelDispatch = IsParallelDispatchEnabled(ProjectedShadowInfo);
+			const bool bDoParallelDispatch = IsParallelDispatchEnabled(ProjectedShadowInfo, ShaderPlatform);
 			ProjectedShadowInfo->RenderDepth(GraphBuilder, this, ShadowDepthTexture, bDoParallelDispatch);
 		}
 
@@ -2020,7 +2022,7 @@ void FSceneRenderer::RenderShadowDepthMaps(FRDGBuilder& GraphBuilder, FInstanceC
 				RDG_GPU_MASK_SCOPE(GraphBuilder, GetGPUMaskForShadow(ProjectedShadowInfo));
 				AddClearShadowDepthPass(GraphBuilder, PreshadowCacheTexture, ProjectedShadowInfo);
 
-				const bool bParallelDispatch = IsParallelDispatchEnabled(ProjectedShadowInfo);
+				const bool bParallelDispatch = IsParallelDispatchEnabled(ProjectedShadowInfo, ShaderPlatform);
 				ProjectedShadowInfo->RenderDepth(GraphBuilder, this, PreshadowCacheTexture, bParallelDispatch);
 				ProjectedShadowInfo->bDepthsCached = true;
 			}
diff --git a/Engine/Source/Runtime/Renderer/Private/VelocityRendering.cpp b/Engine/Source/Runtime/Renderer/Private/VelocityRendering.cpp
--- a/Engine/Source/Runtime/Renderer/Private/VelocityRendering.cpp
+++ b/Engine/Source/Runtime/Renderer/Private/VelocityRendering.cpp
@@ -51,14 +51,16 @@ static TAutoConsoleVariable<int32> CVarRHICmdFlushRenderThreadTasksVelocityPass(
 
 DECLARE_GPU_STAT_NAMED(RenderVelocities, TEXT("Render Velocities"));
 
-bool IsParallelVelocity()
+bool IsParallelVelocity(EShaderPlatform ShaderPlatform)
 {
-	return GRHICommandList.UseParallelAlgorithms() && CVarParallelVelocity.GetValueOnRenderThread();
+	return GRHICommandList.UseParallelAlgorithms() && CVarParallelVelocity.GetValueOnRenderThread()
+		// Parallel dispatch is not supported on mobile platform
+		&& !IsMobilePlatform(ShaderPlatform);
 }
 
-bool IsVelocityWaitForTasksEnabled()
+bool IsVelocityWaitForTasksEnabled(EShaderPlatform ShaderPlatform)
 {
-	return IsParallelVelocity() && (CVarRHICmdFlushRenderThreadTasksVelocityPass.GetValueOnRenderThread() > 0 || CVarRHICmdFlushRenderThreadTasks.GetValueOnRenderThread() > 0);
+	return IsParallelVelocity(ShaderPlatform) && (CVarRHICmdFlushRenderThreadTasksVelocityPass.GetValueOnRenderThread() > 0 || CVarRHICmdFlushRenderThreadTasks.GetValueOnRenderThread() > 0);
 }
 
 bool IsVelocityMergedWithDepthPass()
@@ -227,7 +229,7 @@ void FSceneRenderer::RenderVelocities(
 		: ERenderTargetLoadAction::EClear;
 
 	RDG_GPU_STAT_SCOPE(GraphBuilder, RenderVelocities);
-	RDG_WAIT_FOR_TASKS_CONDITIONAL(GraphBuilder, IsVelocityWaitForTasksEnabled());
+	RDG_WAIT_FOR_TASKS_CONDITIONAL(GraphBuilder, IsVelocityWaitForTasksEnabled(ShaderPlatform));
 
 	const EMeshPass::Type MeshPass = GetMeshPassFromVelocityPass(VelocityPass);
 
@@ -279,7 +281,7 @@ void FSceneRenderer::RenderVelocities(
 
 			PassParameters->RenderTargets[0] = FRenderTargetBinding(SceneTextures.Velocity, VelocityLoadAction);
 
-			if (IsParallelVelocity())
+			if (IsParallelVelocity(ShaderPlatform))
 			{
 				GraphBuilder.AddPass(
 					RDG_EVENT_NAME("VelocityParallel"),
