Add MaxPathExposure to PPV to control path tracer clamping

The maximum path intensity is normally relative to the overall exposure. Since the path tracer accumulates accross frames we cannot tie the clamp to the exposure as it would lead to inconsistent results when interactively adjusting exposure during rendering (or create a strange feedback loop with eye adaptation).

Improve the path tracer setting tooltips in PPV a bit.

Minor cleanups to cvar descriptions, comment out unused cvars.

Cleaner empty buffer handling for LightGridData

Avoid redundant allocation and copy when uploading path tracer lights

diff --git a/Engine/Source/Runtime/Engine/Classes/Engine/Scene.h b/Engine/Source/Runtime/Engine/Classes/Engine/Scene.h
--- a/Engine/Source/Runtime/Engine/Classes/Engine/Scene.h
+++ b/Engine/Source/Runtime/Engine/Classes/Engine/Scene.h
@@ -1195,6 +1195,9 @@ struct FPostProcessSettings
 	UPROPERTY(EditAnywhere, BlueprintReadWrite, Category = Overrides, meta = (PinHiddenByDefault, InlineEditConditionToggle))
 	uint32 bOverride_PathTracingEnableEmissive : 1;
 
+	UPROPERTY(EditAnywhere, BlueprintReadWrite, Category = Overrides, meta = (PinHiddenByDefault, InlineEditConditionToggle))
+	uint32 bOverride_PathTracingMaxPathExposure : 1;
+
 	// -----------------------------------------------------------------------
 
 	/** Enable HQ Gaussian on high end mobile platforms. (ES3_1) */
@@ -1924,14 +1927,17 @@ struct FPostProcessSettings
 	UPROPERTY(interp, EditAnywhere, BlueprintReadWrite, Category = "Rendering Features|PathTracing", meta = (ClampMin = "1", ClampMax = "65536", editcondition = "bOverride_PathTracingSamplesPerPixel", DisplayName = "Samples Per Pixel"))
 	int32 PathTracingSamplesPerPixel;
 
-	/** Sets anti-aliasing filter width for the path tracer. */
+	/** Sets anti-aliasing filter width for the path tracer. Lower values are sharper (and more aliased), larger values are softer (and blurrier). */
 	UPROPERTY(interp, EditAnywhere, BlueprintReadWrite, Category = "Rendering Features|PathTracing", meta = (ClampMin = "1.0", ClampMax = "6.0", editcondition = "bOverride_PathTracingFilterWidth", DisplayName = "Filter Width"))
 	float PathTracingFilterWidth;
 
-	/** Enables emissive materials for the path tracer. */
+	/** Enables emissive materials for the path tracer. This can prevent double-counting of illumination from surfaces that are also represented by light sources, and noise from small emitters. */
 	UPROPERTY(EditAnywhere, BlueprintReadWrite, Category = "Rendering Features|PathTracing", meta = (editcondition = "bOverride_PathTracingEnableEmissive", DisplayName = "Emissive Materials"))
 	uint32 PathTracingEnableEmissive : 1;
 
+	/** Sets the maximum exposure allowed in the path tracer to reduce fireflies. This should be set a few steps higher than the scene exposure. */
+	UPROPERTY(interp, EditAnywhere, BlueprintReadWrite, Category = "Rendering Features|PathTracing", meta = (ClampMin = "-10.0", ClampMax = "30.0", editcondition = "bOverride_PathTracingMaxPathExposure", DisplayName = "Max Path Exposure"))
+	float PathTracingMaxPathExposure;
 
 
 	UPROPERTY()
diff --git a/Engine/Source/Runtime/Engine/Private/Scene.cpp b/Engine/Source/Runtime/Engine/Private/Scene.cpp
--- a/Engine/Source/Runtime/Engine/Private/Scene.cpp
+++ b/Engine/Source/Runtime/Engine/Private/Scene.cpp
@@ -579,6 +579,7 @@ FPostProcessSettings::FPostProcessSettings()
 	PathTracingSamplesPerPixel = 16384;
 	PathTracingFilterWidth = 3.0f;
 	PathTracingEnableEmissive = 1;
+	PathTracingMaxPathExposure = 30.0f;
 	
 	bMobileHQGaussian = false;
 
@@ -759,6 +760,7 @@ FPostProcessSettings::FPostProcessSettings(const FPostProcessSettings& Settings)
 	, bOverride_PathTracingSamplesPerPixel(Settings.bOverride_PathTracingSamplesPerPixel)
 	, bOverride_PathTracingFilterWidth(Settings.bOverride_PathTracingFilterWidth)
 	, bOverride_PathTracingEnableEmissive(Settings.bOverride_PathTracingEnableEmissive)
+	, bOverride_PathTracingMaxPathExposure(Settings.bOverride_PathTracingMaxPathExposure)
 
 	, bMobileHQGaussian(Settings.bMobileHQGaussian)
 	, BloomMethod(Settings.BloomMethod)
@@ -930,6 +932,7 @@ FPostProcessSettings::FPostProcessSettings(const FPostProcessSettings& Settings)
 	, PathTracingSamplesPerPixel(Settings.PathTracingSamplesPerPixel)
 	, PathTracingFilterWidth(Settings.PathTracingFilterWidth)
 	, PathTracingEnableEmissive(Settings.PathTracingEnableEmissive)
+	, PathTracingMaxPathExposure(Settings.PathTracingMaxPathExposure)
 
 	, ScreenPercentage(Settings.ScreenPercentage)
 
diff --git a/Engine/Source/Runtime/Engine/Private/SceneView.cpp b/Engine/Source/Runtime/Engine/Private/SceneView.cpp
--- a/Engine/Source/Runtime/Engine/Private/SceneView.cpp
+++ b/Engine/Source/Runtime/Engine/Private/SceneView.cpp
@@ -1569,6 +1569,11 @@ void FSceneView::OverridePostProcessSettings(const FPostProcessSettings& Src, fl
 			Dest.PathTracingEnableEmissive = Src.PathTracingEnableEmissive;
 		}
 
+		if (Src.bOverride_PathTracingMaxPathExposure)
+		{
+			Dest.PathTracingMaxPathExposure = Src.PathTracingMaxPathExposure;
+		}
+
 
 		if (Src.bOverride_DepthOfFieldBladeCount)
 		{
diff --git a/Engine/Source/Runtime/Renderer/Private/PathTracing.cpp b/Engine/Source/Runtime/Renderer/Private/PathTracing.cpp
--- a/Engine/Source/Runtime/Renderer/Private/PathTracing.cpp
+++ b/Engine/Source/Runtime/Renderer/Private/PathTracing.cpp
@@ -79,7 +79,7 @@ TAutoConsoleVariable<int32> CVarPathTracingMaxSSSBounces(
 TAutoConsoleVariable<float> CVarPathTracingMaxPathIntensity(
 	TEXT("r.PathTracing.MaxPathIntensity"),
 	-1,
-	TEXT("When positive, light paths greater that this amount are clamped to prevent fireflies (default = -1 (off))"),
+	TEXT("When positive, light paths greater that this amount are clamped to prevent fireflies (default = -1 (driven by postprocesing volume))"),
 	ECVF_RenderThreadSafe
 );
 
@@ -93,7 +93,7 @@ TAutoConsoleVariable<int32> CVarPathTracingApproximateCaustics(
 TAutoConsoleVariable<int32> CVarPathTracingEnableEmissive(
 	TEXT("r.PathTracing.EnableEmissive"),
 	-1,
-	TEXT("Indicates if emissive materials should contribute to scene lighting (default = -1 (inherit from PostProcessVolume)"),
+	TEXT("Indicates if emissive materials should contribute to scene lighting (default = -1 (driven by postprocesing volume)"),
 	ECVF_RenderThreadSafe
 );
 
@@ -119,11 +119,13 @@ TAutoConsoleVariable<int32> CVarPathTracingSamplerType(
 	PATHTRACER_SAMPLER_DEFAULT,
 	TEXT("Controls the way the path tracer generates its random numbers\n")
 	TEXT("0: use a different high quality random sequence per pixel\n")
-	TEXT("1: optimize the random sequence across pixels to lower error at the target sample count\n")
+	TEXT("1: optimize the random sequence across pixels to reduce visible error at the target sample count\n")
 	TEXT("2: share random seeds across pixels to improve coherence of execution on the GPU. This trades some correlation across the image in exchange for better performance.\n"),
 	ECVF_RenderThreadSafe
 );
 
+#if 0
+// TODO: re-enable this when multi-gpu is supported again
 // r.PathTracing.GPUCount is read only because ComputeViewGPUMasks results cannot change after UE has been launched
 TAutoConsoleVariable<int32> CVarPathTracingGPUCount(
 	TEXT("r.PathTracing.GPUCount"),
@@ -131,6 +133,7 @@ TAutoConsoleVariable<int32> CVarPathTracingGPUCount(
 	TEXT("Sets the amount of GPUs used for computing the path tracing pass (default = 1 GPU)"),
 	ECVF_RenderThreadSafe | ECVF_ReadOnly
 );
+#endif
 
 TAutoConsoleVariable<int32> CVarPathTracingWiperMode(
 	TEXT("r.PathTracing.WiperMode"),
@@ -158,7 +161,7 @@ TAutoConsoleVariable<int32> CVarPathTracingLightGridResolution(
 TAutoConsoleVariable<int32> CVarPathTracingLightGridMaxCount(
 	TEXT("r.PathTracing.LightGridMaxCount"),
 	128,
-	TEXT("Controls the maximum number of lights per cell in the 2D light grid (default = 128)\n"),
+	TEXT("Controls the maximum number of lights per cell in the 2D light grid. The minimum of this value and the number of lights in the scene is used. (default = 128)\n"),
 	ECVF_RenderThreadSafe
 );
 
@@ -230,8 +233,8 @@ static bool PrepareShaderArgs(const FViewInfo& View, FPathTracingData& PathTraci
 	PathTracingData.MaxPathIntensity = CVarPathTracingMaxPathIntensity.GetValueOnRenderThread();
 	if (PathTracingData.MaxPathIntensity <= 0)
 	{
-		// no clamping required - set to an impossible value
-		PathTracingData.MaxPathIntensity = std::numeric_limits<float>::infinity();
+		// cvar clamp disabled, use PPV exposure value instad
+		PathTracingData.MaxPathIntensity = FMath::Pow(2.0f, View.FinalPostProcessSettings.PathTracingMaxPathExposure);
 	}
 	PathTracingData.ApproximateCaustics = CVarPathTracingApproximateCaustics.GetValueOnRenderThread();
 	PathTracingData.EnableCameraBackfaceCulling = CVarPathTracingEnableCameraBackfaceCulling.GetValueOnRenderThread();
@@ -760,8 +763,7 @@ void PrepareLightGrid(FRDGBuilder& GraphBuilder, FPathTracingLightGrid* LightGri
 		LightGridParameters->LightGrid = GraphBuilder.RegisterExternalTexture(GSystemTextures.BlackDummy);
 		FRDGBufferDesc LightGridDataDesc = FRDGBufferDesc::CreateBufferDesc(sizeof(uint32), 1);
 		FRDGBuffer* LightGridData = GraphBuilder.CreateBuffer(LightGridDataDesc, TEXT("PathTracer.LightGridData"));
-		uint32 Dummy = 0;
-		AddBufferUploadPass(GraphBuilder, LightGridData, &Dummy, sizeof(uint32));
+		AddClearUAVPass(GraphBuilder, GraphBuilder.CreateUAV(LightGridData, PF_R32_UINT), 0);
 		LightGridParameters->LightGridData = GraphBuilder.CreateSRV(LightGridData, PF_R32_UINT);
 	}
 }
@@ -771,16 +773,17 @@ void SetLightParameters(FRDGBuilder& GraphBuilder, FPathTracingRG::FParameters*
 	PassParameters->SceneVisibleLightCount = 0;
 
 	// Lights
-	TArray<FPathTracingLight> Lights;
-	Lights.Reserve(1 + Scene->Lights.Num());
+	uint32 MaxNumLights = 1 + Scene->Lights.Num(); // upper bound
+	// Allocate from the graph builder so that we don't need to copy the data again when queuing the upload
+	FPathTracingLight* Lights = (FPathTracingLight*) GraphBuilder.Alloc(sizeof(FPathTracingLight) * MaxNumLights, 16);
+	uint32 NumLights = 0;
 
 	// Prepend SkyLight to light buffer since it is not part of the regular light list
 	const float Inf = std::numeric_limits<float>::infinity();
 	if (PrepareSkyTexture(GraphBuilder, Scene, View, UseMISCompensation, PassParameters))
 	{
 		check(Scene->SkyLight != nullptr);
-		Lights.AddDefaulted();
-		FPathTracingLight& DestLight = Lights.Last();
+		FPathTracingLight& DestLight = Lights[NumLights++];
 		DestLight.Color = FVector(1, 1, 1); // not used (it is folded into the importance table directly)
 		DestLight.Flags = Scene->SkyLight->bTransmission ? PATHTRACER_FLAG_TRANSMISSION_MASK : 0;
 		DestLight.Flags |= PATHTRACER_FLAG_LIGHTING_CHANNEL_MASK;
@@ -817,8 +820,7 @@ void SetLightParameters(FRDGBuilder& GraphBuilder, FPathTracingRG::FParameters*
 				continue;
 			}
 
-			Lights.AddDefaulted();
-			FPathTracingLight& DestLight = Lights.Last();
+			FPathTracingLight& DestLight = Lights[NumLights++];
 			uint32 Transmission = Light.LightSceneInfo->Proxy->Transmission();
 			uint8 LightingChannelMask = Light.LightSceneInfo->Proxy->GetLightingChannelMask();
 
@@ -846,7 +848,7 @@ void SetLightParameters(FRDGBuilder& GraphBuilder, FPathTracingRG::FParameters*
 		}
 	}
 
-	uint32 NumInfiniteLights = Lights.Num();
+	uint32 NumInfiniteLights = NumLights;
 
 	int32 NextRectTextureIndex = 0;
 
@@ -872,8 +874,7 @@ void SetLightParameters(FRDGBuilder& GraphBuilder, FPathTracingRG::FParameters*
 			continue;
 		}
 
-		Lights.AddDefaulted();
-		FPathTracingLight& DestLight = Lights.Last();
+		FPathTracingLight& DestLight = Lights[NumLights++];
 
 		uint32 Transmission = Light.LightSceneInfo->Proxy->Transmission();
 		uint8 LightingChannelMask = Light.LightSceneInfo->Proxy->GetLightingChannelMask();
@@ -1016,16 +1017,12 @@ void SetLightParameters(FRDGBuilder& GraphBuilder, FPathTracingRG::FParameters*
 		PassParameters->RectLightSampler = TStaticSamplerState<SF_Bilinear, AM_Clamp, AM_Clamp, AM_Clamp>::GetRHI();
 	}
 
-	PassParameters->SceneLightCount = Lights.Num();
+	PassParameters->SceneLightCount = NumLights;
 	{
-		// Buffer creation needs at least one element
-		if (Lights.IsEmpty())
-		{
-			Lights.AddDefaulted();
-		}
 		// Upload the buffer of lights to the GPU
-		size_t DataSize = sizeof(FPathTracingLight) * Lights.Num();
-		PassParameters->SceneLights = GraphBuilder.CreateSRV(FRDGBufferSRVDesc(CreateStructuredBuffer(GraphBuilder, TEXT("PathTracer.LightsBuffer"), sizeof(FPathTracingLight), Lights.Num(), Lights.GetData(), DataSize)));
+		uint32 NumCopyLights = FMath::Max(1u, NumLights); // need at least one since zero-sized buffers are not allowed
+		size_t DataSize = sizeof(FPathTracingLight) * NumCopyLights;
+		PassParameters->SceneLights = GraphBuilder.CreateSRV(FRDGBufferSRVDesc(CreateStructuredBuffer(GraphBuilder, TEXT("PathTracer.LightsBuffer"), sizeof(FPathTracingLight), NumCopyLights, Lights, DataSize, ERDGInitialDataFlags::NoCopy)));
 	}
 
 	if (CVarPathTracingVisibleLights.GetValueOnRenderThread() != 0)
@@ -1082,7 +1079,7 @@ void SetLightParameters(FRDGBuilder& GraphBuilder, FPathTracingRG::FParameters*
 	}
 
 
-	PrepareLightGrid(GraphBuilder, &PassParameters->LightGridParameters, Lights.GetData(), Lights.Num(), NumInfiniteLights, PassParameters->SceneLights, View);
+	PrepareLightGrid(GraphBuilder, &PassParameters->LightGridParameters, Lights, NumLights, NumInfiniteLights, PassParameters->SceneLights, View);
 }
 
 class FPathTracingCompositorPS : public FGlobalShader
