Fix missing atmospheric attenuation applied on hair.

diff --git a/Engine/Shaders/Private/DeferredLightPixelShaders.usf b/Engine/Shaders/Private/DeferredLightPixelShaders.usf
--- a/Engine/Shaders/Private/DeferredLightPixelShaders.usf
+++ b/Engine/Shaders/Private/DeferredLightPixelShaders.usf
@@ -136,6 +136,31 @@ float3 GetAtmosphericLightTransmittance(in float3 StableLargeWorldPosition, in f
 
 #endif // USE_ATMOSPHERE_TRANSMITTANCE
 
+float3 GetAtmosphereAndCloudAttenuation(float3 StableLargeWorldPosition, float2 ScreenUV)
+{
+	float3 Out = 1;
+
+	#if USE_ATMOSPHERE_TRANSMITTANCE
+	Out *= GetAtmosphericLightTransmittance(StableLargeWorldPosition, ScreenUV, DeferredLightUniforms.Direction.xyz);
+	#endif
+
+	#if USE_CLOUD_TRANSMITTANCE
+	float OutOpticalDepth = 0.0f;
+	Out *= lerp(1.0f, GetCloudVolumetricShadow(StableLargeWorldPosition, CloudShadowmapWorldToLightClipMatrix, CloudShadowmapFarDepthKm, CloudShadowmapTexture, CloudShadowmapSampler, OutOpticalDepth), CloudShadowmapStrength);
+	#endif
+
+	return Out;
+}
+
+float3 GetStableLargeWorldPosition(float2 SVPos, float DeviceZ)
+{
+	float3 StableLargeWorldPosition = 0;
+	#if USE_ATMOSPHERE_TRANSMITTANCE || USE_CLOUD_TRANSMITTANCE
+	StableLargeWorldPosition = SvPositionToWorld(float4(SVPos, DeviceZ, 1.0));
+	#endif
+	return StableLargeWorldPosition;
+}
+
 ////////////////////////////////////////////////////////////////////////////////////////////////////
 
 #if STRATA_ENABLED
@@ -276,20 +301,7 @@ void DeferredLightPixelMain(
 		const float  Attenuation = ComputeLightProfileMultiplier(DerivedParams.WorldPosition, DeferredLightUniforms.Position, -DeferredLightUniforms.Direction, DeferredLightUniforms.Tangent);
 
 		OutColor += (Radiance * Attenuation);
-
-	#if USE_ATMOSPHERE_TRANSMITTANCE || USE_CLOUD_TRANSMITTANCE
-		float DeviceZ = LookupDeviceZ(ScreenUV);
-		float3 StableLargeWorldPosition = SvPositionToWorld(float4(SVPos.xy, DeviceZ, 1.0));
-	#endif
-
-	#if USE_ATMOSPHERE_TRANSMITTANCE
-		OutColor.rgb *= GetAtmosphericLightTransmittance(StableLargeWorldPosition, InputParams.ScreenUV, DeferredLightUniforms.Direction.xyz);
-	#endif
-
-	#if USE_CLOUD_TRANSMITTANCE
-		float OutOpticalDepth = 0.0f;
-		OutColor.rgb *= lerp(1.0f, GetCloudVolumetricShadow(StableLargeWorldPosition, CloudShadowmapWorldToLightClipMatrix, CloudShadowmapFarDepthKm, CloudShadowmapTexture, CloudShadowmapSampler, OutOpticalDepth), CloudShadowmapStrength);
-	#endif
+		OutColor.rgb *= GetAtmosphereAndCloudAttenuation(GetStableLargeWorldPosition(InputParams.PixelPos, LookupDeviceZ(InputParams.ScreenUV)), InputParams.ScreenUV);
 	}
 
 	// RGB:SceneColor Specular and Diffuse
@@ -393,11 +405,11 @@ void DeferredLightPixelMain(
 	float SurfaceShadow = 1.0f;
 	const float4 Radiance = GetDynamicLighting(WorldPosition, CameraVector, HairScreenSpaceData.GBuffer, HairScreenSpaceData.AmbientOcclusion, HairScreenSpaceData.GBuffer.ShadingModelID, LightData, LightAttenuation, Dither, PixelCoord, RectTexture, SurfaceShadow);
 	const float  Attenuation = ComputeLightProfileMultiplier(WorldPosition, DeferredLightUniforms.Position, -DeferredLightUniforms.Direction, DeferredLightUniforms.Tangent);
+	const float3 AtmosphericTransmittance = GetAtmosphereAndCloudAttenuation(GetStableLargeWorldPosition(ScreenPosition, Sample.Depth), UV);
 
 	const float LocalCoverage = From8bitCoverage(Sample.Coverage8bit);
 
-	// TODO: Add atomspheric contribution
-	OutColor.rgb = Radiance.xyz * Attenuation * LocalCoverage * GetExposure();
+	OutColor.rgb = AtmosphericTransmittance * Radiance.xyz * Attenuation * LocalCoverage * GetExposure();
 	OutColor.a = LocalCoverage;
 }
 #endif
diff --git a/Engine/Source/Runtime/Renderer/Private/LightRendering.cpp b/Engine/Source/Runtime/Renderer/Private/LightRendering.cpp
--- a/Engine/Source/Runtime/Renderer/Private/LightRendering.cpp
+++ b/Engine/Source/Runtime/Renderer/Private/LightRendering.cpp
@@ -2106,6 +2106,44 @@ void CalculateLightNearFarDepthFromBounds(const FViewInfo& View, const FSphere &
 
 }
 
+static void BindAtmosphereAndCloudResources(
+	const FScene* Scene,
+	const FViewInfo* View,
+	const FLightSceneProxy* Proxy,
+	FRenderLightParams& RenderLightParams,
+	bool& bAtmospherePerPixelTransmittance,
+	bool& bCloudPerPixelTransmittance)
+{
+	bAtmospherePerPixelTransmittance =
+		Proxy->GetLightType() == LightType_Directional &&
+		Proxy->IsUsedAsAtmosphereSunLight() &&
+		Proxy->GetUsePerPixelAtmosphereTransmittance() &&
+		ShouldRenderSkyAtmosphere(Scene, View->Family->EngineShowFlags);
+
+	const FLightSceneProxy* AtmosphereLight0Proxy = Scene->AtmosphereLights[0] ? Scene->AtmosphereLights[0]->Proxy : nullptr;
+	const FLightSceneProxy* AtmosphereLight1Proxy = Scene->AtmosphereLights[1] ? Scene->AtmosphereLights[1]->Proxy : nullptr;
+	const FVolumetricCloudRenderSceneInfo* CloudInfo = Scene->GetVolumetricCloudSceneInfo();
+	const bool VolumetricCloudShadowMap0Valid = View->VolumetricCloudShadowExtractedRenderTarget[0] != nullptr;
+	const bool VolumetricCloudShadowMap1Valid = View->VolumetricCloudShadowExtractedRenderTarget[1] != nullptr;
+	const bool bLight0CloudPerPixelTransmittance = CloudInfo && VolumetricCloudShadowMap0Valid && AtmosphereLight0Proxy == Proxy && AtmosphereLight0Proxy && AtmosphereLight0Proxy->GetCloudShadowOnSurfaceStrength() > 0.0f;
+	const bool bLight1CloudPerPixelTransmittance = CloudInfo && VolumetricCloudShadowMap1Valid && AtmosphereLight1Proxy == Proxy && AtmosphereLight1Proxy && AtmosphereLight1Proxy->GetCloudShadowOnSurfaceStrength() > 0.0f;
+	if (bLight0CloudPerPixelTransmittance)
+	{
+		RenderLightParams.Cloud_ShadowmapTexture = View->VolumetricCloudShadowExtractedRenderTarget[0]->GetShaderResourceRHI();
+		RenderLightParams.Cloud_ShadowmapFarDepthKm = CloudInfo->GetVolumetricCloudCommonShaderParameters().CloudShadowmapFarDepthKm[0].X;
+		RenderLightParams.Cloud_WorldToLightClipShadowMatrix = CloudInfo->GetVolumetricCloudCommonShaderParameters().CloudShadowmapWorldToLightClipMatrix[0];
+		RenderLightParams.Cloud_ShadowmapStrength = AtmosphereLight0Proxy->GetCloudShadowOnSurfaceStrength();
+	}
+	else if (bLight1CloudPerPixelTransmittance)
+	{
+		RenderLightParams.Cloud_ShadowmapTexture = View->VolumetricCloudShadowExtractedRenderTarget[1]->GetShaderResourceRHI();
+		RenderLightParams.Cloud_ShadowmapFarDepthKm = CloudInfo->GetVolumetricCloudCommonShaderParameters().CloudShadowmapFarDepthKm[1].X;
+		RenderLightParams.Cloud_WorldToLightClipShadowMatrix = CloudInfo->GetVolumetricCloudCommonShaderParameters().CloudShadowmapWorldToLightClipMatrix[1];
+		RenderLightParams.Cloud_ShadowmapStrength = AtmosphereLight1Proxy->GetCloudShadowOnSurfaceStrength();
+	}
+	bCloudPerPixelTransmittance = bLight0CloudPerPixelTransmittance || bLight1CloudPerPixelTransmittance;
+}
+
 /**
  * Used by RenderLights to render a light to the scene color buffer.
  *
@@ -2196,30 +2234,15 @@ void FDeferredShadingSceneRenderer::RenderLight(
 		}
 		else
 		{
-			const bool bAtmospherePerPixelTransmittance = LightSceneInfo->Proxy->IsUsedAsAtmosphereSunLight() 
-				&& LightSceneInfo->Proxy->GetUsePerPixelAtmosphereTransmittance() && ShouldRenderSkyAtmosphere(Scene, View.Family->EngineShowFlags);
-
-			FLightSceneProxy* AtmosphereLight0Proxy = Scene->AtmosphereLights[0] ? Scene->AtmosphereLights[0]->Proxy : nullptr;
-			FLightSceneProxy* AtmosphereLight1Proxy = Scene->AtmosphereLights[1] ? Scene->AtmosphereLights[1]->Proxy : nullptr;
-			FVolumetricCloudRenderSceneInfo* CloudInfo = Scene->GetVolumetricCloudSceneInfo();
-			const bool VolumetricCloudShadowMap0Valid = View.VolumetricCloudShadowExtractedRenderTarget[0] != nullptr;
-			const bool VolumetricCloudShadowMap1Valid = View.VolumetricCloudShadowExtractedRenderTarget[1] != nullptr;
-			const bool bLight0CloudPerPixelTransmittance = CloudInfo && VolumetricCloudShadowMap0Valid && AtmosphereLight0Proxy == LightSceneInfo->Proxy && AtmosphereLight0Proxy && AtmosphereLight0Proxy->GetCloudShadowOnSurfaceStrength() > 0.0f;
-			const bool bLight1CloudPerPixelTransmittance = CloudInfo && VolumetricCloudShadowMap1Valid && AtmosphereLight1Proxy == LightSceneInfo->Proxy && AtmosphereLight1Proxy && AtmosphereLight1Proxy->GetCloudShadowOnSurfaceStrength() > 0.0f;
-			if (bLight0CloudPerPixelTransmittance)
-			{
-				RenderLightParams.Cloud_ShadowmapTexture = View.VolumetricCloudShadowExtractedRenderTarget[0]->GetShaderResourceRHI();
-				RenderLightParams.Cloud_ShadowmapFarDepthKm = CloudInfo->GetVolumetricCloudCommonShaderParameters().CloudShadowmapFarDepthKm[0].X;
-				RenderLightParams.Cloud_WorldToLightClipShadowMatrix = CloudInfo->GetVolumetricCloudCommonShaderParameters().CloudShadowmapWorldToLightClipMatrix[0];
-				RenderLightParams.Cloud_ShadowmapStrength = AtmosphereLight0Proxy->GetCloudShadowOnSurfaceStrength();
-			}
-			else if(bLight1CloudPerPixelTransmittance)
-			{
-				RenderLightParams.Cloud_ShadowmapTexture = View.VolumetricCloudShadowExtractedRenderTarget[1]->GetShaderResourceRHI();
-				RenderLightParams.Cloud_ShadowmapFarDepthKm = CloudInfo->GetVolumetricCloudCommonShaderParameters().CloudShadowmapFarDepthKm[1].X;
-				RenderLightParams.Cloud_WorldToLightClipShadowMatrix = CloudInfo->GetVolumetricCloudCommonShaderParameters().CloudShadowmapWorldToLightClipMatrix[1];
-				RenderLightParams.Cloud_ShadowmapStrength = AtmosphereLight1Proxy->GetCloudShadowOnSurfaceStrength();
-			}
+			bool bAtmospherePerPixelTransmittance = false;
+			bool bCloudPerPixelTransmittance = false;
+			BindAtmosphereAndCloudResources(
+				Scene,
+				&View,
+				LightSceneInfo->Proxy,
+				RenderLightParams,
+				bAtmospherePerPixelTransmittance,
+				bCloudPerPixelTransmittance);
 
 			FDeferredLightPS::FPermutationDomain PermutationVector;
 			PermutationVector.Set< FDeferredLightPS::FSourceShapeDim >( ELightSourceShape::Directional );
@@ -2232,7 +2255,7 @@ void FDeferredShadingSceneRenderer::RenderLight(
 			PermutationVector.Set< FDeferredLightPS::FHairLighting>(0);
 			// Only directional lights are rendered in this path, so we only need to check if it is use to light the atmosphere
 			PermutationVector.Set< FDeferredLightPS::FAtmosphereTransmittance >(bAtmospherePerPixelTransmittance);
-			PermutationVector.Set< FDeferredLightPS::FCloudTransmittance >(bLight0CloudPerPixelTransmittance || bLight1CloudPerPixelTransmittance);
+			PermutationVector.Set< FDeferredLightPS::FCloudTransmittance >(bCloudPerPixelTransmittance);
 			PermutationVector.Set< FDeferredLightPS::FStrataFastPath >(StrataTileMaterialType == EStrataTileMaterialType::ESimple);
 
 			TShaderMapRef< FDeferredLightPS > PixelShader( View.ShaderMap, PermutationVector );
@@ -2495,7 +2518,7 @@ void FDeferredShadingSceneRenderer::RenderLightForHair(
 			{},
 			PassParameters,
 			ERDGPassFlags::Raster,
-			[&HairVisibilityData, &View, PassParameters, LightSceneInfo, MaxTransmittanceElementCount, HairShadowMaskTexture, LightingChannelsTexture, bIsShadowMaskValid](FRHICommandList& RHICmdList)
+			[this, &HairVisibilityData, &View, PassParameters, LightSceneInfo, MaxTransmittanceElementCount, HairShadowMaskTexture, LightingChannelsTexture, bIsShadowMaskValid](FRHICommandList& RHICmdList)
 		{
 			RHICmdList.SetViewport(0, 0, 0.0f, HairVisibilityData.SampleLightingViewportResolution.X, HairVisibilityData.SampleLightingViewportResolution.Y, 1.0f);
 
@@ -2512,10 +2535,22 @@ void FDeferredShadingSceneRenderer::RenderLightForHair(
 			FDeferredLightPS::FPermutationDomain PermutationVector;
 			if (LightSceneInfo->Proxy->GetLightType() == LightType_Directional)
 			{
+				bool bAtmospherePerPixelTransmittance = false;
+				bool bCloudPerPixelTransmittance = false;
+				BindAtmosphereAndCloudResources(
+					Scene,
+					&View,
+					LightSceneInfo->Proxy,
+					RenderLightParams,
+					bAtmospherePerPixelTransmittance,
+					bCloudPerPixelTransmittance);
+
 				PermutationVector.Set< FDeferredLightPS::FSourceShapeDim >(ELightSourceShape::Directional);
 				PermutationVector.Set< FDeferredLightPS::FSourceTextureDim >(false);
 				PermutationVector.Set< FDeferredLightPS::FIESProfileDim >(false);
 				PermutationVector.Set< FDeferredLightPS::FInverseSquaredDim >(false);
+				PermutationVector.Set< FDeferredLightPS::FAtmosphereTransmittance >(bAtmospherePerPixelTransmittance);
+				PermutationVector.Set< FDeferredLightPS::FCloudTransmittance >(bCloudPerPixelTransmittance);
 			}
 			else
 			{
@@ -2524,6 +2559,8 @@ void FDeferredShadingSceneRenderer::RenderLightForHair(
 				PermutationVector.Set< FDeferredLightPS::FSourceTextureDim >(LightSceneInfo->Proxy->IsRectLight() && LightSceneInfo->Proxy->HasSourceTexture());
 				PermutationVector.Set< FDeferredLightPS::FIESProfileDim >(bUseIESTexture);
 				PermutationVector.Set< FDeferredLightPS::FInverseSquaredDim >(LightSceneInfo->Proxy->IsInverseSquared());
+				PermutationVector.Set< FDeferredLightPS::FAtmosphereTransmittance >(false);
+				PermutationVector.Set< FDeferredLightPS::FCloudTransmittance >(false);
 			}
 			PermutationVector.Set< FDeferredLightPS::FLightingChannelsDim >(View.bUsesLightingChannels);
 			PermutationVector.Set< FDeferredLightPS::FVisualizeCullingDim >(false);
