Prepare code for removal of EDecalBlendMode.
Clean up logic for render stages.
Removed unneccessary GBufferNormal only target mode.
Remove volumetric decal.

diff --git a/Engine/Shaders/Private/DBufferDecalShared.ush b/Engine/Shaders/Private/DBufferDecalShared.ush
--- a/Engine/Shaders/Private/DBufferDecalShared.ush
+++ b/Engine/Shaders/Private/DBufferDecalShared.ush
@@ -42,19 +42,6 @@ void EncodeDBufferData(FGBufferData GBufferData, float3 MultiOpacity,
 
 	// UNORM 4 channel
 	DBufferC = float4(GBufferData.Metallic, GBufferData.Specular, GBufferData.Roughness, MultiOpacity.z);
-
-	// no need to output any data if the channel isn't used anyway (for DX10 we don't adjust the FrameBufferBlend state - it doesn't support different blend modes for different MRT)
-	{
-#if !MATERIAL_DBUFFERA
-		DBufferA = 0;
-#endif
-#if !MATERIAL_DBUFFERB
-		DBufferB = 0;
-#endif
-#if !MATERIAL_DBUFFERC
-		DBufferC = 0;
-#endif
-	}
 }
 
 /** Populates FDBufferData */
diff --git a/Engine/Shaders/Private/DecalCommon.ush b/Engine/Shaders/Private/DecalCommon.ush
--- a/Engine/Shaders/Private/DecalCommon.ush
+++ b/Engine/Shaders/Private/DecalCommon.ush
@@ -1,9 +1,13 @@
 // Copyright Epic Games, Inc. All Rights Reserved.
 
 /*=============================================================================
-	DecalCommon.usf: decal shared things.
+	DecalCommon.ush: Common decal output logic.
 =============================================================================*/
 
+// DECAL_RENDERSTAGE is set by C++ from EDecalRenderStage
+// DECAL_RENDERTARGETMODE is set by C++ from EDecalRenderTarget
+// DECAL_RENDERTARGET_COUNT is set by C++
+
 #include "Common.ush"
 
 #if SHADING_PATH_MOBILE
@@ -11,90 +15,87 @@
 	#define MobileSceneTextures MobileBasePass.SceneTextures
 #endif
 
-#include "DeferredShadingCommon.ush"		// GBufferData
+#include "DeferredShadingCommon.ush"
 #include "DBufferDecalShared.ush"
 
-half3 FrameBufferDecalBlendOp(half4 Source)
-{
-	half4 Dest = half4(0, 0, 0, 0);
-    
-#if DECAL_BLEND_MODE == DECALBLENDMODEID_STAIN
-	return (Source.rgb*Dest.rgb) + (Dest.rgb*(1.0 - Source.a));
-#elif DECAL_BLEND_MODE == DECALBLENDMODEID_TRANSLUCENT
-	return (Source.rgb*Source.a) + (Dest.rgb*(1.0 - Source.a));
-#elif DECAL_BLEND_MODE == DECALBLENDMODEID_EMISSIVE
-	return (Source.rgb*Source.a) + Dest.rgb;
-#elif DECAL_BLEND_MODE == DECALBLENDMODEID_ALPHACOMPOSITE
-	return Source.rgb + (Dest.rgb*(1.0 - Source.a));
-#else
-	return Source.rgb;
-#endif
-}
-
 void DecalCommonOutput(inout FPixelShaderIn In, inout FPixelShaderOut Out, float3 Color, float Opacity, FGBufferData Data)
 {
-	// RETURN_COLOR not needed unless writing to SceneColor
-	Out.MRT[0] = float4(Color, Opacity);
+#if DECAL_RENDERSTAGE == DECAL_RENDERSTAGE_MOBILE
+	// Cheap mobile forward approximation uses both emissive and color.
+	Color += Data.BaseColor;
+	#if MATERIALBLENDING_MODULATE
+		Color *= Opacity;
+	#endif
+#elif DECAL_RENDERSTAGE == DECAL_RENDERSTAGE_AO
+	// Replace color with AO for AO decals.
+	Color = lerp(1, Data.GBufferAO, Opacity);
+#endif
 
 #if USE_PREEXPOSURE
-	// Blend mode normal should not mult with exposure because Out.MRT[0] is world space normal in that case.
-	#if DECAL_BLEND_MODE != DECALBLENDMODEID_NORMAL
-		Out.MRT[0].xyz *= View.PreExposure;
+	#if DECAL_RENDERTARGETMODE == DECAL_RENDERTARGETMODE_SCENECOLOR || DECAL_RENDERTARGET == DECAL_RENDERTARGETMODE_GBUFFER || DECAL_RENDERTARGET == DECAL_RENDERTARGETMODE_GBUFFER_NONORMAL
+		Color *= View.PreExposure;
 	#endif
 #endif
 
-#if DECAL_RENDERTARGET_COUNT > 1
-	// some MRT rendering
-
-	// Normal may not be updated by decals so we don't need to update that render target
-	float4 OutTarget1 = 0;
-
-	// static lighting isn't updated by decals so we don't need to update that render target
-	float4 OutTarget4 = 0;
-	float4 OutTarget5 = 0;
-	float4 OutTarget6 = 0;
+	Out.MRT[0] = float4(Color, Opacity);
 
-	#if DECAL_RENDERSTAGE == 0
+#if DECAL_RENDERTARGETMODE == DECAL_RENDERTARGETMODE_DBUFFER
 	{
-		// before base pass (DBuffer)
-
-		// @param MultiOpacity .x: Color, .y:Normal, .z:Roughness
 		float3 MultiOpacity = Opacity;
 
 		EncodeDBufferData(Data, MultiOpacity, Out.MRT[0], Out.MRT[1], Out.MRT[2]);
 
 		#if DECAL_RENDERTARGET_COUNT == 4
-		Out.MRT[3] = 1.0f; // DBufferMask
+			Out.MRT[3] = 1.0f; // DBufferMask
 		#endif
 	}
-	#elif DECAL_RENDERSTAGE == 1 || DECAL_RENDERSTAGE == 2 || DECAL_RENDERSTAGE == 3
-	{
-		// 1 == after base pass (GBuffer) || 2 = before lighting (GBuffer)
+#endif
 
-		#if DECAL_OUTPUT_NORMAL
-			EncodeGBuffer(Data, Out.MRT[1], Out.MRT[2], Out.MRT[3], OutTarget4, OutTarget5, OutTarget6);
-		#else
-			EncodeGBuffer(Data, OutTarget1, Out.MRT[1], Out.MRT[2], OutTarget4, OutTarget5, OutTarget6);
-		#endif
-	}
-	#endif // DECAL_RENDERSTAGE
+#if DECAL_RENDERTARGETMODE == DECAL_RENDERTARGETMODE_GBUFFER || DECAL_RENDERTARGETMODE == DECAL_RENDERTARGETMODE_GBUFFER_NONORMAL
+	
+	// Dummy outputs from EncodeGBuffer()
+	float4 OutTarget1 = 0;
+	float4 OutTarget4 = 0;
+	float4 OutTarget5 = 0;
+	float4 OutTarget6 = 0;
 
-	#if DECAL_RENDERSTAGE == 2 || DECAL_RENDERSTAGE == 3
+	#if DECAL_RENDERTARGETMODE == DECAL_RENDERTARGETMODE_GBUFFER
+	{
+		EncodeGBuffer(Data, Out.MRT[1], Out.MRT[2], Out.MRT[3], OutTarget4, OutTarget5, OutTarget6);
+	}
+	#elif DECAL_RENDERTARGETMODE == DECAL_RENDERTARGETMODE_GBUFFER_NONORMAL
 	{
-		#if DECAL_BLEND_MODE == DECALBLENDMODEID_STAIN
-			// BaseColor MRT location changes if we have Normal MRT or not.
-			#if DECAL_OUTPUT_NORMAL
-				Out.MRT[3].rgb *= Opacity;
-			#else
-				Out.MRT[2].rgb *= Opacity;
-			#endif
+		EncodeGBuffer(Data, OutTarget1, Out.MRT[1], Out.MRT[2], OutTarget4, OutTarget5, OutTarget6);
+	}
+	#endif // DECAL_RENDERTARGETMODE
+
+	#if MATERIALBLENDING_MODULATE
+		// BaseColor MRT location changes for render target mode.
+		#if DECAL_RENDERTARGETMODE == DECAL_RENDERTARGETMODE_GBUFFER
+			Out.MRT[3].rgb *= Opacity;
+		#elif DECAL_RENDERTARGETMODE == DECAL_RENDERTARGETMODE_GBUFFER_NONORMAL
+			Out.MRT[2].rgb *= Opacity;
 		#endif
+	#endif // MATERIALBLENDING_MODULATE
 
-		Out.MRT[0].a = Opacity;
-		Out.MRT[1].a = Opacity;
-		Out.MRT[2].a = Opacity;
-		Out.MRT[3].a = Opacity;
-	}
-	#endif // DECAL_RENDERSTAGE
-#endif //DECAL_RENDERTARGET_COUNT > 1
+	Out.MRT[1].a = Opacity;
+	Out.MRT[2].a = Opacity;
+	Out.MRT[3].a = Opacity;
+
+#endif // DECAL_RENDERTARGETMODE
+
+	// Allow shader compiler to cull outputs that we won't use.
+	// Blend state also prevents writing unused outputs, but we want to be sure that we also avoid shader cost.
+#if !DECAL_OUT_MRT0
+	Out.MRT[0] = 0;
+#endif
+#if !DECAL_OUT_MRT1
+	Out.MRT[1] = 0;
+#endif
+#if !DECAL_OUT_MRT2
+	Out.MRT[2] = 0;
+#endif
+#if !DECAL_OUT_MRT3
+	Out.MRT[3] = 0;
+#endif
 }
diff --git a/Engine/Shaders/Private/DeferredDecal.usf b/Engine/Shaders/Private/DeferredDecal.usf
--- a/Engine/Shaders/Private/DeferredDecal.usf
+++ b/Engine/Shaders/Private/DeferredDecal.usf
@@ -4,10 +4,6 @@
 	DeferredDecal.usf: Pixel shader for computing a deferred decal.
 =============================================================================*/
 
-// DECAL_BLEND_MODE is set by C++ from EDecalBlendMode e.g. DECALBLENDMODEID_VOLUMETRIC, DECALBLENDMODEID_NORMAL, ...
-// DECAL_RENDERTARGET_COUNT is set by C++ e.g. 1: single RT, 2: two MRT, 3: three MRT
-// DECAL_RENDERSTAGE is set by C++ from EDecalRenderStage e.g. 0:before base pass, 1:after base pass, 2: before lighting
-
 #include "DecalCommon.ush"
 
 // from SvPosition to decal space (for position), used like this: mul(float4(SvPosition.xyz, 1), SvPositionToDecal);
@@ -18,18 +14,26 @@ float4x4 DecalToWorld;
 float4x4 WorldToDecal;
 // decal orientation in world space, X axis. Used for ObjectOrientation material expression
 float3 DecalOrientation;
-
-// DECAL_PRIMITIVE informs material templates which functions to expose when rendering decals.
-#define DECAL_PRIMITIVE 1
-
 // x - Fade alpha from screen size fading
 // y - Opacity over the lifetime of the decal. 1 -> 0
 float2 DecalParams;
+// from component to clip space (for decal frustum)
+float4x4 FrustumComponentToClip;
 
-#if DECAL_PROJECTION
-	#include "/Engine/Generated/Material.ush"
-#endif
+// Global decal vertex shader
+void MainVS(
+	in float4 InPosition : ATTRIBUTE0,
+	out float4 OutPosition : SV_POSITION
+)
+{
+	OutPosition = mul(InPosition, FrustumComponentToClip);
+}
 
+#if IS_DECAL
+
+// DECAL_PRIMITIVE informs material templates which functions to expose when rendering decals.
+#define DECAL_PRIMITIVE 1
+#include "/Engine/Generated/Material.ush"
 
 // If virtual texture is enabled then use early depth test so that UAV feedback buffer writes respect the depth test
 #if NUM_VIRTUALTEXTURE_SAMPLES
@@ -43,99 +47,6 @@ float2 DecalParams;
 	#endif
 #endif
 
-// from component to clip space (for decal frustum)
-float4x4 FrustumComponentToClip;
-
-// decal vertex shader
-void MainVS(
-	in float4 InPosition : ATTRIBUTE0,
-	out float4 OutPosition : SV_POSITION
-	)
-{
-	OutPosition = mul(InPosition, FrustumComponentToClip);
-}
-
-float DistanceFunction(float4 SvPosition, float3 Pos)
-{
-	float3 SwizzlePos = Pos.zyx;
-
-#if DECAL_PROJECTION && DECAL_BLEND_MODE == DECALBLENDMODEID_VOLUMETRIC
-	
-	FMaterialPixelParameters MaterialParameters = MakeInitializedMaterialPixelParameters();
-#if NUM_MATERIAL_TEXCOORDS
-	for(int CoordinateIndex = 0;CoordinateIndex < NUM_MATERIAL_TEXCOORDS;CoordinateIndex++)
-	{
-		MaterialParameters.TexCoords[CoordinateIndex] = SwizzlePos.xy * 0.5f + 0.5f;
-	}
-#endif
-
-	MaterialParameters.VertexColor = 1;
-	MaterialParameters.SvPosition = SvPosition;
-	MaterialParameters.ScreenPosition = SvPositionToScreenPosition(SvPosition);
-	MaterialParameters.LightVector = SwizzlePos * 0.5f + 0.5f;		// to be compatible with decals but *0.5+0.5 seems wrong
-	MaterialParameters.AbsoluteWorldPosition = 
-	MaterialParameters.WorldPosition_CamRelative = 
-	MaterialParameters.WorldPosition_NoOffsets = 
-	MaterialParameters.WorldPosition_NoOffsets_CamRelative = mul(DecalToWorld, float4(Pos, 1)).xyz;
-	MaterialParameters.CameraVector = normalize(View.WorldCameraOrigin - MaterialParameters.AbsoluteWorldPosition);
-
-	FPixelMaterialInputs PixelMaterialInputs;
-	CalcPixelMaterialInputs(MaterialParameters, PixelMaterialInputs);
-
-	// the material needs to output the distance from "Opacity" depending on LightVector (in 0..1 range)
-	return GetMaterialMaskInputRaw(PixelMaterialInputs);
-
-#endif
-
-	return 0;
-}
-
-// could be exposed to the user
-const static float coneSpread = 0.00002f;
-
-// @return t=0..1, -1: no hit
-float Raycast(float4 SvPosition, float3 StartPos, float3 EndPos)
-{
-	const uint steps = 300; 
-
-	float t = 0;
-
-	float RayLength = length(EndPos - StartPos);
-	float3 Direction = normalize(EndPos - StartPos);
-
-	LOOP for(uint i = 0; i < steps; ++i)
-	{
-		float Distance = DistanceFunction(SvPosition, StartPos + t * Direction);
-
-		if(Distance <= t * coneSpread)
-		{
-			return t / RayLength;
-		}
-
-		if(t > RayLength)
-		{
-		   break;
-		}
-
-		t += Distance;
-	}
-
-	return -1;
-}
-
-// like explained here http://iquilezles.org/www/articles/derivative/derivative.htm
-float3 GradientNotNormalized(float4 SvPosition, float3 OSPos, float t)
-{
-	float eps = t * coneSpread;
-
-	float2 e = float2(eps, 0);
-
-    return float3(DistanceFunction(SvPosition, OSPos + e.xyy) - DistanceFunction(SvPosition, OSPos - e.xyy),
-                  DistanceFunction(SvPosition, OSPos + e.yxy) - DistanceFunction(SvPosition, OSPos - e.yxy),
-                  DistanceFunction(SvPosition, OSPos + e.yyx) - DistanceFunction(SvPosition, OSPos - e.yyx));
-}
-
-
 float4 SvPositionToScreenPosition2(float4 SvPosition)
 {
 	// assumed SvPosition.w = 1
@@ -155,11 +66,6 @@ float4 SvPositionToScreenPosition2(float4 SvPosition)
 	return Ret;
 }
 
-// decal pixel shader
-#if DECAL_PROJECTION
-
-// EARLYDEPTHSTENCIL for volumetric decals could work here, but only if we are ouside of the decal
-
 // is called in MainPS() from PixelShaderOutputCommon.usf
 void FPixelShaderInOut_MainPS(inout FPixelShaderIn In, inout FPixelShaderOut Out)
 {
@@ -170,59 +76,16 @@ void FPixelShaderInOut_MainPS(inout FPixelShaderIn In, inout FPixelShaderOut Out
 	// make SvPosition appear to be rasterized with the depth from the depth buffer
 	In.SvPosition.z = LookupDeviceZ(ScreenUV);
 
-	const bool bVolumetric = DECAL_BLEND_MODE == DECALBLENDMODEID_VOLUMETRIC;
-
-	float3 OSPosition;
-	float3 WSPosition;
-	float HitT;
-
-	if(bVolumetric)
-	{
-		// LineBoxIntersect want's to clip a long ray so we provide one, this could be done differently as it limites the draw distance and reduces precision 
-		float SceneW = 100000000;
-
-		float4 SvPositionInDistance = float4(In.SvPosition.xy, ConvertToDeviceZ(SceneW), 1);
-
-		float3 WSStartRay = View.WorldCameraOrigin;
-		float3 WSEndRay = SvPositionToWorld(SvPositionInDistance);
+	float4 DecalVectorHom = mul(float4(In.SvPosition.xyz,1), SvPositionToDecal);
+	float3 OSPosition = DecalVectorHom.xyz / DecalVectorHom.w;
 
-		// in object/decal space
-		float4 OSStartRayHom = mul(float4(WSStartRay,1), WorldToDecal);
-		float4 OSEndRayHom = mul(float4(WSEndRay, 1), WorldToDecal);
+	// clip content outside the decal
+	// not needed if we stencil out the decal but we do that only on large (screen space) ones
+	clip(OSPosition.xyz + 1.0f);
+	clip(1.0f - OSPosition.xyz);
 
-		float3 OSStartRay = OSStartRayHom.xyz / OSStartRayHom.w;
-		float3 OSEndRay = OSEndRayHom.xyz / OSEndRayHom.w;
-
-		float2 OSBoxMinMax = LineBoxIntersect(OSStartRay, OSEndRay, -1, 1);
-
-		float3 OSBoxStartRay = OSStartRay + OSBoxMinMax.x * (OSEndRay - OSStartRay);
-		float3 OSBoxEndRay = OSStartRay + OSBoxMinMax.y * (OSEndRay - OSStartRay);
-
-		HitT = Raycast(In.SvPosition, OSBoxStartRay, OSBoxEndRay);
-
-		clip(HitT);
-
-		OSPosition = OSBoxStartRay + HitT * (OSBoxEndRay - OSBoxStartRay);
-		WSPosition = mul(float4(OSPosition, 1), DecalToWorld).xyz;
-
-		#if OUTPUT_PIXEL_DEPTH_OFFSET
-			float4 CSHitPos = mul(float4(WSPosition, 1), View.WorldToClip);
-			OutDepth = CSHitPos.z / CSHitPos.w;
-		#endif
-	}
-	else
-	{
-		float4 DecalVectorHom = mul(float4(In.SvPosition.xyz,1), SvPositionToDecal);
-		OSPosition = DecalVectorHom.xyz / DecalVectorHom.w;
-
-		// clip content outside the decal
-		// not needed if we stencil out the decal but we do that only on large (screen space) ones
-		clip(OSPosition.xyz + 1.0f);
-		clip(1.0f - OSPosition.xyz);
-
-		// todo: TranslatedWorld would be better for quality
-		WSPosition = SvPositionToWorld(In.SvPosition);
-	}
+	// todo: TranslatedWorld would be better for quality
+	float3 WSPosition = SvPositionToWorld(In.SvPosition);
 
 	float3 CameraVector = normalize(View.WorldCameraOrigin - WSPosition);
 
@@ -255,71 +118,21 @@ void FPixelShaderInOut_MainPS(inout FPixelShaderIn In, inout FPixelShaderOut Out
 	FPixelMaterialInputs PixelMaterialInputs;
 	CalcPixelMaterialInputs(MaterialParameters, PixelMaterialInputs);
 	
-	float DecalFading = saturate(4 - 4 * abs(SwizzlePos.z * 2 - 1)) * DecalParams.x;
-
-	#if DECAL_BLEND_MODE == DECALBLENDMODEID_AO
-	{
-		float AmbientOcclusion = lerp(1, GetMaterialAmbientOcclusion(PixelMaterialInputs), DecalFading);
-		Out.MRT[0] = float4(AmbientOcclusion, AmbientOcclusion, AmbientOcclusion, AmbientOcclusion);
-		return;
-	}
-	#endif
-
-	if(bVolumetric)
-	{
-		// not normalized
-		float3 OSNormal = GradientNotNormalized(In.SvPosition, OSPosition, HitT);
-			
-		// To transform the normals use tranpose(Inverse(DecalToWorld)) = transpose(WorldToDecal)
-		// But we want to only rotate the normals (we don't want to non-uniformaly scale them).
-		// We assume the matrix is only a scale and rotation, and we remove non-uniform scale:
-		float3 lengthSqr = { length2(DecalToWorld._m00_m01_m02),
-			length2(DecalToWorld._m10_m11_m12),
-			length2(DecalToWorld._m20_m21_m22) };
-
-		float3 scale = rsqrt(lengthSqr);
-
-		// Pre-multiply by the inverse of the non-uniform scale in DecalToWorld
-		float4 ScaledNormal = float4(OSNormal.xyz * scale.xyz, 0.f); 
-		float3 WSNormal = mul(ScaledNormal, DecalToWorld).xyz;
-
-		// need to normalized
-
-		MaterialParameters.WorldNormal = normalize(WSNormal);
-	}
-
-	// Store the results in local variables and reuse instead of calling the functions multiple times.
-	half3 BaseColor = GetMaterialBaseColor(PixelMaterialInputs);
-	half  Metallic = GetMaterialMetallic(PixelMaterialInputs);
-	half  Specular = GetMaterialSpecular(PixelMaterialInputs);
-
-	float3 Color = 1;
-
-	#if DECAL_BLEND_MODE == DECALBLENDMODEID_NORMAL
-		// -1..1 range to 0..1
-		Color = MaterialParameters.WorldNormal * 0.5f + 0.5f;
-	#else
-		Color = GetMaterialEmissive(PixelMaterialInputs);
-	
-		// only for mobile forward		
-		#if (SHADING_PATH_MOBILE && !MOBILE_DEFERRED_SHADING)
-			Color+= BaseColor;
-		#endif
-	#endif
+	float3 Color = GetMaterialEmissive(PixelMaterialInputs);
 	
+	float DecalFading = saturate(4 - 4 * abs(SwizzlePos.z * 2 - 1)) * DecalParams.x;
 	float Opacity = GetMaterialOpacity(PixelMaterialInputs) * DecalFading;
 
 	FGBufferData GBufferData;
-
 	GBufferData.WorldNormal = MaterialParameters.WorldNormal;
-	GBufferData.BaseColor = BaseColor;
-	GBufferData.Metallic = Metallic;
-	GBufferData.Specular = Specular;
+	GBufferData.BaseColor = GetMaterialBaseColor(PixelMaterialInputs);
+	GBufferData.Metallic = GetMaterialMetallic(PixelMaterialInputs);
+	GBufferData.Specular = GetMaterialSpecular(PixelMaterialInputs);
 	GBufferData.Roughness = GetMaterialRoughness(PixelMaterialInputs);
 	GBufferData.CustomData = 0;
 	GBufferData.IndirectIrradiance = 0;
 	GBufferData.PrecomputedShadowFactors = 1;
-	GBufferData.GBufferAO = 1;
+	GBufferData.GBufferAO = GetMaterialAmbientOcclusion(PixelMaterialInputs);
 	GBufferData.ShadingModelID = SHADINGMODELID_DEFAULT_LIT;
 	GBufferData.SelectiveOutputMask = 0;
 	GBufferData.PerObjectGBufferData = 1; 
@@ -339,7 +152,6 @@ void FPixelShaderInOut_MainPS(inout FPixelShaderIn In, inout FPixelShaderOut Out
 #endif
 }
 
-
 //#define PIXELSHADEROUTPUT_MRT0 (DECAL_RENDERTARGET_COUNT > 0)
 //#define PIXELSHADEROUTPUT_MRT1 (DECAL_RENDERTARGET_COUNT > 1)
 //#define PIXELSHADEROUTPUT_MRT2 (DECAL_RENDERTARGET_COUNT > 2)
@@ -349,5 +161,4 @@ void FPixelShaderInOut_MainPS(inout FPixelShaderIn In, inout FPixelShaderOut Out
 // this include generates the wrapper code to call MainPS()
 #include "PixelShaderOutputCommon.ush"
 
-
-#endif // DECAL_PROJECTION
+#endif // IS_DECAL
diff --git a/Engine/Shaders/Private/MaterialTemplate.ush b/Engine/Shaders/Private/MaterialTemplate.ush
--- a/Engine/Shaders/Private/MaterialTemplate.ush
+++ b/Engine/Shaders/Private/MaterialTemplate.ush
@@ -2351,7 +2351,7 @@ half GetMaterialOpacityRaw(FPixelMaterialInputs PixelMaterialInputs)
 	return PixelMaterialInputs.Opacity;
 }
 
-#if MATERIALBLENDING_MASKED || (DECAL_BLEND_MODE == DECALBLENDMODEID_VOLUMETRIC)
+#if MATERIALBLENDING_MASKED
 // Returns the material mask value generated from the material expressions.
 // Use GetMaterialMask() to get the value altered depending on the material blend mode.
 half GetMaterialMaskInputRaw(FPixelMaterialInputs PixelMaterialInputs)
diff --git a/Engine/Shaders/Private/MaterialTexCoordScalesPixelShader.usf b/Engine/Shaders/Private/MaterialTexCoordScalesPixelShader.usf
--- a/Engine/Shaders/Private/MaterialTexCoordScalesPixelShader.usf
+++ b/Engine/Shaders/Private/MaterialTexCoordScalesPixelShader.usf
@@ -163,7 +163,7 @@ void Main(
 		half3 Normal = GetMaterialNormalRaw(PixelMaterialInputs);
 		half3 Emissive = GetMaterialEmissiveRaw(PixelMaterialInputs);
 		half Opacity = GetMaterialOpacityRaw(PixelMaterialInputs);
-#if MATERIALBLENDING_MASKED || (DECAL_BLEND_MODE == DECALBLENDMODEID_VOLUMETRIC)
+#if MATERIALBLENDING_MASKED
 		half Mask = GetMaterialMask(PixelMaterialInputs);
 		if (!bOutputScales)
 		{
diff --git a/Engine/Shaders/Private/MeshDecals.usf b/Engine/Shaders/Private/MeshDecals.usf
--- a/Engine/Shaders/Private/MeshDecals.usf
+++ b/Engine/Shaders/Private/MeshDecals.usf
@@ -88,9 +88,6 @@ void MainVS(
 
 #endif // VERTEXSHADER
 
-
-
-
 // is called in MainPS() from PixelShaderOutputCommon.usf
 void FPixelShaderInOut_MainPS(
 	FVertexFactoryInterpolantsVSToPS FactoryInterpolants,
@@ -123,40 +120,20 @@ void FPixelShaderInOut_MainPS(
 	}
 #endif
 
-	// Store the results in local variables and reuse instead of calling the functions multiple times.
-	half3 BaseColor = GetMaterialBaseColor(PixelMaterialInputs);
-	half  Metallic = GetMaterialMetallic(PixelMaterialInputs);
-	half  Specular = GetMaterialSpecular(PixelMaterialInputs);
-
-	float3 Color = 1;
-
-	#if DECAL_BLEND_MODE == DECALBLENDMODEID_NORMAL
-		// -1..1 range to 0..1
-		Color = MaterialParameters.WorldNormal * 0.5f + 0.5f;
-	#else
-		Color = GetMaterialEmissive(PixelMaterialInputs);
-	
-		#if (ES3_1_PROFILE) // only for mobile
-			Color+= BaseColor;
-		#endif
-	#endif
+	float3 Color = GetMaterialEmissive(PixelMaterialInputs);
 
 	float Opacity = GetMaterialOpacity(PixelMaterialInputs);
 
-	// decal fade
-//	Opacity *= DecalParams.x;
-
 	FGBufferData GBufferData = (FGBufferData)0;
-
 	GBufferData.WorldNormal = MaterialParameters.WorldNormal;
-	GBufferData.BaseColor = BaseColor;
-	GBufferData.Metallic = Metallic;
-	GBufferData.Specular = Specular;
+	GBufferData.BaseColor = GetMaterialBaseColor(PixelMaterialInputs);
+	GBufferData.Metallic = GetMaterialMetallic(PixelMaterialInputs);
+	GBufferData.Specular = GetMaterialSpecular(PixelMaterialInputs);
 	GBufferData.Roughness = GetMaterialRoughness(PixelMaterialInputs);
 	GBufferData.CustomData = 0;
 	GBufferData.IndirectIrradiance = 0;
 	GBufferData.PrecomputedShadowFactors = 1;
-	GBufferData.GBufferAO = 1;
+	GBufferData.GBufferAO = GetMaterialAmbientOcclusion(PixelMaterialInputs);
 	GBufferData.ShadingModelID = SHADINGMODELID_DEFAULT_LIT;
 	GBufferData.SelectiveOutputMask = 0;
 	GBufferData.PerObjectGBufferData = 1;
@@ -176,7 +153,6 @@ void FPixelShaderInOut_MainPS(
 #endif
 }
 
-
 #define PIXELSHADEROUTPUT_MESHDECALPASS NEED_MESH_DECAL_INTERPOLATORS
 #define PIXELSHADEROUTPUT_INTERPOLANTS 1
 //#define PIXELSHADEROUTPUT_MRT0 (DECAL_RENDERTARGET_COUNT > 0)
diff --git a/Engine/Shaders/Private/RequiredTextureResolutionPixelShader.usf b/Engine/Shaders/Private/RequiredTextureResolutionPixelShader.usf
--- a/Engine/Shaders/Private/RequiredTextureResolutionPixelShader.usf
+++ b/Engine/Shaders/Private/RequiredTextureResolutionPixelShader.usf
@@ -108,7 +108,7 @@ void Main(
 		half3 Normal = GetMaterialNormalRaw(PixelMaterialInputs);
 		half3 Emissive = GetMaterialEmissiveRaw(PixelMaterialInputs);
 		half Opacity = GetMaterialOpacityRaw(PixelMaterialInputs);
-#if MATERIALBLENDING_MASKED || (DECAL_BLEND_MODE == DECALBLENDMODEID_VOLUMETRIC)
+#if MATERIALBLENDING_MASKED
 		half Mask = GetMaterialMask(PixelMaterialInputs);
 		clip(GetMaterialMask(PixelMaterialInputs));
 #endif
diff --git a/Engine/Shaders/Private/VirtualTextureCommon.ush b/Engine/Shaders/Private/VirtualTextureCommon.ush
--- a/Engine/Shaders/Private/VirtualTextureCommon.ush
+++ b/Engine/Shaders/Private/VirtualTextureCommon.ush
@@ -63,7 +63,7 @@ void FinalizeVirtualTextureFeedback(in FVirtualTextureFeedbackParams Params, flo
 
 		// When using DBuffer, feedback buffer is blending using decal opacity for opaque objects
 		// This way we avoid overwriting values already written to feedback buffer while rendering dbuffer decals
-#if MATERIALBLENDING_ANY_TRANSLUCENT || DECAL_PROJECTION || (USE_DBUFFER && MATERIALDECALRESPONSEMASK)
+#if MATERIALBLENDING_ANY_TRANSLUCENT || IS_DECAL || (USE_DBUFFER && MATERIALDECALRESPONSEMASK)
 
 		// Stochastic alpha testing
 #if MATERIALBLENDING_ADDITIVE || MATERIALBLENDING_MODULATE
@@ -73,7 +73,7 @@ void FinalizeVirtualTextureFeedback(in FVirtualTextureFeedbackParams Params, flo
 #else
 		// Opacity values smaller/larger than the reciprocal of the downscale factor get clamped so at least one pixel will resolve
 		const float VTDownscaleEpsilon = 1.0f / VIRTUAL_TEXTURE_FEEDBACK_FACTOR;
-		#if DECAL_PROJECTION
+		#if IS_DECAL
 			const float MaxAlpha = VT_FEEDBACK_DECAL_MAX_OPACITY;
 		#else
 			const float MaxAlpha = 1.0f - VTDownscaleEpsilon;
diff --git a/Engine/Source/Runtime/Engine/Private/Materials/HLSLMaterialTranslator.cpp b/Engine/Source/Runtime/Engine/Private/Materials/HLSLMaterialTranslator.cpp
--- a/Engine/Source/Runtime/Engine/Private/Materials/HLSLMaterialTranslator.cpp
+++ b/Engine/Source/Runtime/Engine/Private/Materials/HLSLMaterialTranslator.cpp
@@ -877,8 +877,7 @@ bool FHLSLMaterialTranslator::Translate()
 		bAllowCodeChunkGeneration = false;
 
 		bUsesEmissiveColor = IsMaterialPropertyUsed(MP_EmissiveColor, Chunk[MP_EmissiveColor], FLinearColor(0, 0, 0, 0), 3);
-		bUsesPixelDepthOffset = (AllowPixelDepthOffset(Platform) && IsMaterialPropertyUsed(MP_PixelDepthOffset, Chunk[MP_PixelDepthOffset], FLinearColor(0, 0, 0, 0), 1))
-			|| (Domain == MD_DeferredDecal && Material->GetDecalBlendMode() == DBM_Volumetric_DistanceFunction);
+		bUsesPixelDepthOffset = (AllowPixelDepthOffset(Platform) && IsMaterialPropertyUsed(MP_PixelDepthOffset, Chunk[MP_PixelDepthOffset], FLinearColor(0, 0, 0, 0), 1));
 
 		bool bUsesWorldPositionOffsetCurrent = IsMaterialPropertyUsed(MP_WorldPositionOffset, Chunk[MP_WorldPositionOffset], FLinearColor(0, 0, 0, 0), 3);
 		bool bUsesWorldPositionOffsetPrevious = IsMaterialPropertyUsed(MP_WorldPositionOffset, Chunk[CompiledMP_PrevWorldPositionOffset], FLinearColor(0, 0, 0, 0), 3);
@@ -1024,15 +1023,6 @@ bool FHLSLMaterialTranslator::Translate()
 			}
 		}
 
-		bool bDBufferAllowed = IsUsingDBuffers(Platform);
-		bool bDBufferBlendMode = IsDBufferDecalBlendMode((EDecalBlendMode)Material->GetDecalBlendMode());
-
-		if (bDBufferBlendMode && !bDBufferAllowed)
-		{
-			// Error feedback for when the decal would not be displayed due to project settings
-			Errorf(TEXT("DBuffer decal blend modes are only supported when the 'DBuffer Decals' Rendering Project setting is enabled."));
-		}
-
 		if (Domain == MD_DeferredDecal && BlendMode != BLEND_Translucent)
 		{
 			// We could make the change for the user but it would be confusing when going to DeferredDecal and back
@@ -5325,14 +5315,7 @@ int32 FHLSLMaterialTranslator::TextureSample(
 	const bool bVirtualTexture = TextureType == MCT_TextureVirtual;
 	if (bVirtualTexture)
 	{
-		if (Material->GetMaterialDomain() == MD_DeferredDecal)
-		{
-			if (Material->GetDecalBlendMode() == DBM_Volumetric_DistanceFunction)
-			{
-				return Errorf(TEXT("Sampling a virtual texture is currently only supported inside a volumetric decal."));
-			}
-		}
-		else if (Material->GetMaterialDomain() != MD_Surface)
+		if (Material->GetMaterialDomain() != MD_Surface && Material->GetMaterialDomain() != MD_DeferredDecal)
 		{
 			return Errorf(TEXT("Sampling a virtual texture is currently only supported inside surface and decal shaders."));
 		}
@@ -6050,10 +6033,10 @@ void FHLSLMaterialTranslator::UseSceneTextureId(ESceneTextureId SceneTextureId,
 		bool bRequiresSM5 = (SceneTextureId == PPI_WorldNormal || SceneTextureId == PPI_CustomDepth || SceneTextureId == PPI_CustomStencil || SceneTextureId == PPI_AmbientOcclusion);
 
 		if(bDBuffer)
-		{
+	{
 			if(!(SceneTextureId == PPI_SceneDepth || SceneTextureId == PPI_CustomDepth || SceneTextureId == PPI_CustomStencil || SceneTextureId == PPI_WorldNormal))
-			{
-				// Note: For DBuffer decals: CustomDepth and CustomStencil are only available if r.CustomDepth.Order = 0
+		{
+			// Note: For DBuffer decals: CustomDepth and CustomStencil are only available if r.CustomDepth.Order = 0
 				Errorf(TEXT("DBuffer decals (MaterialDomain=DeferredDecal and DecalBlendMode is using DBuffer) can only access SceneDepth, CustomDepth, CustomStencil, and WorldNormal"));
 			}
 		}
@@ -6062,13 +6045,13 @@ void FHLSLMaterialTranslator::UseSceneTextureId(ESceneTextureId SceneTextureId,
 			if(!(SceneTextureId == PPI_SceneDepth || SceneTextureId == PPI_CustomDepth || SceneTextureId == PPI_CustomStencil || SceneTextureId == PPI_WorldNormal || SceneTextureId == PPI_AmbientOcclusion))
 			{
 				Errorf(TEXT("Decals (MaterialDomain=DeferredDecal) can only access WorldNormal, AmbientOcclusion, SceneDepth, CustomDepth, CustomStencil"));
-			}
+		}
 
-			if (SceneTextureId == PPI_WorldNormal && Material->HasNormalConnected())
-			{
-					// GBuffer can only relate to WorldNormal here.
-				Errorf(TEXT("Decals that read WorldNormal cannot output to normal at the same time"));
-			}
+		if (SceneTextureId == PPI_WorldNormal && Material->HasNormalConnected())
+		{
+			// GBuffer can only relate to WorldNormal here.
+			Errorf(TEXT("Decals that read WorldNormal cannot output to normal at the same time"));
+		}
 		}
 
 		if (bRequiresSM5)
diff --git a/Engine/Source/Runtime/Engine/Private/Materials/MaterialShared.cpp b/Engine/Source/Runtime/Engine/Private/Materials/MaterialShared.cpp
--- a/Engine/Source/Runtime/Engine/Private/Materials/MaterialShared.cpp
+++ b/Engine/Source/Runtime/Engine/Private/Materials/MaterialShared.cpp
@@ -38,7 +38,6 @@
 #include "ShaderPlatformQualitySettings.h"
 #include "MaterialShaderQualitySettings.h"
 #include "Engine/RendererSettings.h"
-#include "DecalRenderingCommon.h"
 #include "ExternalTexture.h"
 #include "ShaderCodeLibrary.h"
 #include "HAL/FileManager.h"
@@ -873,8 +872,7 @@ bool FMaterial::MaterialMayModifyMeshPosition() const
 {
 	// Conservative estimate when called before material translation has occurred. 
 	// This function is only intended for use in deciding whether or not shader permutations are required.
-	return HasVertexPositionOffsetConnected() || HasPixelDepthOffsetConnected() || HasMaterialAttributesConnected()
-		|| (GetMaterialDomain() == MD_DeferredDecal && GetDecalBlendMode() == DBM_Volumetric_DistanceFunction);
+	return HasVertexPositionOffsetConnected() || HasPixelDepthOffsetConnected() || HasMaterialAttributesConnected();
 }
 
 bool FMaterial::MaterialUsesPixelDepthOffset() const
@@ -2013,18 +2011,6 @@ void FMaterial::SetupMaterialEnvironment(
 		OutEnvironment.CompilerFlags.Add(CFLAG_UseFullPrecisionInPS);
 	}
 
-	if(GetMaterialDomain() == MD_DeferredDecal)
-	{
-		// to compare against DECAL_BLEND_MODE, we can expose more if needed
-		OutEnvironment.SetDefine(TEXT("DECALBLENDMODEID_VOLUMETRIC"), (uint32)DBM_Volumetric_DistanceFunction);
-		OutEnvironment.SetDefine(TEXT("DECALBLENDMODEID_STAIN"), (uint32)DBM_Stain);
-		OutEnvironment.SetDefine(TEXT("DECALBLENDMODEID_NORMAL"), (uint32)DBM_Normal);
-		OutEnvironment.SetDefine(TEXT("DECALBLENDMODEID_EMISSIVE"), (uint32)DBM_Emissive);
-		OutEnvironment.SetDefine(TEXT("DECALBLENDMODEID_TRANSLUCENT"), (uint32)DBM_Translucent);
-		OutEnvironment.SetDefine(TEXT("DECALBLENDMODEID_AO"), (uint32)DBM_AmbientOcclusion);
-		OutEnvironment.SetDefine(TEXT("DECALBLENDMODEID_ALPHACOMPOSITE"), (uint32)DBM_AlphaComposite);
-	}
-
 	switch(GetMaterialDomain())
 	{
 		case MD_Surface:				OutEnvironment.SetDefine(TEXT("MATERIAL_DOMAIN_SURFACE"),			TEXT("1")); break;
diff --git a/Engine/Source/Runtime/Engine/Private/ShaderCompiler/ShaderGenerationUtil.cpp b/Engine/Source/Runtime/Engine/Private/ShaderCompiler/ShaderGenerationUtil.cpp
--- a/Engine/Source/Runtime/Engine/Private/ShaderCompiler/ShaderGenerationUtil.cpp
+++ b/Engine/Source/Runtime/Engine/Private/ShaderCompiler/ShaderGenerationUtil.cpp
@@ -183,14 +183,6 @@ void FShaderCompileUtilities::ApplyFetchEnvironment(FShaderMaterialPropertyDefin
 	FETCH_COMPILE_BOOL(MATERIAL_OUTPUT_OPACITY_AS_ALPHA);
 	FETCH_COMPILE_BOOL(TRANSLUCENT_SHADOW_WITH_MASKED_OPACITY);
 
-	FETCH_COMPILE_INT(DECALBLENDMODEID_VOLUMETRIC);
-	FETCH_COMPILE_INT(DECALBLENDMODEID_STAIN);
-	FETCH_COMPILE_INT(DECALBLENDMODEID_NORMAL);
-	FETCH_COMPILE_INT(DECALBLENDMODEID_EMISSIVE);
-	FETCH_COMPILE_INT(DECALBLENDMODEID_TRANSLUCENT);
-	FETCH_COMPILE_INT(DECALBLENDMODEID_AO);
-	FETCH_COMPILE_INT(DECALBLENDMODEID_ALPHACOMPOSITE);
-
 	FETCH_COMPILE_BOOL(MATERIAL_DOMAIN_SURFACE);
 	FETCH_COMPILE_BOOL(MATERIAL_DOMAIN_DEFERREDDECAL);
 	FETCH_COMPILE_BOOL(MATERIAL_DOMAIN_LIGHTFUNCTION);
diff --git a/Engine/Source/Runtime/RenderCore/Public/ShaderMaterial.h b/Engine/Source/Runtime/RenderCore/Public/ShaderMaterial.h
--- a/Engine/Source/Runtime/RenderCore/Public/ShaderMaterial.h
+++ b/Engine/Source/Runtime/RenderCore/Public/ShaderMaterial.h
@@ -148,14 +148,6 @@ struct FShaderMaterialPropertyDefines
 	uint8 MATERIAL_OUTPUT_OPACITY_AS_ALPHA : 1;
 	uint8 TRANSLUCENT_SHADOW_WITH_MASKED_OPACITY : 1;
 
-	uint32 DECALBLENDMODEID_VOLUMETRIC;
-	uint32 DECALBLENDMODEID_STAIN;
-	uint32 DECALBLENDMODEID_NORMAL;
-	uint32 DECALBLENDMODEID_EMISSIVE;
-	uint32 DECALBLENDMODEID_TRANSLUCENT;
-	uint32 DECALBLENDMODEID_AO;
-	uint32 DECALBLENDMODEID_ALPHACOMPOSITE;
-
 	uint8 MATERIAL_DOMAIN_SURFACE : 1;
 	uint8 MATERIAL_DOMAIN_DEFERREDDECAL : 1;
 	uint8 MATERIAL_DOMAIN_LIGHTFUNCTION : 1;
diff --git a/Engine/Source/Runtime/Renderer/Private/CompositionLighting/CompositionLighting.cpp b/Engine/Source/Runtime/Renderer/Private/CompositionLighting/CompositionLighting.cpp
--- a/Engine/Source/Runtime/Renderer/Private/CompositionLighting/CompositionLighting.cpp
+++ b/Engine/Source/Runtime/Renderer/Private/CompositionLighting/CompositionLighting.cpp
@@ -514,7 +514,7 @@ void ProcessBeforeBasePass(
 			if (bDBuffer)
 			{
 				FDeferredDecalPassTextures DecalPassTextures = GetDeferredDecalPassTextures(GraphBuilder, SceneTextures, &DBufferTextures);
-				AddDeferredDecalPass(GraphBuilder, View, DecalPassTextures, DRS_BeforeBasePass);
+				AddDeferredDecalPass(GraphBuilder, View, DecalPassTextures, EDecalRenderStage::BeforeBassPass);
 			}
 
 			if (bNeedSSAO)
@@ -574,21 +574,16 @@ void ProcessAfterBasePass(
 
 	FDeferredDecalPassTextures DecalPassTextures = GetDeferredDecalPassTextures(GraphBuilder, SceneTextures, nullptr);
 
-	if (ViewFamily.EngineShowFlags.Decals && !ViewFamily.EngineShowFlags.ShaderComplexity)
-	{
-		AddDeferredDecalPass(GraphBuilder, View, DecalPassTextures, DRS_AfterBasePass);
-	}
-
 	if (bDoDecal && IsUsingGBuffers(View.GetShaderPlatform()))
 	{
-		// decals are before AmbientOcclusion so the decal can output a normal that AO is affected by
-		AddDeferredDecalPass(GraphBuilder, View, DecalPassTextures, DRS_BeforeLighting);
+		// Decals are before AmbientOcclusion so the decal can output a normal that AO is affected by
+		AddDeferredDecalPass(GraphBuilder, View, DecalPassTextures, EDecalRenderStage::BeforeLighting);
 	}
 
 	if (bDoDecal && !IsSimpleForwardShadingEnabled(View.GetShaderPlatform()))
 	{
 		// DBuffer decals with emissive component
-		AddDeferredDecalPass(GraphBuilder, View, DecalPassTextures, DRS_Emissive);
+		AddDeferredDecalPass(GraphBuilder, View, DecalPassTextures, EDecalRenderStage::Emissive);
 	}
 
 	// Forward shading SSAO is applied before the base pass using only the depth buffer.
@@ -621,7 +616,7 @@ void ProcessAfterBasePass(
 				AmbientOcclusion = AddPostProcessingGTAOPostAsync(GraphBuilder, View, Parameters, GTAOHorizons, FinalTarget);
 
 				ensureMsgf(
-					FDecalRendering::BuildVisibleDecalList(*(FScene*)View.Family->Scene, View, DRS_AmbientOcclusion, nullptr) == false,
+					DecalRendering::BuildVisibleDecalList(*(FScene*)View.Family->Scene, View, EDecalRenderStage::AmbientOcclusion, nullptr) == false,
 					TEXT("Ambient occlusion decals are not supported with Async compute SSAO."));
 			}
 			else
@@ -640,7 +635,7 @@ void ProcessAfterBasePass(
 				if (bDoDecal)
 				{
 					DecalPassTextures.ScreenSpaceAO = AmbientOcclusion.Texture;
-					AddDeferredDecalPass(GraphBuilder, View, DecalPassTextures, DRS_AmbientOcclusion);
+					AddDeferredDecalPass(GraphBuilder, View, DecalPassTextures, EDecalRenderStage::AmbientOcclusion);
 				}
 			}
 		}
diff --git a/Engine/Source/Runtime/Renderer/Private/CompositionLighting/PostProcessDeferredDecals.cpp b/Engine/Source/Runtime/Renderer/Private/CompositionLighting/PostProcessDeferredDecals.cpp
--- a/Engine/Source/Runtime/Renderer/Private/CompositionLighting/PostProcessDeferredDecals.cpp
+++ b/Engine/Source/Runtime/Renderer/Private/CompositionLighting/PostProcessDeferredDecals.cpp
@@ -98,12 +98,9 @@ FDeferredDecalPassTextures GetDeferredDecalPassTextures(FRDGBuilder& GraphBuilde
 void GetDeferredDecalPassParameters(
 	const FViewInfo& View,
 	const FDeferredDecalPassTextures& Textures,
-	FDecalRenderingCommon::ERenderTargetMode RenderTargetMode,
+	EDecalRenderTargetMode RenderTargetMode,
 	FDeferredDecalPassParameters& PassParameters)
 {
-	const bool bWritingToGBufferA = IsWritingToGBufferA(RenderTargetMode);
-	const bool bWritingToDepth = IsWritingToDepth(RenderTargetMode);
-
 	PassParameters.View = View.GetShaderParameters();
 	PassParameters.DeferredDecal = CreateDeferredDecalUniformBuffer(View);
 	PassParameters.SceneTextures = Textures.SceneTexturesUniformBuffer;
@@ -122,38 +119,22 @@ void GetDeferredDecalPassParameters(
 
 	switch (RenderTargetMode)
 	{
-	case FDecalRenderingCommon::RTM_SceneColorAndGBufferWithNormal:
-	case FDecalRenderingCommon::RTM_SceneColorAndGBufferNoNormal:
+	case EDecalRenderTargetMode::SceneColorAndGBuffer:
 		AddColorTarget(Textures.Color);
-		if (bWritingToGBufferA)
-		{
-			AddColorTarget(Textures.GBufferA);
-		}
+		AddColorTarget(Textures.GBufferA);
 		AddColorTarget(Textures.GBufferB);
 		AddColorTarget(Textures.GBufferC);
 		break;
-
-	case FDecalRenderingCommon::RTM_SceneColorAndGBufferDepthWriteWithNormal:
-	case FDecalRenderingCommon::RTM_SceneColorAndGBufferDepthWriteNoNormal:
+	case EDecalRenderTargetMode::SceneColorAndGBufferNoNormal:
 		AddColorTarget(Textures.Color);
-		if (bWritingToGBufferA)
-		{
-			AddColorTarget(Textures.GBufferA);
-		}
 		AddColorTarget(Textures.GBufferB);
 		AddColorTarget(Textures.GBufferC);
-		AddColorTarget(Textures.GBufferE);
-		break;
-
-	case FDecalRenderingCommon::RTM_GBufferNormal:
-		AddColorTarget(Textures.GBufferA);
 		break;
-
-	case FDecalRenderingCommon::RTM_SceneColor:
+	case EDecalRenderTargetMode::SceneColor:
 		AddColorTarget(Textures.Color);
 		break;
 
-	case FDecalRenderingCommon::RTM_DBuffer:
+	case EDecalRenderTargetMode::DBuffer:
 	{
 		check(Textures.DBufferTextures);
 
@@ -176,7 +157,7 @@ void GetDeferredDecalPassParameters(
 		DepthTexture = Textures.Depth.Resolve;
 		break;
 	}
-	case FDecalRenderingCommon::RTM_AmbientOcclusion:
+	case EDecalRenderTargetMode::AmbientOcclusion:
 	{
 		AddColorTarget(Textures.ScreenSpaceAO);
 		break;
@@ -190,7 +171,7 @@ void GetDeferredDecalPassParameters(
 		DepthTexture,
 		ERenderTargetLoadAction::ELoad,
 		ERenderTargetLoadAction::ELoad,
-		bWritingToDepth ? FExclusiveDepthStencil::DepthWrite_StencilWrite : FExclusiveDepthStencil::DepthRead_StencilWrite);
+		FExclusiveDepthStencil::DepthRead_StencilWrite);
 }
 
 TUniformBufferRef<FDeferredDecalUniformParameters> CreateDeferredDecalUniformBuffer(const FViewInfo& View)
@@ -307,7 +288,7 @@ static bool RenderPreStencil(FRHICommandList& RHICmdList, const FViewInfo& View,
 		STENCIL_SANDBOX_MASK, STENCIL_SANDBOX_MASK
 	>::GetRHI();
 
-	FDecalRendering::SetVertexShaderOnly(RHICmdList, GraphicsPSOInit, View, FrustumComponentToClip);
+	DecalRendering::SetVertexShaderOnly(RHICmdList, GraphicsPSOInit, View, FrustumComponentToClip);
 	RHICmdList.SetStencilRef(0);
 
 	// Set stream source after updating cached strides
@@ -319,39 +300,16 @@ static bool RenderPreStencil(FRHICommandList& RHICmdList, const FViewInfo& View,
 	return true;
 }
 
-static EDecalRasterizerState ComputeDecalRasterizerState(bool bInsideDecal, bool bIsInverted, const FViewInfo& View)
-{
-	bool bClockwise = bInsideDecal;
-
-	if (View.bReverseCulling)
-	{
-		bClockwise = !bClockwise;
-	}
-
-	if (bIsInverted)
-	{
-		bClockwise = !bClockwise;
-	}
-	return bClockwise ? DRS_CW : DRS_CCW;
-}
-
 static FDecalDepthState ComputeDecalDepthState(EDecalRenderStage LocalDecalStage, bool bInsideDecal, bool bThisDecalUsesStencil)
 {
 	FDecalDepthState Ret;
 
-	Ret.bDepthOutput = (LocalDecalStage == DRS_AfterBasePass);
-
-	if (Ret.bDepthOutput)
-	{
-		// can be made one enum
-		Ret.DepthTest = DDS_DepthTest;
-		return Ret;
-	}
+	Ret.bDepthOutput = false;
 
 	const bool bUseDecalMask = 
-		LocalDecalStage == DRS_BeforeLighting || 
-		LocalDecalStage == DRS_Emissive || 
-		LocalDecalStage == DRS_AmbientOcclusion;
+		LocalDecalStage == EDecalRenderStage::BeforeLighting || 
+		LocalDecalStage == EDecalRenderStage::Emissive || 
+		LocalDecalStage == EDecalRenderStage::AmbientOcclusion;
 
 	if (bInsideDecal)
 	{
@@ -460,19 +418,9 @@ static FRHIDepthStencilState* GetDecalDepthState(uint32& StencilRef, FDecalDepth
 	}
 }
 
-FRHIRasterizerState* GetDecalRasterizerState(EDecalRasterizerState DecalRasterizerState)
-{
-	switch (DecalRasterizerState)
-	{
-	case DRS_CW: return TStaticRasterizerState<FM_Solid, CM_CW>::GetRHI();
-	case DRS_CCW: return TStaticRasterizerState<FM_Solid, CM_CCW>::GetRHI();
-	default: check(0); return nullptr;
-	}
-}
-
 static bool IsStencilOptimizationAvailable(EDecalRenderStage RenderStage)
 {
-	return RenderStage == DRS_BeforeLighting || RenderStage == DRS_BeforeBasePass || RenderStage == DRS_Emissive;
+	return RenderStage == EDecalRenderStage::BeforeLighting || RenderStage == EDecalRenderStage::BeforeBassPass || RenderStage == EDecalRenderStage::Emissive;
 }
 
 static EDecalDBufferMaskTechnique GetDBufferMaskTechnique(EShaderPlatform ShaderPlatform)
@@ -494,16 +442,14 @@ static EDecalDBufferMaskTechnique GetDBufferMaskTechnique(EShaderPlatform Shader
 
 static const TCHAR* GetStageName(EDecalRenderStage Stage)
 {
-	// could be implemented with enum reflections as well
-
 	switch (Stage)
 	{
-	case DRS_BeforeBasePass: return TEXT("DRS_BeforeBasePass");
-	case DRS_AfterBasePass: return TEXT("DRS_AfterBasePass");
-	case DRS_BeforeLighting: return TEXT("DRS_BeforeLighting");
-	case DRS_Mobile: return TEXT("DRS_Mobile");
-	case DRS_AmbientOcclusion: return TEXT("DRS_AmbientOcclusion");
-	case DRS_Emissive: return TEXT("DRS_Emissive");
+	case EDecalRenderStage::BeforeBassPass: return TEXT("BeforeBassPass");
+	case EDecalRenderStage::BeforeLighting: return TEXT("BeforeLighting");
+	case EDecalRenderStage::Mobile: return TEXT("Mobile");
+	case EDecalRenderStage::MobileBeforeLighting: return TEXT("MobileBeforeLighting");
+	case EDecalRenderStage::Emissive: return TEXT("Emissive");
+	case EDecalRenderStage::AmbientOcclusion: return TEXT("AmbientOcclusion");
 	}
 	return TEXT("<UNKNOWN>");
 }
@@ -557,11 +503,12 @@ void AddDeferredDecalPass(
 	EDecalRenderStage DecalRenderStage)
 {
 	check(PassTextures.Depth.IsValid());
-	check(DecalRenderStage != DRS_BeforeBasePass || PassTextures.DBufferTextures);
+	check(DecalRenderStage != EDecalRenderStage::BeforeBassPass || PassTextures.DBufferTextures);
 
 	const FSceneViewFamily& ViewFamily = *(View.Family);
 
 	// Debug view framework does not yet support decals.
+	// todo: Handle shader complexity mode here for deferred decal.
 	if (!ViewFamily.EngineShowFlags.Decals || ViewFamily.UseDebugViewPS())
 	{
 		return;
@@ -575,13 +522,13 @@ void AddDeferredDecalPass(
 	uint32 SortedDecalCount = 0;
 	FTransientDecalRenderDataList* SortedDecals = nullptr;
 
-	checkf(DecalRenderStage != DRS_AmbientOcclusion || PassTextures.ScreenSpaceAO, TEXT("Attepting to render AO decals without SSAO having emitted a valid render target."));
-	checkf(DecalRenderStage != DRS_BeforeBasePass || IsUsingDBuffers(ShaderPlatform), TEXT("Only DBuffer decals are supported before the base pass."));
+	checkf(DecalRenderStage != EDecalRenderStage::AmbientOcclusion || PassTextures.ScreenSpaceAO, TEXT("Attepting to render AO decals without SSAO having emitted a valid render target."));
+	checkf(DecalRenderStage != EDecalRenderStage::BeforeBassPass || IsUsingDBuffers(ShaderPlatform), TEXT("Only DBuffer decals are supported before the base pass."));
 
 	if (DecalCount)
 	{
 		SortedDecals = GraphBuilder.AllocObject<FTransientDecalRenderDataList>();
-		FDecalRendering::BuildVisibleDecalList(Scene, View, DecalRenderStage, SortedDecals);
+		DecalRendering::BuildVisibleDecalList(Scene, View, DecalRenderStage, SortedDecals);
 		SortedDecalCount = SortedDecals->Num();
 
 		INC_DWORD_STAT_BY(STAT_Decals, SortedDecalCount);
@@ -594,7 +541,7 @@ void AddDeferredDecalPass(
 	// Attempt to clear the D-Buffer if it's appropriate for this view.
 	const EDecalDBufferMaskTechnique DBufferMaskTechnique = GetDBufferMaskTechnique(ShaderPlatform);
 
-	const auto RenderDecals = [&](uint32 DecalIndexBegin, uint32 DecalIndexEnd, FDecalRenderingCommon::ERenderTargetMode RenderTargetMode)
+	const auto RenderDecals = [&](uint32 DecalIndexBegin, uint32 DecalIndexEnd, EDecalRenderTargetMode RenderTargetMode)
 	{
 		auto* PassParameters = GraphBuilder.AllocParameters<FDeferredDecalPassParameters>();
 		GetDeferredDecalPassParameters(View, PassTextures, RenderTargetMode, *PassParameters);
@@ -603,7 +550,7 @@ void AddDeferredDecalPass(
 			RDG_EVENT_NAME("Batch [%d, %d]", DecalIndexBegin, DecalIndexEnd - 1),
 			PassParameters,
 			ERDGPassFlags::Raster,
-			[&View, FeatureLevel, ShaderPlatform, DecalIndexBegin, DecalIndexEnd, SortedDecals, DecalRenderStage, bStencilSizeThreshold, bShaderComplexity](FRHICommandList& RHICmdList)
+			[&View, FeatureLevel, ShaderPlatform, DecalIndexBegin, DecalIndexEnd, SortedDecals, DecalRenderStage, RenderTargetMode, bStencilSizeThreshold, bShaderComplexity](FRHICommandList& RHICmdList)
 		{
 			RHICmdList.SetViewport(View.ViewRect.Min.X, View.ViewRect.Min.Y, 0.0f, View.ViewRect.Max.X, View.ViewRect.Max.Y, 1.0f);
 
@@ -612,10 +559,8 @@ void AddDeferredDecalPass(
 				const FTransientDecalRenderData& DecalData = (*SortedDecals)[DecalIndex];
 				const FDeferredDecalProxy& DecalProxy = *DecalData.DecalProxy;
 				const FMatrix ComponentToWorldMatrix = DecalProxy.ComponentTrans.ToMatrixWithScale();
-				const FMatrix FrustumComponentToClip = FDecalRendering::ComputeComponentToClipMatrix(View, ComponentToWorldMatrix);
-				const EDecalBlendMode DecalBlendMode = bShaderComplexity ? DBM_Emissive : FDecalRenderingCommon::ComputeDecalBlendModeForRenderStage(DecalData.FinalDecalBlendMode, DecalRenderStage);
-				const EDecalRenderStage LocalDecalStage = FDecalRenderingCommon::ComputeRenderStage(ShaderPlatform, DecalBlendMode);
-				const bool bStencilThisDecal = IsStencilOptimizationAvailable(LocalDecalStage);
+				const FMatrix FrustumComponentToClip = DecalRendering::ComputeComponentToClipMatrix(View, ComponentToWorldMatrix);
+				const bool bStencilThisDecal = IsStencilOptimizationAvailable(DecalRenderStage);
 
 				bool bThisDecalUsesStencil = false;
 
@@ -633,38 +578,32 @@ void AddDeferredDecalPass(
 					// Account for the reversal of handedness caused by negative scale on the decal
 					const FVector Scale = DecalProxy.ComponentTrans.GetScale3D();
 					const bool bReverseHanded = Scale.X * Scale.Y * Scale.Z < 0.0f;
-					const EDecalRasterizerState DecalRasterizerState = FDecalRenderingCommon::ComputeDecalRasterizerState(bInsideDecal, bReverseHanded, View.bReverseCulling);
-					GraphicsPSOInit.RasterizerState = GetDecalRasterizerState(DecalRasterizerState);
+					const EDecalRasterizerState DecalRasterizerState = DecalRendering::GetDecalRasterizerState(bInsideDecal, bReverseHanded, View.bReverseCulling);
+					GraphicsPSOInit.RasterizerState = DecalRendering::GetDecalRasterizerState(DecalRasterizerState);
 				}
 
 				uint32 StencilRef = 0;
 
 				{
-					const FDecalDepthState DecalDepthState = ComputeDecalDepthState(LocalDecalStage, bInsideDecal, bThisDecalUsesStencil);
+					const FDecalDepthState DecalDepthState = ComputeDecalDepthState(DecalRenderStage, bInsideDecal, bThisDecalUsesStencil);
 					GraphicsPSOInit.DepthStencilState = GetDecalDepthState(StencilRef, DecalDepthState);
 				}
 
-				GraphicsPSOInit.BlendState = FDecalRendering::GetDecalBlendState(FeatureLevel, DecalRenderStage, DecalBlendMode, DecalData.bHasNormal);
+				GraphicsPSOInit.BlendState = DecalRendering::GetDecalBlendState(DecalData.DecalBlendDesc, DecalRenderStage, RenderTargetMode);
 				GraphicsPSOInit.PrimitiveType = PT_TriangleList;
 
-				FDecalRendering::SetShader(RHICmdList, GraphicsPSOInit, View, DecalData, DecalRenderStage, FrustumComponentToClip);
+				DecalRendering::SetShader(RHICmdList, GraphicsPSOInit, View, DecalData, DecalRenderStage, FrustumComponentToClip);
 				RHICmdList.SetStencilRef(StencilRef);
 				RHICmdList.DrawIndexedPrimitive(GetUnitCubeIndexBuffer(), 0, 0, 8, 0, UE_ARRAY_COUNT(GCubeIndices) / 3, 1);
 			}
 		});
 	};
 
-	const auto GetRenderTargetMode = [&](const FTransientDecalRenderData& DecalData)
-	{
-		const EDecalBlendMode DecalBlendMode = FDecalRenderingCommon::ComputeDecalBlendModeForRenderStage(DecalData.FinalDecalBlendMode, DecalRenderStage);
-		return bShaderComplexity ? FDecalRenderingCommon::RTM_SceneColor : FDecalRenderingCommon::ComputeRenderTargetMode(ShaderPlatform, DecalBlendMode, DecalData.bHasNormal);
-	};
-
 	if (bVisibleDecalsInView)
 	{
 		RDG_EVENT_SCOPE(GraphBuilder, "DeferredDecals %s", GetStageName(DecalRenderStage));
 
-		if (MeshDecalCount > 0 && (DecalRenderStage == DRS_BeforeBasePass || DecalRenderStage == DRS_BeforeLighting || DecalRenderStage == DRS_Emissive))
+		if (MeshDecalCount > 0 && (DecalRenderStage == EDecalRenderStage::BeforeBassPass || DecalRenderStage == EDecalRenderStage::BeforeLighting || DecalRenderStage == EDecalRenderStage::Emissive || DecalRenderStage == EDecalRenderStage::AmbientOcclusion))
 		{
 			RenderMeshDecals(GraphBuilder, View, PassTextures, DecalRenderStage);
 		}
@@ -675,11 +614,11 @@ void AddDeferredDecalPass(
 
 			uint32 SortedDecalIndex = 1;
 			uint32 LastSortedDecalIndex = 0;
-			FDecalRenderingCommon::ERenderTargetMode LastRenderTargetMode = GetRenderTargetMode((*SortedDecals)[0]);
+			EDecalRenderTargetMode LastRenderTargetMode = DecalRendering::GetRenderTargetMode((*SortedDecals)[0].DecalBlendDesc, DecalRenderStage);
 
 			for (; SortedDecalIndex < SortedDecalCount; ++SortedDecalIndex)
 			{
-				const FDecalRenderingCommon::ERenderTargetMode RenderTargetMode = GetRenderTargetMode((*SortedDecals)[SortedDecalIndex]);
+				const EDecalRenderTargetMode RenderTargetMode = DecalRendering::GetRenderTargetMode((*SortedDecals)[SortedDecalIndex].DecalBlendDesc, DecalRenderStage);
 
 				if (LastRenderTargetMode != RenderTargetMode)
 				{
@@ -698,7 +637,7 @@ void AddDeferredDecalPass(
 
 	// Last D-Buffer pass in the frame decodes the write mask (if supported and decals were rendered).
 	if (DBufferMaskTechnique == EDecalDBufferMaskTechnique::WriteMask &&
-		DecalRenderStage == DRS_BeforeBasePass &&
+		DecalRenderStage == EDecalRenderStage::BeforeBassPass &&
 		PassTextures.DBufferTextures->IsValid() &&
 		View.IsLastInFamily())
 	{
diff --git a/Engine/Source/Runtime/Renderer/Private/CompositionLighting/PostProcessDeferredDecals.h b/Engine/Source/Runtime/Renderer/Private/CompositionLighting/PostProcessDeferredDecals.h
--- a/Engine/Source/Runtime/Renderer/Private/CompositionLighting/PostProcessDeferredDecals.h
+++ b/Engine/Source/Runtime/Renderer/Private/CompositionLighting/PostProcessDeferredDecals.h
@@ -3,11 +3,12 @@
 #pragma once
 
 #include "RenderGraph.h"
-#include "DecalRenderingCommon.h"
-
-class FViewInfo;
+#include "SceneView.h"
 
+enum class EDecalRenderStage : uint8;
+enum class EDecalRenderTargetMode : uint8;
 struct FSceneTextures;
+class FViewInfo;
 
 bool IsDBufferEnabled(const FSceneViewFamily& ViewFamily, EShaderPlatform ShaderPlatform);
 
@@ -39,19 +40,6 @@ END_SHADER_PARAMETER_STRUCT()
 
 FDBufferParameters GetDBufferParameters(FRDGBuilder& GraphBuilder, const FDBufferTextures& DBufferTextures, EShaderPlatform ShaderPlatform);
 
-inline bool IsWritingToGBufferA(FDecalRenderingCommon::ERenderTargetMode RenderTargetMode)
-{
-	return RenderTargetMode == FDecalRenderingCommon::RTM_SceneColorAndGBufferWithNormal
-		|| RenderTargetMode == FDecalRenderingCommon::RTM_SceneColorAndGBufferDepthWriteWithNormal
-		|| RenderTargetMode == FDecalRenderingCommon::RTM_GBufferNormal;
-}
-
-inline bool IsWritingToDepth(FDecalRenderingCommon::ERenderTargetMode RenderTargetMode)
-{
-	return RenderTargetMode == FDecalRenderingCommon::RTM_SceneColorAndGBufferDepthWriteWithNormal
-		|| RenderTargetMode == FDecalRenderingCommon::RTM_SceneColorAndGBufferDepthWriteNoNormal;
-}
-
 struct FDeferredDecalPassTextures
 {
 	TRDGUniformBufferRef<FSceneTextureUniformParameters> SceneTexturesUniformBuffer = nullptr;
@@ -96,7 +84,7 @@ END_SHADER_PARAMETER_STRUCT()
 void GetDeferredDecalPassParameters(
 	const FViewInfo& View,
 	const FDeferredDecalPassTextures& DecalPassTextures,
-	FDecalRenderingCommon::ERenderTargetMode RenderTargetMode,
+	EDecalRenderTargetMode RenderTargetMode,
 	FDeferredDecalPassParameters& PassParameters);
 
 void RenderMeshDecals(
diff --git a/Engine/Source/Runtime/Renderer/Private/CompositionLighting/PostProcessMeshDecals.cpp b/Engine/Source/Runtime/Renderer/Private/CompositionLighting/PostProcessMeshDecals.cpp
--- a/Engine/Source/Runtime/Renderer/Private/CompositionLighting/PostProcessMeshDecals.cpp
+++ b/Engine/Source/Runtime/Renderer/Private/CompositionLighting/PostProcessMeshDecals.cpp
@@ -21,15 +21,6 @@
 #include "DebugViewModeRendering.h"
 #include "MeshPassProcessor.inl"
 
-class FMeshDecalAccumulatePolicy
-{
-public:
-	static bool ShouldCompilePermutation(const FMeshMaterialShaderPermutationParameters& Parameters)
-	{
-		return Parameters.MaterialParameters.MaterialDomain == MD_DeferredDecal;
-	}
-};
-
 class FMeshDecalsVS : public FMeshMaterialShader
 {
 public:
@@ -37,7 +28,8 @@ public:
 
 	static bool ShouldCompilePermutation(const FMeshMaterialShaderPermutationParameters& Parameters)
 	{
-		return FMeshDecalAccumulatePolicy::ShouldCompilePermutation(Parameters);
+		return (Parameters.MaterialParameters.MaterialDomain == MD_DeferredDecal) &&
+			DecalRendering::GetBaseRenderStage(DecalRendering::ComputeDecalBlendDesc(Parameters.Platform, Parameters.MaterialParameters)) != EDecalRenderStage::None;
 	}
 
 	FMeshDecalsVS() = default;
@@ -55,13 +47,14 @@ public:
 
 	static bool ShouldCompilePermutation(const FMeshMaterialShaderPermutationParameters& Parameters)
 	{
-		return FMeshDecalAccumulatePolicy::ShouldCompilePermutation(Parameters);
+		return (Parameters.MaterialParameters.MaterialDomain == MD_DeferredDecal) &&
+			DecalRendering::GetBaseRenderStage(DecalRendering::ComputeDecalBlendDesc(Parameters.Platform, Parameters.MaterialParameters)) != EDecalRenderStage::None;
 	}
 
 	static void ModifyCompilationEnvironment(const FMaterialShaderPermutationParameters& Parameters, FShaderCompilerEnvironment& OutEnvironment)
 	{
 		FMeshMaterialShader::ModifyCompilationEnvironment(Parameters, OutEnvironment);
-		FDecalRendering::SetDecalCompilationEnvironment(Parameters, OutEnvironment);
+		DecalRendering::ModifyCompilationEnvironment(DecalRendering::ComputeDecalBlendDesc(Parameters.Platform, Parameters.MaterialParameters), EDecalRenderStage::None, OutEnvironment);
 	}
 
 	FMeshDecalsPS() = default;
@@ -79,15 +72,14 @@ public:
 
 	static bool ShouldCompilePermutation(const FMeshMaterialShaderPermutationParameters& Parameters)
 	{
-		return FMeshDecalsPS::ShouldCompilePermutation(Parameters)
-			&& Parameters.MaterialParameters.bHasEmissiveColorConnected
-			&& IsDBufferDecalBlendMode(FDecalRenderingCommon::ComputeFinalDecalBlendMode(Parameters.Platform, (EDecalBlendMode)Parameters.MaterialParameters.DecalBlendMode, Parameters.MaterialParameters.bHasNormalConnected));
+		return (Parameters.MaterialParameters.MaterialDomain == MD_DeferredDecal) &&
+			DecalRendering::IsCompatibleWithRenderStage(DecalRendering::ComputeDecalBlendDesc(Parameters.Platform, Parameters.MaterialParameters), EDecalRenderStage::Emissive);
 	}
 
 	static void ModifyCompilationEnvironment(const FMaterialShaderPermutationParameters& Parameters, FShaderCompilerEnvironment& OutEnvironment)
 	{
-		FMeshDecalsPS::ModifyCompilationEnvironment(Parameters, OutEnvironment);
-		FDecalRendering::SetEmissiveDBufferDecalCompilationEnvironment(Parameters, OutEnvironment);
+		FMeshMaterialShader::ModifyCompilationEnvironment(Parameters, OutEnvironment);
+		DecalRendering::ModifyCompilationEnvironment(DecalRendering::ComputeDecalBlendDesc(Parameters.Platform, Parameters.MaterialParameters), EDecalRenderStage::Emissive, OutEnvironment);
 	}
 
 	FMeshDecalsEmissivePS() = default;
@@ -98,6 +90,30 @@ public:
 
 IMPLEMENT_MATERIAL_SHADER_TYPE(, FMeshDecalsEmissivePS, TEXT("/Engine/Private/MeshDecals.usf"), TEXT("MainPS"), SF_Pixel);
 
+class FMeshDecalsAmbientOcclusionPS : public FMeshDecalsPS
+{
+public:
+	DECLARE_SHADER_TYPE(FMeshDecalsAmbientOcclusionPS, MeshMaterial);
+
+	static bool ShouldCompilePermutation(const FMeshMaterialShaderPermutationParameters& Parameters)
+	{
+		return (Parameters.MaterialParameters.MaterialDomain == MD_DeferredDecal) &&
+			DecalRendering::IsCompatibleWithRenderStage(DecalRendering::ComputeDecalBlendDesc(Parameters.Platform, Parameters.MaterialParameters), EDecalRenderStage::AmbientOcclusion);
+	}
+
+	static void ModifyCompilationEnvironment(const FMaterialShaderPermutationParameters& Parameters, FShaderCompilerEnvironment& OutEnvironment)
+	{
+		FMeshMaterialShader::ModifyCompilationEnvironment(Parameters, OutEnvironment);
+		DecalRendering::ModifyCompilationEnvironment(DecalRendering::ComputeDecalBlendDesc(Parameters.Platform, Parameters.MaterialParameters), EDecalRenderStage::AmbientOcclusion, OutEnvironment);
+	}
+
+	FMeshDecalsAmbientOcclusionPS() = default;
+	FMeshDecalsAmbientOcclusionPS(const ShaderMetaType::CompiledShaderInitializerType& Initializer)
+		: FMeshDecalsPS(Initializer)
+	{}
+};
+
+IMPLEMENT_MATERIAL_SHADER_TYPE(, FMeshDecalsAmbientOcclusionPS, TEXT("/Engine/Private/MeshDecals.usf"), TEXT("MainPS"), SF_Pixel);
 
 class FMeshDecalMeshProcessor : public FMeshPassProcessor
 {
@@ -105,7 +121,7 @@ public:
 	FMeshDecalMeshProcessor(const FScene* Scene, 
 		const FSceneView* InViewIfDynamicMeshCommand, 
 		EDecalRenderStage InPassDecalStage, 
-		FDecalRenderingCommon::ERenderTargetMode InRenderTargetMode,
+		EDecalRenderTargetMode InRenderTargetMode,
 		FMeshPassDrawListContext* InDrawListContext);
 
 	virtual void AddMeshBatch(const FMeshBatch& RESTRICT MeshBatch, uint64 BatchElementMask, const FPrimitiveSceneProxy* RESTRICT PrimitiveSceneProxy, int32 StaticMeshId = -1) override final;
@@ -123,7 +139,7 @@ private:
 
 	FMeshPassProcessorRenderState PassDrawRenderState;
 	const EDecalRenderStage PassDecalStage;
-	const FDecalRenderingCommon::ERenderTargetMode RenderTargetMode;
+	const EDecalRenderTargetMode RenderTargetMode;
 };
 
 IMPLEMENT_STATIC_UNIFORM_BUFFER_SLOT(DeferredDecals);
@@ -132,7 +148,7 @@ IMPLEMENT_STATIC_UNIFORM_BUFFER_STRUCT(FDeferredDecalUniformParameters, "Deferre
 FMeshDecalMeshProcessor::FMeshDecalMeshProcessor(const FScene* Scene, 
 	const FSceneView* InViewIfDynamicMeshCommand, 
 	EDecalRenderStage InPassDecalStage, 
-	FDecalRenderingCommon::ERenderTargetMode InRenderTargetMode,
+	EDecalRenderTargetMode InRenderTargetMode,
 	FMeshPassDrawListContext* InDrawListContext)
 	: FMeshPassProcessor(Scene, Scene->GetFeatureLevel(), InViewIfDynamicMeshCommand, InDrawListContext)
 	, PassDecalStage(InPassDecalStage)
@@ -154,67 +170,38 @@ void FMeshDecalMeshProcessor::AddMeshBatch(const FMeshBatch& RESTRICT MeshBatch,
 			if (Material->GetRenderingThreadShaderMap())
 			{
 				const EShaderPlatform ShaderPlatform = ViewIfDynamicMeshCommand->GetShaderPlatform();
-				const EDecalBlendMode FinalDecalBlendMode = FDecalRenderingCommon::ComputeFinalDecalBlendMode(ShaderPlatform, Material);
-				const EDecalRenderStage LocalDecalRenderStage = FDecalRenderingCommon::ComputeRenderStage(ShaderPlatform, FinalDecalBlendMode);
+				const FDecalBlendDesc DecalBlendDesc = DecalRendering::ComputeDecalBlendDesc(ShaderPlatform, Material);
 
-				const FMeshDrawingPolicyOverrideSettings OverrideSettings = ComputeMeshOverrideSettings(MeshBatch);
-				ERasterizerFillMode MeshFillMode = ComputeMeshFillMode(MeshBatch, *Material, OverrideSettings);
-				ERasterizerCullMode MeshCullMode = ComputeMeshCullMode(MeshBatch, *Material, OverrideSettings);
-
-				bool bShouldRender = FDecalRenderingCommon::IsCompatibleWithRenderStage(
-					ShaderPlatform,
-					PassDecalStage,
-					LocalDecalRenderStage,
-					FinalDecalBlendMode,
-					Material);
-
-				if (FinalDecalBlendMode == DBM_Normal)
-				{
-					bShouldRender = bShouldRender && RenderTargetMode == FDecalRenderingCommon::RTM_GBufferNormal;
-				}
-				else
-				{
-					bShouldRender = bShouldRender && RenderTargetMode != FDecalRenderingCommon::RTM_GBufferNormal;
-				}
-
-				if (PassDecalStage == DRS_Emissive)
-				{
-					bShouldRender = bShouldRender && Material->HasEmissiveColorConnected();
-				}
+				const bool bShouldRender = 
+					DecalRendering::IsCompatibleWithRenderStage(DecalBlendDesc, PassDecalStage) &&
+					DecalRendering::GetRenderTargetMode(DecalBlendDesc, PassDecalStage) == RenderTargetMode;
 
 				if (bShouldRender)
 				{
-					const bool bHasNormal = Material->HasNormalConnected();
-
-					const EDecalBlendMode DecalBlendMode = FDecalRenderingCommon::ComputeDecalBlendModeForRenderStage(
-						FDecalRenderingCommon::ComputeFinalDecalBlendMode(ShaderPlatform, (EDecalBlendMode)Material->GetDecalBlendMode(), bHasNormal),
-						PassDecalStage);
-
-					FDecalRenderingCommon::ERenderTargetMode DecalRenderTargetMode = FDecalRenderingCommon::ComputeRenderTargetMode(ShaderPlatform, DecalBlendMode, bHasNormal);
+					const FMeshDrawingPolicyOverrideSettings OverrideSettings = ComputeMeshOverrideSettings(MeshBatch);
+					ERasterizerFillMode MeshFillMode = ComputeMeshFillMode(MeshBatch, *Material, OverrideSettings);
+					ERasterizerCullMode MeshCullMode = ComputeMeshCullMode(MeshBatch, *Material, OverrideSettings);
 
-					if (DecalRenderTargetMode == RenderTargetMode)
+					if (ViewIfDynamicMeshCommand->Family->UseDebugViewPS())
 					{
-						if (ViewIfDynamicMeshCommand->Family->UseDebugViewPS())
+						// Deferred decals can only use translucent blend mode
+						if (ViewIfDynamicMeshCommand->Family->EngineShowFlags.ShaderComplexity)
 						{
-							// Deferred decals can only use translucent blend mode
-							if (ViewIfDynamicMeshCommand->Family->EngineShowFlags.ShaderComplexity)
-							{
-								// If we are in the translucent pass then override the blend mode, otherwise maintain additive blending.
-								PassDrawRenderState.SetBlendState(TStaticBlendState<CW_RGBA, BO_Add, BF_One, BF_One, BO_Add, BF_Zero, BF_One>::GetRHI());
-							}
-							else if (ViewIfDynamicMeshCommand->Family->GetDebugViewShaderMode() != DVSM_OutputMaterialTextureScales)
-							{
-								// Otherwise, force translucent blend mode (shaders will use an hardcoded alpha).
-								PassDrawRenderState.SetBlendState(TStaticBlendState<CW_RGBA, BO_Add, BF_SourceAlpha, BF_InverseSourceAlpha, BO_Add, BF_Zero, BF_InverseSourceAlpha>::GetRHI());
-							}
+							// If we are in the translucent pass then override the blend mode, otherwise maintain additive blending.
+							PassDrawRenderState.SetBlendState(TStaticBlendState<CW_RGBA, BO_Add, BF_One, BF_One, BO_Add, BF_Zero, BF_One>::GetRHI());
 						}
-						else
+						else if (ViewIfDynamicMeshCommand->Family->GetDebugViewShaderMode() != DVSM_OutputMaterialTextureScales)
 						{
-							PassDrawRenderState.SetBlendState(FDecalRendering::GetDecalBlendState(FeatureLevel, PassDecalStage, DecalBlendMode, bHasNormal));
+							// Otherwise, force translucent blend mode (shaders will use an hardcoded alpha).
+							PassDrawRenderState.SetBlendState(TStaticBlendState<CW_RGBA, BO_Add, BF_SourceAlpha, BF_InverseSourceAlpha, BO_Add, BF_Zero, BF_InverseSourceAlpha>::GetRHI());
 						}
-
-						Process(MeshBatch, BatchElementMask, StaticMeshId, PrimitiveSceneProxy, *MaterialRenderProxy, *Material, MeshFillMode, MeshCullMode);
 					}
+					else
+					{
+						PassDrawRenderState.SetBlendState(DecalRendering::GetDecalBlendState(DecalBlendDesc, PassDecalStage, RenderTargetMode));
+					}
+
+					Process(MeshBatch, BatchElementMask, StaticMeshId, PrimitiveSceneProxy, *MaterialRenderProxy, *Material, MeshFillMode, MeshCullMode);
 				}
 			}
 		}
@@ -237,10 +224,14 @@ void FMeshDecalMeshProcessor::Process(
 	FMaterialShaderTypes ShaderTypes;
 	ShaderTypes.AddShaderType<FMeshDecalsVS>();
 
-	if (PassDecalStage == DRS_Emissive)
+	if (PassDecalStage == EDecalRenderStage::Emissive)
 	{
 		ShaderTypes.AddShaderType<FMeshDecalsEmissivePS>();
 	}
+	else if (PassDecalStage == EDecalRenderStage::AmbientOcclusion)
+	{
+		ShaderTypes.AddShaderType<FMeshDecalsAmbientOcclusionPS>();
+	}
 	else
 	{
 		ShaderTypes.AddShaderType<FMeshDecalsPS>();
@@ -284,7 +275,7 @@ void DrawDecalMeshCommands(
 	const FViewInfo& View,
 	const FDeferredDecalPassTextures& DecalPassTextures,
 	EDecalRenderStage DecalRenderStage,
-	FDecalRenderingCommon::ERenderTargetMode RenderTargetMode)
+	EDecalRenderTargetMode RenderTargetMode)
 {
 	auto* PassParameters = GraphBuilder.AllocParameters<FDeferredDecalPassParameters>();
 	GetDeferredDecalPassParameters(View, DecalPassTextures, RenderTargetMode, *PassParameters);
@@ -331,50 +322,45 @@ void RenderMeshDecals(
 
 	switch (DecalRenderStage)
 	{
-	case DRS_BeforeBasePass:
-		DrawDecalMeshCommands(GraphBuilder, View, DecalPassTextures, DecalRenderStage, FDecalRenderingCommon::RTM_DBuffer);
+	case EDecalRenderStage::BeforeBassPass:
+		DrawDecalMeshCommands(GraphBuilder, View, DecalPassTextures, DecalRenderStage, EDecalRenderTargetMode::DBuffer);
 		break;
 
-	case DRS_AfterBasePass:
-		DrawDecalMeshCommands(GraphBuilder, View, DecalPassTextures, DecalRenderStage, FDecalRenderingCommon::RTM_SceneColorAndGBufferDepthWriteWithNormal);
+	case EDecalRenderStage::BeforeLighting:
+		DrawDecalMeshCommands(GraphBuilder, View, DecalPassTextures, DecalRenderStage, EDecalRenderTargetMode::SceneColorAndGBuffer);
+		DrawDecalMeshCommands(GraphBuilder, View, DecalPassTextures, DecalRenderStage, EDecalRenderTargetMode::SceneColorAndGBufferNoNormal);
 		break;
 
-	case DRS_BeforeLighting:
-		DrawDecalMeshCommands(GraphBuilder, View, DecalPassTextures, DecalRenderStage, FDecalRenderingCommon::RTM_GBufferNormal);
-		DrawDecalMeshCommands(GraphBuilder, View, DecalPassTextures, DecalRenderStage, FDecalRenderingCommon::RTM_SceneColorAndGBufferWithNormal);
-		DrawDecalMeshCommands(GraphBuilder, View, DecalPassTextures, DecalRenderStage, FDecalRenderingCommon::RTM_SceneColorAndGBufferNoNormal);
+	case EDecalRenderStage::Mobile:
+		DrawDecalMeshCommands(GraphBuilder, View, DecalPassTextures, DecalRenderStage, EDecalRenderTargetMode::SceneColor);
 		break;
 
-	case DRS_Mobile:
-		DrawDecalMeshCommands(GraphBuilder, View, DecalPassTextures, DecalRenderStage, FDecalRenderingCommon::RTM_SceneColor);
+	case EDecalRenderStage::MobileBeforeLighting:
+		DrawDecalMeshCommands(GraphBuilder, View, DecalPassTextures, DecalRenderStage, EDecalRenderTargetMode::SceneColorAndGBuffer);
 		break;
 
-	case DRS_AmbientOcclusion:
-		DrawDecalMeshCommands(GraphBuilder, View, DecalPassTextures, DecalRenderStage, FDecalRenderingCommon::RTM_AmbientOcclusion);
+	case EDecalRenderStage::Emissive:
+		DrawDecalMeshCommands(GraphBuilder, View, DecalPassTextures, DecalRenderStage, EDecalRenderTargetMode::SceneColor);
 		break;
 
-	case DRS_Emissive:
-		DrawDecalMeshCommands(GraphBuilder, View, DecalPassTextures, DecalRenderStage, FDecalRenderingCommon::RTM_SceneColor);
+	case EDecalRenderStage::AmbientOcclusion:
+		DrawDecalMeshCommands(GraphBuilder, View, DecalPassTextures, DecalRenderStage, EDecalRenderTargetMode::AmbientOcclusion);
 		break;
 	}
 }
 
-void RenderMeshDecalsMobile(FRHICommandListImmediate& RHICmdList, const FViewInfo& View)
+void RenderMeshDecalsMobile(FRHICommandListImmediate& RHICmdList, const FViewInfo& View, EDecalRenderStage DecalRenderStage, EDecalRenderTargetMode RenderTargetMode)
 {
-	FDecalRenderingCommon::ERenderTargetMode RenderTargetMode = IsMobileDeferredShadingEnabled(View.GetShaderPlatform()) ? 
-		FDecalRenderingCommon::RTM_SceneColorAndGBufferWithNormal : 
-		FDecalRenderingCommon::RTM_SceneColor;
-
 	FGraphicsPipelineStateInitializer GraphicsPSOInit;
 	RHICmdList.SetViewport(View.ViewRect.Min.X, View.ViewRect.Min.Y, 0, View.ViewRect.Max.X, View.ViewRect.Max.Y, 1);
 	RHICmdList.ApplyCachedRenderTargets(GraphicsPSOInit);
 
-	DrawDynamicMeshPass(View, RHICmdList, [&View, RenderTargetMode](FDynamicPassMeshDrawListContext* DynamicMeshPassContext)
+	DrawDynamicMeshPass(View, RHICmdList, [&View, DecalRenderStage, RenderTargetMode](FDynamicPassMeshDrawListContext* DynamicMeshPassContext)
 	{
 		FMeshDecalMeshProcessor PassMeshProcessor(
 			View.Family->Scene->GetRenderScene(),
 			&View,
-			DRS_Mobile,
+			DecalRenderStage,
 			RenderTargetMode,
 			DynamicMeshPassContext);
 
diff --git a/Engine/Source/Runtime/Renderer/Private/DecalRenderingCommon.cpp b/Engine/Source/Runtime/Renderer/Private/DecalRenderingCommon.cpp
new file mode 100644
--- /dev/null
+++ b/Engine/Source/Runtime/Renderer/Private/DecalRenderingCommon.cpp
@@ -0,0 +1,982 @@
+// Copyright Epic Games, Inc. All Rights Reserved.
+
+#include "DecalRenderingCommon.h"
+#include "RenderUtils.h"
+
+namespace DecalRendering
+{
+	// Legacy logic involving EDecalBlendMode.
+	// EDecalBlendMode will soon be replaced and this code will be removed. But submitting this intermediate step for later reference.
+	void InitBlendDesc(EShaderPlatform Platform, EDecalBlendMode DecalBlendMode, bool bWriteNormal, bool bWriteEmissive, FDecalBlendDesc& Desc)
+	{
+		const bool bIsMobilePlatform = IsMobilePlatform(Platform);
+		const bool bIsMobileDeferredPlatform = bIsMobilePlatform && IsMobileDeferredShadingEnabled(Platform);
+		const bool bIsDBufferPlatform = !bIsMobilePlatform && IsUsingDBuffers(Platform);
+		const bool bIsDBufferMaskPlatform = bIsDBufferPlatform && IsUsingPerPixelDBufferMask(Platform);
+		const bool bIsForwardPlatform = IsAnyForwardShadingEnabled(Platform);
+
+		// Convert DBuffer modes to GBuffer.
+		bool bWriteBaseColor = true;
+		bool bWriteRoughnessSpecularMetallic = true;
+
+		if (!bIsDBufferPlatform && !bIsMobilePlatform)
+		{
+			switch (DecalBlendMode)
+			{
+			case DBM_DBuffer_ColorNormalRoughness:
+				DecalBlendMode = DBM_Translucent;
+				break;
+			case DBM_DBuffer_Color:
+				DecalBlendMode = DBM_Translucent;
+				bWriteNormal = false;
+				bWriteRoughnessSpecularMetallic = false;
+				break;
+			case DBM_DBuffer_ColorNormal:
+				DecalBlendMode = DBM_Translucent;
+				bWriteRoughnessSpecularMetallic = false;
+				break;
+			case DBM_DBuffer_ColorRoughness:
+				DecalBlendMode = DBM_Translucent;
+				bWriteNormal = false;
+				break;
+			case DBM_DBuffer_NormalRoughness:
+				DecalBlendMode = DBM_Translucent;
+				bWriteBaseColor = false;
+				break;
+			case DBM_DBuffer_Normal:
+				DecalBlendMode = DBM_Translucent;
+				bWriteBaseColor = false;
+				bWriteRoughnessSpecularMetallic = false;
+				break;
+			}
+		}
+
+		// Convert GBuffer modes to DBuffer.
+		uint32 DBufferStageMask = 1 << (uint32)EDecalRenderStage::BeforeBassPass;
+		DBufferStageMask |= bWriteEmissive ? (1 << (uint32)EDecalRenderStage::Emissive) : 0;
+
+		if (bIsDBufferPlatform && bIsForwardPlatform && !bIsMobilePlatform)
+		{
+			switch (DecalBlendMode)
+			{
+			case DBM_Translucent:
+			case DBM_Stain:
+				DecalBlendMode = DBM_DBuffer_ColorNormalRoughness;
+				break;
+			case DBM_Normal:
+				DecalBlendMode = DBM_DBuffer_Normal;
+				DBufferStageMask = 1 << (uint32)EDecalRenderStage::BeforeBassPass;
+				break;
+			case DBM_Emissive:
+				DecalBlendMode = DBM_DBuffer_Emissive;
+				break;
+			case DBM_AlphaComposite:
+				DecalBlendMode = DBM_DBuffer_AlphaComposite;
+				break;
+			}
+		}
+
+		// Fill out FDecalBlendDesc.
+		switch (DecalBlendMode)
+		{
+		case DBM_AlphaComposite:
+			Desc.BlendMode = BLEND_AlphaComposite;
+			Desc.bWriteBaseColor = bWriteBaseColor;
+			Desc.bWriteRoughnessSpecularMetallic = bWriteRoughnessSpecularMetallic;
+			Desc.bWriteEmissive = bWriteEmissive;
+			Desc.RenderStageMask = 1 << (uint32)EDecalRenderStage::BeforeLighting;
+			break;
+		case DBM_Stain:
+			Desc.BlendMode = BLEND_Modulate;
+			Desc.bWriteBaseColor = bWriteBaseColor;
+			Desc.bWriteNormal = bWriteNormal;
+			Desc.bWriteRoughnessSpecularMetallic = bWriteRoughnessSpecularMetallic;
+			Desc.bWriteEmissive = bWriteEmissive;
+			Desc.RenderStageMask = 1 << (uint32)EDecalRenderStage::BeforeLighting;
+			break;
+		case DBM_Translucent:
+			Desc.BlendMode = BLEND_Translucent;
+			Desc.bWriteBaseColor = bWriteBaseColor;
+			Desc.bWriteNormal = bWriteNormal;
+			Desc.bWriteRoughnessSpecularMetallic = bWriteRoughnessSpecularMetallic;
+			Desc.bWriteEmissive = bWriteEmissive;
+			Desc.RenderStageMask = 1 << (uint32)EDecalRenderStage::BeforeLighting;
+			break;
+		case DBM_Normal:
+			Desc.BlendMode = BLEND_Translucent;
+			Desc.bWriteNormal = true;
+			Desc.RenderStageMask = 1 << (uint32)EDecalRenderStage::BeforeLighting;
+			break;
+		case DBM_Emissive:
+			Desc.BlendMode = BLEND_Translucent;
+			Desc.bWriteEmissive = true;
+			Desc.RenderStageMask = 1 << (uint32)EDecalRenderStage::BeforeLighting;
+			break;
+		case DBM_DBuffer_ColorNormalRoughness:
+			Desc.BlendMode = BLEND_Translucent;
+			Desc.bWriteBaseColor = Desc.bWriteNormal = Desc.bWriteRoughnessSpecularMetallic = true;
+			Desc.bWriteEmissive = bWriteEmissive;
+			Desc.bWriteDBufferMask = bIsDBufferMaskPlatform;
+			Desc.RenderStageMask = DBufferStageMask;
+			break;
+		case DBM_DBuffer_Color:
+			Desc.BlendMode = BLEND_Translucent;
+			Desc.bWriteBaseColor = true;
+			Desc.bWriteEmissive = bWriteEmissive;
+			Desc.bWriteDBufferMask = bIsDBufferMaskPlatform;
+			Desc.RenderStageMask = DBufferStageMask;
+			break;
+		case DBM_DBuffer_ColorNormal:
+			Desc.BlendMode = BLEND_Translucent;
+			Desc.bWriteBaseColor = Desc.bWriteNormal = true;
+			Desc.bWriteEmissive = bWriteEmissive;
+			Desc.bWriteDBufferMask = bIsDBufferMaskPlatform;
+			Desc.RenderStageMask = DBufferStageMask;
+			break;
+		case DBM_DBuffer_ColorRoughness:
+			Desc.BlendMode = BLEND_Translucent;
+			Desc.bWriteBaseColor = Desc.bWriteRoughnessSpecularMetallic = true;
+			Desc.bWriteEmissive = bWriteEmissive;
+			Desc.bWriteDBufferMask = bIsDBufferMaskPlatform;
+			Desc.RenderStageMask = DBufferStageMask;
+			break;
+		case DBM_DBuffer_Normal:
+			Desc.BlendMode = BLEND_Translucent;
+			Desc.bWriteNormal = true;
+			Desc.bWriteEmissive = bWriteEmissive;
+			Desc.bWriteDBufferMask = bIsDBufferMaskPlatform;
+			Desc.RenderStageMask = DBufferStageMask;
+			break;
+		case DBM_DBuffer_NormalRoughness:
+			Desc.BlendMode = BLEND_Translucent;
+			Desc.bWriteNormal = Desc.bWriteRoughnessSpecularMetallic = true;
+			Desc.bWriteEmissive = bWriteEmissive;
+			Desc.bWriteDBufferMask = bIsDBufferMaskPlatform;
+			Desc.RenderStageMask = DBufferStageMask;
+			break;
+		case DBM_DBuffer_Roughness:
+			Desc.BlendMode = BLEND_Translucent;
+			Desc.bWriteRoughnessSpecularMetallic = true;
+			Desc.bWriteEmissive = bWriteEmissive;
+			Desc.bWriteDBufferMask = bIsDBufferMaskPlatform;
+			Desc.RenderStageMask = DBufferStageMask;
+			break;
+		case DBM_DBuffer_Emissive:
+			Desc.BlendMode = BLEND_Translucent;
+			Desc.bWriteEmissive = true;
+			Desc.RenderStageMask = 1 << (uint32)EDecalRenderStage::Emissive;
+			break;
+		case DBM_DBuffer_AlphaComposite:
+			Desc.BlendMode = BLEND_AlphaComposite;
+			Desc.bWriteBaseColor = Desc.bWriteRoughnessSpecularMetallic = true;
+			Desc.bWriteEmissive = bWriteEmissive;
+			Desc.bWriteDBufferMask = bIsDBufferMaskPlatform;
+			Desc.RenderStageMask = DBufferStageMask;
+			break;
+		case DBM_Volumetric_DistanceFunction:
+			// Ignore
+			break;
+		case DBM_AmbientOcclusion:
+			Desc.BlendMode = BLEND_Translucent;
+			Desc.bWriteAmbientOcclusion = true;
+			Desc.RenderStageMask = 1 << (uint32)EDecalRenderStage::AmbientOcclusion;
+			break;
+		}
+
+		// Fixup for Mobile.
+		if (bIsMobileDeferredPlatform)
+		{
+			Desc.bWriteAmbientOcclusion = false;
+			Desc.bWriteDBufferMask = false;
+			Desc.RenderStageMask = Desc.bWriteEmissive || Desc.bWriteBaseColor || Desc.bWriteNormal || Desc.bWriteRoughnessSpecularMetallic ? 1 << (uint32)EDecalRenderStage::MobileBeforeLighting : 0;
+		}
+		else if (bIsMobilePlatform)
+		{
+			Desc.bWriteNormal = false;
+			Desc.bWriteRoughnessSpecularMetallic = false;
+			Desc.bWriteAmbientOcclusion = false;
+			Desc.bWriteDBufferMask = false;
+			Desc.BlendMode = Desc.bWriteEmissive ? BLEND_Translucent : Desc.BlendMode;
+			Desc.RenderStageMask = Desc.bWriteEmissive || Desc.bWriteBaseColor ? 1 << (uint32)EDecalRenderStage::Mobile : 0;
+		}
+	}
+
+	FDecalBlendDesc ComputeDecalBlendDesc(EShaderPlatform Platform, FMaterial const* Material)
+	{
+		FDecalBlendDesc Desc;
+		InitBlendDesc(Platform, (EDecalBlendMode)Material->GetDecalBlendMode(), Material->HasNormalConnected(), Material->HasEmissiveColorConnected(), Desc);
+		return Desc;
+	}
+
+	FDecalBlendDesc ComputeDecalBlendDesc(EShaderPlatform Platform, FMaterialShaderParameters const& MaterialShaderParameters)
+	{
+		FDecalBlendDesc Desc;
+		InitBlendDesc(Platform, (EDecalBlendMode)MaterialShaderParameters.DecalBlendMode, MaterialShaderParameters.bHasNormalConnected, MaterialShaderParameters.bHasEmissiveColorConnected, Desc);
+		return Desc;
+	}
+
+	bool IsCompatibleWithRenderStage(FDecalBlendDesc DecalBlendDesc, EDecalRenderStage DecalRenderStage)
+	{
+		return (DecalBlendDesc.RenderStageMask & (1 << (uint32)DecalRenderStage)) != 0;
+	}
+
+	EDecalRenderStage GetBaseRenderStage(FDecalBlendDesc DecalBlendDesc)
+	{
+		if (DecalBlendDesc.RenderStageMask & (1 << (uint32)EDecalRenderStage::BeforeBassPass))
+		{
+			return EDecalRenderStage::BeforeBassPass;
+		}
+		if (DecalBlendDesc.RenderStageMask & (1 << (uint32)EDecalRenderStage::BeforeLighting))
+		{
+			return EDecalRenderStage::BeforeLighting;
+		}
+		if (DecalBlendDesc.RenderStageMask & (1 << (uint32)EDecalRenderStage::Mobile))
+		{
+			return EDecalRenderStage::Mobile;
+		}
+		if (DecalBlendDesc.RenderStageMask & (1 << (uint32)EDecalRenderStage::MobileBeforeLighting))
+		{
+			return EDecalRenderStage::MobileBeforeLighting;
+		}
+
+		return EDecalRenderStage::None;
+	}
+
+	EDecalRenderTargetMode GetRenderTargetMode(FDecalBlendDesc DecalBlendDesc, EDecalRenderStage DecalRenderStage)
+	{
+		switch(DecalRenderStage)
+		{
+			case EDecalRenderStage::BeforeBassPass:
+				return EDecalRenderTargetMode::DBuffer;
+			case EDecalRenderStage::BeforeLighting:
+				return DecalBlendDesc.bWriteNormal ? EDecalRenderTargetMode::SceneColorAndGBuffer : EDecalRenderTargetMode::SceneColorAndGBufferNoNormal;
+			case EDecalRenderStage::Mobile:
+				return EDecalRenderTargetMode::SceneColor;
+			case EDecalRenderStage::MobileBeforeLighting:
+				return EDecalRenderTargetMode::SceneColorAndGBuffer;
+			case EDecalRenderStage::Emissive:
+				return EDecalRenderTargetMode::SceneColor;
+			case EDecalRenderStage::AmbientOcclusion:
+				return EDecalRenderTargetMode::AmbientOcclusion;
+		}
+
+		return EDecalRenderTargetMode::None;
+	}
+
+	uint32 GetRenderTargetCount(FDecalBlendDesc DecalBlendDesc, EDecalRenderTargetMode RenderTargetMode)
+	{
+		switch (RenderTargetMode)
+		{
+		case EDecalRenderTargetMode::DBuffer:
+			return DecalBlendDesc.bWriteDBufferMask ? 4 : 3;
+		case EDecalRenderTargetMode::SceneColorAndGBuffer:
+			return 4;
+		case EDecalRenderTargetMode::SceneColorAndGBufferNoNormal:
+			return 3;
+		case EDecalRenderTargetMode::SceneColor:
+			return 1;
+		case EDecalRenderTargetMode::AmbientOcclusion:
+			return 1;
+		}
+
+		return 0;
+	}
+
+	uint32 GetRenderTargetWriteMask(FDecalBlendDesc DecalBlendDesc, EDecalRenderStage DecalRenderStage, EDecalRenderTargetMode RenderTargetMode)
+	{
+		if (RenderTargetMode == EDecalRenderTargetMode::DBuffer)
+		{
+			return (DecalBlendDesc.bWriteBaseColor ? 1 : 0) + (DecalBlendDesc.bWriteNormal ? 2 : 0) + (DecalBlendDesc.bWriteRoughnessSpecularMetallic ? 4 : 0) + (DecalBlendDesc.bWriteDBufferMask ? 8 : 0);
+		}
+		else if (RenderTargetMode == EDecalRenderTargetMode::SceneColorAndGBuffer)
+		{
+			return (DecalBlendDesc.bWriteEmissive ? 1 : 0) + (DecalBlendDesc.bWriteNormal ? 2 : 0) + (DecalBlendDesc.bWriteRoughnessSpecularMetallic ? 4 : 0) + (DecalBlendDesc.bWriteBaseColor ? 8 : 0);
+		}
+		else if (RenderTargetMode == EDecalRenderTargetMode::SceneColorAndGBufferNoNormal)
+		{
+			return (DecalBlendDesc.bWriteEmissive ? 1 : 0) + (DecalBlendDesc.bWriteRoughnessSpecularMetallic ? 2 : 0) + (DecalBlendDesc.bWriteBaseColor ? 4 : 0);
+		}
+		else if (RenderTargetMode == EDecalRenderTargetMode::SceneColor)
+		{
+			if (DecalRenderStage == EDecalRenderStage::Mobile)
+			{
+				return ((DecalBlendDesc.bWriteEmissive || DecalBlendDesc.bWriteBaseColor) ? 1 : 0);
+			}
+			return (DecalBlendDesc.bWriteEmissive ? 1 : 0);
+		}
+		else if (RenderTargetMode == EDecalRenderTargetMode::AmbientOcclusion)
+		{
+			return (DecalBlendDesc.bWriteAmbientOcclusion ? 1 : 0);
+		}
+
+		// Enable all render targets by default.
+		return (1 << GetRenderTargetCount(DecalBlendDesc, RenderTargetMode)) - 1;
+	}
+
+	FRHIBlendState* GetDecalBlendState_DBuffer(FDecalBlendDesc DecalBlendDesc, EDecalRenderStage DecalRenderStage)
+	{
+		// Ignore DBuffer mask bit and always set that MRT active.
+		const uint32 Mask = GetRenderTargetWriteMask(DecalBlendDesc, DecalRenderStage, EDecalRenderTargetMode::DBuffer) & 0x7;
+
+		if (DecalBlendDesc.BlendMode == BLEND_Translucent)
+		{
+			switch (Mask)
+			{
+			case 1:
+				return TStaticBlendState<
+					CW_RGBA, BO_Add, BF_SourceAlpha, BF_InverseSourceAlpha, BO_Add, BF_Zero, BF_InverseSourceAlpha, // BaseColor
+					CW_NONE, BO_Add, BF_Zero, BF_One, BO_Add, BF_Zero, BF_One,										// Normal
+					CW_NONE, BO_Add, BF_Zero, BF_One, BO_Add, BF_Zero, BF_One,										// Metallic, Specular, Roughness
+					CW_RGBA, BO_Add, BF_One, BF_One, BO_Add, BF_One, BF_One											// DBuffer mask
+				>::GetRHI();
+			case 2:
+				return TStaticBlendState<
+					CW_NONE, BO_Add, BF_Zero, BF_One, BO_Add, BF_Zero, BF_One,										// BaseColor
+					CW_RGBA, BO_Add, BF_SourceAlpha, BF_InverseSourceAlpha, BO_Add, BF_Zero, BF_InverseSourceAlpha, // Normal
+					CW_NONE, BO_Add, BF_Zero, BF_One, BO_Add, BF_Zero, BF_One,										// Metallic, Specular, Roughness
+					CW_RGBA, BO_Add, BF_One, BF_One, BO_Add, BF_One, BF_One											// DBuffer mask
+				>::GetRHI();
+			case 3:
+				return TStaticBlendState<
+					CW_RGBA, BO_Add, BF_SourceAlpha, BF_InverseSourceAlpha, BO_Add, BF_Zero, BF_InverseSourceAlpha, // BaseColor
+					CW_RGBA, BO_Add, BF_SourceAlpha, BF_InverseSourceAlpha, BO_Add, BF_Zero, BF_InverseSourceAlpha,	// Normal
+					CW_NONE, BO_Add, BF_Zero, BF_One, BO_Add, BF_Zero, BF_One,										// Metallic, Specular, Roughness
+					CW_RGBA, BO_Add, BF_One, BF_One, BO_Add, BF_One, BF_One											// DBuffer mask
+				>::GetRHI();
+			case 4:
+				return TStaticBlendState<
+					CW_NONE, BO_Add, BF_Zero, BF_One, BO_Add, BF_Zero, BF_One,										// BaseColor
+					CW_NONE, BO_Add, BF_Zero, BF_One, BO_Add, BF_Zero, BF_One,										// Normal
+					CW_RGBA, BO_Add, BF_SourceAlpha, BF_InverseSourceAlpha, BO_Add, BF_Zero, BF_InverseSourceAlpha,	// Metallic, Specular, Roughness
+					CW_RGBA, BO_Add, BF_One, BF_One, BO_Add, BF_One, BF_One											// DBuffer mask
+				>::GetRHI();
+			case 5:
+				return TStaticBlendState<
+					CW_RGBA, BO_Add, BF_SourceAlpha, BF_InverseSourceAlpha, BO_Add, BF_Zero, BF_InverseSourceAlpha, // BaseColor
+					CW_NONE, BO_Add, BF_Zero, BF_One, BO_Add, BF_Zero, BF_One,										// Normal
+					CW_RGBA, BO_Add, BF_SourceAlpha, BF_InverseSourceAlpha, BO_Add, BF_Zero, BF_InverseSourceAlpha,	// Metallic, Specular, Roughness
+					CW_RGBA, BO_Add, BF_One, BF_One, BO_Add, BF_One, BF_One											// DBuffer mask
+				>::GetRHI();
+			case 6:
+				return TStaticBlendState<
+					CW_NONE, BO_Add, BF_Zero, BF_One, BO_Add, BF_Zero, BF_One,										// BaseColor
+					CW_RGBA, BO_Add, BF_SourceAlpha, BF_InverseSourceAlpha, BO_Add, BF_Zero, BF_InverseSourceAlpha,	// Normal
+					CW_RGBA, BO_Add, BF_SourceAlpha, BF_InverseSourceAlpha, BO_Add, BF_Zero, BF_InverseSourceAlpha,	// Metallic, Specular, Roughness
+					CW_RGBA, BO_Add, BF_One, BF_One, BO_Add, BF_One, BF_One											// DBuffer mask
+				>::GetRHI();
+			case 7:
+				return TStaticBlendState<
+					CW_RGBA, BO_Add, BF_SourceAlpha, BF_InverseSourceAlpha, BO_Add, BF_Zero, BF_InverseSourceAlpha,	// BaseColor
+					CW_RGBA, BO_Add, BF_SourceAlpha, BF_InverseSourceAlpha, BO_Add, BF_Zero, BF_InverseSourceAlpha,	// Normal
+					CW_RGBA, BO_Add, BF_SourceAlpha, BF_InverseSourceAlpha, BO_Add, BF_Zero, BF_InverseSourceAlpha,	// Metallic, Specular, Roughness
+					CW_RGBA, BO_Add, BF_One, BF_One, BO_Add, BF_One, BF_One											// DBuffer mask
+				>::GetRHI();
+			}
+		}
+		else if (DecalBlendDesc.BlendMode == BLEND_AlphaComposite)
+		{
+			ensure((Mask & 2) == 0); // AlphaComposite shouldn't write normal.
+
+			switch (Mask)
+			{
+			case 1:
+				return TStaticBlendState<
+					CW_RGBA, BO_Add, BF_One, BF_InverseSourceAlpha, BO_Add, BF_Zero, BF_InverseSourceAlpha,			// BaseColor
+					CW_NONE, BO_Add, BF_Zero, BF_One, BO_Add, BF_Zero, BF_One,										// Normal
+					CW_NONE, BO_Add, BF_Zero, BF_One, BO_Add, BF_Zero, BF_One,										// Metallic, Specular, Roughness
+					CW_RGBA, BO_Add, BF_One, BF_One, BO_Add, BF_One, BF_One											// DBuffer mask
+				>::GetRHI();
+			case 4:
+				return TStaticBlendState<
+					CW_NONE, BO_Add, BF_Zero, BF_One, BO_Add, BF_Zero, BF_One,										// BaseColor
+					CW_NONE, BO_Add, BF_Zero, BF_One, BO_Add, BF_Zero, BF_One,										// Normal
+					CW_RGBA, BO_Add, BF_One, BF_InverseSourceAlpha, BO_Add, BF_Zero, BF_InverseSourceAlpha,			// Metallic, Specular, Roughness
+					CW_RGBA, BO_Add, BF_One, BF_One, BO_Add, BF_One, BF_One											// DBuffer mask
+				>::GetRHI();
+			case 5:
+				return TStaticBlendState<
+					CW_RGBA, BO_Add, BF_One, BF_InverseSourceAlpha, BO_Add, BF_Zero, BF_InverseSourceAlpha,			// BaseColor
+					CW_NONE, BO_Add, BF_Zero, BF_One, BO_Add, BF_Zero, BF_One,										// Normal
+					CW_RGBA, BO_Add, BF_One, BF_InverseSourceAlpha, BO_Add, BF_Zero, BF_InverseSourceAlpha,			// Metallic, Specular, Roughness
+					CW_RGBA, BO_Add, BF_One, BF_One, BO_Add, BF_One, BF_One											// DBuffer mask
+				>::GetRHI();
+			}
+		}
+
+		ensure(0);
+		return TStaticBlendState<>::GetRHI();
+	}
+
+	FRHIBlendState* GetDecalBlendState_SceneColorAndGBuffer(FDecalBlendDesc DecalBlendDesc, EDecalRenderStage DecalRenderStage)
+	{
+		const uint32 Mask = GetRenderTargetWriteMask(DecalBlendDesc, DecalRenderStage, EDecalRenderTargetMode::SceneColorAndGBuffer);
+
+		if (DecalBlendDesc.BlendMode == BLEND_Translucent)
+		{
+			switch (Mask)
+			{
+			case 1:
+				return TStaticBlendState<
+					CW_RGB, BO_Add, BF_SourceAlpha, BF_One, BO_Add, BF_Zero, BF_One,				// Emissive
+					CW_NONE, BO_Add, BF_Zero, BF_One, BO_Add, BF_Zero, BF_One,						// Normal
+					CW_NONE, BO_Add, BF_Zero, BF_One, BO_Add, BF_Zero, BF_One,						// Metallic, Specular, Roughness
+					CW_NONE, BO_Add, BF_Zero, BF_One, BO_Add, BF_Zero, BF_One						// BaseColor
+				>::GetRHI();
+			case 2:
+				return TStaticBlendState<
+					CW_NONE, BO_Add, BF_Zero, BF_One, BO_Add, BF_Zero, BF_One,						// Emissive
+					CW_RGB, BO_Add, BF_SourceAlpha, BF_InverseSourceAlpha, BO_Add, BF_Zero, BF_One,	// Normal
+					CW_NONE, BO_Add, BF_Zero, BF_One, BO_Add, BF_Zero, BF_One,						// Metallic, Specular, Roughness
+					CW_NONE, BO_Add, BF_Zero, BF_One, BO_Add, BF_Zero, BF_One						// BaseColor
+				>::GetRHI();
+			case 3:
+				return TStaticBlendState<
+					CW_RGB, BO_Add, BF_SourceAlpha, BF_One, BO_Add, BF_Zero, BF_One,				// Emissive
+					CW_RGB, BO_Add, BF_SourceAlpha, BF_InverseSourceAlpha, BO_Add, BF_Zero, BF_One,	// Normal
+					CW_NONE, BO_Add, BF_Zero, BF_One, BO_Add, BF_Zero, BF_One,						// Metallic, Specular, Roughness
+					CW_NONE, BO_Add, BF_Zero, BF_One, BO_Add, BF_Zero, BF_One						// BaseColor
+				>::GetRHI();
+			case 4:
+				return TStaticBlendState<
+					CW_NONE, BO_Add, BF_Zero, BF_One, BO_Add, BF_Zero, BF_One,						// Emissive
+					CW_NONE, BO_Add, BF_Zero, BF_One, BO_Add, BF_Zero, BF_One,						// Normal
+					CW_RGB, BO_Add, BF_SourceAlpha, BF_InverseSourceAlpha, BO_Add, BF_Zero, BF_One,	// Metallic, Specular, Roughness
+					CW_NONE, BO_Add, BF_Zero, BF_One, BO_Add, BF_Zero, BF_One						// BaseColor
+				>::GetRHI();
+			case 5:
+				return TStaticBlendState<
+					CW_RGB, BO_Add, BF_SourceAlpha, BF_One, BO_Add, BF_Zero, BF_One,				// Emissive
+					CW_NONE, BO_Add, BF_Zero, BF_One, BO_Add, BF_Zero, BF_One,						// Normal
+					CW_RGB, BO_Add, BF_SourceAlpha, BF_InverseSourceAlpha, BO_Add, BF_Zero, BF_One,	// Metallic, Specular, Roughness
+					CW_NONE, BO_Add, BF_Zero, BF_One, BO_Add, BF_Zero, BF_One						// BaseColor
+				>::GetRHI();
+			case 6:
+				return TStaticBlendState<
+					CW_NONE, BO_Add, BF_Zero, BF_One, BO_Add, BF_Zero, BF_One,						// Emissive
+					CW_RGB, BO_Add, BF_SourceAlpha, BF_InverseSourceAlpha, BO_Add, BF_Zero, BF_One,	// Normal
+					CW_RGB, BO_Add, BF_SourceAlpha, BF_InverseSourceAlpha, BO_Add, BF_Zero, BF_One,	// Metallic, Specular, Roughness
+					CW_NONE, BO_Add, BF_Zero, BF_One, BO_Add, BF_Zero, BF_One						// BaseColor
+				>::GetRHI();
+			case 7:
+				return TStaticBlendState<
+					CW_RGB, BO_Add, BF_SourceAlpha, BF_One, BO_Add, BF_Zero, BF_One,				// Emissive
+					CW_RGB, BO_Add, BF_SourceAlpha, BF_InverseSourceAlpha, BO_Add, BF_Zero, BF_One,	// Normal
+					CW_RGB, BO_Add, BF_SourceAlpha, BF_InverseSourceAlpha, BO_Add, BF_Zero, BF_One,	// Metallic, Specular, Roughness
+					CW_NONE, BO_Add, BF_Zero, BF_One, BO_Add, BF_Zero, BF_One						// BaseColor
+				>::GetRHI();
+			case 8:
+				return TStaticBlendState<
+					CW_NONE, BO_Add, BF_Zero, BF_One, BO_Add, BF_Zero, BF_One,						// Emissive
+					CW_NONE, BO_Add, BF_Zero, BF_One, BO_Add, BF_Zero, BF_One,						// Normal
+					CW_NONE, BO_Add, BF_Zero, BF_One, BO_Add, BF_Zero, BF_One,						// Metallic, Specular, Roughness
+					CW_RGB, BO_Add, BF_SourceAlpha, BF_InverseSourceAlpha, BO_Add, BF_Zero, BF_One	// BaseColor
+				>::GetRHI();
+			case 9:
+				return TStaticBlendState<
+					CW_RGB, BO_Add, BF_SourceAlpha, BF_One, BO_Add, BF_Zero, BF_One,				// Emissive
+					CW_NONE, BO_Add, BF_Zero, BF_One, BO_Add, BF_Zero, BF_One,						// Normal
+					CW_NONE, BO_Add, BF_Zero, BF_One, BO_Add, BF_Zero, BF_One,						// Metallic, Specular, Roughness
+					CW_RGB, BO_Add, BF_SourceAlpha, BF_InverseSourceAlpha, BO_Add, BF_Zero, BF_One	// BaseColor
+				>::GetRHI();
+			case 10:
+				return TStaticBlendState<
+					CW_NONE, BO_Add, BF_Zero, BF_One, BO_Add, BF_Zero, BF_One,						// Emissive
+					CW_RGB, BO_Add, BF_SourceAlpha, BF_InverseSourceAlpha, BO_Add, BF_Zero, BF_One,	// Normal
+					CW_NONE, BO_Add, BF_Zero, BF_One, BO_Add, BF_Zero, BF_One,						// Metallic, Specular, Roughness
+					CW_RGB, BO_Add, BF_SourceAlpha, BF_InverseSourceAlpha, BO_Add, BF_Zero, BF_One	// BaseColor
+				>::GetRHI();
+			case 11:
+				return TStaticBlendState<
+					CW_RGB, BO_Add, BF_SourceAlpha, BF_One, BO_Add, BF_Zero, BF_One,				// Emissive
+					CW_RGB, BO_Add, BF_SourceAlpha, BF_InverseSourceAlpha, BO_Add, BF_Zero, BF_One,	// Normal
+					CW_NONE, BO_Add, BF_Zero, BF_One, BO_Add, BF_Zero, BF_One,						// Metallic, Specular, Roughness
+					CW_RGB, BO_Add, BF_SourceAlpha, BF_InverseSourceAlpha, BO_Add, BF_Zero, BF_One	// BaseColor
+				>::GetRHI();
+			case 12:
+				return TStaticBlendState<
+					CW_NONE, BO_Add, BF_Zero, BF_One, BO_Add, BF_Zero, BF_One,						// Emissive
+					CW_NONE, BO_Add, BF_Zero, BF_One, BO_Add, BF_Zero, BF_One,						// Normal
+					CW_RGB, BO_Add, BF_SourceAlpha, BF_InverseSourceAlpha, BO_Add, BF_Zero, BF_One,	// Metallic, Specular, Roughness
+					CW_RGB, BO_Add, BF_SourceAlpha, BF_InverseSourceAlpha, BO_Add, BF_Zero, BF_One	// BaseColor
+				>::GetRHI();
+			case 13:
+				return TStaticBlendState<
+					CW_RGB, BO_Add, BF_SourceAlpha, BF_One, BO_Add, BF_Zero, BF_One,				// Emissive
+					CW_NONE, BO_Add, BF_Zero, BF_One, BO_Add, BF_Zero, BF_One,						// Normal
+					CW_RGB, BO_Add, BF_SourceAlpha, BF_InverseSourceAlpha, BO_Add, BF_Zero, BF_One,	// Metallic, Specular, Roughness
+					CW_RGB, BO_Add, BF_SourceAlpha, BF_InverseSourceAlpha, BO_Add, BF_Zero, BF_One	// BaseColor
+				>::GetRHI();
+			case 14:
+				return TStaticBlendState<
+					CW_NONE, BO_Add, BF_Zero, BF_One, BO_Add, BF_Zero, BF_One,						// Emissive
+					CW_RGB, BO_Add, BF_SourceAlpha, BF_InverseSourceAlpha, BO_Add, BF_Zero, BF_One,	// Normal
+					CW_RGB, BO_Add, BF_SourceAlpha, BF_InverseSourceAlpha, BO_Add, BF_Zero, BF_One,	// Metallic, Specular, Roughness
+					CW_RGB, BO_Add, BF_SourceAlpha, BF_InverseSourceAlpha, BO_Add, BF_Zero, BF_One	// BaseColor
+				>::GetRHI();
+			case 15:
+				return TStaticBlendState<
+					CW_RGB, BO_Add, BF_SourceAlpha, BF_One, BO_Add, BF_Zero, BF_One,				// Emissive
+					CW_RGB, BO_Add, BF_SourceAlpha, BF_InverseSourceAlpha, BO_Add, BF_Zero, BF_One,	// Normal
+					CW_RGB, BO_Add, BF_SourceAlpha, BF_InverseSourceAlpha, BO_Add, BF_Zero, BF_One,	// Metallic, Specular, Roughness
+					CW_RGB, BO_Add, BF_SourceAlpha, BF_InverseSourceAlpha, BO_Add, BF_Zero, BF_One	// BaseColor
+				>::GetRHI();
+			}
+		}
+		else if (DecalBlendDesc.BlendMode == BLEND_AlphaComposite)
+		{
+			ensure((Mask & 2) == 0); // AlphaComposite shouldn't write normal.
+
+			switch (Mask)
+			{
+			case 1:
+				return TStaticBlendState<
+					CW_RGB, BO_Add, BF_SourceAlpha, BF_One, BO_Add, BF_Zero, BF_One,				// Emissive
+					CW_NONE, BO_Add, BF_Zero, BF_One, BO_Add, BF_Zero, BF_One,						// Normal
+					CW_NONE, BO_Add, BF_Zero, BF_One, BO_Add, BF_Zero, BF_One,						// Metallic, Specular, Roughness
+					CW_NONE, BO_Add, BF_Zero, BF_One, BO_Add, BF_Zero, BF_One						// BaseColor
+				>::GetRHI();
+			case 4:
+				return TStaticBlendState<
+					CW_NONE, BO_Add, BF_Zero, BF_One, BO_Add, BF_Zero, BF_One,						// Emissive
+					CW_NONE, BO_Add, BF_Zero, BF_One, BO_Add, BF_Zero, BF_One,						// Normal
+					CW_RGB, BO_Add, BF_One, BF_InverseSourceAlpha, BO_Add, BF_Zero, BF_One,			// Metallic, Specular, Roughness
+					CW_NONE, BO_Add, BF_Zero, BF_One, BO_Add, BF_Zero, BF_One						// BaseColor
+				>::GetRHI();
+			case 5:
+				return TStaticBlendState<
+					CW_RGB, BO_Add, BF_SourceAlpha, BF_One, BO_Add, BF_Zero, BF_One,				// Emissive
+					CW_NONE, BO_Add, BF_Zero, BF_One, BO_Add, BF_Zero, BF_One,						// Normal
+					CW_RGB, BO_Add, BF_One, BF_InverseSourceAlpha, BO_Add, BF_Zero, BF_One,			// Metallic, Specular, Roughness
+					CW_NONE, BO_Add, BF_Zero, BF_One, BO_Add, BF_Zero, BF_One						// BaseColor
+				>::GetRHI();
+			case 8:
+				return TStaticBlendState<
+					CW_NONE, BO_Add, BF_Zero, BF_One, BO_Add, BF_Zero, BF_One,						// Emissive
+					CW_NONE, BO_Add, BF_Zero, BF_One, BO_Add, BF_Zero, BF_One,						// Normal
+					CW_NONE, BO_Add, BF_Zero, BF_One, BO_Add, BF_Zero, BF_One,						// Metallic, Specular, Roughness
+					CW_RGB, BO_Add, BF_One, BF_InverseSourceAlpha, BO_Add, BF_Zero, BF_One			// BaseColor
+				>::GetRHI();
+			case 9:
+				return TStaticBlendState<
+					CW_RGB, BO_Add, BF_SourceAlpha, BF_One, BO_Add, BF_Zero, BF_One,				// Emissive
+					CW_NONE, BO_Add, BF_Zero, BF_One, BO_Add, BF_Zero, BF_One,						// Normal
+					CW_NONE, BO_Add, BF_Zero, BF_One, BO_Add, BF_Zero, BF_One,						// Metallic, Specular, Roughness
+					CW_RGB, BO_Add, BF_One, BF_InverseSourceAlpha, BO_Add, BF_Zero, BF_One			// BaseColor
+				>::GetRHI();
+			case 12:
+				return TStaticBlendState<
+					CW_NONE, BO_Add, BF_Zero, BF_One, BO_Add, BF_Zero, BF_One,						// Emissive
+					CW_NONE, BO_Add, BF_Zero, BF_One, BO_Add, BF_Zero, BF_One,						// Normal
+					CW_RGB, BO_Add, BF_One, BF_InverseSourceAlpha, BO_Add, BF_Zero, BF_One,			// Metallic, Specular, Roughness
+					CW_RGB, BO_Add, BF_One, BF_InverseSourceAlpha, BO_Add, BF_Zero, BF_One			// BaseColor
+				>::GetRHI();
+			case 13:
+				return TStaticBlendState<
+					CW_RGB, BO_Add, BF_SourceAlpha, BF_One, BO_Add, BF_Zero, BF_One,				// Emissive
+					CW_NONE, BO_Add, BF_Zero, BF_One, BO_Add, BF_Zero, BF_One,						// Normal
+					CW_RGB, BO_Add, BF_One, BF_InverseSourceAlpha, BO_Add, BF_Zero, BF_One,			// Metallic, Specular, Roughness
+					CW_RGB, BO_Add, BF_One, BF_InverseSourceAlpha, BO_Add, BF_Zero, BF_One			// BaseColor
+				>::GetRHI();
+			}
+		}
+		else if (DecalBlendDesc.BlendMode == BLEND_Modulate)
+		{
+			switch (Mask)
+			{
+			case 1:
+				return TStaticBlendState<
+					CW_RGB, BO_Add, BF_SourceAlpha, BF_One, BO_Add, BF_Zero, BF_One,				// Emissive
+					CW_NONE, BO_Add, BF_Zero, BF_One, BO_Add, BF_Zero, BF_One,						// Normal
+					CW_NONE, BO_Add, BF_Zero, BF_One, BO_Add, BF_Zero, BF_One,						// Metallic, Specular, Roughness
+					CW_NONE, BO_Add, BF_Zero, BF_One, BO_Add, BF_Zero, BF_One						// BaseColor
+				>::GetRHI();
+			case 2:
+				return TStaticBlendState<
+					CW_NONE, BO_Add, BF_Zero, BF_One, BO_Add, BF_Zero, BF_One,						// Emissive
+					CW_RGB, BO_Add, BF_SourceAlpha, BF_InverseSourceAlpha, BO_Add, BF_Zero, BF_One,	// Normal
+					CW_NONE, BO_Add, BF_Zero, BF_One, BO_Add, BF_Zero, BF_One,						// Metallic, Specular, Roughness
+					CW_NONE, BO_Add, BF_Zero, BF_One, BO_Add, BF_Zero, BF_One						// BaseColor
+				>::GetRHI();
+			case 3:
+				return TStaticBlendState<
+					CW_RGB, BO_Add, BF_SourceAlpha, BF_One, BO_Add, BF_Zero, BF_One,				// Emissive
+					CW_RGB, BO_Add, BF_SourceAlpha, BF_InverseSourceAlpha, BO_Add, BF_Zero, BF_One,	// Normal
+					CW_NONE, BO_Add, BF_Zero, BF_One, BO_Add, BF_Zero, BF_One,						// Metallic, Specular, Roughness
+					CW_NONE, BO_Add, BF_Zero, BF_One, BO_Add, BF_Zero, BF_One						// BaseColor
+				>::GetRHI();
+			case 4:
+				return TStaticBlendState<
+					CW_NONE, BO_Add, BF_Zero, BF_One, BO_Add, BF_Zero, BF_One,						// Emissive
+					CW_NONE, BO_Add, BF_Zero, BF_One, BO_Add, BF_Zero, BF_One,						// Normal
+					CW_RGB, BO_Add, BF_SourceAlpha, BF_InverseSourceAlpha, BO_Add, BF_Zero, BF_One,	// Metallic, Specular, Roughness
+					CW_NONE, BO_Add, BF_Zero, BF_One, BO_Add, BF_Zero, BF_One						// BaseColor
+				>::GetRHI();
+			case 5:
+				return TStaticBlendState<
+					CW_RGB, BO_Add, BF_SourceAlpha, BF_One, BO_Add, BF_Zero, BF_One,				// Emissive
+					CW_NONE, BO_Add, BF_Zero, BF_One, BO_Add, BF_Zero, BF_One,						// Normal
+					CW_RGB, BO_Add, BF_SourceAlpha, BF_InverseSourceAlpha, BO_Add, BF_Zero, BF_One,	// Metallic, Specular, Roughness
+					CW_NONE, BO_Add, BF_Zero, BF_One, BO_Add, BF_Zero, BF_One						// BaseColor
+				>::GetRHI();
+			case 6:
+				return TStaticBlendState<
+					CW_NONE, BO_Add, BF_Zero, BF_One, BO_Add, BF_Zero, BF_One,						// Emissive
+					CW_RGB, BO_Add, BF_SourceAlpha, BF_InverseSourceAlpha, BO_Add, BF_Zero, BF_One,	// Normal
+					CW_RGB, BO_Add, BF_SourceAlpha, BF_InverseSourceAlpha, BO_Add, BF_Zero, BF_One,	// Metallic, Specular, Roughness
+					CW_NONE, BO_Add, BF_Zero, BF_One, BO_Add, BF_Zero, BF_One						// BaseColor
+				>::GetRHI();
+			case 7:
+				return TStaticBlendState<
+					CW_RGB, BO_Add, BF_SourceAlpha, BF_One, BO_Add, BF_Zero, BF_One,				// Emissive
+					CW_RGB, BO_Add, BF_SourceAlpha, BF_InverseSourceAlpha, BO_Add, BF_Zero, BF_One,	// Normal
+					CW_RGB, BO_Add, BF_SourceAlpha, BF_InverseSourceAlpha, BO_Add, BF_Zero, BF_One,	// Metallic, Specular, Roughness
+					CW_NONE, BO_Add, BF_Zero, BF_One, BO_Add, BF_Zero, BF_One						// BaseColor
+				>::GetRHI();
+			case 8:
+				return TStaticBlendState<
+					CW_NONE, BO_Add, BF_Zero, BF_One, BO_Add, BF_Zero, BF_One,						// Emissive
+					CW_NONE, BO_Add, BF_Zero, BF_One, BO_Add, BF_Zero, BF_One,						// Normal
+					CW_NONE, BO_Add, BF_Zero, BF_One, BO_Add, BF_Zero, BF_One,						// Metallic, Specular, Roughness
+					CW_RGB, BO_Add, BF_DestColor, BF_InverseSourceAlpha, BO_Add, BF_Zero, BF_One	// BaseColor
+				>::GetRHI();
+			case 9:
+				return TStaticBlendState<
+					CW_RGB, BO_Add, BF_SourceAlpha, BF_One, BO_Add, BF_Zero, BF_One,				// Emissive
+					CW_NONE, BO_Add, BF_Zero, BF_One, BO_Add, BF_Zero, BF_One,						// Normal
+					CW_NONE, BO_Add, BF_Zero, BF_One, BO_Add, BF_Zero, BF_One,						// Metallic, Specular, Roughness
+					CW_RGB, BO_Add, BF_DestColor, BF_InverseSourceAlpha, BO_Add, BF_Zero, BF_One	// BaseColor
+				>::GetRHI();
+			case 10:
+				return TStaticBlendState<
+					CW_NONE, BO_Add, BF_Zero, BF_One, BO_Add, BF_Zero, BF_One,						// Emissive
+					CW_RGB, BO_Add, BF_SourceAlpha, BF_InverseSourceAlpha, BO_Add, BF_Zero, BF_One,	// Normal
+					CW_NONE, BO_Add, BF_Zero, BF_One, BO_Add, BF_Zero, BF_One,						// Metallic, Specular, Roughness
+					CW_RGB, BO_Add, BF_DestColor, BF_InverseSourceAlpha, BO_Add, BF_Zero, BF_One	// BaseColor
+				>::GetRHI();
+			case 11:
+				return TStaticBlendState<
+					CW_RGB, BO_Add, BF_SourceAlpha, BF_One, BO_Add, BF_Zero, BF_One,				// Emissive
+					CW_RGB, BO_Add, BF_SourceAlpha, BF_InverseSourceAlpha, BO_Add, BF_Zero, BF_One,	// Normal
+					CW_NONE, BO_Add, BF_Zero, BF_One, BO_Add, BF_Zero, BF_One,						// Metallic, Specular, Roughness
+					CW_RGB, BO_Add, BF_DestColor, BF_InverseSourceAlpha, BO_Add, BF_Zero, BF_One	// BaseColor
+				>::GetRHI();
+			case 12:
+				return TStaticBlendState<
+					CW_NONE, BO_Add, BF_Zero, BF_One, BO_Add, BF_Zero, BF_One,						// Emissive
+					CW_NONE, BO_Add, BF_Zero, BF_One, BO_Add, BF_Zero, BF_One,						// Normal
+					CW_RGB, BO_Add, BF_SourceAlpha, BF_InverseSourceAlpha, BO_Add, BF_Zero, BF_One,	// Metallic, Specular, Roughness
+					CW_RGB, BO_Add, BF_DestColor, BF_InverseSourceAlpha, BO_Add, BF_Zero, BF_One	// BaseColor
+				>::GetRHI();
+			case 13:
+				return TStaticBlendState<
+					CW_RGB, BO_Add, BF_SourceAlpha, BF_One, BO_Add, BF_Zero, BF_One,				// Emissive
+					CW_NONE, BO_Add, BF_Zero, BF_One, BO_Add, BF_Zero, BF_One,						// Normal
+					CW_RGB, BO_Add, BF_SourceAlpha, BF_InverseSourceAlpha, BO_Add, BF_Zero, BF_One,	// Metallic, Specular, Roughness
+					CW_RGB, BO_Add, BF_DestColor, BF_InverseSourceAlpha, BO_Add, BF_Zero, BF_One	// BaseColor
+				>::GetRHI();
+			case 14:
+				return TStaticBlendState<
+					CW_NONE, BO_Add, BF_Zero, BF_One, BO_Add, BF_Zero, BF_One,						// Emissive
+					CW_RGB, BO_Add, BF_SourceAlpha, BF_InverseSourceAlpha, BO_Add, BF_Zero, BF_One,	// Normal
+					CW_RGB, BO_Add, BF_SourceAlpha, BF_InverseSourceAlpha, BO_Add, BF_Zero, BF_One,	// Metallic, Specular, Roughness
+					CW_RGB, BO_Add, BF_DestColor, BF_InverseSourceAlpha, BO_Add, BF_Zero, BF_One	// BaseColor
+				>::GetRHI();
+			case 15:
+				return TStaticBlendState<
+					CW_RGB, BO_Add, BF_SourceAlpha, BF_One, BO_Add, BF_Zero, BF_One,				// Emissive
+					CW_RGB, BO_Add, BF_SourceAlpha, BF_InverseSourceAlpha, BO_Add, BF_Zero, BF_One,	// Normal
+					CW_RGB, BO_Add, BF_SourceAlpha, BF_InverseSourceAlpha, BO_Add, BF_Zero, BF_One,	// Metallic, Specular, Roughness
+					CW_RGB, BO_Add, BF_DestColor, BF_InverseSourceAlpha, BO_Add, BF_Zero, BF_One	// BaseColor
+				>::GetRHI();
+			}
+		}
+
+		ensure(0);
+		return TStaticBlendState<>::GetRHI();
+	}
+
+	FRHIBlendState* GetDecalBlendState_SceneColorAndGBufferNoNormal(FDecalBlendDesc DecalBlendDesc, EDecalRenderStage DecalRenderStage)
+	{
+		const uint32 Mask = GetRenderTargetWriteMask(DecalBlendDesc, DecalRenderStage, EDecalRenderTargetMode::SceneColorAndGBufferNoNormal);
+
+		if (DecalBlendDesc.BlendMode == BLEND_Translucent)
+		{
+			switch (Mask)
+			{
+			case 1:
+				return TStaticBlendState<
+					CW_RGB, BO_Add, BF_SourceAlpha, BF_One, BO_Add, BF_Zero, BF_One,				// Emissive
+					CW_NONE, BO_Add, BF_Zero, BF_One, BO_Add, BF_Zero, BF_One,						// Metallic, Specular, Roughness
+					CW_NONE, BO_Add, BF_Zero, BF_One, BO_Add, BF_Zero, BF_One						// BaseColor
+				>::GetRHI();
+			case 2:
+				return TStaticBlendState<
+					CW_NONE, BO_Add, BF_Zero, BF_One, BO_Add, BF_Zero, BF_One,						// Emissive
+					CW_RGB, BO_Add, BF_SourceAlpha, BF_InverseSourceAlpha, BO_Add, BF_Zero, BF_One,	// Metallic, Specular, Roughness
+					CW_NONE, BO_Add, BF_Zero, BF_One, BO_Add, BF_Zero, BF_One						// BaseColor
+				>::GetRHI();
+			case 3:
+				return TStaticBlendState<
+					CW_RGB, BO_Add, BF_SourceAlpha, BF_One, BO_Add, BF_Zero, BF_One,				// Emissive
+					CW_RGB, BO_Add, BF_SourceAlpha, BF_InverseSourceAlpha, BO_Add, BF_Zero, BF_One,	// Metallic, Specular, Roughness
+					CW_NONE, BO_Add, BF_Zero, BF_One, BO_Add, BF_Zero, BF_One						// BaseColor
+				>::GetRHI();
+			case 4:
+				return TStaticBlendState<
+					CW_NONE, BO_Add, BF_Zero, BF_One, BO_Add, BF_Zero, BF_One,						// Emissive
+					CW_NONE, BO_Add, BF_Zero, BF_One, BO_Add, BF_Zero, BF_One,						// Metallic, Specular, Roughness
+					CW_RGB, BO_Add, BF_SourceAlpha, BF_InverseSourceAlpha, BO_Add, BF_Zero, BF_One	// BaseColor
+				>::GetRHI();
+			case 5:
+				return TStaticBlendState<
+					CW_RGB, BO_Add, BF_SourceAlpha, BF_One, BO_Add, BF_Zero, BF_One,				// Emissive
+					CW_NONE, BO_Add, BF_Zero, BF_One, BO_Add, BF_Zero, BF_One,						// Metallic, Specular, Roughness
+					CW_RGB, BO_Add, BF_SourceAlpha, BF_InverseSourceAlpha, BO_Add, BF_Zero, BF_One	// BaseColor
+				>::GetRHI();
+			case 6:
+				return TStaticBlendState<
+					CW_NONE, BO_Add, BF_Zero, BF_One, BO_Add, BF_Zero, BF_One,						// Emissive
+					CW_RGB, BO_Add, BF_SourceAlpha, BF_InverseSourceAlpha, BO_Add, BF_Zero, BF_One,	// Metallic, Specular, Roughness
+					CW_RGB, BO_Add, BF_SourceAlpha, BF_InverseSourceAlpha, BO_Add, BF_Zero, BF_One	// BaseColor
+				>::GetRHI();
+			case 7:
+				return TStaticBlendState<
+					CW_RGB, BO_Add, BF_SourceAlpha, BF_One, BO_Add, BF_Zero, BF_One,				// Emissive
+					CW_RGB, BO_Add, BF_SourceAlpha, BF_InverseSourceAlpha, BO_Add, BF_Zero, BF_One,	// Metallic, Specular, Roughness
+					CW_RGB, BO_Add, BF_SourceAlpha, BF_InverseSourceAlpha, BO_Add, BF_Zero, BF_One	// BaseColor
+				>::GetRHI();
+			}
+		}
+		else if (DecalBlendDesc.BlendMode == BLEND_AlphaComposite)
+		{
+			switch (Mask)
+			{
+			case 1:
+				return TStaticBlendState<
+					CW_RGB, BO_Add, BF_SourceAlpha, BF_One, BO_Add, BF_Zero, BF_One,				// Emissive
+					CW_NONE, BO_Add, BF_Zero, BF_One, BO_Add, BF_Zero, BF_One,						// Metallic, Specular, Roughness
+					CW_NONE, BO_Add, BF_Zero, BF_One, BO_Add, BF_Zero, BF_One						// BaseColor
+				>::GetRHI();
+			case 2:
+				return TStaticBlendState<
+					CW_NONE, BO_Add, BF_Zero, BF_One, BO_Add, BF_Zero, BF_One,						// Emissive
+					CW_RGB, BO_Add, BF_One, BF_InverseSourceAlpha, BO_Add, BF_Zero, BF_One,			// Metallic, Specular, Roughness
+					CW_NONE, BO_Add, BF_Zero, BF_One, BO_Add, BF_Zero, BF_One						// BaseColor
+				>::GetRHI();
+			case 3:
+				return TStaticBlendState<
+					CW_RGB, BO_Add, BF_SourceAlpha, BF_One, BO_Add, BF_Zero, BF_One,				// Emissive
+					CW_RGB, BO_Add, BF_One, BF_InverseSourceAlpha, BO_Add, BF_Zero, BF_One,			// Metallic, Specular, Roughness
+					CW_NONE, BO_Add, BF_Zero, BF_One, BO_Add, BF_Zero, BF_One						// BaseColor
+				>::GetRHI();
+			case 4:
+				return TStaticBlendState<
+					CW_NONE, BO_Add, BF_Zero, BF_One, BO_Add, BF_Zero, BF_One,						// Emissive
+					CW_NONE, BO_Add, BF_Zero, BF_One, BO_Add, BF_Zero, BF_One,						// Metallic, Specular, Roughness
+					CW_RGB, BO_Add, BF_One, BF_InverseSourceAlpha, BO_Add, BF_Zero, BF_One			// BaseColor
+				>::GetRHI();
+			case 5:
+				return TStaticBlendState<
+					CW_RGB, BO_Add, BF_SourceAlpha, BF_One, BO_Add, BF_Zero, BF_One,				// Emissive
+					CW_NONE, BO_Add, BF_Zero, BF_One, BO_Add, BF_Zero, BF_One,						// Metallic, Specular, Roughness
+					CW_RGB, BO_Add, BF_One, BF_InverseSourceAlpha, BO_Add, BF_Zero, BF_One			// BaseColor
+				>::GetRHI();
+			case 6:
+				return TStaticBlendState<
+					CW_NONE, BO_Add, BF_Zero, BF_One, BO_Add, BF_Zero, BF_One,						// Emissive
+					CW_RGB, BO_Add, BF_One, BF_InverseSourceAlpha, BO_Add, BF_Zero, BF_One,			// Metallic, Specular, Roughness
+					CW_RGB, BO_Add, BF_One, BF_InverseSourceAlpha, BO_Add, BF_Zero, BF_One			// BaseColor
+				>::GetRHI();
+			case 7:
+				return TStaticBlendState<
+					CW_RGB, BO_Add, BF_SourceAlpha, BF_One, BO_Add, BF_Zero, BF_One,				// Emissive
+					CW_RGB, BO_Add, BF_One, BF_InverseSourceAlpha, BO_Add, BF_Zero, BF_One,			// Metallic, Specular, Roughness
+					CW_RGB, BO_Add, BF_One, BF_InverseSourceAlpha, BO_Add, BF_Zero, BF_One			// BaseColor
+				>::GetRHI();
+			}
+		}
+		else if (DecalBlendDesc.BlendMode == BLEND_Modulate)
+		{
+			switch (Mask)
+			{
+			case 1:
+				return TStaticBlendState<
+					CW_RGB, BO_Add, BF_SourceAlpha, BF_One, BO_Add, BF_Zero, BF_One,				// Emissive
+					CW_NONE, BO_Add, BF_Zero, BF_One, BO_Add, BF_Zero, BF_One,						// Metallic, Specular, Roughness
+					CW_NONE, BO_Add, BF_Zero, BF_One, BO_Add, BF_Zero, BF_One						// BaseColor
+				>::GetRHI();
+			case 2:
+				return TStaticBlendState<
+					CW_NONE, BO_Add, BF_Zero, BF_One, BO_Add, BF_Zero, BF_One,						// Emissive
+					CW_RGB, BO_Add, BF_SourceAlpha, BF_InverseSourceAlpha, BO_Add, BF_Zero, BF_One,	// Metallic, Specular, Roughness
+					CW_NONE, BO_Add, BF_Zero, BF_One, BO_Add, BF_Zero, BF_One						// BaseColor
+				>::GetRHI();
+			case 3:
+				return TStaticBlendState<
+					CW_RGB, BO_Add, BF_SourceAlpha, BF_One, BO_Add, BF_Zero, BF_One,				// Emissive
+					CW_RGB, BO_Add, BF_SourceAlpha, BF_InverseSourceAlpha, BO_Add, BF_Zero, BF_One,	// Metallic, Specular, Roughness
+					CW_NONE, BO_Add, BF_Zero, BF_One, BO_Add, BF_Zero, BF_One						// BaseColor
+				>::GetRHI();
+			case 4:
+				return TStaticBlendState<
+					CW_NONE, BO_Add, BF_Zero, BF_One, BO_Add, BF_Zero, BF_One,						// Emissive
+					CW_NONE, BO_Add, BF_Zero, BF_One, BO_Add, BF_Zero, BF_One,						// Metallic, Specular, Roughness
+					CW_RGB, BO_Add, BF_DestColor, BF_InverseSourceAlpha, BO_Add, BF_Zero, BF_One	// BaseColor
+				>::GetRHI();
+			case 5:
+				return TStaticBlendState<
+					CW_RGB, BO_Add, BF_SourceAlpha, BF_One, BO_Add, BF_Zero, BF_One,				// Emissive
+					CW_NONE, BO_Add, BF_Zero, BF_One, BO_Add, BF_Zero, BF_One,						// Metallic, Specular, Roughness
+					CW_RGB, BO_Add, BF_DestColor, BF_InverseSourceAlpha, BO_Add, BF_Zero, BF_One	// BaseColor
+				>::GetRHI();
+			case 6:
+				return TStaticBlendState<
+					CW_NONE, BO_Add, BF_Zero, BF_One, BO_Add, BF_Zero, BF_One,						// Emissive
+					CW_RGB, BO_Add, BF_SourceAlpha, BF_InverseSourceAlpha, BO_Add, BF_Zero, BF_One,	// Metallic, Specular, Roughness
+					CW_RGB, BO_Add, BF_DestColor, BF_InverseSourceAlpha, BO_Add, BF_Zero, BF_One	// BaseColor
+				>::GetRHI();
+			case 7:
+				return TStaticBlendState<
+					CW_RGB, BO_Add, BF_SourceAlpha, BF_One, BO_Add, BF_Zero, BF_One,				// Emissive
+					CW_RGB, BO_Add, BF_SourceAlpha, BF_InverseSourceAlpha, BO_Add, BF_Zero, BF_One,	// Metallic, Specular, Roughness
+					CW_RGB, BO_Add, BF_DestColor, BF_InverseSourceAlpha, BO_Add, BF_Zero, BF_One	// BaseColor
+				>::GetRHI();
+			}
+		}
+
+		ensure(0);
+		return TStaticBlendState<>::GetRHI();
+	}
+
+	FRHIBlendState* GetDecalBlendState_SceneColor(FDecalBlendDesc DecalBlendDesc, EDecalRenderStage DecalRenderStage)
+	{
+		if (DecalRenderStage == EDecalRenderStage::Mobile)
+		{
+			if (DecalBlendDesc.bWriteEmissive)
+			{
+				// Treat blend as emissive
+				return TStaticBlendState<CW_RGB, BO_Add, BF_SourceAlpha, BF_One>::GetRHI();
+			}
+			else
+			{
+				// Treat blend as non-emissive
+				if (DecalBlendDesc.BlendMode == BLEND_Translucent)
+				{
+					return TStaticBlendState<CW_RGB, BO_Add, BF_SourceAlpha, BF_InverseSourceAlpha>::GetRHI();
+				}
+				else if (DecalBlendDesc.BlendMode == BLEND_AlphaComposite)
+				{
+					return TStaticBlendState<CW_RGB, BO_Add, BF_One, BF_InverseSourceAlpha>::GetRHI();
+				}
+				else if (DecalBlendDesc.BlendMode == BLEND_Modulate)
+				{
+					return TStaticBlendState<CW_RGB, BO_Add, BF_DestColor, BF_InverseSourceAlpha>::GetRHI();
+				}
+			}
+		}
+		else
+		{
+			return TStaticBlendState<CW_RGB, BO_Add, BF_SourceAlpha, BF_One>::GetRHI();
+		}
+
+		ensure(0);
+		return TStaticBlendState<>::GetRHI();
+	}
+
+	FRHIBlendState* GetDecalBlendState_AmbientOcclusion(FDecalBlendDesc DecalBlendDesc, EDecalRenderStage DecalRenderStage)
+	{
+		// Modulate with AO target.
+		return TStaticBlendState<CW_RED, BO_Add, BF_DestColor, BF_Zero>::GetRHI();
+	}
+
+	FRHIBlendState* GetDecalBlendState(FDecalBlendDesc DecalBlendDesc, EDecalRenderStage DecalRenderStage, EDecalRenderTargetMode RenderTargetMode)
+	{
+		// As we force the opacity in the shader we don't always _need_ to set different blend states per MRT.
+		// But we want to give the driver as much information as possible about where output isn't required.
+		// An alternative is to call SetRenderTarget per state change. But that is likely be slower (would need testing on various platforms to confirm that).
+
+		switch(RenderTargetMode)
+		{
+			case EDecalRenderTargetMode::DBuffer:
+				return GetDecalBlendState_DBuffer(DecalBlendDesc, DecalRenderStage);
+			case EDecalRenderTargetMode::SceneColorAndGBuffer:
+				return GetDecalBlendState_SceneColorAndGBuffer(DecalBlendDesc, DecalRenderStage);
+			case EDecalRenderTargetMode::SceneColorAndGBufferNoNormal:
+				return GetDecalBlendState_SceneColorAndGBufferNoNormal(DecalBlendDesc, DecalRenderStage);
+			case EDecalRenderTargetMode::SceneColor:
+				return GetDecalBlendState_SceneColor(DecalBlendDesc, DecalRenderStage);
+			case EDecalRenderTargetMode::AmbientOcclusion:
+				return GetDecalBlendState_AmbientOcclusion(DecalBlendDesc, DecalRenderStage);
+		}
+
+		return TStaticBlendState<>::GetRHI();
+	}
+
+	EDecalRasterizerState GetDecalRasterizerState(bool bInsideDecal, bool bIsInverted, bool ViewReverseCulling)
+	{
+		bool bClockwise = bInsideDecal;
+
+		if (ViewReverseCulling)
+		{
+			bClockwise = !bClockwise;
+		}
+
+		if (bIsInverted)
+		{
+			bClockwise = !bClockwise;
+		}
+		
+		return bClockwise ? EDecalRasterizerState::CW : EDecalRasterizerState::CCW;
+	}
+
+	FRHIRasterizerState* GetDecalRasterizerState(EDecalRasterizerState DecalRasterizerState)
+	{
+		switch (DecalRasterizerState)
+		{
+		case EDecalRasterizerState::CW:
+			return TStaticRasterizerState<FM_Solid, CM_CW>::GetRHI();
+		case EDecalRasterizerState::CCW:
+			return TStaticRasterizerState<FM_Solid, CM_CCW>::GetRHI();
+		}
+
+		check(0); 
+		return nullptr;
+	}
+
+	void ModifyCompilationEnvironment(FDecalBlendDesc DecalBlendDesc, EDecalRenderStage DecalRenderStage, FShaderCompilerEnvironment& OutEnvironment)
+	{
+		DecalRenderStage = DecalRenderStage != EDecalRenderStage::None ? DecalRenderStage : GetBaseRenderStage(DecalBlendDesc);
+		check(DecalRenderStage != EDecalRenderStage::None);
+
+		const EDecalRenderTargetMode RenderTargetMode = GetRenderTargetMode(DecalBlendDesc, DecalRenderStage);
+		check(RenderTargetMode != EDecalRenderTargetMode::None);
+
+		const uint32 RenderTargetCount = GetRenderTargetCount(DecalBlendDesc, RenderTargetMode);
+		const uint32 RenderTargetWriteMask = GetRenderTargetWriteMask(DecalBlendDesc, DecalRenderStage, RenderTargetMode);
+
+		OutEnvironment.SetDefine(TEXT("IS_DECAL"), 1);
+		OutEnvironment.SetDefine(TEXT("IS_DBUFFER_DECAL"), DecalRenderStage == EDecalRenderStage::BeforeBassPass ? 1 : 0);
+
+		OutEnvironment.SetDefine(TEXT("DECAL_RENDERSTAGE"), (uint32)DecalRenderStage);
+		OutEnvironment.SetDefine(TEXT("DECAL_RENDERTARGETMODE"), (uint32)RenderTargetMode);
+		OutEnvironment.SetDefine(TEXT("DECAL_RENDERTARGET_COUNT"), RenderTargetCount);
+
+		OutEnvironment.SetDefine(TEXT("DECAL_OUT_MRT0"), (RenderTargetWriteMask & 1) != 0 ? 1 : 0);
+		OutEnvironment.SetDefine(TEXT("DECAL_OUT_MRT1"), (RenderTargetWriteMask & 2) != 0 ? 1 : 0);
+		OutEnvironment.SetDefine(TEXT("DECAL_OUT_MRT2"), (RenderTargetWriteMask & 4) != 0 ? 1 : 0);
+		OutEnvironment.SetDefine(TEXT("DECAL_OUT_MRT3"), (RenderTargetWriteMask & 8) != 0 ? 1 : 0);
+
+		OutEnvironment.SetDefine(TEXT("DECAL_RENDERSTAGE_BEFOREBASEPASS"), (uint32)EDecalRenderStage::BeforeBassPass);
+		OutEnvironment.SetDefine(TEXT("DECAL_RENDERSTAGE_BEFORELIGHTING"), (uint32)EDecalRenderStage::BeforeLighting);
+		OutEnvironment.SetDefine(TEXT("DECAL_RENDERSTAGE_MOBILE"), (uint32)EDecalRenderStage::Mobile);
+		OutEnvironment.SetDefine(TEXT("DECAL_RENDERSTAGE_MOBILEBEFORELIGHTING"), (uint32)EDecalRenderStage::MobileBeforeLighting);
+		OutEnvironment.SetDefine(TEXT("DECAL_RENDERSTAGE_EMISSIVE"), (uint32)EDecalRenderStage::Emissive);
+		OutEnvironment.SetDefine(TEXT("DECAL_RENDERSTAGE_AO"), (uint32)EDecalRenderStage::AmbientOcclusion);
+
+		OutEnvironment.SetDefine(TEXT("DECAL_RENDERTARGETMODE_DBUFFER"), (uint32)EDecalRenderTargetMode::DBuffer);
+		OutEnvironment.SetDefine(TEXT("DECAL_RENDERTARGETMODE_GBUFFER"), (uint32)EDecalRenderTargetMode::SceneColorAndGBuffer);
+		OutEnvironment.SetDefine(TEXT("DECAL_RENDERTARGETMODE_GBUFFER_NONORMAL"), (uint32)EDecalRenderTargetMode::SceneColorAndGBufferNoNormal);
+		OutEnvironment.SetDefine(TEXT("DECAL_RENDERTARGETMODE_SCENECOLOR"), (uint32)EDecalRenderTargetMode::SceneColor);
+		OutEnvironment.SetDefine(TEXT("DECAL_RENDERTARGETMODE_AO"), (uint32)EDecalRenderTargetMode::AmbientOcclusion);
+	}
+}
diff --git a/Engine/Source/Runtime/Renderer/Private/DecalRenderingCommon.h b/Engine/Source/Runtime/Renderer/Private/DecalRenderingCommon.h
new file mode 100644
--- /dev/null
+++ b/Engine/Source/Runtime/Renderer/Private/DecalRenderingCommon.h
@@ -0,0 +1,112 @@
+// Copyright Epic Games, Inc. All Rights Reserved.
+
+#pragma once
+
+#include "CoreMinimal.h"
+#include "Materials/Material.h"
+#include "RenderUtils.h"
+
+class FRHIBlendState;
+class FRHIRasterizerState;
+struct FShaderCompilerEnvironment;
+
+/** Packed decal description which contains all the information to define render passes and shader compilation setup. */
+union FDecalBlendDesc
+{
+	uint32 Packed = 0;
+
+	struct
+	{
+		uint32 BlendMode : 8;
+		uint32 RenderStageMask : 8;
+		uint32 bWriteBaseColor : 1;
+		uint32 bWriteNormal : 1;
+		uint32 bWriteRoughnessSpecularMetallic : 1;
+		uint32 bWriteEmissive : 1;
+		uint32 bWriteAmbientOcclusion : 1;
+		uint32 bWriteDBufferMask : 1;
+	};
+};
+
+/** Enumeration of the points in the frame for decal rendering. */
+enum class EDecalRenderStage : uint8
+{
+	None = 0,
+
+	// DBuffer decal pass.
+	BeforeBassPass = 1,
+	// GBuffer decal pass.
+	BeforeLighting = 2,
+
+	// Mobile decal pass with limited functionality.
+	Mobile = 3,
+	// Mobile decal pass for mobile deferred platforms.
+	MobileBeforeLighting = 4,
+
+	// Emissive decal pass.
+	// DBuffer decals with an emissive component will use this pass.
+	Emissive = 5,
+	// Ambient occlusion decal pass.
+	// A decal can write regular attributes in another pass and then AO in this pass.
+	AmbientOcclusion = 6,
+};
+
+/** Enumeration of the render target layouts for decal rendering. */
+enum class EDecalRenderTargetMode : uint8
+{
+	None = 0,
+
+	DBuffer = 1,
+	SceneColorAndGBuffer = 2,
+	// GBuffer with no normal is necessary for decals sampling the normal from the GBuffer.
+	SceneColorAndGBufferNoNormal = 3,
+	SceneColor = 4,
+	AmbientOcclusion = 5,
+};
+
+/** Enumeration of decal rasterization states. */
+enum class EDecalRasterizerState : uint8
+{
+	Undefined,
+	CCW,
+	CW,
+};
+
+/**
+ * Shared decal functionality for render pass and shader setup.
+ */
+namespace DecalRendering
+{
+	/** Build the packed decal description from a decal material. */
+	FDecalBlendDesc ComputeDecalBlendDesc(EShaderPlatform Platform, FMaterial const* Material);
+
+	/** Build the packed decal description from a decal material. */
+	FDecalBlendDesc ComputeDecalBlendDesc(EShaderPlatform Platform, FMaterialShaderParameters const& MaterialShaderParameters);
+
+	/** Returns true if a decal should be rendered in the render stage. */
+	bool IsCompatibleWithRenderStage(FDecalBlendDesc DecalBlendDesc, EDecalRenderStage DecalRenderStage);
+
+	/** Returns the main render stage for a decal (does not include the emissive, and AO stages). Can return EDecalRenderStage::None if there is no valid main render stage. */
+	EDecalRenderStage GetBaseRenderStage(FDecalBlendDesc DecalBlendDesc);
+
+	/** Get the render target mode that a decal uses for a given stage. Can return EDecalRenderTargetMode::None if there is no valid render target mode. */
+	EDecalRenderTargetMode GetRenderTargetMode(FDecalBlendDesc DecalBlendDesc, EDecalRenderStage DecalRenderStage);
+
+	/** Get render target count for the given render target mode. */
+	uint32 GetRenderTargetCount(FDecalBlendDesc DecalBlendDesc, EDecalRenderTargetMode RenderTargetMode);
+
+	/** Get render target write mask as a bitmask for the given stage and render target mode. */
+	uint32 GetRenderTargetWriteMask(FDecalBlendDesc DecalBlendDesc, EDecalRenderStage DecalRenderStage, EDecalRenderTargetMode RenderTargetMode);
+
+	/** Get the blend state for rendering a decal at the given stage and render target mode. */
+	FRHIBlendState* GetDecalBlendState(FDecalBlendDesc DecalBlendDesc, EDecalRenderStage DecalRenderStage, EDecalRenderTargetMode RenderTargetMode);
+
+	/** Get the EDecalRasterizerState enum from the view setup. */
+	EDecalRasterizerState GetDecalRasterizerState(bool bInsideDecal, bool bIsInverted, bool ViewReverseCulling);
+
+	/** Get the rasterizer state object for a EDecalRasterizerState enum. */
+	FRHIRasterizerState* GetDecalRasterizerState(EDecalRasterizerState DecalRasterizerState);
+
+	/** Modify the shader compilation environment for a given decal and stage. */
+	void ModifyCompilationEnvironment(FDecalBlendDesc DecalBlendDesc, EDecalRenderStage DecalRenderStage, FShaderCompilerEnvironment& OutEnvironment);
+};
diff --git a/Engine/Source/Runtime/Renderer/Private/DecalRenderingShared.cpp b/Engine/Source/Runtime/Renderer/Private/DecalRenderingShared.cpp
--- a/Engine/Source/Runtime/Renderer/Private/DecalRenderingShared.cpp
+++ b/Engine/Source/Runtime/Renderer/Private/DecalRenderingShared.cpp
@@ -29,8 +29,7 @@ FTransientDecalRenderData::FTransientDecalRenderData(const FScene& InScene, cons
 	MaterialProxy = InDecalProxy->DecalMaterial->GetRenderProxy();
 	MaterialResource = &MaterialProxy->GetMaterialWithFallback(InScene.GetFeatureLevel(), MaterialProxy);
 	check(MaterialProxy && MaterialResource);
-	bHasNormal = MaterialResource->HasNormalConnected();
-	FinalDecalBlendMode = FDecalRenderingCommon::ComputeFinalDecalBlendMode(InScene.GetShaderPlatform(), (EDecalBlendMode)MaterialResource->GetDecalBlendMode(), bHasNormal);
+	DecalBlendDesc = DecalRendering::ComputeDecalBlendDesc(InScene.GetShaderPlatform(), MaterialResource);
 }
 
 /**
@@ -71,27 +70,23 @@ IMPLEMENT_SHADER_TYPE(,FDeferredDecalVS,TEXT("/Engine/Private/DeferredDecal.usf"
 class FDeferredDecalPS : public FMaterialShader
 {
 	DECLARE_SHADER_TYPE(FDeferredDecalPS,Material);
-public:
 
-	/**
-	  * Makes sure only shaders for materials that are explicitly flagged
-	  * as 'UsedAsDeferredDecal' in the Material Editor gets compiled into
-	  * the shader cache.
-	  */
+public:
 	static bool ShouldCompilePermutation(const FMaterialShaderPermutationParameters& Parameters)
 	{
-		return Parameters.MaterialParameters.MaterialDomain == MD_DeferredDecal;
+		return (Parameters.MaterialParameters.MaterialDomain == MD_DeferredDecal) &&
+			DecalRendering::GetBaseRenderStage(DecalRendering::ComputeDecalBlendDesc(Parameters.Platform, Parameters.MaterialParameters)) != EDecalRenderStage::None;
 	}
 
-	static void ModifyCompilationEnvironment(const FMaterialShaderPermutationParameters& Parameters, FShaderCompilerEnvironment& OutEnvironment )
+	static void ModifyCompilationEnvironment(const FMaterialShaderPermutationParameters& Parameters, FShaderCompilerEnvironment& OutEnvironment)
 	{
 		FMaterialShader::ModifyCompilationEnvironment(Parameters, OutEnvironment);
-		FDecalRendering::SetDecalCompilationEnvironment(Parameters, OutEnvironment);
+		DecalRendering::ModifyCompilationEnvironment(DecalRendering::ComputeDecalBlendDesc(Parameters.Platform, Parameters.MaterialParameters), EDecalRenderStage::None, OutEnvironment);
 	}
 
 	FDeferredDecalPS() {}
-	FDeferredDecalPS(const ShaderMetaType::CompiledShaderInitializerType& Initializer):
-		FMaterialShader(Initializer)
+	FDeferredDecalPS(const ShaderMetaType::CompiledShaderInitializerType& Initializer)
+		: FMaterialShader(Initializer)
 	{
 		SvPositionToDecal.Bind(Initializer.ParameterMap,TEXT("SvPositionToDecal"));
 		DecalToWorld.Bind(Initializer.ParameterMap,TEXT("DecalToWorld"));
@@ -180,494 +175,262 @@ IMPLEMENT_MATERIAL_SHADER_TYPE(,FDeferredDecalPS,TEXT("/Engine/Private/DeferredD
 class FDeferredDecalEmissivePS : public FDeferredDecalPS
 {
 	DECLARE_SHADER_TYPE(FDeferredDecalEmissivePS, Material);
+
 public:
 	static bool ShouldCompilePermutation(const FMaterialShaderPermutationParameters& Parameters)
 	{
-		return FDeferredDecalPS::ShouldCompilePermutation(Parameters)
-			&& Parameters.MaterialParameters.bHasEmissiveColorConnected
-			&& IsDBufferDecalBlendMode(FDecalRenderingCommon::ComputeFinalDecalBlendMode(Parameters.Platform, (EDecalBlendMode)Parameters.MaterialParameters.DecalBlendMode, Parameters.MaterialParameters.bHasNormalConnected));
+		return (Parameters.MaterialParameters.MaterialDomain == MD_DeferredDecal) &&
+			DecalRendering::IsCompatibleWithRenderStage(DecalRendering::ComputeDecalBlendDesc(Parameters.Platform, Parameters.MaterialParameters), EDecalRenderStage::Emissive);
 	}
 	static void ModifyCompilationEnvironment(const FMaterialShaderPermutationParameters& Parameters, FShaderCompilerEnvironment& OutEnvironment)
 	{
-		FDeferredDecalPS::ModifyCompilationEnvironment(Parameters, OutEnvironment);
-		FDecalRendering::SetEmissiveDBufferDecalCompilationEnvironment(Parameters, OutEnvironment);
+		FMaterialShader::ModifyCompilationEnvironment(Parameters, OutEnvironment);
+		DecalRendering::ModifyCompilationEnvironment(DecalRendering::ComputeDecalBlendDesc(Parameters.Platform, Parameters.MaterialParameters), EDecalRenderStage::Emissive, OutEnvironment);
 	}
 
 	FDeferredDecalEmissivePS() {}
-	FDeferredDecalEmissivePS(const ShaderMetaType::CompiledShaderInitializerType& Initializer) :
-		FDeferredDecalPS(Initializer)
+	FDeferredDecalEmissivePS(const ShaderMetaType::CompiledShaderInitializerType& Initializer)
+		: FDeferredDecalPS(Initializer)
 	{}
 };
 
 IMPLEMENT_MATERIAL_SHADER_TYPE(, FDeferredDecalEmissivePS, TEXT("/Engine/Private/DeferredDecal.usf"), TEXT("MainPS"), SF_Pixel);
 
-bool FDecalRendering::BuildVisibleDecalList(const FScene& Scene, const FViewInfo& View, EDecalRenderStage DecalRenderStage, FTransientDecalRenderDataList* OutVisibleDecals)
+class FDeferredDecalAmbientOcclusionPS : public FDeferredDecalPS
 {
-	QUICK_SCOPE_CYCLE_COUNTER(BuildVisibleDecalList);
+	DECLARE_SHADER_TYPE(FDeferredDecalAmbientOcclusionPS, Material);
 
-	if (OutVisibleDecals)
+public:
+	static bool ShouldCompilePermutation(const FMaterialShaderPermutationParameters& Parameters)
+	{
+		return (Parameters.MaterialParameters.MaterialDomain == MD_DeferredDecal) &&
+			DecalRendering::IsCompatibleWithRenderStage(DecalRendering::ComputeDecalBlendDesc(Parameters.Platform, Parameters.MaterialParameters), EDecalRenderStage::AmbientOcclusion);
+	}
+	static void ModifyCompilationEnvironment(const FMaterialShaderPermutationParameters& Parameters, FShaderCompilerEnvironment& OutEnvironment)
 	{
-		OutVisibleDecals->Empty(Scene.Decals.Num());
+		FMaterialShader::ModifyCompilationEnvironment(Parameters, OutEnvironment);
+		DecalRendering::ModifyCompilationEnvironment(DecalRendering::ComputeDecalBlendDesc(Parameters.Platform, Parameters.MaterialParameters), EDecalRenderStage::AmbientOcclusion, OutEnvironment);
 	}
 
-	const float FadeMultiplier = CVarDecalFadeScreenSizeMultiplier.GetValueOnRenderThread();
-	const EShaderPlatform ShaderPlatform = View.GetShaderPlatform();
+	FDeferredDecalAmbientOcclusionPS() {}
+	FDeferredDecalAmbientOcclusionPS(const ShaderMetaType::CompiledShaderInitializerType& Initializer)
+		: FDeferredDecalPS(Initializer)
+	{}
+};
 
-	const bool bIsPerspectiveProjection = View.IsPerspectiveProjection();
+IMPLEMENT_MATERIAL_SHADER_TYPE(, FDeferredDecalAmbientOcclusionPS, TEXT("/Engine/Private/DeferredDecal.usf"), TEXT("MainPS"), SF_Pixel);
 
-	// Build a list of decals that need to be rendered for this view in SortedDecals
-	for (const FDeferredDecalProxy* DecalProxy : Scene.Decals)
+namespace DecalRendering
+{
+	bool BuildVisibleDecalList(const FScene& Scene, const FViewInfo& View, EDecalRenderStage DecalRenderStage, FTransientDecalRenderDataList* OutVisibleDecals)
 	{
-		if (!DecalProxy->DecalMaterial || !DecalProxy->DecalMaterial->IsValidLowLevelFast())
+		QUICK_SCOPE_CYCLE_COUNTER(BuildVisibleDecalList);
+
+		if (OutVisibleDecals)
 		{
-			continue;
+			OutVisibleDecals->Empty(Scene.Decals.Num());
 		}
 
-		bool bIsShown = true;
-
-		if (!DecalProxy->IsShown(&View))
+		// Don't draw for shader complexity mode.
+		// todo: Handle shader complexity mode for deferred decal.
+		if (View.Family->EngineShowFlags.ShaderComplexity)
 		{
-			bIsShown = false;
+			return false;
 		}
 
-		const FMatrix ComponentToWorldMatrix = DecalProxy->ComponentTrans.ToMatrixWithScale();
+		const float FadeMultiplier = CVarDecalFadeScreenSizeMultiplier.GetValueOnRenderThread();
+		const EShaderPlatform ShaderPlatform = View.GetShaderPlatform();
 
-		// can be optimized as we test against a sphere around the box instead of the box itself
-		const float ConservativeRadius = FMath::Sqrt(
-				ComponentToWorldMatrix.GetScaledAxis(EAxis::X).SizeSquared() +
-				ComponentToWorldMatrix.GetScaledAxis(EAxis::Y).SizeSquared() +
-				ComponentToWorldMatrix.GetScaledAxis(EAxis::Z).SizeSquared());
+		const bool bIsPerspectiveProjection = View.IsPerspectiveProjection();
 
-		// can be optimized as the test is too conservative (sphere instead of OBB)
-		if(ConservativeRadius < SMALL_NUMBER || !View.ViewFrustum.IntersectSphere(ComponentToWorldMatrix.GetOrigin(), ConservativeRadius))
+		// Build a list of decals that need to be rendered for this view in SortedDecals
+		for (const FDeferredDecalProxy* DecalProxy : Scene.Decals)
 		{
-			bIsShown = false;
-		}
+			if (!DecalProxy->DecalMaterial || !DecalProxy->DecalMaterial->IsValidLowLevelFast())
+			{
+				continue;
+			}
 
-		if (bIsShown)
-		{
-			FTransientDecalRenderData Data(Scene, DecalProxy, ConservativeRadius);
-			
-			// filter out decals with blend modes that are not supported on current platform
-			if (FDecalRenderingCommon::IsBlendModeSupported(ShaderPlatform, Data.FinalDecalBlendMode))
+			bool bIsShown = true;
+
+			if (!DecalProxy->IsShown(&View))
 			{
-				if (bIsPerspectiveProjection && Data.DecalProxy->FadeScreenSize != 0.0f)
-				{
-					float Distance = (View.ViewMatrices.GetViewOrigin() - ComponentToWorldMatrix.GetOrigin()).Size();
-					float Radius = ComponentToWorldMatrix.GetMaximumAxisScale();
-					float CurrentScreenSize = ((Radius / Distance) * FadeMultiplier);
-
-					// fading coefficient needs to increase with increasing field of view and decrease with increasing resolution
-					// FadeCoeffScale is an empirically determined constant to bring us back roughly to fraction of screen size for FadeScreenSize
-					const float FadeCoeffScale = 600.0f;
-					float FOVFactor = ((2.0f/View.ViewMatrices.GetProjectionMatrix().M[0][0]) / View.ViewRect.Width()) * FadeCoeffScale;
-					float FadeCoeff = Data.DecalProxy->FadeScreenSize * FOVFactor;
-					float FadeRange = FadeCoeff * 0.5f;
-
-					float Alpha = (CurrentScreenSize - FadeCoeff) / FadeRange;
-					Data.FadeAlpha = FMath::Min(Alpha, 1.0f);
-				}
+				bIsShown = false;
+			}
 
-				EDecalRenderStage LocalDecalRenderStage = FDecalRenderingCommon::ComputeRenderStage(ShaderPlatform, Data.FinalDecalBlendMode);
+			const FMatrix ComponentToWorldMatrix = DecalProxy->ComponentTrans.ToMatrixWithScale();
 
-				const bool bShouldRender = Data.FadeAlpha > 0.0f &&
-					FDecalRenderingCommon::IsCompatibleWithRenderStage(
-						ShaderPlatform,
-						DecalRenderStage,
-						LocalDecalRenderStage,
-						Data.FinalDecalBlendMode,
-						Data.MaterialResource);
+			// can be optimized as we test against a sphere around the box instead of the box itself
+			const float ConservativeRadius = FMath::Sqrt(
+					ComponentToWorldMatrix.GetScaledAxis(EAxis::X).SizeSquared() +
+					ComponentToWorldMatrix.GetScaledAxis(EAxis::Y).SizeSquared() +
+					ComponentToWorldMatrix.GetScaledAxis(EAxis::Z).SizeSquared());
 
-				// we could do this test earlier to avoid the decal intersection but getting DecalBlendMode also costs
-				if (View.Family->EngineShowFlags.ShaderComplexity || bShouldRender)
+			// can be optimized as the test is too conservative (sphere instead of OBB)
+			if(ConservativeRadius < SMALL_NUMBER || !View.ViewFrustum.IntersectSphere(ComponentToWorldMatrix.GetOrigin(), ConservativeRadius))
+			{
+				bIsShown = false;
+			}
+
+			if (bIsShown)
+			{
+				FTransientDecalRenderData Data(Scene, DecalProxy, ConservativeRadius);
+			
+				if (IsCompatibleWithRenderStage(Data.DecalBlendDesc, DecalRenderStage))
 				{
-					if (!OutVisibleDecals)
+					if (bIsPerspectiveProjection && Data.DecalProxy->FadeScreenSize != 0.0f)
+					{
+						float Distance = (View.ViewMatrices.GetViewOrigin() - ComponentToWorldMatrix.GetOrigin()).Size();
+						float Radius = ComponentToWorldMatrix.GetMaximumAxisScale();
+						float CurrentScreenSize = ((Radius / Distance) * FadeMultiplier);
+
+						// fading coefficient needs to increase with increasing field of view and decrease with increasing resolution
+						// FadeCoeffScale is an empirically determined constant to bring us back roughly to fraction of screen size for FadeScreenSize
+						const float FadeCoeffScale = 600.0f;
+						float FOVFactor = ((2.0f/View.ViewMatrices.GetProjectionMatrix().M[0][0]) / View.ViewRect.Width()) * FadeCoeffScale;
+						float FadeCoeff = Data.DecalProxy->FadeScreenSize * FOVFactor;
+						float FadeRange = FadeCoeff * 0.5f;
+
+						float Alpha = (CurrentScreenSize - FadeCoeff) / FadeRange;
+						Data.FadeAlpha = FMath::Min(Alpha, 1.0f);
+					}
+
+					const bool bShouldRender = Data.FadeAlpha > 0.0f;
+
+					if (bShouldRender)
 					{
-						return true;
+						if (!OutVisibleDecals)
+						{
+							return true;
+						}
+						OutVisibleDecals->Add(Data);
 					}
-					OutVisibleDecals->Add(Data);
 				}
 			}
 		}
-	}
 
-	if (!OutVisibleDecals)
-	{
-		return false;
-	}
+		if (!OutVisibleDecals)
+		{
+			return false;
+		}
 
-	if (OutVisibleDecals->Num() > 0)
-	{
-		// Sort by sort order to allow control over composited result
-		// Then sort decals by state to reduce render target switches
-		// Also sort by component since Sort() is not stable
-		struct FCompareFTransientDecalRenderData
+		if (OutVisibleDecals->Num() > 0)
 		{
-			FORCEINLINE bool operator()(const FTransientDecalRenderData& A, const FTransientDecalRenderData& B) const
+			// Sort by sort order to allow control over composited result
+			// Then sort decals by state to reduce render target switches
+			// Also sort by component since Sort() is not stable
+			struct FCompareFTransientDecalRenderData
 			{
-				if (B.DecalProxy->SortOrder != A.DecalProxy->SortOrder)
-				{ 
-					return A.DecalProxy->SortOrder < B.DecalProxy->SortOrder;
-				}
-				// bHasNormal here is more important then blend mode because we want to render every decals that output normals before those that read normal.
-				if (B.bHasNormal != A.bHasNormal)
-				{
-					return B.bHasNormal < A.bHasNormal; // < so that those outputting normal are first.
-				}
-				if (B.FinalDecalBlendMode != A.FinalDecalBlendMode)
+				FORCEINLINE bool operator()(const FTransientDecalRenderData& A, const FTransientDecalRenderData& B) const
 				{
-					return (int32)B.FinalDecalBlendMode < (int32)A.FinalDecalBlendMode;
-				}
-				// Batch decals with the same material together
-				if (B.MaterialProxy != A.MaterialProxy)
-				{
-					return B.MaterialProxy < A.MaterialProxy;
+					if (B.DecalProxy->SortOrder != A.DecalProxy->SortOrder)
+					{ 
+						return A.DecalProxy->SortOrder < B.DecalProxy->SortOrder;
+					}
+					if (B.DecalBlendDesc.bWriteNormal != A.DecalBlendDesc.bWriteNormal)
+					{
+						// bWriteNormal here has priority because we want to render decals that output normals before those could read normals.
+						// Also this is the only flag that can trigger a change of EDecalRenderTargetMode inside a single EDecalRenderStage, and we batch according to this.
+						return B.DecalBlendDesc.bWriteNormal < A.DecalBlendDesc.bWriteNormal; // < so that those outputting normal are first.
+					}
+					if (B.DecalBlendDesc.Packed != A.DecalBlendDesc.Packed)
+					{
+						// Sorting by the FDecalBlendDesc contents will reduce blend state changes.
+						return (int32)B.DecalBlendDesc.Packed < (int32)A.DecalBlendDesc.Packed;
+					}
+					if (B.MaterialProxy != A.MaterialProxy)
+					{
+						// Batch decals with the same material together
+						return B.MaterialProxy < A.MaterialProxy;
+					}
+					return (PTRINT)B.DecalProxy->Component < (PTRINT)A.DecalProxy->Component;
 				}
-				return (PTRINT)B.DecalProxy->Component < (PTRINT)A.DecalProxy->Component;
-			}
-		};
-
-		// Sort decals by blend mode to reduce render target switches
-		OutVisibleDecals->Sort(FCompareFTransientDecalRenderData());
+			};
 
-		return true;
-	}
+			// Sort decals by blend mode to reduce render target switches
+			OutVisibleDecals->Sort(FCompareFTransientDecalRenderData());
 
-	return false;
-}
-
-FMatrix FDecalRendering::ComputeComponentToClipMatrix(const FViewInfo& View, const FMatrix& DecalComponentToWorld)
-{
-	FMatrix ComponentToWorldMatrixTrans = DecalComponentToWorld.ConcatTranslation(View.ViewMatrices.GetPreViewTranslation());
-	return ComponentToWorldMatrixTrans * View.ViewMatrices.GetTranslatedViewProjectionMatrix();
-}
-
-void FDecalRendering::SetShader(FRHICommandList& RHICmdList, FGraphicsPipelineStateInitializer& GraphicsPSOInit, const FViewInfo& View,
-	const FTransientDecalRenderData& DecalData, EDecalRenderStage DecalRenderStage, const FMatrix& FrustumComponentToClip)
-{
-	const FMaterialShaderMap* MaterialShaderMap = DecalData.MaterialResource->GetRenderingThreadShaderMap();
-	const EDebugViewShaderMode DebugViewMode = View.Family->GetDebugViewShaderMode();
-
-	// When in shader complexity, decals get rendered as emissive even though there might not be emissive decals.
-	// FDeferredDecalEmissivePS might not be available depending on the decal blend mode.
-	TShaderRef<FDeferredDecalPS> PixelShader = (DecalRenderStage == DRS_Emissive && DebugViewMode == DVSM_None)
-		? TShaderRef<FDeferredDecalPS>(MaterialShaderMap->GetShader<FDeferredDecalEmissivePS>())
-		: MaterialShaderMap->GetShader<FDeferredDecalPS>();
+			return true;
+		}
 
-	TShaderMapRef<FDeferredDecalVS> VertexShader(View.ShaderMap);
+		return false;
+	}
 
+	FMatrix ComputeComponentToClipMatrix(const FViewInfo& View, const FMatrix& DecalComponentToWorld)
 	{
-		GraphicsPSOInit.BoundShaderState.VertexDeclarationRHI = GetVertexDeclarationFVector4();
-		GraphicsPSOInit.BoundShaderState.VertexShaderRHI = VertexShader.GetVertexShader();
-		GraphicsPSOInit.BoundShaderState.PixelShaderRHI = PixelShader.GetPixelShader();
-		GraphicsPSOInit.PrimitiveType = PT_TriangleList;
-
-		SetGraphicsPipelineState(RHICmdList, GraphicsPSOInit);
-		PixelShader->SetParameters(RHICmdList, View, DecalData.MaterialProxy, *DecalData.DecalProxy, DecalData.FadeAlpha);
+		FMatrix ComponentToWorldMatrixTrans = DecalComponentToWorld.ConcatTranslation(View.ViewMatrices.GetPreViewTranslation());
+		return ComponentToWorldMatrixTrans * View.ViewMatrices.GetTranslatedViewProjectionMatrix();
 	}
 
-	// SetUniformBufferParameter() need to happen after the shader has been set otherwise a DebugBreak could occur.
-
-	// we don't have the Primitive uniform buffer setup for decals (later we want to batch)
+	void SetShader(FRHICommandList& RHICmdList, FGraphicsPipelineStateInitializer& GraphicsPSOInit, const FViewInfo& View,
+		const FTransientDecalRenderData& DecalData, EDecalRenderStage DecalRenderStage, const FMatrix& FrustumComponentToClip)
 	{
-		auto& PrimitiveVS = VertexShader->GetUniformBufferParameter<FPrimitiveUniformShaderParameters>();
-		auto& PrimitivePS = PixelShader->GetUniformBufferParameter<FPrimitiveUniformShaderParameters>();
-
-		// uncomment to track down usage of the Primitive uniform buffer
-		//	check(!PrimitiveVS.IsBound());
-		//	check(!PrimitivePS.IsBound());
+		const FMaterialShaderMap* MaterialShaderMap = DecalData.MaterialResource->GetRenderingThreadShaderMap();
+		const EDebugViewShaderMode DebugViewMode = View.Family->GetDebugViewShaderMode();
 
-		// to prevent potential shader error (UE-18852 ElementalDemo crashes due to nil constant buffer)
-		SetUniformBufferParameter(RHICmdList, VertexShader.GetVertexShader(), PrimitiveVS, GIdentityPrimitiveUniformBuffer);
-
-		if (DebugViewMode == DVSM_None)
+		// When in shader complexity, decals get rendered as emissive even though there might not be emissive decals.
+		// FDeferredDecalEmissivePS might not be available depending on the decal blend mode.
+		TShaderRef<FDeferredDecalPS> PixelShader;
+		if (DecalRenderStage == EDecalRenderStage::Emissive || DebugViewMode != DVSM_None)
 		{
-			SetUniformBufferParameter(RHICmdList, PixelShader.GetPixelShader(), PrimitivePS, GIdentityPrimitiveUniformBuffer);
+			PixelShader = TShaderRef<FDeferredDecalPS>(MaterialShaderMap->GetShader<FDeferredDecalEmissivePS>());
 		}
-	}
-
-	VertexShader->SetParameters(RHICmdList, FrustumComponentToClip);
-
-	// Set stream source after updating cached strides
-	RHICmdList.SetStreamSource(0, GetUnitCubeVertexBuffer(), 0);
-}
-
-void FDecalRendering::SetVertexShaderOnly(FRHICommandList& RHICmdList, FGraphicsPipelineStateInitializer& GraphicsPSOInit, const FViewInfo& View, const FMatrix& FrustumComponentToClip)
-{
-	TShaderMapRef<FDeferredDecalVS> VertexShader(View.ShaderMap);
-
-	GraphicsPSOInit.BoundShaderState.VertexDeclarationRHI = GetVertexDeclarationFVector4();
-	GraphicsPSOInit.BoundShaderState.VertexShaderRHI = VertexShader.GetVertexShader();
-	GraphicsPSOInit.PrimitiveType = PT_TriangleList;
-
-	SetGraphicsPipelineState(RHICmdList, GraphicsPSOInit);
-	VertexShader->SetParameters(RHICmdList, FrustumComponentToClip);
-}
-
-// @return e.g. 1+2+4 means DBufferA(1) + DBufferB(2) + DBufferC(4) is used
-static uint8 ComputeDBufferMRTMask(EDecalBlendMode DecalBlendMode)
-{
-	switch (DecalBlendMode)
-	{
-	case DBM_DBuffer_AlphaComposite:
-		return 1 + 4; // AlphaComposite mode does not touch normals (DBufferB)
-	case DBM_DBuffer_ColorNormalRoughness:
-		return 1 + 2 + 4;
-	case DBM_DBuffer_Emissive:
-	case DBM_DBuffer_EmissiveAlphaComposite:
-	case DBM_DBuffer_Color:
-		return 1;
-	case DBM_DBuffer_ColorNormal:
-		return 1 + 2;
-	case DBM_DBuffer_ColorRoughness:
-		return 1 + 4;
-	case DBM_DBuffer_Normal:
-		return 2;
-	case DBM_DBuffer_NormalRoughness:
-		return 2 + 4;
-	case DBM_DBuffer_Roughness:
-		return 4;
-	}
-
-	return 0;
-}
-
-void FDecalRendering::SetDecalCompilationEnvironment(const FMaterialShaderPermutationParameters& Parameters, FShaderCompilerEnvironment& OutEnvironment)
-{
-	const bool bHasNormalConnected = Parameters.MaterialParameters.bHasNormalConnected;
-	EDecalBlendMode FinalDecalBlendMode = FDecalRenderingCommon::ComputeFinalDecalBlendMode(Parameters.Platform, (EDecalBlendMode)Parameters.MaterialParameters.DecalBlendMode, bHasNormalConnected);
-	EDecalRenderStage DecalRenderStage = FDecalRenderingCommon::ComputeRenderStage(Parameters.Platform, FinalDecalBlendMode);
-	FDecalRenderingCommon::ERenderTargetMode RenderTargetMode = FDecalRenderingCommon::ComputeRenderTargetMode(Parameters.Platform, FinalDecalBlendMode, bHasNormalConnected);
-	uint32 RenderTargetCount = FDecalRenderingCommon::ComputeRenderTargetCount(Parameters.Platform, RenderTargetMode);
-
-	uint32 DecalOutputNormal = (RenderTargetMode == FDecalRenderingCommon::RTM_SceneColorAndGBufferNoNormal || RenderTargetMode == FDecalRenderingCommon::RTM_SceneColorAndGBufferDepthWriteNoNormal) ? 0 : 1;
-	OutEnvironment.SetDefine(TEXT("DECAL_OUTPUT_NORMAL"), DecalOutputNormal);
-
-	// avoid using the index directly, better use DECALBLENDMODEID_VOLUMETRIC, DECALBLENDMODEID_STAIN, ...
-	OutEnvironment.SetDefine(TEXT("DECAL_BLEND_MODE"), (uint32)FinalDecalBlendMode);
-	OutEnvironment.SetDefine(TEXT("DECAL_PROJECTION"), 1);
-	OutEnvironment.SetDefine(TEXT("DECAL_RENDERTARGET_COUNT"), RenderTargetCount);
-	OutEnvironment.SetDefine(TEXT("DECAL_RENDERSTAGE"), (uint32)DecalRenderStage);
-
-	uint8 bDBufferMask = ComputeDBufferMRTMask(FinalDecalBlendMode);
-
-	OutEnvironment.SetDefine(TEXT("MATERIAL_DBUFFERA"), (bDBufferMask & 0x1) != 0);
-	OutEnvironment.SetDefine(TEXT("MATERIAL_DBUFFERB"), (bDBufferMask & 0x2) != 0);
-	OutEnvironment.SetDefine(TEXT("MATERIAL_DBUFFERC"), (bDBufferMask & 0x4) != 0);
-	// define so that we know to output the proper MRTs in the derived defines.
-	OutEnvironment.SetDefine(TEXT("IS_DECAL"), 1);
-
-	if (IsDBufferDecalBlendMode(FinalDecalBlendMode))
-	{
-		OutEnvironment.SetDefine(TEXT("IS_DBUFFER_DECAL"), 1);
-	}
-}
-
-void FDecalRendering::SetEmissiveDBufferDecalCompilationEnvironment(const FMaterialShaderPermutationParameters& Parameters, FShaderCompilerEnvironment& OutEnvironment)
-{
-	OutEnvironment.SetDefine(TEXT("DECAL_OUTPUT_NORMAL"), 0);
-	OutEnvironment.SetDefine(TEXT("DECAL_BLEND_MODE"), (uint32)DBM_DBuffer_Emissive);
-	OutEnvironment.SetDefine(TEXT("DECAL_RENDERTARGET_COUNT"), 1);
-	OutEnvironment.SetDefine(TEXT("DECAL_RENDERSTAGE"), (uint32)DRS_Emissive);
-
-	OutEnvironment.SetDefine(TEXT("MATERIAL_DBUFFERA"), 0);
-	OutEnvironment.SetDefine(TEXT("MATERIAL_DBUFFERB"), 0);
-	OutEnvironment.SetDefine(TEXT("MATERIAL_DBUFFERC"), 0);
-}
-
-extern FRHIBlendState* MobileForward_GetDecalBlendState(EDecalBlendMode DecalBlendMode);
-extern FRHIBlendState* MobileDeferred_GetDecalBlendState(EDecalBlendMode DecalBlendMode, bool bHasNormal);
-
-// @param RenderState 0:before BasePass, 1:before lighting, (later we could add "after lighting" and multiply)
-FRHIBlendState* FDecalRendering::GetDecalBlendState(const ERHIFeatureLevel::Type SMFeatureLevel, EDecalRenderStage InDecalRenderStage, EDecalBlendMode DecalBlendMode, bool bHasNormal)
-{
-	if (InDecalRenderStage == DRS_Mobile)
-	{
-		if (IsMobileDeferredShadingEnabled(GetFeatureLevelShaderPlatform(SMFeatureLevel)))
+		else if (DecalRenderStage == EDecalRenderStage::AmbientOcclusion)
 		{
-			return MobileDeferred_GetDecalBlendState(DecalBlendMode, bHasNormal);
+			PixelShader = TShaderRef<FDeferredDecalPS>(MaterialShaderMap->GetShader<FDeferredDecalAmbientOcclusionPS>());
 		}
 		else
 		{
-			return MobileForward_GetDecalBlendState(DecalBlendMode);
+			PixelShader = MaterialShaderMap->GetShader<FDeferredDecalPS>();
 		}
-	}
-	
-	if (InDecalRenderStage == DRS_BeforeBasePass)
-	{
-		// before base pass (for DBuffer decals)
+
+		TShaderMapRef<FDeferredDecalVS> VertexShader(View.ShaderMap);
+
 		{
-			// As we set the opacity in the shader we don't need to set different frame buffer blend modes but we like to hint to the driver that we
-			// don't need to output there. We also could replace this with many SetRenderTarget calls but it might be slower (needs to be tested).
+			GraphicsPSOInit.BoundShaderState.VertexDeclarationRHI = GetVertexDeclarationFVector4();
+			GraphicsPSOInit.BoundShaderState.VertexShaderRHI = VertexShader.GetVertexShader();
+			GraphicsPSOInit.BoundShaderState.PixelShaderRHI = PixelShader.GetPixelShader();
+			GraphicsPSOInit.PrimitiveType = PT_TriangleList;
 
-			switch (DecalBlendMode)
-			{
-			case DBM_DBuffer_AlphaComposite:
-				return TStaticBlendState<
-					CW_RGBA, BO_Add, BF_One,  BF_InverseSourceAlpha, BO_Add, BF_Zero, BF_InverseSourceAlpha,
-					CW_RGB, BO_Add, BF_Zero, BF_One, BO_Add, BF_Zero, BF_One,
-					CW_RGBA, BO_Add, BF_One,  BF_InverseSourceAlpha, BO_Add, BF_Zero, BF_InverseSourceAlpha,
-					CW_RGBA, BO_Add, BF_One, BF_One, BO_Add, BF_One, BF_One // DBuffer mask
-				>::GetRHI();
-
-			case DBM_DBuffer_ColorNormalRoughness:
-				return TStaticBlendState<
-					CW_RGBA, BO_Add, BF_SourceAlpha, BF_InverseSourceAlpha, BO_Add, BF_Zero, BF_InverseSourceAlpha,
-					CW_RGBA, BO_Add, BF_SourceAlpha, BF_InverseSourceAlpha, BO_Add, BF_Zero, BF_InverseSourceAlpha,
-					CW_RGBA, BO_Add, BF_SourceAlpha, BF_InverseSourceAlpha, BO_Add, BF_Zero, BF_InverseSourceAlpha,
-					CW_RGBA, BO_Add, BF_One, BF_One, BO_Add, BF_One, BF_One // DBuffer mask
-				>::GetRHI();
-
-			case DBM_DBuffer_Color:
-				// we can optimize using less MRT later
-				return TStaticBlendState<
-					CW_RGBA, BO_Add, BF_SourceAlpha, BF_InverseSourceAlpha, BO_Add, BF_Zero, BF_InverseSourceAlpha,
-					CW_RGBA, BO_Add, BF_Zero, BF_One, BO_Add, BF_Zero, BF_One,
-					CW_RGBA, BO_Add, BF_Zero, BF_One, BO_Add, BF_Zero, BF_One,
-					CW_RGBA, BO_Add, BF_One, BF_One, BO_Add, BF_One, BF_One // DBuffer mask
-				>::GetRHI();
-
-			case DBM_DBuffer_ColorNormal:
-				// we can optimize using less MRT later
-				return TStaticBlendState<
-					CW_RGBA, BO_Add, BF_SourceAlpha, BF_InverseSourceAlpha, BO_Add, BF_Zero, BF_InverseSourceAlpha,
-					CW_RGBA, BO_Add, BF_SourceAlpha, BF_InverseSourceAlpha, BO_Add, BF_Zero, BF_InverseSourceAlpha,
-					CW_RGBA, BO_Add, BF_Zero, BF_One, BO_Add, BF_Zero, BF_One,
-					CW_RGBA, BO_Add, BF_One, BF_One, BO_Add, BF_One, BF_One // DBuffer mask
-				>::GetRHI();
-
-			case DBM_DBuffer_ColorRoughness:
-				// we can optimize using less MRT later
-				return TStaticBlendState<
-					CW_RGBA, BO_Add, BF_SourceAlpha, BF_InverseSourceAlpha, BO_Add, BF_Zero, BF_InverseSourceAlpha,
-					CW_RGBA, BO_Add, BF_Zero, BF_One, BO_Add, BF_Zero, BF_One,
-					CW_RGBA, BO_Add, BF_SourceAlpha, BF_InverseSourceAlpha, BO_Add, BF_Zero, BF_InverseSourceAlpha,
-					CW_RGBA, BO_Add, BF_One, BF_One, BO_Add, BF_One, BF_One // DBuffer mask
-				>::GetRHI();
-
-			case DBM_DBuffer_Normal:
-				// we can optimize using less MRT later
-				return TStaticBlendState<
-					CW_RGBA, BO_Add, BF_Zero, BF_One, BO_Add, BF_Zero, BF_One,
-					CW_RGBA, BO_Add, BF_SourceAlpha, BF_InverseSourceAlpha, BO_Add, BF_Zero, BF_InverseSourceAlpha,
-					CW_RGBA, BO_Add, BF_Zero, BF_One, BO_Add, BF_Zero, BF_One,
-					CW_RGBA, BO_Add, BF_One, BF_One, BO_Add, BF_One, BF_One // DBuffer mask
-				>::GetRHI();
-
-			case DBM_DBuffer_NormalRoughness:
-				// we can optimize using less MRT later
-				return TStaticBlendState<
-					CW_RGBA, BO_Add, BF_Zero, BF_One, BO_Add, BF_Zero, BF_One,
-					CW_RGBA, BO_Add, BF_SourceAlpha, BF_InverseSourceAlpha, BO_Add, BF_Zero, BF_InverseSourceAlpha,
-					CW_RGBA, BO_Add, BF_SourceAlpha, BF_InverseSourceAlpha, BO_Add, BF_Zero, BF_InverseSourceAlpha,
-					CW_RGBA, BO_Add, BF_One, BF_One, BO_Add, BF_One, BF_One // DBuffer mask
-				>::GetRHI();
-
-			case DBM_DBuffer_Roughness:
-				// we can optimize using less MRT later
-				return TStaticBlendState<
-					CW_RGBA, BO_Add, BF_Zero, BF_One, BO_Add, BF_Zero, BF_One,
-					CW_RGBA, BO_Add, BF_Zero, BF_One, BO_Add, BF_Zero, BF_One,
-					CW_RGBA, BO_Add, BF_SourceAlpha, BF_InverseSourceAlpha, BO_Add, BF_Zero, BF_InverseSourceAlpha,
-					CW_RGBA, BO_Add, BF_One, BF_One, BO_Add, BF_One, BF_One // DBuffer mask
-				>::GetRHI();
-
-			default:
-				// the decal type should not be rendered in this pass - internal error
-				check(0);
-				return nullptr;
-			}
+			SetGraphicsPipelineState(RHICmdList, GraphicsPSOInit);
+			PixelShader->SetParameters(RHICmdList, View, DecalData.MaterialProxy, *DecalData.DecalProxy, DecalData.FadeAlpha);
 		}
-	}
-	else if (InDecalRenderStage == DRS_AfterBasePass)
-	{
-		ensure(DecalBlendMode == DBM_Volumetric_DistanceFunction);
 
-		return TStaticBlendState<>::GetRHI();
-	}
-	else if (InDecalRenderStage == DRS_AmbientOcclusion)
-	{
-		ensure(DecalBlendMode == DBM_AmbientOcclusion);
+		// SetUniformBufferParameter() need to happen after the shader has been set otherwise a DebugBreak could occur.
 
-		return TStaticBlendState<CW_RED, BO_Add, BF_DestColor, BF_Zero>::GetRHI();
-	}
-	else
-	{
-		// before lighting (for non DBuffer decals)
-
-		switch (DecalBlendMode)
+		// we don't have the Primitive uniform buffer setup for decals (later we want to batch)
 		{
-		case DBM_Translucent:
-			// @todo: Feature Level 10 does not support separate blends modes for each render target. This could result in the
-			// translucent and stain blend modes looking incorrect when running in this mode.
-			if (GSupportsSeparateRenderTargetBlendState)
-			{
-				if (bHasNormal)
-				{
-					return TStaticBlendState<
-						CW_RGB, BO_Add, BF_SourceAlpha, BF_One, BO_Add, BF_Zero, BF_One,	// Emissive
-						CW_RGB, BO_Add, BF_SourceAlpha, BF_InverseSourceAlpha, BO_Add, BF_Zero, BF_One,	// Normal
-						CW_RGB, BO_Add, BF_SourceAlpha, BF_InverseSourceAlpha, BO_Add, BF_Zero, BF_One,	// Metallic, Specular, Roughness
-						CW_RGB, BO_Add, BF_SourceAlpha, BF_InverseSourceAlpha, BO_Add, BF_Zero, BF_One		// BaseColor
-					>::GetRHI();
-				}
-				else
-				{
-					return TStaticBlendState<
-						CW_RGB, BO_Add, BF_SourceAlpha, BF_One, BO_Add, BF_Zero, BF_One,	// Emissive
-						CW_RGB, BO_Add, BF_SourceAlpha, BF_InverseSourceAlpha, BO_Add, BF_Zero, BF_One,	// Metallic, Specular, Roughness
-						CW_RGB, BO_Add, BF_SourceAlpha, BF_InverseSourceAlpha, BO_Add, BF_Zero, BF_One		// BaseColor
-					>::GetRHI();
-				}
-			}
+			auto& PrimitiveVS = VertexShader->GetUniformBufferParameter<FPrimitiveUniformShaderParameters>();
+			auto& PrimitivePS = PixelShader->GetUniformBufferParameter<FPrimitiveUniformShaderParameters>();
+
+			// uncomment to track down usage of the Primitive uniform buffer
+			//	check(!PrimitiveVS.IsBound());
+			//	check(!PrimitivePS.IsBound());
+
+			// to prevent potential shader error (UE-18852 ElementalDemo crashes due to nil constant buffer)
+			SetUniformBufferParameter(RHICmdList, VertexShader.GetVertexShader(), PrimitiveVS, GIdentityPrimitiveUniformBuffer);
 
-		case DBM_Stain:
-			if (GSupportsSeparateRenderTargetBlendState)
+			if (DebugViewMode == DVSM_None)
 			{
-				if (bHasNormal)
-				{
-					return TStaticBlendState<
-						CW_RGB, BO_Add, BF_SourceAlpha, BF_One, BO_Add, BF_Zero, BF_One,	// Emissive
-						CW_RGB, BO_Add, BF_SourceAlpha, BF_InverseSourceAlpha, BO_Add, BF_Zero, BF_One,	// Normal
-						CW_RGB, BO_Add, BF_SourceAlpha, BF_InverseSourceAlpha, BO_Add, BF_Zero, BF_One,	// Metallic, Specular, Roughness
-						CW_RGB, BO_Add, BF_DestColor, BF_InverseSourceAlpha, BO_Add, BF_Zero, BF_One		// BaseColor
-					>::GetRHI();
-				}
-				else
-				{
-					return TStaticBlendState<
-						CW_RGB, BO_Add, BF_SourceAlpha, BF_One, BO_Add, BF_Zero, BF_One,	// Emissive
-						CW_RGB, BO_Add, BF_SourceAlpha, BF_InverseSourceAlpha, BO_Add, BF_Zero, BF_One,	// Metallic, Specular, Roughness
-						CW_RGB, BO_Add, BF_DestColor, BF_InverseSourceAlpha, BO_Add, BF_Zero, BF_One		// BaseColor
-					>::GetRHI();
-				}
+				SetUniformBufferParameter(RHICmdList, PixelShader.GetPixelShader(), PrimitivePS, GIdentityPrimitiveUniformBuffer);
 			}
+		}
 
-		case DBM_Normal:
-			return TStaticBlendState< CW_RGB, BO_Add, BF_SourceAlpha, BF_InverseSourceAlpha >::GetRHI();
+		VertexShader->SetParameters(RHICmdList, FrustumComponentToClip);
 
+		// Set stream source after updating cached strides
+		RHICmdList.SetStreamSource(0, GetUnitCubeVertexBuffer(), 0);
+	}
 
-		case DBM_Emissive:
-		case DBM_DBuffer_Emissive:
-			return TStaticBlendState< CW_RGB, BO_Add, BF_SourceAlpha, BF_One >::GetRHI();
+	void SetVertexShaderOnly(FRHICommandList& RHICmdList, FGraphicsPipelineStateInitializer& GraphicsPSOInit, const FViewInfo& View, const FMatrix& FrustumComponentToClip)
+	{
+		TShaderMapRef<FDeferredDecalVS> VertexShader(View.ShaderMap);
 
-		case DBM_DBuffer_EmissiveAlphaComposite:
-			return TStaticBlendState< CW_RGB, BO_Add, BF_One, BF_One >::GetRHI();
+		GraphicsPSOInit.BoundShaderState.VertexDeclarationRHI = GetVertexDeclarationFVector4();
+		GraphicsPSOInit.BoundShaderState.VertexShaderRHI = VertexShader.GetVertexShader();
+		GraphicsPSOInit.PrimitiveType = PT_TriangleList;
 
-		case DBM_AlphaComposite:
-			if (GSupportsSeparateRenderTargetBlendState)
-			{
-				return TStaticBlendState<
-					CW_RGB, BO_Add, BF_One, BF_InverseSourceAlpha, BO_Add, BF_Zero, BF_One,	// Emissive
-					CW_RGB, BO_Add, BF_Zero, BF_One, BO_Add, BF_Zero, BF_One,				// Normal
-					CW_RGB, BO_Add, BF_One, BF_InverseSourceAlpha, BO_Add, BF_Zero, BF_One,	// Metallic, Specular, Roughness
-					CW_RGB, BO_Add, BF_One, BF_InverseSourceAlpha, BO_Add, BF_Zero, BF_One	// BaseColor
-				>::GetRHI();
-			}
-		default:
-			// the decal type should not be rendered in this pass - internal error
-			check(0);
-			return nullptr;
-		}
+		SetGraphicsPipelineState(RHICmdList, GraphicsPSOInit);
+		VertexShader->SetParameters(RHICmdList, FrustumComponentToClip);
 	}
 }
diff --git a/Engine/Source/Runtime/Renderer/Private/DecalRenderingShared.h b/Engine/Source/Runtime/Renderer/Private/DecalRenderingShared.h
--- a/Engine/Source/Runtime/Renderer/Private/DecalRenderingShared.h
+++ b/Engine/Source/Runtime/Renderer/Private/DecalRenderingShared.h
@@ -1,9 +1,5 @@
 // Copyright Epic Games, Inc. All Rights Reserved.
 
-/*=============================================================================
-	DecalRenderingShared.h
-=============================================================================*/
-
 #pragma once
 
 #include "CoreMinimal.h"
@@ -11,23 +7,22 @@
 #include "DecalRenderingCommon.h"
 
 class FDeferredDecalProxy;
+class FMaterial;
+class FMaterialRenderProxy;
 class FScene;
 class FViewInfo;
-struct FShaderCompilerEnvironment;
-struct FMaterialShaderPermutationParameters;
 
 /**
- * Compact decal data for rendering
+ * Compact deferred decal data for rendering.
  */
 struct FTransientDecalRenderData
 {
 	const FMaterialRenderProxy* MaterialProxy;
 	const FMaterial* MaterialResource;
 	const FDeferredDecalProxy* DecalProxy;
+	FDecalBlendDesc DecalBlendDesc;
 	float FadeAlpha;
 	float ConservativeRadius;
-	EDecalBlendMode FinalDecalBlendMode;
-	bool bHasNormal;
 
 	FTransientDecalRenderData(const FScene& InScene, const FDeferredDecalProxy* InDecalProxy, float InConservativeRadius);
 };
@@ -35,17 +30,12 @@ struct FTransientDecalRenderData
 typedef TArray<FTransientDecalRenderData, SceneRenderingAllocator> FTransientDecalRenderDataList;
 
 /**
- * Shared decal functionality for deferred and forward shading
+ * Shared deferred decal functionality.
  */
-struct FDecalRendering
+namespace DecalRendering
 {
-	static FMatrix ComputeComponentToClipMatrix(const FViewInfo& View, const FMatrix& DecalComponentToWorld);
-	static void SetVertexShaderOnly(FRHICommandList& RHICmdList, FGraphicsPipelineStateInitializer& GraphicsPSOInit, const FViewInfo& View, const FMatrix& FrustumComponentToClip);
-	static bool BuildVisibleDecalList(const FScene& Scene, const FViewInfo& View, EDecalRenderStage DecalRenderStage, FTransientDecalRenderDataList* OutVisibleDecals);
-	static void SetShader(FRHICommandList& RHICmdList, FGraphicsPipelineStateInitializer& GraphicsPSOInit, const FViewInfo& View, const FTransientDecalRenderData& DecalData, EDecalRenderStage DecalRenderStage, const FMatrix& FrustumComponentToClip);
-
-	// Set common compilation environment parameters for decal shaders (FDeferredDecalPS, FMeshDecalsPS, etc.)
-	static void SetDecalCompilationEnvironment(const FMaterialShaderPermutationParameters& Parameters, FShaderCompilerEnvironment& OutEnvironment);
-	static void SetEmissiveDBufferDecalCompilationEnvironment(const FMaterialShaderPermutationParameters& Parameters, FShaderCompilerEnvironment& OutEnvironment);
-	static FRHIBlendState* GetDecalBlendState(const ERHIFeatureLevel::Type SMFeatureLevel, EDecalRenderStage InDecalRenderStage, EDecalBlendMode DecalBlendMode, bool bHasNormal);
+	FMatrix ComputeComponentToClipMatrix(const FViewInfo& View, const FMatrix& DecalComponentToWorld);
+	void SetVertexShaderOnly(FRHICommandList& RHICmdList, FGraphicsPipelineStateInitializer& GraphicsPSOInit, const FViewInfo& View, const FMatrix& FrustumComponentToClip);
+	bool BuildVisibleDecalList(const FScene& Scene, const FViewInfo& View, EDecalRenderStage DecalRenderStage, FTransientDecalRenderDataList* OutVisibleDecals);
+	void SetShader(FRHICommandList& RHICmdList, FGraphicsPipelineStateInitializer& GraphicsPSOInit, const FViewInfo& View, const FTransientDecalRenderData& DecalData, EDecalRenderStage DecalRenderStage, const FMatrix& FrustumComponentToClip);
 };
diff --git a/Engine/Source/Runtime/Renderer/Private/MobileDecalRendering.cpp b/Engine/Source/Runtime/Renderer/Private/MobileDecalRendering.cpp
--- a/Engine/Source/Runtime/Renderer/Private/MobileDecalRendering.cpp
+++ b/Engine/Source/Runtime/Renderer/Private/MobileDecalRendering.cpp
@@ -13,177 +13,11 @@
 #include "PostProcess/SceneRenderTargets.h"
 #include "SceneRendering.h"
 #include "ScenePrivate.h"
+#include "DecalRenderingCommon.h"
 #include "DecalRenderingShared.h"
 
-extern FRHIRasterizerState* GetDecalRasterizerState(EDecalRasterizerState DecalRasterizerState);
-extern void RenderMeshDecalsMobile(FRHICommandListImmediate& RHICmdList, const FViewInfo& View);
-void RenderDeferredDecalsMobile(FRHICommandListImmediate& RHICmdList, const FScene& Scene, const FViewInfo& View);
-
-FRHIBlendState* MobileForward_GetDecalBlendState(EDecalBlendMode DecalBlendMode)
-{
-	switch(DecalBlendMode)
-	{
-	case DBM_Translucent:
-	case DBM_DBuffer_Color:
-	case DBM_DBuffer_ColorNormal:
-	case DBM_DBuffer_ColorRoughness:
-	case DBM_DBuffer_ColorNormalRoughness:
-		return TStaticBlendState<CW_RGB, BO_Add, BF_SourceAlpha, BF_InverseSourceAlpha>::GetRHI();
-	case DBM_Stain:
-		// Modulate
-		return TStaticBlendState<CW_RGB, BO_Add, BF_DestColor, BF_InverseSourceAlpha>::GetRHI();
-	case DBM_Emissive:
-	case DBM_DBuffer_Emissive:
-		// Additive
-		return TStaticBlendState<CW_RGB, BO_Add, BF_SourceAlpha, BF_One>::GetRHI();
-	case DBM_AlphaComposite:
-	case DBM_DBuffer_AlphaComposite:
-	case DBM_DBuffer_EmissiveAlphaComposite:
-		// Premultiplied alpha
-		return TStaticBlendState<CW_RGB, BO_Add, BF_One, BF_InverseSourceAlpha>::GetRHI();
-	default:
-		check(0);
-	};
-	return TStaticBlendState<CW_RGB, BO_Add, BF_SourceAlpha, BF_InverseSourceAlpha>::GetRHI();
-}
-
-FRHIBlendState* MobileDeferred_GetDecalBlendState(EDecalBlendMode DecalBlendMode, bool bHasNormal)
-{
-	switch (DecalBlendMode)
-	{
-	case DBM_Translucent:
-		if (bHasNormal)
-		{
-			return TStaticBlendState<
-				CW_RGB, BO_Add, BF_SourceAlpha, BF_One, BO_Add, BF_Zero, BF_One,				// Emissive
-				CW_RGB, BO_Add, BF_SourceAlpha, BF_InverseSourceAlpha, BO_Add, BF_Zero, BF_One,	// Normal
-				CW_RGB, BO_Add, BF_SourceAlpha, BF_InverseSourceAlpha, BO_Add, BF_Zero, BF_One,	// Metallic, Specular, Roughness
-				CW_RGB, BO_Add, BF_SourceAlpha, BF_InverseSourceAlpha, BO_Add, BF_Zero, BF_One	// BaseColor
-			>::GetRHI();
-		}
-		else
-		{
-			return TStaticBlendState<
-				CW_RGB, BO_Add, BF_SourceAlpha, BF_One, BO_Add, BF_Zero, BF_One,				// Emissive
-				CW_RGB, BO_Add, BF_Zero, BF_One, BO_Add, BF_Zero, BF_One,						
-				CW_RGB, BO_Add, BF_SourceAlpha, BF_InverseSourceAlpha, BO_Add, BF_Zero, BF_One,	// Metallic, Specular, Roughness
-				CW_RGB, BO_Add, BF_SourceAlpha, BF_InverseSourceAlpha, BO_Add, BF_Zero, BF_One	// BaseColor
-			>::GetRHI();
-		}
-	case DBM_Stain:
-		if (bHasNormal)
-		{
-			return TStaticBlendState<
-				CW_RGB, BO_Add, BF_SourceAlpha, BF_One, BO_Add, BF_Zero, BF_One,				// Emissive
-				CW_RGB, BO_Add, BF_SourceAlpha, BF_InverseSourceAlpha, BO_Add, BF_Zero, BF_One,	// Normal
-				CW_RGB, BO_Add, BF_SourceAlpha, BF_InverseSourceAlpha, BO_Add, BF_Zero, BF_One,	// Metallic, Specular, Roughness
-				CW_RGB, BO_Add, BF_DestColor, BF_InverseSourceAlpha, BO_Add, BF_Zero, BF_One	// BaseColor
-			>::GetRHI();
-		}
-		else
-		{
-			return TStaticBlendState<
-				CW_RGB, BO_Add, BF_SourceAlpha, BF_One, BO_Add, BF_Zero, BF_One,				// Emissive
-				CW_RGB, BO_Add, BF_Zero, BF_One, BO_Add, BF_Zero, BF_One,						
-				CW_RGB, BO_Add, BF_SourceAlpha, BF_InverseSourceAlpha, BO_Add, BF_Zero, BF_One,	// Metallic, Specular, Roughness
-				CW_RGB, BO_Add, BF_DestColor, BF_InverseSourceAlpha, BO_Add, BF_Zero, BF_One	// BaseColor
-			>::GetRHI();
-		}
-	case DBM_Emissive:
-	case DBM_DBuffer_Emissive:
-		return TStaticBlendState<
-			CW_RGB, BO_Add, BF_SourceAlpha, BF_One, BO_Add, BF_Zero, BF_One,	// Emissive
-			CW_RGB, BO_Add, BF_Zero, BF_One, BO_Add, BF_Zero, BF_One,				
-			CW_RGB, BO_Add, BF_Zero, BF_One, BO_Add, BF_Zero, BF_One,
-			CW_RGB, BO_Add, BF_Zero, BF_One, BO_Add, BF_Zero, BF_One
-		>::GetRHI();
-	
-	case DBM_DBuffer_EmissiveAlphaComposite:
-		return TStaticBlendState<
-			CW_RGB, BO_Add, BF_One,	 BF_One, BO_Add, BF_Zero, BF_One,	// Emissive
-			CW_RGB, BO_Add, BF_Zero, BF_One, BO_Add, BF_Zero, BF_One,				
-			CW_RGB, BO_Add, BF_Zero, BF_One, BO_Add, BF_Zero, BF_One,
-			CW_RGB, BO_Add, BF_Zero, BF_One, BO_Add, BF_Zero, BF_One
-		>::GetRHI();
-	
-	case DBM_AlphaComposite:
-	case DBM_DBuffer_AlphaComposite:
-		return TStaticBlendState<
-			CW_RGB, BO_Add, BF_One,  BF_One, BO_Add, BF_Zero, BF_One,	// Emissive
-			CW_RGB, BO_Add, BF_Zero, BF_One, BO_Add, BF_Zero, BF_One,				
-			CW_RGB, BO_Add, BF_One, BF_InverseSourceAlpha, BO_Add, BF_Zero, BF_One,	// Metallic, Specular, Roughness
-			CW_RGB, BO_Add, BF_One, BF_InverseSourceAlpha, BO_Add, BF_Zero, BF_One	// BaseColor
-		>::GetRHI();
-
-	case DBM_DBuffer_ColorNormalRoughness:
-		return TStaticBlendState<
-			CW_RGB, BO_Add, BF_Zero, BF_One, BO_Add, BF_Zero, BF_One,
-			CW_RGB, BO_Add, BF_SourceAlpha, BF_InverseSourceAlpha, BO_Add, BF_Zero, BF_One,	// Normal
-			CW_RGB, BO_Add, BF_SourceAlpha, BF_InverseSourceAlpha, BO_Add, BF_Zero, BF_One,	// Metallic, Specular, Roughness
-			CW_RGB, BO_Add, BF_SourceAlpha, BF_InverseSourceAlpha, BO_Add, BF_Zero, BF_One  // BaseColor
-		>::GetRHI();
-	
-	case DBM_DBuffer_ColorRoughness:
-		return TStaticBlendState<
-			CW_RGB, BO_Add, BF_Zero, BF_One, BO_Add, BF_Zero, BF_One,
-			CW_RGB, BO_Add, BF_Zero, BF_One, BO_Add, BF_Zero, BF_One,						
-			CW_RGB, BO_Add, BF_SourceAlpha, BF_InverseSourceAlpha, BO_Add, BF_Zero, BF_One,	// Metallic, Specular, Roughness
-			CW_RGB, BO_Add, BF_SourceAlpha, BF_InverseSourceAlpha, BO_Add, BF_Zero, BF_One  // BaseColor
-		>::GetRHI();
-
-	case DBM_DBuffer_ColorNormal:
-		return TStaticBlendState<
-			CW_RGB, BO_Add, BF_Zero, BF_One, BO_Add, BF_Zero, BF_One,
-			CW_RGB, BO_Add, BF_SourceAlpha, BF_InverseSourceAlpha, BO_Add, BF_Zero, BF_One,	// Normal
-			CW_RGB, BO_Add, BF_Zero, BF_One, BO_Add, BF_Zero, BF_One,						
-			CW_RGB, BO_Add, BF_SourceAlpha, BF_InverseSourceAlpha, BO_Add, BF_Zero, BF_One  // BaseColor
-		>::GetRHI();
-	
-	case DBM_DBuffer_NormalRoughness:
-		return TStaticBlendState<
-			CW_RGB, BO_Add, BF_Zero, BF_One, BO_Add, BF_Zero, BF_One,
-			CW_RGB, BO_Add, BF_SourceAlpha, BF_InverseSourceAlpha, BO_Add, BF_Zero, BF_One,	// Normal
-			CW_RGB, BO_Add, BF_SourceAlpha, BF_InverseSourceAlpha, BO_Add, BF_Zero, BF_One,	// Metallic, Specular, Roughness
-			CW_RGB, BO_Add, BF_Zero, BF_One, BO_Add, BF_Zero, BF_One
-		>::GetRHI();
-
-	case DBM_DBuffer_Color:
-		return TStaticBlendState<
-			CW_RGB, BO_Add, BF_Zero, BF_One, BO_Add, BF_Zero, BF_One,
-			CW_RGB, BO_Add, BF_Zero, BF_One, BO_Add, BF_Zero, BF_One,
-			CW_RGB, BO_Add, BF_Zero, BF_One, BO_Add, BF_Zero, BF_One,						
-			CW_RGB, BO_Add, BF_SourceAlpha, BF_InverseSourceAlpha, BO_Add, BF_Zero, BF_One  // BaseColor
-		>::GetRHI();
-
-	case DBM_DBuffer_Roughness:
-		return TStaticBlendState<
-			CW_RGB, BO_Add, BF_Zero, BF_One, BO_Add, BF_Zero, BF_One,
-			CW_RGB, BO_Add, BF_Zero, BF_One, BO_Add, BF_Zero, BF_One,
-			CW_RGB, BO_Add, BF_SourceAlpha, BF_InverseSourceAlpha, BO_Add, BF_Zero, BF_One,	// Metallic, Specular, Roughness
-			CW_RGB, BO_Add, BF_Zero, BF_One, BO_Add, BF_Zero, BF_One
-		>::GetRHI();
-	
-	case DBM_Normal:
-	case DBM_DBuffer_Normal:
-		return TStaticBlendState<
-			CW_RGB, BO_Add, BF_Zero, BF_One, BO_Add, BF_Zero, BF_One,
-			CW_RGB, BO_Add, BF_SourceAlpha, BF_InverseSourceAlpha, BO_Add, BF_Zero, BF_One,	// Normal
-			CW_RGB, BO_Add, BF_Zero, BF_One, BO_Add, BF_Zero, BF_One,						
-			CW_RGB, BO_Add, BF_Zero, BF_One, BO_Add, BF_Zero, BF_One
-		>::GetRHI();
-
-	case DBM_Volumetric_DistanceFunction:
-		return TStaticBlendState<>::GetRHI();
-	
-	case DBM_AmbientOcclusion:
-		return TStaticBlendState<CW_RED, BO_Add, BF_DestColor, BF_Zero>::GetRHI();
-	
-	default:
-		check(0);
-	};
-
-	return TStaticBlendState<>::GetRHI(); 
-}
+void RenderMeshDecalsMobile(FRHICommandListImmediate& RHICmdList, const FViewInfo& View, EDecalRenderStage DecalRenderStage, EDecalRenderTargetMode RenderTargetMode);
+extern void RenderDeferredDecalsMobile(FRHICommandListImmediate& RHICmdList, const FScene& Scene, const FViewInfo& View, EDecalRenderStage DecalRenderStage, EDecalRenderTargetMode RenderTargetMode);
 
 void FMobileSceneRenderer::RenderDecals(FRHICommandListImmediate& RHICmdList, const FViewInfo& View)
 {
@@ -195,20 +29,24 @@ void FMobileSceneRenderer::RenderDecals(FRHICommandListImmediate& RHICmdList, co
 	CSV_SCOPED_TIMING_STAT_EXCLUSIVE(RenderDecals);
 	SCOPE_CYCLE_COUNTER(STAT_DecalsDrawTime);
 
+	const bool bIsMobileDeferred = IsMobileDeferredShadingEnabled(View.GetShaderPlatform());
+	const EDecalRenderStage DecalRenderStage = bIsMobileDeferred ? EDecalRenderStage::MobileBeforeLighting : EDecalRenderStage::Mobile;
+	const EDecalRenderTargetMode RenderTargetMode = bIsMobileDeferred ? EDecalRenderTargetMode::SceneColorAndGBuffer : EDecalRenderTargetMode::SceneColor;
+
 	// Deferred decals
 	if (Scene->Decals.Num() > 0)
 	{
-		RenderDeferredDecalsMobile(RHICmdList, *Scene, View);
+		RenderDeferredDecalsMobile(RHICmdList, *Scene, View, DecalRenderStage, RenderTargetMode);
 	}
 
 	// Mesh decals
 	if (View.MeshDecalBatches.Num() > 0)
 	{
-		RenderMeshDecalsMobile(RHICmdList, View);
+		RenderMeshDecalsMobile(RHICmdList, View, DecalRenderStage, RenderTargetMode);
 	}
 }
 
-void RenderDeferredDecalsMobile(FRHICommandListImmediate& RHICmdList, const FScene& Scene, const FViewInfo& View)
+void RenderDeferredDecalsMobile(FRHICommandListImmediate& RHICmdList, const FScene& Scene, const FViewInfo& View, EDecalRenderStage DecalRenderStage, EDecalRenderTargetMode RenderTargetMode)
 {
 	const uint32 DecalCount = Scene.Decals.Num();
 	int32 SortedDecalCount = 0;
@@ -217,15 +55,13 @@ void RenderDeferredDecalsMobile(FRHICommandListImmediate& RHICmdList, const FSce
 	if (DecalCount > 0)
 	{
 		// Build a list of decals that need to be rendered for this view
-		FDecalRendering::BuildVisibleDecalList(Scene, View, DRS_Mobile, &SortedDecals);
+		DecalRendering::BuildVisibleDecalList(Scene, View, DecalRenderStage, &SortedDecals);
 		SortedDecalCount = SortedDecals.Num();
 		INC_DWORD_STAT_BY(STAT_Decals, SortedDecalCount);
 	}
 
 	if (SortedDecalCount > 0)
 	{
-		const bool bDeferredShading = IsMobileDeferredShadingEnabled(View.GetShaderPlatform());
-
 		FGraphicsPipelineStateInitializer GraphicsPSOInit;
 		RHICmdList.ApplyCachedRenderTargets(GraphicsPSOInit);
 
@@ -237,7 +73,7 @@ void RenderDeferredDecalsMobile(FRHICommandListImmediate& RHICmdList, const FSce
 			const FTransientDecalRenderData& DecalData = SortedDecals[DecalIndex];
 			const FDeferredDecalProxy& DecalProxy = *DecalData.DecalProxy;
 			const FMatrix ComponentToWorldMatrix = DecalProxy.ComponentTrans.ToMatrixWithScale();
-			const FMatrix FrustumComponentToClip = FDecalRendering::ComputeComponentToClipMatrix(View, ComponentToWorldMatrix);
+			const FMatrix FrustumComponentToClip = DecalRendering::ComputeComponentToClipMatrix(View, ComponentToWorldMatrix);
 
 			const float ConservativeRadius = DecalData.ConservativeRadius;
 			const bool bInsideDecal = ((FVector)View.ViewMatrices.GetViewOrigin() - ComponentToWorldMatrix.GetOrigin()).SizeSquared() < FMath::Square(ConservativeRadius * 1.05f + View.NearClippingDistance * 2.0f);
@@ -247,8 +83,8 @@ void RenderDeferredDecalsMobile(FRHICommandListImmediate& RHICmdList, const FSce
 				const auto& Scale3d = DecalProxy.ComponentTrans.GetScale3D();
 				bReverseHanded = Scale3d[0] * Scale3d[1] * Scale3d[2] < 0.f;
 			}
-			EDecalRasterizerState DecalRasterizerState = FDecalRenderingCommon::ComputeDecalRasterizerState(bInsideDecal, bReverseHanded, View.bReverseCulling);
-			GraphicsPSOInit.RasterizerState = GetDecalRasterizerState(DecalRasterizerState);
+			EDecalRasterizerState DecalRasterizerState = DecalRendering::GetDecalRasterizerState(bInsideDecal, bReverseHanded, View.bReverseCulling);
+			GraphicsPSOInit.RasterizerState = DecalRendering::GetDecalRasterizerState(DecalRasterizerState);
 
 			if (bInsideDecal)
 			{
@@ -267,17 +103,10 @@ void RenderDeferredDecalsMobile(FRHICommandListImmediate& RHICmdList, const FSce
 					GET_STENCIL_BIT_MASK(RECEIVE_DECAL, 1), 0x00>::GetRHI();
 			}
 
-			if (bDeferredShading)
-			{
-				GraphicsPSOInit.BlendState = MobileDeferred_GetDecalBlendState(DecalData.FinalDecalBlendMode, DecalData.bHasNormal);
-			}
-			else
-			{
-				GraphicsPSOInit.BlendState = MobileForward_GetDecalBlendState(DecalData.FinalDecalBlendMode);
-			}
+			GraphicsPSOInit.BlendState = DecalRendering::GetDecalBlendState(DecalData.DecalBlendDesc, DecalRenderStage, RenderTargetMode);
 
 			// Set shader params
-			FDecalRendering::SetShader(RHICmdList, GraphicsPSOInit, View, DecalData, DRS_Mobile, FrustumComponentToClip);
+			DecalRendering::SetShader(RHICmdList, GraphicsPSOInit, View, DecalData, DecalRenderStage, FrustumComponentToClip);
 
 			RHICmdList.DrawIndexedPrimitive(GetUnitCubeIndexBuffer(), 0, 0, 8, 0, UE_ARRAY_COUNT(GCubeIndices) / 3, 1);
 		}
diff --git a/Engine/Source/Runtime/Renderer/Public/DecalRenderingCommon.h b/Engine/Source/Runtime/Renderer/Public/DecalRenderingCommon.h
deleted file mode 100644
