Lower default clipmap range to avoid potential precision issues with non-Nanite instance culling.
Should only really affect planet-scale atmospheric sky rendering with shadows onto cloud type cases, which is not the default config.

diff --git a/Engine/Source/Runtime/Renderer/Private/VirtualShadowMaps/VirtualShadowMapClipmap.cpp b/Engine/Source/Runtime/Renderer/Private/VirtualShadowMaps/VirtualShadowMapClipmap.cpp
--- a/Engine/Source/Runtime/Renderer/Private/VirtualShadowMaps/VirtualShadowMapClipmap.cpp
+++ b/Engine/Source/Runtime/Renderer/Private/VirtualShadowMaps/VirtualShadowMapClipmap.cpp
@@ -25,7 +25,7 @@ static TAutoConsoleVariable<int32> CVarVirtualShadowMapClipmapFirstLevel(
 );
 static TAutoConsoleVariable<int32> CVarVirtualShadowMapClipmapLastLevel(
 	TEXT( "r.Shadow.Virtual.Clipmap.LastLevel" ),
-	28,
+	22,
 	TEXT( "Last level of the virtual climap. Indirectly determines radius the clipmap can cover." ),
 	ECVF_RenderThreadSafe
 );
