Fixes crash in motion blur when enabling Gen5TAA with VR

diff --git a/Engine/Source/Runtime/Renderer/Private/PostProcess/PostProcessMotionBlur.cpp b/Engine/Source/Runtime/Renderer/Private/PostProcess/PostProcessMotionBlur.cpp
--- a/Engine/Source/Runtime/Renderer/Private/PostProcess/PostProcessMotionBlur.cpp
+++ b/Engine/Source/Runtime/Renderer/Private/PostProcess/PostProcessMotionBlur.cpp
@@ -675,12 +675,11 @@ FRDGTextureRef AddMotionBlurFilterPass(
 	MotionBlurFilterParameters.VelocityTileSampler = GetMotionBlurVelocitySampler();
 	MotionBlurFilterParameters.VelocityFlatSampler = GetMotionBlurVelocitySampler();
 
-	FRDGTextureDesc OutColorDesc = ColorTexture->Desc;
-	OutColorDesc.Reset();
-	OutColorDesc.Flags &= ~(TexCreate_RenderTargetable | TexCreate_UAV);
-	OutColorDesc.Flags |= bUseCompute ? TexCreate_UAV : TexCreate_RenderTargetable;
-	OutColorDesc.Flags |= GFastVRamConfig.MotionBlur;
-	OutColorDesc.Format =  IsPostProcessingWithAlphaChannelSupported() ? PF_FloatRGBA : PF_FloatRGB;
+	FRDGTextureDesc OutColorDesc = FRDGTextureDesc::Create2D(
+		ColorTexture->Desc.Extent,
+		IsPostProcessingWithAlphaChannelSupported() ? PF_FloatRGBA : PF_FloatRGB,
+		(!bUseCompute && View.GetOverwriteLoadAction() == ERenderTargetLoadAction::EClear) ? FClearValueBinding::Black : FClearValueBinding::None,
+		(bUseCompute ? TexCreate_UAV : TexCreate_RenderTargetable) | TexCreate_ShaderResource | GFastVRamConfig.MotionBlur);
 
 	FRDGTextureRef ColorTextureOutput = GraphBuilder.CreateTexture(OutColorDesc, TEXT("MotionBlur.SceneColor"));
 
