Apply light function to sub-pixel shadow mask.

diff --git a/Engine/Source/Runtime/Renderer/Private/DeferredShadingRenderer.h b/Engine/Source/Runtime/Renderer/Private/DeferredShadingRenderer.h
--- a/Engine/Source/Runtime/Renderer/Private/DeferredShadingRenderer.h
+++ b/Engine/Source/Runtime/Renderer/Private/DeferredShadingRenderer.h
@@ -689,7 +689,8 @@ private:
 		const FLightSceneInfo* LightSceneInfo,
 		FRDGTextureRef ScreenShadowMaskTexture,
 		bool bLightAttenuationCleared,
-		bool bProjectingForForwardShading);
+		bool bProjectingForForwardShading, 
+		bool bUseHairStrands);
 
 	/** Renders a light function indicating that whole scene shadowing being displayed is for previewing only, and will go away in game. */
 	bool RenderPreviewShadowsIndicator(
@@ -697,7 +698,8 @@ private:
 		const FMinimalSceneTextures& SceneTextures,
 		const FLightSceneInfo* LightSceneInfo,
 		FRDGTextureRef ScreenShadowMaskTexture,
-		bool bLightAttenuationCleared);
+		bool bLightAttenuationCleared,
+		bool bUseHairStrands);
 
 	/** Renders a light function with the given material. */
 	bool RenderLightFunctionForMaterial(
@@ -708,7 +710,8 @@ private:
 		const FMaterialRenderProxy* MaterialProxy,
 		bool bLightAttenuationCleared,
 		bool bProjectingForForwardShading,
-		bool bRenderingPreviewShadowsIndicator);
+		bool bRenderingPreviewShadowsIndicator, 
+		bool bUseHairStrands);
 
 	/** Used by RenderLights to render a light to the scene color buffer. */
 	void RenderLight(
diff --git a/Engine/Source/Runtime/Renderer/Private/LightFunctionRendering.cpp b/Engine/Source/Runtime/Renderer/Private/LightFunctionRendering.cpp
--- a/Engine/Source/Runtime/Renderer/Private/LightFunctionRendering.cpp
+++ b/Engine/Source/Runtime/Renderer/Private/LightFunctionRendering.cpp
@@ -234,11 +234,12 @@ bool FDeferredShadingSceneRenderer::RenderLightFunction(
 	const FLightSceneInfo* LightSceneInfo,
 	FRDGTextureRef ScreenShadowMaskTexture,
 	bool bLightAttenuationCleared,
-	bool bProjectingForForwardShading)
+	bool bProjectingForForwardShading,
+	bool bUseHairStrands)
 {
 	if (ViewFamily.EngineShowFlags.LightFunctions)
 	{
-		return RenderLightFunctionForMaterial(GraphBuilder, SceneTextures, LightSceneInfo, ScreenShadowMaskTexture, LightSceneInfo->Proxy->GetLightFunctionMaterial(), bLightAttenuationCleared, bProjectingForForwardShading, false);
+		return RenderLightFunctionForMaterial(GraphBuilder, SceneTextures, LightSceneInfo, ScreenShadowMaskTexture, LightSceneInfo->Proxy->GetLightFunctionMaterial(), bLightAttenuationCleared, bProjectingForForwardShading, false, bUseHairStrands);
 	}
 	
 	return false;
@@ -249,11 +250,12 @@ bool FDeferredShadingSceneRenderer::RenderPreviewShadowsIndicator(
 	const FMinimalSceneTextures& SceneTextures,
 	const FLightSceneInfo* LightSceneInfo,
 	FRDGTextureRef ScreenShadowMaskTexture,
-	bool bLightAttenuationCleared)
+	bool bLightAttenuationCleared,
+	bool bUseHairStrands)
 {
 	if (GEngine->PreviewShadowsIndicatorMaterial)
 	{
-		return RenderLightFunctionForMaterial(GraphBuilder, SceneTextures, LightSceneInfo, ScreenShadowMaskTexture, GEngine->PreviewShadowsIndicatorMaterial->GetRenderProxy(), bLightAttenuationCleared, false, true);
+		return RenderLightFunctionForMaterial(GraphBuilder, SceneTextures, LightSceneInfo, ScreenShadowMaskTexture, GEngine->PreviewShadowsIndicatorMaterial->GetRenderProxy(), bLightAttenuationCleared, false, true, bUseHairStrands);
 	}
 
 	return false;
@@ -292,7 +294,8 @@ bool FDeferredShadingSceneRenderer::RenderLightFunctionForMaterial(
 	const FMaterialRenderProxy* MaterialProxy,
 	bool bLightAttenuationCleared,
 	bool bProjectingForForwardShading,
-	bool bRenderingPreviewShadowsIndicator)
+	bool bRenderingPreviewShadowsIndicator,
+	bool bUseHairStrands)
 {
 	check(ScreenShadowMaskTexture);
 	check(LightSceneInfo);
@@ -326,7 +329,7 @@ bool FDeferredShadingSceneRenderer::RenderLightFunctionForMaterial(
 				RDG_EVENT_NAME("LightFunction Material=%s", *MaterialForRendering->GetFriendlyName()),
 				PassParameters,
 				ERDGPassFlags::Raster,
-				[&View, LightSceneInfo, LightSceneProxy, MaterialProxyForRendering, MaterialForRendering, MaterialShaders, bLightAttenuationCleared, bProjectingForForwardShading, bRenderingPreviewShadowsIndicator](FRHICommandList& RHICmdList)
+				[&View, LightSceneInfo, LightSceneProxy, MaterialProxyForRendering, MaterialForRendering, MaterialShaders, bLightAttenuationCleared, bProjectingForForwardShading, bRenderingPreviewShadowsIndicator, bUseHairStrands](FRHICommandList& RHICmdList)
 			{
 				FSphere LightBounds = LightSceneProxy->GetBoundingSphere();
 
@@ -399,6 +402,11 @@ bool FDeferredShadingSceneRenderer::RenderLightFunctionForMaterial(
 
 					// Set the light's scissor rectangle.
 					LightSceneProxy->SetScissorRect(RHICmdList, View, View.ViewRect);
+					if (bUseHairStrands)
+					{
+						FIntRect TotalRect = ComputeVisibleHairStrandsMacroGroupsRect(View.ViewRect, View.HairStrandsViewData.MacroGroupDatas);
+						LightSceneProxy->SetScissorRect(RHICmdList, View, TotalRect);
+					}
 
 					// Render a bounding light sphere.
 					SetGraphicsPipelineState(RHICmdList, GraphicsPSOInit);
diff --git a/Engine/Source/Runtime/Renderer/Private/LightGridInjection.cpp b/Engine/Source/Runtime/Renderer/Private/LightGridInjection.cpp
--- a/Engine/Source/Runtime/Renderer/Private/LightGridInjection.cpp
+++ b/Engine/Source/Runtime/Renderer/Private/LightGridInjection.cpp
@@ -872,7 +872,7 @@ void FDeferredShadingSceneRenderer::RenderForwardShadowProjections(
 					SceneTextures,
 					LightSceneInfo,
 					ForwardScreenSpaceShadowMask.Target,
-					true, true);
+					true, true, false);
 			}
 
 			AddCopyToResolveTargetPass(GraphBuilder, ForwardScreenSpaceShadowMask.Target, ForwardScreenSpaceShadowMask.Resolve, FResolveParams());
diff --git a/Engine/Source/Runtime/Renderer/Private/LightRendering.cpp b/Engine/Source/Runtime/Renderer/Private/LightRendering.cpp
--- a/Engine/Source/Runtime/Renderer/Private/LightRendering.cpp
+++ b/Engine/Source/Runtime/Renderer/Private/LightRendering.cpp
@@ -122,6 +122,12 @@ static TAutoConsoleVariable<int32> CVarRayTracingShadowsRectLight(
 	TEXT("Enables ray tracing shadows for rect light (default = 1)"),
 	ECVF_RenderThreadSafe);
 
+static TAutoConsoleVariable<int32> CVarAppliedLightFunctionOnHair(
+	TEXT("r.HairStrands.LightFunction"),
+	1,
+	TEXT("Enables Light function on hair"),
+	ECVF_RenderThreadSafe);
+
 #if ENABLE_DEBUG_DISCARD_PROP
 static float GDebugLightDiscardProp = 0.0f;
 static FAutoConsoleVariableRef CVarDebugLightDiscardProp(
@@ -1883,13 +1889,18 @@ void FDeferredShadingSceneRenderer::RenderLights(
 				{
 					if (bDrawLightFunction)
 					{
-						const bool bLightFunctionRendered = RenderLightFunction(GraphBuilder, SceneTextures, &LightSceneInfo, ScreenShadowMaskTexture, bDrawShadows, false);
+						const bool bLightFunctionRendered = RenderLightFunction(GraphBuilder, SceneTextures, &LightSceneInfo, ScreenShadowMaskTexture, bDrawShadows, false, false);
 						bUsedShadowMaskTexture |= bLightFunctionRendered;
+
+						if (CVarAppliedLightFunctionOnHair.GetValueOnRenderThread() > 0 && bLightFunctionRendered && ScreenShadowMaskSubPixelTexture)
+						{
+							RenderLightFunction(GraphBuilder, SceneTextures, &LightSceneInfo, ScreenShadowMaskSubPixelTexture, bDrawShadows, false, true);
+						}
 					}
 
 					if (bDrawPreviewIndicator)
 					{
-						RenderPreviewShadowsIndicator(GraphBuilder, SceneTextures, &LightSceneInfo, ScreenShadowMaskTexture, bUsedShadowMaskTexture);
+						RenderPreviewShadowsIndicator(GraphBuilder, SceneTextures, &LightSceneInfo, ScreenShadowMaskTexture, bUsedShadowMaskTexture, false);
 					}
 
 					if (!bDrawShadows)
