Make sure primitives remove themselves from the velocity scene data even if the conditions for being put in the velocity scene data have changed.  The underlying bug of FPrimitiveSceneProxy::SetIsBeingMovedByEditor_RenderThread not correctly updating velocity scene data remains - this change is a patch to avoid the crash when a primitive is removed from the scene but doesn't remove itself from velocity data.

diff --git a/Engine/Source/Runtime/Renderer/Private/PrimitiveSceneInfo.cpp b/Engine/Source/Runtime/Renderer/Private/PrimitiveSceneInfo.cpp
--- a/Engine/Source/Runtime/Renderer/Private/PrimitiveSceneInfo.cpp
+++ b/Engine/Source/Runtime/Renderer/Private/PrimitiveSceneInfo.cpp
@@ -155,6 +155,7 @@ FPrimitiveSceneInfo::FPrimitiveSceneInfo(UPrimitiveComponent* InComponent,FScene
 	bNeedsUniformBufferUpdate(false),
 	bIndirectLightingCacheBufferDirty(false),
 	bRegisteredVirtualTextureProducerCallback(false),
+	bRegisteredWithVelocityData(false),
 	InstanceDataOffset(INDEX_NONE),
 	NumInstanceDataEntries(0),
 	LightmapDataOffset(INDEX_NONE),
diff --git a/Engine/Source/Runtime/Renderer/Private/RendererScene.cpp b/Engine/Source/Runtime/Renderer/Private/RendererScene.cpp
--- a/Engine/Source/Runtime/Renderer/Private/RendererScene.cpp
+++ b/Engine/Source/Runtime/Renderer/Private/RendererScene.cpp
@@ -3958,7 +3958,7 @@ void FScene::UpdateAllPrimitiveSceneInfos(FRDGBuilder& GraphBuilder, bool bAsync
 				int32 PrimitiveIndex = PrimitiveSceneInfo->PackedIndex;
 				PrimitiveSceneInfo->PackedIndex = INDEX_NONE;
 
-				if (ShouldPrimitiveOutputVelocity(PrimitiveSceneInfo->Proxy))
+				if (PrimitiveSceneInfo->bRegisteredWithVelocityData)
 				{
 					// Remove primitive's motion blur information.
 					VelocityData.RemoveFromScene(PrimitiveSceneInfo->PrimitiveComponentId);
@@ -4182,6 +4182,7 @@ void FScene::UpdateAllPrimitiveSceneInfos(FRDGBuilder& GraphBuilder, bool bAsync
 
 				if (ShouldPrimitiveOutputVelocity(PrimitiveSceneInfo->Proxy))
 				{
+					PrimitiveSceneInfo->bRegisteredWithVelocityData = true;
 					// We must register the initial LocalToWorld with the velocity state. 
 					// In the case of a moving component with MarkRenderStateDirty() called every frame, UpdateTransform will never happen.
 					VelocityData.UpdateTransform(PrimitiveSceneInfo, PrimitiveTransforms[PrimitiveIndex], PrimitiveTransforms[PrimitiveIndex]);
@@ -4251,6 +4252,7 @@ void FScene::UpdateAllPrimitiveSceneInfos(FRDGBuilder& GraphBuilder, bool bAsync
 
 			if (ShouldPrimitiveOutputVelocity(PrimitiveSceneInfo->Proxy))
 			{
+				PrimitiveSceneInfo->bRegisteredWithVelocityData = true;
 				VelocityData.UpdateTransform(PrimitiveSceneInfo, LocalToWorld, PrimitiveSceneProxy->GetLocalToWorld());
 			}
 
diff --git a/Engine/Source/Runtime/Renderer/Public/PrimitiveSceneInfo.h b/Engine/Source/Runtime/Renderer/Public/PrimitiveSceneInfo.h
--- a/Engine/Source/Runtime/Renderer/Public/PrimitiveSceneInfo.h
+++ b/Engine/Source/Runtime/Renderer/Public/PrimitiveSceneInfo.h
@@ -554,6 +554,9 @@ private:
 	/** If this is TRUE, this primitive has registered with the virtual texture system for a callback on virtual texture changes. */
 	bool bRegisteredVirtualTextureProducerCallback : 1;
 
+	/** True if the primitive registered with velocity data and needs to remove itself when being removed from the scene. */
+	bool bRegisteredWithVelocityData : 1;
+
 	/** Offset into the scene's instance data buffer, when GPUScene is enabled. */
 	int32 InstanceDataOffset;
 
