Add a shadow depth VS permutation for non-nanite vsm to avoid the need for a read-only project setting
PS already has a permutation (bAtomicWrites) so this primarily affects offline compilation
Leaving the function of NonNaniteVSM the same for now to minimize the surface area of the change, but the shader and other logic can be simplified in the future
NonNaniteVSM can become a dynamic debug cvar in the future once we are happy this is working

diff --git a/Engine/Source/Editor/GameProjectGeneration/Private/GameProjectUtils.cpp b/Engine/Source/Editor/GameProjectGeneration/Private/GameProjectUtils.cpp
--- a/Engine/Source/Editor/GameProjectGeneration/Private/GameProjectUtils.cpp
+++ b/Engine/Source/Editor/GameProjectGeneration/Private/GameProjectUtils.cpp
@@ -233,13 +233,6 @@ namespace
 			TEXT("r.Shadow.Virtual.Enable"),
 			TEXT("1"),
 			true /* ShouldReplaceExistingValue */);
-
-		// and also support for non-nanite VSM, this has an impact on all shadow rendering, which is why this is off for existing projects by default.
-		ConfigValues.Emplace(TEXT("DefaultEngine.ini"),
-			TEXT("/Script/Engine.RendererSettings"),
-			TEXT("r.Shadow.Virtual.NonNaniteVSM"),
-			TEXT("True"),
-			true /* ShouldReplaceExistingValue */);
 	}
 
 	/** Get the configuration values for raytracing if enabled. */
diff --git a/Engine/Source/Runtime/Engine/Classes/Engine/RendererSettings.h b/Engine/Source/Runtime/Engine/Classes/Engine/RendererSettings.h
--- a/Engine/Source/Runtime/Engine/Classes/Engine/RendererSettings.h
+++ b/Engine/Source/Runtime/Engine/Classes/Engine/RendererSettings.h
@@ -367,12 +367,6 @@ class ENGINE_API URendererSettings : public UDeveloperSettings
 		ToolTip = "Shadow Map Method"))
 	TEnumAsByte<EShadowMapMethod::Type> ShadowMapMethod;
 
-	UPROPERTY(config, EditAnywhere, Category = Shadows, meta = (
-		ConsoleVariable = "r.Shadow.Virtual.NonNaniteVSM", DisplayName = "Support non-Nanite geometry in Virtual Shadow Maps",
-		ToolTip = "Whether to support non-Nanite geometry in Virtual Shadow Maps.  Enabling this incurs a small overhead if Virtual Shadow Maps is not the active Dynamic Shadowing Method.",
-		ConfigRestartRequired = true))
-	uint32 bEnableNonNaniteVSM:1;
-
 	/**
 	 * "Ray Tracing settings."
 	 */
diff --git a/Engine/Source/Runtime/Engine/Private/RendererSettings.cpp b/Engine/Source/Runtime/Engine/Private/RendererSettings.cpp
--- a/Engine/Source/Runtime/Engine/Private/RendererSettings.cpp
+++ b/Engine/Source/Runtime/Engine/Private/RendererSettings.cpp
@@ -197,19 +197,6 @@ void URendererSettings::PostEditChangeProperty(FPropertyChangedEvent& PropertyCh
 			}
 		}
 
-		if (PropertyChangedEvent.Property->GetFName() == GET_MEMBER_NAME_CHECKED(URendererSettings, ShadowMapMethod)
-			&& ShadowMapMethod == EShadowMapMethod::VirtualShadowMaps)
-		{
-			if (!bEnableNonNaniteVSM)
-			{
-				FMessageDialog::Open(EAppMsgType::Ok, LOCTEXT("Setting automatically enabled, needs restart.", "Virtual Shadow Maps are most efficient when handling both geometry paths, 'Support non-Nanite geometry in Virtual Shadow Maps' has been automatically enabled.  This requires a restart."));
-
-				bEnableNonNaniteVSM = true;
-
-				UpdateDependentPropertyInConfigFile(this, GET_MEMBER_NAME_CHECKED(URendererSettings, bEnableNonNaniteVSM));
-			}
-		}
-
 		if (PropertyChangedEvent.Property->GetFName() == GET_MEMBER_NAME_CHECKED(URendererSettings, VirtualTextureTileSize))
 		{
 			VirtualTextureTileSize = FMath::RoundUpToPowerOfTwo(VirtualTextureTileSize);
@@ -277,11 +264,6 @@ bool URendererSettings::CanEditChange(const FProperty* InProperty) const
 		return ParentVal && bSupportSkyAtmosphere;
 	}
 
-	if ((InProperty->GetFName() == GET_MEMBER_NAME_CHECKED(URendererSettings, bEnableNonNaniteVSM)))
-	{
-		return ParentVal && ShadowMapMethod == EShadowMapMethod::VirtualShadowMaps;
-	}
-
 	return ParentVal;
 }
 #endif // #if WITH_EDITOR
diff --git a/Engine/Source/Runtime/Renderer/Private/ShadowDepthRendering.cpp b/Engine/Source/Runtime/Renderer/Private/ShadowDepthRendering.cpp
--- a/Engine/Source/Runtime/Renderer/Private/ShadowDepthRendering.cpp
+++ b/Engine/Source/Runtime/Renderer/Private/ShadowDepthRendering.cpp
@@ -278,6 +278,7 @@ enum EShadowDepthVertexShaderMode
 	VertexShadowDepth_OutputDepth,
 	VertexShadowDepth_OnePassPointLight,
 	VertexShadowDepth_VSLayer,
+	VertexShadowDepth_VirtualShadowMap,
 };
 
 static TAutoConsoleVariable<int32> CVarSupportPointLightWholeSceneShadows(
@@ -314,6 +315,12 @@ public:
 			return false;
 		}
 
+		if (ShaderMode == VertexShadowDepth_VirtualShadowMap &&
+			(!IsFeatureLevelSupported(Platform, ERHIFeatureLevel::SM5) || !UseGPUScene(Platform) || GEnableNonNaniteVSM == 0 || bIsForGeometryShader))
+		{
+			return false;
+		}
+
 		//Note: This logic needs to stay in sync with OverrideWithDefaultMaterialForShadowDepth!
 		return (Parameters.MaterialParameters.bIsSpecialEngineMaterial
 			// Masked and WPO materials need their shaders but cannot be used with a position only stream.
@@ -338,9 +345,9 @@ public:
 		OutEnvironment.SetDefine(TEXT("ENABLE_FALLBACK_POINTLIGHT_SHADOW_GS"), UseGPUScene(Parameters.Platform) ? 1 : 0);
 #endif // GPUCULL_TODO
 
-		uint32 bEnableNonNaniteVSM = (uint32)(GEnableNonNaniteVSM != 0 && UseGPUScene(Parameters.Platform));
-		OutEnvironment.SetDefine(TEXT("ENABLE_NON_NANITE_VSM"), bEnableNonNaniteVSM);
-		if (bEnableNonNaniteVSM != 0)
+		bool bEnableNonNaniteVSM = (ShaderMode == VertexShadowDepth_VirtualShadowMap);
+		OutEnvironment.SetDefine(TEXT("ENABLE_NON_NANITE_VSM"), bEnableNonNaniteVSM ? 1 : 0);
+		if (bEnableNonNaniteVSM)
 		{
 			FVirtualShadowMapArray::SetShaderDefines(OutEnvironment);
 		}
@@ -432,14 +439,17 @@ IMPLEMENT_SHADER_TYPE(, FOnePassPointShadowDepthGS, TEXT("/Engine/Private/Shadow
 IMPLEMENT_SHADOW_DEPTH_SHADERMODE_SHADERS(VertexShadowDepth_PerspectiveCorrect);
 IMPLEMENT_SHADOW_DEPTH_SHADERMODE_SHADERS(VertexShadowDepth_OutputDepth);
 IMPLEMENT_SHADOW_DEPTH_SHADERMODE_SHADERS(VertexShadowDepth_OnePassPointLight);
+IMPLEMENT_SHADOW_DEPTH_SHADERMODE_SHADERS(VertexShadowDepth_VirtualShadowMap);
 
 // Position only vertex shaders.
 typedef TShadowDepthVS<VertexShadowDepth_PerspectiveCorrect, true> TShadowDepthVSVertexShadowDepth_PerspectiveCorrectPositionOnly;
 typedef TShadowDepthVS<VertexShadowDepth_OutputDepth,        true> TShadowDepthVSVertexShadowDepth_OutputDepthPositionOnly;
 typedef TShadowDepthVS<VertexShadowDepth_OnePassPointLight,  true> TShadowDepthVSVertexShadowDepth_OnePassPointLightPositionOnly;
+typedef TShadowDepthVS<VertexShadowDepth_VirtualShadowMap,   true> TShadowDepthVSVertexShadowDepth_VirtualShadowMapPositionOnly;
 IMPLEMENT_MATERIAL_SHADER_TYPE(template<>, TShadowDepthVSVertexShadowDepth_PerspectiveCorrectPositionOnly, TEXT("/Engine/Private/ShadowDepthVertexShader.usf"), TEXT("PositionOnlyMain"), SF_Vertex);
 IMPLEMENT_MATERIAL_SHADER_TYPE(template<>, TShadowDepthVSVertexShadowDepth_OutputDepthPositionOnly,        TEXT("/Engine/Private/ShadowDepthVertexShader.usf"), TEXT("PositionOnlyMain"), SF_Vertex);
 IMPLEMENT_MATERIAL_SHADER_TYPE(template<>, TShadowDepthVSVertexShadowDepth_OnePassPointLightPositionOnly,  TEXT("/Engine/Private/ShadowDepthVertexShader.usf"), TEXT("PositionOnlyMain"), SF_Vertex);
+IMPLEMENT_MATERIAL_SHADER_TYPE(template<>, TShadowDepthVSVertexShadowDepth_VirtualShadowMapPositionOnly,   TEXT("/Engine/Private/ShadowDepthVertexShader.usf"), TEXT("PositionOnlyMain"), SF_Vertex);
 
 // One pass point light VS for GS shaders.
 typedef TShadowDepthVS<VertexShadowDepth_OnePassPointLight, false, true> TShadowDepthVSForGSVertexShadowDepth_OnePassPointLight;
@@ -511,7 +521,8 @@ public:
 			return false;
 		}
 
-		if (ShaderMode == PixelShadowDepth_VirtualShadowMap && (!IsFeatureLevelSupported(Platform, ERHIFeatureLevel::SM5) || GEnableNonNaniteVSM == 0))
+		if (ShaderMode == PixelShadowDepth_VirtualShadowMap &&
+			(!IsFeatureLevelSupported(Platform, ERHIFeatureLevel::SM5) || !UseGPUScene(Platform) || GEnableNonNaniteVSM == 0))
 		{
 			return false;
 		}
@@ -536,8 +547,8 @@ public:
 		OutEnvironment.SetDefine(TEXT("ONEPASS_POINTLIGHT_SHADOW"), (uint32)(ShaderMode == PixelShadowDepth_OnePassPointLight));
 		OutEnvironment.SetDefine(TEXT("VIRTUAL_TEXTURE_TARGET"), (uint32)(ShaderMode == PixelShadowDepth_VirtualShadowMap));
 
-		uint32 bEnableNonNaniteVSM = (uint32)(GEnableNonNaniteVSM != 0 && UseGPUScene(Parameters.Platform));
-		OutEnvironment.SetDefine(TEXT("ENABLE_NON_NANITE_VSM"), bEnableNonNaniteVSM);
+		bool bEnableNonNaniteVSM = (ShaderMode == PixelShadowDepth_VirtualShadowMap);
+		OutEnvironment.SetDefine(TEXT("ENABLE_NON_NANITE_VSM"), bEnableNonNaniteVSM ? 1 : 0);
 		if (bEnableNonNaniteVSM != 0)
 		{
 			FVirtualShadowMapArray::SetShaderDefines(OutEnvironment);
@@ -596,6 +607,7 @@ bool GetShadowDepthPassShaders(
 	bool bOnePassPointLightShadow,
 	bool bPositionOnlyVS,
 	bool bUsePerspectiveCorrectShadowDepths,
+	bool bVirtualShadowMap,
 	bool bAtomicWrites,
 	TShaderRef<FShadowDepthVS>& VertexShader,
 	TShaderRef<FShadowDepthBasePS>& PixelShader,
@@ -637,6 +649,17 @@ bool GetShadowDepthPassShaders(
 			}
 		}
 	}
+	else if (bVirtualShadowMap)
+	{
+		if (bPositionOnlyVS)
+		{
+			ShaderTypes.AddShaderType<TShadowDepthVS<VertexShadowDepth_VirtualShadowMap, true>>();
+		}
+		else
+		{
+			ShaderTypes.AddShaderType<TShadowDepthVS<VertexShadowDepth_VirtualShadowMap, false>>();
+		}
+	}
 	else if (bUsePerspectiveCorrectShadowDepths)
 	{
 		if (bPositionOnlyVS)
@@ -650,7 +673,6 @@ bool GetShadowDepthPassShaders(
 	}
 	else
 	{
-
 		if (bPositionOnlyVS)
 		{
 			ShaderTypes.AddShaderType<TShadowDepthVS<VertexShadowDepth_OutputDepth, true>>();
@@ -665,7 +687,7 @@ bool GetShadowDepthPassShaders(
 	const bool bNullPixelShader = Material.WritesEveryPixel(true) && !bUsePerspectiveCorrectShadowDepths && !bAtomicWrites && VertexFactory->SupportsNullPixelShader();
 	if (!bNullPixelShader)
 	{
-		if(UseNonNaniteVirtualShadowMaps(GShaderPlatformForFeatureLevel[FeatureLevel], FeatureLevel) && bAtomicWrites )
+		if (bVirtualShadowMap && bAtomicWrites)
 		{
 			ShaderTypes.AddShaderType<TShadowDepthPS<PixelShadowDepth_VirtualShadowMap>>();
 		}
@@ -2088,9 +2110,10 @@ bool FShadowDepthPassMeshProcessor::Process(
 	// One pass point lights don't output a linear depth, so they are already perspective correct.
 	bool bUsePerspectiveCorrectShadowDepths = !ShadowDepthType.bDirectionalLight && !ShadowDepthType.bOnePassPointLightShadow;
 	bool bOnePassPointLightShadow = ShadowDepthType.bOnePassPointLightShadow;
+	
+	bool bVirtualShadowMap = MeshPassTargetType == EMeshPass::VSMShadowDepth;
 	bool bAtomicWrites = false;
-
-	if( MeshPassTargetType == EMeshPass::VSMShadowDepth )
+	if (bVirtualShadowMap)
 	{
 		bUsePerspectiveCorrectShadowDepths = false;
 		bOnePassPointLightShadow = false;
@@ -2105,6 +2128,7 @@ bool FShadowDepthPassMeshProcessor::Process(
 		bOnePassPointLightShadow,
 		bUsePositionOnlyVS,
 		bUsePerspectiveCorrectShadowDepths,
+		bVirtualShadowMap,
 		bAtomicWrites,
 		ShadowDepthPassShaders.VertexShader,
 		ShadowDepthPassShaders.PixelShader,
diff --git a/Engine/Source/Runtime/Renderer/Private/VirtualShadowMaps/VirtualShadowMapArray.cpp b/Engine/Source/Runtime/Renderer/Private/VirtualShadowMaps/VirtualShadowMapArray.cpp
--- a/Engine/Source/Runtime/Renderer/Private/VirtualShadowMaps/VirtualShadowMapArray.cpp
+++ b/Engine/Source/Runtime/Renderer/Private/VirtualShadowMaps/VirtualShadowMapArray.cpp
@@ -140,12 +140,12 @@ static TAutoConsoleVariable<int32> CVarShowClipmapStats(
 	ECVF_RenderThreadSafe
 );
 
-int32 GEnableNonNaniteVSM = 0;
+int32 GEnableNonNaniteVSM = 1;
 
 FAutoConsoleVariableRef CVarEnableNonNaniteVSM(
 	TEXT("r.Shadow.Virtual.NonNaniteVSM"),
 	GEnableNonNaniteVSM,
-	TEXT("EXPERIMENTAL: Enable support for non-nanite Virtual Shadow Maps, this has some performance impact on regular shadow maps.")
+	TEXT("Enable support for non-nanite Virtual Shadow Maps.")
 	TEXT("Read-only and to be set in a config file (requires restart)."),
 	ECVF_RenderThreadSafe | ECVF_ReadOnly
 );
