Improve allocation patterns

https://p4-swarm.epicgames.net/reviews/15932206

diff --git a/Engine/Source/Runtime/Renderer/Private/RendererScene.cpp b/Engine/Source/Runtime/Renderer/Private/RendererScene.cpp
--- a/Engine/Source/Runtime/Renderer/Private/RendererScene.cpp
+++ b/Engine/Source/Runtime/Renderer/Private/RendererScene.cpp
@@ -526,7 +526,7 @@ void FDistanceFieldSceneData::Release()
 	for (int32 BufferIndex = 0; BufferIndex < StreamingRequestReadbackBuffers.Num(); ++BufferIndex)
 	{
 		if (StreamingRequestReadbackBuffers[BufferIndex])
-		{
+	{
 			delete StreamingRequestReadbackBuffers[BufferIndex];
 			StreamingRequestReadbackBuffers[BufferIndex] = nullptr;
 		}
@@ -1154,8 +1154,8 @@ FScene::~FScene()
 
 #if RHI_RAYTRACING
 
-	delete RayTracingDynamicGeometryCollection;
-	RayTracingDynamicGeometryCollection = nullptr;
+		delete RayTracingDynamicGeometryCollection;
+		RayTracingDynamicGeometryCollection = nullptr;
 #endif // RHI_RAYTRACING
 
 	checkf(RemovedPrimitiveSceneInfos.Num() == 0, TEXT("Leaking %d FPrimitiveSceneInfo instances."), RemovedPrimitiveSceneInfos.Num()); // Ensure UpdateAllPrimitiveSceneInfos() is called before destruction.
@@ -3930,33 +3930,33 @@ void FScene::UpdateAllPrimitiveSceneInfos(FRDGBuilder& GraphBuilder, bool bAsync
 
 			for (int32 CheckIndex = StartIndex; CheckIndex < RemovedLocalPrimitiveSceneInfos.Num(); CheckIndex++)
 			{
-				checkfSlow(RemovedLocalPrimitiveSceneInfos[CheckIndex]->PackedIndex >= Primitives.Num() - RemovedLocalPrimitiveSceneInfos.Num(), TEXT("Removed item should be at the end"));
+				checkf(RemovedLocalPrimitiveSceneInfos[CheckIndex]->PackedIndex >= Primitives.Num() - RemovedLocalPrimitiveSceneInfos.Num(), TEXT("Removed item should be at the end"));
 			}
 
-			for (int32 RemoveIndex = StartIndex; RemoveIndex < RemovedLocalPrimitiveSceneInfos.Num(); RemoveIndex++)
-			{
-				int32 SourceIndex = RemovedLocalPrimitiveSceneInfos[RemoveIndex]->PackedIndex;
-				check(SourceIndex >= (Primitives.Num() - RemovedLocalPrimitiveSceneInfos.Num() + StartIndex));
-				Primitives.Pop();
-				PrimitiveTransforms.Pop();
-				PrimitiveSceneProxies.Pop();
-				PrimitiveBounds.Pop();
-				PrimitiveFlagsCompact.Pop();
-				PrimitiveVisibilityIds.Pop();
-				PrimitiveOcclusionFlags.Pop();
-				PrimitiveComponentIds.Pop();
-				PrimitiveVirtualTextureFlags.Pop();
-				PrimitiveVirtualTextureLod.Pop();
-				PrimitiveOcclusionBounds.Pop();
-				PrimitivesAlwaysVisible.RemoveAt(PrimitivesAlwaysVisible.Num() - 1);
+			//Remove all items from the location of StartIndex to the end of the arrays.
+			int RemoveCount = RemovedLocalPrimitiveSceneInfos.Num() - StartIndex;
+			int SourceIndex = Primitives.Num() - RemoveCount;
+
+			Primitives.RemoveAt(SourceIndex, RemoveCount);
+			PrimitiveTransforms.RemoveAt(SourceIndex, RemoveCount);
+			PrimitiveSceneProxies.RemoveAt(SourceIndex, RemoveCount);
+			PrimitiveBounds.RemoveAt(SourceIndex, RemoveCount);
+			PrimitiveFlagsCompact.RemoveAt(SourceIndex, RemoveCount);
+			PrimitiveVisibilityIds.RemoveAt(SourceIndex, RemoveCount);
+			PrimitiveOcclusionFlags.RemoveAt(SourceIndex, RemoveCount);
+			PrimitiveComponentIds.RemoveAt(SourceIndex, RemoveCount);
+			PrimitiveVirtualTextureFlags.RemoveAt(SourceIndex, RemoveCount);
+			PrimitiveVirtualTextureLod.RemoveAt(SourceIndex, RemoveCount);
+			PrimitiveOcclusionBounds.RemoveAt(SourceIndex, RemoveCount);
+			PrimitivesAlwaysVisible.RemoveAt(SourceIndex, RemoveCount);
+
 			#if WITH_EDITOR
-				PrimitivesSelected.RemoveAt(PrimitivesSelected.Num() - 1);
+			PrimitivesSelected.RemoveAt(SourceIndex, RemoveCount);
 			#endif
 			#if RHI_RAYTRACING
-				PrimitiveRayTracingFlags.Pop();
+			PrimitiveRayTracingFlags.RemoveAt(SourceIndex, RemoveCount);
 			#endif
-				PrimitivesNeedingStaticMeshUpdate.RemoveAt(PrimitivesNeedingStaticMeshUpdate.Num() - 1);
-			}
+			PrimitivesNeedingStaticMeshUpdate.RemoveAt(SourceIndex, RemoveCount);
 
 			CheckPrimitiveArrays();
 
@@ -4389,22 +4389,22 @@ void FScene::UpdateAllPrimitiveSceneInfos(FRDGBuilder& GraphBuilder, bool bAsync
 
 void FScene::CreateLightPrimitiveInteractionsForPrimitive(FPrimitiveSceneInfo* PrimitiveInfo, bool bAsyncCreateLPIs)
 {
-	FMemMark MemStackMark(FMemStack::Get());
-	const FBoxSphereBounds& Bounds = PrimitiveInfo->Proxy->GetBounds();
-	const FPrimitiveSceneInfoCompact PrimitiveSceneInfoCompact(PrimitiveInfo);
+		FMemMark MemStackMark(FMemStack::Get());
+		const FBoxSphereBounds& Bounds = PrimitiveInfo->Proxy->GetBounds();
+		const FPrimitiveSceneInfoCompact PrimitiveSceneInfoCompact(PrimitiveInfo);
 
-	// Find local lights that affect the primitive in the light octree.
-	LocalShadowCastingLightOctree.FindElementsWithBoundsTest(Bounds.GetBox(), [&PrimitiveSceneInfoCompact](const FLightSceneInfoCompact& LightSceneInfoCompact)
+		// Find local lights that affect the primitive in the light octree.
+		LocalShadowCastingLightOctree.FindElementsWithBoundsTest(Bounds.GetBox(), [&PrimitiveSceneInfoCompact](const FLightSceneInfoCompact& LightSceneInfoCompact)
 		{
 			LightSceneInfoCompact.LightSceneInfo->CreateLightPrimitiveInteraction(LightSceneInfoCompact, PrimitiveSceneInfoCompact);
 		});
 
-	// Also loop through non-local (directional) shadow-casting lights
-	for (int32 LightID : DirectionalShadowCastingLightIDs)
-	{
-		const FLightSceneInfoCompact& LightSceneInfoCompact = Lights[LightID];
-		LightSceneInfoCompact.LightSceneInfo->CreateLightPrimitiveInteraction(LightSceneInfoCompact, PrimitiveSceneInfoCompact);
-	}
+		// Also loop through non-local (directional) shadow-casting lights
+		for (int32 LightID : DirectionalShadowCastingLightIDs)
+		{
+			const FLightSceneInfoCompact& LightSceneInfoCompact = Lights[LightID];
+			LightSceneInfoCompact.LightSceneInfo->CreateLightPrimitiveInteraction(LightSceneInfoCompact, PrimitiveSceneInfoCompact);
+		}
 }
 
 bool FScene::IsPrimitiveBeingRemoved(FPrimitiveSceneInfo* PrimitiveSceneInfo) const
@@ -4627,20 +4627,20 @@ void UpdateStaticMeshesForMaterials(const TArray<const FMaterial*>& MaterialReso
 			}
 
 			if (UsedMaterialsDependencies.Num() > 0)
-			{
+					{
 				for (const FMaterial* MaterialResourceToUpdate : MaterialResourcesToUpdate)
-				{
+						{
 					UMaterialInterface* UpdatedMaterialInterface = MaterialResourceToUpdate->GetMaterialInterface();
 
 					if (UpdatedMaterialInterface)
-					{
+							{
 						if (UsedMaterialsDependencies.Contains(UpdatedMaterialInterface))
-						{
-							FPrimitiveSceneProxy* SceneProxy = PrimitiveComponent->SceneProxy;
-							FPrimitiveSceneInfo* SceneInfo = SceneProxy->GetPrimitiveSceneInfo();
-							FScene* Scene = SceneInfo->Scene;
-							TArray<FPrimitiveSceneInfo*>& SceneInfos = UsedPrimitives.FindOrAdd(Scene);
-							SceneInfos.Add(SceneInfo);
+				{
+					FPrimitiveSceneProxy* SceneProxy = PrimitiveComponent->SceneProxy;
+					FPrimitiveSceneInfo* SceneInfo = SceneProxy->GetPrimitiveSceneInfo();
+					FScene* Scene = SceneInfo->Scene;
+					TArray<FPrimitiveSceneInfo*>& SceneInfos = UsedPrimitives.FindOrAdd(Scene);
+					SceneInfos.Add(SceneInfo);
 							break;
 						}
 					}
