Refactor path tracer skylight code so it can be shared with RTGI

This allows us to drop the older sampling code from the path tracer implementation.

diff --git a/Engine/Plugins/Experimental/GPULightmass/Source/GPULightmass/Private/LightmapRayTracing.h b/Engine/Plugins/Experimental/GPULightmass/Source/GPULightmass/Private/LightmapRayTracing.h
--- a/Engine/Plugins/Experimental/GPULightmass/Source/GPULightmass/Private/LightmapRayTracing.h
+++ b/Engine/Plugins/Experimental/GPULightmass/Source/GPULightmass/Private/LightmapRayTracing.h
@@ -50,7 +50,6 @@ class FLightmapPathTracingRGS : public FGlobalShader
 	{
 		OutEnvironment.SetDefine(TEXT("GPreviewLightmapPhysicalTileSize"), GPreviewLightmapPhysicalTileSize);
 		OutEnvironment.SetDefine(TEXT("LIGHTMAP_PATH_TRACING_MAIN_RG"), 1);
-		OutEnvironment.SetDefine(TEXT("USE_NEW_SKYDOME"), 1);
 		OutEnvironment.CompilerFlags.Add(CFLAG_ForceDXC);
 	}
 
@@ -102,7 +101,6 @@ class FVolumetricLightmapPathTracingRGS : public FGlobalShader
 	static void ModifyCompilationEnvironment(const FGlobalShaderPermutationParameters& Parameters, FShaderCompilerEnvironment& OutEnvironment)
 	{
 		OutEnvironment.SetDefine(TEXT("GPreviewLightmapPhysicalTileSize"), GPreviewLightmapPhysicalTileSize);
-		OutEnvironment.SetDefine(TEXT("USE_NEW_SKYDOME"), 1);
 		OutEnvironment.CompilerFlags.Add(CFLAG_ForceDXC);
 	}
 
diff --git a/Engine/Shaders/Private/PathTracing/Light/PathTracingSkyLight.ush b/Engine/Shaders/Private/PathTracing/Light/PathTracingSkyLight.ush
--- a/Engine/Shaders/Private/PathTracing/Light/PathTracingSkyLight.ush
+++ b/Engine/Shaders/Private/PathTracing/Light/PathTracingSkyLight.ush
@@ -6,18 +6,13 @@ SkyLight.usf: Light sampling functions for SkyLight implementation
 
 #pragma once
 
-#ifdef USE_NEW_SKYDOME
 Texture2D SkylightTexture;
 Texture2D<float> SkylightPdf;
 SamplerState SkylightTextureSampler;
 float SkylightInvResolution;
 int SkylightMipCount;
-#else
-#include "../../RayTracing/SkyLightMipTreeCommon.ush"
-#include "../../ReflectionEnvironmentShared.ush"
-#endif
 
-#define USE_HIERARCHICAL_IMPORTANCE_SAMPLING
+#define USE_HIERARCHICAL_IMPORTANCE_SAMPLING		1
 
 float SkyLight_EstimateLight(
 	int LightId,
@@ -26,64 +21,28 @@ float SkyLight_EstimateLight(
 	bool IsTransmissiveMaterial
 )
 {
-#ifdef USE_NEW_SKYDOME
 	// cancels out constant factor in pdf below
-	return 4 * PI * SkylightPdf.Load(int3(0, 0, SkylightMipCount - 1));
+#if USE_HIERARCHICAL_IMPORTANCE_SAMPLING == 0
+	return 2 * PI;
 #else
-	// Modifying diffuse contribution to prevent false negatives with PDF estimation
-	float3 SkySHDiffuse = max(GetSkySHDiffuse(WorldNormal), 0.01);
-	float3 SkyLightIrradiance = GetColor(LightId) * SkySHDiffuse;
-	float SkyLightIlluminance = Luminance(SkyLightIrradiance);
-
-	return SkyLightIlluminance * 2.0f * PI;
+	return 4 * PI * SkylightPdf.Load(int3(0, 0, SkylightMipCount - 1));
 #endif
 }
 
 FLightHit SkyLight_TraceLight(RayDesc Ray, int LightId) {
 	if (Ray.TMax == POSITIVE_INFINITY)
 	{
-#ifdef USE_NEW_SKYDOME
 		// NOTE: assumes direction is normalized
 
 		float3 Dir = Ray.Direction;
 		float2 UV = InverseEquiAreaSphericalMapping(Dir.yzx);
 		float4 Result = SkylightTexture.SampleLevel(SkylightTextureSampler, UV, 0);
 		float3 Radiance = Result.xyz;
-#ifndef USE_HIERARCHICAL_IMPORTANCE_SAMPLING
+#if USE_HIERARCHICAL_IMPORTANCE_SAMPLING == 0
 		float Pdf = 1.0 / (2.0 * PI);
 #else
 		float Pdf = Result.w > 0 ? Result.w / (4 * PI * SkylightPdf.Load(int3(0, 0, SkylightMipCount - 1))) : 0.0;
 #endif
-
-#else
-		float3 LightUV = normalize(Ray.Direction);
-		float3 Radiance = SkyLight.Color * SkyLight.Texture.SampleLevel(SkyLight.TextureSampler, LightUV, 0).rgb;
-		float Pdf = 0;
-		if (any(Radiance > 0))
-		{
-#ifndef USE_HIERARCHICAL_IMPORTANCE_SAMPLING
-			Pdf = 1.0 / (2.0 * PI);
-#else
-			// Determine primary direction
-			//uint MipCount = Log2(SkyLight.MipDimensions.x);
-			uint MipCount = log2(SkyLight.MipDimensions.x);
-			uint2 TextureRes;
-			SkyLight.Texture.GetDimensions(TextureRes.x, TextureRes.y);
-			float3 LightDirection = LightUV;
-			uint3 TextureCoord = GetTextureCubeCoordinate(LightDirection, TextureRes);
-			float2 TextureUV = TextureCoord.xy / (float2)TextureRes;
-			uint2 MipCoord = TextureUV * SkyLight.MipDimensions.xy;
-			float FacePdf = PdfFace(MipCount, TextureCoord.z);
-
-			uint StopLevel = min(SkyLight.SamplingStopLevel, MipCount);
-			float MipPdf = PdfMipTree(uint3(MipCoord, TextureCoord.z), StopLevel);
-
-			uint LinearIndex = BufferOffsetAtPixel(MipCoord, StopLevel, SkyLight.MipDimensions.xy);
-			float AreaPdf = SkyLight.SolidAnglePdf[LinearIndex];
-			Pdf = FacePdf * MipPdf * AreaPdf;
-#endif
-		}
-#endif
 		return CreateLightHit(Radiance, Pdf, POSITIVE_INFINITY);
 	}
 	return NullLightHit();
@@ -100,13 +59,8 @@ FLightSample SkyLight_SampleLight(
 	// #dxr_todo: write optimized form that combines radiance and pdf lookup in a single step
 	RayDesc Ray;
 	Ray.Origin = WorldPos;
-	FMaterialClosestHitPayload Payload;
-	Payload.WorldNormal = WorldNormal;
-
-
-#ifdef USE_NEW_SKYDOME
 
-#ifndef USE_HIERARCHICAL_IMPORTANCE_SAMPLING
+#if USE_HIERARCHICAL_IMPORTANCE_SAMPLING == 0
 	float4 Direction_Tangent = UniformSampleHemisphere(RandSample.yz);
 	float3 OutDirection = TangentToWorld(Direction_Tangent.xyz, WorldNormal);
 	float OutPdf = Direction_Tangent.w;
@@ -169,87 +123,4 @@ FLightSample SkyLight_SampleLight(
 
 	return CreateLightSample(Radiance / OutPdf, OutPdf, EquiAreaSphericalMapping(UV).zxy, POSITIVE_INFINITY);
 #endif
-
-#else
-
-#ifndef USE_HIERARCHICAL_IMPORTANCE_SAMPLING
-	float4 Direction_Tangent = UniformSampleHemisphere(RandSample.yz);
-	float3 OutDirection = TangentToWorld(Direction_Tangent.xyz, WorldNormal);
-	float OutPdf = Direction_Tangent.w;
-#else
-	// Determine the appropriate face to use
-	//uint MipCount = Log2(SkyLight.MipDimensions.x);
-	uint MipCount = log2(SkyLight.MipDimensions.x);
-	float FacePdf = 0.0;
-
-	uint FaceIndex = SampleFace(MipCount, RandSample.w, FacePdf);
-
-	uint2 MipPixel = 0;
-	float MipPdf = 1.0;
-	uint StopLevel = min(SkyLight.SamplingStopLevel, MipCount);
-	float DecisionTreeSample = RandSample.y;
-	for (uint Index = 0; Index < MipCount - StopLevel; ++Index)
-	{
-		uint MipIndex = MipCount - Index - 1;
-		uint BufferOffset = BufferOffsetAtPixel(MipPixel, MipIndex, SkyLight.MipDimensions.xy);
-		uint4 BufferOffsets = uint4(BufferOffset, BufferOffset + 1, BufferOffset + 2, BufferOffset + 3);
-
-		float4 Cdf = SampleMipTree(FaceIndex, BufferOffsets);
-		Cdf.y += Cdf.x;
-		Cdf.z += Cdf.y;
-		Cdf.w += Cdf.z;
-		Cdf /= Cdf.w;
-
-		if (DecisionTreeSample < Cdf.x)
-		{
-			MipPdf *= Cdf.x;
-			DecisionTreeSample /= Cdf.x;
-		}
-		else if (DecisionTreeSample < Cdf.y)
-		{
-			MipPdf *= Cdf.y - Cdf.x;
-			MipPixel.x += 1;
-			DecisionTreeSample = (DecisionTreeSample - Cdf.x) / (Cdf.y - Cdf.x);
-		}
-		else if (DecisionTreeSample < Cdf.z)
-		{
-			MipPdf *= Cdf.z - Cdf.y;
-			MipPixel.y += 1;
-			DecisionTreeSample = (DecisionTreeSample - Cdf.y) / (Cdf.z - Cdf.y);
-		}
-		else // if (DecisionTreeSample < Cdf.w)
-		{
-			MipPdf *= Cdf.w - Cdf.z;
-			MipPixel += 1;
-			DecisionTreeSample = (DecisionTreeSample - Cdf.z) / (Cdf.w - Cdf.z);
-		}
-
-		MipPixel = uint2(MipPixel.x << 1, MipPixel.y << 1);
-	}
-	MipPixel = uint2(MipPixel.x >> 1, MipPixel.y >> 1);
-
-	// DEBUG: assert cached PDF is equivalent to re-evaluation
-	//MipPdf = PdfSkyLightOcclusionRay(WorldNormal, RayDirection);
-
-	float2 PixelOffset = float2(DecisionTreeSample, RandSample.z);
-	float2 MipPixelUV = (MipPixel + PixelOffset) / float2(DimensionsAtLevel(SkyLight.MipDimensions.xy, StopLevel));
-
-	// Construct Cubemap coordinate
-	uint2 TextureDimensions;
-	SkyLight.Texture.GetDimensions(TextureDimensions.x, TextureDimensions.y);
-	float3 TexelCoord = float3(MipPixelUV * TextureDimensions, FaceIndex);
-
-	// Determine area pdf
-	uint LinearIndex = BufferOffsetAtPixel(MipPixel, StopLevel, SkyLight.MipDimensions.xy);
-	float AreaPdf = SkyLight.SolidAnglePdf[LinearIndex];
-
-	float3 OutDirection = normalize(GetTextureCubeVector(TexelCoord, TextureDimensions));
-	float OutPdf = FacePdf * MipPdf * AreaPdf;
-#endif
-
-	float3 Radiance = SkyLight.Color * SkyLight.Texture.SampleLevel(SkyLight.TextureSampler, OutDirection, 0).rgb;
-
-	return CreateLightSample(Radiance / OutPdf, OutPdf, OutDirection, POSITIVE_INFINITY);
-#endif
-
 }
diff --git a/Engine/Source/Runtime/Renderer/Private/PathTracing.cpp b/Engine/Source/Runtime/Renderer/Private/PathTracing.cpp
--- a/Engine/Source/Runtime/Renderer/Private/PathTracing.cpp
+++ b/Engine/Source/Runtime/Renderer/Private/PathTracing.cpp
@@ -1,5 +1,6 @@
 // Copyright Epic Games, Inc. All Rights Reserved.
 
+#include "PathTracing.h"
 #include "RHI.h"
 
 #if RHI_RAYTRACING
@@ -466,7 +467,6 @@ class FPathTracingRG : public FGlobalShader
 	static void ModifyCompilationEnvironment(const FGlobalShaderPermutationParameters& Parameters, FShaderCompilerEnvironment& OutEnvironment)
 	{
 		//OutEnvironment.CompilerFlags.Add(CFLAG_WarningsAsErrors);
-		OutEnvironment.SetDefine(TEXT("USE_NEW_SKYDOME"), 1);
 		OutEnvironment.SetDefine(TEXT("USE_RECT_LIGHT_TEXTURES"), 1);
 	}
 
@@ -485,11 +485,8 @@ class FPathTracingRG : public FGlobalShader
 		SHADER_PARAMETER_STRUCT_INCLUDE(FPathTracingLightGrid, LightGridParameters)
 
 		// Skylight
-		SHADER_PARAMETER_RDG_TEXTURE(Texture2D, SkylightTexture)
-		SHADER_PARAMETER_RDG_TEXTURE(Texture2D, SkylightPdf)
-		SHADER_PARAMETER_SAMPLER(SamplerState, SkylightTextureSampler)
-		SHADER_PARAMETER(float, SkylightInvResolution)
-		SHADER_PARAMETER(int32, SkylightMipCount)
+		SHADER_PARAMETER_STRUCT_INCLUDE(FPathTracingSkylight, SkylightParameters)
+
 		// IES Profiles
 		SHADER_PARAMETER_RDG_TEXTURE(Texture2DArray, IESTexture)
 		SHADER_PARAMETER_SAMPLER(SamplerState, IESTextureSampler) // Shared sampler for all IES profiles
@@ -600,25 +597,23 @@ RENDERER_API void PrepareSkyTexture_Internal(
 	}
 }
 
-bool PrepareSkyTexture(FRDGBuilder& GraphBuilder, FScene* Scene, const FViewInfo& View, bool UseMISCompensation, FPathTracingRG::FParameters* PathTracingParameters)
+bool PrepareSkyTexture(FRDGBuilder& GraphBuilder, FScene* Scene, const FViewInfo& View, bool SkylightEnabled, bool UseMISCompensation, bool IsSkylightCachingEnabled, FPathTracingSkylight* SkylightParameters)
 {
-	PathTracingParameters->SkylightTextureSampler = TStaticSamplerState<SF_Bilinear, AM_Clamp, AM_Clamp, AM_Clamp>::GetRHI();
+	SkylightParameters->SkylightTextureSampler = TStaticSamplerState<SF_Bilinear, AM_Clamp, AM_Clamp, AM_Clamp>::GetRHI();
 
 	FReflectionUniformParameters Parameters;
 	SetupReflectionUniformParameters(View, Parameters);
-	if (!(Parameters.SkyLightParameters.Y > 0))
+	if (!SkylightEnabled || !(Parameters.SkyLightParameters.Y > 0))
 	{
 		// textures not ready, or skylight not active
 		// just put in a placeholder
-		PathTracingParameters->SkylightTexture = GraphBuilder.RegisterExternalTexture(GSystemTextures.BlackDummy);
-		PathTracingParameters->SkylightPdf = GraphBuilder.RegisterExternalTexture(GSystemTextures.BlackDummy);
-		PathTracingParameters->SkylightInvResolution = 0;
-		PathTracingParameters->SkylightMipCount = 0;
+		SkylightParameters->SkylightTexture = GraphBuilder.RegisterExternalTexture(GSystemTextures.BlackDummy);
+		SkylightParameters->SkylightPdf = GraphBuilder.RegisterExternalTexture(GSystemTextures.BlackDummy);
+		SkylightParameters->SkylightInvResolution = 0;
+		SkylightParameters->SkylightMipCount = 0;
 		return false;
 	}
 
-	const bool IsSkylightCachingEnabled = CVarPathTracingSkylightCaching.GetValueOnAnyThread() != 0;
-
 	if (!IsSkylightCachingEnabled)
 	{
 		// we don't want any caching - release what we might have been holding onto
@@ -631,10 +626,10 @@ bool PrepareSkyTexture(FRDGBuilder& GraphBuilder, FScene* Scene, const FViewInfo
 	{
 		// we already have a valid texture and pdf, just re-use them!
 		// it is the responsability of code that may invalidate the contents to reset these pointers
-		PathTracingParameters->SkylightTexture = GraphBuilder.RegisterExternalTexture(Scene->PathTracingSkylightTexture, TEXT("PathTracer.Skylight"));
-		PathTracingParameters->SkylightPdf = GraphBuilder.RegisterExternalTexture(Scene->PathTracingSkylightPdf, TEXT("PathTracer.SkylightPdf"));
-		PathTracingParameters->SkylightInvResolution = 1.0f / PathTracingParameters->SkylightTexture->Desc.GetSize().X;
-		PathTracingParameters->SkylightMipCount = PathTracingParameters->SkylightPdf->Desc.NumMips;
+		SkylightParameters->SkylightTexture = GraphBuilder.RegisterExternalTexture(Scene->PathTracingSkylightTexture, TEXT("PathTracer.Skylight"));
+		SkylightParameters->SkylightPdf = GraphBuilder.RegisterExternalTexture(Scene->PathTracingSkylightPdf, TEXT("PathTracer.SkylightPdf"));
+		SkylightParameters->SkylightInvResolution = 1.0f / SkylightParameters->SkylightTexture->Desc.GetSize().X;
+		SkylightParameters->SkylightMipCount = SkylightParameters->SkylightPdf->Desc.NumMips;
 		return true;
 	}
 	RDG_EVENT_SCOPE(GraphBuilder, "Path Tracing SkylightPrepare");
@@ -650,17 +645,17 @@ bool PrepareSkyTexture(FRDGBuilder& GraphBuilder, FScene* Scene, const FViewInfo
 		SkyColor,
 		UseMISCompensation,
 		// Out
-		PathTracingParameters->SkylightTexture,
-		PathTracingParameters->SkylightPdf,
-		PathTracingParameters->SkylightInvResolution,
-		PathTracingParameters->SkylightMipCount
+		SkylightParameters->SkylightTexture,
+		SkylightParameters->SkylightPdf,
+		SkylightParameters->SkylightInvResolution,
+		SkylightParameters->SkylightMipCount
 	);
 
 	// hang onto these for next time (if caching is enabled)
 	if (IsSkylightCachingEnabled)
 	{
-		GraphBuilder.QueueTextureExtraction(PathTracingParameters->SkylightTexture, &Scene->PathTracingSkylightTexture);
-		GraphBuilder.QueueTextureExtraction(PathTracingParameters->SkylightPdf, &Scene->PathTracingSkylightPdf);
+		GraphBuilder.QueueTextureExtraction(SkylightParameters->SkylightTexture, &Scene->PathTracingSkylightTexture);
+		GraphBuilder.QueueTextureExtraction(SkylightParameters->SkylightPdf, &Scene->PathTracingSkylightPdf);
 	}
 	return true;
 }
@@ -780,7 +775,8 @@ void SetLightParameters(FRDGBuilder& GraphBuilder, FPathTracingRG::FParameters*
 
 	// Prepend SkyLight to light buffer since it is not part of the regular light list
 	const float Inf = std::numeric_limits<float>::infinity();
-	if (PrepareSkyTexture(GraphBuilder, Scene, View, UseMISCompensation, PassParameters))
+	const bool IsSkylightCachingEnabled = CVarPathTracingSkylightCaching.GetValueOnAnyThread() != 0;
+	if (PrepareSkyTexture(GraphBuilder, Scene, View, true, UseMISCompensation, IsSkylightCachingEnabled, &PassParameters->SkylightParameters))
 	{
 		check(Scene->SkyLight != nullptr);
 		FPathTracingLight& DestLight = Lights[NumLights++];
diff --git a/Engine/Source/Runtime/Renderer/Private/PathTracing.h b/Engine/Source/Runtime/Renderer/Private/PathTracing.h
new file mode 100644
--- /dev/null
+++ b/Engine/Source/Runtime/Renderer/Private/PathTracing.h
@@ -0,0 +1,20 @@
+// Copyright Epic Games, Inc. All Rights Reserved.
+
+#pragma once
+
+#include "ShaderParameterMacros.h"
+
+// this struct holds skylight parameters
+BEGIN_SHADER_PARAMETER_STRUCT(FPathTracingSkylight, )
+	SHADER_PARAMETER_RDG_TEXTURE(Texture2D, SkylightTexture)
+	SHADER_PARAMETER_RDG_TEXTURE(Texture2D, SkylightPdf)
+	SHADER_PARAMETER_SAMPLER(SamplerState, SkylightTextureSampler)
+	SHADER_PARAMETER(float, SkylightInvResolution)
+	SHADER_PARAMETER(int32, SkylightMipCount)
+END_SHADER_PARAMETER_STRUCT()
+
+class FRDGBuilder;
+class FScene;
+class FViewInfo;
+
+bool PrepareSkyTexture(FRDGBuilder& GraphBuilder, FScene* Scene, const FViewInfo& View, bool SkylightEnabled, bool UseMISCompensation, bool IsSkylightCachingEnabled, FPathTracingSkylight* SkylightParameters);
diff --git a/Engine/Source/Runtime/Renderer/Private/RayTracing/RayTracingGlobalIllumination.cpp b/Engine/Source/Runtime/Renderer/Private/RayTracing/RayTracingGlobalIllumination.cpp
--- a/Engine/Source/Runtime/Renderer/Private/RayTracing/RayTracingGlobalIllumination.cpp
+++ b/Engine/Source/Runtime/Renderer/Private/RayTracing/RayTracingGlobalIllumination.cpp
@@ -24,6 +24,7 @@
 #include "RayTracingDeferredMaterials.h"
 #include "RayTracingTypes.h"
 #include "PathTracingDefinitions.h"
+#include "PathTracing.h"
 
 static TAutoConsoleVariable<int32> CVarRayTracingGlobalIllumination(
 	TEXT("r.RayTracing.GlobalIllumination"),
@@ -244,30 +245,29 @@ DECLARE_GPU_STAT_NAMED(RayTracingGIBruteForce, TEXT("Ray Tracing GI: Brute Force
 DECLARE_GPU_STAT_NAMED(RayTracingGIFinalGather, TEXT("Ray Tracing GI: Final Gather"));
 DECLARE_GPU_STAT_NAMED(RayTracingGICreateGatherPoints, TEXT("Ray Tracing GI: Create Gather Points"));
 
-RENDERER_API void SetupLightParameters(
-	const FScene& Scene,
-	const FViewInfo& View, FRDGBuilder& GraphBuilder, FRDGBufferSRV** OutLightBuffer, uint32* OutLightCount)
+static void SetupLightParameters(
+	FScene* Scene,
+	const FViewInfo& View, FRDGBuilder& GraphBuilder,
+	FRDGBufferSRV** OutLightBuffer, uint32* OutLightCount, FPathTracingSkylight* SkylightParameters)
 {
 	FPathTracingLight Lights[RAY_TRACING_LIGHT_COUNT_MAXIMUM];
 	unsigned LightCount = 0;
 
 	// Get the SkyLight color
 
-	FSkyLightSceneProxy* SkyLight = Scene.SkyLight;
-	FVector SkyLightColor = FVector(0.0f, 0.0f, 0.0f);
-	uint32 SkyLightTransmission = 0;
-	if (SkyLight && SkyLight->bAffectGlobalIllumination && CVarRayTracingGlobalIlluminationSkyLight.GetValueOnRenderThread())
-	{
-		SkyLightColor = FVector(SkyLight->GetEffectiveLightColor());
-		SkyLightTransmission = SkyLight->bTransmission;
-	}
+	FSkyLightSceneProxy* SkyLight = Scene->SkyLight;
+
+	const bool bUseMISCompensation = false; // RTGI doesn't use MIS
+	const bool bUseSkylightCaching = true; // re-use the skylight texture if possible
+	const bool bSkylightEnabled = SkyLight && SkyLight->bAffectGlobalIllumination && CVarRayTracingGlobalIlluminationSkyLight.GetValueOnRenderThread() != 0;
 
 	// Prepend SkyLight to light buffer (if it is active)
+	if (PrepareSkyTexture(GraphBuilder, Scene, View, bSkylightEnabled, bUseMISCompensation, bUseSkylightCaching, SkylightParameters))
 	{
 		FPathTracingLight& DestLight = Lights[LightCount];
 
-		DestLight.Color = SkyLightColor;
-		DestLight.Flags = SkyLightTransmission ? PATHTRACER_FLAG_TRANSMISSION_MASK : 0;
+		DestLight.Color = FVector(1.0f, 1.0f, 1.0f);
+		DestLight.Flags = SkyLight->bTransmission ? PATHTRACER_FLAG_TRANSMISSION_MASK : 0;
 		// SkyLight does not have a LightingChannelMask
 		DestLight.Flags |= PATHTRACER_FLAG_LIGHTING_CHANNEL_MASK;
 		DestLight.Flags |= PATHTRACING_LIGHT_SKY;
@@ -276,8 +276,8 @@ RENDERER_API void SetupLightParameters(
 	}
 
 	
-	uint32 MaxLightCount = FMath::Min(CVarRayTracingGlobalIlluminationMaxLightCount.GetValueOnRenderThread(), RAY_TRACING_LIGHT_COUNT_MAXIMUM);
-	for (auto Light : Scene.Lights)
+	const uint32 MaxLightCount = FMath::Min(CVarRayTracingGlobalIlluminationMaxLightCount.GetValueOnRenderThread(), RAY_TRACING_LIGHT_COUNT_MAXIMUM);
+	for (auto Light : Scene->Lights)
 	{
 		if (LightCount >= MaxLightCount) break;
 
@@ -368,21 +368,6 @@ RENDERER_API void SetupLightParameters(
 	}
 }
 
-void SetupGlobalIlluminationSkyLightParameters(
-	const FScene& Scene,
-	FSkyLightData* SkyLightData)
-{
-	FSkyLightSceneProxy* SkyLight = Scene.SkyLight;
-
-	SetupSkyLightParameters(Scene, SkyLightData);
-
-	// Override the Sky Light color if it should not affect global illumination
-	if (SkyLight && !SkyLight->bAffectGlobalIllumination)
-	{
-		SkyLightData->Color = FVector(0.0f);
-	}
-}
-
 int32 GetRayTracingGlobalIlluminationSamplesPerPixel(const FViewInfo& View)
 {
 	int32 SamplesPerPixel = GRayTracingGlobalIlluminationSamplesPerPixel > -1 ? GRayTracingGlobalIlluminationSamplesPerPixel : View.FinalPostProcessSettings.RayTracingGISamplesPerPixel;
@@ -462,7 +447,7 @@ class FGlobalIlluminationRGS : public FGlobalShader
 		SHADER_PARAMETER_STRUCT_REF(FViewUniformShaderParameters, ViewUniformBuffer)
 		SHADER_PARAMETER_RDG_BUFFER_SRV(StructuredBuffer<FPathTracingLight>, SceneLights)
 		SHADER_PARAMETER(uint32, SceneLightCount)
-		SHADER_PARAMETER_STRUCT_REF(FSkyLightData, SkyLight)
+		SHADER_PARAMETER_STRUCT_INCLUDE(FPathTracingSkylight, SkylightParameters)
 		SHADER_PARAMETER_STRUCT_INCLUDE(FSceneTextureParameters, SceneTextures)
 		SHADER_PARAMETER_RDG_TEXTURE(Texture2D, SSProfilesTexture)
 		SHADER_PARAMETER_SAMPLER(SamplerState, TransmissionProfilesLinearSampler)
@@ -522,7 +507,7 @@ class FRayTracingGlobalIlluminationCreateGatherPointsRGS : public FGlobalShader
 		// Light data
 		SHADER_PARAMETER_RDG_BUFFER_SRV(StructuredBuffer<FPathTracingLight>, SceneLights)
 		SHADER_PARAMETER(uint32, SceneLightCount)
-		SHADER_PARAMETER_STRUCT_REF(FSkyLightData, SkyLight)
+		SHADER_PARAMETER_STRUCT_INCLUDE(FPathTracingSkylight, SkylightParameters)
 
 		// Shading data
 		SHADER_PARAMETER_STRUCT_INCLUDE(FSceneTextureParameters, SceneTextures)
@@ -589,7 +574,7 @@ class FRayTracingGlobalIlluminationCreateGatherPointsTraceRGS : public FGlobalSh
 		// Light data
 		SHADER_PARAMETER_RDG_BUFFER_SRV(StructuredBuffer<FPathTracingLight>, SceneLights)
 		SHADER_PARAMETER(uint32, SceneLightCount)
-		SHADER_PARAMETER_STRUCT_REF(FSkyLightData, SkyLight)
+		SHADER_PARAMETER_STRUCT_INCLUDE(FPathTracingSkylight, SkylightParameters)
 
 		// Shading data
 		SHADER_PARAMETER_STRUCT_INCLUDE(FSceneTextureParameters, SceneTextures)
@@ -840,7 +825,7 @@ void CopyGatherPassParameters(
 
 	NewParameters->SceneLights = PassParameters.SceneLights;
 	NewParameters->SceneLightCount = PassParameters.SceneLightCount;
-	NewParameters->SkyLight = PassParameters.SkyLight;
+	NewParameters->SkylightParameters = PassParameters.SkylightParameters;
 
 	NewParameters->SceneTextures = PassParameters.SceneTextures;
 	NewParameters->SSProfilesTexture = PassParameters.SSProfilesTexture;
@@ -881,7 +866,7 @@ void CopyGatherPassParameters(
 
 	NewParameters->SceneLightCount = PassParameters.SceneLightCount;
 	NewParameters->SceneLights = PassParameters.SceneLights;
-	NewParameters->SkyLight = PassParameters.SkyLight;
+	NewParameters->SkylightParameters = PassParameters.SkylightParameters;
 
 	NewParameters->SceneTextures = PassParameters.SceneTextures;
 	NewParameters->SSProfilesTexture = PassParameters.SSProfilesTexture;
@@ -929,9 +914,6 @@ void FDeferredShadingSceneRenderer::RayTracingGlobalIlluminationCreateGatherPoin
 		MaxShadowDistance = FMath::Max(0.0, 0.99 * Scene->SkyLight->SkyDistanceThreshold);
 	}
 
-	FSkyLightData SkyLightParameters;
-	SetupGlobalIlluminationSkyLightParameters(*Scene, &SkyLightParameters);
-
 	FRayTracingGlobalIlluminationCreateGatherPointsRGS::FParameters* PassParameters = GraphBuilder.AllocParameters<FRayTracingGlobalIlluminationCreateGatherPointsRGS::FParameters>();
 	PassParameters->SampleIndex = SampleIndex;
 	PassParameters->GatherSamplesPerPixel = GatherSamples;
@@ -955,9 +937,8 @@ void FDeferredShadingSceneRenderer::RayTracingGlobalIlluminationCreateGatherPoin
 	PassParameters->ViewUniformBuffer = View.ViewUniformBuffer;
 
 	// Light data
-	SetupLightParameters(*Scene, View, GraphBuilder, &PassParameters->SceneLights, &PassParameters->SceneLightCount);
+	SetupLightParameters(Scene, View, GraphBuilder, &PassParameters->SceneLights, &PassParameters->SceneLightCount, &PassParameters->SkylightParameters);
 	PassParameters->SceneTextures = SceneTextures;
-	PassParameters->SkyLight = CreateUniformBufferImmediate(SkyLightParameters, EUniformBufferUsage::UniformBuffer_SingleFrame);
 
 	// Shading data
 	TRefCountPtr<IPooledRenderTarget> SubsurfaceProfileRT((IPooledRenderTarget*)GetSubsufaceProfileTexture_RT(GraphBuilder.RHICmdList));
@@ -1254,9 +1235,6 @@ void FDeferredShadingSceneRenderer::RenderRayTracingGlobalIlluminationBruteForce
 		MaxShadowDistance = FMath::Max(0.0, 0.99 * Scene->SkyLight->SkyDistanceThreshold);
 	}
 
-	FSkyLightData SkyLightParameters;
-	SetupGlobalIlluminationSkyLightParameters(*Scene, &SkyLightParameters);
-
 	FGlobalIlluminationRGS::FParameters* PassParameters = GraphBuilder.AllocParameters<FGlobalIlluminationRGS::FParameters>();
 	PassParameters->SamplesPerPixel = RayTracingGISamplesPerPixel;
 	int32 CVarRayTracingGlobalIlluminationMaxBouncesValue = CVarRayTracingGlobalIlluminationMaxBounces.GetValueOnRenderThread();
@@ -1278,9 +1256,8 @@ void FDeferredShadingSceneRenderer::RenderRayTracingGlobalIlluminationBruteForce
 	PassParameters->NextEventEstimationSamples = GRayTracingGlobalIlluminationNextEventEstimationSamples;
 	PassParameters->TLAS = View.GetRayTracingSceneViewChecked();
 	PassParameters->ViewUniformBuffer = View.ViewUniformBuffer;
-	SetupLightParameters(*Scene, View, GraphBuilder, &PassParameters->SceneLights, &PassParameters->SceneLightCount);
+	SetupLightParameters(Scene, View, GraphBuilder, &PassParameters->SceneLights, &PassParameters->SceneLightCount, &PassParameters->SkylightParameters);
 	PassParameters->SceneTextures = SceneTextures;
-	PassParameters->SkyLight = CreateUniformBufferImmediate(SkyLightParameters, EUniformBufferUsage::UniformBuffer_SingleDraw);
 
 	// TODO: should be converted to RDG
 	TRefCountPtr<IPooledRenderTarget> SubsurfaceProfileRT((IPooledRenderTarget*) GetSubsufaceProfileTexture_RT(GraphBuilder.RHICmdList));
