Fix for cloud leaking on the side of screen due to out of half resolution depth buffer sampling (when tracing at quarter or half resolution)

diff --git a/Engine/Shaders/Private/VolumetricCloud.usf b/Engine/Shaders/Private/VolumetricCloud.usf
--- a/Engine/Shaders/Private/VolumetricCloud.usf
+++ b/Engine/Shaders/Private/VolumetricCloud.usf
@@ -496,7 +496,9 @@ void MainCommon(in FMaterialPixelParameters MaterialParameters, in float4 SvPosi
 		//
 
 		uint4 TracingCoordToZbufferCoordScaleBias = RenderVolumetricCloudParameters.TracingCoordToZbufferCoordScaleBias;
-		float DeviceZ = RenderVolumetricCloudParameters.SceneDepthTexture.Load(uint3(uint2(SvPosition.xy - 0.5) * TracingCoordToZbufferCoordScaleBias.xy + TracingCoordToZbufferCoordScaleBias.zw, 0)).r;
+		uint2 SceneDepthTextureCoord = uint2(SvPosition.xy - 0.5) * TracingCoordToZbufferCoordScaleBias.xy + TracingCoordToZbufferCoordScaleBias.zw;
+		SceneDepthTextureCoord = clamp(SceneDepthTextureCoord, RenderVolumetricCloudParameters.SceneDepthTextureMinMaxCoord.xy, RenderVolumetricCloudParameters.SceneDepthTextureMinMaxCoord.zw);
+		float DeviceZ = RenderVolumetricCloudParameters.SceneDepthTexture.Load(uint3(SceneDepthTextureCoord, 0)).r;
 #if HAS_INVERTED_Z_BUFFER
 		DeviceZ = max(0.000000000001, DeviceZ);
 #endif
diff --git a/Engine/Source/Runtime/Renderer/Private/VolumetricCloudRendering.cpp b/Engine/Source/Runtime/Renderer/Private/VolumetricCloudRendering.cpp
--- a/Engine/Source/Runtime/Renderer/Private/VolumetricCloudRendering.cpp
+++ b/Engine/Source/Runtime/Renderer/Private/VolumetricCloudRendering.cpp
@@ -496,6 +496,7 @@ BEGIN_GLOBAL_SHADER_PARAMETER_STRUCT(FRenderVolumetricCloudGlobalParameters, )
 	SHADER_PARAMETER(int32, VirtualShadowMapId0)
 //	SHADER_PARAMETER_STRUCT(FBlueNoise, BlueNoise)
 	SHADER_PARAMETER(FUintVector4, TracingCoordToZbufferCoordScaleBias)
+	SHADER_PARAMETER(FUintVector4, SceneDepthTextureMinMaxCoord)
 	SHADER_PARAMETER(uint32, NoiseFrameIndexModPattern)
 	SHADER_PARAMETER(int32, OpaqueIntersectionMode)
 	SHADER_PARAMETER(uint32, VolumetricRenderTargetMode)
@@ -1831,9 +1832,29 @@ static TRDGUniformBufferRef<FRenderVolumetricCloudGlobalParameters> CreateCloudP
 
 	FRenderVolumetricCloudGlobalParameters& VolumetricCloudParams = *GraphBuilder.AllocParameters<FRenderVolumetricCloudGlobalParameters>();
 	SetupDefaultRenderVolumetricCloudGlobalParameters(GraphBuilder, VolumetricCloudParams, CloudInfo, MainView);
-	VolumetricCloudParams.SceneDepthTexture = CloudRC.SceneDepthZ;
 
 	VolumetricCloudParams.SceneDepthTexture = CloudRC.SceneDepthZ;
+	if (CloudRC.bShouldViewRenderVolumetricRenderTarget && MainView.ViewState)
+	{
+		FVolumetricRenderTargetViewStateData& VRT = MainView.ViewState->VolumetricCloudRenderTarget;
+		const uint32 VolumetricReconstructRTDownsampleFactor = VRT.GetVolumetricReconstructRTDownsampleFactor();
+
+		// Use the main view to get the target rect. The clouds are reconstructed at the same resolution as the depth buffer read when tracing.
+		VolumetricCloudParams.SceneDepthTextureMinMaxCoord = FUintVector4(
+			MainView.ViewRect.Min.X / VolumetricReconstructRTDownsampleFactor,
+			MainView.ViewRect.Min.Y / VolumetricReconstructRTDownsampleFactor,
+			(MainView.ViewRect.Max.X - 1) / VolumetricReconstructRTDownsampleFactor,
+			(MainView.ViewRect.Max.Y - 1) / VolumetricReconstructRTDownsampleFactor);
+	}
+	else
+	{
+		VolumetricCloudParams.SceneDepthTextureMinMaxCoord = FUintVector4(
+			MainView.ViewRect.Min.X,
+			MainView.ViewRect.Min.Y,
+			(MainView.ViewRect.Max.X - 1),
+			(MainView.ViewRect.Max.Y - 1));
+	}
+
 	VolumetricCloudParams.Light0Shadow = CloudRC.LightShadowShaderParams0;
 	VolumetricCloudParams.VirtualShadowMapId0 = CloudRC.VirtualShadowMapId0;
 	VolumetricCloudParams.CloudShadowTexture0 = CloudRC.VolumetricCloudShadowTexture[0];
