Move ray tracing acceleration structure updates to overlap with BasePass

diff --git a/Engine/Source/Runtime/Renderer/Private/DeferredShadingRenderer.cpp b/Engine/Source/Runtime/Renderer/Private/DeferredShadingRenderer.cpp
--- a/Engine/Source/Runtime/Renderer/Private/DeferredShadingRenderer.cpp
+++ b/Engine/Source/Runtime/Renderer/Private/DeferredShadingRenderer.cpp
@@ -1168,7 +1168,7 @@
 RDG_BUFFER_ACCESS(RayTracingSceneScratchBuffer, ERHIAccess::UAVCompute)
 END_SHADER_PARAMETER_STRUCT()
 
-bool FDeferredShadingSceneRenderer::DispatchRayTracingWorldUpdates(FRDGBuilder& GraphBuilder)
+bool FDeferredShadingSceneRenderer::SetupRayTracingPipelineStates(FRHICommandListImmediate& RHICmdList)
 {
 	if (!IsRayTracingEnabled() || Views.Num() == 0)
 	{
@@ -1193,32 +1193,11 @@
 		return false;
 	}
 
-	TRACE_CPUPROFILER_EVENT_SCOPE(FDeferredShadingSceneRenderer::DispatchRayTracingWorldUpdates);
+	TRACE_CPUPROFILER_EVENT_SCOPE(FDeferredShadingSceneRenderer::SetupRayTracingPipelineStates);
 
-	// Make sure there are no pending skin cache builds and updates anymore:
-	// FSkeletalMeshObjectGPUSkin::UpdateDynamicData_RenderThread could have enqueued build operations which might not have
-	// been processed by CommitRayTracingGeometryUpdates. All pending builds should be done before adding them to the 
-	// top level BVH
-	if (FGPUSkinCache* GPUSkinCache = Scene->GetGPUSkinCache())
-	{
-		AddPass(GraphBuilder, [this](FRHICommandListImmediate& RHICmdList)
-			{
-				Scene->GetGPUSkinCache()->CommitRayTracingGeometryUpdates(RHICmdList);
-			});
-	}
-
-	GRayTracingGeometryManager.ProcessBuildRequests(GraphBuilder.RHICmdList);
-
 	const int32 ReferenceViewIndex = 0;
 	FViewInfo& ReferenceView = Views[ReferenceViewIndex];
-	FRayTracingScene& RayTracingScene = Scene->RayTracingScene;
 
-	if (RayTracingScene.GeometriesToBuild.Num() > 0)
-	{
-		// Force update all the collected geometries (use stack allocator?)
-		GRayTracingGeometryManager.ForceBuildIfPending(GraphBuilder.RHICmdList, RayTracingScene.GeometriesToBuild);
-	}
-	
 	if (ReferenceView.AddRayTracingMeshBatchTaskList.Num() > 0)
 	{
 		SCOPE_CYCLE_COUNTER(STAT_WaitRayTracingAddMesh);
@@ -1233,8 +1212,6 @@
 		ReferenceView.AddRayTracingMeshBatchTaskList.Empty();
 	}
 
-	RDG_GPU_MASK_SCOPE(GraphBuilder, FRHIGPUMask::All());
-
 	// #dxr_todo: UE-72565: refactor ray tracing effects to not be member functions of DeferredShadingRenderer. register each effect at startup and just loop over them automatically to gather all required shaders
 	TArray<FRHIRayTracingShader*> RayGenShaders;
 
@@ -1259,13 +1236,11 @@
 
 	DeduplicateRayGenerationShaders(RayGenShaders);
 
-	RayTracingScene.BeginCreate(GraphBuilder);
-
 	if (RayGenShaders.Num())
 	{
 		auto DefaultHitShader = ReferenceView.ShaderMap->GetShader<FOpaqueShadowHitGroup>().GetRayTracingShader();
 
-		ReferenceView.RayTracingMaterialPipeline = BindRayTracingMaterialPipeline(GraphBuilder.RHICmdList, ReferenceView,
+		ReferenceView.RayTracingMaterialPipeline = BindRayTracingMaterialPipeline(RHICmdList, ReferenceView,
 			RayGenShaders,
 			DefaultHitShader
 		);
@@ -1273,7 +1248,7 @@
 
 	// Initialize common resources used for lighting in ray tracing effects
 
-	ReferenceView.RayTracingSubSurfaceProfileTexture = GetSubsufaceProfileTexture_RT(GraphBuilder.RHICmdList);
+	ReferenceView.RayTracingSubSurfaceProfileTexture = GetSubsufaceProfileTexture_RT(RHICmdList);
 	if (!ReferenceView.RayTracingSubSurfaceProfileTexture)
 	{
 		ReferenceView.RayTracingSubSurfaceProfileTexture = GSystemTextures.BlackDummy;
@@ -1285,7 +1260,7 @@
 	{
 		FViewInfo& View = Views[ViewIndex];
 
-		View.RayTracingLightData = CreateRayTracingLightData(GraphBuilder.RHICmdList,
+		View.RayTracingLightData = CreateRayTracingLightData(RHICmdList,
 			Scene->Lights, View, EUniformBufferUsage::UniformBuffer_SingleFrame);
 
 		// Send common ray tracing resources from reference view to all others.
@@ -1298,6 +1273,63 @@
 		}
 	}
 
+	return true;
+}
+bool FDeferredShadingSceneRenderer::DispatchRayTracingWorldUpdates(FRDGBuilder& GraphBuilder)
+{
+	if (!IsRayTracingEnabled() || Views.Num() == 0)
+	{
+		return false;
+	}
+
+	bool bAnyRayTracingPassEnabled = false;
+	bool bPathTracingOrDebugViewEnabled = false;
+	for (int32 ViewIndex = 0; ViewIndex < Views.Num(); ViewIndex++)
+	{
+		bAnyRayTracingPassEnabled |= AnyRayTracingPassEnabled(Scene, Views[ViewIndex]);
+		bPathTracingOrDebugViewEnabled |= !CanOverlayRayTracingOutput(Views[ViewIndex]);
+	}
+
+	if (!bAnyRayTracingPassEnabled)
+	{
+		return false;
+	}
+
+	if (!bAnyRayTracingPassEnabled && !bPathTracingOrDebugViewEnabled)
+	{
+		return false;
+	}
+
+	TRACE_CPUPROFILER_EVENT_SCOPE(FDeferredShadingSceneRenderer::DispatchRayTracingWorldUpdates);
+
+	// Make sure there are no pending skin cache builds and updates anymore:
+	// FSkeletalMeshObjectGPUSkin::UpdateDynamicData_RenderThread could have enqueued build operations which might not have
+	// been processed by CommitRayTracingGeometryUpdates. All pending builds should be done before adding them to the 
+	// top level BVH
+	if (FGPUSkinCache* GPUSkinCache = Scene->GetGPUSkinCache())
+	{
+		AddPass(GraphBuilder, [this](FRHICommandListImmediate& RHICmdList)
+			{
+				Scene->GetGPUSkinCache()->CommitRayTracingGeometryUpdates(RHICmdList);
+			});
+	}
+
+	GRayTracingGeometryManager.ProcessBuildRequests(GraphBuilder.RHICmdList);
+
+	const int32 ReferenceViewIndex = 0;
+	FViewInfo& ReferenceView = Views[ReferenceViewIndex];
+	FRayTracingScene& RayTracingScene = Scene->RayTracingScene;
+
+	if (RayTracingScene.GeometriesToBuild.Num() > 0)
+	{
+		// Force update all the collected geometries (use stack allocator?)
+		GRayTracingGeometryManager.ForceBuildIfPending(GraphBuilder.RHICmdList, RayTracingScene.GeometriesToBuild);
+	}
+
+	RDG_GPU_MASK_SCOPE(GraphBuilder, FRHIGPUMask::All());
+
+	RayTracingScene.BeginCreate(GraphBuilder);
+
 	const bool bRayTracingAsyncBuild = CVarRayTracingAsyncBuild.GetValueOnRenderThread() != 0;
 
 	if (bRayTracingAsyncBuild && GRHISupportsRayTracingAsyncBuildAccelerationStructure)
@@ -1310,9 +1342,6 @@
 
 			FRHIAsyncComputeCommandListImmediate& RHIAsyncCmdList = FRHICommandListExecutor::GetImmediateAsyncComputeCommandList();
 
-			// TArray<FRHIUnorderedAccessView*> VertexBuffers;
-			// Scene->GetRayTracingDynamicGeometryCollection()->GetVertexBufferUAVs(VertexBuffers);
-
 			RHICmdList.BeginTransition(RayTracingDynamicGeometryUpdateBeginTransition);
 			RHIAsyncCmdList.EndTransition(RayTracingDynamicGeometryUpdateBeginTransition);
 
@@ -1430,6 +1459,8 @@
 	FBuildAccelerationStructurePassParams* PassParams = GraphBuilder.AllocParameters<FBuildAccelerationStructurePassParams>();
 	PassParams->RayTracingSceneScratchBuffer = Scene->RayTracingScene.BuildScratchBuffer;
 
+	SetupRayTracingPipelineStates(GraphBuilder.RHICmdList);
+
 	GraphBuilder.AddPass(RDG_EVENT_NAME("WaitForRayTracingScene"), PassParams, ERDGPassFlags::Compute | ERDGPassFlags::NeverCull,
 		[this, PassParams](FRHICommandListImmediate& RHICmdList)
 	{
@@ -1553,6 +1584,22 @@
 	});
 }
 
+enum class ERayTracingWorldUpdatesDispatchPoint
+{
+	BeforeLumenSceneLighting,
+	OverlapWithBasePass
+};
+
+ERayTracingWorldUpdatesDispatchPoint GetRayTracingWorldUpdatesDispatchPoint(bool bOcclusionBeforeBasePass, bool bLumenUseHardwareRayTracedShadows)
+{
+	if (bOcclusionBeforeBasePass && bLumenUseHardwareRayTracedShadows)
+	{
+		return ERayTracingWorldUpdatesDispatchPoint::BeforeLumenSceneLighting;
+	}
+
+	return ERayTracingWorldUpdatesDispatchPoint::OverlapWithBasePass;
+}
+
 #endif // RHI_RAYTRACING
 
 static TAutoConsoleVariable<float> CVarStallInitViews(
@@ -2059,11 +2106,6 @@
 		}
 	}
 
-#if RHI_RAYTRACING
-	// Must be done after FGlobalDynamicVertexBuffer::Get().Commit() for dynamic geometries to be updated
-	DispatchRayTracingWorldUpdates(GraphBuilder);
-#endif
-
 	TArray<Nanite::FRasterResults, TInlineAllocator<2>> NaniteRasterResults;
 	if (bNaniteEnabled && Views.Num() > 0)
 	{
@@ -2217,6 +2259,10 @@
 	// Early occlusion queries
 	const bool bOcclusionBeforeBasePass = !bNaniteEnabled && !bAnyLumenEnabled && !bHairEnable && ((DepthPass.EarlyZPassMode == EDepthDrawingMode::DDM_AllOccluders) || bIsEarlyDepthComplete);
 
+#if RHI_RAYTRACING
+	ERayTracingWorldUpdatesDispatchPoint RayTracingWorldUpdatesDispatchPoint = GetRayTracingWorldUpdatesDispatchPoint(bOcclusionBeforeBasePass, Lumen::UseHardwareRayTracedShadows(Views[0]));
+#endif
+
 	if (bOcclusionBeforeBasePass)
 	{
 		RenderOcclusionLambda();
@@ -2286,6 +2332,19 @@
 
 	if (bOcclusionBeforeBasePass)
 	{
+#if RHI_RAYTRACING
+		if (RayTracingWorldUpdatesDispatchPoint == ERayTracingWorldUpdatesDispatchPoint::BeforeLumenSceneLighting)
 		{
+			DispatchRayTracingWorldUpdates(GraphBuilder);
+		}
+
+		// Lumen scene lighting requires ray tracing scene to be ready if HWRT shadows are desired
+		if (Lumen::UseHardwareRayTracedShadows(Views[0]))
+		{
+			WaitForRayTracingScene(GraphBuilder);
+		}
+#endif // RHI_RAYTRACING
+
+		{
 			LLM_SCOPE_BYTAG(Lumen);
 			RenderLumenSceneLighting(GraphBuilder, Views[0]);
@@ -2340,6 +2399,14 @@
 		InitTranslucencyLightingVolumeTextures(GraphBuilder, Views, ERDGPassFlags::AsyncCompute, TranslucencyLightingVolumeTextures);
 	}
 
+#if RHI_RAYTRACING
+	if (RayTracingWorldUpdatesDispatchPoint == ERayTracingWorldUpdatesDispatchPoint::OverlapWithBasePass)
 	{
+		// Async AS builds can potentially overlap with BasePass
+		DispatchRayTracingWorldUpdates(GraphBuilder);
+	}
+#endif
+
+	{
 		RenderBasePass(GraphBuilder, SceneTextures, DBufferTextures, BasePassDepthStencilAccess, ForwardScreenSpaceShadowMaskTexture, InstanceCullingManager);
 		AddServiceLocalQueuePass(GraphBuilder);
@@ -2447,14 +2514,6 @@
 		RenderHairBasePass(GraphBuilder, Scene, SceneTextures, Views, InstanceCullingManager);
 	}
 
-#if RHI_RAYTRACING
-	const bool bRayTracingEnabled = IsRayTracingEnabled();
-	if (bRayTracingEnabled)
-	{
-		WaitForRayTracingScene(GraphBuilder);
-	}
-#endif // RHI_RAYTRACING
-
 	// Shadow and fog after base pass
 	if (bCanOverlayRayTracingOutput && !bOcclusionBeforeBasePass)
 	{
@@ -2462,7 +2521,15 @@
 		AllocateVirtualShadowMaps(bAfterBasePass);
 
 		RenderShadowDepthMaps(GraphBuilder, InstanceCullingManager);
-		
+
+#if RHI_RAYTRACING
+		// Lumen scene lighting requires ray tracing scene to be ready if HWRT shadows are desired
+		if (Lumen::UseHardwareRayTracedShadows(Views[0]))
 		{
+			WaitForRayTracingScene(GraphBuilder);
+		}
+#endif // RHI_RAYTRACING
+
+		{
 			LLM_SCOPE_BYTAG(Lumen);
 			RenderLumenSceneLighting(GraphBuilder, Views[0]);
@@ -2571,6 +2638,14 @@
 		AddClearStencilPass(GraphBuilder, SceneTextures.Depth.Target);
 	}
 
+#if RHI_RAYTRACING
+	// If Lumen is not using HWRT shadows, we can wait until here: before Lumen diffuse indirect
+	if (!Lumen::UseHardwareRayTracedShadows(Views[0]))
+	{
+		WaitForRayTracingScene(GraphBuilder);
+	}
+#endif // RHI_RAYTRACING
+
 	if (bRenderDeferredLighting)
 	{
 		RDG_GPU_STAT_SCOPE(GraphBuilder, RenderDeferredLighting);
@@ -2598,7 +2673,7 @@
 		}
 
 #if RHI_RAYTRACING
-		if (bRayTracingEnabled)
+		if (IsRayTracingEnabled())
 		{
 			RenderDitheredLODFadingOutMask(GraphBuilder, Views[0], SceneTextures.Depth.Target);
 		}
@@ -2847,7 +2922,7 @@
 	}
 
 #if RHI_RAYTRACING
-	if (bRayTracingEnabled)
+	if (IsRayTracingEnabled())
 	{
 		for (int32 ViewIndex = 0; ViewIndex < Views.Num(); ViewIndex++)
 		{
diff --git a/Engine/Source/Runtime/Renderer/Private/DeferredShadingRenderer.h b/Engine/Source/Runtime/Renderer/Private/DeferredShadingRenderer.h
--- a/Engine/Source/Runtime/Renderer/Private/DeferredShadingRenderer.h
+++ b/Engine/Source/Runtime/Renderer/Private/DeferredShadingRenderer.h
@@ -958,6 +958,7 @@ private:
 	/** Fills RayTracingScene instance list for the given View and adds relevant ray tracing data to the view. Does not reset previous scene contents. */
 	bool GatherRayTracingWorldInstancesForView(FRHICommandListImmediate& RHICmdList, FViewInfo& View, FRayTracingScene& RayTracingScene);
 
+	bool SetupRayTracingPipelineStates(FRHICommandListImmediate& RHICmdList);
 	bool DispatchRayTracingWorldUpdates(FRDGBuilder& GraphBuilder);
 	FRayTracingPipelineState* BindRayTracingMaterialPipeline(FRHICommandList& RHICmdList, FViewInfo& View, const TArrayView<FRHIRayTracingShader*>& RayGenShaderTable, FRHIRayTracingShader* DefaultClosestHitShader);
 	FRayTracingPipelineState* BindRayTracingDeferredMaterialGatherPipeline(FRHICommandList& RHICmdList, const FViewInfo& View, const TArrayView<FRHIRayTracingShader*>& RayGenShaderTable);
