Fix visible in ray tracing flag

diff --git a/Engine/Source/Runtime/Renderer/Private/PrimitiveSceneInfo.cpp b/Engine/Source/Runtime/Renderer/Private/PrimitiveSceneInfo.cpp
--- a/Engine/Source/Runtime/Renderer/Private/PrimitiveSceneInfo.cpp
+++ b/Engine/Source/Runtime/Renderer/Private/PrimitiveSceneInfo.cpp
@@ -715,7 +715,14 @@ void FPrimitiveSceneInfo::CacheRayTracingPrimitives(FRHICommandListImmediate& RH
 			{
 				// Legacy path for static meshes.
 				// TODO: convert them to this new path
-				Flags = ERayTracingPrimitiveFlags::ComputeLOD | ERayTracingPrimitiveFlags::CacheMeshCommands;
+				if (Flags == ERayTracingPrimitiveFlags::Dynamic)
+				{
+					Flags = ERayTracingPrimitiveFlags::ComputeLOD | ERayTracingPrimitiveFlags::CacheMeshCommands;
+				}
+				else
+				{
+					Flags = ERayTracingPrimitiveFlags::Excluded;
+				}
 				continue;
 			}
 
