Nukes r.UsePreExposure

diff --git a/Engine/Config/BaseEngine.ini b/Engine/Config/BaseEngine.ini
--- a/Engine/Config/BaseEngine.ini
+++ b/Engine/Config/BaseEngine.ini
@@ -2391,7 +2391,6 @@ DefaultASTCQualityBySize=3
 +VersionedIntRValues=r.Shaders.Optimize
 +VersionedIntRValues=r.CompileShadersForDevelopment
 +VersionedIntRValues=r.MobileHDR
-+VersionedIntRValues=r.UsePreExposure
 
 [/Script/UnrealEd.ChunkDependencyInfo]
 ; Entries can be added to this section of DefaultEngine.ini to set dependencies for package chunks, used by cooking and packaging
diff --git a/Engine/Source/Editor/PixelInspector/Private/PixelInspectorResult.h b/Engine/Source/Editor/PixelInspector/Private/PixelInspectorResult.h
--- a/Engine/Source/Editor/PixelInspector/Private/PixelInspectorResult.h
+++ b/Engine/Source/Editor/PixelInspector/Private/PixelInspectorResult.h
@@ -66,7 +66,7 @@ namespace PixelInspector
 
 
 		//////////////////////////////////////////////////////////////////////////
-		// PreExposure used to render this frame. See "r.UsePreExposure"
+		// PreExposure used to render this frame.
 		float PreExposure;
 		float OneOverPreExposure;
 
diff --git a/Engine/Source/Runtime/Engine/Classes/Engine/RendererSettings.h b/Engine/Source/Runtime/Engine/Classes/Engine/RendererSettings.h
--- a/Engine/Source/Runtime/Engine/Classes/Engine/RendererSettings.h
+++ b/Engine/Source/Runtime/Engine/Classes/Engine/RendererSettings.h
@@ -498,18 +498,6 @@ class ENGINE_API URendererSettings : public UDeveloperSettings
 		ConfigRestartRequired=true))
 	uint32 bExtendDefaultLuminanceRangeInAutoExposureSettings: 1;
 
-	UPROPERTY(config, EditAnywhere, Category = DefaultSettings, meta = (
-		ConsoleVariable = "r.UsePreExposure", DisplayName = "Apply Pre-exposure before writing to the scene color",
-		ToolTip = "Whether to use pre-exposure to remap the range of the scene color around the camera exposure. This limits the render target range required to support HDR lighting value.",
-		ConfigRestartRequired=true))
-	uint32 bUsePreExposure : 1;
-
-	UPROPERTY(config, EditAnywhere, Category = DefaultSettings, meta = (
-		EditCondition = "bUsePreExposure", ConsoleVariable = "r.EyeAdaptation.EditorOnly", DisplayName = "Enable pre-exposure only in the editor",
-		ToolTip = "Whether pre-exposure should be an editor only feature. This is to because it currently has an impact on the renderthread performance.",
-		ConfigRestartRequired=false))
-	uint32 bEnablePreExposureOnlyInTheEditor : 1;
-
 	UPROPERTY(config, EditAnywhere, Category = DefaultSettings, meta = (
 		ConsoleVariable = "r.DefaultFeature.MotionBlur", DisplayName = "Motion Blur",
 		ToolTip = "Whether the default for MotionBlur is enabled or not (postprocess volume/camera/game setting can still override and enable or disable it independently)"))
diff --git a/Engine/Source/Runtime/Renderer/Private/PostProcess/PostProcessEyeAdaptation.cpp b/Engine/Source/Runtime/Renderer/Private/PostProcess/PostProcessEyeAdaptation.cpp
--- a/Engine/Source/Runtime/Renderer/Private/PostProcess/PostProcessEyeAdaptation.cpp
+++ b/Engine/Source/Runtime/Renderer/Private/PostProcess/PostProcessEyeAdaptation.cpp
@@ -66,22 +66,9 @@ namespace
 		TEXT("will have no influence.\n"),
 		ECVF_RenderThreadSafe);
 
-	static TAutoConsoleVariable<int32> CVarEnablePreExposureOnlyInTheEditor(
-		TEXT("r.EyeAdaptation.EditorOnly"),
-		0,
-		TEXT("When pre-exposure is enabled, 0 to enable it everywhere, 1 to enable it only in the editor (default).\n")
-		TEXT("This is to because it currently has an impact on the renderthread performance\n"),
-		ECVF_ReadOnly);
 	const ERHIFeatureLevel::Type BasicEyeAdaptationMinFeatureLevel = ERHIFeatureLevel::ES3_1;
 }
 
-// Function is static because EyeAdaptation is the only system that should be checking this value.
-static bool UsePreExposureEnabled()
-{
-	static const auto CVarUsePreExposure = IConsoleManager::Get().FindTConsoleVariableDataInt(TEXT("r.UsePreExposure"));
-	return CVarUsePreExposure->GetValueOnAnyThread() != 0;
-}
-
 bool IsAutoExposureMethodSupported(ERHIFeatureLevel::Type FeatureLevel, EAutoExposureMethod AutoExposureMethodId)
 {
 	switch (AutoExposureMethodId)
@@ -755,7 +742,7 @@ void FSceneViewState::FEyeAdaptationManager::SwapTextures(FRDGBuilder& GraphBuil
 	const FRHIGPUMask LastFrameGPUMask = AFRUtils::GetPrevSiblingGPUMask(ThisFrameGPUMask);
 	const FRHIGPUMask LastLastFrameGPUMask = AFRUtils::GetPrevSiblingGPUMask(LastFrameGPUMask);
 
-	if (bInUpdateLastExposure && PooledRenderTarget[CurrentBuffer].IsValid() && (GIsEditor || CVarEnablePreExposureOnlyInTheEditor.GetValueOnRenderThread() == 0))
+	if (bInUpdateLastExposure && PooledRenderTarget[CurrentBuffer].IsValid())
 	{
 		FRDGTextureRef CurrentTexture = GraphBuilder.RegisterExternalTexture(PooledRenderTarget[CurrentBuffer], ERenderTargetTexture::ShaderResource, ERDGTextureFlags::MultiFrame);
 
@@ -864,7 +851,7 @@ void FSceneViewState::FEyeAdaptationManager::SwapBuffers(FRDGBuilder& GraphBuild
 	const FRHIGPUMask LastFrameGPUMask = AFRUtils::GetPrevSiblingGPUMask(ThisFrameGPUMask);
 	const FRHIGPUMask LastLastFrameGPUMask = AFRUtils::GetPrevSiblingGPUMask(LastFrameGPUMask);
 
-	if (bInUpdateLastExposure && ExposureBufferData[CurrentBuffer].IsValid() && (GIsEditor || CVarEnablePreExposureOnlyInTheEditor.GetValueOnRenderThread() == 0))
+	if (bInUpdateLastExposure && ExposureBufferData[CurrentBuffer].IsValid())
 	{
 		FRDGBufferRef CurrentRDGBuffer = GraphBuilder.RegisterExternalBuffer(ExposureBufferData[CurrentBuffer], ERDGBufferFlags::MultiFrame);
 
@@ -934,33 +921,19 @@ void FSceneViewState::UpdatePreExposure(FViewInfo& View)
 
 	if (bIsPreExposureRelevant && bEnableAutoExposure)
 	{
-		if (UsePreExposureEnabled())
+		const float PreExposureOverride = CVarEyeAdaptationPreExposureOverride.GetValueOnRenderThread();
+		const float LastExposure = View.GetLastEyeAdaptationExposure();
+		if (PreExposureOverride > 0)
 		{
-			const float PreExposureOverride = CVarEyeAdaptationPreExposureOverride.GetValueOnRenderThread();
-			const float LastExposure = View.GetLastEyeAdaptationExposure();
-			if (PreExposureOverride > 0)
-			{
-				PreExposure = PreExposureOverride;
-			}
-			else if (LastExposure > 0)
-			{
-				PreExposure = LastExposure;
-			}
-
-			bUpdateLastExposure = true;
+			PreExposure = PreExposureOverride;
 		}
-		else
+		else if (LastExposure > 0)
 		{
-			if (View.FinalPostProcessSettings.AutoExposureBiasCurve)
-			{
-				// The exposure compensation curves require the scene average luminance
-				bUpdateLastExposure = true;
-			}
-
-			// Whe PreExposure is not used, we will override the PreExposure value to 1.0, which simulates PreExposure being off.
-			PreExposure = 1.0f;
+			PreExposure = LastExposure;
 		}
 
+		bUpdateLastExposure = true;
+
 		// Mobile LDR does not support post-processing but still can apply Exposure during basepass
 		if (bMobilePlatform && !IsMobileHDR())
 		{
diff --git a/Engine/Source/Runtime/Renderer/Private/PostProcess/TemporalAA.cpp b/Engine/Source/Runtime/Renderer/Private/PostProcess/TemporalAA.cpp
--- a/Engine/Source/Runtime/Renderer/Private/PostProcess/TemporalAA.cpp
+++ b/Engine/Source/Runtime/Renderer/Private/PostProcess/TemporalAA.cpp
@@ -1132,13 +1132,6 @@ static void AddGen5MainTemporalAAPasses(
 	const FTemporalAAHistory& InputHistory = View.PrevViewInfo.TemporalAAHistory;
 	FTemporalAAHistory* OutputHistory = &View.ViewState->PrevFrameViewInfo.TemporalAAHistory;
 
-	// Gen5 Temporal AA no longer use the view.
-	{
-		static IConsoleVariable* CVarShowTransitions = IConsoleManager::Get().FindConsoleVariable(TEXT("r.UsePreExposure"));
-		const bool bUsePreExposure = CVarShowTransitions->GetInt() != 0;
-		ensureMsgf(bUsePreExposure, TEXT("r.TemporalAA.Algorithm=1 requires r.UsePreExposure=1"));
-	}
-
 #if COMPILE_TAA_DEBUG_PASSES
 	const bool bSetupDebugPasses = CVarTAASetupDebugPasses.GetValueOnRenderThread() != 0;
 #endif
diff --git a/Engine/Source/Runtime/Renderer/Private/SceneRendering.cpp b/Engine/Source/Runtime/Renderer/Private/SceneRendering.cpp
--- a/Engine/Source/Runtime/Renderer/Private/SceneRendering.cpp
+++ b/Engine/Source/Runtime/Renderer/Private/SceneRendering.cpp
@@ -170,15 +170,6 @@ static TAutoConsoleVariable<int32> CVarRoundRobinOcclusion(
 	TEXT("0 to disable round-robin occlusion queries for stereo rendering (default), 1 to enable."),
 	ECVF_Scalability | ECVF_RenderThreadSafe);
 
-static TAutoConsoleVariable<int32> CVarUsePreExposure(
-	TEXT("r.UsePreExposure"),
-	1,
-	TEXT("0 to disable pre-exposure, 1 to enable it (default).\n")
-	TEXT("Pre-exposure allows the engine to apply the last frame exposure to luminance values before writing them in rendertargets.\n")
-	TEXT("It avoids rendertarget overflow when using low precision formats like fp16.\n")
-	TEXT("The pre-exposure value can be overriden through r.EyeAdaptation.PreExposureOverride\n"),
-	ECVF_ReadOnly | ECVF_RenderThreadSafe);
-
 static TAutoConsoleVariable<int32> CVarODSCapture(
 	TEXT("vr.ODSCapture"),
 	0,
