Strata stencil categorisation is not using blending nor pixel shader.
Debug can now display both simple/complex path when enabled in green/red.

diff --git a/Engine/Shaders/Private/Strata/StrataMaterialClassification.usf b/Engine/Shaders/Private/Strata/StrataMaterialClassification.usf
--- a/Engine/Shaders/Private/Strata/StrataMaterialClassification.usf
+++ b/Engine/Shaders/Private/Strata/StrataMaterialClassification.usf
@@ -274,6 +274,7 @@ float4 OutputViewSizeAndInvSize;
 float4 OutputBufferSizeAndInvSize;
 float4x4 ViewScreenToTranslatedWorld;
 Buffer<uint> TileListBuffer;
+float4 DebugTileColor;
 
 void StrataTilePassVS(
 	in uint InVertexId : SV_VertexID,
@@ -313,7 +314,8 @@ void StencilMainPS(
 	const bool bTileX = (InTileCoord.x & 1) == 0;
 	const bool bTileY = (InTileCoord.y & 1) == 0;
 	const bool bChecker = (bTileX && bTileY) || (!bTileX && !bTileY);
-	OutColor0 = bChecker ? float4(0.0f, 0.5f, 0.f, 1.0f) : float4(0.f, 1.0f, 0.f, 1.0f);
+	OutColor0 = DebugTileColor * float4(1.0f, 1.0f, 1.0f, bChecker ? 0.33f : 0.66f);
+	OutColor0.rgb *= OutColor0.a; // premultiplied alpha blending assuming alpha is opacity==coverage
 }
 
 #endif //SHADER_STENCIL_CATEGORIZATION
diff --git a/Engine/Source/Runtime/Renderer/Private/Strata/Strata.cpp b/Engine/Source/Runtime/Renderer/Private/Strata/Strata.cpp
--- a/Engine/Source/Runtime/Renderer/Private/Strata/Strata.cpp
+++ b/Engine/Source/Runtime/Renderer/Private/Strata/Strata.cpp
@@ -35,7 +35,7 @@ static TAutoConsoleVariable<int32> CVarStrataClassification(
 static TAutoConsoleVariable<int32> CVarStrataClassificationDebug(
 	TEXT("r.Strata.Classification.Debug"),
 	0,
-	TEXT("Enable strata classification visualization: 1 shows tile with simple material only while 2 shows tile with complex material."),
+	TEXT("Enable strata classification visualization: 1 shows simple material tiles in green and complex material tiles in red."),
 	ECVF_RenderThreadSafe);
 
 static TAutoConsoleVariable<int32> CVarStrataClassificationPassesReadingStrataAreTiled(
@@ -513,6 +513,7 @@ class FStrataMaterialStencilTaggingPassPS : public FGlobalShader
 
 	BEGIN_SHADER_PARAMETER_STRUCT(FParameters, )
 		SHADER_PARAMETER_STRUCT_INCLUDE(Strata::FStrataTilePassVS::FParameters, VS)
+		SHADER_PARAMETER(FVector4, DebugTileColor)
 		RENDER_TARGET_BINDING_SLOTS()
 	END_SHADER_PARAMETER_STRUCT()
 
@@ -574,7 +575,7 @@ static void AddStrataInternalClassificationTilePass(
 	FillUpTiledPassData(TileMaterialType, View, ParametersPS->VS, StrataTilePrimitiveType);
 
 	FStrataTilePassVS::FPermutationDomain VSPermutationVector;
-	VSPermutationVector.Set< FStrataTilePassVS::FEnableDebug >(true);
+	VSPermutationVector.Set< FStrataTilePassVS::FEnableDebug >(bDebug);
 	VSPermutationVector.Set< FStrataTilePassVS::FEnableTexCoordScreenVector >(false);
 	TShaderMapRef<FStrataTilePassVS> VertexShader(View.ShaderMap, VSPermutationVector);
 	TShaderMapRef<FStrataMaterialStencilTaggingPassPS> PixelShader(View.ShaderMap);
@@ -584,6 +585,17 @@ static void AddStrataInternalClassificationTilePass(
 	{
 		check(ColorTexture);
 		ParametersPS->RenderTargets[0] = FRenderTargetBinding(*ColorTexture, ERenderTargetLoadAction::ELoad);
+		switch (TileMaterialType)
+		{
+		case EStrataTileMaterialType::ESimple:
+			ParametersPS->DebugTileColor = FVector4(0.0f, 1.0f, 0.0f, 1.0);
+			break;
+		case EStrataTileMaterialType::EComplex:
+			ParametersPS->DebugTileColor = FVector4(1.0f, 0.0f, 0.0f, 1.0);
+			break;
+		default:
+			check(false);
+		}
 	}
 	else
 	{
@@ -593,6 +605,7 @@ static void AddStrataInternalClassificationTilePass(
 			ERenderTargetLoadAction::ELoad,
 			ERenderTargetLoadAction::ELoad,
 			FExclusiveDepthStencil::DepthNop_StencilWrite);
+		ParametersPS->DebugTileColor = FVector4(ForceInitToZero);
 	}
 	
 	GraphBuilder.AddPass(
@@ -605,14 +618,20 @@ static void AddStrataInternalClassificationTilePass(
 		{
 			FGraphicsPipelineStateInitializer GraphicsPSOInit;
 			RHICmdList.ApplyCachedRenderTargets(GraphicsPSOInit);
-			GraphicsPSOInit.BlendState = TStaticBlendState<CW_RGBA, BO_Add, BF_One, BF_One, BO_Max, BF_SourceAlpha, BF_DestAlpha>::GetRHI();
 			GraphicsPSOInit.RasterizerState = TStaticRasterizerState<>::GetRHI();
 			if (bDebug)
 			{
+				// Use premultiplied alpha blending, pixel shader and depth/stencil is off
+				GraphicsPSOInit.BoundShaderState.PixelShaderRHI = PixelShader.GetPixelShader();
+				SetShaderParameters(RHICmdList, PixelShader, PixelShader.GetPixelShader(), *ParametersPS);
+				GraphicsPSOInit.BlendState = TStaticBlendState<CW_RGBA, BO_Add, BF_One, BF_InverseSourceAlpha, BO_Add, BF_Zero, BF_One>::GetRHI();
 				GraphicsPSOInit.DepthStencilState = TStaticDepthStencilState<false, CF_Always>::GetRHI();
 			}
 			else
 			{
+				// No blending and no pixel shader required. Stencil will be writen to.
+				GraphicsPSOInit.BoundShaderState.PixelShaderRHI = nullptr;
+				GraphicsPSOInit.BlendState = TStaticBlendState<>::GetRHI();
 				GraphicsPSOInit.DepthStencilState = TStaticDepthStencilState<
 					false, CF_Always, 
 					true,  CF_Always, SO_Keep, SO_Keep, SO_Replace,
@@ -621,7 +640,6 @@ static void AddStrataInternalClassificationTilePass(
 			}
 			GraphicsPSOInit.BoundShaderState.VertexDeclarationRHI = GFilterVertexDeclaration.VertexDeclarationRHI;
 			GraphicsPSOInit.BoundShaderState.VertexShaderRHI = VertexShader.GetVertexShader();
-			GraphicsPSOInit.BoundShaderState.PixelShaderRHI = PixelShader.GetPixelShader();
 			GraphicsPSOInit.PrimitiveType = StrataTilePrimitiveType;
 			SetGraphicsPipelineState(RHICmdList, GraphicsPSOInit);
 			SetShaderParameters(RHICmdList, VertexShader, VertexShader.GetVertexShader(), ParametersPS->VS);
@@ -902,10 +920,11 @@ void AddStrataDebugPasses(FRDGBuilder& GraphBuilder, const TArray<FViewInfo>& Vi
 		for (int32 i = 0; i < Views.Num(); ++i)
 		{
 			const FViewInfo& View = Views[i];
-			const int32 StrataTileMaterialTypeIndex = FMath::Clamp(StrataClassificationDebug - 1, 0, 1);
 			const bool bDebugPass = true;
 			AddStrataInternalClassificationTilePass(
-				GraphBuilder, View, nullptr, &SceneColorTexture, (EStrataTileMaterialType)StrataTileMaterialTypeIndex, bDebugPass);
+				GraphBuilder, View, nullptr, &SceneColorTexture, EStrataTileMaterialType::ESimple, bDebugPass);
+			AddStrataInternalClassificationTilePass(
+				GraphBuilder, View, nullptr, &SceneColorTexture, EStrataTileMaterialType::EComplex, bDebugPass);
 		}
 	}
 
