Fixed crash when building HZBFurthest from a 1px RT.

diff --git a/Engine/Source/Runtime/Renderer/Private/SceneTextureReductions.cpp b/Engine/Source/Runtime/Renderer/Private/SceneTextureReductions.cpp
--- a/Engine/Source/Runtime/Renderer/Private/SceneTextureReductions.cpp
+++ b/Engine/Source/Runtime/Renderer/Private/SceneTextureReductions.cpp
@@ -123,7 +123,7 @@ void BuildHZB(
 	HZBSize.X = FMath::Max( FPlatformMath::RoundUpToPowerOfTwo( ViewRect.Width() ) >> 1, 1u );
 	HZBSize.Y = FMath::Max( FPlatformMath::RoundUpToPowerOfTwo( ViewRect.Height() ) >> 1, 1u );
 
-	int32 NumMips = FMath::Log2( FMath::Max<float>( HZBSize.X, HZBSize.Y ) );
+	int32 NumMips = FMath::Max(FMath::FloorToInt(FMath::Log2(FMath::Max<float>(HZBSize.X, HZBSize.Y))), 1);
 
 	bool bReduceClosestDepth = OutClosestHZBTexture != nullptr;
 	bool bUseCompute = bReduceClosestDepth || HZBBuildUseCompute(ShaderPlatform);
