FRayTracingGeometryInstance now only uses array views. Actual bulk data is typically owned by either FRayTracingScene or by CachedRayTracingInstanceTransforms in FSceneProxy.

Added ERayTracingInstanceFlags, which covers all equivalent DXR functionality. Explicitly specified that FRayTracingGeometryInstance.NumTransforms is what defines how many elements will be read from CPU or GPU transform buffers (which may themselves be larger). Added FRayTracingGeometryInstance.CommonUserData that's used if UserData array view is empty. Otherwise UserData must have at least as many elements as NumTransforms.

Also fix System.Renderer.RayTracing.BasicRayTracing failure, as a side-effect.

diff --git a/Engine/Plugins/Experimental/GPULightmass/Source/GPULightmass/Private/LightmapRenderer.cpp b/Engine/Plugins/Experimental/GPULightmass/Source/GPULightmass/Private/LightmapRenderer.cpp
--- a/Engine/Plugins/Experimental/GPULightmass/Source/GPULightmass/Private/LightmapRenderer.cpp
+++ b/Engine/Plugins/Experimental/GPULightmass/Source/GPULightmass/Private/LightmapRenderer.cpp
@@ -613,11 +613,14 @@ void FCachedRayTracingSceneData::SetupFromSceneRenderState(FSceneRenderState& Sc
 				int32 InstanceIndex = RayTracingGeometryInstances.AddDefaulted(1);
 				FRayTracingGeometryInstance& RayTracingInstance = RayTracingGeometryInstances[InstanceIndex];
 				RayTracingInstance.GeometryRHI = Instance.RenderData->LODResources[Instance.ClampedMinLOD].RayTracingGeometry.RayTracingGeometryRHI;
-				RayTracingInstance.Transforms.Add(Instance.LocalToWorld);
+				RayTracingInstance.Transforms = MakeArrayView(&Instance.LocalToWorld, 1);
 				RayTracingInstance.NumTransforms = 1;
-				RayTracingInstance.UserData.Add((uint32)StaticMeshIndex);
+				RayTracingInstance.DefaultUserData = (uint32)StaticMeshIndex;
 				RayTracingInstance.Mask = 0xFF;
-				RayTracingInstance.bForceOpaque = bAllSegmentsOpaque;
+				if (bAllSegmentsOpaque)
+				{
+					RayTracingInstance.Flags |= ERayTracingInstanceFlags::ForceOpaque;
+				}
 
 				for (int32 SegmentIndex = 0; SegmentIndex < MeshBatches.Num(); SegmentIndex++)
 				{
@@ -657,23 +660,32 @@ void FCachedRayTracingSceneData::SetupFromSceneRenderState(FSceneRenderState& Sc
 					FRayTracingGeometryInstance& RayTracingInstance = RayTracingGeometryInstances[InstanceIndex];
 					RayTracingInstance.GeometryRHI = InstanceGroup.ComponentUObject->GetStaticMesh()->GetRenderData()->LODResources[0].RayTracingGeometry.RayTracingGeometryRHI;
 
-					RayTracingInstance.Transforms.AddZeroed(InstanceGroup.InstancedRenderData->PerInstanceRenderData->InstanceBuffer.GetNumInstances());
+					const int32 NumInstances = (int32)InstanceGroup.InstancedRenderData->PerInstanceRenderData->InstanceBuffer.GetNumInstances();
+					TArrayView<FMatrix> NewTransforms = MakeArrayView(
+						OwnedRayTracingInstanceTransforms.Add_GetRef(TUniquePtr<FMatrix>(new FMatrix[NumInstances])).Get(),
+						NumInstances);
 
-					for (int32 InstanceIdx = 0; InstanceIdx < (int32)InstanceGroup.InstancedRenderData->PerInstanceRenderData->InstanceBuffer.GetNumInstances(); InstanceIdx++)
+					for (int32 InstanceIdx = 0; InstanceIdx < NumInstances; InstanceIdx++)
 					{
 						FMatrix Transform;
 						InstanceGroup.InstancedRenderData->PerInstanceRenderData->InstanceBuffer.GetInstanceTransform(InstanceIdx, Transform);
-						Transform.M[3][3] = 1.0f;
+						Transform.M[3][3] = 1.0f; // GPUCULL_TODO: GetInstanceTransform currently sets this to 0 due to *reasons*
+
 						FMatrix InstanceTransform = Transform * InstanceGroup.LocalToWorld;
 
-						RayTracingInstance.Transforms[InstanceIdx] = InstanceTransform;
+						NewTransforms[InstanceIdx] = InstanceTransform;
 					}
 
-					RayTracingInstance.NumTransforms = RayTracingInstance.Transforms.Num();
+					RayTracingInstance.Transforms = NewTransforms;
+					RayTracingInstance.NumTransforms = NumInstances;
 
-					RayTracingInstance.UserData.Add((uint32)(Scene.StaticMeshInstanceRenderStates.Elements.Num() + InstanceGroupIndex));
+					RayTracingInstance.DefaultUserData = (uint32)(Scene.StaticMeshInstanceRenderStates.Elements.Num() + InstanceGroupIndex);
 					RayTracingInstance.Mask = 0xFF;
-					RayTracingInstance.bForceOpaque = bAllSegmentsOpaque;
+
+					if (bAllSegmentsOpaque)
+					{
+						RayTracingInstance.Flags |= ERayTracingInstanceFlags::ForceOpaque;
+					}
 
 					for (int32 SegmentIndex = 0; SegmentIndex < MeshBatches.Num(); SegmentIndex++)
 					{
@@ -913,9 +925,9 @@ void FSceneRenderState::SetupRayTracingScene()
 
 						FRayTracingGeometryInstance& RayTracingInstance = RayTracingGeometryInstances[InstanceIndex];
 						RayTracingInstance.GeometryRHI = Landscape.SectionRayTracingStates[SubSectionIdx]->Geometry.RayTracingGeometryRHI;
-						RayTracingInstance.Transforms.Add(FMatrix::Identity);
+						RayTracingInstance.Transforms = MakeArrayView(&FMatrix::Identity, 1);
 						RayTracingInstance.NumTransforms = 1;
-						RayTracingInstance.UserData.Add((uint32)InstanceIndex);
+						RayTracingInstance.DefaultUserData = (uint32)InstanceIndex;
 						RayTracingInstance.Mask = 0xFF;
 
 						TArray<FMeshBatch> MeshBatches = Landscape.GetMeshBatchesForGBufferRendering(0);
@@ -952,7 +964,10 @@ void FSceneRenderState::SetupRayTracingScene()
 							RayTracingInstance.Mask = 0;
 						}
 
-						RayTracingInstance.bForceOpaque = bAllSegmentsOpaque;
+						if (bAllSegmentsOpaque)
+						{
+							RayTracingInstance.Flags |= ERayTracingInstanceFlags::ForceOpaque;
+						}
 					}
 				}
 			}
@@ -970,7 +985,7 @@ void FSceneRenderState::SetupRayTracingScene()
 			uint32 NumTotalInstances = 0;
 			for (const FRayTracingGeometryInstance& Instance : RayTracingGeometryInstances)
 			{
-				ensure(Instance.NumTransforms >= uint32(Instance.GetTransforms().Num()));
+				ensure(Instance.NumTransforms == uint32(Instance.Transforms.Num()));
 				NumTotalInstances += Instance.NumTransforms;
 			}
 
diff --git a/Engine/Plugins/Experimental/GPULightmass/Source/GPULightmass/Private/Scene/Scene.h b/Engine/Plugins/Experimental/GPULightmass/Source/GPULightmass/Private/Scene/Scene.h
--- a/Engine/Plugins/Experimental/GPULightmass/Source/GPULightmass/Private/Scene/Scene.h
+++ b/Engine/Plugins/Experimental/GPULightmass/Source/GPULightmass/Private/Scene/Scene.h
@@ -9,6 +9,7 @@
 #include "MeshPassProcessor.h"
 #include "IrradianceCaching.h"
 #include "GPULightmassSettings.h"
+#include "Templates/UniquePtr.h"
 
 class FGPULightmass;
 
@@ -75,6 +76,7 @@ struct FCachedRayTracingSceneData
 	TArray<uint32> InstanceDataOriginalOffsets;
 
 	TArray<FRayTracingGeometryInstance> RayTracingGeometryInstances;
+	TArray<TUniquePtr<FMatrix>> OwnedRayTracingInstanceTransforms;
 
 	TUniformBufferRef<FViewUniformShaderParameters> CachedViewUniformBuffer;
 
diff --git a/Engine/Source/Runtime/D3D12RHI/Private/D3D12RayTracing.cpp b/Engine/Source/Runtime/D3D12RHI/Private/D3D12RayTracing.cpp
--- a/Engine/Source/Runtime/D3D12RHI/Private/D3D12RayTracing.cpp
+++ b/Engine/Source/Runtime/D3D12RHI/Private/D3D12RayTracing.cpp
@@ -3333,7 +3333,7 @@ FD3D12RayTracingScene::FD3D12RayTracingScene(FD3D12Adapter* Adapter, const FRayT
 			ReferencedGeometries.Add(Geometry);
 		}
 
-		if (UNLIKELY(Instance.GPUTransformsSRV && Instance.NumTransforms))
+		if (UNLIKELY(Instance.GPUTransformsSRV != nullptr && Instance.NumTransforms))
 		{
 			FInstanceCopyCommand Command;
 			Command.Source = Instance.GPUTransformsSRV;
@@ -3362,41 +3362,48 @@ FD3D12RayTracingScene::FD3D12RayTracingScene(FD3D12Adapter* Adapter, const FRayT
 
 		InstanceDesc.InstanceMask = Instance.Mask;
 		InstanceDesc.InstanceContributionToHitGroupIndex = SegmentPrefixSum[InstanceIndex] * ShaderSlotsPerGeometrySegment;
-		InstanceDesc.Flags = D3D12_RAYTRACING_INSTANCE_FLAG_TRIANGLE_FRONT_COUNTERCLOCKWISE; // #dxr_todo: convert cull mode based on instance mirroring or double-sidedness
 
-		if (Instance.bForceOpaque || GRayTracingDebugForceOpaque)
+		if (EnumHasAnyFlags(Instance.Flags, ERayTracingInstanceFlags::TriangleCullDisable) || GRayTracingDebugDisableTriangleCull)
+		{
+			InstanceDesc.Flags |= D3D12_RAYTRACING_INSTANCE_FLAG_TRIANGLE_CULL_DISABLE;
+		}
+
+		if (!EnumHasAnyFlags(Instance.Flags, ERayTracingInstanceFlags::TriangleCullReverse))
+		{
+			// Counterclockwise is the default for UE. Reversing culling is achieved by *not* setting this flag.
+			InstanceDesc.Flags |= D3D12_RAYTRACING_INSTANCE_FLAG_TRIANGLE_FRONT_COUNTERCLOCKWISE;
+		}
+
+		if (EnumHasAnyFlags(Instance.Flags, ERayTracingInstanceFlags::ForceOpaque) || GRayTracingDebugForceOpaque)
 		{
 			InstanceDesc.Flags |= D3D12_RAYTRACING_INSTANCE_FLAG_FORCE_OPAQUE;
 		}
 
-		if (Instance.bDoubleSided || GRayTracingDebugDisableTriangleCull)
+		if (EnumHasAnyFlags(Instance.Flags, ERayTracingInstanceFlags::ForceNonOpaque))
 		{
-			InstanceDesc.Flags |= D3D12_RAYTRACING_INSTANCE_FLAG_TRIANGLE_CULL_DISABLE;
+			InstanceDesc.Flags |= D3D12_RAYTRACING_INSTANCE_FLAG_FORCE_NON_OPAQUE;
 		}
 
 		const uint32 NumTransforms = Instance.NumTransforms;
 
-		checkf(Instance.UserData.Num() == 0 || Instance.UserData.Num() == 1 || Instance.UserData.Num() == NumTransforms,
-			TEXT("User data array must be either be empty (implicit 0 is assumed), contain a single entry (same value applied to all instances) or contain one entry per entry in Transforms array."));
+		checkf(Instance.UserData.Num() == 0 || Instance.UserData.Num() >= int32(NumTransforms),
+			TEXT("User data array must be either be empty (Instance.DefaultUserData is used), or contain one entry per entry in Transforms array."));
 
-		const bool bUseUniqueUserData = Instance.UserData.Num() > 1;
-		const uint32 CommonUserData = Instance.UserData.Num() == 1 ? Instance.UserData[0] : 0;
+		const bool bUseUniqueUserData = Instance.UserData.Num() != 0;
 
-		const bool bGpuInstance = Instance.GPUTransformsSRV.IsValid();
+		const bool bGpuInstance = Instance.GPUTransformsSRV != nullptr;
 		const bool bCpuInstance = !bGpuInstance;
 
-		TArrayView<const FMatrix> InstanceTransforms = Instance.GetTransforms();
-
 		uint32 DescIndex = BaseInstancePrefixSum[InstanceIndex];
 		for (uint32 TransformIndex = 0; TransformIndex < NumTransforms; ++TransformIndex)
 		{
-			InstanceDesc.InstanceID = bUseUniqueUserData ? Instance.UserData[TransformIndex] : CommonUserData;
+			InstanceDesc.InstanceID = bUseUniqueUserData ? Instance.UserData[TransformIndex] : Instance.DefaultUserData;
 
 			if (LIKELY(bCpuInstance))
 			{
 				// DXR uses a 3x4 transform matrix in row major layout
 
-				const FMatrix& Transform = InstanceTransforms[TransformIndex];
+				const FMatrix& Transform = Instance.Transforms[TransformIndex];
 
 				InstanceDesc.Transform[0][0] = Transform.M[0][0];
 				InstanceDesc.Transform[0][1] = Transform.M[1][0];
diff --git a/Engine/Source/Runtime/RHI/Private/DynamicRHI.cpp b/Engine/Source/Runtime/RHI/Private/DynamicRHI.cpp
--- a/Engine/Source/Runtime/RHI/Private/DynamicRHI.cpp
+++ b/Engine/Source/Runtime/RHI/Private/DynamicRHI.cpp
@@ -22,6 +22,10 @@
 IMPLEMENT_TYPE_LAYOUT(FRayTracingGeometryInitializer);
 IMPLEMENT_TYPE_LAYOUT(FRayTracingGeometrySegment);
 
+static_assert(sizeof(FRayTracingGeometryInstance) <= 64,
+	"Ray tracing instance descriptor is expected to be no more than 64 bytes, "
+	"as there may be a very large number of them.");
+
 #ifndef PLATFORM_ALLOW_NULL_RHI
 	#define PLATFORM_ALLOW_NULL_RHI		0
 #endif
diff --git a/Engine/Source/Runtime/RHI/Public/DynamicRHI.h b/Engine/Source/Runtime/RHI/Public/DynamicRHI.h
--- a/Engine/Source/Runtime/RHI/Public/DynamicRHI.h
+++ b/Engine/Source/Runtime/RHI/Public/DynamicRHI.h
@@ -79,53 +79,49 @@ struct FRHIFlipDetails
 	{}
 };
 
+enum class ERayTracingInstanceFlags : uint8
+{
+	None                  = 0,
+	TriangleCullDisable   = 1 << 1, // No back face culling. Triangle is visible from both sides.
+	TriangleCullReverse   = 1 << 2, // Makes triangle front-facing if its vertices are counterclockwise from ray origin.
+	ForceOpaque           = 1 << 3, // Disable any-hit shader invocation for this instance.
+	ForceNonOpaque        = 1 << 4, // Force any-hit shader invocation even if geometries inside the instance were marked opaque.
+};
+ENUM_CLASS_FLAGS(ERayTracingInstanceFlags);
+
+/**
+* High level descriptor of one or more instances of a mesh in a ray tracing scene.
+* All instances covered by this descriptor will share shader bindings, but may have different transforms and user data.
+*/
 struct FRayTracingGeometryInstance
 {
 	FRHIRayTracingGeometry* GeometryRHI = nullptr;
 
-	// A physical FRayTracingGeometryInstance may be duplicated many times in the scene with different transforms and user data.
+	// A single physical mesh may be duplicated many times in the scene with different transforms and user data.
 	// All copies share the same shader binding table entries and therefore will have the same material and shader resources.
-	TArray<FMatrix, TInlineAllocator<1>> Transforms;
-
-	// Similar to Transforms (but mutually exclusive with it). Allows passing transforms by pointer.
-	TArrayView<const FMatrix> TransformsView;
+	TArrayView<const FMatrix> Transforms;
 
-	TArrayView<const FMatrix> GetTransforms() const
-	{
-		if (TransformsView.Num())
-		{
-			check(Transforms.Num() == 0);
-			return TransformsView;
-		}
-		else
-		{
-			check(TransformsView.Num() == 0);
-			return TArrayView<const FMatrix>(Transforms);
-		}
-	}
+	// Optional buffer that stores GPU transforms. Used instead of CPU-side transform data.
+	FRHIShaderResourceView* GPUTransformsSRV = nullptr;
 
-	// Transforms count. When GPU transforms are used it is a conservative count 
+	// Conservative number of instances. Some of the actual instances may be made inactive if GPU transforms are used.
+	// Must be less or equal to number of entries in Transforms view if CPU transform data is used.
+	// Must be less or equal to number of entries in GPUTransformsSRV if it is non-null.
 	uint32 NumTransforms = 0;
 
-	// Buffer that stores GPU transforms
-	FShaderResourceViewRHIRef GPUTransformsSRV;
-	
 	// Each geometry copy can receive a user-provided integer, which can be used to retrieve extra shader parameters or customize appearance.
 	// This data can be retrieved using GetInstanceUserData() in closest/any hit shaders.
-	// If UserData is empty, then 0 will be implicitly used for all entries.
-	// If UserData contains a single entry, then it will be applied to all instances/copies.
-	// Otherwise one UserData entry must be provided per entry in Transforms array.
-	TArray<uint32, TInlineAllocator<1>> UserData;
+	// If UserData view is empty, then DefaultUserData value will be used for all instances.
+	// If UserData view is used, then it must have the same number of entries as NumInstances.
+	uint32 DefaultUserData = 0;
+	TArrayView<const uint32> UserData;
 
 	// Mask that will be tested against one provided to TraceRay() in shader code.
 	// If binary AND of instance mask with ray mask is zero, then the instance is considered not intersected / invisible.
 	uint8 Mask = 0xFF;
 
-	// No any-hit shaders will be invoked for this geometry instance (only closest hit)
-	bool bForceOpaque = false;
-
-	// Disabling this will allow ray hits to be registered for front and back faces
-	bool bDoubleSided = false;
+	// Flags to control triangle back face culling, whether to allow any-hit shaders, etc.
+	ERayTracingInstanceFlags Flags = ERayTracingInstanceFlags::None;
 };
 
 enum ERayTracingGeometryType
diff --git a/Engine/Source/Runtime/Renderer/Private/DeferredShadingRenderer.cpp b/Engine/Source/Runtime/Renderer/Private/DeferredShadingRenderer.cpp
--- a/Engine/Source/Runtime/Renderer/Private/DeferredShadingRenderer.cpp
+++ b/Engine/Source/Runtime/Renderer/Private/DeferredShadingRenderer.cpp
@@ -946,10 +946,16 @@ bool FDeferredShadingSceneRenderer::GatherRayTracingWorldInstancesForView(FRHICo
 
 					FRayTracingGeometryInstance& RayTracingInstance = RayTracingScene.Instances.AddDefaulted_GetRef();
 					RayTracingInstance.GeometryRHI = Geometry->RayTracingGeometryRHI;
-					RayTracingInstance.UserData.Add((uint32)PrimitiveIndex);
+					RayTracingInstance.DefaultUserData = PrimitiveIndex;
 					RayTracingInstance.Mask = Instance.Mask;
-					RayTracingInstance.bForceOpaque = Instance.bForceOpaque;
-					RayTracingInstance.bDoubleSided = Instance.bDoubleSided;
+					if (Instance.bForceOpaque)
+					{
+						RayTracingInstance.Flags |= ERayTracingInstanceFlags::ForceOpaque;
+					}
+					if (Instance.bDoubleSided)
+					{
+						RayTracingInstance.Flags |= ERayTracingInstanceFlags::TriangleCullDisable;
+					}
 
 					if (Instance.InstanceGPUTransformsSRV.IsValid())
 					{
@@ -962,17 +968,19 @@ bool FDeferredShadingSceneRenderer::GatherRayTracingWorldInstancesForView(FRHICo
 						{
 							// Slow path: copy transforms to the owned storage
 							checkf(Instance.InstanceTransformsView.Num() == 0, TEXT("InstanceTransformsView is expected to be empty if using InstanceTransforms"));
-							RayTracingInstance.NumTransforms = Instance.InstanceTransforms.Num();
-							RayTracingInstance.Transforms.SetNumUninitialized(Instance.InstanceTransforms.Num());
-							FMemory::Memcpy(RayTracingInstance.Transforms.GetData(), Instance.InstanceTransforms.GetData(), Instance.InstanceTransforms.Num() * sizeof(RayTracingInstance.Transforms[0]));
-							static_assert(TIsSame<decltype(RayTracingInstance.Transforms[0]), decltype(Instance.InstanceTransforms[0])>::Value, "Unexpected transform type");
+							TArrayView<FMatrix> SceneOwnedTransforms = RayTracingScene.Allocate<FMatrix>(Instance.InstanceTransforms.Num());
+							FMemory::Memcpy(SceneOwnedTransforms.GetData(), Instance.InstanceTransforms.GetData(), Instance.InstanceTransforms.Num() * sizeof(RayTracingInstance.Transforms[0]));
+							static_assert(TIsSame<decltype(SceneOwnedTransforms[0]), decltype(Instance.InstanceTransforms[0])>::Value, "Unexpected transform type");
+
+							RayTracingInstance.NumTransforms = SceneOwnedTransforms.Num();
+							RayTracingInstance.Transforms = SceneOwnedTransforms;
 						}
 						else
 						{
 							// Fast path: just reference persistently-allocated transforms and avoid a copy
 							checkf(Instance.InstanceTransforms.Num() == 0, TEXT("InstanceTransforms is expected to be empty if using InstanceTransformsView"));
 							RayTracingInstance.NumTransforms = Instance.InstanceTransformsView.Num();
-							RayTracingInstance.TransformsView = Instance.InstanceTransformsView;
+							RayTracingInstance.Transforms = Instance.InstanceTransformsView;
 						}
 					}
 
@@ -1057,8 +1065,59 @@ bool FDeferredShadingSceneRenderer::GatherRayTracingWorldInstancesForView(FRHICo
 			FTaskGraphInterface::Get().WaitUntilTasksComplete(LODTaskList, ENamedThreads::GetRenderThread_Local());
 		}
 
-		Experimental::TSherwoodMap<uint64, int32> InstanceSet;
-		InstanceSet.Reserve(RelevantPrimitives.Num());
+		struct FAutoInstanceBatch
+		{
+			int32 Index = INDEX_NONE;
+
+			// Copies the next transform and user data into the current batch, returns true if arrays were re-allocated.
+			bool Add(FRayTracingScene& RayTracingScene, const FMatrix& InTransform, uint32 InUserData)
+			{
+				// Adhoc TArray-like resize behavior, in lieu of support for using a custom FMemStackBase in TArray.
+				// Idea for future: if batch becomes large enough, we could actually split it into multiple instances to avoid memory waste.
+
+				const bool bNeedReallocation = Cursor == Transforms.Num();
+
+				if (bNeedReallocation)
+				{
+					int32 PrevCount = Transforms.Num();
+					int32 NextCount = FMath::Max(PrevCount * 2, 1);
+
+					TArrayView<FMatrix> NewTransforms = RayTracingScene.Allocate<FMatrix>(NextCount);
+					if (PrevCount)
+					{
+						FMemory::Memcpy(NewTransforms.GetData(), Transforms.GetData(), Transforms.GetTypeSize() * Transforms.Num());
+					}
+					Transforms = NewTransforms;
+
+					TArrayView<uint32> NewUserData = RayTracingScene.Allocate<uint32>(NextCount);
+					if (PrevCount)
+					{
+						FMemory::Memcpy(NewUserData.GetData(), UserData.GetData(), UserData.GetTypeSize() * UserData.Num());
+					}
+					UserData = NewUserData;
+				}
+
+				Transforms[Cursor] = InTransform;
+				UserData[Cursor] = InUserData;
+
+				++Cursor;
+
+				return bNeedReallocation;
+			}
+
+			bool IsValid() const
+			{
+				return Transforms.Num() != 0;
+			}
+
+			TArrayView<FMatrix> Transforms;
+			TArrayView<uint32> UserData;
+			uint32 Cursor = 0;
+		};
+
+		Experimental::TSherwoodMap<uint64, FAutoInstanceBatch> InstanceBatches;
+
+		InstanceBatches.Reserve(RelevantPrimitives.Num());
 
 		// scan relevant primitives computing hash data to look for duplicate instances
 		for (const FRelevantPrimitive& RelevantPrimitive : RelevantPrimitives)
@@ -1105,19 +1164,29 @@ bool FDeferredShadingSceneRenderer::GatherRayTracingWorldInstancesForView(FRHICo
 				const int32 NewInstanceIndex = RayTracingScene.Instances.Num();
 				const uint64 InstanceKey = RelevantPrimitive.InstancingKey();
 
-				const int32 Index = bAutoInstance ? InstanceSet.FindOrAdd(InstanceKey, NewInstanceIndex) : NewInstanceIndex;
+				FAutoInstanceBatch DummyInstanceBatch = { NewInstanceIndex };
+				FAutoInstanceBatch& InstanceBatch = bAutoInstance ? InstanceBatches.FindOrAdd(InstanceKey, DummyInstanceBatch) : DummyInstanceBatch;
 
-				if (Index != NewInstanceIndex)
+				if (InstanceBatch.Index != NewInstanceIndex)
 				{
-					// reusing a previous entry, just append to the instance list
-					FRayTracingGeometryInstance& RayTracingInstance = RayTracingScene.Instances[Index];
-					RayTracingInstance.NumTransforms++;
-					RayTracingInstance.Transforms.Add(Scene->PrimitiveTransforms[PrimitiveIndex]);
-					RayTracingInstance.UserData.Add((uint32)PrimitiveIndex);
+					// Reusing a previous entry, just append to the instance list.
+
+					FRayTracingGeometryInstance& RayTracingInstance = RayTracingScene.Instances[InstanceBatch.Index];
+					bool bReallocated = InstanceBatch.Add(RayTracingScene, Scene->PrimitiveTransforms[PrimitiveIndex], (uint32)PrimitiveIndex);
+
+					++RayTracingInstance.NumTransforms;
+					check(RayTracingInstance.NumTransforms == InstanceBatch.Cursor); // sanity check
 
+					if (bReallocated)
+					{
+						RayTracingInstance.Transforms = InstanceBatch.Transforms;
+						RayTracingInstance.UserData = InstanceBatch.UserData;
+					}
 				}
 				else
 				{
+					// Starting new instance batch
+
 					for (int32 CommandIndex : RelevantPrimitive.CachedRayTracingMeshCommandIndices)
 					{
 						if (CommandIndex >= 0)
@@ -1136,16 +1205,24 @@ bool FDeferredShadingSceneRenderer::GatherRayTracingWorldInstancesForView(FRHICo
 					}
 
 					FRayTracingGeometryInstance& RayTracingInstance = RayTracingScene.Instances.AddDefaulted_GetRef();
-					RayTracingInstance.NumTransforms = 1;
-					RayTracingInstance.Transforms.SetNumUninitialized(1);
-					RayTracingInstance.UserData.SetNumUninitialized(1);
 
 					RayTracingInstance.GeometryRHI = RelevantPrimitive.RayTracingGeometryRHI;
-					RayTracingInstance.Transforms[0] = Scene->PrimitiveTransforms[PrimitiveIndex];
-					RayTracingInstance.UserData[0] = (uint32)PrimitiveIndex;
+
+					InstanceBatch.Add(RayTracingScene, Scene->PrimitiveTransforms[PrimitiveIndex], (uint32)PrimitiveIndex);
+					RayTracingInstance.Transforms = InstanceBatch.Transforms;
+					RayTracingInstance.UserData = InstanceBatch.UserData;
+					RayTracingInstance.NumTransforms = 1;
+
 					RayTracingInstance.Mask = RelevantPrimitive.InstanceMask; // When no cached command is found, InstanceMask == 0 and the instance is effectively filtered out
-					RayTracingInstance.bForceOpaque = RelevantPrimitive.bAllSegmentsOpaque;
-					RayTracingInstance.bDoubleSided = RelevantPrimitive.bTwoSided;
+
+					if (RelevantPrimitive.bAllSegmentsOpaque)
+					{
+						RayTracingInstance.Flags |= ERayTracingInstanceFlags::ForceOpaque;
+					}
+					if (RelevantPrimitive.bTwoSided)
+					{
+						RayTracingInstance.Flags |= ERayTracingInstanceFlags::TriangleCullDisable;
+					}
 				}
 			}
 		}
diff --git a/Engine/Source/Runtime/Renderer/Private/PrimitiveSceneInfo.cpp b/Engine/Source/Runtime/Renderer/Private/PrimitiveSceneInfo.cpp
--- a/Engine/Source/Runtime/Renderer/Private/PrimitiveSceneInfo.cpp
+++ b/Engine/Source/Runtime/Renderer/Private/PrimitiveSceneInfo.cpp
@@ -763,19 +763,26 @@ void FPrimitiveSceneInfo::CacheRayTracingPrimitives(FRHICommandListImmediate& RH
 				SceneInfo->CachedRayTracingInstanceWorldTransforms.Empty();
 				SceneInfo->CachedRayTracingInstanceWorldTransforms.AddUninitialized(CachedRayTracingInstance.NumTransforms);
 				SceneInfo->UpdateCachedRayTracingInstanceTransforms(SceneInfo->Proxy->GetLocalToWorld());
-				SceneInfo->CachedRayTracingInstance.TransformsView = MakeArrayView(SceneInfo->CachedRayTracingInstanceWorldTransforms);
+				SceneInfo->CachedRayTracingInstance.Transforms = MakeArrayView(SceneInfo->CachedRayTracingInstanceWorldTransforms);
 
-				check(SceneInfo->CachedRayTracingInstance.NumTransforms >= uint32(SceneInfo->CachedRayTracingInstance.GetTransforms().Num()));
+				check(SceneInfo->CachedRayTracingInstance.NumTransforms >= uint32(SceneInfo->CachedRayTracingInstance.Transforms.Num()));
 
-				SceneInfo->CachedRayTracingInstance.UserData.SetNumUninitialized(1);
 				SceneInfo->CachedRayTracingInstance.GeometryRHI = CachedRayTracingInstance.Geometry->RayTracingGeometryRHI;
+
 				// At this point (in AddToScene()) PrimitiveIndex has been set
 				check(SceneInfo->GetIndex() != INDEX_NONE);
-				SceneInfo->CachedRayTracingInstance.UserData[0] = (uint32)SceneInfo->GetIndex();
+				SceneInfo->CachedRayTracingInstance.DefaultUserData = (uint32)SceneInfo->GetIndex();
 				SceneInfo->CachedRayTracingInstance.Mask = CachedRayTracingInstance.Mask; // When no cached command is found, InstanceMask == 0 and the instance is effectively filtered out
-				SceneInfo->CachedRayTracingInstance.bForceOpaque = CachedRayTracingInstance.bForceOpaque;
-				SceneInfo->CachedRayTracingInstance.bDoubleSided = CachedRayTracingInstance.bDoubleSided;
 
+				if (CachedRayTracingInstance.bForceOpaque)
+				{
+					SceneInfo->CachedRayTracingInstance.Flags |= ERayTracingInstanceFlags::ForceOpaque;
+				}
+
+				if (CachedRayTracingInstance.bDoubleSided)
+				{
+					SceneInfo->CachedRayTracingInstance.Flags |= ERayTracingInstanceFlags::TriangleCullDisable;
+				}
 			}
 
 		}
diff --git a/Engine/Source/Runtime/Renderer/Private/RayTracing/RayTracingScene.cpp b/Engine/Source/Runtime/Renderer/Private/RayTracing/RayTracingScene.cpp
--- a/Engine/Source/Runtime/Renderer/Private/RayTracing/RayTracingScene.cpp
+++ b/Engine/Source/Runtime/Renderer/Private/RayTracing/RayTracingScene.cpp
@@ -23,10 +23,10 @@ FGraphEventRef FRayTracingScene::BeginCreate(FRDGBuilder& GraphBuilder)
 	WaitForTasks();
 
 	uint32 NumTotalInstances = 0;
-	for (const FRayTracingGeometryInstance& Instance : Instances)
+	for (const FRayTracingGeometryInstance& InstanceDesc : Instances)
 	{
-		ensure(Instance.NumTransforms >= uint32(Instance.GetTransforms().Num()));
-		NumTotalInstances += Instance.NumTransforms;
+		ensure(InstanceDesc.NumTransforms <= uint32(InstanceDesc.Transforms.Num()));
+		NumTotalInstances += InstanceDesc.NumTransforms;
 	}
 	SET_DWORD_STAT(STAT_RayTracingInstances, NumTotalInstances);
 
@@ -133,6 +133,8 @@ void FRayTracingScene::Reset()
 	Instances.Reset();
 	GeometriesToBuild.Reset();
 
+	Allocator.Flush();
+
 	BuildScratchBuffer = nullptr;
 }
 
diff --git a/Engine/Source/Runtime/Renderer/Private/RayTracing/RayTracingScene.h b/Engine/Source/Runtime/Renderer/Private/RayTracing/RayTracingScene.h
--- a/Engine/Source/Runtime/Renderer/Private/RayTracing/RayTracingScene.h
+++ b/Engine/Source/Runtime/Renderer/Private/RayTracing/RayTracingScene.h
@@ -10,6 +10,8 @@
 #include "Async/TaskGraphInterfaces.h"
 #include "RHI.h"
 #include "RenderGraphResources.h"
+#include "Misc/MemStack.h"
+#include "Containers/ArrayView.h"
 
 class FRHIRayTracingScene;
 class FRHIShaderResourceView;
@@ -40,6 +42,14 @@ public:
 	// Similar to Reset(), but also releases any persistent CPU and GPU memory allocations.
 	void ResetAndReleaseResources();
 
+	// Allocates temporary memory that will be valid until the next Reset().
+	// Can be used to store temporary instance transforms, user data, etc.
+	template <typename T>
+	TArrayView<T> Allocate(int32 Count)
+	{
+		return MakeArrayView(new(Allocator) T[Count], Count);
+	}
+
 	// Returns true if RHI ray tracing scene has been created or asynchronous creation has been kicked.
 	// i.e. returns true after BeginCreate() and before Reset().
 	RENDERER_API bool IsCreated() const;
@@ -83,6 +93,10 @@ private:
 
 	// View for the TLAS buffer that should be used in ray tracing shaders
 	FShaderResourceViewRHIRef RayTracingSceneSRV;
+
+	// Transient memory allocator
+	FMemStackBase Allocator;
+
 };
 
 #endif // RHI_RAYTRACING
diff --git a/Engine/Source/Runtime/Renderer/Private/Tests/RayTracingTestbed.cpp b/Engine/Source/Runtime/Renderer/Private/Tests/RayTracingTestbed.cpp
--- a/Engine/Source/Runtime/Renderer/Private/Tests/RayTracingTestbed.cpp
+++ b/Engine/Source/Runtime/Renderer/Private/Tests/RayTracingTestbed.cpp
@@ -14,7 +14,7 @@ IMPLEMENT_SIMPLE_AUTOMATION_TEST(FRayTracingTestbed, "System.Renderer.RayTracing
 #include "GlobalShader.h"
 #include "RayTracingDefinitions.h"
 
-
+// HINT: Execute this test via console command in editor: Automation RunTest System.Renderer.RayTracing.BasicRayTracing
 bool RunRayTracingTestbed_RenderThread(const FString& Parameters)
 {
 	check(IsInRenderingThread());
@@ -101,24 +101,36 @@ bool RunRayTracingTestbed_RenderThread(const FString& Parameters)
 	FRayTracingGeometrySegment Segment;
 	Segment.VertexBuffer = VertexBuffer;
 	Segment.NumPrimitives = 1;
+	Segment.MaxVertices = 3;
 	GeometryInitializer.Segments.Add(Segment);
 	GeometryInitializer.TotalPrimitiveCount = Segment.NumPrimitives;
 	FRayTracingGeometryRHIRef Geometry = RHICreateRayTracingGeometry(GeometryInitializer);
 
-	FShaderResourceViewRHIRef GPUTransforms = nullptr;
-	const uint32 NumTransforms = 1;
+	static constexpr uint32 NumTransforms = 1;
+	static constexpr uint32 NumInstances = 1;
 
-	FRayTracingGeometryInstance Instances[] = {
-		FRayTracingGeometryInstance { Geometry, {FMatrix::Identity}, {}, NumTransforms, GPUTransforms, {0}, 0xFF }
-	};
+	FRayTracingGeometryInstance Instances[NumInstances] = {};
+	Instances[0].GeometryRHI = Geometry;
+	Instances[0].NumTransforms = NumTransforms;
+	Instances[0].Transforms = MakeArrayView(&FMatrix::Identity, 1);
 
 	FRayTracingSceneInitializer Initializer;
 	Initializer.Instances = Instances;
 	Initializer.ShaderSlotsPerGeometrySegment = RAY_TRACING_NUM_SHADER_SLOTS;
-	FRayTracingSceneRHIRef Scene = RHICreateRayTracingScene(Initializer);
+
+	ERayTracingAccelerationStructureFlags SceneBuildFlags = ERayTracingAccelerationStructureFlags::FastTrace;
+	FRayTracingAccelerationStructureSize SceneSizeInfo = RHICalcRayTracingSceneSize(1, SceneBuildFlags);
+	FRHIResourceCreateInfo SceneBufferCreateInfo(TEXT("RayTracingTestBedSceneBuffer"));
+	FBufferRHIRef SceneBuffer = RHICreateBuffer(
+		uint32(SceneSizeInfo.ResultSize),
+		BUF_AccelerationStructure, 0 /*Stride*/, ERHIAccess::BVHWrite,
+		SceneBufferCreateInfo);
 
 	FRHICommandListImmediate& RHICmdList = FRHICommandListExecutor::GetImmediateCommandList();
 
+	FRayTracingSceneRHIRef Scene = RHICreateRayTracingScene(Initializer);
+	RHICmdList.BindAccelerationStructureMemory(Scene, SceneBuffer, 0);
+
 	RHICmdList.BuildAccelerationStructure(Geometry);
 
 	// #yuriy_todo: explicit transitions and state validation for BLAS
