Makes Gen5TAA's sequence length a prime number to remove corelation between jitter and other random sequence

diff --git a/Engine/Source/Runtime/Engine/Public/SceneManagement.h b/Engine/Source/Runtime/Engine/Public/SceneManagement.h
--- a/Engine/Source/Runtime/Engine/Public/SceneManagement.h
+++ b/Engine/Source/Runtime/Engine/Public/SceneManagement.h
@@ -213,8 +213,6 @@ public:
 	//
 	virtual uint32 GetCurrentTemporalAASampleIndex() const = 0;
 
-	virtual uint32 GetCurrentUnclampedTemporalAASampleIndex() const = 0;
-
 	/** Tells if the eye adaptation texture / buffer exists without attempting to allocate it. */
 	bool HasValidEyeAdaptationTexture() const { return bValidEyeAdaptationTexture; }
 	bool HasValidEyeAdaptationBuffer() const { return bValidEyeAdaptationBuffer; }
diff --git a/Engine/Source/Runtime/Renderer/Private/CompositionLighting/PostProcessAmbientOcclusion.cpp b/Engine/Source/Runtime/Renderer/Private/CompositionLighting/PostProcessAmbientOcclusion.cpp
--- a/Engine/Source/Runtime/Renderer/Private/CompositionLighting/PostProcessAmbientOcclusion.cpp
+++ b/Engine/Source/Runtime/Renderer/Private/CompositionLighting/PostProcessAmbientOcclusion.cpp
@@ -408,14 +408,11 @@ static FGTAOShaderParameters GetGTAOShaderParameters(const FViewInfo& View, FInt
 {
 	const FFinalPostProcessSettings& Settings = View.FinalPostProcessSettings;
 
-	uint32 TemporalFrame = 0;
-	uint32 Frame = 0;
-
 	const FSceneViewState* ViewState = static_cast<const FSceneViewState*>(View.State);
 
+	uint32 Frame = 0;
 	if (ViewState && (CVarGTAOPauseJitter.GetValueOnRenderThread() != 1))
 	{
-		TemporalFrame = ViewState->GetCurrentUnclampedTemporalAASampleIndex();
 		Frame = ViewState->GetFrameIndex();
 	}
 
@@ -424,13 +421,13 @@ static FGTAOShaderParameters GetGTAOShaderParameters(const FViewInfo& View, FInt
 	const float Rots[6] = { 60.0f, 300.0f, 180.0f, 240.0f, 120.0f, 0.0f };
 	const float Offsets[4] = { 0.1f, 0.6f, 0.35f, 0.85f };
 
-	float TemporalAngle = Rots[TemporalFrame % 6] * (PI / 360.0f);
+	float TemporalAngle = Rots[Frame % 6] * (PI / 360.0f);
 
 	// Angles of rotation that are set per frame
 	float SinAngle, CosAngle;
 	FMath::SinCos(&SinAngle, &CosAngle, TemporalAngle);
 
-	Result.GTAOParams[0] = FVector4(CosAngle, SinAngle, Offsets[(TemporalFrame / 6) % 4] * 0.25, Offsets[TemporalFrame % 4]);
+	Result.GTAOParams[0] = FVector4(CosAngle, SinAngle, Offsets[(Frame / 6) % 4] * 0.25, Offsets[Frame % 4]);
 
 	// Frame X = number , Y = Thickness param, 
 	float ThicknessBlend = CVarGTAOThicknessBlend.GetValueOnRenderThread();
diff --git a/Engine/Source/Runtime/Renderer/Private/RendererScene.cpp b/Engine/Source/Runtime/Renderer/Private/RendererScene.cpp
--- a/Engine/Source/Runtime/Renderer/Private/RendererScene.cpp
+++ b/Engine/Source/Runtime/Renderer/Private/RendererScene.cpp
@@ -202,7 +202,6 @@ FSceneViewState::FSceneViewState()
 	CachedVisibilityChunkIndex = INDEX_NONE;
 	MIDUsedCount = 0;
 	TemporalAASampleIndex = 0;
-	TemporalAASampleIndexUnclamped = 0;
 	FrameIndex = 0;
 	DistanceFieldTemporalSampleIndex = 0;
 	AOTileIntersectionResources = NULL;
diff --git a/Engine/Source/Runtime/Renderer/Private/ScenePrivate.h b/Engine/Source/Runtime/Renderer/Private/ScenePrivate.h
--- a/Engine/Source/Runtime/Renderer/Private/ScenePrivate.h
+++ b/Engine/Source/Runtime/Renderer/Private/ScenePrivate.h
@@ -917,10 +917,7 @@ private:
 public:
 	
 	// if TemporalAA is on this cycles through 0..TemporalAASampleCount-1, ResetViewState() puts it back to 0
-	int8 TemporalAASampleIndex;
-
-	// if TemporalAA is on this cycles through 0..Onwards, ResetViewState() puts it back to 0
-	uint32 TemporalAASampleIndexUnclamped;
+	int32 TemporalAASampleIndex;
 
 	// counts up by one each frame, warped in 0..7 range, ResetViewState() puts it back to 0
 	uint32 FrameIndex;
@@ -1099,10 +1096,6 @@ public:
 		return TemporalAASampleIndex;
 	}
 
-	virtual uint32 GetCurrentUnclampedTemporalAASampleIndex() const
-	{
-		return TemporalAASampleIndexUnclamped;
-	}
 	// Returns the index of the frame with a desired power of two modulus.
 	inline uint32 GetFrameIndex(uint32 Pow2Modulus) const
 	{
@@ -1120,7 +1113,6 @@ public:
 	virtual void ResetViewState()
 	{
 		TemporalAASampleIndex = 0;
-		TemporalAASampleIndexUnclamped = 0;
 		FrameIndex = 0;
 		DistanceFieldTemporalSampleIndex = 0;
 		PreExposure = 1.f;
diff --git a/Engine/Source/Runtime/Renderer/Private/SceneVisibility.cpp b/Engine/Source/Runtime/Renderer/Private/SceneVisibility.cpp
--- a/Engine/Source/Runtime/Renderer/Private/SceneVisibility.cpp
+++ b/Engine/Source/Runtime/Renderer/Private/SceneVisibility.cpp
@@ -44,6 +44,7 @@
 #include "ProfilingDebugging/DiagnosticTable.h"
 #include "Algo/Unique.h"
 #include "InstanceCulling/InstanceCullingManager.h"
+#include "TemporalAA.h"
 
 /*------------------------------------------------------------------------------
 	Globals
@@ -3276,6 +3277,8 @@ void FSceneRenderer::PreVisibilityFrameSetup(FRDGBuilder& GraphBuilder, const FS
 		// Subpixel jitter for temporal AA
 		int32 CVarTemporalAASamplesValue = CVarTemporalAASamples.GetValueOnRenderThread();
 
+		EMainTAAPassConfig TAAConfig = ITemporalUpscaler::GetMainTAAPassConfig(View);
+
 		bool bTemporalUpsampling = View.PrimaryScreenPercentageMethod == EPrimaryScreenPercentageMethod::TemporalUpscale;
 		
 		// Apply a sub pixel offset to the view.
@@ -3284,8 +3287,19 @@ void FSceneRenderer::PreVisibilityFrameSetup(FRDGBuilder& GraphBuilder, const FS
 			float EffectivePrimaryResolutionFraction = float(View.ViewRect.Width()) / float(View.GetSecondaryViewRectSize().X);
 
 			// Compute number of TAA samples.
-			int32 TemporalAASamples = CVarTemporalAASamplesValue;
+			int32 TemporalAASamples;
 			{
+				if (TAAConfig == EMainTAAPassConfig::Gen5)
+				{
+					// Force the number of AA sample to make sure the quality doesn't get
+					// compromised by previously set settings for Gen4 TAA
+					TemporalAASamples = 8;
+				}
+				else
+				{
+					TemporalAASamples = FMath::Clamp(CVarTemporalAASamplesValue, 1, 255);
+				}
+
 				if (bTemporalUpsampling)
 				{
 					// When doing TAA upsample with screen percentage < 100%, we need extra temporal samples to have a
@@ -3297,11 +3311,27 @@ void FSceneRenderer::PreVisibilityFrameSetup(FRDGBuilder& GraphBuilder, const FS
 					TemporalAASamples = 4;
 				}
 
-				TemporalAASamples = FMath::Clamp(TemporalAASamples, 1, 255);
+				// Use immediately higher prime number to break up coherence between the TAA jitter sequence and any
+				// other random signal that are power of two of View.StateFrameIndex
+				if (TAAConfig == EMainTAAPassConfig::Gen5)
+				{
+					static const int8 kFirstPrimeNumbers[25] = {
+						2, 3, 5, 7, 11, 13, 17, 19, 23, 29, 31, 37, 41, 43, 47, 53, 59, 61, 67, 71, 73, 79, 83, 89, 97,
+					};
+
+					for (int32 PrimeNumberId = 4; PrimeNumberId < UE_ARRAY_COUNT(kFirstPrimeNumbers); PrimeNumberId++)
+					{
+						if (int32(kFirstPrimeNumbers[PrimeNumberId]) >= TemporalAASamples)
+						{
+							TemporalAASamples = int32(kFirstPrimeNumbers[PrimeNumberId]);
+							break;
+						}
+					}
+				}
 			}
 
 			// Compute the new sample index in the temporal sequence.
-			int32 TemporalSampleIndex			= ViewState->TemporalAASampleIndex + 1;
+			int32 TemporalSampleIndex = ViewState->TemporalAASampleIndex + 1;
 			if (TemporalSampleIndex >= TemporalAASamples || View.bCameraCut)
 			{
 				TemporalSampleIndex = 0;
@@ -3317,8 +3347,7 @@ void FSceneRenderer::PreVisibilityFrameSetup(FRDGBuilder& GraphBuilder, const FS
 			// Updates view state.
 			if (!View.bStatePrevViewInfoIsReadOnly && !bFreezeTemporalSequences)
 			{
-				ViewState->TemporalAASampleIndex		  = TemporalSampleIndex;
-				ViewState->TemporalAASampleIndexUnclamped = ViewState->TemporalAASampleIndexUnclamped+1;
+				ViewState->TemporalAASampleIndex = TemporalSampleIndex;
 			}
 
 			// Choose sub pixel sample coordinate in the temporal sequence.
