Fix for Path Tracer running without the ray tracing pipeline having been setup

This fixes a crash due to a null ray tracing pipeline when path tracing is enabled at startup.

diff --git a/Engine/Source/Runtime/Renderer/Private/DeferredShadingRenderer.cpp b/Engine/Source/Runtime/Renderer/Private/DeferredShadingRenderer.cpp
--- a/Engine/Source/Runtime/Renderer/Private/DeferredShadingRenderer.cpp
+++ b/Engine/Source/Runtime/Renderer/Private/DeferredShadingRenderer.cpp
@@ -2717,7 +2717,8 @@ void FDeferredShadingSceneRenderer::Render(FRDGBuilder& GraphBuilder)
 
 #if RHI_RAYTRACING
 	// If Lumen is not using HWRT shadows, we can wait until here: before Lumen diffuse indirect
-	if (!Lumen::UseHardwareRayTracedShadows(Views[0]))
+	// Also catch the case of path tracer or RT debug output
+	if (!Lumen::UseHardwareRayTracedShadows(Views[0]) || !bCanOverlayRayTracingOutput)
 	{
 		WaitForRayTracingScene(GraphBuilder);
 	}
