Optimized distance field streaming requests by skipping a bunch of memory reads and atomics
0.69ms -> 0.12ms of GPU time on 2080, flying with drone

diff --git a/Engine/Shaders/Private/DistanceFieldStreaming.usf b/Engine/Shaders/Private/DistanceFieldStreaming.usf
--- a/Engine/Shaders/Private/DistanceFieldStreaming.usf
+++ b/Engine/Shaders/Private/DistanceFieldStreaming.usf
@@ -75,25 +75,29 @@ void ComputeDistanceFieldAssetWantedMipsCS(
 	{
 		FDFObjectData DFObjectData = LoadDFObjectData(ObjectIndex);
 
+		uint WantedNumMips = 1;
 		if (DebugForceNumMips == 0)
 		{
 			FDFObjectBounds DFObjectBounds = LoadDFObjectBounds(ObjectIndex);
-
 			float OuterDistanceSq = ComputeSquaredDistanceBetweenAABBs(Mip1WorldCenter, Mip1WorldExtent, DFObjectBounds.Center, DFObjectBounds.BoxExtent);
-
-			int WantedNumMips = 1;
-				
+ 
 			if (OuterDistanceSq <= 0)
 			{
 				float InnerDistanceSq = ComputeSquaredDistanceBetweenAABBs(Mip2WorldCenter, Mip2WorldExtent, DFObjectBounds.Center, DFObjectBounds.BoxExtent);
 				WantedNumMips = InnerDistanceSq <= 0 ? 3 : 2;
 			}
-			
-			InterlockedMax(RWDistanceFieldAssetWantedNumMips[DFObjectData.AssetIndex], WantedNumMips);
 		}
 		else
 		{
-			InterlockedMax(RWDistanceFieldAssetWantedNumMips[DFObjectData.AssetIndex], DebugForceNumMips);
+			WantedNumMips = DebugForceNumMips;
+		}
+
+		// RWDistanceFieldAssetWantedNumMips is already cleared to 1 so we can skip some memory loads and atomics
+		BRANCH
+		if (WantedNumMips > 1)
+		{
+			FDFObjectData DFObjectData = LoadDFObjectData(ObjectIndex);
+			InterlockedMax(RWDistanceFieldAssetWantedNumMips[DFObjectData.AssetIndex], WantedNumMips);
 		}
 	}
 }
@@ -115,7 +119,7 @@ void GenerateDistanceFieldAssetStreamingRequestsCS(
 		uint NumMips = LoadDFAssetData(AssetIndex, 0).NumMips;
 		uint WantedNumMips = DistanceFieldAssetWantedNumMips[AssetIndex];
 
-		if (WantedNumMips != 0 && NumMips != WantedNumMips)
+		if (WantedNumMips != 0 && NumMips != 0 && NumMips != WantedNumMips)
 		{
 			uint StreamingRequestIndex;
 			InterlockedAdd(RWDistanceFieldAssetStreamingRequests[0], 1, StreamingRequestIndex);
diff --git a/Engine/Source/Runtime/Renderer/Private/DistanceFieldStreaming.cpp b/Engine/Source/Runtime/Renderer/Private/DistanceFieldStreaming.cpp
--- a/Engine/Source/Runtime/Renderer/Private/DistanceFieldStreaming.cpp
+++ b/Engine/Source/Runtime/Renderer/Private/DistanceFieldStreaming.cpp
@@ -725,7 +725,8 @@ void FDistanceFieldSceneData::GenerateStreamingRequests(
 		FRDGBufferDesc WantedNumMipsDesc = FRDGBufferDesc::CreateStructuredDesc(sizeof(uint32), FMath::RoundUpToPowerOfTwo(NumAssets));
 		FRDGBufferRef WantedNumMips = GraphBuilder.CreateBuffer(WantedNumMipsDesc, TEXT("DistanceFields.DistanceFieldAssetWantedNumMips"));
 
-		AddClearUAVPass(GraphBuilder, GraphBuilder.CreateUAV(FRDGBufferUAVDesc(WantedNumMips)), 0);
+		// Every asset wants at least 1 mipmap
+		AddClearUAVPass(GraphBuilder, GraphBuilder.CreateUAV(FRDGBufferUAVDesc(WantedNumMips)), 1);
 
 		FRDGBufferDesc StreamingRequestsDesc = FRDGBufferDesc::CreateStructuredDesc(sizeof(uint32), MaxStreamingRequests * 2 + 1);
 		StreamingRequestsDesc.Usage = EBufferUsageFlags(StreamingRequestsDesc.Usage | BUF_SourceCopy);
@@ -801,6 +802,8 @@ void FDistanceFieldSceneData::UpdateDistanceFieldAtlas(
 	TRACE_CPUPROFILER_EVENT_SCOPE(FDistanceFieldSceneData::UpdateDistanceFieldAtlas);
 	RDG_EVENT_SCOPE(GraphBuilder, "UpdateDistanceFieldAtlas");
 
+	TArray<FDistanceFieldAssetMipId> AssetDataUploads;
+
 	for (FSetElementId AssetSetId : DistanceFieldAssetRemoves)
 	{
 		const FDistanceFieldAssetState& AssetState = AssetStateArray[AssetSetId];
@@ -817,11 +820,13 @@ void FDistanceFieldSceneData::UpdateDistanceFieldAtlas(
 			}
 		}
 		
+		// Clear GPU data for removed asset
+		AssetDataUploads.Add(FDistanceFieldAssetMipId(AssetSetId, 0));
+
 		AssetStateArray.Remove(AssetSetId);
 	}
 
 	TArray<FDistanceFieldReadRequest> NewReadRequests;
-	TArray<FDistanceFieldAssetMipId> AssetDataUploads;
 	// Lock the most recent streaming request buffer from the GPU, create new read requests for mips we want to load in the Async Task
 	ProcessStreamingRequestsFromGPU(NewReadRequests, AssetDataUploads);
 
@@ -970,40 +975,50 @@ void FDistanceFieldSceneData::UpdateDistanceFieldAtlas(
 
 		for (FDistanceFieldAssetMipId AssetMipUpload : AssetDataUploads)
 		{
-			const FDistanceFieldAssetState& AssetState = AssetStateArray[AssetMipUpload.AssetId];
 			const int32 ReversedMipIndex = AssetMipUpload.ReversedMipIndex;
-			const FDistanceFieldAssetMipState& MipState = AssetState.ReversedMips[ReversedMipIndex];
-			const int32 MipIndex = AssetState.BuiltData->Mips.Num() - ReversedMipIndex - 1;
-			const FSparseDistanceFieldMip& MipBuiltData = AssetState.BuiltData->Mips[MipIndex];
-			const FVector2D DistanceFieldToVolumeScaleBias = MipBuiltData.DistanceFieldToVolumeScaleBias;
-			const int32 NumMips = AssetState.ReversedMips.Num();
-
-			check(NumMips <= DistanceField::NumMips);
-			check(DistanceField::NumMips < 4);
-			check(MipBuiltData.IndirectionDimensions.X < DistanceField::MaxIndirectionDimension 
-				&& MipBuiltData.IndirectionDimensions.Y < DistanceField::MaxIndirectionDimension 
-				&& MipBuiltData.IndirectionDimensions.Z < DistanceField::MaxIndirectionDimension);
-
 			FVector4* UploadAssetData = (FVector4*)AssetDataUploadBuffer.Add_GetRef(AssetMipUpload.AssetId.AsInteger() * DistanceField::NumMips + ReversedMipIndex);
 
-			uint32 IntVector0[4] = 
+			if (AssetStateArray.IsValidId(AssetMipUpload.AssetId))
+			{
+				const FDistanceFieldAssetState& AssetState = AssetStateArray[AssetMipUpload.AssetId];
+				const FDistanceFieldAssetMipState& MipState = AssetState.ReversedMips[ReversedMipIndex];
+				const int32 MipIndex = AssetState.BuiltData->Mips.Num() - ReversedMipIndex - 1;
+				const FSparseDistanceFieldMip& MipBuiltData = AssetState.BuiltData->Mips[MipIndex];
+				const FVector2D DistanceFieldToVolumeScaleBias = MipBuiltData.DistanceFieldToVolumeScaleBias;
+				const int32 NumMips = AssetState.ReversedMips.Num();
+
+				check(NumMips <= DistanceField::NumMips);
+				check(DistanceField::NumMips < 4);
+				check(MipBuiltData.IndirectionDimensions.X < DistanceField::MaxIndirectionDimension
+					&& MipBuiltData.IndirectionDimensions.Y < DistanceField::MaxIndirectionDimension
+					&& MipBuiltData.IndirectionDimensions.Z < DistanceField::MaxIndirectionDimension);
+
+				uint32 IntVector0[4] =
+				{
+					(uint32)MipBuiltData.IndirectionDimensions.X | (uint32)(MipBuiltData.IndirectionDimensions.Y << 10) | (uint32)(MipBuiltData.IndirectionDimensions.Z << 20) | (uint32)(NumMips << 30),
+					(uint32)FFloat16(DistanceFieldToVolumeScaleBias.X).Encoded | ((uint32)FFloat16(DistanceFieldToVolumeScaleBias.Y).Encoded << 16),
+					(uint32)MipState.IndirectionTableOffset,
+					0
+				};
+
+				// Bypass NaN checks in FVector4 ctors
+				FVector4 FloatVector0;
+				FloatVector0.X = *(const float*)&IntVector0[0];
+				FloatVector0.Y = *(const float*)&IntVector0[1];
+				FloatVector0.Z = *(const float*)&IntVector0[2];
+				FloatVector0.W = *(const float*)&IntVector0[3];
+
+				UploadAssetData[0] = FloatVector0;
+				UploadAssetData[1] = FVector4(MipBuiltData.VolumeToVirtualUVScale, 0.0f);
+				UploadAssetData[2] = FVector4(MipBuiltData.VolumeToVirtualUVAdd, 0.0f);
+			}
+			else
 			{
-				(uint32)MipBuiltData.IndirectionDimensions.X | (uint32)(MipBuiltData.IndirectionDimensions.Y << 10) | (uint32)(MipBuiltData.IndirectionDimensions.Z << 20) | (uint32)(NumMips << 30),
-				(uint32)FFloat16(DistanceFieldToVolumeScaleBias.X).Encoded | ((uint32)FFloat16(DistanceFieldToVolumeScaleBias.Y).Encoded << 16),
-				(uint32)MipState.IndirectionTableOffset,
-				0
-			};
-
-			// Bypass NaN checks in FVector4 ctors
-			FVector4 FloatVector0;
-			FloatVector0.X = *(const float*)&IntVector0[0];
-			FloatVector0.Y = *(const float*)&IntVector0[1];
-			FloatVector0.Z = *(const float*)&IntVector0[2];
-			FloatVector0.W = *(const float*)&IntVector0[3];
-
-			UploadAssetData[0] = FloatVector0;
-			UploadAssetData[1] = FVector4(MipBuiltData.VolumeToVirtualUVScale, 0.0f);
-			UploadAssetData[2] = FVector4(MipBuiltData.VolumeToVirtualUVAdd, 0.0f);
+				// Clear invalid entries to zero
+				UploadAssetData[0] = FVector4(0.0f, 0.0f, 0.0f, 0.0f);
+				UploadAssetData[1] = FVector4(0.0f, 0.0f, 0.0f, 0.0f);
+				UploadAssetData[2] = FVector4(0.0f, 0.0f, 0.0f, 0.0f);
+			}
 		}
 
 		AddPass(GraphBuilder, [this](FRHICommandListImmediate& RHICmdList)
