Scale normal bias by distance to camera and projection scale to fix artifacts on flat planes and default sphere with shallow light angles.
Rename parameter to "r.Shadow.Virtual.NormalBias" to better represent its unit-less value and give a reasonable default value.

diff --git a/Engine/Shaders/Private/VirtualShadowMaps/VirtualShadowMapProjection.usf b/Engine/Shaders/Private/VirtualShadowMaps/VirtualShadowMapProjection.usf
--- a/Engine/Shaders/Private/VirtualShadowMaps/VirtualShadowMapProjection.usf
+++ b/Engine/Shaders/Private/VirtualShadowMaps/VirtualShadowMapProjection.usf
@@ -179,7 +179,7 @@ float4 ComputeRandom4(uint2 PixelPosition)
 
 int4 ProjectionRect;
 float ContactShadowLength;
-float NormalOffsetWorld;			// Normal offset in world space
+float NormalBias;
 int SMRTRayCount;					// 0 = off
 int SMRTSamplesPerRay;
 float SMRTRayLengthScale;			// Directional lights
@@ -218,6 +218,9 @@ FProjectLightResult ProjectLight(
 	const half Noise)
 {
 	const float3 WorldPosition = TranslatedWorldPosition - View.PreViewTranslation;
+	// NOTE: If View.PreViewTranslation is guaranteed to be the camera position we can simplify
+	const float DistanceToCamera = length(TranslatedWorldPosition - View.TranslatedWorldCameraOrigin);
+	const float NormalBiasLength = max(0.01f, NormalBias * DistanceToCamera / View.ViewToClip._11);
 
 #if DIRECTIONAL_LIGHT
 	uint LightType = LIGHT_TYPE_DIRECTIONAL;
@@ -247,8 +250,8 @@ FProjectLightResult ProjectLight(
 	{
 		const bool bIsHair = InputType == INPUT_TYPE_HAIRSTRANDS;
 		const half3 EstimatedGeoWorldNormal = bIsHair ? L : GBufferNormal;
-
-		TranslatedWorldPosition += EstimatedGeoWorldNormal * lerp(NormalOffsetWorld, 0.01, saturate(dot(EstimatedGeoWorldNormal, L)));
+		
+		TranslatedWorldPosition += EstimatedGeoWorldNormal * NormalBiasLength;
 
 		if (SMRTRayCount > 0)
 		{
@@ -299,7 +302,7 @@ FProjectLightResult ProjectLight(
 		// depth buffer contains fully opaque depth, which create false positive intersection resulting in wrong self shadowing)
 		if (ContactShadowLengthWorld > 0.0f && !bIsHair)
 		{
-			half ContactShadowHitDistance = VirtualShadowMapScreenRayCast(TranslatedWorldPosition, L, ContactShadowLengthWorld, CONTACT_SHADOW_SAMPLES, Noise);		
+			half ContactShadowHitDistance = VirtualShadowMapScreenRayCast(TranslatedWorldPosition, L, ContactShadowLengthWorld, CONTACT_SHADOW_SAMPLES, Noise);
 			if (ContactShadowHitDistance >= 0.0f)
 			{
 				ShadowFactor = 0.0f;
diff --git a/Engine/Source/Runtime/Renderer/Private/VirtualShadowMaps/VirtualShadowMapProjection.cpp b/Engine/Source/Runtime/Renderer/Private/VirtualShadowMaps/VirtualShadowMapProjection.cpp
--- a/Engine/Source/Runtime/Renderer/Private/VirtualShadowMaps/VirtualShadowMapProjection.cpp
+++ b/Engine/Source/Runtime/Renderer/Private/VirtualShadowMaps/VirtualShadowMapProjection.cpp
@@ -32,10 +32,10 @@ static TAutoConsoleVariable<float> CVarContactShadowLength(
 	ECVF_RenderThreadSafe
 );
 
-static TAutoConsoleVariable<float> CVarNormalOffsetWorld(
-	TEXT( "r.Shadow.Virtual.NormalOffsetWorld" ),
-	0.1f,
-	TEXT( "World space offset along surface normal for shadow lookup." )
+static TAutoConsoleVariable<float> CVarNormalBias(
+	TEXT( "r.Shadow.Virtual.NormalBias" ),
+	0.5f,
+	TEXT( "Receiver offset along surface normal for shadow lookup. Scaled by distance to camera." )
 	TEXT( "Higher values avoid artifacts on surfaces nearly parallel to the light, but also visibility offset shadows and increase the chance of hitting unmapped pages." ),
 	ECVF_RenderThreadSafe
 );
@@ -188,7 +188,7 @@ class FVirtualShadowMapProjectionCS : public FGlobalShader
 		SHADER_PARAMETER_STRUCT_REF(FViewUniformShaderParameters, View)
 		SHADER_PARAMETER(FIntVector4, ProjectionRect)
 		SHADER_PARAMETER(float, ContactShadowLength)
-		SHADER_PARAMETER(float, NormalOffsetWorld)
+		SHADER_PARAMETER(float, NormalBias)
 		SHADER_PARAMETER(uint32, SMRTRayCount)
 		SHADER_PARAMETER(uint32, SMRTSamplesPerRay)
 		SHADER_PARAMETER(float, SMRTRayLengthScale)
@@ -239,6 +239,11 @@ class FVirtualShadowMapProjectionCS : public FGlobalShader
 };
 IMPLEMENT_GLOBAL_SHADER(FVirtualShadowMapProjectionCS, "/Engine/Private/VirtualShadowMaps/VirtualShadowMapProjection.usf", "VirtualShadowMapProjection", SF_Compute);
 
+static float GetNormalBiasForShader()
+{
+	return CVarNormalBias.GetValueOnRenderThread() / 1000.0f;
+}
+
 extern int32 GVirtualShadowMapAtomicWrites;
 
 static void RenderVirtualShadowMapProjectionCommon(
@@ -262,7 +267,7 @@ static void RenderVirtualShadowMapProjectionCommon(
 	PassParameters->ProjectionRect = FIntVector4(ProjectionRect.Min.X, ProjectionRect.Min.Y, ProjectionRect.Max.X, ProjectionRect.Max.Y);
 	PassParameters->DebugOutputType = CVarVirtualShadowMapDebugProjection.GetValueOnRenderThread();
 	PassParameters->ContactShadowLength = CVarContactShadowLength.GetValueOnRenderThread();
-	PassParameters->NormalOffsetWorld = CVarNormalOffsetWorld.GetValueOnRenderThread();
+	PassParameters->NormalBias = GetNormalBiasForShader();
 	PassParameters->InputType = uint32(InputType);
 	PassParameters->HairStrands = HairStrands::BindHairStrandsViewUniformParameters(View);
 	if (bHasHairStrandsData)
