View uniform buffers referenced by FRayTracingDynamicGeometryCollection are now kept alive until EndUpdate() is called. Fixes the View UB lifetime due to UpdateGlobalDistanceFieldVolume() creating a new one while batched dynamic geometry dispatch commands are still referencing the old.

diff --git a/Engine/Source/Runtime/Renderer/Private/RayTracing/RayTracingDynamicGeometry.cpp b/Engine/Source/Runtime/Renderer/Private/RayTracing/RayTracingDynamicGeometry.cpp
--- a/Engine/Source/Runtime/Renderer/Private/RayTracing/RayTracingDynamicGeometry.cpp
+++ b/Engine/Source/Runtime/Renderer/Private/RayTracing/RayTracingDynamicGeometry.cpp
@@ -136,6 +136,12 @@ void FRayTracingDynamicGeometryCollection::AddDynamicMeshBatchForGeometryUpdate(
 	uint32 VertexBufferOffset = 0;
 	bool bUseSharedVertexBuffer = false;
 
+	if (ReferencedUniformBuffers.Num() == 0 || ReferencedUniformBuffers.Last() != View->ViewUniformBuffer)
+	{
+		// Keep ViewUniformBuffer alive until EndUpdate()
+		ReferencedUniformBuffers.Add(View->ViewUniformBuffer);
+	}
+
 	// If update params didn't provide a buffer then use a shared vertex position buffer
 	if (RWBuffer == nullptr)
 	{
@@ -503,6 +509,8 @@ void FRayTracingDynamicGeometryCollection::DispatchUpdates(FRHIComputeCommandLis
 
 void FRayTracingDynamicGeometryCollection::EndUpdate(FRHICommandListImmediate& RHICmdList)
 {
+	ReferencedUniformBuffers.Empty(ReferencedUniformBuffers.Max());
+
 	// Move ownership to RHI thread for another frame
 	RHICmdList.EnqueueLambda([ArrayOwnedByRHIThread = MoveTemp(Segments)](FRHICommandListImmediate&){});
 }
diff --git a/Engine/Source/Runtime/Renderer/Public/RayTracingDynamicGeometryCollection.h b/Engine/Source/Runtime/Renderer/Public/RayTracingDynamicGeometryCollection.h
--- a/Engine/Source/Runtime/Renderer/Public/RayTracingDynamicGeometryCollection.h
+++ b/Engine/Source/Runtime/Renderer/Public/RayTracingDynamicGeometryCollection.h
@@ -35,6 +35,9 @@ private:
 	};
 	TArray<FVertexPositionBuffer*> VertexPositionBuffers;
 
+	// Any uniform buffers that must be kept alive until EndUpdate (after DispatchUpdates is called)
+	TArray<FUniformBufferRHIRef> ReferencedUniformBuffers;
+
 	// Generation ID when the shared vertex buffers have been reset. The current generation ID is stored in the FRayTracingGeometry to keep track
 	// if the vertex buffer data is still valid for that frame - validated before generation the TLAS
 	int64 SharedBufferGenerationID = 0;
