Fixed -d3ddebug warning when shader uses default (dummy) distance field parameters

diff --git a/Engine/Source/Runtime/Renderer/Private/GlobalDistanceField.cpp b/Engine/Source/Runtime/Renderer/Private/GlobalDistanceField.cpp
--- a/Engine/Source/Runtime/Renderer/Private/GlobalDistanceField.cpp
+++ b/Engine/Source/Runtime/Renderer/Private/GlobalDistanceField.cpp
@@ -1036,7 +1036,7 @@ void FViewInfo::SetupDefaultGlobalDistanceFieldUniformBufferParameters(FViewUnif
 	ViewUniformShaderParameters.NumGlobalSDFClipmaps = 0;
 
 	ViewUniformShaderParameters.GlobalDistanceFieldPageAtlasTexture = OrBlack3DIfNull(GBlackVolumeTexture->TextureRHI.GetReference());
-	ViewUniformShaderParameters.GlobalDistanceFieldPageTableTexture = OrBlack3DIfNull(GBlackVolumeTexture->TextureRHI.GetReference());
+	ViewUniformShaderParameters.GlobalDistanceFieldPageTableTexture = OrBlack3DIfNull(GBlackUintVolumeTexture->TextureRHI.GetReference());
 	ViewUniformShaderParameters.GlobalDistanceFieldMipTexture = OrBlack3DIfNull(GBlackVolumeTexture->TextureRHI.GetReference());
 }
 
