Flip broken velocity test in depth pass.
When in DDM_AllOpaqueNoVelocity mode we _don't_ want to render primitives in the depth pass that contribute to velocity.

diff --git a/Engine/Source/Runtime/Renderer/Private/DepthRendering.cpp b/Engine/Source/Runtime/Renderer/Private/DepthRendering.cpp
--- a/Engine/Source/Runtime/Renderer/Private/DepthRendering.cpp
+++ b/Engine/Source/Runtime/Renderer/Private/DepthRendering.cpp
@@ -848,12 +848,12 @@ void FDepthPassMeshProcessor::AddMeshBatch(const FMeshBatch& RESTRICT MeshBatch,
 		// This logic is copy/paste/modified from FOpaqueVelocityMeshProcessor::AddMeshBatch(), but ideally we should clean it up into
 		// a single function that is shared to avoid breakages from code changes.
 		EShaderPlatform ShaderPlatform = ViewIfDynamicMeshCommand->GetShaderPlatform();
-		if (!FOpaqueVelocityMeshProcessor::PrimitiveCanHaveVelocity(ShaderPlatform, PrimitiveSceneProxy))
+		if (FOpaqueVelocityMeshProcessor::PrimitiveCanHaveVelocity(ShaderPlatform, PrimitiveSceneProxy))
 		{
 			bDraw = false;
 		}
 
-		if (!FOpaqueVelocityMeshProcessor::PrimitiveHasVelocityForFrame(PrimitiveSceneProxy))
+		if (FOpaqueVelocityMeshProcessor::PrimitiveHasVelocityForFrame(PrimitiveSceneProxy))
 		{
 			bDraw = false;
 		}
@@ -861,7 +861,7 @@ void FDepthPassMeshProcessor::AddMeshBatch(const FMeshBatch& RESTRICT MeshBatch,
 		checkSlow(ViewIfDynamicMeshCommand->bIsViewInfo);
 		FViewInfo* ViewInfo = (FViewInfo*)ViewIfDynamicMeshCommand;
 
-		if (!FOpaqueVelocityMeshProcessor::PrimitiveHasVelocityForView(*ViewInfo, PrimitiveSceneProxy))
+		if (FOpaqueVelocityMeshProcessor::PrimitiveHasVelocityForView(*ViewInfo, PrimitiveSceneProxy))
 		{
 			bDraw = false;
 		}
