Unify RT Skylight with Path Tracer implementation

Use the path tracer's skydome implementation for RT skylight, replacing the previous cubemap based implementation. The caching of the skydome texure and pdf is now done in a central location, which allows RT skylight, RT GI and the path tracer to share the same map. For static cases (captured sky or explicit cubemap) this work is only done once and re-used. For realtime capture cases (sky/atmosphere system), this work can happen each time the sky is updated (previously this case did not support RT skylight at all).

diff --git a/Engine/Plugins/Experimental/GPULightmass/Source/GPULightmass/Private/LightmapRenderer.cpp b/Engine/Plugins/Experimental/GPULightmass/Source/GPULightmass/Private/LightmapRenderer.cpp
--- a/Engine/Plugins/Experimental/GPULightmass/Source/GPULightmass/Private/LightmapRenderer.cpp
+++ b/Engine/Plugins/Experimental/GPULightmass/Source/GPULightmass/Private/LightmapRenderer.cpp
@@ -13,7 +13,6 @@
 #include "RayTracing/RayTracingMaterialHitShaders.h"
 #include "Async/ParallelFor.h"
 #include "Async/Async.h"
-#include "Rendering/SkyLightImportanceSampling.h"
 #include "LightmapPreviewVirtualTexture.h"
 #include "RHIGPUReadback.h"
 #include "LightmapStorage.h"
@@ -2079,9 +2078,6 @@ void FLightmapRenderer::Finalize(FRDGBuilder& GraphBuilder)
 								RayGuidingCDFY = GraphBuilder.RegisterExternalTexture(LightmapTilePoolGPU.PooledRenderTargets[7], TEXT("RayGuidingCDFY"));
 							}
 
-							// These two buffers must have lifetime extended beyond GraphBuilder.Execute()
-							TUniformBufferRef<FSkyLightData> SkyLightDataUniformBuffer;
-
 							FIntPoint RayTracingResolution;
 							RayTracingResolution.X = GPreviewLightmapPhysicalTileSize * GPUBatchedTileRequests.BatchedTilesDesc.Num();
 							RayTracingResolution.Y = GPreviewLightmapPhysicalTileSize;
diff --git a/Engine/Plugins/Experimental/GPULightmass/Source/GPULightmass/Private/VolumetricLightmap.cpp b/Engine/Plugins/Experimental/GPULightmass/Source/GPULightmass/Private/VolumetricLightmap.cpp
--- a/Engine/Plugins/Experimental/GPULightmass/Source/GPULightmass/Private/VolumetricLightmap.cpp
+++ b/Engine/Plugins/Experimental/GPULightmass/Source/GPULightmass/Private/VolumetricLightmap.cpp
@@ -76,7 +76,6 @@ void ReleaseBrickData(FVolumetricLightmapBrickData& BrickData)
 }
 
 void SetupPathTracingLightParameters(const GPULightmass::FLightSceneRenderState& LightScene, FRDGBuilder& GraphBuilder, FRDGBufferSRV** OutLightBuffer, uint32* OutLightCount);
-FSkyLightData SetupSkyLightParameters(const GPULightmass::FLightSceneRenderState& LightScene);
 
 namespace GPULightmass
 {
@@ -527,9 +526,6 @@ void FVolumetricLightmapRenderer::BackgroundTick()
 #if RHI_RAYTRACING
 			if (IsRayTracingEnabled())
 			{
-				// These two buffers must have lifetime extended beyond RHICmdList.RayTraceDispatch()
-				TUniformBufferRef<FSkyLightData> SkyLightDataUniformBuffer;
-
 				FGlobalShaderMap* GlobalShaderMap = GetGlobalShaderMap(GMaxRHIFeatureLevel);
 
 				{
diff --git a/Engine/Shaders/Private/PathTracing/Light/PathTracingSkyLight.ush b/Engine/Shaders/Private/PathTracing/Light/PathTracingSkyLight.ush
--- a/Engine/Shaders/Private/PathTracing/Light/PathTracingSkyLight.ush
+++ b/Engine/Shaders/Private/PathTracing/Light/PathTracingSkyLight.ush
@@ -6,13 +6,14 @@ SkyLight.usf: Light sampling functions for SkyLight implementation
 
 #pragma once
 
-Texture2D SkylightTexture;
-Texture2D<float> SkylightPdf;
-SamplerState SkylightTextureSampler;
-float SkylightInvResolution;
-int SkylightMipCount;
+#include "../../RayTracing/RayTracingSkyLightCommon.ush"
 
-#define USE_HIERARCHICAL_IMPORTANCE_SAMPLING		1
+// By default, assume MIS is performed in the integrator.
+// However in some cases like RTGI we may want to have _some_ basic MIS
+// without introducing complexity in the integrator.
+#ifndef PATHTRACING_SKY_MIS
+#define PATHTRACING_SKY_MIS 0
+#endif
 
 float SkyLight_EstimateLight(
 	int LightId,
@@ -21,29 +22,19 @@ float SkyLight_EstimateLight(
 	bool IsTransmissiveMaterial
 )
 {
-	// cancels out constant factor in pdf below
-#if USE_HIERARCHICAL_IMPORTANCE_SAMPLING == 0
-	return 2 * PI;
+#if PATHTRACING_SKY_MIS == 0
+	return SkyLight_Estimate();
 #else
-	return 4 * PI * SkylightPdf.Load(int3(0, 0, SkylightMipCount - 1));
+	return 2 * PI;
 #endif
 }
 
-FLightHit SkyLight_TraceLight(RayDesc Ray, int LightId) {
+FLightHit SkyLight_TraceLight(RayDesc Ray, int LightId)
+{
 	if (Ray.TMax == POSITIVE_INFINITY)
 	{
-		// NOTE: assumes direction is normalized
-
-		float3 Dir = Ray.Direction;
-		float2 UV = InverseEquiAreaSphericalMapping(Dir.yzx);
-		float4 Result = SkylightTexture.SampleLevel(SkylightTextureSampler, UV, 0);
-		float3 Radiance = Result.xyz;
-#if USE_HIERARCHICAL_IMPORTANCE_SAMPLING == 0
-		float Pdf = 1.0 / (2.0 * PI);
-#else
-		float Pdf = Result.w > 0 ? Result.w / (4 * PI * SkylightPdf.Load(int3(0, 0, SkylightMipCount - 1))) : 0.0;
-#endif
-		return CreateLightHit(Radiance, Pdf, POSITIVE_INFINITY);
+		float4 Result = SkyLight_EvalLight(Ray.Direction);
+		return CreateLightHit(Result.xyz, Result.w, POSITIVE_INFINITY);
 	}
 	return NullLightHit();
 }
@@ -55,72 +46,46 @@ FLightSample SkyLight_SampleLight(
 	float3 WorldNormal
 )
 {
-	// Make temp structures so we can re-use the calls above
-	// #dxr_todo: write optimized form that combines radiance and pdf lookup in a single step
-	RayDesc Ray;
-	Ray.Origin = WorldPos;
-
-#if USE_HIERARCHICAL_IMPORTANCE_SAMPLING == 0
-	float4 Direction_Tangent = UniformSampleHemisphere(RandSample);
-	float3 OutDirection = TangentToWorld(Direction_Tangent.xyz, WorldNormal);
-	float OutPdf = Direction_Tangent.w;
-	float2 UV = InverseEquiAreaSphericalMapping(OutDirection.yzx);
-	float3 Radiance = SkylightTexture.SampleLevel(SkylightTextureSampler, UV, 0).xyz;
-	return CreateLightSample(Radiance / OutPdf, OutPdf, OutDirection, POSITIVE_INFINITY);
+#if PATHTRACING_SKY_MIS == 0
+	FSkyLightSample Sample = SkyLight_SampleLight(RandSample);
+	return CreateLightSample(Sample.Radiance / Sample.Pdf, Sample.Pdf, Sample.Direction, POSITIVE_INFINITY);
 #else
-	float2 UV = RandSample;
+	// account for the fact that with MIS compensation, we could have a 0 probability of choosing the sky
+	float SkyLightSamplingStrategyPdf = SkyLight_Estimate() > 0.0 ? 0.5 : 0.0;
 
-	int3 Pixel = int3(0, 0, SkylightMipCount - 2);
-	for (; Pixel.z >= 0; Pixel.z--)
+	// Do a simple one-sample MIS between the skylight and cosine sampling on the assumption that
+	// the shading loop is only sampling the light and not doing its own MIS.
+	float3 SkyLightRadiance = 0;
+	float3 Direction = 0;
+	float CosinePdf = 0;
+	float SkyLightPdf = 0;
+	BRANCH if (RandSample.x < SkyLightSamplingStrategyPdf)
 	{
-		Pixel.xy *= 2;
-		float P00 = SkylightPdf.Load(Pixel + int3(0, 0, 0));
-		float P10 = SkylightPdf.Load(Pixel + int3(1, 0, 0));
-		float P01 = SkylightPdf.Load(Pixel + int3(0, 1, 0));
-		float P11 = SkylightPdf.Load(Pixel + int3(1, 1, 0));
+		RandSample.x /= SkyLightSamplingStrategyPdf;
 
-		float L = P00 + P01;
-		float R = P10 + P11;
+		FSkyLightSample SkySample = SkyLight_SampleLight(RandSample);
 
-		float ProbX = L / (L + R);
-		if (UV.x < ProbX)
-		{
-			UV.x /= ProbX;
-			float ProbY = P00 / L;
-			if (UV.y < ProbY)
-			{
-				UV.y /= ProbY;
-			}
-			else
-			{
-				Pixel.y++;
-				UV.y = (UV.y - ProbY) / (1 - ProbY);
-			}
-		}
-		else
-		{
-			Pixel.x++;
-			UV.x = (UV.x - ProbX) / (1 - ProbX);
-			float ProbY = P10 / R;
-			if (UV.y < ProbY)
-			{
-				UV.y /= ProbY;
-			}
-			else
-			{
-				Pixel.y++;
-				UV.y = (UV.y - ProbY) / (1 - ProbY);
-			}
-		}
+		SkyLightRadiance = SkySample.Radiance;
+		Direction = SkySample.Direction;
+		SkyLightPdf = SkySample.Pdf;
+		CosinePdf = saturate(dot(WorldNormal, Direction)) / PI;
 	}
+	else
+	{
+		RandSample.x = (RandSample.x - SkyLightSamplingStrategyPdf) / (1.0 - SkyLightSamplingStrategyPdf);
 
-	Pixel.z = 0;
-	float4 Result = SkylightTexture.Load(Pixel);
-	float3 Radiance = Result.xyz;
-	float OutPdf = Result.w / (4 * PI * SkylightPdf.Load(int3(0, 0, SkylightMipCount - 1)));
+		float4 CosSample = CosineSampleHemisphere(RandSample, WorldNormal);
+		
+		Direction = CosSample.xyz;
+		CosinePdf = CosSample.w;
+
+		float4 SkyEval = SkyLight_EvalLight(Direction);
+		SkyLightRadiance = SkyEval.xyz;
+		SkyLightPdf = SkyEval.w;
+	}
 
-	UV = (float2(Pixel.xy) + UV) * SkylightInvResolution;
+	float Pdf = lerp(CosinePdf, SkyLightPdf, SkyLightSamplingStrategyPdf);
 
-	return CreateLightSample(Radiance / OutPdf, OutPdf, EquiAreaSphericalMapping(UV).zxy, POSITIVE_INFINITY);
+	return CreateLightSample(SkyLightRadiance / Pdf, Pdf, Direction, POSITIVE_INFINITY);
 #endif
 }
diff --git a/Engine/Shaders/Private/RayTracing/BuildMipTreeCS.usf b/Engine/Shaders/Private/RayTracing/BuildMipTreeCS.usf
deleted file mode 100644
diff --git a/Engine/Shaders/Private/RayTracing/BuildMipTreePdfCS.usf b/Engine/Shaders/Private/RayTracing/BuildMipTreePdfCS.usf
deleted file mode 100644
diff --git a/Engine/Shaders/Private/RayTracing/BuildSolidAnglePdfCS.usf b/Engine/Shaders/Private/RayTracing/BuildSolidAnglePdfCS.usf
deleted file mode 100644
diff --git a/Engine/Shaders/Private/RayTracing/GenerateSkyLightVisibilityRaysCS.usf b/Engine/Shaders/Private/RayTracing/GenerateSkyLightVisibilityRaysCS.usf
--- a/Engine/Shaders/Private/RayTracing/GenerateSkyLightVisibilityRaysCS.usf
+++ b/Engine/Shaders/Private/RayTracing/GenerateSkyLightVisibilityRaysCS.usf
@@ -21,10 +21,8 @@ void MainCS(
 {
 	uint2 PixelCoord = uint2(DispatchThreadId.x, DispatchThreadId.y);
 
-	// Cache Cubeface CDF
-	uint MipCount = log2(SkyLight.MipDimensions.x);
-	float FaceCdf[7];
-	BuildFaceCdf(MipCount, FaceCdf);
+	// split samples between strategies unless the skylight pdf is 0 due to MIS compensation (which implies a constant map)
+	const float SkyLightSamplingStrategyPdf = SkyLight_Estimate() > 0 ? 0.5 : 0.0;
 
 	for (uint SampleIndex = 0; SampleIndex < SkyLight.SamplesPerPixel; ++SampleIndex)
 	{
@@ -32,29 +30,37 @@ void MainCS(
 		RandomSequence_Initialize(RandSequence, PixelCoord, SampleIndex, View.StateFrameIndex, SkyLight.SamplesPerPixel);
 
 		// Determine sky light or lambert ray
-		float3 RandSample = RandomSequence_GenerateSample3D(RandSequence);
+		float2 RandSample = RandomSequence_GenerateSample2D(RandSequence);
 
 		float3 RayDirection;
-		float SkyLightSamplingStrategyPdf = 0.5;
-		float MisWeight = 1.0;
-		float RayPdf = 0.0;
+		float UniformPdf = 1.0 / (4 * PI);
+		float SkyLightPdf;
 		BRANCH
 			if (RandSample.x < SkyLightSamplingStrategyPdf)
 			{
 				RandSample.x /= SkyLightSamplingStrategyPdf;
-				GenerateSkyLightOcclusionRay(FaceCdf, RandSample, RayDirection, RayPdf);
-				MisWeight = RayPdf / (SkyLightSamplingStrategyPdf * RayPdf + (1.0 - SkyLightSamplingStrategyPdf) * PdfSphericalOcclusionRay(RayDirection));
+
+				FSkyLightSample SkySample = SkyLight_SampleLight(RandSample);
+
+				RayDirection = SkySample.Direction;
+				SkyLightPdf = SkySample.Pdf;
 			}
 			else
 			{
-				GenerateSphericalOcclusionRay(RandSample.yz, RayDirection, RayPdf);
-				MisWeight = RayPdf / ((1.0 - SkyLightSamplingStrategyPdf) * RayPdf + SkyLightSamplingStrategyPdf * PdfSkyLightOcclusionRay(FaceCdf, RayDirection));
+				RandSample.x = (RandSample.x - SkyLightSamplingStrategyPdf) / (1.0 - SkyLightSamplingStrategyPdf);
+
+				RayDirection = UniformSampleSphere(RandSample).xyz;
+				SkyLightPdf = SkyLight_EvalLight(RayDirection).w;
+
+				
 			}
 
+		float MisWeightOverPdf = 1.0 / lerp(UniformPdf, SkyLightPdf, SkyLightSamplingStrategyPdf);
+
 		// Write out a new Sky Light Visibility Ray for this sample
 
 		const uint SkyLightVisibilityRayIndex = GetSkyLightVisibilityRayIndex(PixelCoord, SampleIndex, SkyLightVisibilityRaysDimensions.xy);
 
-		OutSkyLightVisibilityRays[SkyLightVisibilityRayIndex].DirectionAndPdf = float4(RayDirection.xyz, MisWeight / RayPdf);
+		OutSkyLightVisibilityRays[SkyLightVisibilityRayIndex].DirectionAndPdf = float4(RayDirection, MisWeightOverPdf);
 	}
 }
diff --git a/Engine/Shaders/Private/RayTracing/RayTracingSkyLightCommon.ush b/Engine/Shaders/Private/RayTracing/RayTracingSkyLightCommon.ush
--- a/Engine/Shaders/Private/RayTracing/RayTracingSkyLightCommon.ush
+++ b/Engine/Shaders/Private/RayTracing/RayTracingSkyLightCommon.ush
@@ -7,214 +7,116 @@
 #pragma once
 
 #include "../MonteCarlo.ush"
-#include "SkyLightMipTreeCommon.ush"
 
-#define USE_MIP_PDF_TEXTURE 1
-#define USE_SOLID_ANGLE_TEXTURE 1
+// shader parameters for skylight integration
+Texture2D SkylightTexture;
+Texture2D<float> SkylightPdf;
+SamplerState SkylightTextureSampler;
+float SkylightInvResolution;
+int SkylightMipCount;
 
-float PdfSkyLightOcclusionRay(
-	float FaceCdf[7],
-	float3 WorldDirection
-)
+#define USE_HIERARCHICAL_IMPORTANCE_SAMPLING		1
+
+float SkyLight_Estimate()
 {
-	// Determine primary direction
-	uint MipCount = log2(SkyLight.MipDimensions.x);
-	uint3 TextureCoord = GetTextureCubeCoordinate(WorldDirection, SkyLight.TextureDimensions.xy);
-	float2 TextureUV = TextureCoord.xy / (float2)SkyLight.TextureDimensions.xy;
-	uint2 MipCoord = TextureUV * SkyLight.MipDimensions.xy;
-	float FacePdf = PdfFace(FaceCdf, TextureCoord.z);
-
-	uint StopLevel = min(SkyLight.SamplingStopLevel, MipCount);
-#if USE_MIP_PDF_TEXTURE
-	float MipPdf = PdfMipTree(uint3(MipCoord, TextureCoord.z), StopLevel);
-
-#if USE_SOLID_ANGLE_TEXTURE
-	uint LinearIndex = BufferOffsetAtPixel(MipCoord, StopLevel, SkyLight.MipDimensions.xy);
-	float AreaPdf = SkyLight.SolidAnglePdf[LinearIndex];
-#else
-	float AreaPdf = 1.0 / TexelCoordSolidAngle(MipCoord, DimensionsAtLevel(SkyLight.MipDimensions, StopLevel));
-#endif
-	return FacePdf * MipPdf * AreaPdf;
+	// cancels out constant factor in pdf below
+#if USE_HIERARCHICAL_IMPORTANCE_SAMPLING
+	return 4 * PI * SkylightPdf.Load(int3(0, 0, SkylightMipCount - 1)); 
 #else
-	uint2 MipPixel = 0;
-	float MipPdf = 1.0;
-	for (uint Index = 0; Index < MipCount - StopLevel; ++Index)
-	{
-		// Build Cdf
-		uint MipIndex = MipCount - Index - 1;
-#if USE_MORTON_CODE
-		uint BufferOffset = BufferOffsetAtPixel(MipPixel, MipIndex, SkyLight.MipDimensions.xy);
-		uint4 BufferOffsets = uint4(BufferOffset, BufferOffset + 1, BufferOffset + 2, BufferOffset + 3);
-#else
-		uint4 BufferOffsets;
-		BufferOffsets.x = BufferOffsetAtPixel(MipPixel, MipIndex, SkyLight.MipDimensions.xy);
-		BufferOffsets.y = BufferOffsetAtPixel(MipPixel + uint2(1, 0), MipIndex, SkyLight.MipDimensions.xy);
-		BufferOffsets.z = BufferOffsetAtPixel(MipPixel + uint2(0, 1), MipIndex, SkyLight.MipDimensions.xy);
-		BufferOffsets.w = BufferOffsetAtPixel(MipPixel + uint2(1, 1), MipIndex, SkyLight.MipDimensions.xy);
+	return 4 * PI;
 #endif
-		float4 Cdf = SampleMipTree(TextureCoord.z, BufferOffsets);
-		Cdf.y += Cdf.x;
-		Cdf.z += Cdf.y;
-		Cdf.w += Cdf.z;
-		Cdf /= Cdf.w;
-
-		// Determine texel choice
-		uint2 MipOffset = (MipCoord.xy >> MipIndex) - MipPixel;
-		uint LinearIndex = MipOffset.y << 1 | MipOffset.x;
-		
-		if (LinearIndex > 0)
-		{
-			MipPdf *= Cdf[LinearIndex] - Cdf[LinearIndex - 1] 
-		}
-		else
-		{
-			MipPdf *= Cdf[LinearIndex];
-		}
-
-		MipPixel += MipOffset;
-
-		MipPixel = MipPixel << 1;
-	}
-	MipPdf = max(MipPdf, 0.0);
-	// DEBUG: assert that PdfMipTree is equivalent for evaluating the decision tree..
-	//MipPdf = PdfMipTree(uint3(MipCoord, TextureCoord.z), StopLevel);
+}
 
-#if USE_SOLID_ANGLE_TEXTURE
-	uint LinearIndex = BufferOffsetAtPixel(MipCoord, StopLevel, SkyLight.MipDimensions.xy);
-	float AreaPdf = SkyLight.SolidAnglePdf[LinearIndex];
+// Returns: Radiance and Pdf for the specified direction
+float4 SkyLight_EvalLight(float3 Dir)
+{
+	// NOTE: assumes direction is normalized
+	float2 UV = InverseEquiAreaSphericalMapping(Dir.yzx);
+	float4 Result = SkylightTexture.SampleLevel(SkylightTextureSampler, UV, 0);
+	float3 Radiance = Result.xyz;
+#if USE_HIERARCHICAL_IMPORTANCE_SAMPLING
+	float Pdf = Result.w > 0 ? Result.w / (4 * PI * SkylightPdf.Load(int3(0, 0, SkylightMipCount - 1))) : 0.0; 
 #else
-	float AreaPdf = 1.0 / TexelCoordSolidAngle(MipCoord, DimensionsAtLevel(SkyLight.MipDimensions.xy, StopLevel));
-#endif
-	return FacePdf * MipPdf * AreaPdf;
+	float Pdf = 1.0 / (4.0 * PI);
 #endif
+	return float4(Radiance, Pdf);
 }
 
-bool GenerateSkyLightOcclusionRay(
-	float FaceCdf[7],
-	float3 RandSample,
-	out float3 RayDirection,
-	out float RayPdf
-)
+struct FSkyLightSample {
+	float3 Direction;
+	float3 Radiance;
+	float Pdf;
+};
+
+FSkyLightSample SkyLight_SampleLight(float2 RandSample)
 {
-	// Determine the appropriate face to use
-	uint MipCount = log2(SkyLight.MipDimensions.x);
-	float FacePdf = 0.0;
-	uint FaceIndex = SampleFace(FaceCdf, RandSample.x, FacePdf);
-
-	uint2 MipPixel = 0;
-	float MipPdf = 1.0;
-	uint StopLevel = min(SkyLight.SamplingStopLevel, MipCount);
-	for (uint Index = 0; Index < MipCount - StopLevel; ++Index)
-	{
-		uint MipIndex = MipCount - Index - 1;
-#if USE_MORTON_CODE
-		uint BufferOffset = BufferOffsetAtPixel(MipPixel, MipIndex, SkyLight.MipDimensions.xy);
-		uint4 BufferOffsets = uint4(BufferOffset, BufferOffset + 1, BufferOffset + 2, BufferOffset + 3);
-#else
-		uint4 BufferOffsets;
-		BufferOffsets.x = BufferOffsetAtPixel(MipPixel, MipIndex, SkyLight.MipDimensions);
-		BufferOffsets.y = BufferOffsetAtPixel(MipPixel + uint2(1, 0), MipIndex, SkyLight.MipDimensions.xy);
-		BufferOffsets.z = BufferOffsetAtPixel(MipPixel + uint2(0, 1), MipIndex, SkyLight.MipDimensions.xy);
-		BufferOffsets.w = BufferOffsetAtPixel(MipPixel + uint2(1, 1), MipIndex, SkyLight.MipDimensions.xy);
-#endif
-		float4 Cdf = SampleMipTree(FaceIndex, BufferOffsets);
-		Cdf.y += Cdf.x;
-		Cdf.z += Cdf.y;
-		Cdf.w += Cdf.z;
-		Cdf /= Cdf.w;
+#if USE_HIERARCHICAL_IMPORTANCE_SAMPLING
+	float2 UV = RandSample;
 
-		if (RandSample.y < Cdf.x)
-		{
-			MipPdf *= Cdf.x;
-			RandSample.y /= Cdf.x;
-		}
-		else if (RandSample.y < Cdf.y)
-		{
-			MipPdf *= Cdf.y - Cdf.x;
-			MipPixel.x += 1;
-			RandSample.y = (RandSample.y - Cdf.x) / (Cdf.y - Cdf.x);
-		}
-		else if (RandSample.y < Cdf.z)
+	int3 Pixel = int3(0, 0, SkylightMipCount - 2);
+	for (; Pixel.z >= 0; Pixel.z--)
+	{
+		Pixel.xy *= 2;
+		// TODO: would be nice to have GatherRed available to do this in a single lookup ...
+		float P00 = SkylightPdf.Load(Pixel + int3(0, 0, 0));
+		float P10 = SkylightPdf.Load(Pixel + int3(1, 0, 0));
+		float P01 = SkylightPdf.Load(Pixel + int3(0, 1, 0));
+		float P11 = SkylightPdf.Load(Pixel + int3(1, 1, 0));
+
+		float L = P00 + P01;
+		float R = P10 + P11;
+
+		float ProbX = L / (L + R);
+		if (UV.x < ProbX)
 		{
-			MipPdf *= Cdf.z - Cdf.y;
-			MipPixel.y += 1;
-			RandSample.y = (RandSample.y - Cdf.y) / (Cdf.z - Cdf.y);
+			UV.x /= ProbX;
+			float ProbY = P00 / L;
+			if (UV.y < ProbY)
+			{
+				UV.y /= ProbY;
+			}
+			else
+			{
+				Pixel.y++;
+				UV.y = (UV.y - ProbY) / (1 - ProbY);
+			}
 		}
-		else // if (RandSample.y < Cdf.w)
+		else
 		{
-			MipPdf *= Cdf.w - Cdf.z;
-			MipPixel += 1;
-			RandSample.y = (RandSample.y - Cdf.z) / (Cdf.w - Cdf.z);
+			Pixel.x++;
+			UV.x = (UV.x - ProbX) / (1 - ProbX);
+			float ProbY = P10 / R;
+			if (UV.y < ProbY)
+			{
+				UV.y /= ProbY;
+			}
+			else
+			{
+				Pixel.y++;
+				UV.y = (UV.y - ProbY) / (1 - ProbY);
+			}
 		}
-
-		MipPixel = uint2(MipPixel.x << 1, MipPixel.y << 1);
 	}
-	MipPixel = uint2(MipPixel.x >> 1, MipPixel.y >> 1);
 
-	float2 PixelOffset = RandSample.yz;
-	float2 MipPixelUV = (MipPixel + PixelOffset) / float2(DimensionsAtLevel(SkyLight.MipDimensions.xy, StopLevel));
+	Pixel.z = 0;
+	float4 Result = SkylightTexture.Load(Pixel);
+	float3 Radiance = Result.xyz;
+	float OutPdf = Result.w / (4 * PI * SkylightPdf.Load(int3(0, 0, SkylightMipCount - 1)));
 
-	// Construct Cubemap coordinate
-	float3 TexelCoord = float3(MipPixelUV * SkyLight.TextureDimensions.xy, FaceIndex);
+	UV = (float2(Pixel.xy) + UV) * SkylightInvResolution;
 
-	// Determine area pdf
-#if USE_SOLID_ANGLE_TEXTURE
-	uint LinearIndex = BufferOffsetAtPixel(MipPixel, StopLevel, SkyLight.MipDimensions.xy);
-	float AreaPdf = SkyLight.SolidAnglePdf[LinearIndex];
+	FSkyLightSample Sample;
+	Sample.Direction = EquiAreaSphericalMapping(UV).zxy;
+	Sample.Radiance = Radiance;
+	Sample.Pdf = OutPdf;
+	return Sample;
 #else
-	float AreaPdf = 1.0 / TexelCoordSolidAngle(MipPixel, DimensionsAtLevel(SkyLight.MipDimensions.xy, StopLevel));
-#endif
-
-	RayDirection = normalize(GetTextureCubeVector(TexelCoord, SkyLight.TextureDimensions.xy));
-	RayPdf = FacePdf * MipPdf * AreaPdf;
-
-	return true;
-}
-
-float PdfSphericalOcclusionRay(
-	float3 WorldDirection
-)
-{
-	float Pdf = 1.0 / (4.0 * PI);
-	return Pdf;
-}
+	FSkyLightSample Sample;
 
-bool GenerateSphericalOcclusionRay(
-	float2 RandSample,
-	out float3 RayDirection,
-	out float RayPdf
-)
-{
-	// Perform uniform spherical sampling and convert to world-space
-	float4 Direction_World = UniformSampleSphere(RandSample);
-
-	RayDirection = Direction_World.xyz;
-	RayPdf = Direction_World.w;
-
-	return true;
-}
-
-float PdfCosineHemisphericalOcclusionRay(
-	float3 WorldNormal,
-	float3 WorldDirection
-)
-{
-	float Pdf = max(dot(WorldDirection, WorldNormal), 0.0) / PI;
-	return Pdf;
-}
-
-bool GenerateCosineHemisphericalOcclusionRay(
-	float3 WorldNormal,
-	float2 RandSample,
-	out float3 RayDirection,
-	out float RayPdf
-)
-{
-	// Perform cosine-hemispherical sampling
-	float4 Direction_World = CosineSampleHemisphere(RandSample, WorldNormal);
-
-	RayDirection = Direction_World.xyz;
-	RayPdf = Direction_World.w;
-
-	return true;
+	Sample.Direction = UniformSampleSphere(RandSample).xyz;
+	Sample.Pdf = 1.0 / (4.0 * PI);
+	float2 UV = InverseEquiAreaSphericalMapping(Sample.Direction.yzx);
+	Sample.Radiance = SkylightTexture.SampleLevel(SkylightTextureSampler, UV, 0).xyz;
+	return Sample;
+#endif
 }
diff --git a/Engine/Shaders/Private/RayTracing/RayTracingSkyLightEvaluation.ush b/Engine/Shaders/Private/RayTracing/RayTracingSkyLightEvaluation.ush
--- a/Engine/Shaders/Private/RayTracing/RayTracingSkyLightEvaluation.ush
+++ b/Engine/Shaders/Private/RayTracing/RayTracingSkyLightEvaluation.ush
@@ -25,12 +25,6 @@ StructuredBuffer<FSkyLightVisibilityRays> SkyLightVisibilityRays;
 #include "../HairStrands/HairStrandsRaytracing.ush"
 #endif
 
-// Evaluates the radiance for the current sky light in a given direction
-float3 EvaluateSkyLightRadiance(float3 Direction)
-{
-	return SkyLight.Color * TextureCubeSampleLevel(SkyLight.Texture, SkyLight.TextureSampler, Direction, 0).rgb;
-}
-
 // Evaluates the current sky light at a given location on a surface by integrating over the sky light via casting occlusion rays
 // SampleCoord - Used for accessing the decoupled visibility ray buffer and should reflect some sort of uniformly distributed 2d coordinate to properly importance sample from the sky light. If bGBufferSampleOrigin is true this must be
 //               the pixel location of the g-buffer sample for proper camera ray bias and depth reconstruction.
@@ -58,15 +52,6 @@ void SkyLightEvaluate(
 	inout float AmbientOcclusion,
 	inout float HitDistance)
 {
-	float FaceCdf[7];
-
-	if (!bDecoupleSampleGeneration)
-	{
-		// Cache Cubeface CDF
-		const uint MipCount = log2(SkyLight.MipDimensions.x);
-		BuildFaceCdf(MipCount, FaceCdf);
-	}
-
 	float3 CurrentWorldNormal = WorldNormal;
 	float RayDistance = 0.0;
 	float HitCount = 0.0;
@@ -78,6 +63,9 @@ void SkyLightEvaluate(
 	AmbientOcclusion = 0.0f;
 	HitDistance = -1.0f;
 
+	// split samples between strategies unless the skylight pdf is 0 due to MIS compensation (which implies a constant map)
+	const float SkyLightSamplingStrategyPdf = SkyLight_Estimate() > 0 ? 0.5 : 0.0;
+
 	// Iterate up to the requested sample count
 	for (uint SampleIndex = 0; SampleIndex < SamplesPerPixel; ++SampleIndex)
 	{
@@ -103,29 +91,38 @@ void SkyLightEvaluate(
 			RandomSequence_Initialize(RandSequence, PixelCoord, SampleIndex, View.StateFrameIndex, SamplesPerPixel);
 
 			// Determine sky light or lambert ray
-			float3 RandSample = RandomSequence_GenerateSample3D(RandSequence);
+			float2 RandSample = RandomSequence_GenerateSample2D(RandSequence);
 
 			// Generate a visibility ray for the current sample
-			float SkyLightSamplingStrategyPdf = 0.5;
-			float MisWeight = 1.0;
-			float RayPdf = 0.0;
+			float SkyLightPdf = 0;
+			float CosinePdf = 0;
 			BRANCH
 			if (RandSample.x < SkyLightSamplingStrategyPdf)
 			{
 				RandSample.x /= SkyLightSamplingStrategyPdf;
-				GenerateSkyLightOcclusionRay(FaceCdf, RandSample, Ray.Direction, RayPdf);
-				MisWeight = RayPdf / (SkyLightSamplingStrategyPdf * RayPdf + (1.0 - SkyLightSamplingStrategyPdf) * PdfCosineHemisphericalOcclusionRay(CurrentWorldNormal, Ray.Direction));
+
+				FSkyLightSample SkySample = SkyLight_SampleLight(RandSample);
+				Ray.Direction = SkySample.Direction;
+				SkyLightPdf = SkySample.Pdf;
+
+				CosinePdf = saturate(dot(CurrentWorldNormal, Ray.Direction)) / PI;
 			}
 			else
 			{
-				GenerateCosineHemisphericalOcclusionRay(CurrentWorldNormal, RandSample.yz, Ray.Direction, RayPdf);
-				MisWeight = RayPdf / ((1.0 - SkyLightSamplingStrategyPdf) * RayPdf + SkyLightSamplingStrategyPdf * PdfSkyLightOcclusionRay(FaceCdf, Ray.Direction));
+				RandSample.x = (RandSample.x - SkyLightSamplingStrategyPdf) / (1.0 - SkyLightSamplingStrategyPdf);
+
+				float4 CosSample = CosineSampleHemisphere(RandSample, CurrentWorldNormal);
+				Ray.Direction = CosSample.xyz;
+				CosinePdf = CosSample.w;
+
+				SkyLightPdf = SkyLight_EvalLight(Ray.Direction).w;
 			}
 
 			Ray.Origin = WorldPosition;
 			Ray.TMin = 0.0;
 			Ray.TMax = SkyLight.MaxRayDistance;
-			RayWeight = MisWeight / RayPdf;
+			// MIS / pdf
+			RayWeight = 1.0 / lerp(CosinePdf, SkyLightPdf, SkyLightSamplingStrategyPdf);
 		}
 
 		// Recompute the current WorldNormal if lighting a hair fiber
@@ -216,7 +213,7 @@ void SkyLightEvaluate(
 				LightingSample = EvaluateBxDF(GBufferData, N, V, L, NoL, ShadowTerms);
 			}
 			float3 Brdf = LightingSample.Diffuse + LightingSample.Transmission + LightingSample.Specular;
-			float3 IncomingRadiance = EvaluateSkyLightRadiance(Ray.Direction);
+			float3 IncomingRadiance = SkyLight_EvalLight(Ray.Direction).xyz;
 
 			ExitantRadiance += IncomingRadiance * Brdf * RayWeight;
 			float3 DiffuseThroughput = LightingSample.Diffuse;
diff --git a/Engine/Source/Runtime/Engine/Classes/Components/SkyLightComponent.h b/Engine/Source/Runtime/Engine/Classes/Components/SkyLightComponent.h
--- a/Engine/Source/Runtime/Engine/Classes/Components/SkyLightComponent.h
+++ b/Engine/Source/Runtime/Engine/Classes/Components/SkyLightComponent.h
@@ -10,7 +10,6 @@
 #include "RenderingThread.h"
 #include "Components/LightComponentBase.h"
 #include "Math/SHMath.h"
-#include "Rendering/SkyLightImportanceSampling.h"
 #include "SkyLightComponent.generated.h"
 
 class FSkyLightSceneProxy;
@@ -287,14 +286,9 @@ public:
 	UFUNCTION(BlueprintCallable, Category="Rendering|Components|SkyLight")
 	void RecaptureSky();
 
-	void UpdateImportanceSamplingData();
-
 	virtual void Serialize(FArchive& Ar) override;
 
 	const FTexture* GetProcessedSkyTexture() const { return ProcessedSkyTexture; }
-#if RHI_RAYTRACING
-	const FSkyLightImportanceSamplingData* GetImportanceSamplingData() const { return ImportanceSamplingData; }
-#endif
 	FSHVectorRGB3 GetIrradianceEnvironmentMap() { return IrradianceEnvironmentMap; }
 protected:
 
@@ -312,10 +306,6 @@ protected:
 	FSHVectorRGB3 IrradianceEnvironmentMap;
 	float AverageBrightness=0.0f;
 
-#if RHI_RAYTRACING
-	TRefCountPtr<FSkyLightImportanceSamplingData> ImportanceSamplingData;
-#endif
-
 	/** If 0, no blend is present.  If > 0, BlendDestinationProcessedSkyTexture and BlendDestinationIrradianceEnvironmentMap must be generated and used for rendering. */
 	float BlendFraction;
 
@@ -386,10 +376,6 @@ public:
 	// This has to be refcounted to keep it alive during the handoff without doing a deep copy
 	TRefCountPtr<FSkyTextureCubeResource> ProcessedSkyTexture;
 
-#if RHI_RAYTRACING
-	TRefCountPtr<FSkyLightImportanceSamplingData> ImportanceSamplingData;
-#endif
-
 	FSHVectorRGB3 IrradianceEnvironmentMap;
 };
 
diff --git a/Engine/Source/Runtime/Engine/Private/Components/SkyLightComponent.cpp b/Engine/Source/Runtime/Engine/Private/Components/SkyLightComponent.cpp
--- a/Engine/Source/Runtime/Engine/Private/Components/SkyLightComponent.cpp
+++ b/Engine/Source/Runtime/Engine/Private/Components/SkyLightComponent.cpp
@@ -221,9 +221,6 @@ FSkyLightSceneProxy::FSkyLightSceneProxy(const USkyLightComponent* InLightCompon
 	, CaptureCubeMapResolution(InLightComponent->CubemapResolution)
 	, LowerHemisphereColor(InLightComponent->LowerHemisphereColor)
 	, bLowerHemisphereIsSolidColor(InLightComponent->bLowerHemisphereIsBlack)
-#if RHI_RAYTRACING
-	, ImportanceSamplingData(InLightComponent->ImportanceSamplingData)
-#endif
 	, LightColor(FLinearColor(InLightComponent->LightColor) * InLightComponent->Intensity)
 	, bMovable(InLightComponent->IsMovable())
 {
@@ -505,32 +502,6 @@ void USkyLightComponent::DestroyRenderState_Concurrent()
 	}
 }
 
-void USkyLightComponent::UpdateImportanceSamplingData()
-{
-	check(IsInGameThread());
-
-#if RHI_RAYTRACING
-	if (IsRayTracingEnabled() && ProcessedSkyTexture)
-	{
-		if (!ImportanceSamplingData.IsValid())
-		{
-			ImportanceSamplingData = new FSkyLightImportanceSamplingData();
-			BeginInitResource(ImportanceSamplingData);
-			MarkRenderStateDirty();
-		}
-
-		ENQUEUE_RENDER_COMMAND(UpdateImportanceSamplingDataCmd)(
-			[this](FRHICommandListImmediate& RHICmdList)
-		{
-			if (ProcessedSkyTexture)
-			{
-				ImportanceSamplingData->BuildCDFs(ProcessedSkyTexture);
-			}
-		});
-	}
-#endif
-}
-
 #if WITH_EDITOR
 void USkyLightComponent::PreEditChange(FProperty* PropertyAboutToChange)
 {
@@ -650,10 +621,6 @@ void USkyLightComponent::BeginDestroy()
 	// Release reference
 	ProcessedSkyTexture = NULL;
 
-#if RHI_RAYTRACING
-	ImportanceSamplingData.SafeRelease();
-#endif
-
 	// Begin a fence to track the progress of the above BeginReleaseResource being completed on the RT
 	ReleaseResourcesFence.BeginFence();
 
@@ -672,9 +639,6 @@ TStructOnScope<FActorComponentInstanceData> USkyLightComponent::GetComponentInst
 	FPrecomputedSkyLightInstanceData* SkyLightInstanceData = InstanceData.Cast<FPrecomputedSkyLightInstanceData>();
 	SkyLightInstanceData->LightGuid = LightGuid;
 	SkyLightInstanceData->ProcessedSkyTexture = ProcessedSkyTexture;
-#if RHI_RAYTRACING
-	SkyLightInstanceData->ImportanceSamplingData = ImportanceSamplingData;
-#endif
 
 	// Block until the rendering thread has completed its writes from a previous capture
 	IrradianceMapFence.Wait();
@@ -690,9 +654,6 @@ void USkyLightComponent::ApplyComponentInstanceData(FPrecomputedSkyLightInstance
 
 	LightGuid = (HasStaticShadowing() ? LightMapData->LightGuid : FGuid());
 	ProcessedSkyTexture = LightMapData->ProcessedSkyTexture;
-#if RHI_RAYTRACING
-	ImportanceSamplingData = LightMapData->ImportanceSamplingData;
-#endif
 	IrradianceEnvironmentMap = LightMapData->IrradianceEnvironmentMap;
 	AverageBrightness = LightMapData->AverageBrightness;
 
@@ -761,7 +722,6 @@ void USkyLightComponent::UpdateSkyCaptureContentsArray(UWorld* WorldToUpdate, TA
 					}
 
 					WorldToUpdate->Scene->UpdateSkyCaptureContents(CaptureComponent, CaptureComponent->bCaptureEmissiveOnly, CaptureComponent->Cubemap, CaptureComponent->ProcessedSkyTexture, CaptureComponent->AverageBrightness, CaptureComponent->IrradianceEnvironmentMap, NULL);
-					CaptureComponent->UpdateImportanceSamplingData();
 				}
 				else
 				{
@@ -777,7 +737,6 @@ void USkyLightComponent::UpdateSkyCaptureContentsArray(UWorld* WorldToUpdate, TA
 					}
 
 					WorldToUpdate->Scene->UpdateSkyCaptureContents(CaptureComponent, CaptureComponent->bCaptureEmissiveOnly, CaptureComponent->BlendDestinationCubemap, CaptureComponent->BlendDestinationProcessedSkyTexture, CaptureComponent->BlendDestinationAverageBrightness, CaptureComponent->BlendDestinationIrradianceEnvironmentMap, NULL);
-					CaptureComponent->UpdateImportanceSamplingData();
 				}
 
 				CaptureComponent->IrradianceMapFence.BeginFence();
diff --git a/Engine/Source/Runtime/Engine/Private/Rendering/SkyLightImportanceSampling.cpp b/Engine/Source/Runtime/Engine/Private/Rendering/SkyLightImportanceSampling.cpp
deleted file mode 100644
diff --git a/Engine/Source/Runtime/Engine/Public/Rendering/SkyLightImportanceSampling.h b/Engine/Source/Runtime/Engine/Public/Rendering/SkyLightImportanceSampling.h
deleted file mode 100644
diff --git a/Engine/Source/Runtime/Engine/Public/SceneManagement.h b/Engine/Source/Runtime/Engine/Public/SceneManagement.h
--- a/Engine/Source/Runtime/Engine/Public/SceneManagement.h
+++ b/Engine/Source/Runtime/Engine/Public/SceneManagement.h
@@ -28,7 +28,6 @@
 #include "LightmapUniformShaderParameters.h"
 #include "DynamicBufferAllocator.h"
 #include "Rendering/SkyAtmosphereCommonData.h"
-#include "Rendering/SkyLightImportanceSampling.h"
 
 class FCanvas;
 class FLightMap;
@@ -1139,9 +1138,6 @@ public:
 	uint32 CaptureCubeMapResolution;
 	FLinearColor LowerHemisphereColor;
 	bool bLowerHemisphereIsSolidColor;
-#if RHI_RAYTRACING
-	FSkyLightImportanceSamplingData* ImportanceSamplingData;
-#endif
 
 	bool IsMovable() { return bMovable; }
 
diff --git a/Engine/Source/Runtime/Renderer/Private/DeferredShadingRenderer.h b/Engine/Source/Runtime/Renderer/Private/DeferredShadingRenderer.h
--- a/Engine/Source/Runtime/Renderer/Private/DeferredShadingRenderer.h
+++ b/Engine/Source/Runtime/Renderer/Private/DeferredShadingRenderer.h
@@ -906,7 +906,6 @@ private:
 
 	void VisualizeRectLightMipTree(FRHICommandListImmediate& RHICmdList, const FViewInfo& View, const FRWBuffer& RectLightMipTree, const FIntVector& RectLightMipTreeDimensions);
 
-	void GenerateSkyLightVisibilityRays(FRDGBuilder& GraphBuilder, FRDGBufferRef& SkyLightVisibilityRays, FIntVector& Dimensions);
 	void VisualizeSkyLightMipTree(FRHICommandListImmediate& RHICmdList, const FViewInfo& View, const TRefCountPtr<IPooledRenderTarget>& SceneColor, FRWBuffer& SkyLightMipTreePosX, FRWBuffer& SkyLightMipTreePosY, FRWBuffer& SkyLightMipTreePosZ, FRWBuffer& SkyLightMipTreeNegX, FRWBuffer& SkyLightMipTreeNegY, FRWBuffer& SkyLightMipTreeNegZ, const FIntVector& SkyLightMipDimensions);
 
 	void RenderRayTracingSkyLight(
diff --git a/Engine/Source/Runtime/Renderer/Private/PathTracing.cpp b/Engine/Source/Runtime/Renderer/Private/PathTracing.cpp
--- a/Engine/Source/Runtime/Renderer/Private/PathTracing.cpp
+++ b/Engine/Source/Runtime/Renderer/Private/PathTracing.cpp
@@ -597,7 +597,9 @@ RENDERER_API void PrepareSkyTexture_Internal(
 	}
 }
 
-bool PrepareSkyTexture(FRDGBuilder& GraphBuilder, FScene* Scene, const FViewInfo& View, bool SkylightEnabled, bool UseMISCompensation, bool IsSkylightCachingEnabled, FPathTracingSkylight* SkylightParameters)
+RDG_REGISTER_BLACKBOARD_STRUCT(FPathTracingSkylight)
+
+bool PrepareSkyTexture(FRDGBuilder& GraphBuilder, FScene* Scene, const FViewInfo& View, bool SkylightEnabled, bool UseMISCompensation, FPathTracingSkylight* SkylightParameters)
 {
 	SkylightParameters->SkylightTextureSampler = TStaticSamplerState<SF_Bilinear, AM_Clamp, AM_Clamp, AM_Clamp>::GetRHI();
 
@@ -614,6 +616,17 @@ bool PrepareSkyTexture(FRDGBuilder& GraphBuilder, FScene* Scene, const FViewInfo
 		return false;
 	}
 
+	// the sky is actually enabled, lets see if someone already made use of it for this frame
+	const FPathTracingSkylight* PreviousSkylightParameters = GraphBuilder.Blackboard.Get<FPathTracingSkylight>();
+	if (PreviousSkylightParameters != nullptr)
+	{
+		*SkylightParameters = *PreviousSkylightParameters;
+		return true;
+	}
+
+	// should we remember the skylight prep for the next frame?
+	const bool IsSkylightCachingEnabled = CVarPathTracingSkylightCaching.GetValueOnAnyThread() != 0;
+
 	if (!IsSkylightCachingEnabled)
 	{
 		// we don't want any caching - release what we might have been holding onto
@@ -657,6 +670,10 @@ bool PrepareSkyTexture(FRDGBuilder& GraphBuilder, FScene* Scene, const FViewInfo
 		GraphBuilder.QueueTextureExtraction(SkylightParameters->SkylightTexture, &Scene->PathTracingSkylightTexture);
 		GraphBuilder.QueueTextureExtraction(SkylightParameters->SkylightPdf, &Scene->PathTracingSkylightPdf);
 	}
+
+	// remember the skylight parameters for future passes within this frame
+	GraphBuilder.Blackboard.Create<FPathTracingSkylight>() = *SkylightParameters;
+
 	return true;
 }
 
@@ -775,8 +792,7 @@ void SetLightParameters(FRDGBuilder& GraphBuilder, FPathTracingRG::FParameters*
 
 	// Prepend SkyLight to light buffer since it is not part of the regular light list
 	const float Inf = std::numeric_limits<float>::infinity();
-	const bool IsSkylightCachingEnabled = CVarPathTracingSkylightCaching.GetValueOnAnyThread() != 0;
-	if (PrepareSkyTexture(GraphBuilder, Scene, View, true, UseMISCompensation, IsSkylightCachingEnabled, &PassParameters->SkylightParameters))
+	if (PrepareSkyTexture(GraphBuilder, Scene, View, true, UseMISCompensation, &PassParameters->SkylightParameters))
 	{
 		check(Scene->SkyLight != nullptr);
 		FPathTracingLight& DestLight = Lights[NumLights++];
diff --git a/Engine/Source/Runtime/Renderer/Private/PathTracing.h b/Engine/Source/Runtime/Renderer/Private/PathTracing.h
--- a/Engine/Source/Runtime/Renderer/Private/PathTracing.h
+++ b/Engine/Source/Runtime/Renderer/Private/PathTracing.h
@@ -17,4 +17,4 @@ class FRDGBuilder;
 class FScene;
 class FViewInfo;
 
-bool PrepareSkyTexture(FRDGBuilder& GraphBuilder, FScene* Scene, const FViewInfo& View, bool SkylightEnabled, bool UseMISCompensation, bool IsSkylightCachingEnabled, FPathTracingSkylight* SkylightParameters);
+bool PrepareSkyTexture(FRDGBuilder& GraphBuilder, FScene* Scene, const FViewInfo& View, bool SkylightEnabled, bool UseMISCompensation, FPathTracingSkylight* SkylightParameters);
diff --git a/Engine/Source/Runtime/Renderer/Private/RayTracing/RayTracingGlobalIllumination.cpp b/Engine/Source/Runtime/Renderer/Private/RayTracing/RayTracingGlobalIllumination.cpp
--- a/Engine/Source/Runtime/Renderer/Private/RayTracing/RayTracingGlobalIllumination.cpp
+++ b/Engine/Source/Runtime/Renderer/Private/RayTracing/RayTracingGlobalIllumination.cpp
@@ -257,12 +257,11 @@ static void SetupLightParameters(
 
 	FSkyLightSceneProxy* SkyLight = Scene->SkyLight;
 
-	const bool bUseMISCompensation = false; // RTGI doesn't use MIS
-	const bool bUseSkylightCaching = true; // re-use the skylight texture if possible
+	const bool bUseMISCompensation = true;
 	const bool bSkylightEnabled = SkyLight && SkyLight->bAffectGlobalIllumination && CVarRayTracingGlobalIlluminationSkyLight.GetValueOnRenderThread() != 0;
 
 	// Prepend SkyLight to light buffer (if it is active)
-	if (PrepareSkyTexture(GraphBuilder, Scene, View, bSkylightEnabled, bUseMISCompensation, bUseSkylightCaching, SkylightParameters))
+	if (PrepareSkyTexture(GraphBuilder, Scene, View, bSkylightEnabled, bUseMISCompensation, SkylightParameters))
 	{
 		FPathTracingLight& DestLight = Lights[LightCount];
 
@@ -425,6 +424,12 @@ class FGlobalIlluminationRGS : public FGlobalShader
 		return ShouldCompileRayTracingShadersForProject(Parameters.Platform);
 	}
 
+	static void ModifyCompilationEnvironment(const FGlobalShaderPermutationParameters& Parameters, FShaderCompilerEnvironment& OutEnvironment)
+	{
+		// We need the skylight to do its own form of MIS because RTGI doesn't do its own
+		OutEnvironment.SetDefine(TEXT("PATHTRACING_SKY_MIS"), 1);
+	}
+
 	BEGIN_SHADER_PARAMETER_STRUCT(FParameters, )
 		SHADER_PARAMETER(uint32, SamplesPerPixel)
 		SHADER_PARAMETER(uint32, MaxBounces)
@@ -482,6 +487,12 @@ class FRayTracingGlobalIlluminationCreateGatherPointsRGS : public FGlobalShader
 		return ShouldCompileRayTracingShadersForProject(Parameters.Platform);
 	}
 
+	static void ModifyCompilationEnvironment(const FGlobalShaderPermutationParameters& Parameters, FShaderCompilerEnvironment& OutEnvironment)
+	{
+		// We need the skylight to do its own form of MIS because RTGI doesn't do its own
+		OutEnvironment.SetDefine(TEXT("PATHTRACING_SKY_MIS"), 1);
+	}
+
 	BEGIN_SHADER_PARAMETER_STRUCT(FParameters, )
 		SHADER_PARAMETER(uint32, GatherSamplesPerPixel)
 		SHADER_PARAMETER(uint32, SamplesPerPixel)
diff --git a/Engine/Source/Runtime/Renderer/Private/RayTracing/RayTracingReflections.cpp b/Engine/Source/Runtime/Renderer/Private/RayTracing/RayTracingReflections.cpp
--- a/Engine/Source/Runtime/Renderer/Private/RayTracing/RayTracingReflections.cpp
+++ b/Engine/Source/Runtime/Renderer/Private/RayTracing/RayTracingReflections.cpp
@@ -23,6 +23,7 @@
 #include "RayTracing/RayTracingLighting.h"
 #include "RayTracing/RayTracingSkyLight.h"
 #include "SceneTextureParameters.h"
+#include "PathTracing.h"
 
 static int32 GRayTracingReflections = -1;
 static FAutoConsoleVariableRef CVarReflectionsMethod(
@@ -284,7 +285,7 @@ class FRayTracingReflectionsRGS : public FGlobalShader
 		SHADER_PARAMETER(uint32, RenderTileOffsetX)
 		SHADER_PARAMETER(uint32, RenderTileOffsetY)
 		SHADER_PARAMETER(uint32, EnableTranslucency)
-		SHADER_PARAMETER(int32, SkyLightDecoupleSampleGeneration) 
+		SHADER_PARAMETER(int32, SkyLightDecoupleSampleGeneration)
 		SHADER_PARAMETER(int32, SampleMode)
 		SHADER_PARAMETER(int32, SampleOffset)
 
@@ -303,6 +304,7 @@ class FRayTracingReflectionsRGS : public FGlobalShader
 		SHADER_PARAMETER_STRUCT_REF(FReflectionCaptureShaderData, ReflectionCapture)
 		SHADER_PARAMETER_STRUCT_REF(FForwardLightData, Forward)
 		SHADER_PARAMETER_STRUCT_REF(FSkyLightData, SkyLightData)
+		SHADER_PARAMETER_STRUCT_INCLUDE(FPathTracingSkylight, SkylightParameters)
 
 		// Optional indirection buffer used for sorted materials
 		SHADER_PARAMETER_RDG_BUFFER_UAV(RWStructuredBuffer<FDeferredMaterialPayload>, MaterialBuffer)
@@ -683,14 +685,10 @@ void FDeferredShadingSceneRenderer::RenderRayTracingReflections(
 	CommonParameters.ReflectionCapture = View.ReflectionCaptureUniformBuffer;
 	CommonParameters.Forward = View.ForwardLightingResources->ForwardLightDataUniformBuffer;
 
-	if (bRayTraceSkyLightContribution)
-	{
-		// Fill Sky Light parameters
-		FSkyLightData SkyLightData;
-		SetupSkyLightParameters(*Scene, &SkyLightData);
-
-		CommonParameters.SkyLightData = CreateUniformBufferImmediate(SkyLightData, EUniformBufferUsage::UniformBuffer_SingleDraw);
-	}
+	// Fill Sky Light parameters
+	FSkyLightData SkyLightData;
+	SetupSkyLightParameters(GraphBuilder, Scene, View, bRayTraceSkyLightContribution, &CommonParameters.SkylightParameters, &SkyLightData);
+	CommonParameters.SkyLightData = CreateUniformBufferImmediate(SkyLightData, EUniformBufferUsage::UniformBuffer_SingleDraw);
 
 	for (int32 SamplePassIndex = 0; SamplePassIndex < Options.SamplesPerPixel; SamplePassIndex++)
 	{
diff --git a/Engine/Source/Runtime/Renderer/Private/RayTracing/RayTracingSkyLight.h b/Engine/Source/Runtime/Renderer/Private/RayTracing/RayTracingSkyLight.h
--- a/Engine/Source/Runtime/Renderer/Private/RayTracing/RayTracingSkyLight.h
+++ b/Engine/Source/Runtime/Renderer/Private/RayTracing/RayTracingSkyLight.h
@@ -12,29 +12,10 @@ RaytracingOptions.h declares ray tracing options for use in rendering
 
 BEGIN_GLOBAL_SHADER_PARAMETER_STRUCT(FSkyLightData, RENDERER_API)
 	SHADER_PARAMETER(uint32, SamplesPerPixel)
-	SHADER_PARAMETER(uint32, SamplingStopLevel)
 	SHADER_PARAMETER(float, MaxRayDistance)
-	SHADER_PARAMETER(FVector, Color)
-	SHADER_PARAMETER(FIntVector, MipDimensions)
 	SHADER_PARAMETER(float, MaxNormalBias)
 	SHADER_PARAMETER(float, MaxShadowThickness)
 	SHADER_PARAMETER(int, bTransmission)
-	SHADER_PARAMETER_TEXTURE(TextureCube, Texture)
-	SHADER_PARAMETER_SAMPLER(SamplerState, TextureSampler)
-	SHADER_PARAMETER(FIntVector, TextureDimensions)
-	SHADER_PARAMETER_SRV(Buffer<float>, MipTreePosX)
-	SHADER_PARAMETER_SRV(Buffer<float>, MipTreeNegX)
-	SHADER_PARAMETER_SRV(Buffer<float>, MipTreePosY)
-	SHADER_PARAMETER_SRV(Buffer<float>, MipTreeNegY)
-	SHADER_PARAMETER_SRV(Buffer<float>, MipTreePosZ)
-	SHADER_PARAMETER_SRV(Buffer<float>, MipTreeNegZ)
-	SHADER_PARAMETER_SRV(Buffer<float>, MipTreePdfPosX)
-	SHADER_PARAMETER_SRV(Buffer<float>, MipTreePdfNegX)
-	SHADER_PARAMETER_SRV(Buffer<float>, MipTreePdfPosY)
-	SHADER_PARAMETER_SRV(Buffer<float>, MipTreePdfNegY)
-	SHADER_PARAMETER_SRV(Buffer<float>, MipTreePdfPosZ)
-	SHADER_PARAMETER_SRV(Buffer<float>, MipTreePdfNegZ)
-	SHADER_PARAMETER_SRV(Buffer<float>, SolidAnglePdf)
 END_GLOBAL_SHADER_PARAMETER_STRUCT()
 
 BEGIN_SHADER_PARAMETER_STRUCT(FWritableSkyLightVisibilityRaysData, )
@@ -51,7 +32,14 @@ END_SHADER_PARAMETER_STRUCT()
 
 int32 GetRayTracingSkyLightDecoupleSampleGenerationCVarValue();
 
-extern RENDERER_API bool SetupSkyLightParameters(const FScene& Scene, FSkyLightData* SkyLight);
+class FPathTracingSkylight;
+
+bool SetupSkyLightParameters(
+	FRDGBuilder& GraphBuilder, FScene* Scene, const FViewInfo& View,
+	bool bEnableSkylight,
+	FPathTracingSkylight* SkylightParameters,
+	FSkyLightData* SkyLight);
+
 void SetupSkyLightVisibilityRaysParameters(FRDGBuilder& GraphBuilder, const FViewInfo& View, FSkyLightVisibilityRaysData* OutSkyLightVisibilityRaysData);
 
 extern RENDERER_API void BuildSkyLightCdfs(FRHICommandListImmediate& RHICmdList, class FSkyLightSceneProxy* SkyLight);
diff --git a/Engine/Source/Runtime/Renderer/Private/RayTracing/RaytracingSkylight.cpp b/Engine/Source/Runtime/Renderer/Private/RayTracing/RaytracingSkylight.cpp
--- a/Engine/Source/Runtime/Renderer/Private/RayTracing/RaytracingSkylight.cpp
+++ b/Engine/Source/Runtime/Renderer/Private/RayTracing/RaytracingSkylight.cpp
@@ -25,6 +25,7 @@ static int32 GRayTracingSkyLight = -1;
 #include "SceneTextureParameters.h"
 #include "ScreenSpaceDenoise.h"
 #include "RayTracingDefinitions.h"
+#include "PathTracing.h"
 
 #include "Raytracing/RaytracingOptions.h"
 #include "PostProcess/PostProcessing.h"
@@ -139,9 +140,7 @@ bool ShouldRenderRayTracingSkyLight(const FSkyLightSceneProxy* SkyLightSceneProx
 
 	bRayTracingSkyEnabled = bRayTracingSkyEnabled && (GetSkyLightSamplesPerPixel(SkyLightSceneProxy) > 0);
 
-	return IsRayTracingSkyLightAllowed() && bRayTracingSkyEnabled &&
-		(SkyLightSceneProxy->ImportanceSamplingData != nullptr) &&
-		SkyLightSceneProxy->ImportanceSamplingData->bIsValid;
+	return IsRayTracingSkyLightAllowed() && bRayTracingSkyEnabled;
 }
 
 IMPLEMENT_GLOBAL_SHADER_PARAMETER_STRUCT(FSkyLightData, "SkyLight");
@@ -152,79 +151,32 @@ struct FSkyLightVisibilityRays
 };
 
 bool SetupSkyLightParameters(
-	const FScene& Scene,
+	FRDGBuilder& GraphBuilder, FScene* Scene, const FViewInfo& View,
+	bool bEnableSkylight,
+	FPathTracingSkylight* SkylightParameters,
 	FSkyLightData* SkyLightData)
 {
 	// Check if parameters should be set based on if the sky light's texture has been processed and if its mip tree has been built yet
-	if (
-		Scene.SkyLight &&
-		Scene.SkyLight->ProcessedTexture &&
-		Scene.SkyLight->ImportanceSamplingData)
+
+	const bool bUseMISCompensation = true;
+	if (PrepareSkyTexture(GraphBuilder, Scene, View, bEnableSkylight, bUseMISCompensation, SkylightParameters))
 	{
-		check(Scene.SkyLight->ImportanceSamplingData->bIsValid);
 
-		SkyLightData->SamplesPerPixel = GetSkyLightSamplesPerPixel(Scene.SkyLight);
-		SkyLightData->SamplingStopLevel = GRayTracingSkyLightSamplingStopLevel;
+		SkyLightData->SamplesPerPixel = GetSkyLightSamplesPerPixel(Scene->SkyLight);
 		SkyLightData->MaxRayDistance = GRayTracingSkyLightMaxRayDistance;
 		SkyLightData->MaxNormalBias = GetRaytracingMaxNormalBias();
-		SkyLightData->bTransmission = Scene.SkyLight->bTransmission;
+		SkyLightData->bTransmission = Scene->SkyLight->bTransmission;
 		SkyLightData->MaxShadowThickness = GRayTracingSkyLightMaxShadowThickness;
-
 		ensure(SkyLightData->SamplesPerPixel > 0);
-
-		SkyLightData->Color = FVector(Scene.SkyLight->GetEffectiveLightColor());
-		SkyLightData->Texture = Scene.SkyLight->ProcessedTexture->TextureRHI;
-		SkyLightData->TextureDimensions = FIntVector(Scene.SkyLight->ProcessedTexture->GetSizeX(), Scene.SkyLight->ProcessedTexture->GetSizeY(), 0);
-		SkyLightData->TextureSampler = Scene.SkyLight->ProcessedTexture->SamplerStateRHI;
-		SkyLightData->MipDimensions = Scene.SkyLight->ImportanceSamplingData->MipDimensions;
-
-		SkyLightData->MipTreePosX = Scene.SkyLight->ImportanceSamplingData->MipTreePosX.SRV;
-		SkyLightData->MipTreeNegX = Scene.SkyLight->ImportanceSamplingData->MipTreeNegX.SRV;
-		SkyLightData->MipTreePosY = Scene.SkyLight->ImportanceSamplingData->MipTreePosY.SRV;
-		SkyLightData->MipTreeNegY = Scene.SkyLight->ImportanceSamplingData->MipTreeNegY.SRV;
-		SkyLightData->MipTreePosZ = Scene.SkyLight->ImportanceSamplingData->MipTreePosZ.SRV;
-		SkyLightData->MipTreeNegZ = Scene.SkyLight->ImportanceSamplingData->MipTreeNegZ.SRV;
-
-		SkyLightData->MipTreePdfPosX = Scene.SkyLight->ImportanceSamplingData->MipTreePdfPosX.SRV;
-		SkyLightData->MipTreePdfNegX = Scene.SkyLight->ImportanceSamplingData->MipTreePdfNegX.SRV;
-		SkyLightData->MipTreePdfPosY = Scene.SkyLight->ImportanceSamplingData->MipTreePdfPosY.SRV;
-		SkyLightData->MipTreePdfNegY = Scene.SkyLight->ImportanceSamplingData->MipTreePdfNegY.SRV;
-		SkyLightData->MipTreePdfPosZ = Scene.SkyLight->ImportanceSamplingData->MipTreePdfPosZ.SRV;
-		SkyLightData->MipTreePdfNegZ = Scene.SkyLight->ImportanceSamplingData->MipTreePdfNegZ.SRV;
-		SkyLightData->SolidAnglePdf = Scene.SkyLight->ImportanceSamplingData->SolidAnglePdf.SRV;
-
-		// we prepared a valid SkyLight
 		return true;
 	}
 	else
 	{
+		// skylight is not enabled
 		SkyLightData->SamplesPerPixel = -1;
-		SkyLightData->SamplingStopLevel = 0;
 		SkyLightData->MaxRayDistance = 0.0f;
 		SkyLightData->MaxNormalBias = 0.0f;
 		SkyLightData->MaxShadowThickness = 0.0f;
-
-		SkyLightData->Color = FVector(0.0);
-		SkyLightData->Texture = GBlackTextureCube->TextureRHI;
-		SkyLightData->TextureSampler = TStaticSamplerState<SF_Bilinear, AM_Clamp, AM_Clamp, AM_Clamp>::GetRHI();
-		SkyLightData->MipDimensions = FIntVector(0);
-
-		SkyLightData->MipTreePosX = GBlackTextureWithSRV->ShaderResourceViewRHI;
-		SkyLightData->MipTreeNegX = GBlackTextureWithSRV->ShaderResourceViewRHI;
-		SkyLightData->MipTreePosY = GBlackTextureWithSRV->ShaderResourceViewRHI;
-		SkyLightData->MipTreeNegY = GBlackTextureWithSRV->ShaderResourceViewRHI;
-		SkyLightData->MipTreePosZ = GBlackTextureWithSRV->ShaderResourceViewRHI;
-		SkyLightData->MipTreeNegZ = GBlackTextureWithSRV->ShaderResourceViewRHI;
-
-		SkyLightData->MipTreePdfPosX = GBlackTextureWithSRV->ShaderResourceViewRHI;
-		SkyLightData->MipTreePdfNegX = GBlackTextureWithSRV->ShaderResourceViewRHI;
-		SkyLightData->MipTreePdfPosY = GBlackTextureWithSRV->ShaderResourceViewRHI;
-		SkyLightData->MipTreePdfNegY = GBlackTextureWithSRV->ShaderResourceViewRHI;
-		SkyLightData->MipTreePdfPosZ = GBlackTextureWithSRV->ShaderResourceViewRHI;
-		SkyLightData->MipTreePdfNegZ = GBlackTextureWithSRV->ShaderResourceViewRHI;
-		SkyLightData->SolidAnglePdf = GBlackTextureWithSRV->ShaderResourceViewRHI;
-
-		// we returned a dummy SkyLight
 		return false;
 	}
 }
@@ -296,6 +248,7 @@ class FRayTracingSkyLightRGS : public FGlobalShader
 
 		SHADER_PARAMETER_STRUCT_REF(FViewUniformShaderParameters, ViewUniformBuffer)
 		SHADER_PARAMETER_STRUCT_REF(FSkyLightData, SkyLightData)
+		SHADER_PARAMETER_STRUCT_INCLUDE(FPathTracingSkylight, SkyLightParameters)
 		SHADER_PARAMETER_RDG_UNIFORM_BUFFER(FVirtualVoxelParameters, VirtualVoxel)
 		
 		SHADER_PARAMETER_STRUCT_INCLUDE(FSkyLightVisibilityRaysData, SkyLightVisibilityRaysData)
@@ -358,6 +311,7 @@ class FGenerateSkyLightVisibilityRaysCS : public FGlobalShader
 	BEGIN_SHADER_PARAMETER_STRUCT(FParameters, )
 		SHADER_PARAMETER(int32, SamplesPerPixel)
 
+		SHADER_PARAMETER_STRUCT_INCLUDE(FPathTracingSkylight, SkylightParameters)
 		SHADER_PARAMETER_STRUCT_REF(FSkyLightData, SkyLightData)
 		SHADER_PARAMETER_STRUCT_REF(FViewUniformShaderParameters, ViewUniformBuffer)
 
@@ -368,33 +322,34 @@ class FGenerateSkyLightVisibilityRaysCS : public FGlobalShader
 
 IMPLEMENT_GLOBAL_SHADER(FGenerateSkyLightVisibilityRaysCS, "/Engine/Private/RayTracing/GenerateSkyLightVisibilityRaysCS.usf", "MainCS", SF_Compute);
 
-void FDeferredShadingSceneRenderer::GenerateSkyLightVisibilityRays(
+static void GenerateSkyLightVisibilityRays(
 	FRDGBuilder& GraphBuilder,
+	FViewInfo& View,
+	FPathTracingSkylight& SkylightParameters,
+	FSkyLightData& SkyLightData,
 	FRDGBufferRef& SkyLightVisibilityRaysBuffer,
 	FIntVector& Dimensions
 )
 {
 	FRHIResourceCreateInfo CreateInfo(TEXT("SkyLightVisibilityRays"));
 
-	// SkyLight data setup
-	FSkyLightData SkyLightData;
-	SetupSkyLightParameters(*Scene, &SkyLightData);
-
 	// Allocating mask of 256 x 256 rays
 	Dimensions = FIntVector(256, 256, 0);
 
+	// Compute Pass parameter definition
+	FGenerateSkyLightVisibilityRaysCS::FParameters* PassParameters = GraphBuilder.AllocParameters<FGenerateSkyLightVisibilityRaysCS::FParameters>();
+	PassParameters->ViewUniformBuffer = View.ViewUniformBuffer;
+	PassParameters->SkyLightData = CreateUniformBufferImmediate(SkyLightData, EUniformBufferUsage::UniformBuffer_SingleDraw);
+	PassParameters->SkylightParameters = SkylightParameters;
+
 	// Output structured buffer creation
 	FRDGBufferDesc BufferDesc = FRDGBufferDesc::CreateStructuredDesc(sizeof(FSkyLightVisibilityRays), Dimensions.X * Dimensions.Y * SkyLightData.SamplesPerPixel);
 	SkyLightVisibilityRaysBuffer = GraphBuilder.CreateBuffer(BufferDesc, TEXT("SkyLightVisibilityRays"));
 
-	// Compute Pass parameter definition
-	FGenerateSkyLightVisibilityRaysCS::FParameters* PassParameters = GraphBuilder.AllocParameters<FGenerateSkyLightVisibilityRaysCS::FParameters>();
-	PassParameters->ViewUniformBuffer = Views[0].ViewUniformBuffer;
-	PassParameters->SkyLightData = CreateUniformBufferImmediate(SkyLightData, EUniformBufferUsage::UniformBuffer_SingleDraw);
 	PassParameters->WritableSkyLightVisibilityRaysData.SkyLightVisibilityRaysDimensions = FIntVector(Dimensions.X, Dimensions.Y, 0);
 	PassParameters->WritableSkyLightVisibilityRaysData.OutSkyLightVisibilityRays = GraphBuilder.CreateUAV(SkyLightVisibilityRaysBuffer, EPixelFormat::PF_R32_UINT);
 
-	TShaderMapRef<FGenerateSkyLightVisibilityRaysCS> ComputeShader(GetGlobalShaderMap(FeatureLevel));
+	auto ComputeShader = View.ShaderMap->GetShader<FGenerateSkyLightVisibilityRaysCS>();
 
 	FComputeShaderUtils::AddPass(
 		GraphBuilder,
@@ -414,7 +369,12 @@ void FDeferredShadingSceneRenderer::RenderRayTracingSkyLight(
 	FRDGTextureRef& OutHitDistanceTexture)
 {
 	FSkyLightSceneProxy* SkyLight = Scene->SkyLight;
-	if (!ShouldRenderRayTracingSkyLight(SkyLight))
+	
+	// Fill Sky Light parameters
+	const bool bShouldRenderRayTracingSkyLight = ShouldRenderRayTracingSkyLight(SkyLight);
+	FPathTracingSkylight SkylightParameters;
+	FSkyLightData SkyLightData;
+	if (!SetupSkyLightParameters(GraphBuilder, Scene, Views[0], bShouldRenderRayTracingSkyLight, &SkylightParameters, &SkyLightData))
 	{
 		return;
 	}
@@ -423,8 +383,6 @@ void FDeferredShadingSceneRenderer::RenderRayTracingSkyLight(
 	RDG_GPU_STAT_SCOPE(GraphBuilder, RayTracingSkyLight);
 
 	check(SceneColorTexture);
-	check(SkyLight->ProcessedTexture);
-	check((SkyLight->ImportanceSamplingData != nullptr) && SkyLight->ImportanceSamplingData->bIsValid);
 
 	float ResolutionFraction = 1.0f;
 	if (GRayTracingSkyLightDenoiser != 0)
@@ -447,11 +405,12 @@ void FDeferredShadingSceneRenderer::RenderRayTracingSkyLight(
 		OutHitDistanceTexture = GraphBuilder.CreateTexture(Desc, TEXT("RayTracingSkyLightHitDistance"));
 	}
 
+
 	FRDGBufferRef SkyLightVisibilityRaysBuffer;
 	FIntVector SkyLightVisibilityRaysDimensions;
 	if (CVarRayTracingSkyLightDecoupleSampleGeneration.GetValueOnRenderThread() == 1)
 	{
-		GenerateSkyLightVisibilityRays(GraphBuilder, SkyLightVisibilityRaysBuffer, SkyLightVisibilityRaysDimensions);
+		GenerateSkyLightVisibilityRays(GraphBuilder, Views[0], SkylightParameters, SkyLightData, SkyLightVisibilityRaysBuffer, SkyLightVisibilityRaysDimensions);
 	}
 	else
 	{
@@ -464,10 +423,6 @@ void FDeferredShadingSceneRenderer::RenderRayTracingSkyLight(
 	FRDGTextureUAV* SkyLightkUAV = GraphBuilder.CreateUAV(OutSkyLightTexture);
 	FRDGTextureUAV* RayDistanceUAV = GraphBuilder.CreateUAV(OutHitDistanceTexture);
 
-	// Fill Sky Light parameters
-	FSkyLightData SkyLightData;
-	SetupSkyLightParameters(*Scene, &SkyLightData);
-
 	// Fill Scene Texture parameters
 	FSceneTextureParameters SceneTextures = GetSceneTextureParameters(GraphBuilder);
 
@@ -487,6 +442,7 @@ void FDeferredShadingSceneRenderer::RenderRayTracingSkyLight(
 		FRayTracingSkyLightRGS::FParameters *PassParameters = GraphBuilder.AllocParameters<FRayTracingSkyLightRGS::FParameters>();
 		PassParameters->RWSkyOcclusionMaskUAV = SkyLightkUAV;
 		PassParameters->RWSkyOcclusionRayDistanceUAV = RayDistanceUAV;
+		PassParameters->SkyLightParameters = SkylightParameters;
 		PassParameters->SkyLightData = CreateUniformBufferImmediate(SkyLightData, EUniformBufferUsage::UniformBuffer_SingleDraw);
 		PassParameters->SkyLightVisibilityRaysData.SkyLightVisibilityRaysDimensions = SkyLightVisibilityRaysDimensions;
 		if (CVarRayTracingSkyLightDecoupleSampleGeneration.GetValueOnRenderThread() == 1)
