Renames Gen5TAA to Temporal Super Resolution

diff --git a/Engine/Source/Runtime/Engine/Classes/Engine/RendererSettings.h b/Engine/Source/Runtime/Engine/Classes/Engine/RendererSettings.h
--- a/Engine/Source/Runtime/Engine/Classes/Engine/RendererSettings.h
+++ b/Engine/Source/Runtime/Engine/Classes/Engine/RendererSettings.h
@@ -511,8 +511,8 @@ class ENGINE_API URendererSettings : public UDeveloperSettings
 
 	UPROPERTY(config, EditAnywhere, Category = DefaultSettings, meta = (
 		EditCondition = "DefaultFeatureAntiAliasing == EAntiAliasingMethod::AAM_TemporalAA",
-		ConsoleVariable = "r.TemporalAA.Algorithm", DisplayName = "Temporal Anti-Aliasing for 5th generation of consoles",
-		ToolTip = "Whether to enable UE5's new Temporal Anti-Aliasing method. This new technique is mainly built for handling Nanite amount of details on screen and higher upscale quality with Temporal Upsampling and ScreenPercentage/Dynamic Resolution on desktops and 5th generation of consoles."))
+		ConsoleVariable = "r.TemporalAA.Algorithm", DisplayName = "Temporal Super Resolution",
+		ToolTip = "Enable UE5's Temporal Super Resolution algorithm. This new technique is mainly built for handling Nanite amount of details on screen and higher output fidelity with lower ScreenPercentage/Dynamic Resolution on desktops and consoles."))
 	uint32 bUseGen5TAA : 1;
 
 	UPROPERTY(config, EditAnywhere, Category = DefaultSettings, meta = (
diff --git a/Engine/Source/Runtime/Renderer/Private/PostProcess/PostProcessSubsurface.cpp b/Engine/Source/Runtime/Renderer/Private/PostProcess/PostProcessSubsurface.cpp
--- a/Engine/Source/Runtime/Renderer/Private/PostProcess/PostProcessSubsurface.cpp
+++ b/Engine/Source/Runtime/Renderer/Private/PostProcess/PostProcessSubsurface.cpp
@@ -800,18 +800,18 @@ class FSubsurfaceRecombinePS : public FSubsurfaceShader
 		EMainTAAPassConfig MainTAAConfig = ITemporalUpscaler::GetMainTAAPassConfig(View);
 
 		// Low quality is really bad with modern temporal upscalers.
-		bool bAllowLowQuality = MainTAAConfig == EMainTAAPassConfig::Disabled || MainTAAConfig == EMainTAAPassConfig::Gen4;
+		bool bAllowLowQuality = MainTAAConfig == EMainTAAPassConfig::Disabled || MainTAAConfig == EMainTAAPassConfig::TAA;
 		if (!bAllowLowQuality)
 		{
 			return EQuality::High;
 		}
 
 		// Quality is forced to high when the CVar is set to 'auto' and TAA is NOT enabled.
-		// Gen4 TAA improves quality through temporal filtering and clamping box, making it less necessary to use
+		// TAA improves quality through temporal filtering and clamping box, making it less necessary to use
 		// high quality mode.
 		if (QualityCVar == -1)
 		{
-			if (MainTAAConfig == EMainTAAPassConfig::Gen4)
+			if (MainTAAConfig == EMainTAAPassConfig::TAA)
 			{
 				return EQuality::Low;
 			}
diff --git a/Engine/Source/Runtime/Renderer/Private/PostProcess/PostProcessing.cpp b/Engine/Source/Runtime/Renderer/Private/PostProcess/PostProcessing.cpp
--- a/Engine/Source/Runtime/Renderer/Private/PostProcess/PostProcessing.cpp
+++ b/Engine/Source/Runtime/Renderer/Private/PostProcess/PostProcessing.cpp
@@ -565,7 +565,7 @@ void AddPostProcessingPasses(FRDGBuilder& GraphBuilder, const FViewInfo& View, c
 			RDG_EVENT_SCOPE_CONDITIONAL(
 				GraphBuilder,
 				TAAConfig == EMainTAAPassConfig::ThirdParty,
-				"ThirdParty TAA %s %dx%d -> %dx%d",
+				"ThirdParty %s %dx%d -> %dx%d",
 				UpscalerToUse->GetDebugName(),
 				View.ViewRect.Width(), View.ViewRect.Height(),
 				View.GetSecondaryViewRectSize().X, View.GetSecondaryViewRectSize().Y);
diff --git a/Engine/Source/Runtime/Renderer/Private/PostProcess/TemporalAA.cpp b/Engine/Source/Runtime/Renderer/Private/PostProcess/TemporalAA.cpp
--- a/Engine/Source/Runtime/Renderer/Private/PostProcess/TemporalAA.cpp
+++ b/Engine/Source/Runtime/Renderer/Private/PostProcess/TemporalAA.cpp
@@ -22,8 +22,8 @@ const int32 GTemporalAATileSizeY = 8;
 TAutoConsoleVariable<int32> CVarTAAAlgorithm(
 	TEXT("r.TemporalAA.Algorithm"), 1,
 	TEXT("Algorithm to use for Temporal AA\n")
-	TEXT(" 0: Gen 4 TAAU (default)\n")
-	TEXT(" 1: Gen 5 TAAU (experimental)"),
+	TEXT(" 0: Temporal AA Upsampling\n")
+	TEXT(" 1: Temporal Super Resolution (default)"),
 	ECVF_RenderThreadSafe);
 
 TAutoConsoleVariable<float> CVarTemporalAAFilterSize(
@@ -121,7 +121,7 @@ inline bool DoesPlatformSupportGen5TAA(EShaderPlatform Platform)
 	return FDataDrivenShaderPlatformInfo::GetSupportsGen5TemporalAA(Platform);
 }
 
-BEGIN_SHADER_PARAMETER_STRUCT(FTAACommonParameters, )
+BEGIN_SHADER_PARAMETER_STRUCT(FTSRCommonParameters, )
 	SHADER_PARAMETER_STRUCT(FScreenPassTextureViewportParameters, InputInfo)
 	SHADER_PARAMETER_STRUCT(FScreenPassTextureViewportParameters, LowFrequencyInfo)
 	SHADER_PARAMETER_STRUCT(FScreenPassTextureViewportParameters, RejectionInfo)
@@ -137,27 +137,27 @@ BEGIN_SHADER_PARAMETER_STRUCT(FTAACommonParameters, )
 	SHADER_PARAMETER_STRUCT_REF(FViewUniformShaderParameters, ViewUniformBuffer)
 END_SHADER_PARAMETER_STRUCT()
 
-BEGIN_SHADER_PARAMETER_STRUCT(FTAAHistoryTextures, )
+BEGIN_SHADER_PARAMETER_STRUCT(FTSRHistoryTextures, )
 	SHADER_PARAMETER_RDG_TEXTURE_ARRAY(Texture2D, LowResTextures, [FTemporalAAHistory::kLowResRenderTargetCount])
 	SHADER_PARAMETER_RDG_TEXTURE_ARRAY(Texture2D, Textures, [FTemporalAAHistory::kRenderTargetCount])
 	SHADER_PARAMETER_RDG_TEXTURE_ARRAY(Texture2D, SuperResTextures, [FTemporalAAHistory::kSuperResRenderTargetCount])
 END_SHADER_PARAMETER_STRUCT()
 
-BEGIN_SHADER_PARAMETER_STRUCT(FTAAHistoryUAVs, )
+BEGIN_SHADER_PARAMETER_STRUCT(FTSRHistoryUAVs, )
 	SHADER_PARAMETER_RDG_TEXTURE_UAV_ARRAY(RWTexture2D, LowResTextures, [FTemporalAAHistory::kLowResRenderTargetCount])
 	SHADER_PARAMETER_RDG_TEXTURE_UAV_ARRAY(RWTexture2D, Textures, [FTemporalAAHistory::kRenderTargetCount])
 	SHADER_PARAMETER_RDG_TEXTURE_UAV_ARRAY(RWTexture2D, SuperResTextures, [FTemporalAAHistory::kSuperResRenderTargetCount])
 END_SHADER_PARAMETER_STRUCT()
 
-BEGIN_SHADER_PARAMETER_STRUCT(FTAAPrevHistoryParameters, )
+BEGIN_SHADER_PARAMETER_STRUCT(FTSRPrevHistoryParameters, )
 	SHADER_PARAMETER_STRUCT(FScreenPassTextureViewportParameters, PrevHistoryInfo)
 	SHADER_PARAMETER(FScreenTransform, ScreenPosToPrevHistoryBufferUV)
 	SHADER_PARAMETER(float, HistoryPreExposureCorrection)
 END_SHADER_PARAMETER_STRUCT()
 
-FTAAHistoryUAVs CreateUAVs(FRDGBuilder& GraphBuilder, const FTAAHistoryTextures& Textures)
+FTSRHistoryUAVs CreateUAVs(FRDGBuilder& GraphBuilder, const FTSRHistoryTextures& Textures)
 {
-	FTAAHistoryUAVs UAVs;
+	FTSRHistoryUAVs UAVs;
 	for (int32 i = 0; i < Textures.LowResTextures.Num(); i++)
 	{
 		UAVs.LowResTextures[i] = GraphBuilder.CreateUAV(Textures.LowResTextures[i]);
@@ -173,14 +173,14 @@ FTAAHistoryUAVs CreateUAVs(FRDGBuilder& GraphBuilder, const FTAAHistoryTextures&
 	return UAVs;
 }
 
-class FTAAGen5Shader : public FGlobalShader
+class FTemporalSuperResolutionShader : public FGlobalShader
 {
 public:
-	FTAAGen5Shader(const ShaderMetaType::CompiledShaderInitializerType& Initializer)
+	FTemporalSuperResolutionShader(const ShaderMetaType::CompiledShaderInitializerType& Initializer)
 		: FGlobalShader(Initializer)
 	{ }
 
-	FTAAGen5Shader()
+	FTemporalSuperResolutionShader()
 	{ }
 
 	static bool ShouldCompilePermutation(const FGlobalShaderPermutationParameters& Parameters)
@@ -192,7 +192,7 @@ public:
 	{
 		OutEnvironment.CompilerFlags.Add(CFLAG_AllowRealTypes);
 	}
-}; // class FTAAGen5Shader
+}; // class FTemporalSuperResolutionShader
 
 class FTAAStandaloneCS : public FGlobalShader
 {
@@ -362,26 +362,26 @@ class FTAAStandaloneCS : public FGlobalShader
 	}
 }; // class FTAAStandaloneCS
 
-class FTAAClearPrevTexturesCS : public FTAAGen5Shader
+class FTSRClearPrevTexturesCS : public FTemporalSuperResolutionShader
 {
-	DECLARE_GLOBAL_SHADER(FTAAClearPrevTexturesCS);
-	SHADER_USE_PARAMETER_STRUCT(FTAAClearPrevTexturesCS, FTAAGen5Shader);
+	DECLARE_GLOBAL_SHADER(FTSRClearPrevTexturesCS);
+	SHADER_USE_PARAMETER_STRUCT(FTSRClearPrevTexturesCS, FTemporalSuperResolutionShader);
 
 	BEGIN_SHADER_PARAMETER_STRUCT(FParameters, )
-		SHADER_PARAMETER_STRUCT_INCLUDE(FTAACommonParameters, CommonParameters)
+		SHADER_PARAMETER_STRUCT_INCLUDE(FTSRCommonParameters, CommonParameters)
 
 		SHADER_PARAMETER_RDG_TEXTURE_UAV(RWTexture2D, PrevUseCountOutput)
 		SHADER_PARAMETER_RDG_TEXTURE_UAV(RWTexture2D, PrevClosestDepthOutput)
 	END_SHADER_PARAMETER_STRUCT()
-}; // class FTAAClearPrevTexturesCS
+}; // class FTSRClearPrevTexturesCS
 
-class FTAADilateVelocityCS : public FTAAGen5Shader
+class FTSRDilateVelocityCS : public FTemporalSuperResolutionShader
 {
-	DECLARE_GLOBAL_SHADER(FTAADilateVelocityCS);
-	SHADER_USE_PARAMETER_STRUCT(FTAADilateVelocityCS, FTAAGen5Shader);
+	DECLARE_GLOBAL_SHADER(FTSRDilateVelocityCS);
+	SHADER_USE_PARAMETER_STRUCT(FTSRDilateVelocityCS, FTemporalSuperResolutionShader);
 
 	BEGIN_SHADER_PARAMETER_STRUCT(FParameters, )
-		SHADER_PARAMETER_STRUCT_INCLUDE(FTAACommonParameters, CommonParameters)
+		SHADER_PARAMETER_STRUCT_INCLUDE(FTSRCommonParameters, CommonParameters)
 
 		SHADER_PARAMETER(FVector2D, PrevOutputBufferUVMin)
 		SHADER_PARAMETER(FVector2D, PrevOutputBufferUVMax)
@@ -396,19 +396,19 @@ class FTAADilateVelocityCS : public FTAAGen5Shader
 		SHADER_PARAMETER_RDG_TEXTURE_UAV(RWTexture2D, ParallaxFactorOutput)
 		SHADER_PARAMETER_RDG_TEXTURE_UAV(RWTexture2D, DebugOutput)
 	END_SHADER_PARAMETER_STRUCT()
-}; // class FTAADilateVelocityCS
+}; // class FTSRDilateVelocityCS
 
-class FTAADecimateHistoryCS : public FTAAGen5Shader
+class FTSRDecimateHistoryCS : public FTemporalSuperResolutionShader
 {
-	DECLARE_GLOBAL_SHADER(FTAADecimateHistoryCS);
-	SHADER_USE_PARAMETER_STRUCT(FTAADecimateHistoryCS, FTAAGen5Shader);
+	DECLARE_GLOBAL_SHADER(FTSRDecimateHistoryCS);
+	SHADER_USE_PARAMETER_STRUCT(FTSRDecimateHistoryCS, FTemporalSuperResolutionShader);
 
 	class FOutputHalfRes : SHADER_PERMUTATION_BOOL("DIM_OUTPUT_HALF_RES");
 
 	using FPermutationDomain = TShaderPermutationDomain<FOutputHalfRes>;
 
 	BEGIN_SHADER_PARAMETER_STRUCT(FParameters, )
-		SHADER_PARAMETER_STRUCT_INCLUDE(FTAACommonParameters, CommonParameters)
+		SHADER_PARAMETER_STRUCT_INCLUDE(FTSRCommonParameters, CommonParameters)
 		SHADER_PARAMETER(FMatrix, RotationalClipToPrevClip)
 		SHADER_PARAMETER(FVector, OutputQuantizationError)
 		SHADER_PARAMETER(float, WorldDepthToPixelWorldRadius)
@@ -420,8 +420,8 @@ class FTAADecimateHistoryCS : public FTAAGen5Shader
 		SHADER_PARAMETER_RDG_TEXTURE(Texture2D, PrevClosestDepthTexture)
 		SHADER_PARAMETER_RDG_TEXTURE(Texture2D, ParallaxFactorTexture)
 
-		SHADER_PARAMETER_STRUCT_INCLUDE(FTAAPrevHistoryParameters, PrevHistoryParameters)
-		SHADER_PARAMETER_STRUCT(FTAAHistoryTextures, PrevHistory)
+		SHADER_PARAMETER_STRUCT_INCLUDE(FTSRPrevHistoryParameters, PrevHistoryParameters)
+		SHADER_PARAMETER_STRUCT(FTSRHistoryTextures, PrevHistory)
 
 		SHADER_PARAMETER_RDG_TEXTURE_UAV(RWTexture2D, HalfResSceneColorOutput)
 		SHADER_PARAMETER_RDG_TEXTURE_UAV(RWTexture2D, HalfResPredictionSceneColorOutput)
@@ -431,15 +431,15 @@ class FTAADecimateHistoryCS : public FTAAGen5Shader
 		SHADER_PARAMETER_RDG_TEXTURE_UAV(RWTexture2D, InterferenceSeedOutput)
 		SHADER_PARAMETER_RDG_TEXTURE_UAV(RWTexture2D, DebugOutput)
 	END_SHADER_PARAMETER_STRUCT()
-}; // class FTAADecimateHistoryCS
+}; // class FTSRDecimateHistoryCS
 
-class FTAADetectInterferenceCS : public FTAAGen5Shader
+class FTSRDetectInterferenceCS : public FTemporalSuperResolutionShader
 {
-	DECLARE_GLOBAL_SHADER(FTAADetectInterferenceCS);
-	SHADER_USE_PARAMETER_STRUCT(FTAADetectInterferenceCS, FTAAGen5Shader);
+	DECLARE_GLOBAL_SHADER(FTSRDetectInterferenceCS);
+	SHADER_USE_PARAMETER_STRUCT(FTSRDetectInterferenceCS, FTemporalSuperResolutionShader);
 	
 	BEGIN_SHADER_PARAMETER_STRUCT(FParameters, )
-		SHADER_PARAMETER_STRUCT_INCLUDE(FTAACommonParameters, CommonParameters)
+		SHADER_PARAMETER_STRUCT_INCLUDE(FTSRCommonParameters, CommonParameters)
 
 		SHADER_PARAMETER_RDG_TEXTURE(Texture2D, InputSceneColorTexture)
 		SHADER_PARAMETER_RDG_TEXTURE(Texture2D, DilatedVelocityTexture)
@@ -447,22 +447,22 @@ class FTAADetectInterferenceCS : public FTAAGen5Shader
 		SHADER_PARAMETER_RDG_TEXTURE(Texture2D, ParallaxRejectionMaskTexture)
 		SHADER_PARAMETER_RDG_TEXTURE(Texture2D, InterferenceSeedTexture)
 
-		SHADER_PARAMETER_STRUCT_INCLUDE(FTAAPrevHistoryParameters, PrevHistoryParameters)
-		SHADER_PARAMETER_STRUCT(FTAAHistoryTextures, PrevHistory)
+		SHADER_PARAMETER_STRUCT_INCLUDE(FTSRPrevHistoryParameters, PrevHistoryParameters)
+		SHADER_PARAMETER_STRUCT(FTSRHistoryTextures, PrevHistory)
 
-		SHADER_PARAMETER_STRUCT(FTAAHistoryUAVs, HistoryOutput)
+		SHADER_PARAMETER_STRUCT(FTSRHistoryUAVs, HistoryOutput)
 		SHADER_PARAMETER_RDG_TEXTURE_UAV(RWTexture2D, InterferenceWeightOutput)
 		SHADER_PARAMETER_RDG_TEXTURE_UAV(RWTexture2D, DebugOutput)
 	END_SHADER_PARAMETER_STRUCT()
-}; // class FTAADetectInterferenceCS
+}; // class FTSRDetectInterferenceCS
 
-class FTAAFilterFrequenciesCS : public FTAAGen5Shader
+class FTSRFilterFrequenciesCS : public FTemporalSuperResolutionShader
 {
-	DECLARE_GLOBAL_SHADER(FTAAFilterFrequenciesCS);
-	SHADER_USE_PARAMETER_STRUCT(FTAAFilterFrequenciesCS, FTAAGen5Shader);
+	DECLARE_GLOBAL_SHADER(FTSRFilterFrequenciesCS);
+	SHADER_USE_PARAMETER_STRUCT(FTSRFilterFrequenciesCS, FTemporalSuperResolutionShader);
 
 	BEGIN_SHADER_PARAMETER_STRUCT(FParameters, )
-		SHADER_PARAMETER_STRUCT_INCLUDE(FTAACommonParameters, CommonParameters)
+		SHADER_PARAMETER_STRUCT_INCLUDE(FTSRCommonParameters, CommonParameters)
 		SHADER_PARAMETER(FVector, OutputQuantizationError)
 
 		SHADER_PARAMETER_RDG_TEXTURE(Texture2D, InputTexture)
@@ -474,15 +474,15 @@ class FTAAFilterFrequenciesCS : public FTAAGen5Shader
 		SHADER_PARAMETER_RDG_TEXTURE_UAV(RWTexture2D, FilteredPredictionSceneColorOutput)
 		SHADER_PARAMETER_RDG_TEXTURE_UAV(RWTexture2D, DebugOutput)
 	END_SHADER_PARAMETER_STRUCT()
-}; // class FTAAFilterFrequenciesCS
+}; // class FTSRFilterFrequenciesCS
 
-class FTAACompareHistoryCS : public FTAAGen5Shader
+class FTSRCompareHistoryCS : public FTemporalSuperResolutionShader
 {
-	DECLARE_GLOBAL_SHADER(FTAACompareHistoryCS);
-	SHADER_USE_PARAMETER_STRUCT(FTAACompareHistoryCS, FTAAGen5Shader);
+	DECLARE_GLOBAL_SHADER(FTSRCompareHistoryCS);
+	SHADER_USE_PARAMETER_STRUCT(FTSRCompareHistoryCS, FTemporalSuperResolutionShader);
 
 	BEGIN_SHADER_PARAMETER_STRUCT(FParameters, )
-		SHADER_PARAMETER_STRUCT_INCLUDE(FTAACommonParameters, CommonParameters)
+		SHADER_PARAMETER_STRUCT_INCLUDE(FTSRCommonParameters, CommonParameters)
 		SHADER_PARAMETER_RDG_TEXTURE(Texture2D, ParallaxRejectionMaskTexture)
 		SHADER_PARAMETER_RDG_TEXTURE(Texture2D, FilteredInputTexture)
 		SHADER_PARAMETER_RDG_TEXTURE(Texture2D, FilteredPredictionSceneColorTexture)
@@ -490,61 +490,61 @@ class FTAACompareHistoryCS : public FTAAGen5Shader
 		SHADER_PARAMETER_RDG_TEXTURE_UAV(RWTexture2D, HistoryRejectionOutput)
 		SHADER_PARAMETER_RDG_TEXTURE_UAV(RWTexture2D, DebugOutput)
 	END_SHADER_PARAMETER_STRUCT()
-}; // class FTAACompareHistoryCS
+}; // class FTSRCompareHistoryCS
 
-class FTAAPostfilterRejectionCS : public FTAAGen5Shader
+class FTSRPostfilterRejectionCS : public FTemporalSuperResolutionShader
 {
-	DECLARE_GLOBAL_SHADER(FTAAPostfilterRejectionCS);
-	SHADER_USE_PARAMETER_STRUCT(FTAAPostfilterRejectionCS, FTAAGen5Shader);
+	DECLARE_GLOBAL_SHADER(FTSRPostfilterRejectionCS);
+	SHADER_USE_PARAMETER_STRUCT(FTSRPostfilterRejectionCS, FTemporalSuperResolutionShader);
 
 	BEGIN_SHADER_PARAMETER_STRUCT(FParameters, )
-		SHADER_PARAMETER_STRUCT_INCLUDE(FTAACommonParameters, CommonParameters)
+		SHADER_PARAMETER_STRUCT_INCLUDE(FTSRCommonParameters, CommonParameters)
 		SHADER_PARAMETER_RDG_TEXTURE(Texture2D, HistoryRejectionTexture)
 		SHADER_PARAMETER_RDG_TEXTURE_UAV(RWTexture2D, HistoryRejectionOutput)
 		SHADER_PARAMETER_RDG_TEXTURE_UAV(RWTexture2D, DebugOutput)
 	END_SHADER_PARAMETER_STRUCT()
-}; // class FTAAPostfilterRejectionCS
+}; // class FTSRPostfilterRejectionCS
 
-class FTAADilateRejectionCS : public FTAAGen5Shader
+class FTSRDilateRejectionCS : public FTemporalSuperResolutionShader
 {
-	DECLARE_GLOBAL_SHADER(FTAADilateRejectionCS);
-	SHADER_USE_PARAMETER_STRUCT(FTAADilateRejectionCS, FTAAGen5Shader);
+	DECLARE_GLOBAL_SHADER(FTSRDilateRejectionCS);
+	SHADER_USE_PARAMETER_STRUCT(FTSRDilateRejectionCS, FTemporalSuperResolutionShader);
 
 	BEGIN_SHADER_PARAMETER_STRUCT(FParameters, )
-		SHADER_PARAMETER_STRUCT_INCLUDE(FTAACommonParameters, CommonParameters)
+		SHADER_PARAMETER_STRUCT_INCLUDE(FTSRCommonParameters, CommonParameters)
 		SHADER_PARAMETER_RDG_TEXTURE(Texture2D, HistoryRejectionTexture)
 		SHADER_PARAMETER_RDG_TEXTURE_UAV(RWTexture2D, DilatedHistoryRejectionOutput)
 		SHADER_PARAMETER_RDG_TEXTURE_UAV(RWTexture2D, DebugOutput)
 	END_SHADER_PARAMETER_STRUCT()
-}; // class FTAADilateRejectionCS
+}; // class FTSRDilateRejectionCS
 
-class FTAAUpdateSuperResHistoryCS : public FTAAGen5Shader
+class FTSRUpdateSuperResHistoryCS : public FTemporalSuperResolutionShader
 {
-	DECLARE_GLOBAL_SHADER(FTAAUpdateSuperResHistoryCS);
-	SHADER_USE_PARAMETER_STRUCT(FTAAUpdateSuperResHistoryCS, FTAAGen5Shader);
+	DECLARE_GLOBAL_SHADER(FTSRUpdateSuperResHistoryCS);
+	SHADER_USE_PARAMETER_STRUCT(FTSRUpdateSuperResHistoryCS, FTemporalSuperResolutionShader);
 
 	BEGIN_SHADER_PARAMETER_STRUCT(FParameters, )
-		SHADER_PARAMETER_STRUCT_INCLUDE(FTAACommonParameters, CommonParameters)
+		SHADER_PARAMETER_STRUCT_INCLUDE(FTSRCommonParameters, CommonParameters)
 		SHADER_PARAMETER_RDG_TEXTURE(Texture2D, InputSceneColorTexture)
 
 		SHADER_PARAMETER_RDG_TEXTURE(Texture2D, DilatedVelocityTexture)
 		SHADER_PARAMETER_RDG_TEXTURE(Texture2D, ParallaxRejectionMaskTexture)
 
-		SHADER_PARAMETER_STRUCT_INCLUDE(FTAAPrevHistoryParameters, PrevHistoryParameters)
-		SHADER_PARAMETER_STRUCT(FTAAHistoryTextures, PrevHistory)
+		SHADER_PARAMETER_STRUCT_INCLUDE(FTSRPrevHistoryParameters, PrevHistoryParameters)
+		SHADER_PARAMETER_STRUCT(FTSRHistoryTextures, PrevHistory)
 
-		SHADER_PARAMETER_STRUCT(FTAAHistoryUAVs, HistoryOutput)
+		SHADER_PARAMETER_STRUCT(FTSRHistoryUAVs, HistoryOutput)
 		SHADER_PARAMETER_RDG_TEXTURE_UAV(RWTexture2D, DebugOutput)
 	END_SHADER_PARAMETER_STRUCT()
-}; // class FTAAUpdateSuperResHistoryCS
+}; // class FTSRUpdateSuperResHistoryCS
 
-class FTAAUpdateHistoryCS : public FTAAGen5Shader
+class FTSRUpdateHistoryCS : public FTemporalSuperResolutionShader
 {
-	DECLARE_GLOBAL_SHADER(FTAAUpdateHistoryCS);
-	SHADER_USE_PARAMETER_STRUCT(FTAAUpdateHistoryCS, FTAAGen5Shader);
+	DECLARE_GLOBAL_SHADER(FTSRUpdateHistoryCS);
+	SHADER_USE_PARAMETER_STRUCT(FTSRUpdateHistoryCS, FTemporalSuperResolutionShader);
 
 	BEGIN_SHADER_PARAMETER_STRUCT(FParameters, )
-		SHADER_PARAMETER_STRUCT_INCLUDE(FTAACommonParameters, CommonParameters)
+		SHADER_PARAMETER_STRUCT_INCLUDE(FTSRCommonParameters, CommonParameters)
 		SHADER_PARAMETER_RDG_TEXTURE(Texture2D, InputSceneColorTexture)
 		SHADER_PARAMETER_RDG_TEXTURE_SRV(Texture2D, InputSceneStencilTexture)
 		SHADER_PARAMETER_RDG_TEXTURE(Texture2D, HistoryRejectionTexture)
@@ -557,49 +557,49 @@ class FTAAUpdateHistoryCS : public FTAAGen5Shader
 		SHADER_PARAMETER(FScreenTransform, HistoryPixelPosToPPCo)
 		SHADER_PARAMETER(FVector, HistoryQuantizationError)
 
-		SHADER_PARAMETER_STRUCT_INCLUDE(FTAAPrevHistoryParameters, PrevHistoryParameters)
-		SHADER_PARAMETER_STRUCT(FTAAHistoryTextures, PrevHistory)
+		SHADER_PARAMETER_STRUCT_INCLUDE(FTSRPrevHistoryParameters, PrevHistoryParameters)
+		SHADER_PARAMETER_STRUCT(FTSRHistoryTextures, PrevHistory)
 
-		SHADER_PARAMETER_STRUCT(FTAAHistoryUAVs, HistoryOutput)
+		SHADER_PARAMETER_STRUCT(FTSRHistoryUAVs, HistoryOutput)
 		SHADER_PARAMETER_RDG_TEXTURE_UAV(RWTexture2D, SceneColorOutput)
 		SHADER_PARAMETER_RDG_TEXTURE_UAV(RWTexture2D, DebugOutput)
 	END_SHADER_PARAMETER_STRUCT()
 
 	static void ModifyCompilationEnvironment(const FGlobalShaderPermutationParameters& Parameters, FShaderCompilerEnvironment& OutEnvironment)
 	{
-		FTAAGen5Shader::ModifyCompilationEnvironment(Parameters, OutEnvironment);
+		FTemporalSuperResolutionShader::ModifyCompilationEnvironment(Parameters, OutEnvironment);
 		OutEnvironment.CompilerFlags.Add(CFLAG_Wave32);
 	}
-}; // class FTAAUpdateHistoryCS
+}; // class FTSRUpdateHistoryCS
 
 IMPLEMENT_GLOBAL_SHADER(FTAAStandaloneCS, "/Engine/Private/TemporalAA/TAAStandalone.usf", "MainCS", SF_Compute);
-IMPLEMENT_GLOBAL_SHADER(FTAAClearPrevTexturesCS, "/Engine/Private/TemporalAA/TAAClearPrevTextures.usf", "MainCS", SF_Compute);
-IMPLEMENT_GLOBAL_SHADER(FTAADilateVelocityCS, "/Engine/Private/TemporalAA/TAADilateVelocity.usf", "MainCS", SF_Compute);
-IMPLEMENT_GLOBAL_SHADER(FTAADecimateHistoryCS, "/Engine/Private/TemporalAA/TAADecimateHistory.usf", "MainCS", SF_Compute);
-IMPLEMENT_GLOBAL_SHADER(FTAADetectInterferenceCS, "/Engine/Private/TemporalAA/TAADetectInterference.usf", "MainCS", SF_Compute);
-IMPLEMENT_GLOBAL_SHADER(FTAAFilterFrequenciesCS, "/Engine/Private/TemporalAA/TAAFilterFrequencies.usf", "MainCS", SF_Compute);
-IMPLEMENT_GLOBAL_SHADER(FTAACompareHistoryCS, "/Engine/Private/TemporalAA/TAACompareHistory.usf", "MainCS", SF_Compute);
-IMPLEMENT_GLOBAL_SHADER(FTAAPostfilterRejectionCS, "/Engine/Private/TemporalAA/TAAPostfilterRejection.usf", "MainCS", SF_Compute);
-IMPLEMENT_GLOBAL_SHADER(FTAADilateRejectionCS, "/Engine/Private/TemporalAA/TAADilateRejection.usf", "MainCS", SF_Compute);
-IMPLEMENT_GLOBAL_SHADER(FTAAUpdateSuperResHistoryCS, "/Engine/Private/TemporalAA/TAAUpdateSuperResHistory.usf", "MainCS", SF_Compute);
-IMPLEMENT_GLOBAL_SHADER(FTAAUpdateHistoryCS, "/Engine/Private/TemporalAA/TAAUpdateHistory.usf", "MainCS", SF_Compute);
+IMPLEMENT_GLOBAL_SHADER(FTSRClearPrevTexturesCS, "/Engine/Private/TemporalAA/TAAClearPrevTextures.usf", "MainCS", SF_Compute);
+IMPLEMENT_GLOBAL_SHADER(FTSRDilateVelocityCS, "/Engine/Private/TemporalAA/TAADilateVelocity.usf", "MainCS", SF_Compute);
+IMPLEMENT_GLOBAL_SHADER(FTSRDecimateHistoryCS, "/Engine/Private/TemporalAA/TAADecimateHistory.usf", "MainCS", SF_Compute);
+IMPLEMENT_GLOBAL_SHADER(FTSRDetectInterferenceCS, "/Engine/Private/TemporalAA/TAADetectInterference.usf", "MainCS", SF_Compute);
+IMPLEMENT_GLOBAL_SHADER(FTSRFilterFrequenciesCS, "/Engine/Private/TemporalAA/TAAFilterFrequencies.usf", "MainCS", SF_Compute);
+IMPLEMENT_GLOBAL_SHADER(FTSRCompareHistoryCS, "/Engine/Private/TemporalAA/TAACompareHistory.usf", "MainCS", SF_Compute);
+IMPLEMENT_GLOBAL_SHADER(FTSRPostfilterRejectionCS, "/Engine/Private/TemporalAA/TAAPostfilterRejection.usf", "MainCS", SF_Compute);
+IMPLEMENT_GLOBAL_SHADER(FTSRDilateRejectionCS, "/Engine/Private/TemporalAA/TAADilateRejection.usf", "MainCS", SF_Compute);
+IMPLEMENT_GLOBAL_SHADER(FTSRUpdateSuperResHistoryCS, "/Engine/Private/TemporalAA/TAAUpdateSuperResHistory.usf", "MainCS", SF_Compute);
+IMPLEMENT_GLOBAL_SHADER(FTSRUpdateHistoryCS, "/Engine/Private/TemporalAA/TAAUpdateHistory.usf", "MainCS", SF_Compute);
 
 #if COMPILE_TAA_DEBUG_PASSES
 
-class FTAADebugHistoryCS : public FTAAGen5Shader
+class FTSRDebugHistoryCS : public FTemporalSuperResolutionShader
 {
-	DECLARE_GLOBAL_SHADER(FTAADebugHistoryCS);
-	SHADER_USE_PARAMETER_STRUCT(FTAADebugHistoryCS, FTAAGen5Shader);
+	DECLARE_GLOBAL_SHADER(FTSRDebugHistoryCS);
+	SHADER_USE_PARAMETER_STRUCT(FTSRDebugHistoryCS, FTemporalSuperResolutionShader);
 
 	BEGIN_SHADER_PARAMETER_STRUCT(FParameters, )
-		SHADER_PARAMETER_STRUCT_INCLUDE(FTAACommonParameters, CommonParameters)
-		SHADER_PARAMETER_STRUCT(FTAAHistoryTextures, History)
-		SHADER_PARAMETER_STRUCT(FTAAHistoryTextures, PrevHistory)
+		SHADER_PARAMETER_STRUCT_INCLUDE(FTSRCommonParameters, CommonParameters)
+		SHADER_PARAMETER_STRUCT(FTSRHistoryTextures, History)
+		SHADER_PARAMETER_STRUCT(FTSRHistoryTextures, PrevHistory)
 		SHADER_PARAMETER_RDG_TEXTURE_UAV(RWTexture2D, DebugOutput)
 	END_SHADER_PARAMETER_STRUCT()
-}; // class FTAADebugHistoryCS
+}; // class FTSRDebugHistoryCS
 
-IMPLEMENT_GLOBAL_SHADER(FTAADebugHistoryCS, "/Engine/Private/TemporalAA/TAADebugHistory.usf", "MainCS", SF_Compute);
+IMPLEMENT_GLOBAL_SHADER(FTSRDebugHistoryCS, "/Engine/Private/TemporalAA/TAADebugHistory.usf", "MainCS", SF_Compute);
 
 #endif
 
@@ -701,6 +701,7 @@ void SetupSampleWeightParameters(FTAAStandaloneCS::FParameters* OutTAAParameters
 }
 
 DECLARE_GPU_STAT(TAA)
+DECLARE_GPU_STAT(TemporalSuperResolution)
 
 const TCHAR* const kTAAOutputNames[] = {
 	TEXT("TAA.History"),
@@ -1128,7 +1129,7 @@ FTAAOutputs AddTemporalAAPass(
 	return Outputs;
 } // AddTemporalAAPass()
 
-static void AddGen5MainTemporalAAPasses(
+static void AddTemporalSuperResolutionPasses(
 	FRDGBuilder& GraphBuilder,
 	const FViewInfo& View,
 	const ITemporalUpscaler::FPassInputs& PassInputs,
@@ -1202,14 +1203,14 @@ static void AddGen5MainTemporalAAPasses(
 			FMath::Max(InputExtent.Y, QuantizedHistoryViewSize.Y));
 	}
 
-	RDG_EVENT_SCOPE(GraphBuilder, "TAAU Gen5 %dx%d -> %dx%d", InputRect.Width(), InputRect.Height(), OutputRect.Width(), OutputRect.Height());
-	RDG_GPU_STAT_SCOPE(GraphBuilder, TAA);
+	RDG_EVENT_SCOPE(GraphBuilder, "TemporalSuperResolution %dx%d -> %dx%d", InputRect.Width(), InputRect.Height(), OutputRect.Width(), OutputRect.Height());
+	RDG_GPU_STAT_SCOPE(GraphBuilder, TemporalSuperResolution);
 
 	FRDGTextureRef BlackUintDummy = GSystemTextures.GetZeroUIntDummy(GraphBuilder);
 	FRDGTextureRef BlackDummy = GraphBuilder.RegisterExternalTexture(GSystemTextures.BlackDummy);
 	FRDGTextureRef WhiteDummy = GraphBuilder.RegisterExternalTexture(GSystemTextures.WhiteDummy);
 
-	FTAACommonParameters CommonParameters;
+	FTSRCommonParameters CommonParameters;
 	{
 		CommonParameters.InputInfo = GetScreenPassTextureViewportParameters(FScreenPassTextureViewport(
 			InputExtent, InputRect));
@@ -1255,19 +1256,19 @@ static void AddGen5MainTemporalAAPasses(
 				FClearValueBinding::None,
 				/* InFlags = */ TexCreate_ShaderResource | TexCreate_UAV);
 
-			PrevUseCountTexture = GraphBuilder.CreateTexture(Desc, TEXT("TAA.PrevUseCountTexture"));
-			PrevClosestDepthTexture = GraphBuilder.CreateTexture(Desc, TEXT("TAA.PrevClosestDepthTexture"));
+			PrevUseCountTexture = GraphBuilder.CreateTexture(Desc, TEXT("TSR.PrevUseCountTexture"));
+			PrevClosestDepthTexture = GraphBuilder.CreateTexture(Desc, TEXT("TSR.PrevClosestDepthTexture"));
 		}
 
-		FTAAClearPrevTexturesCS::FParameters* PassParameters = GraphBuilder.AllocParameters<FTAAClearPrevTexturesCS::FParameters>();
+		FTSRClearPrevTexturesCS::FParameters* PassParameters = GraphBuilder.AllocParameters<FTSRClearPrevTexturesCS::FParameters>();
 		PassParameters->CommonParameters = CommonParameters;
 		PassParameters->PrevUseCountOutput = GraphBuilder.CreateUAV(PrevUseCountTexture);
 		PassParameters->PrevClosestDepthOutput = GraphBuilder.CreateUAV(PrevClosestDepthTexture);
 
-		TShaderMapRef<FTAAClearPrevTexturesCS> ComputeShader(View.ShaderMap);
+		TShaderMapRef<FTSRClearPrevTexturesCS> ComputeShader(View.ShaderMap);
 		FComputeShaderUtils::AddPass(
 			GraphBuilder,
-			RDG_EVENT_NAME("TAA ClearPrevTextures %dx%d", InputRect.Width(), InputRect.Height()),
+			RDG_EVENT_NAME("TSR ClearPrevTextures %dx%d", InputRect.Width(), InputRect.Height()),
 			ComputeShader,
 			PassParameters,
 			FComputeShaderUtils::GetGroupCount(InputRect.Size(), 8));
@@ -1285,16 +1286,16 @@ static void AddGen5MainTemporalAAPasses(
 				FClearValueBinding::None,
 				/* InFlags = */ TexCreate_ShaderResource | TexCreate_UAV);
 
-			DilatedVelocityTexture = GraphBuilder.CreateTexture(Desc, TEXT("TAA.DilatedVelocity"));
+			DilatedVelocityTexture = GraphBuilder.CreateTexture(Desc, TEXT("TSR.DilatedVelocity"));
 
 			Desc.Format = PF_R16F;
-			ClosestDepthTexture = GraphBuilder.CreateTexture(Desc, TEXT("TAA.ClosestDepthTexture"));
+			ClosestDepthTexture = GraphBuilder.CreateTexture(Desc, TEXT("TSR.ClosestDepthTexture"));
 
 			Desc.Format = PF_R8_UINT;
-			ParallaxFactorTexture = GraphBuilder.CreateTexture(Desc, TEXT("TAA.ParallaxFactor"));
+			ParallaxFactorTexture = GraphBuilder.CreateTexture(Desc, TEXT("TSR.ParallaxFactor"));
 		}
 
-		FTAADilateVelocityCS::FParameters* PassParameters = GraphBuilder.AllocParameters<FTAADilateVelocityCS::FParameters>();
+		FTSRDilateVelocityCS::FParameters* PassParameters = GraphBuilder.AllocParameters<FTSRDilateVelocityCS::FParameters>();
 		PassParameters->CommonParameters = CommonParameters;
 		PassParameters->PrevOutputBufferUVMin = CommonParameters.InputInfo.UVViewportBilinearMin - CommonParameters.InputInfo.ExtentInverse;
 		PassParameters->PrevOutputBufferUVMax = CommonParameters.InputInfo.UVViewportBilinearMax + CommonParameters.InputInfo.ExtentInverse;
@@ -1305,20 +1306,20 @@ static void AddGen5MainTemporalAAPasses(
 		PassParameters->PrevUseCountOutput = GraphBuilder.CreateUAV(PrevUseCountTexture);
 		PassParameters->PrevClosestDepthOutput = GraphBuilder.CreateUAV(PrevClosestDepthTexture);
 		PassParameters->ParallaxFactorOutput = GraphBuilder.CreateUAV(ParallaxFactorTexture);
-		PassParameters->DebugOutput = CreateDebugUAV(InputExtent, TEXT("Debug.TAA.DilateVelocity"));
+		PassParameters->DebugOutput = CreateDebugUAV(InputExtent, TEXT("Debug.TSR.DilateVelocity"));
 
-		TShaderMapRef<FTAADilateVelocityCS> ComputeShader(View.ShaderMap);
+		TShaderMapRef<FTSRDilateVelocityCS> ComputeShader(View.ShaderMap);
 		FComputeShaderUtils::AddPass(
 			GraphBuilder,
-			RDG_EVENT_NAME("TAA DilateVelocity %dx%d", InputRect.Width(), InputRect.Height()),
+			RDG_EVENT_NAME("TSR DilateVelocity %dx%d", InputRect.Width(), InputRect.Height()),
 			ComputeShader,
 			PassParameters,
 			FComputeShaderUtils::GetGroupCount(InputRect.Size(), 8));
 	}
 
 	// Setup the previous frame history
-	FTAAPrevHistoryParameters PrevHistoryParameters;
-	FTAAHistoryTextures PrevHistory;
+	FTSRPrevHistoryParameters PrevHistoryParameters;
+	FTSRHistoryTextures PrevHistory;
 	{
 		FScreenPassTextureViewport PrevHistoryViewport(InputHistory.ReferenceBufferSize, InputHistory.ViewportRect);
 		if (bCameraCut)
@@ -1374,7 +1375,7 @@ static void AddGen5MainTemporalAAPasses(
 	}
 
 
-	FTAAHistoryTextures History;
+	FTSRHistoryTextures History;
 	{
 		FRDGTextureDesc LowResDesc = FRDGTextureDesc::Create2D(
 			InputExtent,
@@ -1382,7 +1383,7 @@ static void AddGen5MainTemporalAAPasses(
 			FClearValueBinding::None,
 			TexCreate_ShaderResource | TexCreate_UAV);
 
-		History.LowResTextures[0] = GraphBuilder.CreateTexture(LowResDesc, TEXT("TAA.History.LowResMetadata[0]"));
+		History.LowResTextures[0] = GraphBuilder.CreateTexture(LowResDesc, TEXT("TSR.History.LowResMetadata[0]"));
 	}
 
 	{
@@ -1392,14 +1393,14 @@ static void AddGen5MainTemporalAAPasses(
 			FClearValueBinding::None,
 			TexCreate_ShaderResource | TexCreate_UAV);
 
-		History.Textures[0] = GraphBuilder.CreateTexture(Desc, TEXT("TAA.History.LowFrequencies"));
-		History.Textures[1] = GraphBuilder.CreateTexture(Desc, TEXT("TAA.History.HighFrequencies"));
+		History.Textures[0] = GraphBuilder.CreateTexture(Desc, TEXT("TSR.History.LowFrequencies"));
+		History.Textures[1] = GraphBuilder.CreateTexture(Desc, TEXT("TSR.History.HighFrequencies"));
 
 		Desc.Format = PF_R8G8;
-		History.Textures[2] = GraphBuilder.CreateTexture(Desc, TEXT("TAA.History.Metadata"));
+		History.Textures[2] = GraphBuilder.CreateTexture(Desc, TEXT("TSR.History.Metadata"));
 
 		Desc.Format = PF_R16_UINT;
-		History.Textures[3] = GraphBuilder.CreateTexture(Desc, TEXT("TAA.History.SubpixelInfo"));
+		History.Textures[3] = GraphBuilder.CreateTexture(Desc, TEXT("TSR.History.SubpixelInfo"));
 	}
 
 	{
@@ -1409,10 +1410,10 @@ static void AddGen5MainTemporalAAPasses(
 			FClearValueBinding::None,
 			TexCreate_ShaderResource | TexCreate_UAV);
 
-		History.SuperResTextures[0] = GraphBuilder.CreateTexture(SuperResDesc, TEXT("TAA.History.SuperResMetadata[0]"));
+		History.SuperResTextures[0] = GraphBuilder.CreateTexture(SuperResDesc, TEXT("TSR.History.SuperResMetadata[0]"));
 
 		SuperResDesc.Format = PF_R8_UINT;
-		History.SuperResTextures[1] = GraphBuilder.CreateTexture(SuperResDesc, TEXT("TAA.History.SuperResMetadata[1]"));
+		History.SuperResTextures[1] = GraphBuilder.CreateTexture(SuperResDesc, TEXT("TSR.History.SuperResMetadata[1]"));
 	}
 
 	// Decimate input to flicker at same frequency as input.
@@ -1430,14 +1431,14 @@ static void AddGen5MainTemporalAAPasses(
 				FClearValueBinding::None,
 				/* InFlags = */ TexCreate_ShaderResource | TexCreate_UAV);
 
-			ParallaxRejectionMaskTexture = GraphBuilder.CreateTexture(Desc, TEXT("TAA.ParallaxRejectionMask"));
+			ParallaxRejectionMaskTexture = GraphBuilder.CreateTexture(Desc, TEXT("TSR.ParallaxRejectionMask"));
 
 			// TODO(G5TAA): can compress to the history seed's 4bit per pixel
 			Desc.Format = PF_R8G8B8A8;
-			InterferenceSeedTexture = GraphBuilder.CreateTexture(Desc, TEXT("TAA.Interference.Seed"));
+			InterferenceSeedTexture = GraphBuilder.CreateTexture(Desc, TEXT("TSR.Interference.Seed"));
 		}
 
-		FTAADecimateHistoryCS::FParameters* PassParameters = GraphBuilder.AllocParameters<FTAADecimateHistoryCS::FParameters>();
+		FTSRDecimateHistoryCS::FParameters* PassParameters = GraphBuilder.AllocParameters<FTSRDecimateHistoryCS::FParameters>();
 		PassParameters->CommonParameters = CommonParameters;
 		{
 			const FViewMatrices& ViewMatrices = View.ViewMatrices;
@@ -1474,11 +1475,11 @@ static void AddGen5MainTemporalAAPasses(
 				FClearValueBinding::None,
 				/* InFlags = */ TexCreate_ShaderResource | TexCreate_UAV);
 
-			HalfResInputSceneColorTexture = GraphBuilder.CreateTexture(Desc, TEXT("TAA.HalfResInput"));
-			HalfResPredictionSceneColorTexture = GraphBuilder.CreateTexture(Desc, TEXT("TAA.Prediction.SceneColor"));
+			HalfResInputSceneColorTexture = GraphBuilder.CreateTexture(Desc, TEXT("TSR.HalfResInput"));
+			HalfResPredictionSceneColorTexture = GraphBuilder.CreateTexture(Desc, TEXT("TSR.Prediction.SceneColor"));
 
 			Desc.Format = PF_R8;
-			HalfResParallaxRejectionMaskTexture = GraphBuilder.CreateTexture(Desc, TEXT("TAA.HalfResParallaxRejectionMask"));
+			HalfResParallaxRejectionMaskTexture = GraphBuilder.CreateTexture(Desc, TEXT("TSR.HalfResParallaxRejectionMask"));
 
 			PassParameters->HalfResSceneColorOutput = GraphBuilder.CreateUAV(HalfResInputSceneColorTexture);
 			PassParameters->HalfResPredictionSceneColorOutput = GraphBuilder.CreateUAV(HalfResPredictionSceneColorTexture);
@@ -1492,22 +1493,22 @@ static void AddGen5MainTemporalAAPasses(
 				FClearValueBinding::None,
 				/* InFlags = */ TexCreate_ShaderResource | TexCreate_UAV);
 
-			PredictionSceneColorTexture = GraphBuilder.CreateTexture(Desc, TEXT("TAA.Prediction.SceneColor"));
+			PredictionSceneColorTexture = GraphBuilder.CreateTexture(Desc, TEXT("TSR.Prediction.SceneColor"));
 
 			PassParameters->PredictionSceneColorOutput = GraphBuilder.CreateUAV(PredictionSceneColorTexture);
 		}
 
 		PassParameters->ParallaxRejectionMaskOutput = GraphBuilder.CreateUAV(ParallaxRejectionMaskTexture);
 		PassParameters->InterferenceSeedOutput = GraphBuilder.CreateUAV(InterferenceSeedTexture);
-		PassParameters->DebugOutput = CreateDebugUAV(LowFrequencyExtent, TEXT("Debug.TAA.DecimateHistory"));
+		PassParameters->DebugOutput = CreateDebugUAV(LowFrequencyExtent, TEXT("Debug.TSR.DecimateHistory"));
 
-		FTAADecimateHistoryCS::FPermutationDomain PermutationVector;
-		PermutationVector.Set<FTAADecimateHistoryCS::FOutputHalfRes>(bHalfResLowFrequency);
+		FTSRDecimateHistoryCS::FPermutationDomain PermutationVector;
+		PermutationVector.Set<FTSRDecimateHistoryCS::FOutputHalfRes>(bHalfResLowFrequency);
 
-		TShaderMapRef<FTAADecimateHistoryCS> ComputeShader(View.ShaderMap, PermutationVector);
+		TShaderMapRef<FTSRDecimateHistoryCS> ComputeShader(View.ShaderMap, PermutationVector);
 		FComputeShaderUtils::AddPass(
 			GraphBuilder,
-			RDG_EVENT_NAME("TAA DecimateHistory(%s) %dx%d",
+			RDG_EVENT_NAME("TSR DecimateHistory(%s) %dx%d",
 				bHalfResLowFrequency ? TEXT("HalfResShadingOutput") : TEXT(""),
 				InputRect.Width(), InputRect.Height()),
 			ComputeShader,
@@ -1532,10 +1533,10 @@ static void AddGen5MainTemporalAAPasses(
 				FClearValueBinding::None,
 				/* InFlags = */ TexCreate_ShaderResource | TexCreate_UAV);
 
-			InterferenceWeightTexture = GraphBuilder.CreateTexture(Desc, TEXT("TAA.Interference.Weight"));
+			InterferenceWeightTexture = GraphBuilder.CreateTexture(Desc, TEXT("TSR.Interference.Weight"));
 		}
 
-		FTAADetectInterferenceCS::FParameters* PassParameters = GraphBuilder.AllocParameters<FTAADetectInterferenceCS::FParameters>();
+		FTSRDetectInterferenceCS::FParameters* PassParameters = GraphBuilder.AllocParameters<FTSRDetectInterferenceCS::FParameters>();
 		PassParameters->CommonParameters = CommonParameters;
 
 		PassParameters->InputSceneColorTexture = PassInputs.SceneColorTexture;
@@ -1549,9 +1550,9 @@ static void AddGen5MainTemporalAAPasses(
 
 		PassParameters->HistoryOutput = CreateUAVs(GraphBuilder, History);
 		PassParameters->InterferenceWeightOutput = GraphBuilder.CreateUAV(InterferenceWeightTexture);
-		PassParameters->DebugOutput = CreateDebugUAV(LowFrequencyExtent, TEXT("Debug.TAA.DetectInterference"));
+		PassParameters->DebugOutput = CreateDebugUAV(LowFrequencyExtent, TEXT("Debug.TSR.DetectInterference"));
 
-		TShaderMapRef<FTAADetectInterferenceCS> ComputeShader(View.ShaderMap);
+		TShaderMapRef<FTSRDetectInterferenceCS> ComputeShader(View.ShaderMap);
 		ClearUnusedGraphResources(ComputeShader, PassParameters);
 
 		for (int32 i = 0; i < PrevHistory.LowResTextures.Num(); i++)
@@ -1576,7 +1577,7 @@ static void AddGen5MainTemporalAAPasses(
 
 		FComputeShaderUtils::AddPass(
 			GraphBuilder,
-			RDG_EVENT_NAME("TAA DetectInterference %dx%d", InputRect.Width(), InputRect.Height()),
+			RDG_EVENT_NAME("TSR DetectInterference %dx%d", InputRect.Width(), InputRect.Height()),
 			ComputeShader,
 			PassParameters,
 			FComputeShaderUtils::GetGroupCount(InputRect.Size(), 8));
@@ -1606,11 +1607,11 @@ static void AddGen5MainTemporalAAPasses(
 					FClearValueBinding::None,
 					/* InFlags = */ TexCreate_ShaderResource | TexCreate_UAV);
 
-				FilteredInputTexture = GraphBuilder.CreateTexture(Desc, TEXT("TAA.Filtered.SceneColor"));
-				FilteredPredictionSceneColorTexture = GraphBuilder.CreateTexture(Desc, TEXT("TAA.Filtered.Prediction.SceneColor"));
+				FilteredInputTexture = GraphBuilder.CreateTexture(Desc, TEXT("TSR.Filtered.SceneColor"));
+				FilteredPredictionSceneColorTexture = GraphBuilder.CreateTexture(Desc, TEXT("TSR.Filtered.Prediction.SceneColor"));
 			}
 
-			FTAAFilterFrequenciesCS::FParameters* PassParameters = GraphBuilder.AllocParameters<FTAAFilterFrequenciesCS::FParameters>();
+			FTSRFilterFrequenciesCS::FParameters* PassParameters = GraphBuilder.AllocParameters<FTSRFilterFrequenciesCS::FParameters>();
 			PassParameters->CommonParameters = CommonParameters;
 			PassParameters->OutputQuantizationError = ComputePixelFormatQuantizationError(FilteredInputTexture->Desc.Format);
 
@@ -1630,12 +1631,12 @@ static void AddGen5MainTemporalAAPasses(
 
 			PassParameters->FilteredInputOutput = GraphBuilder.CreateUAV(FilteredInputTexture);
 			PassParameters->FilteredPredictionSceneColorOutput = GraphBuilder.CreateUAV(FilteredPredictionSceneColorTexture);
-			PassParameters->DebugOutput = CreateDebugUAV(LowFrequencyExtent, TEXT("Debug.TAA.FilterFrequencies"));
+			PassParameters->DebugOutput = CreateDebugUAV(LowFrequencyExtent, TEXT("Debug.TSR.FilterFrequencies"));
 
-			TShaderMapRef<FTAAFilterFrequenciesCS> ComputeShader(View.ShaderMap);
+			TShaderMapRef<FTSRFilterFrequenciesCS> ComputeShader(View.ShaderMap);
 			FComputeShaderUtils::AddPass(
 				GraphBuilder,
-				RDG_EVENT_NAME("TAA FilterFrequencies %dx%d", LowFrequencyRect.Width(), LowFrequencyRect.Height()),
+				RDG_EVENT_NAME("TSR FilterFrequencies %dx%d", LowFrequencyRect.Width(), LowFrequencyRect.Height()),
 				ComputeShader,
 				PassParameters,
 				FComputeShaderUtils::GetGroupCount(LowFrequencyRect.Size(), 16));
@@ -1650,10 +1651,10 @@ static void AddGen5MainTemporalAAPasses(
 					FClearValueBinding::None,
 					/* InFlags = */ TexCreate_ShaderResource | TexCreate_UAV);
 
-				HistoryRejectionTexture = GraphBuilder.CreateTexture(Desc, TEXT("TAA.HistoryRejection"));
+				HistoryRejectionTexture = GraphBuilder.CreateTexture(Desc, TEXT("TSR.HistoryRejection"));
 			}
 
-			FTAACompareHistoryCS::FParameters* PassParameters = GraphBuilder.AllocParameters<FTAACompareHistoryCS::FParameters>();
+			FTSRCompareHistoryCS::FParameters* PassParameters = GraphBuilder.AllocParameters<FTSRCompareHistoryCS::FParameters>();
 			PassParameters->CommonParameters = CommonParameters;
 			PassParameters->ParallaxRejectionMaskTexture = ParallaxRejectionMaskTexture;
 			PassParameters->FilteredInputTexture = FilteredInputTexture;
@@ -1661,12 +1662,12 @@ static void AddGen5MainTemporalAAPasses(
 			PassParameters->InterferenceWeightTexture = InterferenceWeightTexture;
 
 			PassParameters->HistoryRejectionOutput = GraphBuilder.CreateUAV(HistoryRejectionTexture);
-			PassParameters->DebugOutput = CreateDebugUAV(LowFrequencyExtent, TEXT("Debug.TAA.CompareHistory"));
+			PassParameters->DebugOutput = CreateDebugUAV(LowFrequencyExtent, TEXT("Debug.TSR.CompareHistory"));
 
-			TShaderMapRef<FTAACompareHistoryCS> ComputeShader(View.ShaderMap);
+			TShaderMapRef<FTSRCompareHistoryCS> ComputeShader(View.ShaderMap);
 			FComputeShaderUtils::AddPass(
 				GraphBuilder,
-				RDG_EVENT_NAME("TAA CompareHistory %dx%d", LowFrequencyRect.Width(), LowFrequencyRect.Height()),
+				RDG_EVENT_NAME("TSR CompareHistory %dx%d", LowFrequencyRect.Width(), LowFrequencyRect.Height()),
 				ComputeShader,
 				PassParameters,
 				FComputeShaderUtils::GetGroupCount(LowFrequencyRect.Size(), 16));
@@ -1676,18 +1677,18 @@ static void AddGen5MainTemporalAAPasses(
 	// Post filter the rejection.
 	if (CVarTAAFilterShadingRejection.GetValueOnRenderThread() != 0)
 	{
-		FRDGTextureRef FilteredHistoryRejectionTexture = GraphBuilder.CreateTexture(HistoryRejectionTexture->Desc, TEXT("TAA.FilteredHistoryRejection"));
+		FRDGTextureRef FilteredHistoryRejectionTexture = GraphBuilder.CreateTexture(HistoryRejectionTexture->Desc, TEXT("TSR.FilteredHistoryRejection"));
 
-		FTAAPostfilterRejectionCS::FParameters* PassParameters = GraphBuilder.AllocParameters<FTAAPostfilterRejectionCS::FParameters>();
+		FTSRPostfilterRejectionCS::FParameters* PassParameters = GraphBuilder.AllocParameters<FTSRPostfilterRejectionCS::FParameters>();
 		PassParameters->CommonParameters = CommonParameters;
 		PassParameters->HistoryRejectionTexture = HistoryRejectionTexture;
 		PassParameters->HistoryRejectionOutput = GraphBuilder.CreateUAV(FilteredHistoryRejectionTexture);
-		PassParameters->DebugOutput = CreateDebugUAV(RejectionExtent, TEXT("Debug.TAA.PostfilterRejection"));
+		PassParameters->DebugOutput = CreateDebugUAV(RejectionExtent, TEXT("Debug.TSR.PostfilterRejection"));
 
-		TShaderMapRef<FTAAPostfilterRejectionCS> ComputeShader(View.ShaderMap);
+		TShaderMapRef<FTSRPostfilterRejectionCS> ComputeShader(View.ShaderMap);
 		FComputeShaderUtils::AddPass(
 			GraphBuilder,
-			RDG_EVENT_NAME("TAA PostfilterRejection %dx%d", RejectionRect.Width(), RejectionRect.Height()),
+			RDG_EVENT_NAME("TSR PostfilterRejection %dx%d", RejectionRect.Width(), RejectionRect.Height()),
 			ComputeShader,
 			PassParameters,
 			FComputeShaderUtils::GetGroupCount(RejectionRect.Size(), 8));
@@ -1698,18 +1699,18 @@ static void AddGen5MainTemporalAAPasses(
 	// Dilate the rejection.
 	FRDGTextureRef DilatedHistoryRejectionTexture;
 	{
-		DilatedHistoryRejectionTexture = GraphBuilder.CreateTexture(HistoryRejectionTexture->Desc, TEXT("TAA.DilatedHistoryRejection"));
+		DilatedHistoryRejectionTexture = GraphBuilder.CreateTexture(HistoryRejectionTexture->Desc, TEXT("TSR.DilatedHistoryRejection"));
 
-		FTAADilateRejectionCS::FParameters* PassParameters = GraphBuilder.AllocParameters<FTAADilateRejectionCS::FParameters>();
+		FTSRDilateRejectionCS::FParameters* PassParameters = GraphBuilder.AllocParameters<FTSRDilateRejectionCS::FParameters>();
 		PassParameters->CommonParameters = CommonParameters;
 		PassParameters->HistoryRejectionTexture = HistoryRejectionTexture;
 		PassParameters->DilatedHistoryRejectionOutput = GraphBuilder.CreateUAV(DilatedHistoryRejectionTexture);
-		PassParameters->DebugOutput = CreateDebugUAV(RejectionExtent, TEXT("Debug.TAA.DilateRejection"));
+		PassParameters->DebugOutput = CreateDebugUAV(RejectionExtent, TEXT("Debug.TSR.DilateRejection"));
 
-		TShaderMapRef<FTAADilateRejectionCS> ComputeShader(View.ShaderMap);
+		TShaderMapRef<FTSRDilateRejectionCS> ComputeShader(View.ShaderMap);
 		FComputeShaderUtils::AddPass(
 			GraphBuilder,
-			RDG_EVENT_NAME("TAA DilateRejection %dx%d", RejectionRect.Width(), RejectionRect.Height()),
+			RDG_EVENT_NAME("TSR DilateRejection %dx%d", RejectionRect.Width(), RejectionRect.Height()),
 			ComputeShader,
 			PassParameters,
 			FComputeShaderUtils::GetGroupCount(RejectionRect.Size(), 8));
@@ -1718,7 +1719,7 @@ static void AddGen5MainTemporalAAPasses(
 	TStaticArray<bool, FTemporalAAHistory::kSuperResRenderTargetCount> ExtractSuperResHistoryTexture;
 	if (bEnableInterferenceHeuristic)
 	{
-		FTAAUpdateSuperResHistoryCS::FParameters* PassParameters = GraphBuilder.AllocParameters<FTAAUpdateSuperResHistoryCS::FParameters>();
+		FTSRUpdateSuperResHistoryCS::FParameters* PassParameters = GraphBuilder.AllocParameters<FTSRUpdateSuperResHistoryCS::FParameters>();
 		PassParameters->CommonParameters = CommonParameters;
 		PassParameters->InputSceneColorTexture = PassInputs.SceneColorTexture;
 		PassParameters->DilatedVelocityTexture = DilatedVelocityTexture;
@@ -1728,9 +1729,9 @@ static void AddGen5MainTemporalAAPasses(
 		PassParameters->PrevHistory = PrevHistory;
 
 		PassParameters->HistoryOutput = CreateUAVs(GraphBuilder, History);
-		PassParameters->DebugOutput = CreateDebugUAV(HistoryExtent * 2, TEXT("Debug.TAA.UpdateSuperResHistory"));
+		PassParameters->DebugOutput = CreateDebugUAV(HistoryExtent * 2, TEXT("Debug.TSR.UpdateSuperResHistory"));
 
-		TShaderMapRef<FTAAUpdateSuperResHistoryCS> ComputeShader(View.ShaderMap);
+		TShaderMapRef<FTSRUpdateSuperResHistoryCS> ComputeShader(View.ShaderMap);
 		ClearUnusedGraphResources(ComputeShader, PassParameters);
 
 		for (int32 i = 0; i < PrevHistory.SuperResTextures.Num(); i++)
@@ -1755,7 +1756,7 @@ static void AddGen5MainTemporalAAPasses(
 
 		FComputeShaderUtils::AddPass(
 			GraphBuilder,
-			RDG_EVENT_NAME("TAA UpdateSuperResHistory %dx%d", HistorySize.X * 2, HistorySize.Y * 2),
+			RDG_EVENT_NAME("TSR UpdateSuperResHistory %dx%d", HistorySize.X * 2, HistorySize.Y * 2),
 			ComputeShader,
 			PassParameters,
 			FComputeShaderUtils::GetGroupCount(HistorySize * 2, 8));
@@ -1779,10 +1780,10 @@ static void AddGen5MainTemporalAAPasses(
 				FClearValueBinding::None,
 				/* InFlags = */ TexCreate_ShaderResource | TexCreate_UAV);
 
-			SceneColorOutputTexture = GraphBuilder.CreateTexture(Desc, TEXT("TAA.Output"));
+			SceneColorOutputTexture = GraphBuilder.CreateTexture(Desc, TEXT("TSR.Output"));
 		}
 
-		FTAAUpdateHistoryCS::FParameters* PassParameters = GraphBuilder.AllocParameters<FTAAUpdateHistoryCS::FParameters>();
+		FTSRUpdateHistoryCS::FParameters* PassParameters = GraphBuilder.AllocParameters<FTSRUpdateHistoryCS::FParameters>();
 		PassParameters->CommonParameters = CommonParameters;
 		PassParameters->InputSceneColorTexture = PassInputs.SceneColorTexture;
 		PassParameters->InputSceneStencilTexture = GraphBuilder.CreateSRV(
@@ -1802,9 +1803,9 @@ static void AddGen5MainTemporalAAPasses(
 
 		PassParameters->HistoryOutput = CreateUAVs(GraphBuilder, History);
 		PassParameters->SceneColorOutput = GraphBuilder.CreateUAV(SceneColorOutputTexture);
-		PassParameters->DebugOutput = CreateDebugUAV(HistoryExtent, TEXT("Debug.TAA.UpdateHistory"));
+		PassParameters->DebugOutput = CreateDebugUAV(HistoryExtent, TEXT("Debug.TSR.UpdateHistory"));
 
-		TShaderMapRef<FTAAUpdateHistoryCS> ComputeShader(View.ShaderMap);
+		TShaderMapRef<FTSRUpdateHistoryCS> ComputeShader(View.ShaderMap);
 		ClearUnusedGraphResources(ComputeShader, PassParameters);
 
 		for (int32 i = 0; i < PrevHistory.Textures.Num(); i++)
@@ -1829,7 +1830,7 @@ static void AddGen5MainTemporalAAPasses(
 
 		FComputeShaderUtils::AddPass(
 			GraphBuilder,
-			RDG_EVENT_NAME("TAA UpdateHistory(%s) %dx%d", 
+			RDG_EVENT_NAME("TSR UpdateHistory(%s) %dx%d", 
 				History.Textures[0]->Desc.Format == PF_FloatR11G11B10 ? TEXT("R11G11B10") : TEXT(""),
 				HistorySize.X, HistorySize.Y),
 			ComputeShader,
@@ -1843,16 +1844,16 @@ static void AddGen5MainTemporalAAPasses(
 	{
 		const int32 kHistoryUpscalingFactor = 2;
 
-		FTAADebugHistoryCS::FParameters* PassParameters = GraphBuilder.AllocParameters<FTAADebugHistoryCS::FParameters>();
+		FTSRDebugHistoryCS::FParameters* PassParameters = GraphBuilder.AllocParameters<FTSRDebugHistoryCS::FParameters>();
 		PassParameters->CommonParameters = CommonParameters;
 		PassParameters->History = History;
 		PassParameters->PrevHistory = PrevHistory;
-		PassParameters->DebugOutput = CreateDebugUAV(HistoryExtent * kHistoryUpscalingFactor, TEXT("Debug.TAA.History"));
+		PassParameters->DebugOutput = CreateDebugUAV(HistoryExtent * kHistoryUpscalingFactor, TEXT("Debug.TSR.History"));
 
-		TShaderMapRef<FTAADebugHistoryCS> ComputeShader(View.ShaderMap);
+		TShaderMapRef<FTSRDebugHistoryCS> ComputeShader(View.ShaderMap);
 		FComputeShaderUtils::AddPass(
 			GraphBuilder,
-			RDG_EVENT_NAME("TAA DebugHistory %dx%d", RejectionRect.Width(), RejectionRect.Height()),
+			RDG_EVENT_NAME("TSR DebugHistory %dx%d", RejectionRect.Width(), RejectionRect.Height()),
 			ComputeShader,
 			PassParameters,
 			FComputeShaderUtils::GetGroupCount(HistorySize * kHistoryUpscalingFactor, 8));
@@ -1902,7 +1903,7 @@ static void AddGen5MainTemporalAAPasses(
 
 	*OutSceneColorTexture = SceneColorOutputTexture;
 	*OutSceneColorViewRect = OutputRect;
-} // AddGen5MainTemporalAAPasses()
+} // AddTemporalSuperResolutionPasses()
 
 static void AddGen4MainTemporalAAPasses(
 	FRDGBuilder& GraphBuilder,
@@ -2017,7 +2018,7 @@ public:
 			*OutSceneColorHalfResTexture = nullptr;
 			//*OutSceneColorHalfResViewRect; // TODO.
 
-			return AddGen5MainTemporalAAPasses(
+			return AddTemporalSuperResolutionPasses(
 				GraphBuilder,
 				View,
 				PassInputs,
@@ -2074,10 +2075,10 @@ EMainTAAPassConfig ITemporalUpscaler::GetMainTAAPassConfig(const FViewInfo& View
 	}
 	else if (CVarTAAAlgorithm.GetValueOnRenderThread() && DoesPlatformSupportGen5TAA(View.GetShaderPlatform()))
 	{
-		return EMainTAAPassConfig::Gen5;
+		return EMainTAAPassConfig::TSR;
 	}
 	else
 	{
-		return EMainTAAPassConfig::Gen4;
+		return EMainTAAPassConfig::TAA;
 	}
 }
diff --git a/Engine/Source/Runtime/Renderer/Private/PostProcess/TemporalAA.h b/Engine/Source/Runtime/Renderer/Private/PostProcess/TemporalAA.h
--- a/Engine/Source/Runtime/Renderer/Private/PostProcess/TemporalAA.h
+++ b/Engine/Source/Runtime/Renderer/Private/PostProcess/TemporalAA.h
@@ -14,10 +14,10 @@ enum class EMainTAAPassConfig : uint8
 	Disabled,
 
 	// Uses old UE4's Temporal AA maintained for Gen4 consoles
-	Gen4,
+	TAA,
 
-	// Uses Gen5 Temporal AA
-	Gen5,
+	// Uses Temporal Super Resolution
+	TSR,
 
 	// Uses third party View.Family->GetTemporalUpscalerInterface()
 	ThirdParty,
diff --git a/Engine/Source/Runtime/Renderer/Private/SceneRendering.cpp b/Engine/Source/Runtime/Renderer/Private/SceneRendering.cpp
--- a/Engine/Source/Runtime/Renderer/Private/SceneRendering.cpp
+++ b/Engine/Source/Runtime/Renderer/Private/SceneRendering.cpp
@@ -1611,7 +1611,7 @@ void FViewInfo::SetupUniformBufferParameters(
 		// through whether the velocity has been drawn by the base pass.
 		ViewUniformShaderParameters.ForceDrawAllVelocities = (
 			CVarBasePassForceOutputsVelocity.GetValueOnRenderThread() ||
-			MainTAAPass != EMainTAAPassConfig::Gen4);
+			MainTAAPass != EMainTAAPassConfig::TAA);
 	}
 
 	uint32 FrameIndex = 0;
diff --git a/Engine/Source/Runtime/Renderer/Private/SceneVisibility.cpp b/Engine/Source/Runtime/Renderer/Private/SceneVisibility.cpp
--- a/Engine/Source/Runtime/Renderer/Private/SceneVisibility.cpp
+++ b/Engine/Source/Runtime/Renderer/Private/SceneVisibility.cpp
@@ -3289,7 +3289,7 @@ void FSceneRenderer::PreVisibilityFrameSetup(FRDGBuilder& GraphBuilder, const FS
 			// Compute number of TAA samples.
 			int32 TemporalAASamples;
 			{
-				if (TAAConfig == EMainTAAPassConfig::Gen5)
+				if (TAAConfig == EMainTAAPassConfig::TSR)
 				{
 					// Force the number of AA sample to make sure the quality doesn't get
 					// compromised by previously set settings for Gen4 TAA
@@ -3313,7 +3313,7 @@ void FSceneRenderer::PreVisibilityFrameSetup(FRDGBuilder& GraphBuilder, const FS
 
 				// Use immediately higher prime number to break up coherence between the TAA jitter sequence and any
 				// other random signal that are power of two of View.StateFrameIndex
-				if (TAAConfig == EMainTAAPassConfig::Gen5)
+				if (TAAConfig == EMainTAAPassConfig::TSR)
 				{
 					static const int8 kFirstPrimeNumbers[25] = {
 						2, 3, 5, 7, 11, 13, 17, 19, 23, 29, 31, 37, 41, 43, 47, 53, 59, 61, 67, 71, 73, 79, 83, 89, 97,
