Fix missing UB data for cloud rendering.
Tested with GPUScene disabled.

diff --git a/Engine/Source/Runtime/Renderer/Private/VolumetricCloudRendering.cpp b/Engine/Source/Runtime/Renderer/Private/VolumetricCloudRendering.cpp
--- a/Engine/Source/Runtime/Renderer/Private/VolumetricCloudRendering.cpp
+++ b/Engine/Source/Runtime/Renderer/Private/VolumetricCloudRendering.cpp
@@ -902,7 +902,7 @@ static void GetSingleTriangleMeshBatch(FMeshBatch& LocalSingleTriangleMesh, cons
 	LocalSingleTriangleMesh.Elements[0].MinVertexIndex = 0;
 	LocalSingleTriangleMesh.Elements[0].MaxVertexIndex = 2;
 
-	LocalSingleTriangleMesh.Elements[0].PrimitiveUniformBuffer = nullptr;
+	LocalSingleTriangleMesh.Elements[0].PrimitiveUniformBufferResource = &GIdentityPrimitiveUniformBuffer;
 	LocalSingleTriangleMesh.Elements[0].PrimitiveIdMode = PrimID_ForceZero;
 }
 
