Remove unused virtual shadow map denoiser path
No more need for writing out the "signal" texture info, so reduce that to single F32 for shadow factor
No need to return HitDistance as far up the stack in the shaders
Remove VSM defines and permutations from screen space denoiser

#lock nick.whiting

diff --git a/Engine/Shaders/Private/ScreenSpaceDenoise/SSDCommon.ush b/Engine/Shaders/Private/ScreenSpaceDenoise/SSDCommon.ush
--- a/Engine/Shaders/Private/ScreenSpaceDenoise/SSDCommon.ush
+++ b/Engine/Shaders/Private/ScreenSpaceDenoise/SSDCommon.ush
@@ -30,7 +30,7 @@
 #if CONFIG_SIGNAL_PROCESSING == SIGNAL_PROCESSING_REFLECTIONS
 	#define TARGETED_SAMPLE_COUNT 1024
 
-#elif CONFIG_SIGNAL_PROCESSING == SIGNAL_PROCESSING_SHADOW_VISIBILITY_MASK || CONFIG_SIGNAL_PROCESSING == SIGNAL_PROCESSING_VIRTUAL_SHADOW_MAP_MASK
+#elif CONFIG_SIGNAL_PROCESSING == SIGNAL_PROCESSING_SHADOW_VISIBILITY_MASK
 	// A gray scale valued encoded into a 16bit float only have 10bits mantissa.
 	#define TARGETED_SAMPLE_COUNT 1024
 
diff --git a/Engine/Shaders/Private/ScreenSpaceDenoise/SSDDefinitions.ush b/Engine/Shaders/Private/ScreenSpaceDenoise/SSDDefinitions.ush
--- a/Engine/Shaders/Private/ScreenSpaceDenoise/SSDDefinitions.ush
+++ b/Engine/Shaders/Private/ScreenSpaceDenoise/SSDDefinitions.ush
@@ -191,9 +191,6 @@
 	/** Diffuse indirect using spherical harmonic. */
 	#define SIGNAL_PROCESSING_DIFFUSE_PROBE_HIERARCHY 7
 
-	/** Virtual shadow map mask */
-	#define SIGNAL_PROCESSING_VIRTUAL_SHADOW_MAP_MASK 8
-
 
 
 /** Defines the color space bitfield. */
diff --git a/Engine/Shaders/Private/ScreenSpaceDenoise/SSDSignalFramework.ush b/Engine/Shaders/Private/ScreenSpaceDenoise/SSDSignalFramework.ush
--- a/Engine/Shaders/Private/ScreenSpaceDenoise/SSDSignalFramework.ush
+++ b/Engine/Shaders/Private/ScreenSpaceDenoise/SSDSignalFramework.ush
@@ -30,7 +30,7 @@
 #if !defined(CONFIG_SIGNAL_PROCESSING)
 	#error Missing CONFIG_SIGNAL_PROCESSING
 
-#elif CONFIG_SIGNAL_PROCESSING == SIGNAL_PROCESSING_SHADOW_VISIBILITY_MASK || CONFIG_SIGNAL_PROCESSING == SIGNAL_PROCESSING_VIRTUAL_SHADOW_MAP_MASK
+#elif CONFIG_SIGNAL_PROCESSING == SIGNAL_PROCESSING_SHADOW_VISIBILITY_MASK
 	#include "SSDPublicHarmonics.ush"
 
 #elif CONFIG_SIGNAL_PROCESSING == SIGNAL_PROCESSING_POLYCHROMATIC_PENUMBRA_HARMONIC
@@ -582,7 +582,7 @@ float ComputeBilateralWeight(
 }
 
 float GetSignalWorldBluringRadius(FSSDSignalFrequency SampleFrequency, FSSDSampleSceneInfos SceneMetadata, FSSDSignalDomainKnowledge DomainKnowledge)
-#if CONFIG_SIGNAL_PROCESSING == SIGNAL_PROCESSING_SHADOW_VISIBILITY_MASK || CONFIG_SIGNAL_PROCESSING == SIGNAL_PROCESSING_VIRTUAL_SHADOW_MAP_MASK
+#if CONFIG_SIGNAL_PROCESSING == SIGNAL_PROCESSING_SHADOW_VISIBILITY_MASK
 {
 	if (SampleFrequency.ClosestHitDistance == DENOISER_INVALID_HIT_DISTANCE)
 	{
diff --git a/Engine/Shaders/Private/ScreenSpaceDenoise/SSDSpatialAccumulation.usf b/Engine/Shaders/Private/ScreenSpaceDenoise/SSDSpatialAccumulation.usf
--- a/Engine/Shaders/Private/ScreenSpaceDenoise/SSDSpatialAccumulation.usf
+++ b/Engine/Shaders/Private/ScreenSpaceDenoise/SSDSpatialAccumulation.usf
@@ -121,49 +121,6 @@
 		// #define CONFIG_ACCUMULATOR_VGPR_COMPRESSION ACCUMULATOR_COMPRESSION_PENUMBRA_DRB
 	#endif
 
-#elif CONFIG_SIGNAL_PROCESSING == SIGNAL_PROCESSING_VIRTUAL_SHADOW_MAP_MASK
-	#define MAX_SIGNAL_BATCH_SIZE CONFIG_SIGNAL_BATCH_SIZE
-	#define SIGNAL_ARRAY_SIZE CONFIG_SIGNAL_BATCH_SIZE
-
-	#define CONFIG_BILATERAL_PRESET BILATERAL_PRESET_MONOCHROMATIC_PENUMBRA
-	#define CONFIG_MULTIPLEXED_SIGNALS_PER_SIGNAL_DOMAIN 1
-
-	#if DIM_STAGE == STAGE_RECONSTRUCTION
-		// Input and output layout.
-		#define CONFIG_SIGNAL_INPUT_LAYOUT  SIGNAL_BUFFER_LAYOUT_PENUMBRA_HISTORY
-		#define CONFIG_SIGNAL_OUTPUT_LAYOUT SIGNAL_BUFFER_LAYOUT_PENUMBRA_HISTORY
-
-		#define CONFIG_INPUT_TEXTURE_COUNT ((CONFIG_SIGNAL_BATCH_SIZE + 1) / 2)
-		#define CONFIG_OUTPUT_TEXTURE_COUNT CONFIG_SIGNAL_BATCH_SIZE
-
-		#define CONFIG_SAMPLE_SET           SAMPLE_SET_STACKOWIAK_4_SETS
-		#define CONFIG_BILATERAL_DISTANCE_COMPUTATION SIGNAL_WORLD_FREQUENCY_PRECOMPUTED_BLURING_RADIUS
-		#define CONFIG_MAX_WITH_REF_DISTANCE 1
-		#define CONFIG_OUTPUT_MODE          OUTPUT_MODE_DRB
-
-		#if DIM_SIGNAL_BATCH_SIZE > 1
-			#define CONFIG_CLAMP_UV_PER_SIGNAL 1
-		#endif
-
-	#elif DIM_STAGE == STAGE_PRE_CONVOLUTION
-		// Input and output layout.
-		#define CONFIG_SIGNAL_INPUT_LAYOUT  SIGNAL_BUFFER_LAYOUT_PENUMBRA_HISTORY
-		#define CONFIG_SIGNAL_OUTPUT_LAYOUT SIGNAL_BUFFER_LAYOUT_PENUMBRA_HISTORY
-
-		#define CONFIG_INPUT_TEXTURE_COUNT CONFIG_SIGNAL_BATCH_SIZE
-		#define CONFIG_OUTPUT_TEXTURE_COUNT CONFIG_SIGNAL_BATCH_SIZE
-
-		#define CONFIG_SAMPLE_SET           SAMPLE_SET_HEXAWEB
-		#define CONFIG_CUSTOM_SPREAD_FACTOR 1
-
-		#define CONFIG_BILATERAL_DISTANCE_COMPUTATION SIGNAL_WORLD_FREQUENCY_PRECOMPUTED_BLURING_RADIUS
-		#define CONFIG_MAX_WITH_REF_DISTANCE 1
-		#define CONFIG_OUTPUT_MODE          OUTPUT_MODE_DRB
-
-	#else
-		#error Unexpected stage.
-	#endif
-
 #elif CONFIG_SIGNAL_PROCESSING == SIGNAL_PROCESSING_POLYCHROMATIC_PENUMBRA_HARMONIC
 	// Denoise diffuse and specular harmonics at the same time.
 	#define MAX_SIGNAL_BATCH_SIZE 2
diff --git a/Engine/Shaders/Private/ScreenSpaceDenoise/SSDTemporalAccumulation.usf b/Engine/Shaders/Private/ScreenSpaceDenoise/SSDTemporalAccumulation.usf
--- a/Engine/Shaders/Private/ScreenSpaceDenoise/SSDTemporalAccumulation.usf
+++ b/Engine/Shaders/Private/ScreenSpaceDenoise/SSDTemporalAccumulation.usf
@@ -25,7 +25,7 @@
 
 
 
-#if CONFIG_SIGNAL_PROCESSING == SIGNAL_PROCESSING_SHADOW_VISIBILITY_MASK || CONFIG_SIGNAL_PROCESSING == SIGNAL_PROCESSING_VIRTUAL_SHADOW_MAP_MASK
+#if CONFIG_SIGNAL_PROCESSING == SIGNAL_PROCESSING_SHADOW_VISIBILITY_MASK
 	#define MAX_SIGNAL_BATCH_SIZE CONFIG_SIGNAL_BATCH_SIZE
 	#define SIGNAL_ARRAY_SIZE CONFIG_SIGNAL_BATCH_SIZE
 
@@ -722,7 +722,7 @@ void TemporallyAccumulate(
 
 			// TODO(Denoiser): Should be done for all signals.
 			// TODO(Denoiser): Should also take into account the world bluring radius du to size of the pixel.
-			if (CONFIG_SIGNAL_PROCESSING == SIGNAL_PROCESSING_SHADOW_VISIBILITY_MASK || CONFIG_SIGNAL_PROCESSING == SIGNAL_PROCESSING_VIRTUAL_SHADOW_MAP_MASK)
+			if (CONFIG_SIGNAL_PROCESSING == SIGNAL_PROCESSING_SHADOW_VISIBILITY_MASK)
 				KernelConfig.BilateralDistanceComputation = SIGNAL_WORLD_FREQUENCY_PRECOMPUTED_BLURING_RADIUS;
 			KernelConfig.WorldBluringDistanceMultiplier = CONFIG_BILATERAL_DISTANCE_MULTIPLIER;
 			
@@ -740,7 +740,7 @@ void TemporallyAccumulate(
 			}
 			#else
 			{
-				KernelConfig.SampleSet = (CONFIG_SIGNAL_PROCESSING == SIGNAL_PROCESSING_SHADOW_VISIBILITY_MASK || CONFIG_SIGNAL_PROCESSING == SIGNAL_PROCESSING_VIRTUAL_SHADOW_MAP_MASK) ? SAMPLE_SET_3X3 : SAMPLE_SET_NXN;
+				KernelConfig.SampleSet = (CONFIG_SIGNAL_PROCESSING == SIGNAL_PROCESSING_SHADOW_VISIBILITY_MASK) ? SAMPLE_SET_3X3 : SAMPLE_SET_NXN;
 				KernelConfig.BoxKernelRadius = 3;
 				KernelConfig.bUnroll = false;
 			}
@@ -937,7 +937,7 @@ void TemporallyAccumulate(
 						RejectionFactor = max(RejectionFactor, 1 * max(RejectedDiff.SceneColor.y, RejectedDiff.SceneColor.z));
 						RejectionFactor = max(RejectionFactor, 1 * RejectedDiff.SceneColor.a);
 					}
-					#elif CONFIG_SIGNAL_PROCESSING == SIGNAL_PROCESSING_SHADOW_VISIBILITY_MASK || CONFIG_SIGNAL_PROCESSING == SIGNAL_PROCESSING_POLYCHROMATIC_PENUMBRA_HARMONIC || CONFIG_SIGNAL_PROCESSING == SIGNAL_PROCESSING_AO || CONFIG_SIGNAL_PROCESSING == SIGNAL_PROCESSING_DIFFUSE_INDIRECT_AND_AO || CONFIG_SIGNAL_PROCESSING == SIGNAL_PROCESSING_VIRTUAL_SHADOW_MAP_MASK
+					#elif CONFIG_SIGNAL_PROCESSING == SIGNAL_PROCESSING_SHADOW_VISIBILITY_MASK || CONFIG_SIGNAL_PROCESSING == SIGNAL_PROCESSING_POLYCHROMATIC_PENUMBRA_HARMONIC || CONFIG_SIGNAL_PROCESSING == SIGNAL_PROCESSING_AO || CONFIG_SIGNAL_PROCESSING == SIGNAL_PROCESSING_DIFFUSE_INDIRECT_AND_AO
 					{
 						RejectionFactor = abs(ClampedTransformedHistorySample.MissCount - TransformedHistorySample.MissCount);
 					}
diff --git a/Engine/Shaders/Private/VirtualShadowMaps/ProjectionDirectional.ush b/Engine/Shaders/Private/VirtualShadowMaps/ProjectionDirectional.ush
--- a/Engine/Shaders/Private/VirtualShadowMaps/ProjectionDirectional.ush
+++ b/Engine/Shaders/Private/VirtualShadowMaps/ProjectionDirectional.ush
@@ -158,8 +158,7 @@ float TraceDirectional(
 	uint RayCount,
 	uint SamplesPerRay,
 	float RayLengthScale,
-	const half Noise,
-	inout float InOutHitDistance)
+	const half Noise)
 {
 	float3 WorldPosition = TranslatedWorldPosition - View.PreViewTranslation;
 	float3 ViewPosition = mul(float4(TranslatedWorldPosition, 1.0f), View.TranslatedWorldToView).xyz;
@@ -198,12 +197,6 @@ float TraceDirectional(
 			SamplesPerRay,
 			StepOffset);
 
-		if (Result.bValidHit)
-		{
-			// TODO: min over all the rays if we keep denoiser with > 1 spp?
-			InOutHitDistance = Result.OccluderDistance;
-		}
-
 		RayMissCount += Result.bValidHit ? 0 : 1;
 
 #if SMRT_ADAPTIVE_RAY_COUNT && COMPILER_SUPPORTS_WAVE_VOTE
diff --git a/Engine/Shaders/Private/VirtualShadowMaps/ProjectionSpot.ush b/Engine/Shaders/Private/VirtualShadowMaps/ProjectionSpot.ush
--- a/Engine/Shaders/Private/VirtualShadowMaps/ProjectionSpot.ush
+++ b/Engine/Shaders/Private/VirtualShadowMaps/ProjectionSpot.ush
@@ -221,8 +221,7 @@ float TraceLocalLight(
 	uint SamplesPerRay,
 	float CotMaxRayAngleFromLight,
 	const half Noise,
-	const bool bCullBackface, 
-	inout float InOutHitDistance)
+	const bool bCullBackface)
 {
 	const half3 N = GetGBufferDataUint( PixelPos, true ).WorldNormal;
 	const float3 WorldPosition = TranslatedWorldPosition - View.PreViewTranslation;
@@ -307,12 +306,6 @@ float TraceLocalLight(
 					StepOffset);
 			}
 
-			if ( Result.bValidHit )
-			{
-				// TODO: Min over all the rays with > 1 spp?
-				InOutHitDistance = Result.OccluderDistance;
-			}
-
 			RayMissFactor += Result.bValidHit ? 0 : 1;
 			
 #if SMRT_ADAPTIVE_RAY_COUNT && COMPILER_SUPPORTS_WAVE_VOTE
diff --git a/Engine/Shaders/Private/VirtualShadowMaps/VirtualShadowMapProjection.usf b/Engine/Shaders/Private/VirtualShadowMaps/VirtualShadowMapProjection.usf
--- a/Engine/Shaders/Private/VirtualShadowMaps/VirtualShadowMapProjection.usf
+++ b/Engine/Shaders/Private/VirtualShadowMaps/VirtualShadowMapProjection.usf
@@ -51,20 +51,14 @@ float4 EncodeLightAttenuationFromMask(float ShadowMask)
 	return EncodeLightAttenuation(half4(FadedShadow, FadedSSSShadow, FadedShadow, FadedSSSShadow));
 }
 
-Texture2D<float4> InputSignal;
+Texture2D<float> InputShadowFactor;
 
 void VirtualShadowMapCompositePS(
 	in float4 SvPosition : SV_Position,
-	out float4 OutShadowMask : SV_Target
-	)
+	out float4 OutShadowMask : SV_Target)
 {
-	// NOTE: The signal is encoded as per SSDSignalBufferEncoding.ush, SIGNAL_BUFFER_LAYOUT_PENUMBRA_RECONSTRUCTION
-	float4 Input = InputSignal.Load(int3(SvPosition.xy, 0));
-	float SampleCount = Input.g;
-	float ShadowFactor = Input.r;
-	float Shadow = (SampleCount > 0 ? ShadowFactor : 1.0f);
-	
-	OutShadowMask = EncodeLightAttenuationFromMask(Shadow);
+	float ShadowFactor = InputShadowFactor.Load(int3(SvPosition.xy, 0));	
+	OutShadowMask = EncodeLightAttenuationFromMask(ShadowFactor);
 }
 
 
@@ -89,42 +83,14 @@ float ProjectSingleSample(
 	float3 L,
 	float3 TranslatedWorldPosition,
 	float3 EstimatedGeoWorldNormal,
-	float ContactShadowLengthWorld,
-	inout float InOutHitDistance)
+	float ContactShadowLengthWorld)
 {
 	float3 WorldPosition = TranslatedWorldPosition - View.PreViewTranslation;
-
 	FVirtualShadowMapSampleResult Sample =
 		SampleVirtualShadowMap(VirtualShadowMapId, WorldPosition, ContactShadowLengthWorld, EstimatedGeoWorldNormal);
-
-	if (Sample.bOccluded)
-	{
-		InOutHitDistance = Sample.OccluderDistance;
-	}
-
 	return Sample.ShadowFactor;
 }
 
-float4 EncodeOutput(
-	float3 TranslatedWorldPosition,
-	FLightShaderParameters Light,
-	uint LightType,
-	float RayMissFactor,
-	float HitDistance,
-	int SampleCount)
-{
-	// See signal encoding for SIGNAL_BUFFER_LAYOUT_PENUMBRA_RECONSTRUCTION in SSDSignalBufferEncoding.ush
-	// We fuse that logic in here to avoid needing another pass, but it's hard to call the functions directly as they depend on a large
-	// number of preprocessor defines and logic that we don't want to mirror here.
-	// Unfortunately that means this logic can get out of date when the denoiser changes, but that's the trade-off for now.
-	float4 OutputSignal = float4(RayMissFactor, SampleCount, WORLD_RADIUS_MISS, RayMissFactor);
-	if (RayMissFactor < 1.0f)
-	{
-		OutputSignal.z = ComputeLightSampleWorldBluringRadius(TranslatedWorldPosition, LightType, Light, HitDistance);
-	}
-	return OutputSignal;
-}
-
 float3 GetEstimatedGeoWorldNormal(float3 TranslatedWorldPosition, float2 ScreenUV)
 {
 	// Figure out slope, we do world space since that is the space where we might override using the shading normal...
@@ -193,7 +159,7 @@ RWTexture2D< uint > RWShadowMaskBits;
 // Single light per pass
 // Light parameters loaded via GetRootLightShaderParameters();
 int LightUniformVirtualShadowMapId;
-RWTexture2D< float4 > RWSignal;
+RWTexture2D< float4 > RWShadowFactor;
 
 // Type of input data consume by the page allocation (i.e., data read from the source buffer: Gbuffer, HairStrands data, ...)
 #define INPUT_TYPE_GBUFFER 0
@@ -203,8 +169,6 @@ uint InputType;
 struct FProjectLightResult
 {
 	half ShadowFactor;
-	// NOTE: The signal is encoded as per SSDSignalBufferEncoding.ush, SIGNAL_BUFFER_LAYOUT_PENUMBRA_RECONSTRUCTION
-	float4 Signal;
 };
 
 FProjectLightResult ProjectLight(
@@ -242,8 +206,6 @@ FProjectLightResult ProjectLight(
 #endif // DIRECTIONAL_LIGHT
 
 	float ShadowFactor = 1.0f;
-	float HitDistance = WORLD_RADIUS_MISS;
-	uint SampleCount = 0;
 
 	BRANCH
 	if (bInLightRegion)
@@ -255,7 +217,6 @@ FProjectLightResult ProjectLight(
 
 		if (SMRTRayCount > 0)
 		{
-			SampleCount = SMRTRayCount;
 		#if DIRECTIONAL_LIGHT
 			ShadowFactor = TraceDirectional(
 				VirtualShadowMapId,
@@ -267,8 +228,7 @@ FProjectLightResult ProjectLight(
 				SMRTRayCount,
 				SMRTSamplesPerRay,
 				SMRTRayLengthScale,
-				Noise,
-				/* out */ HitDistance);
+				Noise);
 		#else
 			ShadowFactor = TraceLocalLight(
 				VirtualShadowMapId,
@@ -282,20 +242,17 @@ FProjectLightResult ProjectLight(
 				SMRTSamplesPerRay,
 				SMRTCotMaxRayAngleFromLight,
 				Noise,
-				/*bCullBackface*/ !bIsHair,
-				/* out */ HitDistance);
+				/*bCullBackface*/ !bIsHair);
 		#endif // DIRECTIONAL_LIGHT
 		}
 		else
 		{
-			SampleCount = 1;
 			ShadowFactor = ProjectSingleSample(
 				VirtualShadowMapId,
 				L,
 				TranslatedWorldPosition,
 				EstimatedGeoWorldNormal,
-				ContactShadowLengthWorld,
-				HitDistance);
+				ContactShadowLengthWorld);
 		}
 
 		// Do not run contact shadow when computing the hair shadow mask (i.e. shadow mask applied on hair, has the scene 
@@ -306,7 +263,6 @@ FProjectLightResult ProjectLight(
 			if (ContactShadowHitDistance >= 0.0f)
 			{
 				ShadowFactor = 0.0f;
-				HitDistance = ContactShadowHitDistance;
 			}
 		}
 
@@ -389,7 +345,6 @@ FProjectLightResult ProjectLight(
 					VirtualVoxel.PageTexture,
 					TraversalSettings);
 
-				HitDistance = Result.HitT;
 				ShadowFactor = min(ShadowFactor, saturate(1 - Result.HairCount));
 			}
 		}
@@ -398,7 +353,6 @@ FProjectLightResult ProjectLight(
 
 	FProjectLightResult Result;
 	Result.ShadowFactor = ShadowFactor;
-	Result.Signal = EncodeOutput(TranslatedWorldPosition, Light, LightType, ShadowFactor, HitDistance, SampleCount);
 	return Result;
 }
 
@@ -489,6 +443,6 @@ void VirtualShadowMapProjection(
 		GBufferNormal,
 		Noise);
 
-	RWSignal[ PixelPos ] = Result.Signal;
+	RWShadowFactor[ PixelPos ] = Result.ShadowFactor;
 #endif // ONE_PASS_PROJECTION
 }
diff --git a/Engine/Source/Runtime/Renderer/Private/ScreenSpaceDenoise.cpp b/Engine/Source/Runtime/Renderer/Private/ScreenSpaceDenoise.cpp
--- a/Engine/Source/Runtime/Renderer/Private/ScreenSpaceDenoise.cpp
+++ b/Engine/Source/Runtime/Renderer/Private/ScreenSpaceDenoise.cpp
@@ -118,16 +118,6 @@ static TAutoConsoleVariable<float> CVarGIHistoryConvolutionKernelSpreadFactor(
 	TEXT("Multiplication factor applied on the kernel sample offset (default=3)."),
 	ECVF_RenderThreadSafe);
 
-static TAutoConsoleVariable<int32> CVarVirtualShadowReconstructionSampleCount(
-	TEXT("r.Shadow.Virtual.Denoiser.ReconstructionSamples"), 8,
-	TEXT("Maximum number of samples for the reconstruction pass (default = 8)."),
-	ECVF_RenderThreadSafe);
-
-static TAutoConsoleVariable<int32> CVarVirtualShadowPreConvolutionCount(
-	TEXT("r.Shadow.Virtual.Denoiser.PreConvolution"), 1,
-	TEXT("Number of pre-convolution passes (default = 1)."),
-	ECVF_RenderThreadSafe);
-
 /** The maximum number of mip level supported in the denoiser. */
 // TODO(Denoiser): jump to 3 because bufefr size already have a size multiple of 4.
 static const int32 kMaxMipLevel = 2;
@@ -152,7 +142,6 @@ DECLARE_GPU_STAT(ReflectionsDenoiser)
 DECLARE_GPU_STAT(ShadowsDenoiser)
 DECLARE_GPU_STAT(AmbientOcclusionDenoiser)
 DECLARE_GPU_STAT(DiffuseIndirectDenoiser)
-DECLARE_GPU_STAT(VirtualShadowsDenoiser)
 
 namespace
 {
@@ -206,9 +195,6 @@ enum class ESignalProcessing
 	// Denoise diffuse indirect hierarchy.
 	IndirectProbeHierarchy,
 
-	// Denoise a virtual shadow map mask.
-	VirtualShadowMapMask,
-
 	MAX,
 };
 
@@ -231,8 +217,7 @@ static bool UsesConstantPixelDensityPassLayout(ESignalProcessing SignalProcessin
 		SignalProcessing == ESignalProcessing::DiffuseAndAmbientOcclusion ||
 		SignalProcessing == ESignalProcessing::DiffuseSphericalHarmonic ||
 		SignalProcessing == ESignalProcessing::ScreenSpaceDiffuseIndirect ||
-		SignalProcessing == ESignalProcessing::IndirectProbeHierarchy ||
-		SignalProcessing == ESignalProcessing::VirtualShadowMapMask);
+		SignalProcessing == ESignalProcessing::IndirectProbeHierarchy);
 }
 
 /** Returns whether a signal processing support upscaling. */
@@ -264,8 +249,7 @@ static bool SignalUsesPreConvolution(ESignalProcessing SignalProcessing)
 		SignalProcessing == ESignalProcessing::ShadowVisibilityMask ||
 		SignalProcessing == ESignalProcessing::Reflections ||
 		SignalProcessing == ESignalProcessing::AmbientOcclusion ||
-		SignalProcessing == ESignalProcessing::DiffuseAndAmbientOcclusion ||
-		SignalProcessing == ESignalProcessing::VirtualShadowMapMask;
+		SignalProcessing == ESignalProcessing::DiffuseAndAmbientOcclusion;
 }
 
 /** Returns whether a signal processing uses a history rejection pre convolution pass. */
@@ -322,8 +306,7 @@ static int32 SignalMaxBatchSize(ESignalProcessing SignalProcessing)
 		SignalProcessing == ESignalProcessing::DiffuseAndAmbientOcclusion ||
 		SignalProcessing == ESignalProcessing::DiffuseSphericalHarmonic ||
 		SignalProcessing == ESignalProcessing::ScreenSpaceDiffuseIndirect ||
-		SignalProcessing == ESignalProcessing::IndirectProbeHierarchy ||
-		SignalProcessing == ESignalProcessing::VirtualShadowMapMask)
+		SignalProcessing == ESignalProcessing::IndirectProbeHierarchy)
 	{
 		return 1;
 	}
@@ -349,8 +332,7 @@ static bool SignalSupportMultiSPP(ESignalProcessing SignalProcessing)
 		SignalProcessing == ESignalProcessing::DiffuseAndAmbientOcclusion ||
 		SignalProcessing == ESignalProcessing::DiffuseSphericalHarmonic ||
 		SignalProcessing == ESignalProcessing::ScreenSpaceDiffuseIndirect ||
-		SignalProcessing == ESignalProcessing::IndirectProbeHierarchy ||
-		SignalProcessing == ESignalProcessing::VirtualShadowMapMask);
+		SignalProcessing == ESignalProcessing::IndirectProbeHierarchy);
 }
 
 
@@ -414,12 +396,6 @@ const TCHAR* const kInjestResourceNames[] = {
 	nullptr,
 	nullptr,
 	nullptr,
-
-	// VirtualShadowMapMask
-	TEXT("VirtualShadow.Denoiser.Injest0"),
-	TEXT("VirtualShadow.Denoiser.Injest1"),
-	nullptr,
-	nullptr,
 };
 
 const TCHAR* const kReduceResourceNames[] = {
@@ -470,12 +446,6 @@ const TCHAR* const kReduceResourceNames[] = {
 	nullptr,
 	nullptr,
 	nullptr,
-
-	// VirtualShadowMapMask
-	nullptr,
-	nullptr,
-	nullptr,
-	nullptr,
 };
 
 const TCHAR* const kReconstructionResourceNames[] = {
@@ -526,12 +496,6 @@ const TCHAR* const kReconstructionResourceNames[] = {
 	nullptr,
 	nullptr,
 	nullptr,
-
-	// VirtualShadowMapMask
-	TEXT("VirtualShadow.Denoiser.Reconstruction0"),
-	TEXT("VirtualShadow.Denoiser.Reconstruction1"),
-	TEXT("VirtualShadow.Denoiser.Reconstruction2"),
-	TEXT("VirtualShadow.Denoiser.Reconstruction3"),
 };
 
 const TCHAR* const kPreConvolutionResourceNames[] = {
@@ -582,12 +546,6 @@ const TCHAR* const kPreConvolutionResourceNames[] = {
 	nullptr,
 	nullptr,
 	nullptr,
-
-	// VirtualShadowMapMask
-	TEXT("VirtualShadow.Denoiser.PreConvolution0"),
-	TEXT("VirtualShadow.Denoiser.PreConvolution1"),
-	TEXT("VirtualShadow.Denoiser.PreConvolution2"),
-	TEXT("VirtualShadow.Denoiser.PreConvolution3"),
 };
 
 const TCHAR* const kRejectionPreConvolutionResourceNames[] = {
@@ -638,12 +596,6 @@ const TCHAR* const kRejectionPreConvolutionResourceNames[] = {
 	nullptr,
 	nullptr,
 	nullptr,
-
-	// VirtualShadowMapMask
-	TEXT("VirtualShadow.Denoiser.RejectionPreConvolution0"),
-	TEXT("VirtualShadow.Denoiser.RejectionPreConvolution1"),
-	TEXT("VirtualShadow.Denoiser.RejectionPreConvolution2"),
-	TEXT("VirtualShadow.Denoiser.RejectionPreConvolution3"),
 };
 
 const TCHAR* const kTemporalAccumulationResourceNames[] = {
@@ -694,12 +646,6 @@ const TCHAR* const kTemporalAccumulationResourceNames[] = {
 	TEXT("ProbeHierarchy.TemporalAccumulation1"),
 	TEXT("ProbeHierarchy.TemporalAccumulation2"),
 	nullptr,
-
-	// VirtualShadowMapMask
-	TEXT("VirtualShadow.Denoiser.TemporalAccumulation0"),
-	TEXT("VirtualShadow.Denoiser.TemporalAccumulation1"),
-	TEXT("VirtualShadow.Denoiser.TemporalAccumulation2"),
-	TEXT("VirtualShadow.Denoiser.TemporalAccumulation3"),
 };
 
 const TCHAR* const kHistoryConvolutionResourceNames[] = {
@@ -750,12 +696,6 @@ const TCHAR* const kHistoryConvolutionResourceNames[] = {
 	nullptr,
 	nullptr,
 	nullptr,
-
-	// VirtualShadowMapMask
-	TEXT("VirtualShadow.Denoiser.HistoryConvolution0"),
-	TEXT("VirtualShadow.Denoiser.HistoryConvolution1"),
-	TEXT("VirtualShadow.Denoiser.HistoryConvolution2"),
-	TEXT("VirtualShadow.Denoiser.HistoryConvolution3"),
 };
 
 const TCHAR* const kDenoiserOutputResourceNames[] = {
@@ -806,12 +746,6 @@ const TCHAR* const kDenoiserOutputResourceNames[] = {
 	nullptr,
 	nullptr,
 	nullptr,
-
-	// VirtualShadowMapMask
-	TEXT("VirtualShadow.Denoiser.DenoiserOutput0"),
-	TEXT("VirtualShadow.Denoiser.DenoiserOutput1"),
-	TEXT("VirtualShadow.Denoiser.DenoiserOutput2"),
-	TEXT("VirtualShadow.Denoiser.DenoiserOutput3"),
 };
 
 static_assert(UE_ARRAY_COUNT(kReconstructionResourceNames) == int32(ESignalProcessing::MAX) * kMaxBufferProcessingCount, "You forgot me!");
@@ -824,9 +758,7 @@ static_assert(UE_ARRAY_COUNT(kDenoiserOutputResourceNames) == int32(ESignalProce
 /** Returns whether should compile pipeline for a given shader platform.*/
 static bool ShouldCompileSignalPipeline(ESignalProcessing SignalProcessing, EShaderPlatform Platform)
 {
-	if (SignalProcessing == ESignalProcessing::ScreenSpaceDiffuseIndirect ||
-		SignalProcessing == ESignalProcessing::VirtualShadowMapMask
-		)
+	if (SignalProcessing == ESignalProcessing::ScreenSpaceDiffuseIndirect)
 	{
 		return Platform == SP_PCD3D_SM5 || FDataDrivenShaderPlatformInfo::GetCompileSignalProcessingPipeline(FStaticShaderPlatform(Platform)) || Platform == SP_METAL_SM5 || FDataDrivenShaderPlatformInfo::GetSupportsSSDIndirect(Platform);
 	}
@@ -1575,17 +1507,6 @@ static void DenoiseSignalAtConstantPixelDensity(
 			HistoryTextureCountPerSignal = 3;
 			bHasReconstructionLayoutDifferentFromHistory = true;
 		}
-		else if (Settings.SignalProcessing == ESignalProcessing::VirtualShadowMapMask)
-		{
-			check(Settings.SignalBatchSize == 1);
-
-			ReconstructionDescs[0].Format = PF_FloatRGBA;
-			HistoryDescs[0].Format = PF_FloatRGBA;
-
-			HistoryTextureCountPerSignal = 1;
-			ReconstructionTextureCount = 1;
-			bHasReconstructionLayoutDifferentFromHistory = false;
-		}
 		else
 		{
 			check(0);
@@ -1653,8 +1574,7 @@ static void DenoiseSignalAtConstantPixelDensity(
 
 	// Setup all the metadata to do spatial convolution.
 	FSSDConvolutionMetaData ConvolutionMetaData;
-	if (Settings.SignalProcessing == ESignalProcessing::ShadowVisibilityMask ||
-		Settings.SignalProcessing == ESignalProcessing::VirtualShadowMapMask
+	if (Settings.SignalProcessing == ESignalProcessing::ShadowVisibilityMask
 		)
 	{
 		for (int32 BatchedSignalId = 0; BatchedSignalId < Settings.SignalBatchSize; BatchedSignalId++)
@@ -3023,56 +2943,6 @@ FSSDSignalTextures IScreenSpaceDenoiser::DenoiseIndirectProbeHierarchy(
 	return SignalOutput;
 }
 
-// static
-FSSDSignalTextures IScreenSpaceDenoiser::DenoiseVirtualShadowMapMask(
-	FRDGBuilder& GraphBuilder,
-	const FViewInfo& View,
-	const FSceneTextureParameters& SceneTextures,
-	const FLightSceneInfo* LightSceneInfo,
-	FIntRect LightScissorRect,
-	const FVirtualShadowMapMaskInputs& InputParameters)
-{
-	RDG_GPU_STAT_SCOPE(GraphBuilder, VirtualShadowsDenoiser);
-	const FLightSceneProxy* Proxy = LightSceneInfo->Proxy;
-	const ULightComponent* LightComponent = Proxy->GetLightComponent();
-	ensure(IsSupportedLightType(ELightComponentType(Proxy->GetLightType())));
-	
-	FSSDConstantPixelDensitySettings Settings;
-	Settings.SignalProcessing = ESignalProcessing::VirtualShadowMapMask;
-	Settings.ReconstructionSamples = CVarVirtualShadowReconstructionSampleCount.GetValueOnRenderThread();
-	Settings.PreConvolutionCount = CVarVirtualShadowPreConvolutionCount.GetValueOnRenderThread();
-	Settings.bUseTemporalAccumulation = false;
-	
-	Settings.LightSceneInfo[0] = LightSceneInfo;
-	Settings.SignalScissor[0] = LightScissorRect;
-	// Force viewport to be a multiple of 2, to avoid over frame interference between TAA jitter of the frame, and Stackowiack's SampleTrackId.
-	Settings.FullResViewport = LightScissorRect;
-	Settings.FullResViewport.Min.X &= ~1;
-	Settings.FullResViewport.Min.Y &= ~1;
-	
-	// No temporal reprojection
-	TStaticArray<FScreenSpaceDenoiserHistory*, IScreenSpaceDenoiser::kMaxBatchSize> PrevHistories;
-	TStaticArray<FScreenSpaceDenoiserHistory*, IScreenSpaceDenoiser::kMaxBatchSize> NewHistories;
-	PrevHistories[0] = nullptr;
-	NewHistories[0] = nullptr;
-
-	FSSDSignalTextures InputSignal;
-	InputSignal.Textures[0] = InputParameters.Signal;
-
-	FViewInfoPooledRenderTargets ViewInfoPooledRenderTargets;
-	SetupSceneViewInfoPooledRenderTargets(View, &ViewInfoPooledRenderTargets);
-
-	FSSDSignalTextures SignalOutput;
-	DenoiseSignalAtConstantPixelDensity(
-		GraphBuilder, View, SceneTextures, ViewInfoPooledRenderTargets,
-		InputSignal, Settings,
-		PrevHistories,
-		NewHistories,
-		&SignalOutput);
-
-	return SignalOutput;
-}
-
 // static
 const IScreenSpaceDenoiser* IScreenSpaceDenoiser::GetDefaultDenoiser()
 {
diff --git a/Engine/Source/Runtime/Renderer/Private/ScreenSpaceDenoise.h b/Engine/Source/Runtime/Renderer/Private/ScreenSpaceDenoise.h
--- a/Engine/Source/Runtime/Renderer/Private/ScreenSpaceDenoise.h
+++ b/Engine/Source/Runtime/Renderer/Private/ScreenSpaceDenoise.h
@@ -163,11 +163,6 @@ public:
 		SHADER_PARAMETER_RDG_TEXTURE_UAV_ARRAY(RWTexture2D, SphericalHarmonic, [kSphericalHarmonicTextureCount])
 	END_SHADER_PARAMETER_STRUCT()
 
-	/** All the inputs of the virtual shadow map denoiser. */
-	BEGIN_SHADER_PARAMETER_STRUCT(FVirtualShadowMapMaskInputs, )
-		SHADER_PARAMETER_RDG_TEXTURE(Texture2D, Signal)
-	END_SHADER_PARAMETER_STRUCT()
-
 		
 	/** What the shadow ray tracing needs to output */
 	enum class EShadowRequirements
@@ -340,15 +335,6 @@ public:
 		FRDGTextureRef CompressedDepthTexture,
 		FRDGTextureRef CompressedShadingModelTexture);
 
-	/** Entry point to denoise virtual shadow map mask. */
-	static FSSDSignalTextures DenoiseVirtualShadowMapMask(
-		FRDGBuilder& GraphBuilder,
-		const FViewInfo& View,
-		const FSceneTextureParameters& SceneTextures,
-		const FLightSceneInfo* LightSceneInfo,
-		FIntRect LightScissorRect,
-		const FVirtualShadowMapMaskInputs& InputParameters);
-
 	/** Returns the interface of the default denoiser of the renderer. */
 	static const IScreenSpaceDenoiser* GetDefaultDenoiser();
 
diff --git a/Engine/Source/Runtime/Renderer/Private/ShadowRendering.cpp b/Engine/Source/Runtime/Renderer/Private/ShadowRendering.cpp
--- a/Engine/Source/Runtime/Renderer/Private/ShadowRendering.cpp
+++ b/Engine/Source/Runtime/Renderer/Private/ShadowRendering.cpp
@@ -188,19 +188,8 @@ static TAutoConsoleVariable<float> CVarShadowMaxSlopeScaleDepthBias(
 	TEXT("Higher values give better self-shadowing, but increase self-shadowing artifacts"),
 	ECVF_RenderThreadSafe);
 
-static TAutoConsoleVariable<int32> CVarVirtualSmDenoiserDirectional(
-	TEXT("r.Shadow.Virtual.Denoiser.EnableDirectional"),
-	0,
-	TEXT("Apply denoiser to directional light virtual shadow maps."),
-	ECVF_RenderThreadSafe);
 FForwardScreenSpaceShadowMaskTextureMobileOutputs GScreenSpaceShadowMaskTextureMobileOutputs;
 
-static TAutoConsoleVariable<int32> CVarVirtualSmDenoiserLocal(
-	TEXT("r.Shadow.Virtual.Denoiser.EnableLocal"),
-	0,
-	TEXT("Apply denoiser to local light virtual shadow maps."),
-	ECVF_RenderThreadSafe);
-
 ///////////////////////////////////////////////////////////////////////////////////////////////////
 // Hair
 static TAutoConsoleVariable<int32> CVarHairStrandsCullPerObjectShadowCaster(
@@ -2004,12 +1993,10 @@ void FDeferredShadingSceneRenderer::RenderShadowProjections(
 				if (ScissorRect.Area() > 0)
 				{
 					// Project virtual shadow maps
-					FRDGTextureRef SignalTexture;
-					bool bUseDenoiser = false;
+					FRDGTextureRef ShadowFactorTexture;
 					if (VisibleLightInfo.VirtualShadowMapClipmaps.Num() > 0)
 					{
-						bUseDenoiser = CVarVirtualSmDenoiserDirectional.GetValueOnRenderThread() != 0;
-						SignalTexture = RenderVirtualShadowMapProjection(
+						ShadowFactorTexture = RenderVirtualShadowMapProjection(
 							GraphBuilder,
 							SceneTextures,
 							View,
@@ -2020,9 +2007,8 @@ void FDeferredShadingSceneRenderer::RenderShadowProjections(
 					}
 					else
 					{
-						bUseDenoiser = CVarVirtualSmDenoiserLocal.GetValueOnRenderThread() != 0;
 						check(VirtualShadowMaps.Num() == 1);
-						SignalTexture = RenderVirtualShadowMapProjection(
+						ShadowFactorTexture = RenderVirtualShadowMapProjection(
 							GraphBuilder,
 							SceneTextures,
 							View,
@@ -2032,39 +2018,16 @@ void FDeferredShadingSceneRenderer::RenderShadowProjections(
 							VirtualShadowMaps[0]);
 					}
 
-					// Shadow filtering via denoiser
-					if (bUseDenoiser)
-					{
-						const IScreenSpaceDenoiser* Denoiser = IScreenSpaceDenoiser::GetDefaultDenoiser();
-
-						RDG_EVENT_SCOPE(GraphBuilder, "%s %dx%d",
-							Denoiser->GetDebugName(),
-							View.ViewRect.Width(), View.ViewRect.Height());
-
-						IScreenSpaceDenoiser::FVirtualShadowMapMaskInputs Inputs;
-						Inputs.Signal = SignalTexture;
-
-						FSSDSignalTextures DenoisedSignal = Denoiser->DenoiseVirtualShadowMapMask(
-							GraphBuilder,
-							View,
-							SceneTextureParameters,
-							LightSceneInfo,
-							ScissorRect,
-							Inputs);
-
-						SignalTexture = DenoisedSignal.Textures[0];
-					}
-
 					// Composite into screen shadow mask
-					CompositeVirtualShadowMapMask(GraphBuilder, ScissorRect, SignalTexture, ScreenShadowMaskTexture);
+					CompositeVirtualShadowMapMask(GraphBuilder, ScissorRect, ShadowFactorTexture, ScreenShadowMaskTexture);
 
 					// Sub-pixel shadow (no denoising for hair)
 					if (HairStrands::HasViewHairStrandsData(View) && ScreenShadowMaskSubPixelTexture)
 					{
-						FRDGTextureRef HairSignalTexture;
+						FRDGTextureRef HairShadowFactorTexture;
 						if (VisibleLightInfo.VirtualShadowMapClipmaps.Num() > 0)
 						{
-							HairSignalTexture = RenderVirtualShadowMapProjection(
+							HairShadowFactorTexture = RenderVirtualShadowMapProjection(
 								GraphBuilder,
 								SceneTextures,
 								View,
@@ -2076,7 +2039,7 @@ void FDeferredShadingSceneRenderer::RenderShadowProjections(
 						else
 						{
 							check(VirtualShadowMaps.Num() == 1);
-							HairSignalTexture = RenderVirtualShadowMapProjection(
+							HairShadowFactorTexture = RenderVirtualShadowMapProjection(
 								GraphBuilder,
 								SceneTextures,
 								View,
@@ -2087,7 +2050,7 @@ void FDeferredShadingSceneRenderer::RenderShadowProjections(
 						}
 
 						// Composite into sub pixel mask
-						CompositeVirtualShadowMapMask(GraphBuilder, ScissorRect, HairSignalTexture, ScreenShadowMaskSubPixelTexture);
+						CompositeVirtualShadowMapMask(GraphBuilder, ScissorRect, HairShadowFactorTexture, ScreenShadowMaskSubPixelTexture);
 					}
 				}
 			}
diff --git a/Engine/Source/Runtime/Renderer/Private/VirtualShadowMaps/VirtualShadowMapProjection.cpp b/Engine/Source/Runtime/Renderer/Private/VirtualShadowMaps/VirtualShadowMapProjection.cpp
--- a/Engine/Source/Runtime/Renderer/Private/VirtualShadowMaps/VirtualShadowMapProjection.cpp
+++ b/Engine/Source/Runtime/Renderer/Private/VirtualShadowMaps/VirtualShadowMapProjection.cpp
@@ -116,7 +116,7 @@ class FVirtualShadowMapProjectionCompositePS : public FGlobalShader
 	SHADER_USE_PARAMETER_STRUCT(FVirtualShadowMapProjectionCompositePS, FGlobalShader);
 
 	BEGIN_SHADER_PARAMETER_STRUCT(FParameters, )
-		SHADER_PARAMETER_RDG_TEXTURE(Texture2D<float4>, InputSignal)
+		SHADER_PARAMETER_RDG_TEXTURE(Texture2D<float4>, InputShadowFactor)
 		RENDER_TARGET_BINDING_SLOTS()
 	END_SHADER_PARAMETER_STRUCT()
 
@@ -142,7 +142,7 @@ void CompositeVirtualShadowMapMask(
 	auto ShaderMap = GetGlobalShaderMap(GMaxRHIFeatureLevel);
 
 	FVirtualShadowMapProjectionCompositePS::FParameters* PassParameters = GraphBuilder.AllocParameters<FVirtualShadowMapProjectionCompositePS::FParameters>();
-	PassParameters->InputSignal = Input;
+	PassParameters->InputShadowFactor = Input;
 
 	PassParameters->RenderTargets[0] = FRenderTargetBinding(OutputShadowMaskTexture, ERenderTargetLoadAction::ELoad);
 
@@ -154,7 +154,7 @@ void CompositeVirtualShadowMapMask(
 
 	FPixelShaderUtils::AddFullscreenPass(GraphBuilder,
 		ShaderMap,
-		RDG_EVENT_NAME("Mask Composite"),
+		RDG_EVENT_NAME("MaskComposite"),
 		PixelShader,
 		PassParameters,
 		ScissorRect,
@@ -202,7 +202,7 @@ class FVirtualShadowMapProjectionCS : public FGlobalShader
 		// Pass per light parameters
 		SHADER_PARAMETER_STRUCT(FLightShaderParameters, Light)
 		SHADER_PARAMETER(int32, LightUniformVirtualShadowMapId)
-		SHADER_PARAMETER_RDG_TEXTURE_UAV(RWTexture2D, RWSignal)
+		SHADER_PARAMETER_RDG_TEXTURE_UAV(RWTexture2D, RWShadowFactor)
 	END_SHADER_PARAMETER_STRUCT()
 
 	static void ModifyCompilationEnvironment( const FGlobalShaderPermutationParameters& Parameters, FShaderCompilerEnvironment& OutEnvironment )
@@ -292,7 +292,7 @@ static void RenderVirtualShadowMapProjectionCommon(
 		LightProxy->GetLightShaderParameters(LightParameters);
 		PassParameters->Light = LightParameters;
 		PassParameters->LightUniformVirtualShadowMapId = VirtualShadowMapId;
-		PassParameters->RWSignal = GraphBuilder.CreateUAV( OutputTexture );
+		PassParameters->RWShadowFactor = GraphBuilder.CreateUAV( OutputTexture );
 	}
  
 	if (bDirectionalLight)
@@ -363,17 +363,17 @@ FRDGTextureRef RenderVirtualShadowMapProjectionOnePass(
 	return ShadowMaskBits;
 }
 
-static FRDGTextureRef CreateSignalTexture(FRDGBuilder& GraphBuilder, FIntPoint Extent)
+static FRDGTextureRef CreateShadowFactorTexture(FRDGBuilder& GraphBuilder, FIntPoint Extent)
 {
 	const FLinearColor ClearColor(0.0f, 0.0f, 0.0f, 0.0f);
 
 	FRDGTextureDesc Desc = FRDGTextureDesc::Create2D(
 		Extent,
-		PF_FloatRGBA,
+		PF_R32_FLOAT,
 		FClearValueBinding(ClearColor),
 		TexCreate_ShaderResource | TexCreate_UAV);
 
-	FRDGTextureRef Texture = GraphBuilder.CreateTexture(Desc, TEXT("Shadow.Virtual.Signal"));
+	FRDGTextureRef Texture = GraphBuilder.CreateTexture(Desc, TEXT("Shadow.Virtual.ShadowFactor"));
 	AddClearUAVPass(GraphBuilder, GraphBuilder.CreateUAV(Texture), ClearColor);
 	return Texture;
 }
@@ -387,7 +387,7 @@ FRDGTextureRef RenderVirtualShadowMapProjection(
 	EVirtualShadowMapProjectionInputType InputType,
 	FProjectedShadowInfo* ShadowInfo)
 {	
-	FRDGTextureRef SignalTexture = CreateSignalTexture(GraphBuilder, SceneTextures.Config.Extent);
+	FRDGTextureRef OutputTexture = CreateShadowFactorTexture(GraphBuilder, SceneTextures.Config.Extent);
 	RenderVirtualShadowMapProjectionCommon(
 		GraphBuilder,
 		SceneTextures,
@@ -395,10 +395,10 @@ FRDGTextureRef RenderVirtualShadowMapProjection(
 		VirtualShadowMapArray,
 		ScissorRect,
 		InputType,
-		SignalTexture,
+		OutputTexture,
 		ShadowInfo->GetLightSceneInfo().Proxy,
 		ShadowInfo->VirtualShadowMaps[0]->ID);
-	return SignalTexture;
+	return OutputTexture;
 }
 
 FRDGTextureRef RenderVirtualShadowMapProjection(
@@ -410,7 +410,7 @@ FRDGTextureRef RenderVirtualShadowMapProjection(
 	EVirtualShadowMapProjectionInputType InputType,
 	const TSharedPtr<FVirtualShadowMapClipmap>& Clipmap)
 {
-	FRDGTextureRef SignalTexture = CreateSignalTexture(GraphBuilder, SceneTextures.Config.Extent);
+	FRDGTextureRef OutputTexture = CreateShadowFactorTexture(GraphBuilder, SceneTextures.Config.Extent);
 	RenderVirtualShadowMapProjectionCommon(
 		GraphBuilder,
 		SceneTextures,
@@ -418,8 +418,8 @@ FRDGTextureRef RenderVirtualShadowMapProjection(
 		VirtualShadowMapArray,		
 		ScissorRect,
 		InputType,
-		SignalTexture,
+		OutputTexture,
 		Clipmap->GetLightSceneInfo().Proxy,
 		Clipmap->GetVirtualShadowMap()->ID);
-	return SignalTexture;
+	return OutputTexture;
 }
