Fix enum name typo.

diff --git a/Engine/Source/Runtime/Renderer/Private/CompositionLighting/CompositionLighting.cpp b/Engine/Source/Runtime/Renderer/Private/CompositionLighting/CompositionLighting.cpp
--- a/Engine/Source/Runtime/Renderer/Private/CompositionLighting/CompositionLighting.cpp
+++ b/Engine/Source/Runtime/Renderer/Private/CompositionLighting/CompositionLighting.cpp
@@ -514,7 +514,7 @@ void ProcessBeforeBasePass(
 			if (bDBuffer)
 			{
 				FDeferredDecalPassTextures DecalPassTextures = GetDeferredDecalPassTextures(GraphBuilder, SceneTextures, &DBufferTextures);
-				AddDeferredDecalPass(GraphBuilder, View, DecalPassTextures, EDecalRenderStage::BeforeBassPass);
+				AddDeferredDecalPass(GraphBuilder, View, DecalPassTextures, EDecalRenderStage::BeforeBasePass);
 			}
 
 			if (bNeedSSAO)
diff --git a/Engine/Source/Runtime/Renderer/Private/CompositionLighting/PostProcessDeferredDecals.cpp b/Engine/Source/Runtime/Renderer/Private/CompositionLighting/PostProcessDeferredDecals.cpp
--- a/Engine/Source/Runtime/Renderer/Private/CompositionLighting/PostProcessDeferredDecals.cpp
+++ b/Engine/Source/Runtime/Renderer/Private/CompositionLighting/PostProcessDeferredDecals.cpp
@@ -420,7 +420,7 @@ static FRHIDepthStencilState* GetDecalDepthState(uint32& StencilRef, FDecalDepth
 
 static bool IsStencilOptimizationAvailable(EDecalRenderStage RenderStage)
 {
-	return RenderStage == EDecalRenderStage::BeforeLighting || RenderStage == EDecalRenderStage::BeforeBassPass || RenderStage == EDecalRenderStage::Emissive;
+	return RenderStage == EDecalRenderStage::BeforeLighting || RenderStage == EDecalRenderStage::BeforeBasePass || RenderStage == EDecalRenderStage::Emissive;
 }
 
 static EDecalDBufferMaskTechnique GetDBufferMaskTechnique(EShaderPlatform ShaderPlatform)
@@ -444,7 +444,7 @@ static const TCHAR* GetStageName(EDecalRenderStage Stage)
 {
 	switch (Stage)
 	{
-	case EDecalRenderStage::BeforeBassPass: return TEXT("BeforeBassPass");
+	case EDecalRenderStage::BeforeBasePass: return TEXT("BeforeBasePass");
 	case EDecalRenderStage::BeforeLighting: return TEXT("BeforeLighting");
 	case EDecalRenderStage::Mobile: return TEXT("Mobile");
 	case EDecalRenderStage::MobileBeforeLighting: return TEXT("MobileBeforeLighting");
@@ -503,7 +503,7 @@ void AddDeferredDecalPass(
 	EDecalRenderStage DecalRenderStage)
 {
 	check(PassTextures.Depth.IsValid());
-	check(DecalRenderStage != EDecalRenderStage::BeforeBassPass || PassTextures.DBufferTextures);
+	check(DecalRenderStage != EDecalRenderStage::BeforeBasePass || PassTextures.DBufferTextures);
 
 	const FSceneViewFamily& ViewFamily = *(View.Family);
 
@@ -523,7 +523,7 @@ void AddDeferredDecalPass(
 	FTransientDecalRenderDataList* SortedDecals = nullptr;
 
 	checkf(DecalRenderStage != EDecalRenderStage::AmbientOcclusion || PassTextures.ScreenSpaceAO, TEXT("Attepting to render AO decals without SSAO having emitted a valid render target."));
-	checkf(DecalRenderStage != EDecalRenderStage::BeforeBassPass || IsUsingDBuffers(ShaderPlatform), TEXT("Only DBuffer decals are supported before the base pass."));
+	checkf(DecalRenderStage != EDecalRenderStage::BeforeBasePass || IsUsingDBuffers(ShaderPlatform), TEXT("Only DBuffer decals are supported before the base pass."));
 
 	if (DecalCount)
 	{
@@ -603,7 +603,7 @@ void AddDeferredDecalPass(
 	{
 		RDG_EVENT_SCOPE(GraphBuilder, "DeferredDecals %s", GetStageName(DecalRenderStage));
 
-		if (MeshDecalCount > 0 && (DecalRenderStage == EDecalRenderStage::BeforeBassPass || DecalRenderStage == EDecalRenderStage::BeforeLighting || DecalRenderStage == EDecalRenderStage::Emissive || DecalRenderStage == EDecalRenderStage::AmbientOcclusion))
+		if (MeshDecalCount > 0 && (DecalRenderStage == EDecalRenderStage::BeforeBasePass || DecalRenderStage == EDecalRenderStage::BeforeLighting || DecalRenderStage == EDecalRenderStage::Emissive || DecalRenderStage == EDecalRenderStage::AmbientOcclusion))
 		{
 			RenderMeshDecals(GraphBuilder, View, PassTextures, DecalRenderStage);
 		}
@@ -637,7 +637,7 @@ void AddDeferredDecalPass(
 
 	// Last D-Buffer pass in the frame decodes the write mask (if supported and decals were rendered).
 	if (DBufferMaskTechnique == EDecalDBufferMaskTechnique::WriteMask &&
-		DecalRenderStage == EDecalRenderStage::BeforeBassPass &&
+		DecalRenderStage == EDecalRenderStage::BeforeBasePass &&
 		PassTextures.DBufferTextures->IsValid() &&
 		View.IsLastInFamily())
 	{
diff --git a/Engine/Source/Runtime/Renderer/Private/CompositionLighting/PostProcessMeshDecals.cpp b/Engine/Source/Runtime/Renderer/Private/CompositionLighting/PostProcessMeshDecals.cpp
--- a/Engine/Source/Runtime/Renderer/Private/CompositionLighting/PostProcessMeshDecals.cpp
+++ b/Engine/Source/Runtime/Renderer/Private/CompositionLighting/PostProcessMeshDecals.cpp
@@ -322,7 +322,7 @@ void RenderMeshDecals(
 
 	switch (DecalRenderStage)
 	{
-	case EDecalRenderStage::BeforeBassPass:
+	case EDecalRenderStage::BeforeBasePass:
 		DrawDecalMeshCommands(GraphBuilder, View, DecalPassTextures, DecalRenderStage, EDecalRenderTargetMode::DBuffer);
 		break;
 
diff --git a/Engine/Source/Runtime/Renderer/Private/DecalRenderingCommon.cpp b/Engine/Source/Runtime/Renderer/Private/DecalRenderingCommon.cpp
--- a/Engine/Source/Runtime/Renderer/Private/DecalRenderingCommon.cpp
+++ b/Engine/Source/Runtime/Renderer/Private/DecalRenderingCommon.cpp
@@ -52,7 +52,7 @@ namespace DecalRendering
 		}
 
 		// Convert GBuffer modes to DBuffer.
-		uint32 DBufferStageMask = 1 << (uint32)EDecalRenderStage::BeforeBassPass;
+		uint32 DBufferStageMask = 1 << (uint32)EDecalRenderStage::BeforeBasePass;
 		DBufferStageMask |= bWriteEmissive ? (1 << (uint32)EDecalRenderStage::Emissive) : 0;
 
 		if (bIsDBufferPlatform && bIsForwardPlatform && !bIsMobilePlatform)
@@ -65,7 +65,7 @@ namespace DecalRendering
 				break;
 			case DBM_Normal:
 				DecalBlendMode = DBM_DBuffer_Normal;
-				DBufferStageMask = 1 << (uint32)EDecalRenderStage::BeforeBassPass;
+				DBufferStageMask = 1 << (uint32)EDecalRenderStage::BeforeBasePass;
 				break;
 			case DBM_Emissive:
 				DecalBlendMode = DBM_DBuffer_Emissive;
@@ -222,9 +222,9 @@ namespace DecalRendering
 
 	EDecalRenderStage GetBaseRenderStage(FDecalBlendDesc DecalBlendDesc)
 	{
-		if (DecalBlendDesc.RenderStageMask & (1 << (uint32)EDecalRenderStage::BeforeBassPass))
+		if (DecalBlendDesc.RenderStageMask & (1 << (uint32)EDecalRenderStage::BeforeBasePass))
 		{
-			return EDecalRenderStage::BeforeBassPass;
+			return EDecalRenderStage::BeforeBasePass;
 		}
 		if (DecalBlendDesc.RenderStageMask & (1 << (uint32)EDecalRenderStage::BeforeLighting))
 		{
@@ -246,7 +246,7 @@ namespace DecalRendering
 	{
 		switch(DecalRenderStage)
 		{
-			case EDecalRenderStage::BeforeBassPass:
+			case EDecalRenderStage::BeforeBasePass:
 				return EDecalRenderTargetMode::DBuffer;
 			case EDecalRenderStage::BeforeLighting:
 				return DecalBlendDesc.bWriteNormal ? EDecalRenderTargetMode::SceneColorAndGBuffer : EDecalRenderTargetMode::SceneColorAndGBufferNoNormal;
@@ -955,7 +955,7 @@ namespace DecalRendering
 		const uint32 RenderTargetWriteMask = GetRenderTargetWriteMask(DecalBlendDesc, DecalRenderStage, RenderTargetMode);
 
 		OutEnvironment.SetDefine(TEXT("IS_DECAL"), 1);
-		OutEnvironment.SetDefine(TEXT("IS_DBUFFER_DECAL"), DecalRenderStage == EDecalRenderStage::BeforeBassPass ? 1 : 0);
+		OutEnvironment.SetDefine(TEXT("IS_DBUFFER_DECAL"), DecalRenderStage == EDecalRenderStage::BeforeBasePass ? 1 : 0);
 
 		OutEnvironment.SetDefine(TEXT("DECAL_RENDERSTAGE"), (uint32)DecalRenderStage);
 		OutEnvironment.SetDefine(TEXT("DECAL_RENDERTARGETMODE"), (uint32)RenderTargetMode);
@@ -966,7 +966,7 @@ namespace DecalRendering
 		OutEnvironment.SetDefine(TEXT("DECAL_OUT_MRT2"), (RenderTargetWriteMask & 4) != 0 ? 1 : 0);
 		OutEnvironment.SetDefine(TEXT("DECAL_OUT_MRT3"), (RenderTargetWriteMask & 8) != 0 ? 1 : 0);
 
-		OutEnvironment.SetDefine(TEXT("DECAL_RENDERSTAGE_BEFOREBASEPASS"), (uint32)EDecalRenderStage::BeforeBassPass);
+		OutEnvironment.SetDefine(TEXT("DECAL_RENDERSTAGE_BEFOREBASEPASS"), (uint32)EDecalRenderStage::BeforeBasePass);
 		OutEnvironment.SetDefine(TEXT("DECAL_RENDERSTAGE_BEFORELIGHTING"), (uint32)EDecalRenderStage::BeforeLighting);
 		OutEnvironment.SetDefine(TEXT("DECAL_RENDERSTAGE_MOBILE"), (uint32)EDecalRenderStage::Mobile);
 		OutEnvironment.SetDefine(TEXT("DECAL_RENDERSTAGE_MOBILEBEFORELIGHTING"), (uint32)EDecalRenderStage::MobileBeforeLighting);
diff --git a/Engine/Source/Runtime/Renderer/Private/DecalRenderingCommon.h b/Engine/Source/Runtime/Renderer/Private/DecalRenderingCommon.h
--- a/Engine/Source/Runtime/Renderer/Private/DecalRenderingCommon.h
+++ b/Engine/Source/Runtime/Renderer/Private/DecalRenderingCommon.h
@@ -34,7 +34,7 @@ enum class EDecalRenderStage : uint8
 	None = 0,
 
 	// DBuffer decal pass.
-	BeforeBassPass = 1,
+	BeforeBasePass = 1,
 	// GBuffer decal pass.
 	BeforeLighting = 2,
 
