Convert decal materials to no longer use decal blend mode.
Decal blend properties are inferred from material blend mode and connected outputs.
Decal materials are converted on PostLoad but need to be saved for the conversion to persist.

diff --git a/Engine/Source/Runtime/Core/Public/UObject/UE5MainStreamObjectVersion.h b/Engine/Source/Runtime/Core/Public/UObject/UE5MainStreamObjectVersion.h
--- a/Engine/Source/Runtime/Core/Public/UObject/UE5MainStreamObjectVersion.h
+++ b/Engine/Source/Runtime/Core/Public/UObject/UE5MainStreamObjectVersion.h
@@ -121,6 +121,9 @@ struct CORE_API FUE5MainStreamObjectVersion
 		// Explicitly serialized bSavedCachedExpressionData for Material(Instance)
 		MaterialSavedCachedData,
 
+		// Remove explicit decal blend mode
+		RemoveDecalBlendMode,
+
         // -----<new versions can be added above this line>-------------------------------------------------
 		VersionPlusOne,
 		LatestVersion = VersionPlusOne - 1
diff --git a/Engine/Source/Runtime/Engine/Classes/Materials/Material.h b/Engine/Source/Runtime/Engine/Classes/Materials/Material.h
--- a/Engine/Source/Runtime/Engine/Classes/Materials/Material.h
+++ b/Engine/Source/Runtime/Engine/Classes/Materials/Material.h
@@ -107,26 +107,6 @@ enum EDecalBlendMode
 	DBM_MAX,
 };
 
-inline bool IsDBufferDecalBlendMode(EDecalBlendMode In)
-{
-	switch(In)
-	{
-		case DBM_DBuffer_ColorNormalRoughness:
-		case DBM_DBuffer_Color:
-		case DBM_DBuffer_ColorNormal:
-		case DBM_DBuffer_ColorRoughness:
-		case DBM_DBuffer_Normal:
-		case DBM_DBuffer_NormalRoughness:
-		case DBM_DBuffer_Roughness:
-		case DBM_DBuffer_Emissive:
-		case DBM_DBuffer_AlphaComposite:
-		case DBM_DBuffer_EmissiveAlphaComposite:
-			return true;
-	}
-
-	return false;
-}
-
 /** Defines how the material reacts on DBuffer decals, later we can expose more variants between None and Default. */
 UENUM()
 enum EMaterialDecalResponse
@@ -134,19 +114,19 @@ enum EMaterialDecalResponse
 	/** Do not receive decals (Later we still can read the DBuffer channels to customize the effect, this frees up some interpolators). */
 	MDR_None UMETA(DisplayName="None"),
 
-	/** Receive Decals, applies all DBuffer channels, assumes the decal is non metal and mask the subsurface scattering. */
+	/** Receive Decals, applies all DBuffer channels. */
 	MDR_ColorNormalRoughness UMETA(DisplayName="Color Normal Roughness"),
-	/** Receive Decals, applies color DBuffer channels, assumes the decal is non metal and mask the subsurface scattering. */
+	/** Receive Decals, applies color DBuffer channels. */
 	MDR_Color UMETA(DisplayName="Color"),
-	/** Receive Decals, applies all DBuffer channels, assumes the decal is non metal and mask the subsurface scattering. */
+	/** Receive Decals, applies color and normal DBuffer channels. */
 	MDR_ColorNormal UMETA(DisplayName="Color Normal"),
-	/** Receive Decals, applies all DBuffer channels, assumes the decal is non metal and mask the subsurface scattering. */
+	/** Receive Decals, applies color, roughness, specular and metallic DBuffer channels. */
 	MDR_ColorRoughness UMETA(DisplayName="Color Roughness"),
-	/** Receive Decals, applies all DBuffer channels, assumes the decal is non metal and mask the subsurface scattering. */
+	/** Receive Decals, applies normal DBuffer channels. */
 	MDR_Normal UMETA(DisplayName="Normal"),
-	/** Receive Decals, applies all DBuffer channels, assumes the decal is non metal and mask the subsurface scattering. */
+	/** Receive Decals, applies normal, roughness, specular and metallic DBuffer channels. */
 	MDR_NormalRoughness UMETA(DisplayName="Normal Roughness"),
-	/** Receive Decals, applies all DBuffer channels, assumes the decal is non metal and mask the subsurface scattering. */
+	/** Receive Decals, applies roughness, specular and metallic DBuffer channels. */
 	MDR_Roughness UMETA(DisplayName="Roughness"),
 	MDR_MAX
 };
@@ -385,8 +365,7 @@ class UMaterial : public UMaterialInterface
 	UPROPERTY(BlueprintReadOnly, EditAnywhere, Category=Material, AssetRegistrySearchable)
 	TEnumAsByte<enum EBlendMode> BlendMode;
 
-	/** Defines how the GBuffer chanels are getting manipulated by a decal material pass. (only with MaterialDomain == MD_DeferredDecal) */
-	UPROPERTY(EditAnywhere, Category=Material)
+	UPROPERTY(meta = (DeprecatedProperty, DeprecationMessage = "No longer used."))
 	TEnumAsByte<enum EDecalBlendMode> DecalBlendMode;
 
 	/** 
@@ -1200,6 +1179,7 @@ private:
 
 	void BackwardsCompatibilityInputConversion();
 	void BackwardsCompatibilityVirtualTextureOutputConversion();
+	void BackwardsCompatibilityDecalConversion();
 	void ConvertMaterialToStrataMaterial();
 
 	/** Handles setting up an annotation for this object if a flag has changed value */
@@ -1349,7 +1329,7 @@ public:
 	ENGINE_API virtual bool GetStaticSwitchParameterDefaultValue(const FHashedMaterialParameterInfo& ParameterInfo, bool& OutValue, FGuid& OutExpressionGuid, bool bCheckOwnedGlobalOverrides = false) const override;
 #endif // WITH_EDITOR
 
-	/** Returns the material's decal blend mode, calculated from the DecalBlendMode property and what inputs are connected. */
+	UE_DEPRECATED(5.0, "No longer used.")
 	uint32 GetDecalBlendMode() const { return DecalBlendMode; }
 
 	/** Returns the material's decal response mode */
diff --git a/Engine/Source/Runtime/Engine/Private/Materials/HLSLMaterialTranslator.cpp b/Engine/Source/Runtime/Engine/Private/Materials/HLSLMaterialTranslator.cpp
--- a/Engine/Source/Runtime/Engine/Private/Materials/HLSLMaterialTranslator.cpp
+++ b/Engine/Source/Runtime/Engine/Private/Materials/HLSLMaterialTranslator.cpp
@@ -1023,12 +1023,12 @@ bool FHLSLMaterialTranslator::Translate()
 			}
 		}
 
-		if (Domain == MD_DeferredDecal && BlendMode != BLEND_Translucent)
+		if (Domain == MD_DeferredDecal && !(BlendMode == BLEND_Translucent || BlendMode == BLEND_AlphaComposite || BlendMode == BLEND_Modulate))
 		{
 			// We could make the change for the user but it would be confusing when going to DeferredDecal and back
 			// or we would have to pay a performance cost to make the change more transparently.
 			// The change saves performance as with translucency we don't need to test for MeshDecals in all opaque rendering passes
-			Errorf(TEXT("Material using the DeferredDecal domain need to use the BlendModel Translucent (this saves performance)"));
+			Errorf(TEXT("Material using the DeferredDecal domain can only use the BlendModes Translucent, AlphaComposite or Modulate"));
 		}
 
 		if (MaterialCompilationOutput.bNeedsSceneTextures)
@@ -5907,25 +5907,40 @@ int32 FHLSLMaterialTranslator::SceneTextureLookup(int32 ViewportUV, uint32 InSce
 									SceneTextureId == PPI_SceneDepth ||
 									SceneTextureId == PPI_CustomStencil;
 
-	// special case for DBuffer normals
-	bool bIsUsingReprojectedNormal = false;
+	if (!bSupportedOnMobile && ErrorUnlessFeatureLevelSupported(ERHIFeatureLevel::SM5) == INDEX_NONE)
+	{
+		return INDEX_NONE;
+	}
 
-	if(Material->GetMaterialDomain() == MD_DeferredDecal)
+	if (ShaderFrequency != SF_Pixel && ShaderFrequency != SF_Vertex)
 	{
-		EDecalBlendMode DecalBlendMode = (EDecalBlendMode)Material->GetDecalBlendMode();
-		bool bDBuffer = IsDBufferDecalBlendMode(DecalBlendMode);
+		// we can relax this later if needed
+		return NonPixelShaderExpressionError();
+	}
 
-		if(bDBuffer)
+	// Special case for DBuffer normals.
+	// Note that full reprojection only happens with r.Decal.NormalReprojectionEnabled 1.
+	// Otherwise we use a depth buffer approximation.
+	bool bIsUsingReprojectedNormal = false;
+	if (Material->GetMaterialDomain() == MD_DeferredDecal)
+	{
+		if (SceneTextureId == PPI_WorldNormal)
 		{
-			if(SceneTextureId == PPI_WorldNormal)
+			const bool bUsingDBuffer = IsUsingDBuffers(Platform);
+			const bool bIsDBufferOutput = 
+				MaterialProperty == MP_BaseColor ||
+				MaterialProperty == MP_Normal || 
+				MaterialProperty == MP_Roughness || 
+				MaterialProperty == MP_Specular || 
+				MaterialProperty == MP_Metallic;
+
+			if (bUsingDBuffer && bIsDBufferOutput)
 			{
 				bIsUsingReprojectedNormal = true;
 			}
 		}
 	}
 
-	// The normal is valid, sample from previous frame with reprojection. Otherwise, return the original geometry normal, both for SM5 when the normal is missing,
-	// but also on Mobile as a fallback.
 	if (bIsUsingReprojectedNormal)
 	{
 		int32 BufferUV;
@@ -5955,23 +5970,14 @@ int32 FHLSLMaterialTranslator::SceneTextureLookup(int32 ViewportUV, uint32 InSce
 		}
 		else
 		{
+			// Return the original geometry normal as a fall back.
+			// Shouldn't get here since DBuffer isn't enabled on mobile rendering paths?
 			return AddCodeChunk(MCT_Float4,
 				TEXT("float4(%s,1.0f)"),
 				*CoerceParameter(PixelNormalWS(),MCT_Float3));
 		}
 	}
 
-	if (!bSupportedOnMobile && ErrorUnlessFeatureLevelSupported(ERHIFeatureLevel::SM5) == INDEX_NONE)
-	{
-		return INDEX_NONE;
-	}
-
-	if (ShaderFrequency != SF_Pixel && ShaderFrequency != SF_Vertex)
-	{
-		// we can relax this later if needed
-		return NonPixelShaderExpressionError();
-	}
-		
 	if (SceneTextureId == PPI_DecalMask)
 	{
 		return Error(TEXT("Decal Mask bit was move out of GBuffer to the stencil buffer for performance optimisation and is therefor no longer available"));
@@ -6033,26 +6039,12 @@ void FHLSLMaterialTranslator::UseSceneTextureId(ESceneTextureId SceneTextureId,
 	MaterialCompilationOutput.bNeedsSceneTextures = true;
 	MaterialCompilationOutput.SetIsSceneTextureUsed(SceneTextureId);
 
-	if(Material->GetMaterialDomain() == MD_DeferredDecal)
-	{
-		EDecalBlendMode DecalBlendMode = (EDecalBlendMode)Material->GetDecalBlendMode();
-		bool bDBuffer = IsDBufferDecalBlendMode(DecalBlendMode);
-
-		bool bRequiresSM5 = (SceneTextureId == PPI_WorldNormal || SceneTextureId == PPI_CustomDepth || SceneTextureId == PPI_CustomStencil || SceneTextureId == PPI_AmbientOcclusion);
-
-		if(bDBuffer)
+	if (Material->GetMaterialDomain() == MD_DeferredDecal)
 	{
-			if(!(SceneTextureId == PPI_SceneDepth || SceneTextureId == PPI_CustomDepth || SceneTextureId == PPI_CustomStencil || SceneTextureId == PPI_WorldNormal))
+		if (!(SceneTextureId == PPI_SceneDepth || SceneTextureId == PPI_WorldNormal || SceneTextureId == PPI_CustomDepth || SceneTextureId == PPI_CustomStencil))
 		{
 			// Note: For DBuffer decals: CustomDepth and CustomStencil are only available if r.CustomDepth.Order = 0
-				Errorf(TEXT("DBuffer decals (MaterialDomain=DeferredDecal and DecalBlendMode is using DBuffer) can only access SceneDepth, CustomDepth, CustomStencil, and WorldNormal"));
-			}
-		}
-		else
-		{
-			if(!(SceneTextureId == PPI_SceneDepth || SceneTextureId == PPI_CustomDepth || SceneTextureId == PPI_CustomStencil || SceneTextureId == PPI_WorldNormal || SceneTextureId == PPI_AmbientOcclusion))
-			{
-				Errorf(TEXT("Decals (MaterialDomain=DeferredDecal) can only access WorldNormal, AmbientOcclusion, SceneDepth, CustomDepth, CustomStencil"));
+			Errorf(TEXT("Decals can only access SceneDepth, CustomDepth, CustomStencil, and WorldNormal"));
 		}
 
 		if (SceneTextureId == PPI_WorldNormal && Material->HasNormalConnected())
@@ -6060,8 +6052,8 @@ void FHLSLMaterialTranslator::UseSceneTextureId(ESceneTextureId SceneTextureId,
 			// GBuffer can only relate to WorldNormal here.
 			Errorf(TEXT("Decals that read WorldNormal cannot output to normal at the same time"));
 		}
-		}
 
+		const bool bRequiresSM5 = SceneTextureId == PPI_WorldNormal || SceneTextureId == PPI_CustomDepth || SceneTextureId == PPI_CustomStencil;
 		if (bRequiresSM5)
 		{
 			ErrorUnlessFeatureLevelSupported(ERHIFeatureLevel::SM5);
diff --git a/Engine/Source/Runtime/Engine/Private/Materials/Material.cpp b/Engine/Source/Runtime/Engine/Private/Materials/Material.cpp
--- a/Engine/Source/Runtime/Engine/Private/Materials/Material.cpp
+++ b/Engine/Source/Runtime/Engine/Private/Materials/Material.cpp
@@ -162,6 +162,7 @@ FUObjectAnnotationSparse<FMaterialsWithDirtyUsageFlags,true> GMaterialsWithDirty
 FUObjectAnnotationSparseBool GMaterialsThatNeedExpressionsFlipped;
 FUObjectAnnotationSparseBool GMaterialsThatNeedCoordinateCheck;
 FUObjectAnnotationSparseBool GMaterialsThatNeedCommentFix;
+FUObjectAnnotationSparseBool GMaterialsThatNeedDecalFix;
 
 #endif // #if WITH_EDITOR
 
@@ -3420,6 +3421,17 @@ void UMaterial::Serialize(FArchive& Ar)
 			BlendMode = BLEND_Translucent;
 		}
 	}
+	
+#if WITH_EDITOR
+	if (Ar.IsLoading() && Ar.CustomVer(FUE5MainStreamObjectVersion::GUID) < FUE5MainStreamObjectVersion::RemoveDecalBlendMode)
+	{
+		if (MaterialDomain == MD_DeferredDecal)
+		{
+			GMaterialsThatNeedDecalFix.Set(this);
+		}
+	}
+#endif
+
 #if WITH_EDITOR
 	if (Ar.IsSaving() && Ar.IsCooking() && Ar.IsPersistent() && !Ar.IsObjectReferenceCollector() && FShaderLibraryCooker::NeedsShaderStableKeys(EShaderPlatform::SP_NumPlatforms))
 	{
@@ -3583,6 +3595,100 @@ void UMaterial::BackwardsCompatibilityVirtualTextureOutputConversion()
 #endif // WITH_EDITOR
 }
 
+void UMaterial::BackwardsCompatibilityDecalConversion()
+{
+#if WITH_EDITOR
+	if (GMaterialsThatNeedDecalFix.Get(this))
+	{
+		GMaterialsThatNeedDecalFix.Clear(this);
+
+		// Move stain and alpha composite setting into material blend mode.
+		if (DecalBlendMode == DBM_AlphaComposite)
+		{
+			BlendMode = BLEND_AlphaComposite;
+		}
+		else if (DecalBlendMode == DBM_Stain)
+		{
+			BlendMode = BLEND_Modulate;
+		}
+		else
+		{
+			BlendMode = BLEND_Translucent;
+		}
+
+		// Disconnect outputs according to old DBuffer blend mode.
+		if (DecalBlendMode == DBM_DBuffer_Normal || DecalBlendMode == DBM_DBuffer_Roughness || DecalBlendMode == DBM_DBuffer_NormalRoughness)
+		{
+			BaseColor.Expression = nullptr;
+		}
+		if (DecalBlendMode == DBM_DBuffer_Color || DecalBlendMode == DBM_DBuffer_Roughness || DecalBlendMode == DBM_DBuffer_ColorRoughness || DecalBlendMode == DBM_AlphaComposite)
+		{
+			Normal.Expression = nullptr;
+		}
+		if (DecalBlendMode == DBM_DBuffer_Color || DecalBlendMode == DBM_DBuffer_Normal || DecalBlendMode == DBM_DBuffer_ColorNormal)
+		{
+			Roughness.Expression = Specular.Expression = Metallic.Expression = nullptr;
+		}
+
+		// Previously translucent decals used default values in all unconnected attributes (except for normal).
+		// For backwards compatibility we connect those attributes with defaults.
+		if (DecalBlendMode == DBM_Translucent || DecalBlendMode == DBM_AlphaComposite || DecalBlendMode == DBM_Stain)
+		{
+			if (!BaseColor.IsConnected() || !Metallic.IsConnected())
+			{
+				UMaterialExpressionConstant* Expression = NewObject<UMaterialExpressionConstant>(this);
+				Expressions.Add(Expression);
+
+				Expression->MaterialExpressionEditorX = EditorX - 100;
+				Expression->MaterialExpressionEditorY = EditorY - 120;
+				Expression->R = 0.f;
+
+				if (!BaseColor.IsConnected())
+				{
+					BaseColor.Connect(0, Expression);
+				}
+				if (!Metallic.IsConnected())
+				{
+					Metallic.Connect(0, Expression);
+				}
+			}
+
+			if (!Roughness.IsConnected() || !Specular.IsConnected())
+			{
+				UMaterialExpressionConstant* Expression = NewObject<UMaterialExpressionConstant>(this);
+				Expressions.Add(Expression);
+
+				Expression->MaterialExpressionEditorX = EditorX - 100;
+				Expression->MaterialExpressionEditorY = EditorY - 60;
+				Expression->R = .5f;
+
+				if (!Roughness.IsConnected())
+				{
+					Roughness.Connect(0, Expression);
+				}
+				if (!Specular.IsConnected())
+				{
+					Specular.Connect(0, Expression);
+				}
+			}
+		}
+
+		// Force the material to recompile.
+		FlushResourceShaderMaps();
+
+		// Note we can't mark the package dirty during post load so we will recompile on load until this material is manually resaved. 
+		// So add a map load error here to encourage a save.
+		FFormatNamedArguments Arguments;
+		Arguments.Add(TEXT("Material"), FText::FromString(*GetPathName()));
+		FMessageLog("MapCheck").Warning()
+			->AddToken(FUObjectToken::Create(this))
+			->AddToken(FTextToken::Create(FText::Format(LOCTEXT("MapCheck_Decal", "Material {Material} has been updated to fix up decal settings. If the material asset is not re-saved it may not render correctly when run outside of the editor."), Arguments)))
+			->AddToken(FActionToken::Create(LOCTEXT("MapCheck_Decal_Fix", "Fix"), LOCTEXT("MapCheck_Decal_Action", "Click to mark the material as needing to be saved."), FOnActionTokenExecuted::CreateUObject(this, &UMaterial::FixupMaterialUsageAfterLoad), true));
+		FMessageLog("MapCheck").Open(EMessageSeverity::Warning);
+	}
+#endif // WITH_EDITOR
+}
+
 void UMaterial::GetQualityLevelUsage(TArray<bool, TInlineAllocator<EMaterialQualityLevel::Num> >& OutQualityLevelsUsed, EShaderPlatform ShaderPlatform, bool bCooking)
 {
 	OutQualityLevelsUsed = GetCachedExpressionData().QualityLevelsUsed;
@@ -4105,6 +4211,7 @@ void UMaterial::PostLoad()
 
 	BackwardsCompatibilityInputConversion();
 	BackwardsCompatibilityVirtualTextureOutputConversion();
+	BackwardsCompatibilityDecalConversion();
 	ConvertMaterialToStrataMaterial();
 
 #if WITH_EDITOR
@@ -4390,11 +4497,6 @@ bool UMaterial::CanEditChange(const FProperty* InProperty) const
 			return BlendMode == BLEND_Translucent;
 		}
 
-		if (PropertyName == GET_MEMBER_NAME_STRING_CHECKED(UMaterial, DecalBlendMode))
-		{
-			return MaterialDomain == MD_DeferredDecal;
-		}
-
 		if (PropertyName == GET_MEMBER_NAME_STRING_CHECKED(UMaterial, MaterialDecalResponse))
 		{
 			static auto* CVar = IConsoleManager::Get().FindTConsoleVariableDataInt(TEXT("r.DBuffer"));
@@ -4445,14 +4547,10 @@ bool UMaterial::CanEditChange(const FProperty* InProperty) const
 	
 		if (PropertyName == GET_MEMBER_NAME_STRING_CHECKED(UMaterial, ShadingModel))
 		{
-			return (MaterialDomain == MD_Surface || (MaterialDomain == MD_DeferredDecal && DecalBlendMode == DBM_Volumetric_DistanceFunction));
+			return MaterialDomain == MD_Surface;
 		}
 
-		if (PropertyName == GET_MEMBER_NAME_STRING_CHECKED(UMaterial, DecalBlendMode))
-		{
-			return MaterialDomain == MD_DeferredDecal;
-		}
-		else if (FCString::Strncmp(*PropertyName, TEXT("bUsedWith"), 9) == 0)
+		if (FCString::Strncmp(*PropertyName, TEXT("bUsedWith"), 9) == 0)
 		{
 			return MaterialDomain == MD_DeferredDecal || MaterialDomain == MD_Surface;
 		}
@@ -6491,7 +6589,6 @@ static bool IsPropertyActive_Internal(EMaterialProperty InProperty,
 	EBlendMode BlendMode,
 	FMaterialShadingModelField ShadingModels,
 	ETranslucencyLightingMode TranslucencyLightingMode,
-	EDecalBlendMode DecalBlendMode,
 	bool bBlendableOutputAlpha,
 	bool bHasRefraction,
 	bool bUsesShadingModelFromMaterialExpression)
@@ -6522,118 +6619,39 @@ static bool IsPropertyActive_Internal(EMaterialProperty InProperty,
 			return true;
 		}
 
-		switch (DecalBlendMode)
+		if (BlendMode == BLEND_Translucent)
 		{
-		case DBM_Translucent:
 			return InProperty == MP_EmissiveColor
 				|| InProperty == MP_Normal
 				|| InProperty == MP_Metallic
 				|| InProperty == MP_Specular
 				|| InProperty == MP_BaseColor
 				|| InProperty == MP_Roughness
-				|| InProperty == MP_Opacity;
-
-		case DBM_Stain:
-			return InProperty == MP_EmissiveColor
-				|| InProperty == MP_Normal
-				|| InProperty == MP_Metallic
-				|| InProperty == MP_Specular
-				|| InProperty == MP_BaseColor
-				|| InProperty == MP_Roughness
-				|| InProperty == MP_Opacity;
-
-		case DBM_Normal:
-			return InProperty == MP_Normal
-				|| InProperty == MP_Opacity;
-
-		case DBM_Emissive:
-			// even emissive supports opacity
-			return InProperty == MP_EmissiveColor
-				|| InProperty == MP_Opacity;
-
-		case DBM_AlphaComposite:
-			return InProperty == MP_EmissiveColor
-				|| InProperty == MP_Metallic
-				|| InProperty == MP_Specular
-				|| InProperty == MP_BaseColor
-				|| InProperty == MP_Roughness
-				|| InProperty == MP_Opacity;
-
-		case DBM_DBuffer_AlphaComposite:
-			return InProperty == MP_EmissiveColor
-				|| InProperty == MP_Metallic
-				|| InProperty == MP_Specular
-				|| InProperty == MP_BaseColor
-				|| InProperty == MP_Roughness
-				|| InProperty == MP_Opacity;
-
-		case DBM_DBuffer_ColorNormalRoughness:
-			return InProperty == MP_Normal
-				|| InProperty == MP_BaseColor
-				|| InProperty == MP_Roughness
-				|| InProperty == MP_Metallic
-				|| InProperty == MP_Specular
 				|| InProperty == MP_Opacity
-				|| InProperty == MP_EmissiveColor;
-
-		case DBM_DBuffer_Emissive:
-			return InProperty == MP_EmissiveColor
-				|| InProperty == MP_Opacity;
-
-		case DBM_DBuffer_EmissiveAlphaComposite:
+				|| InProperty == MP_AmbientOcclusion;
+		}
+		else if (BlendMode == BLEND_AlphaComposite)
+		{
+			// AlphaComposite decals never write normal.
 			return InProperty == MP_EmissiveColor
-				|| InProperty == MP_Opacity;
-
-		case DBM_DBuffer_Color:
-			return InProperty == MP_BaseColor
-				|| InProperty == MP_Opacity
-				|| InProperty == MP_EmissiveColor;
-
-		case DBM_DBuffer_ColorNormal:
-			return InProperty == MP_BaseColor
-				|| InProperty == MP_Normal
-				|| InProperty == MP_Opacity
-				|| InProperty == MP_EmissiveColor;
-
-		case DBM_DBuffer_ColorRoughness:
-			return InProperty == MP_BaseColor
-				|| InProperty == MP_Roughness
 				|| InProperty == MP_Metallic
 				|| InProperty == MP_Specular
-				|| InProperty == MP_Opacity
-				|| InProperty == MP_EmissiveColor;
-
-		case DBM_DBuffer_NormalRoughness:
-			return InProperty == MP_Normal
+				|| InProperty == MP_BaseColor
 				|| InProperty == MP_Roughness
-				|| InProperty == MP_Metallic
-				|| InProperty == MP_Specular
 				|| InProperty == MP_Opacity;
-
-		case DBM_DBuffer_Normal:
-			return InProperty == MP_Normal
-				|| InProperty == MP_Opacity;
-
-		case DBM_DBuffer_Roughness:
-			return InProperty == MP_Roughness
-				|| InProperty == MP_Metallic
-				|| InProperty == MP_Specular
-				|| InProperty == MP_Opacity;
-
-		case DBM_Volumetric_DistanceFunction:
+		}
+		else if (BlendMode == BLEND_Modulate)
+		{
 			return InProperty == MP_EmissiveColor
 				|| InProperty == MP_Normal
 				|| InProperty == MP_Metallic
 				|| InProperty == MP_Specular
 				|| InProperty == MP_BaseColor
 				|| InProperty == MP_Roughness
-				|| InProperty == MP_OpacityMask;
-
-		case DBM_AmbientOcclusion:
-			return InProperty == MP_AmbientOcclusion;
-
-		default:
-			// if you create a new mode it needs to expose the right pins
+				|| InProperty == MP_Opacity;
+		}
+		else
+		{
 			return false;
 		}
 	}
@@ -6759,7 +6777,6 @@ bool UMaterial::IsPropertyActiveInEditor(EMaterialProperty InProperty) const
 		BlendMode,
 		ShadingModels,
 		TranslucencyLightingMode,
-		DecalBlendMode,
 		BlendableOutputAlpha,
 		Refraction.IsConnected(),
 		IsShadingModelFromMaterialExpression());
@@ -6773,7 +6790,6 @@ bool UMaterial::IsPropertyActiveInDerived(EMaterialProperty InProperty, const UM
 		DerivedMaterial->GetBlendMode(),
 		DerivedMaterial->GetShadingModels(),
 		TranslucencyLightingMode,
-		DecalBlendMode,
 		BlendableOutputAlpha,
 		Refraction.IsConnected(),
 		DerivedMaterial->IsShadingModelFromMaterialExpression());
diff --git a/Engine/Source/Runtime/Engine/Private/Materials/MaterialShared.cpp b/Engine/Source/Runtime/Engine/Private/Materials/MaterialShared.cpp
--- a/Engine/Source/Runtime/Engine/Private/Materials/MaterialShared.cpp
+++ b/Engine/Source/Runtime/Engine/Private/Materials/MaterialShared.cpp
@@ -1571,11 +1571,6 @@ bool FMaterialResource::ComputeFogPerPixel() const {return Material->bComputeFog
 FString FMaterialResource::GetFriendlyName() const { return *GetNameSafe(Material); } //avoid using the material instance name here, we want materials that share a shadermap to also share a friendly name.
 FString FMaterialResource::GetAssetName() const { return MaterialInstance ? *GetNameSafe(MaterialInstance) : *GetNameSafe(Material); }
 
-uint32 FMaterialResource::GetDecalBlendMode() const
-{
-	return Material->GetDecalBlendMode();
-}
-
 uint32 FMaterialResource::GetMaterialDecalResponse() const
 {
 	return Material->GetMaterialDecalResponse();
diff --git a/Engine/Source/Runtime/Engine/Public/MaterialShared.h b/Engine/Source/Runtime/Engine/Public/MaterialShared.h
--- a/Engine/Source/Runtime/Engine/Public/MaterialShared.h
+++ b/Engine/Source/Runtime/Engine/Public/MaterialShared.h
@@ -1694,7 +1694,6 @@ public:
 	virtual bool HasVertexPositionOffsetConnected() const { return false; }
 	virtual bool HasPixelDepthOffsetConnected() const { return false; }
 	virtual bool HasMaterialAttributesConnected() const { return false; }
-	virtual uint32 GetDecalBlendMode() const { return 0; }
 	virtual uint32 GetMaterialDecalResponse() const { return 0; }
 	virtual bool HasBaseColorConnected() const { return false; }
 	virtual bool HasNormalConnected() const { return false; }
@@ -2545,7 +2544,6 @@ public:
 	ENGINE_API virtual bool UseLmDirectionality() const override;
 	ENGINE_API virtual enum EBlendMode GetBlendMode() const override;
 	ENGINE_API virtual enum ERefractionMode GetRefractionMode() const override;
-	ENGINE_API virtual uint32 GetDecalBlendMode() const override;
 	ENGINE_API virtual uint32 GetMaterialDecalResponse() const override;
 	ENGINE_API virtual bool HasBaseColorConnected() const override;
 	ENGINE_API virtual bool HasNormalConnected() const override;
@@ -3117,7 +3115,6 @@ struct FMaterialShaderParameters
 	ERHIFeatureLevel::Type FeatureLevel;
 	EMaterialQualityLevel::Type QualityLevel;
 	int32 BlendableLocation;
-	uint32 DecalBlendMode;
 	int32 NumCustomizedUVs;
 	uint32 StencilCompare;
 	union
@@ -3182,7 +3179,6 @@ struct FMaterialShaderParameters
 		FeatureLevel = InMaterial->GetFeatureLevel();
 		QualityLevel = InMaterial->GetQualityLevel();
 		BlendableLocation = InMaterial->GetBlendableLocation();
-		DecalBlendMode = InMaterial->GetDecalBlendMode();
 		NumCustomizedUVs = InMaterial->GetNumCustomizedUVs();
 		StencilCompare = InMaterial->GetStencilCompare();
 		bIsDefaultMaterial = InMaterial->IsDefaultMaterial();
diff --git a/Engine/Source/Runtime/Renderer/Private/CompositionLighting/CompositionLighting.cpp b/Engine/Source/Runtime/Renderer/Private/CompositionLighting/CompositionLighting.cpp
--- a/Engine/Source/Runtime/Renderer/Private/CompositionLighting/CompositionLighting.cpp
+++ b/Engine/Source/Runtime/Renderer/Private/CompositionLighting/CompositionLighting.cpp
@@ -571,11 +571,13 @@ void ProcessAfterBasePass(
 
 	// decal are distracting when looking at LightCulling.
 	const bool bDoDecal = ViewFamily.EngineShowFlags.Decals && !ViewFamily.EngineShowFlags.VisualizeLightCulling;
+	const bool bDBuffer = IsUsingDBuffers(View.GetShaderPlatform());
 
 	FDeferredDecalPassTextures DecalPassTextures = GetDeferredDecalPassTextures(GraphBuilder, SceneTextures, nullptr);
 
-	if (bDoDecal && IsUsingGBuffers(View.GetShaderPlatform()))
+	if (bDoDecal && !bDBuffer && IsUsingGBuffers(View.GetShaderPlatform()))
 	{
+		// We can disable this pass if using DBuffer decals
 		// Decals are before AmbientOcclusion so the decal can output a normal that AO is affected by
 		AddDeferredDecalPass(GraphBuilder, View, DecalPassTextures, EDecalRenderStage::BeforeLighting);
 	}
diff --git a/Engine/Source/Runtime/Renderer/Private/DecalRenderingCommon.cpp b/Engine/Source/Runtime/Renderer/Private/DecalRenderingCommon.cpp
--- a/Engine/Source/Runtime/Renderer/Private/DecalRenderingCommon.cpp
+++ b/Engine/Source/Runtime/Renderer/Private/DecalRenderingCommon.cpp
@@ -5,213 +5,99 @@
 
 namespace DecalRendering
 {
-	// Legacy logic involving EDecalBlendMode.
-	// EDecalBlendMode will soon be replaced and this code will be removed. But submitting this intermediate step for later reference.
-	void InitBlendDesc(EShaderPlatform Platform, EDecalBlendMode DecalBlendMode, bool bWriteNormal, bool bWriteEmissive, FDecalBlendDesc& Desc)
+	/** Finalize the initialization of FDecalBlendDesc after BlendMode and bWrite flags have all been set. */
+	void FinalizeBlendDesc(EShaderPlatform Platform, FDecalBlendDesc& Desc)
 	{
 		const bool bIsMobilePlatform = IsMobilePlatform(Platform);
 		const bool bIsMobileDeferredPlatform = bIsMobilePlatform && IsMobileDeferredShadingEnabled(Platform);
 		const bool bIsDBufferPlatform = !bIsMobilePlatform && IsUsingDBuffers(Platform);
 		const bool bIsDBufferMaskPlatform = bIsDBufferPlatform && IsUsingPerPixelDBufferMask(Platform);
-		const bool bIsForwardPlatform = IsAnyForwardShadingEnabled(Platform);
 
-		// Convert DBuffer modes to GBuffer.
-		bool bWriteBaseColor = true;
-		bool bWriteRoughnessSpecularMetallic = true;
+		Desc.bWriteDBufferMask = bIsDBufferMaskPlatform;
 
-		if (!bIsDBufferPlatform && !bIsMobilePlatform)
+		// Enforce platform blend mode limitations.
+		if (Desc.BlendMode != BLEND_Translucent && Desc.BlendMode != BLEND_AlphaComposite && Desc.BlendMode != BLEND_Modulate)
 		{
-			switch (DecalBlendMode)
-			{
-			case DBM_DBuffer_ColorNormalRoughness:
-				DecalBlendMode = DBM_Translucent;
-				break;
-			case DBM_DBuffer_Color:
-				DecalBlendMode = DBM_Translucent;
-				bWriteNormal = false;
-				bWriteRoughnessSpecularMetallic = false;
-				break;
-			case DBM_DBuffer_ColorNormal:
-				DecalBlendMode = DBM_Translucent;
-				bWriteRoughnessSpecularMetallic = false;
-				break;
-			case DBM_DBuffer_ColorRoughness:
-				DecalBlendMode = DBM_Translucent;
-				bWriteNormal = false;
-				break;
-			case DBM_DBuffer_NormalRoughness:
-				DecalBlendMode = DBM_Translucent;
-				bWriteBaseColor = false;
-				break;
-			case DBM_DBuffer_Normal:
-				DecalBlendMode = DBM_Translucent;
-				bWriteBaseColor = false;
-				bWriteRoughnessSpecularMetallic = false;
-				break;
-			}
+			Desc.BlendMode = BLEND_Translucent;
 		}
-
-		// Convert GBuffer modes to DBuffer.
-		uint32 DBufferStageMask = 1 << (uint32)EDecalRenderStage::BeforeBasePass;
-		DBufferStageMask |= bWriteEmissive ? (1 << (uint32)EDecalRenderStage::Emissive) : 0;
-
-		if (bIsDBufferPlatform && bIsForwardPlatform && !bIsMobilePlatform)
+		if (bIsDBufferPlatform && Desc.BlendMode == BLEND_Modulate)
 		{
-			switch (DecalBlendMode)
-			{
-			case DBM_Translucent:
-			case DBM_Stain:
-				DecalBlendMode = DBM_DBuffer_ColorNormalRoughness;
-				break;
-			case DBM_Normal:
-				DecalBlendMode = DBM_DBuffer_Normal;
-				DBufferStageMask = 1 << (uint32)EDecalRenderStage::BeforeBasePass;
-				break;
-			case DBM_Emissive:
-				DecalBlendMode = DBM_DBuffer_Emissive;
-				break;
-			case DBM_AlphaComposite:
-				DecalBlendMode = DBM_DBuffer_AlphaComposite;
-				break;
-			}
-		}
-
-		// Fill out FDecalBlendDesc.
-		switch (DecalBlendMode)
-		{
-		case DBM_AlphaComposite:
-			Desc.BlendMode = BLEND_AlphaComposite;
-			Desc.bWriteBaseColor = bWriteBaseColor;
-			Desc.bWriteRoughnessSpecularMetallic = bWriteRoughnessSpecularMetallic;
-			Desc.bWriteEmissive = bWriteEmissive;
-			Desc.RenderStageMask = 1 << (uint32)EDecalRenderStage::BeforeLighting;
-			break;
-		case DBM_Stain:
-			Desc.BlendMode = BLEND_Modulate;
-			Desc.bWriteBaseColor = bWriteBaseColor;
-			Desc.bWriteNormal = bWriteNormal;
-			Desc.bWriteRoughnessSpecularMetallic = bWriteRoughnessSpecularMetallic;
-			Desc.bWriteEmissive = bWriteEmissive;
-			Desc.RenderStageMask = 1 << (uint32)EDecalRenderStage::BeforeLighting;
-			break;
-		case DBM_Translucent:
-			Desc.BlendMode = BLEND_Translucent;
-			Desc.bWriteBaseColor = bWriteBaseColor;
-			Desc.bWriteNormal = bWriteNormal;
-			Desc.bWriteRoughnessSpecularMetallic = bWriteRoughnessSpecularMetallic;
-			Desc.bWriteEmissive = bWriteEmissive;
-			Desc.RenderStageMask = 1 << (uint32)EDecalRenderStage::BeforeLighting;
-			break;
-		case DBM_Normal:
-			Desc.BlendMode = BLEND_Translucent;
-			Desc.bWriteNormal = true;
-			Desc.RenderStageMask = 1 << (uint32)EDecalRenderStage::BeforeLighting;
-			break;
-		case DBM_Emissive:
 			Desc.BlendMode = BLEND_Translucent;
-			Desc.bWriteEmissive = true;
-			Desc.RenderStageMask = 1 << (uint32)EDecalRenderStage::BeforeLighting;
-			break;
-		case DBM_DBuffer_ColorNormalRoughness:
-			Desc.BlendMode = BLEND_Translucent;
-			Desc.bWriteBaseColor = Desc.bWriteNormal = Desc.bWriteRoughnessSpecularMetallic = true;
-			Desc.bWriteEmissive = bWriteEmissive;
-			Desc.bWriteDBufferMask = bIsDBufferMaskPlatform;
-			Desc.RenderStageMask = DBufferStageMask;
-			break;
-		case DBM_DBuffer_Color:
-			Desc.BlendMode = BLEND_Translucent;
-			Desc.bWriteBaseColor = true;
-			Desc.bWriteEmissive = bWriteEmissive;
-			Desc.bWriteDBufferMask = bIsDBufferMaskPlatform;
-			Desc.RenderStageMask = DBufferStageMask;
-			break;
-		case DBM_DBuffer_ColorNormal:
-			Desc.BlendMode = BLEND_Translucent;
-			Desc.bWriteBaseColor = Desc.bWriteNormal = true;
-			Desc.bWriteEmissive = bWriteEmissive;
-			Desc.bWriteDBufferMask = bIsDBufferMaskPlatform;
-			Desc.RenderStageMask = DBufferStageMask;
-			break;
-		case DBM_DBuffer_ColorRoughness:
-			Desc.BlendMode = BLEND_Translucent;
-			Desc.bWriteBaseColor = Desc.bWriteRoughnessSpecularMetallic = true;
-			Desc.bWriteEmissive = bWriteEmissive;
-			Desc.bWriteDBufferMask = bIsDBufferMaskPlatform;
-			Desc.RenderStageMask = DBufferStageMask;
-			break;
-		case DBM_DBuffer_Normal:
-			Desc.BlendMode = BLEND_Translucent;
-			Desc.bWriteNormal = true;
-			Desc.bWriteEmissive = bWriteEmissive;
-			Desc.bWriteDBufferMask = bIsDBufferMaskPlatform;
-			Desc.RenderStageMask = DBufferStageMask;
-			break;
-		case DBM_DBuffer_NormalRoughness:
-			Desc.BlendMode = BLEND_Translucent;
-			Desc.bWriteNormal = Desc.bWriteRoughnessSpecularMetallic = true;
-			Desc.bWriteEmissive = bWriteEmissive;
-			Desc.bWriteDBufferMask = bIsDBufferMaskPlatform;
-			Desc.RenderStageMask = DBufferStageMask;
-			break;
-		case DBM_DBuffer_Roughness:
-			Desc.BlendMode = BLEND_Translucent;
-			Desc.bWriteRoughnessSpecularMetallic = true;
-			Desc.bWriteEmissive = bWriteEmissive;
-			Desc.bWriteDBufferMask = bIsDBufferMaskPlatform;
-			Desc.RenderStageMask = DBufferStageMask;
-			break;
-		case DBM_DBuffer_Emissive:
-			Desc.BlendMode = BLEND_Translucent;
-			Desc.bWriteEmissive = true;
-			Desc.RenderStageMask = 1 << (uint32)EDecalRenderStage::Emissive;
-			break;
-		case DBM_DBuffer_AlphaComposite:
-			Desc.BlendMode = BLEND_AlphaComposite;
-			Desc.bWriteBaseColor = Desc.bWriteRoughnessSpecularMetallic = true;
-			Desc.bWriteEmissive = bWriteEmissive;
-			Desc.bWriteDBufferMask = bIsDBufferMaskPlatform;
-			Desc.RenderStageMask = DBufferStageMask;
-			break;
-		case DBM_Volumetric_DistanceFunction:
-			// Ignore
-			break;
-		case DBM_AmbientOcclusion:
+		}
+		if (bIsMobilePlatform && !bIsMobileDeferredPlatform && Desc.bWriteEmissive)
+		{
 			Desc.BlendMode = BLEND_Translucent;
-			Desc.bWriteAmbientOcclusion = true;
-			Desc.RenderStageMask = 1 << (uint32)EDecalRenderStage::AmbientOcclusion;
-			break;
 		}
 
-		// Fixup for Mobile.
-		if (bIsMobileDeferredPlatform)
+		// Enforce platform output limitations.
+		if (bIsMobilePlatform && !bIsMobileDeferredPlatform)
+		{
+			Desc.bWriteNormal = false;
+			Desc.bWriteRoughnessSpecularMetallic = false;
+		}
+		if (bIsMobilePlatform)
 		{
 			Desc.bWriteAmbientOcclusion = false;
-			Desc.bWriteDBufferMask = false;
-			Desc.RenderStageMask = Desc.bWriteEmissive || Desc.bWriteBaseColor || Desc.bWriteNormal || Desc.bWriteRoughnessSpecularMetallic ? 1 << (uint32)EDecalRenderStage::MobileBeforeLighting : 0;
 		}
-		else if (bIsMobilePlatform)
+
+		// Enforce blend modes output limitations.
+		if (Desc.BlendMode == BLEND_AlphaComposite)
 		{
 			Desc.bWriteNormal = false;
-			Desc.bWriteRoughnessSpecularMetallic = false;
-			Desc.bWriteAmbientOcclusion = false;
-			Desc.bWriteDBufferMask = false;
-			Desc.BlendMode = Desc.bWriteEmissive ? BLEND_Translucent : Desc.BlendMode;
-			Desc.RenderStageMask = Desc.bWriteEmissive || Desc.bWriteBaseColor ? 1 << (uint32)EDecalRenderStage::Mobile : 0;
+		}
+
+		// Calculate main decal render stage. We set only one (or none) of these for any decal.
+		if (bIsMobileDeferredPlatform && (Desc.bWriteEmissive || Desc.bWriteBaseColor || Desc.bWriteNormal || Desc.bWriteRoughnessSpecularMetallic))
+		{
+			Desc.RenderStageMask |= 1 << (uint32)EDecalRenderStage::MobileBeforeLighting;
+		}
+		else if (bIsMobilePlatform && (Desc.bWriteEmissive || Desc.bWriteBaseColor))
+		{
+			Desc.RenderStageMask |= 1 << (uint32)EDecalRenderStage::Mobile;
+		}
+		else if (bIsDBufferPlatform && (Desc.bWriteBaseColor || Desc.bWriteNormal || Desc.bWriteRoughnessSpecularMetallic))
+		{
+			Desc.RenderStageMask |= 1 << (uint32)EDecalRenderStage::BeforeBasePass;
+		}
+		else if (Desc.bWriteEmissive || Desc.bWriteBaseColor || Desc.bWriteNormal || Desc.bWriteRoughnessSpecularMetallic)
+		{
+			Desc.RenderStageMask |= 1 << (uint32)EDecalRenderStage::BeforeLighting;
+		}
+
+		// Calculate additional decal render stages.
+		if (Desc.bWriteEmissive && bIsDBufferPlatform)
+		{
+			Desc.RenderStageMask |= 1 << (uint32)EDecalRenderStage::Emissive;
+		}
+		if (Desc.bWriteAmbientOcclusion)
+		{
+			Desc.RenderStageMask |= 1 << (uint32)EDecalRenderStage::AmbientOcclusion;
 		}
 	}
 
 	FDecalBlendDesc ComputeDecalBlendDesc(EShaderPlatform Platform, FMaterial const* Material)
 	{
 		FDecalBlendDesc Desc;
-		InitBlendDesc(Platform, (EDecalBlendMode)Material->GetDecalBlendMode(), Material->HasNormalConnected(), Material->HasEmissiveColorConnected(), Desc);
+		Desc.BlendMode = Material->GetBlendMode();
+		Desc.bWriteBaseColor = Material->HasBaseColorConnected();
+		Desc.bWriteNormal = Material->HasNormalConnected();
+		Desc.bWriteRoughnessSpecularMetallic = Material->HasRoughnessConnected() || Material->HasSpecularConnected() || Material->HasMetallicConnected();
+		Desc.bWriteEmissive = Material->HasEmissiveColorConnected();
+		Desc.bWriteAmbientOcclusion = Material->HasAmbientOcclusionConnected();
+		FinalizeBlendDesc(Platform, Desc);
 		return Desc;
 	}
 
 	FDecalBlendDesc ComputeDecalBlendDesc(EShaderPlatform Platform, FMaterialShaderParameters const& MaterialShaderParameters)
 	{
 		FDecalBlendDesc Desc;
-		InitBlendDesc(Platform, (EDecalBlendMode)MaterialShaderParameters.DecalBlendMode, MaterialShaderParameters.bHasNormalConnected, MaterialShaderParameters.bHasEmissiveColorConnected, Desc);
+		Desc.BlendMode = MaterialShaderParameters.BlendMode;
+		Desc.bWriteBaseColor = MaterialShaderParameters.bHasBaseColorConnected;
+		Desc.bWriteNormal = MaterialShaderParameters.bHasNormalConnected;
+		Desc.bWriteRoughnessSpecularMetallic = MaterialShaderParameters.bHasRoughnessConnected || MaterialShaderParameters.bHasSpecularConnected || MaterialShaderParameters.bHasMetallicConnected;
+		Desc.bWriteEmissive = MaterialShaderParameters.bHasEmissiveColorConnected;
+		Desc.bWriteAmbientOcclusion = MaterialShaderParameters.bHasAmbientOcclusionConnected;
+		FinalizeBlendDesc(Platform, Desc);
 		return Desc;
 	}
 
