Fixed incorrect assert condition in FRayTracingScene::BeginCreate for ray tracing instances with GPU transforms.

diff --git a/Engine/Source/Runtime/Renderer/Private/RayTracing/RayTracingScene.cpp b/Engine/Source/Runtime/Renderer/Private/RayTracing/RayTracingScene.cpp
--- a/Engine/Source/Runtime/Renderer/Private/RayTracing/RayTracingScene.cpp
+++ b/Engine/Source/Runtime/Renderer/Private/RayTracing/RayTracingScene.cpp
@@ -25,7 +25,10 @@ FGraphEventRef FRayTracingScene::BeginCreate(FRDGBuilder& GraphBuilder)
 	uint32 NumTotalInstances = 0;
 	for (const FRayTracingGeometryInstance& InstanceDesc : Instances)
 	{
-		ensure(InstanceDesc.NumTransforms <= uint32(InstanceDesc.Transforms.Num()));
+		checkf(InstanceDesc.GPUTransformsSRV || InstanceDesc.NumTransforms <= uint32(InstanceDesc.Transforms.Num()),
+			TEXT("Expected at most %d ray tracing geometry instance transforms, but got %d."),
+			InstanceDesc.NumTransforms, InstanceDesc.Transforms.Num());
+
 		NumTotalInstances += InstanceDesc.NumTransforms;
 	}
 	SET_DWORD_STAT(STAT_RayTracingInstances, NumTotalInstances);
