Converting FMeshPassProcessors to use ODSC.

- FHairMaterialGBufferProcessor
- FHairMaterialProcessor
- FHairVisibilityProcessor

diff --git a/Engine/Source/Runtime/Renderer/Private/HairStrands/HairStrandsVisibility.cpp b/Engine/Source/Runtime/Renderer/Private/HairStrands/HairStrandsVisibility.cpp
--- a/Engine/Source/Runtime/Renderer/Private/HairStrands/HairStrandsVisibility.cpp
+++ b/Engine/Source/Runtime/Renderer/Private/HairStrands/HairStrandsVisibility.cpp
@@ -526,7 +526,17 @@ public:
 	void AddMeshBatch(const FMeshBatch& RESTRICT MeshBatch, uint64 BatchElementMask, const FPrimitiveSceneProxy* RESTRICT PrimitiveSceneProxy, int32 StaticMeshId, int32 MacroGroupId, int32 HairMaterialId);
 
 private:
-	void Process(
+	bool TryAddMeshBatch(
+		const FMeshBatch& RESTRICT MeshBatch,
+		uint64 BatchElementMask,
+		const FPrimitiveSceneProxy* RESTRICT PrimitiveSceneProxy,
+		int32 StaticMeshId,
+		int32 MacroGroupId,
+		int32 HairMaterialId,
+		const FMaterialRenderProxy& MaterialRenderProxy,
+		const FMaterial& Material);
+
+	bool Process(
 		const FMeshBatch& MeshBatch,
 		uint64 BatchElementMask,
 		const FPrimitiveSceneProxy* RESTRICT PrimitiveSceneProxy,
@@ -547,12 +557,36 @@ void FHairMaterialGBufferProcessor::AddMeshBatch(const FMeshBatch& RESTRICT Mesh
 }
 
 void FHairMaterialGBufferProcessor::AddMeshBatch(const FMeshBatch& RESTRICT MeshBatch, uint64 BatchElementMask, const FPrimitiveSceneProxy* RESTRICT PrimitiveSceneProxy, int32 StaticMeshId, int32 MacroGroupId, int32 HairMaterialId)
+{
+	const FMaterialRenderProxy* MaterialRenderProxy = MeshBatch.MaterialRenderProxy;
+	while (MaterialRenderProxy)
+	{
+		const FMaterial* Material = MaterialRenderProxy->GetMaterialNoFallback(FeatureLevel);
+		if (Material)
+		{
+			if (TryAddMeshBatch(MeshBatch, BatchElementMask, PrimitiveSceneProxy, StaticMeshId, MacroGroupId, HairMaterialId, *MaterialRenderProxy, *Material))
+			{
+				break;
+			}
+		}
+
+		MaterialRenderProxy = MaterialRenderProxy->GetFallback(FeatureLevel);
+	}
+}
+
+bool FHairMaterialGBufferProcessor::TryAddMeshBatch(
+	const FMeshBatch& RESTRICT MeshBatch,
+	uint64 BatchElementMask,
+	const FPrimitiveSceneProxy* RESTRICT PrimitiveSceneProxy,
+	int32 StaticMeshId,
+	int32 MacroGroupId,
+	int32 HairMaterialId,
+	const FMaterialRenderProxy& MaterialRenderProxy,
+	const FMaterial& Material)
 {
 	static const FVertexFactoryType* CompatibleVF = FVertexFactoryType::GetVFByName(TEXT("FHairStrandsVertexFactory"));
 
 	// Determine the mesh's material and blend mode.
-	const FMaterialRenderProxy* FallbackMaterialRenderProxyPtr = nullptr;
-	const FMaterial& Material = MeshBatch.MaterialRenderProxy->GetMaterialWithFallback(FeatureLevel, FallbackMaterialRenderProxyPtr);
 	const bool bIsCompatible = IsCompatibleWithHairStrands(&Material, FeatureLevel);
 	const bool bIsHairStrandsFactory = MeshBatch.VertexFactory->GetType()->GetHashedName() == CompatibleVF->GetHashedName();
 	const bool bShouldRender = (!PrimitiveSceneProxy && MeshBatch.Elements.Num() > 0) || (PrimitiveSceneProxy && PrimitiveSceneProxy->ShouldRenderInMainPass());
@@ -579,12 +613,13 @@ void FHairMaterialGBufferProcessor::AddMeshBatch(const FMeshBatch& RESTRICT Mesh
 		GetDrawCommandPrimitiveId(SceneInfo, MeshBatch.Elements[0], PrimitiveId, ScenePrimitiveId);
 		uint32 LightChannelMask = PrimitiveSceneProxy ? PrimitiveSceneProxy->GetLightingChannelMask() : 0;
 
-		const FMaterialRenderProxy& MaterialRenderProxy = FallbackMaterialRenderProxyPtr ? *FallbackMaterialRenderProxyPtr : *MeshBatch.MaterialRenderProxy;
-		Process(MeshBatchCopy, BatchElementMask, PrimitiveSceneProxy, StaticMeshId, MaterialRenderProxy, Material, MacroGroupId, HairMaterialId, PrimitiveId, LightChannelMask);
+		return Process(MeshBatchCopy, BatchElementMask, PrimitiveSceneProxy, StaticMeshId, MaterialRenderProxy, Material, MacroGroupId, HairMaterialId, PrimitiveId, LightChannelMask);
 	}
+
+	return true;
 }
 
-void FHairMaterialGBufferProcessor::Process(
+bool FHairMaterialGBufferProcessor::Process(
 	const FMeshBatch& MeshBatch,
 	uint64 BatchElementMask,
 	const FPrimitiveSceneProxy* RESTRICT PrimitiveSceneProxy,
@@ -602,9 +637,20 @@ void FHairMaterialGBufferProcessor::Process(
 		FHairMaterialGBufferVS,
 		FHairMaterialGBufferPS> PassShaders;
 	{
+		FMaterialShaderTypes ShaderTypes;
+		ShaderTypes.AddShaderType<FHairMaterialGBufferVS>();
+		ShaderTypes.AddShaderType<FHairMaterialGBufferPS>();
+
 		FVertexFactoryType* VertexFactoryType = VertexFactory->GetType();
-		PassShaders.VertexShader = MaterialResource.GetShader<FHairMaterialGBufferVS>(VertexFactoryType);
-		PassShaders.PixelShader = MaterialResource.GetShader<FHairMaterialGBufferPS>(VertexFactoryType);
+
+		FMaterialShaders Shaders;
+		if (!MaterialResource.TryGetShaders(ShaderTypes, VertexFactoryType, Shaders))
+		{
+			return false;
+		}
+
+		Shaders.TryGetVertexShader(PassShaders.VertexShader);
+		Shaders.TryGetPixelShader(PassShaders.PixelShader);
 	}
 
 	FMeshPassProcessorRenderState DrawRenderState(PassDrawRenderState);
@@ -624,6 +670,8 @@ void FHairMaterialGBufferProcessor::Process(
 		FMeshDrawCommandSortKey::Default,
 		EMeshPassFeatures::Default,
 		ShaderElementData);
+
+	return true;
 }
 
 FHairMaterialGBufferProcessor::FHairMaterialGBufferProcessor(
@@ -892,7 +940,17 @@ public:
 	void AddMeshBatch(const FMeshBatch& RESTRICT MeshBatch, uint64 BatchElementMask, const FPrimitiveSceneProxy* RESTRICT PrimitiveSceneProxy, int32 StaticMeshId, int32 MacroGroupId, int32 HairMaterialId);
 
 private:
-	void Process(
+	bool TryAddMeshBatch(
+		const FMeshBatch& RESTRICT MeshBatch,
+		uint64 BatchElementMask,
+		const FPrimitiveSceneProxy* RESTRICT PrimitiveSceneProxy,
+		int32 StaticMeshId,
+		uint32 MacroGroupId,
+		uint32 HairMaterialId,
+		const FMaterialRenderProxy& MaterialRenderProxy,
+		const FMaterial& Material);
+
+	bool Process(
 		const FMeshBatch& MeshBatch,
 		uint64 BatchElementMask,
 		const FPrimitiveSceneProxy* RESTRICT PrimitiveSceneProxy,
@@ -913,15 +971,39 @@ void FHairMaterialProcessor::AddMeshBatch(const FMeshBatch& RESTRICT MeshBatch,
 }
 
 void FHairMaterialProcessor::AddMeshBatch(const FMeshBatch& RESTRICT MeshBatch, uint64 BatchElementMask, const FPrimitiveSceneProxy* RESTRICT PrimitiveSceneProxy, int32 StaticMeshId, int32 MacroGroupId, int32 HairMaterialId)
+{
+	const FMaterialRenderProxy* MaterialRenderProxy = MeshBatch.MaterialRenderProxy;
+	while (MaterialRenderProxy)
+	{
+		const FMaterial* Material = MaterialRenderProxy->GetMaterialNoFallback(FeatureLevel);
+		if (Material)
+		{
+			if (TryAddMeshBatch(MeshBatch, BatchElementMask, PrimitiveSceneProxy, StaticMeshId, MacroGroupId, HairMaterialId, *MaterialRenderProxy, *Material))
+			{
+				break;
+			}
+		}
+
+		MaterialRenderProxy = MaterialRenderProxy->GetFallback(FeatureLevel);
+	}
+}
+
+bool FHairMaterialProcessor::TryAddMeshBatch(
+	const FMeshBatch& RESTRICT MeshBatch,
+	uint64 BatchElementMask,
+	const FPrimitiveSceneProxy* RESTRICT PrimitiveSceneProxy,
+	int32 StaticMeshId,
+	uint32 MacroGroupId,
+	uint32 HairMaterialId,
+	const FMaterialRenderProxy& MaterialRenderProxy,
+	const FMaterial& Material)
 {
 	static const FVertexFactoryType* CompatibleVF = FVertexFactoryType::GetVFByName(TEXT("FHairStrandsVertexFactory"));
 
 	// Determine the mesh's material and blend mode.
-	const FMaterialRenderProxy* FallbackMaterialRenderProxyPtr = nullptr;
-	const FMaterial& Material = MeshBatch.MaterialRenderProxy->GetMaterialWithFallback(FeatureLevel, FallbackMaterialRenderProxyPtr);
 	const bool bIsCompatible = IsCompatibleWithHairStrands(&Material, FeatureLevel);
 	const bool bIsHairStrandsFactory = MeshBatch.VertexFactory->GetType()->GetHashedName() == CompatibleVF->GetHashedName();
-	const bool bShouldRender = (!PrimitiveSceneProxy && MeshBatch.Elements.Num()>0) || (PrimitiveSceneProxy && PrimitiveSceneProxy->ShouldRenderInMainPass());
+	const bool bShouldRender = (!PrimitiveSceneProxy && MeshBatch.Elements.Num() > 0) || (PrimitiveSceneProxy && PrimitiveSceneProxy->ShouldRenderInMainPass());
 
 	if (bIsCompatible
 		&& bIsHairStrandsFactory
@@ -930,7 +1012,7 @@ void FHairMaterialProcessor::AddMeshBatch(const FMeshBatch& RESTRICT MeshBatch,
 	{
 		// For the mesh patch to be rendered a single triangle triangle to spawn the necessary amount of thread
 		FMeshBatch MeshBatchCopy = MeshBatch;
-		for (uint32 ElementIt = 0, ElementCount=uint32(MeshBatch.Elements.Num()); ElementIt < ElementCount; ++ElementIt)
+		for (uint32 ElementIt = 0, ElementCount = uint32(MeshBatch.Elements.Num()); ElementIt < ElementCount; ++ElementIt)
 		{
 			MeshBatchCopy.Elements[ElementIt].FirstIndex = 0;
 			MeshBatchCopy.Elements[ElementIt].NumPrimitives = 1;
@@ -945,12 +1027,13 @@ void FHairMaterialProcessor::AddMeshBatch(const FMeshBatch& RESTRICT MeshBatch,
 		GetDrawCommandPrimitiveId(SceneInfo, MeshBatch.Elements[0], PrimitiveId, ScenePrimitiveId);
 		uint32 LightChannelMask = PrimitiveSceneProxy ? PrimitiveSceneProxy->GetLightingChannelMask() : 0;
 
-		const FMaterialRenderProxy& MaterialRenderProxy = FallbackMaterialRenderProxyPtr ? *FallbackMaterialRenderProxyPtr : *MeshBatch.MaterialRenderProxy;
-		Process(MeshBatchCopy, BatchElementMask, PrimitiveSceneProxy, StaticMeshId, MaterialRenderProxy, Material, MacroGroupId, HairMaterialId, PrimitiveId, LightChannelMask);
+		return Process(MeshBatchCopy, BatchElementMask, PrimitiveSceneProxy, StaticMeshId, MaterialRenderProxy, Material, MacroGroupId, HairMaterialId, PrimitiveId, LightChannelMask);
 	}
+
+	return true;
 }
 
-void FHairMaterialProcessor::Process(
+bool FHairMaterialProcessor::Process(
 	const FMeshBatch& MeshBatch,
 	uint64 BatchElementMask,
 	const FPrimitiveSceneProxy* RESTRICT PrimitiveSceneProxy,
@@ -968,9 +1051,20 @@ void FHairMaterialProcessor::Process(
 		FHairMaterialVS,
 		FHairMaterialPS> PassShaders;
 	{
+		FMaterialShaderTypes ShaderTypes;
+		ShaderTypes.AddShaderType<FHairMaterialVS>();
+		ShaderTypes.AddShaderType<FHairMaterialPS>();
+
 		FVertexFactoryType* VertexFactoryType = VertexFactory->GetType();
-		PassShaders.VertexShader = MaterialResource.GetShader<FHairMaterialVS>(VertexFactoryType);
-		PassShaders.PixelShader = MaterialResource.GetShader<FHairMaterialPS>(VertexFactoryType);
+
+		FMaterialShaders Shaders;
+		if (!MaterialResource.TryGetShaders(ShaderTypes, VertexFactoryType, Shaders))
+		{
+			return false;
+		}
+
+		Shaders.TryGetVertexShader(PassShaders.VertexShader);
+		Shaders.TryGetPixelShader(PassShaders.PixelShader);
 	}
 
 	FMeshPassProcessorRenderState DrawRenderState(PassDrawRenderState);
@@ -990,6 +1084,8 @@ void FHairMaterialProcessor::Process(
 		FMeshDrawCommandSortKey::Default,
 		EMeshPassFeatures::Default,
 		ShaderElementData);
+
+	return true;
 }
 
 FHairMaterialProcessor::FHairMaterialProcessor(
@@ -1562,8 +1658,19 @@ public:
 	void AddMeshBatch(const FMeshBatch& RESTRICT MeshBatch, uint64 BatchElementMask, const FPrimitiveSceneProxy* RESTRICT PrimitiveSceneProxy, int32 StaticMeshId, uint32 HairMacroGroupId, uint32 HairMaterialId, bool bCullingEnable);
 
 private:
+	bool TryAddMeshBatch(
+		const FMeshBatch& RESTRICT MeshBatch,
+		uint64 BatchElementMask,
+		const FPrimitiveSceneProxy* RESTRICT PrimitiveSceneProxy,
+		int32 StaticMeshId,
+		uint32 HairMacroGroupId,
+		uint32 HairMaterialId,
+		bool bCullingEnable,
+		const FMaterialRenderProxy& MaterialRenderProxy,
+		const FMaterial& Material);
+
 	template<EHairVisibilityRenderMode RenderMode, bool bCullingEnable=true>
-	void Process(
+	bool Process(
 		const FMeshBatch& MeshBatch,
 		uint64 BatchElementMask,
 		const FPrimitiveSceneProxy* RESTRICT PrimitiveSceneProxy,
@@ -1586,41 +1693,67 @@ void FHairVisibilityProcessor::AddMeshBatch(const FMeshBatch& RESTRICT MeshBatch
 }
 
 void FHairVisibilityProcessor::AddMeshBatch(const FMeshBatch& RESTRICT MeshBatch, uint64 BatchElementMask, const FPrimitiveSceneProxy* RESTRICT PrimitiveSceneProxy, int32 StaticMeshId, uint32 HairMacroGroupId, uint32 HairMaterialId, bool bCullingEnable)
+{
+	const FMaterialRenderProxy* MaterialRenderProxy = MeshBatch.MaterialRenderProxy;
+	while (MaterialRenderProxy)
+	{
+		const FMaterial* Material = MaterialRenderProxy->GetMaterialNoFallback(FeatureLevel);
+		if (Material)
+		{
+			if (TryAddMeshBatch(MeshBatch, BatchElementMask, PrimitiveSceneProxy, StaticMeshId, HairMacroGroupId, HairMaterialId, bCullingEnable, *MaterialRenderProxy, *Material))
+			{
+				break;
+			}
+		}
+
+		MaterialRenderProxy = MaterialRenderProxy->GetFallback(FeatureLevel);
+	}
+}
+
+bool FHairVisibilityProcessor::TryAddMeshBatch(
+	const FMeshBatch& RESTRICT MeshBatch,
+	uint64 BatchElementMask,
+	const FPrimitiveSceneProxy* RESTRICT PrimitiveSceneProxy,
+	int32 StaticMeshId,
+	uint32 HairMacroGroupId,
+	uint32 HairMaterialId,
+	bool bCullingEnable,
+	const FMaterialRenderProxy& MaterialRenderProxy,
+	const FMaterial& Material)
 {
 	static const FVertexFactoryType* CompatibleVF = FVertexFactoryType::GetVFByName(TEXT("FHairStrandsVertexFactory"));
 
 	// Determine the mesh's material and blend mode.
-	const FMaterialRenderProxy* FallbackMaterialRenderProxyPtr = nullptr;
-	const FMaterial& Material = MeshBatch.MaterialRenderProxy->GetMaterialWithFallback(FeatureLevel, FallbackMaterialRenderProxyPtr);
 	const bool bIsCompatible = IsCompatibleWithHairStrands(&Material, FeatureLevel);
 	const bool bIsHairStrandsFactory = MeshBatch.VertexFactory->GetType()->GetHashedName() == CompatibleVF->GetHashedName();
-	const bool bShouldRender = (!PrimitiveSceneProxy && MeshBatch.Elements.Num()>0) || (PrimitiveSceneProxy && PrimitiveSceneProxy->ShouldRenderInMainPass());
+	const bool bShouldRender = (!PrimitiveSceneProxy && MeshBatch.Elements.Num() > 0) || (PrimitiveSceneProxy && PrimitiveSceneProxy->ShouldRenderInMainPass());
 	const uint32 LightChannelMask = PrimitiveSceneProxy ? PrimitiveSceneProxy->GetLightingChannelMask() : 0;
 
-	if (bIsCompatible 
+	if (bIsCompatible
 		&& bIsHairStrandsFactory
 		&& bShouldRender
 		&& ShouldIncludeDomainInMeshPass(Material.GetMaterialDomain()))
 	{
-		const FMaterialRenderProxy& MaterialRenderProxy = FallbackMaterialRenderProxyPtr ? *FallbackMaterialRenderProxyPtr : *MeshBatch.MaterialRenderProxy;
 		const FMeshDrawingPolicyOverrideSettings OverrideSettings = ComputeMeshOverrideSettings(MeshBatch);
 		const ERasterizerFillMode MeshFillMode = ComputeMeshFillMode(MeshBatch, Material, OverrideSettings);
 		const ERasterizerCullMode MeshCullMode = ComputeMeshCullMode(MeshBatch, Material, OverrideSettings);
 		if (RenderMode == HairVisibilityRenderMode_MSAA_Visibility && bCullingEnable)
-			Process<HairVisibilityRenderMode_MSAA_Visibility, true>(MeshBatch, BatchElementMask, PrimitiveSceneProxy, StaticMeshId, MaterialRenderProxy, Material, HairMacroGroupId, HairMaterialId, LightChannelMask, MeshFillMode, MeshCullMode);
+			return Process<HairVisibilityRenderMode_MSAA_Visibility, true>(MeshBatch, BatchElementMask, PrimitiveSceneProxy, StaticMeshId, MaterialRenderProxy, Material, HairMacroGroupId, HairMaterialId, LightChannelMask, MeshFillMode, MeshCullMode);
 		else if (RenderMode == HairVisibilityRenderMode_MSAA_Visibility && !bCullingEnable)
-			Process<HairVisibilityRenderMode_MSAA_Visibility, false>(MeshBatch, BatchElementMask, PrimitiveSceneProxy, StaticMeshId, MaterialRenderProxy, Material, HairMacroGroupId, HairMaterialId, LightChannelMask, MeshFillMode, MeshCullMode);
+			return Process<HairVisibilityRenderMode_MSAA_Visibility, false>(MeshBatch, BatchElementMask, PrimitiveSceneProxy, StaticMeshId, MaterialRenderProxy, Material, HairMacroGroupId, HairMaterialId, LightChannelMask, MeshFillMode, MeshCullMode);
 		else if (RenderMode == HairVisibilityRenderMode_Transmittance)
-			Process<HairVisibilityRenderMode_Transmittance>(MeshBatch, BatchElementMask, PrimitiveSceneProxy, StaticMeshId, MaterialRenderProxy, Material, HairMacroGroupId, HairMaterialId, LightChannelMask, MeshFillMode, MeshCullMode);
+			return Process<HairVisibilityRenderMode_Transmittance>(MeshBatch, BatchElementMask, PrimitiveSceneProxy, StaticMeshId, MaterialRenderProxy, Material, HairMacroGroupId, HairMaterialId, LightChannelMask, MeshFillMode, MeshCullMode);
 		else if (RenderMode == HairVisibilityRenderMode_TransmittanceAndHairCount)
-			Process<HairVisibilityRenderMode_TransmittanceAndHairCount>(MeshBatch, BatchElementMask, PrimitiveSceneProxy, StaticMeshId, MaterialRenderProxy, Material, HairMacroGroupId, HairMaterialId, LightChannelMask, MeshFillMode, MeshCullMode);
+			return Process<HairVisibilityRenderMode_TransmittanceAndHairCount>(MeshBatch, BatchElementMask, PrimitiveSceneProxy, StaticMeshId, MaterialRenderProxy, Material, HairMacroGroupId, HairMaterialId, LightChannelMask, MeshFillMode, MeshCullMode);
 		else if (RenderMode == HairVisibilityRenderMode_PPLL)
-			Process<HairVisibilityRenderMode_PPLL>(MeshBatch, BatchElementMask, PrimitiveSceneProxy, StaticMeshId, MaterialRenderProxy, Material, HairMacroGroupId, HairMaterialId, LightChannelMask, MeshFillMode, MeshCullMode);
+			return Process<HairVisibilityRenderMode_PPLL>(MeshBatch, BatchElementMask, PrimitiveSceneProxy, StaticMeshId, MaterialRenderProxy, Material, HairMacroGroupId, HairMaterialId, LightChannelMask, MeshFillMode, MeshCullMode);
 	}
+
+	return true;
 }
 
 template<EHairVisibilityRenderMode TRenderMode, bool bCullingEnable>
-void FHairVisibilityProcessor::Process(
+bool FHairVisibilityProcessor::Process(
 	const FMeshBatch& MeshBatch, 
 	uint64 BatchElementMask, 
 	const FPrimitiveSceneProxy* RESTRICT PrimitiveSceneProxy,
@@ -1639,9 +1772,20 @@ void FHairVisibilityProcessor::Process(
 		FHairVisibilityVS<TRenderMode, bCullingEnable>,
 		FHairVisibilityPS<TRenderMode>> PassShaders;
 	{
+		FMaterialShaderTypes ShaderTypes;
+		ShaderTypes.AddShaderType<FHairVisibilityVS<TRenderMode, bCullingEnable>>();
+		ShaderTypes.AddShaderType<FHairVisibilityPS<TRenderMode>>();
+
 		FVertexFactoryType* VertexFactoryType = VertexFactory->GetType();
-		PassShaders.VertexShader = MaterialResource.GetShader<FHairVisibilityVS<TRenderMode,bCullingEnable>>(VertexFactoryType);
-		PassShaders.PixelShader  = MaterialResource.GetShader<FHairVisibilityPS<TRenderMode>>(VertexFactoryType);
+
+		FMaterialShaders Shaders;
+		if (!MaterialResource.TryGetShaders(ShaderTypes, VertexFactoryType, Shaders))
+		{
+			return false;
+		}
+
+		Shaders.TryGetVertexShader(PassShaders.VertexShader);
+		Shaders.TryGetPixelShader(PassShaders.PixelShader);
 	}
 
 	FMeshPassProcessorRenderState DrawRenderState(PassDrawRenderState);
@@ -1661,6 +1805,8 @@ void FHairVisibilityProcessor::Process(
 		FMeshDrawCommandSortKey::Default,
 		EMeshPassFeatures::Default,
 		ShaderElementData);
+
+	return true;
 }
 
 FHairVisibilityProcessor::FHairVisibilityProcessor(
