Converting FMeshPassProcessors to use ODSC.

- FVolumetricCloudRenderShadowMeshProcessor

diff --git a/Engine/Source/Runtime/Renderer/Private/VolumetricCloudRendering.cpp b/Engine/Source/Runtime/Renderer/Private/VolumetricCloudRendering.cpp
--- a/Engine/Source/Runtime/Renderer/Private/VolumetricCloudRendering.cpp
+++ b/Engine/Source/Runtime/Renderer/Private/VolumetricCloudRendering.cpp
@@ -1118,20 +1118,40 @@ public:
 
 	virtual void AddMeshBatch(const FMeshBatch& RESTRICT MeshBatch, uint64 BatchElementMask, const FPrimitiveSceneProxy* RESTRICT PrimitiveSceneProxy, int32 StaticMeshId = -1) override final
 	{
-		// Determine the mesh's material and blend mode.
-		const FMaterialRenderProxy* FallbackMaterialRenderProxyPtr = nullptr;
-		const FMaterial& Material = MeshBatch.MaterialRenderProxy->GetMaterialWithFallback(FeatureLevel, FallbackMaterialRenderProxyPtr);
-		check(Material.GetMaterialDomain() == MD_Volume);
+		const FMaterialRenderProxy* MaterialRenderProxy = MeshBatch.MaterialRenderProxy;
+		while (MaterialRenderProxy)
+		{
+			const FMaterial* Material = MaterialRenderProxy->GetMaterialNoFallback(FeatureLevel);
+			if (Material)
+			{
+				check(Material->GetMaterialDomain() == MD_Volume);
+				if (TryAddMeshBatch(MeshBatch, BatchElementMask, PrimitiveSceneProxy, StaticMeshId, *MaterialRenderProxy, *Material))
+				{
+					break;
+				}
+			}
 
-		const ERasterizerFillMode MeshFillMode = FM_Solid;
-		const ERasterizerCullMode MeshCullMode = CM_None;
-		const FMaterialRenderProxy& MaterialRenderProxy = FallbackMaterialRenderProxyPtr ? *FallbackMaterialRenderProxyPtr : *MeshBatch.MaterialRenderProxy;
-		Process(MeshBatch, BatchElementMask, PrimitiveSceneProxy, MaterialRenderProxy, Material, StaticMeshId, MeshFillMode, MeshCullMode);
+			MaterialRenderProxy = MaterialRenderProxy->GetFallback(FeatureLevel);
+		}
 	}
 
 private:
 
-	void Process(
+	bool TryAddMeshBatch(
+		const FMeshBatch& RESTRICT MeshBatch,
+		uint64 BatchElementMask,
+		const FPrimitiveSceneProxy* RESTRICT PrimitiveSceneProxy,
+		int32 StaticMeshId,
+		const FMaterialRenderProxy& MaterialRenderProxy,
+		const FMaterial& Material)
+	{
+		// Determine the mesh's material and blend mode.
+		const ERasterizerFillMode MeshFillMode = FM_Solid;
+		const ERasterizerCullMode MeshCullMode = CM_None;
+		return Process(MeshBatch, BatchElementMask, PrimitiveSceneProxy, MaterialRenderProxy, Material, StaticMeshId, MeshFillMode, MeshCullMode);
+	}
+
+	bool Process(
 		const FMeshBatch& MeshBatch,
 		uint64 BatchElementMask,
 		const FPrimitiveSceneProxy* RESTRICT PrimitiveSceneProxy,
@@ -1149,8 +1169,16 @@ private:
 		TMeshProcessorShaders<
 			FRenderVolumetricCloudVS,
 			FVolumetricCloudShadowPS> PassShaders;
-		PassShaders.PixelShader = MaterialResource.GetShader<FVolumetricCloudShadowPS>(VertexFactory->GetType());
-		PassShaders.VertexShader = MaterialResource.GetShader<FRenderVolumetricCloudVS>(VertexFactory->GetType());
+		FMaterialShaderTypes ShaderTypes;
+		ShaderTypes.AddShaderType<FRenderVolumetricCloudVS>();
+		ShaderTypes.AddShaderType<FVolumetricCloudShadowPS>();
+		FMaterialShaders Shaders;
+		if (!MaterialResource.TryGetShaders(ShaderTypes, VertexFactory->GetType(), Shaders))
+		{
+			return false;
+		}
+		Shaders.TryGetVertexShader(PassShaders.VertexShader);
+		Shaders.TryGetPixelShader(PassShaders.PixelShader);
 		const FMeshDrawCommandSortKey SortKey = CalculateMeshStaticSortKey(PassShaders.VertexShader, PassShaders.PixelShader);
 		BuildMeshDrawCommands(
 			MeshBatch,
@@ -1165,6 +1193,8 @@ private:
 			SortKey,
 			EMeshPassFeatures::Default,
 			EmptyShaderElementData);
+
+		return true;
 	}
 
 	FMeshPassProcessorRenderState PassDrawRenderState;
