Move some code that dereferences scene proxies behind conditionals to reduce cache misses
#tests a/b tested locally using a replay

diff --git a/Engine/Source/Runtime/Renderer/Private/SceneVisibility.cpp b/Engine/Source/Runtime/Renderer/Private/SceneVisibility.cpp
--- a/Engine/Source/Runtime/Renderer/Private/SceneVisibility.cpp
+++ b/Engine/Source/Runtime/Renderer/Private/SceneVisibility.cpp
@@ -531,10 +531,6 @@ static int32 FrustumCull(const FScene* RESTRICT Scene, FViewInfo& View)
 				for (int32 BitSubIndex = 0; BitSubIndex < NumBitsPerDWORD && WordIndex * NumBitsPerDWORD + BitSubIndex < BitArrayNumInner; BitSubIndex++, Mask <<= 1)
 				{
 					int32 Index = WordIndex * NumBitsPerDWORD + BitSubIndex;
-
-					FPrimitiveSceneProxy* RESTRICT Proxy = Scene->Primitives[Index]->Proxy;
-					const bool bUsingDistanceCullFade = Proxy->IsUsingDistanceCullFade();
-
 					const FPrimitiveBounds& RESTRICT Bounds = Scene->PrimitiveBounds[Index];
 					float DistanceSquared = (Bounds.BoxSphereBounds.Origin - ViewOriginForDistanceCulling).SizeSquared();
 					int32 VisibilityId = INDEX_NONE;
@@ -550,7 +546,7 @@ static int32 FrustumCull(const FScene* RESTRICT Scene, FViewInfo& View)
 					float MinDrawDistanceSq = Bounds.MinDrawDistanceSq;
 
 					// If cull distance is disabled, always show the primitive (except foliage)
-					if (View.Family->EngineShowFlags.DistanceCulledPrimitives && !Proxy->IsDetailMesh())
+					if (View.Family->EngineShowFlags.DistanceCulledPrimitives && !Scene->Primitives[Index]->Proxy->IsDetailMesh())
 					{
 						MaxDrawDistance = FLT_MAX;
 					}
@@ -581,7 +577,7 @@ static int32 FrustumCull(const FScene* RESTRICT Scene, FViewInfo& View)
 					if (Scene->PrimitivesAlwaysVisible[Index])
 					{
 						VisBits |= Mask;
-						if (bUsingDistanceCullFade && DistanceSquared > FMath::Square(MaxDrawDistance - FadeRadius))
+						if (DistanceSquared > FMath::Square(MaxDrawDistance - FadeRadius) && Scene->Primitives[Index]->Proxy->IsUsingDistanceCullFade())
 						{
 							FadingBits |= Mask;
 						}
