Converting FMeshPassProcessors to use ODSC.

- FAnisotropyMeshProcessor

diff --git a/Engine/Source/Runtime/Renderer/Private/AnisotropyRendering.cpp b/Engine/Source/Runtime/Renderer/Private/AnisotropyRendering.cpp
--- a/Engine/Source/Runtime/Renderer/Private/AnisotropyRendering.cpp
+++ b/Engine/Source/Runtime/Renderer/Private/AnisotropyRendering.cpp
@@ -112,7 +112,7 @@ FRegisterPassProcessorCreateFunction RegisterAnisotropyPass(
 	EMeshPassFlags::CachedMeshCommands | EMeshPassFlags::MainView
 	);
 
-void GetAnisotropyPassShaders(
+bool GetAnisotropyPassShaders(
 	const FMaterial& Material,
 	FVertexFactoryType* VertexFactoryType,
 	ERHIFeatureLevel::Type FeatureLevel,
@@ -121,21 +121,23 @@ void GetAnisotropyPassShaders(
 	)
 {
 	static const auto* CVar = IConsoleManager::Get().FindTConsoleVariableDataInt(TEXT("r.ShaderPipelines"));
-	const bool bUseShaderPipelines = RHISupportsShaderPipelines(GShaderPlatformForFeatureLevel[FeatureLevel]) && CVar && CVar->GetValueOnAnyThread() != 0;
 
-	FShaderPipelineRef ShaderPipeline = bUseShaderPipelines ? Material.GetShaderPipeline(&AnisotropyPipeline, VertexFactoryType, false) : FShaderPipelineRef();
-	if (ShaderPipeline.IsValid())
-	{
-		VertexShader = ShaderPipeline.GetShader<FAnisotropyVS>();
-		PixelShader = ShaderPipeline.GetShader<FAnisotropyPS>();
-		check(VertexShader.IsValid() && PixelShader.IsValid());
-	}
-	else
+	FMaterialShaderTypes ShaderTypes;
+	ShaderTypes.PipelineType = &AnisotropyPipeline;
+	ShaderTypes.AddShaderType<FAnisotropyVS>();
+	ShaderTypes.AddShaderType<FAnisotropyPS>();
+
+	FMaterialShaders Shaders;
+	if (!Material.TryGetShaders(ShaderTypes, VertexFactoryType, Shaders))
 	{
-		VertexShader = Material.GetShader<FAnisotropyVS>(VertexFactoryType);
-		PixelShader = Material.GetShader<FAnisotropyPS>(VertexFactoryType);
-		check(VertexShader.IsValid() && PixelShader.IsValid());
+		return false;
 	}
+
+	Shaders.TryGetVertexShader(VertexShader);
+	Shaders.TryGetPixelShader(PixelShader);
+	check(VertexShader.IsValid() && PixelShader.IsValid());
+
+	return true;
 }
 
 void FAnisotropyMeshProcessor::AddMeshBatch(
@@ -145,25 +147,50 @@ void FAnisotropyMeshProcessor::AddMeshBatch(
 	int32 StaticMeshId /* = -1 */ 
 	)
 {
-	if (SupportsAnisotropicMaterials(FeatureLevel, GShaderPlatformForFeatureLevel[FeatureLevel]))
+	if (SupportsAnisotropicMaterials(FeatureLevel, GShaderPlatformForFeatureLevel[FeatureLevel]) && MeshBatch.bUseForMaterial)
 	{
 		const FMaterialRenderProxy* MaterialRenderProxy = MeshBatch.MaterialRenderProxy;
-		const FMaterial* Material = &MaterialRenderProxy->GetMaterialWithFallback(FeatureLevel, MaterialRenderProxy);
-		const EBlendMode BlendMode = Material->GetBlendMode();
-		const bool bIsNotTranslucent = BlendMode == BLEND_Opaque || BlendMode == BLEND_Masked;
-
-		if (MeshBatch.bUseForMaterial && Material->MaterialUsesAnisotropy_RenderThread() && bIsNotTranslucent && Material->GetShadingModels().HasAnyShadingModel({ MSM_DefaultLit, MSM_ClearCoat }))
+		while (MaterialRenderProxy)
 		{
-			const FMeshDrawingPolicyOverrideSettings OverrideSettings = ComputeMeshOverrideSettings(MeshBatch);
-			const ERasterizerFillMode MeshFillMode = ComputeMeshFillMode(MeshBatch, *Material, OverrideSettings);
-			const ERasterizerCullMode MeshCullMode = ComputeMeshCullMode(MeshBatch, *Material, OverrideSettings);
+			const FMaterial* Material = MaterialRenderProxy->GetMaterialNoFallback(FeatureLevel);
+			if (Material)
+			{
+				if (TryAddMeshBatch(MeshBatch, BatchElementMask, PrimitiveSceneProxy, StaticMeshId, *MaterialRenderProxy, *Material))
+				{
+					break;
+				}
+			}
 
-			Process(MeshBatch, BatchElementMask, StaticMeshId, PrimitiveSceneProxy, *MaterialRenderProxy, *Material, MeshFillMode, MeshCullMode);
+			MaterialRenderProxy = MaterialRenderProxy->GetFallback(FeatureLevel);
 		}
 	}
 }
 
-void FAnisotropyMeshProcessor::Process(
+bool FAnisotropyMeshProcessor::TryAddMeshBatch(
+	const FMeshBatch& RESTRICT MeshBatch,
+	uint64 BatchElementMask,
+	const FPrimitiveSceneProxy* RESTRICT PrimitiveSceneProxy,
+	int32 StaticMeshId,
+	const FMaterialRenderProxy& MaterialRenderProxy,
+	const FMaterial& Material)
+{
+	const EBlendMode BlendMode = Material.GetBlendMode();
+	const bool bIsNotTranslucent = BlendMode == BLEND_Opaque || BlendMode == BLEND_Masked;
+
+	bool bResult = true;
+	if (Material.MaterialUsesAnisotropy_RenderThread() && bIsNotTranslucent && Material.GetShadingModels().HasAnyShadingModel({ MSM_DefaultLit, MSM_ClearCoat }))
+	{
+		const FMeshDrawingPolicyOverrideSettings OverrideSettings = ComputeMeshOverrideSettings(MeshBatch);
+		const ERasterizerFillMode MeshFillMode = ComputeMeshFillMode(MeshBatch, Material, OverrideSettings);
+		const ERasterizerCullMode MeshCullMode = ComputeMeshCullMode(MeshBatch, Material, OverrideSettings);
+
+		bResult = Process(MeshBatch, BatchElementMask, StaticMeshId, PrimitiveSceneProxy, MaterialRenderProxy, Material, MeshFillMode, MeshCullMode);
+	}
+
+	return bResult;
+}
+
+bool FAnisotropyMeshProcessor::Process(
 	const FMeshBatch& MeshBatch, 
 	uint64 BatchElementMask, 
 	int32 StaticMeshId, 
@@ -180,13 +207,15 @@ void FAnisotropyMeshProcessor::Process(
 		FAnisotropyVS,
 		FAnisotropyPS> AnisotropyPassShaders;
 
-	GetAnisotropyPassShaders(
+	if (!GetAnisotropyPassShaders(
 		MaterialResource,
 		VertexFactory->GetType(),
 		FeatureLevel,
 		AnisotropyPassShaders.VertexShader,
-		AnisotropyPassShaders.PixelShader
-		);
+		AnisotropyPassShaders.PixelShader))
+	{
+		return false;
+	}
 
 	FMeshMaterialShaderElementData ShaderElementData;
 	ShaderElementData.InitializeMeshMaterialData(ViewIfDynamicMeshCommand, PrimitiveSceneProxy, MeshBatch, StaticMeshId, false);
@@ -207,6 +236,8 @@ void FAnisotropyMeshProcessor::Process(
 		EMeshPassFeatures::Default,
 		ShaderElementData
 		);
+
+	return true;
 }
 
 bool FDeferredShadingSceneRenderer::ShouldRenderAnisotropyPass() const
diff --git a/Engine/Source/Runtime/Renderer/Private/AnisotropyRendering.h b/Engine/Source/Runtime/Renderer/Private/AnisotropyRendering.h
--- a/Engine/Source/Runtime/Renderer/Private/AnisotropyRendering.h
+++ b/Engine/Source/Runtime/Renderer/Private/AnisotropyRendering.h
@@ -30,7 +30,15 @@ public:
 		) override final;
 
 protected:
-	void Process(
+	bool TryAddMeshBatch(
+		const FMeshBatch& RESTRICT MeshBatch,
+		uint64 BatchElementMask,
+		const FPrimitiveSceneProxy* RESTRICT PrimitiveSceneProxy,
+		int32 StaticMeshId,
+		const FMaterialRenderProxy& MaterialRenderProxy,
+		const FMaterial& Material);
+
+	bool Process(
 		const FMeshBatch& MeshBatch,
 		uint64 BatchElementMask,
 		int32 StaticMeshId,
