Remove unused CVar in GetAnisotropyPassShaders()

diff --git a/Engine/Source/Runtime/Renderer/Private/AnisotropyRendering.cpp b/Engine/Source/Runtime/Renderer/Private/AnisotropyRendering.cpp
--- a/Engine/Source/Runtime/Renderer/Private/AnisotropyRendering.cpp
+++ b/Engine/Source/Runtime/Renderer/Private/AnisotropyRendering.cpp
@@ -120,8 +120,6 @@ bool GetAnisotropyPassShaders(
 	TShaderRef<FAnisotropyPS>& PixelShader
 	)
 {
-	static const auto* CVar = IConsoleManager::Get().FindTConsoleVariableDataInt(TEXT("r.ShaderPipelines"));
-
 	FMaterialShaderTypes ShaderTypes;
 	ShaderTypes.PipelineType = &AnisotropyPipeline;
 	ShaderTypes.AddShaderType<FAnisotropyVS>();
