Prevent memory leaks when caching mesh draw commands in parallel thread.

diff --git a/Engine/Source/Runtime/Renderer/Private/PrimitiveSceneInfo.cpp b/Engine/Source/Runtime/Renderer/Private/PrimitiveSceneInfo.cpp
--- a/Engine/Source/Runtime/Renderer/Private/PrimitiveSceneInfo.cpp
+++ b/Engine/Source/Runtime/Renderer/Private/PrimitiveSceneInfo.cpp
@@ -247,6 +247,8 @@ void FPrimitiveSceneInfo::CacheMeshDrawCommands(FRHICommandListImmediate& RHICmd
 	{
 		SCOPED_NAMED_EVENT(FPrimitiveSceneInfo_CacheMeshDrawCommand, FColor::Green);
 
+		FMemMark Mark(FMemStack::Get());
+
 		struct FMeshInfoAndIndex
 		{
 			int32 InfoIndex;
