Parallelize CacheNaniteDrawCommands.

diff --git a/Engine/Source/Runtime/Renderer/Private/PrimitiveSceneInfo.cpp b/Engine/Source/Runtime/Renderer/Private/PrimitiveSceneInfo.cpp
--- a/Engine/Source/Runtime/Renderer/Private/PrimitiveSceneInfo.cpp
+++ b/Engine/Source/Runtime/Renderer/Private/PrimitiveSceneInfo.cpp
@@ -550,107 +550,127 @@ void FPrimitiveSceneInfo::RemoveCachedMeshDrawCommands()
 #endif
 }
 
+static void BuildNaniteDrawCommands(FRHICommandListImmediate& RHICmdList, FScene* Scene, FPrimitiveSceneInfo* PrimitiveSceneInfo);
+
 void FPrimitiveSceneInfo::CacheNaniteDrawCommands(FRHICommandListImmediate& RHICmdList, FScene* Scene, const TArrayView<FPrimitiveSceneInfo*>& SceneInfos)
 {
+	SCOPED_NAMED_EVENT(FPrimitiveSceneInfo_CacheNaniteDrawCommands, FColor::Emerald);
 	QUICK_SCOPE_CYCLE_COUNTER(STAT_CacheNaniteDrawCommands);
-	FMemMark Mark(FMemStack::Get());
 
+	FMemMark Mark(FMemStack::Get());
 	FMaterialRenderProxy::UpdateDeferredCachedUniformExpressions();
 
 	const bool bNaniteEnabled = DoesPlatformSupportNanite(GMaxRHIShaderPlatform);
 	if (bNaniteEnabled)
 	{
-		for (FPrimitiveSceneInfo* PrimitiveSceneInfo : SceneInfos)
+		if (FApp::ShouldUseThreadingForPerformance())
 		{
-			FPrimitiveSceneProxy* Proxy = PrimitiveSceneInfo->Proxy;
-
-			if (Proxy->IsNaniteMesh())
+			ParallelForTemplate(SceneInfos.Num(), [&RHICmdList, &Scene, &SceneInfos](int32 Index)
+			{
+				FMemMark Mark(FMemStack::Get());
+				FTaskTagScope Scope(ETaskTag::EParallelRenderingThread);
+				BuildNaniteDrawCommands(RHICmdList, Scene, SceneInfos[Index]);
+			});
+		}
+		else
+		{
+			for (FPrimitiveSceneInfo* PrimitiveSceneInfo : SceneInfos)
 			{
-				Nanite::FSceneProxyBase* NaniteSceneProxy = static_cast<Nanite::FSceneProxyBase*>(Proxy);
+				BuildNaniteDrawCommands(RHICmdList, Scene, PrimitiveSceneInfo);
+			}
+		}		
+	}
+}
 
-				const TArray<Nanite::FSceneProxyBase::FMaterialSection>& MaterialSections = NaniteSceneProxy->GetMaterialSections();
+void BuildNaniteDrawCommands(FRHICommandListImmediate& RHICmdList, FScene* Scene, FPrimitiveSceneInfo* PrimitiveSceneInfo)
+{
+	FPrimitiveSceneProxy* Proxy = PrimitiveSceneInfo->Proxy;
 
-				for (int32 NaniteMeshPassIndex = 0; NaniteMeshPassIndex < ENaniteMeshPass::Num; ++NaniteMeshPassIndex)
-				{
-					check(PrimitiveSceneInfo->NaniteCommandInfos[NaniteMeshPassIndex].Num() == 0);
-					check(PrimitiveSceneInfo->NaniteMaterialIds[NaniteMeshPassIndex].Num() == 0);
+	if (Proxy->IsNaniteMesh())
+	{
+		Nanite::FSceneProxyBase* NaniteSceneProxy = static_cast<Nanite::FSceneProxyBase*>(Proxy);
 
-					PrimitiveSceneInfo->NaniteMaterialIds[NaniteMeshPassIndex].SetNum(MaterialSections.Num());
-				}
+		const TArray<Nanite::FSceneProxyBase::FMaterialSection>& MaterialSections = NaniteSceneProxy->GetMaterialSections();
+
+		for (int32 NaniteMeshPassIndex = 0; NaniteMeshPassIndex < ENaniteMeshPass::Num; ++NaniteMeshPassIndex)
+		{
+			check(PrimitiveSceneInfo->NaniteCommandInfos[NaniteMeshPassIndex].Num() == 0);
+			check(PrimitiveSceneInfo->NaniteMaterialIds[NaniteMeshPassIndex].Num() == 0);
+
+			PrimitiveSceneInfo->NaniteMaterialIds[NaniteMeshPassIndex].SetNum(MaterialSections.Num());
+		}
 
 #if WITH_EDITOR
-				check(PrimitiveSceneInfo->NaniteHitProxyIds.Num() == 0);
-				PrimitiveSceneInfo->NaniteHitProxyIds.SetNum(MaterialSections.Num());
+		check(PrimitiveSceneInfo->NaniteHitProxyIds.Num() == 0);
+		PrimitiveSceneInfo->NaniteHitProxyIds.SetNum(MaterialSections.Num());
 #endif
 
-				for (int32 NaniteMeshPassIndex = 0; NaniteMeshPassIndex < ENaniteMeshPass::Num; ++NaniteMeshPassIndex)
-				{
-					for (int32 SectionIndex = 0; SectionIndex < MaterialSections.Num(); ++SectionIndex)
-					{
-						PrimitiveSceneInfo->NaniteMaterialIds[NaniteMeshPassIndex][SectionIndex] = INDEX_NONE;
-					}
-				}
+		for (int32 NaniteMeshPassIndex = 0; NaniteMeshPassIndex < ENaniteMeshPass::Num; ++NaniteMeshPassIndex)
+		{
+			for (int32 SectionIndex = 0; SectionIndex < MaterialSections.Num(); ++SectionIndex)
+			{
+				PrimitiveSceneInfo->NaniteMaterialIds[NaniteMeshPassIndex][SectionIndex] = INDEX_NONE;
+			}
+		}
 
 #if WITH_EDITOR
-				for (int32 SectionIndex = 0; SectionIndex < MaterialSections.Num(); ++SectionIndex)
-				{
-					if (MaterialSections[SectionIndex].HitProxy)
-					{
-						PrimitiveSceneInfo->NaniteHitProxyIds[SectionIndex] = MaterialSections[SectionIndex].HitProxy->Id.GetColor().DWColor();
-					}
-					else
-					{
-						// TODO: Is this valid? SME seems to have null proxies, but normal editor doesn't
-						PrimitiveSceneInfo->NaniteHitProxyIds[SectionIndex] = INDEX_NONE;
-					}
-				}
+		for (int32 SectionIndex = 0; SectionIndex < MaterialSections.Num(); ++SectionIndex)
+		{
+			if (MaterialSections[SectionIndex].HitProxy)
+			{
+				PrimitiveSceneInfo->NaniteHitProxyIds[SectionIndex] = MaterialSections[SectionIndex].HitProxy->Id.GetColor().DWColor();
+			}
+			else
+			{
+				// TODO: Is this valid? SME seems to have null proxies, but normal editor doesn't
+				PrimitiveSceneInfo->NaniteHitProxyIds[SectionIndex] = INDEX_NONE;
+			}
+		}
 #endif
 
-				for (int32 MeshPass = 0; MeshPass < ENaniteMeshPass::Num; ++MeshPass)
-				{
-					FNaniteDrawListContext NaniteDrawListContext(Scene->NaniteDrawCommandLock[MeshPass], Scene->NaniteDrawCommands[MeshPass]);
+		for (int32 MeshPass = 0; MeshPass < ENaniteMeshPass::Num; ++MeshPass)
+		{
+			FNaniteDrawListContext NaniteDrawListContext(Scene->NaniteDrawCommandLock[MeshPass], Scene->NaniteDrawCommands[MeshPass]);
 
-					FMeshPassProcessor* NaniteMeshProcessor = nullptr;
-					switch (MeshPass)
-					{
-					case ENaniteMeshPass::BasePass:
-						NaniteMeshProcessor = CreateNaniteMeshProcessor(Scene, nullptr, &NaniteDrawListContext);
-						break;
+			FMeshPassProcessor* NaniteMeshProcessor = nullptr;
+			switch (MeshPass)
+			{
+			case ENaniteMeshPass::BasePass:
+				NaniteMeshProcessor = CreateNaniteMeshProcessor(Scene, nullptr, &NaniteDrawListContext);
+				break;
 
-					case ENaniteMeshPass::LumenCardCapture:
-						NaniteMeshProcessor = CreateLumenCardNaniteMeshProcessor(Scene, nullptr, &NaniteDrawListContext); // TODO: Should skip if !DoesPlatformSupportLumenGI()
-						break;
+			case ENaniteMeshPass::LumenCardCapture:
+				NaniteMeshProcessor = CreateLumenCardNaniteMeshProcessor(Scene, nullptr, &NaniteDrawListContext); // TODO: Should skip if !DoesPlatformSupportLumenGI()
+				break;
 
-					default:
-						check(false);
-					}
-					check(NaniteMeshProcessor);
-
-					int32 StaticMeshesCount = PrimitiveSceneInfo->StaticMeshes.Num();
-					for (int32 MeshIndex = 0; MeshIndex < StaticMeshesCount; ++MeshIndex)
-					{
-						FStaticMeshBatchRelevance& MeshRelevance = PrimitiveSceneInfo->StaticMeshRelevances[MeshIndex];
-						FStaticMeshBatch& Mesh = PrimitiveSceneInfo->StaticMeshes[MeshIndex];
+			default:
+				check(false);
+			}
+			check(NaniteMeshProcessor);
 
-						if (MeshRelevance.bSupportsNaniteRendering)
-						{
-							uint64 BatchElementMask = ~0ull;
-							NaniteMeshProcessor->AddMeshBatch(Mesh, BatchElementMask, Proxy);
+			int32 StaticMeshesCount = PrimitiveSceneInfo->StaticMeshes.Num();
+			for (int32 MeshIndex = 0; MeshIndex < StaticMeshesCount; ++MeshIndex)
+			{
+				FStaticMeshBatchRelevance& MeshRelevance = PrimitiveSceneInfo->StaticMeshRelevances[MeshIndex];
+				FStaticMeshBatch& Mesh = PrimitiveSceneInfo->StaticMeshes[MeshIndex];
 
-							FNaniteCommandInfo CommandInfo = NaniteDrawListContext.GetCommandInfoAndReset();
-							PrimitiveSceneInfo->NaniteCommandInfos[MeshPass].Add(CommandInfo);
-							const uint32 MaterialDepthId = CommandInfo.GetMaterialId();
+				if (MeshRelevance.bSupportsNaniteRendering)
+				{
+					uint64 BatchElementMask = ~0ull;
+					NaniteMeshProcessor->AddMeshBatch(Mesh, BatchElementMask, Proxy);
 
-							const uint32 SectionIndex = Mesh.SegmentIndex;
-							check(SectionIndex < uint32(PrimitiveSceneInfo->NaniteMaterialIds[MeshPass].Num()));
-							check(PrimitiveSceneInfo->NaniteMaterialIds[MeshPass][SectionIndex] == INDEX_NONE || PrimitiveSceneInfo->NaniteMaterialIds[MeshPass][SectionIndex] == MaterialDepthId);
-							PrimitiveSceneInfo->NaniteMaterialIds[MeshPass][SectionIndex] = MaterialDepthId;
-						}
-					}
+					FNaniteCommandInfo CommandInfo = NaniteDrawListContext.GetCommandInfoAndReset();
+					PrimitiveSceneInfo->NaniteCommandInfos[MeshPass].Add(CommandInfo);
+					const uint32 MaterialDepthId = CommandInfo.GetMaterialId();
 
-					NaniteMeshProcessor->~FMeshPassProcessor();
+					const uint32 SectionIndex = Mesh.SegmentIndex;
+					check(SectionIndex < uint32(PrimitiveSceneInfo->NaniteMaterialIds[MeshPass].Num()));
+					check(PrimitiveSceneInfo->NaniteMaterialIds[MeshPass][SectionIndex] == INDEX_NONE || PrimitiveSceneInfo->NaniteMaterialIds[MeshPass][SectionIndex] == MaterialDepthId);
+					PrimitiveSceneInfo->NaniteMaterialIds[MeshPass][SectionIndex] = MaterialDepthId;
 				}
 			}
+
+			NaniteMeshProcessor->~FMeshPassProcessor();
 		}
 	}
 }
