Fix dirty rectangle test that uses Area() >0 calculation.
Area() can overflow for a 2^16 * 2^16 giving a value of 0.

diff --git a/Engine/Source/Runtime/Renderer/Private/VT/RuntimeVirtualTextureSceneProxy.cpp b/Engine/Source/Runtime/Renderer/Private/VT/RuntimeVirtualTextureSceneProxy.cpp
--- a/Engine/Source/Runtime/Renderer/Private/VT/RuntimeVirtualTextureSceneProxy.cpp
+++ b/Engine/Source/Runtime/Renderer/Private/VT/RuntimeVirtualTextureSceneProxy.cpp
@@ -129,7 +129,7 @@ void FRuntimeVirtualTextureSceneProxy::FlushDirtyPages()
 	if (ProducerHandle.PackedValue != 0)
 	{
 		// Don't do any work if we won't mark anything dirty.
-		if (MaxDirtyLevel >= 0 && CombinedDirtyRect.Area() > 0)
+		if (MaxDirtyLevel >= 0 && CombinedDirtyRect.Width() != 0 && CombinedDirtyRect.Height() != 0)
 		{
 			//todo[vt]: 
 			// Profile to work out best heuristic for when we should use the CombinedDirtyRect
