Converting FMeshPassProcessors to use ODSC.

- FRuntimeVirtualTextureMeshProcessor

diff --git a/Engine/Source/Runtime/Renderer/Private/VT/RuntimeVirtualTextureRender.cpp b/Engine/Source/Runtime/Renderer/Private/VT/RuntimeVirtualTextureRender.cpp
--- a/Engine/Source/Runtime/Renderer/Private/VT/RuntimeVirtualTextureRender.cpp
+++ b/Engine/Source/Runtime/Renderer/Private/VT/RuntimeVirtualTextureRender.cpp
@@ -321,8 +321,41 @@ namespace RuntimeVirtualTexture
 		}
 
 	private:
+		bool TryAddMeshBatch(
+			const FMeshBatch& RESTRICT MeshBatch,
+			uint64 BatchElementMask,
+			const FPrimitiveSceneProxy* RESTRICT PrimitiveSceneProxy,
+			int32 StaticMeshId,
+			const FMaterialRenderProxy* MaterialRenderProxy,
+			const FMaterial* Material)
+		{
+			const uint8 OutputAttributeMask = Material->IsDefaultMaterial() ? 0xff : Material->GetRuntimeVirtualTextureOutputAttibuteMask_RenderThread();
+
+			if (OutputAttributeMask != 0)
+			{
+				switch ((ERuntimeVirtualTextureMaterialType)MeshBatch.RuntimeVirtualTextureMaterialType)
+				{
+				case ERuntimeVirtualTextureMaterialType::BaseColor:
+					return Process<FMaterialPolicy_BaseColor>(MeshBatch, BatchElementMask, StaticMeshId, OutputAttributeMask, PrimitiveSceneProxy, *MaterialRenderProxy, *Material);
+					break;
+				case ERuntimeVirtualTextureMaterialType::BaseColor_Normal_Specular:
+				case ERuntimeVirtualTextureMaterialType::BaseColor_Normal_Specular_YCoCg:
+				case ERuntimeVirtualTextureMaterialType::BaseColor_Normal_Specular_Mask_YCoCg:
+					return Process<FMaterialPolicy_BaseColorNormalSpecular>(MeshBatch, BatchElementMask, StaticMeshId, OutputAttributeMask, PrimitiveSceneProxy, *MaterialRenderProxy, *Material);
+					break;
+				case ERuntimeVirtualTextureMaterialType::WorldHeight:
+					return Process<FMaterialPolicy_WorldHeight>(MeshBatch, BatchElementMask, StaticMeshId, OutputAttributeMask, PrimitiveSceneProxy, *MaterialRenderProxy, *Material);
+					break;
+				default:
+					break;
+				}
+			}
+
+			return true;
+		}
+
 		template<class MaterialPolicy>
-		void Process(
+		bool Process(
 			const FMeshBatch& MeshBatch,
 			uint64 BatchElementMask,
 			int32 StaticMeshId,
@@ -335,10 +368,20 @@ namespace RuntimeVirtualTexture
 
 			TMeshProcessorShaders<
 				FShader_VirtualTextureMaterialDraw_VS< MaterialPolicy >,
-				FShader_VirtualTextureMaterialDraw_PS< MaterialPolicy > > Shaders;
+				FShader_VirtualTextureMaterialDraw_PS< MaterialPolicy > > VirtualTexturePassShaders;
+
+			FMaterialShaderTypes ShaderTypes;
+			ShaderTypes.AddShaderType<FShader_VirtualTextureMaterialDraw_VS< MaterialPolicy>>();
+			ShaderTypes.AddShaderType<FShader_VirtualTextureMaterialDraw_PS< MaterialPolicy>>();
+
+			FMaterialShaders Shaders;
+			if (!MaterialResource.TryGetShaders(ShaderTypes, VertexFactory->GetType(), Shaders))
+			{
+				return false;
+			}
 
-			Shaders.VertexShader = MaterialResource.GetShader< FShader_VirtualTextureMaterialDraw_VS< MaterialPolicy > >(VertexFactory->GetType());
-			Shaders.PixelShader = MaterialResource.GetShader< FShader_VirtualTextureMaterialDraw_PS< MaterialPolicy > >(VertexFactory->GetType());
+			Shaders.TryGetVertexShader(VirtualTexturePassShaders.VertexShader);
+			Shaders.TryGetPixelShader(VirtualTexturePassShaders.PixelShader);
 
 			DrawRenderState.SetBlendState(MaterialPolicy::GetBlendState(OutputAttributeMask));
 
@@ -361,12 +404,14 @@ namespace RuntimeVirtualTexture
 				MaterialRenderProxy,
 				MaterialResource,
 				DrawRenderState,
-				Shaders,
+				VirtualTexturePassShaders,
 				MeshFillMode,
 				MeshCullMode,
 				SortKey,
 				EMeshPassFeatures::Default,
 				ShaderElementData);
+
+			return true;
 		}
 
 	public:
@@ -374,29 +419,19 @@ namespace RuntimeVirtualTexture
 		{
 			if (MeshBatch.bRenderToVirtualTexture)
 			{
-				const FMaterialRenderProxy* FallbackMaterialRenderProxyPtr = nullptr;
-				const FMaterial& Material = MeshBatch.MaterialRenderProxy->GetMaterialWithFallback(FeatureLevel, FallbackMaterialRenderProxyPtr);
-				const FMaterialRenderProxy& MaterialRenderProxy = FallbackMaterialRenderProxyPtr ? *FallbackMaterialRenderProxyPtr : *MeshBatch.MaterialRenderProxy;
-				const uint8 OutputAttributeMask = FallbackMaterialRenderProxyPtr ? 0xff : Material.GetRuntimeVirtualTextureOutputAttibuteMask_RenderThread();
-
-				if (OutputAttributeMask != 0)
+				const FMaterialRenderProxy* MaterialRenderProxy = MeshBatch.MaterialRenderProxy;
+				while (MaterialRenderProxy)
 				{
-					switch ((ERuntimeVirtualTextureMaterialType)MeshBatch.RuntimeVirtualTextureMaterialType)
+					const FMaterial* Material = MaterialRenderProxy->GetMaterialNoFallback(FeatureLevel);
+					if (Material && Material->GetRenderingThreadShaderMap())
 					{
-					case ERuntimeVirtualTextureMaterialType::BaseColor:
-						Process<FMaterialPolicy_BaseColor>(MeshBatch, BatchElementMask, StaticMeshId, OutputAttributeMask, PrimitiveSceneProxy, MaterialRenderProxy, Material);
-						break;
-					case ERuntimeVirtualTextureMaterialType::BaseColor_Normal_Specular:
-					case ERuntimeVirtualTextureMaterialType::BaseColor_Normal_Specular_YCoCg:
-					case ERuntimeVirtualTextureMaterialType::BaseColor_Normal_Specular_Mask_YCoCg:
-						Process<FMaterialPolicy_BaseColorNormalSpecular>(MeshBatch, BatchElementMask, StaticMeshId, OutputAttributeMask, PrimitiveSceneProxy, MaterialRenderProxy, Material);
-						break;
-					case ERuntimeVirtualTextureMaterialType::WorldHeight:
-						Process<FMaterialPolicy_WorldHeight>(MeshBatch, BatchElementMask, StaticMeshId, OutputAttributeMask, PrimitiveSceneProxy, MaterialRenderProxy, Material);
-						break;
-					default:
-						break;
+						if (TryAddMeshBatch(MeshBatch, BatchElementMask, PrimitiveSceneProxy, StaticMeshId, MaterialRenderProxy, Material))
+						{
+							break;
+						}
 					}
+
+					MaterialRenderProxy = MaterialRenderProxy->GetFallback(FeatureLevel);
 				}
 			}
 		}
