Converting FMeshPassProcessors to use ODSC.

- FMeshDecalMeshProcessor

diff --git a/Engine/Source/Runtime/Renderer/Private/CompositionLighting/PostProcessMeshDecals.cpp b/Engine/Source/Runtime/Renderer/Private/CompositionLighting/PostProcessMeshDecals.cpp
--- a/Engine/Source/Runtime/Renderer/Private/CompositionLighting/PostProcessMeshDecals.cpp
+++ b/Engine/Source/Runtime/Renderer/Private/CompositionLighting/PostProcessMeshDecals.cpp
@@ -127,7 +127,15 @@ public:
 	virtual void AddMeshBatch(const FMeshBatch& RESTRICT MeshBatch, uint64 BatchElementMask, const FPrimitiveSceneProxy* RESTRICT PrimitiveSceneProxy, int32 StaticMeshId = -1) override final;
 
 private:
-	void Process(
+	bool TryAddMeshBatch(
+		const FMeshBatch& RESTRICT MeshBatch,
+		uint64 BatchElementMask,
+		const FPrimitiveSceneProxy* RESTRICT PrimitiveSceneProxy,
+		int32 StaticMeshId,
+		const FMaterialRenderProxy& MaterialRenderProxy,
+		const FMaterial& Material);
+
+	bool Process(
 		const FMeshBatch& MeshBatch,
 		uint64 BatchElementMask,
 		int32 StaticMeshId,
@@ -162,53 +170,76 @@ void FMeshDecalMeshProcessor::AddMeshBatch(const FMeshBatch& RESTRICT MeshBatch,
 	if (MeshBatch.bUseForMaterial && MeshBatch.IsDecal(FeatureLevel))
 	{
 		const FMaterialRenderProxy* MaterialRenderProxy = MeshBatch.MaterialRenderProxy;
-		const FMaterial* Material = MaterialRenderProxy->GetMaterialNoFallback(FeatureLevel);
-
-		if (Material && Material->IsDeferredDecal())
+		while (MaterialRenderProxy)
 		{
-			// We have no special engine material for decals since we don't want to eat the compilation & memory cost, so just skip if it failed to compile
-			if (Material->GetRenderingThreadShaderMap())
+			const FMaterial* Material = MaterialRenderProxy->GetMaterialNoFallback(FeatureLevel);
+			if (Material)
 			{
-				const EShaderPlatform ShaderPlatform = ViewIfDynamicMeshCommand->GetShaderPlatform();
-				const FDecalBlendDesc DecalBlendDesc = DecalRendering::ComputeDecalBlendDesc(ShaderPlatform, Material);
+				if (TryAddMeshBatch(MeshBatch, BatchElementMask, PrimitiveSceneProxy, StaticMeshId, *MaterialRenderProxy, *Material))
+				{
+					break;
+				}
+			}
 
-				const bool bShouldRender = 
-					DecalRendering::IsCompatibleWithRenderStage(DecalBlendDesc, PassDecalStage) &&
-					DecalRendering::GetRenderTargetMode(DecalBlendDesc, PassDecalStage) == RenderTargetMode;
+			MaterialRenderProxy = MaterialRenderProxy->GetFallback(FeatureLevel);
+		}
+	}
+}
 
-				if (bShouldRender)
-				{
-					const FMeshDrawingPolicyOverrideSettings OverrideSettings = ComputeMeshOverrideSettings(MeshBatch);
-					ERasterizerFillMode MeshFillMode = ComputeMeshFillMode(MeshBatch, *Material, OverrideSettings);
-					ERasterizerCullMode MeshCullMode = ComputeMeshCullMode(MeshBatch, *Material, OverrideSettings);
+bool FMeshDecalMeshProcessor::TryAddMeshBatch(
+	const FMeshBatch& RESTRICT MeshBatch,
+	uint64 BatchElementMask,
+	const FPrimitiveSceneProxy* RESTRICT PrimitiveSceneProxy,
+	int32 StaticMeshId,
+	const FMaterialRenderProxy& MaterialRenderProxy,
+	const FMaterial& Material)
+{
+	if (Material.IsDeferredDecal())
+	{
+		// We have no special engine material for decals since we don't want to eat the compilation & memory cost, so just skip if it failed to compile
+		if (Material.GetRenderingThreadShaderMap())
+		{
+			const EShaderPlatform ShaderPlatform = ViewIfDynamicMeshCommand->GetShaderPlatform();
+			const FDecalBlendDesc DecalBlendDesc = DecalRendering::ComputeDecalBlendDesc(ShaderPlatform, &Material);
+
+			const bool bShouldRender =
+				DecalRendering::IsCompatibleWithRenderStage(DecalBlendDesc, PassDecalStage) &&
+				DecalRendering::GetRenderTargetMode(DecalBlendDesc, PassDecalStage) == RenderTargetMode;
+
+			if (bShouldRender)
+			{
+				const FMeshDrawingPolicyOverrideSettings OverrideSettings = ComputeMeshOverrideSettings(MeshBatch);
+				ERasterizerFillMode MeshFillMode = ComputeMeshFillMode(MeshBatch, Material, OverrideSettings);
+				ERasterizerCullMode MeshCullMode = ComputeMeshCullMode(MeshBatch, Material, OverrideSettings);
 
-					if (ViewIfDynamicMeshCommand->Family->UseDebugViewPS())
+				if (ViewIfDynamicMeshCommand->Family->UseDebugViewPS())
+				{
+					// Deferred decals can only use translucent blend mode
+					if (ViewIfDynamicMeshCommand->Family->EngineShowFlags.ShaderComplexity)
 					{
-						// Deferred decals can only use translucent blend mode
-						if (ViewIfDynamicMeshCommand->Family->EngineShowFlags.ShaderComplexity)
-						{
-							// If we are in the translucent pass then override the blend mode, otherwise maintain additive blending.
-							PassDrawRenderState.SetBlendState(TStaticBlendState<CW_RGBA, BO_Add, BF_One, BF_One, BO_Add, BF_Zero, BF_One>::GetRHI());
-						}
-						else if (ViewIfDynamicMeshCommand->Family->GetDebugViewShaderMode() != DVSM_OutputMaterialTextureScales)
-						{
-							// Otherwise, force translucent blend mode (shaders will use an hardcoded alpha).
-							PassDrawRenderState.SetBlendState(TStaticBlendState<CW_RGBA, BO_Add, BF_SourceAlpha, BF_InverseSourceAlpha, BO_Add, BF_Zero, BF_InverseSourceAlpha>::GetRHI());
-						}
+						// If we are in the translucent pass then override the blend mode, otherwise maintain additive blending.
+						PassDrawRenderState.SetBlendState(TStaticBlendState<CW_RGBA, BO_Add, BF_One, BF_One, BO_Add, BF_Zero, BF_One>::GetRHI());
 					}
-					else
+					else if (ViewIfDynamicMeshCommand->Family->GetDebugViewShaderMode() != DVSM_OutputMaterialTextureScales)
 					{
-						PassDrawRenderState.SetBlendState(DecalRendering::GetDecalBlendState(DecalBlendDesc, PassDecalStage, RenderTargetMode));
+						// Otherwise, force translucent blend mode (shaders will use an hardcoded alpha).
+						PassDrawRenderState.SetBlendState(TStaticBlendState<CW_RGBA, BO_Add, BF_SourceAlpha, BF_InverseSourceAlpha, BO_Add, BF_Zero, BF_InverseSourceAlpha>::GetRHI());
 					}
-
-					Process(MeshBatch, BatchElementMask, StaticMeshId, PrimitiveSceneProxy, *MaterialRenderProxy, *Material, MeshFillMode, MeshCullMode);
 				}
+				else
+				{
+					PassDrawRenderState.SetBlendState(DecalRendering::GetDecalBlendState(DecalBlendDesc, PassDecalStage, RenderTargetMode));
+				}
+
+				return Process(MeshBatch, BatchElementMask, StaticMeshId, PrimitiveSceneProxy, MaterialRenderProxy, Material, MeshFillMode, MeshCullMode);
 			}
 		}
 	}
+
+	return true;
 }
 
-void FMeshDecalMeshProcessor::Process(
+bool FMeshDecalMeshProcessor::Process(
 	const FMeshBatch& MeshBatch,
 	uint64 BatchElementMask,
 	int32 StaticMeshId,
@@ -241,7 +272,7 @@ void FMeshDecalMeshProcessor::Process(
 	if (!MaterialResource.TryGetShaders(ShaderTypes, VertexFactoryType, Shaders))
 	{
 		// Skip rendering if any shaders missing
-		return;
+		return false;
 	}
 
 	TMeshProcessorShaders<
@@ -268,6 +299,8 @@ void FMeshDecalMeshProcessor::Process(
 		SortKey,
 		EMeshPassFeatures::Default,
 		ShaderElementData);
+
+	return true;
 }
 
 void DrawDecalMeshCommands(
