Implements experimental r.TemporalAA.RejectSeparateTranslucency

diff --git a/Engine/Shaders/Private/TemporalAA/TAACompareTranslucency.usf b/Engine/Shaders/Private/TemporalAA/TAACompareTranslucency.usf
new file mode 100644
--- /dev/null
+++ b/Engine/Shaders/Private/TemporalAA/TAACompareTranslucency.usf
@@ -0,0 +1,218 @@
+// Copyright Epic Games, Inc. All Rights Reserved.
+
+#include "TAACommon.ush"
+
+
+//------------------------------------------------------- CONFIG
+
+#define TILE_SIZE 8
+
+#define CONFIG_KERNEL 0
+
+#define CONFIG_OUTPUT_HALF_RES (DIM_OUTPUT_HALF_RES)
+
+#define CONFIG_OUTPUT_INTERFERENCE_SEEDS 0
+
+#define CONFIG_SAMPLES 5
+
+
+//------------------------------------------------------- PARAMETERS
+
+float2 TranslucencyInfo_Extent;
+float2 TranslucencyInfo_ExtentInverse;
+float2 TranslucencyInfo_ScreenPosToViewportScale;
+float2 TranslucencyInfo_ScreenPosToViewportBias;
+uint2  TranslucencyInfo_ViewportMin;
+uint2  TranslucencyInfo_ViewportMax;
+float2 TranslucencyInfo_ViewportSize;
+float2 TranslucencyInfo_ViewportSizeInverse;
+float2 TranslucencyInfo_UVViewportMin;
+float2 TranslucencyInfo_UVViewportMax;
+float2 TranslucencyInfo_UVViewportSize;
+float2 TranslucencyInfo_UVViewportSizeInverse;
+float2 TranslucencyInfo_UVViewportBilinearMin;
+float2 TranslucencyInfo_UVViewportBilinearMax;
+
+float2 PrevTranslucencyInfo_Extent;
+float2 PrevTranslucencyInfo_ExtentInverse;
+float2 PrevTranslucencyInfo_ScreenPosToViewportScale;
+float2 PrevTranslucencyInfo_ScreenPosToViewportBias;
+uint2  PrevTranslucencyInfo_ViewportMin;
+uint2  PrevTranslucencyInfo_ViewportMax;
+float2 PrevTranslucencyInfo_ViewportSize;
+float2 PrevTranslucencyInfo_ViewportSizeInverse;
+float2 PrevTranslucencyInfo_UVViewportMin;
+float2 PrevTranslucencyInfo_UVViewportMax;
+float2 PrevTranslucencyInfo_UVViewportSize;
+float2 PrevTranslucencyInfo_UVViewportSizeInverse;
+float2 PrevTranslucencyInfo_UVViewportBilinearMin;
+float2 PrevTranslucencyInfo_UVViewportBilinearMax;
+float PrevTranslucencyPreExposureCorrection;
+
+FScreenTransform InputPixelPosToScreenPos;
+FScreenTransform ScreenPosToPrevTranslucencyTextureUV;
+
+Texture2D<taa_half4> TranslucencyTexture;
+Texture2D<float2> DilatedVelocityTexture;
+Texture2D<taa_half4> PrevTranslucencyTexture;
+
+RWTexture2D<taa_half> TranslucencyRejectionOutput;
+
+
+//------------------------------------------------------- ENTRY POINT
+
+taa_half4 TransformToComparisonColorSpace(taa_half4 Color)
+{
+	return taa_half4(RGBToYCoCg(Color.rgb), Color.a);
+}
+
+
+//------------------------------------------------------- ENTRY POINT
+
+[numthreads(TILE_SIZE, TILE_SIZE, 1)]
+void MainCS(
+	uint2 GroupId : SV_GroupID,
+	uint GroupThreadIndex : SV_GroupIndex)
+{
+	float4 Debug = 0;
+	
+	uint2 InputPixelPos = (
+		ZOrder2D(GroupThreadIndex, log2(TILE_SIZE)) +
+		GroupId * uint2(TILE_SIZE, TILE_SIZE));
+
+	float2 ScreenPos = ApplyScreenTransform(float2(InputPixelPos), InputPixelPosToScreenPos);
+	float2 EncodedVelocity = DilatedVelocityTexture[InputPixelPos];
+	
+	// Sample current frame translucency's
+	taa_half4 Translucency;
+	taa_half4 TranslucencyMin;
+	taa_half4 TranslucencyMax;
+	{
+		UNROLL_N(CONFIG_SAMPLES)
+		for (uint SampleId = 0; SampleId < CONFIG_SAMPLES; SampleId++)
+		{
+			taa_short2 SampleInputPixelPos;
+			
+			// TODO: handle SeparateTranslucencyScreenPercentage
+			#if CONFIG_SAMPLES == 9
+			{
+				taa_short2 iPixelOffset = taa_short2(kOffsets3x3[kSquareIndexes3x3[SampleId]]);
+				
+				SampleInputPixelPos = AddAndClampPixelOffset(
+					InputPixelPos,
+					iPixelOffset, iPixelOffset,
+					InputPixelPosMin, InputPixelPosMax);
+			}
+			#elif CONFIG_SAMPLES == 5
+			{
+				taa_short2 iPixelOffset = taa_short2(kOffsets3x3[kPlusIndexes3x3[SampleId]]);
+					
+				SampleInputPixelPos = AddAndClampPixelOffset(
+					InputPixelPos,
+					iPixelOffset, iPixelOffset,
+					InputPixelPosMin, InputPixelPosMax);
+			}
+			#else
+				#error Unknown sample count
+			#endif
+
+			taa_half4 TranslucencyColor = TransformToComparisonColorSpace(TranslucencyTexture[SampleInputPixelPos]);
+
+			if (SampleId == 0)
+			{
+				Translucency = TranslucencyColor;
+				TranslucencyMin = TranslucencyColor;
+				TranslucencyMax = TranslucencyColor;
+			}
+			else
+			{
+				TranslucencyMin = min(TranslucencyMin, TranslucencyColor);
+				TranslucencyMax = max(TranslucencyMax, TranslucencyColor);
+			}
+		}
+	}
+
+	// Sample previous frame's translucency
+	taa_half4 PrevTranslucency;
+	taa_half4 PrevTranslucencyMin;
+	taa_half4 PrevTranslucencyMax;
+	{
+		float2 PrevScreenPos = ScreenPos - DecodeVelocityFromTexture(float4(EncodedVelocity, 0.0, 0.0)).xy;
+		float2 PrevTranslucencyTextureUV = ApplyScreenTransform(PrevScreenPos, ScreenPosToPrevTranslucencyTextureUV);
+	
+		UNROLL_N(CONFIG_SAMPLES)
+		for (uint SampleId = 0; SampleId < CONFIG_SAMPLES; SampleId++)
+		{
+			taa_half2 PixelOffset;
+			
+			// TODO: handle r.SeparateTranslucencyScreenPercentage
+			#if CONFIG_SAMPLES == 9
+			{
+				taa_short2 iPixelOffset = taa_short2(kOffsets3x3[kSquareIndexes3x3[SampleId]]);
+				PixelOffset = taa_half2(iPixelOffset);
+			}
+			#elif CONFIG_SAMPLES == 5
+			{
+				taa_short2 iPixelOffset = taa_short2(kOffsets3x3[kPlusIndexes3x3[SampleId]]);
+				PixelOffset = taa_half2(iPixelOffset);
+			}
+			#else
+				#error Unknown sample count
+			#endif
+
+			float2 SampleTextureUV = clamp(
+				PrevTranslucencyTextureUV + PixelOffset * PrevTranslucencyInfo_ExtentInverse,
+				PrevTranslucencyInfo_UVViewportBilinearMin,
+				PrevTranslucencyInfo_UVViewportBilinearMax);
+		
+			taa_half4 PrevTranslucencyColor = TransformToComparisonColorSpace(
+				PrevTranslucencyTexture.SampleLevel(GlobalPointClampedSampler, SampleTextureUV, 0));
+
+			if (SampleId == 0)
+			{
+				PrevTranslucency = PrevTranslucencyColor;
+				PrevTranslucencyMin = PrevTranslucencyColor;
+				PrevTranslucencyMax = PrevTranslucencyColor;
+			}
+			else
+			{
+				PrevTranslucencyMin = min(PrevTranslucencyMin, PrevTranslucencyColor);
+				PrevTranslucencyMax = max(PrevTranslucencyMax, PrevTranslucencyColor);
+			}
+		}
+
+		PrevTranslucency *= PrevTranslucencyPreExposureCorrection;
+		PrevTranslucencyMin *= PrevTranslucencyPreExposureCorrection;
+		PrevTranslucencyMax *= PrevTranslucencyPreExposureCorrection;
+	}
+
+	taa_half Rejection = 1.0;
+	#if 1
+	{
+		taa_half4 ClampedPrevTranslucency = clamp(PrevTranslucency, TranslucencyMin, TranslucencyMax);
+		taa_half4 ClampedTranslucency = clamp(Translucency, PrevTranslucencyMin, PrevTranslucencyMax);
+
+		// TODO: reject with alpha too? Might not need since translucency is premultiplied RGBA 
+		taa_half ForwardRejection = MeasureRejectionFactor(
+			PrevTranslucency.xyz, ClampedPrevTranslucency.xyz,
+			Translucency.xyz, TranslucencyMin.xyz, TranslucencyMax.xyz);
+		
+		taa_half BackwardRejection = MeasureRejectionFactor(
+			Translucency.xyz, ClampedTranslucency.xyz,
+			PrevTranslucency.xyz, PrevTranslucencyMin.xyz, PrevTranslucencyMax.xyz);
+		
+		Rejection = min(ForwardRejection, BackwardRejection);
+	}
+	#endif
+
+	if (all(InputPixelPos < InputInfo_ViewportMax))
+	{
+		TranslucencyRejectionOutput[InputPixelPos] = Rejection;
+
+		#if DEBUG_OUTPUT
+		{
+			DebugOutput[InputPixelPos] = Debug;
+		}
+		#endif
+	}
+}
diff --git a/Engine/Shaders/Private/TemporalAA/TAAUpdateHistory.usf b/Engine/Shaders/Private/TemporalAA/TAAUpdateHistory.usf
--- a/Engine/Shaders/Private/TemporalAA/TAAUpdateHistory.usf
+++ b/Engine/Shaders/Private/TemporalAA/TAAUpdateHistory.usf
@@ -20,6 +20,8 @@
 
 #define CONFIG_LOW_FREQUENCY_DRIFT_REJECTION 1
 
+#define CONFIG_TRANSLUCENCY_REJECTION 1
+
 
 //------------------------------------------------------- CONSTANTS
 
@@ -31,6 +33,7 @@
 Texture2D<taa_half3> InputSceneColorTexture;
 Texture2D<uint2> InputSceneStencilTexture;
 Texture2D<taa_half> HistoryRejectionTexture;
+Texture2D<taa_half> TranslucencyRejectionTexture;
 
 Texture2D<float2> DilatedVelocityTexture;
 Texture2D<uint> ParallaxFactorTexture;
@@ -39,6 +42,7 @@ Texture2D<taa_half> ParallaxRejectionMaskTexture;
 float4 HistoryPixelPosToScreenPos;
 float4 HistoryPixelPosToPPCo;
 float3 HistoryQuantizationError;
+float MinTranslucencyRejection;
 
 Texture2D<taa_half3> PrevHistory_Textures_0;
 Texture2D<taa_half3> PrevHistory_Textures_1;
@@ -72,6 +76,54 @@ taa_half3 TransformColorForClampingBox(taa_half3 Color)
 	return RGBToYCoCg(Color);
 }
 
+void ComputeInputKernelSamplePosition(
+	taa_short2 InputPixelPos,
+	taa_half2 dKO,
+	uint SampleId,
+	out taa_short2 SampleInputPixelPos,
+	out taa_half2 PixelOffset)
+{
+	
+	#if CONFIG_SAMPLES == 9
+	{
+		taa_short2 iPixelOffset = taa_short2(kOffsets3x3[kSquareIndexes3x3[SampleId]]);
+		PixelOffset = taa_half2(iPixelOffset);
+				
+		SampleInputPixelPos = AddAndClampPixelOffset(
+			InputPixelPos,
+			iPixelOffset, iPixelOffset,
+			InputPixelPosMin, InputPixelPosMax);
+	}
+	#elif CONFIG_SAMPLES == 5 || CONFIG_SAMPLES == 6
+	{
+		if (SampleId == 5)
+		{
+			taa_short2 iPixelOffset;
+			#if CONFIG_COMPILE_FP16
+				iPixelOffset = int16_t2(1, 1) - int16_t2((asuint16(dKO) & uint16_t(0x8000)) >> uint16_t(14));
+				PixelOffset = asfloat16(asuint16(half(1.0)).xx | (asuint16(dKO) & uint16_t(0x8000)));
+			#else
+				iPixelOffset = SignFastInt(dKO);
+				PixelOffset = asfloat(asuint(1.0).xx | (asuint(dKO) & uint(0x80000000)));
+			#endif
+						
+			SampleInputPixelPos = ClampPixelOffset(InputPixelPos, InputPixelPosMin, InputPixelPosMax);
+		}
+		else
+		{
+			taa_short2 iPixelOffset = taa_short2(kOffsets3x3[kPlusIndexes3x3[SampleId]]);
+			PixelOffset = taa_half2(iPixelOffset);
+					
+			SampleInputPixelPos = AddAndClampPixelOffset(
+				InputPixelPos,
+				iPixelOffset, iPixelOffset,
+				InputPixelPosMin, InputPixelPosMax);
+		}
+	}
+	#else
+		#error Unknown sample count
+	#endif
+}
 
 //------------------------------------------------------- ENTRY POINT
 
@@ -104,9 +156,14 @@ void MainCS(
 	float2 PrevScreenPos = ScreenPos;
 	taa_half ParallaxRejectionMask = taa_half(1.0);
 	taa_half LowFrequencyRejection = taa_half(1.0);
+	taa_half TranslucencyRejection = taa_half(1.0);
 	taa_half OutputPixelVelocity = taa_half(0.0);
 	#if 1
+	ISOLATE
 	{
+		// Vector in pixel between pixel K -> O.
+		taa_half2 dKO = taa_half2(PPCo - PPCk);
+
 		float2 EncodedVelocity = DilatedVelocityTexture[InputPixelPos];
 		ParallaxRejectionMask = ParallaxRejectionMaskTexture[InputPixelPos];
 
@@ -117,16 +174,38 @@ void MainCS(
 
 		taa_ushort2 RejectionPixelPos = (taa_ushort2(InputPixelPos) - taa_short2(InputPixelPosMin)) / 2;
 		LowFrequencyRejection = HistoryRejectionTexture[RejectionPixelPos];
+		TranslucencyRejection = TranslucencyRejectionTexture[InputPixelPos];
 		
+		UNROLL_N(CONFIG_SAMPLES - 1)
+		for (uint SampleId = 1; SampleId < CONFIG_SAMPLES; SampleId++)
+		{
+			taa_short2 SampleInputPixelPos;
+			taa_half2 PixelOffset;
+			ComputeInputKernelSamplePosition(
+				InputPixelPos, dKO, SampleId,
+				/* out */ SampleInputPixelPos,
+				/* out */ PixelOffset);
+			
+			TranslucencyRejection = min(TranslucencyRejection, TranslucencyRejectionTexture[SampleInputPixelPos]);
+		}
+
+		TranslucencyRejection = max(TranslucencyRejection, taa_half(MinTranslucencyRejection));
+
 		#if !CONFIG_CLAMP
 		{
 			ParallaxRejectionMask = taa_half(1.0);
 			LowFrequencyRejection = taa_half(1.0);
+			TranslucencyRejection = taa_half(1.0);
 		}
 		#endif
 	}
 	#endif
 
+	// Reject the translucency to respond quicker on firefliers.
+	#if CONFIG_TRANSLUCENCY_REJECTION
+		LowFrequencyRejection = min(LowFrequencyRejection, TranslucencyRejection);
+	#endif
+
 	// Fetch wehther the pixel is responsive AA or not.
 	bool bIsResponsiveAAPixel = false;
 	#if CONFIG_RESPONSIVE_STENCIL
@@ -178,46 +257,10 @@ void MainCS(
 		{
 			taa_short2 SampleInputPixelPos;
 			taa_half2 PixelOffset;
-			
-			#if CONFIG_SAMPLES == 9
-			{
-				taa_short2 iPixelOffset = taa_short2(kOffsets3x3[kSquareIndexes3x3[SampleId]]);
-				PixelOffset = taa_half2(iPixelOffset);
-				
-				SampleInputPixelPos = AddAndClampPixelOffset(
-					InputPixelPos,
-					iPixelOffset, iPixelOffset,
-					InputPixelPosMin, InputPixelPosMax);
-			}
-			#elif CONFIG_SAMPLES == 5 || CONFIG_SAMPLES == 6
-			{
-				if (SampleId == 5)
-				{
-					taa_short2 iPixelOffset;
-					#if CONFIG_COMPILE_FP16
-						iPixelOffset = int16_t2(1, 1) - int16_t2((asuint16(dKO) & uint16_t(0x8000)) >> uint16_t(14));
-						PixelOffset = asfloat16(asuint16(half(1.0)).xx | (asuint16(dKO) & uint16_t(0x8000)));
-					#else
-						iPixelOffset = SignFastInt(dKO);
-						PixelOffset = asfloat(asuint(1.0).xx | (asuint(dKO) & uint(0x80000000)));
-					#endif
-						
-					SampleInputPixelPos = ClampPixelOffset(InputPixelPos, InputPixelPosMin, InputPixelPosMax);
-				}
-				else
-				{
-					taa_short2 iPixelOffset = taa_short2(kOffsets3x3[kPlusIndexes3x3[SampleId]]);
-					PixelOffset = taa_half2(iPixelOffset);
-					
-					SampleInputPixelPos = AddAndClampPixelOffset(
-						InputPixelPos,
-						iPixelOffset, iPixelOffset,
-						InputPixelPosMin, InputPixelPosMax);
-				}
-			}
-			#else
-				#error Unknown sample count
-			#endif
+			ComputeInputKernelSamplePosition(
+				InputPixelPos, dKO, SampleId,
+				/* out */ SampleInputPixelPos,
+				/* out */ PixelOffset);
 
 			taa_half3 InputColor = InputSceneColorTexture[SampleInputPixelPos];
 
diff --git a/Engine/Source/Runtime/Renderer/Private/PostProcess/PostProcessing.cpp b/Engine/Source/Runtime/Renderer/Private/PostProcess/PostProcessing.cpp
--- a/Engine/Source/Runtime/Renderer/Private/PostProcess/PostProcessing.cpp
+++ b/Engine/Source/Runtime/Renderer/Private/PostProcess/PostProcessing.cpp
@@ -577,6 +577,7 @@ void AddPostProcessingPasses(FRDGBuilder& GraphBuilder, const FViewInfo& View, c
 			UpscalerPassInputs.SceneColorTexture = SceneColor.Texture;
 			UpscalerPassInputs.SceneDepthTexture = SceneDepth.Texture;
 			UpscalerPassInputs.SceneVelocityTexture = Velocity.Texture;
+			UpscalerPassInputs.SeparateTranslucencyTextures = Inputs.SeparateTranslucencyTextures;
 
 			UpscalerToUse->AddPasses(
 				GraphBuilder,
diff --git a/Engine/Source/Runtime/Renderer/Private/PostProcess/TemporalAA.cpp b/Engine/Source/Runtime/Renderer/Private/PostProcess/TemporalAA.cpp
--- a/Engine/Source/Runtime/Renderer/Private/PostProcess/TemporalAA.cpp
+++ b/Engine/Source/Runtime/Renderer/Private/PostProcess/TemporalAA.cpp
@@ -96,6 +96,11 @@ TAutoConsoleVariable<int32> CVarTAAEnableAntiInterference(
 	TEXT("Enable heuristic to detect geometric interference between input pixel grid alignement and structured geometry."),
 	ECVF_RenderThreadSafe);
 
+TAutoConsoleVariable<int32> CVarTAARejectTranslucency(
+	TEXT("r.TemporalAA.RejectSeparateTranslucency"), 0,
+	TEXT("Enable heuristic to reject based on the Separate Translucency."),
+	ECVF_RenderThreadSafe);
+
 #if COMPILE_TAA_DEBUG_PASSES
 
 TAutoConsoleVariable<int32> CVarTAASetupDebugPasses(
@@ -433,6 +438,30 @@ class FTSRDecimateHistoryCS : public FTemporalSuperResolutionShader
 	END_SHADER_PARAMETER_STRUCT()
 }; // class FTSRDecimateHistoryCS
 
+class FTSRCompareTranslucencyCS : public FTemporalSuperResolutionShader
+{
+	DECLARE_GLOBAL_SHADER(FTSRCompareTranslucencyCS);
+	SHADER_USE_PARAMETER_STRUCT(FTSRCompareTranslucencyCS, FTemporalSuperResolutionShader);
+
+	BEGIN_SHADER_PARAMETER_STRUCT(FParameters, )
+		SHADER_PARAMETER_STRUCT_INCLUDE(FTSRCommonParameters, CommonParameters)
+
+		SHADER_PARAMETER_STRUCT(FScreenPassTextureViewportParameters, TranslucencyInfo)
+		SHADER_PARAMETER_STRUCT(FScreenPassTextureViewportParameters, PrevTranslucencyInfo)
+		SHADER_PARAMETER(float, PrevTranslucencyPreExposureCorrection)
+
+		SHADER_PARAMETER(FScreenTransform, InputPixelPosToScreenPos)
+		SHADER_PARAMETER(FScreenTransform, ScreenPosToPrevTranslucencyTextureUV)
+
+		SHADER_PARAMETER_RDG_TEXTURE(Texture2D, DilatedVelocityTexture)
+		SHADER_PARAMETER_RDG_TEXTURE(Texture2D, TranslucencyTexture)
+		SHADER_PARAMETER_RDG_TEXTURE(Texture2D, PrevTranslucencyTexture)
+
+		SHADER_PARAMETER_RDG_TEXTURE_UAV(RWTexture2D, TranslucencyRejectionOutput)
+		SHADER_PARAMETER_RDG_TEXTURE_UAV(RWTexture2D, DebugOutput)
+	END_SHADER_PARAMETER_STRUCT()
+}; // class FTSRCompareTranslucencyCS
+
 class FTSRDetectInterferenceCS : public FTemporalSuperResolutionShader
 {
 	DECLARE_GLOBAL_SHADER(FTSRDetectInterferenceCS);
@@ -548,6 +577,7 @@ class FTSRUpdateHistoryCS : public FTemporalSuperResolutionShader
 		SHADER_PARAMETER_RDG_TEXTURE(Texture2D, InputSceneColorTexture)
 		SHADER_PARAMETER_RDG_TEXTURE_SRV(Texture2D, InputSceneStencilTexture)
 		SHADER_PARAMETER_RDG_TEXTURE(Texture2D, HistoryRejectionTexture)
+		SHADER_PARAMETER_RDG_TEXTURE(Texture2D, TranslucencyRejectionTexture)
 
 		SHADER_PARAMETER_RDG_TEXTURE(Texture2D, DilatedVelocityTexture)
 		SHADER_PARAMETER_RDG_TEXTURE(Texture2D, ParallaxFactorTexture)
@@ -556,6 +586,7 @@ class FTSRUpdateHistoryCS : public FTemporalSuperResolutionShader
 		SHADER_PARAMETER(FScreenTransform, HistoryPixelPosToScreenPos)
 		SHADER_PARAMETER(FScreenTransform, HistoryPixelPosToPPCo)
 		SHADER_PARAMETER(FVector, HistoryQuantizationError)
+		SHADER_PARAMETER(float, MinTranslucencyRejection)
 
 		SHADER_PARAMETER_STRUCT_INCLUDE(FTSRPrevHistoryParameters, PrevHistoryParameters)
 		SHADER_PARAMETER_STRUCT(FTSRHistoryTextures, PrevHistory)
@@ -576,6 +607,7 @@ IMPLEMENT_GLOBAL_SHADER(FTAAStandaloneCS, "/Engine/Private/TemporalAA/TAAStandal
 IMPLEMENT_GLOBAL_SHADER(FTSRClearPrevTexturesCS, "/Engine/Private/TemporalAA/TAAClearPrevTextures.usf", "MainCS", SF_Compute);
 IMPLEMENT_GLOBAL_SHADER(FTSRDilateVelocityCS, "/Engine/Private/TemporalAA/TAADilateVelocity.usf", "MainCS", SF_Compute);
 IMPLEMENT_GLOBAL_SHADER(FTSRDecimateHistoryCS, "/Engine/Private/TemporalAA/TAADecimateHistory.usf", "MainCS", SF_Compute);
+IMPLEMENT_GLOBAL_SHADER(FTSRCompareTranslucencyCS, "/Engine/Private/TemporalAA/TAACompareTranslucency.usf", "MainCS", SF_Compute);
 IMPLEMENT_GLOBAL_SHADER(FTSRDetectInterferenceCS, "/Engine/Private/TemporalAA/TAADetectInterference.usf", "MainCS", SF_Compute);
 IMPLEMENT_GLOBAL_SHADER(FTSRFilterFrequenciesCS, "/Engine/Private/TemporalAA/TAAFilterFrequencies.usf", "MainCS", SF_Compute);
 IMPLEMENT_GLOBAL_SHADER(FTSRCompareHistoryCS, "/Engine/Private/TemporalAA/TAACompareHistory.usf", "MainCS", SF_Compute);
@@ -1137,7 +1169,6 @@ static void AddTemporalSuperResolutionPasses(
 	FIntRect* OutSceneColorViewRect)
 {
 	const FTemporalAAHistory& InputHistory = View.PrevViewInfo.TemporalAAHistory;
-	FTemporalAAHistory* OutputHistory = &View.ViewState->PrevFrameViewInfo.TemporalAAHistory;
 
 #if COMPILE_TAA_DEBUG_PASSES
 	const bool bSetupDebugPasses = CVarTAASetupDebugPasses.GetValueOnRenderThread() != 0;
@@ -1150,6 +1181,8 @@ static void AddTemporalSuperResolutionPasses(
 
 	bool bEnableInterferenceHeuristic = CVarTAAEnableAntiInterference.GetValueOnRenderThread() != 0;
 
+	bool bRejectSeparateTranslucency = PassInputs.SeparateTranslucencyTextures != nullptr && CVarTAARejectTranslucency.GetValueOnRenderThread() != 0;
+
 	FIntPoint InputExtent = PassInputs.SceneColorTexture->Desc.Extent;
 	FIntRect InputRect = View.ViewRect;
 
@@ -1210,6 +1243,9 @@ static void AddTemporalSuperResolutionPasses(
 	FRDGTextureRef BlackDummy = GraphBuilder.RegisterExternalTexture(GSystemTextures.BlackDummy);
 	FRDGTextureRef WhiteDummy = GraphBuilder.RegisterExternalTexture(GSystemTextures.WhiteDummy);
 
+	FRDGTextureRef SeparateTranslucencyTexture = bRejectSeparateTranslucency ?
+		PassInputs.SeparateTranslucencyTextures->GetColorForRead(GraphBuilder) : nullptr;
+
 	FTSRCommonParameters CommonParameters;
 	{
 		CommonParameters.InputInfo = GetScreenPassTextureViewportParameters(FScreenPassTextureViewport(
@@ -1516,6 +1552,63 @@ static void AddTemporalSuperResolutionPasses(
 			FComputeShaderUtils::GetGroupCount(InputRect.Size(), 8));
 	}
 
+	FRDGTextureRef TranslucencyRejectionTexture = nullptr;
+	if (bRejectSeparateTranslucency && View.PrevViewInfo.SeparateTranslucency != nullptr)
+	{
+		FRDGTextureRef PrevTranslucencyTexture;
+		FScreenPassTextureViewport PrevTranslucencyViewport;
+
+		if (View.PrevViewInfo.SeparateTranslucency)
+		{
+
+			PrevTranslucencyTexture = GraphBuilder.RegisterExternalTexture(View.PrevViewInfo.SeparateTranslucency);
+			PrevTranslucencyViewport = FScreenPassTextureViewport(PrevTranslucencyTexture->Desc.Extent, View.PrevViewInfo.ViewRect);
+
+		}
+		else
+		{
+			PrevTranslucencyTexture = BlackDummy;
+			PrevTranslucencyViewport = FScreenPassTextureViewport(FIntPoint(1, 1), FIntRect(FIntPoint(0, 0), FIntPoint(1, 1)));
+		}
+
+		{
+			FRDGTextureDesc Desc = FRDGTextureDesc::Create2D(
+				InputExtent,
+				PF_R8,
+				FClearValueBinding::None,
+				/* InFlags = */ TexCreate_ShaderResource | TexCreate_UAV);
+
+			TranslucencyRejectionTexture = GraphBuilder.CreateTexture(Desc, TEXT("TSR.TranslucencyRejection"));
+		}
+
+		FTSRCompareTranslucencyCS::FParameters* PassParameters = GraphBuilder.AllocParameters<FTSRCompareTranslucencyCS::FParameters>();
+		PassParameters->CommonParameters = CommonParameters;
+
+		PassParameters->TranslucencyInfo = GetScreenPassTextureViewportParameters(FScreenPassTextureViewport(
+			InputExtent, InputRect));
+		PassParameters->PrevTranslucencyInfo = GetScreenPassTextureViewportParameters(PrevTranslucencyViewport);
+		PassParameters->PrevTranslucencyPreExposureCorrection = PrevHistoryParameters.HistoryPreExposureCorrection;
+
+		PassParameters->InputPixelPosToScreenPos = (FScreenTransform::Identity + 0.5) * CommonParameters.InputInfo.ViewportSizeInverse * FScreenTransform::ViewportUVToScreenPos;
+		PassParameters->ScreenPosToPrevTranslucencyTextureUV = FScreenTransform::ChangeTextureBasisFromTo(
+			PrevTranslucencyViewport, FScreenTransform::ETextureBasis::ScreenPosition, FScreenTransform::ETextureBasis::TextureUV);
+
+		PassParameters->DilatedVelocityTexture = DilatedVelocityTexture;
+		PassParameters->TranslucencyTexture = SeparateTranslucencyTexture;
+		PassParameters->PrevTranslucencyTexture = PrevTranslucencyTexture;
+
+		PassParameters->TranslucencyRejectionOutput = GraphBuilder.CreateUAV(TranslucencyRejectionTexture);
+		PassParameters->DebugOutput = CreateDebugUAV(LowFrequencyExtent, TEXT("Debug.TSR.ComparetTranslucency"));
+
+		TShaderMapRef<FTSRCompareTranslucencyCS> ComputeShader(View.ShaderMap);
+		FComputeShaderUtils::AddPass(
+			GraphBuilder,
+			RDG_EVENT_NAME("TSR ComparetTranslucency %dx%d", InputRect.Width(), InputRect.Height()),
+			ComputeShader,
+			PassParameters,
+			FComputeShaderUtils::GetGroupCount(InputRect.Size(), 8));
+	}
+
 	// Detect interference between geometry and alignement of input texel centers. It is not about awnser whether an
 	// interference has happened in the past, because interference change based on input resolution or camera position.
 	// So to remain stable on camera movement and input resolution change, it is about awnsering the question on whether
@@ -1789,6 +1882,8 @@ static void AddTemporalSuperResolutionPasses(
 		PassParameters->InputSceneStencilTexture = GraphBuilder.CreateSRV(
 			FRDGTextureSRVDesc::CreateWithPixelFormat(PassInputs.SceneDepthTexture, PF_X24_G8));
 		PassParameters->HistoryRejectionTexture = DilatedHistoryRejectionTexture;
+		PassParameters->TranslucencyRejectionTexture = TranslucencyRejectionTexture ? TranslucencyRejectionTexture : BlackDummy;
+
 		PassParameters->DilatedVelocityTexture = DilatedVelocityTexture;
 		PassParameters->ParallaxFactorTexture = ParallaxFactorTexture;
 		PassParameters->ParallaxRejectionMaskTexture = ParallaxRejectionMaskTexture;
@@ -1797,6 +1892,7 @@ static void AddTemporalSuperResolutionPasses(
 		PassParameters->HistoryPixelPosToScreenPos = HistoryPixelPosToViewportUV * FScreenTransform::ViewportUVToScreenPos;
 		PassParameters->HistoryPixelPosToPPCo = HistoryPixelPosToViewportUV * CommonParameters.InputInfo.ViewportSize + CommonParameters.InputJitter + CommonParameters.InputPixelPosMin;
 		PassParameters->HistoryQuantizationError = ComputePixelFormatQuantizationError(History.Textures[0]->Desc.Format);
+		PassParameters->MinTranslucencyRejection = TranslucencyRejectionTexture == nullptr ? 1.0 : 0.0;
 
 		PassParameters->PrevHistoryParameters = PrevHistoryParameters;
 		PassParameters->PrevHistory = PrevHistory;
@@ -1862,6 +1958,7 @@ static void AddTemporalSuperResolutionPasses(
 
 	if (!View.bStatePrevViewInfoIsReadOnly)
 	{
+		FTemporalAAHistory* OutputHistory = &View.ViewState->PrevFrameViewInfo.TemporalAAHistory;
 		OutputHistory->SafeRelease();
 
 		for (int32 i = 0; i < History.LowResTextures.Num(); i++)
@@ -1890,6 +1987,12 @@ static void AddTemporalSuperResolutionPasses(
 
 		OutputHistory->ViewportRect = FIntRect(FIntPoint(0, 0), HistorySize);
 		OutputHistory->ReferenceBufferSize = HistoryExtent;
+
+		if (bRejectSeparateTranslucency)
+		{
+			GraphBuilder.QueueTextureExtraction(
+				SeparateTranslucencyTexture, &View.ViewState->PrevFrameViewInfo.SeparateTranslucency);
+		}
 	}
 
 	// If we upscaled the history buffer, downsize back to the secondary screen percentage size.
diff --git a/Engine/Source/Runtime/Renderer/Private/PostProcess/TemporalAA.h b/Engine/Source/Runtime/Renderer/Private/PostProcess/TemporalAA.h
--- a/Engine/Source/Runtime/Renderer/Private/PostProcess/TemporalAA.h
+++ b/Engine/Source/Runtime/Renderer/Private/PostProcess/TemporalAA.h
@@ -5,6 +5,7 @@
 #include "ScreenPass.h"
 
 struct FTemporalAAHistory;
+class FSeparateTranslucencyTextures;
 
 
 /** Configuration of the main temporal AA pass. */
@@ -170,9 +171,10 @@ public:
 	{
 		bool bAllowDownsampleSceneColor;
 		EPixelFormat DownsampleOverrideFormat;
-		FRDGTextureRef SceneColorTexture;
-		FRDGTextureRef SceneDepthTexture;
-		FRDGTextureRef SceneVelocityTexture;
+		FRDGTextureRef SceneColorTexture = nullptr;
+		FRDGTextureRef SceneDepthTexture = nullptr;
+		FRDGTextureRef SceneVelocityTexture = nullptr;
+		const FSeparateTranslucencyTextures* SeparateTranslucencyTextures = nullptr;
 	};
 
 	virtual ~ITemporalUpscaler() {};
diff --git a/Engine/Source/Runtime/Renderer/Private/SceneRendering.h b/Engine/Source/Runtime/Renderer/Private/SceneRendering.h
--- a/Engine/Source/Runtime/Renderer/Private/SceneRendering.h
+++ b/Engine/Source/Runtime/Renderer/Private/SceneRendering.h
@@ -1006,6 +1006,9 @@ struct FPreviousViewInfo
 	// Bleed free scene color to use for screen space ray tracing.
 	TRefCountPtr<IPooledRenderTarget> ScreenSpaceRayTracingInput;
 
+	// Separate translucency of previous frame to detect translucency change in TAA
+	TRefCountPtr<IPooledRenderTarget> SeparateTranslucency;
+
 	// Temporal AA result of last frame
 	FTemporalAAHistory TemporalAAHistory;
 
