Converting FMeshPassProcessors to use ODSC.

- FTranslucencyDepthPassMeshProcessor

diff --git a/Engine/Source/Runtime/Renderer/Private/TranslucentLighting.cpp b/Engine/Source/Runtime/Renderer/Private/TranslucentLighting.cpp
--- a/Engine/Source/Runtime/Renderer/Private/TranslucentLighting.cpp
+++ b/Engine/Source/Runtime/Renderer/Private/TranslucentLighting.cpp
@@ -350,9 +350,16 @@ public:
 	virtual void AddMeshBatch(const FMeshBatch& RESTRICT MeshBatch, uint64 BatchElementMask, const FPrimitiveSceneProxy* RESTRICT PrimitiveSceneProxy, int32 StaticMeshId = -1) override final;
 
 private:
+	bool TryAddMeshBatch(
+		const FMeshBatch& RESTRICT MeshBatch, 
+		uint64 BatchElementMask, 
+		const FPrimitiveSceneProxy* RESTRICT PrimitiveSceneProxy, 
+		int32 StaticMeshId,
+		const FMaterialRenderProxy& MaterialRenderProxy,
+		const FMaterial& Material);
 
 	template<ETranslucencyShadowDepthShaderMode ShaderMode>
-	void Process(
+	bool Process(
 		const FMeshBatch& MeshBatch,
 		uint64 BatchElementMask,
 		int32 StaticMeshId,
@@ -382,8 +389,40 @@ FTranslucencyDepthPassMeshProcessor::FTranslucencyDepthPassMeshProcessor(const F
 {
 }
 
+bool FTranslucencyDepthPassMeshProcessor::TryAddMeshBatch(
+	const FMeshBatch& RESTRICT MeshBatch, 
+	uint64 BatchElementMask, 
+	const FPrimitiveSceneProxy* RESTRICT PrimitiveSceneProxy, 
+	int32 StaticMeshId,
+	const FMaterialRenderProxy& MaterialRenderProxy,
+	const FMaterial& Material)
+{
+	// Determine the mesh's material and blend mode.
+	const EBlendMode BlendMode = Material.GetBlendMode();
+	const float MaterialTranslucentShadowStartOffset = Material.GetTranslucentShadowStartOffset();
+	const FMeshDrawingPolicyOverrideSettings OverrideSettings = ComputeMeshOverrideSettings(MeshBatch);
+	const ERasterizerFillMode MeshFillMode = ComputeMeshFillMode(MeshBatch, Material, OverrideSettings);
+	const ERasterizerCullMode MeshCullMode = ComputeMeshCullMode(MeshBatch, Material, OverrideSettings);
+	const bool bIsTranslucent = IsTranslucentBlendMode(BlendMode);
+
+	// Only render translucent meshes into the Fourier opacity maps
+	if (bIsTranslucent && ShouldIncludeDomainInMeshPass(Material.GetMaterialDomain()))
+	{
+		if (bDirectionalLight)
+		{
+			return Process<TranslucencyShadowDepth_Standard>(MeshBatch, BatchElementMask, StaticMeshId, PrimitiveSceneProxy, MaterialRenderProxy, Material, MaterialTranslucentShadowStartOffset, MeshFillMode, MeshCullMode);
+		}
+		else
+		{
+			return Process<TranslucencyShadowDepth_PerspectiveCorrect>(MeshBatch, BatchElementMask, StaticMeshId, PrimitiveSceneProxy, MaterialRenderProxy, Material, MaterialTranslucentShadowStartOffset, MeshFillMode, MeshCullMode);
+		}
+	}
+
+	return true;
+}
+
 template<ETranslucencyShadowDepthShaderMode ShaderMode>
-void FTranslucencyDepthPassMeshProcessor::Process(
+bool FTranslucencyDepthPassMeshProcessor::Process(
 	const FMeshBatch& RESTRICT MeshBatch,
 	uint64 BatchElementMask,
 	int32 StaticMeshId,
@@ -400,8 +439,20 @@ void FTranslucencyDepthPassMeshProcessor::Process(
 		TTranslucencyShadowDepthVS<ShaderMode>,
 		TTranslucencyShadowDepthPS<ShaderMode>> PassShaders;
 
-	PassShaders.VertexShader = MaterialResource.GetShader<TTranslucencyShadowDepthVS<ShaderMode> >(VertexFactory->GetType());
-	PassShaders.PixelShader = MaterialResource.GetShader<TTranslucencyShadowDepthPS<ShaderMode> >(VertexFactory->GetType());
+	FMaterialShaderTypes ShaderTypes;
+	ShaderTypes.AddShaderType<TTranslucencyShadowDepthVS<ShaderMode>>();
+	ShaderTypes.AddShaderType<TTranslucencyShadowDepthPS<ShaderMode>>();
+
+	FVertexFactoryType* VertexFactoryType = VertexFactory->GetType();
+
+	FMaterialShaders Shaders;
+	if (!MaterialResource.TryGetShaders(ShaderTypes, VertexFactoryType, Shaders))
+	{
+		return false;
+	}
+
+	Shaders.TryGetVertexShader(PassShaders.VertexShader);
+	Shaders.TryGetPixelShader(PassShaders.PixelShader);
 
 	FMeshPassProcessorRenderState DrawRenderState(PassDrawRenderState);
 
@@ -427,35 +478,27 @@ void FTranslucencyDepthPassMeshProcessor::Process(
 		SortKey,
 		EMeshPassFeatures::Default,
 		ShaderElementData);
+
+	return true;
 }
 
 void FTranslucencyDepthPassMeshProcessor::AddMeshBatch(const FMeshBatch& RESTRICT MeshBatch, uint64 BatchElementMask, const FPrimitiveSceneProxy* RESTRICT PrimitiveSceneProxy, int32 StaticMeshId)
 {
 	if (MeshBatch.CastShadow)
 	{
-		// Determine the mesh's material and blend mode.
-		const FMaterialRenderProxy* FallbackMaterialRenderProxyPtr = nullptr;
-		const FMaterial& Material = MeshBatch.MaterialRenderProxy->GetMaterialWithFallback(FeatureLevel, FallbackMaterialRenderProxyPtr);
-		const FMaterialRenderProxy& MaterialRenderProxy = FallbackMaterialRenderProxyPtr ? *FallbackMaterialRenderProxyPtr : *MeshBatch.MaterialRenderProxy;
-
-		const EBlendMode BlendMode = Material.GetBlendMode();
-		const float MaterialTranslucentShadowStartOffset = Material.GetTranslucentShadowStartOffset();
-		const FMeshDrawingPolicyOverrideSettings OverrideSettings = ComputeMeshOverrideSettings(MeshBatch);
-		const ERasterizerFillMode MeshFillMode = ComputeMeshFillMode(MeshBatch, Material, OverrideSettings);
-		const ERasterizerCullMode MeshCullMode = ComputeMeshCullMode(MeshBatch, Material, OverrideSettings);
-		const bool bIsTranslucent = IsTranslucentBlendMode(BlendMode);
-
-		// Only render translucent meshes into the Fourier opacity maps
-		if (bIsTranslucent && ShouldIncludeDomainInMeshPass(Material.GetMaterialDomain()))
+		const FMaterialRenderProxy* MaterialRenderProxy = MeshBatch.MaterialRenderProxy;
+		while (MaterialRenderProxy)
 		{
-			if (bDirectionalLight)
+			const FMaterial* Material = MaterialRenderProxy->GetMaterialNoFallback(FeatureLevel);
+			if (Material && Material->GetRenderingThreadShaderMap())
 			{
-				Process<TranslucencyShadowDepth_Standard>(MeshBatch, BatchElementMask, StaticMeshId, PrimitiveSceneProxy, MaterialRenderProxy, Material, MaterialTranslucentShadowStartOffset, MeshFillMode, MeshCullMode);
-			}
-			else
-			{
-				Process<TranslucencyShadowDepth_PerspectiveCorrect>(MeshBatch, BatchElementMask, StaticMeshId, PrimitiveSceneProxy, MaterialRenderProxy, Material, MaterialTranslucentShadowStartOffset, MeshFillMode, MeshCullMode);
+				if (TryAddMeshBatch(MeshBatch, BatchElementMask, PrimitiveSceneProxy, StaticMeshId, *MaterialRenderProxy, *Material))
+				{
+					break;
+				}
 			}
+
+			MaterialRenderProxy = MaterialRenderProxy->GetFallback(FeatureLevel);
 		}
 	}
 }
