Add HZB output for hair strands.

diff --git a/Engine/Source/Runtime/Renderer/Private/HairStrands/HairStrandsData.h b/Engine/Source/Runtime/Renderer/Private/HairStrands/HairStrandsData.h
--- a/Engine/Source/Runtime/Renderer/Private/HairStrands/HairStrandsData.h
+++ b/Engine/Source/Runtime/Renderer/Private/HairStrands/HairStrandsData.h
@@ -23,7 +23,9 @@ BEGIN_GLOBAL_SHADER_PARAMETER_STRUCT(FHairStrandsViewUniformParameters, )
 	SHADER_PARAMETER(float, HairDualScatteringRoughnessOverride)						// Override the roughness used for dual scattering (for hack/test purpose only)
 	SHADER_PARAMETER(FIntPoint, HairSampleViewportResolution)							// Maximum viewport resolution of the sample space
 	SHADER_PARAMETER_RDG_TEXTURE(Texture2D<uint4>, HairCategorizationTexture)			// Categorization texture aggregating hair info in screen space (closest depth, coverage, ...)
-	SHADER_PARAMETER_RDG_TEXTURE(Texture2D<float>, HairOnlyDepthTexture)				// Depth texture containing only hair depth
+	SHADER_PARAMETER_RDG_TEXTURE(Texture2D, HairOnlyDepthTexture)						// Depth texture containing only hair depth (not strongly typed for legacy shader code reason)
+	SHADER_PARAMETER_RDG_TEXTURE(Texture2D, HairOnlyDepthClosestHZBTexture)				// HZB closest depth texture containing only hair depth (not strongly typed for legacy shader code reason)
+	SHADER_PARAMETER_RDG_TEXTURE(Texture2D, HairOnlyDepthFurthestHZBTexture)			// HZB furthest depth texture containing only hair depth (not strongly typed for legacy shader code reason)
 	SHADER_PARAMETER_RDG_TEXTURE(Texture2D<uint>, HairSampleOffset)						// Offset & count, for accessing pixel's samples, based on screen pixel position
 	SHADER_PARAMETER_RDG_TEXTURE(Texture2D<uint>, HairSampleCount)						// Total count of hair sample, in sample space
 	SHADER_PARAMETER_RDG_BUFFER_SRV(StructuredBuffer<FPackedHairSample>, HairSampleData)// Sample data (coverage, tangent, base color, ...), in sample space // HAIRSTRANDS_TODO: change this to be a uint4 so that we don't have to include the type for generated contant buffer
@@ -66,6 +68,9 @@ struct FHairStrandsVisibilityData
 	FRDGTextureRef EmissiveTexture = nullptr;
 	FRDGTextureRef HairOnlyDepthTexture = nullptr;
 
+	FRDGTextureRef HairOnlyDepthClosestHZBTexture = nullptr;
+	FRDGTextureRef HairOnlyDepthFurthestHZBTexture = nullptr;
+
 	FRDGTextureRef LightChannelMaskTexture = nullptr;
 
 	uint32			MaxSampleCount = 8;
diff --git a/Engine/Source/Runtime/Renderer/Private/HairStrands/HairStrandsRendering.cpp b/Engine/Source/Runtime/Renderer/Private/HairStrands/HairStrandsRendering.cpp
--- a/Engine/Source/Runtime/Renderer/Private/HairStrands/HairStrandsRendering.cpp
+++ b/Engine/Source/Runtime/Renderer/Private/HairStrands/HairStrandsRendering.cpp
@@ -15,6 +15,8 @@ static TRDGUniformBufferRef<FHairStrandsViewUniformParameters> InternalCreateHai
 	{
 		Parameters->HairCategorizationTexture = In->CategorizationTexture;
 		Parameters->HairOnlyDepthTexture = In->HairOnlyDepthTexture;
+		Parameters->HairOnlyDepthClosestHZBTexture = In->HairOnlyDepthClosestHZBTexture;
+		Parameters->HairOnlyDepthFurthestHZBTexture = In->HairOnlyDepthFurthestHZBTexture;
 		Parameters->HairSampleOffset = In->NodeIndex;
 		Parameters->HairSampleData = GraphBuilder.CreateSRV(In->NodeData);
 		Parameters->HairSampleCoords = GraphBuilder.CreateSRV(In->NodeCoord, FHairStrandsVisibilityData::NodeCoordFormat);
@@ -36,6 +38,13 @@ static TRDGUniformBufferRef<FHairStrandsViewUniformParameters> InternalCreateHai
 			Parameters->HairTileCount = DummyBufferRG16SRV;
 			Parameters->HairTileCountXY = FIntPoint(0,0);
 		}
+
+		if (!Parameters->HairOnlyDepthFurthestHZBTexture)
+		{
+			FRDGTextureRef BlackTexture = GSystemTextures.GetBlackDummy(GraphBuilder);
+			Parameters->HairOnlyDepthFurthestHZBTexture = BlackTexture;
+			Parameters->HairOnlyDepthClosestHZBTexture = BlackTexture;
+		}
 	}
 	else
 	{
@@ -54,6 +63,8 @@ static TRDGUniformBufferRef<FHairStrandsViewUniformParameters> InternalCreateHai
 		FRDGBufferSRVRef DummyBufferRG16SRV = GraphBuilder.CreateSRV(DummyBuffer, PF_R16G16_UINT);
 
 		Parameters->HairOnlyDepthTexture = FarDepth;
+		Parameters->HairOnlyDepthFurthestHZBTexture = BlackTexture;
+		Parameters->HairOnlyDepthClosestHZBTexture = Parameters->HairOnlyDepthFurthestHZBTexture;
 		Parameters->HairCategorizationTexture = ZeroRGBA16_UINT;
 		Parameters->HairSampleCount = ZeroR32_UINT;
 		Parameters->HairSampleOffset = ZeroR32_UINT;
diff --git a/Engine/Source/Runtime/Renderer/Private/HairStrands/HairStrandsVisibility.cpp b/Engine/Source/Runtime/Renderer/Private/HairStrands/HairStrandsVisibility.cpp
--- a/Engine/Source/Runtime/Renderer/Private/HairStrands/HairStrandsVisibility.cpp
+++ b/Engine/Source/Runtime/Renderer/Private/HairStrands/HairStrandsVisibility.cpp
@@ -15,6 +15,7 @@
 #include "PostProcessing.h"
 #include "MeshPassProcessor.inl"
 #include "ScenePrivate.h"
+#include "SceneTextureReductions.h"
 
 DECLARE_GPU_STAT(HairStrandsVisibility);
 
@@ -2929,6 +2930,26 @@ static void AddHairOnlyDepthPass(
 		OutDepthTexture);
 }
 
+static void AddHairOnlyHZBPass(
+	FRDGBuilder& GraphBuilder,
+	const FViewInfo& View,
+	FRDGTextureRef HairDepthTexture,
+	FRDGTextureRef& OutClosestHZBTexture,
+	FRDGTextureRef& OutFurthestHZBTexture)
+{
+	BuildHZB(
+		GraphBuilder,
+		HairDepthTexture,
+		/* VisBufferTexture = */ nullptr,
+		View.ViewRect,
+		View.GetFeatureLevel(),
+		View.GetShaderPlatform(),
+		TEXT("HZBHairClosest"),
+		&OutClosestHZBTexture,
+		TEXT("HZBHairFurthest"),
+		&OutFurthestHZBTexture);
+}
+
 ///////////////////////////////////////////////////////////////////////////////////////////////////
 
 class FHairCountToCoverageCS : public FGlobalShader
@@ -3461,6 +3482,13 @@ void RenderHairStrandsVisibilityBuffer(
 						VisibilityData.TileData,
 						CategorizationTexture,
 						HairOnlyDepthTexture);
+
+					AddHairOnlyHZBPass(
+						GraphBuilder,
+						View,
+						HairOnlyDepthTexture,
+						VisibilityData.HairOnlyDepthClosestHZBTexture,
+						VisibilityData.HairOnlyDepthFurthestHZBTexture);
 				}
 				else
 				{
@@ -3674,6 +3702,13 @@ void RenderHairStrandsVisibilityBuffer(
 						VisibilityData.TileData,
 						CategorizationTexture,
 						HairOnlyDepthTexture);
+
+					AddHairOnlyHZBPass(
+						GraphBuilder,
+						View,
+						HairOnlyDepthTexture,
+						VisibilityData.HairOnlyDepthClosestHZBTexture,
+						VisibilityData.HairOnlyDepthFurthestHZBTexture);
 				}
 				else
 				{
