Remove unnecessary permutation.

diff --git a/Engine/Source/Runtime/Renderer/Private/HairStrands/HairStrandsEnvironment.cpp b/Engine/Source/Runtime/Renderer/Private/HairStrands/HairStrandsEnvironment.cpp
--- a/Engine/Source/Runtime/Renderer/Private/HairStrands/HairStrandsEnvironment.cpp
+++ b/Engine/Source/Runtime/Renderer/Private/HairStrands/HairStrandsEnvironment.cpp
@@ -64,7 +64,8 @@ enum class EHairLightingIntegrationType : uint8
 	SceneColor = 0,
 	AdHoc = 1,
 	Uniform = 2,
-	SH = 3
+	SH = 3,
+	Count
 };
 
 bool GetHairStrandsSkyLightingEnable() { return GHairSkylightingEnable > 0; }
@@ -191,8 +192,8 @@ static void AddHairStrandsEnvironmentAOPass(
 class FHairEnvironmentLighting
 {
 public:
-	class FIntegrationType	: SHADER_PERMUTATION_INT("PERMUTATION_INTEGRATION_TYPE", 4);
-	class FDebug			: SHADER_PERMUTATION_INT("PERMUTATION_DEBUG", 2); 
+	class FIntegrationType	: SHADER_PERMUTATION_INT("PERMUTATION_INTEGRATION_TYPE", uint32(EHairLightingIntegrationType::Count));
+	class FDebug			: SHADER_PERMUTATION_BOOL("PERMUTATION_DEBUG"); 
 	using FPermutationDomain = TShaderPermutationDomain<FIntegrationType, FDebug>;
 
 	static void ModifyCompilationEnvironment(const FGlobalShaderPermutationParameters& Parameters, FShaderCompilerEnvironment& OutEnvironment)
@@ -202,16 +203,21 @@ public:
 
 	static bool ShouldCompilePermutation(const FGlobalShaderPermutationParameters& Parameters)
 	{
-		if (!IsHairStrandsSupported(EHairStrandsShaderType::Strands, Parameters.Platform))
+		// Compile debug permutation only for the uniform integrator
+		FPermutationDomain PermutationVector(Parameters.PermutationId);
+		if (PermutationVector.Get<FIntegrationType>() != int32(EHairLightingIntegrationType::Uniform) && PermutationVector.Get<FDebug>())
 		{
 			return false;
 		}
-
-		return true;
+		return IsHairStrandsSupported(EHairStrandsShaderType::Strands, Parameters.Platform);
 	}
 
 	static FPermutationDomain RemapPermutation(FPermutationDomain PermutationVector)
 	{
+		if (PermutationVector.Get<FIntegrationType>() != int32(EHairLightingIntegrationType::Uniform))
+		{
+			PermutationVector.Set<FDebug>(false);
+		}
 		return PermutationVector;
 	}
 
@@ -293,7 +299,7 @@ class FHairEnvironmentLightingPS : public FGlobalShader
 
 	static bool ShouldCompilePermutation(const FGlobalShaderPermutationParameters& Parameters)
 	{
-		return IsHairStrandsSupported(EHairStrandsShaderType::Strands, Parameters.Platform);
+		return FHairEnvironmentLighting::ShouldCompilePermutation(Parameters);
 	}
 };
 
@@ -386,7 +392,7 @@ static void AddHairStrandsEnvironmentLightingPassPS(
 
 	FHairEnvironmentLightingPS::FPermutationDomain PermutationVector;
 	PermutationVector.Set<FHairEnvironmentLighting::FIntegrationType>(uint32(IntegrationType));
-	PermutationVector.Set<FHairEnvironmentLighting::FDebug>(DebugData ? 1 : 0);
+	PermutationVector.Set<FHairEnvironmentLighting::FDebug>(DebugData != nullptr);
 	PermutationVector = FHairEnvironmentLighting::RemapPermutation(PermutationVector);
 
 	FIntPoint ViewportResolution = VisibilityData.SampleLightingViewportResolution;
