Add permutation for sampling Lumen radiance cache.

diff --git a/Engine/Shaders/Private/HairStrands/HairStrandsEnvironmentLighting.usf b/Engine/Shaders/Private/HairStrands/HairStrandsEnvironmentLighting.usf
--- a/Engine/Shaders/Private/HairStrands/HairStrandsEnvironmentLighting.usf
+++ b/Engine/Shaders/Private/HairStrands/HairStrandsEnvironmentLighting.usf
@@ -18,6 +18,10 @@
 #include "../LightGridCommon.ush"
 #include "../SceneTextureParameters.ush"
 
+#if PERMUTATION_LUMEN
+#include "../Lumen/LumenRadianceCacheCommon.ush"
+#endif
+
 // Disable static lighting, as hair don't have any diffused lighting baked, 
 // and rely solely on specular information (i.e. reflection capture)
 #undef ALLOW_STATIC_LIGHTING
@@ -63,6 +67,12 @@ float GetPreExposure()
 
 float3 GatherRadiance(float3 RayDirection, float Roughness, uint2 PixelPosition, float PixelSceneDepth, float3 AbsoluteWorldPosition)
 {
+#if PERMUTATION_LUMEN
+	// Lumen probe
+	float3 Random = 0;
+	float ConeHalfAngle = 0;
+	return SampleRadianceCacheInterpolated(AbsoluteWorldPosition, RayDirection, ConeHalfAngle, Random);
+#else
 	const float IndirectIrradiance = 1;
 	const int SingleCaptureIndex = 0;
 
@@ -87,6 +97,7 @@ float3 GatherRadiance(float3 RayDirection, float Roughness, uint2 PixelPosition,
 		SingleCaptureIndex,
 		bCompositeSkylight,
 		EyeIndex);
+#endif // PERMUTATION_LUMEN
 }
 
 Texture3D<float4>			HairEnergyLUTTexture;
diff --git a/Engine/Source/Runtime/Renderer/Private/HairStrands/HairStrandsEnvironment.cpp b/Engine/Source/Runtime/Renderer/Private/HairStrands/HairStrandsEnvironment.cpp
--- a/Engine/Source/Runtime/Renderer/Private/HairStrands/HairStrandsEnvironment.cpp
+++ b/Engine/Source/Runtime/Renderer/Private/HairStrands/HairStrandsEnvironment.cpp
@@ -19,6 +19,8 @@
 #include "RenderGraphEvent.h"
 #include "PostProcess/PostProcessing.h"
 #include "GpuDebugRendering.h"
+#include "Lumen/LumenRadianceCache.h"
+#include "Lumen/LumenScreenProbeGather.h"
 
 ///////////////////////////////////////////////////////////////////////////////////////////////////
 
@@ -59,6 +61,14 @@ static FAutoConsoleVariableRef CVarHairStrandsSkyLighting_DebugSample(TEXT("r.Ha
 
 ///////////////////////////////////////////////////////////////////////////////////////////////////
 
+enum class EHairLightingSourceType : uint8
+{
+	SceneColor = 0,
+	ReflectionProbe = 1,
+	Lumen = 2,
+	Count
+};
+
 enum class EHairLightingIntegrationType : uint8
 {
 	SceneColor = 0,
@@ -194,7 +204,8 @@ class FHairEnvironmentLighting
 public:
 	class FIntegrationType	: SHADER_PERMUTATION_INT("PERMUTATION_INTEGRATION_TYPE", uint32(EHairLightingIntegrationType::Count));
 	class FDebug			: SHADER_PERMUTATION_BOOL("PERMUTATION_DEBUG"); 
-	using FPermutationDomain = TShaderPermutationDomain<FIntegrationType, FDebug>;
+	class FLumen			: SHADER_PERMUTATION_BOOL("PERMUTATION_LUMEN"); 
+	using FPermutationDomain = TShaderPermutationDomain<FIntegrationType, FLumen, FDebug>;
 
 	static void ModifyCompilationEnvironment(const FGlobalShaderPermutationParameters& Parameters, FShaderCompilerEnvironment& OutEnvironment)
 	{
@@ -286,6 +297,7 @@ class FHairEnvironmentLightingPS : public FGlobalShader
 
 	BEGIN_SHADER_PARAMETER_STRUCT(FParameters, )
 		SHADER_PARAMETER_STRUCT_INCLUDE(FHairEnvironmentLighting::FParameters, Common)
+		SHADER_PARAMETER_STRUCT_INCLUDE(LumenRadianceCache::FRadianceCacheInterpolationParameters, RadianceCache)
 		SHADER_PARAMETER_STRUCT_INCLUDE(ShaderDrawDebug::FShaderDrawDebugParameters, ShaderDrawParameters)
 		SHADER_PARAMETER_STRUCT_INCLUDE(FHairStrandsDebugData::FWriteParameters, DebugData)
 		RENDER_TARGET_BINDING_SLOTS()
@@ -313,6 +325,7 @@ static void AddHairStrandsEnvironmentLightingPassPS(
 	const FHairStrandsVisibilityData& VisibilityData,
 	const FHairStrandsVoxelResources& VirtualVoxelResources,
 	const FRDGTextureRef SceneColorTexture,
+	const EHairLightingSourceType LightingType,
 	FHairStrandsDebugData::Data* DebugData)
 {
 	FSceneTextureParameters SceneTextures = GetSceneTextureParameters(GraphBuilder);
@@ -328,9 +341,9 @@ static void AddHairStrandsEnvironmentLightingPassPS(
 	PassParameters->HairEnergyLUTTexture = GetHairLUT(GraphBuilder, View, HairLUTType_MeanEnergy);
 
 	EHairLightingIntegrationType IntegrationType = EHairLightingIntegrationType::AdHoc;
-	const bool bUseSceneColor = SceneColorTexture != nullptr;
-	if (bUseSceneColor)
+	if (LightingType == EHairLightingSourceType::SceneColor)
 	{
+		check(SceneColorTexture != nullptr);
 		IntegrationType = EHairLightingIntegrationType::SceneColor;
 		PassParameters->SceneColorTexture = SceneColorTexture;
 	}
@@ -377,6 +390,13 @@ static void AddHairStrandsEnvironmentLightingPassPS(
 		SetupReflectionUniformParameters(View, ReflectionUniformParameters);
 		PassParameters->ReflectionsParameters = CreateUniformBufferImmediate(ReflectionUniformParameters, UniformBuffer_SingleDraw);
 	}
+
+	if (LightingType == EHairLightingSourceType::Lumen)
+	{
+		const FRadianceCacheState& RadianceCacheState = View.ViewState->RadianceCacheState;
+		const LumenRadianceCache::FRadianceCacheInputs RadianceCacheInputs = LumenScreenProbeGatherRadianceCache::SetupRadianceCacheInputs();
+		LumenRadianceCache::GetInterpolationParameters(View, GraphBuilder, RadianceCacheState, RadianceCacheInputs, ParametersPS->RadianceCache);
+	}
 	PassParameters->ForwardLightData = View.ForwardLightingResources->ForwardLightDataUniformBuffer;
 	PassParameters->OutLightingBuffer = nullptr;
 
@@ -392,6 +412,7 @@ static void AddHairStrandsEnvironmentLightingPassPS(
 
 	FHairEnvironmentLightingPS::FPermutationDomain PermutationVector;
 	PermutationVector.Set<FHairEnvironmentLighting::FIntegrationType>(uint32(IntegrationType));
+	PermutationVector.Set<FHairEnvironmentLighting::FLumen>(LightingType == EHairLightingSourceType::Lumen);
 	PermutationVector.Set<FHairEnvironmentLighting::FDebug>(DebugData != nullptr);
 	PermutationVector = FHairEnvironmentLighting::RemapPermutation(PermutationVector);
 
@@ -404,7 +425,7 @@ static void AddHairStrandsEnvironmentLightingPassPS(
 	ParametersPS->RenderTargets[0] = FRenderTargetBinding(VisibilityData.SampleLightingBuffer, ERenderTargetLoadAction::ELoad);
 
 	GraphBuilder.AddPass(
-		bUseSceneColor ? RDG_EVENT_NAME("HairEnvSceneScatterPS") : RDG_EVENT_NAME("HairEnvLightingPS"),
+		RDG_EVENT_NAME("HairEnvLightingPS(%s)", LightingType == EHairLightingSourceType::SceneColor ? TEXT("SceneScatter") : (LightingType == EHairLightingSourceType::Lumen ? TEXT("Lumen") : TEXT("ReflectionProbe")) ),
 		ParametersPS,
 		ERDGPassFlags::Raster,
 		[ParametersPS, VertexShader, PixelShader, ViewportResolution, CapturedView](FRHICommandList& RHICmdList)
@@ -466,30 +487,20 @@ void RenderHairStrandsSceneColorScattering(
 
 		if (bNeedScatterSceneLighting)
 		{
-			AddHairStrandsEnvironmentLightingPassPS(GraphBuilder, Scene, View, VisibilityData, VoxelResources, SceneColorTexture, nullptr);
+			AddHairStrandsEnvironmentLightingPassPS(GraphBuilder, Scene, View, VisibilityData, VoxelResources, SceneColorTexture, EHairLightingSourceType::SceneColor, nullptr);
 		}
 	}
 }
 
-void RenderHairStrandsEnvironmentLighting(
+static void InternalRenderHairStrandsEnvironmentLighting(
 	FRDGBuilder& GraphBuilder,
 	const FScene* Scene,
-	FViewInfo& View)
+	FViewInfo& View,
+	EHairLightingSourceType LightingType)
 {
 	if (!GetHairStrandsSkyLightingEnable() || !HairStrands::HasViewHairStrandsData(View))
 		return;
 
-	struct FHairEnvDebugData
-	{
-		FRDGBufferRef ShadingPointBuffer = nullptr;
-		FRDGBufferRef SampleBuffer = nullptr;
-		FRDGBufferRef ShadingPointCounter = nullptr;
-
-		FRDGBufferUAVRef ShadingPointBufferUAV = nullptr;
-		FRDGBufferUAVRef SampleBufferUAV = nullptr;
-		FRDGBufferUAVRef ShadingPointCounterUAV = nullptr;
-	};
-
 	const FHairStrandsVisibilityData& VisibilityData = View.HairStrandsViewData.VisibilityData;
 	const FHairStrandsVoxelResources& VoxelResources = View.HairStrandsViewData.VirtualVoxelResources;
 	const bool bRenderHairLighting = VoxelResources.IsValid() && VisibilityData.CategorizationTexture;
@@ -504,7 +515,23 @@ void RenderHairStrandsEnvironmentLighting(
 		View.HairStrandsViewData.DebugData.Resources = FHairStrandsDebugData::CreateData(GraphBuilder);
 	}
 
-	AddHairStrandsEnvironmentLightingPassPS(GraphBuilder, Scene, View, VisibilityData, VoxelResources, nullptr, bDebugSamplingEnable ? &View.HairStrandsViewData.DebugData.Resources : nullptr);
+	AddHairStrandsEnvironmentLightingPassPS(GraphBuilder, Scene, View, VisibilityData, VoxelResources, nullptr, LightingType, bDebugSamplingEnable ? &View.HairStrandsViewData.DebugData.Resources : nullptr);
+}
+
+void RenderHairStrandsLumenLighting(
+	FRDGBuilder& GraphBuilder,
+	const FScene* Scene,
+	FViewInfo& View)
+{
+	InternalRenderHairStrandsEnvironmentLighting(GraphBuilder, Scene, View, EHairLightingSourceType::Lumen);
+}
+
+void RenderHairStrandsEnvironmentLighting(
+	FRDGBuilder& GraphBuilder,
+	const FScene* Scene,
+	FViewInfo& View)
+{
+	InternalRenderHairStrandsEnvironmentLighting(GraphBuilder, Scene, View, EHairLightingSourceType::ReflectionProbe);
 }
 
 void RenderHairStrandsAmbientOcclusion(
diff --git a/Engine/Source/Runtime/Renderer/Private/HairStrands/HairStrandsEnvironment.h b/Engine/Source/Runtime/Renderer/Private/HairStrands/HairStrandsEnvironment.h
--- a/Engine/Source/Runtime/Renderer/Private/HairStrands/HairStrandsEnvironment.h
+++ b/Engine/Source/Runtime/Renderer/Private/HairStrands/HairStrandsEnvironment.h
@@ -18,6 +18,11 @@ void RenderHairStrandsAmbientOcclusion(
 	const FViewInfo& View,
 	const FRDGTextureRef& InAOTexture);
 
+void RenderHairStrandsLumenLighting(
+	FRDGBuilder& GraphBuilder,
+	const FScene* Scene,
+	FViewInfo& View);
+
 void RenderHairStrandsEnvironmentLighting(
 	FRDGBuilder& GraphBuilder,
 	const FScene* Scene,
