Change hair strands debug data to be created during the base pass, so that View can be passed as const within the hair env lighting passes.

diff --git a/Engine/Source/Runtime/Renderer/Private/HairStrands/HairStrandsEnvironment.cpp b/Engine/Source/Runtime/Renderer/Private/HairStrands/HairStrandsEnvironment.cpp
--- a/Engine/Source/Runtime/Renderer/Private/HairStrands/HairStrandsEnvironment.cpp
+++ b/Engine/Source/Runtime/Renderer/Private/HairStrands/HairStrandsEnvironment.cpp
@@ -78,10 +78,10 @@ enum class EHairLightingIntegrationType : uint8
 	Count
 };
 
-bool GetHairStrandsSkyLightingEnable() { return GHairSkylightingEnable > 0; }
-static bool GetHairStrandsSkyAOEnable() { return GHairSkyAOEnable > 0; }
-static float GetHairStrandsSkyLightingConeAngle() { return FMath::Max(0.f, GHairSkylightingConeAngle); }
-
+bool GetHairStrandsSkyLightingEnable()				{ return GHairSkylightingEnable > 0; }
+bool GetHairStrandsSkyLightingDebugEnable()			{ return GHairSkylightingEnable > 0 && GHairStrandsSkyLighting_DebugSample > 0; }
+static bool GetHairStrandsSkyAOEnable()				{ return GHairSkyAOEnable > 0; }
+static float GetHairStrandsSkyLightingConeAngle()	{ return FMath::Max(0.f, GHairSkylightingConeAngle); }
 DECLARE_GPU_STAT_NAMED(HairStrandsReflectionEnvironment, TEXT("Hair Strands Reflection Environment"));
 
 ///////////////////////////////////////////////////////////////////////////////////////////////////
@@ -326,7 +326,7 @@ static void AddHairStrandsEnvironmentLightingPassPS(
 	const FHairStrandsVoxelResources& VirtualVoxelResources,
 	const FRDGTextureRef SceneColorTexture,
 	const EHairLightingSourceType LightingType,
-	FHairStrandsDebugData::Data* DebugData)
+	const FHairStrandsDebugData::Data* DebugData)
 {
 	FSceneTextureParameters SceneTextures = GetSceneTextureParameters(GraphBuilder);
 
@@ -495,7 +495,7 @@ void RenderHairStrandsSceneColorScattering(
 static void InternalRenderHairStrandsEnvironmentLighting(
 	FRDGBuilder& GraphBuilder,
 	const FScene* Scene,
-	FViewInfo& View,
+	const FViewInfo& View,
 	EHairLightingSourceType LightingType)
 {
 	if (!GetHairStrandsSkyLightingEnable() || !HairStrands::HasViewHairStrandsData(View))
@@ -508,20 +508,14 @@ static void InternalRenderHairStrandsEnvironmentLighting(
 	{
 		return;
 	}
-
-	const bool bDebugSamplingEnable = GHairStrandsSkyLighting_DebugSample > 0;
-	if (bDebugSamplingEnable)
-	{
-		View.HairStrandsViewData.DebugData.Resources = FHairStrandsDebugData::CreateData(GraphBuilder);
-	}
-
-	AddHairStrandsEnvironmentLightingPassPS(GraphBuilder, Scene, View, VisibilityData, VoxelResources, nullptr, LightingType, bDebugSamplingEnable ? &View.HairStrandsViewData.DebugData.Resources : nullptr);
+	
+	AddHairStrandsEnvironmentLightingPassPS(GraphBuilder, Scene, View, VisibilityData, VoxelResources, nullptr, LightingType, View.HairStrandsViewData.DebugData.IsPlotDataValid() ? &View.HairStrandsViewData.DebugData.Resources : nullptr);
 }
 
 void RenderHairStrandsLumenLighting(
 	FRDGBuilder& GraphBuilder,
 	const FScene* Scene,
-	FViewInfo& View)
+	const FViewInfo& View)
 {
 	InternalRenderHairStrandsEnvironmentLighting(GraphBuilder, Scene, View, EHairLightingSourceType::Lumen);
 }
@@ -529,7 +523,7 @@ void RenderHairStrandsLumenLighting(
 void RenderHairStrandsEnvironmentLighting(
 	FRDGBuilder& GraphBuilder,
 	const FScene* Scene,
-	FViewInfo& View)
+	const FViewInfo& View)
 {
 	InternalRenderHairStrandsEnvironmentLighting(GraphBuilder, Scene, View, EHairLightingSourceType::ReflectionProbe);
 }
diff --git a/Engine/Source/Runtime/Renderer/Private/HairStrands/HairStrandsEnvironment.h b/Engine/Source/Runtime/Renderer/Private/HairStrands/HairStrandsEnvironment.h
--- a/Engine/Source/Runtime/Renderer/Private/HairStrands/HairStrandsEnvironment.h
+++ b/Engine/Source/Runtime/Renderer/Private/HairStrands/HairStrandsEnvironment.h
@@ -21,12 +21,12 @@ void RenderHairStrandsAmbientOcclusion(
 void RenderHairStrandsLumenLighting(
 	FRDGBuilder& GraphBuilder,
 	const FScene* Scene,
-	FViewInfo& View);
+	const FViewInfo& View);
 
 void RenderHairStrandsEnvironmentLighting(
 	FRDGBuilder& GraphBuilder,
 	const FScene* Scene,
-	FViewInfo& View);
+	const FViewInfo& View);
 
 void RenderHairStrandsSceneColorScattering(
 	FRDGBuilder& GraphBuilder,
diff --git a/Engine/Source/Runtime/Renderer/Private/HairStrands/HairStrandsRendering.cpp b/Engine/Source/Runtime/Renderer/Private/HairStrands/HairStrandsRendering.cpp
--- a/Engine/Source/Runtime/Renderer/Private/HairStrands/HairStrandsRendering.cpp
+++ b/Engine/Source/Runtime/Renderer/Private/HairStrands/HairStrandsRendering.cpp
@@ -83,6 +83,7 @@ static TRDGUniformBufferRef<FHairStrandsViewUniformParameters> InternalCreateHai
 IMPLEMENT_GLOBAL_SHADER_PARAMETER_STRUCT(FHairStrandsViewUniformParameters, "HairStrands");
 
 void AddServiceLocalQueuePass(FRDGBuilder& GraphBuilder);
+bool GetHairStrandsSkyLightingDebugEnable();
 
 void RenderHairPrePass(
 	FRDGBuilder& GraphBuilder,
@@ -136,7 +137,12 @@ void RenderHairBasePass(
 				SceneTextures.Depth.Resolve,
 				SceneTextures.Velocity,
 				InstanceCullingManager);
-			
+
+			const bool bDebugSamplingEnable = GetHairStrandsSkyLightingDebugEnable();
+			if (bDebugSamplingEnable)
+			{
+				View.HairStrandsViewData.DebugData.Resources = FHairStrandsDebugData::CreateData(GraphBuilder);
+			}
 		}
 		
 		if (View.HairStrandsViewData.VisibilityData.CategorizationTexture)
