Fix CIS issue.

Add missing files.

diff --git a/Engine/Source/Runtime/Renderer/Private/HairStrands/HairStrandsData.h b/Engine/Source/Runtime/Renderer/Private/HairStrands/HairStrandsData.h
--- a/Engine/Source/Runtime/Renderer/Private/HairStrands/HairStrandsData.h
+++ b/Engine/Source/Runtime/Renderer/Private/HairStrands/HairStrandsData.h
@@ -329,7 +329,7 @@ struct FHairStrandsDebugData
 	}
 
 	static Data CreateData(FRDGBuilder& GraphBuilder);
-	static void SetParameters(FRDGBuilder& GraphBuilder, Data& In, FWriteParameters& Out);
+	static void SetParameters(FRDGBuilder& GraphBuilder, const Data& In, FWriteParameters& Out);
 	static void SetParameters(FRDGBuilder& GraphBuilder, const Data& In, FReadParameters& Out);
 
 	// PPLL debug data
diff --git a/Engine/Source/Runtime/Renderer/Private/HairStrands/HairStrandsDebug.cpp b/Engine/Source/Runtime/Renderer/Private/HairStrands/HairStrandsDebug.cpp
--- a/Engine/Source/Runtime/Renderer/Private/HairStrands/HairStrandsDebug.cpp
+++ b/Engine/Source/Runtime/Renderer/Private/HairStrands/HairStrandsDebug.cpp
@@ -119,7 +119,7 @@ FHairStrandsDebugData::Data FHairStrandsDebugData::CreateData(FRDGBuilder& Graph
 	return Out;
 }
 
-void FHairStrandsDebugData::SetParameters(FRDGBuilder& GraphBuilder, FHairStrandsDebugData::Data& In, FHairStrandsDebugData::FWriteParameters& Out)
+void FHairStrandsDebugData::SetParameters(FRDGBuilder& GraphBuilder, const FHairStrandsDebugData::Data& In, FHairStrandsDebugData::FWriteParameters& Out)
 {
 	Out.Debug_MaxSampleCount = FHairStrandsDebugData::MaxSampleCount;
 	Out.Debug_MaxShadingPointCount = FHairStrandsDebugData::MaxShadingPointCount;
