Fixed crash due pushing ray tracing mesh commands from a different thread. FRayTracingMeshCommandOneFrameArray now does not use SceneRenderingAllocator, which is a stack allocator

diff --git a/Engine/Source/Runtime/Renderer/Public/MeshPassProcessor.h b/Engine/Source/Runtime/Renderer/Public/MeshPassProcessor.h
--- a/Engine/Source/Runtime/Renderer/Public/MeshPassProcessor.h
+++ b/Engine/Source/Runtime/Renderer/Public/MeshPassProcessor.h
@@ -1903,7 +1903,7 @@ struct TUseBitwiseSwap<FVisibleRayTracingMeshCommand>
 	enum { Value = false };
 };
 
-typedef TArray<FVisibleRayTracingMeshCommand, SceneRenderingAllocator> FRayTracingMeshCommandOneFrameArray;
+typedef TArray<FVisibleRayTracingMeshCommand> FRayTracingMeshCommandOneFrameArray;
 
 class FRayTracingMeshCommandContext
 {
