Converting FMeshPassProcessors to use ODSC.

- FDistortionMeshProcessor

diff --git a/Engine/Source/Runtime/Renderer/Private/DistortionRendering.cpp b/Engine/Source/Runtime/Renderer/Private/DistortionRendering.cpp
--- a/Engine/Source/Runtime/Renderer/Private/DistortionRendering.cpp
+++ b/Engine/Source/Runtime/Renderer/Private/DistortionRendering.cpp
@@ -388,40 +388,73 @@ void FDistortionMeshProcessor::AddMeshBatch(const FMeshBatch& RESTRICT MeshBatch
 {
 	if (MeshBatch.bUseForMaterial)
 	{
-		// Determine the mesh's material and blend mode.
-		const FMaterialRenderProxy* FallbackMaterialRenderProxyPtr = nullptr;
-		const FMaterial& Material = MeshBatch.MaterialRenderProxy->GetMaterialWithFallback(FeatureLevel, FallbackMaterialRenderProxyPtr);
-
-		const FMaterialRenderProxy& MaterialRenderProxy = FallbackMaterialRenderProxyPtr ? *FallbackMaterialRenderProxyPtr : *MeshBatch.MaterialRenderProxy;
-
-		const EBlendMode BlendMode = Material.GetBlendMode();
-		const FMeshDrawingPolicyOverrideSettings OverrideSettings = ComputeMeshOverrideSettings(MeshBatch);
-		const ERasterizerFillMode MeshFillMode = ComputeMeshFillMode(MeshBatch, Material, OverrideSettings);
-		const ERasterizerCullMode MeshCullMode = ComputeMeshCullMode(MeshBatch, Material, OverrideSettings);
-		const bool bIsTranslucent = IsTranslucentBlendMode(BlendMode);
-
-		if (bIsTranslucent
-			&& (!PrimitiveSceneProxy || PrimitiveSceneProxy->ShouldRenderInMainPass())
-			&& ShouldIncludeDomainInMeshPass(Material.GetMaterialDomain())
-			&& Material.IsDistorted())
+		const FMaterialRenderProxy* MaterialRenderProxy = MeshBatch.MaterialRenderProxy;
+		while (MaterialRenderProxy)
 		{
-			Process(MeshBatch, BatchElementMask, PrimitiveSceneProxy, StaticMeshId, MaterialRenderProxy, Material, MeshFillMode, MeshCullMode);
+			const FMaterial* Material = MaterialRenderProxy->GetMaterialNoFallback(FeatureLevel);
+			if (Material)
+			{
+				if (TryAddMeshBatch(MeshBatch, BatchElementMask, PrimitiveSceneProxy, StaticMeshId, *MaterialRenderProxy, *Material))
+				{
+					break;
+				}
+			}
+
+			MaterialRenderProxy = MaterialRenderProxy->GetFallback(FeatureLevel);
 		}
 	}
 }
 
-void GetDistortionPassShaders(
+bool FDistortionMeshProcessor::TryAddMeshBatch(
+	const FMeshBatch& RESTRICT MeshBatch,
+	uint64 BatchElementMask,
+	const FPrimitiveSceneProxy* RESTRICT PrimitiveSceneProxy,
+	int32 StaticMeshId,
+	const FMaterialRenderProxy& MaterialRenderProxy,
+	const FMaterial& Material)
+{
+	// Determine the mesh's material and blend mode.
+	const EBlendMode BlendMode = Material.GetBlendMode();
+	const FMeshDrawingPolicyOverrideSettings OverrideSettings = ComputeMeshOverrideSettings(MeshBatch);
+	const ERasterizerFillMode MeshFillMode = ComputeMeshFillMode(MeshBatch, Material, OverrideSettings);
+	const ERasterizerCullMode MeshCullMode = ComputeMeshCullMode(MeshBatch, Material, OverrideSettings);
+	const bool bIsTranslucent = IsTranslucentBlendMode(BlendMode);
+
+	bool bResult = true;
+	if (bIsTranslucent
+		&& (!PrimitiveSceneProxy || PrimitiveSceneProxy->ShouldRenderInMainPass())
+		&& ShouldIncludeDomainInMeshPass(Material.GetMaterialDomain())
+		&& Material.IsDistorted())
+	{
+		bResult = Process(MeshBatch, BatchElementMask, PrimitiveSceneProxy, StaticMeshId, MaterialRenderProxy, Material, MeshFillMode, MeshCullMode);
+	}
+
+	return bResult;
+}
+
+bool GetDistortionPassShaders(
 	const FMaterial& Material,
 	FVertexFactoryType* VertexFactoryType,
 	ERHIFeatureLevel::Type FeatureLevel,
 	TShaderRef<FDistortionMeshVS>& VertexShader,
 	TShaderRef<FDistortionMeshPS>& PixelShader)
 {
-	VertexShader = Material.GetShader<FDistortionMeshVS>(VertexFactoryType);
-	PixelShader = Material.GetShader<FDistortionMeshPS>(VertexFactoryType);
+	FMaterialShaderTypes ShaderTypes;
+	ShaderTypes.AddShaderType<FDistortionMeshVS>();
+	ShaderTypes.AddShaderType<FDistortionMeshPS>();
+
+	FMaterialShaders Shaders;
+	if (!Material.TryGetShaders(ShaderTypes, VertexFactoryType, Shaders))
+	{
+		return false;
+	}
+
+	Shaders.TryGetVertexShader(VertexShader);
+	Shaders.TryGetPixelShader(PixelShader);
+	return true;
 }
 
-void FDistortionMeshProcessor::Process(
+bool FDistortionMeshProcessor::Process(
 	const FMeshBatch& MeshBatch, 
 	uint64 BatchElementMask, 
 	const FPrimitiveSceneProxy* RESTRICT PrimitiveSceneProxy,
@@ -437,13 +470,15 @@ void FDistortionMeshProcessor::Process(
 		FDistortionMeshVS,
 		FDistortionMeshPS> DistortionPassShaders;
 
-	GetDistortionPassShaders(
+	if (!GetDistortionPassShaders(
 		MaterialResource,
 		VertexFactory->GetType(),
 		FeatureLevel,
 		DistortionPassShaders.VertexShader,
-		DistortionPassShaders.PixelShader
-		);
+		DistortionPassShaders.PixelShader))
+	{
+		return false;
+	}
 
 	FMeshMaterialShaderElementData ShaderElementData;
 	ShaderElementData.InitializeMeshMaterialData(ViewIfDynamicMeshCommand, PrimitiveSceneProxy, MeshBatch, StaticMeshId, false);
@@ -463,6 +498,8 @@ void FDistortionMeshProcessor::Process(
 		SortKey,
 		EMeshPassFeatures::Default,
 		ShaderElementData);
+
+	return true;
 }
 
 FDistortionMeshProcessor::FDistortionMeshProcessor(const FScene* Scene, const FSceneView* InViewIfDynamicMeshCommand, const FMeshPassProcessorRenderState& InPassDrawRenderState, FMeshPassDrawListContext* InDrawListContext)
diff --git a/Engine/Source/Runtime/Renderer/Private/DistortionRendering.h b/Engine/Source/Runtime/Renderer/Private/DistortionRendering.h
--- a/Engine/Source/Runtime/Renderer/Private/DistortionRendering.h
+++ b/Engine/Source/Runtime/Renderer/Private/DistortionRendering.h
@@ -84,7 +84,15 @@ public:
 	FMeshPassProcessorRenderState PassDrawRenderState;
 
 private:
-	void Process(
+	bool TryAddMeshBatch(
+		const FMeshBatch& RESTRICT MeshBatch,
+		uint64 BatchElementMask,
+		const FPrimitiveSceneProxy* RESTRICT PrimitiveSceneProxy,
+		int32 StaticMeshId,
+		const FMaterialRenderProxy& MaterialRenderProxy,
+		const FMaterial& Material);
+
+	bool Process(
 		const FMeshBatch& MeshBatch,
 		uint64 BatchElementMask,
 		const FPrimitiveSceneProxy* RESTRICT PrimitiveSceneProxy,
