Add hair tile indirect buffer for dispatching one thread per tile.

diff --git a/Engine/Shaders/Private/HairStrands/HairStrandsVisibilityTile.usf b/Engine/Shaders/Private/HairStrands/HairStrandsVisibilityTile.usf
--- a/Engine/Shaders/Private/HairStrands/HairStrandsVisibilityTile.usf
+++ b/Engine/Shaders/Private/HairStrands/HairStrandsVisibilityTile.usf
@@ -70,10 +70,12 @@ void TileMainCS(uint2 DispatchThreadId : SV_DispatchThreadID, uint LinearIndex :
 
 #if SHADER_TILE_COPY_ARGS
 uint2 TileCountXY;
+uint  TilePerThread_GroupSize;
 uint bRectPrimitive;
 Buffer<uint> TileCountBuffer;
 RWBuffer<uint> TileIndirectDrawBuffer;
 RWBuffer<uint> TileIndirectDispatchBuffer;
+RWBuffer<uint> TilePerThreadIndirectDispatchBuffer;
 
 [numthreads(1, 1, 1)]
 void MainCS(uint2 DispatchThreadId : SV_DispatchThreadID)
@@ -91,6 +93,13 @@ void MainCS(uint2 DispatchThreadId : SV_DispatchThreadID)
 		TileIndirectDispatchBuffer[0] = TileCount % TileCountXY.x;
 		TileIndirectDispatchBuffer[1] = (TileCount+TileCountXY.x-1) / TileCountXY.x;
 		TileIndirectDispatchBuffer[2] = 1;
+
+		// Indirect dispatch with one thread
+		// At 4k, with 8x8 tiles: 4096 x 4096 / 8 x 8 = 512 x 512 tiles 
+		// With a group size of 64 (8x8), the total dispatch count should be 512 x 512 / 8 x 8 = 64 x 64 = 4096, which is below the 65k limit per dimension
+		TilePerThreadIndirectDispatchBuffer[0] = (TileCount + TilePerThread_GroupSize -1) / TilePerThread_GroupSize;
+		TilePerThreadIndirectDispatchBuffer[1] = 1;
+		TilePerThreadIndirectDispatchBuffer[2] = 1;
 	}
 }
 #endif // SHADER_TILE_COPY_ARGS
diff --git a/Engine/Source/Runtime/Renderer/Private/HairStrands/HairStrandsData.h b/Engine/Source/Runtime/Renderer/Private/HairStrands/HairStrandsData.h
--- a/Engine/Source/Runtime/Renderer/Private/HairStrands/HairStrandsData.h
+++ b/Engine/Source/Runtime/Renderer/Private/HairStrands/HairStrandsData.h
@@ -42,6 +42,7 @@ struct FHairStrandsTiles
 	FIntPoint			Resolution = FIntPoint(0, 0);
 	static const uint32 GroupSize = 64;
 	static const uint32	TileSize = 8;
+	static const uint32	TilePerThread_GroupSize = 64;
 	uint32				TileCount = 0;
 	FIntPoint			TileCountXY = FIntPoint(0, 0);
 	bool				bRectPrimitive = false;
@@ -51,6 +52,7 @@ struct FHairStrandsTiles
 	FRDGBufferRef		TileCountBuffer = nullptr;
 	FRDGBufferRef		TileIndirectDrawBuffer = nullptr;
 	FRDGBufferRef		TileIndirectDispatchBuffer = nullptr;
+	FRDGBufferRef		TilePerThreadIndirectDispatchBuffer = nullptr;
 
 	bool IsValid() const { return TileCount > 0 && TileDataBuffer != nullptr; }
 };
diff --git a/Engine/Source/Runtime/Renderer/Private/HairStrands/HairStrandsTile.cpp b/Engine/Source/Runtime/Renderer/Private/HairStrands/HairStrandsTile.cpp
--- a/Engine/Source/Runtime/Renderer/Private/HairStrands/HairStrandsTile.cpp
+++ b/Engine/Source/Runtime/Renderer/Private/HairStrands/HairStrandsTile.cpp
@@ -39,10 +39,12 @@ class FHairStrandsTileCopyArgsPassCS : public FGlobalShader
 	BEGIN_SHADER_PARAMETER_STRUCT(FParameters, )
 		SHADER_PARAMETER_STRUCT_REF(FViewUniformShaderParameters, ViewUniformBuffer)
 		SHADER_PARAMETER(FIntPoint, TileCountXY)
+		SHADER_PARAMETER(uint32, TilePerThread_GroupSize)
 		SHADER_PARAMETER(uint32, bRectPrimitive)
 		SHADER_PARAMETER_RDG_BUFFER_SRV(Buffer, TileCountBuffer)
 		SHADER_PARAMETER_RDG_BUFFER_UAV(RWBuffer, TileIndirectDrawBuffer)
 		SHADER_PARAMETER_RDG_BUFFER_UAV(RWBuffer, TileIndirectDispatchBuffer)
+		SHADER_PARAMETER_RDG_BUFFER_UAV(RWBuffer, TilePerThreadIndirectDispatchBuffer)
 	END_SHADER_PARAMETER_STRUCT()
 
 	static bool ShouldCompilePermutation(const FGlobalShaderPermutationParameters& Parameters)
@@ -67,10 +69,12 @@ void AddHairStrandsCopyArgsTilesPass(
 	TShaderMapRef<FHairStrandsTileCopyArgsPassCS> ComputeShader(View.ShaderMap);
 	FHairStrandsTileCopyArgsPassCS::FParameters* PassParameters = GraphBuilder.AllocParameters<FHairStrandsTileCopyArgsPassCS::FParameters>();
 	PassParameters->TileCountXY = TileData.TileCountXY;
+	PassParameters->TilePerThread_GroupSize = TileData.TilePerThread_GroupSize;
 	PassParameters->bRectPrimitive = TileData.bRectPrimitive ? 1 : 0;
 	PassParameters->TileCountBuffer = GraphBuilder.CreateSRV(TileData.TileCountBuffer, PF_R32_UINT);
 	PassParameters->TileIndirectDrawBuffer = GraphBuilder.CreateUAV(TileData.TileIndirectDrawBuffer, PF_R32_UINT);
 	PassParameters->TileIndirectDispatchBuffer = GraphBuilder.CreateUAV(TileData.TileIndirectDispatchBuffer, PF_R32_UINT);
+	PassParameters->TilePerThreadIndirectDispatchBuffer = GraphBuilder.CreateUAV(TileData.TilePerThreadIndirectDispatchBuffer, PF_R32_UINT);
 
 	FComputeShaderUtils::AddPass(
 		GraphBuilder,
@@ -118,6 +122,7 @@ FHairStrandsTiles AddHairStrandsGenerateTilesPass(
 {
 	FHairStrandsTiles Out;
 
+	check(FHairStrandsTiles::TilePerThread_GroupSize == 64); // If this value change, we need to update the shaders using 
 	check(FHairStrandsTiles::TileSize == 8); // only size supported for now
 	const FIntPoint InputResolution = InputTexture->Desc.Extent;
 	Out.TileCountXY = FIntPoint(FMath::CeilToInt(InputResolution.X / float(FHairStrandsTiles::TileSize)), FMath::CeilToInt(InputResolution.Y / float(FHairStrandsTiles::TileSize)));
@@ -127,6 +132,7 @@ FHairStrandsTiles AddHairStrandsGenerateTilesPass(
 	Out.TileCountBuffer = GraphBuilder.CreateBuffer(FRDGBufferDesc::CreateBufferDesc(4, 1), TEXT("Hair.TileCountBuffer"));
 	Out.TileIndirectDrawBuffer = GraphBuilder.CreateBuffer(FRDGBufferDesc::CreateIndirectDesc<FRHIDrawIndirectParameters>(1), TEXT("Hair.TileIndirectDrawBuffer"));
 	Out.TileIndirectDispatchBuffer = GraphBuilder.CreateBuffer(FRDGBufferDesc::CreateIndirectDesc<FRHIDispatchIndirectParameters>(1), TEXT("Hair.TileIndirectDispatchBuffer"));
+	Out.TilePerThreadIndirectDispatchBuffer = GraphBuilder.CreateBuffer(FRDGBufferDesc::CreateIndirectDesc<FRHIDispatchIndirectParameters>(1), TEXT("Hair.TilePerThreadIndirectDispatchBuffer"));
 	Out.TileDataBuffer = GraphBuilder.CreateBuffer(FRDGBufferDesc::CreateBufferDesc(sizeof(uint32), Out.TileCount), TEXT("Hair.TileDataBuffer"));
 
 	FRDGBufferUAVRef TileCountUAV = GraphBuilder.CreateUAV(Out.TileCountBuffer, PF_R32_UINT);
