- PopFirstUnsubmittedBreadcrumb also for ComputeContext
- Fix scoping of event for async compute path

diff --git a/Engine/Source/Runtime/RHI/Private/RHICommandList.cpp b/Engine/Source/Runtime/RHI/Private/RHICommandList.cpp
--- a/Engine/Source/Runtime/RHI/Private/RHICommandList.cpp
+++ b/Engine/Source/Runtime/RHI/Private/RHICommandList.cpp
@@ -445,6 +445,10 @@ void FRHICommandListExecutor::ExecuteInner_DoExecute(FRHICommandListBase& CmdLis
 	{
 		CmdList.Context->RHISetBreadcrumbStackTop(CmdList.PopFirstUnsubmittedBreadcrumb());
 	}
+	else if (CmdList.ComputeContext)
+	{
+		CmdList.ComputeContext->RHISetBreadcrumbStackTop(CmdList.PopFirstUnsubmittedBreadcrumb());
+	}
 #endif
 
 #if WITH_MGPU
diff --git a/Engine/Source/Runtime/Renderer/Private/DeferredShadingRenderer.cpp b/Engine/Source/Runtime/Renderer/Private/DeferredShadingRenderer.cpp
--- a/Engine/Source/Runtime/Renderer/Private/DeferredShadingRenderer.cpp
+++ b/Engine/Source/Runtime/Renderer/Private/DeferredShadingRenderer.cpp
@@ -1441,8 +1441,10 @@ bool FDeferredShadingSceneRenderer::DispatchRayTracingWorldUpdates(FRDGBuilder&
 
 			RHIAsyncCmdList.BindAccelerationStructureMemory(RayTracingSceneRHI, Scene->RayTracingScene.GetBufferChecked(), 0);
 
-			SCOPED_DRAW_EVENT(RHIAsyncCmdList, RayTracingScene);
-			RHIAsyncCmdList.BuildAccelerationStructure(RayTracingSceneRHI);
+			{
+				SCOPED_DRAW_EVENT(RHIAsyncCmdList, RayTracingScene);
+				RHIAsyncCmdList.BuildAccelerationStructure(RayTracingSceneRHI);				
+			}
 
 			RHIAsyncCmdList.BeginTransition(RayTracingDynamicGeometryUpdateEndTransition);
 			FRHIAsyncComputeCommandListImmediate::ImmediateDispatch(RHIAsyncCmdList);
