Adds an option for r.TemporalAA.ShadingRejection.SpatialFilter to run at history rejection resolution

diff --git a/Engine/Shaders/Private/TemporalAA/TAACompareHistory.usf b/Engine/Shaders/Private/TemporalAA/TAACompareHistory.usf
--- a/Engine/Shaders/Private/TemporalAA/TAACompareHistory.usf
+++ b/Engine/Shaders/Private/TemporalAA/TAACompareHistory.usf
@@ -5,7 +5,7 @@
 
 //------------------------------------------------------- CONFIG
 
-#define CONFIG_OUTPUT_HALF_RES 1
+#define CONFIG_OUTPUT_HALF_RES DIM_OUTPUT_HALF_RES
 
 #define TILE_SIZE 16
 
@@ -128,24 +128,21 @@ void MainCS(
 
 		HistoryRejectionOutput[OutputPixelPos] = Rejection;
 	}
-	#endif
-
+	#else
 	{
 		bool bIsValid = all(InputPixelPos <= taa_short2(InputPixelPosMax));
 
 		taa_ushort2 OutputPixelPos = taa_ushort2(InputPixelPos - taa_short2(InputPixelPosMin));
 		OutputPixelPos.x = bIsValid ? OutputPixelPos.x : ~taa_ushort(0);
 
-		#if !CONFIG_OUTPUT_HALF_RES
-		{
-			HistoryRejectionOutput[OutputPixelPos] = Rejection;
-		}
-		#endif
+		HistoryRejectionOutput[OutputPixelPos] = Rejection;
+	}
+	#endif
 
-		#if DEBUG_OUTPUT
-		{
-			DebugOutput[OutputPixelPos] = Debug;
-		}
-		#endif
+	#if DEBUG_OUTPUT
+	{
+		taa_short2 OutputPixelPos = InvalidateOutputPixelPos(taa_ushort2(InputPixelPos), InputPixelPosMax);
+		DebugOutput[OutputPixelPos] = Debug;
 	}
+	#endif
 }
diff --git a/Engine/Shaders/Private/TemporalAA/TAAPostfilterRejection.usf b/Engine/Shaders/Private/TemporalAA/TAAPostfilterRejection.usf
--- a/Engine/Shaders/Private/TemporalAA/TAAPostfilterRejection.usf
+++ b/Engine/Shaders/Private/TemporalAA/TAAPostfilterRejection.usf
@@ -7,13 +7,22 @@
 
 #define TILE_SIZE 8
 
+#define CONFIG_OUTPUT_HALF_RES DIM_OUTPUT_HALF_RES
+
 
 //------------------------------------------------------- PARAMETERS
 
+uint4 HistoryRejectionViewport;
+
 Texture2D<taa_half> HistoryRejectionTexture;
 RWTexture2D<taa_half> HistoryRejectionOutput;
 
 
+//------------------------------------------------------- LDS
+
+groupshared taa_half SharedRejection[TILE_SIZE * TILE_SIZE];
+
+
 //------------------------------------------------------- FUNCTIONS
 
 // Operator that output, lowest, median and highest values from 3 input values.
@@ -76,29 +85,55 @@ void MainCS(
 		taa_short2(GroupId) * taa_short2(TILE_SIZE, TILE_SIZE) +
 		Map8x8Tile2x2Lane(GroupThreadIndex));
 
+	RejectionPixelPos = ClampPixelOffset(RejectionPixelPos, HistoryRejectionViewport.xy, HistoryRejectionViewport.zw - 1);
+
 	// Fetch neighborhood of the kernel
 	taa_half InSamples[9];
 	UNROLL
 	for (uint i = 0; i < 9; i++)
 	{
-		taa_short2 SampleInputPixelPos = RejectionPixelPos + ClampPixelOffset(
+		taa_short2 SampleInputPixelPos = AddAndClampPixelOffset(
 			RejectionPixelPos,
 			kOffsets3x3[i], kOffsets3x3[i],
-			RejectionInfo_ViewportMin, RejectionInfo_ViewportMax - 1);
+			HistoryRejectionViewport.xy, HistoryRejectionViewport.zw - 1);
 
 		InSamples[i] = HistoryRejectionTexture[SampleInputPixelPos];
 	}
 
 	taa_half Rejection = Median9(InSamples);
 
-	bool bIsValid = all(RejectionPixelPos < taa_short2(RejectionInfo_ViewportMax));
-	RejectionPixelPos.x = bIsValid ? RejectionPixelPos.x : ~taa_ushort(0);
+	#if CONFIG_OUTPUT_HALF_RES
+	{
+		SharedRejection[GroupThreadIndex] = Rejection;
+
+		GroupMemoryBarrierWithGroupSync();
 
-	HistoryRejectionOutput[RejectionPixelPos] = Rejection;
+		UNROLL
+		for (uint i = 1; i < 4; i++)
+		{
+			uint ButterflySwap = i;
 
+			Rejection = min(Rejection, SharedRejection[GroupThreadIndex ^ ButterflySwap]);
+		}
+		
+		bool bIsValid = all(RejectionPixelPos <= taa_short2(HistoryRejectionViewport.zw)) && all((taa_ushort2(RejectionPixelPos) % taa_ushort(2)) == taa_ushort(0));
+
+		taa_ushort2 OutputPixelPos = taa_ushort2(RejectionPixelPos) / taa_ushort(2);
+		OutputPixelPos.x = bIsValid ? OutputPixelPos.x : ~taa_ushort(0);
+
+		HistoryRejectionOutput[OutputPixelPos] = Rejection;
+	}
+	#else
+	{
+		taa_short2 OutputPixelPos = InvalidateOutputPixelPos(RejectionPixelPos, HistoryRejectionViewport.zw);
+		HistoryRejectionOutput[OutputPixelPos] = Rejection;
+	}
+	#endif
+	
 	#if DEBUG_OUTPUT
 	{
-		DebugOutput[RejectionPixelPos] = Debug;
+		taa_short2 OutputPixelPos = InvalidateOutputPixelPos(RejectionPixelPos, HistoryRejectionViewport.zw);
+		DebugOutput[OutputPixelPos] = Debug;
 	}
 	#endif
 }
diff --git a/Engine/Source/Runtime/Renderer/Private/PostProcess/TemporalAA.cpp b/Engine/Source/Runtime/Renderer/Private/PostProcess/TemporalAA.cpp
--- a/Engine/Source/Runtime/Renderer/Private/PostProcess/TemporalAA.cpp
+++ b/Engine/Source/Runtime/Renderer/Private/PostProcess/TemporalAA.cpp
@@ -88,7 +88,10 @@ TAutoConsoleVariable<int32> CVarTAAHalfResShadingRejection(
 
 TAutoConsoleVariable<int32> CVarTAAFilterShadingRejection(
 	TEXT("r.TemporalAA.ShadingRejection.SpatialFilter"), 1,
-	TEXT("Whether the shading rejection should have spatial statistical filtering pass to reduce flickering (default = 1)."),
+	TEXT("Whether the shading rejection should have spatial statistical filtering pass to reduce flickering (default = 1).\n")
+	TEXT(" 0: Disabled;\n")
+	TEXT(" 1: Spatial filter pass is run at lower resolution than CompareHistory pass (default);\n")
+	TEXT(" 2: Spatial filter pass is run CompareHistory pass resolution to improve stability."),
 	ECVF_RenderThreadSafe);
 
 TAutoConsoleVariable<int32> CVarTAAEnableAntiInterference(
@@ -510,6 +513,9 @@ class FTSRCompareHistoryCS : public FTemporalSuperResolutionShader
 	DECLARE_GLOBAL_SHADER(FTSRCompareHistoryCS);
 	SHADER_USE_PARAMETER_STRUCT(FTSRCompareHistoryCS, FTemporalSuperResolutionShader);
 
+	class FOutputHalfRes : SHADER_PERMUTATION_BOOL("DIM_OUTPUT_HALF_RES");
+	using FPermutationDomain = TShaderPermutationDomain<FOutputHalfRes>;
+
 	BEGIN_SHADER_PARAMETER_STRUCT(FParameters, )
 		SHADER_PARAMETER_STRUCT_INCLUDE(FTSRCommonParameters, CommonParameters)
 		SHADER_PARAMETER_RDG_TEXTURE(Texture2D, ParallaxRejectionMaskTexture)
@@ -526,8 +532,11 @@ class FTSRPostfilterRejectionCS : public FTemporalSuperResolutionShader
 	DECLARE_GLOBAL_SHADER(FTSRPostfilterRejectionCS);
 	SHADER_USE_PARAMETER_STRUCT(FTSRPostfilterRejectionCS, FTemporalSuperResolutionShader);
 
+	class FOutputHalfRes : SHADER_PERMUTATION_BOOL("DIM_OUTPUT_HALF_RES");
+	using FPermutationDomain = TShaderPermutationDomain<FOutputHalfRes>;
+
 	BEGIN_SHADER_PARAMETER_STRUCT(FParameters, )
-		SHADER_PARAMETER_STRUCT_INCLUDE(FTSRCommonParameters, CommonParameters)
+		SHADER_PARAMETER(FIntRect, HistoryRejectionViewport)
 		SHADER_PARAMETER_RDG_TEXTURE(Texture2D, HistoryRejectionTexture)
 		SHADER_PARAMETER_RDG_TEXTURE_UAV(RWTexture2D, HistoryRejectionOutput)
 		SHADER_PARAMETER_RDG_TEXTURE_UAV(RWTexture2D, DebugOutput)
@@ -1183,6 +1192,15 @@ static void AddTemporalSuperResolutionPasses(
 
 	bool bRejectSeparateTranslucency = PassInputs.SeparateTranslucencyTextures != nullptr && CVarTAARejectTranslucency.GetValueOnRenderThread() != 0;
 
+	enum class ERejectionPostFilter : uint8
+	{
+		Disabled,
+		PostRejectionDownsample,
+		PreRejectionDownsample,
+	};
+
+	ERejectionPostFilter PostFilter = ERejectionPostFilter(FMath::Clamp(CVarTAAFilterShadingRejection.GetValueOnRenderThread(), 0, 2));
+
 	FIntPoint InputExtent = PassInputs.SceneColorTexture->Desc.Extent;
 	FIntRect InputRect = View.ViewRect;
 
@@ -1737,9 +1755,11 @@ static void AddTemporalSuperResolutionPasses(
 
 		// Compare the low frequencies
 		{
+			bool bOutputHalfRes = PostFilter != ERejectionPostFilter::PreRejectionDownsample;
+
 			{
 				FRDGTextureDesc Desc = FRDGTextureDesc::Create2D(
-					RejectionExtent,
+					bOutputHalfRes ? RejectionExtent : LowFrequencyExtent,
 					PF_R8,
 					FClearValueBinding::None,
 					/* InFlags = */ TexCreate_ShaderResource | TexCreate_UAV);
@@ -1757,7 +1777,10 @@ static void AddTemporalSuperResolutionPasses(
 			PassParameters->HistoryRejectionOutput = GraphBuilder.CreateUAV(HistoryRejectionTexture);
 			PassParameters->DebugOutput = CreateDebugUAV(LowFrequencyExtent, TEXT("Debug.TSR.CompareHistory"));
 
-			TShaderMapRef<FTSRCompareHistoryCS> ComputeShader(View.ShaderMap);
+			FTSRCompareHistoryCS::FPermutationDomain PermutationVector;
+			PermutationVector.Set<FTSRCompareHistoryCS::FOutputHalfRes>(bOutputHalfRes);
+
+			TShaderMapRef<FTSRCompareHistoryCS> ComputeShader(View.ShaderMap, PermutationVector);
 			FComputeShaderUtils::AddPass(
 				GraphBuilder,
 				RDG_EVENT_NAME("TSR CompareHistory %dx%d", LowFrequencyRect.Width(), LowFrequencyRect.Height()),
@@ -1768,23 +1791,38 @@ static void AddTemporalSuperResolutionPasses(
 	}
 
 	// Post filter the rejection.
-	if (CVarTAAFilterShadingRejection.GetValueOnRenderThread() != 0)
+	if (PostFilter != ERejectionPostFilter::Disabled)
 	{
-		FRDGTextureRef FilteredHistoryRejectionTexture = GraphBuilder.CreateTexture(HistoryRejectionTexture->Desc, TEXT("TSR.FilteredHistoryRejection"));
+		bool bOutputHalfRes = PostFilter == ERejectionPostFilter::PreRejectionDownsample;
+		FIntRect Rect = bOutputHalfRes ? LowFrequencyRect : RejectionRect;
+
+		FRDGTextureRef FilteredHistoryRejectionTexture;
+		{
+			FRDGTextureDesc Desc = FRDGTextureDesc::Create2D(
+				RejectionExtent,
+				PF_R8,
+				FClearValueBinding::None,
+				/* InFlags = */ TexCreate_ShaderResource | TexCreate_UAV);
+
+			FilteredHistoryRejectionTexture = GraphBuilder.CreateTexture(Desc, TEXT("TSR.HistoryRejection"));
+		}
 
 		FTSRPostfilterRejectionCS::FParameters* PassParameters = GraphBuilder.AllocParameters<FTSRPostfilterRejectionCS::FParameters>();
-		PassParameters->CommonParameters = CommonParameters;
+		PassParameters->HistoryRejectionViewport = Rect;
 		PassParameters->HistoryRejectionTexture = HistoryRejectionTexture;
 		PassParameters->HistoryRejectionOutput = GraphBuilder.CreateUAV(FilteredHistoryRejectionTexture);
-		PassParameters->DebugOutput = CreateDebugUAV(RejectionExtent, TEXT("Debug.TSR.PostfilterRejection"));
+		PassParameters->DebugOutput = CreateDebugUAV(bOutputHalfRes ? LowFrequencyExtent : RejectionExtent, TEXT("Debug.TSR.PostfilterRejection"));
+
+		FTSRPostfilterRejectionCS::FPermutationDomain PermutationVector;
+		PermutationVector.Set<FTSRPostfilterRejectionCS::FOutputHalfRes>(PostFilter == ERejectionPostFilter::PreRejectionDownsample);
 
-		TShaderMapRef<FTSRPostfilterRejectionCS> ComputeShader(View.ShaderMap);
+		TShaderMapRef<FTSRPostfilterRejectionCS> ComputeShader(View.ShaderMap, PermutationVector);
 		FComputeShaderUtils::AddPass(
 			GraphBuilder,
-			RDG_EVENT_NAME("TSR PostfilterRejection %dx%d", RejectionRect.Width(), RejectionRect.Height()),
+			RDG_EVENT_NAME("TSR PostfilterRejection %dx%d", Rect.Width(), Rect.Height()),
 			ComputeShader,
 			PassParameters,
-			FComputeShaderUtils::GetGroupCount(RejectionRect.Size(), 8));
+			FComputeShaderUtils::GetGroupCount(Rect.Size(), 8));
 
 		HistoryRejectionTexture = FilteredHistoryRejectionTexture;
 	}
@@ -1949,7 +1987,7 @@ static void AddTemporalSuperResolutionPasses(
 		TShaderMapRef<FTSRDebugHistoryCS> ComputeShader(View.ShaderMap);
 		FComputeShaderUtils::AddPass(
 			GraphBuilder,
-			RDG_EVENT_NAME("TSR DebugHistory %dx%d", RejectionRect.Width(), RejectionRect.Height()),
+			RDG_EVENT_NAME("TSR DebugHistory %dx%d", HistorySize.X, HistorySize.Y),
 			ComputeShader,
 			PassParameters,
 			FComputeShaderUtils::GetGroupCount(HistorySize * kHistoryUpscalingFactor, 8));
