Fix post processing using 1x1 dummy depth stencil texture causing RTV/DSV size mismatch and GPU crashes.

diff --git a/Engine/Source/Runtime/Renderer/Private/DeferredShadingRenderer.cpp b/Engine/Source/Runtime/Renderer/Private/DeferredShadingRenderer.cpp
--- a/Engine/Source/Runtime/Renderer/Private/DeferredShadingRenderer.cpp
+++ b/Engine/Source/Runtime/Renderer/Private/DeferredShadingRenderer.cpp
@@ -3098,6 +3098,7 @@ void FDeferredShadingSceneRenderer::Render(FRDGBuilder& GraphBuilder)
 
 		FPostProcessingInputs PostProcessingInputs;
 		PostProcessingInputs.ViewFamilyTexture = ViewFamilyTexture;
+		PostProcessingInputs.CustomDepthTexture = SceneTextures.CustomDepth.Depth;
 		PostProcessingInputs.SeparateTranslucencyTextures = &SeparateTranslucencyTextures;
 		PostProcessingInputs.SceneTextures = SceneTextures.UniformBuffer;
 
diff --git a/Engine/Source/Runtime/Renderer/Private/PostProcess/PostProcessing.cpp b/Engine/Source/Runtime/Renderer/Private/PostProcess/PostProcessing.cpp
--- a/Engine/Source/Runtime/Renderer/Private/PostProcess/PostProcessing.cpp
+++ b/Engine/Source/Runtime/Renderer/Private/PostProcess/PostProcessing.cpp
@@ -273,7 +273,7 @@ void AddPostProcessingPasses(FRDGBuilder& GraphBuilder, const FViewInfo& View, c
 	const FScreenPassRenderTarget ViewFamilyOutput = FScreenPassRenderTarget::CreateViewFamilyOutput(Inputs.ViewFamilyTexture, View);
 	const FScreenPassTexture SceneDepth(SceneTextureParameters.SceneDepthTexture, PrimaryViewRect);
 	const FScreenPassTexture SeparateTranslucency(Inputs.SeparateTranslucencyTextures->GetColorForRead(GraphBuilder), PrimaryViewRect);
-	const FScreenPassTexture CustomDepth((*Inputs.SceneTextures)->CustomDepthTexture, PrimaryViewRect);
+	const FScreenPassTexture CustomDepth(Inputs.CustomDepthTexture, PrimaryViewRect);
 	const FScreenPassTexture Velocity(SceneTextureParameters.GBufferVelocityTexture, PrimaryViewRect);
 	const FScreenPassTexture BlackDummy(GSystemTextures.GetBlackDummy(GraphBuilder));
 
diff --git a/Engine/Source/Runtime/Renderer/Private/PostProcess/PostProcessing.h b/Engine/Source/Runtime/Renderer/Private/PostProcess/PostProcessing.h
--- a/Engine/Source/Runtime/Renderer/Private/PostProcess/PostProcessing.h
+++ b/Engine/Source/Runtime/Renderer/Private/PostProcess/PostProcessing.h
@@ -47,12 +47,14 @@ struct FPostProcessingInputs
 {
 	TRDGUniformBufferRef<FSceneTextureUniformParameters> SceneTextures = nullptr;
 	FRDGTextureRef ViewFamilyTexture = nullptr;
+	FRDGTextureRef CustomDepthTexture = nullptr;
 	const FSeparateTranslucencyTextures* SeparateTranslucencyTextures = nullptr;
 
 	void Validate() const
 	{
 		check(SceneTextures);
 		check(ViewFamilyTexture);
+		check(CustomDepthTexture);
 		check(SeparateTranslucencyTextures);
 	}
 };
