Fixed ray tracing primitive index. PrimitiveIndex was cached in CachedRayTracingInstance, but wasn't updated inside FScene_SwapPrimitiveSceneInfos leading to stale indices.

diff --git a/Engine/Source/Runtime/Renderer/Private/RendererScene.cpp b/Engine/Source/Runtime/Renderer/Private/RendererScene.cpp
--- a/Engine/Source/Runtime/Renderer/Private/RendererScene.cpp
+++ b/Engine/Source/Runtime/Renderer/Private/RendererScene.cpp
@@ -3965,6 +3965,12 @@ void FScene::UpdateAllPrimitiveSceneInfos(FRDGBuilder& GraphBuilder, bool bAsync
 							LumenSceneData->AddPrimitiveToUpdate(SourceIndex);
 							LumenSceneData->AddPrimitiveToUpdate(DestIndex);
 
+						#if RHI_RAYTRACING
+							// Update cached PrimitiveIndex after an index swap
+							Primitives[SourceIndex]->CachedRayTracingInstance.DefaultUserData = SourceIndex;
+							Primitives[DestIndex]->CachedRayTracingInstance.DefaultUserData = DestIndex;
+						#endif
+
 							SourceIndex = DestIndex;
 						}
 					}
