Add info to hair view uniform buffer to indicate of the tile data are valid or not.

diff --git a/Engine/Source/Runtime/Renderer/Private/HairStrands/HairStrandsData.h b/Engine/Source/Runtime/Renderer/Private/HairStrands/HairStrandsData.h
--- a/Engine/Source/Runtime/Renderer/Private/HairStrands/HairStrandsData.h
+++ b/Engine/Source/Runtime/Renderer/Private/HairStrands/HairStrandsData.h
@@ -22,6 +22,7 @@ BEGIN_GLOBAL_SHADER_PARAMETER_STRUCT(FHairStrandsViewUniformParameters, )
 	SHADER_PARAMETER(FIntPoint, HairTileCountXY)										// Tile count in X/Y
 	SHADER_PARAMETER(float, HairDualScatteringRoughnessOverride)						// Override the roughness used for dual scattering (for hack/test purpose only)
 	SHADER_PARAMETER(FIntPoint, HairSampleViewportResolution)							// Maximum viewport resolution of the sample space
+	SHADER_PARAMETER(uint32, bHairTileValid)											// True if tile data are valid
 	SHADER_PARAMETER_RDG_TEXTURE(Texture2D<uint4>, HairCategorizationTexture)			// Categorization texture aggregating hair info in screen space (closest depth, coverage, ...)
 	SHADER_PARAMETER_RDG_TEXTURE(Texture2D, HairOnlyDepthTexture)						// Depth texture containing only hair depth (not strongly typed for legacy shader code reason)
 	SHADER_PARAMETER_RDG_TEXTURE(Texture2D, HairOnlyDepthClosestHZBTexture)				// HZB closest depth texture containing only hair depth (not strongly typed for legacy shader code reason)
diff --git a/Engine/Source/Runtime/Renderer/Private/HairStrands/HairStrandsRendering.cpp b/Engine/Source/Runtime/Renderer/Private/HairStrands/HairStrandsRendering.cpp
--- a/Engine/Source/Runtime/Renderer/Private/HairStrands/HairStrandsRendering.cpp
+++ b/Engine/Source/Runtime/Renderer/Private/HairStrands/HairStrandsRendering.cpp
@@ -27,6 +27,7 @@ static TRDGUniformBufferRef<FHairStrandsViewUniformParameters> InternalCreateHai
 			Parameters->HairTileData = In->TileData.TileDataSRV;
 			Parameters->HairTileCount = GraphBuilder.CreateSRV(In->TileData.TileCountBuffer, PF_R32_UINT);
 			Parameters->HairTileCountXY = In->TileData.TileCountXY;
+			Parameters->bHairTileValid = true;
 		}
 		else
 		{
@@ -34,9 +35,10 @@ static TRDGUniformBufferRef<FHairStrandsViewUniformParameters> InternalCreateHai
 			AddClearUAVPass(GraphBuilder, GraphBuilder.CreateUAV(DummyBuffer, PF_R16G16_UINT), 0);
 			FRDGBufferSRVRef DummyBufferR32SRV = GraphBuilder.CreateSRV(DummyBuffer, PF_R32_UINT);
 			FRDGBufferSRVRef DummyBufferRG16SRV = GraphBuilder.CreateSRV(DummyBuffer, PF_R16G16_UINT);
-			Parameters->HairTileData = DummyBufferR32SRV;
-			Parameters->HairTileCount = DummyBufferRG16SRV;
+			Parameters->HairTileData = DummyBufferRG16SRV;
+			Parameters->HairTileCount = DummyBufferR32SRV;
 			Parameters->HairTileCountXY = FIntPoint(0,0);
+			Parameters->bHairTileValid = false;
 		}
 
 		if (!Parameters->HairOnlyDepthFurthestHZBTexture)
