Add indirect tile dispatch for marking shadow page covered by hair.

Review:  https://p4-swarm.epicgames.net/reviews/15901055/v1/

diff --git a/Engine/Shaders/Private/VirtualShadowMaps/PageManagement.usf b/Engine/Shaders/Private/VirtualShadowMaps/PageManagement.usf
--- a/Engine/Shaders/Private/VirtualShadowMaps/PageManagement.usf
+++ b/Engine/Shaders/Private/VirtualShadowMaps/PageManagement.usf
@@ -27,7 +27,8 @@ RWStructuredBuffer<uint> OutStatsBuffer;
 
 // Type of input data consume by the page allocation (i.e., data read from the source buffer: Gbuffer, HairStrands data, ...)
 #define INPUT_TYPE_GBUFFER 0
-#define INPUT_TYPE_HAIRSTRANDS 1
+#define INPUT_TYPE_NANITE 1
+#define INPUT_TYPE_HAIRSTRANDS 2
 
 // Flags generated by per-pixel pass to determine which pages are required to provide shadow for the visible geometry
 RWStructuredBuffer<uint> OutPageRequestFlags;
@@ -64,7 +65,7 @@ float GetDeviceZFromPixel(uint2 PixelPos)
 {
 	// Merge Nanite depth with depth buffer if enabled (i.e. we're before the base pass)
 	float DeviceZ = 0.0f;
-#if LOAD_DEPTH_FROM_NANITE_BUFFER
+#if PERMUTATION_INPUT_TYPE == INPUT_TYPE_NANITE
 	// Load Nanite Depth
 	UlongType VisPixel = VisBuffer64[uint2(View.ViewRectMin.xy) + PixelPos];
 	uint DepthInt = 0;
@@ -174,10 +175,29 @@ void MarkPage(uint VirtualShadowMapId, uint MipLevel, float3 WorldPosition, floa
 }
 
 [numthreads(VSM_DEFAULT_CS_GROUP_XY, VSM_DEFAULT_CS_GROUP_XY, 1)]
-void GeneratePageFlagsFromPixels(uint3 GroupId : SV_GroupID, uint GroupIndex : SV_GroupIndex, uint3 DispatchThreadId : SV_DispatchThreadID)
+void GeneratePageFlagsFromPixels(
+	uint3 InGroupId : SV_GroupID,
+	uint  GroupIndex : SV_GroupIndex,
+	uint3 GroupThreadId : SV_GroupThreadID,
+	uint3 DispatchThreadId : SV_DispatchThreadID)
 {
 	// Morton order within a group so page access/atomics are more coherent and wave-swizzled gradients are possible.
-	uint2 PixelPos = VSM_DEFAULT_CS_GROUP_XY * GroupId.xy + MortonDecode(GroupIndex);
+#if PERMUTATION_INPUT_TYPE == INPUT_TYPE_HAIRSTRANDS
+	uint2 GroupId		= InGroupId.xy;
+	const uint TileCount= HairStrands.HairTileCount[0];
+	if (HairStrands.bHairTileValid>0)
+	{
+		const uint LinearIndex  = InGroupId.x + InGroupId.y * HairStrands.HairTileCountXY.x;
+		if (LinearIndex >= TileCount)
+		{
+			return;
+		}
+		GroupId = HairStrands.HairTileData[LinearIndex];
+	}
+#else // PERMUTATION_INPUT_TYPE == INPUT_TYPE_GBUFFER || PERMUTATION_INPUT_TYPE == INPUT_TYPE_NANITE
+	const uint2 GroupId = InGroupId.xy;
+#endif
+	const uint2 PixelPos = VSM_DEFAULT_CS_GROUP_XY * GroupId.xy + MortonDecode(GroupIndex);
 
 	if (any(PixelPos >= uint2(View.ViewSizeAndInvSize.xy)))
 	{
@@ -187,7 +207,7 @@ void GeneratePageFlagsFromPixels(uint3 GroupId : SV_GroupID, uint GroupIndex : S
 	float3 WorldPosition = 0;
 	float  SceneDepth    = 0;
 	bool   bIsValid		 = false;
-	if (InputType == INPUT_TYPE_HAIRSTRANDS)
+	#if PERMUTATION_INPUT_TYPE == INPUT_TYPE_HAIRSTRANDS
 	{
 		FCategorizationData CatData = DecodeCategorizationData(HairStrands.HairCategorizationTexture.Load(uint3(PixelPos, 0)));
 		FScreenPixelData HairPixel = GetScreenPixelData(PixelPos, CatData.ClosestDepth);
@@ -195,7 +215,7 @@ void GeneratePageFlagsFromPixels(uint3 GroupId : SV_GroupID, uint GroupIndex : S
 		SceneDepth = HairPixel.SceneDepth;
 		bIsValid = CatData.PixelCoverage > 0;
 	}
-	else// (InputType == INPUT_TYPE_GBUFFER)
+	#else // PERMUTATION_INPUT_TYPE == INPUT_TYPE_GBUFFER || PERMUTATION_INPUT_TYPE == INPUT_TYPE_NANITE
 	{
 		FScreenPixelData Pixel = GetScreenPixelData(PixelPos, GetDeviceZFromPixel(PixelPos));
 		WorldPosition = Pixel.WorldPosition;
@@ -206,6 +226,7 @@ void GeneratePageFlagsFromPixels(uint3 GroupId : SV_GroupID, uint GroupIndex : S
 			bIsValid = GetGBufferData(Pixel.ScreenUV).ShadingModelID != SHADINGMODELID_UNLIT;
 		}
 	}
+	#endif
 
 	// Dither pattern for page dilation
 	// We don't need to to check all 8 adjacent pages; as long as there's at least a single pixel near the edge
diff --git a/Engine/Source/Runtime/Renderer/Private/VirtualShadowMaps/VirtualShadowMapArray.cpp b/Engine/Source/Runtime/Renderer/Private/VirtualShadowMaps/VirtualShadowMapArray.cpp
--- a/Engine/Source/Runtime/Renderer/Private/VirtualShadowMaps/VirtualShadowMapArray.cpp
+++ b/Engine/Source/Runtime/Renderer/Private/VirtualShadowMaps/VirtualShadowMapArray.cpp
@@ -375,8 +375,8 @@ class FGeneratePageFlagsFromPixelsCS : public FVirtualPageManagementShader
 	DECLARE_GLOBAL_SHADER(FGeneratePageFlagsFromPixelsCS);
 	SHADER_USE_PARAMETER_STRUCT(FGeneratePageFlagsFromPixelsCS, FVirtualPageManagementShader)
 
-	class FNaniteDepthBufferDim : SHADER_PERMUTATION_BOOL("LOAD_DEPTH_FROM_NANITE_BUFFER");
-	using FPermutationDomain = TShaderPermutationDomain<FNaniteDepthBufferDim>;
+	class FInputType : SHADER_PERMUTATION_INT("PERMUTATION_INPUT_TYPE", 3); 
+	using FPermutationDomain = TShaderPermutationDomain<FInputType>;
 	
 	BEGIN_SHADER_PARAMETER_STRUCT(FParameters, )
 		SHADER_PARAMETER_RDG_UNIFORM_BUFFER(FVirtualShadowMapUniformParameters, VirtualShadowMap)
@@ -387,6 +387,7 @@ class FGeneratePageFlagsFromPixelsCS : public FVirtualPageManagementShader
 		SHADER_PARAMETER_RDG_TEXTURE(Texture2D<uint2>, VisBuffer64)
 		SHADER_PARAMETER_RDG_BUFFER_UAV(RWStructuredBuffer<uint>, OutPageRequestFlags)
 		SHADER_PARAMETER_RDG_BUFFER_SRV(StructuredBuffer< uint >, VirtualShadowMapIdRemap)
+		RDG_BUFFER_ACCESS(IndirectBufferArgs, ERHIAccess::IndirectArgs)
 		SHADER_PARAMETER(uint32, InputType)
 		SHADER_PARAMETER(uint32, NumDirectionalLightSmInds)
 		SHADER_PARAMETER(uint32, bPostBasePass)
@@ -912,9 +913,11 @@ void FVirtualShadowMapArray::BuildPageAllocations(
 				{
 					auto GeneratePageFlags = [&](bool bHairPass)
 					{
-						FGeneratePageFlagsFromPixelsCS::FPermutationDomain PermutationVector;
-						PermutationVector.Set<FGeneratePageFlagsFromPixelsCS::FNaniteDepthBufferDim>(NaniteVisBuffer64 != nullptr && !bPostBasePass);
+						const bool bUseNaniteDepth = NaniteVisBuffer64 != nullptr && !bPostBasePass;
+						const uint32 InputType = bHairPass ? 2u : (bUseNaniteDepth ? 1u : 0u); // HairStrands, Nanite, or GBuffer
 
+						FGeneratePageFlagsFromPixelsCS::FPermutationDomain PermutationVector;
+						PermutationVector.Set<FGeneratePageFlagsFromPixelsCS::FInputType>(InputType);
 						FGeneratePageFlagsFromPixelsCS::FParameters* PassParameters = GraphBuilder.AllocParameters< FGeneratePageFlagsFromPixelsCS::FParameters >();
 						PassParameters->VirtualShadowMap = GetUniformBuffer(GraphBuilder);
 
@@ -922,7 +925,6 @@ void FVirtualShadowMapArray::BuildPageAllocations(
 						PassParameters->bPostBasePass = bPostBasePass;
 
 						PassParameters->VisBuffer64 = VisBuffer64;
-						PassParameters->InputType = bHairPass ? 1 : 0;
 						PassParameters->HairStrands = HairStrands::BindHairStrandsViewUniformParameters(View);
 						PassParameters->View = View.ViewUniformBuffer;
 						PassParameters->OutPageRequestFlags = PageRequestFlagsUAV;
@@ -937,13 +939,26 @@ void FVirtualShadowMapArray::BuildPageAllocations(
 						static_assert((FVirtualPageManagementShader::DefaultCSGroupXY % 2) == 0, "GeneratePageFlagsFromPixels requires even-sized CS groups for quad swizzling.");
 						const FIntPoint GridSize = FIntPoint::DivideAndRoundUp(View.ViewRect.Size(), FVirtualPageManagementShader::DefaultCSGroupXY);
 
-						FComputeShaderUtils::AddPass(
-							GraphBuilder,
-							RDG_EVENT_NAME("GeneratePageFlagsFromPixels(%s)", bHairPass ? TEXT("HairStrands") : TEXT("GBuffer")),
-							ComputeShader,
-							PassParameters,
-							FIntVector(GridSize.X, GridSize.Y, 1)
-						);
+						if (bHairPass && View.HairStrandsViewData.VisibilityData.TileData.IsValid())
+						{
+							PassParameters->IndirectBufferArgs = View.HairStrandsViewData.VisibilityData.TileData.TileIndirectDispatchBuffer;
+							FComputeShaderUtils::AddPass(
+								GraphBuilder,
+								RDG_EVENT_NAME("GeneratePageFlagsFromPixels(HairStrands,Tile)"),
+								ComputeShader,
+								PassParameters,
+								View.HairStrandsViewData.VisibilityData.TileData.TileIndirectDispatchBuffer,
+								0);
+						}
+						else
+						{
+							FComputeShaderUtils::AddPass(
+								GraphBuilder,
+								RDG_EVENT_NAME("GeneratePageFlagsFromPixels(%s)", bHairPass ? TEXT("HairStrands") : (bUseNaniteDepth ? TEXT("Nanite") : TEXT("GBuffer"))),
+								ComputeShader,
+								PassParameters,
+								FIntVector(GridSize.X, GridSize.Y, 1));
+						}
 					};
 
 					GeneratePageFlags(false);
