Add render graph support for default textures and default buffers with custom format and custom default value.

diff --git a/Engine/Source/Runtime/RHI/Public/RHIUtilities.h b/Engine/Source/Runtime/RHI/Public/RHIUtilities.h
--- a/Engine/Source/Runtime/RHI/Public/RHIUtilities.h
+++ b/Engine/Source/Runtime/RHI/Public/RHIUtilities.h
@@ -64,6 +64,28 @@ static inline bool IsBlockCompressedFormat(EPixelFormat Format)
 	return false;
 }
 
+static bool IsFloatFormat(EPixelFormat Format)
+{
+	switch (Format)
+	{
+	case PF_A32B32G32R32F:
+	case PF_FloatRGB:
+	case PF_FloatRGBA:
+	case PF_R32_FLOAT:
+	case PF_G16R16F:
+	case PF_G16R16F_FILTER:
+	case PF_G32R32F:
+	case PF_R16F:
+	case PF_R16F_FILTER:
+	case PF_FloatR11G11B10:
+		return true;
+
+	default:
+		break;
+	}
+	return false;
+}
+
 inline bool RHISupportsTextureBuffers(const FStaticShaderPlatform Platform)
 {
 	return IsFeatureLevelSupported(Platform, ERHIFeatureLevel::SM5)
diff --git a/Engine/Source/Runtime/RenderCore/Private/RenderGraphUtils.cpp b/Engine/Source/Runtime/RenderCore/Private/RenderGraphUtils.cpp
--- a/Engine/Source/Runtime/RenderCore/Private/RenderGraphUtils.cpp
+++ b/Engine/Source/Runtime/RenderCore/Private/RenderGraphUtils.cpp
@@ -402,6 +402,31 @@ void AddClearUAVPass(FRDGBuilder& GraphBuilder, FRDGTextureUAVRef TextureUAV, co
 	AddClearUAVPass(GraphBuilder, TextureUAV, FVector4(ClearColor.R, ClearColor.G, ClearColor.B, ClearColor.A));
 }
 
+void AddClearUAVPass(FRDGBuilder& GraphBuilder, FRDGTextureUAVRef TextureUAV, uint32 Value)
+{
+	AddClearUAVPass(GraphBuilder, TextureUAV, { Value, Value , Value , Value });
+}
+
+void AddClearUAVPass(FRDGBuilder& GraphBuilder, FRDGTextureUAVRef TextureUAV, float Value)
+{
+	AddClearUAVPass(GraphBuilder, TextureUAV, { Value, Value , Value , Value });
+}
+
+void AddClearUAVPass(FRDGBuilder& GraphBuilder, FRDGTextureUAVRef TextureUAV, const FVector& Value)
+{
+	AddClearUAVPass(GraphBuilder, TextureUAV, { Value.X, Value.Y , Value.Z , 0.f });
+}
+
+void AddClearUAVPass(FRDGBuilder& GraphBuilder, FRDGTextureUAVRef TextureUAV, const FIntPoint& Value)
+{
+	AddClearUAVPass(GraphBuilder, TextureUAV, { uint32(Value.X), uint32(Value.Y), 0u, 0u });
+}
+
+void AddClearUAVPass(FRDGBuilder& GraphBuilder, FRDGTextureUAVRef TextureUAV, const FVector2D& Value)
+{
+	AddClearUAVPass(GraphBuilder, TextureUAV, { Value.X, Value.Y , 0.f, 0.f });
+}
+
 class FClearUAVRectsPS : public FGlobalShader
 {
 	DECLARE_GLOBAL_SHADER(FClearUAVRectsPS);
diff --git a/Engine/Source/Runtime/RenderCore/Public/RenderGraphUtils.h b/Engine/Source/Runtime/RenderCore/Public/RenderGraphUtils.h
--- a/Engine/Source/Runtime/RenderCore/Public/RenderGraphUtils.h
+++ b/Engine/Source/Runtime/RenderCore/Public/RenderGraphUtils.h
@@ -523,10 +523,23 @@ RENDERCORE_API void AddCopyToResolveTargetPass(
 
 /** Adds a render graph pass to clear a texture or buffer UAV with a single typed value. */
 RENDERCORE_API void AddClearUAVPass(FRDGBuilder& GraphBuilder, FRDGBufferUAVRef BufferUAV, uint32 Value);
+
 RENDERCORE_API void AddClearUAVFloatPass(FRDGBuilder& GraphBuilder, FRDGBufferUAVRef BufferUAV, float Value);
 
+RENDERCORE_API void AddClearUAVPass(FRDGBuilder& GraphBuilder, FRDGTextureUAVRef TextureUAV, float ClearValue);
+
+RENDERCORE_API void AddClearUAVPass(FRDGBuilder& GraphBuilder, FRDGTextureUAVRef TextureUAV, uint32 ClearValue);
+
+RENDERCORE_API void AddClearUAVPass(FRDGBuilder& GraphBuilder, FRDGTextureUAVRef TextureUAV, const FIntPoint& ClearValue);
+
+RENDERCORE_API void AddClearUAVPass(FRDGBuilder& GraphBuilder, FRDGTextureUAVRef TextureUAV, const FVector2D& ClearValue);
+
+RENDERCORE_API void AddClearUAVPass(FRDGBuilder& GraphBuilder, FRDGTextureUAVRef TextureUAV, const FVector& ClearValue);
+
 RENDERCORE_API void AddClearUAVPass(FRDGBuilder& GraphBuilder, FRDGTextureUAVRef TextureUAV, const FUintVector4& ClearValues);
 
+RENDERCORE_API void AddClearUAVPass(FRDGBuilder& GraphBuilder, FRDGTextureUAVRef TextureUAV, const FVector4& ClearValues);
+
 RENDERCORE_API void AddClearUAVPass(FRDGBuilder& GraphBuilder, FRDGTextureUAVRef TextureUAV, const float(&ClearValues)[4]);
 
 RENDERCORE_API void AddClearUAVPass(FRDGBuilder& GraphBuilder, FRDGTextureUAVRef TextureUAV, const uint32(&ClearValues)[4]);
diff --git a/Engine/Source/Runtime/Renderer/Private/SystemTextures.cpp b/Engine/Source/Runtime/Renderer/Private/SystemTextures.cpp
--- a/Engine/Source/Runtime/Renderer/Private/SystemTextures.cpp
+++ b/Engine/Source/Runtime/Renderer/Private/SystemTextures.cpp
@@ -891,3 +891,264 @@ FRDGTextureRef FSystemTextures::GetZeroUShort4Dummy(FRDGBuilder& GraphBuilder) c
 	return GraphBuilder.RegisterExternalTexture(ZeroUShort4Dummy, TEXT("ZeroUShort4Dummy"));
 }
 
+//////////////////////////////////////////////////////////////////////////////////////////////////////////
+// Default textures
+
+template<typename TClearType>
+FClearValueBinding GetClearBinding(TClearType Value)
+{
+	return FClearValueBinding::None;
+}
+
+template<>
+FClearValueBinding GetClearBinding(FClearValueBinding Value)
+{
+	return Value;
+}
+
+bool operator !=(const FDefaultTextureKey& A, const FDefaultTextureKey& B)
+{
+	return A.Format != B.Format ||
+		A.Dimension != B.Dimension ||
+		A.ValueAsUInt[0] != B.ValueAsUInt[0] ||
+		A.ValueAsUInt[1] != B.ValueAsUInt[1] ||
+		A.ValueAsUInt[2] != B.ValueAsUInt[2] ||
+		A.ValueAsUInt[3] != B.ValueAsUInt[3];
+}
+
+template<typename T>
+static FDefaultTextureKey GetDefaultTextureKey(EPixelFormat Format, const T& In)
+{
+	FDefaultTextureKey Out;
+	const uint32 Size = sizeof(T);
+	const uint32* InAsUInt = (const uint32*)&In;
+	Out.ValueAsUInt[0] = InAsUInt[0];
+	Out.ValueAsUInt[1] = Size > 4  ? InAsUInt[1] : 0u;
+	Out.ValueAsUInt[2] = Size > 8  ? InAsUInt[2] : 0u;
+	Out.ValueAsUInt[3] = Size > 12 ? InAsUInt[3] : 0u;
+	Out.Format = Format;
+	return Out;
+}
+
+static void AddClearUAVPass(FRDGBuilder& GraphBuilder, FRDGTextureUAVRef TextureUAV, const FClearValueBinding& Value)
+{
+	// Nothing. This is just for compiling template function
+}
+
+template<typename TClearValue>
+FRDGTextureRef GetInternalDefaultTexture(
+	FRDGBuilder& GraphBuilder, 
+	TArray<FDefaultTexture>& DefaultTextures,
+	FHashTable& HashDefaultTextures,
+	ETextureDimension Dimension, 
+	EPixelFormat Format, 
+	TClearValue Value)
+{
+	// Check this is a valid format
+	check(Format != PF_Unknown && Format != PF_MAX && GPixelFormats[Format].BlockSizeX == 1 && GPixelFormats[Format].BlockSizeY == 1 && GPixelFormats[Format].BlockSizeZ == 1);
+
+	const FDefaultTextureKey Key = GetDefaultTextureKey(Format, Value);
+	const uint32 Hash = Murmur32({uint32(Key.Dimension), uint32(Key.Format), Key.ValueAsUInt[0], Key.ValueAsUInt[1], Key.ValueAsUInt[2], Key.ValueAsUInt[3]});
+
+	uint32 Index = HashDefaultTextures.First(Hash);
+	while (HashDefaultTextures.IsValid(Index) && DefaultTextures[Index].Key != Key)
+	{
+		Index = HashDefaultTextures.Next(Index);
+		check(DefaultTextures[Index].Hash == Hash); //Sanitycheck
+	}
+
+	if (HashDefaultTextures.IsValid(Index) && DefaultTextures[Index].Texture != nullptr)
+	{
+		return GraphBuilder.RegisterExternalTexture(DefaultTextures[Index].Texture);
+	}
+
+	bool bSupportUAVWrite = true;
+	bool bSupportGraphicsWrite = true;
+	FClearValueBinding ClearValue = GetClearBinding(Value);
+	if (IsDepthOrStencilFormat(Format))
+	{
+		check(
+			ClearValue == FClearValueBinding::DepthZero || 
+			ClearValue == FClearValueBinding::DepthOne  || 
+			ClearValue == FClearValueBinding::DepthFar  ||
+			ClearValue == FClearValueBinding::DepthNear);
+		bSupportUAVWrite = false;
+	}
+
+	ETextureCreateFlags Flags =
+		TexCreate_ShaderResource |
+		(bSupportUAVWrite ? TexCreate_UAV : TexCreate_None) |
+		(bSupportGraphicsWrite ? TexCreate_RenderTargetable : TexCreate_None);
+
+	FRDGTextureRef Texture = nullptr; 
+	if (Dimension == ETextureDimension::Texture2D)
+	{
+		Texture = GraphBuilder.CreateTexture(FRDGTextureDesc::Create2D(FIntPoint(1, 1), Format, ClearValue, TexCreate_ShaderResource | Flags), TEXT("DefaultTexture"), ERDGTextureFlags::MultiFrame); // TODO Create better name e.g., DefaultTexture(2D,PF_R32INT,Zero)
+	}
+	else if (Dimension == ETextureDimension::Texture2DArray)
+	{
+		Texture = GraphBuilder.CreateTexture(FRDGTextureDesc::Create2DArray(FIntPoint(1, 1), Format, ClearValue, TexCreate_ShaderResource | Flags, 1), TEXT("DefaultTexture"), ERDGTextureFlags::MultiFrame);
+	}
+	else if (Dimension == ETextureDimension::Texture3D)
+	{
+		Texture = GraphBuilder.CreateTexture(FRDGTextureDesc::Create3D(FIntVector(1, 1, 1), Format, ClearValue, TexCreate_ShaderResource | Flags), TEXT("DefaultTexture"), ERDGTextureFlags::MultiFrame);
+	}
+	else if (Dimension == ETextureDimension::TextureCube)
+	{
+		Texture = GraphBuilder.CreateTexture(FRDGTextureDesc::CreateCube(1, Format, ClearValue, TexCreate_ShaderResource | Flags), TEXT("DefaultTexture"), ERDGTextureFlags::MultiFrame);
+	}
+	else if (Dimension == ETextureDimension::TextureCubeArray)
+	{
+		Texture = GraphBuilder.CreateTexture(FRDGTextureDesc::CreateCubeArray(1, Format, ClearValue, TexCreate_ShaderResource | Flags, 1), TEXT("DefaultTexture"), ERDGTextureFlags::MultiFrame);
+	}
+	else
+	{
+		return nullptr;
+	}
+
+	if (ClearValue == FClearValueBinding::None)
+	{
+		check(bSupportUAVWrite);
+		AddClearUAVPass(GraphBuilder, GraphBuilder.CreateUAV(Texture), Value);
+	}
+
+	FDefaultTexture Entry;
+	Entry.Key = Key;
+	Entry.Hash = Hash;
+	Entry.Texture = GraphBuilder.ConvertToExternalTexture(Texture);
+
+	Index = DefaultTextures.Add(Entry);
+	HashDefaultTextures.Add(Hash, Index);
+	return Texture;
+}
+
+//////////////////////////////////////////////////////////////////////////////////////////////////////////
+// Default Buffers
+
+template<typename T>
+static FDefaultBufferKey GetDefaultBufferKey(uint32 NumBytePerElement, bool bIsStructuredBuffer, const T* In)
+{
+	FDefaultBufferKey Out;
+	if (In)
+	{
+		const uint32* InAsUInt = (const uint32*)In;
+		Out.ValueAsUInt[0] = InAsUInt[0];
+		Out.ValueAsUInt[1] = NumBytePerElement > 4  ? InAsUInt[1] : 0u;
+		Out.ValueAsUInt[2] = NumBytePerElement > 8  ? InAsUInt[2] : 0u;
+		Out.ValueAsUInt[3] = NumBytePerElement > 12 ? InAsUInt[3] : 0u;
+	}
+
+	Out.NumBytePerElement	= NumBytePerElement;
+	Out.bIsStructuredBuffer = bIsStructuredBuffer;
+	return Out;
+}
+
+bool operator !=(const FDefaultBufferKey& A, const FDefaultBufferKey& B)
+{
+	return A.NumBytePerElement != B.NumBytePerElement ||
+		A.bIsStructuredBuffer != B.bIsStructuredBuffer ||
+		A.ValueAsUInt[0] != B.ValueAsUInt[0] ||
+		A.ValueAsUInt[1] != B.ValueAsUInt[1] ||
+		A.ValueAsUInt[2] != B.ValueAsUInt[2] ||
+		A.ValueAsUInt[3] != B.ValueAsUInt[3];
+}
+
+template<typename TClearValue>
+FRDGBufferRef GetInternalDefaultBuffer(
+	FRDGBuilder& GraphBuilder, 
+	TArray<FDefaultBuffer>& DefaultBuffers,
+	FHashTable& HashDefaultBuffers,
+	uint32 NumBytePerElement,
+	bool bIsStructuredBuffer,
+	const TClearValue* Value)
+{
+	// Buffer key
+	const uint32 NumElements = 1;
+	const FDefaultBufferKey Key = GetDefaultBufferKey(NumBytePerElement, bIsStructuredBuffer, Value);
+	const uint32 Hash = Murmur32({uint32(Key.bIsStructuredBuffer ? 0x20000000u : 0x10000000u) | Key.NumBytePerElement, Key.ValueAsUInt[0], Key.ValueAsUInt[1], Key.ValueAsUInt[2], Key.ValueAsUInt[3] });
+
+	// Find exsting buffer ("fast" path)
+	uint32 Index = HashDefaultBuffers.First(Hash);
+	while (HashDefaultBuffers.IsValid(Index) && DefaultBuffers[Index].Key != Key)
+	{
+		Index = HashDefaultBuffers.Next(Index);
+		check(DefaultBuffers[Index].Hash == Hash); //Sanitycheck
+	}
+
+	if (HashDefaultBuffers.IsValid(Index) && DefaultBuffers[Index].Buffer != nullptr)
+	{
+		return GraphBuilder.RegisterExternalBuffer(DefaultBuffers[Index].Buffer);
+	}
+
+	// Adding new buffer if there is no fit (slow path)
+	FRDGBufferRef Buffer = nullptr; 
+	if (bIsStructuredBuffer)
+	{
+		Buffer = GraphBuilder.CreateBuffer(FRDGBufferDesc::CreateStructuredDesc(NumBytePerElement, NumElements), TEXT("DefaultBuffer"), ERDGBufferFlags::MultiFrame);
+	}
+	else
+	{
+		Buffer = GraphBuilder.CreateBuffer(FRDGBufferDesc::CreateBufferDesc(NumBytePerElement, NumElements), TEXT("DefaultStructuredBuffer"), ERDGBufferFlags::MultiFrame);
+	}
+
+	// Initialize the entire buffer with the provided data
+	if (Value)
+	{
+		AddBufferUploadPass(GraphBuilder, Buffer, Value, NumElements * NumBytePerElement, ERDGInitialDataFlags::None);
+	}
+	// Initialize buffer to 0
+	else
+	{
+		TArray<uint8> DefaultValue;
+		DefaultValue.Init(0u, NumElements * NumBytePerElement);
+		AddBufferUploadPass(GraphBuilder, Buffer, DefaultValue.GetData(), DefaultValue.Num(), ERDGInitialDataFlags::None);
+	}
+
+	FDefaultBuffer Entry;
+	Entry.Key = Key;
+	Entry.Hash = Hash;
+	Entry.Buffer = GraphBuilder.ConvertToExternalBuffer(Buffer);
+
+	Index = DefaultBuffers.Add(Entry);
+	HashDefaultBuffers.Add(Hash, Index);
+	return Buffer;
+}
+
+//////////////////////////////////////////////////////////////////////////////////////////////////////////////////////////////////////////////////////////////////////////////////////////////////////////////////////////////////////////////////////////////////////////////////////////////////////////////////
+// Textures 
+
+FRDGTextureRef FSystemTextures::GetDefaultTexture2D(FRDGBuilder& GraphBuilder, EPixelFormat Format, float Value)												{ return GetInternalDefaultTexture(GraphBuilder, DefaultTextures, HashDefaultTextures, ETextureDimension::Texture2D, Format, Value); }
+FRDGTextureRef FSystemTextures::GetDefaultTexture2D(FRDGBuilder& GraphBuilder, EPixelFormat Format, uint32 Value)												{ return GetInternalDefaultTexture(GraphBuilder, DefaultTextures, HashDefaultTextures, ETextureDimension::Texture2D, Format, Value); }
+FRDGTextureRef FSystemTextures::GetDefaultTexture2D(FRDGBuilder& GraphBuilder, EPixelFormat Format, const FVector& Value)										{ return GetInternalDefaultTexture(GraphBuilder, DefaultTextures, HashDefaultTextures, ETextureDimension::Texture2D, Format, Value); }
+FRDGTextureRef FSystemTextures::GetDefaultTexture2D(FRDGBuilder& GraphBuilder, EPixelFormat Format, const FVector4& Value)										{ return GetInternalDefaultTexture(GraphBuilder, DefaultTextures, HashDefaultTextures, ETextureDimension::Texture2D, Format, Value); }
+FRDGTextureRef FSystemTextures::GetDefaultTexture2D(FRDGBuilder& GraphBuilder, EPixelFormat Format, const FUintVector4& Value)									{ return GetInternalDefaultTexture(GraphBuilder, DefaultTextures, HashDefaultTextures, ETextureDimension::Texture2D, Format, Value); }
+FRDGTextureRef FSystemTextures::GetDefaultTexture2D(FRDGBuilder& GraphBuilder, EPixelFormat Format, const FClearValueBinding& Value)							{ return GetInternalDefaultTexture(GraphBuilder, DefaultTextures, HashDefaultTextures, ETextureDimension::Texture2D, Format, Value); }
+
+FRDGTextureRef FSystemTextures::GetDefaultTexture(FRDGBuilder& GraphBuilder, ETextureDimension Dimension, EPixelFormat Format, float Value)						{ return GetInternalDefaultTexture(GraphBuilder, DefaultTextures, HashDefaultTextures, Dimension, Format, Value); }
+FRDGTextureRef FSystemTextures::GetDefaultTexture(FRDGBuilder& GraphBuilder, ETextureDimension Dimension, EPixelFormat Format, uint32 Value)					{ return GetInternalDefaultTexture(GraphBuilder, DefaultTextures, HashDefaultTextures, Dimension, Format, Value); }
+FRDGTextureRef FSystemTextures::GetDefaultTexture(FRDGBuilder& GraphBuilder, ETextureDimension Dimension, EPixelFormat Format, const FVector2D& Value)			{ return GetInternalDefaultTexture(GraphBuilder, DefaultTextures, HashDefaultTextures, Dimension, Format, Value); }
+FRDGTextureRef FSystemTextures::GetDefaultTexture(FRDGBuilder& GraphBuilder, ETextureDimension Dimension, EPixelFormat Format, const FIntPoint& Value)			{ return GetInternalDefaultTexture(GraphBuilder, DefaultTextures, HashDefaultTextures, Dimension, Format, Value); }
+FRDGTextureRef FSystemTextures::GetDefaultTexture(FRDGBuilder& GraphBuilder, ETextureDimension Dimension, EPixelFormat Format, const FVector& Value)			{ return GetInternalDefaultTexture(GraphBuilder, DefaultTextures, HashDefaultTextures, Dimension, Format, Value); }
+FRDGTextureRef FSystemTextures::GetDefaultTexture(FRDGBuilder& GraphBuilder, ETextureDimension Dimension, EPixelFormat Format, const FVector4& Value)			{ return GetInternalDefaultTexture(GraphBuilder, DefaultTextures, HashDefaultTextures, Dimension, Format, Value); }
+FRDGTextureRef FSystemTextures::GetDefaultTexture(FRDGBuilder& GraphBuilder, ETextureDimension Dimension, EPixelFormat Format, const FUintVector4& Value)		{ return GetInternalDefaultTexture(GraphBuilder, DefaultTextures, HashDefaultTextures, Dimension, Format, Value); }
+FRDGTextureRef FSystemTextures::GetDefaultTexture(FRDGBuilder& GraphBuilder, ETextureDimension Dimension, EPixelFormat Format, const FClearValueBinding& Value)	{ return GetInternalDefaultTexture(GraphBuilder, DefaultTextures, HashDefaultTextures, Dimension, Format, Value); }
+
+//////////////////////////////////////////////////////////////////////////////////////////////////////////////////////////////////////////////////////////////////////////////////////////////////////////////////////////////////////////////////////////////////////////////////////////////////////////////////
+// Buffers
+
+// Default init to 0
+FRDGBufferRef FSystemTextures::GetDefaultBuffer(FRDGBuilder& GraphBuilder, uint32 NumBytePerElement)										{ return GetInternalDefaultBuffer(GraphBuilder, DefaultBuffers, HashDefaultBuffers, NumBytePerElement, false, (uint32*)nullptr); }
+FRDGBufferRef FSystemTextures::GetDefaultStructuredBuffer(FRDGBuilder& GraphBuilder, uint32 NumBytePerElement)								{ return GetInternalDefaultBuffer(GraphBuilder, DefaultBuffers, HashDefaultBuffers, NumBytePerElement, true,  (uint32*)nullptr); }
+
+// Default value of an element
+FRDGBufferRef FSystemTextures::GetDefaultBuffer(FRDGBuilder& GraphBuilder, uint32 NumBytePerElement, float Value)							{ return GetInternalDefaultBuffer(GraphBuilder, DefaultBuffers, HashDefaultBuffers, NumBytePerElement, false/* Vertex */, &Value); }
+FRDGBufferRef FSystemTextures::GetDefaultBuffer(FRDGBuilder& GraphBuilder, uint32 NumBytePerElement, uint32 Value)							{ return GetInternalDefaultBuffer(GraphBuilder, DefaultBuffers, HashDefaultBuffers, NumBytePerElement, false/* Vertex */, &Value); }
+FRDGBufferRef FSystemTextures::GetDefaultBuffer(FRDGBuilder& GraphBuilder, uint32 NumBytePerElement, const FVector& Value)					{ return GetInternalDefaultBuffer(GraphBuilder, DefaultBuffers, HashDefaultBuffers, NumBytePerElement, false/* Vertex */, &Value); }
+FRDGBufferRef FSystemTextures::GetDefaultBuffer(FRDGBuilder& GraphBuilder, uint32 NumBytePerElement, const FVector4& Value)					{ return GetInternalDefaultBuffer(GraphBuilder, DefaultBuffers, HashDefaultBuffers, NumBytePerElement, false/* Vertex */, &Value); }
+FRDGBufferRef FSystemTextures::GetDefaultBuffer(FRDGBuilder& GraphBuilder, uint32 NumBytePerElement, const FUintVector4& Value)				{ return GetInternalDefaultBuffer(GraphBuilder, DefaultBuffers, HashDefaultBuffers, NumBytePerElement, false/* Vertex */, &Value); }
+
+FRDGBufferRef FSystemTextures::GetDefaultStructuredBuffer(FRDGBuilder& GraphBuilder, uint32 NumBytePerElement, float Value)					{ return GetInternalDefaultBuffer(GraphBuilder, DefaultBuffers, HashDefaultBuffers, NumBytePerElement, true /* Structured */, &Value); }
+FRDGBufferRef FSystemTextures::GetDefaultStructuredBuffer(FRDGBuilder& GraphBuilder, uint32 NumBytePerElement, uint32 Value)				{ return GetInternalDefaultBuffer(GraphBuilder, DefaultBuffers, HashDefaultBuffers, NumBytePerElement, true /* Structured */, &Value); }
+FRDGBufferRef FSystemTextures::GetDefaultStructuredBuffer(FRDGBuilder& GraphBuilder, uint32 NumBytePerElement, const FVector& Value)		{ return GetInternalDefaultBuffer(GraphBuilder, DefaultBuffers, HashDefaultBuffers, NumBytePerElement, true /* Structured */, &Value); }
+FRDGBufferRef FSystemTextures::GetDefaultStructuredBuffer(FRDGBuilder& GraphBuilder, uint32 NumBytePerElement, const FVector4& Value)		{ return GetInternalDefaultBuffer(GraphBuilder, DefaultBuffers, HashDefaultBuffers, NumBytePerElement, true /* Structured */, &Value); }
+FRDGBufferRef FSystemTextures::GetDefaultStructuredBuffer(FRDGBuilder& GraphBuilder, uint32 NumBytePerElement, const FUintVector4& Value)	{ return GetInternalDefaultBuffer(GraphBuilder, DefaultBuffers, HashDefaultBuffers, NumBytePerElement, true /* Structured */, &Value); }
diff --git a/Engine/Source/Runtime/Renderer/Private/SystemTextures.h b/Engine/Source/Runtime/Renderer/Private/SystemTextures.h
--- a/Engine/Source/Runtime/Renderer/Private/SystemTextures.h
+++ b/Engine/Source/Runtime/Renderer/Private/SystemTextures.h
@@ -34,6 +34,33 @@ struct FRDGSystemTextures
 	FRDGTextureSRVRef StencilDummySRV{};
 };
 
+
+struct FDefaultTextureKey
+{
+	uint32 ValueAsUInt[4] = { 0u,0u,0u,0u };
+	EPixelFormat Format = PF_Unknown;
+	ETextureDimension Dimension = ETextureDimension::Texture2D;
+};
+struct FDefaultTexture
+{
+	uint32 Hash = 0;
+	FDefaultTextureKey Key;
+	TRefCountPtr<IPooledRenderTarget> Texture;
+};
+
+struct FDefaultBufferKey
+{
+	uint32 ValueAsUInt[4] = { 0u,0u,0u,0u };
+	uint32 NumBytePerElement = 0;
+	bool bIsStructuredBuffer = false;
+};
+struct FDefaultBuffer
+{
+	uint32 Hash = 0;
+	FDefaultBufferKey Key;
+	TRefCountPtr<FRDGPooledBuffer> Buffer;
+};
+
 /**
  * Encapsulates the system textures used for scene rendering.
  */
@@ -114,6 +141,7 @@ public:
 	// SRV for StencilDummy Texture.
 	TRefCountPtr<FRHIShaderResourceView> StencilDummySRV;
 
+	// Create simple default texture
 	FRDGTextureRef GetWhiteDummy(FRDGBuilder& GraphBuilder) const;
 	FRDGTextureRef GetBlackDummy(FRDGBuilder& GraphBuilder) const;
 	FRDGTextureRef GetZeroUIntDummy(FRDGBuilder& GraphBuilder) const;
@@ -127,10 +155,50 @@ public:
 	FRDGTextureRef GetMidGreyDummy(FRDGBuilder& GraphBuilder) const;
 	FRDGTextureRef GetVolumetricBlackDummy(FRDGBuilder& GraphBuilder) const;
 
+	// Create default 2D texture (1x1) with specific format and initialize value 
+	FRDGTextureRef GetDefaultTexture2D(FRDGBuilder& GraphBuilder, EPixelFormat Format, float Value);
+	FRDGTextureRef GetDefaultTexture2D(FRDGBuilder& GraphBuilder, EPixelFormat Format, uint32 Value);
+	FRDGTextureRef GetDefaultTexture2D(FRDGBuilder& GraphBuilder, EPixelFormat Format, const FVector& Value);
+	FRDGTextureRef GetDefaultTexture2D(FRDGBuilder& GraphBuilder, EPixelFormat Format, const FVector4& Value);
+	FRDGTextureRef GetDefaultTexture2D(FRDGBuilder& GraphBuilder, EPixelFormat Format, const FUintVector4& Value);
+	FRDGTextureRef GetDefaultTexture2D(FRDGBuilder& GraphBuilder, EPixelFormat Format, const FClearValueBinding& Value);
+
+	// Create default 2D/3D/Cube/Array texture (1x1) with specific format and initialize value 
+	FRDGTextureRef GetDefaultTexture(FRDGBuilder& GraphBuilder, ETextureDimension Dimension, EPixelFormat Format, float Value);
+	FRDGTextureRef GetDefaultTexture(FRDGBuilder& GraphBuilder, ETextureDimension Dimension, EPixelFormat Format, uint32 Value);
+	FRDGTextureRef GetDefaultTexture(FRDGBuilder& GraphBuilder, ETextureDimension Dimension, EPixelFormat Format, const FVector2D& Value);
+	FRDGTextureRef GetDefaultTexture(FRDGBuilder& GraphBuilder, ETextureDimension Dimension, EPixelFormat Format, const FIntPoint& Value);
+	FRDGTextureRef GetDefaultTexture(FRDGBuilder& GraphBuilder, ETextureDimension Dimension, EPixelFormat Format, const FVector& Value);
+	FRDGTextureRef GetDefaultTexture(FRDGBuilder& GraphBuilder, ETextureDimension Dimension, EPixelFormat Format, const FVector4& Value);
+	FRDGTextureRef GetDefaultTexture(FRDGBuilder& GraphBuilder, ETextureDimension Dimension, EPixelFormat Format, const FUintVector4& Value);
+	FRDGTextureRef GetDefaultTexture(FRDGBuilder& GraphBuilder, ETextureDimension Dimension, EPixelFormat Format, const FClearValueBinding& Value);
+
+	// Create default buffer initialize to zero.
+	FRDGBufferRef GetDefaultBuffer(FRDGBuilder& GraphBuilder, uint32 NumBytePerElement);
+	FRDGBufferRef GetDefaultStructuredBuffer(FRDGBuilder& GraphBuilder, uint32 NumBytePerElement);
+
+	// Create a default buffer initialized with a reference element.
+	FRDGBufferRef GetDefaultBuffer(FRDGBuilder& GraphBuilder, uint32 NumBytePerElement, float Value);
+	FRDGBufferRef GetDefaultBuffer(FRDGBuilder& GraphBuilder, uint32 NumBytePerElement, uint32 Value);
+	FRDGBufferRef GetDefaultBuffer(FRDGBuilder& GraphBuilder, uint32 NumBytePerElement, const FVector& Value);
+	FRDGBufferRef GetDefaultBuffer(FRDGBuilder& GraphBuilder, uint32 NumBytePerElement, const FVector4& Value);
+	FRDGBufferRef GetDefaultBuffer(FRDGBuilder& GraphBuilder, uint32 NumBytePerElement, const FUintVector4& Value);
+	FRDGBufferRef GetDefaultStructuredBuffer(FRDGBuilder& GraphBuilder, uint32 NumBytePerElement, float Value);
+	FRDGBufferRef GetDefaultStructuredBuffer(FRDGBuilder& GraphBuilder, uint32 NumBytePerElement, uint32 Value);
+	FRDGBufferRef GetDefaultStructuredBuffer(FRDGBuilder& GraphBuilder, uint32 NumBytePerElement, const FVector& Value);
+	FRDGBufferRef GetDefaultStructuredBuffer(FRDGBuilder& GraphBuilder, uint32 NumBytePerElement, const FVector4& Value);
+	FRDGBufferRef GetDefaultStructuredBuffer(FRDGBuilder& GraphBuilder, uint32 NumBytePerElement, const FUintVector4& Value);
+
 protected:
 	/** Maximum feature level that the textures have been initialized up to */
 	ERHIFeatureLevel::Type FeatureLevelInitializedTo;
 
+	/** Default textures allocated on-demand */
+	TArray<FDefaultTexture> DefaultTextures;
+	TArray<FDefaultBuffer> DefaultBuffers;
+	FHashTable HashDefaultTextures;
+	FHashTable HashDefaultBuffers;
+
 	void InitializeCommonTextures(FRHICommandListImmediate& RHICmdList);
 	void InitializeFeatureLevelDependentTextures(FRHICommandListImmediate& RHICmdList, const ERHIFeatureLevel::Type InFeatureLevel);
 };
