Change some hair resources to use the new default resource API.

diff --git a/Engine/Source/Runtime/Renderer/Private/HairStrands/HairStrandsRendering.cpp b/Engine/Source/Runtime/Renderer/Private/HairStrands/HairStrandsRendering.cpp
--- a/Engine/Source/Runtime/Renderer/Private/HairStrands/HairStrandsRendering.cpp
+++ b/Engine/Source/Runtime/Renderer/Private/HairStrands/HairStrandsRendering.cpp
@@ -31,8 +31,7 @@ static TRDGUniformBufferRef<FHairStrandsViewUniformParameters> InternalCreateHai
 		}
 		else
 		{
-			FRDGBufferRef DummyBuffer = GraphBuilder.CreateBuffer(FRDGBufferDesc::CreateBufferDesc(4, 1), TEXT("Hair.DummyBuffer"));
-			AddClearUAVPass(GraphBuilder, GraphBuilder.CreateUAV(DummyBuffer, PF_R16G16_UINT), 0);
+			FRDGBufferRef DummyBuffer = GSystemTextures.GetDefaultBuffer(GraphBuilder, 4);
 			FRDGBufferSRVRef DummyBufferR32SRV = GraphBuilder.CreateSRV(DummyBuffer, PF_R32_UINT);
 			FRDGBufferSRVRef DummyBufferRG16SRV = GraphBuilder.CreateSRV(DummyBuffer, PF_R16G16_UINT);
 			Parameters->HairTileData = DummyBufferRG16SRV;
@@ -50,16 +49,13 @@ static TRDGUniformBufferRef<FHairStrandsViewUniformParameters> InternalCreateHai
 	}
 	else
 	{
-		FRDGBufferRef DummyBuffer = GraphBuilder.CreateBuffer(FRDGBufferDesc::CreateBufferDesc(4,1), TEXT("Hair.DummyBuffer"));
-		FRDGBufferRef DummyNodeBuffer = GraphBuilder.CreateBuffer(FRDGBufferDesc::CreateStructuredDesc(20, 1), TEXT("Hair.DummyNodeBuffer"));
+		FRDGBufferRef DummyBuffer = GSystemTextures.GetDefaultBuffer(GraphBuilder, 4);
+		FRDGBufferRef DummyNodeBuffer = GSystemTextures.GetDefaultStructuredBuffer(GraphBuilder, 20);
 
-		FRDGTextureRef BlackTexture = GSystemTextures.GetBlackDummy(GraphBuilder);
-		FRDGTextureRef ZeroR32_UINT = GSystemTextures.GetZeroUIntDummy(GraphBuilder);
-		FRDGTextureRef ZeroRGBA16_UINT = GSystemTextures.GetZeroUShort4Dummy(GraphBuilder);
-		FRDGTextureRef FarDepth = GSystemTextures.GetDepthDummy(GraphBuilder);
-
-		AddClearUAVPass(GraphBuilder, GraphBuilder.CreateUAV(DummyNodeBuffer), 0);
-		AddClearUAVPass(GraphBuilder, GraphBuilder.CreateUAV(DummyBuffer, PF_R16G16_UINT), 0);
+		FRDGTextureRef BlackTexture 	= GSystemTextures.GetBlackDummy(GraphBuilder);
+		FRDGTextureRef ZeroR32_UINT 	= GSystemTextures.GetDefaultTexture2D(GraphBuilder, PF_R32_UINT, 0u);
+		FRDGTextureRef ZeroRGBA16_UINT 	= GSystemTextures.GetDefaultTexture2D(GraphBuilder, PF_R16G16_UINT, 0u);
+		FRDGTextureRef FarDepth 		= GSystemTextures.GetDepthDummy(GraphBuilder);
 
 		FRDGBufferSRVRef DummyBufferR32SRV = GraphBuilder.CreateSRV(DummyBuffer, PF_R32_UINT);
 		FRDGBufferSRVRef DummyBufferRG16SRV = GraphBuilder.CreateSRV(DummyBuffer, PF_R16G16_UINT);
@@ -202,4 +198,4 @@ bool HasViewHairStrandsData(const TArrayView<FViewInfo>& Views)
 	return false;
 }
 
-}
\ No newline at end of file
+}
