When resizing VT page table, make sure texture size is large enough that lower mip levels of non-power-2 VTs will also fit

diff --git a/Engine/Source/Runtime/Renderer/Private/VT/VirtualTextureAllocator.cpp b/Engine/Source/Runtime/Renderer/Private/VT/VirtualTextureAllocator.cpp
--- a/Engine/Source/Runtime/Renderer/Private/VT/VirtualTextureAllocator.cpp
+++ b/Engine/Source/Runtime/Renderer/Private/VT/VirtualTextureAllocator.cpp
@@ -469,8 +469,13 @@ uint32 FVirtualTextureAllocator::Alloc(FAllocatedVirtualTexture* VT)
 
 		check(AddressBlocks[AllocIndex].State != EBlockState::FreeList);
 
-		AllocatedWidth = FMath::Max(AllocatedWidth, vTileX + WidthInTiles);
-		AllocatedHeight = FMath::Max(AllocatedHeight, vTileY + HeightInTiles);
+		// Make sure we allocate enough space in the backing texture so all the mip levels fit
+		const uint32 SizeAlign = 1u << MaxLevel;
+		const uint32 AlignedWidthInTiles = Align(WidthInTiles, SizeAlign);
+		const uint32 AlignedHeightInTiles = Align(HeightInTiles, SizeAlign);
+
+		AllocatedWidth = FMath::Max(AllocatedWidth, vTileX + AlignedWidthInTiles);
+		AllocatedHeight = FMath::Max(AllocatedHeight, vTileY + AlignedHeightInTiles);
 	}
 
 	return vAddress;
