Remove VT UDIM page table wrapping
- It's not possible to make wrapping work correctly for lower mip levels, depending on the sizes of the VTs in the stack
- Now UDIMs that are smaller than stack dimensions will just map the lowest resolution mip to out-of-bounds addresses
- Entire stack will continue to wrap, based on the stack dimensions
- End result is UDIM wrapping should no longer by relied on

diff --git a/Engine/Source/Runtime/Engine/Private/VT/UploadingVirtualTexture.cpp b/Engine/Source/Runtime/Engine/Private/VT/UploadingVirtualTexture.cpp
--- a/Engine/Source/Runtime/Engine/Private/VT/UploadingVirtualTexture.cpp
+++ b/Engine/Source/Runtime/Engine/Private/VT/UploadingVirtualTexture.cpp
@@ -97,17 +97,18 @@ FUploadingVirtualTexture::~FUploadingVirtualTexture()
 	}
 }
 
-uint32 FUploadingVirtualTexture::GetLocalMipBias(uint8 vLevel, uint32 vAddress) const
+uint32 FUploadingVirtualTexture::GetLocalMipBias(uint8 vLevel, uint32 vAddress, uint8 MaxLevel) const
 {
 	const uint32 NumMips = Data->NumMips;
-	uint32 NumNonResidentLevels = 0u;
-	while (vLevel < NumMips)
+	uint32 Current_vLevel = vLevel;
+	uint32 Current_vAddress = vAddress;
+	while (Current_vLevel <= MaxLevel)
 	{
-		const uint32 TileIndex = Data->GetTileIndex(vLevel, vAddress);
+		const uint32 TileIndex = Data->GetTileIndex(Current_vLevel, Current_vAddress);
 		if (TileIndex == ~0u)
 		{
-			// vAddress is out-of-bounds for the given producer, this is a 
-			NumNonResidentLevels += (NumMips - vLevel);
+			// vAddress is out-of-bounds for the given producer
+			Current_vLevel = MaxLevel;
 			break;
 		}
 
@@ -117,12 +118,11 @@ uint32 FUploadingVirtualTexture::GetLocalMipBias(uint8 vLevel, uint32 vAddress)
 			break;
 		}
 
-		++NumNonResidentLevels;
-		++vLevel;
-		vAddress >>= 2;
+		Current_vLevel++;
+		Current_vAddress >>= 2;
 	}
 
-	return NumNonResidentLevels;
+	return Current_vLevel - vLevel;
 }
 
 FVTRequestPageResult FUploadingVirtualTexture::RequestPageData(const FVirtualTextureProducerHandle& ProducerHandle, uint8 LayerMask, uint8 vLevel, uint64 vAddress, EVTRequestPagePriority Priority)
diff --git a/Engine/Source/Runtime/Engine/Private/VT/UploadingVirtualTexture.h b/Engine/Source/Runtime/Engine/Private/VT/UploadingVirtualTexture.h
--- a/Engine/Source/Runtime/Engine/Private/VT/UploadingVirtualTexture.h
+++ b/Engine/Source/Runtime/Engine/Private/VT/UploadingVirtualTexture.h
@@ -76,7 +76,7 @@ public:
 	virtual ~FUploadingVirtualTexture();
 
 	// IVirtualTexture interface
-	virtual uint32 GetLocalMipBias(uint8 vLevel, uint32 vAddress) const override;
+	virtual uint32 GetLocalMipBias(uint8 vLevel, uint32 vAddress, uint8 MaxLevel) const override;
 	virtual FVTRequestPageResult RequestPageData(const FVirtualTextureProducerHandle& ProducerHandle, uint8 LayerMask, uint8 vLevel, uint64 vAddress, EVTRequestPagePriority Priority) override;
 	virtual IVirtualTextureFinalizer* ProducePageData(FRHICommandListImmediate& RHICmdList,
 		ERHIFeatureLevel::Type FeatureLevel,
diff --git a/Engine/Source/Runtime/RenderCore/Public/VirtualTexturing.h b/Engine/Source/Runtime/RenderCore/Public/VirtualTexturing.h
--- a/Engine/Source/Runtime/RenderCore/Public/VirtualTexturing.h
+++ b/Engine/Source/Runtime/RenderCore/Public/VirtualTexturing.h
@@ -259,7 +259,7 @@ public:
 	 * @param vAddress Virtual address to check
 	 * @return Mip bias to be added to vLevel to reach a resident page at the given address
 	 */
-	virtual uint32 GetLocalMipBias(uint8 vLevel, uint32 vAddress) const { return 0u; }
+	virtual uint32 GetLocalMipBias(uint8 vLevel, uint32 vAddress, uint8 MaxLevel) const { return 0u; }
 
 	/**
 	* Makes a request for the given page data.
diff --git a/Engine/Source/Runtime/Renderer/Private/VT/VirtualTextureShared.h b/Engine/Source/Runtime/Renderer/Private/VT/VirtualTextureShared.h
--- a/Engine/Source/Runtime/Renderer/Private/VT/VirtualTextureShared.h
+++ b/Engine/Source/Runtime/Renderer/Private/VT/VirtualTextureShared.h
@@ -44,7 +44,7 @@ struct FPageTableUpdate
 	inline void Check(uint8 vDimensions)
 	{
 		const uint32 LowBitMask = (1u << (vDimensions * vLogSize)) - 1;
-		checkSlow((vAddress & LowBitMask) == 0);
+		check((vAddress & LowBitMask) == 0);
 		//checkSlow(vLogSize <= vLevel);
 	}
 };
diff --git a/Engine/Source/Runtime/Renderer/Private/VT/VirtualTextureSystem.cpp b/Engine/Source/Runtime/Renderer/Private/VT/VirtualTextureSystem.cpp
--- a/Engine/Source/Runtime/Renderer/Private/VT/VirtualTextureSystem.cpp
+++ b/Engine/Source/Runtime/Renderer/Private/VT/VirtualTextureSystem.cpp
@@ -439,8 +439,8 @@ IAllocatedVirtualTexture* FVirtualTextureSystem::AllocateVirtualTexture(const FA
 
 	uint32 BlockWidthInTiles = 0u;
 	uint32 BlockHeightInTiles = 0u;
-	uint32 MinWidthInBlocks = ~0u;
-	uint32 MinHeightInBlocks = ~0u;
+	uint32 WidthInBlocks = 0u;
+	uint32 HeightInBlocks = 0u;
 	uint32 DepthInTiles = 0u;
 	bool bSupport16BitPageTable = true;
 	FVirtualTextureProducer* ProducerForLayer[VIRTUALTEXTURE_SPACE_MAXLAYERS] = { nullptr };
@@ -454,8 +454,8 @@ IAllocatedVirtualTexture* FVirtualTextureSystem::AllocateVirtualTexture(const FA
 			const FVTProducerDescription& ProducerDesc = Producer->GetDescription();
 			BlockWidthInTiles = FMath::Max(BlockWidthInTiles, ProducerDesc.BlockWidthInTiles);
 			BlockHeightInTiles = FMath::Max(BlockHeightInTiles, ProducerDesc.BlockHeightInTiles);
-			MinWidthInBlocks = FMath::Min<uint32>(MinWidthInBlocks, ProducerDesc.WidthInBlocks);
-			MinHeightInBlocks = FMath::Min<uint32>(MinHeightInBlocks, ProducerDesc.HeightInBlocks);
+			WidthInBlocks = FMath::Max<uint32>(WidthInBlocks, ProducerDesc.WidthInBlocks);
+			HeightInBlocks = FMath::Max<uint32>(HeightInBlocks, ProducerDesc.HeightInBlocks);
 			DepthInTiles = FMath::Max(DepthInTiles, ProducerDesc.DepthInTiles);
 
 			uint32 ProducerLayerIndex = Desc.ProducerLayerIndex[LayerIndex];
@@ -472,55 +472,8 @@ IAllocatedVirtualTexture* FVirtualTextureSystem::AllocateVirtualTexture(const FA
 	check(BlockWidthInTiles > 0u);
 	check(BlockHeightInTiles > 0u);
 	check(DepthInTiles > 0u);
-
-	// Find a block width that is evenly divided by all layers (least common multiple)
-	// Start with min size, then increment by min size until a valid size is found
-	uint32 WidthInBlocks = MinWidthInBlocks;
-	{
-		bool bFoundValidWidthInBlocks = false;
-		while (!bFoundValidWidthInBlocks)
-		{
-			bFoundValidWidthInBlocks = true;
-			for (uint32 LayerIndex = 0u; LayerIndex < Desc.NumTextureLayers; ++LayerIndex)
-			{
-				const FVirtualTextureProducer* Producer = ProducerForLayer[LayerIndex];
-				if (Producer)
-				{
-					if ((WidthInBlocks % Producer->GetDescription().WidthInBlocks) != 0u)
-					{
-						WidthInBlocks += MinWidthInBlocks;
-						check(WidthInBlocks > MinWidthInBlocks); // check for overflow
-						bFoundValidWidthInBlocks = false;
-						break;
-					}
-				}
-			}
-		}
-	}
-
-	// Same thing for height
-	uint32 HeightInBlocks = MinHeightInBlocks;
-	{
-		bool bFoundValidHeightInBlocks = false;
-		while (!bFoundValidHeightInBlocks)
-		{
-			bFoundValidHeightInBlocks = true;
-			for (uint32 LayerIndex = 0u; LayerIndex < Desc.NumTextureLayers; ++LayerIndex)
-			{
-				const FVirtualTextureProducer* Producer = ProducerForLayer[LayerIndex];
-				if (Producer)
-				{
-					if ((HeightInBlocks % Producer->GetDescription().HeightInBlocks) != 0u)
-					{
-						HeightInBlocks += MinHeightInBlocks;
-						check(HeightInBlocks > MinHeightInBlocks); // check for overflow
-						bFoundValidHeightInBlocks = false;
-						break;
-					}
-				}
-			}
-		}
-	}
+	check(WidthInBlocks > 0u);
+	check(HeightInBlocks > 0u);
 
 	// Sum the total number of physical groups from all producers
 	uint32 NumPhysicalGroups = 0;
@@ -1574,7 +1527,12 @@ void FVirtualTextureSystem::GatherRequestsTask(const FGatherRequestsParameters&
 		}
 
 		const uint32 vDimensions = Space->GetDimensions();
+		const uint32 AllocatedPageX = AllocatedVT->GetVirtualPageX();
+		const uint32 AllocatedPageY = AllocatedVT->GetVirtualPageY();
+
 		check(vAddress >= AllocatedVT->GetVirtualAddress());
+		check(vPageX >= AllocatedPageX);
+		check(vPageY >= AllocatedPageY);
 
 		for (uint32 ProducerIndex = 0u; ProducerIndex < NumUniqueProducers; ++ProducerIndex)
 		{
@@ -1592,6 +1550,7 @@ void FVirtualTextureSystem::GatherRequestsTask(const FGatherRequestsParameters&
 			}
 
 			const uint32 ProducerMipBias = AllocatedVT->GetUniqueProducerMipBias(ProducerIndex);
+			const uint32 MaxLocalLevel = MaxLevel - ProducerMipBias;
 
 			// here vLevel is clamped against ProducerMipBias, as ProducerMipBias represents the most detailed level of this producer, relative to the allocated VT
 			// used to rescale vAddress to the correct tile within the given mip level
@@ -1601,33 +1560,11 @@ void FVirtualTextureSystem::GatherRequestsTask(const FGatherRequestsParameters&
 			// here we subtract ProducerMipBias, which effectively matches more detailed mips of lower resolution producers with less detailed mips of higher resolution producers
 			uint32 Local_vLevel = Mapping_vLevel - ProducerMipBias;
 
-			// Wrap vAddress for the given producer
-			uint32 Wrapped_vAddress = vAddress;
-			uint32 Local_vAddress;
-			{
-				// Scale size of producer up to be relative to size of entire allocated VT
-				const uint32 ProducerScaleFactor = (1u << ProducerMipBias);
-				const uint32 ProducerWidthInPages = Producer->GetWidthInTiles() * ProducerScaleFactor;
-				const uint32 ProducerHeightInPages = Producer->GetHeightInTiles() * ProducerScaleFactor;
-				const uint32 AllocatedPageX = AllocatedVT->GetVirtualPageX();
-				const uint32 AllocatedPageY = AllocatedVT->GetVirtualPageY();
-				check(vPageX >= AllocatedPageX);
-				check(vPageY >= AllocatedPageY);
-
-				uint32 Local_vPageX = vPageX - AllocatedPageX;
-				uint32 Local_vPageY = vPageY - AllocatedPageY;
-				if (Local_vPageX >= ProducerWidthInPages || Local_vPageY >= ProducerHeightInPages)
-				{
-					Local_vPageX %= ProducerWidthInPages;
-					Local_vPageY %= ProducerHeightInPages;
-					Wrapped_vAddress = FMath::MortonCode2(AllocatedPageX + Local_vPageX) | (FMath::MortonCode2(AllocatedPageY + Local_vPageY) << 1);
-
-					Wrapped_vAddress &= (0xffffffff << (vLevel * vDimensions));
-				}
-				Local_vAddress = FMath::MortonCode2(Local_vPageX >> Mapping_vLevel) | (FMath::MortonCode2(Local_vPageY >> Mapping_vLevel) << 1);
-			}
+			const uint32 Local_vPageX = (vPageX - AllocatedPageX) >> Mapping_vLevel;
+			const uint32 Local_vPageY = (vPageY - AllocatedPageY) >> Mapping_vLevel;
+			uint32 Local_vAddress = FMath::MortonCode2(Local_vPageX) | (FMath::MortonCode2(Local_vPageY) << 1);
 
-			const uint32 LocalMipBias = Producer->GetVirtualTexture()->GetLocalMipBias(Local_vLevel, Local_vAddress);
+			const uint32 LocalMipBias = Producer->GetVirtualTexture()->GetLocalMipBias(Local_vLevel, Local_vAddress, MaxLocalLevel);
 			if (LocalMipBias > 0u)
 			{
 				Local_vLevel += LocalMipBias;
@@ -1675,7 +1612,6 @@ void FVirtualTextureSystem::GatherRequestsTask(const FGatherRequestsParameters&
 
 					ensure(Allocated_vLevel <= MaxLevel);
 					const uint32 Allocated_vAddress = vAddress & (0xffffffff << (Allocated_vLevel * vDimensions));
-					const uint32 AllocatedWrapped_vAddress = Wrapped_vAddress & (0xffffffff << (Allocated_vLevel * vDimensions));
 
 					AddPageUpdate(PageUpdateBuffers, PageUpdateFlushCount, PhysicalSpace->GetID(), Allocated_pAddress);
 
@@ -1700,22 +1636,6 @@ void FVirtualTextureSystem::GatherRequestsTask(const FGatherRequestsParameters&
 								// map the page now if it wasn't already mapped
 								RequestList->AddDirectMappingRequest(Space->GetID(), PhysicalSpace->GetID(), PageTableLayerIndex, MaxLevel, Allocated_vAddress, Allocated_vLevel, AllocatedMapping_vLevel, Allocated_pAddress);
 							}
-
-							if (AllocatedWrapped_vAddress != Allocated_vAddress)
-							{
-								const FPhysicalSpaceIDAndAddress PrevWrappedPhysicalSpaceIDAndAddress = PageMap.FindPagePhysicalSpaceIDAndAddress(Allocated_vLevel, AllocatedWrapped_vAddress);
-
-								// Wrapped vAddress must be mapped to the same physical address unwrapped vAddress
-								ensure(PrevWrappedPhysicalSpaceIDAndAddress.Packed == ~0u || PrevWrappedPhysicalSpaceIDAndAddress == PhysicalSpaceIDAndAddress);
-								ensure(PrevPhysicalSpaceIDAndAddress.Packed == ~0u ||
-									PrevWrappedPhysicalSpaceIDAndAddress.Packed == ~0u ||
-									PrevPhysicalSpaceIDAndAddress == PrevWrappedPhysicalSpaceIDAndAddress);
-
-								if (PrevWrappedPhysicalSpaceIDAndAddress.Packed == ~0u)
-								{
-									RequestList->AddDirectMappingRequest(Space->GetID(), PhysicalSpace->GetID(), PageTableLayerIndex, MaxLevel, AllocatedWrapped_vAddress, Allocated_vLevel, AllocatedMapping_vLevel, Allocated_pAddress);
-								}
-							}
 							++NumMappedPages;
 						}
 					}
@@ -1775,7 +1695,6 @@ void FVirtualTextureSystem::GatherRequestsTask(const FGatherRequestsParameters&
 						const uint32 PrefetchMapping_vLevel = PrefetchLocal_vLevel + ProducerMipBias;
 						ensure(PrefetchMapping_vLevel <= MaxLevel);
 						const uint32 Prefetch_vAddress = vAddress & (0xffffffff << (PrefetchMapping_vLevel * vDimensions));
-						const uint32 PrefetchWrapped_vAddress = Wrapped_vAddress & (0xffffffff << (PrefetchMapping_vLevel * vDimensions));
 						for (uint32 LoadLayerIndex = 0u; LoadLayerIndex < NumPageTableLayersToLoad; ++LoadLayerIndex)
 						{
 							const uint32 LayerIndex = PageTableLayersToLoad[LoadLayerIndex];
@@ -1785,10 +1704,6 @@ void FVirtualTextureSystem::GatherRequestsTask(const FGatherRequestsParameters&
 								if (ProducerPhysicalGroupMaskToPrefetch & (1u << ProducerPhysicalGroupIndex))
 								{
 									RequestList->AddMappingRequest(LoadRequestIndex, ProducerPhysicalGroupIndex, ID, LayerIndex, MaxLevel, Prefetch_vAddress, PrefetchMapping_vLevel, PrefetchMapping_vLevel);
-									if (PrefetchWrapped_vAddress != Prefetch_vAddress)
-									{
-										RequestList->AddMappingRequest(LoadRequestIndex, ProducerPhysicalGroupIndex, ID, LayerIndex, MaxLevel, PrefetchWrapped_vAddress, PrefetchMapping_vLevel, PrefetchMapping_vLevel);
-									}
 								}
 							}
 						}
@@ -1810,10 +1725,6 @@ void FVirtualTextureSystem::GatherRequestsTask(const FGatherRequestsParameters&
 							if (GroupMaskToLoad & (1u << ProducerPhysicalGroupIndex))
 							{
 								RequestList->AddMappingRequest(LoadRequestIndex, ProducerPhysicalGroupIndex, ID, LayerIndex, MaxLevel, vAddress, vLevel, Mapping_vLevel);
-								if (Wrapped_vAddress != vAddress)
-								{
-									RequestList->AddMappingRequest(LoadRequestIndex, ProducerPhysicalGroupIndex, ID, LayerIndex, MaxLevel, Wrapped_vAddress, vLevel, Mapping_vLevel);
-								}
 							}
 						}
 					}
