Triple buffer for eye adaptation read back on mobile.

diff --git a/Engine/Source/Runtime/Renderer/Private/PostProcess/PostProcessEyeAdaptation.cpp b/Engine/Source/Runtime/Renderer/Private/PostProcess/PostProcessEyeAdaptation.cpp
--- a/Engine/Source/Runtime/Renderer/Private/PostProcess/PostProcessEyeAdaptation.cpp
+++ b/Engine/Source/Runtime/Renderer/Private/PostProcess/PostProcessEyeAdaptation.cpp
@@ -722,12 +722,9 @@ void FSceneViewState::FEyeAdaptationManager::SafeRelease()
 	{
 		PooledRenderTarget[Index].SafeRelease();
 		ExposureTextureReadback[Index] = nullptr;
+		ExposureBufferData[Index].SafeRelease();
+		ExposureBufferReadback[Index] = nullptr;
 	}
-
-	ExposureBufferData[0].SafeRelease();
-	ExposureBufferData[1].SafeRelease();
-	ExposureBufferReadback = nullptr;
-
 }
 
 void FSceneViewState::FEyeAdaptationManager::SwapTextures(FRDGBuilder& GraphBuilder, bool bInUpdateLastExposure)
@@ -811,14 +808,14 @@ const TRefCountPtr<IPooledRenderTarget>& FSceneViewState::FEyeAdaptationManager:
 
 const TRefCountPtr<FRDGPooledBuffer>& FSceneViewState::FEyeAdaptationManager::GetBuffer(uint32 BufferIndex) const
 {
-	check(BufferIndex == 0 || BufferIndex == 1);
+	check(0 <= BufferIndex && BufferIndex < 3);
 
 	return ExposureBufferData[BufferIndex];
 }
 
 const TRefCountPtr<FRDGPooledBuffer>& FSceneViewState::FEyeAdaptationManager::GetOrCreateBuffer(FRDGBuilder& GraphBuilder, uint32 BufferIndex)
 {
-	check(BufferIndex == 0 || BufferIndex == 1);
+	check(0 <= BufferIndex && BufferIndex < 3);
 
 	// Create textures if needed.
 	if (!ExposureBufferData[BufferIndex].IsValid())
@@ -843,8 +840,6 @@ void FSceneViewState::FEyeAdaptationManager::SwapBuffers(FRDGBuilder& GraphBuild
 
 	FRHICommandListImmediate& RHICmdList = FRHICommandListExecutor::GetImmediateCommandList();
 
-	check(CurrentBuffer == 0 || CurrentBuffer == 1);
-
 	// When reading back last frame's exposure data in AFR, make sure we do it on the
 	// GPUs that were actually used last frame.
 	const FRHIGPUMask ThisFrameGPUMask = RHICmdList.GetGPUMask();
@@ -855,8 +850,9 @@ void FSceneViewState::FEyeAdaptationManager::SwapBuffers(FRDGBuilder& GraphBuild
 	{
 		FRDGBufferRef CurrentRDGBuffer = GraphBuilder.RegisterExternalBuffer(ExposureBufferData[CurrentBuffer], ERDGBufferFlags::MultiFrame);
 
-		bool bReadbackCompleted = false;
-		if (ExposureBufferReadback != nullptr && ExposureBufferReadback->IsReady())
+		// first, read the value from two frames ago
+		int32 PreviousPreviousBuffer = GetPreviousPreviousIndex();
+		if (ExposureBufferReadback[PreviousPreviousBuffer] != nullptr && ExposureBufferReadback[PreviousPreviousBuffer]->IsReady())
 		{
 			// Workaround until FRHIGPUTextureReadback::Lock has multigpu support
 			FRHIGPUMask ReadBackGPUMask = LastLastFrameGPUMask;
@@ -864,31 +860,26 @@ void FSceneViewState::FEyeAdaptationManager::SwapBuffers(FRDGBuilder& GraphBuild
 			RDG_GPU_MASK_SCOPE(GraphBuilder, ReadBackGPUMask);
 
 			// Read the last request results.
-			FVector4* ReadbackData = (FVector4*)ExposureBufferReadback->Lock(sizeof(FVector4));
+			FVector4* ReadbackData = (FVector4*)ExposureBufferReadback[PreviousPreviousBuffer]->Lock(sizeof(FVector4));
 			if (ReadbackData)
 			{
 				LastExposure = ReadbackData->X;
 				LastAverageSceneLuminance = ReadbackData->Z;
 
-				ExposureBufferReadback->Unlock();
+				ExposureBufferReadback[PreviousPreviousBuffer]->Unlock();
 			}
-
-			bReadbackCompleted = true;
 		}
 
-		if (bReadbackCompleted || !ExposureBufferReadback)
+		if (!ExposureBufferReadback[CurrentBuffer])
 		{
-			if (!ExposureBufferReadback)
-			{
-				static const FName ExposureValueName(TEXT("Scene view state exposure readback"));
-				ExposureBufferReadback.Reset(new FRHIGPUBufferReadback(ExposureValueName));
-			}
-
-			AddEnqueueCopyPass(GraphBuilder, ExposureBufferReadback.Get(), CurrentRDGBuffer, 0);
+			static const FName ExposureValueName(TEXT("Scene view state exposure readback"));
+			ExposureBufferReadback[CurrentBuffer].Reset(new FRHIGPUBufferReadback(ExposureValueName));
 		}
+
+		AddEnqueueCopyPass(GraphBuilder, ExposureBufferReadback[CurrentBuffer].Get(), CurrentRDGBuffer, 0);
 	}
 
-	CurrentBuffer = 1 - CurrentBuffer;
+	CurrentBuffer = (CurrentBuffer + 1) % 3;
 }
 
 void FSceneViewState::UpdatePreExposure(FViewInfo& View)
diff --git a/Engine/Source/Runtime/Renderer/Private/ScenePrivate.h b/Engine/Source/Runtime/Renderer/Private/ScenePrivate.h
--- a/Engine/Source/Runtime/Renderer/Private/ScenePrivate.h
+++ b/Engine/Source/Runtime/Renderer/Private/ScenePrivate.h
@@ -860,7 +860,7 @@ private:
 
 		const TRefCountPtr<FRDGPooledBuffer>& GetLastBuffer(FRDGBuilder& GraphBuilder)
 		{
-			return GetOrCreateBuffer(GraphBuilder, 1 - CurrentBuffer);
+			return GetOrCreateBuffer(GraphBuilder, GetPreviousPreviousIndex());
 		}
 
 		void SwapBuffers(FRDGBuilder& GraphBuilder, bool bUpdateLastExposure);
@@ -883,8 +883,8 @@ private:
 		TUniquePtr<FRHIGPUTextureReadback> ExposureTextureReadback[3];
 
 		// ES3.1 feature level. For efficent readback use buffers instead of textures
-		TRefCountPtr<FRDGPooledBuffer> ExposureBufferData[2];
-		TUniquePtr<FRHIGPUBufferReadback> ExposureBufferReadback;
+		TRefCountPtr<FRDGPooledBuffer> ExposureBufferData[3];
+		TUniquePtr<FRHIGPUBufferReadback> ExposureBufferReadback[3];
 
 	} EyeAdaptationManager;
 
