Don't force PrecomputedIndirectLightingColorScale to 0 if the platform doesn't support Lumen. This is a workaround for missing indirect illumination in Vulkan.

diff --git a/Engine/Source/Runtime/Renderer/Private/SceneRendering.cpp b/Engine/Source/Runtime/Renderer/Private/SceneRendering.cpp
--- a/Engine/Source/Runtime/Renderer/Private/SceneRendering.cpp
+++ b/Engine/Source/Runtime/Renderer/Private/SceneRendering.cpp
@@ -1703,6 +1703,9 @@ void FViewInfo::SetupUniformBufferParameters(
 
 	ViewUniformShaderParameters.DecalDepthBias = CVarDecalDepthBias.GetValueOnRenderThread();
 
+	ERHIFeatureLevel::Type RHIFeatureLevel = Scene == nullptr ? GMaxRHIFeatureLevel : Scene->GetFeatureLevel();
+	EShaderPlatform ShaderPlatform = GShaderPlatformForFeatureLevel[RHIFeatureLevel];
+
 	ViewUniformShaderParameters.IndirectLightingColorScale = FVector(FinalPostProcessSettings.IndirectLightingColor.R * FinalPostProcessSettings.IndirectLightingIntensity,
 		FinalPostProcessSettings.IndirectLightingColor.G * FinalPostProcessSettings.IndirectLightingIntensity,
 		FinalPostProcessSettings.IndirectLightingColor.B * FinalPostProcessSettings.IndirectLightingIntensity);
@@ -1711,7 +1714,7 @@ void FViewInfo::SetupUniformBufferParameters(
 
 	// If Lumen Dynamic GI is enabled then we don't want GI from Lightmaps
 	// Note: this has the side effect of removing direct lighting from Static Lights
-	if (FinalPostProcessSettings.DynamicGlobalIlluminationMethod == EDynamicGlobalIlluminationMethod::Lumen)
+	if (FinalPostProcessSettings.DynamicGlobalIlluminationMethod == EDynamicGlobalIlluminationMethod::Lumen && DoesPlatformSupportLumenGI(ShaderPlatform))
 	{
 		ViewUniformShaderParameters.PrecomputedIndirectLightingColorScale = FVector::ZeroVector;
 	}
@@ -1729,8 +1732,6 @@ void FViewInfo::SetupUniformBufferParameters(
 
 	ViewUniformShaderParameters.CircleDOFParams = DiaphragmDOF::CircleDofHalfCoc(*this);
 
-	ERHIFeatureLevel::Type RHIFeatureLevel = Scene == nullptr ? GMaxRHIFeatureLevel : Scene->GetFeatureLevel();
-
 	if (Scene && Scene->SkyLight)
 	{
 		FSkyLightSceneProxy* SkyLight = Scene->SkyLight;
