Avoid copying scene primitive bounds array.

diff --git a/Engine/Source/Runtime/Renderer/Private/DistanceFieldObjectManagement.cpp b/Engine/Source/Runtime/Renderer/Private/DistanceFieldObjectManagement.cpp
--- a/Engine/Source/Runtime/Renderer/Private/DistanceFieldObjectManagement.cpp
+++ b/Engine/Source/Runtime/Renderer/Private/DistanceFieldObjectManagement.cpp
@@ -626,7 +626,7 @@ void FDistanceFieldSceneData::UpdateDistanceFieldObjectBuffers(
 					int32 RangeCount = PartitionUpdateRangesDFO(ParallelRanges, IndicesToUpdateInObjectBuffers.Num(), bExecuteInParallel);
 
 					ParallelFor(RangeCount,
-						[this, &ParallelRanges, PrimitiveBounds, RangeCount](int32 RangeIndex)
+						[this, &ParallelRanges, &PrimitiveBounds, RangeCount](int32 RangeIndex)
 						{
 							for (int32 ItemIndex = ParallelRanges.Range[RangeIndex].ItemStart; ItemIndex < ParallelRanges.Range[RangeIndex].ItemStart + ParallelRanges.Range[RangeIndex].ItemCount; ++ItemIndex)
 							{
