Avoid creating debug SSR pass in test and shipping builds.

diff --git a/Engine/Source/Runtime/Renderer/Private/ScreenSpaceRayTracing.cpp b/Engine/Source/Runtime/Renderer/Private/ScreenSpaceRayTracing.cpp
--- a/Engine/Source/Runtime/Renderer/Private/ScreenSpaceRayTracing.cpp
+++ b/Engine/Source/Runtime/Renderer/Private/ScreenSpaceRayTracing.cpp
@@ -175,7 +175,7 @@ bool IsSSRTemporalPassRequired(const FViewInfo& View)
 
 
 FRDGTextureUAV* CreateScreenSpaceRayTracingDebugUAV(FRDGBuilder& GraphBuilder, const FRDGTextureDesc& Desc, const TCHAR* Name, bool bClear = false)
-#if 1
+#if (!UE_BUILD_SHIPPING && !UE_BUILD_TEST)
 {
 	FRDGTextureDesc DebugDesc = FRDGTextureDesc::Create2D(
 		Desc.Extent,
