Make sure we're always able to map the lowest resolution VT mips, otherwise the system will continue to try (and fail) to do this every frame, costing lots of perf
- Need to include LocalMipBias, otherwise may not be any data at AllocatedVT's MaxLevel for sparse UDIMs
- Check for overflow when adding lock requests, keep overflow requests around to try again next frame

diff --git a/Engine/Source/Runtime/Engine/Private/VT/UploadingVirtualTexture.cpp b/Engine/Source/Runtime/Engine/Private/VT/UploadingVirtualTexture.cpp
--- a/Engine/Source/Runtime/Engine/Private/VT/UploadingVirtualTexture.cpp
+++ b/Engine/Source/Runtime/Engine/Private/VT/UploadingVirtualTexture.cpp
@@ -97,18 +97,18 @@ FUploadingVirtualTexture::~FUploadingVirtualTexture()
 	}
 }
 
-uint32 FUploadingVirtualTexture::GetLocalMipBias(uint8 vLevel, uint32 vAddress, uint8 MaxLevel) const
+uint32 FUploadingVirtualTexture::GetLocalMipBias(uint8 vLevel, uint32 vAddress) const
 {
 	const uint32 NumMips = Data->NumMips;
 	uint32 Current_vLevel = vLevel;
 	uint32 Current_vAddress = vAddress;
-	while (Current_vLevel <= MaxLevel)
+	while (Current_vLevel < NumMips)
 	{
 		const uint32 TileIndex = Data->GetTileIndex(Current_vLevel, Current_vAddress);
 		if (TileIndex == ~0u)
 		{
 			// vAddress is out-of-bounds for the given producer
-			Current_vLevel = MaxLevel;
+			Current_vLevel = NumMips - 1u;
 			break;
 		}
 
diff --git a/Engine/Source/Runtime/Engine/Private/VT/UploadingVirtualTexture.h b/Engine/Source/Runtime/Engine/Private/VT/UploadingVirtualTexture.h
--- a/Engine/Source/Runtime/Engine/Private/VT/UploadingVirtualTexture.h
+++ b/Engine/Source/Runtime/Engine/Private/VT/UploadingVirtualTexture.h
@@ -76,7 +76,7 @@ public:
 	virtual ~FUploadingVirtualTexture();
 
 	// IVirtualTexture interface
-	virtual uint32 GetLocalMipBias(uint8 vLevel, uint32 vAddress, uint8 MaxLevel) const override;
+	virtual uint32 GetLocalMipBias(uint8 vLevel, uint32 vAddress) const override;
 	virtual FVTRequestPageResult RequestPageData(const FVirtualTextureProducerHandle& ProducerHandle, uint8 LayerMask, uint8 vLevel, uint64 vAddress, EVTRequestPagePriority Priority) override;
 	virtual IVirtualTextureFinalizer* ProducePageData(FRHICommandListImmediate& RHICmdList,
 		ERHIFeatureLevel::Type FeatureLevel,
diff --git a/Engine/Source/Runtime/RenderCore/Public/VirtualTexturing.h b/Engine/Source/Runtime/RenderCore/Public/VirtualTexturing.h
--- a/Engine/Source/Runtime/RenderCore/Public/VirtualTexturing.h
+++ b/Engine/Source/Runtime/RenderCore/Public/VirtualTexturing.h
@@ -259,7 +259,7 @@ public:
 	 * @param vAddress Virtual address to check
 	 * @return Mip bias to be added to vLevel to reach a resident page at the given address
 	 */
-	virtual uint32 GetLocalMipBias(uint8 vLevel, uint32 vAddress, uint8 MaxLevel) const { return 0u; }
+	virtual uint32 GetLocalMipBias(uint8 vLevel, uint32 vAddress) const { return 0u; }
 
 	/**
 	* Makes a request for the given page data.
diff --git a/Engine/Source/Runtime/Renderer/Private/VT/AllocatedVirtualTexture.cpp b/Engine/Source/Runtime/Renderer/Private/VT/AllocatedVirtualTexture.cpp
--- a/Engine/Source/Runtime/Renderer/Private/VT/AllocatedVirtualTexture.cpp
+++ b/Engine/Source/Runtime/Renderer/Private/VT/AllocatedVirtualTexture.cpp
@@ -61,39 +61,8 @@
 
 	MaxLevel = FMath::Min(MaxLevel, VIRTUALTEXTURE_LOG2_MAX_PAGETABLE_SIZE - 1u);
 
-	// Lock lowest resolution mip from each producer
-	// Depending on the block dimensions of the producers that make up this allocated VT, different allocated VTs may need to lock different low resolution mips from the same producer
-	// In the common case where block dimensions match, same mip will be locked by all allocated VTs that make use of the same producer
-	for (int32 ProducerIndex = 0u; ProducerIndex < UniqueProducers.Num(); ++ProducerIndex)
-	{
-		FVirtualTextureProducerHandle ProducerHandle = UniqueProducers[ProducerIndex].Handle;
-		FVirtualTextureProducer* Producer = InSystem->FindProducer(ProducerHandle);
-		if (Producer && Producer->GetDescription().bPersistentHighestMip)
-		{
-			const uint32 MipBias = UniqueProducers[ProducerIndex].MipBias;
-			check(MipBias <= MaxLevel);
-			const uint32 Local_vLevel = MaxLevel - MipBias;
-			checkf(Local_vLevel <= Producer->GetMaxLevel(), TEXT("Invalid Local_vLevel %d for VT producer %s, Producer MaxLevel %d, MipBias %d, AllocatedVT MaxLevel %d"),
-				Local_vLevel,
-				*Producer->GetName().ToString(),
-				Producer->GetMaxLevel(),
-				MipBias,
-				MaxLevel);
+	// Lock tiles
+	LockOrUnlockTiles(InSystem, true);
 
-			const uint32 MipScaleFactor = (1u << Local_vLevel);
-			const uint32 RootWidthInTiles = FMath::DivideAndRoundUp(Producer->GetWidthInTiles(), MipScaleFactor);
-			const uint32 RootHeightInTiles = FMath::DivideAndRoundUp(Producer->GetHeightInTiles(), MipScaleFactor);
-			for (uint32 TileY = 0u; TileY < RootHeightInTiles; ++TileY)
-			{
-				for (uint32 TileX = 0u; TileX < RootWidthInTiles; ++TileX)
-				{
-					const uint32 Local_vAddress = FMath::MortonCode2(TileX) | (FMath::MortonCode2(TileY) << 1);
-					const FVirtualTextureLocalTile TileToUnlock(ProducerHandle, Local_vAddress, Local_vLevel);
-					InSystem->LockTile(TileToUnlock);
-				}
-			}
-		}
-	}
-
 	// Use 16bit page table entries if all physical spaces are small enough
 	bool bSupport16BitPageTable = true;
@@ -151,31 +120,7 @@
 	check(RefCount.GetValue() == 0);
 
 	// Unlock any locked tiles
-	for (int32 ProducerIndex = 0u; ProducerIndex < UniqueProducers.Num(); ++ProducerIndex)
-	{
-		const FVirtualTextureProducerHandle ProducerHandle = UniqueProducers[ProducerIndex].Handle;
-		FVirtualTextureProducer* Producer = System->FindProducer(ProducerHandle);
-		if (Producer && Producer->GetDescription().bPersistentHighestMip)
-		{
-			const uint32 MipBias = UniqueProducers[ProducerIndex].MipBias;
-			check(MipBias <= MaxLevel);
-			const uint32 Local_vLevel = MaxLevel - MipBias;
-			check(Local_vLevel <= Producer->GetMaxLevel());
+	LockOrUnlockTiles(System, false);
 
-			const uint32 MipScaleFactor = (1u << Local_vLevel);
-			const uint32 RootWidthInTiles = FMath::DivideAndRoundUp(Producer->GetWidthInTiles(), MipScaleFactor);
-			const uint32 RootHeightInTiles = FMath::DivideAndRoundUp(Producer->GetHeightInTiles(), MipScaleFactor);
-			for (uint32 TileY = 0u; TileY < RootHeightInTiles; ++TileY)
-			{
-				for (uint32 TileX = 0u; TileX < RootWidthInTiles; ++TileX)
-				{
-					const uint32 Local_vAddress = FMath::MortonCode2(TileX) | (FMath::MortonCode2(TileY) << 1);
-					const FVirtualTextureLocalTile TileToUnlock(ProducerHandle, Local_vAddress, Local_vLevel);
-					System->UnlockTile(TileToUnlock, Producer);
-				}
-			}
-		}
-	}
-
 	// Physical pool needs to evict all pages that belong to this VT
 	{
@@ -236,6 +181,59 @@
 	delete this;
 }
 
+void FAllocatedVirtualTexture::LockOrUnlockTiles(FVirtualTextureSystem* InSystem, bool bLock) const
+{
+	// (Un)Lock lowest resolution mip from each producer
+	// Depending on the block dimensions of the producers that make up this allocated VT, different allocated VTs may need to lock different low resolution mips from the same producer
+	// In the common case where block dimensions match, same mip will be locked by all allocated VTs that make use of the same producer
+	for (int32 ProducerIndex = 0u; ProducerIndex < UniqueProducers.Num(); ++ProducerIndex)
+	{
+		FVirtualTextureProducerHandle ProducerHandle = UniqueProducers[ProducerIndex].Handle;
+		FVirtualTextureProducer* Producer = InSystem->FindProducer(ProducerHandle);
+		if (Producer && Producer->GetDescription().bPersistentHighestMip)
+		{
+			const uint32 MipBias = UniqueProducers[ProducerIndex].MipBias;
+			check(MipBias <= MaxLevel);
+			const uint32 Local_vLevel = MaxLevel - MipBias;
+			checkf(Local_vLevel <= Producer->GetMaxLevel(), TEXT("Invalid Local_vLevel %d for VT producer %s, Producer MaxLevel %d, MipBias %d, AllocatedVT MaxLevel %d"),
+				Local_vLevel,
+				*Producer->GetName().ToString(),
+				Producer->GetMaxLevel(),
+				MipBias,
+				MaxLevel);
+
+			const uint32 MipScaleFactor = (1u << Local_vLevel);
+			const uint32 RootWidthInTiles = FMath::DivideAndRoundUp(Producer->GetWidthInTiles(), MipScaleFactor);
+			const uint32 RootHeightInTiles = FMath::DivideAndRoundUp(Producer->GetHeightInTiles(), MipScaleFactor);
+
+			for (uint32 TileY = 0u; TileY < RootHeightInTiles; ++TileY)
+			{
+				for (uint32 TileX = 0u; TileX < RootWidthInTiles; ++TileX)
+				{
+					const uint32 Local_vAddress = FMath::MortonCode2(TileX) | (FMath::MortonCode2(TileY) << 1);
+					FVirtualTextureLocalTile Tile(ProducerHandle, Local_vAddress, Local_vLevel);
+
+					const uint32 LocalMipBias = Producer->GetVirtualTexture()->GetLocalMipBias(Local_vLevel, Local_vAddress);
+					if (LocalMipBias > 0u)
+					{
+						Tile.Local_vAddress >>= (LocalMipBias * Description.Dimensions);
+						Tile.Local_vLevel += LocalMipBias;
+					}
+
+					if (bLock)
+					{
+						InSystem->LockTile(Tile);
+					}
+					else
+					{
+						InSystem->UnlockTile(Tile, Producer);
+					}
+				}
+			}
+		}
+	}
+}
+
 uint32 FAllocatedVirtualTexture::AddUniqueProducer(FVirtualTextureProducerHandle const& InHandle, const FVirtualTextureProducer* InProducer)
 {
 	for (int32 Index = 0u; Index < UniqueProducers.Num(); ++Index)
diff --git a/Engine/Source/Runtime/Renderer/Private/VT/AllocatedVirtualTexture.h b/Engine/Source/Runtime/Renderer/Private/VT/AllocatedVirtualTexture.h
--- a/Engine/Source/Runtime/Renderer/Private/VT/AllocatedVirtualTexture.h
+++ b/Engine/Source/Runtime/Renderer/Private/VT/AllocatedVirtualTexture.h
@@ -36,6 +36,8 @@ public:
 
 	void AssignVirtualAddress(uint32 vAddress);
 
+	void LockOrUnlockTiles(FVirtualTextureSystem* InSystem, bool bLock) const;
+
 	// begin IAllocatedVirtualTexture
 	virtual uint32 GetNumPageTableTextures() const override;
 	virtual FRHITexture* GetPageTableTexture(uint32 InPageTableIndex) const override;
diff --git a/Engine/Source/Runtime/Renderer/Private/VT/VirtualTextureSystem.cpp b/Engine/Source/Runtime/Renderer/Private/VT/VirtualTextureSystem.cpp
--- a/Engine/Source/Runtime/Renderer/Private/VT/VirtualTextureSystem.cpp
+++ b/Engine/Source/Runtime/Renderer/Private/VT/VirtualTextureSystem.cpp
@@ -907,20 +907,23 @@ static int32 WrapTilePosition(int32 Position, int32 Size)
 
 void FVirtualTextureSystem::RequestTilesForRegionInternal(const IAllocatedVirtualTexture* AllocatedVT, const FVector2D& InScreenSpaceSize, const FVector2D& InViewportPosition, const FVector2D& InViewportSize, const FVector2D& InUV0, const FVector2D& InUV1, int32 InMipLevel)
 {
+	const float ScreenSpaceSizeX = FMath::Max(InScreenSpaceSize.X, 1.0f);
+	const float ScreenSpaceSizeY = FMath::Max(InScreenSpaceSize.Y, 1.0f);
+
 	// Determine the screen-space coordinates of the texture we're trying to display
 	const float SrcPositionX0 = FMath::Max(InViewportPosition.X, 0.0f);
 	const float SrcPositionY0 = FMath::Max(InViewportPosition.Y, 0.0f);
-	const float SrcPositionX1 = FMath::Min(InViewportPosition.X + InViewportSize.X, InScreenSpaceSize.X);
-	const float SrcPositionY1 = FMath::Min(InViewportPosition.Y + InViewportSize.Y, InScreenSpaceSize.Y);
+	const float SrcPositionX1 = FMath::Min(InViewportPosition.X + InViewportSize.X, ScreenSpaceSizeX);
+	const float SrcPositionY1 = FMath::Min(InViewportPosition.Y + InViewportSize.Y, ScreenSpaceSizeY);
 
 	const int32 WidthInBlocks = AllocatedVT->GetWidthInBlocks();
 	const int32 HeightInBlocks = AllocatedVT->GetHeightInBlocks();
 
 	// Map coordinates to UV space
-	const float PositionU0 = FMath::Lerp(InUV0.X, InUV1.X, SrcPositionX0 / InScreenSpaceSize.X) / WidthInBlocks;
-	const float PositionV0 = FMath::Lerp(InUV0.Y, InUV1.Y, SrcPositionY0 / InScreenSpaceSize.Y) / HeightInBlocks;
-	const float PositionU1 = FMath::Lerp(InUV0.X, InUV1.X, SrcPositionX1 / InScreenSpaceSize.X) / WidthInBlocks;
-	const float PositionV1 = FMath::Lerp(InUV0.Y, InUV1.Y, SrcPositionY1 / InScreenSpaceSize.Y) / HeightInBlocks;
+	const float PositionU0 = FMath::Lerp(InUV0.X, InUV1.X, SrcPositionX0 / ScreenSpaceSizeX) / WidthInBlocks;
+	const float PositionV0 = FMath::Lerp(InUV0.Y, InUV1.Y, SrcPositionY0 / ScreenSpaceSizeY) / HeightInBlocks;
+	const float PositionU1 = FMath::Lerp(InUV0.X, InUV1.X, SrcPositionX1 / ScreenSpaceSizeX) / WidthInBlocks;
+	const float PositionV1 = FMath::Lerp(InUV0.Y, InUV1.Y, SrcPositionY1 / ScreenSpaceSizeY) / HeightInBlocks;
 
 	// Map UVs to tile coordinates
 	const int32 MipWidthInTiles = FMath::Max<int32>(AllocatedVT->GetWidthInTiles() >> InMipLevel, 1);
@@ -1139,6 +1142,7 @@ void FVirtualTextureSystem::Update(FRDGBuilder& GraphBuilder, ERHIFeatureLevel::
 
 	// Collect tiles to lock
 	{
+		TArray<FVirtualTextureLocalTile> RemainingTilesToLock;
 		for (const FVirtualTextureLocalTile& Tile : TilesToLock)
 		{
 			const FVirtualTextureProducerHandle ProducerHandle = Tile.GetProducerHandle();
@@ -1162,14 +1166,20 @@ void FVirtualTextureSystem::Update(FRDGBuilder& GraphBuilder, ERHIFeatureLevel::
 						PagePool.Lock(pAddress);
 					}
 				}
+
 				if (ProducerLayerMaskToLoad != 0u)
 				{
-					MergedRequestList->LockLoadRequest(FVirtualTextureLocalTile(Tile.GetProducerHandle(), Tile.Local_vAddress, Tile.Local_vLevel), ProducerLayerMaskToLoad);
+					const uint16 LoadRequestIndex = MergedRequestList->LockLoadRequest(FVirtualTextureLocalTile(Tile.GetProducerHandle(), Tile.Local_vAddress, Tile.Local_vLevel), ProducerLayerMaskToLoad);
+					if (LoadRequestIndex == 0xffff)
+					{
+						// Overflowed the request list...try to lock the tile again next frame
+						RemainingTilesToLock.Add(Tile);
+					}
 				}
 			}
 		}
 
-		TilesToLock.Reset();
+		TilesToLock = MoveTemp(RemainingTilesToLock);
 	}
 
 	TArray<uint32> PackedTiles;
@@ -1550,7 +1560,7 @@ void FVirtualTextureSystem::GatherRequestsTask(const FGatherRequestsParameters&
 			}
 
 			const uint32 ProducerMipBias = AllocatedVT->GetUniqueProducerMipBias(ProducerIndex);
-			const uint32 MaxLocalLevel = MaxLevel - ProducerMipBias;
+			const uint32 ProducerMaxLevel = Producer->GetMaxLevel();
 
 			// here vLevel is clamped against ProducerMipBias, as ProducerMipBias represents the most detailed level of this producer, relative to the allocated VT
 			// used to rescale vAddress to the correct tile within the given mip level
@@ -1564,7 +1574,7 @@ void FVirtualTextureSystem::GatherRequestsTask(const FGatherRequestsParameters&
 			const uint32 Local_vPageY = (vPageY - AllocatedPageY) >> Mapping_vLevel;
 			uint32 Local_vAddress = FMath::MortonCode2(Local_vPageX) | (FMath::MortonCode2(Local_vPageY) << 1);
 
-			const uint32 LocalMipBias = Producer->GetVirtualTexture()->GetLocalMipBias(Local_vLevel, Local_vAddress, MaxLocalLevel);
+			const uint32 LocalMipBias = Producer->GetVirtualTexture()->GetLocalMipBias(Local_vLevel, Local_vAddress);
 			if (LocalMipBias > 0u)
 			{
 				Local_vLevel += LocalMipBias;
@@ -1587,7 +1597,7 @@ void FVirtualTextureSystem::GatherRequestsTask(const FGatherRequestsParameters&
 				const FTexturePagePool& RESTRICT PagePool = PhysicalSpace->GetPagePool();
 
 				// Find the highest resolution tile that's currently loaded
-				const uint32 Allocated_pAddress = PagePool.FindNearestPageAddress(ProducerHandle, ProducerGroupIndex, Local_vAddress, Local_vLevel, MaxLevel);
+				const uint32 Allocated_pAddress = PagePool.FindNearestPageAddress(ProducerHandle, ProducerGroupIndex, Local_vAddress, Local_vLevel, ProducerMaxLevel);
 
 				bool bRequestedPageWasResident = false;
 				uint32 AllocatedLocal_vLevel = MaxLevel;
@@ -1597,7 +1607,7 @@ void FVirtualTextureSystem::GatherRequestsTask(const FGatherRequestsParameters&
 					check(AllocatedLocal_vLevel >= Local_vLevel);
 
 					const uint32 AllocatedMapping_vLevel = AllocatedLocal_vLevel + ProducerMipBias;
-					uint32 Allocated_vLevel = AllocatedMapping_vLevel;
+					uint32 Allocated_vLevel = FMath::Min(AllocatedMapping_vLevel, MaxLevel);
 					if (AllocatedLocal_vLevel == Local_vLevel)
 					{
 						// page at the requested level was already resident, no longer need to load
@@ -1946,6 +1956,7 @@ void FVirtualTextureSystem::SubmitRequests(FRDGBuilder& GraphBuilder, ERHIFeatur
 
 			const EVTRequestPagePriority Priority = bLockTile ? EVTRequestPagePriority::High : EVTRequestPagePriority::Normal;
 			FVTRequestPageResult RequestPageResult = Producer.GetVirtualTexture()->RequestPageData(ProducerHandle, ProducerTextureLayerMask, TileToLoad.Local_vLevel, TileToLoad.Local_vAddress, Priority);
+			checkf(!bLockTile || RequestPageResult.Status != EVTRequestPageStatus::Invalid, TEXT("Tried to lock an invalid VT tile"));
 			if (RequestPageResult.Status == EVTRequestPageStatus::Pending && bForceProduceTile)
 			{
 				// If we're trying to lock this tile, we're OK producing data now (and possibly waiting) as long as data is pending
@@ -2225,12 +2236,15 @@ void FVirtualTextureSystem::SubmitRequests(FRDGBuilder& GraphBuilder, ERHIFeatur
 							uint32 pAddress = PageMap.FindPageAddress(vLevel, vAddress);
 							if (pAddress == ~0u)
 							{
-								const uint32 Local_vAddress = FMath::MortonCode2(TileX) | (FMath::MortonCode2(TileY) << 1);
+								uint32 Local_vAddress = FMath::MortonCode2(TileX) | (FMath::MortonCode2(TileY) << 1);
+
+								const uint32 LocalMipBias = Producer->GetVirtualTexture()->GetLocalMipBias(Local_vLevel, Local_vAddress);
+								Local_vAddress >>= (LocalMipBias * 2u);
 
-								pAddress = PagePool.FindPageAddress(ProducerHandle, ProducerPhysicalGroupIndex, Local_vAddress, Local_vLevel);
+								pAddress = PagePool.FindPageAddress(ProducerHandle, ProducerPhysicalGroupIndex, Local_vAddress, Local_vLevel + LocalMipBias);
 								if (pAddress != ~0u)
 								{
-									PagePool.MapPage(Space, PhysicalSpace, PageTableLayerIndex, MaxLevel, vLevel, vAddress, vLevel, pAddress);
+									PagePool.MapPage(Space, PhysicalSpace, PageTableLayerIndex, MaxLevel, vLevel, vAddress, vLevel + LocalMipBias, pAddress);
 								}
 								else
 								{
