Avoid binding hair resources when hair is not enabled.

diff --git a/Engine/Source/Runtime/Renderer/Private/DeferredShadingRenderer.cpp b/Engine/Source/Runtime/Renderer/Private/DeferredShadingRenderer.cpp
--- a/Engine/Source/Runtime/Renderer/Private/DeferredShadingRenderer.cpp
+++ b/Engine/Source/Runtime/Renderer/Private/DeferredShadingRenderer.cpp
@@ -2126,13 +2126,6 @@ void FDeferredShadingSceneRenderer::Render(FRDGBuilder& GraphBuilder)
 			}
 		}
 	}
-	else
-	{
-		for (FViewInfo& View : Views)
-		{
-			View.HairStrandsViewData.UniformBuffer = HairStrands::CreateDefaultHairStrandsViewUniformBuffer(GraphBuilder, View);
-		}
-	}
 
 	if (bNaniteEnabled)
 	{
diff --git a/Engine/Source/Runtime/Renderer/Private/HairStrands/HairStrandsData.h b/Engine/Source/Runtime/Renderer/Private/HairStrands/HairStrandsData.h
--- a/Engine/Source/Runtime/Renderer/Private/HairStrands/HairStrandsData.h
+++ b/Engine/Source/Runtime/Renderer/Private/HairStrands/HairStrandsData.h
@@ -352,7 +352,7 @@ typedef TArray<FHairStrandsMacroGroupData, SceneRenderingAllocator> FHairStrands
 // View Data
 struct FHairStrandsViewData
 {
-	TRDGUniformBufferRef<FHairStrandsViewUniformParameters> UniformBuffer;
+	TRDGUniformBufferRef<FHairStrandsViewUniformParameters> UniformBuffer{};
 	bool bIsValid = false;
 
 	// Internal data
