Avoid binding strata RDG resources when disabled.

diff --git a/Engine/Source/Runtime/Renderer/Private/LightRendering.cpp b/Engine/Source/Runtime/Renderer/Private/LightRendering.cpp
--- a/Engine/Source/Runtime/Renderer/Private/LightRendering.cpp
+++ b/Engine/Source/Runtime/Renderer/Private/LightRendering.cpp
@@ -592,8 +592,10 @@ private:
 		FRenderLightParams* RenderLightParams)
 	{
 		FGlobalShader::SetParameters<FViewUniformShaderParameters>(RHICmdList, ShaderRHI, View.ViewUniformBuffer);
-		TRDGUniformBufferRef<FStrataGlobalUniformParameters> StrataUniformBuffer = Strata::BindStrataGlobalUniformParameters(View.StrataSceneData);
-		FGlobalShader::SetParameters<FStrataGlobalUniformParameters>(RHICmdList, ShaderRHI, StrataUniformBuffer->GetRHIRef());
+		if (TRDGUniformBufferRef<FStrataGlobalUniformParameters> StrataUniformBuffer = Strata::BindStrataGlobalUniformParameters(View.StrataSceneData))
+		{
+			FGlobalShader::SetParameters<FStrataGlobalUniformParameters>(RHICmdList, ShaderRHI, StrataUniformBuffer->GetRHIRef());
+		}
 
 		if(LightAttenuationTexture.IsBound())
 		{
diff --git a/Engine/Source/Runtime/Renderer/Private/Strata/Strata.cpp b/Engine/Source/Runtime/Renderer/Private/Strata/Strata.cpp
--- a/Engine/Source/Runtime/Renderer/Private/Strata/Strata.cpp
+++ b/Engine/Source/Runtime/Renderer/Private/Strata/Strata.cpp
@@ -255,9 +255,9 @@ void InitialiseStrataFrameSceneData(FSceneRenderer& SceneRenderer, FRDGBuilder&
 	}
 
 	// Create the readable uniform buffers for each views once for all (it is view independent and all the views should be tiled into the render target textures & material buffer)
-	FStrataGlobalUniformParameters* StrataUniformParameters = GraphBuilder.AllocParameters<FStrataGlobalUniformParameters>();
 	if (IsStrataEnabled())
 	{
+		FStrataGlobalUniformParameters* StrataUniformParameters = GraphBuilder.AllocParameters<FStrataGlobalUniformParameters>();
 		StrataUniformParameters->MaxBytesPerPixel = StrataSceneData.MaxBytesPerPixel;
 		StrataUniformParameters->MaterialLobesBuffer = StrataSceneData.MaterialLobesBufferSRV;
 		StrataUniformParameters->ClassificationTexture = StrataSceneData.ClassificationTexture;
@@ -269,20 +269,6 @@ void InitialiseStrataFrameSceneData(FSceneRenderer& SceneRenderer, FRDGBuilder&
 		StrataUniformParameters->GGXEnergyLUTScaleBias = GetStrataGGXEnergyLUTScaleBias();
 		StrataSceneData.StrataGlobalUniformParameters = GraphBuilder.CreateUniformBuffer(StrataUniformParameters);
 	}
-	else
-	{
-		// Create each time. This path will go away when Strata is always enabled anyway.
-		StrataUniformParameters->MaxBytesPerPixel = 0;
-		StrataUniformParameters->MaterialLobesBuffer = GraphBuilder.CreateSRV(GraphBuilder.RegisterExternalBuffer(GWhiteVertexBufferWithRDG->Buffer), PF_R32_UINT);
-		StrataUniformParameters->ClassificationTexture = GSystemTextures.GetZeroUIntDummy(GraphBuilder);
-		StrataUniformParameters->TopLayerNormalTexture = GSystemTextures.GetZeroUIntDummy(GraphBuilder);
-		StrataUniformParameters->SSSTexture = GSystemTextures.GetZeroUIntDummy(GraphBuilder);
-		StrataUniformParameters->GGXEnergyLUT3DTexture = GSystemTextures.VolumetricBlackDummy->GetRenderTargetItem().ShaderResourceTexture;
-		StrataUniformParameters->GGXEnergyLUT2DTexture = GSystemTextures.BlackDummy->GetRenderTargetItem().ShaderResourceTexture;
-		StrataUniformParameters->GGXEnergyLUTSampler = TStaticSamplerState<SF_Bilinear, AM_Clamp, AM_Clamp, AM_Clamp>::GetRHI();
-		StrataUniformParameters->GGXEnergyLUTScaleBias = GetStrataGGXEnergyLUTScaleBias();
-		StrataSceneData.StrataGlobalUniformParameters = GraphBuilder.CreateUniformBuffer(StrataUniformParameters);
-	}
 }
 
 void BindStrataBasePassUniformParameters(FRDGBuilder& GraphBuilder, FStrataSceneData* StrataSceneData, FStrataBasePassUniformParameters& OutStrataUniformParameters)
@@ -309,7 +295,7 @@ void BindStrataBasePassUniformParameters(FRDGBuilder& GraphBuilder, FStrataScene
 TRDGUniformBufferRef<FStrataGlobalUniformParameters> BindStrataGlobalUniformParameters(FStrataSceneData* StrataSceneData)
 {
 	check(StrataSceneData);
-	check(StrataSceneData->StrataGlobalUniformParameters != nullptr);
+	check(StrataSceneData->StrataGlobalUniformParameters != nullptr || !IsStrataEnabled());
 	return StrataSceneData->StrataGlobalUniformParameters;
 }
 
diff --git a/Engine/Source/Runtime/Renderer/Private/Strata/Strata.h b/Engine/Source/Runtime/Renderer/Private/Strata/Strata.h
--- a/Engine/Source/Runtime/Renderer/Private/Strata/Strata.h
+++ b/Engine/Source/Runtime/Renderer/Private/Strata/Strata.h
@@ -62,7 +62,7 @@ struct FStrataSceneData
 	FRDGTextureRef TopLayerNormalTexture;
 	FRDGTextureRef SSSTexture;
 
-	TRDGUniformBufferRef<FStrataGlobalUniformParameters> StrataGlobalUniformParameters;
+	TRDGUniformBufferRef<FStrataGlobalUniformParameters> StrataGlobalUniformParameters{};
 
 	// Resources computed once for multiple frames
 
