Fix for async SSAO overlap. Async SSAO is done before the base pass (so that it doesn't touch the GBuffer).

diff --git a/Engine/Source/Runtime/Renderer/Private/BasePassRendering.cpp b/Engine/Source/Runtime/Renderer/Private/BasePassRendering.cpp
--- a/Engine/Source/Runtime/Renderer/Private/BasePassRendering.cpp
+++ b/Engine/Source/Runtime/Renderer/Private/BasePassRendering.cpp
@@ -758,6 +758,8 @@ void FDeferredShadingSceneRenderer::RenderBasePass(
 			ViewFamily.EngineShowFlags.LODColoration ||
 			ViewFamily.EngineShowFlags.HLODColoration);
 
+	const bool bForwardShadingEnabled = IsForwardShadingEnabled(SceneTextures.Config.ShaderPlatform);
+
 	const FExclusiveDepthStencil ExclusiveDepthStencil(BasePassDepthStencilAccess);
 
 	TStaticArray<FRDGTextureRef, MaxSimultaneousRenderTargets> BasePassTextures;
@@ -853,9 +855,13 @@ void FDeferredShadingSceneRenderer::RenderBasePass(
 	BasePassRenderTargets.ShadingRateTexture = GVRSImageManager.GetVariableRateShadingImage(GraphBuilder, ViewFamily, nullptr, EVRSType::None);
 	
 	FForwardBasePassTextures ForwardBasePassTextures{};
-	ForwardBasePassTextures.SceneDepthIfResolved = SceneTextures.Depth.IsSeparate() ? SceneTextures.Depth.Resolve : nullptr;
-	ForwardBasePassTextures.ScreenSpaceAO = SceneTextures.ScreenSpaceAO;
-	ForwardBasePassTextures.ScreenSpaceShadowMask = ForwardShadowMaskTexture;
+
+	if (bForwardShadingEnabled)
+	{
+		ForwardBasePassTextures.SceneDepthIfResolved = SceneTextures.Depth.IsSeparate() ? SceneTextures.Depth.Resolve : nullptr;
+		ForwardBasePassTextures.ScreenSpaceAO = SceneTextures.ScreenSpaceAO;
+		ForwardBasePassTextures.ScreenSpaceShadowMask = ForwardShadowMaskTexture;
+	}
 
 	GraphBuilder.SetCommandListStat(GET_STATID(STAT_CLM_BasePass));
 	RenderBasePassInternal(GraphBuilder, BasePassRenderTargets, BasePassDepthStencilAccess, ForwardBasePassTextures, DBufferTextures, SceneTextures.QuadOverdraw, bDoParallelBasePass, bRenderLightmapDensity, InstanceCullingManager);
diff --git a/Engine/Source/Runtime/Renderer/Private/CompositionLighting/CompositionLighting.cpp b/Engine/Source/Runtime/Renderer/Private/CompositionLighting/CompositionLighting.cpp
--- a/Engine/Source/Runtime/Renderer/Private/CompositionLighting/CompositionLighting.cpp
+++ b/Engine/Source/Runtime/Renderer/Private/CompositionLighting/CompositionLighting.cpp
@@ -501,7 +501,7 @@ void ProcessBeforeBasePass(
 
 	const auto ProcessView = [&](const FViewInfo& View, uint32 SSAOLevels)
 	{
-		const bool bNeedSSAO = SSAOLevels && FSSAOHelper::GetGTAOPassType(View, SSAOLevels) != EGTAOType::ENonAsync;
+		const bool bNeedSSAO = SSAOLevels && FSSAOHelper::GetGTAOPassType(View, SSAOLevels) == EGTAOType::EOff;
 
 		// so that the passes can register themselves to the graph
 		if (bDBuffer || bNeedSSAO)
@@ -539,8 +539,8 @@ void ProcessBeforeBasePass(
 
 		uint32 SSAOLevels = FSSAOHelper::ComputeAmbientOcclusionPassCount(View);
 
-		// In forward, if depth prepass is off - as SSAO here requires a valid HZB buffer - disable SSAO.
-		if (!IsForwardShadingEnabled(ShaderPlatform) || !View.HZB || FSSAOHelper::IsAmbientOcclusionAsyncCompute(View, SSAOLevels))
+		// Disable SSAO if no valid HZB exists, or if we are deferred and non-async (this is handled in ProcessAfterBasePass instead).
+		if (!View.HZB || (!IsForwardShadingEnabled(View.GetShaderPlatform()) && !FSSAOHelper::IsAmbientOcclusionAsyncCompute(View, SSAOLevels)))
 		{
 			SSAOLevels = 0;
 		}
@@ -597,13 +597,7 @@ void ProcessAfterBasePass(
 			const bool bScreenSpaceAOIsProduced = SceneTextures.ScreenSpaceAO->HasBeenProduced();
 			FScreenPassRenderTarget FinalTarget = FScreenPassRenderTarget(SceneTextures.ScreenSpaceAO, View.ViewRect, bScreenSpaceAOIsProduced ? ERenderTargetLoadAction::ELoad : ERenderTargetLoadAction::ENoAction);
 
-			FScreenPassTexture AmbientOcclusion;
-#if RHI_RAYTRACING
-			if (ShouldRenderRayTracingAmbientOcclusion(View) && bScreenSpaceAOIsProduced)
-			{
-				AmbientOcclusion = FinalTarget;
-			}
-#endif
+			FScreenPassTexture AmbientOcclusion = bScreenSpaceAOIsProduced ? FinalTarget : FScreenPassTexture();
 
 			const EGTAOType GTAOType = FSSAOHelper::GetGTAOPassType(View, SSAOLevels);
 
@@ -628,7 +622,7 @@ void ProcessAfterBasePass(
 					FGTAOCommonParameters Parameters = GetGTAOCommonParameters(GraphBuilder, View, SceneTextures.UniformBuffer, GTAOType);
 					AmbientOcclusion = AddPostProcessingGTAOAllPasses(GraphBuilder, View, Parameters, FinalTarget);
 				}
-				else
+				else if (!AmbientOcclusion.IsValid())
 				{
 					FSSAOCommonParameters Parameters = GetSSAOCommonParameters(GraphBuilder, View, SceneTextures.UniformBuffer, SSAOLevels, true);
 					AmbientOcclusion = AddPostProcessingAmbientOcclusion(GraphBuilder, View, Parameters, FinalTarget);
