Display on-screen message when VT pool is oversubscribed

diff --git a/Engine/Source/Runtime/Renderer/Private/VT/VirtualTexturePhysicalSpace.cpp b/Engine/Source/Runtime/Renderer/Private/VT/VirtualTexturePhysicalSpace.cpp
--- a/Engine/Source/Runtime/Renderer/Private/VT/VirtualTexturePhysicalSpace.cpp
+++ b/Engine/Source/Runtime/Renderer/Private/VT/VirtualTexturePhysicalSpace.cpp
@@ -11,6 +11,7 @@
 FVirtualTexturePhysicalSpace::FVirtualTexturePhysicalSpace(const FVTPhysicalSpaceDescription& InDesc, uint16 InID)
 	: Description(InDesc)
 	, NumRefs(0u)
+	, LastFrameOversubscribed(0u)
 	, ID(InID)
 	, bPageTableLimit(false)
 	, bGpuTextureLimit(false)
diff --git a/Engine/Source/Runtime/Renderer/Private/VT/VirtualTexturePhysicalSpace.h b/Engine/Source/Runtime/Renderer/Private/VT/VirtualTexturePhysicalSpace.h
--- a/Engine/Source/Runtime/Renderer/Private/VT/VirtualTexturePhysicalSpace.h
+++ b/Engine/Source/Runtime/Renderer/Private/VT/VirtualTexturePhysicalSpace.h
@@ -70,6 +70,9 @@ public:
 	inline uint32 Release() { check(NumRefs > 0u); return --NumRefs; }
 	inline uint32 GetRefCount() const { return NumRefs; }
 
+	inline uint32 GetLastFrameOversubscribed() const { return LastFrameOversubscribed; }
+	inline void SetLastFrameOversubscribed(uint32 InFrame) { LastFrameOversubscribed = InFrame; }
+
 	FRHITexture* GetPhysicalTexture(int32 Layer) const
 	{
 		check(PooledRenderTarget[Layer].IsValid());
@@ -112,6 +115,7 @@ private:
 
 	uint32 TextureSizeInTiles;
 	uint32 NumRefs;
+	uint32 LastFrameOversubscribed;
 	uint16 ID;
 	bool bPageTableLimit; // True if the physical size was limited by the page table format requested
 	bool bGpuTextureLimit; // True if the physical size was limited by the maximum GPU texture size
diff --git a/Engine/Source/Runtime/Renderer/Private/VT/VirtualTextureSystem.cpp b/Engine/Source/Runtime/Renderer/Private/VT/VirtualTextureSystem.cpp
--- a/Engine/Source/Runtime/Renderer/Private/VT/VirtualTextureSystem.cpp
+++ b/Engine/Source/Runtime/Renderer/Private/VT/VirtualTextureSystem.cpp
@@ -20,6 +20,8 @@
 #include "VT/VirtualTextureScalability.h"
 #include "VT/VirtualTextureSpace.h"
 
+#define LOCTEXT_NAMESPACE "VirtualTexture"
+
 DECLARE_CYCLE_STAT(TEXT("VirtualTextureSystem Update"), STAT_VirtualTextureSystem_Update, STATGROUP_VirtualTexturing);
 
 DECLARE_CYCLE_STAT(TEXT("Gather Requests"), STAT_ProcessRequests_Gather, STATGROUP_VirtualTexturing);
@@ -224,7 +226,7 @@ FVirtualTextureSystem& FVirtualTextureSystem::Get()
 }
 
 FVirtualTextureSystem::FVirtualTextureSystem()
-	: Frame(1u) // Need to start on Frame 1, otherwise the first call to update will fail to allocate any pages
+	: Frame(1024u) // Need to start on a high enough value that we'll be able to allocate pages
 	, bFlushCaches(false)
 	, FlushCachesCommand(TEXT("r.VT.Flush"), TEXT("Flush all the physical caches in the VT system."),
 		FConsoleCommandDelegate::CreateRaw(this, &FVirtualTextureSystem::FlushCachesFromConsole))
@@ -237,10 +239,17 @@ FVirtualTextureSystem::FVirtualTextureSystem()
 		FConsoleCommandDelegate::CreateRaw(this, &FVirtualTextureSystem::SaveAllocatorImagesFromConsole))
 #endif
 {
+#if !UE_BUILD_SHIPPING
+	FCoreDelegates::OnGetOnScreenMessages.AddRaw(this, &FVirtualTextureSystem::GetOnScreenMessages);
+#endif
 }
 
 FVirtualTextureSystem::~FVirtualTextureSystem()
 {
+#if !UE_BUILD_SHIPPING
+	FCoreDelegates::OnGetOnScreenMessages.RemoveAll(this);
+#endif
+
 	DestroyPendingVirtualTextures();
 
 	check(AllocatedVTs.Num() == 0);
@@ -406,6 +415,14 @@ void FVirtualTextureSystem::ListPhysicalPoolsFromConsole()
 	UE_LOG(LogConsoleResponse, Display, TEXT("TotalPhysicalMemory: %fMB"), (double)TotalPhysicalMemory / 1024.0 / 1024.0);
 }
 
+#if !UE_BUILD_SHIPPING
+void FVirtualTextureSystem::GetOnScreenMessages(FCoreDelegates::FSeverityMessageMap& OutMessages)
+{
+	FScopeLock Locker(&OnScreenMessageLock);
+	OutMessages.Append(OnScreenMessages);
+}
+#endif
+
 #if WITH_EDITOR
 void FVirtualTextureSystem::SaveAllocatorImagesFromConsole()
 {
@@ -1936,6 +1953,7 @@ void FVirtualTextureSystem::SubmitRequests(FRDGBuilder& GraphBuilder, ERHIFeatur
 		uint32 NumStacksProduced = 0u;
 		uint32 NumPagesProduced = 0u;
 		uint32 NumPageAllocateFails = 0u;
+
 		for (uint32 RequestIndex = 0u; RequestIndex < RequestList->GetNumLoadRequests(); ++RequestIndex)
 		{
 			const bool bLockTile = RequestList->IsLocked(RequestIndex);
@@ -2018,17 +2036,7 @@ void FVirtualTextureSystem::SubmitRequests(FRDGBuilder& GraphBuilder, ERHIFeatur
 						}
 						else
 						{
-							static bool bWarnedOnce = false;
-							if (!bWarnedOnce)
-							{
-								UE_LOG(LogConsoleResponse, Display, TEXT("Failed to allocate VT page from pool %d"), PhysicalSpace->GetID());
-								for (int TextureIndex = 0; TextureIndex < PhysicalSpace->GetDescription().NumLayers; ++TextureIndex)
-								{
-									const FPixelFormatInfo& PoolFormatInfo = GPixelFormats[PhysicalSpace->GetFormat(TextureIndex)];
-									UE_LOG(LogConsoleResponse, Display, TEXT("  PF_%s"), PoolFormatInfo.Name);
-								}
-								bWarnedOnce = true;
-							}
+							PhysicalSpace->SetLastFrameOversubscribed(Frame);
 							bProduceTargetValid = false;
 							NumPageAllocateFails++;
 							break;
@@ -2146,6 +2154,35 @@ void FVirtualTextureSystem::SubmitRequests(FRDGBuilder& GraphBuilder, ERHIFeatur
 		INC_DWORD_STAT_BY(STAT_NumStacksRequested, RequestList->GetNumLoadRequests());
 		INC_DWORD_STAT_BY(STAT_NumStacksProduced, NumStacksProduced);
 		INC_DWORD_STAT_BY(STAT_NumPageAllocateFails, NumPageAllocateFails);
+
+#if !UE_BUILD_SHIPPING
+		{
+			FScopeLock Locker(&OnScreenMessageLock);
+			OnScreenMessages.Reset();
+
+			for (int32 SpaceIndex = 0u; SpaceIndex < PhysicalSpaces.Num(); ++SpaceIndex)
+			{
+				const FVirtualTexturePhysicalSpace* PhysicalSpace = PhysicalSpaces[SpaceIndex];
+				// Once a pool fails an allocation, keep displaying the message for a minimum number of frames, to avoid the message quickly flickering on/off
+				if (Frame <= PhysicalSpace->GetLastFrameOversubscribed() + 60u)
+				{
+					const FVTPhysicalSpaceDescription& Desc = PhysicalSpace->GetDescription();
+
+					FString FormatString;
+					for (uint32 Layer = 0u; Layer < Desc.NumLayers; ++Layer)
+					{
+						FormatString += GPixelFormats[Desc.Format[Layer]].Name;
+						if (Layer + 1u < Desc.NumLayers)
+						{
+							FormatString += TEXT(", ");
+						}
+					}
+
+					OnScreenMessages.Add(FCoreDelegates::EOnScreenMessageSeverity::Warning, FText::Format(LOCTEXT("VTOversubscribed", "VT Pool [{0}] is oversubscribed"), FText::FromString(FormatString)));
+				}
+			}
+		}
+#endif // !UE_BUILD_SHIPPING
 	}
 
 	{
@@ -2326,3 +2363,5 @@ void FVirtualTextureSystem::ReleasePendingResources()
 {
 	ReleasePendingSpaces();
 }
+
+#undef LOCTEXT_NAMESPACE
\ No newline at end of file
diff --git a/Engine/Source/Runtime/Renderer/Private/VT/VirtualTextureSystem.h b/Engine/Source/Runtime/Renderer/Private/VT/VirtualTextureSystem.h
--- a/Engine/Source/Runtime/Renderer/Private/VT/VirtualTextureSystem.h
+++ b/Engine/Source/Runtime/Renderer/Private/VT/VirtualTextureSystem.h
@@ -3,6 +3,9 @@
 #pragma once
 
 #include "CoreMinimal.h"
+#if !UE_BUILD_SHIPPING
+#include "Misc/CoreDelegates.h"
+#endif
 #include "RHI.h"
 #include "RendererInterface.h"
 #include "Templates/UniquePtr.h"
@@ -102,6 +105,13 @@ private:
 
 	void UpdateCSVStats() const;
 
+#if !UE_BUILD_SHIPPING
+	void GetOnScreenMessages(FCoreDelegates::FSeverityMessageMap& OutMessages);
+
+	FCoreDelegates::FSeverityMessageMap OnScreenMessages;
+	FCriticalSection OnScreenMessageLock;
+#endif
+
 	uint32	Frame;
 
 	static const uint32 MaxNumTasks = 16;
