Support CSM Caching.

diff --git a/Engine/Shaders/Private/CopyShadowMaps.usf b/Engine/Shaders/Private/CopyShadowMaps.usf
--- a/Engine/Shaders/Private/CopyShadowMaps.usf
+++ b/Engine/Shaders/Private/CopyShadowMaps.usf
@@ -96,3 +96,18 @@ void CopyCubeDepthPS(
 	OutColor = 0;
 	OutDepth = TextureCubeSampleDepthLevel(ShadowDepthCubeTexture, ShadowDepthSampler, CubeCoordinates, 0);
 }
+
+float4 DepthOffsetScale;
+
+void Scrolling2DDepthPS(
+	FScreenVertexOutput Input,
+	out float OutDepth : SV_DEPTH,
+	out float4 OutColor : SV_Target0
+)
+{
+	OutColor = 0;
+	
+	OutDepth = Texture2DSampleLevel(ShadowDepthTexture, ShadowDepthSampler, Input.UV, 0).x;
+
+	OutDepth = OutDepth < 1.0f ? OutDepth * DepthOffsetScale.y + DepthOffsetScale.x : 1.0f;
+}
\ No newline at end of file
diff --git a/Engine/Source/Runtime/Engine/Classes/Engine/RendererSettings.h b/Engine/Source/Runtime/Engine/Classes/Engine/RendererSettings.h
--- a/Engine/Source/Runtime/Engine/Classes/Engine/RendererSettings.h
+++ b/Engine/Source/Runtime/Engine/Classes/Engine/RendererSettings.h
@@ -580,6 +580,11 @@ class ENGINE_API URendererSettings : public UDeveloperSettings
 		ConfigRestartRequired = true))
 	uint32 bEarlyZPassOnlyMaterialMasking : 1;
 
+	UPROPERTY(config, EditAnywhere, Category = Optimizations, meta = (
+		ConsoleVariable = "r.Shadow.CSMCaching", DisplayName = "Enable CSM Caching",
+		ToolTip = "Enable caching CSM to reduce draw calls for casting CSM and probably improve performance."))
+		uint32 bEnableCSMCaching : 1;
+
 	UPROPERTY(config, EditAnywhere, Category=MiscLighting, meta=(
 		ConsoleVariable="r.DBuffer",DisplayName="DBuffer Decals",
 		ToolTip="Whether to accumulate decal properties to a buffer before the base pass.  DBuffer decals correctly affect lightmap and sky lighting, unlike regular deferred decals.  DBuffer enabled forces a full prepass.  Changing this setting requires restarting the editor.",
diff --git a/Engine/Source/Runtime/Renderer/Private/LightSceneInfo.cpp b/Engine/Source/Runtime/Renderer/Private/LightSceneInfo.cpp
--- a/Engine/Source/Runtime/Renderer/Private/LightSceneInfo.cpp
+++ b/Engine/Source/Runtime/Renderer/Private/LightSceneInfo.cpp
@@ -385,6 +385,35 @@ void FLightSceneInfo::ConditionalUpdateMobileMovablePointLightUniformBuffer(cons
 	}
 }
 
+bool FLightSceneInfo::ShouldRecordShadowSubjectsForMobile() const
+{
+	if (Proxy == nullptr)
+	{
+		return false;
+	}
+
+	// record shadow casters if CSM culling is enabled for the light's mobility type and the culling mode requires the list of casters.
+	static auto* MyCVarMobileEnableStaticAndCSMShadowReceivers = IConsoleManager::Get().FindTConsoleVariableDataInt(TEXT("r.Mobile.EnableStaticAndCSMShadowReceivers"));
+	const bool bCombinedStaticAndCSMEnabled = MyCVarMobileEnableStaticAndCSMShadowReceivers->GetValueOnAnyThread() != 0;
+
+	static auto* CVarMobileEnableMovableLightCSMShaderCulling = IConsoleManager::Get().FindTConsoleVariableDataInt(TEXT("r.Mobile.EnableMovableLightCSMShaderCulling"));
+	const bool bMobileEnableMovableLightCSMShaderCulling = CVarMobileEnableMovableLightCSMShaderCulling->GetValueOnAnyThread() == 1;
+
+	static auto* CVarMobileCSMShaderCullingMethod = IConsoleManager::Get().FindTConsoleVariableDataInt(TEXT("r.Mobile.Shadow.CSMShaderCullingMethod"));
+	const uint32 MobileCSMCullingMode = CVarMobileCSMShaderCullingMethod->GetValueOnAnyThread() & 0xF;
+
+	const FLightSceneProxy* LightSceneProxy = Proxy;
+
+	bool bRenderMovableDirectionalLightCSM = LightSceneProxy->IsMovable() && ShouldRenderViewIndependentWholeSceneShadows();
+
+	bool bLightHasCombinedStaticAndCSMEnabled = bCombinedStaticAndCSMEnabled && LightSceneProxy->UseCSMForDynamicObjects();
+	bool bMovableLightUsingCSM = bMobileEnableMovableLightCSMShaderCulling && bRenderMovableDirectionalLightCSM;
+
+	bool bShouldRecordShadowSubjectsForMobile = (MobileCSMCullingMode == 2 || MobileCSMCullingMode == 3) && (bLightHasCombinedStaticAndCSMEnabled || bMovableLightUsingCSM);
+
+	return bShouldRecordShadowSubjectsForMobile;
+}
+
 /** Determines whether two bounding spheres intersect. */
 FORCEINLINE bool AreSpheresNotIntersecting(
 	const VectorRegister& A_XYZ,
diff --git a/Engine/Source/Runtime/Renderer/Private/LightSceneInfo.h b/Engine/Source/Runtime/Renderer/Private/LightSceneInfo.h
--- a/Engine/Source/Runtime/Renderer/Private/LightSceneInfo.h
+++ b/Engine/Source/Runtime/Renderer/Private/LightSceneInfo.h
@@ -348,6 +348,8 @@ public:
 
 	// Update the mobile movable point light uniform buffer before it is used for mobile base pass rendering.
 	void ConditionalUpdateMobileMovablePointLightUniformBuffer(const class FSceneRenderer* SceneRenderer);
+
+	bool ShouldRecordShadowSubjectsForMobile() const;
 };
 
 /** Defines how the light is stored in the scene's light octree. */
diff --git a/Engine/Source/Runtime/Renderer/Private/MobileBasePassRendering.cpp b/Engine/Source/Runtime/Renderer/Private/MobileBasePassRendering.cpp
--- a/Engine/Source/Runtime/Renderer/Private/MobileBasePassRendering.cpp
+++ b/Engine/Source/Runtime/Renderer/Private/MobileBasePassRendering.cpp
@@ -252,7 +252,7 @@ void SetupMobileDirectionalLightUniformParameters(
 			{
 				const FProjectedShadowInfo* ShadowInfo = DirectionalLightShadowInfos[i];
 
-				if (ShadowInfo->ShadowDepthView && !ShadowInfo->bRayTracedDistanceField)
+				if (ShadowInfo->ShadowDepthView && !ShadowInfo->bRayTracedDistanceField && ShadowInfo->CacheMode != SDCM_StaticPrimitivesOnly)
 				{
 					if (OutShadowIndex == 0)
 					{
diff --git a/Engine/Source/Runtime/Renderer/Private/RendererScene.cpp b/Engine/Source/Runtime/Renderer/Private/RendererScene.cpp
--- a/Engine/Source/Runtime/Renderer/Private/RendererScene.cpp
+++ b/Engine/Source/Runtime/Renderer/Private/RendererScene.cpp
@@ -2358,13 +2358,16 @@ int64 FScene::GetCachedWholeSceneShadowMapsSize() const
 {
 	int64 CachedShadowmapMemory = 0;
 
-	for (TMap<int32, FCachedShadowMapData>::TConstIterator CachedShadowMapIt(CachedShadowMaps); CachedShadowMapIt; ++CachedShadowMapIt)
+	for (TMap<int32, TArray<FCachedShadowMapData>>::TConstIterator CachedShadowMapIt(CachedShadowMaps); CachedShadowMapIt; ++CachedShadowMapIt)
 	{
-		const FCachedShadowMapData& ShadowMapData = CachedShadowMapIt.Value();
+		const TArray<FCachedShadowMapData>& ShadowMapDatas = CachedShadowMapIt.Value();
 
-		if (ShadowMapData.ShadowMap.IsValid())
+		for (const auto& ShadowMapData : ShadowMapDatas)
 		{
-			CachedShadowmapMemory += ShadowMapData.ShadowMap.ComputeMemorySize();
+			if (ShadowMapData.ShadowMap.IsValid())
+			{
+				CachedShadowmapMemory += ShadowMapData.ShadowMap.ComputeMemorySize();
+			}
 		}
 	}
 
@@ -3959,6 +3962,19 @@ void FScene::UpdateAllPrimitiveSceneInfos(FRDGBuilder& GraphBuilder, bool bAsync
 						#endif
 							TBitArraySwapElements(PrimitivesNeedingStaticMeshUpdate, DestIndex, SourceIndex);
 
+							for (TMap<int32, TArray<FCachedShadowMapData>>::TIterator CachedShadowMapIt(CachedShadowMaps); CachedShadowMapIt; ++CachedShadowMapIt)
+							{
+								TArray<FCachedShadowMapData>& ShadowMapDatas = CachedShadowMapIt.Value();
+
+								for (auto& ShadowMapData : ShadowMapDatas)
+								{
+									if (ShadowMapData.StaticShadowSubjectMap.Num() > 0)
+									{
+										TBitArraySwapElements(ShadowMapData.StaticShadowSubjectMap, DestIndex, SourceIndex);
+									}
+								}
+							}
+
 							GPUScene.AddPrimitiveToUpdate(SourceIndex);
 							GPUScene.AddPrimitiveToUpdate(DestIndex);
 
@@ -4056,6 +4072,36 @@ void FScene::UpdateAllPrimitiveSceneInfos(FRDGBuilder& GraphBuilder, bool bAsync
 				LumenSceneData->RemovePrimitive(PrimitiveSceneInfo, PrimitiveIndex);
 
 				DeletedSceneInfos.Add(PrimitiveSceneInfo);
+
+				for (const FLightSceneInfo* LightSceneInfo : DirectionalLights)
+				{
+					TArray<FCachedShadowMapData>* CachedShadowMapDatas = GetCachedShadowMapDatas(LightSceneInfo->Id);
+
+					if (CachedShadowMapDatas)
+					{
+						for (auto& CachedShadowMapData : *CachedShadowMapDatas)
+						{
+							if (CachedShadowMapData.StaticShadowSubjectMap[PrimitiveIndex] == true)
+							{
+								CachedShadowMapData.InvalidateCachedShadow();
+							}
+						}
+					}
+				}
+			}
+
+			for (TMap<int32, TArray<FCachedShadowMapData>>::TIterator CachedShadowMapIt(CachedShadowMaps); CachedShadowMapIt; ++CachedShadowMapIt)
+			{
+				TArray<FCachedShadowMapData>& ShadowMapDatas = CachedShadowMapIt.Value();
+
+				for (auto& ShadowMapData : ShadowMapDatas)
+				{
+					if (ShadowMapData.StaticShadowSubjectMap.Num() > 0)
+					{
+						ShadowMapData.StaticShadowSubjectMap.RemoveAt(SourceIndex, RemoveCount);
+						check(Primitives.Num() == ShadowMapData.StaticShadowSubjectMap.Num());
+					}
+				}
 			}
 			RemovedLocalPrimitiveSceneInfos.RemoveAt(StartIndex, RemovedLocalPrimitiveSceneInfos.Num() - StartIndex, false);
 		}
@@ -4085,6 +4131,19 @@ void FScene::UpdateAllPrimitiveSceneInfos(FRDGBuilder& GraphBuilder, bool bAsync
 			PrimitiveRayTracingFlags.Reserve(PrimitiveRayTracingFlags.Num() + AddedLocalPrimitiveSceneInfos.Num());
 		#endif
 			PrimitivesNeedingStaticMeshUpdate.Reserve(PrimitivesNeedingStaticMeshUpdate.Num() + AddedLocalPrimitiveSceneInfos.Num());
+
+			for (TMap<int32, TArray<FCachedShadowMapData>>::TIterator CachedShadowMapIt(CachedShadowMaps); CachedShadowMapIt; ++CachedShadowMapIt)
+			{
+				TArray<FCachedShadowMapData>& ShadowMapDatas = CachedShadowMapIt.Value();
+
+				for (auto& ShadowMapData : ShadowMapDatas)
+				{
+					if (ShadowMapData.StaticShadowSubjectMap.Num() > 0)
+					{
+						ShadowMapData.StaticShadowSubjectMap.Reserve(ShadowMapData.StaticShadowSubjectMap.Num() + AddedLocalPrimitiveSceneInfos.Num());
+					}
+				}
+			}
 		}
 
 		while (AddedLocalPrimitiveSceneInfos.Num())
@@ -4121,6 +4180,19 @@ void FScene::UpdateAllPrimitiveSceneInfos(FRDGBuilder& GraphBuilder, bool bAsync
 			#endif
 				PrimitivesNeedingStaticMeshUpdate.Add(false);
 
+				for (TMap<int32, TArray<FCachedShadowMapData>>::TIterator CachedShadowMapIt(CachedShadowMaps); CachedShadowMapIt; ++CachedShadowMapIt)
+				{
+					TArray<FCachedShadowMapData>& ShadowMapDatas = CachedShadowMapIt.Value();
+
+					for (auto& ShadowMapData : ShadowMapDatas)
+					{
+						if (ShadowMapData.StaticShadowSubjectMap.Num() > 0)
+						{
+							ShadowMapData.StaticShadowSubjectMap.Add(false);
+						}
+					}
+				}
+
 				const int32 SourceIndex = PrimitiveSceneProxies.Num() - 1;
 				PrimitiveSceneInfo->PackedIndex = SourceIndex;
 
@@ -4207,6 +4279,19 @@ void FScene::UpdateAllPrimitiveSceneInfos(FRDGBuilder& GraphBuilder, bool bAsync
 						#endif
 							TBitArraySwapElements(PrimitivesNeedingStaticMeshUpdate, DestIndex, SourceIndex);
 
+							for (TMap<int32, TArray<FCachedShadowMapData>>::TIterator CachedShadowMapIt(CachedShadowMaps); CachedShadowMapIt; ++CachedShadowMapIt)
+							{
+								TArray<FCachedShadowMapData>& ShadowMapDatas = CachedShadowMapIt.Value();
+
+								for (auto& ShadowMapData : ShadowMapDatas)
+								{
+									if (ShadowMapData.StaticShadowSubjectMap.Num() > 0)
+									{
+										TBitArraySwapElements(ShadowMapData.StaticShadowSubjectMap, DestIndex, SourceIndex);
+									}
+								}
+							}
+
 							GPUScene.AddPrimitiveToUpdate(DestIndex);
 							LumenSceneData->AddPrimitiveToUpdate(DestIndex);
 						}
diff --git a/Engine/Source/Runtime/Renderer/Private/SceneCore.cpp b/Engine/Source/Runtime/Renderer/Private/SceneCore.cpp
--- a/Engine/Source/Runtime/Renderer/Private/SceneCore.cpp
+++ b/Engine/Source/Runtime/Renderer/Private/SceneCore.cpp
@@ -310,11 +310,14 @@ void FLightPrimitiveInteraction::FlushCachedShadowMapData()
 	{
 		if (bCastShadow && !PrimitiveSceneInfo->Proxy->IsMeshShapeOftenMoving())
 		{
-			FCachedShadowMapData* CachedShadowMapData = PrimitiveSceneInfo->Scene->CachedShadowMaps.Find(LightSceneInfo->Id);
+			TArray<FCachedShadowMapData>* CachedShadowMapDatas = PrimitiveSceneInfo->Scene->GetCachedShadowMapDatas(LightSceneInfo->Id);
 
-			if (CachedShadowMapData)
+			if (CachedShadowMapDatas)
 			{
-				CachedShadowMapData->ShadowMap.Release();
+				for (auto& CachedShadowMapData : *CachedShadowMapDatas)
+				{
+					CachedShadowMapData.InvalidateCachedShadow();
+				}
 			}
 		}
 	}
diff --git a/Engine/Source/Runtime/Renderer/Private/ScenePrivate.h b/Engine/Source/Runtime/Renderer/Private/ScenePrivate.h
--- a/Engine/Source/Runtime/Renderer/Private/ScenePrivate.h
+++ b/Engine/Source/Runtime/Renderer/Private/ScenePrivate.h
@@ -2460,6 +2460,9 @@ private:
 	void HideNodeChildren(FSceneViewState* ViewState, FLODSceneNode& Node);
 };
 
+// Enable the DEBUG_CSM_CACHING to make the debugging CSM caching with RenderDoc more easier
+#define DEBUG_CSM_CACHING 0
+
 class FCachedShadowMapData
 {
 public:
@@ -2469,6 +2472,28 @@ public:
 	bool bCachedShadowMapHasPrimitives;
 	bool bCachedShadowMapHasNaniteGeometry;
 
+	/**
+	* The static meshes cast shadow on this cached csm
+	*/
+	TBitArray<> StaticShadowSubjectMap;
+
+	FIntPoint ShadowBufferResolution;
+	FVector PreShadowTranslation;
+	float MaxSubjectZ;
+	float MinSubjectZ;
+
+	/**
+	* The extra static meshes cast shadow in last frame, if it exceeds the r.Shadow.MaxCSMScrollingStaticShadowSubjects, the cached csm should be rebuilt.
+	*/
+	int32 LastFrameExtraStaticShadowSubjects;
+
+	void InvalidateCachedShadow()
+	{
+		ShadowMap.Release();
+
+		StaticShadowSubjectMap.SetRange(0, StaticShadowSubjectMap.Num(), false);
+	}
+
 	FCachedShadowMapData(const FWholeSceneProjectedShadowInitializer& InInitializer, float InLastUsedTime) :
 		Initializer(InInitializer),
 		LastUsedTime(InLastUsedTime),
@@ -2739,7 +2764,7 @@ public:
 	FNaniteMaterialTables MaterialTables[ENaniteMeshPass::Num];
 
 	/** Map from light id to the cached shadowmap data for that light. */
-	TMap<int32, FCachedShadowMapData> CachedShadowMaps;
+	TMap<int32, TArray<FCachedShadowMapData>> CachedShadowMaps;
 	
 	/** Atlas HZB textures from the previous render. */
 	TArray<TRefCountPtr<IPooledRenderTarget>>	PrevAtlasHZBs;
@@ -3160,6 +3185,38 @@ public:
 
 	void CreateLightPrimitiveInteractionsForPrimitive(FPrimitiveSceneInfo* PrimitiveInfo, bool bAsyncCreateLPIs);
 
+	FORCEINLINE TArray<FCachedShadowMapData>* GetCachedShadowMapDatas(int32 LightID)
+	{
+		return CachedShadowMaps.Find(LightID);
+	}
+
+	FORCEINLINE FCachedShadowMapData& GetCachedShadowMapDataRef(int32 LightID, int32 ShadowMapIndex = 0)
+	{
+		TArray<FCachedShadowMapData>& CachedShadowMapDatas = CachedShadowMaps.FindChecked(LightID);
+
+		checkSlow(ShadowMapIndex >= 0 && ShadowMapIndex < CachedShadowMapDatas.Num());
+
+		return CachedShadowMapDatas[ShadowMapIndex];
+	}
+
+	FORCEINLINE const FCachedShadowMapData& GetCachedShadowMapDataRef(int32 LightID, int32 ShadowMapIndex = 0) const
+	{
+		const TArray<FCachedShadowMapData>& CachedShadowMapDatas = CachedShadowMaps.FindChecked(LightID);
+
+		checkSlow(ShadowMapIndex >= 0 && ShadowMapIndex < CachedShadowMapDatas.Num());
+
+		return CachedShadowMapDatas[ShadowMapIndex];
+	}
+
+	FORCEINLINE const FCachedShadowMapData* GetCachedShadowMapData(int32 LightID, int32 ShadowMapIndex = 0) const
+	{
+		const TArray<FCachedShadowMapData>& CachedShadowMapDatas = CachedShadowMaps.FindChecked(LightID);
+
+		checkSlow(ShadowMapIndex >= 0 && ShadowMapIndex < CachedShadowMapDatas.Num());
+
+		return &CachedShadowMapDatas[ShadowMapIndex];
+	}
+
 	bool IsPrimitiveBeingRemoved(FPrimitiveSceneInfo* PrimitiveSceneInfo) const;
 
 protected:
diff --git a/Engine/Source/Runtime/Renderer/Private/SceneRendering.h b/Engine/Source/Runtime/Renderer/Private/SceneRendering.h
--- a/Engine/Source/Runtime/Renderer/Private/SceneRendering.h
+++ b/Engine/Source/Runtime/Renderer/Private/SceneRendering.h
@@ -2011,13 +2011,14 @@ protected:
 		TArray<FProjectedShadowInfo*, SceneRenderingAllocator>& ShadowInfos, 
 		TArray<FProjectedShadowInfo*, SceneRenderingAllocator>& ShadowInfosThatNeedCulling, 
 		FVisibleLightInfo& VisibleLightInfo, 
-		FLightSceneInfo& LightSceneInfo);
+		FLightSceneInfo& LightSceneInfo,
+		uint32& NumCSMCachesUpdatedThisFrame);
 
 	void AllocateShadowDepthTargets(FRHICommandListImmediate& RHICmdList);
 	
 	void AllocateAtlasedShadowDepthTargets(FRHICommandListImmediate& RHICmdList, TArray<FProjectedShadowInfo*, SceneRenderingAllocator>& Shadows, TArray<FSortedShadowMapAtlas,SceneRenderingAllocator>& OutAtlases);
 
-	void AllocateCachedSpotlightShadowDepthTargets(FRHICommandListImmediate& RHICmdList, TArray<FProjectedShadowInfo*, SceneRenderingAllocator>& CachedShadows);
+	void AllocateCachedShadowDepthTargets(FRHICommandListImmediate& RHICmdList, TArray<FProjectedShadowInfo*, SceneRenderingAllocator>& CachedShadows);
 
 	void AllocateCSMDepthTargets(FRHICommandListImmediate& RHICmdList, const TArray<FProjectedShadowInfo*, SceneRenderingAllocator>& WholeSceneDirectionalShadows, TArray<FSortedShadowMapAtlas,SceneRenderingAllocator>& OutAtlases);
 
diff --git a/Engine/Source/Runtime/Renderer/Private/ShadowDepthRendering.cpp b/Engine/Source/Runtime/Renderer/Private/ShadowDepthRendering.cpp
--- a/Engine/Source/Runtime/Renderer/Private/ShadowDepthRendering.cpp
+++ b/Engine/Source/Runtime/Renderer/Private/ShadowDepthRendering.cpp
@@ -147,12 +147,16 @@ void SetupShadowDepthPassUniformBuffer(
 	const FViewInfo& View,
 	FShadowDepthPassUniformParameters& ShadowDepthPassParameters)
 {
+	static const auto CSMCachingCVar = IConsoleManager::Get().FindTConsoleVariableDataInt(TEXT("r.Shadow.CSMCaching"));
+	const bool bCSMCachingEnabled = CSMCachingCVar && CSMCachingCVar->GetValueOnAnyThread() != 0;
+
 	SetupSceneTextureUniformParameters(GraphBuilder, View.FeatureLevel, ESceneTextureSetupMode::None, ShadowDepthPassParameters.SceneTextures);
 
 	ShadowDepthPassParameters.ProjectionMatrix = FTranslationMatrix(ShadowInfo->PreShadowTranslation - View.ViewMatrices.GetPreViewTranslation()) * ShadowInfo->TranslatedWorldToClipOuterMatrix;	
 	ShadowDepthPassParameters.ViewMatrix = ShadowInfo->TranslatedWorldToView;
 
-	ShadowDepthPassParameters.ShadowParams = FVector4(ShadowInfo->GetShaderDepthBias(), ShadowInfo->GetShaderSlopeDepthBias(), ShadowInfo->GetShaderMaxSlopeDepthBias(), ShadowInfo->bOnePassPointLightShadow ? 1 : ShadowInfo->InvMaxSubjectDepth);
+	// Disable the SlopeDepthBias because we couldn't reconstruct the depth offset if it is not 0.0f when scrolling the cached shadow map.
+	ShadowDepthPassParameters.ShadowParams = FVector4(ShadowInfo->GetShaderDepthBias(), bCSMCachingEnabled ? 0.0f : ShadowInfo->GetShaderSlopeDepthBias(), ShadowInfo->GetShaderMaxSlopeDepthBias(), ShadowInfo->bOnePassPointLightShadow ? 1 : ShadowInfo->InvMaxSubjectDepth);
 	ShadowDepthPassParameters.bClampToNearPlane = ShadowInfo->ShouldClampToNearPlane() ? 1.0f : 0.0f;
 
 	if (ShadowInfo->bOnePassPointLightShadow)
@@ -900,6 +904,29 @@ public:
 
 IMPLEMENT_GLOBAL_SHADER(FCopyShadowMaps2DPS, "/Engine/Private/CopyShadowMaps.usf", "Copy2DDepthPS", SF_Pixel);
 
+/** */
+class FScrollingShadowMaps2DPS : public FGlobalShader
+{
+public:
+	DECLARE_GLOBAL_SHADER(FScrollingShadowMaps2DPS);
+	SHADER_USE_PARAMETER_STRUCT(FScrollingShadowMaps2DPS, FGlobalShader);
+
+	BEGIN_SHADER_PARAMETER_STRUCT(FParameters, )
+		SHADER_PARAMETER_STRUCT_REF(FViewUniformShaderParameters, View)
+		SHADER_PARAMETER_RDG_TEXTURE(Texture2D, ShadowDepthTexture)
+		SHADER_PARAMETER_SAMPLER(SamplerState, ShadowDepthSampler)
+		SHADER_PARAMETER(FVector4, DepthOffsetScale)
+		RENDER_TARGET_BINDING_SLOTS()
+	END_SHADER_PARAMETER_STRUCT()
+
+	static bool ShouldCompilePermutation(const FGlobalShaderPermutationParameters& Parameters)
+	{
+		return true;
+	}
+};
+
+IMPLEMENT_GLOBAL_SHADER(FScrollingShadowMaps2DPS, "/Engine/Private/CopyShadowMaps.usf", "Scrolling2DDepthPS", SF_Pixel);
+
 void FProjectedShadowInfo::CopyCachedShadowMap(
 	FRDGBuilder& GraphBuilder,
 	const FViewInfo& View,
@@ -907,8 +934,8 @@ void FProjectedShadowInfo::CopyCachedShadowMap(
 	const FRenderTargetBindingSlots& RenderTargetBindingSlots,
 	const FMeshPassProcessorRenderState& DrawRenderState)
 {
-	check(CacheMode == SDCM_MovablePrimitivesOnly);
-	const FCachedShadowMapData& CachedShadowMapData = SceneRenderer->Scene->CachedShadowMaps.FindChecked(GetLightSceneInfo().Id);
+	check(CacheMode == SDCM_MovablePrimitivesOnly || CacheMode == SDCM_CSMScrolling);
+	const FCachedShadowMapData& CachedShadowMapData = SceneRenderer->Scene->GetCachedShadowMapDataRef(GetLightSceneInfo().Id, FMath::Max(CascadeSettings.ShadowSplitIndex, 0));
 
 	if (CachedShadowMapData.bCachedShadowMapHasPrimitives && CachedShadowMapData.ShadowMap.IsValid())
 	{
@@ -927,98 +954,162 @@ void FProjectedShadowInfo::CopyCachedShadowMap(
 		extern TGlobalResource<FFilterVertexDeclaration> GFilterVertexDeclaration;
 		GraphicsPSOInit.BoundShaderState.VertexDeclarationRHI = GFilterVertexDeclaration.VertexDeclarationRHI;
 
-		if (bOnePassPointLightShadow)
+		if (CacheMode == SDCM_MovablePrimitivesOnly)
 		{
-			TShaderRef<FScreenVS> ScreenVertexShader;
-			TShaderMapRef<FCopyShadowMapsCubePS> PixelShader(View.ShaderMap);
-			GraphicsPSOInit.BoundShaderState.PixelShaderRHI = PixelShader.GetPixelShader();
-
-			int32 InstanceCount = 1;
-
-#if PLATFORM_SUPPORTS_GEOMETRY_SHADERS
-			if (RHISupportsGeometryShaders(GShaderPlatformForFeatureLevel[SceneRenderer->FeatureLevel]))
-			{
-				TShaderMapRef<TScreenVSForGS<false>> VertexShader(View.ShaderMap);
-				TShaderMapRef<FCopyShadowMapsCubeGS> GeometryShader(View.ShaderMap);
-				GraphicsPSOInit.BoundShaderState.VertexShaderRHI = VertexShader.GetVertexShader();
-				GraphicsPSOInit.BoundShaderState.GeometryShaderRHI = GeometryShader.GetGeometryShader();
-				ScreenVertexShader = VertexShader;
-			}
-			else
-#endif
+			if (bOnePassPointLightShadow)
 			{
-				check(RHISupportsVertexShaderLayer(GShaderPlatformForFeatureLevel[SceneRenderer->FeatureLevel]));
-				TShaderMapRef<TScreenVSForGS<true>> VertexShader(View.ShaderMap);
-				GraphicsPSOInit.BoundShaderState.VertexShaderRHI = VertexShader.GetVertexShader();
-				ScreenVertexShader = VertexShader;
+				TShaderRef<FScreenVS> ScreenVertexShader;
+				TShaderMapRef<FCopyShadowMapsCubePS> PixelShader(View.ShaderMap);
+				GraphicsPSOInit.BoundShaderState.PixelShaderRHI = PixelShader.GetPixelShader();
 
-				InstanceCount = 6;
-			}
+				int32 InstanceCount = 1;
 
-			auto* PassParameters = GraphBuilder.AllocParameters<FCopyShadowMapsCubePS::FParameters>();
-			PassParameters->RenderTargets = RenderTargetBindingSlots;
-			PassParameters->ShadowDepthCubeTexture = ShadowDepthTexture;
-			PassParameters->ShadowDepthSampler = TStaticSamplerState<SF_Point, AM_Clamp, AM_Clamp, AM_Clamp>::GetRHI();
+	#if PLATFORM_SUPPORTS_GEOMETRY_SHADERS
+				if (RHISupportsGeometryShaders(GShaderPlatformForFeatureLevel[SceneRenderer->FeatureLevel]))
+				{
+					TShaderMapRef<TScreenVSForGS<false>> VertexShader(View.ShaderMap);
+					TShaderMapRef<FCopyShadowMapsCubeGS> GeometryShader(View.ShaderMap);
+					GraphicsPSOInit.BoundShaderState.VertexShaderRHI = VertexShader.GetVertexShader();
+					GraphicsPSOInit.BoundShaderState.GeometryShaderRHI = GeometryShader.GetGeometryShader();
+					ScreenVertexShader = VertexShader;
+				}
+				else
+	#endif
+				{
+					check(RHISupportsVertexShaderLayer(GShaderPlatformForFeatureLevel[SceneRenderer->FeatureLevel]));
+					TShaderMapRef<TScreenVSForGS<true>> VertexShader(View.ShaderMap);
+					GraphicsPSOInit.BoundShaderState.VertexShaderRHI = VertexShader.GetVertexShader();
+					ScreenVertexShader = VertexShader;
 
-			GraphBuilder.AddPass(
-				RDG_EVENT_NAME("CopyCachedShadowMap"),
-				PassParameters,
-				ERDGPassFlags::Raster,
-				[this, ScreenVertexShader, PixelShader, GraphicsPSOInit, PassParameters, ShadowDepthExtent, InstanceCount, StencilRef](FRHICommandList& RHICmdList) mutable
-			{
-				SetStateForView(RHICmdList);
-				RHICmdList.ApplyCachedRenderTargets(GraphicsPSOInit);
-				SetGraphicsPipelineState(RHICmdList, GraphicsPSOInit);
-				SetShaderParameters(RHICmdList, PixelShader, PixelShader.GetPixelShader(), *PassParameters);
-				RHICmdList.SetStencilRef(StencilRef);
+					InstanceCount = 6;
+				}
 
-				FIntPoint ResolutionWithBorder = FIntPoint(ResolutionX + 2 * BorderSize, ResolutionY + 2 * BorderSize);
+				auto* PassParameters = GraphBuilder.AllocParameters<FCopyShadowMapsCubePS::FParameters>();
+				PassParameters->RenderTargets = RenderTargetBindingSlots;
+				PassParameters->ShadowDepthCubeTexture = ShadowDepthTexture;
+				PassParameters->ShadowDepthSampler = TStaticSamplerState<SF_Point, AM_Clamp, AM_Clamp, AM_Clamp>::GetRHI();
 
-				DrawRectangle(
-					RHICmdList,
-					0, 0,
-					ResolutionWithBorder.X, ResolutionWithBorder.Y,
-					0, 0,
-					ResolutionWithBorder.X, ResolutionWithBorder.Y,
-					ResolutionWithBorder,
-					ShadowDepthExtent,
-					ScreenVertexShader,
-					EDRF_Default,
-					InstanceCount);
-			});
+				GraphBuilder.AddPass(
+					RDG_EVENT_NAME("CopyCachedShadowMap"),
+					PassParameters,
+					ERDGPassFlags::Raster,
+					[this, ScreenVertexShader, PixelShader, GraphicsPSOInit, PassParameters, ShadowDepthExtent, InstanceCount, StencilRef](FRHICommandList& RHICmdList) mutable
+				{
+					SetStateForView(RHICmdList);
+					RHICmdList.ApplyCachedRenderTargets(GraphicsPSOInit);
+					SetGraphicsPipelineState(RHICmdList, GraphicsPSOInit);
+					SetShaderParameters(RHICmdList, PixelShader, PixelShader.GetPixelShader(), *PassParameters);
+					RHICmdList.SetStencilRef(StencilRef);
+
+					FIntPoint ResolutionWithBorder = FIntPoint(ResolutionX + 2 * BorderSize, ResolutionY + 2 * BorderSize);
+
+					DrawRectangle(
+						RHICmdList,
+						0, 0,
+						ResolutionWithBorder.X, ResolutionWithBorder.Y,
+						0, 0,
+						ResolutionWithBorder.X, ResolutionWithBorder.Y,
+						ResolutionWithBorder,
+						ShadowDepthExtent,
+						ScreenVertexShader,
+						EDRF_Default,
+						InstanceCount);
+				});
+			}
+			else
+			{
+				TShaderMapRef<FScreenVS> ScreenVertexShader(View.ShaderMap);
+				TShaderMapRef<FCopyShadowMaps2DPS> PixelShader(View.ShaderMap);
+				GraphicsPSOInit.BoundShaderState.VertexShaderRHI = ScreenVertexShader.GetVertexShader();
+				GraphicsPSOInit.BoundShaderState.PixelShaderRHI = PixelShader.GetPixelShader();
+
+				auto* PassParameters = GraphBuilder.AllocParameters<FCopyShadowMaps2DPS::FParameters>();
+				PassParameters->RenderTargets = RenderTargetBindingSlots;
+				PassParameters->ShadowDepthTexture = ShadowDepthTexture;
+				PassParameters->ShadowDepthSampler = TStaticSamplerState<SF_Point, AM_Clamp, AM_Clamp, AM_Clamp>::GetRHI();
+
+				GraphBuilder.AddPass(
+					RDG_EVENT_NAME("CopyCachedShadowMap"),
+					PassParameters,
+					ERDGPassFlags::Raster,
+					[this, ScreenVertexShader, PixelShader, GraphicsPSOInit, PassParameters, ShadowDepthExtent, StencilRef](FRHICommandList& RHICmdList) mutable
+				{
+					SetStateForView(RHICmdList);
+					RHICmdList.ApplyCachedRenderTargets(GraphicsPSOInit);
+					SetGraphicsPipelineState(RHICmdList, GraphicsPSOInit);
+					SetShaderParameters(RHICmdList, PixelShader, PixelShader.GetPixelShader(), *PassParameters);
+					RHICmdList.SetStencilRef(StencilRef);
+
+					FIntPoint ResolutionWithBorder = FIntPoint(ResolutionX + 2 * BorderSize, ResolutionY + 2 * BorderSize);
+
+					DrawRectangle(
+						RHICmdList,
+						0, 0,
+						ResolutionWithBorder.X, ResolutionWithBorder.Y,
+						0, 0,
+						ResolutionWithBorder.X, ResolutionWithBorder.Y,
+						ResolutionWithBorder,
+						ShadowDepthExtent,
+						ScreenVertexShader,
+						EDRF_Default);
+				});
+			}
 		}
-		else
+		else // CacheMode == SDCM_CSMScrolling
 		{
 			TShaderMapRef<FScreenVS> ScreenVertexShader(View.ShaderMap);
-			TShaderMapRef<FCopyShadowMaps2DPS> PixelShader(View.ShaderMap);
+			TShaderMapRef<FScrollingShadowMaps2DPS> PixelShader(View.ShaderMap);
 			GraphicsPSOInit.BoundShaderState.VertexShaderRHI = ScreenVertexShader.GetVertexShader();
 			GraphicsPSOInit.BoundShaderState.PixelShaderRHI = PixelShader.GetPixelShader();
 
-			auto* PassParameters = GraphBuilder.AllocParameters<FCopyShadowMaps2DPS::FParameters>();
+			auto* PassParameters = GraphBuilder.AllocParameters<FScrollingShadowMaps2DPS::FParameters>();
 			PassParameters->RenderTargets = RenderTargetBindingSlots;
 			PassParameters->ShadowDepthTexture = ShadowDepthTexture;
 			PassParameters->ShadowDepthSampler = TStaticSamplerState<SF_Point, AM_Clamp, AM_Clamp, AM_Clamp>::GetRHI();
+			/** According to the implementation in ShadowDepthVertexShader.usf, the formula is DeviceZ = 1 - ((MaxZ - SceneDepth) / (MaxZ - MinZ)) * InvMaxSubjectDepth + SlopeDepthBias * Slope + ConstantDepthBias;
+			For short C0 = InvMaxSubjectDepth; C1 = 1 + SlopeDepthBias * Slope + ConstantDepthBias;
+			So the SceneDepth0 = MaxZ0 - (C1 - DeviceZ0) * (MaxZ0 - MinZ0) / C0 ;
+			SceneDepth1 = SceneDepth0 + ZOffset;
+			The reconstruct DeviceZ1 = C1 - ((MaxZ1 - SceneDepth1) / (MaxZ1 - MinZ1)) * C0;
+			So DeviceZ1 = DeviceZ0 * (MaxZ0 - MinZ0) / (MaxZ1 - MinZ1) + (C0 * (MaxZ0 + ZOffset - MaxZ1) - C1 * (MaxZ0 - MinZ0)) / (MaxZ1 - MinZ1) + C1;
+			*/
+			float MaxZ0MinusMinZ0 = CachedShadowMapData.MaxSubjectZ - CachedShadowMapData.MinSubjectZ;
+			float MaxZ1MinusMinZ1 = MaxSubjectZ - MinSubjectZ;
+			float MaxZ0PlusZOffsetMinusMaxZ1 = CachedShadowMapData.MaxSubjectZ + CSMScrollingZOffset - MaxSubjectZ;
+			float C1 = 1 + GetShaderDepthBias();
+			PassParameters->DepthOffsetScale = FVector4((InvMaxSubjectDepth * MaxZ0PlusZOffsetMinusMaxZ1 - C1 * MaxZ0MinusMinZ0) / MaxZ1MinusMinZ1 + C1, MaxZ0MinusMinZ0 / MaxZ1MinusMinZ1, 0.0f, 0.0f);
 
 			GraphBuilder.AddPass(
-				RDG_EVENT_NAME("CopyCachedShadowMap"),
+				RDG_EVENT_NAME("ScrollingCachedWholeSceneDirectionalShadowMap"),
 				PassParameters,
 				ERDGPassFlags::Raster,
 				[this, ScreenVertexShader, PixelShader, GraphicsPSOInit, PassParameters, ShadowDepthExtent, StencilRef](FRHICommandList& RHICmdList) mutable
 			{
+				checkSlow(OverlappedUVOnCachedShadowMap != FVector4(-1.0f, -1.0f, -1.0f, -1.0f));
+				checkSlow(OverlappedUVOnCurrentShadowMap != FVector4(-1.0f, -1.0f, -1.0f, -1.0f));
+
+				FIntPoint ResolutionWithBorder = FIntPoint(ResolutionX + 2 * BorderSize, ResolutionY + 2 * BorderSize);
+
+				uint32 UStart = OverlappedUVOnCachedShadowMap.X * ResolutionWithBorder.X + 0.5f;
+				uint32 USize = (OverlappedUVOnCachedShadowMap.Z - OverlappedUVOnCachedShadowMap.X) * ResolutionWithBorder.X + 0.5f;
+
+				uint32 VStart = OverlappedUVOnCachedShadowMap.Y * ResolutionWithBorder.Y + 0.5f;
+				uint32 VSize = (OverlappedUVOnCachedShadowMap.W - OverlappedUVOnCachedShadowMap.Y) * ResolutionWithBorder.Y + 0.5f;
+
+				FIntVector4 OutputViewport = FIntVector4(OverlappedUVOnCurrentShadowMap.X * ResolutionWithBorder.X + 0.5f, OverlappedUVOnCurrentShadowMap.Y * ResolutionWithBorder.Y + 0.5f, OverlappedUVOnCurrentShadowMap.Z * ResolutionWithBorder.X + 0.5f, OverlappedUVOnCurrentShadowMap.W * ResolutionWithBorder.Y + 0.5f);
+
 				SetStateForView(RHICmdList);
 				RHICmdList.ApplyCachedRenderTargets(GraphicsPSOInit);
 				SetGraphicsPipelineState(RHICmdList, GraphicsPSOInit);
 				SetShaderParameters(RHICmdList, PixelShader, PixelShader.GetPixelShader(), *PassParameters);
 				RHICmdList.SetStencilRef(StencilRef);
 
-				FIntPoint ResolutionWithBorder = FIntPoint(ResolutionX + 2 * BorderSize, ResolutionY + 2 * BorderSize);
-
 				DrawRectangle(
 					RHICmdList,
-					0, 0,
-					ResolutionWithBorder.X, ResolutionWithBorder.Y,
-					0, 0,
-					ResolutionWithBorder.X, ResolutionWithBorder.Y,
+					OutputViewport.X, OutputViewport.Y,
+					OutputViewport.Z - OutputViewport.X, OutputViewport.W - OutputViewport.Y,
+					UStart, VStart,
+					USize, VSize,
 					ResolutionWithBorder,
 					ShadowDepthExtent,
 					ScreenVertexShader,
@@ -1096,7 +1187,7 @@ void FProjectedShadowInfo::RenderDepth(
 		ERenderTargetLoadAction::ENoAction,
 		FExclusiveDepthStencil::DepthWrite_StencilNop);
 
-	if (CacheMode == SDCM_MovablePrimitivesOnly)
+	if (CacheMode == SDCM_MovablePrimitivesOnly || CacheMode == SDCM_CSMScrolling)
 	{
 		// Copy in depths of static primitives before we render movable primitives.
 		FMeshPassProcessorRenderState DrawRenderState;
@@ -1273,7 +1364,7 @@ FRHIGPUMask FSceneRenderer::GetGPUMaskForShadow(FProjectedShadowInfo* ProjectedS
 	}
 	// SDCM_StaticPrimitivesOnly shadows don't update every frame so we need to render
 	// their depths on all possible GPUs.
-	else if (ProjectedShadowInfo->CacheMode == SDCM_StaticPrimitivesOnly)
+	else if (!ProjectedShadowInfo->IsWholeSceneDirectionalShadow() && ProjectedShadowInfo->CacheMode == SDCM_StaticPrimitivesOnly)
 	{
 		// Cached whole scene shadows shouldn't be view dependent.
 		checkSlow(ProjectedShadowInfo->DependentView == nullptr);
@@ -1892,8 +1983,10 @@ void FSceneRenderer::RenderShadowDepthMaps(FRDGBuilder& GraphBuilder, FInstanceC
 		GetLightNameForDrawEvent(ProjectedShadowInfo->GetLightSceneInfo().Proxy, LightNameWithLevel);
 		RDG_EVENT_SCOPE(GraphBuilder, "Cubemap %s %u^2", *LightNameWithLevel, TargetSize.X, TargetSize.Y);
 
+		const FCachedShadowMapData& CachedShadowMapData = Scene->GetCachedShadowMapDataRef(ProjectedShadowInfo->GetLightSceneInfo().Id, FMath::Max(ProjectedShadowInfo->CascadeSettings.ShadowSplitIndex, 0));
+
 		// Only clear when we're not copying from a cached shadow map.
-		if (ProjectedShadowInfo->CacheMode != SDCM_MovablePrimitivesOnly || !Scene->CachedShadowMaps.FindChecked(ProjectedShadowInfo->GetLightSceneInfo().Id).bCachedShadowMapHasPrimitives)
+		if (ProjectedShadowInfo->CacheMode != SDCM_MovablePrimitivesOnly || !CachedShadowMapData.bCachedShadowMapHasPrimitives)
 		{
 			AddClearShadowDepthPass(GraphBuilder, ShadowDepthTexture);
 		}
diff --git a/Engine/Source/Runtime/Renderer/Private/ShadowRendering.h b/Engine/Source/Runtime/Renderer/Private/ShadowRendering.h
--- a/Engine/Source/Runtime/Renderer/Private/ShadowRendering.h
+++ b/Engine/Source/Runtime/Renderer/Private/ShadowRendering.h
@@ -143,6 +143,7 @@ enum EShadowDepthCacheMode
 {
 	SDCM_MovablePrimitivesOnly,
 	SDCM_StaticPrimitivesOnly,
+	SDCM_CSMScrolling,
 	SDCM_Uncached
 };
 
@@ -391,6 +392,15 @@ public:
 	float PerObjectShadowFadeStart;
 	float InvPerObjectShadowFadeLength;
 
+	/** The Z offset in shadow light coordinate of the cached shadow bound center and current shadow bound center when scrolling the cached shadow map. */
+	float CSMScrollingZOffset;
+
+	FVector4 OverlappedUVOnCachedShadowMap;
+	FVector4 OverlappedUVOnCurrentShadowMap;
+
+	/** The extra culling planes to cull the static meshes which already in the overlapped area when scrolling the cached shadow map. */
+	TArray<FPlane, TInlineAllocator<4>> CSMScrollingExtraCullingPlanes;
+
 	/** Projection index for light types that use multiple projections for shadows. */
 	int32 ProjectionIndex = 0;
 
diff --git a/Engine/Source/Runtime/Renderer/Private/ShadowSetup.cpp b/Engine/Source/Runtime/Renderer/Private/ShadowSetup.cpp
--- a/Engine/Source/Runtime/Renderer/Private/ShadowSetup.cpp
+++ b/Engine/Source/Runtime/Renderer/Private/ShadowSetup.cpp
@@ -308,6 +308,52 @@ static TAutoConsoleVariable<int32> CVarMobileShadowmapRoundUpToPowerOfTwo(
 	ECVF_RenderThreadSafe
 );
 
+static TAutoConsoleVariable<int32> CVarCSMCaching(
+	TEXT("r.Shadow.CSMCaching"),
+	0,
+	TEXT("0: Render CSM every frame.\n")
+	TEXT("1: Enable CSM caching. (default)"),
+	ECVF_RenderThreadSafe);
+
+static TAutoConsoleVariable<float> CVarCSMScrollingOverlapAreaThrottle(
+	TEXT("r.Shadow.CSMScrollingOverlapAreaThrottle"),
+	0.75f,
+	TEXT("The minimum ratio of the overlap area for CSM scrolling."),
+	ECVF_RenderThreadSafe);
+
+static TAutoConsoleVariable<int32> CVarCSMScrollingMaxExtraStaticShadowSubjects(
+	TEXT("r.Shadow.MaxCSMScrollingStaticShadowSubjects"),
+	5,
+	TEXT("The maximum number of extra static shadow subjects need to be drawed when scrolling CSM."),
+	ECVF_RenderThreadSafe);
+
+/**
+ * Helper function to maintain the common login of caching the whole scene shadow and be able to handle the specific-purpose through the CacheLambda and UnCacheLambda
+ */
+template <typename CacheLambdaType, typename UnCacheLambdaType>
+void TryToCacheShadowMap(const FScene* Scene, int32& OutNumShadowMaps, EShadowDepthCacheMode* OutCacheModes, uint32& NumCachesUpdatedThisFrame, CacheLambdaType&& CacheLambda, UnCacheLambdaType&& UnCacheLambda)
+{
+	int64 CachedShadowMapsSize = Scene->GetCachedWholeSceneShadowMapsSize();
+
+	if (CachedShadowMapsSize < static_cast<int64>(GWholeSceneShadowCacheMb) * 1024 * 1024)
+	{
+		OutNumShadowMaps = 2;
+		// Note: ShadowMap with static primitives rendered first so movable shadowmap can composite
+		OutCacheModes[0] = SDCM_StaticPrimitivesOnly;
+		OutCacheModes[1] = SDCM_MovablePrimitivesOnly;
+		++NumCachesUpdatedThisFrame;
+
+		CacheLambda();
+	}
+	else
+	{
+		OutNumShadowMaps = 1;
+		OutCacheModes[0] = SDCM_Uncached;
+
+		UnCacheLambda();
+	}
+}
+
 /**
  * Helper function to determine fade alpha value for shadows based on resolution. In the below ASCII art (1) is
  * the MinShadowResolution and (2) is the ShadowFadeResolution. Alpha will be 0 below the min resolution and 1
@@ -646,6 +692,8 @@ FProjectedShadowInfo::FProjectedShadowInfo()
 	, bIncludeInScreenSpaceShadowMask(true)
 	, PerObjectShadowFadeStart(WORLD_MAX)
 	, InvPerObjectShadowFadeLength(0.0f)
+	, OverlappedUVOnCachedShadowMap(-1.0f, -1.0f, -1.0f, -1.0f)
+	, OverlappedUVOnCurrentShadowMap(-1.0f, -1.0f, -1.0f, -1.0f)
 	, LightSceneInfo(0)
 	, ParentSceneInfo(0)
 	, NumDynamicSubjectMeshElements(0)
@@ -1087,7 +1135,7 @@ struct FAddSubjectPrimitiveResult
 			uint32 bNeedPrimitiveFadingStateUpdate : 1;
 			uint32 bFadingIn : 1;
 			uint32 bAddOnRenderThread : 1;
-			uint32 bRecordShadowSubjectsForMobile : 1;
+			uint32 bShouldRecordShadowSubjectsForMobile : 1;
 
 			union
 			{
@@ -1489,7 +1537,7 @@ bool FProjectedShadowInfo::ShouldDrawStaticMeshes_AnyThread(
 	return bDrawingStaticMeshes;
 }
 
-void FProjectedShadowInfo::AddSubjectPrimitive(FPrimitiveSceneInfo* PrimitiveSceneInfo, TArray<FViewInfo>* ViewArray, ERHIFeatureLevel::Type FeatureLevel, bool bRecordShadowSubjectsForMobileShading)
+void FProjectedShadowInfo::AddSubjectPrimitive(FPrimitiveSceneInfo* PrimitiveSceneInfo, TArray<FViewInfo>* ViewArray, ERHIFeatureLevel::Type FeatureLevel, bool bShouldRecordShadowSubjectsForMobileShading)
 {
 	// Ray traced shadows use the GPU managed distance field object buffers, no CPU culling should be used
 	check(!bRayTracedDistanceField);
@@ -1638,7 +1686,7 @@ void FProjectedShadowInfo::AddSubjectPrimitive(FPrimitiveSceneInfo* PrimitiveSce
 
 			if (bDrawingStaticMeshes)
 			{
-				if (bRecordShadowSubjectsForMobileShading)
+				if (bShouldRecordShadowSubjectsForMobileShading)
 				{
 					DependentView->VisibleLightInfos[GetLightSceneInfo().Id].MobileCSMSubjectPrimitives.AddSubjectPrimitive(PrimitiveSceneInfo, PrimitiveId);
 				}
@@ -1648,7 +1696,7 @@ void FProjectedShadowInfo::AddSubjectPrimitive(FPrimitiveSceneInfo* PrimitiveSce
 				// Add the primitive to the subject primitive list.
 				DynamicSubjectPrimitives.Add(PrimitiveSceneInfo);
 
-				if (bRecordShadowSubjectsForMobileShading)
+				if (bShouldRecordShadowSubjectsForMobileShading)
 				{
 					DependentView->VisibleLightInfos[GetLightSceneInfo().Id].MobileCSMSubjectPrimitives.AddSubjectPrimitive(PrimitiveSceneInfo, PrimitiveId);
 				}
@@ -1704,20 +1752,12 @@ uint64 FProjectedShadowInfo::AddSubjectPrimitive_AnyThread(
 
 	if (FSceneInterface::GetShadingPath(FeatureLevel) == EShadingPath::Mobile)
 	{
-		// record shadow casters if CSM culling is enabled for the light's mobility type and the culling mode requires the list of casters.
-		static auto* CVarMobileEnableStaticAndCSMShadowReceivers = IConsoleManager::Get().FindTConsoleVariableDataInt(TEXT("r.Mobile.EnableStaticAndCSMShadowReceivers"));
-		static auto* CVarMobileEnableMovableLightCSMShaderCulling = IConsoleManager::Get().FindTConsoleVariableDataInt(TEXT("r.Mobile.EnableMovableLightCSMShaderCulling"));
-		static auto* CVarMobileCSMShaderCullingMethod = IConsoleManager::Get().FindTConsoleVariableDataInt(TEXT("r.Mobile.Shadow.CSMShaderCullingMethod"));
-		const uint32 MobileCSMCullingMode = CVarMobileCSMShaderCullingMethod->GetValueOnAnyThread() & 0xF;
-		const bool bRecordShadowSubjectsForMobile =
-			(MobileCSMCullingMode == 2 || MobileCSMCullingMode == 3)
-			&& ((CVarMobileEnableMovableLightCSMShaderCulling->GetValueOnAnyThread() && GetLightSceneInfo().Proxy->IsMovable() && GetLightSceneInfo().ShouldRenderViewIndependentWholeSceneShadows())
-				|| (CVarMobileEnableStaticAndCSMShadowReceivers->GetValueOnAnyThread() && GetLightSceneInfo().Proxy->UseCSMForDynamicObjects()));
-
-		if (bRecordShadowSubjectsForMobile)
+		const bool bShouldRecordShadowSubjectsForMobile = GetLightSceneInfo().ShouldRecordShadowSubjectsForMobile();
+
+		if (bShouldRecordShadowSubjectsForMobile)
 		{
 			Result.bAddOnRenderThread = true;
-			Result.bRecordShadowSubjectsForMobile = true;
+			Result.bShouldRecordShadowSubjectsForMobile = true;
 			++OutStats.NumDeferredPrimitives;
 			return Result.Qword;
 		}
@@ -1906,9 +1946,27 @@ void FProjectedShadowInfo::FinalizeAddSubjectPrimitive(
 	FPrimitiveSceneInfo* PrimitiveSceneInfo = Op.PrimitiveSceneInfo;
 	const FAddSubjectPrimitiveResult& Result = Op.Result;
 
+	if (CacheMode != SDCM_Uncached && IsWholeSceneDirectionalShadow())
+	{
+		checkSlow(!PrimitiveSceneInfo->Proxy->IsMeshShapeOftenMoving() || CacheMode == SDCM_CSMScrolling || CacheMode == SDCM_MovablePrimitivesOnly);
+
+		FCachedShadowMapData& CachedShadowMapData = PrimitiveSceneInfo->Scene->GetCachedShadowMapDataRef(LightSceneInfo->Id, CascadeSettings.ShadowSplitIndex);
+
+		if (CacheMode == SDCM_StaticPrimitivesOnly)
+		{
+			// Record all static meshes in the cached shadow map.
+			CachedShadowMapData.StaticShadowSubjectMap[PrimitiveSceneInfo->GetIndex()] = true;
+		}
+		else if (!PrimitiveSceneInfo->Proxy->IsMeshShapeOftenMoving())
+		{
+			// Count the number of extra draw calls of static meshes for filling the scrolling area.
+			CachedShadowMapData.LastFrameExtraStaticShadowSubjects += 1;
+		}
+	}
+
 	if (Result.bAddOnRenderThread)
 	{
-		AddSubjectPrimitive(PrimitiveSceneInfo, ViewArray, FeatureLevel, Result.bRecordShadowSubjectsForMobile);
+		AddSubjectPrimitive(PrimitiveSceneInfo, ViewArray, FeatureLevel, Result.bShouldRecordShadowSubjectsForMobile);
 		return;
 	}
 
@@ -1991,7 +2049,8 @@ bool FProjectedShadowInfo::HasSubjectPrims() const
 {
 	return DynamicSubjectPrimitives.Num() > 0
 		|| ShadowDepthPass.HasAnyDraw()
-		|| SubjectMeshCommandBuildRequests.Num() > 0;
+		|| SubjectMeshCommandBuildRequests.Num() > 0
+		|| ShadowDepthPassVisibleCommands.Num() > 0;
 }
 
 void FProjectedShadowInfo::AddReceiverPrimitive(FPrimitiveSceneInfo* PrimitiveSceneInfo)
@@ -3019,10 +3078,12 @@ void ComputeWholeSceneShadowCacheModes(
 	if (GCacheWholeSceneShadows 
 		&& (!bCubeShadowMap || RHISupportsGeometryShaders(GShaderPlatformForFeatureLevel[Scene->GetFeatureLevel()]) || RHISupportsVertexShaderLayer(GShaderPlatformForFeatureLevel[Scene->GetFeatureLevel()])))
 	{
-		FCachedShadowMapData* CachedShadowMapData = Scene->CachedShadowMaps.Find(LightSceneInfo->Id);
+		TArray<FCachedShadowMapData>* CachedShadowMapDatas = Scene->GetCachedShadowMapDatas(LightSceneInfo->Id);
 
-		if (CachedShadowMapData)
+		if (CachedShadowMapDatas)
 		{
+			FCachedShadowMapData* CachedShadowMapData = &(*CachedShadowMapDatas)[0];
+
 			if (InOutProjectedShadowInitializer.IsCachedShadowValid(CachedShadowMapData->Initializer) 
 				&& CachedShadowMapData->bCachedShadowMapHasNaniteGeometry == !bNeedsVirtualShadowMap)
 			{
@@ -3033,16 +3094,9 @@ void ComputeWholeSceneShadowCacheModes(
 				}
 				else
 				{
-					int64 CachedShadowMapsSize = Scene->GetCachedWholeSceneShadowMapsSize();
-
-					if (CachedShadowMapsSize < static_cast<int64>(GWholeSceneShadowCacheMb) * 1024 * 1024)
+					TryToCacheShadowMap(Scene, OutNumShadowMaps, OutCacheModes, *NumCachesUpdatedThisFrame,
+					[&]()
 					{
-						OutNumShadowMaps = 2;
-						// Note: ShadowMap with static primitives rendered first so movable shadowmap can composite
-						OutCacheModes[0] = SDCM_StaticPrimitivesOnly;
-						OutCacheModes[1] = SDCM_MovablePrimitivesOnly;
-						++*NumCachesUpdatedThisFrame;
-						
 						// Check if update is caused by resolution change
 						if (CanFallbackToOldShadowMapCache(CachedShadowMapData->ShadowMap, MaxShadowResolution))
 						{
@@ -3083,43 +3137,186 @@ void ComputeWholeSceneShadowCacheModes(
 								--*NumCachesUpdatedThisFrame;
 							}
 						}
-					}
-					else
+					},
+					[&]()
 					{
-						OutNumShadowMaps = 1;
-						OutCacheModes[0] = SDCM_Uncached;
-						CachedShadowMapData->ShadowMap.Release();
-					}
+						CachedShadowMapData->InvalidateCachedShadow();
+					});
 				}
 			}
 			else
 			{
 				OutNumShadowMaps = 1;
 				OutCacheModes[0] = SDCM_Uncached;
-				CachedShadowMapData->ShadowMap.Release();
+				CachedShadowMapData->InvalidateCachedShadow();
 			}
-
+			
 			CachedShadowMapData->Initializer = InOutProjectedShadowInitializer;
 			CachedShadowMapData->LastUsedTime = RealTime;
 		}
 		else
 		{
-			int64 CachedShadowMapsSize = Scene->GetCachedWholeSceneShadowMapsSize();
+			TryToCacheShadowMap(Scene, OutNumShadowMaps, OutCacheModes, *NumCachesUpdatedThisFrame,
+			[&]()
+			{
+				Scene->CachedShadowMaps.Add(LightSceneInfo->Id).Add(FCachedShadowMapData(InOutProjectedShadowInitializer, RealTime));
+			},
+			[&]()
+			{
+			});
+		}
+	}
+	else
+	{
+		OutNumShadowMaps = 1;
+		OutCacheModes[0] = SDCM_Uncached;
+		Scene->CachedShadowMaps.Remove(LightSceneInfo->Id);
+	}
+
+	if (OutNumShadowMaps > 0)
+	{
+		int32 NumOcclusionQueryableShadows = 0;
+
+		for (int32 i = 0; i < OutNumShadowMaps; i++)
+		{
+			NumOcclusionQueryableShadows += IsShadowCacheModeOcclusionQueryable(OutCacheModes[i]);
+		}
+
+		// Verify only one of the shadows will be occlusion queried, since they are all for the same light bounds
+		check(NumOcclusionQueryableShadows == 1);
+	}
+}
+
+void ComputeViewDependentWholeSceneShadowCacheModes(
+	const FLightSceneInfo* LightSceneInfo,
+	float RealTime,
+	FScene* Scene,
+	const FWholeSceneProjectedShadowInitializer& ProjectedShadowInitializer,
+	const FIntPoint& ShadowMapSize,
+	uint32& InOutNumCSMCachesUpdatedThisFrame,
+	int32& OutNumShadowMaps,
+	EShadowDepthCacheMode* OutCacheModes)
+{
+	if (CVarCSMCaching.GetValueOnAnyThread() == 1)
+	{
+		TArray<FCachedShadowMapData>* CachedShadowMapDatas = Scene->GetCachedShadowMapDatas(LightSceneInfo->Id);
+
+		checkSlow(ProjectedShadowInitializer.CascadeSettings.ShadowSplitIndex >= 0);
+		FCachedShadowMapData* CachedShadowMapData = (CachedShadowMapDatas && CachedShadowMapDatas->Num() > ProjectedShadowInitializer.CascadeSettings.ShadowSplitIndex) ? &(*CachedShadowMapDatas)[ProjectedShadowInitializer.CascadeSettings.ShadowSplitIndex] : nullptr;
 
-			if (CachedShadowMapsSize < GWholeSceneShadowCacheMb * 1024 * 1024)
+		if (CachedShadowMapData)
+		{
+			const FWholeSceneProjectedShadowInitializer& CachedShadowInitializer = CachedShadowMapData->Initializer;
+
+			// The cached shadow map could be valid if the light direction doesn't change.
+			bool bIsValidCachedViewDependentWholeSceneShadow = ProjectedShadowInitializer.WorldToLight == CachedShadowInitializer.WorldToLight
+				&& ProjectedShadowInitializer.WAxis == CachedShadowInitializer.WAxis
+				&& ProjectedShadowInitializer.bOnePassPointLightShadow == CachedShadowInitializer.bOnePassPointLightShadow
+				&& ProjectedShadowInitializer.bRayTracedDistanceField == CachedShadowInitializer.bRayTracedDistanceField;
+
+			if (bIsValidCachedViewDependentWholeSceneShadow)
 			{
-				OutNumShadowMaps = 2;
-				// Note: ShadowMap with static primitives rendered first so movable shadowmap can composite
-				OutCacheModes[0] = SDCM_StaticPrimitivesOnly;
-				OutCacheModes[1] = SDCM_MovablePrimitivesOnly;
-				++*NumCachesUpdatedThisFrame;
-				Scene->CachedShadowMaps.Add(LightSceneInfo->Id, FCachedShadowMapData(InOutProjectedShadowInitializer, RealTime));
+				bool bExactlyEqual = ProjectedShadowInitializer.PreShadowTranslation == CachedShadowInitializer.PreShadowTranslation
+					&& ProjectedShadowInitializer.Scales == CachedShadowInitializer.Scales
+					&& ProjectedShadowInitializer.SubjectBounds.Origin == CachedShadowInitializer.SubjectBounds.Origin
+					&& ProjectedShadowInitializer.SubjectBounds.BoxExtent == CachedShadowInitializer.SubjectBounds.BoxExtent
+					&& ProjectedShadowInitializer.SubjectBounds.SphereRadius == CachedShadowInitializer.SubjectBounds.SphereRadius
+					&& ProjectedShadowInitializer.MinLightW == CachedShadowInitializer.MinLightW
+					&& ProjectedShadowInitializer.MaxDistanceToCastInLightW == CachedShadowInitializer.MaxDistanceToCastInLightW;
+
+				if (CachedShadowMapData->ShadowMap.IsValid() && CachedShadowMapData->ShadowBufferResolution == ShadowMapSize)
+				{
+					if (bExactlyEqual)
+					{
+						OutNumShadowMaps = 1;
+						OutCacheModes[0] = SDCM_MovablePrimitivesOnly;
+					}
+					else
+					{
+						const FVector FaceDirection(1, 0, 0);
+						FVector	XAxis, YAxis;
+						FaceDirection.FindBestAxisVectors(XAxis, YAxis);
+						const FMatrix WorldToLight = ProjectedShadowInitializer.WorldToLight * FBasisVectorMatrix(-XAxis, YAxis, FaceDirection.GetSafeNormal(), FVector::ZeroVector);
+
+						FVector UpVector = FVector(WorldToLight.M[0][1], WorldToLight.M[1][1], WorldToLight.M[2][1]);
+						FVector RightVector = FVector(WorldToLight.M[0][0], WorldToLight.M[1][0], WorldToLight.M[2][0]);
+
+						// Projected the centers to the light coordinate
+						FVector ProjectedCurrentCenter = WorldToLight.TransformPosition(-ProjectedShadowInitializer.PreShadowTranslation);
+						FVector ProjectedCachedCenter = WorldToLight.TransformPosition(-CachedShadowMapData->Initializer.PreShadowTranslation);
+						FVector ProjectedCenterVector = ProjectedCachedCenter - ProjectedCurrentCenter;
+
+						FBox CurrentViewport = FBox(FVector(-ProjectedShadowInitializer.SubjectBounds.SphereRadius, -ProjectedShadowInitializer.SubjectBounds.SphereRadius, 0.0f), FVector(ProjectedShadowInitializer.SubjectBounds.SphereRadius, ProjectedShadowInitializer.SubjectBounds.SphereRadius, 0.0f));
+						FBox CachedViewport = FBox(FVector(ProjectedCenterVector.X - CachedShadowMapData->Initializer.SubjectBounds.SphereRadius, ProjectedCenterVector.Y - CachedShadowMapData->Initializer.SubjectBounds.SphereRadius, 0.0f), FVector(ProjectedCenterVector.X + CachedShadowMapData->Initializer.SubjectBounds.SphereRadius, ProjectedCenterVector.Y + CachedShadowMapData->Initializer.SubjectBounds.SphereRadius, 0.0f));
+
+						FBox OverlappedViewport = CurrentViewport.Overlap(CachedViewport);
+
+						FVector OverlappedArea = OverlappedViewport.GetSize();
+
+						float OverlappedAreaRatio = (OverlappedArea.X * OverlappedArea.Y) / (CachedShadowMapData->Initializer.SubjectBounds.SphereRadius * CachedShadowMapData->Initializer.SubjectBounds.SphereRadius * 4);
+
+						// if the overlapped area of the cached shadow map and the current shadow map smaller than a throttle or the extra draw calls of the static meshes when scrolling the shadow are greater than a throttle, update the cached shadow map, otherwise scroll the shadow map instead.
+						if (OverlappedAreaRatio > CVarCSMScrollingOverlapAreaThrottle.GetValueOnAnyThread() && CachedShadowMapData->LastFrameExtraStaticShadowSubjects <= CVarCSMScrollingMaxExtraStaticShadowSubjects.GetValueOnAnyThread())
+						{
+							OutNumShadowMaps = 1;
+							OutCacheModes[0] = SDCM_CSMScrolling;
+						}
+						else
+						{
+							OutNumShadowMaps = 2;
+							OutCacheModes[0] = SDCM_StaticPrimitivesOnly;
+							OutCacheModes[1] = SDCM_MovablePrimitivesOnly;
+							++InOutNumCSMCachesUpdatedThisFrame;
+
+							CachedShadowMapData->Initializer = ProjectedShadowInitializer;
+							CachedShadowMapData->InvalidateCachedShadow();
+							checkSlow(CachedShadowMapData->StaticShadowSubjectMap.Num() == Scene->Primitives.Num());
+						}
+					}
+				}
+				else
+				{
+					TryToCacheShadowMap(Scene, OutNumShadowMaps, OutCacheModes, InOutNumCSMCachesUpdatedThisFrame,
+					[&]()
+					{
+						CachedShadowMapData->ShadowBufferResolution = ShadowMapSize;
+						CachedShadowMapData->Initializer = ProjectedShadowInitializer;
+						checkSlow(CachedShadowMapData->StaticShadowSubjectMap.Num() == Scene->Primitives.Num());
+					},
+					[&]()
+					{
+						CachedShadowMapData->InvalidateCachedShadow();
+						CachedShadowMapData->Initializer = ProjectedShadowInitializer;
+					});
+				}
 			}
 			else
 			{
 				OutNumShadowMaps = 1;
 				OutCacheModes[0] = SDCM_Uncached;
+				CachedShadowMapData->InvalidateCachedShadow();
+				CachedShadowMapData->Initializer = ProjectedShadowInitializer;
 			}
+
+			CachedShadowMapData->LastUsedTime = RealTime;
+		}
+		else
+		{
+			TryToCacheShadowMap(Scene, OutNumShadowMaps, OutCacheModes, InOutNumCSMCachesUpdatedThisFrame,
+			[&]()
+			{
+				if (CachedShadowMapDatas == nullptr)
+				{
+					CachedShadowMapDatas = &Scene->CachedShadowMaps.Add(LightSceneInfo->Id);
+				}
+				CachedShadowMapDatas->Add(FCachedShadowMapData(ProjectedShadowInitializer, RealTime));
+				CachedShadowMapData = &CachedShadowMapDatas->Last();
+				CachedShadowMapData->ShadowBufferResolution = ShadowMapSize;
+				CachedShadowMapData->StaticShadowSubjectMap.Init(false, Scene->Primitives.Num());
+			},
+			[&]()
+			{
+			});
 		}
 	}
 	else
@@ -3549,7 +3746,7 @@ void FSceneRenderer::CreateWholeSceneProjectedShadow(
 					{
 						const bool bHasStaticPrimitives = ProjectedShadowInfo->HasSubjectPrims() || bContainsNaniteSubjects;
 						bRenderShadow = bHasStaticPrimitives;
-						FCachedShadowMapData& CachedShadowMapData = Scene->CachedShadowMaps.FindChecked(ProjectedShadowInfo->GetLightSceneInfo().Id);
+						FCachedShadowMapData& CachedShadowMapData = Scene->GetCachedShadowMapDataRef(ProjectedShadowInfo->GetLightSceneInfo().Id);
 						CachedShadowMapData.bCachedShadowMapHasPrimitives = bHasStaticPrimitives;
 						CachedShadowMapData.bCachedShadowMapHasNaniteGeometry = ProjectedShadowInfo->bNaniteGeometry;
 					}
@@ -4019,7 +4216,39 @@ struct FGatherShadowPrimitivesPacket
 					bScreenSpaceSizeCulled = FMath::Square(PrimitiveBounds.SphereRadius) < FMath::Square(MinScreenRadiusForShadowCaster) * DistanceSquared * LODScaleSquared;
 				}
 
+				bool bShadowRelevance = true;
+
+				if (ProjectedShadowInfo->CacheMode == SDCM_MovablePrimitivesOnly)
+				{
+					const FCachedShadowMapData& CachedShadowMapData = Scene->GetCachedShadowMapDataRef(LightSceneInfo.Id, ProjectedShadowInfo->CascadeSettings.ShadowSplitIndex);
+
+					// if the mesh is dynamic or a new added mesh, it should cast csm.
+					bShadowRelevance = PrimitiveProxy->IsMeshShapeOftenMoving() || !CachedShadowMapData.StaticShadowSubjectMap[PrimitiveSceneInfo->GetIndex()];
+				}
+				else if (ProjectedShadowInfo->CacheMode == SDCM_StaticPrimitivesOnly)
+				{
+					bShadowRelevance = !PrimitiveProxy->IsMeshShapeOftenMoving();
+				}
+				else if (ProjectedShadowInfo->CacheMode == SDCM_CSMScrolling)
+				{
+					const FCachedShadowMapData& CachedShadowMapData = Scene->GetCachedShadowMapDataRef(LightSceneInfo.Id, ProjectedShadowInfo->CascadeSettings.ShadowSplitIndex);
+
+					// if the mesh is dynamic or a new added mesh, it should cast csm.
+					bShadowRelevance = PrimitiveProxy->IsMeshShapeOftenMoving() || !CachedShadowMapData.StaticShadowSubjectMap[PrimitiveSceneInfo->GetIndex()];
+
+					// if the static mesh is rendered in the cached shadow map, check if it intersects with the non-overlapped area
+					if (!bShadowRelevance)
+					{
+						for (int32 i = 0; i < ProjectedShadowInfo->CSMScrollingExtraCullingPlanes.Num() && !bShadowRelevance; ++i)
+						{
+							const FPlane& CSMScrollingExtraCullingPlane = ProjectedShadowInfo->CSMScrollingExtraCullingPlanes[i];
+							bShadowRelevance |= FMath::PlaneAABBRelativePosition(CSMScrollingExtraCullingPlane, PrimitiveBounds.GetBox()) <= 0;
+						}
+					}
+				}
+
 				if (!bScreenSpaceSizeCulled
+					&& bShadowRelevance
 					&& ProjectedShadowInfo->GetLightSceneInfoCompact().AffectsPrimitive(PrimitiveBounds, PrimitiveProxy)
 					// Include all primitives for movable lights, but only statically shadowed primitives from a light with static shadowing,
 					// Since lights with static shadowing still create per-object shadows for primitives without static shadowing.
@@ -4083,10 +4312,28 @@ struct FGatherShadowPrimitivesPacket
 			Context.OverflowedMDCIndices = OverflowBuffer.MDCIndices.GetData();
 			Context.OverflowedMeshIndices = OverflowBuffer.MeshIndices.GetData();
 
+			FCachedShadowMapData* CachedShadowMapData = nullptr;
+
+			if (ProjectedShadowInfo->CacheMode != SDCM_Uncached && ProjectedShadowInfo->IsWholeSceneDirectionalShadow())
+			{
+				CachedShadowMapData = const_cast<FCachedShadowMapData*>(Scene->GetCachedShadowMapData(ProjectedShadowInfo->GetLightSceneInfo().Id, ProjectedShadowInfo->CascadeSettings.ShadowSplitIndex));
+				
+				if (ProjectedShadowInfo->CacheMode == SDCM_MovablePrimitivesOnly || ProjectedShadowInfo->CacheMode == SDCM_CSMScrolling)
+				{
+					CachedShadowMapData->LastFrameExtraStaticShadowSubjects = 0;
+				}
+			}
+
 			for (int32 PrimitiveIndex = 0; PrimitiveIndex < SubjectPrimitives.Num(); PrimitiveIndex++)
 			{
 				ProjectedShadowInfo->FinalizeAddSubjectPrimitive(SubjectPrimitives[PrimitiveIndex], nullptr, FeatureLevel, Context);
 			}
+
+			if (ProjectedShadowInfo->CacheMode == SDCM_StaticPrimitivesOnly && ProjectedShadowInfo->IsWholeSceneDirectionalShadow())
+			{
+				checkSlow(CachedShadowMapData != nullptr);
+				CachedShadowMapData->bCachedShadowMapHasPrimitives = ProjectedShadowInfo->HasSubjectPrims();
+			}
 		}
 	}
 };
@@ -4234,7 +4481,8 @@ void FSceneRenderer::AddViewDependentWholeSceneShadowsForView(
 	TArray<FProjectedShadowInfo*, SceneRenderingAllocator>& ShadowInfos, 
 	TArray<FProjectedShadowInfo*, SceneRenderingAllocator>& ShadowInfosThatNeedCulling,
 	FVisibleLightInfo& VisibleLightInfo, 
-	FLightSceneInfo& LightSceneInfo)
+	FLightSceneInfo& LightSceneInfo,
+	uint32& NumCSMCachesUpdatedThisFrame)
 {
 	SCOPE_CYCLE_COUNTER(STAT_AddViewDependentWholeSceneShadowsForView);
 
@@ -4307,35 +4555,136 @@ void FSceneRenderer::AddViewDependentWholeSceneShadowsForView(
 					FIntPoint ShadowBufferResolution = FIntPoint( FMath::Clamp( MaxCSMResolution, MinCSMResolution, (int32)GMaxShadowDepthBufferSizeX ) - ShadowBorder * 2,
 														          FMath::Clamp( MaxCSMResolution, MinCSMResolution, (int32)GMaxShadowDepthBufferSizeY ) - ShadowBorder * 2 );
 
-					// Create the projected shadow info.
-					FProjectedShadowInfo* ProjectedShadowInfo = new(FMemStack::Get(), 1, 16) FProjectedShadowInfo;
-					ProjectedShadowInfo->SetupWholeSceneProjection(
-						&LightSceneInfo,
-						&View,
-						ProjectedShadowInitializer,
-						ShadowBufferResolution.X,
-						ShadowBufferResolution.Y,
-						ShadowBufferResolution.X,
-						ShadowBufferResolution.Y,
-						ShadowBorder
-					);
-					ProjectedShadowInfo->FadeAlphas = FadeAlphas;
-					ProjectedShadowInfo->ProjectionIndex = Index;
-
-					VisibleLightInfo.MemStackProjectedShadows.Add(ProjectedShadowInfo);
-					VisibleLightInfo.AllProjectedShadows.Add(ProjectedShadowInfo);
-					ShadowInfos.Add(ProjectedShadowInfo);
+					int32 NumShadowMaps = 1;
+					EShadowDepthCacheMode CacheMode[2] = { SDCM_Uncached, SDCM_Uncached };
 
-					if (bNeedsVirtualShadowMap && !ProjectedShadowInitializer.bRayTracedDistanceField)
+					if (!View.bIsSceneCapture && !ProjectedShadowInitializer.bRayTracedDistanceField)
 					{
-						// If we have a virtual shadow map, disable nanite rendering into the regular shadow map or else we'd get double-shadowing
-						ProjectedShadowInfo->bNaniteGeometry = false;
+						ComputeViewDependentWholeSceneShadowCacheModes(
+							&LightSceneInfo,
+							ViewFamily.CurrentRealTime,
+							Scene,
+							ProjectedShadowInitializer,
+							ShadowBufferResolution,
+							NumCSMCachesUpdatedThisFrame,
+							NumShadowMaps,
+							CacheMode);
 					}
 
-					// Ray traced shadows use the GPU managed distance field object buffers, no CPU culling needed
-					if (!ProjectedShadowInfo->bRayTracedDistanceField)
+					for (int32 CacheModeIndex = 0; CacheModeIndex < NumShadowMaps; CacheModeIndex++)
 					{
-						ShadowInfosThatNeedCulling.Add(ProjectedShadowInfo);
+						// Create the projected shadow info.
+						FProjectedShadowInfo* ProjectedShadowInfo = new(FMemStack::Get(), 1, 16) FProjectedShadowInfo;
+						ProjectedShadowInfo->SetupWholeSceneProjection(
+							&LightSceneInfo,
+							&View,
+							ProjectedShadowInitializer,
+							ShadowBufferResolution.X,
+							ShadowBufferResolution.Y,
+							ShadowBufferResolution.X,
+							ShadowBufferResolution.Y,
+							ShadowBorder
+						);
+						ProjectedShadowInfo->FadeAlphas = FadeAlphas;
+						ProjectedShadowInfo->ProjectionIndex = Index;
+						ProjectedShadowInfo->CacheMode = CacheMode[CacheModeIndex];
+
+						VisibleLightInfo.MemStackProjectedShadows.Add(ProjectedShadowInfo);
+						VisibleLightInfo.AllProjectedShadows.Add(ProjectedShadowInfo);
+						ShadowInfos.Add(ProjectedShadowInfo);
+
+						if (bNeedsVirtualShadowMap && !ProjectedShadowInitializer.bRayTracedDistanceField)
+						{
+							// If we have a virtual shadow map, disable nanite rendering into the regular shadow map or else we'd get double-shadowing
+							ProjectedShadowInfo->bNaniteGeometry = false;
+						}
+
+						// Ray traced shadows use the GPU managed distance field object buffers, no CPU culling needed
+						if (!ProjectedShadowInfo->bRayTracedDistanceField)
+						{
+							ShadowInfosThatNeedCulling.Add(ProjectedShadowInfo);
+
+							if (CacheMode[CacheModeIndex] == SDCM_StaticPrimitivesOnly || CacheMode[CacheModeIndex] == SDCM_CSMScrolling)
+							{
+								FCachedShadowMapData& CachedShadowMapData = Scene->GetCachedShadowMapDataRef(LightSceneInfo.Id, ProjectedShadowInfo->CascadeSettings.ShadowSplitIndex);
+
+								if (CacheMode[CacheModeIndex] == SDCM_StaticPrimitivesOnly)
+								{
+									CachedShadowMapData.MaxSubjectZ = ProjectedShadowInfo->MaxSubjectZ;
+									CachedShadowMapData.MinSubjectZ = ProjectedShadowInfo->MinSubjectZ;
+									CachedShadowMapData.PreShadowTranslation = ProjectedShadowInfo->PreShadowTranslation;
+								}
+								else // CacheMode[CacheModeIndex] == SDCM_CSMScrolling
+								{
+									const FVector FaceDirection(1, 0, 0);
+									FVector	XAxis, YAxis;
+									FaceDirection.FindBestAxisVectors(XAxis, YAxis);
+									const FMatrix WorldToLight = ProjectedShadowInitializer.WorldToLight * FBasisVectorMatrix(-XAxis, YAxis, FaceDirection.GetSafeNormal(), FVector::ZeroVector);
+									const FMatrix LightToWorld = WorldToLight.InverseFast();
+
+									FVector UpVector = LightToWorld.TransformVector(FVector(0.0f, 1.0f, 0.0f));
+									FVector RightVector = LightToWorld.TransformVector(FVector(1.0f, 0.0f, 0.0f));
+
+									// Projected the centers to the light coordinate, recalculate the projected center because the ProjectedShadowInfo's PreShadowTranslation was modified
+									FVector ProjectedCurrentCenter = WorldToLight.TransformPosition(-ProjectedShadowInfo->PreShadowTranslation);
+									FVector ProjectedCachedCenter = WorldToLight.TransformPosition(-CachedShadowMapData.PreShadowTranslation);
+									FVector ProjectedCenterVector = ProjectedCachedCenter - ProjectedCurrentCenter;
+
+									ProjectedShadowInfo->CSMScrollingZOffset = ProjectedCenterVector.Z;
+
+									FVector2D BorderScale = FVector2D(float(ProjectedShadowInfo->ResolutionX + 2 * ProjectedShadowInfo->BorderSize) / ProjectedShadowInfo->ResolutionX, float(ProjectedShadowInfo->ResolutionY + 2 * ProjectedShadowInfo->BorderSize) / ProjectedShadowInfo->ResolutionY);
+									FVector2D CurrentRadius = ProjectedShadowInitializer.SubjectBounds.SphereRadius * BorderScale;
+									FVector2D CachedRadius = CachedShadowMapData.Initializer.SubjectBounds.SphereRadius * BorderScale;
+
+									FBox CurrentViewport = FBox(FVector(-CurrentRadius.X, -CurrentRadius.Y, 0.0f), FVector(CurrentRadius.X, CurrentRadius.Y, 0.0f));
+									FBox CachedViewport = FBox(FVector(ProjectedCenterVector.X - CachedRadius.X, ProjectedCenterVector.Y - CachedRadius.Y, 0.0f), FVector(ProjectedCenterVector.X + CachedRadius.X, ProjectedCenterVector.Y + CachedRadius.Y, 0.0f));
+
+									FBox OverlappedViewport = CurrentViewport.Overlap(CachedViewport);
+
+									float CachedViewportWidth = CachedViewport.Max.X - CachedViewport.Min.X;
+									float CachedViewportHeight = CachedViewport.Max.Y - CachedViewport.Min.Y;
+									ProjectedShadowInfo->OverlappedUVOnCachedShadowMap = FVector4((OverlappedViewport.Min.X - CachedViewport.Min.X) / CachedViewportWidth
+										, (CachedViewport.Max.Y - OverlappedViewport.Max.Y) / CachedViewportHeight
+										, (OverlappedViewport.Max.X - CachedViewport.Min.X) / CachedViewportWidth
+										, (CachedViewport.Max.Y - OverlappedViewport.Min.Y) / CachedViewportHeight);
+
+									float CurrentViewportWidth = CurrentViewport.Max.X - CurrentViewport.Min.X;
+									float CurrentViewportHeight = CurrentViewport.Max.Y - CurrentViewport.Min.Y;
+									ProjectedShadowInfo->OverlappedUVOnCurrentShadowMap = FVector4((OverlappedViewport.Min.X - CurrentViewport.Min.X) / CurrentViewportWidth
+										, (CurrentViewport.Max.Y - OverlappedViewport.Max.Y) / CurrentViewportHeight
+										, (OverlappedViewport.Max.X - CurrentViewport.Min.X) / CurrentViewportWidth
+										, (CurrentViewport.Max.Y - OverlappedViewport.Min.Y) / CurrentViewportHeight);
+
+									TArray<FPlane, TInlineAllocator<4>>& ExtraCullingPlanes = ProjectedShadowInfo->CSMScrollingExtraCullingPlanes;
+
+									ExtraCullingPlanes.Empty();
+
+									//Top
+									if (CurrentViewport.Max.Y > OverlappedViewport.Max.Y)
+									{
+										ExtraCullingPlanes.Add(FPlane(LightToWorld.TransformPosition(ProjectedCurrentCenter + FVector(0.0f, OverlappedViewport.Max.Y, 0.0f)), -UpVector));
+									}
+
+									//Bottom
+									if (CurrentViewport.Min.Y < OverlappedViewport.Min.Y)
+									{
+										ExtraCullingPlanes.Add(FPlane(LightToWorld.TransformPosition(ProjectedCurrentCenter + FVector(0.0f, OverlappedViewport.Min.Y, 0.0f)), UpVector));
+									}
+
+									//Left
+									if (CurrentViewport.Min.X < OverlappedViewport.Min.X)
+									{
+										ExtraCullingPlanes.Add(FPlane(LightToWorld.TransformPosition(ProjectedCurrentCenter + FVector(OverlappedViewport.Min.X, 0.0f, 0.0f)), RightVector));
+									}
+
+									//Right
+									if (CurrentViewport.Max.X > OverlappedViewport.Max.X)
+									{
+										ExtraCullingPlanes.Add(FPlane(LightToWorld.TransformPosition(ProjectedCurrentCenter + FVector(OverlappedViewport.Max.X, 0.0f, 0.0f)), -RightVector));
+									}
+								}
+							}
+						}
 					}
 				}
 			}
@@ -4389,6 +4738,8 @@ void FSceneRenderer::AllocateShadowDepthTargets(FRHICommandListImmediate& RHICmd
 	TArray<FProjectedShadowInfo*, SceneRenderingAllocator> MobileWholeSceneDirectionalShadows;
 	TArray<FProjectedShadowInfo*, SceneRenderingAllocator> MobileDynamicSpotlightShadows;
 
+	TArray<FProjectedShadowInfo*, SceneRenderingAllocator> CachedWholeSceneDirectionalShadows;
+
 	const bool bMobile = FeatureLevel < ERHIFeatureLevel::SM5;
 
 	for (TSparseArray<FLightSceneInfoCompact>::TConstIterator LightIt(Scene->Lights); LightIt; ++LightIt)
@@ -4425,9 +4776,10 @@ void FSceneRenderer::AllocateShadowDepthTargets(FRHICommandListImmediate& RHICmd
 				bShadowIsVisible |= bHasViewRelevance && VisibleLightViewInfo.ProjectedShadowVisibilityMap[ShadowIndex];
 			}
 
-			if (ProjectedShadowInfo->CacheMode == SDCM_MovablePrimitivesOnly && !ProjectedShadowInfo->HasSubjectPrims())
+			// Skip the shadow rendering if there is no primitive in the cached shadow map and the current shadow map either
+			if ((ProjectedShadowInfo->CacheMode == SDCM_StaticPrimitivesOnly || (ProjectedShadowInfo->CacheMode == SDCM_MovablePrimitivesOnly && !ProjectedShadowInfo->IsWholeSceneDirectionalShadow())) && !ProjectedShadowInfo->HasSubjectPrims())
 			{
-				FCachedShadowMapData& CachedShadowMapData = Scene->CachedShadowMaps.FindChecked(ProjectedShadowInfo->GetLightSceneInfo().Id);
+				const FCachedShadowMapData& CachedShadowMapData = Scene->GetCachedShadowMapDataRef(ProjectedShadowInfo->GetLightSceneInfo().Id, FMath::Max(ProjectedShadowInfo->CascadeSettings.ShadowSplitIndex, 0));
 
 				// A shadowmap for movable primitives when there are no movable primitives would normally read directly from the cached shadowmap
 				// However if the cached shadowmap also had no primitives then we need to skip rendering the shadow entirely
@@ -4493,7 +4845,14 @@ void FSceneRenderer::AllocateShadowDepthTargets(FRHICommandListImmediate& RHICmd
 					}
 					else if (ProjectedShadowInfo->bDirectionalLight && ProjectedShadowInfo->bWholeSceneShadow)
 					{
-						WholeSceneDirectionalShadows.Add(ProjectedShadowInfo);
+						if (ProjectedShadowInfo->CacheMode == SDCM_StaticPrimitivesOnly)
+						{
+							CachedWholeSceneDirectionalShadows.Add(ProjectedShadowInfo);
+						}
+						else
+						{
+							WholeSceneDirectionalShadows.Add(ProjectedShadowInfo);
+						}
 					}
 					else if (ProjectedShadowInfo->bOnePassPointLightShadow)
 					{
@@ -4525,7 +4884,11 @@ void FSceneRenderer::AllocateShadowDepthTargets(FRHICommandListImmediate& RHICmd
 
 		if (!bMobile)
 		{
+			AllocateCachedShadowDepthTargets(RHICmdList, CachedWholeSceneDirectionalShadows);
+
 			AllocateCSMDepthTargets(RHICmdList, WholeSceneDirectionalShadows, SortedShadowsForShadowDepthPass.ShadowMapAtlases);
+
+			CachedWholeSceneDirectionalShadows.Empty();
 		}
 		else
 		{
@@ -4579,7 +4942,11 @@ void FSceneRenderer::AllocateShadowDepthTargets(FRHICommandListImmediate& RHICmd
 	}
 
 	AllocateOnePassPointLightDepthTargets(RHICmdList, WholeScenePointShadows);
-	AllocateCachedSpotlightShadowDepthTargets(RHICmdList, CachedSpotlightShadows);
+	AllocateCachedShadowDepthTargets(RHICmdList, CachedSpotlightShadows);
+	if (bMobile)
+	{
+		AllocateCachedShadowDepthTargets(RHICmdList, CachedWholeSceneDirectionalShadows);
+	}
 	AllocateAtlasedShadowDepthTargets(RHICmdList, Shadows, SortedShadowsForShadowDepthPass.ShadowMapAtlases);
 	AllocateTranslucentShadowDepthTargets(RHICmdList, TranslucentShadows);
 	AllocateMobileCSMAndSpotLightShadowDepthTargets(RHICmdList, MobileWholeSceneDirectionalShadows);
@@ -4605,13 +4972,16 @@ void FSceneRenderer::AllocateShadowDepthTargets(FRHICommandListImmediate& RHICmd
 	}
 
 	// Remove cache entries that haven't been used in a while
-	for (TMap<int32, FCachedShadowMapData>::TIterator CachedShadowMapIt(Scene->CachedShadowMaps); CachedShadowMapIt; ++CachedShadowMapIt)
+	for (TMap<int32, TArray<FCachedShadowMapData>>::TIterator CachedShadowMapIt(Scene->CachedShadowMaps); CachedShadowMapIt; ++CachedShadowMapIt)
 	{
-		FCachedShadowMapData& ShadowMapData = CachedShadowMapIt.Value();
+		TArray<FCachedShadowMapData>& ShadowMapDatas = CachedShadowMapIt.Value();
 
-		if (ShadowMapData.ShadowMap.IsValid() && ViewFamily.CurrentRealTime - ShadowMapData.LastUsedTime > 2.0f)
+		for (auto& ShadowMapData : ShadowMapDatas)
 		{
-			ShadowMapData.ShadowMap.Release();
+			if (ShadowMapData.ShadowMap.IsValid() && ViewFamily.CurrentRealTime - ShadowMapData.LastUsedTime > 2.0f)
+			{
+				ShadowMapData.InvalidateCachedShadow();
+			}
 		}
 	}
 
@@ -4658,7 +5028,7 @@ void FSceneRenderer::AllocateAtlasedShadowDepthTargets(
 
 		if (ProjectedShadowInfo->CacheMode == SDCM_MovablePrimitivesOnly && !ProjectedShadowInfo->HasSubjectPrims())
 		{
-			FCachedShadowMapData& CachedShadowMapData = Scene->CachedShadowMaps.FindChecked(ProjectedShadowInfo->GetLightSceneInfo().Id);
+			const FCachedShadowMapData& CachedShadowMapData = Scene->GetCachedShadowMapDataRef(ProjectedShadowInfo->GetLightSceneInfo().Id);
 			ProjectedShadowInfo->X = ProjectedShadowInfo->Y = 0;
 			ProjectedShadowInfo->bAllocated = true;
 			// Skip the shadow depth pass since there are no movable primitives to composite, project from the cached shadowmap directly which contains static primitive depths
@@ -4725,20 +5095,23 @@ void FSceneRenderer::AllocateAtlasedShadowDepthTargets(
 	}
 }
 
-void FSceneRenderer::AllocateCachedSpotlightShadowDepthTargets(FRHICommandListImmediate& RHICmdList, TArray<FProjectedShadowInfo*, SceneRenderingAllocator>& CachedSpotlightShadows)
+void FSceneRenderer::AllocateCachedShadowDepthTargets(FRHICommandListImmediate& RHICmdList, TArray<FProjectedShadowInfo*, SceneRenderingAllocator>& CachedShadows)
 {
-	for (int32 ShadowIndex = 0; ShadowIndex < CachedSpotlightShadows.Num(); ShadowIndex++)
+	for (int32 ShadowIndex = 0; ShadowIndex < CachedShadows.Num(); ShadowIndex++)
 	{
-		FProjectedShadowInfo* ProjectedShadowInfo = CachedSpotlightShadows[ShadowIndex];
+		FProjectedShadowInfo* ProjectedShadowInfo = CachedShadows[ShadowIndex];
+
+		bool bIsWholeSceneDirectionalShadow = ProjectedShadowInfo->IsWholeSceneDirectionalShadow();
+
 		SortedShadowsForShadowDepthPass.ShadowMapAtlases.AddDefaulted();
 		FSortedShadowMapAtlas& ShadowMap = SortedShadowsForShadowDepthPass.ShadowMapAtlases.Last();
 
 		FIntPoint ShadowResolution(ProjectedShadowInfo->ResolutionX + ProjectedShadowInfo->BorderSize * 2, ProjectedShadowInfo->ResolutionY + ProjectedShadowInfo->BorderSize * 2);
 		FPooledRenderTargetDesc ShadowMapDesc2D = FPooledRenderTargetDesc::Create2DDesc(ShadowResolution, PF_ShadowDepth, FClearValueBinding::DepthOne, TexCreate_None, TexCreate_DepthStencilTargetable | TexCreate_ShaderResource, false, 1, false);
-		GRenderTargetPool.FindFreeElement(RHICmdList, ShadowMapDesc2D, ShadowMap.RenderTargets.DepthTarget, TEXT("CachedShadowDepthMap"), ERenderTargetTransience::NonTransient);
+		GRenderTargetPool.FindFreeElement(RHICmdList, ShadowMapDesc2D, ShadowMap.RenderTargets.DepthTarget, bIsWholeSceneDirectionalShadow ? *FString::Printf(TEXT("CachedShadowDepthMap_Split%d"), ProjectedShadowInfo->CascadeSettings.ShadowSplitIndex) : TEXT("CachedShadowDepthMap"), ERenderTargetTransience::NonTransient);
 
 		check(ProjectedShadowInfo->CacheMode == SDCM_StaticPrimitivesOnly);
-		FCachedShadowMapData& CachedShadowMapData = Scene->CachedShadowMaps.FindChecked(ProjectedShadowInfo->GetLightSceneInfo().Id);
+		FCachedShadowMapData& CachedShadowMapData = Scene->GetCachedShadowMapDataRef(ProjectedShadowInfo->GetLightSceneInfo().Id, FMath::Max(ProjectedShadowInfo->CascadeSettings.ShadowSplitIndex, 0));
 		CachedShadowMapData.ShadowMap = ShadowMap.RenderTargets;
 
 		ProjectedShadowInfo->X = ProjectedShadowInfo->Y = 0;
@@ -4844,7 +5217,7 @@ void FSceneRenderer::AllocateOnePassPointLightDepthTargets(FRHICommandListImmedi
 
 			if (ProjectedShadowInfo->CacheMode == SDCM_MovablePrimitivesOnly && !ProjectedShadowInfo->HasSubjectPrims())
 			{
-				FCachedShadowMapData& CachedShadowMapData = Scene->CachedShadowMaps.FindChecked(ProjectedShadowInfo->GetLightSceneInfo().Id);
+				const FCachedShadowMapData& CachedShadowMapData = Scene->GetCachedShadowMapDataRef(ProjectedShadowInfo->GetLightSceneInfo().Id);
 				ProjectedShadowInfo->X = ProjectedShadowInfo->Y = 0;
 				ProjectedShadowInfo->bAllocated = true;
 				// Skip the shadow depth pass since there are no movable primitives to composite, project from the cached shadowmap directly which contains static primitive depths
@@ -4862,7 +5235,7 @@ void FSceneRenderer::AllocateOnePassPointLightDepthTargets(FRHICommandListImmedi
 
 				if (ProjectedShadowInfo->CacheMode == SDCM_StaticPrimitivesOnly)
 				{
-					FCachedShadowMapData& CachedShadowMapData = Scene->CachedShadowMaps.FindChecked(ProjectedShadowInfo->GetLightSceneInfo().Id);
+					FCachedShadowMapData& CachedShadowMapData = Scene->GetCachedShadowMapDataRef(ProjectedShadowInfo->GetLightSceneInfo().Id);
 					CachedShadowMapData.ShadowMap = ShadowMapCubemap.RenderTargets;
 				}
 
@@ -4985,6 +5358,7 @@ void FSceneRenderer::InitDynamicShadows(FRHICommandListImmediate& RHICmdList, FG
 
 	uint32 NumPointShadowCachesUpdatedThisFrame = 0;
 	uint32 NumSpotShadowCachesUpdatedThisFrame = 0;
+	uint32 NumCSMCachesUpdatedThisFrame = 0;
 
 	TArray<FProjectedShadowInfo*,SceneRenderingAllocator> PreShadows;
 	TArray<FProjectedShadowInfo*,SceneRenderingAllocator> ViewDependentWholeSceneShadows;
@@ -5090,7 +5464,7 @@ void FSceneRenderer::InitDynamicShadows(FRHICommandListImmediate& RHICmdList, FG
 								// Mobile uses the scene's MobileDirectionalLights only for whole scene shadows.
 								&& (LightSceneInfo == Scene->MobileDirectionalLights[0] || LightSceneInfo == Scene->MobileDirectionalLights[1] || LightSceneInfo == Scene->MobileDirectionalLights[2])))
 						{
-							AddViewDependentWholeSceneShadowsForView(ViewDependentWholeSceneShadows, ViewDependentWholeSceneShadowsThatNeedCulling, VisibleLightInfo, *LightSceneInfo);
+							AddViewDependentWholeSceneShadowsForView(ViewDependentWholeSceneShadows, ViewDependentWholeSceneShadowsThatNeedCulling, VisibleLightInfo, *LightSceneInfo, NumCSMCachesUpdatedThisFrame);
 						}
 
 						// Disable per-object shadows when non-nanite VSMs are enabled
@@ -5131,7 +5505,7 @@ void FSceneRenderer::InitDynamicShadows(FRHICommandListImmediate& RHICmdList, FG
 			}
 		}
 
-		CSV_CUSTOM_STAT(LightCount, UpdatedShadowMaps, float(NumPointShadowCachesUpdatedThisFrame + NumSpotShadowCachesUpdatedThisFrame), ECsvCustomStatOp::Set);
+		CSV_CUSTOM_STAT(LightCount, UpdatedShadowMaps, float(NumPointShadowCachesUpdatedThisFrame + NumSpotShadowCachesUpdatedThisFrame + NumCSMCachesUpdatedThisFrame), ECsvCustomStatOp::Set);
 		CSV_CUSTOM_STAT_GLOBAL(ShadowCacheUsageMB, (float(Scene->GetCachedWholeSceneShadowMapsSize()) / 1024) / 1024, ECsvCustomStatOp::Set);
 
 		// Calculate visibility of the projected shadows.
diff --git a/Engine/Source/Runtime/Renderer/Private/ShadowSetupMobile.cpp b/Engine/Source/Runtime/Renderer/Private/ShadowSetupMobile.cpp
--- a/Engine/Source/Runtime/Renderer/Private/ShadowSetupMobile.cpp
+++ b/Engine/Source/Runtime/Renderer/Private/ShadowSetupMobile.cpp
@@ -245,20 +245,15 @@ void FMobileSceneRenderer::InitDynamicShadows(FRHICommandListImmediate& RHICmdLi
 	// initialize CSMVisibilityInfo for each eligible light.
 	for (FLightSceneInfo* MobileDirectionalLightSceneInfo : Scene->MobileDirectionalLights)
 	{
-		const FLightSceneProxy* LightSceneProxy = MobileDirectionalLightSceneInfo ? MobileDirectionalLightSceneInfo->Proxy : nullptr;
-		if (LightSceneProxy)
-		{
-			bool bLightHasCombinedStaticAndCSMEnabled = bCombinedStaticAndCSMEnabled && LightSceneProxy->UseCSMForDynamicObjects();
-			bool bMovableLightUsingCSM = bMobileEnableMovableLightCSMShaderCulling && LightSceneProxy->IsMovable() && MobileDirectionalLightSceneInfo->ShouldRenderViewIndependentWholeSceneShadows();
+		const bool bShouldRecordShadowSubjectsForMobile = MobileDirectionalLightSceneInfo ? MobileDirectionalLightSceneInfo->ShouldRecordShadowSubjectsForMobile() : false;
 
-			if (bLightHasCombinedStaticAndCSMEnabled || bMovableLightUsingCSM)
+		if (bShouldRecordShadowSubjectsForMobile)
+		{
+			int32 PrimitiveCount = Scene->Primitives.Num();
+			for (auto& View : Views)
 			{
-				int32 PrimitiveCount = Scene->Primitives.Num();
-				for (auto& View : Views)
-				{
-					FMobileCSMSubjectPrimitives& MobileCSMSubjectPrimitives = View.VisibleLightInfos[MobileDirectionalLightSceneInfo->Id].MobileCSMSubjectPrimitives;
-					MobileCSMSubjectPrimitives.InitShadowSubjectPrimitives(PrimitiveCount);
-				}
+				FMobileCSMSubjectPrimitives& MobileCSMSubjectPrimitives = View.VisibleLightInfos[MobileDirectionalLightSceneInfo->Id].MobileCSMSubjectPrimitives;
+				MobileCSMSubjectPrimitives.InitShadowSubjectPrimitives(PrimitiveCount);
 			}
 		}
 	}
