#ue5 - Fix crash when loading editor with help from Brandon.Dawson

diff --git a/Engine/Source/Runtime/Renderer/Private/PrimitiveSceneInfo.cpp b/Engine/Source/Runtime/Renderer/Private/PrimitiveSceneInfo.cpp
--- a/Engine/Source/Runtime/Renderer/Private/PrimitiveSceneInfo.cpp
+++ b/Engine/Source/Runtime/Renderer/Private/PrimitiveSceneInfo.cpp
@@ -116,11 +116,11 @@ FPrimitiveFlagsCompact::FPrimitiveFlagsCompact(const FPrimitiveSceneProxy* Proxy
 {}
 
 FPrimitiveSceneInfoCompact::FPrimitiveSceneInfoCompact(FPrimitiveSceneInfo* InPrimitiveSceneInfo) :
-	Bounds(PrimitiveSceneInfo->Proxy->GetBounds()),
 	PrimitiveFlagsCompact(InPrimitiveSceneInfo->Proxy)
 {
 	PrimitiveSceneInfo = InPrimitiveSceneInfo;
 	Proxy = PrimitiveSceneInfo->Proxy;
+	Bounds = PrimitiveSceneInfo->Proxy->GetBounds();
 	MinDrawDistance = PrimitiveSceneInfo->Proxy->GetMinDrawDistance();
 	MaxDrawDistance = PrimitiveSceneInfo->Proxy->GetMaxDrawDistance();
 
