Add cvar to control number of frames VT page must be unused before it's considered 'free'
With lots of VTs per material, pages may not be 'touched' by the CPU every frame they're used, this threshold will help prevent these pages from being re-used too early
Effectively this will make the 'VT oversubscription' message more likely to appear, but this should more accurately reflect the actual load on the system

diff --git a/Engine/Source/Runtime/Engine/Private/VT/VirtualTextureScalability.cpp b/Engine/Source/Runtime/Engine/Private/VT/VirtualTextureScalability.cpp
--- a/Engine/Source/Runtime/Engine/Private/VT/VirtualTextureScalability.cpp
+++ b/Engine/Source/Runtime/Engine/Private/VT/VirtualTextureScalability.cpp
@@ -126,6 +126,14 @@ namespace VirtualTextureScalability
 		0,
 		TEXT("Is anisotropic filtering for VTs enabled?"),
 		ECVF_RenderThreadSafe | ECVF_ReadOnly);
+
+	static TAutoConsoleVariable<int32> CVarVTPageFreeThreshold(
+		TEXT("r.VT.PageFreeThreshold"),
+		60,
+		TEXT("Number of frames since the last time a VT page was used, before it's considered free.\n")
+		TEXT("VT pages are not necesarily marked as used on the CPU every time they're accessed by the GPU.\n")
+		TEXT("Increasing this threshold reduces the chances that an in-use frame is considered free."),
+		ECVF_RenderThreadSafe);
 	
 	/** Track changes and apply to relevant systems. This allows us to dynamically change the scalability settings. */
 	static void OnUpdate()
@@ -246,4 +254,9 @@ namespace VirtualTextureScalability
 	{
 		return GroupIndex < NumScalabilityGroups ? GTileCountBiases[GroupIndex] : 0;
 	}
+
+	uint32 GetPageFreeThreshold()
+	{
+		return FMath::Max(CVarVTPageFreeThreshold.GetValueOnRenderThread(), 0);
+	}
 }
diff --git a/Engine/Source/Runtime/Engine/Public/VT/VirtualTextureScalability.h b/Engine/Source/Runtime/Engine/Public/VT/VirtualTextureScalability.h
--- a/Engine/Source/Runtime/Engine/Public/VT/VirtualTextureScalability.h
+++ b/Engine/Source/Runtime/Engine/Public/VT/VirtualTextureScalability.h
@@ -28,4 +28,6 @@ namespace VirtualTextureScalability
 	ENGINE_API float GetPoolSizeScale(uint32 GroupIndex);
 	/** Get resolution bias for runtime virtual textures. */
 	ENGINE_API int32 GetRuntimeVirtualTextureSizeBias(uint32 GroupIndex);
+	/** Get the number of frames a page must be unused, before it's considered free */
+	ENGINE_API uint32 GetPageFreeThreshold();
 }
diff --git a/Engine/Source/Runtime/Renderer/Private/VT/TexturePagePool.cpp b/Engine/Source/Runtime/Renderer/Private/VT/TexturePagePool.cpp
--- a/Engine/Source/Runtime/Renderer/Private/VT/TexturePagePool.cpp
+++ b/Engine/Source/Runtime/Renderer/Private/VT/TexturePagePool.cpp
@@ -164,7 +164,7 @@ FVirtualTextureLocalTile FTexturePagePool::GetLocalTileFromPhysicalAddress(uint1
 	return FVirtualTextureLocalTile(FVirtualTextureProducerHandle(Pages[pAddress].PackedProducerHandle), Pages[pAddress].Local_vAddress, Pages[pAddress].Local_vLevel);
 }
 
-bool FTexturePagePool::AnyFreeAvailable( uint32 Frame ) const
+bool FTexturePagePool::AnyFreeAvailable( uint32 Frame, uint32 FreeThreshold) const
 {
 	if( FreeHeap.Num() > 0 )
 	{
@@ -172,7 +172,7 @@ bool FTexturePagePool::AnyFreeAvailable( uint32 Frame ) const
 		const uint16 pAddress = FreeHeap.Top();
 		const uint32 PageFrame = FreeHeap.GetKey(pAddress) >> 4;
 		// Don't free any pages that were touched this frame
-		return PageFrame != Frame;
+		return PageFrame + FreeThreshold < Frame;
 	}
 
 	return false;
@@ -239,7 +239,6 @@ uint32 FTexturePagePool::FindNearestPageLevel(const FVirtualTextureProducerHandl
 uint32 FTexturePagePool::Alloc(FVirtualTextureSystem* System, uint32 Frame, const FVirtualTextureProducerHandle& ProducerHandle, uint8 GroupIndex, uint32 Local_vAddress, uint8 Local_vLevel, bool bLock)
 {
 	check(ProducerHandle.PackedValue != 0u);
-	check(AnyFreeAvailable(Frame));
 	checkSlow(FindPageAddress(ProducerHandle, GroupIndex, Local_vAddress, Local_vLevel) == ~0u);
 
 	// Grab the LRU free page
diff --git a/Engine/Source/Runtime/Renderer/Private/VT/TexturePagePool.h b/Engine/Source/Runtime/Renderer/Private/VT/TexturePagePool.h
--- a/Engine/Source/Runtime/Renderer/Private/VT/TexturePagePool.h
+++ b/Engine/Source/Runtime/Renderer/Private/VT/TexturePagePool.h
@@ -74,7 +74,7 @@ public:
 	/**
 	 * Check if there are any free pages available at the moment.
 	 */
-	bool		AnyFreeAvailable( uint32 Frame ) const;
+	bool		AnyFreeAvailable( uint32 Frame, uint32 FreeThreshold ) const;
 
 	/**
 	 * Find physical address of the page allocated for the given VT address, or ~0 if not allocated
diff --git a/Engine/Source/Runtime/Renderer/Private/VT/VirtualTextureSystem.cpp b/Engine/Source/Runtime/Renderer/Private/VT/VirtualTextureSystem.cpp
--- a/Engine/Source/Runtime/Renderer/Private/VT/VirtualTextureSystem.cpp
+++ b/Engine/Source/Runtime/Renderer/Private/VT/VirtualTextureSystem.cpp
@@ -1950,6 +1950,7 @@ void FVirtualTextureSystem::SubmitRequests(FRDGBuilder& GraphBuilder, ERHIFeatur
 		PrepareTasks.Reserve(RequestList->GetNumLoadRequests());
 
 		const uint32 MaxPagesProduced = VirtualTextureScalability::GetMaxPagesProducedPerFrame();
+		const uint32 PageFreeThreshold = VirtualTextureScalability::GetPageFreeThreshold();
 		uint32 NumStacksProduced = 0u;
 		uint32 NumPagesProduced = 0u;
 		uint32 NumPageAllocateFails = 0u;
@@ -2011,7 +2012,7 @@ void FVirtualTextureSystem::SubmitRequests(FRDGBuilder& GraphBuilder, ERHIFeatur
 					{
 						FVirtualTexturePhysicalSpace* RESTRICT PhysicalSpace = Producer.GetPhysicalSpaceForPhysicalGroup(ProducerPhysicalGroupIndex);
 						FTexturePagePool& RESTRICT PagePool = PhysicalSpace->GetPagePool();
-						if (PagePool.AnyFreeAvailable(Frame))
+						if (PagePool.AnyFreeAvailable(Frame, PageFreeThreshold))
 						{
 							const uint32 pAddress = PagePool.Alloc(this, Frame, ProducerHandle, ProducerPhysicalGroupIndex, TileToLoad.Local_vAddress, TileToLoad.Local_vLevel, bLockTile);
 							check(pAddress != ~0u);
