Removed several RHI thread dispatch passes which are no longer necessary.

diff --git a/Engine/Source/Runtime/Renderer/Private/DeferredShadingRenderer.cpp b/Engine/Source/Runtime/Renderer/Private/DeferredShadingRenderer.cpp
--- a/Engine/Source/Runtime/Renderer/Private/DeferredShadingRenderer.cpp
+++ b/Engine/Source/Runtime/Renderer/Private/DeferredShadingRenderer.cpp
@@ -2163,7 +2163,6 @@ void FDeferredShadingSceneRenderer::Render(FRDGBuilder& GraphBuilder)
 		}
 
 		GraphBuilder.SetCommandListStat(GET_STATID(STAT_CLM_AfterPrePass));
-		AddServiceLocalQueuePass(GraphBuilder);
 
 		// special pass for DDM_AllOpaqueNoVelocity, which uses the velocity pass to finish the early depth pass write
 		if (bShouldRenderVelocities && Scene->EarlyZPassMode == DDM_AllOpaqueNoVelocity)
@@ -2172,7 +2171,6 @@ void FDeferredShadingSceneRenderer::Render(FRDGBuilder& GraphBuilder)
 			GraphBuilder.SetCommandListStat(GET_STATID(STAT_CLM_Velocity));
 			RenderVelocities(GraphBuilder, SceneTextures, EVelocityPass::Opaque, bHairEnable);
 			GraphBuilder.SetCommandListStat(GET_STATID(STAT_CLM_AfterVelocity));
-			AddServiceLocalQueuePass(GraphBuilder);
 		}
 
 		if (bDoInitViewAftersPrepass)
@@ -2193,8 +2191,6 @@ void FDeferredShadingSceneRenderer::Render(FRDGBuilder& GraphBuilder)
 				DynamicIndexBufferForInitShadows.Commit();
 				DynamicReadBufferForInitShadows.Commit();
 			}
-
-			AddServiceLocalQueuePass(GraphBuilder);
 		}
 	}
 
@@ -2370,7 +2366,6 @@ void FDeferredShadingSceneRenderer::Render(FRDGBuilder& GraphBuilder)
 		AllocateVirtualShadowMaps(bAfterBasePass);
 
 		RenderShadowDepthMaps(GraphBuilder, InstanceCullingManager);
-		AddServiceLocalQueuePass(GraphBuilder);
 	}
 	// End early Shadow depth rendering
 
@@ -2416,7 +2411,6 @@ void FDeferredShadingSceneRenderer::Render(FRDGBuilder& GraphBuilder)
 		{
 			SceneTextures.SetupMode |= ESceneTextureSetupMode::CustomDepth;
 			SceneTextures.UniformBuffer = CreateSceneTextureUniformBuffer(GraphBuilder, FeatureLevel, SceneTextures.SetupMode);
-			AddServiceLocalQueuePass(GraphBuilder);
 		}
 	}
 
@@ -2590,8 +2584,6 @@ void FDeferredShadingSceneRenderer::Render(FRDGBuilder& GraphBuilder)
 		RenderOcclusionLambda();
 	}
 
-	AddServiceLocalQueuePass(GraphBuilder);
-
 	// End occlusion after base
 
 	if (!bUseGBuffer)
@@ -2628,7 +2620,6 @@ void FDeferredShadingSceneRenderer::Render(FRDGBuilder& GraphBuilder)
 		}
 
 		ComputeVolumetricFog(GraphBuilder, SceneTextures);
-		AddServiceLocalQueuePass(GraphBuilder);
 	}
 	// End shadow and fog after base pass
 
@@ -2670,7 +2661,6 @@ void FDeferredShadingSceneRenderer::Render(FRDGBuilder& GraphBuilder)
 		{
 			SceneTextures.SetupMode |= ESceneTextureSetupMode::CustomDepth;
 			SceneTextures.UniformBuffer = CreateSceneTextureUniformBuffer(GraphBuilder, FeatureLevel, SceneTextures.SetupMode);
-			AddServiceLocalQueuePass(GraphBuilder);
 		}
 	}
 
@@ -2684,7 +2674,6 @@ void FDeferredShadingSceneRenderer::Render(FRDGBuilder& GraphBuilder)
 		GraphBuilder.SetCommandListStat(GET_STATID(STAT_CLM_Velocity));
 		RenderVelocities(GraphBuilder, SceneTextures, EVelocityPass::Opaque, bHairEnable);
 		GraphBuilder.SetCommandListStat(GET_STATID(STAT_CLM_AfterVelocity));
-		AddServiceLocalQueuePass(GraphBuilder);
 
 		// TODO: Populate velocity buffer from Nanite visibility buffer.
 	}
@@ -2775,11 +2764,9 @@ void FDeferredShadingSceneRenderer::Render(FRDGBuilder& GraphBuilder)
 		GraphBuilder.SetCommandListStat(GET_STATID(STAT_CLM_Lighting));
 		RenderLights(GraphBuilder, SceneTextures, TranslucencyLightingVolumeTextures, LightingChannelsTexture, SortedLightSet);
 		GraphBuilder.SetCommandListStat(GET_STATID(STAT_CLM_AfterLighting));
-		AddServiceLocalQueuePass(GraphBuilder);
 
 		InjectTranslucencyLightingVolumeAmbientCubemap(GraphBuilder, Views, TranslucencyLightingVolumeTextures);
 		FilterTranslucencyLightingVolume(GraphBuilder, Views, TranslucencyLightingVolumeTextures);
-		AddServiceLocalQueuePass(GraphBuilder);
 
 		// Render diffuse sky lighting and reflections that only operate on opaque pixels
 		RenderDeferredReflectionsAndSkyLighting(GraphBuilder, SceneTextures, DynamicBentNormalAOTexture);
@@ -2799,8 +2786,6 @@ void FDeferredShadingSceneRenderer::Render(FRDGBuilder& GraphBuilder)
 			CompositeRayTracingSkyLight(GraphBuilder, SceneTextures, SkyLightTexture, SkyLightHitDistanceTexture);
 		}
 	#endif
-
-		AddServiceLocalQueuePass(GraphBuilder);
 	}
 	else if (HairStrands::HasViewHairStrandsData(Views) && ViewFamily.EngineShowFlags.Lighting)
 	{
@@ -2846,7 +2831,6 @@ void FDeferredShadingSceneRenderer::Render(FRDGBuilder& GraphBuilder)
 
 		GraphBuilder.SetCommandListStat(GET_STATID(STAT_CLM_WaterPass));
 		RenderSingleLayerWater(GraphBuilder, SceneTextures, bShouldRenderVolumetricCloud, SceneWithoutWaterTextures);
-		AddServiceLocalQueuePass(GraphBuilder);
 	}
 
 	// Rebuild scene textures to include scene color.
@@ -2937,7 +2921,6 @@ void FDeferredShadingSceneRenderer::Render(FRDGBuilder& GraphBuilder)
 		// Render all remaining translucency views.
 		GraphBuilder.SetCommandListStat(GET_STATID(STAT_CLM_Translucency));
 		RenderTranslucency(GraphBuilder, SceneTextures, TranslucencyLightingVolumeTextures, &SeparateTranslucencyTextures, TranslucencyViewsToRender, InstanceCullingManager);
-		AddServiceLocalQueuePass(GraphBuilder);
 		TranslucencyViewsToRender = ETranslucencyView::None;
 
 		// Compose hair before velocity/distortion pass since these pass write depth value, 
@@ -2952,7 +2935,6 @@ void FDeferredShadingSceneRenderer::Render(FRDGBuilder& GraphBuilder)
 		{
 			GraphBuilder.SetCommandListStat(GET_STATID(STAT_CLM_Distortion));
 			RenderDistortion(GraphBuilder, SceneTextures.Color.Target, SceneTextures.Depth.Target);
-			AddServiceLocalQueuePass(GraphBuilder);
 		}
 
 		if (bShouldRenderVelocities)
@@ -2961,7 +2943,6 @@ void FDeferredShadingSceneRenderer::Render(FRDGBuilder& GraphBuilder)
 
 			GraphBuilder.SetCommandListStat(GET_STATID(STAT_CLM_TranslucentVelocity));
 			RenderVelocities(GraphBuilder, SceneTextures, EVelocityPass::Translucent, false);
-			AddServiceLocalQueuePass(GraphBuilder);
 
 			if (bRecreateSceneTextures)
 			{
@@ -3005,7 +2986,6 @@ void FDeferredShadingSceneRenderer::Render(FRDGBuilder& GraphBuilder)
 		SCOPE_CYCLE_COUNTER(STAT_FDeferredShadingSceneRenderer_RenderLightShaftBloom);
 		GraphBuilder.SetCommandListStat(GET_STATID(STAT_CLM_LightShaftBloom));
 		RenderLightShaftBloom(GraphBuilder, SceneTextures, SeparateTranslucencyTextures);
-		AddServiceLocalQueuePass(GraphBuilder);
 	}
 
 	if (bUseVirtualTexturing)
@@ -3042,14 +3022,12 @@ void FDeferredShadingSceneRenderer::Render(FRDGBuilder& GraphBuilder)
 		const float OcclusionMaxDistance = Scene->SkyLight && !Scene->SkyLight->bWantsStaticShadowing ? Scene->SkyLight->OcclusionMaxDistance : Scene->DefaultMaxDistanceFieldOcclusionDistance;
 		FRDGTextureRef DummyOutput = nullptr;
 		RenderDistanceFieldLighting(GraphBuilder, SceneTextures, FDistanceFieldAOParameters(OcclusionMaxDistance), DummyOutput, false, ViewFamily.EngineShowFlags.VisualizeDistanceFieldAO);
-		AddServiceLocalQueuePass(GraphBuilder);
 	}
 
 	// Draw visualizations just before use to avoid target contamination
 	if (ViewFamily.EngineShowFlags.VisualizeMeshDistanceFields || ViewFamily.EngineShowFlags.VisualizeGlobalDistanceField)
 	{
 		RenderMeshDistanceFieldVisualization(GraphBuilder, SceneTextures, FDistanceFieldAOParameters(Scene->DefaultMaxDistanceFieldOcclusionDistance));
-		AddServiceLocalQueuePass(GraphBuilder);
 	}
 
 	RenderLumenSceneVisualization(GraphBuilder, SceneTextures);
@@ -3059,7 +3037,6 @@ void FDeferredShadingSceneRenderer::Render(FRDGBuilder& GraphBuilder)
 	if (ViewFamily.EngineShowFlags.StationaryLightOverlap)
 	{
 		RenderStationaryLightOverlap(GraphBuilder, SceneTextures, LightingChannelsTexture);
-		AddServiceLocalQueuePass(GraphBuilder);
 	}
 
 	if (bShouldVisualizeVolumetricCloud && bCanOverlayRayTracingOutput)
@@ -3068,7 +3045,6 @@ void FDeferredShadingSceneRenderer::Render(FRDGBuilder& GraphBuilder)
 		ReconstructVolumetricRenderTarget(GraphBuilder, Views, SceneTextures.Depth.Resolve, HalfResolutionDepthCheckerboardMinMaxTexture, false);
 		ComposeVolumetricRenderTargetOverSceneForVisualization(GraphBuilder, Views, SceneTextures.Color.Target, SceneTextures);
 		RenderVolumetricCloud(GraphBuilder, SceneTextures, true, false, HalfResolutionDepthCheckerboardMinMaxTexture, false, InstanceCullingManager);
-		AddServiceLocalQueuePass(GraphBuilder);
 	}
 
 	// Resolve the scene color for post processing.
