Don't try to map all resident pages to newly allocated VTs.  Too slow, especially for large VTs
- Instead, keep VT cached for a few frames after they are deleted, this way they can be re-used if they're deleted/recreated due to state changes
- This should solve the problem of material changes causing pages to become unmapped for a few frames

diff --git a/Engine/Source/Runtime/RenderCore/Public/VirtualTexturing.h b/Engine/Source/Runtime/RenderCore/Public/VirtualTexturing.h
--- a/Engine/Source/Runtime/RenderCore/Public/VirtualTexturing.h
+++ b/Engine/Source/Runtime/RenderCore/Public/VirtualTexturing.h
@@ -333,6 +333,8 @@ public:
 		, WidthInBlocks(InWidthInBlocks)
 		, HeightInBlocks(InHeightInBlocks)
 		, DepthInTiles(InDepthInTiles)
+		, FrameDeleted(0u)
+		, NumRefs(0)
 		, PageTableFormat(EVTPageTableFormat::UInt32)
 		, SpaceID(~0u)
 		, MaxLevel(0u)
@@ -382,7 +384,7 @@ public:
 
 protected:
 	friend class FVirtualTextureSystem;
-	virtual void Destroy(class FVirtualTextureSystem* System) = 0;
+	virtual void Destroy(FVirtualTextureSystem* InSystem) = 0;
 	virtual ~IAllocatedVirtualTexture() {}
 
 	FAllocatedVTDescription Description;
@@ -391,7 +393,8 @@ protected:
 	uint32 WidthInBlocks;
 	uint32 HeightInBlocks;
 	uint32 DepthInTiles;
-	
+	uint32 FrameDeleted;
+	int32 NumRefs;
 
 	// should be set explicitly by derived class constructor
 	EVTPageTableFormat PageTableFormat;
diff --git a/Engine/Source/Runtime/Renderer/Private/VT/AllocatedVirtualTexture.cpp b/Engine/Source/Runtime/Renderer/Private/VT/AllocatedVirtualTexture.cpp
--- a/Engine/Source/Runtime/Renderer/Private/VT/AllocatedVirtualTexture.cpp
+++ b/Engine/Source/Runtime/Renderer/Private/VT/AllocatedVirtualTexture.cpp
@@ -18,7 +18,6 @@ FAllocatedVirtualTexture::FAllocatedVirtualTexture(FVirtualTextureSystem* InSyst
 	uint32 InHeightInBlocks,
 	uint32 InDepthInTiles)
 	: IAllocatedVirtualTexture(InDesc, InBlockWidthInTiles, InBlockHeightInTiles, InWidthInBlocks, InHeightInBlocks, InDepthInTiles)
-	, RefCount(1)
 	, FrameAllocated(InFrame)
 	, Space(nullptr)
 {
@@ -105,19 +104,9 @@ void FAllocatedVirtualTexture::AssignVirtualAddress(uint32 vAddress)
 }
 
 void FAllocatedVirtualTexture::Destroy(FVirtualTextureSystem* System)
-{
-	const int32 NewRefCount = RefCount.Decrement();
-	check(NewRefCount >= 0);
-	if (NewRefCount == 0)
-	{
-		System->ReleaseVirtualTexture(this);
-	}
-}
-
-void FAllocatedVirtualTexture::Release(FVirtualTextureSystem* System)
 {
 	check(IsInRenderingThread());
-	check(RefCount.GetValue() == 0);
+	check(NumRefs == 0);
 
 	// Unlock any locked tiles
 	LockOrUnlockTiles(System, false);
@@ -175,10 +164,7 @@ void FAllocatedVirtualTexture::Release(FVirtualTextureSystem* System)
 	}
 
 	Space->FreeVirtualTexture(this);
-	System->RemoveAllocatedVT(this);
 	System->ReleaseSpace(Space);
-
-	delete this;
 }
 
 void FAllocatedVirtualTexture::LockOrUnlockTiles(FVirtualTextureSystem* InSystem, bool bLock) const
diff --git a/Engine/Source/Runtime/Renderer/Private/VT/AllocatedVirtualTexture.h b/Engine/Source/Runtime/Renderer/Private/VT/AllocatedVirtualTexture.h
--- a/Engine/Source/Runtime/Renderer/Private/VT/AllocatedVirtualTexture.h
+++ b/Engine/Source/Runtime/Renderer/Private/VT/AllocatedVirtualTexture.h
@@ -31,7 +31,6 @@ public:
 
 	void Release(FVirtualTextureSystem* System);
 
-	inline void IncrementRefCount() { RefCount.Increment(); }
 	inline uint32 GetFrameAllocated() const { return FrameAllocated; }
 
 	void AssignVirtualAddress(uint32 vAddress);
@@ -47,7 +46,7 @@ public:
 	virtual FRHIShaderResourceView* GetPhysicalTextureSRV(uint32 InLayerIndex, bool bSRGB) const override;
 	virtual void GetPackedPageTableUniform(FUintVector4* OutUniform) const override;
 	virtual void GetPackedUniform(FUintVector4* OutUniform, uint32 LayerIndex) const override;
-	virtual void Destroy(FVirtualTextureSystem* System) override;
+	virtual void Destroy(FVirtualTextureSystem* InSystem) override;
 	// end IAllocatedVirtualTexture
 
 	inline FVirtualTextureSpace* GetSpace() const { return Space; }
@@ -66,7 +65,6 @@ private:
 	uint32 AddUniqueProducer(FVirtualTextureProducerHandle const& InHandle, const FVirtualTextureProducer* InProducer);
 	uint32 AddUniquePhysicalSpace(FVirtualTexturePhysicalSpace* InPhysicalSpace, uint32 InUniqueProducerIndex, uint32 InProducerPhysicalSpaceIndex);
 
-	FThreadSafeCounter RefCount;
 	uint32 FrameAllocated;
 
 	FVirtualTextureSpace* Space;
diff --git a/Engine/Source/Runtime/Renderer/Private/VT/VirtualTextureSystem.cpp b/Engine/Source/Runtime/Renderer/Private/VT/VirtualTextureSystem.cpp
--- a/Engine/Source/Runtime/Renderer/Private/VT/VirtualTextureSystem.cpp
+++ b/Engine/Source/Runtime/Renderer/Private/VT/VirtualTextureSystem.cpp
@@ -28,7 +28,6 @@ DECLARE_CYCLE_STAT(TEXT("Gather Requests"), STAT_ProcessRequests_Gather, STATGRO
 DECLARE_CYCLE_STAT(TEXT("Sort Requests"), STAT_ProcessRequests_Sort, STATGROUP_VirtualTexturing);
 DECLARE_CYCLE_STAT(TEXT("Submit Requests"), STAT_ProcessRequests_Submit, STATGROUP_VirtualTexturing);
 DECLARE_CYCLE_STAT(TEXT("Map Requests"), STAT_ProcessRequests_Map, STATGROUP_VirtualTexturing);
-DECLARE_CYCLE_STAT(TEXT("Map New VTs"), STAT_ProcessRequests_MapNew, STATGROUP_VirtualTexturing);
 DECLARE_CYCLE_STAT(TEXT("Finalize Requests"), STAT_ProcessRequests_Finalize, STATGROUP_VirtualTexturing);
 DECLARE_CYCLE_STAT(TEXT("Merge Unique Pages"), STAT_ProcessRequests_MergePages, STATGROUP_VirtualTexturing);
 DECLARE_CYCLE_STAT(TEXT("Merge Requests"), STAT_ProcessRequests_MergeRequests, STATGROUP_VirtualTexturing);
@@ -250,7 +249,7 @@ FVirtualTextureSystem::~FVirtualTextureSystem()
 	FCoreDelegates::OnGetOnScreenMessages.RemoveAll(this);
 #endif
 
-	DestroyPendingVirtualTextures();
+	DestroyPendingVirtualTextures(true);
 
 	check(AllocatedVTs.Num() == 0);
 
@@ -441,16 +440,20 @@ IAllocatedVirtualTexture* FVirtualTextureSystem::AllocateVirtualTexture(const FA
 {
 	check(Desc.NumTextureLayers <= VIRTUALTEXTURE_SPACE_MAXLAYERS);
 
-	// Make sure any pending VTs are destroyed before attempting to allocate a new one
-	// Otherwise, we might find/return an existing IAllocatedVirtualTexture* that's pending deletion
-	DestroyPendingVirtualTextures();
-
 	// Check to see if we already have an allocated VT that matches this description
 	// This can happen often as multiple material instances will share the same textures
 	FAllocatedVirtualTexture*& AllocatedVT = AllocatedVTs.FindOrAdd(Desc);
 	if (AllocatedVT)
 	{
-		AllocatedVT->IncrementRefCount();
+		FScopeLock Lock(&AllocatedVTLock);
+		const int32 PrevNumRefs = AllocatedVT->NumRefs++;
+		check(PrevNumRefs >= 0);
+		if (PrevNumRefs == 0)
+		{
+			// Bringing a VT 'back to life', remove it from the pending delete list
+			verify(PendingDeleteAllocatedVTs.RemoveSwap(AllocatedVT, false) == 1);
+		}
+
 		return AllocatedVT;
 	}
 
@@ -515,50 +518,61 @@ IAllocatedVirtualTexture* FVirtualTextureSystem::AllocateVirtualTexture(const FA
 	}
 
 	AllocatedVT = new FAllocatedVirtualTexture(this, Frame, Desc, ProducerForLayer, BlockWidthInTiles, BlockHeightInTiles, WidthInBlocks, HeightInBlocks, DepthInTiles);
-	if (bAnyLayerProducerWantsPersistentHighestMip)
-	{
-		AllocatedVTsToMap.Add(AllocatedVT);
-	}
+	AllocatedVT->NumRefs = 1;
 	return AllocatedVT;
 }
 
 void FVirtualTextureSystem::DestroyVirtualTexture(IAllocatedVirtualTexture* AllocatedVT)
 {
-	AllocatedVT->Destroy(this);
-}
-
-void FVirtualTextureSystem::ReleaseVirtualTexture(FAllocatedVirtualTexture* AllocatedVT)
-{
-	if (IsInRenderingThread())
+	FScopeLock Lock(&AllocatedVTLock);
+	const int32 NewNumRefs = --AllocatedVT->NumRefs;
+	check(NewNumRefs >= 0);
+	if (NewNumRefs == 0)
 	{
-		AllocatedVT->Release(this);
-	}
-	else
-	{
-		FScopeLock Lock(&PendingDeleteLock);
+		AllocatedVT->FrameDeleted = Frame;
 		PendingDeleteAllocatedVTs.Add(AllocatedVT);
 	}
 }
 
-void FVirtualTextureSystem::RemoveAllocatedVT(FAllocatedVirtualTexture* AllocatedVT)
-{
-	// shouldn't be more than 1 instance of this in the list
-	verify(AllocatedVTsToMap.Remove(AllocatedVT) <= 1);
-	// should always exist in this map
-	verify(AllocatedVTs.Remove(AllocatedVT->GetDescription()) == 1);
-}
-
-void FVirtualTextureSystem::DestroyPendingVirtualTextures()
+void FVirtualTextureSystem::DestroyPendingVirtualTextures(bool bForceDestroyAll)
 {
 	check(IsInRenderingThread());
-	TArray<FAllocatedVirtualTexture*> AllocatedVTsToDelete;
+	
+	TArray<IAllocatedVirtualTexture*> AllocatedVTsToDelete;
 	{
-		FScopeLock Lock(&PendingDeleteLock);
-		AllocatedVTsToDelete = MoveTemp(PendingDeleteAllocatedVTs);
+		FScopeLock Lock(&AllocatedVTLock);
+		if (bForceDestroyAll)
+		{
+			AllocatedVTsToDelete = MoveTemp(PendingDeleteAllocatedVTs);
+			PendingDeleteAllocatedVTs.Reset();
+		}
+		else
+		{
+			const uint32 CurrentFrame = Frame;
+			int32 Index = 0;
+			while (Index < PendingDeleteAllocatedVTs.Num())
+			{
+				IAllocatedVirtualTexture* AllocatedVT = PendingDeleteAllocatedVTs[Index];
+				check(AllocatedVT->NumRefs == 0);
+
+				if (CurrentFrame >= AllocatedVT->FrameDeleted + 60u)
+				{
+					AllocatedVTsToDelete.Add(AllocatedVT);
+					PendingDeleteAllocatedVTs.RemoveAtSwap(Index, 1, false);
+				}
+				else
+				{
+					Index++;
+				}
+			}
+		}
 	}
-	for (FAllocatedVirtualTexture* AllocatedVT : AllocatedVTsToDelete)
+
+	for (IAllocatedVirtualTexture* AllocatedVT : AllocatedVTsToDelete)
 	{
-		AllocatedVT->Release(this);
+		verify(AllocatedVTs.Remove(AllocatedVT->GetDescription()) == 1);
+		AllocatedVT->Destroy(this);
+		delete AllocatedVT;
 	}
 }
 
@@ -1044,7 +1058,7 @@ void FVirtualTextureSystem::Update(FRDGBuilder& GraphBuilder, ERHIFeatureLevel::
 		bFlushCaches = false;
 	}
 
-	DestroyPendingVirtualTextures();
+	DestroyPendingVirtualTextures(false);
 
 	// Early out when no allocated VTs
 	if (AllocatedVTs.Num() == 0)
@@ -2218,104 +2232,6 @@ void FVirtualTextureSystem::SubmitRequests(FRDGBuilder& GraphBuilder, ERHIFeatur
 		}
 	}
 
-	// Map any resident tiles to newly allocated VTs
-	{
-		SCOPE_CYCLE_COUNTER(STAT_ProcessRequests_MapNew);
-
-		uint32 Index = 0u;
-		while (Index < (uint32)AllocatedVTsToMap.Num())
-		{
-			const FAllocatedVirtualTexture* AllocatedVT = AllocatedVTsToMap[Index];
-			const uint32 vDimensions = AllocatedVT->GetDimensions();
-			const uint32 BaseTileX = AllocatedVT->GetVirtualPageX();
-			const uint32 BaseTileY = AllocatedVT->GetVirtualPageY();
-			FVirtualTextureSpace* Space = AllocatedVT->GetSpace();
-
-			uint32 NumFullyMappedLayers = 0u;
-			for (uint32 PageTableLayerIndex = 0u; PageTableLayerIndex < AllocatedVT->GetNumPageTableLayers(); ++PageTableLayerIndex)
-			{
-				uint32 ProducerIndex = AllocatedVT->GetProducerIndexForPageTableLayer(PageTableLayerIndex);
-				const FVirtualTextureProducerHandle ProducerHandle = AllocatedVT->GetUniqueProducerHandle(ProducerIndex);
-				const FVirtualTextureProducer* Producer = Producers.FindProducer(ProducerHandle);
-				if (!Producer)
-				{
-					++NumFullyMappedLayers;
-					continue;
-				}
-
-				uint32 ProducerPhysicalGroupIndex = AllocatedVT->GetProducerPhysicalGroupIndexForPageTableLayer(PageTableLayerIndex);
-
-				const uint32 ProducerMipBias = AllocatedVT->GetUniqueProducerMipBias(ProducerIndex);
-				const uint32 WidthInTiles = Producer->GetWidthInTiles();
-				const uint32 HeightInTiles = Producer->GetHeightInTiles();
-				const uint32 MaxLevel = AllocatedVT->GetMaxLevel();
-				const uint32 LocalMaxLevel = FMath::Min(Producer->GetMaxLevel(), MaxLevel - ProducerMipBias);
-
-				FVirtualTexturePhysicalSpace* PhysicalSpace = AllocatedVT->GetPhysicalSpaceForPageTableLayer(PageTableLayerIndex);
-				FTexturePagePool& PagePool = PhysicalSpace->GetPagePool();
-				FTexturePageMap& PageMap = Space->GetPageMapForPageTableLayer(PageTableLayerIndex);
-				
-				bool bIsLayerFullyMapped = false;
-				for (uint32 Local_vLevel = 0; Local_vLevel <= LocalMaxLevel; ++Local_vLevel)
-				{
-					const uint32 vLevel = Local_vLevel + ProducerMipBias;
-					check(vLevel <= MaxLevel);
-
-					const uint32 MipScaleFactor = (1u << Local_vLevel);
-					const uint32 LevelWidthInTiles = FMath::DivideAndRoundUp(WidthInTiles, MipScaleFactor);
-					const uint32 LevelHeightInTiles = FMath::DivideAndRoundUp(HeightInTiles, MipScaleFactor);
-
-					uint32 NumNonResidentPages = 0u;
-					for (uint32 TileY = 0u; TileY < LevelHeightInTiles; ++TileY)
-					{
-						for (uint32 TileX = 0u; TileX < LevelWidthInTiles; ++TileX)
-						{
-							const uint32 vAddress = FMath::MortonCode2(BaseTileX + (TileX << vLevel)) | (FMath::MortonCode2(BaseTileY + (TileY << vLevel)) << 1);
-							uint32 pAddress = PageMap.FindPageAddress(vLevel, vAddress);
-							if (pAddress == ~0u)
-							{
-								uint32 Local_vAddress = FMath::MortonCode2(TileX) | (FMath::MortonCode2(TileY) << 1);
-
-								const uint32 LocalMipBias = Producer->GetVirtualTexture()->GetLocalMipBias(Local_vLevel, Local_vAddress);
-								Local_vAddress >>= (LocalMipBias * 2u);
-
-								pAddress = PagePool.FindPageAddress(ProducerHandle, ProducerPhysicalGroupIndex, Local_vAddress, Local_vLevel + LocalMipBias);
-								if (pAddress != ~0u)
-								{
-									PagePool.MapPage(Space, PhysicalSpace, PageTableLayerIndex, MaxLevel, vLevel, vAddress, vLevel + LocalMipBias, pAddress);
-								}
-								else
-								{
-									++NumNonResidentPages;
-								}
-							}
-						}
-					}
-
-					if (NumNonResidentPages == 0u && !bIsLayerFullyMapped)
-					{
-						bIsLayerFullyMapped = true;
-						++NumFullyMappedLayers;
-					}
-				}
-			}
-
-			if (NumFullyMappedLayers < AllocatedVT->GetNumPageTableLayers())
-			{
-				++Index;
-			}
-			else
-			{
-				// Remove from list as long as we can fully map at least one mip level of the VT....this way we guarantee all tiles at least have some valid data (even if low resolution)
-				// Normally we expect to be able to at least map the least-detailed mip, since those tiles should always be locked/resident
-				// It's possible during loading that they may not be available for a few frames however
-				AllocatedVTsToMap.RemoveAtSwap(Index, 1, false);
-			}
-		}
-
-		AllocatedVTsToMap.Shrink();
-	}
-
 	// Finalize requests
 	{
 		SCOPE_CYCLE_COUNTER(STAT_ProcessRequests_Finalize);
diff --git a/Engine/Source/Runtime/Renderer/Private/VT/VirtualTextureSystem.h b/Engine/Source/Runtime/Renderer/Private/VT/VirtualTextureSystem.h
--- a/Engine/Source/Runtime/Renderer/Private/VT/VirtualTextureSystem.h
+++ b/Engine/Source/Runtime/Renderer/Private/VT/VirtualTextureSystem.h
@@ -44,8 +44,6 @@ public:
 
 	IAllocatedVirtualTexture* AllocateVirtualTexture(const FAllocatedVTDescription& Desc);
 	void DestroyVirtualTexture(IAllocatedVirtualTexture* AllocatedVT);
-	void ReleaseVirtualTexture(FAllocatedVirtualTexture* AllocatedVT);
-	void RemoveAllocatedVT(FAllocatedVirtualTexture* AllocatedVT);
 
 	FVirtualTextureProducerHandle RegisterProducer(const FVTProducerDescription& InDesc, IVirtualTexture* InProducer);
 	void ReleaseProducer(const FVirtualTextureProducerHandle& Handle);
@@ -82,7 +80,7 @@ private:
 	FVirtualTextureSystem();
 	~FVirtualTextureSystem();
 
-	void DestroyPendingVirtualTextures();
+	void DestroyPendingVirtualTextures(bool bForceDestroyAll);
 	void ReleasePendingSpaces();
 
 	void RequestTilesForRegionInternal(const IAllocatedVirtualTexture* AllocatedVT, const FVector2D& InScreenSpaceSize, const FVector2D& InViewportPosition, const FVector2D& InViewportSize, const FVector2D& InUV0, const FVector2D& InUV1, int32 InMipLevel);
@@ -120,8 +118,8 @@ private:
 	TArray<FVirtualTexturePhysicalSpace*> PhysicalSpaces;
 	FVirtualTextureProducerCollection Producers;
 
-	FCriticalSection PendingDeleteLock;
-	TArray<FAllocatedVirtualTexture*> PendingDeleteAllocatedVTs;
+	FCriticalSection AllocatedVTLock;
+	TArray<IAllocatedVirtualTexture*> PendingDeleteAllocatedVTs;
 
 	TMap<FAllocatedVTDescription, FAllocatedVirtualTexture*> AllocatedVTs;
 
@@ -151,7 +149,6 @@ private:
 	TSet<FVirtualTextureLocalTile> ContinuousUpdateTilesToProduce;
 	TSet<FVirtualTextureLocalTile> MappedTilesToProduce;
 	TArray<FVirtualTextureLocalTile> TransientCollectedPages;
-	TArray<FAllocatedVirtualTexture*> AllocatedVTsToMap;
 	TArray<IVirtualTextureFinalizer*> Finalizers;
 };
 
