Fix a bug of CSM Caching.

diff --git a/Engine/Source/Runtime/Renderer/Private/ShadowDepthRendering.cpp b/Engine/Source/Runtime/Renderer/Private/ShadowDepthRendering.cpp
--- a/Engine/Source/Runtime/Renderer/Private/ShadowDepthRendering.cpp
+++ b/Engine/Source/Runtime/Renderer/Private/ShadowDepthRendering.cpp
@@ -1983,10 +1983,8 @@ void FSceneRenderer::RenderShadowDepthMaps(FRDGBuilder& GraphBuilder, FInstanceC
 		GetLightNameForDrawEvent(ProjectedShadowInfo->GetLightSceneInfo().Proxy, LightNameWithLevel);
 		RDG_EVENT_SCOPE(GraphBuilder, "Cubemap %s %u^2", *LightNameWithLevel, TargetSize.X, TargetSize.Y);
 
-		const FCachedShadowMapData& CachedShadowMapData = Scene->GetCachedShadowMapDataRef(ProjectedShadowInfo->GetLightSceneInfo().Id, FMath::Max(ProjectedShadowInfo->CascadeSettings.ShadowSplitIndex, 0));
-
 		// Only clear when we're not copying from a cached shadow map.
-		if (ProjectedShadowInfo->CacheMode != SDCM_MovablePrimitivesOnly || !CachedShadowMapData.bCachedShadowMapHasPrimitives)
+		if (ProjectedShadowInfo->CacheMode != SDCM_MovablePrimitivesOnly || !Scene->GetCachedShadowMapDataRef(ProjectedShadowInfo->GetLightSceneInfo().Id, FMath::Max(ProjectedShadowInfo->CascadeSettings.ShadowSplitIndex, 0)).bCachedShadowMapHasPrimitives)
 		{
 			AddClearShadowDepthPass(GraphBuilder, ShadowDepthTexture);
 		}
