Fix hair LUTs not being correctly initialized.

diff --git a/Engine/Source/Runtime/Renderer/Private/HairStrands/HairStrandsLUT.cpp b/Engine/Source/Runtime/Renderer/Private/HairStrands/HairStrandsLUT.cpp
--- a/Engine/Source/Runtime/Renderer/Private/HairStrands/HairStrandsLUT.cpp
+++ b/Engine/Source/Runtime/Renderer/Private/HairStrands/HairStrandsLUT.cpp
@@ -354,7 +354,7 @@ static FRDGTextureRef InternalGetHairLUT(FRDGBuilder& GraphBuilder, const FViewI
 	}
 	else if (Type == FHairLUTType::HairLUTType_MeanEnergy)
 	{
-		const bool bNeedUpdate = GSystemTextures.HairLUT1.GetReference() == nullptr;
+		const bool bNeedUpdate = GSystemTextures.HairLUT1.GetReference() == nullptr || GSystemTextures.HairLUT1.GetReference()->GetRenderTargetItem().ShaderResourceTexture->GetSizeXYZ() != FIntVector(GHairLUTIncidentAngleCount, GHairLUTRoughnessCount, GHairLUTAbsorptionCount);
 		if (bNeedUpdate)
 		{
 			Out = AddHairLUTPass(GraphBuilder, View, HairLUTType_MeanEnergy);
@@ -367,7 +367,7 @@ static FRDGTextureRef InternalGetHairLUT(FRDGBuilder& GraphBuilder, const FViewI
 	}
 	else if (Type == FHairLUTType::HairLUTType_Coverage)
 	{
-		const bool bNeedUpdate = GSystemTextures.HairLUT2.GetReference() == nullptr;
+		const bool bNeedUpdate = GSystemTextures.HairLUT2.GetReference() == nullptr || GSystemTextures.HairLUT2.GetReference()->GetRenderTargetItem().ShaderResourceTexture->GetSizeXYZ() != FIntVector(FHairCountToCoverageData::HairCount, FHairCountToCoverageData::HairRadiusCount, 1);
 		if (bNeedUpdate)
 		{
 			Out = AddHairCoverageLUTPass(GraphBuilder, View);
