Implemented a ray tracing instance visualization mode

diff --git a/Engine/Shaders/Private/RayTracing/RayTracingDebug.usf b/Engine/Shaders/Private/RayTracing/RayTracingDebug.usf
--- a/Engine/Shaders/Private/RayTracing/RayTracingDebug.usf
+++ b/Engine/Shaders/Private/RayTracing/RayTracingDebug.usf
@@ -14,14 +14,37 @@ uint VisualizationMode;
 uint ShouldUsePreExposure;
 uint OpaqueOnly;
 
-RAY_TRACING_ENTRY_RAYGEN(RayTracingDebugMainRGS)
+// TODO: move MurmurMix and IntToColor into a common header, shared with Nanite/Visualize.usf
+
+uint MurmurMix(uint Hash)
 {
-	uint2 PixelCoord = DispatchRaysIndex().xy + View.ViewRectMin;
+	Hash ^= Hash >> 16;
+	Hash *= 0x85ebca6b;
+	Hash ^= Hash >> 13;
+	Hash *= 0xc2b2ae35;
+	Hash ^= Hash >> 16;
+	return Hash;
+}
 
-	float2 RenderTargetUV = (float2(PixelCoord) + .5f) * View.BufferSizeAndInvSize.zw;
+float3 IntToColor(uint Index)
+{
+	uint Hash = MurmurMix(Index);
 
-	RayDesc Ray = CreatePrimaryRay(RenderTargetUV);
-    FRayCone RayCone = (FRayCone)0;
+	float3 Color = float3
+	(
+		(Hash >> 0) & 255,
+		(Hash >> 8) & 255,
+		(Hash >> 16) & 255
+	);
+
+	return Color * (1.0f / 255.0f);
+}
+
+float4 DebugRayTracingMaterial(RayDesc Ray)
+{
+	uint2 PixelCoord = DispatchRaysIndex().xy + View.ViewRectMin;
+
+	FRayCone RayCone = (FRayCone)0;
 	RayCone.SpreadAngle = View.EyeToPixelSpreadAngle;
 
 	const uint RayFlags = 0;
@@ -113,6 +136,58 @@ RAY_TRACING_ENTRY_RAYGEN(RayTracingDebugMainRGS)
 			}
 		}
 	}
+	return Result;
+}
+
+struct FRayTracingDebugPayload : FMinimalPayload
+{
+	uint InstanceHash;
+};
+
+float4 DebugRayTracingInstance(RayDesc Ray)
+{
+	float4 Result = (float4)0;
+
+	FRayTracingDebugPayload Payload = (FRayTracingDebugPayload)0;
+	Payload.SetMiss();
+
+	TraceRay(
+		TLAS,
+		0 /*RayFlags*/,
+		0xff /*InstanceInclusionMask*/,
+		0 /*RayContributionToHitGroupIndex*/,
+		RAY_TRACING_NUM_SHADER_SLOTS /*MultiplierForGeometryContributionToShaderIndex*/,
+		0 /*MissShaderIndex*/,
+		Ray,
+		Payload);
+
+	if (Payload.IsHit())
+	{
+		Result.rgb = IntToColor(Payload.InstanceHash);
+		Result.a = 1;
+	}
+
+	return Result;
+}
+
+RAY_TRACING_ENTRY_RAYGEN(RayTracingDebugMainRGS)
+{
+	uint2 PixelCoord = DispatchRaysIndex().xy + View.ViewRectMin;
+
+	float2 RenderTargetUV = (float2(PixelCoord) + .5f) * View.BufferSizeAndInvSize.zw;
+
+	RayDesc Ray = CreatePrimaryRay(RenderTargetUV);
+
+	float4 Result = (float4)0;
+
+	if (VisualizationMode == RAY_TRACING_DEBUG_VIZ_INSTANCES)
+	{
+		Result = DebugRayTracingInstance(Ray);
+	}
+	else
+	{
+		Result = DebugRayTracingMaterial(Ray);
+	}
 
 	if (ShouldUsePreExposure)
 	{
@@ -122,3 +197,26 @@ RAY_TRACING_ENTRY_RAYGEN(RayTracingDebugMainRGS)
 
 	Output[PixelCoord] = Result;
 }
+
+RAY_TRACING_ENTRY_CLOSEST_HIT(RayTracingDebugMainCHS,
+	FRayTracingDebugPayload, Payload,
+	FDefaultAttributes, Attributes)
+{
+	Payload.HitT = RayTCurrent();
+
+	// Can't use InstanceIndex() because it's not stable due to culling.
+	// Instad we can use InstanceID() (user data on the instance) combined with
+	// a hash of the instance transform. This gives us a reasonable overall instance hash
+	// that can be used for debug visualizations. 
+	// Dynamic objects can't be visualized in a stable way, as their transform is constantly updated.
+
+	Payload.InstanceHash = InstanceID();
+	float4x3 Mat = ObjectToWorld4x3();
+	for (int i = 0; i < 4; ++i)
+	{
+		for (int j = 0; j < 3; ++j)
+		{
+			Payload.InstanceHash += MurmurMix(asuint(Mat[i][j]));
+		}
+	}
+}
diff --git a/Engine/Shaders/Shared/RaytracingDebugDefinitions.h b/Engine/Shaders/Shared/RaytracingDebugDefinitions.h
--- a/Engine/Shaders/Shared/RaytracingDebugDefinitions.h
+++ b/Engine/Shaders/Shared/RaytracingDebugDefinitions.h
@@ -29,4 +29,5 @@
 #define RAY_TRACING_DEBUG_VIZ_PRIMARY_RAYS				19  
 #define RAY_TRACING_DEBUG_VIZ_WORLD_TANGENT				20 
 #define RAY_TRACING_DEBUG_VIZ_ANISOTROPY				21 
+#define RAY_TRACING_DEBUG_VIZ_INSTANCES					22
 
diff --git a/Engine/Source/Editor/UnrealEd/Private/RayTracingDebugVisualizationMenuCommands.cpp b/Engine/Source/Editor/UnrealEd/Private/RayTracingDebugVisualizationMenuCommands.cpp
--- a/Engine/Source/Editor/UnrealEd/Private/RayTracingDebugVisualizationMenuCommands.cpp
+++ b/Engine/Source/Editor/UnrealEd/Private/RayTracingDebugVisualizationMenuCommands.cpp
@@ -53,6 +53,7 @@ void FRayTracingDebugVisualizationMenuCommands::CreateRayTracingDebugVisualizati
 	RayTracingDebugModeNames.Add(LOCTEXT("PrimaryRays", "PrimaryRays"));
 	RayTracingDebugModeNames.Add(LOCTEXT("World Tangent", "World Tangent"));
 	RayTracingDebugModeNames.Add(LOCTEXT("Anisotropy", "Anisotropy"));
+	RayTracingDebugModeNames.Add(LOCTEXT("Instances", "Instances"));
 
 	for ( int32 RayTracingDebugIndex = 0; RayTracingDebugIndex < RayTracingDebugModeNames.Num(); ++RayTracingDebugIndex)
 	{
diff --git a/Engine/Source/Runtime/Renderer/Private/RayTracing/RayTracingDebug.cpp b/Engine/Source/Runtime/Renderer/Private/RayTracing/RayTracingDebug.cpp
--- a/Engine/Source/Runtime/Renderer/Private/RayTracing/RayTracingDebug.cpp
+++ b/Engine/Source/Runtime/Renderer/Private/RayTracing/RayTracingDebug.cpp
@@ -53,6 +53,30 @@ class FRayTracingDebugRGS : public FGlobalShader
 };
 IMPLEMENT_GLOBAL_SHADER(FRayTracingDebugRGS, "/Engine/Private/RayTracing/RayTracingDebug.usf", "RayTracingDebugMainRGS", SF_RayGen);
 
+class FRayTracingDebugCHS : public FGlobalShader
+{
+	DECLARE_SHADER_TYPE(FRayTracingDebugCHS, Global);
+
+	static void ModifyCompilationEnvironment(const FGlobalShaderPermutationParameters& Parameters, FShaderCompilerEnvironment& OutEnvironment)
+	{
+		FGlobalShader::ModifyCompilationEnvironment(Parameters, OutEnvironment);
+	}
+
+public:
+
+	static bool ShouldCompilePermutation(const FGlobalShaderPermutationParameters& Parameters)
+	{
+		return ShouldCompileRayTracingShadersForProject(Parameters.Platform);
+	}
+
+	FRayTracingDebugCHS() = default;
+	FRayTracingDebugCHS(const ShaderMetaType::CompiledShaderInitializerType& Initializer)
+		: FGlobalShader(Initializer)
+	{}
+};
+IMPLEMENT_SHADER_TYPE(, FRayTracingDebugCHS, TEXT("/Engine/Private/RayTracing/RayTracingDebug.usf"), TEXT("RayTracingDebugMainCHS"), SF_RayHitGroup);
+
+
 void FDeferredShadingSceneRenderer::PrepareRayTracingDebug(const FViewInfo& View, TArray<FRHIRayTracingShader*>& OutRayGenShaders)
 {
 	// Declare all RayGen shaders that require material closest hit shaders to be bound
@@ -90,6 +114,7 @@ void FDeferredShadingSceneRenderer::RenderRayTracingDebug(FRDGBuilder& GraphBuil
 		RayTracingDebugVisualizationModes.Emplace(FName(*LOCTEXT("PrimaryRays", "PrimaryRays").ToString()),										RAY_TRACING_DEBUG_VIZ_PRIMARY_RAYS);
 		RayTracingDebugVisualizationModes.Emplace(FName(*LOCTEXT("World Tangent", "World Tangent").ToString()),									RAY_TRACING_DEBUG_VIZ_WORLD_TANGENT);
 		RayTracingDebugVisualizationModes.Emplace(FName(*LOCTEXT("Anisotropy", "Anisotropy").ToString()),										RAY_TRACING_DEBUG_VIZ_ANISOTROPY);
+		RayTracingDebugVisualizationModes.Emplace(FName(*LOCTEXT("Instances", "Instances").ToString()),											RAY_TRACING_DEBUG_VIZ_INSTANCES);
 	}
 
 	uint32 DebugVisualizationMode;
@@ -129,13 +154,27 @@ void FDeferredShadingSceneRenderer::RenderRayTracingDebug(FRDGBuilder& GraphBuil
 		return;
 	}
 
-
 	FGlobalShaderMap* ShaderMap = GetGlobalShaderMap(FeatureLevel);
-
 	auto RayGenShader = ShaderMap->GetShader<FRayTracingDebugRGS>();
 
 	FRayTracingPipelineState* Pipeline = View.RayTracingMaterialPipeline;
 
+	if (DebugVisualizationMode == RAY_TRACING_DEBUG_VIZ_INSTANCES)
+	{
+		FRayTracingPipelineStateInitializer Initializer;
+
+		FRHIRayTracingShader* RayGenShaderTable[] = { RayGenShader.GetRayTracingShader() };
+		Initializer.SetRayGenShaderTable(RayGenShaderTable);
+
+		auto ClosestHitShader = ShaderMap->GetShader<FRayTracingDebugCHS>();
+		FRHIRayTracingShader* HitGroupTable[] = { ClosestHitShader.GetRayTracingShader() };
+		Initializer.SetHitGroupTable(HitGroupTable);
+		Initializer.bAllowHitGroupIndexing = false; // Use the same hit shader for all geometry in the scene by disabling SBT indexing.
+		Initializer.MaxPayloadSizeInBytes = RAY_TRACING_MAX_ALLOWED_PAYLOAD_SIZE;
+
+		Pipeline = PipelineStateCache::GetAndOrCreateRayTracingPipelineState(GraphBuilder.RHICmdList, Initializer);
+	}
+
 	FRayTracingDebugRGS::FParameters* RayGenParameters = GraphBuilder.AllocParameters<FRayTracingDebugRGS::FParameters>();
 
 	RayGenParameters->VisualizationMode = DebugVisualizationMode;
