Add hair composition point after translucent, but before translucent render after DOF.

diff --git a/Engine/Source/Runtime/Renderer/Private/HairStrands/HairStrandsUtils.cpp b/Engine/Source/Runtime/Renderer/Private/HairStrands/HairStrandsUtils.cpp
--- a/Engine/Source/Runtime/Renderer/Private/HairStrands/HairStrandsUtils.cpp
+++ b/Engine/Source/Runtime/Renderer/Private/HairStrands/HairStrandsUtils.cpp
@@ -38,7 +38,7 @@ static int32 GHairVisibilityRectOptimEnable = 0;
 static FAutoConsoleVariableRef CVarHairVisibilityRectOptimEnable(TEXT("r.HairStrands.RectLightingOptim"), GHairVisibilityRectOptimEnable, TEXT("Hair Visibility use projected view rect to light only relevant pixels"));
 
 static int32 GHairStrandsComposeAfterTranslucency = 1;
-static FAutoConsoleVariableRef CVarHairStrandsComposeAfterTranslucency(TEXT("r.HairStrands.ComposeAfterTranslucency"), GHairStrandsComposeAfterTranslucency, TEXT("0: Compose hair before translucent objects. 1: Compose hair after translucent objects, but before separate translucent objects. 2: Compose hair after seperate translucent objects"));
+static FAutoConsoleVariableRef CVarHairStrandsComposeAfterTranslucency(TEXT("r.HairStrands.ComposeAfterTranslucency"), GHairStrandsComposeAfterTranslucency, TEXT("0: Compose hair before translucent objects. 1: Compose hair after translucent objects, but before separate translucent objects. 2: Compose hair after all/seperate translucent objects, 3: Compose hair after translucent objects but before translucent render after DOF (which allows depth testing against hair depth)"));
 
 static float GHairDualScatteringRoughnessOverride = 0;
 static FAutoConsoleVariableRef CVarHairDualScatteringRoughnessOverride(TEXT("r.HairStrands.DualScatteringRoughness"), GHairDualScatteringRoughnessOverride, TEXT("Override all roughness for the dual scattering evaluation. 0 means no override. Default:0"));
@@ -63,6 +63,7 @@ EHairStrandsCompositionType GetHairStrandsComposition()
 	case 0	: return EHairStrandsCompositionType::BeforeTranslucent;
 	case 1	: return EHairStrandsCompositionType::AfterTranslucent;
 	case 2	: return EHairStrandsCompositionType::AfterSeparateTranslucent;
+	case 3	: return EHairStrandsCompositionType::AfterTranslucentBeforeTranslucentAfterDOF;
 	default	: return EHairStrandsCompositionType::BeforeTranslucent;
 	}
 }
diff --git a/Engine/Source/Runtime/Renderer/Private/HairStrands/HairStrandsUtils.h b/Engine/Source/Runtime/Renderer/Private/HairStrands/HairStrandsUtils.h
--- a/Engine/Source/Runtime/Renderer/Private/HairStrands/HairStrandsUtils.h
+++ b/Engine/Source/Runtime/Renderer/Private/HairStrands/HairStrandsUtils.h
@@ -67,6 +67,7 @@ enum class  EHairStrandsCompositionType : uint8
 	BeforeTranslucent,
 	AfterTranslucent,
 	AfterSeparateTranslucent,
+	AfterTranslucentBeforeTranslucentAfterDOF
 };
 
 EHairStrandsCompositionType GetHairStrandsComposition();
diff --git a/Engine/Source/Runtime/Renderer/Private/TranslucentRendering.cpp b/Engine/Source/Runtime/Renderer/Private/TranslucentRendering.cpp
--- a/Engine/Source/Runtime/Renderer/Private/TranslucentRendering.cpp
+++ b/Engine/Source/Runtime/Renderer/Private/TranslucentRendering.cpp
@@ -15,6 +15,7 @@
 #include "Lumen/LumenTranslucencyVolumeLighting.h"
 #include "VirtualShadowMaps/VirtualShadowMapArray.h"
 #include "Strata/Strata.h"
+#include "HairStrands/HairStrandsUtils.h"
 
 DECLARE_CYCLE_STAT(TEXT("TranslucencyTimestampQueryFence Wait"), STAT_TranslucencyTimestampQueryFence_Wait, STATGROUP_SceneRendering);
 DECLARE_CYCLE_STAT(TEXT("TranslucencyTimestampQuery Wait"), STAT_TranslucencyTimestampQuery_Wait, STATGROUP_SceneRendering);
@@ -1176,6 +1177,10 @@ void FDeferredShadingSceneRenderer::RenderTranslucency(
 	if (ViewFamily.AllowTranslucencyAfterDOF())
 	{
 		RenderTranslucencyInner(GraphBuilder, SceneTextures, TranslucentLightingVolumeTextures, OutSeparateTranslucencyTextures, ViewsToRender, SceneColorCopyTexture, ETranslucencyPass::TPT_StandardTranslucency, InstanceCullingManager);
+		if (GetHairStrandsComposition() == EHairStrandsCompositionType::AfterTranslucentBeforeTranslucentAfterDOF)
+		{
+			RenderHairComposition(GraphBuilder, Views, SceneTextures.Color.Target, SceneTextures.Depth.Target);
+		}
 		RenderTranslucencyInner(GraphBuilder, SceneTextures, TranslucentLightingVolumeTextures, OutSeparateTranslucencyTextures, ViewsToRender, SceneColorCopyTexture, ETranslucencyPass::TPT_TranslucencyAfterDOF, InstanceCullingManager);
 		RenderTranslucencyInner(GraphBuilder, SceneTextures, TranslucentLightingVolumeTextures, OutSeparateTranslucencyTextures, ViewsToRender, SceneColorCopyTexture, ETranslucencyPass::TPT_TranslucencyAfterDOFModulate, InstanceCullingManager);
 	}
