Fix nullptr dereference.

diff --git a/Engine/Source/Runtime/Renderer/Private/VT/VirtualTextureSystem.cpp b/Engine/Source/Runtime/Renderer/Private/VT/VirtualTextureSystem.cpp
--- a/Engine/Source/Runtime/Renderer/Private/VT/VirtualTextureSystem.cpp
+++ b/Engine/Source/Runtime/Renderer/Private/VT/VirtualTextureSystem.cpp
@@ -2179,7 +2179,7 @@ void FVirtualTextureSystem::SubmitRequests(FRDGBuilder& GraphBuilder, ERHIFeatur
 			{
 				const FVirtualTexturePhysicalSpace* PhysicalSpace = PhysicalSpaces[SpaceIndex];
 				// Once a pool fails an allocation, keep displaying the message for a minimum number of frames, to avoid the message quickly flickering on/off
-				if (Frame <= PhysicalSpace->GetLastFrameOversubscribed() + 60u)
+				if (PhysicalSpace != nullptr && Frame <= PhysicalSpace->GetLastFrameOversubscribed() + 60u)
 				{
 					const FVTPhysicalSpaceDescription& Desc = PhysicalSpace->GetDescription();
 
