AllocatedVTs that use private space are deleted immediately, rather than staying cached for some frames

diff --git a/Engine/Source/Runtime/Renderer/Private/VT/VirtualTextureSystem.cpp b/Engine/Source/Runtime/Renderer/Private/VT/VirtualTextureSystem.cpp
--- a/Engine/Source/Runtime/Renderer/Private/VT/VirtualTextureSystem.cpp
+++ b/Engine/Source/Runtime/Renderer/Private/VT/VirtualTextureSystem.cpp
@@ -457,6 +457,9 @@ IAllocatedVirtualTexture* FVirtualTextureSystem::AllocateVirtualTexture(const FA
 		return AllocatedVT;
 	}
 
+	// Clear out any AllocatedVTs that are ready to be deleted, before trying a new allocation
+	DestroyPendingVirtualTextures(false);
+
 	uint32 BlockWidthInTiles = 0u;
 	uint32 BlockHeightInTiles = 0u;
 	uint32 WidthInBlocks = 0u;
@@ -553,9 +556,12 @@ void FVirtualTextureSystem::DestroyPendingVirtualTextures(bool bForceDestroyAll)
 			while (Index < PendingDeleteAllocatedVTs.Num())
 			{
 				IAllocatedVirtualTexture* AllocatedVT = PendingDeleteAllocatedVTs[Index];
+				const FAllocatedVTDescription& Desc = AllocatedVT->GetDescription();
 				check(AllocatedVT->NumRefs == 0);
 
-				if (CurrentFrame >= AllocatedVT->FrameDeleted + 60u)
+				// If the AllocatedVT is using a private space release it immediately, we don't want to hold references to these private spaces any longer then needed
+				// Otherwise keep deleted VTs around for a few frames, in case they are reused
+				if (Desc.bPrivateSpace || CurrentFrame >= AllocatedVT->FrameDeleted + 60u)
 				{
 					AllocatedVTsToDelete.Add(AllocatedVT);
 					PendingDeleteAllocatedVTs.RemoveAtSwap(Index, 1, false);
