Fix PrimitiveId doesn't get updated during primitive array swapping

diff --git a/Engine/Source/Runtime/Renderer/Private/RendererScene.cpp b/Engine/Source/Runtime/Renderer/Private/RendererScene.cpp
--- a/Engine/Source/Runtime/Renderer/Private/RendererScene.cpp
+++ b/Engine/Source/Runtime/Renderer/Private/RendererScene.cpp
@@ -4291,6 +4291,11 @@ void FScene::UpdateAllPrimitiveSceneInfos(FRDGBuilder& GraphBuilder, bool bAsync
 
 							GPUScene.AddPrimitiveToUpdate(DestIndex);
 							LumenSceneData->AddPrimitiveToUpdate(DestIndex);
+
+						#if RHI_RAYTRACING
+							// Update cached PrimitiveIndex after an index swap
+							Primitives[DestIndex]->CachedRayTracingInstance.DefaultUserData = DestIndex;
+						#endif
 						}
 					}
 				}
