Integrate Project Mobile CSM to the ShadowMaskTexture to have a fade between the last Cascade and SDF

diff --git a/Engine/Shaders/Private/MobileBasePassPixelShader.usf b/Engine/Shaders/Private/MobileBasePassPixelShader.usf
--- a/Engine/Shaders/Private/MobileBasePassPixelShader.usf
+++ b/Engine/Shaders/Private/MobileBasePassPixelShader.usf
@@ -752,7 +752,7 @@ void Main(
 	IndirectIrradiance *= MaterialAO;
 
 	float  ShadowPositionZ = 0;
-#if DIRECTIONAL_LIGHT_CSM && !MATERIAL_SHADINGMODEL_SINGLELAYERWATER
+#if DIRECTIONAL_LIGHT_CSM && !MATERIAL_SHADINGMODEL_SINGLELAYERWATER && !APPLY_DISTANCE_FIELD
 	// Cascaded Shadow Map
 	if (IsCSMEnabled())
 	{
diff --git a/Engine/Source/Runtime/Renderer/Private/DeferredShadingRenderer.h b/Engine/Source/Runtime/Renderer/Private/DeferredShadingRenderer.h
--- a/Engine/Source/Runtime/Renderer/Private/DeferredShadingRenderer.h
+++ b/Engine/Source/Runtime/Renderer/Private/DeferredShadingRenderer.h
@@ -664,14 +664,6 @@ private:
 		TRDGUniformBufferRef<FSceneTextureUniformParameters> SceneTexturesUniformBuffer,
 		FRDGTextureRef& BentNormalOutput) const;
 
-	void RenderShadowProjections(
-		FRDGBuilder& GraphBuilder,
-		const FMinimalSceneTextures& SceneTextures,
-		FRDGTextureRef ScreenShadowMaskTexture,
-		FRDGTextureRef ScreenShadowMaskSubPixelTexture,
-		const FLightSceneInfo* LightSceneInfo,
-		bool bProjectingForForwardShading);
-
 	void RenderDeferredShadowProjections(
 		FRDGBuilder& GraphBuilder,
 		const FMinimalSceneTextures& SceneTextures,
diff --git a/Engine/Source/Runtime/Renderer/Private/MobileShadingRenderer.cpp b/Engine/Source/Runtime/Renderer/Private/MobileShadingRenderer.cpp
--- a/Engine/Source/Runtime/Renderer/Private/MobileShadingRenderer.cpp
+++ b/Engine/Source/Runtime/Renderer/Private/MobileShadingRenderer.cpp
@@ -791,8 +791,8 @@ void FMobileSceneRenderer::Render(FRDGBuilder& GraphBuilder)
 
 		if (bRequiresDistanceFieldShadowingPass)
 		{
-			CSV_SCOPED_TIMING_STAT_EXCLUSIVE(RenderSDFShadowing);
-			RenderMobileSDFShadowing(GraphBuilder, SceneTextures.Depth.Resolve, Scene, Views, VisibleLightInfos);
+			CSV_SCOPED_TIMING_STAT_EXCLUSIVE(RenderMobileShadowProjections);
+			RenderMobileShadowProjections(GraphBuilder, SceneTextures.Depth.Resolve);
 		}
 
 		if (bShouldRenderHZB)
diff --git a/Engine/Source/Runtime/Renderer/Private/SceneRendering.h b/Engine/Source/Runtime/Renderer/Private/SceneRendering.h
--- a/Engine/Source/Runtime/Renderer/Private/SceneRendering.h
+++ b/Engine/Source/Runtime/Renderer/Private/SceneRendering.h
@@ -1981,6 +1981,15 @@ protected:
 		bool bProjectingForForwardShading,
 		bool bMobileModulatedProjections);
 
+
+	void RenderShadowProjections(
+		FRDGBuilder& GraphBuilder,
+		const FMinimalSceneTextures& SceneTextures,
+		FRDGTextureRef ScreenShadowMaskTexture,
+		FRDGTextureRef ScreenShadowMaskSubPixelTexture,
+		const FLightSceneInfo* LightSceneInfo,
+		bool bProjectingForForwardShading);
+
 	/** Finds a matching cached preshadow, if one exists. */
 	TRefCountPtr<FProjectedShadowInfo> GetCachedPreshadow(
 		const FLightPrimitiveInteraction* InParentInteraction,
@@ -2283,6 +2292,8 @@ protected:
 
 	/** Before SetupMobileBasePassAfterShadowInit, we need to update the uniform buffer and shadow info for all movable point lights.*/
 	void UpdateMovablePointLightUniformBufferAndShadowInfo();
+
+	void RenderMobileShadowProjections(FRDGBuilder& GraphBuilder, FRDGTextureRef SceneDepthTexture);
 private:
 	const bool bGammaSpace;
 	const bool bDeferredShading;
diff --git a/Engine/Source/Runtime/Renderer/Private/ShadowRendering.cpp b/Engine/Source/Runtime/Renderer/Private/ShadowRendering.cpp
--- a/Engine/Source/Runtime/Renderer/Private/ShadowRendering.cpp
+++ b/Engine/Source/Runtime/Renderer/Private/ShadowRendering.cpp
@@ -1866,9 +1866,14 @@ void FSceneRenderer::RenderShadowProjections(
 		}
 
 		FShadowProjectionPassParameters CommonPassParameters;
-		CommonPassParameters.SceneTextures = GetSceneTextureShaderParameters(SceneTextures.UniformBuffer);
+		CommonPassParameters.SceneTextures = SceneTextures.GetSceneTextureShaderParameters(View.FeatureLevel);
 		CommonPassParameters.HairStrands = HairStrands::BindHairStrandsViewUniformParameters(View);
-		CommonPassParameters.Strata = Strata::BindStrataGlobalUniformParameters(View.StrataSceneData);
+
+		if (Strata::IsStrataEnabled())
+		{
+			CommonPassParameters.Strata = Strata::BindStrataGlobalUniformParameters(View.StrataSceneData);
+		}
+		
 		CommonPassParameters.RenderTargets[0] = FRenderTargetBinding(OutputTexture, ERenderTargetLoadAction::ELoad);
 		CommonPassParameters.RenderTargets.DepthStencil =
 			bSubPixelShadow ?
@@ -1900,7 +1905,7 @@ void FSceneRenderer::RenderShadowProjections(
 	}
 }
 
-void FDeferredShadingSceneRenderer::RenderShadowProjections(
+void FSceneRenderer::RenderShadowProjections(
 	FRDGBuilder& GraphBuilder,
 	const FMinimalSceneTextures& SceneTextures,
 	FRDGTextureRef ScreenShadowMaskTexture,
@@ -2275,12 +2280,9 @@ void InitMobileSDFShadowingOutputs(FRHICommandListImmediate& RHICmdList, const F
 	}
 }
 
-void RenderMobileSDFShadowing(
+void FMobileSceneRenderer::RenderMobileShadowProjections(
 	FRDGBuilder& GraphBuilder, 
-	FRDGTextureRef SceneDepthTexture, 
-	const FScene* Scene,
-	const TArrayView<const FViewInfo> Views,
-	const TArrayView<const FVisibleLightInfo> VisibleLightInfos)
+	FRDGTextureRef SceneDepthTexture)
 {
 	FRDGTextureRef ScreenShadowMaskTexture = GraphBuilder.RegisterExternalTexture(GScreenSpaceShadowMaskTextureMobileOutputs.ScreenSpaceShadowMaskTextureMobile, TEXT("ScreenSpaceShadowMaskTextureMobile"));
 	AddClearRenderTargetPass(GraphBuilder, ScreenShadowMaskTexture);
@@ -2294,52 +2296,11 @@ void RenderMobileSDFShadowing(
 		const FVisibleLightInfo& VisibleLightInfo = VisibleLightInfos[LightSceneInfo->Id];
 		const FLightSceneProxy* LightSceneProxy = LightSceneInfo->Proxy;
 
-		// Allocate arrays using the graph allocator so we can safely reference them in passes.
-		using FProjectedShadowInfoArray = TArray<const FProjectedShadowInfo*, SceneRenderingAllocator>;
-		auto& DistanceFieldShadows = *GraphBuilder.AllocObject<FProjectedShadowInfoArray>();
-
-		for (int32 ShadowIndex = 0; ShadowIndex < VisibleLightInfo.ShadowsToProject.Num(); ShadowIndex++)
-		{
-			const FProjectedShadowInfo* ProjectedShadowInfo = VisibleLightInfo.ShadowsToProject[ShadowIndex];
-			if (ProjectedShadowInfo->bRayTracedDistanceField)
-			{
-				DistanceFieldShadows.Add(ProjectedShadowInfo);
-			}
-		}
-
-		if (DistanceFieldShadows.Num() > 0)
-		{
-			RDG_EVENT_SCOPE(GraphBuilder, "DistanceFieldShadows");
-
-			// Distance field shadows need to be renderer last as they blend over far shadow cascades.
-			for (int32 ViewIndex = 0; ViewIndex < Views.Num(); ViewIndex++)
-			{
-				const FViewInfo& View = Views[ViewIndex];
-				RDG_GPU_MASK_SCOPE(GraphBuilder, View.GPUMask);
-				RDG_EVENT_SCOPE_CONDITIONAL(GraphBuilder, Views.Num() > 1, "View%d", ViewIndex);
-
-				FIntRect ScissorRect;
-				if (!LightSceneProxy->GetScissorRect(ScissorRect, View, View.ViewRect))
-				{
-					ScissorRect = View.ViewRect;
-				}
-
-				if (ScissorRect.Area() > 0)
-				{
-					for (int32 ShadowIndex = 0; ShadowIndex < DistanceFieldShadows.Num(); ShadowIndex++)
-					{
-						const FProjectedShadowInfo* ProjectedShadowInfo = DistanceFieldShadows[ShadowIndex];
-						ProjectedShadowInfo->RenderRayTracedDistanceFieldProjection(
-							GraphBuilder,
-							SceneTextures,
-							ScreenShadowMaskTexture,
-							View,
-							ScissorRect,
-							false);
-					}
-				}
-			}
-		}
+		RenderShadowProjections(GraphBuilder, SceneTextures,
+			ScreenShadowMaskTexture,
+			nullptr,
+			LightSceneInfo,
+			false);
 	}
 }
 
diff --git a/Engine/Source/Runtime/Renderer/Private/ShadowRendering.h b/Engine/Source/Runtime/Renderer/Private/ShadowRendering.h
--- a/Engine/Source/Runtime/Renderer/Private/ShadowRendering.h
+++ b/Engine/Source/Runtime/Renderer/Private/ShadowRendering.h
@@ -58,11 +58,6 @@ void OverrideWithDefaultMaterialForShadowDepth(
 	);
 
 void InitMobileSDFShadowingOutputs(FRHICommandListImmediate& RHICmdList, const FIntPoint& Extent);
-void RenderMobileSDFShadowing(FRDGBuilder& GraphBuilder, 
-	FRDGTextureRef SceneDepthTexture, 
-	const FScene* Scene,
-	const TArrayView<const FViewInfo> Views,
-	const TArrayView<const FVisibleLightInfo> VisibleLightInfos);
 void ReleaseMobileSDFShadowingOutputs();
 
 enum EShadowDepthRenderMode
@@ -1140,7 +1135,8 @@ public:
 
 	static bool ShouldCompilePermutation(const FGlobalShaderPermutationParameters& Parameters)
 	{
-		return IsFeatureLevelSupported(Parameters.Platform, ERHIFeatureLevel::SM5);
+		return IsFeatureLevelSupported(Parameters.Platform, ERHIFeatureLevel::SM5)
+			|| (bUseTransmission == 0 && SubPixelShadow == 0 && IsMobileDistanceFieldEnabled(Parameters.Platform));
 	}
 
 	/**
