Groundwork for ray tracing performance visualization view modes

diff --git a/Engine/Shaders/Private/RayTracing/RayTracingDebug.usf b/Engine/Shaders/Private/RayTracing/RayTracingDebug.usf
--- a/Engine/Shaders/Private/RayTracing/RayTracingDebug.usf
+++ b/Engine/Shaders/Private/RayTracing/RayTracingDebug.usf
@@ -2,6 +2,8 @@
 
 #include "../Common.ush"
 #include "../ShadingCommon.ush"
+#include "../ColorMap.ush"
+
 #include "RayTracingCommon.ush"
 #include "RayTracingHitGroupCommon.ush"
 
@@ -13,6 +15,7 @@ RaytracingAccelerationStructure TLAS;
 uint VisualizationMode;
 uint ShouldUsePreExposure;
 uint OpaqueOnly;
+float TimingScale;
 
 // TODO: move MurmurMix and IntToColor into a common header, shared with Nanite/Visualize.usf
 
@@ -52,6 +55,10 @@ float4 DebugRayTracingMaterial(RayDesc Ray)
 	const bool bEnableSkyLightContribution = true;
 	const bool bIgnoreTranslucentMaterials = false;
 
+#if PLATFORM_SUPPORTS_SHADER_TIMESTAMP
+	uint64_t TimeBegin = GetShaderTimestamp();
+#endif
+
 	FMaterialClosestHitPayload Payload = TraceMaterialRay(
 		TLAS,
 		RayFlags,
@@ -62,6 +69,10 @@ float4 DebugRayTracingMaterial(RayDesc Ray)
 		bEnableSkyLightContribution,
 		bIgnoreTranslucentMaterials);
 
+#if PLATFORM_SUPPORTS_SHADER_TIMESTAMP
+	uint64_t TimeEnd = GetShaderTimestamp();
+#endif
+
 	float4 Result = float4(0, 0, 0, 1);
 
 	if (Payload.IsHit())
@@ -134,6 +145,11 @@ float4 DebugRayTracingMaterial(RayDesc Ray)
 			Result = float4(0.0, pow(AniG, 2.2), pow(AniB, 2.2), 1.0f);
 			break;
 			}
+		#if PLATFORM_SUPPORTS_SHADER_TIMESTAMP
+		case RAY_TRACING_DEBUG_VIZ_PERFORMANCE:
+			Result.rgb = ColorMapInferno(float(TimeEnd - TimeBegin) * TimingScale);
+			break;
+		#endif // PLATFORM_SUPPORTS_SHADER_TIMESTAMP
 		}
 	}
 	return Result;
@@ -194,7 +210,7 @@ RAY_TRACING_ENTRY_RAYGEN(RayTracingDebugMainRGS)
 		// Only  when the output is tonemapped
 		Result *= View.PreExposure;
 	}
-
+	
 	Output[PixelCoord] = Result;
 }
 
diff --git a/Engine/Shaders/Shared/RaytracingDebugDefinitions.h b/Engine/Shaders/Shared/RaytracingDebugDefinitions.h
--- a/Engine/Shaders/Shared/RaytracingDebugDefinitions.h
+++ b/Engine/Shaders/Shared/RaytracingDebugDefinitions.h
@@ -30,4 +30,5 @@
 #define RAY_TRACING_DEBUG_VIZ_WORLD_TANGENT				20 
 #define RAY_TRACING_DEBUG_VIZ_ANISOTROPY				21 
 #define RAY_TRACING_DEBUG_VIZ_INSTANCES					22
+#define RAY_TRACING_DEBUG_VIZ_PERFORMANCE				23 
 
diff --git a/Engine/Source/Editor/UnrealEd/Private/RayTracingDebugVisualizationMenuCommands.cpp b/Engine/Source/Editor/UnrealEd/Private/RayTracingDebugVisualizationMenuCommands.cpp
--- a/Engine/Source/Editor/UnrealEd/Private/RayTracingDebugVisualizationMenuCommands.cpp
+++ b/Engine/Source/Editor/UnrealEd/Private/RayTracingDebugVisualizationMenuCommands.cpp
@@ -55,6 +55,11 @@ void FRayTracingDebugVisualizationMenuCommands::CreateRayTracingDebugVisualizati
 	RayTracingDebugModeNames.Add(LOCTEXT("Anisotropy", "Anisotropy"));
 	RayTracingDebugModeNames.Add(LOCTEXT("Instances", "Instances"));
 
+	if (GRHISupportsShaderTimestamp)
+	{
+		RayTracingDebugModeNames.Add(LOCTEXT("Performance", "Performance"));
+	}
+
 	for ( int32 RayTracingDebugIndex = 0; RayTracingDebugIndex < RayTracingDebugModeNames.Num(); ++RayTracingDebugIndex)
 	{
 		const FText CommandNameText = RayTracingDebugModeNames[RayTracingDebugIndex];
diff --git a/Engine/Source/Runtime/RHI/Private/RHI.cpp b/Engine/Source/Runtime/RHI/Private/RHI.cpp
--- a/Engine/Source/Runtime/RHI/Private/RHI.cpp
+++ b/Engine/Source/Runtime/RHI/Private/RHI.cpp
@@ -1312,6 +1312,8 @@ bool GRHISupportsPixelShaderUAVs = true;
 
 bool GRHISupportsMeshShaders = false;
 
+bool GRHISupportsShaderTimestamp = false;
+
 FVertexElementTypeSupportInfo GVertexElementTypeSupport;
 
 RHI_API int32 volatile GCurrentTextureMemorySize = 0;
diff --git a/Engine/Source/Runtime/RHI/Public/RHI.h b/Engine/Source/Runtime/RHI/Public/RHI.h
--- a/Engine/Source/Runtime/RHI/Public/RHI.h
+++ b/Engine/Source/Runtime/RHI/Public/RHI.h
@@ -684,6 +684,12 @@ extern RHI_API bool GRHISupportsConservativeRasterization;
 /** true if the RHI supports Mesh and Amplification shaders */
 extern RHI_API bool GRHISupportsMeshShaders;
 
+/**
+* True if the RHI supports reading system timer in shaders via GetShaderTimestamp().
+* Individual shaders must be compiled with an appropriate vendor extension and check PLATFORM_SUPPORTS_SHADER_TIMESTAMP.
+*/
+extern RHI_API bool GRHISupportsShaderTimestamp;
+
 /** Tables of all MSAA sample offset for all MSAA supported. Use GetMSAASampleOffsets() to read it. */
 extern RHI_API FVector2D GRHIDefaultMSAASampleOffsets[1 + 2 + 4 + 8 + 16];
 
diff --git a/Engine/Source/Runtime/Renderer/Private/RayTracing/RayTracingDebug.cpp b/Engine/Source/Runtime/Renderer/Private/RayTracing/RayTracingDebug.cpp
--- a/Engine/Source/Runtime/Renderer/Private/RayTracing/RayTracingDebug.cpp
+++ b/Engine/Source/Runtime/Renderer/Private/RayTracing/RayTracingDebug.cpp
@@ -32,6 +32,12 @@ TAutoConsoleVariable<int32> CVarRayTracingDebugModeOpaqueOnly(
 	ECVF_RenderThreadSafe
 );
 
+static TAutoConsoleVariable<float> CVarRayTracingDebugTimingScale(
+	TEXT("r.RayTracing.DebugTimingScale"),
+	1.0f,
+	TEXT("Scaling factor for ray timing heat map visualization. (default = 1)\n")
+);
+
 class FRayTracingDebugRGS : public FGlobalShader
 {
 	DECLARE_GLOBAL_SHADER(FRayTracingDebugRGS)
@@ -40,6 +46,7 @@ class FRayTracingDebugRGS : public FGlobalShader
 	BEGIN_SHADER_PARAMETER_STRUCT(FParameters, )
 		SHADER_PARAMETER(uint32, VisualizationMode)
 		SHADER_PARAMETER(int32, ShouldUsePreExposure)
+		SHADER_PARAMETER(float, TimingScale)
 		SHADER_PARAMETER(int32, OpaqueOnly)
 		SHADER_PARAMETER_SRV(RaytracingAccelerationStructure, TLAS)
 		SHADER_PARAMETER_RDG_TEXTURE_UAV(RWTexture2D<float4>, Output)
@@ -115,6 +122,7 @@ void FDeferredShadingSceneRenderer::RenderRayTracingDebug(FRDGBuilder& GraphBuil
 		RayTracingDebugVisualizationModes.Emplace(FName(*LOCTEXT("World Tangent", "World Tangent").ToString()),									RAY_TRACING_DEBUG_VIZ_WORLD_TANGENT);
 		RayTracingDebugVisualizationModes.Emplace(FName(*LOCTEXT("Anisotropy", "Anisotropy").ToString()),										RAY_TRACING_DEBUG_VIZ_ANISOTROPY);
 		RayTracingDebugVisualizationModes.Emplace(FName(*LOCTEXT("Instances", "Instances").ToString()),											RAY_TRACING_DEBUG_VIZ_INSTANCES);
+		RayTracingDebugVisualizationModes.Emplace(FName(*LOCTEXT("Performance", "Performance").ToString()),										RAY_TRACING_DEBUG_VIZ_PERFORMANCE);
 	}
 
 	uint32 DebugVisualizationMode;
@@ -179,6 +187,7 @@ void FDeferredShadingSceneRenderer::RenderRayTracingDebug(FRDGBuilder& GraphBuil
 
 	RayGenParameters->VisualizationMode = DebugVisualizationMode;
 	RayGenParameters->ShouldUsePreExposure = View.Family->EngineShowFlags.Tonemapper;
+	RayGenParameters->TimingScale = CVarRayTracingDebugTimingScale.GetValueOnAnyThread() / 25000.0f;
 	RayGenParameters->OpaqueOnly = CVarRayTracingDebugModeOpaqueOnly.GetValueOnRenderThread();
 	RayGenParameters->TLAS = View.GetRayTracingSceneViewChecked();
 	RayGenParameters->ViewUniformBuffer = View.ViewUniformBuffer;
