Add option to r.CustomDepth.Order to select before or after base pass based on DBuffer project setting.
Made this new option the default.
Means that custom stencil "just works" with DBuffer decals.

diff --git a/Engine/Source/Runtime/Engine/Private/Materials/HLSLMaterialTranslator.cpp b/Engine/Source/Runtime/Engine/Private/Materials/HLSLMaterialTranslator.cpp
--- a/Engine/Source/Runtime/Engine/Private/Materials/HLSLMaterialTranslator.cpp
+++ b/Engine/Source/Runtime/Engine/Private/Materials/HLSLMaterialTranslator.cpp
@@ -6151,7 +6151,7 @@ void FHLSLMaterialTranslator::UseSceneTextureId(ESceneTextureId SceneTextureId,
 		const bool bSceneTextureSupportsDecal = SceneTextureId == PPI_SceneDepth || SceneTextureId == PPI_WorldNormal || SceneTextureId == PPI_CustomDepth || SceneTextureId == PPI_CustomStencil;
 		if (!bSceneTextureSupportsDecal)
 		{
-			// Note: For DBuffer decals CustomDepth and CustomStencil are only available if r.CustomDepth.Order = 0
+			// Note: For DBuffer decals CustomDepth and CustomStencil are not available if r.CustomDepth.Order = 1
 			Errorf(TEXT("Decals can only access SceneDepth, CustomDepth, CustomStencil, and WorldNormal."));
 		}
 
diff --git a/Engine/Source/Runtime/Renderer/Private/CustomDepthRendering.cpp b/Engine/Source/Runtime/Renderer/Private/CustomDepthRendering.cpp
--- a/Engine/Source/Runtime/Renderer/Private/CustomDepthRendering.cpp
+++ b/Engine/Source/Runtime/Renderer/Private/CustomDepthRendering.cpp
@@ -19,10 +19,11 @@ static TAutoConsoleVariable<int32> CVarCustomDepth(
 
 static TAutoConsoleVariable<int32> CVarCustomDepthOrder(
 	TEXT("r.CustomDepth.Order"),
-	1,
+	2,
 	TEXT("When CustomDepth (and CustomStencil) is getting rendered\n")
-	TEXT("  0: Before Base Pass (can be more efficient with AsyncCompute, allows using it in DBuffer pass, no GBuffer blending decals allow GBuffer compression)\n")
-	TEXT("  1: After Base Pass (default)"),
+	TEXT("  0: Before Base Pass (Allows samping in DBuffer pass. Can be more efficient with AsyncCompute.)\n")
+	TEXT("  1: After Base Pass\n")
+	TEXT("  2: Default (Before Base Pass if DBuffer enabled.)\n"),
 	ECVF_RenderThreadSafe);
 
 static TAutoConsoleVariable<int32> CVarMobileCustomDepthDownSample(
@@ -41,13 +42,11 @@ static TAutoConsoleVariable<int32> CVarCustomDepthTemporalAAJitter(
 
 DECLARE_GPU_STAT_NAMED(CustomDepth, TEXT("Custom Depth"));
 
-ECustomDepthPassLocation GetCustomDepthPassLocation()
+ECustomDepthPassLocation GetCustomDepthPassLocation(EShaderPlatform Platform)
 {
-	switch (CVarCustomDepthOrder.GetValueOnRenderThread())
-	{
-	case 0: return ECustomDepthPassLocation::BeforeBasePass;
-	}
-	return ECustomDepthPassLocation::AfterBasePass;
+	const int32 CustomDepthOrder = CVarCustomDepthOrder.GetValueOnRenderThread();
+	const bool bCustomDepthBeforeBasePase = CustomDepthOrder == 0 || (CustomDepthOrder == 2 && IsUsingDBuffers(Platform));
+	return bCustomDepthBeforeBasePase ? ECustomDepthPassLocation::BeforeBasePass : ECustomDepthPassLocation::AfterBasePass;
 }
 
 ECustomDepthMode GetCustomDepthMode()
diff --git a/Engine/Source/Runtime/Renderer/Private/CustomDepthRendering.h b/Engine/Source/Runtime/Renderer/Private/CustomDepthRendering.h
--- a/Engine/Source/Runtime/Renderer/Private/CustomDepthRendering.h
+++ b/Engine/Source/Runtime/Renderer/Private/CustomDepthRendering.h
@@ -14,7 +14,7 @@ enum class ECustomDepthPassLocation : uint32
 };
 
 // Returns the location in the frame where custom depth is rendered.
-extern RENDERER_API ECustomDepthPassLocation GetCustomDepthPassLocation();
+extern ECustomDepthPassLocation GetCustomDepthPassLocation(EShaderPlatform Platform);
 
 enum class ECustomDepthMode : uint32
 {
@@ -29,20 +29,17 @@ enum class ECustomDepthMode : uint32
 };
 
 // The custom depth mode currently configured.
-extern RENDERER_API ECustomDepthMode GetCustomDepthMode();
-
-// Whether custom depth passes should write stencil.
-extern RENDERER_API bool IsCustomDepthPassWritingStencil(ERHIFeatureLevel::Type InFeatureLevel);
+extern ECustomDepthMode GetCustomDepthMode();
 
 // Returns the requested downsample factor for custom depth textures.
-extern RENDERER_API uint32 GetCustomDepthDownsampleFactor(ERHIFeatureLevel::Type InFeatureLevel);
+extern uint32 GetCustomDepthDownsampleFactor(ERHIFeatureLevel::Type InFeatureLevel);
 
 inline bool IsCustomDepthPassEnabled()
 {
 	return GetCustomDepthMode() != ECustomDepthMode::Disabled;
 }
 
-struct RENDERER_API FCustomDepthTextures
+struct FCustomDepthTextures
 {
 	static FCustomDepthTextures Create(FRDGBuilder& GraphBuilder, FIntPoint Extent, ERHIFeatureLevel::Type FeatureLevel, uint32 DownsampleFactor);
 
@@ -62,4 +59,4 @@ struct RENDERER_API FCustomDepthTextures
 	ERenderTargetLoadAction StencilAction = ERenderTargetLoadAction::EClear;
 
 	uint32 DownsampleFactor = 1;
-};
\ No newline at end of file
+};
diff --git a/Engine/Source/Runtime/Renderer/Private/DeferredShadingRenderer.cpp b/Engine/Source/Runtime/Renderer/Private/DeferredShadingRenderer.cpp
--- a/Engine/Source/Runtime/Renderer/Private/DeferredShadingRenderer.cpp
+++ b/Engine/Source/Runtime/Renderer/Private/DeferredShadingRenderer.cpp
@@ -2404,7 +2404,8 @@ void FDeferredShadingSceneRenderer::Render(FRDGBuilder& GraphBuilder)
 	const bool bStrataEnabled = Strata::IsStrataEnabled();
 	Strata::InitialiseStrataFrameSceneData(*this, GraphBuilder);
 
-	if (GetCustomDepthPassLocation() == ECustomDepthPassLocation::BeforeBasePass)
+	const ECustomDepthPassLocation CustomDepthPassLocation = GetCustomDepthPassLocation(ShaderPlatform);
+	if (CustomDepthPassLocation == ECustomDepthPassLocation::BeforeBasePass)
 	{
 		QUICK_SCOPE_CYCLE_COUNTER(STAT_FDeferredShadingSceneRenderer_CustomDepthPass_BeforeBasePass);
 		if (RenderCustomDepthPass(GraphBuilder, SceneTextures.CustomDepth, SceneTextures.GetSceneTextureShaderParameters(FeatureLevel)))
@@ -2654,7 +2655,7 @@ void FDeferredShadingSceneRenderer::Render(FRDGBuilder& GraphBuilder)
 		bAsyncComputeVolumetricCloud = RenderVolumetricCloud(GraphBuilder, SceneTextures, bSkipVolumetricRenderTarget, bSkipPerPixelTracing, HalfResolutionDepthCheckerboardMinMaxTexture, true, InstanceCullingManager);
 	}
 
-	if (GetCustomDepthPassLocation() == ECustomDepthPassLocation::AfterBasePass)
+	if (CustomDepthPassLocation == ECustomDepthPassLocation::AfterBasePass)
 	{
 		QUICK_SCOPE_CYCLE_COUNTER(STAT_FDeferredShadingSceneRenderer_CustomDepthPass_AfterBasePass);
 		if (RenderCustomDepthPass(GraphBuilder, SceneTextures.CustomDepth, SceneTextures.GetSceneTextureShaderParameters(FeatureLevel)))
