Fix one texel border of non black values in mask texture.

diff --git a/Engine/Source/Runtime/Renderer/Private/RendererScene.cpp b/Engine/Source/Runtime/Renderer/Private/RendererScene.cpp
--- a/Engine/Source/Runtime/Renderer/Private/RendererScene.cpp
+++ b/Engine/Source/Runtime/Renderer/Private/RendererScene.cpp
@@ -707,14 +707,17 @@ static void UpdateEarlyZPassModeCVarSinkFunction()
 	static int32 CachedEarlyZPass = CVarEarlyZPass.GetValueOnGameThread();
 	static int32 CachedBasePassWriteDepthEvenWithFullPrepass = CVarBasePassWriteDepthEvenWithFullPrepass.GetValueOnGameThread();
 	static int32 CachedMobileEarlyZPass = CVarMobileEarlyZPass.GetValueOnGameThread();
+	static int32 CachedVelocityMergedWithDepthPass = IsVelocityMergedWithDepthPass();
 
 	const int32 EarlyZPass = CVarEarlyZPass.GetValueOnGameThread();
 	const int32 BasePassWriteDepthEvenWithFullPrepass = CVarBasePassWriteDepthEvenWithFullPrepass.GetValueOnGameThread();
 	const int32 MobileEarlyZPass = CVarMobileEarlyZPass.GetValueOnGameThread();
+	const int32 VelocityMergedWithDepthPass = IsVelocityMergedWithDepthPass();
 
 	if (EarlyZPass != CachedEarlyZPass
 		|| BasePassWriteDepthEvenWithFullPrepass != CachedBasePassWriteDepthEvenWithFullPrepass
-		|| MobileEarlyZPass != CachedMobileEarlyZPass)
+		|| MobileEarlyZPass != CachedMobileEarlyZPass
+		|| VelocityMergedWithDepthPass != CachedVelocityMergedWithDepthPass)
 	{
 		for (TObjectIterator<UWorld> It; It; ++It)
 		{
@@ -723,6 +726,7 @@ static void UpdateEarlyZPassModeCVarSinkFunction()
 			{
 				FScene* Scene = (FScene*)(World->Scene);
 				Scene->UpdateEarlyZPassMode();
+				Scene->UpdateStaticDrawLists();
 			}
 		}
 
