Converting FMeshPassProcessors to use ODSC.

- FEditorLevelInstanceMeshProcessor

diff --git a/Engine/Source/Runtime/Renderer/Private/SceneHitProxyRendering.cpp b/Engine/Source/Runtime/Renderer/Private/SceneHitProxyRendering.cpp
--- a/Engine/Source/Runtime/Renderer/Private/SceneHitProxyRendering.cpp
+++ b/Engine/Source/Runtime/Renderer/Private/SceneHitProxyRendering.cpp
@@ -1037,34 +1037,50 @@ void FEditorLevelInstanceMeshProcessor::AddMeshBatch(const FMeshBatch& RESTRICT
 		&& PrimitiveSceneProxy
 		&& PrimitiveSceneProxy->IsEditingLevelInstanceChild())
 	{
-		// Determine the mesh's material and blend mode.
-		const FMaterialRenderProxy* MaterialRenderProxy = nullptr;
-		const FMaterial* Material = &MeshBatch.MaterialRenderProxy->GetMaterialWithFallback(FeatureLevel, MaterialRenderProxy);
-
-		const FMeshDrawingPolicyOverrideSettings OverrideSettings = ComputeMeshOverrideSettings(MeshBatch);
-		const ERasterizerFillMode MeshFillMode = ComputeMeshFillMode(MeshBatch, *Material, OverrideSettings);
-		const ERasterizerCullMode MeshCullMode = CM_None;
-
-		if (Material->WritesEveryPixel() && !Material->IsTwoSided() && !Material->MaterialModifiesMeshPosition_RenderThread())
+		const FMaterialRenderProxy* MaterialRenderProxy = MeshBatch.MaterialRenderProxy;
+		while (MaterialRenderProxy)
 		{
-			// Default material doesn't handle masked, and doesn't have the correct bIsTwoSided setting.
-			MaterialRenderProxy = UMaterial::GetDefaultMaterial(MD_Surface)->GetRenderProxy();
-			check(MaterialRenderProxy);
-			Material = MaterialRenderProxy->GetMaterialNoFallback(FeatureLevel);
-		}
+			const FMaterial* Material = MaterialRenderProxy->GetMaterialNoFallback(FeatureLevel);
+			if (Material && Material->GetRenderingThreadShaderMap())
+			{
+				if (TryAddMeshBatch(MeshBatch, BatchElementMask, PrimitiveSceneProxy, StaticMeshId, MaterialRenderProxy, Material))
+				{
+					break;
+				}
+			}
 
-		if (!MaterialRenderProxy)
-		{
-			MaterialRenderProxy = MeshBatch.MaterialRenderProxy;
+			MaterialRenderProxy = MaterialRenderProxy->GetFallback(FeatureLevel);
 		}
+	}
+}
 
-		check(Material && MaterialRenderProxy);
+bool FEditorLevelInstanceMeshProcessor::TryAddMeshBatch(
+	const FMeshBatch& RESTRICT MeshBatch, 
+	uint64 BatchElementMask, 
+	const FPrimitiveSceneProxy* RESTRICT PrimitiveSceneProxy, 
+	int32 StaticMeshId, 
+	const FMaterialRenderProxy* MaterialRenderProxy, 
+	const FMaterial* Material)
+{
+	// Determine the mesh's material and blend mode.
+	const FMeshDrawingPolicyOverrideSettings OverrideSettings = ComputeMeshOverrideSettings(MeshBatch);
+	const ERasterizerFillMode MeshFillMode = ComputeMeshFillMode(MeshBatch, *Material, OverrideSettings);
+	const ERasterizerCullMode MeshCullMode = CM_None;
 
-		Process(MeshBatch, BatchElementMask, StaticMeshId, PrimitiveSceneProxy, *MaterialRenderProxy, *Material, MeshFillMode, MeshCullMode);
+	if (Material->WritesEveryPixel() && !Material->IsTwoSided() && !Material->MaterialModifiesMeshPosition_RenderThread())
+	{
+		// Default material doesn't handle masked, and doesn't have the correct bIsTwoSided setting.
+		MaterialRenderProxy = UMaterial::GetDefaultMaterial(MD_Surface)->GetRenderProxy();
+		check(MaterialRenderProxy);
+		Material = MaterialRenderProxy->GetMaterialNoFallback(FeatureLevel);
 	}
+
+	check(Material && MaterialRenderProxy);
+
+	return Process(MeshBatch, BatchElementMask, StaticMeshId, PrimitiveSceneProxy, *MaterialRenderProxy, *Material, MeshFillMode, MeshCullMode);
 }
 
-void FEditorLevelInstanceMeshProcessor::Process(
+bool FEditorLevelInstanceMeshProcessor::Process(
 	const FMeshBatch& MeshBatch,
 	uint64 BatchElementMask,
 	int32 StaticMeshId,
@@ -1080,13 +1096,16 @@ void FEditorLevelInstanceMeshProcessor::Process(
 		FHitProxyVS,
 		FHitProxyPS> HitProxyPassShaders;
 
-	GetHitProxyPassShaders(
+	if (!GetHitProxyPassShaders(
 		MaterialResource,
 		VertexFactory->GetType(),
 		FeatureLevel,
 		HitProxyPassShaders.VertexShader,
 		HitProxyPassShaders.PixelShader
-	);
+	))
+	{
+		return false;
+	}
 
 	const int32 StencilRef = GetStencilValue(ViewIfDynamicMeshCommand, PrimitiveSceneProxy);
 	PassDrawRenderState.SetStencilRef(StencilRef);
@@ -1110,6 +1129,8 @@ void FEditorLevelInstanceMeshProcessor::Process(
 		SortKey,
 		EMeshPassFeatures::Default,
 		ShaderElementData);
+
+	return true;
 }
 
 int32 FEditorLevelInstanceMeshProcessor::GetStencilValue(const FSceneView* View, const FPrimitiveSceneProxy* PrimitiveSceneProxy)
diff --git a/Engine/Source/Runtime/Renderer/Private/SceneHitProxyRendering.h b/Engine/Source/Runtime/Renderer/Private/SceneHitProxyRendering.h
--- a/Engine/Source/Runtime/Renderer/Private/SceneHitProxyRendering.h
+++ b/Engine/Source/Runtime/Renderer/Private/SceneHitProxyRendering.h
@@ -85,7 +85,10 @@ public:
 
 	FMeshPassProcessorRenderState PassDrawRenderState;
 
-	void Process(
+private:
+	bool TryAddMeshBatch(const FMeshBatch& RESTRICT MeshBatch, uint64 BatchElementMask, const FPrimitiveSceneProxy* RESTRICT PrimitiveSceneProxy, int32 StaticMeshId, const FMaterialRenderProxy* MaterialRenderProxy, const FMaterial* Material);
+
+	bool Process(
 		const FMeshBatch& MeshBatch,
 		uint64 BatchElementMask,
 		int32 StaticMeshId,
