Try to make it more obvious that allocated memory will be initialized, attempt to satisfy static analysis

diff --git a/Engine/Source/Runtime/Renderer/Private/VT/VirtualTextureSystem.cpp b/Engine/Source/Runtime/Renderer/Private/VT/VirtualTextureSystem.cpp
--- a/Engine/Source/Runtime/Renderer/Private/VT/VirtualTextureSystem.cpp
+++ b/Engine/Source/Runtime/Renderer/Private/VT/VirtualTextureSystem.cpp
@@ -1950,7 +1950,7 @@ void FVirtualTextureSystem::SubmitRequests(FRDGBuilder& GraphBuilder, ERHIFeatur
 	LLM_SCOPE(ELLMTag::VirtualTextureSystem);
 
 	// Allocate space to hold the physical address we allocate for each page load (1 page per layer per request)
-	uint32* RequestPhysicalAddress = new(MemStack, MEM_Oned) uint32[RequestList->GetNumLoadRequests() * VIRTUALTEXTURE_SPACE_MAXLAYERS];
+	uint32* RequestPhysicalAddress = NewOned<uint32>(MemStack, RequestList->GetNumLoadRequests() * VIRTUALTEXTURE_SPACE_MAXLAYERS);
 	{
 		SCOPE_CYCLE_COUNTER(STAT_ProcessRequests_Submit);
 
