Make sure we pick the no-lightmap shader permutation when using the path tracer

This is another small step towards enabling path tracing specific shader permutations.

This requires re-caching the raytracing mesh commands when we flip to and from path traced rendering.

Cleanup some flag handling by removing the redundant ERayTracingRenderMode enum since it is already capture by show flags.

Avoid checking the path tracing flag per view since it is always common to all views

Re-introduce the deduplication over views for the non-path tracing case since ray tracing effects can be enabled in a view dependent fashion.

Rename CanOverlayRayTracingOutput to HasRayTracedOverlay which I think is a bit clearer. Note that the meaning has flipped - it now returns true when path tracing or ray tracing debug modes are enabled.

Remove leftovers of GAveragePathTracedMRays (which wasn't hooked up to anything anymore)

Refactor static variables used to store path tracing state into a pimpl object off of the ViewState

diff --git a/Engine/Plugins/Experimental/GPULightmass/Source/GPULightmass/Private/LightmapRayTracing.h b/Engine/Plugins/Experimental/GPULightmass/Source/GPULightmass/Private/LightmapRayTracing.h
--- a/Engine/Plugins/Experimental/GPULightmass/Source/GPULightmass/Private/LightmapRayTracing.h
+++ b/Engine/Plugins/Experimental/GPULightmass/Source/GPULightmass/Private/LightmapRayTracing.h
@@ -18,7 +18,7 @@ class FLightmapRayTracingMeshProcessor : public FRayTracingMeshProcessor
 {
 public:
 	FLightmapRayTracingMeshProcessor(FRayTracingMeshCommandContext* InCommandContext, FMeshPassProcessorRenderState InPassDrawRenderState)
-		: FRayTracingMeshProcessor(InCommandContext, nullptr, nullptr, InPassDrawRenderState)
+		: FRayTracingMeshProcessor(InCommandContext, nullptr, nullptr, InPassDrawRenderState, ERayTracingMeshCommandsMode::PATH_TRACING) // NOTE: RayTracingMeshCommandsMode does not really matter here since Process is overridden
 	{}
 
 	virtual bool Process(
diff --git a/Engine/Plugins/MovieScene/MovieRenderPipeline/Source/MovieRenderPipelineRenderPasses/Private/MoviePipelineImagePassBase.cpp b/Engine/Plugins/MovieScene/MovieRenderPipeline/Source/MovieRenderPipelineRenderPasses/Private/MoviePipelineImagePassBase.cpp
--- a/Engine/Plugins/MovieScene/MovieRenderPipeline/Source/MovieRenderPipelineRenderPasses/Private/MoviePipelineImagePassBase.cpp
+++ b/Engine/Plugins/MovieScene/MovieRenderPipeline/Source/MovieRenderPipelineRenderPasses/Private/MoviePipelineImagePassBase.cpp
@@ -82,9 +82,6 @@ TSharedPtr<FSceneViewFamilyContext> UMoviePipelineImagePassBase::CalculateViewFa
 
 	// View is added as a child of the OutViewFamily-> 
 	FSceneView* View = GetSceneViewForSampleState(OutViewFamily.Get(), /*InOut*/ InOutSampleState);
-#if RHI_RAYTRACING
-	View->SetupRayTracedRendering();
-#endif
 	
 	SetupViewForViewModeOverride(View);
 
diff --git a/Engine/Plugins/Runtime/nDisplay/Source/DisplayCluster/Private/Game/EngineClasses/Basics/DisplayClusterViewportClient.cpp b/Engine/Plugins/Runtime/nDisplay/Source/DisplayCluster/Private/Game/EngineClasses/Basics/DisplayClusterViewportClient.cpp
--- a/Engine/Plugins/Runtime/nDisplay/Source/DisplayCluster/Private/Game/EngineClasses/Basics/DisplayClusterViewportClient.cpp
+++ b/Engine/Plugins/Runtime/nDisplay/Source/DisplayCluster/Private/Game/EngineClasses/Basics/DisplayClusterViewportClient.cpp
@@ -416,11 +416,6 @@ void UDisplayClusterViewportClient::Draw(FViewport* InViewport, FCanvas* SceneCa
 										}
 									}
 								}
-
-#if RHI_RAYTRACING
-								View->SetupRayTracedRendering();
-#endif
-
 							}
 
 							// Add view information for resource streaming. Allow up to 5X boost for small FOV.
diff --git a/Engine/Source/Editor/UnrealEd/Private/LevelEditorViewport.cpp b/Engine/Source/Editor/UnrealEd/Private/LevelEditorViewport.cpp
--- a/Engine/Source/Editor/UnrealEd/Private/LevelEditorViewport.cpp
+++ b/Engine/Source/Editor/UnrealEd/Private/LevelEditorViewport.cpp
@@ -1909,10 +1909,6 @@ FSceneView* FLevelEditorViewportClient::CalcSceneView(FSceneViewFamily* ViewFami
 	View->bCameraCut = bEditorCameraCut;
 	View->bHasSelectedComponents = GEditor->GetSelectedComponentCount() > 0;
 
-#if RHI_RAYTRACING
-	View->SetupRayTracedRendering();
-#endif
-
 	return View;
 
 }
diff --git a/Engine/Source/Runtime/Engine/Private/GameViewportClient.cpp b/Engine/Source/Runtime/Engine/Private/GameViewportClient.cpp
--- a/Engine/Source/Runtime/Engine/Private/GameViewportClient.cpp
+++ b/Engine/Source/Runtime/Engine/Private/GameViewportClient.cpp
@@ -1541,10 +1541,6 @@ void UGameViewportClient::Draw(FViewport* InViewport, FCanvas* SceneCanvas)
 								}
 							}
 						}
-
-					#if RHI_RAYTRACING
-						View->SetupRayTracedRendering();
-					#endif
 					}
 
 					// Add view information for resource streaming. Allow up to 5X boost for small FOV.
diff --git a/Engine/Source/Runtime/Engine/Private/SceneView.cpp b/Engine/Source/Runtime/Engine/Private/SceneView.cpp
--- a/Engine/Source/Runtime/Engine/Private/SceneView.cpp
+++ b/Engine/Source/Runtime/Engine/Private/SceneView.cpp
@@ -2696,29 +2696,6 @@ FSceneViewFamilyContext::~FSceneViewFamilyContext()
 	}
 }
 
-#if RHI_RAYTRACING
-void FSceneView::SetupRayTracedRendering()
-{
-	RayTracingRenderMode = ERayTracingRenderMode::Disabled;
-
-	if (!IsRayTracingEnabled())
-	{
-		return;
-	}
-
-	const FEngineShowFlags& ShowFlags = Family->EngineShowFlags;
-
-	if (ShowFlags.PathTracing && FDataDrivenShaderPlatformInfo::GetSupportsPathTracing(GetShaderPlatform()))
-	{
-		RayTracingRenderMode = ERayTracingRenderMode::PathTracing;
-	}	
-	else if (ShowFlags.RayTracingDebug)
-	{
-		RayTracingRenderMode = ERayTracingRenderMode::RayTracingDebug;
-	}
-}
-#endif // RHI_RAYTRACING
-
 #if !(UE_BUILD_SHIPPING || UE_BUILD_TEST)
 
 EDebugViewShaderMode FSceneViewFamily::ChooseDebugViewShaderMode() const
diff --git a/Engine/Source/Runtime/Engine/Private/UnrealEngine.cpp b/Engine/Source/Runtime/Engine/Private/UnrealEngine.cpp
--- a/Engine/Source/Runtime/Engine/Private/UnrealEngine.cpp
+++ b/Engine/Source/Runtime/Engine/Private/UnrealEngine.cpp
@@ -498,7 +498,6 @@ void FTemporaryPlayInEditorIDOverride::SetID(int32 NewID)
 // We expose these variables to everyone as we need to access them in other files via an extern
 ENGINE_API float GAverageFPS = 0.0f;
 ENGINE_API float GAverageMS = 0.0f;
-ENGINE_API float GAveragePathTracedMRays = 0.0f;
 ENGINE_API double GLastMemoryWarningTime = 0.f;
 
 static FCachedSystemScalabilityCVars GCachedScalabilityCVars;
@@ -15271,30 +15270,6 @@ int32 UEngine::RenderStatFPS(UWorld* World, FViewport* Viewport, FCanvas* Canvas
 	);
 	Y += RowHeight;
 
-#if RHI_RAYTRACING
-	if (GAveragePathTracedMRays > 0.0)
-	{
-		// Draw the MRays/frame counter.
-		Canvas->DrawShadowedString(
-			X,
-			Y,
-			*FString::Printf(TEXT("%5.2f MRays/fr"), GAveragePathTracedMRays),
-			Font,
-			FPSColor
-		);
-		Y += RowHeight;
-
-		// Draw the MRays/s counter.
-		Canvas->DrawShadowedString(
-			X,
-			Y,
-			*FString::Printf(TEXT("%5.2f MRays/s"), GAveragePathTracedMRays*GAverageFPS),
-			Font,
-			FPSColor
-		);
-		Y += RowHeight;
-	}
-#endif
 	return Y;
 }
 
diff --git a/Engine/Source/Runtime/Engine/Public/SceneView.h b/Engine/Source/Runtime/Engine/Public/SceneView.h
--- a/Engine/Source/Runtime/Engine/Public/SceneView.h
+++ b/Engine/Source/Runtime/Engine/Public/SceneView.h
@@ -29,13 +29,6 @@ class FVolumetricFogViewResources;
 class FIESLightProfileResource;
 class ITemporalUpscaler;
 
-enum class ERayTracingRenderMode
-{
-	Disabled			= 0,
-	PathTracing			= 1,
-	RayTracingDebug		= 2,
-};
-
 // Projection data for a FSceneView
 struct FSceneViewProjectionData
 {
@@ -1383,11 +1376,6 @@ public:
 		const FViewMatrices& InPrevViewMatrices) const;
 
 #if RHI_RAYTRACING
-	/** Setup ray tracing based rendering */
-	void SetupRayTracedRendering();
-
-	ERayTracingRenderMode RayTracingRenderMode = ERayTracingRenderMode::Disabled;
-
 	/** Current ray tracing debug visualization mode */
 	FName CurrentRayTracingDebugVisualizationMode;
 #endif
diff --git a/Engine/Source/Runtime/Renderer/Private/CompositionLighting/CompositionLighting.cpp b/Engine/Source/Runtime/Renderer/Private/CompositionLighting/CompositionLighting.cpp
--- a/Engine/Source/Runtime/Renderer/Private/CompositionLighting/CompositionLighting.cpp
+++ b/Engine/Source/Runtime/Renderer/Private/CompositionLighting/CompositionLighting.cpp
@@ -556,13 +556,12 @@ void ProcessAfterBasePass(
 	const FAsyncResults& AsyncResults,
 	bool bEnableSSAO)
 {
-	if (!CanOverlayRayTracingOutput(View))
+	const FSceneViewFamily& ViewFamily = *View.Family;
+	if (HasRayTracedOverlay(ViewFamily))
 	{
 		return;
 	}
 
-	const FSceneViewFamily& ViewFamily = *View.Family;
-
 	RDG_GPU_MASK_SCOPE(GraphBuilder, View.GPUMask);
 	RDG_EVENT_SCOPE(GraphBuilder, "LightCompositionTasks_PreLighting");
 	RDG_GPU_STAT_SCOPE(GraphBuilder, CompositionPreLighting);
@@ -640,17 +639,15 @@ void ProcessAfterBasePass(
 
 bool CanProcessAsync(TArrayView<const FViewInfo> Views)
 {
-	bool bAnyAsyncSSAO = true;
-	for (int32 i = 0; i < Views.Num(); ++i)
+	for (const FViewInfo& View : Views)
 	{
-		uint32 Levels = FSSAOHelper::ComputeAmbientOcclusionPassCount(Views[i]);
-		if (!FSSAOHelper::IsAmbientOcclusionAsyncCompute(Views[i], Levels) || !CanOverlayRayTracingOutput(Views[i]))
+		uint32 Levels = FSSAOHelper::ComputeAmbientOcclusionPassCount(View);
+		if (!FSSAOHelper::IsAmbientOcclusionAsyncCompute(View, Levels) || HasRayTracedOverlay(*View.Family))
 		{
-			bAnyAsyncSSAO = false;
-			break;
+			return false;
 		}
 	}
-	return bAnyAsyncSSAO;
+	return true;
 }
 
 FAsyncResults ProcessAsync(FRDGBuilder& GraphBuilder, TArrayView<const FViewInfo> Views, const FMinimalSceneTextures& SceneTextures)
diff --git a/Engine/Source/Runtime/Renderer/Private/DeferredShadingRenderer.cpp b/Engine/Source/Runtime/Renderer/Private/DeferredShadingRenderer.cpp
--- a/Engine/Source/Runtime/Renderer/Private/DeferredShadingRenderer.cpp
+++ b/Engine/Source/Runtime/Renderer/Private/DeferredShadingRenderer.cpp
@@ -885,7 +885,7 @@ bool FDeferredShadingSceneRenderer::GatherRayTracingWorldInstancesForView(FRHICo
 									*TaskDynamicCommandStorage, *TaskVisibleCommands,
 									SegmentIndex, WorkItem.InstanceIndex);
 								FMeshPassProcessorRenderState PassDrawRenderState(Scene->UniformBuffers.ViewUniformBuffer);
-								FRayTracingMeshProcessor RayTracingMeshProcessor(&CommandContext, Scene, &View, PassDrawRenderState);
+								FRayTracingMeshProcessor RayTracingMeshProcessor(&CommandContext, Scene, &View, PassDrawRenderState, Scene->CachedRayTracingMeshCommandsMode);
 								RayTracingMeshProcessor.AddMeshBatch(MeshBatch, 1, WorkItem.SceneProxy);
 							}
 						}
@@ -1056,7 +1056,7 @@ bool FDeferredShadingSceneRenderer::GatherRayTracingWorldInstancesForView(FRHICo
 							const FMeshBatch& MeshBatch = InstanceMaterials[SegmentIndex];
 							FDynamicRayTracingMeshCommandContext CommandContext(View.DynamicRayTracingMeshCommandStorage, View.VisibleRayTracingMeshCommands, SegmentIndex, InstanceIndex);
 							FMeshPassProcessorRenderState PassDrawRenderState(Scene->UniformBuffers.ViewUniformBuffer);
-							FRayTracingMeshProcessor RayTracingMeshProcessor(&CommandContext, Scene, &View, PassDrawRenderState);
+							FRayTracingMeshProcessor RayTracingMeshProcessor(&CommandContext, Scene, &View, PassDrawRenderState, Scene->CachedRayTracingMeshCommandsMode);
 							RayTracingMeshProcessor.AddMeshBatch(MeshBatch, 1, SceneProxy);
 						}
 					}
@@ -1281,9 +1281,9 @@ bool FDeferredShadingSceneRenderer::SetupRayTracingPipelineStates(FRHICommandLis
 	}
 
 	bool bAnyRayTracingPassEnabled = false;
-	for (int32 ViewIndex = 0; ViewIndex < Views.Num(); ViewIndex++)
+	for (const FViewInfo& View : Views)
 	{
-		bAnyRayTracingPassEnabled |= AnyRayTracingPassEnabled(Scene, Views[ViewIndex]);
+		bAnyRayTracingPassEnabled |= AnyRayTracingPassEnabled(Scene, View);
 	}
 
 	if (!bAnyRayTracingPassEnabled)
@@ -1312,30 +1312,36 @@ bool FDeferredShadingSceneRenderer::SetupRayTracingPipelineStates(FRHICommandLis
 
 	// #dxr_todo: UE-72565: refactor ray tracing effects to not be member functions of DeferredShadingRenderer. register each effect at startup and just loop over them automatically to gather all required shaders
 	TArray<FRHIRayTracingShader*> RayGenShaders; // TODO: inline allocator here?
-
-	if (ReferenceView.RayTracingRenderMode == ERayTracingRenderMode::PathTracing)
+	const bool bIsPathTracing = ViewFamily.EngineShowFlags.PathTracing;
+	if (bIsPathTracing)
 	{
-		// this view only needs the path tracing raygen shaders as all other
-		// passes should be disabled
-		PreparePathTracing(ReferenceView, RayGenShaders);
+		// This view only needs the path tracing raygen shaders as all other
+		// passes should be disabled.
+		PreparePathTracing(ViewFamily, RayGenShaders);
 	}
 	else
 	{
-		// path tracing is disabled, get all other possible raygen shaders
-		PrepareRayTracingReflections(ReferenceView, *Scene, RayGenShaders);
-		PrepareSingleLayerWaterRayTracingReflections(ReferenceView, *Scene, RayGenShaders);
-		PrepareRayTracingShadows(ReferenceView, RayGenShaders);
-		PrepareRayTracingAmbientOcclusion(ReferenceView, RayGenShaders);
-		PrepareRayTracingSkyLight(ReferenceView, RayGenShaders);
-		PrepareRayTracingGlobalIllumination(ReferenceView, RayGenShaders);
-		PrepareRayTracingTranslucency(ReferenceView, RayGenShaders);
-		PrepareRayTracingDebug(ReferenceView, RayGenShaders);
+		// Path tracing is disabled, get all other possible raygen shaders
+		PrepareRayTracingDebug(ViewFamily, RayGenShaders);
 
-		PrepareRayTracingLumenDirectLighting(ReferenceView, *Scene, RayGenShaders);
-		PrepareLumenHardwareRayTracingScreenProbeGather(ReferenceView, RayGenShaders);
-		PrepareLumenHardwareRayTracingRadianceCache(ReferenceView, RayGenShaders);
-		PrepareLumenHardwareRayTracingReflections(ReferenceView, RayGenShaders);
-		PrepareLumenHardwareRayTracingVisualize(ReferenceView, RayGenShaders);
+		// These other cases do potentially depend on the camera position since they are
+		// driven by FinalPostProcessSettings, which is why we need to merge them across views
+		for (const FViewInfo& View : Views)
+		{
+			PrepareRayTracingReflections(View, *Scene, RayGenShaders);
+			PrepareSingleLayerWaterRayTracingReflections(View, *Scene, RayGenShaders);
+			PrepareRayTracingShadows(View, RayGenShaders);
+			PrepareRayTracingAmbientOcclusion(View, RayGenShaders);
+			PrepareRayTracingSkyLight(View, RayGenShaders);
+			PrepareRayTracingGlobalIllumination(View, RayGenShaders);
+			PrepareRayTracingTranslucency(View, RayGenShaders);
+			PrepareRayTracingLumenDirectLighting(View, *Scene, RayGenShaders);
+			PrepareLumenHardwareRayTracingScreenProbeGather(View, RayGenShaders);
+			PrepareLumenHardwareRayTracingRadianceCache(View, RayGenShaders);
+			PrepareLumenHardwareRayTracingReflections(View, RayGenShaders);
+			PrepareLumenHardwareRayTracingVisualize(View, RayGenShaders);
+		}
+		DeduplicateRayGenerationShaders(RayGenShaders);
 	}
 
 	if (RayGenShaders.Num())
@@ -1357,15 +1363,23 @@ bool FDeferredShadingSceneRenderer::SetupRayTracingPipelineStates(FRHICommandLis
 	{
 		FViewInfo& View = Views[ViewIndex];
 
-		View.RayTracingLightData = CreateRayTracingLightData(RHICmdList,
-			Scene->Lights, View, EUniformBufferUsage::UniformBuffer_SingleFrame);
+		if (bIsPathTracing)
+		{
+			// Path Tracing currently uses its own code to manage lights, so doesn't need to run this.
+			// TODO: merge the lighting representations between ray traced and path traced cases?
+		}
+		else
+		{
+			// This light data is a function of the camera position, so must be computed per view.
+			View.RayTracingLightData = CreateRayTracingLightData(RHICmdList,
+				Scene->Lights, View, EUniformBufferUsage::UniformBuffer_SingleFrame);
+		}
 
 		// Send common ray tracing resources from reference view to all others.
 		if (ViewIndex != ReferenceViewIndex)
 		{
 			View.RayTracingSubSurfaceProfileTexture = ReferenceView.RayTracingSubSurfaceProfileTexture;
 			View.RayTracingSubSurfaceProfileSRV = ReferenceView.RayTracingSubSurfaceProfileSRV;
-			View.RayTracingLightData = ReferenceView.RayTracingLightData;
 			View.RayTracingMaterialPipeline = ReferenceView.RayTracingMaterialPipeline;
 		}
 	}
@@ -1557,7 +1571,7 @@ void FDeferredShadingSceneRenderer::WaitForRayTracingScene(FRDGBuilder& GraphBui
 		const int32 ReferenceViewIndex = 0;
 		FViewInfo& ReferenceView = Views[ReferenceViewIndex];
 
-		const bool bIsPathTracing = ReferenceView.RayTracingRenderMode == ERayTracingRenderMode::PathTracing;
+		const bool bIsPathTracing = ViewFamily.EngineShowFlags.PathTracing;
 
 		check(ReferenceView.RayTracingMaterialPipeline || ReferenceView.RayTracingMaterialBindings.Num() == 0);
 
@@ -1612,9 +1626,8 @@ void FDeferredShadingSceneRenderer::WaitForRayTracingScene(FRDGBuilder& GraphBui
 			if (!bIsPathTracing)
 			{
 				TArray<FRHIRayTracingShader*> DeferredMaterialRayGenShaders;
-				for (int32 ViewIndex = 0; ViewIndex < Views.Num(); ++ViewIndex)
+				for (const FViewInfo& View : Views)
 				{
-					FViewInfo& View = Views[ViewIndex];
 					PrepareRayTracingReflectionsDeferredMaterial(View, *Scene, DeferredMaterialRayGenShaders);
 					PrepareRayTracingDeferredReflectionsDeferredMaterial(View, *Scene, DeferredMaterialRayGenShaders);
 					PrepareRayTracingGlobalIlluminationDeferredMaterial(View, DeferredMaterialRayGenShaders);
@@ -1632,9 +1645,8 @@ void FDeferredShadingSceneRenderer::WaitForRayTracingScene(FRDGBuilder& GraphBui
 
 				// Add Lumen hardware ray tracing materials
 				TArray<FRHIRayTracingShader*> LumenHardwareRayTracingRayGenShaders;
-				for (int32 ViewIndex = 0; ViewIndex < Views.Num(); ++ViewIndex)
+				for (const FViewInfo& View : Views)
 				{
-					FViewInfo& View = Views[ViewIndex];
 					PrepareLumenHardwareRayTracingVisualizeLumenMaterial(View, LumenHardwareRayTracingRayGenShaders);
 					PrepareLumenHardwareRayTracingRadianceCacheLumenMaterial(View, LumenHardwareRayTracingRayGenShaders);
 					PrepareLumenHardwareRayTracingReflectionsLumenMaterial(View, LumenHardwareRayTracingRayGenShaders);
@@ -1760,6 +1772,19 @@ void FDeferredShadingSceneRenderer::Render(FRDGBuilder& GraphBuilder)
 {
 	const bool bNaniteEnabled = UseNanite(ShaderPlatform) && ViewFamily.EngineShowFlags.NaniteMeshes;
 
+#if RHI_RAYTRACING
+	{
+		ERayTracingMeshCommandsMode CurrentMode = ViewFamily.EngineShowFlags.PathTracing ? ERayTracingMeshCommandsMode::PATH_TRACING : ERayTracingMeshCommandsMode::RAY_TRACING;
+		if (CurrentMode != Scene->CachedRayTracingMeshCommandsMode)
+		{
+			// If we change to or from a path traced render, we need to refresh the cached ray tracing mesh commands
+			// because they contain data about the currently bound shader. This operation is a bit expensive but
+			// only happens once as we transition between modes which should be rare.
+			Scene->RefreshRayTracingMeshCommandCache(GraphBuilder.RHICmdList, CurrentMode);
+		}
+	}
+#endif
+
 	Scene->UpdateAllPrimitiveSceneInfos(GraphBuilder, true);
 
 	FGPUSceneScopeBeginEndHelper GPUSceneScopeBeginEndHelper(Scene->GPUScene, GPUSceneDynamicContext, Scene);
@@ -1930,12 +1955,6 @@ void FDeferredShadingSceneRenderer::Render(FRDGBuilder& GraphBuilder)
 	// Prepare the scene for rendering this frame.
 	GatherRayTracingWorldInstancesForView(GraphBuilder.RHICmdList, ReferenceView, RayTracingScene);
 
-	if (ReferenceView.RayTracingRenderMode != ERayTracingRenderMode::PathTracing)
-	{
-		extern ENGINE_API float GAveragePathTracedMRays;
-		GAveragePathTracedMRays = 0.0f;
-	}
-
 #endif // RHI_RAYTRACING
 
 	{
@@ -2006,13 +2025,13 @@ void FDeferredShadingSceneRenderer::Render(FRDGBuilder& GraphBuilder)
 	}
 
 	const bool bUseGBuffer = IsUsingGBuffers(ShaderPlatform);
-	bool bCanOverlayRayTracingOutput = CanOverlayRayTracingOutput(Views[0]);// #dxr_todo: UE-72557 multi-view case
+	const bool bHasRayTracedOverlay = HasRayTracedOverlay(ViewFamily);
 	
 	const bool bRenderDeferredLighting = ViewFamily.EngineShowFlags.Lighting
 		&& FeatureLevel >= ERHIFeatureLevel::SM5
 		&& ViewFamily.EngineShowFlags.DeferredLighting
 		&& bUseGBuffer
-		&& bCanOverlayRayTracingOutput;
+		&& !bHasRayTracedOverlay;
 
 	bool bComputeLightGrid = false;
 	bool bAnyLumenEnabled = false;
@@ -2360,7 +2379,7 @@ void FDeferredShadingSceneRenderer::Render(FRDGBuilder& GraphBuilder)
 	// End early occlusion queries
 
 	// Early Shadow depth rendering
-	if (bCanOverlayRayTracingOutput && bOcclusionBeforeBasePass)
+	if (!bHasRayTracedOverlay && bOcclusionBeforeBasePass)
 	{
 		const bool bAfterBasePass = false;
 		AllocateVirtualShadowMaps(bAfterBasePass);
@@ -2375,7 +2394,7 @@ void FDeferredShadingSceneRenderer::Render(FRDGBuilder& GraphBuilder)
 	const bool bShouldVisualizeVolumetricCloud = bShouldRenderVolumetricCloudBase && !!ViewFamily.EngineShowFlags.VisualizeVolumetricCloudConservativeDensity;
 	bool bAsyncComputeVolumetricCloud = IsVolumetricRenderTargetEnabled() && IsVolumetricRenderTargetAsyncCompute();
 	bool bHasHalfResCheckerboardMinMaxDepth = false;
-	bool bVolumetricRenderTargetRequired = bShouldRenderVolumetricCloud && bCanOverlayRayTracingOutput;
+	bool bVolumetricRenderTargetRequired = bShouldRenderVolumetricCloud && !bHasRayTracedOverlay;
 
 	if (bShouldRenderVolumetricCloudBase)
 	{
@@ -2443,7 +2462,7 @@ void FDeferredShadingSceneRenderer::Render(FRDGBuilder& GraphBuilder)
 	FRDGTextureRef HalfResolutionDepthCheckerboardMinMaxTexture = nullptr;
 
 	// Kick off async compute cloud eraly if all depth has been written in the prepass
-	if (bShouldRenderVolumetricCloud && bAsyncComputeVolumetricCloud && DepthPass.EarlyZPassMode == DDM_AllOpaque && bCanOverlayRayTracingOutput)
+	if (bShouldRenderVolumetricCloud && bAsyncComputeVolumetricCloud && DepthPass.EarlyZPassMode == DDM_AllOpaque && !bHasRayTracedOverlay)
 	{
 		HalfResolutionDepthCheckerboardMinMaxTexture = CreateHalfResolutionDepthCheckerboardMinMax(GraphBuilder, Views, SceneTextures.Depth.Resolve);
 		bHasHalfResCheckerboardMinMaxDepth = true;
@@ -2600,7 +2619,7 @@ void FDeferredShadingSceneRenderer::Render(FRDGBuilder& GraphBuilder)
 	}
 
 	// Shadow and fog after base pass
-	if (bCanOverlayRayTracingOutput && !bOcclusionBeforeBasePass)
+	if (!bHasRayTracedOverlay && !bOcclusionBeforeBasePass)
 	{
 		const bool bAfterBasePass = true;
 		AllocateVirtualShadowMaps(bAfterBasePass);
@@ -2645,7 +2664,7 @@ void FDeferredShadingSceneRenderer::Render(FRDGBuilder& GraphBuilder)
 	}
 
 	// If not all depth is written during the prepass, kick off async compute cloud after basepass
-	if (bShouldRenderVolumetricCloud && bAsyncComputeVolumetricCloud && DepthPass.EarlyZPassMode != DDM_AllOpaque && bCanOverlayRayTracingOutput)
+	if (bShouldRenderVolumetricCloud && bAsyncComputeVolumetricCloud && DepthPass.EarlyZPassMode != DDM_AllOpaque && !bHasRayTracedOverlay)
 	{
 		HalfResolutionDepthCheckerboardMinMaxTexture = CreateHalfResolutionDepthCheckerboardMinMax(GraphBuilder, Views, SceneTextures.Depth.Resolve);
 		bHasHalfResCheckerboardMinMaxDepth = true;
@@ -2723,7 +2742,7 @@ void FDeferredShadingSceneRenderer::Render(FRDGBuilder& GraphBuilder)
 #if RHI_RAYTRACING
 	// If Lumen is not using HWRT shadows, we can wait until here: before Lumen diffuse indirect
 	// Also catch the case of path tracer or RT debug output
-	if (!Lumen::UseHardwareRayTracedShadows(Views[0]) || !bCanOverlayRayTracingOutput)
+	if (!Lumen::UseHardwareRayTracedShadows(Views[0]) || bHasRayTracedOverlay)
 	{
 		WaitForRayTracingScene(GraphBuilder);
 	}
@@ -2794,12 +2813,12 @@ void FDeferredShadingSceneRenderer::Render(FRDGBuilder& GraphBuilder)
 		RenderDeferredReflectionsAndSkyLightingHair(GraphBuilder);
 	}
 
-	if (bShouldRenderVolumetricCloud && IsVolumetricRenderTargetEnabled() && !bHasHalfResCheckerboardMinMaxDepth && bCanOverlayRayTracingOutput)
+	if (bShouldRenderVolumetricCloud && IsVolumetricRenderTargetEnabled() && !bHasHalfResCheckerboardMinMaxDepth && !bHasRayTracedOverlay)
 	{
 		HalfResolutionDepthCheckerboardMinMaxTexture = CreateHalfResolutionDepthCheckerboardMinMax(GraphBuilder, Views, SceneTextures.Depth.Resolve);
 	}
 
-	if (bShouldRenderVolumetricCloud && bCanOverlayRayTracingOutput)
+	if (bShouldRenderVolumetricCloud && !bHasRayTracedOverlay)
 	{
 		if (!bAsyncComputeVolumetricCloud)
 		{
@@ -2812,12 +2831,12 @@ void FDeferredShadingSceneRenderer::Render(FRDGBuilder& GraphBuilder)
 		ReconstructVolumetricRenderTarget(GraphBuilder, Views, SceneTextures.Depth.Resolve, HalfResolutionDepthCheckerboardMinMaxTexture, bAsyncComputeVolumetricCloud);
 	}
 
-	const bool bShouldRenderTranslucency = bCanOverlayRayTracingOutput && ShouldRenderTranslucency();
+	const bool bShouldRenderTranslucency = !bHasRayTracedOverlay && ShouldRenderTranslucency();
 
 	// Union of all translucency view render flags.
 	ETranslucencyView TranslucencyViewsToRender = bShouldRenderTranslucency ? GetTranslucencyViews(Views) : ETranslucencyView::None;
 
-	const bool bShouldRenderSingleLayerWater = bCanOverlayRayTracingOutput && ShouldRenderSingleLayerWater(Views);
+	const bool bShouldRenderSingleLayerWater = !bHasRayTracedOverlay && ShouldRenderSingleLayerWater(Views);
 	FSceneWithoutWaterTextures SceneWithoutWaterTextures;
 	if (bShouldRenderSingleLayerWater)
 	{
@@ -2840,28 +2859,28 @@ void FDeferredShadingSceneRenderer::Render(FRDGBuilder& GraphBuilder)
 	FRDGTextureRef LightShaftOcclusionTexture = nullptr;
 
 	// Draw Lightshafts
-	if (bCanOverlayRayTracingOutput && ViewFamily.EngineShowFlags.LightShafts)
+	if (!bHasRayTracedOverlay && ViewFamily.EngineShowFlags.LightShafts)
 	{
 		SCOPE_CYCLE_COUNTER(STAT_FDeferredShadingSceneRenderer_RenderLightShaftOcclusion);
 		LightShaftOcclusionTexture = RenderLightShaftOcclusion(GraphBuilder, SceneTextures);
 	}
 
 	// Draw atmosphere
-	if (bCanOverlayRayTracingOutput && ShouldRenderAtmosphere(ViewFamily))
+	if (!bHasRayTracedOverlay && ShouldRenderAtmosphere(ViewFamily))
 	{
 		SCOPE_CYCLE_COUNTER(STAT_FDeferredShadingSceneRenderer_RenderAtmosphere);
 		RenderAtmosphere(GraphBuilder, SceneTextures, LightShaftOcclusionTexture);
 	}
 
 	// Draw the sky atmosphere
-	if (bCanOverlayRayTracingOutput && bShouldRenderSkyAtmosphere)
+	if (!bHasRayTracedOverlay && bShouldRenderSkyAtmosphere)
 	{
 		SCOPE_CYCLE_COUNTER(STAT_FDeferredShadingSceneRenderer_RenderSkyAtmosphere);
 		RenderSkyAtmosphere(GraphBuilder, SceneTextures);
 	}
 
 	// Draw fog.
-	if (bCanOverlayRayTracingOutput && ShouldRenderFog(ViewFamily))
+	if (!bHasRayTracedOverlay && ShouldRenderFog(ViewFamily))
 	{
 		RDG_CSV_STAT_EXCLUSIVE_SCOPE(GraphBuilder, RenderFog);
 		SCOPE_CYCLE_COUNTER(STAT_FDeferredShadingSceneRenderer_RenderFog);
@@ -2869,7 +2888,7 @@ void FDeferredShadingSceneRenderer::Render(FRDGBuilder& GraphBuilder)
 	}
 
 	// After the height fog, Draw volumetric clouds (having fog applied on them already) when using per pixel tracing,
-	if (bCanOverlayRayTracingOutput && bShouldRenderVolumetricCloud)
+	if (!bHasRayTracedOverlay && bShouldRenderVolumetricCloud)
 	{
 		bool bSkipVolumetricRenderTarget = true;
 		bool bSkipPerPixelTracing = false;
@@ -2887,7 +2906,7 @@ void FDeferredShadingSceneRenderer::Render(FRDGBuilder& GraphBuilder)
 
 	RenderOpaqueFX(GraphBuilder, Views, FXSystem, SceneTextures.UniformBuffer);
 
-	if (bCanOverlayRayTracingOutput && bShouldRenderSkyAtmosphere)
+	if (!bHasRayTracedOverlay && bShouldRenderSkyAtmosphere)
 	{
 		// Debug the sky atmosphere. Critically rendered before translucency to avoid emissive leaking over visualization by writing depth. 
 		// Alternative: render in post process chain as VisualizeHDR.
@@ -2903,7 +2922,7 @@ void FDeferredShadingSceneRenderer::Render(FRDGBuilder& GraphBuilder)
 	FSeparateTranslucencyTextures SeparateTranslucencyTextures(SeparateTranslucencyDimensions);
 
 	// Draw translucency.
-	if (bCanOverlayRayTracingOutput && TranslucencyViewsToRender != ETranslucencyView::None)
+	if (!bHasRayTracedOverlay && TranslucencyViewsToRender != ETranslucencyView::None)
 	{
 		RDG_CSV_STAT_EXCLUSIVE_SCOPE(GraphBuilder, RenderTranslucency);
 		SCOPE_CYCLE_COUNTER(STAT_TranslucencyDrawTime);
@@ -2982,7 +3001,7 @@ void FDeferredShadingSceneRenderer::Render(FRDGBuilder& GraphBuilder)
 		Strata::AddStrataDebugPasses(GraphBuilder, Views, SceneTextures.Color.Target, Scene->GetShaderPlatform());
 	}
 
-	if (bCanOverlayRayTracingOutput && ViewFamily.EngineShowFlags.LightShafts)
+	if (!bHasRayTracedOverlay && ViewFamily.EngineShowFlags.LightShafts)
 	{
 		SCOPE_CYCLE_COUNTER(STAT_FDeferredShadingSceneRenderer_RenderLightShaftBloom);
 		GraphBuilder.SetCommandListStat(GET_STATID(STAT_CLM_LightShaftBloom));
@@ -2998,16 +3017,18 @@ void FDeferredShadingSceneRenderer::Render(FRDGBuilder& GraphBuilder)
 #if RHI_RAYTRACING
 	if (IsRayTracingEnabled())
 	{
-		for (int32 ViewIndex = 0; ViewIndex < Views.Num(); ViewIndex++)
+		if (ViewFamily.EngineShowFlags.PathTracing && FDataDrivenShaderPlatformInfo::GetSupportsPathTracing(Scene->GetShaderPlatform()))
 		{
-			if (Views[ViewIndex].RayTracingRenderMode == ERayTracingRenderMode::PathTracing
-				&& FDataDrivenShaderPlatformInfo::GetSupportsPathTracing(Views[ViewIndex].GetShaderPlatform()))
+			for (const FViewInfo& View : Views)
 			{
-				RenderPathTracing(GraphBuilder, Views[ViewIndex], SceneTextures.UniformBuffer, SceneTextures.Color.Target);
+				RenderPathTracing(GraphBuilder, View, SceneTextures.UniformBuffer, SceneTextures.Color.Target);
 			}
-			else if (Views[ViewIndex].RayTracingRenderMode == ERayTracingRenderMode::RayTracingDebug)
+		}
+		else if (ViewFamily.EngineShowFlags.RayTracingDebug)
+		{
+			for (const FViewInfo& View : Views)
 			{
-				RenderRayTracingDebug(GraphBuilder, Views[ViewIndex], SceneTextures.Color.Target);
+				RenderRayTracingDebug(GraphBuilder, View, SceneTextures.Color.Target);
 			}
 		}
 	}
@@ -3040,7 +3061,7 @@ void FDeferredShadingSceneRenderer::Render(FRDGBuilder& GraphBuilder)
 		RenderStationaryLightOverlap(GraphBuilder, SceneTextures, LightingChannelsTexture);
 	}
 
-	if (bShouldVisualizeVolumetricCloud && bCanOverlayRayTracingOutput)
+	if (bShouldVisualizeVolumetricCloud && !bHasRayTracedOverlay)
 	{
 		RenderVolumetricCloud(GraphBuilder, SceneTextures, false, true, HalfResolutionDepthCheckerboardMinMaxTexture, false, InstanceCullingManager);
 		ReconstructVolumetricRenderTarget(GraphBuilder, Views, SceneTextures.Depth.Resolve, HalfResolutionDepthCheckerboardMinMaxTexture, false);
@@ -3165,23 +3186,14 @@ void FDeferredShadingSceneRenderer::Render(FRDGBuilder& GraphBuilder)
 
 bool AnyRayTracingPassEnabled(const FScene* Scene, const FViewInfo& View)
 {
-	if (ShouldRenderRayTracingAmbientOcclusion(View)
+	return ShouldRenderRayTracingAmbientOcclusion(View)
 		|| ShouldRenderRayTracingReflections(View)
 		|| ShouldRenderRayTracingGlobalIllumination(View)
 		|| ShouldRenderRayTracingTranslucency(View)
 		|| ShouldRenderRayTracingSkyLight(Scene ? Scene->SkyLight : nullptr)
 		|| ShouldRenderRayTracingShadows()
         || Lumen::AnyLumenHardwareRayTracingPassEnabled(Scene, View)
-		|| View.RayTracingRenderMode == ERayTracingRenderMode::PathTracing
-		|| View.RayTracingRenderMode == ERayTracingRenderMode::RayTracingDebug
-		)
-	{
-		return true;
-	}
-	else
-	{
-		return false;
-	}	
+		|| HasRayTracedOverlay(*View.Family);
 }
 
 bool ShouldRenderRayTracingEffect(bool bEffectEnabled)
@@ -3191,8 +3203,7 @@ bool ShouldRenderRayTracingEffect(bool bEffectEnabled)
 		return false;
 	}
 
-	static auto CVar = IConsoleManager::Get().FindConsoleVariable(TEXT("r.RayTracing.ForceAllRayTracingEffects"));
-	const int32 OverrideMode = CVar != nullptr ? CVar->GetInt() : -1;
+	const int32 OverrideMode = CVarForceAllRayTracingEffects.GetValueOnRenderThread();
 
 	if (OverrideMode >= 0)
 	{
@@ -3204,11 +3215,12 @@ bool ShouldRenderRayTracingEffect(bool bEffectEnabled)
 	}
 }
 
-bool CanOverlayRayTracingOutput(const FViewInfo& View)
+bool HasRayTracedOverlay(const FSceneViewFamily& ViewFamily)
 {
-	// Return false if a full screen ray tracing pass will be displayed on top of the raster pass
+	// Return true if a full screen ray tracing pass will be displayed on top of the raster pass
 	// This can be used to skip certain calculations
-	return (View.RayTracingRenderMode != ERayTracingRenderMode::PathTracing &&
-			View.RayTracingRenderMode != ERayTracingRenderMode::RayTracingDebug);
+	return
+		ViewFamily.EngineShowFlags.PathTracing ||
+		ViewFamily.EngineShowFlags.RayTracingDebug;
 }
 #endif // RHI_RAYTRACING
diff --git a/Engine/Source/Runtime/Renderer/Private/DeferredShadingRenderer.h b/Engine/Source/Runtime/Renderer/Private/DeferredShadingRenderer.h
--- a/Engine/Source/Runtime/Renderer/Private/DeferredShadingRenderer.h
+++ b/Engine/Source/Runtime/Renderer/Private/DeferredShadingRenderer.h
@@ -969,8 +969,8 @@ private:
 	static void PrepareRayTracingSkyLight(const FViewInfo& View, TArray<FRHIRayTracingShader*>& OutRayGenShaders);
 	static void PrepareRayTracingGlobalIllumination(const FViewInfo& View, TArray<FRHIRayTracingShader*>& OutRayGenShaders);
 	static void PrepareRayTracingTranslucency(const FViewInfo& View, TArray<FRHIRayTracingShader*>& OutRayGenShaders);
-	static void PrepareRayTracingDebug(const FViewInfo& View, TArray<FRHIRayTracingShader*>& OutRayGenShaders);
-	static void PreparePathTracing(const FViewInfo& View, TArray<FRHIRayTracingShader*>& OutRayGenShaders);
+	static void PrepareRayTracingDebug(const FSceneViewFamily& ViewFamily, TArray<FRHIRayTracingShader*>& OutRayGenShaders);
+	static void PreparePathTracing(const FSceneViewFamily& ViewFamily, TArray<FRHIRayTracingShader*>& OutRayGenShaders);
 	static void PrepareRayTracingLumenDirectLighting(const FViewInfo& View,const FScene& Scene, TArray<FRHIRayTracingShader*>& OutRayGenShaders);
 	static void PrepareLumenHardwareRayTracingScreenProbeGather(const FViewInfo& View, TArray<FRHIRayTracingShader*>& OutRayGenShaders);
 	static void PrepareLumenHardwareRayTracingScreenProbeGatherDeferredMaterial(const FViewInfo& View, TArray<FRHIRayTracingShader*>& OutRayGenShaders);
diff --git a/Engine/Source/Runtime/Renderer/Private/PathTracing.cpp b/Engine/Source/Runtime/Renderer/Private/PathTracing.cpp
--- a/Engine/Source/Runtime/Renderer/Private/PathTracing.cpp
+++ b/Engine/Source/Runtime/Renderer/Private/PathTracing.cpp
@@ -178,9 +178,46 @@ TAutoConsoleVariable<int32> CVarPathTracingLightGridVisualize(
 	ECVF_RenderThreadSafe
 );
 
+// Store the rendering options used on the previous frame so we can correctly invalidate when things change
+struct FPathTracingConfig
+{
+	FPathTracingData PathTracingData;
+	FIntRect ViewRect;
+	int LightShowFlags;
+	int LightGridResolution;
+	int LightGridMaxCount;
+	bool VisibleLights;
+	bool UseMISCompensation;
+	bool LockedSamplingPattern;
+
+	bool IsDifferent(const FPathTracingConfig& Other) const
+	{
+		// If any of these parameters if different, we will need to restart path tracing accuulation
+		return
+			PathTracingData.MaxSamples != Other.PathTracingData.MaxSamples ||
+			PathTracingData.MaxBounces != Other.PathTracingData.MaxBounces ||
+			PathTracingData.MaxSSSBounces != Other.PathTracingData.MaxSSSBounces ||
+			PathTracingData.MISMode != Other.PathTracingData.MISMode ||
+			PathTracingData.SamplerType != Other.PathTracingData.SamplerType ||
+			PathTracingData.ApproximateCaustics != Other.PathTracingData.ApproximateCaustics ||
+			PathTracingData.EnableCameraBackfaceCulling != Other.PathTracingData.EnableCameraBackfaceCulling ||
+			PathTracingData.EnableDirectLighting != Other.PathTracingData.EnableDirectLighting ||
+			PathTracingData.EnableEmissive != Other.PathTracingData.EnableEmissive ||
+			PathTracingData.VisualizeLightGrid != Other.PathTracingData.VisualizeLightGrid ||
+			PathTracingData.MaxPathIntensity != Other.PathTracingData.MaxPathIntensity ||
+			PathTracingData.FilterWidth != Other.PathTracingData.FilterWidth ||
+			ViewRect != Other.ViewRect ||
+			LightShowFlags != Other.LightShowFlags ||
+			LightGridResolution != Other.LightGridResolution ||
+			LightGridMaxCount != Other.LightGridMaxCount ||
+			VisibleLights != Other.VisibleLights ||
+			UseMISCompensation != Other.UseMISCompensation ||
+			LockedSamplingPattern != Other.LockedSamplingPattern;
+	}
+};
 
 // This function prepares the portion of shader arguments that may involve invalidating the path traced state
-static bool PrepareShaderArgs(const FViewInfo& View, FPathTracingData& PathTracingData) {
+static void PrepareShaderArgs(const FViewInfo& View, FPathTracingData& PathTracingData) {
 	PathTracingData.EnableDirectLighting = true;
 	int32 MaxBounces = CVarPathTracingMaxBounces.GetValueOnRenderThread();
 	if (MaxBounces < 0)
@@ -213,7 +250,6 @@ static bool PrepareShaderArgs(const FViewInfo& View, FPathTracingData& PathTraci
 	PathTracingData.MaxSSSBounces = CVarPathTracingMaxSSSBounces.GetValueOnRenderThread();
 	PathTracingData.MaxNormalBias = GetRaytracingMaxNormalBias();
 	PathTracingData.MISMode = CVarPathTracingMISMode.GetValueOnRenderThread();
-	uint32 VisibleLights = CVarPathTracingVisibleLights.GetValueOnRenderThread();
 	PathTracingData.MaxPathIntensity = CVarPathTracingMaxPathIntensity.GetValueOnRenderThread();
 	if (PathTracingData.MaxPathIntensity <= 0)
 	{
@@ -232,109 +268,6 @@ static bool PrepareShaderArgs(const FViewInfo& View, FPathTracingData& PathTraci
 		FilterWidth = View.FinalPostProcessSettings.PathTracingFilterWidth;
 	}
 	PathTracingData.FilterWidth = FilterWidth;
-
-	bool NeedInvalidation = false;
-
-	// If any of the parameters above changed since last time -- reset the accumulation
-	// TODO: find something cleaner than just using static variables here. Should all
-	// the state used for comparison go into ViewState?
-	static uint32 PrevMaxBounces = PathTracingData.MaxBounces;
-	if (PathTracingData.MaxBounces != PrevMaxBounces)
-	{
-		NeedInvalidation = true;
-		PrevMaxBounces = PathTracingData.MaxBounces;
-	}
-
-	// Changing the number of SSS bounces requires starting over
-	static uint32 PrevMaxSSSBounces = PathTracingData.MaxSSSBounces;
-	if (PathTracingData.MaxSSSBounces != PrevMaxSSSBounces)
-	{
-		NeedInvalidation = true;
-		PrevMaxSSSBounces = PathTracingData.MaxSSSBounces;
-	}
-
-	// Changing MIS mode requires starting over
-	static uint32 PreviousMISMode = PathTracingData.MISMode;
-	if (PreviousMISMode != PathTracingData.MISMode)
-	{
-		NeedInvalidation = true;
-		PreviousMISMode = PathTracingData.MISMode;
-	}
-
-	// Changing VisibleLights requires starting over
-	static uint32 PreviousVisibleLights = VisibleLights;
-	if (PreviousVisibleLights != VisibleLights)
-	{
-		NeedInvalidation = true;
-		PreviousVisibleLights = VisibleLights;
-	}
-
-	// Changing MaxPathIntensity requires starting over
-	static float PreviousMaxPathIntensity = PathTracingData.MaxPathIntensity;
-	if (PreviousMaxPathIntensity != PathTracingData.MaxPathIntensity)
-	{
-		NeedInvalidation = true;
-		PreviousMaxPathIntensity = PathTracingData.MaxPathIntensity;
-	}
-
-	// Changing approximate caustics requires starting over
-	static uint32 PreviousApproximateCaustics = PathTracingData.ApproximateCaustics;
-	if (PreviousApproximateCaustics != PathTracingData.ApproximateCaustics)
-	{
-		NeedInvalidation = true;
-		PreviousApproximateCaustics = PathTracingData.ApproximateCaustics;
-	}
-
-	// Changing filter width requires starting over
-	static float PreviousFilterWidth = PathTracingData.FilterWidth;
-	if (PreviousFilterWidth != PathTracingData.FilterWidth)
-	{
-		NeedInvalidation = true;
-		PreviousFilterWidth = PathTracingData.FilterWidth;
-	}
-
-	// Changing backface culling status requires starting over
-	static uint32 PreviousBackfaceCulling = PathTracingData.EnableCameraBackfaceCulling;
-	if (PreviousBackfaceCulling != PathTracingData.EnableCameraBackfaceCulling)
-	{
-		NeedInvalidation = true;
-		PreviousBackfaceCulling = PathTracingData.EnableCameraBackfaceCulling;
-	}
-
-	// Changing direct lighting requires starting over
-	static uint32 PreviousEnableDirectLighting = PathTracingData.EnableDirectLighting;
-	if (PreviousEnableDirectLighting != PathTracingData.EnableDirectLighting)
-	{
-		NeedInvalidation = true;
-		PreviousEnableDirectLighting = PathTracingData.EnableDirectLighting;
-	}
-
-	// Changing enable emissive requires starting over
-	static uint32 PreviousEnableEmissive = PathTracingData.EnableEmissive;
-	if (PreviousEnableEmissive != PathTracingData.EnableEmissive)
-	{
-		NeedInvalidation = true;
-		PreviousEnableEmissive = PathTracingData.EnableEmissive;
-	}
-
-	// Changing sampler type requires starting over
-	static uint32 PreviousSamplerType = PathTracingData.SamplerType;
-	if (PreviousSamplerType != PathTracingData.SamplerType)
-	{
-		NeedInvalidation = true;
-		PreviousSamplerType = PathTracingData.SamplerType;
-	}
-
-	// Changing visualization mode requires starting over
-	static int32 PreviousVisualizeLightGrid = PathTracingData.VisualizeLightGrid;
-	if (PreviousVisualizeLightGrid != PathTracingData.VisualizeLightGrid)
-	{
-		NeedInvalidation = true;
-		PreviousVisualizeLightGrid = PathTracingData.VisualizeLightGrid;
-	}
-
-	// the rest of PathTracingData and AdaptiveSamplingData is filled in by SetParameters below
-	return NeedInvalidation;
 }
 
 class FPathTracingSkylightPrepareCS : public FGlobalShader
@@ -1099,13 +1032,13 @@ class FPathTracingCompositorPS : public FGlobalShader
 };
 IMPLEMENT_SHADER_TYPE(, FPathTracingCompositorPS, TEXT("/Engine/Private/PathTracing/PathTracingCompositingPixelShader.usf"), TEXT("CompositeMain"), SF_Pixel);
 
-void FDeferredShadingSceneRenderer::PreparePathTracing(const FViewInfo& View, TArray<FRHIRayTracingShader*>& OutRayGenShaders)
+void FDeferredShadingSceneRenderer::PreparePathTracing(const FSceneViewFamily& ViewFamily, TArray<FRHIRayTracingShader*>& OutRayGenShaders)
 {
-	if (View.RayTracingRenderMode == ERayTracingRenderMode::PathTracing
-		&& FDataDrivenShaderPlatformInfo::GetSupportsPathTracing(View.GetShaderPlatform()))
+	if (ViewFamily.EngineShowFlags.PathTracing
+		&& FDataDrivenShaderPlatformInfo::GetSupportsPathTracing(ViewFamily.GetShaderPlatform()))
 	{
 		// Declare all RayGen shaders that require material closest hit shaders to be bound
-		auto RayGenShader = View.ShaderMap->GetShader<FPathTracingRG>();
+		auto RayGenShader = GetGlobalShaderMap(ViewFamily.GetShaderPlatform())->GetShader<FPathTracingRG>();
 		OutRayGenShaders.Add(RayGenShader.GetRayTracingShader());
 	}
 }
@@ -1132,97 +1065,72 @@ void FDeferredShadingSceneRenderer::RenderPathTracing(
 		return;
 	}
 
-	bool bArgsChanged = false;
+	FPathTracingConfig Config = {};
 
 	// Get current value of MaxSPP and reset render if it has changed
 	// NOTE: we ignore the CVar when using offline rendering
 	int32 SamplesPerPixelCVar = View.bIsOfflineRender ? -1 : CVarPathTracingSamplesPerPixel.GetValueOnRenderThread();
 	uint32 MaxSPP = SamplesPerPixelCVar > -1 ? SamplesPerPixelCVar : View.FinalPostProcessSettings.PathTracingSamplesPerPixel;
 	MaxSPP = FMath::Max(MaxSPP, 1u);
-	if (View.ViewState->PathTracingTargetSPP != MaxSPP)
-	{
-		// Store MaxSPP in the view state because we may have multiple views, each targetting a different sample count
-		View.ViewState->PathTracingTargetSPP = MaxSPP;
-		bArgsChanged = true;
-	}
-	// Changing FrameIndependentTemporalSeed requires starting over
-	bool LockedSamplingPattern = CVarPathTracingFrameIndependentTemporalSeed.GetValueOnRenderThread() == 0;
-	static bool PreviousLockedSamplingPattern = LockedSamplingPattern;
-	if (PreviousLockedSamplingPattern != LockedSamplingPattern)
-	{
-		PreviousLockedSamplingPattern = LockedSamplingPattern;
-		bArgsChanged = true;
-	}
+	Config.LockedSamplingPattern = CVarPathTracingFrameIndependentTemporalSeed.GetValueOnRenderThread() == 0;
+
 
 	// compute an integer code of what show flags related to lights are currently enabled so we can detect changes
-	int CurrentLightShowFlags = 0;
-	CurrentLightShowFlags |= View.Family->EngineShowFlags.SkyLighting           ? 1 << 0 : 0;
-	CurrentLightShowFlags |= View.Family->EngineShowFlags.DirectionalLights     ? 1 << 1 : 0;
-	CurrentLightShowFlags |= View.Family->EngineShowFlags.RectLights            ? 1 << 2 : 0;
-	CurrentLightShowFlags |= View.Family->EngineShowFlags.SpotLights            ? 1 << 3 : 0;
-	CurrentLightShowFlags |= View.Family->EngineShowFlags.PointLights           ? 1 << 4 : 0;
-	CurrentLightShowFlags |= View.Family->EngineShowFlags.TexturedLightProfiles ? 1 << 5 : 0;
-
-	static int PreviousLightShowFlags = CurrentLightShowFlags;
-	if (PreviousLightShowFlags != CurrentLightShowFlags)
-	{
-		PreviousLightShowFlags = CurrentLightShowFlags;
-		bArgsChanged = true;
-	}
+	Config.LightShowFlags = 0;
+	Config.LightShowFlags |= View.Family->EngineShowFlags.SkyLighting           ? 1 << 0 : 0;
+	Config.LightShowFlags |= View.Family->EngineShowFlags.DirectionalLights     ? 1 << 1 : 0;
+	Config.LightShowFlags |= View.Family->EngineShowFlags.RectLights            ? 1 << 2 : 0;
+	Config.LightShowFlags |= View.Family->EngineShowFlags.SpotLights            ? 1 << 3 : 0;
+	Config.LightShowFlags |= View.Family->EngineShowFlags.PointLights           ? 1 << 4 : 0;
+	Config.LightShowFlags |= View.Family->EngineShowFlags.TexturedLightProfiles ? 1 << 5 : 0;
 
-	bool UseMISCompensation = CVarPathTracingMISMode.GetValueOnRenderThread() == 2 && CVarPathTracingMISCompensation.GetValueOnRenderThread() != 0;
-	static bool PreviousUseMISCompensation = UseMISCompensation;
-	if (PreviousUseMISCompensation != UseMISCompensation)
-	{
-		PreviousUseMISCompensation = UseMISCompensation;
-		bArgsChanged = true;
-		// if the mode changes we need to rebuild the importance table
-		Scene->PathTracingSkylightTexture.SafeRelease();
-		Scene->PathTracingSkylightPdf.SafeRelease();
-	}
+	PrepareShaderArgs(View, Config.PathTracingData);
+
+	Config.VisibleLights = CVarPathTracingVisibleLights.GetValueOnRenderThread() != 0;
+	Config.UseMISCompensation = Config.PathTracingData.MISMode == 2 && CVarPathTracingMISCompensation.GetValueOnRenderThread() != 0;
 
-	const uint32 LightGridResolution = FMath::RoundUpToPowerOfTwo(CVarPathTracingLightGridResolution.GetValueOnRenderThread());
-	const uint32 LightGridMaxCount = FMath::Clamp(CVarPathTracingLightGridMaxCount.GetValueOnRenderThread(), 1, RAY_TRACING_LIGHT_COUNT_MAXIMUM);
+	Config.ViewRect = View.ViewRect;
 
-	static uint32 PreviousLightGridResolution = LightGridResolution;
-	if (PreviousLightGridResolution != LightGridResolution)
+	Config.LightGridResolution = FMath::RoundUpToPowerOfTwo(CVarPathTracingLightGridResolution.GetValueOnRenderThread());
+	Config.LightGridMaxCount = FMath::Clamp(CVarPathTracingLightGridMaxCount.GetValueOnRenderThread(), 1, RAY_TRACING_LIGHT_COUNT_MAXIMUM);
+
+	Config.PathTracingData.MaxSamples = MaxSPP;
+
+	bool FirstTime = false;
+	if (!View.ViewState->PathTracingLastConfig.IsValid())
 	{
-		PreviousLightGridResolution = LightGridResolution;
-		bArgsChanged = true;
+		View.ViewState->PathTracingLastConfig = MakePimpl<FPathTracingConfig>();
+		FirstTime = true; // we just initialized the option state for this view -- don't bother comparing in this case
 	}
 
-	static uint32_t PreviousLightGridMaxCount = LightGridMaxCount;
-	if (PreviousLightGridMaxCount != LightGridMaxCount)
+	if (FirstTime || Config.UseMISCompensation != View.ViewState->PathTracingLastConfig->UseMISCompensation)
 	{
-		PreviousLightGridMaxCount = LightGridMaxCount;
-		bArgsChanged = true;
+		// if the mode changes we need to rebuild the importance table
+		Scene->PathTracingSkylightTexture.SafeRelease();
+		Scene->PathTracingSkylightPdf.SafeRelease();
 	}
 
-	// Get other basic path tracing settings and see if we need to invalidate the current state
-	FPathTracingData PathTracingData;
-	bArgsChanged |= PrepareShaderArgs(View, PathTracingData);
-
 	// If the scene has changed in some way (camera move, object movement, etc ...)
 	// we must invalidate the ViewState to start over from scratch
-	if (bArgsChanged || View.ViewState->PathTracingRect != View.ViewRect)
+	if (FirstTime || Config.IsDifferent(*View.ViewState->PathTracingLastConfig))
 	{
+		// remember the options we used for next time
+		*View.ViewState->PathTracingLastConfig = Config;
 		View.ViewState->PathTracingInvalidate();
-		View.ViewState->PathTracingRect = View.ViewRect;
 	}
 
 	// Setup temporal seed _after_ invalidation in case we got reset
-	if (LockedSamplingPattern)
+	if (Config.LockedSamplingPattern)
 	{
 		// Count samples from 0 for deterministic results
-		PathTracingData.TemporalSeed = View.ViewState->PathTracingSampleIndex;
+		Config.PathTracingData.TemporalSeed = View.ViewState->PathTracingSampleIndex;
 	}
 	else
 	{
 		// Count samples from an ever-increasing counter to avoid screen-door effect
-		PathTracingData.TemporalSeed = View.ViewState->PathTracingFrameIndex;
+		Config.PathTracingData.TemporalSeed = View.ViewState->PathTracingFrameIndex;
 	}
-	PathTracingData.Iteration = View.ViewState->PathTracingSampleIndex;
-	PathTracingData.MaxSamples = MaxSPP;
+	Config.PathTracingData.Iteration = View.ViewState->PathTracingSampleIndex;
 
 	// Prepare radiance buffer (will be shared with display pass)
 	FRDGTexture* RadianceTexture = nullptr;
@@ -1241,17 +1149,17 @@ void FDeferredShadingSceneRenderer::RenderPathTracing(
 			TexCreate_ShaderResource | TexCreate_UAV);
 		RadianceTexture = GraphBuilder.CreateTexture(RadianceTextureDesc, TEXT("PathTracer.Radiance"), ERDGTextureFlags::MultiFrame);
 	}
-	bool NeedsMoreRays = PathTracingData.Iteration < MaxSPP;
+	const bool bNeedsMoreRays = Config.PathTracingData.Iteration < MaxSPP;
 
-	if (NeedsMoreRays)
+	if (bNeedsMoreRays)
 	{
 		FPathTracingRG::FParameters* PassParameters = GraphBuilder.AllocParameters<FPathTracingRG::FParameters>();
 		PassParameters->TLAS = View.GetRayTracingSceneViewChecked();
 		PassParameters->ViewUniformBuffer = View.ViewUniformBuffer;
-		PassParameters->PathTracingData = PathTracingData;
+		PassParameters->PathTracingData = Config.PathTracingData;
 		// upload sky/lights data
-		SetLightParameters(GraphBuilder, PassParameters, Scene, View, UseMISCompensation);
-		if (PathTracingData.EnableDirectLighting == 0)
+		SetLightParameters(GraphBuilder, PassParameters, Scene, View, Config.UseMISCompensation);
+		if (Config.PathTracingData.EnableDirectLighting == 0)
 		{
 			PassParameters->SceneVisibleLightCount = 0;
 		}
@@ -1299,7 +1207,7 @@ void FDeferredShadingSceneRenderer::RenderPathTracing(
 	// now add a pixel shader pass to display our Radiance buffer
 
 	FPathTracingCompositorPS::FParameters* DisplayParameters = GraphBuilder.AllocParameters<FPathTracingCompositorPS::FParameters>();
-	DisplayParameters->Iteration = PathTracingData.Iteration;
+	DisplayParameters->Iteration = Config.PathTracingData.Iteration;
 	DisplayParameters->MaxSamples = MaxSPP;
 	DisplayParameters->ProgressDisplayEnabled = CVarPathTracingProgressDisplay.GetValueOnRenderThread();
 	DisplayParameters->ViewUniformBuffer = View.ViewUniformBuffer;
diff --git a/Engine/Source/Runtime/Renderer/Private/PrimitiveSceneInfo.cpp b/Engine/Source/Runtime/Renderer/Private/PrimitiveSceneInfo.cpp
--- a/Engine/Source/Runtime/Renderer/Private/PrimitiveSceneInfo.cpp
+++ b/Engine/Source/Runtime/Renderer/Private/PrimitiveSceneInfo.cpp
@@ -664,24 +664,39 @@ void FPrimitiveSceneInfo::RemoveCachedNaniteDrawCommands()
 }
 
 #if RHI_RAYTRACING
+void FScene::RefreshRayTracingMeshCommandCache(FRHICommandListImmediate& RHICmdList, ERayTracingMeshCommandsMode Mode)
+{
+	if (Mode != CachedRayTracingMeshCommandsMode)
+	{
+		// remember if were using path tracing or not when caching the commands
+		CachedRayTracingMeshCommandsMode = Mode;
+		// get rid of all existing cached commands
+		CachedRayTracingMeshCommands.Empty(CachedRayTracingMeshCommands.Num());
+		// re-cache all current primitives
+		FPrimitiveSceneInfo::CacheRayTracingPrimitives(RHICmdList, this, Primitives);
+	}
+}
+
+
+// TODO: Remove RHICmdList argument? It does not appear to be used
 void FPrimitiveSceneInfo::CacheRayTracingPrimitives(FRHICommandListImmediate& RHICmdList, FScene* Scene, const TArrayView<FPrimitiveSceneInfo*>& SceneInfos)
 {
 	if (IsRayTracingEnabled() && !(Scene->World->WorldType == EWorldType::EditorPreview || Scene->World->WorldType == EWorldType::GamePreview))
 	{
 		checkf(RHISupportsMultithreadedShaderCreation(GMaxRHIShaderPlatform), TEXT("Raytracing code needs the ability to create shaders from task threads."));
 
-		FCachedRayTracingMeshCommandContext CommandContext(Scene->CachedRayTracingMeshCommands);
+		FCachedRayTracingMeshCommandStorage& CachedRayTracingMeshCommands = Scene->CachedRayTracingMeshCommands;
+		FCachedRayTracingMeshCommandContext CommandContext(CachedRayTracingMeshCommands);
 		FMeshPassProcessorRenderState PassDrawRenderState(Scene->UniformBuffers.ViewUniformBuffer);
-		FRayTracingMeshProcessor RayTracingMeshProcessor(&CommandContext, Scene, nullptr, PassDrawRenderState);
+		FRayTracingMeshProcessor RayTracingMeshProcessor(&CommandContext, Scene, nullptr, PassDrawRenderState, Scene->CachedRayTracingMeshCommandsMode);
 
 		for (FPrimitiveSceneInfo* SceneInfo : SceneInfos)
 		{
 			if (SceneInfo->RayTracingGeometries.Num() > 0 && SceneInfo->StaticMeshes.Num() > 0)
 			{
 				int32 MaxLOD = -1;
-				for (int32 MeshIndex = 0; MeshIndex < SceneInfo->StaticMeshes.Num(); MeshIndex++)
+				for (const FStaticMeshBatch& Mesh : SceneInfo->StaticMeshes)
 				{
-					FStaticMeshBatch& Mesh = SceneInfo->StaticMeshes[MeshIndex];
 					MaxLOD = MaxLOD < Mesh.LODIndex ? Mesh.LODIndex : MaxLOD;
 				}
 
@@ -691,17 +706,18 @@ void FPrimitiveSceneInfo::CacheRayTracingPrimitives(FRHICommandListImmediate& RH
 				SceneInfo->CachedRayTracingMeshCommandsHashPerLOD.Empty(MaxLOD + 1);
 				SceneInfo->CachedRayTracingMeshCommandsHashPerLOD.AddZeroed(MaxLOD + 1);
 
-				for (int32 MeshIndex = 0; MeshIndex < SceneInfo->StaticMeshes.Num(); MeshIndex++)
+				for (const FStaticMeshBatch& Mesh : SceneInfo->StaticMeshes)
 				{
-					FStaticMeshBatch& Mesh = SceneInfo->StaticMeshes[MeshIndex];
-
-					RayTracingMeshProcessor.AddMeshBatch(Mesh, ~0ull, SceneInfo->Proxy);
+					// Why do we pass a full mask here when the dynamic case only uses a mask of 1?
+					// Also note that the code below assumes only a single command was generated per batch.
+					const uint64 BatchElementMask = ~0ull;
+					RayTracingMeshProcessor.AddMeshBatch(Mesh, BatchElementMask, SceneInfo->Proxy);
 
 					if (CommandContext.CommandIndex >= 0)
 					{
 						uint64& Hash = SceneInfo->CachedRayTracingMeshCommandsHashPerLOD[Mesh.LODIndex];
 						Hash <<= 1;
-						Hash ^= Scene->CachedRayTracingMeshCommands[CommandContext.CommandIndex].ShaderBindings.GetDynamicInstancingHash();
+						Hash ^= CachedRayTracingMeshCommands[CommandContext.CommandIndex].ShaderBindings.GetDynamicInstancingHash();
 
 						SceneInfo->CachedRayTracingMeshCommandIndicesPerLOD[Mesh.LODIndex].Add(CommandContext.CommandIndex);
 						CommandContext.CommandIndex = -1;
@@ -743,16 +759,19 @@ void FPrimitiveSceneInfo::CacheRayTracingPrimitives(FRHICommandListImmediate& RH
 				SceneInfo->CachedRayTracingMeshCommandsHashPerLOD.Empty(1);
 				SceneInfo->CachedRayTracingMeshCommandsHashPerLOD.AddZeroed(1);
 
-				for (FMeshBatch& Mesh : CachedRayTracingInstance.Materials)
+				for (const FMeshBatch& Mesh : CachedRayTracingInstance.Materials)
 				{
-					RayTracingMeshProcessor.AddMeshBatch(Mesh, ~0ull, SceneInfo->Proxy);
+					// Why do we pass a full mask here when the dynamic case only uses a mask of 1?
+					// Also note that the code below assumes only a single command was generated per batch.
+					const uint64 BatchElementMask = ~0ull;
+					RayTracingMeshProcessor.AddMeshBatch(Mesh, BatchElementMask, SceneInfo->Proxy);
 
 					// The material section must emit a command. Otherwise, it should have been excluded earlier
 					check(CommandContext.CommandIndex >= 0);
 
 					uint64& Hash = SceneInfo->CachedRayTracingMeshCommandsHashPerLOD[Mesh.LODIndex];
 					Hash <<= 1;
-					Hash ^= Scene->CachedRayTracingMeshCommands[CommandContext.CommandIndex].ShaderBindings.GetDynamicInstancingHash();
+					Hash ^= CachedRayTracingMeshCommands[CommandContext.CommandIndex].ShaderBindings.GetDynamicInstancingHash();
 
 					SceneInfo->CachedRayTracingMeshCommandIndicesPerLOD[Mesh.LODIndex].Add(CommandContext.CommandIndex);
 					CommandContext.CommandIndex = -1;
@@ -1378,6 +1397,7 @@ void FPrimitiveSceneInfo::UpdateStaticMeshes(FRHICommandListImmediate& RHICmdLis
 
 		SceneInfo->RemoveCachedMeshDrawCommands();
 		SceneInfo->RemoveCachedNaniteDrawCommands();
+		SceneInfo->RemoveCachedRayTracingPrimitives();
 	}
 
 	if (bReAddToDrawLists)
diff --git a/Engine/Source/Runtime/Renderer/Private/RayTracing/RayTracingDebug.cpp b/Engine/Source/Runtime/Renderer/Private/RayTracing/RayTracingDebug.cpp
--- a/Engine/Source/Runtime/Renderer/Private/RayTracing/RayTracingDebug.cpp
+++ b/Engine/Source/Runtime/Renderer/Private/RayTracing/RayTracingDebug.cpp
@@ -84,12 +84,12 @@ public:
 IMPLEMENT_SHADER_TYPE(, FRayTracingDebugCHS, TEXT("/Engine/Private/RayTracing/RayTracingDebug.usf"), TEXT("RayTracingDebugMainCHS"), SF_RayHitGroup);
 
 
-void FDeferredShadingSceneRenderer::PrepareRayTracingDebug(const FViewInfo& View, TArray<FRHIRayTracingShader*>& OutRayGenShaders)
+void FDeferredShadingSceneRenderer::PrepareRayTracingDebug(const FSceneViewFamily& ViewFamily, TArray<FRHIRayTracingShader*>& OutRayGenShaders)
 {
 	// Declare all RayGen shaders that require material closest hit shaders to be bound
-	if (View.RayTracingRenderMode == ERayTracingRenderMode::RayTracingDebug)
+	if (ViewFamily.EngineShowFlags.RayTracingDebug)
 	{
-		auto RayGenShader = View.ShaderMap->GetShader<FRayTracingDebugRGS>();
+		auto RayGenShader = GetGlobalShaderMap(ViewFamily.GetShaderPlatform())->GetShader<FRayTracingDebugRGS>();
 		OutRayGenShaders.Add(RayGenShader.GetRayTracingShader());
 	}
 }
diff --git a/Engine/Source/Runtime/Renderer/Private/RayTracing/RayTracingMaterialHitShaders.cpp b/Engine/Source/Runtime/Renderer/Private/RayTracing/RayTracingMaterialHitShaders.cpp
--- a/Engine/Source/Runtime/Renderer/Private/RayTracing/RayTracingMaterialHitShaders.cpp
+++ b/Engine/Source/Runtime/Renderer/Private/RayTracing/RayTracingMaterialHitShaders.cpp
@@ -392,7 +392,9 @@ bool FRayTracingMeshProcessor::TryAddMeshBatch(
 		// Check for a cached light-map.
 		const bool bIsLitMaterial = ShadingModels.IsLit();
 		static const auto AllowStaticLightingVar = IConsoleManager::Get().FindTConsoleVariableDataInt(TEXT("r.AllowStaticLighting"));
-		const bool bAllowStaticLighting = (!AllowStaticLightingVar || AllowStaticLightingVar->GetValueOnRenderThread() != 0);
+		// NOTE: path tracing is not interested in baked static lighting (TODO: introduce a call for path tracer specific materials)
+		const bool bAllowStaticLighting = RayTracingMeshCommandsMode == ERayTracingMeshCommandsMode::RAY_TRACING &&
+			(!AllowStaticLightingVar || AllowStaticLightingVar->GetValueOnRenderThread() != 0);
 
 		const FLightMapInteraction LightMapInteraction = (bAllowStaticLighting && MeshBatch.LCI && bIsLitMaterial)
 			? MeshBatch.LCI->GetLightMapInteraction(FeatureLevel)
@@ -540,7 +542,7 @@ FRayTracingPipelineState* FDeferredShadingSceneRenderer::BindRayTracingMaterialP
 	TRACE_CPUPROFILER_EVENT_SCOPE(FDeferredShadingSceneRenderer::BindRayTracingMaterialPipeline);
 	SCOPE_CYCLE_COUNTER(STAT_BindRayTracingPipeline);
 
-	const bool bIsPathTracing = View.RayTracingRenderMode == ERayTracingRenderMode::PathTracing;
+	const bool bIsPathTracing = ViewFamily.EngineShowFlags.PathTracing;
 
 	FRayTracingPipelineStateInitializer Initializer;
 
diff --git a/Engine/Source/Runtime/Renderer/Private/RayTracing/RayTracingMaterialHitShaders.h b/Engine/Source/Runtime/Renderer/Private/RayTracing/RayTracingMaterialHitShaders.h
--- a/Engine/Source/Runtime/Renderer/Private/RayTracing/RayTracingMaterialHitShaders.h
+++ b/Engine/Source/Runtime/Renderer/Private/RayTracing/RayTracingMaterialHitShaders.h
@@ -14,13 +14,14 @@ class RENDERER_API FRayTracingMeshProcessor
 {
 public:
 
-	FRayTracingMeshProcessor(FRayTracingMeshCommandContext* InCommandContext, const FScene* InScene, const FSceneView* InViewIfDynamicMeshCommand, FMeshPassProcessorRenderState InPassDrawRenderState)
+	FRayTracingMeshProcessor(FRayTracingMeshCommandContext* InCommandContext, const FScene* InScene, const FSceneView* InViewIfDynamicMeshCommand, FMeshPassProcessorRenderState InPassDrawRenderState, ERayTracingMeshCommandsMode InRayTracingMeshCommandsMode)
 		:
 		CommandContext(InCommandContext),
 		Scene(InScene),
 		ViewIfDynamicMeshCommand(InViewIfDynamicMeshCommand),
 		FeatureLevel(InScene ? InScene->GetFeatureLevel() : ERHIFeatureLevel::SM5),
-		PassDrawRenderState(InPassDrawRenderState)
+		PassDrawRenderState(InPassDrawRenderState),
+		RayTracingMeshCommandsMode(InRayTracingMeshCommandsMode)
 	{}
 
 	virtual ~FRayTracingMeshProcessor() = default;
@@ -33,6 +34,7 @@ protected:
 	const FSceneView* ViewIfDynamicMeshCommand;
 	ERHIFeatureLevel::Type FeatureLevel;
 	FMeshPassProcessorRenderState PassDrawRenderState;
+	ERayTracingMeshCommandsMode RayTracingMeshCommandsMode;
 
 	virtual bool Process(
 		const FMeshBatch& RESTRICT MeshBatch,
diff --git a/Engine/Source/Runtime/Renderer/Private/RayTracing/RayTracingPrimaryRays.cpp b/Engine/Source/Runtime/Renderer/Private/RayTracing/RayTracingPrimaryRays.cpp
--- a/Engine/Source/Runtime/Renderer/Private/RayTracing/RayTracingPrimaryRays.cpp
+++ b/Engine/Source/Runtime/Renderer/Private/RayTracing/RayTracingPrimaryRays.cpp
@@ -72,7 +72,7 @@ void FDeferredShadingSceneRenderer::PrepareRayTracingTranslucency(const FViewInf
 {
 	// Declare all RayGen shaders that require material closest hit shaders to be bound.
 	// NOTE: Translucency shader may be used for primary ray debug view mode.
-	if (GetRayTracingTranslucencyOptions().bEnabled || View.RayTracingRenderMode == ERayTracingRenderMode::RayTracingDebug)
+	if (GetRayTracingTranslucencyOptions().bEnabled || View.Family->EngineShowFlags.RayTracingDebug)
 	{
 		FRayTracingPrimaryRaysRGS::FPermutationDomain PermutationVector;
 
diff --git a/Engine/Source/Runtime/Renderer/Private/RayTracing/RaytracingOptions.h b/Engine/Source/Runtime/Renderer/Private/RayTracing/RaytracingOptions.h
--- a/Engine/Source/Runtime/Renderer/Private/RayTracing/RaytracingOptions.h
+++ b/Engine/Source/Runtime/Renderer/Private/RayTracing/RaytracingOptions.h
@@ -57,7 +57,7 @@ extern bool ShouldRenderRayTracingTranslucency(const FViewInfo& View);
 extern bool ShouldRenderRayTracingShadows();
 extern bool ShouldRenderRayTracingShadowsForLight(const FLightSceneProxy& LightProxy);
 extern bool ShouldRenderRayTracingShadowsForLight(const FLightSceneInfoCompact& LightInfo);
-extern bool CanOverlayRayTracingOutput(const FViewInfo& View);
+extern bool HasRayTracedOverlay(const FSceneViewFamily& ViewFamily);
 
 extern bool EnableRayTracingShadowTwoSidedGeometry();
 extern float GetRaytracingMaxNormalBias();
@@ -116,9 +116,9 @@ FORCEINLINE bool ShouldRenderRayTracingShadowsForLight(const FLightSceneInfoComp
 	return false;
 }
 
-FORCEINLINE bool CanOverlayRayTracingOutput(const FViewInfo& View)
+FORCEINLINE bool HasRayTracedOverlay(const FSceneViewFamily& ViewFamily)
 {
-	return true;
+	return false;
 }
 
 FORCEINLINE bool CanUseRayTracingAMDHitToken()
diff --git a/Engine/Source/Runtime/Renderer/Private/RendererScene.cpp b/Engine/Source/Runtime/Renderer/Private/RendererScene.cpp
--- a/Engine/Source/Runtime/Renderer/Private/RendererScene.cpp
+++ b/Engine/Source/Runtime/Renderer/Private/RendererScene.cpp
@@ -242,8 +242,6 @@ FSceneViewState::FSceneViewState()
 	bUpdateLastExposure = false;
 
 #if RHI_RAYTRACING
-	PathTracingRect = FIntRect(0, 0, 0, 0);
-	PathTracingTargetSPP = 0;
 	PathTracingSampleIndex = 0;
 	PathTracingFrameIndex = 0;
 
@@ -1009,6 +1007,7 @@ FScene::FScene(UWorld* InWorld, bool bInRequiresHitProxies, bool bInIsEditorScen
 ,	FXSystem(nullptr)
 ,	bScenesPrimitivesNeedStaticMeshElementUpdate(false)
 ,	bPathTracingNeedsInvalidation(true)
+,   CachedRayTracingMeshCommandsMode(ERayTracingMeshCommandsMode::RAY_TRACING)
 ,	SkyLight(NULL)
 ,	ConvolvedSkyRenderTargetReadyIndex(-1)
 ,	RealTimeSlicedReflectionCaptureFirstFrameState(ERealTimeSlicedReflectionCaptureFirstFrameState::INIT)
diff --git a/Engine/Source/Runtime/Renderer/Private/ScenePrivate.h b/Engine/Source/Runtime/Renderer/Private/ScenePrivate.h
--- a/Engine/Source/Runtime/Renderer/Private/ScenePrivate.h
+++ b/Engine/Source/Runtime/Renderer/Private/ScenePrivate.h
@@ -89,6 +89,7 @@ class FLumenSceneData;
 class FVirtualShadowMapArrayCacheManager;
 class FComputeFramework;
 struct FHairStrandsInstance;
+struct FPathTracingConfig;
 
 /** Holds information about a single primitive's occlusion. */
 class FPrimitiveOcclusionHistory
@@ -980,10 +981,8 @@ public:
 
 	// Reference path tracing cached results
 	TRefCountPtr<IPooledRenderTarget> PathTracingRadianceRT;
-	// Keep track of the rectangle of pixels the Radiance texture is valid for so that path tracing can restart if this changes
-	FIntRect PathTracingRect;
-	// Target sampling count for the path tracer - to allow different views to target different quality levels
-	uint32   PathTracingTargetSPP;
+	// Keeps track of the internal path tracer options relevant to detecting when to restart the path tracer accumulation
+	TPimplPtr<FPathTracingConfig> PathTracingLastConfig;
 
 	// Current sample index to be rendered by the path tracer - this gets incremented each time the path tracer accumulates a frame of samples
 	uint32 PathTracingSampleIndex;
@@ -2552,6 +2551,12 @@ struct FMeshComputeDispatchCommand
 	uint32 NumMaxVertices;
 	FRWBuffer* TargetBuffer;
 };
+
+enum class ERayTracingMeshCommandsMode : uint8 {
+	RAY_TRACING,
+	PATH_TRACING
+};
+
 #endif
 
 /** 
@@ -2578,6 +2583,9 @@ public:
 
 #if RHI_RAYTRACING
 	FCachedRayTracingMeshCommandStorage CachedRayTracingMeshCommands;
+
+	/** This method is meant to be used when transitioning to and from path traced rendering since shader bindings must be updated in the cached commands. */
+	void RefreshRayTracingMeshCommandCache(FRHICommandListImmediate& RHICmdList, ERayTracingMeshCommandsMode Mode);
 #endif
 	/** Nanite state buckets. These are stored on the scene as they are computed at FPrimitiveSceneInfo::AddToScene time. */
 	FRWLock NaniteDrawCommandLock[ENaniteMeshPass::Num];
@@ -2667,6 +2675,9 @@ public:
 	/** True if a change to the scene that requires to invalidate the path tracer buffers has happened. */
 	bool bPathTracingNeedsInvalidation;
 
+	/** What mode where the cached RT commands prepared for last? */
+	ERayTracingMeshCommandsMode CachedRayTracingMeshCommandsMode;
+
 	/** The scene's sky light, if any. */
 	FSkyLightSceneProxy* SkyLight;
 
@@ -2705,7 +2716,7 @@ public:
 	/**
 	 * The path tracer uses its own representation of the skylight. These textures
 	 * are updated lazily by the path tracer when missing. Any code that modifies
-	 * the skylight appearance should simplify reset these pointers.
+	 * the skylight appearance should simply reset these pointers.
 	 */
 	TRefCountPtr<IPooledRenderTarget> PathTracingSkylightTexture;
 	TRefCountPtr<IPooledRenderTarget> PathTracingSkylightPdf;
diff --git a/Engine/Source/Runtime/Renderer/Private/SceneRendering.cpp b/Engine/Source/Runtime/Renderer/Private/SceneRendering.cpp
--- a/Engine/Source/Runtime/Renderer/Private/SceneRendering.cpp
+++ b/Engine/Source/Runtime/Renderer/Private/SceneRendering.cpp
@@ -2579,7 +2579,7 @@ void FSceneRenderer::PrepareViewRectsForRendering()
 			const bool bWillApplyTemporalAA = (IsPostProcessingEnabled(View) || View.bIsPlanarReflection)
 #if RHI_RAYTRACING
 				// path tracer does its own anti-aliasing
-				&& (View.RayTracingRenderMode != ERayTracingRenderMode::PathTracing)
+				&& (!ViewFamily.EngineShowFlags.PathTracing)
 #endif
 			;
 
diff --git a/Engine/Source/Runtime/Renderer/Private/SceneVisibility.cpp b/Engine/Source/Runtime/Renderer/Private/SceneVisibility.cpp
--- a/Engine/Source/Runtime/Renderer/Private/SceneVisibility.cpp
+++ b/Engine/Source/Runtime/Renderer/Private/SceneVisibility.cpp
@@ -4694,12 +4694,12 @@ void FDeferredShadingSceneRenderer::InitViewsPossiblyAfterPrepass(FRDGBuilder& G
 
 	FRHICommandListImmediate& RHICmdList = GraphBuilder.RHICmdList;
 
-	const bool bCanOverlayRayTracingOutput = CanOverlayRayTracingOutput(Views[0]);// #dxr_todo: UE-72557 multi-view case
+	const bool bHasRayTracedOverlay = HasRayTracedOverlay(ViewFamily);
 
 	if (ViewFamily.EngineShowFlags.DynamicShadows 
 		&& !IsSimpleForwardShadingEnabled(ShaderPlatform)
 		&& !ViewFamily.EngineShowFlags.HitProxies
-		&& bCanOverlayRayTracingOutput)
+		&& !bHasRayTracedOverlay)
 	{
 		// Setup dynamic shadows.
 		InitDynamicShadows(RHICmdList, DynamicIndexBufferForInitShadows, DynamicVertexBufferForInitShadows, DynamicReadBufferForInitShadows, InstanceCullingManager);
diff --git a/Engine/Source/Runtime/Renderer/Public/PrimitiveSceneInfo.h b/Engine/Source/Runtime/Renderer/Public/PrimitiveSceneInfo.h
--- a/Engine/Source/Runtime/Renderer/Public/PrimitiveSceneInfo.h
+++ b/Engine/Source/Runtime/Renderer/Public/PrimitiveSceneInfo.h
@@ -355,12 +355,6 @@ public:
 	FRayTracingGeometryInstance CachedRayTracingInstance;
 	TArray<FMatrix> CachedRayTracingInstanceLocalTransforms;
 	TArray<FMatrix> CachedRayTracingInstanceWorldTransforms;
-
-	struct FStaticMeshOrCommandIndex
-	{
-		int32 StaticMeshIndex;
-		int32 CommandIndex;
-	};
 #endif
 
 	/** Initialization constructor. */
