Fix missing include in ScenePrivate.h

diff --git a/Engine/Source/Runtime/Renderer/Private/ScenePrivate.h b/Engine/Source/Runtime/Renderer/Private/ScenePrivate.h
--- a/Engine/Source/Runtime/Renderer/Private/ScenePrivate.h
+++ b/Engine/Source/Runtime/Renderer/Private/ScenePrivate.h
@@ -12,6 +12,7 @@
 #include "Misc/Guid.h"
 #include "Math/RandomStream.h"
 #include "Engine/EngineTypes.h"
+#include "Templates/PimplPtr.h"
 #include "RHI.h"
 #include "RenderResource.h"
 #include "RenderingThread.h"
