Globally disable "r.EnableAsyncComputeTranslucencyLightingVolumeClear" due to GPU crashes in async compute FClearTranslucentLightingVolumeCS shader.
 - Was previously disabled only on selected platforms (UE-62964)

diff --git a/Engine/Source/Runtime/Renderer/Private/DeferredShadingRenderer.cpp b/Engine/Source/Runtime/Renderer/Private/DeferredShadingRenderer.cpp
--- a/Engine/Source/Runtime/Renderer/Private/DeferredShadingRenderer.cpp
+++ b/Engine/Source/Runtime/Renderer/Private/DeferredShadingRenderer.cpp
@@ -82,7 +82,7 @@ static TAutoConsoleVariable<int32> CVarIrisNormal(
 	TEXT(" 1: on"),
 	ECVF_ReadOnly);
 
-int32 GbEnableAsyncComputeTranslucencyLightingVolumeClear = 1;
+int32 GbEnableAsyncComputeTranslucencyLightingVolumeClear = 0; // @todo: disabled due to GPU crashes
 static FAutoConsoleVariableRef CVarEnableAsyncComputeTranslucencyLightingVolumeClear(
 	TEXT("r.EnableAsyncComputeTranslucencyLightingVolumeClear"),
 	GbEnableAsyncComputeTranslucencyLightingVolumeClear,
