Fix compile errors on non-raytracing platforms

diff --git a/Engine/Source/Runtime/Renderer/Private/PrimitiveSceneInfo.cpp b/Engine/Source/Runtime/Renderer/Private/PrimitiveSceneInfo.cpp
--- a/Engine/Source/Runtime/Renderer/Private/PrimitiveSceneInfo.cpp
+++ b/Engine/Source/Runtime/Renderer/Private/PrimitiveSceneInfo.cpp
@@ -1397,7 +1397,9 @@ void FPrimitiveSceneInfo::UpdateStaticMeshes(FRHICommandListImmediate& RHICmdLis
 
 		SceneInfo->RemoveCachedMeshDrawCommands();
 		SceneInfo->RemoveCachedNaniteDrawCommands();
+	#if RHI_RAYTRACING
 		SceneInfo->RemoveCachedRayTracingPrimitives();
+	#endif
 	}
 
 	if (bReAddToDrawLists)
diff --git a/Engine/Source/Runtime/Renderer/Private/ScenePrivate.h b/Engine/Source/Runtime/Renderer/Private/ScenePrivate.h
--- a/Engine/Source/Runtime/Renderer/Private/ScenePrivate.h
+++ b/Engine/Source/Runtime/Renderer/Private/ScenePrivate.h
@@ -2676,8 +2676,10 @@ public:
 	/** True if a change to the scene that requires to invalidate the path tracer buffers has happened. */
 	bool bPathTracingNeedsInvalidation;
 
+#if RHI_RAYTRACING
 	/** What mode where the cached RT commands prepared for last? */
 	ERayTracingMeshCommandsMode CachedRayTracingMeshCommandsMode;
+#endif
 
 	/** The scene's sky light, if any. */
 	FSkyLightSceneProxy* SkyLight;
