Fix compile errors on non-raytracing platforms

diff --git a/Engine/Source/Runtime/Renderer/Private/RendererScene.cpp b/Engine/Source/Runtime/Renderer/Private/RendererScene.cpp
--- a/Engine/Source/Runtime/Renderer/Private/RendererScene.cpp
+++ b/Engine/Source/Runtime/Renderer/Private/RendererScene.cpp
@@ -1007,7 +1007,9 @@ FScene::FScene(UWorld* InWorld, bool bInRequiresHitProxies, bool bInIsEditorScen
 ,	FXSystem(nullptr)
 ,	bScenesPrimitivesNeedStaticMeshElementUpdate(false)
 ,	bPathTracingNeedsInvalidation(true)
+#if RHI_RAYTRACING
 ,   CachedRayTracingMeshCommandsMode(ERayTracingMeshCommandsMode::RAY_TRACING)
+#endif
 ,	SkyLight(NULL)
 ,	ConvolvedSkyRenderTargetReadyIndex(-1)
 ,	RealTimeSlicedReflectionCaptureFirstFrameState(ERealTimeSlicedReflectionCaptureFirstFrameState::INIT)
