Don't do custom stencil if no custom depth texture exists.

diff --git a/Engine/Source/Runtime/Renderer/Private/PostProcess/PostProcessMaterial.cpp b/Engine/Source/Runtime/Renderer/Private/PostProcess/PostProcessMaterial.cpp
--- a/Engine/Source/Runtime/Renderer/Private/PostProcess/PostProcessMaterial.cpp
+++ b/Engine/Source/Runtime/Renderer/Private/PostProcess/PostProcessMaterial.cpp
@@ -385,7 +385,7 @@ FScreenPassTexture AddPostProcessMaterialPass(
 	// Allocate custom depth stencil texture(s) and depth stencil state.
 	const EMaterialCustomDepthPolicy CustomStencilPolicy = GetMaterialCustomDepthPolicy(Material, FeatureLevel);
 
-	if (CustomStencilPolicy == EMaterialCustomDepthPolicy::Enabled)
+	if (CustomStencilPolicy == EMaterialCustomDepthPolicy::Enabled && HasBeenProduced(Inputs.CustomDepthTexture))
 	{
 		check(Inputs.CustomDepthTexture);
 		DepthStencilTexture = Inputs.CustomDepthTexture;
