Skip always visible primitives in FetchVisibilityForPrimitives_Range.

#rnx

diff --git a/Engine/Source/Runtime/Renderer/Private/SceneVisibility.cpp b/Engine/Source/Runtime/Renderer/Private/SceneVisibility.cpp
--- a/Engine/Source/Runtime/Renderer/Private/SceneVisibility.cpp
+++ b/Engine/Source/Runtime/Renderer/Private/SceneVisibility.cpp
@@ -910,6 +910,12 @@ static void FetchVisibilityForPrimitives_Range(FVisForPrimParams& Params, FGloba
 		}
 #endif
 
+		if (!bCanBeOccluded)
+		{
+			View.PrimitiveDefinitelyUnoccludedMap.AccessCorrespondingBit(BitIt) = true;
+			continue;
+		}
+
 		//we can't allow the prim history insertion array to realloc or it will invalidate pointers in the other output arrays.
 		const bool bCanAllocPrimHistory = bSingleThreaded || InsertPrimitiveOcclusionHistory->Num() < InsertPrimitiveOcclusionHistory->Max();		
 
