Fixed a crash due to AnyRayTracingPassEnabled returning true in some cases while ray tracing may be not supported (i.e. if certain RT features are force-enabled via CVars). Returning true from this function results in a check failure inside WaitForRayTracingScene() later.

diff --git a/Engine/Source/Runtime/Renderer/Private/DeferredShadingRenderer.cpp b/Engine/Source/Runtime/Renderer/Private/DeferredShadingRenderer.cpp
--- a/Engine/Source/Runtime/Renderer/Private/DeferredShadingRenderer.cpp
+++ b/Engine/Source/Runtime/Renderer/Private/DeferredShadingRenderer.cpp
@@ -3186,6 +3186,11 @@ void FDeferredShadingSceneRenderer::Render(FRDGBuilder& GraphBuilder)
 
 bool AnyRayTracingPassEnabled(const FScene* Scene, const FViewInfo& View)
 {
+	if (!IsRayTracingEnabled())
+	{
+		return false;
+	}
+
 	return ShouldRenderRayTracingAmbientOcclusion(View)
 		|| ShouldRenderRayTracingReflections(View)
 		|| ShouldRenderRayTracingGlobalIllumination(View)
