Restore functionality to automatically map newly allocated VTs, but only for lowest resolution mip level
- This should be much faster than previous code which attempted to map all mip levels
- Need to ensure at least lowest resolution mip is mapped, in order to avoid rendering frames with garbage data mapped
- Add warning for AllocatedVTs that never manage to complete this mappings

diff --git a/Engine/Source/Runtime/RenderCore/Public/VirtualTexturing.h b/Engine/Source/Runtime/RenderCore/Public/VirtualTexturing.h
--- a/Engine/Source/Runtime/RenderCore/Public/VirtualTexturing.h
+++ b/Engine/Source/Runtime/RenderCore/Public/VirtualTexturing.h
@@ -385,6 +385,7 @@ public:
 protected:
 	friend class FVirtualTextureSystem;
 	virtual void Destroy(class FVirtualTextureSystem* InSystem) = 0;
+	virtual bool TryMapLockedTiles(class FVirtualTextureSystem* InSystem) const = 0;
 	virtual ~IAllocatedVirtualTexture() {}
 
 	FAllocatedVTDescription Description;
diff --git a/Engine/Source/Runtime/Renderer/Private/VT/AllocatedVirtualTexture.cpp b/Engine/Source/Runtime/Renderer/Private/VT/AllocatedVirtualTexture.cpp
--- a/Engine/Source/Runtime/Renderer/Private/VT/AllocatedVirtualTexture.cpp
+++ b/Engine/Source/Runtime/Renderer/Private/VT/AllocatedVirtualTexture.cpp
@@ -220,6 +220,67 @@ void FAllocatedVirtualTexture::LockOrUnlockTiles(FVirtualTextureSystem* InSystem
 	}
 }
 
+bool FAllocatedVirtualTexture::TryMapLockedTiles(FVirtualTextureSystem* InSystem) const
+{
+	bool bHasMissingTiles = false;
+	for (int32 PageTableLayerIndex = 0u; PageTableLayerIndex < UniquePageTableLayers.Num(); ++PageTableLayerIndex)
+	{
+		const FPageTableLayerDesc& PageTableLayer = UniquePageTableLayers[PageTableLayerIndex];
+		const FProducerDesc& UniqueProducer = UniqueProducers[PageTableLayer.UniqueProducerIndex];
+		const FVirtualTextureProducer* Producer = InSystem->FindProducer(UniqueProducer.Handle);
+		if (!Producer)
+		{
+			continue;
+		}
+
+		const uint32 WidthInTiles = Producer->GetWidthInTiles();
+		const uint32 HeightInTiles = Producer->GetHeightInTiles();
+		const uint32 Local_vLevel = FMath::Min(Producer->GetMaxLevel(), MaxLevel - UniqueProducer.MipBias);
+		const uint32 MipScaleFactor = (1u << Local_vLevel);
+		const uint32 RootWidthInTiles = FMath::DivideAndRoundUp(WidthInTiles, MipScaleFactor);
+		const uint32 RootHeightInTiles = FMath::DivideAndRoundUp(HeightInTiles, MipScaleFactor);
+
+		FTexturePagePool& PagePool = PageTableLayer.PhysicalSpace->GetPagePool();
+		FTexturePageMap& PageMap = Space->GetPageMapForPageTableLayer(PageTableLayerIndex);
+
+		uint32 NumNonResidentPages = 0u;
+		for (uint32 TileY = 0u; TileY < RootHeightInTiles; ++TileY)
+		{
+			for (uint32 TileX = 0u; TileX < RootWidthInTiles; ++TileX)
+			{
+				const uint32 vAddress = FMath::MortonCode2(VirtualPageX + (TileX << MaxLevel)) | (FMath::MortonCode2(VirtualPageY + (TileY << MaxLevel)) << 1);
+				uint32 pAddress = PageMap.FindPageAddress(MaxLevel, vAddress);
+				if (pAddress == ~0u)
+				{
+					uint32 Local_vAddress = FMath::MortonCode2(TileX) | (FMath::MortonCode2(TileY) << 1);
+
+					const uint32 LocalMipBias = Producer->GetVirtualTexture()->GetLocalMipBias(Local_vLevel, Local_vAddress);
+					Local_vAddress >>= (LocalMipBias * Description.Dimensions);
+
+					pAddress = PagePool.FindPageAddress(UniqueProducer.Handle, PageTableLayer.ProducerPhysicalGroupIndex, Local_vAddress, Local_vLevel + LocalMipBias);
+					if (pAddress != ~0u)
+					{
+						PagePool.MapPage(Space, PageTableLayer.PhysicalSpace, PageTableLayerIndex, MaxLevel, MaxLevel, vAddress, MaxLevel + LocalMipBias, pAddress);
+					}
+					else
+					{
+						bHasMissingTiles = true;
+					}
+				}
+			}
+		}
+	}
+
+	// Display a warning message (once) if we've failed to map pages for this after a set number of frames
+	// Generally there should be no delay, but if the system is saturated, it's possible that locked pages may not be loaded immediately
+	if (bHasMissingTiles && InSystem->GetFrame() == FrameAllocated + 30u)
+	{
+		UE_LOG(LogVirtualTexturing, Warning, TEXT("Failed to map lowest resolution mip for AllocatedVT %s"), *Description.Name.ToString());
+	}
+
+	return !bHasMissingTiles;
+}
+
 uint32 FAllocatedVirtualTexture::AddUniqueProducer(FVirtualTextureProducerHandle const& InHandle, const FVirtualTextureProducer* InProducer)
 {
 	for (int32 Index = 0u; Index < UniqueProducers.Num(); ++Index)
diff --git a/Engine/Source/Runtime/Renderer/Private/VT/AllocatedVirtualTexture.h b/Engine/Source/Runtime/Renderer/Private/VT/AllocatedVirtualTexture.h
--- a/Engine/Source/Runtime/Renderer/Private/VT/AllocatedVirtualTexture.h
+++ b/Engine/Source/Runtime/Renderer/Private/VT/AllocatedVirtualTexture.h
@@ -29,8 +29,6 @@ public:
 
 	virtual ~FAllocatedVirtualTexture();
 
-	void Release(FVirtualTextureSystem* System);
-
 	inline uint32 GetFrameAllocated() const { return FrameAllocated; }
 
 	void AssignVirtualAddress(uint32 vAddress);
@@ -47,6 +45,7 @@ public:
 	virtual void GetPackedPageTableUniform(FUintVector4* OutUniform) const override;
 	virtual void GetPackedUniform(FUintVector4* OutUniform, uint32 LayerIndex) const override;
 	virtual void Destroy(FVirtualTextureSystem* InSystem) override;
+	virtual bool TryMapLockedTiles(FVirtualTextureSystem* InSystem) const override;
 	// end IAllocatedVirtualTexture
 
 	inline FVirtualTextureSpace* GetSpace() const { return Space; }
diff --git a/Engine/Source/Runtime/Renderer/Private/VT/VirtualTextureSystem.cpp b/Engine/Source/Runtime/Renderer/Private/VT/VirtualTextureSystem.cpp
--- a/Engine/Source/Runtime/Renderer/Private/VT/VirtualTextureSystem.cpp
+++ b/Engine/Source/Runtime/Renderer/Private/VT/VirtualTextureSystem.cpp
@@ -20,6 +20,8 @@
 #include "VT/VirtualTextureScalability.h"
 #include "VT/VirtualTextureSpace.h"
 
+PRAGMA_DISABLE_OPTIMIZATION
+
 #define LOCTEXT_NAMESPACE "VirtualTexture"
 
 DECLARE_CYCLE_STAT(TEXT("VirtualTextureSystem Update"), STAT_VirtualTextureSystem_Update, STATGROUP_VirtualTexturing);
@@ -28,6 +30,7 @@ DECLARE_CYCLE_STAT(TEXT("Gather Requests"), STAT_ProcessRequests_Gather, STATGRO
 DECLARE_CYCLE_STAT(TEXT("Sort Requests"), STAT_ProcessRequests_Sort, STATGROUP_VirtualTexturing);
 DECLARE_CYCLE_STAT(TEXT("Submit Requests"), STAT_ProcessRequests_Submit, STATGROUP_VirtualTexturing);
 DECLARE_CYCLE_STAT(TEXT("Map Requests"), STAT_ProcessRequests_Map, STATGROUP_VirtualTexturing);
+DECLARE_CYCLE_STAT(TEXT("Map New VTs"), STAT_ProcessRequests_MapNew, STATGROUP_VirtualTexturing);
 DECLARE_CYCLE_STAT(TEXT("Finalize Requests"), STAT_ProcessRequests_Finalize, STATGROUP_VirtualTexturing);
 DECLARE_CYCLE_STAT(TEXT("Merge Unique Pages"), STAT_ProcessRequests_MergePages, STATGROUP_VirtualTexturing);
 DECLARE_CYCLE_STAT(TEXT("Merge Requests"), STAT_ProcessRequests_MergeRequests, STATGROUP_VirtualTexturing);
@@ -522,6 +525,10 @@ IAllocatedVirtualTexture* FVirtualTextureSystem::AllocateVirtualTexture(const FA
 
 	AllocatedVT = new FAllocatedVirtualTexture(this, Frame, Desc, ProducerForLayer, BlockWidthInTiles, BlockHeightInTiles, WidthInBlocks, HeightInBlocks, DepthInTiles);
 	AllocatedVT->NumRefs = 1;
+	if (bAnyLayerProducerWantsPersistentHighestMip)
+	{
+		AllocatedVTsToMap.Add(AllocatedVT);
+	}
 	return AllocatedVT;
 }
 
@@ -576,6 +583,8 @@ void FVirtualTextureSystem::DestroyPendingVirtualTextures(bool bForceDestroyAll)
 
 	for (IAllocatedVirtualTexture* AllocatedVT : AllocatedVTsToDelete)
 	{
+		// shouldn't be more than 1 instance of this in the list
+		verify(AllocatedVTsToMap.Remove(AllocatedVT) <= 1);
 		verify(AllocatedVTs.Remove(AllocatedVT->GetDescription()) == 1);
 		AllocatedVT->Destroy(this);
 		delete AllocatedVT;
@@ -2238,6 +2247,28 @@ void FVirtualTextureSystem::SubmitRequests(FRDGBuilder& GraphBuilder, ERHIFeatur
 		}
 	}
 
+	// Map any resident tiles to newly allocated VTs
+	if(AllocatedVTsToMap.Num() > 0)
+	{
+		SCOPE_CYCLE_COUNTER(STAT_ProcessRequests_MapNew);
+
+		uint32 Index = 0u;
+		while (Index < (uint32)AllocatedVTsToMap.Num())
+		{
+			const IAllocatedVirtualTexture* AllocatedVT = AllocatedVTsToMap[Index];
+			if (AllocatedVT->TryMapLockedTiles(this))
+			{
+				AllocatedVTsToMap.RemoveAtSwap(Index, 1, false);
+			}
+			else
+			{
+				Index++;
+			}
+		}
+
+		AllocatedVTsToMap.Shrink();
+	}
+
 	// Finalize requests
 	{
 		SCOPE_CYCLE_COUNTER(STAT_ProcessRequests_Finalize);
@@ -2287,4 +2318,6 @@ void FVirtualTextureSystem::ReleasePendingResources()
 	ReleasePendingSpaces();
 }
 
-#undef LOCTEXT_NAMESPACE
\ No newline at end of file
+#undef LOCTEXT_NAMESPACE
+
+PRAGMA_ENABLE_OPTIMIZATION
\ No newline at end of file
diff --git a/Engine/Source/Runtime/Renderer/Private/VT/VirtualTextureSystem.h b/Engine/Source/Runtime/Renderer/Private/VT/VirtualTextureSystem.h
--- a/Engine/Source/Runtime/Renderer/Private/VT/VirtualTextureSystem.h
+++ b/Engine/Source/Runtime/Renderer/Private/VT/VirtualTextureSystem.h
@@ -37,6 +37,8 @@ public:
 	static void Shutdown();
 	static FVirtualTextureSystem& Get();
 
+	uint32 GetFrame() const { return Frame; }
+
 	void AllocateResources(FRDGBuilder& GraphBuilder, ERHIFeatureLevel::Type FeatureLevel);
 	void CallPendingCallbacks();
 	void Update( FRDGBuilder& GraphBuilder, ERHIFeatureLevel::Type FeatureLevel, FScene* Scene);
@@ -122,6 +124,7 @@ private:
 	TArray<IAllocatedVirtualTexture*> PendingDeleteAllocatedVTs;
 
 	TMap<FAllocatedVTDescription, FAllocatedVirtualTexture*> AllocatedVTs;
+	TArray<IAllocatedVirtualTexture*> AllocatedVTsToMap;
 
 	FAdaptiveVirtualTexture* AdaptiveVTs[MaxSpaces] = { nullptr };
 
