Remove PRAGMA_DISABLE_OPTIMIZATION left in by mistake

diff --git a/Engine/Source/Runtime/Renderer/Private/VT/VirtualTextureSystem.cpp b/Engine/Source/Runtime/Renderer/Private/VT/VirtualTextureSystem.cpp
--- a/Engine/Source/Runtime/Renderer/Private/VT/VirtualTextureSystem.cpp
+++ b/Engine/Source/Runtime/Renderer/Private/VT/VirtualTextureSystem.cpp
@@ -20,8 +20,6 @@
 #include "VT/VirtualTextureScalability.h"
 #include "VT/VirtualTextureSpace.h"
 
-PRAGMA_DISABLE_OPTIMIZATION
-
 #define LOCTEXT_NAMESPACE "VirtualTexture"
 
 DECLARE_CYCLE_STAT(TEXT("VirtualTextureSystem Update"), STAT_VirtualTextureSystem_Update, STATGROUP_VirtualTexturing);
@@ -2319,5 +2317,3 @@ void FVirtualTextureSystem::ReleasePendingResources()
 }
 
 #undef LOCTEXT_NAMESPACE
-
-PRAGMA_ENABLE_OPTIMIZATION
\ No newline at end of file
