Also destroy pending AllocatedVTs when flushing destroyed VT resources

diff --git a/Engine/Source/Runtime/Renderer/Private/VT/VirtualTextureSystem.cpp b/Engine/Source/Runtime/Renderer/Private/VT/VirtualTextureSystem.cpp
--- a/Engine/Source/Runtime/Renderer/Private/VT/VirtualTextureSystem.cpp
+++ b/Engine/Source/Runtime/Renderer/Private/VT/VirtualTextureSystem.cpp
@@ -2313,6 +2313,7 @@ void FVirtualTextureSystem::CallPendingCallbacks()
 
 void FVirtualTextureSystem::ReleasePendingResources()
 {
+	DestroyPendingVirtualTextures(true);
 	ReleasePendingSpaces();
 }
 
