Don't set the depth stencil binding on the ShaderDrawDebug screen pass, as it's unused and causes RHI crashes when TAA Upsampling is enabled.

diff --git a/Engine/Source/Runtime/Renderer/Private/GpuDebugRendering.cpp b/Engine/Source/Runtime/Renderer/Private/GpuDebugRendering.cpp
--- a/Engine/Source/Runtime/Renderer/Private/GpuDebugRendering.cpp
+++ b/Engine/Source/Runtime/Renderer/Private/GpuDebugRendering.cpp
@@ -296,7 +296,7 @@ namespace ShaderDrawDebug
 
 			FShaderDrawVSPSParameters * PassParameters = GraphBuilder.AllocParameters<FShaderDrawVSPSParameters >();
 			PassParameters->ShaderDrawPSParameters.RenderTargets[0] = FRenderTargetBinding(OutputTexture, ERenderTargetLoadAction::ELoad);
-			PassParameters->ShaderDrawPSParameters.RenderTargets.DepthStencil = FDepthStencilBinding(DepthTexture, ERenderTargetLoadAction::ELoad, FExclusiveDepthStencil::DepthRead_StencilNop);
+			//PassParameters->ShaderDrawPSParameters.RenderTargets.DepthStencil = FDepthStencilBinding(DepthTexture, ERenderTargetLoadAction::ELoad, FExclusiveDepthStencil::DepthRead_StencilNop);
 			PassParameters->ShaderDrawPSParameters.DepthTexture = DepthTexture;
 			PassParameters->ShaderDrawPSParameters.DepthTextureResolution    = FIntPoint(DepthTexture->Desc.Extent.X, DepthTexture->Desc.Extent.Y);
 			PassParameters->ShaderDrawPSParameters.DepthTextureInvResolution = FVector2D(1.f/DepthTexture->Desc.Extent.X, 1.f / DepthTexture->Desc.Extent.Y);
