Disable indirect capsule shadow processing when disabled.

diff --git a/Engine/Source/Runtime/Renderer/Private/CapsuleShadowRendering.cpp b/Engine/Source/Runtime/Renderer/Private/CapsuleShadowRendering.cpp
--- a/Engine/Source/Runtime/Renderer/Private/CapsuleShadowRendering.cpp
+++ b/Engine/Source/Runtime/Renderer/Private/CapsuleShadowRendering.cpp
@@ -865,6 +865,11 @@ void FDeferredShadingSceneRenderer::CreateIndirectCapsuleShadows()
 {
 	QUICK_SCOPE_CYCLE_COUNTER(STAT_CreateIndirectCapsuleShadows);
 
+	if (!SupportsCapsuleIndirectShadows(FeatureLevel, ShaderPlatform))
+	{
+		return;
+	}
+
 	for (int32 PrimitiveIndex = 0; PrimitiveIndex < Scene->DynamicIndirectCasterPrimitives.Num(); PrimitiveIndex++)
 	{
 		FPrimitiveSceneInfo* PrimitiveSceneInfo = Scene->DynamicIndirectCasterPrimitives[PrimitiveIndex];
