Skip creating RHI buffer every frame. Early out from scene view extensions if none exist.

diff --git a/Engine/Source/Runtime/Renderer/Private/SceneRendering.cpp b/Engine/Source/Runtime/Renderer/Private/SceneRendering.cpp
--- a/Engine/Source/Runtime/Renderer/Private/SceneRendering.cpp
+++ b/Engine/Source/Runtime/Renderer/Private/SceneRendering.cpp
@@ -3609,6 +3609,11 @@ void FSceneRenderer::UpdatePrimitiveIndirectLightingCacheBuffers()
 */
 void FSceneRenderer::ViewExtensionPreRender_RenderThread(FRHICommandListImmediate& RHICmdList, FSceneRenderer* SceneRenderer)
 {
+	if (SceneRenderer->ViewFamily.ViewExtensions.IsEmpty())
+	{
+		return;
+	}
+
 	FMemMark MemStackMark(FMemStack::Get());
 
 	{
@@ -4774,6 +4779,11 @@ void FSceneRenderer::UpdateSkyIrradianceGpuBuffer(FRHICommandListImmediate& RHIC
 		return;
 	}
 
+	if (!Scene->SkyIrradianceEnvironmentMap.Buffer)
+	{
+		Scene->SkyIrradianceEnvironmentMap.Initialize(TEXT("SkyIrradianceEnvironmentMap"), sizeof(FVector4), 7, 0);
+	}
+
 	FVector4 OutSkyIrradianceEnvironmentMap[7];
 	// Make sure there's no padding since we're going to cast to FVector4*
 	checkSlow(sizeof(OutSkyIrradianceEnvironmentMap) == sizeof(FVector4) * 7);
@@ -4790,10 +4800,6 @@ void FSceneRenderer::UpdateSkyIrradianceGpuBuffer(FRHICommandListImmediate& RHIC
 		const FSHVectorRGB3& SkyIrradiance = Scene->SkyLight->IrradianceEnvironmentMap;
 		SetupSkyIrradianceEnvironmentMapConstantsFromSkyIrradiance(OutSkyIrradianceEnvironmentMap, SkyIrradiance);
 
-		// Create a buffer for this frame 
-		Scene->SkyIrradianceEnvironmentMap.Release();
-		Scene->SkyIrradianceEnvironmentMap.Initialize(TEXT("SkyIrradianceEnvironmentMap"), sizeof(FVector4), 7, 0);
-
 		// Set the captured environment map data
 		void* DataPtr = RHICmdList.LockBuffer(Scene->SkyIrradianceEnvironmentMap.Buffer, 0, Scene->SkyIrradianceEnvironmentMap.NumBytes, RLM_WriteOnly);
 		checkSlow(Scene->SkyIrradianceEnvironmentMap.NumBytes == sizeof(OutSkyIrradianceEnvironmentMap));
@@ -4802,8 +4808,6 @@ void FSceneRenderer::UpdateSkyIrradianceGpuBuffer(FRHICommandListImmediate& RHIC
 	}
 	else if (Scene->SkyIrradianceEnvironmentMap.NumBytes == 0)
 	{
-		Scene->SkyIrradianceEnvironmentMap.Initialize(TEXT("SkyIrradianceEnvironmentMap"), sizeof(FVector4), 7, 0);
-
 		// Ensure that sky irradiance SH buffer contains sensible initial values (zero init).
 		// If there is no sky in the level, then nothing else may fill this buffer.
 		void* DataPtr = RHICmdList.LockBuffer(Scene->SkyIrradianceEnvironmentMap.Buffer, 0, Scene->SkyIrradianceEnvironmentMap.NumBytes, RLM_WriteOnly);
