Don't create a distance field streaming and pass task if no streaming is requested.

diff --git a/Engine/Source/Runtime/Renderer/Private/DistanceFieldStreaming.cpp b/Engine/Source/Runtime/Renderer/Private/DistanceFieldStreaming.cpp
--- a/Engine/Source/Runtime/Renderer/Private/DistanceFieldStreaming.cpp
+++ b/Engine/Source/Runtime/Renderer/Private/DistanceFieldStreaming.cpp
@@ -897,42 +897,47 @@ void FDistanceFieldSceneData::UpdateDistanceFieldAtlas(
 			UpdateParameters.BrickUploadCoordinatesPtr = AtlasUpload.BrickUploadCoordinatesPtr;
 		}
 
-		UpdateParameters.NewReadRequests = MoveTemp(NewReadRequests);
-		UpdateParameters.ReadRequestsToUpload = MoveTemp(ReadRequestsToUpload);
-		UpdateParameters.ReadRequestsToCleanUp = MoveTemp(ReadRequestsToCleanUp);
-
 		check(AsyncTaskEvents.IsEmpty());
-		// Kick off an async task to copy completed read requests into upload staging buffers, and issue new read requests
-		AsyncTaskEvents.Add(TGraphTask<FDistanceFieldStreamingUpdateTask>::CreateTask().ConstructAndDispatchWhenReady(UpdateParameters));
 
-		AddPass(GraphBuilder, [this, AtlasUpload, NumBrickUploads, NumIndirectionTableAdds](FRHICommandListImmediate& RHICmdList)
+		if (NewReadRequests.Num() || ReadRequestsToUpload.Num() || ReadRequestsToCleanUp.Num())
 		{
-			QUICK_SCOPE_CYCLE_COUNTER(STAT_WaitOnDistanceFieldStreamingUpdate);
-			TRACE_CPUPROFILER_EVENT_SCOPE(WaitOnDistanceFieldStreamingUpdate);
+			UpdateParameters.NewReadRequests = MoveTemp(NewReadRequests);
+			UpdateParameters.ReadRequestsToUpload = MoveTemp(ReadRequestsToUpload);
+			UpdateParameters.ReadRequestsToCleanUp = MoveTemp(ReadRequestsToCleanUp);
 
-			check(!AsyncTaskEvents.IsEmpty());
-			// Block on the async task before RDG execution of compute scatter uploads
-			FTaskGraphInterface::Get().WaitUntilTasksComplete(AsyncTaskEvents, ENamedThreads::GetRenderThread_Local());
-			AsyncTaskEvents.Empty();
+			// Kick off an async task to copy completed read requests into upload staging buffers, and issue new read requests
+			AsyncTaskEvents.Add(TGraphTask<FDistanceFieldStreamingUpdateTask>::CreateTask().ConstructAndDispatchWhenReady(UpdateParameters));
 
-			if (NumBrickUploads > 0)
+			AddPass(GraphBuilder, [this, AtlasUpload, NumBrickUploads, NumIndirectionTableAdds](FRHICommandListImmediate& RHICmdList)
 			{
-				AtlasUpload.Unlock();
-			}
+				QUICK_SCOPE_CYCLE_COUNTER(STAT_WaitOnDistanceFieldStreamingUpdate);
+				TRACE_CPUPROFILER_EVENT_SCOPE(WaitOnDistanceFieldStreamingUpdate);
 
-			if (NumIndirectionTableAdds > 0)
-			{
-				RHICmdList.Transition({
-					FRHITransitionInfo(IndirectionTable.UAV, ERHIAccess::Unknown, ERHIAccess::UAVCompute)
-					});
+				check(!AsyncTaskEvents.IsEmpty());
 
-				IndirectionTableUploadBuffer.ResourceUploadTo(RHICmdList, IndirectionTable, false);
+				// Block on the async task before RDG execution of compute scatter uploads
+				FTaskGraphInterface::Get().WaitUntilTasksComplete(AsyncTaskEvents, ENamedThreads::GetRenderThread_Local());
+				AsyncTaskEvents.Empty();
 
-				RHICmdList.Transition({
-					FRHITransitionInfo(IndirectionTable.UAV, ERHIAccess::UAVCompute, ERHIAccess::SRVMask)
-					});
-			}
-		});
+				if (NumBrickUploads > 0)
+				{
+					AtlasUpload.Unlock();
+				}
+
+				if (NumIndirectionTableAdds > 0)
+				{
+					RHICmdList.Transition({
+						FRHITransitionInfo(IndirectionTable.UAV, ERHIAccess::Unknown, ERHIAccess::UAVCompute)
+						});
+
+					IndirectionTableUploadBuffer.ResourceUploadTo(RHICmdList, IndirectionTable, false);
+
+					RHICmdList.Transition({
+						FRHITransitionInfo(IndirectionTable.UAV, ERHIAccess::UAVCompute, ERHIAccess::SRVMask)
+						});
+				}
+			});
+		}
 
 		FRDGTextureRef DistanceFieldBrickVolumeTextureRDG = GraphBuilder.RegisterExternalTexture(DistanceFieldBrickVolumeTexture, TEXT("DistanceFields.DistanceFieldBrickVolumeTexture"));
 
