Minor optimization to destruction of memstack shadows.

diff --git a/Engine/Source/Runtime/Renderer/Private/SceneRendering.cpp b/Engine/Source/Runtime/Renderer/Private/SceneRendering.cpp
--- a/Engine/Source/Runtime/Renderer/Private/SceneRendering.cpp
+++ b/Engine/Source/Runtime/Renderer/Private/SceneRendering.cpp
@@ -2933,24 +2933,12 @@ bool FSceneRenderer::DoOcclusionQueries() const
 }
 
 FSceneRenderer::~FSceneRenderer()
-{	
-	if(Scene)
+{
+	for (FProjectedShadowInfo* ProjectedShadow : MemStackProjectedShadows)
 	{
-		// Destruct the projected shadow infos.
-		for(TSparseArray<FLightSceneInfoCompact>::TConstIterator LightIt(Scene->Lights);LightIt;++LightIt)
-		{
-			if( VisibleLightInfos.IsValidIndex(LightIt.GetIndex()) )
-			{
-				FVisibleLightInfo& VisibleLightInfo = VisibleLightInfos[LightIt.GetIndex()];
-				for(int32 ShadowIndex = 0;ShadowIndex < VisibleLightInfo.MemStackProjectedShadows.Num();ShadowIndex++)
-				{
-					// FProjectedShadowInfo's in MemStackProjectedShadows were allocated on the rendering thread mem stack, 
-					// Their memory will be freed when the stack is freed with no destructor call, so invoke the destructor explicitly
-					VisibleLightInfo.MemStackProjectedShadows[ShadowIndex]->~FProjectedShadowInfo();
-				}
-				VisibleLightInfo.VirtualShadowMapId = INDEX_NONE;
-			}
-		}
+		// FProjectedShadowInfo's in MemStackProjectedShadows were allocated on the rendering thread mem stack, 
+		// Their memory will be freed when the stack is freed with no destructor call, so invoke the destructor explicitly
+		ProjectedShadow->~FProjectedShadowInfo();
 	}
 
 	// Manually release references to TRefCountPtrs that are allocated on the mem stack, which doesn't call dtors
diff --git a/Engine/Source/Runtime/Renderer/Private/ShadowSetup.cpp b/Engine/Source/Runtime/Renderer/Private/ShadowSetup.cpp
--- a/Engine/Source/Runtime/Renderer/Private/ShadowSetup.cpp
+++ b/Engine/Source/Runtime/Renderer/Private/ShadowSetup.cpp
@@ -2869,7 +2869,7 @@ void FSceneRenderer::CreatePerObjectProjectedShadow(
 					false))					// no translucent shadow
 				{
 					ProjectedShadowInfo->FadeAlphas = ResolutionFadeAlphas;
-					VisibleLightInfo.MemStackProjectedShadows.Add(ProjectedShadowInfo);
+					MemStackProjectedShadows.Add(ProjectedShadowInfo);
 
 					if (bOpaqueShadowIsVisibleThisFrame)
 					{
@@ -2911,7 +2911,7 @@ void FSceneRenderer::CreatePerObjectProjectedShadow(
 					true))					// translucent shadow
 				{
 					ProjectedShadowInfo->FadeAlphas = ResolutionFadeAlphas,
-					VisibleLightInfo.MemStackProjectedShadows.Add(ProjectedShadowInfo);
+					MemStackProjectedShadows.Add(ProjectedShadowInfo);
 
 					if (bTranslucentShadowIsVisibleThisFrame)
 					{
@@ -3647,7 +3647,7 @@ void FSceneRenderer::CreateWholeSceneProjectedShadow(
 					// Create the projected shadow info.
 					FProjectedShadowInfo* ProjectedShadowInfo = new(FMemStack::Get(), 1, 16) FProjectedShadowInfo;
 					// Add to remember-to-call-dtor list
-					VisibleLightInfo.MemStackProjectedShadows.Add(ProjectedShadowInfo);
+					MemStackProjectedShadows.Add(ProjectedShadowInfo);
 
 					// Rescale size to fit whole virtual SM but keeping aspect ratio
 					int32 VirtualSizeX = SizeX >= SizeY ? FVirtualShadowMap::VirtualMaxResolutionXY : (FVirtualShadowMap::VirtualMaxResolutionXY * SizeX) / SizeY;
@@ -3713,7 +3713,7 @@ void FSceneRenderer::CreateWholeSceneProjectedShadow(
 					ProjectedShadowInfo->CacheMode = CacheMode[CacheModeIndex];
 					ProjectedShadowInfo->FadeAlphas = FadeAlphas;
 
-					VisibleLightInfo.MemStackProjectedShadows.Add(ProjectedShadowInfo);
+					MemStackProjectedShadows.Add(ProjectedShadowInfo);
 
 					// If we have a virtual shadow map, disable nanite rendering into the regular shadow map or else we'd get double-shadowing
 					if (bNeedsVirtualShadowMap)
@@ -4589,7 +4589,7 @@ void FSceneRenderer::AddViewDependentWholeSceneShadowsForView(
 						ProjectedShadowInfo->ProjectionIndex = Index;
 						ProjectedShadowInfo->CacheMode = CacheMode[CacheModeIndex];
 
-						VisibleLightInfo.MemStackProjectedShadows.Add(ProjectedShadowInfo);
+						MemStackProjectedShadows.Add(ProjectedShadowInfo);
 						VisibleLightInfo.AllProjectedShadows.Add(ProjectedShadowInfo);
 						ShadowInfos.Add(ProjectedShadowInfo);
 
@@ -4712,7 +4712,7 @@ void FSceneRenderer::AddViewDependentWholeSceneShadowsForView(
 					// Create the projected shadow info to make sure that culling happens.
 					FProjectedShadowInfo* ProjectedShadowInfo = new(FMemStack::Get(), 1, 16) FProjectedShadowInfo;
 					// Add to remember-to-call-dtor list
-					VisibleLightInfo.MemStackProjectedShadows.Add(ProjectedShadowInfo);
+					MemStackProjectedShadows.Add(ProjectedShadowInfo);
 					ProjectedShadowInfo->SetupClipmapProjection(&LightSceneInfo, &View, VirtualShadowMapClipmap, CVarVsmUseFarShadowRules.GetValueOnRenderThread() != 0 ? MaxNonFarCascadeDistance : -1.0f);
 					VisibleLightInfo.AllProjectedShadows.Add(ProjectedShadowInfo);
 					ShadowInfosThatNeedCulling.Add(ProjectedShadowInfo);
