Optimized GPU scene to move instance updates onto task. Removes instance updates from RT critical path.

diff --git a/Engine/Source/Runtime/Renderer/Private/GPUScene.cpp b/Engine/Source/Runtime/Renderer/Private/GPUScene.cpp
--- a/Engine/Source/Runtime/Renderer/Private/GPUScene.cpp
+++ b/Engine/Source/Runtime/Renderer/Private/GPUScene.cpp
@@ -186,11 +186,9 @@ inline FPrimitiveTransforms InitPrimitiveTransforms(const FScene& Scene, const F
 	FPrimitiveTransforms Transforms;
 	Transforms.LocalToWorld = PrimitiveSceneProxy->GetLocalToWorld();
 
+	if (!Scene.VelocityData.GetComponentPreviousLocalToWorldWithoutFrameUpdate(PrimitiveSceneInfo->PrimitiveComponentId, Transforms.PreviousLocalToWorld))
 	{
-		bool bHasPrecomputedVolumetricLightmap{};
-		bool bOutputVelocity{};
-		int32 SingleCaptureIndex{};
-		Scene.GetPrimitiveUniformShaderParameters_RenderThread(PrimitiveSceneInfo, bHasPrecomputedVolumetricLightmap, Transforms.PreviousLocalToWorld, SingleCaptureIndex, bOutputVelocity);
+		Transforms.PreviousLocalToWorld = Transforms.LocalToWorld;
 	}
 
 	return Transforms;
@@ -260,6 +258,42 @@ struct FLightMapUploadInfo
  */
 struct FUploadDataSourceAdapterScenePrimitives
 {
+	class FUpdateInstanceTransformsTask
+	{
+	public:
+		FUpdateInstanceTransformsTask(const FUploadDataSourceAdapterScenePrimitives& InAdapter)
+			: Adapter(InAdapter)
+		{}
+
+		static TStatId GetStatId()
+		{
+			RETURN_QUICK_DECLARE_CYCLE_STAT(FUpdateInstanceTransformsTask, STATGROUP_TaskGraphTasks);
+		}
+
+		static ENamedThreads::Type GetDesiredThread()
+		{
+			return ENamedThreads::AnyHiPriThreadHiPriTask;
+		}
+
+		static ESubsequentsMode::Type GetSubsequentsMode()
+		{
+			return ESubsequentsMode::TrackSubsequents;
+		}
+
+		void DoTask(ENamedThreads::Type, const FGraphEventRef&)
+		{
+			TRACE_CPUPROFILER_EVENT_SCOPE(GPUScene_UpdateInstanceTransforms);
+
+			for (int32 Index = 0; Index < Adapter.NumPrimitivesToUpload(); ++Index)
+			{
+				Adapter.UpdateInstanceTransforms(Index);
+			}
+		}
+
+	private:
+		const FUploadDataSourceAdapterScenePrimitives& Adapter;
+	};
+
 	static constexpr bool bUpdateNaniteMaterialTables = true;
 
 	FUploadDataSourceAdapterScenePrimitives(FScene& InScene, uint32 InSceneFrameNumber, const TArray<int32> &InPrimitivesToUpdate)
@@ -349,11 +383,6 @@ struct FUploadDataSourceAdapterScenePrimitives
 		return false;
 	}
 
-	FORCEINLINE bool ShouldUpdateInstanceTransforms() const
-	{
-		return true;
-	}
-
 	FORCEINLINE void UpdateInstanceTransforms(int32 ItemIndex) const
 	{
 		const int32 PrimitiveID = PrimitivesToUpdate[ItemIndex];
@@ -402,9 +431,25 @@ struct FUploadDataSourceAdapterScenePrimitives
 		return false;
 	}
 
+	FORCEINLINE void UpdateInstanceTransforms()
+	{
+		if (PrimitivesToUpdate.IsEmpty())
+		{
+			return;
+		}
+
+		UpdateInstanceTransformsTask = TGraphTask<FUpdateInstanceTransformsTask>::CreateTask(nullptr, ENamedThreads::GetRenderThread_Local()).ConstructAndDispatchWhenReady(*this);
+	}
+
+	FORCEINLINE void WaitForTasks() const
+	{
+		UpdateInstanceTransformsTask->Wait(ENamedThreads::GetRenderThread_Local());
+	}
+
 	FScene& Scene;
 	const uint32 SceneFrameNumber;
 	TArray<int, SceneRenderingAllocator> PrimitivesToUpdate;
+	FGraphEventRef UpdateInstanceTransformsTask;
 };
 
 void FGPUScene::SetEnabled(ERHIFeatureLevel::Type InFeatureLevel)
@@ -476,7 +521,7 @@ void FGPUScene::UpdateInternal(FRDGBuilder& GraphBuilder, FScene& Scene)
 	// Store in GPU-scene to enable validation that update has been carried out.
 	SceneFrameNumber = Scene.GetFrameNumber();
 
-	FUploadDataSourceAdapterScenePrimitives Adapter(Scene, SceneFrameNumber, PrimitivesToUpdate);
+	FUploadDataSourceAdapterScenePrimitives& Adapter = *GraphBuilder.AllocObject<FUploadDataSourceAdapterScenePrimitives>(Scene, SceneFrameNumber, PrimitivesToUpdate);
 	FGPUSceneBufferState BufferState = UpdateBufferState(GraphBuilder, &Scene, Adapter);
 
 	// The adapter copies the IDs of primitives to update such that any that are (incorrectly) marked for update after are not lost.
@@ -488,7 +533,7 @@ void FGPUScene::UpdateInternal(FRDGBuilder& GraphBuilder, FScene& Scene)
 
 
 	AddPass(GraphBuilder, RDG_EVENT_NAME("GPUSceneUpdate"), 
-		[this, &Scene, Adapter = MoveTemp(Adapter), BufferState = MoveTemp(BufferState)](FRHICommandListImmediate& RHICmdList)
+		[this, &Scene, &Adapter = Adapter, BufferState = MoveTemp(BufferState)](FRHICommandListImmediate& RHICmdList)
 	{
 		SCOPED_NAMED_EVENT(STAT_UpdateGPUScene, FColor::Green);
 		CSV_SCOPED_TIMING_STAT_EXCLUSIVE(UpdateGPUScene);
@@ -571,9 +616,9 @@ namespace
 };
 
 template<typename FUploadDataSourceAdapter>
-FGPUSceneBufferState FGPUScene::UpdateBufferState(FRDGBuilder& GraphBuilder, FScene* Scene, const FUploadDataSourceAdapter& UploadDataSourceAdapter)
+FGPUSceneBufferState FGPUScene::UpdateBufferState(FRDGBuilder& GraphBuilder, FScene* Scene, FUploadDataSourceAdapter& UploadDataSourceAdapter)
 {
-	UpdatePrimitiveInstances(Scene, UploadDataSourceAdapter);
+	UploadDataSourceAdapter.UpdateInstanceTransforms();
 
 	FGPUSceneBufferState BufferState;
 	ensure(bInBeginEndBlock);
@@ -625,39 +670,6 @@ FGPUSceneBufferState FGPUScene::UpdateBufferState(FRDGBuilder& GraphBuilder, FSc
 	return BufferState;
 }
 
-template <typename FUploadDataSourceAdapter>
-void FGPUScene::UpdatePrimitiveInstances(FScene* Scene, const FUploadDataSourceAdapter& UploadDataSourceAdapter)
-{
-	const int32 NumPrimitiveDataUploads = UploadDataSourceAdapter.NumPrimitivesToUpload();
-
-	if (!NumPrimitiveDataUploads || !UploadDataSourceAdapter.ShouldUpdateInstanceTransforms())
-	{
-		return;
-	}
-
-	// Note: we iterate over the primitives, whether they have instances or not (which is a bit wasteful) but this is the way we currently get to the instance data.
-	// GPUCULL_TODO: move instance data ownership to GPU-scene such that it can be put in a compact list or something, and be tracked independent of primitives?
-
-	TRACE_CPUPROFILER_EVENT_SCOPE(GPUScene_UpdateInstanceTransforms);
-
-	FParallelUpdateRanges ParallelRanges;
-
-	const bool bExecuteInParallel = GGPUSceneParallelUpdate != 0 && FApp::ShouldUseThreadingForPerformance();
-	const int32 RangeCount = PartitionUpdateRanges(ParallelRanges, NumPrimitiveDataUploads, bExecuteInParallel);
-	const uint32 InstanceDataNumArrays = FInstanceSceneShaderData::InstanceDataStrideInFloat4s;
-
-	// Upload any out of date instance slots.
-	ParallelFor(RangeCount,
-		[&UploadDataSourceAdapter, this, &ParallelRanges, RangeCount, InstanceDataNumArrays](int32 RangeIndex)
-	{
-		for (int32 ItemIndex = ParallelRanges.Range[RangeIndex].ItemStart; ItemIndex < ParallelRanges.Range[RangeIndex].ItemStart + ParallelRanges.Range[RangeIndex].ItemCount; ++ItemIndex)
-		{
-			UploadDataSourceAdapter.UpdateInstanceTransforms(ItemIndex);
-		}
-
-	}, RangeCount == 1);
-}
-
 template<typename FUploadDataSourceAdapter>
 void FGPUScene::UploadGeneral(FRHICommandListImmediate& RHICmdList, FScene *Scene, const FUploadDataSourceAdapter& UploadDataSourceAdapter, const FGPUSceneBufferState& BufferState)
 {
@@ -988,6 +1000,8 @@ void FGPUScene::UploadGeneral(FRHICommandListImmediate& RHICmdList, FScene *Scen
 
 				if (NumPrimitiveDataUploads > 0)
 				{
+					UploadDataSourceAdapter.WaitForTasks();
+
 					// Note: we iterate over the primitives, whether they have instances or not (which is a bit wasteful) but this is the way we currently get to the instance data.
 					// GPUCULL_TODO: move instance data ownership to GPU-scene such that it can be put in a compact list or something, and be tracked independent of primitives?
 					RangeCount = PartitionUpdateRanges(ParallelRanges, NumPrimitiveDataUploads, bExecuteInParallel);
@@ -1166,12 +1180,12 @@ struct FUploadDataSourceAdapterDynamicPrimitives
 		return false;
 	}
 
-	FORCEINLINE bool ShouldUpdateInstanceTransforms() const
+	FORCEINLINE void UpdateInstanceTransforms() const
 	{
-		return false;
+		// Nothing to do.
 	}
 
-	FORCEINLINE void UpdateInstanceTransforms(int32 ItemIndex) const
+	FORCEINLINE void WaitForTasks() const
 	{
 		// Nothing to do.
 	}
diff --git a/Engine/Source/Runtime/Renderer/Private/GPUScene.h b/Engine/Source/Runtime/Renderer/Private/GPUScene.h
--- a/Engine/Source/Runtime/Renderer/Private/GPUScene.h
+++ b/Engine/Source/Runtime/Renderer/Private/GPUScene.h
@@ -264,7 +264,7 @@ private:
 	ERHIFeatureLevel::Type FeatureLevel;
 
 	template<typename FUploadDataSourceAdapter>
-	FGPUSceneBufferState UpdateBufferState(FRDGBuilder& GraphBuilder, FScene* Scene, const FUploadDataSourceAdapter& UploadDataSourceAdapter);
+	FGPUSceneBufferState UpdateBufferState(FRDGBuilder& GraphBuilder, FScene* Scene, FUploadDataSourceAdapter& UploadDataSourceAdapter);
 
 	template <typename FUploadDataSourceAdapter>
 	void UpdatePrimitiveInstances(FScene* Scene, const FUploadDataSourceAdapter& UploadDataSourceAdapter);
diff --git a/Engine/Source/Runtime/Renderer/Private/ScenePrivate.h b/Engine/Source/Runtime/Renderer/Private/ScenePrivate.h
--- a/Engine/Source/Runtime/Renderer/Private/ScenePrivate.h
+++ b/Engine/Source/Runtime/Renderer/Private/ScenePrivate.h
@@ -2338,6 +2338,20 @@ public:
 		return false;
 	}
 
+	inline bool GetComponentPreviousLocalToWorldWithoutFrameUpdate(FPrimitiveComponentId PrimitiveComponentId, FMatrix& OutPreviousLocalToWorld) const
+	{
+		const FComponentVelocityData* VelocityData = ComponentData.Find(PrimitiveComponentId);
+
+		if (VelocityData)
+		{
+			check(VelocityData->bPreviousLocalToWorldValid);
+			OutPreviousLocalToWorld = VelocityData->PreviousLocalToWorld;
+			return true;
+		}
+
+		return false;
+	}
+
 	/** 
 	 * Updates a primitives current LocalToWorld state.
 	 */
