Fix for distance field streaming pass / task filtering.

diff --git a/Engine/Source/Runtime/Renderer/Private/DistanceFieldStreaming.cpp b/Engine/Source/Runtime/Renderer/Private/DistanceFieldStreaming.cpp
--- a/Engine/Source/Runtime/Renderer/Private/DistanceFieldStreaming.cpp
+++ b/Engine/Source/Runtime/Renderer/Private/DistanceFieldStreaming.cpp
@@ -907,17 +907,21 @@ void FDistanceFieldSceneData::UpdateDistanceFieldAtlas(
 
 			// Kick off an async task to copy completed read requests into upload staging buffers, and issue new read requests
 			AsyncTaskEvents.Add(TGraphTask<FDistanceFieldStreamingUpdateTask>::CreateTask().ConstructAndDispatchWhenReady(UpdateParameters));
+		}
 
-			AddPass(GraphBuilder, [this, AtlasUpload, NumBrickUploads, NumIndirectionTableAdds](FRHICommandListImmediate& RHICmdList)
+		if (AsyncTaskEvents.Num() || NumBrickUploads > 0 || NumIndirectionTableAdds > 0)
+		{
+			AddPass(GraphBuilder, RDG_EVENT_NAME("WaitOnDistanceFieldStreamingUpdate"), [this, AtlasUpload, NumBrickUploads, NumIndirectionTableAdds](FRHICommandListImmediate& RHICmdList)
 			{
 				QUICK_SCOPE_CYCLE_COUNTER(STAT_WaitOnDistanceFieldStreamingUpdate);
 				TRACE_CPUPROFILER_EVENT_SCOPE(WaitOnDistanceFieldStreamingUpdate);
 
-				check(!AsyncTaskEvents.IsEmpty());
-
-				// Block on the async task before RDG execution of compute scatter uploads
-				FTaskGraphInterface::Get().WaitUntilTasksComplete(AsyncTaskEvents, ENamedThreads::GetRenderThread_Local());
-				AsyncTaskEvents.Empty();
+				if (!AsyncTaskEvents.IsEmpty())
+				{
+					// Block on the async task before RDG execution of compute scatter uploads
+					FTaskGraphInterface::Get().WaitUntilTasksComplete(AsyncTaskEvents, ENamedThreads::GetRenderThread_Local());
+					AsyncTaskEvents.Empty();
+				}
 
 				if (NumBrickUploads > 0)
 				{
