Revert inset shadow logic to match UE4 (leftover from Reverb)

diff --git a/Engine/Source/Runtime/Renderer/Private/SceneCore.cpp b/Engine/Source/Runtime/Renderer/Private/SceneCore.cpp
--- a/Engine/Source/Runtime/Renderer/Private/SceneCore.cpp
+++ b/Engine/Source/Runtime/Renderer/Private/SceneCore.cpp
@@ -100,8 +100,7 @@ void FLightPrimitiveInteraction::Create(FLightSceneInfo* LightSceneInfo,FPrimiti
 		const bool bTranslucentObjectShadow = LightSceneInfo->Proxy->CastsTranslucentShadows() && PrimitiveSceneInfo->Proxy->CastsVolumetricTranslucentShadow();
 		const bool bInsetObjectShadow = 
 			// Currently only supporting inset shadows on directional lights, but could be made to work with any whole scene shadows
-			(LightSceneInfo->Proxy->GetLightType() == LightType_Directional
-				|| LightSceneInfo->Proxy->GetLightType() == LightType_Spot)
+			LightSceneInfo->Proxy->GetLightType() == LightType_Directional
 			&& PrimitiveSceneInfo->Proxy->CastsInsetShadow();
 
 		// Movable directional lights determine shadow relevance dynamically based on the view and CSM settings. Interactions are only required for per-object cases.
