Increase default r.AOGlobalDistanceField.OccupancyRatio to .5 to fix artifacts in cave scenes, probably caused by [REDACTED_PERSON] not having as large of a negative region (by design)

diff --git a/Engine/Source/Runtime/Renderer/Private/GlobalDistanceField.cpp b/Engine/Source/Runtime/Renderer/Private/GlobalDistanceField.cpp
--- a/Engine/Source/Runtime/Renderer/Private/GlobalDistanceField.cpp
+++ b/Engine/Source/Runtime/Renderer/Private/GlobalDistanceField.cpp
@@ -25,7 +25,7 @@ FAutoConsoleVariableRef CVarAOGlobalDistanceField(
 	ECVF_Scalability | ECVF_RenderThreadSafe
 	);
 
-float GGlobalDistanceFieldOccupancyRatio = 0.3f;
+float GGlobalDistanceFieldOccupancyRatio = 0.5f;
 FAutoConsoleVariableRef CVarAOGlobalDistanceFieldOccupancyRatio(
 	TEXT("r.AOGlobalDistanceField.OccupancyRatio"),
 	GGlobalDistanceFieldOccupancyRatio,
