Fix a bug of triple buffer of mobile eye adaptation.

diff --git a/Engine/Source/Runtime/Renderer/Private/PostProcess/PostProcessMobile.cpp b/Engine/Source/Runtime/Renderer/Private/PostProcess/PostProcessMobile.cpp
--- a/Engine/Source/Runtime/Renderer/Private/PostProcess/PostProcessMobile.cpp
+++ b/Engine/Source/Runtime/Renderer/Private/PostProcess/PostProcessMobile.cpp
@@ -1947,7 +1947,7 @@ void AddMobileEyeAdaptationPass(FRDGBuilder& GraphBuilder, const FViewInfo& View
 	// Get the custom 1x1 target used to store exposure value and Toggle the two render targets used to store new and old.
 	View.SwapEyeAdaptationBuffers(GraphBuilder);
 
-	FRDGBufferRef EyeAdaptationBuffer = GraphBuilder.RegisterExternalBuffer(View.GetLastEyeAdaptationBuffer(GraphBuilder), ERDGBufferFlags::MultiFrame);
+	FRDGBufferRef EyeAdaptationBuffer = Inputs.EyeAdaptationBuffer;
 	FRDGBufferSRVRef EyeAdaptationBufferSRV = GraphBuilder.CreateSRV(EyeAdaptationBuffer, PF_A32B32G32R32F);
 
 	FRDGBufferRef OutputBuffer = GraphBuilder.RegisterExternalBuffer(View.GetEyeAdaptationBuffer(GraphBuilder), ERDGBufferFlags::MultiFrame);
diff --git a/Engine/Source/Runtime/Renderer/Private/PostProcess/PostProcessMobile.h b/Engine/Source/Runtime/Renderer/Private/PostProcess/PostProcessMobile.h
--- a/Engine/Source/Runtime/Renderer/Private/PostProcess/PostProcessMobile.h
+++ b/Engine/Source/Runtime/Renderer/Private/PostProcess/PostProcessMobile.h
@@ -168,6 +168,7 @@ FMobileEyeAdaptationSetupOutputs AddMobileEyeAdaptationSetupPass(FRDGBuilder& Gr
 struct FMobileEyeAdaptationInputs
 {
 	FRDGBufferSRVRef EyeAdaptationSetupSRV;
+	FRDGBufferRef EyeAdaptationBuffer;
 	bool bUseBasicEyeAdaptation;
 	bool bUseHistogramEyeAdaptation;
 };
diff --git a/Engine/Source/Runtime/Renderer/Private/PostProcess/PostProcessing.cpp b/Engine/Source/Runtime/Renderer/Private/PostProcess/PostProcessing.cpp
--- a/Engine/Source/Runtime/Renderer/Private/PostProcess/PostProcessing.cpp
+++ b/Engine/Source/Runtime/Renderer/Private/PostProcess/PostProcessing.cpp
@@ -1378,6 +1378,7 @@ void AddMobilePostProcessingPasses(FRDGBuilder& GraphBuilder, FScene* Scene, con
 
 	// Default the new eye adaptation to the last one in case it's not generated this frame.
 	const FEyeAdaptationParameters EyeAdaptationParameters = GetEyeAdaptationParameters(View, ERHIFeatureLevel::ES3_1);
+	FRDGBufferRef LastEyeAdaptationBuffer = GetEyeAdaptationBuffer(GraphBuilder, View);
 
 	enum class EPass : uint32
 	{
@@ -1618,11 +1619,6 @@ void AddMobilePostProcessingPasses(FRDGBuilder& GraphBuilder, FScene* Scene, con
 			BloomSetupInputs.SunShaftAndDof = PostProcessSunShaftAndDof;
 
 			BloomSetupOutputs = AddMobileBloomSetupPass(GraphBuilder, View, EyeAdaptationParameters, BloomSetupInputs);
-
-			if (bHasEyeAdaptationPass && View.ViewState && !View.bStatePrevViewInfoIsReadOnly)
-			{
-				GraphBuilder.QueueTextureExtraction(BloomSetupOutputs.EyeAdaptation.Texture, &View.ViewState->PrevFrameViewInfo.MobileBloomSetup_EyeAdaptation);
-			}
 		}
 
 		if (PassSequence.IsEnabled(EPass::DepthOfField))
@@ -1820,8 +1816,14 @@ void AddMobilePostProcessingPasses(FRDGBuilder& GraphBuilder, FScene* Scene, con
 			EyeAdaptationInputs.bUseBasicEyeAdaptation = bUseBasicEyeAdaptation;
 			EyeAdaptationInputs.bUseHistogramEyeAdaptation = bUseHistogramEyeAdaptation;
 			EyeAdaptationInputs.EyeAdaptationSetupSRV = EyeAdaptationSetupOutputs.EyeAdaptationSetupSRV;
+			EyeAdaptationInputs.EyeAdaptationBuffer = LastEyeAdaptationBuffer;
 
 			AddMobileEyeAdaptationPass(GraphBuilder, View, EyeAdaptationParameters, EyeAdaptationInputs);
+
+			if ((bUseBasicEyeAdaptation || bUseHistogramEyeAdaptation) && View.ViewState && !View.bStatePrevViewInfoIsReadOnly)
+			{
+				GraphBuilder.QueueTextureExtraction(BloomSetupOutputs.EyeAdaptation.Texture, &View.ViewState->PrevFrameViewInfo.MobileBloomSetup_EyeAdaptation);
+			}
 		}
 
 		if (PassSequence.IsEnabled(EPass::SunMerge))
@@ -1985,7 +1987,7 @@ void AddMobilePostProcessingPasses(FRDGBuilder& GraphBuilder, FScene* Scene, con
 		TonemapperInputs.bOutputInHDR = bHDRTonemapperOutput;
 		TonemapperInputs.bGammaOnly = bDoGammaOnly;
 		TonemapperInputs.bMetalMSAAHDRDecode = bMetalMSAAHDRDecode;
-		TonemapperInputs.EyeAdaptationBuffer = bUseEyeAdaptation ? TryRegisterExternalBuffer(GraphBuilder, View.GetLastEyeAdaptationBuffer(GraphBuilder), ERDGBufferFlags::MultiFrame) : nullptr;
+		TonemapperInputs.EyeAdaptationBuffer = bUseEyeAdaptation ? LastEyeAdaptationBuffer : nullptr;
 
 		SceneColor = AddTonemapPass(GraphBuilder, View, TonemapperInputs);
 
