Call WaitForRayTracingScene() before lighting if an earlier wait point was not forced by Lumen

diff --git a/Engine/Source/Runtime/Renderer/Private/DeferredShadingRenderer.cpp b/Engine/Source/Runtime/Renderer/Private/DeferredShadingRenderer.cpp
--- a/Engine/Source/Runtime/Renderer/Private/DeferredShadingRenderer.cpp
+++ b/Engine/Source/Runtime/Renderer/Private/DeferredShadingRenderer.cpp
@@ -2368,6 +2368,7 @@ void FDeferredShadingSceneRenderer::Render(FRDGBuilder& GraphBuilder)
 
 #if RHI_RAYTRACING
 	ERayTracingWorldUpdatesDispatchPoint RayTracingWorldUpdatesDispatchPoint = GetRayTracingWorldUpdatesDispatchPoint(bOcclusionBeforeBasePass, Lumen::UseHardwareRayTracedShadows(Views[0]));
+	bool bRayTracingSceneReady = false;
 #endif
 
 	if (bOcclusionBeforeBasePass)
@@ -2448,6 +2449,7 @@ void FDeferredShadingSceneRenderer::Render(FRDGBuilder& GraphBuilder)
 		if (Lumen::UseHardwareRayTracedShadows(Views[0]))
 		{
 			WaitForRayTracingScene(GraphBuilder);
+			bRayTracingSceneReady = true;
 		}
 #endif // RHI_RAYTRACING
 
@@ -2631,6 +2633,7 @@ void FDeferredShadingSceneRenderer::Render(FRDGBuilder& GraphBuilder)
 		if (Lumen::UseHardwareRayTracedShadows(Views[0]))
 		{
 			WaitForRayTracingScene(GraphBuilder);
+			bRayTracingSceneReady = true;
 		}
 #endif // RHI_RAYTRACING
 
@@ -2740,11 +2743,11 @@ void FDeferredShadingSceneRenderer::Render(FRDGBuilder& GraphBuilder)
 	}
 
 #if RHI_RAYTRACING
-	// If Lumen is not using HWRT shadows, we can wait until here: before Lumen diffuse indirect
-	// Also catch the case of path tracer or RT debug output
-	if (!Lumen::UseHardwareRayTracedShadows(Views[0]) || bHasRayTracedOverlay)
+	// If Lumen did not force an earlier ray tracing scene sync, we must wait for it here.
+	if (!bRayTracingSceneReady)
 	{
 		WaitForRayTracingScene(GraphBuilder);
+		bRayTracingSceneReady = true;
 	}
 #endif // RHI_RAYTRACING
 
