Add adaptive voxel allocation to avoid overallocation.
This CL adds GPU->CPU readback back to adpat voxel size based on previous frame allocated page count. If we run out of page, we increase voxel size, so that next allocations fit into the memory constraint. This approach is not perfect as at least one /two frames will have glitches (over allocation not-compensated) and continuous change, can create oscillation which are not so easy to counter balance due to latency.

Manual port of CL 16346703

diff --git a/Engine/Shaders/Private/HairStrands/HairStrandsVoxelPageAllocation.usf b/Engine/Shaders/Private/HairStrands/HairStrandsVoxelPageAllocation.usf
--- a/Engine/Shaders/Private/HairStrands/HairStrandsVoxelPageAllocation.usf
+++ b/Engine/Shaders/Private/HairStrands/HairStrandsVoxelPageAllocation.usf
@@ -656,6 +656,7 @@ uint IndirectGroupSize;
 
 RWBuffer<uint> OutPageIndexGlobalCounter;
 RWBuffer<uint> OutIndirectArgsBuffer;
+RWBuffer<uint> OutTotalRequestedPageAllocationBuffer; 
 
 [numthreads(1, 1, 1)]
 void AddIndirectBufferCS(uint GroupIndex : SV_GroupIndex, uint3 DispatchThreadId : SV_DispatchThreadID)
@@ -676,6 +677,9 @@ void AddIndirectBufferCS(uint GroupIndex : SV_GroupIndex, uint3 DispatchThreadId
 
 	// Copy current total to previous total
 	OutPageIndexGlobalCounter[1] = TotalAllocatedPageCount;
+
+	// The value will be written several times, but only the value written during the last macro group will give the final page count allocated/requested
+	OutTotalRequestedPageAllocationBuffer[0] = TotalAllocatedPageCount;
 }
 
 #endif
diff --git a/Engine/Source/Runtime/Renderer/Private/HairStrands/HairStrandsData.h b/Engine/Source/Runtime/Renderer/Private/HairStrands/HairStrandsData.h
--- a/Engine/Source/Runtime/Renderer/Private/HairStrands/HairStrandsData.h
+++ b/Engine/Source/Runtime/Renderer/Private/HairStrands/HairStrandsData.h
@@ -364,6 +364,22 @@ struct FHairStrandsViewData
 	FHairStrandsDebugData DebugData;
 };
 
+// View State data (i.e., persistent accross fram)
+struct FHairStrandsViewStateData
+{
+	FRHIGPUBufferReadback* GetBuffer() const { return VoxelPageAllocationCountReadback; }
+	bool IsReady() const { return VoxelPageAllocationCountReadback->IsReady(); }
+	bool IsInit() const { return VoxelPageAllocationCountReadback != nullptr; }
+	void Init();
+	void Release();
+
+	float VoxelWorldSize = 0; // Voxel size used during the last frame allocation
+	uint32 AllocatedPageCount = 0; // Number of voxels allocated last frame
+
+	// Buffer used for reading back the number of voxels allocated on the GPU
+	FRHIGPUBufferReadback* VoxelPageAllocationCountReadback = nullptr;
+};
+
 namespace HairStrands
 {
 	TRDGUniformBufferRef<FHairStrandsViewUniformParameters> CreateDefaultHairStrandsViewUniformBuffer(FRDGBuilder& GraphBuilder, FViewInfo& View);
diff --git a/Engine/Source/Runtime/Renderer/Private/HairStrands/HairStrandsInterface.cpp b/Engine/Source/Runtime/Renderer/Private/HairStrands/HairStrandsInterface.cpp
--- a/Engine/Source/Runtime/Renderer/Private/HairStrands/HairStrandsInterface.cpp
+++ b/Engine/Source/Runtime/Renderer/Private/HairStrands/HairStrandsInterface.cpp
@@ -395,7 +395,6 @@ void RunHairStrandsBookmark(EHairStrandsBookmark Bookmark, FHairStrandsBookmarkP
 	}
 }
 
-bool IsHairStrandsClusterCullingUseHzb();
 FHairStrandsBookmarkParameters CreateHairStrandsBookmarkParameters(FScene* Scene, FViewInfo& View)
 {
 	FHairStrandsBookmarkParameters Out;
@@ -423,7 +422,7 @@ FHairStrandsBookmarkParameters CreateHairStrandsBookmarkParameters(FScene* Scene
 	{
 		GHairStrandsParameterFunction(Out);
 	}
-	Out.bHzbRequest = Out.bHasElements && Out.bStrandsGeometryEnabled && IsHairStrandsClusterCullingUseHzb();
+	Out.bHzbRequest = false; // Out.bHasElements&& Out.bStrandsGeometryEnabled;
 
 	return Out;
 }
diff --git a/Engine/Source/Runtime/Renderer/Private/HairStrands/HairStrandsRendering.cpp b/Engine/Source/Runtime/Renderer/Private/HairStrands/HairStrandsRendering.cpp
--- a/Engine/Source/Runtime/Renderer/Private/HairStrands/HairStrandsRendering.cpp
+++ b/Engine/Source/Runtime/Renderer/Private/HairStrands/HairStrandsRendering.cpp
@@ -82,6 +82,7 @@ IMPLEMENT_GLOBAL_SHADER_PARAMETER_STRUCT(FHairStrandsViewUniformParameters, "Hai
 
 void AddServiceLocalQueuePass(FRDGBuilder& GraphBuilder);
 bool GetHairStrandsSkyLightingDebugEnable();
+bool IsHairStrandsAdaptiveVoxelAllocationEnable();
 
 void RenderHairPrePass(
 	FRDGBuilder& GraphBuilder,
@@ -98,6 +99,23 @@ void RenderHairPrePass(
 
 		const ERHIFeatureLevel::Type FeatureLevel = Scene->GetFeatureLevel();
 
+		// Allocate voxel page allocation readback buffers
+		const bool bAdaptiveAllocationEnable = IsHairStrandsAdaptiveVoxelAllocationEnable();
+		if (View.ViewState)
+		{
+			const bool bIsInit = View.ViewState->HairStrandsViewStateData.IsInit();
+			// Init resources if the adaptive allocation is enabled, but the resources are not initialized yet.
+			if (bAdaptiveAllocationEnable && !bIsInit)
+			{
+				View.ViewState->HairStrandsViewStateData.Init();
+			}
+			// Release resources if the adaptive allocation is disabled, but the resources are initialized.
+			else if (!bAdaptiveAllocationEnable && bIsInit)
+			{
+				View.ViewState->HairStrandsViewStateData.Release();
+			}
+		}
+
 		//SCOPED_GPU_STAT(RHICmdList, HairRendering);
 		CreateHairStrandsMacroGroups(GraphBuilder, Scene, View);
 		AddServiceLocalQueuePass(GraphBuilder);
@@ -156,6 +174,27 @@ void RenderHairBasePass(
 	}
 }
 
+void FHairStrandsViewStateData::Init()
+{
+	VoxelWorldSize = 0;
+	AllocatedPageCount = 0;
+	if (VoxelPageAllocationCountReadback == nullptr)
+	{
+		VoxelPageAllocationCountReadback = new FRHIGPUBufferReadback(TEXT("Voxel page allocation readback"));
+	}
+}
+
+void FHairStrandsViewStateData::Release()
+{
+	VoxelWorldSize = 0;
+	AllocatedPageCount = 0;
+	if (VoxelPageAllocationCountReadback)
+	{
+		delete VoxelPageAllocationCountReadback;
+		VoxelPageAllocationCountReadback = nullptr;
+	}
+}
+
 namespace HairStrands
 {
 
diff --git a/Engine/Source/Runtime/Renderer/Private/HairStrands/HairStrandsVoxelization.cpp b/Engine/Source/Runtime/Renderer/Private/HairStrands/HairStrandsVoxelization.cpp
--- a/Engine/Source/Runtime/Renderer/Private/HairStrands/HairStrandsVoxelization.cpp
+++ b/Engine/Source/Runtime/Renderer/Private/HairStrands/HairStrandsVoxelization.cpp
@@ -110,6 +110,18 @@ static float GHairVirtualVoxelRaytracing_SkyOcclusionThreshold = 1;
 static FAutoConsoleVariableRef CVarHairVirtualVoxelRaytracing_ShadowOcclusionThreshold(TEXT("r.RayTracing.Shadows.HairOcclusionThreshold"), GHairVirtualVoxelRaytracing_ShadowOcclusionThreshold, TEXT("Define the number of hair that need to be crossed, before casting occlusion (default = 1)"), ECVF_RenderThreadSafe);
 static FAutoConsoleVariableRef CVarHairVirtualVoxelRaytracing_SkyOcclusionThreshold(TEXT("r.RayTracing.Sky.HairOcclusionThreshold"), GHairVirtualVoxelRaytracing_SkyOcclusionThreshold, TEXT("Define the number of hair that need to be crossed, before casting occlusion (default = 1)"), ECVF_RenderThreadSafe);
 
+static int32 GHairVirtualVoxelAdaptive_Enable = 1;
+static float GHairVirtualVoxelAdaptive_CorrectionSpeed = 0.1f;
+static float GHairVirtualVoxelAdaptive_CorrectionThreshold = 0.90f;
+static FAutoConsoleVariableRef CVarHairVirtualVoxelAdaptive_Enable(TEXT("r.HairStrands.Voxelization.Virtual.Adaptive"), GHairVirtualVoxelAdaptive_Enable, TEXT("Enable adaptive voxel allocation (default = 1)"), ECVF_RenderThreadSafe);
+static FAutoConsoleVariableRef CVarHairVirtualVoxelAdaptive_CorrectionSpeed(TEXT("r.HairStrands.Voxelization.Virtual.Adaptive.CorrectionSpeed"), GHairVirtualVoxelAdaptive_CorrectionSpeed, TEXT("Define the speed at which allocation adaption runs (value in 0..1, default = 0.25). A higher number means faster adaptation, but with a risk of oscillation i.e. over and under allocation"), ECVF_RenderThreadSafe);
+static FAutoConsoleVariableRef CVarHairVirtualVoxelAdaptive_CorrectionThreshold(TEXT("r.HairStrands.Voxelization.Virtual.Adaptive.CorrectionThreshold"), GHairVirtualVoxelAdaptive_CorrectionThreshold, TEXT("Define the allocation margin to limit over allocation (value in 0..1, default = 0.95)"), ECVF_RenderThreadSafe);
+
+bool IsHairStrandsAdaptiveVoxelAllocationEnable()
+{
+	return GHairVirtualVoxelAdaptive_Enable > 0;
+}
+
 bool IsHairStrandsVoxelizationEnable()
 {
 	return GHairVoxelizationEnable > 0;
@@ -401,7 +413,8 @@ class FVoxelAddIndirectBufferCS : public FGlobalShader
 		SHADER_PARAMETER(uint32, PageResolution)
 		SHADER_PARAMETER_RDG_BUFFER_UAV(RWBuffer, OutPageIndexGlobalCounter)
 		SHADER_PARAMETER_RDG_BUFFER_UAV(RWBuffer, OutIndirectArgsBuffer)
-		END_SHADER_PARAMETER_STRUCT()
+		SHADER_PARAMETER_RDG_BUFFER_UAV(RWBuffer, OutTotalRequestedPageAllocationBuffer)
+	END_SHADER_PARAMETER_STRUCT()
 
 public:
 	static bool ShouldCompilePermutation(const FGlobalShaderPermutationParameters& Parameters) { return IsHairStrandsSupported(EHairStrandsShaderType::Strands, Parameters.Platform); }
@@ -489,7 +502,8 @@ static void AddAllocateVoxelPagesPass(
 	FRDGBufferRef& OutNodeDescBuffer,
 	FRDGBufferRef& OutIndirectArgsBuffer,
 	FRDGBufferRef& OutPageIndexGlobalCounter, 
-	FRDGBufferRef& OutVoxelizationViewInfoBuffer)
+	FRDGBufferRef& OutVoxelizationViewInfoBuffer,
+	FRDGBufferRef& OutTotalRequestedPageAllocationBuffer)
 {
 	const uint32 GroupSize = 32;
 	const bool bIsGPUDriven = GHairVirtualVoxelGPUDriven > 0;
@@ -543,7 +557,9 @@ static void AddAllocateVoxelPagesPass(
 	// Over-allocation (upper bound)
 	if (bIsGPUDriven)
 	{
-		OutTotalPageIndexCount = GHairVirtualVoxelGPUDrivenMaxPageIndexRes * GHairVirtualVoxelGPUDrivenMaxPageIndexRes * GHairVirtualVoxelGPUDrivenMaxPageIndexRes;
+		// Use the max between the estimated size on CPU and a pseudo-conservative side driven by settings. The CPU estimation is no necessarely correct as the bounds are not reliable on skel. mesh.
+		const uint32 MaxPageIndexCount = GHairVirtualVoxelGPUDrivenMaxPageIndexRes * GHairVirtualVoxelGPUDrivenMaxPageIndexRes * GHairVirtualVoxelGPUDrivenMaxPageIndexRes;
+		OutTotalPageIndexCount = FMath::Max(MaxPageIndexCount, OutTotalPageIndexCount);
 	}
 	check(OutTotalPageIndexCount > 0);
 	
@@ -557,6 +573,8 @@ static void AddAllocateVoxelPagesPass(
 	const uint32 TotalPageCount = PageCountResolution.X * PageCountResolution.Y * PageCountResolution.Z;
 	FRDGBufferRef PageToPageIndexBuffer = GraphBuilder.CreateBuffer(FRDGBufferDesc::CreateBufferDesc(sizeof(uint32), TotalPageCount), TEXT("Hair.PageToPageIndexBuffer"));
 
+	FRDGBufferRef ReadBackBuffer = GraphBuilder.CreateBuffer(FRDGBufferDesc::CreateBufferDesc(sizeof(uint32), 1), TEXT("ReadBackAllocations"));
+
 	FRDGBufferUAVRef PageIndexBufferUAV = GraphBuilder.CreateUAV(PageIndexBuffer, PF_R32_UINT);
 	FRDGBufferUAVRef PageIndexOccupancyBufferUAV = GraphBuilder.CreateUAV(PageIndexOccupancyBuffer, PF_R32G32_UINT);
 	FRDGBufferUAVRef PageIndexGlobalCounterUAV = GraphBuilder.CreateUAV(PageIndexGlobalCounter, PF_R32_UINT);
@@ -567,14 +585,20 @@ static void AddAllocateVoxelPagesPass(
 	FRDGBufferRef PageIndexResolutionBuffer = GraphBuilder.CreateBuffer(FRDGBufferDesc::CreateBufferDesc(MacroGroupCount * 4 * sizeof(uint32), OutTotalPageIndexCount), TEXT("Hair.PageIndexResolutionBuffer"));
 	FRDGBufferRef PageIndexAllocationIndirectBufferArgs = GraphBuilder.CreateBuffer(FRDGBufferDesc::CreateIndirectDesc<FRHIDispatchIndirectParameters>(MacroGroupCount), TEXT("Hair.PageIndexAllocationIndirectBufferArgs"));
 
+	// Store the total requested page allocation (for feedback purpose)
+	FRDGBufferDesc TotalRequestDesc = FRDGBufferDesc::CreateBufferDesc(sizeof(uint32), 1);
+	TotalRequestDesc.Usage |= BUF_SourceCopy;
+	OutTotalRequestedPageAllocationBuffer = GraphBuilder.CreateBuffer(TotalRequestDesc, TEXT("TotalRequestedPageAllocationBuffer"));
+	FRDGBufferUAVRef TotalRequestedPageAllocationBufferUAV = GraphBuilder.CreateUAV(OutTotalRequestedPageAllocationBuffer, PF_R32_UINT);
+	AddClearUAVPass(GraphBuilder, TotalRequestedPageAllocationBufferUAV, 0u);
+	
 	AddClearUAVPass(GraphBuilder, PageIndexBufferUAV, 0u);
 	AddClearUAVPass(GraphBuilder, PageIndexOccupancyBufferUAV, 0u);
 	AddClearUAVPass(GraphBuilder, PageIndexGlobalCounterUAV, 0u);
 
 	// Allocate page index for all instance group
 //	if (bIsGPUDriven)
-	{		
-
+	{
 		FVoxelAllocatePageIndexCS::FParameters* Parameters = GraphBuilder.AllocParameters<FVoxelAllocatePageIndexCS::FParameters>();
 		Parameters->PageWorldSize = PageWorldSize;
 		Parameters->TotalPageIndexCount = OutTotalPageIndexCount;
@@ -836,6 +860,7 @@ static void AddAllocateVoxelPagesPass(
 			Parameters->IndirectGroupSize = IndirectDispatchGroupSize;
 			Parameters->OutPageIndexGlobalCounter = GraphBuilder.CreateUAV(PageIndexGlobalCounter, PF_R32_UINT);
 			Parameters->OutIndirectArgsBuffer = GraphBuilder.CreateUAV(IndirectArgsBuffer);
+			Parameters->OutTotalRequestedPageAllocationBuffer = TotalRequestedPageAllocationBufferUAV;
 
 			const FIntVector DispatchCount(1, 1, 1);
 			TShaderMapRef<FVoxelAddIndirectBufferCS> ComputeShader(View.ShaderMap);
@@ -853,22 +878,108 @@ static void AddAllocateVoxelPagesPass(
 	OutVoxelizationViewInfoBuffer = VoxelizationViewInfoBuffer;
 }
 
-FHairStrandsVoxelResources AllocateVirtualVoxelResources(
+static FHairStrandsVoxelResources AllocateVirtualVoxelResources(
 	FRDGBuilder& GraphBuilder,
 	const FViewInfo& View,
 	FHairStrandsMacroGroupDatas& MacroGroupDatas,
 	FHairStrandsMacroGroupResources& MacroGroupResources,
-	FRDGBufferRef& PageToPageIndexBuffer)
+	FRDGBufferRef& PageToPageIndexBuffer,
+	FHairStrandsViewStateData* OutViewStateData)
 {
 	DECLARE_GPU_STAT(HairStrandsVoxelPageAllocation);
 	RDG_EVENT_SCOPE(GraphBuilder, "HairStrandsVoxelPageAllocation");
 	RDG_GPU_STAT_SCOPE(GraphBuilder, HairStrandsVoxelPageAllocation);
 
+
+	// Init. default page table size and voxel size
+	FIntVector PageCountResolution = FIntVector(GHairVirtualVoxel_PageCountPerDim, GHairVirtualVoxel_PageCountPerDim, GHairVirtualVoxel_PageCountPerDim);
+	uint32 PageCount = PageCountResolution.X * PageCountResolution.Y * PageCountResolution.Z;
+	const float MinVoxelWorldSize = 0.01f;
+	const float MaxVoxelWorldSize = 10.f;
+
+	float VoxelWorldSize = FMath::Clamp(GHairVirtualVoxel_VoxelWorldSize, MinVoxelWorldSize, MaxVoxelWorldSize);
+
+	// Readback allocated value to adapt the voxel size in order to fit max page allocation
+	const bool bAdaptiveResolution = OutViewStateData && OutViewStateData->IsInit();
+	FRHIGPUBufferReadback* ReadbackBuffer = nullptr;
+	bool bEnqueueNewReadbackQuery = false;
+	if (bAdaptiveResolution)
+	{
+		// First initialization (no query has been issued yet)
+		if (OutViewStateData->VoxelWorldSize == 0)
+		{
+			OutViewStateData->VoxelWorldSize = VoxelWorldSize;
+			bEnqueueNewReadbackQuery = true;
+		}
+
+		uint32 AllocatedPageCount = PageCount;
+		ReadbackBuffer = OutViewStateData->GetBuffer();
+		if (ReadbackBuffer->IsReady())
+		{
+			bEnqueueNewReadbackQuery = true;
+
+			ReadbackBuffer = OutViewStateData->GetBuffer();
+			AllocatedPageCount = *(uint32*)(ReadbackBuffer->Lock(sizeof(uint32)));
+			ReadbackBuffer->Unlock();
+
+			// ...
+			const float Threshold = FMath::Clamp(GHairVirtualVoxelAdaptive_CorrectionThreshold, 0.f, 1.f);
+			const float Factor = FMath::Clamp(GHairVirtualVoxelAdaptive_CorrectionSpeed, 0.f, 1.f);
+
+			// Voxel pages are represent a volume. To derive a better estimate of the ratio by which voxel size needs to be scale, 
+			// compute the cubic root of this ratio.
+			//
+			// AllocatedPage   AllocatedRes^3
+			// ------------- = --------------  = VolumeRatio = LinearRatio^3
+			//    MaxPage          MaxRes^3
+
+			// Ratio used for predicting voxel size increase
+			const float VolumeRatio = float(AllocatedPageCount) / float(PageCount);
+			const float LinearRatio = FMath::Pow(VolumeRatio, 1.f / 3.f);
+
+			// Ratio used for predicting voxel size decrease (i.e. when requested allocation fit, 
+			// but the voxel size does not match the (more precise) target).
+			// In this case, we add a threshold/margin to to the target, so that there is no oscillation.
+			const float VolumeRatio_Thres = float(AllocatedPageCount) / float(PageCount * Threshold);
+			const float LinearRatio_Thres = FMath::Pow(VolumeRatio_Thres, 1.f / 3.f);
+
+			const float PrevWorldVoxelSize = OutViewStateData->VoxelWorldSize;
+
+			// If the page pool is not large enough increase voxel size
+			if (AllocatedPageCount > PageCount)
+			{
+				VoxelWorldSize = PrevWorldVoxelSize * LinearRatio;
+			}
+			// If the page pool is large enough but the voxel are larger than the requested size decrease voxel size
+			else if (AllocatedPageCount < PageCount && PrevWorldVoxelSize > VoxelWorldSize)
+			{
+				const float TargetVoxelWorldSize = PrevWorldVoxelSize * LinearRatio_Thres;
+				VoxelWorldSize = FMath::Max(VoxelWorldSize, FMath::Lerp(PrevWorldVoxelSize, TargetVoxelWorldSize, Factor));
+			}
+			else
+			{
+				VoxelWorldSize = PrevWorldVoxelSize;
+			}
+
+			// Clamp voxel size into a reasonable amount (e.g. 0.1mm - 100mm)
+			VoxelWorldSize = FMath::Clamp(VoxelWorldSize, MinVoxelWorldSize, MaxVoxelWorldSize);
+		}
+		else
+		{
+			// Use previous frame prediction by default (a readback is currently in-flight, but not ready for this frame)
+			VoxelWorldSize = OutViewStateData->VoxelWorldSize;
+		}
+
+		// Update state data
+		OutViewStateData->VoxelWorldSize = VoxelWorldSize;
+		OutViewStateData->AllocatedPageCount = AllocatedPageCount;
+	}
+
 	FHairStrandsVoxelResources Out;
 
-	Out.Parameters.Common.PageCountResolution		= FIntVector(GHairVirtualVoxel_PageCountPerDim, GHairVirtualVoxel_PageCountPerDim, GHairVirtualVoxel_PageCountPerDim);
-	Out.Parameters.Common.PageCount					= Out.Parameters.Common.PageCountResolution.X * Out.Parameters.Common.PageCountResolution.Y * Out.Parameters.Common.PageCountResolution.Z;
-	Out.Parameters.Common.VoxelWorldSize			= FMath::Clamp(GHairVirtualVoxel_VoxelWorldSize, 0.01f, 10.f);
+	Out.Parameters.Common.PageCountResolution		= PageCountResolution;
+	Out.Parameters.Common.PageCount					= PageCount;
+	Out.Parameters.Common.VoxelWorldSize			= VoxelWorldSize;
 	Out.Parameters.Common.PageResolution			= FMath::RoundUpToPowerOfTwo(FMath::Clamp(GHairVirtualVoxel_PageResolution, 2, 256));
 	Out.Parameters.Common.PageTextureResolution		= Out.Parameters.Common.PageCountResolution * Out.Parameters.Common.PageResolution;
 
@@ -895,6 +1006,7 @@ FHairStrandsVoxelResources AllocateVirtualVoxelResources(
 
 	Out.Parameters.Common.HairCoveragePixelRadiusAtDepth1	= ComputeMinStrandRadiusAtDepth1(FIntPoint(View.ViewRect.Width(), View.ViewRect.Height()), View.FOV, 1/*SampleCount*/, 1/*RasterizationScale*/).Primary;
 
+	FRDGBufferRef TotalRequestedPageAllocationBuffer;
 	AddAllocateVoxelPagesPass(
 		GraphBuilder, 
 		View, 
@@ -914,7 +1026,14 @@ FHairStrandsVoxelResources AllocateVirtualVoxelResources(
 		Out.NodeDescBuffer,
 		Out.IndirectArgsBuffer,
 		Out.PageIndexGlobalCounter,
-		Out.VoxelizationViewInfoBuffer);
+		Out.VoxelizationViewInfoBuffer,
+		TotalRequestedPageAllocationBuffer);
+
+	// Enque next adaptive feedback buffer
+	if (bEnqueueNewReadbackQuery)
+	{
+		AddEnqueueCopyPass(GraphBuilder, ReadbackBuffer, TotalRequestedPageAllocationBuffer, 4);
+	}
 
 	{
 		// Allocation should be conservative
@@ -1419,7 +1538,8 @@ void VoxelizeHairStrands(
 
 	{
 		FRDGBufferRef PageToPageIndexBuffer = nullptr;
-		VirtualVoxelResources = AllocateVirtualVoxelResources(GraphBuilder, View, MacroGroupDatas, MacroGroupResources, PageToPageIndexBuffer);
+		FHairStrandsViewStateData* HairStrandsViewStateData = View.ViewState ? &View.ViewState->HairStrandsViewStateData : nullptr;
+		VirtualVoxelResources = AllocateVirtualVoxelResources(GraphBuilder, View, MacroGroupDatas, MacroGroupResources, PageToPageIndexBuffer, HairStrandsViewStateData);
 
 		FRDGBufferRef ClearIndArgsBuffer = IndirectVoxelPageClear(GraphBuilder, View, VirtualVoxelResources);
 
diff --git a/Engine/Source/Runtime/Renderer/Private/RendererScene.cpp b/Engine/Source/Runtime/Renderer/Private/RendererScene.cpp
--- a/Engine/Source/Runtime/Renderer/Private/RendererScene.cpp
+++ b/Engine/Source/Runtime/Renderer/Private/RendererScene.cpp
@@ -288,6 +288,8 @@ FSceneViewState::~FSceneViewState()
 	AOTileIntersectionResources = NULL;
 	DestroyRenderResource(AOScreenGridResources);
 	AOScreenGridResources = NULL;
+
+	HairStrandsViewStateData.Release();
 }
 
 #if WITH_EDITOR
diff --git a/Engine/Source/Runtime/Renderer/Private/ScenePrivate.h b/Engine/Source/Runtime/Renderer/Private/ScenePrivate.h
--- a/Engine/Source/Runtime/Renderer/Private/ScenePrivate.h
+++ b/Engine/Source/Runtime/Renderer/Private/ScenePrivate.h
@@ -1084,6 +1084,8 @@ public:
 	FVector VolumetricCloudShadowmapPreviousAnchorPoint[NUM_ATMOSPHERE_LIGHTS];
 	FVector VolumetricCloudShadowmapPreviousAtmosphericLightDir[NUM_ATMOSPHERE_LIGHTS];
 
+	FHairStrandsViewStateData HairStrandsViewStateData;
+
 	// call after OnFrameRenderingSetup()
 	virtual uint32 GetCurrentTemporalAASampleIndex() const
 	{
