Remove HZB fetching from hair cluster selection as this path is no longer used.

This also fixes NaN arising due to HZB data being invalid under certain conditions.

diff --git a/Engine/Shaders/Private/HairStrands/HairStrandsClusterCulling.usf b/Engine/Shaders/Private/HairStrands/HairStrandsClusterCulling.usf
--- a/Engine/Shaders/Private/HairStrands/HairStrandsClusterCulling.usf
+++ b/Engine/Shaders/Private/HairStrands/HairStrandsClusterCulling.usf
@@ -529,11 +529,6 @@ RWBuffer<uint> DrawIndirectParameters;
 RWStructuredBuffer<FHairClusterDebugInfo> ClusterDebugInfoBuffer;
 #endif
 
-float3			HZBUvFactor;
-float4			HZBSize;
-Texture2D		HZBTexture;
-SamplerState	HZBSampler;
-
 bool ViewFrustumIntersectSphere(float3 SphereOrigin, float SphereRadius)
 {
 	for (uint PlaneIndex = 0; PlaneIndex < NumConvexHullPlanes; PlaneIndex++)
@@ -549,70 +544,6 @@ bool ViewFrustumIntersectSphere(float3 SphereOrigin, float SphereRadius)
 	return true;
 }
 
-// From HZBOcclusion.usf. TODO: make common
-// Rect contains screen coordinates as {minx, miny, maxx, maxy}. Those coordinates are not clamped to screen edges so they mnight be outside the [0,1]^2.
-bool ViewHierarchicalZBufferCulling(float3 WorldBoundsMin, float3 WorldBoundsMax, inout float4 NotClampedRect)
-{
-	float3 WorldBounds[2] = { WorldBoundsMin, WorldBoundsMax };
-
-	// Screen rect from bounds
-	float3 RectMin = float3(1, 1, 1);
-	float3 RectMax = float3(-1, -1, -1);
-	/*UNROLL*/ for (int i = 0; i < 8; i++)
-	{
-		float3 PointSrc;
-		PointSrc.x = WorldBounds[(i >> 0) & 1].x;
-		PointSrc.y = WorldBounds[(i >> 1) & 1].y;
-		PointSrc.z = WorldBounds[(i >> 2) & 1].z;
-
-		float4 PointClip = mul(float4(PointSrc, 1), WorldToClipMatrix);
-		float3 PointScreen = PointClip.xyz / PointClip.w;
-
-		RectMin = min(RectMin, PointScreen);
-		RectMax = max(RectMax, PointScreen);
-	}
-
-	// FIXME assumes DX
-	NotClampedRect = (float4(RectMin.xy, RectMax.xy) * float2(0.5, -0.5).xyxy + 0.5).xwzy;
-	float4 Rect = saturate(NotClampedRect);
-	float4 RectPixels = Rect * HZBSize.xyxy;
-	float2 RectSize = (RectPixels.zw - RectPixels.xy) * 0.5;	// 0.5 for 4x4
-	float Level = max(ceil(log2(max(RectSize.x, RectSize.y))), HZBUvFactor.z);
-
-#if PERMUTATION_HZBCULLING
-	// Check if we can drop one level lower
-	float LevelLower = max(Level - 1, 0);
-	float4 LowerRect = RectPixels * exp2(-LevelLower);
-	float2 LowerRectSize = ceil(LowerRect.zw) - floor(LowerRect.xy);
-	if (all(LowerRectSize <= 4))
-	{
-		Level = LevelLower;
-	}
-
-	// 4x4 samples
-	float2 Scale = HZBUvFactor.xy * (Rect.zw - Rect.xy) / 3;
-	float2 Bias = HZBUvFactor.xy * Rect.xy;
-
-	float4 MinDepth = 1;
-	/*UNROLL*/ for (int j = 0; j < 4; j++)
-	{
-		// TODO could vectorize this
-		float4 Depth;
-		Depth.x = HZBTexture.SampleLevel(HZBSampler, float2(j, 0) * Scale + Bias, Level).r;
-		Depth.y = HZBTexture.SampleLevel(HZBSampler, float2(j, 1) * Scale + Bias, Level).r;
-		Depth.z = HZBTexture.SampleLevel(HZBSampler, float2(j, 2) * Scale + Bias, Level).r;
-		Depth.w = HZBTexture.SampleLevel(HZBSampler, float2(j, 3) * Scale + Bias, Level).r;
-		MinDepth = min(MinDepth, Depth);
-	}
-	MinDepth.x = min(min(MinDepth.x, MinDepth.y), min(MinDepth.z, MinDepth.w));
-
-	// Inverted Z buffer
-	return RectMax.z >= MinDepth.x ? true : false;
-#else
-	return true;
-#endif
-}
-
 [numthreads(64, 1, 1)]
 void MainClusterCullingCS(uint3 DispatchThreadId : SV_DispatchThreadID)
 {
@@ -622,8 +553,7 @@ void MainClusterCullingCS(uint3 DispatchThreadId : SV_DispatchThreadID)
 		// Note on the cluster culling & LOD selection
 		// Run:
 		// 1. Frustum culling
-		// 2. Visibility culling (using HZB)
-		// 3. LOD selection
+		// 2. LOD selection
 		//
 		// When a cluster is culled:
 		// . Cluster's strands position still updated (HairInterpolation), for the lowest LOD
@@ -654,7 +584,6 @@ void MainClusterCullingCS(uint3 DispatchThreadId : SV_DispatchThreadID)
 		if (bVisible || bForceVisible) 
 		{
 			float4 NotClampedRect = float4(0.0f, 1.0f, 0.0f, 1.0f);
-			bVisible = bVisible && ViewHierarchicalZBufferCulling(Min, Max, NotClampedRect);
 			if (bVisible || bForceVisible)
 			{
 				const FHairClusterInfo ClusterInfo = UnpackHairClusterInfo(ClusterInfoBuffer[ClusterId]);
diff --git a/Engine/Source/Runtime/Renderer/Private/HairStrands/HairStrandsClusters.cpp b/Engine/Source/Runtime/Renderer/Private/HairStrands/HairStrandsClusters.cpp
--- a/Engine/Source/Runtime/Renderer/Private/HairStrands/HairStrandsClusters.cpp
+++ b/Engine/Source/Runtime/Renderer/Private/HairStrands/HairStrandsClusters.cpp
@@ -13,9 +13,6 @@
 
 ///////////////////////////////////////////////////////////////////////////////////////////////////
 
-static int32 GHairStrandsClusterCullingUsesHzb = 1;
-static FAutoConsoleVariableRef CVarHairCullingUseHzb(TEXT("r.HairStrands.Cluster.CullingUsesHzb"), GHairStrandsClusterCullingUsesHzb, TEXT("Enable/disable the use of HZB to help cull more hair clusters."));
-
 static int32 GHairStrandsClusterForceLOD = -1;
 static FAutoConsoleVariableRef CVarHairClusterCullingLodMode(TEXT("r.HairStrands.Cluster.ForceLOD"), GHairStrandsClusterForceLOD, TEXT("Force a specific hair LOD."));
 
@@ -29,11 +26,6 @@ bool IsHairStrandsClusterCullingEnable()
 	return true;
 }
 
-bool IsHairStrandsClusterCullingUseHzb()
-{
-	return GHairStrandsClusterCullingUsesHzb > 0;
-}
-
 ///////////////////////////////////////////////////////////////////////////////////////////////////
 
 class FHairIndBufferClearCS : public FGlobalShader
@@ -68,9 +60,8 @@ class FHairClusterCullingCS : public FGlobalShader
 	DECLARE_GLOBAL_SHADER(FHairClusterCullingCS);
 	SHADER_USE_PARAMETER_STRUCT(FHairClusterCullingCS, FGlobalShader);
 
-	class FHZBCulling : SHADER_PERMUTATION_INT("PERMUTATION_HZBCULLING", 2);
 	class FDebugAABBBuffer : SHADER_PERMUTATION_INT("PERMUTATION_DEBUGAABBBUFFER", 2);
-	using FPermutationDomain = TShaderPermutationDomain<FHZBCulling, FDebugAABBBuffer>;
+	using FPermutationDomain = TShaderPermutationDomain<FDebugAABBBuffer>;
 
 	BEGIN_SHADER_PARAMETER_STRUCT(FParameters, )
 		SHADER_PARAMETER(FVector3f, CameraWorldPos)
@@ -92,10 +83,6 @@ class FHairClusterCullingCS : public FGlobalShader
 		SHADER_PARAMETER_RDG_BUFFER_UAV(RWStructuredBuffer, ClusterDebugInfoBuffer)
 		SHADER_PARAMETER_RDG_BUFFER_UAV(RWBuffer, DispatchIndirectParametersClusterCount)
 		SHADER_PARAMETER_RDG_BUFFER_UAV(RWBuffer, DrawIndirectParameters)
-		SHADER_PARAMETER(FVector3f, HZBUvFactor)
-		SHADER_PARAMETER(FVector4, HZBSize)
-		SHADER_PARAMETER_RDG_TEXTURE(Texture2D<float>, HZBTexture)
-		SHADER_PARAMETER_SAMPLER(SamplerState, HZBSampler)
 	END_SHADER_PARAMETER_STRUCT()
 
 public:
@@ -247,14 +234,6 @@ static FVector CapturedCameraWorldPos;
 static FMatrix CapturedWorldToClipMatrix;
 static FMatrix CapturedProjMatrix;
 
-struct FHairHZBParameters
-{
-	// Set from renderer for view culling
-	FVector HZBUvFactorValue;
-	FVector4 HZBSizeValue;
-	FRDGTextureRef HZB = nullptr;
-};
-
 bool IsHairStrandsClusterDebugEnable();
 bool IsHairStrandsClusterDebugAABBEnable();
 
@@ -263,7 +242,6 @@ static void AddClusterCullingPass(
 	FGlobalShaderMap* ShaderMap,
 	const FViewInfo& View,
 	const FHairCullingParams& CullingParameters,
-	const FHairHZBParameters& HZBParameters,
 	FHairStrandClusterData::FHairGroup& ClusterData)
 {
 	FRDGBufferRef DispatchIndirectParametersClusterCount = GraphBuilder.CreateBuffer(FRDGBufferDesc::CreateIndirectDesc<FRHIDispatchIndirectParameters>(), TEXT("Hair.DispatchIndirectParametersClusterCount"));
@@ -375,13 +353,7 @@ static void AddClusterCullingPass(
 		Parameters->ClusterDebugInfoBuffer = GraphBuilder.CreateUAV(ClusterDebugInfoBuffer, PF_R32_SINT);
 #endif
 
-		Parameters->HZBUvFactor = HZBParameters.HZBUvFactorValue;
-		Parameters->HZBSize = HZBParameters.HZBSizeValue;
-		Parameters->HZBTexture = HZBParameters.HZB;
-		Parameters->HZBSampler = TStaticSamplerState<SF_Point, AM_Clamp, AM_Clamp, AM_Clamp>::GetRHI();
-
 		FHairClusterCullingCS::FPermutationDomain Permutation;
-		Permutation.Set<FHairClusterCullingCS::FHZBCulling>((HZBParameters.HZB && !bClusterCullingFrozenCamera && GHairStrandsClusterCullingUsesHzb) ? 1 : 0);
 		Permutation.Set<FHairClusterCullingCS::FDebugAABBBuffer>(bClusterDebugAABBBuffer ? 1 : 0);
 		TShaderMapRef<FHairClusterCullingCS> ComputeShader(ShaderMap, Permutation);
 		const FIntVector DispatchCount = DispatchCount.DivideAndRoundUp(FIntVector(ClusterData.ClusterCount, 1, 1), FIntVector(64, 1, 1));
@@ -578,29 +550,6 @@ void ComputeHairStrandsClustersCulling(
 		SceneViews[ViewId] = &Views[ViewId];
 	}
 
-	FHairHZBParameters HZBParameters;
-	if (ViewCount > 0) // only handling one view for now
-	{
-		const FViewInfo& ViewInfo = Views[0];
-		HZBParameters.HZB = ViewInfo.HZB;
-
-		const float kHZBTestMaxMipmap = 9.0f;
-		const float HZBMipmapCounts = FMath::Log2(FMath::Max<float>(ViewInfo.HZBMipmap0Size.X, ViewInfo.HZBMipmap0Size.Y));
-		const FVector HZBUvFactorValue(
-			float(ViewInfo.ViewRect.Width()) / float(2 * ViewInfo.HZBMipmap0Size.X),
-			float(ViewInfo.ViewRect.Height()) / float(2 * ViewInfo.HZBMipmap0Size.Y),
-			FMath::Max(HZBMipmapCounts - kHZBTestMaxMipmap, 0.0f)
-		);
-		const FVector4 HZBSizeValue(
-			ViewInfo.HZBMipmap0Size.X,
-			ViewInfo.HZBMipmap0Size.Y,
-			1.0f / float(ViewInfo.HZBMipmap0Size.X),
-			1.0f / float(ViewInfo.HZBMipmap0Size.Y)
-		);
-		HZBParameters.HZBUvFactorValue = HZBUvFactorValue;
-		HZBParameters.HZBSizeValue = HZBSizeValue;
-	}
-
 	const bool bClusterCulling = IsHairStrandsClusterCullingEnable();
 	for (const FViewInfo& View : Views)
 	{
@@ -619,7 +568,6 @@ void ComputeHairStrandsClustersCulling(
 					&ShaderMap, 
 					View, 
 					CullingParameters,
-					HZBParameters,
 					ClusterData);
 			}
 		}
