Rewrites eye adaptation histogram using InterlockedAdd()

diff --git a/Engine/Shaders/Private/PostProcessHistogram.usf b/Engine/Shaders/Private/PostProcessHistogram.usf
--- a/Engine/Shaders/Private/PostProcessHistogram.usf
+++ b/Engine/Shaders/Private/PostProcessHistogram.usf
@@ -1,126 +1,179 @@
 // Copyright Epic Games, Inc. All Rights Reserved.
 
-/*=============================================================================
-	PostProcessHistogram.usf: PostProcessing histogram
-=============================================================================*/
-
 #include "Common.ush"
-#include "ScreenPass.ush"
 #include "PostProcessHistogramCommon.ush"
 
-SCREEN_PASS_TEXTURE_VIEWPORT(Input)
+//------------------------------------------------------- CONFIG
 
-Texture2D InputTexture;
+#define GROUP_SIZE 16
 
-// Output histogram texture (UAV)
-RWTexture2D<float4> HistogramRWTexture;
+// Number of independent histogram in DLS to reduce atomic contention.
+#define CONFIG_HISTOGRAM_COUNT ((GROUP_SIZE * GROUP_SIZE) / HISTOGRAM_SIZE)
 
-// Number of thread groups in the dispatch
-uint2 ThreadGroupCount;
+#define CONFIG_LOOP_SIZE 8
 
-// THREADGROUP_SIZEX*THREADGROUP_SIZEY histograms of the size HISTOGRAM_SIZE
-groupshared float SharedHistogram[HISTOGRAM_SIZE][THREADGROUP_SIZEX][THREADGROUP_SIZEY];
+#define CONFIG_QUANTIZATION 64.0
 
-[numthreads(THREADGROUP_SIZEX, THREADGROUP_SIZEY, 1)]
-void MainCS(
-	uint3 GroupId : SV_GroupID,
-	uint3 DispatchThreadId : SV_DispatchThreadID,
-    uint3 GroupThreadId : SV_GroupThreadID,
-	uint GroupIndex: SV_GroupIndex)
-{
-	// todo: can be cleared more efficiently
-	// clear all THREADGROUP_SIZEX*THREADGROUP_SIZEY histograms
-	UNROLL for (uint i = 0; i < HISTOGRAM_SIZE; ++i)
-	{
-		SharedHistogram[i][GroupThreadId.x][GroupThreadId.y] = 0.0f;
-	}
-	
-	GroupMemoryBarrierWithGroupSync();
 
-	// Each thread group processes LoopX * LoopY texels of the input.
-	uint2 TileSize = uint2(LOOP_SIZEX, LOOP_SIZEY);
+//------------------------------------------------------- PARAMETERS
 
-	// Top left input texel for this group.
-	uint2 LeftTop = Input_ViewportMin + DispatchThreadId.xy * TileSize;
 
-	float2 InvViewSize = Input_ViewportSizeInverse.xy;
+float2 InputSceneColorViewportExtentInverse;
+uint4 InputSceneColorViewportMinMax;
+Texture2D<float3> InputSceneColorTexture;
+RWStructuredBuffer<uint> HistogramOutput;
 
-	// Accumulate all pixels into THREADGROUP_SIZEX*THREADGROUP_SIZEY histograms
-	LOOP for (uint y = 0; y < LOOP_SIZEY; ++y)
-	{
-		LOOP for (uint x = 0; x < LOOP_SIZEX; ++x)
-		{
-			uint2 Tile = uint2(x, y);
-			uint2 TexelPos = LeftTop + Tile;
-
-			if(TexelPos.x < Input_ViewportMax.x && TexelPos.y < Input_ViewportMax.y)
-			{
-				float4 SceneColor = InputTexture.Load(int3(TexelPos, 0));
+float AreaNormalizeFactor;
+Texture2D EyeAdaptationTexture;
+StructuredBuffer<uint> HistogramBuffer;
+RWTexture2D<float4> HistogramTextureOutput;
 
-				SceneColor.xyz *= View.OneOverPreExposure;
 
-				float LuminanceVal = max(dot(SceneColor.xyz, float3(1.0f, 1.0f, 1.0f)/3.0f), EyeAdaptation_LuminanceMin);
+//------------------------------------------------------- LDS
 
-				float LogLuminance = ComputeHistogramPositionFromLuminance(LuminanceVal);
-				float2 ScreenUV = (TexelPos.xy - Input_ViewportMin) * InvViewSize;
-				float ScreenWeight = AdaptationWeightTexture(ScreenUV);
+groupshared uint SharedHistograms[HISTOGRAM_SIZE * CONFIG_HISTOGRAM_COUNT];
 
-				// Map the normalized histogram position into texels.
-				float fBucket = saturate(LogLuminance) * (HISTOGRAM_SIZE-1);// * 0.9999f;
 
-				// Find two discrete buckets that straddle the continuous histogram position.
-				uint Bucket0 = (uint)(fBucket);
-				uint Bucket1 = Bucket0 + 1;
+//------------------------------------------------------- ENTRY POINT
 
-				Bucket0 = min(Bucket0, uint(HISTOGRAM_SIZE - 1));
-				Bucket1 = min(Bucket1, uint(HISTOGRAM_SIZE - 1));
+[numthreads(GROUP_SIZE, GROUP_SIZE, 1)]
+void MainCS(
+	uint2 GroupId : SV_GroupID,
+	uint GroupThreadIndex : SV_GroupIndex)
+{
+	// Clear the shared histogram in one go. Use all 256 threads to clear 4x 64 bin histograms.
+	SharedHistograms[GroupThreadIndex] = 0;
+	 
+	GroupMemoryBarrierWithGroupSync();
 
-				// Weighted blend between the two buckets.
-				float Weight1 = frac(fBucket);
-				float Weight0 = 1.0f - Weight1;
+	uint2 GroupThreadId = uint2(GroupThreadIndex % GROUP_SIZE, GroupThreadIndex / GROUP_SIZE);
+	uint2 GroupOffset = uint(GROUP_SIZE * CONFIG_LOOP_SIZE) * GroupId.xy + InputSceneColorViewportMinMax.xy;
 
-				// When EyeAdaptation_BlackHistogramBucketInfluence=.0, we will ignore the last bucket. The main use
-				// case is so the black background pixels in the editor have no effect. But if we have cases where
-				// pixel values can actually be black, we want to set EyeAdaptation_LastHistogramBucketInfluence=1.0.
-				// This value is controlled by the cvar "r.EyeAdaptation.BlackHistogramBucketInfluence"
-				if (Bucket0 == 0)
-				{
-					Weight0 *= EyeAdaptation_BlackHistogramBucketInfluence;
-				}
+	LOOP
+	for (uint y = 0; y < CONFIG_LOOP_SIZE; y ++)
+	{
+		BRANCH
+		if (GroupOffset.y + y * CONFIG_LOOP_SIZE >= InputSceneColorViewportMinMax.w)
+		{
+			break;
+		}
 
-				// Accumulate the weight to the nearby history buckets.
-				SharedHistogram[Bucket0][GroupThreadId.x][GroupThreadId.y] += Weight0 * ScreenWeight;
-				SharedHistogram[Bucket1][GroupThreadId.x][GroupThreadId.y] += Weight1 * ScreenWeight;
+		bool bProcess[CONFIG_LOOP_SIZE];
+		float3 SceneColor[CONFIG_LOOP_SIZE];
+		float ScreenWeight[CONFIG_LOOP_SIZE];
+		ISOLATE
+		{
+			UNROLL_N(CONFIG_LOOP_SIZE)
+			for (uint x = 0; x < CONFIG_LOOP_SIZE; x++)
+			{
+				const uint2 TileOffset = uint2(x, y) * GROUP_SIZE;
+				uint2 TexelPos = (GroupOffset + TileOffset) + GroupThreadId;
+				
+				const uint SampleId = x;
+				bProcess[SampleId] = all(TexelPos < InputSceneColorViewportMinMax.zw);
+				SceneColor[SampleId] = InputSceneColorTexture[TexelPos];
+
+				float2 ScreenUV = float2(TexelPos - InputSceneColorViewportMinMax.xy) * InputSceneColorViewportExtentInverse;
+				ScreenWeight[SampleId] = AdaptationWeightTexture(ScreenUV);
 			}
 		}
-	}
+			
+		UNROLL_N(CONFIG_LOOP_SIZE)
+		for (uint SampleId = 0; SampleId < CONFIG_LOOP_SIZE; SampleId++)
+		{
+			float LuminanceVal = max(dot(SceneColor[SampleId].xyz, float3(1.0f, 1.0f, 1.0f)/3.0f), EyeAdaptation_LuminanceMin);
+			float LogLuminance = ComputeHistogramPositionFromLuminance(LuminanceVal);
+			
+			float fBucket = saturate(LogLuminance) * float((HISTOGRAM_SIZE - 1) * 0.9999f);
+			fBucket = clamp(fBucket, float(0.0), float(HISTOGRAM_SIZE - 1.0));
+
+			uint Bucket0 = (uint)(fBucket);
+			Bucket0 = clamp(fBucket, 0, HISTOGRAM_SIZE - 1);
+
+			float Weight1 = frac(fBucket) * ScreenWeight[SampleId];
+			float Weight0 = (1.0f - Weight1) * ScreenWeight[SampleId];
+			
+			// accumulate the weight to the nearby history buckets
+			// NOTE: Pack "overflow" to next bucket to same counter high bits. This halves the number of atomics. 
+			uint PackedHistogram = (uint(Weight1 * CONFIG_QUANTIZATION) << 16) | uint(Weight0 * CONFIG_QUANTIZATION);
+                   
+			uint Offset = HISTOGRAM_SIZE * (GroupThreadIndex % CONFIG_HISTOGRAM_COUNT);
+			
+			if (bProcess[SampleId])
+			{
+				InterlockedAdd(SharedHistograms[Bucket0 + Offset], PackedHistogram);
+			}
+		}
+	} // for (uint y = 0; y < CONFIG_LOOP_SIZE; y ++)
 
 	GroupMemoryBarrierWithGroupSync();
 
-	// Accumulate all histograms into one.
-	if (GroupIndex < HISTOGRAM_SIZE / 4) 
+	BRANCH
+	if (GroupThreadIndex < HISTOGRAM_SIZE)
 	{
-		float4 Sum = 0;
-
-		LOOP for (uint y = 0; y < THREADGROUP_SIZEY; ++y)
+		uint Sum = 0;
+		
 		{
-			LOOP for (uint x = 0; x < THREADGROUP_SIZEX; ++x)
+			uint SumLow = 0;
+			uint SumHi = 0;
+
+			UNROLL_N(CONFIG_HISTOGRAM_COUNT)
+			for (uint LDSHistogramId = 0; LDSHistogramId < CONFIG_HISTOGRAM_COUNT; LDSHistogramId++)
 			{
-				Sum += float4(
-					SharedHistogram[GroupIndex * 4 + 0][x][y],
-					SharedHistogram[GroupIndex * 4 + 1][x][y],
-					SharedHistogram[GroupIndex * 4 + 2][x][y],
-					SharedHistogram[GroupIndex * 4 + 3][x][y]);
+				SumLow += SharedHistograms[GroupThreadIndex + HISTOGRAM_SIZE * LDSHistogramId];
+				SumHi += SharedHistograms[GroupThreadIndex + HISTOGRAM_SIZE * LDSHistogramId - 1];
 			}
+
+			FLATTEN
+			if (GroupThreadIndex == 0)
+			{
+				Sum = SumLow;
+			}
+			else
+			{
+				Sum = (SumLow + (SumHi >> 16));
+			}
+
+			Sum = Sum & 0xFFFF;
 		}
 
-		float2 MaxExtent = Input_ViewportSize;
-		float Area = MaxExtent.x * MaxExtent.y;
+		// Accumulate groups to shared histogram
+		InterlockedAdd(HistogramOutput[GroupThreadIndex], Sum);
+	}
+}
+
+
+//------------------------------------------------------- ENTRY POINT
+
+[numthreads(HISTOGRAM_SIZE / 4, 1, 1)]
+void MainConvertCS(
+	uint GroupId : SV_GroupID,
+	uint GroupThreadIndex : SV_GroupIndex)
+{
+	BRANCH
+	if (GroupId.x == 0)
+	{
+		uint4 HistogramUint = uint4(
+			HistogramBuffer[GroupThreadIndex * 4 + 0],
+			HistogramBuffer[GroupThreadIndex * 4 + 1],
+			HistogramBuffer[GroupThreadIndex * 4 + 2],
+			HistogramBuffer[GroupThreadIndex * 4 + 3]);
 	
-		// Fixed to include borders.
-		float NormalizeFactor = 1.0f / Area;
+		float4 Luma = float4(HistogramUint) * (rcp(CONFIG_QUANTIZATION) * AreaNormalizeFactor);
+
+		// When EyeAdaptation_BlackHistogramBucketInfluence=.0, we will ignore the last bucket. The main use
+		// case is so the black background pixels in the editor have no effect. But if we have cases where
+		// pixel values can actually be black, we want to set EyeAdaptation_LastHistogramBucketInfluence=1.0.
+		// This value is controlled by the cvar "r.EyeAdaptation.BlackHistogramBucketInfluence"
+		Luma.r *= GroupThreadIndex == 0 ? EyeAdaptation_BlackHistogramBucketInfluence : 1.0;
 
-		// Output texture with one histogram per line, x and y unwrapped into all the lines
-		HistogramRWTexture[uint2(GroupIndex, GroupId.x + GroupId.y * ThreadGroupCount.x)] = Sum * NormalizeFactor;
+		HistogramTextureOutput[uint2(GroupThreadIndex, 0)] = Luma;
+	}
+	else
+	{
+		// line 1: store 4 channels of EyeAdaptationTexture (copied over so we can read the value in EyeAdaptation pass which is writing to eye adaptation)
+		float4 OldExposureScale = EyeAdaptationTexture.Load(int3(0, 0, 0));
+		
+		HistogramTextureOutput[uint2(GroupThreadIndex, 1)] = OldExposureScale;
 	}
 }
diff --git a/Engine/Shaders/Private/PostProcessHistogramReduce.usf b/Engine/Shaders/Private/PostProcessHistogramReduce.usf
deleted file mode 100644
diff --git a/Engine/Source/Runtime/Renderer/Private/PostProcess/PostProcessHistogram.cpp b/Engine/Source/Runtime/Renderer/Private/PostProcess/PostProcessHistogram.cpp
--- a/Engine/Source/Runtime/Renderer/Private/PostProcess/PostProcessHistogram.cpp
+++ b/Engine/Source/Runtime/Renderer/Private/PostProcess/PostProcessHistogram.cpp
@@ -15,88 +15,48 @@ class FHistogramCS : public FGlobalShader
 {
 public:
 	// Changing these numbers requires Histogram.usf to be recompiled.
-	static const uint32 ThreadGroupSizeX = 8;
-	static const uint32 ThreadGroupSizeY = 4;
-	static const uint32 LoopCountX = 8;
-	static const uint32 LoopCountY = 8;
-	static const uint32 HistogramSize = 64;
-
-	// /4 as we store 4 buckets in one ARGB texel.
-	static const uint32 HistogramTexelCount = HistogramSize / 4;
-
-	// The number of texels on each axis processed by a single thread group.
-	static const FIntPoint TexelsPerThreadGroup;
+	static const uint32 kHistogramSize = 64;
 
 	DECLARE_GLOBAL_SHADER(FHistogramCS);
 	SHADER_USE_PARAMETER_STRUCT(FHistogramCS, FGlobalShader);
 
 	BEGIN_SHADER_PARAMETER_STRUCT(FParameters, )
 		SHADER_PARAMETER_STRUCT_REF(FViewUniformShaderParameters, View)
-		SHADER_PARAMETER_STRUCT(FScreenPassTextureViewportParameters, Input)
 		SHADER_PARAMETER_STRUCT(FEyeAdaptationParameters, EyeAdaptation)
-		SHADER_PARAMETER_RDG_TEXTURE(Texture2D, InputTexture)
-		SHADER_PARAMETER_RDG_TEXTURE_UAV(RWTexture2D<float4>, HistogramRWTexture)
-		SHADER_PARAMETER(FIntPoint, ThreadGroupCount)
+		SHADER_PARAMETER(FIntRect, InputSceneColorViewportMinMax)
+		SHADER_PARAMETER(FVector2D, InputSceneColorViewportExtentInverse)
+		SHADER_PARAMETER_RDG_TEXTURE(Texture2D, InputSceneColorTexture)
+		SHADER_PARAMETER_RDG_BUFFER_UAV(RWStructuredBuffer<uint>, HistogramOutput)
 	END_SHADER_PARAMETER_STRUCT()
 
 	static bool ShouldCompilePermutation(const FGlobalShaderPermutationParameters& Parameters)
 	{
 		return IsFeatureLevelSupported(Parameters.Platform, ERHIFeatureLevel::SM5);
 	}
-
-	static void ModifyCompilationEnvironment(const FGlobalShaderPermutationParameters& Parameters, FShaderCompilerEnvironment& OutEnvironment)
-	{
-		FGlobalShader::ModifyCompilationEnvironment(Parameters, OutEnvironment);
-		OutEnvironment.SetDefine(TEXT("THREADGROUP_SIZEX"), ThreadGroupSizeX);
-		OutEnvironment.SetDefine(TEXT("THREADGROUP_SIZEY"), ThreadGroupSizeY);
-		OutEnvironment.SetDefine(TEXT("LOOP_SIZEX"), LoopCountX);
-		OutEnvironment.SetDefine(TEXT("LOOP_SIZEY"), LoopCountY);
-		OutEnvironment.SetDefine(TEXT("HISTOGRAM_SIZE"), HistogramSize);
-		OutEnvironment.CompilerFlags.Add( CFLAG_StandardOptimization );
-	}
-
-	// One ThreadGroup processes LoopCountX*LoopCountY blocks of size ThreadGroupSizeX*ThreadGroupSizeY
-	static FIntPoint GetThreadGroupCount(FIntPoint InputExtent)
-	{
-		return FIntPoint::DivideAndRoundUp(InputExtent, TexelsPerThreadGroup);
-	}
 };
 
-const FIntPoint FHistogramCS::TexelsPerThreadGroup(ThreadGroupSizeX * LoopCountX, ThreadGroupSizeY * LoopCountY);
-
-IMPLEMENT_GLOBAL_SHADER(FHistogramCS, "/Engine/Private/PostProcessHistogram.usf", "MainCS", SF_Compute);
-
-class FHistogramReducePS : public FGlobalShader
+class FHistogramConvertCS : public FGlobalShader
 {
 public:
-	DECLARE_GLOBAL_SHADER(FHistogramReducePS);
-	SHADER_USE_PARAMETER_STRUCT(FHistogramReducePS, FGlobalShader);
+	DECLARE_GLOBAL_SHADER(FHistogramConvertCS);
+	SHADER_USE_PARAMETER_STRUCT(FHistogramConvertCS, FGlobalShader);
 	
 	BEGIN_SHADER_PARAMETER_STRUCT(FParameters, )
-		SHADER_PARAMETER_STRUCT(FScreenPassTextureViewportParameters, Input)
-		SHADER_PARAMETER_SAMPLER(SamplerState, InputSampler)
-		SHADER_PARAMETER_RDG_TEXTURE(Texture2D, InputTexture)
+		SHADER_PARAMETER_STRUCT(FEyeAdaptationParameters, EyeAdaptation)
+		SHADER_PARAMETER(float, AreaNormalizeFactor)
 		SHADER_PARAMETER_RDG_TEXTURE(Texture2D, EyeAdaptationTexture)
-		SHADER_PARAMETER(uint32, LoopSize)
-		RENDER_TARGET_BINDING_SLOTS()
+		SHADER_PARAMETER_RDG_BUFFER_SRV(StructuredBuffer<uint>, HistogramBuffer)
+		SHADER_PARAMETER_RDG_TEXTURE_UAV(RWTexture2D<float4>, HistogramTextureOutput)
 	END_SHADER_PARAMETER_STRUCT()
 
-	// Uses full float4 to get best quality for smooth eye adaptation transitions.
-	static const EPixelFormat OutputFormat = PF_A32B32G32R32F;
-
 	static bool ShouldCompilePermutation(const FGlobalShaderPermutationParameters& Parameters)
 	{
 		return IsFeatureLevelSupported(Parameters.Platform, ERHIFeatureLevel::SM5);
 	}
-
-	static void ModifyCompilationEnvironment(const FGlobalShaderPermutationParameters& Parameters, FShaderCompilerEnvironment& OutEnvironment)
-	{
-		FGlobalShader::ModifyCompilationEnvironment(Parameters, OutEnvironment);
-		OutEnvironment.SetRenderTargetOutputFormat(0, OutputFormat);
-	}
 };
 
-IMPLEMENT_GLOBAL_SHADER(FHistogramReducePS, "/Engine/Private/PostProcessHistogramReduce.usf", "MainPS", SF_Pixel);
+IMPLEMENT_GLOBAL_SHADER(FHistogramCS, "/Engine/Private/PostProcessHistogram.usf", "MainCS", SF_Compute);\
+IMPLEMENT_GLOBAL_SHADER(FHistogramConvertCS, "/Engine/Private/PostProcessHistogram.usf", "MainConvertCS", SF_Compute);
 
 } //! namespace
 
@@ -110,85 +70,72 @@ FRDGTextureRef AddHistogramPass(
 	check(SceneColor.IsValid());
 	check(EyeAdaptationTexture);
 
-	const FIntPoint HistogramThreadGroupCount = FIntPoint::DivideAndRoundUp(SceneColor.ViewRect.Size(), FHistogramCS::TexelsPerThreadGroup);
-
-	const uint32 HistogramThreadGroupCountTotal = HistogramThreadGroupCount.X * HistogramThreadGroupCount.Y;
-
-	FRDGTextureRef HistogramTexture = nullptr;
-
-	RDG_EVENT_SCOPE(GraphBuilder, "Histogram");
+	RDG_EVENT_SCOPE(GraphBuilder, "Histogram %dx%d", SceneColor.ViewRect.Width(), SceneColor.ViewRect.Height());
 
 	// First pass outputs one flattened histogram per group.
+	FRDGBufferRef HistogramBuffer = nullptr;
 	{
-		const FIntPoint TextureExtent = FIntPoint(FHistogramCS::HistogramTexelCount, HistogramThreadGroupCountTotal);
+		{
+			FRDGBufferDesc Desc = FRDGBufferDesc::CreateStructuredDesc(sizeof(uint32), FHistogramCS::kHistogramSize);
 
-		const FRDGTextureDesc TextureDesc = FRDGTextureDesc::Create2D(
-			TextureExtent,
-			PF_FloatRGBA,
-			FClearValueBinding::None,
-			GFastVRamConfig.Histogram | TexCreate_RenderTargetable | TexCreate_UAV | TexCreate_ShaderResource);
-
-		HistogramTexture = GraphBuilder.CreateTexture(TextureDesc, TEXT("Histogram"));
+			if (GFastVRamConfig.Histogram & TexCreate_FastVRAM)
+			{
+				Desc.Usage |= BUF_FastVRAM | BUF_Transient;
+			}
 
+			HistogramBuffer = GraphBuilder.CreateBuffer(Desc, TEXT("Histogram"));
+		}
+		
 		FHistogramCS::FParameters* PassParameters = GraphBuilder.AllocParameters<FHistogramCS::FParameters>();
 		PassParameters->View = View.ViewUniformBuffer;
-		PassParameters->Input = GetScreenPassTextureViewportParameters(FScreenPassTextureViewport(SceneColor));
-		PassParameters->InputTexture = SceneColor.Texture;
-		PassParameters->HistogramRWTexture = GraphBuilder.CreateUAV(HistogramTexture);
-		PassParameters->ThreadGroupCount = HistogramThreadGroupCount;
 		PassParameters->EyeAdaptation = EyeAdaptationParameters;
+		PassParameters->InputSceneColorViewportMinMax = SceneColor.ViewRect;
+		PassParameters->InputSceneColorViewportExtentInverse.X = 1.0f / float(SceneColor.ViewRect.Width());
+		PassParameters->InputSceneColorViewportExtentInverse.Y = 1.0f / float(SceneColor.ViewRect.Height());
+		PassParameters->InputSceneColorTexture = SceneColor.Texture;
+		PassParameters->HistogramOutput = GraphBuilder.CreateUAV(HistogramBuffer);
 
-		TShaderMapRef<FHistogramCS> ComputeShader(View.ShaderMap);
+		static const int32 kGroupSize = 16;
+		static const int32 kLoopSize = 8;
 
+		TShaderMapRef<FHistogramCS> ComputeShader(View.ShaderMap);
 		FComputeShaderUtils::AddPass(
 			GraphBuilder,
-			RDG_EVENT_NAME("Histogram %dx%d (CS)", SceneColor.ViewRect.Width(), SceneColor.ViewRect.Height()),
+			RDG_EVENT_NAME("BuildHistogram %dx%d", SceneColor.ViewRect.Width(), SceneColor.ViewRect.Height()),
 			ComputeShader,
 			PassParameters,
-			FIntVector(HistogramThreadGroupCount.X, HistogramThreadGroupCount.Y, 1));
+			FComputeShaderUtils::GetGroupCount(SceneColor.ViewRect.Size(), kGroupSize * kLoopSize));
 	}
 
+	// Converts the histogram to a single line in a RGBA texture.
+	// The second line contains the eye adaptation value (two line texture).
 	FRDGTextureRef HistogramReduceTexture = nullptr;
-
-	// Second pass further reduces the histogram to a single line. The second line contains the eye adaptation value (two line texture).
 	{
-		const FIntPoint TextureExtent = FIntPoint(FHistogramCS::HistogramTexelCount, 2);
-
-		const FRDGTextureDesc TextureDesc = FRDGTextureDesc::Create2D(
-			TextureExtent,
-			FHistogramReducePS::OutputFormat,
-			FClearValueBinding::None,
-			GFastVRamConfig.HistogramReduce | TexCreate_RenderTargetable | TexCreate_ShaderResource);
-
-		HistogramReduceTexture = GraphBuilder.CreateTexture(TextureDesc, TEXT("HistogramReduce"));
+		{
+			const FRDGTextureDesc TextureDesc = FRDGTextureDesc::Create2D(
+				FIntPoint(FHistogramCS::kHistogramSize / 4, 2),
+				PF_A32B32G32R32F,
+				FClearValueBinding::None,
+				GFastVRamConfig.HistogramReduce | TexCreate_UAV | TexCreate_ShaderResource);
 
-		const FScreenPassTextureViewport InputViewport(HistogramTexture);
-		const FScreenPassTextureViewport OutputViewport(HistogramReduceTexture);
+			HistogramReduceTexture = GraphBuilder.CreateTexture(TextureDesc, TEXT("Histogram"));
+		}
 
-		FHistogramReducePS::FParameters* PassParameters = GraphBuilder.AllocParameters<FHistogramReducePS::FParameters>();
-		PassParameters->Input = GetScreenPassTextureViewportParameters(InputViewport);
-		PassParameters->InputTexture = HistogramTexture;
-		PassParameters->InputSampler = TStaticSamplerState<SF_Bilinear, AM_Clamp, AM_Clamp, AM_Clamp>::GetRHI();
-		PassParameters->LoopSize = HistogramThreadGroupCountTotal;
+		FHistogramConvertCS::FParameters* PassParameters = GraphBuilder.AllocParameters<FHistogramConvertCS::FParameters>();
+		PassParameters->EyeAdaptation = EyeAdaptationParameters;
+		PassParameters->AreaNormalizeFactor = 1.0f / float(SceneColor.ViewRect.Area());
 		PassParameters->EyeAdaptationTexture = EyeAdaptationTexture;
-		PassParameters->RenderTargets[0] = FRenderTargetBinding(HistogramReduceTexture, ERenderTargetLoadAction::ENoAction);
-
-		TShaderMapRef<FHistogramReducePS> PixelShader(View.ShaderMap);
+		PassParameters->HistogramBuffer = GraphBuilder.CreateSRV(HistogramBuffer);
+		PassParameters->HistogramTextureOutput = GraphBuilder.CreateUAV(HistogramReduceTexture);
 
-		AddDrawScreenPass(
+		TShaderMapRef<FHistogramConvertCS> ComputeShader(View.ShaderMap);
+		FComputeShaderUtils::AddPass(
 			GraphBuilder,
-			RDG_EVENT_NAME("HistogramReduce %dx%d (PS)", InputViewport.Extent.X, InputViewport.Extent.Y),
-			View,
-			OutputViewport,
-			InputViewport,
-			PixelShader,
-			PassParameters);
+			RDG_EVENT_NAME("HistogramConvert"),
+			ComputeShader,
+			PassParameters,
+			FIntVector(2, 1, 1));
 	}
 
 	return HistogramReduceTexture;
 }
-
-FIntPoint GetHistogramTexelsPerGroup()
-{
-	return FHistogramCS::TexelsPerThreadGroup;
-}
\ No newline at end of file
