Nukes FEyeAdaptationPS

diff --git a/Engine/Shaders/Private/PostProcessEyeAdaptation.usf b/Engine/Shaders/Private/PostProcessEyeAdaptation.usf
--- a/Engine/Shaders/Private/PostProcessEyeAdaptation.usf
+++ b/Engine/Shaders/Private/PostProcessEyeAdaptation.usf
@@ -103,7 +103,9 @@ float2 ComputeWeightedTextureAverageAlphaSubRegion(
 	return float2(Value, max(WeightTotal,0.0000000001f));
 }
 
-float4 EyeAdaptationCommon()
+#if COMPUTESHADER
+[numthreads(1, 1, 1)]
+void EyeAdaptationCS(uint2 DispatchThreadId : SV_DispatchThreadID)
 {
 	float4 OutColor = 0;
 
@@ -131,22 +133,7 @@ float4 EyeAdaptationCommon()
 	OutColor.y = MiddleGreyExposureCompensation * TargetExposureScale;
 	OutColor.z = AverageSceneLuminance;
 	OutColor.w = MiddleGreyExposureCompensation;
-
-	return OutColor;
-}
-
-void EyeAdaptationPS(float4 UVAndScreenPos : TEXCOORD0, 
-	out float4 OutColor : SV_Target0)
-{
-	OutColor = EyeAdaptationCommon();
-}
-
-#if COMPUTESHADER
-[numthreads(1, 1, 1)]
-void EyeAdaptationCS(uint2 DispatchThreadId : SV_DispatchThreadID)
-{
-	float4 OutColor = EyeAdaptationCommon();
-
+	
 	uint2 PixelPos = DispatchThreadId;
 	RWEyeAdaptationTexture[PixelPos] = OutColor;
 }
diff --git a/Engine/Source/Runtime/Renderer/Private/PostProcess/PostProcessEyeAdaptation.cpp b/Engine/Source/Runtime/Renderer/Private/PostProcess/PostProcessEyeAdaptation.cpp
--- a/Engine/Source/Runtime/Renderer/Private/PostProcess/PostProcessEyeAdaptation.cpp
+++ b/Engine/Source/Runtime/Renderer/Private/PostProcess/PostProcessEyeAdaptation.cpp
@@ -417,59 +417,23 @@ float GetEyeAdaptationFixedExposure(const FViewInfo& View)
 //! Histogram Eye Adaptation
 //////////////////////////////////////////////////////////////////////////
 
-class FEyeAdaptationShader : public FGlobalShader
+class FEyeAdaptationCS : public FGlobalShader
 {
-public:
+	DECLARE_GLOBAL_SHADER(FEyeAdaptationCS);
+	SHADER_USE_PARAMETER_STRUCT(FEyeAdaptationCS, FGlobalShader);
+
 	BEGIN_SHADER_PARAMETER_STRUCT(FParameters, )
 		SHADER_PARAMETER_STRUCT(FEyeAdaptationParameters, EyeAdaptation)
 		SHADER_PARAMETER_RDG_TEXTURE(Texture2D, HistogramTexture)
-		END_SHADER_PARAMETER_STRUCT()
+		SHADER_PARAMETER_RDG_TEXTURE_UAV(RWTexture2D<float4>, RWEyeAdaptationTexture)
+	END_SHADER_PARAMETER_STRUCT()
 
-		static const EPixelFormat OutputFormat = PF_A32B32G32R32F;
+	static const EPixelFormat OutputFormat = PF_A32B32G32R32F;
 
 	static bool ShouldCompilePermutation(const FGlobalShaderPermutationParameters& Parameters)
 	{
 		return IsFeatureLevelSupported(Parameters.Platform, ERHIFeatureLevel::SM5);
 	}
-
-	static void ModifyCompilationEnvironment(const FGlobalShaderPermutationParameters& Parameters, FShaderCompilerEnvironment& OutEnvironment)
-	{
-		FGlobalShader::ModifyCompilationEnvironment(Parameters, OutEnvironment);
-		OutEnvironment.SetRenderTargetOutputFormat(0, OutputFormat);
-	}
-
-	FEyeAdaptationShader() = default;
-	FEyeAdaptationShader(const CompiledShaderInitializerType& Initializer)
-		: FGlobalShader(Initializer)
-	{}
-};
-
-class FEyeAdaptationPS : public FEyeAdaptationShader
-{
-	using Super = FEyeAdaptationShader;
-public:
-	DECLARE_GLOBAL_SHADER(FEyeAdaptationPS);
-	SHADER_USE_PARAMETER_STRUCT(FEyeAdaptationPS, Super);
-
-	BEGIN_SHADER_PARAMETER_STRUCT(FParameters, )
-		SHADER_PARAMETER_STRUCT_INCLUDE(Super::FParameters, Base)
-		RENDER_TARGET_BINDING_SLOTS()
-		END_SHADER_PARAMETER_STRUCT()
-};
-
-IMPLEMENT_GLOBAL_SHADER(FEyeAdaptationPS, "/Engine/Private/PostProcessEyeAdaptation.usf", "EyeAdaptationPS", SF_Pixel);
-
-class FEyeAdaptationCS : public FEyeAdaptationShader
-{
-	using Super = FEyeAdaptationShader;
-public:
-	DECLARE_GLOBAL_SHADER(FEyeAdaptationCS);
-	SHADER_USE_PARAMETER_STRUCT(FEyeAdaptationCS, Super);
-
-	BEGIN_SHADER_PARAMETER_STRUCT(FParameters, )
-		SHADER_PARAMETER_STRUCT_INCLUDE(Super::FParameters, Base)
-		SHADER_PARAMETER_RDG_TEXTURE_UAV(RWTexture2D<float4>, RWEyeAdaptationTexture)
-		END_SHADER_PARAMETER_STRUCT()
 };
 
 IMPLEMENT_GLOBAL_SHADER(FEyeAdaptationCS, "/Engine/Private/PostProcessEyeAdaptation.usf", "EyeAdaptationCS", SF_Compute);
@@ -484,42 +448,18 @@ FRDGTextureRef AddHistogramEyeAdaptationPass(
 
 	FRDGTextureRef OutputTexture = GraphBuilder.RegisterExternalTexture(View.GetEyeAdaptationTexture(GraphBuilder.RHICmdList), ERenderTargetTexture::Targetable, ERDGTextureFlags::MultiFrame);
 
-	FEyeAdaptationShader::FParameters PassBaseParameters;
-	PassBaseParameters.EyeAdaptation = GetEyeAdaptationParameters(View, ERHIFeatureLevel::SM5);
-	PassBaseParameters.HistogramTexture = HistogramTexture;
-
-	if (View.bUseComputePasses)
-	{
-		FEyeAdaptationCS::FParameters* PassParameters = GraphBuilder.AllocParameters<FEyeAdaptationCS::FParameters>();
-		PassParameters->Base = PassBaseParameters;
-		PassParameters->RWEyeAdaptationTexture = GraphBuilder.CreateUAV(OutputTexture);
-
-		TShaderMapRef<FEyeAdaptationCS> ComputeShader(View.ShaderMap);
-
-		FComputeShaderUtils::AddPass(
-			GraphBuilder,
-			RDG_EVENT_NAME("HistogramEyeAdaptation (CS)"),
-			ComputeShader,
-			PassParameters,
-			FIntVector(1, 1, 1));
-	}
-	else
-	{
-		FEyeAdaptationPS::FParameters* PassParameters = GraphBuilder.AllocParameters<FEyeAdaptationPS::FParameters>();
-		PassParameters->Base = PassBaseParameters;
-		PassParameters->RenderTargets[0] = FRenderTargetBinding(OutputTexture, ERenderTargetLoadAction::ENoAction);
-
-		TShaderMapRef<FEyeAdaptationPS> PixelShader(View.ShaderMap);
+	FEyeAdaptationCS::FParameters* PassParameters = GraphBuilder.AllocParameters<FEyeAdaptationCS::FParameters>();
+	PassParameters->EyeAdaptation = GetEyeAdaptationParameters(View, ERHIFeatureLevel::SM5);
+	PassParameters->HistogramTexture = HistogramTexture;
+	PassParameters->RWEyeAdaptationTexture = GraphBuilder.CreateUAV(OutputTexture);
 
-		AddDrawScreenPass(
-			GraphBuilder,
-			RDG_EVENT_NAME("HistogramEyeAdaptation (PS)"),
-			View,
-			FScreenPassTextureViewport(OutputTexture),
-			FScreenPassTextureViewport(HistogramTexture),
-			PixelShader,
-			PassParameters);
-	}
+	TShaderMapRef<FEyeAdaptationCS> ComputeShader(View.ShaderMap);
+	FComputeShaderUtils::AddPass(
+		GraphBuilder,
+		RDG_EVENT_NAME("HistogramEyeAdaptation (CS)"),
+		ComputeShader,
+		PassParameters,
+		FIntVector(1, 1, 1));
 
 	return OutputTexture;
 }
