Add UAVFormat field to FPooledRenderTargetDesc to allow selection of an aliasing format when retrieving surfaces from the pool
 - Used by the VT system to alias DXT surfaces with integer formats.

diff --git a/Engine/Source/Runtime/RenderCore/Private/RenderTargetPool.cpp b/Engine/Source/Runtime/RenderCore/Private/RenderTargetPool.cpp
--- a/Engine/Source/Runtime/RenderCore/Private/RenderTargetPool.cpp
+++ b/Engine/Source/Runtime/RenderCore/Private/RenderTargetPool.cpp
@@ -92,6 +92,7 @@ static uint64 GetTypeHash(FPooledRenderTargetDesc Desc)
 		sizeof(FPooledRenderTargetDesc::Flags) +
 		sizeof(FPooledRenderTargetDesc::TargetableFlags) +
 		sizeof(FPooledRenderTargetDesc::Format) +
+		sizeof(FPooledRenderTargetDesc::UAVFormat) +
 		sizeof(FPooledRenderTargetDesc::Extent) +
 		sizeof(FPooledRenderTargetDesc::Depth) +
 		sizeof(FPooledRenderTargetDesc::ArraySize) +
@@ -503,9 +504,13 @@ Done:
 
 		if ((Desc.TargetableFlags & TexCreate_UAV))
 		{
+			EPixelFormat AliasFormat = Desc.UAVFormat != PF_Unknown
+				? Desc.UAVFormat
+				: Desc.Format;
+
 			// The render target desc is invalid if a UAV is requested with an RHI that doesn't support the high-end feature level.
 			check(GMaxRHIFeatureLevel >= ERHIFeatureLevel::SM5 || GMaxRHIFeatureLevel == ERHIFeatureLevel::ES3_1);
-			Found->RenderTargetItem.UAV = RHICreateUnorderedAccessView(Found->RenderTargetItem.TargetableTexture, 0);
+			Found->RenderTargetItem.UAV = RHICreateUnorderedAccessView(Found->RenderTargetItem.TargetableTexture, 0, AliasFormat);
 		}
 
 		AllocationLevelInKB += ComputeSizeInKB(*Found);
diff --git a/Engine/Source/Runtime/RenderCore/Public/RendererInterface.h b/Engine/Source/Runtime/RenderCore/Public/RendererInterface.h
--- a/Engine/Source/Runtime/RenderCore/Public/RendererInterface.h
+++ b/Engine/Source/Runtime/RenderCore/Public/RendererInterface.h
@@ -235,6 +235,7 @@ public:
 			&& LhsFlags == RhsFlags
 			&& TargetableFlags == rhs.TargetableFlags
 			&& Format == rhs.Format
+			&& UAVFormat == rhs.UAVFormat
 			&& Extent == rhs.Extent
 			&& Depth == rhs.Depth
 			&& ArraySize == rhs.ArraySize
@@ -378,6 +379,8 @@ public:
 	ETextureCreateFlags TargetableFlags = TexCreate_None;
 	/** Texture format e.g. PF_B8G8R8A8 */
 	EPixelFormat Format = PF_Unknown;
+	/** Texture format used when creating the UAV (if TexCreate_UAV is also passed in TargetableFlags, ignored otherwise). PF_Unknown == use default (same as Format) */
+	EPixelFormat UAVFormat = PF_Unknown;
 	/** In pixels, (0,0) if not set, (x,0) for cube maps, todo: make 3d int vector for volume textures */
 	FIntPoint Extent = FIntPoint::ZeroValue;
 	/** 0, unless it's texture array or volume texture */
diff --git a/Engine/Source/Runtime/Renderer/Private/VT/VirtualTexturePhysicalSpace.cpp b/Engine/Source/Runtime/Renderer/Private/VT/VirtualTexturePhysicalSpace.cpp
--- a/Engine/Source/Runtime/Renderer/Private/VT/VirtualTexturePhysicalSpace.cpp
+++ b/Engine/Source/Runtime/Renderer/Private/VT/VirtualTexturePhysicalSpace.cpp
@@ -104,7 +104,7 @@ void FVirtualTexturePhysicalSpace::InitRHI()
 		
 		// Allocate physical texture from the render target pool
 		const uint32 TextureSize = GetTextureSize();
-		const FPooledRenderTargetDesc Desc = FPooledRenderTargetDesc::Create2DDesc(
+		FPooledRenderTargetDesc Desc = FPooledRenderTargetDesc::Create2DDesc(
 			FIntPoint(TextureSize, TextureSize),
 			FormatSRV,
 			FClearValueBinding::None,
@@ -113,6 +113,11 @@ void FVirtualTexturePhysicalSpace::InitRHI()
 			(bCreateAliasedUAV || FormatSRV == PF_A32B32G32R32F) ? TexCreate_ShaderResource | TexCreate_UAV : TexCreate_ShaderResource,
 			false);
 
+		if (bCreateAliasedUAV)
+		{
+			Desc.UAVFormat = FormatUAV;
+		}
+
 		GRenderTargetPool.FindFreeElement(RHICmdList, Desc, PooledRenderTarget[Layer], TEXT("VirtualPhysicalTexture"));
 		FRHITexture* TextureRHI = PooledRenderTarget[Layer]->GetRenderTargetItem().ShaderResourceTexture;
 
@@ -127,7 +132,7 @@ void FVirtualTexturePhysicalSpace::InitRHI()
 
 		if (bCreateAliasedUAV)
 		{
-			TextureUAV[Layer] = RHICreateUnorderedAccessView(TextureRHI, 0, FormatUAV);
+			TextureUAV[Layer] = PooledRenderTarget[Layer]->GetRenderTargetItem().UAV;
 		}
 	}
 }
