Fix CIS unity build.

diff --git a/Engine/Source/Runtime/Renderer/Private/HairStrands/HairStrandsData.h b/Engine/Source/Runtime/Renderer/Private/HairStrands/HairStrandsData.h
--- a/Engine/Source/Runtime/Renderer/Private/HairStrands/HairStrandsData.h
+++ b/Engine/Source/Runtime/Renderer/Private/HairStrands/HairStrandsData.h
@@ -12,6 +12,7 @@
 
 class FPrimitiveSceneProxy;
 class FViewInfo;
+class FRHIGPUBufferReadback;
 struct FMeshBatch;
 
 ////////////////////////////////////////////////////////////////////////////////////
