Change hair rendering from legacy dynamic mesh pass drawing to GPU-Scene Instance Culling compatible

diff --git a/Engine/Source/Runtime/Renderer/Private/HairStrands/HairStrandsRasterCommon.cpp b/Engine/Source/Runtime/Renderer/Private/HairStrands/HairStrandsRasterCommon.cpp
--- a/Engine/Source/Runtime/Renderer/Private/HairStrands/HairStrandsRasterCommon.cpp
+++ b/Engine/Source/Runtime/Renderer/Private/HairStrands/HairStrandsRasterCommon.cpp
@@ -286,18 +286,13 @@ void AddHairStrandsRasterPass(
 		ViewInfo->CachedViewUniformShaderParameters->ViewForward = SavedViewForward;
 	}
 
-	GraphBuilder.AddPass(
-		GetPassName(RasterPassType),
-		PassParameters,
-		ERDGPassFlags::Raster,
-		[PassParameters, Scene = Scene, ViewInfo, RasterPassType, &PrimitiveSceneInfos, ViewportRect, HairRenderInfo, HairRenderInfoBits, RasterDirection](FRHICommandListImmediate& RHICmdList)
+	AddSimpleMeshPass(GraphBuilder, PassParameters, Scene, *ViewInfo, &InstanceCullingManager, GetPassName(RasterPassType), ViewportRect,
+		[PassParameters, Scene = Scene, ViewInfo, RasterPassType, &PrimitiveSceneInfos, HairRenderInfo, HairRenderInfoBits, RasterDirection](FDynamicPassMeshDrawListContext* ShadowContext)
 	{
 		SCOPE_CYCLE_COUNTER(STAT_RenderPerObjectShadowDepthsTime);
 
 		FMeshPassProcessorRenderState DrawRenderState;
 
-		RHICmdList.SetViewport(ViewportRect.Min.X, ViewportRect.Min.Y, 0.0f, ViewportRect.Max.X, ViewportRect.Max.Y, 1.0f);
-
 		if (RasterPassType == EHairStrandsRasterPassType::DeepOpacityMap)
 		{
 			DrawRenderState.SetBlendState(TStaticBlendState<
@@ -313,13 +308,7 @@ void AddHairStrandsRasterPass(
 			DrawRenderState.SetDepthStencilState(TStaticDepthStencilState<true, CF_DepthNearOrEqual>::GetRHI());
 		}
 
-		FDynamicMeshDrawCommandStorage DynamicMeshDrawCommandStorage; // << Were would thid be stored?
-		FMeshCommandOneFrameArray VisibleMeshDrawCommands;
-		FGraphicsMinimalPipelineStateSet GraphicsMinimalPipelineStateSet;
-		bool bNeedsInitialization;
-		FDynamicPassMeshDrawListContext ShadowContext(DynamicMeshDrawCommandStorage, VisibleMeshDrawCommands, GraphicsMinimalPipelineStateSet, bNeedsInitialization);
-
-		FHairRasterMeshProcessor HairRasterMeshProcessor(Scene, ViewInfo /* is a SceneView */, DrawRenderState, &ShadowContext, RasterPassType);
+		FHairRasterMeshProcessor HairRasterMeshProcessor(Scene, ViewInfo /* is a SceneView */, DrawRenderState, ShadowContext, RasterPassType);
 
 		for (const FHairStrandsMacroGroupData::PrimitiveInfo& PrimitiveInfo : PrimitiveSceneInfos)
 		{
@@ -328,13 +317,6 @@ void AddHairStrandsRasterPass(
 			const uint64 BatchElementMask = ~0ull;
 			HairRasterMeshProcessor.AddMeshBatch(MeshBatch, BatchElementMask, PrimitiveInfo.PrimitiveSceneProxy, -1 , bCullingEnable);
 		}
-
-		if (VisibleMeshDrawCommands.Num() > 0)
-		{
-			FRHIBuffer* PrimitiveIdVertexBuffer = nullptr;
-			SortAndMergeDynamicPassMeshDrawCommands(ViewInfo->GetFeatureLevel(), VisibleMeshDrawCommands, DynamicMeshDrawCommandStorage, PrimitiveIdVertexBuffer, 1, ViewInfo->DynamicPrimitiveCollector.GetPrimitiveIdRange());
-			SubmitMeshDrawCommands(VisibleMeshDrawCommands, GraphicsMinimalPipelineStateSet, PrimitiveIdVertexBuffer, 0, false, 1, RHICmdList);
-		}
 	});
 }
 
diff --git a/Engine/Source/Runtime/Renderer/Private/HairStrands/HairStrandsVisibility.cpp b/Engine/Source/Runtime/Renderer/Private/HairStrands/HairStrandsVisibility.cpp
--- a/Engine/Source/Runtime/Renderer/Private/HairStrands/HairStrandsVisibility.cpp
+++ b/Engine/Source/Runtime/Renderer/Private/HairStrands/HairStrandsVisibility.cpp
@@ -775,26 +775,16 @@ static void AddHairMaterialGBufferPass(
 		ERenderTargetLoadAction::ELoad,
 		FExclusiveDepthStencil::DepthWrite_StencilNop);
 
-	GraphBuilder.AddPass(
-		RDG_EVENT_NAME("HairStrandsMaterialGBufferPass"),
-		PassParameters,
-		ERDGPassFlags::Raster,
-		[PassParameters, Scene = Scene, ViewInfo, &MacroGroupDatas, Resolution](FRHICommandListImmediate& RHICmdList)
+	AddSimpleMeshPass(GraphBuilder, PassParameters, Scene, *ViewInfo, &InstanceCullingManager, RDG_EVENT_NAME("HairStrandsMaterialGBufferPass"), FIntRect(0, 0, Resolution.X, Resolution.Y),
+		[PassParameters, Scene = Scene, ViewInfo, &MacroGroupDatas, Resolution](FDynamicPassMeshDrawListContext* ShadowContext)
 		{
 			FMeshPassProcessorRenderState DrawRenderState;
 
 			{
-				RHICmdList.SetViewport(0, 0, 0.0f, Resolution.X, Resolution.Y, 1.0f);
 				DrawRenderState.SetBlendState(TStaticBlendState<>::GetRHI());
 				DrawRenderState.SetDepthStencilState(TStaticDepthStencilState <true, CF_Always> ::GetRHI());
 
-				FDynamicMeshDrawCommandStorage DynamicMeshDrawCommandStorage;
-				FMeshCommandOneFrameArray VisibleMeshDrawCommands;
-				FGraphicsMinimalPipelineStateSet PipelineStateSet;
-				// @todo loadtime arnes: do we need to pass this along to somewhere?
-				bool NeedsShaderInitialization;
-				FDynamicPassMeshDrawListContext ShadowContext(DynamicMeshDrawCommandStorage, VisibleMeshDrawCommands, PipelineStateSet, NeedsShaderInitialization);
-				FHairMaterialGBufferProcessor MeshProcessor(Scene, ViewInfo, DrawRenderState, &ShadowContext);
+				FHairMaterialGBufferProcessor MeshProcessor(Scene, ViewInfo, DrawRenderState, ShadowContext);
 
 				for (const FHairStrandsMacroGroupData& MacroGroupData : MacroGroupDatas)
 				{
@@ -805,13 +795,6 @@ static void AddHairMaterialGBufferPass(
 						MeshProcessor.AddMeshBatch(MeshBatch, BatchElementMask, PrimitiveInfo.PrimitiveSceneProxy, -1, MacroGroupData.MacroGroupId, PrimitiveInfo.MaterialId);
 					}
 				}
-
-				if (VisibleMeshDrawCommands.Num() > 0)
-				{
-					FRHIBuffer* PrimitiveIdVertexBuffer = nullptr;
-					SortAndMergeDynamicPassMeshDrawCommands(ViewInfo->GetFeatureLevel(), VisibleMeshDrawCommands, DynamicMeshDrawCommandStorage, PrimitiveIdVertexBuffer, 1, ViewInfo->DynamicPrimitiveCollector.GetPrimitiveIdRange());
-					SubmitMeshDrawCommands(VisibleMeshDrawCommands, PipelineStateSet, PrimitiveIdVertexBuffer, 0, false, 1, RHICmdList);
-				}
 			}
 		});
 }
@@ -1257,16 +1240,12 @@ static FMaterialPassOutput AddHairMaterialPass(
 	}
 	
 
-	GraphBuilder.AddPass(
-		RDG_EVENT_NAME("HairStrandsMaterialPass"),
-		PassParameters,
-		ERDGPassFlags::Raster,
-		[PassParameters, Scene = Scene, ViewInfo, &MacroGroupDatas, MaxNodeCount, Resolution, NodeGroupSize, bUpdateSampleCoverage, bOutputEmissive](FRHICommandListImmediate& RHICmdList)
+	AddSimpleMeshPass(GraphBuilder, PassParameters, Scene, *ViewInfo, &InstanceCullingManager, RDG_EVENT_NAME("HairStrandsMaterialPass"), FIntRect(0, 0, Resolution.X, Resolution.Y),
+		[PassParameters, Scene = Scene, ViewInfo, &MacroGroupDatas, MaxNodeCount, Resolution, NodeGroupSize, bUpdateSampleCoverage, bOutputEmissive](FDynamicPassMeshDrawListContext* ShadowContext)
 	{
 		FMeshPassProcessorRenderState DrawRenderState;
 
 		{
-			RHICmdList.SetViewport(0, 0, 0.0f, Resolution.X, Resolution.Y, 1.0f);
 			if (bOutputEmissive)
 			{
 				DrawRenderState.SetBlendState(TStaticBlendState<CW_RGBA, BO_Add, BF_One, BF_One, BO_Add, BF_One, BF_Zero>::GetRHI());
@@ -1276,14 +1255,7 @@ static FMaterialPassOutput AddHairMaterialPass(
 				DrawRenderState.SetBlendState(TStaticBlendState<>::GetRHI());
 			}
 			DrawRenderState.SetDepthStencilState(TStaticDepthStencilState <false, CF_Always> ::GetRHI());
-			
-			FDynamicMeshDrawCommandStorage DynamicMeshDrawCommandStorage;
-			FMeshCommandOneFrameArray VisibleMeshDrawCommands;
-			FGraphicsMinimalPipelineStateSet PipelineStateSet;
-			// @todo loadtime arnes: do we need to pass this along to somewhere?
-			bool NeedsShaderInitialization;
-			FDynamicPassMeshDrawListContext ShadowContext(DynamicMeshDrawCommandStorage, VisibleMeshDrawCommands, PipelineStateSet, NeedsShaderInitialization);
-			FHairMaterialProcessor MeshProcessor(Scene, ViewInfo, DrawRenderState, &ShadowContext);
+			FHairMaterialProcessor MeshProcessor(Scene, ViewInfo, DrawRenderState, ShadowContext);
 
 			for (const FHairStrandsMacroGroupData& MacroGroupData : MacroGroupDatas)
 			{
@@ -1294,13 +1266,6 @@ static FMaterialPassOutput AddHairMaterialPass(
 					MeshProcessor.AddMeshBatch(MeshBatch, BatchElementMask, PrimitiveInfo.PrimitiveSceneProxy, -1, MacroGroupData.MacroGroupId, PrimitiveInfo.MaterialId);
 				}
 			}
-
-			if (VisibleMeshDrawCommands.Num() > 0)
-			{
-				FRHIBuffer* PrimitiveIdVertexBuffer = nullptr;
-				SortAndMergeDynamicPassMeshDrawCommands(ViewInfo->GetFeatureLevel(), VisibleMeshDrawCommands, DynamicMeshDrawCommandStorage, PrimitiveIdVertexBuffer, 1, ViewInfo->DynamicPrimitiveCollector.GetPrimitiveIdRange());
-				SubmitMeshDrawCommands(VisibleMeshDrawCommands, PipelineStateSet, PrimitiveIdVertexBuffer, 0, false, 1, RHICmdList);
-			}
 		}
 	});
 
@@ -2468,19 +2433,14 @@ static void AddHairVisibilityCommonPass(
 		PassParameters->View = ViewInfo->ViewUniformBuffer;
 	}
 
-	GraphBuilder.AddPass(
-		GetPassName(),
-		PassParameters,
-		ERDGPassFlags::Raster,
-		[PassParameters, Scene = Scene, ViewInfo, &MacroGroupDatas, RenderMode](FRHICommandListImmediate& RHICmdList)
+	AddSimpleMeshPass(GraphBuilder, PassParameters, Scene, *ViewInfo, &InstanceCullingManager, GetPassName(), FIntRect(0, 0, ViewInfo->ViewRect.Width(), ViewInfo->ViewRect.Height()),
+		[PassParameters, Scene = Scene, ViewInfo, &MacroGroupDatas, RenderMode](FDynamicPassMeshDrawListContext* ShadowContext)
 	{
-		check(RHICmdList.IsInsideRenderPass());
 		check(IsInRenderingThread());
 
 		FMeshPassProcessorRenderState DrawRenderState;
 
 		{
-			RHICmdList.SetViewport(0, 0, 0.0f, ViewInfo->ViewRect.Width(), ViewInfo->ViewRect.Height(), 1.0f);
 			if (RenderMode == HairVisibilityRenderMode_MSAA_Visibility)
 			{
 				DrawRenderState.SetBlendState(TStaticBlendState<CW_RGBA, BO_Add, BF_One, BF_Zero, BO_Add, BF_One, BF_Zero>::GetRHI());
@@ -2504,12 +2464,7 @@ static void AddHairVisibilityCommonPass(
 				DrawRenderState.SetDepthStencilState(TStaticDepthStencilState<false, CF_DepthNearOrEqual>::GetRHI());
 			}
 
-			FDynamicMeshDrawCommandStorage DynamicMeshDrawCommandStorage;
-			FMeshCommandOneFrameArray VisibleMeshDrawCommands;
-			FGraphicsMinimalPipelineStateSet PipelineStateSet;
-			bool NeedsShaderInitialization;
-			FDynamicPassMeshDrawListContext ShadowContext(DynamicMeshDrawCommandStorage, VisibleMeshDrawCommands, PipelineStateSet, NeedsShaderInitialization);
-			FHairVisibilityProcessor MeshProcessor(Scene, ViewInfo, DrawRenderState, RenderMode, &ShadowContext);
+			FHairVisibilityProcessor MeshProcessor(Scene, ViewInfo, DrawRenderState, RenderMode, ShadowContext);
 			
 			for (const FHairStrandsMacroGroupData& MacroGroupData : MacroGroupDatas)
 			{
@@ -2520,13 +2475,6 @@ static void AddHairVisibilityCommonPass(
 					MeshProcessor.AddMeshBatch(MeshBatch, BatchElementMask, PrimitiveInfo.PrimitiveSceneProxy, -1, MacroGroupData.MacroGroupId, PrimitiveInfo.MaterialId, PrimitiveInfo.IsCullingEnable());
 				}
 			}
-
-			if (VisibleMeshDrawCommands.Num() > 0)
-			{
-				FRHIBuffer* PrimitiveIdVertexBuffer = nullptr;
-				SortAndMergeDynamicPassMeshDrawCommands(ViewInfo->GetFeatureLevel(), VisibleMeshDrawCommands, DynamicMeshDrawCommandStorage, PrimitiveIdVertexBuffer, 1, ViewInfo->DynamicPrimitiveCollector.GetPrimitiveIdRange());
-				SubmitMeshDrawCommands(VisibleMeshDrawCommands, PipelineStateSet, PrimitiveIdVertexBuffer, 0, false, 1, RHICmdList);
-			}
 		}
 	});
 }
