Fix validation error when running in Low scalability settings.

diff --git a/Engine/Source/Runtime/Renderer/Private/IndirectLightRendering.cpp b/Engine/Source/Runtime/Renderer/Private/IndirectLightRendering.cpp
--- a/Engine/Source/Runtime/Renderer/Private/IndirectLightRendering.cpp
+++ b/Engine/Source/Runtime/Renderer/Private/IndirectLightRendering.cpp
@@ -865,10 +865,12 @@ void FDeferredShadingSceneRenderer::RenderDiffuseIndirectAndAmbientOcclusion(
 			}
 		}
 
+		bool bWritableAmbientOcclusionMask = true;
 		if (ViewPipelineState.AmbientOcclusionMethod == EAmbientOcclusionMethod::Disabled)
 		{
 			ensure(!HasBeenProduced(SceneTextures.ScreenSpaceAO));
 			AmbientOcclusionMask = nullptr;
+			bWritableAmbientOcclusionMask = false;
 		}
 		else if (ViewPipelineState.AmbientOcclusionMethod == EAmbientOcclusionMethod::RTAO)
 		{
@@ -892,11 +894,13 @@ void FDeferredShadingSceneRenderer::RenderDiffuseIndirectAndAmbientOcclusion(
 			else
 			{
 				AmbientOcclusionMask = GetScreenSpaceAOFallback(SystemTextures);
+				bWritableAmbientOcclusionMask = false;
 			}
 		}
 		else
 		{
 			unimplemented();
+			bWritableAmbientOcclusionMask = false;
 		}
 
 		// Extract the dynamic AO for application of AO beyond RenderDiffuseIndirectAndAmbientOcclusion()
@@ -907,8 +911,9 @@ void FDeferredShadingSceneRenderer::RenderDiffuseIndirectAndAmbientOcclusion(
 			SceneTextures.ScreenSpaceAO = AmbientOcclusionMask;
 		}
 
-		if (HairStrands::HasViewHairStrandsData(View) && (ViewPipelineState.AmbientOcclusionMethod == EAmbientOcclusionMethod::SSGI || ViewPipelineState.AmbientOcclusionMethod == EAmbientOcclusionMethod::SSAO))
+		if (HairStrands::HasViewHairStrandsData(View) && (ViewPipelineState.AmbientOcclusionMethod == EAmbientOcclusionMethod::SSGI || ViewPipelineState.AmbientOcclusionMethod == EAmbientOcclusionMethod::SSAO) && bWritableAmbientOcclusionMask)
 		{
+			if (!HasBeenProduced(SceneTextures.ScreenSpaceAO))
 			RenderHairStrandsAmbientOcclusion(
 				GraphBuilder,
 				View,
