Fix typo.

diff --git a/Engine/Source/Runtime/Renderer/Private/IndirectLightRendering.cpp b/Engine/Source/Runtime/Renderer/Private/IndirectLightRendering.cpp
--- a/Engine/Source/Runtime/Renderer/Private/IndirectLightRendering.cpp
+++ b/Engine/Source/Runtime/Renderer/Private/IndirectLightRendering.cpp
@@ -913,7 +913,6 @@ void FDeferredShadingSceneRenderer::RenderDiffuseIndirectAndAmbientOcclusion(
 
 		if (HairStrands::HasViewHairStrandsData(View) && (ViewPipelineState.AmbientOcclusionMethod == EAmbientOcclusionMethod::SSGI || ViewPipelineState.AmbientOcclusionMethod == EAmbientOcclusionMethod::SSAO) && bWritableAmbientOcclusionMask)
 		{
-			if (!HasBeenProduced(SceneTextures.ScreenSpaceAO))
 			RenderHairStrandsAmbientOcclusion(
 				GraphBuilder,
 				View,
