Fix crash when hair voxelization is disabled.

diff --git a/Engine/Source/Runtime/Renderer/Private/HairStrands/HairStrandsVoxelization.cpp b/Engine/Source/Runtime/Renderer/Private/HairStrands/HairStrandsVoxelization.cpp
--- a/Engine/Source/Runtime/Renderer/Private/HairStrands/HairStrandsVoxelization.cpp
+++ b/Engine/Source/Runtime/Renderer/Private/HairStrands/HairStrandsVoxelization.cpp
@@ -1066,6 +1066,75 @@ static FHairStrandsVoxelResources AllocateVirtualVoxelResources(
 	return Out;
 }
 
+static FHairStrandsVoxelResources AllocateDummyVirtualVoxelResources(
+	FRDGBuilder& GraphBuilder,
+	const FViewInfo& View,
+	FHairStrandsMacroGroupDatas& MacroGroupDatas)
+{
+	FHairStrandsVoxelResources Out;
+
+	Out.Parameters.Common.PageCountResolution		= FIntVector(0,0,0);
+	Out.Parameters.Common.PageCount					= 0;
+	Out.Parameters.Common.VoxelWorldSize			= 1;
+	Out.Parameters.Common.PageResolution			= 1;
+	Out.Parameters.Common.PageTextureResolution		= FIntVector(1, 1, 1);
+
+	Out.Parameters.Common.DensityScale				= 0;
+	Out.Parameters.Common.DensityScale_AO			= 0;
+	Out.Parameters.Common.DensityScale_Shadow		= 0;
+	Out.Parameters.Common.DensityScale_Transmittance= 0;
+	Out.Parameters.Common.DensityScale_Environment	= 0;
+	Out.Parameters.Common.DensityScale_Raytracing   = 0;
+
+	Out.Parameters.Common.DepthBiasScale_Shadow			= 0;
+	Out.Parameters.Common.DepthBiasScale_Transmittance	= 0;
+	Out.Parameters.Common.DepthBiasScale_Environment	= 0;
+
+	Out.Parameters.Common.SteppingScale_Shadow		  = 1;
+	Out.Parameters.Common.SteppingScale_Transmittance = 1;
+	Out.Parameters.Common.SteppingScale_Environment	  = 1;
+	Out.Parameters.Common.SteppingScale_Raytracing	  = 1;
+
+	Out.Parameters.Common.NodeDescCount				= 0;
+	Out.Parameters.Common.IndirectDispatchGroupSize = 64;	
+	Out.Parameters.Common.Raytracing_ShadowOcclusionThreshold	= 1;
+	Out.Parameters.Common.Raytracing_SkyOcclusionThreshold		= 1;
+
+	Out.Parameters.Common.HairCoveragePixelRadiusAtDepth1	= ComputeMinStrandRadiusAtDepth1(FIntPoint(View.ViewRect.Width(), View.ViewRect.Height()), View.FOV, 1/*SampleCount*/, 1/*RasterizationScale*/).Primary;
+	Out.Parameters.Common.PageIndexCount = 0;
+	for (FHairStrandsMacroGroupData& MacroGroup : MacroGroupDatas)
+	{
+		MacroGroup.VirtualVoxelNodeDesc.PageIndexResolution = FIntVector(0,0,0);
+	}
+
+	FRDGBufferRef Dummy4BytesBuffer = GSystemTextures.GetDefaultBuffer(GraphBuilder, 4, 0u);
+	FRDGBufferRef Dummy8BytesBuffer = GSystemTextures.GetDefaultBuffer(GraphBuilder, 8, 0u);
+	FRDGBufferRef DummyStructBuffer = GSystemTextures.GetDefaultStructuredBuffer(GraphBuilder, sizeof(FPackedVirtualVoxelNodeDesc), 0u);
+
+	Out.PageIndexBuffer				= Dummy4BytesBuffer;
+	Out.PageIndexOccupancyBuffer	= Dummy8BytesBuffer;
+	Out.PageIndexCoordBuffer		= Dummy4BytesBuffer;
+	Out.NodeDescBuffer				= DummyStructBuffer;
+	Out.IndirectArgsBuffer			= Dummy4BytesBuffer;
+	Out.PageIndexGlobalCounter		= Dummy4BytesBuffer;
+	Out.VoxelizationViewInfoBuffer	= DummyStructBuffer;
+	Out.PageTexture					= GSystemTextures.GetDefaultTexture(GraphBuilder, ETextureDimension::Texture3D, PF_R32_UINT, 0u);
+
+	Out.Parameters.Common.PageIndexBuffer			= GraphBuilder.CreateSRV(Out.PageIndexBuffer, PF_R32_UINT);
+	Out.Parameters.Common.PageIndexOccupancyBuffer	= GraphBuilder.CreateSRV(Out.PageIndexOccupancyBuffer, PF_R32G32_UINT);
+	Out.Parameters.Common.PageIndexCoordBuffer		= GraphBuilder.CreateSRV(Out.PageIndexCoordBuffer, PF_R8G8B8A8_UINT);
+	Out.Parameters.Common.NodeDescBuffer			= GraphBuilder.CreateSRV(Out.NodeDescBuffer); 
+	Out.Parameters.PageTexture						= Out.PageTexture;
+	
+	if (Out.PageIndexBuffer && Out.NodeDescBuffer)
+	{
+		FVirtualVoxelParameters* Parameters = GraphBuilder.AllocParameters<FVirtualVoxelParameters>();
+		*Parameters = Out.Parameters;
+		Out.UniformBuffer = GraphBuilder.CreateUniformBuffer(Parameters);
+	}
+
+	return Out;
+}
 
 static FRDGBufferRef IndirectVoxelPageClear(
 	FRDGBuilder& GraphBuilder,
@@ -1521,16 +1590,16 @@ void VoxelizeHairStrands(
 	const FScene* Scene, 
 	FViewInfo& View,
 	FInstanceCullingManager& InstanceCullingManager)
-{
-	if (!IsHairStrandsVoxelizationEnable())
-		return;
-
+{	
 	FHairStrandsMacroGroupDatas& MacroGroupDatas = View.HairStrandsViewData.MacroGroupDatas;
 	FHairStrandsVoxelResources& VirtualVoxelResources = View.HairStrandsViewData.VirtualVoxelResources;
 	FHairStrandsMacroGroupResources& MacroGroupResources = View.HairStrandsViewData.MacroGroupResources;
 
-	if (MacroGroupDatas.Num() == 0)
+	if (!IsHairStrandsVoxelizationEnable() || MacroGroupDatas.Num() == 0)
+	{
+		VirtualVoxelResources = AllocateDummyVirtualVoxelResources(GraphBuilder, View, MacroGroupDatas);
 		return;
+	}
 
 	DECLARE_GPU_STAT(HairStrandsVoxelization);
 	RDG_EVENT_SCOPE(GraphBuilder, "HairStrandsVoxelization");
diff --git a/Engine/Source/Runtime/Renderer/Private/LightRendering.cpp b/Engine/Source/Runtime/Renderer/Private/LightRendering.cpp
--- a/Engine/Source/Runtime/Renderer/Private/LightRendering.cpp
+++ b/Engine/Source/Runtime/Renderer/Private/LightRendering.cpp
@@ -1941,10 +1941,13 @@ void FDeferredShadingSceneRenderer::RenderLights(
 				{
 					for (FViewInfo& View : Views)
 					{
-						FHairStrandsTransmittanceMaskData TransmittanceMaskData = DummyTransmittanceMaskData;
 						if (bDrawHairShadow && HairStrands::HasViewHairStrandsData(View))
 						{
-							TransmittanceMaskData = RenderHairStrandsTransmittanceMask(GraphBuilder, View, &LightSceneInfo, ScreenShadowMaskSubPixelTexture);
+							FHairStrandsTransmittanceMaskData TransmittanceMaskData = RenderHairStrandsTransmittanceMask(GraphBuilder, View, &LightSceneInfo, ScreenShadowMaskSubPixelTexture);							
+							if (TransmittanceMaskData.TransmittanceMask == nullptr)
+							{
+								TransmittanceMaskData = DummyTransmittanceMaskData;
+							}
 
 							// Note: ideally the light should still be evaluated for hair when not casting shadow, but for preserving the old behavior, and not adding 
 							// any perf. regression, we disable this light for hair rendering 
@@ -2476,7 +2479,7 @@ void FDeferredShadingSceneRenderer::RenderLightForHair(
 	FRDGTextureRef LightingChannelsTexture,
 	const FHairStrandsTransmittanceMaskData& InTransmittanceMaskData)
 {
-	const bool bHairRenderingEnabled = HairStrands::HasViewHairStrandsData(View) && (LightSceneInfo->Proxy->CastsHairStrandsDeepShadow() || IsHairStrandsVoxelizationEnable());
+	const bool bHairRenderingEnabled = HairStrands::HasViewHairStrandsData(View);
 	if (!bHairRenderingEnabled)
 	{
 		return;
