Fixed crash with Lumen Direct Lighting when FDeferredShadingSceneRenderer::ShouldPrepareDistanceFieldScene() returns false, as is the case on Intel.

diff --git a/Engine/Source/Runtime/Renderer/Private/DistanceFieldObjectManagement.cpp b/Engine/Source/Runtime/Renderer/Private/DistanceFieldObjectManagement.cpp
--- a/Engine/Source/Runtime/Renderer/Private/DistanceFieldObjectManagement.cpp
+++ b/Engine/Source/Runtime/Renderer/Private/DistanceFieldObjectManagement.cpp
@@ -124,21 +124,31 @@ FDistanceFieldAtlasParameters DistanceField::SetupAtlasParameters(const FDistanc
 {
 	FDistanceFieldAtlasParameters SceneParameters;
 
-	SceneParameters.SceneDistanceFieldAssetData = DistanceFieldSceneData.AssetDataBuffer.SRV;
-
-	SceneParameters.DistanceFieldIndirectionTable = DistanceFieldSceneData.IndirectionTable.SRV;
-	SceneParameters.DistanceFieldBrickTexture = DistanceFieldSceneData.DistanceFieldBrickVolumeTexture->GetRenderTargetItem().ShaderResourceTexture;
-	SceneParameters.DistanceFieldSampler = TStaticSamplerState<SF_Bilinear,AM_Clamp,AM_Clamp,AM_Clamp>::GetRHI();
-
-	SceneParameters.DistanceFieldBrickSize = FVector(DistanceField::BrickSize);
-	SceneParameters.DistanceFieldUniqueDataBrickSize = FVector(DistanceField::UniqueDataBrickSize);
-	SceneParameters.DistanceFieldBrickAtlasSizeInBricks = DistanceFieldSceneData.BrickTextureDimensionsInBricks;
-	SceneParameters.DistanceFieldBrickAtlasMask = DistanceFieldSceneData.BrickTextureDimensionsInBricks - FIntVector(1);
-	SceneParameters.DistanceFieldBrickAtlasSizeLog2 = FIntVector(
-		FMath::FloorLog2(DistanceFieldSceneData.BrickTextureDimensionsInBricks.X),
-		FMath::FloorLog2(DistanceFieldSceneData.BrickTextureDimensionsInBricks.Y),
-		FMath::FloorLog2(DistanceFieldSceneData.BrickTextureDimensionsInBricks.Z));
-	SceneParameters.DistanceFieldBrickAtlasTexelSize = FVector(1.0f) / FVector(DistanceFieldSceneData.BrickTextureDimensionsInBricks * DistanceField::BrickSize);
+	if (DistanceFieldSceneData.DistanceFieldBrickVolumeTexture)
+	{
+		SceneParameters.SceneDistanceFieldAssetData = DistanceFieldSceneData.AssetDataBuffer.SRV;
+
+		SceneParameters.DistanceFieldIndirectionTable = DistanceFieldSceneData.IndirectionTable.SRV;
+		SceneParameters.DistanceFieldBrickTexture = DistanceFieldSceneData.DistanceFieldBrickVolumeTexture->GetRenderTargetItem().ShaderResourceTexture;
+		SceneParameters.DistanceFieldSampler = TStaticSamplerState<SF_Bilinear,AM_Clamp,AM_Clamp,AM_Clamp>::GetRHI();
+
+		SceneParameters.DistanceFieldBrickSize = FVector(DistanceField::BrickSize);
+		SceneParameters.DistanceFieldUniqueDataBrickSize = FVector(DistanceField::UniqueDataBrickSize);
+		SceneParameters.DistanceFieldBrickAtlasSizeInBricks = DistanceFieldSceneData.BrickTextureDimensionsInBricks;
+		SceneParameters.DistanceFieldBrickAtlasMask = DistanceFieldSceneData.BrickTextureDimensionsInBricks - FIntVector(1);
+		SceneParameters.DistanceFieldBrickAtlasSizeLog2 = FIntVector(
+			FMath::FloorLog2(DistanceFieldSceneData.BrickTextureDimensionsInBricks.X),
+			FMath::FloorLog2(DistanceFieldSceneData.BrickTextureDimensionsInBricks.Y),
+			FMath::FloorLog2(DistanceFieldSceneData.BrickTextureDimensionsInBricks.Z));
+		SceneParameters.DistanceFieldBrickAtlasTexelSize = FVector(1.0f) / FVector(DistanceFieldSceneData.BrickTextureDimensionsInBricks * DistanceField::BrickSize);
+	}
+	else
+	{
+		SceneParameters.SceneDistanceFieldAssetData = nullptr;
+		SceneParameters.DistanceFieldIndirectionTable = nullptr;
+		SceneParameters.DistanceFieldBrickTexture = nullptr;
+		SceneParameters.DistanceFieldSampler = nullptr;
+	}
 
 	return SceneParameters;
 }
