[GPUScene] - Adding LLM tags for GPUScene.

diff --git a/Engine/Source/Runtime/Renderer/Private/GPUScene.cpp b/Engine/Source/Runtime/Renderer/Private/GPUScene.cpp
--- a/Engine/Source/Runtime/Renderer/Private/GPUScene.cpp
+++ b/Engine/Source/Runtime/Renderer/Private/GPUScene.cpp
@@ -18,6 +18,8 @@
 #include "Async/ParallelFor.h"
 #include "VirtualShadowMaps/VirtualShadowMapCacheManager.h"
 #include "NaniteSceneProxy.h"
+#include "HAL/LowLevelMemTracker.h"
+#include "HAL/LowLevelMemStats.h"
 
 int32 GGPUSceneUploadEveryFrame = 0;
 FAutoConsoleVariableRef CVarGPUSceneUploadEveryFrame(
@@ -59,6 +61,10 @@ FAutoConsoleVariableRef CVarGPUSceneInstanceBVH(
 	ECVF_RenderThreadSafe | ECVF_ReadOnly
 );
 
+LLM_DECLARE_TAG_API(GPUScene, RENDERER_API);
+DECLARE_LLM_MEMORY_STAT(TEXT("GPUScene"), STAT_GPUSceneLLM, STATGROUP_LLMFULL);
+DECLARE_LLM_MEMORY_STAT(TEXT("GPUScene"), STAT_GPUSceneSummaryLLM, STATGROUP_LLM);
+LLM_DEFINE_TAG(GPUScene, NAME_None, NAME_None, GET_STATFNAME(STAT_GPUSceneLLM), GET_STATFNAME(STAT_GPUSceneSummaryLLM));
 
 FVector OrthonormalizeTransform(FMatrix& Matrix)
 {
@@ -446,6 +452,8 @@ void FGPUScene::EndRender()
 
 void FGPUScene::UpdateInternal(FRDGBuilder& GraphBuilder, FScene& Scene)
 {
+	LLM_SCOPE_BYTAG(GPUScene);
+
 	ensure(bInBeginEndBlock);
 	ensure(bIsEnabled == UseGPUScene(GMaxRHIShaderPlatform, Scene.GetFeatureLevel()));
 	ensure(NumScenePrimitives == Scene.Primitives.Num());
@@ -573,6 +581,8 @@ namespace
 template<typename FUploadDataSourceAdapter>
 FGPUSceneBufferState FGPUScene::UpdateBufferState(FRDGBuilder& GraphBuilder, FScene* Scene, const FUploadDataSourceAdapter& UploadDataSourceAdapter)
 {
+	LLM_SCOPE_BYTAG(GPUScene);
+
 	UpdatePrimitiveInstances(Scene, UploadDataSourceAdapter);
 
 	FGPUSceneBufferState BufferState;
@@ -628,6 +638,8 @@ FGPUSceneBufferState FGPUScene::UpdateBufferState(FRDGBuilder& GraphBuilder, FSc
 template <typename FUploadDataSourceAdapter>
 void FGPUScene::UpdatePrimitiveInstances(FScene* Scene, const FUploadDataSourceAdapter& UploadDataSourceAdapter)
 {
+	LLM_SCOPE_BYTAG(GPUScene);
+
 	const int32 NumPrimitiveDataUploads = UploadDataSourceAdapter.NumPrimitivesToUpload();
 
 	if (!NumPrimitiveDataUploads || !UploadDataSourceAdapter.ShouldUpdateInstanceTransforms())
@@ -661,6 +673,8 @@ void FGPUScene::UpdatePrimitiveInstances(FScene* Scene, const FUploadDataSourceA
 template<typename FUploadDataSourceAdapter>
 void FGPUScene::UploadGeneral(FRHICommandListImmediate& RHICmdList, FScene *Scene, const FUploadDataSourceAdapter& UploadDataSourceAdapter, const FGPUSceneBufferState& BufferState)
 {
+	LLM_SCOPE_BYTAG(GPUScene);
+
 	if (Scene != nullptr)
 	{
 		ensure(bIsEnabled == UseGPUScene(GMaxRHIShaderPlatform, Scene->GetFeatureLevel()));
@@ -1211,6 +1225,8 @@ struct FUploadDataSourceAdapterDynamicPrimitives
 
 void FGPUScene::UploadDynamicPrimitiveShaderDataForViewInternal(FRDGBuilder& GraphBuilder, FScene *Scene, FViewInfo& View)
 {
+	LLM_SCOPE_BYTAG(GPUScene);
+
 	ensure(bInBeginEndBlock);
 	ensure(Scene == nullptr || DynamicPrimitivesOffset >= Scene->Primitives.Num());
 
@@ -1269,6 +1285,8 @@ void AddPrimitiveToUpdateGPU(FScene& Scene, int32 PrimitiveId)
 
 void FGPUScene::AddPrimitiveToUpdate(int32 PrimitiveId)
 {
+	LLM_SCOPE_BYTAG(GPUScene);
+
 	if (bIsEnabled)
 	{
 		if (PrimitiveId + 1 > PrimitivesMarkedToUpdate.Num())
@@ -1311,6 +1329,8 @@ void FGPUScene::UploadDynamicPrimitiveShaderDataForView(FRDGBuilder& GraphBuilde
 
 int32 FGPUScene::AllocateInstanceSlots(int32 NumInstanceDataEntries)
 {
+	LLM_SCOPE_BYTAG(GPUScene);
+
 	if (bIsEnabled)
 	{
 		if (NumInstanceDataEntries > 0)
@@ -1347,6 +1367,8 @@ int32 FGPUScene::AllocateInstanceSlots(int32 NumInstanceDataEntries)
 
 void FGPUScene::FreeInstanceSlots(int InstanceDataOffset, int32 NumInstanceDataEntries)
 {
+	LLM_SCOPE_BYTAG(GPUScene);
+
 	if (bIsEnabled)
 	{
 		InstanceDataAllocator.Free(InstanceDataOffset, NumInstanceDataEntries);
@@ -1425,6 +1447,8 @@ FGPUSceneDynamicContext::~FGPUSceneDynamicContext()
 
 FGPUScenePrimitiveCollector::FUploadData* FGPUSceneDynamicContext::AllocateDynamicPrimitiveData()
 {
+	LLM_SCOPE_BYTAG(GPUScene);
+
 	FGPUScenePrimitiveCollector::FUploadData* UploadData = new FGPUScenePrimitiveCollector::FUploadData;
 	DymamicPrimitiveUploadData.Add(UploadData);
 	return UploadData;
