Free FProjectedShadowInfo::ProjectionStencilingPasses in FProjectedShadowInfo destructor, else could leak if ClearTransientArrays() wasn't called.

diff --git a/Engine/Source/Runtime/Renderer/Private/ShadowSetup.cpp b/Engine/Source/Runtime/Renderer/Private/ShadowSetup.cpp
--- a/Engine/Source/Runtime/Renderer/Private/ShadowSetup.cpp
+++ b/Engine/Source/Runtime/Renderer/Private/ShadowSetup.cpp
@@ -705,6 +705,10 @@ FProjectedShadowInfo::FProjectedShadowInfo()
 
 FProjectedShadowInfo::~FProjectedShadowInfo()
 {
+	for (auto ProjectionStencilingPass : ProjectionStencilingPasses)
+	{
+		delete ProjectionStencilingPass;
+	}
 }
 
 /** Shadow border needs to be wide enough to prevent the shadow filtering from picking up content in other shadowmaps in the atlas. */
