Fix crash when path tracer is active

Post processing was trying to handle screen space reflection post processing passes even though the actual passes were skipped

diff --git a/Engine/Source/Runtime/Renderer/Private/ScreenSpaceRayTracing.cpp b/Engine/Source/Runtime/Renderer/Private/ScreenSpaceRayTracing.cpp
--- a/Engine/Source/Runtime/Renderer/Private/ScreenSpaceRayTracing.cpp
+++ b/Engine/Source/Runtime/Renderer/Private/ScreenSpaceRayTracing.cpp
@@ -121,7 +121,7 @@ bool ShouldKeepBleedFreeSceneColor(const FViewInfo& View)
 
 bool ShouldRenderScreenSpaceReflections(const FViewInfo& View)
 {
-	if(!View.Family->EngineShowFlags.ScreenSpaceReflections || View.FinalPostProcessSettings.ReflectionMethod != EReflectionMethod::ScreenSpace)
+	if(!View.Family->EngineShowFlags.ScreenSpaceReflections || View.FinalPostProcessSettings.ReflectionMethod != EReflectionMethod::ScreenSpace || HasRayTracedOverlay(*View.Family))
 	{
 		return false;
 	}
