PhysicalPagePool needs to be UInt compatible and PhysicalPagePoolHw needs to be float compatible.   Slightly modified version of CL by [REDACTED_PERSON] after testing on macOS Big Sur.

diff --git a/Engine/Source/Runtime/Renderer/Private/VirtualShadowMaps/VirtualShadowMapArray.cpp b/Engine/Source/Runtime/Renderer/Private/VirtualShadowMaps/VirtualShadowMapArray.cpp
--- a/Engine/Source/Runtime/Renderer/Private/VirtualShadowMaps/VirtualShadowMapArray.cpp
+++ b/Engine/Source/Runtime/Renderer/Private/VirtualShadowMaps/VirtualShadowMapArray.cpp
@@ -269,8 +269,8 @@ void FVirtualShadowMapArray::Initialize(FRDGBuilder& GraphBuilder, FVirtualShado
 	UniformParameters.PageTable = GraphBuilder.CreateSRV(GSystemTextures.GetDefaultStructuredBuffer(GraphBuilder, sizeof(DummyPageElement), DummyPageElement));
 	UniformParameters.ProjectionData = GraphBuilder.CreateSRV(GSystemTextures.GetDefaultStructuredBuffer(GraphBuilder, sizeof(FVirtualShadowMapProjectionShaderData)));
 
-	UniformParameters.PhysicalPagePool = GraphBuilder.RegisterExternalTexture(GSystemTextures.BlackDummy);
-	UniformParameters.PhysicalPagePoolHw = GraphBuilder.RegisterExternalTexture(GSystemTextures.BlackDummy);
+	UniformParameters.PhysicalPagePool = GSystemTextures.GetZeroUIntDummy(GraphBuilder);
+	UniformParameters.PhysicalPagePoolHw = GSystemTextures.GetBlackDummy(GraphBuilder);
 }
 
 FVirtualShadowMapArray::~FVirtualShadowMapArray()
@@ -590,7 +590,7 @@ void FVirtualShadowMapArray::ClearPhysicalMemory(FRDGBuilder& GraphBuilder, FRDG
 			bool bCacheDataAvailable = CacheManager && CacheManager->PrevBuffers.PhysicalPageMetaData;
 			if (bCacheDataAvailable)
 			{
-				PassParameters->CachedPhysicalPagesTexture = RegisterExternalTextureWithFallback(GraphBuilder, CacheManager->PrevBuffers.PhysicalPagePool, GSystemTextures.BlackDummy);
+				PassParameters->CachedPhysicalPagesTexture = RegisterExternalTextureWithFallback(GraphBuilder, CacheManager->PrevBuffers.PhysicalPagePool, GSystemTextures.ZeroUIntDummy);
 				PassParameters->CachedPageInfos = GraphBuilder.CreateSRV(CachedPageInfosRDG);
 				PassParameters->PrevPhysicalPageMetaData = GraphBuilder.CreateSRV(GraphBuilder.RegisterExternalBuffer(CacheManager->PrevBuffers.PhysicalPageMetaData));
 			}
@@ -1158,8 +1158,8 @@ void FVirtualShadowMapArray::BuildPageAllocations(
 
 void FVirtualShadowMapArray::SetupProjectionParameters(FRDGBuilder& GraphBuilder)
 {
-	UniformParameters.PhysicalPagePool = PhysicalPagePoolRDG ? PhysicalPagePoolRDG : GraphBuilder.RegisterExternalTexture(GSystemTextures.BlackDummy);
-	UniformParameters.PhysicalPagePoolHw = PhysicalPagePoolHw ? PhysicalPagePoolHw : GraphBuilder.RegisterExternalTexture(GSystemTextures.BlackDummy);
+	UniformParameters.PhysicalPagePool = PhysicalPagePoolRDG ? PhysicalPagePoolRDG : GSystemTextures.GetZeroUIntDummy(GraphBuilder);
+	UniformParameters.PhysicalPagePoolHw = PhysicalPagePoolHw ? PhysicalPagePoolHw : GSystemTextures.GetBlackDummy(GraphBuilder);
 }
 
 class FDebugVisualizeVirtualSmCS : public FVirtualPageManagementShader
diff --git a/Engine/Source/Runtime/Renderer/Private/VirtualShadowMaps/VirtualShadowMapCacheManager.cpp b/Engine/Source/Runtime/Renderer/Private/VirtualShadowMaps/VirtualShadowMapCacheManager.cpp
--- a/Engine/Source/Runtime/Renderer/Private/VirtualShadowMaps/VirtualShadowMapCacheManager.cpp
+++ b/Engine/Source/Runtime/Renderer/Private/VirtualShadowMaps/VirtualShadowMapCacheManager.cpp
@@ -450,8 +450,8 @@ void FVirtualShadowMapArrayCacheManager::ProcessInstanceRangeInvalidation(FRDGBu
 	PrevUniformParameters.ProjectionData = RegExtCreateSrv(PrevBuffers.ShadowMapProjectionDataBuffer, TEXT("Shadow.Virtual.PrevProjectionData"));
 	PrevUniformParameters.PageTable = RegExtCreateSrv(PrevBuffers.PageTable, TEXT("Shadow.Virtual.PrevPageTable"));
 	// Unused in this path
-	PrevUniformParameters.PhysicalPagePool = GraphBuilder.RegisterExternalTexture(GSystemTextures.BlackDummy);
-	PrevUniformParameters.PhysicalPagePoolHw = GraphBuilder.RegisterExternalTexture(GSystemTextures.BlackDummy);
+	PrevUniformParameters.PhysicalPagePool = GSystemTextures.GetZeroUIntDummy(GraphBuilder);
+	PrevUniformParameters.PhysicalPagePoolHw = GSystemTextures.GetBlackDummy(GraphBuilder);
 
 	FRDGBufferUploader BufferUploader;
 	FRDGBufferRef InstanceRangesSmallRDG = !InstanceRangesSmall.IsEmpty() ? CreateStructuredBuffer(GraphBuilder, BufferUploader, TEXT("Shadow.Virtual.InstanceRangesSmall"), InstanceRangesSmall) : nullptr;
