Corrected the naming of default buffers

#rnx

diff --git a/Engine/Source/Runtime/Renderer/Private/SystemTextures.cpp b/Engine/Source/Runtime/Renderer/Private/SystemTextures.cpp
--- a/Engine/Source/Runtime/Renderer/Private/SystemTextures.cpp
+++ b/Engine/Source/Runtime/Renderer/Private/SystemTextures.cpp
@@ -1089,11 +1089,11 @@ FRDGBufferRef GetInternalDefaultBuffer(
 	FRDGBufferRef Buffer = nullptr; 
 	if (bIsStructuredBuffer)
 	{
-		Buffer = GraphBuilder.CreateBuffer(FRDGBufferDesc::CreateStructuredDesc(NumBytePerElement, NumElements), TEXT("DefaultBuffer"));
+		Buffer = GraphBuilder.CreateBuffer(FRDGBufferDesc::CreateStructuredDesc(NumBytePerElement, NumElements), TEXT("DefaultStructuredBuffer"));
 	}
 	else
 	{
-		Buffer = GraphBuilder.CreateBuffer(FRDGBufferDesc::CreateBufferDesc(NumBytePerElement, NumElements), TEXT("DefaultStructuredBuffer"));
+		Buffer = GraphBuilder.CreateBuffer(FRDGBufferDesc::CreateBufferDesc(NumBytePerElement, NumElements), TEXT("DefaultBuffer"));
 	}
 
 	FRDGBufferUploader BufferUploader;
