Reworked scene renderer lifetime to allow deferral of scene render cleanup.
This avoids a long stall on mesh draw command recording tasks at the end of the frame and saves about .1 ms of work off the render thread.

diff --git a/Engine/Source/Runtime/RenderCore/Private/RenderingThread.cpp b/Engine/Source/Runtime/RenderCore/Private/RenderingThread.cpp
--- a/Engine/Source/Runtime/RenderCore/Private/RenderingThread.cpp
+++ b/Engine/Source/Runtime/RenderCore/Private/RenderingThread.cpp
@@ -794,6 +794,17 @@ PRAGMA_ENABLE_DEPRECATION_WARNINGS
 	ResumeTextureStreamingRenderTasks();
 }
 
+static FStopRenderingThread GStopRenderingThreadDelegate;
+
+FDelegateHandle RegisterStopRenderingThreadDelegate(const FStopRenderingThread::FDelegate& InDelegate)
+{
+	return GStopRenderingThreadDelegate.Add(InDelegate);
+}
+
+void UnregisterStopRenderingThreadDelegate(FDelegateHandle InDelegateHandle)
+{
+	GStopRenderingThreadDelegate.Remove(InDelegateHandle);
+}
 
 void StopRenderingThread()
 {
@@ -817,6 +828,8 @@ void StopRenderingThread()
 
 	if( GIsThreadedRendering )
 	{
+		GStopRenderingThreadDelegate.Broadcast();
+
 		// Get the list of objects which need to be cleaned up when the rendering thread is done with them.
 		FPendingCleanupObjects* PendingCleanupObjects = GetPendingCleanupObjects();
 
diff --git a/Engine/Source/Runtime/RenderCore/Public/RenderingThread.h b/Engine/Source/Runtime/RenderCore/Public/RenderingThread.h
--- a/Engine/Source/Runtime/RenderCore/Public/RenderingThread.h
+++ b/Engine/Source/Runtime/RenderCore/Public/RenderingThread.h
@@ -12,6 +12,7 @@
 #include "Templates/Atomic.h"
 #include "Trace/Trace.h"
 #include "Serialization/MemoryLayout.h"
+#include "Delegates/Delegate.h"
 
 class FRHICommandListImmediate;
 
@@ -350,3 +351,10 @@ public:
 private:
 	RenderCommandFunctionArray* RenderCommands;
 };
+
+DECLARE_MULTICAST_DELEGATE(FStopRenderingThread);
+using FStopRenderingThreadDelegate = FStopRenderingThread::FDelegate;
+
+extern RENDERCORE_API FDelegateHandle RegisterStopRenderingThreadDelegate(const FStopRenderingThreadDelegate& InDelegate);
+
+extern RENDERCORE_API void UnregisterStopRenderingThreadDelegate(FDelegateHandle InDelegateHandle);
\ No newline at end of file
diff --git a/Engine/Source/Runtime/Renderer/Private/MeshDrawCommands.cpp b/Engine/Source/Runtime/Renderer/Private/MeshDrawCommands.cpp
--- a/Engine/Source/Runtime/Renderer/Private/MeshDrawCommands.cpp
+++ b/Engine/Source/Runtime/Renderer/Private/MeshDrawCommands.cpp
@@ -1276,20 +1276,20 @@ bool FParallelMeshDrawCommandPass::IsOnDemandShaderCreationEnabled()
 		(RHISupportsMultithreadedShaderCreation(GMaxRHIShaderPlatform) || (bIsMobileRenderer && (!GSupportsParallelRenderingTasksWithSeparateRHIThread && IsRunningRHIInSeparateThread())));
 }
 
-void FParallelMeshDrawCommandPass::WaitForMeshPassSetupTask() const
+void FParallelMeshDrawCommandPass::WaitForMeshPassSetupTask(EWaitThread WaitThread) const
 {
 	if (TaskEventRef.IsValid())
 	{
 		// Need to wait on GetRenderThread_Local, as mesh pass setup task can wait on rendering thread inside InitResourceFromPossiblyParallelRendering().
 		QUICK_SCOPE_CYCLE_COUNTER(STAT_WaitForMeshPassSetupTask);
-		FTaskGraphInterface::Get().WaitUntilTaskCompletes(TaskEventRef, ENamedThreads::GetRenderThread_Local());
+		FTaskGraphInterface::Get().WaitUntilTaskCompletes(TaskEventRef, WaitThread == EWaitThread::Render ? ENamedThreads::GetRenderThread_Local() : ENamedThreads::AnyThread);
 	}
 }
 
-void FParallelMeshDrawCommandPass::WaitForTasksAndEmpty()
+void FParallelMeshDrawCommandPass::WaitForTasksAndEmpty(EWaitThread WaitThread)
 {
 	// Need to wait in case if someone dispatched sort and draw merge task, but didn't draw it.
-	WaitForMeshPassSetupTask();
+	WaitForMeshPassSetupTask(WaitThread);
 	TaskEventRef = nullptr;
 
 	DumpInstancingStats();
diff --git a/Engine/Source/Runtime/Renderer/Private/MeshDrawCommands.h b/Engine/Source/Runtime/Renderer/Private/MeshDrawCommands.h
--- a/Engine/Source/Runtime/Renderer/Private/MeshDrawCommands.h
+++ b/Engine/Source/Runtime/Renderer/Private/MeshDrawCommands.h
@@ -138,6 +138,12 @@ public:
 class FParallelMeshDrawCommandPass
 {
 public:
+	enum class EWaitThread
+	{
+		Render,
+		Task
+	};
+
 	FParallelMeshDrawCommandPass()
 		: bPrimitiveIdBufferDataOwnedByRHIThread(false)
 		, bHasQueuedInstanceCullingBuild(false)
@@ -192,7 +198,7 @@ public:
 	 */
 	void DispatchDraw(FParallelCommandListSet* ParallelCommandListSet, FRHICommandList& RHICmdList, const FInstanceCullingDrawParams* InstanceCullingDrawParams = nullptr) const;
 
-	void WaitForTasksAndEmpty();
+	void WaitForTasksAndEmpty(EWaitThread WaitThread = EWaitThread::Render);
 	void SetDumpInstancingStats(const FString& InPassName);
 	bool HasAnyDraw() const { return MaxNumDraws > 0; }
 
@@ -225,7 +231,7 @@ private:
 	int32 MaxNumDraws;
 
 	void DumpInstancingStats() const;
-	void WaitForMeshPassSetupTask() const;
+	void WaitForMeshPassSetupTask(EWaitThread WaitThread = EWaitThread::Render) const;
 };
 
 RENDERER_API extern void SortAndMergeDynamicPassMeshDrawCommands(
diff --git a/Engine/Source/Runtime/Renderer/Private/MobileSceneCaptureRendering.cpp b/Engine/Source/Runtime/Renderer/Private/MobileSceneCaptureRendering.cpp
--- a/Engine/Source/Runtime/Renderer/Private/MobileSceneCaptureRendering.cpp
+++ b/Engine/Source/Runtime/Renderer/Private/MobileSceneCaptureRendering.cpp
@@ -309,7 +309,7 @@ void UpdateSceneCaptureContentMobile_RenderThread(
 	const FGenerateMipsParams& GenerateMipsParams,
 	bool bDisableFlipCopyGLES)
 {
-	FMemMark MemStackMark(FMemStack::Get());
+	SceneRenderer->RenderThreadBegin(RHICmdList);
 
 	// update any resources that needed a deferred update
 	FDeferredUpdateResource::UpdateResources(RHICmdList);
@@ -426,5 +426,6 @@ void UpdateSceneCaptureContentMobile_RenderThread(
 
 		GraphBuilder.Execute();
 	}
-	FSceneRenderer::WaitForTasksClearSnapshotsAndDeleteSceneRenderer(RHICmdList, SceneRenderer);
+
+	SceneRenderer->RenderThreadEnd(RHICmdList);
 }
diff --git a/Engine/Source/Runtime/Renderer/Private/PlanarReflectionRendering.cpp b/Engine/Source/Runtime/Renderer/Private/PlanarReflectionRendering.cpp
--- a/Engine/Source/Runtime/Renderer/Private/PlanarReflectionRendering.cpp
+++ b/Engine/Source/Runtime/Renderer/Private/PlanarReflectionRendering.cpp
@@ -270,7 +270,7 @@ static void UpdatePlanarReflectionContents_RenderThread(
 {
 	QUICK_SCOPE_CYCLE_COUNTER(STAT_RenderPlanarReflection);
 
-	FMemMark MemStackMark(FMemStack::Get());
+	SceneRenderer->RenderThreadBegin(RHICmdList);
 
 	// Make sure we render to the same set of GPUs as the main scene renderer.
 	if (MainSceneRenderer->ViewFamily.RenderTarget != nullptr)
@@ -402,7 +402,8 @@ static void UpdatePlanarReflectionContents_RenderThread(
 			}
 		}
 	}
-	FSceneRenderer::WaitForTasksClearSnapshotsAndDeleteSceneRenderer(RHICmdList, SceneRenderer);
+
+	SceneRenderer->RenderThreadEnd(RHICmdList);
 }
 
 // Used for generate valid data to update planar reflection uniform buffer but don't actually render the reflection scene when we are using mobile pixel projected reflection.
@@ -416,8 +417,7 @@ static void UpdatePlanarReflectionContentsWithoutRendering_RenderThread(
 	const FName OwnerName)
 {
 	QUICK_SCOPE_CYCLE_COUNTER(STAT_RenderPlanarReflection);
-
-	FMemMark MemStackMark(FMemStack::Get());
+	SceneRenderer->RenderThreadBegin(RHICmdList);
 
 	FBox PlanarReflectionBounds = SceneProxy->WorldBounds;
 
@@ -470,7 +470,8 @@ static void UpdatePlanarReflectionContentsWithoutRendering_RenderThread(
 			SceneProxy->ViewRect[ViewIndex] = SceneRenderer->Views[ViewIndex].ViewRect;
 		}
 	}
-	FSceneRenderer::WaitForTasksClearSnapshotsAndDeleteSceneRenderer(RHICmdList, SceneRenderer);
+
+	SceneRenderer->RenderThreadEnd(RHICmdList);
 }
 
 extern void BuildProjectionMatrix(FIntPoint RenderTargetSize, ECameraProjectionMode::Type ProjectionType, float FOV, float OrthoWidth, float InNearClippingPlane, FMatrix& ProjectionMatrix);
diff --git a/Engine/Source/Runtime/Renderer/Private/ReflectionEnvironmentCapture.cpp b/Engine/Source/Runtime/Renderer/Private/ReflectionEnvironmentCapture.cpp
--- a/Engine/Source/Runtime/Renderer/Private/ReflectionEnvironmentCapture.cpp
+++ b/Engine/Source/Runtime/Renderer/Private/ReflectionEnvironmentCapture.cpp
@@ -712,7 +712,7 @@ void ClearScratchCubemaps(FRHICommandListImmediate& RHICmdList, int32 TargetSize
 /** Captures the scene for a reflection capture by rendering the scene multiple times and copying into a cubemap texture. */
 void CaptureSceneToScratchCubemap(FRHICommandListImmediate& RHICmdList, FSceneRenderer* SceneRenderer, ECubeFace CubeFace, int32 CubemapSize, bool bCapturingForSkyLight, bool bLowerHemisphereIsBlack, const FLinearColor& LowerHemisphereColor, bool bCapturingForMobile)
 {
-	FMemMark MemStackMark(FMemStack::Get());
+	SceneRenderer->RenderThreadBegin(RHICmdList);
 
 	// update any resources that needed a deferred update
 	FDeferredUpdateResource::UpdateResources(RHICmdList);
@@ -827,7 +827,7 @@ void CaptureSceneToScratchCubemap(FRHICommandListImmediate& RHICmdList, FSceneRe
 		GraphBuilder.Execute();
 	}
 
-	FSceneRenderer::WaitForTasksClearSnapshotsAndDeleteSceneRenderer(RHICmdList, SceneRenderer);
+	SceneRenderer->RenderThreadEnd(RHICmdList);
 }
 
 void CopyCubemapToScratchCubemap(FRHICommandListImmediate& RHICmdList, ERHIFeatureLevel::Type FeatureLevel, UTextureCube* SourceCubemap, int32 CubemapSize, bool bIsSkyLight, bool bLowerHemisphereIsBlack, float SourceCubemapRotation, const FLinearColor& LowerHemisphereColorValue)
diff --git a/Engine/Source/Runtime/Renderer/Private/Renderer.cpp b/Engine/Source/Runtime/Renderer/Private/Renderer.cpp
--- a/Engine/Source/Runtime/Renderer/Private/Renderer.cpp
+++ b/Engine/Source/Runtime/Renderer/Private/Renderer.cpp
@@ -57,10 +57,21 @@ void FRendererModule::StartupModule()
 	GTemporalUpscaler = ITemporalUpscaler::GetDefaultTemporalUpscaler();
 
 	FVirtualTextureSystem::Initialize();
+
+	StopRenderingThreadDelegate = RegisterStopRenderingThreadDelegate(FStopRenderingThreadDelegate::CreateLambda([this]
+	{
+		ENQUEUE_RENDER_COMMAND(FSceneRendererCleanUp)(
+			[](FRHICommandListImmediate& RHICmdList)
+		{
+			FSceneRenderer::CleanUp(RHICmdList);
+		});
+	}));
 }
 
 void FRendererModule::ShutdownModule()
 {
+	UnregisterStopRenderingThreadDelegate(StopRenderingThreadDelegate);
+
 	FVirtualTextureSystem::Shutdown();
 
 	// Free up the memory of the default denoiser. Responsibility of the plugin to free up theirs.
diff --git a/Engine/Source/Runtime/Renderer/Private/RendererModule.h b/Engine/Source/Runtime/Renderer/Private/RendererModule.h
--- a/Engine/Source/Runtime/Renderer/Private/RendererModule.h
+++ b/Engine/Source/Runtime/Renderer/Private/RendererModule.h
@@ -127,6 +127,7 @@ private:
 	FOnPostOpaqueRender PostOpaqueRenderDelegate;
 	FOnPostOpaqueRender OverlayRenderDelegate;
 	FOnResolvedSceneColor PostResolvedSceneColorCallbacks;
+	FDelegateHandle StopRenderingThreadDelegate;
 };
 
 extern ICustomCulling* GCustomCullingImpl;
diff --git a/Engine/Source/Runtime/Renderer/Private/SceneCaptureRendering.cpp b/Engine/Source/Runtime/Renderer/Private/SceneCaptureRendering.cpp
--- a/Engine/Source/Runtime/Renderer/Private/SceneCaptureRendering.cpp
+++ b/Engine/Source/Runtime/Renderer/Private/SceneCaptureRendering.cpp
@@ -246,7 +246,7 @@ static void UpdateSceneCaptureContentDeferred_RenderThread(
 	const FGenerateMipsParams& GenerateMipsParams
 	)
 {
-	FMemMark MemStackMark(FMemStack::Get());
+	SceneRenderer->RenderThreadBegin(RHICmdList);
 
 	// update any resources that needed a deferred update
 	FDeferredUpdateResource::UpdateResources(RHICmdList);
@@ -268,10 +268,10 @@ static void UpdateSceneCaptureContentDeferred_RenderThread(
 			SceneRenderer->Render(GraphBuilder);
 		}
 
-			if (bGenerateMips)
-			{
+		if (bGenerateMips)
+		{
 			FGenerateMips::Execute(GraphBuilder, TargetTexture, GenerateMipsParams);
-			}
+		}
 
 		FRDGTextureRef ResolveTexture = RegisterExternalTexture(GraphBuilder, RenderTargetTexture->TextureRHI, TEXT("SceneCaptureResolve"));
 		AddCopyToResolveTargetPass(GraphBuilder, TargetTexture, ResolveTexture, ResolveParams);
@@ -279,7 +279,7 @@ static void UpdateSceneCaptureContentDeferred_RenderThread(
 		GraphBuilder.Execute();
 	}
 
-	FSceneRenderer::WaitForTasksClearSnapshotsAndDeleteSceneRenderer(RHICmdList, SceneRenderer);
+	SceneRenderer->RenderThreadEnd(RHICmdList);
 }
 
 static void ODSCapture_RenderThread(
diff --git a/Engine/Source/Runtime/Renderer/Private/SceneRendering.cpp b/Engine/Source/Runtime/Renderer/Private/SceneRendering.cpp
--- a/Engine/Source/Runtime/Renderer/Private/SceneRendering.cpp
+++ b/Engine/Source/Runtime/Renderer/Private/SceneRendering.cpp
@@ -980,6 +980,14 @@ void FViewInfo::Init()
 	GPUSceneViewId = INDEX_NONE;
 }
 
+void FViewInfo::WaitForTasks(FParallelMeshDrawCommandPass::EWaitThread WaitThread)
+{
+	for (int32 MeshDrawIndex = 0; MeshDrawIndex < EMeshPass::Num; MeshDrawIndex++)
+	{
+		ParallelMeshDrawCommandPasses[MeshDrawIndex].WaitForTasksAndEmpty(WaitThread);
+	}
+}
+
 FViewInfo::~FViewInfo()
 {
 	for(int32 ResourceIndex = 0;ResourceIndex < DynamicResources.Num();ResourceIndex++)
@@ -991,11 +999,6 @@ FViewInfo::~FViewInfo()
 		CustomVisibilityQuery->Release();
 	}
 
-	for (int32 MeshDrawIndex = 0; MeshDrawIndex < EMeshPass::Num; MeshDrawIndex++)
-	{
-		ParallelMeshDrawCommandPasses[MeshDrawIndex].WaitForTasksAndEmpty();
-	}
-
 	//this uses memstack allocation for strongrefs, so we need to manually empty to get the destructor called to not leak the uniformbuffers stored here.
 	TranslucentSelfShadowUniformBufferMap.Empty();
 
@@ -2064,11 +2067,10 @@ FViewInfo* FViewInfo::CreateSnapshot() const
 	return Result;
 }
 
-void FViewInfo::DestroyAllSnapshots()
+void FViewInfo::DestroyAllSnapshots(FParallelMeshDrawCommandPass::EWaitThread WaitThread)
 {
 	QUICK_SCOPE_CYCLE_COUNTER(STAT_FViewInfo_DestroyAllSnapshots);
 
-	check(IsInRenderingThread());
 	// we will only keep double the number actually used, plus a few
 	int32 NumToRemove = FreeViewInfoSnapshots.Num() - (ViewInfoSnapshots.Num() + 2);
 	if (NumToRemove > 0)
@@ -2086,7 +2088,7 @@ void FViewInfo::DestroyAllSnapshots()
 
 		for (int32 Index = 0; Index < Snapshot->ParallelMeshDrawCommandPasses.Num(); ++Index)
 		{
-			Snapshot->ParallelMeshDrawCommandPasses[Index].WaitForTasksAndEmpty();
+			Snapshot->ParallelMeshDrawCommandPasses[Index].WaitForTasksAndEmpty(WaitThread);
 		}
 
 		for (int i = 0; i < EMeshPass::Num; i++)
@@ -2329,7 +2331,6 @@ FSceneRenderer::FSceneRenderer(const FSceneViewFamily* InViewFamily,FHitProxyCon
 ,	bHasRequestedToggleFreeze(false)
 ,	bUsedPrecomputedVisibility(false)
 ,	InstancedStereoWidth(0)
-,	RootMark(nullptr)
 ,	FamilySize(0, 0)
 ,	GPUSceneDynamicContext(CheckPointer(Scene)->GPUScene)
 ,	bShadowDepthRenderCompleted(false)
@@ -3484,88 +3485,6 @@ bool FSceneRenderer::ShouldCompositeEditorPrimitives(const FViewInfo& View)
 	return false;
 }
 
-void FSceneRenderer::WaitForTasksClearSnapshotsAndDeleteSceneRenderer(FRHICommandListImmediate& RHICmdList, FSceneRenderer* SceneRenderer, bool bWaitForTasks)
-{
-	TRACE_CPUPROFILER_EVENT_SCOPE(FSceneRenderer::WaitForTasksClearSnapshotsAndDeleteSceneRenderer);
-
-	// we are about to destroy things that are being used for async tasks, so we wait here for them.
-	if (bWaitForTasks)
-	{
-		QUICK_SCOPE_CYCLE_COUNTER(STAT_DeleteSceneRenderer_WaitForTasks);
-		RHICmdList.ImmediateFlush(EImmediateFlushType::WaitForOutstandingTasksOnly);
-	}
-
-	// Wait for all dispatched shadow mesh draw tasks.
-	for (int32 PassIndex = 0; PassIndex < SceneRenderer->DispatchedShadowDepthPasses.Num(); ++PassIndex)
-	{
-		SceneRenderer->DispatchedShadowDepthPasses[PassIndex]->WaitForTasksAndEmpty();
-	}
-
-	FViewInfo::DestroyAllSnapshots(); // this destroys viewinfo snapshots
-	static const IConsoleVariable* AsyncDispatch	= IConsoleManager::Get().FindConsoleVariable(TEXT("r.RHICmdAsyncRHIThreadDispatch"));
-
-	if (AsyncDispatch->GetInt() == 0)
-	{
-		QUICK_SCOPE_CYCLE_COUNTER(STAT_DeleteSceneRenderer_Dispatch);
-		RHICmdList.ImmediateFlush(EImmediateFlushType::WaitForDispatchToRHIThread); // we want to make sure this all gets to the rhi thread this frame and doesn't hang around
-	}
-
-
-	FMemMark* LocalRootMark = SceneRenderer->RootMark;
-	SceneRenderer->RootMark = nullptr;
-	// Delete the scene renderer.
-	{
-		QUICK_SCOPE_CYCLE_COUNTER(STAT_DeleteSceneRenderer);
-		delete SceneRenderer;
-	}
-
-	// Can relase only after all mesh pass tasks are finished.
-	GPrimitiveIdVertexBufferPool.DiscardAll();
-	FGraphicsMinimalPipelineStateId::ResetLocalPipelineIdTableSize();
-
-	delete LocalRootMark;
-}
-
-class FClearSnapshotsAndDeleteSceneRendererTask
-{
-	FSceneRenderer* SceneRenderer;
-
-public:
-
-	FClearSnapshotsAndDeleteSceneRendererTask(FSceneRenderer* InSceneRenderer)
-		: SceneRenderer(InSceneRenderer)
-	{
-	}
-
-	FORCEINLINE TStatId GetStatId() const
-	{
-		RETURN_QUICK_DECLARE_CYCLE_STAT(FClearSnapshotsAndDeleteSceneRendererTask, STATGROUP_TaskGraphTasks);
-	}
-
-	ENamedThreads::Type GetDesiredThread()
-	{
-		return ENamedThreads::GetRenderThread_Local();
-	}
-
-	static ESubsequentsMode::Type GetSubsequentsMode() { return ESubsequentsMode::TrackSubsequents; }
-
-	void DoTask(ENamedThreads::Type CurrentThread, const FGraphEventRef& MyCompletionGraphEvent)
-	{
-		FSceneRenderer::WaitForTasksClearSnapshotsAndDeleteSceneRenderer(FRHICommandListExecutor::GetImmediateCommandList(), SceneRenderer, false);
-	}
-};
-
-void FSceneRenderer::DelayWaitForTasksClearSnapshotsAndDeleteSceneRenderer(FRHICommandListImmediate& RHICmdList, FSceneRenderer* SceneRenderer)
-{
-	FGraphEventArray& WaitOutstandingTasks = RHICmdList.GetRenderThreadTaskArray();
-	FGraphEventRef ClearSnapshotsAndDeleteSceneRendererTask = TGraphTask<FClearSnapshotsAndDeleteSceneRendererTask>::CreateTask(&WaitOutstandingTasks, ENamedThreads::GetRenderThread()).ConstructAndDispatchWhenReady(SceneRenderer);
-
-	WaitOutstandingTasks.Empty();
-	WaitOutstandingTasks.Add(ClearSnapshotsAndDeleteSceneRendererTask);
-
-	FTaskGraphInterface::Get().ProcessThreadUntilIdle(ENamedThreads::GetRenderThread_Local()); 
-}
-
 void FSceneRenderer::UpdatePrimitiveIndirectLightingCacheBuffers()
 {
 	// Use a bit array to prevent primitives from being updated more than once.
@@ -3650,43 +3569,171 @@ void FSceneRenderer::ViewExtensionPreRender_RenderThread(FRHICommandListImmediat
 	FDeferredUpdateResource::UpdateResources(RHICmdList);
 }
 
-
-static TAutoConsoleVariable<int32> CVarDelaySceneRenderCompletion(
-	TEXT("r.DelaySceneRenderCompletion"),
-	0,
-	TEXT("Experimental option to postpone the cleanup of the scene renderer until later. This does NOT currently work because it is possible for the scene to be modified before ~FSceneRenderer, and that assumes the scene is unchanged."),
+static int32 GSceneRenderCleanUpMode = 2;
+static FAutoConsoleVariableRef CVarSceneRenderCleanUpMode(
+	TEXT("r.SceneRenderCleanUpMode"),
+	GSceneRenderCleanUpMode,
+	TEXT("Controls when to perform clean up of the scene renderer.\n")
+	TEXT(" 0: clean up is performed immediately after render on the render thread.\n")
+	TEXT(" 1: clean up deferred until the start of the next scene render on the render thread.\n")
+	TEXT(" 2: clean up deferred until the start of the next scene render on the render thread, with some work distributed to an async task. (default)\n"),
 	ECVF_RenderThreadSafe
 );
 
-/**
- * Helper function performing actual work in render thread.
- *
- * @param SceneRenderer	Scene renderer to use for rendering.
- */
-static void RenderViewFamily_RenderThread(FRHICommandListImmediate& RHICmdList, FSceneRenderer* SceneRenderer)
+enum class ESceneRenderCleanUpMode
 {
-	LLM_SCOPE(ELLMTag::SceneRender);
+	Immediate,
+	Deferred,
+	DeferredAndAsync
+};
 
+inline ESceneRenderCleanUpMode GetSceneRenderCleanUpMode()
+{
+	if (GSceneRenderCleanUpMode != 1 && GSceneRenderCleanUpMode != 2)
 	{
-		QUICK_SCOPE_CYCLE_COUNTER(STAT_DelaySceneRenderCompletion_TaskWait);
-		FRHICommandListExecutor::GetImmediateCommandList().ImmediateFlush(EImmediateFlushType::WaitForOutstandingTasksOnly);
+		return ESceneRenderCleanUpMode::Immediate;
 	}
 
 	static const IConsoleVariable* AsyncDispatch = IConsoleManager::Get().FindConsoleVariable(TEXT("r.RHICmdAsyncRHIThreadDispatch"));
 
-#if WITH_EDITOR
-	bool bDelayCleanup = false;
-#else
-	bool bDelayCleanup = !!AsyncDispatch->GetInt() && IsRunningRHIInSeparateThread() && !!CVarDelaySceneRenderCompletion.GetValueOnRenderThread();
-#endif
+	if (!AsyncDispatch->GetInt() || !IsRunningRHIInSeparateThread())
+	{
+		return ESceneRenderCleanUpMode::Immediate;
+	}
+
+	return (ESceneRenderCleanUpMode)GSceneRenderCleanUpMode;
+}
+
+static void DeleteSceneRenderer(FRHICommandListImmediate& RHICmdList, FSceneRenderer* SceneRenderer, FMemMark* MemStackMark)
+{
+	static const IConsoleVariable* AsyncDispatch = IConsoleManager::Get().FindConsoleVariable(TEXT("r.RHICmdAsyncRHIThreadDispatch"));
+
+	if (AsyncDispatch->GetInt() == 0)
+	{
+		QUICK_SCOPE_CYCLE_COUNTER(STAT_DeleteSceneRenderer_Dispatch);
+		RHICmdList.ImmediateFlush(EImmediateFlushType::WaitForDispatchToRHIThread); // we want to make sure this all gets to the rhi thread this frame and doesn't hang around
+	}
+
+	TRACE_CPUPROFILER_EVENT_SCOPE(DeleteSceneRenderer);
+	delete SceneRenderer;
+	delete MemStackMark;
+
+	// Can release only after all mesh pass tasks are finished.
+	GPrimitiveIdVertexBufferPool.DiscardAll();
+	FGraphicsMinimalPipelineStateId::ResetLocalPipelineIdTableSize();
+}
+
+static void ReleaseSceneRenderer(FRHICommandListImmediate& RHICmdList, FSceneRenderer* SceneRenderer, FMemMark* MemStackMark)
+{
+	{
+		QUICK_SCOPE_CYCLE_COUNTER(STAT_DeleteSceneRenderer_WaitForTasks);
+		RHICmdList.ImmediateFlush(EImmediateFlushType::WaitForOutstandingTasksOnly);
+	}
+
+	SceneRenderer->WaitForTasksAndClearSnapshots(FParallelMeshDrawCommandPass::EWaitThread::Render);
+	DeleteSceneRenderer(RHICmdList, SceneRenderer, MemStackMark);
+}
+
+void FSceneRenderer::RenderThreadBegin(FRHICommandListImmediate& RHICmdList)
+{
+	CleanUp(RHICmdList);
+
+	FMemStackBase& MemStack = FMemStack::Get();
+	checkf(MemStack.GetNumMarks() == 0, TEXT("FSceneRenderer::RenderThreadBegin called with an active MemMark."));
+	MemStackMark = new FMemMark(MemStack);
+}
 
-	FMemMark* MemStackMark = new FMemMark(FMemStack::Get());
+struct FSceneRenderCleanUpState
+{
+	FSceneRenderer* Renderer{};
+	FMemMark* MemStackMark{};
+	FGraphEventRef Task;
+	ESceneRenderCleanUpMode CompletionMode = ESceneRenderCleanUpMode::Immediate;
+};
+
+static FSceneRenderCleanUpState GSceneRenderCleanUpState;
+
+void FSceneRenderer::RenderThreadEnd(FRHICommandListImmediate& RHICmdList)
+{
+	check(!GSceneRenderCleanUpState.Renderer);
+
+	const ESceneRenderCleanUpMode SceneRenderCleanUpMode = GetSceneRenderCleanUpMode();
 
-	if (bDelayCleanup)
+	GSceneRenderCleanUpState.CompletionMode = SceneRenderCleanUpMode;
+
+	if (GSceneRenderCleanUpState.CompletionMode == ESceneRenderCleanUpMode::Immediate)
 	{
-		SceneRenderer->RootMark = MemStackMark;
-		MemStackMark = nullptr;
+		ReleaseSceneRenderer(RHICmdList, this, MemStackMark);
 	}
+	else
+	{
+		GSceneRenderCleanUpState.Renderer = this;
+		GSceneRenderCleanUpState.MemStackMark = MemStackMark;
+
+		if (SceneRenderCleanUpMode == ESceneRenderCleanUpMode::DeferredAndAsync)
+		{
+			FGraphEventArray& WaitOutstandingTasks = RHICmdList.GetRenderThreadTaskArray();
+
+			GSceneRenderCleanUpState.Task = FFunctionGraphTask::CreateAndDispatchWhenReady([this]
+			{
+				WaitForTasksAndClearSnapshots(FParallelMeshDrawCommandPass::EWaitThread::Task);
+
+			}, TStatId(), &WaitOutstandingTasks);
+
+			WaitOutstandingTasks.Empty();
+		}
+	}
+}
+
+void FSceneRenderer::CleanUp(FRHICommandListImmediate& RHICmdList)
+{
+	if (GSceneRenderCleanUpState.CompletionMode == ESceneRenderCleanUpMode::Immediate)
+	{
+		return;
+	}
+
+	switch (GSceneRenderCleanUpState.CompletionMode)
+	{
+	case ESceneRenderCleanUpMode::Deferred:
+		ReleaseSceneRenderer(RHICmdList, GSceneRenderCleanUpState.Renderer, GSceneRenderCleanUpState.MemStackMark);
+		break;
+	case ESceneRenderCleanUpMode::DeferredAndAsync:
+		GSceneRenderCleanUpState.Task->Wait(ENamedThreads::GetRenderThread());
+		DeleteSceneRenderer(RHICmdList, GSceneRenderCleanUpState.Renderer, GSceneRenderCleanUpState.MemStackMark);
+		break;
+	}
+
+	GSceneRenderCleanUpState = {};
+}
+
+void FSceneRenderer::WaitForTasksAndClearSnapshots(FParallelMeshDrawCommandPass::EWaitThread WaitThread)
+{
+	TRACE_CPUPROFILER_EVENT_SCOPE(FSceneRenderer::WaitForTasksAndClearSnapshots);
+
+	// Wait for all dispatched shadow mesh draw tasks.
+	for (int32 PassIndex = 0; PassIndex < DispatchedShadowDepthPasses.Num(); ++PassIndex)
+	{
+		DispatchedShadowDepthPasses[PassIndex]->WaitForTasksAndEmpty(WaitThread);
+	}
+
+	for (FViewInfo& View : Views)
+	{
+		View.WaitForTasks(WaitThread);
+	}
+
+	FViewInfo::DestroyAllSnapshots(WaitThread);
+}
+
+/**
+ * Helper function performing actual work in render thread.
+ *
+ * @param SceneRenderer	Scene renderer to use for rendering.
+ */
+static void RenderViewFamily_RenderThread(FRHICommandListImmediate& RHICmdList, FSceneRenderer* SceneRenderer)
+{
+	LLM_SCOPE(ELLMTag::SceneRender);
+
+	SceneRenderer->RenderThreadBegin(RHICmdList);
 
 	// update any resources that needed a deferred update
 	FDeferredUpdateResource::UpdateResources(RHICmdList);
@@ -3723,7 +3770,7 @@ static void RenderViewFamily_RenderThread(FRHICommandListImmediate& RHICmdList,
 			GraphBuilder.Execute();
 
 			CSV_SCOPED_TIMING_STAT_EXCLUSIVE(PostRenderCleanUp);
-			
+
 			if (IsHairStrandsEnabled(EHairStrandsShaderType::All, SceneRenderer->Scene->GetShaderPlatform()) && SceneRenderer->Views.Num() > 0 && !ViewFamily.EngineShowFlags.HitProxies)
 			{
 				FHairStrandsBookmarkParameters Parameters = CreateHairStrandsBookmarkParameters(SceneRenderer->Scene, SceneRenderer->Views);
@@ -3774,16 +3821,6 @@ static void RenderViewFamily_RenderThread(FRHICommandListImmediate& RHICmdList,
 #endif
 
 		GRenderTargetPool.SetEventRecordingActive(false);
-
-		if (bDelayCleanup)
-		{
-			check(!MemStackMark && SceneRenderer->RootMark);
-			FSceneRenderer::DelayWaitForTasksClearSnapshotsAndDeleteSceneRenderer(RHICmdList, SceneRenderer);
-		}
-		else
-		{
-			FSceneRenderer::WaitForTasksClearSnapshotsAndDeleteSceneRenderer(RHICmdList, SceneRenderer);
-		}
 	}
 
 #if STATS
@@ -3795,10 +3832,7 @@ static void RenderViewFamily_RenderThread(FRHICommandListImmediate& RHICmdList,
 	}
 #endif
 
-	{
-		TRACE_CPUPROFILER_EVENT_SCOPE(DeleteMemStackMark);
-		delete MemStackMark;
-	}
+	SceneRenderer->RenderThreadEnd(RHICmdList);
 }
 
 void OnChangeSimpleForwardShading(IConsoleVariable* Var)
@@ -3984,7 +4018,6 @@ void FRendererModule::BeginRenderingViewFamily(FCanvas* Canvas, FSceneViewFamily
 			{
 				const float StartDelayMillisec = FPlatformTime::ToMilliseconds(FPlatformTime::Cycles() - DrawSceneEnqueue);
 				CSV_CUSTOM_STAT_GLOBAL(DrawSceneCommand_StartDelay, StartDelayMillisec, ECsvCustomStatOp::Set);
-
 				RenderViewFamily_RenderThread(RHICmdList, SceneRenderer);
 				FlushPendingDeleteRHIResources_RenderThread();
 			});
@@ -4009,8 +4042,9 @@ void FRendererModule::PerFrameCleanupIfSkipRenderer()
 {
 	// Some systems (e.g. Slate) can still draw (via FRendererModule::DrawTileMesh for example) when scene renderer is not used
 	ENQUEUE_RENDER_COMMAND(CmdPerFrameCleanupIfSkipRenderer)(
-		[](FRHICommandListImmediate&)
+		[](FRHICommandListImmediate& RHICmdList)
 	{
+		FSceneRenderer::CleanUp(RHICmdList);
 		GPrimitiveIdVertexBufferPool.DiscardAll();
 	});
 }
@@ -4797,6 +4831,8 @@ void FSceneRenderer::UpdateSkyIrradianceGpuBuffer(FRHICommandListImmediate& RHIC
 		Scene->SkyIrradianceEnvironmentMap.Initialize(TEXT("SkyIrradianceEnvironmentMap"), sizeof(FVector4), 7, 0);
 	}
 
+	TRACE_CPUPROFILER_EVENT_SCOPE(UpdateSkyIrradianceGpuBuffer);
+
 	FVector4 OutSkyIrradianceEnvironmentMap[7];
 	// Make sure there's no padding since we're going to cast to FVector4*
 	checkSlow(sizeof(OutSkyIrradianceEnvironmentMap) == sizeof(FVector4) * 7);
diff --git a/Engine/Source/Runtime/Renderer/Private/SceneRendering.h b/Engine/Source/Runtime/Renderer/Private/SceneRendering.h
--- a/Engine/Source/Runtime/Renderer/Private/SceneRendering.h
+++ b/Engine/Source/Runtime/Renderer/Private/SceneRendering.h
@@ -1616,7 +1616,9 @@ public:
 	FViewInfo* CreateSnapshot() const;
 
 	/** Destroy all snapshots before we wipe the scene allocator. */
-	static void DestroyAllSnapshots();
+	static void DestroyAllSnapshots(FParallelMeshDrawCommandPass::EWaitThread WaitThread);
+
+	void WaitForTasks(FParallelMeshDrawCommandPass::EWaitThread WaitThread);
 
 	// Get the range in DynamicMeshElements[] for a given PrimitiveIndex
 	// @return range (start is inclusive, end is exclusive)
@@ -1820,14 +1822,15 @@ public:
 	 */
 	uint32 InstancedStereoWidth;
 
-	/** Only used if we are going to delay the deletion of the scene renderer until later. */
-	FMemMark* RootMark;
-
 public:
 
 	FSceneRenderer(const FSceneViewFamily* InViewFamily,FHitProxyConsumer* HitProxyConsumer);
 	virtual ~FSceneRenderer();
 
+	// Initializes the scene renderer on the render thread.
+	void RenderThreadBegin(FRHICommandListImmediate& RHICmdList);
+	void RenderThreadEnd(FRHICommandListImmediate& RHICmdList);
+
 	// FSceneRenderer interface
 
 	virtual void Render(FRDGBuilder& GraphBuilder) = 0;
@@ -1880,9 +1883,11 @@ public:
 	static void ViewExtensionPreRender_RenderThread(FRHICommandListImmediate& RHICmdList, FSceneRenderer* SceneRenderer);
 
 	/** the last thing we do with a scene renderer, lots of cleanup related to the threading **/
-	static void WaitForTasksClearSnapshotsAndDeleteSceneRenderer(FRHICommandListImmediate& RHICmdList, FSceneRenderer* SceneRenderer, bool bWaitForTasks = true);
-	static void DelayWaitForTasksClearSnapshotsAndDeleteSceneRenderer(FRHICommandListImmediate& RHICmdList, FSceneRenderer* SceneRenderer);
-	
+	void WaitForTasksAndClearSnapshots(FParallelMeshDrawCommandPass::EWaitThread WaitThread);
+
+	/** Called to release any deallocations that were deferred until the next render. */
+	static void CleanUp(FRHICommandListImmediate& RHICmdList);
+
 	/** Apply the ResolutionFraction on ViewSize, taking into account renderer's requirements. */
 	static FIntPoint ApplyResolutionFraction(
 		const FSceneViewFamily& ViewFamily, const FIntPoint& UnscaledViewSize, float ResolutionFraction);
@@ -2177,6 +2182,8 @@ private:
 	 * Note #2: Much more elegant if the pool itself managed this!
 	 */
 	TArray<FProjectedShadowInfo*, SceneRenderingAllocator> MemStackProjectedShadows;
+
+	FMemMark* MemStackMark = nullptr;
 };
 
 struct FForwardScreenSpaceShadowMaskTextureMobileOutputs
