Updated Nanite unsupported error message

diff --git a/Engine/Source/Runtime/Renderer/Private/SceneRendering.cpp b/Engine/Source/Runtime/Renderer/Private/SceneRendering.cpp
--- a/Engine/Source/Runtime/Renderer/Private/SceneRendering.cpp
+++ b/Engine/Source/Runtime/Renderer/Private/SceneRendering.cpp
@@ -3043,6 +3043,9 @@ void FSceneRenderer::RenderFinish(FRDGBuilder& GraphBuilder, FRDGTextureRef View
 			}
 		}
 
+		static const auto* CVarNaniteRequireAtomics = IConsoleManager::Get().FindTConsoleVariableDataInt(TEXT("r.Nanite.RequireAtomic64Support"));
+		const bool bNaniteRequireAtomics = CVarNaniteRequireAtomics != nullptr ? CVarNaniteRequireAtomics->GetValueOnAnyThread() != 0 : true;
+
 		// Mobile-specific warnings
 		const bool bMobile = (FeatureLevel <= ERHIFeatureLevel::ES3_1);
 		const bool bShowMobileLowQualityLightmapWarning = bMobile && !ReadOnlyCVARCache.bEnableLowQualityLightmaps && ReadOnlyCVARCache.bAllowStaticLighting;
@@ -3090,7 +3093,7 @@ void FSceneRenderer::RenderFinish(FRDGBuilder& GraphBuilder, FRDGTextureRef View
 						bShowAtmosphericFogWarning, bViewParentOrFrozen, bShowSkylightWarning, bShowPointLightWarning, bShowShadowedLightOverflowWarning,
 						bShowMobileLowQualityLightmapWarning, bShowMobileMovableDirectionalLightWarning, bShowMobileDynamicCSMWarning, bMobileShowVertexFogWarning,
 						bShowSkinCacheOOM, bSingleLayerWaterWarning, bShowNoSkyAtmosphereComponentWarning, bFxDebugDraw, FXInterface, bShowSkyAtmosphereFogComponentsConflicts, 
-						bLumenEnabledButNoSoftwareTracing, bNaniteEnabledButNoAtomics, bRealTimeSkyCaptureButNothingToCapture]
+						bLumenEnabledButNoSoftwareTracing, bNaniteEnabledButNoAtomics, bNaniteRequireAtomics, bRealTimeSkyCaptureButNothingToCapture]
 						(FCanvas& Canvas)
 					{
 						// so it can get the screen size
@@ -3243,12 +3246,21 @@ void FSceneRenderer::RenderFinish(FRDGBuilder& GraphBuilder, FRDGTextureRef View
 
 						if (bNaniteEnabledButNoAtomics)
 						{
-						#if PLATFORM_WINDOWS || PLATFORM_LINUX
-							FString String = TEXT("Nanite is enabled and used in the scene, but 64bit atomics are not supported (try upgrading video driver). Expect corrupt or incomplete rendering.");
-						#else
-							FString String = TEXT("Nanite is enabled and used in the scene, but 64bit atomics are not supported (not currently supported by this platform).  Expect corrupt or incomplete rendering.");
-						#endif
-							Canvas.DrawShadowedText(10, Y, FText::FromString(String), GetStatsFont(), FLinearColor(1.0, 0.05, 0.05, 1.0));
+							FString NaniteError;
+							if (bNaniteRequireAtomics)
+							{
+								NaniteError = TEXT("Nanite is used in the scene but not supported by your graphics hardware and/or driver. Meshes requiring Nanite will not render");
+							}
+							else
+							{
+							#if PLATFORM_WINDOWS || PLATFORM_LINUX
+								NaniteError = TEXT("Nanite is enabled and used in the scene, but 64bit atomics are not supported (try upgrading video driver). Expect corrupt or incomplete rendering.");
+							#else
+								NaniteError = TEXT("Nanite is enabled and used in the scene, but 64bit atomics are not supported (not currently supported by this platform).  Expect corrupt or incomplete rendering.");
+							#endif
+							}
+
+							Canvas.DrawShadowedText(10, Y, FText::FromString(NaniteError), GetStatsFont(), FLinearColor(1.0, 0.05, 0.05, 1.0));
 							Y += 14;
 						}
 
