Use FindConsoleVariable instead of FindTConsoleVariableDataInt

diff --git a/Engine/Source/Runtime/Renderer/Private/SceneRendering.cpp b/Engine/Source/Runtime/Renderer/Private/SceneRendering.cpp
--- a/Engine/Source/Runtime/Renderer/Private/SceneRendering.cpp
+++ b/Engine/Source/Runtime/Renderer/Private/SceneRendering.cpp
@@ -3043,8 +3043,8 @@ void FSceneRenderer::RenderFinish(FRDGBuilder& GraphBuilder, FRDGTextureRef View
 			}
 		}
 
-		static const auto* CVarNaniteRequireAtomics = IConsoleManager::Get().FindTConsoleVariableDataInt(TEXT("r.Nanite.RequireAtomic64Support"));
-		const bool bNaniteRequireAtomics = CVarNaniteRequireAtomics != nullptr ? CVarNaniteRequireAtomics->GetValueOnAnyThread() != 0 : true;
+		static const auto CVarNaniteRequireAtomics = IConsoleManager::Get().FindConsoleVariable(TEXT("r.Nanite.RequireAtomic64Support"));
+		const bool bNaniteRequireAtomics = CVarNaniteRequireAtomics != nullptr ? CVarNaniteRequireAtomics->GetInt() != 0 : true;
 
 		// Mobile-specific warnings
 		const bool bMobile = (FeatureLevel <= ERHIFeatureLevel::ES3_1);
