Fix incorrect allocation of hair strands cluster data when a hair group component is instanciated.

diff --git a/Engine/Plugins/Runtime/HairStrands/Source/HairStrandsCore/Private/GroomManager.cpp b/Engine/Plugins/Runtime/HairStrands/Source/HairStrandsCore/Private/GroomManager.cpp
--- a/Engine/Plugins/Runtime/HairStrands/Source/HairStrandsCore/Private/GroomManager.cpp
+++ b/Engine/Plugins/Runtime/HairStrands/Source/HairStrandsCore/Private/GroomManager.cpp
@@ -669,6 +669,9 @@ static void RunHairLODSelection(FRDGBuilder& GraphBuilder, const FHairStrandsIns
 		}
 		else if (GeometryType == EHairGeometryType::Strands)
 		{
+			check(Instance->HairGroupPublicData);
+			Instance->HairGroupPublicData->Allocate(GraphBuilder);
+
 			if (Instance->Strands.RestRootResource)			Instance->Strands.RestRootResource->Allocate(GraphBuilder);
 			if (Instance->Strands.RestResource)				Instance->Strands.RestResource->Allocate(GraphBuilder);
 			if (Instance->Strands.ClusterCullingResource)	Instance->Strands.ClusterCullingResource->Allocate(GraphBuilder);
diff --git a/Engine/Source/Runtime/Renderer/Private/HairStrands/HairStrandsInterface.cpp b/Engine/Source/Runtime/Renderer/Private/HairStrands/HairStrandsInterface.cpp
--- a/Engine/Source/Runtime/Renderer/Private/HairStrands/HairStrandsInterface.cpp
+++ b/Engine/Source/Runtime/Renderer/Private/HairStrands/HairStrandsInterface.cpp
@@ -303,6 +303,23 @@ void FHairGroupPublicData::SetClusters(uint32 InClusterCount, uint32 InVertexCou
 
 void FHairGroupPublicData::InitRHI()
 {
+	if (bIsInitialized || GUsingNullRHI) { return; }
+
+	// Resource are allocated on-demand
+	#if 0
+	FMemMark Mark(FMemStack::Get());
+	FRHICommandListImmediate& RHICmdList = FRHICommandListExecutor::GetImmediateCommandList();
+	FRDGBuilder GraphBuilder(RHICmdList);
+	Allocate(GraphBuilder);
+	GraphBuilder.Execute();
+	#endif
+}
+
+void FHairGroupPublicData::Allocate(FRDGBuilder& GraphBuilder)
+{
+	if (bIsInitialized)
+		return;
+
 	if (ClusterCount == 0)
 		return;
 
@@ -318,9 +335,6 @@ void FHairGroupPublicData::InitRHI()
 	
 	if (GUsingNullRHI || !bHasStrands) { return; }
 
-	FMemMark Mark(FMemStack::Get());
-	FRHICommandListImmediate& RHICmdList = FRHICommandListExecutor::GetImmediateCommandList();
-	FRDGBuilder GraphBuilder(RHICmdList);
 	InternalCreateIndirectBufferRDG(GraphBuilder, DrawIndirectBuffer, TEXT("Hair.Cluster_DrawIndirectBuffer"), FUintVector4(GroupControlTriangleStripVertexCount, 1, 0, 0));
 	InternalCreateIndirectBufferRDG(GraphBuilder, DrawIndirectRasterComputeBuffer, TEXT("Hair.Cluster_DrawIndirectRasterComputeBuffer"), FUintVector4(0, 1, 0, 0));
 
@@ -330,10 +344,15 @@ void FHairGroupPublicData::InitRHI()
 	InternalCreateVertexBufferRDG(GraphBuilder, sizeof(int32), VertexCount, EPixelFormat::PF_R32_UINT, CulledVertexIdBuffer, TEXT("Hair.Cluster_CulledVertexIdBuffer"));
 	InternalCreateVertexBufferRDG(GraphBuilder, sizeof(float), VertexCount, EPixelFormat::PF_R32_FLOAT, CulledVertexRadiusScaleBuffer, TEXT("Hair.Cluster_CulledVertexRadiusScaleBuffer"), true);
 
-	GraphBuilder.Execute();
+	bIsInitialized = true;
 }
 
 void FHairGroupPublicData::ReleaseRHI()
+{
+	//Release();
+}
+
+void FHairGroupPublicData::Release()
 {
 	DrawIndirectBuffer.Release();
 	DrawIndirectRasterComputeBuffer.Release();
@@ -341,6 +360,7 @@ void FHairGroupPublicData::ReleaseRHI()
 	GroupAABBBuffer.Release();
 	CulledVertexIdBuffer.Release();
 	CulledVertexRadiusScaleBuffer.Release();
+	bIsInitialized = false;
 }
 
 uint32 FHairGroupPublicData::GetResourcesSize() const
diff --git a/Engine/Source/Runtime/Renderer/Public/HairStrandsInterface.h b/Engine/Source/Runtime/Renderer/Public/HairStrandsInterface.h
--- a/Engine/Source/Runtime/Renderer/Public/HairStrandsInterface.h
+++ b/Engine/Source/Runtime/Renderer/Public/HairStrandsInterface.h
@@ -154,6 +154,8 @@ public:
 	virtual void InitRHI() override;
 	virtual void ReleaseRHI() override;
 	virtual FString GetFriendlyName() const override { return TEXT("FHairGroupPublicData"); }
+	void Allocate(FRDGBuilder& GraphBuilder);
+	void Release();
 	uint32 GetResourcesSize() const;
 
 	// The primitive count when no culling and neither lod happens
@@ -302,6 +304,7 @@ public:
 	FRDGExternalBuffer CulledVertexRadiusScaleBuffer;
 	bool bCullingResultAvailable = false;
 	bool bSupportVoxelization = true;
+	bool bIsInitialized = false;
 
 	/* CPU LOD selection. Hair LOD selection can be done by CPU or GPU. If bUseCPULODSelection is true, 
 	   CPU LOD selection is enabled otherwise the GPU selection is used. CPU LOD selection use the CPU 
