Disabled deferred scene render clean as the default due to crash. Needs investigation.

diff --git a/Engine/Source/Runtime/Renderer/Private/SceneRendering.cpp b/Engine/Source/Runtime/Renderer/Private/SceneRendering.cpp
--- a/Engine/Source/Runtime/Renderer/Private/SceneRendering.cpp
+++ b/Engine/Source/Runtime/Renderer/Private/SceneRendering.cpp
@@ -3581,7 +3581,7 @@ void FSceneRenderer::ViewExtensionPreRender_RenderThread(FRHICommandListImmediat
 	FDeferredUpdateResource::UpdateResources(RHICmdList);
 }
 
-static int32 GSceneRenderCleanUpMode = 2;
+static int32 GSceneRenderCleanUpMode = 0;
 static FAutoConsoleVariableRef CVarSceneRenderCleanUpMode(
 	TEXT("r.SceneRenderCleanUpMode"),
 	GSceneRenderCleanUpMode,
