Change default texture initialization to be done on CPU with lock/unlock to avoid running though RDG, which adds overhead tracking.

diff --git a/Engine/Source/Runtime/Renderer/Private/SystemTextures.cpp b/Engine/Source/Runtime/Renderer/Private/SystemTextures.cpp
--- a/Engine/Source/Runtime/Renderer/Private/SystemTextures.cpp
+++ b/Engine/Source/Runtime/Renderer/Private/SystemTextures.cpp
@@ -899,18 +899,6 @@ FRDGTextureRef FSystemTextures::GetZeroUShort4Dummy(FRDGBuilder& GraphBuilder) c
 //////////////////////////////////////////////////////////////////////////////////////////////////////////
 // Default textures
 
-template<typename TClearType>
-FClearValueBinding GetClearBinding(TClearType Value)
-{
-	return FClearValueBinding::None;
-}
-
-template<>
-FClearValueBinding GetClearBinding(FClearValueBinding Value)
-{
-	return Value;
-}
-
 bool operator !=(const FDefaultTextureKey& A, const FDefaultTextureKey& B)
 {
 	return A.Format != B.Format ||
@@ -935,9 +923,285 @@ static FDefaultTextureKey GetDefaultTextureKey(EPixelFormat Format, const T& In)
 	return Out;
 }
 
-static void AddClearUAVPass(FRDGBuilder& GraphBuilder, FRDGTextureUAVRef TextureUAV, const FClearValueBinding& Value)
+// Convert from X to 4 components data float/uint/int. Supported input are:
+// * float
+// * int32
+// * uint32
+// * FVector2D
+// * FIntPoint
+// * FVector
+// * FVector4
+// * FUintVector4
+// * FClearValueBinding
+FIntVector4		ToVector(int32 Value)						{ return FIntVector4(Value, Value, Value, Value); }
+FVector4		ToVector(float Value)						{ return FVector4(Value, Value, Value, Value); }
+FUintVector4	ToVector(uint32 Value)						{ return FUintVector4(Value, Value, Value, Value); }
+FVector4		ToVector(const FVector & Value)				{ return FVector4(Value.X, Value.Y, Value.Z, 0); }
+FVector4		ToVector(const FVector4 & Value)			{ return Value; }
+FVector4		ToVector(const FVector2D& Value)			{ return FVector4(Value.X, Value.Y, 0, 0); }
+FIntVector4		ToVector(const FIntPoint& Value)			{ return FIntVector4(Value.X, Value.Y, 0, 0); }
+FUintVector4	ToVector(const FUintVector4& Value)			{ return Value; }
+FVector4		ToVector(const FClearValueBinding & Value)	{ return FVector4(Value.Value.Color[0], Value.Value.Color[1], Value.Value.Color[2], Value.Value.Color[3]); }
+
+template <typename TInputType>	struct TFormatConversionTraits				{ /*Error*/ };
+template <>						struct TFormatConversionTraits<FVector4>	{ typedef float  Type; };
+template <>						struct TFormatConversionTraits<FUintVector4>{ typedef uint32 Type; };
+template <>						struct TFormatConversionTraits<FIntVector4> { typedef int32  Type; };
+
+enum class EDefaultInputType
+{
+	Typed,
+	UNorm,
+	SNorm,
+	UNorm10,
+	UNorm11,
+	UNorm2
+};
+
+// Convert input type into the final type. This function manages UNorm/SNorm type by assuming if the input if float, its value is normalized in [0..1].
+template <typename TInType, typename TOutType, EDefaultInputType InputFormatType>
+TOutType ConvertInputFormat(const TInType& In)
+{
+	return TOutType(In);
+}
+template<> uint32 ConvertInputFormat<float, uint32, EDefaultInputType::UNorm>(const float& In) { return FMath::Clamp(In,  0.f, 1.f) * MAX_uint32; }
+template<>  int32 ConvertInputFormat<float,  int32, EDefaultInputType::SNorm>(const float& In) { return FMath::Clamp(In, -1.f, 1.f) * MAX_int32; }
+template<> uint16 ConvertInputFormat<float, uint16, EDefaultInputType::UNorm>(const float& In) { return FMath::Clamp(In,  0.f, 1.f) * MAX_uint16; }
+template<>  int16 ConvertInputFormat<float,  int16, EDefaultInputType::SNorm>(const float& In) { return FMath::Clamp(In, -1.f, 1.f) * MAX_int16; }
+template<>  uint8 ConvertInputFormat<float,  uint8, EDefaultInputType::UNorm>(const float& In) { return FMath::Clamp(In,  0.f, 1.f) * MAX_uint8; }
+template<>   int8 ConvertInputFormat<float,   int8, EDefaultInputType::SNorm>(const float& In) { return FMath::Clamp(In, -1.f, 1.f) * MAX_int8; }
+
+template<> uint32 ConvertInputFormat<float, uint32, EDefaultInputType::UNorm10>(const float& In) { return FMath::Clamp(In, 0.f, 1.f) * 1024u; }
+template<> uint32 ConvertInputFormat<float, uint32, EDefaultInputType::UNorm11>(const float& In) { return FMath::Clamp(In, 0.f, 1.f) * 2048u; }
+template<> uint32 ConvertInputFormat<float, uint32, EDefaultInputType::UNorm2> (const float& In) { return FMath::Clamp(In, 0.f, 1.f) * 3u; }
+
+// 4 components conversion with swizzling
+template <EDefaultInputType InputFormatType, typename TInType, typename TOutType, uint32 SwizzleX, uint32 SwizzleY, uint32 SwizzleZ, uint32 SwizzleW>
+void FormatData(const TInType& In, uint8* Out, uint32& OutByteCount)
+{
+	TOutType* OutTyped = (TOutType*)Out;
+	OutTyped[0] = ConvertInputFormat<TFormatConversionTraits<TInType>::Type, TOutType, InputFormatType>(In[SwizzleX]);
+	OutTyped[1] = ConvertInputFormat<TFormatConversionTraits<TInType>::Type, TOutType, InputFormatType>(In[SwizzleY]);
+	OutTyped[2] = ConvertInputFormat<TFormatConversionTraits<TInType>::Type, TOutType, InputFormatType>(In[SwizzleZ]);
+	OutTyped[3] = ConvertInputFormat<TFormatConversionTraits<TInType>::Type, TOutType, InputFormatType>(In[SwizzleW]);
+	OutByteCount = 4 * sizeof(TOutType);
+}
+
+// 3 components conversion with swizzling
+template <EDefaultInputType InputFormatType, typename TInType, typename TOutType, uint32 SwizzleX, uint32 SwizzleY, uint32 SwizzleZ>
+void FormatData(const TInType& In, uint8* Out, uint32& OutByteCount)
+{
+	TOutType* OutTyped = (TOutType*)Out;
+	OutTyped[0] = ConvertInputFormat<TFormatConversionTraits<TInType>::Type, TOutType, InputFormatType>(In[SwizzleX]);
+	OutTyped[1] = ConvertInputFormat<TFormatConversionTraits<TInType>::Type, TOutType, InputFormatType>(In[SwizzleY]);
+	OutTyped[2] = ConvertInputFormat<TFormatConversionTraits<TInType>::Type, TOutType, InputFormatType>(In[SwizzleZ]);
+	OutByteCount = 3 * sizeof(TOutType);
+}
+
+// 2 components conversion with swizzling
+template <EDefaultInputType InputFormatType, typename TInType, typename TOutType, uint32 SwizzleX, uint32 SwizzleY>
+void FormatData(const TInType& In, uint8* Out, uint32& OutByteCount)
+{
+	TOutType* OutTyped = (TOutType*)Out;
+	OutTyped[0] = ConvertInputFormat<TFormatConversionTraits<TInType>::Type, TOutType, InputFormatType>(In[SwizzleX]);
+	OutTyped[1] = ConvertInputFormat<TFormatConversionTraits<TInType>::Type, TOutType, InputFormatType>(In[SwizzleY]);
+	OutByteCount = 2 * sizeof(TOutType);
+}
+
+// 1 component conversion
+template <EDefaultInputType InputFormatType, typename TInType, typename TOutType>
+void FormatData(const TInType& In, uint8* Out, uint32& OutByteCount)
+{
+	TOutType* OutTyped = (TOutType*)Out;
+	OutTyped[0] = ConvertInputFormat<TFormatConversionTraits<TInType>::Type, TOutType, InputFormatType>(In[0]);
+	OutByteCount = 4;
+}
+
+template <typename TInType>
+void FormatData111110(const TInType& In, uint8* Out, uint32& OutByteCount)
+{
+	uint32* OutTyped = (uint32*)Out;
+	*OutTyped = 
+		 (2048u & ConvertInputFormat<TFormatConversionTraits<TInType>::Type, uint32, EDefaultInputType::UNorm11>(In[0]))     |
+		((2048u & ConvertInputFormat<TFormatConversionTraits<TInType>::Type, uint32, EDefaultInputType::UNorm11>(In[1]))<<11)|
+		((1024u & ConvertInputFormat<TFormatConversionTraits<TInType>::Type, uint32, EDefaultInputType::UNorm10>(In[2]))<<22);
+	OutByteCount = 4;
+}
+
+template <typename TInType>
+void FormatData1010102(const TInType& In, uint8* Out, uint32& OutByteCount)
 {
-	// Nothing. This is just for compiling template function
+	uint32* OutTyped = (uint32*)Out;
+	*OutTyped =
+		 (1024u & ConvertInputFormat<TFormatConversionTraits<TInType>::Type, uint32, EDefaultInputType::UNorm10>(In[0]))        |
+		((1024u & ConvertInputFormat<TFormatConversionTraits<TInType>::Type, uint32, EDefaultInputType::UNorm10>(In[1])) << 10) |
+		((1024u & ConvertInputFormat<TFormatConversionTraits<TInType>::Type, uint32, EDefaultInputType::UNorm10>(In[2])) << 20) |
+		((   3u & ConvertInputFormat<TFormatConversionTraits<TInType>::Type, uint32, EDefaultInputType::UNorm2> (In[3])) << 30);
+	OutByteCount = 4;
+}
+
+template<typename TInType>
+void InitializeData(const TInType& InData, EPixelFormat InFormat, uint8* OutData, uint32& OutByteCount)
+{
+	// If a new format is added insure that it is either supported here, or at least flagged as not supported
+	static_assert(PF_MAX == 72);
+
+	switch (InFormat)
+	{
+		// 32bits
+		case PF_R32G32B32A32_UINT:		{ FormatData<EDefaultInputType::Typed, TInType,	uint32,  0, 1, 2, 3>	(InData, OutData, OutByteCount); } break;
+		case PF_A32B32G32R32F:			{ FormatData<EDefaultInputType::Typed, TInType,	float,   3, 2, 1, 0>	(InData, OutData, OutByteCount); } break;
+		case PF_R32G32_UINT:			{ FormatData<EDefaultInputType::Typed, TInType,	uint32,  0, 1>			(InData, OutData, OutByteCount); } break;
+		case PF_G32R32F:				{ FormatData<EDefaultInputType::Typed, TInType,	float,   1, 0>			(InData, OutData, OutByteCount); } break;
+		case PF_R32_UINT:				{ FormatData<EDefaultInputType::Typed, TInType,	uint32>					(InData, OutData, OutByteCount); } break;
+		case PF_R32_SINT:				{ FormatData<EDefaultInputType::Typed, TInType,	int32>					(InData, OutData, OutByteCount); } break;
+		case PF_R32_FLOAT:				{ FormatData<EDefaultInputType::Typed, TInType,	float>					(InData, OutData, OutByteCount); } break;
+
+		// 16bits
+		case PF_R16G16B16A16_UINT:		{ FormatData<EDefaultInputType::Typed, TInType,	uint16,   0, 1, 2, 3>	(InData, OutData, OutByteCount); } break;
+		case PF_R16G16B16A16_SINT:		{ FormatData<EDefaultInputType::Typed, TInType,	int16,    0, 1, 2, 3>	(InData, OutData, OutByteCount); } break;
+		case PF_R16G16B16A16_UNORM:		{ FormatData<EDefaultInputType::UNorm, TInType,	uint16,   0, 1, 2, 3>	(InData, OutData, OutByteCount); } break;
+		case PF_R16G16B16A16_SNORM:		{ FormatData<EDefaultInputType::SNorm, TInType,	int16,    0, 1, 2, 3>	(InData, OutData, OutByteCount); } break;
+		case PF_A16B16G16R16:			{ FormatData<EDefaultInputType::UNorm, TInType,	uint16,   3, 2, 1, 0>	(InData, OutData, OutByteCount); } break;
+		case PF_FloatRGBA:				{ FormatData<EDefaultInputType::Typed, TInType,	FFloat16, 0, 1, 2, 3>	(InData, OutData, OutByteCount); } break;
+		case PF_R16G16_UINT:			{ FormatData<EDefaultInputType::Typed, TInType,	uint16,   0, 1>			(InData, OutData, OutByteCount); } break;
+		case PF_G16R16:					{ FormatData<EDefaultInputType::UNorm, TInType,	uint16,   1, 0>			(InData, OutData, OutByteCount); } break;
+		case PF_G16R16F:				{ FormatData<EDefaultInputType::Typed, TInType,	FFloat16, 0, 1>			(InData, OutData, OutByteCount); } break;
+		case PF_G16R16F_FILTER:			{ FormatData<EDefaultInputType::Typed, TInType,	FFloat16, 0, 1>			(InData, OutData, OutByteCount); } break;
+		case PF_R16F_FILTER:			{ FormatData<EDefaultInputType::Typed, TInType,	FFloat16>				(InData, OutData, OutByteCount); } break;
+		case PF_R16F:					{ FormatData<EDefaultInputType::Typed, TInType,	FFloat16>				(InData, OutData, OutByteCount); } break;
+		case PF_G16:					{ FormatData<EDefaultInputType::UNorm, TInType,	uint16>					(InData, OutData, OutByteCount); } break;
+		case PF_R16_UINT:				{ FormatData<EDefaultInputType::Typed, TInType,	uint16>					(InData, OutData, OutByteCount); } break;
+		case PF_R16_SINT:				{ FormatData<EDefaultInputType::Typed, TInType,	int16>					(InData, OutData, OutByteCount); } break;
+
+		// 8bits
+		case PF_B8G8R8A8:				{ FormatData<EDefaultInputType::UNorm, TInType,	uint8,    2, 1, 0, 3>	(InData, OutData, OutByteCount); } break;
+		case PF_R8G8B8A8:				{ FormatData<EDefaultInputType::UNorm, TInType,	uint8,    0, 1, 2, 3>	(InData, OutData, OutByteCount); } break;
+		case PF_A8R8G8B8:				{ FormatData<EDefaultInputType::UNorm, TInType,	uint8,    3, 2, 1, 0>	(InData, OutData, OutByteCount); } break;
+		case PF_R8G8B8A8_UINT:			{ FormatData<EDefaultInputType::Typed, TInType,	uint8,    0, 1, 2, 3>	(InData, OutData, OutByteCount); } break;
+		case PF_R8G8B8A8_SNORM:			{ FormatData<EDefaultInputType::SNorm, TInType,	int8,     0, 1, 2, 3>	(InData, OutData, OutByteCount); } break;
+		case PF_R8G8:					{ FormatData<EDefaultInputType::UNorm, TInType,	uint8,    0, 1>			(InData, OutData, OutByteCount); } break;
+		case PF_R8_UINT:				{ FormatData<EDefaultInputType::Typed, TInType,	uint8>					(InData, OutData, OutByteCount); } break;
+		case PF_R8:						{ FormatData<EDefaultInputType::UNorm, TInType,	uint8>					(InData, OutData, OutByteCount); } break;
+		case PF_G8:						{ FormatData<EDefaultInputType::UNorm, TInType,	uint8>					(InData, OutData, OutByteCount); } break;
+		case PF_L8:						{ FormatData<EDefaultInputType::UNorm, TInType,	uint8>					(InData, OutData, OutByteCount); } break;
+		case PF_A1:						{ FormatData<EDefaultInputType::UNorm, TInType,	uint8>					(InData, OutData, OutByteCount); } break;
+		case PF_A8:						{ FormatData<EDefaultInputType::UNorm, TInType,	uint8>					(InData, OutData, OutByteCount); } break;
+
+		// Depth/Stencil. Since these texture will only be used as SRV, we handle them as regular float/float16.
+		case PF_D24:					{ FormatData<EDefaultInputType::Typed, TInType,	float>					(InData, OutData, OutByteCount); } break;
+		case PF_DepthStencil:			{ FormatData<EDefaultInputType::Typed, TInType,	float>					(InData, OutData, OutByteCount); } break;
+		case PF_ShadowDepth:			{ FormatData<EDefaultInputType::Typed, TInType,	FFloat16>				(InData, OutData, OutByteCount); } break;
+
+		// Custom
+		case PF_FloatRGB:				{ FormatData111110<TInType>	(InData, OutData, OutByteCount); } break;
+		case PF_A2B10G10R10:			{ FormatData1010102<TInType>(InData, OutData, OutByteCount); } break;
+		case PF_FloatR11G11B10:			{ FormatData111110<TInType>	(InData, OutData, OutByteCount); } break;
+			return;
+
+		// Not supported
+		case PF_R5G6B5_UNORM:
+		case PF_BC5:
+		case PF_V8U8:
+		case PF_PVRTC2:
+		case PF_PVRTC4:
+		case PF_UYVY:
+		case PF_DXT1:
+		case PF_DXT3:
+		case PF_DXT5:
+		case PF_BC4:
+		case PF_ATC_RGB:
+		case PF_ATC_RGBA_E:
+		case PF_ATC_RGBA_I:
+		case PF_X24_G8:
+		case PF_ETC1:
+		case PF_ETC2_RGB:
+		case PF_ETC2_RGBA:
+		case PF_ASTC_4x4:
+		case PF_ASTC_6x6:
+		case PF_ASTC_8x8:
+		case PF_ASTC_10x10:
+		case PF_ASTC_12x12:
+		case PF_BC6H:
+		case PF_BC7:
+		case PF_XGXR8:
+		case PF_PLATFORM_HDR_0:
+		case PF_PLATFORM_HDR_1:
+		case PF_PLATFORM_HDR_2:
+		case PF_NV12:
+		case PF_ETC2_R11_EAC:
+		case PF_ETC2_RG11_EAC:
+		case PF_Unknown:
+		case PF_MAX:
+			OutByteCount = 0;
+			return;
+	}
+}
+
+template <typename DataType>
+void SetDefaultTextureData2D(FRHITexture2D* Texture, const DataType& InData)
+{
+	uint8 SrcData[16] = { 0,0,0,0, 0,0,0,0, 0,0,0,0, 0,0,0,0, };
+	uint32 SrcByteCount = 0;
+	const EPixelFormat Format = Texture->GetFormat();
+	InitializeData(ToVector(InData), Format, SrcData, SrcByteCount);
+
+	uint32 DestStride;
+	uint8* Dest = (uint8*)RHILockTexture2D(Texture, 0, RLM_WriteOnly, DestStride, false);
+	FMemory::Memcpy(Dest, SrcData, SrcByteCount);
+	RHIUnlockTexture2D(Texture, 0, false);
+}
+
+template <typename DataType>
+void SetDefaultTextureData2DArray(FRHITexture2DArray* Texture, const DataType& InData)
+{
+	uint8 SrcData[16] = { 0,0,0,0, 0,0,0,0, 0,0,0,0, 0,0,0,0, };
+	uint32 SrcByteCount = 0;
+	const EPixelFormat Format = Texture->GetFormat();
+	InitializeData(ToVector(InData), Format, SrcData, SrcByteCount);
+
+	uint32 DestStride;
+	uint8* Dest = (uint8*)RHILockTexture2DArray(Texture, 0, 0, RLM_WriteOnly, DestStride, false);
+	FMemory::Memcpy(Dest, SrcData, SrcByteCount);
+	RHIUnlockTexture2DArray(Texture, 0, 0, false);
+}
+
+template <typename DataType>
+void SetDefaultTextureData3D(FRHICommandListImmediate& RHICmdList, FRHITexture3D* Texture, const DataType& InData)
+{
+	uint8 SrcData[16] = { 0,0,0,0, 0,0,0,0, 0,0,0,0, 0,0,0,0, };
+	uint32 SrcByteCount = 0;
+	const EPixelFormat Format = Texture->GetFormat();
+	InitializeData(ToVector(InData), Format, SrcData, SrcByteCount);
+
+	FUpdateTextureRegion3D Region(0, 0, 0, 0, 0, 0, 1, 1, 1);
+	RHICmdList.UpdateTexture3D(
+		Texture,
+		0,
+		Region,
+		SrcByteCount,
+		SrcByteCount,
+		SrcData);
+
+	// UpdateTexture3D before and after state is currently undefined
+	RHICmdList.Transition(FRHITransitionInfo(Texture, ERHIAccess::Unknown, ERHIAccess::SRVMask));
+}
+
+template <typename DataType>
+void SetDefaultTextureDataCube(FRHITextureCube* Texture, const DataType& InData)
+{
+	uint8 SrcData[16] = { 0,0,0,0, 0,0,0,0, 0,0,0,0, 0,0,0,0, };
+	uint32 SrcByteCount = 0;
+	const EPixelFormat Format = Texture->GetFormat();
+	InitializeData(ToVector(InData), Format, SrcData, SrcByteCount);
+
+	for (uint32 FaceIt = 0; FaceIt < 6; ++FaceIt)
+	{
+		uint32 DestStride;
+		uint8* Dest = (uint8*)RHILockTextureCubeFace(Texture, FaceIt, 0u, 0u, RLM_WriteOnly, DestStride, false);
+		FMemory::Memcpy(Dest, SrcData, SrcByteCount);
+		RHIUnlockTextureCubeFace(Texture, FaceIt, 0, 0, false);
+	}
 }
 
 template<typename TClearValue>
@@ -952,6 +1216,10 @@ FRDGTextureRef GetInternalDefaultTexture(
 	// Check this is a valid format
 	check(Format != PF_Unknown && Format != PF_MAX && GPixelFormats[Format].BlockSizeX == 1 && GPixelFormats[Format].BlockSizeY == 1 && GPixelFormats[Format].BlockSizeZ == 1);
 
+	// Convert Depth/Stencil format to float/float16 since these texture will only be used as SRV
+	if (Format == PF_D24 || Format == PF_DepthStencil)	{ Format = PF_R32_FLOAT; }
+	if (Format == PF_ShadowDepth)						{ Format = PF_R32_FLOAT; }
+
 	const FDefaultTextureKey Key = GetDefaultTextureKey(Format, Value);
 	const uint32 Hash = Murmur32({uint32(Key.Dimension), uint32(Key.Format), Key.ValueAsUInt[0], Key.ValueAsUInt[1], Key.ValueAsUInt[2], Key.ValueAsUInt[3]});
 
@@ -967,64 +1235,56 @@ FRDGTextureRef GetInternalDefaultTexture(
 		return GraphBuilder.RegisterExternalTexture(DefaultTextures[Index].Texture);
 	}
 
-	bool bSupportUAVWrite = true;
-	bool bSupportGraphicsWrite = true;
-	FClearValueBinding ClearValue = GetClearBinding(Value);
-	if (IsDepthOrStencilFormat(Format))
-	{
-		check(
-			ClearValue == FClearValueBinding::DepthZero || 
-			ClearValue == FClearValueBinding::DepthOne  || 
-			ClearValue == FClearValueBinding::DepthFar  ||
-			ClearValue == FClearValueBinding::DepthNear);
-		bSupportUAVWrite = false;
-	}
-
-	ETextureCreateFlags Flags =
-		TexCreate_ShaderResource |
-		(bSupportUAVWrite ? TexCreate_UAV : TexCreate_None) |
-		(bSupportGraphicsWrite ? TexCreate_RenderTargetable : TexCreate_None);
-
-	FRDGTextureRef Texture = nullptr; 
-	if (Dimension == ETextureDimension::Texture2D)
-	{
-		Texture = GraphBuilder.CreateTexture(FRDGTextureDesc::Create2D(FIntPoint(1, 1), Format, ClearValue, TexCreate_ShaderResource | Flags), TEXT("DefaultTexture"), ERDGTextureFlags::MultiFrame); // TODO Create better name e.g., DefaultTexture(2D,PF_R32INT,Zero)
-	}
-	else if (Dimension == ETextureDimension::Texture2DArray)
-	{
-		Texture = GraphBuilder.CreateTexture(FRDGTextureDesc::Create2DArray(FIntPoint(1, 1), Format, ClearValue, TexCreate_ShaderResource | Flags, 1), TEXT("DefaultTexture"), ERDGTextureFlags::MultiFrame);
-	}
-	else if (Dimension == ETextureDimension::Texture3D)
-	{
-		Texture = GraphBuilder.CreateTexture(FRDGTextureDesc::Create3D(FIntVector(1, 1, 1), Format, ClearValue, TexCreate_ShaderResource | Flags), TEXT("DefaultTexture"), ERDGTextureFlags::MultiFrame);
-	}
-	else if (Dimension == ETextureDimension::TextureCube)
-	{
-		Texture = GraphBuilder.CreateTexture(FRDGTextureDesc::CreateCube(1, Format, ClearValue, TexCreate_ShaderResource | Flags), TEXT("DefaultTexture"), ERDGTextureFlags::MultiFrame);
-	}
-	else if (Dimension == ETextureDimension::TextureCubeArray)
-	{
-		Texture = GraphBuilder.CreateTexture(FRDGTextureDesc::CreateCubeArray(1, Format, ClearValue, TexCreate_ShaderResource | Flags, 1), TEXT("DefaultTexture"), ERDGTextureFlags::MultiFrame);
-	}
-	else
-	{
-		return nullptr;
-	}
-
-	if (ClearValue == FClearValueBinding::None)
-	{
-		check(bSupportUAVWrite);
-		AddClearUAVPass(GraphBuilder, GraphBuilder.CreateUAV(Texture), Value);
-	}
-
 	FDefaultTexture Entry;
 	Entry.Key = Key;
 	Entry.Hash = Hash;
-	Entry.Texture = GraphBuilder.ConvertToExternalTexture(Texture);
+	Entry.Texture = nullptr;
+	
+	{
+		if (Dimension == ETextureDimension::Texture2D)
+		{
+			FRHIResourceCreateInfo CreateInfo(TEXT("DefaultTexture2D"));
+			FTexture2DRHIRef Texture = RHICreateTexture2D(1, 1, Format, 1, 1, TexCreate_ShaderResource, CreateInfo);
+			SetDefaultTextureData2D(Texture, Value);
+			Entry.Texture = CreateRenderTarget(Texture, CreateInfo.DebugName);
+		}
+		else if (Dimension == ETextureDimension::Texture2DArray)
+		{
+			FRHIResourceCreateInfo CreateInfo(TEXT("DefaultTexture2DArray"));
+			FTexture2DArrayRHIRef Texture = RHICreateTexture2DArray(1, 1, 1, Format, 1, 1, TexCreate_ShaderResource, CreateInfo);
+			SetDefaultTextureData2DArray(Texture, Value);
+			Entry.Texture = CreateRenderTarget(Texture, CreateInfo.DebugName);
+		}
+		else if (Dimension == ETextureDimension::Texture3D)
+		{
+			FRHIResourceCreateInfo CreateInfo(TEXT("DefaultTexture3D"));
+			FTexture3DRHIRef Texture = RHICreateTexture3D(1, 1, 1, Format, 1, TexCreate_ShaderResource, CreateInfo);
+			SetDefaultTextureData3D(GraphBuilder.RHICmdList, Texture, Value);
+			Entry.Texture = CreateRenderTarget(Texture, CreateInfo.DebugName);
+		}
+		else if (Dimension == ETextureDimension::TextureCube)
+		{
+			FRHIResourceCreateInfo CreateInfo(TEXT("DefaultTextureCube"));
+			FTextureCubeRHIRef Texture = RHICreateTextureCube(1, Format, 1, TexCreate_ShaderResource, CreateInfo);
+			SetDefaultTextureDataCube(Texture, Value);
+			Entry.Texture = CreateRenderTarget(Texture, CreateInfo.DebugName);
+		}
+		else if (Dimension == ETextureDimension::TextureCubeArray)
+		{
+			FRHIResourceCreateInfo CreateInfo(TEXT("DefaultTextureCubeArray"));
+			FTextureCubeRHIRef Texture = RHICreateTextureCubeArray(1, 1, Format, 1, TexCreate_ShaderResource, CreateInfo);
+			SetDefaultTextureDataCube(Texture, Value);
+			Entry.Texture = CreateRenderTarget(Texture, CreateInfo.DebugName);
+		}
+		else
+		{
+			return nullptr;
+		}
+	}
 
 	Index = DefaultTextures.Add(Entry);
 	HashDefaultTextures.Add(Hash, Index);
-	return Texture;
+	return GraphBuilder.RegisterExternalTexture(Entry.Texture);
 }
 
 //////////////////////////////////////////////////////////////////////////////////////////////////////////
@@ -1123,6 +1383,18 @@ FRDGBufferRef GetInternalDefaultBuffer(
 	return Buffer;
 }
 
+FVector4 GetClearBindingValue(EPixelFormat Format, FClearValueBinding Value)
+{
+	if (IsDepthOrStencilFormat(Format))
+	{
+		return FVector4(Value.Value.DSValue.Depth, Value.Value.DSValue.Depth, Value.Value.DSValue.Depth, Value.Value.DSValue.Depth);
+	}
+	else
+	{
+		return FVector4(Value.Value.Color[0], Value.Value.Color[1], Value.Value.Color[2], Value.Value.Color[3]);
+	}
+}
+
 //////////////////////////////////////////////////////////////////////////////////////////////////////////////////////////////////////////////////////////////////////////////////////////////////////////////////////////////////////////////////////////////////////////////////////////////////////////////////
 // Textures 
 
@@ -1131,7 +1403,7 @@ FRDGTextureRef FSystemTextures::GetDefaultTexture2D(FRDGBuilder& GraphBuilder, E
 FRDGTextureRef FSystemTextures::GetDefaultTexture2D(FRDGBuilder& GraphBuilder, EPixelFormat Format, const FVector& Value)										{ return GetInternalDefaultTexture(GraphBuilder, DefaultTextures, HashDefaultTextures, ETextureDimension::Texture2D, Format, Value); }
 FRDGTextureRef FSystemTextures::GetDefaultTexture2D(FRDGBuilder& GraphBuilder, EPixelFormat Format, const FVector4& Value)										{ return GetInternalDefaultTexture(GraphBuilder, DefaultTextures, HashDefaultTextures, ETextureDimension::Texture2D, Format, Value); }
 FRDGTextureRef FSystemTextures::GetDefaultTexture2D(FRDGBuilder& GraphBuilder, EPixelFormat Format, const FUintVector4& Value)									{ return GetInternalDefaultTexture(GraphBuilder, DefaultTextures, HashDefaultTextures, ETextureDimension::Texture2D, Format, Value); }
-FRDGTextureRef FSystemTextures::GetDefaultTexture2D(FRDGBuilder& GraphBuilder, EPixelFormat Format, const FClearValueBinding& Value)							{ return GetInternalDefaultTexture(GraphBuilder, DefaultTextures, HashDefaultTextures, ETextureDimension::Texture2D, Format, Value); }
+FRDGTextureRef FSystemTextures::GetDefaultTexture2D(FRDGBuilder& GraphBuilder, EPixelFormat Format, const FClearValueBinding& Value)							{ return GetInternalDefaultTexture(GraphBuilder, DefaultTextures, HashDefaultTextures, ETextureDimension::Texture2D, Format, GetClearBindingValue(Format, Value)); }
 
 FRDGTextureRef FSystemTextures::GetDefaultTexture(FRDGBuilder& GraphBuilder, ETextureDimension Dimension, EPixelFormat Format, float Value)						{ return GetInternalDefaultTexture(GraphBuilder, DefaultTextures, HashDefaultTextures, Dimension, Format, Value); }
 FRDGTextureRef FSystemTextures::GetDefaultTexture(FRDGBuilder& GraphBuilder, ETextureDimension Dimension, EPixelFormat Format, uint32 Value)					{ return GetInternalDefaultTexture(GraphBuilder, DefaultTextures, HashDefaultTextures, Dimension, Format, Value); }
@@ -1140,7 +1412,7 @@ FRDGTextureRef FSystemTextures::GetDefaultTexture(FRDGBuilder& GraphBuilder, ETe
 FRDGTextureRef FSystemTextures::GetDefaultTexture(FRDGBuilder& GraphBuilder, ETextureDimension Dimension, EPixelFormat Format, const FVector& Value)			{ return GetInternalDefaultTexture(GraphBuilder, DefaultTextures, HashDefaultTextures, Dimension, Format, Value); }
 FRDGTextureRef FSystemTextures::GetDefaultTexture(FRDGBuilder& GraphBuilder, ETextureDimension Dimension, EPixelFormat Format, const FVector4& Value)			{ return GetInternalDefaultTexture(GraphBuilder, DefaultTextures, HashDefaultTextures, Dimension, Format, Value); }
 FRDGTextureRef FSystemTextures::GetDefaultTexture(FRDGBuilder& GraphBuilder, ETextureDimension Dimension, EPixelFormat Format, const FUintVector4& Value)		{ return GetInternalDefaultTexture(GraphBuilder, DefaultTextures, HashDefaultTextures, Dimension, Format, Value); }
-FRDGTextureRef FSystemTextures::GetDefaultTexture(FRDGBuilder& GraphBuilder, ETextureDimension Dimension, EPixelFormat Format, const FClearValueBinding& Value)	{ return GetInternalDefaultTexture(GraphBuilder, DefaultTextures, HashDefaultTextures, Dimension, Format, Value); }
+FRDGTextureRef FSystemTextures::GetDefaultTexture(FRDGBuilder& GraphBuilder, ETextureDimension Dimension, EPixelFormat Format, const FClearValueBinding& Value)	{ return GetInternalDefaultTexture(GraphBuilder, DefaultTextures, HashDefaultTextures, Dimension, Format, GetClearBindingValue(Format, Value)); }
 
 //////////////////////////////////////////////////////////////////////////////////////////////////////////////////////////////////////////////////////////////////////////////////////////////////////////////////////////////////////////////////////////////////////////////////////////////////////////////////
 // Buffers
