Fix CIS issue.

diff --git a/Engine/Source/Runtime/Renderer/Private/SystemTextures.cpp b/Engine/Source/Runtime/Renderer/Private/SystemTextures.cpp
--- a/Engine/Source/Runtime/Renderer/Private/SystemTextures.cpp
+++ b/Engine/Source/Runtime/Renderer/Private/SystemTextures.cpp
@@ -980,10 +980,10 @@ template <EDefaultInputType InputFormatType, typename TInType, typename TOutType
 void FormatData(const TInType& In, uint8* Out, uint32& OutByteCount)
 {
 	TOutType* OutTyped = (TOutType*)Out;
-	OutTyped[0] = ConvertInputFormat<TFormatConversionTraits<TInType>::Type, TOutType, InputFormatType>(In[SwizzleX]);
-	OutTyped[1] = ConvertInputFormat<TFormatConversionTraits<TInType>::Type, TOutType, InputFormatType>(In[SwizzleY]);
-	OutTyped[2] = ConvertInputFormat<TFormatConversionTraits<TInType>::Type, TOutType, InputFormatType>(In[SwizzleZ]);
-	OutTyped[3] = ConvertInputFormat<TFormatConversionTraits<TInType>::Type, TOutType, InputFormatType>(In[SwizzleW]);
+	OutTyped[0] = ConvertInputFormat<typename TFormatConversionTraits<TInType>::Type, TOutType, InputFormatType>(In[SwizzleX]);
+	OutTyped[1] = ConvertInputFormat<typename TFormatConversionTraits<TInType>::Type, TOutType, InputFormatType>(In[SwizzleY]);
+	OutTyped[2] = ConvertInputFormat<typename TFormatConversionTraits<TInType>::Type, TOutType, InputFormatType>(In[SwizzleZ]);
+	OutTyped[3] = ConvertInputFormat<typename TFormatConversionTraits<TInType>::Type, TOutType, InputFormatType>(In[SwizzleW]);
 	OutByteCount = 4 * sizeof(TOutType);
 }
 
@@ -992,9 +992,9 @@ template <EDefaultInputType InputFormatType, typename TInType, typename TOutType
 void FormatData(const TInType& In, uint8* Out, uint32& OutByteCount)
 {
 	TOutType* OutTyped = (TOutType*)Out;
-	OutTyped[0] = ConvertInputFormat<TFormatConversionTraits<TInType>::Type, TOutType, InputFormatType>(In[SwizzleX]);
-	OutTyped[1] = ConvertInputFormat<TFormatConversionTraits<TInType>::Type, TOutType, InputFormatType>(In[SwizzleY]);
-	OutTyped[2] = ConvertInputFormat<TFormatConversionTraits<TInType>::Type, TOutType, InputFormatType>(In[SwizzleZ]);
+	OutTyped[0] = ConvertInputFormat<typename TFormatConversionTraits<TInType>::Type, TOutType, InputFormatType>(In[SwizzleX]);
+	OutTyped[1] = ConvertInputFormat<typename TFormatConversionTraits<TInType>::Type, TOutType, InputFormatType>(In[SwizzleY]);
+	OutTyped[2] = ConvertInputFormat<typename TFormatConversionTraits<TInType>::Type, TOutType, InputFormatType>(In[SwizzleZ]);
 	OutByteCount = 3 * sizeof(TOutType);
 }
 
@@ -1003,8 +1003,8 @@ template <EDefaultInputType InputFormatType, typename TInType, typename TOutType
 void FormatData(const TInType& In, uint8* Out, uint32& OutByteCount)
 {
 	TOutType* OutTyped = (TOutType*)Out;
-	OutTyped[0] = ConvertInputFormat<TFormatConversionTraits<TInType>::Type, TOutType, InputFormatType>(In[SwizzleX]);
-	OutTyped[1] = ConvertInputFormat<TFormatConversionTraits<TInType>::Type, TOutType, InputFormatType>(In[SwizzleY]);
+	OutTyped[0] = ConvertInputFormat<typename TFormatConversionTraits<TInType>::Type, TOutType, InputFormatType>(In[SwizzleX]);
+	OutTyped[1] = ConvertInputFormat<typename TFormatConversionTraits<TInType>::Type, TOutType, InputFormatType>(In[SwizzleY]);
 	OutByteCount = 2 * sizeof(TOutType);
 }
 
@@ -1013,7 +1013,7 @@ template <EDefaultInputType InputFormatType, typename TInType, typename TOutType
 void FormatData(const TInType& In, uint8* Out, uint32& OutByteCount)
 {
 	TOutType* OutTyped = (TOutType*)Out;
-	OutTyped[0] = ConvertInputFormat<TFormatConversionTraits<TInType>::Type, TOutType, InputFormatType>(In[0]);
+	OutTyped[0] = ConvertInputFormat<typename TFormatConversionTraits<TInType>::Type, TOutType, InputFormatType>(In[0]);
 	OutByteCount = 4;
 }
 
@@ -1022,9 +1022,9 @@ void FormatData111110(const TInType& In, uint8* Out, uint32& OutByteCount)
 {
 	uint32* OutTyped = (uint32*)Out;
 	*OutTyped = 
-		 (2048u & ConvertInputFormat<TFormatConversionTraits<TInType>::Type, uint32, EDefaultInputType::UNorm11>(In[0]))     |
-		((2048u & ConvertInputFormat<TFormatConversionTraits<TInType>::Type, uint32, EDefaultInputType::UNorm11>(In[1]))<<11)|
-		((1024u & ConvertInputFormat<TFormatConversionTraits<TInType>::Type, uint32, EDefaultInputType::UNorm10>(In[2]))<<22);
+		 (2048u & ConvertInputFormat<typename TFormatConversionTraits<TInType>::Type, uint32, EDefaultInputType::UNorm11>(In[0]))     |
+		((2048u & ConvertInputFormat<typename TFormatConversionTraits<TInType>::Type, uint32, EDefaultInputType::UNorm11>(In[1]))<<11)|
+		((1024u & ConvertInputFormat<typename TFormatConversionTraits<TInType>::Type, uint32, EDefaultInputType::UNorm10>(In[2]))<<22);
 	OutByteCount = 4;
 }
 
@@ -1033,10 +1033,10 @@ void FormatData1010102(const TInType& In, uint8* Out, uint32& OutByteCount)
 {
 	uint32* OutTyped = (uint32*)Out;
 	*OutTyped =
-		 (1024u & ConvertInputFormat<TFormatConversionTraits<TInType>::Type, uint32, EDefaultInputType::UNorm10>(In[0]))        |
-		((1024u & ConvertInputFormat<TFormatConversionTraits<TInType>::Type, uint32, EDefaultInputType::UNorm10>(In[1])) << 10) |
-		((1024u & ConvertInputFormat<TFormatConversionTraits<TInType>::Type, uint32, EDefaultInputType::UNorm10>(In[2])) << 20) |
-		((   3u & ConvertInputFormat<TFormatConversionTraits<TInType>::Type, uint32, EDefaultInputType::UNorm2> (In[3])) << 30);
+		 (1024u & ConvertInputFormat<typename TFormatConversionTraits<TInType>::Type, uint32, EDefaultInputType::UNorm10>(In[0]))        |
+		((1024u & ConvertInputFormat<typename TFormatConversionTraits<TInType>::Type, uint32, EDefaultInputType::UNorm10>(In[1])) << 10) |
+		((1024u & ConvertInputFormat<typename TFormatConversionTraits<TInType>::Type, uint32, EDefaultInputType::UNorm10>(In[2])) << 20) |
+		((   3u & ConvertInputFormat<typename TFormatConversionTraits<TInType>::Type, uint32, EDefaultInputType::UNorm2> (In[3])) << 30);
 	OutByteCount = 4;
 }
 
