Fix another CIS issue with Clang.

diff --git a/Engine/Source/Runtime/Renderer/Private/SystemTextures.cpp b/Engine/Source/Runtime/Renderer/Private/SystemTextures.cpp
--- a/Engine/Source/Runtime/Renderer/Private/SystemTextures.cpp
+++ b/Engine/Source/Runtime/Renderer/Private/SystemTextures.cpp
@@ -964,8 +964,8 @@ TOutType ConvertInputFormat(const TInType& In)
 {
 	return TOutType(In);
 }
-template<> uint32 ConvertInputFormat<float, uint32, EDefaultInputType::UNorm>(const float& In) { return FMath::Clamp(In,  0.f, 1.f) * MAX_uint32; }
-template<>  int32 ConvertInputFormat<float,  int32, EDefaultInputType::SNorm>(const float& In) { return FMath::Clamp(In, -1.f, 1.f) * MAX_int32; }
+template<> uint32 ConvertInputFormat<float, uint32, EDefaultInputType::UNorm>(const float& In) { return FMath::Clamp(In,  0.f, 1.f) * float(MAX_uint32); }
+template<>  int32 ConvertInputFormat<float,  int32, EDefaultInputType::SNorm>(const float& In) { return FMath::Clamp(In, -1.f, 1.f) * float(MAX_int32); }
 template<> uint16 ConvertInputFormat<float, uint16, EDefaultInputType::UNorm>(const float& In) { return FMath::Clamp(In,  0.f, 1.f) * MAX_uint16; }
 template<>  int16 ConvertInputFormat<float,  int16, EDefaultInputType::SNorm>(const float& In) { return FMath::Clamp(In, -1.f, 1.f) * MAX_int16; }
 template<>  uint8 ConvertInputFormat<float,  uint8, EDefaultInputType::UNorm>(const float& In) { return FMath::Clamp(In,  0.f, 1.f) * MAX_uint8; }
