Convert a couple of template shader classes to use FPermutationDomain instead.

diff --git a/Engine/Source/Runtime/Renderer/Private/DistanceFieldObjectCulling.cpp b/Engine/Source/Runtime/Renderer/Private/DistanceFieldObjectCulling.cpp
--- a/Engine/Source/Runtime/Renderer/Private/DistanceFieldObjectCulling.cpp
+++ b/Engine/Source/Runtime/Renderer/Private/DistanceFieldObjectCulling.cpp
@@ -35,51 +35,6 @@ FAutoConsoleVariableRef CVarMaxDistanceFieldObjectsPerCullTile(
 	ECVF_RenderThreadSafe | ECVF_ReadOnly
 	);
 
-class FCircleVertexBuffer : public FVertexBuffer
-{
-public:
-
-	int32 NumSections;
-
-	FCircleVertexBuffer()
-	{
-		NumSections = 8;
-	}
-
-	virtual void InitRHI() override
-	{
-		// Used as a non-indexed triangle list, so 3 vertices per triangle
-		const uint32 Size = 3 * NumSections * sizeof(FScreenVertex);
-		FRHIResourceCreateInfo CreateInfo(TEXT("FCircleVertexBuffer"));
-		VertexBufferRHI = RHICreateBuffer(Size, BUF_Static | BUF_VertexBuffer, 0, ERHIAccess::VertexOrIndexBuffer, CreateInfo);
-		FScreenVertex* DestVertex = (FScreenVertex*)RHILockBuffer(VertexBufferRHI, 0, Size, RLM_WriteOnly);
-
-		const float RadiansPerRingSegment = PI / (float)NumSections;
-
-		// Boost the effective radius so that the edges of the circle approximation lie on the circle, instead of the vertices
-		const float Radius = 1.0f / FMath::Cos(RadiansPerRingSegment);
-
-		for (int32 SectionIndex = 0; SectionIndex < NumSections; SectionIndex++)
-		{
-			float Fraction = SectionIndex / (float)NumSections;
-			float CurrentAngle = Fraction * 2 * PI;
-			float NextAngle = ((SectionIndex + 1) / (float)NumSections) * 2 * PI;
-			FVector2D CurrentPosition(Radius * FMath::Cos(CurrentAngle), Radius * FMath::Sin(CurrentAngle));
-			FVector2D NextPosition(Radius * FMath::Cos(NextAngle), Radius * FMath::Sin(NextAngle));
-
-			DestVertex[SectionIndex * 3 + 0].Position = FVector2D(0, 0);
-			DestVertex[SectionIndex * 3 + 0].UV = CurrentPosition;
-			DestVertex[SectionIndex * 3 + 1].Position = FVector2D(0, 0);
-			DestVertex[SectionIndex * 3 + 1].UV = NextPosition;
-			DestVertex[SectionIndex * 3 + 2].Position = FVector2D(0, 0);
-			DestVertex[SectionIndex * 3 + 2].UV = FVector2D(.5f, .5f);
-		}
-
-		RHIUnlockBuffer(VertexBufferRHI);      
-	}
-};
-
-TGlobalResource<FCircleVertexBuffer> GCircleVertexBuffer;
 TGlobalResource<FDistanceFieldObjectBufferResource> GAOCulledObjectBuffers;
 
 void FTileIntersectionResources::InitDynamicRHI()
@@ -353,12 +308,14 @@ private:
 
 IMPLEMENT_SHADER_TYPE(,FObjectCullVS,TEXT("/Engine/Private/DistanceFieldObjectCulling.usf"),TEXT("ObjectCullVS"),SF_Vertex);
 
-template <bool bCountingPass>
-class TObjectCullPS : public FGlobalShader
+class FObjectCullPS : public FGlobalShader
 {
-	DECLARE_SHADER_TYPE(TObjectCullPS, Global);
+	DECLARE_GLOBAL_SHADER(FObjectCullPS);
 public:
 
+	class FCountingPass : SHADER_PERMUTATION_BOOL("SCATTER_CULLING_COUNT_PASS");
+	using FPermutationDomain = TShaderPermutationDomain<FCountingPass>;
+
 	static bool ShouldCompilePermutation(const FGlobalShaderPermutationParameters& Parameters)
 	{
 		return IsFeatureLevelSupported(Parameters.Platform, ERHIFeatureLevel::SM5) && DoesPlatformSupportDistanceFieldAO(Parameters.Platform) && IsUsingDistanceFields(Parameters.Platform);
@@ -370,14 +327,13 @@ public:
 		FTileIntersectionParameters::ModifyCompilationEnvironment(Parameters.Platform, OutEnvironment);
 
 		OutEnvironment.SetDefine(TEXT("DOWNSAMPLE_FACTOR"), GAODownsampleFactor);
-		OutEnvironment.SetDefine(TEXT("SCATTER_CULLING_COUNT_PASS"), bCountingPass ? 1 : 0);
 	}
 
 	/** Default constructor. */
-	TObjectCullPS() {}
+	FObjectCullPS() {}
 
 	/** Initialization constructor. */
-	TObjectCullPS(const ShaderMetaType::CompiledShaderInitializerType& Initializer)
+	FObjectCullPS(const ShaderMetaType::CompiledShaderInitializerType& Initializer)
 		: FGlobalShader(Initializer)
 	{
 		ObjectParameters.Bind(Initializer.ParameterMap);
@@ -421,8 +377,7 @@ private:
 	LAYOUT_FIELD(FShaderParameter, NumGroups);
 };
 
-IMPLEMENT_SHADER_TYPE(template<>,TObjectCullPS<true>,TEXT("/Engine/Private/DistanceFieldObjectCulling.usf"),TEXT("ObjectCullPS"),SF_Pixel);
-IMPLEMENT_SHADER_TYPE(template<>,TObjectCullPS<false>,TEXT("/Engine/Private/DistanceFieldObjectCulling.usf"),TEXT("ObjectCullPS"),SF_Pixel);
+IMPLEMENT_GLOBAL_SHADER(FObjectCullPS, "/Engine/Private/DistanceFieldObjectCulling.usf", "ObjectCullPS", SF_Pixel);
 
 const uint32 ComputeStartOffsetGroupSize = 64;
 
@@ -498,11 +453,13 @@ private:
 
 IMPLEMENT_SHADER_TYPE(,FComputeCulledTilesStartOffsetCS,TEXT("/Engine/Private/DistanceFieldObjectCulling.usf"),TEXT("ComputeCulledTilesStartOffsetCS"),SF_Compute);
 
-template<bool bCountingPass>
-void ScatterTilesToObjects(FRHICommandListImmediate& RHICmdList, const FViewInfo& View, const FDistanceFieldSceneData& DistanceFieldSceneData, FIntPoint TileListGroupSize, const FDistanceFieldAOParameters& Parameters)
+void ScatterTilesToObjects(FRHICommandListImmediate& RHICmdList, bool bCountingPass, const FViewInfo& View, const FDistanceFieldSceneData& DistanceFieldSceneData, FIntPoint TileListGroupSize, const FDistanceFieldAOParameters& Parameters)
 {
-	TShaderMapRef<FObjectCullVS> VertexShader(View.ShaderMap);
-	TShaderMapRef<TObjectCullPS<bCountingPass>> PixelShader(View.ShaderMap);
+	FObjectCullPS::FPermutationDomain PermutationVector;
+	PermutationVector.Set<FObjectCullPS::FCountingPass>(bCountingPass);
+
+	auto VertexShader = View.ShaderMap->GetShader<FObjectCullVS>();
+	auto PixelShader = View.ShaderMap->GetShader<FObjectCullPS>(PermutationVector);
 
 	TArray<FRHIUnorderedAccessView*> UAVs;
 	PixelShader->GetUAVs(View, UAVs);
@@ -640,7 +597,7 @@ void BuildTileObjectLists(FRDGBuilder& GraphBuilder, FScene* Scene, TArray<FView
 					RHICmdList.ClearUAVUint(TileIntersectionResources->NumCulledTilesArray.UAV, FUintVector4(0, 0, 0, 0));
 
 					// Rasterize object bounding shapes and intersect with screen tiles to compute how many tiles intersect each object
-					ScatterTilesToObjects<true>(RHICmdList, View, Scene->DistanceFieldSceneData, TileListGroupSize, Parameters);
+					ScatterTilesToObjects(RHICmdList, true, View, Scene->DistanceFieldSceneData, TileListGroupSize, Parameters);
 				}
 
 				{
@@ -669,7 +626,7 @@ void BuildTileObjectLists(FRDGBuilder& GraphBuilder, FScene* Scene, TArray<FView
 					RHICmdList.ClearUAVUint(TileIntersectionResources->NumCulledTilesArray.UAV, FUintVector4(0, 0, 0, 0));
 
 					// Rasterize object bounding shapes and intersect with screen tiles, and write out intersecting tile indices for the cone tracing pass
-					ScatterTilesToObjects<false>(RHICmdList, View, Scene->DistanceFieldSceneData, TileListGroupSize, Parameters);
+					ScatterTilesToObjects(RHICmdList, false, View, Scene->DistanceFieldSceneData, TileListGroupSize, Parameters);
 				}
 			}
 			else
diff --git a/Engine/Source/Runtime/Renderer/Private/DistanceFieldScreenGridLighting.cpp b/Engine/Source/Runtime/Renderer/Private/DistanceFieldScreenGridLighting.cpp
--- a/Engine/Source/Runtime/Renderer/Private/DistanceFieldScreenGridLighting.cpp
+++ b/Engine/Source/Runtime/Renderer/Private/DistanceFieldScreenGridLighting.cpp
@@ -70,12 +70,14 @@ void FAOScreenGridResources::InitDynamicRHI()
 	ScreenGridConeVisibility.Initialize(TEXT("ScreenGridConeVisibility"), sizeof(uint32), NumConeSampleDirections * ScreenGridDimensions.X * ScreenGridDimensions.Y, PF_R32_UINT, BUF_Static | FastVRamFlag);
 }
 
-template<bool bUseGlobalDistanceField>
-class TConeTraceScreenGridObjectOcclusionCS : public FGlobalShader
+class FConeTraceScreenGridObjectOcclusionCS : public FGlobalShader
 {
-	DECLARE_SHADER_TYPE(TConeTraceScreenGridObjectOcclusionCS,Global)
+	DECLARE_GLOBAL_SHADER(FConeTraceScreenGridObjectOcclusionCS);
 public:
 
+	class FUseGlobalDistanceField : SHADER_PERMUTATION_BOOL("USE_GLOBAL_DISTANCE_FIELD");
+	using FPermutationDomain = TShaderPermutationDomain<FUseGlobalDistanceField>;
+
 	static bool ShouldCompilePermutation(const FGlobalShaderPermutationParameters& Parameters)
 	{
 		return IsFeatureLevelSupported(Parameters.Platform, ERHIFeatureLevel::SM5) && DoesPlatformSupportDistanceFieldAO(Parameters.Platform) && IsUsingDistanceFields(Parameters.Platform);
@@ -85,13 +87,12 @@ public:
 	{
 		FGlobalShader::ModifyCompilationEnvironment(Parameters, OutEnvironment);
 		FTileIntersectionParameters::ModifyCompilationEnvironment(Parameters.Platform, OutEnvironment);
-		OutEnvironment.SetDefine(TEXT("USE_GLOBAL_DISTANCE_FIELD"), bUseGlobalDistanceField);
 
 		// To reduce shader compile time of compute shaders with shared memory, doesn't have an impact on generated code with current compiler (June 2010 DX SDK)
 		OutEnvironment.CompilerFlags.Add(CFLAG_StandardOptimization);
 	}
 
-	TConeTraceScreenGridObjectOcclusionCS(const ShaderMetaType::CompiledShaderInitializerType& Initializer)
+	FConeTraceScreenGridObjectOcclusionCS(const ShaderMetaType::CompiledShaderInitializerType& Initializer)
 		: FGlobalShader(Initializer)
 	{
 		ObjectParameters.Bind(Initializer.ParameterMap);
@@ -106,7 +107,7 @@ public:
 		ConeDepthVisibilityFunction.Bind(Initializer.ParameterMap, TEXT("ConeDepthVisibilityFunction"));
 	}
 
-	TConeTraceScreenGridObjectOcclusionCS()
+	FConeTraceScreenGridObjectOcclusionCS()
 	{
 	}
 
@@ -116,6 +117,7 @@ public:
 		const FDistanceFieldSceneData& DistanceFieldSceneData,
 		FRHITexture* DistanceFieldNormal,  
 		const FDistanceFieldAOParameters& Parameters,
+		bool bUseGlobalDistanceField,
 		const FGlobalDistanceFieldInfo& GlobalDistanceFieldInfo)
 	{
 		FRHIComputeShader* ShaderRHI = RHICmdList.GetBoundComputeShader();
@@ -191,12 +193,7 @@ private:
 	LAYOUT_FIELD(FRWShaderParameter, ConeDepthVisibilityFunction);
 };
 
-#define IMPLEMENT_CONETRACE_CS_TYPE(bUseGlobalDistanceField) \
-	typedef TConeTraceScreenGridObjectOcclusionCS<bUseGlobalDistanceField> TConeTraceScreenGridObjectOcclusionCS##bUseGlobalDistanceField; \
-	IMPLEMENT_SHADER_TYPE(template<>,TConeTraceScreenGridObjectOcclusionCS##bUseGlobalDistanceField,TEXT("/Engine/Private/DistanceFieldScreenGridLighting.usf"),TEXT("ConeTraceObjectOcclusionCS"),SF_Compute);
-
-IMPLEMENT_CONETRACE_CS_TYPE(true)
-IMPLEMENT_CONETRACE_CS_TYPE(false)
+IMPLEMENT_GLOBAL_SHADER(FConeTraceScreenGridObjectOcclusionCS, "/Engine/Private/DistanceFieldScreenGridLighting.usf", "ConeTraceObjectOcclusionCS", SF_Compute);
 
 const int32 GConeTraceGlobalDFTileSize = 8;
 
@@ -562,30 +559,17 @@ void FDeferredShadingSceneRenderer::RenderDistanceFieldAOScreenGrid(
 
 				RHICmdList.Transition(FRHITransitionInfo(ScreenGridResources->ScreenGridConeVisibility.UAV, ERHIAccess::Unknown, ERHIAccess::UAVCompute));
 
-				if (bUseGlobalDistanceField)
-				{
-					check(View.GlobalDistanceFieldInfo.Clipmaps.Num() > 0);
+				FConeTraceScreenGridObjectOcclusionCS::FPermutationDomain PermutationVector;
+				PermutationVector.Set<FConeTraceScreenGridObjectOcclusionCS::FUseGlobalDistanceField>(bUseGlobalDistanceField);
 
-					{
-						TShaderMapRef<TConeTraceScreenGridObjectOcclusionCS<true> > ComputeShader(View.ShaderMap);
+				auto ComputeShader = View.ShaderMap->GetShader<FConeTraceScreenGridObjectOcclusionCS>(PermutationVector);
 
-						RHICmdList.SetComputeShader(ComputeShader.GetComputeShader());
-						ComputeShader->SetParameters(RHICmdList, View, DistanceFieldSceneData, DistanceFieldNormal->GetRHI(), Parameters, View.GlobalDistanceFieldInfo);
-						DispatchIndirectComputeShader(RHICmdList, ComputeShader.GetShader(), TileIntersectionResources->ObjectTilesIndirectArguments.Buffer, 0);
-						ComputeShader->UnsetParameters(RHICmdList, View);
-					}
-				}
-				else
-				{
-					{
-						TShaderMapRef<TConeTraceScreenGridObjectOcclusionCS<false> > ComputeShader(View.ShaderMap);
-
-						RHICmdList.SetComputeShader(ComputeShader.GetComputeShader());
-						ComputeShader->SetParameters(RHICmdList, View, DistanceFieldSceneData, DistanceFieldNormal->GetRHI(), Parameters, View.GlobalDistanceFieldInfo);
-						DispatchIndirectComputeShader(RHICmdList, ComputeShader.GetShader(), TileIntersectionResources->ObjectTilesIndirectArguments.Buffer, 0);
-						ComputeShader->UnsetParameters(RHICmdList, View);
-					}
-				}
+				check(!bUseGlobalDistanceField || View.GlobalDistanceFieldInfo.Clipmaps.Num() > 0);
+
+				RHICmdList.SetComputeShader(ComputeShader.GetComputeShader());
+				ComputeShader->SetParameters(RHICmdList, View, DistanceFieldSceneData, DistanceFieldNormal->GetRHI(), Parameters, bUseGlobalDistanceField, View.GlobalDistanceFieldInfo);
+				DispatchIndirectComputeShader(RHICmdList, ComputeShader.GetShader(), TileIntersectionResources->ObjectTilesIndirectArguments.Buffer, 0);
+				ComputeShader->UnsetParameters(RHICmdList, View);
 			}
 
 			// Compute heightfield occlusion after heightfield GI, otherwise it self-shadows incorrectly
