[INSIGHTS] - add insights zones to illuminate overall frame time.

diff --git a/Engine/Source/Runtime/Core/Private/Async/TaskGraph.cpp b/Engine/Source/Runtime/Core/Private/Async/TaskGraph.cpp
--- a/Engine/Source/Runtime/Core/Private/Async/TaskGraph.cpp
+++ b/Engine/Source/Runtime/Core/Private/Async/TaskGraph.cpp
@@ -705,6 +705,7 @@ PRAGMA_ENABLE_DEPRECATION_WARNINGS
 #endif
 				if (bAllowStall)
 				{
+					TRACE_CPUPROFILER_EVENT_SCOPE(WaitForTasks);
 					{
 						FScopeCycleCounter Scope(StallStatId, EStatFlags::Verbose);
 						Queue(QueueIndex).StallRestartEvent->Wait(bIsRenderThreadAndPolling ? GRenderThreadPollPeriodMs : MAX_uint32, bCountAsStall);
diff --git a/Engine/Source/Runtime/CoreUObject/Private/UObject/LinkerLoad.cpp b/Engine/Source/Runtime/CoreUObject/Private/UObject/LinkerLoad.cpp
--- a/Engine/Source/Runtime/CoreUObject/Private/UObject/LinkerLoad.cpp
+++ b/Engine/Source/Runtime/CoreUObject/Private/UObject/LinkerLoad.cpp
@@ -4155,7 +4155,7 @@ void FLinkerLoad::Preload( UObject* Object )
 
 		if (Object->GetLinker() == this)
 		{
-			TRACE_CPUPROFILER_EVENT_SCOPE(FLinkerLoad::Preload);
+			TRACE_CPUPROFILER_EVENT_SCOPE_TEXT(*FString::Printf(TEXT("FLinkerLoad::Preload %s"), *Object->GetFName().ToString()));
 
 			UClass* Cls = Cast<UClass>(Object);
 			check(!GEventDrivenLoaderEnabled || !bLockoutLegacyOperations || !EVENT_DRIVEN_ASYNC_LOAD_ACTIVE_AT_RUNTIME);
diff --git a/Engine/Source/Runtime/Engine/Private/Components/SkeletalMeshComponent.cpp b/Engine/Source/Runtime/Engine/Private/Components/SkeletalMeshComponent.cpp
--- a/Engine/Source/Runtime/Engine/Private/Components/SkeletalMeshComponent.cpp
+++ b/Engine/Source/Runtime/Engine/Private/Components/SkeletalMeshComponent.cpp
@@ -3922,6 +3922,7 @@ void USkeletalMeshComponent::CompleteParallelAnimationEvaluation(bool bDoPostAni
 
 bool USkeletalMeshComponent::HandleExistingParallelEvaluationTask(bool bBlockOnTask, bool bPerformPostAnimEvaluation)
 {
+	TRACE_CPUPROFILER_EVENT_SCOPE(USkeletalMeshComponent);
 	if (IsRunningParallelEvaluation()) // We are already processing eval on another thread
 	{
 		if (bBlockOnTask)
diff --git a/Engine/Source/Runtime/Engine/Private/LevelTick.cpp b/Engine/Source/Runtime/Engine/Private/LevelTick.cpp
--- a/Engine/Source/Runtime/Engine/Private/LevelTick.cpp
+++ b/Engine/Source/Runtime/Engine/Private/LevelTick.cpp
@@ -1123,6 +1123,7 @@ void UWorld::SendAllEndOfFrameUpdates()
 #endif
 		](int32 Index)
 		{
+			TRACE_CPUPROFILER_EVENT_SCOPE(DeferredRenderUpdates);
 			FOptionalTaskTagScope Scope(ETaskTag::EParallelGameThread);
 			UActorComponent* NextComponent = LocalComponentsThatNeedEndOfFrameUpdate[Index];
 			if (NextComponent)
@@ -1671,12 +1672,14 @@ void UWorld::Tick( ELevelTick TickType, float DeltaSeconds )
 	// Tick LevelInstanceSubsystem outside of FTickTaskManagerInterface::StartFrame/EndFrame because it can cause levels to be deleted and invalidate its LevelList
 	if (ULevelInstanceSubsystem* LevelInstanceSubsystem = GetSubsystem<ULevelInstanceSubsystem>())
 	{
+		TRACE_CPUPROFILER_EVENT_SCOPE(LevelInstanceSubsystem);
 		LevelInstanceSubsystem->Tick();
 	}
 #endif
 
 	if (bDoingActorTicks)
 	{
+		TRACE_CPUPROFILER_EVENT_SCOPE(bDoingActorTicks);
 		SCOPE_CYCLE_COUNTER(STAT_TickTime);
 
 		FWorldDelegates::OnWorldPostActorTick.Broadcast(this, TickType, DeltaSeconds);
@@ -1684,12 +1687,14 @@ void UWorld::Tick( ELevelTick TickType, float DeltaSeconds )
 #if PHYSICS_INTERFACE_PHYSX
 		if ( PhysicsScene != nullptr )
 		{
+			TRACE_CPUPROFILER_EVENT_SCOPE(GPhysCommandHandler);
 			GPhysCommandHandler->Flush();
 		}
 #endif // WITH_PHYSX
 		
 		// All tick is done, execute async trace
 		{
+			TRACE_CPUPROFILER_EVENT_SCOPE(FinishAsyncTrace);
 			SCOPE_CYCLE_COUNTER(STAT_FinishAsyncTraceTickTime);
 			SCOPE_TIME_GUARD_MS(TEXT("UWorld::Tick - FinishAsyncTrace"), 5);
 			FinishAsyncTrace();
@@ -1703,6 +1708,7 @@ void UWorld::Tick( ELevelTick TickType, float DeltaSeconds )
 	// Update net and flush networking.
     // Tick all net drivers
 	{
+		TRACE_CPUPROFILER_EVENT_SCOPE(NetBroadcastTickTime);
 		SCOPE_CYCLE_COUNTER(STAT_NetBroadcastTickTime);
 		LLM_SCOPE(ELLMTag::Networking);
 		BroadcastTickFlush(RealDeltaSeconds); // note: undilated time is being used here
@@ -1710,12 +1716,14 @@ void UWorld::Tick( ELevelTick TickType, float DeltaSeconds )
 	
      // PostTick all net drivers
 	{
+		TRACE_CPUPROFILER_EVENT_SCOPE(BroadcastPostTickFlush);
 		LLM_SCOPE(ELLMTag::Networking);
 		BroadcastPostTickFlush(RealDeltaSeconds); // note: undilated time is being used here
 	}
 
 	if( Scene )
 	{
+		TRACE_CPUPROFILER_EVENT_SCOPE(UpdateSpeedTreeWind);
 		// Update SpeedTree wind objects.
 		Scene->UpdateSpeedTreeWind(TimeSeconds);
 	}
@@ -1723,6 +1731,7 @@ void UWorld::Tick( ELevelTick TickType, float DeltaSeconds )
 	// Tick the FX system.
 	if (!bIsPaused && FXSystem != nullptr)
 	{
+		TRACE_CPUPROFILER_EVENT_SCOPE(FXSystem);
 		SCOPE_TIME_GUARD_MS(TEXT("UWorld::Tick - FX"), 5);
 		FXSystem->Tick(DeltaSeconds);
 	}
@@ -1740,7 +1749,11 @@ void UWorld::Tick( ELevelTick TickType, float DeltaSeconds )
 	bInTick = false;
 	Mark.Pop();
 
-	GEngine->ConditionalCollectGarbage();
+	{
+		TRACE_CPUPROFILER_EVENT_SCOPE(ConditionalCollectGarbage);
+		GEngine->ConditionalCollectGarbage();
+	}
+	
 
 	// players only request from last frame
 	if (bPlayersOnlyPending)
diff --git a/Engine/Source/Runtime/Engine/Private/TickTaskManager.cpp b/Engine/Source/Runtime/Engine/Private/TickTaskManager.cpp
--- a/Engine/Source/Runtime/Engine/Private/TickTaskManager.cpp
+++ b/Engine/Source/Runtime/Engine/Private/TickTaskManager.cpp
@@ -560,6 +560,7 @@ public:
 				CA_SUPPRESS(6385);
 				if (TickCompletionEvents[Block].Num())
 				{
+					TRACE_CPUPROFILER_EVENT_SCOPE(TickCompletionEvents);
 					FTaskGraphInterface::Get().WaitUntilTasksComplete(TickCompletionEvents[Block], ENamedThreads::GameThread);
 					if (SingleThreadedMode() || Block == TG_NewlySpawned || CVarAllowAsyncTickCleanup.GetValueOnGameThread() == 0 || TickCompletionEvents[Block].Num() < 50)
 					{
diff --git a/Engine/Source/Runtime/Renderer/Private/PrimitiveSceneInfo.cpp b/Engine/Source/Runtime/Renderer/Private/PrimitiveSceneInfo.cpp
--- a/Engine/Source/Runtime/Renderer/Private/PrimitiveSceneInfo.cpp
+++ b/Engine/Source/Runtime/Renderer/Private/PrimitiveSceneInfo.cpp
@@ -838,6 +838,7 @@ void FPrimitiveSceneInfo::RemoveCachedRayTracingPrimitives()
 
 void FPrimitiveSceneInfo::UpdateCachedRayTracingInstanceTransforms(FMatrix NewPrimitiveLocalToWorld)
 {
+	TRACE_CPUPROFILER_EVENT_SCOPE(UpdateCachedRayTracingInstanceTransforms);
 	for (int32 Index = 0; Index < CachedRayTracingInstanceLocalTransforms.Num(); Index++)
 	{
 		CachedRayTracingInstanceWorldTransforms[Index] = CachedRayTracingInstanceLocalTransforms[Index] * NewPrimitiveLocalToWorld;
@@ -1380,6 +1381,7 @@ bool FPrimitiveSceneInfo::NeedsUpdateStaticMeshes()
 void FPrimitiveSceneInfo::UpdateStaticMeshes(FRHICommandListImmediate& RHICmdList, FScene* Scene, const TArrayView<FPrimitiveSceneInfo*>& SceneInfos, bool bReAddToDrawLists)
 {
 	QUICK_SCOPE_CYCLE_COUNTER(STAT_FPrimitiveSceneInfo_UpdateStaticMeshes);
+	TRACE_CPUPROFILER_EVENT_SCOPE(FPrimitiveSceneInfo_UpdateStaticMeshes);
 
 	const bool bNeedsStaticMeshUpdate = !bReAddToDrawLists;
 
diff --git a/Engine/Source/Runtime/Renderer/Private/SceneRendering.cpp b/Engine/Source/Runtime/Renderer/Private/SceneRendering.cpp
--- a/Engine/Source/Runtime/Renderer/Private/SceneRendering.cpp
+++ b/Engine/Source/Runtime/Renderer/Private/SceneRendering.cpp
@@ -3923,6 +3923,7 @@ extern CORE_API bool GRenderThreadPollingOn;
 
 void FRendererModule::BeginRenderingViewFamily(FCanvas* Canvas, FSceneViewFamily* ViewFamily)
 {
+	TRACE_CPUPROFILER_EVENT_SCOPE(BeginRenderingViewFamily);
 	check(Canvas);
 	check(ViewFamily);
 	check(ViewFamily->Scene);
diff --git a/Engine/Source/Runtime/Renderer/Private/SceneVisibility.cpp b/Engine/Source/Runtime/Renderer/Private/SceneVisibility.cpp
--- a/Engine/Source/Runtime/Renderer/Private/SceneVisibility.cpp
+++ b/Engine/Source/Runtime/Renderer/Private/SceneVisibility.cpp
@@ -820,6 +820,7 @@ struct FVisForPrimParams
 template<bool bSingleThreaded>
 static void FetchVisibilityForPrimitives_Range(FVisForPrimParams& Params, FGlobalDynamicVertexBuffer* DynamicVertexBufferIfSingleThreaded)
 {	
+	TRACE_CPUPROFILER_EVENT_SCOPE(FetchVisibilityForPrimitives_Range);
 	int32 NumOccludedPrimitives = 0;
 	
 	const FScene* Scene				= Params.Scene;
@@ -891,414 +892,417 @@ static void FetchVisibilityForPrimitives_Range(FVisForPrimParams& Params, FGloba
 	int32 NumTotalPrims = View.PrimitiveVisibilityMap.Num();
 	int32 NumTotalDefUnoccluded = View.PrimitiveDefinitelyUnoccludedMap.Num();
 
-	//if we are load balanced then we iterate only the set bits, and the ranges have been pre-selected to evenly distribute set bits among the tasks with no overlaps.
-	//if not, then the entire array is evenly divided by range.
+	{
+		TRACE_CPUPROFILER_EVENT_SCOPE_TEXT(*FString::Printf(TEXT("forEach over %d entries"), NumToProcess));
+	
+		//if we are load balanced then we iterate only the set bits, and the ranges have been pre-selected to evenly distribute set bits among the tasks with no overlaps.
+		//if not, then the entire array is evenly divided by range.
 #if BALANCE_LOAD
-	for (FSceneSetBitIterator BitIt(View.PrimitiveVisibilityMap, StartIndex); BitIt && (NumProcessed < NumToProcess); ++BitIt, ++NumProcessed)
+		for (FSceneSetBitIterator BitIt(View.PrimitiveVisibilityMap, StartIndex); BitIt && (NumProcessed < NumToProcess); ++BitIt, ++NumProcessed)
 #else
-	for (TBitArray<SceneRenderingBitArrayAllocator>::FIterator BitIt(View.PrimitiveVisibilityMap, StartIndex); BitIt && (NumProcessed < NumToProcess); ++BitIt, ++NumProcessed)
+		for (TBitArray<SceneRenderingBitArrayAllocator>::FIterator BitIt(View.PrimitiveVisibilityMap, StartIndex); BitIt && (NumProcessed < NumToProcess); ++BitIt, ++NumProcessed)
 #endif
-	{
-		const bool bAlwaysVisible = PrimitivesAlwaysVisible[BitIt.GetIndex()];
-		uint8 OcclusionFlags = Scene->PrimitiveOcclusionFlags[BitIt.GetIndex()];
-		bool bCanBeOccluded = !bAlwaysVisible && (OcclusionFlags & EOcclusionFlags::CanBeOccluded) != 0;
+		{
+			const bool bAlwaysVisible = PrimitivesAlwaysVisible[BitIt.GetIndex()];
+			uint8 OcclusionFlags = Scene->PrimitiveOcclusionFlags[BitIt.GetIndex()];
+			bool bCanBeOccluded = !bAlwaysVisible && (OcclusionFlags & EOcclusionFlags::CanBeOccluded) != 0;
 
 #if !BALANCE_LOAD		
-		if (!View.PrimitiveVisibilityMap.AccessCorrespondingBit(BitIt))
-		{
-			continue;
-		}
+			if (!View.PrimitiveVisibilityMap.AccessCorrespondingBit(BitIt))
+			{
+				continue;
+			}
 #endif
 
-		if (!bCanBeOccluded)
-		{
-			View.PrimitiveDefinitelyUnoccludedMap.AccessCorrespondingBit(BitIt) = true;
-			continue;
-		}
-
-		//we can't allow the prim history insertion array to realloc or it will invalidate pointers in the other output arrays.
-		const bool bCanAllocPrimHistory = bSingleThreaded || InsertPrimitiveOcclusionHistory->Num() < InsertPrimitiveOcclusionHistory->Max();		
-
-	#if WITH_EDITOR
-		if (GIsEditor)
-		{
-			if (Scene->PrimitivesSelected[BitIt.GetIndex()])
+			if (!bCanBeOccluded)
 			{
-				// to render occluded outline for selected objects
-				bCanBeOccluded = false;
+				View.PrimitiveDefinitelyUnoccludedMap.AccessCorrespondingBit(BitIt) = true;
+				continue;
 			}
-		}
-	#endif
 
-		int32 NumSubQueries = 1;
-		bool bSubQueries = false;
-		const TArray<FBoxSphereBounds>* SubBounds = nullptr;
+			//we can't allow the prim history insertion array to realloc or it will invalidate pointers in the other output arrays.
+			const bool bCanAllocPrimHistory = bSingleThreaded || InsertPrimitiveOcclusionHistory->Num() < InsertPrimitiveOcclusionHistory->Max();
 
-		check(Params.SubIsOccluded);
-		TArray<bool>& SubIsOccluded = *Params.SubIsOccluded;
-		int32 SubIsOccludedStart = SubIsOccluded.Num();
-		if ((OcclusionFlags & EOcclusionFlags::HasSubprimitiveQueries) && GAllowSubPrimitiveQueries && !View.bDisableQuerySubmissions && !bAlwaysVisible)
-		{
-			FPrimitiveSceneProxy* Proxy = Scene->Primitives[BitIt.GetIndex()]->Proxy;
-			SubBounds = Proxy->GetOcclusionQueries(&View);
-			NumSubQueries = SubBounds->Num();
-			bSubQueries = true;
-			if (!NumSubQueries)
+#if WITH_EDITOR
+			if (GIsEditor)
 			{
-				View.PrimitiveVisibilityMap.AccessCorrespondingBit(BitIt) = false;
-				continue;
+				if (Scene->PrimitivesSelected[BitIt.GetIndex()])
+				{
+					// to render occluded outline for selected objects
+					bCanBeOccluded = false;
+				}
 			}
-			SubIsOccluded.Reserve(NumSubQueries);
-		}
+#endif
 
-		bool bAllSubOcclusionStateIsDefinite = true;
-		bool bAllSubOccluded = true;
-		FPrimitiveComponentId PrimitiveId = Scene->PrimitiveComponentIds[BitIt.GetIndex()];
+			int32 NumSubQueries = 1;
+			bool bSubQueries = false;
+			const TArray<FBoxSphereBounds>* SubBounds = nullptr;
 
-		for (int32 SubQuery = 0; SubQuery < NumSubQueries; SubQuery++)
-		{
-			FPrimitiveOcclusionHistory* PrimitiveOcclusionHistory = ViewPrimitiveOcclusionHistory.Find(FPrimitiveOcclusionHistoryKey(PrimitiveId, SubQuery));
-			
-			bool bIsOccluded = false;
-			bool bOcclusionStateIsDefinite = false;			
-
-			if (!PrimitiveOcclusionHistory)
-			{
-				// If the primitive doesn't have an occlusion history yet, create it.
-				if (bSingleThreaded)
-				{					
-					// In singlethreaded mode we can safely modify the view's history directly.
-					PrimitiveOcclusionHistory = &ViewPrimitiveOcclusionHistory[
-						ViewPrimitiveOcclusionHistory.Add(FPrimitiveOcclusionHistory(PrimitiveId, SubQuery))
-					];
-				}
-				else if (bCanAllocPrimHistory)
+			check(Params.SubIsOccluded);
+			TArray<bool>& SubIsOccluded = *Params.SubIsOccluded;
+			int32 SubIsOccludedStart = SubIsOccluded.Num();
+			if ((OcclusionFlags & EOcclusionFlags::HasSubprimitiveQueries) && GAllowSubPrimitiveQueries && !View.bDisableQuerySubmissions && !bAlwaysVisible)
+			{
+				FPrimitiveSceneProxy* Proxy = Scene->Primitives[BitIt.GetIndex()]->Proxy;
+				SubBounds = Proxy->GetOcclusionQueries(&View);
+				NumSubQueries = SubBounds->Num();
+				bSubQueries = true;
+				if (!NumSubQueries)
 				{
-					// In multithreaded mode we have to buffer the new histories and add them to the view during a post-combine
-					PrimitiveOcclusionHistory = &(*InsertPrimitiveOcclusionHistory)[
-						InsertPrimitiveOcclusionHistory->Add(FPrimitiveOcclusionHistory(PrimitiveId, SubQuery))
-					];
-				}				
-				
-				// If the primitive hasn't been visible recently enough to have a history, treat it as unoccluded this frame so it will be rendered as an occluder and its true occlusion state can be determined.
-				// already set bIsOccluded = false;
-
-				// Flag the primitive's occlusion state as indefinite, which will force it to be queried this frame.
-				// The exception is if the primitive isn't occludable, in which case we know that it's definitely unoccluded.
-				bOcclusionStateIsDefinite = bCanBeOccluded ? false : true;
+					View.PrimitiveVisibilityMap.AccessCorrespondingBit(BitIt) = false;
+					continue;
+				}
+				SubIsOccluded.Reserve(NumSubQueries);
 			}
-			else
+
+			bool bAllSubOcclusionStateIsDefinite = true;
+			bool bAllSubOccluded = true;
+			FPrimitiveComponentId PrimitiveId = Scene->PrimitiveComponentIds[BitIt.GetIndex()];
+
+			for (int32 SubQuery = 0; SubQuery < NumSubQueries; SubQuery++)
 			{
-				if (View.bIgnoreExistingQueries)
+				FPrimitiveOcclusionHistory* PrimitiveOcclusionHistory = ViewPrimitiveOcclusionHistory.Find(FPrimitiveOcclusionHistoryKey(PrimitiveId, SubQuery));
+
+				bool bIsOccluded = false;
+				bool bOcclusionStateIsDefinite = false;
+
+				if (!PrimitiveOcclusionHistory)
 				{
-					// If the view is ignoring occlusion queries, the primitive is definitely unoccluded.
+					// If the primitive doesn't have an occlusion history yet, create it.
+					if (bSingleThreaded)
+					{
+						// In singlethreaded mode we can safely modify the view's history directly.
+						PrimitiveOcclusionHistory = &ViewPrimitiveOcclusionHistory[
+							ViewPrimitiveOcclusionHistory.Add(FPrimitiveOcclusionHistory(PrimitiveId, SubQuery))
+						];
+					}
+					else if (bCanAllocPrimHistory)
+					{
+						// In multithreaded mode we have to buffer the new histories and add them to the view during a post-combine
+						PrimitiveOcclusionHistory = &(*InsertPrimitiveOcclusionHistory)[
+							InsertPrimitiveOcclusionHistory->Add(FPrimitiveOcclusionHistory(PrimitiveId, SubQuery))
+						];
+					}
+
+					// If the primitive hasn't been visible recently enough to have a history, treat it as unoccluded this frame so it will be rendered as an occluder and its true occlusion state can be determined.
 					// already set bIsOccluded = false;
-					bOcclusionStateIsDefinite = View.bDisableQuerySubmissions;
+
+					// Flag the primitive's occlusion state as indefinite, which will force it to be queried this frame.
+					// The exception is if the primitive isn't occludable, in which case we know that it's definitely unoccluded.
+					bOcclusionStateIsDefinite = bCanBeOccluded ? false : true;
 				}
-				else if (bCanBeOccluded)
+				else
 				{
-					if (bHZBOcclusion)
+					if (View.bIgnoreExistingQueries)
 					{
-						if (HZBOcclusionTests.IsValidFrame(PrimitiveOcclusionHistory->LastTestFrameNumber))
-						{
-							bIsOccluded = !HZBOcclusionTests.IsVisible(PrimitiveOcclusionHistory->HZBTestIndex);
-							bOcclusionStateIsDefinite = true;
-						}
+						// If the view is ignoring occlusion queries, the primitive is definitely unoccluded.
+						// already set bIsOccluded = false;
+						bOcclusionStateIsDefinite = View.bDisableQuerySubmissions;
 					}
-					else
+					else if (bCanBeOccluded)
 					{
-						// Read the occlusion query results.
-						uint64 NumSamples = 0;
-						bool bGrouped = false;
-						FRHIRenderQuery* PastQuery = PrimitiveOcclusionHistory->GetQueryForReading(OcclusionFrameCounter, NumBufferedFrames, ReadBackLagTolerance, bGrouped);
-						if (PastQuery)
+						if (bHZBOcclusion)
+						{
+							if (HZBOcclusionTests.IsValidFrame(PrimitiveOcclusionHistory->LastTestFrameNumber))
+							{
+								bIsOccluded = !HZBOcclusionTests.IsVisible(PrimitiveOcclusionHistory->HZBTestIndex);
+								bOcclusionStateIsDefinite = true;
+							}
+						}
+						else
 						{
-							//int32 RefCount = PastQuery.GetReference()->GetRefCount();
-							// NOTE: RHIGetOcclusionQueryResult should never fail when using a blocking call, rendering artifacts may show up.
-							//if (RHICmdList.GetRenderQueryResult(PastQuery, NumSamples, true))
-							if (GDynamicRHI->RHIGetRenderQueryResult(PastQuery, NumSamples, true))
+							// Read the occlusion query results.
+							uint64 NumSamples = 0;
+							bool bGrouped = false;
+							FRHIRenderQuery* PastQuery = PrimitiveOcclusionHistory->GetQueryForReading(OcclusionFrameCounter, NumBufferedFrames, ReadBackLagTolerance, bGrouped);
+							if (PastQuery)
 							{
-								// we render occlusion without MSAA
-								uint32 NumPixels = (uint32)NumSamples;
+								//int32 RefCount = PastQuery.GetReference()->GetRefCount();
+								// NOTE: RHIGetOcclusionQueryResult should never fail when using a blocking call, rendering artifacts may show up.
+								//if (RHICmdList.GetRenderQueryResult(PastQuery, NumSamples, true))
+								if (GDynamicRHI->RHIGetRenderQueryResult(PastQuery, NumSamples, true))
+								{
+									// we render occlusion without MSAA
+									uint32 NumPixels = (uint32)NumSamples;
 
-								// The primitive is occluded if none of its bounding box's pixels were visible in the previous frame's occlusion query.
-								bIsOccluded = (NumPixels == 0);
+									// The primitive is occluded if none of its bounding box's pixels were visible in the previous frame's occlusion query.
+									bIsOccluded = (NumPixels == 0);
 
-								
-								if (!bIsOccluded)
-								{
-									checkSlow(View.OneOverNumPossiblePixels > 0.0f);
-									PrimitiveOcclusionHistory->LastPixelsPercentage = NumPixels * View.OneOverNumPossiblePixels;
+
+									if (!bIsOccluded)
+									{
+										checkSlow(View.OneOverNumPossiblePixels > 0.0f);
+										PrimitiveOcclusionHistory->LastPixelsPercentage = NumPixels * View.OneOverNumPossiblePixels;
+									}
+									else
+									{
+										PrimitiveOcclusionHistory->LastPixelsPercentage = 0.0f;
+									}
+
+
+									// Flag the primitive's occlusion state as definite if it wasn't grouped.
+									bOcclusionStateIsDefinite = !bGrouped;
 								}
 								else
 								{
-									PrimitiveOcclusionHistory->LastPixelsPercentage = 0.0f;
-								}								
-
-
-								// Flag the primitive's occlusion state as definite if it wasn't grouped.
-								bOcclusionStateIsDefinite = !bGrouped;
+									// If the occlusion query failed, treat the primitive as visible.  
+									// already set bIsOccluded = false;
+								}
 							}
 							else
 							{
-								// If the occlusion query failed, treat the primitive as visible.  
-								// already set bIsOccluded = false;
+								if (NumBufferedFrames > 1 || GRHIMaximumReccommendedOustandingOcclusionQueries < MAX_int32)
+								{
+									// If there's no occlusion query for the primitive, assume it is whatever it was last frame
+									bIsOccluded = PrimitiveOcclusionHistory->WasOccludedLastFrame;
+									bOcclusionStateIsDefinite = PrimitiveOcclusionHistory->OcclusionStateWasDefiniteLastFrame;
+								}
+								else
+								{
+									// If there's no occlusion query for the primitive, set it's visibility state to whether it has been unoccluded recently.
+									bIsOccluded = (PrimitiveOcclusionHistory->LastProvenVisibleTime + GEngine->PrimitiveProbablyVisibleTime < CurrentRealTime);
+									// the state was definite last frame, otherwise we would have ran a query
+									bOcclusionStateIsDefinite = true;
+								}
+								if (bIsOccluded)
+								{
+									PrimitiveOcclusionHistory->LastPixelsPercentage = 0.0f;
+								}
+								else
+								{
+									PrimitiveOcclusionHistory->LastPixelsPercentage = GEngine->MaxOcclusionPixelsFraction;
+								}
 							}
 						}
-						else
+
+						if (GVisualizeOccludedPrimitives && OcclusionPDI && bIsOccluded)
 						{
-							if (NumBufferedFrames > 1 || GRHIMaximumReccommendedOustandingOcclusionQueries < MAX_int32)
-							{
-								// If there's no occlusion query for the primitive, assume it is whatever it was last frame
-								bIsOccluded = PrimitiveOcclusionHistory->WasOccludedLastFrame;
-								bOcclusionStateIsDefinite = PrimitiveOcclusionHistory->OcclusionStateWasDefiniteLastFrame;
-							}
-							else
-							{
-								// If there's no occlusion query for the primitive, set it's visibility state to whether it has been unoccluded recently.
-								bIsOccluded = (PrimitiveOcclusionHistory->LastProvenVisibleTime + GEngine->PrimitiveProbablyVisibleTime < CurrentRealTime);
-								// the state was definite last frame, otherwise we would have ran a query
-								bOcclusionStateIsDefinite = true;
-							}
-							if (bIsOccluded)
-							{
-								PrimitiveOcclusionHistory->LastPixelsPercentage = 0.0f;
-							}
-							else
-							{
-								PrimitiveOcclusionHistory->LastPixelsPercentage = GEngine->MaxOcclusionPixelsFraction;
-							}
+							const FBoxSphereBounds& Bounds = bSubQueries ? (*SubBounds)[SubQuery] : Scene->PrimitiveOcclusionBounds[BitIt.GetIndex()];
+							DrawWireBox(OcclusionPDI, Bounds.GetBox(), FColor(50, 255, 50), SDPG_Foreground);
 						}
 					}
-
-					if (GVisualizeOccludedPrimitives && OcclusionPDI && bIsOccluded)
+					else
 					{
-						const FBoxSphereBounds& Bounds = bSubQueries ? (*SubBounds)[SubQuery] : Scene->PrimitiveOcclusionBounds[BitIt.GetIndex()];
-						DrawWireBox(OcclusionPDI, Bounds.GetBox(), FColor(50, 255, 50), SDPG_Foreground);
+						// Primitives that aren't occludable are considered definitely unoccluded.
+						// already set bIsOccluded = false;
+						bOcclusionStateIsDefinite = true;
 					}
-				}
-				else
-				{					
-					// Primitives that aren't occludable are considered definitely unoccluded.
-					// already set bIsOccluded = false;
-					bOcclusionStateIsDefinite = true;
-				}
 
-				if (bClearQueries)
-				{					
-					if (bSingleThreaded)
-					{						
-						PrimitiveOcclusionHistory->ReleaseQuery(OcclusionFrameCounter, NumBufferedFrames);
-					}
-					else
+					if (bClearQueries)
 					{
-						if (PrimitiveOcclusionHistory->GetQueryForEviction(OcclusionFrameCounter, NumBufferedFrames))
+						if (bSingleThreaded)
 						{
-							QueriesToRelease->Add(PrimitiveOcclusionHistory);							
+							PrimitiveOcclusionHistory->ReleaseQuery(OcclusionFrameCounter, NumBufferedFrames);
+						}
+						else
+						{
+							if (PrimitiveOcclusionHistory->GetQueryForEviction(OcclusionFrameCounter, NumBufferedFrames))
+							{
+								QueriesToRelease->Add(PrimitiveOcclusionHistory);
+							}
 						}
 					}
 				}
-			}
 
-			if (PrimitiveOcclusionHistory)
-			{
-				if (bSubmitQueries && bCanBeOccluded)
+				if (PrimitiveOcclusionHistory)
 				{
-					bool bSkipNewlyConsidered = false;
-
-					if (bNewlyConsideredBBoxExpandActive)
+					if (bSubmitQueries && bCanBeOccluded)
 					{
-						if (!PrimitiveOcclusionHistory->BecameEligibleForQueryCooldown && OcclusionFrameCounter - PrimitiveOcclusionHistory->LastConsideredFrameNumber > uint32(GFramesNotOcclusionTestedToExpandBBoxes))
+						bool bSkipNewlyConsidered = false;
+
+						if (bNewlyConsideredBBoxExpandActive)
 						{
-							PrimitiveOcclusionHistory->BecameEligibleForQueryCooldown = GFramesToExpandNewlyOcclusionTestedBBoxes;
-						}
+							if (!PrimitiveOcclusionHistory->BecameEligibleForQueryCooldown && OcclusionFrameCounter - PrimitiveOcclusionHistory->LastConsideredFrameNumber > uint32(GFramesNotOcclusionTestedToExpandBBoxes))
+							{
+								PrimitiveOcclusionHistory->BecameEligibleForQueryCooldown = GFramesToExpandNewlyOcclusionTestedBBoxes;
+							}
 
-						bSkipNewlyConsidered = !!PrimitiveOcclusionHistory->BecameEligibleForQueryCooldown;
+							bSkipNewlyConsidered = !!PrimitiveOcclusionHistory->BecameEligibleForQueryCooldown;
 
-						if (bSkipNewlyConsidered)
-						{
-							PrimitiveOcclusionHistory->BecameEligibleForQueryCooldown--;
+							if (bSkipNewlyConsidered)
+							{
+								PrimitiveOcclusionHistory->BecameEligibleForQueryCooldown--;
+							}
 						}
-					}
-
 
-					bool bAllowBoundsTest;
-					const FBoxSphereBounds OcclusionBounds = (bSubQueries ? (*SubBounds)[SubQuery] : Scene->PrimitiveOcclusionBounds[BitIt.GetIndex()]).ExpandBy(GExpandAllTestedBBoxesAmount + (bSkipNewlyConsidered ? GExpandNewlyOcclusionTestedBBoxesAmount : 0.0));
-					if (FVector::DistSquared(ViewOrigin, OcclusionBounds.Origin) < NeverOcclusionTestDistanceSquared)
-					{
-						bAllowBoundsTest = false;
-					}
-					else if (View.bHasNearClippingPlane)
-					{
-						bAllowBoundsTest = View.NearClippingPlane.PlaneDot(OcclusionBounds.Origin) <
-							-(FVector::BoxPushOut(View.NearClippingPlane, OcclusionBounds.BoxExtent));
 
-					}
-					else if (!View.IsPerspectiveProjection())
-					{
-						// Transform parallel near plane
-						static_assert((int32)ERHIZBuffer::IsInverted != 0, "Check equation for culling!");
-						bAllowBoundsTest = View.WorldToScreen(OcclusionBounds.Origin).Z - View.ViewMatrices.GetProjectionMatrix().M[2][2] * OcclusionBounds.SphereRadius < 1;
-					}
-					else
-					{
-						bAllowBoundsTest = OcclusionBounds.SphereRadius < HALF_WORLD_MAX;
-					}
+						bool bAllowBoundsTest;
+						const FBoxSphereBounds OcclusionBounds = (bSubQueries ? (*SubBounds)[SubQuery] : Scene->PrimitiveOcclusionBounds[BitIt.GetIndex()]).ExpandBy(GExpandAllTestedBBoxesAmount + (bSkipNewlyConsidered ? GExpandNewlyOcclusionTestedBBoxesAmount : 0.0));
+						if (FVector::DistSquared(ViewOrigin, OcclusionBounds.Origin) < NeverOcclusionTestDistanceSquared)
+						{
+							bAllowBoundsTest = false;
+						}
+						else if (View.bHasNearClippingPlane)
+						{
+							bAllowBoundsTest = View.NearClippingPlane.PlaneDot(OcclusionBounds.Origin) <
+								-(FVector::BoxPushOut(View.NearClippingPlane, OcclusionBounds.BoxExtent));
 
-					if (bAllowBoundsTest)
-					{
-						PrimitiveOcclusionHistory->LastTestFrameNumber = OcclusionFrameCounter;
-						if (bHZBOcclusion)
+						}
+						else if (!View.IsPerspectiveProjection())
 						{
-							// Always run
-							if (bSingleThreaded)
-							{								
-								PrimitiveOcclusionHistory->HZBTestIndex = HZBOcclusionTests.AddBounds(OcclusionBounds.Origin, OcclusionBounds.BoxExtent);
-							}
-							else
-							{
-								HZBBoundsToAdd->Emplace(PrimitiveOcclusionHistory, OcclusionBounds.Origin, OcclusionBounds.BoxExtent);
-							}
+							// Transform parallel near plane
+							static_assert((int32)ERHIZBuffer::IsInverted != 0, "Check equation for culling!");
+							bAllowBoundsTest = View.WorldToScreen(OcclusionBounds.Origin).Z - View.ViewMatrices.GetProjectionMatrix().M[2][2] * OcclusionBounds.SphereRadius < 1;
 						}
 						else
 						{
-							// decide if a query should be run this frame
-							bool bRunQuery, bGroupedQuery;
+							bAllowBoundsTest = OcclusionBounds.SphereRadius < HALF_WORLD_MAX;
+						}
 
-							if (!bSubQueries && // sub queries are never grouped, we assume the custom code knows what it is doing and will group internally if it wants
-								(OcclusionFlags & EOcclusionFlags::AllowApproximateOcclusion))
+						if (bAllowBoundsTest)
+						{
+							PrimitiveOcclusionHistory->LastTestFrameNumber = OcclusionFrameCounter;
+							if (bHZBOcclusion)
 							{
-								if (bIsOccluded)
+								// Always run
+								if (bSingleThreaded)
 								{
-									// Primitives that were occluded the previous frame use grouped queries.
-									bGroupedQuery = true;
-									bRunQuery = true;
+									PrimitiveOcclusionHistory->HZBTestIndex = HZBOcclusionTests.AddBounds(OcclusionBounds.Origin, OcclusionBounds.BoxExtent);
 								}
-								else if (bOcclusionStateIsDefinite)
+								else
 								{
-									bGroupedQuery = false;
-									float Rnd = GOcclusionRandomStream.GetFraction();
-									if (GRHISupportsExactOcclusionQueries)
+									HZBBoundsToAdd->Emplace(PrimitiveOcclusionHistory, OcclusionBounds.Origin, OcclusionBounds.BoxExtent);
+								}
+							}
+							else
+							{
+								// decide if a query should be run this frame
+								bool bRunQuery, bGroupedQuery;
+
+								if (!bSubQueries && // sub queries are never grouped, we assume the custom code knows what it is doing and will group internally if it wants
+									(OcclusionFlags & EOcclusionFlags::AllowApproximateOcclusion))
+								{
+									if (bIsOccluded)
+									{
+										// Primitives that were occluded the previous frame use grouped queries.
+										bGroupedQuery = true;
+										bRunQuery = true;
+									}
+									else if (bOcclusionStateIsDefinite)
 									{
-										float FractionMultiplier = FMath::Max(PrimitiveOcclusionHistory->LastPixelsPercentage / GEngine->MaxOcclusionPixelsFraction, 1.0f);
-										bRunQuery = (FractionMultiplier * Rnd) < GEngine->MaxOcclusionPixelsFraction;
+										bGroupedQuery = false;
+										float Rnd = GOcclusionRandomStream.GetFraction();
+										if (GRHISupportsExactOcclusionQueries)
+										{
+											float FractionMultiplier = FMath::Max(PrimitiveOcclusionHistory->LastPixelsPercentage / GEngine->MaxOcclusionPixelsFraction, 1.0f);
+											bRunQuery = (FractionMultiplier * Rnd) < GEngine->MaxOcclusionPixelsFraction;
+										}
+										else
+										{
+											bRunQuery = CurrentRealTime - PrimitiveOcclusionHistory->LastProvenVisibleTime > PrimitiveProbablyVisibleTime * (0.5f * 0.25f * Rnd);
+										}
 									}
 									else
 									{
-										bRunQuery = CurrentRealTime - PrimitiveOcclusionHistory->LastProvenVisibleTime > PrimitiveProbablyVisibleTime * (0.5f * 0.25f * Rnd);
-									}									
+										bGroupedQuery = false;
+										bRunQuery = true;
+									}
 								}
 								else
 								{
+									// Primitives that need precise occlusion results use individual queries.
 									bGroupedQuery = false;
 									bRunQuery = true;
 								}
-							}
-							else
-							{
-								// Primitives that need precise occlusion results use individual queries.
-								bGroupedQuery = false;
-								bRunQuery = true;
-							}
-
-							if (bRunQuery)
-							{
-								const FVector BoundOrigin = OcclusionBounds.Origin + View.ViewMatrices.GetPreViewTranslation();
-								const FVector BoundExtent = OcclusionBounds.BoxExtent;
 
-								if (bSingleThreaded)
+								if (bRunQuery)
 								{
-									checkSlow(DynamicVertexBufferIfSingleThreaded);
+									const FVector BoundOrigin = OcclusionBounds.Origin + View.ViewMatrices.GetPreViewTranslation();
+									const FVector BoundExtent = OcclusionBounds.BoxExtent;
 
-									if (GRHIMaximumReccommendedOustandingOcclusionQueries < MAX_int32 && !bGroupedQuery)
+									if (bSingleThreaded)
 									{
-										QueriesToAdd->Emplace(FPrimitiveOcclusionHistoryKey(PrimitiveId, SubQuery), BoundOrigin, BoundExtent, PrimitiveOcclusionHistory->LastQuerySubmitFrame());
+										checkSlow(DynamicVertexBufferIfSingleThreaded);
+
+										if (GRHIMaximumReccommendedOustandingOcclusionQueries < MAX_int32 && !bGroupedQuery)
+										{
+											QueriesToAdd->Emplace(FPrimitiveOcclusionHistoryKey(PrimitiveId, SubQuery), BoundOrigin, BoundExtent, PrimitiveOcclusionHistory->LastQuerySubmitFrame());
+										}
+										else
+										{
+											PrimitiveOcclusionHistory->SetCurrentQuery(OcclusionFrameCounter,
+												bGroupedQuery ?
+												View.GroupedOcclusionQueries.BatchPrimitive(BoundOrigin, BoundExtent, *DynamicVertexBufferIfSingleThreaded) :
+												View.IndividualOcclusionQueries.BatchPrimitive(BoundOrigin, BoundExtent, *DynamicVertexBufferIfSingleThreaded),
+												NumBufferedFrames,
+												bGroupedQuery,
+												Params.bNeedsScanOnRead
+											);
+										}
 									}
 									else
 									{
-										PrimitiveOcclusionHistory->SetCurrentQuery(OcclusionFrameCounter,
-											bGroupedQuery ?
-											View.GroupedOcclusionQueries.BatchPrimitive(BoundOrigin, BoundExtent, *DynamicVertexBufferIfSingleThreaded) :
-											View.IndividualOcclusionQueries.BatchPrimitive(BoundOrigin, BoundExtent, *DynamicVertexBufferIfSingleThreaded),
-											NumBufferedFrames,
-											bGroupedQuery,
-											Params.bNeedsScanOnRead
-										);
+										check(GRHIMaximumReccommendedOustandingOcclusionQueries < MAX_int32); // it would be fairly easy to set up this path to optimize when there are a limited number, but it hasn't been done yet
+										QueriesToAdd->Emplace(PrimitiveOcclusionHistory, BoundOrigin, BoundExtent, bGroupedQuery);
 									}
 								}
-								else
-								{
-									check(GRHIMaximumReccommendedOustandingOcclusionQueries < MAX_int32); // it would be fairly easy to set up this path to optimize when there are a limited number, but it hasn't been done yet
-									QueriesToAdd->Emplace(PrimitiveOcclusionHistory, BoundOrigin, BoundExtent, bGroupedQuery);
-								}
 							}
 						}
+						else
+						{
+							// If the primitive's bounding box intersects the near clipping plane, treat it as definitely unoccluded.
+							bIsOccluded = false;
+							bOcclusionStateIsDefinite = true;
+						}
 					}
-					else
+					// Set the primitive's considered time to keep its occlusion history from being trimmed.
+					PrimitiveOcclusionHistory->LastConsideredTime = CurrentRealTime;
+					if (!bIsOccluded && bOcclusionStateIsDefinite)
 					{
-						// If the primitive's bounding box intersects the near clipping plane, treat it as definitely unoccluded.
-						bIsOccluded = false;
-						bOcclusionStateIsDefinite = true;
+						PrimitiveOcclusionHistory->LastProvenVisibleTime = CurrentRealTime;
 					}
+					PrimitiveOcclusionHistory->LastConsideredFrameNumber = OcclusionFrameCounter;
+					PrimitiveOcclusionHistory->WasOccludedLastFrame = bIsOccluded;
+					PrimitiveOcclusionHistory->OcclusionStateWasDefiniteLastFrame = bOcclusionStateIsDefinite;
 				}
-				// Set the primitive's considered time to keep its occlusion history from being trimmed.
-				PrimitiveOcclusionHistory->LastConsideredTime = CurrentRealTime;
-				if (!bIsOccluded && bOcclusionStateIsDefinite)
+
+				if (bSubQueries)
+				{
+					SubIsOccluded.Add(bIsOccluded);
+					if (!bIsOccluded)
+					{
+						bAllSubOccluded = false;
+					}
+					if (bIsOccluded || !bOcclusionStateIsDefinite)
+					{
+						bAllSubOcclusionStateIsDefinite = false;
+					}
+				}
+				else
 				{
-					PrimitiveOcclusionHistory->LastProvenVisibleTime = CurrentRealTime;
+
+					if (bIsOccluded)
+					{
+						View.PrimitiveVisibilityMap.AccessCorrespondingBit(BitIt) = false;
+						STAT(NumOccludedPrimitives++);
+					}
+					else if (bOcclusionStateIsDefinite)
+					{
+						View.PrimitiveDefinitelyUnoccludedMap.AccessCorrespondingBit(BitIt) = true;
+					}
 				}
-				PrimitiveOcclusionHistory->LastConsideredFrameNumber = OcclusionFrameCounter;
-				PrimitiveOcclusionHistory->WasOccludedLastFrame = bIsOccluded;
-				PrimitiveOcclusionHistory->OcclusionStateWasDefiniteLastFrame = bOcclusionStateIsDefinite;
 			}
 
 			if (bSubQueries)
 			{
-				SubIsOccluded.Add(bIsOccluded);
-				if (!bIsOccluded)
-				{
-					bAllSubOccluded = false;
-				}
-				if (bIsOccluded || !bOcclusionStateIsDefinite)
+				if (SubIsOccluded.Num() > 0)
 				{
-					bAllSubOcclusionStateIsDefinite = false;
+					FPrimitiveSceneProxy* Proxy = Scene->Primitives[BitIt.GetIndex()]->Proxy;
+					Proxy->AcceptOcclusionResults(&View, &SubIsOccluded, SubIsOccludedStart, SubIsOccluded.Num() - SubIsOccludedStart);
 				}
-			}
-			else
-			{
-					
-				if (bIsOccluded)
+
+				if (bAllSubOccluded)
 				{
 					View.PrimitiveVisibilityMap.AccessCorrespondingBit(BitIt) = false;
 					STAT(NumOccludedPrimitives++);
 				}
-				else if (bOcclusionStateIsDefinite)
+				else if (bAllSubOcclusionStateIsDefinite)
 				{
 					View.PrimitiveDefinitelyUnoccludedMap.AccessCorrespondingBit(BitIt) = true;
-				}					
-			}			
-		}
-
-		if (bSubQueries)
-		{
-			if (SubIsOccluded.Num() > 0)
-			{
-				FPrimitiveSceneProxy* Proxy = Scene->Primitives[BitIt.GetIndex()]->Proxy;
-				Proxy->AcceptOcclusionResults(&View, &SubIsOccluded, SubIsOccludedStart, SubIsOccluded.Num() - SubIsOccludedStart);
-			}
-
-			if (bAllSubOccluded)
-			{
-				View.PrimitiveVisibilityMap.AccessCorrespondingBit(BitIt) = false;
-				STAT(NumOccludedPrimitives++);
-			}
-			else if (bAllSubOcclusionStateIsDefinite)
-			{
-				View.PrimitiveDefinitelyUnoccludedMap.AccessCorrespondingBit(BitIt) = true;
+				}
 			}
 		}
 	}
-
 	check(NumTotalDefUnoccluded == View.PrimitiveDefinitelyUnoccludedMap.Num());
 	check(NumTotalPrims == View.PrimitiveVisibilityMap.Num());
 	check(!InsertPrimitiveOcclusionHistory || InsertPrimitiveOcclusionHistory->Num() <= ReserveAmount);
@@ -1347,6 +1351,7 @@ static int32 FetchVisibilityForPrimitives(const FScene* Scene, FViewInfo& View,
 	CSV_SCOPED_TIMING_STAT_EXCLUSIVE(FetchVisibilityForPrimitives);
 	QUICK_SCOPE_CYCLE_COUNTER(STAT_FetchVisibilityForPrimitives);
 	FSceneViewState* ViewState = (FSceneViewState*)View.State;
+	TRACE_CPUPROFILER_EVENT_SCOPE(FetchVisibilityForPrimitives);
 	
 	static int32 SubIsOccludedArrayIndex = 0;
 	SubIsOccludedArrayIndex = 1 - SubIsOccludedArrayIndex;
@@ -1357,6 +1362,7 @@ static int32 FetchVisibilityForPrimitives(const FScene* Scene, FViewInfo& View,
 
 	if (GOcclusionCullParallelPrimFetch && GSupportsParallelOcclusionQueries)
 	{		
+		TRACE_CPUPROFILER_EVENT_SCOPE(FetchVisibilityParallel);
 		static const int32 MaxNumCullTasks = 4;
 		static const int32 ActualNumCullTasks = 4;
 		static const int32 NumOutputArrays = MaxNumCullTasks;
@@ -1457,6 +1463,7 @@ static int32 FetchVisibilityForPrimitives(const FScene* Scene, FViewInfo& View,
 		int32 NumOccludedPrims = 0;
 		{
 			QUICK_SCOPE_CYCLE_COUNTER(STAT_FetchVisibilityForPrimitivesCombine);
+			TRACE_CPUPROFILER_EVENT_SCOPE(FetchVisibilityForPrimitivesCombine);
 
 			//wait for them all so we don't start modifying the prim histories while the gather is running
 			FTaskGraphInterface::Get().WaitUntilTasksComplete(TaskWaitArray, ENamedThreads::GetRenderThread_Local());
@@ -1535,6 +1542,7 @@ static int32 FetchVisibilityForPrimitives(const FScene* Scene, FViewInfo& View,
 	}
 	else
 	{
+		TRACE_CPUPROFILER_EVENT_SCOPE(FetchVisibilityOther);
 		//SubIsOccluded stuff needs a frame's lifetime
 		TArray<bool>& SubIsOccluded = View.FrameSubIsOccluded[SubIsOccludedArrayIndex];
 		SubIsOccluded.Reset();
@@ -2638,7 +2646,7 @@ static void ComputeAndMarkRelevanceForViewParallel(
 	FPrimitiveViewMasks& OutHasDynamicEditorMeshElementsMasks
 	)
 {
-	SCOPED_NAMED_EVENT(FSceneRenderer_ComputeAndMarkRelevanceForViewParallel, FColor::Blue);
+	TRACE_CPUPROFILER_EVENT_SCOPE(FSceneRenderer_ComputeAndMarkRelevanceForViewParallel);
 
 	check(OutHasDynamicMeshElementsMasks.Num() == Scene->Primitives.Num());
 
@@ -4003,267 +4011,272 @@ void FSceneRenderer::ComputeViewVisibility(
 	}
 
 	uint8 ViewBit = 0x1;
-	for (int32 ViewIndex = 0; ViewIndex < Views.Num(); ++ViewIndex, ViewBit <<= 1)
 	{
-		STAT(NumProcessedPrimitives += NumPrimitives);
-
-		FViewInfo& View = Views[ViewIndex];
-		FViewCommands& ViewCommands = ViewCommandsPerView[ViewIndex];
-		FSceneViewState* ViewState = (FSceneViewState*)View.State;
-
-		// Allocate the view's visibility maps.
-		View.PrimitiveVisibilityMap.Init(false,Scene->Primitives.Num());
-
-		// These are not initialized here, as we overwrite the whole array in GatherDynamicMeshElements().
-		View.DynamicMeshEndIndices.SetNumUninitialized(Scene->Primitives.Num());
-		View.PrimitiveDefinitelyUnoccludedMap.Init(false,Scene->Primitives.Num());
-		View.PotentiallyFadingPrimitiveMap.Init(false,Scene->Primitives.Num());
-		View.PrimitiveFadeUniformBuffers.AddZeroed(Scene->Primitives.Num());
-		View.PrimitiveFadeUniformBufferMap.Init(false, Scene->Primitives.Num());
-		View.StaticMeshVisibilityMap.Init(false,Scene->StaticMeshes.GetMaxIndex());
-		View.StaticMeshFadeOutDitheredLODMap.Init(false,Scene->StaticMeshes.GetMaxIndex());
-		View.StaticMeshFadeInDitheredLODMap.Init(false,Scene->StaticMeshes.GetMaxIndex());
-		View.PrimitivesLODMask.Init(FLODMask(), Scene->Primitives.Num());
-		View.DistanceCullingPrimitiveMap.Init(false, Scene->Primitives.Num());
-
-		View.VisibleLightInfos.Empty(Scene->Lights.GetMaxIndex());
-
-		// The dirty list allocation must take into account the max possible size because when GILCUpdatePrimTaskEnabled is true,
-		// the indirect lighting cache will be update on by threaded job, which can not do reallocs on the buffer (since it uses the SceneRenderingAllocator).
-		View.DirtyIndirectLightingCacheBufferPrimitives.Reserve(Scene->Primitives.Num());
-
-		for(int32 LightIndex = 0;LightIndex < Scene->Lights.GetMaxIndex();LightIndex++)
+		TRACE_CPUPROFILER_EVENT_SCOPE(FSceneRenderer_Views);
+		for (int32 ViewIndex = 0; ViewIndex < Views.Num(); ++ViewIndex, ViewBit <<= 1)
 		{
-			if( LightIndex+2 < Scene->Lights.GetMaxIndex() )
+			TRACE_CPUPROFILER_EVENT_SCOPE_TEXT(*FString::Printf(TEXT("View %d"), ViewIndex));
+			STAT(NumProcessedPrimitives += NumPrimitives);
+
+			FViewInfo& View = Views[ViewIndex];
+			FViewCommands& ViewCommands = ViewCommandsPerView[ViewIndex];
+			FSceneViewState* ViewState = (FSceneViewState*)View.State;
+
+			// Allocate the view's visibility maps.
+			View.PrimitiveVisibilityMap.Init(false, Scene->Primitives.Num());
+
+			// These are not initialized here, as we overwrite the whole array in GatherDynamicMeshElements().
+			View.DynamicMeshEndIndices.SetNumUninitialized(Scene->Primitives.Num());
+			View.PrimitiveDefinitelyUnoccludedMap.Init(false, Scene->Primitives.Num());
+			View.PotentiallyFadingPrimitiveMap.Init(false, Scene->Primitives.Num());
+			View.PrimitiveFadeUniformBuffers.AddZeroed(Scene->Primitives.Num());
+			View.PrimitiveFadeUniformBufferMap.Init(false, Scene->Primitives.Num());
+			View.StaticMeshVisibilityMap.Init(false, Scene->StaticMeshes.GetMaxIndex());
+			View.StaticMeshFadeOutDitheredLODMap.Init(false, Scene->StaticMeshes.GetMaxIndex());
+			View.StaticMeshFadeInDitheredLODMap.Init(false, Scene->StaticMeshes.GetMaxIndex());
+			View.PrimitivesLODMask.Init(FLODMask(), Scene->Primitives.Num());
+			View.DistanceCullingPrimitiveMap.Init(false, Scene->Primitives.Num());
+
+			View.VisibleLightInfos.Empty(Scene->Lights.GetMaxIndex());
+
+			// The dirty list allocation must take into account the max possible size because when GILCUpdatePrimTaskEnabled is true,
+			// the indirect lighting cache will be update on by threaded job, which can not do reallocs on the buffer (since it uses the SceneRenderingAllocator).
+			View.DirtyIndirectLightingCacheBufferPrimitives.Reserve(Scene->Primitives.Num());
+
+			for (int32 LightIndex = 0;LightIndex < Scene->Lights.GetMaxIndex();LightIndex++)
 			{
-				if (LightIndex > 2)
+				if (LightIndex + 2 < Scene->Lights.GetMaxIndex())
 				{
-					FLUSH_CACHE_LINE(&View.VisibleLightInfos(LightIndex-2));
+					if (LightIndex > 2)
+					{
+						FLUSH_CACHE_LINE(&View.VisibleLightInfos(LightIndex - 2));
+					}
+					// @todo optimization These prefetches cause asserts since LightIndex > View.VisibleLightInfos.Num() - 1
+					//FPlatformMisc::Prefetch(&View.VisibleLightInfos[LightIndex+2]);
+					//FPlatformMisc::Prefetch(&View.VisibleLightInfos[LightIndex+1]);
 				}
-				// @todo optimization These prefetches cause asserts since LightIndex > View.VisibleLightInfos.Num() - 1
-				//FPlatformMisc::Prefetch(&View.VisibleLightInfos[LightIndex+2]);
-				//FPlatformMisc::Prefetch(&View.VisibleLightInfos[LightIndex+1]);
+				new(View.VisibleLightInfos) FVisibleLightViewInfo();
 			}
-			new(View.VisibleLightInfos) FVisibleLightViewInfo();
-		}
 
-		View.PrimitiveViewRelevanceMap.Reset(Scene->Primitives.Num());
-		View.PrimitiveViewRelevanceMap.AddZeroed(Scene->Primitives.Num());
+			View.PrimitiveViewRelevanceMap.Reset(Scene->Primitives.Num());
+			View.PrimitiveViewRelevanceMap.AddZeroed(Scene->Primitives.Num());
 
-		// If this is the visibility-parent of other views, reset its ParentPrimitives list.
-		const bool bIsParent = ViewState && ViewState->IsViewParent();
-		if ( bIsParent )
-		{
-			// PVS-Studio does not understand the validation of ViewState above, so we're disabling
-			// its warning that ViewState may be null:
-			ViewState->ParentPrimitives.Reset(); //-V595
-		}
+			// If this is the visibility-parent of other views, reset its ParentPrimitives list.
+			const bool bIsParent = ViewState && ViewState->IsViewParent();
+			if (bIsParent)
+			{
+				// PVS-Studio does not understand the validation of ViewState above, so we're disabling
+				// its warning that ViewState may be null:
+				ViewState->ParentPrimitives.Reset(); //-V595
+			}
 
-		if (ViewState)
-		{	
-			SCOPE_CYCLE_COUNTER(STAT_DecompressPrecomputedOcclusion);
-			View.PrecomputedVisibilityData = ViewState->GetPrecomputedVisibilityData(View, Scene);
-		}
-		else
-		{
-			View.PrecomputedVisibilityData = nullptr;
-		}
+			if (ViewState)
+			{
+				SCOPE_CYCLE_COUNTER(STAT_DecompressPrecomputedOcclusion);
+				View.PrecomputedVisibilityData = ViewState->GetPrecomputedVisibilityData(View, Scene);
+			}
+			else
+			{
+				View.PrecomputedVisibilityData = nullptr;
+			}
 
-		if (View.PrecomputedVisibilityData)
-		{
-			bUsedPrecomputedVisibility = true;
-		}
+			if (View.PrecomputedVisibilityData)
+			{
+				bUsedPrecomputedVisibility = true;
+			}
 
-		bool bNeedsFrustumCulling = true;
+			bool bNeedsFrustumCulling = true;
 
-		// Development builds sometimes override frustum culling, e.g. dependent views in the editor.
+			// Development builds sometimes override frustum culling, e.g. dependent views in the editor.
 #if !(UE_BUILD_SHIPPING || UE_BUILD_TEST)
-		if( ViewState )
-		{
-#if WITH_EDITOR
-			// For visibility child views, check if the primitive was visible in the parent view.
-			const FSceneViewState* const ViewParent = (FSceneViewState*)ViewState->GetViewParent();
-			if(ViewParent)
+			if (ViewState)
 			{
-				bNeedsFrustumCulling = false;
-				for (FSceneBitArray::FIterator BitIt(View.PrimitiveVisibilityMap); BitIt; ++BitIt)
+#if WITH_EDITOR
+				// For visibility child views, check if the primitive was visible in the parent view.
+				const FSceneViewState* const ViewParent = (FSceneViewState*)ViewState->GetViewParent();
+				if (ViewParent)
 				{
-					if (ViewParent->ParentPrimitives.Contains(Scene->PrimitiveComponentIds[BitIt.GetIndex()]) ||
-						Scene->PrimitivesAlwaysVisible[BitIt.GetIndex()])
+					bNeedsFrustumCulling = false;
+					for (FSceneBitArray::FIterator BitIt(View.PrimitiveVisibilityMap); BitIt; ++BitIt)
 					{
-						BitIt.GetValue() = true;
+						if (ViewParent->ParentPrimitives.Contains(Scene->PrimitiveComponentIds[BitIt.GetIndex()]) ||
+							Scene->PrimitivesAlwaysVisible[BitIt.GetIndex()])
+						{
+							BitIt.GetValue() = true;
+						}
 					}
 				}
-			}
 #endif
-			// For views with frozen visibility, check if the primitive is in the frozen visibility set.
-			if (ViewState->bIsFrozen)
-			{
-				bNeedsFrustumCulling = false;
-				for (FSceneBitArray::FIterator BitIt(View.PrimitiveVisibilityMap); BitIt; ++BitIt)
+				// For views with frozen visibility, check if the primitive is in the frozen visibility set.
+				if (ViewState->bIsFrozen)
 				{
-					if (ViewState->FrozenPrimitives.Contains(Scene->PrimitiveComponentIds[BitIt.GetIndex()]) ||
-						Scene->PrimitivesAlwaysVisible[BitIt.GetIndex()])
+					bNeedsFrustumCulling = false;
+					for (FSceneBitArray::FIterator BitIt(View.PrimitiveVisibilityMap); BitIt; ++BitIt)
 					{
-						BitIt.GetValue() = true;
+						if (ViewState->FrozenPrimitives.Contains(Scene->PrimitiveComponentIds[BitIt.GetIndex()]) ||
+							Scene->PrimitivesAlwaysVisible[BitIt.GetIndex()])
+						{
+							BitIt.GetValue() = true;
+						}
 					}
 				}
 			}
-		}
 #endif
 
-		// Most views use standard frustum culling.
-		if (bNeedsFrustumCulling)
-		{
-			// Update HLOD transition/visibility states to allow use during distance culling
-			FLODSceneTree& HLODTree = Scene->SceneLODHierarchy;
-			if (HLODTree.IsActive())
-			{
-				QUICK_SCOPE_CYCLE_COUNTER(STAT_ViewVisibilityTime_HLODUpdate);
-				HLODTree.UpdateVisibilityStates(View);
-			}
-			else
-			{
-				HLODTree.ClearVisibilityState(View);
-			}
-
-			int32 NumCulledPrimitivesForView;
-			const bool bUseFastIntersect = (View.ViewFrustum.PermutedPlanes.Num() == 8) && CVarUseFastIntersect.GetValueOnRenderThread();
-			if (View.CustomVisibilityQuery && View.CustomVisibilityQuery->Prepare())
+			// Most views use standard frustum culling.
+			if (bNeedsFrustumCulling)
 			{
-				if (CVarAlsoUseSphereForFrustumCull.GetValueOnRenderThread())
+				TRACE_CPUPROFILER_EVENT_SCOPE(FSceneRenderer_Cull);
+				// Update HLOD transition/visibility states to allow use during distance culling
+				FLODSceneTree& HLODTree = Scene->SceneLODHierarchy;
+				if (HLODTree.IsActive())
 				{
-					NumCulledPrimitivesForView = bUseFastIntersect ? FrustumCull<true, true, true>(Scene, View) : FrustumCull<true, true, false>(Scene, View);
+					QUICK_SCOPE_CYCLE_COUNTER(STAT_ViewVisibilityTime_HLODUpdate);
+					HLODTree.UpdateVisibilityStates(View);
 				}
 				else
 				{
-					NumCulledPrimitivesForView = bUseFastIntersect ? FrustumCull<true, false, true>(Scene, View) : FrustumCull<true, false, false>(Scene, View);
+					HLODTree.ClearVisibilityState(View);
 				}
-			}
-			else
-			{
-				if (CVarAlsoUseSphereForFrustumCull.GetValueOnRenderThread())
+
+				int32 NumCulledPrimitivesForView;
+				const bool bUseFastIntersect = (View.ViewFrustum.PermutedPlanes.Num() == 8) && CVarUseFastIntersect.GetValueOnRenderThread();
+				if (View.CustomVisibilityQuery && View.CustomVisibilityQuery->Prepare())
 				{
-					NumCulledPrimitivesForView = bUseFastIntersect ? FrustumCull<false, true, true>(Scene, View) : FrustumCull<false, true, false>(Scene, View);
+					if (CVarAlsoUseSphereForFrustumCull.GetValueOnRenderThread())
+					{
+						NumCulledPrimitivesForView = bUseFastIntersect ? FrustumCull<true, true, true>(Scene, View) : FrustumCull<true, true, false>(Scene, View);
+					}
+					else
+					{
+						NumCulledPrimitivesForView = bUseFastIntersect ? FrustumCull<true, false, true>(Scene, View) : FrustumCull<true, false, false>(Scene, View);
+					}
 				}
 				else
 				{
-					NumCulledPrimitivesForView = bUseFastIntersect ? FrustumCull<false, false, true>(Scene, View) : FrustumCull<false, false, false>(Scene, View);
+					if (CVarAlsoUseSphereForFrustumCull.GetValueOnRenderThread())
+					{
+						NumCulledPrimitivesForView = bUseFastIntersect ? FrustumCull<false, true, true>(Scene, View) : FrustumCull<false, true, false>(Scene, View);
+					}
+					else
+					{
+						NumCulledPrimitivesForView = bUseFastIntersect ? FrustumCull<false, false, true>(Scene, View) : FrustumCull<false, false, false>(Scene, View);
+					}
 				}
+				STAT(NumCulledPrimitives += NumCulledPrimitivesForView);
+				UpdatePrimitiveFading(Scene, View);
 			}
-			STAT(NumCulledPrimitives += NumCulledPrimitivesForView);
-			UpdatePrimitiveFading(Scene, View);			
-		}
 
-		// If any primitives are explicitly hidden, remove them now.
-		if (View.HiddenPrimitives.Num())
-		{
-			for (FSceneSetBitIterator BitIt(View.PrimitiveVisibilityMap); BitIt; ++BitIt)
+			// If any primitives are explicitly hidden, remove them now.
+			if (View.HiddenPrimitives.Num())
 			{
-				if (View.HiddenPrimitives.Contains(Scene->PrimitiveComponentIds[BitIt.GetIndex()]))
+				for (FSceneSetBitIterator BitIt(View.PrimitiveVisibilityMap); BitIt; ++BitIt)
 				{
-					// Should respect this over the always visible bit set, since we want the editor
-					// to support hiding layers of primitives for various reasons.
-					View.PrimitiveVisibilityMap.AccessCorrespondingBit(BitIt) = false;
+					if (View.HiddenPrimitives.Contains(Scene->PrimitiveComponentIds[BitIt.GetIndex()]))
+					{
+						// Should respect this over the always visible bit set, since we want the editor
+						// to support hiding layers of primitives for various reasons.
+						View.PrimitiveVisibilityMap.AccessCorrespondingBit(BitIt) = false;
+					}
 				}
 			}
-		}
 
-		// If the view has any show only primitives, hide everything else
-		if (View.ShowOnlyPrimitives.IsSet())
-		{
-			View.bHasNoVisiblePrimitive = View.ShowOnlyPrimitives->Num() == 0;
-			for (FSceneSetBitIterator BitIt(View.PrimitiveVisibilityMap); BitIt; ++BitIt)
+			// If the view has any show only primitives, hide everything else
+			if (View.ShowOnlyPrimitives.IsSet())
 			{
-				if (!View.ShowOnlyPrimitives->Contains(Scene->PrimitiveComponentIds[BitIt.GetIndex()]))
+				View.bHasNoVisiblePrimitive = View.ShowOnlyPrimitives->Num() == 0;
+				for (FSceneSetBitIterator BitIt(View.PrimitiveVisibilityMap); BitIt; ++BitIt)
 				{
-					View.PrimitiveVisibilityMap.AccessCorrespondingBit(BitIt) = false;
+					if (!View.ShowOnlyPrimitives->Contains(Scene->PrimitiveComponentIds[BitIt.GetIndex()]))
+					{
+						View.PrimitiveVisibilityMap.AccessCorrespondingBit(BitIt) = false;
+					}
 				}
 			}
-		}
 
-		if (View.bStaticSceneOnly)
-		{
-			for (FSceneSetBitIterator BitIt(View.PrimitiveVisibilityMap); BitIt; ++BitIt)
+			if (View.bStaticSceneOnly)
 			{
-				// Reflection captures should only capture objects that won't move, since reflection captures won't update at runtime
-				if (!Scene->Primitives[BitIt.GetIndex()]->Proxy->HasStaticLighting())
+				for (FSceneSetBitIterator BitIt(View.PrimitiveVisibilityMap); BitIt; ++BitIt)
 				{
-					View.PrimitiveVisibilityMap.AccessCorrespondingBit(BitIt) = false;
+					// Reflection captures should only capture objects that won't move, since reflection captures won't update at runtime
+					if (!Scene->Primitives[BitIt.GetIndex()]->Proxy->HasStaticLighting())
+					{
+						View.PrimitiveVisibilityMap.AccessCorrespondingBit(BitIt) = false;
+					}
 				}
 			}
-		}
 
-		// Cull small objects in wireframe in ortho views
-		// This is important for performance in the editor because wireframe disables any kind of occlusion culling
-		if (View.Family->EngineShowFlags.Wireframe)
-		{
-			float ScreenSizeScale = FMath::Max(View.ViewMatrices.GetProjectionMatrix().M[0][0] * View.ViewRect.Width(), View.ViewMatrices.GetProjectionMatrix().M[1][1] * View.ViewRect.Height());
-			for (FSceneSetBitIterator BitIt(View.PrimitiveVisibilityMap); BitIt; ++BitIt)
+			// Cull small objects in wireframe in ortho views
+			// This is important for performance in the editor because wireframe disables any kind of occlusion culling
+			if (View.Family->EngineShowFlags.Wireframe)
 			{
-				if (ScreenSizeScale * Scene->PrimitiveBounds[BitIt.GetIndex()].BoxSphereBounds.SphereRadius <= GWireframeCullThreshold)
+				float ScreenSizeScale = FMath::Max(View.ViewMatrices.GetProjectionMatrix().M[0][0] * View.ViewRect.Width(), View.ViewMatrices.GetProjectionMatrix().M[1][1] * View.ViewRect.Height());
+				for (FSceneSetBitIterator BitIt(View.PrimitiveVisibilityMap); BitIt; ++BitIt)
 				{
-					View.PrimitiveVisibilityMap.AccessCorrespondingBit(BitIt) = false;
+					if (ScreenSizeScale * Scene->PrimitiveBounds[BitIt.GetIndex()].BoxSphereBounds.SphereRadius <= GWireframeCullThreshold)
+					{
+						View.PrimitiveVisibilityMap.AccessCorrespondingBit(BitIt) = false;
+					}
 				}
 			}
-		}
-
-		// Occlusion cull for all primitives in the view frustum, but not in wireframe.
-		if (!View.Family->EngineShowFlags.Wireframe)
-		{
-			int32 NumOccludedPrimitivesInView = OcclusionCull(RHICmdList, Scene, View, DynamicVertexBuffer);
-			STAT(NumOccludedPrimitives += NumOccludedPrimitivesInView);
-		}
-
-		{
-			QUICK_SCOPE_CYCLE_COUNTER(STAT_ViewVisibilityTime_ConditionalUpdateStaticMeshes);
-			SCOPED_NAMED_EVENT(FSceneRenderer_UpdateStaticMeshes, FColor::Red);
 
-			TArray<FPrimitiveSceneInfo*> AddedSceneInfos;
-			for (TConstDualSetBitIterator<SceneRenderingBitArrayAllocator, FDefaultBitArrayAllocator> BitIt(View.PrimitiveVisibilityMap, Scene->PrimitivesNeedingStaticMeshUpdate); BitIt; ++BitIt)
+			// Occlusion cull for all primitives in the view frustum, but not in wireframe.
+			if (!View.Family->EngineShowFlags.Wireframe)
 			{
-				int32 PrimitiveIndex = BitIt.GetIndex();
-				AddedSceneInfos.Add(Scene->Primitives[PrimitiveIndex]);
+				int32 NumOccludedPrimitivesInView = OcclusionCull(RHICmdList, Scene, View, DynamicVertexBuffer);
+				STAT(NumOccludedPrimitives += NumOccludedPrimitivesInView);
 			}
 
-			if (AddedSceneInfos.Num() > 0)
 			{
-				FPrimitiveSceneInfo::UpdateStaticMeshes(RHICmdList, Scene, AddedSceneInfos);
+				QUICK_SCOPE_CYCLE_COUNTER(STAT_ViewVisibilityTime_ConditionalUpdateStaticMeshes);
+				SCOPED_NAMED_EVENT(FSceneRenderer_UpdateStaticMeshes, FColor::Red);
+
+				TArray<FPrimitiveSceneInfo*> AddedSceneInfos;
+				for (TConstDualSetBitIterator<SceneRenderingBitArrayAllocator, FDefaultBitArrayAllocator> BitIt(View.PrimitiveVisibilityMap, Scene->PrimitivesNeedingStaticMeshUpdate); BitIt; ++BitIt)
+				{
+					int32 PrimitiveIndex = BitIt.GetIndex();
+					AddedSceneInfos.Add(Scene->Primitives[PrimitiveIndex]);
+				}
+
+				if (AddedSceneInfos.Num() > 0)
+				{
+					FPrimitiveSceneInfo::UpdateStaticMeshes(RHICmdList, Scene, AddedSceneInfos);
+				}
 			}
-		}
 
-		// ISR views can't compute relevance until all views are frustum culled
-		if (!bIsInstancedStereo)
-		{
-			SCOPE_CYCLE_COUNTER(STAT_ViewRelevance);
-			ComputeAndMarkRelevanceForViewParallel(RHICmdList, Scene, View, ViewCommands, ViewBit, HasDynamicMeshElementsMasks, HasDynamicEditorMeshElementsMasks);
-		}
+			// ISR views can't compute relevance until all views are frustum culled
+			if (!bIsInstancedStereo)
+			{
+				SCOPE_CYCLE_COUNTER(STAT_ViewRelevance);
+				ComputeAndMarkRelevanceForViewParallel(RHICmdList, Scene, View, ViewCommands, ViewBit, HasDynamicMeshElementsMasks, HasDynamicEditorMeshElementsMasks);
+			}
 
 #if !(UE_BUILD_SHIPPING || UE_BUILD_TEST)
-		// Store the primitive for parent occlusion rendering.
-		if (FPlatformProperties::SupportsWindowedMode() && ViewState && ViewState->IsViewParent())
-		{
-			for (FSceneDualSetBitIterator BitIt(View.PrimitiveVisibilityMap, View.PrimitiveDefinitelyUnoccludedMap); BitIt; ++BitIt)
+			// Store the primitive for parent occlusion rendering.
+			if (FPlatformProperties::SupportsWindowedMode() && ViewState && ViewState->IsViewParent())
 			{
-				ViewState->ParentPrimitives.Add(Scene->PrimitiveComponentIds[BitIt.GetIndex()]);
+				for (FSceneDualSetBitIterator BitIt(View.PrimitiveVisibilityMap, View.PrimitiveDefinitelyUnoccludedMap); BitIt; ++BitIt)
+				{
+					ViewState->ParentPrimitives.Add(Scene->PrimitiveComponentIds[BitIt.GetIndex()]);
+				}
 			}
-		}
 #endif
 
 #if !(UE_BUILD_SHIPPING || UE_BUILD_TEST)
-		// if we are freezing the scene, then remember the primitives that are rendered.
-		if (ViewState && ViewState->bIsFreezing)
-		{
-			for (FSceneSetBitIterator BitIt(View.PrimitiveVisibilityMap); BitIt; ++BitIt)
+			// if we are freezing the scene, then remember the primitives that are rendered.
+			if (ViewState && ViewState->bIsFreezing)
 			{
-				ViewState->FrozenPrimitives.Add(Scene->PrimitiveComponentIds[BitIt.GetIndex()]);
+				for (FSceneSetBitIterator BitIt(View.PrimitiveVisibilityMap); BitIt; ++BitIt)
+				{
+					ViewState->FrozenPrimitives.Add(Scene->PrimitiveComponentIds[BitIt.GetIndex()]);
+				}
 			}
-		}
 #endif
 
-		// TODO: right now decals visibility computed right before rendering them, ideally it should be done in InitViews and this flag should be replaced with list of visible decals  
-	    // Currently used to disable stencil operations in forward base pass when scene has no any decals
-		View.bSceneHasDecals = (Scene->Decals.Num() > 0) || (GForceSceneHasDecals != 0);
+			// TODO: right now decals visibility computed right before rendering them, ideally it should be done in InitViews and this flag should be replaced with list of visible decals  
+			// Currently used to disable stencil operations in forward base pass when scene has no any decals
+			View.bSceneHasDecals = (Scene->Decals.Num() > 0) || (GForceSceneHasDecals != 0);
+		}
 	}
 
 	if ((Views.Num() > 1) && bIsInstancedStereo)
