Fixed assert with -onethread caused by main thread MemMark

diff --git a/Engine/Source/Runtime/Renderer/Private/SceneRendering.cpp b/Engine/Source/Runtime/Renderer/Private/SceneRendering.cpp
--- a/Engine/Source/Runtime/Renderer/Private/SceneRendering.cpp
+++ b/Engine/Source/Runtime/Renderer/Private/SceneRendering.cpp
@@ -3650,9 +3650,7 @@ void FSceneRenderer::RenderThreadBegin(FRHICommandListImmediate& RHICmdList)
 {
 	CleanUp(RHICmdList);
 
-	FMemStackBase& MemStack = FMemStack::Get();
-	checkf(MemStack.GetNumMarks() == 0, TEXT("FSceneRenderer::RenderThreadBegin called with an active MemMark."));
-	MemStackMark = new FMemMark(MemStack);
+	MemStackMark = new FMemMark(FMemStack::Get());
 }
 
 struct FSceneRenderCleanUpState
