Reworking default textures to be read-only to reduce RDG tracking overhead.

diff --git a/Engine/Source/Runtime/RenderCore/Private/RenderUtils.cpp b/Engine/Source/Runtime/RenderCore/Private/RenderUtils.cpp
--- a/Engine/Source/Runtime/RenderCore/Private/RenderUtils.cpp
+++ b/Engine/Source/Runtime/RenderCore/Private/RenderUtils.cpp
@@ -16,29 +16,6 @@
 #include "Interfaces/ITargetPlatformManagerModule.h"
 #endif
 
-FTextureWithRDG::FTextureWithRDG() = default;
-FTextureWithRDG::FTextureWithRDG(const FTextureWithRDG& Other) = default;
-FTextureWithRDG& FTextureWithRDG::operator=(const FTextureWithRDG & Other) = default;
-FTextureWithRDG::~FTextureWithRDG() = default;
-
-FRDGTexture* FTextureWithRDG::GetRDG(FRDGBuilder& GraphBuilder) const
-{
-	checkf(RenderTarget, TEXT("InitRDG was not called before use."));
-	return GraphBuilder.RegisterExternalTexture(RenderTarget, ERenderTargetTexture::ShaderResource, ERDGTextureFlags::ReadOnly);
-}
-
-void FTextureWithRDG::ReleaseRHI()
-{
-	RenderTarget = nullptr;
-	FTexture::ReleaseRHI();
-}
-
-void FTextureWithRDG::InitRDG(const TCHAR* Name)
-{
-	check(TextureRHI);
-	RenderTarget = CreateRenderTarget(TextureRHI, Name);
-}
-
 FBufferWithRDG::FBufferWithRDG() = default;
 FBufferWithRDG::FBufferWithRDG(const FBufferWithRDG & Other) = default;
 FBufferWithRDG& FBufferWithRDG::operator=(const FBufferWithRDG & Other) = default;
@@ -151,7 +128,7 @@ private:
 /**
  * A solid-colored 1x1 texture.
  */
-template <int32 R, int32 G, int32 B, int32 A, bool bWithUAV = false>
+template <int32 R, int32 G, int32 B, int32 A>
 class FColoredTexture : public FTextureWithSRV
 {
 public:
@@ -162,10 +139,6 @@ public:
 		FBlackVolumeTextureResourceBulkDataInterface BlackTextureBulkData(FColor(R, G, B, A));
 		FRHIResourceCreateInfo CreateInfo(TEXT("ColoredTexture"), &BlackTextureBulkData);
 		ETextureCreateFlags CreateFlags = TexCreate_ShaderResource;
-		if(bWithUAV)
-		{
-			CreateFlags |= TexCreate_UAV;
-		}
 		// BGRA typed UAV is unsupported per D3D spec, use RGBA here.
 		FTexture2DRHIRef Texture2D = RHICreateTexture2D(1, 1, PF_R8G8B8A8, 1, 1, CreateFlags, CreateInfo);
 		TextureRHI = Texture2D;
@@ -176,10 +149,6 @@ public:
 
 		// Create a view of the texture
 		ShaderResourceViewRHI = RHICreateShaderResourceView(TextureRHI, 0u);
-		if(bWithUAV)
-		{
-			UnorderedAccessViewRHI = RHICreateUnorderedAccessView(TextureRHI, 0u);
-		}
 	}
 
 	/** Returns the width of the texture in pixels. */
@@ -232,7 +201,6 @@ FTextureWithSRV* GTransparentBlackTextureWithSRV = new TGlobalResource<FColoredT
 FTexture* GWhiteTexture = GWhiteTextureWithSRV;
 FTexture* GBlackTexture = GBlackTextureWithSRV;
 FTexture* GTransparentBlackTexture = GTransparentBlackTextureWithSRV;
-FTextureWithSRV* GBlackTextureWithUAV = new TGlobalResource<FColoredTexture<0,0,0,0,true> >;
 
 FVertexBufferWithSRV* GEmptyVertexBufferWithUAV = new TGlobalResource<FEmptyVertexBuffer>;
 
@@ -286,7 +254,7 @@ FBufferWithRDG* GWhiteVertexBufferWithRDG = new TGlobalResource<FWhiteVertexBuff
  * A class representing a 1x1x1 black volume texture.
  */
 template <EPixelFormat PixelFormat, uint8 Alpha>
-class FBlackVolumeTexture : public FTextureWithRDG
+class FBlackVolumeTexture : public FTexture
 {
 public:
 	
@@ -317,8 +285,6 @@ public:
 		// Create the sampler state.
 		FSamplerStateInitializerRHI SamplerStateInitializer(SF_Point, AM_Wrap, AM_Wrap, AM_Wrap);
 		SamplerStateRHI = GetOrCreateSamplerState(SamplerStateInitializer);
-
-		FTextureWithRDG::InitRDG(Name);
 	}
 
 	/**
@@ -339,11 +305,11 @@ public:
 };
 
 /** Global black volume texture resource. */
-FTextureWithRDG* GBlackVolumeTexture = new TGlobalResource<FBlackVolumeTexture<PF_B8G8R8A8, 0>>();
-FTextureWithRDG* GBlackAlpha1VolumeTexture = new TGlobalResource<FBlackVolumeTexture<PF_B8G8R8A8, 255>>();
+FTexture* GBlackVolumeTexture = new TGlobalResource<FBlackVolumeTexture<PF_B8G8R8A8, 0>>();
+FTexture* GBlackAlpha1VolumeTexture = new TGlobalResource<FBlackVolumeTexture<PF_B8G8R8A8, 255>>();
 
 /** Global black volume texture resource. */
-FTextureWithRDG* GBlackUintVolumeTexture = new TGlobalResource<FBlackVolumeTexture<PF_R8G8B8A8_UINT, 0>>();
+FTexture* GBlackUintVolumeTexture = new TGlobalResource<FBlackVolumeTexture<PF_R8G8B8A8_UINT, 0>>();
 
 class FBlackArrayTexture : public FTexture
 {
@@ -468,7 +434,7 @@ RENDERCORE_API int32 GMipColorTextureMipLevels = FMipColorTexture::NumMips;
 RENDERCORE_API const uint32 GDiffuseConvolveMipLevel = 4;
 
 /** A solid color cube texture. */
-class FSolidColorTextureCube : public FTextureWithRDG
+class FSolidColorTextureCube : public FTexture
 {
 public:
 	FSolidColorTextureCube(const FColor& InColor)
@@ -512,8 +478,6 @@ public:
 		// Create the sampler state RHI resource.
 		FSamplerStateInitializerRHI SamplerStateInitializer(SF_Point, AM_Wrap, AM_Wrap, AM_Wrap);
 		SamplerStateRHI = GetOrCreateSamplerState(SamplerStateInitializer);
-
-		FTextureWithRDG::InitRDG(Name);
 	}
 
 	/** Returns the width of the texture in pixels. */
@@ -540,7 +504,7 @@ class FWhiteTextureCube : public FSolidColorTextureCube
 public:
 	FWhiteTextureCube() : FSolidColorTextureCube(FColor::White) {}
 };
-FTextureWithRDG* GWhiteTextureCube = new TGlobalResource<FWhiteTextureCube>;
+FTexture* GWhiteTextureCube = new TGlobalResource<FWhiteTextureCube>;
 
 /** A black cube texture. */
 class FBlackTextureCube : public FSolidColorTextureCube
@@ -548,7 +512,7 @@ class FBlackTextureCube : public FSolidColorTextureCube
 public:
 	FBlackTextureCube() : FSolidColorTextureCube(FColor::Black) {}
 };
-FTextureWithRDG* GBlackTextureCube = new TGlobalResource<FBlackTextureCube>;
+FTexture* GBlackTextureCube = new TGlobalResource<FBlackTextureCube>;
 
 /** A black cube texture. */
 class FBlackTextureDepthCube : public FSolidColorTextureCube
@@ -556,9 +520,9 @@ class FBlackTextureDepthCube : public FSolidColorTextureCube
 public:
 	FBlackTextureDepthCube() : FSolidColorTextureCube(PF_ShadowDepth) {}
 };
-FTextureWithRDG* GBlackTextureDepthCube = new TGlobalResource<FBlackTextureDepthCube>;
+FTexture* GBlackTextureDepthCube = new TGlobalResource<FBlackTextureDepthCube>;
 
-class FBlackCubeArrayTexture : public FTextureWithRDG
+class FBlackCubeArrayTexture : public FTexture
 {
 public:
 	// FResource interface.
@@ -585,8 +549,6 @@ public:
 			// Create the sampler state RHI resource.
 			FSamplerStateInitializerRHI SamplerStateInitializer(SF_Point,AM_Wrap,AM_Wrap,AM_Wrap);
 			SamplerStateRHI = GetOrCreateSamplerState(SamplerStateInitializer);
-
-			FTextureWithRDG::InitRDG(Name);
 		}
 	}
 
@@ -602,7 +564,7 @@ public:
 		return 1;
 	}
 };
-FTextureWithRDG* GBlackCubeArrayTexture = new TGlobalResource<FBlackCubeArrayTexture>;
+FTexture* GBlackCubeArrayTexture = new TGlobalResource<FBlackCubeArrayTexture>;
 
 /**
  * A UINT 1x1 texture.
diff --git a/Engine/Source/Runtime/RenderCore/Public/RenderResource.h b/Engine/Source/Runtime/RenderCore/Public/RenderResource.h
--- a/Engine/Source/Runtime/RenderCore/Public/RenderResource.h
+++ b/Engine/Source/Runtime/RenderCore/Public/RenderResource.h
@@ -437,16 +437,11 @@ public:
 	/** SRV that views the entire texture */
 	FShaderResourceViewRHIRef ShaderResourceViewRHI;
 
-	/** *optional* UAV that views the entire texture */
-	FUnorderedAccessViewRHIRef UnorderedAccessViewRHI;
-
-
 	virtual ~FTextureWithSRV() {}
 
 	virtual void ReleaseRHI() override
 	{
 		ShaderResourceViewRHI.SafeRelease();
-		UnorderedAccessViewRHI.SafeRelease();
 		FTexture::ReleaseRHI();
 	}
 };
diff --git a/Engine/Source/Runtime/RenderCore/Public/RenderUtils.h b/Engine/Source/Runtime/RenderCore/Public/RenderUtils.h
--- a/Engine/Source/Runtime/RenderCore/Public/RenderUtils.h
+++ b/Engine/Source/Runtime/RenderCore/Public/RenderUtils.h
@@ -12,31 +12,10 @@
 class FBufferWithRDG;
 class FRDGBuffer;
 class FRDGPooledBuffer;
-class FTextureWithRDG;
 class FRDGTexture;
 class FRDGBuilder;
 struct IPooledRenderTarget;
 
-/** An FTexture variant that includes more efficient support for registering with RDG. */
-class RENDERCORE_API FTextureWithRDG : public FTexture
-{
-public:
-	FTextureWithRDG();
-	FTextureWithRDG(const FTextureWithRDG& Other);
-	FTextureWithRDG& operator=(const FTextureWithRDG& Other);
-	~FTextureWithRDG() override;
-
-	FRDGTexture* GetRDG(FRDGBuilder& GraphBuilder) const;
-
-	void ReleaseRHI() override;
-
-protected:
-	void InitRDG(const TCHAR* Name);
-
-private:
-	TRefCountPtr<IPooledRenderTarget> RenderTarget;
-};
-
 class RENDERCORE_API FBufferWithRDG : public FRenderResource
 {
 public:
@@ -144,7 +123,6 @@ extern RENDERCORE_API class FTextureWithSRV* GWhiteTextureWithSRV;
 /** A global black texture. */
 extern RENDERCORE_API class FTexture* GBlackTexture;
 extern RENDERCORE_API class FTextureWithSRV* GBlackTextureWithSRV;
-extern RENDERCORE_API class FTextureWithSRV* GBlackTextureWithUAV;
 
 extern RENDERCORE_API class FTexture* GTransparentBlackTexture;
 extern RENDERCORE_API class FTextureWithSRV* GTransparentBlackTextureWithSRV;
@@ -159,28 +137,28 @@ extern RENDERCORE_API class FBufferWithRDG* GWhiteVertexBufferWithRDG;
 extern RENDERCORE_API class FTexture* GBlackArrayTexture;
 
 /** A global black volume texture. */
-extern RENDERCORE_API class FTextureWithRDG* GBlackVolumeTexture;
+extern RENDERCORE_API class FTexture* GBlackVolumeTexture;
 
 /** A global black volume texture, with alpha=1. */
-extern RENDERCORE_API class FTextureWithRDG* GBlackAlpha1VolumeTexture;
+extern RENDERCORE_API class FTexture* GBlackAlpha1VolumeTexture;
 
 /** A global black texture<uint> */
 extern RENDERCORE_API class FTexture* GBlackUintTexture;
 
 /** A global black volume texture<uint>  */
-extern RENDERCORE_API class FTextureWithRDG* GBlackUintVolumeTexture;
+extern RENDERCORE_API class FTexture* GBlackUintVolumeTexture;
 
 /** A global white cube texture. */
-extern RENDERCORE_API class FTextureWithRDG* GWhiteTextureCube;
+extern RENDERCORE_API class FTexture* GWhiteTextureCube;
 
 /** A global black cube texture. */
-extern RENDERCORE_API class FTextureWithRDG* GBlackTextureCube;
+extern RENDERCORE_API class FTexture* GBlackTextureCube;
 
 /** A global black cube depth texture. */
-extern RENDERCORE_API class FTextureWithRDG* GBlackTextureDepthCube;
+extern RENDERCORE_API class FTexture* GBlackTextureDepthCube;
 
 /** A global black cube array texture. */
-extern RENDERCORE_API class FTextureWithRDG* GBlackCubeArrayTexture;
+extern RENDERCORE_API class FTexture* GBlackCubeArrayTexture;
 
 /** A global texture that has a different solid color in each mip-level. */
 extern RENDERCORE_API class FTexture* GMipColorTexture;
diff --git a/Engine/Source/Runtime/Renderer/Private/FogRendering.cpp b/Engine/Source/Runtime/Renderer/Private/FogRendering.cpp
--- a/Engine/Source/Runtime/Renderer/Private/FogRendering.cpp
+++ b/Engine/Source/Runtime/Renderer/Private/FogRendering.cpp
@@ -88,7 +88,8 @@ void SetupFogUniformParameters(FRDGBuilder& GraphBuilder, const FViewInfo& View,
 		}
 		else
 		{
-			OutParameters.IntegratedLightScattering = GBlackAlpha1VolumeTexture->GetRDG(GraphBuilder);
+			const FRDGSystemTextures& SystemTextures = FRDGSystemTextures::Get(GraphBuilder);
+			OutParameters.IntegratedLightScattering = SystemTextures.VolumetricBlackAlphaOne;
 			OutParameters.ApplyVolumetricFog = 0.0f;
 		}
 		OutParameters.IntegratedLightScatteringSampler = TStaticSamplerState<SF_Bilinear, AM_Clamp, AM_Clamp, AM_Clamp>::GetRHI();
diff --git a/Engine/Source/Runtime/Renderer/Private/ShadowRendering.cpp b/Engine/Source/Runtime/Renderer/Private/ShadowRendering.cpp
--- a/Engine/Source/Runtime/Renderer/Private/ShadowRendering.cpp
+++ b/Engine/Source/Runtime/Renderer/Private/ShadowRendering.cpp
@@ -233,7 +233,7 @@ void GetOnePassPointShadowProjectionParameters(FRDGBuilder& GraphBuilder, const
 	//@todo DynamicGI: remove duplication with FOnePassPointShadowProjectionShaderParameters
 	FRDGTexture* ShadowDepthTextureValue = ShadowInfo
 		? GraphBuilder.RegisterExternalTexture(ShadowInfo->RenderTargets.DepthTarget)
-		: GBlackTextureDepthCube->GetRDG(GraphBuilder);
+		: SystemTextures.BlackDepthCube;
 
 	OutParameters.ShadowDepthCubeTexture = ShadowDepthTextureValue;
 	OutParameters.ShadowDepthCubeTexture2 = ShadowDepthTextureValue;
diff --git a/Engine/Source/Runtime/Renderer/Private/SystemTextures.cpp b/Engine/Source/Runtime/Renderer/Private/SystemTextures.cpp
--- a/Engine/Source/Runtime/Renderer/Private/SystemTextures.cpp
+++ b/Engine/Source/Runtime/Renderer/Private/SystemTextures.cpp
@@ -7,6 +7,7 @@
 #include "SystemTextures.h"
 #include "Math/RandomStream.h"
 #include "Math/Sobol.h"
+#include "Math/Float16.h"
 #include "RenderTargetPool.h"
 #include "ClearQuad.h"
 #include "LTC.h"
@@ -30,12 +31,14 @@ const FRDGSystemTextures& FRDGSystemTextures::Create(FRDGBuilder& GraphBuilder)
 	SystemTextures.BlackAlphaOne = Register(GSystemTextures.BlackAlphaOneDummy);
 	SystemTextures.MaxFP16Depth = Register(GSystemTextures.MaxFP16Depth);
 	SystemTextures.DepthDummy = Register(GSystemTextures.DepthDummy);
+	SystemTextures.BlackDepthCube = Register(GSystemTextures.BlackDepthCube);
 	SystemTextures.StencilDummy = Register(GSystemTextures.StencilDummy);
 	SystemTextures.Green = Register(GSystemTextures.GreenDummy);
 	SystemTextures.DefaultNormal8Bit = Register(GSystemTextures.DefaultNormal8Bit);
 	SystemTextures.MidGrey = Register(GSystemTextures.MidGreyDummy);
 	SystemTextures.VolumetricBlack = Register(GSystemTextures.VolumetricBlackDummy);
-	SystemTextures.StencilDummySRV = GraphBuilder.CreateSRV(FRDGTextureSRVDesc::CreateWithPixelFormat(SystemTextures.DepthDummy, PF_X24_G8));
+	SystemTextures.VolumetricBlackAlphaOne = Register(GSystemTextures.VolumetricBlackAlphaOneDummy);
+	SystemTextures.StencilDummySRV = GraphBuilder.CreateSRV(FRDGTextureSRVDesc::Create(SystemTextures.DepthDummy));
 	return SystemTextures;
 }
 
@@ -74,142 +77,98 @@ void FSystemTextures::InitializeTextures(FRHICommandListImmediate& RHICmdList, c
 	// there's no needed setup for those feature levels lower or identical to the current one
 }
 
+template <typename DataType>
+void SetDummyTextureData(FRHITexture2D* Texture2D, const DataType& DummyData)
+{
+	uint32 DestStride;
+	DataType* DestBuffer = (DataType*)RHILockTexture2D(Texture2D, 0, RLM_WriteOnly, DestStride, false);
+	*DestBuffer = DummyData;
+	RHIUnlockTexture2D(Texture2D, 0, false);
+}
+
 void FSystemTextures::InitializeCommonTextures(FRHICommandListImmediate& RHICmdList)
 {
 	// First initialize textures that are common to all feature levels. This is always done the first time we come into this function, as doesn't care about the
 	// requested feature level
 
-		// Create a WhiteDummy texture
-		{
-			FPooledRenderTargetDesc Desc(FPooledRenderTargetDesc::Create2DDesc(FIntPoint(1, 1), PF_B8G8R8A8, FClearValueBinding::White, TexCreate_HideInVisualizeTexture, TexCreate_RenderTargetable | TexCreate_NoFastClear | TexCreate_ShaderResource, false));
-			Desc.AutoWritable = false;
-			GRenderTargetPool.FindFreeElement(RHICmdList, Desc, WhiteDummy, TEXT("WhiteDummy"), ERenderTargetTransience::NonTransient);
-
-			RHICmdList.Transition(FRHITransitionInfo(WhiteDummy->GetRenderTargetItem().TargetableTexture, ERHIAccess::SRVMask, ERHIAccess::RTV));
-
-			FRHIRenderPassInfo RPInfo(WhiteDummy->GetRenderTargetItem().TargetableTexture, ERenderTargetActions::Clear_Store);
-			RHICmdList.BeginRenderPass(RPInfo, TEXT("WhiteDummy"));
-			RHICmdList.EndRenderPass();
-			RHICmdList.CopyToResolveTarget(WhiteDummy->GetRenderTargetItem().TargetableTexture, WhiteDummy->GetRenderTargetItem().ShaderResourceTexture, FResolveParams());
-
-			WhiteDummySRV = RHICreateShaderResourceView((FRHITexture2D*)WhiteDummy->GetRenderTargetItem().ShaderResourceTexture.GetReference(), 0);
-		}
-
-		// Create a BlackDummy texture
-		{
-			FPooledRenderTargetDesc Desc(FPooledRenderTargetDesc::Create2DDesc(FIntPoint(1, 1), PF_B8G8R8A8, FClearValueBinding::Transparent, TexCreate_HideInVisualizeTexture, TexCreate_RenderTargetable | TexCreate_NoFastClear | TexCreate_ShaderResource, false));
-			Desc.AutoWritable = false;
-			GRenderTargetPool.FindFreeElement(RHICmdList, Desc, BlackDummy, TEXT("BlackDummy"), ERenderTargetTransience::NonTransient);
-
-			RHICmdList.Transition(FRHITransitionInfo(BlackDummy->GetRenderTargetItem().TargetableTexture, ERHIAccess::SRVMask, ERHIAccess::RTV));
+	// Create a WhiteDummy texture
+	{
+		WhiteDummy = CreateRenderTarget(GWhiteTexture->TextureRHI, TEXT("WhiteDummy"));
+		WhiteDummySRV = RHICreateShaderResourceView((FRHITexture2D*)WhiteDummy->GetShaderResourceRHI(), 0);
+	}
 
-			FRHIRenderPassInfo RPInfo(BlackDummy->GetRenderTargetItem().TargetableTexture, ERenderTargetActions::Clear_Store);
-			RHICmdList.BeginRenderPass(RPInfo, TEXT("BlackDummy"));
-			RHICmdList.EndRenderPass();
-			RHICmdList.CopyToResolveTarget(BlackDummy->GetRenderTargetItem().TargetableTexture, BlackDummy->GetRenderTargetItem().ShaderResourceTexture, FResolveParams());
-		}
+	// Create a BlackDummy texture
+	{
+		BlackDummy = CreateRenderTarget(GBlackTexture->TextureRHI, TEXT("BlackDummy"));
+	}
 	
-		// Create a texture that is a single UInt32 value set to 0
-		{
-			FPooledRenderTargetDesc Desc(FPooledRenderTargetDesc::Create2DDesc(FIntPoint(1,1), PF_R32_UINT, FClearValueBinding::Transparent, TexCreate_HideInVisualizeTexture, TexCreate_RenderTargetable | TexCreate_NoFastClear | TexCreate_ShaderResource, false));
-			Desc.AutoWritable = false;
-			GRenderTargetPool.FindFreeElement(RHICmdList, Desc, ZeroUIntDummy, TEXT("ZeroUIntDummy"), ERenderTargetTransience::NonTransient);
-
-			RHICmdList.Transition(FRHITransitionInfo(ZeroUIntDummy->GetRenderTargetItem().TargetableTexture, ERHIAccess::SRVMask, ERHIAccess::RTV));
-
-			FRHIRenderPassInfo RPInfo(ZeroUIntDummy->GetRenderTargetItem().TargetableTexture, ERenderTargetActions::Clear_Store);
-			RHICmdList.BeginRenderPass(RPInfo, TEXT("ClearZeroUIntDummy"));
-			RHICmdList.EndRenderPass();
-			RHICmdList.CopyToResolveTarget(ZeroUIntDummy->GetRenderTargetItem().TargetableTexture, ZeroUIntDummy->GetRenderTargetItem().ShaderResourceTexture, FResolveParams());
-		}
-
-		// Create a texture that is a single 4xUInt16 (UShort) value set to 0
-		{
-			FPooledRenderTargetDesc Desc(FPooledRenderTargetDesc::Create2DDesc(FIntPoint(1, 1), PF_R16G16B16A16_UINT, FClearValueBinding::Transparent, TexCreate_HideInVisualizeTexture, TexCreate_RenderTargetable | TexCreate_NoFastClear | TexCreate_ShaderResource, false));
-			Desc.AutoWritable = false;
-			GRenderTargetPool.FindFreeElement(RHICmdList, Desc, ZeroUShort4Dummy, TEXT("ZeroUShort4Dummy"), ERenderTargetTransience::NonTransient);
-
-			RHICmdList.Transition(FRHITransitionInfo(ZeroUShort4Dummy->GetRenderTargetItem().TargetableTexture, ERHIAccess::SRVMask, ERHIAccess::RTV));
-
-			FRHIRenderPassInfo RPInfo(ZeroUShort4Dummy->GetRenderTargetItem().TargetableTexture, ERenderTargetActions::Clear_Store);
-			RHICmdList.BeginRenderPass(RPInfo, TEXT("ClearZeroUShort4Dummy"));
-			RHICmdList.EndRenderPass();
-			RHICmdList.CopyToResolveTarget(ZeroUShort4Dummy->GetRenderTargetItem().TargetableTexture, ZeroUShort4Dummy->GetRenderTargetItem().ShaderResourceTexture, FResolveParams());
-		}
-
-		// Create a BlackAlphaOneDummy texture
-		{
-			FPooledRenderTargetDesc Desc(FPooledRenderTargetDesc::Create2DDesc(FIntPoint(1, 1), PF_B8G8R8A8, FClearValueBinding::Black, TexCreate_HideInVisualizeTexture, TexCreate_RenderTargetable | TexCreate_NoFastClear | TexCreate_ShaderResource, false));
-			Desc.AutoWritable = false;
-			GRenderTargetPool.FindFreeElement(RHICmdList, Desc, BlackAlphaOneDummy, TEXT("BlackAlphaOneDummy"), ERenderTargetTransience::NonTransient);
-
-			RHICmdList.Transition(FRHITransitionInfo(BlackAlphaOneDummy->GetRenderTargetItem().TargetableTexture, ERHIAccess::SRVMask, ERHIAccess::RTV));
+	// Create a texture that is a single UInt32 value set to 0
+	{	
+		FRHIResourceCreateInfo CreateInfo(TEXT("ZeroUIntDummy"));
+		FTexture2DRHIRef Texture2D = RHICreateTexture2D(1, 1, PF_R32_UINT, 1, 1, TexCreate_ShaderResource, CreateInfo);
+		SetDummyTextureData<uint32>(Texture2D, 0u);
+		ZeroUIntDummy = CreateRenderTarget(Texture2D, CreateInfo.DebugName);
+	}
 
-			FRHIRenderPassInfo RPInfo(BlackAlphaOneDummy->GetRenderTargetItem().TargetableTexture, ERenderTargetActions::Clear_Store);
-			RHICmdList.BeginRenderPass(RPInfo, TEXT("BlackAlphaOneDummy"));
-			RHICmdList.EndRenderPass();
-			RHICmdList.CopyToResolveTarget(BlackAlphaOneDummy->GetRenderTargetItem().TargetableTexture, BlackAlphaOneDummy->GetRenderTargetItem().ShaderResourceTexture, FResolveParams());
-		}
+	// Create a BlackAlphaOneDummy texture
+	{
+		FRHIResourceCreateInfo CreateInfo(TEXT("BlackAlphaOneDummy"));
+		FTexture2DRHIRef Texture2D = RHICreateTexture2D(1, 1, PF_B8G8R8A8, 1, 1, TexCreate_ShaderResource, CreateInfo);
+		SetDummyTextureData<FColor>(Texture2D, FColor(0, 0, 0, 255));
+		BlackAlphaOneDummy = CreateRenderTarget(Texture2D, CreateInfo.DebugName);
+	}
 
-		// Create a GreenDummy texture
-		{
-			FPooledRenderTargetDesc Desc(FPooledRenderTargetDesc::Create2DDesc(FIntPoint(1, 1), PF_B8G8R8A8, FClearValueBinding::Green, TexCreate_HideInVisualizeTexture, TexCreate_RenderTargetable | TexCreate_NoFastClear | TexCreate_ShaderResource, false));
-			Desc.AutoWritable = false;
-			GRenderTargetPool.FindFreeElement(RHICmdList, Desc, GreenDummy, TEXT("GreenDummy"), ERenderTargetTransience::NonTransient);
+	// Create a GreenDummy texture
+	{
+		FRHIResourceCreateInfo CreateInfo(TEXT("GreenDummy"));
+		FTexture2DRHIRef Texture2D = RHICreateTexture2D(1, 1, PF_B8G8R8A8, 1, 1, TexCreate_ShaderResource, CreateInfo);
+		SetDummyTextureData<FColor>(Texture2D, FColor(0, 255, 0, 255));
+		GreenDummy = CreateRenderTarget(Texture2D, CreateInfo.DebugName);
+	}
 
-			RHICmdList.Transition(FRHITransitionInfo(GreenDummy->GetRenderTargetItem().TargetableTexture, ERHIAccess::SRVMask, ERHIAccess::RTV));
+	// Create a DefaultNormal8Bit texture
+	{
+		FRHIResourceCreateInfo CreateInfo(TEXT("DefaultNormal8Bit"));
+		FTexture2DRHIRef Texture2D = RHICreateTexture2D(1, 1, PF_B8G8R8A8, 1, 1, TexCreate_ShaderResource, CreateInfo);
+		SetDummyTextureData<FColor>(Texture2D, FColor(128, 128, 128, 255));
+		DefaultNormal8Bit = CreateRenderTarget(Texture2D, CreateInfo.DebugName);
+	}
 
-			FRHIRenderPassInfo RPInfo(GreenDummy->GetRenderTargetItem().TargetableTexture, ERenderTargetActions::Clear_Store);
-			RHICmdList.BeginRenderPass(RPInfo, TEXT("GreenDummy"));
-			RHICmdList.EndRenderPass();
-			RHICmdList.CopyToResolveTarget(GreenDummy->GetRenderTargetItem().TargetableTexture, GreenDummy->GetRenderTargetItem().ShaderResourceTexture, FResolveParams());
-		}
+	// Create the PerlinNoiseGradient texture
+	{
+		const uint32 Extent = 128;
+	
+		FRHIResourceCreateInfo CreateInfo(TEXT("PerlinNoiseGradient"));
+		FTexture2DRHIRef Texture2D = RHICreateTexture2D(Extent, Extent, PF_B8G8R8A8, 1, 1, TexCreate_ShaderResource, CreateInfo);
 
-		// Create a DefaultNormal8Bit texture
+		// Write the contents of the texture.
+		uint32 DestStride;
+		uint8* DestBuffer = (uint8*)RHICmdList.LockTexture2D(Texture2D, 0, RLM_WriteOnly, DestStride, false);
+		// seed the pseudo random stream with a good value
+		FRandomStream RandomStream(12345);
+		// Values represent float3 values in the -1..1 range.
+		// The vectors are the edge mid point of a cube from -1 .. 1
+		static uint32 gradtable[] =
 		{
-			FPooledRenderTargetDesc Desc(FPooledRenderTargetDesc::Create2DDesc(FIntPoint(1, 1), PF_B8G8R8A8, FClearValueBinding::DefaultNormal8Bit, TexCreate_HideInVisualizeTexture, TexCreate_RenderTargetable | TexCreate_NoFastClear | TexCreate_ShaderResource, false));
-			Desc.AutoWritable = false;
-			GRenderTargetPool.FindFreeElement(RHICmdList, Desc, DefaultNormal8Bit, TEXT("DefaultNormal8Bit"), ERenderTargetTransience::NonTransient);
-
-			RHICmdList.Transition(FRHITransitionInfo(DefaultNormal8Bit->GetRenderTargetItem().TargetableTexture, ERHIAccess::SRVMask, ERHIAccess::RTV));
-
-			FRHIRenderPassInfo RPInfo(DefaultNormal8Bit->GetRenderTargetItem().TargetableTexture, ERenderTargetActions::Clear_Store);
-			RHICmdList.BeginRenderPass(RPInfo, TEXT("DefaultNormal8Bit"));
-			RHICmdList.EndRenderPass();
-			RHICmdList.CopyToResolveTarget(DefaultNormal8Bit->GetRenderTargetItem().TargetableTexture, DefaultNormal8Bit->GetRenderTargetItem().ShaderResourceTexture, FResolveParams());
-		}
-
-		// Create the PerlinNoiseGradient texture
+			0x88ffff, 0xff88ff, 0xffff88,
+			0x88ff00, 0xff8800, 0xff0088,
+			0x8800ff, 0x0088ff, 0x00ff88,
+			0x880000, 0x008800, 0x000088,
+		};
+		for (int32 y = 0; y < Extent; ++y)
 		{
-			FPooledRenderTargetDesc Desc(FPooledRenderTargetDesc::Create2DDesc(FIntPoint(128, 128), PF_B8G8R8A8, FClearValueBinding::None, TexCreate_HideInVisualizeTexture, TexCreate_None | TexCreate_NoFastClear | TexCreate_ShaderResource, false));
-			Desc.AutoWritable = false;
-			GRenderTargetPool.FindFreeElement(RHICmdList, Desc, PerlinNoiseGradient, TEXT("PerlinNoiseGradient"), ERenderTargetTransience::NonTransient);
-			// Write the contents of the texture.
-			uint32 DestStride;
-			uint8* DestBuffer = (uint8*)RHICmdList.LockTexture2D((FTexture2DRHIRef&)PerlinNoiseGradient->GetRenderTargetItem().ShaderResourceTexture, 0, RLM_WriteOnly, DestStride, false);
-			// seed the pseudo random stream with a good value
-			FRandomStream RandomStream(12345);
-			// Values represent float3 values in the -1..1 range.
-			// The vectors are the edge mid point of a cube from -1 .. 1
-			static uint32 gradtable[] =
-			{
-				0x88ffff, 0xff88ff, 0xffff88,
-				0x88ff00, 0xff8800, 0xff0088,
-				0x8800ff, 0x0088ff, 0x00ff88,
-				0x880000, 0x008800, 0x000088,
-			};
-			for (int32 y = 0; y < Desc.Extent.Y; ++y)
+			for (int32 x = 0; x < Extent; ++x)
 			{
-				for (int32 x = 0; x < Desc.Extent.X; ++x)
-				{
 				uint32* Dest = (uint32*)(DestBuffer + x * sizeof(uint32) + y * DestStride);
 
-					// pick a random direction (hacky way to overcome the quality issues FRandomStream has)
-					*Dest = gradtable[(uint32)(RandomStream.GetFraction() * 11.9999999f)];
-				}
+				// pick a random direction (hacky way to overcome the quality issues FRandomStream has)
+				*Dest = gradtable[(uint32)(RandomStream.GetFraction() * 11.9999999f)];
 			}
-			RHICmdList.UnlockTexture2D((FTexture2DRHIRef&)PerlinNoiseGradient->GetRenderTargetItem().ShaderResourceTexture, 0, false);
 		}
+		RHICmdList.UnlockTexture2D(Texture2D, 0, false);
+
+		PerlinNoiseGradient = CreateRenderTarget(Texture2D, CreateInfo.DebugName);
+	}
 
 	if (GPixelFormats[PF_FloatRGBA].Supported)
 	{
@@ -225,49 +184,27 @@ void FSystemTextures::InitializeCommonTextures(FRHICommandListImmediate& RHICmdL
 		RHICmdList.CopyToResolveTarget(MaxFP16Depth->GetRenderTargetItem().TargetableTexture, MaxFP16Depth->GetRenderTargetItem().ShaderResourceTexture, FResolveParams());
 	}
 
-	// Create dummy 1x1 depth texture		
 	{
-		FPooledRenderTargetDesc Desc(FPooledRenderTargetDesc::Create2DDesc(FIntPoint(1, 1), PF_DepthStencil, FClearValueBinding::DepthFar, TexCreate_None, TexCreate_DepthStencilTargetable, false));
-		Desc.AutoWritable = false;
-		GRenderTargetPool.FindFreeElement(RHICmdList, Desc, DepthDummy, TEXT("DepthDummy"), ERenderTargetTransience::NonTransient);
+		DepthDummy = static_cast<bool>(ERHIZBuffer::IsInverted) ? BlackDummy : WhiteDummy;
 
-		RHICmdList.Transition(FRHITransitionInfo(DepthDummy->GetRenderTargetItem().TargetableTexture, ERHIAccess::SRVMask, ERHIAccess::DSVWrite));
-
-		FRHIRenderPassInfo RPInfo(DepthDummy->GetRenderTargetItem().TargetableTexture, EDepthStencilTargetActions::ClearDepthStencil_StoreDepthStencil, nullptr, FExclusiveDepthStencil::DepthWrite_StencilWrite);
-		RHICmdList.BeginRenderPass(RPInfo, TEXT("DepthDummy"));
-		RHICmdList.EndRenderPass();
-		RHICmdList.CopyToResolveTarget(DepthDummy->GetRenderTargetItem().TargetableTexture, DepthDummy->GetRenderTargetItem().ShaderResourceTexture, FResolveParams());
+		BlackDepthCube = CreateRenderTarget(GBlackTextureDepthCube->TextureRHI, TEXT("BlackDepthCube"));
 	}
 
 	// Create a dummy stencil SRV.
 	{
-		FPooledRenderTargetDesc Desc(FPooledRenderTargetDesc::Create2DDesc(FIntPoint(1, 1), PF_R8G8B8A8_UINT, FClearValueBinding::White, TexCreate_HideInVisualizeTexture, TexCreate_RenderTargetable | TexCreate_NoFastClear, false));
-		Desc.AutoWritable = false;
-		GRenderTargetPool.FindFreeElement(RHICmdList, Desc, StencilDummy, TEXT("StencilDummy"), ERenderTargetTransience::NonTransient);
-
-		RHICmdList.Transition(FRHITransitionInfo(StencilDummy->GetRenderTargetItem().TargetableTexture, ERHIAccess::SRVMask, ERHIAccess::RTV));
-
-		FRHIRenderPassInfo RPInfo(StencilDummy->GetRenderTargetItem().TargetableTexture, ERenderTargetActions::Clear_Store);
-		RHICmdList.BeginRenderPass(RPInfo, TEXT("StencilDummy"));
-		RHICmdList.EndRenderPass();
-		RHICmdList.CopyToResolveTarget(StencilDummy->GetRenderTargetItem().TargetableTexture, StencilDummy->GetRenderTargetItem().ShaderResourceTexture, FResolveParams());
-
-		StencilDummySRV = RHICreateShaderResourceView((FRHITexture2D*)StencilDummy->GetRenderTargetItem().ShaderResourceTexture.GetReference(), 0);
+		FRHIResourceCreateInfo CreateInfo(TEXT("StencilDummy"));
+		FTexture2DRHIRef Texture2D = RHICreateTexture2D(1, 1, PF_R8G8B8A8_UINT, 1, 1, TexCreate_ShaderResource, CreateInfo);
+		SetDummyTextureData<FColor>(Texture2D, FColor::White);
+		StencilDummy = CreateRenderTarget(Texture2D, CreateInfo.DebugName);
+		StencilDummySRV = RHICreateShaderResourceView((FRHITexture2D*)StencilDummy->GetShaderResourceRHI(), 0);
 	}
 
 	if (GPixelFormats[PF_FloatRGBA].Supported)
 	{
-		// PF_FloatRGBA to encode exactly the 0.5.
-		FPooledRenderTargetDesc Desc(FPooledRenderTargetDesc::Create2DDesc(FIntPoint(1, 1), PF_FloatRGBA, FClearValueBinding(FLinearColor(0.5f, 0.5f, 0.5f, 0.5f)), TexCreate_HideInVisualizeTexture, TexCreate_RenderTargetable | TexCreate_NoFastClear | TexCreate_ShaderResource, false));
-		Desc.AutoWritable = false;
-		GRenderTargetPool.FindFreeElement(RHICmdList, Desc, MidGreyDummy, TEXT("MidGreyDummy"), ERenderTargetTransience::NonTransient);
-
-		RHICmdList.Transition(FRHITransitionInfo(MidGreyDummy->GetRenderTargetItem().TargetableTexture, ERHIAccess::SRVMask, ERHIAccess::RTV));
-
-		FRHIRenderPassInfo RPInfo(MidGreyDummy->GetRenderTargetItem().TargetableTexture, ERenderTargetActions::Clear_Store);
-		RHICmdList.BeginRenderPass(RPInfo, TEXT("MidGreyDummy"));
-		RHICmdList.EndRenderPass();
-		RHICmdList.CopyToResolveTarget(MidGreyDummy->GetRenderTargetItem().TargetableTexture, MidGreyDummy->GetRenderTargetItem().ShaderResourceTexture, FResolveParams());
+		FRHIResourceCreateInfo CreateInfo(TEXT("MidGreyDummy"));
+		FTexture2DRHIRef Texture2D = RHICreateTexture2D(1, 1, PF_FloatRGBA, 1, 1, TexCreate_ShaderResource, CreateInfo);
+		SetDummyTextureData<FFloat16Color>(Texture2D, FFloat16Color(FLinearColor(0.5f, 0.5f, 0.5f, 0.5f)));
+		MidGreyDummy = CreateRenderTarget(Texture2D, CreateInfo.DebugName);
 	}
 }
 
@@ -282,12 +219,11 @@ void FSystemTextures::InitializeFeatureLevelDependentTextures(FRHICommandListImm
 		// Create the SobolSampling texture
 	if (CurrentFeatureLevel < ERHIFeatureLevel::ES3_1 && InFeatureLevel >= ERHIFeatureLevel::ES3_1 && GPixelFormats[PF_R16_UINT].Supported)
 	{
-		FPooledRenderTargetDesc Desc(FPooledRenderTargetDesc::Create2DDesc(FIntPoint(32, 16), PF_R16_UINT, FClearValueBinding::None, TexCreate_HideInVisualizeTexture, TexCreate_NoFastClear | TexCreate_ShaderResource, false));
-		Desc.AutoWritable = false;
-		GRenderTargetPool.FindFreeElement(RHICmdList, Desc, SobolSampling, TEXT("SobolSampling"));
+		FRHIResourceCreateInfo CreateInfo(TEXT("SobolSampling"));
+		FTexture2DRHIRef Texture2D = RHICreateTexture2D(32, 16, PF_R16_UINT, 1, 1, TexCreate_ShaderResource, CreateInfo);
 		// Write the contents of the texture.
 		uint32 DestStride;
-		uint8* DestBuffer = (uint8*)RHICmdList.LockTexture2D((FTexture2DRHIRef&)SobolSampling->GetRenderTargetItem().ShaderResourceTexture, 0, RLM_WriteOnly, DestStride, false);
+		uint8* DestBuffer = (uint8*)RHICmdList.LockTexture2D(Texture2D, 0, RLM_WriteOnly, DestStride, false);
 
 		uint16 *Dest;
 		for (int y = 0; y < 16; ++y)
@@ -306,28 +242,16 @@ void FSystemTextures::InitializeFeatureLevelDependentTextures(FRHICommandListImm
 				*Dest = FSobol::ComputeGPUSpatialSeed(x, y, /* Index = */ 1);
 			}
 		}
-		RHICmdList.UnlockTexture2D((FTexture2DRHIRef&)SobolSampling->GetRenderTargetItem().ShaderResourceTexture, 0, false);
+		RHICmdList.UnlockTexture2D(Texture2D, 0, false);
+
+		SobolSampling = CreateRenderTarget(Texture2D, CreateInfo.DebugName);
 	}
 
 	// Create a VolumetricBlackDummy texture
 	if (CurrentFeatureLevel < ERHIFeatureLevel::SM5 && InFeatureLevel >= ERHIFeatureLevel::SM5)
 	{
-		FPooledRenderTargetDesc Desc(FPooledRenderTargetDesc::CreateVolumeDesc(1, 1, 1, PF_B8G8R8A8, FClearValueBinding::Transparent, TexCreate_HideInVisualizeTexture, TexCreate_ShaderResource | TexCreate_RenderTargetable | TexCreate_NoFastClear, false));
-		Desc.AutoWritable = false;
-		GRenderTargetPool.FindFreeElement(RHICmdList, Desc, VolumetricBlackDummy, TEXT("VolumetricBlackDummy"), ERenderTargetTransience::NonTransient);
-
-		const uint8 BlackBytes[4] = { 0, 0, 0, 0 };
-		FUpdateTextureRegion3D Region(0, 0, 0, 0, 0, 0, Desc.Extent.X, Desc.Extent.Y, Desc.Depth);
-		RHICmdList.UpdateTexture3D(
-			(FTexture3DRHIRef&)VolumetricBlackDummy->GetRenderTargetItem().ShaderResourceTexture,
-			0,
-			Region,
-			Desc.Extent.X * sizeof(BlackBytes),
-			Desc.Extent.X * Desc.Extent.Y * sizeof(BlackBytes),
-			BlackBytes);
-
-		// UpdateTexture3D before and after state is currently undefined
-		RHICmdList.Transition(FRHITransitionInfo(VolumetricBlackDummy->GetTargetableRHI(), ERHIAccess::Unknown, ERHIAccess::SRVMask));
+		VolumetricBlackDummy = CreateRenderTarget(GBlackVolumeTexture->TextureRHI, TEXT("VolumetricBlackDummy"));
+		VolumetricBlackAlphaOneDummy = CreateRenderTarget(GBlackAlpha1VolumeTexture->TextureRHI, TEXT("VolumetricBlackAlphaOneDummy"));
 	}
 
 	if (CurrentFeatureLevel < ERHIFeatureLevel::SM5 && InFeatureLevel >= ERHIFeatureLevel::SM5)
@@ -359,29 +283,31 @@ void FSystemTextures::InitializeFeatureLevelDependentTextures(FRHICommandListImm
 			Format = PF_G16R16;
 		}
 
-		FPooledRenderTargetDesc Desc(FPooledRenderTargetDesc::Create2DDesc(FIntPoint(128, 32), Format, FClearValueBinding::None, TexCreate_None, TexCreate_ShaderResource, false));
-		Desc.AutoWritable = false;
+		FIntPoint Extent(128, 32);
+
 		if (bReference)
 		{
-			Desc.Extent.X = 128;
-			Desc.Extent.Y = 128;
+			Extent.X = 128;
+			Extent.Y = 128;
 		}
 
-		GRenderTargetPool.FindFreeElement(RHICmdList, Desc, PreintegratedGF, TEXT("PreintegratedGF"), ERenderTargetTransience::NonTransient);
+		FRHIResourceCreateInfo CreateInfo(TEXT("PreintegratedGF"));
+		FTexture2DRHIRef Texture2D = RHICreateTexture2D(Extent.X, Extent.Y, Format, 1, 1, TexCreate_ShaderResource, CreateInfo);
+
 		// Write the contents of the texture.
 		uint32 DestStride;
-		uint8* DestBuffer = (uint8*)RHICmdList.LockTexture2D((FTexture2DRHIRef&)PreintegratedGF->GetRenderTargetItem().ShaderResourceTexture, 0, RLM_WriteOnly, DestStride, false);
+		uint8* DestBuffer = (uint8*)RHICmdList.LockTexture2D(Texture2D, 0, RLM_WriteOnly, DestStride, false);
 
 		// x is NoV, y is roughness
-		for (int32 y = 0; y < Desc.Extent.Y; y++)
+		for (int32 y = 0; y < Extent.Y; y++)
 		{
-			float Roughness = (float)(y + 0.5f) / Desc.Extent.Y;
+			float Roughness = (float)(y + 0.5f) / Extent.Y;
 			float m = Roughness * Roughness;
 			float m2 = m * m;
 
-			for (int32 x = 0; x < Desc.Extent.X; x++)
+			for (int32 x = 0; x < Extent.X; x++)
 			{
-				float NoV = (float)(x + 0.5f) / Desc.Extent.X;
+				float NoV = (float)(x + 0.5f) / Extent.X;
 
 				FVector3f V;
 				V.X = FMath::Sqrt(1.0f - NoV * NoV);	// sin
@@ -450,14 +376,14 @@ void FSystemTextures::InitializeFeatureLevelDependentTextures(FRHICommandListImm
 				B /= NumSamples;
 				C /= NumSamples;
 
-				if (Desc.Format == PF_A16B16G16R16)
+				if (Format == PF_A16B16G16R16)
 				{
 					uint16* Dest = (uint16*)(DestBuffer + x * 8 + y * DestStride);
 					Dest[0] = (int32)(FMath::Clamp(A, 0.0f, 1.0f) * 65535.0f + 0.5f);
 					Dest[1] = (int32)(FMath::Clamp(B, 0.0f, 1.0f) * 65535.0f + 0.5f);
 					Dest[2] = (int32)(FMath::Clamp(C, 0.0f, 1.0f) * 65535.0f + 0.5f);
 				}
-				else if (Desc.Format == PF_G16R16)
+				else if (Format == PF_G16R16)
 				{
 					uint16* Dest = (uint16*)(DestBuffer + x * 4 + y * DestStride);
 					Dest[0] = (int32)(FMath::Clamp(A, 0.0f, 1.0f) * 65535.0f + 0.5f);
@@ -465,7 +391,7 @@ void FSystemTextures::InitializeFeatureLevelDependentTextures(FRHICommandListImm
 				}
 				else
 				{
-					check(Desc.Format == PF_R8G8);
+					check(Format == PF_R8G8);
 
 					uint8* Dest = (uint8*)(DestBuffer + x * 2 + y * DestStride);
 					Dest[0] = (int32)(FMath::Clamp(A, 0.0f, 1.0f) * 255.f + 0.5f);
@@ -473,20 +399,22 @@ void FSystemTextures::InitializeFeatureLevelDependentTextures(FRHICommandListImm
 				}
 			}
 		}
-		RHICmdList.UnlockTexture2D((FTexture2DRHIRef&)PreintegratedGF->GetRenderTargetItem().ShaderResourceTexture, 0, false);
+		RHICmdList.UnlockTexture2D(Texture2D, 0, false);
+
+		PreintegratedGF = CreateRenderTarget(Texture2D, CreateInfo.DebugName);
 	}
 
 	if (CurrentFeatureLevel < ERHIFeatureLevel::SM5 && InFeatureLevel >= ERHIFeatureLevel::SM5)
 	{
 	    // Create the PerlinNoise3D texture (similar to http://prettyprocs.wordpress.com/2012/10/20/fast-perlin-noise/)
 	    {
-		    uint32 Extent = 16;
+		    const uint32 Extent = 16;
     
 		    const uint32 Square = Extent * Extent;
-    
-		    FPooledRenderTargetDesc Desc(FPooledRenderTargetDesc::CreateVolumeDesc(Extent, Extent, Extent, PF_B8G8R8A8, FClearValueBinding::None, TexCreate_ShaderResource | TexCreate_HideInVisualizeTexture | TexCreate_NoTiling, TexCreate_ShaderResource, false));
-		    Desc.AutoWritable = false;
-		    GRenderTargetPool.FindFreeElement(RHICmdList, Desc, PerlinNoise3D, TEXT("PerlinNoise3D"), ERenderTargetTransience::NonTransient);
+
+			FRHIResourceCreateInfo CreateInfo(TEXT("PerlinNoise3D"));
+			FTexture3DRHIRef Texture3D = RHICreateTexture3D(Extent, Extent, Extent, PF_B8G8R8A8, 1, TexCreate_ShaderResource | TexCreate_NoTiling, CreateInfo);
+
 		    // Write the contents of the texture.
 		    TArray<uint32> DestBuffer;
     
@@ -550,11 +478,11 @@ void FSystemTextures::InitializeFeatureLevelDependentTextures(FRHICommandListImm
 			{
 			    uint32* Dest = DestBuffer.GetData();
     
-				for (uint32 z = 0; z < Desc.Depth; ++z)
+				for (uint32 z = 0; z < Extent; ++z)
 			    {
-					for (uint32 y = 0; y < (uint32)Desc.Extent.Y; ++y)
+					for (uint32 y = 0; y < Extent; ++y)
 				    {
-						for (uint32 x = 0; x < (uint32)Desc.Extent.X; ++x)
+						for (uint32 x = 0; x < Extent; ++x)
 					    {
 						    uint32 Value = *Dest;
     
@@ -578,65 +506,73 @@ void FSystemTextures::InitializeFeatureLevelDependentTextures(FRHICommandListImm
 			    }
 		    }
     
-		    FUpdateTextureRegion3D Region(0, 0, 0, 0, 0, 0, Desc.Extent.X, Desc.Extent.Y, Desc.Depth);
+		    FUpdateTextureRegion3D Region(0, 0, 0, 0, 0, 0, Extent, Extent, Extent);
     
 		    RHICmdList.UpdateTexture3D(
-			    (FTexture3DRHIRef&)PerlinNoise3D->GetRenderTargetItem().ShaderResourceTexture,
+			    Texture3D,
 			    0,
 			    Region,
-			    Desc.Extent.X * sizeof(uint32),
-			    Desc.Extent.X * Desc.Extent.Y * sizeof(uint32),
+			    Extent * sizeof(uint32),
+			    Extent * Extent * sizeof(uint32),
 			    (const uint8*)DestBuffer.GetData());
+
+			PerlinNoise3D = CreateRenderTarget(Texture3D, CreateInfo.DebugName);
+
 		} // end Create the PerlinNoise3D texture
 
 		// GTAO Randomization texture	
 		{
+			{
+				FRHIResourceCreateInfo CreateInfo(TEXT("LTCMat"));
+				FTexture2DRHIRef Texture2D = RHICreateTexture2D(LTC_Size, LTC_Size, PF_FloatRGBA, 1, 1, TexCreate_FastVRAM | TexCreate_ShaderResource, CreateInfo);
 
-		    {
-			    FPooledRenderTargetDesc Desc(FPooledRenderTargetDesc::Create2DDesc(FIntPoint(LTC_Size, LTC_Size), PF_FloatRGBA, FClearValueBinding::None, TexCreate_FastVRAM, TexCreate_ShaderResource, false));
-			    Desc.AutoWritable = false;
-    
-			    GRenderTargetPool.FindFreeElement(RHICmdList, Desc, LTCMat, TEXT("LTCMat"));
-			    // Write the contents of the texture.
-			    uint32 DestStride;
-			    uint8* DestBuffer = (uint8*)RHICmdList.LockTexture2D((FTexture2DRHIRef&)LTCMat->GetRenderTargetItem().ShaderResourceTexture, 0, RLM_WriteOnly, DestStride, false);
-    
-				    for (int32 y = 0; y < Desc.Extent.Y; ++y)
-			    {
-					    for (int32 x = 0; x < Desc.Extent.X; ++x)
-				    {
-					    uint16* Dest = (uint16*)(DestBuffer + x * 4 * sizeof(uint16) + y * DestStride);
-    
-						    for (int k = 0; k < 4; k++)
-							    Dest[k] = FFloat16(LTC_Mat[4 * (x + y * LTC_Size) + k]).Encoded;
-				    }
-			    }
-			    RHICmdList.UnlockTexture2D((FTexture2DRHIRef&)LTCMat->GetRenderTargetItem().ShaderResourceTexture, 0, false);
-		    }
-    
-		    {
-			    FPooledRenderTargetDesc Desc(FPooledRenderTargetDesc::Create2DDesc(FIntPoint(LTC_Size, LTC_Size), PF_G16R16F, FClearValueBinding::None, TexCreate_FastVRAM, TexCreate_ShaderResource, false));
-			    Desc.AutoWritable = false;
-    
-			    GRenderTargetPool.FindFreeElement(RHICmdList, Desc, LTCAmp, TEXT("LTCAmp"));
-			    // Write the contents of the texture.
-			    uint32 DestStride;
-			    uint8* DestBuffer = (uint8*)RHICmdList.LockTexture2D((FTexture2DRHIRef&)LTCAmp->GetRenderTargetItem().ShaderResourceTexture, 0, RLM_WriteOnly, DestStride, false);
+				// Write the contents of the texture.
+				uint32 DestStride;
+				uint8* DestBuffer = (uint8*)RHICmdList.LockTexture2D(Texture2D, 0, RLM_WriteOnly, DestStride, false);
     
-				for (int32 y = 0; y < Desc.Extent.Y; ++y)
-			    {
-					for (int32 x = 0; x < Desc.Extent.X; ++x)
-				    {
-					    uint16* Dest = (uint16*)(DestBuffer + x * 2 * sizeof(uint16) + y * DestStride);
+				for (int32 y = 0; y < LTC_Size; ++y)
+				{
+					for (int32 x = 0; x < LTC_Size; ++x)
+					{
+						uint16* Dest = (uint16*)(DestBuffer + x * 4 * sizeof(uint16) + y * DestStride);
     
+						for (int k = 0; k < 4; k++)
+						{
+							Dest[k] = FFloat16(LTC_Mat[4 * (x + y * LTC_Size) + k]).Encoded;
+						}
+					}
+				}
+				RHICmdList.UnlockTexture2D(Texture2D, 0, false);
+
+				LTCMat = CreateRenderTarget(Texture2D, CreateInfo.DebugName);
+			}
+
+			{
+				FRHIResourceCreateInfo CreateInfo(TEXT("LTCAmp"));
+				FTexture2DRHIRef Texture2D = RHICreateTexture2D(LTC_Size, LTC_Size, PF_G16R16F, 1, 1, TexCreate_FastVRAM | TexCreate_ShaderResource, CreateInfo);
+
+				// Write the contents of the texture.
+				uint32 DestStride;
+				uint8* DestBuffer = (uint8*)RHICmdList.LockTexture2D(Texture2D, 0, RLM_WriteOnly, DestStride, false);
+
+				for (int32 y = 0; y < LTC_Size; ++y)
+				{
+					for (int32 x = 0; x < LTC_Size; ++x)
+					{
+						uint16* Dest = (uint16*)(DestBuffer + x * 2 * sizeof(uint16) + y * DestStride);
+
 						for (int k = 0; k < 2; k++)
-						    Dest[k] = FFloat16(LTC_Amp[4 * (x + y * LTC_Size) + k]).Encoded;
-				    }
-			    }
-			    RHICmdList.UnlockTexture2D((FTexture2DRHIRef&)LTCAmp->GetRenderTargetItem().ShaderResourceTexture, 0, false);
-		    }
-		} // end Create the GTAO  randomization texture
-	} // end if (FeatureLevelInitializedTo < ERHIFeatureLevel::SM5 && InFeatureLevel >= ERHIFeatureLevel::SM5)
+						{
+							Dest[k] = FFloat16(LTC_Amp[4 * (x + y * LTC_Size) + k]).Encoded;
+						}
+					}
+				}
+				RHICmdList.UnlockTexture2D(Texture2D, 0, false);
+
+				LTCAmp = CreateRenderTarget(Texture2D, CreateInfo.DebugName);
+			}
+		}
+	}
 
 	// Create the SSAO randomization texture
 	static const auto MobileAmbientOcclusionCVar = IConsoleManager::Get().FindTConsoleVariableDataInt(TEXT("r.Mobile.AmbientOcclusion"));
@@ -669,17 +605,18 @@ void FSystemTextures::InitializeFeatureLevelDependentTextures(FRHICommandListImm
 			}
 
 			{
-				// could be PF_V8U8 to save shader instructions but that doesn't work on all hardware
-				FPooledRenderTargetDesc Desc(FPooledRenderTargetDesc::Create2DDesc(FIntPoint(64, 64), PF_R8G8, FClearValueBinding::None, TexCreate_HideInVisualizeTexture, TexCreate_NoFastClear | TexCreate_ShaderResource, false));
-				Desc.AutoWritable = false;
-				GRenderTargetPool.FindFreeElement(RHICmdList, Desc, SSAORandomization, TEXT("SSAORandomization"), ERenderTargetTransience::NonTransient);
+				const uint32 Extent = 64;
+
+				FRHIResourceCreateInfo CreateInfo(TEXT("SSAORandomization"));
+				FTexture2DRHIRef Texture2D = RHICreateTexture2D(Extent, Extent, PF_R8G8, 1, 1, TexCreate_ShaderResource, CreateInfo);
+
 				// Write the contents of the texture.
 				uint32 DestStride;
-				uint8* DestBuffer = (uint8*)RHICmdList.LockTexture2D((FTexture2DRHIRef&)SSAORandomization->GetRenderTargetItem().ShaderResourceTexture, 0, RLM_WriteOnly, DestStride, false);
+				uint8* DestBuffer = (uint8*)RHICmdList.LockTexture2D(Texture2D, 0, RLM_WriteOnly, DestStride, false);
 
-				for (int32 y = 0; y < Desc.Extent.Y; ++y)
+				for (int32 y = 0; y < Extent; ++y)
 				{
-					for (int32 x = 0; x < Desc.Extent.X; ++x)
+					for (int32 x = 0; x < Extent; ++x)
 					{
 						uint8* Dest = (uint8*)(DestBuffer + x * sizeof(uint16) + y * DestStride);
 
@@ -689,8 +626,11 @@ void FSystemTextures::InitializeFeatureLevelDependentTextures(FRHICommandListImm
 						Dest[1] = Bases[Index].G;
 					}
 				}
+
+				RHICmdList.UnlockTexture2D(Texture2D, 0, false);
+
+				SSAORandomization = CreateRenderTarget(Texture2D, CreateInfo.DebugName);
 			}
-			RHICmdList.UnlockTexture2D((FTexture2DRHIRef&)SSAORandomization->GetRenderTargetItem().ShaderResourceTexture, 0, false);
 		}
 	}
 		
@@ -704,18 +644,17 @@ void FSystemTextures::InitializeFeatureLevelDependentTextures(FRHICommandListImm
 
 		bool bGTAOPreIngegratedUsingVolumeLUT = MobileGTAOPreIntegratedTextureTypeCVar->GetValueOnAnyThread() == 2;
 
-		FPooledRenderTargetDesc Desc;
+		TRefCountPtr<FRHITexture> Texture;
+
+		FRHIResourceCreateInfo CreateInfo(TEXT("GTAOPreIntegrated"));
 		if (bGTAOPreIngegratedUsingVolumeLUT)
 		{
-			Desc = FPooledRenderTargetDesc(FPooledRenderTargetDesc::CreateVolumeDesc(Extent, Extent, Extent, PF_R16F, FClearValueBinding::None, TexCreate_HideInVisualizeTexture | TexCreate_NoTiling | TexCreate_ShaderResource, TexCreate_ShaderResource, false));
+			Texture = RHICreateTexture3D(Extent, Extent, Extent, PF_R16F, 1, TexCreate_ShaderResource, CreateInfo);
 		}
 		else
 		{
-			Desc = FPooledRenderTargetDesc(FPooledRenderTargetDesc::Create2DDesc(FIntPoint(Square, Extent), PF_R16F, FClearValueBinding::None, TexCreate_HideInVisualizeTexture | TexCreate_NoTiling | TexCreate_ShaderResource, TexCreate_ShaderResource, false));
+			Texture = RHICreateTexture2D(Square, Extent, PF_R16F, 1, 1, TexCreate_ShaderResource, CreateInfo);
 		}
-		
-		Desc.AutoWritable = false;
-		GRenderTargetPool.FindFreeElement(RHICmdList, Desc, GTAOPreIntegrated, TEXT("GTAOPreIntegrated"), ERenderTargetTransience::NonTransient);
 
 		// Write the contents of the texture.
 		TArray<FFloat16> TempBuffer;
@@ -730,7 +669,7 @@ void FSystemTextures::InitializeFeatureLevelDependentTextures(FRHICommandListImm
 		else
 		{
 			uint32 DestStride;
-			DestBuffer = (FFloat16*)RHICmdList.LockTexture2D((FTexture2DRHIRef&)GTAOPreIntegrated->GetRenderTargetItem().ShaderResourceTexture, 0, RLM_WriteOnly, DestStride, false);
+			DestBuffer = (FFloat16*)RHICmdList.LockTexture2D((FRHITexture2D*)Texture.GetReference(), 0, RLM_WriteOnly, DestStride, false);
 		}
 
 		for (uint32 z = 0; z < Extent; ++z)
@@ -776,20 +715,22 @@ void FSystemTextures::InitializeFeatureLevelDependentTextures(FRHICommandListImm
 
 		if (bGTAOPreIngegratedUsingVolumeLUT)
 		{
-			FUpdateTextureRegion3D Region(0, 0, 0, 0, 0, 0, Desc.Extent.X, Desc.Extent.Y, Desc.Depth);
+			FUpdateTextureRegion3D Region(0, 0, 0, 0, 0, 0, Extent, Extent, Extent);
 
 			RHICmdList.UpdateTexture3D(
-				(FTexture3DRHIRef&)GTAOPreIntegrated->GetRenderTargetItem().ShaderResourceTexture,
+				(FRHITexture3D*)Texture.GetReference(),
 				0,
 				Region,
-				Desc.Extent.X * sizeof(FFloat16),
-				Desc.Extent.X * Desc.Extent.Y * sizeof(FFloat16),
+				Extent * sizeof(FFloat16),
+				Extent * Extent * sizeof(FFloat16),
 				(const uint8*)DestBuffer);
 		}
 		else
 		{
-			RHICmdList.UnlockTexture2D((FTexture2DRHIRef&)GTAOPreIntegrated->GetRenderTargetItem().ShaderResourceTexture, 0, false);
+			RHICmdList.UnlockTexture2D((FRHITexture2D*)Texture.GetReference(), 0, false);
 		}
+
+		GTAOPreIntegrated = CreateRenderTarget(Texture, CreateInfo.DebugName);
 	}
 
 	// Initialize textures only once.
@@ -818,11 +759,12 @@ void FSystemTextures::ReleaseDynamicRHI()
 	GreenDummy.SafeRelease();
 	DefaultNormal8Bit.SafeRelease();
 	VolumetricBlackDummy.SafeRelease();
+	VolumetricBlackAlphaOneDummy.SafeRelease();
 	ZeroUIntDummy.SafeRelease();
-	ZeroUShort4Dummy.SafeRelease();
 	MidGreyDummy.SafeRelease();
 	StencilDummy.SafeRelease();
 	StencilDummySRV.SafeRelease();
+	BlackDepthCube.SafeRelease();
 	GTAOPreIntegrated.SafeRelease();
 
 	DefaultTextures.Empty();
@@ -891,11 +833,6 @@ FRDGTextureRef FSystemTextures::GetZeroUIntDummy(FRDGBuilder& GraphBuilder) cons
 	return GraphBuilder.RegisterExternalTexture(ZeroUIntDummy, TEXT("ZeroUIntDummy"));
 }
 
-FRDGTextureRef FSystemTextures::GetZeroUShort4Dummy(FRDGBuilder& GraphBuilder) const
-{
-	return GraphBuilder.RegisterExternalTexture(ZeroUShort4Dummy, TEXT("ZeroUShort4Dummy"));
-}
-
 //////////////////////////////////////////////////////////////////////////////////////////////////////////
 // Default textures
 
@@ -1345,42 +1282,50 @@ FRDGBufferRef GetInternalDefaultBuffer(
 		return GraphBuilder.RegisterExternalBuffer(DefaultBuffers[Index].Buffer);
 	}
 
+	const uint32 BufferSize = NumBytePerElement * NumElements;
+
+	FRDGBufferDesc BufferDesc;
+	FRHIResourceCreateInfo CreateInfo(TEXT(""));
+
 	// Adding new buffer if there is no fit (slow path)
-	FRDGBufferRef Buffer = nullptr; 
+	TRefCountPtr<FRHIBuffer> RHIBuffer;
 	if (bIsStructuredBuffer)
 	{
-		Buffer = GraphBuilder.CreateBuffer(FRDGBufferDesc::CreateStructuredDesc(NumBytePerElement, NumElements), TEXT("DefaultStructuredBuffer"));
+		CreateInfo.DebugName = TEXT("DefaultStructuredBuffer");
+
+		BufferDesc = FRDGBufferDesc::CreateStructuredDesc(NumBytePerElement, NumElements);
+		BufferDesc.Usage = BUF_Static | BUF_ShaderResource;
+
+		RHIBuffer = RHICreateStructuredBuffer(NumBytePerElement, BufferSize, BufferDesc.Usage, CreateInfo);
 	}
 	else
 	{
-		Buffer = GraphBuilder.CreateBuffer(FRDGBufferDesc::CreateBufferDesc(NumBytePerElement, NumElements), TEXT("DefaultBuffer"));
-	}
+		CreateInfo.DebugName = TEXT("DefaultBuffer");
+
+		BufferDesc = FRDGBufferDesc::CreateUploadDesc(NumBytePerElement, NumElements);
 
-	FRDGBufferUploader BufferUploader;
+		RHIBuffer = RHICreateVertexBuffer(BufferSize, BUF_Static | BUF_ShaderResource, CreateInfo);
+	}
 
-	// Initialize the entire buffer with the provided data
+	void* DestPtr = GraphBuilder.RHICmdList.LockBuffer(RHIBuffer, 0, BufferSize, RLM_WriteOnly);
 	if (Value)
 	{
-		BufferUploader.Upload(GraphBuilder, Buffer, Value, NumElements * NumBytePerElement);
+		FMemory::Memcpy(DestPtr, Value, BufferSize);
 	}
-	// Initialize buffer to 0
 	else
 	{
-		TArray<uint8> DefaultValue;
-		DefaultValue.Init(0u, NumElements * NumBytePerElement);
-		BufferUploader.Upload(GraphBuilder, Buffer, DefaultValue.GetData(), DefaultValue.Num());
+		FMemory::Memzero(DestPtr, BufferSize);
 	}
-
-	BufferUploader.Submit(GraphBuilder);
+	GraphBuilder.RHICmdList.UnlockBuffer(RHIBuffer);
 
 	FDefaultBuffer Entry;
 	Entry.Key = Key;
 	Entry.Hash = Hash;
-	Entry.Buffer = GraphBuilder.ConvertToExternalBuffer(Buffer);
+	Entry.Buffer = new FRDGPooledBuffer(RHIBuffer, BufferDesc, NumElements, CreateInfo.DebugName);
 
 	Index = DefaultBuffers.Add(Entry);
 	HashDefaultBuffers.Add(Hash, Index);
-	return Buffer;
+	return GraphBuilder.RegisterExternalBuffer(Entry.Buffer);
 }
 
 FVector4 GetClearBindingValue(EPixelFormat Format, FClearValueBinding Value)
diff --git a/Engine/Source/Runtime/Renderer/Private/SystemTextures.h b/Engine/Source/Runtime/Renderer/Private/SystemTextures.h
--- a/Engine/Source/Runtime/Renderer/Private/SystemTextures.h
+++ b/Engine/Source/Runtime/Renderer/Private/SystemTextures.h
@@ -26,10 +26,12 @@ struct FRDGSystemTextures
 	FRDGTextureRef MaxFP16Depth{};
 	FRDGTextureRef DepthDummy{};
 	FRDGTextureRef StencilDummy{};
+	FRDGTextureRef BlackDepthCube{};
 	FRDGTextureRef Green{};
 	FRDGTextureRef DefaultNormal8Bit{};
 	FRDGTextureRef MidGrey{};
 	FRDGTextureRef VolumetricBlack{};
+	FRDGTextureRef VolumetricBlackAlphaOne{};
 
 	FRDGTextureSRVRef StencilDummySRV{};
 };
@@ -122,6 +124,7 @@ public:
 	TRefCountPtr<IPooledRenderTarget> DepthDummy;
 	/** Stencil texture that holds a single stencil value. */
 	TRefCountPtr<IPooledRenderTarget> StencilDummy;
+	TRefCountPtr<IPooledRenderTarget> BlackDepthCube;
 	// float4(0,1,0,1)
 	TRefCountPtr<IPooledRenderTarget> GreenDummy;
 	// float4(0.5,0.5,0.5,1)
@@ -131,10 +134,10 @@ public:
 
 	/** float4(0,0,0,0) volumetric texture. */
 	TRefCountPtr<IPooledRenderTarget> VolumetricBlackDummy;
+	TRefCountPtr<IPooledRenderTarget> VolumetricBlackAlphaOneDummy;
 	
 	// Dummy 0 Uint texture for RHIs that need explicit overloads
 	TRefCountPtr<IPooledRenderTarget> ZeroUIntDummy;
-	TRefCountPtr<IPooledRenderTarget> ZeroUShort4Dummy;
 
 	// SRV for WhiteDummy Texture.
 	TRefCountPtr<FRHIShaderResourceView> WhiteDummySRV;
diff --git a/Engine/Source/Runtime/Renderer/Private/VolumeLighting.cpp b/Engine/Source/Runtime/Renderer/Private/VolumeLighting.cpp
--- a/Engine/Source/Runtime/Renderer/Private/VolumeLighting.cpp
+++ b/Engine/Source/Runtime/Renderer/Private/VolumeLighting.cpp
@@ -34,7 +34,7 @@ const FProjectedShadowInfo* GetFirstWholeSceneShadowMap(const FVisibleLightInfo&
 static auto SetVolumeShadowingDefaultShaderParametersGlobal = [](FRDGBuilder& GraphBuilder, auto& ShaderParams)
 {
 	const FRDGSystemTextures& SystemTextures = FRDGSystemTextures::Get(GraphBuilder);
-	FRDGTextureRef BlackDepthCubeTexture = GBlackTextureDepthCube->GetRDG(GraphBuilder);
+	FRDGTextureRef BlackDepthCubeTexture = SystemTextures.BlackDepthCube;
 
 	ShaderParams.WorldToShadowMatrix = FMatrix::Identity;
 	ShaderParams.ShadowmapMinMax = FVector4(1.0f);
@@ -153,8 +153,8 @@ static auto GetVolumeShadowingShaderParametersGlobal = [](
 	// See FOnePassPointShadowProjectionShaderParameters from ShadowRendering.h
 	//
 	FRDGTexture* ShadowDepthTextureValue = ShadowInfo
-		? GraphBuilder.RegisterExternalTexture(ShadowInfo->RenderTargets.DepthTarget)
-		: GBlackTextureDepthCube->GetRDG(GraphBuilder);
+		? GraphBuilder.RegisterExternalTexture(ShadowInfo->RenderTargets.DepthTarget) 
+		: SystemTextures.BlackDepthCube;
 
 	ShaderParams.OnePassPointShadowProjection.ShadowDepthCubeTexture = ShadowDepthTextureValue;
 	ShaderParams.OnePassPointShadowProjection.ShadowDepthCubeTexture2 = ShadowDepthTextureValue;
