Fix for FSceneRenderer crashing in destructor when r.SceneRenderCleanUpMode > 0. Re-enabled async mode.

diff --git a/Engine/Source/Runtime/Renderer/Private/GPUScene.cpp b/Engine/Source/Runtime/Renderer/Private/GPUScene.cpp
--- a/Engine/Source/Runtime/Renderer/Private/GPUScene.cpp
+++ b/Engine/Source/Runtime/Renderer/Private/GPUScene.cpp
@@ -1417,6 +1417,11 @@ TRange<int32> FGPUScene::CommitPrimitiveCollector(FGPUScenePrimitiveCollector& P
 
 
 FGPUSceneDynamicContext::~FGPUSceneDynamicContext()
+{
+	Release();
+}
+
+void FGPUSceneDynamicContext::Release()
 {
 	for (auto UploadData : DymamicPrimitiveUploadData)
 	{
diff --git a/Engine/Source/Runtime/Renderer/Private/GPUScene.h b/Engine/Source/Runtime/Renderer/Private/GPUScene.h
--- a/Engine/Source/Runtime/Renderer/Private/GPUScene.h
+++ b/Engine/Source/Runtime/Renderer/Private/GPUScene.h
@@ -126,6 +126,8 @@ public:
 	FGPUSceneDynamicContext(FGPUScene& InGPUScene) : GPUScene(InGPUScene) {}
 	~FGPUSceneDynamicContext();
 
+	void Release();
+
 private:
 	friend class FGPUScene;
 	friend class FGPUScenePrimitiveCollector;
diff --git a/Engine/Source/Runtime/Renderer/Private/SceneRendering.cpp b/Engine/Source/Runtime/Renderer/Private/SceneRendering.cpp
--- a/Engine/Source/Runtime/Renderer/Private/SceneRendering.cpp
+++ b/Engine/Source/Runtime/Renderer/Private/SceneRendering.cpp
@@ -3507,7 +3507,7 @@ void FSceneRenderer::ViewExtensionPreRender_RenderThread(FRHICommandListImmediat
 	FDeferredUpdateResource::UpdateResources(RHICmdList);
 }
 
-static int32 GSceneRenderCleanUpMode = 0;
+static int32 GSceneRenderCleanUpMode = 2;
 static FAutoConsoleVariableRef CVarSceneRenderCleanUpMode(
 	TEXT("r.SceneRenderCleanUpMode"),
 	GSceneRenderCleanUpMode,
@@ -3603,6 +3603,8 @@ void FSceneRenderer::RenderThreadEnd(FRHICommandListImmediate& RHICmdList)
 	}
 	else
 	{
+		GPUSceneDynamicContext.Release();
+
 		GSceneRenderCleanUpState.Renderer = this;
 		GSceneRenderCleanUpState.MemStackMark = MemStackMark;
 
