Create VolumetricBlackDummy and VolumetricBlackAlphaOneDummy system textures for mobile feature level as well

diff --git a/Engine/Source/Runtime/Renderer/Private/SystemTextures.cpp b/Engine/Source/Runtime/Renderer/Private/SystemTextures.cpp
--- a/Engine/Source/Runtime/Renderer/Private/SystemTextures.cpp
+++ b/Engine/Source/Runtime/Renderer/Private/SystemTextures.cpp
@@ -206,6 +206,12 @@ void FSystemTextures::InitializeCommonTextures(FRHICommandListImmediate& RHICmdL
 		SetDummyTextureData<FFloat16Color>(Texture2D, FFloat16Color(FLinearColor(0.5f, 0.5f, 0.5f, 0.5f)));
 		MidGreyDummy = CreateRenderTarget(Texture2D, CreateInfo.DebugName);
 	}
+
+	// Create a VolumetricBlackDummy texture
+	{
+		VolumetricBlackDummy = CreateRenderTarget(GBlackVolumeTexture->TextureRHI, TEXT("VolumetricBlackDummy"));
+		VolumetricBlackAlphaOneDummy = CreateRenderTarget(GBlackAlpha1VolumeTexture->TextureRHI, TEXT("VolumetricBlackAlphaOneDummy"));
+	}
 }
 
 void FSystemTextures::InitializeFeatureLevelDependentTextures(FRHICommandListImmediate& RHICmdList, const ERHIFeatureLevel::Type InFeatureLevel)
@@ -247,13 +253,6 @@ void FSystemTextures::InitializeFeatureLevelDependentTextures(FRHICommandListImm
 		SobolSampling = CreateRenderTarget(Texture2D, CreateInfo.DebugName);
 	}
 
-	// Create a VolumetricBlackDummy texture
-	if (CurrentFeatureLevel < ERHIFeatureLevel::SM5 && InFeatureLevel >= ERHIFeatureLevel::SM5)
-	{
-		VolumetricBlackDummy = CreateRenderTarget(GBlackVolumeTexture->TextureRHI, TEXT("VolumetricBlackDummy"));
-		VolumetricBlackAlphaOneDummy = CreateRenderTarget(GBlackAlpha1VolumeTexture->TextureRHI, TEXT("VolumetricBlackAlphaOneDummy"));
-	}
-
 	if (CurrentFeatureLevel < ERHIFeatureLevel::SM5 && InFeatureLevel >= ERHIFeatureLevel::SM5)
 	{
 		FPooledRenderTargetDesc Desc(FPooledRenderTargetDesc::CreateVolumeDesc(1, 1, 1, PF_B8G8R8A8, FClearValueBinding::Transparent, TexCreate_HideInVisualizeTexture, TexCreate_ShaderResource | TexCreate_RenderTargetable | TexCreate_NoFastClear, false));
