Remove unused FMobileRenderOpaqueFXPassParameters UB

diff --git a/Engine/Source/Runtime/Renderer/Private/MobileShadingRenderer.cpp b/Engine/Source/Runtime/Renderer/Private/MobileShadingRenderer.cpp
--- a/Engine/Source/Runtime/Renderer/Private/MobileShadingRenderer.cpp
+++ b/Engine/Source/Runtime/Renderer/Private/MobileShadingRenderer.cpp
@@ -150,10 +150,6 @@ static bool UsesCustomDepthStencilLookup(const FViewInfo& View)
 	return bUsesCustomDepthStencil;
 }
 
-BEGIN_SHADER_PARAMETER_STRUCT(FMobileRenderOpaqueFXPassParameters, )
-	SHADER_PARAMETER_RDG_UNIFORM_BUFFER(FMobileSceneTextureUniformParameters, MobileSceneTextures)
-END_SHADER_PARAMETER_STRUCT()
-
 static void RenderOpaqueFX(
 	FRDGBuilder& GraphBuilder,
 	TArrayView<const FViewInfo> Views,
