Added UI for VisualiseIndirectLighting using probes.

diff --git a/Engine/Source/Runtime/Engine/Public/ShowFlags.h b/Engine/Source/Runtime/Engine/Public/ShowFlags.h
--- a/Engine/Source/Runtime/Engine/Public/ShowFlags.h
+++ b/Engine/Source/Runtime/Engine/Public/ShowFlags.h
@@ -439,6 +439,7 @@ private:
 		SetVisualizeGlobalDistanceField(false);
 		SetVisualizeLumenScene(false);
 		SetVisualizeLumenIndirectDiffuse(false);
+		SetVisualizeIndirectLighting(false);
 		SetVisualizeDistanceFieldAO(false);
 		SetVisualizeSSR(false);
 		SetVisualizeSSS(false);
diff --git a/Engine/Source/Runtime/Engine/Public/ShowFlagsValues.inl b/Engine/Source/Runtime/Engine/Public/ShowFlagsValues.inl
--- a/Engine/Source/Runtime/Engine/Public/ShowFlagsValues.inl
+++ b/Engine/Source/Runtime/Engine/Public/ShowFlagsValues.inl
@@ -309,6 +309,8 @@ SHOWFLAG_FIXED_IN_SHIPPING(0, VisualizeGlobalDistanceField, SFG_Visualize, NSLOC
 SHOWFLAG_FIXED_IN_SHIPPING(0, VisualizeLumenScene, SFG_Visualize, NSLOCTEXT("UnrealEd", "LumenSceneSF", "Lumen Scene"))
 /** Visualization of Lumen indirect diffuse */
 SHOWFLAG_FIXED_IN_SHIPPING(0, VisualizeLumenIndirectDiffuse, SFG_Visualize, NSLOCTEXT("UnrealEd", "LumenIndirectDiffuseSF", "Lumen Global Illumination"))
+/** Enable the debug visualisation of indirect diffuse and specular lighting using probes */
+SHOWFLAG_FIXED_IN_SHIPPING(0, VisualizeIndirectLighting, SFG_Visualize, NSLOCTEXT("UnrealEd", "VisualizeIndirectLightingSF", "Visualize Indirect Lighting"))
 /** Screen space AO, for now SHOWFLAG_ALWAYS_ACCESSIBLE because r.GBuffer need that */
 SHOWFLAG_ALWAYS_ACCESSIBLE(ScreenSpaceAO, SFG_LightingFeatures, NSLOCTEXT("UnrealEd", "ScreenSpaceAOSF", "Screen Space Ambient Occlusion"))
 /** Distance field AO, for now SHOWFLAG_ALWAYS_ACCESSIBLE because it's exposed in SceneCapture */
diff --git a/Engine/Source/Runtime/Renderer/Private/DebugProbeRendering.cpp b/Engine/Source/Runtime/Renderer/Private/DebugProbeRendering.cpp
--- a/Engine/Source/Runtime/Renderer/Private/DebugProbeRendering.cpp
+++ b/Engine/Source/Runtime/Renderer/Private/DebugProbeRendering.cpp
@@ -6,8 +6,8 @@
 
 
 // Changing this causes a full shader recompile
-static TAutoConsoleVariable<int32> CVarDebugProbes(
-	TEXT("r.DebugProbes"),
+static TAutoConsoleVariable<int32> CVarDebugIndirectLightingProbes(
+	TEXT("r.DebugIndirectLightingProbes"),
 	0,
 	TEXT("Enables debug probes rendering to visualise diffuse/specular lighting on simple sphere scattered in the world.") \
 	TEXT(" 0: disabled.\n")
@@ -61,8 +61,8 @@ static void CommonStampDeferredDebugProbeDrawCall(
 	FStampDeferredDebugProbePS::FParameters* PassParameters)
 {
 	PassParameters->ViewUniformBuffer = View.ViewUniformBuffer;
-	PassParameters->DebugProbesMode = FMath::Clamp(CVarDebugProbes.GetValueOnRenderThread(), 0, 3);
-
+	PassParameters->DebugProbesMode = View.Family->EngineShowFlags.VisualizeIndirectLighting ? 3 : FMath::Clamp(CVarDebugIndirectLightingProbes.GetValueOnRenderThread(), 0, 3);
+		
 	FStampDeferredDebugProbePS::FPermutationDomain PermutationVector;
 	TShaderMapRef<FStampDeferredDebugProbePS> PixelShader(View.ShaderMap, PermutationVector);
 
@@ -82,13 +82,14 @@ void StampDeferredDebugProbeDepthPS(
 	RDG_EVENT_SCOPE(GraphBuilder, "StampDeferredDebugProbeDepth");
 	RDG_GPU_STAT_SCOPE(GraphBuilder, StampDeferredDebugProbe);
 
-	if (CVarDebugProbes.GetValueOnRenderThread() <= 0)
-	{
-		return;
-	}
-
+	const bool bVisualizeIndirectLighting = CVarDebugIndirectLightingProbes.GetValueOnRenderThread() > 0;
 	for (const FViewInfo& View : Views)
 	{
+		if (!(bVisualizeIndirectLighting || View.Family->EngineShowFlags.VisualizeIndirectLighting))
+		{
+			continue;
+		}
+
 		FStampDeferredDebugProbePS::FParameters* PassParameters = GraphBuilder.AllocParameters<FStampDeferredDebugProbePS::FParameters>();
 		PassParameters->ViewUniformBuffer = View.ViewUniformBuffer;
 		PassParameters->RenderTargets.DepthStencil = FDepthStencilBinding(SceneDepthTexture, ERenderTargetLoadAction::ELoad, ERenderTargetLoadAction::ELoad, FExclusiveDepthStencil::DepthWrite_StencilWrite);
@@ -105,13 +106,14 @@ void StampDeferredDebugProbeMaterialPS(
 	RDG_EVENT_SCOPE(GraphBuilder, "StampDeferredDebugProbeMaterial");
 	RDG_GPU_STAT_SCOPE(GraphBuilder, StampDeferredDebugProbe);
 
-	if (CVarDebugProbes.GetValueOnRenderThread() <= 0)
-	{
-		return;
-	}
-
+	const bool bVisualizeIndirectLighting = CVarDebugIndirectLightingProbes.GetValueOnRenderThread() > 0;
 	for (const FViewInfo& View : Views)
 	{
+		if (!(bVisualizeIndirectLighting || View.Family->EngineShowFlags.VisualizeIndirectLighting))
+		{
+			continue;
+		}
+
 		FStampDeferredDebugProbePS::FParameters* PassParameters = GraphBuilder.AllocParameters<FStampDeferredDebugProbePS::FParameters>();
 		PassParameters->RenderTargets = BasePassRenderTargets;
 
