Converted SetShaderUniformBufferImmediate to local uniform buffers to move creation off the render thread. Removed uses of RHICommandListImmediate. Moved local uniform buffer functions onto the compute command list.

diff --git a/Engine/Source/Runtime/RenderCore/Public/ShaderParameterUtils.h b/Engine/Source/Runtime/RenderCore/Public/ShaderParameterUtils.h
--- a/Engine/Source/Runtime/RenderCore/Public/ShaderParameterUtils.h
+++ b/Engine/Source/Runtime/RenderCore/Public/ShaderParameterUtils.h
@@ -604,11 +604,13 @@ inline void SetUniformBufferParameterImmediate(
 	checkSlow(Parameter.IsInitialized());
 	if(Parameter.IsBound())
 	{
-		RHICmdList.SetShaderUniformBuffer(
+		const FRHIUniformBufferLayout& UniformBufferLayout = TBufferStruct::StaticStructMetadata.GetLayout();
+		FLocalUniformBuffer UniformBuffer = RHICmdList.BuildLocalUniformBuffer(&UniformBufferValue, UniformBufferLayout.ConstantBufferSize, UniformBufferLayout);
+
+		RHICmdList.SetLocalShaderUniformBuffer(
 			Shader,
 			Parameter.GetBaseIndex(),
-			RHICreateUniformBuffer(&UniformBufferValue,TBufferStruct::StaticStructMetadata.GetLayout(),UniformBuffer_SingleDraw)
-			);
+			UniformBuffer);
 	}
 }
 
@@ -625,10 +627,12 @@ inline void SetUniformBufferParameterImmediate(
 	checkSlow(Parameter.IsInitialized());
 	if(Parameter.IsBound())
 	{
-		RHICmdList.SetShaderUniformBuffer(
+		const FRHIUniformBufferLayout& UniformBufferLayout = TBufferStruct::StaticStructMetadata.GetLayout();
+		FLocalUniformBuffer UniformBuffer = RHICmdList.BuildLocalUniformBuffer(&UniformBufferValue, UniformBufferLayout.ConstantBufferSize, UniformBufferLayout);
+
+		RHICmdList.SetLocalShaderUniformBuffer(
 			Shader,
 			Parameter.GetBaseIndex(),
-			RHICreateUniformBuffer(&UniformBufferValue,TBufferStruct::StaticStructMetadata.GetLayout(),UniformBuffer_SingleDraw)
-			);
+			UniformBuffer);
 	}
-}
+}
\ No newline at end of file
diff --git a/Engine/Source/Runtime/Renderer/Private/CompositionLighting/PostProcessAmbientOcclusion.cpp b/Engine/Source/Runtime/Renderer/Private/CompositionLighting/PostProcessAmbientOcclusion.cpp
--- a/Engine/Source/Runtime/Renderer/Private/CompositionLighting/PostProcessAmbientOcclusion.cpp
+++ b/Engine/Source/Runtime/Renderer/Private/CompositionLighting/PostProcessAmbientOcclusion.cpp
@@ -893,7 +893,7 @@ FScreenPassTexture AddAmbientOcclusionPass(
 				RDG_EVENT_NAME("DepthBounds ClearQuad(%s)", Output.Texture->Name),
 				ClearParameters,
 				ERDGPassFlags::Raster,
-			[OutputViewport, DepthFar](FRHICommandListImmediate& RHICmdList)
+			[OutputViewport, DepthFar](FRHICommandList& RHICmdList)
 			{
 				// We must clear all pixels that won't be touched by AO shader.
 				FClearQuadCallbacks Callbacks;
@@ -946,7 +946,7 @@ FScreenPassTexture AddAmbientOcclusionPass(
 			MoveTemp(EventName),
 			PassParameters,
 			ERDGPassFlags::Raster,
-			[&View, OutputViewport, InputViewport, VertexShader, PixelShader, PassParameters, bDepthBoundsTestEnabled, DepthFar] (FRHICommandListImmediate& RHICmdList)
+			[&View, OutputViewport, InputViewport, VertexShader, PixelShader, PassParameters, bDepthBoundsTestEnabled, DepthFar] (FRHICommandList& RHICmdList)
 		{
 			const FIntRect InputRect = InputViewport.Rect;
 			const FIntPoint InputSize = InputViewport.Extent;
diff --git a/Engine/Source/Runtime/Renderer/Private/DistanceFieldShadowing.cpp b/Engine/Source/Runtime/Renderer/Private/DistanceFieldShadowing.cpp
--- a/Engine/Source/Runtime/Renderer/Private/DistanceFieldShadowing.cpp
+++ b/Engine/Source/Runtime/Renderer/Private/DistanceFieldShadowing.cpp
@@ -397,7 +397,7 @@ void ScatterObjectsToShadowTiles(
 			RDG_EVENT_NAME("ScatterMeshSDFsToLightGrid %ux%u", LightTileDimensions.X, LightTileDimensions.Y),
 			PassParameters,
 			ERDGPassFlags::Raster,
-			[LightTileDimensions, bReverseCulling, VertexShader, PixelShader, PassParameters](FRHICommandListImmediate& RHICmdList)
+			[LightTileDimensions, bReverseCulling, VertexShader, PixelShader, PassParameters](FRHICommandList& RHICmdList)
 		{
 			FGraphicsPipelineStateInitializer GraphicsPSOInit;
 			RHICmdList.ApplyCachedRenderTargets(GraphicsPSOInit);
@@ -979,7 +979,7 @@ void FProjectedShadowInfo::RenderRayTracedDistanceFieldProjection(
 			RDG_EVENT_NAME("Upsample"),
 			PassParameters,
 			ERDGPassFlags::Raster,
-			[this, &View, PixelShader, ScissorRect, bProjectingForForwardShading, PassParameters](FRHICommandListImmediate& RHICmdList)
+			[this, &View, PixelShader, ScissorRect, bProjectingForForwardShading, PassParameters](FRHICommandList& RHICmdList)
 		{
 			FGraphicsPipelineStateInitializer GraphicsPSOInit;
 			RHICmdList.ApplyCachedRenderTargets(GraphicsPSOInit);
diff --git a/Engine/Source/Runtime/Renderer/Private/GpuDebugRendering.cpp b/Engine/Source/Runtime/Renderer/Private/GpuDebugRendering.cpp
--- a/Engine/Source/Runtime/Renderer/Private/GpuDebugRendering.cpp
+++ b/Engine/Source/Runtime/Renderer/Private/GpuDebugRendering.cpp
@@ -321,7 +321,7 @@ namespace ShaderDrawDebug
 				RDG_EVENT_NAME("ShaderDrawDebug"),
 				PassParameters,
 				ERDGPassFlags::Raster,
-				[VertexShader, PixelShader, PassParameters, IndirectBuffer, LockedIndirectBuffer, bUseRdgInput](FRHICommandListImmediate& RHICmdListImmediate)
+				[VertexShader, PixelShader, PassParameters, IndirectBuffer, LockedIndirectBuffer, bUseRdgInput](FRHICommandList& RHICmdList)
 			{
 				// Marks the indirect draw parameter as used by the pass, given it's not used directly by any of the shaders.
 				if (bUseRdgInput)
@@ -330,7 +330,7 @@ namespace ShaderDrawDebug
 				}
 
 				FGraphicsPipelineStateInitializer GraphicsPSOInit;
-				RHICmdListImmediate.ApplyCachedRenderTargets(GraphicsPSOInit);
+				RHICmdList.ApplyCachedRenderTargets(GraphicsPSOInit);
 				GraphicsPSOInit.DepthStencilState = TStaticDepthStencilState<false, CF_Always>::GetRHI();
 				GraphicsPSOInit.BlendState = TStaticBlendState<CW_RGBA, BO_Add, BF_One, BF_InverseSourceAlpha, BO_Add, BF_Zero, BF_One>::GetRHI(); // Premultiplied-alpha composition
 				GraphicsPSOInit.RasterizerState = TStaticRasterizerState<FM_Solid, CM_None, true>::GetRHI();
@@ -338,10 +338,10 @@ namespace ShaderDrawDebug
 				GraphicsPSOInit.BoundShaderState.VertexDeclarationRHI = GEmptyVertexDeclaration.VertexDeclarationRHI;
 				GraphicsPSOInit.BoundShaderState.VertexShaderRHI = VertexShader.GetVertexShader();
 				GraphicsPSOInit.BoundShaderState.PixelShaderRHI = PixelShader.GetPixelShader();
-				SetGraphicsPipelineState(RHICmdListImmediate, GraphicsPSOInit);
+				SetGraphicsPipelineState(RHICmdList, GraphicsPSOInit);
 
-				SetShaderParameters(RHICmdListImmediate, VertexShader, VertexShader.GetVertexShader(), PassParameters->ShaderDrawVSParameters);
-				SetShaderParameters(RHICmdListImmediate, PixelShader, PixelShader.GetPixelShader(), PassParameters->ShaderDrawPSParameters);
+				SetShaderParameters(RHICmdList, VertexShader, VertexShader.GetVertexShader(), PassParameters->ShaderDrawVSParameters);
+				SetShaderParameters(RHICmdList, PixelShader, PixelShader.GetPixelShader(), PassParameters->ShaderDrawPSParameters);
 
 
 				if (bUseRdgInput)
@@ -349,11 +349,11 @@ namespace ShaderDrawDebug
 					// Marks the indirect draw parameter as used by the pass, given it's not used directly by any of the shaders.
 					FRHIBuffer* IndirectBufferRHI = PassParameters->ShaderDrawVSParameters.IndirectBuffer->GetIndirectRHICallBuffer();
 					check(IndirectBufferRHI != nullptr);
-					RHICmdListImmediate.DrawPrimitiveIndirect(IndirectBufferRHI, 0);
+					RHICmdList.DrawPrimitiveIndirect(IndirectBufferRHI, 0);
 				}
 				else
 				{
-					RHICmdListImmediate.DrawPrimitiveIndirect(LockedIndirectBuffer, 0);
+					RHICmdList.DrawPrimitiveIndirect(LockedIndirectBuffer, 0);
 				}
 			});
 		};
diff --git a/Engine/Source/Runtime/Renderer/Private/PostProcess/PostProcessBloomSetup.cpp b/Engine/Source/Runtime/Renderer/Private/PostProcess/PostProcessBloomSetup.cpp
--- a/Engine/Source/Runtime/Renderer/Private/PostProcess/PostProcessBloomSetup.cpp
+++ b/Engine/Source/Runtime/Renderer/Private/PostProcess/PostProcessBloomSetup.cpp
@@ -156,7 +156,7 @@ FScreenPassTexture AddBloomSetupPass(FRDGBuilder& GraphBuilder, const FViewInfo&
 			Viewport,
 			FScreenPassPipelineState(VertexShader, PixelShader),
 			PassParameters,
-			[VertexShader, PixelShader, PassParameters] (FRHICommandListImmediate& RHICmdList)
+			[VertexShader, PixelShader, PassParameters] (FRHICommandList& RHICmdList)
 		{
 			SetShaderParameters(RHICmdList, VertexShader, VertexShader.GetVertexShader(), PassParameters->BloomSetup);
 			SetShaderParameters(RHICmdList, PixelShader, PixelShader.GetPixelShader(), *PassParameters);
diff --git a/Engine/Source/Runtime/Renderer/Private/PostProcess/PostProcessCompositeEditorPrimitives.cpp b/Engine/Source/Runtime/Renderer/Private/PostProcess/PostProcessCompositeEditorPrimitives.cpp
--- a/Engine/Source/Runtime/Renderer/Private/PostProcess/PostProcessCompositeEditorPrimitives.cpp
+++ b/Engine/Source/Runtime/Renderer/Private/PostProcess/PostProcessCompositeEditorPrimitives.cpp
@@ -125,7 +125,7 @@ IMPLEMENT_GLOBAL_SHADER(FTemporalUpsampleEditorDepthPS, "/Engine/Private/PostPro
 IMPLEMENT_GLOBAL_SHADER(FPopulateEditorDepthPS, "/Engine/Private/PostProcessCompositeEditorPrimitives.usf", "MainPopulateSceneDepthPS", SF_Pixel);
 IMPLEMENT_GLOBAL_SHADER(FCompositeEditorPrimitivesPS, "/Engine/Private/PostProcessCompositeEditorPrimitives.usf", "MainCompositeEditorPrimitivesPS", SF_Pixel);
 
-void RenderEditorPrimitives(FRHICommandListImmediate& RHICmdList, const FViewInfo& View, FMeshPassProcessorRenderState& DrawRenderState, FInstanceCullingManager& InstanceCullingManager)
+void RenderEditorPrimitives(FRHICommandList& RHICmdList, const FViewInfo& View, FMeshPassProcessorRenderState& DrawRenderState, FInstanceCullingManager& InstanceCullingManager)
 {
 	// Always depth test against other editor primitives
 	DrawRenderState.SetDepthStencilState(TStaticDepthStencilState<
@@ -175,7 +175,7 @@ void RenderEditorPrimitives(FRHICommandListImmediate& RHICmdList, const FViewInf
 	View.BatchedViewElements.Draw(RHICmdList, DrawRenderState, FeatureLevel, bNeedToSwitchVerticalAxis, View, false, 1.0f);
 }
 
-void RenderForegroundEditorPrimitives(FRHICommandListImmediate& RHICmdList, const FViewInfo& View, FMeshPassProcessorRenderState& DrawRenderState, FInstanceCullingManager& InstanceCullingManager)
+void RenderForegroundEditorPrimitives(FRHICommandList& RHICmdList, const FViewInfo& View, FMeshPassProcessorRenderState& DrawRenderState, FInstanceCullingManager& InstanceCullingManager)
 {
 	const auto FeatureLevel = View.GetFeatureLevel();
 	const auto ShaderPlatform = GShaderPlatformForFeatureLevel[FeatureLevel];
@@ -493,7 +493,7 @@ FScreenPassTexture AddEditorPrimitivePass(
 					NumSamples),
 				PassParameters,
 				ERDGPassFlags::Raster,
-				[&View, &InstanceCullingManager, PassParameters, EditorView, EditorPrimitivesViewport, BasePassType, NumSamples](FRHICommandListImmediate& RHICmdList)
+				[&View, &InstanceCullingManager, PassParameters, EditorView, EditorPrimitivesViewport, BasePassType, NumSamples](FRHICommandList& RHICmdList)
 			{
 				RHICmdList.SetViewport(EditorPrimitivesViewport.Rect.Min.X, EditorPrimitivesViewport.Rect.Min.Y, 0.0f, EditorPrimitivesViewport.Rect.Max.X, EditorPrimitivesViewport.Rect.Max.Y, 1.0f);
 
diff --git a/Engine/Source/Runtime/Renderer/Private/PostProcess/PostProcessMaterial.cpp b/Engine/Source/Runtime/Renderer/Private/PostProcess/PostProcessMaterial.cpp
--- a/Engine/Source/Runtime/Renderer/Private/PostProcess/PostProcessMaterial.cpp
+++ b/Engine/Source/Runtime/Renderer/Private/PostProcess/PostProcessMaterial.cpp
@@ -569,7 +569,7 @@ FScreenPassTexture AddPostProcessMaterialPass(
 		FScreenPassPipelineState(VertexShader, PixelShader, BlendState, DepthStencilState),
 		PostProcessMaterialParameters,
 		ScreenPassFlags,
-		[&View, VertexShader, PixelShader, MaterialRenderProxy, Material, PostProcessMaterialParameters, MaterialStencilRef](FRHICommandListImmediate& RHICmdList)
+		[&View, VertexShader, PixelShader, MaterialRenderProxy, Material, PostProcessMaterialParameters, MaterialStencilRef](FRHICommandList& RHICmdList)
 	{
 		FPostProcessMaterialVS::SetParameters(RHICmdList, VertexShader, View, MaterialRenderProxy, *Material, *PostProcessMaterialParameters);
 		FPostProcessMaterialPS::SetParameters(RHICmdList, PixelShader, View, MaterialRenderProxy, *Material, *PostProcessMaterialParameters);
diff --git a/Engine/Source/Runtime/Renderer/Private/PostProcess/PostProcessWeightedSampleSum.cpp b/Engine/Source/Runtime/Renderer/Private/PostProcess/PostProcessWeightedSampleSum.cpp
--- a/Engine/Source/Runtime/Renderer/Private/PostProcess/PostProcessWeightedSampleSum.cpp
+++ b/Engine/Source/Runtime/Renderer/Private/PostProcess/PostProcessWeightedSampleSum.cpp
@@ -425,7 +425,7 @@ FScreenPassTexture AddGaussianBlurPass(
 				FScreenPassTextureViewport(Filter),
 				FScreenPassPipelineState(VertexShader, PixelShader),
 				PassParameters,
-				[VertexShader, PixelShader, PassParameters] (FRHICommandListImmediate& RHICmdList)
+				[VertexShader, PixelShader, PassParameters] (FRHICommandList& RHICmdList)
 			{
 				SetShaderParameters(RHICmdList, VertexShader, VertexShader.GetVertexShader(), PassParameters->Filter);
 				SetShaderParameters(RHICmdList, PixelShader, PixelShader.GetPixelShader(), *PassParameters);
diff --git a/Engine/Source/Runtime/Renderer/Private/ScreenPass.h b/Engine/Source/Runtime/Renderer/Private/ScreenPass.h
--- a/Engine/Source/Runtime/Renderer/Private/ScreenPass.h
+++ b/Engine/Source/Runtime/Renderer/Private/ScreenPass.h
@@ -390,7 +390,7 @@ ENUM_CLASS_FLAGS(EScreenPassDrawFlags);
  */
 template<typename TSetupFunction>
 void DrawScreenPass(
-	FRHICommandListImmediate& RHICmdList,
+	FRHICommandList& RHICmdList,
 	const FViewInfo& View,
 	const FScreenPassTextureViewport& OutputViewport,
 	const FScreenPassTextureViewport& InputViewport,
@@ -472,9 +472,9 @@ FORCEINLINE void AddDrawScreenPass(
 		Forward<FRDGEventName&&>(PassName),
 		PixelShaderParameters,
 		ERDGPassFlags::Raster,
-		[&View, OutputViewport, InputViewport, PipelineState, PixelShader, PixelShaderParameters, Flags](FRHICommandListImmediate& RHICmdList)
+		[&View, OutputViewport, InputViewport, PipelineState, PixelShader, PixelShaderParameters, Flags](FRHICommandList& RHICmdList)
 	{
-		DrawScreenPass(RHICmdList, View, OutputViewport, InputViewport, PipelineState, Flags, [&](FRHICommandListImmediate&)
+		DrawScreenPass(RHICmdList, View, OutputViewport, InputViewport, PipelineState, Flags, [&](FRHICommandList&)
 		{
 			SetShaderParameters(RHICmdList, PixelShader, PixelShader.GetPixelShader(), *PixelShaderParameters);
 		});
@@ -571,7 +571,7 @@ FORCEINLINE void AddDrawScreenPass(
 		Forward<FRDGEventName&&>(PassName),
 		PassParameterStruct,
 		ERDGPassFlags::Raster,
-		[&View, OutputViewport, InputViewport, PipelineState, SetupFunction, Flags] (FRHICommandListImmediate& RHICmdList)
+		[&View, OutputViewport, InputViewport, PipelineState, SetupFunction, Flags] (FRHICommandList& RHICmdList)
 	{
 		DrawScreenPass(RHICmdList, View, OutputViewport, InputViewport, PipelineState, Flags, SetupFunction);
 	});
