Add OpenImageDenoise plugin for path tracer

This plugin implements a denoiser for path traced images using the OpenImageDenoise library. The denoiser is enabled by both loading the plugin and setting r.PathTracing.Denoiser=1

Add Albedo/Normal outputs to the path tracer to support the denoiser.

diff --git a/Engine/Plugins/Experimental/OpenImageDenoise/OpenImageDenoise.uplugin b/Engine/Plugins/Experimental/OpenImageDenoise/OpenImageDenoise.uplugin
new file mode 100644
--- /dev/null
+++ b/Engine/Plugins/Experimental/OpenImageDenoise/OpenImageDenoise.uplugin
@@ -0,0 +1,28 @@
+{
+	"FileVersion" : 3,
+	"Version" : 1,
+	"VersionName" : "1.0",
+	"FriendlyName" : "OpenImageDenoise",
+	"Description" : "Denoising engine for the Unreal Path Tracer based on Intel's OpenImageDenoise library.",
+	"Category" : "Denoising",
+	"CreatedBy" : "Epic Games, Inc.",
+	"CreatedByURL" : "http://epicgames.com",
+	"DocsURL" : "",
+	"MarketplaceURL" : "",
+	"SupportURL" : "",
+	"EnabledByDefault" : false,
+	"CanContainContent" : false,
+	"IsBetaVersion" : false,
+	"Installed" : false,
+	"Modules" :
+	[
+		{
+			"Name" : "OpenImageDenoise",
+			"Type" : "RuntimeAndProgram",
+			"LoadingPhase" : "Default",
+			"WhitelistPlatforms": [
+				"Win64"
+			]
+		}
+	]
+}
diff --git a/Engine/Plugins/Experimental/OpenImageDenoise/Source/OpenImageDenoise/OpenImageDenoise.Build.cs b/Engine/Plugins/Experimental/OpenImageDenoise/Source/OpenImageDenoise/OpenImageDenoise.Build.cs
new file mode 100644
--- /dev/null
+++ b/Engine/Plugins/Experimental/OpenImageDenoise/Source/OpenImageDenoise/OpenImageDenoise.Build.cs
@@ -0,0 +1,26 @@
+// Copyright Epic Games, Inc. All Rights Reserved.
+
+namespace UnrealBuildTool.Rules
+{
+	public class OpenImageDenoise : ModuleRules
+	{
+		public OpenImageDenoise(ReadOnlyTargetRules Target) : base(Target)
+		{
+			PrivateDependencyModuleNames.AddRange(
+				new string[] {
+					"Core",
+					"CoreUObject",
+					"Engine",
+					"RenderCore",
+					"Renderer",
+					"RHI",
+					"IntelOIDN"
+				}
+			);
+			if (Target.bBuildEditor)
+			{
+				PrivateDependencyModuleNames.Add("MessageLog");
+			}
+		}
+	}
+}
diff --git a/Engine/Plugins/Experimental/OpenImageDenoise/Source/OpenImageDenoise/Private/Plugin.cpp b/Engine/Plugins/Experimental/OpenImageDenoise/Source/OpenImageDenoise/Private/Plugin.cpp
new file mode 100644
--- /dev/null
+++ b/Engine/Plugins/Experimental/OpenImageDenoise/Source/OpenImageDenoise/Private/Plugin.cpp
@@ -0,0 +1,104 @@
+// Copyright Epic Games, Inc. All Rights Reserved.
+
+#include "CoreMinimal.h"
+#include "Modules/ModuleInterface.h"
+#include "Renderer/Public/PathTracingDenoiser.h"
+#include "RHIResources.h"
+
+#include "OpenImageDenoise/oidn.hpp"
+
+class FOpenImageDenoiseModule : public IModuleInterface
+{
+public:
+	/** IModuleInterface implementation */
+	virtual void StartupModule() override;
+	virtual void ShutdownModule() override;
+};
+
+#if WITH_EDITOR
+DECLARE_LOG_CATEGORY_EXTERN(LogOpenImageDenoise, Log, All);
+DEFINE_LOG_CATEGORY(LogOpenImageDenoise);
+#endif
+
+IMPLEMENT_MODULE(FOpenImageDenoiseModule, OpenImageDenoise)
+
+static void Denoise(FRHICommandListImmediate& RHICmdList, FRHITexture2D* ColorTex, FRHITexture2D* AlbedoTex, FRHITexture2D* NormalTex, FRHITexture2D* OutputTex)
+{
+	const int DenoiserMode = 2; // TODO: Expose setting for this
+
+#if WITH_EDITOR
+	uint64 FilterExecuteTime = 0;
+	FilterExecuteTime -= FPlatformTime::Cycles64();
+#endif
+
+	FIntPoint Size = ColorTex->GetSizeXY();
+	FIntRect Rect = FIntRect(0, 0, Size.X, Size.Y);
+	TArray<FLinearColor> RawPixels;
+	TArray<FLinearColor> RawAlbedo;
+	TArray<FLinearColor> RawNormal;
+	FReadSurfaceDataFlags ReadDataFlags(ERangeCompressionMode::RCM_MinMax);
+	ReadDataFlags.SetLinearToGamma(false);
+	RHICmdList.ReadSurfaceData(ColorTex, Rect, RawPixels, ReadDataFlags);
+	if (DenoiserMode >= 2)
+	{
+		RHICmdList.ReadSurfaceData(AlbedoTex, Rect, RawAlbedo, ReadDataFlags);
+		RHICmdList.ReadSurfaceData(NormalTex, Rect, RawNormal, ReadDataFlags);
+	}
+
+	check(RawPixels.Num() == Size.X * Size.Y);
+
+	uint32_t DestStride;
+	FLinearColor* DestBuffer = (FLinearColor*)RHICmdList.LockTexture2D(OutputTex, 0, RLM_WriteOnly, DestStride, false);
+
+	// create device only once?
+	oidn::DeviceRef OIDNDevice = oidn::newDevice();
+	OIDNDevice.commit();
+	oidn::FilterRef OIDNFilter = OIDNDevice.newFilter("RT");
+	OIDNFilter.setImage("color", RawPixels.GetData(), oidn::Format::Float3, Size.X, Size.Y, 0, sizeof(FLinearColor), sizeof(FLinearColor) * Size.X);
+	if (DenoiserMode >= 2)
+	{
+		OIDNFilter.setImage("albedo", RawAlbedo.GetData(), oidn::Format::Float3, Size.X, Size.Y, 0, sizeof(FLinearColor), sizeof(FLinearColor) * Size.X);
+		OIDNFilter.setImage("normal", RawNormal.GetData(), oidn::Format::Float3, Size.X, Size.Y, 0, sizeof(FLinearColor), sizeof(FLinearColor) * Size.X);
+	}
+	if (DenoiserMode >= 3)
+	{
+		OIDNFilter.set("cleanAux", true);
+	}
+	OIDNFilter.setImage("output", DestBuffer, oidn::Format::Float3, Size.X, Size.Y, 0, sizeof(FLinearColor), DestStride);
+	OIDNFilter.set("hdr", true);
+	OIDNFilter.commit();
+
+	OIDNFilter.execute();
+
+	RHICmdList.UnlockTexture2D(OutputTex, 0, false);
+
+#if WITH_EDITOR
+	const char* errorMessage;
+	if (OIDNDevice.getError(errorMessage) != oidn::Error::None)
+	{
+		UE_LOG(LogOpenImageDenoise, Warning, TEXT("Denoiser failed: %s"), *FString(errorMessage));
+		return;
+	}
+
+	FilterExecuteTime += FPlatformTime::Cycles64();
+	const double FilterExecuteTimeMS = 1000.0 * FPlatformTime::ToSeconds64(FilterExecuteTime);
+	UE_LOG(LogOpenImageDenoise, Log, TEXT("Denoised %d x %d pixels in %.2f ms"), Size.X, Size.Y, FilterExecuteTimeMS);
+#endif
+}
+
+
+void FOpenImageDenoiseModule::StartupModule()
+{
+#if WITH_EDITOR
+	UE_LOG(LogOpenImageDenoise, Log, TEXT("OIDN starting up"));
+#endif
+	GPathTracingDenoiserFunc = &Denoise;
+}
+
+void FOpenImageDenoiseModule::ShutdownModule()
+{
+#if WITH_EDITOR
+	UE_LOG(LogOpenImageDenoise, Log, TEXT("OIDN shutting down"));
+#endif
+	GPathTracingDenoiserFunc = nullptr;
+}
diff --git a/Engine/Shaders/Private/PathTracing/PathTracing.usf b/Engine/Shaders/Private/PathTracing/PathTracing.usf
--- a/Engine/Shaders/Private/PathTracing/PathTracing.usf
+++ b/Engine/Shaders/Private/PathTracing/PathTracing.usf
@@ -49,6 +49,8 @@ float MaxNormalBias;
 float FilterWidth;
 
 RWTexture2D<float4> RadianceTexture;
+RWTexture2D<float4> AlbedoTexture;
+RWTexture2D<float4> NormalTexture;
 RaytracingAccelerationStructure TLAS;
 int3 TileOffset;
 uint SceneVisibleLightCount;
@@ -69,7 +71,7 @@ void AccumulateRadiance(inout float3 TotalRadiance, float3 PathRadiance, bool bI
 	TotalRadiance += PathRadiance;
 }
 
-FMaterialClosestHitPayload TraceTransparentRay(RayDesc Ray, FRayCone RayCone, bool bIsCameraRay, bool bLastBounce, bool bIncludeEmission, uint2 LaunchIndex, uint NumLights, inout RandomSequence RandSequence, inout float3 PathThroughput, inout float3 Radiance)
+FMaterialClosestHitPayload TraceTransparentRay(RayDesc Ray, FRayCone RayCone, bool bIsCameraRay, bool bLastBounce, bool bIncludeEmission, uint2 LaunchIndex, uint NumLights, inout RandomSequence RandSequence, inout float3 PathThroughput, inout float3 Radiance, inout float3 Albedo, inout float3 Normal)
 {
 	const uint RayFlags = bIsCameraRay && EnableCameraBackfaceCulling ? RAY_FLAG_CULL_BACK_FACING_TRIANGLES : 0;
 	const uint MissShaderIndex = 0;
@@ -143,6 +145,11 @@ FMaterialClosestHitPayload TraceTransparentRay(RayDesc Ray, FRayCone RayCone, bo
 		if (!bLastBounce)
 		{
 			float3 Contrib = PathThroughput * EstimateMaterialAlbedo(HitPayload);
+			if (bIsCameraRay)
+			{
+				Albedo += Contrib;
+				Normal += PathThroughput * (1 - Luminance(Transparency)) * HitPayload.WorldNormal;
+			}
 
 			float SelectionWeight = max3(Contrib.x, Contrib.y, Contrib.z);
 			SelectionWeightSum += SelectionWeight;
@@ -517,6 +524,8 @@ RAY_TRACING_ENTRY_RAYGEN(PathTracingMainRG)
 
 
 	float3 Radiance = 0;
+	float3 Albedo = 0;
+	float3 Normal = 0;
 
 	// Initialize ray and payload
 	RayDesc Ray;
@@ -551,7 +560,7 @@ RAY_TRACING_ENTRY_RAYGEN(PathTracingMainRG)
 		const bool bIsLastBounce = Bounce == MaxBounces;
 		const bool bIncludeEmissive = (EnableDirectLighting != 0 || Bounce > 1) &&
 									  (EnableEmissive       != 0 || bIsCameraRay);
-		FMaterialClosestHitPayload Payload = TraceTransparentRay(Ray, RayCone, bIsCameraRay, bIsLastBounce, bIncludeEmissive, LaunchIndex, NumVisibleLights, RandSequence, PathThroughput, Radiance);
+		FMaterialClosestHitPayload Payload = TraceTransparentRay(Ray, RayCone, bIsCameraRay, bIsLastBounce, bIncludeEmissive, LaunchIndex, NumVisibleLights, RandSequence, PathThroughput, Radiance, Albedo, Normal);
 
 		if (Payload.IsMiss())
 		{
@@ -790,6 +799,8 @@ RAY_TRACING_ENTRY_RAYGEN(PathTracingMainRG)
 	float BlendFactor = 1.0 / float(NumSamples);
 	// Avoid reading the old pixel on the first sample on the off-chance there is a NaN/Inf pixel ...
 	float4 OldPixel = NumSamples > 1 ? RadianceTexture[PixelIndex] : 0;
+	float4 OldAlbedo = NumSamples > 1 ? AlbedoTexture[PixelIndex] : 0;
+	float4 OldNormal = NumSamples > 1 ? NormalTexture[PixelIndex] : 0;
 	float3 OldRadiance = OldPixel.rgb;
 	float OldVariance = OldPixel.a;
 
@@ -802,4 +813,6 @@ RAY_TRACING_ENTRY_RAYGEN(PathTracingMainRG)
 	float NewVariance = lerp(OldVariance, DeviationSquared, BlendFactor);
 
 	RadianceTexture[PixelIndex] = float4(NewRadiance, NewVariance);
+	AlbedoTexture[PixelIndex] = lerp(OldAlbedo, float4(Albedo, 0), BlendFactor);
+	NormalTexture[PixelIndex] = lerp(OldNormal, float4(Normal, 0), BlendFactor);
 }
diff --git a/Engine/Source/Runtime/Renderer/Private/PathTracing.cpp b/Engine/Source/Runtime/Renderer/Private/PathTracing.cpp
--- a/Engine/Source/Runtime/Renderer/Private/PathTracing.cpp
+++ b/Engine/Source/Runtime/Renderer/Private/PathTracing.cpp
@@ -2,6 +2,9 @@
 
 #include "PathTracing.h"
 #include "RHI.h"
+#include "PathTracingDenoiser.h"
+
+PathTracingDenoiserFunction* GPathTracingDenoiserFunc = nullptr;
 
 TAutoConsoleVariable<int32> CVarPathTracing(
 	TEXT("r.PathTracing"),
@@ -186,6 +189,15 @@ TAutoConsoleVariable<int32> CVarPathTracingLightGridVisualize(
 	ECVF_RenderThreadSafe
 );
 
+TAutoConsoleVariable<int32> CVarPathTracingDenoiser(
+	TEXT("r.PathTracing.Denoiser"),
+	0,
+	TEXT("Enable denoising of the path traced output (default = 0)\n")
+	TEXT("0: off (default)")
+	TEXT("1: enabled (if a denoiser plugin is active)"),
+	ECVF_RenderThreadSafe
+);
+
 BEGIN_SHADER_PARAMETER_STRUCT(FPathTracingData, )
 	SHADER_PARAMETER(uint32, Iteration)
 	SHADER_PARAMETER(uint32, TemporalSeed)
@@ -244,7 +256,8 @@ struct FPathTracingConfig
 };
 
 // This function prepares the portion of shader arguments that may involve invalidating the path traced state
-static void PrepareShaderArgs(const FViewInfo& View, FPathTracingData& PathTracingData) {
+static void PrepareShaderArgs(const FViewInfo& View, FPathTracingData& PathTracingData)
+{
 	PathTracingData.EnableDirectLighting = true;
 	int32 MaxBounces = CVarPathTracingMaxBounces.GetValueOnRenderThread();
 	if (MaxBounces < 0)
@@ -424,6 +437,8 @@ class FPathTracingRG : public FGlobalShader
 
 	BEGIN_SHADER_PARAMETER_STRUCT(FParameters, )
 		SHADER_PARAMETER_RDG_TEXTURE_UAV(RWTexture2D<float4>, RadianceTexture)
+		SHADER_PARAMETER_RDG_TEXTURE_UAV(RWTexture2D<float4>, AlbedoTexture)
+		SHADER_PARAMETER_RDG_TEXTURE_UAV(RWTexture2D<float4>, NormalTexture)
 		SHADER_PARAMETER_SRV(RaytracingAccelerationStructure, TLAS)
 
 		SHADER_PARAMETER_STRUCT_REF(FViewUniformShaderParameters, ViewUniformBuffer)
@@ -1191,9 +1206,19 @@ void FDeferredShadingSceneRenderer::PreparePathTracing(const FSceneViewFamily& V
 void FSceneViewState::PathTracingInvalidate()
 {
 	PathTracingRadianceRT.SafeRelease();
+	PathTracingAlbedoRT.SafeRelease();
+	PathTracingNormalRT.SafeRelease();
+	PathTracingRadianceDenoisedRT.SafeRelease();
 	PathTracingSampleIndex = 0;
 }
 
+BEGIN_SHADER_PARAMETER_STRUCT(FDenoiseTextureParameters, )
+	RDG_TEXTURE_ACCESS(InputTexture, ERHIAccess::CopySrc)
+	RDG_TEXTURE_ACCESS(InputAlbedo, ERHIAccess::CopySrc)
+	RDG_TEXTURE_ACCESS(InputNormal, ERHIAccess::CopySrc)
+	RDG_TEXTURE_ACCESS(OutputTexture, ERHIAccess::CopyDest)
+END_SHADER_PARAMETER_STRUCT()
+
 DECLARE_GPU_STAT_NAMED(Stat_GPU_PathTracing, TEXT("Path Tracing"));
 void FDeferredShadingSceneRenderer::RenderPathTracing(
 	FRDGBuilder& GraphBuilder,
@@ -1286,20 +1311,26 @@ void FDeferredShadingSceneRenderer::RenderPathTracing(
 
 	// Prepare radiance buffer (will be shared with display pass)
 	FRDGTexture* RadianceTexture = nullptr;
+	FRDGTexture* AlbedoTexture = nullptr;
+	FRDGTexture* NormalTexture = nullptr;
 	if (View.ViewState->PathTracingRadianceRT)
 	{
 		// we already have a valid radiance texture, re-use it
 		RadianceTexture = GraphBuilder.RegisterExternalTexture(View.ViewState->PathTracingRadianceRT, TEXT("PathTracer.Radiance"));
+		AlbedoTexture   = GraphBuilder.RegisterExternalTexture(View.ViewState->PathTracingAlbedoRT, TEXT("PathTracer.Albedo"));
+		NormalTexture   = GraphBuilder.RegisterExternalTexture(View.ViewState->PathTracingNormalRT, TEXT("PathTracer.Normal"));
 	}
 	else
 	{
 		// First time through, need to make a new texture
-		FRDGTextureDesc RadianceTextureDesc = FRDGTextureDesc::Create2D(
+		FRDGTextureDesc Desc = FRDGTextureDesc::Create2D(
 			View.ViewRect.Size(),
 			PF_A32B32G32R32F,
 			FClearValueBinding::None,
 			TexCreate_ShaderResource | TexCreate_UAV);
-		RadianceTexture = GraphBuilder.CreateTexture(RadianceTextureDesc, TEXT("PathTracer.Radiance"), ERDGTextureFlags::MultiFrame);
+		RadianceTexture = GraphBuilder.CreateTexture(Desc, TEXT("PathTracer.Radiance"), ERDGTextureFlags::MultiFrame);
+		AlbedoTexture   = GraphBuilder.CreateTexture(Desc, TEXT("PathTracer.Albedo")  , ERDGTextureFlags::MultiFrame);
+		NormalTexture   = GraphBuilder.CreateTexture(Desc, TEXT("PathTracer.Normal")  , ERDGTextureFlags::MultiFrame);
 	}
 	const bool bNeedsMoreRays = Config.PathTracingData.Iteration < MaxSPP;
 
@@ -1318,6 +1349,8 @@ void FDeferredShadingSceneRenderer::RenderPathTracing(
 
 		PassParameters->IESTextureSampler = TStaticSamplerState<SF_Bilinear, AM_Clamp, AM_Clamp, AM_Clamp>::GetRHI();
 		PassParameters->RadianceTexture = GraphBuilder.CreateUAV(RadianceTexture);
+		PassParameters->AlbedoTexture   = GraphBuilder.CreateUAV(AlbedoTexture);
+		PassParameters->NormalTexture   = GraphBuilder.CreateUAV(NormalTexture);
 
 		PassParameters->SSProfilesTexture = GetSubsufaceProfileTexture_RT(GraphBuilder.RHICmdList)->GetShaderResourceRHI();
 
@@ -1354,7 +1387,65 @@ void FDeferredShadingSceneRenderer::RenderPathTracing(
 
 		// After we are done, make sure we remember our texture for next time so that we can accumulate samples across frames
 		GraphBuilder.QueueTextureExtraction(RadianceTexture, &View.ViewState->PathTracingRadianceRT);
+		GraphBuilder.QueueTextureExtraction(AlbedoTexture  , &View.ViewState->PathTracingAlbedoRT  );
+		GraphBuilder.QueueTextureExtraction(NormalTexture  , &View.ViewState->PathTracingNormalRT  );
+
+	}
+
+	FRDGTexture* DenoisedRadianceTexture = nullptr;
+	const int DenoiserMode = CVarPathTracingDenoiser.GetValueOnRenderThread();
+	const bool IsDenoiserEnabled = DenoiserMode != 0 && GPathTracingDenoiserFunc != nullptr;
+	static int PrevDenoiserMode = DenoiserMode;
+	if (IsDenoiserEnabled)
+	{
+		// request denoise if this is the last sample
+		bool NeedsDenoise = (Config.PathTracingData.Iteration + 1) == MaxSPP;
+		// also allow turning on the denoiser after the image has stopped accumulating samples
+		if (!bNeedsMoreRays)
+		{
+			// we aren't currently rendering, run the denoiser if we just turned it on
+			NeedsDenoise |= DenoiserMode != PrevDenoiserMode;
+		}
+
+		if (View.ViewState->PathTracingRadianceDenoisedRT)
+		{
+			// we already have a texture for this
+			DenoisedRadianceTexture = GraphBuilder.RegisterExternalTexture(View.ViewState->PathTracingRadianceDenoisedRT, TEXT("PathTracer.DenoisedRadiance"));
+		}
+
+		if (NeedsDenoise)
+		{
+			if (DenoisedRadianceTexture == nullptr)
+			{
+				// First time through, need to make a new texture
+				FRDGTextureDesc RadianceTextureDesc = FRDGTextureDesc::Create2D(
+					View.ViewRect.Size(),
+					PF_A32B32G32R32F,
+					FClearValueBinding::None,
+					TexCreate_ShaderResource | TexCreate_UAV);
+				DenoisedRadianceTexture = GraphBuilder.CreateTexture(RadianceTextureDesc, TEXT("PathTracer.DenoisedRadiance"), ERDGTextureFlags::MultiFrame);
+			}
+
+			FDenoiseTextureParameters* DenoiseParameters = GraphBuilder.AllocParameters<FDenoiseTextureParameters>();
+			DenoiseParameters->InputTexture = RadianceTexture;
+			DenoiseParameters->InputAlbedo = AlbedoTexture;
+			DenoiseParameters->InputNormal = NormalTexture;
+			DenoiseParameters->OutputTexture = DenoisedRadianceTexture;
+			GraphBuilder.AddPass(RDG_EVENT_NAME("Path Tracer Denoiser Plugin"), DenoiseParameters, ERDGPassFlags::Readback, 
+				[DenoiseParameters, DenoiserMode](FRHICommandListImmediate& RHICmdList)
+				{
+					GPathTracingDenoiserFunc(RHICmdList,
+						DenoiseParameters->InputTexture->GetRHI()->GetTexture2D(),
+						DenoiseParameters->InputAlbedo->GetRHI()->GetTexture2D(),
+						DenoiseParameters->InputNormal->GetRHI()->GetTexture2D(),
+						DenoiseParameters->OutputTexture->GetRHI()->GetTexture2D());
+				}
+			);
+
+			GraphBuilder.QueueTextureExtraction(DenoisedRadianceTexture, &View.ViewState->PathTracingRadianceDenoisedRT);
+		}
 	}
+	PrevDenoiserMode = DenoiserMode;
 
 	// now add a pixel shader pass to display our Radiance buffer
 
@@ -1363,7 +1454,7 @@ void FDeferredShadingSceneRenderer::RenderPathTracing(
 	DisplayParameters->MaxSamples = MaxSPP;
 	DisplayParameters->ProgressDisplayEnabled = CVarPathTracingProgressDisplay.GetValueOnRenderThread();
 	DisplayParameters->ViewUniformBuffer = View.ViewUniformBuffer;
-	DisplayParameters->RadianceTexture = GraphBuilder.CreateSRV(FRDGTextureSRVDesc::Create(RadianceTexture));
+	DisplayParameters->RadianceTexture = GraphBuilder.CreateSRV(FRDGTextureSRVDesc::Create(DenoisedRadianceTexture ? DenoisedRadianceTexture : RadianceTexture));
 	DisplayParameters->RenderTargets[0] = FRenderTargetBinding(SceneColorOutputTexture, ERenderTargetLoadAction::ELoad);
 
 	FScreenPassTextureViewport Viewport(SceneColorOutputTexture, View.ViewRect);
diff --git a/Engine/Source/Runtime/Renderer/Private/ScenePrivate.h b/Engine/Source/Runtime/Renderer/Private/ScenePrivate.h
--- a/Engine/Source/Runtime/Renderer/Private/ScenePrivate.h
+++ b/Engine/Source/Runtime/Renderer/Private/ScenePrivate.h
@@ -980,6 +980,9 @@ public:
 
 	// Reference path tracing cached results
 	TRefCountPtr<IPooledRenderTarget> PathTracingRadianceRT;
+	TRefCountPtr<IPooledRenderTarget> PathTracingAlbedoRT;
+	TRefCountPtr<IPooledRenderTarget> PathTracingNormalRT;
+	TRefCountPtr<IPooledRenderTarget> PathTracingRadianceDenoisedRT;
 	// Keeps track of the internal path tracer options relevant to detecting when to restart the path tracer accumulation
 	TPimplPtr<FPathTracingConfig> PathTracingLastConfig;
 
diff --git a/Engine/Source/Runtime/Renderer/Public/PathTracingDenoiser.h b/Engine/Source/Runtime/Renderer/Public/PathTracingDenoiser.h
new file mode 100644
--- /dev/null
+++ b/Engine/Source/Runtime/Renderer/Public/PathTracingDenoiser.h
@@ -0,0 +1,10 @@
+// Copyright Epic Games, Inc. All Rights Reserved.
+
+#pragma once
+
+class FRHITexture2D;
+class FRHICommandListImmediate;
+
+using PathTracingDenoiserFunction = void(FRHICommandListImmediate& RHICmdList, FRHITexture2D* ColorTex, FRHITexture2D* AlbedoTex, FRHITexture2D* NormalTex, FRHITexture2D* OutputTex);
+
+extern RENDERER_API PathTracingDenoiserFunction* GPathTracingDenoiserFunc;
