Add a pass uniform buffer for decal passes.
Added EyeAdaptationTexture to the uniform buffer.
Fixes not being able to sample EyeAdaptation in decal materials.
Hand edited integration of CL16494400 from 4.27 release.
Needed to break out DBufferTextures into new file to prevent build include errors.

diff --git a/Engine/Shaders/Private/DecalCommon.ush b/Engine/Shaders/Private/DecalCommon.ush
--- a/Engine/Shaders/Private/DecalCommon.ush
+++ b/Engine/Shaders/Private/DecalCommon.ush
@@ -13,6 +13,10 @@
 #if SHADING_PATH_MOBILE
 	// Reroute MobileSceneTextures uniform buffer references to the base pass uniform buffer 
 	#define MobileSceneTextures MobileBasePass.SceneTextures
+	#define EyeAdaptationStruct MobileBasePass
+#else
+	#define SceneTexturesStruct DecalPass.SceneTextures
+	#define EyeAdaptationStruct DecalPass
 #endif
 
 #include "DeferredShadingCommon.ush"
diff --git a/Engine/Source/Runtime/Renderer/Private/BasePassRendering.h b/Engine/Source/Runtime/Renderer/Private/BasePassRendering.h
--- a/Engine/Source/Runtime/Renderer/Private/BasePassRendering.h
+++ b/Engine/Source/Runtime/Renderer/Private/BasePassRendering.h
@@ -16,7 +16,7 @@
 #include "SceneManagement.h"
 #include "Materials/Material.h"
 #include "PostProcess/SceneRenderTargets.h"
-#include "CompositionLighting/PostProcessDeferredDecals.h"
+#include "DBufferTextures.h"
 #include "LightMapRendering.h"
 #include "VelocityRendering.h"
 #include "MeshMaterialShaderType.h"
diff --git a/Engine/Source/Runtime/Renderer/Private/CompositionLighting/CompositionLighting.cpp b/Engine/Source/Runtime/Renderer/Private/CompositionLighting/CompositionLighting.cpp
--- a/Engine/Source/Runtime/Renderer/Private/CompositionLighting/CompositionLighting.cpp
+++ b/Engine/Source/Runtime/Renderer/Private/CompositionLighting/CompositionLighting.cpp
@@ -513,7 +513,7 @@ void ProcessBeforeBasePass(
 			// decals are before AmbientOcclusion so the decal can output a normal that AO is affected by
 			if (bDBuffer)
 			{
-				FDeferredDecalPassTextures DecalPassTextures = GetDeferredDecalPassTextures(GraphBuilder, SceneTextures, &DBufferTextures);
+				FDeferredDecalPassTextures DecalPassTextures = GetDeferredDecalPassTextures(GraphBuilder, View, SceneTextures, &DBufferTextures);
 				AddDeferredDecalPass(GraphBuilder, View, DecalPassTextures, EDecalRenderStage::BeforeBasePass);
 			}
 
@@ -572,7 +572,7 @@ void ProcessAfterBasePass(
 	const bool bDoDecal = ViewFamily.EngineShowFlags.Decals && !ViewFamily.EngineShowFlags.VisualizeLightCulling;
 	const bool bDBuffer = IsUsingDBuffers(View.GetShaderPlatform());
 
-	FDeferredDecalPassTextures DecalPassTextures = GetDeferredDecalPassTextures(GraphBuilder, SceneTextures, nullptr);
+	FDeferredDecalPassTextures DecalPassTextures = GetDeferredDecalPassTextures(GraphBuilder, View, SceneTextures, nullptr);
 
 	if (bDoDecal && !bDBuffer && IsUsingGBuffers(View.GetShaderPlatform()))
 	{
diff --git a/Engine/Source/Runtime/Renderer/Private/CompositionLighting/PostProcessDeferredDecals.cpp b/Engine/Source/Runtime/Renderer/Private/CompositionLighting/PostProcessDeferredDecals.cpp
--- a/Engine/Source/Runtime/Renderer/Private/CompositionLighting/PostProcessDeferredDecals.cpp
+++ b/Engine/Source/Runtime/Renderer/Private/CompositionLighting/PostProcessDeferredDecals.cpp
@@ -1,13 +1,15 @@
 // Copyright Epic Games, Inc. All Rights Reserved.
 
 #include "CompositionLighting/PostProcessDeferredDecals.h"
-#include "SceneUtils.h"
-#include "ScenePrivate.h"
-#include "DecalRenderingShared.h"
+
 #include "ClearQuad.h"
+#include "DBufferTextures.h"
+#include "DecalRenderingShared.h"
 #include "PipelineStateCache.h"
-#include "VisualizeTexture.h"
 #include "RendererUtils.h"
+#include "SceneUtils.h"
+#include "ScenePrivate.h"
+#include "VisualizeTexture.h"
 
 static TAutoConsoleVariable<float> CVarStencilSizeThreshold(
 	TEXT("r.Decal.StencilSizeThreshold"),
@@ -53,38 +55,22 @@ bool IsDBufferEnabled(const FSceneViewFamily& ViewFamily, EShaderPlatform Shader
 	return !ViewFamily.EngineShowFlags.ShaderComplexity && ViewFamily.EngineShowFlags.Decals && IsUsingDBuffers(ShaderPlatform);
 }
 
-FDBufferParameters GetDBufferParameters(FRDGBuilder& GraphBuilder, const FDBufferTextures& DBufferTextures, EShaderPlatform ShaderPlatform)
-{
-	const FRDGSystemTextures& SystemTextures = FRDGSystemTextures::Get(GraphBuilder);
-
-	FDBufferParameters Parameters;
-	Parameters.DBufferATextureSampler = TStaticSamplerState<>::GetRHI();
-	Parameters.DBufferBTextureSampler = TStaticSamplerState<>::GetRHI();
-	Parameters.DBufferCTextureSampler = TStaticSamplerState<>::GetRHI();
-	Parameters.DBufferATexture = SystemTextures.BlackAlphaOne;
-	Parameters.DBufferBTexture = SystemTextures.DefaultNormal8Bit;
-	Parameters.DBufferCTexture = SystemTextures.BlackAlphaOne;
-	Parameters.DBufferRenderMask = SystemTextures.White;
-
-	if (DBufferTextures.IsValid())
-	{
-		Parameters.DBufferATexture = DBufferTextures.DBufferA;
-		Parameters.DBufferBTexture = DBufferTextures.DBufferB;
-		Parameters.DBufferCTexture = DBufferTextures.DBufferC;
+IMPLEMENT_STATIC_UNIFORM_BUFFER_STRUCT(FDecalPassUniformParameters, "DecalPass", SceneTextures);
 
-		if (DBufferTextures.DBufferMask)
-		{
-			Parameters.DBufferRenderMask = DBufferTextures.DBufferMask;
-		}
-	}
-
-	return Parameters;
-}
-
-FDeferredDecalPassTextures GetDeferredDecalPassTextures(FRDGBuilder& GraphBuilder, const FSceneTextures& SceneTextures, FDBufferTextures* DBufferTextures)
+FDeferredDecalPassTextures GetDeferredDecalPassTextures(
+	FRDGBuilder& GraphBuilder, 
+	const FSceneView& View,
+	const FSceneTextures& SceneTextures, 
+	FDBufferTextures* DBufferTextures)
 {
 	FDeferredDecalPassTextures PassTextures;
-	PassTextures.SceneTexturesUniformBuffer = SceneTextures.UniformBuffer;
+
+	auto* Parameters = GraphBuilder.AllocParameters<FDecalPassUniformParameters>();
+	const ESceneTextureSetupMode TextureReadAccess = ESceneTextureSetupMode::GBufferA | ESceneTextureSetupMode::SceneDepth | ESceneTextureSetupMode::CustomDepth;
+	SetupSceneTextureUniformParameters(GraphBuilder, View.FeatureLevel, TextureReadAccess, Parameters->SceneTextures);
+	Parameters->EyeAdaptationTexture = GetEyeAdaptationTexture(GraphBuilder, View);
+	PassTextures.DecalPassUniformBuffer = GraphBuilder.CreateUniformBuffer(Parameters);
+
 	PassTextures.Depth = SceneTextures.Depth;
 	PassTextures.Color = SceneTextures.Color.Target;
 	PassTextures.GBufferA = (*SceneTextures.UniformBuffer)->GBufferATexture;
@@ -92,6 +78,7 @@ FDeferredDecalPassTextures GetDeferredDecalPassTextures(FRDGBuilder& GraphBuilde
 	PassTextures.GBufferC = (*SceneTextures.UniformBuffer)->GBufferCTexture;
 	PassTextures.GBufferE = (*SceneTextures.UniformBuffer)->GBufferETexture;
 	PassTextures.DBufferTextures = DBufferTextures;
+
 	return PassTextures;
 }
 
@@ -103,7 +90,7 @@ void GetDeferredDecalPassParameters(
 {
 	PassParameters.View = View.GetShaderParameters();
 	PassParameters.DeferredDecal = CreateDeferredDecalUniformBuffer(View);
-	PassParameters.SceneTextures = Textures.SceneTexturesUniformBuffer;
+	PassParameters.DecalPass = Textures.DecalPassUniformBuffer;
 
 	FRDGTextureRef DepthTexture = Textures.Depth.Target;
 
@@ -176,10 +163,6 @@ void GetDeferredDecalPassParameters(
 
 TUniformBufferRef<FDeferredDecalUniformParameters> CreateDeferredDecalUniformBuffer(const FViewInfo& View)
 {
-	// This extern is not pretty, but it's only temporary. The plan is to make velocity in DepthPass the only option, but requires some testing to 
-	// ensure quality and perf is what we expect.
-	extern bool IsVelocityMergedWithDepthPass();
-
 	bool bIsMotionInDepth = IsVelocityMergedWithDepthPass();
 	// if we have early motion vectors (bIsMotionInDepth) and the cvar is enabled and we actually have a buffer from the previous frame (View.PrevViewInfo.GBufferA.IsValid())
 	bool bIsNormalReprojectionEnabled = (bIsMotionInDepth && CVarDBufferDecalNormalReprojectionEnabled.GetValueOnRenderThread() && View.PrevViewInfo.GBufferA.IsValid());
@@ -218,18 +201,6 @@ enum EDecalDepthInputState
 	DDS_DepthTest_StencilEqual0,
 };
 
-enum class EDecalDBufferMaskTechnique
-{
-	// DBufferMask is not enabled.
-	Disabled,
-
-	// DBufferMask is written explicitly by the shader during the DBuffer pass.
-	PerPixel,
-
-	// DBufferMask is constructed after the DBuffer pass by compositing DBuffer write mask planes together in a compute shader.
-	WriteMask,
-};
-
 struct FDecalDepthState
 {
 	EDecalDepthInputState DepthTest;
@@ -423,23 +394,6 @@ static bool IsStencilOptimizationAvailable(EDecalRenderStage RenderStage)
 	return RenderStage == EDecalRenderStage::BeforeLighting || RenderStage == EDecalRenderStage::BeforeBasePass || RenderStage == EDecalRenderStage::Emissive;
 }
 
-static EDecalDBufferMaskTechnique GetDBufferMaskTechnique(EShaderPlatform ShaderPlatform)
-{
-	const bool bWriteMaskDBufferMask = RHISupportsRenderTargetWriteMask(ShaderPlatform);
-	const bool bPerPixelDBufferMask = IsUsingPerPixelDBufferMask(ShaderPlatform);
-	checkf(!bWriteMaskDBufferMask || !bPerPixelDBufferMask, TEXT("The WriteMask and PerPixel DBufferMask approaches cannot be enabled at the same time. They are mutually exclusive."));
-
-	if (bWriteMaskDBufferMask)
-	{
-		return EDecalDBufferMaskTechnique::WriteMask;
-	}
-	else if (bPerPixelDBufferMask)
-	{
-		return EDecalDBufferMaskTechnique::PerPixel;
-	}
-	return EDecalDBufferMaskTechnique::Disabled;
-}
-
 static const TCHAR* GetStageName(EDecalRenderStage Stage)
 {
 	switch (Stage)
@@ -454,48 +408,6 @@ static const TCHAR* GetStageName(EDecalRenderStage Stage)
 	return TEXT("<UNKNOWN>");
 }
 
-FDBufferTextures CreateDBufferTextures(FRDGBuilder& GraphBuilder, FIntPoint Extent, EShaderPlatform ShaderPlatform)
-{
-	FDBufferTextures DBufferTextures;
-
-	if (IsUsingDBuffers(ShaderPlatform))
-	{
-		const EDecalDBufferMaskTechnique DBufferMaskTechnique = GetDBufferMaskTechnique(ShaderPlatform);
-		const ETextureCreateFlags WriteMaskFlags = DBufferMaskTechnique == EDecalDBufferMaskTechnique::WriteMask ? TexCreate_NoFastClearFinalize | TexCreate_DisableDCC : TexCreate_None;
-		const ETextureCreateFlags BaseFlags = WriteMaskFlags | TexCreate_ShaderResource | TexCreate_RenderTargetable;
-		const ERDGTextureFlags TextureFlags = DBufferMaskTechnique != EDecalDBufferMaskTechnique::Disabled
-			? ERDGTextureFlags::MaintainCompression
-			: ERDGTextureFlags::None;
-
-		FRDGTextureDesc Desc = FRDGTextureDesc::Create2D(Extent, PF_B8G8R8A8, FClearValueBinding::None, BaseFlags);
-
-		Desc.Flags = BaseFlags | GFastVRamConfig.DBufferA;
-		Desc.ClearValue = FClearValueBinding::Black;
-		DBufferTextures.DBufferA = GraphBuilder.CreateTexture(Desc, TEXT("DBufferA"), TextureFlags);
-
-		Desc.Flags = BaseFlags | GFastVRamConfig.DBufferB;
-		Desc.ClearValue = FClearValueBinding(FLinearColor(128.0f / 255.0f, 128.0f / 255.0f, 128.0f / 255.0f, 1));
-		DBufferTextures.DBufferB = GraphBuilder.CreateTexture(Desc, TEXT("DBufferB"), TextureFlags);
-
-		Desc.Flags = BaseFlags | GFastVRamConfig.DBufferC;
-		Desc.ClearValue = FClearValueBinding(FLinearColor(0, 0, 0, 1));
-		DBufferTextures.DBufferC = GraphBuilder.CreateTexture(Desc, TEXT("DBufferC"), TextureFlags);
-
-		if (DBufferMaskTechnique == EDecalDBufferMaskTechnique::PerPixel)
-		{
-			// Note: 32bpp format is used here to utilize color compression hardware (same as other DBuffer targets).
-			// This significantly reduces bandwidth for clearing, writing and reading on some GPUs.
-			// While a smaller format, such as R8_UINT, will use less video memory, it will result in slower clears and higher bandwidth requirements.
-			check(Desc.Format == PF_B8G8R8A8);
-			Desc.Flags = TexCreate_ShaderResource | TexCreate_RenderTargetable;
-			Desc.ClearValue = FClearValueBinding::Transparent;
-			DBufferTextures.DBufferMask = GraphBuilder.CreateTexture(Desc, TEXT("DBufferMask"));
-		}
-	}
-
-	return DBufferTextures;
-}
-
 void AddDeferredDecalPass(
 	FRDGBuilder& GraphBuilder,
 	const FViewInfo& View,
@@ -661,8 +573,8 @@ void ExtractNormalsForNextFrameReprojection(FRDGBuilder& GraphBuilder, const FSc
 		{
 			if (Views[Index].bStatePrevViewInfoIsReadOnly == false)
 			{
-				GraphBuilder.QueueTextureExtraction(SceneTextures.GBufferA,&Views[Index].ViewState->PrevFrameViewInfo.GBufferA);
+				GraphBuilder.QueueTextureExtraction(SceneTextures.GBufferA, &Views[Index].ViewState->PrevFrameViewInfo.GBufferA);
 			}
 		}
 	}
-}
\ No newline at end of file
+}
diff --git a/Engine/Source/Runtime/Renderer/Private/CompositionLighting/PostProcessDeferredDecals.h b/Engine/Source/Runtime/Renderer/Private/CompositionLighting/PostProcessDeferredDecals.h
--- a/Engine/Source/Runtime/Renderer/Private/CompositionLighting/PostProcessDeferredDecals.h
+++ b/Engine/Source/Runtime/Renderer/Private/CompositionLighting/PostProcessDeferredDecals.h
@@ -3,46 +3,25 @@
 #pragma once
 
 #include "RenderGraph.h"
+#include "SceneTextureParameters.h"
 #include "SceneView.h"
 
 enum class EDecalRenderStage : uint8;
 enum class EDecalRenderTargetMode : uint8;
+struct FDBufferTextures;
 struct FSceneTextures;
 class FViewInfo;
 
 bool IsDBufferEnabled(const FSceneViewFamily& ViewFamily, EShaderPlatform ShaderPlatform);
 
-struct FDBufferTextures
-{
-	bool IsValid() const
-	{
-		check(!DBufferA || (DBufferB && DBufferC));
-		return HasBeenProduced(DBufferA);
-	}
-
-	FRDGTextureRef DBufferA = nullptr;
-	FRDGTextureRef DBufferB = nullptr;
-	FRDGTextureRef DBufferC = nullptr;
-	FRDGTextureRef DBufferMask = nullptr;
-};
-
-FDBufferTextures CreateDBufferTextures(FRDGBuilder& GraphBuilder, FIntPoint Extent, EShaderPlatform ShaderPlatform);
-
-BEGIN_SHADER_PARAMETER_STRUCT(FDBufferParameters, )
-	SHADER_PARAMETER_RDG_TEXTURE(Texture2D, DBufferATexture)
-	SHADER_PARAMETER_RDG_TEXTURE(Texture2D, DBufferBTexture)
-	SHADER_PARAMETER_RDG_TEXTURE(Texture2D, DBufferCTexture)
-	SHADER_PARAMETER_RDG_TEXTURE(Texture2D<uint>, DBufferRenderMask)
-	SHADER_PARAMETER_SAMPLER(SamplerState, DBufferATextureSampler)
-	SHADER_PARAMETER_SAMPLER(SamplerState, DBufferBTextureSampler)
-	SHADER_PARAMETER_SAMPLER(SamplerState, DBufferCTextureSampler)
-END_SHADER_PARAMETER_STRUCT()
-
-FDBufferParameters GetDBufferParameters(FRDGBuilder& GraphBuilder, const FDBufferTextures& DBufferTextures, EShaderPlatform ShaderPlatform);
+BEGIN_GLOBAL_SHADER_PARAMETER_STRUCT(FDecalPassUniformParameters, )
+	SHADER_PARAMETER_STRUCT(FSceneTextureUniformParameters, SceneTextures)
+	SHADER_PARAMETER_RDG_TEXTURE(Texture2D, EyeAdaptationTexture)
+END_GLOBAL_SHADER_PARAMETER_STRUCT()
 
 struct FDeferredDecalPassTextures
 {
-	TRDGUniformBufferRef<FSceneTextureUniformParameters> SceneTexturesUniformBuffer = nullptr;
+	TRDGUniformBufferRef<FDecalPassUniformParameters> DecalPassUniformBuffer = nullptr;
 
 	// Potential render targets for the decal pass.
 	FRDGTextureMSAA Depth;
@@ -55,7 +34,11 @@ struct FDeferredDecalPassTextures
 	FDBufferTextures* DBufferTextures = nullptr;
 };
 
-FDeferredDecalPassTextures GetDeferredDecalPassTextures(FRDGBuilder& GraphBuilder, const FSceneTextures& SceneTextures, FDBufferTextures* DBufferTextures);
+FDeferredDecalPassTextures GetDeferredDecalPassTextures(
+	FRDGBuilder& GraphBuilder, 
+	const FSceneView& ViewInfo,
+	const FSceneTextures& SceneTextures, 
+	FDBufferTextures* DBufferTextures);
 
 void AddDeferredDecalPass(
 	FRDGBuilder& GraphBuilder,
@@ -77,7 +60,7 @@ TUniformBufferRef<FDeferredDecalUniformParameters> CreateDeferredDecalUniformBuf
 BEGIN_SHADER_PARAMETER_STRUCT(FDeferredDecalPassParameters, )
 	SHADER_PARAMETER_STRUCT_INCLUDE(FViewShaderParameters, View)
 	SHADER_PARAMETER_STRUCT_REF(FDeferredDecalUniformParameters, DeferredDecal)
-	SHADER_PARAMETER_RDG_UNIFORM_BUFFER(FSceneTextureUniformParameters, SceneTextures)
+	SHADER_PARAMETER_RDG_UNIFORM_BUFFER(FDecalPassUniformParameters, DecalPass)
 	RENDER_TARGET_BINDING_SLOTS()
 END_SHADER_PARAMETER_STRUCT()
 
diff --git a/Engine/Source/Runtime/Renderer/Private/DBufferTextures.cpp b/Engine/Source/Runtime/Renderer/Private/DBufferTextures.cpp
new file mode 100644
--- /dev/null
+++ b/Engine/Source/Runtime/Renderer/Private/DBufferTextures.cpp
@@ -0,0 +1,101 @@
+// Copyright Epic Games, Inc. All Rights Reserved.
+
+#include "DBufferTextures.h"
+
+#include "RendererUtils.h"
+#include "RenderGraphUtils.h"
+#include "SceneUtils.h"
+#include "ScenePrivate.h"
+
+bool FDBufferTextures::IsValid() const
+{
+	check(!DBufferA || (DBufferB && DBufferC));
+	return HasBeenProduced(DBufferA);
+}
+
+EDecalDBufferMaskTechnique GetDBufferMaskTechnique(EShaderPlatform ShaderPlatform)
+{
+	const bool bWriteMaskDBufferMask = RHISupportsRenderTargetWriteMask(ShaderPlatform);
+	const bool bPerPixelDBufferMask = IsUsingPerPixelDBufferMask(ShaderPlatform);
+	checkf(!bWriteMaskDBufferMask || !bPerPixelDBufferMask, TEXT("The WriteMask and PerPixel DBufferMask approaches cannot be enabled at the same time. They are mutually exclusive."));
+
+	if (bWriteMaskDBufferMask)
+	{
+		return EDecalDBufferMaskTechnique::WriteMask;
+	}
+	else if (bPerPixelDBufferMask)
+	{
+		return EDecalDBufferMaskTechnique::PerPixel;
+	}
+	return EDecalDBufferMaskTechnique::Disabled;
+}
+
+FDBufferTextures CreateDBufferTextures(FRDGBuilder& GraphBuilder, FIntPoint Extent, EShaderPlatform ShaderPlatform)
+{
+	FDBufferTextures DBufferTextures;
+
+	if (IsUsingDBuffers(ShaderPlatform))
+	{
+		const EDecalDBufferMaskTechnique DBufferMaskTechnique = GetDBufferMaskTechnique(ShaderPlatform);
+		const ETextureCreateFlags WriteMaskFlags = DBufferMaskTechnique == EDecalDBufferMaskTechnique::WriteMask ? TexCreate_NoFastClearFinalize | TexCreate_DisableDCC : TexCreate_None;
+		const ETextureCreateFlags BaseFlags = WriteMaskFlags | TexCreate_ShaderResource | TexCreate_RenderTargetable;
+		const ERDGTextureFlags TextureFlags = DBufferMaskTechnique != EDecalDBufferMaskTechnique::Disabled
+			? ERDGTextureFlags::MaintainCompression
+			: ERDGTextureFlags::None;
+
+		FRDGTextureDesc Desc = FRDGTextureDesc::Create2D(Extent, PF_B8G8R8A8, FClearValueBinding::None, BaseFlags);
+
+		Desc.Flags = BaseFlags | GFastVRamConfig.DBufferA;
+		Desc.ClearValue = FClearValueBinding::Black;
+		DBufferTextures.DBufferA = GraphBuilder.CreateTexture(Desc, TEXT("DBufferA"), TextureFlags);
+
+		Desc.Flags = BaseFlags | GFastVRamConfig.DBufferB;
+		Desc.ClearValue = FClearValueBinding(FLinearColor(128.0f / 255.0f, 128.0f / 255.0f, 128.0f / 255.0f, 1));
+		DBufferTextures.DBufferB = GraphBuilder.CreateTexture(Desc, TEXT("DBufferB"), TextureFlags);
+
+		Desc.Flags = BaseFlags | GFastVRamConfig.DBufferC;
+		Desc.ClearValue = FClearValueBinding(FLinearColor(0, 0, 0, 1));
+		DBufferTextures.DBufferC = GraphBuilder.CreateTexture(Desc, TEXT("DBufferC"), TextureFlags);
+
+		if (DBufferMaskTechnique == EDecalDBufferMaskTechnique::PerPixel)
+		{
+			// Note: 32bpp format is used here to utilize color compression hardware (same as other DBuffer targets).
+			// This significantly reduces bandwidth for clearing, writing and reading on some GPUs.
+			// While a smaller format, such as R8_UINT, will use less video memory, it will result in slower clears and higher bandwidth requirements.
+			check(Desc.Format == PF_B8G8R8A8);
+			Desc.Flags = TexCreate_ShaderResource | TexCreate_RenderTargetable;
+			Desc.ClearValue = FClearValueBinding::Transparent;
+			DBufferTextures.DBufferMask = GraphBuilder.CreateTexture(Desc, TEXT("DBufferMask"));
+		}
+	}
+
+	return DBufferTextures;
+}
+
+FDBufferParameters GetDBufferParameters(FRDGBuilder& GraphBuilder, const FDBufferTextures& DBufferTextures, EShaderPlatform ShaderPlatform)
+{
+	const FRDGSystemTextures& SystemTextures = FRDGSystemTextures::Get(GraphBuilder);
+
+	FDBufferParameters Parameters;
+	Parameters.DBufferATextureSampler = TStaticSamplerState<>::GetRHI();
+	Parameters.DBufferBTextureSampler = TStaticSamplerState<>::GetRHI();
+	Parameters.DBufferCTextureSampler = TStaticSamplerState<>::GetRHI();
+	Parameters.DBufferATexture = SystemTextures.BlackAlphaOne;
+	Parameters.DBufferBTexture = SystemTextures.DefaultNormal8Bit;
+	Parameters.DBufferCTexture = SystemTextures.BlackAlphaOne;
+	Parameters.DBufferRenderMask = SystemTextures.White;
+
+	if (DBufferTextures.IsValid())
+	{
+		Parameters.DBufferATexture = DBufferTextures.DBufferA;
+		Parameters.DBufferBTexture = DBufferTextures.DBufferB;
+		Parameters.DBufferCTexture = DBufferTextures.DBufferC;
+
+		if (DBufferTextures.DBufferMask)
+		{
+			Parameters.DBufferRenderMask = DBufferTextures.DBufferMask;
+		}
+	}
+
+	return Parameters;
+}
diff --git a/Engine/Source/Runtime/Renderer/Private/DBufferTextures.h b/Engine/Source/Runtime/Renderer/Private/DBufferTextures.h
new file mode 100644
--- /dev/null
+++ b/Engine/Source/Runtime/Renderer/Private/DBufferTextures.h
@@ -0,0 +1,45 @@
+// Copyright Epic Games, Inc. All Rights Reserved.
+
+#pragma once
+
+#include "RenderGraphDefinitions.h"
+#include "RHIDefinitions.h"
+#include "ShaderParameterMacros.h"
+
+struct FSceneTextures;
+class FSceneViewFamily;
+class FSceneTextureUniformParameters;
+class FViewInfo;
+
+enum class EDecalDBufferMaskTechnique
+{
+	Disabled,	// DBufferMask is not enabled.
+	PerPixel,	// DBufferMask is written explicitly by the shader during the DBuffer pass.
+	WriteMask,	// DBufferMask is constructed after the DBuffer pass by compositing DBuffer write mask planes together in a compute shader.
+};
+
+EDecalDBufferMaskTechnique GetDBufferMaskTechnique(EShaderPlatform ShaderPlatform);
+
+struct FDBufferTextures
+{
+	bool IsValid() const;
+
+	FRDGTextureRef DBufferA = nullptr;
+	FRDGTextureRef DBufferB = nullptr;
+	FRDGTextureRef DBufferC = nullptr;
+	FRDGTextureRef DBufferMask = nullptr;
+};
+
+FDBufferTextures CreateDBufferTextures(FRDGBuilder& GraphBuilder, FIntPoint Extent, EShaderPlatform ShaderPlatform);
+
+BEGIN_SHADER_PARAMETER_STRUCT(FDBufferParameters, )
+	SHADER_PARAMETER_RDG_TEXTURE(Texture2D, DBufferATexture)
+	SHADER_PARAMETER_RDG_TEXTURE(Texture2D, DBufferBTexture)
+	SHADER_PARAMETER_RDG_TEXTURE(Texture2D, DBufferCTexture)
+	SHADER_PARAMETER_RDG_TEXTURE(Texture2D<uint>, DBufferRenderMask)
+	SHADER_PARAMETER_SAMPLER(SamplerState, DBufferATextureSampler)
+	SHADER_PARAMETER_SAMPLER(SamplerState, DBufferBTextureSampler)
+	SHADER_PARAMETER_SAMPLER(SamplerState, DBufferCTextureSampler)
+END_SHADER_PARAMETER_STRUCT()
+
+FDBufferParameters GetDBufferParameters(FRDGBuilder& GraphBuilder, const FDBufferTextures& DBufferTextures, EShaderPlatform ShaderPlatform);
diff --git a/Engine/Source/Runtime/Renderer/Private/ScenePrivate.h b/Engine/Source/Runtime/Renderer/Private/ScenePrivate.h
--- a/Engine/Source/Runtime/Renderer/Private/ScenePrivate.h
+++ b/Engine/Source/Runtime/Renderer/Private/ScenePrivate.h
@@ -57,7 +57,6 @@
 #include "Lumen/LumenViewState.h"
 #include "VolumetricRenderTargetViewStateData.h"
 #include "GPUScene.h"
-#include "CompositionLighting/PostProcessDeferredDecals.h"
 #include "DynamicBVH.h"
 
 /** Factor by which to grow occlusion tests **/
