fixed empty pagetable on Mobile devices

diff --git a/Engine/Source/Runtime/Renderer/Private/VT/VirtualTextureSpace.cpp b/Engine/Source/Runtime/Renderer/Private/VT/VirtualTextureSpace.cpp
--- a/Engine/Source/Runtime/Renderer/Private/VT/VirtualTextureSpace.cpp
+++ b/Engine/Source/Runtime/Renderer/Private/VT/VirtualTextureSpace.cpp
@@ -480,7 +480,7 @@ void FVirtualTextureSpace::ApplyUpdates(FVirtualTextureSystem* System, FRDGBuild
 			if (NumUpdates)
 			{
 				auto* PassParameters = GraphBuilder.AllocParameters<FRenderTargetParameters>();
-				PassParameters->RenderTargets[0] = FRenderTargetBinding(PageTableTexture, ERenderTargetLoadAction::ENoAction, Mip);
+				PassParameters->RenderTargets[0] = FRenderTargetBinding(PageTableTexture, ERenderTargetLoadAction::ELoad, Mip);
 
 				GraphBuilder.AddPass(
 					RDG_EVENT_NAME("PageTableUpdate (Mip: %d)", Mip),
