Renamed VisualizeIndirectLighting to VisualizeLightingOnProbes and made the name consistent on ui and cvar.

diff --git a/Engine/Source/Runtime/Engine/Public/ShowFlags.h b/Engine/Source/Runtime/Engine/Public/ShowFlags.h
--- a/Engine/Source/Runtime/Engine/Public/ShowFlags.h
+++ b/Engine/Source/Runtime/Engine/Public/ShowFlags.h
@@ -439,7 +439,7 @@ private:
 		SetVisualizeGlobalDistanceField(false);
 		SetVisualizeLumenScene(false);
 		SetVisualizeLumenIndirectDiffuse(false);
-		SetVisualizeIndirectLighting(false);
+		SetVisualizeLightingOnProbes(false);
 		SetVisualizeDistanceFieldAO(false);
 		SetVisualizeSSR(false);
 		SetVisualizeSSS(false);
diff --git a/Engine/Source/Runtime/Engine/Public/ShowFlagsValues.inl b/Engine/Source/Runtime/Engine/Public/ShowFlagsValues.inl
--- a/Engine/Source/Runtime/Engine/Public/ShowFlagsValues.inl
+++ b/Engine/Source/Runtime/Engine/Public/ShowFlagsValues.inl
@@ -309,8 +309,8 @@ SHOWFLAG_FIXED_IN_SHIPPING(0, VisualizeGlobalDistanceField, SFG_Visualize, NSLOC
 SHOWFLAG_FIXED_IN_SHIPPING(0, VisualizeLumenScene, SFG_Visualize, NSLOCTEXT("UnrealEd", "LumenSceneSF", "Lumen Scene"))
 /** Visualization of Lumen indirect diffuse */
 SHOWFLAG_FIXED_IN_SHIPPING(0, VisualizeLumenIndirectDiffuse, SFG_Visualize, NSLOCTEXT("UnrealEd", "LumenIndirectDiffuseSF", "Lumen Global Illumination"))
-/** Enable the debug visualisation of indirect diffuse and specular lighting using probes */
-SHOWFLAG_FIXED_IN_SHIPPING(0, VisualizeIndirectLighting, SFG_Visualize, NSLOCTEXT("UnrealEd", "VisualizeIndirectLightingSF", "Visualize Indirect Lighting"))
+/** Enable the debug visualization of diffuse/specular lighting (direct and indirect) using probes */
+SHOWFLAG_FIXED_IN_SHIPPING(0, VisualizeLightingOnProbes, SFG_Visualize, NSLOCTEXT("UnrealEd", "VisualizeLightingOnProbesSF", "Visualize Lighting on Probes"))
 /** Screen space AO, for now SHOWFLAG_ALWAYS_ACCESSIBLE because r.GBuffer need that */
 SHOWFLAG_ALWAYS_ACCESSIBLE(ScreenSpaceAO, SFG_LightingFeatures, NSLOCTEXT("UnrealEd", "ScreenSpaceAOSF", "Screen Space Ambient Occlusion"))
 /** Distance field AO, for now SHOWFLAG_ALWAYS_ACCESSIBLE because it's exposed in SceneCapture */
diff --git a/Engine/Source/Runtime/Renderer/Private/DebugProbeRendering.cpp b/Engine/Source/Runtime/Renderer/Private/DebugProbeRendering.cpp
--- a/Engine/Source/Runtime/Renderer/Private/DebugProbeRendering.cpp
+++ b/Engine/Source/Runtime/Renderer/Private/DebugProbeRendering.cpp
@@ -6,10 +6,10 @@
 
 
 // Changing this causes a full shader recompile
-static TAutoConsoleVariable<int32> CVarDebugIndirectLightingProbes(
-	TEXT("r.DebugIndirectLightingProbes"),
+static TAutoConsoleVariable<int32> CVarVisualizeLightingOnProbes(
+	TEXT("r.VisualizeLightingOnProbes"),
 	0,
-	TEXT("Enables debug probes rendering to visualise diffuse/specular lighting on simple sphere scattered in the world.") \
+	TEXT("Enables debug probes rendering to visualise diffuse/specular lighting (direct and indirect) on simple sphere scattered in the world.") \
 	TEXT(" 0: disabled.\n")
 	TEXT(" 1: camera probes only.\n")
 	TEXT(" 2: world probes only.\n")
@@ -63,7 +63,7 @@ static void CommonStampDeferredDebugProbeDrawCall(
 	bool bRenderVelocity = false)
 {
 	PassParameters->ViewUniformBuffer = View.ViewUniformBuffer;
-	PassParameters->DebugProbesMode = View.Family->EngineShowFlags.VisualizeIndirectLighting ? 3 : FMath::Clamp(CVarDebugIndirectLightingProbes.GetValueOnRenderThread(), 0, 3);
+	PassParameters->DebugProbesMode = View.Family->EngineShowFlags.VisualizeLightingOnProbes ? 3 : FMath::Clamp(CVarVisualizeLightingOnProbes.GetValueOnRenderThread(), 0, 3);
 		
 	FStampDeferredDebugProbePS::FPermutationDomain PermutationVector;
 	PermutationVector.Set<FStampDeferredDebugProbePS::FRenderVelocity>(bRenderVelocity);
@@ -85,10 +85,10 @@ void StampDeferredDebugProbeDepthPS(
 	RDG_EVENT_SCOPE(GraphBuilder, "StampDeferredDebugProbeDepth");
 	RDG_GPU_STAT_SCOPE(GraphBuilder, StampDeferredDebugProbe);
 
-	const bool bVisualizeIndirectLighting = CVarDebugIndirectLightingProbes.GetValueOnRenderThread() > 0;
+	const bool bVisualizeLightingOnProbes = CVarVisualizeLightingOnProbes.GetValueOnRenderThread() > 0;
 	for (const FViewInfo& View : Views)
 	{
-		if (!(bVisualizeIndirectLighting || View.Family->EngineShowFlags.VisualizeIndirectLighting))
+		if (!(bVisualizeLightingOnProbes || View.Family->EngineShowFlags.VisualizeLightingOnProbes))
 		{
 			continue;
 		}
@@ -109,10 +109,10 @@ void StampDeferredDebugProbeMaterialPS(
 	RDG_EVENT_SCOPE(GraphBuilder, "StampDeferredDebugProbeMaterial");
 	RDG_GPU_STAT_SCOPE(GraphBuilder, StampDeferredDebugProbe);
 
-	const bool bVisualizeIndirectLighting = CVarDebugIndirectLightingProbes.GetValueOnRenderThread() > 0;
+	const bool bVisualizeLightingOnProbes = CVarVisualizeLightingOnProbes.GetValueOnRenderThread() > 0;
 	for (const FViewInfo& View : Views)
 	{
-		if (!(bVisualizeIndirectLighting || View.Family->EngineShowFlags.VisualizeIndirectLighting))
+		if (!(bVisualizeLightingOnProbes || View.Family->EngineShowFlags.VisualizeLightingOnProbes))
 		{
 			continue;
 		}
@@ -132,10 +132,10 @@ void StampDeferredDebugProbeVelocityPS(
 	RDG_EVENT_SCOPE(GraphBuilder, "StampDeferredDebugProbeVelocity");
 	RDG_GPU_STAT_SCOPE(GraphBuilder, StampDeferredDebugProbe);
 
-	const bool bVisualizeIndirectLighting = CVarDebugIndirectLightingProbes.GetValueOnRenderThread() > 0;
+	const bool bVisualizeLightingOnProbes = CVarVisualizeLightingOnProbes.GetValueOnRenderThread() > 0;
 	for (const FViewInfo& View : Views)
 	{
-		if (!(bVisualizeIndirectLighting || View.Family->EngineShowFlags.VisualizeIndirectLighting))
+		if (!(bVisualizeLightingOnProbes || View.Family->EngineShowFlags.VisualizeLightingOnProbes))
 		{
 			continue;
 		}
