Follow up with code clean up post AtmosphericFog removal

diff --git a/Engine/Shaders/Public/ShaderVersion.ush b/Engine/Shaders/Public/ShaderVersion.ush
--- a/Engine/Shaders/Public/ShaderVersion.ush
+++ b/Engine/Shaders/Public/ShaderVersion.ush
@@ -3,4 +3,4 @@
 // in Platform.ush (which should be included in any shader) it allows to invalidate the shader DDC.
 // 
 // If you are merging streams and there is a conflict with this GUID you should make a new GUID rather than taking one or the other.
-// GUID = 4A6DD7DA42BD444FBC64C4D1C186EABD
+// GUID = FF18584666CC4406A481AB221199BBAA
diff --git a/Engine/Source/Runtime/Engine/Classes/Materials/MaterialExpressionAtmosphericLightColor.h b/Engine/Source/Runtime/Engine/Classes/Materials/MaterialExpressionAtmosphericLightColor.h
--- a/Engine/Source/Runtime/Engine/Classes/Materials/MaterialExpressionAtmosphericLightColor.h
+++ b/Engine/Source/Runtime/Engine/Classes/Materials/MaterialExpressionAtmosphericLightColor.h
@@ -8,7 +8,7 @@
 #include "Materials/MaterialExpression.h"
 #include "MaterialExpressionAtmosphericLightColor.generated.h"
 
-UCLASS()
+UCLASS(collapsecategories, hidecategories = Object, DisplayName = "Atmosphere Sun Light Illuminance On Ground")
 class UMaterialExpressionAtmosphericLightColor : public UMaterialExpression
 {
 	GENERATED_UCLASS_BODY()
diff --git a/Engine/Source/Runtime/Engine/Classes/Materials/MaterialExpressionAtmosphericLightVector.h b/Engine/Source/Runtime/Engine/Classes/Materials/MaterialExpressionAtmosphericLightVector.h
--- a/Engine/Source/Runtime/Engine/Classes/Materials/MaterialExpressionAtmosphericLightVector.h
+++ b/Engine/Source/Runtime/Engine/Classes/Materials/MaterialExpressionAtmosphericLightVector.h
@@ -8,13 +8,11 @@
 #include "Materials/MaterialExpression.h"
 #include "MaterialExpressionAtmosphericLightVector.generated.h"
 
-UCLASS()
+UCLASS(collapsecategories, hidecategories = Object, DisplayName = "Atmosphere Sun Light Vector")
 class UMaterialExpressionAtmosphericLightVector : public UMaterialExpression
 {
 	GENERATED_UCLASS_BODY()
 
-
-
 	//~ Begin UMaterialExpression Interface
 #if WITH_EDITOR
 	virtual int32 Compile(class FMaterialCompiler* Compiler, int32 OutputIndex) override;
diff --git a/Engine/Source/Runtime/Engine/Private/Components/DirectionalLightComponent.cpp b/Engine/Source/Runtime/Engine/Private/Components/DirectionalLightComponent.cpp
--- a/Engine/Source/Runtime/Engine/Private/Components/DirectionalLightComponent.cpp
+++ b/Engine/Source/Runtime/Engine/Private/Components/DirectionalLightComponent.cpp
@@ -224,11 +224,6 @@ public:
 	 */
 	FLinearColor SunDiscOuterSpaceLuminance;
 
-	/**
-	 * The light outer space illuminance used to feed the SkyAtmosphere simulation
-	 */
-	FLinearColor SunOuterSpaceIlluminance;
-
 	/**
 	 * The atmosphere transmittance from ground to current light direction.
 	 */
@@ -296,9 +291,8 @@ public:
 		OcclusionMaskDarkness(Component->OcclusionMaskDarkness),
 		OcclusionDepthRange(Component->OcclusionDepthRange),
 		LightShaftOverrideDirection(Component->LightShaftOverrideDirection),
-		bPerPixelTransmittanceEnabled(false),
+		bPerPixelTransmittanceEnabled(Component->bPerPixelAtmosphereTransmittance),
 		SunDiscOuterSpaceLuminance(FLinearColor::White),
-		SunOuterSpaceIlluminance(FLinearColor::White),
 		TransmittanceTowardSun(FLinearColor::White),
 		DynamicShadowCascades(Component->DynamicShadowCascades > 0 ? Component->DynamicShadowCascades : 0),
 		CascadeDistributionExponent(Component->CascadeDistributionExponent),
@@ -537,12 +531,10 @@ public:
 		return DoesPlatformSupportDistanceFieldShadowing(GShaderPlatformForFeatureLevel[InFeatureLevel]) && (bCreateWithCSM || bCreateWithoutCSM);
 	}
 
-	virtual void SetAtmosphereRelatedProperties(FLinearColor TransmittanceTowardSunIn, FLinearColor SunOuterSpaceIlluminanceIn, FLinearColor SunOuterSpaceLuminance, bool bPerPixelTransmittanceEnabledIn) override
+	virtual void SetAtmosphereRelatedProperties(FLinearColor TransmittanceTowardSunIn, FLinearColor SunDiscOuterSpaceLuminanceIn) override
 	{
 		TransmittanceTowardSun = TransmittanceTowardSunIn;
-		SunDiscOuterSpaceLuminance = SunOuterSpaceLuminance;
-		SunOuterSpaceIlluminance = SunOuterSpaceIlluminanceIn;
-		bPerPixelTransmittanceEnabled = bPerPixelTransmittanceEnabledIn;
+		SunDiscOuterSpaceLuminance = SunDiscOuterSpaceLuminanceIn;
 	}
 
 	virtual FLinearColor GetOuterSpaceLuminance() const override
diff --git a/Engine/Source/Runtime/Engine/Private/Components/SkyAtmosphereComponent.cpp b/Engine/Source/Runtime/Engine/Private/Components/SkyAtmosphereComponent.cpp
--- a/Engine/Source/Runtime/Engine/Private/Components/SkyAtmosphereComponent.cpp
+++ b/Engine/Source/Runtime/Engine/Private/Components/SkyAtmosphereComponent.cpp
@@ -429,8 +429,6 @@ FSkyAtmosphereSceneProxy::FSkyAtmosphereSceneProxy(const USkyAtmosphereComponent
 	TraceSampleCountScale = InComponent->TraceSampleCountScale;
 
 	InComponent->GetOverrideLightStatus(OverrideAtmosphericLight, OverrideAtmosphericLightDirection);
-
-	TransmittanceAtZenith = AtmosphereSetup.GetTransmittanceAtGroundLevel(FVector(0.0f, 0.0f, 1.0f));
 }
 
 FSkyAtmosphereSceneProxy::~FSkyAtmosphereSceneProxy()
diff --git a/Engine/Source/Runtime/Engine/Private/Materials/MaterialExpressions.cpp b/Engine/Source/Runtime/Engine/Private/Materials/MaterialExpressions.cpp
--- a/Engine/Source/Runtime/Engine/Private/Materials/MaterialExpressions.cpp
+++ b/Engine/Source/Runtime/Engine/Private/Materials/MaterialExpressions.cpp
@@ -18079,7 +18079,7 @@ int32 UMaterialExpressionAtmosphericFogColor::Compile(class FMaterialCompiler* C
 
 void UMaterialExpressionAtmosphericFogColor::GetCaption(TArray<FString>& OutCaptions) const
 {
-	OutCaptions.Add(TEXT("Atmospheric Fog Color"));
+	OutCaptions.Add(TEXT("Atmospheric Fog Color (deprecated)"));
 }
 #endif // WITH_EDITOR
 
@@ -18563,7 +18563,7 @@ int32 UMaterialExpressionAtmosphericLightVector::Compile(class FMaterialCompiler
 
 void UMaterialExpressionAtmosphericLightVector::GetCaption(TArray<FString>& OutCaptions) const
 {
-	OutCaptions.Add(TEXT("AtmosphericLightVector"));
+	OutCaptions.Add(TEXT("Atmosphere Sun Light Vector"));
 }
 #endif // WITH_EDITOR
 
@@ -18598,7 +18598,7 @@ int32 UMaterialExpressionAtmosphericLightColor::Compile(class FMaterialCompiler*
 
 void UMaterialExpressionAtmosphericLightColor::GetCaption(TArray<FString>& OutCaptions) const
 {
-	OutCaptions.Add(TEXT("AtmosphericLightColor"));
+	OutCaptions.Add(TEXT("Atmosphere Sun Light Illuminance On Ground (not per pixel)"));
 }
 #endif // WITH_EDITOR
 
diff --git a/Engine/Source/Runtime/Engine/Public/SceneManagement.h b/Engine/Source/Runtime/Engine/Public/SceneManagement.h
--- a/Engine/Source/Runtime/Engine/Public/SceneManagement.h
+++ b/Engine/Source/Runtime/Engine/Public/SceneManagement.h
@@ -1162,7 +1162,6 @@ public:
 	~FSkyAtmosphereSceneProxy();
 
 	FLinearColor GetSkyLuminanceFactor() const { return SkyLuminanceFactor; }
-	FLinearColor GetTransmittanceAtZenith() const { return TransmittanceAtZenith; };
 	float GetAerialPespectiveViewDistanceScale() const { return AerialPespectiveViewDistanceScale; }
 	float GetHeightFogContribution() const { return HeightFogContribution; }
 	float GetAerialPerspectiveStartDepthKm() const { return AerialPerspectiveStartDepthKm; }
@@ -1181,7 +1180,6 @@ private:
 
 	FAtmosphereSetup AtmosphereSetup;
 
-	FLinearColor TransmittanceAtZenith;
 	FLinearColor SkyLuminanceFactor;
 	float AerialPespectiveViewDistanceScale;
 	float HeightFogContribution;
@@ -1515,7 +1513,7 @@ public:
 	inline bool IsUsedAsAtmosphereSunLight() const { return bUsedAsAtmosphereSunLight; }
 	inline uint8 GetAtmosphereSunLightIndex() const { return AtmosphereSunLightIndex; }
 	inline FLinearColor GetAtmosphereSunDiskColorScale() const { return AtmosphereSunDiskColorScale; }
-	virtual void SetAtmosphereRelatedProperties(FLinearColor TransmittanceTowardSun, FLinearColor SunOuterSpaceIlluminance, FLinearColor SunOuterSpaceLuminance, bool bPerPixelTransmittanceEnabled) {}
+	virtual void SetAtmosphereRelatedProperties(FLinearColor TransmittanceTowardSunIn, FLinearColor SunDiscOuterSpaceLuminanceIn) {}
 	virtual FLinearColor GetOuterSpaceLuminance() const { return FLinearColor::White; }
 	virtual FLinearColor GetOuterSpaceIlluminance() const { return GetColor(); }
 	virtual FLinearColor GetAtmosphereTransmittanceTowardSun() const { return FLinearColor::White; }
diff --git a/Engine/Source/Runtime/Engine/Public/SceneView.h b/Engine/Source/Runtime/Engine/Public/SceneView.h
--- a/Engine/Source/Runtime/Engine/Public/SceneView.h
+++ b/Engine/Source/Runtime/Engine/Public/SceneView.h
@@ -677,7 +677,6 @@ enum ETranslucencyVolumeCascade
 	VIEW_UNIFORM_BUFFER_MEMBER(FVector3f, IndirectLightingColorScale) \
 	VIEW_UNIFORM_BUFFER_MEMBER(FVector3f, PrecomputedIndirectLightingColorScale) \
 	VIEW_UNIFORM_BUFFER_MEMBER_ARRAY(FVector4, AtmosphereLightDirection, [NUM_ATMOSPHERE_LIGHTS]) \
-	VIEW_UNIFORM_BUFFER_MEMBER_ARRAY(FLinearColor, AtmosphereLightIlluminanceOnGround, [NUM_ATMOSPHERE_LIGHTS]) \
 	VIEW_UNIFORM_BUFFER_MEMBER_ARRAY(FLinearColor, AtmosphereLightIlluminanceOnGroundPostTransmittance, [NUM_ATMOSPHERE_LIGHTS]) \
 	VIEW_UNIFORM_BUFFER_MEMBER_ARRAY(FLinearColor, AtmosphereLightIlluminanceOuterSpace, [NUM_ATMOSPHERE_LIGHTS]) \
 	VIEW_UNIFORM_BUFFER_MEMBER_ARRAY(FLinearColor, AtmosphereLightDiscLuminance, [NUM_ATMOSPHERE_LIGHTS]) \
diff --git a/Engine/Source/Runtime/Renderer/Private/SceneRendering.cpp b/Engine/Source/Runtime/Renderer/Private/SceneRendering.cpp
--- a/Engine/Source/Runtime/Renderer/Private/SceneRendering.cpp
+++ b/Engine/Source/Runtime/Renderer/Private/SceneRendering.cpp
@@ -1294,8 +1294,6 @@ void FViewInfo::SetupUniformBufferParameters(
 		check(Index < NUM_ATMOSPHERE_LIGHTS);
 		ViewUniformShaderParameters.AtmosphereLightDiscCosHalfApexAngle[Index] = FVector4(1.0f);
 		ViewUniformShaderParameters.AtmosphereLightDiscLuminance[Index] = FLinearColor::Black;
-		ViewUniformShaderParameters.AtmosphereLightIlluminanceOnGround[Index] = FLinearColor::Black;
-		ViewUniformShaderParameters.AtmosphereLightIlluminanceOnGround[Index].A = 0.0f;
 		ViewUniformShaderParameters.AtmosphereLightIlluminanceOnGroundPostTransmittance[Index] = FLinearColor::Black;
 		ViewUniformShaderParameters.AtmosphereLightIlluminanceOnGroundPostTransmittance[Index].A = 0.0f;
 		ViewUniformShaderParameters.AtmosphereLightIlluminanceOuterSpace[Index] = FLinearColor::Black;
@@ -1317,17 +1315,16 @@ void FViewInfo::SetupUniformBufferParameters(
 			ViewUniformShaderParameters.DirectionalLightDirection = FVector::ZeroVector;
 		}
 
-		// Atmospheric fog parameters
+		// Set default atmosphere lights parameters
 		FLightSceneInfo* SunLight = Scene->AtmosphereLights[0];	// Atmospheric fog only takes into account the a single sun light with index 0.
 		const float SunLightDiskHalfApexAngleRadian = SunLight ? SunLight->Proxy->GetSunLightHalfApexAngleRadian() : FLightSceneProxy::GetSunOnEarthHalfApexAngleRadian();
 
 		ViewUniformShaderParameters.AtmosphereLightDiscCosHalfApexAngle[0] = FVector4(FMath::Cos(SunLightDiskHalfApexAngleRadian));
 		//Added check so atmospheric light color and vector can use a directional light without needing an atmospheric fog actor in the scene
 		ViewUniformShaderParameters.AtmosphereLightDiscLuminance[0] = SunLight ? SunLight->Proxy->GetOuterSpaceLuminance() : FLinearColor::Black;
-		ViewUniformShaderParameters.AtmosphereLightIlluminanceOnGround[0] = SunLight ? SunLight->Proxy->GetColor() : FLinearColor::Black;
-		ViewUniformShaderParameters.AtmosphereLightIlluminanceOnGround[0].A = 1.0f;
-		ViewUniformShaderParameters.AtmosphereLightIlluminanceOnGroundPostTransmittance[0] = ViewUniformShaderParameters.AtmosphereLightIlluminanceOnGround[0];
+		ViewUniformShaderParameters.AtmosphereLightIlluminanceOnGroundPostTransmittance[0] = SunLight ? SunLight->Proxy->GetColor() : FLinearColor::Black;
 		ViewUniformShaderParameters.AtmosphereLightIlluminanceOnGroundPostTransmittance[0].A = 0.0f;
+		ViewUniformShaderParameters.AtmosphereLightIlluminanceOuterSpace[0] = ViewUniformShaderParameters.AtmosphereLightIlluminanceOnGroundPostTransmittance[0];
 		ViewUniformShaderParameters.AtmosphereLightIlluminanceOuterSpace[0].A = 0.0f;
 		ViewUniformShaderParameters.AtmosphereLightDirection[0] = SunLight ? -SunLight->Proxy->GetDirection() : DefaultSunDirection;
 
@@ -1404,8 +1401,6 @@ void FViewInfo::SetupUniformBufferParameters(
 			{
 				ViewUniformShaderParameters.AtmosphereLightDiscCosHalfApexAngle[Index] = FVector4(FMath::Cos(Light->Proxy->GetSunLightHalfApexAngleRadian()));
 				ViewUniformShaderParameters.AtmosphereLightDiscLuminance[Index] = Light->Proxy->GetOuterSpaceLuminance();
-				ViewUniformShaderParameters.AtmosphereLightIlluminanceOnGround[Index] = Light->Proxy->GetSunIlluminanceAccountingForSkyAtmospherePerPixelTransmittance();
-				ViewUniformShaderParameters.AtmosphereLightIlluminanceOnGround[Index].A = 1.0f;
 				ViewUniformShaderParameters.AtmosphereLightIlluminanceOnGroundPostTransmittance[Index] = Light->Proxy->GetSunIlluminanceOnGroundPostTransmittance();
 				ViewUniformShaderParameters.AtmosphereLightIlluminanceOnGroundPostTransmittance[Index].A = 1.0f; // interactions with HeightFogComponent
 				ViewUniformShaderParameters.AtmosphereLightIlluminanceOuterSpace[Index] = Light->Proxy->GetOuterSpaceIlluminance();
@@ -1456,11 +1451,10 @@ void FViewInfo::SetupUniformBufferParameters(
 				{
 					ViewUniformShaderParameters.AtmosphereLightDiscCosHalfApexAngle[Index] = FVector4(1.0f);
 					ViewUniformShaderParameters.AtmosphereLightDiscLuminance[Index] = FLinearColor::Black;
-					ViewUniformShaderParameters.AtmosphereLightIlluminanceOnGround[Index] = Light->Proxy->GetColor();
-					ViewUniformShaderParameters.AtmosphereLightIlluminanceOnGround[Index].A = 1.0f;
 					ViewUniformShaderParameters.AtmosphereLightIlluminanceOnGroundPostTransmittance[Index] = Light->Proxy->GetColor();
 					ViewUniformShaderParameters.AtmosphereLightIlluminanceOnGroundPostTransmittance[Index].A = 0.0f; // no interactions with HeightFogComponent
 					ViewUniformShaderParameters.AtmosphereLightIlluminanceOuterSpace[Index] = Light->Proxy->GetColor();
+					ViewUniformShaderParameters.AtmosphereLightIlluminanceOuterSpace[0].A = 0.0f;
 					ViewUniformShaderParameters.AtmosphereLightDirection[Index] = -Light->Proxy->GetDirection();
 				}
 				else
diff --git a/Engine/Source/Runtime/Renderer/Private/SkyAtmosphereRendering.cpp b/Engine/Source/Runtime/Renderer/Private/SkyAtmosphereRendering.cpp
--- a/Engine/Source/Runtime/Renderer/Private/SkyAtmosphereRendering.cpp
+++ b/Engine/Source/Runtime/Renderer/Private/SkyAtmosphereRendering.cpp
@@ -407,32 +407,23 @@ static void CopyAtmosphereSetupToUniformShaderParameters(FAtmosphereUniformShade
 #undef COPYMACRO
 }
 
-static FLinearColor GetLightDiskLuminance(FLightSceneInfo& Light, FLinearColor LightIlluminance)
+static FLinearColor GetLightDiskLuminance(FLightSceneInfo& Light)
 {
-	const float SunSolidAngle = 2.0f * PI * (1.0f - FMath::Cos(Light.Proxy->GetSunLightHalfApexAngleRadian())); // Solid angle from aperture https://en.wikipedia.org/wiki/Solid_angle 
-	return  Light.Proxy->GetAtmosphereSunDiskColorScale() * LightIlluminance / SunSolidAngle; // approximation
+	
+	const float SunSolidAngle = 2.0f * PI * (1.0f - FMath::Cos(Light.Proxy->GetSunLightHalfApexAngleRadian()));			// Solid angle from aperture https://en.wikipedia.org/wiki/Solid_angle 
+	return  Light.Proxy->GetAtmosphereSunDiskColorScale() * Light.Proxy->GetOuterSpaceIlluminance() / SunSolidAngle;	// approximation
 }
 
 void PrepareSunLightProxy(const FSkyAtmosphereRenderSceneInfo& SkyAtmosphere, uint32 AtmosphereLightIndex, FLightSceneInfo& AtmosphereLight)
 {
 	// See explanation in "Physically Based Sky, Atmosphere	and Cloud Rendering in Frostbite" page 26
-	const bool bAtmosphereAffectsSunIlluminance = true;
 	const FSkyAtmosphereSceneProxy& SkyAtmosphereProxy = SkyAtmosphere.GetSkyAtmosphereSceneProxy();
 	const FVector AtmosphereLightDirection = SkyAtmosphereProxy.GetAtmosphereLightDirection(AtmosphereLightIndex, -AtmosphereLight.Proxy->GetDirection());
-	const FLinearColor TransmittanceTowardSun = bAtmosphereAffectsSunIlluminance ? SkyAtmosphereProxy.GetAtmosphereSetup().GetTransmittanceAtGroundLevel(AtmosphereLightDirection) : FLinearColor(FLinearColor::White);
-
-	// TODO: when the AtmosphericFog is gone, we hsould be able to remove this, as well as TransmittanceAtZenith on the SkyAtmosphereProxy.
-	const FLinearColor TransmittanceAtZenithFinal = bAtmosphereAffectsSunIlluminance ? SkyAtmosphereProxy.GetTransmittanceAtZenith() : FLinearColor(FLinearColor::White);
-
-	const FLinearColor SunZenithIlluminanceOnGround = AtmosphereLight.Proxy->GetColor();
-	const FLinearColor SunOuterSpaceIlluminance = SunZenithIlluminanceOnGround / TransmittanceAtZenithFinal;
-	const FLinearColor SunDiskOuterSpaceLuminance = GetLightDiskLuminance(AtmosphereLight, SunOuterSpaceIlluminance);
+	const FLinearColor TransmittanceTowardSun = SkyAtmosphereProxy.GetAtmosphereSetup().GetTransmittanceAtGroundLevel(AtmosphereLightDirection);
 
-	// We always set the transmittance on the proxy. Shaders using atmospheric light color then have to decide which sun illuminance to use (without or with transmittance).
-	// We also set whether or not the light should apply the simple transmittance computed on CPU during lighting pass. If per pixel transmittance is enabled, it should not be done.
-	const bool bPerPixelTransmittanceEnabled = AtmosphereLight.Proxy->GetUsePerPixelAtmosphereTransmittance();
+	const FLinearColor SunDiskOuterSpaceLuminance = GetLightDiskLuminance(AtmosphereLight);
 
-	AtmosphereLight.Proxy->SetAtmosphereRelatedProperties(TransmittanceTowardSun, SunOuterSpaceIlluminance, SunDiskOuterSpaceLuminance, bPerPixelTransmittanceEnabled);
+	AtmosphereLight.Proxy->SetAtmosphereRelatedProperties(TransmittanceTowardSun, SunDiskOuterSpaceLuminance);
 }
 
 
@@ -535,9 +526,8 @@ void FScene::ResetAtmosphereLightsProperties()
 		FLightSceneInfo* Light = AtmosphereLights[LightIndex];
 		if (Light)
 		{
-			FLinearColor LightZenithIlluminance = Light->Proxy->GetColor();
-			const bool bPerPixelTransmittanceEnabled = false;
-			Light->Proxy->SetAtmosphereRelatedProperties(FLinearColor::White, LightZenithIlluminance, GetLightDiskLuminance(*Light, LightZenithIlluminance), bPerPixelTransmittanceEnabled);
+			const FLinearColor TransmittanceTowardSun = FLinearColor::White;
+			Light->Proxy->SetAtmosphereRelatedProperties(TransmittanceTowardSun, GetLightDiskLuminance(*Light));
 		}
 	}
 }
