Convert ReflectionEnvironmentRealTimeCapture to not use DrawDynamicMeshPass to stop using workaround path.

diff --git a/Engine/Source/Runtime/Renderer/Private/ReflectionEnvironmentRealTimeCapture.cpp b/Engine/Source/Runtime/Renderer/Private/ReflectionEnvironmentRealTimeCapture.cpp
--- a/Engine/Source/Runtime/Renderer/Private/ReflectionEnvironmentRealTimeCapture.cpp
+++ b/Engine/Source/Runtime/Renderer/Private/ReflectionEnvironmentRealTimeCapture.cpp
@@ -584,42 +584,33 @@ void FScene::AllocateAndCaptureFrameSkyEnvMap(
 								PassParameters->RenderTargets.DepthStencil = FDepthStencilBinding(CubeDepthTexture, ERenderTargetLoadAction::EClear, FExclusiveDepthStencil::DepthWrite_StencilNop);
 							}
 
-							GraphBuilder.AddPass(
-								RDG_EVENT_NAME("CaptureSkyMeshReflection"),
-								PassParameters,
-								ERDGPassFlags::Raster,
-								[&MainView, CubeViewUniformBuffer, bUseDepthBuffer](FRHICommandListImmediate& RHICmdList)
+							AddSimpleMeshPass(GraphBuilder, PassParameters, Scene, MainView, &InstanceCullingManager, RDG_EVENT_NAME("CaptureSkyMeshReflection"), MainView.ViewRect,
+							[&MainView, &CubeViewUniformBuffer, bUseDepthBuffer, Scene](FDynamicPassMeshDrawListContext* DynamicMeshPassContext)
+							{
+								FMeshPassProcessorRenderState DrawRenderState;
+
+								FExclusiveDepthStencil::Type BasePassDepthStencilAccess_Sky = bUseDepthBuffer ? FExclusiveDepthStencil::Type(Scene->DefaultBasePassDepthStencilAccess | FExclusiveDepthStencil::DepthWrite)
+									: FExclusiveDepthStencil::Type(Scene->DefaultBasePassDepthStencilAccess & ~FExclusiveDepthStencil::DepthWrite);
+								SetupBasePassState(BasePassDepthStencilAccess_Sky, false, DrawRenderState);
+
+								FSkyPassMeshProcessor PassMeshProcessor(Scene, nullptr, DrawRenderState, DynamicMeshPassContext);
+								const int32 SkyRealTimeReflectionOnlyMeshBatcheCount = MainView.SkyMeshBatches.Num();
+								for (int32 MeshBatchIndex = 0; MeshBatchIndex < SkyRealTimeReflectionOnlyMeshBatcheCount; ++MeshBatchIndex)
 								{
-									DrawDynamicMeshPass(MainView, RHICmdList,
-										[&MainView, &CubeViewUniformBuffer, bUseDepthBuffer](FDynamicPassMeshDrawListContext* DynamicMeshPassContext)
-										{
-											FScene* Scene = MainView.Family->Scene->GetRenderScene();
-
-											FMeshPassProcessorRenderState DrawRenderState;
-
-											FExclusiveDepthStencil::Type BasePassDepthStencilAccess_Sky = bUseDepthBuffer ? FExclusiveDepthStencil::Type(Scene->DefaultBasePassDepthStencilAccess | FExclusiveDepthStencil::DepthWrite)
-												: FExclusiveDepthStencil::Type(Scene->DefaultBasePassDepthStencilAccess & ~FExclusiveDepthStencil::DepthWrite);
-											SetupBasePassState(BasePassDepthStencilAccess_Sky, false, DrawRenderState);
-
-											FSkyPassMeshProcessor PassMeshProcessor(Scene, nullptr, DrawRenderState, DynamicMeshPassContext);
-											const int32 SkyRealTimeReflectionOnlyMeshBatcheCount = MainView.SkyMeshBatches.Num();
-											for (int32 MeshBatchIndex = 0; MeshBatchIndex < SkyRealTimeReflectionOnlyMeshBatcheCount; ++MeshBatchIndex)
-											{
-												FSkyMeshBatch& SkyMeshBatch = MainView.SkyMeshBatches[MeshBatchIndex];
-												if (!SkyMeshBatch.bVisibleInRealTimeSkyCapture)
-												{
-													continue;
-												}
-
-												const FMeshBatch* MeshBatch = SkyMeshBatch.Mesh;
-												const FPrimitiveSceneProxy* PrimitiveSceneProxy = SkyMeshBatch.Proxy;
-												const FPrimitiveSceneInfo* PrimitiveSceneInfo = PrimitiveSceneProxy->GetPrimitiveSceneInfo();
-
-												const uint64 DefaultBatchElementMask = ~0ull;
-												PassMeshProcessor.AddMeshBatch(*MeshBatch, DefaultBatchElementMask, PrimitiveSceneProxy);
-											}
-										});
-								});
+									FSkyMeshBatch& SkyMeshBatch = MainView.SkyMeshBatches[MeshBatchIndex];
+									if (!SkyMeshBatch.bVisibleInRealTimeSkyCapture)
+									{
+										continue;
+									}
+
+									const FMeshBatch* MeshBatch = SkyMeshBatch.Mesh;
+									const FPrimitiveSceneProxy* PrimitiveSceneProxy = SkyMeshBatch.Proxy;
+									const FPrimitiveSceneInfo* PrimitiveSceneInfo = PrimitiveSceneProxy->GetPrimitiveSceneInfo();
+
+									const uint64 DefaultBatchElementMask = ~0ull;
+									PassMeshProcessor.AddMeshBatch(*MeshBatch, DefaultBatchElementMask, PrimitiveSceneProxy);
+								}
+							});
 						}
 						else
 						{
