Convert RuntimeVirtualTexture to not use SubmitMeshDrawCommands directly, such that GPU-scene instance culling passes can be inserted.

diff --git a/Engine/Source/Runtime/Renderer/Private/VT/RuntimeVirtualTextureRender.cpp b/Engine/Source/Runtime/Renderer/Private/VT/RuntimeVirtualTextureRender.cpp
--- a/Engine/Source/Runtime/Renderer/Private/VT/RuntimeVirtualTextureRender.cpp
+++ b/Engine/Source/Runtime/Renderer/Private/VT/RuntimeVirtualTextureRender.cpp
@@ -451,21 +451,14 @@ namespace RuntimeVirtualTexture
 	FRegisterPassProcessorCreateFunction RegisterVirtualTexturePassMobile(&CreateRuntimeVirtualTexturePassProcessor, EShadingPath::Mobile, EMeshPass::VirtualTexture, EMeshPassFlags::CachedMeshCommands);
 
 
-	/** Collect meshes and draw. */
-	void DrawMeshes(FRHICommandListImmediate& RHICmdList, FScene const* Scene, FViewInfo const* View, ERuntimeVirtualTextureMaterialType MaterialType, uint32 RuntimeVirtualTextureMask, uint8 vLevel, uint8 MaxLevel)
+	/** Collect meshes to draw. */
+	void GatherMeshesToDraw(FDynamicPassMeshDrawListContext* DynamicMeshPassContext, FScene const* Scene, FViewInfo* View, ERuntimeVirtualTextureMaterialType MaterialType, uint32 RuntimeVirtualTextureMask, uint8 vLevel, uint8 MaxLevel)
 	{
 		// Cached draw command collectors
 		const FCachedPassMeshDrawList& SceneDrawList = Scene->CachedDrawLists[EMeshPass::VirtualTexture];
-		TArray<FVisibleMeshDrawCommand, TInlineAllocator<256>> CachedDrawCommands;
 
 		// Uncached mesh processor
-		FDynamicMeshDrawCommandStorage MeshDrawCommandStorage;
-		FMeshCommandOneFrameArray AllocatedCommands;
-		FGraphicsMinimalPipelineStateSet GraphicsMinimalPipelineStateSet;
-		bool NeedsShaderInitialisation;
-
-		FDynamicPassMeshDrawListContext DynamicMeshPassContext(MeshDrawCommandStorage, AllocatedCommands, GraphicsMinimalPipelineStateSet, NeedsShaderInitialisation);
-		FRuntimeVirtualTextureMeshProcessor MeshProcessor(Scene, View, &DynamicMeshPassContext);
+		FRuntimeVirtualTextureMeshProcessor MeshProcessor(Scene, View, DynamicMeshPassContext);
 
 		// Pre-calculate view factors used for culling
 		const float RcpWorldSize = 1.f / (View->ViewMatrices.GetInvProjectionMatrix().M[0][0]);
@@ -550,7 +543,7 @@ namespace RuntimeVirtualTexture
 								CachedMeshDrawCommand.Flags,
 								CachedMeshDrawCommand.SortKey);
 
-							CachedDrawCommands.Add(NewVisibleMeshDrawCommand);
+							DynamicMeshPassContext->AddVisibleMeshDrawCommand(NewVisibleMeshDrawCommand);
 							bCachedDraw = true;
 						}
 					}
@@ -564,30 +557,6 @@ namespace RuntimeVirtualTexture
 				}
 			}
 		}
-
-		// Combine cached and uncached draw command lists
-		const int32 NumCachedCommands = CachedDrawCommands.Num();
-		if (NumCachedCommands > 0)
-		{
-			const int32 CurrentCount = AllocatedCommands.Num();
-			AllocatedCommands.AddUninitialized(NumCachedCommands);
-			FMemory::Memcpy(&AllocatedCommands[CurrentCount], &CachedDrawCommands[0], NumCachedCommands * sizeof(CachedDrawCommands[0]));
-		}
-
-		// Sort and submit
-		if (AllocatedCommands.Num() > 0)
-		{
-			FRHIBuffer* PrimitiveIdsBuffer;
-
-			// GPUCULL_TODO: workaround for the fact that DrawDynamicMeshPassPrivate et al. don't work with GPU-Scene instancing
-			//               we don't support dynamic instancing for this path since we require one primitive per draw command
-			//               This is because the stride on the instance data buffer is set to 0 so only the first will ever be fetched.
-			const bool bDynamicInstancing = false;
-			const uint32 InstanceFactor = 1;
-
-			SortAndMergeDynamicPassMeshDrawCommands(View->FeatureLevel, AllocatedCommands, MeshDrawCommandStorage, PrimitiveIdsBuffer, InstanceFactor, View->DynamicPrimitiveCollector.GetPrimitiveIdRange());
-			SubmitMeshDrawCommands(AllocatedCommands, GraphicsMinimalPipelineStateSet, PrimitiveIdsBuffer, 0, bDynamicInstancing, InstanceFactor, RHICmdList);
-		}
 	}
 
 	/** BC Compression compute shader */
@@ -1068,8 +1037,7 @@ namespace RuntimeVirtualTexture
 		ViewInitOptions.BackgroundColor = FLinearColor::Black;
 		ViewInitOptions.OverlayColor = FLinearColor::White;
 
-		FViewInfo& ViewInfo = *GraphBuilder.AllocObject<FViewInfo>(ViewInitOptions);
-		FViewInfo* View = &ViewInfo;
+		FViewInfo* View = GraphBuilder.AllocObject<FViewInfo>(ViewInitOptions);
 		ViewFamily.Views.Add(View);
 
 		View->bIsVirtualTexture = true;
@@ -1096,13 +1064,10 @@ namespace RuntimeVirtualTexture
 			PassParameters->RenderTargets[1] = GraphSetup.RenderTexture1 ? FRenderTargetBinding(GraphSetup.RenderTexture1, LoadAction) : FRenderTargetBinding();
 			PassParameters->RenderTargets[2] = GraphSetup.RenderTexture2 ? FRenderTargetBinding(GraphSetup.RenderTexture2, LoadAction) : FRenderTargetBinding();
 
-			GraphBuilder.AddPass(
-				RDG_EVENT_NAME("VirtualTextureDraw"),
-				PassParameters,
-				ERDGPassFlags::Raster,
-				[Scene, View, MaterialType, RuntimeVirtualTextureMask, vLevel, MaxLevel](FRHICommandListImmediate& RHICmdListImmediate)
+			AddSimpleMeshPass(GraphBuilder, PassParameters, Scene, *View, nullptr, RDG_EVENT_NAME("VirtualTextureDraw"), View->ViewRect,
+			[Scene, View, MaterialType, RuntimeVirtualTextureMask, vLevel, MaxLevel](FDynamicPassMeshDrawListContext* DynamicMeshPassContext)
 			{
-				DrawMeshes(RHICmdListImmediate, Scene, View, MaterialType, RuntimeVirtualTextureMask, vLevel, MaxLevel);
+				GatherMeshesToDraw(DynamicMeshPassContext, Scene, View, MaterialType, RuntimeVirtualTextureMask, vLevel, MaxLevel);
 			});
 		}
 
diff --git a/Engine/Source/Runtime/Renderer/Public/MeshPassProcessor.h b/Engine/Source/Runtime/Renderer/Public/MeshPassProcessor.h
--- a/Engine/Source/Runtime/Renderer/Public/MeshPassProcessor.h
+++ b/Engine/Source/Runtime/Renderer/Public/MeshPassProcessor.h
@@ -1392,6 +1392,14 @@ public:
 		DrawList.Add(NewVisibleMeshDrawCommand);
 	}
 
+	/**
+	 * Use to add pre-built (cached) draw commands to a dynamic context. 
+	 */
+	FORCEINLINE void AddVisibleMeshDrawCommand(const FVisibleMeshDrawCommand& VisibleMeshDrawCommand)
+	{
+		DrawList.Add(VisibleMeshDrawCommand);
+	}
+
 private:
 	FDynamicMeshDrawCommandStorage& DrawListStorage;
 	FMeshCommandOneFrameArray& DrawList;
