Fix missing hair debug output for non-strands hair geometry.

diff --git a/Engine/Source/Runtime/Renderer/Private/HairStrands/HairStrandsDebug.cpp b/Engine/Source/Runtime/Renderer/Private/HairStrands/HairStrandsDebug.cpp
--- a/Engine/Source/Runtime/Renderer/Private/HairStrands/HairStrandsDebug.cpp
+++ b/Engine/Source/Runtime/Renderer/Private/HairStrands/HairStrandsDebug.cpp
@@ -285,7 +285,7 @@ static void AddDebugHairPrintPass(
 	const FHairStrandsMacroGroupResources& MacroGroupResources,
 	FRDGTextureSRVRef InDepthStencilTexture)
 {
-	if (!VisibilityData.CategorizationTexture || !VisibilityData.NodeIndex || !VisibilityData.NodeData || !InDepthStencilTexture) return;
+	if (!VisibilityData.CategorizationTexture || !VisibilityData.NodeIndex || !VisibilityData.NodeData || !InDepthStencilTexture || !ShaderDrawDebug::IsShaderDrawDebugEnabled(*View)) return;
 
 	FRDGTextureRef ViewHairCountTexture = VisibilityData.ViewHairCountTexture ? VisibilityData.ViewHairCountTexture : GSystemTextures.GetBlackDummy(GraphBuilder);
 	FRDGTextureRef ViewHairCountUintTexture = VisibilityData.ViewHairCountUintTexture ? VisibilityData.ViewHairCountUintTexture : GSystemTextures.GetBlackDummy(GraphBuilder);
@@ -1049,7 +1049,12 @@ static void InternalRenderHairStrandsDebugInfo(
 	const struct FHairStrandClusterData& HairClusterData,
 	FRDGTextureRef SceneColorTexture)
 {
-	check(HairStrands::HasViewHairStrandsData(View));
+	FHairStrandsBookmarkParameters Params = CreateHairStrandsBookmarkParameters(Scene, View);
+	Params.SceneColorTexture = SceneColorTexture;
+	if (!Params.bHasElements)
+	{
+		return;
+	}
 
 	const float YStep = 14;
 	const float ColumnWidth = 200;
@@ -1066,12 +1071,17 @@ static void InternalRenderHairStrandsDebugInfo(
 	// Only render debug information for the main view
 	const FSceneTextures& SceneTextures = FSceneTextures::Get(GraphBuilder);
 
+	// Bookmark for calling debug rendering from the plugin
 	{
-		FHairStrandsBookmarkParameters Params = CreateHairStrandsBookmarkParameters(Scene, View);
-		Params.SceneColorTexture = SceneColorTexture;
 		RunHairStrandsBookmark(GraphBuilder, EHairStrandsBookmark::ProcessDebug, Params);
 	}
 
+	// Pass this point, all debug rendering concern only hair strands data
+	if (!HairStrands::HasViewHairStrandsData(View))
+	{
+		return;
+	}
+
 	const FScreenPassRenderTarget SceneColor(SceneColorTexture, View.ViewRect, ERenderTargetLoadAction::ELoad);
 
 	// Debug mode name only
@@ -1374,9 +1384,6 @@ void RenderHairStrandsDebugInfo(
 	bool bHasHairData = false;
 	for (FViewInfo& View : Views)
 	{
-		if (HairStrands::HasViewHairStrandsData(View))
-		{
-			InternalRenderHairStrandsDebugInfo(GraphBuilder, Scene, View, HairClusterData, SceneColorTexture);
-		}
+		InternalRenderHairStrandsDebugInfo(GraphBuilder, Scene, View, HairClusterData, SceneColorTexture);
 	}
 }
\ No newline at end of file
