Fix CPU-side structure to align with padding added to GPU structure.

diff --git a/Engine/Source/Runtime/Renderer/Private/VirtualShadowMaps/VirtualShadowMapArray.cpp b/Engine/Source/Runtime/Renderer/Private/VirtualShadowMaps/VirtualShadowMapArray.cpp
--- a/Engine/Source/Runtime/Renderer/Private/VirtualShadowMaps/VirtualShadowMapArray.cpp
+++ b/Engine/Source/Runtime/Renderer/Private/VirtualShadowMaps/VirtualShadowMapArray.cpp
@@ -21,9 +21,10 @@ IMPLEMENT_STATIC_UNIFORM_BUFFER_STRUCT(FVirtualShadowMapUniformParameters, "Virt
 struct FShadowMapCacheData
 {
 	// XY offset in pages to the location of the previous frame's page table.
-	FIntPoint SmPageOffset;
+	FIntPoint SmPageOffset = FIntPoint(0, 0);
 	// ID of the corresponding virtual SM in the chached data 
-	int32 VirtualShadowMapId;
+	int32 VirtualShadowMapId = INDEX_NONE;
+	int Padding = 0;
 };
 
 
