Fix BlackDummy to not have alpha = 1.0 (fixes black reflection captures in views witout SSR pass)

diff --git a/Engine/Source/Runtime/Renderer/Private/SystemTextures.cpp b/Engine/Source/Runtime/Renderer/Private/SystemTextures.cpp
--- a/Engine/Source/Runtime/Renderer/Private/SystemTextures.cpp
+++ b/Engine/Source/Runtime/Renderer/Private/SystemTextures.cpp
@@ -99,7 +99,10 @@ void FSystemTextures::InitializeCommonTextures(FRHICommandListImmediate& RHICmdL
 
 	// Create a BlackDummy texture
 	{
-		BlackDummy = CreateRenderTarget(GBlackTexture->TextureRHI, TEXT("BlackDummy"));
+		FRHIResourceCreateInfo CreateInfo(TEXT("BlackDummy"));
+		FTexture2DRHIRef Texture2D = RHICreateTexture2D(1, 1, PF_B8G8R8A8, 1, 1, TexCreate_ShaderResource, CreateInfo);
+		SetDummyTextureData<FColor>(Texture2D, FColor(0, 0, 0, 0));
+		BlackDummy = CreateRenderTarget(Texture2D, CreateInfo.DebugName);
 	}
 	
 	// Create a texture that is a single UInt32 value set to 0
