Enable Path Tracing denoiser by default

Expose a checkbox on PostProcessVolume to control the path tracing denoiser.

Fix the hard clamp on samples/pixel since we don't have any internal limit. A UI clamp is sufficient.

Store the denoiser state in the path tracer config object so it does not get confused per View.

diff --git a/Engine/Plugins/Experimental/OpenImageDenoise/OpenImageDenoise.uplugin b/Engine/Plugins/Experimental/OpenImageDenoise/OpenImageDenoise.uplugin
--- a/Engine/Plugins/Experimental/OpenImageDenoise/OpenImageDenoise.uplugin
+++ b/Engine/Plugins/Experimental/OpenImageDenoise/OpenImageDenoise.uplugin
@@ -10,7 +10,7 @@
 	"DocsURL" : "",
 	"MarketplaceURL" : "",
 	"SupportURL" : "",
-	"EnabledByDefault" : false,
+	"EnabledByDefault" : true,
 	"CanContainContent" : false,
 	"IsBetaVersion" : false,
 	"Installed" : false,
diff --git a/Engine/Source/Runtime/Engine/Classes/Engine/Scene.h b/Engine/Source/Runtime/Engine/Classes/Engine/Scene.h
--- a/Engine/Source/Runtime/Engine/Classes/Engine/Scene.h
+++ b/Engine/Source/Runtime/Engine/Classes/Engine/Scene.h
@@ -1198,6 +1198,9 @@ struct FPostProcessSettings
 	UPROPERTY(EditAnywhere, BlueprintReadWrite, Category = Overrides, meta = (PinHiddenByDefault, InlineEditConditionToggle))
 	uint32 bOverride_PathTracingMaxPathExposure : 1;
 
+	UPROPERTY(EditAnywhere, BlueprintReadWrite, Category = Overrides, meta = (PinHiddenByDefault, InlineEditConditionToggle))
+	uint32 bOverride_PathTracingEnableDenoiser : 1;
+
 	// -----------------------------------------------------------------------
 
 	/** Enable HQ Gaussian on high end mobile platforms. (ES3_1) */
@@ -1924,7 +1927,7 @@ struct FPostProcessSettings
 	int32 PathTracingMaxBounces;
 
 	/** Sets the samples per pixel for the path tracer. */
-	UPROPERTY(interp, EditAnywhere, BlueprintReadWrite, Category = "Rendering Features|PathTracing", meta = (ClampMin = "1", ClampMax = "65536", editcondition = "bOverride_PathTracingSamplesPerPixel", DisplayName = "Samples Per Pixel"))
+	UPROPERTY(interp, EditAnywhere, BlueprintReadWrite, Category = "Rendering Features|PathTracing", meta = (ClampMin = "1", UIMax = "65536", editcondition = "bOverride_PathTracingSamplesPerPixel", DisplayName = "Samples Per Pixel"))
 	int32 PathTracingSamplesPerPixel;
 
 	/** Sets anti-aliasing filter width for the path tracer. Lower values are sharper (and more aliased), larger values are softer (and blurrier). */
@@ -1939,6 +1942,10 @@ struct FPostProcessSettings
 	UPROPERTY(interp, EditAnywhere, BlueprintReadWrite, Category = "Rendering Features|PathTracing", meta = (ClampMin = "-10.0", ClampMax = "30.0", editcondition = "bOverride_PathTracingMaxPathExposure", DisplayName = "Max Path Exposure"))
 	float PathTracingMaxPathExposure;
 
+	/** Run the currently loaded denoiser plugin on the last sample to remove noise from the output. Has no effect if a plug-in is not loaded. */
+	UPROPERTY(EditAnywhere, BlueprintReadWrite, Category = "Rendering Features|PathTracing", meta = (editcondition = "bOverride_PathTracingEnableDenoiser", DisplayName = "Denoiser"))
+	uint32 PathTracingEnableDenoiser : 1;
+
 
 	UPROPERTY()
 	float LPVFadeRange_DEPRECATED;
diff --git a/Engine/Source/Runtime/Engine/Private/Scene.cpp b/Engine/Source/Runtime/Engine/Private/Scene.cpp
--- a/Engine/Source/Runtime/Engine/Private/Scene.cpp
+++ b/Engine/Source/Runtime/Engine/Private/Scene.cpp
@@ -580,6 +580,7 @@ FPostProcessSettings::FPostProcessSettings()
 	PathTracingFilterWidth = 3.0f;
 	PathTracingEnableEmissive = 1;
 	PathTracingMaxPathExposure = 30.0f;
+	PathTracingEnableDenoiser = 1;
 	
 	bMobileHQGaussian = false;
 
@@ -763,6 +764,7 @@ FPostProcessSettings::FPostProcessSettings(const FPostProcessSettings& Settings)
 	, bOverride_PathTracingFilterWidth(Settings.bOverride_PathTracingFilterWidth)
 	, bOverride_PathTracingEnableEmissive(Settings.bOverride_PathTracingEnableEmissive)
 	, bOverride_PathTracingMaxPathExposure(Settings.bOverride_PathTracingMaxPathExposure)
+	, bOverride_PathTracingEnableDenoiser(Settings.bOverride_PathTracingEnableDenoiser)
 
 	, bMobileHQGaussian(Settings.bMobileHQGaussian)
 	, BloomMethod(Settings.BloomMethod)
@@ -935,6 +937,7 @@ FPostProcessSettings::FPostProcessSettings(const FPostProcessSettings& Settings)
 	, PathTracingFilterWidth(Settings.PathTracingFilterWidth)
 	, PathTracingEnableEmissive(Settings.PathTracingEnableEmissive)
 	, PathTracingMaxPathExposure(Settings.PathTracingMaxPathExposure)
+	, PathTracingEnableDenoiser(Settings.PathTracingEnableDenoiser)
 
 	, ScreenPercentage(Settings.ScreenPercentage)
 
diff --git a/Engine/Source/Runtime/Engine/Private/SceneView.cpp b/Engine/Source/Runtime/Engine/Private/SceneView.cpp
--- a/Engine/Source/Runtime/Engine/Private/SceneView.cpp
+++ b/Engine/Source/Runtime/Engine/Private/SceneView.cpp
@@ -1586,6 +1586,10 @@ void FSceneView::OverridePostProcessSettings(const FPostProcessSettings& Src, fl
 			Dest.PathTracingMaxPathExposure = Src.PathTracingMaxPathExposure;
 		}
 
+		if (Src.bOverride_PathTracingEnableDenoiser)
+		{
+			Dest.PathTracingEnableDenoiser = Src.PathTracingEnableDenoiser;
+		}
 
 		if (Src.bOverride_DepthOfFieldBladeCount)
 		{
diff --git a/Engine/Source/Runtime/Renderer/Private/PathTracing.cpp b/Engine/Source/Runtime/Renderer/Private/PathTracing.cpp
--- a/Engine/Source/Runtime/Renderer/Private/PathTracing.cpp
+++ b/Engine/Source/Runtime/Renderer/Private/PathTracing.cpp
@@ -183,18 +183,19 @@ TAutoConsoleVariable<int32> CVarPathTracingLightGridVisualize(
 	0,
 	TEXT("Enables a visualization mode of the light grid density where red indicates the maximum light count has been reached (default = 0)\n")
 	TEXT("0: off (default)\n")
-	TEXT("1: light count heatmap (red - close to overflow, increase r.PathTracing.LightGridMaxCount)")
-	TEXT("2: unique light lists (colors are a function of which lights occupy each cell)")
-	TEXT("3: area light visualization (green: point light sources only, blue: some area light sources)"),
+	TEXT("1: light count heatmap (red - close to overflow, increase r.PathTracing.LightGridMaxCount)\n")
+	TEXT("2: unique light lists (colors are a function of which lights occupy each cell)\n")
+	TEXT("3: area light visualization (green: point light sources only, blue: some area light sources)\n"),
 	ECVF_RenderThreadSafe
 );
 
 TAutoConsoleVariable<int32> CVarPathTracingDenoiser(
 	TEXT("r.PathTracing.Denoiser"),
-	0,
-	TEXT("Enable denoising of the path traced output (default = 0)\n")
-	TEXT("0: off (default)")
-	TEXT("1: enabled (if a denoiser plugin is active)"),
+	-1,
+	TEXT("Enable denoising of the path traced output (if a denoiser plugin is active) (default = -1 (driven by postprocesing volume))\n")
+	TEXT("-1: inherit from PostProcessVolume\n")
+	TEXT("0: disable denoiser\n")
+	TEXT("1: enable denoiser (if a denoiser plugin is active)\n"),
 	ECVF_RenderThreadSafe
 );
 
@@ -228,6 +229,7 @@ struct FPathTracingConfig
 	bool VisibleLights;
 	bool UseMISCompensation;
 	bool LockedSamplingPattern;
+	int DenoiserMode; // NOTE: does not require path tracing invalidation
 
 	bool IsDifferent(const FPathTracingConfig& Other) const
 	{
@@ -1393,9 +1395,12 @@ void FDeferredShadingSceneRenderer::RenderPathTracing(
 	}
 
 	FRDGTexture* DenoisedRadianceTexture = nullptr;
-	const int DenoiserMode = CVarPathTracingDenoiser.GetValueOnRenderThread();
+	int DenoiserMode = CVarPathTracingDenoiser.GetValueOnRenderThread();
+	if (DenoiserMode < 0)
+	{
+		DenoiserMode = View.FinalPostProcessSettings.PathTracingEnableDenoiser;
+	}
 	const bool IsDenoiserEnabled = DenoiserMode != 0 && GPathTracingDenoiserFunc != nullptr;
-	static int PrevDenoiserMode = DenoiserMode;
 	if (IsDenoiserEnabled)
 	{
 		// request denoise if this is the last sample
@@ -1404,7 +1409,7 @@ void FDeferredShadingSceneRenderer::RenderPathTracing(
 		if (!bNeedsMoreRays)
 		{
 			// we aren't currently rendering, run the denoiser if we just turned it on
-			NeedsDenoise |= DenoiserMode != PrevDenoiserMode;
+			NeedsDenoise |= DenoiserMode != View.ViewState->PathTracingLastConfig->DenoiserMode;
 		}
 
 		if (View.ViewState->PathTracingRadianceDenoisedRT)
@@ -1445,7 +1450,7 @@ void FDeferredShadingSceneRenderer::RenderPathTracing(
 			GraphBuilder.QueueTextureExtraction(DenoisedRadianceTexture, &View.ViewState->PathTracingRadianceDenoisedRT);
 		}
 	}
-	PrevDenoiserMode = DenoiserMode;
+	View.ViewState->PathTracingLastConfig->DenoiserMode = DenoiserMode;
 
 	// now add a pixel shader pass to display our Radiance buffer
 
