Fixed skylight  show flag not working properly for stationary skylight.

diff --git a/Engine/Source/Runtime/Renderer/Private/SceneRendering.cpp b/Engine/Source/Runtime/Renderer/Private/SceneRendering.cpp
--- a/Engine/Source/Runtime/Renderer/Private/SceneRendering.cpp
+++ b/Engine/Source/Runtime/Renderer/Private/SceneRendering.cpp
@@ -1680,7 +1680,11 @@ void FViewInfo::SetupUniformBufferParameters(
 	{
 		FSkyLightSceneProxy* SkyLight = Scene->SkyLight;
 
-		ViewUniformShaderParameters.SkyLightColor = SkyLight->GetEffectiveLightColor();
+		// Setup the sky color mulitpler, and use it to nullify the sky contribution in case SkyLighting is disabled.
+		// Note: we cannot simply select the base pass shader permutation skylight=0 because we would need to trigger bScenesPrimitivesNeedStaticMeshElementUpdate.
+		// However, this would need to be done per view (showflag is per view) and this is not possible today as it is selected within the scene. 
+		// So we simply nullify the sky light diffuse contribution. Reflection are handled by the indirect lighting render pass.
+		ViewUniformShaderParameters.SkyLightColor = Family->EngineShowFlags.SkyLighting ? SkyLight->GetEffectiveLightColor() : FLinearColor::Black;
 
 		bool bApplyPrecomputedBentNormalShadowing = 
 			SkyLight->bCastShadows 
