Move instance ID buffer to RDG and out of the View uniform buffer and into own UB.
 - Enables overlapping batched and non-batched instance culling (needed for batching work).
 - Removes some explicit transitions & minor cleanup.
 - Added tracking the required number of instances (fixes non-nanite VSM for large ISMs)

diff --git a/Engine/Plugins/Experimental/GPULightmass/Source/GPULightmass/Private/LightmapRenderer.cpp b/Engine/Plugins/Experimental/GPULightmass/Source/GPULightmass/Private/LightmapRenderer.cpp
--- a/Engine/Plugins/Experimental/GPULightmass/Source/GPULightmass/Private/LightmapRenderer.cpp
+++ b/Engine/Plugins/Experimental/GPULightmass/Source/GPULightmass/Private/LightmapRenderer.cpp
@@ -549,25 +549,6 @@ void FCachedRayTracingSceneData::SetupViewUniformBufferFromSceneRenderState(FSce
 			InstanceSceneDataSOAStride = InstanceSceneData.Num();
 		}
 
-		{
-			TRACE_CPUPROFILER_EVENT_SCOPE(InstanceIdsIdentityBuffer);
-
-			TResourceArray<uint32> InstanceIdsIdentity;
-			for (int32 Index = 0; Index < InstanceSceneData.Num(); Index++)
-			{
-				InstanceIdsIdentity.Add(Index);
-			}
-
-			FRHIResourceCreateInfo CreateInfo(TEXT("InstanceIdsIdentityBuffer"), &InstanceIdsIdentity);
-			if (InstanceIdsIdentity.GetResourceDataSize() == 0)
-			{
-				InstanceIdsIdentity.Add(0);
-			}
-
-			InstanceIdsIdentityBufferRHI = RHICreateStructuredBuffer(sizeof(uint32), InstanceIdsIdentity.GetResourceDataSize(), BUF_Static | BUF_ShaderResource, CreateInfo);
-			InstanceIdsIdentityBufferSRV = RHICreateShaderResourceView(InstanceIdsIdentityBufferRHI);
-		}
-
 		FViewUniformShaderParameters ViewUniformBufferParameters;
 		CachedViewUniformBuffer = TUniformBufferRef<FViewUniformShaderParameters>::CreateUniformBufferImmediate(ViewUniformBufferParameters, UniformBuffer_MultiFrame, EUniformBufferValidation::None);
 	}
@@ -842,7 +823,6 @@ void FSceneRenderState::SetupRayTracingScene(int32 LODIndex)
 				*View.CachedViewUniformShaderParameters);
 
 			View.CachedViewUniformShaderParameters->InstanceDataSOAStride = CachedRayTracingScene->InstanceSceneDataSOAStride;
-			View.CachedViewUniformShaderParameters->InstanceIdsBuffer = CachedRayTracingScene->InstanceIdsIdentityBufferSRV;
 
 			if (LightSceneRenderState.SkyLight.IsSet())
 			{
@@ -1204,6 +1184,7 @@ void FSceneRenderState::CalculateDistributionPrefixSumForAllLightmaps()
 
 BEGIN_SHADER_PARAMETER_STRUCT(FLightmapGBufferPassParameters, )
 	SHADER_PARAMETER_RDG_UNIFORM_BUFFER(FLightmapGBufferParams, PassUniformBuffer)
+	SHADER_PARAMETER_RDG_UNIFORM_BUFFER(FInstanceCullingGlobalUniforms, InstanceCulling)
 	RENDER_TARGET_BINDING_SLOTS()
 END_SHADER_PARAMETER_STRUCT()
 
@@ -1922,6 +1903,29 @@ void FLightmapRenderer::Finalize(FRDGBuilder& GraphBuilder)
 		PassUniformBuffer = GraphBuilder.CreateUniformBuffer(LightmapGBufferParameters);
 	}
 
+	TRDGUniformBufferRef<FInstanceCullingGlobalUniforms> InstanceCullingUniformBuffer = nullptr;
+	{
+		FRDGBufferRef InstanceIdsIdentityBuffer = nullptr;
+		{
+			TRACE_CPUPROFILER_EVENT_SCOPE(InstanceIdsIdentityBuffer);
+
+			TArray<uint32, SceneRenderingAllocator> InstanceIdsIdentity;
+			for (uint32 Index = 0U; Index < FMath::Max(1U, Scene->CachedRayTracingScene->InstanceSceneDataSOAStride); ++Index)
+			{
+				InstanceIdsIdentity.Add(Index);
+			}
+
+			InstanceIdsIdentityBuffer = CreateStructuredBuffer(GraphBuilder, TEXT("InstanceIdsIdentityBuffer"), InstanceIdsIdentity, ERDGInitialDataFlags::NoCopy);
+		}
+		FInstanceCullingGlobalUniforms* InstanceCullingUniforms = GraphBuilder.AllocParameters<FInstanceCullingGlobalUniforms>();
+
+		InstanceCullingUniforms->InstanceIdsBuffer = GraphBuilder.CreateSRV(InstanceIdsIdentityBuffer);
+		// Note redundant, but must have non-null reference even if not used it would seem
+		InstanceCullingUniforms->PageInfoBuffer = GraphBuilder.CreateSRV(InstanceIdsIdentityBuffer);
+		InstanceCullingUniforms->BufferCapacity = 0;
+		InstanceCullingUniformBuffer = GraphBuilder.CreateUniformBuffer(InstanceCullingUniforms);
+	}
+
 	RDG_GPU_MASK_SCOPE(GraphBuilder, FRHIGPUMask::GPU0());
 
 	IPooledRenderTarget* OutputRenderTargets[3] = { nullptr, nullptr, nullptr };
@@ -2101,6 +2105,7 @@ void FLightmapRenderer::Finalize(FRDGBuilder& GraphBuilder)
 
 							auto* PassParameters = GraphBuilder.AllocParameters<FLightmapGBufferPassParameters>();
 							PassParameters->PassUniformBuffer = PassUniformBuffer;
+							PassParameters->InstanceCulling = InstanceCullingUniformBuffer;
 
 							GraphBuilder.AddPass(
 								RDG_EVENT_NAME("LightmapGBuffer"),
@@ -2534,6 +2539,7 @@ void FLightmapRenderer::Finalize(FRDGBuilder& GraphBuilder)
 					{
 						auto* PassParameters = GraphBuilder.AllocParameters<FLightmapGBufferPassParameters>();
 						PassParameters->PassUniformBuffer = PassUniformBuffer;
+						PassParameters->InstanceCulling = InstanceCullingUniformBuffer;
 
 						GraphBuilder.AddPass(
 							RDG_EVENT_NAME("LightmapGBuffer"),
diff --git a/Engine/Shaders/Private/InstanceCulling/BuildInstanceDrawCommands.usf b/Engine/Shaders/Private/InstanceCulling/BuildInstanceDrawCommands.usf
--- a/Engine/Shaders/Private/InstanceCulling/BuildInstanceDrawCommands.usf
+++ b/Engine/Shaders/Private/InstanceCulling/BuildInstanceDrawCommands.usf
@@ -26,12 +26,9 @@ struct FBatchInfo
 };
 
 
-#if OUTPUT_COMMAND_IDS
-RWBuffer<uint> InstanceIdsBufferOut;
-RWBuffer<uint> DrawCommandIdsBufferOut;
-#else
-RWBuffer<uint> InstanceIdsBufferLegacyOut;
-#endif
+RWStructuredBuffer<uint> InstanceIdsBufferOut;
+RWStructuredBuffer<uint> DrawCommandIdsBufferOut;
+
 RWBuffer<uint> DrawIndirectArgsBufferOut;
 RWBuffer<uint> InstanceIdOffsetBufferOut;
 RWStructuredBuffer<uint> OutputOffsetBufferOut;
@@ -110,11 +107,9 @@ bool TestInstanceVisibilityBit(uint InstanceId, uint ViewId)
 void WriteInstanceId(uint Offset, uint InstanceId, uint ViewIdIndex, uint DrawCommandId)
 {
 	uint PackedId = InstanceId | (ViewIdIndex << 28U);
-#if OUTPUT_COMMAND_IDS
 	InstanceIdsBufferOut[Offset] = PackedId;
+#if OUTPUT_COMMAND_IDS
 	DrawCommandIdsBufferOut[Offset] = DrawCommandId;
-#else
-	InstanceIdsBufferLegacyOut[Offset] = PackedId;
 #endif
 }
 
diff --git a/Engine/Shaders/Private/SceneData.ush b/Engine/Shaders/Private/SceneData.ush
--- a/Engine/Shaders/Private/SceneData.ush
+++ b/Engine/Shaders/Private/SceneData.ush
@@ -380,11 +380,10 @@ FSceneDataIntermediates GetSceneDataIntermediates(uint InstanceIdOffset, uint Dr
 	}
 	else
 	{
-		// TODO: move InstanceIdsBuffer out of view and into some global Scene structure?
-		Intermediates.InstanceId = View.InstanceIdsBuffer[InstanceIdOffset + DrawInstanceId] & ((1U << 28U) - 1);
+		Intermediates.InstanceId = InstanceCulling.InstanceIdsBuffer[InstanceIdOffset + DrawInstanceId] & ((1U << 28U) - 1);
 		// We store the view index (which can be used for instanced stereo or other multi-view in the top four bits of the instance ID)
 		// Note: this is an index of views for this render pass, not the view ID in the culling manager.
-		Intermediates.ViewIndex = View.InstanceIdsBuffer[InstanceIdOffset + DrawInstanceId] >> 28U;
+		Intermediates.ViewIndex = InstanceCulling.InstanceIdsBuffer[InstanceIdOffset + DrawInstanceId] >> 28U;
 	}
 
 	Intermediates.InstanceData = GetInstanceData(Intermediates.InstanceId, View.InstanceDataSOAStride);
diff --git a/Engine/Shaders/Private/ShadowDepthVertexShader.usf b/Engine/Shaders/Private/ShadowDepthVertexShader.usf
--- a/Engine/Shaders/Private/ShadowDepthVertexShader.usf
+++ b/Engine/Shaders/Private/ShadowDepthVertexShader.usf
@@ -391,7 +391,7 @@ void Main(
 	{
 		// Get the offset from which we loaded the instance ID
 		uint InstanceIdIndex = VertexFactoryGetInstanceIdLoadIndex(VFIntermediates);
-		FPageInfo PageInfo = UnPackPageInfo(View.PageInfoBuffer[InstanceIdIndex]);
+		FPageInfo PageInfo = UnPackPageInfo(InstanceCulling.PageInfoBuffer[InstanceIdIndex]);
 
 		TransformToVirtualSmPage(OutPosition, OutVirtualSmPageClip, PageInfo, WorldPos.xyz);
 		
@@ -518,7 +518,7 @@ void PositionOnlyMain(
 		uint InstanceIdIndex = VertexFactoryGetInstanceIdLoadIndex(Input);
 		// TODO: Maybe offset index to local buffer for this? We may not want to use a global since most passes are not supporting this
 		//       Or perhaps both are different, as they are managed somewhat differently anyway, maybe.
-		FPageInfo PageInfo = UnPackPageInfo(View.PageInfoBuffer[InstanceIdIndex]);
+		FPageInfo PageInfo = UnPackPageInfo(InstanceCulling.PageInfoBuffer[InstanceIdIndex]);
 		TransformToVirtualSmPage(OutPosition, OutVirtualSmPageClip, PageInfo, WorldPos.xyz);
 
 		ViewId = PageInfo.ViewId;
diff --git a/Engine/Shaders/Private/VirtualShadowMaps/BuildPerPageDrawCommands.usf b/Engine/Shaders/Private/VirtualShadowMaps/BuildPerPageDrawCommands.usf
--- a/Engine/Shaders/Private/VirtualShadowMaps/BuildPerPageDrawCommands.usf
+++ b/Engine/Shaders/Private/VirtualShadowMaps/BuildPerPageDrawCommands.usf
@@ -37,8 +37,8 @@ uint PackPageInfo(FPageInfo PageInfo)
 
 
 RWStructuredBuffer<FVisibleInstanceCmd> VisibleInstancesOut;
-Buffer<uint> InstanceIdsBuffer;
-Buffer<uint> DrawCommandIdsBuffer;
+StructuredBuffer<uint> InstanceIdsBuffer;
+StructuredBuffer<uint> DrawCommandIdsBuffer;
 StructuredBuffer<uint> NumInstanceIdsBuffer;
 uint FirstPrimaryView;
 uint NumPrimaryViews;
@@ -469,8 +469,8 @@ void AllocateCommandInstanceOutputSpaceCs(uint DrawCommandId : SV_DispatchThread
 StructuredBuffer<FVisibleInstanceCmd> VisibleInstances;
 StructuredBuffer <uint> VisibleInstanceCountBuffer;
 //RWStructuredBuffer<uint> TmpInstanceIdOffsetBufferOut;
-RWBuffer<uint> InstanceIdsBufferLegacyOut;
-RWBuffer<uint> PageInfoBufferLegacyOut;
+RWStructuredBuffer<uint> InstanceIdsBufferOut;
+RWStructuredBuffer<uint> PageInfoBufferOut;
 
 
 [numthreads(NUM_THREADS_PER_GROUP, 1, 1)]
@@ -486,7 +486,7 @@ void OutputCommandInstanceListsCs(uint VisibleInstanceIndex : SV_DispatchThreadI
 		uint InstanceIdOutputOffset = 0;
 		InterlockedAdd(TmpInstanceIdOffsetBufferOut[VisibleInstanceCmd.DrawCommandId], 1U, InstanceIdOutputOffset);
 		// TODO: maybe repack as uint2 since that might be better for these type of presumably scalar loads.
-		InstanceIdsBufferLegacyOut[InstanceIdOutputOffset] = VisibleInstanceCmd.InstanceId;
-		PageInfoBufferLegacyOut[InstanceIdOutputOffset] = VisibleInstanceCmd.PackedPageInfo;
+		InstanceIdsBufferOut[InstanceIdOutputOffset] = VisibleInstanceCmd.InstanceId;
+		PageInfoBufferOut[InstanceIdOutputOffset] = VisibleInstanceCmd.PackedPageInfo;
 	}
 }
\ No newline at end of file
diff --git a/Engine/Source/Runtime/Engine/Private/SceneManagement.cpp b/Engine/Source/Runtime/Engine/Private/SceneManagement.cpp
--- a/Engine/Source/Runtime/Engine/Private/SceneManagement.cpp
+++ b/Engine/Source/Runtime/Engine/Private/SceneManagement.cpp
@@ -805,10 +805,6 @@ FViewUniformShaderParameters::FViewUniformShaderParameters()
 	EditorVisualizeLevelInstanceIds = GIdentityPrimitiveBuffer.EditorVisualizeLevelInstanceDataBufferSRV;
 	EditorSelectedHitProxyIds = GIdentityPrimitiveBuffer.EditorSelectedDataBufferSRV;
 //#endif
-
-	// Automatically bind these for everyone to use!
-	InstanceIdsBuffer = GInstanceCullingManagerResources.GetInstancesIdBufferSrv();
-	PageInfoBuffer = GInstanceCullingManagerResources.GetPageInfoBufferSrv();
 }
 
 FInstancedViewUniformShaderParameters::FInstancedViewUniformShaderParameters()
@@ -1283,20 +1279,3 @@ void FDefaultMobileReflectionCaptureUniformBuffer::InitDynamicRHI()
 
 /** Global uniform buffer containing the default reflection data used in mobile renderer. */
 TGlobalResource<FDefaultMobileReflectionCaptureUniformBuffer> GDefaultMobileReflectionCaptureUniformBuffer;
-
-
-
-void FInstanceCullingManagerResources::InitRHI()
-{
-	InstanceIdsBuffer.Initialize(TEXT("InstanceCulling.InstanceIds"), sizeof(uint32), MaxIndirectInstances, EPixelFormat::PF_R32_UINT, ERHIAccess::SRVGraphics, BUF_Static);
-	PageInfoBuffer.Initialize(TEXT("InstanceCulling.PageInfoBuffer"), sizeof(uint32), MaxIndirectInstances, EPixelFormat::PF_R32_UINT, ERHIAccess::SRVGraphics, BUF_Static);
-}
-
-void FInstanceCullingManagerResources::ReleaseRHI()
-{
-	InstanceIdsBuffer.Release();
-	PageInfoBuffer.Release();
-}
-
-
-TGlobalResource<FInstanceCullingManagerResources> GInstanceCullingManagerResources;
diff --git a/Engine/Source/Runtime/Engine/Public/SceneManagement.h b/Engine/Source/Runtime/Engine/Public/SceneManagement.h
--- a/Engine/Source/Runtime/Engine/Public/SceneManagement.h
+++ b/Engine/Source/Runtime/Engine/Public/SceneManagement.h
@@ -2180,31 +2180,6 @@ private:
 };
 
 
-class FInstanceCullingManagerResources : public FRenderResource
-{
-public:
-	/**
-	 */
-	static constexpr uint32 MaxIndirectInstances = 1024 * 1024;
-
-	virtual void InitRHI() override;
-	virtual void ReleaseRHI() override;
-
-	FRHIBuffer* GetInstancesIdBuffer() const { return InstanceIdsBuffer.Buffer; }
-
-	FRHIShaderResourceView* GetInstancesIdBufferSrv() const { return InstanceIdsBuffer.SRV.GetReference(); }
-	FRHIShaderResourceView* GetPageInfoBufferSrv() const { return PageInfoBuffer.SRV.GetReference(); }
-
-	FUnorderedAccessViewRHIRef GetInstancesIdBufferUav() const { return InstanceIdsBuffer.UAV; }
-	FUnorderedAccessViewRHIRef GetPageInfoBufferUav() const { return PageInfoBuffer.UAV; }
-
-private:
-	FRWBuffer PageInfoBuffer;
-	FRWBuffer InstanceIdsBuffer;
-};
-
-extern ENGINE_API TGlobalResource<FInstanceCullingManagerResources> GInstanceCullingManagerResources;
-
 /** 
  * Base class for a resource allocated from a FMeshElementCollector with AllocateOneFrameResource, which the collector releases.
  * This is useful for per-frame structures which are referenced by a mesh batch given to the FMeshElementCollector.
diff --git a/Engine/Source/Runtime/Engine/Public/SceneView.h b/Engine/Source/Runtime/Engine/Public/SceneView.h
--- a/Engine/Source/Runtime/Engine/Public/SceneView.h
+++ b/Engine/Source/Runtime/Engine/Public/SceneView.h
@@ -853,9 +853,6 @@ BEGIN_GLOBAL_SHADER_PARAMETER_STRUCT_WITH_CONSTRUCTOR(FViewUniformShaderParamete
 
 	SHADER_PARAMETER_SRV(Buffer<float>, PhysicsFieldClipmapBuffer)
 
-	SHADER_PARAMETER_SRV(Buffer<uint>, InstanceIdsBuffer)
-	SHADER_PARAMETER_SRV(Buffer<uint>, PageInfoBuffer)
-
 END_GLOBAL_SHADER_PARAMETER_STRUCT()
 
 /** Copy of the view uniform shader parameters associated with a view for instanced stereo. */
diff --git a/Engine/Source/Runtime/Landscape/Private/LandscapeGrass.cpp b/Engine/Source/Runtime/Landscape/Private/LandscapeGrass.cpp
--- a/Engine/Source/Runtime/Landscape/Private/LandscapeGrass.cpp
+++ b/Engine/Source/Runtime/Landscape/Private/LandscapeGrass.cpp
@@ -553,6 +553,7 @@ public:
 
 	BEGIN_SHADER_PARAMETER_STRUCT(FLandscapeGrassPassParameters, )
 		SHADER_PARAMETER_STRUCT_REF(FViewUniformShaderParameters, View)
+		SHADER_PARAMETER_RDG_UNIFORM_BUFFER(FInstanceCullingGlobalUniforms, InstanceCulling)
 		RENDER_TARGET_BINDING_SLOTS()
 	END_SHADER_PARAMETER_STRUCT()
 
@@ -580,6 +581,7 @@ public:
 
 		auto* PassParameters = GraphBuilder.AllocParameters<FLandscapeGrassPassParameters>();
 		PassParameters->View = View->ViewUniformBuffer;
+		PassParameters->InstanceCulling = FInstanceCullingContext::CreateDummyInstanceCullingUniformBuffer(GraphBuilder);
 		PassParameters->RenderTargets[0] = FRenderTargetBinding(OutputTexture, ERenderTargetLoadAction::EClear);
 
 		GraphBuilder.AddPass(
diff --git a/Engine/Source/Runtime/Landscape/Private/LandscapePhysicalMaterial.cpp b/Engine/Source/Runtime/Landscape/Private/LandscapePhysicalMaterial.cpp
--- a/Engine/Source/Runtime/Landscape/Private/LandscapePhysicalMaterial.cpp
+++ b/Engine/Source/Runtime/Landscape/Private/LandscapePhysicalMaterial.cpp
@@ -276,7 +276,7 @@ namespace
 	
 	BEGIN_SHADER_PARAMETER_STRUCT(FLandscapePhysicalMaterialPassParameters, )
 		SHADER_PARAMETER_STRUCT_REF(FViewUniformShaderParameters, View)
-		SHADER_PARAMETER_STRUCT_INCLUDE(FInstanceCullingDrawParams, InstanceCullingDrawParams)
+		SHADER_PARAMETER_RDG_UNIFORM_BUFFER(FInstanceCullingGlobalUniforms, InstanceCulling)
 		RENDER_TARGET_BINDING_SLOTS()
 	END_SHADER_PARAMETER_STRUCT()
 
@@ -316,6 +316,7 @@ namespace
 
 		auto* PassParameters = GraphBuilder.AllocParameters<FLandscapePhysicalMaterialPassParameters>();
 		PassParameters->View = View->ViewUniformBuffer;
+		PassParameters->InstanceCulling = FInstanceCullingContext::CreateDummyInstanceCullingUniformBuffer(GraphBuilder);
 		PassParameters->RenderTargets[0] = FRenderTargetBinding(OutputTexture, ERenderTargetLoadAction::EClear);
 
 		GraphBuilder.AddPass(
diff --git a/Engine/Source/Runtime/Renderer/Private/CompositionLighting/PostProcessDeferredDecals.cpp b/Engine/Source/Runtime/Renderer/Private/CompositionLighting/PostProcessDeferredDecals.cpp
--- a/Engine/Source/Runtime/Renderer/Private/CompositionLighting/PostProcessDeferredDecals.cpp
+++ b/Engine/Source/Runtime/Renderer/Private/CompositionLighting/PostProcessDeferredDecals.cpp
@@ -83,6 +83,7 @@ FDeferredDecalPassTextures GetDeferredDecalPassTextures(
 }
 
 void GetDeferredDecalPassParameters(
+	FRDGBuilder &GraphBuilder,
 	const FViewInfo& View,
 	const FDeferredDecalPassTextures& Textures,
 	EDecalRenderTargetMode RenderTargetMode,
@@ -91,6 +92,9 @@ void GetDeferredDecalPassParameters(
 	PassParameters.View = View.GetShaderParameters();
 	PassParameters.DeferredDecal = CreateDeferredDecalUniformBuffer(View);
 	PassParameters.DecalPass = Textures.DecalPassUniformBuffer;
+	
+	// TODO: hook up to instance culling manager and convert to simple mesh pass.
+	PassParameters.InstanceCulling = FInstanceCullingContext::CreateDummyInstanceCullingUniformBuffer(GraphBuilder);
 
 	FRDGTextureRef DepthTexture = Textures.Depth.Target;
 
@@ -456,7 +460,7 @@ void AddDeferredDecalPass(
 	const auto RenderDecals = [&](uint32 DecalIndexBegin, uint32 DecalIndexEnd, EDecalRenderTargetMode RenderTargetMode)
 	{
 		auto* PassParameters = GraphBuilder.AllocParameters<FDeferredDecalPassParameters>();
-		GetDeferredDecalPassParameters(View, PassTextures, RenderTargetMode, *PassParameters);
+		GetDeferredDecalPassParameters(GraphBuilder, View, PassTextures, RenderTargetMode, *PassParameters);
 
 		GraphBuilder.AddPass(
 			RDG_EVENT_NAME("Batch [%d, %d]", DecalIndexBegin, DecalIndexEnd - 1),
diff --git a/Engine/Source/Runtime/Renderer/Private/CompositionLighting/PostProcessDeferredDecals.h b/Engine/Source/Runtime/Renderer/Private/CompositionLighting/PostProcessDeferredDecals.h
--- a/Engine/Source/Runtime/Renderer/Private/CompositionLighting/PostProcessDeferredDecals.h
+++ b/Engine/Source/Runtime/Renderer/Private/CompositionLighting/PostProcessDeferredDecals.h
@@ -5,6 +5,7 @@
 #include "RenderGraph.h"
 #include "SceneTextureParameters.h"
 #include "SceneView.h"
+#include "InstanceCulling/InstanceCullingContext.h"
 
 enum class EDecalRenderStage : uint8;
 enum class EDecalRenderTargetMode : uint8;
@@ -61,10 +62,12 @@ BEGIN_SHADER_PARAMETER_STRUCT(FDeferredDecalPassParameters, )
 	SHADER_PARAMETER_STRUCT_INCLUDE(FViewShaderParameters, View)
 	SHADER_PARAMETER_STRUCT_REF(FDeferredDecalUniformParameters, DeferredDecal)
 	SHADER_PARAMETER_RDG_UNIFORM_BUFFER(FDecalPassUniformParameters, DecalPass)
+	SHADER_PARAMETER_RDG_UNIFORM_BUFFER(FInstanceCullingGlobalUniforms, InstanceCulling)
 	RENDER_TARGET_BINDING_SLOTS()
 END_SHADER_PARAMETER_STRUCT()
 
 void GetDeferredDecalPassParameters(
+	FRDGBuilder& GraphBuilder,
 	const FViewInfo& View,
 	const FDeferredDecalPassTextures& DecalPassTextures,
 	EDecalRenderTargetMode RenderTargetMode,
diff --git a/Engine/Source/Runtime/Renderer/Private/CompositionLighting/PostProcessMeshDecals.cpp b/Engine/Source/Runtime/Renderer/Private/CompositionLighting/PostProcessMeshDecals.cpp
--- a/Engine/Source/Runtime/Renderer/Private/CompositionLighting/PostProcessMeshDecals.cpp
+++ b/Engine/Source/Runtime/Renderer/Private/CompositionLighting/PostProcessMeshDecals.cpp
@@ -311,7 +311,7 @@ void DrawDecalMeshCommands(
 	EDecalRenderTargetMode RenderTargetMode)
 {
 	auto* PassParameters = GraphBuilder.AllocParameters<FDeferredDecalPassParameters>();
-	GetDeferredDecalPassParameters(View, DecalPassTextures, RenderTargetMode, *PassParameters);
+	GetDeferredDecalPassParameters(GraphBuilder, View, DecalPassTextures, RenderTargetMode, *PassParameters);
 
 	GraphBuilder.AddPass(
 		RDG_EVENT_NAME("MeshDecals"),
diff --git a/Engine/Source/Runtime/Renderer/Private/DeferredShadingRenderer.cpp b/Engine/Source/Runtime/Renderer/Private/DeferredShadingRenderer.cpp
--- a/Engine/Source/Runtime/Renderer/Private/DeferredShadingRenderer.cpp
+++ b/Engine/Source/Runtime/Renderer/Private/DeferredShadingRenderer.cpp
@@ -1875,7 +1875,7 @@ void FDeferredShadingSceneRenderer::Render(FRDGBuilder& GraphBuilder)
 	// Find the visible primitives.
 	GraphBuilder.RHICmdList.ImmediateFlush(EImmediateFlushType::DispatchToRHIThread);
 
-	FInstanceCullingManager InstanceCullingManager(GInstanceCullingManagerResources, Scene->GPUScene.IsEnabled());
+	FInstanceCullingManager InstanceCullingManager(Scene->GPUScene.IsEnabled());
 
 	bool bDoInitViewAftersPrepass = false;
 	{
diff --git a/Engine/Source/Runtime/Renderer/Private/InstanceCulling/InstanceCullingContext.cpp b/Engine/Source/Runtime/Renderer/Private/InstanceCulling/InstanceCullingContext.cpp
--- a/Engine/Source/Runtime/Renderer/Private/InstanceCulling/InstanceCullingContext.cpp
+++ b/Engine/Source/Runtime/Renderer/Private/InstanceCulling/InstanceCullingContext.cpp
@@ -14,6 +14,8 @@
 
 #define ENABLE_DETERMINISTIC_INSTANCE_CULLING 0
 
+IMPLEMENT_STATIC_UNIFORM_BUFFER_SLOT(InstanceCullingUbSlot);
+IMPLEMENT_STATIC_UNIFORM_BUFFER_STRUCT(FInstanceCullingGlobalUniforms, "InstanceCulling", InstanceCullingUbSlot);
 
 struct FBatchInfo
 {
@@ -45,6 +47,7 @@ void FInstanceCullingContext::ResetCommands(int32 MaxNumCommands)
 	CullingCommands.Empty(MaxNumCommands);
 	InstanceRuns.Reset();
 	PrimitiveIds.Reset();
+	TotalInstances = 0U;
 }
 
 void FInstanceCullingContext::BeginCullingCommand(EPrimitiveType BatchType, uint32 BaseVertexIndex, uint32 FirstIndex, uint32 NumPrimitives, bool bInMaterialMayModifyPosition)
@@ -80,8 +83,9 @@ void FInstanceCullingContext::BeginCullingCommand(EPrimitiveType BatchType, uint
 	}
 }
 
-void FInstanceCullingContext::AddPrimitiveToCullingCommand(int32 ScenePrimitiveId)
+void FInstanceCullingContext::AddPrimitiveToCullingCommand(int32 ScenePrimitiveId, uint32 NumInstances)
 {
+	TotalInstances += NumInstances;
 	PrimitiveIds.Add(ScenePrimitiveId);
 }
 
@@ -90,7 +94,10 @@ void FInstanceCullingContext::AddInstanceRunToCullingCommand(int32 ScenePrimitiv
 	//InstanceRuns.AddDefaulted(NumRuns);
 	for (uint32 Index = 0; Index < NumRuns; ++Index)
 	{
-		InstanceRuns.Add(FInstanceRun{ Runs[Index * 2], Runs[Index * 2 + 1], ScenePrimitiveId });
+		uint32 RunStart = Runs[Index * 2];
+		uint32 RunEndIncl = Runs[Index * 2 + 1];
+		TotalInstances += (RunEndIncl + 1U) - RunStart;
+		InstanceRuns.Add(FInstanceRun{ RunStart, RunEndIncl, ScenePrimitiveId });
 	}
 }
 
@@ -231,7 +238,7 @@ public:
 		SHADER_PARAMETER_RDG_BUFFER_UAV(RWBuffer<uint>, InstanceIdsBufferOut)
 		SHADER_PARAMETER_RDG_BUFFER_UAV(RWBuffer<uint>, DrawCommandIdsBufferOut)
 		// Using the wrong kind of buffer for RDG...
-		SHADER_PARAMETER_UAV(RWBuffer<uint>, InstanceIdsBufferLegacyOut)
+		SHADER_PARAMETER_UAV(RWStructuredBuffer<uint>, InstanceIdsBufferOut)
 
 		SHADER_PARAMETER_RDG_BUFFER_UAV(RWBuffer<uint>, DrawIndirectArgsBufferOut)
 		SHADER_PARAMETER(int32, NumPrimitiveIds)
@@ -305,10 +312,8 @@ public:
 		SHADER_PARAMETER_RDG_BUFFER_UAV(RWStructuredBuffer<uint>, OutputOffsetBufferOut)
 		SHADER_PARAMETER_RDG_BUFFER_UAV(RWBuffer<uint>, InstanceIdOffsetBufferOut)
 
-		SHADER_PARAMETER_RDG_BUFFER_UAV(RWBuffer<uint>, InstanceIdsBufferOut)
-		SHADER_PARAMETER_RDG_BUFFER_UAV(RWBuffer<uint>, DrawCommandIdsBufferOut)
-		// Using the wrong kind of buffer for RDG...
-		SHADER_PARAMETER_UAV(RWBuffer<uint>, InstanceIdsBufferLegacyOut)
+		SHADER_PARAMETER_RDG_BUFFER_UAV(RWStructuredBuffer<uint>, InstanceIdsBufferOut)
+		SHADER_PARAMETER_RDG_BUFFER_UAV(RWStructuredBuffer<uint>, DrawCommandIdsBufferOut)
 
 		SHADER_PARAMETER_RDG_BUFFER_UAV(RWBuffer<uint>, DrawIndirectArgsBufferOut)
 		SHADER_PARAMETER(int32, NumPrimitiveIds)
@@ -326,11 +331,25 @@ public:
 };
 IMPLEMENT_GLOBAL_SHADER(FBuildInstanceIdBufferAndCommandsFromPrimitiveIdsCs, "/Engine/Private/InstanceCulling/BuildInstanceDrawCommands.usf", "BuildInstanceIdBufferAndCommandsFromPrimitiveIdsCs", SF_Compute);
 
+const TRDGUniformBufferRef<FInstanceCullingGlobalUniforms> FInstanceCullingContext::CreateDummyInstanceCullingUniformBuffer(FRDGBuilder& GraphBuilder)
+{
+	FInstanceCullingGlobalUniforms* InstanceCullingGlobalUniforms = GraphBuilder.AllocParameters<FInstanceCullingGlobalUniforms>();
+	FRDGBufferRef DummyBuffer = GSystemTextures.GetDefaultStructuredBuffer(GraphBuilder, 4);
+	InstanceCullingGlobalUniforms->InstanceIdsBuffer = GraphBuilder.CreateSRV(DummyBuffer);
+	InstanceCullingGlobalUniforms->PageInfoBuffer = GraphBuilder.CreateSRV(DummyBuffer);
+	InstanceCullingGlobalUniforms->BufferCapacity = 0;
+	return GraphBuilder.CreateUniformBuffer(InstanceCullingGlobalUniforms);
+}
+
 void FInstanceCullingContext::BuildRenderingCommands(FRDGBuilder& GraphBuilder, const FGPUScene& GPUScene, const TRange<int32>& DynamicPrimitiveIdRange, FInstanceCullingResult& Results) const
 {
 	Results = FInstanceCullingResult();
-	if (CullingCommands.Num() == 0)
+	if (CullingCommands.Num() == 0 || TotalInstances == 0)
 	{
+		if (InstanceCullingManager)
+		{
+			Results.UniformBuffer = InstanceCullingManager->GetDummyInstanceCullingUniformBuffer();
+		}
 		return;
 	}
 
@@ -342,9 +361,9 @@ void FInstanceCullingContext::BuildRenderingCommands(FRDGBuilder& GraphBuilder,
 	FGlobalShaderMap* ShaderMap = GetGlobalShaderMap(FeatureLevel);
 
 	// Note: use start at zero offset if there is no instance culling manager, this means each build rendering commands pass will overwrite the same ID range. Which is only ok assuming correct barriers (should be erring on this side by default).
-	TArray<uint32> NullArray;
+	TArray<uint32, TInlineAllocator<1> > NullArray;
 	NullArray.AddZeroed(1);
-	FRDGBufferRef InstanceIdOutOffsetBufferRDG = InstanceCullingManager ? InstanceCullingManager->CullingIntermediate.InstanceIdOutOffsetBuffer : CreateStructuredBuffer(GraphBuilder, TEXT("InstanceCulling.OutputOffsetBufferOutTransient"), NullArray);
+	FRDGBufferRef InstanceIdOutOffsetBufferRDG = CreateStructuredBuffer(GraphBuilder, TEXT("InstanceCulling.OutputOffsetBufferOut"), NullArray);
 	// If there is no manager, then there is no data on culling, so set flag to skip that and ignore buffers.
 	FRDGBufferRef VisibleInstanceFlagsRDG = InstanceCullingManager != nullptr ? InstanceCullingManager->CullingIntermediate.VisibleInstanceFlags : nullptr;
 	const bool bCullInstances = InstanceCullingManager != nullptr;
@@ -360,6 +379,9 @@ void FInstanceCullingContext::BuildRenderingCommands(FRDGBuilder& GraphBuilder,
 	FRDGBufferRef PrimitiveInstanceRunsBuffer = CreateStructuredBuffer(GraphBuilder, TEXT("InstanceCulling.InstanceRuns"), InstanceRuns);
 	FRDGBufferRef ViewIdsBuffer = CreateStructuredBuffer(GraphBuilder, TEXT("InstanceCulling.ViewIds"), ViewIds);
 
+	const uint32 InstanceIdBufferSize = TotalInstances * ViewIds.Num();
+	FRDGBufferRef InstanceIdsBuffer = GraphBuilder.CreateBuffer(FRDGBufferDesc::CreateStructuredDesc(sizeof(uint32), InstanceIdBufferSize), TEXT("InstanceCulling.InstanceIdsBuffer"));
+
 #if ENABLE_DETERMINISTIC_INSTANCE_CULLING
 	// 1. Compute output sizes for all commands
 	FRDGBufferRef InstanceCountsRDG = GraphBuilder.CreateBuffer(FRDGBufferDesc::CreateStructuredDesc(sizeof(FIntVector2), CullingCommands.Num()), TEXT("InstanceCulling.InstanceCounts"));
@@ -460,15 +482,9 @@ void FInstanceCullingContext::BuildRenderingCommands(FRDGBuilder& GraphBuilder,
 		PassParameters->ViewIds = GraphBuilder.CreateSRV(ViewIdsBuffer);
 		PassParameters->NumViewIds = ViewIds.Num();
 
-		// TODO: Remove this when everything is properly RDG'd
-		AddPass(GraphBuilder, [&GraphBuilder](FRHICommandList& RHICmdList)
-		{
-			RHICmdList.Transition(FRHITransitionInfo(GInstanceCullingManagerResources.GetInstancesIdBufferUav(), ERHIAccess::Unknown, ERHIAccess::UAVCompute));
-		});
-
 		//PassParameters->InstanceIdsBufferOut = GraphBuilder.CreateUAV(InstanceIdsBufferRDG, PF_R32_UINT);
 		// TODO: Access resources through manager rather than global
-		PassParameters->InstanceIdsBufferLegacyOut = GInstanceCullingManagerResources.GetInstancesIdBufferUav();
+		PassParameters->InstanceIdsBufferOut = GInstanceCullingManagerResources.GetInstancesIdBufferUav();
 		PassParameters->DrawIndirectArgsBufferOut = GraphBuilder.CreateUAV(DrawIndirectArgsRDG, PF_R32_UINT);
 		PassParameters->NumPrimitiveIds = PrimitiveIds.Num();
 		PassParameters->NumInstanceRuns = InstanceRuns.Num();
@@ -480,7 +496,6 @@ void FInstanceCullingContext::BuildRenderingCommands(FRDGBuilder& GraphBuilder,
 		PassParameters->NumCulledInstances = uint32(NumCulledInstances);
 
 		FOutputInstanceIdsAtOffsetCs::FPermutationDomain PermutationVector;
-		// NOTE: this also switches between legacy buffer and RDG for Id output
 		PermutationVector.Set<FOutputInstanceIdsAtOffsetCs::FOutputCommandIdDim>(0);
 		PermutationVector.Set<FOutputInstanceIdsAtOffsetCs::FCullInstancesDim>(bCullInstances);
 		PermutationVector.Set<FOutputInstanceIdsAtOffsetCs::FStereoModeDim>(InstanceCullingMode == EInstanceCullingMode::Stereo);
@@ -520,22 +535,13 @@ void FInstanceCullingContext::BuildRenderingCommands(FRDGBuilder& GraphBuilder,
 
 	FRDGBufferRef DrawIndirectArgsRDG = GraphBuilder.CreateBuffer(FRDGBufferDesc::CreateIndirectDesc(IndirectArgsNumWords * CullingCommands.Num()), TEXT("InstanceCulling.DrawIndirectArgsBuffer"));
 	// not using structured buffer as we want/have toget at it as a vertex buffer 
-	//FRDGBufferRef InstanceIdsBufferRDG = GraphBuilder.CreateBuffer(FRDGBufferDesc::CreateBufferDesc(sizeof(uint32), PrimitiveIds.Num() * FInstanceCullingManager::MaxAverageInstanceFactor), TEXT("InstanceCulling.InstanceIdsBuffer"));
 	FRDGBufferRef InstanceIdOffsetBufferRDG = GraphBuilder.CreateBuffer(FRDGBufferDesc::CreateBufferDesc(sizeof(uint32), CullingCommands.Num()), TEXT("InstanceCulling.InstanceIdOffsetBuffer"));
 
 	PassParameters->ViewIds = GraphBuilder.CreateSRV(ViewIdsBuffer);
 	PassParameters->NumViewIds = ViewIds.Num();
 	PassParameters->bDrawOnlyVSMInvalidatingGeometry = bDrawOnlyVSMInvalidatingGeometry;
 
-	// TODO: Remove this when everything is properly RDG'd
-	AddPass(GraphBuilder, [&GraphBuilder](FRHICommandList& RHICmdList)
-	{
-		RHICmdList.Transition(FRHITransitionInfo(GInstanceCullingManagerResources.GetInstancesIdBufferUav(), ERHIAccess::Unknown, ERHIAccess::UAVCompute));
-	});
-
-	//PassParameters->InstanceIdsBufferOut = GraphBuilder.CreateUAV(InstanceIdsBufferRDG, PF_R32_UINT);
-	// TODO: Access resources through manager rather than global
-	PassParameters->InstanceIdsBufferLegacyOut = GInstanceCullingManagerResources.GetInstancesIdBufferUav();
+	PassParameters->InstanceIdsBufferOut = GraphBuilder.CreateUAV(InstanceIdsBuffer);
 	PassParameters->DrawIndirectArgsBufferOut = GraphBuilder.CreateUAV(DrawIndirectArgsRDG, PF_R32_UINT);
 	PassParameters->InstanceIdOffsetBufferOut = GraphBuilder.CreateUAV(InstanceIdOffsetBufferRDG, PF_R32_UINT);
 	PassParameters->NumPrimitiveIds = PrimitiveIds.Num();
@@ -548,7 +554,6 @@ void FInstanceCullingContext::BuildRenderingCommands(FRDGBuilder& GraphBuilder,
 	PassParameters->NumCulledViews = NumCulledViews;
 
 	FBuildInstanceIdBufferAndCommandsFromPrimitiveIdsCs::FPermutationDomain PermutationVector;
-	// NOTE: this also switches between legacy buffer and RDG for Id output
 	PermutationVector.Set<FBuildInstanceIdBufferAndCommandsFromPrimitiveIdsCs::FOutputCommandIdDim>(0);
 	PermutationVector.Set<FBuildInstanceIdBufferAndCommandsFromPrimitiveIdsCs::FCullInstancesDim>(bCullInstances);
 	PermutationVector.Set<FBuildInstanceIdBufferAndCommandsFromPrimitiveIdsCs::FStereoModeDim>(InstanceCullingMode == EInstanceCullingMode::Stereo);
@@ -568,18 +573,13 @@ void FInstanceCullingContext::BuildRenderingCommands(FRDGBuilder& GraphBuilder,
 #endif // ENABLE_DETERMINISTIC_INSTANCE_CULLING
 
 	Results.DrawIndirectArgsBuffer = DrawIndirectArgsRDG;
-	//ConvertToExternalBuffer(GraphBuilder, DrawIndirectArgsRDG, Results.DrawIndirectArgsBuffer);
-	//GraphBuilder.QueueBufferExtraction(InstanceIdsBufferRDG, &Results.InstanceIdsBuffer);
 	Results.InstanceIdOffsetBuffer = InstanceIdOffsetBufferRDG;
-	//ConvertToExternalBuffer(GraphBuilder, InstanceIdOffsetBufferRDG, Results.InstanceIdOffsetBuffer);
-	//GraphBuilder.Transition(FRHITransitionInfo(GInstanceCullingManagerResources.GetInstancesIdBufferUav(), ERHIAccess::Unknown, ERHIAccess::SRVGraphics));
 
-	// TODO: Remove this when everything is properly RDG'd
-	AddPass(GraphBuilder, [&GraphBuilder](FRHICommandList& RHICmdList)
-	{
-		RHICmdList.Transition(FRHITransitionInfo(GInstanceCullingManagerResources.GetInstancesIdBufferUav(), ERHIAccess::UAVCompute, ERHIAccess::SRVGraphics));
-		RHICmdList.Transition(FRHITransitionInfo(GInstanceCullingManagerResources.GetPageInfoBufferUav(), ERHIAccess::Unknown, ERHIAccess::SRVGraphics));
-	});
+	FInstanceCullingGlobalUniforms* UniformParameters = GraphBuilder.AllocParameters<FInstanceCullingGlobalUniforms>();
+	UniformParameters->InstanceIdsBuffer = GraphBuilder.CreateSRV(InstanceIdsBuffer);
+	UniformParameters->PageInfoBuffer = GraphBuilder.CreateSRV(InstanceIdsBuffer);
+	UniformParameters->BufferCapacity = InstanceIdBufferSize;
+	Results.UniformBuffer = GraphBuilder.CreateUniformBuffer(UniformParameters);
 }
 
 void FInstanceCullingContext::BuildRenderingCommands(FRDGBuilder& GraphBuilder, FGPUScene& GPUScene, const TRange<int32>& DynamicPrimitiveIdRange, FInstanceCullingRdgParams& Params) const
@@ -616,13 +616,12 @@ void FInstanceCullingContext::BuildRenderingCommands(FRDGBuilder& GraphBuilder,
 	PassParameters->InstanceRuns = GraphBuilder.CreateSRV(CreateStructuredBuffer(GraphBuilder, TEXT("InstanceCulling.InstanceRuns"), InstanceRuns));
 
 
-	FRDGBufferRef InstanceIdOutOffsetBufferRDG = Intermediate.InstanceIdOutOffsetBuffer;
 	FRDGBufferRef VisibleInstanceFlagsRDG = Intermediate.VisibleInstanceFlags;
 
 	// Create and initialize if not allocated
 	if (Params.InstanceIdWriteOffsetBuffer == nullptr)
 	{
-		TArray<uint32> NullArray;
+		TArray<uint32, TInlineAllocator<1> > NullArray;
 		NullArray.AddZeroed(1);
 		Params.InstanceIdWriteOffsetBuffer = CreateStructuredBuffer(GraphBuilder, TEXT("InstanceCulling.InstanceIdWriteOffsetBuffer"), NullArray);
 	}
@@ -631,7 +630,6 @@ void FInstanceCullingContext::BuildRenderingCommands(FRDGBuilder& GraphBuilder,
 
 	Params.DrawIndirectArgs = GraphBuilder.CreateBuffer(FRDGBufferDesc::CreateIndirectDesc(IndirectArgsNumWords * CullingCommands.Num()), TEXT("InstanceCulling.DrawIndirectArgsBuffer"));
 	// not using structured buffer as we want/have to get at it as a vertex buffer 
-	//FRDGBufferRef InstanceIdsBufferRDG = GraphBuilder.CreateBuffer(FRDGBufferDesc::CreateBufferDesc(sizeof(uint32), PrimitiveIds.Num() * FInstanceCullingManager::MaxAverageInstanceFactor), TEXT("InstanceIdsBuffer"));
 	Params.InstanceIdStartOffsetBuffer = GraphBuilder.CreateBuffer(FRDGBufferDesc::CreateBufferDesc(sizeof(uint32), CullingCommands.Num()), TEXT("InstanceCulling.InstanceIdOffsetBuffer"));
 
 	PassParameters->ViewIds = GraphBuilder.CreateSRV(CreateStructuredBuffer(GraphBuilder, TEXT("InstanceCulling.ViewIds"), ViewIds));
@@ -639,7 +637,6 @@ void FInstanceCullingContext::BuildRenderingCommands(FRDGBuilder& GraphBuilder,
 	PassParameters->bDrawOnlyVSMInvalidatingGeometry = 0;
 
 
-	//PassParameters->InstanceIdsBufferOut = GraphBuilder.CreateUAV(InstanceIdsBufferRDG, PF_R32_UINT);
 	// TODO: Access resources through manager rather than global
 	PassParameters->DrawIndirectArgsBufferOut = GraphBuilder.CreateUAV(Params.DrawIndirectArgs, PF_R32_UINT);
 	PassParameters->InstanceIdOffsetBufferOut = GraphBuilder.CreateUAV(Params.InstanceIdStartOffsetBuffer, PF_R32_UINT);
@@ -651,14 +648,14 @@ void FInstanceCullingContext::BuildRenderingCommands(FRDGBuilder& GraphBuilder,
 	if (Params.InstanceIdsBuffer == nullptr)
 	{
 		// TODO: we could compute the max instance count from the MDCs.
-		const int32 InstanceIdBufferSize = CullingCommands.Num() * FInstanceCullingManager::MaxAverageInstanceFactor * 64;
-		Params.InstanceIdsBuffer = GraphBuilder.CreateBuffer(FRDGBufferDesc::CreateBufferDesc(sizeof(uint32), InstanceIdBufferSize), TEXT("InstanceCulling.InstanceIdsBuffer"));
-		Params.DrawCommandIdsBuffer = GraphBuilder.CreateBuffer(FRDGBufferDesc::CreateBufferDesc(sizeof(uint32), InstanceIdBufferSize), TEXT("InstanceCulling.DrawCommandIdsBuffer"));
+		const int32 InstanceIdBufferSize = TotalInstances * ViewIds.Num();
+		Params.InstanceIdsBuffer = GraphBuilder.CreateBuffer(FRDGBufferDesc::CreateStructuredDesc(sizeof(uint32), InstanceIdBufferSize), TEXT("InstanceCulling.InstanceIdsBuffer"));
+		Params.DrawCommandIdsBuffer = GraphBuilder.CreateBuffer(FRDGBufferDesc::CreateStructuredDesc(sizeof(uint32), InstanceIdBufferSize), TEXT("InstanceCulling.DrawCommandIdsBuffer"));
 
 	}
 	
-	PassParameters->InstanceIdsBufferOut = GraphBuilder.CreateUAV(Params.InstanceIdsBuffer, PF_R32_UINT);
-	PassParameters->DrawCommandIdsBufferOut = GraphBuilder.CreateUAV(Params.DrawCommandIdsBuffer, PF_R32_UINT);
+	PassParameters->InstanceIdsBufferOut = GraphBuilder.CreateUAV(Params.InstanceIdsBuffer);
+	PassParameters->DrawCommandIdsBufferOut = GraphBuilder.CreateUAV(Params.DrawCommandIdsBuffer);
 
 	int32 NumInstanceFlagWords = FMath::DivideAndRoundUp(Intermediate.NumInstances, int32(sizeof(uint32) * 8));
 	PassParameters->NumInstanceFlagWords = uint32(NumInstanceFlagWords);
@@ -667,7 +664,6 @@ void FInstanceCullingContext::BuildRenderingCommands(FRDGBuilder& GraphBuilder,
 
 
 	FBuildInstanceIdBufferAndCommandsFromPrimitiveIdsCs::FPermutationDomain PermutationVector;
-	// NOTE: this also switches between legacy buffer and RDG for Id output
 	PermutationVector.Set<FBuildInstanceIdBufferAndCommandsFromPrimitiveIdsCs::FOutputCommandIdDim>(1);
 	PermutationVector.Set<FBuildInstanceIdBufferAndCommandsFromPrimitiveIdsCs::FCullInstancesDim>(false);
 	PermutationVector.Set<FBuildInstanceIdBufferAndCommandsFromPrimitiveIdsCs::FStereoModeDim>(false);
@@ -705,7 +701,7 @@ void FInstanceCullingContext::BuildRenderingCommandsBatched(FRDGBuilder& GraphBu
 	TArray<int32, SceneRenderingAllocator> ViewIdsAll;
 	TArray<FBatchInfo, SceneRenderingAllocator> BatchInfos;
 	BatchInfos.AddDefaulted(Batches.Num());
-
+	uint32 InstanceIdBufferSize = 0U;
 	// Index that maps from each command to the corresponding batch - maybe not the utmost efficiency
 	TArray<int32, SceneRenderingAllocator> BatchIndsAll;
 	for (int32 BatchIndex = 0; BatchIndex < Batches.Num(); ++BatchIndex)
@@ -748,6 +744,8 @@ void FInstanceCullingContext::BuildRenderingCommandsBatched(FRDGBuilder& GraphBu
 		{
 			BatchIndsAll[Index] = BatchIndex;
 		}
+
+		InstanceIdBufferSize += InstanceCullingContext.TotalInstances * InstanceCullingContext.ViewIds.Num();
 	}
 
 	FRDGBufferRef CullingCommandsRDG = CreateStructuredBuffer(GraphBuilder, TEXT("CullingCommandsAll"), CullingCommandsAll);
@@ -757,12 +755,14 @@ void FInstanceCullingContext::BuildRenderingCommandsBatched(FRDGBuilder& GraphBu
 	FRDGBufferRef BatchInfosRDG = CreateStructuredBuffer(GraphBuilder, TEXT("BatchInfos"), BatchInfos);
 	FRDGBufferRef BatchIndsRDG = CreateStructuredBuffer(GraphBuilder, TEXT("BatchIndsAll"), BatchIndsAll);
 
+	FRDGBufferRef InstanceIdsBuffer = GraphBuilder.CreateBuffer(FRDGBufferDesc::CreateStructuredDesc(sizeof(uint32), InstanceIdBufferSize), TEXT("InstanceCulling.InstanceIdsBuffer"));
+
 	FBuildInstanceIdBufferAndCommandsFromPrimitiveIdsCs::FParameters* PassParameters = GraphBuilder.AllocParameters<FBuildInstanceIdBufferAndCommandsFromPrimitiveIdsCs::FParameters>();
 
 	// Note: use start at zero offset if there is no instance culling manager, this means each build rendering commands pass will overwrite the same ID range. Which is only ok assuming correct barriers (should be erring on this side by default).
-	TArray<uint32> NullArray;
+	TArray<uint32, TInlineAllocator<1> > NullArray;
 	NullArray.AddZeroed(1);
-	FRDGBufferRef InstanceIdOutOffsetBufferRDG = InstanceCullingManager ? InstanceCullingManager->CullingIntermediate.InstanceIdOutOffsetBuffer : CreateStructuredBuffer(GraphBuilder, TEXT("OutputOffsetBufferOutTransient"), NullArray);
+	FRDGBufferRef InstanceIdOutOffsetBufferRDG = CreateStructuredBuffer(GraphBuilder, TEXT("InstanceCulling.OutputOffsetBufferOut"), NullArray);
 	// If there is no manager, then there is no data on culling, so set flag to skip that and ignore buffers.
 	FRDGBufferRef VisibleInstanceFlagsRDG = InstanceCullingManager != nullptr ? InstanceCullingManager->CullingIntermediate.VisibleInstanceFlags : nullptr;
 	const bool bCullInstances = InstanceCullingManager != nullptr;
@@ -794,15 +794,7 @@ void FInstanceCullingContext::BuildRenderingCommandsBatched(FRDGBuilder& GraphBu
 	PassParameters->ViewIds = GraphBuilder.CreateSRV(ViewIdsRDG);
 	PassParameters->NumViewIds = ViewIdsAll.Num();
 
-	// TODO: Remove this when everything is properly RDG'd
-	AddPass(GraphBuilder, [&GraphBuilder](FRHICommandList& RHICmdList)
-	{
-		RHICmdList.Transition(FRHITransitionInfo(GInstanceCullingManagerResources.GetInstancesIdBufferUav(), ERHIAccess::Unknown, ERHIAccess::UAVCompute));
-	});
-
-	//PassParameters->InstanceIdsBufferOut = GraphBuilder.CreateUAV(InstanceIdsBufferRDG, PF_R32_UINT);
-	// TODO: Access resources through manager rather than global
-	PassParameters->InstanceIdsBufferLegacyOut = GInstanceCullingManagerResources.GetInstancesIdBufferUav();
+	PassParameters->InstanceIdsBufferOut = GraphBuilder.CreateUAV(InstanceIdsBuffer);
 	PassParameters->DrawIndirectArgsBufferOut = GraphBuilder.CreateUAV(DrawIndirectArgsRDG, PF_R32_UINT);
 	PassParameters->InstanceIdOffsetBufferOut = GraphBuilder.CreateUAV(InstanceIdOffsetBufferRDG, PF_R32_UINT);
 	PassParameters->NumPrimitiveIds = PrimitiveIdsAll.Num();
@@ -815,7 +807,6 @@ void FInstanceCullingContext::BuildRenderingCommandsBatched(FRDGBuilder& GraphBu
 	PassParameters->NumCulledViews = NumCulledViews;
 
 	FBuildInstanceIdBufferAndCommandsFromPrimitiveIdsCs::FPermutationDomain PermutationVector;
-	// NOTE: this also switches between legacy buffer and RDG for Id output
 	PermutationVector.Set<FBuildInstanceIdBufferAndCommandsFromPrimitiveIdsCs::FOutputCommandIdDim>(0);
 	PermutationVector.Set<FBuildInstanceIdBufferAndCommandsFromPrimitiveIdsCs::FCullInstancesDim>(bCullInstances);
 	PermutationVector.Set<FBuildInstanceIdBufferAndCommandsFromPrimitiveIdsCs::FBatchedDim>(true);
@@ -829,12 +820,11 @@ void FInstanceCullingContext::BuildRenderingCommandsBatched(FRDGBuilder& GraphBu
 		PassParameters,
 		FIntVector(CullingCommandsAll.Num(), 1, 1)
 	);
-
-	// TODO: Remove this when everything is properly RDG'd
-	AddPass(GraphBuilder, [&GraphBuilder](FRHICommandList& RHICmdList)
-	{
-		RHICmdList.Transition(FRHITransitionInfo(GInstanceCullingManagerResources.GetInstancesIdBufferUav(), ERHIAccess::UAVCompute, ERHIAccess::SRVGraphics));
-	});
+	FInstanceCullingGlobalUniforms* UniformParameters = GraphBuilder.AllocParameters<FInstanceCullingGlobalUniforms>();
+	UniformParameters->InstanceIdsBuffer = GraphBuilder.CreateSRV(InstanceIdsBuffer);
+	UniformParameters->PageInfoBuffer = GraphBuilder.CreateSRV(InstanceIdsBuffer);
+	UniformParameters->BufferCapacity = InstanceIdBufferSize;
+	TRDGUniformBufferRef<FInstanceCullingGlobalUniforms> UniformBuffer = GraphBuilder.CreateUniformBuffer(UniformParameters);
 
 
 	for (int32 BatchIndex = 0; BatchIndex < Batches.Num(); ++BatchIndex)
@@ -847,6 +837,7 @@ void FInstanceCullingContext::BuildRenderingCommandsBatched(FRDGBuilder& GraphBu
 		Result.DrawIndirectArgsBuffer = DrawIndirectArgsRDG;
 		Result.InstanceIdOffsetBuffer = InstanceIdOffsetBufferRDG;
 		Result.DrawCommandDataOffset = BatchInfo.CullingCommandsOffset;
+		Result.UniformBuffer = UniformBuffer;
 	}	
 }
 
diff --git a/Engine/Source/Runtime/Renderer/Private/InstanceCulling/InstanceCullingManager.cpp b/Engine/Source/Runtime/Renderer/Private/InstanceCulling/InstanceCullingManager.cpp
--- a/Engine/Source/Runtime/Renderer/Private/InstanceCulling/InstanceCullingManager.cpp
+++ b/Engine/Source/Runtime/Renderer/Private/InstanceCulling/InstanceCullingManager.cpp
@@ -103,18 +103,15 @@ void FInstanceCullingManager::CullInstances(FRDGBuilder& GraphBuilder, FGPUScene
 	int32 NumInstances = GPUScene.InstanceDataAllocator.GetMaxSize();
 	RDG_EVENT_SCOPE(GraphBuilder, "CullInstances [%d Views X %d Instances]", NumViews, NumInstances);
 
-	check(!CullingIntermediate.InstanceIdOutOffsetBuffer);
 	check(!CullingIntermediate.VisibleInstanceFlags);
 
-	TArray<uint32> NullArray;
-	NullArray.AddZeroed(1);
-	CullingIntermediate.InstanceIdOutOffsetBuffer = CreateStructuredBuffer(GraphBuilder, TEXT("InstanceCulling.OutputOffsetBufferOut"), NullArray);
-
 	int32 NumInstanceFlagWords = FMath::DivideAndRoundUp(NumInstances, int32(sizeof(uint32) * 8));
 
 	CullingIntermediate.NumInstances = NumInstances;
 	CullingIntermediate.NumViews = NumViews;
 
+	CullingIntermediate.DummyUniformBuffer = FInstanceCullingContext::CreateDummyInstanceCullingUniformBuffer(GraphBuilder);
+
 	if (NumInstances && NumViews)
 	{
 		// Create a buffer to record one bit for each instance per view
diff --git a/Engine/Source/Runtime/Renderer/Private/InstanceCulling/InstanceCullingManager.h b/Engine/Source/Runtime/Renderer/Private/InstanceCulling/InstanceCullingManager.h
--- a/Engine/Source/Runtime/Renderer/Private/InstanceCulling/InstanceCullingManager.h
+++ b/Engine/Source/Runtime/Renderer/Private/InstanceCulling/InstanceCullingManager.h
@@ -19,13 +19,10 @@ public:
 	 */
 	FRDGBufferRef VisibleInstanceFlags = nullptr;
 	
-	/**
-	 * Write offset used by all the instance ID expand passes to allocate space in the global instance Id buffer.
-	 * Initialized to zero by CullInstances.
-	 */
-	FRDGBufferRef InstanceIdOutOffsetBuffer = nullptr;
 	int32 NumInstances = 0;
 	int32 NumViews = 0;
+
+	TRDGUniformBufferRef<FInstanceCullingGlobalUniforms> DummyUniformBuffer;
 };
 
 struct FInstanceCullingResult
@@ -35,6 +32,7 @@ struct FInstanceCullingResult
 	FRDGBufferRef InstanceIdOffsetBuffer = nullptr;
 	// Offset (in items, not bytes or something) for both buffers to start fetching data at, used when batching multiple culling jobs in the same buffer
 	uint32 DrawCommandDataOffset = 0U;
+	TRDGUniformBufferRef<FInstanceCullingGlobalUniforms> UniformBuffer = nullptr;
 
 	//FRHIBuffer* GetDrawIndirectArgsBufferRHI() const { return DrawIndirectArgsBuffer.IsValid() ? DrawIndirectArgsBuffer->GetVertexBufferRHI() : nullptr; }
 	//FRHIBuffer* GetInstanceIdOffsetBufferRHI() const { return InstanceIdOffsetBuffer.IsValid() ? InstanceIdOffsetBuffer->GetVertexBufferRHI() : nullptr; }
@@ -44,6 +42,7 @@ struct FInstanceCullingResult
 		OutParams.DrawIndirectArgsBuffer = DrawIndirectArgsBuffer;
 		OutParams.InstanceIdOffsetBuffer = InstanceIdOffsetBuffer;
 		OutParams.DrawCommandDataOffset = DrawCommandDataOffset;
+		OutParams.InstanceCulling = UniformBuffer;
 	}
 
 	static void CondGetDrawParameters(const FInstanceCullingResult* InstanceCullingResult, FInstanceCullingDrawParams& OutParams)
@@ -57,6 +56,7 @@ struct FInstanceCullingResult
 			OutParams.DrawIndirectArgsBuffer = nullptr;
 			OutParams.InstanceIdOffsetBuffer = nullptr;
 			OutParams.DrawCommandDataOffset = 0U;
+			OutParams.InstanceCulling = nullptr;
 		}
 	}
 };
@@ -71,9 +71,6 @@ public:
 
 	bool IsEnabled() const { return bIsEnabled; }
 
-	// Max average number of instances that primitives are expanded to. GPUCULL_TODO: Not very robust
-	static constexpr uint32 MaxAverageInstanceFactor = 128;
-
 	// Register a view for culling, returns integer ID of the view.
 	int32 RegisterView(const Nanite::FPackedViewParams& Params);
 	
@@ -88,12 +85,12 @@ public:
 	 */
 	void CullInstances(FRDGBuilder& GraphBuilder, FGPUScene& GPUScene);
 
-	FInstanceCullingManager(FInstanceCullingManagerResources& InResources, bool bInIsEnabled)
-	: Resources(InResources),
-		bIsEnabled(bInIsEnabled)
+	FInstanceCullingManager(bool bInIsEnabled)
+	: bIsEnabled(bInIsEnabled)
 	{
 	}
-
+	const TRDGUniformBufferRef<FInstanceCullingGlobalUniforms> GetDummyInstanceCullingUniformBuffer() const { return CullingIntermediate.DummyUniformBuffer; }
+	
 	// Populated by CullInstances, used when performing final culling & rendering 
 	FInstanceCullingIntermediate CullingIntermediate;
 
@@ -101,7 +98,6 @@ private:
 	FInstanceCullingManager() = delete;
 	FInstanceCullingManager(FInstanceCullingManager &) = delete;
 
-	FInstanceCullingManagerResources& Resources;
 	TArray<Nanite::FPackedView> CullingViews;
 	bool bIsEnabled;
 };
diff --git a/Engine/Source/Runtime/Renderer/Private/MeshDrawCommands.cpp b/Engine/Source/Runtime/Renderer/Private/MeshDrawCommands.cpp
--- a/Engine/Source/Runtime/Renderer/Private/MeshDrawCommands.cpp
+++ b/Engine/Source/Runtime/Renderer/Private/MeshDrawCommands.cpp
@@ -597,7 +597,7 @@ void SetupGPUInstancedDraws(
 			}
 			else
 			{
-				InstanceCullingContext.AddPrimitiveToCullingCommand(VisibleMeshDrawCommand.DrawPrimitiveId);
+				InstanceCullingContext.AddPrimitiveToCullingCommand(VisibleMeshDrawCommand.DrawPrimitiveId, VisibleMeshDrawCommand.MeshDrawCommand->NumInstances);
 			}
 		}
 	}
@@ -1417,13 +1417,10 @@ void FParallelMeshDrawCommandPass::BuildRenderingCommands(FRDGBuilder& GraphBuil
 			return;
 		}
 		WaitForMeshPassSetupTask();
-		if (MaxNumDraws > 0 && TaskContext.InstanceCullingContext.HasCullingCommands())
-		{
-			// 2. Run finalize culling commands pass
-			TaskContext.InstanceCullingContext.BuildRenderingCommands(GraphBuilder, GPUScene, TaskContext.View->DynamicPrimitiveCollector.GetPrimitiveIdRange(), TaskContext.InstanceCullingResult);
-			TaskContext.InstanceCullingResult.GetDrawParameters(OutInstanceCullingDrawParams);
-			return;
-		}
+		// 2. Run finalize culling commands pass
+		TaskContext.InstanceCullingContext.BuildRenderingCommands(GraphBuilder, GPUScene, TaskContext.View->DynamicPrimitiveCollector.GetPrimitiveIdRange(), TaskContext.InstanceCullingResult);
+		TaskContext.InstanceCullingResult.GetDrawParameters(OutInstanceCullingDrawParams);
+		return;
 	}
 	OutInstanceCullingDrawParams.DrawIndirectArgsBuffer = nullptr;
 	OutInstanceCullingDrawParams.InstanceIdOffsetBuffer = nullptr;
diff --git a/Engine/Source/Runtime/Renderer/Private/MobileShadingRenderer.cpp b/Engine/Source/Runtime/Renderer/Private/MobileShadingRenderer.cpp
--- a/Engine/Source/Runtime/Renderer/Private/MobileShadingRenderer.cpp
+++ b/Engine/Source/Runtime/Renderer/Private/MobileShadingRenderer.cpp
@@ -694,7 +694,7 @@ void FMobileSceneRenderer::Render(FRDGBuilder& GraphBuilder)
 	GSystemTextures.InitializeTextures(GraphBuilder.RHICmdList, FeatureLevel);
 	FRDGSystemTextures::Create(GraphBuilder);
 
-	FInstanceCullingManager InstanceCullingManager(GInstanceCullingManagerResources, Scene->GPUScene.IsEnabled());
+	FInstanceCullingManager InstanceCullingManager(Scene->GPUScene.IsEnabled());
 
 	// Important that this uses consistent logic throughout the frame, so evaluate once and pass in the flag from here
 	// NOTE: Must be done after  system texture initialization
diff --git a/Engine/Source/Runtime/Renderer/Private/PostProcess/PostProcessCompositeEditorPrimitives.cpp b/Engine/Source/Runtime/Renderer/Private/PostProcess/PostProcessCompositeEditorPrimitives.cpp
--- a/Engine/Source/Runtime/Renderer/Private/PostProcess/PostProcessCompositeEditorPrimitives.cpp
+++ b/Engine/Source/Runtime/Renderer/Private/PostProcess/PostProcessCompositeEditorPrimitives.cpp
@@ -277,6 +277,7 @@ BEGIN_SHADER_PARAMETER_STRUCT(FEditorPrimitivesPassParameters, )
 	SHADER_PARAMETER_STRUCT_REF(FReflectionCaptureShaderData, ReflectionCapture)
 	SHADER_PARAMETER_RDG_UNIFORM_BUFFER(FOpaqueBasePassUniformParameters, BasePass)
 	SHADER_PARAMETER_RDG_UNIFORM_BUFFER(FMobileBasePassUniformParameters, MobileBasePass)
+	SHADER_PARAMETER_RDG_UNIFORM_BUFFER(FInstanceCullingGlobalUniforms, InstanceCulling)
 	RENDER_TARGET_BINDING_SLOTS()
 END_SHADER_PARAMETER_STRUCT()
 
@@ -472,6 +473,7 @@ FScreenPassTexture AddEditorPrimitivePass(
 			FEditorPrimitivesPassParameters* PassParameters = GraphBuilder.AllocParameters<FEditorPrimitivesPassParameters>();
 			PassParameters->View = EditorView->ViewUniformBuffer;
 			PassParameters->ReflectionCapture = View.ReflectionCaptureUniformBuffer;
+			PassParameters->InstanceCulling = InstanceCullingManager.GetDummyInstanceCullingUniformBuffer();
 			PassParameters->RenderTargets[0] = FRenderTargetBinding(EditorPrimitiveColor, ERenderTargetLoadAction::ELoad);
 			PassParameters->RenderTargets.DepthStencil = FDepthStencilBinding(EditorPrimitiveDepth, ERenderTargetLoadAction::ELoad, ERenderTargetLoadAction::ELoad, FExclusiveDepthStencil::DepthWrite_StencilWrite);
 
diff --git a/Engine/Source/Runtime/Renderer/Private/Renderer.cpp b/Engine/Source/Runtime/Renderer/Private/Renderer.cpp
--- a/Engine/Source/Runtime/Renderer/Private/Renderer.cpp
+++ b/Engine/Source/Runtime/Renderer/Private/Renderer.cpp
@@ -103,7 +103,7 @@ void FRendererModule::InitializeSystemTextures(FRHICommandListImmediate& RHICmdL
 
 BEGIN_SHADER_PARAMETER_STRUCT(FDrawTileMeshPassParameters, )
 	SHADER_PARAMETER_STRUCT_REF(FViewUniformShaderParameters, View)
-	SHADER_PARAMETER_STRUCT_INCLUDE(FInstanceCullingDrawParams, InstanceCullingDrawParams)
+	SHADER_PARAMETER_RDG_UNIFORM_BUFFER(FInstanceCullingGlobalUniforms, InstanceCulling)
 	SHADER_PARAMETER_STRUCT_REF(FReflectionCaptureShaderData, ReflectionCapture)
 	SHADER_PARAMETER_RDG_UNIFORM_BUFFER(FDebugViewModePassUniformParameters, DebugViewMode)
 	SHADER_PARAMETER_RDG_UNIFORM_BUFFER(FTranslucentBasePassUniformParameters, TranslucentBasePass)
@@ -237,6 +237,7 @@ void FRendererModule::DrawTileMesh(FCanvasRenderContext& RenderContext, FMeshPas
 		PassParameters->RenderTargets[0] = FRenderTargetBinding(RenderContext.GetRenderTarget(), ERenderTargetLoadAction::ELoad);
 		PassParameters->View = View.ViewUniformBuffer;
 		PassParameters->ReflectionCapture = EmptyReflectionCaptureUniformBuffer;
+		PassParameters->InstanceCulling = FInstanceCullingContext::CreateDummyInstanceCullingUniformBuffer(GraphBuilder);
 
 		// handle translucent material blend modes, not relevant in MaterialTexCoordScalesAnalysis since it outputs the scales.
 		if (ViewFamily->GetDebugViewShaderMode() == DVSM_OutputMaterialTextureScales)
diff --git a/Engine/Source/Runtime/Renderer/Private/SceneHitProxyRendering.cpp b/Engine/Source/Runtime/Renderer/Private/SceneHitProxyRendering.cpp
--- a/Engine/Source/Runtime/Renderer/Private/SceneHitProxyRendering.cpp
+++ b/Engine/Source/Runtime/Renderer/Private/SceneHitProxyRendering.cpp
@@ -551,7 +551,7 @@ void FMobileSceneRenderer::RenderHitProxies(FRDGBuilder& GraphBuilder)
 	FRDGTextureRef HitProxyDepthTexture = nullptr;
 	InitHitProxyRender(GraphBuilder, this, SceneTexturesConfig, HitProxyTexture, HitProxyDepthTexture);
 
-	FInstanceCullingManager InstanceCullingManager(GInstanceCullingManagerResources, Scene->GPUScene.IsEnabled());
+	FInstanceCullingManager InstanceCullingManager(Scene->GPUScene.IsEnabled());
 
 	// Find the visible primitives.
 	InitViews(GraphBuilder, SceneTexturesConfig, InstanceCullingManager);
@@ -591,7 +591,7 @@ void FDeferredShadingSceneRenderer::RenderHitProxies(FRDGBuilder& GraphBuilder)
 
 	const FIntPoint HitProxyTextureSize = HitProxyDepthTexture->Desc.Extent;
 
-	FInstanceCullingManager InstanceCullingManager(GInstanceCullingManagerResources, Scene->GPUScene.IsEnabled());
+	FInstanceCullingManager InstanceCullingManager(Scene->GPUScene.IsEnabled());
 
 	// Find the visible primitives.
 	{
diff --git a/Engine/Source/Runtime/Renderer/Private/VirtualShadowMaps/VirtualShadowMapArray.cpp b/Engine/Source/Runtime/Renderer/Private/VirtualShadowMaps/VirtualShadowMapArray.cpp
--- a/Engine/Source/Runtime/Renderer/Private/VirtualShadowMaps/VirtualShadowMapArray.cpp
+++ b/Engine/Source/Runtime/Renderer/Private/VirtualShadowMaps/VirtualShadowMapArray.cpp
@@ -1590,8 +1590,8 @@ public:
 		SHADER_PARAMETER(uint32, GPUSceneFrameNumber)
 
 
-		SHADER_PARAMETER_RDG_BUFFER_SRV(Buffer< uint>, InstanceIdsBuffer)
-		SHADER_PARAMETER_RDG_BUFFER_SRV(Buffer< uint>, DrawCommandIdsBuffer)
+		SHADER_PARAMETER_RDG_BUFFER_SRV(StructuredBuffer< uint>, InstanceIdsBuffer)
+		SHADER_PARAMETER_RDG_BUFFER_SRV(StructuredBuffer< uint>, DrawCommandIdsBuffer)
 		SHADER_PARAMETER_RDG_BUFFER_SRV(StructuredBuffer< uint32 >, NumInstanceIdsBuffer)
 		SHADER_PARAMETER(int32, FirstPrimaryView)
 		SHADER_PARAMETER(int32, NumPrimaryViews)
@@ -1676,8 +1676,8 @@ public:
 
 	BEGIN_SHADER_PARAMETER_STRUCT(FParameters, )
 		SHADER_PARAMETER_RDG_BUFFER_SRV(StructuredBuffer< FVisibleInstanceCmd >, VisibleInstances)
-		SHADER_PARAMETER_UAV(RWBuffer<uint>, InstanceIdsBufferLegacyOut)
-		SHADER_PARAMETER_UAV(RWBuffer<uint>, PageInfoBufferLegacyOut)
+		SHADER_PARAMETER_RDG_BUFFER_UAV(RWStructuredBuffer<uint>, InstanceIdsBufferOut)
+		SHADER_PARAMETER_RDG_BUFFER_UAV(RWStructuredBuffer<uint>, PageInfoBufferOut)
 		SHADER_PARAMETER_RDG_BUFFER_UAV(RWStructuredBuffer<uint>, TmpInstanceIdOffsetBufferOut)
 		SHADER_PARAMETER_RDG_BUFFER_SRV(StructuredBuffer<uint>, VisibleInstanceCountBuffer)
 
@@ -1819,27 +1819,17 @@ void FVirtualShadowMapArray::RenderVirtualShadowMapsHw(FRDGBuilder& GraphBuilder
 			//if (IsClipMap())
 			// AddPageRenderCommandPass(GraphBuilder, GPUScene, InstanceIdBuffer, InstanceCountBuffer, ShadowMapViews);
 			// 2.1. Pass to perform the culling & replication to each page, store 
-			const uint32 MaxNumInstancesPerPass = NumDrawCommands * FInstanceCullingManager::MaxAverageInstanceFactor * 64u;
+
+			// TODO: This is both not right, and also over conservative when running with the atomic path
+			const uint32 MaxNumInstancesPerPass = MeshCommandPass.GetInstanceCullingContext()->TotalInstances * 64u;
 			FRDGBufferRef VisibleInstancesRdg = CreateStructuredBuffer(GraphBuilder, TEXT("Shadow.Virtual.VisibleInstances"), sizeof(FVisibleInstanceCmd), MaxNumInstancesPerPass, nullptr, 0);
 
-			TArray<uint32> NullArray;
+			TArray<uint32, TInlineAllocator<1> > NullArray;
 			NullArray.AddZeroed(1);
 			FRDGBufferRef VisibleInstanceWriteOffsetRDG = CreateStructuredBuffer(GraphBuilder, TEXT("Shadow.Virtual.VisibleInstanceWriteOffset"), NullArray);
 			FRDGBufferRef OutputOffsetBufferRDG = CreateStructuredBuffer(GraphBuilder, TEXT("Shadow.Virtual.OutputOffsetBuffer"), NullArray);
 			FRDGBufferRef VirtualShadowViewsRDG = CreateStructuredBuffer(GraphBuilder, TEXT("Shadow.Virtual.VirtualShadowViews"), VirtualShadowViews);
 
-			// TODO: Remove this when everything is properly RDG'd
-			AddPass(GraphBuilder, [](FRHICommandList& RHICmdList)
-			{
-				FRHITransitionInfo 	Transitions[2] =
-				{
-					FRHITransitionInfo(GInstanceCullingManagerResources.GetInstancesIdBufferUav(), ERHIAccess::Unknown, ERHIAccess::UAVCompute),
-					FRHITransitionInfo(GInstanceCullingManagerResources.GetPageInfoBufferUav(), ERHIAccess::Unknown, ERHIAccess::UAVCompute)
-				};
-
-				RHICmdList.Transition(Transitions);
-			});
-
 			{
 				FRDGBufferRef IndirectArgs = AddIndirectArgsSetupCsPass1D(GraphBuilder, Params.InstanceIdWriteOffsetBuffer, 1, FCullPerPageDrawCommandsCs::NumThreadsPerGroup);
 
@@ -1851,8 +1841,8 @@ void FVirtualShadowMapArray::RenderVirtualShadowMapsHw(FRDGBuilder& GraphBuilder
 				PassParameters->GPUScenePrimitiveSceneData = GPUScene.PrimitiveBuffer.SRV;
 				PassParameters->InstanceDataSOAStride = GPUScene.InstanceDataSOAStride;
 				PassParameters->GPUSceneFrameNumber = GPUScene.GetSceneFrameNumber();
-				PassParameters->InstanceIdsBuffer = GraphBuilder.CreateSRV(Params.InstanceIdsBuffer, PF_R32_UINT);
-				PassParameters->DrawCommandIdsBuffer = GraphBuilder.CreateSRV(Params.DrawCommandIdsBuffer, PF_R32_UINT);
+				PassParameters->InstanceIdsBuffer = GraphBuilder.CreateSRV(Params.InstanceIdsBuffer);
+				PassParameters->DrawCommandIdsBuffer = GraphBuilder.CreateSRV(Params.DrawCommandIdsBuffer);
 				PassParameters->NumInstanceIdsBuffer = GraphBuilder.CreateSRV(Params.InstanceIdWriteOffsetBuffer);
 				PassParameters->FirstPrimaryView = 0;
 				PassParameters->NumPrimaryViews = NumPrimaryViews;
@@ -1884,8 +1874,7 @@ void FVirtualShadowMapArray::RenderVirtualShadowMapsHw(FRDGBuilder& GraphBuilder
 			{
 				FAllocateCommandInstanceOutputSpaceCs::FParameters* PassParameters = GraphBuilder.AllocParameters<FAllocateCommandInstanceOutputSpaceCs::FParameters>();
 
-				// get this in a neater way
-				FRDGBufferRef InstanceIdOutOffsetBufferRDG = MeshCommandPass.GetInstanceCullingContext()->InstanceCullingManager->CullingIntermediate.InstanceIdOutOffsetBuffer;
+				FRDGBufferRef InstanceIdOutOffsetBufferRDG = CreateStructuredBuffer(GraphBuilder, TEXT("InstanceCulling.OutputOffsetBufferOut"), NullArray);
 
 				PassParameters->NumDrawCommands = NumDrawCommands;
 				PassParameters->DrawIndirectArgsBufferOut = GraphBuilder.CreateUAV(Params.DrawIndirectArgs, PF_R32_UINT);
@@ -1906,14 +1895,17 @@ void FVirtualShadowMapArray::RenderVirtualShadowMapsHw(FRDGBuilder& GraphBuilder
 
 			}
 			// 2.3. Perform final pass to re-shuffle the instance ID's to their final resting places
+			FRDGBufferRef InstanceIdsBuffer = GraphBuilder.CreateBuffer(FRDGBufferDesc::CreateStructuredDesc(sizeof(uint32), MaxNumInstancesPerPass), TEXT("Shadow.Virtual.InstanceIdsBuffer"));
+			FRDGBufferRef PageInfoBuffer = GraphBuilder.CreateBuffer(FRDGBufferDesc::CreateStructuredDesc(sizeof(uint32), MaxNumInstancesPerPass), TEXT("Shadow.Virtual.PageInfoBuffer"));
+
 			{
 				FRDGBufferRef IndirectArgs = AddIndirectArgsSetupCsPass1D(GraphBuilder, VisibleInstanceWriteOffsetRDG, 1, FOutputCommandInstanceListsCs::NumThreadsPerGroup);
 
 				FOutputCommandInstanceListsCs::FParameters* PassParameters = GraphBuilder.AllocParameters<FOutputCommandInstanceListsCs::FParameters>();
 
 				PassParameters->VisibleInstances = GraphBuilder.CreateSRV(VisibleInstancesRdg);
-				PassParameters->PageInfoBufferLegacyOut = GInstanceCullingManagerResources.GetPageInfoBufferUav();
-				PassParameters->InstanceIdsBufferLegacyOut = GInstanceCullingManagerResources.GetInstancesIdBufferUav();
+				PassParameters->PageInfoBufferOut = GraphBuilder.CreateUAV(PageInfoBuffer);
+				PassParameters->InstanceIdsBufferOut = GraphBuilder.CreateUAV(InstanceIdsBuffer);
 				PassParameters->TmpInstanceIdOffsetBufferOut = GraphBuilder.CreateUAV(TmpInstanceIdOffsetBufferRDG);
 				PassParameters->VisibleInstanceCountBuffer = GraphBuilder.CreateSRV(VisibleInstanceWriteOffsetRDG);
 				PassParameters->IndirectArgs = IndirectArgs;
@@ -1972,23 +1964,18 @@ void FVirtualShadowMapArray::RenderVirtualShadowMapsHw(FRDGBuilder& GraphBuilder
 			FInstanceCullingResult InstanceCullingResult;
 			InstanceCullingResult.DrawIndirectArgsBuffer = Params.DrawIndirectArgs;
 			InstanceCullingResult.InstanceIdOffsetBuffer = Params.InstanceIdStartOffsetBuffer;
+
+			FInstanceCullingGlobalUniforms* InstanceCullingGlobalUniforms = GraphBuilder.AllocParameters<FInstanceCullingGlobalUniforms>();
+			InstanceCullingGlobalUniforms->InstanceIdsBuffer = GraphBuilder.CreateSRV(InstanceIdsBuffer);
+			InstanceCullingGlobalUniforms->PageInfoBuffer = GraphBuilder.CreateSRV(PageInfoBuffer);
+			InstanceCullingGlobalUniforms->BufferCapacity = MaxNumInstancesPerPass;
+			InstanceCullingResult.UniformBuffer = GraphBuilder.CreateUniformBuffer(InstanceCullingGlobalUniforms);
+
 			InstanceCullingResult.GetDrawParameters(PassParameters->InstanceCullingDrawParams);
 
 			FIntRect ViewRect;
 			ViewRect.Max = bAtomicWrites ? FVirtualShadowMap::VirtualMaxResolutionXY : GetPhysicalPoolSize();
 
-			// TODO: Remove this when everything is properly RDG'd
-			AddPass(GraphBuilder, [](FRHICommandList& RHICmdList)
-			{
-				FRHITransitionInfo 	Transitions[2] =
-				{
-					FRHITransitionInfo(GInstanceCullingManagerResources.GetInstancesIdBufferUav(), ERHIAccess::UAVCompute, ERHIAccess::SRVGraphics),
-					FRHITransitionInfo(GInstanceCullingManagerResources.GetPageInfoBufferUav(), ERHIAccess::UAVCompute, ERHIAccess::SRVGraphics)
-				};
-
-				RHICmdList.Transition(Transitions);
-			});
-
 			GraphBuilder.AddPass(
 				RDG_EVENT_NAME("RenderVirtualShadowMapsHw"),
 				PassParameters,
diff --git a/Engine/Source/Runtime/Renderer/Private/VolumetricCloudRendering.cpp b/Engine/Source/Runtime/Renderer/Private/VolumetricCloudRendering.cpp
--- a/Engine/Source/Runtime/Renderer/Private/VolumetricCloudRendering.cpp
+++ b/Engine/Source/Runtime/Renderer/Private/VolumetricCloudRendering.cpp
@@ -670,10 +670,10 @@ enum EVolumetricCloudRenderViewPsPermutations
 
 BEGIN_SHADER_PARAMETER_STRUCT(FRenderVolumetricCloudRenderViewParametersPS, )
 	SHADER_PARAMETER_RDG_UNIFORM_BUFFER(FRenderVolumetricCloudGlobalParameters, VolumetricCloudRenderViewParamsUB)
-	SHADER_PARAMETER_STRUCT_INCLUDE(FInstanceCullingDrawParams, InstanceCullingDrawParams)
 	SHADER_PARAMETER_RDG_TEXTURE(Texture2D, CloudShadowTexture0)
 	SHADER_PARAMETER_RDG_TEXTURE(Texture2D, CloudShadowTexture1)
 	SHADER_PARAMETER_STRUCT_INCLUDE(FVirtualShadowMapSamplingParameters, VirtualShadowMap)
+	SHADER_PARAMETER_RDG_UNIFORM_BUFFER(FInstanceCullingGlobalUniforms, InstanceCulling)
 	RENDER_TARGET_BINDING_SLOTS()
 END_SHADER_PARAMETER_STRUCT()
 
@@ -1096,8 +1096,8 @@ private:
 
 BEGIN_SHADER_PARAMETER_STRUCT(FVolumetricCloudShadowParametersPS, )
 	SHADER_PARAMETER_STRUCT_INCLUDE(FViewShaderParameters, View)
-	SHADER_PARAMETER_STRUCT_INCLUDE(FInstanceCullingDrawParams, InstanceCullingDrawParams)
 	SHADER_PARAMETER_RDG_UNIFORM_BUFFER(FRenderVolumetricCloudGlobalParameters, TraceVolumetricCloudParamsUB)
+	SHADER_PARAMETER_RDG_UNIFORM_BUFFER(FInstanceCullingGlobalUniforms, InstanceCulling)
 	RENDER_TARGET_BINDING_SLOTS()
 END_SHADER_PARAMETER_STRUCT()
 
@@ -1630,6 +1630,7 @@ void FSceneRenderer::InitVolumetricCloudsForViews(FRDGBuilder& GraphBuilder, boo
 						FVolumetricCloudShadowParametersPS* CloudShadowParameters = GraphBuilder.AllocParameters<FVolumetricCloudShadowParametersPS>();
 						CloudShadowParameters->View = ViewInfo.GetShaderParameters();
 						CloudShadowParameters->TraceVolumetricCloudParamsUB = TraceVolumetricCloudParamsUB;
+						CloudShadowParameters->InstanceCulling = InstanceCullingManager.GetDummyInstanceCullingUniformBuffer();
 						CloudShadowParameters->RenderTargets[0] = FRenderTargetBinding(CloudTextureTracedOutput, ERenderTargetLoadAction::ENoAction);
 
 						GraphBuilder.AddPass(
@@ -2099,6 +2100,7 @@ void FSceneRenderer::RenderVolumetricCloudsInternal(FRDGBuilder& GraphBuilder, F
 		RenderViewPassParameters->CloudShadowTexture0 = CloudRC.VolumetricCloudShadowTexture[0];
 		RenderViewPassParameters->CloudShadowTexture1 = CloudRC.VolumetricCloudShadowTexture[1];
 		RenderViewPassParameters->VirtualShadowMap = VirtualShadowMapArray.GetSamplingParameters(GraphBuilder);
+		RenderViewPassParameters->InstanceCulling = InstanceCullingManager.GetDummyInstanceCullingUniformBuffer();
 		RenderViewPassParameters->RenderTargets = CloudRC.RenderTargets;
 		
 		GraphBuilder.AddPass(
diff --git a/Engine/Source/Runtime/Renderer/Private/VolumetricFogVoxelization.cpp b/Engine/Source/Runtime/Renderer/Private/VolumetricFogVoxelization.cpp
--- a/Engine/Source/Runtime/Renderer/Private/VolumetricFogVoxelization.cpp
+++ b/Engine/Source/Runtime/Renderer/Private/VolumetricFogVoxelization.cpp
@@ -547,6 +547,7 @@ void VoxelizeVolumePrimitive(FVoxelizeVolumeMeshProcessor& PassMeshProcessor,
 BEGIN_SHADER_PARAMETER_STRUCT(FVoxelizeVolumePassParameters, )
 	SHADER_PARAMETER_STRUCT_REF(FViewUniformShaderParameters, View)
 	SHADER_PARAMETER_RDG_UNIFORM_BUFFER(FVoxelizeVolumePassUniformParameters, Pass)
+	SHADER_PARAMETER_RDG_UNIFORM_BUFFER(FInstanceCullingGlobalUniforms, InstanceCulling)
 	RENDER_TARGET_BINDING_SLOTS()
 END_SHADER_PARAMETER_STRUCT()
 
@@ -567,6 +568,7 @@ void FDeferredShadingSceneRenderer::VoxelizeFogVolumePrimitives(
 
 		auto* PassParameters = GraphBuilder.AllocParameters<FVoxelizeVolumePassParameters>();
 		PassParameters->Pass = CreateVoxelizeVolumePassUniformBuffer(GraphBuilder, View, IntegrationData, Jitter, Scene->GetVolumetricCloudSceneInfo());
+		PassParameters->InstanceCulling = FInstanceCullingContext::CreateDummyInstanceCullingUniformBuffer(GraphBuilder);
 		PassParameters->RenderTargets[0] = FRenderTargetBinding(IntegrationData.VBufferA, ERenderTargetLoadAction::ELoad);
 		if (bVoxelizeEmissive)
 		{
diff --git a/Engine/Source/Runtime/Renderer/Public/InstanceCulling/InstanceCullingContext.h b/Engine/Source/Runtime/Renderer/Public/InstanceCulling/InstanceCullingContext.h
--- a/Engine/Source/Runtime/Renderer/Public/InstanceCulling/InstanceCullingContext.h
+++ b/Engine/Source/Runtime/Renderer/Public/InstanceCulling/InstanceCullingContext.h
@@ -51,10 +51,18 @@ struct FInstanceCullingRdgParams
 	FRDGBufferRef PrimitiveCullingCommands = nullptr;
 };
 
+
+BEGIN_GLOBAL_SHADER_PARAMETER_STRUCT(FInstanceCullingGlobalUniforms, RENDERER_API)
+	SHADER_PARAMETER_RDG_BUFFER_SRV(StructuredBuffer<uint>, InstanceIdsBuffer)
+	SHADER_PARAMETER_RDG_BUFFER_SRV(StructuredBuffer<uint>, PageInfoBuffer)
+	SHADER_PARAMETER(uint32, BufferCapacity)
+END_GLOBAL_SHADER_PARAMETER_STRUCT()
+
 BEGIN_SHADER_PARAMETER_STRUCT(FInstanceCullingDrawParams, )
 	RDG_BUFFER_ACCESS(DrawIndirectArgsBuffer, ERHIAccess::IndirectArgs)
 	RDG_BUFFER_ACCESS(InstanceIdOffsetBuffer, ERHIAccess::VertexOrIndexBuffer)
 	SHADER_PARAMETER(uint32, DrawCommandDataOffset)
+	SHADER_PARAMETER_RDG_UNIFORM_BUFFER(FInstanceCullingGlobalUniforms, InstanceCulling)	
 END_SHADER_PARAMETER_STRUCT()
 
 
@@ -76,6 +84,8 @@ public:
 
 	FInstanceCullingContext(FInstanceCullingManager* InInstanceCullingManager, TArrayView<const int32> InViewIds, EInstanceCullingMode InInstanceCullingMode = EInstanceCullingMode::Normal, bool bInDrawOnlyVSMInvalidatingGeometry = false);
 
+	static RENDERER_API const TRDGUniformBufferRef<FInstanceCullingGlobalUniforms> CreateDummyInstanceCullingUniformBuffer(FRDGBuilder& GraphBuilder);
+
 	struct FPrimCullingCommand
 	{
 		uint32 BaseVertexIndex;
@@ -102,7 +112,7 @@ public:
 	 * Command that is executed in the per-view, post-cull pass to gather up the instances belonging to this primitive.
 	 * Multiple commands may add to the same slot, ordering is not preserved.
 	 */
-	void AddPrimitiveToCullingCommand(int32 ScenePrimitiveId);
+	void AddPrimitiveToCullingCommand(int32 ScenePrimitiveId, uint32 NumInstances);
 
 	/**
 	 * Command that is executed in the per-view, post-cull pass to gather up the instances belonging to this primitive.
@@ -150,5 +160,7 @@ public:
 	bool bIsEnabled = false;
 	EInstanceCullingMode InstanceCullingMode = EInstanceCullingMode::Normal;
 	bool bDrawOnlyVSMInvalidatingGeometry = false;
+
+	uint32 TotalInstances = 0U;
 };
 
