Adding tunable noise to volume fog in two places
1. Jitter the offset of the sample position when writing to the light scattering volume texture, controlled with new cvar r.VolumetricFog.LightScatteringSampleJitterMultiplier
2. Per-pixel upsampling noise when reading from the volume texture, controlled with new cvar r.VolumetricFog.UpsampleJitterMultiplier

Occupancy is the same and overhead minimal.

Thanks to [REDACTED_PERSON]!

diff --git a/Engine/Shaders/Private/HeightFogPixelShader.usf b/Engine/Shaders/Private/HeightFogPixelShader.usf
--- a/Engine/Shaders/Private/HeightFogPixelShader.usf
+++ b/Engine/Shaders/Private/HeightFogPixelShader.usf
@@ -8,6 +8,9 @@
 #include "SceneTexturesCommon.ush"
 #include "SHCommon.ush"
 #include "HeightFogCommon.ush"
+#include "Random.ush"
+
+float UpsampleJitterMultiplier;
 
 float bOnlyOnRenderedOpaque;
 float bUseLinearDepthTexture;
@@ -51,7 +54,12 @@ void ExponentialPixelMain(
 	float3 WorldPositionRelativeToCamera = ScreenVector.xyz * SceneDepth;
 
 	float ZSlice = log2(SceneDepth * View.VolumetricFogGridZParams.x + View.VolumetricFogGridZParams.y) * View.VolumetricFogGridZParams.z * View.VolumetricFogInvGridSize.z;
-	float3 VolumeUV = float3((SVPos.xy - View.ViewRectMin.xy) * View.VolumetricFogSVPosToVolumeUV, ZSlice);
+	
+	uint3 Rand16Bits = Rand3DPCG16(int3(SVPos.xy, View.StateFrameIndexMod8));	
+	float3 Rand3D = (float3(Rand16Bits) / float(uint(0xffff))) * 2.0f - 1.0f;
+	float3 CellOffset = UpsampleJitterMultiplier * Rand3D;
+
+	float3 VolumeUV = float3(((SVPos.xy - View.ViewRectMin.xy) + CellOffset.xy) * View.VolumetricFogSVPosToVolumeUV, ZSlice);
 
 	float4 HeightFogInscatteringAndOpacity = CalculateHeightFog(WorldPositionRelativeToCamera);
 
diff --git a/Engine/Shaders/Private/VolumetricFog.usf b/Engine/Shaders/Private/VolumetricFog.usf
--- a/Engine/Shaders/Private/VolumetricFog.usf
+++ b/Engine/Shaders/Private/VolumetricFog.usf
@@ -16,6 +16,7 @@
 #include "ForwardShadowingCommon.ush"
 #include "ParticipatingMediaCommon.ush"
 #include "Lumen/LumenTranslucencyVolumeShared.ush"
+#include "Random.ush"
 
 #if VIRTUAL_SHADOW_MAP
 #include "VirtualShadowMaps/ProjectionCommon.ush"
@@ -192,6 +193,7 @@ void WriteToBoundingSphereVS(
 }
 
 float HistoryWeight;
+float LightScatteringSampleJitterMultiplier;
 float4 FrameJitterOffsets[16]; 
 uint HistoryMissSuperSampleCount;
 float PhaseG;
@@ -568,8 +570,10 @@ void LightScatteringCS(
 
 	for (uint SampleIndex = 0; SampleIndex < NumSuperSamples; SampleIndex++)
 	{
-		float3 CellOffset = FrameJitterOffsets[SampleIndex].xyz;
-		//float CellOffset = .5f;
+		uint3 Rand32Bits = Rand4DPCG32(int4(GridCoordinate.xyz, View.StateFrameIndexMod8 + 8 * SampleIndex)).xyz;
+		float3 Rand3D = (float3(Rand32Bits) / float(uint(0xffffffff))) * 2.0f - 1.0f;
+		float3 CellOffset = FrameJitterOffsets[SampleIndex].xyz + LightScatteringSampleJitterMultiplier * Rand3D;
+		//CellOffset = 0.5f;
 
 		float SceneDepth;
 		float3 WorldPosition = ComputeCellWorldPosition(GridCoordinate, CellOffset, SceneDepth);
diff --git a/Engine/Source/Runtime/Renderer/Private/FogRendering.cpp b/Engine/Source/Runtime/Renderer/Private/FogRendering.cpp
--- a/Engine/Source/Runtime/Renderer/Private/FogRendering.cpp
+++ b/Engine/Source/Runtime/Renderer/Private/FogRendering.cpp
@@ -44,6 +44,19 @@ static TAutoConsoleVariable<bool> CVarFogUseDepthBounds(
 	TEXT(" true: enabled (default)"),
 	ECVF_RenderThreadSafe);
 
+static TAutoConsoleVariable<float> CVarUpsampleJitterMultiplier(
+	TEXT("r.VolumetricFog.UpsampleJitterMultiplier"),
+		0.0f,
+	TEXT("Multiplier for random offset value used to jitter the sample position of the 3D fog volume to hide fog pixelization due to sampling from a lower resolution texture."),
+	ECVF_RenderThreadSafe | ECVF_Scalability);
+
+float GLightScatteringSampleJitterMultiplier = 0;
+FAutoConsoleVariableRef CVarLightScatteringSampleJitterMultiplier(
+	TEXT("r.VolumetricFog.LightScatteringSampleJitterMultiplier"),
+	GLightScatteringSampleJitterMultiplier,
+	TEXT("Multiplier for random offset value used to jitter each world sample position when generating the 3D fog volume. Enable/disable with r.VolumetricFog.Jitter"),
+	ECVF_RenderThreadSafe | ECVF_Scalability);
+
 IMPLEMENT_GLOBAL_SHADER_PARAMETER_STRUCT(FFogUniformParameters, "FogStruct");
 
 struct FHeightFogRenderingParameters
@@ -196,6 +209,7 @@ public:
 		LinearDepthSampler.Bind(Initializer.ParameterMap, TEXT("LinearDepthSampler"));
 		bOnlyOnRenderedOpaque.Bind(Initializer.ParameterMap, TEXT("bOnlyOnRenderedOpaque"));
 		bUseLinearDepthTexture.Bind(Initializer.ParameterMap, TEXT("bUseLinearDepthTexture"));
+		UpsampleJitterMultiplier.Bind(Initializer.ParameterMap, TEXT("UpsampleJitterMultiplier"));
 		LinearDepthTextureMinMaxUV.Bind(Initializer.ParameterMap, TEXT("LinearDepthTextureMinMaxUV"));
 	}
 
@@ -232,6 +246,8 @@ public:
 		SetShaderValue(RHICmdList, RHICmdList.GetBoundPixelShader(), bOnlyOnRenderedOpaque, View.bFogOnlyOnRenderedOpaque ? 1.0f : 0.0f);
 		SetShaderValue(RHICmdList, RHICmdList.GetBoundPixelShader(), bUseLinearDepthTexture, bUseLinearDepthTextureEnabled ? Params.LinearDepthReadScale : 0.0f);
 		SetShaderValue(RHICmdList, RHICmdList.GetBoundPixelShader(), LinearDepthTextureMinMaxUV, Params.LinearDepthMinMaxUV);
+		extern int32 GVolumetricFogGridPixelSize;
+		SetShaderValue(RHICmdList, RHICmdList.GetBoundPixelShader(), UpsampleJitterMultiplier, CVarUpsampleJitterMultiplier.GetValueOnRenderThread() * GVolumetricFogGridPixelSize);
 	}
 
 private:
@@ -241,6 +257,7 @@ private:
 	LAYOUT_FIELD(FShaderResourceParameter, LinearDepthSampler);
 	LAYOUT_FIELD(FShaderParameter, bOnlyOnRenderedOpaque)
 	LAYOUT_FIELD(FShaderParameter, bUseLinearDepthTexture);
+	LAYOUT_FIELD(FShaderParameter, UpsampleJitterMultiplier);
 	LAYOUT_FIELD(FShaderParameter, LinearDepthTextureMinMaxUV);
 };
 
diff --git a/Engine/Source/Runtime/Renderer/Private/VolumetricFogShared.h b/Engine/Source/Runtime/Renderer/Private/VolumetricFogShared.h
--- a/Engine/Source/Runtime/Renderer/Private/VolumetricFogShared.h
+++ b/Engine/Source/Runtime/Renderer/Private/VolumetricFogShared.h
@@ -53,6 +53,7 @@ public:
 		UnjitteredPrevWorldToClip.Bind(ParameterMap, TEXT("UnjitteredPrevWorldToClip")); 
 		FrameJitterOffsets.Bind(ParameterMap, TEXT("FrameJitterOffsets")); 
 		HistoryWeight.Bind(ParameterMap, TEXT("HistoryWeight")); 
+		LightScatteringSampleJitterMultiplier.Bind(ParameterMap, TEXT("LightScatteringSampleJitterMultiplier"));
 		HistoryMissSuperSampleCount.Bind(ParameterMap, TEXT("HistoryMissSuperSampleCount")); 
 	}
 
@@ -90,6 +91,10 @@ public:
 		extern float GVolumetricFogHistoryWeight;
 		SetShaderValue(RHICmdList, ShaderRHI, HistoryWeight, IntegrationData.bTemporalHistoryIsValid ? GVolumetricFogHistoryWeight : 0.0f);
 
+		extern float GLightScatteringSampleJitterMultiplier;
+		extern int32 GVolumetricFogJitter;
+		SetShaderValue(RHICmdList, ShaderRHI, LightScatteringSampleJitterMultiplier, GVolumetricFogJitter ? GLightScatteringSampleJitterMultiplier : 0);
+
 		extern int32 GVolumetricFogHistoryMissSupersampleCount;
 		SetShaderValue(RHICmdList, ShaderRHI, HistoryMissSuperSampleCount, FMath::Clamp(GVolumetricFogHistoryMissSupersampleCount, 1, 16));
 	}
@@ -122,6 +127,7 @@ private:
 	LAYOUT_FIELD(FShaderParameter, UnjitteredPrevWorldToClip)
 	LAYOUT_FIELD(FShaderParameter, FrameJitterOffsets)
 	LAYOUT_FIELD(FShaderParameter, HistoryWeight)
+	LAYOUT_FIELD(FShaderParameter, LightScatteringSampleJitterMultiplier)
 	LAYOUT_FIELD(FShaderParameter, HistoryMissSuperSampleCount)
 };
 
