Clean up of shader parameter from 16546528.

diff --git a/Engine/Source/Runtime/Renderer/Private/FogRendering.cpp b/Engine/Source/Runtime/Renderer/Private/FogRendering.cpp
--- a/Engine/Source/Runtime/Renderer/Private/FogRendering.cpp
+++ b/Engine/Source/Runtime/Renderer/Private/FogRendering.cpp
@@ -50,13 +50,6 @@ static TAutoConsoleVariable<float> CVarUpsampleJitterMultiplier(
 	TEXT("Multiplier for random offset value used to jitter the sample position of the 3D fog volume to hide fog pixelization due to sampling from a lower resolution texture."),
 	ECVF_RenderThreadSafe | ECVF_Scalability);
 
-float GLightScatteringSampleJitterMultiplier = 0;
-FAutoConsoleVariableRef CVarLightScatteringSampleJitterMultiplier(
-	TEXT("r.VolumetricFog.LightScatteringSampleJitterMultiplier"),
-	GLightScatteringSampleJitterMultiplier,
-	TEXT("Multiplier for random offset value used to jitter each world sample position when generating the 3D fog volume. Enable/disable with r.VolumetricFog.Jitter"),
-	ECVF_RenderThreadSafe | ECVF_Scalability);
-
 IMPLEMENT_GLOBAL_SHADER_PARAMETER_STRUCT(FFogUniformParameters, "FogStruct");
 
 struct FHeightFogRenderingParameters
diff --git a/Engine/Source/Runtime/Renderer/Private/VolumetricFog.cpp b/Engine/Source/Runtime/Renderer/Private/VolumetricFog.cpp
--- a/Engine/Source/Runtime/Renderer/Private/VolumetricFog.cpp
+++ b/Engine/Source/Runtime/Renderer/Private/VolumetricFog.cpp
@@ -133,6 +133,14 @@ FAutoConsoleVariableRef CVarVolumetricFogLightFunction(
 	ECVF_Scalability | ECVF_RenderThreadSafe
 );
 
+float GLightScatteringSampleJitterMultiplier = 0;
+FAutoConsoleVariableRef CVarLightScatteringSampleJitterMultiplier(
+	TEXT("r.VolumetricFog.LightScatteringSampleJitterMultiplier"),
+	GLightScatteringSampleJitterMultiplier,
+	TEXT("Multiplier for random offset value used to jitter each world sample position when generating the 3D fog volume. Enable/disable with r.VolumetricFog.Jitter"),
+	ECVF_RenderThreadSafe | ECVF_Scalability
+);
+
 IMPLEMENT_GLOBAL_SHADER_PARAMETER_STRUCT(FVolumetricFogGlobalData, "VolumetricFog");
 
 DECLARE_GPU_STAT(VolumetricFog);
@@ -857,6 +865,7 @@ public:
 		UseDirectionalLightShadowing.Bind(Initializer.ParameterMap, TEXT("UseDirectionalLightShadowing"));
 		AOParameters.Bind(Initializer.ParameterMap);
 		GlobalDistanceFieldParameters.Bind(Initializer.ParameterMap);
+		LightScatteringSampleJitterMultiplier.Bind(Initializer.ParameterMap, TEXT("LightScatteringSampleJitterMultiplier"));
 
 		CloudShadowmapTexture.Bind(Initializer.ParameterMap, TEXT("CloudShadowmapTexture"));
 		CloudShadowmapSampler.Bind(Initializer.ParameterMap, TEXT("CloudShadowmapSampler"));
@@ -992,6 +1001,8 @@ public:
 				CloudShadowmapStrength,
 				CloudShadowmap_Strength);
 		}
+
+		SetShaderValue(RHICmdList, ShaderRHI, LightScatteringSampleJitterMultiplier, GVolumetricFogJitter ? GLightScatteringSampleJitterMultiplier : 0);
 	}
 
 private:
@@ -1021,6 +1032,7 @@ private:
 	LAYOUT_FIELD(FShaderResourceParameter, PrevConservativeDepthTexture);
 	LAYOUT_FIELD(FShaderParameter, PrevConservativeDepthTextureSize);
 	LAYOUT_FIELD(FShaderParameter, UseConservativeDepthTexture);
+	LAYOUT_FIELD(FShaderParameter, LightScatteringSampleJitterMultiplier)
 };
 
 IMPLEMENT_GLOBAL_SHADER(TVolumetricFogLightScatteringCS, "/Engine/Private/VolumetricFog.usf", "LightScatteringCS", SF_Compute);
diff --git a/Engine/Source/Runtime/Renderer/Private/VolumetricFogShared.h b/Engine/Source/Runtime/Renderer/Private/VolumetricFogShared.h
--- a/Engine/Source/Runtime/Renderer/Private/VolumetricFogShared.h
+++ b/Engine/Source/Runtime/Renderer/Private/VolumetricFogShared.h
@@ -52,8 +52,7 @@ public:
 		UnjitteredClipToTranslatedWorld.Bind(ParameterMap, TEXT("UnjitteredClipToTranslatedWorld")); 
 		UnjitteredPrevWorldToClip.Bind(ParameterMap, TEXT("UnjitteredPrevWorldToClip")); 
 		FrameJitterOffsets.Bind(ParameterMap, TEXT("FrameJitterOffsets")); 
-		HistoryWeight.Bind(ParameterMap, TEXT("HistoryWeight")); 
-		LightScatteringSampleJitterMultiplier.Bind(ParameterMap, TEXT("LightScatteringSampleJitterMultiplier"));
+		HistoryWeight.Bind(ParameterMap, TEXT("HistoryWeight"));
 		HistoryMissSuperSampleCount.Bind(ParameterMap, TEXT("HistoryMissSuperSampleCount")); 
 	}
 
@@ -91,10 +90,6 @@ public:
 		extern float GVolumetricFogHistoryWeight;
 		SetShaderValue(RHICmdList, ShaderRHI, HistoryWeight, IntegrationData.bTemporalHistoryIsValid ? GVolumetricFogHistoryWeight : 0.0f);
 
-		extern float GLightScatteringSampleJitterMultiplier;
-		extern int32 GVolumetricFogJitter;
-		SetShaderValue(RHICmdList, ShaderRHI, LightScatteringSampleJitterMultiplier, GVolumetricFogJitter ? GLightScatteringSampleJitterMultiplier : 0);
-
 		extern int32 GVolumetricFogHistoryMissSupersampleCount;
 		SetShaderValue(RHICmdList, ShaderRHI, HistoryMissSuperSampleCount, FMath::Clamp(GVolumetricFogHistoryMissSupersampleCount, 1, 16));
 	}
@@ -127,7 +122,6 @@ private:
 	LAYOUT_FIELD(FShaderParameter, UnjitteredPrevWorldToClip)
 	LAYOUT_FIELD(FShaderParameter, FrameJitterOffsets)
 	LAYOUT_FIELD(FShaderParameter, HistoryWeight)
-	LAYOUT_FIELD(FShaderParameter, LightScatteringSampleJitterMultiplier)
 	LAYOUT_FIELD(FShaderParameter, HistoryMissSuperSampleCount)
 };
 
