Fix build break caused by not wrapping code using PrimitivesSelected with WITH_EDITOR

diff --git a/Engine/Source/Runtime/Renderer/Private/RendererScene.cpp b/Engine/Source/Runtime/Renderer/Private/RendererScene.cpp
--- a/Engine/Source/Runtime/Renderer/Private/RendererScene.cpp
+++ b/Engine/Source/Runtime/Renderer/Private/RendererScene.cpp
@@ -1426,6 +1426,7 @@ void FScene::UpdatePrimitiveSelectedState_RenderThread(const FPrimitiveSceneInfo
 {
 	check(IsInRenderingThread());
 
+#if WITH_EDITOR
 	if (PrimitiveSceneInfo)
 	{
 		if (PrimitiveSceneInfo->GetIndex() != INDEX_NONE)
@@ -1433,6 +1434,7 @@ void FScene::UpdatePrimitiveSelectedState_RenderThread(const FPrimitiveSceneInfo
 			PrimitivesSelected[PrimitiveSceneInfo->GetIndex()] = bIsSelected;
 		}
 	}
+#endif // WITH_EDITOR
 }
 
 void FScene::UpdatePrimitiveLightingAttachmentRoot(UPrimitiveComponent* Primitive)
