Do not extract HZBFurthest when nanite is disabled. This allows it to be allocated as a transient resource. Also reinstate the option to use half-res scene color as input to SSR. The option was accidentally lost during merge

diff --git a/Engine/Source/Runtime/Renderer/Private/DeferredShadingRenderer.cpp b/Engine/Source/Runtime/Renderer/Private/DeferredShadingRenderer.cpp
--- a/Engine/Source/Runtime/Renderer/Private/DeferredShadingRenderer.cpp
+++ b/Engine/Source/Runtime/Renderer/Private/DeferredShadingRenderer.cpp
@@ -382,10 +382,14 @@ bool FDeferredShadingSceneRenderer::RenderHzb(FRDGBuilder& GraphBuilder, FRDGTex
 				View.HZB = FurthestHZBTexture;
 
 				// Extract furthest HZB texture.
-				if (View.ViewState)
+				if (View.ViewState && IsNaniteEnabled())
 				{
 					GraphBuilder.QueueTextureExtraction(FurthestHZBTexture, &View.ViewState->PrevFrameViewInfo.HZB);
 				}
+				else
+				{
+					View.ViewState->PrevFrameViewInfo.HZB = nullptr;
+				}
 
 				// Extract closest HZB texture.
 				if (ViewPipelineState.bClosestHZB)
@@ -1740,9 +1744,14 @@ void FDeferredShadingSceneRenderer::CommitFinalPipelineState()
 	} 
 }
 
+bool FDeferredShadingSceneRenderer::IsNaniteEnabled() const
+{
+	return UseNanite(ShaderPlatform) && ViewFamily.EngineShowFlags.NaniteMeshes;
+}
+
 void FDeferredShadingSceneRenderer::Render(FRDGBuilder& GraphBuilder)
 {
-	const bool bNaniteEnabled = UseNanite(ShaderPlatform) && ViewFamily.EngineShowFlags.NaniteMeshes;
+	const bool bNaniteEnabled = IsNaniteEnabled();
 
 #if RHI_RAYTRACING
 	{
diff --git a/Engine/Source/Runtime/Renderer/Private/DeferredShadingRenderer.h b/Engine/Source/Runtime/Renderer/Private/DeferredShadingRenderer.h
--- a/Engine/Source/Runtime/Renderer/Private/DeferredShadingRenderer.h
+++ b/Engine/Source/Runtime/Renderer/Private/DeferredShadingRenderer.h
@@ -806,6 +806,8 @@ private:
 	
 	bool ShouldRenderDistanceFieldAO() const;
 
+	bool IsNaniteEnabled() const;
+
 	void CopySceneCaptureComponentToTarget(
 		FRDGBuilder& GraphBuilder,
 		TRDGUniformBufferRef<FSceneTextureUniformParameters> SceneTexturesUniformBuffer,
diff --git a/Engine/Source/Runtime/Renderer/Private/ScreenSpaceRayTracing.cpp b/Engine/Source/Runtime/Renderer/Private/ScreenSpaceRayTracing.cpp
--- a/Engine/Source/Runtime/Renderer/Private/ScreenSpaceRayTracing.cpp
+++ b/Engine/Source/Runtime/Renderer/Private/ScreenSpaceRayTracing.cpp
@@ -1039,6 +1039,10 @@ void RenderScreenSpaceReflections(
 		{
 			InputColor = GraphBuilder.RegisterExternalTexture(View.PrevViewInfo.CustomSSRInput);
 		}
+		else if (GSSRHalfResSceneColor && View.PrevViewInfo.HalfResTemporalAAHistory.IsValid())
+		{
+			InputColor = GraphBuilder.RegisterExternalTexture(View.PrevViewInfo.HalfResTemporalAAHistory);
+		}
 		else if (View.PrevViewInfo.TemporalAAHistory.IsValid())
 		{
 			InputColor = GraphBuilder.RegisterExternalTexture(View.PrevViewInfo.TemporalAAHistory.RT[0]);
