Fix hair tile data debug visualization when the view rect does not match the underlying buffer.

diff --git a/Engine/Source/Runtime/Renderer/Private/HairStrands/HairStrandsTile.cpp b/Engine/Source/Runtime/Renderer/Private/HairStrands/HairStrandsTile.cpp
--- a/Engine/Source/Runtime/Renderer/Private/HairStrands/HairStrandsTile.cpp
+++ b/Engine/Source/Runtime/Renderer/Private/HairStrands/HairStrandsTile.cpp
@@ -17,10 +17,15 @@
 ////////////////////////////////////////////////////////////////////////////////////////////////////////////////////////////////////////////////
 
 FHairStrandsTilePassVS::FParameters GetHairStrandsTileParameters(const FHairStrandsTiles& In)
+{
+	return GetHairStrandsTileParameters(In, In.Resolution);
+}
+
+FHairStrandsTilePassVS::FParameters GetHairStrandsTileParameters(const FHairStrandsTiles& In, const FIntPoint& InOutputResolution)
 {
 	FHairStrandsTilePassVS::FParameters Out;
 	Out.bRectPrimitive			= In.bRectPrimitive ? 1 : 0;
-	Out.TileOutputResolution	= In.Resolution;
+	Out.TileOutputResolution	= InOutputResolution;
 	Out.TileDataBuffer			= In.TileDataSRV;
 	Out.TileIndirectBuffer		= In.TileIndirectDrawBuffer;
 	return Out;
@@ -216,8 +221,8 @@ void AddHairStrandsDebugTilePass(
 	const FIntPoint OutputResolution = View.ViewRect.Size();
 	
 	FHairStrandsTileDebugPassPS::FParameters* PassParameters = GraphBuilder.AllocParameters<FHairStrandsTileDebugPassPS::FParameters>();
-	PassParameters->TileParameters = GetHairStrandsTileParameters(TileData);
-	PassParameters->OutputResolution = OutputResolution;
+	PassParameters->TileParameters = GetHairStrandsTileParameters(TileData, OutputResolution);
+	PassParameters->OutputResolution = OutputResolution; // Use the Original Buffer for 
 
 	TShaderMapRef<FHairStrandsTilePassVS> VertexShader(View.ShaderMap);
 	TShaderMapRef<FHairStrandsTileDebugPassPS> PixelShader(View.ShaderMap);
diff --git a/Engine/Source/Runtime/Renderer/Private/HairStrands/HairStrandsTile.h b/Engine/Source/Runtime/Renderer/Private/HairStrands/HairStrandsTile.h
--- a/Engine/Source/Runtime/Renderer/Private/HairStrands/HairStrandsTile.h
+++ b/Engine/Source/Runtime/Renderer/Private/HairStrands/HairStrandsTile.h
@@ -32,6 +32,7 @@ class FHairStrandsTilePassVS : public FGlobalShader
 };
 
 FHairStrandsTilePassVS::FParameters GetHairStrandsTileParameters(const FHairStrandsTiles& In);
+FHairStrandsTilePassVS::FParameters GetHairStrandsTileParameters(const FHairStrandsTiles& In, const FIntPoint& InOutputResolution);
 
 FHairStrandsTiles AddHairStrandsGenerateTilesPass(
 	FRDGBuilder& GraphBuilder,
