Enable hair tiling by default.

diff --git a/Engine/Source/Runtime/Renderer/Private/HairStrands/HairStrandsVisibility.cpp b/Engine/Source/Runtime/Renderer/Private/HairStrands/HairStrandsVisibility.cpp
--- a/Engine/Source/Runtime/Renderer/Private/HairStrands/HairStrandsVisibility.cpp
+++ b/Engine/Source/Runtime/Renderer/Private/HairStrands/HairStrandsVisibility.cpp
@@ -89,7 +89,7 @@ static FAutoConsoleVariableRef CVarHairStrandsVisibility_OutputEmissiveData(TEXT
 static int32 GHairStrandsDebugPPLL = 0;
 static FAutoConsoleVariableRef CVarHairStrandsDebugPPLL(TEXT("r.HairStrands.Visibility.PPLL.Debug"), GHairStrandsDebugPPLL, TEXT("Draw debug per pixel light list rendering."));
 
-static int32 GHairStrandsTile = 0;
+static int32 GHairStrandsTile = 1;
 static FAutoConsoleVariableRef CVarHairStrandsTile(TEXT("r.HairStrands.Tile"), GHairStrandsTile, TEXT("Enable tile generation & usage for hair strands."));
 
 
