Fix a crash due to incorrect conditioning

diff --git a/Engine/Source/Runtime/Renderer/Private/DeferredShadingRenderer.cpp b/Engine/Source/Runtime/Renderer/Private/DeferredShadingRenderer.cpp
--- a/Engine/Source/Runtime/Renderer/Private/DeferredShadingRenderer.cpp
+++ b/Engine/Source/Runtime/Renderer/Private/DeferredShadingRenderer.cpp
@@ -382,13 +382,16 @@ bool FDeferredShadingSceneRenderer::RenderHzb(FRDGBuilder& GraphBuilder, FRDGTex
 				View.HZB = FurthestHZBTexture;
 
 				// Extract furthest HZB texture.
-				if (View.ViewState && IsNaniteEnabled())
+				if (View.ViewState)
 				{
-					GraphBuilder.QueueTextureExtraction(FurthestHZBTexture, &View.ViewState->PrevFrameViewInfo.HZB);
-				}
-				else
-				{
-					View.ViewState->PrevFrameViewInfo.HZB = nullptr;
+					if (IsNaniteEnabled())
+					{
+						GraphBuilder.QueueTextureExtraction(FurthestHZBTexture, &View.ViewState->PrevFrameViewInfo.HZB);
+					}
+					else
+					{
+						View.ViewState->PrevFrameViewInfo.HZB = nullptr;
+					}
 				}
 
 				// Extract closest HZB texture.
