Expose current path tracer sample index to be able to retrieve path tracer progression

Avoid incrementing the path tracer sample and frame index when not accumulating samples

diff --git a/Engine/Source/Runtime/Engine/Public/SceneManagement.h b/Engine/Source/Runtime/Engine/Public/SceneManagement.h
--- a/Engine/Source/Runtime/Engine/Public/SceneManagement.h
+++ b/Engine/Source/Runtime/Engine/Public/SceneManagement.h
@@ -233,6 +233,19 @@ public:
 	 * returns the occlusion frame counter 
 	 */
 	virtual uint32 GetOcclusionFrameCounter() const = 0;
+
+#if RHI_RAYTRACING
+	/**
+	* returns the path tracer sample index
+	*/
+	virtual uint32 GetPathTracingSampleIndex() const = 0;
+
+	/**
+	* returns the path tracer sample count
+	*/
+	virtual uint32 GetPathTracingSampleCount() const = 0;
+#endif
+
 protected:
 	// Don't allow direct deletion of the view state, Destroy should be called instead.
 	virtual ~FSceneViewStateInterface() {}
diff --git a/Engine/Source/Runtime/Renderer/Private/PathTracing.cpp b/Engine/Source/Runtime/Renderer/Private/PathTracing.cpp
--- a/Engine/Source/Runtime/Renderer/Private/PathTracing.cpp
+++ b/Engine/Source/Runtime/Renderer/Private/PathTracing.cpp
@@ -1214,6 +1214,20 @@ void FSceneViewState::PathTracingInvalidate()
 	PathTracingSampleIndex = 0;
 }
 
+uint32 FSceneViewState::GetPathTracingSampleIndex() const {
+	return PathTracingSampleIndex;
+}
+
+uint32 FSceneViewState::GetPathTracingSampleCount() const {
+	FPathTracingConfig* LastConfig = PathTracingLastConfig.Get();
+	if (LastConfig)
+	{
+		return LastConfig->PathTracingData.MaxSamples;
+	}
+	return 0;
+}
+
+
 BEGIN_SHADER_PARAMETER_STRUCT(FDenoiseTextureParameters, )
 	RDG_TEXTURE_ACCESS(InputTexture, ERHIAccess::CopySrc)
 	RDG_TEXTURE_ACCESS(InputAlbedo, ERHIAccess::CopySrc)
@@ -1392,6 +1406,9 @@ void FDeferredShadingSceneRenderer::RenderPathTracing(
 		GraphBuilder.QueueTextureExtraction(AlbedoTexture  , &View.ViewState->PathTracingAlbedoRT  );
 		GraphBuilder.QueueTextureExtraction(NormalTexture  , &View.ViewState->PathTracingNormalRT  );
 
+		// Bump counters for next frame
+		++View.ViewState->PathTracingSampleIndex;
+		++View.ViewState->PathTracingFrameIndex;
 	}
 
 	FRDGTexture* DenoisedRadianceTexture = nullptr;
@@ -1482,10 +1499,6 @@ void FDeferredShadingSceneRenderer::RenderPathTracing(
 		PixelShader,
 		DisplayParameters
 	);
-
-	// Bump counters for next frame
-	++View.ViewState->PathTracingSampleIndex;
-	++View.ViewState->PathTracingFrameIndex;
 }
 
 #endif
diff --git a/Engine/Source/Runtime/Renderer/Private/ScenePrivate.h b/Engine/Source/Runtime/Renderer/Private/ScenePrivate.h
--- a/Engine/Source/Runtime/Renderer/Private/ScenePrivate.h
+++ b/Engine/Source/Runtime/Renderer/Private/ScenePrivate.h
@@ -976,6 +976,8 @@ public:
 #if RHI_RAYTRACING
 	// Invalidates cached results related to the path tracer so accumulated rendering can start over
 	void PathTracingInvalidate();
+	virtual uint32 GetPathTracingSampleIndex() const override;
+	virtual uint32 GetPathTracingSampleCount() const override;
 
 	// Reference path tracing cached results
 	TRefCountPtr<IPooledRenderTarget> PathTracingRadianceRT;
