Fixed shadow matrices type mismatch when LWC is enabled

diff --git a/Engine/Source/Runtime/Renderer/Private/ShadowRendering.h b/Engine/Source/Runtime/Renderer/Private/ShadowRendering.h
--- a/Engine/Source/Runtime/Renderer/Private/ShadowRendering.h
+++ b/Engine/Source/Runtime/Renderer/Private/ShadowRendering.h
@@ -1396,29 +1396,42 @@ public:
 				);
 		}
 
+		const int32 NumShaderMatrices = FMath::DivideAndRoundUp<int32>(ShadowViewProjectionMatrices.GetNumBytes(), sizeof(FMatrix44f));
+
 		if (ShadowInfo)
 		{
-			SetShaderValueArray<ShaderRHIParamRef, FMatrix>(
-				RHICmdList, 
+			const int32 NumUsableMatrices = FMath::Min<int32>(NumShaderMatrices, ShadowInfo->OnePassShadowViewProjectionMatrices.Num());
+			TArray<FMatrix44f, SceneRenderingAllocator> TypeCastedMatrices;
+			TypeCastedMatrices.AddUninitialized(NumUsableMatrices);
+			for (int32 i = 0; i < NumUsableMatrices; i++)
+			{
+				TypeCastedMatrices[i] = ShadowInfo->OnePassShadowViewProjectionMatrices[i];
+			}
+			if (NumUsableMatrices < NumShaderMatrices)
+			{
+				TypeCastedMatrices.AddZeroed(NumShaderMatrices - NumUsableMatrices);
+			}
+			SetShaderValueArray<ShaderRHIParamRef, FMatrix44f>(
+				RHICmdList,
 				ShaderRHI,
 				ShadowViewProjectionMatrices,
-				ShadowInfo->OnePassShadowViewProjectionMatrices.GetData(),
-				ShadowInfo->OnePassShadowViewProjectionMatrices.Num()
+				TypeCastedMatrices.GetData(),
+				NumShaderMatrices
 				);
 
 			SetShaderValue(RHICmdList, ShaderRHI,InvShadowmapResolution,1.0f / ShadowInfo->ResolutionX);
 		}
 		else
 		{
-			TArray<FMatrix, SceneRenderingAllocator> ZeroMatrices;
-			ZeroMatrices.AddZeroed(FMath::DivideAndRoundUp<int32>(ShadowViewProjectionMatrices.GetNumBytes(), sizeof(FMatrix)));
+			TArray<FMatrix44f, SceneRenderingAllocator> ZeroMatrices;
+			ZeroMatrices.AddZeroed(NumShaderMatrices);
 
-			SetShaderValueArray<ShaderRHIParamRef, FMatrix>(
+			SetShaderValueArray<ShaderRHIParamRef, FMatrix44f>(
 				RHICmdList, 
 				ShaderRHI,
 				ShadowViewProjectionMatrices,
 				ZeroMatrices.GetData(),
-				ZeroMatrices.Num()
+				NumShaderMatrices
 				);
 
 			SetShaderValue(RHICmdList, ShaderRHI,InvShadowmapResolution,0);
