Remove dynamic vertex buffer creation/updates from shadow projection
- Whole scene shadows now generate their planes in the VS directly
- Persistent projected frustum geometry VS gets unprojected and reprojected in the VS for per-object and similar cases. This could be done entirely in a VS as well, but since this VS is used for sphere and cone geometry as well it avoids an additional permutation to keep it a (static) VB for now.
- Split unrelated VS's into separate entry points and parameters

diff --git a/Engine/Shaders/Private/ShadowProjectionVertexShader.usf b/Engine/Shaders/Private/ShadowProjectionVertexShader.usf
--- a/Engine/Shaders/Private/ShadowProjectionVertexShader.usf
+++ b/Engine/Shaders/Private/ShadowProjectionVertexShader.usf
@@ -6,25 +6,45 @@
 
 #include "Common.ush"
 
-#ifndef USE_TRANSFORM
-	#define USE_TRANSFORM 1
-#endif
+void ShadowProjectionNoTransformVS(
+	in float4 InPosition : ATTRIBUTE0,
+	out float4 OutPosition : SV_POSITION)
+{
+	OutPosition = float4(InPosition.xyz, 1);
+}
+
 
-#if USE_TRANSFORM
-	float4 StencilingGeometryPosAndScale;
-#endif
+float4 StencilingGeometryPosAndScale;
+float4x4 InvReceiverInnerMatrix;
+float4 PreShadowToPreViewTranslation;
 
-void Main(
+void ShadowVolumeBoundProjectionVS(
 	in float4 InPosition : ATTRIBUTE0,
-	out float4 OutPosition : SV_POSITION
-	)
+	out float4 OutPosition : SV_POSITION)
 {
-	#if USE_TRANSFORM
-		// Transform geometry to clip space
-		float3 WorldPosition = InPosition.xyz * StencilingGeometryPosAndScale.w + StencilingGeometryPosAndScale.xyz;
-		OutPosition = mul(float4(WorldPosition,1),View.TranslatedWorldToClip);
-	#else
-		// Pass position straight through, as geometry is defined in clip space already.
-		OutPosition = float4(InPosition.xyz, 1);
-	#endif
+	float4 UnprojectedVertex = mul(float4(InPosition.xyz, 1), InvReceiverInnerMatrix);
+	float3 OutVertex = UnprojectedVertex.xyz / UnprojectedVertex.w + PreShadowToPreViewTranslation.xyz;
+
+	float3 WorldPosition = OutVertex.xyz * StencilingGeometryPosAndScale.w + StencilingGeometryPosAndScale.xyz;
+	OutPosition = mul(float4(WorldPosition, 1), View.TranslatedWorldToClip);
+}
+
+
+float4 ClipZValues;
+
+void WholeSceneDirectionalShadowStencilVS(
+	in float4 InPosition : ATTRIBUTE0,
+	in uint VertexID : SV_VertexID,
+	out float4 OutPosition : SV_POSITION)
+{
+	uint IndexInTriangle = VertexID % 3;
+	bool FirstTriangle = VertexID < 3;
+
+	OutPosition.x = IndexInTriangle == 1 ? 3.0f : -1.0f;
+	OutPosition.y = IndexInTriangle == 2 ? 3.0f : -1.0f;
+	OutPosition.z = FirstTriangle ? ClipZValues.x : ClipZValues.y;
+	OutPosition.w = 1.0f;
+
+	// Flip facing of second triangle
+	OutPosition.x = FirstTriangle ? OutPosition.x : -OutPosition.x;
 }
diff --git a/Engine/Source/Runtime/Renderer/Private/ShadowRendering.cpp b/Engine/Source/Runtime/Renderer/Private/ShadowRendering.cpp
--- a/Engine/Source/Runtime/Renderer/Private/ShadowRendering.cpp
+++ b/Engine/Source/Runtime/Renderer/Private/ShadowRendering.cpp
@@ -260,32 +260,43 @@ void GetOnePassPointShadowProjectionParameters(FRDGBuilder& GraphBuilder, const
 	FShadowVolumeBoundProjectionVS
 -----------------------------------------------------------------------------*/
 
-void FShadowVolumeBoundProjectionVS::SetParameters(FRHICommandList& RHICmdList, const FSceneView& View, const FProjectedShadowInfo* ShadowInfo)
+void FShadowVolumeBoundProjectionVS::SetParameters(
+	FRHICommandList& RHICmdList,
+	const FSceneView& View,
+	const FProjectedShadowInfo* ShadowInfo,
+	EShadowProjectionVertexShaderFlags Flags)
 {
-	if(ShadowInfo->IsWholeSceneDirectionalShadow())
-	{
-		// Calculate bounding geometry transform for whole scene directional shadow.
-		// Use a pair of pre-transformed planes for stenciling.
-		StencilingGeometryParameters.Set(RHICmdList, this, FVector4(0,0,0,1));
-	}
-	else if(ShadowInfo->IsWholeScenePointLightShadow())
+	FRHIVertexShader* ShaderRHI = RHICmdList.GetBoundVertexShader();
+
+	if(ShadowInfo->IsWholeScenePointLightShadow())
 	{
 		// Handle stenciling sphere for point light.
 		StencilingGeometryParameters.Set(RHICmdList, this, View, ShadowInfo->LightSceneInfo);
 	}
 	else
 	{
-		// Other bounding geometry types are pre-transformed.
 		StencilingGeometryParameters.Set(RHICmdList, this, FVector4(0,0,0,1));
 	}
+
+	if ((Flags & EShadowProjectionVertexShaderFlags::DrawingFrustum) != EShadowProjectionVertexShaderFlags::None)
+	{
+		const FVector PreShadowToPreView(View.ViewMatrices.GetPreViewTranslation() - ShadowInfo->PreShadowTranslation);
+		SetShaderValue(RHICmdList, ShaderRHI, InvReceiverInnerMatrix, ShadowInfo->InvReceiverInnerMatrix);
+		SetShaderValue(RHICmdList, ShaderRHI, PreShadowToPreViewTranslation, FVector4(PreShadowToPreView, 0));
+	}
+	else
+	{
+		FMatrix Identity = FMatrix::Identity;
+		SetShaderValue(RHICmdList, ShaderRHI, InvReceiverInnerMatrix, Identity);
+		SetShaderValue(RHICmdList, ShaderRHI, PreShadowToPreViewTranslation, FVector4(0, 0, 0, 0));
+	}
 }
 
-IMPLEMENT_TYPE_LAYOUT(FShadowProjectionVertexShaderInterface);
 IMPLEMENT_TYPE_LAYOUT(FShadowProjectionPixelShaderInterface);
 
-IMPLEMENT_SHADER_TYPE(,FShadowProjectionNoTransformVS,TEXT("/Engine/Private/ShadowProjectionVertexShader.usf"),TEXT("Main"),SF_Vertex);
+IMPLEMENT_SHADER_TYPE(,FShadowProjectionNoTransformVS,TEXT("/Engine/Private/ShadowProjectionVertexShader.usf"),TEXT("ShadowProjectionNoTransformVS"),SF_Vertex);
 
-IMPLEMENT_SHADER_TYPE(,FShadowVolumeBoundProjectionVS,TEXT("/Engine/Private/ShadowProjectionVertexShader.usf"),TEXT("Main"),SF_Vertex);
+IMPLEMENT_SHADER_TYPE(,FShadowVolumeBoundProjectionVS,TEXT("/Engine/Private/ShadowProjectionVertexShader.usf"),TEXT("ShadowVolumeBoundProjectionVS"),SF_Vertex);
 
 /**
  * Implementations for TShadowProjectionPS.  
@@ -403,7 +414,7 @@ static void BindShaderShaders(FRHICommandList& RHICmdList, FGraphicsPipelineStat
 	GraphicsPSOInit.BoundShaderState.PixelShaderRHI = PixelShader.GetPixelShader();
 	SetGraphicsPipelineState(RHICmdList, GraphicsPSOInit);
 
-	VertexShader->SetParameters(RHICmdList, View, ShadowInfo);
+	VertexShader->SetParameters(RHICmdList, View, ShadowInfo, EShadowProjectionVertexShaderFlags::DrawingFrustum);
 	PixelShader->SetParameters(RHICmdList, ViewIndex, View, ShadowInfo);
 
 	if (Strata::IsStrataEnabled())
@@ -709,6 +720,35 @@ FRHIBlendState* FProjectedShadowInfo::GetBlendStateForProjection(bool bProjectin
 		bMobileModulatedProjections);
 }
 
+class FFrustumVertexBuffer : public FVertexBuffer
+{
+public:
+	virtual void InitRHI() override
+	{
+		FRHIResourceCreateInfo CreateInfo(TEXT("FProjectedShadowInfoStencilFrustum"));
+		VertexBufferRHI = RHICreateVertexBuffer(sizeof(FVector4) * 8, BUF_Volatile, CreateInfo);
+		FVector4* OutFrustumVertices = reinterpret_cast<FVector4*>(RHILockBuffer(VertexBufferRHI, 0, sizeof(FVector4) * 8, RLM_WriteOnly));
+		
+		for(uint32 vZ = 0;vZ < 2;vZ++)
+		{
+			for(uint32 vY = 0;vY < 2;vY++)
+			{
+				for(uint32 vX = 0;vX < 2;vX++)
+				{
+					OutFrustumVertices[GetCubeVertexIndex(vX,vY,vZ)] = FVector4(
+						(vX ? -1.0f : 1.0f),
+						(vY ? -1.0f : 1.0f),
+						(vZ ?  1.0f : 0.0f),
+						1.0f);
+				}
+			}
+		}
+
+		RHIUnlockBuffer(VertexBufferRHI);
+	}
+};
+TGlobalResource<FFrustumVertexBuffer> GFrustumVertexBuffer;
+
 void FProjectedShadowInfo::SetupFrustumForProjection(const FViewInfo* View, TArray<FVector4, TInlineAllocator<8>>& OutFrustumVertices, bool& bOutCameraInsideShadowFrustum, FPlane* OutPlanes) const
 {
 	bOutCameraInsideShadowFrustum = true;
@@ -734,8 +774,8 @@ void FProjectedShadowInfo::SetupFrustumForProjection(const FViewInfo* View, TArr
 							(vY ? -1.0f : 1.0f),
 							(vZ ?  1.0f : 0.0f),
 							1.0f
-							)
-						);
+						)
+					);
 					const FVector ProjectedVertex = UnprojectedVertex / UnprojectedVertex.W + PreShadowToPreViewTranslation;
 					OutFrustumVertices[GetCubeVertexIndex(vX,vY,vZ)] = FVector4(ProjectedVertex, 0);
 				}
@@ -793,6 +833,17 @@ void FProjectedShadowInfo::SetupFrustumForProjection(const FViewInfo* View, TArr
 	}
 }
 
+class FWholeSceneDirectionalShadowStencilVS : public FGlobalShader
+{
+	DECLARE_GLOBAL_SHADER(FWholeSceneDirectionalShadowStencilVS);
+	SHADER_USE_PARAMETER_STRUCT(FWholeSceneDirectionalShadowStencilVS, FGlobalShader)
+
+	BEGIN_SHADER_PARAMETER_STRUCT(FParameters, )
+		SHADER_PARAMETER(FVector4, ClipZValues)
+	END_SHADER_PARAMETER_STRUCT()
+};
+IMPLEMENT_GLOBAL_SHADER(FWholeSceneDirectionalShadowStencilVS, "/Engine/Private/ShadowProjectionVertexShader.usf", "WholeSceneDirectionalShadowStencilVS", SF_Vertex);
+
 void FProjectedShadowInfo::SetupProjectionStencilMask(
 	FRHICommandListImmediate& RHICmdList, 
 	const FViewInfo* View, 
@@ -860,48 +911,25 @@ void FProjectedShadowInfo::SetupProjectionStencilMask(
 		checkSlow(bDirectionalLight);
 
 		// Draw 2 fullscreen planes, front facing one at the near subfrustum plane, and back facing one at the far.
+		FVector4 Near = View->ViewMatrices.GetProjectionMatrix().TransformFVector4(FVector4(0, 0, CascadeSettings.SplitNear));
+		FVector4 Far = View->ViewMatrices.GetProjectionMatrix().TransformFVector4(FVector4(0, 0, CascadeSettings.SplitFar));
+		float StencilNear = Near.Z / Near.W;
+		float StencilFar = Far.Z / Far.W;
 
-		// Find the projection shaders.
-		TShaderMapRef<FShadowProjectionNoTransformVS> VertexShaderNoTransform(View->ShaderMap);
-		VertexShaderNoTransform->SetParameters(RHICmdList, View->ViewUniformBuffer);
-
+		TShaderMapRef<FWholeSceneDirectionalShadowStencilVS> VertexShader(View->ShaderMap);
+			
 		GraphicsPSOInit.BoundShaderState.VertexDeclarationRHI = GetVertexDeclarationFVector4();
-		GraphicsPSOInit.BoundShaderState.VertexShaderRHI = VertexShaderNoTransform.GetVertexShader();
-		GraphicsPSOInit.PrimitiveType = PT_TriangleList;
+		GraphicsPSOInit.BoundShaderState.VertexShaderRHI = VertexShader.GetVertexShader();
 
 		SetGraphicsPipelineState(RHICmdList, GraphicsPSOInit);
 
-		FVector4 Near = View->ViewMatrices.GetProjectionMatrix().TransformFVector4(FVector4(0, 0, CascadeSettings.SplitNear));
-		FVector4 Far = View->ViewMatrices.GetProjectionMatrix().TransformFVector4(FVector4(0, 0, CascadeSettings.SplitFar));
-		float StencilNear = Near.Z / Near.W;
-		float StencilFar = Far.Z / Far.W;
+		FWholeSceneDirectionalShadowStencilVS::FParameters Parameters;
+		Parameters.ClipZValues = FVector4(StencilFar, StencilNear, 0, 0);
 
-		FRHIResourceCreateInfo CreateInfo(TEXT("FProjectedShadowInfo"));
-		FBufferRHIRef VertexBufferRHI = RHICreateVertexBuffer(sizeof(FVector4) * 12, BUF_Volatile, CreateInfo);
-		void* VoidPtr = RHILockBuffer(VertexBufferRHI, 0, sizeof(FVector4) * 12, RLM_WriteOnly);
-
-		// Generate the vertices used
-		FVector4* Vertices = (FVector4*)VoidPtr;
-
-			// Far Plane
-		Vertices[0] = FVector4( 1,  1,  StencilFar);
-		Vertices[1] = FVector4(-1,  1,  StencilFar);
-		Vertices[2] = FVector4( 1, -1,  StencilFar);
-		Vertices[3] = FVector4( 1, -1,  StencilFar);
-		Vertices[4] = FVector4(-1,  1,  StencilFar);
-		Vertices[5] = FVector4(-1, -1,  StencilFar);
-
-			// Near Plane
-		Vertices[6]  = FVector4(-1,  1, StencilNear);
-		Vertices[7]  = FVector4( 1,  1, StencilNear);
-		Vertices[8]  = FVector4(-1, -1, StencilNear);
-		Vertices[9]  = FVector4(-1, -1, StencilNear);
-		Vertices[10] = FVector4( 1,  1, StencilNear);
-		Vertices[11] = FVector4( 1, -1, StencilNear);
+		SetShaderParameters(RHICmdList, VertexShader, VertexShader.GetVertexShader(), Parameters);
 
-		RHIUnlockBuffer(VertexBufferRHI);
-		RHICmdList.SetStreamSource(0, VertexBufferRHI, 0);
-		RHICmdList.DrawPrimitive(0, (CascadeSettings.ShadowSplitIndex > 0) ? 4 : 2, 1);
+		RHICmdList.SetStreamSource(0, nullptr, 0);
+		RHICmdList.DrawPrimitive(0, (CascadeSettings.ShadowSplitIndex > 0) ? 2 : 1, 1);
 	}
 	// Not a preshadow, mask the projection to any pixels inside the frustum.
 	else
@@ -949,18 +977,12 @@ void FProjectedShadowInfo::SetupProjectionStencilMask(
 		SetGraphicsPipelineState(RHICmdList, GraphicsPSOInit);
 
 		// Set the projection vertex shader parameters
-		VertexShader->SetParameters(RHICmdList, *View, this);
+		VertexShader->SetParameters(RHICmdList, *View, this, EShadowProjectionVertexShaderFlags::DrawingFrustum);
 
-		FRHIResourceCreateInfo CreateInfo(TEXT("FProjectedShadowInfoStencilFrustum"));
-		FBufferRHIRef VertexBufferRHI = RHICreateVertexBuffer(sizeof(FVector4) * FrustumVertices.Num(), BUF_Volatile, CreateInfo);
-		void* VoidPtr = RHILockBuffer(VertexBufferRHI, 0, sizeof(FVector4) * FrustumVertices.Num(), RLM_WriteOnly);
-		FPlatformMemory::Memcpy(VoidPtr, FrustumVertices.GetData(), sizeof(FVector4) * FrustumVertices.Num());
-		RHIUnlockBuffer(VertexBufferRHI);
+		RHICmdList.SetStreamSource(0, GFrustumVertexBuffer.VertexBufferRHI, 0);
 
-		RHICmdList.SetStreamSource(0, VertexBufferRHI, 0);
 		// Draw the frustum using the stencil buffer to mask just the pixels which are inside the shadow frustum.
 		RHICmdList.DrawIndexedPrimitive(GCubeIndexBuffer.IndexBufferRHI, 0, 0, 8, 0, 12, 1);
-		VertexBufferRHI.SafeRelease();
 
 		// if rendering modulated shadows mask out subject mesh elements to prevent self shadowing.
 		if (bMobileModulatedProjections && !CVarEnableModulatedSelfShadow.GetValueOnRenderThread())
@@ -1264,16 +1286,9 @@ void FProjectedShadowInfo::RenderProjectionInternal(
 	}
 	else
 	{
-		FRHIResourceCreateInfo CreateInfo(TEXT("FProjectedShadowInfoFrustum"));
-		FBufferRHIRef VertexBufferRHI = RHICreateVertexBuffer(sizeof(FVector4) * FrustumVertices.Num(), BUF_Volatile, CreateInfo);
-		void* VoidPtr = RHILockBuffer(VertexBufferRHI, 0, sizeof(FVector4) * FrustumVertices.Num(), RLM_WriteOnly);
-		FPlatformMemory::Memcpy(VoidPtr, FrustumVertices.GetData(), sizeof(FVector4) * FrustumVertices.Num());
-		RHIUnlockBuffer(VertexBufferRHI);
-
-		RHICmdList.SetStreamSource(0, VertexBufferRHI, 0);
+		RHICmdList.SetStreamSource(0, GFrustumVertexBuffer.VertexBufferRHI, 0);
 		// Draw the frustum using the projection shader..
 		RHICmdList.DrawIndexedPrimitive(GCubeIndexBuffer.IndexBufferRHI, 0, 0, 8, 0, 12, 1);
-		VertexBufferRHI.SafeRelease();
 	}
 
 	if (!bDepthBoundsTestEnabled && bStencilTestEnabled)
@@ -1336,7 +1351,7 @@ static void SetPointLightShaderTempl(FRHICommandList& RHICmdList, FGraphicsPipel
 
 	SetGraphicsPipelineState(RHICmdList, GraphicsPSOInit);
 	
-	VertexShader->SetParameters(RHICmdList, View, ShadowInfo);
+	VertexShader->SetParameters(RHICmdList, View, ShadowInfo, EShadowProjectionVertexShaderFlags::None);
 	PixelShader->SetParameters(RHICmdList, ViewIndex, View, ShadowInfo, HairStrandsUniformBuffer);
 }
 
diff --git a/Engine/Source/Runtime/Renderer/Private/ShadowRendering.h b/Engine/Source/Runtime/Renderer/Private/ShadowRendering.h
--- a/Engine/Source/Runtime/Renderer/Private/ShadowRendering.h
+++ b/Engine/Source/Runtime/Renderer/Private/ShadowRendering.h
@@ -852,88 +852,51 @@ private:
 	friend class FShadowProjectionPS;
 };
 
-/**
-* A generic vertex shader for projecting a shadow depth buffer onto the scene.
-*/
-class FShadowProjectionVertexShaderInterface : public FGlobalShader
+enum class EShadowProjectionVertexShaderFlags
 {
-	DECLARE_TYPE_LAYOUT(FShadowProjectionVertexShaderInterface, NonVirtual);
-public:
-	FShadowProjectionVertexShaderInterface() {}
-	FShadowProjectionVertexShaderInterface(const ShaderMetaType::CompiledShaderInitializerType& Initializer)
-		: FGlobalShader(Initializer)
-	{ }
-
-	void SetParameters(FRHICommandList& RHICmdList, const FSceneView& View, const FProjectedShadowInfo* ShadowInfo) {}
-
-	
-	
+	None,
+	DrawingFrustum  = 0x01
 };
+ENUM_CLASS_FLAGS(EShadowProjectionVertexShaderFlags)
 
 /**
 * A vertex shader for projecting a shadow depth buffer onto the scene.
 */
-class FShadowVolumeBoundProjectionVS : public FShadowProjectionVertexShaderInterface
+class FShadowVolumeBoundProjectionVS : public FGlobalShader
 {
 	DECLARE_SHADER_TYPE(FShadowVolumeBoundProjectionVS,Global);
 public:
 
 	FShadowVolumeBoundProjectionVS() {}
 	FShadowVolumeBoundProjectionVS(const ShaderMetaType::CompiledShaderInitializerType& Initializer)
-		: FShadowProjectionVertexShaderInterface(Initializer) 
+		: FGlobalShader(Initializer) 
 	{
-		StencilingGeometryParameters.Bind(Initializer.ParameterMap);
-	}
+		const FShaderParameterMap& ParameterMap = Initializer.ParameterMap;
 
-	static bool ShouldCompilePermutation(const FGlobalShaderPermutationParameters& Parameters)
-	{
-		return true;
-	}
-	
-	static void ModifyCompilationEnvironment(const FGlobalShaderPermutationParameters& Parameters, FShaderCompilerEnvironment& OutEnvironment)
-	{
-		FShadowProjectionVertexShaderInterface::ModifyCompilationEnvironment(Parameters, OutEnvironment);
-		OutEnvironment.SetDefine(TEXT("USE_TRANSFORM"), (uint32)1);
+		StencilingGeometryParameters.Bind(ParameterMap);
+		InvReceiverInnerMatrix.Bind(ParameterMap, TEXT("InvReceiverInnerMatrix"));
+		PreShadowToPreViewTranslation.Bind(ParameterMap, TEXT("PreShadowToPreViewTranslation"));
 	}
 
-	void SetParameters(FRHICommandList& RHICmdList, const FSceneView& View, const FProjectedShadowInfo* ShadowInfo);
+	void SetParameters(FRHICommandList& RHICmdList, const FSceneView& View,	const FProjectedShadowInfo* ShadowInfo, EShadowProjectionVertexShaderFlags Flags);
 
 private:
 	LAYOUT_FIELD(FStencilingGeometryShaderParameters, StencilingGeometryParameters);
+	LAYOUT_FIELD(FShaderParameter, InvReceiverInnerMatrix);
+	LAYOUT_FIELD(FShaderParameter, PreShadowToPreViewTranslation);
 };
 
-class FShadowProjectionNoTransformVS : public FShadowProjectionVertexShaderInterface
+class FShadowProjectionNoTransformVS : public FGlobalShader
 {
 	DECLARE_SHADER_TYPE(FShadowProjectionNoTransformVS,Global);
 public:
 	FShadowProjectionNoTransformVS() {}
 	FShadowProjectionNoTransformVS(const ShaderMetaType::CompiledShaderInitializerType& Initializer)
-		: FShadowProjectionVertexShaderInterface(Initializer) 
-	{
-	}
-
-	/**
-	 * Add any defines required by the shader
-	 * @param OutEnvironment - shader environment to modify
-	 */
-	static void ModifyCompilationEnvironment(const FGlobalShaderPermutationParameters& Parameters, FShaderCompilerEnvironment& OutEnvironment)
+		: FGlobalShader(Initializer) 
 	{
-		FShadowProjectionVertexShaderInterface::ModifyCompilationEnvironment(Parameters, OutEnvironment);
-		OutEnvironment.SetDefine(TEXT("USE_TRANSFORM"), (uint32)0);
 	}
 
-	static bool ShouldCompilePermutation(const FGlobalShaderPermutationParameters& Parameters)
-	{
-		return true;
-	}
-
-	void SetParameters(FRHICommandList& RHICmdList, FRHIUniformBuffer* ViewUniformBuffer)
-	{
-	}
-
-	void SetParameters(FRHICommandList& RHICmdList, const FSceneView& View, const FProjectedShadowInfo*)
-	{
-	}
+	void SetParameters(FRHICommandList& RHICmdList, const FSceneView& View, const FProjectedShadowInfo* ShadowInfo, EShadowProjectionVertexShaderFlags Flags) {}
 };
 
 /**
