Fix LWC errors in ShadowRendering.cpp

diff --git a/Engine/Source/Runtime/Renderer/Private/ShadowRendering.cpp b/Engine/Source/Runtime/Renderer/Private/ShadowRendering.cpp
--- a/Engine/Source/Runtime/Renderer/Private/ShadowRendering.cpp
+++ b/Engine/Source/Runtime/Renderer/Private/ShadowRendering.cpp
@@ -286,8 +286,7 @@ void FShadowVolumeBoundProjectionVS::SetParameters(
 	}
 	else
 	{
-		FMatrix Identity = FMatrix::Identity;
-		SetShaderValue(RHICmdList, ShaderRHI, InvReceiverInnerMatrix, Identity);
+		SetShaderValue(RHICmdList, ShaderRHI, InvReceiverInnerMatrix, FMatrix44f::Identity);
 		SetShaderValue(RHICmdList, ShaderRHI, PreShadowToPreViewTranslation, FVector4(0, 0, 0, 0));
 	}
 }
@@ -1524,11 +1523,11 @@ void FProjectedShadowInfo::RenderFrustumWireframe(FPrimitiveDrawInterface* PDI)
 		SubjectPrimitiveId = DynamicSubjectPrimitives[0]->GetIndex();
 	}
 
-	const FMatrix InvShadowTransform = (bWholeSceneShadow || bPreShadow) ? TranslatedWorldToClipInnerMatrix.InverseFast() : InvReceiverInnerMatrix;
+	const FMatrix44f InvShadowTransform = (bWholeSceneShadow || bPreShadow) ? TranslatedWorldToClipInnerMatrix.InverseFast() : InvReceiverInnerMatrix;
 
 	FColor Color;
 
-	if(IsWholeSceneDirectionalShadow())
+	if (IsWholeSceneDirectionalShadow())
 	{
 		Color = FColor::White;
 		switch(CascadeSettings.ShadowSplitIndex)
diff --git a/Engine/Source/Runtime/Renderer/Private/ShadowRendering.h b/Engine/Source/Runtime/Renderer/Private/ShadowRendering.h
--- a/Engine/Source/Runtime/Renderer/Private/ShadowRendering.h
+++ b/Engine/Source/Runtime/Renderer/Private/ShadowRendering.h
@@ -266,10 +266,10 @@ public:
 	 * Matrix used for rendering the shadow depth buffer.  
 	 * Note that this does not necessarily contain all of the shadow casters with CSM, since the vertex shader flattens them onto the near plane of the projection.
 	 */
-	FMatrix TranslatedWorldToClipInnerMatrix;
-	FMatrix TranslatedWorldToClipOuterMatrix;
+	FMatrix44f TranslatedWorldToClipInnerMatrix;
+	FMatrix44f TranslatedWorldToClipOuterMatrix;
 
-	FMatrix InvReceiverInnerMatrix;
+	FMatrix44f InvReceiverInnerMatrix;
 
 	float InvMaxSubjectDepth;
 
