Fix shader draw debug lock.

diff --git a/Engine/Shaders/Private/ShaderDrawDebug.usf b/Engine/Shaders/Private/ShaderDrawDebug.usf
--- a/Engine/Shaders/Private/ShaderDrawDebug.usf
+++ b/Engine/Shaders/Private/ShaderDrawDebug.usf
@@ -6,12 +6,7 @@
 #include "ShaderDrawDebug.ush"
 
 // The shader used to render the debug element on view
-
-#if PERMUTATION_INPUT_TYPE == 0 
 StructuredBuffer<FPackedShaderDrawElement>	ShaderDrawDebugPrimitive;
-#else
-StructuredBuffer<FPackedShaderDrawElement>	LockedShaderDrawDebugPrimitive;
-#endif
 
 Texture2D<float> DepthTexture;
 int2 DepthTextureResolution;
@@ -25,11 +20,7 @@ void ShaderDrawDebugVS(
 	out float4 Position : SV_POSITION,
 	out float4 ColorAlpha : TEXCOORD0)
 {
-#if PERMUTATION_INPUT_TYPE == 0
 	FPackedShaderDrawElement PackedLine = ShaderDrawDebugPrimitive[InstanceId];
-#else
-	FPackedShaderDrawElement PackedLine = LockedShaderDrawDebugPrimitive[InstanceId];
-#endif
 	FShaderDrawElement Line = UnpackShaderElement(PackedLine);
 
 	float4 VertexPosition = float4(VertexId == 0 ? Line.Pos0 : Line.Pos1, 1.0f);
@@ -72,36 +63,4 @@ void ShaderDrawDebugClearCS(uint3 DispatchThreadId : SV_DispatchThreadID)
 	IndirectBuffer[2] = 0;
 	IndirectBuffer[3] = 0;
 }
-#endif
-
-#if GPU_DEBUG_RENDERING_COPY_CS
-
-	#if PERMUTATION_BUFFER_TYPE == 0
-	uint NumElements;
-	StructuredBuffer<FPackedShaderDrawElement>	InStructuredBuffer;
-	RWStructuredBuffer<FPackedShaderDrawElement>	OutStructuredBuffer;
-	[numthreads(1024, 1, 1)]
-	void ShaderDrawDebugCopyCS(uint3 DispatchThreadId : SV_DispatchThreadID)
-	{
-		const uint Index = DispatchThreadId.x;
-		if (Index < NumElements)
-		{
-			OutStructuredBuffer[Index] = InStructuredBuffer[Index];
-		}
-	}
-	#endif
-
-	#if PERMUTATION_BUFFER_TYPE == 1
-	Buffer<uint> InBuffer;
-	RWBuffer<uint> OutBuffer;
-	[numthreads(1, 1, 1)]
-	void ShaderDrawDebugCopyCS(uint3 DispatchThreadId : SV_DispatchThreadID)
-	{
-		OutBuffer[0] = InBuffer[0];
-		OutBuffer[1] = InBuffer[1];
-		OutBuffer[2] = InBuffer[2];
-		OutBuffer[3] = InBuffer[3];
-	}
-	#endif
-
-#endif
+#endif
\ No newline at end of file
diff --git a/Engine/Source/Runtime/Renderer/Private/GpuDebugRendering.cpp b/Engine/Source/Runtime/Renderer/Private/GpuDebugRendering.cpp
--- a/Engine/Source/Runtime/Renderer/Private/GpuDebugRendering.cpp
+++ b/Engine/Source/Runtime/Renderer/Private/GpuDebugRendering.cpp
@@ -40,10 +40,10 @@ namespace ShaderDrawDebug
 #endif
 	}
 
-	bool IsShaderDrawLocked()
+	static bool IsShaderDrawLocked()
 	{
 #if WITH_EDITOR
-		return CVarShaderDrawLock.GetValueOnAnyThread() > 0;
+		return CVarShaderDrawLock.GetValueOnRenderThread() > 0;
 #else
 		return false;
 #endif
@@ -95,8 +95,8 @@ namespace ShaderDrawDebug
 	// This needs to be allocated per view, or move into a more persistent place
 	struct FLockedData
 	{
-		FRWBufferStructured Buffer;
-		FRWBuffer IndirectBuffer;
+		TRefCountPtr<FRDGPooledBuffer> Buffer;
+		TRefCountPtr<FRDGPooledBuffer> IndirectBuffer;
 		bool bIsLocked = false;
 	};
 
@@ -129,41 +129,6 @@ namespace ShaderDrawDebug
 
 	IMPLEMENT_GLOBAL_SHADER(FShaderDrawDebugClearCS, "/Engine/Private/ShaderDrawDebug.usf", "ShaderDrawDebugClearCS", SF_Compute);
 
-
-	//////////////////////////////////////////////////////////////////////////
-
-	class FShaderDrawDebugCopyCS : public FGlobalShader
-	{
-		DECLARE_GLOBAL_SHADER(FShaderDrawDebugCopyCS);
-		SHADER_USE_PARAMETER_STRUCT(FShaderDrawDebugCopyCS, FGlobalShader);
-
-		class FBufferType : SHADER_PERMUTATION_INT("PERMUTATION_BUFFER_TYPE", 2);
-		using FPermutationDomain = TShaderPermutationDomain<FBufferType>;
-
-		BEGIN_SHADER_PARAMETER_STRUCT(FParameters, )
-			SHADER_PARAMETER(uint32, NumElements)
-			SHADER_PARAMETER_RDG_BUFFER_SRV(Buffer, InBuffer)
-			SHADER_PARAMETER_UAV(RWBuffer, OutBuffer)
-
-			SHADER_PARAMETER_RDG_BUFFER_SRV(StructuredBuffer, InStructuredBuffer)
-			SHADER_PARAMETER_UAV(RWStructuredBuffer, OutStructuredBuffer)
-		END_SHADER_PARAMETER_STRUCT()
-
-		static bool ShouldCompilePermutation(const FGlobalShaderPermutationParameters& Parameters)
-		{
-			return IsShaderDrawDebugEnabled(Parameters.Platform);
-		}
-
-		static void ModifyCompilationEnvironment(const FGlobalShaderPermutationParameters& Parameters, FShaderCompilerEnvironment& OutEnvironment)
-		{
-			FGlobalShader::ModifyCompilationEnvironment(Parameters, OutEnvironment);
-			OutEnvironment.SetDefine(TEXT("GPU_DEBUG_RENDERING"), 1);
-			OutEnvironment.SetDefine(TEXT("GPU_DEBUG_RENDERING_COPY_CS"), 1);
-		}
-	};
-
-	IMPLEMENT_GLOBAL_SHADER(FShaderDrawDebugCopyCS, "/Engine/Private/ShaderDrawDebug.usf", "ShaderDrawDebugCopyCS", SF_Compute);
-
 	//////////////////////////////////////////////////////////////////////////
 
 	class FShaderDrawDebugVS : public FGlobalShader
@@ -171,9 +136,6 @@ namespace ShaderDrawDebug
 		DECLARE_GLOBAL_SHADER(FShaderDrawDebugVS);
 		SHADER_USE_PARAMETER_STRUCT(FShaderDrawDebugVS, FGlobalShader);
 
-		class FInputType : SHADER_PERMUTATION_INT("PERMUTATION_INPUT_TYPE", 2);
-		using FPermutationDomain = TShaderPermutationDomain<FInputType>;
-
 		BEGIN_SHADER_PARAMETER_STRUCT(FParameters, )
 			SHADER_PARAMETER_STRUCT_REF(FViewUniformShaderParameters, View)
 			SHADER_PARAMETER_SRV(StructuredBuffer, LockedShaderDrawDebugPrimitive)
@@ -237,6 +199,63 @@ namespace ShaderDrawDebug
 
 	//////////////////////////////////////////////////////////////////////////
 
+	static void InternalDrawView(
+		FRDGBuilder& GraphBuilder,
+		const FViewInfo& View,
+		FRDGBufferRef DataBuffer,
+		FRDGBufferRef IndirectBuffer,
+		FRDGTextureRef OutputTexture,
+		FRDGTextureRef DepthTexture)
+	{
+		TShaderMapRef<FShaderDrawDebugVS> VertexShader(View.ShaderMap);
+		TShaderMapRef<FShaderDrawDebugPS> PixelShader(View.ShaderMap);
+
+		FShaderDrawVSPSParameters* PassParameters = GraphBuilder.AllocParameters<FShaderDrawVSPSParameters >();
+		PassParameters->ShaderDrawPSParameters.RenderTargets[0] = FRenderTargetBinding(OutputTexture, ERenderTargetLoadAction::ELoad);
+		//PassParameters->ShaderDrawPSParameters.RenderTargets.DepthStencil = FDepthStencilBinding(DepthTexture, ERenderTargetLoadAction::ELoad, FExclusiveDepthStencil::DepthRead_StencilNop);
+		PassParameters->ShaderDrawPSParameters.DepthTexture = DepthTexture;
+		PassParameters->ShaderDrawPSParameters.DepthTextureResolution = FIntPoint(DepthTexture->Desc.Extent.X, DepthTexture->Desc.Extent.Y);
+		PassParameters->ShaderDrawPSParameters.DepthTextureInvResolution = FVector2D(1.f / DepthTexture->Desc.Extent.X, 1.f / DepthTexture->Desc.Extent.Y);
+		PassParameters->ShaderDrawPSParameters.DepthSampler = TStaticSamplerState<SF_Point, AM_Clamp, AM_Clamp, AM_Clamp>::GetRHI();
+		PassParameters->ShaderDrawVSParameters.View = View.ViewUniformBuffer;
+		PassParameters->ShaderDrawVSParameters.ShaderDrawDebugPrimitive = GraphBuilder.CreateSRV(DataBuffer);
+		PassParameters->ShaderDrawVSParameters.IndirectBuffer = IndirectBuffer;
+
+		ValidateShaderParameters(PixelShader, PassParameters->ShaderDrawPSParameters);
+		ClearUnusedGraphResources(PixelShader, &PassParameters->ShaderDrawPSParameters, { IndirectBuffer });
+		ValidateShaderParameters(VertexShader, PassParameters->ShaderDrawVSParameters);
+		ClearUnusedGraphResources(VertexShader, &PassParameters->ShaderDrawVSParameters, { IndirectBuffer });
+
+		GraphBuilder.AddPass(
+			RDG_EVENT_NAME("ShaderDrawDebug"),
+			PassParameters,
+			ERDGPassFlags::Raster,
+			[VertexShader, PixelShader, PassParameters, IndirectBuffer](FRHICommandList& RHICmdList)
+			{
+				// Marks the indirect draw parameter as used by the pass, given it's not used directly by any of the shaders.
+				PassParameters->ShaderDrawVSParameters.IndirectBuffer->MarkResourceAsUsed();
+
+				FGraphicsPipelineStateInitializer GraphicsPSOInit;
+				RHICmdList.ApplyCachedRenderTargets(GraphicsPSOInit);
+				GraphicsPSOInit.DepthStencilState = TStaticDepthStencilState<false, CF_Always>::GetRHI();
+				GraphicsPSOInit.BlendState = TStaticBlendState<CW_RGBA, BO_Add, BF_One, BF_InverseSourceAlpha, BO_Add, BF_Zero, BF_One>::GetRHI(); // Premultiplied-alpha composition
+				GraphicsPSOInit.RasterizerState = TStaticRasterizerState<FM_Solid, CM_None, true>::GetRHI();
+				GraphicsPSOInit.PrimitiveType = PT_LineList;
+				GraphicsPSOInit.BoundShaderState.VertexDeclarationRHI = GEmptyVertexDeclaration.VertexDeclarationRHI;
+				GraphicsPSOInit.BoundShaderState.VertexShaderRHI = VertexShader.GetVertexShader();
+				GraphicsPSOInit.BoundShaderState.PixelShaderRHI = PixelShader.GetPixelShader();
+				SetGraphicsPipelineState(RHICmdList, GraphicsPSOInit);
+
+				SetShaderParameters(RHICmdList, VertexShader, VertexShader.GetVertexShader(), PassParameters->ShaderDrawVSParameters);
+				SetShaderParameters(RHICmdList, PixelShader, PixelShader.GetPixelShader(), PassParameters->ShaderDrawPSParameters);
+
+				// Marks the indirect draw parameter as used by the pass, given it's not used directly by any of the shaders.
+				FRHIBuffer* IndirectBufferRHI = PassParameters->ShaderDrawVSParameters.IndirectBuffer->GetIndirectRHICallBuffer();
+				check(IndirectBufferRHI != nullptr);
+				RHICmdList.DrawPrimitiveIndirect(IndirectBufferRHI, 0);
+			});
+	}
+
 	void BeginView(FRDGBuilder& GraphBuilder, FViewInfo& View)
 	{
 		if (!IsShaderDrawDebugEnabled(View) || !IsShaderDrawDebugEnabled())
@@ -244,8 +263,11 @@ namespace ShaderDrawDebug
 			return;
 		}
 
-		FRDGBufferRef DataBuffer = GraphBuilder.CreateBuffer(FRDGBufferDesc::CreateStructuredDesc(sizeof(FPackedShaderDrawElement), GetMaxElementCount()), TEXT("ShaderDrawDataBuffer"));
-		FRDGBufferRef IndirectBuffer = GraphBuilder.CreateBuffer(FRDGBufferDesc::CreateIndirectDesc<FRHIDrawIndirectParameters>(1), TEXT("ShaderDrawDataIndirectBuffer"));
+		const bool bLockBufferThisFrame = IsShaderDrawLocked() && !LockedData.bIsLocked;
+		ERDGBufferFlags Flags = bLockBufferThisFrame ? ERDGBufferFlags::MultiFrame : ERDGBufferFlags::None;
+
+		FRDGBufferRef DataBuffer = GraphBuilder.CreateBuffer(FRDGBufferDesc::CreateStructuredDesc(sizeof(FPackedShaderDrawElement), GetMaxElementCount()), TEXT("ShaderDrawDataBuffer"), Flags);
+		FRDGBufferRef IndirectBuffer = GraphBuilder.CreateBuffer(FRDGBufferDesc::CreateIndirectDesc<FRHIDrawIndirectParameters>(1), TEXT("ShaderDrawDataIndirectBuffer"), Flags);
 
 		FShaderDrawDebugClearCS::FParameters* Parameters = GraphBuilder.AllocParameters<FShaderDrawDebugClearCS::FParameters>();
 		Parameters->DataBuffer = GraphBuilder.CreateUAV(DataBuffer);
@@ -265,14 +287,21 @@ namespace ShaderDrawDebug
 
 		View.ShaderDrawData.Buffer = DataBuffer;
 		View.ShaderDrawData.IndirectBuffer = IndirectBuffer;
+		View.ShaderDrawData.CursorPosition = View.CursorPos;
 
 		if (IsShaderDrawLocked() && !LockedData.bIsLocked)
 		{
-			LockedData.Buffer.Initialize(TEXT("ShaderDrawDataBuffer"), sizeof(FPackedShaderDrawElement), GetMaxElementCount(), 0U);
-			LockedData.IndirectBuffer.Initialize(TEXT("ShaderDrawDataIndirectBuffer"), sizeof(uint32), 4, PF_R32_UINT, BUF_DrawIndirect);
+			LockedData.Buffer = GraphBuilder.ConvertToExternalBuffer(View.ShaderDrawData.Buffer);
+			LockedData.IndirectBuffer = GraphBuilder.ConvertToExternalBuffer(View.ShaderDrawData.IndirectBuffer);
+			LockedData.bIsLocked = true;
 		}
 
-		View.ShaderDrawData.CursorPosition = View.CursorPos;
+		if (!IsShaderDrawLocked() && LockedData.bIsLocked)
+		{
+			LockedData.Buffer = nullptr;
+			LockedData.IndirectBuffer = nullptr;
+			LockedData.bIsLocked = false;
+		}
 	}
 
 	void DrawView(FRDGBuilder& GraphBuilder, const FViewInfo& View, FRDGTextureRef OutputTexture, FRDGTextureRef DepthTexture)
@@ -282,119 +311,17 @@ namespace ShaderDrawDebug
 			return;
 		}
 
-		auto RunPass = [&](
-			bool bUseRdgInput,
-			FRDGBufferRef DataBuffer, 
-			FRDGBufferRef IndirectBuffer, 
-			FShaderResourceViewRHIRef LockedDataBuffer,
-			FRHIBuffer* LockedIndirectBuffer)
 		{
-			FShaderDrawDebugVS::FPermutationDomain PermutationVector;
-			PermutationVector.Set<FShaderDrawDebugVS::FInputType>(bUseRdgInput ? 0 : 1);
-			TShaderMapRef<FShaderDrawDebugVS> VertexShader(View.ShaderMap, PermutationVector);
-			TShaderMapRef<FShaderDrawDebugPS> PixelShader(View.ShaderMap);
-
-			FShaderDrawVSPSParameters * PassParameters = GraphBuilder.AllocParameters<FShaderDrawVSPSParameters >();
-			PassParameters->ShaderDrawPSParameters.RenderTargets[0] = FRenderTargetBinding(OutputTexture, ERenderTargetLoadAction::ELoad);
-			//PassParameters->ShaderDrawPSParameters.RenderTargets.DepthStencil = FDepthStencilBinding(DepthTexture, ERenderTargetLoadAction::ELoad, FExclusiveDepthStencil::DepthRead_StencilNop);
-			PassParameters->ShaderDrawPSParameters.DepthTexture = DepthTexture;
-			PassParameters->ShaderDrawPSParameters.DepthTextureResolution    = FIntPoint(DepthTexture->Desc.Extent.X, DepthTexture->Desc.Extent.Y);
-			PassParameters->ShaderDrawPSParameters.DepthTextureInvResolution = FVector2D(1.f/DepthTexture->Desc.Extent.X, 1.f / DepthTexture->Desc.Extent.Y);
-			PassParameters->ShaderDrawPSParameters.DepthSampler = TStaticSamplerState<SF_Point, AM_Clamp, AM_Clamp, AM_Clamp>::GetRHI();
-			PassParameters->ShaderDrawVSParameters.View = View.ViewUniformBuffer;
-			if (bUseRdgInput)
-			{
-				PassParameters->ShaderDrawVSParameters.ShaderDrawDebugPrimitive = GraphBuilder.CreateSRV(DataBuffer);
-				PassParameters->ShaderDrawVSParameters.IndirectBuffer = IndirectBuffer;
-			}
-			else
-			{
-				PassParameters->ShaderDrawVSParameters.LockedShaderDrawDebugPrimitive = LockedDataBuffer;
-			}
-
-			ValidateShaderParameters(PixelShader, PassParameters->ShaderDrawPSParameters);
-			ClearUnusedGraphResources(PixelShader, &PassParameters->ShaderDrawPSParameters, { IndirectBuffer });
-			ValidateShaderParameters(VertexShader, PassParameters->ShaderDrawVSParameters);
-			ClearUnusedGraphResources(VertexShader, &PassParameters->ShaderDrawVSParameters, { IndirectBuffer });
-
-			GraphBuilder.AddPass(
-				RDG_EVENT_NAME("ShaderDrawDebug"),
-				PassParameters,
-				ERDGPassFlags::Raster,
-				[VertexShader, PixelShader, PassParameters, IndirectBuffer, LockedIndirectBuffer, bUseRdgInput](FRHICommandList& RHICmdList)
-			{
-				// Marks the indirect draw parameter as used by the pass, given it's not used directly by any of the shaders.
-				if (bUseRdgInput)
-				{
-					PassParameters->ShaderDrawVSParameters.IndirectBuffer->MarkResourceAsUsed();
-				}
-
-				FGraphicsPipelineStateInitializer GraphicsPSOInit;
-				RHICmdList.ApplyCachedRenderTargets(GraphicsPSOInit);
-				GraphicsPSOInit.DepthStencilState = TStaticDepthStencilState<false, CF_Always>::GetRHI();
-				GraphicsPSOInit.BlendState = TStaticBlendState<CW_RGBA, BO_Add, BF_One, BF_InverseSourceAlpha, BO_Add, BF_Zero, BF_One>::GetRHI(); // Premultiplied-alpha composition
-				GraphicsPSOInit.RasterizerState = TStaticRasterizerState<FM_Solid, CM_None, true>::GetRHI();
-				GraphicsPSOInit.PrimitiveType = PT_LineList;
-				GraphicsPSOInit.BoundShaderState.VertexDeclarationRHI = GEmptyVertexDeclaration.VertexDeclarationRHI;
-				GraphicsPSOInit.BoundShaderState.VertexShaderRHI = VertexShader.GetVertexShader();
-				GraphicsPSOInit.BoundShaderState.PixelShaderRHI = PixelShader.GetPixelShader();
-				SetGraphicsPipelineState(RHICmdList, GraphicsPSOInit);
-
-				SetShaderParameters(RHICmdList, VertexShader, VertexShader.GetVertexShader(), PassParameters->ShaderDrawVSParameters);
-				SetShaderParameters(RHICmdList, PixelShader, PixelShader.GetPixelShader(), PassParameters->ShaderDrawPSParameters);
-
-
-				if (bUseRdgInput)
-				{
-					// Marks the indirect draw parameter as used by the pass, given it's not used directly by any of the shaders.
-					FRHIBuffer* IndirectBufferRHI = PassParameters->ShaderDrawVSParameters.IndirectBuffer->GetIndirectRHICallBuffer();
-					check(IndirectBufferRHI != nullptr);
-					RHICmdList.DrawPrimitiveIndirect(IndirectBufferRHI, 0);
-				}
-				else
-				{
-					RHICmdList.DrawPrimitiveIndirect(LockedIndirectBuffer, 0);
-				}
-			});
-		};
-
-		FRDGBufferRef DataBuffer = View.ShaderDrawData.Buffer;
-		FRDGBufferRef IndirectBuffer = View.ShaderDrawData.IndirectBuffer;
-		{
-			RunPass(true, DataBuffer, IndirectBuffer, nullptr, nullptr);
+			FRDGBufferRef DataBuffer = View.ShaderDrawData.Buffer;
+			FRDGBufferRef IndirectBuffer = View.ShaderDrawData.IndirectBuffer;
+			InternalDrawView(GraphBuilder, View, DataBuffer, IndirectBuffer, OutputTexture, DepthTexture);
 		}
 
-		if (LockedData.bIsLocked)
+		if (IsShaderDrawLocked())
 		{
-			FShaderResourceViewRHIRef LockedDataBuffer = LockedData.Buffer.SRV;
-			FRHIBuffer* LockedIndirectBuffer = LockedData.IndirectBuffer.Buffer;
-			RunPass(false, nullptr, nullptr, LockedDataBuffer, LockedIndirectBuffer);
-		}
-
-		if (IsShaderDrawLocked() && !LockedData.bIsLocked)
-		{
-			{
-				const uint32 NumElements = LockedData.Buffer.NumBytes / sizeof(FPackedShaderDrawElement);
-
-				FShaderDrawDebugCopyCS::FPermutationDomain PermutationVector;
-				PermutationVector.Set<FShaderDrawDebugCopyCS::FBufferType>(0);
-				TShaderMapRef<FShaderDrawDebugCopyCS> ComputeShader(View.ShaderMap, PermutationVector);
-				FShaderDrawDebugCopyCS::FParameters* Parameters = GraphBuilder.AllocParameters<FShaderDrawDebugCopyCS::FParameters>();
-				Parameters->NumElements = NumElements;
-				Parameters->InStructuredBuffer = GraphBuilder.CreateSRV(DataBuffer);
-				Parameters->OutStructuredBuffer = LockedData.Buffer.UAV;
-				FComputeShaderUtils::AddPass(GraphBuilder, RDG_EVENT_NAME("ShaderDrawDebugCopy"), ComputeShader, Parameters, FIntVector(FMath::CeilToInt(NumElements / 1024.f), 1, 1));
-			}
-			{
-				FShaderDrawDebugCopyCS::FPermutationDomain PermutationVector;
-				PermutationVector.Set<FShaderDrawDebugCopyCS::FBufferType>(1);
-				TShaderMapRef<FShaderDrawDebugCopyCS> ComputeShader(View.ShaderMap, PermutationVector);
-
-				FShaderDrawDebugCopyCS::FParameters* Parameters = GraphBuilder.AllocParameters<FShaderDrawDebugCopyCS::FParameters>();
-				Parameters->InBuffer = GraphBuilder.CreateSRV(IndirectBuffer);
-				Parameters->OutBuffer = LockedData.IndirectBuffer.UAV;
-				FComputeShaderUtils::AddPass(GraphBuilder, RDG_EVENT_NAME("ShaderDrawDebugCopy"), ComputeShader, Parameters, FIntVector(1, 1, 1));
-			}
+			FRDGBufferRef DataBuffer = GraphBuilder.RegisterExternalBuffer(LockedData.Buffer);
+			FRDGBufferRef IndirectBuffer = GraphBuilder.RegisterExternalBuffer(LockedData.IndirectBuffer);
+			InternalDrawView(GraphBuilder, View, DataBuffer, IndirectBuffer, OutputTexture, DepthTexture);
 		}
 	}
 
@@ -404,18 +331,6 @@ namespace ShaderDrawDebug
 		{
 			return;
 		}
-
-		if (IsShaderDrawLocked() && !LockedData.bIsLocked)
-		{
-			LockedData.bIsLocked = true;
-		}
-
-		if (!IsShaderDrawLocked() && LockedData.bIsLocked)
-		{
-			LockedData.Buffer.Release();
-			LockedData.IndirectBuffer.Release();
-			LockedData.bIsLocked = false;
-		}
 	}
 
 	void SetParameters(FRDGBuilder& GraphBuilder, const FShaderDrawDebugData& Data, FShaderDrawDebugParameters& OutParameters)
