Batch BuildRenderingCommands from major mesh passes inside the main render function

diff --git a/Engine/Source/Runtime/RenderCore/Private/RenderGraphBuilder.cpp b/Engine/Source/Runtime/RenderCore/Private/RenderGraphBuilder.cpp
--- a/Engine/Source/Runtime/RenderCore/Private/RenderGraphBuilder.cpp
+++ b/Engine/Source/Runtime/RenderCore/Private/RenderGraphBuilder.cpp
@@ -376,6 +376,11 @@ void FRDGBuilder::TickPoolElements()
 #endif
 }
 
+bool FRDGBuilder::IsImmediateMode()
+{
+	return ::IsImmediateMode();
+}
+
 ERDGPassFlags FRDGBuilder::OverridePassFlags(const TCHAR* PassName, ERDGPassFlags PassFlags, bool bAsyncComputeSupported)
 {
 	const bool bDebugAllowedForPass =
@@ -1902,10 +1907,27 @@ void FRDGBuilder::UploadBuffers()
 
 	for (FUploadedBuffer& UploadedBuffer : UploadedBuffers)
 	{
-		const TRefCountPtr<FRDGPooledBuffer>& PooledBuffer = ConvertToExternalBuffer(UploadedBuffer.Buffer);
-		void* DestPtr = RHICmdList.LockBuffer(PooledBuffer->GetRHI(), 0, UploadedBuffer.DataSize, RLM_WriteOnly);
-		FMemory::Memcpy(DestPtr, UploadedBuffer.Data, UploadedBuffer.DataSize);
-		RHICmdList.UnlockBuffer(PooledBuffer->GetRHI());
+		const void* InitialData;
+		uint64 InitialDataSize;
+
+		if (UploadedBuffer.bUseCallbacks)
+		{
+			InitialData = UploadedBuffer.DataCallback();
+			InitialDataSize = UploadedBuffer.DataSizeCallback();
+		}
+		else
+		{
+			InitialData = UploadedBuffer.Data;
+			InitialDataSize = UploadedBuffer.DataSize;
+		}
+
+		if (InitialData && InitialDataSize)
+		{
+			const TRefCountPtr<FRDGPooledBuffer>& PooledBuffer = ConvertToExternalBuffer(UploadedBuffer.Buffer);
+			void* DestPtr = RHICmdList.LockBuffer(PooledBuffer->GetRHI(), 0, InitialDataSize, RLM_WriteOnly);
+			FMemory::Memcpy(DestPtr, InitialData, InitialDataSize);
+			RHICmdList.UnlockBuffer(PooledBuffer->GetRHI());
+		}
 	}
 	UploadedBuffers.Reset();
 }
@@ -2685,6 +2707,7 @@ void FRDGBuilder::BeginResourceRHI(FRDGPassHandle PassHandle, FRDGBufferRef Buff
 		check(GetPrologueBarrierPassHandle(PassHandle) == PassHandle);
 	}
 #endif
+	Buffer->FinalizeDesc();
 
 	// If transient then create the resource on the transient allocator. External or extracted resource can't be transient because of lifetime tracking issues.
 	if (IsTransient(Buffer))
diff --git a/Engine/Source/Runtime/RenderCore/Private/RenderGraphUtils.cpp b/Engine/Source/Runtime/RenderCore/Private/RenderGraphUtils.cpp
--- a/Engine/Source/Runtime/RenderCore/Private/RenderGraphUtils.cpp
+++ b/Engine/Source/Runtime/RenderCore/Private/RenderGraphUtils.cpp
@@ -776,6 +776,19 @@ FRDGBufferRef CreateStructuredBuffer(
 	return Buffer;
 }
 
+FRDGBufferRef CreateStructuredBuffer(
+	FRDGBuilder& GraphBuilder,
+	const TCHAR* Name,
+	uint32 BytesPerElement,
+	FRDGBufferNumElementsCallback&& NumElementsCallback,
+	FRDGBufferInitialDataCallback&& InitialDataCallback,
+	FRDGBufferInitialDataSizeCallback&& InitialDataSizeCallback)
+{
+	FRDGBufferRef Buffer = GraphBuilder.CreateBuffer(FRDGBufferDesc::CreateStructuredDesc(BytesPerElement, 1), Name, MoveTemp(NumElementsCallback));
+	GraphBuilder.QueueBufferUpload(Buffer, MoveTemp(InitialDataCallback), MoveTemp(InitialDataSizeCallback));
+	return Buffer;
+}
+
 FRDGBufferRef CreateVertexBuffer(
 	FRDGBuilder& GraphBuilder,
 	const TCHAR* Name,
diff --git a/Engine/Source/Runtime/RenderCore/Private/RenderGraphValidation.cpp b/Engine/Source/Runtime/RenderCore/Private/RenderGraphValidation.cpp
--- a/Engine/Source/Runtime/RenderCore/Private/RenderGraphValidation.cpp
+++ b/Engine/Source/Runtime/RenderCore/Private/RenderGraphValidation.cpp
@@ -382,6 +382,13 @@ void FRDGUserValidation::ValidateUploadBuffer(FRDGBufferRef Buffer, const void*
 	check(InitialData || InitialDataSize == 0);
 }
 
+void FRDGUserValidation::ValidateUploadBuffer(FRDGBufferRef Buffer, const FRDGBufferInitialDataCallback& InitialDataCallback, const FRDGBufferInitialDataSizeCallback& InitialDataSizeCallback)
+{
+	check(Buffer);
+	checkf(!Buffer->bQueuedForUpload, TEXT("Buffer %s already has an upload queued. Only one upload can be done for each graph."), Buffer->Name);
+	check(InitialDataCallback && InitialDataSizeCallback);
+}
+
 void FRDGUserValidation::ValidateExtractTexture(FRDGTextureRef Texture, TRefCountPtr<IPooledRenderTarget>* OutTexturePtr)
 {
 	ValidateExtractResource(Texture);
diff --git a/Engine/Source/Runtime/RenderCore/Public/RenderGraphBuilder.h b/Engine/Source/Runtime/RenderCore/Public/RenderGraphBuilder.h
--- a/Engine/Source/Runtime/RenderCore/Public/RenderGraphBuilder.h
+++ b/Engine/Source/Runtime/RenderCore/Public/RenderGraphBuilder.h
@@ -71,6 +71,12 @@ public:
 	 */
 	FRDGBufferRef CreateBuffer(const FRDGBufferDesc& Desc, const TCHAR* Name, ERDGBufferFlags Flags = ERDGBufferFlags::None);
 
+	/** A variant of CreateBuffer where users supply NumElements through a callback. This allows creating buffers with
+	 *  sizes unknown at creation time. The callback is called before executing the most recent RDG pass that references
+	 *  the buffer so data must be ready before that.
+	 */
+	FRDGBufferRef CreateBuffer(const FRDGBufferDesc& Desc, const TCHAR* Name, FRDGBufferNumElementsCallback&& NumElementsCallback, ERDGBufferFlags Flags = ERDGBufferFlags::None);
+
 	/** Create graph tracked SRV for a texture from a descriptor. */
 	FRDGTextureSRVRef CreateSRV(const FRDGTextureSRVDesc& Desc);
 
@@ -173,6 +179,11 @@ public:
 	/** Queues a buffer upload operation prior to execution. The resource lifetime is extended and the data is uploaded prior to executing passes. */
 	void QueueBufferUpload(FRDGBufferRef Buffer, const void* InitialData, uint64 InitialDataSize, ERDGInitialDataFlags InitialDataFlags = ERDGInitialDataFlags::None);
 
+	/** A variant where InitialData and InitialDataSize are supplied through callbacks. This allows queuing an upload with information unknown at
+	 *  creation time. The callbacks are called before RDG pass execution so data must be ready before that.
+	 */
+	void QueueBufferUpload(FRDGBufferRef Buffer, FRDGBufferInitialDataCallback&& InitialDataCallback, FRDGBufferInitialDataSizeCallback&& InitialDataSizeCallback);
+
 	/** Queues a pooled render target extraction to happen at the end of graph execution. For graph-created textures, this extends
 	 *  the lifetime of the GPU resource until execution, at which point the pointer is filled. If specified, the texture is transitioned
 	 *  to the AccessFinal state, or kDefaultAccessFinal otherwise.
@@ -250,6 +261,9 @@ public:
 	/** Per-frame update of the render graph resource pool. */
 	static void TickPoolElements();
 
+	/** Whether RDG is running in immediate mode. */
+	static bool IsImmediateMode();
+
 	/** The RHI command list used for the render graph. */
 	FRHICommandListImmediate& RHICmdList;
 
@@ -472,14 +486,25 @@ private:
 		FUploadedBuffer() = default;
 
 		FUploadedBuffer(FRDGBuffer* InBuffer, const void* InData, uint64 InDataSize)
-			: Buffer(InBuffer)
+			: bUseCallbacks(false)
+			, Buffer(InBuffer)
 			, Data(InData)
 			, DataSize(InDataSize)
 		{}
 
+		FUploadedBuffer(FRDGBuffer* InBuffer, FRDGBufferInitialDataCallback&& InDataCallback, FRDGBufferInitialDataSizeCallback&& InDataSizeCallback)
+			: bUseCallbacks(true)
+			, Buffer(InBuffer)
+			, DataCallback(MoveTemp(InDataCallback))
+			, DataSizeCallback(MoveTemp(InDataSizeCallback))
+		{}
+
+		bool bUseCallbacks = false;
 		FRDGBuffer* Buffer{};
 		const void* Data{};
 		uint64 DataSize{};
+		FRDGBufferInitialDataCallback DataCallback;
+		FRDGBufferInitialDataSizeCallback DataSizeCallback;
 	};
 
 	TArray<FUploadedBuffer, FRDGArrayAllocator> UploadedBuffers;
diff --git a/Engine/Source/Runtime/RenderCore/Public/RenderGraphBuilder.inl b/Engine/Source/Runtime/RenderCore/Public/RenderGraphBuilder.inl
--- a/Engine/Source/Runtime/RenderCore/Public/RenderGraphBuilder.inl
+++ b/Engine/Source/Runtime/RenderCore/Public/RenderGraphBuilder.inl
@@ -52,6 +52,23 @@ inline FRDGBufferRef FRDGBuilder::CreateBuffer(
 	return Buffer;
 }
 
+inline FRDGBufferRef FRDGBuilder::CreateBuffer(
+	const FRDGBufferDesc& Desc,
+	const TCHAR* Name,
+	FRDGBufferNumElementsCallback&& NumElementsCallback,
+	ERDGBufferFlags Flags)
+{
+	// RDG no longer supports the legacy transient resource API.
+	FRDGBufferDesc OverrideDesc = Desc;
+	EnumRemoveFlags(OverrideDesc.Usage, BUF_Transient);
+
+	IF_RDG_ENABLE_DEBUG(UserValidation.ValidateCreateBuffer(Desc, Name, Flags));
+	FRDGBufferRef Buffer = Buffers.Allocate(Allocator, Name, OverrideDesc, Flags, MoveTemp(NumElementsCallback));
+	IF_RDG_ENABLE_DEBUG(UserValidation.ValidateCreateBuffer(Buffer));
+	IF_RDG_ENABLE_TRACE(Trace.AddResource(Buffer));
+	return Buffer;
+}
+
 inline FRDGTextureSRVRef FRDGBuilder::CreateSRV(const FRDGTextureSRVDesc& Desc)
 {
 	IF_RDG_ENABLE_DEBUG(UserValidation.ValidateCreateSRV(Desc));
@@ -204,6 +221,14 @@ inline void FRDGBuilder::QueueBufferUpload(FRDGBufferRef Buffer, const void* Ini
 	Buffer->bQueuedForUpload = 1;
 }
 
+inline void FRDGBuilder::QueueBufferUpload(FRDGBufferRef Buffer, FRDGBufferInitialDataCallback&& InitialDataCallback, FRDGBufferInitialDataSizeCallback&& InitialDataSizeCallback)
+{
+	IF_RDG_ENABLE_DEBUG(UserValidation.ValidateUploadBuffer(Buffer, InitialDataCallback, InitialDataSizeCallback));
+
+	UploadedBuffers.Emplace(Buffer, MoveTemp(InitialDataCallback), MoveTemp(InitialDataSizeCallback));
+	Buffer->bQueuedForUpload = 1;
+}
+
 inline void FRDGBuilder::QueueTextureExtraction(FRDGTextureRef Texture, TRefCountPtr<IPooledRenderTarget>* OutTexturePtr, ERHIAccess AccessFinal)
 {
 	QueueTextureExtraction(Texture, OutTexturePtr);
diff --git a/Engine/Source/Runtime/RenderCore/Public/RenderGraphDefinitions.h b/Engine/Source/Runtime/RenderCore/Public/RenderGraphDefinitions.h
--- a/Engine/Source/Runtime/RenderCore/Public/RenderGraphDefinitions.h
+++ b/Engine/Source/Runtime/RenderCore/Public/RenderGraphDefinitions.h
@@ -576,4 +576,9 @@ using TRDGTextureSubresourceArray = TArray<InElementType, TInlineAllocator<1, In
 using FRDGPassHandlesByPipeline = TRHIPipelineArray<FRDGPassHandle>;
 using FRDGPassesByPipeline = TRHIPipelineArray<FRDGPass*>;
 
-class FRDGTrace;
\ No newline at end of file
+class FRDGTrace;
+
+using FRDGBufferNumElementsCallback = TFunction<uint32()>;
+using FRDGBufferInitialDataCallback = TFunction<const void*()>;
+using FRDGBufferInitialDataSizeCallback = TFunction<uint64()>;
+using FRDGDispatchGroupCountCallback = TFunction<FIntVector()>;
\ No newline at end of file
diff --git a/Engine/Source/Runtime/RenderCore/Public/RenderGraphResources.h b/Engine/Source/Runtime/RenderCore/Public/RenderGraphResources.h
--- a/Engine/Source/Runtime/RenderCore/Public/RenderGraphResources.h
+++ b/Engine/Source/Runtime/RenderCore/Public/RenderGraphResources.h
@@ -931,7 +931,7 @@ class RENDERCORE_API FRDGBuffer final
 	: public FRDGParentResource
 {
 public:
-	const FRDGBufferDesc Desc;
+	FRDGBufferDesc Desc;
 	const ERDGBufferFlags Flags;
 
 	//////////////////////////////////////////////////////////////////////////
@@ -978,6 +978,12 @@ private:
 		, Flags(InFlags)
 	{}
 
+	FRDGBuffer(const TCHAR* InName, const FRDGBufferDesc& InDesc, ERDGBufferFlags InFlags, FRDGBufferNumElementsCallback&& InNumElementsCallback)
+		: FRDGBuffer(InName, InDesc, InFlags)
+	{
+		NumElementsCallback = MoveTemp(InNumElementsCallback);
+	}
+
 	/** Assigns a pooled buffer as the backing RHI resource. */
 	void SetRHI(FRDGPooledBuffer* InPooledBuffer);
 
@@ -987,6 +993,15 @@ private:
 	/** Finalizes the buffer for execution; no other transitions are allowed after calling this. */
 	void Finalize(FRDGPooledBufferArray& PooledBufferArray);
 
+	/** Finalizes any pending field of the buffer descriptor. */
+	void FinalizeDesc()
+	{
+		if (NumElementsCallback)
+		{
+			Desc.NumElements = FMath::Max(NumElementsCallback(), 1u);
+		}
+	}
+
 	FRHIBuffer* GetRHIUnchecked() const
 	{
 		return static_cast<FRHIBuffer*>(FRDGResource::GetRHIUnchecked());
@@ -1029,6 +1044,9 @@ private:
 	/** Tracks the last pass that produced this resource as the graph is built. */
 	FRDGProducerStatesByPipeline LastProducer;
 
+	/** Optional callback to supply NumElements after the creation of this FRDGBuffer. */
+	FRDGBufferNumElementsCallback NumElementsCallback;
+
 #if RDG_ENABLE_DEBUG
 	struct FRDGBufferDebugData* BufferDebugData = nullptr;
 	FRDGBufferDebugData& GetBufferDebugData() const;
diff --git a/Engine/Source/Runtime/RenderCore/Public/RenderGraphUtils.h b/Engine/Source/Runtime/RenderCore/Public/RenderGraphUtils.h
--- a/Engine/Source/Runtime/RenderCore/Public/RenderGraphUtils.h
+++ b/Engine/Source/Runtime/RenderCore/Public/RenderGraphUtils.h
@@ -382,6 +382,41 @@ struct RENDERCORE_API FComputeShaderUtils
 		});
 	}
 
+	/** Dispatch a compute shader to render graph builder with its parameters. GroupCount is supplied through a callback.
+	 *  This allows adding a dispatch with unknown GroupCount but the value must be ready before the pass is executed.
+	 */
+	template<typename TShaderClass>
+	static void AddPass(
+		FRDGBuilder& GraphBuilder,
+		FRDGEventName&& PassName,
+		ERDGPassFlags PassFlags,
+		const TShaderRef<TShaderClass>& ComputeShader,
+		const FShaderParametersMetadata* ParametersMetadata,
+		typename TShaderClass::FParameters* Parameters,
+		FRDGDispatchGroupCountCallback&& GroupCountCallback)
+	{
+		checkf(
+			EnumHasAnyFlags(PassFlags, ERDGPassFlags::Compute | ERDGPassFlags::AsyncCompute) &&
+			!EnumHasAnyFlags(PassFlags, ERDGPassFlags::Copy | ERDGPassFlags::Raster), TEXT("AddPass only supports 'Compute' or 'AsyncCompute'."));
+
+		ClearUnusedGraphResources(ComputeShader, ParametersMetadata, Parameters);
+
+		GraphBuilder.AddPass(
+			Forward<FRDGEventName>(PassName),
+			ParametersMetadata,
+			Parameters,
+			PassFlags,
+			[ParametersMetadata, Parameters, ComputeShader, GroupCountCallback = MoveTemp(GroupCountCallback)](FRHIComputeCommandList& RHICmdList)
+			{
+				const FIntVector GroupCount = GroupCountCallback();
+				if (GroupCount.X > 0 && GroupCount.Y > 0 && GroupCount.Z > 0)
+				{
+					ValidateGroupCount(GroupCount);
+					FComputeShaderUtils::Dispatch(RHICmdList, ComputeShader, ParametersMetadata, *Parameters, GroupCount);
+				}
+			});
+	}
+
 	template<typename TShaderClass>
 	static void AddPass(
 		FRDGBuilder& GraphBuilder,
@@ -407,6 +442,18 @@ struct RENDERCORE_API FComputeShaderUtils
 		AddPass(GraphBuilder, Forward<FRDGEventName>(PassName), ERDGPassFlags::Compute, ComputeShader, ParametersMetadata, Parameters, GroupCount);
 	}
 
+	template <typename TShaderClass>
+	static FORCEINLINE void AddPass(
+		FRDGBuilder& GraphBuilder,
+		FRDGEventName&& PassName,
+		const TShaderRef<TShaderClass>& ComputeShader,
+		typename TShaderClass::FParameters* Parameters,
+		FRDGDispatchGroupCountCallback&& GroupCountCallback)
+	{
+		const FShaderParametersMetadata* ParametersMetadata = TShaderClass::FParameters::FTypeInfo::GetStructMetadata();
+		AddPass(GraphBuilder, Forward<FRDGEventName>(PassName), ERDGPassFlags::Compute, ComputeShader, ParametersMetadata, Parameters, MoveTemp(GroupCountCallback));
+	}
+
 	/** Dispatch a compute shader to render graph builder with its parameters. */
 	template<typename TShaderClass>
 	static void AddPass(
@@ -607,6 +654,18 @@ RENDERCORE_API FRDGBufferRef CreateStructuredBuffer(
 	uint64 InitialDataSize,
 	ERDGInitialDataFlags InitialDataFlags = ERDGInitialDataFlags::None);
 
+/** A variant where NumElements, InitialData, and InitialDataSize are supplied through callbacks. This allows creating a buffer with
+ *  with information unknown at creation time. Though, data must be ready before the most recent RDG pass that references the buffer
+ *  is executed.
+ */
+RENDERCORE_API FRDGBufferRef CreateStructuredBuffer(
+	FRDGBuilder& GraphBuilder,
+	const TCHAR* Name,
+	uint32 BytesPerElement,
+	FRDGBufferNumElementsCallback&& NumElementsCallback,
+	FRDGBufferInitialDataCallback&& InitialDataCallback,
+	FRDGBufferInitialDataSizeCallback&& InitialDataSizeCallback);
+
 /**
  * Helper to create a structured buffer with initial data from a TArray.
  */
diff --git a/Engine/Source/Runtime/RenderCore/Public/RenderGraphValidation.h b/Engine/Source/Runtime/RenderCore/Public/RenderGraphValidation.h
--- a/Engine/Source/Runtime/RenderCore/Public/RenderGraphValidation.h
+++ b/Engine/Source/Runtime/RenderCore/Public/RenderGraphValidation.h
@@ -60,6 +60,7 @@ public:
 	void ValidateRegisterExternalBuffer(FRDGBufferRef Buffer);
 
 	void ValidateUploadBuffer(FRDGBufferRef Buffer, const void* InitialData, uint64 InitialDataSize);
+	void ValidateUploadBuffer(FRDGBufferRef Buffer, const FRDGBufferInitialDataCallback& InitialDataCallback, const FRDGBufferInitialDataSizeCallback& InitialDataSizeCallback);
 
 	/** Validates a resource extraction operation. */
 	void ValidateExtractTexture(FRDGTextureRef Texture, TRefCountPtr<IPooledRenderTarget>* OutTexturePtr);
diff --git a/Engine/Source/Runtime/Renderer/Private/DeferredShadingRenderer.cpp b/Engine/Source/Runtime/Renderer/Private/DeferredShadingRenderer.cpp
--- a/Engine/Source/Runtime/Renderer/Private/DeferredShadingRenderer.cpp
+++ b/Engine/Source/Runtime/Renderer/Private/DeferredShadingRenderer.cpp
@@ -1887,7 +1887,7 @@ void FDeferredShadingSceneRenderer::Render(FRDGBuilder& GraphBuilder)
 	// Find the visible primitives.
 	GraphBuilder.RHICmdList.ImmediateFlush(EImmediateFlushType::DispatchToRHIThread);
 
-	FInstanceCullingManager InstanceCullingManager(Scene->GPUScene.IsEnabled());
+	FInstanceCullingManager& InstanceCullingManager = *GraphBuilder.AllocObject<FInstanceCullingManager>(Scene->GPUScene.IsEnabled());
 
 	bool bDoInitViewAftersPrepass = false;
 	{
@@ -1972,6 +1972,9 @@ void FDeferredShadingSceneRenderer::Render(FRDGBuilder& GraphBuilder)
 		{
 			// GPUCULL_TODO: Possibly fold into unpack step
 			InstanceCullingManager.CullInstances(GraphBuilder, Scene->GPUScene);
+
+			// This sets up a deferred context that is used by all subsequent culling passes until it is either flushed by RDG execute or explicitly closed
+			InstanceCullingManager.BeginDeferredCulling(GraphBuilder, Scene->GPUScene);
 		}
 
 		if (!bDoInitViewAftersPrepass)
diff --git a/Engine/Source/Runtime/Renderer/Private/InstanceCulling/InstanceCullingContext.cpp b/Engine/Source/Runtime/Renderer/Private/InstanceCulling/InstanceCullingContext.cpp
--- a/Engine/Source/Runtime/Renderer/Private/InstanceCulling/InstanceCullingContext.cpp
+++ b/Engine/Source/Runtime/Renderer/Private/InstanceCulling/InstanceCullingContext.cpp
@@ -17,20 +17,12 @@
 IMPLEMENT_STATIC_UNIFORM_BUFFER_SLOT(InstanceCullingUbSlot);
 IMPLEMENT_STATIC_UNIFORM_BUFFER_STRUCT(FInstanceCullingGlobalUniforms, "InstanceCulling", InstanceCullingUbSlot);
 
-struct FBatchInfo
-{
-	uint32 CullingCommandsOffset;
-	uint32 NumCullingCommands;
-	uint32 PrimitiveIdsOffset;
-	uint32 NumPrimitiveIds;
-	uint32 InstanceRunsOffset;
-	uint32 NumInstanceRuns;
-	uint32 ViewIdsOffset;
-	uint32 NumViewIds;
-	int32 DynamicPrimitiveIdOffset;
-	int32 DynamicPrimitiveIdMax;
-};
-
+static int32 GAllowBatchedBuildRenderingCommands = 1;
+static FAutoConsoleVariableRef CVarAllowBatchedBuildRenderingCommands(
+	TEXT("r.InstanceCulling.AllowBatchedBuildRenderingCommands"),
+	GAllowBatchedBuildRenderingCommands,
+	TEXT("Whether to allow batching BuildRenderingCommands for GPU instance culling"),
+	ECVF_RenderThreadSafe);
 
 FInstanceCullingContext::FInstanceCullingContext(FInstanceCullingManager* InInstanceCullingManager, TArrayView<const int32> InViewIds, enum EInstanceCullingMode InInstanceCullingMode, bool bInDrawOnlyVSMInvalidatingGeometry) :
 	InstanceCullingManager(InInstanceCullingManager),
@@ -341,10 +333,16 @@ const TRDGUniformBufferRef<FInstanceCullingGlobalUniforms> FInstanceCullingConte
 	return GraphBuilder.CreateUniformBuffer(InstanceCullingGlobalUniforms);
 }
 
-void FInstanceCullingContext::BuildRenderingCommands(FRDGBuilder& GraphBuilder, const FGPUScene& GPUScene, const TRange<int32>& DynamicPrimitiveIdRange, FInstanceCullingResult& Results) const
+void FInstanceCullingContext::BuildRenderingCommands(
+	FRDGBuilder& GraphBuilder,
+	const FGPUScene& GPUScene,
+	const TRange<int32>& DynamicPrimitiveIdRange,
+	FInstanceCullingResult& Results,
+	FInstanceCullingDrawParams* InstanceCullingDrawParams) const
 {
 	Results = FInstanceCullingResult();
-	if (CullingCommands.Num() == 0 || TotalInstances == 0)
+
+	if (!HasCullingCommands())
 	{
 		if (InstanceCullingManager)
 		{
@@ -353,6 +351,15 @@ void FInstanceCullingContext::BuildRenderingCommands(FRDGBuilder& GraphBuilder,
 		return;
 	}
 
+	if (InstanceCullingDrawParams && InstanceCullingManager && InstanceCullingManager->IsDeferredCullingActive())
+	{
+		Results.DrawIndirectArgsBuffer = InstanceCullingManager->BatchedCullingScratch.DrawIndirectArgsBuffer;
+		Results.InstanceIdOffsetBuffer = InstanceCullingManager->BatchedCullingScratch.InstanceIdOffsetBuffer;
+		Results.UniformBuffer = InstanceCullingManager->BatchedCullingScratch.UniformBuffer;
+		InstanceCullingManager->BatchedCullingScratch.Batches.Add(FBatchItem{ this, InstanceCullingDrawParams, DynamicPrimitiveIdRange });
+		return;
+	}
+
 	ensure(InstanceCullingMode == EInstanceCullingMode::Normal || ViewIds.Num() == 2);
 
 	RDG_EVENT_SCOPE(GraphBuilder, "BuildRenderingCommands");
@@ -680,94 +687,129 @@ void FInstanceCullingContext::BuildRenderingCommands(FRDGBuilder& GraphBuilder,
 	);
 }
 
-
-void FInstanceCullingContext::BuildRenderingCommandsBatched(FRDGBuilder& GraphBuilder, FGPUScene& GPUScene, TArrayView<FInstanceCullingContext::FBatchItem> Batches)
+/** Helper function to merge GPU instance culling input batches queued throughout a frame. */
+static void MergeInstanceCullingContext(
+	const TArray<FInstanceCullingContext::FBatchItem, SceneRenderingAllocator>& Batches,
+	FBuildInstanceIdBufferAndCommandsFromPrimitiveIdsCs::FParameters* PassParameters,
+	FInstanceCullingContextMerged& MergedContext)
 {
-	if (Batches.Num() == 0)
-	{
-		return;
-	}
-	
-	RDG_EVENT_SCOPE(GraphBuilder, "BuildRenderingCommandsBatched");
-
-	FInstanceCullingManager* InstanceCullingManager = nullptr;
-	// TODO: In the future, we should pre-upload the commands for static meshes when cached, only dynamic need uploading each frame, they should be kicked asynchronously or flushed in batches.
-	//       Then this would only need indices.
-
-	// 1. conglomerate and upload data for all batches - consider pre-allocating this in shared place directly to avoid copy.
-	TArray<FPrimCullingCommand, SceneRenderingAllocator> CullingCommandsAll;
-	TArray<int32, SceneRenderingAllocator> PrimitiveIdsAll;
-	TArray<FInstanceRun, SceneRenderingAllocator> InstanceRunsAll;
-	TArray<int32, SceneRenderingAllocator> ViewIdsAll;
-	TArray<FBatchInfo, SceneRenderingAllocator> BatchInfos;
-	BatchInfos.AddDefaulted(Batches.Num());
-	uint32 InstanceIdBufferSize = 0U;
-	// Index that maps from each command to the corresponding batch - maybe not the utmost efficiency
-	TArray<int32, SceneRenderingAllocator> BatchIndsAll;
-	for (int32 BatchIndex = 0; BatchIndex < Batches.Num(); ++BatchIndex)
+	if (!MergedContext.bInitialized)
 	{
-		FBatchItem& BatchItem = Batches[BatchIndex];
-		FInstanceCullingContext& InstanceCullingContext = *BatchItem.Context;
-
-		if (InstanceCullingContext.InstanceCullingManager != nullptr)
+		MergedContext.bInitialized = true;
+
+		// TODO: In the future, we should pre-upload the commands for static meshes when cached, only dynamic need uploading each frame, they should be kicked asynchronously or flushed in batches.
+		//       Then this would only need indices.
+
+		// 1. conglomerate and upload data for all batches - consider pre-allocating this in shared place directly to avoid copy.
+		TArray<FInstanceCullingContext::FPrimCullingCommand, SceneRenderingAllocator>& CullingCommandsAll = MergedContext.CullingCommandsAll;
+		TArray<int32, SceneRenderingAllocator>& PrimitiveIdsAll = MergedContext.PrimitiveIdsAll;
+		TArray<FInstanceCullingContext::FInstanceRun, SceneRenderingAllocator>& InstanceRunsAll = MergedContext.InstanceRunsAll;
+		TArray<int32, SceneRenderingAllocator>& ViewIdsAll = MergedContext.ViewIdsAll;
+		TArray<FInstanceCullingContext::FBatchInfo, SceneRenderingAllocator>& BatchInfos = MergedContext.BatchInfos;
+		BatchInfos.AddDefaulted(Batches.Num());
+		MergedContext.InstanceIdBufferSize = 0U;
+
+		// Index that maps from each command to the corresponding batch - maybe not the utmost efficiency
+		TArray<int32, SceneRenderingAllocator>& BatchIndsAll = MergedContext.BatchIndsAll;
+		for (int32 BatchIndex = 0; BatchIndex < Batches.Num(); ++BatchIndex)
 		{
-			// Be very weird if there were different ones..
-			ensure(InstanceCullingManager == nullptr || InstanceCullingManager == InstanceCullingContext.InstanceCullingManager);
-			InstanceCullingManager = InstanceCullingContext.InstanceCullingManager;
-		}
+			const FInstanceCullingContext::FBatchItem& BatchItem = Batches[BatchIndex];
+			const FInstanceCullingContext& InstanceCullingContext = *BatchItem.Context;
 
-		FBatchInfo& BatchInfo = BatchInfos[BatchIndex];
+			FInstanceCullingContext::FBatchInfo& BatchInfo = BatchInfos[BatchIndex];
 
-		BatchInfo.CullingCommandsOffset = CullingCommandsAll.Num();
-		BatchInfo.NumCullingCommands = InstanceCullingContext.CullingCommands.Num();
-		CullingCommandsAll.Append(InstanceCullingContext.CullingCommands);
-		
-		BatchInfo.PrimitiveIdsOffset = PrimitiveIdsAll.Num();
-		BatchInfo.NumPrimitiveIds = InstanceCullingContext.PrimitiveIds.Num();
-		PrimitiveIdsAll.Append(InstanceCullingContext.PrimitiveIds);
+			BatchInfo.CullingCommandsOffset = CullingCommandsAll.Num();
+			BatchInfo.NumCullingCommands = InstanceCullingContext.CullingCommands.Num();
+			CullingCommandsAll.Append(InstanceCullingContext.CullingCommands);
 
-		BatchInfo.InstanceRunsOffset = InstanceRunsAll.Num();
-		BatchInfo.NumInstanceRuns = InstanceCullingContext.InstanceRuns.Num();
-		InstanceRunsAll.Append(InstanceCullingContext.InstanceRuns);
+			BatchInfo.PrimitiveIdsOffset = PrimitiveIdsAll.Num();
+			BatchInfo.NumPrimitiveIds = InstanceCullingContext.PrimitiveIds.Num();
+			PrimitiveIdsAll.Append(InstanceCullingContext.PrimitiveIds);
 
-		BatchInfo.ViewIdsOffset = ViewIdsAll.Num();
-		BatchInfo.NumViewIds = InstanceCullingContext.ViewIds.Num();
-		ViewIdsAll.Append(InstanceCullingContext.ViewIds);
+			BatchInfo.InstanceRunsOffset = InstanceRunsAll.Num();
+			BatchInfo.NumInstanceRuns = InstanceCullingContext.InstanceRuns.Num();
+			InstanceRunsAll.Append(InstanceCullingContext.InstanceRuns);
 
-		BatchInfo.DynamicPrimitiveIdOffset = BatchItem.DynamicPrimitiveIdRange.GetLowerBoundValue();
-		BatchInfo.DynamicPrimitiveIdMax = BatchItem.DynamicPrimitiveIdRange.GetUpperBoundValue();
+			BatchInfo.ViewIdsOffset = ViewIdsAll.Num();
+			BatchInfo.NumViewIds = InstanceCullingContext.ViewIds.Num();
+			ViewIdsAll.Append(InstanceCullingContext.ViewIds);
 
-		// Map to batch from index in command list
-		// TODO: could be made more efficient than having an index back to the batch for each command
-		BatchIndsAll.AddDefaulted(InstanceCullingContext.CullingCommands.Num());
-		for (int32 Index = BatchInfo.CullingCommandsOffset; Index < BatchIndsAll.Num(); ++Index)
-		{
-			BatchIndsAll[Index] = BatchIndex;
+			BatchInfo.DynamicPrimitiveIdOffset = BatchItem.DynamicPrimitiveIdRange.GetLowerBoundValue();
+			BatchInfo.DynamicPrimitiveIdMax = BatchItem.DynamicPrimitiveIdRange.GetUpperBoundValue();
+
+			// Map to batch from index in command list
+			// TODO: could be made more efficient than having an index back to the batch for each command
+			BatchIndsAll.AddDefaulted(InstanceCullingContext.CullingCommands.Num());
+			for (int32 Index = BatchInfo.CullingCommandsOffset; Index < BatchIndsAll.Num(); ++Index)
+			{
+				BatchIndsAll[Index] = BatchIndex;
+			}
+
+			MergedContext.InstanceIdBufferSize += InstanceCullingContext.TotalInstances * InstanceCullingContext.ViewIds.Num();
+
+			FInstanceCullingDrawParams& Result = *BatchItem.Result;
+			Result.DrawCommandDataOffset = BatchInfo.CullingCommandsOffset;
 		}
 
-		InstanceIdBufferSize += InstanceCullingContext.TotalInstances * InstanceCullingContext.ViewIds.Num();
+		PassParameters->NumViewIds = ViewIdsAll.Num();
+		PassParameters->NumPrimitiveIds = PrimitiveIdsAll.Num();
+		PassParameters->NumInstanceRuns = InstanceRunsAll.Num();
+		PassParameters->NumCommands = CullingCommandsAll.Num();
+	}
+}
+
+void FInstanceCullingContext::BuildRenderingCommandsDeferred(
+	FRDGBuilder& GraphBuilder,
+	FGPUScene& GPUScene,
+	FInstanceCullingManager& InstanceCullingManager)
+{
+#define INST_CULL_CALLBACK(CustomCode) \
+	[&Batches, PassParameters, &MergedContext]() \
+	{ \
+		MergeInstanceCullingContext(Batches, PassParameters, MergedContext); \
+		return CustomCode; \
+	}
+#define INST_CULL_CREATE_STRUCT_BUFF_ARGS(ArrayName) \
+	GraphBuilder, \
+	TEXT("InstanceCulling.") TEXT(#ArrayName), \
+	MergedContext.ArrayName.GetTypeSize(), \
+	INST_CULL_CALLBACK(MergedContext.ArrayName.Num()), \
+	INST_CULL_CALLBACK(MergedContext.ArrayName.GetData()), \
+	INST_CULL_CALLBACK(MergedContext.ArrayName.Num() * MergedContext.ArrayName.GetTypeSize())
+
+	// Cannot defer pass execution in immediate mode.
+	// TODO: support batching when r.RDG.Drain is enabled
+	if (!AllowBatchedBuildRenderingCommands(GPUScene))
+	{
+		return;
 	}
 
-	FRDGBufferRef CullingCommandsRDG = CreateStructuredBuffer(GraphBuilder, TEXT("CullingCommandsAll"), CullingCommandsAll);
-	FRDGBufferRef PrimitiveIdsRDG = CreateStructuredBuffer(GraphBuilder, TEXT("PrimitiveIdsAll"), PrimitiveIdsAll);
-	FRDGBufferRef InstanceRunsRDG = CreateStructuredBuffer(GraphBuilder, TEXT("InstanceRunsAll"), InstanceRunsAll);
-	FRDGBufferRef ViewIdsRDG = CreateStructuredBuffer(GraphBuilder, TEXT("ViewIdsAll"), ViewIdsAll);
-	FRDGBufferRef BatchInfosRDG = CreateStructuredBuffer(GraphBuilder, TEXT("BatchInfos"), BatchInfos);
-	FRDGBufferRef BatchIndsRDG = CreateStructuredBuffer(GraphBuilder, TEXT("BatchIndsAll"), BatchIndsAll);
+	FBatchedInstanceCullingScratchSpace& BatchedCullingScratch = InstanceCullingManager.BatchedCullingScratch;
 
-	FRDGBufferRef InstanceIdsBuffer = GraphBuilder.CreateBuffer(FRDGBufferDesc::CreateStructuredDesc(sizeof(uint32), InstanceIdBufferSize), TEXT("InstanceCulling.InstanceIdsBuffer"));
+	checkf(!BatchedCullingScratch.bBatchingActive, TEXT("Batched GPU instance culling has already been enabled."));
+	BatchedCullingScratch.bBatchingActive = true;
+
+	const TArray<FBatchItem, SceneRenderingAllocator>& Batches = BatchedCullingScratch.Batches;
+	FInstanceCullingContextMerged& MergedContext = BatchedCullingScratch.MergedContext;
 
 	FBuildInstanceIdBufferAndCommandsFromPrimitiveIdsCs::FParameters* PassParameters = GraphBuilder.AllocParameters<FBuildInstanceIdBufferAndCommandsFromPrimitiveIdsCs::FParameters>();
 
+	FRDGBufferRef CullingCommandsRDG = CreateStructuredBuffer(INST_CULL_CREATE_STRUCT_BUFF_ARGS(CullingCommandsAll));
+	FRDGBufferRef PrimitiveIdsRDG = CreateStructuredBuffer(INST_CULL_CREATE_STRUCT_BUFF_ARGS(PrimitiveIdsAll));
+	FRDGBufferRef InstanceRunsRDG = CreateStructuredBuffer(INST_CULL_CREATE_STRUCT_BUFF_ARGS(InstanceRunsAll));
+	FRDGBufferRef ViewIdsRDG = CreateStructuredBuffer(INST_CULL_CREATE_STRUCT_BUFF_ARGS(ViewIdsAll));
+	FRDGBufferRef BatchInfosRDG = CreateStructuredBuffer(INST_CULL_CREATE_STRUCT_BUFF_ARGS(BatchInfos));
+	FRDGBufferRef BatchIndsRDG = CreateStructuredBuffer(INST_CULL_CREATE_STRUCT_BUFF_ARGS(BatchIndsAll));
+	FRDGBufferRef InstanceIdsBuffer = GraphBuilder.CreateBuffer(FRDGBufferDesc::CreateStructuredDesc(sizeof(uint32), 1), TEXT("InstanceCulling.InstanceIdsBuffer"), INST_CULL_CALLBACK(MergedContext.InstanceIdBufferSize));
+
 	// Note: use start at zero offset if there is no instance culling manager, this means each build rendering commands pass will overwrite the same ID range. Which is only ok assuming correct barriers (should be erring on this side by default).
 	TArray<uint32, TInlineAllocator<1> > NullArray;
 	NullArray.AddZeroed(1);
 	FRDGBufferRef InstanceIdOutOffsetBufferRDG = CreateStructuredBuffer(GraphBuilder, TEXT("InstanceCulling.OutputOffsetBufferOut"), NullArray);
 	// If there is no manager, then there is no data on culling, so set flag to skip that and ignore buffers.
-	FRDGBufferRef VisibleInstanceFlagsRDG = InstanceCullingManager != nullptr ? InstanceCullingManager->CullingIntermediate.VisibleInstanceFlags : nullptr;
-	const bool bCullInstances = InstanceCullingManager != nullptr;
-	const int32 NumCulledInstances = InstanceCullingManager != nullptr ? InstanceCullingManager->CullingIntermediate.NumInstances : 0;
-	const uint32 NumCulledViews = InstanceCullingManager != nullptr ? uint32(InstanceCullingManager->CullingIntermediate.NumViews) : 0U;
+	FRDGBufferRef VisibleInstanceFlagsRDG = InstanceCullingManager.CullingIntermediate.VisibleInstanceFlags;
+	const int32 NumCulledInstances = InstanceCullingManager.CullingIntermediate.NumInstances;
+	const uint32 NumCulledViews = uint32(InstanceCullingManager.CullingIntermediate.NumViews);
 	int32 NumInstanceFlagWords = FMath::DivideAndRoundUp(NumCulledInstances, int32(sizeof(uint32) * 8));
 
 	// Because the view uniforms are not set up by the time this runs
@@ -786,20 +828,16 @@ void FInstanceCullingContext::BuildRenderingCommandsBatched(FRDGBuilder& GraphBu
 
 	PassParameters->OutputOffsetBufferOut = GraphBuilder.CreateUAV(InstanceIdOutOffsetBufferRDG);
 
-	FRDGBufferRef DrawIndirectArgsRDG = GraphBuilder.CreateBuffer(FRDGBufferDesc::CreateIndirectDesc(IndirectArgsNumWords * CullingCommandsAll.Num()), TEXT("DrawIndirectArgsBuffer"));
+	BatchedCullingScratch.DrawIndirectArgsBuffer = GraphBuilder.CreateBuffer(FRDGBufferDesc::CreateIndirectDesc(), TEXT("DrawIndirectArgsBuffer"), INST_CULL_CALLBACK(FInstanceCullingContext::IndirectArgsNumWords * MergedContext.CullingCommandsAll.Num()));
 	// not using structured buffer as we want/have to get at it as a vertex buffer 
 	//FRDGBufferRef InstanceIdsBufferRDG = GraphBuilder.CreateBuffer(FRDGBufferDesc::CreateBufferDesc(sizeof(uint32), PrimitiveIds.Num() * FInstanceCullingManager::MaxAverageInstanceFactor), TEXT("InstanceIdsBuffer"));
-	FRDGBufferRef InstanceIdOffsetBufferRDG = GraphBuilder.CreateBuffer(FRDGBufferDesc::CreateBufferDesc(sizeof(uint32), CullingCommandsAll.Num()), TEXT("InstanceIdOffsetBuffer"));
+	BatchedCullingScratch.InstanceIdOffsetBuffer = GraphBuilder.CreateBuffer(FRDGBufferDesc::CreateBufferDesc(sizeof(uint32), 1), TEXT("InstanceIdOffsetBuffer"), INST_CULL_CALLBACK(MergedContext.CullingCommandsAll.Num()));
 
 	PassParameters->ViewIds = GraphBuilder.CreateSRV(ViewIdsRDG);
-	PassParameters->NumViewIds = ViewIdsAll.Num();
 
 	PassParameters->InstanceIdsBufferOut = GraphBuilder.CreateUAV(InstanceIdsBuffer);
-	PassParameters->DrawIndirectArgsBufferOut = GraphBuilder.CreateUAV(DrawIndirectArgsRDG, PF_R32_UINT);
-	PassParameters->InstanceIdOffsetBufferOut = GraphBuilder.CreateUAV(InstanceIdOffsetBufferRDG, PF_R32_UINT);
-	PassParameters->NumPrimitiveIds = PrimitiveIdsAll.Num();
-	PassParameters->NumInstanceRuns = InstanceRunsAll.Num();
-	PassParameters->NumCommands = CullingCommandsAll.Num();
+	PassParameters->DrawIndirectArgsBufferOut = GraphBuilder.CreateUAV(BatchedCullingScratch.DrawIndirectArgsBuffer, PF_R32_UINT);
+	PassParameters->InstanceIdOffsetBufferOut = GraphBuilder.CreateUAV(BatchedCullingScratch.InstanceIdOffsetBuffer, PF_R32_UINT);
 	PassParameters->VisibleInstanceFlags = VisibleInstanceFlagsRDG ? GraphBuilder.CreateSRV(VisibleInstanceFlagsRDG) : nullptr;
 
 	PassParameters->NumInstanceFlagWords = uint32(NumInstanceFlagWords);
@@ -808,7 +846,7 @@ void FInstanceCullingContext::BuildRenderingCommandsBatched(FRDGBuilder& GraphBu
 
 	FBuildInstanceIdBufferAndCommandsFromPrimitiveIdsCs::FPermutationDomain PermutationVector;
 	PermutationVector.Set<FBuildInstanceIdBufferAndCommandsFromPrimitiveIdsCs::FOutputCommandIdDim>(0);
-	PermutationVector.Set<FBuildInstanceIdBufferAndCommandsFromPrimitiveIdsCs::FCullInstancesDim>(bCullInstances);
+	PermutationVector.Set<FBuildInstanceIdBufferAndCommandsFromPrimitiveIdsCs::FCullInstancesDim>(true);
 	PermutationVector.Set<FBuildInstanceIdBufferAndCommandsFromPrimitiveIdsCs::FBatchedDim>(true);
 
 	auto ComputeShader = GetGlobalShaderMap(GMaxRHIFeatureLevel)->GetShader<FBuildInstanceIdBufferAndCommandsFromPrimitiveIdsCs>(PermutationVector);
@@ -818,27 +856,22 @@ void FInstanceCullingContext::BuildRenderingCommandsBatched(FRDGBuilder& GraphBu
 		RDG_EVENT_NAME("BuildInstanceIdBufferAndCommandsFromPrimitiveIds"),
 		ComputeShader,
 		PassParameters,
-		FIntVector(CullingCommandsAll.Num(), 1, 1)
-	);
+		INST_CULL_CALLBACK(FIntVector(MergedContext.CullingCommandsAll.Num(), 1, 1)));
+	
 	FInstanceCullingGlobalUniforms* UniformParameters = GraphBuilder.AllocParameters<FInstanceCullingGlobalUniforms>();
 	UniformParameters->InstanceIdsBuffer = GraphBuilder.CreateSRV(InstanceIdsBuffer);
 	UniformParameters->PageInfoBuffer = GraphBuilder.CreateSRV(InstanceIdsBuffer);
-	UniformParameters->BufferCapacity = InstanceIdBufferSize;
-	TRDGUniformBufferRef<FInstanceCullingGlobalUniforms> UniformBuffer = GraphBuilder.CreateUniformBuffer(UniformParameters);
+	UniformParameters->BufferCapacity = 0U; // TODO: this is not used at the moment, but is intended for range checks so would have been good.
+	BatchedCullingScratch.UniformBuffer = GraphBuilder.CreateUniformBuffer(UniformParameters);
 
 
-	for (int32 BatchIndex = 0; BatchIndex < Batches.Num(); ++BatchIndex)
-	{
-		FBatchItem& BatchItem = Batches[BatchIndex];
-		FInstanceCullingResult& Result = *BatchItem.Result;
-
-		const FBatchInfo& BatchInfo = BatchInfos[BatchIndex];
-
-		Result.DrawIndirectArgsBuffer = DrawIndirectArgsRDG;
-		Result.InstanceIdOffsetBuffer = InstanceIdOffsetBufferRDG;
-		Result.DrawCommandDataOffset = BatchInfo.CullingCommandsOffset;
-		Result.UniformBuffer = UniformBuffer;
-	}	
+#undef INST_CULL_CREATE_STRUCT_BUFF_ARGS
+#undef INST_CULL_CALLBACK
 }
 
-
+bool FInstanceCullingContext::AllowBatchedBuildRenderingCommands(const FGPUScene& GPUScene)
+{
+	static IConsoleVariable* CVarRDGDrain = IConsoleManager::Get().FindConsoleVariable(TEXT("r.RDG.Drain"));
+	check(CVarRDGDrain);
+	return GPUScene.IsEnabled() && !!GAllowBatchedBuildRenderingCommands && !FRDGBuilder::IsImmediateMode() && !CVarRDGDrain->GetInt();
+}
diff --git a/Engine/Source/Runtime/Renderer/Private/InstanceCulling/InstanceCullingManager.cpp b/Engine/Source/Runtime/Renderer/Private/InstanceCulling/InstanceCullingManager.cpp
--- a/Engine/Source/Runtime/Renderer/Private/InstanceCulling/InstanceCullingManager.cpp
+++ b/Engine/Source/Runtime/Renderer/Private/InstanceCulling/InstanceCullingManager.cpp
@@ -153,3 +153,8 @@ void FInstanceCullingManager::CullInstances(FRDGBuilder& GraphBuilder, FGPUScene
 		}
 	}
 }
+
+void FInstanceCullingManager::BeginDeferredCulling(FRDGBuilder& GraphBuilder, FGPUScene& GPUScene)
+{
+	FInstanceCullingContext::BuildRenderingCommandsDeferred(GraphBuilder, GPUScene, *this);
+}
diff --git a/Engine/Source/Runtime/Renderer/Private/InstanceCulling/InstanceCullingManager.h b/Engine/Source/Runtime/Renderer/Private/InstanceCulling/InstanceCullingManager.h
--- a/Engine/Source/Runtime/Renderer/Private/InstanceCulling/InstanceCullingManager.h
+++ b/Engine/Source/Runtime/Renderer/Private/InstanceCulling/InstanceCullingManager.h
@@ -61,6 +61,22 @@ struct FInstanceCullingResult
 	}
 };
 
+struct FBatchedInstanceCullingScratchSpace
+{
+	/** Whether we defer and batch GPU instance culling work throughout a frame. */
+	bool bBatchingActive = false;
+
+	FRDGBufferRef DrawIndirectArgsBuffer = nullptr;
+	FRDGBufferRef InstanceIdOffsetBuffer = nullptr;
+	TRDGUniformBufferRef<FInstanceCullingGlobalUniforms> UniformBuffer = nullptr;
+
+	/** GPU instance culling input data merged from multiple batches throughout a frame. */
+	FInstanceCullingContextMerged MergedContext;
+
+	/** Batches of GPU instance culling input data. */
+	TArray<FInstanceCullingContext::FBatchItem, SceneRenderingAllocator> Batches;
+};
+
 /**
  * Manages allocation of indirect arguments and culling jobs for all instanced draws (use the GPU Scene culling).
  */
@@ -91,10 +107,22 @@ public:
 	}
 	const TRDGUniformBufferRef<FInstanceCullingGlobalUniforms> GetDummyInstanceCullingUniformBuffer() const { return CullingIntermediate.DummyUniformBuffer; }
 	
+	/** Starts batching GPU instance culling work items if possible. */
+	void BeginDeferredCulling(FRDGBuilder& GraphBuilder, FGPUScene& GPUScene);
+
+	/** Whether we are actively batching GPU instance culling work. */
+	bool IsDeferredCullingActive() const { return BatchedCullingScratch.bBatchingActive; }
+
 	// Populated by CullInstances, used when performing final culling & rendering 
 	FInstanceCullingIntermediate CullingIntermediate;
 
+
 private:
+	friend class FInstanceCullingContext;
+
+	// Polulated by FInstanceCullingContext::BuildRenderingCommandsDeferred, used to hold instance culling related data that needs to be passed around
+	FBatchedInstanceCullingScratchSpace BatchedCullingScratch;
+
 	FInstanceCullingManager() = delete;
 	FInstanceCullingManager(FInstanceCullingManager &) = delete;
 
diff --git a/Engine/Source/Runtime/Renderer/Private/MeshDrawCommands.cpp b/Engine/Source/Runtime/Renderer/Private/MeshDrawCommands.cpp
--- a/Engine/Source/Runtime/Renderer/Private/MeshDrawCommands.cpp
+++ b/Engine/Source/Runtime/Renderer/Private/MeshDrawCommands.cpp
@@ -1407,40 +1407,25 @@ public:
 	}
 };
 
-void FParallelMeshDrawCommandPass::BuildRenderingCommands(FRDGBuilder& GraphBuilder, FGPUScene& GPUScene, FInstanceCullingDrawParams& OutInstanceCullingDrawParams)
+void FParallelMeshDrawCommandPass::BuildRenderingCommands(
+	FRDGBuilder& GraphBuilder,
+	FGPUScene& GPUScene,
+	FInstanceCullingDrawParams& OutInstanceCullingDrawParams)
 {
 	if (TaskContext.InstanceCullingContext.IsEnabled())
 	{
-		if (bHasQueuedInstanceCullingBuild)
-		{
-			TaskContext.InstanceCullingResult.GetDrawParameters(OutInstanceCullingDrawParams);
-			return;
-		}
+		check(!bHasInstanceCullingDrawParameters);
 		WaitForMeshPassSetupTask();
-		// 2. Run finalize culling commands pass
-		TaskContext.InstanceCullingContext.BuildRenderingCommands(GraphBuilder, GPUScene, TaskContext.View->DynamicPrimitiveCollector.GetPrimitiveIdRange(), TaskContext.InstanceCullingResult);
+		// 2. Run or queue finalize culling commands pass
+		TaskContext.InstanceCullingContext.BuildRenderingCommands(GraphBuilder, GPUScene, TaskContext.View->DynamicPrimitiveCollector.GetPrimitiveIdRange(), TaskContext.InstanceCullingResult, &OutInstanceCullingDrawParams);
 		TaskContext.InstanceCullingResult.GetDrawParameters(OutInstanceCullingDrawParams);
+		bHasInstanceCullingDrawParameters = true;
+		check(!TaskContext.InstanceCullingContext.HasCullingCommands() || OutInstanceCullingDrawParams.DrawIndirectArgsBuffer && OutInstanceCullingDrawParams.InstanceIdOffsetBuffer);
 		return;
 	}
 	OutInstanceCullingDrawParams.DrawIndirectArgsBuffer = nullptr;
 	OutInstanceCullingDrawParams.InstanceIdOffsetBuffer = nullptr;
 	OutInstanceCullingDrawParams.DrawCommandDataOffset = 0U;
-
-}
-
-
-void FParallelMeshDrawCommandPass::QueueBatchedBuildRenderingCommands(TArray<FInstanceCullingContext::FBatchItem, SceneRenderingAllocator> & BatchItems)
-{
-	if (TaskContext.InstanceCullingContext.IsEnabled())
-	{
-		WaitForMeshPassSetupTask();
-		if (MaxNumDraws > 0 && TaskContext.InstanceCullingContext.HasCullingCommands())
-		{
-			// 2. Queue finalize culling commands pass
-			BatchItems.Add(FInstanceCullingContext::FBatchItem{ &TaskContext.InstanceCullingContext, &TaskContext.InstanceCullingResult, TaskContext.View->DynamicPrimitiveCollector.GetPrimitiveIdRange() });
-			bHasQueuedInstanceCullingBuild = true;
-		}
-	}
 }
 
 
diff --git a/Engine/Source/Runtime/Renderer/Private/MeshDrawCommands.h b/Engine/Source/Runtime/Renderer/Private/MeshDrawCommands.h
--- a/Engine/Source/Runtime/Renderer/Private/MeshDrawCommands.h
+++ b/Engine/Source/Runtime/Renderer/Private/MeshDrawCommands.h
@@ -146,7 +146,7 @@ public:
 
 	FParallelMeshDrawCommandPass()
 		: bPrimitiveIdBufferDataOwnedByRHIThread(false)
-		, bHasQueuedInstanceCullingBuild(false)
+		, bHasInstanceCullingDrawParameters(false)
 		, MaxNumDraws(0)
 	{
 	}
@@ -178,14 +178,10 @@ public:
 	 * Sync with setup task and run post-instance culling job to create the render commands and instance ID lists and optionally vertex instance data.
 	 * Needs to happen after DispatchPassSetup and before DispatchDraw, but not before global instance culling has been done.
 	 */
-	void BuildRenderingCommands(FRDGBuilder& GraphBuilder, FGPUScene& GPUScene, FInstanceCullingDrawParams& OutInstanceCullingDrawParams);
-
-
-	/**
-	 * Sync with setup task and return item to run batch processing on, sets a flag to signal batched processing has been requested (which is then used in BuildRenderingCommands to skip processing).
-	 * Needs to happen after DispatchPassSetup and before DispatchDraw, but not before global instance culling has been done.
-	 */
-	void QueueBatchedBuildRenderingCommands(TArray<FInstanceCullingContext::FBatchItem, SceneRenderingAllocator>& BatchItems);
+	void BuildRenderingCommands(
+		FRDGBuilder& GraphBuilder,
+		FGPUScene& GPUScene,
+		FInstanceCullingDrawParams& OutInstanceCullingDrawParams);
 
 	/**
 	 * Sync with setup task and run post-instance culling job unpack surviving instance IDs into two lists of instance ID + draw command ID.
@@ -224,7 +220,7 @@ private:
 
 	// If TaskContext::PrimitiveIdBufferData will be released by RHI Thread.
 	mutable bool bPrimitiveIdBufferDataOwnedByRHIThread;
-	bool bHasQueuedInstanceCullingBuild;
+	bool bHasInstanceCullingDrawParameters;
 
 	// Maximum number of draws for this pass. Used to prealocate resources on rendering thread. 
 	// Has a guarantee that if there won't be any draws, then MaxNumDraws = 0;
diff --git a/Engine/Source/Runtime/Renderer/Private/MobileShadingRenderer.cpp b/Engine/Source/Runtime/Renderer/Private/MobileShadingRenderer.cpp
--- a/Engine/Source/Runtime/Renderer/Private/MobileShadingRenderer.cpp
+++ b/Engine/Source/Runtime/Renderer/Private/MobileShadingRenderer.cpp
@@ -572,6 +572,7 @@ void FMobileSceneRenderer::InitViews(FRDGBuilder& GraphBuilder, FSceneTexturesCo
 	{
 		// GPUCULL_TODO: Possibly fold into unpack step
 		InstanceCullingManager.CullInstances(GraphBuilder, Scene->GPUScene);
+		InstanceCullingManager.BeginDeferredCulling(GraphBuilder, Scene->GPUScene);
 	}
 	extern TSet<IPersistentViewUniformBufferExtension*> PersistentViewUniformBufferExtensions;
 
@@ -694,7 +695,7 @@ void FMobileSceneRenderer::Render(FRDGBuilder& GraphBuilder)
 	GSystemTextures.InitializeTextures(GraphBuilder.RHICmdList, FeatureLevel);
 	FRDGSystemTextures::Create(GraphBuilder);
 
-	FInstanceCullingManager InstanceCullingManager(Scene->GPUScene.IsEnabled());
+	FInstanceCullingManager& InstanceCullingManager = *GraphBuilder.AllocObject<FInstanceCullingManager>(Scene->GPUScene.IsEnabled());
 
 	// Important that this uses consistent logic throughout the frame, so evaluate once and pass in the flag from here
 	// NOTE: Must be done after  system texture initialization
diff --git a/Engine/Source/Runtime/Renderer/Private/ShadowDepthRendering.cpp b/Engine/Source/Runtime/Renderer/Private/ShadowDepthRendering.cpp
--- a/Engine/Source/Runtime/Renderer/Private/ShadowDepthRendering.cpp
+++ b/Engine/Source/Runtime/Renderer/Private/ShadowDepthRendering.cpp
@@ -1942,31 +1942,6 @@ void FSceneRenderer::RenderShadowDepthMaps(FRDGBuilder& GraphBuilder, FInstanceC
 		VirtualShadowMapArray.RenderVirtualShadowMapsHw(GraphBuilder, SortedShadowsForShadowDepthPass.VirtualShadowMapShadows, *Scene);
 	}
 
-	// Batch all culling commands to minimize transition stalls and improve throughput etc
-	{
-		TArray<FInstanceCullingContext::FBatchItem, SceneRenderingAllocator> InstanceCullingBatches;
-		for (FSortedShadowMapAtlas& ShadowMapAtlas : SortedShadowsForShadowDepthPass.ShadowMapAtlases)
-		{
-			for (FProjectedShadowInfo* ProjectedShadowInfo : ShadowMapAtlas.Shadows)
-			{
-				ProjectedShadowInfo->QueueBatchedBuildRenderingCommands(InstanceCullingBatches);
-			}
-		}
-		for (FSortedShadowMapAtlas& ShadowMap : SortedShadowsForShadowDepthPass.ShadowMapCubemaps)
-		{
-			FProjectedShadowInfo* ProjectedShadowInfo = ShadowMap.Shadows[0];
-			ProjectedShadowInfo->QueueBatchedBuildRenderingCommands(InstanceCullingBatches);
-		}
-		for (FProjectedShadowInfo* ProjectedShadowInfo : SortedShadowsForShadowDepthPass.PreshadowCache.Shadows)
-		{
-			if (!ProjectedShadowInfo->bDepthsCached)
-			{
-				ProjectedShadowInfo->QueueBatchedBuildRenderingCommands(InstanceCullingBatches);
-			}
-		}
-		FInstanceCullingContext::BuildRenderingCommandsBatched(GraphBuilder, Scene->GPUScene, InstanceCullingBatches);
-	}
-
 	VirtualShadowMapArray.SetupProjectionParameters(GraphBuilder);
 
 	// Render non-VSM shadows
diff --git a/Engine/Source/Runtime/Renderer/Private/ShadowRendering.h b/Engine/Source/Runtime/Renderer/Private/ShadowRendering.h
--- a/Engine/Source/Runtime/Renderer/Private/ShadowRendering.h
+++ b/Engine/Source/Runtime/Renderer/Private/ShadowRendering.h
@@ -698,10 +698,13 @@ public:
 
 	float GetMaxNonFarCascadeDistance() const { return MaxNonFarCascadeDistance; }
 
-	void QueueBatchedBuildRenderingCommands(TArray<FInstanceCullingContext::FBatchItem, SceneRenderingAllocator>& BatchItems)
+	void BuildRenderingCommands(
+		FRDGBuilder& GraphBuilder,
+		FGPUScene& GPUScene,
+		FInstanceCullingDrawParams& InstanceCullingDrawParams)
 	{
 		ShadowDepthView->DynamicPrimitiveCollector.Commit();
-		return ShadowDepthPass.QueueBatchedBuildRenderingCommands(BatchItems);
+		return ShadowDepthPass.BuildRenderingCommands(GraphBuilder, GPUScene, InstanceCullingDrawParams);
 	}
 
 private:
diff --git a/Engine/Source/Runtime/Renderer/Public/InstanceCulling/InstanceCullingContext.h b/Engine/Source/Runtime/Renderer/Public/InstanceCulling/InstanceCullingContext.h
--- a/Engine/Source/Runtime/Renderer/Public/InstanceCulling/InstanceCullingContext.h
+++ b/Engine/Source/Runtime/Renderer/Public/InstanceCulling/InstanceCullingContext.h
@@ -96,6 +96,27 @@ public:
 		uint32 bMaterialMayModifyPosition;
 	};
 
+	struct FBatchInfo
+	{
+		uint32 CullingCommandsOffset;
+		uint32 NumCullingCommands;
+		uint32 PrimitiveIdsOffset;
+		uint32 NumPrimitiveIds;
+		uint32 InstanceRunsOffset;
+		uint32 NumInstanceRuns;
+		uint32 ViewIdsOffset;
+		uint32 NumViewIds;
+		int32 DynamicPrimitiveIdOffset;
+		int32 DynamicPrimitiveIdMax;
+	};
+
+	struct FBatchItem
+	{
+		const FInstanceCullingContext* Context = nullptr;
+		FInstanceCullingDrawParams* Result = nullptr;
+		TRange<int32> DynamicPrimitiveIdRange;
+	};
+
 	/**
 	 * Call to empty out the culling commands & other culling data.
 	 */
@@ -120,27 +141,35 @@ public:
 	 */
 	void AddInstanceRunToCullingCommand(int32 ScenePrimitiveId, const uint32* Runs, uint32 NumRuns);
 
-
-	void BuildRenderingCommands(FRDGBuilder& GraphBuilder, const FGPUScene& GPUScene, const TRange<int32> &DynamicPrimitiveIdRange, FInstanceCullingResult& Results) const;
+	/**
+	 * If InstanceCullingDrawParams is not null, this BuildRenderingCommands operation will be deferred and merged into a global pass when possible.
+	 */
+	void BuildRenderingCommands(
+		FRDGBuilder& GraphBuilder,
+		const FGPUScene& GPUScene,
+		const TRange<int32> &DynamicPrimitiveIdRange,
+		FInstanceCullingResult& Results,
+		FInstanceCullingDrawParams* InstanceCullingDrawParams = nullptr) const;
 
 	/**
 	 * Build instance ID lists & render commands, but without extracting results and also accepting RGD resources for output buffers etc (rather than registering global resources)
 	 */
 	void BuildRenderingCommands(FRDGBuilder& GraphBuilder, FGPUScene& GPUScene, const TRange<int32>& DynamicPrimitiveIdRange, FInstanceCullingRdgParams& Params) const;
 
-	inline bool HasCullingCommands() const { return CullingCommands.Num() > 0; 	}
+	inline bool HasCullingCommands() const { return CullingCommands.Num() > 0 && TotalInstances > 0; 	}
 
 	EInstanceCullingMode GetInstanceCullingMode() const { return InstanceCullingMode; }
 
+	/**
+	 * Add a batched BuildRenderingCommands pass. Each batch represents a BuildRenderingCommands call from a mesh pass.
+	 * Batches are gradually added as we walk through the main render function.
+	 */
+	static void BuildRenderingCommandsDeferred(
+		FRDGBuilder& GraphBuilder,
+		FGPUScene& GPUScene,
+		FInstanceCullingManager& InstanceCullingManager);
 
-	struct FBatchItem
-	{
-		FInstanceCullingContext* Context = nullptr;
-		FInstanceCullingResult* Result = nullptr;
-		TRange<int32> DynamicPrimitiveIdRange;
-	};
-
-	static void BuildRenderingCommandsBatched(FRDGBuilder& GraphBuilder, FGPUScene& GPUScene, TArrayView<FBatchItem> Batches);
+	static bool AllowBatchedBuildRenderingCommands(const FGPUScene& GPUScene);
 
 	// GPUCULL_TODO: These should not be dynamically heap-allocated, all except instance runs are easy to pre-size on memstack.
 	//           Must be presized as populated from task threads.
@@ -164,3 +193,15 @@ public:
 	uint32 TotalInstances = 0U;
 };
 
+struct FInstanceCullingContextMerged
+{
+	bool bInitialized = false;
+	TArray<FInstanceCullingContext::FPrimCullingCommand, SceneRenderingAllocator> CullingCommandsAll;
+	TArray<int32, SceneRenderingAllocator> PrimitiveIdsAll;
+	TArray<FInstanceCullingContext::FInstanceRun, SceneRenderingAllocator> InstanceRunsAll;
+	TArray<int32, SceneRenderingAllocator> ViewIdsAll;
+	TArray<FInstanceCullingContext::FBatchInfo, SceneRenderingAllocator> BatchInfos;
+	TArray<int32, SceneRenderingAllocator> BatchIndsAll;
+
+	uint32 InstanceIdBufferSize = 0U;
+};
