Add 64 group size support for hair strands voxelization.

diff --git a/Engine/Plugins/Runtime/HairStrands/Source/HairStrandsCore/Private/HairStrandsRendering.cpp b/Engine/Plugins/Runtime/HairStrands/Source/HairStrandsCore/Private/HairStrandsRendering.cpp
--- a/Engine/Plugins/Runtime/HairStrands/Source/HairStrandsCore/Private/HairStrandsRendering.cpp
+++ b/Engine/Plugins/Runtime/HairStrands/Source/HairStrandsCore/Private/HairStrandsRendering.cpp
@@ -60,18 +60,18 @@ bool NeedsUpdateCardsMeshTriangles()
 
 static FIntVector ComputeDispatchCount(uint32 ItemCount, uint32 GroupSize)
 {
-	const uint32 BatchCount = FMath::DivideAndRoundUp(ItemCount, GroupSize);
-	const uint32 DispatchCountX = FMath::FloorToInt(FMath::Sqrt(static_cast<float>(BatchCount)));
-	const uint32 DispatchCountY = DispatchCountX + FMath::DivideAndRoundUp(BatchCount - DispatchCountX * DispatchCountX, DispatchCountX);
+	const uint32 GroupCount = FMath::DivideAndRoundUp(ItemCount, GroupSize);
+	const uint32 DispatchCountX = FMath::FloorToInt(FMath::Sqrt(static_cast<float>(GroupCount)));
+	const uint32 DispatchCountY = DispatchCountX + FMath::DivideAndRoundUp(GroupCount - DispatchCountX * DispatchCountX, DispatchCountX);
 
 	check(DispatchCountX <= 65535);
 	check(DispatchCountY <= 65535);
-	check(BatchCount <= DispatchCountX * DispatchCountY);
+	check(GroupCount <= DispatchCountX * DispatchCountY);
 	return FIntVector(DispatchCountX, DispatchCountY, 1);
 }
 
 // Same as above but the group count is what matters and is preserved
-static FIntVector ComputeDispatchGroupCount2D(uint32 GroupCount)
+static FIntVector ComputeDispatchCount(uint32 GroupCount)
 {
 	const uint32 DispatchCountX = FMath::FloorToInt(FMath::Sqrt(static_cast<float>(GroupCount)));
 	const uint32 DispatchCountY = DispatchCountX + FMath::DivideAndRoundUp(GroupCount - DispatchCountX * DispatchCountX, DispatchCountX);
@@ -817,7 +817,7 @@ static void AddHairClusterAABBPass(
 	// * If the instance has deformatin (simulation, global interpolation, skinning, then we update all clusters (for voxelization purpose) and the instance AABB
 	// * If the instance is static, we only update the instance AABB
 	const uint32 GroupSize = ComputeGroupSize();
-	const FIntVector DispatchCount = (UpdateType == EHairAABBUpdateType::UpdateClusterAABB && ClusterData) ? ComputeDispatchGroupCount2D(ClusterData->ClusterCount) : FIntVector(1,1,1);
+	const FIntVector DispatchCount = (UpdateType == EHairAABBUpdateType::UpdateClusterAABB && ClusterData) ? ComputeDispatchCount(ClusterData->ClusterCount) : FIntVector(1,1,1);
 
 	FHairClusterAABBCS::FParameters* Parameters = GraphBuilder.AllocParameters<FHairClusterAABBCS::FParameters>();
 	Parameters->CPUBoundMin = TransformedBounds.GetBox().Min;
diff --git a/Engine/Shaders/Private/HairStrands/HairStrandsClusterCulling.usf b/Engine/Shaders/Private/HairStrands/HairStrandsClusterCulling.usf
--- a/Engine/Shaders/Private/HairStrands/HairStrandsClusterCulling.usf
+++ b/Engine/Shaders/Private/HairStrands/HairStrandsClusterCulling.usf
@@ -640,6 +640,14 @@ void MainClusterCullingCS(uint3 DispatchThreadId : SV_DispatchThreadID)
 RWBuffer<uint> DrawIndirectBuffer; 
 RWBuffer<uint> DispatchIndirectBuffer; 
 
+#if PERMUTATION_GROUP_SIZE == 64
+	#define GROUP_SIZE 64
+#elif PERMUTATION_GROUP_SIZE == 32
+	#define GROUP_SIZE 32
+#else
+	#error Unknown group size
+#endif
+
 [numthreads(1, 1, 1)]
 void MainClusterCullingPrepareIndirectDispatchCS(
 	uint3 DispatchThreadId : SV_DispatchThreadID, 
@@ -648,13 +656,10 @@ void MainClusterCullingPrepareIndirectDispatchCS(
 {
 	if (GroupId.x == 0 && GroupThreadId.x == 0)
 	{
-		// See HairStrandsVisibilityRasterCompute.usf
-		const uint RasterComputeGroupSize = 32; // TODO change this so that we can have dispatch for different GroupSize
-
 		// DrawIndirectBuffer contains the total number of vertex count needed for drawing quad triangle strips, so 6 vertices per quad / control point
 		const uint VertexCount    = DrawIndirectBuffer[0];
 		const uint ControlPoint   = VertexCount / 6;
-		const uint DispatchCount  = DivideAndRoundUp(ControlPoint, RasterComputeGroupSize);
+		const uint DispatchCount  = DivideAndRoundUp(ControlPoint, GROUP_SIZE);
 		const uint DispatchCountX = DispatchCount >= INDIRECT_GROUP_COUNT_X ? INDIRECT_GROUP_COUNT_X : (DispatchCount % INDIRECT_GROUP_COUNT_X);
 		const uint DispatchCountY = DivideAndRoundUp(DispatchCount, INDIRECT_GROUP_COUNT_X);
 
diff --git a/Engine/Shaders/Private/HairStrands/HairStrandsVisibilityRasterCompute.usf b/Engine/Shaders/Private/HairStrands/HairStrandsVisibilityRasterCompute.usf
--- a/Engine/Shaders/Private/HairStrands/HairStrandsVisibilityRasterCompute.usf
+++ b/Engine/Shaders/Private/HairStrands/HairStrandsVisibilityRasterCompute.usf
@@ -13,6 +13,15 @@
 #define RASTER_AMD_DX11 2
 #define RASTER_AMD_DX12 3
 
+// Group size
+#if PERMUTATION_GROUP_SIZE == 64
+	#define GROUP_SIZE 64
+#elif PERMUTATION_GROUP_SIZE == 32
+	#define GROUP_SIZE 32
+#else
+	#error Unknown group size	
+#endif
+
 // Add permutation for AMD/other
 #if PERMUTATION_RASTER_ATOMIC == RASTER_NVIDIA
 #define NV_SHADER_EXTN_SLOT u7
@@ -261,11 +270,6 @@ uint2 Rasterize(uint2 PixelCoord, bool bIsSteep, float FracY, float RatioScale,
 	return SeedCoord;
 }
 //////////////////////////////////////////////////////////////////////////////////////////
-
-// These number need to match HairStrandsClusterCulling.usf
-// TODO
-#define GROUP_SIZE 32
-
 [numthreads(GROUP_SIZE, 1, 1)]
 void MainCS(uint2 DispatchThreadID : SV_DispatchThreadID)
 {
diff --git a/Engine/Shaders/Private/HairStrands/HairStrandsVoxelRasterCompute.usf b/Engine/Shaders/Private/HairStrands/HairStrandsVoxelRasterCompute.usf
--- a/Engine/Shaders/Private/HairStrands/HairStrandsVoxelRasterCompute.usf
+++ b/Engine/Shaders/Private/HairStrands/HairStrandsVoxelRasterCompute.usf
@@ -8,6 +8,7 @@
 #include "HairStrandsVoxelPageCommon.ush" 
 
 ///////////////////////////////////////////////////////////////////////////
+
 uint			DispatchCountX;
 uint			MaxRasterCount;
 uint			FrameIdMod8;
@@ -42,7 +43,20 @@ StructuredBuffer<FPackedVirtualVoxelNodeDesc> VirtualVoxelParams_NodeDescBuffer;
 RWTexture3D<uint> OutPageTexture;
 
 ///////////////////////////////////////////////////////////////////////////
+// Define
+
+#define JITTER_ENABLE 0
+
+#if PERMUTATION_GROUP_SIZE == 64
+	#define GROUP_SIZE 64
+#elif PERMUTATION_GROUP_SIZE == 32
+	#define GROUP_SIZE 32
+#else
+	#error Unknown group size	
+#endif
 
+///////////////////////////////////////////////////////////////////////////
+// Utils
 float3 GetHairVoxelJitter(uint2 PixelCoord, uint Seed)
 {
 	return float3(
@@ -51,11 +65,7 @@ float3 GetHairVoxelJitter(uint2 PixelCoord, uint Seed)
 		InterleavedGradientNoise(PixelCoord.xy, Seed * 7901));
 }
 
-#define JITTER_ENABLE 0
-
-// These number need to match HairStrandsClusterCulling.usf
-// TODO: create permutation for different group size
-#define GROUP_SIZE 32
+///////////////////////////////////////////////////////////////////////////
 
 [numthreads(GROUP_SIZE, 1, 1)]
 void MainCS(uint2 DispatchThreadID : SV_DispatchThreadID)
diff --git a/Engine/Source/Runtime/Renderer/Private/HairStrands/HairStrandsClusters.cpp b/Engine/Source/Runtime/Renderer/Private/HairStrands/HairStrandsClusters.cpp
--- a/Engine/Source/Runtime/Renderer/Private/HairStrands/HairStrandsClusters.cpp
+++ b/Engine/Source/Runtime/Renderer/Private/HairStrands/HairStrandsClusters.cpp
@@ -128,6 +128,9 @@ class FMainClusterCullingPrepareIndirectDispatchCS : public FGlobalShader
 	DECLARE_GLOBAL_SHADER(FMainClusterCullingPrepareIndirectDispatchCS);
 	SHADER_USE_PARAMETER_STRUCT(FMainClusterCullingPrepareIndirectDispatchCS, FGlobalShader);
 
+	class FGroupSize : SHADER_PERMUTATION_SPARSE_INT("PERMUTATION_GROUP_SIZE", 32, 64);
+	using FPermutationDomain = TShaderPermutationDomain<FGroupSize>;
+
 	BEGIN_SHADER_PARAMETER_STRUCT(FParameters, )
 		SHADER_PARAMETER_RDG_BUFFER_UAV(RWBuffer, DrawIndirectBuffer)
 		SHADER_PARAMETER_RDG_BUFFER_UAV(RWBuffer, DispatchIndirectBuffer)
@@ -471,7 +474,9 @@ static void AddClusterCullingPass(
 		Parameters->DrawIndirectBuffer = DrawIndirectParametersBuffer.UAV;
 		Parameters->DispatchIndirectBuffer = DrawIndirectParametersRasterComputeBuffer.UAV;
 		
-		TShaderMapRef<FMainClusterCullingPrepareIndirectDispatchCS> ComputeShader(ShaderMap);
+		FMainClusterCullingPrepareIndirectDispatchCS::FPermutationDomain PermutationVector;
+		PermutationVector.Set<FMainClusterCullingPrepareIndirectDispatchCS::FGroupSize>(GetVendorOptimalGroupSize1D());
+		TShaderMapRef<FMainClusterCullingPrepareIndirectDispatchCS> ComputeShader(ShaderMap, PermutationVector);
 		FComputeShaderUtils::AddPass(
 			GraphBuilder,
 			RDG_EVENT_NAME("PrepareIndirectDispatchCS"),
diff --git a/Engine/Source/Runtime/Renderer/Private/HairStrands/HairStrandsUtils.cpp b/Engine/Source/Runtime/Renderer/Private/HairStrands/HairStrandsUtils.cpp
--- a/Engine/Source/Runtime/Renderer/Private/HairStrands/HairStrandsUtils.cpp
+++ b/Engine/Source/Runtime/Renderer/Private/HairStrands/HairStrandsUtils.cpp
@@ -273,7 +273,15 @@ bool IsHairStrandsViewRectOptimEnable()
 	return GHairVisibilityRectOptimEnable > 0;
 }
 
-EHairVisibilityVendor GetVendor()
+enum EHairVisibilityVendor
+{
+	HairVisibilityVendor_AMD,
+	HairVisibilityVendor_NVIDIA,
+	HairVisibilityVendor_INTEL,
+	HairVisibilityVendorCount
+};
+
+inline EHairVisibilityVendor GetVendor()
 {
 	return IsRHIDeviceAMD() ? HairVisibilityVendor_AMD : (IsRHIDeviceNVIDIA() ? HairVisibilityVendor_NVIDIA : HairVisibilityVendor_INTEL);
 }
@@ -300,6 +308,29 @@ FIntPoint GetVendorOptimalGroupSize2D()
 	}
 }
 
+FIntVector ComputeDispatchCount(uint32 ItemCount, uint32 GroupSize)
+{
+	const uint32 GroupCount = FMath::DivideAndRoundUp(ItemCount, GroupSize);
+	const uint32 DispatchCountX = FMath::FloorToInt(FMath::Sqrt(static_cast<float>(GroupCount)));
+	const uint32 DispatchCountY = DispatchCountX + FMath::DivideAndRoundUp(GroupCount - DispatchCountX * DispatchCountX, DispatchCountX);
+
+	check(DispatchCountX <= 65535);
+	check(DispatchCountY <= 65535);
+	check(GroupCount <= DispatchCountX * DispatchCountY);
+	return FIntVector(DispatchCountX, DispatchCountY, 1);
+}
+
+FIntVector ComputeDispatchCount(uint32 GroupCount)
+{
+	const uint32 DispatchCountX = FMath::FloorToInt(FMath::Sqrt(static_cast<float>(GroupCount)));
+	const uint32 DispatchCountY = DispatchCountX + FMath::DivideAndRoundUp(GroupCount - DispatchCountX * DispatchCountX, DispatchCountX);
+
+	check(DispatchCountX <= 65535);
+	check(DispatchCountY <= 65535);
+	check(GroupCount <= DispatchCountX * DispatchCountY);
+	return FIntVector(DispatchCountX, DispatchCountY, 1);
+}
+
 FVector4 PackHairRenderInfo(
 	float PrimaryRadiusAtDepth1,
 	float StableRadiusAtDepth1,
diff --git a/Engine/Source/Runtime/Renderer/Private/HairStrands/HairStrandsUtils.h b/Engine/Source/Runtime/Renderer/Private/HairStrands/HairStrandsUtils.h
--- a/Engine/Source/Runtime/Renderer/Private/HairStrands/HairStrandsUtils.h
+++ b/Engine/Source/Runtime/Renderer/Private/HairStrands/HairStrandsUtils.h
@@ -50,18 +50,12 @@ FIntRect ComputeVisibleHairStrandsMacroGroupsRect(const FIntRect& ViewRect, cons
 
 bool IsHairStrandsViewRectOptimEnable();
 
-enum EHairVisibilityVendor
-{
-	HairVisibilityVendor_AMD,
-	HairVisibilityVendor_NVIDIA,
-	HairVisibilityVendor_INTEL,
-	HairVisibilityVendorCount
-};
-
-EHairVisibilityVendor GetVendor();
 uint32 GetVendorOptimalGroupSize1D();
 FIntPoint GetVendorOptimalGroupSize2D();
 
+FIntVector ComputeDispatchCount(uint32 ItemCount, uint32 GroupSize);
+FIntVector ComputeDispatchCount(uint32 GroupCount);
+
 enum class  EHairStrandsCompositionType : uint8
 {
 	BeforeTranslucent,
diff --git a/Engine/Source/Runtime/Renderer/Private/HairStrands/HairStrandsVisibility.cpp b/Engine/Source/Runtime/Renderer/Private/HairStrands/HairStrandsVisibility.cpp
--- a/Engine/Source/Runtime/Renderer/Private/HairStrands/HairStrandsVisibility.cpp
+++ b/Engine/Source/Runtime/Renderer/Private/HairStrands/HairStrandsVisibility.cpp
@@ -3010,11 +3010,12 @@ class FVisiblityRasterComputeCS : public FGlobalShader
 	DECLARE_GLOBAL_SHADER(FVisiblityRasterComputeCS);
 	SHADER_USE_PARAMETER_STRUCT(FVisiblityRasterComputeCS, FGlobalShader);
 
+	class FGroupSize : SHADER_PERMUTATION_SPARSE_INT("PERMUTATION_GROUP_SIZE", 32, 64);
 	class FRasterAtomic : SHADER_PERMUTATION_INT("PERMUTATION_RASTER_ATOMIC", 4);
 	class FSPP : SHADER_PERMUTATION_SPARSE_INT("PERMUTATION_SPP", 1, 2, 4); 
 	class FCulling : SHADER_PERMUTATION_BOOL("PERMUTATION_CULLING");
 	class FStochastic : SHADER_PERMUTATION_BOOL("PERMUTATION_STOCHASTIC");
-	using FPermutationDomain = TShaderPermutationDomain<FRasterAtomic, FSPP, FCulling, FStochastic>;
+	using FPermutationDomain = TShaderPermutationDomain<FRasterAtomic, FSPP, FCulling, FStochastic, FGroupSize>;
 
 	BEGIN_SHADER_PARAMETER_STRUCT(FParameters, )
 		SHADER_PARAMETER(uint32, MacroGroupId)
@@ -3158,8 +3159,7 @@ static FRasterComputeOutput AddVisibilityComputeRasterPass(
 	ViewUniformShaderParameters = TUniformBufferRef<FViewUniformShaderParameters>::CreateUniformBufferImmediate(*ViewInfo.CachedViewUniformShaderParameters, UniformBuffer_SingleFrame);
 
 	const uint32 FrameIdMode8 = ViewInfo.ViewState ? (ViewInfo.ViewState->GetFrameIndex() % 8) : 0;
-	const uint32 GroupSize = 32;
-	const uint32 DispatchCountX = 64;
+	const uint32 GroupSize = GetVendorOptimalGroupSize1D();
 
 	FVisiblityRasterComputeCS::FPermutationDomain PermutationVector0;
 	FVisiblityRasterComputeCS::FPermutationDomain PermutationVector1;
@@ -3180,6 +3180,7 @@ static FRasterComputeOutput AddVisibilityComputeRasterPass(
 #endif
 	PermutationVector0.Set<FVisiblityRasterComputeCS::FStochastic>(bStochasticRaster);
 	PermutationVector0.Set<FVisiblityRasterComputeCS::FSPP>(SamplePerPixelCount);
+	PermutationVector0.Set<FVisiblityRasterComputeCS::FGroupSize>(GroupSize);
 	PermutationVector1 = PermutationVector0; 
 
 	PermutationVector0.Set<FVisiblityRasterComputeCS::FCulling>(false);
@@ -3197,7 +3198,7 @@ static FRasterComputeOutput AddVisibilityComputeRasterPass(
 			PassParameters->OutputResolution = Out.SuperResolution;
 			PassParameters->ResolutionMultiplier = Out.ResolutionMultiplier;
 			PassParameters->MacroGroupId = MacroGroup.MacroGroupId;
-			PassParameters->DispatchCountX = DispatchCountX;
+			PassParameters->DispatchCountX = 1;
 			PassParameters->MaxRasterCount = FMath::Clamp(GHairVisibilityComputeRaster_MaxPixelCount, 1, 256);			
 			PassParameters->FrameIdMod8 = FrameIdMode8;
 			PassParameters->HairMaterialId = PrimitiveInfo.MaterialId;
@@ -3239,8 +3240,8 @@ static FRasterComputeOutput AddVisibilityComputeRasterPass(
 			}
 			else
 			{
-				const uint32 DispatchCountY = FMath::CeilToInt(PassParameters->HairStrandsVF_VertexCount / float(GroupSize * DispatchCountX));
-				const FIntVector DispatchCount(DispatchCountX, DispatchCountY, 1);
+				const FIntVector DispatchCount = ComputeDispatchCount(PassParameters->HairStrandsVF_VertexCount, GroupSize);
+				PassParameters->DispatchCountX = DispatchCount.X;
 				FComputeShaderUtils::AddPass(GraphBuilder, RDG_EVENT_NAME("HairStrandsVisibilityComputeRaster(culling=off)"), ComputeShader_CullingOff, PassParameters, DispatchCount);
 			}
 		}
diff --git a/Engine/Source/Runtime/Renderer/Private/HairStrands/HairStrandsVoxelization.cpp b/Engine/Source/Runtime/Renderer/Private/HairStrands/HairStrandsVoxelization.cpp
--- a/Engine/Source/Runtime/Renderer/Private/HairStrands/HairStrandsVoxelization.cpp
+++ b/Engine/Source/Runtime/Renderer/Private/HairStrands/HairStrandsVoxelization.cpp
@@ -1190,8 +1190,9 @@ class FVoxelRasterComputeCS : public FGlobalShader
 	DECLARE_GLOBAL_SHADER(FVoxelRasterComputeCS);
 	SHADER_USE_PARAMETER_STRUCT(FVoxelRasterComputeCS, FGlobalShader);
 
+	class FGroupSize : SHADER_PERMUTATION_SPARSE_INT("PERMUTATION_GROUP_SIZE", 32, 64);
 	class FCulling : SHADER_PERMUTATION_INT("PERMUTATION_CULLING", 2);
-	using FPermutationDomain = TShaderPermutationDomain<FCulling>;
+	using FPermutationDomain = TShaderPermutationDomain<FGroupSize, FCulling>;
 
 	BEGIN_SHADER_PARAMETER_STRUCT(FParameters, )
 		SHADER_PARAMETER_STRUCT(FHairStrandsVoxelCommonParameters, VirtualVoxelParams)
@@ -1246,13 +1247,14 @@ static void AddVirtualVoxelizationComputeRasterPass(
 		FRDGTextureUAVRef PageTextureUAV = GraphBuilder.CreateUAV(VoxelResources.PageTexture);
 
 		const uint32 FrameIdMode8 = ViewInfo && ViewInfo->ViewState ? (ViewInfo->ViewState->GetFrameIndex() % 8) : 0;
-		const uint32 GroupSize = 32;
-		const uint32 DispatchCountX = 64;
+		const uint32 GroupSize = GetVendorOptimalGroupSize1D();
 
 		FVoxelRasterComputeCS::FPermutationDomain PermutationVector_Off;
 		FVoxelRasterComputeCS::FPermutationDomain PermutationVector_On;
 		PermutationVector_Off.Set<FVoxelRasterComputeCS::FCulling>(0);
+		PermutationVector_Off.Set<FVoxelRasterComputeCS::FGroupSize>(GroupSize);
 		PermutationVector_On.Set<FVoxelRasterComputeCS::FCulling>(1);
+		PermutationVector_On.Set<FVoxelRasterComputeCS::FGroupSize>(GroupSize);
 
 		TShaderMapRef<FVoxelRasterComputeCS> ComputeShader_CullingOff(ViewInfo->ShaderMap, PermutationVector_Off);
 		TShaderMapRef<FVoxelRasterComputeCS> ComputeShader_CullingOn(ViewInfo->ShaderMap, PermutationVector_On);
@@ -1274,7 +1276,7 @@ static void AddVirtualVoxelizationComputeRasterPass(
 				PassParameters->VirtualVoxelParams = VoxelResources.Parameters.Common;
 				PassParameters->MacroGroupId = MacroGroup.MacroGroupId;
 				PassParameters->VoxelizationViewInfoBuffer = VoxelizationViewInfoBufferSRV;
-				PassParameters->DispatchCountX = DispatchCountX;
+				PassParameters->DispatchCountX = 1;
 				PassParameters->OutPageTexture = PageTextureUAV;
 				PassParameters->FrameIdMod8 = FrameIdMode8;
 
@@ -1303,11 +1305,10 @@ static void AddVirtualVoxelizationComputeRasterPass(
 				}
 				else
 				{
-					const uint32 DispatchCountY = FMath::CeilToInt(PassParameters->HairStrandsVF_VertexCount / float(GroupSize * DispatchCountX));
-					const FIntVector DispatchCount(DispatchCountX, DispatchCountY, 1);
+					const FIntVector DispatchCount = ComputeDispatchCount(PassParameters->HairStrandsVF_VertexCount, GroupSize);
+					PassParameters->DispatchCountX = DispatchCount.X;
 					FComputeShaderUtils::AddPass(GraphBuilder, RDG_EVENT_NAME("HairStrandsVoxelComputeRaster(culling=off)"), ComputeShader_CullingOff, PassParameters, DispatchCount);
 				}
-
 			}
 		}
 	}
