Increase limit of instances possible in CullInstances to 16M

diff --git a/Engine/Shaders/Private/InstanceCulling/CullInstances.usf b/Engine/Shaders/Private/InstanceCulling/CullInstances.usf
--- a/Engine/Shaders/Private/InstanceCulling/CullInstances.usf
+++ b/Engine/Shaders/Private/InstanceCulling/CullInstances.usf
@@ -5,7 +5,6 @@
 #include "../Nanite/NaniteDataDecode.ush"
 #include "../Nanite/HZBCull.ush"
 
-//StructuredBuffer<FPackedNaniteView> Views;
 RWStructuredBuffer<uint> InstanceVisibilityFlagsOut;
 uint NumInstances;
 uint NumInstanceFlagWords;
diff --git a/Engine/Source/Runtime/Renderer/Private/InstanceCulling/InstanceCullingManager.cpp b/Engine/Source/Runtime/Renderer/Private/InstanceCulling/InstanceCullingManager.cpp
--- a/Engine/Source/Runtime/Renderer/Private/InstanceCulling/InstanceCullingManager.cpp
+++ b/Engine/Source/Runtime/Renderer/Private/InstanceCulling/InstanceCullingManager.cpp
@@ -59,7 +59,7 @@ class FCullInstancesCs : public FGlobalShader
 	SHADER_USE_PARAMETER_STRUCT(FCullInstancesCs, FGlobalShader)
 public:
 
-	static constexpr int32 NumThreadsPerGroup = 64;
+	static constexpr int32 NumThreadsPerGroup = 256;
 
 	static bool ShouldCompilePermutation(const FGlobalShaderPermutationParameters& Parameters)
 	{
