1. Fix a bug of missing velocity due to wrong LoadAction for parallel rendering.
2. Add a project setting option to avoid generating velocity and TAA shader unexpectly for mobile platform.

diff --git a/Engine/Source/Runtime/Core/Private/HAL/ConsoleManager.cpp b/Engine/Source/Runtime/Core/Private/HAL/ConsoleManager.cpp
--- a/Engine/Source/Runtime/Core/Private/HAL/ConsoleManager.cpp
+++ b/Engine/Source/Runtime/Core/Private/HAL/ConsoleManager.cpp
@@ -2873,6 +2873,14 @@ static TAutoConsoleVariable<int32> CVarMobileDefaultAntiAliasing(
 	TEXT("Mobile default AntiAliasingMethod\n")
 	TEXT(" 0: off (no anti-aliasing)\n")
 	TEXT(" 1: FXAA (default, faster than TemporalAA but much more shimmering for non static cases)\n")
-	TEXT(" 2: TemporalAA\n")
+	TEXT(" 2: TemporalAA(it will fallback to FXAA if SupportsGen4TAA is disabled) \n")
 	TEXT(" 3: MSAA"),
-	ECVF_RenderThreadSafe);
\ No newline at end of file
+	ECVF_RenderThreadSafe);
+
+static TAutoConsoleVariable<int32> CVarMobileSupportsGen4TAA(
+	TEXT("r.Mobile.SupportsGen4TAA"),
+	1,
+	TEXT("Support desktop Gen4 TAA with mobile rendering\n")
+	TEXT("0: Fallback to FXAA")
+	TEXT("1: Support Desktop Gen4 TAA (default)"),
+	ECVF_ReadOnly | ECVF_RenderThreadSafe);
\ No newline at end of file
diff --git a/Engine/Source/Runtime/Engine/Classes/Engine/RendererSettings.h b/Engine/Source/Runtime/Engine/Classes/Engine/RendererSettings.h
--- a/Engine/Source/Runtime/Engine/Classes/Engine/RendererSettings.h
+++ b/Engine/Source/Runtime/Engine/Classes/Engine/RendererSettings.h
@@ -903,6 +903,12 @@ class ENGINE_API URendererSettings : public UDeveloperSettings
 		ConfigRestartRequired = true))
 		TEnumAsByte<EMobilePlanarReflectionMode::Type> MobilePlanarReflectionMode;
 
+	UPROPERTY(config, EditAnywhere, Category = Mobile, meta = (
+		ConsoleVariable = "r.Mobile.SupportsGen4TAA", DisplayName = "Supports desktop Gen4 TAA on mobile",
+		ToolTip = "Support desktop Gen4 TAA with mobile rendering. Changing this setting requires restarting the editor.",
+		ConfigRestartRequired = true))
+		uint32 bMobileSupportsGen4TAA : 1;
+
 	UPROPERTY(config, EditAnywhere, Category="Experimental|LOD Streaming|Skeletal Mesh", meta=(
 		DisplayName="Stream LODs by default",
 		ToolTip="Whether to stream skeletal mesh LODs by default."))
diff --git a/Engine/Source/Runtime/Engine/Private/SceneView.cpp b/Engine/Source/Runtime/Engine/Private/SceneView.cpp
--- a/Engine/Source/Runtime/Engine/Private/SceneView.cpp
+++ b/Engine/Source/Runtime/Engine/Private/SceneView.cpp
@@ -909,7 +909,7 @@ void FSceneView::SetupAntiAliasingMethod()
 
 		if (AntiAliasingMethod == AAM_TemporalAA)
 		{
-			if (!Family->EngineShowFlags.TemporalAA || !Family->bRealtimeUpdate)
+			if (!Family->EngineShowFlags.TemporalAA || !Family->bRealtimeUpdate || !SupportsGen4TAA(GetShaderPlatform()))
 			{
 				AntiAliasingMethod = AAM_FXAA;
 			}
diff --git a/Engine/Source/Runtime/RenderCore/Private/RenderUtils.cpp b/Engine/Source/Runtime/RenderCore/Private/RenderUtils.cpp
--- a/Engine/Source/Runtime/RenderCore/Private/RenderUtils.cpp
+++ b/Engine/Source/Runtime/RenderCore/Private/RenderUtils.cpp
@@ -1139,6 +1139,17 @@ RENDERCORE_API bool IsMobileDistanceFieldShadowingEnabled(const FStaticShaderPla
 	return GRHISupportsPixelShaderUAVs && bDistanceFieldShadowingEnabled && IsMobileDistanceFieldEnabled(Platform);
 }
 
+RENDERCORE_API bool SupportsGen4TAA(const FStaticShaderPlatform Platform)
+{
+	if (IsMobilePlatform(Platform))
+	{
+		static FShaderPlatformCachedIniValue<bool> MobileSupportsGen4TAAIniValue(TEXT("/Script/Engine.RendererSettings"), TEXT("r.Mobile.SupportsGen4TAA"));
+		return (MobileSupportsGen4TAAIniValue.Get(Platform) != 0);
+	}
+
+	return true;
+}
+
 template<typename Type>
 Type FShaderPlatformCachedIniValue<Type>::Get(EShaderPlatform ShaderPlatform)
 {
@@ -1598,6 +1609,17 @@ RENDERCORE_API bool UseVirtualTextureLightmap(const FStaticFeatureLevel InFeatur
 	return bUseVirtualTextureLightmap;
 }
 
+RENDERCORE_API bool PlatformSupportsVelocityRendering(const FStaticShaderPlatform Platform)
+{
+	if (IsMobilePlatform(Platform))
+	{
+		// Enable velocity rendering if desktop Gen4 TAA is supported on mobile.
+		return SupportsGen4TAA(Platform);
+	}
+
+	return true;
+}
+
 RENDERCORE_API bool DoesPlatformSupportNanite(EShaderPlatform Platform)
 {
 	// Make sure the current platform has DDPI definitions.
diff --git a/Engine/Source/Runtime/RenderCore/Public/RenderUtils.h b/Engine/Source/Runtime/RenderCore/Public/RenderUtils.h
--- a/Engine/Source/Runtime/RenderCore/Public/RenderUtils.h
+++ b/Engine/Source/Runtime/RenderCore/Public/RenderUtils.h
@@ -503,6 +503,10 @@ RENDERCORE_API bool IsMobileDistanceFieldEnabled(const FStaticShaderPlatform Pla
 
 RENDERCORE_API bool IsMobileDistanceFieldShadowingEnabled(const FStaticShaderPlatform Platform);
 
+RENDERCORE_API bool SupportsGen4TAA(const FStaticShaderPlatform Platform);
+
+RENDERCORE_API bool PlatformSupportsVelocityRendering(const FStaticShaderPlatform Platform);
+
 RENDERCORE_API bool IsUsingDBuffers(const FStaticShaderPlatform Platform);
 
 
diff --git a/Engine/Source/Runtime/Renderer/Private/PostProcess/TemporalAA.cpp b/Engine/Source/Runtime/Renderer/Private/PostProcess/TemporalAA.cpp
--- a/Engine/Source/Runtime/Renderer/Private/PostProcess/TemporalAA.cpp
+++ b/Engine/Source/Runtime/Renderer/Private/PostProcess/TemporalAA.cpp
@@ -355,7 +355,7 @@ class FTAAStandaloneCS : public FGlobalShader
 		}
 		
 		//Only Main and MainUpsampling config without DownSample permutations are supported on mobile platform.
-		return IsFeatureLevelSupported(Parameters.Platform, ERHIFeatureLevel::SM5) || ((PermutationVector.Get<FTAAPassConfigDim>() == ETAAPassConfig::Main || PermutationVector.Get<FTAAPassConfigDim>() == ETAAPassConfig::MainUpsampling) && !PermutationVector.Get<FTAADownsampleDim>());
+		return SupportsGen4TAA(Parameters.Platform) && (!IsMobilePlatform(Parameters.Platform) || ((PermutationVector.Get<FTAAPassConfigDim>() == ETAAPassConfig::Main || PermutationVector.Get<FTAAPassConfigDim>() == ETAAPassConfig::MainUpsampling) && !PermutationVector.Get<FTAADownsampleDim>()));
 	}
 
 	static void ModifyCompilationEnvironment(const FGlobalShaderPermutationParameters& Parameters, FShaderCompilerEnvironment& OutEnvironment)
diff --git a/Engine/Source/Runtime/Renderer/Private/RendererScene.cpp b/Engine/Source/Runtime/Renderer/Private/RendererScene.cpp
--- a/Engine/Source/Runtime/Renderer/Private/RendererScene.cpp
+++ b/Engine/Source/Runtime/Renderer/Private/RendererScene.cpp
@@ -3886,9 +3886,13 @@ struct FPrimitiveArraySortKey
 	}
 };
 
-static bool ShouldPrimitiveOutputVelocity(const FPrimitiveSceneProxy* Proxy)
+static bool ShouldPrimitiveOutputVelocity(const FPrimitiveSceneProxy* Proxy, const FStaticShaderPlatform ShaderPlatform)
 {
-	return Proxy->DrawsVelocity() || (!!CVarWPOPrimitivesOutputVelocity.GetValueOnRenderThread() && Proxy->IsUsingWPOMaterial());
+	bool bShouldPrimitiveOutputVelocity = Proxy->DrawsVelocity() || (!!CVarWPOPrimitivesOutputVelocity.GetValueOnRenderThread() && Proxy->IsUsingWPOMaterial());
+
+	bool bPlatformSupportsVelocityRendering = PlatformSupportsVelocityRendering(ShaderPlatform);
+
+	return bPlatformSupportsVelocityRendering && bShouldPrimitiveOutputVelocity;
 }
 
 void FScene::UpdateAllPrimitiveSceneInfos(FRDGBuilder& GraphBuilder, bool bAsyncCreateLPIs)
@@ -4390,7 +4394,7 @@ void FScene::UpdateAllPrimitiveSceneInfos(FRDGBuilder& GraphBuilder, bool bAsync
 				FPrimitiveSceneInfo* PrimitiveSceneInfo = AddedLocalPrimitiveSceneInfos[AddIndex];
 				int32 PrimitiveIndex = PrimitiveSceneInfo->PackedIndex;
 
-				if (ShouldPrimitiveOutputVelocity(PrimitiveSceneInfo->Proxy))
+				if (ShouldPrimitiveOutputVelocity(PrimitiveSceneInfo->Proxy, GetShaderPlatform()))
 				{
 					PrimitiveSceneInfo->bRegisteredWithVelocityData = true;
 					// We must register the initial LocalToWorld with the velocity state. 
@@ -4460,7 +4464,7 @@ void FScene::UpdateAllPrimitiveSceneInfos(FRDGBuilder& GraphBuilder, bool bAsync
 			// (note that the octree update relies on the bounds not being modified yet).
 			PrimitiveSceneInfo->RemoveFromScene(bUpdateStaticDrawLists);
 
-			if (ShouldPrimitiveOutputVelocity(PrimitiveSceneInfo->Proxy))
+			if (ShouldPrimitiveOutputVelocity(PrimitiveSceneInfo->Proxy, GetShaderPlatform()))
 			{
 				PrimitiveSceneInfo->bRegisteredWithVelocityData = true;
 				VelocityData.UpdateTransform(PrimitiveSceneInfo, LocalToWorld, PrimitiveSceneProxy->GetLocalToWorld());
diff --git a/Engine/Source/Runtime/Renderer/Private/VelocityRendering.cpp b/Engine/Source/Runtime/Renderer/Private/VelocityRendering.cpp
--- a/Engine/Source/Runtime/Renderer/Private/VelocityRendering.cpp
+++ b/Engine/Source/Runtime/Renderer/Private/VelocityRendering.cpp
@@ -90,6 +90,8 @@ public:
 		// Compile for materials which modify meshes.
 		const bool bMayModifyMeshes = Parameters.MaterialParameters.bMaterialMayModifyMeshPosition;
 
+		const bool bHasPlatformSupport = PlatformSupportsVelocityRendering(Parameters.Platform);
+
 		/**
 		 * Any material with a vertex factory incompatible with base pass velocity generation must generate
 		 * permutations for this shader. Shaders which don't fall into this set are considered "simple" enough
@@ -101,7 +103,7 @@ public:
 		// The material may explicitly override and request that it be rendered into the velocity pass.
 		const bool bIsSeparateVelocityPassRequiredByMaterial = Parameters.MaterialParameters.bIsTranslucencyWritingVelocity;
 
-		return (bIsSeparateVelocityPassRequired || bIsSeparateVelocityPassRequiredByMaterial);
+		return bHasPlatformSupport && (bIsSeparateVelocityPassRequired || bIsSeparateVelocityPassRequiredByMaterial);
 	}
 
 	FVelocityVS() = default;
@@ -188,7 +190,7 @@ bool FDeferredShadingSceneRenderer::ShouldRenderVelocities() const
 
 bool FMobileSceneRenderer::ShouldRenderVelocities() const
 {
-	if (!FVelocityRendering::IsSeparateVelocityPassSupported(ShaderPlatform) || ViewFamily.UseDebugViewPS())
+	if (!FVelocityRendering::IsSeparateVelocityPassSupported(ShaderPlatform) || ViewFamily.UseDebugViewPS() || !PlatformSupportsVelocityRendering(ShaderPlatform))
 	{
 		return false;
 	}
@@ -256,19 +258,22 @@ void FSceneRenderer::RenderVelocities(
 
 			RDG_GPU_MASK_SCOPE(GraphBuilder, View.GPUMask);
 
-			// Avoid adding a separate clear pass, clear it in the velocity pass if it is possible.
-			if (bForceVelocity && !bHasAnyDraw)
+			const bool bIsParallelVelocity = IsParallelVelocity(ShaderPlatform);
+
+			// Clear velocity render target explicitly when velocity rendering in parallel or no draw but force to.
+			// Avoid adding a separate clear pass in non parallel rendering.
+			const bool bExplicitlyClearVelocity = (VelocityLoadAction == ERenderTargetLoadAction::EClear) && (bIsParallelVelocity || (bForceVelocity && !bHasAnyDraw));
+
+			if (bExplicitlyClearVelocity)
 			{
-				if (VelocityLoadAction == ERenderTargetLoadAction::EClear)
-				{
-					AddClearRenderTargetPass(GraphBuilder, SceneTextures.Velocity);
+				AddClearRenderTargetPass(GraphBuilder, SceneTextures.Velocity);
 
-					if (!View.Family->bMultiGPUForkAndJoin)
-					{
-						VelocityLoadAction = ERenderTargetLoadAction::ELoad;
-					}
-				}
+				// Parallel render need to use Load action in any case.
+				VelocityLoadAction = ERenderTargetLoadAction::ELoad;
+			}
 
+			if (!bHasAnyDraw)
+			{
 				continue;
 			}
 
@@ -286,7 +291,7 @@ void FSceneRenderer::RenderVelocities(
 
 			PassParameters->RenderTargets[0] = FRenderTargetBinding(SceneTextures.Velocity, VelocityLoadAction);
 
-			if (IsParallelVelocity(ShaderPlatform))
+			if (bIsParallelVelocity)
 			{
 				GraphBuilder.AddPass(
 					RDG_EVENT_NAME("VelocityParallel"),
@@ -403,7 +408,7 @@ bool FVelocityMeshProcessor::PrimitiveHasVelocityForView(const FViewInfo& View,
 
 bool FOpaqueVelocityMeshProcessor::PrimitiveCanHaveVelocity(EShaderPlatform ShaderPlatform, const FPrimitiveSceneProxy* PrimitiveSceneProxy)
 {
-	if (!FVelocityRendering::IsSeparateVelocityPassSupported(ShaderPlatform))
+	if (!FVelocityRendering::IsSeparateVelocityPassSupported(ShaderPlatform) || !PlatformSupportsVelocityRendering(ShaderPlatform))
 	{
 		return false;
 	}
@@ -541,7 +546,7 @@ bool FTranslucentVelocityMeshProcessor::PrimitiveCanHaveVelocity(EShaderPlatform
 	 * Therefore, the primitive can't be filtered based on motion, or it will break post
 	 * effects like depth of field which rely on depth information.
 	 */
-	return FVelocityRendering::IsSeparateVelocityPassSupported(ShaderPlatform);
+	return FVelocityRendering::IsSeparateVelocityPassSupported(ShaderPlatform) && PlatformSupportsVelocityRendering(ShaderPlatform);
 }
 
 bool FTranslucentVelocityMeshProcessor::PrimitiveHasVelocityForFrame(const FPrimitiveSceneProxy* PrimitiveSceneProxy)
