Fixed lighting debug probes to not render in refelction captures.

diff --git a/Engine/Source/Runtime/Renderer/Private/DebugProbeRendering.cpp b/Engine/Source/Runtime/Renderer/Private/DebugProbeRendering.cpp
--- a/Engine/Source/Runtime/Renderer/Private/DebugProbeRendering.cpp
+++ b/Engine/Source/Runtime/Renderer/Private/DebugProbeRendering.cpp
@@ -103,7 +103,7 @@ void StampDeferredDebugProbeDepthPS(
 	const bool bVisualizeLightingOnProbes = CVarVisualizeLightingOnProbes.GetValueOnRenderThread() > 0;
 	for (const FViewInfo& View : Views)
 	{
-		if (!(bVisualizeLightingOnProbes || View.Family->EngineShowFlags.VisualizeLightingOnProbes))
+		if (!(bVisualizeLightingOnProbes || View.Family->EngineShowFlags.VisualizeLightingOnProbes) || View.bIsReflectionCapture)
 		{
 			continue;
 		}
@@ -127,7 +127,7 @@ void StampDeferredDebugProbeMaterialPS(
 	const bool bVisualizeLightingOnProbes = CVarVisualizeLightingOnProbes.GetValueOnRenderThread() > 0;
 	for (const FViewInfo& View : Views)
 	{
-		if (!(bVisualizeLightingOnProbes || View.Family->EngineShowFlags.VisualizeLightingOnProbes))
+		if (!(bVisualizeLightingOnProbes || View.Family->EngineShowFlags.VisualizeLightingOnProbes) || View.bIsReflectionCapture)
 		{
 			continue;
 		}
@@ -161,7 +161,7 @@ void StampDeferredDebugProbeVelocityPS(
 	const bool bVisualizeLightingOnProbes = CVarVisualizeLightingOnProbes.GetValueOnRenderThread() > 0;
 	for (const FViewInfo& View : Views)
 	{
-		if (!(bVisualizeLightingOnProbes || View.Family->EngineShowFlags.VisualizeLightingOnProbes))
+		if (!(bVisualizeLightingOnProbes || View.Family->EngineShowFlags.VisualizeLightingOnProbes) || View.bIsReflectionCapture)
 		{
 			continue;
 		}
