Fix for mesh lod selection when rendering into virtual texture

- virtual texture mesh lod range can have no intersection with available lods for the mesh
  - this resulted in attempted rendering of the mesh lod corresponding to the range max which may not be streamed in
    - we now render the highest detail mesh lod streamed in into the VT in this case even if it falls outside the desired range

diff --git a/Engine/Source/Runtime/Renderer/Private/VT/RuntimeVirtualTextureRender.cpp b/Engine/Source/Runtime/Renderer/Private/VT/RuntimeVirtualTextureRender.cpp
--- a/Engine/Source/Runtime/Renderer/Private/VT/RuntimeVirtualTextureRender.cpp
+++ b/Engine/Source/Runtime/Renderer/Private/VT/RuntimeVirtualTextureRender.cpp
@@ -510,8 +510,9 @@ namespace RuntimeVirtualTexture
 			const float AreaRatio = 2.f * SphereBounds.W * RcpWorldSize;
 			const int32 CurFirstLODIdx = PrimitiveSceneInfo->Proxy->GetCurrentFirstLODIdx_RenderThread();
 			const int32 MinLODIdx = FMath::Max((int32)LodInfo.MinLod, CurFirstLODIdx);
+			const int32 MaxLODIdx = FMath::Max((int32)LodInfo.MaxLod, CurFirstLODIdx);
 			const int32 LodBias = (int32)LodInfo.LodBias - FPrimitiveVirtualTextureLodInfo::LodBiasOffset;
-			const int32 LodIndex = FMath::Clamp<int32>(LodBias - FMath::FloorToInt(FMath::Log2(AreaRatio)), MinLODIdx, LodInfo.MaxLod);
+			const int32 LodIndex = FMath::Clamp<int32>(LodBias - FMath::FloorToInt(FMath::Log2(AreaRatio)), MinLODIdx, MaxLODIdx);
 
 			// Process meshes
 			for (int32 MeshIndex = 0; MeshIndex < PrimitiveSceneInfo->StaticMeshes.Num(); ++MeshIndex)
