Moved Allocate/Free of GPU-Scene Instance out of AddToScene/RemoveFromScene as these are called for primitive transform updates.

diff --git a/Engine/Source/Runtime/Renderer/Private/PrimitiveSceneInfo.cpp b/Engine/Source/Runtime/Renderer/Private/PrimitiveSceneInfo.cpp
--- a/Engine/Source/Runtime/Renderer/Private/PrimitiveSceneInfo.cpp
+++ b/Engine/Source/Runtime/Renderer/Private/PrimitiveSceneInfo.cpp
@@ -959,6 +959,76 @@
 	return LCIs.Num();
 }
 
+void FPrimitiveSceneInfo::AllocateGPUSceneInstances(FScene* Scene, const TArrayView<FPrimitiveSceneInfo*>& SceneInfos)
+{
+	if (Scene->GPUScene.IsEnabled())
+	{
+		SCOPE_CYCLE_COUNTER(STAT_UpdateGPUSceneTime);
+		for (FPrimitiveSceneInfo* SceneInfo : SceneInfos)
+		{
+			check(SceneInfo->InstanceDataOffset == INDEX_NONE);
+			check(SceneInfo->NumInstanceDataEntries == 0);
+			if (SceneInfo->Proxy->SupportsInstanceDataBuffer())
+			{
+				if (const TArray<FPrimitiveInstance>* PrimitiveInstances = SceneInfo->Proxy->GetPrimitiveInstances())
+				{
+					SceneInfo->InstanceDataOffset = Scene->GPUScene.AllocateInstanceSlots(PrimitiveInstances->Num());
+					SceneInfo->NumInstanceDataEntries = PrimitiveInstances->Num();
+
+					if (GGPUSceneInstanceBVH)
+					{
+						for (int32 InstanceIndex = 0; InstanceIndex < PrimitiveInstances->Num(); ++InstanceIndex)
+						{
+							const FPrimitiveInstance& PrimitiveInstance = (*PrimitiveInstances)[InstanceIndex];
+							const FRenderBounds WorldBounds = PrimitiveInstance.LocalBounds.TransformBy(SceneInfo->Proxy->GetLocalToWorld());
+							// TODO: Replace Instance BVH FBounds with FRenderBounds
+							Scene->InstanceBVH.Add(FBounds({ WorldBounds.GetMin(), WorldBounds.GetMax() }), SceneInfo->InstanceDataOffset + InstanceIndex);
+						}
+					}
+				}
+			}
+			else
+			{
+				// Allocate a single 'dummy/fallback' instance for the primitive that gets automatically populated with the data from the primitive
+				SceneInfo->InstanceDataOffset = Scene->GPUScene.AllocateInstanceSlots(1);
+				SceneInfo->NumInstanceDataEntries = 1;
+			}
+
+			// Force a primitive update in the GPU scene, 
+			// NOTE: does not set Added as this is handled elsewhere.
+			Scene->GPUScene.AddPrimitiveToUpdate(SceneInfo->PackedIndex);
+
+			// Force a primitive update in the Lumen scene
+			if (Scene->LumenSceneData)
+			{
+				Scene->LumenSceneData->UpdatePrimitiveInstanceOffset(SceneInfo->PackedIndex);
+			}
+		}
+	}
+}
+
+void FPrimitiveSceneInfo::FreeGPUSceneInstances()
+{
+	// Release all instance data slots associated with this primitive.
+	if (InstanceDataOffset != INDEX_NONE && Scene->GPUScene.IsEnabled())
+	{
+		SCOPE_CYCLE_COUNTER(STAT_UpdateGPUSceneTime);
+
+		check(Proxy->SupportsInstanceDataBuffer() || NumInstanceDataEntries == 1);
+		if (GGPUSceneInstanceBVH)
+		{
+			for (int32 InstanceIndex = 0; InstanceIndex < NumInstanceDataEntries; InstanceIndex++)
+			{
+				Scene->InstanceBVH.Remove(InstanceDataOffset + InstanceIndex);
+			}
+		}
+
+		Scene->GPUScene.FreeInstanceSlots(InstanceDataOffset, NumInstanceDataEntries);
+		InstanceDataOffset = INDEX_NONE;
+		NumInstanceDataEntries = 0;
+	}
+}
+
 void FPrimitiveSceneInfo::AddToScene(FRHICommandListImmediate& RHICmdList, FScene* Scene, const TArrayView<FPrimitiveSceneInfo*>& SceneInfos, bool bUpdateStaticDrawLists, bool bAddToStaticDrawLists, bool bAsyncCreateLPIs)
 {
 	check(IsInRenderingThread());
@@ -1037,53 +1107,6 @@
 		}
 	}
 
-	if (Scene->GPUScene.IsEnabled())
-	{
-		SCOPE_CYCLE_COUNTER(STAT_UpdateGPUSceneTime);
-		for (FPrimitiveSceneInfo* SceneInfo : SceneInfos)
-		{
-			check(SceneInfo->InstanceDataOffset == INDEX_NONE);
-			check(SceneInfo->NumInstanceDataEntries == 0);
-			if (SceneInfo->Proxy->SupportsInstanceDataBuffer())
-			{
-				if (const TArray<FPrimitiveInstance>* PrimitiveInstances = SceneInfo->Proxy->GetPrimitiveInstances())
-				{
-					SceneInfo->InstanceDataOffset = Scene->GPUScene.AllocateInstanceSlots(PrimitiveInstances->Num());
-					SceneInfo->NumInstanceDataEntries = PrimitiveInstances->Num();
-
-					if (GGPUSceneInstanceBVH)
-					{
-						for (int32 InstanceIndex = 0; InstanceIndex < PrimitiveInstances->Num(); ++InstanceIndex)
-						{
-							const FPrimitiveInstance& PrimitiveInstance = (*PrimitiveInstances)[InstanceIndex];
-							const FRenderBounds WorldBounds = PrimitiveInstance.LocalBounds.TransformBy(SceneInfo->Proxy->GetLocalToWorld());
-							// TODO: Replace Instance BVH FBounds with FRenderBounds
-							Scene->InstanceBVH.Add(FBounds({ WorldBounds.GetMin(), WorldBounds.GetMax() }), SceneInfo->InstanceDataOffset + InstanceIndex);
-						}
-					}
-				}
-			}
-			else
-			{
-				// Allocate a single 'dummy/fallback' instance for the primitive that gets automatically populated with the data from the primitive
-				SceneInfo->InstanceDataOffset = Scene->GPUScene.AllocateInstanceSlots(1);
-				SceneInfo->NumInstanceDataEntries = 1;
-			}
-
-			// Force a primitive update in the GPU scene
-			if (!Scene->GPUScene.PrimitivesMarkedToUpdate[SceneInfo->PackedIndex])
-			{
-				Scene->GPUScene.PrimitivesToUpdate.Add(SceneInfo->PackedIndex);
-				Scene->GPUScene.PrimitivesMarkedToUpdate[SceneInfo->PackedIndex] = true;
-			}
-
-			// Force a primitive update in the Lumen scene
-			if (Scene->LumenSceneData)
-			{
-				Scene->LumenSceneData->UpdatePrimitiveInstanceOffset(SceneInfo->PackedIndex);
-			}
-		}
-	}
 
 	{
 		SCOPED_NAMED_EVENT(FPrimitiveSceneInfo_AddToScene_ReflectionCaptures, FColor::Yellow);
@@ -1267,25 +1290,6 @@
 		Scene->GPUScene.LightmapDataAllocator.Free(LightmapDataOffset, NumLightmapDataEntries);
 	}
 
-	// Release all instance data slots associated with this primitive.
-	if (InstanceDataOffset != INDEX_NONE && Scene->GPUScene.IsEnabled())
-	{
-		SCOPE_CYCLE_COUNTER(STAT_UpdateGPUSceneTime);
-
-		check(Proxy->SupportsInstanceDataBuffer() || NumInstanceDataEntries == 1);
-		if( GGPUSceneInstanceBVH )
-		{
-			for( int32 InstanceIndex = 0; InstanceIndex < NumInstanceDataEntries; InstanceIndex++ )
-			{
-				Scene->InstanceBVH.Remove( InstanceDataOffset + InstanceIndex );
-			}
-		}
-
-		Scene->GPUScene.FreeInstanceSlots(InstanceDataOffset, NumInstanceDataEntries);
-		InstanceDataOffset = INDEX_NONE;
-		NumInstanceDataEntries = 0;
-	}
-
 	if (Proxy->CastsDynamicIndirectShadow())
 	{
 		Scene->DynamicIndirectCasterPrimitives.RemoveSingleSwap(this);
diff --git a/Engine/Source/Runtime/Renderer/Private/RendererScene.cpp b/Engine/Source/Runtime/Renderer/Private/RendererScene.cpp
--- a/Engine/Source/Runtime/Renderer/Private/RendererScene.cpp
+++ b/Engine/Source/Runtime/Renderer/Private/RendererScene.cpp
@@ -4109,6 +4109,8 @@ void FScene::UpdateAllPrimitiveSceneInfos(FRDGBuilder& GraphBuilder, bool bAsync
 				// Remove the primitive from the scene.
 				PrimitiveSceneInfo->RemoveFromScene(true);
 
+				PrimitiveSceneInfo->FreeGPUSceneInstances();
+
 				// Update the primitive that was swapped to this index
 				GPUScene.AddPrimitiveToUpdate(PrimitiveIndex);
 
@@ -4362,6 +4364,7 @@ void FScene::UpdateAllPrimitiveSceneInfos(FRDGBuilder& GraphBuilder, bool bAsync
 				GPUScene.MarkPrimitiveAdded(PrimitiveIndex);
 			}
 
+			FPrimitiveSceneInfo::AllocateGPUSceneInstances(this, TArrayView<FPrimitiveSceneInfo*>(&AddedLocalPrimitiveSceneInfos[StartIndex], AddedLocalPrimitiveSceneInfos.Num() - StartIndex));
 
 			{
 				SCOPED_NAMED_EVENT(FScene_AddPrimitiveSceneInfoToScene, FColor::Turquoise);
diff --git a/Engine/Source/Runtime/Renderer/Public/PrimitiveSceneInfo.h b/Engine/Source/Runtime/Renderer/Public/PrimitiveSceneInfo.h
--- a/Engine/Source/Runtime/Renderer/Public/PrimitiveSceneInfo.h
+++ b/Engine/Source/Runtime/Renderer/Public/PrimitiveSceneInfo.h
@@ -372,6 +372,10 @@ public:
 	/** Removes the primitive from the scene. */
 	void RemoveFromScene(bool bUpdateStaticDrawLists);
 
+	/** Allocate/Free slots for instance data in GPU-Scene */
+	static void AllocateGPUSceneInstances(FScene* Scene, const TArrayView<FPrimitiveSceneInfo*>& SceneInfos);
+	void FreeGPUSceneInstances();
+
 	/** return true if we need to call ConditionalUpdateStaticMeshes */
 	bool NeedsUpdateStaticMeshes();
 
