UE5 - Fix texture not getting released
#rnx

diff --git a/Engine/Source/Runtime/Renderer/Private/SystemTextures.cpp b/Engine/Source/Runtime/Renderer/Private/SystemTextures.cpp
--- a/Engine/Source/Runtime/Renderer/Private/SystemTextures.cpp
+++ b/Engine/Source/Runtime/Renderer/Private/SystemTextures.cpp
@@ -831,6 +831,7 @@ void FSystemTextures::ReleaseDynamicRHI()
 	StencilDummySRV.SafeRelease();
 	BlackDepthCube.SafeRelease();
 	GTAOPreIntegrated.SafeRelease();
+	AsciiTexture.SafeRelease();
 
 	DefaultTextures.Empty();
 	DefaultBuffers.Empty();
