Move shader draw debug lock data onto the view state.

diff --git a/Engine/Shaders/Private/ShaderDrawDebug.ush b/Engine/Shaders/Private/ShaderDrawDebug.ush
--- a/Engine/Shaders/Private/ShaderDrawDebug.ush
+++ b/Engine/Shaders/Private/ShaderDrawDebug.ush
@@ -6,6 +6,21 @@
 #define SHADER_DRAW_DEBUG 1
 #endif
 
+// Set of predefined colors
+#define ColorWhite		float4(1, 1, 1, 1)
+#define ColorBlack		float4(0, 0, 0, 1)
+#define ColorRed		float4(1, 0, 0, 1)
+#define ColorGreen		float4(0, 1, 0, 1)
+#define ColorBlue		float4(0, 0, 1, 1)
+#define ColorYellow		float4(1, 1, 0, 1)
+#define ColorCyan		float4(0, 1, 1, 1)
+#define ColorMagenta	float4(1, 0, 1, 1)
+#define ColorOrange		float4(243.f / 255.f, 156.f / 255.f,  18.f / 255.f, 1)
+#define ColorPurple		float4(169.f / 255.f,   7.f / 255.f, 228.f / 255.f, 1)
+#define ColorTurquoise	float4( 26.f / 255.f, 188.f / 255.f, 156.f / 255.f, 1)
+#define ColorSilver		float4(189.f / 255.f, 195.f / 255.f, 199.f / 255.f, 1)
+#define ColorEmerald	float4( 46.f / 255.f, 204.f / 255.f, 113.f / 255.f, 1)
+
 struct FShaderDrawElement 
 {
 	float3 Pos0;
@@ -48,7 +63,7 @@ FShaderDrawElement UnpackShaderElement(FPackedShaderDrawElement PackedE)
 
 #if SHADER_DRAW_DEBUG == 1
 
-RWStructuredBuffer<FPackedShaderDrawElement>	OutShaderDrawPrimitive;
+RWStructuredBuffer<FPackedShaderDrawElement>OutShaderDrawPrimitive;
 RWBuffer<uint>								OutputShaderDrawIndirect;
 
 int2 ShaderDrawCursorPos;
@@ -79,6 +94,11 @@ void AddLine(float3 Pos0, float3 Pos1, float4 Color0, float4 Color1)
 	}
 }
 
+void AddLine(float3 Pos0, float3 Pos1, float4 Color)
+{
+	AddLine(Pos0, Pos1, Color, Color);
+}
+
 void AddQuad(float3 Pos0, float3 Pos1, float3 Pos2, float3 Pos3, float4 Color)
 {
 	uint PreviousPrimitiveCount = 0;
diff --git a/Engine/Source/Runtime/Renderer/Private/GpuDebugRendering.cpp b/Engine/Source/Runtime/Renderer/Private/GpuDebugRendering.cpp
--- a/Engine/Source/Runtime/Renderer/Private/GpuDebugRendering.cpp
+++ b/Engine/Source/Runtime/Renderer/Private/GpuDebugRendering.cpp
@@ -7,6 +7,7 @@
 #include "RenderGraphUtils.h"
 #include "Containers/ResourceArray.h"
 #include "CommonRenderResources.h"
+#include "ScenePrivate.h"
 
 namespace ShaderDrawDebug 
 {
@@ -92,16 +93,6 @@ namespace ShaderDrawDebug
 		float Pos1_ColorY[4];		// float3 pos1 + packed color1
 	};
 
-	// This needs to be allocated per view, or move into a more persistent place
-	struct FLockedData
-	{
-		TRefCountPtr<FRDGPooledBuffer> Buffer;
-		TRefCountPtr<FRDGPooledBuffer> IndirectBuffer;
-		bool bIsLocked = false;
-	};
-
-	static FLockedData LockedData = {};
-
 	//////////////////////////////////////////////////////////////////////////
 
 	class FShaderDrawDebugClearCS : public FGlobalShader
@@ -263,7 +254,9 @@ namespace ShaderDrawDebug
 			return;
 		}
 
-		const bool bLockBufferThisFrame = IsShaderDrawLocked() && !LockedData.bIsLocked;
+		FSceneViewState* ViewState = View.ViewState;
+
+		const bool bLockBufferThisFrame = IsShaderDrawLocked() && ViewState && !ViewState->ShaderDrawDebugStateData.bIsLocked;
 		ERDGBufferFlags Flags = bLockBufferThisFrame ? ERDGBufferFlags::MultiFrame : ERDGBufferFlags::None;
 
 		FRDGBufferRef DataBuffer = GraphBuilder.CreateBuffer(FRDGBufferDesc::CreateStructuredDesc(sizeof(FPackedShaderDrawElement), GetMaxElementCount()), TEXT("ShaderDrawDataBuffer"), Flags);
@@ -285,22 +278,23 @@ namespace ShaderDrawDebug
 			FComputeShaderUtils::Dispatch(RHICmdList, ComputeShader, *Parameters, FIntVector(1,1,1));
 		});
 
+
 		View.ShaderDrawData.Buffer = DataBuffer;
 		View.ShaderDrawData.IndirectBuffer = IndirectBuffer;
 		View.ShaderDrawData.CursorPosition = View.CursorPos;
 
-		if (IsShaderDrawLocked() && !LockedData.bIsLocked)
+		if (IsShaderDrawLocked() && ViewState && !ViewState->ShaderDrawDebugStateData.bIsLocked)
 		{
-			LockedData.Buffer = GraphBuilder.ConvertToExternalBuffer(View.ShaderDrawData.Buffer);
-			LockedData.IndirectBuffer = GraphBuilder.ConvertToExternalBuffer(View.ShaderDrawData.IndirectBuffer);
-			LockedData.bIsLocked = true;
+			ViewState->ShaderDrawDebugStateData.Buffer = GraphBuilder.ConvertToExternalBuffer(View.ShaderDrawData.Buffer);
+			ViewState->ShaderDrawDebugStateData.IndirectBuffer = GraphBuilder.ConvertToExternalBuffer(View.ShaderDrawData.IndirectBuffer);
+			ViewState->ShaderDrawDebugStateData.bIsLocked = true;
 		}
 
-		if (!IsShaderDrawLocked() && LockedData.bIsLocked)
+		if (!IsShaderDrawLocked() && ViewState && ViewState->ShaderDrawDebugStateData.bIsLocked)
 		{
-			LockedData.Buffer = nullptr;
-			LockedData.IndirectBuffer = nullptr;
-			LockedData.bIsLocked = false;
+			ViewState->ShaderDrawDebugStateData.Buffer = nullptr;
+			ViewState->ShaderDrawDebugStateData.IndirectBuffer = nullptr;
+			ViewState->ShaderDrawDebugStateData.bIsLocked = false;
 		}
 	}
 
@@ -317,10 +311,10 @@ namespace ShaderDrawDebug
 			InternalDrawView(GraphBuilder, View, DataBuffer, IndirectBuffer, OutputTexture, DepthTexture);
 		}
 
-		if (IsShaderDrawLocked())
+		if (View.ViewState && View.ViewState->ShaderDrawDebugStateData.bIsLocked)
 		{
-			FRDGBufferRef DataBuffer = GraphBuilder.RegisterExternalBuffer(LockedData.Buffer);
-			FRDGBufferRef IndirectBuffer = GraphBuilder.RegisterExternalBuffer(LockedData.IndirectBuffer);
+			FRDGBufferRef DataBuffer = GraphBuilder.RegisterExternalBuffer(View.ViewState->ShaderDrawDebugStateData.Buffer);
+			FRDGBufferRef IndirectBuffer = GraphBuilder.RegisterExternalBuffer(View.ViewState->ShaderDrawDebugStateData.IndirectBuffer);
 			InternalDrawView(GraphBuilder, View, DataBuffer, IndirectBuffer, OutputTexture, DepthTexture);
 		}
 	}
diff --git a/Engine/Source/Runtime/Renderer/Private/RendererScene.cpp b/Engine/Source/Runtime/Renderer/Private/RendererScene.cpp
--- a/Engine/Source/Runtime/Renderer/Private/RendererScene.cpp
+++ b/Engine/Source/Runtime/Renderer/Private/RendererScene.cpp
@@ -290,6 +290,7 @@ FSceneViewState::~FSceneViewState()
 	AOScreenGridResources = NULL;
 
 	HairStrandsViewStateData.Release();
+	ShaderDrawDebugStateData.Release();
 }
 
 #if WITH_EDITOR
diff --git a/Engine/Source/Runtime/Renderer/Private/ScenePrivate.h b/Engine/Source/Runtime/Renderer/Private/ScenePrivate.h
--- a/Engine/Source/Runtime/Renderer/Private/ScenePrivate.h
+++ b/Engine/Source/Runtime/Renderer/Private/ScenePrivate.h
@@ -679,6 +679,20 @@ struct FHLODSceneNodeVisibilityState
 	uint16 bIsFading	: 1;
 };
 
+struct FShaderDrawDebugStateData
+{
+	TRefCountPtr<FRDGPooledBuffer> Buffer;
+	TRefCountPtr<FRDGPooledBuffer> IndirectBuffer;
+	bool bIsLocked = false;
+
+	void Release()
+	{
+		Buffer = nullptr;
+		IndirectBuffer = nullptr;
+		bIsLocked = false;
+	}
+};
+
 /**
  * The scene manager's private implementation of persistent view state.
  * This class is associated with a particular camera across multiple frames by the game thread.
@@ -1089,6 +1103,8 @@ public:
 
 	FHairStrandsViewStateData HairStrandsViewStateData;
 
+	FShaderDrawDebugStateData ShaderDrawDebugStateData;
+
 	// call after OnFrameRenderingSetup()
 	virtual uint32 GetCurrentTemporalAASampleIndex() const
 	{
