Use material asset names, and not friendly names, for debug markers.

- Friendly name ignores material instances for historical reasons (it should really be using a debug name on the shadermap).

diff --git a/Engine/Source/Runtime/Renderer/Private/MeshPassProcessor.cpp b/Engine/Source/Runtime/Renderer/Private/MeshPassProcessor.cpp
--- a/Engine/Source/Runtime/Renderer/Private/MeshPassProcessor.cpp
+++ b/Engine/Source/Runtime/Renderer/Private/MeshPassProcessor.cpp
@@ -1133,7 +1133,7 @@ void FMeshDrawCommand::SetDebugData(const FPrimitiveSceneProxy* PrimitiveScenePr
 	DebugData.PixelShader = UntypedShaders.PixelShader;
 	DebugData.VertexFactory = VertexFactory;
 	DebugData.ResourceName =  PrimitiveSceneProxy ? PrimitiveSceneProxy->GetResourceName() : FName();
-	DebugData.MaterialName = Material->GetFriendlyName();
+	DebugData.MaterialName = Material->GetAssetName();
 }
 #endif
 
diff --git a/Engine/Source/Runtime/Renderer/Private/PostProcess/PostProcessMaterial.cpp b/Engine/Source/Runtime/Renderer/Private/PostProcess/PostProcessMaterial.cpp
--- a/Engine/Source/Runtime/Renderer/Private/PostProcess/PostProcessMaterial.cpp
+++ b/Engine/Source/Runtime/Renderer/Private/PostProcess/PostProcessMaterial.cpp
@@ -467,7 +467,7 @@ FScreenPassTexture AddPostProcessMaterialPass(
 	const FScreenPassTextureViewport SceneColorViewport(SceneColor);
 	const FScreenPassTextureViewport OutputViewport(Output);
 
-	RDG_EVENT_SCOPE(GraphBuilder, "PostProcessMaterial %dx%d Material=%s", SceneColorViewport.Rect.Width(), SceneColorViewport.Rect.Height(), *Material->GetFriendlyName());
+	RDG_EVENT_SCOPE(GraphBuilder, "PostProcessMaterial %dx%d Material=%s", SceneColorViewport.Rect.Width(), SceneColorViewport.Rect.Height(), *Material->GetAssetName());
 
 	const uint32 MaterialStencilRef = Material->GetStencilRefValue();
 
