Rename GpuDebugRendering -> ShaderDebug to be consistent with ShaderPrint.

diff --git a/Engine/Plugins/Runtime/HairStrands/Source/HairStrandsCore/Private/GroomTextureBuilder.cpp b/Engine/Plugins/Runtime/HairStrands/Source/HairStrandsCore/Private/GroomTextureBuilder.cpp
--- a/Engine/Plugins/Runtime/HairStrands/Source/HairStrandsCore/Private/GroomTextureBuilder.cpp
+++ b/Engine/Plugins/Runtime/HairStrands/Source/HairStrandsCore/Private/GroomTextureBuilder.cpp
@@ -9,7 +9,7 @@
 #include "RenderGraphUtils.h"
 #include "ShaderParameterStruct.h"
 #include "GlobalShader.h"
-#include "GpuDebugRendering.h"
+#include "ShaderDebug.h"
 #include "CommonRenderResources.h"
 #include "HairStrandsMeshProjection.h"
 #include "Engine/StaticMesh.h"
diff --git a/Engine/Plugins/Runtime/HairStrands/Source/HairStrandsCore/Private/HairCardsBuilder.cpp b/Engine/Plugins/Runtime/HairStrands/Source/HairStrandsCore/Private/HairCardsBuilder.cpp
--- a/Engine/Plugins/Runtime/HairStrands/Source/HairStrandsCore/Private/HairCardsBuilder.cpp
+++ b/Engine/Plugins/Runtime/HairStrands/Source/HairStrandsCore/Private/HairCardsBuilder.cpp
@@ -16,7 +16,7 @@
 #include "Containers/ResourceArray.h"
 #include "Rendering/SkeletalMeshRenderData.h"
 #include "HAL/ConsoleManager.h"
-#include "GpuDebugRendering.h"
+#include "ShaderDebug.h"
 #include "Async/ParallelFor.h"
 #include "Misc/ScopedSlowTask.h"
 #include "CommonRenderResources.h"
diff --git a/Engine/Plugins/Runtime/HairStrands/Source/HairStrandsCore/Private/HairStrandsRendering.cpp b/Engine/Plugins/Runtime/HairStrands/Source/HairStrandsCore/Private/HairStrandsRendering.cpp
--- a/Engine/Plugins/Runtime/HairStrands/Source/HairStrandsCore/Private/HairStrandsRendering.cpp
+++ b/Engine/Plugins/Runtime/HairStrands/Source/HairStrandsCore/Private/HairStrandsRendering.cpp
@@ -15,7 +15,7 @@
 #include "Containers/ResourceArray.h"
 #include "Rendering/SkeletalMeshRenderData.h"
 #include "HAL/ConsoleManager.h"
-#include "GpuDebugRendering.h"
+#include "ShaderDebug.h"
 #include "Async/ParallelFor.h"
 #include "RenderTargetPool.h"
 #include "GroomTextureBuilder.h"
diff --git a/Engine/Source/Runtime/Renderer/Private/DeferredShadingRenderer.cpp b/Engine/Source/Runtime/Renderer/Private/DeferredShadingRenderer.cpp
--- a/Engine/Source/Runtime/Renderer/Private/DeferredShadingRenderer.cpp
+++ b/Engine/Source/Runtime/Renderer/Private/DeferredShadingRenderer.cpp
@@ -45,7 +45,7 @@
 #include "RayTracingDefinitions.h"
 #include "RayTracingInstance.h"
 #include "ShaderPrint.h"
-#include "GpuDebugRendering.h"
+#include "ShaderDebug.h"
 #include "GPUSortManager.h"
 #include "HairStrands/HairStrandsRendering.h"
 #include "HairStrands/HairStrandsData.h"
diff --git a/Engine/Source/Runtime/Renderer/Private/GpuDebugRendering.cpp b/Engine/Source/Runtime/Renderer/Private/GpuDebugRendering.cpp
deleted file mode 100644
diff --git a/Engine/Source/Runtime/Renderer/Private/HairStrands/HairStrandsDebug.cpp b/Engine/Source/Runtime/Renderer/Private/HairStrands/HairStrandsDebug.cpp
--- a/Engine/Source/Runtime/Renderer/Private/HairStrands/HairStrandsDebug.cpp
+++ b/Engine/Source/Runtime/Renderer/Private/HairStrands/HairStrandsDebug.cpp
@@ -24,7 +24,7 @@
 #include "CanvasTypes.h"
 #include "ShaderPrintParameters.h"
 #include "RenderGraphUtils.h"
-#include "GpuDebugRendering.h"
+#include "ShaderDebug.h"
 #include "ShaderPrint.h"
 #include "ScreenPass.h"
 
diff --git a/Engine/Source/Runtime/Renderer/Private/HairStrands/HairStrandsEnvironment.cpp b/Engine/Source/Runtime/Renderer/Private/HairStrands/HairStrandsEnvironment.cpp
--- a/Engine/Source/Runtime/Renderer/Private/HairStrands/HairStrandsEnvironment.cpp
+++ b/Engine/Source/Runtime/Renderer/Private/HairStrands/HairStrandsEnvironment.cpp
@@ -18,7 +18,7 @@
 #include "ScenePrivate.h"
 #include "RenderGraphEvent.h"
 #include "PostProcess/PostProcessing.h"
-#include "GpuDebugRendering.h"
+#include "ShaderDebug.h"
 #include "Lumen/LumenRadianceCache.h"
 #include "Lumen/LumenScreenProbeGather.h"
 
diff --git a/Engine/Source/Runtime/Renderer/Private/HairStrands/HairStrandsScatter.cpp b/Engine/Source/Runtime/Renderer/Private/HairStrands/HairStrandsScatter.cpp
--- a/Engine/Source/Runtime/Renderer/Private/HairStrands/HairStrandsScatter.cpp
+++ b/Engine/Source/Runtime/Renderer/Private/HairStrands/HairStrandsScatter.cpp
@@ -9,7 +9,7 @@
 #include "RenderGraphUtils.h"
 #include "PostProcessing.h"
 #include "HairStrandsRendering.h"
-#include "GpuDebugRendering.h"
+#include "ShaderDebug.h"
 
 /////////////////////////////////////////////////////////////////////////////////////////
 
diff --git a/Engine/Source/Runtime/Renderer/Private/HairStrands/HairStrandsTransmittance.cpp b/Engine/Source/Runtime/Renderer/Private/HairStrands/HairStrandsTransmittance.cpp
--- a/Engine/Source/Runtime/Renderer/Private/HairStrands/HairStrandsTransmittance.cpp
+++ b/Engine/Source/Runtime/Renderer/Private/HairStrands/HairStrandsTransmittance.cpp
@@ -14,7 +14,7 @@
 #include "SceneTextureParameters.h"
 #include "RenderGraphUtils.h"
 #include "PostProcessing.h"
-#include "GpuDebugRendering.h"
+#include "ShaderDebug.h"
 #include "ShaderPrintParameters.h"
 #include "LightSceneInfo.h"
 #include "ShaderPrint.h"
diff --git a/Engine/Source/Runtime/Renderer/Private/PostProcess/PostProcessing.cpp b/Engine/Source/Runtime/Renderer/Private/PostProcess/PostProcessing.cpp
--- a/Engine/Source/Runtime/Renderer/Private/PostProcess/PostProcessing.cpp
+++ b/Engine/Source/Runtime/Renderer/Private/PostProcess/PostProcessing.cpp
@@ -38,7 +38,7 @@
 #include "PostProcess/PostProcessStreamingAccuracyLegend.h"
 #include "PostProcess/PostProcessSubsurface.h"
 #include "ShaderPrint.h"
-#include "GpuDebugRendering.h"
+#include "ShaderDebug.h"
 #include "HighResScreenshot.h"
 #include "IHeadMountedDisplay.h"
 #include "IXRTrackingSystem.h"
diff --git a/Engine/Source/Runtime/Renderer/Private/SceneRendering.h b/Engine/Source/Runtime/Renderer/Private/SceneRendering.h
--- a/Engine/Source/Runtime/Renderer/Private/SceneRendering.h
+++ b/Engine/Source/Runtime/Renderer/Private/SceneRendering.h
@@ -30,7 +30,7 @@
 #include "Templates/UniquePtr.h"
 #include "RenderGraphUtils.h"
 #include "MeshDrawCommands.h"
-#include "GpuDebugRendering.h"
+#include "ShaderDebug.h"
 #include "PostProcess/PostProcessAmbientOcclusionMobile.h"
 #include "Nanite/NaniteRender.h"
 #include "VirtualShadowMaps/VirtualShadowMapArray.h"
diff --git a/Engine/Source/Runtime/Renderer/Private/ShaderDebug.cpp b/Engine/Source/Runtime/Renderer/Private/ShaderDebug.cpp
new file mode 100644
--- /dev/null
+++ b/Engine/Source/Runtime/Renderer/Private/ShaderDebug.cpp
@@ -0,0 +1,340 @@
+// Copyright Epic Games, Inc. All Rights Reserved.
+
+#include "ShaderDebug.h"
+#include "SceneRendering.h"
+#include "GlobalShader.h"
+#include "RenderGraphBuilder.h"
+#include "RenderGraphUtils.h"
+#include "Containers/ResourceArray.h"
+#include "CommonRenderResources.h"
+#include "ScenePrivate.h"
+
+namespace ShaderDrawDebug 
+{
+	// Console variables
+	static int32 GShaderDrawDebug_Enable = 1;
+	static FAutoConsoleVariableRef CVarShaderDrawEnable(
+		TEXT("r.ShaderDrawDebug"),
+		GShaderDrawDebug_Enable,
+		TEXT("ShaderDrawDebug debugging toggle.\n"),
+		ECVF_Cheat | ECVF_RenderThreadSafe);
+
+	static int32 GShaderDrawDebug_MaxElementCount;
+	static FAutoConsoleVariableRef CVarShaderDrawMaxElementCount(
+		TEXT("r.ShaderDrawDebug.MaxElementCount"),
+		GShaderDrawDebug_MaxElementCount,
+		TEXT("ShaderDraw output buffer size in element.\n"),
+		ECVF_Cheat | ECVF_RenderThreadSafe);
+
+	static TAutoConsoleVariable<int32> CVarShaderDrawLock(
+		TEXT("r.ShaderDrawDebug.Lock"),
+		0,
+		TEXT("Lock the shader draw buffer.\n"),
+		ECVF_Cheat | ECVF_RenderThreadSafe);
+
+	bool IsShaderDrawDebugEnabled()
+	{
+#if WITH_EDITOR
+		return GShaderDrawDebug_Enable > 0;
+#else
+		return false;
+#endif
+	}
+
+	static bool IsShaderDrawLocked()
+	{
+#if WITH_EDITOR
+		return CVarShaderDrawLock.GetValueOnRenderThread() > 0;
+#else
+		return false;
+#endif
+	}
+
+	static bool IsShaderDrawDebugEnabled(const EShaderPlatform Platform)
+	{
+		return IsFeatureLevelSupported(Platform, ERHIFeatureLevel::SM5) && IsPCPlatform(Platform) && !IsOpenGLPlatform(Platform);
+	}
+
+	void SetEnabled(bool bInEnabled)
+	{
+#if WITH_EDITOR
+		GShaderDrawDebug_Enable = bInEnabled ? 1 : 0;
+#endif		
+	}
+
+	void SetMaxElementCount(uint32 MaxCount)
+	{
+#if WITH_EDITOR
+		GShaderDrawDebug_MaxElementCount = FMath::Max(1024, int32(MaxCount));
+#endif		
+	}
+
+	uint32 GetMaxElementCount()
+	{
+#if WITH_EDITOR
+		return uint32(FMath::Max(1, GShaderDrawDebug_MaxElementCount));
+#else
+		return 0;
+#endif			
+	}
+
+	bool IsShaderDrawDebugEnabled(const FViewInfo& View)
+	{
+		return IsShaderDrawDebugEnabled() && IsShaderDrawDebugEnabled(View.GetShaderPlatform());
+	}
+
+	// Note: Unaligned structures used for structured buffers is an unsupported and/or sparsely
+	//         supported feature in VK (VK_EXT_scalar_block_layout) and Metal. Consequently, we
+	//         do manual packing in order to accommodate.
+	struct FPackedShaderDrawElement
+	{
+		// This is not packed as fp16 to be able to debug large scale data while preserving accuracy at short range.
+		float Pos0_ColorX[4];		// float3 pos0 + packed color0
+		float Pos1_ColorY[4];		// float3 pos1 + packed color1
+	};
+
+	//////////////////////////////////////////////////////////////////////////
+
+	class FShaderDrawDebugClearCS : public FGlobalShader
+	{
+		DECLARE_GLOBAL_SHADER(FShaderDrawDebugClearCS);
+		SHADER_USE_PARAMETER_STRUCT(FShaderDrawDebugClearCS, FGlobalShader);
+
+		BEGIN_SHADER_PARAMETER_STRUCT(FParameters, )
+		SHADER_PARAMETER_RDG_BUFFER_UAV(RWStructuredBuffer, DataBuffer)
+		SHADER_PARAMETER_RDG_BUFFER_UAV(RWBuffer, IndirectBuffer)
+		END_SHADER_PARAMETER_STRUCT()
+
+		static bool ShouldCompilePermutation(const FGlobalShaderPermutationParameters& Parameters)
+		{
+			return IsShaderDrawDebugEnabled(Parameters.Platform);
+		}
+
+		static void ModifyCompilationEnvironment(const FGlobalShaderPermutationParameters& Parameters, FShaderCompilerEnvironment& OutEnvironment)
+		{
+			FGlobalShader::ModifyCompilationEnvironment(Parameters, OutEnvironment);
+			OutEnvironment.SetDefine(TEXT("GPU_DEBUG_RENDERING"), 1);
+			OutEnvironment.SetDefine(TEXT("GPU_DEBUG_RENDERING_CLEAR_CS"), 1);
+		}
+	};
+
+	IMPLEMENT_GLOBAL_SHADER(FShaderDrawDebugClearCS, "/Engine/Private/ShaderDrawDebug.usf", "ShaderDrawDebugClearCS", SF_Compute);
+
+	//////////////////////////////////////////////////////////////////////////
+
+	class FShaderDrawDebugVS : public FGlobalShader
+	{
+		DECLARE_GLOBAL_SHADER(FShaderDrawDebugVS);
+		SHADER_USE_PARAMETER_STRUCT(FShaderDrawDebugVS, FGlobalShader);
+
+		BEGIN_SHADER_PARAMETER_STRUCT(FParameters, )
+			SHADER_PARAMETER_STRUCT_REF(FViewUniformShaderParameters, View)
+			SHADER_PARAMETER_SRV(StructuredBuffer, LockedShaderDrawDebugPrimitive)
+			SHADER_PARAMETER_RDG_BUFFER_SRV(StructuredBuffer, ShaderDrawDebugPrimitive)
+			RDG_BUFFER_ACCESS(IndirectBuffer, ERHIAccess::IndirectArgs)
+		END_SHADER_PARAMETER_STRUCT()
+
+		static bool ShouldCompilePermutation(const FGlobalShaderPermutationParameters& Parameters)
+		{
+			return IsShaderDrawDebugEnabled(Parameters.Platform);
+		}
+
+		static void ModifyCompilationEnvironment(const FGlobalShaderPermutationParameters& Parameters, FShaderCompilerEnvironment& OutEnvironment)
+		{
+			FGlobalShader::ModifyCompilationEnvironment(Parameters, OutEnvironment);
+			OutEnvironment.SetDefine(TEXT("GPU_DEBUG_RENDERING"), 1);
+			OutEnvironment.SetDefine(TEXT("GPU_DEBUG_RENDERING_VS"), 1);
+			OutEnvironment.SetDefine(TEXT("GPU_DEBUG_RENDERING_PS"), 0);
+		}
+	};
+
+	IMPLEMENT_GLOBAL_SHADER(FShaderDrawDebugVS, "/Engine/Private/ShaderDrawDebug.usf", "ShaderDrawDebugVS", SF_Vertex);
+
+	//////////////////////////////////////////////////////////////////////////
+
+	class FShaderDrawDebugPS : public FGlobalShader
+	{
+		DECLARE_GLOBAL_SHADER(FShaderDrawDebugPS);
+		SHADER_USE_PARAMETER_STRUCT(FShaderDrawDebugPS, FGlobalShader);
+		
+		BEGIN_SHADER_PARAMETER_STRUCT(FParameters, )
+			SHADER_PARAMETER_RDG_TEXTURE(Texture2D, DepthTexture)
+			SHADER_PARAMETER(FIntPoint, DepthTextureResolution)
+			SHADER_PARAMETER(FVector2D, DepthTextureInvResolution)
+			SHADER_PARAMETER_SAMPLER(SamplerState, DepthSampler)
+			RENDER_TARGET_BINDING_SLOTS()
+		END_SHADER_PARAMETER_STRUCT()
+
+		using FPermutationDomain = TShaderPermutationDomain<>;
+
+		static bool ShouldCompilePermutation(const FGlobalShaderPermutationParameters& Parameters)
+		{
+			return IsShaderDrawDebugEnabled(Parameters.Platform);
+		}
+
+		static void ModifyCompilationEnvironment(const FGlobalShaderPermutationParameters& Parameters, FShaderCompilerEnvironment& OutEnvironment)
+		{
+			FGlobalShader::ModifyCompilationEnvironment(Parameters, OutEnvironment);
+			OutEnvironment.SetDefine(TEXT("GPU_DEBUG_RENDERING"), 1);
+			OutEnvironment.SetDefine(TEXT("GPU_DEBUG_RENDERING_VS"), 0);
+			OutEnvironment.SetDefine(TEXT("GPU_DEBUG_RENDERING_PS"), 1);
+		}
+	};
+
+	IMPLEMENT_GLOBAL_SHADER(FShaderDrawDebugPS, "/Engine/Private/ShaderDrawDebug.usf", "ShaderDrawDebugPS", SF_Pixel);
+
+	BEGIN_SHADER_PARAMETER_STRUCT(FShaderDrawVSPSParameters , )
+		SHADER_PARAMETER_STRUCT_INCLUDE(FShaderDrawDebugVS::FParameters, ShaderDrawVSParameters)
+		SHADER_PARAMETER_STRUCT_INCLUDE(FShaderDrawDebugPS::FParameters, ShaderDrawPSParameters)
+	END_SHADER_PARAMETER_STRUCT()
+
+	//////////////////////////////////////////////////////////////////////////
+
+	static void InternalDrawView(
+		FRDGBuilder& GraphBuilder,
+		const FViewInfo& View,
+		FRDGBufferRef DataBuffer,
+		FRDGBufferRef IndirectBuffer,
+		FRDGTextureRef OutputTexture,
+		FRDGTextureRef DepthTexture)
+	{
+		TShaderMapRef<FShaderDrawDebugVS> VertexShader(View.ShaderMap);
+		TShaderMapRef<FShaderDrawDebugPS> PixelShader(View.ShaderMap);
+
+		FShaderDrawVSPSParameters* PassParameters = GraphBuilder.AllocParameters<FShaderDrawVSPSParameters >();
+		PassParameters->ShaderDrawPSParameters.RenderTargets[0] = FRenderTargetBinding(OutputTexture, ERenderTargetLoadAction::ELoad);
+		//PassParameters->ShaderDrawPSParameters.RenderTargets.DepthStencil = FDepthStencilBinding(DepthTexture, ERenderTargetLoadAction::ELoad, FExclusiveDepthStencil::DepthRead_StencilNop);
+		PassParameters->ShaderDrawPSParameters.DepthTexture = DepthTexture;
+		PassParameters->ShaderDrawPSParameters.DepthTextureResolution = FIntPoint(DepthTexture->Desc.Extent.X, DepthTexture->Desc.Extent.Y);
+		PassParameters->ShaderDrawPSParameters.DepthTextureInvResolution = FVector2D(1.f / DepthTexture->Desc.Extent.X, 1.f / DepthTexture->Desc.Extent.Y);
+		PassParameters->ShaderDrawPSParameters.DepthSampler = TStaticSamplerState<SF_Point, AM_Clamp, AM_Clamp, AM_Clamp>::GetRHI();
+		PassParameters->ShaderDrawVSParameters.View = View.ViewUniformBuffer;
+		PassParameters->ShaderDrawVSParameters.ShaderDrawDebugPrimitive = GraphBuilder.CreateSRV(DataBuffer);
+		PassParameters->ShaderDrawVSParameters.IndirectBuffer = IndirectBuffer;
+
+		ValidateShaderParameters(PixelShader, PassParameters->ShaderDrawPSParameters);
+		ClearUnusedGraphResources(PixelShader, &PassParameters->ShaderDrawPSParameters, { IndirectBuffer });
+		ValidateShaderParameters(VertexShader, PassParameters->ShaderDrawVSParameters);
+		ClearUnusedGraphResources(VertexShader, &PassParameters->ShaderDrawVSParameters, { IndirectBuffer });
+
+		GraphBuilder.AddPass(
+			RDG_EVENT_NAME("ShaderDrawDebug"),
+			PassParameters,
+			ERDGPassFlags::Raster,
+			[VertexShader, PixelShader, PassParameters, IndirectBuffer](FRHICommandList& RHICmdList)
+			{
+				// Marks the indirect draw parameter as used by the pass, given it's not used directly by any of the shaders.
+				PassParameters->ShaderDrawVSParameters.IndirectBuffer->MarkResourceAsUsed();
+
+				FGraphicsPipelineStateInitializer GraphicsPSOInit;
+				RHICmdList.ApplyCachedRenderTargets(GraphicsPSOInit);
+				GraphicsPSOInit.DepthStencilState = TStaticDepthStencilState<false, CF_Always>::GetRHI();
+				GraphicsPSOInit.BlendState = TStaticBlendState<CW_RGBA, BO_Add, BF_One, BF_InverseSourceAlpha, BO_Add, BF_Zero, BF_One>::GetRHI(); // Premultiplied-alpha composition
+				GraphicsPSOInit.RasterizerState = TStaticRasterizerState<FM_Solid, CM_None, true>::GetRHI();
+				GraphicsPSOInit.PrimitiveType = PT_LineList;
+				GraphicsPSOInit.BoundShaderState.VertexDeclarationRHI = GEmptyVertexDeclaration.VertexDeclarationRHI;
+				GraphicsPSOInit.BoundShaderState.VertexShaderRHI = VertexShader.GetVertexShader();
+				GraphicsPSOInit.BoundShaderState.PixelShaderRHI = PixelShader.GetPixelShader();
+				SetGraphicsPipelineState(RHICmdList, GraphicsPSOInit);
+
+				SetShaderParameters(RHICmdList, VertexShader, VertexShader.GetVertexShader(), PassParameters->ShaderDrawVSParameters);
+				SetShaderParameters(RHICmdList, PixelShader, PixelShader.GetPixelShader(), PassParameters->ShaderDrawPSParameters);
+
+				// Marks the indirect draw parameter as used by the pass, given it's not used directly by any of the shaders.
+				FRHIBuffer* IndirectBufferRHI = PassParameters->ShaderDrawVSParameters.IndirectBuffer->GetIndirectRHICallBuffer();
+				check(IndirectBufferRHI != nullptr);
+				RHICmdList.DrawPrimitiveIndirect(IndirectBufferRHI, 0);
+			});
+	}
+
+	void BeginView(FRDGBuilder& GraphBuilder, FViewInfo& View)
+	{
+		if (!IsShaderDrawDebugEnabled(View) || !IsShaderDrawDebugEnabled())
+		{
+			return;
+		}
+
+		FSceneViewState* ViewState = View.ViewState;
+
+		const bool bLockBufferThisFrame = IsShaderDrawLocked() && ViewState && !ViewState->ShaderDrawDebugStateData.bIsLocked;
+		ERDGBufferFlags Flags = bLockBufferThisFrame ? ERDGBufferFlags::MultiFrame : ERDGBufferFlags::None;
+
+		FRDGBufferRef DataBuffer = GraphBuilder.CreateBuffer(FRDGBufferDesc::CreateStructuredDesc(sizeof(FPackedShaderDrawElement), GetMaxElementCount()), TEXT("ShaderDrawDataBuffer"), Flags);
+		FRDGBufferRef IndirectBuffer = GraphBuilder.CreateBuffer(FRDGBufferDesc::CreateIndirectDesc<FRHIDrawIndirectParameters>(1), TEXT("ShaderDrawDataIndirectBuffer"), Flags);
+
+		FShaderDrawDebugClearCS::FParameters* Parameters = GraphBuilder.AllocParameters<FShaderDrawDebugClearCS::FParameters>();
+		Parameters->DataBuffer = GraphBuilder.CreateUAV(DataBuffer);
+		Parameters->IndirectBuffer = GraphBuilder.CreateUAV(IndirectBuffer);
+
+		TShaderMapRef<FShaderDrawDebugClearCS> ComputeShader(View.ShaderMap);
+
+		// Note: we do not call ClearUnusedGraphResources here as we want to for the allocation of DataBuffer
+		GraphBuilder.AddPass(
+			RDG_EVENT_NAME("ShaderDrawClear"),
+			Parameters,
+			ERDGPassFlags::Compute,
+			[Parameters, ComputeShader](FRHICommandList& RHICmdList)
+		{
+			FComputeShaderUtils::Dispatch(RHICmdList, ComputeShader, *Parameters, FIntVector(1,1,1));
+		});
+
+
+		View.ShaderDrawData.Buffer = DataBuffer;
+		View.ShaderDrawData.IndirectBuffer = IndirectBuffer;
+		View.ShaderDrawData.CursorPosition = View.CursorPos;
+
+		if (IsShaderDrawLocked() && ViewState && !ViewState->ShaderDrawDebugStateData.bIsLocked)
+		{
+			ViewState->ShaderDrawDebugStateData.Buffer = GraphBuilder.ConvertToExternalBuffer(View.ShaderDrawData.Buffer);
+			ViewState->ShaderDrawDebugStateData.IndirectBuffer = GraphBuilder.ConvertToExternalBuffer(View.ShaderDrawData.IndirectBuffer);
+			ViewState->ShaderDrawDebugStateData.bIsLocked = true;
+		}
+
+		if (!IsShaderDrawLocked() && ViewState && ViewState->ShaderDrawDebugStateData.bIsLocked)
+		{
+			ViewState->ShaderDrawDebugStateData.Buffer = nullptr;
+			ViewState->ShaderDrawDebugStateData.IndirectBuffer = nullptr;
+			ViewState->ShaderDrawDebugStateData.bIsLocked = false;
+		}
+	}
+
+	void DrawView(FRDGBuilder& GraphBuilder, const FViewInfo& View, FRDGTextureRef OutputTexture, FRDGTextureRef DepthTexture)
+	{
+		if (!IsShaderDrawDebugEnabled(View))
+		{
+			return;
+		}
+
+		{
+			FRDGBufferRef DataBuffer = View.ShaderDrawData.Buffer;
+			FRDGBufferRef IndirectBuffer = View.ShaderDrawData.IndirectBuffer;
+			InternalDrawView(GraphBuilder, View, DataBuffer, IndirectBuffer, OutputTexture, DepthTexture);
+		}
+
+		if (View.ViewState && View.ViewState->ShaderDrawDebugStateData.bIsLocked)
+		{
+			FRDGBufferRef DataBuffer = GraphBuilder.RegisterExternalBuffer(View.ViewState->ShaderDrawDebugStateData.Buffer);
+			FRDGBufferRef IndirectBuffer = GraphBuilder.RegisterExternalBuffer(View.ViewState->ShaderDrawDebugStateData.IndirectBuffer);
+			InternalDrawView(GraphBuilder, View, DataBuffer, IndirectBuffer, OutputTexture, DepthTexture);
+		}
+	}
+
+	void EndView(FViewInfo& View)
+	{
+		if (!IsShaderDrawDebugEnabled(View))
+		{
+			return;
+		}
+	}
+
+	void SetParameters(FRDGBuilder& GraphBuilder, const FShaderDrawDebugData& Data, FShaderDrawDebugParameters& OutParameters)
+	{
+		FRDGBufferRef DataBuffer = Data.Buffer;
+		FRDGBufferRef IndirectBuffer = Data.IndirectBuffer;
+
+		OutParameters.ShaderDrawCursorPos = Data.CursorPosition;
+		OutParameters.ShaderDrawMaxElementCount = GetMaxElementCount();
+		OutParameters.OutShaderDrawPrimitive = GraphBuilder.CreateUAV(DataBuffer);
+		OutParameters.OutputShaderDrawIndirect = GraphBuilder.CreateUAV(IndirectBuffer);
+	}
+}
diff --git a/Engine/Source/Runtime/Renderer/Private/VolumetricCloudRendering.cpp b/Engine/Source/Runtime/Renderer/Private/VolumetricCloudRendering.cpp
--- a/Engine/Source/Runtime/Renderer/Private/VolumetricCloudRendering.cpp
+++ b/Engine/Source/Runtime/Renderer/Private/VolumetricCloudRendering.cpp
@@ -17,7 +17,7 @@
 #include "SkyAtmosphereRendering.h"
 #include "VolumeLighting.h"
 #include "DynamicPrimitiveDrawing.h"
-#include "GpuDebugRendering.h"
+#include "ShaderDebug.h"
 #include "CanvasTypes.h"
 #include "RenderTargetTemp.h"
 #include "VolumetricRenderTarget.h"
diff --git a/Engine/Source/Runtime/Renderer/Public/GpuDebugRendering.h b/Engine/Source/Runtime/Renderer/Public/GpuDebugRendering.h
--- a/Engine/Source/Runtime/Renderer/Public/GpuDebugRendering.h
+++ b/Engine/Source/Runtime/Renderer/Public/GpuDebugRendering.h
@@ -2,45 +2,6 @@
 
 #pragma once
 
-#include "RendererInterface.h"
-#include "RHICommandList.h"
-#include "RenderGraph.h"
-
-class FViewInfo;
-
-struct FShaderDrawDebugData
-{
-	FIntPoint CursorPosition;
-	FRDGBufferRef Buffer = nullptr;
-	FRDGBufferRef IndirectBuffer = nullptr;
-};
-
-namespace ShaderDrawDebug 
-{
-	// Call this to know if this is even just available (for exemple in shipping mode buffers won't exists)
-	RENDERER_API bool IsShaderDrawDebugEnabled();
-	RENDERER_API void SetEnabled(bool bEnable);
-	RENDERER_API void SetMaxElementCount(uint32 MaxCount);
-	RENDERER_API uint32 GetMaxElementCount();
-
-	// Call this to know if a view can render this debug information
-	bool IsShaderDrawDebugEnabled(const FViewInfo& View);
-
-	// Allocate the debug print buffer associated with the view
-	void BeginView(FRDGBuilder& GraphBuilder, FViewInfo& View);
-	// Draw info from the debug print buffer to the given output target
-	void DrawView(FRDGBuilder& GraphBuilder, const FViewInfo& View, FRDGTextureRef OutputTexture, FRDGTextureRef DepthTexture);
-	// Release the debug print buffer associated with the view
-	void EndView(FViewInfo& View);
-
-	// The structure to be set on the debug shader outputting debug primitive
-	BEGIN_SHADER_PARAMETER_STRUCT(FShaderDrawDebugParameters, )
-		SHADER_PARAMETER(FIntPoint, ShaderDrawCursorPos)
-		SHADER_PARAMETER(int32, ShaderDrawMaxElementCount)
-		SHADER_PARAMETER_RDG_BUFFER_UAV(RWStructuredBuffer, OutShaderDrawPrimitive)
-		SHADER_PARAMETER_RDG_BUFFER_UAV(RWBuffer<uint>, OutputShaderDrawIndirect)
-	END_SHADER_PARAMETER_STRUCT()
-
-	// Call this to fill the FShaderDrawParameters
-	RENDERER_API void SetParameters(FRDGBuilder& GraphBuilder, const FShaderDrawDebugData& Data, FShaderDrawDebugParameters& OutParameters);
-}
+// This file is a deprecated include file. 
+// Forward to the new include file
+#include "ShaderDebug.h"
diff --git a/Engine/Source/Runtime/Renderer/Public/HairStrandsInterface.h b/Engine/Source/Runtime/Renderer/Public/HairStrandsInterface.h
--- a/Engine/Source/Runtime/Renderer/Public/HairStrandsInterface.h
+++ b/Engine/Source/Runtime/Renderer/Public/HairStrandsInterface.h
@@ -13,7 +13,7 @@
 #include "Shader.h"
 #include "RenderResource.h"
 #include "RenderGraphResources.h"
-#include "GpuDebugRendering.h"
+#include "ShaderDebug.h"
 
 class UTexture2D;
 
diff --git a/Engine/Source/Runtime/Renderer/Public/ShaderDebug.h b/Engine/Source/Runtime/Renderer/Public/ShaderDebug.h
new file mode 100644
--- /dev/null
+++ b/Engine/Source/Runtime/Renderer/Public/ShaderDebug.h
@@ -0,0 +1,46 @@
+// Copyright Epic Games, Inc. All Rights Reserved.
+
+#pragma once
+
+#include "RendererInterface.h"
+#include "RHICommandList.h"
+#include "RenderGraph.h"
+
+class FViewInfo;
+
+struct FShaderDrawDebugData
+{
+	FIntPoint CursorPosition;
+	FRDGBufferRef Buffer = nullptr;
+	FRDGBufferRef IndirectBuffer = nullptr;
+};
+
+namespace ShaderDrawDebug 
+{
+	// Call this to know if this is even just available (for exemple in shipping mode buffers won't exists)
+	RENDERER_API bool IsShaderDrawDebugEnabled();
+	RENDERER_API void SetEnabled(bool bEnable);
+	RENDERER_API void SetMaxElementCount(uint32 MaxCount);
+	RENDERER_API uint32 GetMaxElementCount();
+
+	// Call this to know if a view can render this debug information
+	bool IsShaderDrawDebugEnabled(const FViewInfo& View);
+
+	// Allocate the debug print buffer associated with the view
+	void BeginView(FRDGBuilder& GraphBuilder, FViewInfo& View);
+	// Draw info from the debug print buffer to the given output target
+	void DrawView(FRDGBuilder& GraphBuilder, const FViewInfo& View, FRDGTextureRef OutputTexture, FRDGTextureRef DepthTexture);
+	// Release the debug print buffer associated with the view
+	void EndView(FViewInfo& View);
+
+	// The structure to be set on the debug shader outputting debug primitive
+	BEGIN_SHADER_PARAMETER_STRUCT(FShaderDrawDebugParameters, )
+		SHADER_PARAMETER(FIntPoint, ShaderDrawCursorPos)
+		SHADER_PARAMETER(int32, ShaderDrawMaxElementCount)
+		SHADER_PARAMETER_RDG_BUFFER_UAV(RWStructuredBuffer, OutShaderDrawPrimitive)
+		SHADER_PARAMETER_RDG_BUFFER_UAV(RWBuffer<uint>, OutputShaderDrawIndirect)
+	END_SHADER_PARAMETER_STRUCT()
+
+	// Call this to fill the FShaderDrawParameters
+	RENDERER_API void SetParameters(FRDGBuilder& GraphBuilder, const FShaderDrawDebugData& Data, FShaderDrawDebugParameters& OutParameters);
+}
