Fix memory stomp in SystemTextures.cpp
 - AsciiTexture initialization writes past the end of the RHI allocation
 - Loop was written as though the pixel format of the texture is PF_RGBA8, but is actually PF_R8, causing a 3 byte stomp at the end of the allocation.

diff --git a/Engine/Source/Runtime/Renderer/Private/SystemTextures.cpp b/Engine/Source/Runtime/Renderer/Private/SystemTextures.cpp
--- a/Engine/Source/Runtime/Renderer/Private/SystemTextures.cpp
+++ b/Engine/Source/Runtime/Renderer/Private/SystemTextures.cpp
@@ -328,9 +328,6 @@ void FSystemTextures::InitializeFeatureLevelDependentTextures(FRHICommandListImm
 				const uint32 C = y < CharacterRes/2 ? Hi : Lo;
 				const bool bFilled = (C & (1u << ((y * CharacterRes + x) & 31u)));
 				Dest[0] = bFilled ? 0xFF : 0x0;
-				Dest[1] = bFilled ? 0xFF : 0x0;
-				Dest[2] = bFilled ? 0xFF : 0x0;
-				Dest[3] = bFilled ? 0xFF : 0x0;
 			}
 		}
 		RHICmdList.UnlockTexture2D((FTexture2DRHIRef&)AsciiTexture->GetRenderTargetItem().ShaderResourceTexture, 0, false);
