Fix crash in batched instance culling from using null VisibleInstanceFlags when the scene is empty
- also removed some dead code (debug code that was inactive)

diff --git a/Engine/Shaders/Private/InstanceCulling/BuildInstanceDrawCommands.usf b/Engine/Shaders/Private/InstanceCulling/BuildInstanceDrawCommands.usf
--- a/Engine/Shaders/Private/InstanceCulling/BuildInstanceDrawCommands.usf
+++ b/Engine/Shaders/Private/InstanceCulling/BuildInstanceDrawCommands.usf
@@ -113,292 +113,6 @@ void WriteInstanceId(uint Offset, uint InstanceId, uint ViewIdIndex, uint DrawCo
 #endif
 }
 
-#ifndef ENABLE_DETERMINISTIC_INSTANCE_CULLING
-	#define ENABLE_DETERMINISTIC_INSTANCE_CULLING 0
-#endif
-
-#if ENABLE_DETERMINISTIC_INSTANCE_CULLING
-
-RWStructuredBuffer<uint2> InstanceCountsOut;
-
-[numthreads(NUM_THREADS_PER_GROUP, 1, 1)]
-void ComputeInstanceIdOutputSize(uint GroupThreadIndex : SV_GroupIndex, uint GroupId : SV_GroupID)
-{
-	int CommandSlotIndex = int(GroupId);
-	FPrimCullingCommand Cmd = PrimitiveCullingCommands[CommandSlotIndex];
-
-	// 1. Figure out the total number of instances
-	PrimitiveInstanceCountGroupShared = 0;
-	RunInstanceCountGroupShared = 0;
-	FinalInstanceCountGroupShared = 0;
-	GroupMemoryBarrierWithGroupSync();
-
-	// Get start of next CMD or go to end.
-	uint LastOffset = (CommandSlotIndex + 1 < NumCommands) ? PrimitiveCullingCommands[CommandSlotIndex + 1].FirstPrimitiveIdOffset : NumPrimitiveIds;
-	// Note: the max here is purely to safe-guard against corrupted buffer data leading to an infinite loop
-	uint NumCmdPrimitiveIds = min(NumPrimitiveIds, LastOffset - Cmd.FirstPrimitiveIdOffset);
-
-	uint LastInstanceRunOffset = (CommandSlotIndex + 1 < NumCommands) ? PrimitiveCullingCommands[CommandSlotIndex + 1].FirstInstanceRunOffset : NumInstanceRuns;
-	// Note: the min here is purely to safe-guard against corrupted buffer data leading to an infinite loop
-	uint NumCmdInstanceRuns = min(NumInstanceRuns, LastInstanceRunOffset - Cmd.FirstInstanceRunOffset);
-
-	for (uint Index = GroupThreadIndex; Index < NumCmdPrimitiveIds + NumCmdInstanceRuns; Index += NUM_THREADS_PER_GROUP)
-	{
-		uint Num = 0;
-		if (Index < NumCmdPrimitiveIds)
-		{
-			uint PrimitiveId = TranslatePrimitiveId(PrimitiveIds[Cmd.FirstPrimitiveIdOffset + Index]);
-			if (GetPrimitiveData(PrimitiveId).InstanceDataOffset >= 0)
-			{
-				Num = GetPrimitiveData(PrimitiveId).NumInstanceDataEntries;
-			}
-			else
-			{
-				// Fake entry for invalid ones -> shaders that will not use the instancing data anyway
-				Num = 1;
-			}
-			InterlockedAdd(PrimitiveInstanceCountGroupShared, Num);
-		}
-		else
-		{
-			FInstanceRun Run = InstanceRuns[Cmd.FirstInstanceRunOffset + Index - NumCmdPrimitiveIds];
-			Num = (Run.EndInclusive + 1) - Run.Start;
-
-			InterlockedAdd(RunInstanceCountGroupShared, Num);
-		}
-
-	}
-
-	GroupMemoryBarrierWithGroupSync();
-	if (GroupThreadIndex == 0)
-	{
-		InstanceCountsOut[CommandSlotIndex] = uint2(PrimitiveInstanceCountGroupShared, RunInstanceCountGroupShared);
-	}
-}
-
-
-groupshared uint GroupSharedPrefixBuffer[NUM_THREADS_PER_GROUP];
-
-
-uint CalcPrefixSum(uint Value, inout uint InclusiveSum, uint ThreadIdx)
-{
-	// TODO: make efficient using: wave shuffle intrinsics when applicable
-	InclusiveSum = Value;
-	GroupSharedPrefixBuffer[ThreadIdx] = InclusiveSum;
-	UNROLL
-	for (uint Offset = 1; Offset < NUM_THREADS_PER_GROUP; Offset *= 2)
-	{
-		// Must sync so all threads have written data before loading it
-		GroupMemoryBarrierWithGroupSync();
-		if (ThreadIdx >= Offset)
-		{
-			InclusiveSum += GroupSharedPrefixBuffer[ThreadIdx - Offset];
-		}
-		// Must sync so all threads have loaded data before overwriting 
-		GroupMemoryBarrierWithGroupSync();
-		if (ThreadIdx >= Offset)
-		{
-			GroupSharedPrefixBuffer[ThreadIdx] = InclusiveSum;
-		}
-	}
-
-	uint ExclusiveSum = InclusiveSum - Value;
-	GroupMemoryBarrierWithGroupSync();
-	return ExclusiveSum;
-}
-
-uint CalcBinaryPrefixSum(bool bItemValid, inout uint InclusiveSum, uint ThreadIdx)
-{
-	// TODO: make efficient using: a ballot-based prefix for each warp/wave, or the built-in one in SM6, and then reduce the 
-	// wave offsets by brute force. 
-	return CalcPrefixSum(bItemValid ? 1 : 0, InclusiveSum, ThreadIdx);
-}
-
-StructuredBuffer<uint2> InstanceCounts;
-
-[numthreads(NUM_THREADS_PER_GROUP, 1, 1)]
-void CalcOutputOffsets(uint GroupThreadIndex : SV_GroupIndex, uint GroupId : SV_GroupID)
-{
-	// Run single thread group and compute prefix sum.
-	if (GroupThreadIndex == 0)
-	{
-		// We assume there is only one of these running at any given time (beware!)
-		// Load the start offset in the global output buffer and broadcast
-		FinalInstanceCountGroupShared = OutputOffsetBufferOut[0];
-	}
-	GroupMemoryBarrierWithGroupSync();
-
-	for (uint BaseIndex = 0; BaseIndex < NumCommands; BaseIndex += NUM_THREADS_PER_GROUP)
-	{
-		uint CommandSlotIndex = BaseIndex + GroupThreadIndex;
-		bool bValidIndex = CommandSlotIndex < NumCommands;
-
-		uint CommandInstanceCount = 0;
-		if (bValidIndex)
-		{
-			uint PrimitiveInstanceCount = InstanceCounts[CommandSlotIndex].x;
-			uint RunInstanceCount = InstanceCounts[CommandSlotIndex].y;
-			CommandInstanceCount = NumViewIds * (PrimitiveInstanceCount + RunInstanceCount);
-		}
-
-		uint InclusiveSum = 0U;
-		uint OutputOffset = CalcPrefixSum(CommandInstanceCount, InclusiveSum, GroupThreadIndex);
-		if (bValidIndex)
-		{
-			InstanceIdOffsetBufferOut[CommandSlotIndex] = OutputOffset + FinalInstanceCountGroupShared;
-		}
-		// broadcast the total
-		GroupMemoryBarrierWithGroupSync();
-		if (GroupThreadIndex == NUM_THREADS_PER_GROUP - 1)
-		{
-			FinalInstanceCountGroupShared += InclusiveSum;
-		}
-		GroupMemoryBarrierWithGroupSync();
-	}
-	GroupMemoryBarrierWithGroupSync();
-	if (GroupThreadIndex == 0)
-	{
-		// Write back the total count
-		OutputOffsetBufferOut[0] = FinalInstanceCountGroupShared;
-	}
-}
-
-Buffer<uint> InstanceIdOffsetBuffer;
-
-[numthreads(NUM_THREADS_PER_GROUP, 1, 1)]
-void OutputInstanceIdsAtOffset(uint GroupThreadIndex : SV_GroupIndex, uint GroupId : SV_GroupID)
-{
-	FinalInstanceCountGroupShared = 0;
-	GroupMemoryBarrierWithGroupSync();
-
-	int CommandSlotIndex = int(GroupId);
-	FPrimCullingCommand Cmd = PrimitiveCullingCommands[CommandSlotIndex];
-
-	uint PrimitiveInstanceCount = InstanceCounts[CommandSlotIndex].x; 
-	uint RunInstanceCount = InstanceCounts[CommandSlotIndex].y;
-	uint InstanceIdOutputOffset = InstanceIdOffsetBuffer[CommandSlotIndex];
-
-	// Note: NumViewIds - replicate for each view in instanced stereo or other multi-view
-	uint CommandInstanceCount = NumViewIds * (PrimitiveInstanceCount + RunInstanceCount);
-
-	// Bail if no instances to output
-	if (CommandInstanceCount > 0)
-	{
-		// 3. Broadcast to all threads
-		uint PrimitiveId = TranslatePrimitiveId(PrimitiveIds[Cmd.FirstPrimitiveIdOffset]);
-		if (GroupThreadIndex == 0)
-		{
-			WriteInstanceId(InstanceIdOutputOffset, GetPrimitiveData(PrimitiveId).InstanceDataOffset, 0U, CommandSlotIndex);
-		}
-		GroupMemoryBarrierWithGroupSync();
-
-		int CurrentPrimitiveIdIndex = -1;
-
-		uint CurrentRangeStart = 0;
-		uint CurrentRangeEnd = 0;
-		int InstanceDataOffset = 0;
-
-		// 4. loop over output primitive instance number
-		for (uint BaseOutputIndex1 = 0; BaseOutputIndex1 < PrimitiveInstanceCount; BaseOutputIndex1 += NUM_THREADS_PER_GROUP)
-		{
-			uint OutputIndex = BaseOutputIndex1 + GroupThreadIndex;
-
-			bool bValidItemIndex = OutputIndex < PrimitiveInstanceCount;
-
-			if (bValidItemIndex)
-			{
-				// Naive search for range containing the output index (reasonably efficient for large runs, bad for many small runs)
-				while (OutputIndex >= CurrentRangeEnd)
-				{
-					CurrentPrimitiveIdIndex += 1;
-					uint PrimitiveId = TranslatePrimitiveId(PrimitiveIds[Cmd.FirstPrimitiveIdOffset + CurrentPrimitiveIdIndex]);
-
-					CurrentRangeStart = CurrentRangeEnd;
-					InstanceDataOffset = GetPrimitiveData(PrimitiveId).InstanceDataOffset;
-					CurrentRangeEnd += max(1, GetPrimitiveData(PrimitiveId).NumInstanceDataEntries);
-				}
-			}
-			uint OutId = (bValidItemIndex && InstanceDataOffset >= 0) ? uint(InstanceDataOffset) + OutputIndex - CurrentRangeStart : 0U;
-			for (uint ViewIdIndex = 0; ViewIdIndex < NumViewIds; ++ViewIdIndex)
-			{
-				bool bWriteItem = bValidItemIndex && (InstanceDataOffset < 0 || TestInstanceVisibilityBit(OutId, ViewIds[ViewIdIndex]));
-				uint InclusiveSum = 0U;
-				uint OutputOffset = CalcBinaryPrefixSum(bWriteItem, InclusiveSum, GroupThreadIndex);
-				if (bWriteItem)
-				{
-					WriteInstanceId(InstanceIdOutputOffset + FinalInstanceCountGroupShared + OutputOffset, OutId, ViewIdIndex, CommandSlotIndex);
-				}
-				// broadcast the total
-				GroupMemoryBarrierWithGroupSync();
-				if (GroupThreadIndex == NUM_THREADS_PER_GROUP - 1)
-				{
-					FinalInstanceCountGroupShared += InclusiveSum;
-				}
-				GroupMemoryBarrierWithGroupSync();
-
-			}
-		}
-		// repeat process for the runs
-		CurrentRangeStart = 0;
-		CurrentRangeEnd = 0;
-		InstanceDataOffset = 0;
-		FInstanceRun Run = (FInstanceRun)0;
-		int CurrentRunIndex = -1;
-		for (uint BaseOutputIndex2 = 0; BaseOutputIndex2 < RunInstanceCount; BaseOutputIndex2 += NUM_THREADS_PER_GROUP)
-		{
-			uint OutputIndex = BaseOutputIndex2 + GroupThreadIndex;
-			bool bValidItemIndex = OutputIndex < RunInstanceCount;
-
-			if (bValidItemIndex)
-			{
-				// Naive search for range containing the output index (reasonably efficient for large runs, bad for many small runs)
-				while (OutputIndex >= CurrentRangeEnd)
-				{
-					CurrentRunIndex += 1;
-					Run = InstanceRuns[Cmd.FirstInstanceRunOffset + CurrentRunIndex];
-
-					CurrentRangeStart = CurrentRangeEnd;
-					CurrentRangeEnd += (Run.EndInclusive + 1) - Run.Start;
-					InstanceDataOffset = GetPrimitiveData(Run.PrimitiveId).InstanceDataOffset;
-				}
-			}
-			uint OutId = bValidItemIndex ? (uint(InstanceDataOffset) + Run.Start + OutputIndex - CurrentRangeStart) : 0;
-			for (uint ViewIdIndex = 0; ViewIdIndex < NumViewIds; ++ViewIdIndex)
-			{
-				bool bWriteItem = bValidItemIndex && TestInstanceVisibilityBit(OutId, ViewIds[ViewIdIndex]);
-
-				uint InclusiveSum = 0U;
-				uint OutputOffset = CalcBinaryPrefixSum(bWriteItem, InclusiveSum, GroupThreadIndex);
-				if (bWriteItem)
-				{
-					WriteInstanceId(InstanceIdOutputOffset + FinalInstanceCountGroupShared + OutputOffset, OutId, ViewIdIndex, CommandSlotIndex);
-				}
-				// broadcast the total
-				GroupMemoryBarrierWithGroupSync();
-				if (GroupThreadIndex == NUM_THREADS_PER_GROUP - 1)
-				{
-					FinalInstanceCountGroupShared += InclusiveSum;
-				}
-				GroupMemoryBarrierWithGroupSync();
-			}
-		}
-	}
-	GroupMemoryBarrierWithGroupSync();
-
-	// Generate draw command
-	if (GroupThreadIndex == 0)
-	{
-		DrawIndirectArgsBufferOut[CommandSlotIndex * INDIRECT_ARGS_NUM_WORDS + 0] = Cmd.NumVerticesOrIndices;
-		DrawIndirectArgsBufferOut[CommandSlotIndex * INDIRECT_ARGS_NUM_WORDS + 1] = FinalInstanceCountGroupShared;
-		DrawIndirectArgsBufferOut[CommandSlotIndex * INDIRECT_ARGS_NUM_WORDS + 2] = Cmd.FirstIndex;
-		DrawIndirectArgsBufferOut[CommandSlotIndex * INDIRECT_ARGS_NUM_WORDS + 3] = Cmd.BaseVertexIndex;
-		DrawIndirectArgsBufferOut[CommandSlotIndex * INDIRECT_ARGS_NUM_WORDS + 4] = 0U;
-
-	}
-	GroupMemoryBarrierWithGroupSync();
-}
-
-#endif // ENABLE_DETERMINISTIC_INSTANCE_CULLING
 
 /**
  */
diff --git a/Engine/Source/Runtime/Renderer/Private/InstanceCulling/InstanceCullingContext.cpp b/Engine/Source/Runtime/Renderer/Private/InstanceCulling/InstanceCullingContext.cpp
--- a/Engine/Source/Runtime/Renderer/Private/InstanceCulling/InstanceCullingContext.cpp
+++ b/Engine/Source/Runtime/Renderer/Private/InstanceCulling/InstanceCullingContext.cpp
@@ -12,7 +12,6 @@
 #include "ScenePrivate.h"
 #include "InstanceCulling/InstanceCullingManager.h"
 
-#define ENABLE_DETERMINISTIC_INSTANCE_CULLING 0
 
 IMPLEMENT_STATIC_UNIFORM_BUFFER_SLOT(InstanceCullingUbSlot);
 IMPLEMENT_STATIC_UNIFORM_BUFFER_STRUCT(FInstanceCullingGlobalUniforms, "InstanceCulling", InstanceCullingUbSlot);
@@ -93,163 +92,6 @@ void FInstanceCullingContext::AddInstanceRunToCullingCommand(int32 ScenePrimitiv
 	}
 }
 
-#if ENABLE_DETERMINISTIC_INSTANCE_CULLING
-
-class FComputeInstanceIdOutputSizeCs : public FGlobalShader
-{
-	DECLARE_GLOBAL_SHADER(FComputeInstanceIdOutputSizeCs);
-	SHADER_USE_PARAMETER_STRUCT(FComputeInstanceIdOutputSizeCs, FGlobalShader)
-public:
-	class FCullInstancesDim : SHADER_PERMUTATION_BOOL("CULL_INSTANCES");
-	using FPermutationDomain = TShaderPermutationDomain<FCullInstancesDim>;
-
-	static constexpr int32 NumThreadsPerGroup = 64;
-
-	static bool ShouldCompilePermutation(const FGlobalShaderPermutationParameters& Parameters)
-	{
-		return UseGPUScene(Parameters.Platform);
-	}
-
-	static void ModifyCompilationEnvironment(const FGlobalShaderPermutationParameters& Parameters, FShaderCompilerEnvironment& OutEnvironment)
-	{
-		FGlobalShader::ModifyCompilationEnvironment(Parameters, OutEnvironment);
-		OutEnvironment.SetDefine(TEXT("INDIRECT_ARGS_NUM_WORDS"), FInstanceCullingContext::IndirectArgsNumWords);
-		OutEnvironment.SetDefine(TEXT("VF_SUPPORTS_PRIMITIVE_SCENE_DATA"), 1);
-		OutEnvironment.SetDefine(TEXT("USE_GLOBAL_GPU_SCENE_DATA"), 1);
-		OutEnvironment.SetDefine(TEXT("NUM_THREADS_PER_GROUP"), NumThreadsPerGroup);
-		OutEnvironment.SetDefine(TEXT("NANITE_MULTI_VIEW"), 1);
-		OutEnvironment.SetDefine(TEXT("ENABLE_DETERMINISTIC_INSTANCE_CULLING"), 1);
-		OutEnvironment.SetDefine(TEXT("PRIM_ID_DYNAMIC_FLAG"), GPrimIDDynamicFlag);
-	}
-
-	BEGIN_SHADER_PARAMETER_STRUCT(FParameters, )
-		SHADER_PARAMETER_SRV(StructuredBuffer<float4>, GPUSceneInstanceSceneData)
-		SHADER_PARAMETER_SRV(StructuredBuffer<float4>, GPUScenePrimitiveSceneData)
-		SHADER_PARAMETER(uint32, InstanceDataSOAStride)
-
-		SHADER_PARAMETER_RDG_BUFFER_SRV(StructuredBuffer< FInstanceCullingContext::FPrimCullingCommand >, PrimitiveCullingCommands)
-		SHADER_PARAMETER_RDG_BUFFER_SRV(StructuredBuffer< int32 >, PrimitiveIds)
-		SHADER_PARAMETER_RDG_BUFFER_SRV(StructuredBuffer< FInstanceCullingContext::FInstanceRun >, InstanceRuns)
-		SHADER_PARAMETER_RDG_BUFFER_SRV(StructuredBuffer< uint32 >, VisibleInstanceFlags)
-		SHADER_PARAMETER_RDG_BUFFER_SRV(StructuredBuffer< uint32 >, ViewIds)
-
-		SHADER_PARAMETER_RDG_BUFFER_UAV(RWStructuredBuffer<uint>, OutputOffsetBufferOut)
-		SHADER_PARAMETER_RDG_BUFFER_UAV(RWStructuredBuffer<uint2>, InstanceCountsOut)
-
-		SHADER_PARAMETER(int32, NumPrimitiveIds)
-		SHADER_PARAMETER(int32, NumInstanceRuns)
-		SHADER_PARAMETER(int32, NumCommands)
-		SHADER_PARAMETER(uint32, NumInstanceFlagWords)
-		SHADER_PARAMETER(uint32, NumCulledInstances)
-		SHADER_PARAMETER(uint32, NumCulledViews)
-		SHADER_PARAMETER(int32, NumViewIds)
-
-		SHADER_PARAMETER(int32, DynamicPrimitiveIdOffset)
-		SHADER_PARAMETER(int32, DynamicPrimitiveIdMax)
-	END_SHADER_PARAMETER_STRUCT()
-};
-IMPLEMENT_GLOBAL_SHADER(FComputeInstanceIdOutputSizeCs, "/Engine/Private/InstanceCulling/BuildInstanceDrawCommands.usf", "ComputeInstanceIdOutputSize", SF_Compute);
-
-class FCalcOutputOffsetsCs : public FGlobalShader
-{
-	DECLARE_GLOBAL_SHADER(FCalcOutputOffsetsCs);
-	SHADER_USE_PARAMETER_STRUCT(FCalcOutputOffsetsCs, FGlobalShader)
-public:
-	static constexpr int32 NumThreadsPerGroup = 64;
-
-	static bool ShouldCompilePermutation(const FGlobalShaderPermutationParameters& Parameters)
-	{
-		return UseGPUScene(Parameters.Platform);
-	}
-
-	static void ModifyCompilationEnvironment(const FGlobalShaderPermutationParameters& Parameters, FShaderCompilerEnvironment& OutEnvironment)
-	{
-		FGlobalShader::ModifyCompilationEnvironment(Parameters, OutEnvironment);
-		OutEnvironment.SetDefine(TEXT("INDIRECT_ARGS_NUM_WORDS"), FInstanceCullingContext::IndirectArgsNumWords);
-		OutEnvironment.SetDefine(TEXT("VF_SUPPORTS_PRIMITIVE_SCENE_DATA"), 1);
-		OutEnvironment.SetDefine(TEXT("USE_GLOBAL_GPU_SCENE_DATA"), 1);
-		OutEnvironment.SetDefine(TEXT("NUM_THREADS_PER_GROUP"), NumThreadsPerGroup);
-		OutEnvironment.SetDefine(TEXT("NANITE_MULTI_VIEW"), 1);
-		OutEnvironment.SetDefine(TEXT("ENABLE_DETERMINISTIC_INSTANCE_CULLING"), 1);
-		OutEnvironment.SetDefine(TEXT("PRIM_ID_DYNAMIC_FLAG"), GPrimIDDynamicFlag);
-	}
-
-	BEGIN_SHADER_PARAMETER_STRUCT(FParameters, )
-		SHADER_PARAMETER_RDG_BUFFER_UAV(RWBuffer<uint>, InstanceIdOffsetBufferOut)
-		SHADER_PARAMETER_RDG_BUFFER_UAV(RWStructuredBuffer<uint>, OutputOffsetBufferOut)
-		SHADER_PARAMETER_RDG_BUFFER_SRV(StructuredBuffer<uint2>, InstanceCounts)
-		SHADER_PARAMETER(int32, NumCommands)
-		SHADER_PARAMETER(int32, NumViewIds)
-	END_SHADER_PARAMETER_STRUCT()
-};
-IMPLEMENT_GLOBAL_SHADER(FCalcOutputOffsetsCs, "/Engine/Private/InstanceCulling/BuildInstanceDrawCommands.usf", "CalcOutputOffsets", SF_Compute);
-
-class FOutputInstanceIdsAtOffsetCs : public FGlobalShader
-{
-	DECLARE_GLOBAL_SHADER(FOutputInstanceIdsAtOffsetCs);
-	SHADER_USE_PARAMETER_STRUCT(FOutputInstanceIdsAtOffsetCs, FGlobalShader)
-public:
-	static constexpr int32 NumThreadsPerGroup = 64;
-
-	// GPUCULL_TODO: remove once buffer is somehow unified
-	class FOutputCommandIdDim : SHADER_PERMUTATION_BOOL("OUTPUT_COMMAND_IDS");
-	class FCullInstancesDim : SHADER_PERMUTATION_BOOL("CULL_INSTANCES");
-	using FPermutationDomain = TShaderPermutationDomain<FOutputCommandIdDim, FCullInstancesDim>;
-
-	static bool ShouldCompilePermutation(const FGlobalShaderPermutationParameters& Parameters)
-	{
-		return UseGPUScene(Parameters.Platform);
-	}
-
-	static void ModifyCompilationEnvironment(const FGlobalShaderPermutationParameters& Parameters, FShaderCompilerEnvironment& OutEnvironment)
-	{
-		FGlobalShader::ModifyCompilationEnvironment(Parameters, OutEnvironment);
-		OutEnvironment.SetDefine(TEXT("INDIRECT_ARGS_NUM_WORDS"), FInstanceCullingContext::IndirectArgsNumWords);
-		OutEnvironment.SetDefine(TEXT("VF_SUPPORTS_PRIMITIVE_SCENE_DATA"), 1);
-		OutEnvironment.SetDefine(TEXT("USE_GLOBAL_GPU_SCENE_DATA"), 1);
-		OutEnvironment.SetDefine(TEXT("NUM_THREADS_PER_GROUP"), NumThreadsPerGroup);
-		OutEnvironment.SetDefine(TEXT("NANITE_MULTI_VIEW"), 1);
-		OutEnvironment.SetDefine(TEXT("ENABLE_DETERMINISTIC_INSTANCE_CULLING"), 1);
-		OutEnvironment.SetDefine(TEXT("PRIM_ID_DYNAMIC_FLAG"), GPrimIDDynamicFlag);
-	}
-
-	BEGIN_SHADER_PARAMETER_STRUCT(FParameters, )
-		SHADER_PARAMETER_RDG_BUFFER_SRV(Buffer<uint>, InstanceIdOffsetBuffer)
-		SHADER_PARAMETER_RDG_BUFFER_SRV(StructuredBuffer<uint2>, InstanceCounts)
-
-		SHADER_PARAMETER_SRV(StructuredBuffer<float4>, GPUSceneInstanceSceneData)
-		SHADER_PARAMETER_SRV(StructuredBuffer<float4>, GPUScenePrimitiveSceneData)
-		SHADER_PARAMETER(uint32, InstanceDataSOAStride)
-
-		SHADER_PARAMETER_RDG_BUFFER_SRV(StructuredBuffer< FInstanceCullingContext::FPrimCullingCommand >, PrimitiveCullingCommands)
-		SHADER_PARAMETER_RDG_BUFFER_SRV(StructuredBuffer< int32 >, PrimitiveIds)
-		SHADER_PARAMETER_RDG_BUFFER_SRV(StructuredBuffer< FInstanceCullingContext::FInstanceRun >, InstanceRuns)
-		SHADER_PARAMETER_RDG_BUFFER_SRV(StructuredBuffer< uint32 >, VisibleInstanceFlags)
-		SHADER_PARAMETER_RDG_BUFFER_SRV(StructuredBuffer< uint32 >, ViewIds)
-
-		SHADER_PARAMETER_RDG_BUFFER_UAV(RWBuffer<uint>, InstanceIdsBufferOut)
-		SHADER_PARAMETER_RDG_BUFFER_UAV(RWBuffer<uint>, DrawCommandIdsBufferOut)
-		// Using the wrong kind of buffer for RDG...
-		SHADER_PARAMETER_UAV(RWStructuredBuffer<uint>, InstanceIdsBufferOut)
-
-		SHADER_PARAMETER_RDG_BUFFER_UAV(RWBuffer<uint>, DrawIndirectArgsBufferOut)
-		SHADER_PARAMETER(int32, NumPrimitiveIds)
-		SHADER_PARAMETER(int32, NumInstanceRuns)
-		SHADER_PARAMETER(int32, NumCommands)
-		SHADER_PARAMETER(uint32, NumInstanceFlagWords)
-		SHADER_PARAMETER(uint32, NumCulledInstances)
-		SHADER_PARAMETER(uint32, NumCulledViews)
-		SHADER_PARAMETER(int32, NumViewIds)
-
-		SHADER_PARAMETER(int32, DynamicPrimitiveIdOffset)
-		SHADER_PARAMETER(int32, DynamicPrimitiveIdMax)
-
-	END_SHADER_PARAMETER_STRUCT()
-};
-IMPLEMENT_GLOBAL_SHADER(FOutputInstanceIdsAtOffsetCs, "/Engine/Private/InstanceCulling/BuildInstanceDrawCommands.usf", "OutputInstanceIdsAtOffset", SF_Compute);
-
-#endif // ENABLE_DETERMINISTIC_INSTANCE_CULLING
-
 class FBuildInstanceIdBufferAndCommandsFromPrimitiveIdsCs : public FGlobalShader
 {
 	DECLARE_GLOBAL_SHADER(FBuildInstanceIdBufferAndCommandsFromPrimitiveIdsCs);
@@ -389,136 +231,6 @@ void FInstanceCullingContext::BuildRenderingCommands(
 	const uint32 InstanceIdBufferSize = TotalInstances * ViewIds.Num();
 	FRDGBufferRef InstanceIdsBuffer = GraphBuilder.CreateBuffer(FRDGBufferDesc::CreateStructuredDesc(sizeof(uint32), InstanceIdBufferSize), TEXT("InstanceCulling.InstanceIdsBuffer"));
 
-#if ENABLE_DETERMINISTIC_INSTANCE_CULLING
-	// 1. Compute output sizes for all commands
-	FRDGBufferRef InstanceCountsRDG = GraphBuilder.CreateBuffer(FRDGBufferDesc::CreateStructuredDesc(sizeof(FIntVector2), CullingCommands.Num()), TEXT("InstanceCulling.InstanceCounts"));
-	FRDGBufferRef InstanceIdOffsetBufferRDG = GraphBuilder.CreateBuffer(FRDGBufferDesc::CreateBufferDesc(sizeof(uint32), CullingCommands.Num()), TEXT("InstanceCulling.InstanceIdOffsetBuffer"));
-	FRDGBufferRef DrawIndirectArgsRDG = GraphBuilder.CreateBuffer(FRDGBufferDesc::CreateIndirectDesc(IndirectArgsNumWords * CullingCommands.Num()), TEXT("InstanceCulling.DrawIndirectArgsBuffer"));
-
-
-	{
-		FComputeInstanceIdOutputSizeCs::FParameters* PassParameters = GraphBuilder.AllocParameters<FComputeInstanceIdOutputSizeCs::FParameters>();
-
-		PassParameters->InstanceCountsOut = GraphBuilder.CreateUAV(InstanceCountsRDG);
-
-		// Because the view uniforms are not set up by the time this runs
-		// PassParameters->View = View.ViewUniformBuffer;
-		// Set up global GPU-scene data instead...
-		PassParameters->GPUSceneInstanceSceneData = GPUScene.InstanceDataBuffer.SRV;
-		PassParameters->GPUScenePrimitiveSceneData = GPUScene.PrimitiveBuffer.SRV;
-		PassParameters->InstanceDataSOAStride = GPUScene.InstanceDataSOAStride;
-		// Upload data etc
-		PassParameters->PrimitiveCullingCommands = GraphBuilder.CreateSRV(PrimitiveCullingCommandsBuffer, CullingCommands));
-
-		PassParameters->PrimitiveIds = GraphBuilder.CreateSRV(PrimitiveIdsBuffer, PrimitiveIds));
-		PassParameters->DynamicPrimitiveIdOffset = DynamicPrimitiveIdRange.GetLowerBoundValue();
-		PassParameters->DynamicPrimitiveIdMax = DynamicPrimitiveIdRange.GetUpperBoundValue();
-
-		PassParameters->InstanceRuns = GraphBuilder.CreateSRV(PrimitiveInstanceRunsBuffer);
-
-		PassParameters->OutputOffsetBufferOut = GraphBuilder.CreateUAV(InstanceIdOutOffsetBufferRDG);
-
-		PassParameters->ViewIds = GraphBuilder.CreateSRV(ViewIdsBuffer);
-		PassParameters->NumViewIds = ViewIds.Num();
-		PassParameters->NumPrimitiveIds = PrimitiveIds.Num();
-		PassParameters->NumInstanceRuns = InstanceRuns.Num();
-		PassParameters->NumCommands = CullingCommands.Num();
-		PassParameters->VisibleInstanceFlags = VisibleInstanceFlagsRDG ? GraphBuilder.CreateSRV(VisibleInstanceFlagsRDG) : nullptr;
-
-		PassParameters->NumInstanceFlagWords = NumInstanceFlagWords;
-		PassParameters->NumCulledInstances = NumCulledInstances;
-		PassParameters->NumCulledViews = NumCulledViews;
-		PassParameters->NumCulledInstances = uint32(NumCulledInstances);
-
-		FComputeInstanceIdOutputSizeCs::FPermutationDomain PermutationVector;
-		PermutationVector.Set<FComputeInstanceIdOutputSizeCs::FCullInstancesDim>(bCullInstances);
-		PermutationVector.Set<FComputeInstanceIdOutputSizeCs::FStereoModeDim>(InstanceCullingMode == EInstanceCullingMode::Stereo);
-		auto ComputeShader = ShaderMap->GetShader<FComputeInstanceIdOutputSizeCs>(PermutationVector);
-
-		FComputeShaderUtils::AddPass(
-			GraphBuilder,
-			RDG_EVENT_NAME("ComputeInstanceIdOutputSize"),
-			ComputeShader,
-			PassParameters,
-			FIntVector(CullingCommands.Num(), 1, 1)
-		);
-	}
-	// 2. Allocate output slots for each command
-	{
-		FCalcOutputOffsetsCs::FParameters* PassParameters = GraphBuilder.AllocParameters<FCalcOutputOffsetsCs::FParameters>();
-
-		PassParameters->InstanceCounts = GraphBuilder.CreateSRV(InstanceCountsRDG);
-		PassParameters->OutputOffsetBufferOut = GraphBuilder.CreateUAV(InstanceIdOutOffsetBufferRDG);
-		PassParameters->InstanceIdOffsetBufferOut = GraphBuilder.CreateUAV(InstanceIdOffsetBufferRDG, PF_R32_UINT);
-		PassParameters->NumViewIds = ViewIds.Num();
-		PassParameters->NumCommands = CullingCommands.Num();
-
-		auto ComputeShader = ShaderMap->GetShader<FCalcOutputOffsetsCs>();
-
-		FComputeShaderUtils::AddPass(
-			GraphBuilder,
-			RDG_EVENT_NAME("CalcOutputOffsets"),
-			ComputeShader,
-			PassParameters,
-			FIntVector(1, 1, 1)
-		);
-	}
-	// 3. Populate the output buffers
-	{
-		FOutputInstanceIdsAtOffsetCs::FParameters* PassParameters = GraphBuilder.AllocParameters<FOutputInstanceIdsAtOffsetCs::FParameters>();
-
-		PassParameters->InstanceCounts = GraphBuilder.CreateSRV(InstanceCountsRDG);
-		PassParameters->InstanceIdOffsetBuffer = GraphBuilder.CreateSRV(InstanceIdOffsetBufferRDG, PF_R32_UINT);
-
-
-		// Because the view uniforms are not set up by the time this runs
-		// PassParameters->View = View.ViewUniformBuffer;
-		// Set up global GPU-scene data instead...
-		PassParameters->GPUSceneInstanceSceneData = GPUScene.InstanceDataBuffer.SRV;
-		PassParameters->GPUScenePrimitiveSceneData = GPUScene.PrimitiveBuffer.SRV;
-		PassParameters->InstanceDataSOAStride = GPUScene.InstanceDataSOAStride;
-		// Upload data etc
-		PassParameters->PrimitiveCullingCommands = GraphBuilder.CreateSRV(PrimitiveCullingCommandsBuffer);
-
-		PassParameters->PrimitiveIds = GraphBuilder.CreateSRV(PrimitiveIdsBuffer);
-		PassParameters->DynamicPrimitiveIdOffset = DynamicPrimitiveIdRange.GetLowerBoundValue();
-		PassParameters->DynamicPrimitiveIdMax = DynamicPrimitiveIdRange.GetUpperBoundValue();
-
-		PassParameters->InstanceRuns = GraphBuilder.CreateSRV(PrimitiveInstanceRunsBuffer);
-
-		PassParameters->ViewIds = GraphBuilder.CreateSRV(ViewIdsBuffer);
-		PassParameters->NumViewIds = ViewIds.Num();
-
-		//PassParameters->InstanceIdsBufferOut = GraphBuilder.CreateUAV(InstanceIdsBufferRDG, PF_R32_UINT);
-		// TODO: Access resources through manager rather than global
-		PassParameters->InstanceIdsBufferOut = GInstanceCullingManagerResources.GetInstancesIdBufferUav();
-		PassParameters->DrawIndirectArgsBufferOut = GraphBuilder.CreateUAV(DrawIndirectArgsRDG, PF_R32_UINT);
-		PassParameters->NumPrimitiveIds = PrimitiveIds.Num();
-		PassParameters->NumInstanceRuns = InstanceRuns.Num();
-		PassParameters->NumCommands = CullingCommands.Num();
-		PassParameters->VisibleInstanceFlags = VisibleInstanceFlagsRDG ? GraphBuilder.CreateSRV(VisibleInstanceFlagsRDG) : nullptr;
-		PassParameters->NumInstanceFlagWords = NumInstanceFlagWords;
-		PassParameters->NumCulledInstances = NumCulledInstances;
-		PassParameters->NumCulledViews = NumCulledViews;
-		PassParameters->NumCulledInstances = uint32(NumCulledInstances);
-
-		FOutputInstanceIdsAtOffsetCs::FPermutationDomain PermutationVector;
-		PermutationVector.Set<FOutputInstanceIdsAtOffsetCs::FOutputCommandIdDim>(0);
-		PermutationVector.Set<FOutputInstanceIdsAtOffsetCs::FCullInstancesDim>(bCullInstances);
-		PermutationVector.Set<FOutputInstanceIdsAtOffsetCs::FStereoModeDim>(InstanceCullingMode == EInstanceCullingMode::Stereo);
-		auto ComputeShader = ShaderMap->GetShader<FOutputInstanceIdsAtOffsetCs>(PermutationVector);
-
-		FComputeShaderUtils::AddPass(
-			GraphBuilder,
-			RDG_EVENT_NAME("OutputInstanceIdsAtOffset"),
-			ComputeShader,
-			PassParameters,
-			FIntVector(CullingCommands.Num(), 1, 1)
-		);
-	}
-
-#else // !ENABLE_DETERMINISTIC_INSTANCE_CULLING
-
 	FBuildInstanceIdBufferAndCommandsFromPrimitiveIdsCs::FParameters* PassParameters = GraphBuilder.AllocParameters<FBuildInstanceIdBufferAndCommandsFromPrimitiveIdsCs::FParameters>();
 
 	// Because the view uniforms are not set up by the time this runs
@@ -577,8 +289,6 @@ void FInstanceCullingContext::BuildRenderingCommands(
 		FIntVector(CullingCommands.Num(), 1, 1)
 	);
 
-#endif // ENABLE_DETERMINISTIC_INSTANCE_CULLING
-
 	Results.DrawIndirectArgsBuffer = DrawIndirectArgsRDG;
 	Results.InstanceIdOffsetBuffer = InstanceIdOffsetBufferRDG;
 
@@ -784,6 +494,12 @@ void FInstanceCullingContext::BuildRenderingCommandsDeferred(
 		return;
 	}
 
+	// If there are no instances, there can be no work to perform later.
+	if (InstanceCullingManager.CullingIntermediate.NumInstances == 0 || InstanceCullingManager.CullingIntermediate.NumViews == 0)
+	{
+		return;
+	}
+
 	FBatchedInstanceCullingScratchSpace& BatchedCullingScratch = InstanceCullingManager.BatchedCullingScratch;
 
 	checkf(!BatchedCullingScratch.bBatchingActive, TEXT("Batched GPU instance culling has already been enabled."));
