Implements r.TemporalAA.EnableResponsiveAA

diff --git a/Engine/Shaders/Private/TemporalAA/TAAUpdateHistory.usf b/Engine/Shaders/Private/TemporalAA/TAAUpdateHistory.usf
--- a/Engine/Shaders/Private/TemporalAA/TAAUpdateHistory.usf
+++ b/Engine/Shaders/Private/TemporalAA/TAAUpdateHistory.usf
@@ -45,6 +45,7 @@ float4 HistoryPixelPosToScreenPos;
 float4 HistoryPixelPosToPPCo;
 float3 HistoryQuantizationError;
 float MinTranslucencyRejection;
+uint ResponsiveStencilMask;
 
 Texture2D<taa_half3> PrevHistory_Textures_0;
 Texture2D<taa_half3> PrevHistory_Textures_1;
@@ -253,9 +254,7 @@ void MainCS(
 		// Fetch wehther the pixel is responsive AA or not.
 		#if CONFIG_RESPONSIVE_STENCIL
 		{
-			const uint kResponsiveStencilMask = 1 << 3;
-			
-			bIsResponsiveAAPixel = (SceneStencilRef & kResponsiveStencilMask) != 0;
+			bIsResponsiveAAPixel = (SceneStencilRef & ResponsiveStencilMask) != 0;
 		}
 		#endif
 
diff --git a/Engine/Source/Runtime/Renderer/Private/PostProcess/TemporalAA.cpp b/Engine/Source/Runtime/Renderer/Private/PostProcess/TemporalAA.cpp
--- a/Engine/Source/Runtime/Renderer/Private/PostProcess/TemporalAA.cpp
+++ b/Engine/Source/Runtime/Renderer/Private/PostProcess/TemporalAA.cpp
@@ -104,6 +104,11 @@ TAutoConsoleVariable<int32> CVarTAARejectTranslucency(
 	TEXT("Enable heuristic to reject based on the Separate Translucency."),
 	ECVF_RenderThreadSafe);
 
+TAutoConsoleVariable<int32> CVarTAAEnableResponiveAA(
+	TEXT("r.TemporalAA.EnableResponiveAA"), 1,
+	TEXT("Whether the responsive AA should be enabled."),
+	ECVF_RenderThreadSafe);
+
 #if COMPILE_TAA_DEBUG_PASSES
 
 TAutoConsoleVariable<int32> CVarTAASetupDebugPasses(
@@ -596,6 +601,7 @@ class FTSRUpdateHistoryCS : public FTemporalSuperResolutionShader
 		SHADER_PARAMETER(FScreenTransform, HistoryPixelPosToPPCo)
 		SHADER_PARAMETER(FVector3f, HistoryQuantizationError)
 		SHADER_PARAMETER(float, MinTranslucencyRejection)
+		SHADER_PARAMETER(int32, ResponsiveStencilMask)
 
 		SHADER_PARAMETER_STRUCT_INCLUDE(FTSRPrevHistoryParameters, PrevHistoryParameters)
 		SHADER_PARAMETER_STRUCT(FTSRHistoryTextures, PrevHistory)
@@ -1931,6 +1937,7 @@ static void AddTemporalSuperResolutionPasses(
 		PassParameters->HistoryPixelPosToPPCo = HistoryPixelPosToViewportUV * CommonParameters.InputInfo.ViewportSize + CommonParameters.InputJitter + CommonParameters.InputPixelPosMin;
 		PassParameters->HistoryQuantizationError = ComputePixelFormatQuantizationError(History.Textures[0]->Desc.Format);
 		PassParameters->MinTranslucencyRejection = TranslucencyRejectionTexture == nullptr ? 1.0 : 0.0;
+		PassParameters->ResponsiveStencilMask = CVarTAAEnableResponiveAA.GetValueOnRenderThread() ? (STENCIL_TEMPORAL_RESPONSIVE_AA_MASK) : 0;
 
 		PassParameters->PrevHistoryParameters = PrevHistoryParameters;
 		PassParameters->PrevHistory = PrevHistory;
