Separates TSR source code from TAA

diff --git a/Engine/Source/Runtime/Renderer/Private/PostProcess/TemporalAA.cpp b/Engine/Source/Runtime/Renderer/Private/PostProcess/TemporalAA.cpp
--- a/Engine/Source/Runtime/Renderer/Private/PostProcess/TemporalAA.cpp
+++ b/Engine/Source/Runtime/Renderer/Private/PostProcess/TemporalAA.cpp
@@ -11,8 +11,6 @@
 #include "ScenePrivate.h"
 #include "RendererModule.h"
 
-#define COMPILE_TAA_DEBUG_PASSES (!UE_BUILD_SHIPPING)
-
 namespace
 {
 
@@ -81,42 +79,6 @@ TAutoConsoleVariable<int32> CVarTAAR11G11B10History(
 	TEXT("Select the bitdepth of the history."),
 	ECVF_RenderThreadSafe);
 
-TAutoConsoleVariable<int32> CVarTAAHalfResShadingRejection(
-	TEXT("r.TemporalAA.ShadingRejection.HalfRes"), 0,
-	TEXT("Whether the shading rejection should be done at half res. Saves performance but may introduce back some flickering (default = 0)."),
-	ECVF_RenderThreadSafe);
-
-TAutoConsoleVariable<int32> CVarTAAFilterShadingRejection(
-	TEXT("r.TemporalAA.ShadingRejection.SpatialFilter"), 1,
-	TEXT("Whether the shading rejection should have spatial statistical filtering pass to reduce flickering (default = 1).\n")
-	TEXT(" 0: Disabled;\n")
-	TEXT(" 1: Spatial filter pass is run at lower resolution than CompareHistory pass (default);\n")
-	TEXT(" 2: Spatial filter pass is run CompareHistory pass resolution to improve stability."),
-	ECVF_RenderThreadSafe);
-
-TAutoConsoleVariable<int32> CVarTAAEnableAntiInterference(
-	TEXT("r.TemporalAA.AntiInterference"), 0,
-	TEXT("Enable heuristic to detect geometric interference between input pixel grid alignement and structured geometry."),
-	ECVF_RenderThreadSafe);
-
-TAutoConsoleVariable<int32> CVarTAARejectTranslucency(
-	TEXT("r.TemporalAA.RejectSeparateTranslucency"), 0,
-	TEXT("Enable heuristic to reject based on the Separate Translucency."),
-	ECVF_RenderThreadSafe);
-
-TAutoConsoleVariable<int32> CVarTAAEnableResponiveAA(
-	TEXT("r.TemporalAA.EnableResponiveAA"), 1,
-	TEXT("Whether the responsive AA should be enabled."),
-	ECVF_RenderThreadSafe);
-
-#if COMPILE_TAA_DEBUG_PASSES
-
-TAutoConsoleVariable<int32> CVarTAASetupDebugPasses(
-	TEXT("r.TemporalAA.Debug.SetupExtraPasses"), 0,
-	TEXT("Whether to enable the debug passes"),
-	ECVF_RenderThreadSafe);
-
-#endif
 
 inline bool DoesPlatformSupportTemporalHistoryUpscale(EShaderPlatform Platform)
 {
@@ -124,93 +86,10 @@ inline bool DoesPlatformSupportTemporalHistoryUpscale(EShaderPlatform Platform)
 		&& IsFeatureLevelSupported(Platform, ERHIFeatureLevel::SM5);
 }
 
-inline bool DoesPlatformSupportGen5TAA(EShaderPlatform Platform)
-{
-	// TODO(Gen5TAA): alpha channel is not supported yet
-	if (IsPostProcessingWithAlphaChannelSupported())
-	{
-		return false;
-	}
-	return FDataDrivenShaderPlatformInfo::GetSupportsGen5TemporalAA(Platform);
-}
-
-BEGIN_SHADER_PARAMETER_STRUCT(FTSRCommonParameters, )
-	SHADER_PARAMETER_STRUCT(FScreenPassTextureViewportParameters, InputInfo)
-	SHADER_PARAMETER_STRUCT(FScreenPassTextureViewportParameters, LowFrequencyInfo)
-	SHADER_PARAMETER_STRUCT(FScreenPassTextureViewportParameters, RejectionInfo)
-	SHADER_PARAMETER_STRUCT(FScreenPassTextureViewportParameters, HistoryInfo)
-
-	SHADER_PARAMETER(FIntPoint, InputPixelPosMin)
-	SHADER_PARAMETER(FIntPoint, InputPixelPosMax)
-
-	SHADER_PARAMETER(FVector2D, InputJitter)
-	SHADER_PARAMETER(int32, bCameraCut)
-	SHADER_PARAMETER(int32, bEnableInterferenceHeuristic)
-
-	SHADER_PARAMETER_STRUCT_REF(FViewUniformShaderParameters, ViewUniformBuffer)
-END_SHADER_PARAMETER_STRUCT()
-
-BEGIN_SHADER_PARAMETER_STRUCT(FTSRHistoryTextures, )
-	SHADER_PARAMETER_RDG_TEXTURE_ARRAY(Texture2D, LowResTextures, [FTemporalAAHistory::kLowResRenderTargetCount])
-	SHADER_PARAMETER_RDG_TEXTURE_ARRAY(Texture2D, Textures, [FTemporalAAHistory::kRenderTargetCount])
-	SHADER_PARAMETER_RDG_TEXTURE_ARRAY(Texture2D, SuperResTextures, [FTemporalAAHistory::kSuperResRenderTargetCount])
-END_SHADER_PARAMETER_STRUCT()
-
-BEGIN_SHADER_PARAMETER_STRUCT(FTSRHistoryUAVs, )
-	SHADER_PARAMETER_RDG_TEXTURE_UAV_ARRAY(RWTexture2D, LowResTextures, [FTemporalAAHistory::kLowResRenderTargetCount])
-	SHADER_PARAMETER_RDG_TEXTURE_UAV_ARRAY(RWTexture2D, Textures, [FTemporalAAHistory::kRenderTargetCount])
-	SHADER_PARAMETER_RDG_TEXTURE_UAV_ARRAY(RWTexture2D, SuperResTextures, [FTemporalAAHistory::kSuperResRenderTargetCount])
-END_SHADER_PARAMETER_STRUCT()
-
-BEGIN_SHADER_PARAMETER_STRUCT(FTSRPrevHistoryParameters, )
-	SHADER_PARAMETER_STRUCT(FScreenPassTextureViewportParameters, PrevHistoryInfo)
-	SHADER_PARAMETER(FScreenTransform, ScreenPosToPrevHistoryBufferUV)
-	SHADER_PARAMETER(float, HistoryPreExposureCorrection)
-END_SHADER_PARAMETER_STRUCT()
-
-FTSRHistoryUAVs CreateUAVs(FRDGBuilder& GraphBuilder, const FTSRHistoryTextures& Textures)
-{
-	FTSRHistoryUAVs UAVs;
-	for (int32 i = 0; i < Textures.LowResTextures.Num(); i++)
-	{
-		UAVs.LowResTextures[i] = GraphBuilder.CreateUAV(Textures.LowResTextures[i]);
-	}
-	for (int32 i = 0; i < Textures.Textures.Num(); i++)
-	{
-		UAVs.Textures[i] = GraphBuilder.CreateUAV(Textures.Textures[i]);
-	}
-	for (int32 i = 0; i < Textures.SuperResTextures.Num(); i++)
-	{
-		UAVs.SuperResTextures[i] = GraphBuilder.CreateUAV(Textures.SuperResTextures[i]);
-	}
-	return UAVs;
-}
-
-class FTemporalSuperResolutionShader : public FGlobalShader
-{
-public:
-	FTemporalSuperResolutionShader(const ShaderMetaType::CompiledShaderInitializerType& Initializer)
-		: FGlobalShader(Initializer)
-	{ }
-
-	FTemporalSuperResolutionShader()
-	{ }
-
-	static bool ShouldCompilePermutation(const FGlobalShaderPermutationParameters& Parameters)
-	{
-		return DoesPlatformSupportGen5TAA(Parameters.Platform);
-	}
-
-	static void ModifyCompilationEnvironment(const FGlobalShaderPermutationParameters& Parameters, FShaderCompilerEnvironment& OutEnvironment)
-	{
-		OutEnvironment.CompilerFlags.Add(CFLAG_AllowRealTypes);
-	}
-}; // class FTemporalSuperResolutionShader
-
-class FTAAStandaloneCS : public FGlobalShader
+class FTemporalAACS : public FGlobalShader
 {
-	DECLARE_GLOBAL_SHADER(FTAAStandaloneCS);
-	SHADER_USE_PARAMETER_STRUCT(FTAAStandaloneCS, FGlobalShader);
+	DECLARE_GLOBAL_SHADER(FTemporalAACS);
+	SHADER_USE_PARAMETER_STRUCT(FTemporalAACS, FGlobalShader);
 
 	class FTAAPassConfigDim : SHADER_PERMUTATION_ENUM_CLASS("TAA_PASS_CONFIG", ETAAPassConfig);
 	class FTAAFastDim : SHADER_PERMUTATION_BOOL("TAA_FAST");
@@ -373,282 +252,9 @@ class FTAAStandaloneCS : public FGlobalShader
 		// There are some mobile specific shader optimizations need to be set in the shader, such as disable shared memory usage, disable stencil texture sampling.
 		OutEnvironment.SetDefine(TEXT("AA_MOBILE_CONFIG"), bIsMobileTiledGPU ? 1 : 0);
 	}
-}; // class FTAAStandaloneCS
-
-class FTSRClearPrevTexturesCS : public FTemporalSuperResolutionShader
-{
-	DECLARE_GLOBAL_SHADER(FTSRClearPrevTexturesCS);
-	SHADER_USE_PARAMETER_STRUCT(FTSRClearPrevTexturesCS, FTemporalSuperResolutionShader);
-
-	BEGIN_SHADER_PARAMETER_STRUCT(FParameters, )
-		SHADER_PARAMETER_STRUCT_INCLUDE(FTSRCommonParameters, CommonParameters)
-
-		SHADER_PARAMETER_RDG_TEXTURE_UAV(RWTexture2D, PrevUseCountOutput)
-		SHADER_PARAMETER_RDG_TEXTURE_UAV(RWTexture2D, PrevClosestDepthOutput)
-	END_SHADER_PARAMETER_STRUCT()
-}; // class FTSRClearPrevTexturesCS
-
-class FTSRDilateVelocityCS : public FTemporalSuperResolutionShader
-{
-	DECLARE_GLOBAL_SHADER(FTSRDilateVelocityCS);
-	SHADER_USE_PARAMETER_STRUCT(FTSRDilateVelocityCS, FTemporalSuperResolutionShader);
-
-	BEGIN_SHADER_PARAMETER_STRUCT(FParameters, )
-		SHADER_PARAMETER_STRUCT_INCLUDE(FTSRCommonParameters, CommonParameters)
-
-		SHADER_PARAMETER(FVector2D, PrevOutputBufferUVMin)
-		SHADER_PARAMETER(FVector2D, PrevOutputBufferUVMax)
-
-		SHADER_PARAMETER_RDG_TEXTURE(Texture2D, SceneDepthTexture)
-		SHADER_PARAMETER_RDG_TEXTURE(Texture2D, SceneVelocityTexture)
-
-		SHADER_PARAMETER_RDG_TEXTURE_UAV(RWTexture2D, DilatedVelocityOutput)
-		SHADER_PARAMETER_RDG_TEXTURE_UAV(RWTexture2D, ClosestDepthOutput)
-		SHADER_PARAMETER_RDG_TEXTURE_UAV(RWTexture2D, PrevUseCountOutput)
-		SHADER_PARAMETER_RDG_TEXTURE_UAV(RWTexture2D, PrevClosestDepthOutput)
-		SHADER_PARAMETER_RDG_TEXTURE_UAV(RWTexture2D, ParallaxFactorOutput)
-		SHADER_PARAMETER_RDG_TEXTURE_UAV(RWTexture2D, DebugOutput)
-	END_SHADER_PARAMETER_STRUCT()
-}; // class FTSRDilateVelocityCS
-
-class FTSRDecimateHistoryCS : public FTemporalSuperResolutionShader
-{
-	DECLARE_GLOBAL_SHADER(FTSRDecimateHistoryCS);
-	SHADER_USE_PARAMETER_STRUCT(FTSRDecimateHistoryCS, FTemporalSuperResolutionShader);
-
-	class FOutputHalfRes : SHADER_PERMUTATION_BOOL("DIM_OUTPUT_HALF_RES");
-
-	using FPermutationDomain = TShaderPermutationDomain<FOutputHalfRes>;
-
-	BEGIN_SHADER_PARAMETER_STRUCT(FParameters, )
-		SHADER_PARAMETER_STRUCT_INCLUDE(FTSRCommonParameters, CommonParameters)
-		SHADER_PARAMETER(FMatrix44f, RotationalClipToPrevClip)
-		SHADER_PARAMETER(FVector3f, OutputQuantizationError)
-		SHADER_PARAMETER(float, WorldDepthToPixelWorldRadius)
-
-		SHADER_PARAMETER_RDG_TEXTURE(Texture2D, InputSceneColorTexture)
-		SHADER_PARAMETER_RDG_TEXTURE(Texture2D, DilatedVelocityTexture)
-		SHADER_PARAMETER_RDG_TEXTURE(Texture2D, ClosestDepthTexture)
-		SHADER_PARAMETER_RDG_TEXTURE(Texture2D, PrevUseCountTexture)
-		SHADER_PARAMETER_RDG_TEXTURE(Texture2D, PrevClosestDepthTexture)
-		SHADER_PARAMETER_RDG_TEXTURE(Texture2D, ParallaxFactorTexture)
-
-		SHADER_PARAMETER_STRUCT_INCLUDE(FTSRPrevHistoryParameters, PrevHistoryParameters)
-		SHADER_PARAMETER_STRUCT(FTSRHistoryTextures, PrevHistory)
-
-		SHADER_PARAMETER_RDG_TEXTURE_UAV(RWTexture2D, HalfResSceneColorOutput)
-		SHADER_PARAMETER_RDG_TEXTURE_UAV(RWTexture2D, HalfResPredictionSceneColorOutput)
-		SHADER_PARAMETER_RDG_TEXTURE_UAV(RWTexture2D, HalfResParallaxRejectionMaskOutput)
-		SHADER_PARAMETER_RDG_TEXTURE_UAV(RWTexture2D, PredictionSceneColorOutput)
-		SHADER_PARAMETER_RDG_TEXTURE_UAV(RWTexture2D, ParallaxRejectionMaskOutput)
-		SHADER_PARAMETER_RDG_TEXTURE_UAV(RWTexture2D, InterferenceSeedOutput)
-		SHADER_PARAMETER_RDG_TEXTURE_UAV(RWTexture2D, DebugOutput)
-	END_SHADER_PARAMETER_STRUCT()
-}; // class FTSRDecimateHistoryCS
-
-class FTSRCompareTranslucencyCS : public FTemporalSuperResolutionShader
-{
-	DECLARE_GLOBAL_SHADER(FTSRCompareTranslucencyCS);
-	SHADER_USE_PARAMETER_STRUCT(FTSRCompareTranslucencyCS, FTemporalSuperResolutionShader);
-
-	BEGIN_SHADER_PARAMETER_STRUCT(FParameters, )
-		SHADER_PARAMETER_STRUCT_INCLUDE(FTSRCommonParameters, CommonParameters)
-
-		SHADER_PARAMETER_STRUCT(FScreenPassTextureViewportParameters, TranslucencyInfo)
-		SHADER_PARAMETER_STRUCT(FScreenPassTextureViewportParameters, PrevTranslucencyInfo)
-		SHADER_PARAMETER(float, PrevTranslucencyPreExposureCorrection)
-
-		SHADER_PARAMETER(FScreenTransform, InputPixelPosToScreenPos)
-		SHADER_PARAMETER(FScreenTransform, ScreenPosToPrevTranslucencyTextureUV)
-
-		SHADER_PARAMETER_RDG_TEXTURE(Texture2D, DilatedVelocityTexture)
-		SHADER_PARAMETER_RDG_TEXTURE(Texture2D, TranslucencyTexture)
-		SHADER_PARAMETER_RDG_TEXTURE(Texture2D, PrevTranslucencyTexture)
-
-		SHADER_PARAMETER_RDG_TEXTURE_UAV(RWTexture2D, TranslucencyRejectionOutput)
-		SHADER_PARAMETER_RDG_TEXTURE_UAV(RWTexture2D, DebugOutput)
-	END_SHADER_PARAMETER_STRUCT()
-}; // class FTSRCompareTranslucencyCS
-
-class FTSRDetectInterferenceCS : public FTemporalSuperResolutionShader
-{
-	DECLARE_GLOBAL_SHADER(FTSRDetectInterferenceCS);
-	SHADER_USE_PARAMETER_STRUCT(FTSRDetectInterferenceCS, FTemporalSuperResolutionShader);
-	
-	BEGIN_SHADER_PARAMETER_STRUCT(FParameters, )
-		SHADER_PARAMETER_STRUCT_INCLUDE(FTSRCommonParameters, CommonParameters)
-
-		SHADER_PARAMETER_RDG_TEXTURE(Texture2D, InputSceneColorTexture)
-		SHADER_PARAMETER_RDG_TEXTURE(Texture2D, DilatedVelocityTexture)
-		SHADER_PARAMETER_RDG_TEXTURE(Texture2D, PredictionSceneColorTexture)
-		SHADER_PARAMETER_RDG_TEXTURE(Texture2D, ParallaxRejectionMaskTexture)
-		SHADER_PARAMETER_RDG_TEXTURE(Texture2D, InterferenceSeedTexture)
-
-		SHADER_PARAMETER_STRUCT_INCLUDE(FTSRPrevHistoryParameters, PrevHistoryParameters)
-		SHADER_PARAMETER_STRUCT(FTSRHistoryTextures, PrevHistory)
-
-		SHADER_PARAMETER_STRUCT(FTSRHistoryUAVs, HistoryOutput)
-		SHADER_PARAMETER_RDG_TEXTURE_UAV(RWTexture2D, InterferenceWeightOutput)
-		SHADER_PARAMETER_RDG_TEXTURE_UAV(RWTexture2D, DebugOutput)
-	END_SHADER_PARAMETER_STRUCT()
-}; // class FTSRDetectInterferenceCS
-
-class FTSRFilterFrequenciesCS : public FTemporalSuperResolutionShader
-{
-	DECLARE_GLOBAL_SHADER(FTSRFilterFrequenciesCS);
-	SHADER_USE_PARAMETER_STRUCT(FTSRFilterFrequenciesCS, FTemporalSuperResolutionShader);
-
-	BEGIN_SHADER_PARAMETER_STRUCT(FParameters, )
-		SHADER_PARAMETER_STRUCT_INCLUDE(FTSRCommonParameters, CommonParameters)
-		SHADER_PARAMETER(FVector3f, OutputQuantizationError)
-
-		SHADER_PARAMETER_RDG_TEXTURE(Texture2D, InputTexture)
-		SHADER_PARAMETER_RDG_TEXTURE(Texture2D, PredictionSceneColorTexture)
-		SHADER_PARAMETER_RDG_TEXTURE(Texture2D, ParallaxRejectionMaskTexture)
-		SHADER_PARAMETER_RDG_TEXTURE(Texture2D, InterferenceWeightTexture)
-
-		SHADER_PARAMETER_RDG_TEXTURE_UAV(RWTexture2D, FilteredInputOutput)
-		SHADER_PARAMETER_RDG_TEXTURE_UAV(RWTexture2D, FilteredPredictionSceneColorOutput)
-		SHADER_PARAMETER_RDG_TEXTURE_UAV(RWTexture2D, DebugOutput)
-	END_SHADER_PARAMETER_STRUCT()
-}; // class FTSRFilterFrequenciesCS
-
-class FTSRCompareHistoryCS : public FTemporalSuperResolutionShader
-{
-	DECLARE_GLOBAL_SHADER(FTSRCompareHistoryCS);
-	SHADER_USE_PARAMETER_STRUCT(FTSRCompareHistoryCS, FTemporalSuperResolutionShader);
-
-	class FOutputHalfRes : SHADER_PERMUTATION_BOOL("DIM_OUTPUT_HALF_RES");
-	using FPermutationDomain = TShaderPermutationDomain<FOutputHalfRes>;
-
-	BEGIN_SHADER_PARAMETER_STRUCT(FParameters, )
-		SHADER_PARAMETER_STRUCT_INCLUDE(FTSRCommonParameters, CommonParameters)
-		SHADER_PARAMETER_RDG_TEXTURE(Texture2D, ParallaxRejectionMaskTexture)
-		SHADER_PARAMETER_RDG_TEXTURE(Texture2D, FilteredInputTexture)
-		SHADER_PARAMETER_RDG_TEXTURE(Texture2D, FilteredPredictionSceneColorTexture)
-		SHADER_PARAMETER_RDG_TEXTURE(Texture2D, InterferenceWeightTexture)
-		SHADER_PARAMETER_RDG_TEXTURE_UAV(RWTexture2D, HistoryRejectionOutput)
-		SHADER_PARAMETER_RDG_TEXTURE_UAV(RWTexture2D, DebugOutput)
-	END_SHADER_PARAMETER_STRUCT()
-}; // class FTSRCompareHistoryCS
-
-class FTSRPostfilterRejectionCS : public FTemporalSuperResolutionShader
-{
-	DECLARE_GLOBAL_SHADER(FTSRPostfilterRejectionCS);
-	SHADER_USE_PARAMETER_STRUCT(FTSRPostfilterRejectionCS, FTemporalSuperResolutionShader);
-
-	class FOutputHalfRes : SHADER_PERMUTATION_BOOL("DIM_OUTPUT_HALF_RES");
-	using FPermutationDomain = TShaderPermutationDomain<FOutputHalfRes>;
-
-	BEGIN_SHADER_PARAMETER_STRUCT(FParameters, )
-		SHADER_PARAMETER(FIntRect, HistoryRejectionViewport)
-		SHADER_PARAMETER_RDG_TEXTURE(Texture2D, HistoryRejectionTexture)
-		SHADER_PARAMETER_RDG_TEXTURE_UAV(RWTexture2D, HistoryRejectionOutput)
-		SHADER_PARAMETER_RDG_TEXTURE_UAV(RWTexture2D, DebugOutput)
-	END_SHADER_PARAMETER_STRUCT()
-}; // class FTSRPostfilterRejectionCS
-
-class FTSRDilateRejectionCS : public FTemporalSuperResolutionShader
-{
-	DECLARE_GLOBAL_SHADER(FTSRDilateRejectionCS);
-	SHADER_USE_PARAMETER_STRUCT(FTSRDilateRejectionCS, FTemporalSuperResolutionShader);
-
-	BEGIN_SHADER_PARAMETER_STRUCT(FParameters, )
-		SHADER_PARAMETER_STRUCT_INCLUDE(FTSRCommonParameters, CommonParameters)
-		SHADER_PARAMETER_RDG_TEXTURE(Texture2D, HistoryRejectionTexture)
-		SHADER_PARAMETER_RDG_TEXTURE_UAV(RWTexture2D, DilatedHistoryRejectionOutput)
-		SHADER_PARAMETER_RDG_TEXTURE_UAV(RWTexture2D, DebugOutput)
-	END_SHADER_PARAMETER_STRUCT()
-}; // class FTSRDilateRejectionCS
-
-class FTSRUpdateSuperResHistoryCS : public FTemporalSuperResolutionShader
-{
-	DECLARE_GLOBAL_SHADER(FTSRUpdateSuperResHistoryCS);
-	SHADER_USE_PARAMETER_STRUCT(FTSRUpdateSuperResHistoryCS, FTemporalSuperResolutionShader);
-
-	BEGIN_SHADER_PARAMETER_STRUCT(FParameters, )
-		SHADER_PARAMETER_STRUCT_INCLUDE(FTSRCommonParameters, CommonParameters)
-		SHADER_PARAMETER_RDG_TEXTURE(Texture2D, InputSceneColorTexture)
-
-		SHADER_PARAMETER_RDG_TEXTURE(Texture2D, DilatedVelocityTexture)
-		SHADER_PARAMETER_RDG_TEXTURE(Texture2D, ParallaxRejectionMaskTexture)
-
-		SHADER_PARAMETER_STRUCT_INCLUDE(FTSRPrevHistoryParameters, PrevHistoryParameters)
-		SHADER_PARAMETER_STRUCT(FTSRHistoryTextures, PrevHistory)
-
-		SHADER_PARAMETER_STRUCT(FTSRHistoryUAVs, HistoryOutput)
-		SHADER_PARAMETER_RDG_TEXTURE_UAV(RWTexture2D, DebugOutput)
-	END_SHADER_PARAMETER_STRUCT()
-}; // class FTSRUpdateSuperResHistoryCS
-
-class FTSRUpdateHistoryCS : public FTemporalSuperResolutionShader
-{
-	DECLARE_GLOBAL_SHADER(FTSRUpdateHistoryCS);
-	SHADER_USE_PARAMETER_STRUCT(FTSRUpdateHistoryCS, FTemporalSuperResolutionShader);
-
-	BEGIN_SHADER_PARAMETER_STRUCT(FParameters, )
-		SHADER_PARAMETER_STRUCT_INCLUDE(FTSRCommonParameters, CommonParameters)
-		SHADER_PARAMETER_RDG_TEXTURE(Texture2D, InputSceneColorTexture)
-		SHADER_PARAMETER_RDG_TEXTURE_SRV(Texture2D, InputSceneStencilTexture)
-		SHADER_PARAMETER_RDG_TEXTURE(Texture2D, HistoryRejectionTexture)
-		SHADER_PARAMETER_RDG_TEXTURE(Texture2D, TranslucencyRejectionTexture)
-
-		SHADER_PARAMETER_RDG_TEXTURE(Texture2D, DilatedVelocityTexture)
-		SHADER_PARAMETER_RDG_TEXTURE(Texture2D, ParallaxFactorTexture)
-		SHADER_PARAMETER_RDG_TEXTURE(Texture2D, ParallaxRejectionMaskTexture)
-
-		SHADER_PARAMETER(FScreenTransform, HistoryPixelPosToScreenPos)
-		SHADER_PARAMETER(FScreenTransform, HistoryPixelPosToPPCo)
-		SHADER_PARAMETER(FVector3f, HistoryQuantizationError)
-		SHADER_PARAMETER(float, MinTranslucencyRejection)
-		SHADER_PARAMETER(int32, ResponsiveStencilMask)
-
-		SHADER_PARAMETER_STRUCT_INCLUDE(FTSRPrevHistoryParameters, PrevHistoryParameters)
-		SHADER_PARAMETER_STRUCT(FTSRHistoryTextures, PrevHistory)
-
-		SHADER_PARAMETER_STRUCT(FTSRHistoryUAVs, HistoryOutput)
-		SHADER_PARAMETER_RDG_TEXTURE_UAV(RWTexture2D, SceneColorOutput)
-		SHADER_PARAMETER_RDG_TEXTURE_UAV(RWTexture2D, DebugOutput)
-	END_SHADER_PARAMETER_STRUCT()
-
-	static void ModifyCompilationEnvironment(const FGlobalShaderPermutationParameters& Parameters, FShaderCompilerEnvironment& OutEnvironment)
-	{
-		FTemporalSuperResolutionShader::ModifyCompilationEnvironment(Parameters, OutEnvironment);
-		OutEnvironment.CompilerFlags.Add(CFLAG_Wave32);
-	}
-}; // class FTSRUpdateHistoryCS
-
-IMPLEMENT_GLOBAL_SHADER(FTAAStandaloneCS, "/Engine/Private/TemporalAA/TAAStandalone.usf", "MainCS", SF_Compute);
-IMPLEMENT_GLOBAL_SHADER(FTSRClearPrevTexturesCS, "/Engine/Private/TemporalAA/TAAClearPrevTextures.usf", "MainCS", SF_Compute);
-IMPLEMENT_GLOBAL_SHADER(FTSRDilateVelocityCS, "/Engine/Private/TemporalAA/TAADilateVelocity.usf", "MainCS", SF_Compute);
-IMPLEMENT_GLOBAL_SHADER(FTSRDecimateHistoryCS, "/Engine/Private/TemporalAA/TAADecimateHistory.usf", "MainCS", SF_Compute);
-IMPLEMENT_GLOBAL_SHADER(FTSRCompareTranslucencyCS, "/Engine/Private/TemporalAA/TAACompareTranslucency.usf", "MainCS", SF_Compute);
-IMPLEMENT_GLOBAL_SHADER(FTSRDetectInterferenceCS, "/Engine/Private/TemporalAA/TAADetectInterference.usf", "MainCS", SF_Compute);
-IMPLEMENT_GLOBAL_SHADER(FTSRFilterFrequenciesCS, "/Engine/Private/TemporalAA/TAAFilterFrequencies.usf", "MainCS", SF_Compute);
-IMPLEMENT_GLOBAL_SHADER(FTSRCompareHistoryCS, "/Engine/Private/TemporalAA/TAACompareHistory.usf", "MainCS", SF_Compute);
-IMPLEMENT_GLOBAL_SHADER(FTSRPostfilterRejectionCS, "/Engine/Private/TemporalAA/TAAPostfilterRejection.usf", "MainCS", SF_Compute);
-IMPLEMENT_GLOBAL_SHADER(FTSRDilateRejectionCS, "/Engine/Private/TemporalAA/TAADilateRejection.usf", "MainCS", SF_Compute);
-IMPLEMENT_GLOBAL_SHADER(FTSRUpdateSuperResHistoryCS, "/Engine/Private/TemporalAA/TAAUpdateSuperResHistory.usf", "MainCS", SF_Compute);
-IMPLEMENT_GLOBAL_SHADER(FTSRUpdateHistoryCS, "/Engine/Private/TemporalAA/TAAUpdateHistory.usf", "MainCS", SF_Compute);
-
-#if COMPILE_TAA_DEBUG_PASSES
-
-class FTSRDebugHistoryCS : public FTemporalSuperResolutionShader
-{
-	DECLARE_GLOBAL_SHADER(FTSRDebugHistoryCS);
-	SHADER_USE_PARAMETER_STRUCT(FTSRDebugHistoryCS, FTemporalSuperResolutionShader);
-
-	BEGIN_SHADER_PARAMETER_STRUCT(FParameters, )
-		SHADER_PARAMETER_STRUCT_INCLUDE(FTSRCommonParameters, CommonParameters)
-		SHADER_PARAMETER_STRUCT(FTSRHistoryTextures, History)
-		SHADER_PARAMETER_STRUCT(FTSRHistoryTextures, PrevHistory)
-		SHADER_PARAMETER_RDG_TEXTURE_UAV(RWTexture2D, DebugOutput)
-	END_SHADER_PARAMETER_STRUCT()
-}; // class FTSRDebugHistoryCS
-
-IMPLEMENT_GLOBAL_SHADER(FTSRDebugHistoryCS, "/Engine/Private/TemporalAA/TAADebugHistory.usf", "MainCS", SF_Compute);
+}; // class FTemporalAACS
 
-#endif
+IMPLEMENT_GLOBAL_SHADER(FTemporalAACS, "/Engine/Private/TemporalAA.usf", "MainCS", SF_Compute);
 
 float CatmullRom(float x)
 {
@@ -659,26 +265,7 @@ float CatmullRom(float x)
 		return (1.5f * ax - 2.5f) * ax*ax + 1.0f;
 }
 
-FVector ComputePixelFormatQuantizationError(EPixelFormat PixelFormat)
-{
-	FVector Error;
-	if (PixelFormat == PF_FloatRGBA || PixelFormat == PF_FloatR11G11B10)
-	{
-		FIntVector HistoryColorMantissaBits = PixelFormat == PF_FloatR11G11B10 ? FIntVector(6, 6, 5) : FIntVector(10, 10, 10);
-
-		Error.X = FMath::Pow(0.5f, HistoryColorMantissaBits.X);
-		Error.Y = FMath::Pow(0.5f, HistoryColorMantissaBits.Y);
-		Error.Z = FMath::Pow(0.5f, HistoryColorMantissaBits.Z);
-	}
-	else
-	{
-		check(0);
-	}
-
-	return Error;
-}
-
-void SetupSampleWeightParameters(FTAAStandaloneCS::FParameters* OutTAAParameters, const FTAAPassParameters& PassParameters, FVector2D TemporalJitterPixels)
+void SetupSampleWeightParameters(FTemporalAACS::FParameters* OutTAAParameters, const FTAAPassParameters& PassParameters, FVector2D TemporalJitterPixels)
 {
 	float JitterX = TemporalJitterPixels.X;
 	float JitterY = TemporalJitterPixels.Y;
@@ -748,7 +335,6 @@ void SetupSampleWeightParameters(FTAAStandaloneCS::FParameters* OutTAAParameters
 }
 
 DECLARE_GPU_STAT(TAA)
-DECLARE_GPU_STAT(TemporalSuperResolution)
 
 const TCHAR* const kTAAOutputNames[] = {
 	TEXT("TAA.History"),
@@ -779,6 +365,25 @@ bool IsTemporalAASceneDownsampleAllowed(const FViewInfo& View)
 	return CVarTemporalAAAllowDownsampling.GetValueOnRenderThread() != 0;
 }
 
+FVector ComputePixelFormatQuantizationError(EPixelFormat PixelFormat)
+{
+	FVector Error;
+	if (PixelFormat == PF_FloatRGBA || PixelFormat == PF_FloatR11G11B10)
+	{
+		FIntVector HistoryColorMantissaBits = PixelFormat == PF_FloatR11G11B10 ? FIntVector(6, 6, 5) : FIntVector(10, 10, 10);
+
+		Error.X = FMath::Pow(0.5f, HistoryColorMantissaBits.X);
+		Error.Y = FMath::Pow(0.5f, HistoryColorMantissaBits.Y);
+		Error.Z = FMath::Pow(0.5f, HistoryColorMantissaBits.Z);
+	}
+	else
+	{
+		check(0);
+	}
+
+	return Error;
+}
+
 float GetTemporalAAHistoryUpscaleFactor(const FViewInfo& View)
 {
 	float UpscaleFactor = 1.0f;
@@ -916,39 +521,39 @@ FTAAOutputs AddTemporalAAPass(
 	TStaticArray<bool, FTemporalAAHistory::kRenderTargetCount> bUseHistoryTexture;
 
 	{
-		FTAAStandaloneCS::FPermutationDomain PermutationVector;
-		PermutationVector.Set<FTAAStandaloneCS::FTAAPassConfigDim>(Inputs.Pass);
-		PermutationVector.Set<FTAAStandaloneCS::FTAAFastDim>(Inputs.bUseFast);
-		PermutationVector.Set<FTAAStandaloneCS::FTAADownsampleDim>(Inputs.bDownsample);
-		PermutationVector.Set<FTAAStandaloneCS::FTAAUpsampleFilteredDim>(true);
+		FTemporalAACS::FPermutationDomain PermutationVector;
+		PermutationVector.Set<FTemporalAACS::FTAAPassConfigDim>(Inputs.Pass);
+		PermutationVector.Set<FTemporalAACS::FTAAFastDim>(Inputs.bUseFast);
+		PermutationVector.Set<FTemporalAACS::FTAADownsampleDim>(Inputs.bDownsample);
+		PermutationVector.Set<FTemporalAACS::FTAAUpsampleFilteredDim>(true);
 
 		if (IsTAAUpsamplingConfig(Inputs.Pass))
 		{
 			const bool bUpsampleFiltered = CVarTemporalAAUpsampleFiltered.GetValueOnRenderThread() != 0 || Inputs.Pass != ETAAPassConfig::MainUpsampling;
-			PermutationVector.Set<FTAAStandaloneCS::FTAAUpsampleFilteredDim>(bUpsampleFiltered);
+			PermutationVector.Set<FTemporalAACS::FTAAUpsampleFilteredDim>(bUpsampleFiltered);
 
 			// If screen percentage > 100% on X or Y axes, then use screen percentage range = 2 shader permutation to disable LDS caching.
 			if (SrcRect.Width() > DestRect.Width() ||
 				SrcRect.Height() > DestRect.Height())
 			{
-				PermutationVector.Set<FTAAStandaloneCS::FTAAScreenPercentageDim>(2);
+				PermutationVector.Set<FTemporalAACS::FTAAScreenPercentageDim>(2);
 			}
 			// If screen percentage < 50% on X and Y axes, then use screen percentage range = 3 shader permutation.
 			else if (SrcRect.Width() * 100 < 50 * DestRect.Width() &&
 				SrcRect.Height() * 100 < 50 * DestRect.Height() &&
 				Inputs.Pass == ETAAPassConfig::MainSuperSampling)
 			{
-				PermutationVector.Set<FTAAStandaloneCS::FTAAScreenPercentageDim>(3);
+				PermutationVector.Set<FTemporalAACS::FTAAScreenPercentageDim>(3);
 			}
 			// If screen percentage < 71% on X and Y axes, then use screen percentage range = 1 shader permutation to have smaller LDS caching.
 			else if (SrcRect.Width() * 100 < 71 * DestRect.Width() &&
 				SrcRect.Height() * 100 < 71 * DestRect.Height())
 			{
-				PermutationVector.Set<FTAAStandaloneCS::FTAAScreenPercentageDim>(1);
+				PermutationVector.Set<FTemporalAACS::FTAAScreenPercentageDim>(1);
 			}
 		}
 
-		FTAAStandaloneCS::FParameters* PassParameters = GraphBuilder.AllocParameters<FTAAStandaloneCS::FParameters>();
+		FTemporalAACS::FParameters* PassParameters = GraphBuilder.AllocParameters<FTemporalAACS::FParameters>();
 
 		// Setups common shader parameters
 		const FIntPoint InputExtent = Inputs.SceneColorInput->Desc.Extent;
@@ -1138,7 +743,7 @@ FTAAOutputs AddTemporalAAPass(
 			PassParameters->DebugOutput = GraphBuilder.CreateUAV(DebugTexture);
 		}
 
-		TShaderMapRef<FTAAStandaloneCS> ComputeShader(View.ShaderMap, PermutationVector);
+		TShaderMapRef<FTemporalAACS> ComputeShader(View.ShaderMap, PermutationVector);
 
 		ClearUnusedGraphResources(ComputeShader, PassParameters);
 		for (int32 i = 0; i < FTemporalAAHistory::kRenderTargetCount; i++)
@@ -1176,970 +781,102 @@ FTAAOutputs AddTemporalAAPass(
 	return Outputs;
 } // AddTemporalAAPass()
 
-static void AddTemporalSuperResolutionPasses(
+static void AddGen4MainTemporalAAPasses(
 	FRDGBuilder& GraphBuilder,
 	const FViewInfo& View,
 	const ITemporalUpscaler::FPassInputs& PassInputs,
 	FRDGTextureRef* OutSceneColorTexture,
-	FIntRect* OutSceneColorViewRect)
+	FIntRect* OutSceneColorViewRect,
+	FRDGTextureRef* OutSceneColorHalfResTexture,
+	FIntRect* OutSceneColorHalfResViewRect)
 {
-	const FTemporalAAHistory& InputHistory = View.PrevViewInfo.TemporalAAHistory;
-
-#if COMPILE_TAA_DEBUG_PASSES
-	const bool bSetupDebugPasses = CVarTAASetupDebugPasses.GetValueOnRenderThread() != 0;
-#endif
-
-	// Whether to use camera cut shader permutation or not.
-	bool bCameraCut = !InputHistory.IsValid() || View.bCameraCut;
-
-	bool bHalfResLowFrequency = CVarTAAHalfResShadingRejection.GetValueOnRenderThread() != 0;
-
-	bool bEnableInterferenceHeuristic = CVarTAAEnableAntiInterference.GetValueOnRenderThread() != 0;
+	check(View.AntiAliasingMethod == AAM_TemporalAA && View.ViewState);
 
-	bool bRejectSeparateTranslucency = PassInputs.SeparateTranslucencyTextures != nullptr && CVarTAARejectTranslucency.GetValueOnRenderThread() != 0;
+	FTAAPassParameters TAAParameters(View);
 
-	enum class ERejectionPostFilter : uint8
-	{
-		Disabled,
-		PostRejectionDownsample,
-		PreRejectionDownsample,
-	};
+	TAAParameters.Pass = View.PrimaryScreenPercentageMethod == EPrimaryScreenPercentageMethod::TemporalUpscale
+		? ETAAPassConfig::MainUpsampling
+		: ETAAPassConfig::Main;
 
-	ERejectionPostFilter PostFilter = ERejectionPostFilter(FMath::Clamp(CVarTAAFilterShadingRejection.GetValueOnRenderThread(), 0, 2));
+	TAAParameters.SetupViewRect(View);
 
-	FIntPoint InputExtent = PassInputs.SceneColorTexture->Desc.Extent;
-	FIntRect InputRect = View.ViewRect;
+	const EPostProcessAAQuality LowQualityTemporalAA = EPostProcessAAQuality::Medium;
 
-	FIntPoint LowFrequencyExtent = InputExtent;
-	FIntRect LowFrequencyRect = FIntRect(FIntPoint(0, 0), InputRect.Size());
+	TAAParameters.bUseFast = GetPostProcessAAQuality() == LowQualityTemporalAA;
 
-	if (bHalfResLowFrequency)
-	{
-		LowFrequencyExtent = InputExtent / 2;
-		LowFrequencyRect = FIntRect(FIntPoint(0, 0), FIntPoint::DivideAndRoundUp(InputRect.Size(), 2));
-	}
+	const FIntRect SecondaryViewRect = TAAParameters.OutputViewRect;
 
-	FIntPoint RejectionExtent = LowFrequencyExtent / 2;
-	FIntRect RejectionRect = FIntRect(FIntPoint(0, 0), FIntPoint::DivideAndRoundUp(LowFrequencyRect.Size(), 2));
+	const float HistoryUpscaleFactor = GetTemporalAAHistoryUpscaleFactor(View);
 
-	FIntPoint OutputExtent;
-	FIntRect OutputRect;
-	if (View.PrimaryScreenPercentageMethod == EPrimaryScreenPercentageMethod::TemporalUpscale)
+	// Configures TAA to upscale the history buffer; this is in addition to the secondary screen percentage upscale.
+	// We end up with a scene color that is larger than the secondary screen percentage. We immediately downscale
+	// afterwards using a Mitchel-Netravali filter.
+	if (HistoryUpscaleFactor > 1.0f)
 	{
-		OutputRect.Min = FIntPoint(0, 0);
-		OutputRect.Max = View.GetSecondaryViewRectSize();
+		const FIntPoint HistoryViewSize(
+			TAAParameters.OutputViewRect.Width() * HistoryUpscaleFactor,
+			TAAParameters.OutputViewRect.Height() * HistoryUpscaleFactor);
 
-		FIntPoint QuantizedPrimaryUpscaleViewSize;
-		QuantizeSceneBufferSize(OutputRect.Max, QuantizedPrimaryUpscaleViewSize);
+		TAAParameters.Pass = ETAAPassConfig::MainSuperSampling;
+		TAAParameters.bUseFast = false;
 
-		OutputExtent = FIntPoint(
-			FMath::Max(InputExtent.X, QuantizedPrimaryUpscaleViewSize.X),
-			FMath::Max(InputExtent.Y, QuantizedPrimaryUpscaleViewSize.Y));
-	}
-	else
-	{
-		OutputRect.Min = FIntPoint(0, 0);
-		OutputRect.Max = View.ViewRect.Size();
-		OutputExtent = InputExtent;
+		TAAParameters.OutputViewRect.Min.X = 0;
+		TAAParameters.OutputViewRect.Min.Y = 0;
+		TAAParameters.OutputViewRect.Max = HistoryViewSize;
 	}
 
-	FIntPoint HistoryExtent;
-	FIntPoint HistorySize;
-	{
-		float UpscaleFactor = FMath::Clamp(CVarTemporalAAHistorySP.GetValueOnRenderThread() / 100.0f, 1.0f, 2.0f);
+	TAAParameters.DownsampleOverrideFormat = PassInputs.DownsampleOverrideFormat;
+
+	TAAParameters.bDownsample = PassInputs.bAllowDownsampleSceneColor && TAAParameters.bUseFast;
 
-		HistorySize = FIntPoint(
-			FMath::CeilToInt(OutputRect.Width() * UpscaleFactor),
-			FMath::CeilToInt(OutputRect.Height() * UpscaleFactor));
+	TAAParameters.SceneDepthTexture = PassInputs.SceneDepthTexture;
+	TAAParameters.SceneVelocityTexture = PassInputs.SceneVelocityTexture;
+	TAAParameters.SceneColorInput = PassInputs.SceneColorTexture;
 
-		FIntPoint QuantizedHistoryViewSize;
-		QuantizeSceneBufferSize(HistorySize, QuantizedHistoryViewSize);
+	const FTemporalAAHistory& InputHistory = View.PrevViewInfo.TemporalAAHistory;
 
-		HistoryExtent = FIntPoint(
-			FMath::Max(InputExtent.X, QuantizedHistoryViewSize.X),
-			FMath::Max(InputExtent.Y, QuantizedHistoryViewSize.Y));
-	}
+	FTemporalAAHistory& OutputHistory = View.ViewState->PrevFrameViewInfo.TemporalAAHistory;
 
-	RDG_EVENT_SCOPE(GraphBuilder, "TemporalSuperResolution %dx%d -> %dx%d", InputRect.Width(), InputRect.Height(), OutputRect.Width(), OutputRect.Height());
-	RDG_GPU_STAT_SCOPE(GraphBuilder, TemporalSuperResolution);
 
-	FRDGTextureRef BlackUintDummy = GSystemTextures.GetZeroUIntDummy(GraphBuilder);
-	FRDGTextureRef BlackDummy = GraphBuilder.RegisterExternalTexture(GSystemTextures.BlackDummy);
-	FRDGTextureRef WhiteDummy = GraphBuilder.RegisterExternalTexture(GSystemTextures.WhiteDummy);
+	const FTAAOutputs TAAOutputs = ::AddTemporalAAPass(
+		GraphBuilder,
+		View,
+		TAAParameters,
+		InputHistory,
+		&OutputHistory);
 
-	FRDGTextureRef SeparateTranslucencyTexture = bRejectSeparateTranslucency ?
-		PassInputs.SeparateTranslucencyTextures->GetColorForRead(GraphBuilder) : nullptr;
+	FRDGTextureRef SceneColorTexture = TAAOutputs.SceneColor;
 
-	FTSRCommonParameters CommonParameters;
+	// If we upscaled the history buffer, downsize back to the secondary screen percentage size.
+	if (HistoryUpscaleFactor > 1.0f)
 	{
-		CommonParameters.InputInfo = GetScreenPassTextureViewportParameters(FScreenPassTextureViewport(
-			InputExtent, InputRect));
-		CommonParameters.InputPixelPosMin = CommonParameters.InputInfo.ViewportMin;
-		CommonParameters.InputPixelPosMax = CommonParameters.InputInfo.ViewportMax - 1;
-
-		CommonParameters.LowFrequencyInfo = GetScreenPassTextureViewportParameters(FScreenPassTextureViewport(
-			LowFrequencyExtent, LowFrequencyRect));
+		const FIntRect InputViewport = TAAParameters.OutputViewRect;
 
-		CommonParameters.RejectionInfo = GetScreenPassTextureViewportParameters(FScreenPassTextureViewport(
-			RejectionExtent, RejectionRect));
+		FIntPoint QuantizedOutputSize;
+		QuantizeSceneBufferSize(SecondaryViewRect.Size(), QuantizedOutputSize);
 
-		CommonParameters.HistoryInfo = GetScreenPassTextureViewportParameters(FScreenPassTextureViewport(
-			HistoryExtent, FIntRect(FIntPoint(0, 0), HistorySize)));
+		FScreenPassTextureViewport OutputViewport;
+		OutputViewport.Rect = SecondaryViewRect;
+		OutputViewport.Extent.X = FMath::Max(PassInputs.SceneColorTexture->Desc.Extent.X, QuantizedOutputSize.X);
+		OutputViewport.Extent.Y = FMath::Max(PassInputs.SceneColorTexture->Desc.Extent.Y, QuantizedOutputSize.Y);
 
-		CommonParameters.InputJitter = View.TemporalJitterPixels;
-		CommonParameters.bCameraCut = bCameraCut;
-		CommonParameters.bEnableInterferenceHeuristic = bEnableInterferenceHeuristic;
-		CommonParameters.ViewUniformBuffer = View.ViewUniformBuffer;
+		SceneColorTexture = ComputeMitchellNetravaliDownsample(GraphBuilder, View, FScreenPassTexture(SceneColorTexture, InputViewport), OutputViewport);
 	}
 
-	auto CreateDebugUAV = [&](const FIntPoint& Extent, const TCHAR* DebugName)
-	{
-		FRDGTextureDesc DebugDesc = FRDGTextureDesc::Create2D(
-			Extent,
-			PF_FloatRGBA,
-			FClearValueBinding::None,
-			/* InFlags = */ TexCreate_ShaderResource | TexCreate_UAV);
-
-		FRDGTextureRef DebugTexture = GraphBuilder.CreateTexture(DebugDesc, DebugName);
-
-		return GraphBuilder.CreateUAV(DebugTexture);
-	};
+	*OutSceneColorTexture = SceneColorTexture;
+	*OutSceneColorViewRect = SecondaryViewRect;
+	*OutSceneColorHalfResTexture = TAAOutputs.DownsampledSceneColor;
+	*OutSceneColorHalfResViewRect = FIntRect::DivideAndRoundUp(SecondaryViewRect, 2);
+} // AddGen4MainTemporalAAPasses()
 
-	// Clear atomic scattered texture.
-	FRDGTextureRef PrevUseCountTexture;
-	FRDGTextureRef PrevClosestDepthTexture;
-	{
-		{
-			FRDGTextureDesc Desc = FRDGTextureDesc::Create2D(
-				InputExtent,
-				PF_R32_UINT,
-				FClearValueBinding::None,
-				/* InFlags = */ TexCreate_ShaderResource | TexCreate_UAV);
+bool DoesPlatformSupportTSR(EShaderPlatform Platform);
 
-			PrevUseCountTexture = GraphBuilder.CreateTexture(Desc, TEXT("TSR.PrevUseCountTexture"));
-			PrevClosestDepthTexture = GraphBuilder.CreateTexture(Desc, TEXT("TSR.PrevClosestDepthTexture"));
-		}
-
-		FTSRClearPrevTexturesCS::FParameters* PassParameters = GraphBuilder.AllocParameters<FTSRClearPrevTexturesCS::FParameters>();
-		PassParameters->CommonParameters = CommonParameters;
-		PassParameters->PrevUseCountOutput = GraphBuilder.CreateUAV(PrevUseCountTexture);
-		PassParameters->PrevClosestDepthOutput = GraphBuilder.CreateUAV(PrevClosestDepthTexture);
-
-		TShaderMapRef<FTSRClearPrevTexturesCS> ComputeShader(View.ShaderMap);
-		FComputeShaderUtils::AddPass(
-			GraphBuilder,
-			RDG_EVENT_NAME("TSR ClearPrevTextures %dx%d", InputRect.Width(), InputRect.Height()),
-			ComputeShader,
-			PassParameters,
-			FComputeShaderUtils::GetGroupCount(InputRect.Size(), 8));
-	}
-
-	// Dilate the velocity texture & scatter reprojection into previous frame
-	FRDGTextureRef DilatedVelocityTexture;
-	FRDGTextureRef ClosestDepthTexture;
-	FRDGTextureRef ParallaxFactorTexture;
-	{
-		{
-			FRDGTextureDesc Desc = FRDGTextureDesc::Create2D(
-				InputExtent,
-				PF_G16R16,
-				FClearValueBinding::None,
-				/* InFlags = */ TexCreate_ShaderResource | TexCreate_UAV);
-
-			DilatedVelocityTexture = GraphBuilder.CreateTexture(Desc, TEXT("TSR.DilatedVelocity"));
-
-			Desc.Format = PF_R16F;
-			ClosestDepthTexture = GraphBuilder.CreateTexture(Desc, TEXT("TSR.ClosestDepthTexture"));
-
-			Desc.Format = PF_R8_UINT;
-			ParallaxFactorTexture = GraphBuilder.CreateTexture(Desc, TEXT("TSR.ParallaxFactor"));
-		}
-
-		FTSRDilateVelocityCS::FParameters* PassParameters = GraphBuilder.AllocParameters<FTSRDilateVelocityCS::FParameters>();
-		PassParameters->CommonParameters = CommonParameters;
-		PassParameters->PrevOutputBufferUVMin = CommonParameters.InputInfo.UVViewportBilinearMin - CommonParameters.InputInfo.ExtentInverse;
-		PassParameters->PrevOutputBufferUVMax = CommonParameters.InputInfo.UVViewportBilinearMax + CommonParameters.InputInfo.ExtentInverse;
-		PassParameters->SceneDepthTexture = PassInputs.SceneDepthTexture;
-		PassParameters->SceneVelocityTexture = PassInputs.SceneVelocityTexture;
-		PassParameters->DilatedVelocityOutput = GraphBuilder.CreateUAV(DilatedVelocityTexture);
-		PassParameters->ClosestDepthOutput = GraphBuilder.CreateUAV(ClosestDepthTexture);
-		PassParameters->PrevUseCountOutput = GraphBuilder.CreateUAV(PrevUseCountTexture);
-		PassParameters->PrevClosestDepthOutput = GraphBuilder.CreateUAV(PrevClosestDepthTexture);
-		PassParameters->ParallaxFactorOutput = GraphBuilder.CreateUAV(ParallaxFactorTexture);
-		PassParameters->DebugOutput = CreateDebugUAV(InputExtent, TEXT("Debug.TSR.DilateVelocity"));
-
-		TShaderMapRef<FTSRDilateVelocityCS> ComputeShader(View.ShaderMap);
-		FComputeShaderUtils::AddPass(
-			GraphBuilder,
-			RDG_EVENT_NAME("TSR DilateVelocity %dx%d", InputRect.Width(), InputRect.Height()),
-			ComputeShader,
-			PassParameters,
-			FComputeShaderUtils::GetGroupCount(InputRect.Size(), 8));
-	}
-
-	// Setup the previous frame history
-	FTSRPrevHistoryParameters PrevHistoryParameters;
-	FTSRHistoryTextures PrevHistory;
-	{
-		FScreenPassTextureViewport PrevHistoryViewport(InputHistory.ReferenceBufferSize, InputHistory.ViewportRect);
-		if (bCameraCut)
-		{
-			PrevHistoryViewport.Extent = FIntPoint(1, 1);
-			PrevHistoryViewport.Rect = FIntRect(FIntPoint(0, 0), FIntPoint(1, 1));
-		}
-
-		PrevHistoryParameters.PrevHistoryInfo = GetScreenPassTextureViewportParameters(PrevHistoryViewport);
-		PrevHistoryParameters.ScreenPosToPrevHistoryBufferUV = FScreenTransform::ChangeTextureBasisFromTo(
-			PrevHistoryViewport, FScreenTransform::ETextureBasis::ScreenPosition, FScreenTransform::ETextureBasis::TextureUV);
-		PrevHistoryParameters.HistoryPreExposureCorrection = View.PreExposure / View.PrevViewInfo.SceneColorPreExposure;
-
-		for (int32 i = 0; i < PrevHistory.LowResTextures.Num(); i++)
-		{
-			if (InputHistory.LowResRT[i].IsValid() && !bCameraCut)
-			{
-				PrevHistory.LowResTextures[i] = GraphBuilder.RegisterExternalTexture(InputHistory.LowResRT[i]);
-			}
-			else
-			{
-				PrevHistory.LowResTextures[i] = BlackDummy;
-			}
-		}
-
-		for (int32 i = 0; i < PrevHistory.Textures.Num(); i++)
-		{
-			if (InputHistory.RT[i].IsValid() && !bCameraCut)
-			{
-				PrevHistory.Textures[i] = GraphBuilder.RegisterExternalTexture(InputHistory.RT[i]);
-			}
-			else if (i == 3)
-			{
-				PrevHistory.Textures[i] = BlackUintDummy;
-			}
-			else
-			{
-				PrevHistory.Textures[i] = BlackDummy;
-			}
-		}
-
-		for (int32 i = 0; i < PrevHistory.SuperResTextures.Num(); i++)
-		{
-			if (InputHistory.SuperResRT[i].IsValid() && !bCameraCut)
-			{
-				PrevHistory.SuperResTextures[i] = GraphBuilder.RegisterExternalTexture(InputHistory.SuperResRT[i]);
-			}
-			else
-			{
-				PrevHistory.SuperResTextures[i] = BlackUintDummy;
-			}
-		}
-	}
-
-
-	FTSRHistoryTextures History;
-	{
-		FRDGTextureDesc LowResDesc = FRDGTextureDesc::Create2D(
-			InputExtent,
-			PF_R8,
-			FClearValueBinding::None,
-			TexCreate_ShaderResource | TexCreate_UAV);
-
-		History.LowResTextures[0] = GraphBuilder.CreateTexture(LowResDesc, TEXT("TSR.History.LowResMetadata[0]"));
-	}
-
-	{
-		FRDGTextureDesc Desc = FRDGTextureDesc::Create2D(
-			HistoryExtent,
-			(CVarTAAR11G11B10History.GetValueOnRenderThread() != 0) ? PF_FloatR11G11B10 : PF_FloatRGBA,
-			FClearValueBinding::None,
-			TexCreate_ShaderResource | TexCreate_UAV);
-
-		History.Textures[0] = GraphBuilder.CreateTexture(Desc, TEXT("TSR.History.LowFrequencies"));
-		History.Textures[1] = GraphBuilder.CreateTexture(Desc, TEXT("TSR.History.HighFrequencies"));
-
-		Desc.Format = PF_R8G8;
-		History.Textures[2] = GraphBuilder.CreateTexture(Desc, TEXT("TSR.History.Metadata"));
-
-		Desc.Format = PF_R16_UINT;
-		History.Textures[3] = GraphBuilder.CreateTexture(Desc, TEXT("TSR.History.SubpixelInfo"));
-	}
-
-	{
-		FRDGTextureDesc SuperResDesc = FRDGTextureDesc::Create2D(
-			HistoryExtent * 2,
-			PF_R16_UINT,
-			FClearValueBinding::None,
-			TexCreate_ShaderResource | TexCreate_UAV);
-
-		History.SuperResTextures[0] = GraphBuilder.CreateTexture(SuperResDesc, TEXT("TSR.History.SuperResMetadata[0]"));
-
-		SuperResDesc.Format = PF_R8_UINT;
-		History.SuperResTextures[1] = GraphBuilder.CreateTexture(SuperResDesc, TEXT("TSR.History.SuperResMetadata[1]"));
-	}
-
-	// Decimate input to flicker at same frequency as input.
-	FRDGTextureRef HalfResInputSceneColorTexture = nullptr;
-	FRDGTextureRef HalfResPredictionSceneColorTexture = nullptr;
-	FRDGTextureRef HalfResParallaxRejectionMaskTexture = nullptr;
-	FRDGTextureRef PredictionSceneColorTexture = nullptr;
-	FRDGTextureRef ParallaxRejectionMaskTexture = nullptr;
-	FRDGTextureRef InterferenceSeedTexture = nullptr;
-	{
-		{
-			FRDGTextureDesc Desc = FRDGTextureDesc::Create2D(
-				InputExtent,
-				PF_R8,
-				FClearValueBinding::None,
-				/* InFlags = */ TexCreate_ShaderResource | TexCreate_UAV);
-
-			ParallaxRejectionMaskTexture = GraphBuilder.CreateTexture(Desc, TEXT("TSR.ParallaxRejectionMask"));
-
-			// TODO(G5TAA): can compress to the history seed's 4bit per pixel
-			Desc.Format = PF_R8G8B8A8;
-			InterferenceSeedTexture = GraphBuilder.CreateTexture(Desc, TEXT("TSR.Interference.Seed"));
-		}
-
-		FTSRDecimateHistoryCS::FParameters* PassParameters = GraphBuilder.AllocParameters<FTSRDecimateHistoryCS::FParameters>();
-		PassParameters->CommonParameters = CommonParameters;
-		{
-			const FViewMatrices& ViewMatrices = View.ViewMatrices;
-			const FViewMatrices& PrevViewMatrices = View.PrevViewInfo.ViewMatrices;
-
-			FMatrix RotationalInvViewProj = ViewMatrices.ComputeInvProjectionNoAAMatrix() * (ViewMatrices.GetTranslatedViewMatrix().RemoveTranslation().GetTransposed());
-			FMatrix RotationalPrevViewProj = (PrevViewMatrices.GetTranslatedViewMatrix().RemoveTranslation()) * PrevViewMatrices.ComputeProjectionNoAAMatrix();
-
-			PassParameters->RotationalClipToPrevClip = RotationalInvViewProj * RotationalPrevViewProj;
-		}
-		PassParameters->OutputQuantizationError = ComputePixelFormatQuantizationError(PF_FloatR11G11B10);
-		{
-			float TanHalfFieldOfView = View.ViewMatrices.GetInvProjectionMatrix().M[0][0];
-
-			// Should be multiplied 0.5* for the diameter to radius, and by 2.0 because GetTanHalfFieldOfView() cover only half of the pixels.
-			PassParameters->WorldDepthToPixelWorldRadius = TanHalfFieldOfView / float(View.ViewRect.Width());
-		}
-
-		PassParameters->InputSceneColorTexture = PassInputs.SceneColorTexture;
-		PassParameters->DilatedVelocityTexture = DilatedVelocityTexture;
-		PassParameters->ClosestDepthTexture = ClosestDepthTexture;
-		PassParameters->PrevUseCountTexture = PrevUseCountTexture;
-		PassParameters->PrevClosestDepthTexture = PrevClosestDepthTexture;
-		PassParameters->ParallaxFactorTexture = ParallaxFactorTexture;
-
-		PassParameters->PrevHistoryParameters = PrevHistoryParameters;
-		PassParameters->PrevHistory = PrevHistory;
-
-		if (bHalfResLowFrequency)
-		{
-			FRDGTextureDesc Desc = FRDGTextureDesc::Create2D(
-				LowFrequencyExtent,
-				PF_FloatR11G11B10,
-				FClearValueBinding::None,
-				/* InFlags = */ TexCreate_ShaderResource | TexCreate_UAV);
-
-			HalfResInputSceneColorTexture = GraphBuilder.CreateTexture(Desc, TEXT("TSR.HalfResInput"));
-			HalfResPredictionSceneColorTexture = GraphBuilder.CreateTexture(Desc, TEXT("TSR.Prediction.SceneColor"));
-
-			Desc.Format = PF_R8;
-			HalfResParallaxRejectionMaskTexture = GraphBuilder.CreateTexture(Desc, TEXT("TSR.HalfResParallaxRejectionMask"));
-
-			PassParameters->HalfResSceneColorOutput = GraphBuilder.CreateUAV(HalfResInputSceneColorTexture);
-			PassParameters->HalfResPredictionSceneColorOutput = GraphBuilder.CreateUAV(HalfResPredictionSceneColorTexture);
-			PassParameters->HalfResParallaxRejectionMaskOutput = GraphBuilder.CreateUAV(HalfResParallaxRejectionMaskTexture);
-		}
-		else
-		{
-			FRDGTextureDesc Desc = FRDGTextureDesc::Create2D(
-				InputExtent,
-				PF_FloatR11G11B10,
-				FClearValueBinding::None,
-				/* InFlags = */ TexCreate_ShaderResource | TexCreate_UAV);
-
-			PredictionSceneColorTexture = GraphBuilder.CreateTexture(Desc, TEXT("TSR.Prediction.SceneColor"));
-
-			PassParameters->PredictionSceneColorOutput = GraphBuilder.CreateUAV(PredictionSceneColorTexture);
-		}
-
-		PassParameters->ParallaxRejectionMaskOutput = GraphBuilder.CreateUAV(ParallaxRejectionMaskTexture);
-		PassParameters->InterferenceSeedOutput = GraphBuilder.CreateUAV(InterferenceSeedTexture);
-		PassParameters->DebugOutput = CreateDebugUAV(LowFrequencyExtent, TEXT("Debug.TSR.DecimateHistory"));
-
-		FTSRDecimateHistoryCS::FPermutationDomain PermutationVector;
-		PermutationVector.Set<FTSRDecimateHistoryCS::FOutputHalfRes>(bHalfResLowFrequency);
-
-		TShaderMapRef<FTSRDecimateHistoryCS> ComputeShader(View.ShaderMap, PermutationVector);
-		FComputeShaderUtils::AddPass(
-			GraphBuilder,
-			RDG_EVENT_NAME("TSR DecimateHistory(%s) %dx%d",
-				bHalfResLowFrequency ? TEXT("HalfResShadingOutput") : TEXT(""),
-				InputRect.Width(), InputRect.Height()),
-			ComputeShader,
-			PassParameters,
-			FComputeShaderUtils::GetGroupCount(InputRect.Size(), 8));
-	}
-
-	FRDGTextureRef TranslucencyRejectionTexture = nullptr;
-	if (bRejectSeparateTranslucency && View.PrevViewInfo.SeparateTranslucency != nullptr)
-	{
-		FRDGTextureRef PrevTranslucencyTexture;
-		FScreenPassTextureViewport PrevTranslucencyViewport;
-
-		if (View.PrevViewInfo.SeparateTranslucency)
-		{
-
-			PrevTranslucencyTexture = GraphBuilder.RegisterExternalTexture(View.PrevViewInfo.SeparateTranslucency);
-			PrevTranslucencyViewport = FScreenPassTextureViewport(PrevTranslucencyTexture->Desc.Extent, View.PrevViewInfo.ViewRect);
-
-		}
-		else
-		{
-			PrevTranslucencyTexture = BlackDummy;
-			PrevTranslucencyViewport = FScreenPassTextureViewport(FIntPoint(1, 1), FIntRect(FIntPoint(0, 0), FIntPoint(1, 1)));
-		}
-
-		{
-			FRDGTextureDesc Desc = FRDGTextureDesc::Create2D(
-				InputExtent,
-				PF_R8,
-				FClearValueBinding::None,
-				/* InFlags = */ TexCreate_ShaderResource | TexCreate_UAV);
-
-			TranslucencyRejectionTexture = GraphBuilder.CreateTexture(Desc, TEXT("TSR.TranslucencyRejection"));
-		}
-
-		FTSRCompareTranslucencyCS::FParameters* PassParameters = GraphBuilder.AllocParameters<FTSRCompareTranslucencyCS::FParameters>();
-		PassParameters->CommonParameters = CommonParameters;
-
-		PassParameters->TranslucencyInfo = GetScreenPassTextureViewportParameters(FScreenPassTextureViewport(
-			InputExtent, InputRect));
-		PassParameters->PrevTranslucencyInfo = GetScreenPassTextureViewportParameters(PrevTranslucencyViewport);
-		PassParameters->PrevTranslucencyPreExposureCorrection = PrevHistoryParameters.HistoryPreExposureCorrection;
-
-		PassParameters->InputPixelPosToScreenPos = (FScreenTransform::Identity + 0.5) * CommonParameters.InputInfo.ViewportSizeInverse * FScreenTransform::ViewportUVToScreenPos;
-		PassParameters->ScreenPosToPrevTranslucencyTextureUV = FScreenTransform::ChangeTextureBasisFromTo(
-			PrevTranslucencyViewport, FScreenTransform::ETextureBasis::ScreenPosition, FScreenTransform::ETextureBasis::TextureUV);
-
-		PassParameters->DilatedVelocityTexture = DilatedVelocityTexture;
-		PassParameters->TranslucencyTexture = SeparateTranslucencyTexture;
-		PassParameters->PrevTranslucencyTexture = PrevTranslucencyTexture;
-
-		PassParameters->TranslucencyRejectionOutput = GraphBuilder.CreateUAV(TranslucencyRejectionTexture);
-		PassParameters->DebugOutput = CreateDebugUAV(LowFrequencyExtent, TEXT("Debug.TSR.ComparetTranslucency"));
-
-		TShaderMapRef<FTSRCompareTranslucencyCS> ComputeShader(View.ShaderMap);
-		FComputeShaderUtils::AddPass(
-			GraphBuilder,
-			RDG_EVENT_NAME("TSR ComparetTranslucency %dx%d", InputRect.Width(), InputRect.Height()),
-			ComputeShader,
-			PassParameters,
-			FComputeShaderUtils::GetGroupCount(InputRect.Size(), 8));
-	}
-
-	// Detect interference between geometry and alignement of input texel centers. It is not about awnser whether an
-	// interference has happened in the past, because interference change based on input resolution or camera position.
-	// So to remain stable on camera movement and input resolution change, it is about awnsering the question on whether
-	// an interference is possible.
-	// TODO(G5TAA): Could sample the interference seed in the DilateVelocity and detect interference in the decimate.
-	TStaticArray<bool, FTemporalAAHistory::kLowResRenderTargetCount> ExtractLowResHistoryTexture;
-	FRDGTextureRef InterferenceWeightTexture;
-	if (bEnableInterferenceHeuristic)
-	{
-		{
-			// TODO(G5TAA): Compress to 1bit per pixel
-			FRDGTextureDesc Desc = FRDGTextureDesc::Create2D(
-				InputExtent,
-				PF_R8,
-				FClearValueBinding::None,
-				/* InFlags = */ TexCreate_ShaderResource | TexCreate_UAV);
-
-			InterferenceWeightTexture = GraphBuilder.CreateTexture(Desc, TEXT("TSR.Interference.Weight"));
-		}
-
-		FTSRDetectInterferenceCS::FParameters* PassParameters = GraphBuilder.AllocParameters<FTSRDetectInterferenceCS::FParameters>();
-		PassParameters->CommonParameters = CommonParameters;
-
-		PassParameters->InputSceneColorTexture = PassInputs.SceneColorTexture;
-		PassParameters->DilatedVelocityTexture = DilatedVelocityTexture;
-		PassParameters->PredictionSceneColorTexture = PredictionSceneColorTexture;
-		PassParameters->ParallaxRejectionMaskTexture = ParallaxRejectionMaskTexture;
-		PassParameters->InterferenceSeedTexture = InterferenceSeedTexture;
-
-		PassParameters->PrevHistoryParameters = PrevHistoryParameters;
-		PassParameters->PrevHistory = PrevHistory;
-
-		PassParameters->HistoryOutput = CreateUAVs(GraphBuilder, History);
-		PassParameters->InterferenceWeightOutput = GraphBuilder.CreateUAV(InterferenceWeightTexture);
-		PassParameters->DebugOutput = CreateDebugUAV(LowFrequencyExtent, TEXT("Debug.TSR.DetectInterference"));
-
-		TShaderMapRef<FTSRDetectInterferenceCS> ComputeShader(View.ShaderMap);
-		ClearUnusedGraphResources(ComputeShader, PassParameters);
-
-		for (int32 i = 0; i < PrevHistory.LowResTextures.Num(); i++)
-		{
-			bool bNeedsExtractForNextFrame = PassParameters->PrevHistory.LowResTextures[i] != nullptr;
-			bool bPrevFrameIsntAvailable = PassParameters->PrevHistory.LowResTextures[i] == BlackDummy;
-			bool bOutputHistory = PassParameters->HistoryOutput.LowResTextures[i] != nullptr;
-
-			ExtractLowResHistoryTexture[i] = bNeedsExtractForNextFrame;
-
-			if (bPrevFrameIsntAvailable && !PassParameters->CommonParameters.bCameraCut)
-			{
-				//ensureMsgf(false, TEXT("Shaders read PrevHistory[%d] but doesn't write HistoryOutput[%d]"), i, i);
-				PassParameters->CommonParameters.bCameraCut = true;
-			}
-
-			if (bOutputHistory && !bNeedsExtractForNextFrame)
-			{
-				ensureMsgf(false, TEXT("Shaders write HistoryOutput[%d] but doesn't read PrevHistory[%d]"), i, i);
-			}
-		}
-
-		FComputeShaderUtils::AddPass(
-			GraphBuilder,
-			RDG_EVENT_NAME("TSR DetectInterference %dx%d", InputRect.Width(), InputRect.Height()),
-			ComputeShader,
-			PassParameters,
-			FComputeShaderUtils::GetGroupCount(InputRect.Size(), 8));
-	}
-	else
-	{
-		for (int32 i = 0; i < ExtractLowResHistoryTexture.Num(); i++)
-		{
-			ExtractLowResHistoryTexture[i] = false;
-		}
-
-		// TODO(G5TAA): Shader permutation.
-		InterferenceWeightTexture = WhiteDummy;
-	}
-
-	// Reject the history with frequency decomposition.
-	FRDGTextureRef HistoryRejectionTexture;
-	{
-		// Filter out the high frquencies
-		FRDGTextureRef FilteredInputTexture;
-		FRDGTextureRef FilteredPredictionSceneColorTexture;
-		{
-			{
-				FRDGTextureDesc Desc = FRDGTextureDesc::Create2D(
-					LowFrequencyExtent,
-					PF_FloatR11G11B10,
-					FClearValueBinding::None,
-					/* InFlags = */ TexCreate_ShaderResource | TexCreate_UAV);
-
-				FilteredInputTexture = GraphBuilder.CreateTexture(Desc, TEXT("TSR.Filtered.SceneColor"));
-				FilteredPredictionSceneColorTexture = GraphBuilder.CreateTexture(Desc, TEXT("TSR.Filtered.Prediction.SceneColor"));
-			}
-
-			FTSRFilterFrequenciesCS::FParameters* PassParameters = GraphBuilder.AllocParameters<FTSRFilterFrequenciesCS::FParameters>();
-			PassParameters->CommonParameters = CommonParameters;
-			PassParameters->OutputQuantizationError = ComputePixelFormatQuantizationError(FilteredInputTexture->Desc.Format);
-
-			if (bHalfResLowFrequency)
-			{
-				PassParameters->InputTexture = HalfResInputSceneColorTexture;
-				PassParameters->PredictionSceneColorTexture = HalfResPredictionSceneColorTexture;
-				PassParameters->ParallaxRejectionMaskTexture = HalfResParallaxRejectionMaskTexture;
-			}
-			else
-			{
-				PassParameters->InputTexture = PassInputs.SceneColorTexture;
-				PassParameters->PredictionSceneColorTexture = PredictionSceneColorTexture;
-				PassParameters->ParallaxRejectionMaskTexture = ParallaxRejectionMaskTexture;
-				PassParameters->InterferenceWeightTexture = InterferenceWeightTexture;
-			}
-
-			PassParameters->FilteredInputOutput = GraphBuilder.CreateUAV(FilteredInputTexture);
-			PassParameters->FilteredPredictionSceneColorOutput = GraphBuilder.CreateUAV(FilteredPredictionSceneColorTexture);
-			PassParameters->DebugOutput = CreateDebugUAV(LowFrequencyExtent, TEXT("Debug.TSR.FilterFrequencies"));
-
-			TShaderMapRef<FTSRFilterFrequenciesCS> ComputeShader(View.ShaderMap);
-			FComputeShaderUtils::AddPass(
-				GraphBuilder,
-				RDG_EVENT_NAME("TSR FilterFrequencies %dx%d", LowFrequencyRect.Width(), LowFrequencyRect.Height()),
-				ComputeShader,
-				PassParameters,
-				FComputeShaderUtils::GetGroupCount(LowFrequencyRect.Size(), 16));
-		}
-
-		// Compare the low frequencies
-		{
-			bool bOutputHalfRes = PostFilter != ERejectionPostFilter::PreRejectionDownsample;
-
-			{
-				FRDGTextureDesc Desc = FRDGTextureDesc::Create2D(
-					bOutputHalfRes ? RejectionExtent : LowFrequencyExtent,
-					PF_R8,
-					FClearValueBinding::None,
-					/* InFlags = */ TexCreate_ShaderResource | TexCreate_UAV);
-
-				HistoryRejectionTexture = GraphBuilder.CreateTexture(Desc, TEXT("TSR.HistoryRejection"));
-			}
-
-			FTSRCompareHistoryCS::FParameters* PassParameters = GraphBuilder.AllocParameters<FTSRCompareHistoryCS::FParameters>();
-			PassParameters->CommonParameters = CommonParameters;
-			PassParameters->ParallaxRejectionMaskTexture = ParallaxRejectionMaskTexture;
-			PassParameters->FilteredInputTexture = FilteredInputTexture;
-			PassParameters->FilteredPredictionSceneColorTexture = FilteredPredictionSceneColorTexture;
-			PassParameters->InterferenceWeightTexture = InterferenceWeightTexture;
-
-			PassParameters->HistoryRejectionOutput = GraphBuilder.CreateUAV(HistoryRejectionTexture);
-			PassParameters->DebugOutput = CreateDebugUAV(LowFrequencyExtent, TEXT("Debug.TSR.CompareHistory"));
-
-			FTSRCompareHistoryCS::FPermutationDomain PermutationVector;
-			PermutationVector.Set<FTSRCompareHistoryCS::FOutputHalfRes>(bOutputHalfRes);
-
-			TShaderMapRef<FTSRCompareHistoryCS> ComputeShader(View.ShaderMap, PermutationVector);
-			FComputeShaderUtils::AddPass(
-				GraphBuilder,
-				RDG_EVENT_NAME("TSR CompareHistory %dx%d", LowFrequencyRect.Width(), LowFrequencyRect.Height()),
-				ComputeShader,
-				PassParameters,
-				FComputeShaderUtils::GetGroupCount(LowFrequencyRect.Size(), 16));
-		}
-	}
-
-	// Post filter the rejection.
-	if (PostFilter != ERejectionPostFilter::Disabled)
-	{
-		bool bOutputHalfRes = PostFilter == ERejectionPostFilter::PreRejectionDownsample;
-		FIntRect Rect = bOutputHalfRes ? LowFrequencyRect : RejectionRect;
-
-		FRDGTextureRef FilteredHistoryRejectionTexture;
-		{
-			FRDGTextureDesc Desc = FRDGTextureDesc::Create2D(
-				RejectionExtent,
-				PF_R8,
-				FClearValueBinding::None,
-				/* InFlags = */ TexCreate_ShaderResource | TexCreate_UAV);
-
-			FilteredHistoryRejectionTexture = GraphBuilder.CreateTexture(Desc, TEXT("TSR.HistoryRejection"));
-		}
-
-		FTSRPostfilterRejectionCS::FParameters* PassParameters = GraphBuilder.AllocParameters<FTSRPostfilterRejectionCS::FParameters>();
-		PassParameters->HistoryRejectionViewport = Rect;
-		PassParameters->HistoryRejectionTexture = HistoryRejectionTexture;
-		PassParameters->HistoryRejectionOutput = GraphBuilder.CreateUAV(FilteredHistoryRejectionTexture);
-		PassParameters->DebugOutput = CreateDebugUAV(bOutputHalfRes ? LowFrequencyExtent : RejectionExtent, TEXT("Debug.TSR.PostfilterRejection"));
-
-		FTSRPostfilterRejectionCS::FPermutationDomain PermutationVector;
-		PermutationVector.Set<FTSRPostfilterRejectionCS::FOutputHalfRes>(PostFilter == ERejectionPostFilter::PreRejectionDownsample);
-
-		TShaderMapRef<FTSRPostfilterRejectionCS> ComputeShader(View.ShaderMap, PermutationVector);
-		FComputeShaderUtils::AddPass(
-			GraphBuilder,
-			RDG_EVENT_NAME("TSR PostfilterRejection %dx%d", Rect.Width(), Rect.Height()),
-			ComputeShader,
-			PassParameters,
-			FComputeShaderUtils::GetGroupCount(Rect.Size(), 8));
-
-		HistoryRejectionTexture = FilteredHistoryRejectionTexture;
-	}
-
-	// Dilate the rejection.
-	FRDGTextureRef DilatedHistoryRejectionTexture;
-	{
-		DilatedHistoryRejectionTexture = GraphBuilder.CreateTexture(HistoryRejectionTexture->Desc, TEXT("TSR.DilatedHistoryRejection"));
-
-		FTSRDilateRejectionCS::FParameters* PassParameters = GraphBuilder.AllocParameters<FTSRDilateRejectionCS::FParameters>();
-		PassParameters->CommonParameters = CommonParameters;
-		PassParameters->HistoryRejectionTexture = HistoryRejectionTexture;
-		PassParameters->DilatedHistoryRejectionOutput = GraphBuilder.CreateUAV(DilatedHistoryRejectionTexture);
-		PassParameters->DebugOutput = CreateDebugUAV(RejectionExtent, TEXT("Debug.TSR.DilateRejection"));
-
-		TShaderMapRef<FTSRDilateRejectionCS> ComputeShader(View.ShaderMap);
-		FComputeShaderUtils::AddPass(
-			GraphBuilder,
-			RDG_EVENT_NAME("TSR DilateRejection %dx%d", RejectionRect.Width(), RejectionRect.Height()),
-			ComputeShader,
-			PassParameters,
-			FComputeShaderUtils::GetGroupCount(RejectionRect.Size(), 8));
-	}
-
-	TStaticArray<bool, FTemporalAAHistory::kSuperResRenderTargetCount> ExtractSuperResHistoryTexture;
-	if (bEnableInterferenceHeuristic)
-	{
-		FTSRUpdateSuperResHistoryCS::FParameters* PassParameters = GraphBuilder.AllocParameters<FTSRUpdateSuperResHistoryCS::FParameters>();
-		PassParameters->CommonParameters = CommonParameters;
-		PassParameters->InputSceneColorTexture = PassInputs.SceneColorTexture;
-		PassParameters->DilatedVelocityTexture = DilatedVelocityTexture;
-		PassParameters->ParallaxRejectionMaskTexture = ParallaxRejectionMaskTexture;
-
-		PassParameters->PrevHistoryParameters = PrevHistoryParameters;
-		PassParameters->PrevHistory = PrevHistory;
-
-		PassParameters->HistoryOutput = CreateUAVs(GraphBuilder, History);
-		PassParameters->DebugOutput = CreateDebugUAV(HistoryExtent * 2, TEXT("Debug.TSR.UpdateSuperResHistory"));
-
-		TShaderMapRef<FTSRUpdateSuperResHistoryCS> ComputeShader(View.ShaderMap);
-		ClearUnusedGraphResources(ComputeShader, PassParameters);
-
-		for (int32 i = 0; i < PrevHistory.SuperResTextures.Num(); i++)
-		{
-			bool bNeedsExtractForNextFrame = PassParameters->PrevHistory.SuperResTextures[i] != nullptr;
-			bool bPrevFrameIsntAvailable = PassParameters->PrevHistory.SuperResTextures[i] == BlackDummy;
-			bool bOutputHistory = PassParameters->HistoryOutput.SuperResTextures[i] != nullptr;
-
-			ExtractSuperResHistoryTexture[i] = bNeedsExtractForNextFrame;
-
-			if (bPrevFrameIsntAvailable && !PassParameters->CommonParameters.bCameraCut)
-			{
-				//ensureMsgf(false, TEXT("Shaders read PrevHistory[%d] but doesn't write HistoryOutput[%d]"), i, i);
-				PassParameters->CommonParameters.bCameraCut = true;
-			}
-
-			if (bOutputHistory && !bNeedsExtractForNextFrame)
-			{
-				ensureMsgf(false, TEXT("Shaders write HistoryOutput[%d] but doesn't read PrevHistory[%d]"), i, i);
-			}
-		}
-
-		FComputeShaderUtils::AddPass(
-			GraphBuilder,
-			RDG_EVENT_NAME("TSR UpdateSuperResHistory %dx%d", HistorySize.X * 2, HistorySize.Y * 2),
-			ComputeShader,
-			PassParameters,
-			FComputeShaderUtils::GetGroupCount(HistorySize * 2, 8));
-	}
-	else
-	{
-		for (int32 i = 0; i < ExtractSuperResHistoryTexture.Num(); i++)
-		{
-			ExtractSuperResHistoryTexture[i] = false;
-		}
-	}
-
-	TStaticArray<bool, FTemporalAAHistory::kRenderTargetCount> ExtractHistoryTexture;
-	FRDGTextureRef SceneColorOutputTexture;
-	{
-		// Allocate output
-		{
-			FRDGTextureDesc Desc = FRDGTextureDesc::Create2D(
-				HistoryExtent,
-				PF_FloatR11G11B10,
-				FClearValueBinding::None,
-				/* InFlags = */ TexCreate_ShaderResource | TexCreate_UAV);
-
-			SceneColorOutputTexture = GraphBuilder.CreateTexture(Desc, TEXT("TSR.Output"));
-		}
-
-		FTSRUpdateHistoryCS::FParameters* PassParameters = GraphBuilder.AllocParameters<FTSRUpdateHistoryCS::FParameters>();
-		PassParameters->CommonParameters = CommonParameters;
-		PassParameters->InputSceneColorTexture = PassInputs.SceneColorTexture;
-		PassParameters->InputSceneStencilTexture = GraphBuilder.CreateSRV(
-			FRDGTextureSRVDesc::CreateWithPixelFormat(PassInputs.SceneDepthTexture, PF_X24_G8));
-		PassParameters->HistoryRejectionTexture = DilatedHistoryRejectionTexture;
-		PassParameters->TranslucencyRejectionTexture = TranslucencyRejectionTexture ? TranslucencyRejectionTexture : BlackDummy;
-
-		PassParameters->DilatedVelocityTexture = DilatedVelocityTexture;
-		PassParameters->ParallaxFactorTexture = ParallaxFactorTexture;
-		PassParameters->ParallaxRejectionMaskTexture = ParallaxRejectionMaskTexture;
-
-		FScreenTransform HistoryPixelPosToViewportUV = (FScreenTransform::Identity + 0.5f) * CommonParameters.HistoryInfo.ViewportSizeInverse;
-		PassParameters->HistoryPixelPosToScreenPos = HistoryPixelPosToViewportUV * FScreenTransform::ViewportUVToScreenPos;
-		PassParameters->HistoryPixelPosToPPCo = HistoryPixelPosToViewportUV * CommonParameters.InputInfo.ViewportSize + CommonParameters.InputJitter + CommonParameters.InputPixelPosMin;
-		PassParameters->HistoryQuantizationError = ComputePixelFormatQuantizationError(History.Textures[0]->Desc.Format);
-		PassParameters->MinTranslucencyRejection = TranslucencyRejectionTexture == nullptr ? 1.0 : 0.0;
-		PassParameters->ResponsiveStencilMask = CVarTAAEnableResponiveAA.GetValueOnRenderThread() ? (STENCIL_TEMPORAL_RESPONSIVE_AA_MASK) : 0;
-
-		PassParameters->PrevHistoryParameters = PrevHistoryParameters;
-		PassParameters->PrevHistory = PrevHistory;
-
-		PassParameters->HistoryOutput = CreateUAVs(GraphBuilder, History);
-		PassParameters->SceneColorOutput = GraphBuilder.CreateUAV(SceneColorOutputTexture);
-		PassParameters->DebugOutput = CreateDebugUAV(HistoryExtent, TEXT("Debug.TSR.UpdateHistory"));
-
-		TShaderMapRef<FTSRUpdateHistoryCS> ComputeShader(View.ShaderMap);
-		ClearUnusedGraphResources(ComputeShader, PassParameters);
-
-		for (int32 i = 0; i < PrevHistory.Textures.Num(); i++)
-		{
-			bool bNeedsExtractForNextFrame = PassParameters->PrevHistory.Textures[i] != nullptr;
-			bool bPrevFrameIsntAvailable = PassParameters->PrevHistory.Textures[i] == BlackDummy;
-			bool bOutputHistory = PassParameters->HistoryOutput.Textures[i] != nullptr;
-
-			ExtractHistoryTexture[i] = bNeedsExtractForNextFrame;
-
-			if (bPrevFrameIsntAvailable && !bCameraCut)
-			{
-				//ensureMsgf(false, TEXT("Shaders read PrevHistory[%d] but doesn't write HistoryOutput[%d]"), i, i);
-				PassParameters->CommonParameters.bCameraCut = true;
-			}
-
-			if (bOutputHistory && !bNeedsExtractForNextFrame)
-			{
-				ensureMsgf(false, TEXT("Shaders write HistoryOutput[%d] but doesn't read PrevHistory[%d]"), i, i);
-			}
-		}
-
-		FComputeShaderUtils::AddPass(
-			GraphBuilder,
-			RDG_EVENT_NAME("TSR UpdateHistory(%s) %dx%d", 
-				History.Textures[0]->Desc.Format == PF_FloatR11G11B10 ? TEXT("R11G11B10") : TEXT(""),
-				HistorySize.X, HistorySize.Y),
-			ComputeShader,
-			PassParameters,
-			FComputeShaderUtils::GetGroupCount(HistorySize, 8));
-	}
-
-	// Debug the history.
-	#if COMPILE_TAA_DEBUG_PASSES
-	if (bSetupDebugPasses)
-	{
-		const int32 kHistoryUpscalingFactor = 2;
-
-		FTSRDebugHistoryCS::FParameters* PassParameters = GraphBuilder.AllocParameters<FTSRDebugHistoryCS::FParameters>();
-		PassParameters->CommonParameters = CommonParameters;
-		PassParameters->History = History;
-		PassParameters->PrevHistory = PrevHistory;
-		PassParameters->DebugOutput = CreateDebugUAV(HistoryExtent * kHistoryUpscalingFactor, TEXT("Debug.TSR.History"));
-
-		TShaderMapRef<FTSRDebugHistoryCS> ComputeShader(View.ShaderMap);
-		FComputeShaderUtils::AddPass(
-			GraphBuilder,
-			RDG_EVENT_NAME("TSR DebugHistory %dx%d", HistorySize.X, HistorySize.Y),
-			ComputeShader,
-			PassParameters,
-			FComputeShaderUtils::GetGroupCount(HistorySize * kHistoryUpscalingFactor, 8));
-	}
-	#endif
-
-	if (!View.bStatePrevViewInfoIsReadOnly)
-	{
-		FTemporalAAHistory* OutputHistory = &View.ViewState->PrevFrameViewInfo.TemporalAAHistory;
-		OutputHistory->SafeRelease();
-
-		for (int32 i = 0; i < History.LowResTextures.Num(); i++)
-		{
-			if (ExtractLowResHistoryTexture[i])
-			{
-				GraphBuilder.QueueTextureExtraction(History.LowResTextures[i], &OutputHistory->LowResRT[i]);
-			}
-		}
-
-		for (int32 i = 0; i < History.Textures.Num(); i++)
-		{
-			if (ExtractHistoryTexture[i])
-			{
-				GraphBuilder.QueueTextureExtraction(History.Textures[i], &OutputHistory->RT[i]);
-			}
-		}
-
-		for (int32 i = 0; i < History.SuperResTextures.Num(); i++)
-		{
-			if (ExtractSuperResHistoryTexture[i])
-			{
-				GraphBuilder.QueueTextureExtraction(History.SuperResTextures[i], &OutputHistory->SuperResRT[i]);
-			}
-		}
-
-		OutputHistory->ViewportRect = FIntRect(FIntPoint(0, 0), HistorySize);
-		OutputHistory->ReferenceBufferSize = HistoryExtent;
-
-		if (bRejectSeparateTranslucency)
-		{
-			GraphBuilder.QueueTextureExtraction(
-				SeparateTranslucencyTexture, &View.ViewState->PrevFrameViewInfo.SeparateTranslucency);
-		}
-	}
-
-	// If we upscaled the history buffer, downsize back to the secondary screen percentage size.
-	if (HistorySize != OutputRect.Size())
-	{
-		SceneColorOutputTexture = ComputeMitchellNetravaliDownsample(
-			GraphBuilder, View,
-			/* InputViewport = */ FScreenPassTexture(SceneColorOutputTexture, FIntRect(FIntPoint(0, 0), HistorySize)),
-			/* OutputViewport = */ FScreenPassTextureViewport(OutputExtent, OutputRect));
-	}
-
-	*OutSceneColorTexture = SceneColorOutputTexture;
-	*OutSceneColorViewRect = OutputRect;
-} // AddTemporalSuperResolutionPasses()
-
-static void AddGen4MainTemporalAAPasses(
+void AddTemporalSuperResolutionPasses(
 	FRDGBuilder& GraphBuilder,
 	const FViewInfo& View,
 	const ITemporalUpscaler::FPassInputs& PassInputs,
 	FRDGTextureRef* OutSceneColorTexture,
-	FIntRect* OutSceneColorViewRect,
-	FRDGTextureRef* OutSceneColorHalfResTexture,
-	FIntRect* OutSceneColorHalfResViewRect)
-{
-	check(View.AntiAliasingMethod == AAM_TemporalAA && View.ViewState);
-
-	FTAAPassParameters TAAParameters(View);
-
-	TAAParameters.Pass = View.PrimaryScreenPercentageMethod == EPrimaryScreenPercentageMethod::TemporalUpscale
-		? ETAAPassConfig::MainUpsampling
-		: ETAAPassConfig::Main;
-
-	TAAParameters.SetupViewRect(View);
-
-	const EPostProcessAAQuality LowQualityTemporalAA = EPostProcessAAQuality::Medium;
-
-	TAAParameters.bUseFast = GetPostProcessAAQuality() == LowQualityTemporalAA;
-
-	const FIntRect SecondaryViewRect = TAAParameters.OutputViewRect;
-
-	const float HistoryUpscaleFactor = GetTemporalAAHistoryUpscaleFactor(View);
-
-	// Configures TAA to upscale the history buffer; this is in addition to the secondary screen percentage upscale.
-	// We end up with a scene color that is larger than the secondary screen percentage. We immediately downscale
-	// afterwards using a Mitchel-Netravali filter.
-	if (HistoryUpscaleFactor > 1.0f)
-	{
-		const FIntPoint HistoryViewSize(
-			TAAParameters.OutputViewRect.Width() * HistoryUpscaleFactor,
-			TAAParameters.OutputViewRect.Height() * HistoryUpscaleFactor);
-
-		TAAParameters.Pass = ETAAPassConfig::MainSuperSampling;
-		TAAParameters.bUseFast = false;
-
-		TAAParameters.OutputViewRect.Min.X = 0;
-		TAAParameters.OutputViewRect.Min.Y = 0;
-		TAAParameters.OutputViewRect.Max = HistoryViewSize;
-	}
-
-	TAAParameters.DownsampleOverrideFormat = PassInputs.DownsampleOverrideFormat;
-
-	TAAParameters.bDownsample = PassInputs.bAllowDownsampleSceneColor && TAAParameters.bUseFast;
-
-	TAAParameters.SceneDepthTexture = PassInputs.SceneDepthTexture;
-	TAAParameters.SceneVelocityTexture = PassInputs.SceneVelocityTexture;
-	TAAParameters.SceneColorInput = PassInputs.SceneColorTexture;
-
-	const FTemporalAAHistory& InputHistory = View.PrevViewInfo.TemporalAAHistory;
-
-	FTemporalAAHistory& OutputHistory = View.ViewState->PrevFrameViewInfo.TemporalAAHistory;
-
-
-	const FTAAOutputs TAAOutputs = ::AddTemporalAAPass(
-		GraphBuilder,
-		View,
-		TAAParameters,
-		InputHistory,
-		&OutputHistory);
-
-	FRDGTextureRef SceneColorTexture = TAAOutputs.SceneColor;
-
-	// If we upscaled the history buffer, downsize back to the secondary screen percentage size.
-	if (HistoryUpscaleFactor > 1.0f)
-	{
-		const FIntRect InputViewport = TAAParameters.OutputViewRect;
-
-		FIntPoint QuantizedOutputSize;
-		QuantizeSceneBufferSize(SecondaryViewRect.Size(), QuantizedOutputSize);
-
-		FScreenPassTextureViewport OutputViewport;
-		OutputViewport.Rect = SecondaryViewRect;
-		OutputViewport.Extent.X = FMath::Max(PassInputs.SceneColorTexture->Desc.Extent.X, QuantizedOutputSize.X);
-		OutputViewport.Extent.Y = FMath::Max(PassInputs.SceneColorTexture->Desc.Extent.Y, QuantizedOutputSize.Y);
-
-		SceneColorTexture = ComputeMitchellNetravaliDownsample(GraphBuilder, View, FScreenPassTexture(SceneColorTexture, InputViewport), OutputViewport);
-	}
-
-	*OutSceneColorTexture = SceneColorTexture;
-	*OutSceneColorViewRect = SecondaryViewRect;
-	*OutSceneColorHalfResTexture = TAAOutputs.DownsampledSceneColor;
-	*OutSceneColorHalfResViewRect = FIntRect::DivideAndRoundUp(SecondaryViewRect, 2);
-} // AddGen4MainTemporalAAPasses()
+	FIntRect* OutSceneColorViewRect);
 
 const ITemporalUpscaler* GTemporalUpscaler = nullptr;
 
@@ -2161,7 +898,7 @@ public:
 		FRDGTextureRef* OutSceneColorHalfResTexture,
 		FIntRect* OutSceneColorHalfResViewRect) const final
 	{
-		if (CVarTAAAlgorithm.GetValueOnRenderThread() && DoesPlatformSupportGen5TAA(View.GetShaderPlatform()))
+		if (CVarTAAAlgorithm.GetValueOnRenderThread() && DoesPlatformSupportTSR(View.GetShaderPlatform()))
 		{
 			*OutSceneColorHalfResTexture = nullptr;
 			//*OutSceneColorHalfResViewRect; // TODO.
@@ -2221,7 +958,7 @@ EMainTAAPassConfig ITemporalUpscaler::GetMainTAAPassConfig(const FViewInfo& View
 	{
 		return EMainTAAPassConfig::ThirdParty;
 	}
-	else if (CVarTAAAlgorithm.GetValueOnRenderThread() && DoesPlatformSupportGen5TAA(View.GetShaderPlatform()))
+	else if (CVarTAAAlgorithm.GetValueOnRenderThread() && DoesPlatformSupportTSR(View.GetShaderPlatform()))
 	{
 		return EMainTAAPassConfig::TSR;
 	}
