Separates TSR source code from TAA

diff --git a/Engine/Source/Runtime/Renderer/Private/PostProcess/TemporalSuperResolution.cpp b/Engine/Source/Runtime/Renderer/Private/PostProcess/TemporalSuperResolution.cpp
new file mode 100644
--- /dev/null
+++ b/Engine/Source/Runtime/Renderer/Private/PostProcess/TemporalSuperResolution.cpp
@@ -0,0 +1,1307 @@
+// Copyright Epic Games, Inc. All Rights Reserved.
+
+#include "PostProcess/TemporalAA.h"
+#include "PostProcess/PostProcessTonemap.h"
+#include "PostProcess/PostProcessMitchellNetravali.h"
+#include "PostProcess/PostProcessing.h"
+#include "ClearQuad.h"
+#include "PostProcessing.h"
+#include "SceneTextureParameters.h"
+#include "PixelShaderUtils.h"
+#include "ScenePrivate.h"
+#include "RendererModule.h"
+
+#define COMPILE_TSR_DEBUG_PASSES (!UE_BUILD_SHIPPING)
+
+bool DoesPlatformSupportTSR(EShaderPlatform Platform)
+{
+	// TODO(TSR): alpha channel is not supported yet
+	if (IsPostProcessingWithAlphaChannelSupported())
+	{
+		return false;
+	}
+	return FDataDrivenShaderPlatformInfo::GetSupportsGen5TemporalAA(Platform);
+}
+
+namespace
+{
+
+TAutoConsoleVariable<float> CVarTSRHistorySP(
+	TEXT("r.TSR.HistoryScreenPercentage"),
+	100.0f,
+	TEXT("Size of TSR's history."),
+	ECVF_RenderThreadSafe);
+
+TAutoConsoleVariable<int32> CVarTSRR11G11B10History(
+	TEXT("r.TSR.R11G11B10History"), 1,
+	TEXT("Select the bitdepth of the history."),
+	ECVF_RenderThreadSafe);
+
+TAutoConsoleVariable<int32> CVarTSRHalfResShadingRejection(
+	TEXT("r.TSR.ShadingRejection.HalfRes"), 0,
+	TEXT("Whether the shading rejection should be done at half res. Saves performance but may introduce back some flickering (default = 0)."),
+	ECVF_RenderThreadSafe);
+
+TAutoConsoleVariable<int32> CVarTSRFilterShadingRejection(
+	TEXT("r.TSR.ShadingRejection.SpatialFilter"), 1,
+	TEXT("Whether the shading rejection should have spatial statistical filtering pass to reduce flickering (default = 1).\n")
+	TEXT(" 0: Disabled;\n")
+	TEXT(" 1: Spatial filter pass is run at lower resolution than CompareHistory pass (default);\n")
+	TEXT(" 2: Spatial filter pass is run CompareHistory pass resolution to improve stability."),
+	ECVF_RenderThreadSafe);
+
+TAutoConsoleVariable<int32> CVarTSREnableAntiInterference(
+	TEXT("r.TSR.AntiInterference"), 0,
+	TEXT("Enable heuristic to detect geometric interference between input pixel grid alignement and structured geometry."),
+	ECVF_RenderThreadSafe);
+
+TAutoConsoleVariable<int32> CVarTSRRejectTranslucency(
+	TEXT("r.TSR.RejectSeparateTranslucency"), 0,
+	TEXT("Enable heuristic to reject based on the Separate Translucency."),
+	ECVF_RenderThreadSafe);
+
+TAutoConsoleVariable<int32> CVarTSREnableResponiveAA(
+	TEXT("r.TSR.EnableResponiveAA"), 1,
+	TEXT("Whether the responsive AA should be enabled."),
+	ECVF_RenderThreadSafe);
+
+#if COMPILE_TSR_DEBUG_PASSES
+
+TAutoConsoleVariable<int32> CVarTSRSetupDebugPasses(
+	TEXT("r.TSR.Debug.SetupExtraPasses"), 0,
+	TEXT("Whether to enable the debug passes"),
+	ECVF_RenderThreadSafe);
+
+#endif
+
+BEGIN_SHADER_PARAMETER_STRUCT(FTSRCommonParameters, )
+	SHADER_PARAMETER_STRUCT(FScreenPassTextureViewportParameters, InputInfo)
+	SHADER_PARAMETER_STRUCT(FScreenPassTextureViewportParameters, LowFrequencyInfo)
+	SHADER_PARAMETER_STRUCT(FScreenPassTextureViewportParameters, RejectionInfo)
+	SHADER_PARAMETER_STRUCT(FScreenPassTextureViewportParameters, HistoryInfo)
+
+	SHADER_PARAMETER(FIntPoint, InputPixelPosMin)
+	SHADER_PARAMETER(FIntPoint, InputPixelPosMax)
+
+	SHADER_PARAMETER(FVector2D, InputJitter)
+	SHADER_PARAMETER(int32, bCameraCut)
+	SHADER_PARAMETER(int32, bEnableInterferenceHeuristic)
+
+	SHADER_PARAMETER_STRUCT_REF(FViewUniformShaderParameters, ViewUniformBuffer)
+END_SHADER_PARAMETER_STRUCT()
+
+BEGIN_SHADER_PARAMETER_STRUCT(FTSRHistoryTextures, )
+	SHADER_PARAMETER_RDG_TEXTURE_ARRAY(Texture2D, LowResTextures, [FTemporalAAHistory::kLowResRenderTargetCount])
+	SHADER_PARAMETER_RDG_TEXTURE_ARRAY(Texture2D, Textures, [FTemporalAAHistory::kRenderTargetCount])
+	SHADER_PARAMETER_RDG_TEXTURE_ARRAY(Texture2D, SuperResTextures, [FTemporalAAHistory::kSuperResRenderTargetCount])
+END_SHADER_PARAMETER_STRUCT()
+
+BEGIN_SHADER_PARAMETER_STRUCT(FTSRHistoryUAVs, )
+	SHADER_PARAMETER_RDG_TEXTURE_UAV_ARRAY(RWTexture2D, LowResTextures, [FTemporalAAHistory::kLowResRenderTargetCount])
+	SHADER_PARAMETER_RDG_TEXTURE_UAV_ARRAY(RWTexture2D, Textures, [FTemporalAAHistory::kRenderTargetCount])
+	SHADER_PARAMETER_RDG_TEXTURE_UAV_ARRAY(RWTexture2D, SuperResTextures, [FTemporalAAHistory::kSuperResRenderTargetCount])
+END_SHADER_PARAMETER_STRUCT()
+
+BEGIN_SHADER_PARAMETER_STRUCT(FTSRPrevHistoryParameters, )
+	SHADER_PARAMETER_STRUCT(FScreenPassTextureViewportParameters, PrevHistoryInfo)
+	SHADER_PARAMETER(FScreenTransform, ScreenPosToPrevHistoryBufferUV)
+	SHADER_PARAMETER(float, HistoryPreExposureCorrection)
+END_SHADER_PARAMETER_STRUCT()
+
+FTSRHistoryUAVs CreateUAVs(FRDGBuilder& GraphBuilder, const FTSRHistoryTextures& Textures)
+{
+	FTSRHistoryUAVs UAVs;
+	for (int32 i = 0; i < Textures.LowResTextures.Num(); i++)
+	{
+		UAVs.LowResTextures[i] = GraphBuilder.CreateUAV(Textures.LowResTextures[i]);
+	}
+	for (int32 i = 0; i < Textures.Textures.Num(); i++)
+	{
+		UAVs.Textures[i] = GraphBuilder.CreateUAV(Textures.Textures[i]);
+	}
+	for (int32 i = 0; i < Textures.SuperResTextures.Num(); i++)
+	{
+		UAVs.SuperResTextures[i] = GraphBuilder.CreateUAV(Textures.SuperResTextures[i]);
+	}
+	return UAVs;
+}
+
+class FTSRShader : public FGlobalShader
+{
+public:
+	FTSRShader(const ShaderMetaType::CompiledShaderInitializerType& Initializer)
+		: FGlobalShader(Initializer)
+	{ }
+
+	FTSRShader()
+	{ }
+
+	static bool ShouldCompilePermutation(const FGlobalShaderPermutationParameters& Parameters)
+	{
+		return DoesPlatformSupportTSR(Parameters.Platform);
+	}
+
+	static void ModifyCompilationEnvironment(const FGlobalShaderPermutationParameters& Parameters, FShaderCompilerEnvironment& OutEnvironment)
+	{
+		OutEnvironment.CompilerFlags.Add(CFLAG_AllowRealTypes);
+	}
+}; // class FTemporalSuperResolutionShader
+
+class FTSRClearPrevTexturesCS : public FTSRShader
+{
+	DECLARE_GLOBAL_SHADER(FTSRClearPrevTexturesCS);
+	SHADER_USE_PARAMETER_STRUCT(FTSRClearPrevTexturesCS, FTSRShader);
+
+	BEGIN_SHADER_PARAMETER_STRUCT(FParameters, )
+		SHADER_PARAMETER_STRUCT_INCLUDE(FTSRCommonParameters, CommonParameters)
+
+		SHADER_PARAMETER_RDG_TEXTURE_UAV(RWTexture2D, PrevUseCountOutput)
+		SHADER_PARAMETER_RDG_TEXTURE_UAV(RWTexture2D, PrevClosestDepthOutput)
+	END_SHADER_PARAMETER_STRUCT()
+}; // class FTSRClearPrevTexturesCS
+
+class FTSRDilateVelocityCS : public FTSRShader
+{
+	DECLARE_GLOBAL_SHADER(FTSRDilateVelocityCS);
+	SHADER_USE_PARAMETER_STRUCT(FTSRDilateVelocityCS, FTSRShader);
+
+	BEGIN_SHADER_PARAMETER_STRUCT(FParameters, )
+		SHADER_PARAMETER_STRUCT_INCLUDE(FTSRCommonParameters, CommonParameters)
+
+		SHADER_PARAMETER(FVector2D, PrevOutputBufferUVMin)
+		SHADER_PARAMETER(FVector2D, PrevOutputBufferUVMax)
+
+		SHADER_PARAMETER_RDG_TEXTURE(Texture2D, SceneDepthTexture)
+		SHADER_PARAMETER_RDG_TEXTURE(Texture2D, SceneVelocityTexture)
+
+		SHADER_PARAMETER_RDG_TEXTURE_UAV(RWTexture2D, DilatedVelocityOutput)
+		SHADER_PARAMETER_RDG_TEXTURE_UAV(RWTexture2D, ClosestDepthOutput)
+		SHADER_PARAMETER_RDG_TEXTURE_UAV(RWTexture2D, PrevUseCountOutput)
+		SHADER_PARAMETER_RDG_TEXTURE_UAV(RWTexture2D, PrevClosestDepthOutput)
+		SHADER_PARAMETER_RDG_TEXTURE_UAV(RWTexture2D, ParallaxFactorOutput)
+		SHADER_PARAMETER_RDG_TEXTURE_UAV(RWTexture2D, DebugOutput)
+	END_SHADER_PARAMETER_STRUCT()
+}; // class FTSRDilateVelocityCS
+
+class FTSRDecimateHistoryCS : public FTSRShader
+{
+	DECLARE_GLOBAL_SHADER(FTSRDecimateHistoryCS);
+	SHADER_USE_PARAMETER_STRUCT(FTSRDecimateHistoryCS, FTSRShader);
+
+	class FOutputHalfRes : SHADER_PERMUTATION_BOOL("DIM_OUTPUT_HALF_RES");
+
+	using FPermutationDomain = TShaderPermutationDomain<FOutputHalfRes>;
+
+	BEGIN_SHADER_PARAMETER_STRUCT(FParameters, )
+		SHADER_PARAMETER_STRUCT_INCLUDE(FTSRCommonParameters, CommonParameters)
+		SHADER_PARAMETER(FMatrix44f, RotationalClipToPrevClip)
+		SHADER_PARAMETER(FVector3f, OutputQuantizationError)
+		SHADER_PARAMETER(float, WorldDepthToPixelWorldRadius)
+
+		SHADER_PARAMETER_RDG_TEXTURE(Texture2D, InputSceneColorTexture)
+		SHADER_PARAMETER_RDG_TEXTURE(Texture2D, DilatedVelocityTexture)
+		SHADER_PARAMETER_RDG_TEXTURE(Texture2D, ClosestDepthTexture)
+		SHADER_PARAMETER_RDG_TEXTURE(Texture2D, PrevUseCountTexture)
+		SHADER_PARAMETER_RDG_TEXTURE(Texture2D, PrevClosestDepthTexture)
+		SHADER_PARAMETER_RDG_TEXTURE(Texture2D, ParallaxFactorTexture)
+
+		SHADER_PARAMETER_STRUCT_INCLUDE(FTSRPrevHistoryParameters, PrevHistoryParameters)
+		SHADER_PARAMETER_STRUCT(FTSRHistoryTextures, PrevHistory)
+
+		SHADER_PARAMETER_RDG_TEXTURE_UAV(RWTexture2D, HalfResSceneColorOutput)
+		SHADER_PARAMETER_RDG_TEXTURE_UAV(RWTexture2D, HalfResPredictionSceneColorOutput)
+		SHADER_PARAMETER_RDG_TEXTURE_UAV(RWTexture2D, HalfResParallaxRejectionMaskOutput)
+		SHADER_PARAMETER_RDG_TEXTURE_UAV(RWTexture2D, PredictionSceneColorOutput)
+		SHADER_PARAMETER_RDG_TEXTURE_UAV(RWTexture2D, ParallaxRejectionMaskOutput)
+		SHADER_PARAMETER_RDG_TEXTURE_UAV(RWTexture2D, InterferenceSeedOutput)
+		SHADER_PARAMETER_RDG_TEXTURE_UAV(RWTexture2D, DebugOutput)
+	END_SHADER_PARAMETER_STRUCT()
+}; // class FTSRDecimateHistoryCS
+
+class FTSRCompareTranslucencyCS : public FTSRShader
+{
+	DECLARE_GLOBAL_SHADER(FTSRCompareTranslucencyCS);
+	SHADER_USE_PARAMETER_STRUCT(FTSRCompareTranslucencyCS, FTSRShader);
+
+	BEGIN_SHADER_PARAMETER_STRUCT(FParameters, )
+		SHADER_PARAMETER_STRUCT_INCLUDE(FTSRCommonParameters, CommonParameters)
+
+		SHADER_PARAMETER_STRUCT(FScreenPassTextureViewportParameters, TranslucencyInfo)
+		SHADER_PARAMETER_STRUCT(FScreenPassTextureViewportParameters, PrevTranslucencyInfo)
+		SHADER_PARAMETER(float, PrevTranslucencyPreExposureCorrection)
+
+		SHADER_PARAMETER(FScreenTransform, InputPixelPosToScreenPos)
+		SHADER_PARAMETER(FScreenTransform, ScreenPosToPrevTranslucencyTextureUV)
+
+		SHADER_PARAMETER_RDG_TEXTURE(Texture2D, DilatedVelocityTexture)
+		SHADER_PARAMETER_RDG_TEXTURE(Texture2D, TranslucencyTexture)
+		SHADER_PARAMETER_RDG_TEXTURE(Texture2D, PrevTranslucencyTexture)
+
+		SHADER_PARAMETER_RDG_TEXTURE_UAV(RWTexture2D, TranslucencyRejectionOutput)
+		SHADER_PARAMETER_RDG_TEXTURE_UAV(RWTexture2D, DebugOutput)
+	END_SHADER_PARAMETER_STRUCT()
+}; // class FTSRCompareTranslucencyCS
+
+class FTSRDetectInterferenceCS : public FTSRShader
+{
+	DECLARE_GLOBAL_SHADER(FTSRDetectInterferenceCS);
+	SHADER_USE_PARAMETER_STRUCT(FTSRDetectInterferenceCS, FTSRShader);
+	
+	BEGIN_SHADER_PARAMETER_STRUCT(FParameters, )
+		SHADER_PARAMETER_STRUCT_INCLUDE(FTSRCommonParameters, CommonParameters)
+
+		SHADER_PARAMETER_RDG_TEXTURE(Texture2D, InputSceneColorTexture)
+		SHADER_PARAMETER_RDG_TEXTURE(Texture2D, DilatedVelocityTexture)
+		SHADER_PARAMETER_RDG_TEXTURE(Texture2D, PredictionSceneColorTexture)
+		SHADER_PARAMETER_RDG_TEXTURE(Texture2D, ParallaxRejectionMaskTexture)
+		SHADER_PARAMETER_RDG_TEXTURE(Texture2D, InterferenceSeedTexture)
+
+		SHADER_PARAMETER_STRUCT_INCLUDE(FTSRPrevHistoryParameters, PrevHistoryParameters)
+		SHADER_PARAMETER_STRUCT(FTSRHistoryTextures, PrevHistory)
+
+		SHADER_PARAMETER_STRUCT(FTSRHistoryUAVs, HistoryOutput)
+		SHADER_PARAMETER_RDG_TEXTURE_UAV(RWTexture2D, InterferenceWeightOutput)
+		SHADER_PARAMETER_RDG_TEXTURE_UAV(RWTexture2D, DebugOutput)
+	END_SHADER_PARAMETER_STRUCT()
+}; // class FTSRDetectInterferenceCS
+
+class FTSRFilterFrequenciesCS : public FTSRShader
+{
+	DECLARE_GLOBAL_SHADER(FTSRFilterFrequenciesCS);
+	SHADER_USE_PARAMETER_STRUCT(FTSRFilterFrequenciesCS, FTSRShader);
+
+	BEGIN_SHADER_PARAMETER_STRUCT(FParameters, )
+		SHADER_PARAMETER_STRUCT_INCLUDE(FTSRCommonParameters, CommonParameters)
+		SHADER_PARAMETER(FVector3f, OutputQuantizationError)
+
+		SHADER_PARAMETER_RDG_TEXTURE(Texture2D, InputTexture)
+		SHADER_PARAMETER_RDG_TEXTURE(Texture2D, PredictionSceneColorTexture)
+		SHADER_PARAMETER_RDG_TEXTURE(Texture2D, ParallaxRejectionMaskTexture)
+		SHADER_PARAMETER_RDG_TEXTURE(Texture2D, InterferenceWeightTexture)
+
+		SHADER_PARAMETER_RDG_TEXTURE_UAV(RWTexture2D, FilteredInputOutput)
+		SHADER_PARAMETER_RDG_TEXTURE_UAV(RWTexture2D, FilteredPredictionSceneColorOutput)
+		SHADER_PARAMETER_RDG_TEXTURE_UAV(RWTexture2D, DebugOutput)
+	END_SHADER_PARAMETER_STRUCT()
+}; // class FTSRFilterFrequenciesCS
+
+class FTSRCompareHistoryCS : public FTSRShader
+{
+	DECLARE_GLOBAL_SHADER(FTSRCompareHistoryCS);
+	SHADER_USE_PARAMETER_STRUCT(FTSRCompareHistoryCS, FTSRShader);
+
+	class FOutputHalfRes : SHADER_PERMUTATION_BOOL("DIM_OUTPUT_HALF_RES");
+	using FPermutationDomain = TShaderPermutationDomain<FOutputHalfRes>;
+
+	BEGIN_SHADER_PARAMETER_STRUCT(FParameters, )
+		SHADER_PARAMETER_STRUCT_INCLUDE(FTSRCommonParameters, CommonParameters)
+		SHADER_PARAMETER_RDG_TEXTURE(Texture2D, ParallaxRejectionMaskTexture)
+		SHADER_PARAMETER_RDG_TEXTURE(Texture2D, FilteredInputTexture)
+		SHADER_PARAMETER_RDG_TEXTURE(Texture2D, FilteredPredictionSceneColorTexture)
+		SHADER_PARAMETER_RDG_TEXTURE(Texture2D, InterferenceWeightTexture)
+		SHADER_PARAMETER_RDG_TEXTURE_UAV(RWTexture2D, HistoryRejectionOutput)
+		SHADER_PARAMETER_RDG_TEXTURE_UAV(RWTexture2D, DebugOutput)
+	END_SHADER_PARAMETER_STRUCT()
+}; // class FTSRCompareHistoryCS
+
+class FTSRPostfilterRejectionCS : public FTSRShader
+{
+	DECLARE_GLOBAL_SHADER(FTSRPostfilterRejectionCS);
+	SHADER_USE_PARAMETER_STRUCT(FTSRPostfilterRejectionCS, FTSRShader);
+
+	class FOutputHalfRes : SHADER_PERMUTATION_BOOL("DIM_OUTPUT_HALF_RES");
+	using FPermutationDomain = TShaderPermutationDomain<FOutputHalfRes>;
+
+	BEGIN_SHADER_PARAMETER_STRUCT(FParameters, )
+		SHADER_PARAMETER(FIntRect, HistoryRejectionViewport)
+		SHADER_PARAMETER_RDG_TEXTURE(Texture2D, HistoryRejectionTexture)
+		SHADER_PARAMETER_RDG_TEXTURE_UAV(RWTexture2D, HistoryRejectionOutput)
+		SHADER_PARAMETER_RDG_TEXTURE_UAV(RWTexture2D, DebugOutput)
+	END_SHADER_PARAMETER_STRUCT()
+}; // class FTSRPostfilterRejectionCS
+
+class FTSRDilateRejectionCS : public FTSRShader
+{
+	DECLARE_GLOBAL_SHADER(FTSRDilateRejectionCS);
+	SHADER_USE_PARAMETER_STRUCT(FTSRDilateRejectionCS, FTSRShader);
+
+	BEGIN_SHADER_PARAMETER_STRUCT(FParameters, )
+		SHADER_PARAMETER_STRUCT_INCLUDE(FTSRCommonParameters, CommonParameters)
+		SHADER_PARAMETER_RDG_TEXTURE(Texture2D, HistoryRejectionTexture)
+		SHADER_PARAMETER_RDG_TEXTURE_UAV(RWTexture2D, DilatedHistoryRejectionOutput)
+		SHADER_PARAMETER_RDG_TEXTURE_UAV(RWTexture2D, DebugOutput)
+	END_SHADER_PARAMETER_STRUCT()
+}; // class FTSRDilateRejectionCS
+
+class FTSRUpdateSuperResHistoryCS : public FTSRShader
+{
+	DECLARE_GLOBAL_SHADER(FTSRUpdateSuperResHistoryCS);
+	SHADER_USE_PARAMETER_STRUCT(FTSRUpdateSuperResHistoryCS, FTSRShader);
+
+	BEGIN_SHADER_PARAMETER_STRUCT(FParameters, )
+		SHADER_PARAMETER_STRUCT_INCLUDE(FTSRCommonParameters, CommonParameters)
+		SHADER_PARAMETER_RDG_TEXTURE(Texture2D, InputSceneColorTexture)
+
+		SHADER_PARAMETER_RDG_TEXTURE(Texture2D, DilatedVelocityTexture)
+		SHADER_PARAMETER_RDG_TEXTURE(Texture2D, ParallaxRejectionMaskTexture)
+
+		SHADER_PARAMETER_STRUCT_INCLUDE(FTSRPrevHistoryParameters, PrevHistoryParameters)
+		SHADER_PARAMETER_STRUCT(FTSRHistoryTextures, PrevHistory)
+
+		SHADER_PARAMETER_STRUCT(FTSRHistoryUAVs, HistoryOutput)
+		SHADER_PARAMETER_RDG_TEXTURE_UAV(RWTexture2D, DebugOutput)
+	END_SHADER_PARAMETER_STRUCT()
+}; // class FTSRUpdateSuperResHistoryCS
+
+class FTSRUpdateHistoryCS : public FTSRShader
+{
+	DECLARE_GLOBAL_SHADER(FTSRUpdateHistoryCS);
+	SHADER_USE_PARAMETER_STRUCT(FTSRUpdateHistoryCS, FTSRShader);
+
+	BEGIN_SHADER_PARAMETER_STRUCT(FParameters, )
+		SHADER_PARAMETER_STRUCT_INCLUDE(FTSRCommonParameters, CommonParameters)
+		SHADER_PARAMETER_RDG_TEXTURE(Texture2D, InputSceneColorTexture)
+		SHADER_PARAMETER_RDG_TEXTURE_SRV(Texture2D, InputSceneStencilTexture)
+		SHADER_PARAMETER_RDG_TEXTURE(Texture2D, HistoryRejectionTexture)
+		SHADER_PARAMETER_RDG_TEXTURE(Texture2D, TranslucencyRejectionTexture)
+
+		SHADER_PARAMETER_RDG_TEXTURE(Texture2D, DilatedVelocityTexture)
+		SHADER_PARAMETER_RDG_TEXTURE(Texture2D, ParallaxFactorTexture)
+		SHADER_PARAMETER_RDG_TEXTURE(Texture2D, ParallaxRejectionMaskTexture)
+
+		SHADER_PARAMETER(FScreenTransform, HistoryPixelPosToScreenPos)
+		SHADER_PARAMETER(FScreenTransform, HistoryPixelPosToPPCo)
+		SHADER_PARAMETER(FVector3f, HistoryQuantizationError)
+		SHADER_PARAMETER(float, MinTranslucencyRejection)
+		SHADER_PARAMETER(int32, ResponsiveStencilMask)
+
+		SHADER_PARAMETER_STRUCT_INCLUDE(FTSRPrevHistoryParameters, PrevHistoryParameters)
+		SHADER_PARAMETER_STRUCT(FTSRHistoryTextures, PrevHistory)
+
+		SHADER_PARAMETER_STRUCT(FTSRHistoryUAVs, HistoryOutput)
+		SHADER_PARAMETER_RDG_TEXTURE_UAV(RWTexture2D, SceneColorOutput)
+		SHADER_PARAMETER_RDG_TEXTURE_UAV(RWTexture2D, DebugOutput)
+	END_SHADER_PARAMETER_STRUCT()
+
+	static void ModifyCompilationEnvironment(const FGlobalShaderPermutationParameters& Parameters, FShaderCompilerEnvironment& OutEnvironment)
+	{
+		FTSRShader::ModifyCompilationEnvironment(Parameters, OutEnvironment);
+		OutEnvironment.CompilerFlags.Add(CFLAG_Wave32);
+	}
+}; // class FTSRUpdateHistoryCS
+
+#if COMPILE_TSR_DEBUG_PASSES
+
+class FTSRDebugHistoryCS : public FTSRShader
+{
+	DECLARE_GLOBAL_SHADER(FTSRDebugHistoryCS);
+	SHADER_USE_PARAMETER_STRUCT(FTSRDebugHistoryCS, FTSRShader);
+
+	BEGIN_SHADER_PARAMETER_STRUCT(FParameters, )
+		SHADER_PARAMETER_STRUCT_INCLUDE(FTSRCommonParameters, CommonParameters)
+		SHADER_PARAMETER_STRUCT(FTSRHistoryTextures, History)
+		SHADER_PARAMETER_STRUCT(FTSRHistoryTextures, PrevHistory)
+		SHADER_PARAMETER_RDG_TEXTURE_UAV(RWTexture2D, DebugOutput)
+		END_SHADER_PARAMETER_STRUCT()
+}; // class FTSRDebugHistoryCS
+
+#endif
+
+IMPLEMENT_GLOBAL_SHADER(FTSRClearPrevTexturesCS,     "/Engine/Private/TemporalSuperResolution/TSRClearPrevTextures.usf",     "MainCS", SF_Compute);
+IMPLEMENT_GLOBAL_SHADER(FTSRDilateVelocityCS,        "/Engine/Private/TemporalSuperResolution/TSRDilateVelocity.usf",        "MainCS", SF_Compute);
+IMPLEMENT_GLOBAL_SHADER(FTSRDecimateHistoryCS,       "/Engine/Private/TemporalSuperResolution/TSRDecimateHistory.usf",       "MainCS", SF_Compute);
+IMPLEMENT_GLOBAL_SHADER(FTSRCompareTranslucencyCS,   "/Engine/Private/TemporalSuperResolution/TSRCompareTranslucency.usf",   "MainCS", SF_Compute);
+IMPLEMENT_GLOBAL_SHADER(FTSRDetectInterferenceCS,    "/Engine/Private/TemporalSuperResolution/TSRDetectInterference.usf",    "MainCS", SF_Compute);
+IMPLEMENT_GLOBAL_SHADER(FTSRFilterFrequenciesCS,     "/Engine/Private/TemporalSuperResolution/TSRFilterFrequencies.usf",     "MainCS", SF_Compute);
+IMPLEMENT_GLOBAL_SHADER(FTSRCompareHistoryCS,        "/Engine/Private/TemporalSuperResolution/TSRCompareHistory.usf",        "MainCS", SF_Compute);
+IMPLEMENT_GLOBAL_SHADER(FTSRPostfilterRejectionCS,   "/Engine/Private/TemporalSuperResolution/TSRPostfilterRejection.usf",   "MainCS", SF_Compute);
+IMPLEMENT_GLOBAL_SHADER(FTSRDilateRejectionCS,       "/Engine/Private/TemporalSuperResolution/TSRDilateRejection.usf",       "MainCS", SF_Compute);
+IMPLEMENT_GLOBAL_SHADER(FTSRUpdateSuperResHistoryCS, "/Engine/Private/TemporalSuperResolution/TSRUpdateSuperResHistory.usf", "MainCS", SF_Compute);
+IMPLEMENT_GLOBAL_SHADER(FTSRUpdateHistoryCS,         "/Engine/Private/TemporalSuperResolution/TSRUpdateHistory.usf",         "MainCS", SF_Compute);
+
+#if COMPILE_TSR_DEBUG_PASSES
+IMPLEMENT_GLOBAL_SHADER(FTSRDebugHistoryCS,          "/Engine/Private/TemporalSuperResolution/TSRDebugHistory.usf",          "MainCS", SF_Compute);
+#endif
+
+DECLARE_GPU_STAT(TemporalSuperResolution)
+
+} //! namespace
+
+FVector ComputePixelFormatQuantizationError(EPixelFormat PixelFormat);
+
+void AddTemporalSuperResolutionPasses(
+	FRDGBuilder& GraphBuilder,
+	const FViewInfo& View,
+	const ITemporalUpscaler::FPassInputs& PassInputs,
+	FRDGTextureRef* OutSceneColorTexture,
+	FIntRect* OutSceneColorViewRect)
+{
+	const FTemporalAAHistory& InputHistory = View.PrevViewInfo.TemporalAAHistory;
+
+#if COMPILE_TSR_DEBUG_PASSES
+	const bool bSetupDebugPasses = CVarTSRSetupDebugPasses.GetValueOnRenderThread() != 0;
+#endif
+
+	// Whether to use camera cut shader permutation or not.
+	bool bCameraCut = !InputHistory.IsValid() || View.bCameraCut;
+
+	bool bHalfResLowFrequency = CVarTSRHalfResShadingRejection.GetValueOnRenderThread() != 0;
+
+	bool bEnableInterferenceHeuristic = CVarTSREnableAntiInterference.GetValueOnRenderThread() != 0;
+
+	bool bRejectSeparateTranslucency = PassInputs.SeparateTranslucencyTextures != nullptr && CVarTSRRejectTranslucency.GetValueOnRenderThread() != 0;
+
+	enum class ERejectionPostFilter : uint8
+	{
+		Disabled,
+		PostRejectionDownsample,
+		PreRejectionDownsample,
+	};
+
+	ERejectionPostFilter PostFilter = ERejectionPostFilter(FMath::Clamp(CVarTSRFilterShadingRejection.GetValueOnRenderThread(), 0, 2));
+
+	FIntPoint InputExtent = PassInputs.SceneColorTexture->Desc.Extent;
+	FIntRect InputRect = View.ViewRect;
+
+	FIntPoint LowFrequencyExtent = InputExtent;
+	FIntRect LowFrequencyRect = FIntRect(FIntPoint(0, 0), InputRect.Size());
+
+	if (bHalfResLowFrequency)
+	{
+		LowFrequencyExtent = InputExtent / 2;
+		LowFrequencyRect = FIntRect(FIntPoint(0, 0), FIntPoint::DivideAndRoundUp(InputRect.Size(), 2));
+	}
+
+	FIntPoint RejectionExtent = LowFrequencyExtent / 2;
+	FIntRect RejectionRect = FIntRect(FIntPoint(0, 0), FIntPoint::DivideAndRoundUp(LowFrequencyRect.Size(), 2));
+
+	FIntPoint OutputExtent;
+	FIntRect OutputRect;
+	if (View.PrimaryScreenPercentageMethod == EPrimaryScreenPercentageMethod::TemporalUpscale)
+	{
+		OutputRect.Min = FIntPoint(0, 0);
+		OutputRect.Max = View.GetSecondaryViewRectSize();
+
+		FIntPoint QuantizedPrimaryUpscaleViewSize;
+		QuantizeSceneBufferSize(OutputRect.Max, QuantizedPrimaryUpscaleViewSize);
+
+		OutputExtent = FIntPoint(
+			FMath::Max(InputExtent.X, QuantizedPrimaryUpscaleViewSize.X),
+			FMath::Max(InputExtent.Y, QuantizedPrimaryUpscaleViewSize.Y));
+	}
+	else
+	{
+		OutputRect.Min = FIntPoint(0, 0);
+		OutputRect.Max = View.ViewRect.Size();
+		OutputExtent = InputExtent;
+	}
+
+	FIntPoint HistoryExtent;
+	FIntPoint HistorySize;
+	{
+		float UpscaleFactor = FMath::Clamp(CVarTSRHistorySP.GetValueOnRenderThread() / 100.0f, 1.0f, 2.0f);
+
+		HistorySize = FIntPoint(
+			FMath::CeilToInt(OutputRect.Width() * UpscaleFactor),
+			FMath::CeilToInt(OutputRect.Height() * UpscaleFactor));
+
+		FIntPoint QuantizedHistoryViewSize;
+		QuantizeSceneBufferSize(HistorySize, QuantizedHistoryViewSize);
+
+		HistoryExtent = FIntPoint(
+			FMath::Max(InputExtent.X, QuantizedHistoryViewSize.X),
+			FMath::Max(InputExtent.Y, QuantizedHistoryViewSize.Y));
+	}
+
+	RDG_EVENT_SCOPE(GraphBuilder, "TemporalSuperResolution %dx%d -> %dx%d", InputRect.Width(), InputRect.Height(), OutputRect.Width(), OutputRect.Height());
+	RDG_GPU_STAT_SCOPE(GraphBuilder, TemporalSuperResolution);
+
+	FRDGTextureRef BlackUintDummy = GSystemTextures.GetZeroUIntDummy(GraphBuilder);
+	FRDGTextureRef BlackDummy = GraphBuilder.RegisterExternalTexture(GSystemTextures.BlackDummy);
+	FRDGTextureRef WhiteDummy = GraphBuilder.RegisterExternalTexture(GSystemTextures.WhiteDummy);
+
+	FRDGTextureRef SeparateTranslucencyTexture = bRejectSeparateTranslucency ?
+		PassInputs.SeparateTranslucencyTextures->GetColorForRead(GraphBuilder) : nullptr;
+
+	FTSRCommonParameters CommonParameters;
+	{
+		CommonParameters.InputInfo = GetScreenPassTextureViewportParameters(FScreenPassTextureViewport(
+			InputExtent, InputRect));
+		CommonParameters.InputPixelPosMin = CommonParameters.InputInfo.ViewportMin;
+		CommonParameters.InputPixelPosMax = CommonParameters.InputInfo.ViewportMax - 1;
+
+		CommonParameters.LowFrequencyInfo = GetScreenPassTextureViewportParameters(FScreenPassTextureViewport(
+			LowFrequencyExtent, LowFrequencyRect));
+
+		CommonParameters.RejectionInfo = GetScreenPassTextureViewportParameters(FScreenPassTextureViewport(
+			RejectionExtent, RejectionRect));
+
+		CommonParameters.HistoryInfo = GetScreenPassTextureViewportParameters(FScreenPassTextureViewport(
+			HistoryExtent, FIntRect(FIntPoint(0, 0), HistorySize)));
+
+		CommonParameters.InputJitter = View.TemporalJitterPixels;
+		CommonParameters.bCameraCut = bCameraCut;
+		CommonParameters.bEnableInterferenceHeuristic = bEnableInterferenceHeuristic;
+		CommonParameters.ViewUniformBuffer = View.ViewUniformBuffer;
+	}
+
+	auto CreateDebugUAV = [&](const FIntPoint& Extent, const TCHAR* DebugName)
+	{
+		FRDGTextureDesc DebugDesc = FRDGTextureDesc::Create2D(
+			Extent,
+			PF_FloatRGBA,
+			FClearValueBinding::None,
+			/* InFlags = */ TexCreate_ShaderResource | TexCreate_UAV);
+
+		FRDGTextureRef DebugTexture = GraphBuilder.CreateTexture(DebugDesc, DebugName);
+
+		return GraphBuilder.CreateUAV(DebugTexture);
+	};
+
+	// Clear atomic scattered texture.
+	FRDGTextureRef PrevUseCountTexture;
+	FRDGTextureRef PrevClosestDepthTexture;
+	{
+		{
+			FRDGTextureDesc Desc = FRDGTextureDesc::Create2D(
+				InputExtent,
+				PF_R32_UINT,
+				FClearValueBinding::None,
+				/* InFlags = */ TexCreate_ShaderResource | TexCreate_UAV);
+
+			PrevUseCountTexture = GraphBuilder.CreateTexture(Desc, TEXT("TSR.PrevUseCountTexture"));
+			PrevClosestDepthTexture = GraphBuilder.CreateTexture(Desc, TEXT("TSR.PrevClosestDepthTexture"));
+		}
+
+		FTSRClearPrevTexturesCS::FParameters* PassParameters = GraphBuilder.AllocParameters<FTSRClearPrevTexturesCS::FParameters>();
+		PassParameters->CommonParameters = CommonParameters;
+		PassParameters->PrevUseCountOutput = GraphBuilder.CreateUAV(PrevUseCountTexture);
+		PassParameters->PrevClosestDepthOutput = GraphBuilder.CreateUAV(PrevClosestDepthTexture);
+
+		TShaderMapRef<FTSRClearPrevTexturesCS> ComputeShader(View.ShaderMap);
+		FComputeShaderUtils::AddPass(
+			GraphBuilder,
+			RDG_EVENT_NAME("TSR ClearPrevTextures %dx%d", InputRect.Width(), InputRect.Height()),
+			ComputeShader,
+			PassParameters,
+			FComputeShaderUtils::GetGroupCount(InputRect.Size(), 8));
+	}
+
+	// Dilate the velocity texture & scatter reprojection into previous frame
+	FRDGTextureRef DilatedVelocityTexture;
+	FRDGTextureRef ClosestDepthTexture;
+	FRDGTextureRef ParallaxFactorTexture;
+	{
+		{
+			FRDGTextureDesc Desc = FRDGTextureDesc::Create2D(
+				InputExtent,
+				PF_G16R16,
+				FClearValueBinding::None,
+				/* InFlags = */ TexCreate_ShaderResource | TexCreate_UAV);
+
+			DilatedVelocityTexture = GraphBuilder.CreateTexture(Desc, TEXT("TSR.DilatedVelocity"));
+
+			Desc.Format = PF_R16F;
+			ClosestDepthTexture = GraphBuilder.CreateTexture(Desc, TEXT("TSR.ClosestDepthTexture"));
+
+			Desc.Format = PF_R8_UINT;
+			ParallaxFactorTexture = GraphBuilder.CreateTexture(Desc, TEXT("TSR.ParallaxFactor"));
+		}
+
+		FTSRDilateVelocityCS::FParameters* PassParameters = GraphBuilder.AllocParameters<FTSRDilateVelocityCS::FParameters>();
+		PassParameters->CommonParameters = CommonParameters;
+		PassParameters->PrevOutputBufferUVMin = CommonParameters.InputInfo.UVViewportBilinearMin - CommonParameters.InputInfo.ExtentInverse;
+		PassParameters->PrevOutputBufferUVMax = CommonParameters.InputInfo.UVViewportBilinearMax + CommonParameters.InputInfo.ExtentInverse;
+		PassParameters->SceneDepthTexture = PassInputs.SceneDepthTexture;
+		PassParameters->SceneVelocityTexture = PassInputs.SceneVelocityTexture;
+		PassParameters->DilatedVelocityOutput = GraphBuilder.CreateUAV(DilatedVelocityTexture);
+		PassParameters->ClosestDepthOutput = GraphBuilder.CreateUAV(ClosestDepthTexture);
+		PassParameters->PrevUseCountOutput = GraphBuilder.CreateUAV(PrevUseCountTexture);
+		PassParameters->PrevClosestDepthOutput = GraphBuilder.CreateUAV(PrevClosestDepthTexture);
+		PassParameters->ParallaxFactorOutput = GraphBuilder.CreateUAV(ParallaxFactorTexture);
+		PassParameters->DebugOutput = CreateDebugUAV(InputExtent, TEXT("Debug.TSR.DilateVelocity"));
+
+		TShaderMapRef<FTSRDilateVelocityCS> ComputeShader(View.ShaderMap);
+		FComputeShaderUtils::AddPass(
+			GraphBuilder,
+			RDG_EVENT_NAME("TSR DilateVelocity %dx%d", InputRect.Width(), InputRect.Height()),
+			ComputeShader,
+			PassParameters,
+			FComputeShaderUtils::GetGroupCount(InputRect.Size(), 8));
+	}
+
+	// Setup the previous frame history
+	FTSRPrevHistoryParameters PrevHistoryParameters;
+	FTSRHistoryTextures PrevHistory;
+	{
+		FScreenPassTextureViewport PrevHistoryViewport(InputHistory.ReferenceBufferSize, InputHistory.ViewportRect);
+		if (bCameraCut)
+		{
+			PrevHistoryViewport.Extent = FIntPoint(1, 1);
+			PrevHistoryViewport.Rect = FIntRect(FIntPoint(0, 0), FIntPoint(1, 1));
+		}
+
+		PrevHistoryParameters.PrevHistoryInfo = GetScreenPassTextureViewportParameters(PrevHistoryViewport);
+		PrevHistoryParameters.ScreenPosToPrevHistoryBufferUV = FScreenTransform::ChangeTextureBasisFromTo(
+			PrevHistoryViewport, FScreenTransform::ETextureBasis::ScreenPosition, FScreenTransform::ETextureBasis::TextureUV);
+		PrevHistoryParameters.HistoryPreExposureCorrection = View.PreExposure / View.PrevViewInfo.SceneColorPreExposure;
+
+		for (int32 i = 0; i < PrevHistory.LowResTextures.Num(); i++)
+		{
+			if (InputHistory.LowResRT[i].IsValid() && !bCameraCut)
+			{
+				PrevHistory.LowResTextures[i] = GraphBuilder.RegisterExternalTexture(InputHistory.LowResRT[i]);
+			}
+			else
+			{
+				PrevHistory.LowResTextures[i] = BlackDummy;
+			}
+		}
+
+		for (int32 i = 0; i < PrevHistory.Textures.Num(); i++)
+		{
+			if (InputHistory.RT[i].IsValid() && !bCameraCut)
+			{
+				PrevHistory.Textures[i] = GraphBuilder.RegisterExternalTexture(InputHistory.RT[i]);
+			}
+			else if (i == 3)
+			{
+				PrevHistory.Textures[i] = BlackUintDummy;
+			}
+			else
+			{
+				PrevHistory.Textures[i] = BlackDummy;
+			}
+		}
+
+		for (int32 i = 0; i < PrevHistory.SuperResTextures.Num(); i++)
+		{
+			if (InputHistory.SuperResRT[i].IsValid() && !bCameraCut)
+			{
+				PrevHistory.SuperResTextures[i] = GraphBuilder.RegisterExternalTexture(InputHistory.SuperResRT[i]);
+			}
+			else
+			{
+				PrevHistory.SuperResTextures[i] = BlackUintDummy;
+			}
+		}
+	}
+
+
+	FTSRHistoryTextures History;
+	{
+		FRDGTextureDesc LowResDesc = FRDGTextureDesc::Create2D(
+			InputExtent,
+			PF_R8,
+			FClearValueBinding::None,
+			TexCreate_ShaderResource | TexCreate_UAV);
+
+		History.LowResTextures[0] = GraphBuilder.CreateTexture(LowResDesc, TEXT("TSR.History.LowResMetadata[0]"));
+	}
+
+	{
+		FRDGTextureDesc Desc = FRDGTextureDesc::Create2D(
+			HistoryExtent,
+			(CVarTSRR11G11B10History.GetValueOnRenderThread() != 0) ? PF_FloatR11G11B10 : PF_FloatRGBA,
+			FClearValueBinding::None,
+			TexCreate_ShaderResource | TexCreate_UAV);
+
+		History.Textures[0] = GraphBuilder.CreateTexture(Desc, TEXT("TSR.History.LowFrequencies"));
+		History.Textures[1] = GraphBuilder.CreateTexture(Desc, TEXT("TSR.History.HighFrequencies"));
+
+		Desc.Format = PF_R8G8;
+		History.Textures[2] = GraphBuilder.CreateTexture(Desc, TEXT("TSR.History.Metadata"));
+
+		Desc.Format = PF_R16_UINT;
+		History.Textures[3] = GraphBuilder.CreateTexture(Desc, TEXT("TSR.History.SubpixelInfo"));
+	}
+
+	{
+		FRDGTextureDesc SuperResDesc = FRDGTextureDesc::Create2D(
+			HistoryExtent * 2,
+			PF_R16_UINT,
+			FClearValueBinding::None,
+			TexCreate_ShaderResource | TexCreate_UAV);
+
+		History.SuperResTextures[0] = GraphBuilder.CreateTexture(SuperResDesc, TEXT("TSR.History.SuperResMetadata[0]"));
+
+		SuperResDesc.Format = PF_R8_UINT;
+		History.SuperResTextures[1] = GraphBuilder.CreateTexture(SuperResDesc, TEXT("TSR.History.SuperResMetadata[1]"));
+	}
+
+	// Decimate input to flicker at same frequency as input.
+	FRDGTextureRef HalfResInputSceneColorTexture = nullptr;
+	FRDGTextureRef HalfResPredictionSceneColorTexture = nullptr;
+	FRDGTextureRef HalfResParallaxRejectionMaskTexture = nullptr;
+	FRDGTextureRef PredictionSceneColorTexture = nullptr;
+	FRDGTextureRef ParallaxRejectionMaskTexture = nullptr;
+	FRDGTextureRef InterferenceSeedTexture = nullptr;
+	{
+		{
+			FRDGTextureDesc Desc = FRDGTextureDesc::Create2D(
+				InputExtent,
+				PF_R8,
+				FClearValueBinding::None,
+				/* InFlags = */ TexCreate_ShaderResource | TexCreate_UAV);
+
+			ParallaxRejectionMaskTexture = GraphBuilder.CreateTexture(Desc, TEXT("TSR.ParallaxRejectionMask"));
+
+			// TODO(TSR): can compress to the history seed's 4bit per pixel
+			Desc.Format = PF_R8G8B8A8;
+			InterferenceSeedTexture = GraphBuilder.CreateTexture(Desc, TEXT("TSR.Interference.Seed"));
+		}
+
+		FTSRDecimateHistoryCS::FParameters* PassParameters = GraphBuilder.AllocParameters<FTSRDecimateHistoryCS::FParameters>();
+		PassParameters->CommonParameters = CommonParameters;
+		{
+			const FViewMatrices& ViewMatrices = View.ViewMatrices;
+			const FViewMatrices& PrevViewMatrices = View.PrevViewInfo.ViewMatrices;
+
+			FMatrix RotationalInvViewProj = ViewMatrices.ComputeInvProjectionNoAAMatrix() * (ViewMatrices.GetTranslatedViewMatrix().RemoveTranslation().GetTransposed());
+			FMatrix RotationalPrevViewProj = (PrevViewMatrices.GetTranslatedViewMatrix().RemoveTranslation()) * PrevViewMatrices.ComputeProjectionNoAAMatrix();
+
+			PassParameters->RotationalClipToPrevClip = RotationalInvViewProj * RotationalPrevViewProj;
+		}
+		PassParameters->OutputQuantizationError = ComputePixelFormatQuantizationError(PF_FloatR11G11B10);
+		{
+			float TanHalfFieldOfView = View.ViewMatrices.GetInvProjectionMatrix().M[0][0];
+
+			// Should be multiplied 0.5* for the diameter to radius, and by 2.0 because GetTanHalfFieldOfView() cover only half of the pixels.
+			PassParameters->WorldDepthToPixelWorldRadius = TanHalfFieldOfView / float(View.ViewRect.Width());
+		}
+
+		PassParameters->InputSceneColorTexture = PassInputs.SceneColorTexture;
+		PassParameters->DilatedVelocityTexture = DilatedVelocityTexture;
+		PassParameters->ClosestDepthTexture = ClosestDepthTexture;
+		PassParameters->PrevUseCountTexture = PrevUseCountTexture;
+		PassParameters->PrevClosestDepthTexture = PrevClosestDepthTexture;
+		PassParameters->ParallaxFactorTexture = ParallaxFactorTexture;
+
+		PassParameters->PrevHistoryParameters = PrevHistoryParameters;
+		PassParameters->PrevHistory = PrevHistory;
+
+		if (bHalfResLowFrequency)
+		{
+			FRDGTextureDesc Desc = FRDGTextureDesc::Create2D(
+				LowFrequencyExtent,
+				PF_FloatR11G11B10,
+				FClearValueBinding::None,
+				/* InFlags = */ TexCreate_ShaderResource | TexCreate_UAV);
+
+			HalfResInputSceneColorTexture = GraphBuilder.CreateTexture(Desc, TEXT("TSR.HalfResInput"));
+			HalfResPredictionSceneColorTexture = GraphBuilder.CreateTexture(Desc, TEXT("TSR.Prediction.SceneColor"));
+
+			Desc.Format = PF_R8;
+			HalfResParallaxRejectionMaskTexture = GraphBuilder.CreateTexture(Desc, TEXT("TSR.HalfResParallaxRejectionMask"));
+
+			PassParameters->HalfResSceneColorOutput = GraphBuilder.CreateUAV(HalfResInputSceneColorTexture);
+			PassParameters->HalfResPredictionSceneColorOutput = GraphBuilder.CreateUAV(HalfResPredictionSceneColorTexture);
+			PassParameters->HalfResParallaxRejectionMaskOutput = GraphBuilder.CreateUAV(HalfResParallaxRejectionMaskTexture);
+		}
+		else
+		{
+			FRDGTextureDesc Desc = FRDGTextureDesc::Create2D(
+				InputExtent,
+				PF_FloatR11G11B10,
+				FClearValueBinding::None,
+				/* InFlags = */ TexCreate_ShaderResource | TexCreate_UAV);
+
+			PredictionSceneColorTexture = GraphBuilder.CreateTexture(Desc, TEXT("TSR.Prediction.SceneColor"));
+
+			PassParameters->PredictionSceneColorOutput = GraphBuilder.CreateUAV(PredictionSceneColorTexture);
+		}
+
+		PassParameters->ParallaxRejectionMaskOutput = GraphBuilder.CreateUAV(ParallaxRejectionMaskTexture);
+		PassParameters->InterferenceSeedOutput = GraphBuilder.CreateUAV(InterferenceSeedTexture);
+		PassParameters->DebugOutput = CreateDebugUAV(LowFrequencyExtent, TEXT("Debug.TSR.DecimateHistory"));
+
+		FTSRDecimateHistoryCS::FPermutationDomain PermutationVector;
+		PermutationVector.Set<FTSRDecimateHistoryCS::FOutputHalfRes>(bHalfResLowFrequency);
+
+		TShaderMapRef<FTSRDecimateHistoryCS> ComputeShader(View.ShaderMap, PermutationVector);
+		FComputeShaderUtils::AddPass(
+			GraphBuilder,
+			RDG_EVENT_NAME("TSR DecimateHistory(%s) %dx%d",
+				bHalfResLowFrequency ? TEXT("HalfResShadingOutput") : TEXT(""),
+				InputRect.Width(), InputRect.Height()),
+			ComputeShader,
+			PassParameters,
+			FComputeShaderUtils::GetGroupCount(InputRect.Size(), 8));
+	}
+
+	FRDGTextureRef TranslucencyRejectionTexture = nullptr;
+	if (bRejectSeparateTranslucency && View.PrevViewInfo.SeparateTranslucency != nullptr)
+	{
+		FRDGTextureRef PrevTranslucencyTexture;
+		FScreenPassTextureViewport PrevTranslucencyViewport;
+
+		if (View.PrevViewInfo.SeparateTranslucency)
+		{
+
+			PrevTranslucencyTexture = GraphBuilder.RegisterExternalTexture(View.PrevViewInfo.SeparateTranslucency);
+			PrevTranslucencyViewport = FScreenPassTextureViewport(PrevTranslucencyTexture->Desc.Extent, View.PrevViewInfo.ViewRect);
+
+		}
+		else
+		{
+			PrevTranslucencyTexture = BlackDummy;
+			PrevTranslucencyViewport = FScreenPassTextureViewport(FIntPoint(1, 1), FIntRect(FIntPoint(0, 0), FIntPoint(1, 1)));
+		}
+
+		{
+			FRDGTextureDesc Desc = FRDGTextureDesc::Create2D(
+				InputExtent,
+				PF_R8,
+				FClearValueBinding::None,
+				/* InFlags = */ TexCreate_ShaderResource | TexCreate_UAV);
+
+			TranslucencyRejectionTexture = GraphBuilder.CreateTexture(Desc, TEXT("TSR.TranslucencyRejection"));
+		}
+
+		FTSRCompareTranslucencyCS::FParameters* PassParameters = GraphBuilder.AllocParameters<FTSRCompareTranslucencyCS::FParameters>();
+		PassParameters->CommonParameters = CommonParameters;
+
+		PassParameters->TranslucencyInfo = GetScreenPassTextureViewportParameters(FScreenPassTextureViewport(
+			InputExtent, InputRect));
+		PassParameters->PrevTranslucencyInfo = GetScreenPassTextureViewportParameters(PrevTranslucencyViewport);
+		PassParameters->PrevTranslucencyPreExposureCorrection = PrevHistoryParameters.HistoryPreExposureCorrection;
+
+		PassParameters->InputPixelPosToScreenPos = (FScreenTransform::Identity + 0.5) * CommonParameters.InputInfo.ViewportSizeInverse * FScreenTransform::ViewportUVToScreenPos;
+		PassParameters->ScreenPosToPrevTranslucencyTextureUV = FScreenTransform::ChangeTextureBasisFromTo(
+			PrevTranslucencyViewport, FScreenTransform::ETextureBasis::ScreenPosition, FScreenTransform::ETextureBasis::TextureUV);
+
+		PassParameters->DilatedVelocityTexture = DilatedVelocityTexture;
+		PassParameters->TranslucencyTexture = SeparateTranslucencyTexture;
+		PassParameters->PrevTranslucencyTexture = PrevTranslucencyTexture;
+
+		PassParameters->TranslucencyRejectionOutput = GraphBuilder.CreateUAV(TranslucencyRejectionTexture);
+		PassParameters->DebugOutput = CreateDebugUAV(LowFrequencyExtent, TEXT("Debug.TSR.ComparetTranslucency"));
+
+		TShaderMapRef<FTSRCompareTranslucencyCS> ComputeShader(View.ShaderMap);
+		FComputeShaderUtils::AddPass(
+			GraphBuilder,
+			RDG_EVENT_NAME("TSR ComparetTranslucency %dx%d", InputRect.Width(), InputRect.Height()),
+			ComputeShader,
+			PassParameters,
+			FComputeShaderUtils::GetGroupCount(InputRect.Size(), 8));
+	}
+
+	// Detect interference between geometry and alignement of input texel centers. It is not about awnser whether an
+	// interference has happened in the past, because interference change based on input resolution or camera position.
+	// So to remain stable on camera movement and input resolution change, it is about awnsering the question on whether
+	// an interference is possible.
+	// TODO(TSR): Could sample the interference seed in the DilateVelocity and detect interference in the decimate.
+	TStaticArray<bool, FTemporalAAHistory::kLowResRenderTargetCount> ExtractLowResHistoryTexture;
+	FRDGTextureRef InterferenceWeightTexture;
+	if (bEnableInterferenceHeuristic)
+	{
+		{
+			// TODO(TSR): Compress to 1bit per pixel
+			FRDGTextureDesc Desc = FRDGTextureDesc::Create2D(
+				InputExtent,
+				PF_R8,
+				FClearValueBinding::None,
+				/* InFlags = */ TexCreate_ShaderResource | TexCreate_UAV);
+
+			InterferenceWeightTexture = GraphBuilder.CreateTexture(Desc, TEXT("TSR.Interference.Weight"));
+		}
+
+		FTSRDetectInterferenceCS::FParameters* PassParameters = GraphBuilder.AllocParameters<FTSRDetectInterferenceCS::FParameters>();
+		PassParameters->CommonParameters = CommonParameters;
+
+		PassParameters->InputSceneColorTexture = PassInputs.SceneColorTexture;
+		PassParameters->DilatedVelocityTexture = DilatedVelocityTexture;
+		PassParameters->PredictionSceneColorTexture = PredictionSceneColorTexture;
+		PassParameters->ParallaxRejectionMaskTexture = ParallaxRejectionMaskTexture;
+		PassParameters->InterferenceSeedTexture = InterferenceSeedTexture;
+
+		PassParameters->PrevHistoryParameters = PrevHistoryParameters;
+		PassParameters->PrevHistory = PrevHistory;
+
+		PassParameters->HistoryOutput = CreateUAVs(GraphBuilder, History);
+		PassParameters->InterferenceWeightOutput = GraphBuilder.CreateUAV(InterferenceWeightTexture);
+		PassParameters->DebugOutput = CreateDebugUAV(LowFrequencyExtent, TEXT("Debug.TSR.DetectInterference"));
+
+		TShaderMapRef<FTSRDetectInterferenceCS> ComputeShader(View.ShaderMap);
+		ClearUnusedGraphResources(ComputeShader, PassParameters);
+
+		for (int32 i = 0; i < PrevHistory.LowResTextures.Num(); i++)
+		{
+			bool bNeedsExtractForNextFrame = PassParameters->PrevHistory.LowResTextures[i] != nullptr;
+			bool bPrevFrameIsntAvailable = PassParameters->PrevHistory.LowResTextures[i] == BlackDummy;
+			bool bOutputHistory = PassParameters->HistoryOutput.LowResTextures[i] != nullptr;
+
+			ExtractLowResHistoryTexture[i] = bNeedsExtractForNextFrame;
+
+			if (bPrevFrameIsntAvailable && !PassParameters->CommonParameters.bCameraCut)
+			{
+				//ensureMsgf(false, TEXT("Shaders read PrevHistory[%d] but doesn't write HistoryOutput[%d]"), i, i);
+				PassParameters->CommonParameters.bCameraCut = true;
+			}
+
+			if (bOutputHistory && !bNeedsExtractForNextFrame)
+			{
+				ensureMsgf(false, TEXT("Shaders write HistoryOutput[%d] but doesn't read PrevHistory[%d]"), i, i);
+			}
+		}
+
+		FComputeShaderUtils::AddPass(
+			GraphBuilder,
+			RDG_EVENT_NAME("TSR DetectInterference %dx%d", InputRect.Width(), InputRect.Height()),
+			ComputeShader,
+			PassParameters,
+			FComputeShaderUtils::GetGroupCount(InputRect.Size(), 8));
+	}
+	else
+	{
+		for (int32 i = 0; i < ExtractLowResHistoryTexture.Num(); i++)
+		{
+			ExtractLowResHistoryTexture[i] = false;
+		}
+
+		// TODO(TSR): Shader permutation.
+		InterferenceWeightTexture = WhiteDummy;
+	}
+
+	// Reject the history with frequency decomposition.
+	FRDGTextureRef HistoryRejectionTexture;
+	{
+		// Filter out the high frquencies
+		FRDGTextureRef FilteredInputTexture;
+		FRDGTextureRef FilteredPredictionSceneColorTexture;
+		{
+			{
+				FRDGTextureDesc Desc = FRDGTextureDesc::Create2D(
+					LowFrequencyExtent,
+					PF_FloatR11G11B10,
+					FClearValueBinding::None,
+					/* InFlags = */ TexCreate_ShaderResource | TexCreate_UAV);
+
+				FilteredInputTexture = GraphBuilder.CreateTexture(Desc, TEXT("TSR.Filtered.SceneColor"));
+				FilteredPredictionSceneColorTexture = GraphBuilder.CreateTexture(Desc, TEXT("TSR.Filtered.Prediction.SceneColor"));
+			}
+
+			FTSRFilterFrequenciesCS::FParameters* PassParameters = GraphBuilder.AllocParameters<FTSRFilterFrequenciesCS::FParameters>();
+			PassParameters->CommonParameters = CommonParameters;
+			PassParameters->OutputQuantizationError = ComputePixelFormatQuantizationError(FilteredInputTexture->Desc.Format);
+
+			if (bHalfResLowFrequency)
+			{
+				PassParameters->InputTexture = HalfResInputSceneColorTexture;
+				PassParameters->PredictionSceneColorTexture = HalfResPredictionSceneColorTexture;
+				PassParameters->ParallaxRejectionMaskTexture = HalfResParallaxRejectionMaskTexture;
+			}
+			else
+			{
+				PassParameters->InputTexture = PassInputs.SceneColorTexture;
+				PassParameters->PredictionSceneColorTexture = PredictionSceneColorTexture;
+				PassParameters->ParallaxRejectionMaskTexture = ParallaxRejectionMaskTexture;
+				PassParameters->InterferenceWeightTexture = InterferenceWeightTexture;
+			}
+
+			PassParameters->FilteredInputOutput = GraphBuilder.CreateUAV(FilteredInputTexture);
+			PassParameters->FilteredPredictionSceneColorOutput = GraphBuilder.CreateUAV(FilteredPredictionSceneColorTexture);
+			PassParameters->DebugOutput = CreateDebugUAV(LowFrequencyExtent, TEXT("Debug.TSR.FilterFrequencies"));
+
+			TShaderMapRef<FTSRFilterFrequenciesCS> ComputeShader(View.ShaderMap);
+			FComputeShaderUtils::AddPass(
+				GraphBuilder,
+				RDG_EVENT_NAME("TSR FilterFrequencies %dx%d", LowFrequencyRect.Width(), LowFrequencyRect.Height()),
+				ComputeShader,
+				PassParameters,
+				FComputeShaderUtils::GetGroupCount(LowFrequencyRect.Size(), 16));
+		}
+
+		// Compare the low frequencies
+		{
+			bool bOutputHalfRes = PostFilter != ERejectionPostFilter::PreRejectionDownsample;
+
+			{
+				FRDGTextureDesc Desc = FRDGTextureDesc::Create2D(
+					bOutputHalfRes ? RejectionExtent : LowFrequencyExtent,
+					PF_R8,
+					FClearValueBinding::None,
+					/* InFlags = */ TexCreate_ShaderResource | TexCreate_UAV);
+
+				HistoryRejectionTexture = GraphBuilder.CreateTexture(Desc, TEXT("TSR.HistoryRejection"));
+			}
+
+			FTSRCompareHistoryCS::FParameters* PassParameters = GraphBuilder.AllocParameters<FTSRCompareHistoryCS::FParameters>();
+			PassParameters->CommonParameters = CommonParameters;
+			PassParameters->ParallaxRejectionMaskTexture = ParallaxRejectionMaskTexture;
+			PassParameters->FilteredInputTexture = FilteredInputTexture;
+			PassParameters->FilteredPredictionSceneColorTexture = FilteredPredictionSceneColorTexture;
+			PassParameters->InterferenceWeightTexture = InterferenceWeightTexture;
+
+			PassParameters->HistoryRejectionOutput = GraphBuilder.CreateUAV(HistoryRejectionTexture);
+			PassParameters->DebugOutput = CreateDebugUAV(LowFrequencyExtent, TEXT("Debug.TSR.CompareHistory"));
+
+			FTSRCompareHistoryCS::FPermutationDomain PermutationVector;
+			PermutationVector.Set<FTSRCompareHistoryCS::FOutputHalfRes>(bOutputHalfRes);
+
+			TShaderMapRef<FTSRCompareHistoryCS> ComputeShader(View.ShaderMap, PermutationVector);
+			FComputeShaderUtils::AddPass(
+				GraphBuilder,
+				RDG_EVENT_NAME("TSR CompareHistory %dx%d", LowFrequencyRect.Width(), LowFrequencyRect.Height()),
+				ComputeShader,
+				PassParameters,
+				FComputeShaderUtils::GetGroupCount(LowFrequencyRect.Size(), 16));
+		}
+	}
+
+	// Post filter the rejection.
+	if (PostFilter != ERejectionPostFilter::Disabled)
+	{
+		bool bOutputHalfRes = PostFilter == ERejectionPostFilter::PreRejectionDownsample;
+		FIntRect Rect = bOutputHalfRes ? LowFrequencyRect : RejectionRect;
+
+		FRDGTextureRef FilteredHistoryRejectionTexture;
+		{
+			FRDGTextureDesc Desc = FRDGTextureDesc::Create2D(
+				RejectionExtent,
+				PF_R8,
+				FClearValueBinding::None,
+				/* InFlags = */ TexCreate_ShaderResource | TexCreate_UAV);
+
+			FilteredHistoryRejectionTexture = GraphBuilder.CreateTexture(Desc, TEXT("TSR.HistoryRejection"));
+		}
+
+		FTSRPostfilterRejectionCS::FParameters* PassParameters = GraphBuilder.AllocParameters<FTSRPostfilterRejectionCS::FParameters>();
+		PassParameters->HistoryRejectionViewport = Rect;
+		PassParameters->HistoryRejectionTexture = HistoryRejectionTexture;
+		PassParameters->HistoryRejectionOutput = GraphBuilder.CreateUAV(FilteredHistoryRejectionTexture);
+		PassParameters->DebugOutput = CreateDebugUAV(bOutputHalfRes ? LowFrequencyExtent : RejectionExtent, TEXT("Debug.TSR.PostfilterRejection"));
+
+		FTSRPostfilterRejectionCS::FPermutationDomain PermutationVector;
+		PermutationVector.Set<FTSRPostfilterRejectionCS::FOutputHalfRes>(PostFilter == ERejectionPostFilter::PreRejectionDownsample);
+
+		TShaderMapRef<FTSRPostfilterRejectionCS> ComputeShader(View.ShaderMap, PermutationVector);
+		FComputeShaderUtils::AddPass(
+			GraphBuilder,
+			RDG_EVENT_NAME("TSR PostfilterRejection %dx%d", Rect.Width(), Rect.Height()),
+			ComputeShader,
+			PassParameters,
+			FComputeShaderUtils::GetGroupCount(Rect.Size(), 8));
+
+		HistoryRejectionTexture = FilteredHistoryRejectionTexture;
+	}
+
+	// Dilate the rejection.
+	FRDGTextureRef DilatedHistoryRejectionTexture;
+	{
+		DilatedHistoryRejectionTexture = GraphBuilder.CreateTexture(HistoryRejectionTexture->Desc, TEXT("TSR.DilatedHistoryRejection"));
+
+		FTSRDilateRejectionCS::FParameters* PassParameters = GraphBuilder.AllocParameters<FTSRDilateRejectionCS::FParameters>();
+		PassParameters->CommonParameters = CommonParameters;
+		PassParameters->HistoryRejectionTexture = HistoryRejectionTexture;
+		PassParameters->DilatedHistoryRejectionOutput = GraphBuilder.CreateUAV(DilatedHistoryRejectionTexture);
+		PassParameters->DebugOutput = CreateDebugUAV(RejectionExtent, TEXT("Debug.TSR.DilateRejection"));
+
+		TShaderMapRef<FTSRDilateRejectionCS> ComputeShader(View.ShaderMap);
+		FComputeShaderUtils::AddPass(
+			GraphBuilder,
+			RDG_EVENT_NAME("TSR DilateRejection %dx%d", RejectionRect.Width(), RejectionRect.Height()),
+			ComputeShader,
+			PassParameters,
+			FComputeShaderUtils::GetGroupCount(RejectionRect.Size(), 8));
+	}
+
+	TStaticArray<bool, FTemporalAAHistory::kSuperResRenderTargetCount> ExtractSuperResHistoryTexture;
+	if (bEnableInterferenceHeuristic)
+	{
+		FTSRUpdateSuperResHistoryCS::FParameters* PassParameters = GraphBuilder.AllocParameters<FTSRUpdateSuperResHistoryCS::FParameters>();
+		PassParameters->CommonParameters = CommonParameters;
+		PassParameters->InputSceneColorTexture = PassInputs.SceneColorTexture;
+		PassParameters->DilatedVelocityTexture = DilatedVelocityTexture;
+		PassParameters->ParallaxRejectionMaskTexture = ParallaxRejectionMaskTexture;
+
+		PassParameters->PrevHistoryParameters = PrevHistoryParameters;
+		PassParameters->PrevHistory = PrevHistory;
+
+		PassParameters->HistoryOutput = CreateUAVs(GraphBuilder, History);
+		PassParameters->DebugOutput = CreateDebugUAV(HistoryExtent * 2, TEXT("Debug.TSR.UpdateSuperResHistory"));
+
+		TShaderMapRef<FTSRUpdateSuperResHistoryCS> ComputeShader(View.ShaderMap);
+		ClearUnusedGraphResources(ComputeShader, PassParameters);
+
+		for (int32 i = 0; i < PrevHistory.SuperResTextures.Num(); i++)
+		{
+			bool bNeedsExtractForNextFrame = PassParameters->PrevHistory.SuperResTextures[i] != nullptr;
+			bool bPrevFrameIsntAvailable = PassParameters->PrevHistory.SuperResTextures[i] == BlackDummy;
+			bool bOutputHistory = PassParameters->HistoryOutput.SuperResTextures[i] != nullptr;
+
+			ExtractSuperResHistoryTexture[i] = bNeedsExtractForNextFrame;
+
+			if (bPrevFrameIsntAvailable && !PassParameters->CommonParameters.bCameraCut)
+			{
+				//ensureMsgf(false, TEXT("Shaders read PrevHistory[%d] but doesn't write HistoryOutput[%d]"), i, i);
+				PassParameters->CommonParameters.bCameraCut = true;
+			}
+
+			if (bOutputHistory && !bNeedsExtractForNextFrame)
+			{
+				ensureMsgf(false, TEXT("Shaders write HistoryOutput[%d] but doesn't read PrevHistory[%d]"), i, i);
+			}
+		}
+
+		FComputeShaderUtils::AddPass(
+			GraphBuilder,
+			RDG_EVENT_NAME("TSR UpdateSuperResHistory %dx%d", HistorySize.X * 2, HistorySize.Y * 2),
+			ComputeShader,
+			PassParameters,
+			FComputeShaderUtils::GetGroupCount(HistorySize * 2, 8));
+	}
+	else
+	{
+		for (int32 i = 0; i < ExtractSuperResHistoryTexture.Num(); i++)
+		{
+			ExtractSuperResHistoryTexture[i] = false;
+		}
+	}
+
+	TStaticArray<bool, FTemporalAAHistory::kRenderTargetCount> ExtractHistoryTexture;
+	FRDGTextureRef SceneColorOutputTexture;
+	{
+		// Allocate output
+		{
+			FRDGTextureDesc Desc = FRDGTextureDesc::Create2D(
+				HistoryExtent,
+				PF_FloatR11G11B10,
+				FClearValueBinding::None,
+				/* InFlags = */ TexCreate_ShaderResource | TexCreate_UAV);
+
+			SceneColorOutputTexture = GraphBuilder.CreateTexture(Desc, TEXT("TSR.Output"));
+		}
+
+		FTSRUpdateHistoryCS::FParameters* PassParameters = GraphBuilder.AllocParameters<FTSRUpdateHistoryCS::FParameters>();
+		PassParameters->CommonParameters = CommonParameters;
+		PassParameters->InputSceneColorTexture = PassInputs.SceneColorTexture;
+		PassParameters->InputSceneStencilTexture = GraphBuilder.CreateSRV(
+			FRDGTextureSRVDesc::CreateWithPixelFormat(PassInputs.SceneDepthTexture, PF_X24_G8));
+		PassParameters->HistoryRejectionTexture = DilatedHistoryRejectionTexture;
+		PassParameters->TranslucencyRejectionTexture = TranslucencyRejectionTexture ? TranslucencyRejectionTexture : BlackDummy;
+
+		PassParameters->DilatedVelocityTexture = DilatedVelocityTexture;
+		PassParameters->ParallaxFactorTexture = ParallaxFactorTexture;
+		PassParameters->ParallaxRejectionMaskTexture = ParallaxRejectionMaskTexture;
+
+		FScreenTransform HistoryPixelPosToViewportUV = (FScreenTransform::Identity + 0.5f) * CommonParameters.HistoryInfo.ViewportSizeInverse;
+		PassParameters->HistoryPixelPosToScreenPos = HistoryPixelPosToViewportUV * FScreenTransform::ViewportUVToScreenPos;
+		PassParameters->HistoryPixelPosToPPCo = HistoryPixelPosToViewportUV * CommonParameters.InputInfo.ViewportSize + CommonParameters.InputJitter + CommonParameters.InputPixelPosMin;
+		PassParameters->HistoryQuantizationError = ComputePixelFormatQuantizationError(History.Textures[0]->Desc.Format);
+		PassParameters->MinTranslucencyRejection = TranslucencyRejectionTexture == nullptr ? 1.0 : 0.0;
+		PassParameters->ResponsiveStencilMask = CVarTSREnableResponiveAA.GetValueOnRenderThread() ? (STENCIL_TEMPORAL_RESPONSIVE_AA_MASK) : 0;
+
+		PassParameters->PrevHistoryParameters = PrevHistoryParameters;
+		PassParameters->PrevHistory = PrevHistory;
+
+		PassParameters->HistoryOutput = CreateUAVs(GraphBuilder, History);
+		PassParameters->SceneColorOutput = GraphBuilder.CreateUAV(SceneColorOutputTexture);
+		PassParameters->DebugOutput = CreateDebugUAV(HistoryExtent, TEXT("Debug.TSR.UpdateHistory"));
+
+		TShaderMapRef<FTSRUpdateHistoryCS> ComputeShader(View.ShaderMap);
+		ClearUnusedGraphResources(ComputeShader, PassParameters);
+
+		for (int32 i = 0; i < PrevHistory.Textures.Num(); i++)
+		{
+			bool bNeedsExtractForNextFrame = PassParameters->PrevHistory.Textures[i] != nullptr;
+			bool bPrevFrameIsntAvailable = PassParameters->PrevHistory.Textures[i] == BlackDummy;
+			bool bOutputHistory = PassParameters->HistoryOutput.Textures[i] != nullptr;
+
+			ExtractHistoryTexture[i] = bNeedsExtractForNextFrame;
+
+			if (bPrevFrameIsntAvailable && !bCameraCut)
+			{
+				//ensureMsgf(false, TEXT("Shaders read PrevHistory[%d] but doesn't write HistoryOutput[%d]"), i, i);
+				PassParameters->CommonParameters.bCameraCut = true;
+			}
+
+			if (bOutputHistory && !bNeedsExtractForNextFrame)
+			{
+				ensureMsgf(false, TEXT("Shaders write HistoryOutput[%d] but doesn't read PrevHistory[%d]"), i, i);
+			}
+		}
+
+		FComputeShaderUtils::AddPass(
+			GraphBuilder,
+			RDG_EVENT_NAME("TSR UpdateHistory(%s) %dx%d", 
+				History.Textures[0]->Desc.Format == PF_FloatR11G11B10 ? TEXT("R11G11B10") : TEXT(""),
+				HistorySize.X, HistorySize.Y),
+			ComputeShader,
+			PassParameters,
+			FComputeShaderUtils::GetGroupCount(HistorySize, 8));
+	}
+
+	// Debug the history.
+	#if COMPILE_TSR_DEBUG_PASSES
+	if (bSetupDebugPasses)
+	{
+		const int32 kHistoryUpscalingFactor = 2;
+
+		FTSRDebugHistoryCS::FParameters* PassParameters = GraphBuilder.AllocParameters<FTSRDebugHistoryCS::FParameters>();
+		PassParameters->CommonParameters = CommonParameters;
+		PassParameters->History = History;
+		PassParameters->PrevHistory = PrevHistory;
+		PassParameters->DebugOutput = CreateDebugUAV(HistoryExtent * kHistoryUpscalingFactor, TEXT("Debug.TSR.History"));
+
+		TShaderMapRef<FTSRDebugHistoryCS> ComputeShader(View.ShaderMap);
+		FComputeShaderUtils::AddPass(
+			GraphBuilder,
+			RDG_EVENT_NAME("TSR DebugHistory %dx%d", HistorySize.X, HistorySize.Y),
+			ComputeShader,
+			PassParameters,
+			FComputeShaderUtils::GetGroupCount(HistorySize * kHistoryUpscalingFactor, 8));
+	}
+	#endif
+
+	if (!View.bStatePrevViewInfoIsReadOnly)
+	{
+		FTemporalAAHistory* OutputHistory = &View.ViewState->PrevFrameViewInfo.TemporalAAHistory;
+		OutputHistory->SafeRelease();
+
+		for (int32 i = 0; i < History.LowResTextures.Num(); i++)
+		{
+			if (ExtractLowResHistoryTexture[i])
+			{
+				GraphBuilder.QueueTextureExtraction(History.LowResTextures[i], &OutputHistory->LowResRT[i]);
+			}
+		}
+
+		for (int32 i = 0; i < History.Textures.Num(); i++)
+		{
+			if (ExtractHistoryTexture[i])
+			{
+				GraphBuilder.QueueTextureExtraction(History.Textures[i], &OutputHistory->RT[i]);
+			}
+		}
+
+		for (int32 i = 0; i < History.SuperResTextures.Num(); i++)
+		{
+			if (ExtractSuperResHistoryTexture[i])
+			{
+				GraphBuilder.QueueTextureExtraction(History.SuperResTextures[i], &OutputHistory->SuperResRT[i]);
+			}
+		}
+
+		OutputHistory->ViewportRect = FIntRect(FIntPoint(0, 0), HistorySize);
+		OutputHistory->ReferenceBufferSize = HistoryExtent;
+
+		if (bRejectSeparateTranslucency)
+		{
+			GraphBuilder.QueueTextureExtraction(
+				SeparateTranslucencyTexture, &View.ViewState->PrevFrameViewInfo.SeparateTranslucency);
+		}
+	}
+
+	// If we upscaled the history buffer, downsize back to the secondary screen percentage size.
+	if (HistorySize != OutputRect.Size())
+	{
+		SceneColorOutputTexture = ComputeMitchellNetravaliDownsample(
+			GraphBuilder, View,
+			/* InputViewport = */ FScreenPassTexture(SceneColorOutputTexture, FIntRect(FIntPoint(0, 0), HistorySize)),
+			/* OutputViewport = */ FScreenPassTextureViewport(OutputExtent, OutputRect));
+	}
+
+	*OutSceneColorTexture = SceneColorOutputTexture;
+	*OutSceneColorViewRect = OutputRect;
+} // AddTemporalSuperResolutionPasses()
