Make sure instance culling always initializes dummy resources on mobile

diff --git a/Engine/Source/Runtime/Renderer/Private/MobileShadingRenderer.cpp b/Engine/Source/Runtime/Renderer/Private/MobileShadingRenderer.cpp
--- a/Engine/Source/Runtime/Renderer/Private/MobileShadingRenderer.cpp
+++ b/Engine/Source/Runtime/Renderer/Private/MobileShadingRenderer.cpp
@@ -568,7 +568,6 @@ void FMobileSceneRenderer::InitViews(FRDGBuilder& GraphBuilder, FSceneTexturesCo
 		PrepareDistanceFieldScene(GraphBuilder, false);
 	}
 	
-	if (InstanceCullingManager.IsEnabled())
 	{
 		// GPUCULL_TODO: Possibly fold into unpack step
 		InstanceCullingManager.CullInstances(GraphBuilder, Scene->GPUScene);
