Improving the occlusion culling parallel code path.

diff --git a/Engine/Source/Runtime/Renderer/Private/SceneVisibility.cpp b/Engine/Source/Runtime/Renderer/Private/SceneVisibility.cpp
--- a/Engine/Source/Runtime/Renderer/Private/SceneVisibility.cpp
+++ b/Engine/Source/Runtime/Renderer/Private/SceneVisibility.cpp
@@ -509,6 +509,8 @@ static int32 FrustumCull(const FScene* RESTRICT Scene, FViewInfo& View)
 	ParallelFor(NumTasks, 
 		[&NumCulledPrimitives, Scene, &View, MaxDrawDistanceScale, HLODState](int32 TaskIndex)
 		{
+			SCOPED_NAMED_EVENT(SceneVisibility_FrustumCull, FColor::Red);
+
 			QUICK_SCOPE_CYCLE_COUNTER(STAT_FrustumCull_Loop);
 			const FPlane* PermutedPlanePtr = View.ViewFrustum.PermutedPlanes.GetData();
 			const int32 BitArrayNumInner = View.PrimitiveVisibilityMap.Num();
@@ -724,9 +726,9 @@ struct FHZBBound
 	, BoundsExtent(InBoundsExtent)
 	{}
 
-	FPrimitiveOcclusionHistory* TargetHistory;
-	FVector BoundsOrigin;
-	FVector BoundsExtent;
+	FPrimitiveOcclusionHistory* const TargetHistory;
+	const FVector BoundsOrigin;
+	const FVector BoundsExtent;
 };
 
 #define BALANCE_LOAD 1
@@ -736,82 +738,79 @@ struct FVisForPrimParams
 {
 	FVisForPrimParams(){}
 
-	FVisForPrimParams(const FScene* InScene, 
-						FViewInfo* InView, 
-						FViewElementPDI* InOcclusionPDI, 						
-						const int32 InStartIndex, 
-						const int32 InNumToProcess, 
-						const bool bInSubmitQueries, 
-						const bool bInHZBOcclusion,						
-						TArray<FPrimitiveOcclusionHistory>* OutOcclusionHistory,
-						TArray<FPrimitiveOcclusionHistory*>* OutQueriesToRelease,
-						TArray<FHZBBound>* OutHZBBounds,
-						TArray<FOcclusionBounds>* OutQueriesToRun,
-						TArray<bool>* OutSubIsOccluded)
+	FVisForPrimParams(const FScene* InScene,
+						FViewInfo* InView,
+						FViewElementPDI* InOcclusionPDI,
+						int32 InStartIndex,
+						int32 InNumToProcess,
+						bool bInSubmitQueries,
+						bool bInHZBOcclusion,
+						TArray<FOcclusionBounds>& OutQueriesToRun,
+						TArray<bool>& OutSubIsOccluded)
 		: Scene(InScene)
 		, View(InView)
 		, OcclusionPDI(InOcclusionPDI)		
 		, StartIndex(InStartIndex)
 		, NumToProcess(InNumToProcess)
 		, bSubmitQueries(bInSubmitQueries)
-		, bHZBOcclusion(bInHZBOcclusion)		
-		, bNeedsScanOnRead(false)
-		, InsertPrimitiveOcclusionHistory(OutOcclusionHistory)
-		, QueriesToRelease(OutQueriesToRelease)
-		, HZBBoundsToAdd(OutHZBBounds)
-		, QueriesToAdd(OutQueriesToRun)	
-		, SubIsOccluded(OutSubIsOccluded)
+		, bHZBOcclusion(bInHZBOcclusion)
+		, QueriesToAdd(&OutQueriesToRun)
+		, SubIsOccluded(&OutSubIsOccluded)
 	{
-
+		OutQueriesToRun.Reset();
+		OutSubIsOccluded.Reset();
 	}
 
 	void Init(	const FScene* InScene,
 				FViewInfo* InView,
-				FViewElementPDI* InOcclusionPDI,
-				const int32 InStartIndex,
-				const int32 InNumToProcess,
-				const bool bInSubmitQueries,
-				const bool bInHZBOcclusion,				
-				TArray<FPrimitiveOcclusionHistory>* OutOcclusionHistory,
-				TArray<FPrimitiveOcclusionHistory*>* OutQueriesToRelease,
-				TArray<FHZBBound>* OutHZBBounds,
-				TArray<FOcclusionBounds>* OutQueriesToRun,
-				TArray<bool>* OutSubIsOccluded)
-			
+				int32 InStartIndex,
+				int32 InNumToProcess,
+				bool bInSubmitQueries,
+				bool bInHZBOcclusion,
+				TArray<FPrimitiveOcclusionHistory>& OutOcclusionHistory,
+				TArray<FPrimitiveOcclusionHistory*>& OutQueriesToRelease,
+				TArray<FHZBBound>& OutHZBBounds,
+				TArray<FOcclusionBounds>& OutQueriesToRun,
+				TArray<bool>& OutSubIsOccluded)
 	{
+		OutOcclusionHistory.Reset();
+		OutQueriesToRelease.Reset();
+		OutHZBBounds.Reset();
+		OutQueriesToRun.Reset();
+		OutSubIsOccluded.Reset();
+
 		Scene = InScene;
 		View = InView;
-		OcclusionPDI = InOcclusionPDI;
 		StartIndex = InStartIndex;
 		NumToProcess = InNumToProcess;
 		bSubmitQueries = bInSubmitQueries;
 		bHZBOcclusion = bInHZBOcclusion;
-		InsertPrimitiveOcclusionHistory = OutOcclusionHistory;
-		QueriesToRelease = OutQueriesToRelease;
-		HZBBoundsToAdd = OutHZBBounds;
-		QueriesToAdd = OutQueriesToRun;
-		SubIsOccluded = OutSubIsOccluded;
+		InsertPrimitiveOcclusionHistory = &OutOcclusionHistory;
+		QueriesToRelease = &OutQueriesToRelease;
+		HZBBoundsToAdd = &OutHZBBounds;
+		QueriesToAdd = &OutQueriesToRun;
+		SubIsOccluded = &OutSubIsOccluded;
 	}
 
-	const FScene* Scene;
-	FViewInfo* View;
-	FViewElementPDI* OcclusionPDI;
-	int32 StartIndex;
-	int32 NumToProcess;
-	bool bSubmitQueries;
-	bool bHZBOcclusion;	
+	const FScene* Scene{};
+	FViewInfo* View{};
+	FViewElementPDI* OcclusionPDI{};
+	int32 StartIndex{};
+	int32 NumToProcess{};
+	bool bSubmitQueries{};
+	bool bHZBOcclusion{};
 
 	// Whether the entries written into the history need to be read using a scan search (see FPrimitiveOcclusionHistory::bNeedsScanOnRead)
-	bool bNeedsScanOnRead;
+	bool bNeedsScanOnRead{};
 
 	//occlusion history to insert into.  In parallel these will be all merged back into the view's history on the main thread.
 	//use TChunkedArray so pointers to the new FPrimitiveOcclusionHistory's won't change if the array grows.	
-	TArray<FPrimitiveOcclusionHistory>*		InsertPrimitiveOcclusionHistory;
-	TArray<FPrimitiveOcclusionHistory*>*	QueriesToRelease;
-	TArray<FHZBBound>*						HZBBoundsToAdd;
-	TArray<FOcclusionBounds>*				QueriesToAdd;
-	int32									NumOccludedPrims;
-	TArray<bool>*							SubIsOccluded;
+	TArray<FPrimitiveOcclusionHistory>*		InsertPrimitiveOcclusionHistory{};
+	TArray<FPrimitiveOcclusionHistory*>*	QueriesToRelease{};
+	TArray<FHZBBound>*						HZBBoundsToAdd{};
+	TArray<FOcclusionBounds>*				QueriesToAdd{};
+	int32									NumOccludedPrims{};
+	TArray<bool>*							SubIsOccluded{};
 };
 
 //This function is shared between the single and multi-threaded versions.  Modifications to any primitives indexed by BitIt should be ok
@@ -819,8 +818,8 @@ struct FVisForPrimParams
 //to be recombined later.
 template<bool bSingleThreaded>
 static void FetchVisibilityForPrimitives_Range(FVisForPrimParams& Params, FGlobalDynamicVertexBuffer* DynamicVertexBufferIfSingleThreaded)
-{	
-	TRACE_CPUPROFILER_EVENT_SCOPE(FetchVisibilityForPrimitives_Range);
+{
+	SCOPED_NAMED_EVENT(FetchVisibilityForPrimitives_Range, FColor::Magenta);
 	int32 NumOccludedPrimitives = 0;
 	
 	const FScene* Scene				= Params.Scene;
@@ -840,7 +839,6 @@ static void FetchVisibilityForPrimitives_Range(FVisForPrimParams& Params, FGloba
 	bool bClearQueries = !View.Family->EngineShowFlags.HitProxies;
 	const float CurrentRealTime = View.Family->CurrentRealTime;
 	uint32 OcclusionFrameCounter = ViewState->OcclusionFrameCounter;
-	FRHIRenderQueryPool* OcclusionQueryPool = ViewState->OcclusionQueryPool;
 	FHZBOcclusionTester& HZBOcclusionTests = ViewState->HZBOcclusionTests;
 
 	int32 ReadBackLagTolerance = NumBufferedFrames;
@@ -893,7 +891,7 @@ static void FetchVisibilityForPrimitives_Range(FVisForPrimParams& Params, FGloba
 	int32 NumTotalDefUnoccluded = View.PrimitiveDefinitelyUnoccludedMap.Num();
 
 	{
-		TRACE_CPUPROFILER_EVENT_SCOPE_TEXT(*FString::Printf(TEXT("forEach over %d entries"), NumToProcess));
+		SCOPED_NAMED_EVENT_F(TEXT("forEach over %d entries"), FColor::Magenta, NumToProcess);
 	
 		//if we are load balanced then we iterate only the set bits, and the ranges have been pre-selected to evenly distribute set bits among the tasks with no overlaps.
 		//if not, then the entire array is evenly divided by range.
@@ -1309,49 +1307,12 @@ static void FetchVisibilityForPrimitives_Range(FVisForPrimParams& Params, FGloba
 	Params.NumOccludedPrims = NumOccludedPrimitives;	
 }
 
-FAutoConsoleTaskPriority CPrio_FetchVisibilityForPrimitivesTask(
-	TEXT("TaskGraph.TaskPriorities.FetchVisibilityForPrimitivesTask"),
-	TEXT("Task and thread priority for FetchVisibilityForPrimitivesTask."),
-	ENamedThreads::HighThreadPriority, // if we have high priority task threads, then use them...
-	ENamedThreads::NormalTaskPriority, // .. at normal task priority
-	ENamedThreads::HighTaskPriority // if we don't have hi pri threads, then use normal priority threads at high task priority instead
-	);
-
-class FetchVisibilityForPrimitivesTask
-{
-	FVisForPrimParams& Params;
-
-public:
-
-	FetchVisibilityForPrimitivesTask(FVisForPrimParams& InParams)
-		: Params(InParams)
-	{
-	}
-
-	FORCEINLINE TStatId GetStatId() const
-	{
-		RETURN_QUICK_DECLARE_CYCLE_STAT(FetchVisibilityForPrimitivesTask, STATGROUP_TaskGraphTasks);
-	}
-
-	ENamedThreads::Type GetDesiredThread()
-	{
-		return CPrio_FetchVisibilityForPrimitivesTask.Get();
-	}
-
-	static ESubsequentsMode::Type GetSubsequentsMode() { return ESubsequentsMode::TrackSubsequents; }
-
-	void DoTask(ENamedThreads::Type CurrentThread, const FGraphEventRef& MyCompletionGraphEvent)
-	{
-		FetchVisibilityForPrimitives_Range<false>(Params, nullptr);
-	}
-};
-
 static int32 FetchVisibilityForPrimitives(const FScene* Scene, FViewInfo& View, const bool bSubmitQueries, const bool bHZBOcclusion, FGlobalDynamicVertexBuffer& DynamicVertexBuffer)
 {
 	CSV_SCOPED_TIMING_STAT_EXCLUSIVE(FetchVisibilityForPrimitives);
 	QUICK_SCOPE_CYCLE_COUNTER(STAT_FetchVisibilityForPrimitives);
 	FSceneViewState* ViewState = (FSceneViewState*)View.State;
-	TRACE_CPUPROFILER_EVENT_SCOPE(FetchVisibilityForPrimitives);
+	SCOPED_NAMED_EVENT(FetchVisibilityForPrimitives, FColor::Magenta);
 	
 	static int32 SubIsOccludedArrayIndex = 0;
 	SubIsOccludedArrayIndex = 1 - SubIsOccludedArrayIndex;
@@ -1361,13 +1322,11 @@ static int32 FetchVisibilityForPrimitives(const FScene* Scene, FViewInfo& View,
 	TSet<FPrimitiveOcclusionHistory, FPrimitiveOcclusionHistoryKeyFuncs>& ViewPrimitiveOcclusionHistory = ViewState->PrimitiveOcclusionHistorySet;
 
 	if (GOcclusionCullParallelPrimFetch && GSupportsParallelOcclusionQueries)
-	{		
-		TRACE_CPUPROFILER_EVENT_SCOPE(FetchVisibilityParallel);
-		static const int32 MaxNumCullTasks = 4;
-		static const int32 ActualNumCullTasks = 4;
-		static const int32 NumOutputArrays = MaxNumCullTasks;
-		
-		FGraphEventRef TaskRefArray[NumOutputArrays];
+	{
+		SCOPED_NAMED_EVENT(FetchVisibilityParallel, FColor::Magenta);
+		constexpr int32 MaxNumCullTasks = 8;
+		constexpr int32 ActualNumCullTasks = 8;
+		constexpr int32 NumOutputArrays = MaxNumCullTasks;
 
 		//params for each task
 		FVisForPrimParams Params[NumOutputArrays];
@@ -1376,7 +1335,7 @@ static int32 FetchVisibilityForPrimitives(const FScene* Scene, FViewInfo& View,
 		TArray<FPrimitiveOcclusionHistory> OutputOcclusionHistory[NumOutputArrays];
 		TArray<FPrimitiveOcclusionHistory*> OutQueriesToRelease[NumOutputArrays];
 		TArray<FHZBBound> OutHZBBounds[NumOutputArrays];
-		TArray<FOcclusionBounds> OutQueriesToRun[NumOutputArrays];	
+		TArray<FOcclusionBounds> OutQueriesToRun[NumOutputArrays];
 
 		static TArray<bool> FrameSubIsOccluded[NumOutputArrays][FSceneView::NumBufferedSubIsOccludedArrays];
 
@@ -1391,7 +1350,7 @@ static int32 FetchVisibilityForPrimitives(const FScene* Scene, FViewInfo& View,
 			for (FSceneSetBitIterator BitIt(View.PrimitiveVisibilityMap); BitIt; ++BitIt, ++NumBitsSet)
 			{
 			}
-			
+
 			int32 BitsPerTask = NumBitsSet / ActualNumCullTasks;
 			int32 NumBitsForRange = 0;
 			int32 CurrentStartIndex = 0;
@@ -1421,20 +1380,15 @@ static int32 FetchVisibilityForPrimitives(const FScene* Scene, FViewInfo& View,
 		const int32 NumPrims = View.PrimitiveVisibilityMap.Num();
 		const int32 NumPerTask = NumPrims / ActualNumCullTasks;
 		int32 StartIndex = 0;
-		int32 NumToProcess = NumPerTask;
 
-		FGraphEventArray TaskWaitArray;
 		int32 NumTasks = 0;
 		for (int32 i = 0; i < ActualNumCullTasks && (StartIndex < NumPrims); ++i, ++NumTasks)
 		{
-			NumToProcess = (i == (ActualNumCullTasks - 1)) ? (NumPrims - StartIndex) : NumPerTask;
-			TArray<bool>& SubIsOccluded = FrameSubIsOccluded[i][SubIsOccludedArrayIndex];
-			SubIsOccluded.Reset();
+			const int32 NumToProcess = (i == (ActualNumCullTasks - 1)) ? (NumPrims - StartIndex) : NumPerTask;
 
 			Params[i].Init(
 				Scene,
 				&View,
-				nullptr,
 #if BALANCE_LOAD
 				StartIndices[i],
 				ProcessRange[i],
@@ -1443,30 +1397,31 @@ static int32 FetchVisibilityForPrimitives(const FScene* Scene, FViewInfo& View,
 				NumToProcess,
 #endif
 				bSubmitQueries,
-				bHZBOcclusion,				
-				&OutputOcclusionHistory[i],
-				&OutQueriesToRelease[i],
-				&OutHZBBounds[i],
-				&OutQueriesToRun[i],
-				&SubIsOccluded
-				);
-
-			TaskRefArray[i] = TGraphTask<FetchVisibilityForPrimitivesTask>::CreateTask().ConstructAndDispatchWhenReady(Params[i]);			
-			TaskWaitArray.Add(TaskRefArray[i]);
+				bHZBOcclusion,
+				OutputOcclusionHistory[i],
+				OutQueriesToRelease[i],
+				OutHZBBounds[i],
+				OutQueriesToRun[i],
+				FrameSubIsOccluded[i][SubIsOccludedArrayIndex]
+			);
 
 			StartIndex += NumToProcess;
 		}
 
-		FRHIRenderQueryPool* OcclusionQueryPool = ViewState->OcclusionQueryPool;
+		ParallelFor(NumTasks,
+			[&Params](int32 Index)
+			{
+				FetchVisibilityForPrimitives_Range<false>(Params[Index], nullptr);
+			},
+			!(FApp::ShouldUseThreadingForPerformance() && CVarParallelInitViews.GetValueOnRenderThread() > 0 && IsInActualRenderingThread())
+		);
+
 		FHZBOcclusionTester& HZBOcclusionTests = ViewState->HZBOcclusionTests;		
 
 		int32 NumOccludedPrims = 0;
 		{
 			QUICK_SCOPE_CYCLE_COUNTER(STAT_FetchVisibilityForPrimitivesCombine);
-			TRACE_CPUPROFILER_EVENT_SCOPE(FetchVisibilityForPrimitivesCombine);
-
-			//wait for them all so we don't start modifying the prim histories while the gather is running
-			FTaskGraphInterface::Get().WaitUntilTasksComplete(TaskWaitArray, ENamedThreads::GetRenderThread_Local());
+			SCOPED_NAMED_EVENT(FetchVisibilityForPrimitivesCombine, FColor::Magenta);
 
 #if QUERY_SANITY_CHECK
 			{
@@ -1477,16 +1432,15 @@ static int32 FetchVisibilityForPrimitives(const FScene* Scene, FViewInfo& View,
 				for (int32 i = 0; i < NumTasks; ++i)
 				{
 					bool bAlreadyIn = false;
-					for (auto ReleaseQueryIter = OutQueriesToRelease[i].CreateIterator(); ReleaseQueryIter; ++ReleaseQueryIter)
+					for (FPrimitiveOcclusionHistory* History : OutQueriesToRelease[i])
 					{
-						FPrimitiveOcclusionHistory* History = *ReleaseQueryIter;
 						ReleaseQuerySet.Add(History->PrimitiveId.PrimIDValue, &bAlreadyIn);
 						checkf(!bAlreadyIn, TEXT("Prim: %i double released query."), History->PrimitiveId.PrimIDValue);
 					}
 
-					for (auto RunQueriesIter = OutQueriesToRun[i].CreateIterator(); RunQueriesIter; ++RunQueriesIter)
+					for (const FOcclusionBounds& OcclusionBounds : OutQueriesToRun[i])
 					{
-						FPrimitiveOcclusionHistory* History = RunQueriesIter->PrimitiveOcclusionHistory;
+						FPrimitiveOcclusionHistory* History = OcclusionBounds->PrimitiveOcclusionHistory;
 						RunQuerySet.Add(History->PrimitiveId.PrimIDValue, &bAlreadyIn);
 						checkf(!bAlreadyIn, TEXT("Prim: %i double run query."), History->PrimitiveId.PrimIDValue);
 					}					
@@ -1499,27 +1453,26 @@ static int32 FetchVisibilityForPrimitives(const FScene* Scene, FViewInfo& View,
 			for (int32 i = 0; i < NumTasks; ++i)
 			{
 				//HZB output
-				for (auto HZBBoundIter = OutHZBBounds[i].CreateIterator(); HZBBoundIter; ++HZBBoundIter)
+				for (const FHZBBound& HZBBounds : OutHZBBounds[i])
 				{
-					HZBBoundIter->TargetHistory->HZBTestIndex = HZBOcclusionTests.AddBounds(HZBBoundIter->BoundsOrigin, HZBBoundIter->BoundsExtent);
+					HZBBounds.TargetHistory->HZBTestIndex = HZBOcclusionTests.AddBounds(HZBBounds.BoundsOrigin, HZBBounds.BoundsExtent);
 				}
 
 				//Manual query release handling
-				for (auto ReleaseQueryIter = OutQueriesToRelease[i].CreateIterator(); ReleaseQueryIter; ++ReleaseQueryIter)
+				for (FPrimitiveOcclusionHistory* History : OutQueriesToRelease[i])
 				{
-					FPrimitiveOcclusionHistory* History = *ReleaseQueryIter;
 					History->ReleaseQuery(OcclusionFrameCounter, NumBufferedFrames);
 				}
 				
 				//New query batching
-				for (auto RunQueriesIter = OutQueriesToRun[i].CreateIterator(); RunQueriesIter; ++RunQueriesIter)
+				for (const FOcclusionBounds& OcclusionBounds : OutQueriesToRun[i])
 				{
-					RunQueriesIter->PrimitiveOcclusionHistory->SetCurrentQuery(OcclusionFrameCounter,
-						RunQueriesIter->bGroupedQuery ?
-						View.GroupedOcclusionQueries.BatchPrimitive(RunQueriesIter->BoundsOrigin, RunQueriesIter->BoundsExtent, DynamicVertexBuffer) :
-						View.IndividualOcclusionQueries.BatchPrimitive(RunQueriesIter->BoundsOrigin, RunQueriesIter->BoundsExtent, DynamicVertexBuffer),
+					OcclusionBounds.PrimitiveOcclusionHistory->SetCurrentQuery(OcclusionFrameCounter,
+						OcclusionBounds.bGroupedQuery ?
+						View.GroupedOcclusionQueries.BatchPrimitive(OcclusionBounds.BoundsOrigin, OcclusionBounds.BoundsExtent, DynamicVertexBuffer) :
+						View.IndividualOcclusionQueries.BatchPrimitive(OcclusionBounds.BoundsOrigin, OcclusionBounds.BoundsExtent, DynamicVertexBuffer),
 						NumBufferedFrames,
-						RunQueriesIter->bGroupedQuery,
+						OcclusionBounds.bGroupedQuery,
 						Params[i].bNeedsScanOnRead
 						);
 				}
@@ -1528,10 +1481,7 @@ static int32 FetchVisibilityForPrimitives(const FScene* Scene, FViewInfo& View,
 			//now add new primitive histories to the view. may resize the view's array.
 			for (int32 i = 0; i < NumTasks; ++i)
 			{											
-				for (int32 HistoryIndex = 0; HistoryIndex < OutputOcclusionHistory[i].Num(); ++HistoryIndex)
-				{
-					ViewPrimitiveOcclusionHistory.Add(MoveTemp(OutputOcclusionHistory[i][HistoryIndex]));
-				}
+				ViewPrimitiveOcclusionHistory.Append(MoveTemp(OutputOcclusionHistory[i]));
 
 				//accumulate occluded prims across tasks
 				NumOccludedPrims += Params[i].NumOccludedPrims;
@@ -1543,15 +1493,8 @@ static int32 FetchVisibilityForPrimitives(const FScene* Scene, FViewInfo& View,
 	else
 	{
 		TRACE_CPUPROFILER_EVENT_SCOPE(FetchVisibilityOther);
-		//SubIsOccluded stuff needs a frame's lifetime
-		TArray<bool>& SubIsOccluded = View.FrameSubIsOccluded[SubIsOccludedArrayIndex];
-		SubIsOccluded.Reset();
-
+		
 		static TArray<FOcclusionBounds> PendingIndividualQueriesWhenOptimizing;
-		PendingIndividualQueriesWhenOptimizing.Reset();
-
-		static TArray<FOcclusionBounds*> PendingIndividualQueriesWhenOptimizingSorter;
-		PendingIndividualQueriesWhenOptimizingSorter.Reset();
 
 		FViewElementPDI OcclusionPDI(&View, nullptr, nullptr);
 		int32 StartIndex = 0;
@@ -1563,12 +1506,10 @@ static int32 FetchVisibilityForPrimitives(const FScene* Scene, FViewInfo& View,
 			StartIndex,
 			NumToProcess,
 			bSubmitQueries,
-			bHZBOcclusion,			
-			nullptr,
-			nullptr,
-			nullptr,
-			&PendingIndividualQueriesWhenOptimizing,
-			&SubIsOccluded
+			bHZBOcclusion,
+			PendingIndividualQueriesWhenOptimizing,
+			//SubIsOccluded stuff needs a frame's lifetime
+			View.FrameSubIsOccluded[SubIsOccludedArrayIndex]
 			);
 
 		FetchVisibilityForPrimitives_Range<true>(Params, &DynamicVertexBuffer);
@@ -1576,6 +1517,9 @@ static int32 FetchVisibilityForPrimitives(const FScene* Scene, FViewInfo& View,
 		int32 IndQueries = PendingIndividualQueriesWhenOptimizing.Num();
 		if (IndQueries)
 		{
+			static TArray<FOcclusionBounds*> PendingIndividualQueriesWhenOptimizingSorter;
+			PendingIndividualQueriesWhenOptimizingSorter.Reset();
+
 			int32 SoftMaxQueries = GRHIMaximumReccommendedOustandingOcclusionQueries / FMath::Min(NumBufferedFrames, 2); // extra RHIT frame does not count
 			int32 UsedQueries = View.GroupedOcclusionQueries.GetNumBatchOcclusionQueries();
 
