Restore changes to PathTracing code caused by the last engine merge

diff --git a/Engine/Plugins/Experimental/OpenImageDenoise/OpenImageDenoise.uplugin b/Engine/Plugins/Experimental/OpenImageDenoise/OpenImageDenoise.uplugin
new file mode 100644
--- /dev/null
+++ b/Engine/Plugins/Experimental/OpenImageDenoise/OpenImageDenoise.uplugin
@@ -0,0 +1,28 @@
+{
+	"FileVersion" : 3,
+	"Version" : 1,
+	"VersionName" : "1.0",
+	"FriendlyName" : "OpenImageDenoise",
+	"Description" : "Denoising engine for the Unreal Path Tracer based on Intel's OpenImageDenoise library.",
+	"Category" : "Denoising",
+	"CreatedBy" : "Epic Games, Inc.",
+	"CreatedByURL" : "http://epicgames.com",
+	"DocsURL" : "",
+	"MarketplaceURL" : "",
+	"SupportURL" : "",
+	"EnabledByDefault" : true,
+	"CanContainContent" : false,
+	"IsBetaVersion" : false,
+	"Installed" : false,
+	"Modules" :
+	[
+		{
+			"Name" : "OpenImageDenoise",
+			"Type" : "RuntimeAndProgram",
+			"LoadingPhase" : "Default",
+			"WhitelistPlatforms": [
+				"Win64"
+			]
+		}
+	]
+}
diff --git a/Engine/Plugins/Experimental/OpenImageDenoise/Source/OpenImageDenoise/OpenImageDenoise.Build.cs b/Engine/Plugins/Experimental/OpenImageDenoise/Source/OpenImageDenoise/OpenImageDenoise.Build.cs
new file mode 100644
--- /dev/null
+++ b/Engine/Plugins/Experimental/OpenImageDenoise/Source/OpenImageDenoise/OpenImageDenoise.Build.cs
@@ -0,0 +1,26 @@
+// Copyright Epic Games, Inc. All Rights Reserved.
+
+namespace UnrealBuildTool.Rules
+{
+	public class OpenImageDenoise : ModuleRules
+	{
+		public OpenImageDenoise(ReadOnlyTargetRules Target) : base(Target)
+		{
+			PrivateDependencyModuleNames.AddRange(
+				new string[] {
+					"Core",
+					"CoreUObject",
+					"Engine",
+					"RenderCore",
+					"Renderer",
+					"RHI",
+					"IntelOIDN"
+				}
+			);
+			if (Target.bBuildEditor)
+			{
+				PrivateDependencyModuleNames.Add("MessageLog");
+			}
+		}
+	}
+}
diff --git a/Engine/Plugins/Experimental/OpenImageDenoise/Source/OpenImageDenoise/Private/Plugin.cpp b/Engine/Plugins/Experimental/OpenImageDenoise/Source/OpenImageDenoise/Private/Plugin.cpp
new file mode 100644
--- /dev/null
+++ b/Engine/Plugins/Experimental/OpenImageDenoise/Source/OpenImageDenoise/Private/Plugin.cpp
@@ -0,0 +1,106 @@
+// Copyright Epic Games, Inc. All Rights Reserved.
+
+#include "CoreMinimal.h"
+#include "Modules/ModuleInterface.h"
+#include "Modules/ModuleManager.h"
+#include "Renderer/Public/PathTracingDenoiser.h"
+#include "RHI.h"
+#include "RHIResources.h"
+
+#include "OpenImageDenoise/oidn.hpp"
+
+class FOpenImageDenoiseModule : public IModuleInterface
+{
+public:
+	/** IModuleInterface implementation */
+	virtual void StartupModule() override;
+	virtual void ShutdownModule() override;
+};
+
+#if WITH_EDITOR
+DECLARE_LOG_CATEGORY_EXTERN(LogOpenImageDenoise, Log, All);
+DEFINE_LOG_CATEGORY(LogOpenImageDenoise);
+#endif
+
+IMPLEMENT_MODULE(FOpenImageDenoiseModule, OpenImageDenoise)
+
+static void Denoise(FRHICommandListImmediate& RHICmdList, FRHITexture2D* ColorTex, FRHITexture2D* AlbedoTex, FRHITexture2D* NormalTex, FRHITexture2D* OutputTex)
+{
+	const int DenoiserMode = 2; // TODO: Expose setting for this
+
+#if WITH_EDITOR
+	uint64 FilterExecuteTime = 0;
+	FilterExecuteTime -= FPlatformTime::Cycles64();
+#endif
+
+	FIntPoint Size = ColorTex->GetSizeXY();
+	FIntRect Rect = FIntRect(0, 0, Size.X, Size.Y);
+	TArray<FLinearColor> RawPixels;
+	TArray<FLinearColor> RawAlbedo;
+	TArray<FLinearColor> RawNormal;
+	FReadSurfaceDataFlags ReadDataFlags(ERangeCompressionMode::RCM_MinMax);
+	ReadDataFlags.SetLinearToGamma(false);
+	RHICmdList.ReadSurfaceData(ColorTex, Rect, RawPixels, ReadDataFlags);
+	if (DenoiserMode >= 2)
+	{
+		RHICmdList.ReadSurfaceData(AlbedoTex, Rect, RawAlbedo, ReadDataFlags);
+		RHICmdList.ReadSurfaceData(NormalTex, Rect, RawNormal, ReadDataFlags);
+	}
+
+	check(RawPixels.Num() == Size.X * Size.Y);
+
+	uint32_t DestStride;
+	FLinearColor* DestBuffer = (FLinearColor*)RHICmdList.LockTexture2D(OutputTex, 0, RLM_WriteOnly, DestStride, false);
+
+	// create device only once?
+	oidn::DeviceRef OIDNDevice = oidn::newDevice();
+	OIDNDevice.commit();
+	oidn::FilterRef OIDNFilter = OIDNDevice.newFilter("RT");
+	OIDNFilter.setImage("color", RawPixels.GetData(), oidn::Format::Float3, Size.X, Size.Y, 0, sizeof(FLinearColor), sizeof(FLinearColor) * Size.X);
+	if (DenoiserMode >= 2)
+	{
+		OIDNFilter.setImage("albedo", RawAlbedo.GetData(), oidn::Format::Float3, Size.X, Size.Y, 0, sizeof(FLinearColor), sizeof(FLinearColor) * Size.X);
+		OIDNFilter.setImage("normal", RawNormal.GetData(), oidn::Format::Float3, Size.X, Size.Y, 0, sizeof(FLinearColor), sizeof(FLinearColor) * Size.X);
+	}
+	if (DenoiserMode >= 3)
+	{
+		OIDNFilter.set("cleanAux", true);
+	}
+	OIDNFilter.setImage("output", DestBuffer, oidn::Format::Float3, Size.X, Size.Y, 0, sizeof(FLinearColor), DestStride);
+	OIDNFilter.set("hdr", true);
+	OIDNFilter.commit();
+
+	OIDNFilter.execute();
+
+	RHICmdList.UnlockTexture2D(OutputTex, 0, false);
+
+#if WITH_EDITOR
+	const char* errorMessage;
+	if (OIDNDevice.getError(errorMessage) != oidn::Error::None)
+	{
+		UE_LOG(LogOpenImageDenoise, Warning, TEXT("Denoiser failed: %s"), *FString(errorMessage));
+		return;
+	}
+
+	FilterExecuteTime += FPlatformTime::Cycles64();
+	const double FilterExecuteTimeMS = 1000.0 * FPlatformTime::ToSeconds64(FilterExecuteTime);
+	UE_LOG(LogOpenImageDenoise, Log, TEXT("Denoised %d x %d pixels in %.2f ms"), Size.X, Size.Y, FilterExecuteTimeMS);
+#endif
+}
+
+
+void FOpenImageDenoiseModule::StartupModule()
+{
+#if WITH_EDITOR
+	UE_LOG(LogOpenImageDenoise, Log, TEXT("OIDN starting up"));
+#endif
+	GPathTracingDenoiserFunc = &Denoise;
+}
+
+void FOpenImageDenoiseModule::ShutdownModule()
+{
+#if WITH_EDITOR
+	UE_LOG(LogOpenImageDenoise, Log, TEXT("OIDN shutting down"));
+#endif
+	GPathTracingDenoiserFunc = nullptr;
+}
diff --git a/Engine/Source/Runtime/Engine/Classes/Engine/Scene.h b/Engine/Source/Runtime/Engine/Classes/Engine/Scene.h
--- a/Engine/Source/Runtime/Engine/Classes/Engine/Scene.h
+++ b/Engine/Source/Runtime/Engine/Classes/Engine/Scene.h
@@ -1198,6 +1198,9 @@ struct FPostProcessSettings
 	UPROPERTY(EditAnywhere, BlueprintReadWrite, Category = Overrides, meta = (PinHiddenByDefault, InlineEditConditionToggle))
 	uint32 bOverride_PathTracingMaxPathExposure : 1;
 
+	UPROPERTY(EditAnywhere, BlueprintReadWrite, Category = Overrides, meta = (PinHiddenByDefault, InlineEditConditionToggle))
+	uint32 bOverride_PathTracingEnableDenoiser : 1;
+
 	// -----------------------------------------------------------------------
 
 	/** Enable HQ Gaussian on high end mobile platforms. (ES3_1) */
@@ -1924,7 +1927,7 @@ struct FPostProcessSettings
 	int32 PathTracingMaxBounces;
 
 	/** Sets the samples per pixel for the path tracer. */
-	UPROPERTY(interp, EditAnywhere, BlueprintReadWrite, Category = "Rendering Features|PathTracing", meta = (ClampMin = "1", ClampMax = "65536", editcondition = "bOverride_PathTracingSamplesPerPixel", DisplayName = "Samples Per Pixel"))
+	UPROPERTY(interp, EditAnywhere, BlueprintReadWrite, Category = "Rendering Features|PathTracing", meta = (ClampMin = "1", UIMax = "65536", editcondition = "bOverride_PathTracingSamplesPerPixel", DisplayName = "Samples Per Pixel"))
 	int32 PathTracingSamplesPerPixel;
 
 	/** Sets anti-aliasing filter width for the path tracer. Lower values are sharper (and more aliased), larger values are softer (and blurrier). */
@@ -1939,6 +1942,10 @@ struct FPostProcessSettings
 	UPROPERTY(interp, EditAnywhere, BlueprintReadWrite, Category = "Rendering Features|PathTracing", meta = (ClampMin = "-10.0", ClampMax = "30.0", editcondition = "bOverride_PathTracingMaxPathExposure", DisplayName = "Max Path Exposure"))
 	float PathTracingMaxPathExposure;
 
+	/** Run the currently loaded denoiser plugin on the last sample to remove noise from the output. Has no effect if a plug-in is not loaded. */
+	UPROPERTY(EditAnywhere, BlueprintReadWrite, Category = "Rendering Features|PathTracing", meta = (editcondition = "bOverride_PathTracingEnableDenoiser", DisplayName = "Denoiser"))
+	uint32 PathTracingEnableDenoiser : 1;
+
 
 	UPROPERTY()
 	float LPVFadeRange_DEPRECATED;
diff --git a/Engine/Source/Runtime/Engine/Private/Scene.cpp b/Engine/Source/Runtime/Engine/Private/Scene.cpp
--- a/Engine/Source/Runtime/Engine/Private/Scene.cpp
+++ b/Engine/Source/Runtime/Engine/Private/Scene.cpp
@@ -580,6 +580,7 @@ FPostProcessSettings::FPostProcessSettings()
 	PathTracingFilterWidth = 3.0f;
 	PathTracingEnableEmissive = 1;
 	PathTracingMaxPathExposure = 30.0f;
+	PathTracingEnableDenoiser = 1;
 	
 	bMobileHQGaussian = false;
 
@@ -763,6 +764,7 @@ FPostProcessSettings::FPostProcessSettings(const FPostProcessSettings& Settings)
 	, bOverride_PathTracingFilterWidth(Settings.bOverride_PathTracingFilterWidth)
 	, bOverride_PathTracingEnableEmissive(Settings.bOverride_PathTracingEnableEmissive)
 	, bOverride_PathTracingMaxPathExposure(Settings.bOverride_PathTracingMaxPathExposure)
+	, bOverride_PathTracingEnableDenoiser(Settings.bOverride_PathTracingEnableDenoiser)
 
 	, bMobileHQGaussian(Settings.bMobileHQGaussian)
 	, BloomMethod(Settings.BloomMethod)
@@ -935,6 +937,7 @@ FPostProcessSettings::FPostProcessSettings(const FPostProcessSettings& Settings)
 	, PathTracingFilterWidth(Settings.PathTracingFilterWidth)
 	, PathTracingEnableEmissive(Settings.PathTracingEnableEmissive)
 	, PathTracingMaxPathExposure(Settings.PathTracingMaxPathExposure)
+	, PathTracingEnableDenoiser(Settings.PathTracingEnableDenoiser)
 
 	, ScreenPercentage(Settings.ScreenPercentage)
 
diff --git a/Engine/Source/Runtime/Engine/Private/SceneView.cpp b/Engine/Source/Runtime/Engine/Private/SceneView.cpp
--- a/Engine/Source/Runtime/Engine/Private/SceneView.cpp
+++ b/Engine/Source/Runtime/Engine/Private/SceneView.cpp
@@ -1585,6 +1585,10 @@ void FSceneView::OverridePostProcessSettings(const FPostProcessSettings& Src, fl
 			Dest.PathTracingMaxPathExposure = Src.PathTracingMaxPathExposure;
 		}
 
+		if (Src.bOverride_PathTracingEnableDenoiser)
+		{
+			Dest.PathTracingEnableDenoiser = Src.PathTracingEnableDenoiser;
+		}
 
 		if (Src.bOverride_DepthOfFieldBladeCount)
 		{
diff --git a/Engine/Source/Runtime/Engine/Public/SceneManagement.h b/Engine/Source/Runtime/Engine/Public/SceneManagement.h
--- a/Engine/Source/Runtime/Engine/Public/SceneManagement.h
+++ b/Engine/Source/Runtime/Engine/Public/SceneManagement.h
@@ -233,6 +233,19 @@ public:
 	 * returns the occlusion frame counter 
 	 */
 	virtual uint32 GetOcclusionFrameCounter() const = 0;
+
+#if RHI_RAYTRACING
+	/**
+	* returns the path tracer sample index
+	*/
+	virtual uint32 GetPathTracingSampleIndex() const = 0;
+
+	/**
+	* returns the path tracer sample count
+	*/
+	virtual uint32 GetPathTracingSampleCount() const = 0;
+#endif
+
 protected:
 	// Don't allow direct deletion of the view state, Destroy should be called instead.
 	virtual ~FSceneViewStateInterface() {}
diff --git a/Engine/Source/Runtime/Renderer/Private/PathTracing.cpp b/Engine/Source/Runtime/Renderer/Private/PathTracing.cpp
--- a/Engine/Source/Runtime/Renderer/Private/PathTracing.cpp
+++ b/Engine/Source/Runtime/Renderer/Private/PathTracing.cpp
@@ -2,6 +2,9 @@
 
 #include "PathTracing.h"
 #include "RHI.h"
+#include "PathTracingDenoiser.h"
+
+PathTracingDenoiserFunction* GPathTracingDenoiserFunc = nullptr;
 
 TAutoConsoleVariable<int32> CVarPathTracing(
 	TEXT("r.PathTracing"),
@@ -180,9 +183,19 @@ TAutoConsoleVariable<int32> CVarPathTracingLightGridVisualize(
 	0,
 	TEXT("Enables a visualization mode of the light grid density where red indicates the maximum light count has been reached (default = 0)\n")
 	TEXT("0: off (default)\n")
-	TEXT("1: light count heatmap (red - close to overflow, increase r.PathTracing.LightGridMaxCount)")
-	TEXT("2: unique light lists (colors are a function of which lights occupy each cell)")
-	TEXT("3: area light visualization (green: point light sources only, blue: some area light sources)"),
+	TEXT("1: light count heatmap (red - close to overflow, increase r.PathTracing.LightGridMaxCount)\n")
+	TEXT("2: unique light lists (colors are a function of which lights occupy each cell)\n")
+	TEXT("3: area light visualization (green: point light sources only, blue: some area light sources)\n"),
+	ECVF_RenderThreadSafe
+);
+
+TAutoConsoleVariable<int32> CVarPathTracingDenoiser(
+	TEXT("r.PathTracing.Denoiser"),
+	-1,
+	TEXT("Enable denoising of the path traced output (if a denoiser plugin is active) (default = -1 (driven by postprocesing volume))\n")
+	TEXT("-1: inherit from PostProcessVolume\n")
+	TEXT("0: disable denoiser\n")
+	TEXT("1: enable denoiser (if a denoiser plugin is active)\n"),
 	ECVF_RenderThreadSafe
 );
 
@@ -216,6 +229,7 @@ struct FPathTracingConfig
 	bool VisibleLights;
 	bool UseMISCompensation;
 	bool LockedSamplingPattern;
+	int DenoiserMode; // NOTE: does not require path tracing invalidation
 
 	bool IsDifferent(const FPathTracingConfig& Other) const
 	{
@@ -244,7 +258,8 @@ struct FPathTracingConfig
 };
 
 // This function prepares the portion of shader arguments that may involve invalidating the path traced state
-static void PrepareShaderArgs(const FViewInfo& View, FPathTracingData& PathTracingData) {
+static void PrepareShaderArgs(const FViewInfo& View, FPathTracingData& PathTracingData)
+{
 	PathTracingData.EnableDirectLighting = true;
 	int32 MaxBounces = CVarPathTracingMaxBounces.GetValueOnRenderThread();
 	if (MaxBounces < 0)
@@ -424,6 +439,8 @@ class FPathTracingRG : public FGlobalShader
 
 	BEGIN_SHADER_PARAMETER_STRUCT(FParameters, )
 		SHADER_PARAMETER_RDG_TEXTURE_UAV(RWTexture2D<float4>, RadianceTexture)
+		SHADER_PARAMETER_RDG_TEXTURE_UAV(RWTexture2D<float4>, AlbedoTexture)
+		SHADER_PARAMETER_RDG_TEXTURE_UAV(RWTexture2D<float4>, NormalTexture)
 		SHADER_PARAMETER_SRV(RaytracingAccelerationStructure, TLAS)
 
 		SHADER_PARAMETER_STRUCT_REF(FViewUniformShaderParameters, ViewUniformBuffer)
@@ -1196,9 +1213,33 @@ void FDeferredShadingSceneRenderer::PreparePathTracing(const FSceneViewFamily& V
 void FSceneViewState::PathTracingInvalidate()
 {
 	PathTracingRadianceRT.SafeRelease();
+	PathTracingAlbedoRT.SafeRelease();
+	PathTracingNormalRT.SafeRelease();
+	PathTracingRadianceDenoisedRT.SafeRelease();
 	PathTracingSampleIndex = 0;
 }
 
+uint32 FSceneViewState::GetPathTracingSampleIndex() const {
+	return PathTracingSampleIndex;
+}
+
+uint32 FSceneViewState::GetPathTracingSampleCount() const {
+	FPathTracingConfig* LastConfig = PathTracingLastConfig.Get();
+	if (LastConfig)
+	{
+		return LastConfig->PathTracingData.MaxSamples;
+	}
+	return 0;
+}
+
+
+BEGIN_SHADER_PARAMETER_STRUCT(FDenoiseTextureParameters, )
+	RDG_TEXTURE_ACCESS(InputTexture, ERHIAccess::CopySrc)
+	RDG_TEXTURE_ACCESS(InputAlbedo, ERHIAccess::CopySrc)
+	RDG_TEXTURE_ACCESS(InputNormal, ERHIAccess::CopySrc)
+	RDG_TEXTURE_ACCESS(OutputTexture, ERHIAccess::CopyDest)
+END_SHADER_PARAMETER_STRUCT()
+
 DECLARE_GPU_STAT_NAMED(Stat_GPU_PathTracing, TEXT("Path Tracing"));
 void FDeferredShadingSceneRenderer::RenderPathTracing(
 	FRDGBuilder& GraphBuilder,
@@ -1291,20 +1332,26 @@ void FDeferredShadingSceneRenderer::RenderPathTracing(
 
 	// Prepare radiance buffer (will be shared with display pass)
 	FRDGTexture* RadianceTexture = nullptr;
+	FRDGTexture* AlbedoTexture = nullptr;
+	FRDGTexture* NormalTexture = nullptr;
 	if (View.ViewState->PathTracingRadianceRT)
 	{
 		// we already have a valid radiance texture, re-use it
 		RadianceTexture = GraphBuilder.RegisterExternalTexture(View.ViewState->PathTracingRadianceRT, TEXT("PathTracer.Radiance"));
+		AlbedoTexture   = GraphBuilder.RegisterExternalTexture(View.ViewState->PathTracingAlbedoRT, TEXT("PathTracer.Albedo"));
+		NormalTexture   = GraphBuilder.RegisterExternalTexture(View.ViewState->PathTracingNormalRT, TEXT("PathTracer.Normal"));
 	}
 	else
 	{
 		// First time through, need to make a new texture
-		FRDGTextureDesc RadianceTextureDesc = FRDGTextureDesc::Create2D(
+		FRDGTextureDesc Desc = FRDGTextureDesc::Create2D(
 			View.ViewRect.Size(),
 			PF_A32B32G32R32F,
 			FClearValueBinding::None,
 			TexCreate_ShaderResource | TexCreate_UAV);
-		RadianceTexture = GraphBuilder.CreateTexture(RadianceTextureDesc, TEXT("PathTracer.Radiance"), ERDGTextureFlags::MultiFrame);
+		RadianceTexture = GraphBuilder.CreateTexture(Desc, TEXT("PathTracer.Radiance"), ERDGTextureFlags::MultiFrame);
+		AlbedoTexture   = GraphBuilder.CreateTexture(Desc, TEXT("PathTracer.Albedo")  , ERDGTextureFlags::MultiFrame);
+		NormalTexture   = GraphBuilder.CreateTexture(Desc, TEXT("PathTracer.Normal")  , ERDGTextureFlags::MultiFrame);
 	}
 	const bool bNeedsMoreRays = Config.PathTracingData.Iteration < MaxSPP;
 
@@ -1323,6 +1370,8 @@ void FDeferredShadingSceneRenderer::RenderPathTracing(
 
 		PassParameters->IESTextureSampler = TStaticSamplerState<SF_Bilinear, AM_Clamp, AM_Clamp, AM_Clamp>::GetRHI();
 		PassParameters->RadianceTexture = GraphBuilder.CreateUAV(RadianceTexture);
+		PassParameters->AlbedoTexture   = GraphBuilder.CreateUAV(AlbedoTexture);
+		PassParameters->NormalTexture   = GraphBuilder.CreateUAV(NormalTexture);
 
 		PassParameters->SSProfilesTexture = GetSubsufaceProfileTexture_RT(GraphBuilder.RHICmdList)->GetShaderResourceRHI();
 
@@ -1359,8 +1408,72 @@ void FDeferredShadingSceneRenderer::RenderPathTracing(
 
 		// After we are done, make sure we remember our texture for next time so that we can accumulate samples across frames
 		GraphBuilder.QueueTextureExtraction(RadianceTexture, &View.ViewState->PathTracingRadianceRT);
+		GraphBuilder.QueueTextureExtraction(AlbedoTexture  , &View.ViewState->PathTracingAlbedoRT  );
+		GraphBuilder.QueueTextureExtraction(NormalTexture  , &View.ViewState->PathTracingNormalRT  );
+
+		// Bump counters for next frame
+		++View.ViewState->PathTracingSampleIndex;
+		++View.ViewState->PathTracingFrameIndex;
 	}
 
+	FRDGTexture* DenoisedRadianceTexture = nullptr;
+	int DenoiserMode = CVarPathTracingDenoiser.GetValueOnRenderThread();
+	if (DenoiserMode < 0)
+	{
+		DenoiserMode = View.FinalPostProcessSettings.PathTracingEnableDenoiser;
+	}
+	const bool IsDenoiserEnabled = DenoiserMode != 0 && GPathTracingDenoiserFunc != nullptr;
+	if (IsDenoiserEnabled)
+	{
+		// request denoise if this is the last sample
+		bool NeedsDenoise = (Config.PathTracingData.Iteration + 1) == MaxSPP;
+		// also allow turning on the denoiser after the image has stopped accumulating samples
+		if (!bNeedsMoreRays)
+		{
+			// we aren't currently rendering, run the denoiser if we just turned it on
+			NeedsDenoise |= DenoiserMode != View.ViewState->PathTracingLastConfig->DenoiserMode;
+		}
+
+		if (View.ViewState->PathTracingRadianceDenoisedRT)
+		{
+			// we already have a texture for this
+			DenoisedRadianceTexture = GraphBuilder.RegisterExternalTexture(View.ViewState->PathTracingRadianceDenoisedRT, TEXT("PathTracer.DenoisedRadiance"));
+		}
+
+		if (NeedsDenoise)
+		{
+			if (DenoisedRadianceTexture == nullptr)
+			{
+				// First time through, need to make a new texture
+				FRDGTextureDesc RadianceTextureDesc = FRDGTextureDesc::Create2D(
+					View.ViewRect.Size(),
+					PF_A32B32G32R32F,
+					FClearValueBinding::None,
+					TexCreate_ShaderResource | TexCreate_UAV);
+				DenoisedRadianceTexture = GraphBuilder.CreateTexture(RadianceTextureDesc, TEXT("PathTracer.DenoisedRadiance"), ERDGTextureFlags::MultiFrame);
+			}
+
+			FDenoiseTextureParameters* DenoiseParameters = GraphBuilder.AllocParameters<FDenoiseTextureParameters>();
+			DenoiseParameters->InputTexture = RadianceTexture;
+			DenoiseParameters->InputAlbedo = AlbedoTexture;
+			DenoiseParameters->InputNormal = NormalTexture;
+			DenoiseParameters->OutputTexture = DenoisedRadianceTexture;
+			GraphBuilder.AddPass(RDG_EVENT_NAME("Path Tracer Denoiser Plugin"), DenoiseParameters, ERDGPassFlags::Readback, 
+				[DenoiseParameters, DenoiserMode](FRHICommandListImmediate& RHICmdList)
+				{
+					GPathTracingDenoiserFunc(RHICmdList,
+						DenoiseParameters->InputTexture->GetRHI()->GetTexture2D(),
+						DenoiseParameters->InputAlbedo->GetRHI()->GetTexture2D(),
+						DenoiseParameters->InputNormal->GetRHI()->GetTexture2D(),
+						DenoiseParameters->OutputTexture->GetRHI()->GetTexture2D());
+				}
+			);
+
+			GraphBuilder.QueueTextureExtraction(DenoisedRadianceTexture, &View.ViewState->PathTracingRadianceDenoisedRT);
+		}
+	}
+	View.ViewState->PathTracingLastConfig->DenoiserMode = DenoiserMode;
+
 	// now add a pixel shader pass to display our Radiance buffer
 
 	FPathTracingCompositorPS::FParameters* DisplayParameters = GraphBuilder.AllocParameters<FPathTracingCompositorPS::FParameters>();
@@ -1368,7 +1481,7 @@ void FDeferredShadingSceneRenderer::RenderPathTracing(
 	DisplayParameters->MaxSamples = MaxSPP;
 	DisplayParameters->ProgressDisplayEnabled = CVarPathTracingProgressDisplay.GetValueOnRenderThread();
 	DisplayParameters->ViewUniformBuffer = View.ViewUniformBuffer;
-	DisplayParameters->RadianceTexture = GraphBuilder.CreateSRV(FRDGTextureSRVDesc::Create(RadianceTexture));
+	DisplayParameters->RadianceTexture = GraphBuilder.CreateSRV(FRDGTextureSRVDesc::Create(DenoisedRadianceTexture ? DenoisedRadianceTexture : RadianceTexture));
 	DisplayParameters->RenderTargets[0] = FRenderTargetBinding(SceneColorOutputTexture, ERenderTargetLoadAction::ELoad);
 
 	FScreenPassTextureViewport Viewport(SceneColorOutputTexture, View.ViewRect);
@@ -1391,10 +1504,6 @@ void FDeferredShadingSceneRenderer::RenderPathTracing(
 		PixelShader,
 		DisplayParameters
 	);
-
-	// Bump counters for next frame
-	++View.ViewState->PathTracingSampleIndex;
-	++View.ViewState->PathTracingFrameIndex;
 }
 
 #endif
diff --git a/Engine/Source/Runtime/Renderer/Private/ScenePrivate.h b/Engine/Source/Runtime/Renderer/Private/ScenePrivate.h
--- a/Engine/Source/Runtime/Renderer/Private/ScenePrivate.h
+++ b/Engine/Source/Runtime/Renderer/Private/ScenePrivate.h
@@ -990,9 +990,14 @@ public:
 #if RHI_RAYTRACING
 	// Invalidates cached results related to the path tracer so accumulated rendering can start over
 	void PathTracingInvalidate();
+	virtual uint32 GetPathTracingSampleIndex() const override;
+	virtual uint32 GetPathTracingSampleCount() const override;
 
 	// Reference path tracing cached results
 	TRefCountPtr<IPooledRenderTarget> PathTracingRadianceRT;
+	TRefCountPtr<IPooledRenderTarget> PathTracingAlbedoRT;
+	TRefCountPtr<IPooledRenderTarget> PathTracingNormalRT;
+	TRefCountPtr<IPooledRenderTarget> PathTracingRadianceDenoisedRT;
 	// Keeps track of the internal path tracer options relevant to detecting when to restart the path tracer accumulation
 	TPimplPtr<FPathTracingConfig> PathTracingLastConfig;
 
diff --git a/Engine/Source/Runtime/Renderer/Public/PathTracingDenoiser.h b/Engine/Source/Runtime/Renderer/Public/PathTracingDenoiser.h
new file mode 100644
--- /dev/null
+++ b/Engine/Source/Runtime/Renderer/Public/PathTracingDenoiser.h
@@ -0,0 +1,10 @@
+// Copyright Epic Games, Inc. All Rights Reserved.
+
+#pragma once
+
+class FRHITexture2D;
+class FRHICommandListImmediate;
+
+using PathTracingDenoiserFunction = void(FRHICommandListImmediate& RHICmdList, FRHITexture2D* ColorTex, FRHITexture2D* AlbedoTex, FRHITexture2D* NormalTex, FRHITexture2D* OutputTex);
+
+extern RENDERER_API PathTracingDenoiserFunction* GPathTracingDenoiserFunc;
