Fixed mobile specific issues when project has GPUScene enabled

diff --git a/Engine/Source/Runtime/Renderer/Private/DepthRendering.cpp b/Engine/Source/Runtime/Renderer/Private/DepthRendering.cpp
--- a/Engine/Source/Runtime/Renderer/Private/DepthRendering.cpp
+++ b/Engine/Source/Runtime/Renderer/Private/DepthRendering.cpp
@@ -619,7 +619,7 @@ bool FMobileSceneRenderer::ShouldRenderPrePass() const
 	return Scene->EarlyZPassMode == DDM_MaskedOnly || Scene->EarlyZPassMode == DDM_AllOpaque;
 }
 
-void FMobileSceneRenderer::RenderPrePass(FRHICommandListImmediate& RHICmdList, const FViewInfo& View, const FInstanceCullingDrawParams* InstanceCullingDrawParams/* = nullptr*/)
+void FMobileSceneRenderer::RenderPrePass(FRHICommandListImmediate& RHICmdList, const FViewInfo& View)
 {
 	if (!ShouldRenderPrePass())
 	{
@@ -636,7 +636,7 @@ void FMobileSceneRenderer::RenderPrePass(FRHICommandListImmediate& RHICmdList, c
 	SCOPED_GPU_STAT(RHICmdList, Prepass);
 
 	SetStereoViewport(RHICmdList, View);
-	View.ParallelMeshDrawCommandPasses[EMeshPass::DepthPass].DispatchDraw(nullptr, RHICmdList, InstanceCullingDrawParams);
+	View.ParallelMeshDrawCommandPasses[EMeshPass::DepthPass].DispatchDraw(nullptr, RHICmdList, &MeshPassInstanceCullingDrawParams[EMeshPass::DepthPass]);
 }
 
 void FDeferredShadingSceneRenderer::RenderPrePassHMD(FRDGBuilder& GraphBuilder, FRDGTextureRef DepthTexture)
diff --git a/Engine/Source/Runtime/Renderer/Private/GPUScene.cpp b/Engine/Source/Runtime/Renderer/Private/GPUScene.cpp
--- a/Engine/Source/Runtime/Renderer/Private/GPUScene.cpp
+++ b/Engine/Source/Runtime/Renderer/Private/GPUScene.cpp
@@ -1705,11 +1705,11 @@ public:
 		SHADER_PARAMETER_STRUCT_INCLUDE(FGPUSceneWriterParameters, GPUSceneWriterParameters)
 		END_SHADER_PARAMETER_STRUCT()
 
-		static constexpr int32 NumThreadsPerGroup = FInstanceProcessingGPULoadBalancer::ThreadGroupSize;
+	static constexpr int32 NumThreadsPerGroup = FInstanceProcessingGPULoadBalancer::ThreadGroupSize;
 
 	static bool ShouldCompilePermutation(const FGlobalShaderPermutationParameters& Parameters)
 	{
-		return IsFeatureLevelSupported(Parameters.Platform, ERHIFeatureLevel::SM5);
+		return UseGPUScene(Parameters.Platform, GetMaxSupportedFeatureLevel(Parameters.Platform));
 	}
 
 	static void ModifyCompilationEnvironment(const FGlobalShaderPermutationParameters& Parameters, FShaderCompilerEnvironment& OutEnvironment)
@@ -1736,7 +1736,7 @@ void FGPUScene::AddUpdatePrimitiveIdsPass(FRDGBuilder& GraphBuilder, FInstancePr
 
 		GetWriteParameters(PassParameters->GPUSceneWriterParameters);
 
-		auto ComputeShader = GetGlobalShaderMap(GMaxRHIFeatureLevel)->GetShader<FGPUSceneSetInstancePrimitiveIdCS>();
+		auto ComputeShader = GetGlobalShaderMap(FeatureLevel)->GetShader<FGPUSceneSetInstancePrimitiveIdCS>();
 
 		FComputeShaderUtils::AddPass(
 			GraphBuilder,
diff --git a/Engine/Source/Runtime/Renderer/Private/MobileBasePassRendering.cpp b/Engine/Source/Runtime/Renderer/Private/MobileBasePassRendering.cpp
--- a/Engine/Source/Runtime/Renderer/Private/MobileBasePassRendering.cpp
+++ b/Engine/Source/Runtime/Renderer/Private/MobileBasePassRendering.cpp
@@ -300,7 +300,7 @@ void SetupMobileSkyReflectionUniformParameters(FSkyLightSceneProxy* SkyLight, FM
 	Parameters.TextureSampler = CaptureTexture->SamplerStateRHI;
 }
 
-void FMobileSceneRenderer::RenderMobileBasePass(FRHICommandListImmediate& RHICmdList, const FViewInfo& View, const FInstanceCullingDrawParams* InstanceCullingDrawParams)
+void FMobileSceneRenderer::RenderMobileBasePass(FRHICommandListImmediate& RHICmdList, const FViewInfo& View)
 {
 	CSV_SCOPED_TIMING_STAT_EXCLUSIVE(RenderBasePass);
 	SCOPED_DRAW_EVENT(RHICmdList, MobileBasePass);
@@ -308,11 +308,11 @@ void FMobileSceneRenderer::RenderMobileBasePass(FRHICommandListImmediate& RHICmd
 	SCOPED_GPU_STAT(RHICmdList, Basepass);
 
 	RHICmdList.SetViewport(View.ViewRect.Min.X, View.ViewRect.Min.Y, 0, View.ViewRect.Max.X, View.ViewRect.Max.Y, 1);
-	View.ParallelMeshDrawCommandPasses[EMeshPass::BasePass].DispatchDraw(nullptr, RHICmdList, InstanceCullingDrawParams);
+	View.ParallelMeshDrawCommandPasses[EMeshPass::BasePass].DispatchDraw(nullptr, RHICmdList, &MeshPassInstanceCullingDrawParams[EMeshPass::BasePass]);
 		
 	if (View.Family->EngineShowFlags.Atmosphere)
 	{
-		View.ParallelMeshDrawCommandPasses[EMeshPass::SkyPass].DispatchDraw(nullptr, RHICmdList);
+		View.ParallelMeshDrawCommandPasses[EMeshPass::SkyPass].DispatchDraw(nullptr, RHICmdList, &MeshPassInstanceCullingDrawParams[EMeshPass::SkyPass]);
 	}
 
 	// editor primitives
diff --git a/Engine/Source/Runtime/Renderer/Private/MobileShadingRenderer.cpp b/Engine/Source/Runtime/Renderer/Private/MobileShadingRenderer.cpp
--- a/Engine/Source/Runtime/Renderer/Private/MobileShadingRenderer.cpp
+++ b/Engine/Source/Runtime/Renderer/Private/MobileShadingRenderer.cpp
@@ -173,11 +173,7 @@ BEGIN_SHADER_PARAMETER_STRUCT(FMobileRenderPassParameters, )
 	SHADER_PARAMETER_STRUCT_INCLUDE(FViewShaderParameters, View)
 	SHADER_PARAMETER_RDG_UNIFORM_BUFFER(FMobileBasePassUniformParameters, MobileBasePass)
 	SHADER_PARAMETER_RDG_UNIFORM_BUFFER(FDebugViewModePassUniformParameters, DebugViewMode)
-	SHADER_PARAMETER_STRUCT_INCLUDE(FInstanceCullingDrawParams, InstanceCullingDrawParamsDepth)
-	SHADER_PARAMETER_STRUCT_INCLUDE(FInstanceCullingDrawParams, InstanceCullingDrawParamsOpaque)
-	SHADER_PARAMETER_STRUCT_INCLUDE(FInstanceCullingDrawParams, InstanceCullingDrawParamsSky)
-	SHADER_PARAMETER_STRUCT_INCLUDE(FInstanceCullingDrawParams, InstanceCullingDrawParamsTranslucency)
-	SHADER_PARAMETER_STRUCT_INCLUDE(FInstanceCullingDrawParams, InstanceCullingDrawParamsDebugView)
+	SHADER_PARAMETER_STRUCT_INCLUDE(FInstanceCullingDrawParams, InstanceCullingDrawParams)
 	RENDER_TARGET_BINDING_SLOTS()
 END_SHADER_PARAMETER_STRUCT()
 
@@ -251,6 +247,8 @@ FMobileSceneRenderer::FMobileSceneRenderer(const FSceneViewFamily* InViewFamily,
 		}
 	}
 
+	FMemory::Memzero(MeshPassInstanceCullingDrawParams);
+
 	NumMSAASamples = GetDefaultMSAACount(ERHIFeatureLevel::ES3_1);
 }
 
@@ -609,7 +607,7 @@ void FMobileSceneRenderer::RenderFullDepthPrepass(FRDGBuilder& GraphBuilder, FSc
 
 	for (int32 ViewIndex = 0; ViewIndex < Views.Num(); ViewIndex++)
 	{
-		const FViewInfo& View = Views[ViewIndex];
+		FViewInfo& View = Views[ViewIndex];
 		if (!View.ShouldRenderView())
 		{
 			continue;
@@ -621,13 +619,15 @@ void FMobileSceneRenderer::RenderFullDepthPrepass(FRDGBuilder& GraphBuilder, FSc
 		PassParameters->RenderTargets = BasePassRenderTargets;
 		PassParameters->View = View.GetShaderParameters();
 
+		View.ParallelMeshDrawCommandPasses[EMeshPass::DepthPass].BuildRenderingCommands(GraphBuilder, Scene->GPUScene, PassParameters->InstanceCullingDrawParams);
+
 		GraphBuilder.AddPass(
 			RDG_EVENT_NAME("FullDepthPrepass"),
 			PassParameters,
 			ERDGPassFlags::Raster,
 			[this, PassParameters, &View](FRHICommandListImmediate& RHICmdList)
 			{
-				RenderPrePass(RHICmdList, View, &PassParameters->InstanceCullingDrawParamsDepth);
+				RenderPrePass(RHICmdList, View);
 
 				// Issue occlusion queries
 				RHICmdList.SetCurrentStat(GET_STATID(STAT_CLMM_Occlusion));
@@ -636,11 +636,11 @@ void FMobileSceneRenderer::RenderFullDepthPrepass(FRDGBuilder& GraphBuilder, FSc
 	}
 }
 
-void FMobileSceneRenderer::RenderMaskedPrePass(FRHICommandListImmediate& RHICmdList, const FViewInfo& View, const FInstanceCullingDrawParams* InstanceCullingDrawParams)
+void FMobileSceneRenderer::RenderMaskedPrePass(FRHICommandListImmediate& RHICmdList, const FViewInfo& View)
 {
 	if (bIsMaskedOnlyDepthPrepassEnabled)
 	{
-		RenderPrePass(RHICmdList, View, InstanceCullingDrawParams);
+		RenderPrePass(RHICmdList, View);
 	}
 }
 
@@ -882,14 +882,18 @@ void FMobileSceneRenderer::Render(FRDGBuilder& GraphBuilder)
 
 void FMobileSceneRenderer::BuildInstanceCullingDrawParams(FRDGBuilder& GraphBuilder, FViewInfo& View, FMobileRenderPassParameters* PassParameters)
 {
-	View.ParallelMeshDrawCommandPasses[EMeshPass::DepthPass].BuildRenderingCommands(GraphBuilder, Scene->GPUScene, PassParameters->InstanceCullingDrawParamsDepth);
-	View.ParallelMeshDrawCommandPasses[EMeshPass::BasePass].BuildRenderingCommands(GraphBuilder, Scene->GPUScene, PassParameters->InstanceCullingDrawParamsOpaque);
-	View.ParallelMeshDrawCommandPasses[EMeshPass::SkyPass].BuildRenderingCommands(GraphBuilder, Scene->GPUScene, PassParameters->InstanceCullingDrawParamsSky);
-	View.ParallelMeshDrawCommandPasses[StandardTranslucencyMeshPass].BuildRenderingCommands(GraphBuilder, Scene->GPUScene, PassParameters->InstanceCullingDrawParamsTranslucency);
+	View.ParallelMeshDrawCommandPasses[EMeshPass::DepthPass].BuildRenderingCommands(GraphBuilder, Scene->GPUScene, MeshPassInstanceCullingDrawParams[EMeshPass::DepthPass]);
+	View.ParallelMeshDrawCommandPasses[EMeshPass::BasePass].BuildRenderingCommands(GraphBuilder, Scene->GPUScene, MeshPassInstanceCullingDrawParams[EMeshPass::BasePass]);
+	View.ParallelMeshDrawCommandPasses[EMeshPass::SkyPass].BuildRenderingCommands(GraphBuilder, Scene->GPUScene, MeshPassInstanceCullingDrawParams[EMeshPass::SkyPass]);
+	View.ParallelMeshDrawCommandPasses[StandardTranslucencyMeshPass].BuildRenderingCommands(GraphBuilder, Scene->GPUScene, MeshPassInstanceCullingDrawParams[StandardTranslucencyMeshPass]);
 	if (ViewFamily.UseDebugViewPS())
 	{
-		View.ParallelMeshDrawCommandPasses[EMeshPass::DebugViewMode].BuildRenderingCommands(GraphBuilder, Scene->GPUScene, PassParameters->InstanceCullingDrawParamsDebugView);
+		View.ParallelMeshDrawCommandPasses[EMeshPass::DebugViewMode].BuildRenderingCommands(GraphBuilder, Scene->GPUScene, MeshPassInstanceCullingDrawParams[EMeshPass::DebugViewMode]);
 	}
+
+	// these resources should be shared among all mesh passes
+	// we assign it a pass parameters to make sure RDG does resource transtions
+	PassParameters->InstanceCullingDrawParams = MeshPassInstanceCullingDrawParams[EMeshPass::BasePass];
 }
 
 void FMobileSceneRenderer::RenderForward(FRDGBuilder& GraphBuilder, FRDGTextureRef ViewFamilyTexture, FSceneTextures& SceneTextures)
@@ -1012,10 +1016,10 @@ void FMobileSceneRenderer::RenderForwardSinglePass(FRDGBuilder& GraphBuilder, FM
 
 		// Depth pre-pass
 		RHICmdList.SetCurrentStat(GET_STATID(STAT_CLM_MobilePrePass));
-		RenderMaskedPrePass(RHICmdList, View, &PassParameters->InstanceCullingDrawParamsDepth);
+		RenderMaskedPrePass(RHICmdList, View);
 		// Opaque and masked
 		RHICmdList.SetCurrentStat(GET_STATID(STAT_CLMM_Opaque));
-		RenderMobileBasePass(RHICmdList, View, &PassParameters->InstanceCullingDrawParamsOpaque);
+		RenderMobileBasePass(RHICmdList, View);
 		RenderMobileDebugView(RHICmdList, View);
 		RHICmdList.ImmediateFlush(EImmediateFlushType::DispatchToRHIThread);
 		const bool bAdrenoOcclusionMode = (CVarMobileAdrenoOcclusionMode.GetValueOnRenderThread() != 0 && IsOpenGLPlatform(ShaderPlatform));
@@ -1032,7 +1036,7 @@ void FMobileSceneRenderer::RenderForwardSinglePass(FRDGBuilder& GraphBuilder, FM
 		RenderDecals(RHICmdList, View);
 		RenderModulatedShadowProjections(RHICmdList, ViewIndex, View);
 		// Draw translucency.
-		RenderTranslucency(RHICmdList, View, &PassParameters->InstanceCullingDrawParamsTranslucency);
+		RenderTranslucency(RHICmdList, View);
 		if (bAdrenoOcclusionMode)
 		{
 			RHICmdList.SetCurrentStat(GET_STATID(STAT_CLMM_Occlusion));
@@ -1061,10 +1065,10 @@ void FMobileSceneRenderer::RenderForwardMultiPass(FRDGBuilder& GraphBuilder, FMo
 
 		// Depth pre-pass
 		RHICmdList.SetCurrentStat(GET_STATID(STAT_CLM_MobilePrePass));
-		RenderMaskedPrePass(RHICmdList, View, &PassParameters->InstanceCullingDrawParamsDepth);
+		RenderMaskedPrePass(RHICmdList, View);
 		// Opaque and masked
 		RHICmdList.SetCurrentStat(GET_STATID(STAT_CLMM_Opaque));
-		RenderMobileBasePass(RHICmdList, View, &PassParameters->InstanceCullingDrawParamsOpaque);
+		RenderMobileBasePass(RHICmdList, View);
 		RenderMobileDebugView(RHICmdList, View);
 		// Issue occlusion queries
 		RHICmdList.SetCurrentStat(GET_STATID(STAT_CLMM_Occlusion));
@@ -1107,7 +1111,7 @@ void FMobileSceneRenderer::RenderForwardMultiPass(FRDGBuilder& GraphBuilder, FMo
 		RenderDecals(RHICmdList, View);
 		RenderModulatedShadowProjections(RHICmdList, ViewIndex, View);
 		// Draw translucency.
-		RenderTranslucency(RHICmdList, View, &SecondPassParameters->InstanceCullingDrawParamsTranslucency);
+		RenderTranslucency(RHICmdList, View);
 		// Pre-tonemap before MSAA resolve (iOS only)
 		PreTonemapMSAA(RHICmdList, SceneTextures);
 	});
@@ -1293,10 +1297,10 @@ void FMobileSceneRenderer::RenderDeferredSinglePass(FRDGBuilder& GraphBuilder, c
 		}
 		// Depth pre-pass
 		RHICmdList.SetCurrentStat(GET_STATID(STAT_CLM_MobilePrePass));
-		RenderMaskedPrePass(RHICmdList, View, &PassParameters->InstanceCullingDrawParamsDepth);
+		RenderMaskedPrePass(RHICmdList, View);
 		// Opaque and masked
 		RHICmdList.SetCurrentStat(GET_STATID(STAT_CLMM_Opaque));
-		RenderMobileBasePass(RHICmdList, View, &PassParameters->InstanceCullingDrawParamsOpaque);
+		RenderMobileBasePass(RHICmdList, View);
 		RHICmdList.ImmediateFlush(EImmediateFlushType::DispatchToRHIThread);
 		// Issue occlusion queries
 		RHICmdList.SetCurrentStat(GET_STATID(STAT_CLMM_Occlusion));
@@ -1315,7 +1319,7 @@ void FMobileSceneRenderer::RenderDeferredSinglePass(FRDGBuilder& GraphBuilder, c
 		}
 		MobileDeferredShadingPass(RHICmdList, ViewIndex, NumViews, View, *Scene, SortedLightSet);
 		// Draw translucency.
-		RenderTranslucency(RHICmdList, View, &PassParameters->InstanceCullingDrawParamsTranslucency);
+		RenderTranslucency(RHICmdList, View);
 	});
 }
 
@@ -1334,10 +1338,10 @@ void FMobileSceneRenderer::RenderDeferredMultiPass(FRDGBuilder& GraphBuilder, cl
 
 		// Depth pre-pass
 		RHICmdList.SetCurrentStat(GET_STATID(STAT_CLM_MobilePrePass));
-		RenderMaskedPrePass(RHICmdList, View, &PassParameters->InstanceCullingDrawParamsDepth);
+		RenderMaskedPrePass(RHICmdList, View);
 		// Opaque and masked
 		RHICmdList.SetCurrentStat(GET_STATID(STAT_CLMM_Opaque));
-		RenderMobileBasePass(RHICmdList, View, &PassParameters->InstanceCullingDrawParamsOpaque);
+		RenderMobileBasePass(RHICmdList, View);
 		// Issue occlusion queries
 		RHICmdList.SetCurrentStat(GET_STATID(STAT_CLMM_Occlusion));
 		RenderOcclusion(RHICmdList, View);
@@ -1394,7 +1398,7 @@ void FMobileSceneRenderer::RenderDeferredMultiPass(FRDGBuilder& GraphBuilder, cl
 		RHICmdList.NextSubpass();
 		MobileDeferredShadingPass(RHICmdList, ViewIndex, NumViews, View, *Scene, SortedLightSet);
 		// Draw translucency.
-		RenderTranslucency(RHICmdList, View, &ThirdPassParameters->InstanceCullingDrawParamsTranslucency);
+		RenderTranslucency(RHICmdList, View);
 	});
 }
 
@@ -1425,7 +1429,7 @@ void FMobileSceneRenderer::RenderMobileDebugView(FRHICommandListImmediate& RHICm
 		DrawClearQuad(RHICmdList, FLinearColor::Black);
 
 		RHICmdList.SetViewport(View.ViewRect.Min.X, View.ViewRect.Min.Y, 0, View.ViewRect.Max.X, View.ViewRect.Max.Y, 1);
-		View.ParallelMeshDrawCommandPasses[EMeshPass::DebugViewMode].DispatchDraw(nullptr, RHICmdList);
+		View.ParallelMeshDrawCommandPasses[EMeshPass::DebugViewMode].DispatchDraw(nullptr, RHICmdList, &MeshPassInstanceCullingDrawParams[EMeshPass::DebugViewMode]);
 	}
 #endif // WITH_DEBUG_VIEW_MODES
 }
diff --git a/Engine/Source/Runtime/Renderer/Private/MobileTranslucentRendering.cpp b/Engine/Source/Runtime/Renderer/Private/MobileTranslucentRendering.cpp
--- a/Engine/Source/Runtime/Renderer/Private/MobileTranslucentRendering.cpp
+++ b/Engine/Source/Runtime/Renderer/Private/MobileTranslucentRendering.cpp
@@ -31,7 +31,7 @@
 #include "MeshPassProcessor.inl"
 #include "ClearQuad.h"
 
-void FMobileSceneRenderer::RenderTranslucency(FRHICommandListImmediate& RHICmdList, const FViewInfo& View, const FInstanceCullingDrawParams* InstanceCullingDrawParams)
+void FMobileSceneRenderer::RenderTranslucency(FRHICommandListImmediate& RHICmdList, const FViewInfo& View)
 {	
 	const bool bShouldRenderTranslucency = ShouldRenderTranslucency(StandardTranslucencyPass) && ViewFamily.EngineShowFlags.Translucency && !ViewFamily.UseDebugViewPS();
 	if (bShouldRenderTranslucency)
@@ -42,7 +42,7 @@ void FMobileSceneRenderer::RenderTranslucency(FRHICommandListImmediate& RHICmdLi
 		SCOPED_GPU_STAT(RHICmdList, Translucency);
 
 		RHICmdList.SetViewport(View.ViewRect.Min.X, View.ViewRect.Min.Y, 0.0f, View.ViewRect.Max.X, View.ViewRect.Max.Y, 1.0f);
-		View.ParallelMeshDrawCommandPasses[StandardTranslucencyMeshPass].DispatchDraw(nullptr, RHICmdList, InstanceCullingDrawParams);
+		View.ParallelMeshDrawCommandPasses[StandardTranslucencyMeshPass].DispatchDraw(nullptr, RHICmdList, &MeshPassInstanceCullingDrawParams[StandardTranslucencyMeshPass]);
 	}
 }
 
diff --git a/Engine/Source/Runtime/Renderer/Private/SceneHitProxyRendering.cpp b/Engine/Source/Runtime/Renderer/Private/SceneHitProxyRendering.cpp
--- a/Engine/Source/Runtime/Renderer/Private/SceneHitProxyRendering.cpp
+++ b/Engine/Source/Runtime/Renderer/Private/SceneHitProxyRendering.cpp
@@ -551,7 +551,7 @@ void FMobileSceneRenderer::RenderHitProxies(FRDGBuilder& GraphBuilder)
 	FRDGTextureRef HitProxyDepthTexture = nullptr;
 	InitHitProxyRender(GraphBuilder, this, SceneTexturesConfig, HitProxyTexture, HitProxyDepthTexture);
 
-	FInstanceCullingManager InstanceCullingManager(Scene->GPUScene.IsEnabled());
+	FInstanceCullingManager& InstanceCullingManager = *GraphBuilder.AllocObject<FInstanceCullingManager>(Scene->GPUScene.IsEnabled());
 
 	// Find the visible primitives.
 	InitViews(GraphBuilder, SceneTexturesConfig, InstanceCullingManager);
diff --git a/Engine/Source/Runtime/Renderer/Private/SceneRendering.h b/Engine/Source/Runtime/Renderer/Private/SceneRendering.h
--- a/Engine/Source/Runtime/Renderer/Private/SceneRendering.h
+++ b/Engine/Source/Runtime/Renderer/Private/SceneRendering.h
@@ -2237,12 +2237,12 @@ protected:
 
 	void InitViews(FRDGBuilder& GraphBuilder, FSceneTexturesConfig& SceneTexturesConfig, FInstanceCullingManager& InstanceCullingManager);
 
-	void RenderPrePass(FRHICommandListImmediate& RHICmdList, const FViewInfo& View, const FInstanceCullingDrawParams* InstanceCullingDrawParams = nullptr);
-	void RenderMaskedPrePass(FRHICommandListImmediate& RHICmdList, const FViewInfo& View, const FInstanceCullingDrawParams* InstanceCullingDrawParams);
+	void RenderPrePass(FRHICommandListImmediate& RHICmdList, const FViewInfo& View);
+	void RenderMaskedPrePass(FRHICommandListImmediate& RHICmdList, const FViewInfo& View);
 	void RenderFullDepthPrepass(FRDGBuilder& GraphBuilder, FSceneTextures& SceneTextures);
 
 	/** Renders the opaque base pass for mobile. */
-	void RenderMobileBasePass(FRHICommandListImmediate& RHICmdList, const FViewInfo& View, const FInstanceCullingDrawParams* InstanceCullingDrawParams);
+	void RenderMobileBasePass(FRHICommandListImmediate& RHICmdList, const FViewInfo& View);
 
 	void PostRenderBasePass(FRHICommandListImmediate& RHICmdList, FViewInfo& View);
 
@@ -2274,7 +2274,7 @@ protected:
 	void RenderDecals(FRHICommandListImmediate& RHICmdList, const FViewInfo& View);
 
 	/** Renders the base pass for translucency. */
-	void RenderTranslucency(FRHICommandListImmediate& RHICmdList, const FViewInfo& View, const FInstanceCullingDrawParams* InstanceCullingDrawParams);
+	void RenderTranslucency(FRHICommandListImmediate& RHICmdList, const FViewInfo& View);
 
 	/** On chip pre-tonemap before scene color MSAA resolve (iOS only) */
 	void PreTonemapMSAA(FRHICommandListImmediate& RHICmdList, const FMinimalSceneTextures& SceneTextures);
@@ -2325,6 +2325,9 @@ private:
 	ETranslucencyPass::Type StandardTranslucencyPass;
 	EMeshPass::Type StandardTranslucencyMeshPass;
 
+	// TODO: remove this after instance culling is specialized for mobile
+	FInstanceCullingDrawParams MeshPassInstanceCullingDrawParams[EMeshPass::Num];
+
 	static FGlobalDynamicIndexBuffer DynamicIndexBuffer;
 	static FGlobalDynamicVertexBuffer DynamicVertexBuffer;
 	static TGlobalResource<FGlobalDynamicReadBuffer> DynamicReadBuffer;
