RG11B10 + SceneDepthAux 16/32F + Alpha Propagate

For Forward ES31
 Default SceneColor RG11B10 + R16F\32F Depth texture
 With PropagateAlpha on RGBA16F + R16F\32F
 PostProcess we sample SceneDepthAux for Depth
For Deferred ES31
 SceneDepthAux only for Metal
 PropagateAlpha not working yet
 PostProcess we sample SceneDepthTexture for Depth

cvar to change Depth texture from 16 to 32Fr.Mobile.SceneDepthAux
cvar for AlphaPropagate r.Mobile.PropagateAlpha

diff --git a/Engine/Build/Android/Java/src/com/epicgames/unreal/GameActivity.java.template b/Engine/Build/Android/Java/src/com/epicgames/unreal/GameActivity.java.template
--- a/Engine/Build/Android/Java/src/com/epicgames/unreal/GameActivity.java.template
+++ b/Engine/Build/Android/Java/src/com/epicgames/unreal/GameActivity.java.template
@@ -385,6 +385,9 @@ public class GameActivity extends $${gameActivitySuperClass}$$ implements Surfac
 	// depthbuffer preference from manifest
 	int DepthBufferPreference = 0;
 	
+	// PropagatetAlpha from the manifest
+	int PropagateAlpha = 0;
+
 	/** AssetManger reference - populated on start up and used when the OBB is packed into the APK */
 	private AssetManager			AssetManagerReference;
 	
@@ -2651,6 +2654,16 @@ public class GameActivity extends $${gameActivitySuperClass}$$ implements Surfac
 				Log.debug( "Did not find DepthBufferPreference, using default.");
 			}
 
+			if (bundle.containsKey("com.epicgames.ue4.GameActivity.PropagateAlpha"))
+			{
+				PropagateAlpha = bundle.getInt("com.epicgames.ue4.GameActivity.PropagateAlpha");
+				Log.debug( "Found PropagateAlpha = " + PropagateAlpha);
+			}
+			else
+			{
+				Log.debug( "Did not find PropagateAlpha, using default.");
+			}
+
 			// Determine if data is embedded in APK from AndroidManifest.xml
 			if (bundle.containsKey("com.epicgames.unreal.GameActivity.bPackageDataInsideApk"))
 			{
@@ -2812,7 +2825,7 @@ public class GameActivity extends $${gameActivitySuperClass}$$ implements Surfac
 		nativeSetGlobalActivity(UseExternalFilesDir, PublicLogFiles, getFilesDir().getPath(), getExternalFilesDir(null).getPath(), PackageDataInsideApkValue==1, APKPath);
 
 		// tell the engine if this is a portrait app
-		nativeSetWindowInfo(getResources().getConfiguration().orientation == Configuration.ORIENTATION_PORTRAIT, DepthBufferPreference);
+		nativeSetWindowInfo(getResources().getConfiguration().orientation == Configuration.ORIENTATION_PORTRAIT, DepthBufferPreference, PropagateAlpha);
 
 		// get the full language code, like en-US
 		// note: this may need to be Locale.getDefault().getLanguage()
@@ -3256,7 +3269,7 @@ public class GameActivity extends $${gameActivitySuperClass}$$ implements Surfac
 		}
 
 		// invalidate window cache
-		nativeSetWindowInfo(getResources().getConfiguration().orientation == Configuration.ORIENTATION_PORTRAIT, DepthBufferPreference);
+		nativeSetWindowInfo(getResources().getConfiguration().orientation == Configuration.ORIENTATION_PORTRAIT, DepthBufferPreference, PropagateAlpha);
 
 		// only do this on KitKat and above
 		if (ShouldHideUI)
@@ -6496,7 +6509,7 @@ public class GameActivity extends $${gameActivitySuperClass}$$ implements Surfac
 	public native void nativeSetAndroidStartupState(boolean bDebuggerAttached);
 	public native void nativeSetGlobalActivity(boolean bUseExternalFilesDir, boolean bPublicLogFiles, String internalFilePath, String externalFilePath, boolean bOBBInAPK, String APKPath);
 	public native void nativeSetObbFilePaths(String OBBMainFilePath, String OBBPatchFilePath, String OBBOverflow1FilePath, String OBBOverflow2FilePath);
-	public native void nativeSetWindowInfo(boolean bIsPortrait, int DepthBufferPreference);
+	public native void nativeSetWindowInfo(boolean bIsPortrait, int DepthBufferPreference, int PropagateAlpha);
 	public native void nativeSetObbInfo(String ProjectName, String PackageName, int Version, int PatchVersion, String AppType);
 	public native void nativeSetAndroidVersionInformation( String AndroidVersion, int TargetSDK, String PhoneMake, String PhoneModel, String PhoneBuildNumber, String OSLanguage );
 
diff --git a/Engine/Shaders/Private/FXAAShader.usf b/Engine/Shaders/Private/FXAAShader.usf
--- a/Engine/Shaders/Private/FXAAShader.usf
+++ b/Engine/Shaders/Private/FXAAShader.usf
@@ -106,7 +106,7 @@ void FxaaVS(
 			fxaaConsoleEdgeThresholdMin,
 			fxaaConsole360ConstDir);
 
-		#if (POST_PROCESS_ALPHA != 2)
+		#if (POST_PROCESS_ALPHA != 2) && !MOBILE_PROPAGATE_ALPHA
 			OutColor.a = 1.0;
 		#endif
 	}
diff --git a/Engine/Shaders/Private/MaterialTemplate.ush b/Engine/Shaders/Private/MaterialTemplate.ush
--- a/Engine/Shaders/Private/MaterialTemplate.ush
+++ b/Engine/Shaders/Private/MaterialTemplate.ush
@@ -1895,7 +1895,7 @@ MaterialFloat4 MobileSceneTextureLookup(inout FMaterialPixelParameters Parameter
 		#if MOBILE_DEFERRED_SHADING
 			MaterialFloat Depth = ConvertFromDeviceZ(Texture2DSample(MobileSceneTextures.SceneDepthTexture, MobileSceneTextures.SceneDepthTextureSampler, UV).r);
 		#else
-			MaterialFloat Depth = ConvertFromDeviceZ(Texture2DSample(MobileSceneTextures.SceneColorTexture, MobileSceneTextures.SceneColorTextureSampler, UV).a);
+			MaterialFloat Depth = ConvertFromDeviceZ(Texture2DSample(MobileSceneTextures.SceneDepthAuxTexture, MobileSceneTextures.SceneDepthAuxTextureSampler, UV).r);
 		#endif
 		return MaterialFloat4(Depth.rrr, 0.f);
 	}
diff --git a/Engine/Shaders/Private/MobileBasePassPixelShader.usf b/Engine/Shaders/Private/MobileBasePassPixelShader.usf
--- a/Engine/Shaders/Private/MobileBasePassPixelShader.usf
+++ b/Engine/Shaders/Private/MobileBasePassPixelShader.usf
@@ -92,10 +92,16 @@ int NumDynamicPointLights;
 #define FORCE_VERTEX_NORMAL (MOBILE_QL_DISABLE_MATERIAL_NORMAL)
 #define SUPPORT_SPOTLIGHTS_SHADOW (MATERIALBLENDING_SOLID || MATERIALBLENDING_MASKED) && PROJECT_MOBILE_ENABLE_MOVABLE_SPOTLIGHTS_SHADOW
 #define DEFERRED_SHADING_PATH (MOBILE_DEFERRED_SHADING && ((MATERIALBLENDING_SOLID || MATERIALBLENDING_MASKED) && !MATERIAL_SHADINGMODEL_SINGLELAYERWATER))
-#define USE_SCENE_DEPTH_AUX (DEFERRED_SHADING_PATH && METAL_PROFILE)
+#define USE_SCENE_DEPTH_AUX (DEFERRED_SHADING_PATH && METAL_PROFILE) || !DEFERRED_SHADING_PATH
+#if DEFERRED_SHADING_PATH
+	#define SV_TargetDepthAux SV_Target4
+#else
+	#define SV_TargetDepthAux SV_Target1
+#endif
 #define TRANSLUCENCY_NON_DIRECTIONAL (MATERIALBLENDING_ANY_TRANSLUCENT && (TRANSLUCENCY_LIGHTING_VOLUMETRIC_NONDIRECTIONAL || TRANSLUCENCY_LIGHTING_VOLUMETRIC_PERVERTEX_NONDIRECTIONAL) && !MATERIAL_SHADINGMODEL_SINGLELAYERWATER)
 
 #define ALLOW_CUBE_REFLECTIONS HQ_REFLECTIONS
+#define MOBILE_PROPAGATE_ALPHA_PREMULTIPLY 2
 
 #include "SHCommon.ush"
 #include "MobileBasePassCommon.ush"
@@ -566,7 +572,7 @@ void Main(
 	, out HALF4_TYPE OutColor		: SV_Target0
 #endif
 #if USE_SCENE_DEPTH_AUX
-	, out float OutSceneDepthAux : SV_Target4
+	, out float OutSceneDepthAux : SV_TargetDepthAux
 #endif
 #if OUTPUT_PIXEL_DEPTH_OFFSET
     , out float OutDepth : SV_Depth
@@ -996,11 +1002,10 @@ void Main(
 		#if !MATERIAL_USE_ALPHA_TO_COVERAGE
 			// Scene color alpha is not used yet so we set it to 1
 			OutColor.a = 1.0;
-
-			#if OUTPUT_MOBILE_HDR 
-				// Store depth in FP16 alpha. This depth value can be fetched during translucency or sampled in post-processing
-				OutColor.a = SvPosition.z;
-			#endif
+			
+		#if USE_SCENE_DEPTH_AUX
+			OutSceneDepthAux = SvPosition.z;
+		#endif
 		#else
 			half MaterialOpacityMask = GetMaterialMaskInputRaw(PixelMaterialInputs);
 			OutColor.a = GetMaterialMask(PixelMaterialInputs) / max(abs(ddx(MaterialOpacityMask)) + abs(ddy(MaterialOpacityMask)), 0.0001f) + 0.5f;
@@ -1018,10 +1023,6 @@ void Main(
 		OutColor.rgb = min(OutColor.rgb, Max10BitsFloat.xxx * 0.5f);
 	#endif
 
-#if USE_SCENE_DEPTH_AUX
-	OutSceneDepthAux = SvPosition.z;
-#endif
-
 #if USE_EDITOR_COMPOSITING && (MOBILE_EMULATION)
 	// Editor primitive depth testing
 	OutColor.a = 1.0;
@@ -1035,7 +1036,9 @@ void Main(
 		OutColor.rgb = sqrt(OutColor.rgb);
 	#endif
 #endif
-
+#if (MATERIALBLENDING_TRANSLUCENT || MATERIALBLENDING_ALPHACOMPOSITE) && (MOBILE_PROPAGATE_ALPHA == MOBILE_PROPAGATE_ALPHA_PREMULTIPLY)
+		OutColor.rgb = OutColor.rgb * OutColor.a;
+#endif
 #if NUM_VIRTUALTEXTURE_SAMPLES || LIGHTMAP_VT_ENABLED
 	FinalizeVirtualTextureFeedback(
 		MaterialParameters.VirtualTextureFeedback,
diff --git a/Engine/Shaders/Private/PostProcessMobile.usf b/Engine/Shaders/Private/PostProcessMobile.usf
--- a/Engine/Shaders/Private/PostProcessMobile.usf
+++ b/Engine/Shaders/Private/PostProcessMobile.usf
@@ -147,14 +147,10 @@ void SunMaskPS_Mobile(
 {
 	half4 SceneColor = SceneColorTexture.Sample(SceneColorSampler, InUVPos.xy);
 
-	#if MOBILE_USEDEPTHTEXTURE
+	#if MOBILE_USEDEPTHTEXTURE 
 		half InDepth = ConvertFromDeviceZ(Texture2DSampleLevel(MobileSceneTextures.SceneDepthTexture, MobileSceneTextures.SceneDepthTextureSampler, InUVPos.xy, 0).r);
 	#else
-		#if MOBILE_DEFERRED_SHADING
-			half InDepth = ConvertFromDeviceZ(Texture2DSampleLevel(MobileSceneTextures.SceneDepthAuxTexture, MobileSceneTextures.SceneDepthAuxTextureSampler, InUVPos.xy, 0).r);
-		#else
-			half InDepth = ConvertFromDeviceZ(SceneColor.a);
-		#endif
+		half InDepth = ConvertFromDeviceZ(Texture2DSampleLevel(MobileSceneTextures.SceneDepthAuxTexture, MobileSceneTextures.SceneDepthAuxTextureSampler, InUVPos.xy, 0).r);
 	#endif
 	
 	#if MOBILE_USESUN
@@ -853,6 +849,7 @@ void IntegrateDOFPS_Mobile(
 	half4 Dof = Texture2DSample(DofBlurTexture, DofBlurSampler, TexCoords.zw);
 	half DofCoc = Texture2DSample(SunShaftAndDofTexture, SunShaftAndDofSampler, TexCoords.xy).r;
 	// Convert alpha back into circle of confusion.
+	OutColor.a = SceneColor.a;
 	SceneColor.a = max(Dof.a, abs(DofCoc * 2.0 - 1.0));
 	// Convert circle of confusion into blend factors.		
 	half2 ScaleBias = CocBlendScaleBias(); // Constant.
@@ -863,7 +860,6 @@ void IntegrateDOFPS_Mobile(
 	OutColor.rgb = lerp(SceneColor.rgb, DofFine.rgb, DofAmountFine);
 	// Blend in coarse DOF.
 	OutColor.rgb = lerp(OutColor.rgb, Dof.rgb, DofAmount);
-	OutColor.a = SceneColor.a;
 }
 
 //
diff --git a/Engine/Shaders/Private/PostProcessTonemap.usf b/Engine/Shaders/Private/PostProcessTonemap.usf
--- a/Engine/Shaders/Private/PostProcessTonemap.usf
+++ b/Engine/Shaders/Private/PostProcessTonemap.usf
@@ -615,7 +615,7 @@ float4 TonemapCommonPS(
 
 	#endif
 	
-	#if (FEATURE_LEVEL <= FEATURE_LEVEL_ES3_1 && USE_FXAA == 0) || POST_PROCESS_ALPHA == 2 || (POST_PROCESS_ALPHA == 1 && DIM_OUTPUT_DEVICE == TONEMAPPER_OUTPUT_NoToneCurve)  || (POST_PROCESS_ALPHA == 1 && DIM_OUTPUT_DEVICE == TONEMAPPER_OUTPUT_WithToneCurve)
+	#if (FEATURE_LEVEL <= FEATURE_LEVEL_ES3_1 && (USE_FXAA == 0 || MOBILE_PROPAGATE_ALPHA > 0)) || POST_PROCESS_ALPHA == 2 || (POST_PROCESS_ALPHA == 1 && DIM_OUTPUT_DEVICE == TONEMAPPER_OUTPUT_NoToneCurve)  || (POST_PROCESS_ALPHA == 1 && DIM_OUTPUT_DEVICE == TONEMAPPER_OUTPUT_WithToneCurve)
 		OutColor.a = SceneColor.a;
 	#endif
 
diff --git a/Engine/Shaders/Private/PostProcessUpscale.usf b/Engine/Shaders/Private/PostProcessUpscale.usf
--- a/Engine/Shaders/Private/PostProcessUpscale.usf
+++ b/Engine/Shaders/Private/PostProcessUpscale.usf
@@ -290,7 +290,7 @@ void MainPS(noperspective float4 UVAndScreenPos : TEXCOORD0, float4 SvPosition :
 	}
 #endif
 
-#if !POST_PROCESS_ALPHA
+#if !POST_PROCESS_ALPHA && !MOBILE_PROPAGATE_ALPHA
 	OutColor.a = 0; // Skip all computations related to alpha
 #endif
 }
\ No newline at end of file
diff --git a/Engine/Shaders/Private/SceneTexturesCommon.ush b/Engine/Shaders/Private/SceneTexturesCommon.ush
--- a/Engine/Shaders/Private/SceneTexturesCommon.ush
+++ b/Engine/Shaders/Private/SceneTexturesCommon.ush
@@ -138,8 +138,8 @@ float LookupDeviceZ( float2 ScreenUV )
 		return Texture2DSample(MobileSceneTextures.SceneDepthTexture, MobileSceneTextures.SceneDepthTextureSampler, ScreenUV).r;
 	#else
 		// SceneDepth texture is not accessible during post-processing as we discard it at the end of mobile BasePass
-		// instead fetch DeviceZ from SceneColor.A
-		return Texture2DSample(MobileSceneTextures.SceneColorTexture, MobileSceneTextures.SceneColorTextureSampler, ScreenUV).a;
+		// instead fetch DeviceZ from SceneDepthAuxTexture
+		return Texture2DSample(MobileSceneTextures.SceneDepthAuxTexture, MobileSceneTextures.SceneDepthAuxTextureSampler, ScreenUV);
 	#endif
 #elif COMPILER_GLSL_ES3_1 && PIXELSHADER
 	return DepthbufferFetchES2();
diff --git a/Engine/Shaders/Public/Platform/Metal/MetalCommon.ush b/Engine/Shaders/Public/Platform/Metal/MetalCommon.ush
--- a/Engine/Shaders/Public/Platform/Metal/MetalCommon.ush
+++ b/Engine/Shaders/Public/Platform/Metal/MetalCommon.ush
@@ -39,13 +39,16 @@
 	// 1 component: R_4
 	Texture2D<float> gl_LastFragDataR_4;
 	#define SubpassFetchR_4() gl_LastFragDataR_4.Load(uint3(0, 0, 0), 0)
+
+	Texture2D<float> gl_LastFragDataR_1;
+	#define SubpassFetchR_1() gl_LastFragDataR_1.Load(uint3(0, 0, 0), 0)
 	//
 	// Rest of SubpassFetch can be found in MetalSubpassSupport.ush
 	//
 	#if MOBILE_DEFERRED_SHADING
 		#define DepthbufferFetchES2() SubpassFetchR_4()
 	#else
-		#define DepthbufferFetchES2() SubpassFetchRGBA_0().w
+		#define DepthbufferFetchES2() SubpassFetchR_1()
 	#endif
 #else
 	#define SubpassFetchRGBA_0() ((float4)0.0)
diff --git a/Engine/Source/Developer/ShaderFormatOpenGL/Private/GlslBackend.cpp b/Engine/Source/Developer/ShaderFormatOpenGL/Private/GlslBackend.cpp
--- a/Engine/Source/Developer/ShaderFormatOpenGL/Private/GlslBackend.cpp
+++ b/Engine/Source/Developer/ShaderFormatOpenGL/Private/GlslBackend.cpp
@@ -1347,7 +1347,7 @@ class ir_gen_glsl_visitor : public ir_visitor
 				ralloc_asprintf_append(buffer, "\n#ifdef GL_ARM_shader_framebuffer_fetch_depth_stencil\n");
 				ralloc_asprintf_append(buffer, "  float DepthbufferFetchES2() { return gl_LastFragDepthARM; }\n");
 				ralloc_asprintf_append(buffer, "#elif defined(GL_EXT_shader_framebuffer_fetch)\n");
-				ralloc_asprintf_append(buffer, "  float DepthbufferFetchES2() { return out_Target0.w; }\n");
+				ralloc_asprintf_append(buffer, "  float DepthbufferFetchES2() { return out_Target1.r; }\n");
 				ralloc_asprintf_append(buffer, "#else\n");
 				ralloc_asprintf_append(buffer, "  float DepthbufferFetchES2() { return 0.0; }\n");
 				ralloc_asprintf_append(buffer, "#endif\n\n");
diff --git a/Engine/Source/Runtime/ApplicationCore/Private/Android/AndroidWindow.cpp b/Engine/Source/Runtime/ApplicationCore/Private/Android/AndroidWindow.cpp
--- a/Engine/Source/Runtime/ApplicationCore/Private/Android/AndroidWindow.cpp
+++ b/Engine/Source/Runtime/ApplicationCore/Private/Android/AndroidWindow.cpp
@@ -14,6 +14,8 @@
 #include "Misc/CommandLine.h"
 #include "HAL/PlatformStackWalk.h"
 
+int GAndroidPropagateAlpha = 0;
+
 struct FCachedWindowRect
 {
 	FCachedWindowRect() : WindowWidth(-1), WindowHeight(-1), WindowInit(false), ContentScaleFactor(-1.0f), Window_EventThread(nullptr)
@@ -89,12 +91,14 @@ static int GAndroidDepthBufferPreference = 0;
 static FVector4 GAndroidPortraitSafezone = FVector4(-1.0f, -1.0f, -1.0f, -1.0f);
 static FVector4 GAndroidLandscapeSafezone = FVector4(-1.0f, -1.0f, -1.0f, -1.0f);
 #if USE_ANDROID_JNI
-JNI_METHOD void Java_com_epicgames_unreal_GameActivity_nativeSetWindowInfo(JNIEnv* jenv, jobject thiz, jboolean bIsPortrait, jint DepthBufferPreference)
+JNI_METHOD void Java_com_epicgames_unreal_GameActivity_nativeSetWindowInfo(JNIEnv* jenv, jobject thiz, jboolean bIsPortrait, jint DepthBufferPreference, jint PropagateAlpha)
 {
 	ClearCachedWindowRects();
 	GAndroidIsPortrait = bIsPortrait == JNI_TRUE;
 	GAndroidDepthBufferPreference = DepthBufferPreference;
+	GAndroidPropagateAlpha = PropagateAlpha;
 	FPlatformMisc::LowLevelOutputDebugStringf(TEXT("App is running in %s\n"), GAndroidIsPortrait ? TEXT("Portrait") : TEXT("Landscape"));
+	FPlatformMisc::LowLevelOutputDebugStringf(TEXT("AndroidPropagateAlpha =  %d\n"), GAndroidPropagateAlpha);
 }
 
 JNI_METHOD void Java_com_epicgames_unreal_GameActivity_nativeSetSurfaceViewInfo(JNIEnv* jenv, jobject thiz, jint width, jint height)
diff --git a/Engine/Source/Runtime/Core/Private/Android/AndroidPlatformMisc.cpp b/Engine/Source/Runtime/Core/Private/Android/AndroidPlatformMisc.cpp
--- a/Engine/Source/Runtime/Core/Private/Android/AndroidPlatformMisc.cpp
+++ b/Engine/Source/Runtime/Core/Private/Android/AndroidPlatformMisc.cpp
@@ -3007,6 +3007,12 @@ bool FAndroidMisc::Expand16BitIndicesTo32BitOnLoad()
 	return  (CVarMaliMidgardIndexingBug.GetValueOnAnyThread() > 0);
 }
 
+int FAndroidMisc::GetMobilePropagateAlphaSetting()
+{
+	extern int GAndroidPropagateAlpha;
+	return GAndroidPropagateAlpha;
+}
+
 TArray<int32> FAndroidMisc::GetSupportedNativeDisplayRefreshRates()
 {
 	TArray<int32> Result;
diff --git a/Engine/Source/Runtime/Core/Private/GenericPlatform/GenericPlatformMisc.cpp b/Engine/Source/Runtime/Core/Private/GenericPlatform/GenericPlatformMisc.cpp
--- a/Engine/Source/Runtime/Core/Private/GenericPlatform/GenericPlatformMisc.cpp
+++ b/Engine/Source/Runtime/Core/Private/GenericPlatform/GenericPlatformMisc.cpp
@@ -1817,8 +1817,12 @@ bool FGenericPlatformMisc::IsPGOEnabled()
 	return PLATFORM_COMPILER_OPTIMIZATION_PG;
 }
 
-bool FGenericPlatformMisc::CheckPersistentDownloadStorageSpaceAvailable( uint64 BytesRequired, bool bAttemptToUseUI )
+int FGenericPlatformMisc::GetMobilePropagateAlphaSetting()
 {
-	return true;
-}
-
+	static int PropagateAlpha = -1;
+	if (PropagateAlpha < 0)
+	{
+		GConfig->GetInt(TEXT("/Script/Engine.RendererSettings"), TEXT("r.Mobile.PropagateAlpha"), PropagateAlpha, GEngineIni);
+	}
+	return PropagateAlpha;
+}
\ No newline at end of file
diff --git a/Engine/Source/Runtime/Core/Public/Android/AndroidPlatformMisc.h b/Engine/Source/Runtime/Core/Public/Android/AndroidPlatformMisc.h
--- a/Engine/Source/Runtime/Core/Public/Android/AndroidPlatformMisc.h
+++ b/Engine/Source/Runtime/Core/Public/Android/AndroidPlatformMisc.h
@@ -302,6 +302,11 @@ public:
 	 */
 	static bool Expand16BitIndicesTo32BitOnLoad();
 
+	/**
+	 * Will return true if we wish to propagate the alpha to the backbuffer
+	 */
+	static int GetMobilePropagateAlphaSetting();
+
 	static bool SupportsBackbufferSampling();
 
 	static void SetMemoryWarningHandler(void (*Handler)(const FGenericMemoryWarningContext& Context));
diff --git a/Engine/Source/Runtime/Core/Public/GenericPlatform/GenericPlatformMisc.h b/Engine/Source/Runtime/Core/Public/GenericPlatform/GenericPlatformMisc.h
--- a/Engine/Source/Runtime/Core/Public/GenericPlatform/GenericPlatformMisc.h
+++ b/Engine/Source/Runtime/Core/Public/GenericPlatform/GenericPlatformMisc.h
@@ -1546,6 +1546,11 @@ public:
 	 */
 	static bool IsPGOEnabled();
 
+	/**
+	 * For mobile devices, this function will return true if we wish to propagate the alpha to the backbuffer
+	 */
+	static int GetMobilePropagateAlphaSetting();
+
 #if !UE_BUILD_SHIPPING
 	/**
 	 * Returns any platform specific warning messages we want printed on screen
diff --git a/Engine/Source/Runtime/Engine/Private/SceneUtils.cpp b/Engine/Source/Runtime/Engine/Private/SceneUtils.cpp
--- a/Engine/Source/Runtime/Engine/Private/SceneUtils.cpp
+++ b/Engine/Source/Runtime/Engine/Private/SceneUtils.cpp
@@ -13,6 +13,16 @@ bool IsMobileHDR()
 	return MobileHDRCvar->GetValueOnAnyThread() == 1;
 }
 
+bool IsMobilePropagateAlphaEnabled(EShaderPlatform Platform)
+{
+	return IsMobilePlatform(Platform) && (FPlatformMisc::GetMobilePropagateAlphaSetting() > 0);
+}
+
+bool IsMobilePropagateAlphaPremultiply(EShaderPlatform Platform)
+{
+	return IsMobilePlatform(Platform) && (FPlatformMisc::GetMobilePropagateAlphaSetting() == 2);
+}
+
 ENGINE_API EMobileHDRMode GetMobileHDRMode()
 {
 	EMobileHDRMode HDRMode = EMobileHDRMode::EnabledFloat16;
diff --git a/Engine/Source/Runtime/Engine/Private/ShaderCompiler/ShaderCompiler.cpp b/Engine/Source/Runtime/Engine/Private/ShaderCompiler/ShaderCompiler.cpp
--- a/Engine/Source/Runtime/Engine/Private/ShaderCompiler/ShaderCompiler.cpp
+++ b/Engine/Source/Runtime/Engine/Private/ShaderCompiler/ShaderCompiler.cpp
@@ -4869,6 +4869,11 @@ void GlobalBeginCompileShader(
 			bool bGLESDeferredShading = Target.Platform == SP_OPENGL_ES3_1_ANDROID;
 			Input.Environment.SetDefine(TEXT("USE_GLES_FBF_DEFERRED"), bGLESDeferredShading ? 1 : 0);
 		}
+
+		static FShaderPlatformCachedIniValue<int32> MobilePropagateAlphaIniValue(TEXT("/Script/Engine.RendererSettings"), TEXT("r.Mobile.PropagateAlpha"));
+		int MobilePropagateAlphaIniValueInt = MobilePropagateAlphaIniValue.Get((EShaderPlatform)ShaderPlatform);
+
+		Input.Environment.SetDefine(TEXT("MOBILE_PROPAGATE_ALPHA"), MobilePropagateAlphaIniValueInt);
 	}
 
 	// Set VR definitions
diff --git a/Engine/Source/Runtime/Engine/Public/SceneUtils.h b/Engine/Source/Runtime/Engine/Public/SceneUtils.h
--- a/Engine/Source/Runtime/Engine/Public/SceneUtils.h
+++ b/Engine/Source/Runtime/Engine/Public/SceneUtils.h
@@ -45,6 +45,10 @@ enum EAntiAliasingMethod
 /** True if HDR is enabled for the mobile renderer. */
 ENGINE_API bool IsMobileHDR();
 
+/** True if Alpha Propagate is enabled for the mobile renderer. */
+ENGINE_API bool IsMobilePropagateAlphaEnabled(EShaderPlatform Platform);
+ENGINE_API bool IsMobilePropagateAlphaPremultiply(EShaderPlatform Platform);
+
 ENGINE_API EMobileHDRMode GetMobileHDRMode();
 
 ENGINE_API bool IsMobileColorsRGB();
diff --git a/Engine/Source/Runtime/OpenGLDrv/Private/Android/AndroidEGL.cpp b/Engine/Source/Runtime/OpenGLDrv/Private/Android/AndroidEGL.cpp
--- a/Engine/Source/Runtime/OpenGLDrv/Private/Android/AndroidEGL.cpp
+++ b/Engine/Source/Runtime/OpenGLDrv/Private/Android/AndroidEGL.cpp
@@ -137,6 +137,10 @@ validConfig (0)
 	int DepthBufferPreference = (int)FAndroidWindow::GetDepthBufferPreference();
 	if (DepthBufferPreference > 0)
 		depthSize = DepthBufferPreference;
+	if (FAndroidMisc::GetMobilePropagateAlphaSetting() > 0)
+	{
+		alphaSize = 8;
+	}
 }
 
 AndroidEGL::AndroidEGL()
diff --git a/Engine/Source/Runtime/OpenGLDrv/Private/OpenGLES.cpp b/Engine/Source/Runtime/OpenGLDrv/Private/OpenGLES.cpp
--- a/Engine/Source/Runtime/OpenGLDrv/Private/OpenGLES.cpp
+++ b/Engine/Source/Runtime/OpenGLDrv/Private/OpenGLES.cpp
@@ -250,7 +250,7 @@ void FOpenGLES::ProcessExtensions(const FString& ExtensionsString)
 		glBufferStorageEXT = (PFNGLBUFFERSTORAGEEXTPROC)((void*)eglGetProcAddress("glBufferStorageEXT"));
 	}
 
-	if (ExtensionsString.Contains(TEXT("GL_EXT_multisampled_render_to_texture")))
+	if (ExtensionsString.Contains(TEXT("GL_EXT_multisampled_render_to_texture2")))
 	{
 		glFramebufferTexture2DMultisampleEXT = (PFNGLFRAMEBUFFERTEXTURE2DMULTISAMPLEEXTPROC)((void*)eglGetProcAddress("glFramebufferTexture2DMultisampleEXT"));
 		glRenderbufferStorageMultisampleEXT = (PFNGLRENDERBUFFERSTORAGEMULTISAMPLEEXTPROC)((void*)eglGetProcAddress("glRenderbufferStorageMultisampleEXT"));
diff --git a/Engine/Source/Runtime/OpenGLDrv/Private/OpenGLRenderTarget.cpp b/Engine/Source/Runtime/OpenGLDrv/Private/OpenGLRenderTarget.cpp
--- a/Engine/Source/Runtime/OpenGLDrv/Private/OpenGLRenderTarget.cpp
+++ b/Engine/Source/Runtime/OpenGLDrv/Private/OpenGLRenderTarget.cpp
@@ -208,7 +208,6 @@ GLuint FOpenGLDynamicRHI::GetOpenGLFramebuffer(uint32 NumSimultaneousRenderTarge
 				if (NumSamplesTileMem > 1 && glFramebufferTexture2DMultisampleEXT)
 				{
 					// GL_EXT_multisampled_render_to_texture requires GL_COLOR_ATTACHMENT0
-					check(RenderTargetIndex == 0);
 					glFramebufferTexture2DMultisampleEXT(GL_FRAMEBUFFER, GL_COLOR_ATTACHMENT0 + RenderTargetIndex, RenderTarget->Target, RenderTarget->GetResource(), MipmapLevels[RenderTargetIndex], NumSamplesTileMem);
 					VERIFY_GL(glFramebufferTexture2DMultisampleEXT);
 				}
@@ -248,7 +247,6 @@ GLuint FOpenGLDynamicRHI::GetOpenGLFramebuffer(uint32 NumSimultaneousRenderTarge
 				if (NumSamplesTileMem > 1 && glFramebufferTexture2DMultisampleEXT)
 				{
 					// GL_EXT_multisampled_render_to_texture requires GL_COLOR_ATTACHMENT0
-					check(RenderTargetIndex == 0);
 					glFramebufferTexture2DMultisampleEXT(GL_FRAMEBUFFER, GL_COLOR_ATTACHMENT0 + RenderTargetIndex, RenderTarget->Target, RenderTarget->GetResource(), MipmapLevels[RenderTargetIndex], NumSamplesTileMem);
 					VERIFY_GL(glFramebufferTexture2DMultisampleEXT);
 				}
diff --git a/Engine/Source/Runtime/OpenGLDrv/Private/OpenGLTexture.cpp b/Engine/Source/Runtime/OpenGLDrv/Private/OpenGLTexture.cpp
--- a/Engine/Source/Runtime/OpenGLDrv/Private/OpenGLTexture.cpp
+++ b/Engine/Source/Runtime/OpenGLDrv/Private/OpenGLTexture.cpp
@@ -356,8 +356,8 @@ FRHITexture* FOpenGLDynamicRHI::CreateOpenGLRHITextureOnly(const uint32 SizeX, c
 
 	// Move NumSamples to on-chip MSAA if supported
 	uint32 NumSamplesTileMem = 1;
-	GLint MaxSamplesTileMem = FOpenGL::GetMaxMSAASamplesTileMem(); /* RHIs which do not support tiled GPU MSAA return 0 */
-	if (MaxSamplesTileMem > 0)
+	GLint MaxSamplesTileMem = FOpenGL::GetMaxMSAASamplesTileMem(); /* RHIs which do not support tiled GPU MSAA return 1 */
+	if (MaxSamplesTileMem > 1)
 	{
 		NumSamplesTileMem = FMath::Min<uint32>(NumSamples, MaxSamplesTileMem);
 		NumSamples = 1;
@@ -689,7 +689,6 @@ void FOpenGLDynamicRHI::InitializeGLTextureInternal(GLuint TextureID, FRHITextur
 	}
 	else
 	{
-		check( FOpenGL::SupportsMultisampledTextures() );
 		check( BulkData == NULL);
 
 		// Try to create an immutable texture and fallback if it fails
diff --git a/Engine/Source/Runtime/OpenGLDrv/Public/OpenGL.h b/Engine/Source/Runtime/OpenGLDrv/Public/OpenGL.h
--- a/Engine/Source/Runtime/OpenGLDrv/Public/OpenGL.h
+++ b/Engine/Source/Runtime/OpenGLDrv/Public/OpenGL.h
@@ -94,7 +94,6 @@ public:
 	static FORCEINLINE bool SupportsDepthStencilReadSurface()			{ return true; }
 	static FORCEINLINE bool SupportsFloatReadSurface()					{ return true; }
 	static FORCEINLINE bool SupportsWideMRT()							{ return true; }
-	static FORCEINLINE bool SupportsMultisampledTextures()				{ return true; }
 	static FORCEINLINE bool SupportsPolygonMode()						{ return true; }
 	static FORCEINLINE bool SupportsTexture3D()							{ return true; }
 	static FORCEINLINE bool SupportsMobileMultiView()					{ return false; }
diff --git a/Engine/Source/Runtime/OpenGLDrv/Public/OpenGLES.h b/Engine/Source/Runtime/OpenGLDrv/Public/OpenGLES.h
--- a/Engine/Source/Runtime/OpenGLDrv/Public/OpenGLES.h
+++ b/Engine/Source/Runtime/OpenGLDrv/Public/OpenGLES.h
@@ -188,7 +188,6 @@ struct FOpenGLES : public FOpenGLBase
 	static FORCEINLINE bool SupportsDepthStencilRead() { return false; }
 	static FORCEINLINE bool SupportsFloatReadSurface() { return SupportsColorBufferHalfFloat(); }
 	static FORCEINLINE bool SupportsWideMRT() { return true; }
-	static FORCEINLINE bool SupportsMultisampledTextures() { return false; }
 	static FORCEINLINE bool SupportsPolygonMode() { return false; }
 	static FORCEINLINE bool SupportsTexture3D() { return true; }
 	static FORCEINLINE bool SupportsMobileMultiView() { return bSupportsMobileMultiView; }
@@ -371,6 +370,12 @@ struct FOpenGLES : public FOpenGLBase
 		glCopyImageSubDataEXT(SrcName, SrcTarget, SrcLevel, SrcX, SrcY, SrcZ, DstName, DstTarget, DstLevel, DstX, DstY, DstZ, Width, Height, Depth);
 	}
 
+	static FORCEINLINE bool TexStorage2DMultisample(GLenum Target, GLsizei Samples, GLint InternalFormat, GLsizei Width, GLsizei Height, GLboolean FixedSampleLocations)
+	{
+		glTexStorage2DMultisample(Target, Samples, InternalFormat, Width, Height, FixedSampleLocations);
+		return true;
+	}
+
 	static FORCEINLINE void ClearBufferfv(GLenum Buffer, GLint DrawBufferIndex, const GLfloat* Value)
 	{
 		glClearBufferfv(Buffer, DrawBufferIndex, Value);
diff --git a/Engine/Source/Runtime/RenderCore/Private/RenderUtils.cpp b/Engine/Source/Runtime/RenderCore/Private/RenderUtils.cpp
--- a/Engine/Source/Runtime/RenderCore/Private/RenderUtils.cpp
+++ b/Engine/Source/Runtime/RenderCore/Private/RenderUtils.cpp
@@ -1080,7 +1080,8 @@ RENDERCORE_API bool IsMobileDeferredShadingEnabled(const FStaticShaderPlatform P
 
 RENDERCORE_API bool MobileRequiresSceneDepthAux(const FStaticShaderPlatform Platform)
 {
-	if (IsMetalMobilePlatform(Platform) && IsMobileDeferredShadingEnabled(Platform))
+	// SceneDepth is used on all mobile platforms when the forward shading is enabled and on IOS with deferred shading.
+	if (!IsMobileDeferredShadingEnabled(Platform) || IsMetalMobilePlatform(Platform))
 	{
 		return true;
 	}
diff --git a/Engine/Source/Runtime/RenderCore/Private/Shader.cpp b/Engine/Source/Runtime/RenderCore/Private/Shader.cpp
--- a/Engine/Source/Runtime/RenderCore/Private/Shader.cpp
+++ b/Engine/Source/Runtime/RenderCore/Private/Shader.cpp
@@ -1530,6 +1530,20 @@ void ShaderMapAppendKeyString(EShaderPlatform Platform, FString& KeyString)
 			KeyString += (CVar && CVar->GetInt() != 0) ? TEXT("_MobileHDR") : TEXT("");
 		}
 
+		{
+			static FShaderPlatformCachedIniValue<int32> MobilePropagateAlphaIniValue(TEXT("/Script/Engine.RendererSettings"), TEXT("r.Mobile.PropagateAlpha"));
+			int MobilePropagateAlphaIniValueInt = MobilePropagateAlphaIniValue.Get((EShaderPlatform)Platform);
+
+			if (MobilePropagateAlphaIniValueInt == 1)
+			{
+				KeyString += TEXT("_MbPropA");
+			}
+			else if (MobilePropagateAlphaIniValueInt == 2)
+			{
+				KeyString += TEXT("_MbPrePropA");
+			}
+		}
+
 		{
 			static IConsoleVariable* CVar = IConsoleManager::Get().FindConsoleVariable(TEXT("r.Mobile.ShadingPath"));
 			KeyString += (CVar && CVar->GetInt() != 0) ? TEXT("_MobDSh") : TEXT("");
diff --git a/Engine/Source/Runtime/Renderer/Private/MobileBasePass.cpp b/Engine/Source/Runtime/Renderer/Private/MobileBasePass.cpp
--- a/Engine/Source/Runtime/Renderer/Private/MobileBasePass.cpp
+++ b/Engine/Source/Runtime/Renderer/Private/MobileBasePass.cpp
@@ -406,7 +406,7 @@ void MobileBasePass::SetTranslucentRenderState(FMeshPassProcessorRenderState& Dr
 	if (Material.GetShadingModels().HasShadingModel(MSM_ThinTranslucent))
 	{
 		// the mobile thin translucent fallback uses a similar mode as BLEND_Translucent, but multiplies color by 1 insead of SrcAlpha.
-		DrawRenderState.SetBlendState(TStaticBlendState<CW_RGB, BO_Add, BF_One, BF_InverseSourceAlpha, BO_Add, BF_Zero, BF_InverseSourceAlpha>::GetRHI());
+		DrawRenderState.SetBlendState(TStaticBlendState<CW_RGBA, BO_Add, BF_One, BF_InverseSourceAlpha, BO_Add, BF_Zero, BF_InverseSourceAlpha>::GetRHI());
 	}
 	else
 	{
@@ -419,7 +419,7 @@ void MobileBasePass::SetTranslucentRenderState(FMeshPassProcessorRenderState& Dr
 			}
 			else
 			{
-				DrawRenderState.SetBlendState(TStaticBlendState<CW_RGB, BO_Add, BF_SourceAlpha, BF_InverseSourceAlpha, BO_Add, BF_Zero, BF_InverseSourceAlpha>::GetRHI());
+				DrawRenderState.SetBlendState(TStaticBlendState<CW_RGBA, BO_Add, BF_SourceAlpha, BF_InverseSourceAlpha, BO_Add, BF_Zero, BF_InverseSourceAlpha>::GetRHI());
 			}
 			break;
 		case BLEND_Additive:
@@ -428,11 +428,11 @@ void MobileBasePass::SetTranslucentRenderState(FMeshPassProcessorRenderState& Dr
 			break;
 		case BLEND_Modulate:
 			// Modulate with the existing scene color
-			DrawRenderState.SetBlendState(TStaticBlendState<CW_RGB, BO_Add, BF_DestColor, BF_Zero>::GetRHI());
+			DrawRenderState.SetBlendState(TStaticBlendState<CW_RGBA, BO_Add, BF_DestColor, BF_Zero>::GetRHI());
 			break;
 		case BLEND_AlphaComposite:
 			// Blend with existing scene color. New color is already pre-multiplied by alpha.
-			DrawRenderState.SetBlendState(TStaticBlendState<CW_RGB, BO_Add, BF_One, BF_InverseSourceAlpha, BO_Add, BF_Zero, BF_InverseSourceAlpha>::GetRHI());
+			DrawRenderState.SetBlendState(TStaticBlendState<CW_RGBA, BO_Add, BF_One, BF_InverseSourceAlpha, BO_Add, BF_Zero, BF_InverseSourceAlpha>::GetRHI());
 			break;
 		case BLEND_AlphaHoldout:
 			// Blend by holding out the matte shape of the source alpha
diff --git a/Engine/Source/Runtime/Renderer/Private/MobileBasePassRendering.cpp b/Engine/Source/Runtime/Renderer/Private/MobileBasePassRendering.cpp
--- a/Engine/Source/Runtime/Renderer/Private/MobileBasePassRendering.cpp
+++ b/Engine/Source/Runtime/Renderer/Private/MobileBasePassRendering.cpp
@@ -43,6 +43,22 @@ static TAutoConsoleVariable<int32> CVarMobileMaxVisibleMovableSpotLightsShadow(
 	TEXT("The max number of visible spotlighs can cast shadow sorted by screen size, should be as less as possible for performance reason"),
 	ECVF_RenderThreadSafe);
 
+static TAutoConsoleVariable<int32> CVarMobileSceneDepthAux(
+	TEXT("r.Mobile.SceneDepthAux"),
+	1,
+	TEXT("1: 16F SceneDepthAux Format")
+	TEXT("2: 32F SceneDepthAux Format"),
+	ECVF_Scalability | ECVF_RenderThreadSafe);
+
+static TAutoConsoleVariable<int32> CVarMobilePropagateAlpha(
+	TEXT("r.Mobile.PropagateAlpha"),
+	0,
+	TEXT("0: Disabled")
+	TEXT("1: Propagate Full Alpha Propagate")
+	TEXT("2: Propagate Full Alpha Propagate + Premultiply"),
+	ECVF_ReadOnly | ECVF_RenderThreadSafe);
+
+
 IMPLEMENT_STATIC_UNIFORM_BUFFER_STRUCT(FMobileBasePassUniformParameters, "MobileBasePass", SceneTextures);
 
 static TAutoConsoleVariable<int32> CVarMobileUseHWsRGBEncoding(
diff --git a/Engine/Source/Runtime/Renderer/Private/MobileDistortionPass.cpp b/Engine/Source/Runtime/Renderer/Private/MobileDistortionPass.cpp
--- a/Engine/Source/Runtime/Renderer/Private/MobileDistortionPass.cpp
+++ b/Engine/Source/Runtime/Renderer/Private/MobileDistortionPass.cpp
@@ -42,12 +42,21 @@ TRDGUniformBufferRef<FMobileDistortionPassUniformParameters> CreateMobileDistort
 {
 	auto* Parameters = GraphBuilder.AllocParameters<FMobileDistortionPassUniformParameters>();
 
-	EMobileSceneTextureSetupMode SetupMode = EMobileSceneTextureSetupMode::SceneColor;
+	EMobileSceneTextureSetupMode SetupMode = EMobileSceneTextureSetupMode::None;
 	if (View.bCustomDepthStencilValid)
 	{
 		SetupMode |= EMobileSceneTextureSetupMode::CustomDepth;
 	}
 
+	if (IsMobileDeferredShadingEnabled(View.GetShaderPlatform()))
+	{
+		SetupMode |= EMobileSceneTextureSetupMode::SceneDepth;
+	}
+	else
+	{
+		SetupMode |= EMobileSceneTextureSetupMode::SceneDepthAux;
+	}
+
 	SetupMobileSceneTextureUniformParameters(GraphBuilder, SetupMode, Parameters->SceneTextures);
 
 	SetupDistortionParams(Parameters->DistortionParams, View);
diff --git a/Engine/Source/Runtime/Renderer/Private/MobileShadingRenderer.cpp b/Engine/Source/Runtime/Renderer/Private/MobileShadingRenderer.cpp
--- a/Engine/Source/Runtime/Renderer/Private/MobileShadingRenderer.cpp
+++ b/Engine/Source/Runtime/Renderer/Private/MobileShadingRenderer.cpp
@@ -391,7 +391,7 @@ void FMobileSceneRenderer::InitViews(FRDGBuilder& GraphBuilder, FSceneTexturesCo
 	PostVisibilityFrameSetup(ILCTaskData);
 
 	const FIntPoint RenderTargetSize = (ViewFamily.RenderTarget->GetRenderTargetTexture().IsValid()) ? ViewFamily.RenderTarget->GetRenderTargetTexture()->GetSizeXY() : ViewFamily.RenderTarget->GetSizeXY();
-	const bool bRequiresUpscale = ((int32)RenderTargetSize.X > FamilySize.X || (int32)RenderTargetSize.Y > FamilySize.Y);
+	const bool bRequiresUpscale = ((int32)RenderTargetSize.X > FamilySize.X || (int32)RenderTargetSize.Y > FamilySize.Y) || IsMobilePropagateAlphaEnabled(ViewFamily.GetShaderPlatform());
 	// ES requires that the back buffer and depth match dimensions.
 	// For the most part this is not the case when using scene captures. Thus scene captures always render to scene color target.
 	const bool bStereoRenderingAndHMD = ViewFamily.EngineShowFlags.StereoRendering && ViewFamily.EngineShowFlags.HMDDistortion;
@@ -936,6 +936,7 @@ void FMobileSceneRenderer::RenderForward(FRDGBuilder& GraphBuilder, FRDGTextureR
 
 	FRenderTargetBindingSlots BasePassRenderTargets;
 	BasePassRenderTargets[0] = FRenderTargetBinding(SceneColor, SceneColorResolve, ERenderTargetLoadAction::EClear);
+	BasePassRenderTargets[1] = FRenderTargetBinding(SceneTextures.DepthAux.Target, SceneTextures.DepthAux.Resolve, ERenderTargetLoadAction::EClear);
 	BasePassRenderTargets.DepthStencil = FDepthStencilBinding(SceneDepth, bIsFullDepthPrepassEnabled ? ERenderTargetLoadAction::ELoad: ERenderTargetLoadAction::EClear, FExclusiveDepthStencil::DepthWrite_StencilWrite);
 	BasePassRenderTargets.ShadingRateTexture = (!MainView.bIsSceneCapture && !MainView.bIsReflectionCapture && ShadingRateTarget.IsValid()) ? RegisterExternalTexture(GraphBuilder, ShadingRateTarget->GetRenderTargetItem().ShaderResourceTexture, TEXT("ShadingRateTexture")) : nullptr;
 	BasePassRenderTargets.SubpassHint = ESubpassHint::DepthReadSubpass;
@@ -1087,6 +1088,7 @@ void FMobileSceneRenderer::RenderForwardMultiPass(FRDGBuilder& GraphBuilder, FMo
 	}
 
 	BasePassRenderTargets[0].SetLoadAction(ERenderTargetLoadAction::ELoad);
+	BasePassRenderTargets[1].SetLoadAction(ERenderTargetLoadAction::ELoad);
 	BasePassRenderTargets.DepthStencil.SetDepthLoadAction(ERenderTargetLoadAction::ELoad);
 	BasePassRenderTargets.DepthStencil.SetStencilLoadAction(ERenderTargetLoadAction::ELoad);
 	BasePassRenderTargets.DepthStencil.SetDepthStencilAccess(ExclusiveDepthStencil);
@@ -1220,7 +1222,7 @@ void FMobileSceneRenderer::RenderDeferred(FRDGBuilder& GraphBuilder, const FSort
 		ColorTargets.Add(SceneTextures.GBufferC);
 		if(MobileRequiresSceneDepthAux(ShaderPlatform))
 		{
-			ColorTargets.Add(SceneTextures.DepthAux);
+			ColorTargets.Add(SceneTextures.DepthAux.Target);
 		}
 	}
 
diff --git a/Engine/Source/Runtime/Renderer/Private/PostProcess/PostProcessMobile.cpp b/Engine/Source/Runtime/Renderer/Private/PostProcess/PostProcessMobile.cpp
--- a/Engine/Source/Runtime/Renderer/Private/PostProcess/PostProcessMobile.cpp
+++ b/Engine/Source/Runtime/Renderer/Private/PostProcess/PostProcessMobile.cpp
@@ -732,7 +732,14 @@ FMobileSunMaskOutputs AddMobileSunMaskPass(FRDGBuilder& GraphBuilder, const FVie
 
 		if (!Inputs.bUseDepthTexture)
 		{
-			OutputDesc.Format = PF_FloatR11G11B10;
+			if (IsMobilePropagateAlphaEnabled(View.GetShaderPlatform()))
+			{
+				OutputDesc.Format = PF_FloatRGBA;
+			}
+			else
+			{
+				OutputDesc.Format = PF_FloatR11G11B10;
+			}
 			SceneColorOutput = FScreenPassRenderTarget(GraphBuilder.CreateTexture(OutputDesc, TEXT("SceneColor")), Inputs.SceneColor.ViewRect, ERenderTargetLoadAction::ENoAction);
 		}
 
diff --git a/Engine/Source/Runtime/Renderer/Private/PostProcess/PostProcessing.cpp b/Engine/Source/Runtime/Renderer/Private/PostProcess/PostProcessing.cpp
--- a/Engine/Source/Runtime/Renderer/Private/PostProcess/PostProcessing.cpp
+++ b/Engine/Source/Runtime/Renderer/Private/PostProcess/PostProcessing.cpp
@@ -1464,6 +1464,7 @@ void AddMobilePostProcessingPasses(FRDGBuilder& GraphBuilder, FScene* Scene, con
 
 	// Scene color is updated incrementally through the post process pipeline.
 	FScreenPassTexture SceneColor((*Inputs.SceneTextures)->SceneColorTexture, FinalOutputViewRect);
+	FScreenPassTexture SceneDepthAux((*Inputs.SceneTextures)->SceneDepthAuxTexture, FinalOutputViewRect);
 
 	// Default the new eye adaptation to the last one in case it's not generated this frame.
 	const FEyeAdaptationParameters EyeAdaptationParameters = GetEyeAdaptationParameters(View, ERHIFeatureLevel::ES3_1);
@@ -1560,7 +1561,7 @@ void AddMobilePostProcessingPasses(FRDGBuilder& GraphBuilder, FScene* Scene, con
 	static const auto VarTonemapperUpscale = IConsoleManager::Get().FindTConsoleVariableDataInt(TEXT("r.MobileTonemapperUpscale"));
 	bool bDisableUpscaleInTonemapper = !VarTonemapperUpscale || VarTonemapperUpscale->GetValueOnRenderThread() == 0;
 
-	bool bShouldPrimaryUpscale = (View.PrimaryScreenPercentageMethod == EPrimaryScreenPercentageMethod::SpatialUpscale && View.UnscaledViewRect != View.ViewRect);
+	bool bShouldPrimaryUpscale = IsMobilePropagateAlphaEnabled(View.GetShaderPlatform()) || (View.PrimaryScreenPercentageMethod == EPrimaryScreenPercentageMethod::SpatialUpscale && View.UnscaledViewRect != View.ViewRect);
 
 	PassSequence.SetEnabled(EPass::Tonemap, bUseToneMapper);
 	PassSequence.SetEnabled(EPass::HighResolutionScreenshotMask, bUseHighResolutionScreenshotMask);
@@ -1668,7 +1669,7 @@ void AddMobilePostProcessingPasses(FRDGBuilder& GraphBuilder, FScene* Scene, con
 		if (PassSequence.IsEnabled(EPass::SunMask))
 		{
 			PassSequence.AcceptPass(EPass::SunMask);
-			bool bUseDepthTexture = SceneColor.Texture->Desc.Format == PF_FloatR11G11B10;
+			bool bUseDepthTexture = SceneColor.Texture->Desc.Format == PF_FloatR11G11B10 && !SceneDepthAux.IsValid();
 
 			FMobileSunMaskInputs SunMaskInputs;
 			SunMaskInputs.bUseDepthTexture = bUseDepthTexture;
@@ -2073,7 +2074,7 @@ void AddMobilePostProcessingPasses(FRDGBuilder& GraphBuilder, FScene* Scene, con
 		TonemapperInputs.Bloom = BloomOutput;
 		TonemapperInputs.EyeAdaptationTexture = nullptr;
 		TonemapperInputs.ColorGradingTexture = ColorGradingTexture;
-		TonemapperInputs.bWriteAlphaChannel = View.AntiAliasingMethod == AAM_FXAA || IsPostProcessingWithAlphaChannelSupported() || bUseMobileDof;
+		TonemapperInputs.bWriteAlphaChannel = View.AntiAliasingMethod == AAM_FXAA || IsPostProcessingWithAlphaChannelSupported() || bUseMobileDof || IsMobilePropagateAlphaEnabled(View.GetShaderPlatform());
 		TonemapperInputs.bFlipYAxis = RHINeedsToSwitchVerticalAxis(View.GetShaderPlatform()) && !PassSequence.IsEnabled(EPass::PostProcessMaterialAfterTonemapping);
 		TonemapperInputs.bOutputInHDR = bHDRTonemapperOutput;
 		TonemapperInputs.bGammaOnly = bDoGammaOnly;
diff --git a/Engine/Source/Runtime/Renderer/Private/SceneTextures.cpp b/Engine/Source/Runtime/Renderer/Private/SceneTextures.cpp
--- a/Engine/Source/Runtime/Renderer/Private/SceneTextures.cpp
+++ b/Engine/Source/Runtime/Renderer/Private/SceneTextures.cpp
@@ -107,15 +107,14 @@ static EPixelFormat GetGBufferFFormat()
 	return NormalGBufferFormat;
 }
 
-static EPixelFormat GetMobileSceneColorFormat()
+static EPixelFormat GetMobileSceneColorFormat(EShaderPlatform ShaderPlatform)
 {
 	const EPixelFormat DefaultLowPrecisionFormat = IHeadMountedDisplayModule::IsAvailable() && IHeadMountedDisplayModule::Get().IsStandaloneStereoOnlyDevice()
 		? PF_R8G8B8A8 : PF_B8G8R8A8;
-	EPixelFormat DefaultColorFormat = (!IsMobileHDR() || !GSupportsRenderTargetFormat_PF_FloatRGBA) ? DefaultLowPrecisionFormat : PF_FloatRGBA;
-	if (IsMobileDeferredShadingEnabled(GMaxRHIShaderPlatform))
-	{
-		DefaultColorFormat = PF_FloatR11G11B10;
-	}
+	const EPixelFormat DefaultPrecisionFormat = IsMobilePropagateAlphaEnabled(ShaderPlatform) ? PF_FloatRGBA : PF_FloatR11G11B10;
+
+	EPixelFormat DefaultColorFormat = (!IsMobileHDR() || !GSupportsRenderTargetFormat_PF_FloatRGBA) ? DefaultLowPrecisionFormat : DefaultPrecisionFormat;
+
 	check(GPixelFormats[DefaultColorFormat].Supported);
 
 	EPixelFormat Format = DefaultColorFormat;
@@ -184,6 +183,27 @@ static EPixelFormat GetSceneColorFormat(const FSceneViewFamily& ViewFamily)
 	return Format;
 }
 
+inline EPixelFormat GetMobileSceneDepthAuxPixelFormat(EShaderPlatform ShaderPlatform)
+{
+	if (IsMobileDeferredShadingEnabled(ShaderPlatform))
+	{
+		return PF_R32_FLOAT;
+	}
+
+	static const auto CVar = IConsoleManager::Get().FindTConsoleVariableDataInt(TEXT("r.Mobile.SceneDepthAux"));
+	EPixelFormat Format = PF_R16F;
+	switch (CVar->GetValueOnAnyThread())
+	{
+	case 1:
+		Format =  PF_R16F;
+		break;
+	case 2:
+		Format = PF_R32_FLOAT;
+		break;
+	}
+	return Format;
+}
+
 static uint32 GetEditorPrimitiveNumSamples(ERHIFeatureLevel::Type FeatureLevel)
 {
 	uint32 SampleCount = 1;
@@ -515,7 +535,7 @@ FSceneTexturesConfig FSceneTexturesConfig::Create(const FSceneViewFamily& ViewFa
 
 	case EShadingPath::Mobile:
 	{
-		Config.ColorFormat = GetMobileSceneColorFormat();
+		Config.ColorFormat = GetMobileSceneColorFormat(Config.ShaderPlatform);
 
 		// On mobile the scene depth is calculated from the alpha component of the scene color
 		// Use FarPlane for alpha to ensure un-rendered pixels have max depth...
@@ -760,15 +780,25 @@ FSceneTextures& FSceneTextures::Create(FRDGBuilder& GraphBuilder, const FSceneTe
 			const FRDGTextureDesc Desc = FRDGTextureDesc::Create2D(Config.Extent, GetGBufferFFormat(), FClearValueBinding({ 0.5f, 0.5f, 0.5f, 0.5f }), TexCreate_RenderTargetable | TexCreate_ShaderResource | FlagsToAdd | GFastVRamConfig.GBufferF);
 			SceneTextures.GBufferF = GraphBuilder.CreateTexture(Desc, TEXT("GBufferF"));
 		}
-		else if (MobileRequiresSceneDepthAux(Config.ShaderPlatform))
+	}
+
+
+	if(Config.ShadingPath == EShadingPath::Mobile && MobileRequiresSceneDepthAux(Config.ShaderPlatform))
+	{
+		const float FarDepth = (float)ERHIZBuffer::FarPlane;
+		const FLinearColor FarDepthColor(FarDepth, FarDepth, FarDepth, FarDepth);
+		ETextureCreateFlags FlagsToAdd = IsMobileDeferredShadingEnabled(Config.ShaderPlatform)? TexCreate_Memoryless : TexCreate_None;
+		FRDGTextureDesc Desc = FRDGTextureDesc::Create2D(Config.Extent, GetMobileSceneDepthAuxPixelFormat(Config.ShaderPlatform), FClearValueBinding(FarDepthColor), TexCreate_RenderTargetable | TexCreate_ShaderResource | TexCreate_InputAttachmentRead| FlagsToAdd);
+		Desc.NumSamples = Config.NumSamples;
+		SceneTextures.DepthAux = GraphBuilder.CreateTexture(Desc, TEXT("SceneDepthAux"));
+		
+		if (Desc.NumSamples > 1)
 		{
-			const float FarDepth = (float)ERHIZBuffer::FarPlane;
-			const FLinearColor FarDepthColor(FarDepth, FarDepth, FarDepth, FarDepth);
-			const FRDGTextureDesc Desc = FRDGTextureDesc::Create2D(Config.Extent, PF_R32_FLOAT, FClearValueBinding(FarDepthColor), TexCreate_RenderTargetable | TexCreate_ShaderResource | FlagsToAdd);
-			SceneTextures.DepthAux = GraphBuilder.CreateTexture(Desc, TEXT("SceneDepthAux"));
+			Desc.NumSamples = 1;
+			Desc.Flags = TexCreate_ResolveTargetable | TexCreate_ShaderResource;
+			SceneTextures.DepthAux.Resolve = GraphBuilder.CreateTexture(Desc, TEXT("SceneDepthAux"));
 		}
 	}
-
 #if WITH_EDITOR
 	{
 		const FRDGTextureDesc ColorDesc(FRDGTextureDesc::Create2D(Config.Extent, PF_B8G8R8A8, FClearValueBinding::Transparent, TexCreate_ShaderResource | TexCreate_RenderTargetable, 1, Config.EditorPrimitiveNumSamples));
@@ -1116,10 +1146,13 @@ void SetupMobileSceneTextureUniformParameters(
 			{
 				SceneTextureParameters.GBufferDTexture = SceneTextures->GBufferD;
 			}
+		}
 
-			if (HasBeenProduced(SceneTextures->DepthAux))
+		if (EnumHasAnyFlags(SetupMode, EMobileSceneTextureSetupMode::SceneDepthAux))
+		{
+			if (HasBeenProduced(SceneTextures->DepthAux.Resolve))
 			{
-				SceneTextureParameters.SceneDepthAuxTexture = SceneTextures->DepthAux;
+				SceneTextureParameters.SceneDepthAuxTexture = SceneTextures->DepthAux.Resolve;
 			}
 		}
 
diff --git a/Engine/Source/Runtime/Renderer/Private/SceneTextures.h b/Engine/Source/Runtime/Renderer/Private/SceneTextures.h
--- a/Engine/Source/Runtime/Renderer/Private/SceneTextures.h
+++ b/Engine/Source/Runtime/Renderer/Private/SceneTextures.h
@@ -161,8 +161,8 @@ struct RENDERER_API FSceneTextures : public FMinimalSceneTextures
 	FRDGTextureRef GBufferE{};
 	FRDGTextureRef GBufferF{};
 
-	// (Mobile Deferred) Additional GBuffer texture used by mobile
-	FRDGTextureRef DepthAux{};
+	// Additional Buffer texture used by mobile
+	FRDGTextureMSAA DepthAux{};
 
 	// Texture containing dynamic motion vectors. Can be bound by the base pass or its own velocity pass.
 	FRDGTextureRef Velocity{};
