Fix an issue where r.FastVRam.PostProcessMaterial has no effect

diff --git a/Engine/Source/Runtime/Renderer/Private/PostProcess/PostProcessMaterial.cpp b/Engine/Source/Runtime/Renderer/Private/PostProcess/PostProcessMaterial.cpp
--- a/Engine/Source/Runtime/Renderer/Private/PostProcess/PostProcessMaterial.cpp
+++ b/Engine/Source/Runtime/Renderer/Private/PostProcess/PostProcessMaterial.cpp
@@ -444,6 +444,7 @@ FScreenPassTexture AddPostProcessMaterialPass(
 				OutputDesc.Format = Inputs.OutputFormat;
 			}
 			OutputDesc.ClearValue = FClearValueBinding(FLinearColor::Black);
+			OutputDesc.Flags &= (~ETextureCreateFlags::FastVRAM);
 			OutputDesc.Flags |= GFastVRamConfig.PostProcessMaterial;
 
 			Output = FScreenPassRenderTarget(GraphBuilder.CreateTexture(OutputDesc, TEXT("PostProcessMaterial")), SceneColor.ViewRect, View.GetOverwriteLoadAction());
