RenderRayTracingSkyLight() out parameters are now initialized to safe placeholder values if early out path is taken (when SetupSkyLightParameters() returns false).

diff --git a/Engine/Source/Runtime/Renderer/Private/RayTracing/RaytracingSkylight.cpp b/Engine/Source/Runtime/Renderer/Private/RayTracing/RaytracingSkylight.cpp
--- a/Engine/Source/Runtime/Renderer/Private/RayTracing/RaytracingSkylight.cpp
+++ b/Engine/Source/Runtime/Renderer/Private/RayTracing/RaytracingSkylight.cpp
@@ -376,6 +376,8 @@ void FDeferredShadingSceneRenderer::RenderRayTracingSkyLight(
 	FSkyLightData SkyLightData;
 	if (!SetupSkyLightParameters(GraphBuilder, Scene, Views[0], bShouldRenderRayTracingSkyLight, &SkylightParameters, &SkyLightData))
 	{
+		OutSkyLightTexture = GraphBuilder.RegisterExternalTexture(GSystemTextures.BlackDummy);
+		OutHitDistanceTexture = GraphBuilder.RegisterExternalTexture(GSystemTextures.BlackDummy);
 		return;
 	}
 
