Updated most of distance field ambient occlusion codepath to use RDG.
- Resource state transitions handled by RDG.
- FDistanceFieldObjectBufferResource replaced by FDistanceFieldCulledObjectBufferParameters to use RDG.
- FTileIntersectionResources replaced by FTileIntersectionParameters to use RDG.
- FAOScreenGridResources replaced by FAOScreenGridParameters to use RDG.
- Don't store DFAO resources in View.State.
- Split FDeferredShadingSceneRenderer::RenderDistanceFieldAOScreenGrid() into multiple passes.
- Split BuildTileObjectLists() into multiple passes.

diff --git a/Engine/Source/Runtime/Renderer/Private/DeferredShadingRenderer.h b/Engine/Source/Runtime/Renderer/Private/DeferredShadingRenderer.h
--- a/Engine/Source/Runtime/Renderer/Private/DeferredShadingRenderer.h
+++ b/Engine/Source/Runtime/Renderer/Private/DeferredShadingRenderer.h
@@ -25,6 +25,8 @@
 enum class ERayTracingPrimaryRaysFlag : uint32;
 
 class FSceneTextureParameters;
+class FDistanceFieldCulledObjectBufferParameters;
+class FTileIntersectionParameters;
 class FDistanceFieldAOParameters;
 class UStaticMeshComponent;
 class FExponentialHeightFogSceneInfo;
@@ -485,6 +487,9 @@ private:
 		FRDGBuilder& GraphBuilder,
 		const FSceneTextures& SceneTextures,
 		const FViewInfo& View,
+		const FDistanceFieldCulledObjectBufferParameters& CulledObjectBufferParameters,
+		FRDGBufferRef ObjectTilesIndirectArguments,
+		const FTileIntersectionParameters& TileIntersectionParameters,
 		const FDistanceFieldAOParameters& Parameters,
 		FRDGTextureRef DistanceFieldNormal,
 		FRDGTextureRef& OutDynamicBentNormalAO);
diff --git a/Engine/Source/Runtime/Renderer/Private/DistanceFieldAmbientOcclusion.cpp b/Engine/Source/Runtime/Renderer/Private/DistanceFieldAmbientOcclusion.cpp
--- a/Engine/Source/Runtime/Renderer/Private/DistanceFieldAmbientOcclusion.cpp
+++ b/Engine/Source/Runtime/Renderer/Private/DistanceFieldAmbientOcclusion.cpp
@@ -24,7 +24,6 @@
 IMPLEMENT_TYPE_LAYOUT(FAOParameters);
 IMPLEMENT_TYPE_LAYOUT(FDFAOUpsampleParameters);
 IMPLEMENT_TYPE_LAYOUT(FScreenGridParameters);
-IMPLEMENT_TYPE_LAYOUT(FTileIntersectionParameters);
 
 int32 GDistanceFieldAO = 1;
 FAutoConsoleVariableRef CVarDistanceFieldAO(
@@ -612,6 +611,49 @@ void GenerateBestSpacedVectors()
 	}
 }
 
+void AllocateTileIntersectionBuffers(
+	FRDGBuilder& GraphBuilder,
+	FIntPoint TileListGroupSize,
+	uint32 MaxSceneObjects,
+	bool bAllow16BitIndices,
+	FRDGBufferRef& OutObjectTilesIndirectArguments,
+	FTileIntersectionParameters& OutParameters)
+{
+	// Can only use 16 bit for CulledTileDataArray if few enough objects and tiles
+	const bool b16BitObjectIndices = MaxSceneObjects < (1 << 16);
+	const bool b16BitCulledTileIndexBuffer = bAllow16BitIndices && b16BitObjectIndices && (TileListGroupSize.X * TileListGroupSize.Y < (1 << 16));
+
+	FRDGBufferRef TileConeAxisAndCos = GraphBuilder.CreateBuffer(FRDGBufferDesc::CreateBufferDesc(sizeof(FVector4), TileListGroupSize.X * TileListGroupSize.Y), TEXT("TileConeAxisAndCos"));
+	OutParameters.RWTileConeAxisAndCos = GraphBuilder.CreateUAV(TileConeAxisAndCos, PF_A32B32G32R32F);
+	OutParameters.TileConeAxisAndCos = GraphBuilder.CreateSRV(TileConeAxisAndCos, PF_A32B32G32R32F);
+
+	FRDGBufferRef TileConeDepthRanges = GraphBuilder.CreateBuffer(FRDGBufferDesc::CreateBufferDesc(sizeof(FVector4), TileListGroupSize.X * TileListGroupSize.Y), TEXT("TileConeDepthRanges"));
+	OutParameters.RWTileConeDepthRanges = GraphBuilder.CreateUAV(TileConeDepthRanges, PF_A32B32G32R32F);
+	OutParameters.TileConeDepthRanges = GraphBuilder.CreateSRV(TileConeDepthRanges, PF_A32B32G32R32F);
+
+	FRDGBufferRef NumCulledTilesArray = GraphBuilder.CreateBuffer(FRDGBufferDesc::CreateBufferDesc(sizeof(uint32), MaxSceneObjects), TEXT("NumCulledTilesArray"));
+	OutParameters.RWNumCulledTilesArray = GraphBuilder.CreateUAV(NumCulledTilesArray, PF_R32_UINT);
+	OutParameters.NumCulledTilesArray = GraphBuilder.CreateSRV(NumCulledTilesArray, PF_R32_UINT);
+
+	FRDGBufferRef CulledTilesStartOffsetArray = GraphBuilder.CreateBuffer(FRDGBufferDesc::CreateBufferDesc(sizeof(uint32), MaxSceneObjects), TEXT("CulledTilesStartOffsetArray"));
+	OutParameters.RWCulledTilesStartOffsetArray = GraphBuilder.CreateUAV(CulledTilesStartOffsetArray, PF_R32_UINT);
+	OutParameters.CulledTilesStartOffsetArray = GraphBuilder.CreateSRV(CulledTilesStartOffsetArray, PF_R32_UINT);
+
+	extern int32 GAverageDistanceFieldObjectsPerCullTile;
+
+	FRDGBufferRef CulledTileDataArray = GraphBuilder.CreateBuffer(
+		FRDGBufferDesc::CreateBufferDesc(b16BitCulledTileIndexBuffer ? sizeof(uint16) : sizeof(uint32), GAverageDistanceFieldObjectsPerCullTile * TileListGroupSize.X * TileListGroupSize.Y * CulledTileDataStride),
+		TEXT("CulledTileDataArray"));
+	OutParameters.RWCulledTileDataArray = GraphBuilder.CreateUAV(CulledTileDataArray, b16BitCulledTileIndexBuffer ? PF_R16_UINT : PF_R32_UINT);
+	OutParameters.CulledTileDataArray = GraphBuilder.CreateSRV(CulledTileDataArray, b16BitCulledTileIndexBuffer ? PF_R16_UINT : PF_R32_UINT);
+
+	OutObjectTilesIndirectArguments = GraphBuilder.CreateBuffer(FRDGBufferDesc::CreateIndirectDesc<FRHIDispatchIndirectParameters>(), TEXT("ObjectTilesIndirectArguments"));
+	OutParameters.RWObjectTilesIndirectArguments = GraphBuilder.CreateUAV(OutObjectTilesIndirectArguments, PF_R32_UINT);
+	OutParameters.ObjectTilesIndirectArguments = GraphBuilder.CreateSRV(OutObjectTilesIndirectArguments, PF_R32_UINT);
+
+	OutParameters.TileListGroupSize = TileListGroupSize;
+}
+
 void ListDistanceFieldLightingMemory(const FViewInfo& View, FSceneRenderer& SceneRenderer)
 {
 #if !NO_LOGGING
@@ -622,26 +664,6 @@ void ListDistanceFieldLightingMemory(const FViewInfo& View, FSceneRenderer& Scen
 	{
 		UE_LOG(LogRenderer, Log, TEXT("   Scene Object data %.3fMb"), Scene->DistanceFieldSceneData.GetCurrentObjectBuffers()->GetSizeBytes() / 1024.0f / 1024.0f);
 	}
-	
-	UE_LOG(LogRenderer, Log, TEXT(""));
-	UE_LOG(LogRenderer, Log, TEXT("Distance Field AO"));
-	
-	UE_LOG(LogRenderer, Log, TEXT("   Culled objects %.3fMb"), GAOCulledObjectBuffers.Buffers.GetSizeBytes() / 1024.0f / 1024.0f);
-
-	FTileIntersectionResources* TileIntersectionResources = ((FSceneViewState*)View.State)->AOTileIntersectionResources;
-
-	if (TileIntersectionResources)
-	{
-		UE_LOG(LogRenderer, Log, TEXT("   Tile Culled objects %.3fMb"), TileIntersectionResources->GetSizeBytes() / 1024.0f / 1024.0f);
-	}
-	
-	FAOScreenGridResources* ScreenGridResources = ((FSceneViewState*)View.State)->AOScreenGridResources;
-
-	if (ScreenGridResources)
-	{
-		UE_LOG(LogRenderer, Log, TEXT("   Screen grid temporaries %.3fMb"), ScreenGridResources->GetSizeBytesForAO() / 1024.0f / 1024.0f);
-	}
-
 #endif // !NO_LOGGING
 }
 
@@ -742,8 +764,6 @@ bool FDeferredShadingSceneRenderer::ShouldRenderDistanceFieldLighting() const
 
 	return SupportsDistanceFieldAO(View.GetFeatureLevel(), View.GetShaderPlatform())
 		&& Views.Num() == 1
-		// ViewState is used to cache tile intersection resources which have to be sized based on the view
-		&& View.State
 		&& View.IsPerspectiveProjection()
 		&& Scene->DistanceFieldSceneData.NumObjectsInBuffer;
 }
@@ -788,9 +808,29 @@ void FDeferredShadingSceneRenderer::RenderDistanceFieldLighting(
 		DistanceFieldNormal = GraphBuilder.CreateTexture(Desc, TEXT("DistanceFieldNormal"));
 	}
 
+	const FIntPoint TileListGroupSize = GetTileListGroupSizeForView(View);
+	const int32 MaxSceneObjects = Scene->DistanceFieldSceneData.NumObjectsInBuffer;
+	const bool bAllow16BitIndices = !IsMetalPlatform(GShaderPlatformForFeatureLevel[View.FeatureLevel]);
+
+	FRDGBufferRef ObjectTilesIndirectArguments = nullptr;
+	FTileIntersectionParameters TileIntersectionParameters;
+
+	AllocateTileIntersectionBuffers(GraphBuilder, TileListGroupSize, MaxSceneObjects, bAllow16BitIndices, ObjectTilesIndirectArguments, TileIntersectionParameters);
+
+	FRDGBufferRef ObjectIndirectArguments = nullptr;
+	FDistanceFieldCulledObjectBufferParameters CulledObjectBufferParameters;
+
 	if (UseAOObjectDistanceField())
 	{
-		CullObjectsToView(GraphBuilder, Scene, View, Parameters, GAOCulledObjectBuffers);
+		AllocateDistanceFieldCulledObjectBuffers(
+			GraphBuilder,
+			false,
+			FMath::DivideAndRoundUp(MaxSceneObjects, 256) * 256,
+			DFPT_SignedDistanceField,
+			ObjectIndirectArguments,
+			CulledObjectBufferParameters);
+
+		CullObjectsToView(GraphBuilder, Scene, View, Parameters, CulledObjectBufferParameters);
 	}
 
 	ComputeDistanceFieldNormal(GraphBuilder, Views, SceneTextures.UniformBuffer, DistanceFieldNormal, Parameters);
@@ -798,7 +838,8 @@ void FDeferredShadingSceneRenderer::RenderDistanceFieldLighting(
 	// Intersect objects with screen tiles, build lists
 	if (UseAOObjectDistanceField())
 	{
-		BuildTileObjectLists(GraphBuilder, Scene, Views, DistanceFieldNormal, Parameters);
+		//@todo - support multiple views - should pass one TileIntersectionParameters per view
+		BuildTileObjectLists(GraphBuilder, Scene, Views, ObjectIndirectArguments, CulledObjectBufferParameters, TileIntersectionParameters, DistanceFieldNormal, Parameters);
 	}
 
 	FRDGTextureRef BentNormalOutput = nullptr;
@@ -807,21 +848,13 @@ void FDeferredShadingSceneRenderer::RenderDistanceFieldLighting(
 		GraphBuilder,
 		SceneTextures,
 		View,
+		CulledObjectBufferParameters,
+		ObjectTilesIndirectArguments,
+		TileIntersectionParameters,
 		Parameters,
 		DistanceFieldNormal,
 		BentNormalOutput);
 
-	if (IsTransientResourceBufferAliasingEnabled() && UseAOObjectDistanceField())
-	{
-		AddPass(GraphBuilder, [&View](FRHICommandListImmediate&)
-		{
-			GAOCulledObjectBuffers.Buffers.DiscardTransientResource();
-
-			FTileIntersectionResources* TileIntersectionResources = ((FSceneViewState*)View.State)->AOTileIntersectionResources;
-			TileIntersectionResources->DiscardTransientResource();
-		});
-	}
-
 	RenderCapsuleShadowsForMovableSkylight(GraphBuilder, SceneTextures.UniformBuffer, BentNormalOutput);
 
 	// Upsample to full resolution, write to output in case of debug AO visualization or scene color modulation (standard upsampling is done later together with sky lighting and reflection environment)
diff --git a/Engine/Source/Runtime/Renderer/Private/DistanceFieldAmbientOcclusion.h b/Engine/Source/Runtime/Renderer/Private/DistanceFieldAmbientOcclusion.h
--- a/Engine/Source/Runtime/Renderer/Private/DistanceFieldAmbientOcclusion.h
+++ b/Engine/Source/Runtime/Renderer/Private/DistanceFieldAmbientOcclusion.h
@@ -38,240 +38,44 @@ public:
 	FDistanceFieldAOParameters(float InOcclusionMaxDistance, float InContrast = 0);
 };
 
-/**  */
-class FTileIntersectionResources : public FRenderResource
-{
-public:
-
-	FRWBuffer TileConeAxisAndCos;
-	FRWBuffer TileConeDepthRanges;
-
-	FRWBuffer NumCulledTilesArray;
-	FRWBuffer CulledTilesStartOffsetArray;
-	FRWBuffer CulledTileDataArray;
-	FRWBuffer ObjectTilesIndirectArguments;
-
-	FIntPoint TileDimensions;
-	int32 MaxSceneObjects;
-	bool bAllow16BitIndices;
-
-	FTileIntersectionResources(bool bInAllow16BitIndices) :
-		MaxSceneObjects(0), bAllow16BitIndices(bInAllow16BitIndices)
-	{}
-
-	bool HasAllocatedEnoughFor(FIntPoint TestTileDimensions, int32 TestMaxSceneObjects) const
-	{
-		return TestTileDimensions == TileDimensions && TestMaxSceneObjects <= MaxSceneObjects;
-	}
+BEGIN_SHADER_PARAMETER_STRUCT(FTileIntersectionParameters, )
+	SHADER_PARAMETER_RDG_BUFFER_UAV(RWBuffer<FVector4>, RWTileConeAxisAndCos)
+	SHADER_PARAMETER_RDG_BUFFER_UAV(RWBuffer<FVector4>, RWTileConeDepthRanges)
 
-	void SetupParameters(FIntPoint InTileDimensions, int32 InMaxSceneObjects)
-	{
-		TileDimensions = InTileDimensions;
-		MaxSceneObjects = InMaxSceneObjects;
-	}
+	SHADER_PARAMETER_RDG_BUFFER_UAV(RWBuffer<uint>, RWNumCulledTilesArray)
+	SHADER_PARAMETER_RDG_BUFFER_UAV(RWBuffer<uint>, RWCulledTilesStartOffsetArray)
+	SHADER_PARAMETER_RDG_BUFFER_UAV(RWBuffer<uint>, RWCulledTileDataArray)
+	SHADER_PARAMETER_RDG_BUFFER_UAV(RWBuffer<uint>, RWObjectTilesIndirectArguments)
 
-	virtual void InitDynamicRHI() override;
+	SHADER_PARAMETER_RDG_BUFFER_SRV(Buffer<FVector4>, TileConeAxisAndCos)
+	SHADER_PARAMETER_RDG_BUFFER_SRV(Buffer<FVector4>, TileConeDepthRanges)
 
-	
+	SHADER_PARAMETER_RDG_BUFFER_SRV(Buffer<uint>, NumCulledTilesArray)
+	SHADER_PARAMETER_RDG_BUFFER_SRV(Buffer<uint>, CulledTilesStartOffsetArray)
+	SHADER_PARAMETER_RDG_BUFFER_SRV(Buffer<uint>, CulledTileDataArray)
+	SHADER_PARAMETER_RDG_BUFFER_SRV(Buffer<uint>, ObjectTilesIndirectArguments)
 
-	virtual void ReleaseDynamicRHI() override
-	{
-		TileConeAxisAndCos.Release();
-		TileConeDepthRanges.Release();
-
-		NumCulledTilesArray.Release();
-		CulledTilesStartOffsetArray.Release();
-		CulledTileDataArray.Release();
-		ObjectTilesIndirectArguments.Release();
-	}
-
-	void AcquireTransientResource()
-	{
-		TileConeAxisAndCos.AcquireTransientResource();
-		TileConeDepthRanges.AcquireTransientResource();
-		NumCulledTilesArray.AcquireTransientResource();
-		CulledTilesStartOffsetArray.AcquireTransientResource();
-		CulledTileDataArray.AcquireTransientResource();
-	}
-
-	void DiscardTransientResource()
-	{
-		TileConeAxisAndCos.DiscardTransientResource();
-		TileConeDepthRanges.DiscardTransientResource();
-		NumCulledTilesArray.DiscardTransientResource();
-		CulledTilesStartOffsetArray.DiscardTransientResource();
-		CulledTileDataArray.DiscardTransientResource();
-	}
-
-	size_t GetSizeBytes() const
-	{
-		return TileConeAxisAndCos.NumBytes + TileConeDepthRanges.NumBytes
-			+ NumCulledTilesArray.NumBytes + CulledTilesStartOffsetArray.NumBytes + CulledTileDataArray.NumBytes + ObjectTilesIndirectArguments.NumBytes;
-	}
-};
+	SHADER_PARAMETER(FIntPoint, TileListGroupSize)
+END_SHADER_PARAMETER_STRUCT()
 
 static int32 CulledTileDataStride = 2;
 static int32 ConeTraceObjectsThreadGroupSize = 64;
 
-class FTileIntersectionParameters
+inline void TileIntersectionModifyCompilationEnvironment(EShaderPlatform Platform, FShaderCompilerEnvironment& OutEnvironment)
 {
-	DECLARE_TYPE_LAYOUT(FTileIntersectionParameters, NonVirtual);
-public:
-
-	static void ModifyCompilationEnvironment(EShaderPlatform Platform, FShaderCompilerEnvironment& OutEnvironment)
-	{
-		OutEnvironment.SetDefine(TEXT("CULLED_TILE_DATA_STRIDE"), CulledTileDataStride);
-		extern int32 GDistanceFieldAOTileSizeX;
-		OutEnvironment.SetDefine(TEXT("CULLED_TILE_SIZEX"), GDistanceFieldAOTileSizeX);
-		extern int32 GConeTraceDownsampleFactor;
-		OutEnvironment.SetDefine(TEXT("TRACE_DOWNSAMPLE_FACTOR"), GConeTraceDownsampleFactor);
-		OutEnvironment.SetDefine(TEXT("CONE_TRACE_OBJECTS_THREADGROUP_SIZE"), ConeTraceObjectsThreadGroupSize);
-	}
-
-	void Bind(const FShaderParameterMap& ParameterMap)
-	{
-		TileListGroupSize.Bind(ParameterMap, TEXT("TileListGroupSize"));
-		NumCulledTilesArray.Bind(ParameterMap, TEXT("NumCulledTilesArray"));
-		CulledTilesStartOffsetArray.Bind(ParameterMap, TEXT("CulledTilesStartOffsetArray"));
-		CulledTileDataArray.Bind(ParameterMap, TEXT("CulledTileDataArray"));
-		ObjectTilesIndirectArguments.Bind(ParameterMap, TEXT("ObjectTilesIndirectArguments"));
-	}
-
-	template<typename TParamRef>
-	void Set(FRHICommandList& RHICmdList, const TParamRef& ShaderRHI, const FTileIntersectionResources& TileIntersectionResources)
-	{
-		SetShaderValue(RHICmdList, ShaderRHI, TileListGroupSize, TileIntersectionResources.TileDimensions);
-
-		NumCulledTilesArray.SetBuffer(RHICmdList, ShaderRHI, TileIntersectionResources.NumCulledTilesArray);
-		CulledTilesStartOffsetArray.SetBuffer(RHICmdList, ShaderRHI, TileIntersectionResources.CulledTilesStartOffsetArray);
-		CulledTileDataArray.SetBuffer(RHICmdList, ShaderRHI, TileIntersectionResources.CulledTileDataArray);
-		ObjectTilesIndirectArguments.SetBuffer(RHICmdList, ShaderRHI, TileIntersectionResources.ObjectTilesIndirectArguments);
-	}
-
-	void GetUAVs(FTileIntersectionResources& TileIntersectionResources, TArray<FRHIUnorderedAccessView*>& UAVs)
-	{
-		uint32 MaxIndex = 0;
-		
-		MaxIndex = FMath::Max(MaxIndex, NumCulledTilesArray.GetUAVIndex());
-		MaxIndex = FMath::Max(MaxIndex, CulledTilesStartOffsetArray.GetUAVIndex());
-		MaxIndex = FMath::Max(MaxIndex, CulledTileDataArray.GetUAVIndex());
-		MaxIndex = FMath::Max(MaxIndex, ObjectTilesIndirectArguments.GetUAVIndex());
-
-		UAVs.AddZeroed(MaxIndex + 1);
-
-		if (NumCulledTilesArray.IsUAVBound())
-		{
-			UAVs[NumCulledTilesArray.GetUAVIndex()] = TileIntersectionResources.NumCulledTilesArray.UAV;
-		}
-
-		if (CulledTilesStartOffsetArray.IsUAVBound())
-		{
-			UAVs[CulledTilesStartOffsetArray.GetUAVIndex()] = TileIntersectionResources.CulledTilesStartOffsetArray.UAV;
-		}
-
-		if (CulledTileDataArray.IsUAVBound())
-		{
-			UAVs[CulledTileDataArray.GetUAVIndex()] = TileIntersectionResources.CulledTileDataArray.UAV;
-		}
-
-		if (ObjectTilesIndirectArguments.IsUAVBound())
-		{
-			UAVs[ObjectTilesIndirectArguments.GetUAVIndex()] = TileIntersectionResources.ObjectTilesIndirectArguments.UAV;
-		}
-
-		check(UAVs.Num() > 0);
-	}
-
-	void GetReadableTransitions(FTileIntersectionResources& TileIntersectionResources, TArray<FRHITransitionInfo>& TransitionInfos)
-	{
-		if (NumCulledTilesArray.IsUAVBound())
-		{
-			TransitionInfos.Add(FRHITransitionInfo(TileIntersectionResources.NumCulledTilesArray.UAV, ERHIAccess::Unknown, ERHIAccess::SRVMask));
-		}
-
-		if (CulledTilesStartOffsetArray.IsUAVBound())
-		{
-			TransitionInfos.Add(FRHITransitionInfo(TileIntersectionResources.CulledTilesStartOffsetArray.UAV, ERHIAccess::Unknown, ERHIAccess::SRVMask));
-		}
-
-		if (CulledTileDataArray.IsUAVBound())
-		{
-			TransitionInfos.Add(FRHITransitionInfo(TileIntersectionResources.CulledTileDataArray.UAV, ERHIAccess::Unknown, ERHIAccess::SRVMask));
-		}
-
-		if (ObjectTilesIndirectArguments.IsUAVBound())
-		{
-			TransitionInfos.Add(FRHITransitionInfo(TileIntersectionResources.ObjectTilesIndirectArguments.UAV, ERHIAccess::Unknown, ERHIAccess::SRVMask | ERHIAccess::IndirectArgs));
-		}
-
-		check(TransitionInfos.Num() > 0);
-	}
-
-	template<typename TParamRef>
-	void UnsetParameters(FRHICommandList& RHICmdList, const TParamRef& ShaderRHI)
-	{
-		NumCulledTilesArray.UnsetUAV(RHICmdList, ShaderRHI);
-		CulledTilesStartOffsetArray.UnsetUAV(RHICmdList, ShaderRHI);
-		CulledTileDataArray.UnsetUAV(RHICmdList, ShaderRHI);
-		ObjectTilesIndirectArguments.UnsetUAV(RHICmdList, ShaderRHI);
-	}
-
-	friend FArchive& operator<<(FArchive& Ar, FTileIntersectionParameters& P)
-	{
-		Ar << P.TileListGroupSize;
-		Ar << P.NumCulledTilesArray;
-		Ar << P.CulledTilesStartOffsetArray;
-		Ar << P.CulledTileDataArray;
-		Ar << P.ObjectTilesIndirectArguments;
-		return Ar;
-	}
-
-private:
-	
-		LAYOUT_FIELD(FShaderParameter, TileListGroupSize)
-
-		LAYOUT_FIELD(FRWShaderParameter, NumCulledTilesArray)
-		LAYOUT_FIELD(FRWShaderParameter, CulledTilesStartOffsetArray)
-		LAYOUT_FIELD(FRWShaderParameter, CulledTileDataArray)
-		LAYOUT_FIELD(FRWShaderParameter, ObjectTilesIndirectArguments)
-	
-};
-
-class FAOScreenGridResources : public FRenderResource
-{
-public:
-
-	FAOScreenGridResources()
-	{}
-
-	virtual void InitDynamicRHI() override;
-
-	virtual void ReleaseDynamicRHI() override
-	{
-		ScreenGridConeVisibility.Release();
-		ConeDepthVisibilityFunction.Release();
-	}
-
-	void AcquireTransientResource()
-	{
-		ScreenGridConeVisibility.AcquireTransientResource();
-	}
-
-	void DiscardTransientResource()
-	{
-		ScreenGridConeVisibility.DiscardTransientResource();
-	}
-
-	FIntPoint ScreenGridDimensions;
-
-	FRWBuffer ScreenGridConeVisibility;
-	FRWBuffer ConeDepthVisibilityFunction;
+	OutEnvironment.SetDefine(TEXT("CULLED_TILE_DATA_STRIDE"), CulledTileDataStride);
+	extern int32 GDistanceFieldAOTileSizeX;
+	OutEnvironment.SetDefine(TEXT("CULLED_TILE_SIZEX"), GDistanceFieldAOTileSizeX);
+	extern int32 GConeTraceDownsampleFactor;
+	OutEnvironment.SetDefine(TEXT("TRACE_DOWNSAMPLE_FACTOR"), GConeTraceDownsampleFactor);
+	OutEnvironment.SetDefine(TEXT("CONE_TRACE_OBJECTS_THREADGROUP_SIZE"), ConeTraceObjectsThreadGroupSize);
+}
 
-	size_t GetSizeBytesForAO() const
-	{
-		return ScreenGridConeVisibility.NumBytes;
-	}
-};
+BEGIN_SHADER_PARAMETER_STRUCT(FAOScreenGridParameters, )
+	SHADER_PARAMETER_RDG_BUFFER_UAV(RWBuffer<uint>, RWScreenGridConeVisibility)
+	SHADER_PARAMETER_RDG_BUFFER_SRV(Buffer<uint>, ScreenGridConeVisibility)
+	SHADER_PARAMETER(FIntPoint, ScreenGridConeVisibilitySize)
+END_SHADER_PARAMETER_STRUCT()
 
 extern void GetSpacedVectors(uint32 FrameNumber, TArray<FVector, TInlineAllocator<9> >& OutVectors);
 
@@ -444,7 +248,6 @@ public:
 	{
 		BaseLevelTexelSize.Bind(ParameterMap, TEXT("BaseLevelTexelSize"));
 		JitterOffset.Bind(ParameterMap, TEXT("JitterOffset"));
-		ScreenGridConeVisibilitySize.Bind(ParameterMap, TEXT("ScreenGridConeVisibilitySize"));
 		DistanceFieldNormalTexture.Bind(ParameterMap, TEXT("DistanceFieldNormalTexture"));
 		DistanceFieldNormalSampler.Bind(ParameterMap, TEXT("DistanceFieldNormalSampler"));
 	}
@@ -459,9 +262,6 @@ public:
 		extern FVector2D GetJitterOffset(int32 SampleIndex);
 		SetShaderValue(RHICmdList, ShaderRHI, JitterOffset, GetJitterOffset(View.ViewState->GetDistanceFieldTemporalSampleIndex()));
 
-		FAOScreenGridResources* ScreenGridResources = View.ViewState->AOScreenGridResources;
-
-		SetShaderValue(RHICmdList, ShaderRHI, ScreenGridConeVisibilitySize, ScreenGridResources->ScreenGridDimensions);
 
 		SetTextureParameter(
 			RHICmdList,
@@ -475,7 +275,7 @@ public:
 
 	friend FArchive& operator<<(FArchive& Ar,FScreenGridParameters& P)
 	{
-		Ar << P.BaseLevelTexelSize << P.JitterOffset << P.ScreenGridConeVisibilitySize << P.DistanceFieldNormalTexture << P.DistanceFieldNormalSampler;
+		Ar << P.BaseLevelTexelSize << P.JitterOffset << P.DistanceFieldNormalTexture << P.DistanceFieldNormalSampler;
 		return Ar;
 	}
 
@@ -483,7 +283,6 @@ private:
 	
 		LAYOUT_FIELD(FShaderParameter, BaseLevelTexelSize)
 		LAYOUT_FIELD(FShaderParameter, JitterOffset)
-		LAYOUT_FIELD(FShaderParameter, ScreenGridConeVisibilitySize)
 		LAYOUT_FIELD(FShaderResourceParameter, DistanceFieldNormalTexture)
 		LAYOUT_FIELD(FShaderResourceParameter, DistanceFieldNormalSampler)
 	
@@ -493,6 +292,15 @@ extern void TrackGPUProgress(FRHICommandListImmediate& RHICmdList, uint32 DebugI
 
 extern bool ShouldRenderDeferredDynamicSkyLight(const FScene* Scene, const FSceneViewFamily& ViewFamily);
 
-extern void CullObjectsToView(FRDGBuilder& GraphBuilder, FScene* Scene, const FViewInfo& View, const FDistanceFieldAOParameters& Parameters, class FDistanceFieldObjectBufferResource& CulledObjectBuffers);
-extern void BuildTileObjectLists(FRDGBuilder& GraphBuilder, FScene* Scene, TArray<FViewInfo>& Views, FRDGTextureRef DistanceFieldNormal, const FDistanceFieldAOParameters& Parameters);
+class FDistanceFieldCulledObjectBufferParameters;
+
+extern void CullObjectsToView(FRDGBuilder& GraphBuilder, FScene* Scene, const FViewInfo& View, const FDistanceFieldAOParameters& Parameters, FDistanceFieldCulledObjectBufferParameters& CulledObjectBuffers);
+extern void BuildTileObjectLists(FRDGBuilder& GraphBuilder,
+	FScene* Scene,
+	TArray<FViewInfo>& Views,
+	FRDGBufferRef ObjectIndirectArguments,
+	const FDistanceFieldCulledObjectBufferParameters& CulledObjectBufferParameters,
+	FTileIntersectionParameters TileIntersectionParameters,
+	FRDGTextureRef DistanceFieldNormal,
+	const FDistanceFieldAOParameters& Parameters);
 extern FIntPoint GetTileListGroupSizeForView(const FViewInfo& View);
diff --git a/Engine/Source/Runtime/Renderer/Private/DistanceFieldLightingShared.h b/Engine/Source/Runtime/Renderer/Private/DistanceFieldLightingShared.h
--- a/Engine/Source/Runtime/Renderer/Private/DistanceFieldLightingShared.h
+++ b/Engine/Source/Runtime/Renderer/Private/DistanceFieldLightingShared.h
@@ -43,6 +43,17 @@ enum EDistanceFieldPrimitiveType
 	DFPT_Num
 };
 
+// Must match equivalent shader defines
+static const int32 GDistanceFieldObjectDataStride = 9;
+static const int32 GDistanceFieldObjectBoundsStride = 2;
+static const int32 GDistanceFieldCulledObjectDataStride = GDistanceFieldObjectDataStride;
+static const int32 GDistanceFieldCulledObjectBoxBoundsStride = 5;
+
+static const int32 GHeightFieldObjectDataStride = 6;
+static const int32 GHeightFieldObjectBoundsStride = 2;
+static const int32 GHeightFieldCulledObjectDataStride = GHeightFieldObjectDataStride;
+static const int32 GHeightFieldCulledObjectBoxBoundsStride = 5;
+
 template <EDistanceFieldPrimitiveType PrimitiveType>
 class TDistanceFieldObjectBuffers
 {
@@ -277,29 +288,6 @@ public:
 	}
 };
 
-class FDistanceFieldCulledObjectBuffers : public TDistanceFieldCulledObjectBuffers<DFPT_SignedDistanceField> {};
-class FHeightFieldCulledObjectBuffers : public TDistanceFieldCulledObjectBuffers<DFPT_HeightField> {};
-
-template <EDistanceFieldPrimitiveType PrimitiveType>
-class TDistanceFieldObjectBufferResource : public FRenderResource
-{
-public:
-	typename TChooseClass<PrimitiveType == DFPT_HeightField, FHeightFieldCulledObjectBuffers, FDistanceFieldCulledObjectBuffers>::Result Buffers;
-
-	virtual void InitDynamicRHI()  override
-	{
-		Buffers.Initialize();
-	}
-
-	virtual void ReleaseDynamicRHI() override
-	{
-		Buffers.Release();
-	}
-};
-
-class FDistanceFieldObjectBufferResource : public TDistanceFieldObjectBufferResource<DFPT_SignedDistanceField> {};
-class FHeightFieldObjectBufferResource : public TDistanceFieldObjectBufferResource<DFPT_HeightField> {};
-
 BEGIN_SHADER_PARAMETER_STRUCT(FDistanceFieldCulledObjectBufferParameters, )
 	SHADER_PARAMETER_RDG_BUFFER_UAV(RWBuffer<uint>, RWObjectIndirectArguments)
 	SHADER_PARAMETER_RDG_BUFFER_UAV(RWBuffer<uint>, RWCulledObjectBounds)
@@ -312,10 +300,10 @@ BEGIN_SHADER_PARAMETER_STRUCT(FDistanceFieldCulledObjectBufferParameters, )
 END_SHADER_PARAMETER_STRUCT()
 
 extern void AllocateDistanceFieldCulledObjectBuffers(
-	FRDGBuilder& GraphBuilder, 
-	bool bWantBoxBounds, 
-	uint32 MaxObjects, 
-	uint32 NumBoundsElementsScale,
+	FRDGBuilder& GraphBuilder,
+	bool bWantBoxBounds,
+	uint32 MaxObjects,
+	EDistanceFieldPrimitiveType PrimitiveType,
 	FRDGBufferRef& OutObjectIndirectArguments,
 	FDistanceFieldCulledObjectBufferParameters& OutParameters);
 
@@ -542,8 +530,6 @@ extern void CullDistanceFieldObjectsForLight(
 	FDistanceFieldCulledObjectBufferParameters& CulledObjectBufferParameters,
 	FLightTileIntersectionParameters& LightTileIntersectionParameters);
 
-extern TGlobalResource<FDistanceFieldObjectBufferResource> GAOCulledObjectBuffers;
-
 extern bool SupportsDistanceFieldAO(ERHIFeatureLevel::Type FeatureLevel, EShaderPlatform ShaderPlatform);
 
 template <bool bIsAType> struct TSelector {};
diff --git a/Engine/Source/Runtime/Renderer/Private/DistanceFieldObjectCulling.cpp b/Engine/Source/Runtime/Renderer/Private/DistanceFieldObjectCulling.cpp
--- a/Engine/Source/Runtime/Renderer/Private/DistanceFieldObjectCulling.cpp
+++ b/Engine/Source/Runtime/Renderer/Private/DistanceFieldObjectCulling.cpp
@@ -35,34 +35,24 @@ FAutoConsoleVariableRef CVarMaxDistanceFieldObjectsPerCullTile(
 	ECVF_RenderThreadSafe | ECVF_ReadOnly
 	);
 
-TGlobalResource<FDistanceFieldObjectBufferResource> GAOCulledObjectBuffers;
-
-void FTileIntersectionResources::InitDynamicRHI()
-{
-	const uint32 FastVRamFlag = GFastVRamConfig.DistanceFieldTileIntersectionResources | (IsTransientResourceBufferAliasingEnabled() ? BUF_Transient : BUF_None);
-	TileConeAxisAndCos.Initialize(TEXT("TileConeAxisAndCos"), sizeof(FVector4), TileDimensions.X * TileDimensions.Y, PF_A32B32G32R32F, BUF_Static | FastVRamFlag);
-	TileConeDepthRanges.Initialize(TEXT("TileConeDepthRanges"), sizeof(FVector4), TileDimensions.X * TileDimensions.Y, PF_A32B32G32R32F, BUF_Static | FastVRamFlag);
-
-	NumCulledTilesArray.Initialize(TEXT("NumCulledTilesArray"), sizeof(uint32), MaxSceneObjects, PF_R32_UINT, BUF_Static | FastVRamFlag);
-	CulledTilesStartOffsetArray.Initialize(TEXT("CulledTilesStartOffsetArray"), sizeof(uint32), MaxSceneObjects, PF_R32_UINT, BUF_Static | FastVRamFlag);
-
-	// Can only use 16 bit for CulledTileDataArray if few enough objects and tiles
-	const bool b16BitObjectIndices = MaxSceneObjects < (1 << 16);
-	const bool b16BitCulledTileIndexBuffer = bAllow16BitIndices && b16BitObjectIndices && TileDimensions.X * TileDimensions.Y < (1 << 16);
-	CulledTileDataArray.Initialize(
-		TEXT("CulledTileDataArray"),
-		b16BitCulledTileIndexBuffer ? sizeof(uint16) : sizeof(uint32), 
-		GAverageDistanceFieldObjectsPerCullTile * TileDimensions.X * TileDimensions.Y * CulledTileDataStride, 
-		b16BitCulledTileIndexBuffer ? PF_R16_UINT : PF_R32_UINT, 
-		BUF_Static | FastVRamFlag);
-	ObjectTilesIndirectArguments.Initialize(TEXT("ObjectTilesIndirectArguments"), sizeof(uint32), 3, PF_R32_UINT, BUF_Static | BUF_DrawIndirect);
-}
-
 class FCullObjectsForViewCS : public FGlobalShader
 {
-	DECLARE_SHADER_TYPE(FCullObjectsForViewCS,Global)
+	DECLARE_GLOBAL_SHADER(FCullObjectsForViewCS);
+	SHADER_USE_PARAMETER_STRUCT(FCullObjectsForViewCS, FGlobalShader);
 public:
 
+	BEGIN_SHADER_PARAMETER_STRUCT(FParameters, )
+		SHADER_PARAMETER_STRUCT_INCLUDE(FDistanceFieldCulledObjectBufferParameters, CulledObjectBufferParameters)
+		SHADER_PARAMETER_STRUCT_INCLUDE(FDistanceFieldObjectBufferParameters, ObjectBufferParameters)
+
+		SHADER_PARAMETER_STRUCT_REF(FViewUniformShaderParameters, View)
+		SHADER_PARAMETER(uint32, NumConvexHullPlanes)
+		SHADER_PARAMETER_ARRAY(FVector4, ViewFrustumConvexHull, [6])
+		SHADER_PARAMETER(uint32, ObjectBoundingGeometryIndexCount)
+		SHADER_PARAMETER(float, AOObjectMaxDistance)
+		SHADER_PARAMETER(float, AOMaxViewDistance)
+	END_SHADER_PARAMETER_STRUCT()
+
 	static bool ShouldCompilePermutation(const FGlobalShaderPermutationParameters& Parameters)
 	{
 		return IsFeatureLevelSupported(Parameters.Platform, ERHIFeatureLevel::SM5) && DoesPlatformSupportDistanceFieldAO(Parameters.Platform) && IsUsingDistanceFields(Parameters.Platform);
@@ -73,108 +63,60 @@ public:
 		FGlobalShader::ModifyCompilationEnvironment(Parameters,OutEnvironment);
 		OutEnvironment.SetDefine(TEXT("UPDATEOBJECTS_THREADGROUP_SIZE"), UpdateObjectsGroupSize);
 	}
+};
 
-	FCullObjectsForViewCS(const ShaderMetaType::CompiledShaderInitializerType& Initializer)
-		: FGlobalShader(Initializer)
-	{
-		ObjectBufferParameters.Bind(Initializer.ParameterMap);
-		CulledObjectParameters.Bind(Initializer.ParameterMap);
-		AOParameters.Bind(Initializer.ParameterMap);
-		NumConvexHullPlanes.Bind(Initializer.ParameterMap, TEXT("NumConvexHullPlanes"));
-		ViewFrustumConvexHull.Bind(Initializer.ParameterMap, TEXT("ViewFrustumConvexHull"));
-		ObjectBoundingGeometryIndexCount.Bind(Initializer.ParameterMap, TEXT("ObjectBoundingGeometryIndexCount"));
-	}
-
-	FCullObjectsForViewCS()
-	{
-	}
-
-	void SetParameters(FRHICommandList& RHICmdList, const FScene* Scene, const FSceneView& View, const FDistanceFieldAOParameters& Parameters)
-	{
-		FRHITransitionInfo UAVTransitions[6];
-		UAVTransitions[0] = FRHITransitionInfo(GAOCulledObjectBuffers.Buffers.ObjectIndirectArguments.UAV, ERHIAccess::Unknown, ERHIAccess::UAVCompute);
-		UAVTransitions[1] = FRHITransitionInfo(GAOCulledObjectBuffers.Buffers.Bounds.UAV, ERHIAccess::Unknown, ERHIAccess::UAVCompute);
-		UAVTransitions[2] = FRHITransitionInfo(GAOCulledObjectBuffers.Buffers.Data.UAV, ERHIAccess::Unknown, ERHIAccess::UAVCompute);
-		UAVTransitions[3] = FRHITransitionInfo(GAOCulledObjectBuffers.Buffers.BoxBounds.UAV, ERHIAccess::Unknown, ERHIAccess::UAVCompute);
-		UAVTransitions[4] = FRHITransitionInfo(Scene->DistanceFieldSceneData.GetCurrentObjectBuffers()->Data.UAV, ERHIAccess::Unknown, ERHIAccess::SRVCompute);
-		UAVTransitions[5] = FRHITransitionInfo(Scene->DistanceFieldSceneData.GetCurrentObjectBuffers()->Bounds.UAV, ERHIAccess::Unknown, ERHIAccess::SRVCompute);
-		RHICmdList.Transition(MakeArrayView(UAVTransitions, UE_ARRAY_COUNT(UAVTransitions)));
-
-		FRHIComputeShader* ShaderRHI = RHICmdList.GetBoundComputeShader();
-		FGlobalShader::SetParameters<FViewUniformShaderParameters>(RHICmdList, ShaderRHI, View.ViewUniformBuffer);
-
-		ObjectBufferParameters.Set(RHICmdList, ShaderRHI, *(Scene->DistanceFieldSceneData.GetCurrentObjectBuffers()), Scene->DistanceFieldSceneData.NumObjectsInBuffer);
-		CulledObjectParameters.Set(RHICmdList, ShaderRHI, GAOCulledObjectBuffers.Buffers, Scene->DistanceFieldSceneData);
-
-		AOParameters.Set(RHICmdList, ShaderRHI, Parameters);
+IMPLEMENT_GLOBAL_SHADER(FCullObjectsForViewCS, "/Engine/Private/DistanceFieldObjectCulling.usf", "CullObjectsForViewCS", SF_Compute);
 
-		// Shader assumes max 6
-		check(View.ViewFrustum.Planes.Num() <= 6);
-		SetShaderValue(RHICmdList, ShaderRHI, NumConvexHullPlanes, View.ViewFrustum.Planes.Num());
-		SetShaderValueArray(RHICmdList, ShaderRHI, ViewFrustumConvexHull, LWC::DemoteArrayType<FPlane4f, FPlane, TInlineAllocator<6>>(View.ViewFrustum.Planes).GetData(), View.ViewFrustum.Planes.Num());	// LWC_TODO: Perf pessimization
-		SetShaderValue(RHICmdList, ShaderRHI, ObjectBoundingGeometryIndexCount, StencilingGeometry::GLowPolyStencilSphereIndexBuffer.GetIndexCount());
-	}
+void CullObjectsToView(FRDGBuilder& GraphBuilder, FScene* Scene, const FViewInfo& View, const FDistanceFieldAOParameters& Parameters, FDistanceFieldCulledObjectBufferParameters& CulledObjectBufferParameters)
+{
+	AddClearUAVPass(GraphBuilder, CulledObjectBufferParameters.RWObjectIndirectArguments, 0);
 
-	void UnsetParameters(FRHICommandList& RHICmdList, const FScene* Scene)
 	{
-		ObjectBufferParameters.UnsetParameters(RHICmdList, RHICmdList.GetBoundComputeShader(), *(Scene->DistanceFieldSceneData.GetCurrentObjectBuffers()));
-		CulledObjectParameters.UnsetParameters(RHICmdList, RHICmdList.GetBoundComputeShader());
-
-		FRHITransitionInfo SRVTransitions[6];
-		SRVTransitions[0] = FRHITransitionInfo(GAOCulledObjectBuffers.Buffers.ObjectIndirectArguments.UAV, ERHIAccess::Unknown, ERHIAccess::IndirectArgs | ERHIAccess::SRVMask);
-		SRVTransitions[1] = FRHITransitionInfo(GAOCulledObjectBuffers.Buffers.Bounds.UAV, ERHIAccess::Unknown, ERHIAccess::SRVMask);
-		SRVTransitions[2] = FRHITransitionInfo(GAOCulledObjectBuffers.Buffers.Data.UAV, ERHIAccess::Unknown, ERHIAccess::SRVMask);
-		SRVTransitions[3] = FRHITransitionInfo(GAOCulledObjectBuffers.Buffers.BoxBounds.UAV, ERHIAccess::Unknown, ERHIAccess::SRVMask);
-		SRVTransitions[4] = FRHITransitionInfo(Scene->DistanceFieldSceneData.GetCurrentObjectBuffers()->Data.UAV, ERHIAccess::Unknown, ERHIAccess::SRVMask);
-		SRVTransitions[5] = FRHITransitionInfo(Scene->DistanceFieldSceneData.GetCurrentObjectBuffers()->Bounds.UAV, ERHIAccess::Unknown, ERHIAccess::SRVMask);
-		RHICmdList.Transition(MakeArrayView(SRVTransitions, UE_ARRAY_COUNT(SRVTransitions)));
-	}
+		const int32 NumObjectsInBuffer = Scene->DistanceFieldSceneData.NumObjectsInBuffer;
 
-private:
+		auto* PassParameters = GraphBuilder.AllocParameters<FCullObjectsForViewCS::FParameters>();
 
-	LAYOUT_FIELD((TDistanceFieldObjectBufferParameters<DFPT_SignedDistanceField>), ObjectBufferParameters);
-	LAYOUT_FIELD((TDistanceFieldCulledObjectBufferParameters<DFPT_SignedDistanceField>), CulledObjectParameters);
-	LAYOUT_FIELD(FAOParameters, AOParameters);
-	LAYOUT_FIELD(FShaderParameter, NumConvexHullPlanes);
-	LAYOUT_FIELD(FShaderParameter, ViewFrustumConvexHull);
-	LAYOUT_FIELD(FShaderParameter, ObjectBoundingGeometryIndexCount);
-};
+		PassParameters->CulledObjectBufferParameters = CulledObjectBufferParameters;
+		PassParameters->ObjectBufferParameters = DistanceField::SetupObjectBufferParameters(Scene->DistanceFieldSceneData);
 
-IMPLEMENT_SHADER_TYPE(,FCullObjectsForViewCS,TEXT("/Engine/Private/DistanceFieldObjectCulling.usf"),TEXT("CullObjectsForViewCS"),SF_Compute);
+		PassParameters->View = View.ViewUniformBuffer;
+		PassParameters->NumConvexHullPlanes = View.ViewFrustum.Planes.Num();
 
-void CullObjectsToView(FRDGBuilder& GraphBuilder, FScene* Scene, const FViewInfo& View, const FDistanceFieldAOParameters& Parameters, FDistanceFieldObjectBufferResource& CulledObjectBuffers)
-{
-	AddPass(GraphBuilder, RDG_EVENT_NAME("ObjectFrustumCulling"), [Scene, &View, Parameters, &CulledObjectBuffers] (FRHICommandListImmediate& RHICmdList)
-	{
-		if (!CulledObjectBuffers.IsInitialized()
-			|| CulledObjectBuffers.Buffers.MaxObjects < Scene->DistanceFieldSceneData.NumObjectsInBuffer
-			|| CulledObjectBuffers.Buffers.MaxObjects > 3 * Scene->DistanceFieldSceneData.NumObjectsInBuffer)
+		for (int32 i = 0; i < View.ViewFrustum.Planes.Num(); i++)
 		{
-			CulledObjectBuffers.Buffers.MaxObjects = Scene->DistanceFieldSceneData.NumObjectsInBuffer * 5 / 4;
-			CulledObjectBuffers.ReleaseResource();
-			CulledObjectBuffers.InitResource();
+			PassParameters->ViewFrustumConvexHull[i] = FVector4(View.ViewFrustum.Planes[i], View.ViewFrustum.Planes[i].W);
 		}
 
-		CulledObjectBuffers.Buffers.AcquireTransientResource();
-
-		RHICmdList.Transition(FRHITransitionInfo(CulledObjectBuffers.Buffers.ObjectIndirectArguments.UAV, ERHIAccess::Unknown, ERHIAccess::UAVCompute));
-		RHICmdList.ClearUAVUint(CulledObjectBuffers.Buffers.ObjectIndirectArguments.UAV, FUintVector4(0, 0, 0, 0));
+		PassParameters->ObjectBoundingGeometryIndexCount = StencilingGeometry::GLowPolyStencilSphereIndexBuffer.GetIndexCount();
+		PassParameters->AOObjectMaxDistance = Parameters.ObjectMaxOcclusionDistance;
+		PassParameters->AOMaxViewDistance = GetMaxAOViewDistance();
 
-		TShaderMapRef<FCullObjectsForViewCS> ComputeShader(GetGlobalShaderMap(Scene->GetFeatureLevel()));
-		RHICmdList.SetComputeShader(ComputeShader.GetComputeShader());
-		ComputeShader->SetParameters(RHICmdList, Scene, View, Parameters);
+		auto ComputeShader = View.ShaderMap->GetShader<FCullObjectsForViewCS>();
+		const int32 GroupSize = FMath::DivideAndRoundUp<uint32>(NumObjectsInBuffer, UpdateObjectsGroupSize);
 
-		DispatchComputeShader(RHICmdList, ComputeShader.GetShader(), FMath::DivideAndRoundUp<uint32>(Scene->DistanceFieldSceneData.NumObjectsInBuffer, UpdateObjectsGroupSize), 1, 1);
-		ComputeShader->UnsetParameters(RHICmdList, Scene);
-	});
+		FComputeShaderUtils::AddPass(
+			GraphBuilder,
+			RDG_EVENT_NAME("ObjectFrustumCulling"),
+			ComputeShader,
+			PassParameters,
+			FIntVector(GroupSize, 1, 1));
+	}
 }
 
 /**  */
 class FBuildTileConesCS : public FGlobalShader
 {
-	DECLARE_SHADER_TYPE(FBuildTileConesCS,Global)
+	DECLARE_GLOBAL_SHADER(FBuildTileConesCS);
 public:
 
+	BEGIN_SHADER_PARAMETER_STRUCT(FParameters, )
+		SHADER_PARAMETER_RDG_BUFFER_UAV(RWBuffer<FVector4>, RWTileConeAxisAndCos)
+		SHADER_PARAMETER_RDG_BUFFER_UAV(RWBuffer<FVector4>, RWTileConeDepthRanges)
+		SHADER_PARAMETER_RDG_UNIFORM_BUFFER(FSceneTextureUniformParameters, SceneTextures)
+		SHADER_PARAMETER_RDG_TEXTURE(Texture2D, DistanceFieldNormalTexture)
+		SHADER_PARAMETER_SAMPLER(SamplerState, DistanceFieldNormalSampler)
+	END_SHADER_PARAMETER_STRUCT()
+
 	static bool ShouldCompilePermutation(const FGlobalShaderPermutationParameters& Parameters)
 	{
 		return IsFeatureLevelSupported(Parameters.Platform, ERHIFeatureLevel::SM5) && DoesPlatformSupportDistanceFieldAO(Parameters.Platform) && IsUsingDistanceFields(Parameters.Platform);
@@ -194,80 +136,46 @@ public:
 	FBuildTileConesCS(const ShaderMetaType::CompiledShaderInitializerType& Initializer)
 		: FGlobalShader(Initializer)
 	{
+		BindForLegacyShaderParameters<FParameters>(this, Initializer.PermutationId, Initializer.ParameterMap, false);
 		AOParameters.Bind(Initializer.ParameterMap);
-		TileConeAxisAndCos.Bind(Initializer.ParameterMap, TEXT("TileConeAxisAndCos"));
-		TileConeDepthRanges.Bind(Initializer.ParameterMap, TEXT("TileConeDepthRanges"));
 		NumGroups.Bind(Initializer.ParameterMap, TEXT("NumGroups"));
 		ViewDimensionsParameter.Bind(Initializer.ParameterMap, TEXT("ViewDimensions"));
-		DistanceFieldNormalTexture.Bind(Initializer.ParameterMap, TEXT("DistanceFieldNormalTexture"));
-		DistanceFieldNormalSampler.Bind(Initializer.ParameterMap, TEXT("DistanceFieldNormalSampler"));
 	}
 
 	FBuildTileConesCS()
 	{
 	}
-	void SetParameters(FRHICommandList& RHICmdList, const FViewInfo& View, FRHITexture* DistanceFieldNormal, FScene* Scene, FVector2D NumGroupsValue, const FDistanceFieldAOParameters& Parameters)
+	void SetParameters(FRHICommandList& RHICmdList, const FViewInfo& View, FScene* Scene, FVector2D NumGroupsValue, const FDistanceFieldAOParameters& Parameters)
 	{
 		FRHIComputeShader* ShaderRHI = RHICmdList.GetBoundComputeShader();
 
 		FGlobalShader::SetParameters<FViewUniformShaderParameters>(RHICmdList, ShaderRHI, View.ViewUniformBuffer);
 		AOParameters.Set(RHICmdList, ShaderRHI, Parameters);
 
-		FTileIntersectionResources* TileIntersectionResources = ((FSceneViewState*)View.State)->AOTileIntersectionResources;
-
-		FRHITransitionInfo UAVTransitions[2];
-		UAVTransitions[0] = FRHITransitionInfo(TileIntersectionResources->TileConeAxisAndCos.UAV, ERHIAccess::Unknown, ERHIAccess::UAVCompute);
-		UAVTransitions[1] = FRHITransitionInfo(TileIntersectionResources->TileConeDepthRanges.UAV, ERHIAccess::Unknown, ERHIAccess::UAVCompute);
-		RHICmdList.Transition(MakeArrayView(UAVTransitions, UE_ARRAY_COUNT(UAVTransitions)));
-
-		TileConeAxisAndCos.SetBuffer(RHICmdList, ShaderRHI, TileIntersectionResources->TileConeAxisAndCos);
-		TileConeDepthRanges.SetBuffer(RHICmdList, ShaderRHI, TileIntersectionResources->TileConeDepthRanges);
-
 		SetShaderValue(RHICmdList, ShaderRHI, ViewDimensionsParameter, View.ViewRect);
 
 		SetShaderValue(RHICmdList, ShaderRHI, NumGroups, NumGroupsValue);
-
-		SetTextureParameter(
-			RHICmdList,
-			ShaderRHI,
-			DistanceFieldNormalTexture,
-			DistanceFieldNormalSampler,
-			TStaticSamplerState<SF_Point,AM_Clamp,AM_Clamp,AM_Clamp>::GetRHI(),
-			DistanceFieldNormal
-			);
-	}
-
-	void UnsetParameters(FRHICommandList& RHICmdList, const FSceneView& View)
-	{
-		TileConeAxisAndCos.UnsetUAV(RHICmdList, RHICmdList.GetBoundComputeShader());
-		TileConeDepthRanges.UnsetUAV(RHICmdList, RHICmdList.GetBoundComputeShader());
-
-		FTileIntersectionResources* TileIntersectionResources = ((FSceneViewState*)View.State)->AOTileIntersectionResources;
-
-		FRHITransitionInfo SRVTransitions[2];
-		SRVTransitions[0] = FRHITransitionInfo(TileIntersectionResources->TileConeAxisAndCos.UAV, ERHIAccess::Unknown, ERHIAccess::SRVMask);
-		SRVTransitions[1] = FRHITransitionInfo(TileIntersectionResources->TileConeDepthRanges.UAV, ERHIAccess::Unknown, ERHIAccess::SRVMask);
-		RHICmdList.Transition(MakeArrayView(SRVTransitions, UE_ARRAY_COUNT(SRVTransitions)));
 	}
 
 private:
 	LAYOUT_FIELD(FAOParameters, AOParameters);
-	LAYOUT_FIELD(FRWShaderParameter, TileConeAxisAndCos);
-	LAYOUT_FIELD(FRWShaderParameter, TileConeDepthRanges);
 	LAYOUT_FIELD(FShaderParameter, ViewDimensionsParameter);
 	LAYOUT_FIELD(FShaderParameter, NumGroups);
-	LAYOUT_FIELD(FShaderResourceParameter, DistanceFieldNormalTexture);
-	LAYOUT_FIELD(FShaderResourceParameter, DistanceFieldNormalSampler);
 };
 
-IMPLEMENT_SHADER_TYPE(,FBuildTileConesCS,TEXT("/Engine/Private/DistanceFieldObjectCulling.usf"),TEXT("BuildTileConesMain"),SF_Compute);
+IMPLEMENT_GLOBAL_SHADER(FBuildTileConesCS, "/Engine/Private/DistanceFieldObjectCulling.usf", "BuildTileConesMain", SF_Compute);
 
 /**  */
 class FObjectCullVS : public FGlobalShader
 {
-	DECLARE_SHADER_TYPE(FObjectCullVS,Global);
+	DECLARE_GLOBAL_SHADER(FObjectCullVS);
 public:
 
+	BEGIN_SHADER_PARAMETER_STRUCT(FParameters, )
+		SHADER_PARAMETER_STRUCT_INCLUDE(FDistanceFieldCulledObjectBufferParameters, DistanceFieldCulledObjectBuffers)
+		SHADER_PARAMETER_STRUCT_INCLUDE(FDistanceFieldAtlasParameters, DistanceFieldAtlas)
+	END_SHADER_PARAMETER_STRUCT()
+
 	static bool ShouldCompilePermutation(const FGlobalShaderPermutationParameters& Parameters) 
 	{ 
 		return IsFeatureLevelSupported(Parameters.Platform, ERHIFeatureLevel::SM5) && DoesPlatformSupportDistanceFieldAO(Parameters.Platform) && IsUsingDistanceFields(Parameters.Platform);
@@ -276,19 +184,18 @@ public:
 	FObjectCullVS(const ShaderMetaType::CompiledShaderInitializerType& Initializer):
 		FGlobalShader(Initializer)
 	{
-		ObjectParameters.Bind(Initializer.ParameterMap);
+		BindForLegacyShaderParameters<FParameters>(this, Initializer.PermutationId, Initializer.ParameterMap, false);
 		AOParameters.Bind(Initializer.ParameterMap);
 		ConservativeRadiusScale.Bind(Initializer.ParameterMap, TEXT("ConservativeRadiusScale"));
 	}
 
 	FObjectCullVS() {}
 
-	void SetParameters(FRHICommandList& RHICmdList, const FSceneView& View, const FDistanceFieldSceneData& DistanceFieldSceneData, const FDistanceFieldAOParameters& Parameters)
+	void SetParameters(FRHICommandList& RHICmdList, const FSceneView& View, const FDistanceFieldAOParameters& Parameters)
 	{
 		FRHIVertexShader* ShaderRHI = RHICmdList.GetBoundVertexShader();
 		FGlobalShader::SetParameters<FViewUniformShaderParameters>(RHICmdList, ShaderRHI, View.ViewUniformBuffer);
 		
-		ObjectParameters.Set(RHICmdList, ShaderRHI, GAOCulledObjectBuffers.Buffers, DistanceFieldSceneData);
 		AOParameters.Set(RHICmdList, ShaderRHI, Parameters);
 
 		const int32 NumRings = StencilingGeometry::GLowPolyStencilSphereVertexBuffer.GetNumRings();
@@ -301,18 +208,23 @@ public:
 	}
 
 private:
-	LAYOUT_FIELD((TDistanceFieldCulledObjectBufferParameters<DFPT_SignedDistanceField>), ObjectParameters);
 	LAYOUT_FIELD(FAOParameters, AOParameters);
 	LAYOUT_FIELD(FShaderParameter, ConservativeRadiusScale);
 };
 
-IMPLEMENT_SHADER_TYPE(,FObjectCullVS,TEXT("/Engine/Private/DistanceFieldObjectCulling.usf"),TEXT("ObjectCullVS"),SF_Vertex);
+IMPLEMENT_GLOBAL_SHADER(FObjectCullVS, "/Engine/Private/DistanceFieldObjectCulling.usf", "ObjectCullVS", SF_Vertex);
 
 class FObjectCullPS : public FGlobalShader
 {
 	DECLARE_GLOBAL_SHADER(FObjectCullPS);
 public:
 
+	BEGIN_SHADER_PARAMETER_STRUCT(FParameters, )
+		SHADER_PARAMETER_STRUCT_INCLUDE(FTileIntersectionParameters, TileIntersectionParameters)
+		SHADER_PARAMETER_STRUCT_INCLUDE(FDistanceFieldCulledObjectBufferParameters, DistanceFieldCulledObjectBuffers)
+		SHADER_PARAMETER_STRUCT_INCLUDE(FDistanceFieldAtlasParameters, DistanceFieldAtlas)
+	END_SHADER_PARAMETER_STRUCT()
+
 	class FCountingPass : SHADER_PERMUTATION_BOOL("SCATTER_CULLING_COUNT_PASS");
 	using FPermutationDomain = TShaderPermutationDomain<FCountingPass>;
 
@@ -324,7 +236,7 @@ public:
 	static void ModifyCompilationEnvironment(const FGlobalShaderPermutationParameters& Parameters, FShaderCompilerEnvironment& OutEnvironment)
 	{
 		FGlobalShader::ModifyCompilationEnvironment(Parameters, OutEnvironment);
-		FTileIntersectionParameters::ModifyCompilationEnvironment(Parameters.Platform, OutEnvironment);
+		TileIntersectionModifyCompilationEnvironment(Parameters.Platform, OutEnvironment);
 
 		OutEnvironment.SetDefine(TEXT("DOWNSAMPLE_FACTOR"), GAODownsampleFactor);
 	}
@@ -336,57 +248,54 @@ public:
 	FObjectCullPS(const ShaderMetaType::CompiledShaderInitializerType& Initializer)
 		: FGlobalShader(Initializer)
 	{
-		ObjectParameters.Bind(Initializer.ParameterMap);
+		BindForLegacyShaderParameters<FParameters>(this, Initializer.PermutationId, Initializer.ParameterMap, false);
 		AOParameters.Bind(Initializer.ParameterMap);
-		TileIntersectionParameters.Bind(Initializer.ParameterMap);
-		TileConeAxisAndCos.Bind(Initializer.ParameterMap, TEXT("TileConeAxisAndCos"));
-		TileConeDepthRanges.Bind(Initializer.ParameterMap, TEXT("TileConeDepthRanges"));
 		NumGroups.Bind(Initializer.ParameterMap, TEXT("NumGroups"));
 	}
 
-	void SetParameters(FRHICommandList& RHICmdList, const FSceneView& View, const FDistanceFieldSceneData& DistanceFieldSceneData, FVector2D NumGroupsValue, const FDistanceFieldAOParameters& Parameters)
+	void SetParameters(FRHICommandList& RHICmdList, const FSceneView& View, FVector2D NumGroupsValue, const FDistanceFieldAOParameters& Parameters)
 	{
 		FRHIPixelShader* ShaderRHI = RHICmdList.GetBoundPixelShader();
 		FGlobalShader::SetParameters<FViewUniformShaderParameters>(RHICmdList, ShaderRHI, View.ViewUniformBuffer);
 
-		ObjectParameters.Set(RHICmdList, ShaderRHI, GAOCulledObjectBuffers.Buffers, DistanceFieldSceneData);
 		AOParameters.Set(RHICmdList, ShaderRHI, Parameters);
 
-		FTileIntersectionResources* TileIntersectionResources = ((FSceneViewState*)View.State)->AOTileIntersectionResources;
-
-		SetSRVParameter(RHICmdList, ShaderRHI, TileConeAxisAndCos, TileIntersectionResources->TileConeAxisAndCos.SRV);
-		SetSRVParameter(RHICmdList, ShaderRHI, TileConeDepthRanges, TileIntersectionResources->TileConeDepthRanges.SRV);
-		TileIntersectionParameters.Set(RHICmdList, ShaderRHI, *TileIntersectionResources);
-
 		SetShaderValue(RHICmdList, ShaderRHI, NumGroups, NumGroupsValue);
 	}
 
-	void GetUAVs(const FSceneView& View, TArray<FRHIUnorderedAccessView*>& UAVs)
-	{
-		FTileIntersectionResources* TileIntersectionResources = ((FSceneViewState*)View.State)->AOTileIntersectionResources;
-		TileIntersectionParameters.GetUAVs(*TileIntersectionResources, UAVs);
-		check(UAVs.Num() > 0);
-	}
-
 private:
-	LAYOUT_FIELD((TDistanceFieldCulledObjectBufferParameters<DFPT_SignedDistanceField>), ObjectParameters);
 	LAYOUT_FIELD(FAOParameters, AOParameters);
-	LAYOUT_FIELD(FTileIntersectionParameters, TileIntersectionParameters);
-	LAYOUT_FIELD(FShaderResourceParameter, TileConeAxisAndCos);
-	LAYOUT_FIELD(FShaderResourceParameter, TileConeDepthRanges);
 	LAYOUT_FIELD(FShaderParameter, NumGroups);
 };
 
 IMPLEMENT_GLOBAL_SHADER(FObjectCullPS, "/Engine/Private/DistanceFieldObjectCulling.usf", "ObjectCullPS", SF_Pixel);
 
+BEGIN_SHADER_PARAMETER_STRUCT(FObjectCullParameters, )
+	SHADER_PARAMETER_STRUCT_INCLUDE(FObjectCullVS::FParameters, VS)
+	SHADER_PARAMETER_STRUCT_INCLUDE(FObjectCullPS::FParameters, PS)
+	SHADER_PARAMETER_RDG_UNIFORM_BUFFER(FSceneTextureUniformParameters, SceneTextures)
+	RDG_BUFFER_ACCESS(ObjectIndirectArguments, ERHIAccess::IndirectArgs)
+	RENDER_TARGET_BINDING_SLOTS()
+END_SHADER_PARAMETER_STRUCT()
+
 const uint32 ComputeStartOffsetGroupSize = 64;
 
 /**  */
 class FComputeCulledTilesStartOffsetCS : public FGlobalShader
 {
-	DECLARE_SHADER_TYPE(FComputeCulledTilesStartOffsetCS,Global)
+	DECLARE_GLOBAL_SHADER(FComputeCulledTilesStartOffsetCS);
+	SHADER_USE_PARAMETER_STRUCT(FComputeCulledTilesStartOffsetCS, FGlobalShader);
+
 public:
 
+	BEGIN_SHADER_PARAMETER_STRUCT(FParameters, )
+		SHADER_PARAMETER_STRUCT_REF(FViewUniformShaderParameters, View)
+		SHADER_PARAMETER_STRUCT_INCLUDE(FTileIntersectionParameters, TileIntersectionParameters)
+		SHADER_PARAMETER_STRUCT_INCLUDE(FDistanceFieldCulledObjectBufferParameters, DistanceFieldCulledObjectBuffers)
+		SHADER_PARAMETER_STRUCT_INCLUDE(FDistanceFieldAtlasParameters, DistanceFieldAtlas)
+		SHADER_PARAMETER_RDG_UNIFORM_BUFFER(FSceneTextureUniformParameters, SceneTextures)
+	END_SHADER_PARAMETER_STRUCT()
+
 	static bool ShouldCompilePermutation(const FGlobalShaderPermutationParameters& Parameters)
 	{
 		return IsFeatureLevelSupported(Parameters.Platform, ERHIFeatureLevel::SM5) && DoesPlatformSupportDistanceFieldAO(Parameters.Platform) && IsUsingDistanceFields(Parameters.Platform);
@@ -395,65 +304,25 @@ public:
 	static void ModifyCompilationEnvironment(const FGlobalShaderPermutationParameters& Parameters, FShaderCompilerEnvironment& OutEnvironment)
 	{
 		FGlobalShader::ModifyCompilationEnvironment(Parameters,OutEnvironment);
-		FTileIntersectionParameters::ModifyCompilationEnvironment(Parameters.Platform, OutEnvironment);
+		TileIntersectionModifyCompilationEnvironment(Parameters.Platform, OutEnvironment);
 
 		OutEnvironment.SetDefine(TEXT("COMPUTE_START_OFFSET_GROUP_SIZE"), ComputeStartOffsetGroupSize);
 	}
-
-	FComputeCulledTilesStartOffsetCS(const ShaderMetaType::CompiledShaderInitializerType& Initializer)
-		: FGlobalShader(Initializer)
-	{
-		ObjectParameters.Bind(Initializer.ParameterMap);
-		TileIntersectionParameters.Bind(Initializer.ParameterMap);
-	}
-
-	FComputeCulledTilesStartOffsetCS()
-	{
-	}
-	void SetParameters(FRHICommandList& RHICmdList, const FSceneView& View, const FDistanceFieldSceneData& DistanceFieldSceneData)
-	{
-		FRHIComputeShader* ShaderRHI = RHICmdList.GetBoundComputeShader();
-
-		FGlobalShader::SetParameters<FViewUniformShaderParameters>(RHICmdList, ShaderRHI, View.ViewUniformBuffer);
-
-		ObjectParameters.Set(RHICmdList, ShaderRHI, GAOCulledObjectBuffers.Buffers, DistanceFieldSceneData);
-
-		FTileIntersectionResources* TileIntersectionResources = ((FSceneViewState*)View.State)->AOTileIntersectionResources;
-		
-		TArray<FRHIUnorderedAccessView*> UAVs;
-		TileIntersectionParameters.GetUAVs(*TileIntersectionResources, UAVs);
-
-		TArray<FRHITransitionInfo> TransitionInfos;
-		for (FRHIUnorderedAccessView* UAV : UAVs)
-		{
-			TransitionInfos.Add(FRHITransitionInfo(UAV, ERHIAccess::Unknown, ERHIAccess::UAVCompute));
-		}
-		TransitionInfos.Add(FRHITransitionInfo(TileIntersectionResources->NumCulledTilesArray.UAV, ERHIAccess::Unknown, ERHIAccess::SRVMask));
-		RHICmdList.Transition(MakeArrayView(TransitionInfos.GetData(), TransitionInfos.Num()));
-
-		TileIntersectionParameters.Set(RHICmdList, ShaderRHI, *TileIntersectionResources);
-	}
-
-	void UnsetParameters(FRHICommandList& RHICmdList, const FSceneView& View)
-	{
-		FTileIntersectionResources* TileIntersectionResources = ((FSceneViewState*)View.State)->AOTileIntersectionResources;
-
-		TileIntersectionParameters.UnsetParameters(RHICmdList, RHICmdList.GetBoundComputeShader());
-
-		TArray<FRHITransitionInfo> SRVTransitions;
-		TileIntersectionParameters.GetReadableTransitions(*TileIntersectionResources, SRVTransitions);
-		RHICmdList.Transition(MakeArrayView(SRVTransitions.GetData(), SRVTransitions.Num()));
-	}
-
-private:
-
-	LAYOUT_FIELD((TDistanceFieldCulledObjectBufferParameters<DFPT_SignedDistanceField>), ObjectParameters);
-	LAYOUT_FIELD(FTileIntersectionParameters, TileIntersectionParameters);
 };
 
-IMPLEMENT_SHADER_TYPE(,FComputeCulledTilesStartOffsetCS,TEXT("/Engine/Private/DistanceFieldObjectCulling.usf"),TEXT("ComputeCulledTilesStartOffsetCS"),SF_Compute);
-
-void ScatterTilesToObjects(FRHICommandListImmediate& RHICmdList, bool bCountingPass, const FViewInfo& View, const FDistanceFieldSceneData& DistanceFieldSceneData, FIntPoint TileListGroupSize, const FDistanceFieldAOParameters& Parameters)
+IMPLEMENT_GLOBAL_SHADER(FComputeCulledTilesStartOffsetCS, "/Engine/Private/DistanceFieldObjectCulling.usf", "ComputeCulledTilesStartOffsetCS", SF_Compute);
+
+void ScatterTilesToObjects(
+	FRDGBuilder& GraphBuilder,
+	bool bCountingPass,
+	const FViewInfo& View,
+	const FDistanceFieldSceneData& DistanceFieldSceneData,
+	FIntPoint TileListGroupSize,
+	const FDistanceFieldAOParameters& Parameters,
+	FRDGBufferRef ObjectIndirectArguments,
+	const FDistanceFieldCulledObjectBufferParameters& CulledObjectBufferParameters,
+	const FTileIntersectionParameters& TileIntersectionParameters,
+	TRDGUniformBufferRef<FSceneTextureUniformParameters> SceneTexturesUniformBuffer)
 {
 	FObjectCullPS::FPermutationDomain PermutationVector;
 	PermutationVector.Set<FObjectCullPS::FCountingPass>(bCountingPass);
@@ -461,62 +330,60 @@ void ScatterTilesToObjects(FRHICommandListImmediate& RHICmdList, bool bCountingP
 	auto VertexShader = View.ShaderMap->GetShader<FObjectCullVS>();
 	auto PixelShader = View.ShaderMap->GetShader<FObjectCullPS>(PermutationVector);
 
-	TArray<FRHIUnorderedAccessView*> UAVs;
-	PixelShader->GetUAVs(View, UAVs);
-	TArray<FRHITransitionInfo> TransitionInfos;
-	for (FRHIUnorderedAccessView* UAV : UAVs)
-	{
-		TransitionInfos.Add(FRHITransitionInfo(UAV, ERHIAccess::Unknown, ERHIAccess::UAVMask));
-	}
-	RHICmdList.Transition(MakeArrayView(TransitionInfos.GetData(), TransitionInfos.Num()));
+	auto* PassParameters = GraphBuilder.AllocParameters<FObjectCullParameters>();
+	PassParameters->VS.DistanceFieldCulledObjectBuffers = CulledObjectBufferParameters;
+	PassParameters->VS.DistanceFieldAtlas = DistanceField::SetupAtlasParameters(DistanceFieldSceneData);
+	PassParameters->PS.TileIntersectionParameters = TileIntersectionParameters;
+	PassParameters->PS.DistanceFieldCulledObjectBuffers = CulledObjectBufferParameters;
+	PassParameters->PS.DistanceFieldAtlas = DistanceField::SetupAtlasParameters(DistanceFieldSceneData);
+	PassParameters->SceneTextures = SceneTexturesUniformBuffer;
+	PassParameters->ObjectIndirectArguments = ObjectIndirectArguments;
 
-	FRHIRenderPassInfo RPInfo(FRHIRenderPassInfo::NoRenderTargets);
 	if (GRHIRequiresRenderTargetForPixelShaderUAVs)
 	{
-		TRefCountPtr<IPooledRenderTarget> Dummy;
-		FPooledRenderTargetDesc Desc(FPooledRenderTargetDesc::Create2DDesc(TileListGroupSize, PF_B8G8R8A8, FClearValueBinding::None, TexCreate_None, TexCreate_RenderTargetable, false));
-		if (!GRenderTargetPool.FindFreeElement(RHICmdList, Desc, Dummy, TEXT("Dummy")))
-		{
-			RHICmdList.Transition(FRHITransitionInfo(Dummy->GetRenderTargetItem().TargetableTexture, ERHIAccess::Unknown, ERHIAccess::RTV));
-		}
-
-		RPInfo.ColorRenderTargets[0].Action = ERenderTargetActions::DontLoad_DontStore;
-		RPInfo.ColorRenderTargets[0].ArraySlice = 0;
-		RPInfo.ColorRenderTargets[0].MipIndex = 0;
-		RPInfo.ColorRenderTargets[0].RenderTarget = Dummy->GetRenderTargetItem().TargetableTexture;
+		FRDGTextureDesc DummyDesc = FRDGTextureDesc::Create2D(TileListGroupSize, PF_B8G8R8A8, FClearValueBinding::Black, TexCreate_RenderTargetable);
+		PassParameters->RenderTargets[0] = FRenderTargetBinding(GraphBuilder.CreateTexture(DummyDesc, TEXT("Dummy")), ERenderTargetLoadAction::ENoAction);
 	}
 
-	RHICmdList.BeginRenderPass(RPInfo, TEXT("ScatterTilesToObjects"));
-	{
-		FGraphicsPipelineStateInitializer GraphicsPSOInit;
-		RHICmdList.ApplyCachedRenderTargets(GraphicsPSOInit);
+	ClearUnusedGraphResources(VertexShader, &PassParameters->VS);
+	ClearUnusedGraphResources(PixelShader, &PassParameters->PS);
 
-		RHICmdList.SetViewport(0, 0, 0.0f, TileListGroupSize.X, TileListGroupSize.Y, 1.0f);
+	GraphBuilder.AddPass(
+		bCountingPass ? RDG_EVENT_NAME("CountTileObjectIntersections") : RDG_EVENT_NAME("CullTilesToObjects"),
+		PassParameters,
+		ERDGPassFlags::Raster,
+		[PassParameters, VertexShader, PixelShader, &View, TileListGroupSize, Parameters, ObjectIndirectArguments](FRHICommandListImmediate& RHICmdList)
+		{
+			FGraphicsPipelineStateInitializer GraphicsPSOInit;
+			RHICmdList.ApplyCachedRenderTargets(GraphicsPSOInit);
 
-		// Render backfaces since camera may intersect
-		GraphicsPSOInit.RasterizerState = View.bReverseCulling ? TStaticRasterizerState<FM_Solid, CM_CW>::GetRHI() : TStaticRasterizerState<FM_Solid, CM_CCW>::GetRHI();
-		GraphicsPSOInit.DepthStencilState = TStaticDepthStencilState<false, CF_Always>::GetRHI();
-		GraphicsPSOInit.BlendState = TStaticBlendState<>::GetRHI();
-		GraphicsPSOInit.PrimitiveType = PT_TriangleList;
+			RHICmdList.SetViewport(0, 0, 0.0f, TileListGroupSize.X, TileListGroupSize.Y, 1.0f);
 
-		GraphicsPSOInit.BoundShaderState.VertexDeclarationRHI = GetVertexDeclarationFVector4();
-		GraphicsPSOInit.BoundShaderState.VertexShaderRHI = VertexShader.GetVertexShader();
-		GraphicsPSOInit.BoundShaderState.PixelShaderRHI = PixelShader.GetPixelShader();
+			// Render backfaces since camera may intersect
+			GraphicsPSOInit.RasterizerState = View.bReverseCulling ? TStaticRasterizerState<FM_Solid, CM_CW>::GetRHI() : TStaticRasterizerState<FM_Solid, CM_CCW>::GetRHI();
+			GraphicsPSOInit.DepthStencilState = TStaticDepthStencilState<false, CF_Always>::GetRHI();
+			GraphicsPSOInit.BlendState = TStaticBlendState<>::GetRHI();
+			GraphicsPSOInit.PrimitiveType = PT_TriangleList;
 
-		SetGraphicsPipelineState(RHICmdList, GraphicsPSOInit);
+			GraphicsPSOInit.BoundShaderState.VertexDeclarationRHI = GetVertexDeclarationFVector4();
+			GraphicsPSOInit.BoundShaderState.VertexShaderRHI = VertexShader.GetVertexShader();
+			GraphicsPSOInit.BoundShaderState.PixelShaderRHI = PixelShader.GetPixelShader();
 
-		VertexShader->SetParameters(RHICmdList, View, DistanceFieldSceneData, Parameters);
-		PixelShader->SetParameters(RHICmdList, View, DistanceFieldSceneData, FVector2D(TileListGroupSize.X, TileListGroupSize.Y), Parameters);
+			SetGraphicsPipelineState(RHICmdList, GraphicsPSOInit);
 
-		RHICmdList.SetStreamSource(0, StencilingGeometry::GLowPolyStencilSphereVertexBuffer.VertexBufferRHI, 0);
+			VertexShader->SetParameters(RHICmdList, View, Parameters);
+			PixelShader->SetParameters(RHICmdList, View, FVector2D(TileListGroupSize.X, TileListGroupSize.Y), Parameters);
 
-		RHICmdList.DrawIndexedPrimitiveIndirect(
-			StencilingGeometry::GLowPolyStencilSphereIndexBuffer.IndexBufferRHI,
-			GAOCulledObjectBuffers.Buffers.ObjectIndirectArguments.Buffer,
-			0);
-	}
-	RHICmdList.EndRenderPass();
-	RHICmdList.Transition(MakeArrayView(TransitionInfos.GetData(), TransitionInfos.Num()));
+			SetShaderParameters(RHICmdList, VertexShader, VertexShader.GetVertexShader(), PassParameters->VS);
+			SetShaderParameters(RHICmdList, PixelShader, PixelShader.GetPixelShader(), PassParameters->PS);
+
+			RHICmdList.SetStreamSource(0, StencilingGeometry::GLowPolyStencilSphereVertexBuffer.VertexBufferRHI, 0);
+
+			RHICmdList.DrawIndexedPrimitiveIndirect(
+				StencilingGeometry::GLowPolyStencilSphereIndexBuffer.IndexBufferRHI,
+				ObjectIndirectArguments->GetIndirectRHICallBuffer(),
+				0);
+		});
 }
 
 FIntPoint GetTileListGroupSizeForView(const FViewInfo& View)
@@ -526,13 +393,18 @@ FIntPoint GetTileListGroupSizeForView(const FViewInfo& View)
 		FMath::DivideAndRoundUp(FMath::Max(View.ViewRect.Size().Y / GAODownsampleFactor, 1), GDistanceFieldAOTileSizeY));
 }
 
-BEGIN_SHADER_PARAMETER_STRUCT(FBuildTileObjectListParameters, )
-	SHADER_PARAMETER_RDG_UNIFORM_BUFFER(FSceneTextureUniformParameters, SceneTextures)
-	RDG_TEXTURE_ACCESS(DistanceFieldNormal, ERHIAccess::SRVCompute)
-END_SHADER_PARAMETER_STRUCT()
-
-void BuildTileObjectLists(FRDGBuilder& GraphBuilder, FScene* Scene, TArray<FViewInfo>& Views, FRDGTextureRef DistanceFieldNormal, const FDistanceFieldAOParameters& Parameters)
+void BuildTileObjectLists(
+	FRDGBuilder& GraphBuilder,
+	FScene* Scene,
+	TArray<FViewInfo>& Views,
+	FRDGBufferRef ObjectIndirectArguments,
+	const FDistanceFieldCulledObjectBufferParameters& CulledObjectBufferParameters,
+	FTileIntersectionParameters TileIntersectionParameters,
+	FRDGTextureRef DistanceFieldNormal,
+	const FDistanceFieldAOParameters& Parameters)
 {
+	ensure(GAOScatterTileCulling);
+
 	RDG_EVENT_SCOPE(GraphBuilder, "BuildTileList");
 
 	TRDGUniformBufferRef<FSceneTextureUniformParameters> SceneTexturesUniformBuffer = CreateSceneTextureUniformBuffer(GraphBuilder, Scene->GetFeatureLevel());
@@ -542,97 +414,69 @@ void BuildTileObjectLists(FRDGBuilder& GraphBuilder, FScene* Scene, TArray<FView
 		const FViewInfo& View = Views[ViewIndex];
 		RDG_GPU_MASK_SCOPE(GraphBuilder, View.GPUMask);
 
-		auto* PassParameters = GraphBuilder.AllocParameters<FBuildTileObjectListParameters>();
-		PassParameters->SceneTextures = SceneTexturesUniformBuffer;
-		PassParameters->DistanceFieldNormal = DistanceFieldNormal;
+		const FIntPoint TileListGroupSize = GetTileListGroupSizeForView(View);
 
-		// TODO(RDG) Port all of these into respective passes.
-		GraphBuilder.AddPass(
-			{},
-			PassParameters,
-			ERDGPassFlags::Compute | ERDGPassFlags::NeverCull,
-			[Scene, &View, DistanceFieldNormal, Parameters](FRHICommandListImmediate& RHICmdList)
 		{
-			const FIntPoint TileListGroupSize = GetTileListGroupSizeForView(View);
-
-			FTileIntersectionResources*& TileIntersectionResources = ((FSceneViewState*)View.State)->AOTileIntersectionResources;
-
-			if (!TileIntersectionResources
-				|| !TileIntersectionResources->IsInitialized()
-				|| !TileIntersectionResources->HasAllocatedEnoughFor(TileListGroupSize, Scene->DistanceFieldSceneData.NumObjectsInBuffer)
-				|| GFastVRamConfig.bDirty)
-			{
-				if (TileIntersectionResources)
-				{
-					TileIntersectionResources->ReleaseResource();
-				}
-				else
+			auto* PassParameters = GraphBuilder.AllocParameters<FBuildTileConesCS::FParameters>();
+			PassParameters->RWTileConeAxisAndCos = TileIntersectionParameters.RWTileConeAxisAndCos;
+			PassParameters->RWTileConeDepthRanges = TileIntersectionParameters.RWTileConeDepthRanges;
+			PassParameters->SceneTextures = SceneTexturesUniformBuffer;
+			PassParameters->DistanceFieldNormalTexture = DistanceFieldNormal;
+			PassParameters->DistanceFieldNormalSampler = TStaticSamplerState<SF_Point, AM_Clamp, AM_Clamp, AM_Clamp>::GetRHI();
+
+			auto ComputeShader = View.ShaderMap->GetShader<FBuildTileConesCS>();
+
+			ClearUnusedGraphResources(ComputeShader, PassParameters);
+
+			GraphBuilder.AddPass(
+				RDG_EVENT_NAME("BuildTileCones"),
+				PassParameters,
+				ERDGPassFlags::Compute,
+				[PassParameters, ComputeShader, Scene, &View, TileListGroupSize, Parameters](FRHICommandListImmediate& RHICmdList)
 				{
-					TileIntersectionResources = new FTileIntersectionResources(!IsMetalPlatform(GShaderPlatformForFeatureLevel[View.FeatureLevel]));
-				}
+					RHICmdList.SetComputeShader(ComputeShader.GetComputeShader());
 
-				TileIntersectionResources->SetupParameters(TileListGroupSize, Scene->DistanceFieldSceneData.NumObjectsInBuffer);
-				TileIntersectionResources->InitResource();
-			}
-			TileIntersectionResources->AcquireTransientResource();
+					ComputeShader->SetParameters(RHICmdList, View, Scene, FVector2D(TileListGroupSize.X, TileListGroupSize.Y), Parameters);
+					SetShaderParameters(RHICmdList, ComputeShader, ComputeShader.GetComputeShader(), *PassParameters);
 
-			if (GAOScatterTileCulling)
-			{
-				{
-					SCOPED_DRAW_EVENT(RHICmdList, BuildTileCones);
-					TShaderMapRef<FBuildTileConesCS> ComputeShader(View.ShaderMap);
-
-					RHICmdList.SetComputeShader(ComputeShader.GetComputeShader());
-					ComputeShader->SetParameters(RHICmdList, View, DistanceFieldNormal->GetRHI(), Scene, FVector2D(TileListGroupSize.X, TileListGroupSize.Y), Parameters);
 					DispatchComputeShader(RHICmdList, ComputeShader.GetShader(), TileListGroupSize.X, TileListGroupSize.Y, 1);
 
-					ComputeShader->UnsetParameters(RHICmdList, View);
-				}
+					UnsetShaderUAVs(RHICmdList, ComputeShader, ComputeShader.GetComputeShader());
+				});
+		}
 
-				{
-					SCOPED_DRAW_EVENT(RHICmdList, CountTileObjectIntersections);
+		// Start at 0 tiles per object
+		AddClearUAVPass(GraphBuilder, TileIntersectionParameters.RWNumCulledTilesArray, 0);
 
-					// Start at 0 tiles per object
-					RHICmdList.Transition(FRHITransitionInfo(TileIntersectionResources->NumCulledTilesArray.UAV, ERHIAccess::Unknown, ERHIAccess::UAVCompute));
-					RHICmdList.ClearUAVUint(TileIntersectionResources->NumCulledTilesArray.UAV, FUintVector4(0, 0, 0, 0));
+		// Rasterize object bounding shapes and intersect with screen tiles to compute how many tiles intersect each object
+		ScatterTilesToObjects(GraphBuilder, true, View, Scene->DistanceFieldSceneData, TileListGroupSize, Parameters, ObjectIndirectArguments, CulledObjectBufferParameters, TileIntersectionParameters, SceneTexturesUniformBuffer);
 
-					// Rasterize object bounding shapes and intersect with screen tiles to compute how many tiles intersect each object
-					ScatterTilesToObjects(RHICmdList, true, View, Scene->DistanceFieldSceneData, TileListGroupSize, Parameters);
-				}
+		// Start at 0 threadgroups
+		AddClearUAVPass(GraphBuilder, TileIntersectionParameters.RWObjectTilesIndirectArguments, 0);
 
-				{
-					SCOPED_DRAW_EVENT(RHICmdList, ComputeStartOffsets);
-					// Start at 0 threadgroups
-					RHICmdList.Transition(FRHITransitionInfo(TileIntersectionResources->ObjectTilesIndirectArguments.UAV, ERHIAccess::Unknown, ERHIAccess::UAVCompute));
-					RHICmdList.ClearUAVUint(TileIntersectionResources->ObjectTilesIndirectArguments.UAV, FUintVector4(0, 0, 0, 0));
-
-					// Accumulate how many cone trace threadgroups we should dispatch, and also compute the start offset for each object's culled tile data
-					TShaderMapRef<FComputeCulledTilesStartOffsetCS> ComputeShader(View.ShaderMap);
-					const uint32 GroupSize = FMath::DivideAndRoundUp<uint32>(Scene->DistanceFieldSceneData.NumObjectsInBuffer, ComputeStartOffsetGroupSize);
-					// Must write to RWObjectTilesIndirectArguments
-					check(GroupSize != 0);
-					RHICmdList.SetComputeShader(ComputeShader.GetComputeShader());
-					ComputeShader->SetParameters(RHICmdList, View, Scene->DistanceFieldSceneData);
-					DispatchComputeShader(RHICmdList, ComputeShader.GetShader(), GroupSize, 1, 1);
+		{
+			auto* PassParameters = GraphBuilder.AllocParameters<FComputeCulledTilesStartOffsetCS::FParameters>();
+			PassParameters->View = View.ViewUniformBuffer;
+			PassParameters->TileIntersectionParameters = TileIntersectionParameters;
+			PassParameters->DistanceFieldCulledObjectBuffers = CulledObjectBufferParameters;
+			PassParameters->DistanceFieldAtlas = DistanceField::SetupAtlasParameters(Scene->DistanceFieldSceneData);
+			PassParameters->SceneTextures = SceneTexturesUniformBuffer;
+
+			auto ComputeShader = View.ShaderMap->GetShader<FComputeCulledTilesStartOffsetCS>();
+			const int32 GroupSize = FMath::DivideAndRoundUp<uint32>(Scene->DistanceFieldSceneData.NumObjectsInBuffer, ComputeStartOffsetGroupSize);
+
+			FComputeShaderUtils::AddPass(
+				GraphBuilder,
+				RDG_EVENT_NAME("ComputeStartOffsets"),
+				ComputeShader,
+				PassParameters,
+				FIntVector(GroupSize, 1, 1));
+		}
 
-					ComputeShader->UnsetParameters(RHICmdList, View);
-				}
+		// Start at 0 tiles per object
+		AddClearUAVPass(GraphBuilder, TileIntersectionParameters.RWNumCulledTilesArray, 0);
 
-				{
-					SCOPED_DRAW_EVENT(RHICmdList, CullTilesToObjects);
-
-					// Start at 0 tiles per object
-					RHICmdList.Transition(FRHITransitionInfo(TileIntersectionResources->NumCulledTilesArray.UAV, ERHIAccess::Unknown, ERHIAccess::UAVCompute));
-					RHICmdList.ClearUAVUint(TileIntersectionResources->NumCulledTilesArray.UAV, FUintVector4(0, 0, 0, 0));
-
-					// Rasterize object bounding shapes and intersect with screen tiles, and write out intersecting tile indices for the cone tracing pass
-					ScatterTilesToObjects(RHICmdList, false, View, Scene->DistanceFieldSceneData, TileListGroupSize, Parameters);
-				}
-			}
-			else
-			{
-				ensure(0);
-			}
-		});
+		// Rasterize object bounding shapes and intersect with screen tiles, and write out intersecting tile indices for the cone tracing pass
+		ScatterTilesToObjects(GraphBuilder, false, View, Scene->DistanceFieldSceneData, TileListGroupSize, Parameters, ObjectIndirectArguments, CulledObjectBufferParameters, TileIntersectionParameters, SceneTexturesUniformBuffer);
 	}
 }
diff --git a/Engine/Source/Runtime/Renderer/Private/DistanceFieldObjectManagement.cpp b/Engine/Source/Runtime/Renderer/Private/DistanceFieldObjectManagement.cpp
--- a/Engine/Source/Runtime/Renderer/Private/DistanceFieldObjectManagement.cpp
+++ b/Engine/Source/Runtime/Renderer/Private/DistanceFieldObjectManagement.cpp
@@ -85,14 +85,14 @@ LLM_DEFINE_TAG(DistanceFields, NAME_None, NAME_None, GET_STATFNAME(STAT_Distance
 #endif // ENABLE_LOW_LEVEL_MEM_TRACKER
 
 // Must match equivalent shader defines
-template<> int32 TDistanceFieldObjectBuffers<DFPT_SignedDistanceField>::ObjectDataStride = 9;
-template<> int32 TDistanceFieldObjectBuffers<DFPT_SignedDistanceField>::ObjectBoundsStride = 2;
-template<> int32 TDistanceFieldObjectBuffers<DFPT_HeightField>::ObjectDataStride = 6;
-template<> int32 TDistanceFieldObjectBuffers<DFPT_HeightField>::ObjectBoundsStride = 2;
-template<> int32 TDistanceFieldCulledObjectBuffers<DFPT_SignedDistanceField>::ObjectDataStride = TDistanceFieldObjectBuffers<DFPT_SignedDistanceField>::ObjectDataStride;
-template<> int32 TDistanceFieldCulledObjectBuffers<DFPT_SignedDistanceField>::ObjectBoxBoundsStride = 5;
-template<> int32 TDistanceFieldCulledObjectBuffers<DFPT_HeightField>::ObjectDataStride = 6;
-template<> int32 TDistanceFieldCulledObjectBuffers<DFPT_HeightField>::ObjectBoxBoundsStride = 5;
+template<> int32 TDistanceFieldObjectBuffers<DFPT_SignedDistanceField>::ObjectDataStride = GDistanceFieldObjectDataStride;
+template<> int32 TDistanceFieldObjectBuffers<DFPT_SignedDistanceField>::ObjectBoundsStride = GDistanceFieldObjectBoundsStride;
+template<> int32 TDistanceFieldObjectBuffers<DFPT_HeightField>::ObjectDataStride = GHeightFieldObjectDataStride;
+template<> int32 TDistanceFieldObjectBuffers<DFPT_HeightField>::ObjectBoundsStride = GHeightFieldObjectBoundsStride;
+template<> int32 TDistanceFieldCulledObjectBuffers<DFPT_SignedDistanceField>::ObjectDataStride = GDistanceFieldCulledObjectDataStride;
+template<> int32 TDistanceFieldCulledObjectBuffers<DFPT_SignedDistanceField>::ObjectBoxBoundsStride = GDistanceFieldCulledObjectBoxBoundsStride;
+template<> int32 TDistanceFieldCulledObjectBuffers<DFPT_HeightField>::ObjectDataStride = GHeightFieldCulledObjectDataStride;
+template<> int32 TDistanceFieldCulledObjectBuffers<DFPT_HeightField>::ObjectBoxBoundsStride = GHeightFieldCulledObjectBoxBoundsStride;
 
 template <EDistanceFieldPrimitiveType PrimitiveType>
 void TDistanceFieldObjectBuffers<PrimitiveType>::Initialize()
diff --git a/Engine/Source/Runtime/Renderer/Private/DistanceFieldScreenGridLighting.cpp b/Engine/Source/Runtime/Renderer/Private/DistanceFieldScreenGridLighting.cpp
--- a/Engine/Source/Runtime/Renderer/Private/DistanceFieldScreenGridLighting.cpp
+++ b/Engine/Source/Runtime/Renderer/Private/DistanceFieldScreenGridLighting.cpp
@@ -63,18 +63,22 @@ FVector2D GetJitterOffset(int32 SampleIndex)
 	return FVector2D(0, 0);
 }
 
-void FAOScreenGridResources::InitDynamicRHI()
-{
-	//@todo - 2d textures
-	const uint32 FastVRamFlag = GFastVRamConfig.DistanceFieldAOScreenGridResources | (IsTransientResourceBufferAliasingEnabled() ? BUF_Transient : BUF_None);
-	ScreenGridConeVisibility.Initialize(TEXT("ScreenGridConeVisibility"), sizeof(uint32), NumConeSampleDirections * ScreenGridDimensions.X * ScreenGridDimensions.Y, PF_R32_UINT, BUF_Static | FastVRamFlag);
-}
-
 class FConeTraceScreenGridObjectOcclusionCS : public FGlobalShader
 {
 	DECLARE_GLOBAL_SHADER(FConeTraceScreenGridObjectOcclusionCS);
+
 public:
 
+	BEGIN_SHADER_PARAMETER_STRUCT(FParameters, )
+		SHADER_PARAMETER_STRUCT_INCLUDE(FDistanceFieldCulledObjectBufferParameters, DistanceFieldCulledObjectBuffers)
+		SHADER_PARAMETER_STRUCT_INCLUDE(FDistanceFieldAtlasParameters, DistanceFieldAtlas)
+		SHADER_PARAMETER_STRUCT_INCLUDE(FTileIntersectionParameters, TileIntersectionParameters)
+		SHADER_PARAMETER_STRUCT_INCLUDE(FAOScreenGridParameters, AOScreenGridParameters)
+		RDG_BUFFER_ACCESS(ObjectTilesIndirectArguments, ERHIAccess::IndirectArgs)
+		SHADER_PARAMETER_RDG_UNIFORM_BUFFER(FSceneTextureUniformParameters, SceneTextures)
+		RDG_TEXTURE_ACCESS(DistanceFieldNormal, ERHIAccess::SRVCompute)
+	END_SHADER_PARAMETER_STRUCT()
+
 	class FUseGlobalDistanceField : SHADER_PERMUTATION_BOOL("USE_GLOBAL_DISTANCE_FIELD");
 	using FPermutationDomain = TShaderPermutationDomain<FUseGlobalDistanceField>;
 
@@ -86,7 +90,7 @@ public:
 	static void ModifyCompilationEnvironment(const FGlobalShaderPermutationParameters& Parameters, FShaderCompilerEnvironment& OutEnvironment)
 	{
 		FGlobalShader::ModifyCompilationEnvironment(Parameters, OutEnvironment);
-		FTileIntersectionParameters::ModifyCompilationEnvironment(Parameters.Platform, OutEnvironment);
+		TileIntersectionModifyCompilationEnvironment(Parameters.Platform, OutEnvironment);
 
 		// To reduce shader compile time of compute shaders with shared memory, doesn't have an impact on generated code with current compiler (June 2010 DX SDK)
 		OutEnvironment.CompilerFlags.Add(CFLAG_StandardOptimization);
@@ -95,16 +99,12 @@ public:
 	FConeTraceScreenGridObjectOcclusionCS(const ShaderMetaType::CompiledShaderInitializerType& Initializer)
 		: FGlobalShader(Initializer)
 	{
-		ObjectParameters.Bind(Initializer.ParameterMap);
+		BindForLegacyShaderParameters<FParameters>(this, Initializer.PermutationId, Initializer.ParameterMap, false);
 		AOParameters.Bind(Initializer.ParameterMap);
 		ScreenGridParameters.Bind(Initializer.ParameterMap);
 		GlobalDistanceFieldParameters.Bind(Initializer.ParameterMap);
-		TileConeDepthRanges.Bind(Initializer.ParameterMap, TEXT("TileConeDepthRanges"));
-		TileIntersectionParameters.Bind(Initializer.ParameterMap);
 		TanConeHalfAngle.Bind(Initializer.ParameterMap, TEXT("TanConeHalfAngle"));
 		BentNormalNormalizeFactor.Bind(Initializer.ParameterMap, TEXT("BentNormalNormalizeFactor"));
-		ScreenGridConeVisibility.Bind(Initializer.ParameterMap, TEXT("ScreenGridConeVisibility"));
-		ConeDepthVisibilityFunction.Bind(Initializer.ParameterMap, TEXT("ConeDepthVisibilityFunction"));
 	}
 
 	FConeTraceScreenGridObjectOcclusionCS()
@@ -114,7 +114,6 @@ public:
 	void SetParameters(
 		FRHICommandList& RHICmdList, 
 		const FViewInfo& View, 
-		const FDistanceFieldSceneData& DistanceFieldSceneData,
 		FRHITexture* DistanceFieldNormal,  
 		const FDistanceFieldAOParameters& Parameters,
 		bool bUseGlobalDistanceField,
@@ -123,8 +122,6 @@ public:
 		FRHIComputeShader* ShaderRHI = RHICmdList.GetBoundComputeShader();
 		FGlobalShader::SetParameters<FViewUniformShaderParameters>(RHICmdList, ShaderRHI, View.ViewUniformBuffer);
 
-		ObjectParameters.Set(RHICmdList, ShaderRHI, GAOCulledObjectBuffers.Buffers, DistanceFieldSceneData);
-
 		AOParameters.Set(RHICmdList, ShaderRHI, Parameters);
 		ScreenGridParameters.Set(RHICmdList, ShaderRHI, View, DistanceFieldNormal);
 
@@ -145,13 +142,6 @@ public:
 
 		SetUniformBufferParameterImmediate(RHICmdList, ShaderRHI, GetUniformBufferParameter<FAOSampleData2>(), AOSampleData);
 
-		FTileIntersectionResources* TileIntersectionResources = View.ViewState->AOTileIntersectionResources;
-		SetSRVParameter(RHICmdList, ShaderRHI, TileConeDepthRanges, TileIntersectionResources->TileConeDepthRanges.SRV);
-
-		RHICmdList.Transition(FRHITransitionInfo(TileIntersectionResources->CulledTileDataArray.UAV, ERHIAccess::Unknown, ERHIAccess::SRVMask));
-
-		TileIntersectionParameters.Set(RHICmdList, ShaderRHI, *TileIntersectionResources);
-
 		extern float GAOConeHalfAngle;
 		SetShaderValue(RHICmdList, ShaderRHI, TanConeHalfAngle, FMath::Tan(GAOConeHalfAngle));
 
@@ -168,41 +158,37 @@ public:
 		float ConeSampleAverageSize = (UnoccludedVector / NumConeSampleDirections).Size();
 		float BentNormalNormalizeFactorValue = ConeSampleAverageSize ? 1.0f / ConeSampleAverageSize : 0.f;
 		SetShaderValue(RHICmdList, ShaderRHI, BentNormalNormalizeFactor, BentNormalNormalizeFactorValue);
-
-		FAOScreenGridResources* ScreenGridResources = View.ViewState->AOScreenGridResources;
-
-		ScreenGridConeVisibility.SetBuffer(RHICmdList, ShaderRHI, ScreenGridResources->ScreenGridConeVisibility);
-	}
-
-	void UnsetParameters(FRHICommandList& RHICmdList, const FViewInfo& View)
-	{
-		ScreenGridConeVisibility.UnsetUAV(RHICmdList, RHICmdList.GetBoundComputeShader());
-		ConeDepthVisibilityFunction.UnsetUAV(RHICmdList, RHICmdList.GetBoundComputeShader());
 	}
 
 private:
-	LAYOUT_FIELD((TDistanceFieldCulledObjectBufferParameters<DFPT_SignedDistanceField>), ObjectParameters);
 	LAYOUT_FIELD(FAOParameters, AOParameters);
 	LAYOUT_FIELD(FScreenGridParameters, ScreenGridParameters);
 	LAYOUT_FIELD(FGlobalDistanceFieldParameters, GlobalDistanceFieldParameters);
-	LAYOUT_FIELD(FShaderResourceParameter, TileConeDepthRanges);
-	LAYOUT_FIELD(FTileIntersectionParameters, TileIntersectionParameters);
 	LAYOUT_FIELD(FShaderParameter, TanConeHalfAngle);
 	LAYOUT_FIELD(FShaderParameter, BentNormalNormalizeFactor);
-	LAYOUT_FIELD(FRWShaderParameter, ScreenGridConeVisibility);
-	LAYOUT_FIELD(FRWShaderParameter, ConeDepthVisibilityFunction);
 };
 
 IMPLEMENT_GLOBAL_SHADER(FConeTraceScreenGridObjectOcclusionCS, "/Engine/Private/DistanceFieldScreenGridLighting.usf", "ConeTraceObjectOcclusionCS", SF_Compute);
 
 const int32 GConeTraceGlobalDFTileSize = 8;
 
-template<bool bConeTraceObjects>
-class TConeTraceScreenGridGlobalOcclusionCS : public FGlobalShader
+class FConeTraceScreenGridGlobalOcclusionCS : public FGlobalShader
 {
-	DECLARE_SHADER_TYPE(TConeTraceScreenGridGlobalOcclusionCS,Global)
+	DECLARE_GLOBAL_SHADER(FConeTraceScreenGridGlobalOcclusionCS);
 public:
 
+	BEGIN_SHADER_PARAMETER_STRUCT(FParameters, )
+		SHADER_PARAMETER_STRUCT_INCLUDE(FDistanceFieldCulledObjectBufferParameters, DistanceFieldCulledObjectBuffers)
+		SHADER_PARAMETER_STRUCT_INCLUDE(FDistanceFieldAtlasParameters, DistanceFieldAtlas)
+		SHADER_PARAMETER_STRUCT_INCLUDE(FAOScreenGridParameters, AOScreenGridParameters)
+		SHADER_PARAMETER_RDG_BUFFER_SRV(Buffer<FVector4>, TileConeDepthRanges)
+		SHADER_PARAMETER_RDG_UNIFORM_BUFFER(FSceneTextureUniformParameters, SceneTextures)
+		RDG_TEXTURE_ACCESS(DistanceFieldNormal, ERHIAccess::SRVCompute)
+	END_SHADER_PARAMETER_STRUCT()
+
+	class FConeTraceObjects : SHADER_PERMUTATION_BOOL("CONE_TRACE_OBJECTS");
+	using FPermutationDomain = TShaderPermutationDomain<FConeTraceObjects>;
+
 	static bool ShouldCompilePermutation(const FGlobalShaderPermutationParameters& Parameters)
 	{
 		return IsFeatureLevelSupported(Parameters.Platform, ERHIFeatureLevel::SM5) && DoesPlatformSupportDistanceFieldAO(Parameters.Platform) && IsUsingDistanceFields(Parameters.Platform);
@@ -211,7 +197,6 @@ public:
 	static void ModifyCompilationEnvironment(const FGlobalShaderPermutationParameters& Parameters, FShaderCompilerEnvironment& OutEnvironment)
 	{
 		FGlobalShader::ModifyCompilationEnvironment(Parameters,OutEnvironment);
-		OutEnvironment.SetDefine(TEXT("CONE_TRACE_OBJECTS"), bConeTraceObjects);
 		OutEnvironment.SetDefine(TEXT("CONE_TRACE_GLOBAL_DISPATCH_SIZEX"), GConeTraceGlobalDFTileSize);
 		OutEnvironment.SetDefine(TEXT("OUTPUT_VISIBILITY_DIRECTLY"), TEXT("1"));
 		OutEnvironment.SetDefine(TEXT("TRACE_DOWNSAMPLE_FACTOR"), GConeTraceDownsampleFactor);
@@ -221,28 +206,25 @@ public:
 		OutEnvironment.CompilerFlags.Add(CFLAG_StandardOptimization);
 	}
 
-	TConeTraceScreenGridGlobalOcclusionCS(const ShaderMetaType::CompiledShaderInitializerType& Initializer)
+	FConeTraceScreenGridGlobalOcclusionCS(const ShaderMetaType::CompiledShaderInitializerType& Initializer)
 		: FGlobalShader(Initializer)
 	{
-		ObjectParameters.Bind(Initializer.ParameterMap);
+		BindForLegacyShaderParameters<FParameters>(this, Initializer.PermutationId, Initializer.ParameterMap, false);
 		AOParameters.Bind(Initializer.ParameterMap);
 		ScreenGridParameters.Bind(Initializer.ParameterMap);
 		GlobalDistanceFieldParameters.Bind(Initializer.ParameterMap);
-		TileConeDepthRanges.Bind(Initializer.ParameterMap, TEXT("TileConeDepthRanges"));
 		TileListGroupSize.Bind(Initializer.ParameterMap, TEXT("TileListGroupSize"));
 		TanConeHalfAngle.Bind(Initializer.ParameterMap, TEXT("TanConeHalfAngle"));
 		BentNormalNormalizeFactor.Bind(Initializer.ParameterMap, TEXT("BentNormalNormalizeFactor"));
-		ScreenGridConeVisibility.Bind(Initializer.ParameterMap, TEXT("ScreenGridConeVisibility"));
 	}
 
-	TConeTraceScreenGridGlobalOcclusionCS()
+	FConeTraceScreenGridGlobalOcclusionCS()
 	{
 	}
 
 	void SetParameters(
 		FRHICommandList& RHICmdList, 
-		const FViewInfo& View, 
-		const FDistanceFieldSceneData& DistanceFieldSceneData,
+		const FViewInfo& View,
 		FIntPoint TileListGroupSizeValue, 
 		FRHITexture* DistanceFieldNormal, 
 		const FDistanceFieldAOParameters& Parameters,
@@ -251,8 +233,6 @@ public:
 		FRHIComputeShader* ShaderRHI = RHICmdList.GetBoundComputeShader();
 		FGlobalShader::SetParameters<FViewUniformShaderParameters>(RHICmdList, ShaderRHI, View.ViewUniformBuffer);
 
-		ObjectParameters.Set(RHICmdList, ShaderRHI, GAOCulledObjectBuffers.Buffers, DistanceFieldSceneData);
-
 		AOParameters.Set(RHICmdList, ShaderRHI, Parameters);
 		ScreenGridParameters.Set(RHICmdList, ShaderRHI, View, DistanceFieldNormal);
 		GlobalDistanceFieldParameters.Set(RHICmdList, ShaderRHI, GlobalDistanceFieldInfo.ParameterData);
@@ -269,12 +249,6 @@ public:
 
 		SetUniformBufferParameterImmediate(RHICmdList, ShaderRHI, GetUniformBufferParameter<FAOSampleData2>(), AOSampleData);
 
-		if (UseAOObjectDistanceField())
-		{
-			FTileIntersectionResources* TileIntersectionResources = View.ViewState->AOTileIntersectionResources;
-			SetSRVParameter(RHICmdList, ShaderRHI, TileConeDepthRanges, TileIntersectionResources->TileConeDepthRanges.SRV);
-		}
-
 		SetShaderValue(RHICmdList, ShaderRHI, TileListGroupSize, TileListGroupSizeValue);
 
 		extern float GAOConeHalfAngle;
@@ -290,52 +264,32 @@ public:
 		float ConeSampleAverageSize = (UnoccludedVector / NumConeSampleDirections).Size();
 		float BentNormalNormalizeFactorValue = ConeSampleAverageSize ? 1.0f / ConeSampleAverageSize : 0.f;
 		SetShaderValue(RHICmdList, ShaderRHI, BentNormalNormalizeFactor, BentNormalNormalizeFactorValue);
-
-		FAOScreenGridResources* ScreenGridResources = View.ViewState->AOScreenGridResources;
-			
-		RHICmdList.Transition(FRHITransitionInfo(ScreenGridResources->ScreenGridConeVisibility.UAV, ERHIAccess::Unknown, ERHIAccess::UAVMask));
-
-		// Note: no transition, want to overlap object cone tracing and global DF cone tracing since both shaders use atomics to ScreenGridConeVisibility
-		RHICmdList.BeginUAVOverlap(ScreenGridResources->ScreenGridConeVisibility.UAV);
-
-		ScreenGridConeVisibility.SetBuffer(RHICmdList, ShaderRHI, ScreenGridResources->ScreenGridConeVisibility);
-	}
-
-	void UnsetParameters(FRHICommandList& RHICmdList, const FViewInfo& View)
-	{
-		ScreenGridConeVisibility.UnsetUAV(RHICmdList, RHICmdList.GetBoundComputeShader());
-		RHICmdList.EndUAVOverlap(View.ViewState->AOScreenGridResources->ScreenGridConeVisibility.UAV);
-
-		RHICmdList.Transition(FRHITransitionInfo(View.ViewState->AOScreenGridResources->ScreenGridConeVisibility.UAV, ERHIAccess::UAVMask, ERHIAccess::SRVMask));
 	}
 
 private:
-	LAYOUT_FIELD((TDistanceFieldCulledObjectBufferParameters<DFPT_SignedDistanceField>), ObjectParameters);
 	LAYOUT_FIELD(FAOParameters, AOParameters);
 	LAYOUT_FIELD(FScreenGridParameters, ScreenGridParameters);
 	LAYOUT_FIELD(FGlobalDistanceFieldParameters, GlobalDistanceFieldParameters);
-	LAYOUT_FIELD(FShaderResourceParameter, TileConeDepthRanges);
 	LAYOUT_FIELD(FShaderParameter, TileListGroupSize);
 	LAYOUT_FIELD(FShaderParameter, TanConeHalfAngle);
 	LAYOUT_FIELD(FShaderParameter, BentNormalNormalizeFactor);
-	LAYOUT_FIELD(FRWShaderParameter, ScreenGridConeVisibility);
 };
 
-#define IMPLEMENT_CONETRACE_GLOBAL_CS_TYPE(bConeTraceObjects) \
-	typedef TConeTraceScreenGridGlobalOcclusionCS<bConeTraceObjects> TConeTraceScreenGridGlobalOcclusionCS##bConeTraceObjects; \
-	IMPLEMENT_SHADER_TYPE(template<>,TConeTraceScreenGridGlobalOcclusionCS##bConeTraceObjects,TEXT("/Engine/Private/DistanceFieldScreenGridLighting.usf"),TEXT("ConeTraceGlobalOcclusionCS"),SF_Compute);
-
-IMPLEMENT_CONETRACE_GLOBAL_CS_TYPE(true)
-IMPLEMENT_CONETRACE_GLOBAL_CS_TYPE(false)
-
-
+IMPLEMENT_GLOBAL_SHADER(FConeTraceScreenGridGlobalOcclusionCS, "/Engine/Private/DistanceFieldScreenGridLighting.usf", "ConeTraceGlobalOcclusionCS", SF_Compute);
 
 const int32 GCombineConesSizeX = 8;
 
 class FCombineConeVisibilityCS : public FGlobalShader
 {
-	DECLARE_SHADER_TYPE(FCombineConeVisibilityCS,Global)
+	DECLARE_GLOBAL_SHADER(FCombineConeVisibilityCS);
 public:
+
+	BEGIN_SHADER_PARAMETER_STRUCT(FParameters, )
+		SHADER_PARAMETER_STRUCT_INCLUDE(FAOScreenGridParameters, AOScreenGridParameters)
+		SHADER_PARAMETER_RDG_TEXTURE_UAV(RWTexture2D, RWDistanceFieldBentNormal)
+		RDG_TEXTURE_ACCESS(DistanceFieldNormal, ERHIAccess::SRVCompute)
+	END_SHADER_PARAMETER_STRUCT()
+
 	static bool ShouldCompilePermutation(const FGlobalShaderPermutationParameters& Parameters)
 	{
 		return IsFeatureLevelSupported(Parameters.Platform, ERHIFeatureLevel::SM5) && DoesPlatformSupportDistanceFieldAO(Parameters.Platform) && IsUsingDistanceFields(Parameters.Platform);
@@ -354,10 +308,9 @@ public:
 	FCombineConeVisibilityCS(const ShaderMetaType::CompiledShaderInitializerType& Initializer)
 		: FGlobalShader(Initializer)
 	{
+		BindForLegacyShaderParameters<FParameters>(this, Initializer.PermutationId, Initializer.ParameterMap, false);
 		ScreenGridParameters.Bind(Initializer.ParameterMap);
 		BentNormalNormalizeFactor.Bind(Initializer.ParameterMap, TEXT("BentNormalNormalizeFactor"));
-		ScreenGridConeVisibility.Bind(Initializer.ParameterMap, TEXT("ScreenGridConeVisibility"));
-		DistanceFieldBentNormal.Bind(Initializer.ParameterMap, TEXT("DistanceFieldBentNormal"));
 		ConeBufferMax.Bind(Initializer.ParameterMap, TEXT("ConeBufferMax"));
 		DFNormalBufferUVMax.Bind(Initializer.ParameterMap, TEXT("DFNormalBufferUVMax"));
 	}
@@ -366,12 +319,7 @@ public:
 	{
 	}
 
-	void SetParameters(
-		FRHICommandList& RHICmdList, 
-		const FViewInfo& View, 
-		FRHITexture* DistanceFieldNormal, 
-		FRHITexture* DownsampledBentNormal,
-		FRHIUnorderedAccessView* DownsampleBentNormalUAV)
+	void SetParameters(FRHICommandList& RHICmdList, const FViewInfo& View, FRHITexture* DistanceFieldNormal)
 	{
 		FRHIComputeShader* ShaderRHI = RHICmdList.GetBoundComputeShader();
 		FGlobalShader::SetParameters<FViewUniformShaderParameters>(RHICmdList, ShaderRHI, View.ViewUniformBuffer);
@@ -400,12 +348,6 @@ public:
 		float BentNormalNormalizeFactorValue = ConeSampleAverageSize ? 1.0f / ConeSampleAverageSize : 0.f;
 		SetShaderValue(RHICmdList, ShaderRHI, BentNormalNormalizeFactor, BentNormalNormalizeFactorValue);
 
-		FAOScreenGridResources* ScreenGridResources = View.ViewState->AOScreenGridResources;
-
-		DistanceFieldBentNormal.SetTexture(RHICmdList, ShaderRHI, DownsampledBentNormal, DownsampleBentNormalUAV);
-
-		SetSRVParameter(RHICmdList, ShaderRHI, ScreenGridConeVisibility, ScreenGridResources->ScreenGridConeVisibility.SRV);
-
 		FIntPoint const ConeBufferMaxValue(View.ViewRect.Width() / GAODownsampleFactor / GConeTraceDownsampleFactor - 1, View.ViewRect.Height() / GAODownsampleFactor / GConeTraceDownsampleFactor - 1);
 		SetShaderValue(RHICmdList, ShaderRHI, ConeBufferMax, ConeBufferMaxValue);
 
@@ -422,11 +364,9 @@ private:
 	LAYOUT_FIELD(FShaderParameter, BentNormalNormalizeFactor);
 	LAYOUT_FIELD(FShaderParameter, DFNormalBufferUVMax);
 	LAYOUT_FIELD(FShaderParameter, ConeBufferMax);
-	LAYOUT_FIELD(FShaderResourceParameter, ScreenGridConeVisibility);
-	LAYOUT_FIELD(FRWShaderParameter, DistanceFieldBentNormal);
 };
 
-IMPLEMENT_SHADER_TYPE(,FCombineConeVisibilityCS,TEXT("/Engine/Private/DistanceFieldScreenGridLighting.usf"),TEXT("CombineConeVisibilityCS"),SF_Compute);
+IMPLEMENT_GLOBAL_SHADER(FCombineConeVisibilityCS, "/Engine/Private/DistanceFieldScreenGridLighting.usf", "CombineConeVisibilityCS", SF_Compute);
 
 void PostProcessBentNormalAOScreenGrid(
 	FRDGBuilder& GraphBuilder, 
@@ -456,23 +396,13 @@ void PostProcessBentNormalAOScreenGrid(
 		Parameters);
 }
 
-// TODO(RDG) Replace with individual passes.
-BEGIN_SHADER_PARAMETER_STRUCT(FDistanceFieldAOScreenGridParameters, )
-	SHADER_PARAMETER_RDG_UNIFORM_BUFFER(FSceneTextureUniformParameters, SceneTextures)
-	RDG_TEXTURE_ACCESS(DistanceFieldNormal, ERHIAccess::SRVCompute)
-END_SHADER_PARAMETER_STRUCT()
-
-// TODO(RDG) Replace with individual passes.
-BEGIN_SHADER_PARAMETER_STRUCT(FDistanceFieldAOScreenGridCombineConeVisibilityParameters, )
-	RDG_TEXTURE_ACCESS(DownsampledBentNormal, ERHIAccess::UAVCompute)
-	SHADER_PARAMETER_RDG_TEXTURE_UAV(RWTexture2D, DownsampledBentNormalUAV)
-	RDG_TEXTURE_ACCESS(DistanceFieldNormal, ERHIAccess::SRVCompute)
-END_SHADER_PARAMETER_STRUCT()
-
 void FDeferredShadingSceneRenderer::RenderDistanceFieldAOScreenGrid(
 	FRDGBuilder& GraphBuilder,
 	const FSceneTextures& SceneTextures,
 	const FViewInfo& View,
+	const FDistanceFieldCulledObjectBufferParameters& CulledObjectBufferParameters,
+	FRDGBufferRef ObjectTilesIndirectArguments,
+	const FTileIntersectionParameters& TileIntersectionParameters,
 	const FDistanceFieldAOParameters& Parameters,
 	FRDGTextureRef DistanceFieldNormal,
 	FRDGTextureRef& OutDynamicBentNormalAO)
@@ -483,102 +413,104 @@ void FDeferredShadingSceneRenderer::RenderDistanceFieldAOScreenGrid(
 	const FIntPoint ConeTraceBufferSize = GetBufferSizeForConeTracing();
 	const FIntPoint TileListGroupSize = GetTileListGroupSizeForView(View);
 
-	FAOScreenGridResources*& ScreenGridResources = View.ViewState->AOScreenGridResources;
+	FAOScreenGridParameters AOScreenGridParameters;
 
-	if (!ScreenGridResources
-		|| ScreenGridResources->ScreenGridDimensions != ConeTraceBufferSize
-		|| !ScreenGridResources->IsInitialized()
-		|| GFastVRamConfig.bDirty)
 	{
-		if (ScreenGridResources)
-		{
-			ScreenGridResources->ReleaseResource();
-		}
-		else
-		{
-			ScreenGridResources = new FAOScreenGridResources();
-		}
+		AOScreenGridParameters.ScreenGridConeVisibilitySize = ConeTraceBufferSize;
+		//@todo - 2d textures
+		//@todo - FastVRamFlag
+		// @todo - ScreenGridConeVisibility should probably be R32_FLOAT format.
+		FRDGBufferRef ScreenGridConeVisibility = GraphBuilder.CreateBuffer(FRDGBufferDesc::CreateBufferDesc(sizeof(uint32), NumConeSampleDirections * ConeTraceBufferSize.X * ConeTraceBufferSize.Y), TEXT("ScreenGridConeVisibility"));
+		AOScreenGridParameters.RWScreenGridConeVisibility = GraphBuilder.CreateUAV(ScreenGridConeVisibility, PF_R32_UINT);
+		AOScreenGridParameters.ScreenGridConeVisibility = GraphBuilder.CreateSRV(ScreenGridConeVisibility, PF_R32_UINT);
+	}
 
-		ScreenGridResources->ScreenGridDimensions = ConeTraceBufferSize;
+	float ConeVisibilityClearValue = 1.0f;
+	AddClearUAVPass(GraphBuilder, AOScreenGridParameters.RWScreenGridConeVisibility, *(uint32*)&ConeVisibilityClearValue);
 
-		ScreenGridResources->InitResource();
-	}
+	const FDistanceFieldSceneData& DistanceFieldSceneData = Scene->DistanceFieldSceneData;
 
+	// Note: no transition, want to overlap object cone tracing and global DF cone tracing since both shaders use atomics to ScreenGridConeVisibility
+
+	if (bUseGlobalDistanceField)
 	{
-		FDistanceFieldAOScreenGridParameters* PassParameters = GraphBuilder.AllocParameters<FDistanceFieldAOScreenGridParameters>();
+		check(View.GlobalDistanceFieldInfo.Clipmaps.Num() > 0);
+
+		auto* PassParameters = GraphBuilder.AllocParameters<FConeTraceScreenGridGlobalOcclusionCS::FParameters>();
+		PassParameters->DistanceFieldCulledObjectBuffers = CulledObjectBufferParameters;
+		PassParameters->DistanceFieldAtlas = DistanceField::SetupAtlasParameters(DistanceFieldSceneData);
+		PassParameters->AOScreenGridParameters = AOScreenGridParameters;
+		PassParameters->TileConeDepthRanges = TileIntersectionParameters.TileConeDepthRanges;
 		PassParameters->SceneTextures = SceneTextures.UniformBuffer;
 		PassParameters->DistanceFieldNormal = DistanceFieldNormal;
-		const FDistanceFieldSceneData& DistanceFieldSceneData = Scene->DistanceFieldSceneData;
 
-		GraphBuilder.AddPass(
-			{},
-			PassParameters,
-			ERDGPassFlags::Compute | ERDGPassFlags::NeverCull,
-			[&View, &ScreenGridResources, &DistanceFieldSceneData, Parameters, DistanceFieldNormal, TileListGroupSize, bUseGlobalDistanceField, bUseObjectDistanceField]
-			(FRHICommandListImmediate& RHICmdList)
-		{
-			ScreenGridResources->AcquireTransientResource();
+		FConeTraceScreenGridGlobalOcclusionCS::FPermutationDomain PermutationVector;
+		PermutationVector.Set<FConeTraceScreenGridGlobalOcclusionCS::FConeTraceObjects>(bUseObjectDistanceField);
 
-			float ConeVisibilityClearValue = 1.0f;
-			RHICmdList.Transition(FRHITransitionInfo(ScreenGridResources->ScreenGridConeVisibility.UAV, ERHIAccess::Unknown, ERHIAccess::UAVCompute));
-			RHICmdList.ClearUAVUint(ScreenGridResources->ScreenGridConeVisibility.UAV, FUintVector4(*(uint32*)&ConeVisibilityClearValue, 0, 0, 0)); // @todo - ScreenGridConeVisibility should probably be R32_FLOAT format.
+		auto ComputeShader = View.ShaderMap->GetShader<FConeTraceScreenGridGlobalOcclusionCS>(PermutationVector);
 
-			if (bUseGlobalDistanceField)
-			{
-				SCOPED_DRAW_EVENT(RHICmdList, ConeTraceGlobal);
+		ClearUnusedGraphResources(ComputeShader, PassParameters);
 
+		GraphBuilder.AddPass(
+			RDG_EVENT_NAME("ConeTraceGlobal"),
+			PassParameters,
+			ERDGPassFlags::Compute,
+			[PassParameters, ComputeShader, &View, Parameters, DistanceFieldNormal, TileListGroupSize, bUseObjectDistanceField](FRHICommandListImmediate& RHICmdList)
+			{
 				const uint32 GroupSizeX = FMath::DivideAndRoundUp(View.ViewRect.Size().X / GAODownsampleFactor / GConeTraceDownsampleFactor, GConeTraceGlobalDFTileSize);
 				const uint32 GroupSizeY = FMath::DivideAndRoundUp(View.ViewRect.Size().Y / GAODownsampleFactor / GConeTraceDownsampleFactor, GConeTraceGlobalDFTileSize);
 
-				check(View.GlobalDistanceFieldInfo.Clipmaps.Num() > 0);
-
-				if (bUseObjectDistanceField)
-				{
-					TShaderMapRef<TConeTraceScreenGridGlobalOcclusionCS<true> > ComputeShader(View.ShaderMap);
-
-					RHICmdList.SetComputeShader(ComputeShader.GetComputeShader());
-					ComputeShader->SetParameters(RHICmdList, View, DistanceFieldSceneData, TileListGroupSize, DistanceFieldNormal->GetRHI(), Parameters, View.GlobalDistanceFieldInfo);
-					DispatchComputeShader(RHICmdList, ComputeShader.GetShader(), GroupSizeX, GroupSizeY, 1);
-					ComputeShader->UnsetParameters(RHICmdList, View);
-				}
-				else
-				{
-					TShaderMapRef<TConeTraceScreenGridGlobalOcclusionCS<false> > ComputeShader(View.ShaderMap);
-
-					RHICmdList.SetComputeShader(ComputeShader.GetComputeShader());
-					ComputeShader->SetParameters(RHICmdList, View, DistanceFieldSceneData, TileListGroupSize, DistanceFieldNormal->GetRHI(), Parameters, View.GlobalDistanceFieldInfo);
-					DispatchComputeShader(RHICmdList, ComputeShader.GetShader(), GroupSizeX, GroupSizeY, 1);
-					ComputeShader->UnsetParameters(RHICmdList, View);
-				}
-			}
-
-			if (bUseObjectDistanceField)
-			{
-				SCOPED_DRAW_EVENT(RHICmdList, ConeTraceObjects);
-				FTileIntersectionResources* TileIntersectionResources = ((FSceneViewState*)View.State)->AOTileIntersectionResources;
+				RHICmdList.SetComputeShader(ComputeShader.GetComputeShader());
+
+				ComputeShader->SetParameters(RHICmdList, View, TileListGroupSize, DistanceFieldNormal->GetRHI(), Parameters, View.GlobalDistanceFieldInfo);
+				SetShaderParameters(RHICmdList, ComputeShader, ComputeShader.GetComputeShader(), *PassParameters);
+
+				DispatchComputeShader(RHICmdList, ComputeShader.GetShader(), GroupSizeX, GroupSizeY, 1);
+
+				UnsetShaderUAVs(RHICmdList, ComputeShader, ComputeShader.GetComputeShader());
+			});
+	}
 
-				RHICmdList.Transition(FRHITransitionInfo(ScreenGridResources->ScreenGridConeVisibility.UAV, ERHIAccess::Unknown, ERHIAccess::UAVCompute));
+	if (bUseObjectDistanceField)
+	{
+		check(!bUseGlobalDistanceField || View.GlobalDistanceFieldInfo.Clipmaps.Num() > 0);
+
+		auto* PassParameters = GraphBuilder.AllocParameters<FConeTraceScreenGridObjectOcclusionCS::FParameters>();
+		PassParameters->DistanceFieldCulledObjectBuffers = CulledObjectBufferParameters;
+		PassParameters->DistanceFieldAtlas = DistanceField::SetupAtlasParameters(DistanceFieldSceneData);
+		PassParameters->TileIntersectionParameters = TileIntersectionParameters;
+		PassParameters->AOScreenGridParameters = AOScreenGridParameters;
+		PassParameters->ObjectTilesIndirectArguments = ObjectTilesIndirectArguments;
+		PassParameters->SceneTextures = SceneTextures.UniformBuffer;
+		PassParameters->DistanceFieldNormal = DistanceFieldNormal;
 
-				FConeTraceScreenGridObjectOcclusionCS::FPermutationDomain PermutationVector;
-				PermutationVector.Set<FConeTraceScreenGridObjectOcclusionCS::FUseGlobalDistanceField>(bUseGlobalDistanceField);
+		FConeTraceScreenGridObjectOcclusionCS::FPermutationDomain PermutationVector;
+		PermutationVector.Set<FConeTraceScreenGridObjectOcclusionCS::FUseGlobalDistanceField>(bUseGlobalDistanceField);
 
-				auto ComputeShader = View.ShaderMap->GetShader<FConeTraceScreenGridObjectOcclusionCS>(PermutationVector);
+		auto ComputeShader = View.ShaderMap->GetShader<FConeTraceScreenGridObjectOcclusionCS>(PermutationVector);
 
-				check(!bUseGlobalDistanceField || View.GlobalDistanceFieldInfo.Clipmaps.Num() > 0);
+		ClearUnusedGraphResources(ComputeShader, PassParameters);
 
+		GraphBuilder.AddPass(
+			RDG_EVENT_NAME("ConeTraceObjects"),
+			PassParameters,
+			ERDGPassFlags::Compute,
+			[PassParameters, ComputeShader, &View, Parameters, DistanceFieldNormal, bUseGlobalDistanceField, ObjectTilesIndirectArguments](FRHICommandListImmediate& RHICmdList)
+			{
 				RHICmdList.SetComputeShader(ComputeShader.GetComputeShader());
-				ComputeShader->SetParameters(RHICmdList, View, DistanceFieldSceneData, DistanceFieldNormal->GetRHI(), Parameters, bUseGlobalDistanceField, View.GlobalDistanceFieldInfo);
-				DispatchIndirectComputeShader(RHICmdList, ComputeShader.GetShader(), TileIntersectionResources->ObjectTilesIndirectArguments.Buffer, 0);
-				ComputeShader->UnsetParameters(RHICmdList, View);
-			}
 
-			// Compute heightfield occlusion after heightfield GI, otherwise it self-shadows incorrectly
-			View.HeightfieldLightingViewInfo.ComputeOcclusionForScreenGrid(View, RHICmdList, DistanceFieldNormal->GetRHI(), *ScreenGridResources, Parameters);
+				ComputeShader->SetParameters(RHICmdList, View, DistanceFieldNormal->GetRHI(), Parameters, bUseGlobalDistanceField, View.GlobalDistanceFieldInfo);
+				SetShaderParameters(RHICmdList, ComputeShader, ComputeShader.GetComputeShader(), *PassParameters);
+
+				DispatchIndirectComputeShader(RHICmdList, ComputeShader.GetShader(), ObjectTilesIndirectArguments->GetIndirectRHICallBuffer(), 0);
 
-			RHICmdList.Transition(FRHITransitionInfo(ScreenGridResources->ScreenGridConeVisibility.UAV, ERHIAccess::UAVCompute, ERHIAccess::SRVCompute));
-		});
+				UnsetShaderUAVs(RHICmdList, ComputeShader, ComputeShader.GetComputeShader());
+			});
 	}
 
+	// Compute heightfield occlusion after heightfield GI, otherwise it self-shadows incorrectly
+	View.HeightfieldLightingViewInfo.ComputeOcclusionForScreenGrid(GraphBuilder, View, SceneTextures, DistanceFieldNormal, AOScreenGridParameters, Parameters);
+
 	FRDGTextureRef DownsampledBentNormal = nullptr;
 
 	{
@@ -590,31 +522,30 @@ void FDeferredShadingSceneRenderer::RenderDistanceFieldAOScreenGrid(
 		const uint32 GroupSizeX = FMath::DivideAndRoundUp(ConeTraceBufferSize.X, GCombineConesSizeX);
 		const uint32 GroupSizeY = FMath::DivideAndRoundUp(ConeTraceBufferSize.Y, GCombineConesSizeX);
 
-		TShaderMapRef<FCombineConeVisibilityCS> ComputeShader(View.ShaderMap);
-
-		auto* PassParameters = GraphBuilder.AllocParameters<FDistanceFieldAOScreenGridCombineConeVisibilityParameters>();
-		PassParameters->DownsampledBentNormal = DownsampledBentNormal;
-		PassParameters->DownsampledBentNormalUAV = GraphBuilder.CreateUAV(DownsampledBentNormal);
+		auto* PassParameters = GraphBuilder.AllocParameters<FCombineConeVisibilityCS::FParameters>();
+		PassParameters->AOScreenGridParameters = AOScreenGridParameters;
+		PassParameters->RWDistanceFieldBentNormal = GraphBuilder.CreateUAV(DownsampledBentNormal);
 		PassParameters->DistanceFieldNormal = DistanceFieldNormal;
 
+		auto ComputeShader = View.ShaderMap->GetShader<FCombineConeVisibilityCS>();
+
+		ClearUnusedGraphResources(ComputeShader, PassParameters);
+
 		GraphBuilder.AddPass(
 			RDG_EVENT_NAME("CombineCones"),
 			PassParameters,
 			ERDGPassFlags::Compute,
 			[&View, ComputeShader, PassParameters, GroupSizeX, GroupSizeY](FRHICommandList& RHICmdList)
-		{
-			RHICmdList.SetComputeShader(ComputeShader.GetComputeShader());
-			ComputeShader->SetParameters(RHICmdList, View, PassParameters->DistanceFieldNormal->GetRHI(), PassParameters->DownsampledBentNormal->GetRHI(), PassParameters->DownsampledBentNormalUAV->GetRHI());
-			DispatchComputeShader(RHICmdList, ComputeShader.GetShader(), GroupSizeX, GroupSizeY, 1);
-		});
-	}
+			{
+				RHICmdList.SetComputeShader(ComputeShader.GetComputeShader());
 
-	if (IsTransientResourceBufferAliasingEnabled())
-	{
-		AddPass(GraphBuilder, [&ScreenGridResources](FRHICommandList&)
-		{
-			ScreenGridResources->DiscardTransientResource();
-		});
+				ComputeShader->SetParameters(RHICmdList, View, PassParameters->DistanceFieldNormal->GetRHI());
+				SetShaderParameters(RHICmdList, ComputeShader, ComputeShader.GetComputeShader(), *PassParameters);
+
+				DispatchComputeShader(RHICmdList, ComputeShader.GetShader(), GroupSizeX, GroupSizeY, 1);
+
+				UnsetShaderUAVs(RHICmdList, ComputeShader, ComputeShader.GetComputeShader());
+			});
 	}
 
 	PostProcessBentNormalAOScreenGrid(
diff --git a/Engine/Source/Runtime/Renderer/Private/DistanceFieldShadowing.cpp b/Engine/Source/Runtime/Renderer/Private/DistanceFieldShadowing.cpp
--- a/Engine/Source/Runtime/Renderer/Private/DistanceFieldShadowing.cpp
+++ b/Engine/Source/Runtime/Renderer/Private/DistanceFieldShadowing.cpp
@@ -128,9 +128,6 @@ FIntPoint GetBufferSizeForDFShadows()
 	return FIntPoint::DivideAndRoundDown(GetSceneTextureExtent(), GetDFShadowDownsampleFactor());
 }
 
-TGlobalResource<FDistanceFieldObjectBufferResource> GShadowCulledObjectBuffers;
-TGlobalResource<FHeightFieldObjectBufferResource> GShadowCulledHeightFieldObjectBuffers;
-
 class FCullObjectsForShadowCS : public FGlobalShader
 {
 	DECLARE_GLOBAL_SHADER(FCullObjectsForShadowCS);
diff --git a/Engine/Source/Runtime/Renderer/Private/DistanceFieldVisualization.cpp b/Engine/Source/Runtime/Renderer/Private/DistanceFieldVisualization.cpp
--- a/Engine/Source/Runtime/Renderer/Private/DistanceFieldVisualization.cpp
+++ b/Engine/Source/Runtime/Renderer/Private/DistanceFieldVisualization.cpp
@@ -29,6 +29,8 @@ public:
 		SHADER_PARAMETER_STRUCT_REF(FViewUniformShaderParameters, View)
 		SHADER_PARAMETER_STRUCT_REF(FForwardLightData, ForwardLightData)
 		SHADER_PARAMETER_STRUCT_REF(FReflectionUniformParameters, ReflectionStruct)
+		SHADER_PARAMETER_STRUCT_INCLUDE(FDistanceFieldCulledObjectBufferParameters, DistanceFieldCulledObjectBuffers)
+		SHADER_PARAMETER_STRUCT_INCLUDE(FDistanceFieldAtlasParameters, DistanceFieldAtlas)
 		SHADER_PARAMETER_RDG_UNIFORM_BUFFER(FSceneTextureUniformParameters, SceneTextures)
 		SHADER_PARAMETER_RDG_TEXTURE_UAV(RWTexture2D<float4>, RWVisualizeMeshDistanceFields)
 		SHADER_PARAMETER(FVector2D, NumGroups)
@@ -51,21 +53,17 @@ public:
 		: FGlobalShader(Initializer)
 	{
 		BindForLegacyShaderParameters<FParameters>(this, Initializer.PermutationId, Initializer.ParameterMap);
-		ObjectParameters.Bind(Initializer.ParameterMap);
 		AOParameters.Bind(Initializer.ParameterMap);
 		GlobalDistanceFieldParameters.Bind(Initializer.ParameterMap);
 	}
 
 	void SetParameters(
 		FRHICommandList& RHICmdList,
-		const FDistanceFieldSceneData& DistanceFieldSceneData,
 		const FDistanceFieldAOParameters& Parameters,
 		const FGlobalDistanceFieldInfo& GlobalDistanceFieldInfo)
 	{
 		FRHIComputeShader* ShaderRHI = RHICmdList.GetBoundComputeShader();
 
-		ObjectParameters.Set(RHICmdList, ShaderRHI, GAOCulledObjectBuffers.Buffers, DistanceFieldSceneData);
-
 		AOParameters.Set(RHICmdList, ShaderRHI, Parameters);
 
 		if (GlobalDistanceFieldParameters.IsBound())
@@ -75,7 +73,6 @@ public:
 	}
 
 private:
-	LAYOUT_FIELD((TDistanceFieldCulledObjectBufferParameters<DFPT_SignedDistanceField>), ObjectParameters);
 	LAYOUT_FIELD(FAOParameters, AOParameters);
 	LAYOUT_FIELD(FGlobalDistanceFieldParameters, GlobalDistanceFieldParameters);
 };
@@ -93,8 +90,6 @@ public:
 		SHADER_PARAMETER_RDG_UNIFORM_BUFFER(FSceneTextureUniformParameters, SceneTextures)	
 		SHADER_PARAMETER_RDG_TEXTURE(Texture2D, VisualizeDistanceFieldTexture)
 		SHADER_PARAMETER_SAMPLER(SamplerState, VisualizeDistanceFieldSampler)
-		SHADER_PARAMETER_STRUCT_INCLUDE(FDistanceFieldObjectBufferParameters, DistanceFieldObjectBuffers)
-		SHADER_PARAMETER_STRUCT_INCLUDE(FDistanceFieldAtlasParameters, DistanceFieldAtlas)
 		RENDER_TARGET_BINDING_SLOTS()
 	END_SHADER_PARAMETER_STRUCT()
 
@@ -121,7 +116,6 @@ void FDeferredShadingSceneRenderer::RenderMeshDistanceFieldVisualization(
 	if (!UseDistanceFieldAO()
 		|| FeatureLevel < ERHIFeatureLevel::SM5
 		|| !DoesPlatformSupportDistanceFieldAO(FirstView.GetShaderPlatform())
-		|| Views.Num() != 1
 		|| Scene->DistanceFieldSceneData.NumObjectsInBuffer == 0)
 	{
 		return;
@@ -135,8 +129,6 @@ void FDeferredShadingSceneRenderer::RenderMeshDistanceFieldVisualization(
 
 	RDG_EVENT_SCOPE(GraphBuilder, "VisualizeMeshDistanceFields");
 
-	CullObjectsToView(GraphBuilder, Scene, FirstView, Parameters, GAOCulledObjectBuffers);
-
 	FRDGTextureRef VisualizeResultTexture = nullptr;
 
 	{
@@ -151,6 +143,21 @@ void FDeferredShadingSceneRenderer::RenderMeshDistanceFieldVisualization(
 
 	for (const FViewInfo& View : Views)
 	{
+		check(!bUseGlobalDistanceField || View.GlobalDistanceFieldInfo.Clipmaps.Num() > 0);
+
+		FRDGBufferRef ObjectIndirectArguments = nullptr;
+		FDistanceFieldCulledObjectBufferParameters CulledObjectBufferParameters;
+
+		AllocateDistanceFieldCulledObjectBuffers(
+			GraphBuilder,
+			false,
+			FMath::DivideAndRoundUp(Scene->DistanceFieldSceneData.NumObjectsInBuffer, 256) * 256,
+			DFPT_SignedDistanceField,
+			ObjectIndirectArguments,
+			CulledObjectBufferParameters);
+
+		CullObjectsToView(GraphBuilder, Scene, View, Parameters, CulledObjectBufferParameters);
+
 		uint32 GroupSizeX = FMath::DivideAndRoundUp(View.ViewRect.Size().X / GAODownsampleFactor, GDistanceFieldAOTileSizeX);
 		uint32 GroupSizeY = FMath::DivideAndRoundUp(View.ViewRect.Size().Y / GAODownsampleFactor, GDistanceFieldAOTileSizeY);
 
@@ -160,36 +167,29 @@ void FDeferredShadingSceneRenderer::RenderMeshDistanceFieldVisualization(
 		PassParameters->View = View.ViewUniformBuffer;
 		PassParameters->ForwardLightData = View.ForwardLightingResources->ForwardLightDataUniformBuffer;
 		PassParameters->ReflectionStruct = CreateReflectionUniformBuffer(View, UniformBuffer_MultiFrame);
+		PassParameters->DistanceFieldCulledObjectBuffers = CulledObjectBufferParameters;
+		PassParameters->DistanceFieldAtlas = DistanceField::SetupAtlasParameters(Scene->DistanceFieldSceneData);
 		PassParameters->NumGroups = FVector2D(GroupSizeX, GroupSizeY);
 		PassParameters->SceneTextures = SceneTextures.UniformBuffer;
 		PassParameters->RWVisualizeMeshDistanceFields = GraphBuilder.CreateUAV(VisualizeResultTexture);
-		const FDistanceFieldSceneData& DistanceFieldSceneData = Scene->DistanceFieldSceneData;
 
-		check(!bUseGlobalDistanceField || View.GlobalDistanceFieldInfo.Clipmaps.Num() > 0);
+		ClearUnusedGraphResources(ComputeShader, PassParameters);
 
 		GraphBuilder.AddPass(
 			RDG_EVENT_NAME("VisualizeMeshDistanceFieldCS"),
 			PassParameters,
 			ERDGPassFlags::Compute,
-			[&View, &DistanceFieldSceneData, Parameters, ComputeShader, PassParameters, GroupSizeX, GroupSizeY](FRHICommandList& RHICmdList)
+			[&View, Parameters, ComputeShader, PassParameters, GroupSizeX, GroupSizeY](FRHICommandList& RHICmdList)
 		{
 			FRHIComputeShader* ShaderRHI = ComputeShader.GetComputeShader();
 			RHICmdList.SetComputeShader(ComputeShader.GetComputeShader());
 			SetShaderParameters(RHICmdList, ComputeShader, ShaderRHI, *PassParameters);
-			ComputeShader->SetParameters(RHICmdList, DistanceFieldSceneData, Parameters, View.GlobalDistanceFieldInfo);
+			ComputeShader->SetParameters(RHICmdList, Parameters, View.GlobalDistanceFieldInfo);
 			DispatchComputeShader(RHICmdList, ComputeShader.GetShader(), GroupSizeX, GroupSizeY, 1);
 			UnsetShaderUAVs(RHICmdList, ComputeShader, ShaderRHI);
 		});
 	}
 
-	if (IsTransientResourceBufferAliasingEnabled())
-	{
-		AddPass(GraphBuilder, [](FRHICommandList&)
-		{
-			GAOCulledObjectBuffers.Buffers.DiscardTransientResource();
-		});
-	}
-
 	TShaderMapRef<FVisualizeDistanceFieldUpsamplePS> PixelShader(FirstView.ShaderMap);
 
 	for (const FViewInfo& View : Views)
@@ -202,9 +202,6 @@ void FDeferredShadingSceneRenderer::RenderMeshDistanceFieldVisualization(
 		PassParameters->RenderTargets[0] = FRenderTargetBinding(SceneTextures.Color.Target, ERenderTargetLoadAction::ELoad);
 		PassParameters->RenderTargets.DepthStencil = FDepthStencilBinding(SceneTextures.Depth.Target, ERenderTargetLoadAction::ELoad, ERenderTargetLoadAction::ELoad, FExclusiveDepthStencil::DepthRead_StencilRead);
 
-		PassParameters->DistanceFieldObjectBuffers = DistanceField::SetupObjectBufferParameters(Scene->DistanceFieldSceneData);
-		PassParameters->DistanceFieldAtlas = DistanceField::SetupAtlasParameters(Scene->DistanceFieldSceneData);
-
 		const FScreenPassTextureViewport InputViewport(VisualizeResultTexture, GetDownscaledRect(View.ViewRect, GAODownsampleFactor));
 		const FScreenPassTextureViewport OutputViewport(SceneTextures.Color.Target, View.ViewRect);
 
diff --git a/Engine/Source/Runtime/Renderer/Private/HeightfieldLighting.cpp b/Engine/Source/Runtime/Renderer/Private/HeightfieldLighting.cpp
--- a/Engine/Source/Runtime/Renderer/Private/HeightfieldLighting.cpp
+++ b/Engine/Source/Runtime/Renderer/Private/HeightfieldLighting.cpp
@@ -1088,9 +1088,16 @@ const int32 GHeightfieldOcclusionDispatchSize = 8;
 
 class FCalculateHeightfieldOcclusionScreenGridCS : public FGlobalShader
 {
-	DECLARE_SHADER_TYPE(FCalculateHeightfieldOcclusionScreenGridCS, Global)
+	DECLARE_GLOBAL_SHADER(FCalculateHeightfieldOcclusionScreenGridCS);
+
 public:
 
+	BEGIN_SHADER_PARAMETER_STRUCT(FParameters, )
+		SHADER_PARAMETER_STRUCT_INCLUDE(FAOScreenGridParameters, AOScreenGridParameters)
+		SHADER_PARAMETER_RDG_UNIFORM_BUFFER(FSceneTextureUniformParameters, SceneTextures)
+		RDG_TEXTURE_ACCESS(DistanceFieldNormal, ERHIAccess::SRVCompute)
+	END_SHADER_PARAMETER_STRUCT()
+
 	static bool ShouldCompilePermutation(const FGlobalShaderPermutationParameters& Parameters)
 	{
 		return IsFeatureLevelSupported(Parameters.Platform, ERHIFeatureLevel::SM5) && DoesPlatformSupportDistanceFieldGI(Parameters.Platform);
@@ -1111,11 +1118,11 @@ public:
 	FCalculateHeightfieldOcclusionScreenGridCS(const ShaderMetaType::CompiledShaderInitializerType& Initializer)
 		: FGlobalShader(Initializer)
 	{
+		BindForLegacyShaderParameters<FParameters>(this, Initializer.PermutationId, Initializer.ParameterMap, false);
 		AOParameters.Bind(Initializer.ParameterMap);
 		ScreenGridParameters.Bind(Initializer.ParameterMap);
 		HeightfieldDescriptionParameters.Bind(Initializer.ParameterMap);
 		HeightfieldTextureParameters.Bind(Initializer.ParameterMap);
-		ScreenGridConeVisibility.Bind(Initializer.ParameterMap, TEXT("ScreenGridConeVisibility"));
 		TanConeHalfAngle.Bind(Initializer.ParameterMap, TEXT("TanConeHalfAngle"));
 	}
 
@@ -1129,7 +1136,6 @@ public:
 		UTexture2D* HeightfieldTextureValue,
 		int32 NumHeightfieldsValue,
 		FRHITexture* DistanceFieldNormal,
-		const FAOScreenGridResources& ScreenGridResources,
 		const FDistanceFieldAOParameters& Parameters)
 	{
 		FRHIComputeShader* ShaderRHI = RHICmdList.GetBoundComputeShader();
@@ -1140,9 +1146,6 @@ public:
 		HeightfieldDescriptionParameters.Set(RHICmdList, ShaderRHI, GetHeightfieldDescriptionsSRV(), NumHeightfieldsValue);
 		HeightfieldTextureParameters.Set(RHICmdList, ShaderRHI, HeightfieldTextureValue, nullptr, nullptr);
 
-		RHICmdList.Transition(FRHITransitionInfo(ScreenGridResources.ScreenGridConeVisibility.UAV, ERHIAccess::Unknown, ERHIAccess::UAVCompute));
-		ScreenGridConeVisibility.SetBuffer(RHICmdList, ShaderRHI, ScreenGridResources.ScreenGridConeVisibility);
-
 		FAOSampleData2 AOSampleData;
 
 		TArray<FVector, TInlineAllocator<9> > SampleDirections;
@@ -1159,60 +1162,66 @@ public:
 		SetShaderValue(RHICmdList, ShaderRHI, TanConeHalfAngle, FMath::Tan(GAOConeHalfAngle));
 	}
 
-	void UnsetParameters(FRHICommandList& RHICmdList, const FAOScreenGridResources& ScreenGridResources)
-	{
-		ScreenGridConeVisibility.UnsetUAV(RHICmdList, RHICmdList.GetBoundComputeShader());
-		RHICmdList.Transition(FRHITransitionInfo(ScreenGridResources.ScreenGridConeVisibility.UAV, ERHIAccess::UAVCompute, ERHIAccess::SRVMask));
-	}
-
 private:
 
 	LAYOUT_FIELD(FAOParameters, AOParameters);
 	LAYOUT_FIELD(FScreenGridParameters, ScreenGridParameters);
 	LAYOUT_FIELD(FHeightfieldDescriptionParameters, HeightfieldDescriptionParameters);
 	LAYOUT_FIELD(FHeightfieldTextureParameters, HeightfieldTextureParameters);
-	LAYOUT_FIELD(FRWShaderParameter, ScreenGridConeVisibility);
 	LAYOUT_FIELD(FShaderParameter, TanConeHalfAngle);
 };
 
-IMPLEMENT_SHADER_TYPE(, FCalculateHeightfieldOcclusionScreenGridCS, TEXT("/Engine/Private/HeightfieldLighting.usf"), TEXT("CalculateHeightfieldOcclusionScreenGridCS"), SF_Compute);
+IMPLEMENT_GLOBAL_SHADER(FCalculateHeightfieldOcclusionScreenGridCS, "/Engine/Private/HeightfieldLighting.usf", "CalculateHeightfieldOcclusionScreenGridCS", SF_Compute);
 
 
 void FHeightfieldLightingViewInfo::ComputeOcclusionForScreenGrid(
+	FRDGBuilder& GraphBuilder,
 	const FViewInfo& View,
-	FRHICommandListImmediate& RHICmdList,
-	FRHITexture* DistanceFieldNormal,
-	const FAOScreenGridResources& ScreenGridResources,
+	const FSceneTextures& SceneTextures,
+	FRDGTextureRef DistanceFieldNormal,
+	const FAOScreenGridParameters& AOScreenGridParameters,
 	const FDistanceFieldAOParameters& Parameters) const
 {
-	const FScene* Scene = (const FScene*)View.Family->Scene;
+	int32 DownsampleFactor = Heightfield.DownsampleFactor;
 
 	if (Heightfield.ComponentDescriptions.Num() > 0 && GAOHeightfieldOcclusion)
 	{
-		SCOPED_DRAW_EVENT(RHICmdList, HeightfieldOcclusion);
-
-		FSceneViewState* ViewState = (FSceneViewState*)View.State;
-
+		for (TMap<FHeightfieldComponentTextures, TArray<FHeightfieldComponentDescription>>::TConstIterator It(Heightfield.ComponentDescriptions); It; ++It)
 		{
-			for (TMap<FHeightfieldComponentTextures, TArray<FHeightfieldComponentDescription>>::TConstIterator It(Heightfield.ComponentDescriptions); It; ++It)
+			UTexture2D* HeightfieldTexture = It.Key().HeightAndNormal;
+			const TArray<FHeightfieldComponentDescription>& HeightfieldDescriptions = It.Value();
+
+			if (HeightfieldDescriptions.Num() > 0)
 			{
-				const TArray<FHeightfieldComponentDescription>& HeightfieldDescriptions = It.Value();
+				auto* PassParameters = GraphBuilder.AllocParameters<FCalculateHeightfieldOcclusionScreenGridCS::FParameters>();
+				PassParameters->AOScreenGridParameters = AOScreenGridParameters;
+				PassParameters->SceneTextures = SceneTextures.UniformBuffer;
+				PassParameters->DistanceFieldNormal = DistanceFieldNormal;
 
-				if (HeightfieldDescriptions.Num() > 0)
-				{
-					UploadHeightfieldDescriptions(HeightfieldDescriptions, FVector2D(1, 1), 1.0f / Heightfield.DownsampleFactor);
+				auto ComputeShader = View.ShaderMap->GetShader<FCalculateHeightfieldOcclusionScreenGridCS>();
 
-					UTexture2D* HeightfieldTexture = It.Key().HeightAndNormal;
+				ClearUnusedGraphResources(ComputeShader, PassParameters);
 
-					const uint32 GroupSizeX = FMath::DivideAndRoundUp(View.ViewRect.Size().X / GAODownsampleFactor, GHeightfieldOcclusionDispatchSize);
-					const uint32 GroupSizeY = FMath::DivideAndRoundUp(View.ViewRect.Size().Y / GAODownsampleFactor, GHeightfieldOcclusionDispatchSize);
+				GraphBuilder.AddPass(
+					RDG_EVENT_NAME("HeightfieldOcclusion"),
+					PassParameters,
+					ERDGPassFlags::Compute,
+					[PassParameters, ComputeShader, &View, Parameters, DistanceFieldNormal, &HeightfieldDescriptions, HeightfieldTexture, DownsampleFactor](FRHICommandListImmediate& RHICmdList)
+					{
+						UploadHeightfieldDescriptions(HeightfieldDescriptions, FVector2D(1, 1), 1.0f / DownsampleFactor);
 
-					TShaderMapRef<FCalculateHeightfieldOcclusionScreenGridCS> ComputeShader(View.ShaderMap);
-					RHICmdList.SetComputeShader(ComputeShader.GetComputeShader());
-					ComputeShader->SetParameters(RHICmdList, View, HeightfieldTexture, HeightfieldDescriptions.Num(), DistanceFieldNormal, ScreenGridResources, Parameters);
-					DispatchComputeShader(RHICmdList, ComputeShader.GetShader(), GroupSizeX, GroupSizeY, 1);
-					ComputeShader->UnsetParameters(RHICmdList, ScreenGridResources);
-				}
+						const uint32 GroupSizeX = FMath::DivideAndRoundUp(View.ViewRect.Size().X / GAODownsampleFactor, GHeightfieldOcclusionDispatchSize);
+						const uint32 GroupSizeY = FMath::DivideAndRoundUp(View.ViewRect.Size().Y / GAODownsampleFactor, GHeightfieldOcclusionDispatchSize);
+
+						RHICmdList.SetComputeShader(ComputeShader.GetComputeShader());
+
+						ComputeShader->SetParameters(RHICmdList, View, HeightfieldTexture, HeightfieldDescriptions.Num(), DistanceFieldNormal->GetRHI(), Parameters);
+						SetShaderParameters(RHICmdList, ComputeShader, ComputeShader.GetComputeShader(), *PassParameters);
+
+						DispatchComputeShader(RHICmdList, ComputeShader.GetShader(), GroupSizeX, GroupSizeY, 1);
+
+						UnsetShaderUAVs(RHICmdList, ComputeShader, ComputeShader.GetComputeShader());
+					});
 			}
 		}
 	}
diff --git a/Engine/Source/Runtime/Renderer/Private/HeightfieldLighting.h b/Engine/Source/Runtime/Renderer/Private/HeightfieldLighting.h
--- a/Engine/Source/Runtime/Renderer/Private/HeightfieldLighting.h
+++ b/Engine/Source/Runtime/Renderer/Private/HeightfieldLighting.h
@@ -15,7 +15,7 @@
 #include "PrimitiveSceneProxy.h"
 #include "RenderGraphResources.h"
 
-class FAOScreenGridResources;
+class FAOScreenGridParameters;
 class FDistanceFieldAOParameters;
 class FLightSceneInfo;
 class FProjectedShadowInfo;
@@ -109,17 +109,11 @@ public:
 	void ComputeLighting(FRDGBuilder& GraphBuilder, const FViewInfo& View, const FLightSceneInfo& LightSceneInfo) const;
 
 	void ComputeOcclusionForScreenGrid(
-		const FViewInfo& View, 
-		FRHICommandListImmediate& RHICmdList, 
-		FRHITexture* DistanceFieldNormal,
-		const class FAOScreenGridResources& ScreenGridResources,
-		const class FDistanceFieldAOParameters& Parameters) const;
-
-	void ComputeIrradianceForScreenGrid(
-		const FViewInfo& View, 
-		FRHICommandListImmediate& RHICmdList,
-		FRHITexture* DistanceFieldNormal,
-		const FAOScreenGridResources& ScreenGridResources,
+		FRDGBuilder& GraphBuilder,
+		const FViewInfo& View,
+		const FSceneTextures& SceneTextures,
+		FRDGTextureRef DistanceFieldNormal,
+		const FAOScreenGridParameters& AOScreenGridParameters,
 		const FDistanceFieldAOParameters& Parameters) const;
 
 private:
diff --git a/Engine/Source/Runtime/Renderer/Private/RendererScene.cpp b/Engine/Source/Runtime/Renderer/Private/RendererScene.cpp
--- a/Engine/Source/Runtime/Renderer/Private/RendererScene.cpp
+++ b/Engine/Source/Runtime/Renderer/Private/RendererScene.cpp
@@ -208,8 +208,6 @@ FSceneViewState::FSceneViewState()
 	TemporalAASampleIndex = 0;
 	FrameIndex = 0;
 	DistanceFieldTemporalSampleIndex = 0;
-	AOTileIntersectionResources = NULL;
-	AOScreenGridResources = NULL;
 	bDOFHistory = true;
 	bDOFHistory2 = true;
 	
@@ -286,10 +284,6 @@ FSceneViewState::~FSceneViewState()
 	}
 
 	DestroyRenderResource(HeightfieldLightingAtlas);
-	DestroyRenderResource(AOTileIntersectionResources);
-	AOTileIntersectionResources = NULL;
-	DestroyRenderResource(AOScreenGridResources);
-	AOScreenGridResources = NULL;
 
 	HairStrandsViewStateData.Release();
 	ShaderDrawDebugStateData.Release();
diff --git a/Engine/Source/Runtime/Renderer/Private/ScenePrivate.h b/Engine/Source/Runtime/Renderer/Private/ScenePrivate.h
--- a/Engine/Source/Runtime/Renderer/Private/ScenePrivate.h
+++ b/Engine/Source/Runtime/Renderer/Private/ScenePrivate.h
@@ -1039,11 +1039,6 @@ public:
 	TRefCountPtr<IPooledRenderTarget> LightScatteringHistory;
 	TRefCountPtr<IPooledRenderTarget> PrevLightScatteringConservativeDepthTexture;
 
-	/** Distance field AO tile intersection GPU resources.  Last frame's state is not used, but they must be sized exactly to the view so stored here. */
-	class FTileIntersectionResources* AOTileIntersectionResources;
-
-	class FAOScreenGridResources* AOScreenGridResources;
-
 	bool bInitializedGlobalDistanceFieldOrigins;
 	bool bGlobalDistanceFieldPendingReset = false;
 	FGlobalDistanceFieldClipmapState GlobalDistanceFieldClipmapState[GMaxGlobalDistanceFieldClipmaps];
