Fixed crash when using GetMiniFontTexture.

diff --git a/Engine/Source/Runtime/Renderer/Private/ScreenPass.cpp b/Engine/Source/Runtime/Renderer/Private/ScreenPass.cpp
--- a/Engine/Source/Runtime/Renderer/Private/ScreenPass.cpp
+++ b/Engine/Source/Runtime/Renderer/Private/ScreenPass.cpp
@@ -17,11 +17,11 @@ const FTextureRHIRef& GetMiniFontTexture()
 {
 	if (GSystemTextures.AsciiTexture)
 	{
-		return GSystemTextures.AsciiTexture->GetRenderTargetItem().TargetableTexture;
+		return GSystemTextures.AsciiTexture->GetRenderTargetItem().ShaderResourceTexture;
 	}
 	else
 	{
-		return GSystemTextures.WhiteDummy->GetRenderTargetItem().TargetableTexture;
+		return GSystemTextures.WhiteDummy->GetRenderTargetItem().ShaderResourceTexture;
 	}
 }
 
