Fix buffer error in Vulkan for shadows.

diff --git a/Engine/Source/Runtime/Renderer/Private/ShadowRendering.cpp b/Engine/Source/Runtime/Renderer/Private/ShadowRendering.cpp
--- a/Engine/Source/Runtime/Renderer/Private/ShadowRendering.cpp
+++ b/Engine/Source/Runtime/Renderer/Private/ShadowRendering.cpp
@@ -725,7 +725,7 @@ public:
 	virtual void InitRHI() override
 	{
 		FRHIResourceCreateInfo CreateInfo(TEXT("FProjectedShadowInfoStencilFrustum"));
-		VertexBufferRHI = RHICreateVertexBuffer(sizeof(FVector4) * 8, BUF_Volatile, CreateInfo);
+		VertexBufferRHI = RHICreateVertexBuffer(sizeof(FVector4) * 8, BUF_Static, CreateInfo);
 		FVector4* OutFrustumVertices = reinterpret_cast<FVector4*>(RHILockBuffer(VertexBufferRHI, 0, sizeof(FVector4) * 8, RLM_WriteOnly));
 		
 		for(uint32 vZ = 0;vZ < 2;vZ++)
