Fix debug build error due to GPU Scene renames within a DO_GUARD block

diff --git a/Engine/Source/Runtime/Renderer/Private/GPUScene.cpp b/Engine/Source/Runtime/Renderer/Private/GPUScene.cpp
--- a/Engine/Source/Runtime/Renderer/Private/GPUScene.cpp
+++ b/Engine/Source/Runtime/Renderer/Private/GPUScene.cpp
@@ -1138,9 +1138,9 @@ void FGPUScene::UploadGeneral(FRHICommandListImmediate& RHICmdList, FScene *Scen
 
 #if DO_GUARD_SLOW
 		// Verify that if InstanceClearList is empty all the upload bits are cleared as well.
-		if (InstanceClearList.Num() == 0)
+		if (InstanceSceneDataClearList.Num() == 0)
 		{
-			ensure(InstanceDataToClear.Find(true) == INDEX_NONE);
+			ensure(InstanceSceneDataToClear.Find(true) == INDEX_NONE);
 		}
 #endif
 
