Skip GPUScene updates when it's disabled

diff --git a/Engine/Source/Runtime/Renderer/Private/GPUScene.h b/Engine/Source/Runtime/Renderer/Private/GPUScene.h
--- a/Engine/Source/Runtime/Renderer/Private/GPUScene.h
+++ b/Engine/Source/Runtime/Renderer/Private/GPUScene.h
@@ -227,29 +227,32 @@ public:
 	 */
 	FORCEINLINE void RecordPrimitiveIdSwap(int32 PrimitiveIdA, int32 PrimitiveIdB)
 	{
-		// We should never call this on a non-existent primitive, so no need to resize
-		checkSlow(PrimitiveIdA < PrimitiveDirtyState.Num());
-		checkSlow(PrimitiveIdB < PrimitiveDirtyState.Num());
-
-		if (PrimitiveDirtyState[PrimitiveIdA] == EPrimitiveDirtyState::None)
-		{
-			PrimitivesToUpdate.Add(PrimitiveIdA);
-		}
-		PrimitiveDirtyState[PrimitiveIdA] |= EPrimitiveDirtyState::ChangedId;
-		if (PrimitiveDirtyState[PrimitiveIdB] == EPrimitiveDirtyState::None)
+		if (IsEnabled())
 		{
-			PrimitivesToUpdate.Add(PrimitiveIdB);
+			// We should never call this on a non-existent primitive, so no need to resize
+			checkSlow(PrimitiveIdA < PrimitiveDirtyState.Num());
+			checkSlow(PrimitiveIdB < PrimitiveDirtyState.Num());
+
+			if (PrimitiveDirtyState[PrimitiveIdA] == EPrimitiveDirtyState::None)
+			{
+				PrimitivesToUpdate.Add(PrimitiveIdA);
+			}
+			PrimitiveDirtyState[PrimitiveIdA] |= EPrimitiveDirtyState::ChangedId;
+			if (PrimitiveDirtyState[PrimitiveIdB] == EPrimitiveDirtyState::None)
+			{
+				PrimitivesToUpdate.Add(PrimitiveIdB);
+			}
+			PrimitiveDirtyState[PrimitiveIdB] |= EPrimitiveDirtyState::ChangedId;
+
+			Swap(PrimitiveDirtyState[PrimitiveIdA], PrimitiveDirtyState[PrimitiveIdB]);
 		}
-		PrimitiveDirtyState[PrimitiveIdB] |= EPrimitiveDirtyState::ChangedId;
-
-		Swap(PrimitiveDirtyState[PrimitiveIdA], PrimitiveDirtyState[PrimitiveIdB]);
 	}
 
 	FORCEINLINE EPrimitiveDirtyState GetPrimitiveDirtyState(int32 PrimitiveId) const { return PrimitiveDirtyState[PrimitiveId]; }
 
 	FORCEINLINE void ResizeDirtyState(int32 NewSizeIn)
 	{
-		if (NewSizeIn > PrimitiveDirtyState.Num())
+		if (IsEnabled() && NewSizeIn > PrimitiveDirtyState.Num())
 		{
 			const int32 NewSize = Align(NewSizeIn, 64);
 			static_assert(static_cast<uint32>(EPrimitiveDirtyState::None) == 0U, "Using AddZeroed to ensure efficent add, requires None == 0");
