Fix missing clearing hair-only depth texture, for tile not covered by hair.

diff --git a/Engine/Shaders/Private/HairStrands/HairStrandsVisibilityDepthPS.usf b/Engine/Shaders/Private/HairStrands/HairStrandsVisibilityDepthPS.usf
--- a/Engine/Shaders/Private/HairStrands/HairStrandsVisibilityDepthPS.usf
+++ b/Engine/Shaders/Private/HairStrands/HairStrandsVisibilityDepthPS.usf
@@ -38,19 +38,19 @@ void MainPS(
 ////////////////////////////////////////////////////////////////////////////////////////////////////////////////////////////////////
 
 #if PERMUTATION_OUTPUT_TYPE == 1
+uint bClear;
 void MainPS(
 	in FScreenVertexOutput Input,
 	out float SvDepth : SV_DEPTH)
 {
-	const FCategorizationData CategorisationData = DecodeCategorizationData(CategorisationTexture.Load(uint3(Input.Position.xy, 0)));
-	if (CategorisationData.PixelCoverage > 0.0f)
-	{
-		SvDepth = CategorisationData.ClosestDepth;
-	}
-	else
+	SvDepth = 0; // Inverse-Z
+	if (bClear == 0)
 	{
-		SvDepth = 0; // Inverse-Z|
+		const FCategorizationData CategorisationData = DecodeCategorizationData(CategorisationTexture.Load(uint3(Input.Position.xy, 0)));
+		if (CategorisationData.PixelCoverage > 0.0f)
+		{
+			SvDepth = CategorisationData.ClosestDepth;
+		}
 	}
-
 }
 #endif
\ No newline at end of file
diff --git a/Engine/Shaders/Private/HairStrands/HairStrandsVisibilityTile.usf b/Engine/Shaders/Private/HairStrands/HairStrandsVisibilityTile.usf
--- a/Engine/Shaders/Private/HairStrands/HairStrandsVisibilityTile.usf
+++ b/Engine/Shaders/Private/HairStrands/HairStrandsVisibilityTile.usf
@@ -14,6 +14,7 @@ int2 ViewResolution;
 Texture2D<uint4> InputTexture;
 RWBuffer<uint> TileCountBuffer;
 RWBuffer<uint2> TileDataBuffer;
+RWBuffer<uint2> TileClearBuffer;
 
 groupshared bool s_HasTileHair[GROUP_THREAD_COUNT];
 
@@ -57,11 +58,20 @@ void TileMainCS(uint2 DispatchThreadId : SV_DispatchThreadID, uint LinearIndex :
 		s_HasTileHair[LinearIndex] = s_HasTileHair[LinearIndex] || s_HasTileHair[LinearIndex + 2];
 	}
 	GroupMemoryBarrierWithGroupSync();
-	if (LinearIndex < 1 && (s_HasTileHair[LinearIndex] || s_HasTileHair[LinearIndex + 1]))
+	if (LinearIndex < 1)
 	{
-		uint WriteToIndex;
-		InterlockedAdd(TileCountBuffer[0], 1, WriteToIndex);
-		TileDataBuffer[WriteToIndex] = GroupId.xy;
+		if ((s_HasTileHair[LinearIndex] || s_HasTileHair[LinearIndex + 1]))
+		{
+			uint WriteToIndex;
+			InterlockedAdd(TileCountBuffer[0], 1, WriteToIndex);
+			TileDataBuffer[WriteToIndex] = GroupId.xy;
+		}
+		else
+		{
+			uint WriteToIndex;
+			InterlockedAdd(TileCountBuffer[1], 1, WriteToIndex);
+			TileClearBuffer[WriteToIndex] = GroupId.xy;
+		}
 	}
 }
 #endif // SHADER_TILE_GENERATION
@@ -76,6 +86,7 @@ Buffer<uint> TileCountBuffer;
 RWBuffer<uint> TileIndirectDrawBuffer;
 RWBuffer<uint> TileIndirectDispatchBuffer;
 RWBuffer<uint> TilePerThreadIndirectDispatchBuffer;
+RWBuffer<uint> TileClearIndirectDrawBuffer;
 
 [numthreads(1, 1, 1)]
 void MainCS(uint2 DispatchThreadId : SV_DispatchThreadID)
@@ -100,6 +111,12 @@ void MainCS(uint2 DispatchThreadId : SV_DispatchThreadID)
 		TilePerThreadIndirectDispatchBuffer[0] = (TileCount + TilePerThread_GroupSize -1) / TilePerThread_GroupSize;
 		TilePerThreadIndirectDispatchBuffer[1] = 1;
 		TilePerThreadIndirectDispatchBuffer[2] = 1;
+
+		const uint TileClearCount = TileCountBuffer[1];
+
+		TileClearIndirectDrawBuffer[0] = bRectPrimitive > 0 ? 4 : 6;
+		TileClearIndirectDrawBuffer[1] = TileClearCount;
+		TileClearIndirectDrawBuffer[2] = 1;
 	}
 }
 #endif // SHADER_TILE_COPY_ARGS
diff --git a/Engine/Source/Runtime/Renderer/Private/HairStrands/HairStrandsData.h b/Engine/Source/Runtime/Renderer/Private/HairStrands/HairStrandsData.h
--- a/Engine/Source/Runtime/Renderer/Private/HairStrands/HairStrandsData.h
+++ b/Engine/Source/Runtime/Renderer/Private/HairStrands/HairStrandsData.h
@@ -41,6 +41,8 @@ END_GLOBAL_SHADER_PARAMETER_STRUCT()
 
 struct FHairStrandsTiles
 {
+	enum class ETileType : uint8 { Hair, Other };
+
 	FIntPoint			Resolution = FIntPoint(0, 0);
 	static const uint32 GroupSize = 64;
 	static const uint32	TileSize = 8;
@@ -56,6 +58,10 @@ struct FHairStrandsTiles
 	FRDGBufferRef		TileIndirectDispatchBuffer = nullptr;
 	FRDGBufferRef		TilePerThreadIndirectDispatchBuffer = nullptr;
 
+	FRDGBufferSRVRef	TileClearSRV = nullptr;
+	FRDGBufferRef		TileClearBuffer = nullptr;
+	FRDGBufferRef		TileClearIndirectDrawBuffer = nullptr;
+
 	bool IsValid() const { return TileCount > 0 && TileDataBuffer != nullptr; }
 };
 
diff --git a/Engine/Source/Runtime/Renderer/Private/HairStrands/HairStrandsTile.cpp b/Engine/Source/Runtime/Renderer/Private/HairStrands/HairStrandsTile.cpp
--- a/Engine/Source/Runtime/Renderer/Private/HairStrands/HairStrandsTile.cpp
+++ b/Engine/Source/Runtime/Renderer/Private/HairStrands/HairStrandsTile.cpp
@@ -16,18 +16,18 @@
 
 ////////////////////////////////////////////////////////////////////////////////////////////////////////////////////////////////////////////////
 
-FHairStrandsTilePassVS::FParameters GetHairStrandsTileParameters(const FHairStrandsTiles& In)
+FHairStrandsTilePassVS::FParameters GetHairStrandsTileParameters(const FHairStrandsTiles& In, FHairStrandsTiles::ETileType TileType)
 {
-	return GetHairStrandsTileParameters(In, In.Resolution);
+	return GetHairStrandsTileParameters(In, In.Resolution, TileType);
 }
 
-FHairStrandsTilePassVS::FParameters GetHairStrandsTileParameters(const FHairStrandsTiles& In, const FIntPoint& InOutputResolution)
+FHairStrandsTilePassVS::FParameters GetHairStrandsTileParameters(const FHairStrandsTiles& In, const FIntPoint& InOutputResolution, FHairStrandsTiles::ETileType TileType)
 {
 	FHairStrandsTilePassVS::FParameters Out;
 	Out.bRectPrimitive			= In.bRectPrimitive ? 1 : 0;
 	Out.TileOutputResolution	= InOutputResolution;
-	Out.TileDataBuffer			= In.TileDataSRV;
-	Out.TileIndirectBuffer		= In.TileIndirectDrawBuffer;
+	Out.TileDataBuffer			= TileType == FHairStrandsTiles::ETileType::Hair ? In.TileDataSRV : In.TileClearSRV;
+	Out.TileIndirectBuffer		= TileType == FHairStrandsTiles::ETileType::Hair ? In.TileIndirectDrawBuffer : In.TileClearIndirectDrawBuffer;
 	return Out;
 }
 
@@ -50,6 +50,7 @@ class FHairStrandsTileCopyArgsPassCS : public FGlobalShader
 		SHADER_PARAMETER_RDG_BUFFER_UAV(RWBuffer, TileIndirectDrawBuffer)
 		SHADER_PARAMETER_RDG_BUFFER_UAV(RWBuffer, TileIndirectDispatchBuffer)
 		SHADER_PARAMETER_RDG_BUFFER_UAV(RWBuffer, TilePerThreadIndirectDispatchBuffer)
+		SHADER_PARAMETER_RDG_BUFFER_UAV(RWBuffer, TileClearIndirectDrawBuffer)
 	END_SHADER_PARAMETER_STRUCT()
 
 	static bool ShouldCompilePermutation(const FGlobalShaderPermutationParameters& Parameters)
@@ -81,6 +82,8 @@ void AddHairStrandsCopyArgsTilesPass(
 	PassParameters->TileIndirectDispatchBuffer = GraphBuilder.CreateUAV(TileData.TileIndirectDispatchBuffer, PF_R32_UINT);
 	PassParameters->TilePerThreadIndirectDispatchBuffer = GraphBuilder.CreateUAV(TileData.TilePerThreadIndirectDispatchBuffer, PF_R32_UINT);
 
+	PassParameters->TileClearIndirectDrawBuffer = GraphBuilder.CreateUAV(TileData.TileClearIndirectDrawBuffer, PF_R32_UINT);
+
 	FComputeShaderUtils::AddPass(
 		GraphBuilder,
 		RDG_EVENT_NAME("HairStrandsTileCopyArgs"),
@@ -105,6 +108,7 @@ class FHairStrandsTileGenerationPassCS : public FGlobalShader
 		SHADER_PARAMETER_RDG_TEXTURE(Texture2D<uint4>, InputTexture)
 		SHADER_PARAMETER_RDG_BUFFER_UAV(RWBuffer, TileCountBuffer)
 		SHADER_PARAMETER_RDG_BUFFER_UAV(RWBuffer, TileDataBuffer)
+		SHADER_PARAMETER_RDG_BUFFER_UAV(RWBuffer, TileClearBuffer)
 	END_SHADER_PARAMETER_STRUCT()
 
 	static bool ShouldCompilePermutation(const FGlobalShaderPermutationParameters& Parameters)
@@ -134,12 +138,15 @@ FHairStrandsTiles AddHairStrandsGenerateTilesPass(
 	Out.TileCount = Out.TileCountXY.X * Out.TileCountXY.Y;
 	Out.Resolution = InputResolution;
 	Out.bRectPrimitive = GRHISupportsRectTopology;
-	Out.TileCountBuffer = GraphBuilder.CreateBuffer(FRDGBufferDesc::CreateBufferDesc(4, 1), TEXT("Hair.TileCountBuffer"));
+	Out.TileCountBuffer = GraphBuilder.CreateBuffer(FRDGBufferDesc::CreateBufferDesc(4, 2), TEXT("Hair.TileCountBuffer"));
 	Out.TileIndirectDrawBuffer = GraphBuilder.CreateBuffer(FRDGBufferDesc::CreateIndirectDesc<FRHIDrawIndirectParameters>(1), TEXT("Hair.TileIndirectDrawBuffer"));
 	Out.TileIndirectDispatchBuffer = GraphBuilder.CreateBuffer(FRDGBufferDesc::CreateIndirectDesc<FRHIDispatchIndirectParameters>(1), TEXT("Hair.TileIndirectDispatchBuffer"));
 	Out.TilePerThreadIndirectDispatchBuffer = GraphBuilder.CreateBuffer(FRDGBufferDesc::CreateIndirectDesc<FRHIDispatchIndirectParameters>(1), TEXT("Hair.TilePerThreadIndirectDispatchBuffer"));
 	Out.TileDataBuffer = GraphBuilder.CreateBuffer(FRDGBufferDesc::CreateBufferDesc(sizeof(uint32), Out.TileCount), TEXT("Hair.TileDataBuffer"));
 
+	Out.TileClearIndirectDrawBuffer = GraphBuilder.CreateBuffer(FRDGBufferDesc::CreateIndirectDesc<FRHIDrawIndirectParameters>(1), TEXT("Hair.TileClearIndirectDrawBuffer"));
+	Out.TileClearBuffer = GraphBuilder.CreateBuffer(FRDGBufferDesc::CreateBufferDesc(sizeof(uint32), Out.TileCount), TEXT("Hair.TileClearBuffer"));
+
 	FRDGBufferUAVRef TileCountUAV = GraphBuilder.CreateUAV(Out.TileCountBuffer, PF_R32_UINT);
 	AddClearUAVPass(GraphBuilder, TileCountUAV, 0u);
 
@@ -148,6 +155,7 @@ FHairStrandsTiles AddHairStrandsGenerateTilesPass(
 	PassParameters->ViewResolution = View.ViewRect.Size();
 	PassParameters->InputTexture = InputTexture;
 	PassParameters->TileDataBuffer = GraphBuilder.CreateUAV(Out.TileDataBuffer, PF_R16G16_UINT);
+	PassParameters->TileClearBuffer = GraphBuilder.CreateUAV(Out.TileClearBuffer, PF_R16G16_UINT);
 	PassParameters->TileCountBuffer = TileCountUAV;
 
 	FComputeShaderUtils::AddPass(
@@ -158,6 +166,7 @@ FHairStrandsTiles AddHairStrandsGenerateTilesPass(
 		FComputeShaderUtils::GetGroupCount(InputResolution, FHairStrandsTiles::TileSize));
 
 	Out.TileDataSRV = GraphBuilder.CreateSRV(Out.TileDataBuffer, PF_R16G16_UINT);
+	Out.TileClearSRV = GraphBuilder.CreateSRV(Out.TileClearBuffer, PF_R16G16_UINT);
 
 	// Initialize indirect dispatch buffer, based on the indirect draw bugger
 	AddHairStrandsCopyArgsTilesPass(GraphBuilder, View, Out);
diff --git a/Engine/Source/Runtime/Renderer/Private/HairStrands/HairStrandsTile.h b/Engine/Source/Runtime/Renderer/Private/HairStrands/HairStrandsTile.h
--- a/Engine/Source/Runtime/Renderer/Private/HairStrands/HairStrandsTile.h
+++ b/Engine/Source/Runtime/Renderer/Private/HairStrands/HairStrandsTile.h
@@ -31,8 +31,8 @@ class FHairStrandsTilePassVS : public FGlobalShader
 	static void ModifyCompilationEnvironment(const FGlobalShaderPermutationParameters& Parameters, FShaderCompilerEnvironment& OutEnvironment);
 };
 
-FHairStrandsTilePassVS::FParameters GetHairStrandsTileParameters(const FHairStrandsTiles& In);
-FHairStrandsTilePassVS::FParameters GetHairStrandsTileParameters(const FHairStrandsTiles& In, const FIntPoint& InOutputResolution);
+FHairStrandsTilePassVS::FParameters GetHairStrandsTileParameters(const FHairStrandsTiles& In, FHairStrandsTiles::ETileType TileType = FHairStrandsTiles::ETileType::Hair);
+FHairStrandsTilePassVS::FParameters GetHairStrandsTileParameters(const FHairStrandsTiles& In, const FIntPoint& InOutputResolution, FHairStrandsTiles::ETileType TileType = FHairStrandsTiles::ETileType::Hair);
 
 FHairStrandsTiles AddHairStrandsGenerateTilesPass(
 	FRDGBuilder& GraphBuilder,
diff --git a/Engine/Source/Runtime/Renderer/Private/HairStrands/HairStrandsVisibility.cpp b/Engine/Source/Runtime/Renderer/Private/HairStrands/HairStrandsVisibility.cpp
--- a/Engine/Source/Runtime/Renderer/Private/HairStrands/HairStrandsVisibility.cpp
+++ b/Engine/Source/Runtime/Renderer/Private/HairStrands/HairStrandsVisibility.cpp
@@ -2723,6 +2723,7 @@ class FHairVisibilityDepthPS : public FGlobalShader
 
 	BEGIN_SHADER_PARAMETER_STRUCT(FParameters, )
 		SHADER_PARAMETER_STRUCT_INCLUDE(FHairStrandsTilePassVS::FParameters, TileData)
+		SHADER_PARAMETER(uint32, bClear)
 		SHADER_PARAMETER_RDG_TEXTURE(Texture2D, CategorisationTexture)
 		RENDER_TARGET_BINDING_SLOTS()
 	END_SHADER_PARAMETER_STRUCT()
@@ -2744,11 +2745,18 @@ public:
 
 IMPLEMENT_GLOBAL_SHADER(FHairVisibilityDepthPS, "/Engine/Private/HairStrands/HairStrandsVisibilityDepthPS.usf", "MainPS", SF_Pixel);
 
+enum class EHairPatchPass
+{
+	GBufferPatch,
+	DepthPatch,
+	DepthClear
+};
+
 static void AddHairVisibilityCommonPatchPass(
 	FRDGBuilder& GraphBuilder,
 	const FViewInfo& View,
 	const FHairStrandsTiles& TileData,
-	const bool bDepthOnly,
+	const EHairPatchPass PassType,
 	const FRDGTextureRef& CategorisationTexture,
 	FRDGTextureRef OutGBufferBTexture,
 	FRDGTextureRef OutGBufferCTexture,
@@ -2756,14 +2764,15 @@ static void AddHairVisibilityCommonPatchPass(
 	FRDGTextureRef OutDepthTexture)
 {
 	FHairVisibilityDepthPS::FParameters* Parameters = GraphBuilder.AllocParameters<FHairVisibilityDepthPS::FParameters>();
-	Parameters->TileData = GetHairStrandsTileParameters(TileData);
+	Parameters->bClear = PassType == EHairPatchPass::DepthClear ? 1u : 0u;
+	Parameters->TileData = GetHairStrandsTileParameters(TileData, PassType == EHairPatchPass::DepthClear ?  FHairStrandsTiles::ETileType::Other : FHairStrandsTiles::ETileType::Hair);
 	Parameters->CategorisationTexture = CategorisationTexture;
 	Parameters->RenderTargets.DepthStencil = FDepthStencilBinding(
 		OutDepthTexture,
 		ERenderTargetLoadAction::ELoad,
 		ERenderTargetLoadAction::ELoad,
 		FExclusiveDepthStencil::DepthWrite_StencilNop);
-	if (!bDepthOnly)
+	if (PassType == EHairPatchPass::GBufferPatch)
 	{
 		check(OutGBufferBTexture && OutGBufferCTexture && OutColorTexture);
 		Parameters->RenderTargets[0] = FRenderTargetBinding(OutGBufferBTexture, ERenderTargetLoadAction::ELoad);
@@ -2775,14 +2784,14 @@ static void AddHairVisibilityCommonPatchPass(
 	TShaderMapRef<FHairStrandsTilePassVS> TileVertexShader(View.ShaderMap);
 
 	FHairVisibilityDepthPS::FPermutationDomain PermutationVector;
-	PermutationVector.Set<FHairVisibilityDepthPS::FOutputType>(bDepthOnly ? 1 : 0);
+	PermutationVector.Set<FHairVisibilityDepthPS::FOutputType>( (PassType == EHairPatchPass::DepthClear || PassType == EHairPatchPass::DepthPatch) ? 1 : 0);
 	TShaderMapRef<FHairVisibilityDepthPS> PixelShader(View.ShaderMap, PermutationVector);
 	const FIntRect Viewport = View.ViewRect;
 	const FIntPoint Resolution = OutDepthTexture->Desc.Extent;
 	const bool bUseTile = TileData.IsValid();
 
 	GraphBuilder.AddPass(
-		RDG_EVENT_NAME("%s(%s)", bDepthOnly ? TEXT("HairStrandsVisibilityHairOnlyDepth") : TEXT("HairStrandsVisibilityWriteColorAndDepth"), bUseTile ? TEXT("Tile") : TEXT("Screen")),
+		RDG_EVENT_NAME("%s(%s)", (PassType == EHairPatchPass::DepthClear || PassType == EHairPatchPass::DepthPatch) ? TEXT("HairStrandsVisibilityPatchHairOnlyDepth") : TEXT("HairStrandsVisibilityWriteColorAndDepth"), bUseTile ? TEXT("Tile") : TEXT("Screen")),
 		Parameters,
 		ERDGPassFlags::Raster,
 		[Parameters, ScreenVertexShader, TileVertexShader, PixelShader, Viewport, Resolution, bUseTile](FRHICommandList& RHICmdList)
@@ -2847,7 +2856,7 @@ static void AddHairVisibilityColorAndDepthPatchPass(
 		GraphBuilder,
 		View,
 		TileData,
-		false,
+		EHairPatchPass::GBufferPatch,
 		CategorisationTexture,
 		OutGBufferBTexture,
 		OutGBufferCTexture,
@@ -2866,11 +2875,29 @@ static void AddHairOnlyDepthPass(
 	{
 		return;
 	}
+
+	// If tile data are available, we dispatch a complementary set of tile to clear non-hair tile
+	// If tile data are not available, then the clearly is done prior to that.
+	if (TileData.IsValid())
+	{
+		AddHairVisibilityCommonPatchPass(
+			GraphBuilder,
+			View,
+			TileData,
+			EHairPatchPass::DepthClear,
+			CategorisationTexture,
+			nullptr,
+			nullptr,
+			nullptr,
+			OutDepthTexture);
+	}
+
+	// Depth value
 	AddHairVisibilityCommonPatchPass(
 		GraphBuilder,
 		View,
 		TileData,
-		true,
+		EHairPatchPass::DepthPatch,
 		CategorisationTexture,
 		nullptr,
 		nullptr,
