Added extra ray tracing scene build input validation on the render thread timeline, to attempt catching issues before RHI thread where some of the context is no longer available.

diff --git a/Engine/Source/Runtime/Renderer/Private/DeferredShadingRenderer.cpp b/Engine/Source/Runtime/Renderer/Private/DeferredShadingRenderer.cpp
--- a/Engine/Source/Runtime/Renderer/Private/DeferredShadingRenderer.cpp
+++ b/Engine/Source/Runtime/Renderer/Private/DeferredShadingRenderer.cpp
@@ -951,6 +951,8 @@ bool FDeferredShadingSceneRenderer::GatherRayTracingWorldInstancesForView(FRHICo
 
 					FRayTracingGeometryInstance& RayTracingInstance = RayTracingScene.Instances.AddDefaulted_GetRef();
 					RayTracingInstance.GeometryRHI = Geometry->RayTracingGeometryRHI;
+					checkf(RayTracingInstance.GeometryRHI, TEXT("Ray tracing instance must have a valid geometry."));
+
 					RayTracingInstance.DefaultUserData = PrimitiveIndex;
 					RayTracingInstance.Mask = Instance.Mask;
 					if (Instance.bForceOpaque)
@@ -1150,6 +1152,7 @@ bool FDeferredShadingSceneRenderer::GatherRayTracingWorldInstancesForView(FRHICo
 					View.VisibleRayTracingMeshCommands.Add(NewVisibleMeshCommand);
 				}
 
+				checkf(SceneInfo->CachedRayTracingInstance.GeometryRHI, TEXT("Ray tracing instance must have a valid geometry."));
 				RayTracingScene.Instances.Add(SceneInfo->CachedRayTracingInstance);
 				AddDebugRayTracingInstanceFlags(RayTracingScene.Instances.Last().Flags);
 			}
@@ -1214,6 +1217,7 @@ bool FDeferredShadingSceneRenderer::GatherRayTracingWorldInstancesForView(FRHICo
 					FRayTracingGeometryInstance& RayTracingInstance = RayTracingScene.Instances.AddDefaulted_GetRef();
 
 					RayTracingInstance.GeometryRHI = RelevantPrimitive.RayTracingGeometryRHI;
+					checkf(RayTracingInstance.GeometryRHI, TEXT("Ray tracing instance must have a valid geometry."));
 
 					InstanceBatch.Add(RayTracingScene, Scene->PrimitiveTransforms[PrimitiveIndex], (uint32)PrimitiveIndex);
 					RayTracingInstance.Transforms = InstanceBatch.Transforms;
@@ -1463,13 +1467,28 @@ bool FDeferredShadingSceneRenderer::DispatchRayTracingWorldUpdates(FRDGBuilder&
 				[this, PassParams](FRHICommandList& RHICmdList)
 			{
 				FRHIRayTracingScene* RayTracingSceneRHI = Scene->RayTracingScene.GetRHIRayTracingSceneChecked();
-
-				RHICmdList.BindAccelerationStructureMemory(RayTracingSceneRHI, Scene->RayTracingScene.GetBufferChecked(), 0);
+				FRHIBuffer* AccelerationStructureBuffer = Scene->RayTracingScene.GetBufferChecked();
+				FRHIBuffer* ScratchBuffer = PassParams->RayTracingSceneScratchBuffer->GetRHI();
 
 				FRayTracingSceneBuildParams BuildParams;
 				BuildParams.Scene = RayTracingSceneRHI;
-				BuildParams.ScratchBuffer = PassParams->RayTracingSceneScratchBuffer->GetRHI();
+				BuildParams.ScratchBuffer = ScratchBuffer;
 				BuildParams.ScratchBufferOffset = 0;
+
+				// Sanity check acceleration structure buffer sizes
+			#if DO_CHECK
+				{
+					FRayTracingAccelerationStructureSize SizeInfo = RHICalcRayTracingSceneSize(
+						Scene->RayTracingScene.NumNativeInstances, ERayTracingAccelerationStructureFlags::FastTrace);
+
+					check(SizeInfo.ResultSize <= Scene->RayTracingScene.SizeInfo.ResultSize);
+					check(SizeInfo.BuildScratchSize <= Scene->RayTracingScene.SizeInfo.BuildScratchSize);
+					check(SizeInfo.ResultSize <= AccelerationStructureBuffer->GetSize());
+					check(SizeInfo.BuildScratchSize <= ScratchBuffer->GetSize());
+				}
+			#endif // DO_CHECK
+
+				RHICmdList.BindAccelerationStructureMemory(RayTracingSceneRHI, AccelerationStructureBuffer, 0);
 				RHICmdList.BuildAccelerationStructure(BuildParams);
 
 				// Submit potentially expensive BVH build commands to the GPU as soon as possible.
diff --git a/Engine/Source/Runtime/Renderer/Private/RayTracing/RayTracingScene.cpp b/Engine/Source/Runtime/Renderer/Private/RayTracing/RayTracingScene.cpp
--- a/Engine/Source/Runtime/Renderer/Private/RayTracing/RayTracingScene.cpp
+++ b/Engine/Source/Runtime/Renderer/Private/RayTracing/RayTracingScene.cpp
@@ -22,22 +22,25 @@ FGraphEventRef FRayTracingScene::BeginCreate(FRDGBuilder& GraphBuilder)
 {
 	WaitForTasks();
 
-	uint32 NumTotalInstances = 0;
+	NumNativeInstances = 0;
+
 	for (const FRayTracingGeometryInstance& InstanceDesc : Instances)
 	{
 		checkf(InstanceDesc.GPUTransformsSRV || InstanceDesc.NumTransforms <= uint32(InstanceDesc.Transforms.Num()),
 			TEXT("Expected at most %d ray tracing geometry instance transforms, but got %d."),
 			InstanceDesc.NumTransforms, InstanceDesc.Transforms.Num());
 
-		NumTotalInstances += InstanceDesc.NumTransforms;
+		checkf(InstanceDesc.GeometryRHI, TEXT("Ray tracing instance must have a valid geometry."));
+
+		NumNativeInstances += InstanceDesc.NumTransforms;
 	}
-	SET_DWORD_STAT(STAT_RayTracingInstances, NumTotalInstances);
+	SET_DWORD_STAT(STAT_RayTracingInstances, NumNativeInstances);
 
 	// Round up number of instances to some multiple to avoid pathological growth reallocations.
 	static constexpr uint32 AllocationGranularity = 8 * 1024;
-	NumTotalInstances = FMath::DivideAndRoundUp(NumTotalInstances, AllocationGranularity) * AllocationGranularity;
+	NumNativeInstances = FMath::DivideAndRoundUp(NumNativeInstances, AllocationGranularity) * AllocationGranularity;
 
-	SizeInfo = RHICalcRayTracingSceneSize(NumTotalInstances, ERayTracingAccelerationStructureFlags::FastTrace);
+	SizeInfo = RHICalcRayTracingSceneSize(NumNativeInstances, ERayTracingAccelerationStructureFlags::FastTrace);
 	check(SizeInfo.ResultSize < ~0u);
 
 	// Allocate GPU buffer if current one is too small or significantly larger than what we need.
diff --git a/Engine/Source/Runtime/Renderer/Private/RayTracing/RayTracingScene.h b/Engine/Source/Runtime/Renderer/Private/RayTracing/RayTracingScene.h
--- a/Engine/Source/Runtime/Renderer/Private/RayTracing/RayTracingScene.h
+++ b/Engine/Source/Runtime/Renderer/Private/RayTracing/RayTracingScene.h
@@ -74,6 +74,9 @@ public:
 	// This must be filled before calling BeginCreateTask().
 	TArray<FRayTracingGeometryInstance> Instances;
 
+	// Total flattened number of ray tracing geometry instances (a single FRayTracingGeometryInstance may represent many).
+	uint32 NumNativeInstances = 0;
+
 	// Geometries which still have a pending build request but are used this frame and require a force build.
 	TArray<const FRayTracingGeometry*> GeometriesToBuild;
 
