Use conservative CPU cull volume for VSM clipmaps to avoid missing non-nanite geometry in cached pages.
Better labeling for some GPU profiling regions.
Fix to frustum vertex buffer usage flag.

diff --git a/Engine/Source/Runtime/Renderer/Private/ShadowDepthRendering.cpp b/Engine/Source/Runtime/Renderer/Private/ShadowDepthRendering.cpp
--- a/Engine/Source/Runtime/Renderer/Private/ShadowDepthRendering.cpp
+++ b/Engine/Source/Runtime/Renderer/Private/ShadowDepthRendering.cpp
@@ -1706,7 +1706,7 @@ void FSceneRenderer::RenderShadowDepthMaps(FRDGBuilder& GraphBuilder, FInstanceC
 		const TRefCountPtr<IPooledRenderTarget> PrevHZBPhysical = bVSMUseHZB ? CacheManager->PrevBuffers.HZBPhysical : nullptr;
 		
 		{
-			RDG_EVENT_SCOPE(GraphBuilder, "RenderVirtualShadowMaps");
+			RDG_EVENT_SCOPE(GraphBuilder, "RenderVirtualShadowMaps(Nanite)");
 
 			const FIntPoint VirtualShadowSize = VirtualShadowMapArray.GetPhysicalPoolSize();
 			const FIntRect VirtualShadowViewRect = FIntRect(0, 0, VirtualShadowSize.X, VirtualShadowSize.Y);
@@ -1906,13 +1906,13 @@ void FSceneRenderer::RenderShadowDepthMaps(FRDGBuilder& GraphBuilder, FInstanceC
 			};
 
 			{
-				RDG_EVENT_SCOPE(GraphBuilder, "Directional Lights");
+				RDG_EVENT_SCOPE(GraphBuilder, "DirectionalLights");
 				static FString VirtualFilterName = TEXT("VSM_Directional");
 				FilterAndRenderVirtualShadowMaps(true, VirtualFilterName);
 			}
 
 			{
-				RDG_EVENT_SCOPE(GraphBuilder, "Perspective Lights (DepthClip)");
+				RDG_EVENT_SCOPE(GraphBuilder, "LocalLights");
 				static FString VirtualFilterName = TEXT("VSM_Perspective");
 				FilterAndRenderVirtualShadowMaps(false, VirtualFilterName);
 			}
diff --git a/Engine/Source/Runtime/Renderer/Private/ShadowSetup.cpp b/Engine/Source/Runtime/Renderer/Private/ShadowSetup.cpp
--- a/Engine/Source/Runtime/Renderer/Private/ShadowSetup.cpp
+++ b/Engine/Source/Runtime/Renderer/Private/ShadowSetup.cpp
@@ -277,6 +277,14 @@ static TAutoConsoleVariable<int32> CVarVsmUseFarShadowRules(
 	ECVF_RenderThreadSafe
 );
 
+static TAutoConsoleVariable<int32> CVarClipmapUseConservativeCulling(
+	TEXT("r.Shadow.Virtual.Clipmap.UseConservativeCulling"),
+	1,
+	TEXT(""),
+	ECVF_RenderThreadSafe
+);
+
+
 CSV_DECLARE_CATEGORY_EXTERN(LightCount);
 
 #if !UE_BUILD_SHIPPING
@@ -1049,8 +1057,19 @@ void FProjectedShadowInfo::SetupClipmapProjection(FLightSceneInfo* InLightSceneI
 	GetViewFrustumBounds(CasterOuterFrustum, CasterMatrix, true);
 	ReceiverInnerFrustum = CasterOuterFrustum;
 
-	// Magically get the accurate shadow bounds set up - ka-ching!
-	ShadowBounds = InLightSceneInfo->Proxy->GetShadowSplitBoundsDepthRange(*DependentView, DependentView->ViewMatrices.GetViewOrigin(), 10.0f, VirtualShadowMapClipmap->GetMaxRadius(), &CascadeSettings);
+	if (CVarClipmapUseConservativeCulling.GetValueOnRenderThread() != 0)
+	{
+		ShadowBounds = VirtualShadowMapClipmap->GetBoundingSphere();
+		CascadeSettings.ShadowBoundsAccurate = VirtualShadowMapClipmap->GetViewFrustumBounds();
+	}
+	else
+	{
+		// NOTE: Only here for debug! This version causes issues with VSM caching as it may cull objects
+		// that are outside the current camera frustum but overlap a page that will then be cached. We need
+		// all cached pages to be complete.
+		ShadowBounds = InLightSceneInfo->Proxy->GetShadowSplitBoundsDepthRange(*DependentView, DependentView->ViewMatrices.GetViewOrigin(), 10.0f, VirtualShadowMapClipmap->GetMaxRadius(), &CascadeSettings);
+	}
+	
 	// Um... it's checked in IsWholeSceneDirectionalShadow()
 	CascadeSettings.ShadowSplitIndex = 1000;
 
diff --git a/Engine/Source/Runtime/Renderer/Private/VirtualShadowMaps/VirtualShadowMapArray.cpp b/Engine/Source/Runtime/Renderer/Private/VirtualShadowMaps/VirtualShadowMapArray.cpp
--- a/Engine/Source/Runtime/Renderer/Private/VirtualShadowMaps/VirtualShadowMapArray.cpp
+++ b/Engine/Source/Runtime/Renderer/Private/VirtualShadowMaps/VirtualShadowMapArray.cpp
@@ -1694,7 +1694,7 @@ void FVirtualShadowMapArray::RenderVirtualShadowMapsHw(FRDGBuilder& GraphBuilder
 		return;
 	}
 
-	RDG_EVENT_SCOPE(GraphBuilder, "RenderVirtualShadowMapsHw");
+	RDG_EVENT_SCOPE(GraphBuilder, "RenderVirtualShadowMaps(Non-Nanite)");
 
 	const bool bAllocatePageRectAtlas = CVarAllocatePagesUsingRects.GetValueOnRenderThread() != 0;
 	const bool bAtomicWrites = GVirtualShadowMapAtomicWrites != 0;
diff --git a/Engine/Source/Runtime/Renderer/Private/VirtualShadowMaps/VirtualShadowMapClipmap.cpp b/Engine/Source/Runtime/Renderer/Private/VirtualShadowMaps/VirtualShadowMapClipmap.cpp
--- a/Engine/Source/Runtime/Renderer/Private/VirtualShadowMaps/VirtualShadowMapClipmap.cpp
+++ b/Engine/Source/Runtime/Renderer/Private/VirtualShadowMaps/VirtualShadowMapClipmap.cpp
@@ -75,9 +75,9 @@ FVirtualShadowMapClipmap::FVirtualShadowMapClipmap(
 		FPlane(-1, 0, 0, 0 ),
 		FPlane( 0, 0, 0, 1 ));
 
-	WorldToViewRotationMatrix = WorldToLightRotationMatrix * FaceMatrix;
+	WorldToLightViewRotationMatrix = WorldToLightRotationMatrix * FaceMatrix;
 	// Pure rotation matrix
-	FMatrix ViewToWorldRotationMatrix = WorldToViewRotationMatrix.GetTransposed();
+	FMatrix ViewToWorldRotationMatrix = WorldToLightViewRotationMatrix.GetTransposed();
 	
 	// NOTE: Rotational (roll) invariance of the directional light depends on square pixels so we just base everything on the camera X scales/resolution
 	// NOTE: 0.5 because we double the size of the clipmap region below to handle snapping
@@ -116,7 +116,7 @@ FVirtualShadowMapClipmap::FVirtualShadowMapClipmap(
 		float HalfLevelDim = 2.0f * RawLevelRadius;
 		float SnapSize = RawLevelRadius;
 
-		FVector ViewCenter = WorldToViewRotationMatrix.TransformPosition(WorldOrigin);
+		FVector ViewCenter = WorldToLightViewRotationMatrix.TransformPosition(WorldOrigin);
 		FIntPoint CenterSnapUnits(
 			FMath::RoundToInt(ViewCenter.X / SnapSize),
 			FMath::RoundToInt(ViewCenter.Y / SnapSize));
@@ -174,6 +174,26 @@ FVirtualShadowMapClipmap::FVirtualShadowMapClipmap(
 		const float ZOffset = ViewRadiusZ + ViewCenterDeltaZ;
 		Level.ViewToClip = FReversedZOrthoMatrix(HalfLevelDim, HalfLevelDim, ZScale, ZOffset);
 	}
+
+	ComputeBoundingVolumes(CameraViewMatrices);
+}
+
+void FVirtualShadowMapClipmap::ComputeBoundingVolumes(const FViewMatrices& CameraViewMatrices)
+{
+	// We don't really do much CPU culling with clipmaps. After various testing the fact that we are culling
+	// a single frustum that goes out and basically the entire map, and we have to extrude towards (and away!) from
+	// the light, and dilate to cover full pages at every clipmap level (to avoid culling something that will go
+	// into a page that then gets cached with incomplete geometry), in many situations there is effectively no
+	// culling that happens. For instance, as soon as the camera looks vaguely towards or away from the light direction,
+	// the extruded frustum effectively covers the whole world.
+
+	const FVector CameraOrigin = CameraViewMatrices.GetViewOrigin();
+	const FVector CameraDirection = CameraViewMatrices.GetViewMatrix().GetColumn(2);
+
+	// Thus we don't spend a lot of time trying to optimize for the easy cases and instead just pick an extremely
+	// conservative frustum.
+	ViewFrustumBounds = FConvexVolume();
+	BoundingSphere = FSphere(CameraOrigin, GetMaxRadius());
 }
 
 float FVirtualShadowMapClipmap::GetMaxRadius() const
@@ -191,7 +211,7 @@ FViewMatrices FVirtualShadowMapClipmap::GetViewMatrices(int32 ClipmapIndex) cons
 	// NOTE: Be careful here! There's special logic in FViewMatrices around ViewOrigin for ortho projections we need to bypass...
 	// There's also the fact that some of this data is going to be "wrong", due to the "overridden" matrix thing that shadows do
 	Initializer.ViewOrigin = Level.WorldCenter;
-	Initializer.ViewRotationMatrix = WorldToViewRotationMatrix;
+	Initializer.ViewRotationMatrix = WorldToLightViewRotationMatrix;
 	Initializer.ProjectionMatrix = Level.ViewToClip;
 
 	// TODO: This is probably unused in the shadows/nanite path, but coupling here is not ideal
@@ -207,10 +227,10 @@ FVirtualShadowMapProjectionShaderData FVirtualShadowMapClipmap::GetProjectionSha
 	
 	// NOTE: Some shader logic (projection, etc) assumes some of these parameters are constant across all levels in a clipmap
 	FVirtualShadowMapProjectionShaderData Data;
-	Data.TranslatedWorldToShadowViewMatrix = WorldToViewRotationMatrix;
+	Data.TranslatedWorldToShadowViewMatrix = WorldToLightViewRotationMatrix;
 	Data.ShadowViewToClipMatrix = Level.ViewToClip;
-	Data.TranslatedWorldToShadowUVMatrix = CalcTranslatedWorldToShadowUVMatrix(WorldToViewRotationMatrix, Level.ViewToClip);
-	Data.TranslatedWorldToShadowUVNormalMatrix = CalcTranslatedWorldToShadowUVNormalMatrix(WorldToViewRotationMatrix, Level.ViewToClip);
+	Data.TranslatedWorldToShadowUVMatrix = CalcTranslatedWorldToShadowUVMatrix(WorldToLightViewRotationMatrix, Level.ViewToClip);
+	Data.TranslatedWorldToShadowUVNormalMatrix = CalcTranslatedWorldToShadowUVNormalMatrix(WorldToLightViewRotationMatrix, Level.ViewToClip);
 	Data.ShadowPreViewTranslation = -Level.WorldCenter;
 	Data.VirtualShadowMapId = Level.VirtualShadowMap->ID;
 	Data.LightType = ELightComponentType::LightType_Directional;
diff --git a/Engine/Source/Runtime/Renderer/Private/VirtualShadowMaps/VirtualShadowMapClipmap.h b/Engine/Source/Runtime/Renderer/Private/VirtualShadowMaps/VirtualShadowMapClipmap.h
--- a/Engine/Source/Runtime/Renderer/Private/VirtualShadowMaps/VirtualShadowMapClipmap.h
+++ b/Engine/Source/Runtime/Renderer/Private/VirtualShadowMaps/VirtualShadowMapClipmap.h
@@ -59,6 +59,8 @@ public:
 	// Returns the max radius the clipmap is guaranteed to cover (i.e. the radius of the last clipmap level)
 	// Note that this is not a conservative radius of the level projection, which is snapped
 	float GetMaxRadius() const;
+	FSphere GetBoundingSphere() const { return BoundingSphere; }
+	FConvexVolume GetViewFrustumBounds() const { return ViewFrustumBounds; }
 
 	const FViewInfo* GetDependentView() const { return DependentView; }
 
@@ -67,6 +69,8 @@ public:
 	static uint32 GetCoarsePageClipmapIndexMask();
 
 private:
+	void ComputeBoundingVolumes(const FViewMatrices& CameraViewMatrices);
+
 	const FLightSceneInfo& LightSceneInfo;
 
 	/**
@@ -82,7 +86,7 @@ private:
 	FVector WorldOrigin;
 
 	/** Directional light rotation matrix (no translation) */
-	FMatrix WorldToViewRotationMatrix;
+	FMatrix WorldToLightViewRotationMatrix;
 
 	int32 FirstLevel;
 	float ResolutionLodBias;
@@ -95,6 +99,8 @@ private:
 		FVector WorldCenter;
 		FIntPoint CornerOffset;
 	};
+	TArray< FLevelData, TInlineAllocator<32> > LevelData;
 
-	TArray< FLevelData, TInlineAllocator<16> > LevelData;
+	FSphere BoundingSphere;
+	FConvexVolume ViewFrustumBounds;
 };
