Fixed crash creating ray tracing BVH TLAS due rounding up the number of instances in the high level code, which created a mismatch in the computations of buffer sizes between low and high level code.

diff --git a/Engine/Source/Runtime/Renderer/Private/RayTracing/RayTracingScene.cpp b/Engine/Source/Runtime/Renderer/Private/RayTracing/RayTracingScene.cpp
--- a/Engine/Source/Runtime/Renderer/Private/RayTracing/RayTracingScene.cpp
+++ b/Engine/Source/Runtime/Renderer/Private/RayTracing/RayTracingScene.cpp
@@ -38,9 +38,13 @@ FGraphEventRef FRayTracingScene::BeginCreate(FRDGBuilder& GraphBuilder)
 
 	// Round up number of instances to some multiple to avoid pathological growth reallocations.
 	static constexpr uint32 AllocationGranularity = 8 * 1024;
-	NumNativeInstances = FMath::DivideAndRoundUp(NumNativeInstances, AllocationGranularity) * AllocationGranularity;
+	uint32 NumNativeInstancesAligned = FMath::DivideAndRoundUp(NumNativeInstances, AllocationGranularity) * AllocationGranularity;
 
 	SizeInfo = RHICalcRayTracingSceneSize(NumNativeInstances, ERayTracingAccelerationStructureFlags::FastTrace);
+	FRayTracingAccelerationStructureSize SizeInfoAligned = RHICalcRayTracingSceneSize(NumNativeInstancesAligned, ERayTracingAccelerationStructureFlags::FastTrace);
+	SizeInfo.ResultSize = FMath::Max(SizeInfo.ResultSize, SizeInfoAligned.ResultSize);
+	SizeInfo.BuildScratchSize = FMath::Max(SizeInfo.BuildScratchSize, SizeInfoAligned.BuildScratchSize);
+
 	check(SizeInfo.ResultSize < ~0u);
 
 	// Allocate GPU buffer if current one is too small or significantly larger than what we need.
