Add an option to render distance field shadows using async compute. Also fix a bug where a raytraced DF FProjectedShadowInfo is rendered for all views instead of just its dependent view

diff --git a/Engine/Source/Runtime/Renderer/Private/DeferredShadingRenderer.cpp b/Engine/Source/Runtime/Renderer/Private/DeferredShadingRenderer.cpp
--- a/Engine/Source/Runtime/Renderer/Private/DeferredShadingRenderer.cpp
+++ b/Engine/Source/Runtime/Renderer/Private/DeferredShadingRenderer.cpp
@@ -2402,6 +2402,8 @@ void FDeferredShadingSceneRenderer::Render(FRDGBuilder& GraphBuilder)
 
 	// End early occlusion queries
 
+	BeginAsyncDistanceFieldShadowProjections(GraphBuilder, SceneTextures);
+
 	// Early Shadow depth rendering
 	if (!bHasRayTracedOverlay && bOcclusionBeforeBasePass)
 	{
diff --git a/Engine/Source/Runtime/Renderer/Private/DistanceFieldShadowing.cpp b/Engine/Source/Runtime/Renderer/Private/DistanceFieldShadowing.cpp
--- a/Engine/Source/Runtime/Renderer/Private/DistanceFieldShadowing.cpp
+++ b/Engine/Source/Runtime/Renderer/Private/DistanceFieldShadowing.cpp
@@ -86,6 +86,13 @@ FAutoConsoleVariableRef CVarAverageObjectsPerShadowCullTile(
 	ECVF_RenderThreadSafe | ECVF_ReadOnly
 	);
 
+int32 GDFShadowAsyncCompute = 0;
+static FAutoConsoleVariableRef CVarDFShadowAsyncCompute(
+	TEXT("r.DFShadowAsyncCompute"),
+	GDFShadowAsyncCompute,
+	TEXT("Whether render distance field shadows using async compute if possible"),
+	ECVF_RenderThreadSafe);
+
 static int32 GHeightFieldShadowing = 0;
 FAutoConsoleVariableRef CVarHeightFieldShadowing(
 	TEXT("r.HeightFieldShadowing"),
@@ -768,23 +775,27 @@ void RayTraceShadows(
 		FComputeShaderUtils::AddPass(
 			GraphBuilder,
 			RDG_EVENT_NAME("DistanceFieldShadowing %ux%u", GroupSizeX * GDistanceFieldShadowTileSizeX, GroupSizeY * GDistanceFieldShadowTileSizeY),
+			!!GDFShadowAsyncCompute ? ERDGPassFlags::AsyncCompute : ERDGPassFlags::Compute,
 			ComputeShader,
 			PassParameters,
 			FIntVector(GroupSizeX, GroupSizeY, 1));
 	}
 }
 
-FRDGTextureRef FProjectedShadowInfo::BeginRenderRayTracedDistanceFieldProjection(
+void FProjectedShadowInfo::BeginRenderRayTracedDistanceFieldProjection(
 	FRDGBuilder& GraphBuilder,
 	const FMinimalSceneTextures& SceneTextures,
-	const FViewInfo& View) const
+	const FViewInfo& View)
 {
+	if (RayTracedShadowsTexture)
+	{
+		return;
+	}
+
 	const bool bDFShadowSupported = SupportsDistanceFieldShadows(View.GetFeatureLevel(), View.GetShaderPlatform());
 	const bool bHFShadowSupported = SupportsHeightFieldShadows(View.GetFeatureLevel(), View.GetShaderPlatform());
 	const FScene* Scene = (const FScene*)(View.Family->Scene);
 
-	FRDGTextureRef RayTracedShadowsTexture = nullptr;
-
 	if (bDFShadowSupported && View.Family->EngineShowFlags.RayTracedDistanceFieldShadows)
 	{
 		QUICK_SCOPE_CYCLE_COUNTER(STAT_BeginRenderRayTracedDistanceFieldShadows);
@@ -909,8 +920,6 @@ FRDGTextureRef FProjectedShadowInfo::BeginRenderRayTracedDistanceFieldProjection
 
 		RayTraceShadows(GraphBuilder, SceneTextures, RayTracedShadowsTexture, View, Scene->DistanceFieldSceneData, this, DFPT_HeightField, bHasPrevOutput, PrevOutputTexture, ObjectBufferParameters, CulledObjectBufferParameters, LightTileIntersectionParameters);
 	}
-
-	return RayTracedShadowsTexture;
 }
 
 BEGIN_SHADER_PARAMETER_STRUCT(FDistanceFieldShadowingUpsample, )
@@ -924,11 +933,11 @@ void FProjectedShadowInfo::RenderRayTracedDistanceFieldProjection(
 	FRDGTextureRef ScreenShadowMaskTexture,
 	const FViewInfo& View,
 	FIntRect ScissorRect,
-	bool bProjectingForForwardShading) const
+	bool bProjectingForForwardShading)
 {
 	check(ScissorRect.Area() > 0);
 
-	FRDGTextureRef RayTracedShadowsTexture = BeginRenderRayTracedDistanceFieldProjection(GraphBuilder, SceneTextures, View);
+	BeginRenderRayTracedDistanceFieldProjection(GraphBuilder, SceneTextures, View);
 
 	if (RayTracedShadowsTexture)
 	{
diff --git a/Engine/Source/Runtime/Renderer/Private/SceneRendering.h b/Engine/Source/Runtime/Renderer/Private/SceneRendering.h
--- a/Engine/Source/Runtime/Renderer/Private/SceneRendering.h
+++ b/Engine/Source/Runtime/Renderer/Private/SceneRendering.h
@@ -1990,6 +1990,10 @@ protected:
 		const FLightSceneInfo* LightSceneInfo,
 		bool bProjectingForForwardShading);
 
+	void BeginAsyncDistanceFieldShadowProjections(
+		FRDGBuilder& GraphBuilder,
+		const FMinimalSceneTextures& SceneTextures) const;
+
 	/** Finds a matching cached preshadow, if one exists. */
 	TRefCountPtr<FProjectedShadowInfo> GetCachedPreshadow(
 		const FLightPrimitiveInteraction* InParentInteraction,
@@ -2178,6 +2182,9 @@ private:
 #endif
 	bool bShadowDepthRenderCompleted;
 
+	/** Distance field shadows to project. Used to avoid iterating through the scene lights array. */
+	TArray<FProjectedShadowInfo*, TInlineAllocator<2, SceneRenderingAllocator>> ProjectedDistanceFieldShadows;
+
 	/**
 	 * Projected shadows allocated on the scene rendering mem stack.
 	 * Note: this exists purely such that the destructor can be called as the memory is allocated from a pool.
diff --git a/Engine/Source/Runtime/Renderer/Private/ShadowRendering.cpp b/Engine/Source/Runtime/Renderer/Private/ShadowRendering.cpp
--- a/Engine/Source/Runtime/Renderer/Private/ShadowRendering.cpp
+++ b/Engine/Source/Runtime/Renderer/Private/ShadowRendering.cpp
@@ -2119,14 +2119,53 @@ void FSceneRenderer::RenderShadowProjections(
 			{
 				for (int32 ShadowIndex = 0; ShadowIndex < DistanceFieldShadows.Num(); ShadowIndex++)
 				{
-					const FProjectedShadowInfo* ProjectedShadowInfo = DistanceFieldShadows[ShadowIndex];
-					ProjectedShadowInfo->RenderRayTracedDistanceFieldProjection(
-						GraphBuilder,
-						SceneTextures,
-						ScreenShadowMaskTexture,
-						View,
-						ScissorRect,
-						bProjectingForForwardShading);
+					FProjectedShadowInfo* ProjectedShadowInfo = DistanceFieldShadows[ShadowIndex];
+
+					if (Views.Num() == 1 || ProjectedShadowInfo->DependentView == &View || !ProjectedShadowInfo->DependentView)
+					{
+						ProjectedShadowInfo->RenderRayTracedDistanceFieldProjection(
+							GraphBuilder,
+							SceneTextures,
+							ScreenShadowMaskTexture,
+							View,
+							ScissorRect,
+							bProjectingForForwardShading);
+					}
+				}
+			}
+		}
+	}
+}
+
+void FSceneRenderer::BeginAsyncDistanceFieldShadowProjections(FRDGBuilder& GraphBuilder, const FMinimalSceneTextures& SceneTextures) const
+{
+	extern int32 GDFShadowAsyncCompute;
+
+	if (!!GDFShadowAsyncCompute && ViewFamily.EngineShowFlags.DynamicShadows && GetShadowQuality() > 0 && ProjectedDistanceFieldShadows.Num() > 0)
+	{
+		RDG_EVENT_SCOPE(GraphBuilder, "DistanceFieldShadows");
+
+		for (int32 ViewIndex = 0; ViewIndex < Views.Num(); ViewIndex++)
+		{
+			const FViewInfo& View = Views[ViewIndex];
+			RDG_GPU_MASK_SCOPE(GraphBuilder, View.GPUMask);
+			RDG_EVENT_SCOPE_CONDITIONAL(GraphBuilder, Views.Num() > 1, "View%d", ViewIndex);
+
+			for (int32 DFShadowIndex = 0; DFShadowIndex < ProjectedDistanceFieldShadows.Num(); ++DFShadowIndex)
+			{
+				FProjectedShadowInfo* ProjectedShadowInfo = ProjectedDistanceFieldShadows[DFShadowIndex];
+
+				check(ProjectedShadowInfo->bIncludeInScreenSpaceShadowMask);
+
+				FIntRect ScissorRect;
+				if (ProjectedShadowInfo->bDirectionalLight || !ProjectedShadowInfo->GetLightSceneInfo().Proxy->GetScissorRect(ScissorRect, View, View.ViewRect))
+				{
+					ScissorRect = View.ViewRect;
+				}
+
+				if (ScissorRect.Area() > 0 && (Views.Num() == 1 || ProjectedShadowInfo->DependentView == &View || !ProjectedShadowInfo->DependentView))
+				{
+					ProjectedShadowInfo->BeginRenderRayTracedDistanceFieldProjection(GraphBuilder, SceneTextures, View);
 				}
 			}
 		}
diff --git a/Engine/Source/Runtime/Renderer/Private/ShadowRendering.h b/Engine/Source/Runtime/Renderer/Private/ShadowRendering.h
--- a/Engine/Source/Runtime/Renderer/Private/ShadowRendering.h
+++ b/Engine/Source/Runtime/Renderer/Private/ShadowRendering.h
@@ -546,10 +546,10 @@ public:
 		FRDGTextureRef ShadowDepthTexture,
 		bool bDoParallelDispatch);
 
-	FRDGTextureRef BeginRenderRayTracedDistanceFieldProjection(
+	void BeginRenderRayTracedDistanceFieldProjection(
 		FRDGBuilder& GraphBuilder,
 		const FMinimalSceneTextures& SceneTextures,
-		const FViewInfo& View) const;
+		const FViewInfo& View);
 
 	/** Renders ray traced distance field shadows. */
 	void RenderRayTracedDistanceFieldProjection(
@@ -558,7 +558,7 @@ public:
 		FRDGTextureRef ScreenShadowMaskTexture,
 		const FViewInfo& View,
 		FIntRect ScissorRect,
-		bool bProjectingForForwardShading) const;
+		bool bProjectingForForwardShading);
 
 	/** Render one pass point light shadow projections. */
 	void RenderOnePassPointLightProjection(
@@ -758,6 +758,9 @@ private:
 	float ShaderSlopeDepthBias;
 	float ShaderMaxSlopeDepthBias;
 
+	/** Ray traced DF shadow intermediate output. Populated by BeginRenderRayTracedDistanceFieldProjection and consumed by RenderRayTracedDistanceFieldProjection. */
+	FRDGTextureRef RayTracedShadowsTexture;
+
 	void CopyCachedShadowMap(
 		FRDGBuilder& GraphBuilder,
 		const FViewInfo& View,
diff --git a/Engine/Source/Runtime/Renderer/Private/ShadowSetup.cpp b/Engine/Source/Runtime/Renderer/Private/ShadowSetup.cpp
--- a/Engine/Source/Runtime/Renderer/Private/ShadowSetup.cpp
+++ b/Engine/Source/Runtime/Renderer/Private/ShadowSetup.cpp
@@ -708,6 +708,7 @@ FProjectedShadowInfo::FProjectedShadowInfo()
 	, NumSubjectMeshCommandBuildRequestElements(0)
 	, ShaderDepthBias(0.0f)
 	, ShaderSlopeDepthBias(0.0f)
+	, RayTracedShadowsTexture(nullptr)
 {
 }
 
@@ -4841,12 +4842,17 @@ void FSceneRenderer::AllocateShadowDepthTargets(FRHICommandListImmediate& RHICmd
 					&& (!bMobile || ProjectedShadowInfo->bPerObjectOpaqueShadow || ProjectedShadowInfo->bRayTracedDistanceField);
 
 				if (bNeedsProjection)
-				{if (ProjectedShadowInfo->bCapsuleShadow)
+				{
+					if (ProjectedShadowInfo->bCapsuleShadow)
 					{
 						VisibleLightInfo.CapsuleShadowsToProject.Add(ProjectedShadowInfo);
 					}
 					else
 					{
+						if (ProjectedShadowInfo->bRayTracedDistanceField)
+						{
+							ProjectedDistanceFieldShadows.Add(ProjectedShadowInfo);
+						}
 						VisibleLightInfo.ShadowsToProject.Add(ProjectedShadowInfo);
 					}
 				}
