Use StructuredBuffers in FTileIntersectionParameters.
Also use rounded MaxSceneObjects in AllocateTileIntersectionBuffers.

diff --git a/Engine/Shaders/Private/DistanceFieldAOShared.ush b/Engine/Shaders/Private/DistanceFieldAOShared.ush
--- a/Engine/Shaders/Private/DistanceFieldAOShared.ush
+++ b/Engine/Shaders/Private/DistanceFieldAOShared.ush
@@ -292,7 +292,7 @@ float2 GetScreenUVFromScreenGrid(uint2 OutputCoordinate)
 	return GetScreenUVFromScreenGrid(OutputCoordinate, 1);
 }
 
-Buffer<uint> CulledTilesStartOffsetArray;
+StructuredBuffer<uint> CulledTilesStartOffsetArray;
 
 #define INVALID_TILE_INDEX 0xFFFF
 
diff --git a/Engine/Shaders/Private/DistanceFieldObjectCulling.usf b/Engine/Shaders/Private/DistanceFieldObjectCulling.usf
--- a/Engine/Shaders/Private/DistanceFieldObjectCulling.usf
+++ b/Engine/Shaders/Private/DistanceFieldObjectCulling.usf
@@ -128,8 +128,8 @@ groupshared uint IntegerTileMaxZ2;
 uint4 ViewDimensions;
 float2 NumGroups;
 
-RWBuffer<float4> RWTileConeAxisAndCos;
-RWBuffer<float4> RWTileConeDepthRanges;
+RWStructuredBuffer<float4> RWTileConeAxisAndCos;
+RWStructuredBuffer<float4> RWTileConeDepthRanges;
 
 /** Builds tile depth ranges and bounding cones. */
 [numthreads(THREADGROUP_SIZEX, THREADGROUP_SIZEY, 1)]
@@ -317,11 +317,11 @@ bool IntersectObjectWithConeDepthRange(
 	return false;
 }
 
-Buffer<float4> TileConeAxisAndCos;
-Buffer<float4> TileConeDepthRanges;
+StructuredBuffer<float4> TileConeAxisAndCos;
+StructuredBuffer<float4> TileConeDepthRanges;
 
-RWBuffer<uint> RWNumCulledTilesArray;
-RWBuffer<uint> RWCulledTilesStartOffsetArray;
+RWStructuredBuffer<uint> RWNumCulledTilesArray;
+RWStructuredBuffer<uint> RWCulledTilesStartOffsetArray;
 RWBuffer<uint> RWCulledTileDataArray;
 
 /** Intersects a single object with the tile and adds to the intersection list if needed. */
@@ -381,7 +381,7 @@ void ObjectCullPS(
 }
 
 RWBuffer<uint> RWObjectTilesIndirectArguments;
-Buffer<uint> NumCulledTilesArray;
+StructuredBuffer<uint> NumCulledTilesArray;
 
 #ifndef COMPUTE_START_OFFSET_GROUP_SIZE
 #define COMPUTE_START_OFFSET_GROUP_SIZE 1
diff --git a/Engine/Source/Runtime/Renderer/Private/DistanceFieldAmbientOcclusion.cpp b/Engine/Source/Runtime/Renderer/Private/DistanceFieldAmbientOcclusion.cpp
--- a/Engine/Source/Runtime/Renderer/Private/DistanceFieldAmbientOcclusion.cpp
+++ b/Engine/Source/Runtime/Renderer/Private/DistanceFieldAmbientOcclusion.cpp
@@ -623,33 +623,34 @@ void AllocateTileIntersectionBuffers(
 	const bool b16BitObjectIndices = MaxSceneObjects < (1 << 16);
 	const bool b16BitCulledTileIndexBuffer = bAllow16BitIndices && b16BitObjectIndices && (TileListGroupSize.X * TileListGroupSize.Y < (1 << 16));
 
-	FRDGBufferRef TileConeAxisAndCos = GraphBuilder.CreateBuffer(FRDGBufferDesc::CreateBufferDesc(sizeof(FVector4), TileListGroupSize.X * TileListGroupSize.Y), TEXT("TileConeAxisAndCos"));
-	OutParameters.RWTileConeAxisAndCos = GraphBuilder.CreateUAV(TileConeAxisAndCos, PF_A32B32G32R32F);
-	OutParameters.TileConeAxisAndCos = GraphBuilder.CreateSRV(TileConeAxisAndCos, PF_A32B32G32R32F);
+	int32 TileCount = TileListGroupSize.X * TileListGroupSize.Y;
 
-	FRDGBufferRef TileConeDepthRanges = GraphBuilder.CreateBuffer(FRDGBufferDesc::CreateBufferDesc(sizeof(FVector4), TileListGroupSize.X * TileListGroupSize.Y), TEXT("TileConeDepthRanges"));
-	OutParameters.RWTileConeDepthRanges = GraphBuilder.CreateUAV(TileConeDepthRanges, PF_A32B32G32R32F);
-	OutParameters.TileConeDepthRanges = GraphBuilder.CreateSRV(TileConeDepthRanges, PF_A32B32G32R32F);
+	FRDGBufferRef TileConeAxisAndCos = GraphBuilder.CreateBuffer(FRDGBufferDesc::CreateStructuredDesc(sizeof(FVector4), TileCount), TEXT("TileConeAxisAndCos"));
+	OutParameters.RWTileConeAxisAndCos = GraphBuilder.CreateUAV(TileConeAxisAndCos);
+	OutParameters.TileConeAxisAndCos = GraphBuilder.CreateSRV(TileConeAxisAndCos);
 
-	FRDGBufferRef NumCulledTilesArray = GraphBuilder.CreateBuffer(FRDGBufferDesc::CreateBufferDesc(sizeof(uint32), MaxSceneObjects), TEXT("NumCulledTilesArray"));
-	OutParameters.RWNumCulledTilesArray = GraphBuilder.CreateUAV(NumCulledTilesArray, PF_R32_UINT);
-	OutParameters.NumCulledTilesArray = GraphBuilder.CreateSRV(NumCulledTilesArray, PF_R32_UINT);
+	FRDGBufferRef TileConeDepthRanges = GraphBuilder.CreateBuffer(FRDGBufferDesc::CreateStructuredDesc(sizeof(FVector4), TileCount), TEXT("TileConeDepthRanges"));
+	OutParameters.RWTileConeDepthRanges = GraphBuilder.CreateUAV(TileConeDepthRanges);
+	OutParameters.TileConeDepthRanges = GraphBuilder.CreateSRV(TileConeDepthRanges);
 
-	FRDGBufferRef CulledTilesStartOffsetArray = GraphBuilder.CreateBuffer(FRDGBufferDesc::CreateBufferDesc(sizeof(uint32), MaxSceneObjects), TEXT("CulledTilesStartOffsetArray"));
-	OutParameters.RWCulledTilesStartOffsetArray = GraphBuilder.CreateUAV(CulledTilesStartOffsetArray, PF_R32_UINT);
-	OutParameters.CulledTilesStartOffsetArray = GraphBuilder.CreateSRV(CulledTilesStartOffsetArray, PF_R32_UINT);
+	FRDGBufferRef NumCulledTilesArray = GraphBuilder.CreateBuffer(FRDGBufferDesc::CreateStructuredDesc(sizeof(uint32), MaxSceneObjects), TEXT("NumCulledTilesArray"));
+	OutParameters.RWNumCulledTilesArray = GraphBuilder.CreateUAV(NumCulledTilesArray);
+	OutParameters.NumCulledTilesArray = GraphBuilder.CreateSRV(NumCulledTilesArray);
+
+	FRDGBufferRef CulledTilesStartOffsetArray = GraphBuilder.CreateBuffer(FRDGBufferDesc::CreateStructuredDesc(sizeof(uint32), MaxSceneObjects), TEXT("CulledTilesStartOffsetArray"));
+	OutParameters.RWCulledTilesStartOffsetArray = GraphBuilder.CreateUAV(CulledTilesStartOffsetArray);
+	OutParameters.CulledTilesStartOffsetArray = GraphBuilder.CreateSRV(CulledTilesStartOffsetArray);
 
 	extern int32 GAverageDistanceFieldObjectsPerCullTile;
 
 	FRDGBufferRef CulledTileDataArray = GraphBuilder.CreateBuffer(
-		FRDGBufferDesc::CreateBufferDesc(b16BitCulledTileIndexBuffer ? sizeof(uint16) : sizeof(uint32), GAverageDistanceFieldObjectsPerCullTile * TileListGroupSize.X * TileListGroupSize.Y * CulledTileDataStride),
+		FRDGBufferDesc::CreateBufferDesc(b16BitCulledTileIndexBuffer ? sizeof(uint16) : sizeof(uint32), GAverageDistanceFieldObjectsPerCullTile * TileCount * CulledTileDataStride),
 		TEXT("CulledTileDataArray"));
 	OutParameters.RWCulledTileDataArray = GraphBuilder.CreateUAV(CulledTileDataArray, b16BitCulledTileIndexBuffer ? PF_R16_UINT : PF_R32_UINT);
 	OutParameters.CulledTileDataArray = GraphBuilder.CreateSRV(CulledTileDataArray, b16BitCulledTileIndexBuffer ? PF_R16_UINT : PF_R32_UINT);
 
 	OutObjectTilesIndirectArguments = GraphBuilder.CreateBuffer(FRDGBufferDesc::CreateIndirectDesc<FRHIDispatchIndirectParameters>(), TEXT("ObjectTilesIndirectArguments"));
 	OutParameters.RWObjectTilesIndirectArguments = GraphBuilder.CreateUAV(OutObjectTilesIndirectArguments, PF_R32_UINT);
-	OutParameters.ObjectTilesIndirectArguments = GraphBuilder.CreateSRV(OutObjectTilesIndirectArguments, PF_R32_UINT);
 
 	OutParameters.TileListGroupSize = TileListGroupSize;
 }
@@ -809,7 +810,7 @@ void FDeferredShadingSceneRenderer::RenderDistanceFieldLighting(
 	}
 
 	const FIntPoint TileListGroupSize = GetTileListGroupSizeForView(View);
-	const int32 MaxSceneObjects = Scene->DistanceFieldSceneData.NumObjectsInBuffer;
+	const int32 MaxSceneObjects = FMath::DivideAndRoundUp(Scene->DistanceFieldSceneData.NumObjectsInBuffer, 256) * 256;
 	const bool bAllow16BitIndices = !IsMetalPlatform(GShaderPlatformForFeatureLevel[View.FeatureLevel]);
 
 	FRDGBufferRef ObjectTilesIndirectArguments = nullptr;
@@ -825,7 +826,7 @@ void FDeferredShadingSceneRenderer::RenderDistanceFieldLighting(
 		AllocateDistanceFieldCulledObjectBuffers(
 			GraphBuilder,
 			false,
-			FMath::DivideAndRoundUp(MaxSceneObjects, 256) * 256,
+			MaxSceneObjects,
 			DFPT_SignedDistanceField,
 			ObjectIndirectArguments,
 			CulledObjectBufferParameters);
diff --git a/Engine/Source/Runtime/Renderer/Private/DistanceFieldAmbientOcclusion.h b/Engine/Source/Runtime/Renderer/Private/DistanceFieldAmbientOcclusion.h
--- a/Engine/Source/Runtime/Renderer/Private/DistanceFieldAmbientOcclusion.h
+++ b/Engine/Source/Runtime/Renderer/Private/DistanceFieldAmbientOcclusion.h
@@ -39,21 +39,20 @@ public:
 };
 
 BEGIN_SHADER_PARAMETER_STRUCT(FTileIntersectionParameters, )
-	SHADER_PARAMETER_RDG_BUFFER_UAV(RWBuffer<FVector4>, RWTileConeAxisAndCos)
-	SHADER_PARAMETER_RDG_BUFFER_UAV(RWBuffer<FVector4>, RWTileConeDepthRanges)
+	SHADER_PARAMETER_RDG_BUFFER_UAV(RWStructuredBuffer<FVector4>, RWTileConeAxisAndCos)
+	SHADER_PARAMETER_RDG_BUFFER_UAV(RWStructuredBuffer<FVector4>, RWTileConeDepthRanges)
 
-	SHADER_PARAMETER_RDG_BUFFER_UAV(RWBuffer<uint>, RWNumCulledTilesArray)
-	SHADER_PARAMETER_RDG_BUFFER_UAV(RWBuffer<uint>, RWCulledTilesStartOffsetArray)
+	SHADER_PARAMETER_RDG_BUFFER_UAV(RWStructuredBuffer<uint>, RWNumCulledTilesArray)
+	SHADER_PARAMETER_RDG_BUFFER_UAV(RWStructuredBuffer<uint>, RWCulledTilesStartOffsetArray)
 	SHADER_PARAMETER_RDG_BUFFER_UAV(RWBuffer<uint>, RWCulledTileDataArray)
 	SHADER_PARAMETER_RDG_BUFFER_UAV(RWBuffer<uint>, RWObjectTilesIndirectArguments)
 
-	SHADER_PARAMETER_RDG_BUFFER_SRV(Buffer<FVector4>, TileConeAxisAndCos)
-	SHADER_PARAMETER_RDG_BUFFER_SRV(Buffer<FVector4>, TileConeDepthRanges)
+	SHADER_PARAMETER_RDG_BUFFER_SRV(StructuredBuffer<FVector4>, TileConeAxisAndCos)
+	SHADER_PARAMETER_RDG_BUFFER_SRV(StructuredBuffer<FVector4>, TileConeDepthRanges)
 
-	SHADER_PARAMETER_RDG_BUFFER_SRV(Buffer<uint>, NumCulledTilesArray)
-	SHADER_PARAMETER_RDG_BUFFER_SRV(Buffer<uint>, CulledTilesStartOffsetArray)
+	SHADER_PARAMETER_RDG_BUFFER_SRV(StructuredBuffer<uint>, NumCulledTilesArray)
+	SHADER_PARAMETER_RDG_BUFFER_SRV(StructuredBuffer<uint>, CulledTilesStartOffsetArray)
 	SHADER_PARAMETER_RDG_BUFFER_SRV(Buffer<uint>, CulledTileDataArray)
-	SHADER_PARAMETER_RDG_BUFFER_SRV(Buffer<uint>, ObjectTilesIndirectArguments)
 
 	SHADER_PARAMETER(FIntPoint, TileListGroupSize)
 END_SHADER_PARAMETER_STRUCT()
diff --git a/Engine/Source/Runtime/Renderer/Private/DistanceFieldObjectCulling.cpp b/Engine/Source/Runtime/Renderer/Private/DistanceFieldObjectCulling.cpp
--- a/Engine/Source/Runtime/Renderer/Private/DistanceFieldObjectCulling.cpp
+++ b/Engine/Source/Runtime/Renderer/Private/DistanceFieldObjectCulling.cpp
@@ -110,8 +110,8 @@ class FBuildTileConesCS : public FGlobalShader
 public:
 
 	BEGIN_SHADER_PARAMETER_STRUCT(FParameters, )
-		SHADER_PARAMETER_RDG_BUFFER_UAV(RWBuffer<FVector4>, RWTileConeAxisAndCos)
-		SHADER_PARAMETER_RDG_BUFFER_UAV(RWBuffer<FVector4>, RWTileConeDepthRanges)
+		SHADER_PARAMETER_RDG_BUFFER_UAV(RWStructuredBuffer<FVector4>, RWTileConeAxisAndCos)
+		SHADER_PARAMETER_RDG_BUFFER_UAV(RWStructuredBuffer<FVector4>, RWTileConeDepthRanges)
 		SHADER_PARAMETER_RDG_UNIFORM_BUFFER(FSceneTextureUniformParameters, SceneTextures)
 		SHADER_PARAMETER_RDG_TEXTURE(Texture2D, DistanceFieldNormalTexture)
 		SHADER_PARAMETER_SAMPLER(SamplerState, DistanceFieldNormalSampler)
