Use explicit 'float' types for FVirtualShadowMapProjectionShaderData, as it is memcpied to the GPU directly

diff --git a/Engine/Source/Runtime/Renderer/Private/VirtualShadowMaps/VirtualShadowMapArray.h b/Engine/Source/Runtime/Renderer/Private/VirtualShadowMaps/VirtualShadowMapArray.h
--- a/Engine/Source/Runtime/Renderer/Private/VirtualShadowMaps/VirtualShadowMapArray.h
+++ b/Engine/Source/Runtime/Renderer/Private/VirtualShadowMaps/VirtualShadowMapArray.h
@@ -80,15 +80,15 @@ struct FVirtualShadowMapProjectionShaderData
 	 * Transform from shadow-pre-translated world space to shadow view space, example use: (WorldSpacePos + ShadowPreViewTranslation) * TranslatedWorldToShadowViewMatrix
 	 * TODO: Why don't we call it a rotation and store in a 3x3? Does it ever have translation in?
 	 */
-	FMatrix TranslatedWorldToShadowViewMatrix;
-	FMatrix ShadowViewToClipMatrix;
-	FMatrix TranslatedWorldToShadowUVMatrix;
-	FMatrix TranslatedWorldToShadowUVNormalMatrix;
+	FMatrix44f TranslatedWorldToShadowViewMatrix;
+	FMatrix44f ShadowViewToClipMatrix;
+	FMatrix44f TranslatedWorldToShadowUVMatrix;
+	FMatrix44f TranslatedWorldToShadowUVNormalMatrix;
 
-	FVector ShadowPreViewTranslation;
+	FVector3f ShadowPreViewTranslation;
 	uint32 LightType = ELightComponentType::LightType_Directional;
 	
-	FVector ClipmapWorldOrigin;
+	FVector3f ClipmapWorldOrigin;
 	int32 VirtualShadowMapId = INDEX_NONE;
 	
 	FIntPoint ClipmapCornerOffset;
