Matinee/Rendering: Removed Matinee time from rendering stats.

diff --git a/Engine/Source/Runtime/Engine/Public/SceneView.h b/Engine/Source/Runtime/Engine/Public/SceneView.h
--- a/Engine/Source/Runtime/Engine/Public/SceneView.h
+++ b/Engine/Source/Runtime/Engine/Public/SceneView.h
@@ -1380,14 +1380,11 @@ struct FDisplayInternalsData
 {
 	//
 	int32 DisplayInternalsCVarValue;
-	// current time Matinee location (in seconds) of the single playing playing actor, -1 if none is playing, -2 if multiple are playing
-	float MatineeTime;
 	// -1 if not set, from IStreamingManager::Get().StreamAllResources(Duration) in FStreamAllResourcesLatentCommand
 	uint32 NumPendingStreamingRequests;
 
 	FDisplayInternalsData()
 		: DisplayInternalsCVarValue(0)
-		, MatineeTime(-1.0f)
 		, NumPendingStreamingRequests(-1)
 	{
 		check(!IsValid());
diff --git a/Engine/Source/Runtime/Renderer/Private/SceneRendering.cpp b/Engine/Source/Runtime/Renderer/Private/SceneRendering.cpp
--- a/Engine/Source/Runtime/Renderer/Private/SceneRendering.cpp
+++ b/Engine/Source/Runtime/Renderer/Private/SceneRendering.cpp
@@ -2211,27 +2211,6 @@ void FDisplayInternalsData::Setup(UWorld *World)
 
 	if(IsValid())
 	{
-		MatineeTime = -1.0f;
-		uint32 Count = 0;
-
-		for (TObjectIterator<AMatineeActor> It; It; ++It)
-		{
-			AMatineeActor* MatineeActor = *It;
-
-			if(MatineeActor->GetWorld() == World && MatineeActor->bIsPlaying)
-			{
-				MatineeTime = MatineeActor->InterpPosition;
-				++Count;
-			}
-		}
-
-		if(Count > 1)
-		{
-			MatineeTime = -2;
-		}
-
-		check(IsValid());
-		
 		extern ENGINE_API uint32 GStreamAllResourcesStillInFlight;
 		NumPendingStreamingRequests = GStreamAllResourcesStillInFlight;
 	}
@@ -4385,7 +4364,6 @@ static void DisplayInternals(FRHICommandListImmediate& RHICmdList, FViewInfo& In
 
 		CANVAS_HEADER(TEXT("Family:"))
 		CANVAS_LINE(false, TEXT("  Time (Real/World/DeltaWorld): %.2f/%.2f/%.2f"), Family->CurrentRealTime, Family->CurrentWorldTime, Family->DeltaWorldTime)
-		CANVAS_LINE(false, TEXT("  MatineeTime: %f"), Family->DisplayInternalsData.MatineeTime)
 		CANVAS_LINE(false, TEXT("  FrameNumber: %u"), Family->FrameNumber)
 		CANVAS_LINE(false, TEXT("  ExposureSettings: %s"), *Family->ExposureSettings.ToString())
 		CANVAS_LINE(false, TEXT("  GammaCorrection: %.2f"), Family->GammaCorrection)
