Modified FCapsuleShadowingCS and FCapsuleShadowingUpsamplePS to use FPermutationDomain instead of templates.

diff --git a/Engine/Shaders/Private/CapsuleShadowShaders.usf b/Engine/Shaders/Private/CapsuleShadowShaders.usf
--- a/Engine/Shaders/Private/CapsuleShadowShaders.usf
+++ b/Engine/Shaders/Private/CapsuleShadowShaders.usf
@@ -24,6 +24,7 @@
 #define LIGHT_SOURCE_MODE 0
 #endif
 
+// must match CapsuleShadowRendering.cpp
 #define LIGHT_SOURCE_PUNCTUAL 0
 #define LIGHT_SOURCE_FROM_CAPSULE 1
 #define LIGHT_SOURCE_FROM_RECEIVER 2
diff --git a/Engine/Source/Runtime/Renderer/Private/CapsuleShadowRendering.cpp b/Engine/Source/Runtime/Renderer/Private/CapsuleShadowRendering.cpp
--- a/Engine/Source/Runtime/Renderer/Private/CapsuleShadowRendering.cpp
+++ b/Engine/Source/Runtime/Renderer/Private/CapsuleShadowRendering.cpp
@@ -163,28 +163,43 @@ FIntPoint GetBufferSizeForCapsuleShadows()
 	return FIntPoint::DivideAndRoundDown(GetSceneTextureExtent(), GetCapsuleShadowDownsampleFactor());
 }
 
-enum ECapsuleShadowingType
+enum class ECapsuleShadowingType
 {
-	ShapeShadow_DirectionalLightTiledCulling,
-	ShapeShadow_PointLightTiledCulling,
-	ShapeShadow_IndirectTiledCulling,
-	ShapeShadow_MovableSkylightTiledCulling,
-	ShapeShadow_MovableSkylightTiledCullingGatherFromReceiverBentNormal
+	DirectionalLightTiledCulling,
+	PointLightTiledCulling,
+	IndirectTiledCulling,
+	MovableSkylightTiledCulling,
+	MovableSkylightTiledCullingGatherFromReceiverBentNormal,
+
+	MAX
+};
+
+enum class EIndirectShadowingPrimitiveTypes
+{
+	CapsuleShapes,
+	MeshDistanceFields,
+	CapsuleShapesAndMeshDistanceFields,
+
+	MAX
 };
 
-enum EIndirectShadowingPrimitiveTypes
+enum class ELightSourceMode
 {
-	IPT_CapsuleShapes = 1,
-	IPT_MeshDistanceFields = 2,
-	IPT_CapsuleShapesAndMeshDistanceFields = IPT_CapsuleShapes | IPT_MeshDistanceFields
+	// must match CapsuleShadowShaders.usf
+	Punctual = 0,
+	FromCapsule = 1,
+	FromReceiver = 2,
 };
 
-template<ECapsuleShadowingType ShadowingType>
-class TCapsuleShadowingBaseCS : public FGlobalShader
+class FCapsuleShadowingCS : public FGlobalShader
 {
-	DECLARE_INLINE_TYPE_LAYOUT(TCapsuleShadowingBaseCS, NonVirtual);
+	DECLARE_GLOBAL_SHADER(FCapsuleShadowingCS);
 public:
 
+	class FShapeShadow : SHADER_PERMUTATION_ENUM_CLASS("SHAPE_SHADOW", ECapsuleShadowingType);
+	class FIndirectPrimitiveType : SHADER_PERMUTATION_ENUM_CLASS("INDIRECT_PRIMITIVE_TYPE", EIndirectShadowingPrimitiveTypes);
+	using FPermutationDomain = TShaderPermutationDomain<FShapeShadow, FIndirectPrimitiveType>;
+
 	static bool ShouldCompilePermutation(const FGlobalShaderPermutationParameters& Parameters)
 	{
 		return IsFeatureLevelSupported(Parameters.Platform, ERHIFeatureLevel::SM5) && DoesPlatformSupportCapsuleShadows(Parameters.Platform);
@@ -192,39 +207,58 @@ public:
 
 	static void ModifyCompilationEnvironment(const FGlobalShaderPermutationParameters& Parameters, FShaderCompilerEnvironment& OutEnvironment)
 	{
+		OutEnvironment.CompilerFlags.Add(CFLAG_StandardOptimization);
+
+		FPermutationDomain PermutationVector(Parameters.PermutationId);
+
 		OutEnvironment.SetDefine(TEXT("THREADGROUP_SIZEX"), GShadowShapeTileSize);
 		OutEnvironment.SetDefine(TEXT("THREADGROUP_SIZEY"), GShadowShapeTileSize);
-		OutEnvironment.SetDefine(TEXT("POINT_LIGHT"), ShadowingType == ShapeShadow_PointLightTiledCulling);
-		uint32 LightSourceMode = 0;
 
-		if (ShadowingType == ShapeShadow_DirectionalLightTiledCulling || ShadowingType == ShapeShadow_PointLightTiledCulling)
+		ECapsuleShadowingType ShadowingType = PermutationVector.Get<FShapeShadow>();
+
+		OutEnvironment.SetDefine(TEXT("POINT_LIGHT"), ShadowingType == ECapsuleShadowingType::PointLightTiledCulling);
+
+		ELightSourceMode LightSourceMode = (ELightSourceMode)0;
+
+		if (ShadowingType == ECapsuleShadowingType::DirectionalLightTiledCulling || ShadowingType == ECapsuleShadowingType::PointLightTiledCulling)
 		{
-			LightSourceMode = 0;
+			LightSourceMode = ELightSourceMode::Punctual;
 		}
-		else if (ShadowingType == ShapeShadow_IndirectTiledCulling || ShadowingType == ShapeShadow_MovableSkylightTiledCulling)
+		else if (ShadowingType == ECapsuleShadowingType::IndirectTiledCulling || ShadowingType == ECapsuleShadowingType::MovableSkylightTiledCulling)
 		{
-			LightSourceMode = 1;
+			LightSourceMode = ELightSourceMode::FromCapsule;
 		}
-		else if (ShadowingType == ShapeShadow_MovableSkylightTiledCullingGatherFromReceiverBentNormal)
+		else if (ShadowingType == ECapsuleShadowingType::MovableSkylightTiledCullingGatherFromReceiverBentNormal)
 		{
-			LightSourceMode = 2;
+			LightSourceMode = ELightSourceMode::FromReceiver;
 		}
 		else
 		{
 			check(0);
 		}
 
-		OutEnvironment.SetDefine(TEXT("LIGHT_SOURCE_MODE"), LightSourceMode);
-		const bool bApplyToBentNormal = ShadowingType == ShapeShadow_MovableSkylightTiledCulling || ShadowingType == ShapeShadow_MovableSkylightTiledCullingGatherFromReceiverBentNormal;
+		OutEnvironment.SetDefine(TEXT("LIGHT_SOURCE_MODE"), (uint32)LightSourceMode);
+		const bool bApplyToBentNormal = ShadowingType == ECapsuleShadowingType::MovableSkylightTiledCulling || ShadowingType == ECapsuleShadowingType::MovableSkylightTiledCullingGatherFromReceiverBentNormal;
 		OutEnvironment.SetDefine(TEXT("APPLY_TO_BENT_NORMAL"), bApplyToBentNormal);
-		OutEnvironment.CompilerFlags.Add(CFLAG_StandardOptimization);
+
+		EIndirectShadowingPrimitiveTypes PrimitiveTypes = PermutationVector.Get<FIndirectPrimitiveType>();
+
+		if (PrimitiveTypes == EIndirectShadowingPrimitiveTypes::CapsuleShapes || PrimitiveTypes == EIndirectShadowingPrimitiveTypes::CapsuleShapesAndMeshDistanceFields)
+		{
+			OutEnvironment.SetDefine(TEXT("SUPPORT_CAPSULE_SHAPES"), 1);
+		}
+
+		if (PrimitiveTypes == EIndirectShadowingPrimitiveTypes::MeshDistanceFields || PrimitiveTypes == EIndirectShadowingPrimitiveTypes::CapsuleShapesAndMeshDistanceFields)
+		{
+			OutEnvironment.SetDefine(TEXT("SUPPORT_MESH_DISTANCE_FIELDS"), 1);
+		}
 	}
 
 	/** Default constructor. */
-	TCapsuleShadowingBaseCS() {}
+	FCapsuleShadowingCS() {}
 
 	/** Initialization constructor. */
-	TCapsuleShadowingBaseCS(const ShaderMetaType::CompiledShaderInitializerType& Initializer)
+	FCapsuleShadowingCS(const ShaderMetaType::CompiledShaderInitializerType& Initializer)
 		: FGlobalShader(Initializer)
 	{
 		ShadowFactors.Bind(Initializer.ParameterMap, TEXT("ShadowFactors"));
@@ -270,7 +304,8 @@ public:
 		int32 NumMeshDistanceFieldCastersValue,
 		FRHIShaderResourceView* MeshDistanceFieldCasterIndicesSRV,
 		FRHIShaderResourceView* LightDirectionDataSRV,
-		FRHITexture* ReceiverBentNormalTextureValue)
+		FRHITexture* ReceiverBentNormalTextureValue,
+		ECapsuleShadowingType ShadowingType)
 	{
 		FRHIComputeShader* ShaderRHI = RHICmdList.GetBoundComputeShader();
 
@@ -281,7 +316,7 @@ public:
 			RHICmdList.Transition(FRHITransitionInfo(TileIntersectionCountsBuffer->UAV, ERHIAccess::Unknown, ERHIAccess::UAVCompute));
 		}
 
-		if (ShadowingType == ShapeShadow_MovableSkylightTiledCulling)
+		if (ShadowingType == ECapsuleShadowingType::MovableSkylightTiledCulling)
 		{
 			check(!ShadowFactors.IsBound());
 			BentNormalTexture.SetTexture(RHICmdList, ShaderRHI, OutputTexture, OutputTextureUAV);
@@ -303,7 +338,7 @@ public:
 
 		SetShaderValue(RHICmdList, ShaderRHI, TileDimensions, TileDimensionsValue);
 
-		if (ShadowingType == ShapeShadow_MovableSkylightTiledCulling)
+		if (ShadowingType == ECapsuleShadowingType::MovableSkylightTiledCulling)
 		{
 			check(ReceiverBentNormalTextureValue);
 			SetTextureParameter(RHICmdList, ShaderRHI, ReceiverBentNormalTexture, ReceiverBentNormalTextureValue);
@@ -317,7 +352,7 @@ public:
 
 		if (LightSceneInfo)
 		{
-			check(ShadowingType == ShapeShadow_DirectionalLightTiledCulling || ShadowingType == ShapeShadow_PointLightTiledCulling);
+			check(ShadowingType == ECapsuleShadowingType::DirectionalLightTiledCulling || ShadowingType == ECapsuleShadowingType::PointLightTiledCulling);
 			
 			const FLightSceneProxy& LightProxy = *LightSceneInfo->Proxy;
 
@@ -338,7 +373,7 @@ public:
 		}
 		else
 		{
-			check(ShadowingType == ShapeShadow_IndirectTiledCulling || ShadowingType == ShapeShadow_MovableSkylightTiledCulling || ShadowingType == ShapeShadow_MovableSkylightTiledCullingGatherFromReceiverBentNormal);
+			check(ShadowingType == ECapsuleShadowingType::IndirectTiledCulling || ShadowingType == ECapsuleShadowingType::MovableSkylightTiledCulling || ShadowingType == ECapsuleShadowingType::MovableSkylightTiledCullingGatherFromReceiverBentNormal);
 			check(!LightDirection.IsBound() && !LightPositionAndInvRadius.IsBound());
 		}
 
@@ -420,55 +455,7 @@ private:
 	LAYOUT_FIELD((TDistanceFieldCulledObjectBufferParameters<DFPT_SignedDistanceField>), DistanceFieldCulledObjectParameters);
 };
 
-template<ECapsuleShadowingType ShadowingType, EIndirectShadowingPrimitiveTypes PrimitiveTypes>
-class TCapsuleShadowingCS : public TCapsuleShadowingBaseCS<ShadowingType>
-{
-	DECLARE_SHADER_TYPE(TCapsuleShadowingCS, Global);
-
-	TCapsuleShadowingCS() {}
-
-	/** Initialization constructor. */
-	TCapsuleShadowingCS(const ShaderMetaType::CompiledShaderInitializerType& Initializer)
-		: TCapsuleShadowingBaseCS<ShadowingType>(Initializer)
-	{}
-
-	static void ModifyCompilationEnvironment(const FGlobalShaderPermutationParameters& Parameters, FShaderCompilerEnvironment& OutEnvironment)
-	{
-		TCapsuleShadowingBaseCS<ShadowingType>::ModifyCompilationEnvironment(Parameters, OutEnvironment);
-
-		if (PrimitiveTypes & IPT_CapsuleShapes)
-		{
-			OutEnvironment.SetDefine(TEXT("SUPPORT_CAPSULE_SHAPES"), 1);
-		}
-		
-		if (PrimitiveTypes & IPT_MeshDistanceFields)
-		{
-			OutEnvironment.SetDefine(TEXT("SUPPORT_MESH_DISTANCE_FIELDS"), 1);
-		}
-	}
-};
-
-#define IMPLEMENT_CAPSULE_SHADOW_TYPE(ShadowType,PrimitiveType) \
-	typedef TCapsuleShadowingCS<ShadowType,PrimitiveType> TCapsuleShadowingCS##ShadowType##PrimitiveType; \
-	IMPLEMENT_SHADER_TYPE(template<>,TCapsuleShadowingCS##ShadowType##PrimitiveType,TEXT("/Engine/Private/CapsuleShadowShaders.usf"),TEXT("CapsuleShadowingCS"),SF_Compute);
-
-IMPLEMENT_CAPSULE_SHADOW_TYPE(ShapeShadow_DirectionalLightTiledCulling, IPT_CapsuleShapes);
-IMPLEMENT_CAPSULE_SHADOW_TYPE(ShapeShadow_PointLightTiledCulling, IPT_CapsuleShapes);
-IMPLEMENT_CAPSULE_SHADOW_TYPE(ShapeShadow_IndirectTiledCulling, IPT_CapsuleShapes);
-IMPLEMENT_CAPSULE_SHADOW_TYPE(ShapeShadow_MovableSkylightTiledCulling, IPT_CapsuleShapes);
-IMPLEMENT_CAPSULE_SHADOW_TYPE(ShapeShadow_MovableSkylightTiledCullingGatherFromReceiverBentNormal, IPT_CapsuleShapes);
-
-IMPLEMENT_CAPSULE_SHADOW_TYPE(ShapeShadow_DirectionalLightTiledCulling, IPT_MeshDistanceFields);
-IMPLEMENT_CAPSULE_SHADOW_TYPE(ShapeShadow_PointLightTiledCulling, IPT_MeshDistanceFields);
-IMPLEMENT_CAPSULE_SHADOW_TYPE(ShapeShadow_IndirectTiledCulling, IPT_MeshDistanceFields);
-IMPLEMENT_CAPSULE_SHADOW_TYPE(ShapeShadow_MovableSkylightTiledCulling, IPT_MeshDistanceFields);
-IMPLEMENT_CAPSULE_SHADOW_TYPE(ShapeShadow_MovableSkylightTiledCullingGatherFromReceiverBentNormal, IPT_MeshDistanceFields);
-
-IMPLEMENT_CAPSULE_SHADOW_TYPE(ShapeShadow_DirectionalLightTiledCulling, IPT_CapsuleShapesAndMeshDistanceFields);
-IMPLEMENT_CAPSULE_SHADOW_TYPE(ShapeShadow_PointLightTiledCulling, IPT_CapsuleShapesAndMeshDistanceFields);
-IMPLEMENT_CAPSULE_SHADOW_TYPE(ShapeShadow_IndirectTiledCulling, IPT_CapsuleShapesAndMeshDistanceFields);
-IMPLEMENT_CAPSULE_SHADOW_TYPE(ShapeShadow_MovableSkylightTiledCulling, IPT_CapsuleShapesAndMeshDistanceFields);
-IMPLEMENT_CAPSULE_SHADOW_TYPE(ShapeShadow_MovableSkylightTiledCullingGatherFromReceiverBentNormal, IPT_CapsuleShapesAndMeshDistanceFields);
+IMPLEMENT_GLOBAL_SHADER(FCapsuleShadowingCS, "/Engine/Private/CapsuleShadowShaders.usf", "CapsuleShadowingCS", SF_Compute);
 
 class FCapsuleShadowingUpsampleVS : public FGlobalShader
 {
@@ -519,15 +506,17 @@ private:
 	LAYOUT_FIELD(FShaderResourceParameter, TileIntersectionCounts);
 };
 
-IMPLEMENT_SHADER_TYPE(,FCapsuleShadowingUpsampleVS,TEXT("/Engine/Private/CapsuleShadowShaders.usf"),TEXT("CapsuleShadowingUpsampleVS"),SF_Vertex);
-
+IMPLEMENT_GLOBAL_SHADER(FCapsuleShadowingUpsampleVS, "/Engine/Private/CapsuleShadowShaders.usf", "CapsuleShadowingUpsampleVS", SF_Vertex);
 
-template<bool bUpsampleRequired, bool bApplyToSSAO>
-class TCapsuleShadowingUpsamplePS : public FGlobalShader
+class FCapsuleShadowingUpsamplePS : public FGlobalShader
 {
-	DECLARE_SHADER_TYPE(TCapsuleShadowingUpsamplePS, Global);
+	DECLARE_GLOBAL_SHADER(FCapsuleShadowingUpsamplePS);
 public:
 
+	class FUpsampleRequired : SHADER_PERMUTATION_BOOL("UPSAMPLE_REQUIRED");
+	class FApplySSAO : SHADER_PERMUTATION_BOOL("APPLY_TO_SSAO");
+	using FPermutationDomain = TShaderPermutationDomain<FUpsampleRequired, FApplySSAO>;
+
 	static bool ShouldCompilePermutation(const FGlobalShaderPermutationParameters& Parameters)
 	{
 		return IsFeatureLevelSupported(Parameters.Platform, ERHIFeatureLevel::SM5) && DoesPlatformSupportCapsuleShadows(Parameters.Platform);
@@ -536,15 +525,13 @@ public:
 	static void ModifyCompilationEnvironment(const FGlobalShaderPermutationParameters& Parameters, FShaderCompilerEnvironment& OutEnvironment)
 	{
 		OutEnvironment.SetDefine(TEXT("DOWNSAMPLE_FACTOR"), 2);
-		OutEnvironment.SetDefine(TEXT("UPSAMPLE_REQUIRED"), bUpsampleRequired);
-		OutEnvironment.SetDefine(TEXT("APPLY_TO_SSAO"), bApplyToSSAO);
 	}
 
 	/** Default constructor. */
-	TCapsuleShadowingUpsamplePS() {}
+	FCapsuleShadowingUpsamplePS() {}
 
 	/** Initialization constructor. */
-	TCapsuleShadowingUpsamplePS(const ShaderMetaType::CompiledShaderInitializerType& Initializer)
+	FCapsuleShadowingUpsamplePS(const ShaderMetaType::CompiledShaderInitializerType& Initializer)
 		: FGlobalShader(Initializer)
 	{
 		ShadowFactorsTexture.Bind(Initializer.ParameterMap,TEXT("ShadowFactorsTexture"));
@@ -572,14 +559,7 @@ private:
 	LAYOUT_FIELD(FShaderParameter, OutputtingToLightAttenuation);
 };
 
-#define IMPLEMENT_CAPSULE_APPLY_SHADER_TYPE(bUpsampleRequired,bApplyToSSAO) \
-	typedef TCapsuleShadowingUpsamplePS<bUpsampleRequired,bApplyToSSAO> TCapsuleShadowingUpsamplePS##bUpsampleRequired##bApplyToSSAO; \
-	IMPLEMENT_SHADER_TYPE(template<>,TCapsuleShadowingUpsamplePS##bUpsampleRequired##bApplyToSSAO,TEXT("/Engine/Private/CapsuleShadowShaders.usf"),TEXT("CapsuleShadowingUpsamplePS"),SF_Pixel)
-
-IMPLEMENT_CAPSULE_APPLY_SHADER_TYPE(true, true);
-IMPLEMENT_CAPSULE_APPLY_SHADER_TYPE(true, false);
-IMPLEMENT_CAPSULE_APPLY_SHADER_TYPE(false, true);
-IMPLEMENT_CAPSULE_APPLY_SHADER_TYPE(false, false);
+IMPLEMENT_GLOBAL_SHADER(FCapsuleShadowingUpsamplePS, "/Engine/Private/CapsuleShadowShaders.usf", "CapsuleShadowingUpsamplePS", SF_Pixel);
 
 void AllocateCapsuleTileIntersectionCountsBuffer(FIntPoint GroupSize, FSceneViewState* ViewState)
 {
@@ -725,44 +705,21 @@ bool FDeferredShadingSceneRenderer::RenderCapsuleDirectShadows(
 				PassParameters->RayTracedShadowsUAV = GraphBuilder.CreateUAV(RayTracedShadowsRT);
 				PassParameters->SceneTextures = SceneTexturesUniformBuffer;
 
+				ECapsuleShadowingType ShadowingType = bDirectionalLight ? ECapsuleShadowingType::DirectionalLightTiledCulling : ECapsuleShadowingType::PointLightTiledCulling;
+
+				FCapsuleShadowingCS::FPermutationDomain PermutationVector;
+				PermutationVector.Set<FCapsuleShadowingCS::FShapeShadow>(ShadowingType);
+				PermutationVector.Set<FCapsuleShadowingCS::FIndirectPrimitiveType>(EIndirectShadowingPrimitiveTypes::CapsuleShapes);
+				auto ComputeShader = View.ShaderMap->GetShader<FCapsuleShadowingCS>(PermutationVector);
+
 				GraphBuilder.AddPass(
 					RDG_EVENT_NAME("TiledCapsuleShadowing"),
 					PassParameters,
 					ERDGPassFlags::Compute,
-					[this, &View, &LightSceneInfo, RayTracedShadows = PassParameters->RayTracedShadows.GetTexture(), RayTracedShadowsUAV = PassParameters->RayTracedShadowsUAV, GroupSize, ScissorRect, bDirectionalLight, NumCapsuleShapeData](FRHIComputeCommandList& RHICmdList)
-				{
-					if (bDirectionalLight)
-					{
-						TShaderMapRef<TCapsuleShadowingCS<ShapeShadow_DirectionalLightTiledCulling, IPT_CapsuleShapes> > ComputeShader(View.ShaderMap);
-						RHICmdList.SetComputeShader(ComputeShader.GetComputeShader());
-
-						ComputeShader->SetParameters(
-							RHICmdList,
-							Scene,
-							View,
-							&LightSceneInfo,
-							RayTracedShadows->GetRHI(),
-							RayTracedShadowsUAV->GetRHI(),
-							GroupSize,
-							&View.ViewState->CapsuleTileIntersectionCountsBuffer,
-							FVector2D(GroupSize.X, GroupSize.Y),
-							GCapsuleMaxDirectOcclusionDistance,
-							ScissorRect,
-							GetCapsuleShadowDownsampleFactor(),
-							NumCapsuleShapeData,
-							LightSceneInfo.ShadowCapsuleShapesSRV.GetReference(),
-							0,
-							NULL,
-							NULL,
-							NULL);
-
-						DispatchComputeShader(RHICmdList, ComputeShader.GetShader(), GroupSize.X, GroupSize.Y, 1);
-						ComputeShader->UnsetParameters(RHICmdList, &View.ViewState->CapsuleTileIntersectionCountsBuffer);
-					}
-					else
+					[ComputeShader, this, &View, &LightSceneInfo, RayTracedShadows = PassParameters->RayTracedShadows.GetTexture(), RayTracedShadowsUAV = PassParameters->RayTracedShadowsUAV, GroupSize, ScissorRect, NumCapsuleShapeData, ShadowingType](FRHIComputeCommandList& RHICmdList)
 					{
-						TShaderMapRef<TCapsuleShadowingCS<ShapeShadow_PointLightTiledCulling, IPT_CapsuleShapes> > ComputeShader(View.ShaderMap);
-						RHICmdList.SetComputeShader(ComputeShader.GetComputeShader());
+						FRHIComputeShader* ShaderRHI = ComputeShader.GetComputeShader();
+						RHICmdList.SetComputeShader(ShaderRHI);
 
 						ComputeShader->SetParameters(
 							RHICmdList,
@@ -780,14 +737,14 @@ bool FDeferredShadingSceneRenderer::RenderCapsuleDirectShadows(
 							NumCapsuleShapeData,
 							LightSceneInfo.ShadowCapsuleShapesSRV.GetReference(),
 							0,
-							NULL,
-							NULL,
-							NULL);
+							nullptr,
+							nullptr,
+							nullptr,
+							ShadowingType);
 
 						DispatchComputeShader(RHICmdList, ComputeShader.GetShader(), GroupSize.X, GroupSize.Y, 1);
 						ComputeShader->UnsetParameters(RHICmdList, &View.ViewState->CapsuleTileIntersectionCountsBuffer);
-					}
-				});
+					});
 			}
 
 			{
@@ -816,35 +773,22 @@ bool FDeferredShadingSceneRenderer::RenderCapsuleDirectShadows(
 						bProjectingForForwardShading,
 						false);
 
-					TShaderMapRef<FCapsuleShadowingUpsampleVS> VertexShader(View.ShaderMap);
-					GraphicsPSOInit.PrimitiveType = PT_TriangleList;
-
-					if (GCapsuleShadowsFullResolution)
-					{
-						TShaderMapRef<TCapsuleShadowingUpsamplePS<false, false> > PixelShader(View.ShaderMap);
-
-						GraphicsPSOInit.BoundShaderState.VertexDeclarationRHI = GTileVertexDeclaration.VertexDeclarationRHI;
-						GraphicsPSOInit.BoundShaderState.VertexShaderRHI = VertexShader.GetVertexShader();
-						GraphicsPSOInit.BoundShaderState.PixelShaderRHI = PixelShader.GetPixelShader();
-
-						SetGraphicsPipelineState(RHICmdList, GraphicsPSOInit);
+					auto VertexShader = View.ShaderMap->GetShader<FCapsuleShadowingUpsampleVS>();
 
-						VertexShader->SetParameters(RHICmdList, View, GroupSize, ScissorRect, View.ViewState->CapsuleTileIntersectionCountsBuffer);
-						PixelShader->SetParameters(RHICmdList, View, ScissorRect, RayTracedShadowsRT->GetRHI(), true);
-					}
-					else
-					{
-						TShaderMapRef<TCapsuleShadowingUpsamplePS<true, false> > PixelShader(View.ShaderMap);
+					FCapsuleShadowingUpsamplePS::FPermutationDomain PermutationVector;
+					PermutationVector.Set<FCapsuleShadowingUpsamplePS::FUpsampleRequired>(!GCapsuleShadowsFullResolution);
+					PermutationVector.Set<FCapsuleShadowingUpsamplePS::FApplySSAO>(false);
+					auto PixelShader = View.ShaderMap->GetShader<FCapsuleShadowingUpsamplePS>(PermutationVector);
 
-						GraphicsPSOInit.BoundShaderState.VertexDeclarationRHI = GTileVertexDeclaration.VertexDeclarationRHI;
-						GraphicsPSOInit.BoundShaderState.VertexShaderRHI = VertexShader.GetVertexShader();
-						GraphicsPSOInit.BoundShaderState.PixelShaderRHI = PixelShader.GetPixelShader();
+					GraphicsPSOInit.PrimitiveType = PT_TriangleList;
+					GraphicsPSOInit.BoundShaderState.VertexDeclarationRHI = GTileVertexDeclaration.VertexDeclarationRHI;
+					GraphicsPSOInit.BoundShaderState.VertexShaderRHI = VertexShader.GetVertexShader();
+					GraphicsPSOInit.BoundShaderState.PixelShaderRHI = PixelShader.GetPixelShader();
 
-						SetGraphicsPipelineState(RHICmdList, GraphicsPSOInit);
+					SetGraphicsPipelineState(RHICmdList, GraphicsPSOInit);
 
-						VertexShader->SetParameters(RHICmdList, View, GroupSize, ScissorRect, View.ViewState->CapsuleTileIntersectionCountsBuffer);
-						PixelShader->SetParameters(RHICmdList, View, ScissorRect, RayTracedShadowsRT->GetRHI(), true);
-					}
+					VertexShader->SetParameters(RHICmdList, View, GroupSize, ScissorRect, View.ViewState->CapsuleTileIntersectionCountsBuffer);
+					PixelShader->SetParameters(RHICmdList, View, ScissorRect, RayTracedShadowsRT->GetRHI(), true);
 
 					RHICmdList.SetStreamSource(0, GTileTexCoordVertexBuffer.VertexBufferRHI, 0);
 					RHICmdList.DrawIndexedPrimitive(GTileIndexBuffer.IndexBufferRHI, 
@@ -1256,56 +1200,60 @@ void FDeferredShadingSceneRenderer::RenderIndirectCapsuleShadows(FRDGBuilder& Gr
 				PassParameters->RayTracedShadowsUAV = GraphBuilder.CreateUAV(RayTracedShadowsRT);
 				PassParameters->SceneTextures = SceneTextures.UniformBuffer;
 
+				EIndirectShadowingPrimitiveTypes PrimitiveTypes;
+
+				if (NumCapsuleShapes > 0 && NumMeshDistanceFieldCasters > 0)
+				{
+					PrimitiveTypes = EIndirectShadowingPrimitiveTypes::CapsuleShapesAndMeshDistanceFields;
+				}
+				else if (NumCapsuleShapes > 0)
+				{
+					PrimitiveTypes = EIndirectShadowingPrimitiveTypes::CapsuleShapes;
+				}
+				else
+				{
+					check(NumMeshDistanceFieldCasters > 0);
+					PrimitiveTypes = EIndirectShadowingPrimitiveTypes::MeshDistanceFields;
+				}
+
+				FCapsuleShadowingCS::FPermutationDomain PermutationVector;
+				PermutationVector.Set<FCapsuleShadowingCS::FShapeShadow>(ECapsuleShadowingType::IndirectTiledCulling);
+				PermutationVector.Set<FCapsuleShadowingCS::FIndirectPrimitiveType>(PrimitiveTypes);
+				auto ComputeShader = View.ShaderMap->GetShader<FCapsuleShadowingCS>(PermutationVector);
+
 				GraphBuilder.AddPass(
 					RDG_EVENT_NAME("TiledCapsuleShadowing %u capsules among %u meshes", NumCapsuleShapes, NumMeshesWithCapsules),
 					PassParameters,
 					ERDGPassFlags::Compute,
-					[this, &View, RayTracedShadows = PassParameters->RayTracedShadows.GetTexture(), RayTracedShadowsUAV = PassParameters->RayTracedShadowsUAV, IndirectShadowLightDirectionSRV, GroupSize, ScissorRect, NumCapsuleShapes, NumMeshDistanceFieldCasters](FRHIComputeCommandList& RHICmdList)
-				{
-					TShaderRef<TCapsuleShadowingBaseCS<ShapeShadow_IndirectTiledCulling>> ComputeShaderBase;
-
-					if (NumCapsuleShapes > 0 && NumMeshDistanceFieldCasters > 0)
+					[ComputeShader, this, &View, RayTracedShadows = PassParameters->RayTracedShadows.GetTexture(), RayTracedShadowsUAV = PassParameters->RayTracedShadowsUAV, IndirectShadowLightDirectionSRV, GroupSize, ScissorRect, NumCapsuleShapes, NumMeshDistanceFieldCasters](FRHIComputeCommandList& RHICmdList)
 					{
-						TShaderMapRef<TCapsuleShadowingCS<ShapeShadow_IndirectTiledCulling, IPT_CapsuleShapesAndMeshDistanceFields> > ComputeShader(View.ShaderMap);
-						ComputeShaderBase = ComputeShader;
-					}
-					else if (NumCapsuleShapes > 0)
-					{
-						TShaderMapRef<TCapsuleShadowingCS<ShapeShadow_IndirectTiledCulling, IPT_CapsuleShapes> > ComputeShader(View.ShaderMap);
-						ComputeShaderBase = ComputeShader;
-					}
-					else
-					{
-						check(NumMeshDistanceFieldCasters > 0);
-						TShaderMapRef<TCapsuleShadowingCS<ShapeShadow_IndirectTiledCulling, IPT_MeshDistanceFields> > ComputeShader(View.ShaderMap);
-						ComputeShaderBase = ComputeShader;
-					}
+						FRHIComputeShader* ShaderRHI = ComputeShader.GetComputeShader();
+						RHICmdList.SetComputeShader(ShaderRHI);
 
-					RHICmdList.SetComputeShader(ComputeShaderBase.GetComputeShader());
-
-					ComputeShaderBase->SetParameters(
-						RHICmdList,
-						Scene,
-						View,
-						NULL,
-						RayTracedShadows->GetRHI(),
-						RayTracedShadowsUAV->GetRHI(),
-						GroupSize,
-						&View.ViewState->CapsuleTileIntersectionCountsBuffer,
-						FVector2D(GroupSize.X, GroupSize.Y),
-						GCapsuleMaxIndirectOcclusionDistance,
-						ScissorRect,
-						GetCapsuleShadowDownsampleFactor(),
-						NumCapsuleShapes,
-						View.ViewState->IndirectShadowCapsuleShapesSRV ? View.ViewState->IndirectShadowCapsuleShapesSRV.GetReference() : NULL,
-						NumMeshDistanceFieldCasters,
-						View.ViewState->IndirectShadowMeshDistanceFieldCasterIndicesSRV ? View.ViewState->IndirectShadowMeshDistanceFieldCasterIndicesSRV.GetReference() : NULL,
-						IndirectShadowLightDirectionSRV,
-						NULL);
-
-					DispatchComputeShader(RHICmdList, ComputeShaderBase.GetShader(), GroupSize.X, GroupSize.Y, 1);
-					ComputeShaderBase->UnsetParameters(RHICmdList, &View.ViewState->CapsuleTileIntersectionCountsBuffer);
-				});
+						ComputeShader->SetParameters(
+							RHICmdList,
+							Scene,
+							View,
+							nullptr,
+							RayTracedShadows->GetRHI(),
+							RayTracedShadowsUAV->GetRHI(),
+							GroupSize,
+							&View.ViewState->CapsuleTileIntersectionCountsBuffer,
+							FVector2D(GroupSize.X, GroupSize.Y),
+							GCapsuleMaxIndirectOcclusionDistance,
+							ScissorRect,
+							GetCapsuleShadowDownsampleFactor(),
+							NumCapsuleShapes,
+							View.ViewState->IndirectShadowCapsuleShapesSRV ? View.ViewState->IndirectShadowCapsuleShapesSRV.GetReference() : nullptr,
+							NumMeshDistanceFieldCasters,
+							View.ViewState->IndirectShadowMeshDistanceFieldCasterIndicesSRV ? View.ViewState->IndirectShadowMeshDistanceFieldCasterIndicesSRV.GetReference() : nullptr,
+							IndirectShadowLightDirectionSRV,
+							nullptr,
+							ECapsuleShadowingType::IndirectTiledCulling);
+
+						DispatchComputeShader(RHICmdList, ComputeShader, GroupSize.X, GroupSize.Y, 1);
+						ComputeShader->UnsetParameters(RHICmdList, &View.ViewState->CapsuleTileIntersectionCountsBuffer);
+					});
 			}
 
 			{
@@ -1345,49 +1293,21 @@ void FDeferredShadingSceneRenderer::RenderIndirectCapsuleShadows(FRDGBuilder& Gr
 						GraphicsPSOInit.BlendState = TStaticBlendState<CW_RGB, BO_Add, BF_DestColor, BF_Zero>::GetRHI();
 					}
 
-					TShaderMapRef<FCapsuleShadowingUpsampleVS> VertexShader(View.ShaderMap);
+					auto VertexShader = View.ShaderMap->GetShader<FCapsuleShadowingUpsampleVS>();
+
+					FCapsuleShadowingUpsamplePS::FPermutationDomain PermutationVector;
+					PermutationVector.Set<FCapsuleShadowingUpsamplePS::FUpsampleRequired>(!GCapsuleShadowsFullResolution);
+					PermutationVector.Set<FCapsuleShadowingUpsamplePS::FApplySSAO>(RenderTargetCount > 1);
+					auto PixelShader = View.ShaderMap->GetShader<FCapsuleShadowingUpsamplePS>(PermutationVector);
+
 					GraphicsPSOInit.BoundShaderState.VertexDeclarationRHI = GTileVertexDeclaration.VertexDeclarationRHI;
 					GraphicsPSOInit.BoundShaderState.VertexShaderRHI = VertexShader.GetVertexShader();
 					GraphicsPSOInit.PrimitiveType = PT_TriangleList;
+					GraphicsPSOInit.BoundShaderState.PixelShaderRHI = PixelShader.GetPixelShader();
 
-					if (RenderTargetCount > 1)
-					{
-						if (GCapsuleShadowsFullResolution)
-						{
-							TShaderMapRef<TCapsuleShadowingUpsamplePS<false, true> > PixelShader(View.ShaderMap);
-							GraphicsPSOInit.BoundShaderState.PixelShaderRHI = PixelShader.GetPixelShader();
-							SetGraphicsPipelineState(RHICmdList, GraphicsPSOInit);
-							VertexShader->SetParameters(RHICmdList, View, GroupSize, ScissorRect, View.ViewState->CapsuleTileIntersectionCountsBuffer);
-							PixelShader->SetParameters(RHICmdList, View, ScissorRect, RayTracedShadowsRT->GetRHI(), false);
-						}
-						else
-						{
-							TShaderMapRef<TCapsuleShadowingUpsamplePS<true, true> > PixelShader(View.ShaderMap);
-							GraphicsPSOInit.BoundShaderState.PixelShaderRHI = PixelShader.GetPixelShader();
-							SetGraphicsPipelineState(RHICmdList, GraphicsPSOInit);
-							VertexShader->SetParameters(RHICmdList, View, GroupSize, ScissorRect, View.ViewState->CapsuleTileIntersectionCountsBuffer);
-							PixelShader->SetParameters(RHICmdList, View, ScissorRect, RayTracedShadowsRT->GetRHI(), false);
-						}
-					}
-					else
-					{
-						if (GCapsuleShadowsFullResolution)
-						{
-							TShaderMapRef<TCapsuleShadowingUpsamplePS<false, false> > PixelShader(View.ShaderMap);
-							GraphicsPSOInit.BoundShaderState.PixelShaderRHI = PixelShader.GetPixelShader();
-							SetGraphicsPipelineState(RHICmdList, GraphicsPSOInit);
-							VertexShader->SetParameters(RHICmdList, View, GroupSize, ScissorRect, View.ViewState->CapsuleTileIntersectionCountsBuffer);
-							PixelShader->SetParameters(RHICmdList, View, ScissorRect, RayTracedShadowsRT->GetRHI(), false);
-						}
-						else
-						{
-							TShaderMapRef<TCapsuleShadowingUpsamplePS<true, false> > PixelShader(View.ShaderMap);
-							GraphicsPSOInit.BoundShaderState.PixelShaderRHI = PixelShader.GetPixelShader();
-							SetGraphicsPipelineState(RHICmdList, GraphicsPSOInit);
-							VertexShader->SetParameters(RHICmdList, View, GroupSize, ScissorRect, View.ViewState->CapsuleTileIntersectionCountsBuffer);
-							PixelShader->SetParameters(RHICmdList, View, ScissorRect, RayTracedShadowsRT->GetRHI(), false);
-						}
-					}
+					SetGraphicsPipelineState(RHICmdList, GraphicsPSOInit);
+					VertexShader->SetParameters(RHICmdList, View, GroupSize, ScissorRect, View.ViewState->CapsuleTileIntersectionCountsBuffer);
+					PixelShader->SetParameters(RHICmdList, View, ScissorRect, RayTracedShadowsRT->GetRHI(), false);
 
 					RHICmdList.SetStreamSource(0, GTileTexCoordVertexBuffer.VertexBufferRHI, 0);
 					RHICmdList.DrawIndexedPrimitive(GTileIndexBuffer.IndexBufferRHI,
@@ -1503,34 +1423,38 @@ void FDeferredShadingSceneRenderer::RenderCapsuleShadowsForMovableSkylight(
 							PassParameters->BentNormalOutputUAV = GraphBuilder.CreateUAV(NewBentNormal);
 							PassParameters->SceneTextures = SceneTexturesUniformBuffer;
 
+							EIndirectShadowingPrimitiveTypes PrimitiveTypes;
+
+							if (NumCapsuleShapes > 0 && NumMeshDistanceFieldCasters > 0)
+							{
+								PrimitiveTypes = EIndirectShadowingPrimitiveTypes::CapsuleShapesAndMeshDistanceFields;
+							}
+							else if (NumCapsuleShapes > 0)
+							{
+								PrimitiveTypes = EIndirectShadowingPrimitiveTypes::CapsuleShapes;
+							}
+							else
+							{
+								check(NumMeshDistanceFieldCasters > 0);
+								PrimitiveTypes = EIndirectShadowingPrimitiveTypes::MeshDistanceFields;
+							}
+
+							FCapsuleShadowingCS::FPermutationDomain PermutationVector;
+							PermutationVector.Set<FCapsuleShadowingCS::FShapeShadow>(ECapsuleShadowingType::MovableSkylightTiledCulling);
+							PermutationVector.Set<FCapsuleShadowingCS::FIndirectPrimitiveType>(PrimitiveTypes);
+							auto ComputeShader = View.ShaderMap->GetShader<FCapsuleShadowingCS>(PermutationVector);
+
 							GraphBuilder.AddPass(
 								RDG_EVENT_NAME("TiledCapsuleShadowing % u capsules among % u meshes", NumCapsuleShapes, NumMeshesWithCapsules),
 								PassParameters,
 								ERDGPassFlags::Compute,
-								[this, &View, NewBentNormal, PassParameters, GroupSizeX, GroupSizeY, NumCapsuleShapes, NumMeshDistanceFieldCasters, ScissorRect, IndirectShadowLightDirectionSRV]
-								(FRHICommandList& RHICmdList)
-							{
+								[PassParameters, ComputeShader, this, &View, GroupSizeX, GroupSizeY, NumCapsuleShapes, NumMeshDistanceFieldCasters, ScissorRect, IndirectShadowLightDirectionSRV]
+							(FRHICommandList& RHICmdList)
 								{
-									TShaderRef<TCapsuleShadowingBaseCS<ShapeShadow_MovableSkylightTiledCulling>> ComputeShaderBase;
-									if (NumCapsuleShapes > 0 && NumMeshDistanceFieldCasters > 0)
-									{
-										TShaderMapRef<TCapsuleShadowingCS<ShapeShadow_MovableSkylightTiledCulling, IPT_CapsuleShapesAndMeshDistanceFields> > ComputeShader(View.ShaderMap);
-										ComputeShaderBase = ComputeShader;
-									}
-									else if (NumCapsuleShapes > 0)
-									{
-										TShaderMapRef<TCapsuleShadowingCS<ShapeShadow_MovableSkylightTiledCulling, IPT_CapsuleShapes> > ComputeShader(View.ShaderMap);
-										ComputeShaderBase = ComputeShader;
-									}
-									else
-									{
-										TShaderMapRef<TCapsuleShadowingCS<ShapeShadow_MovableSkylightTiledCulling, IPT_MeshDistanceFields> > ComputeShader(View.ShaderMap);
-										ComputeShaderBase = ComputeShader;
-									}
-
-									RHICmdList.SetComputeShader(ComputeShaderBase.GetComputeShader());
-
-									ComputeShaderBase->SetParameters(
+									FRHIComputeShader* ShaderRHI = ComputeShader.GetComputeShader();
+									RHICmdList.SetComputeShader(ShaderRHI);
+
+									ComputeShader->SetParameters(
 										RHICmdList,
 										Scene,
 										View,
@@ -1548,12 +1472,12 @@ void FDeferredShadingSceneRenderer::RenderCapsuleShadowsForMovableSkylight(
 										NumMeshDistanceFieldCasters,
 										View.ViewState->IndirectShadowMeshDistanceFieldCasterIndicesSRV ? View.ViewState->IndirectShadowMeshDistanceFieldCasterIndicesSRV.GetReference() : NULL,
 										IndirectShadowLightDirectionSRV,
-										PassParameters->BentNormalInput->GetRHI());
+										PassParameters->BentNormalInput->GetRHI(),
+										ECapsuleShadowingType::MovableSkylightTiledCulling);
 
-									DispatchComputeShader(RHICmdList, ComputeShaderBase.GetShader(), GroupSizeX, GroupSizeY, 1);
-									ComputeShaderBase->UnsetParameters(RHICmdList, nullptr);
-								}
-							});
+									DispatchComputeShader(RHICmdList, ComputeShader.GetShader(), GroupSizeX, GroupSizeY, 1);
+									ComputeShader->UnsetParameters(RHICmdList, nullptr);
+								});
 						}
 
 						// Replace the pipeline output with our output that has capsule shadows applied
