Temp fix for RHI validation for GPUScene.InstanceSceneData

#rnx

diff --git a/Engine/Source/Runtime/Renderer/Private/GPUScene.cpp b/Engine/Source/Runtime/Renderer/Private/GPUScene.cpp
--- a/Engine/Source/Runtime/Renderer/Private/GPUScene.cpp
+++ b/Engine/Source/Runtime/Renderer/Private/GPUScene.cpp
@@ -963,7 +963,7 @@ void FGPUScene::UploadGeneral(FRHICommandListImmediate& RHICmdList, FScene *Scen
 		const bool bNaniteEnabled = DoesPlatformSupportNanite(GMaxRHIShaderPlatform);
 
 		const uint32 InstanceSceneDataNumArrays = FInstanceSceneShaderData::DataStrideInFloat4s;
-		FUAVTransitionStateScopeHelper InstanceSceneDataTransitionHelper(RHICmdList, BufferState.InstanceSceneDataBuffer.UAV, BufferState.bResizedInstanceSceneData ? ERHIAccess::Unknown : ERHIAccess::SRVMask, ERHIAccess::SRVMask);
+		FUAVTransitionStateScopeHelper InstanceSceneDataTransitionHelper(RHICmdList, BufferState.InstanceSceneDataBuffer.UAV, ERHIAccess::Unknown, ERHIAccess::SRVMask);
 
 		const uint32 LightMapDataBufferSize = BufferState.LightMapDataBufferSize;
 
@@ -1830,8 +1830,6 @@ void FGPUScene::AddUpdatePrimitiveIdsPass(FRDGBuilder& GraphBuilder, FInstancePr
 			IdOnlyUpdateItems.GetCsGroupCount()
 		);
 
-		//GraphBuilder.AddPass()
-
 		EndReadWriteAccess(GraphBuilder, ERHIAccess::UAVCompute);
 	}
 }
