Fix crash accessing View.ViewState in DFAO.
Fallback to temporal sample index 0  when there is no view state.

diff --git a/Engine/Source/Runtime/Engine/Private/SceneView.cpp b/Engine/Source/Runtime/Engine/Private/SceneView.cpp
--- a/Engine/Source/Runtime/Engine/Private/SceneView.cpp
+++ b/Engine/Source/Runtime/Engine/Private/SceneView.cpp
@@ -951,6 +951,15 @@ float FSceneView::GetTemporalLODTransition() const
 	return 0.0f;
 }
 
+uint32 FSceneView::GetDistanceFieldTemporalSampleIndex() const
+{
+	if (State)
+	{
+		return State->GetDistanceFieldTemporalSampleIndex();
+	}
+	return 0;
+}
+
 uint32 FSceneView::GetViewKey() const
 {
 	if (State)
diff --git a/Engine/Source/Runtime/Engine/Public/SceneManagement.h b/Engine/Source/Runtime/Engine/Public/SceneManagement.h
--- a/Engine/Source/Runtime/Engine/Public/SceneManagement.h
+++ b/Engine/Source/Runtime/Engine/Public/SceneManagement.h
@@ -212,6 +212,11 @@ public:
 	//
 	virtual uint32 GetCurrentTemporalAASampleIndex() const = 0;
 
+	/**
+	 * returns the distance field temporal sample index
+	 */
+	virtual uint32 GetDistanceFieldTemporalSampleIndex() const = 0;
+
 	/** Tells if the eye adaptation texture / buffer exists without attempting to allocate it. */
 	bool HasValidEyeAdaptationTexture() const { return bValidEyeAdaptationTexture; }
 	bool HasValidEyeAdaptationBuffer() const { return bValidEyeAdaptationBuffer; }
diff --git a/Engine/Source/Runtime/Engine/Public/SceneView.h b/Engine/Source/Runtime/Engine/Public/SceneView.h
--- a/Engine/Source/Runtime/Engine/Public/SceneView.h
+++ b/Engine/Source/Runtime/Engine/Public/SceneView.h
@@ -1284,6 +1284,11 @@ public:
 	 */
 	float GetTemporalLODTransition() const;
 
+	/**
+	 * returns the distance field temporal sample index or 0 if there is no view state
+	 */
+	uint32 GetDistanceFieldTemporalSampleIndex() const;
+
 	/** 
 	 * returns a unique key for the view state if one exists, otherwise returns zero
 	 */
diff --git a/Engine/Source/Runtime/Renderer/Private/DistanceFieldAmbientOcclusion.h b/Engine/Source/Runtime/Renderer/Private/DistanceFieldAmbientOcclusion.h
--- a/Engine/Source/Runtime/Renderer/Private/DistanceFieldAmbientOcclusion.h
+++ b/Engine/Source/Runtime/Renderer/Private/DistanceFieldAmbientOcclusion.h
@@ -259,7 +259,7 @@ public:
 		SetShaderValue(RHICmdList, ShaderRHI, BaseLevelTexelSize, BaseLevelTexelSizeValue);
 
 		extern FVector2D GetJitterOffset(int32 SampleIndex);
-		SetShaderValue(RHICmdList, ShaderRHI, JitterOffset, GetJitterOffset(View.ViewState->GetDistanceFieldTemporalSampleIndex()));
+		SetShaderValue(RHICmdList, ShaderRHI, JitterOffset, GetJitterOffset(View.GetDistanceFieldTemporalSampleIndex()));
 
 
 		SetTextureParameter(
diff --git a/Engine/Source/Runtime/Renderer/Private/DistanceFieldLightingPost.cpp b/Engine/Source/Runtime/Renderer/Private/DistanceFieldLightingPost.cpp
--- a/Engine/Source/Runtime/Renderer/Private/DistanceFieldLightingPost.cpp
+++ b/Engine/Source/Runtime/Renderer/Private/DistanceFieldLightingPost.cpp
@@ -109,7 +109,7 @@ public:
 		);
 
 		extern FVector2D GetJitterOffset(int32 SampleIndex);
-		FVector2D const JitterOffsetValue = GetJitterOffset(View.ViewState->GetDistanceFieldTemporalSampleIndex());
+		FVector2D const JitterOffsetValue = GetJitterOffset(View.GetDistanceFieldTemporalSampleIndex());
 
 		const FIntPoint DownsampledBufferSize = GetBufferSizeForAO();
 		const FVector2D BaseLevelTexelSizeValue(1.0f / DownsampledBufferSize.X, 1.0f / DownsampledBufferSize.Y);
diff --git a/Engine/Source/Runtime/Renderer/Private/ScenePrivate.h b/Engine/Source/Runtime/Renderer/Private/ScenePrivate.h
--- a/Engine/Source/Runtime/Renderer/Private/ScenePrivate.h
+++ b/Engine/Source/Runtime/Renderer/Private/ScenePrivate.h
@@ -1148,7 +1148,7 @@ public:
 		}
 	}
 
-	int32 GetDistanceFieldTemporalSampleIndex() const
+	uint32 GetDistanceFieldTemporalSampleIndex() const
 	{
 		return DistanceFieldTemporalSampleIndex;
 	}
