Resolve hang during FSceneRenderer::CleanUp.

diff --git a/Engine/Source/Runtime/Renderer/Private/SceneRendering.cpp b/Engine/Source/Runtime/Renderer/Private/SceneRendering.cpp
--- a/Engine/Source/Runtime/Renderer/Private/SceneRendering.cpp
+++ b/Engine/Source/Runtime/Renderer/Private/SceneRendering.cpp
@@ -3673,7 +3673,7 @@ void FSceneRenderer::CleanUp(FRHICommandListImmediate& RHICmdList)
 		ReleaseSceneRenderer(RHICmdList, GSceneRenderCleanUpState.Renderer, GSceneRenderCleanUpState.MemStackMark);
 		break;
 	case ESceneRenderCleanUpMode::DeferredAndAsync:
-		GSceneRenderCleanUpState.Task->Wait(ENamedThreads::GetRenderThread());
+		GSceneRenderCleanUpState.Task->Wait(ENamedThreads::GetRenderThread_Local());
 		DeleteSceneRenderer(RHICmdList, GSceneRenderCleanUpState.Renderer, GSceneRenderCleanUpState.MemStackMark);
 		break;
 	}
