Reuse previous frame SceneColorHalfRes for scene color downsample output instead of creating a new one. This prevents having two non-transient SceneColorHalfRes textures allocated on fast VRAM

diff --git a/Engine/Source/Runtime/Renderer/Private/PostProcess/PostProcessDownsample.cpp b/Engine/Source/Runtime/Renderer/Private/PostProcess/PostProcessDownsample.cpp
--- a/Engine/Source/Runtime/Renderer/Private/PostProcess/PostProcessDownsample.cpp
+++ b/Engine/Source/Runtime/Renderer/Private/PostProcess/PostProcessDownsample.cpp
@@ -135,7 +135,14 @@ FScreenPassTexture AddDownsamplePass(
 			Desc.Format = Inputs.FormatOverride;
 		}
 
-		Output.Texture = GraphBuilder.CreateTexture(Desc, Inputs.Name);
+		if (Inputs.UserSuppliedOutput && Translate(Inputs.UserSuppliedOutput->GetDesc()) == Desc)
+		{
+			Output.Texture = GraphBuilder.RegisterExternalTexture(Inputs.UserSuppliedOutput, Inputs.Name);
+		}
+		else
+		{
+			Output.Texture = GraphBuilder.CreateTexture(Desc, Inputs.Name);
+		}
 		Output.ViewRect = FIntRect::DivideAndRoundUp(Inputs.SceneColor.ViewRect, 2);
 		Output.LoadAction = ERenderTargetLoadAction::ENoAction;
 	}
diff --git a/Engine/Source/Runtime/Renderer/Private/PostProcess/PostProcessDownsample.h b/Engine/Source/Runtime/Renderer/Private/PostProcess/PostProcessDownsample.h
--- a/Engine/Source/Runtime/Renderer/Private/PostProcess/PostProcessDownsample.h
+++ b/Engine/Source/Runtime/Renderer/Private/PostProcess/PostProcessDownsample.h
@@ -37,6 +37,9 @@ struct FDownsamplePassInputs
 	// Friendly name of the pass. Used for logging and profiling.
 	const TCHAR* Name = nullptr;
 
+	// Optional user supplied output buffer.
+	IPooledRenderTarget* UserSuppliedOutput = nullptr;
+
 	// Input scene color RDG texture / view rect. Must not be null.
 	FScreenPassTexture SceneColor;
 
diff --git a/Engine/Source/Runtime/Renderer/Private/PostProcess/PostProcessing.cpp b/Engine/Source/Runtime/Renderer/Private/PostProcess/PostProcessing.cpp
--- a/Engine/Source/Runtime/Renderer/Private/PostProcess/PostProcessing.cpp
+++ b/Engine/Source/Runtime/Renderer/Private/PostProcess/PostProcessing.cpp
@@ -741,6 +741,7 @@ void AddPostProcessingPasses(FRDGBuilder& GraphBuilder, const FViewInfo& View, c
 			PassInputs.SceneColor = SceneColor;
 			PassInputs.Quality = DownsampleQuality;
 			PassInputs.FormatOverride = DownsampleOverrideFormat;
+			PassInputs.UserSuppliedOutput = View.PrevViewInfo.HalfResTemporalAAHistory;
 
 			DownsampledSceneColor = AddDownsamplePass(GraphBuilder, View, PassInputs);
 		}
diff --git a/Engine/Source/Runtime/Renderer/Private/PostProcess/TemporalAA.cpp b/Engine/Source/Runtime/Renderer/Private/PostProcess/TemporalAA.cpp
--- a/Engine/Source/Runtime/Renderer/Private/PostProcess/TemporalAA.cpp
+++ b/Engine/Source/Runtime/Renderer/Private/PostProcess/TemporalAA.cpp
@@ -511,8 +511,16 @@ FTAAOutputs AddTemporalAAPass(
 				Inputs.DownsampleOverrideFormat != PF_Unknown ? Inputs.DownsampleOverrideFormat : Inputs.SceneColorInput->Desc.Format,
 				FClearValueBinding::Black,
 				TexCreate_ShaderResource | TexCreate_UAV | GFastVRamConfig.Downsample);
+			const TRefCountPtr<IPooledRenderTarget>& PrevFrameHalfResTAAHistory = View.PrevViewInfo.HalfResTemporalAAHistory;
 
-			Outputs.DownsampledSceneColor = GraphBuilder.CreateTexture(HalfResSceneColorDesc, TEXT("SceneColorHalfRes"));
+			if (PrevFrameHalfResTAAHistory && Translate(PrevFrameHalfResTAAHistory->GetDesc(), ERenderTargetTexture::ShaderResource) == HalfResSceneColorDesc)
+			{
+				Outputs.DownsampledSceneColor = GraphBuilder.RegisterExternalTexture(PrevFrameHalfResTAAHistory);
+			}
+			else
+			{
+				Outputs.DownsampledSceneColor = GraphBuilder.CreateTexture(HalfResSceneColorDesc, TEXT("SceneColorHalfRes"));
+			}
 		}
 	}
 
