Implements r.Nanite.ViewMeshLODBias.*

diff --git a/Engine/Source/Runtime/Renderer/Private/DeferredShadingRenderer.cpp b/Engine/Source/Runtime/Renderer/Private/DeferredShadingRenderer.cpp
--- a/Engine/Source/Runtime/Renderer/Private/DeferredShadingRenderer.cpp
+++ b/Engine/Source/Runtime/Renderer/Private/DeferredShadingRenderer.cpp
@@ -16,6 +16,7 @@
 #include "PostProcess/SceneFilterRendering.h"
 #include "PostProcess/PostProcessSubsurface.h"
 #include "PostProcess/PostProcessVisualizeCalibrationMaterial.h"
+#include "PostProcess/TemporalAA.h"
 #include "CompositionLighting/CompositionLighting.h"
 #include "FXSystem.h"
 #include "OneColorShader.h"
@@ -222,6 +223,22 @@ static TAutoConsoleVariable<int32> CVarForceBlackVelocityBuffer(
 	ECVF_RenderThreadSafe);
 #endif
 
+static TAutoConsoleVariable<int32> CVarNaniteViewMeshLODBiasEnable(
+	TEXT("r.Nanite.ViewMeshLODBias.Enable"), 1,
+	TEXT("Whether LOD offset to apply for rasterized Nanite meshes for the main viewport should be based off TSR's ScreenPercentage (Enabled by default)."),
+	ECVF_RenderThreadSafe);
+
+static TAutoConsoleVariable<float> CVarNaniteViewMeshLODBiasOffset(
+	TEXT("r.Nanite.ViewMeshLODBias.Offset"), 0.0f,
+	TEXT("LOD offset to apply for rasterized Nanite meshes for the main viewport when using TSR (Default = 0)."),
+	ECVF_RenderThreadSafe);
+
+static TAutoConsoleVariable<float> CVarNaniteViewMeshLODBiasMin(
+	TEXT("r.Nanite.ViewMeshLODBias.Min"), -2.0f,
+	TEXT("Minimum LOD offset for rasterizing Nanite meshes for the main viewport (Default = -2)."),
+	ECVF_RenderThreadSafe);
+
+
 namespace Lumen
 {
 	extern bool AnyLumenHardwareRayTracingPassEnabled();
@@ -2282,7 +2299,23 @@ void FDeferredShadingSceneRenderer::Render(FRDGBuilder& GraphBuilder)
 				static FString EmptyFilterName = TEXT(""); // Empty filter represents primary view.
 				const bool bExtractStats = Nanite::IsStatFilterActive(EmptyFilterName);
 
-				Nanite::FPackedView PackedView = Nanite::CreatePackedViewFromViewInfo(View, RasterTextureSize, VIEW_FLAG_HZBTEST, /*StreamingPriorityCategory*/ 3);
+				float LODScaleFactor = 1.0f;
+				if (View.PrimaryScreenPercentageMethod == EPrimaryScreenPercentageMethod::TemporalUpscale &&
+					CVarNaniteViewMeshLODBiasEnable.GetValueOnRenderThread() != 0)
+				{
+					float TemporalUpscaleFactor = float(View.GetSecondaryViewRectSize().X) / float(View.ViewRect.Width());
+
+					LODScaleFactor = TemporalUpscaleFactor * FMath::Exp2(-CVarNaniteViewMeshLODBiasOffset.GetValueOnRenderThread());
+					LODScaleFactor = FMath::Min(LODScaleFactor, FMath::Exp2(-CVarNaniteViewMeshLODBiasMin.GetValueOnRenderThread()));
+				}
+
+				Nanite::FPackedView PackedView = Nanite::CreatePackedViewFromViewInfo(
+					View,
+					RasterTextureSize,
+					VIEW_FLAG_HZBTEST,
+					/* StreamingPriorityCategory = */ 3,
+					/* MinBoundsRadius = */ 0.0f,
+					LODScaleFactor);
 
 				Nanite::CullRasterize(
 					GraphBuilder,
