Add mip map bias for virtual texture sampling.
Respects the existing r.MiipMapLodBias CVar.
Also add residency tracking for each virtual texture physical pool.
If the pool has been created from a config with the bEnableResidencyMipMapBias flag set then a mip map lod bias is applied to bring the pool back within budget.
The residency mip map lod bias is applied globally (combined with the global mip map lod bias).
In future it could be good to apply the mip map lod bias per physical pool, but this will require tracking all physical pool combinations for each page table.
Physical pool residency tracking can be shown on screen with r.vt.residency.show CVar.

diff --git a/Engine/Shaders/Private/VirtualTextureCommon.ush b/Engine/Shaders/Private/VirtualTextureCommon.ush
--- a/Engine/Shaders/Private/VirtualTextureCommon.ush
+++ b/Engine/Shaders/Private/VirtualTextureCommon.ush
@@ -220,7 +220,7 @@ int TextureComputeVirtualMipLevel(
 	const float Noise = 0;
 #endif
 
-	return (int)floor(ComputedLevel + MipBias + Noise) + int(PageTableUniform.vPageTableMipBias);
+	return (int)floor(ComputedLevel + MipBias + View.GlobalVirtualTextureMipBias + Noise) + int(PageTableUniform.vPageTableMipBias);
 }
 
 /** Samples page table indirection and any apply changes to UV, vLevel and PageTable information for adaptive page tables. */
@@ -332,7 +332,7 @@ VTPageTableResult TextureLoadVirtualPageTable(
 {
 	VTPageTableResult Result = (VTPageTableResult)0;
 	UV = UV * PageTableUniform.UVScale;
-	int vLevel = 0;
+	int vLevel = View.GlobalVirtualTextureMipBias;
 #if PIXELSHADER
 	vLevel = TextureComputeVirtualMipLevel(Result, ddx(UV), ddy(UV), MipBias, SvPositionXY, PageTableUniform);
 #endif // PIXELSHADER
@@ -351,7 +351,7 @@ VTPageTableResult TextureLoadVirtualPageTable(
 {
 	VTPageTableResult Result = (VTPageTableResult)0;
 	UV = UV * PageTableUniform.UVScale;
-	int vLevel = 0;
+	int vLevel = View.GlobalVirtualTextureMipBias;
 #if PIXELSHADER
 	vLevel = TextureComputeVirtualMipLevel(Result, ddx(UV), ddy(UV), MipBias, SvPositionXY, PageTableUniform);
 #endif // PIXELSHADER
@@ -371,7 +371,7 @@ VTPageTableResult TextureLoadVirtualPageTable(
 {
 	VTPageTableResult Result = (VTPageTableResult)0;
 	UV = UV * PageTableUniform.UVScale;
-	int vLevel = 0;
+	int vLevel = View.GlobalVirtualTextureMipBias;
 #if PIXELSHADER
 	vLevel = TextureComputeVirtualMipLevel(Result, ddx(UV), ddy(UV), MipBias, SvPositionXY, PageTableUniform);
 #endif // PIXELSHADER
@@ -390,7 +390,7 @@ VTPageTableResult TextureLoadVirtualPageTable(
 {
 	VTPageTableResult Result = (VTPageTableResult)0;
 	UV = UV * PageTableUniform.UVScale;
-	int vLevel = 0;
+	int vLevel = View.GlobalVirtualTextureMipBias;
 #if PIXELSHADER
 	vLevel = TextureComputeVirtualMipLevel(Result, ddx(UV), ddy(UV), MipBias, SvPositionXY, PageTableUniform);
 #endif // PIXELSHADER
@@ -446,7 +446,7 @@ VTPageTableResult TextureLoadVirtualPageTableLevel(
 {
 	VTPageTableResult Result = (VTPageTableResult)0;
 	UV = UV * PageTableUniform.UVScale;
-	int vLevel = (int)floor(Level);
+	int vLevel = (int)floor(Level + View.GlobalVirtualTextureMipBias);
 	UV = ApplyAddressMode(UV, AddressU, AddressV);
 	TextureLoadVirtualPageTableInternal(Result, PageTable0, PageTableUniform, UV, vLevel);
 	StoreVirtualTextureFeedback(Feedback, SampleIndex, Result.PackedRequest);
@@ -462,7 +462,7 @@ VTPageTableResult TextureLoadVirtualPageTableLevel(
 {
 	VTPageTableResult Result = (VTPageTableResult)0;
 	UV = UV * PageTableUniform.UVScale;
-	int vLevel = (int)floor(Level);
+	int vLevel = (int)floor(Level + View.GlobalVirtualTextureMipBias);
 	UV = ApplyAddressMode(UV, AddressU, AddressV);
 	TextureLoadVirtualPageTableInternal(Result, PageTable0, PageTableUniform, UV, vLevel);
 	return Result;
@@ -479,7 +479,7 @@ VTPageTableResult TextureLoadVirtualPageTableLevel(
 {
 	VTPageTableResult Result = (VTPageTableResult)0;
 	UV = UV * PageTableUniform.UVScale;
-	int vLevel = (int)floor(Level);
+	int vLevel = (int)floor(Level + View.GlobalVirtualTextureMipBias);
 	UV = ApplyAddressMode(UV, AddressU, AddressV);
 	TextureLoadVirtualPageTableInternal(Result, PageTable0, PageTable1, PageTableUniform, UV, vLevel);
 	StoreVirtualTextureFeedback(Feedback, SampleIndex, Result.PackedRequest);
@@ -495,7 +495,7 @@ VTPageTableResult TextureLoadVirtualPageTableLevel(
 {
 	VTPageTableResult Result = (VTPageTableResult)0;
 	UV = UV * PageTableUniform.UVScale;
-	int vLevel = (int)floor(Level);
+	int vLevel = (int)floor(Level + View.GlobalVirtualTextureMipBias);
 	UV = ApplyAddressMode(UV, AddressU, AddressV);
 	TextureLoadVirtualPageTableInternal(Result, PageTable0, PageTable1, PageTableUniform, UV, vLevel);
 	return Result;
@@ -513,7 +513,7 @@ VTPageTableResult TextureLoadVirtualPageTableAdaptive(
 {
 	VTPageTableResult Result = (VTPageTableResult)0;
 	UV = UV * PageTableUniform.UVScale;
-	int vLevel = 0;
+	int vLevel = View.GlobalVirtualTextureMipBias;
 #if PIXELSHADER
 	vLevel = TextureComputeVirtualMipLevel(Result, ddx(UV), ddy(UV), MipBias, SvPositionXY, PageTableUniform);
 #endif // PIXELSHADER
@@ -556,7 +556,7 @@ VTPageTableResult TextureLoadVirtualPageTableAdaptiveLevel(
 {
 	VTPageTableResult Result = (VTPageTableResult)0;
 	// Level is an index into the full size adaptive VT. AdaptiveLevelBias shifts it relative to the low mips allocated VT.
-	int vLevel = (int)floor(Level) - PageTableUniform.AdaptiveLevelBias;
+	int vLevel = (int)floor(max(Level, View.GlobalVirtualTextureMipBias)) - PageTableUniform.AdaptiveLevelBias;
 	UV = UV * PageTableUniform.UVScale;
 	UV = ApplyAddressMode(UV, AddressU, AddressV);
 	ApplyAdaptivePageTableUniform(PageTableIndirection, Result, PageTableUniform, UV, vLevel);
@@ -575,7 +575,7 @@ VTPageTableResult TextureLoadVirtualPageTableAdaptiveLevel(
 {
 	VTPageTableResult Result = (VTPageTableResult)0;
 	// Level is an index into the full size adaptive VT. AdaptiveLevelBias shifts it relative to the low mips allocated VT.
-	int vLevel = (int)floor(Level) - PageTableUniform.AdaptiveLevelBias;
+	int vLevel = (int)floor(max(Level, View.GlobalVirtualTextureMipBias)) - PageTableUniform.AdaptiveLevelBias;
 	UV = UV * PageTableUniform.UVScale;
 	UV = ApplyAddressMode(UV, AddressU, AddressV);
 	ApplyAdaptivePageTableUniform(PageTableIndirection, Result, PageTableUniform, UV, vLevel);
diff --git a/Engine/Source/Runtime/Engine/Classes/VT/VirtualTexturePoolConfig.h b/Engine/Source/Runtime/Engine/Classes/VT/VirtualTexturePoolConfig.h
--- a/Engine/Source/Runtime/Engine/Classes/VT/VirtualTexturePoolConfig.h
+++ b/Engine/Source/Runtime/Engine/Classes/VT/VirtualTexturePoolConfig.h
@@ -18,7 +18,14 @@ struct FVirtualTextureSpacePoolConfig
 {
 	GENERATED_USTRUCT_BODY()
 
-	FVirtualTextureSpacePoolConfig() : MinTileSize(0), MaxTileSize(0), SizeInMegabyte(0), bAllowSizeScale(false), ScalabilityGroup(0) {}
+	FVirtualTextureSpacePoolConfig() 
+		: MinTileSize(0)
+		, MaxTileSize(0)
+		, SizeInMegabyte(0)
+		, bEnableResidencyMipMapBias(false)
+		, bAllowSizeScale(false)
+		, ScalabilityGroup(0) 
+	{}
 
 	/** Minimum tile size to match (including tile border). */
 	UPROPERTY()
@@ -36,6 +43,10 @@ struct FVirtualTextureSpacePoolConfig
 	UPROPERTY()
 	int32 SizeInMegabyte;
 
+	/** Enable MipMapBias based on pool residency tracking. */
+	UPROPERTY()
+	bool bEnableResidencyMipMapBias;
+
 	/** Allow the size to allocate for the pool to be scaled by some factor. */
 	UPROPERTY()
 	bool bAllowSizeScale;
@@ -45,7 +56,7 @@ struct FVirtualTextureSpacePoolConfig
 	uint32 ScalabilityGroup;
 
 	/** Is this the default config? Use this setting when we can't find any other match. */
-	bool IsDefault() const { return Formats.Num() == 0; }
+	bool IsDefault() const { return Formats.Num() == 0 && SizeInMegabyte > 0; }
 };
 
 UCLASS(config = Engine, transient)
diff --git a/Engine/Source/Runtime/Engine/Private/VT/VirtualTexturePoolConfig.cpp b/Engine/Source/Runtime/Engine/Private/VT/VirtualTexturePoolConfig.cpp
--- a/Engine/Source/Runtime/Engine/Private/VT/VirtualTexturePoolConfig.cpp
+++ b/Engine/Source/Runtime/Engine/Private/VT/VirtualTexturePoolConfig.cpp
@@ -11,6 +11,10 @@ UVirtualTexturePoolConfig::UVirtualTexturePoolConfig(const FObjectInitializer& O
 
 void UVirtualTexturePoolConfig::FindPoolConfig(TEnumAsByte<EPixelFormat> const* InFormats, int32 InNumLayers, int32 InTileSize, FVirtualTextureSpacePoolConfig& OutConfig) const
 {
+	bool bFoundDefaultConfig = false;
+	FVirtualTextureSpacePoolConfig DefaultConfig;
+	DefaultConfig.SizeInMegabyte = DefaultSizeInMegabyte;
+
 	// Reverse iterate so that project config can override base config
 	for (int32 Id = Pools.Num() - 1; Id >= 0 ; Id--)
 	{
@@ -34,8 +38,13 @@ void UVirtualTexturePoolConfig::FindPoolConfig(TEnumAsByte<EPixelFormat> const*
 				return;
 			}
 		}
+
+		if (!bFoundDefaultConfig && Config.IsDefault())
+		{
+			DefaultConfig = Config;
+			bFoundDefaultConfig = true;
+		}
 	}
 
-	OutConfig = FVirtualTextureSpacePoolConfig();
-	OutConfig.SizeInMegabyte = DefaultSizeInMegabyte;
+	OutConfig = DefaultConfig;
 }
diff --git a/Engine/Source/Runtime/Engine/Public/SceneView.h b/Engine/Source/Runtime/Engine/Public/SceneView.h
--- a/Engine/Source/Runtime/Engine/Public/SceneView.h
+++ b/Engine/Source/Runtime/Engine/Public/SceneView.h
@@ -744,6 +744,7 @@ enum ETranslucencyVolumeCascade
 	VIEW_UNIFORM_BUFFER_MEMBER(float, EyeToPixelSpreadAngle) \
 	VIEW_UNIFORM_BUFFER_MEMBER(FMatrix44f, WorldToVirtualTexture) \
 	VIEW_UNIFORM_BUFFER_MEMBER_ARRAY(FVector4, XRPassthroughCameraUVs, [2]) \
+	VIEW_UNIFORM_BUFFER_MEMBER(float, GlobalVirtualTextureMipBias) \
 	VIEW_UNIFORM_BUFFER_MEMBER(uint32, VirtualTextureFeedbackStride) \
 	VIEW_UNIFORM_BUFFER_MEMBER(uint32, VirtualTextureFeedbackJitterOffset) \
 	VIEW_UNIFORM_BUFFER_MEMBER(uint32, VirtualTextureFeedbackSampleOffset) \
diff --git a/Engine/Source/Runtime/Engine/Public/ShowFlagsValues.inl b/Engine/Source/Runtime/Engine/Public/ShowFlagsValues.inl
--- a/Engine/Source/Runtime/Engine/Public/ShowFlagsValues.inl
+++ b/Engine/Source/Runtime/Engine/Public/ShowFlagsValues.inl
@@ -388,5 +388,8 @@ SHOWFLAG_FIXED_IN_SHIPPING(0, DrawOnlyVSMInvalidatingGeo, SFG_Visualize, NSLOCTE
 /** If single layer water should use a full precision depth buffer for SceneDepthWithoutSingleLayerWater. For now SHOWFLAG_ALWAYS_ACCESSIBLE because it's exposed in SceneCapture */
 SHOWFLAG_ALWAYS_ACCESSIBLE(SingleLayerWaterRefractionFullPrecision, SFG_Advanced, NSLOCTEXT("UnrealEd", "SingleLayerWaterRefractionFullPrecision", "Single Layer Water Refraction Full Precision"))
 
+/** Debug the Virtual Texture System */
+SHOWFLAG_FIXED_IN_SHIPPING(0, VirtualTextureResidency, SFG_Hidden, NSLOCTEXT("UnrealEd", "VirtualTextureResidency", "Virtual Texture Residency"))
+
 #undef SHOWFLAG_ALWAYS_ACCESSIBLE
 #undef SHOWFLAG_FIXED_IN_SHIPPING
diff --git a/Engine/Source/Runtime/Renderer/Private/SceneRendering.cpp b/Engine/Source/Runtime/Renderer/Private/SceneRendering.cpp
--- a/Engine/Source/Runtime/Renderer/Private/SceneRendering.cpp
+++ b/Engine/Source/Runtime/Renderer/Private/SceneRendering.cpp
@@ -1802,6 +1802,8 @@ void FViewInfo::SetupUniformBufferParameters(
 
 	ViewUniformShaderParameters.PreIntegratedBRDF = GEngine->PreIntegratedSkinBRDFTexture->GetResource()->TextureRHI;
 
+	ViewUniformShaderParameters.GlobalVirtualTextureMipBias = FVirtualTextureSystem::Get().GetGlobalMipBias();
+
 	ViewUniformShaderParameters.VirtualTextureFeedbackStride = GetVirtualTextureFeedbackBufferSize(SceneTexturesConfig.Extent).X;
 	// Use some low(ish) discrepancy sequence to run over every pixel in the virtual texture feedback tile.
 	ViewUniformShaderParameters.VirtualTextureFeedbackJitterOffset = SampleVirtualTextureFeedbackSequence(FrameIndex);
diff --git a/Engine/Source/Runtime/Renderer/Private/VT/VirtualTexturePhysicalSpace.cpp b/Engine/Source/Runtime/Renderer/Private/VT/VirtualTexturePhysicalSpace.cpp
--- a/Engine/Source/Runtime/Renderer/Private/VT/VirtualTexturePhysicalSpace.cpp
+++ b/Engine/Source/Runtime/Renderer/Private/VT/VirtualTexturePhysicalSpace.cpp
@@ -1,20 +1,65 @@
 // Copyright Epic Games, Inc. All Rights Reserved.
 
 #include "VirtualTexturePhysicalSpace.h"
-#include "VirtualTextureSystem.h"
+
+#include "Engine/Canvas.h"
 #include "RenderTargetPool.h"
 #include "VisualizeTexture.h"
-#include "Stats/Stats.h"
 #include "VT/VirtualTexturePoolConfig.h"
+#include "VT/VirtualTextureScalability.h"
+#include "VT/VirtualTextureSystem.h"
+
+static TAutoConsoleVariable<float> CVarVTResidencyMaxMipMapBias(
+	TEXT("r.VT.Residency.MaxMipMapBias"),
+	8,
+	TEXT("Maximum mip bias to apply to prevent Virtual Texture pool residency over-subscription.\n")
+	TEXT("Default 8"),
+	ECVF_RenderThreadSafe);
+
+static TAutoConsoleVariable<float> CVarVTResidencyUpperBound(
+	TEXT("r.VT.Residency.UpperBound"),
+	0.9f,
+	TEXT("Virtual Texture pool residency above which we increase mip bias.\n")
+	TEXT("Default 0.95"),
+	ECVF_RenderThreadSafe);
+
+static TAutoConsoleVariable<float> CVarVTResidencyLowerBound(
+	TEXT("r.VT.Residency.LowerBound"),
+	0.75f,
+	TEXT("Virtual Texture pool residency below which we decrease mip bias.\n")
+	TEXT("Default 0.85"),
+	ECVF_RenderThreadSafe);
+
+static TAutoConsoleVariable<float> CVarVTResidencyLockedUpperBound(
+	TEXT("r.VT.Residency.LockedUpperBound"),
+	0.65f,
+	TEXT("Virtual Texture pool locked page residency above which we kill any mip bias.\n")
+	TEXT("That's because locked pages are never affected by the mip bias setting. So it is unlikely that we can get the pool within budget.\n")
+	TEXT("Default 0.65"),
+	ECVF_RenderThreadSafe);
+
+static TAutoConsoleVariable<float> CVarVTResidencyAdjustmentRate(
+	TEXT("r.VT.Residency.AdjustmentRate"),
+	0.2,
+	TEXT("Rate at which we adjust mip bias due to Virtual Texture pool residency.\n")
+	TEXT("Default 0.25"),
+	ECVF_RenderThreadSafe);
 
 
 FVirtualTexturePhysicalSpace::FVirtualTexturePhysicalSpace(const FVTPhysicalSpaceDescription& InDesc, uint16 InID)
 	: Description(InDesc)
 	, NumRefs(0u)
-	, LastFrameOversubscribed(0u)
 	, ID(InID)
-	, bPageTableLimit(false)
 	, bGpuTextureLimit(false)
+	, bEnableResidencyMipMapBias(true)
+	, ResidencyMipMapBias(0.0f)
+	, LastFrameOversubscribed(0)
+#if !UE_BUILD_SHIPPING
+	, VisibleHistory(HistorySize)
+	, LockedHistory(HistorySize)
+	, MipMapBiasHistory(HistorySize)
+	, HistoryIndex(0)
+#endif
 {
 	// Find matching physical pool
 	FVirtualTextureSpacePoolConfig Config;
@@ -41,15 +86,22 @@ FVirtualTexturePhysicalSpace::FVirtualTexturePhysicalSpace(const FVTPhysicalSpac
 		bGpuTextureLimit = true;
 	}
 
+	bEnableResidencyMipMapBias = Config.bEnableResidencyMipMapBias;
+
 	Pool.Initialize(GetNumTiles());
 
 	// Initialize this resource FeatureLevel, so it gets re-created on FeatureLevel changes
 	SetFeatureLevel(GMaxRHIFeatureLevel);
 
-#if STATS
-	const FString LongName = FString::Printf(TEXT("WorkingSet %s %%"), FormatInfo.Name);
-	WorkingSetSizeStatID = FDynamicStats::CreateStatIdDouble<STAT_GROUP_TO_FStatGroup(STATGROUP_VirtualTexturing)>(LongName);
-#endif // STATS
+	// Store string for logging.
+	for (uint32 Layer = 0u; Layer < Description.NumLayers; ++Layer)
+	{
+		FormatString += GPixelFormats[Description.Format[Layer]].Name;
+		if (Layer + 1u < Description.NumLayers)
+		{
+			FormatString += TEXT(", ");
+		}
+	}
 }
 
 FVirtualTexturePhysicalSpace::~FVirtualTexturePhysicalSpace()
@@ -150,10 +202,122 @@ uint32 FVirtualTexturePhysicalSpace::GetSizeInBytes() const
 	return GetNumTiles() * TileSizeBytes;
 }
 
-#if STATS
-void FVirtualTexturePhysicalSpace::UpdateWorkingSetStat()
+void FVirtualTexturePhysicalSpace::UpdateResidencyTracking(uint32 Frame)
+{
+	float LockedUpperBound = CVarVTResidencyLockedUpperBound.GetValueOnRenderThread();
+	float LowerBound = CVarVTResidencyLowerBound.GetValueOnRenderThread();
+	float UpperBound = CVarVTResidencyUpperBound.GetValueOnRenderThread();
+	float AdjustmentRate = CVarVTResidencyAdjustmentRate.GetValueOnRenderThread();
+	float MaxMipMapBias = CVarVTResidencyMaxMipMapBias.GetValueOnRenderThread();
+
+	const uint32 NumPages = Pool.GetNumPages();
+	const uint32 NumLockedPages = Pool.GetNumLockedPages();
+	const float LockedPageResidency = (float)NumLockedPages / (float)NumPages;
+
+	const uint32 PageFreeThreshold = FMath::Max(VirtualTextureScalability::GetPageFreeThreshold(), 0u);
+	const uint32 FrameMinusThreshold = Frame > PageFreeThreshold ? Frame - PageFreeThreshold : 0;
+	const uint32 NumVisiblePages = Pool.GetNumVisiblePages(FrameMinusThreshold);
+	const float VisiblePageResidency = (float)NumVisiblePages / (float)NumPages;
+	
+	if (ResidencyMipMapBias > 0.f && VisiblePageResidency < LowerBound)
+	{
+		ResidencyMipMapBias -= AdjustmentRate * (LowerBound - VisiblePageResidency);
+	}
+	else if (VisiblePageResidency > UpperBound)
+	{
+		ResidencyMipMapBias += AdjustmentRate * (VisiblePageResidency - UpperBound);
+	}
+
+	ResidencyMipMapBias = FMath::Clamp(ResidencyMipMapBias, 0.f, MaxMipMapBias);
+
+	if (ResidencyMipMapBias > 0.f)
+	{
+		LastFrameOversubscribed = Frame;
+	}
+
+	if (!bEnableResidencyMipMapBias || LockedPageResidency > LockedUpperBound)
+	{
+		ResidencyMipMapBias = 0.f;
+	}
+
+#if !UE_BUILD_SHIPPING
+	LockedHistory[HistoryIndex+1] = LockedPageResidency;
+	VisibleHistory[HistoryIndex+1] = VisiblePageResidency;
+	MipMapBiasHistory[HistoryIndex+1] = ResidencyMipMapBias / MaxMipMapBias;
+	HistoryIndex++;
+#endif
+}
+
+void FVirtualTexturePhysicalSpace::DrawResidencyGraph(FCanvas* Canvas, FBox2D CanvasPosition)
 {
-	const double Value = (double)WorkingSetSize.GetValue() / (double)GetNumTiles() * 100.0;
-	FThreadStats::AddMessage(WorkingSetSizeStatID.GetName(), EStatOperation::Set, Value);
+	// Note that this is called on game thread and reads the history values written on the render thread.
+	// But it should be safe and any race condition will only lead to a slightly incorrect graph.
+
+#if !UE_BUILD_SHIPPING
+	const FLinearColor BackgroundColor(0.0f, 0.0f, 0.0f, 0.7f);
+	const FLinearColor GraphBorderColor(0.1f, 0.1f, 0.1f);
+	const FLinearColor GraphResidencyColor(0.8f, 0.1f, 0.1f);
+	const FLinearColor GraphLockedPageColor(0.8f, 0.8f, 0.1f);
+	const FLinearColor GraphMipMapBiasColor(0.1f, 0.8f, 0.1f);
+
+	FCanvasTileItem BackgroundTile(CanvasPosition.Min, CanvasPosition.GetSize(), BackgroundColor);
+	BackgroundTile.BlendMode = SE_BLEND_AlphaBlend;
+	Canvas->DrawItem(BackgroundTile);
+
+	const int32 BorderSize = 10;
+
+	FString Title = FString::Printf(TEXT("%s (%dMB)"), *FormatString, GetSizeInBytes() / (1024 * 1024));
+	Canvas->DrawShadowedString(CanvasPosition.Min.X + BorderSize, CanvasPosition.Min.Y, *Title, GEngine->GetSmallFont(), FLinearColor::White);
+
+	CanvasPosition.Min += FVector2D(BorderSize, BorderSize);
+	CanvasPosition.Max -= FVector2D(BorderSize, BorderSize);
+
+	const uint32 NumSamples = FMath::Min3<uint32>((uint32)CanvasPosition.GetSize().X, HistorySize, HistoryIndex);
+
+	FHitProxyId HitProxyId = Canvas->GetHitProxyId();
+	FBatchedElements* BatchedElements = Canvas->GetBatchedElements(FCanvas::ET_Line);
+	BatchedElements->AddReserveLines(2 + 3 * NumSamples);
+
+	BatchedElements->AddLine(
+		FVector(CanvasPosition.Min.X - 1.0f, CanvasPosition.Max.Y, 0.0f),
+		FVector(CanvasPosition.Min.X - 1.0f, CanvasPosition.Min.Y - 1.0f, 0.0f),
+		GraphBorderColor,
+		HitProxyId);
+
+	BatchedElements->AddLine(
+		FVector(CanvasPosition.Min.X, CanvasPosition.Max.Y - 1.0f, 0.0f),
+		FVector(CanvasPosition.Max.X, CanvasPosition.Max.Y - 1.0f, 0.0f),
+		GraphBorderColor,
+		HitProxyId);
+
+	for (uint32 SampleIndex = 0; SampleIndex < NumSamples; ++SampleIndex)
+	{
+		float Visible0 = VisibleHistory[HistoryIndex - NumSamples + SampleIndex];
+		float Visible1 = VisibleHistory[HistoryIndex - NumSamples + SampleIndex + 1];
+
+		BatchedElements->AddLine(
+			FVector(CanvasPosition.Min.X + SampleIndex, CanvasPosition.Max.Y - Visible0 * CanvasPosition.GetSize().Y, 0.0f),
+			FVector(CanvasPosition.Min.X + SampleIndex + 1, CanvasPosition.Max.Y - Visible1 * CanvasPosition.GetSize().Y, 0.0f),
+			GraphResidencyColor,
+			HitProxyId);
+
+		float Locked0 = LockedHistory[HistoryIndex - NumSamples + SampleIndex];
+		float Locked1 = LockedHistory[HistoryIndex - NumSamples + SampleIndex + 1];
+
+		BatchedElements->AddLine(
+			FVector(CanvasPosition.Min.X + SampleIndex, CanvasPosition.Max.Y - Locked0 * CanvasPosition.GetSize().Y, 0.0f),
+			FVector(CanvasPosition.Min.X + SampleIndex + 1, CanvasPosition.Max.Y - Locked1 * CanvasPosition.GetSize().Y, 0.0f),
+			GraphLockedPageColor,
+			HitProxyId);
+
+		float MipMapBias0 = MipMapBiasHistory[HistoryIndex - NumSamples + SampleIndex];
+		float MipMapBias1 = MipMapBiasHistory[HistoryIndex - NumSamples + SampleIndex + 1];
+
+		BatchedElements->AddLine(
+			FVector(CanvasPosition.Min.X + SampleIndex, CanvasPosition.Max.Y - MipMapBias0 * CanvasPosition.GetSize().Y, 0.0f),
+			FVector(CanvasPosition.Min.X + SampleIndex + 1, CanvasPosition.Max.Y - MipMapBias1 * CanvasPosition.GetSize().Y, 0.0f),
+			GraphMipMapBiasColor,
+			HitProxyId);
+	}
+#endif // UE_BUILD_SHIPPING
 }
-#endif // STATS
diff --git a/Engine/Source/Runtime/Renderer/Private/VT/VirtualTexturePhysicalSpace.h b/Engine/Source/Runtime/Renderer/Private/VT/VirtualTexturePhysicalSpace.h
--- a/Engine/Source/Runtime/Renderer/Private/VT/VirtualTexturePhysicalSpace.h
+++ b/Engine/Source/Runtime/Renderer/Private/VT/VirtualTexturePhysicalSpace.h
@@ -3,9 +3,11 @@
 #pragma once
 
 #include "CoreMinimal.h"
-#include "VirtualTextureShared.h"
-#include "TexturePagePool.h"
+
+#include "Containers/CircularBuffer.h"
 #include "RendererInterface.h"
+#include "TexturePagePool.h"
+#include "VirtualTextureShared.h"
 #include "VirtualTexturing.h"
 
 struct FVTPhysicalSpaceDescription
@@ -47,6 +49,7 @@ public:
 	virtual ~FVirtualTexturePhysicalSpace();
 
 	inline const FVTPhysicalSpaceDescription& GetDescription() const { return Description; }
+	inline const FString& GetFormatString() const { return FormatString; }
 	inline EPixelFormat GetFormat(int32 Layer) const { return Description.Format[Layer]; }
 	inline uint16 GetID() const { return ID; }
 	inline uint32 GetNumTiles() const { return TextureSizeInTiles * TextureSizeInTiles; }
@@ -70,9 +73,6 @@ public:
 	inline uint32 Release() { check(NumRefs > 0u); return --NumRefs; }
 	inline uint32 GetRefCount() const { return NumRefs; }
 
-	inline uint32 GetLastFrameOversubscribed() const { return LastFrameOversubscribed; }
-	inline void SetLastFrameOversubscribed(uint32 InFrame) { LastFrameOversubscribed = InFrame; }
-
 	FRHITexture* GetPhysicalTexture(int32 Layer) const
 	{
 		check(PooledRenderTarget[Layer].IsValid());
@@ -95,18 +95,18 @@ public:
 		return PooledRenderTarget[Layer];
 	}
 
-#if STATS
-	inline void ResetWorkingSetSize() { WorkingSetSize.Reset(); }
-	inline void IncrementWorkingSetSize(int32 Amount) { WorkingSetSize.Add(Amount); }
-	void UpdateWorkingSetStat();
-#else // STATS
-	inline void ResetWorkingSetSize() {}
-	inline void IncrementWorkingSetSize(int32 Amount) {}
-	inline void UpdateWorkingSetStat() {}
-#endif // !STATS
+	/** Update internal tracking of residency. This is used to update stats and to calculate a mip bias to keep within the pool budget. */
+	void UpdateResidencyTracking(uint32 Frame);
+	/** Get dynamic mip bias used to keep within residency budget. */
+	inline float GetResidencyMipMapBias() const { return ResidencyMipMapBias; }
+	/** Get frame at which residency tracking last saw over-subscription. */
+	inline uint32 GetLastFrameOversubscribed() const { return LastFrameOversubscribed; }
+	/** Draw residency graph on screen. */
+	void DrawResidencyGraph(class FCanvas* Canvas, FBox2D CanvasPosition);
 
 private:
 	FVTPhysicalSpaceDescription Description;
+	FString FormatString;
 	FTexturePagePool Pool;
 	TRefCountPtr<IPooledRenderTarget> PooledRenderTarget[VIRTUALTEXTURE_SPACE_MAXLAYERS];
 	FShaderResourceViewRHIRef TextureSRV[VIRTUALTEXTURE_SPACE_MAXLAYERS];
@@ -115,13 +115,19 @@ private:
 
 	uint32 TextureSizeInTiles;
 	uint32 NumRefs;
-	uint32 LastFrameOversubscribed;
 	uint16 ID;
-	bool bPageTableLimit; // True if the physical size was limited by the page table format requested
+	
 	bool bGpuTextureLimit; // True if the physical size was limited by the maximum GPU texture size
+	
+	bool bEnableResidencyMipMapBias;
+	float ResidencyMipMapBias;
+	uint32 LastFrameOversubscribed;
 
-#if STATS
-	TStatId WorkingSetSizeStatID;
-	FThreadSafeCounter WorkingSetSize;
-#endif // STATS
+#if !UE_BUILD_SHIPPING
+	static const int32 HistorySize = 512;
+	TCircularBuffer<float> VisibleHistory;
+	TCircularBuffer<float> LockedHistory;
+	TCircularBuffer<float> MipMapBiasHistory;
+	uint32 HistoryIndex;
+#endif
 };
diff --git a/Engine/Source/Runtime/Renderer/Private/VT/VirtualTextureSystem.cpp b/Engine/Source/Runtime/Renderer/Private/VT/VirtualTextureSystem.cpp
--- a/Engine/Source/Runtime/Renderer/Private/VT/VirtualTextureSystem.cpp
+++ b/Engine/Source/Runtime/Renderer/Private/VT/VirtualTextureSystem.cpp
@@ -3,6 +3,8 @@
 #include "VirtualTextureSystem.h"
 
 #include "AllocatedVirtualTexture.h"
+#include "Engine/Canvas.h"
+#include "Debug/DebugDrawService.h"
 #include "HAL/IConsoleManager.h"
 #include "PostProcess/SceneRenderTargets.h"
 #include "ProfilingDebugging/CsvProfiler.h"
@@ -42,7 +44,7 @@ DECLARE_CYCLE_STAT(TEXT("Feedback Map"), STAT_FeedbackMap, STATGROUP_VirtualText
 DECLARE_CYCLE_STAT(TEXT("Feedback Analysis"), STAT_FeedbackAnalysis, STATGROUP_VirtualTexturing);
 DECLARE_CYCLE_STAT(TEXT("Page Table Updates"), STAT_PageTableUpdates, STATGROUP_VirtualTexturing);
 DECLARE_CYCLE_STAT(TEXT("Flush Cache"), STAT_FlushCache, STATGROUP_VirtualTexturing);
-DECLARE_CYCLE_STAT(TEXT("Update Stats"), STAT_UpdateStats, STATGROUP_VirtualTexturing);
+DECLARE_CYCLE_STAT(TEXT("Update Residency Tracking"), STAT_ResidencyTracking, STATGROUP_VirtualTexturing);
 
 DECLARE_DWORD_COUNTER_STAT(TEXT("Num page visible"), STAT_NumPageVisible, STATGROUP_VirtualTexturing);
 DECLARE_DWORD_COUNTER_STAT(TEXT("Num page visible resident"), STAT_NumPageVisibleResident, STATGROUP_VirtualTexturing);
@@ -114,6 +116,12 @@ static TAutoConsoleVariable<int32> CVarVTProduceLockedTilesOnFlush(
 	TEXT("Should locked tiles be (re)produced when flushing the cache"),
 	ECVF_RenderThreadSafe
 );
+static TAutoConsoleVariable<int32> CVarVTResidencyShow(
+	TEXT("r.VT.Residency.Show"),
+	0,
+	TEXT("Show on screen HUD for virtual texture physical pool residency"),
+	ECVF_Default
+);
 
 static FORCEINLINE uint32 EncodePage(uint32 ID, uint32 vLevel, uint32 vTileX, uint32 vTileY)
 {
@@ -241,6 +249,7 @@ FVirtualTextureSystem::FVirtualTextureSystem()
 {
 #if !UE_BUILD_SHIPPING
 	FCoreDelegates::OnGetOnScreenMessages.AddRaw(this, &FVirtualTextureSystem::GetOnScreenMessages);
+	DrawResidencyHudDelegateHandle = UDebugDrawService::Register(TEXT("VirtualTextureResidency"), FDebugDrawDelegate::CreateRaw(this, &FVirtualTextureSystem::DrawResidencyHud));
 #endif
 }
 
@@ -248,6 +257,7 @@ FVirtualTextureSystem::~FVirtualTextureSystem()
 {
 #if !UE_BUILD_SHIPPING
 	FCoreDelegates::OnGetOnScreenMessages.RemoveAll(this);
+	UDebugDrawService::Unregister(DrawResidencyHudDelegateHandle);
 #endif
 
 	DestroyPendingVirtualTextures(true);
@@ -348,13 +358,8 @@ void FVirtualTextureSystem::ListPhysicalPoolsFromConsole()
 			const FTexturePagePool& PagePool = PhysicalSpace.GetPagePool();
 			const uint32 TotalSizeInBytes = PhysicalSpace.GetSizeInBytes();
 
-			UE_LOG(LogConsoleResponse, Display, TEXT("PhysicaPool: [%i] %ix%i:"), i, Desc.TileSize, Desc.TileSize);
+			UE_LOG(LogConsoleResponse, Display, TEXT("PhysicaPool: [%i] %s (%ix%i):"), i, *PhysicalSpace.GetFormatString(), Desc.TileSize, Desc.TileSize);
 			
-			for (int Layer = 0; Layer < Desc.NumLayers; ++Layer)
-			{
-				UE_LOG(LogConsoleResponse, Display, TEXT("  Layer %i=%s"), Layer, GPixelFormats[Desc.Format[Layer]].Name);
-			}
-
 			const int32 AllocatedTiles = PagePool.GetNumAllocatedPages();
 			const float AllocatedLoad = (float)AllocatedTiles / (float)PhysicalSpace.GetNumTiles();
 			const float AllocatedMemory = AllocatedLoad * TotalSizeInBytes / 1024.0f / 1024.0f;
@@ -1293,7 +1298,7 @@ void FVirtualTextureSystem::Update(FRDGBuilder& GraphBuilder, ERHIFeatureLevel::
 	// Submit the merged requests
 	SubmitRequests(GraphBuilder, FeatureLevel, MemStack, MergedRequestList, true);
 
-	UpdateCSVStats();
+	UpdateResidencyTracking();
 
 	ReleasePendingSpaces();
 }
@@ -1797,11 +1802,6 @@ void FVirtualTextureSystem::GatherRequestsTask(const FGatherRequestsParameters&
 		FVirtualTexturePhysicalSpace* RESTRICT PhysicalSpace = GetPhysicalSpace(PhysicalSpaceID);
 		FPageUpdateBuffer& RESTRICT Buffer = PageUpdateBuffers[PhysicalSpaceID];
 
-		if (Buffer.WorkingSetSize > 0u)
-		{
-			PhysicalSpace->IncrementWorkingSetSize(Buffer.WorkingSetSize);
-		}
-
 		if (Buffer.NumPages > 0u)
 		{
 			Buffer.NumPageUpdates += Buffer.NumPages;
@@ -1839,28 +1839,18 @@ void FVirtualTextureSystem::GetContinuousUpdatesToProduce(FUniqueRequestList con
 	}
 }
 
-void FVirtualTextureSystem::UpdateCSVStats() const
+void FVirtualTextureSystem::UpdateResidencyTracking() const
 {
-#if CSV_PROFILER
-	SCOPE_CYCLE_COUNTER(STAT_UpdateStats);
+	SCOPE_CYCLE_COUNTER(STAT_ResidencyTracking);
 
-	uint32 TotalPages = 0;
-	uint32 CurrentPages = 0;
-	const uint32 AgeTolerance = 5; // Include some tolerance/smoothing for previous frames
 	for (int32 i = 0; i < PhysicalSpaces.Num(); ++i)
 	{
 		FVirtualTexturePhysicalSpace* PhysicalSpace = PhysicalSpaces[i];
 		if (PhysicalSpace)
 		{
-			FTexturePagePool const& PagePool = PhysicalSpace->GetPagePool();
-			TotalPages += PagePool.GetNumPages();
-			CurrentPages += PagePool.GetNumVisiblePages(Frame > AgeTolerance ? Frame - AgeTolerance : 0);
+			PhysicalSpace->UpdateResidencyTracking(Frame);
 		}
 	}
-
-	const float PhysicalPoolUsage = TotalPages > 0 ? (float)CurrentPages / (float)TotalPages : 0.f;
-	CSV_CUSTOM_STAT_GLOBAL(VirtualTexturePageUsage, PhysicalPoolUsage, ECsvCustomStatOp::Set);
-#endif
 }
 
 void FVirtualTextureSystem::SubmitRequestsFromLocalTileList(TArray<FVirtualTextureLocalTile>& OutDeferredTiles, const TSet<FVirtualTextureLocalTile>& LocalTileList, EVTProducePageFlags Flags, FRDGBuilder& GraphBuilder, ERHIFeatureLevel::Type FeatureLevel)
@@ -2064,7 +2054,6 @@ void FVirtualTextureSystem::SubmitRequests(FRDGBuilder& GraphBuilder, ERHIFeatur
 						}
 						else
 						{
-							PhysicalSpace->SetLastFrameOversubscribed(Frame);
 							bProduceTargetValid = false;
 							NumPageAllocateFails++;
 							break;
@@ -2194,19 +2183,12 @@ void FVirtualTextureSystem::SubmitRequests(FRDGBuilder& GraphBuilder, ERHIFeatur
 				// Once a pool fails an allocation, keep displaying the message for a minimum number of frames, to avoid the message quickly flickering on/off
 				if (PhysicalSpace != nullptr && Frame <= PhysicalSpace->GetLastFrameOversubscribed() + 60u)
 				{
-					const FVTPhysicalSpaceDescription& Desc = PhysicalSpace->GetDescription();
-
-					FString FormatString;
-					for (uint32 Layer = 0u; Layer < Desc.NumLayers; ++Layer)
-					{
-						FormatString += GPixelFormats[Desc.Format[Layer]].Name;
-						if (Layer + 1u < Desc.NumLayers)
-						{
-							FormatString += TEXT(", ");
-						}
-					}
+					FString const& FormatString = PhysicalSpace->GetFormatString();
+					const float MipBias = PhysicalSpace->GetResidencyMipMapBias();
 
-					OnScreenMessages.Add(FCoreDelegates::EOnScreenMessageSeverity::Warning, FText::Format(LOCTEXT("VTOversubscribed", "VT Pool [{0}] is oversubscribed"), FText::FromString(FormatString)));
+					OnScreenMessages.Add(
+						FCoreDelegates::EOnScreenMessageSeverity::Warning, 
+						FText::Format(LOCTEXT("VTOversubscribed", "VT Pool [{0}] is oversubscribed. Setting MipBias {1}"), FText::FromString(FormatString), FText::AsNumber(MipBias)));
 				}
 			}
 		}
@@ -2317,4 +2299,82 @@ void FVirtualTextureSystem::ReleasePendingResources()
 	ReleasePendingSpaces();
 }
 
+float FVirtualTextureSystem::GetGlobalMipBias() const
+{
+	float MaxResidencyMipMapBias = 0.f;
+	for (int32 SpaceIndex = 0; SpaceIndex < PhysicalSpaces.Num(); ++SpaceIndex)
+	{
+		FVirtualTexturePhysicalSpace const* PhysicalSpace = PhysicalSpaces[SpaceIndex];
+		const float ResidencyMipMapBias = PhysicalSpace != nullptr ? PhysicalSpace->GetResidencyMipMapBias() : 0.f;
+		MaxResidencyMipMapBias = FMath::Max(MaxResidencyMipMapBias, ResidencyMipMapBias);
+	}
+
+	return UTexture2D::GetGlobalMipMapLODBias() + MaxResidencyMipMapBias;
+}
+
+void FVirtualTextureSystem::DrawResidencyHud(UCanvas* InCanvas, APlayerController* InController)
+{
+#if !UE_BUILD_SHIPPING
+	if (CVarVTResidencyShow.GetValueOnGameThread() == 0)
+	{
+		return;
+	}
+
+	int32 NumGraphs = 0;
+	for (int32 SpaceIndex = 0; SpaceIndex < PhysicalSpaces.Num(); ++SpaceIndex)
+	{
+		FVirtualTexturePhysicalSpace* PhysicalSpace = PhysicalSpaces[SpaceIndex];
+		if (PhysicalSpace)
+		{
+			NumGraphs++;
+		}
+	}
+
+	if (NumGraphs == 0)
+	{
+		return;
+	}
+
+	const FIntPoint GraphSize(400, 200);
+	const FIntPoint BorderSize(25, 25);
+	const FIntPoint GraphWithBorderSize = GraphSize + BorderSize * 2;
+
+	const FIntPoint CanvasSize = InCanvas->Canvas->GetParentCanvasSize();
+	const int32 NumGraphsInRow = FMath::Max(CanvasSize.X / GraphWithBorderSize.X, 1);
+	const int32 CanvasOffsetY = 90;
+
+	int32 GraphIndex = 0;
+	for (int32 SpaceIndex = 0; SpaceIndex < PhysicalSpaces.Num(); ++SpaceIndex)
+	{
+		FVirtualTexturePhysicalSpace* PhysicalSpace = PhysicalSpaces[SpaceIndex];
+		if (PhysicalSpace)
+		{
+			int32 GraphX = GraphIndex % NumGraphsInRow;
+			int32 GraphY = GraphIndex / NumGraphsInRow;
+
+			FBox2D CanvasPosition;
+			CanvasPosition.Min.X = GraphX * GraphWithBorderSize.X + BorderSize.X;
+			CanvasPosition.Min.Y = GraphY * GraphWithBorderSize.Y + BorderSize.Y + CanvasOffsetY;
+			CanvasPosition.Max = CanvasPosition.Min + GraphSize;
+
+			if (CanvasPosition.Min.Y > CanvasSize.Y)
+			{
+				// Off screen so early out.
+				break;
+			}
+
+			PhysicalSpace->DrawResidencyGraph(InCanvas->Canvas, CanvasPosition);
+		
+			GraphIndex++;
+		}
+	}
+
+	// Key with same colors as used by FVirtualTexturePhysicalSpace::DrawResidencyGraph()
+	InCanvas->Canvas->DrawShadowedString(25, 320, TEXT("MipMap Bias"), GEngine->GetSmallFont(), FLinearColor(0.1f, 0.8f, 0.1f));
+	InCanvas->Canvas->DrawShadowedString(125, 320, TEXT("Page Residency"), GEngine->GetSmallFont(), FLinearColor(0.8f, 0.1f, 0.1f));
+	InCanvas->Canvas->DrawShadowedString(225, 320, TEXT("LockedPage Residency"), GEngine->GetSmallFont(), FLinearColor(0.8f, 0.8f, 0.1f));
+
+#endif
+}
+
 #undef LOCTEXT_NAMESPACE
diff --git a/Engine/Source/Runtime/Renderer/Private/VT/VirtualTextureSystem.h b/Engine/Source/Runtime/Renderer/Private/VT/VirtualTextureSystem.h
--- a/Engine/Source/Runtime/Renderer/Private/VT/VirtualTextureSystem.h
+++ b/Engine/Source/Runtime/Renderer/Private/VT/VirtualTextureSystem.h
@@ -75,6 +75,8 @@ public:
 	void FlushCache();
 	void FlushCache(FVirtualTextureProducerHandle const& ProducerHandle, FIntRect const& TextureRegion, uint32 MaxLevel);
 
+	float GetGlobalMipBias() const;
+
 private:
 	friend class FFeedbackAnalysisTask;
 	friend class FGatherRequestsTask;
@@ -103,7 +105,7 @@ private:
 
 	void GetContinuousUpdatesToProduce(FUniqueRequestList const* RequestList, int32 MaxTilesToProduce);
 
-	void UpdateCSVStats() const;
+	void UpdateResidencyTracking() const;
 
 #if !UE_BUILD_SHIPPING
 	void GetOnScreenMessages(FCoreDelegates::FSeverityMessageMap& OutMessages);
@@ -153,6 +155,9 @@ private:
 	TSet<FVirtualTextureLocalTile> MappedTilesToProduce;
 	TArray<FVirtualTextureLocalTile> TransientCollectedPages;
 	TArray<IVirtualTextureFinalizer*> Finalizers;
+
+	void DrawResidencyHud(class UCanvas*, class APlayerController*);
+	FDelegateHandle	DrawResidencyHudDelegateHandle;
 };
 
 
