Strata more preparation for rough refraction
- Added roughness scatter buffer (not fed properly)
- Added basic scenbe color mip generation (next todo: blur it)
- Mip level is selected based on roughness when applying distortion

shooter game tested

diff --git a/Engine/Shaders/Private/DistortAccumulatePS.usf b/Engine/Shaders/Private/DistortAccumulatePS.usf
--- a/Engine/Shaders/Private/DistortAccumulatePS.usf
+++ b/Engine/Shaders/Private/DistortAccumulatePS.usf
@@ -39,6 +39,10 @@ void Main(
 	in float4 SvPosition	: SV_Position
 	OPTIONAL_IsFrontFace,
 	out float4 OutColor		: SV_Target0
+#if STRATA_ENABLED
+	,
+	out float4 OutRough		: SV_Target1
+#endif
 	)
 {
 #if INSTANCED_STEREO
@@ -124,6 +128,10 @@ void Main(
 
 	// output positives in R|G channels and negatives in B|A channels
 	OutColor = float4(PosOffset.x,PosOffset.y,NegOffset.x,NegOffset.y);
+
+#if STRATA_ENABLED
+	OutRough = saturate( length(abs(PosOffset - NegOffset)) );				// STRATA_TODO use a proper model
+#endif
 }
 
 
diff --git a/Engine/Shaders/Private/DistortApplyScreenPS.usf b/Engine/Shaders/Private/DistortApplyScreenPS.usf
--- a/Engine/Shaders/Private/DistortApplyScreenPS.usf
+++ b/Engine/Shaders/Private/DistortApplyScreenPS.usf
@@ -22,9 +22,12 @@ SamplerState	SceneColorTextureSampler;
  #if USE_MSAA
 Texture2DMS<float4> DistortionMSAATexture;
 Texture2DMS<float4> SceneColorMSAATexture;
+Texture2DMS<float4> RoughnessScatterMSAATexture;
 #else
 Texture2D		DistortionTexture;
 SamplerState	DistortionTextureSampler;
+Texture2D<float4> RoughnessScatterTexture;
+SamplerState      RoughnessScatterSampler;
 #endif
 
 /** distorts screen texture using accumulated distortion offsets */
@@ -50,7 +53,37 @@ void Main(
 	//Scale by the screen size and a fudge factor to come close to the offset values we would have had under normal circumstances before my changes. Also flip Y and invert the precision bias scale.
 	DistBufferUVOffset *= InvDistortionScaleBias;
 	
-	
+#if USE_ROUGH_REFRACTION
+
+#if USE_MSAA
+	half RoughnessScatter = RoughnessScatterMSAATexture.Load(int2(TexCoord.xy * View.BufferSizeAndInvSize.xy), SampleIndex) * 10.0f;
+#else
+	half RoughnessScatter = Texture2DSample(RoughnessScatterTexture, RoughnessScatterSampler, TexCoord.xy) * 10.0f;
+#endif
+
+	float2 NewBufferUV = ((TexCoord.xy + DistBufferUVOffset) * View.BufferSizeAndInvSize.xy - View.ViewRectMin.xy) * View.BufferSizeAndInvSize.zw;
+
+	FLATTEN if (any(NewBufferUV < 0.0f) || any(NewBufferUV > 1.0f))
+	{
+		NewBufferUV = (TexCoord.xy * View.BufferSizeAndInvSize.xy - View.ViewRectMin.xy) * View.BufferSizeAndInvSize.zw;
+	}
+
+#if USE_MSAA
+	if (RoughnessScatter == 0.0f)
+	{
+		// Sample the full res MSAA source texture
+		OutColor = SceneColorMSAATexture.Load(int2(TexCoord * View.BufferSizeAndInvSize.xy), SampleIndex);
+	}
+	else
+	{
+		OutColor = Texture2DSampleLevel(SceneColorTexture, SceneColorTextureSampler, NewBufferUV, RoughnessScatter);
+	}
+#else
+	OutColor = Texture2DSampleLevel(SceneColorTexture, SceneColorTextureSampler, NewBufferUV, RoughnessScatter);
+#endif
+
+#else // USE_ROUGH_REFRACTION
+
 	float2 NewBufferUV = TexCoord.xy + DistBufferUVOffset;
 
 	// If we're about to sample outside the valid BufferUVMinMax, set to 0 distortion.
@@ -66,6 +99,8 @@ void Main(
 #else
 	OutColor = Texture2DSample(SceneColorTexture, SceneColorTextureSampler, NewBufferUV);
 #endif
+
+#endif // USE_ROUGH_REFRACTION
 }
 
 void Merge(
diff --git a/Engine/Shaders/Private/DistortFiltering.usf b/Engine/Shaders/Private/DistortFiltering.usf
new file mode 100644
--- /dev/null
+++ b/Engine/Shaders/Private/DistortFiltering.usf
@@ -0,0 +1,32 @@
+// Copyright Epic Games, Inc. All Rights Reserved.
+
+/*=============================================================================
+	DistortFiltering.usf: Pixel shader for filtering the scene color in preparation for rough refraction
+=============================================================================*/
+
+#include "Common.ush"
+
+
+uint SrcMipIndex;
+int2 SrcMipResolution;
+int2 DstMipResolution;
+SamplerState SourceSampler;
+Texture2D<float4> SourceTexture;
+RWTexture2D<float4> OutTextureMipColor;
+
+
+/** output distortion offsets as color so they can be accumulated (via blending) */
+[numthreads(THREADGROUP_SIZE, THREADGROUP_SIZE, 1)]
+void DownsampleColorCS(uint3 ThreadId : SV_DispatchThreadID)
+{
+	if (any(ThreadId.xy >= uint2(DstMipResolution)))
+	{
+		return;
+	}
+	float2 ScaledUVs = (float2(ThreadId.xy) + 0.5f) / float2(DstMipResolution);
+
+	float4 OutColor = SourceTexture.SampleLevel(SourceSampler, ScaledUVs, SrcMipIndex);
+
+	OutTextureMipColor[uint2(ThreadId.xy)] = OutColor;
+
+}
diff --git a/Engine/Source/Runtime/Renderer/Private/DistortionRendering.cpp b/Engine/Source/Runtime/Renderer/Private/DistortionRendering.cpp
--- a/Engine/Source/Runtime/Renderer/Private/DistortionRendering.cpp
+++ b/Engine/Source/Runtime/Renderer/Private/DistortionRendering.cpp
@@ -13,6 +13,7 @@
 #include "ScenePrivate.h"
 #include "ScreenPass.h"
 #include "MeshPassProcessor.inl"
+#include "Strata/Strata.h"
 
 DECLARE_GPU_STAT(Distortion);
 
@@ -67,18 +68,27 @@ TRDGUniformBufferRef<FDistortionPassUniformParameters> CreateDistortionPassUnifo
 	return GraphBuilder.CreateUniformBuffer(Parameters);
 }
 
+static bool GetUseRoughRefraction()
+{
+	return Strata::IsStrataEnabled();
+}
+
 class FDistortionScreenPS : public FGlobalShader
 {
 public:
 	class FUseMSAADim : SHADER_PERMUTATION_BOOL("USE_MSAA");
-	using FPermutationDomain = TShaderPermutationDomain<FUseMSAADim>;
+	class FUseRoughRefractionDim : SHADER_PERMUTATION_BOOL("USE_ROUGH_REFRACTION");
+	using FPermutationDomain = TShaderPermutationDomain<FUseMSAADim, FUseRoughRefractionDim>;
 
 	BEGIN_SHADER_PARAMETER_STRUCT(FParameters, )
 		SHADER_PARAMETER_STRUCT_REF(FViewUniformShaderParameters, View)
+		SHADER_PARAMETER_RDG_TEXTURE(Texture2DMS<float4>, RoughnessScatterMSAATexture)
 		SHADER_PARAMETER_RDG_TEXTURE(Texture2DMS<float4>, DistortionMSAATexture)
 		SHADER_PARAMETER_RDG_TEXTURE(Texture2DMS<float4>, SceneColorMSAATexture)
+		SHADER_PARAMETER_RDG_TEXTURE(Texture2D, RoughnessScatterTexture)
 		SHADER_PARAMETER_RDG_TEXTURE(Texture2D, DistortionTexture)
 		SHADER_PARAMETER_RDG_TEXTURE(Texture2D, SceneColorTexture)
+		SHADER_PARAMETER_SAMPLER(SamplerState, RoughnessScatterSampler)
 		SHADER_PARAMETER_SAMPLER(SamplerState, DistortionTextureSampler)
 		SHADER_PARAMETER_SAMPLER(SamplerState, SceneColorTextureSampler)
 		RENDER_TARGET_BINDING_SLOTS()
@@ -195,6 +205,35 @@ BEGIN_SHADER_PARAMETER_STRUCT(FDistortionPassParameters, )
 	RENDER_TARGET_BINDING_SLOTS()
 END_SHADER_PARAMETER_STRUCT()
 
+class FDownsampleSceneColorCS : public FGlobalShader
+{
+	DECLARE_GLOBAL_SHADER(FDownsampleSceneColorCS);
+	SHADER_USE_PARAMETER_STRUCT(FDownsampleSceneColorCS, FGlobalShader);
+
+	static const uint32 ThreadGroupSize = 8;
+
+	using FPermutationDomain = TShaderPermutationDomain<>;
+
+	BEGIN_SHADER_PARAMETER_STRUCT(FParameters, )
+		SHADER_PARAMETER(uint32, SrcMipIndex)
+		SHADER_PARAMETER(FIntPoint, SrcMipResolution)
+		SHADER_PARAMETER(FIntPoint, DstMipResolution)
+		SHADER_PARAMETER_RDG_TEXTURE_SRV(Texture2D, SourceTexture)
+		SHADER_PARAMETER_SAMPLER(SamplerState, SourceSampler)
+		SHADER_PARAMETER_RDG_TEXTURE_UAV(RWTexture2D, OutTextureMipColor)
+	END_SHADER_PARAMETER_STRUCT()
+
+public:
+	static bool ShouldCompilePermutation(const FGlobalShaderPermutationParameters& Parameters) { return GetMaxSupportedFeatureLevel(Parameters.Platform) >= ERHIFeatureLevel::SM5; }
+	static void ModifyCompilationEnvironment(const FGlobalShaderPermutationParameters& Parameters, FShaderCompilerEnvironment& OutEnvironment)
+	{
+		FGlobalShader::ModifyCompilationEnvironment(Parameters, OutEnvironment);
+		OutEnvironment.SetDefine(TEXT("THREADGROUP_SIZE"), ThreadGroupSize);
+		OutEnvironment.SetDefine(TEXT("USE_COMPUTE"), 1);
+	}
+};
+IMPLEMENT_GLOBAL_SHADER(FDownsampleSceneColorCS, "/Engine/Private/DistortFiltering.usf", "DownsampleColorCS", SF_Compute);
+
 void FDeferredShadingSceneRenderer::RenderDistortion(FRDGBuilder& GraphBuilder, FRDGTextureRef SceneColorTexture, FRDGTextureRef SceneDepthTexture)
 {
 	check(SceneDepthTexture);
@@ -213,8 +252,84 @@ void FDeferredShadingSceneRenderer::RenderDistortion(FRDGBuilder& GraphBuilder,
 	FDepthStencilBinding StencilWriteBinding(SceneDepthTexture, ERenderTargetLoadAction::ELoad, ERenderTargetLoadAction::EClear, FExclusiveDepthStencil::DepthRead_StencilWrite);
 
 	FRDGTextureRef DistortionTexture = nullptr;
+	FRDGTextureRef RoughnessScatterTexture = nullptr;
 	FGlobalShaderMap* ShaderMap = GetGlobalShaderMap(FeatureLevel);
 
+	// Create a texture of the scene color with MIP levels, unfiltered
+	// NOTE: we cannot use mips on the scene color render target it self because multiple views are continuous in there and they would leak color one onto another.
+	TArray<FRDGTextureRef> ViewsSceneColorMipchain;
+	ViewsSceneColorMipchain.SetNum(Views.Num());
+	const bool bUseRoughRefraction = GetUseRoughRefraction();
+	if(bUseRoughRefraction)
+	{
+		// We keep the same format, for AddCopyTexturePass to not complain for now. Do something special to rely on PF_FloatR11G11B10 only.
+		EPixelFormat SceneColorFormat = SceneColorTexture->Desc.Format;	
+
+		for (int32 ViewIndex = 0, Num = Views.Num(); ViewIndex < Num; ++ViewIndex)
+		{
+			// Create the mip chain for the scene color in a separated texture
+
+			const FViewInfo& View = Views[ViewIndex];
+			const FIntPoint SceneColorMip0Resolution = View.ViewRect.Size();
+
+			// We do not use max and do not add 1 to the result, this to not go down to the lowest mip level as this is not required.
+			const int32 MipCount = FMath::Max((int32)1, (int32)FMath::CeilLogTwo(FMath::Min(SceneColorMip0Resolution.X, SceneColorMip0Resolution.Y)));
+
+			FRDGTextureDesc SceneColorMipchainDesc = FRDGTextureDesc::Create2D(SceneColorMip0Resolution, SceneColorFormat, FClearValueBinding::None,
+				TexCreate_TargetArraySlicesIndependently | TexCreate_ShaderResource | TexCreate_UAV | TexCreate_RenderTargetable, MipCount);
+
+			FRDGTextureRef SceneColorMipchainTexture = GraphBuilder.CreateTexture(SceneColorMipchainDesc, TEXT("SceneColorMipchain"));
+			ViewsSceneColorMipchain[ViewIndex] = SceneColorMipchainTexture;
+
+			// Copy scene color into the first mip level
+
+			FRHICopyTextureInfo CopyInfo;
+			CopyInfo.SourcePosition.X = View.ViewRect.Min.X;
+			CopyInfo.SourcePosition.Y = View.ViewRect.Min.Y;
+			CopyInfo.Size.X = SceneColorMip0Resolution.X;
+			CopyInfo.Size.Y = SceneColorMip0Resolution.Y;
+			CopyInfo.Size.Z = 1;
+			CopyInfo.NumMips = 1;
+			AddCopyTexturePass(GraphBuilder, SceneColorTexture, SceneColorMipchainTexture, CopyInfo);
+
+			// Now render the mip chain
+
+			FDownsampleSceneColorCS::FPermutationDomain PermutationVector;
+			TShaderMapRef<FDownsampleSceneColorCS> ComputeShader(GetGlobalShaderMap(FeatureLevel), PermutationVector);
+
+			FIntPoint SrcMipResolution = SceneColorMip0Resolution;
+			FIntPoint DstMipResolution = SceneColorMip0Resolution / 2;
+			for (int32 DstMipIndex = 1; DstMipIndex < MipCount; DstMipIndex++)
+			{
+				FDownsampleSceneColorCS::FParameters* PassParameters = GraphBuilder.AllocParameters<FDownsampleSceneColorCS::FParameters>();
+				PassParameters->SrcMipIndex = DstMipIndex - 1;
+				PassParameters->SrcMipResolution = SrcMipResolution;
+				PassParameters->DstMipResolution = DstMipResolution;
+				PassParameters->SourceSampler = TStaticSamplerState<SF_Bilinear>::GetRHI();	// To average 4 texel with one sample
+
+				PassParameters->SourceTexture = GraphBuilder.CreateSRV(FRDGTextureSRVDesc::CreateForMipLevel(SceneColorMipchainTexture, DstMipIndex - 1));
+				PassParameters->OutTextureMipColor = GraphBuilder.CreateUAV(FRDGTextureUAVDesc(SceneColorMipchainTexture, DstMipIndex));
+
+				FIntVector NumGroups = FIntVector::DivideAndRoundUp(FIntVector(DstMipResolution.X, DstMipResolution.Y, 1), FIntVector(FDownsampleSceneColorCS::ThreadGroupSize, FDownsampleSceneColorCS::ThreadGroupSize, 1));
+
+				// Dispatch with GenerateMips: reading from a slice through SRV and writing into lower mip through UAV.
+				ClearUnusedGraphResources(ComputeShader, PassParameters);
+				GraphBuilder.AddPass(
+					Forward<FRDGEventName>(RDG_EVENT_NAME("DistoMipGen")),
+					PassParameters,
+					ERDGPassFlags::Compute | ERDGPassFlags::NeverCull,
+					[PassParameters, ComputeShader, NumGroups](FRHICommandList& RHICmdList)
+					{
+						FComputeShaderUtils::Dispatch(RHICmdList, ComputeShader, *PassParameters, NumGroups);
+					});
+
+				SrcMipResolution = DstMipResolution;
+				DstMipResolution = DstMipResolution / 2;
+			}
+		}
+	}
+
+
 	// Use stencil mask to optimize cases with lower screen coverage.
 	// Note: This adds an extra pass which is actually slower as distortion tends towards full-screen.
 	//       It could be worth testing object screen bounds then reverting to a target flip and single pass.
@@ -239,6 +354,20 @@ void FDeferredShadingSceneRenderer::RenderDistortion(FRDGBuilder& GraphBuilder,
 				SceneDepthTexture->Desc.NumSamples),
 			TEXT("Distortion"));
 
+		if (bUseRoughRefraction)
+		{
+			// This is the texture containing information about the surface back scattering process
+			RoughnessScatterTexture = GraphBuilder.CreateTexture(
+				FRDGTextureDesc::Create2D(
+					SceneDepthTexture->Desc.Extent,
+					PF_G8,
+					FClearValueBinding::Transparent,
+					GFastVRamConfig.Distortion | TexCreate_RenderTargetable | TexCreate_ShaderResource,
+					1,
+					SceneDepthTexture->Desc.NumSamples),
+				TEXT("DistortionRoughnessScatter"));
+		}
+
 		ERenderTargetLoadAction LoadAction = ERenderTargetLoadAction::EClear;
 
 		for (int32 ViewIndex = 0; ViewIndex < Views.Num(); ViewIndex++)
@@ -259,6 +388,10 @@ void FDeferredShadingSceneRenderer::RenderDistortion(FRDGBuilder& GraphBuilder,
 			PassParameters->View = View.GetShaderParameters();
 			PassParameters->Pass = CreateDistortionPassUniformBuffer(GraphBuilder, View);
 			PassParameters->RenderTargets[0] = FRenderTargetBinding(DistortionTexture, LoadAction);
+			if (bUseRoughRefraction)
+			{
+				PassParameters->RenderTargets[1] = FRenderTargetBinding(RoughnessScatterTexture, LoadAction);
+			}
 			PassParameters->RenderTargets.DepthStencil = StencilWriteBinding;
 
 			View.ParallelMeshDrawCommandPasses[EMeshPass::Distortion].BuildRenderingCommands(GraphBuilder, Scene->GPUScene, PassParameters->InstanceCullingDrawParams);
@@ -288,11 +421,18 @@ void FDeferredShadingSceneRenderer::RenderDistortion(FRDGBuilder& GraphBuilder,
 	FDistortionScreenPS::FParameters CommonParameters;
 	CommonParameters.DistortionMSAATexture = DistortionTexture;
 	CommonParameters.DistortionTexture = DistortionTexture;
-	CommonParameters.SceneColorTextureSampler = TStaticSamplerState<>::GetRHI();
+	if (bUseRoughRefraction)
+	{
+		CommonParameters.RoughnessScatterMSAATexture = RoughnessScatterTexture;
+		CommonParameters.RoughnessScatterTexture = RoughnessScatterTexture;
+	}
+	CommonParameters.SceneColorTextureSampler = bUseRoughRefraction ? TStaticSamplerState<SF_Trilinear>::GetRHI() : TStaticSamplerState<>::GetRHI();
 	CommonParameters.DistortionTextureSampler = TStaticSamplerState<>::GetRHI();
+	CommonParameters.RoughnessScatterSampler = TStaticSamplerState<>::GetRHI();
 
 	FDistortionScreenPS::FPermutationDomain PermutationVector;
 	PermutationVector.Set<FDistortionScreenPS::FUseMSAADim>(SceneColorTexture->Desc.NumSamples > 1);
+	PermutationVector.Set<FDistortionScreenPS::FUseRoughRefractionDim>(bUseRoughRefraction);
 
 	TShaderMapRef<FScreenPassVS> VertexShader(ShaderMap);
 	TShaderMapRef<FDistortionApplyScreenPS> ApplyPixelShader(ShaderMap, PermutationVector);
@@ -327,6 +467,11 @@ void FDeferredShadingSceneRenderer::RenderDistortion(FRDGBuilder& GraphBuilder,
 
 			auto* PassParameters = GraphBuilder.AllocParameters<FDistortionScreenPS::FParameters>();
 			*PassParameters = CommonParameters;
+			if (bUseRoughRefraction)
+			{
+				PassParameters->SceneColorMSAATexture = ViewsSceneColorMipchain[ViewIndex];
+				PassParameters->SceneColorTexture = ViewsSceneColorMipchain[ViewIndex];
+			}
 			PassParameters->View = View.ViewUniformBuffer;
 			PassParameters->RenderTargets[0] = FRenderTargetBinding(DistortionSceneColorTexture, LoadAction);
 
@@ -520,8 +665,16 @@ FMeshPassProcessor* CreateDistortionPassProcessor(const FScene* Scene, const FSc
 
 	DistortionPassState.SetStencilRef(kStencilMaskBit);
 
-	// additive blending of offsets (or complexity if the shader complexity viewmode is enabled)
-	DistortionPassState.SetBlendState(TStaticBlendState<CW_RGBA, BO_Add, BF_One, BF_One, BO_Add, BF_One, BF_One>::GetRHI());
+	if (GetUseRoughRefraction())
+	{
+		DistortionPassState.SetBlendState(TStaticBlendState<CW_RGBA, BO_Add, BF_One, BF_One, BO_Add, BF_One, BF_One,
+															CW_RGBA, BO_Add, BF_One, BF_One, BO_Add, BF_One, BF_One>::GetRHI());
+	}
+	else
+	{
+		// additive blending of offsets (or complexity if the shader complexity viewmode is enabled)
+		DistortionPassState.SetBlendState(TStaticBlendState<CW_RGBA, BO_Add, BF_One, BF_One, BO_Add, BF_One, BF_One>::GetRHI());
+	}
 
 	return new(FMemStack::Get()) FDistortionMeshProcessor(Scene, InViewIfDynamicMeshCommand, DistortionPassState, InDrawListContext);
 }
@@ -532,8 +685,17 @@ FMeshPassProcessor* CreateMobileDistortionPassProcessor(const FScene* Scene, con
 
 	// We don't have depth, render all pixels, pixel shader will sample SceneDepth from SceneColor.A and discard if occluded
 	DistortionPassState.SetDepthStencilState(TStaticDepthStencilState<false, CF_Always>::GetRHI());
-	// additive blending of offsets
-	DistortionPassState.SetBlendState(TStaticBlendState<CW_RGBA, BO_Add, BF_One, BF_One, BO_Add, BF_One, BF_One>::GetRHI());
+
+	if (GetUseRoughRefraction())
+	{
+		DistortionPassState.SetBlendState(TStaticBlendState<CW_RGBA, BO_Add, BF_One, BF_One, BO_Add, BF_One, BF_One,
+															CW_RGBA, BO_Add, BF_One, BF_One, BO_Add, BF_One, BF_One>::GetRHI());
+	}
+	else
+	{
+		// additive blending of offsets
+		DistortionPassState.SetBlendState(TStaticBlendState<CW_RGBA, BO_Add, BF_One, BF_One, BO_Add, BF_One, BF_One>::GetRHI());
+	}
 
 	return new(FMemStack::Get()) FDistortionMeshProcessor(Scene, InViewIfDynamicMeshCommand, DistortionPassState, InDrawListContext);
 }
