Removed ES 3.1 redundant feature level checks

diff --git a/Engine/Source/Runtime/Engine/Private/Particles/ParticleSortingGPU.cpp b/Engine/Source/Runtime/Engine/Private/Particles/ParticleSortingGPU.cpp
--- a/Engine/Source/Runtime/Engine/Private/Particles/ParticleSortingGPU.cpp
+++ b/Engine/Source/Runtime/Engine/Private/Particles/ParticleSortingGPU.cpp
@@ -229,22 +229,19 @@ int32 GenerateParticleSortKeys(
  */
 void FParticleSortBuffers::InitRHI()
 {
-	if (GetFeatureLevel() >= ERHIFeatureLevel::ES3_1)
+	for (int32 BufferIndex = 0; BufferIndex < 2; ++BufferIndex)
 	{
-		for (int32 BufferIndex = 0; BufferIndex < 2; ++BufferIndex)
-		{
-			FRHIResourceCreateInfo CreateInfo(TEXT("PartialSortKeyBuffer"));
+		FRHIResourceCreateInfo CreateInfo(TEXT("PartialSortKeyBuffer"));
 
-			KeyBuffers[BufferIndex] = RHICreateVertexBuffer( BufferSize * sizeof(uint32), BUF_Static | BUF_ShaderResource | BUF_UnorderedAccess, CreateInfo);
-			KeyBufferSRVs[BufferIndex] = RHICreateShaderResourceView( KeyBuffers[BufferIndex], /*Stride=*/ sizeof(uint32), PF_R32_UINT );
-			KeyBufferUAVs[BufferIndex] = RHICreateUnorderedAccessView( KeyBuffers[BufferIndex], PF_R32_UINT );
+		KeyBuffers[BufferIndex] = RHICreateVertexBuffer( BufferSize * sizeof(uint32), BUF_Static | BUF_ShaderResource | BUF_UnorderedAccess, CreateInfo);
+		KeyBufferSRVs[BufferIndex] = RHICreateShaderResourceView( KeyBuffers[BufferIndex], /*Stride=*/ sizeof(uint32), PF_R32_UINT );
+		KeyBufferUAVs[BufferIndex] = RHICreateUnorderedAccessView( KeyBuffers[BufferIndex], PF_R32_UINT );
 
-			CreateInfo.DebugName = TEXT("PartialSortVertexBuffer");
-			VertexBuffers[BufferIndex] = RHICreateVertexBuffer( BufferSize * sizeof(uint32), BUF_Static | BUF_ShaderResource | BUF_UnorderedAccess, CreateInfo);
+		CreateInfo.DebugName = TEXT("PartialSortVertexBuffer");
+		VertexBuffers[BufferIndex] = RHICreateVertexBuffer( BufferSize * sizeof(uint32), BUF_Static | BUF_ShaderResource | BUF_UnorderedAccess, CreateInfo);
 
-			VertexBufferSortSRVs[BufferIndex] = RHICreateShaderResourceView(VertexBuffers[BufferIndex], /*Stride=*/ sizeof(uint32), PF_R32_UINT);
-			VertexBufferSortUAVs[BufferIndex] = RHICreateUnorderedAccessView(VertexBuffers[BufferIndex], PF_R32_UINT);
-		}
+		VertexBufferSortSRVs[BufferIndex] = RHICreateShaderResourceView(VertexBuffers[BufferIndex], /*Stride=*/ sizeof(uint32), PF_R32_UINT);
+		VertexBufferSortUAVs[BufferIndex] = RHICreateUnorderedAccessView(VertexBuffers[BufferIndex], PF_R32_UINT);
 	}
 }
 
diff --git a/Engine/Source/Runtime/RenderCore/Private/ShaderParameters.cpp b/Engine/Source/Runtime/RenderCore/Private/ShaderParameters.cpp
--- a/Engine/Source/Runtime/RenderCore/Private/ShaderParameters.cpp
+++ b/Engine/Source/Runtime/RenderCore/Private/ShaderParameters.cpp
@@ -383,12 +383,7 @@ void FShaderType::AddReferencedUniformBufferIncludes(FShaderCompilerEnvironment&
 
 	FString& GeneratedUniformBuffersInclude = OutEnvironment.IncludeVirtualPathToContentsMap.FindOrAdd("/Engine/Generated/GeneratedUniformBuffers.ush");
 	GeneratedUniformBuffersInclude.Append(UniformBufferIncludes);
-
-	ERHIFeatureLevel::Type MaxFeatureLevel = GetMaxSupportedFeatureLevel(Platform);
-	if (MaxFeatureLevel >= ERHIFeatureLevel::ES3_1)
-	{
-		OutEnvironment.SetDefine(TEXT("PLATFORM_SUPPORTS_SRV_UB"), TEXT("1"));
-	}
+	OutEnvironment.SetDefine(TEXT("PLATFORM_SUPPORTS_SRV_UB"), TEXT("1"));
 }
 
 void FShaderType::DumpDebugInfo()
@@ -505,10 +500,5 @@ void FVertexFactoryType::AddReferencedUniformBufferIncludes(FShaderCompilerEnvir
 
 	FString& GeneratedUniformBuffersInclude = OutEnvironment.IncludeVirtualPathToContentsMap.FindOrAdd("/Engine/Generated/GeneratedUniformBuffers.ush");
 	GeneratedUniformBuffersInclude.Append(UniformBufferIncludes);
-
-	ERHIFeatureLevel::Type MaxFeatureLevel = GetMaxSupportedFeatureLevel(Platform);
-	if (MaxFeatureLevel >= ERHIFeatureLevel::ES3_1)
-	{
-		OutEnvironment.SetDefine(TEXT("PLATFORM_SUPPORTS_SRV_UB"), TEXT("1"));
-	}
+	OutEnvironment.SetDefine(TEXT("PLATFORM_SUPPORTS_SRV_UB"), TEXT("1"));
 }
diff --git a/Engine/Source/Runtime/Renderer/Private/LightMapRendering.cpp b/Engine/Source/Runtime/Renderer/Private/LightMapRendering.cpp
--- a/Engine/Source/Runtime/Renderer/Private/LightMapRendering.cpp
+++ b/Engine/Source/Runtime/Renderer/Private/LightMapRendering.cpp
@@ -376,7 +376,6 @@ void GetIndirectLightingCacheParameters(
 			Parameters.IndirectLightingCacheTextureSampler2 = TStaticSamplerState<SF_Bilinear,AM_Clamp,AM_Clamp,AM_Clamp>::GetRHI();
 		}
 		else
-		if (FeatureLevel >= ERHIFeatureLevel::ES3_1)
 		{
 			Parameters.IndirectLightingCacheTexture0 = GBlackVolumeTexture->TextureRHI;
 			Parameters.IndirectLightingCacheTexture1 = GBlackVolumeTexture->TextureRHI;
@@ -386,16 +385,6 @@ void GetIndirectLightingCacheParameters(
 			Parameters.IndirectLightingCacheTextureSampler1 = GBlackVolumeTexture->SamplerStateRHI;
 			Parameters.IndirectLightingCacheTextureSampler2 = GBlackVolumeTexture->SamplerStateRHI;
 		}
-		else
-		{
-			Parameters.IndirectLightingCacheTexture0 = GBlackTexture->TextureRHI;
-			Parameters.IndirectLightingCacheTexture1 = GBlackTexture->TextureRHI;
-			Parameters.IndirectLightingCacheTexture2 = GBlackTexture->TextureRHI;
-
-			Parameters.IndirectLightingCacheTextureSampler0 = GBlackTexture->SamplerStateRHI;
-			Parameters.IndirectLightingCacheTextureSampler1 = GBlackTexture->SamplerStateRHI;
-			Parameters.IndirectLightingCacheTextureSampler2 = GBlackTexture->SamplerStateRHI;
-		}
 	}
 }
 
diff --git a/Engine/Source/Runtime/Renderer/Private/PostProcess/PostProcessEyeAdaptation.cpp b/Engine/Source/Runtime/Renderer/Private/PostProcess/PostProcessEyeAdaptation.cpp
--- a/Engine/Source/Runtime/Renderer/Private/PostProcess/PostProcessEyeAdaptation.cpp
+++ b/Engine/Source/Runtime/Renderer/Private/PostProcess/PostProcessEyeAdaptation.cpp
@@ -76,7 +76,7 @@ bool IsAutoExposureMethodSupported(ERHIFeatureLevel::Type FeatureLevel, EAutoExp
 	case EAutoExposureMethod::AEM_Histogram:
 	case EAutoExposureMethod::AEM_Basic:
 	case EAutoExposureMethod::AEM_Manual:
-		return FeatureLevel >= ERHIFeatureLevel::ES3_1;
+		return true;
 	}
 	return false;
 }
diff --git a/Engine/Source/Runtime/Renderer/Private/SceneCaptureRendering.cpp b/Engine/Source/Runtime/Renderer/Private/SceneCaptureRendering.cpp
--- a/Engine/Source/Runtime/Renderer/Private/SceneCaptureRendering.cpp
+++ b/Engine/Source/Runtime/Renderer/Private/SceneCaptureRendering.cpp
@@ -854,7 +854,7 @@ void FScene::UpdateSceneCaptureContents(USceneCaptureComponentCube* CaptureCompo
 	{
 		UTextureRenderTargetCube* const TextureTarget = TextureTargets[CaptureIter];
 
-		if (GetFeatureLevel() >= ERHIFeatureLevel::ES3_1 && TextureTarget)
+		if (TextureTarget)
 		{
 			const float FOV = 90 * (float)PI / 360.0f;
 			for (int32 faceidx = 0; faceidx < (int32)ECubeFace::CubeFace_MAX; faceidx++)
diff --git a/Engine/Source/Runtime/Renderer/Private/SystemTextures.cpp b/Engine/Source/Runtime/Renderer/Private/SystemTextures.cpp
--- a/Engine/Source/Runtime/Renderer/Private/SystemTextures.cpp
+++ b/Engine/Source/Runtime/Renderer/Private/SystemTextures.cpp
@@ -226,7 +226,7 @@ void FSystemTextures::InitializeFeatureLevelDependentTextures(FRHICommandListImm
     ERHIFeatureLevel::Type CurrentFeatureLevel = FeatureLevelInitializedTo == ERHIFeatureLevel::Num ? ERHIFeatureLevel::ES2_REMOVED : FeatureLevelInitializedTo;
 
 		// Create the SobolSampling texture
-	if (CurrentFeatureLevel < ERHIFeatureLevel::ES3_1 && InFeatureLevel >= ERHIFeatureLevel::ES3_1 && GPixelFormats[PF_R16_UINT].Supported)
+	if (CurrentFeatureLevel < ERHIFeatureLevel::ES3_1 && GPixelFormats[PF_R16_UINT].Supported)
 	{
 		FRHIResourceCreateInfo CreateInfo(TEXT("SobolSampling"));
 		FTexture2DRHIRef Texture2D = RHICreateTexture2D(32, 16, PF_R16_UINT, 1, 1, TexCreate_ShaderResource, CreateInfo);
@@ -639,7 +639,7 @@ void FSystemTextures::InitializeFeatureLevelDependentTextures(FRHICommandListImm
 	// Create the SSAO randomization texture
 	static const auto MobileAmbientOcclusionCVar = IConsoleManager::Get().FindTConsoleVariableDataInt(TEXT("r.Mobile.AmbientOcclusion"));
 	if ((CurrentFeatureLevel < ERHIFeatureLevel::SM5 && InFeatureLevel >= ERHIFeatureLevel::SM5) ||
-		(CurrentFeatureLevel < ERHIFeatureLevel::ES3_1 && InFeatureLevel >= ERHIFeatureLevel::ES3_1 && MobileAmbientOcclusionCVar != nullptr && MobileAmbientOcclusionCVar->GetValueOnAnyThread()>0))
+		(CurrentFeatureLevel < ERHIFeatureLevel::ES3_1 && MobileAmbientOcclusionCVar != nullptr && MobileAmbientOcclusionCVar->GetValueOnAnyThread()>0))
 	{
 		{
 			float g_AngleOff1 = 127;
@@ -698,7 +698,7 @@ void FSystemTextures::InitializeFeatureLevelDependentTextures(FRHICommandListImm
 		
 	static const auto MobileGTAOPreIntegratedTextureTypeCVar = IConsoleManager::Get().FindTConsoleVariableDataInt(TEXT("r.Mobile.GTAOPreIntegratedTextureType"));
 
-	if (CurrentFeatureLevel < ERHIFeatureLevel::ES3_1 && InFeatureLevel >= ERHIFeatureLevel::ES3_1 && MobileGTAOPreIntegratedTextureTypeCVar && MobileGTAOPreIntegratedTextureTypeCVar->GetValueOnAnyThread() > 0)
+	if (CurrentFeatureLevel < ERHIFeatureLevel::ES3_1 && MobileGTAOPreIntegratedTextureTypeCVar && MobileGTAOPreIntegratedTextureTypeCVar->GetValueOnAnyThread() > 0)
 	{
 		uint32 Extent = 16; // should be consistent with LUTSize in PostprocessMobile.usf
 
