Saves over 100ms (RT) on console (~10ms on PC) in Test target in a large world using HLODs in FVirtualTextureProducerCollection::RemoveAllCallbacks.

diff --git a/Engine/Source/Runtime/Renderer/Private/VT/VirtualTextureProducer.cpp b/Engine/Source/Runtime/Renderer/Private/VT/VirtualTextureProducer.cpp
--- a/Engine/Source/Runtime/Renderer/Private/VT/VirtualTextureProducer.cpp
+++ b/Engine/Source/Runtime/Renderer/Private/VT/VirtualTextureProducer.cpp
@@ -129,6 +129,15 @@ void FVirtualTextureProducerCollection::CallPendingCallbacks()
 		// Make a copy, then release the callback entry before calling the callback function
 		// (The destroyed callback may try to remove this or other callbacks, so need to make sure state is valid before calling)
 		const FCallbackEntry CallbackCopy(Callback);
+		if (Callback.Baton)
+		{
+			TArray<uint32>& BatonCallbacks = CallbacksMap.FindChecked(Callback.Baton);
+			BatonCallbacks.Remove(CallbackIndex);
+			if (BatonCallbacks.IsEmpty())
+			{
+				CallbacksMap.Remove(Callback.Baton);
+			}
+		}
 		Callback.DestroyedFunction = nullptr;
 		Callback.Baton = nullptr;
 		Callback.OwnerHandle = FVirtualTextureProducerHandle();
@@ -156,6 +165,7 @@ void FVirtualTextureProducerCollection::AddDestroyedCallback(const FVirtualTextu
 {
 	check(IsInRenderingThread());
 	check(Function);
+	check(Baton);
 
 	FProducerEntry* Entry = GetEntry(Handle);
 	if (Entry)
@@ -167,6 +177,7 @@ void FVirtualTextureProducerCollection::AddDestroyedCallback(const FVirtualTextu
 		Callback.Baton = Baton;
 		Callback.OwnerHandle = Handle;
 		Callback.PackedFlags = 0u;
+		CallbacksMap.FindOrAdd(Baton).Add(CallbackIndex);
 	}
 }
 
@@ -176,26 +187,25 @@ uint32 FVirtualTextureProducerCollection::RemoveAllCallbacks(const void* Baton)
 	check(Baton);
 
 	uint32 NumRemoved = 0u;
-	for (int32 CallbackIndex = CallbackList_Count; CallbackIndex < Callbacks.Num(); ++CallbackIndex)
+	for (int32 CallbackIndex : CallbacksMap.FindChecked(Baton))
 	{
 		FCallbackEntry& Callback = Callbacks[CallbackIndex];
-		if (Callback.Baton == Baton)
+		check(Callback.Baton == Baton);
+		check(Callback.DestroyedFunction);
+		Callback.DestroyedFunction = nullptr;
+		Callback.Baton = nullptr;
+		Callback.OwnerHandle = FVirtualTextureProducerHandle();
+
+		// If callback is already pending, we can't move it back to free list, or we risk corrupting the pending list while it's being iterated
+		// Setting DestroyedFunction to nullptr here will ensure callback is no longer invoked, and it will be moved to free list later when it's removed from pending list
+		if (!Callback.bPending)
 		{
-			check(Callback.DestroyedFunction);
-			Callback.DestroyedFunction = nullptr;
-			Callback.Baton = nullptr;
-			Callback.OwnerHandle = FVirtualTextureProducerHandle();
-
-			// If callback is already pending, we can't move it back to free list, or we risk corrupting the pending list while it's being iterated
-			// Setting DestroyedFunction to nullptr here will ensure callback is no longer invoked, and it will be moved to free list later when it's removed from pending list
-			if (!Callback.bPending)
-			{
-				Callback.PackedFlags = 0u;
-				ReleaseCallback(CallbackIndex);
-			}
-			++NumRemoved;
+			Callback.PackedFlags = 0u;
+			ReleaseCallback(CallbackIndex);
 		}
+		++NumRemoved;
 	}
+	CallbacksMap.Remove(Baton);
 	return NumRemoved;
 }
 
diff --git a/Engine/Source/Runtime/Renderer/Private/VT/VirtualTextureProducer.h b/Engine/Source/Runtime/Renderer/Private/VT/VirtualTextureProducer.h
--- a/Engine/Source/Runtime/Renderer/Private/VT/VirtualTextureProducer.h
+++ b/Engine/Source/Runtime/Renderer/Private/VT/VirtualTextureProducer.h
@@ -203,5 +203,6 @@ private:
 
 	TArray<FProducerEntry> Producers;
 	TArray<FCallbackEntry> Callbacks;
+	TMap<void*, TArray<uint32>> CallbacksMap;
 	uint32 NumPendingCallbacks;
 };
