Move InstanceCullingLoadBalancer into a sparate set of files

diff --git a/Engine/Shaders/Private/GPUScene/GPUSceneDataManagement.usf b/Engine/Shaders/Private/GPUScene/GPUSceneDataManagement.usf
--- a/Engine/Shaders/Private/GPUScene/GPUSceneDataManagement.usf
+++ b/Engine/Shaders/Private/GPUScene/GPUSceneDataManagement.usf
@@ -2,211 +2,22 @@
 
 #include "../Common.ush"
 #include "GPUSceneWriter.ush"
-
-
-#if COMPILER_SUPPORTS_WAVE_VOTE
-	// TODO: Need to also check that wave width == thread group size (or implement workaround for other wave sizes)
-	#define USE_WAVE_BIT_PREFIX_WORK_DISTRIBUTION 0
-#endif
-
-struct FPackedInstanceBatch
-{
-	uint FirstItem_NumItems;
-};
-
-struct FPackedInstanceBatchItem
-{
-	// packed 32-NumInstancesItemBits:NumInstancesItemBits - need one more bit for the case where one item has ThreadGroupSize work to do
-	uint InstanceSceneDataOffset_NumInstances;
-	// packed 32-PrefixBits:PrefixBits
-	uint Payload_BatchPrefixOffset;
-};
-
-
-StructuredBuffer<FPackedInstanceBatch> BatchBuffer;
-StructuredBuffer<FPackedInstanceBatchItem> ItemBuffer;
-uint NumBatches;
-uint NumItems;
-
-
-struct FInstanceBatch
-{
-	uint FirstItem;
-	uint NumItems;
-};
-
-FInstanceBatch UnpackBatch(FPackedInstanceBatch PackedBatch)
-{
-	FInstanceBatch Result;
-	Result.FirstItem = PackedBatch.FirstItem_NumItems >> NUM_INSTANCES_ITEM_BITS;
-	Result.NumItems = PackedBatch.FirstItem_NumItems & NUM_INSTANCES_ITEM_MASK;
-
-	return Result;
-};
-
-struct FInstanceBatchItem
-{
-	uint InstanceSceneDataOffset;
-	uint NumInstances;
-	uint Payload;
-	uint BatchPrefixOffset;
-};
-
-FInstanceBatchItem UnpackItem(FPackedInstanceBatchItem PackedItem)
-{
-	FInstanceBatchItem Result;
-
-	Result.InstanceSceneDataOffset = PackedItem.InstanceSceneDataOffset_NumInstances >> NUM_INSTANCES_ITEM_BITS;
-	Result.NumInstances = PackedItem.InstanceSceneDataOffset_NumInstances & NUM_INSTANCES_ITEM_MASK;
-	Result.Payload = PackedItem.Payload_BatchPrefixOffset >> PREFIX_BITS;
-	Result.BatchPrefixOffset = PackedItem.Payload_BatchPrefixOffset & PREFIX_BIT_MASK;
-
-	return Result;
-};
-
-
-#if USE_WAVE_BIT_PREFIX_WORK_DISTRIBUTION
-
-groupshared uint WorkBoundary[NUM_THREADS_PER_GROUP];
-
-#else
-
-groupshared uint ItemIndex[NUM_THREADS_PER_GROUP];
-
-void PrefixMaxStep(uint ElementIndex, uint StepSize, inout uint CurrentValue)
-{
-	if (ElementIndex >= StepSize)
-	{
-		CurrentValue = max(CurrentValue, ItemIndex[ElementIndex - StepSize]);
-	}
-	GroupMemoryBarrierWithGroupSync();
-
-	if (ElementIndex >= StepSize)
-	{
-		ItemIndex[ElementIndex] = CurrentValue;
-	}
-	GroupMemoryBarrierWithGroupSync();
-}
-
-uint PrefixMaxLastStep(uint ElementIndex, inout uint CurrentValue)
-{
-	GroupMemoryBarrierWithGroupSync();
-
-	uint StepSize = NUM_THREADS_PER_GROUP / 2U;
-
-	if (ElementIndex >= StepSize)
-	{
-		uint NewValue = max(CurrentValue, ItemIndex[ElementIndex - StepSize]);
-		CurrentValue = NewValue;
-	}
-
-	// No sync needed since we don't look at ItemIndex[] again
-	return CurrentValue;
-}
-
-uint PrefixMax(uint ElementIndex)
-{
-	GroupMemoryBarrierWithGroupSync();
-
-	uint CurrentValue = ItemIndex[ElementIndex];
-	GroupMemoryBarrierWithGroupSync();
-
-	UNROLL
-	for (uint StepSize = 1U; StepSize < NUM_THREADS_PER_GROUP / 2U; ++StepSize)
-	{
-		PrefixMaxStep(ElementIndex, 1U, CurrentValue);
-	}
-	return PrefixMaxLastStep(ElementIndex, CurrentValue);
-}
-
-#endif // USE_WAVE_BIT_PREFIX_WORK_DISTRIBUTION
-
-
-groupshared FInstanceBatchItem Items[NUM_THREADS_PER_GROUP];
-
-// TODO: Move this into own file and namespace things a bit.
-FInstanceBatchItem InstanceProcessingGPULoadBalancer_Setup(uint GroupId, int GroupThreadIndex)
-{
-	uint BatchIndex = GroupId;
-
-	// Load work description for whole work group
-	FInstanceBatch GroupBatch = UnpackBatch(BatchBuffer[GroupId]);
-
-	// Special case where we have an 1:1 item/thread no need to do prefix sums or whatever
-	// TODO: this case can also store less data per item
-	if (GroupBatch.NumItems == NUM_THREADS_PER_GROUP)
-	{
-		checkSlow(UnpackItem(ItemBuffer[GroupBatch.FirstItem + GroupThreadIndex]).BatchPrefixOffset == GroupThreadIndex);
-
-		return UnpackItem(ItemBuffer[GroupBatch.FirstItem + GroupThreadIndex]);
-	}
-
-	// Load work Items (guaranteed to be less than group size)
-#if USE_WAVE_BIT_PREFIX_WORK_DISTRIBUTION
-	WorkBoundary[GroupThreadIndex] = 0;
-#else 
-	ItemIndex[GroupThreadIndex] = 0U;
-#endif
-
-	GroupMemoryBarrierWithGroupSync();
-
-	if (GroupThreadIndex < GroupBatch.NumItems)
-	{
-		FInstanceBatchItem WorkItem = UnpackItem(ItemBuffer[GroupBatch.FirstItem + GroupThreadIndex]);
-#if USE_WAVE_BIT_PREFIX_WORK_DISTRIBUTION
-		// Mark last instance as boundary
-		WorkBoundary[(WorkItem.BatchPrefixOffset + WorkItem.NumInstances) - 1U] = 1;
-#else 
-		ItemIndex[WorkItem.BatchPrefixOffset] = GroupThreadIndex;
-#endif
-		Items[GroupThreadIndex] = WorkItem;
-	}
-	GroupMemoryBarrierWithGroupSync();
-	uint WorkItemIndex = 0U;
-	if (GroupBatch.NumItems < 4U)
-	{
-		// Distribute using linear search (no syncs requires like the prefix max op)
-		for (int Index = int(GroupBatch.NumItems) - 1; Index >= 0; --Index)
-		{
-			// TODO: pack the offset data into a single dword and store in group batch? 6x5 bits can fit and would be far more efficient that all the shared mumbojumbo
-			FInstanceBatchItem WorkItem = Items[Index];
-			if (GroupThreadIndex >= WorkItem.BatchPrefixOffset)
-			{
-				WorkItemIndex = uint(Index);
-				break;
-			}
-		}
-	}
-	else
-	{
-#if USE_WAVE_BIT_PREFIX_WORK_DISTRIBUTION
-		// Ballot based boundary search for when wave ops are available.
-		bool bIsBoundary = WorkBoundary[GroupThreadIndex];
-		WorkItemIndex = WavePrefixCountBits(bIsBoundary);
-#else
-		// Distribute the index to the other threads that should work on the same item
-		WorkItemIndex = PrefixMax(GroupThreadIndex);
-#endif
-	}
-
-	return Items[WorkItemIndex];
-}
+#include "../InstanceCulling/InstanceCullingLoadBalancer.ush"
 
 /**
  * Each thread loops over a range on instances loaded from a buffer. The instance bounds are projected to all cached virtual shadow map address space 
  * and any overlapped pages are marked as invalid.
  */
 [numthreads(NUM_THREADS_PER_GROUP, 1, 1)]
-void GPUSceneSetInstancePrimitiveIdCS(uint GroupId : SV_GroupID, int GroupThreadIndex : SV_GroupIndex)
+void GPUSceneSetInstancePrimitiveIdCS(uint3 GroupId : SV_GroupID, int GroupThreadIndex : SV_GroupIndex)
 {
-	FInstanceBatchItem WorkItem = InstanceProcessingGPULoadBalancer_Setup(GroupId, GroupThreadIndex);
+	FInstanceWorkSetup WorkSetup = InstanceCullingLoadBalancer_Setup(GroupId, GroupThreadIndex, 0U);
 
-	// Work out who I am in expansion of  the item in range [0,WorkItem.Count], may be negative when no work at end of batch
-	int LocalItemIndex = GroupThreadIndex - WorkItem.BatchPrefixOffset;
-
-	if (LocalItemIndex >= 0 && LocalItemIndex < WorkItem.NumInstances)
+	if (!WorkSetup.bValid)
 	{
-		int InstanceId = WorkItem.InstanceSceneDataOffset + LocalItemIndex;
-		WriteInstancePrimitiveId(InstanceId, WorkItem.Payload);
+		return;
 	}
+
+	uint InstanceId = WorkSetup.Item.InstanceDataOffset + uint(WorkSetup.LocalItemIndex);
+	WriteInstancePrimitiveId(InstanceId, WorkSetup.Item.Payload);
 }
diff --git a/Engine/Shaders/Private/InstanceCulling/InstanceCullingLoadBalancer.ush b/Engine/Shaders/Private/InstanceCulling/InstanceCullingLoadBalancer.ush
new file mode 100644
--- /dev/null
+++ b/Engine/Shaders/Private/InstanceCulling/InstanceCullingLoadBalancer.ush
@@ -0,0 +1,224 @@
+// Copyright Epic Games, Inc. All Rights Reserved.
+
+#pragma once 
+
+#include "../Common.ush"
+#include "../WaveOpUtil.ush"
+#include "../ComputeShaderUtils.ush"
+
+#if COMPILER_SUPPORTS_WAVE_VOTE
+	// TODO: Need to also check that wave width == thread group size (or implement workaround for other wave sizes)
+#define USE_WAVE_BIT_PREFIX_WORK_DISTRIBUTION 0
+#endif
+
+struct FPackedInstanceBatch
+{
+	uint FirstItem_NumItems;
+};
+
+struct FPackedInstanceBatchItem
+{
+	// packed 32-NumInstancesItemBits:NumInstancesItemBits - need one more bit for the case where one item has ThreadGroupSize work to do
+	uint InstanceDataOffset_NumInstances;
+	// packed 32-PrefixBits:PrefixBits
+	uint Payload_BatchPrefixOffset;
+};
+
+
+StructuredBuffer<FPackedInstanceBatch> BatchBuffer;
+StructuredBuffer<FPackedInstanceBatchItem> ItemBuffer;
+uint NumBatches;
+uint NumItems;
+
+
+struct FInstanceBatch
+{
+	uint FirstItem;
+	uint NumItems;
+};
+
+FInstanceBatch UnpackBatch(FPackedInstanceBatch PackedBatch, uint ItemDataOffset)
+{
+	FInstanceBatch Result;
+	Result.FirstItem = ItemDataOffset + (PackedBatch.FirstItem_NumItems >> NUM_INSTANCES_ITEM_BITS);
+	Result.NumItems = PackedBatch.FirstItem_NumItems & NUM_INSTANCES_ITEM_MASK;
+
+	return Result;
+};
+
+struct FInstanceBatchItem
+{
+	uint InstanceDataOffset;
+	uint NumInstances;
+	uint Payload;
+	uint BatchPrefixOffset;
+};
+
+FInstanceBatchItem UnpackItem(FPackedInstanceBatchItem PackedItem)
+{
+	FInstanceBatchItem Result;
+
+	Result.InstanceDataOffset = PackedItem.InstanceDataOffset_NumInstances >> NUM_INSTANCES_ITEM_BITS;
+	Result.NumInstances = PackedItem.InstanceDataOffset_NumInstances & NUM_INSTANCES_ITEM_MASK;
+	Result.Payload = PackedItem.Payload_BatchPrefixOffset >> PREFIX_BITS;
+	Result.BatchPrefixOffset = PackedItem.Payload_BatchPrefixOffset & PREFIX_BIT_MASK;
+
+	return Result;
+};
+
+
+#if USE_WAVE_BIT_PREFIX_WORK_DISTRIBUTION
+
+groupshared uint WorkBoundary[NUM_THREADS_PER_GROUP];
+
+#else // !USE_WAVE_BIT_PREFIX_WORK_DISTRIBUTION
+
+groupshared uint ItemIndex[NUM_THREADS_PER_GROUP];
+
+void PrefixMaxStep(uint ElementIndex, uint StepSize, inout uint CurrentValue)
+{
+	if (ElementIndex >= StepSize)
+	{
+		CurrentValue = max(CurrentValue, ItemIndex[ElementIndex - StepSize]);
+	}
+	GroupMemoryBarrierWithGroupSync();
+
+	if (ElementIndex >= StepSize)
+	{
+		ItemIndex[ElementIndex] = CurrentValue;
+	}
+	GroupMemoryBarrierWithGroupSync();
+}
+
+uint PrefixMaxLastStep(uint ElementIndex, inout uint CurrentValue)
+{
+	GroupMemoryBarrierWithGroupSync();
+
+	uint StepSize = NUM_THREADS_PER_GROUP / 2U;
+
+	if (ElementIndex >= StepSize)
+	{
+		uint NewValue = max(CurrentValue, ItemIndex[ElementIndex - StepSize]);
+		CurrentValue = NewValue;
+	}
+
+	// No sync needed since we don't look at ItemIndex[] again
+	return CurrentValue;
+}
+
+uint PrefixMax(uint ElementIndex)
+{
+	GroupMemoryBarrierWithGroupSync();
+
+	uint CurrentValue = ItemIndex[ElementIndex];
+	GroupMemoryBarrierWithGroupSync();
+
+	UNROLL
+	for (uint StepSize = 1U; StepSize < NUM_THREADS_PER_GROUP / 2U; ++StepSize)
+	{
+		PrefixMaxStep(ElementIndex, 1U, CurrentValue);
+	}
+	return PrefixMaxLastStep(ElementIndex, CurrentValue);
+}
+
+#endif // USE_WAVE_BIT_PREFIX_WORK_DISTRIBUTION
+
+
+groupshared FInstanceBatchItem Items[NUM_THREADS_PER_GROUP];
+
+struct FInstanceWorkSetup
+{
+	bool bValid;
+	FInstanceBatchItem Item;
+	// Work out who I am in expansion of  the item in range [0,WorkItem.Count], may be negative when no work at end of batch
+	int LocalItemIndex;
+	uint WorkItemIndex;
+};
+
+
+FInstanceWorkSetup InstanceCullingLoadBalancer_Setup(uint3 GroupId, uint GroupThreadIndex, uint ItemDataOffset)
+{
+	FInstanceWorkSetup  Setup = (FInstanceWorkSetup)0;
+
+	uint BatchIndex = GetUnWrappedDispatchGroupId(GroupId);
+	if (BatchIndex >= NumBatches)
+	{
+		return Setup;
+	}
+
+	// Load work description for whole work group
+	FInstanceBatch GroupBatch = UnpackBatch(BatchBuffer[BatchIndex], ItemDataOffset);
+
+	// Special case where we have an 1:1 item/thread no need to do prefix sums or whatever
+	// TODO: this case can also store less data per item
+	if (GroupBatch.NumItems == NUM_THREADS_PER_GROUP)
+	{
+		checkSlow(UnpackItem(ItemBuffer[GroupBatch.FirstItem + GroupThreadIndex]).BatchPrefixOffset == GroupThreadIndex);
+
+		Setup.Item = UnpackItem(ItemBuffer[GroupBatch.FirstItem + GroupThreadIndex]);
+		Setup.bValid = true;
+
+		return Setup;
+	}
+
+	// Load work Items (guaranteed to be less than group size)
+#if USE_WAVE_BIT_PREFIX_WORK_DISTRIBUTION
+	WorkBoundary[GroupThreadIndex] = 0;
+#else 
+	ItemIndex[GroupThreadIndex] = 0U;
+#endif
+
+	GroupMemoryBarrierWithGroupSync();
+
+	if (GroupThreadIndex < GroupBatch.NumItems)
+	{
+		FInstanceBatchItem WorkItem = UnpackItem(ItemBuffer[GroupBatch.FirstItem + GroupThreadIndex]);
+#if USE_WAVE_BIT_PREFIX_WORK_DISTRIBUTION
+		// Mark last instance as boundary
+		WorkBoundary[(WorkItem.BatchPrefixOffset + WorkItem.NumInstances) - 1U] = 1;
+#else 
+		ItemIndex[WorkItem.BatchPrefixOffset] = GroupThreadIndex;
+#endif
+		Items[GroupThreadIndex] = WorkItem;
+	}
+	GroupMemoryBarrierWithGroupSync();
+	uint WorkItemIndex = 0U;
+	if (GroupBatch.NumItems < 4U)
+	{
+		// Distribute using linear search (no syncs requires like the prefix max op)
+		for (int Index = int(GroupBatch.NumItems) - 1; Index >= 0; --Index)
+		{
+			// TODO: pack the offset data into a single dword and store in group batch? 6x5 bits can fit and would be far more efficient that all the shared mumbojumbo
+			FInstanceBatchItem WorkItem = Items[Index];
+			if (GroupThreadIndex >= WorkItem.BatchPrefixOffset)
+			{
+				WorkItemIndex = uint(Index);
+				break;
+			}
+		}
+	}
+	else
+	{
+#if USE_WAVE_BIT_PREFIX_WORK_DISTRIBUTION
+		// Ballot based boundary search for when wave ops are available.
+		bool bIsBoundary = WorkBoundary[GroupThreadIndex];
+		WorkItemIndex = WavePrefixCountBits(bIsBoundary);
+#else
+		// Distribute the index to the other threads that should work on the same item
+		WorkItemIndex = PrefixMax(GroupThreadIndex);
+#endif
+	}
+
+	Setup.WorkItemIndex = WorkItemIndex;
+	Setup.Item = Items[WorkItemIndex];
+	// Work out who I am in expansion of  the item in range [0,WorkItem.Count], may be negative when no work at end of batch
+	Setup.LocalItemIndex = int(GroupThreadIndex - Setup.Item.BatchPrefixOffset);
+	Setup.bValid = Setup.LocalItemIndex >= 0 && Setup.LocalItemIndex < int(Setup.Item.NumInstances);
+	return Setup;
+}
+
+
+uint InstanceCullingLoadBalancer_GetNumBatches()
+{
+	return NumBatches;
+}
\ No newline at end of file
diff --git a/Engine/Source/Runtime/Renderer/Private/GPUScene.cpp b/Engine/Source/Runtime/Renderer/Private/GPUScene.cpp
--- a/Engine/Source/Runtime/Renderer/Private/GPUScene.cpp
+++ b/Engine/Source/Runtime/Renderer/Private/GPUScene.cpp
@@ -371,163 +371,6 @@ void FGPUScene::EndRender()
 }
 
 
-
-// Helper to build the needed data to run per-instance operation on the GPU in a balanced way
-// TODO: Move to own files
-class FInstanceProcessingGPULoadBalancer
-{
-public:
-	static constexpr int32 ThreadGroupSize = 64;
-
-	// Number of bits needed for prefix sum storage
-	static constexpr uint32 PrefixBits = ILog2Const(uint32(ThreadGroupSize));
-	static constexpr uint32 PrefixBitMask = (1U << PrefixBits) - 1U;
-
-	static constexpr uint32 NumInstancesItemBits = PrefixBits + 1U;
-	static constexpr uint32 NumInstancesItemMask = (1U << NumInstancesItemBits) - 1U;
-
-	FInstanceProcessingGPULoadBalancer()
-	{ 
-	}
-
-	void Add(int32 InstanceSceneDataOffset, int32 NumInstanceSceneDataEntries, int32 Payload)
-	{
-		int32 InstancesAdded = 0;
-		while (InstancesAdded < NumInstanceSceneDataEntries)
-		{
-			int32 MaxInstancesThisBatch = ThreadGroupSize - CurrentBatchPrefixSum;
-
-			if (MaxInstancesThisBatch > 0)
-			{
-				const int NumInstancesThisItem = FMath::Min(MaxInstancesThisBatch, NumInstanceSceneDataEntries - InstancesAdded);
-				Items.Add(PackItem(InstanceSceneDataOffset + InstancesAdded, NumInstancesThisItem, Payload, CurrentBatchPrefixSum));
-				if (CurrentBatchNumItems * PrefixBits < sizeof(CurrentBatchPackedPrefixSum) * 8U)
-				{
-					CurrentBatchPackedPrefixSum |= CurrentBatchPackedPrefixSum << (PrefixBits * CurrentBatchNumItems);
-				}
-				CurrentBatchNumItems += 1;
-				InstancesAdded += NumInstancesThisItem;
-				CurrentBatchPrefixSum += NumInstancesThisItem;
-			}
-
-			// Flush batch if it is not possible to add any more items (for one of the reasons)
-			if (MaxInstancesThisBatch <= 0 || CurrentBatchPrefixSum >= ThreadGroupSize)
-			{
-				Batches.Add(PackBatch(CurrentBatchFirstItem, CurrentBatchNumItems));
-				CurrentBatchFirstItem = uint32(Items.Num());
-				CurrentBatchPrefixSum = 0;
-				CurrentBatchNumItems = 0U;
-				CurrentBatchPackedPrefixSum = 0U;
-			}
-		}
-	}
-
-	bool IsEmpty() const
-	{
-		return Items.IsEmpty();
-	}
-
-	struct FPackedBatch
-	{
-		uint32 FirstItem_NumItems;
-	};
-
-	FPackedBatch PackBatch(uint32 FirstItem, uint32 NumItems)
-	{
-		checkSlow(NumItems < (1U << NumInstancesItemBits));
-		checkSlow(FirstItem < (1U << (32U - NumInstancesItemBits)));
-
-		return FPackedBatch{ (FirstItem << NumInstancesItemBits) | (NumItems & NumInstancesItemMask) };
-	}
-
-	struct FPackedItem
-	{
-		// packed 32-NumInstancesItemBits:NumInstancesItemBits - need one more bit for the case where one item has ThreadGroupSize work to do
-		uint32 InstanceSceneDataOffset_NumInstances;
-		// packed 32-PrefixBits:PrefixBits
-		uint32 Payload_BatchPrefixOffset;
-	};
-	FPackedItem PackItem(uint32 InstanceSceneDataOffset, uint32 NumInstances, uint32 Payload, uint32 BatchPrefixSum)
-	{
-		checkSlow(NumInstances < (1U << NumInstancesItemBits));
-		checkSlow(InstanceSceneDataOffset < (1U << (32U - NumInstancesItemBits)));
-		checkSlow(BatchPrefixSum < (1U << PrefixBits));
-		checkSlow(Payload < (1U << (32U - PrefixBits)));
-
-		return FPackedItem
-		{ 
-			(InstanceSceneDataOffset << NumInstancesItemBits) | (NumInstances & NumInstancesItemMask),
-			(Payload << PrefixBits) | (BatchPrefixSum & PrefixBitMask)
-		};
-	}
-
-	BEGIN_SHADER_PARAMETER_STRUCT(FShaderParameters, )
-		SHADER_PARAMETER_RDG_BUFFER_SRV(StructuredBuffer< FPackedBatch >, BatchBuffer)
-		SHADER_PARAMETER_RDG_BUFFER_SRV(StructuredBuffer< FPackedItem >, ItemBuffer)
-		SHADER_PARAMETER(uint32, NumBatches)
-		SHADER_PARAMETER(uint32, NumItems)
-	END_SHADER_PARAMETER_STRUCT()
-
-	struct FGPUData
-	{
-		int32 NumBatches = 0;
-		int32 NumItems = 0;
-		FRDGBufferRef BatchBuffer = nullptr;
-		FRDGBufferRef ItemBuffer = nullptr;
-
-		void GetShaderParameters(FRDGBuilder& GraphBuilder, FShaderParameters& ShaderParameters)
-		{
-			ShaderParameters.BatchBuffer = GraphBuilder.CreateSRV(BatchBuffer);
-			ShaderParameters.ItemBuffer = GraphBuilder.CreateSRV(ItemBuffer);
-			ShaderParameters.NumBatches = NumBatches;
-			ShaderParameters.NumItems = NumItems;
-		}
-	};
-
-	// publish constants to the shader
-	static void SetShaderDefines(FShaderCompilerEnvironment& OutEnvironment)
-	{
-		OutEnvironment.SetDefine(TEXT("NUM_THREADS_PER_GROUP"), ThreadGroupSize);
-		OutEnvironment.SetDefine(TEXT("NUM_INSTANCES_ITEM_BITS"), NumInstancesItemBits);
-		OutEnvironment.SetDefine(TEXT("NUM_INSTANCES_ITEM_MASK"), NumInstancesItemMask);
-		OutEnvironment.SetDefine(TEXT("PREFIX_BITS"), PrefixBits);
-		OutEnvironment.SetDefine(TEXT("PREFIX_BIT_MASK"), PrefixBitMask);
-	}
-
-	FGPUData Upload(FRDGBuilder& GraphBuilder)
-	{
-		if (CurrentBatchNumItems != 0)
-		{
-			Batches.Add(PackBatch(CurrentBatchFirstItem, CurrentBatchNumItems));
-			CurrentBatchNumItems = 0;
-		}
-
-		FGPUData Result;
-		Result.BatchBuffer = CreateStructuredBuffer(GraphBuilder, TEXT("InstanceProcessingGPULoadBalancer.Batches"), Batches, ERDGInitialDataFlags::NoCopy);
-		Result.ItemBuffer = CreateStructuredBuffer(GraphBuilder, TEXT("InstanceProcessingGPULoadBalancer.Items"), Items, ERDGInitialDataFlags::NoCopy);
-		Result.NumBatches = Batches.Num();
-		Result.NumItems = Items.Num();
-
-		return Result;
-	}
-
-	FIntVector GetCsGroupCount() const
-	{
-		return FIntVector(Batches.Num(), 1, 1);
-	}
-
-	TArray<FPackedBatch, SceneRenderingAllocator>  Batches;
-	TArray<FPackedItem, SceneRenderingAllocator>  Items;
-	
-	int32 CurrentBatchPrefixSum = 0;
-	uint32 CurrentBatchNumItems = 0U;
-	uint32 CurrentBatchPackedPrefixSum = 0U;
-	uint32 CurrentBatchFirstItem = 0U;
-};
-
-
-
-
 void FGPUScene::UpdateInternal(FRDGBuilder& GraphBuilder, FScene& Scene)
 {
 	LLM_SCOPE_BYTAG(GPUScene);
@@ -567,7 +410,7 @@ void FGPUScene::UpdateInternal(FRDGBuilder& GraphBuilder, FScene& Scene)
 
 	// Pull out instances needing only primitive ID update, they still have to go to the general update such that the primitive gets updated (as it moved)
 	{
-		FInstanceProcessingGPULoadBalancer IdOnlyUpdateData;
+		FInstanceGPULoadBalancer IdOnlyUpdateData;
 		for (int32 Index = 0; Index < Adapter.PrimitivesToUpdate.Num(); ++Index)
 		{
 			int32 PrimitiveId = Adapter.PrimitivesToUpdate[Index];
@@ -1785,11 +1628,11 @@ class FGPUSceneSetInstancePrimitiveIdCS : public FGlobalShader
 	SHADER_USE_PARAMETER_STRUCT(FGPUSceneSetInstancePrimitiveIdCS, FGlobalShader)
 public:
 	BEGIN_SHADER_PARAMETER_STRUCT(FParameters, )
-		SHADER_PARAMETER_STRUCT_INCLUDE(FInstanceProcessingGPULoadBalancer::FShaderParameters, BatcherParameters)
+		SHADER_PARAMETER_STRUCT_INCLUDE(FGPUScene::FInstanceGPULoadBalancer::FShaderParameters, BatcherParameters)
 		SHADER_PARAMETER_STRUCT_INCLUDE(FGPUSceneWriterParameters, GPUSceneWriterParameters)
 		END_SHADER_PARAMETER_STRUCT()
 
-	static constexpr int32 NumThreadsPerGroup = FInstanceProcessingGPULoadBalancer::ThreadGroupSize;
+	static constexpr int32 NumThreadsPerGroup = FGPUScene::FInstanceGPULoadBalancer::ThreadGroupSize;
 
 	static bool ShouldCompilePermutation(const FGlobalShaderPermutationParameters& Parameters)
 	{
@@ -1800,7 +1643,7 @@ public:
 	{
 		FGlobalShader::ModifyCompilationEnvironment(Parameters, OutEnvironment);
 
-		FInstanceProcessingGPULoadBalancer::SetShaderDefines(OutEnvironment);
+		FGPUScene::FInstanceGPULoadBalancer::SetShaderDefines(OutEnvironment);
 
 		OutEnvironment.SetDefine(TEXT("VF_SUPPORTS_PRIMITIVE_SCENE_DATA"), 1);
 	}
@@ -1808,7 +1651,7 @@ public:
 IMPLEMENT_GLOBAL_SHADER(FGPUSceneSetInstancePrimitiveIdCS, "/Engine/Private/GPUScene/GPUSceneDataManagement.usf", "GPUSceneSetInstancePrimitiveIdCS", SF_Compute);
 
 
-void FGPUScene::AddUpdatePrimitiveIdsPass(FRDGBuilder& GraphBuilder, FInstanceProcessingGPULoadBalancer& IdOnlyUpdateItems)
+void FGPUScene::AddUpdatePrimitiveIdsPass(FRDGBuilder& GraphBuilder, FInstanceGPULoadBalancer& IdOnlyUpdateItems)
 {
 	if (!IdOnlyUpdateItems.IsEmpty())
 	{
@@ -1827,7 +1670,7 @@ void FGPUScene::AddUpdatePrimitiveIdsPass(FRDGBuilder& GraphBuilder, FInstancePr
 			RDG_EVENT_NAME("GPUSceneSetInstancePrimitiveIdCS"),
 			ComputeShader,
 			PassParameters,
-			IdOnlyUpdateItems.GetCsGroupCount()
+			IdOnlyUpdateItems.GetWrappedCsGroupCount()
 		);
 
 		EndReadWriteAccess(GraphBuilder, ERHIAccess::UAVCompute);
diff --git a/Engine/Source/Runtime/Renderer/Private/GPUScene.h b/Engine/Source/Runtime/Renderer/Private/GPUScene.h
--- a/Engine/Source/Runtime/Renderer/Private/GPUScene.h
+++ b/Engine/Source/Runtime/Renderer/Private/GPUScene.h
@@ -7,6 +7,7 @@
 #include "PrimitiveUniformShaderParameters.h"
 #include "PrimitiveSceneInfo.h"
 #include "GrowOnlySpanAllocator.h"
+#include "InstanceCulling/InstanceCullingLoadBalancer.h"
 
 class FRHICommandList;
 class FScene;
@@ -38,8 +39,8 @@ public:
 	 */
 	FORCEINLINE int32 Add(const FPrimitiveUniformShaderParameters* Data, int32 Num)
 	{
-		ensure(GPUSceneDynamicContext != nullptr);
-		ensure(!bCommitted);
+		check(GPUSceneDynamicContext != nullptr);
+		check(!bCommitted);
 
 		// Lazy allocation of the upload data to not waste space and processing if none was needed.
 		if (UploadData == nullptr)
@@ -58,10 +59,10 @@ public:
 	 */
 	FORCEINLINE int32 GetFinalId(int32 LocalIndex) const
 	{
-		ensure(bCommitted);
+		check(bCommitted);
 
-		ensure(PrimitiveIdRange.GetLowerBoundValue() >= 0);
-		ensure(PrimitiveIdRange.Size<int32>() >= LocalIndex);
+		check(PrimitiveIdRange.GetLowerBoundValue() >= 0);
+		check(PrimitiveIdRange.Size<int32>() >= LocalIndex);
 		return PrimitiveIdRange.GetLowerBoundValue() + LocalIndex;
 	}
 
@@ -75,23 +76,32 @@ public:
 	/**
 	 * Get the range of Primitive IDs in GPU-Scene for this batch of dynamic primitives, only valid to call after commit.
 	 */
-	FORCEINLINE const TRange<int32> &GetPrimitiveIdRange() const
-	{ 
-		ensure(bCommitted || UploadData == nullptr);
-		return PrimitiveIdRange; 
+	FORCEINLINE const TRange<int32>& GetPrimitiveIdRange() const
+	{
+		check(bCommitted || UploadData == nullptr);
+		return PrimitiveIdRange;
 	}
 
 	FORCEINLINE int32 GetInstanceSceneDataOffset(int32 PrimitiveId) const
 	{
-		ensure(bCommitted);
-		ensure(PrimitiveIdRange.Contains(PrimitiveId));
-		ensure(UploadData->bIsUploaded);
+		check(bCommitted);
+		check(PrimitiveIdRange.Contains(PrimitiveId));
+		check(UploadData->bIsUploaded);
 
 		// Assume a 1:1 mapping between primitive ID and instance ID
 		return UploadData->InstanceSceneDataOffset + (PrimitiveId - PrimitiveIdRange.GetLowerBoundValue());
 	}
 
+	FORCEINLINE int32 GetInstanceSceneDataOffset() const
+	{
+		check(bCommitted || UploadData == nullptr);
+
+		// Assume a 1:1 mapping between primitive ID and instance ID
+		return UploadData != nullptr ? UploadData->InstanceSceneDataOffset : 0;
+	}
+
 	int32 Num() const {	return UploadData != nullptr ? UploadData->PrimitiveShaderData.Num() : 0; }
+	int32 NumInstances() const { return Num(); }
 private:
 
 	friend class FGPUScene;
@@ -277,6 +287,9 @@ public:
 
 	uint32 GetSceneFrameNumber() const { return SceneFrameNumber; }
 
+	int32 GetNumInstances() const { return InstanceSceneDataAllocator.GetMaxSize(); }
+	int32 GetNumPrimitives() const { return DynamicPrimitivesOffset; }
+
 	bool bUpdateAllPrimitives;
 
 	/** Indices of primitives that need to be updated in GPU Scene */
@@ -299,6 +312,7 @@ public:
 	FRWBufferStructured		LightmapDataBuffer;
 	FScatterUploadBuffer	LightmapUploadBuffer;
 
+	using FInstanceGPULoadBalancer = TInstanceCullingLoadBalancer<SceneRenderingAllocator>;
 private:
 	TArray<EPrimitiveDirtyState> PrimitiveDirtyState;
 
@@ -339,7 +353,7 @@ private:
 
 	void UpdateInternal(FRDGBuilder& GraphBuilder, FScene& Scene);
 
-	void AddUpdatePrimitiveIdsPass(FRDGBuilder& GraphBuilder, FInstanceProcessingGPULoadBalancer& IdOnlyUpdateItems);
+	void AddUpdatePrimitiveIdsPass(FRDGBuilder& GraphBuilder, FInstanceGPULoadBalancer& IdOnlyUpdateItems);
 };
 
 class FGPUSceneScopeBeginEndHelper
diff --git a/Engine/Source/Runtime/Renderer/Private/InstanceCulling/InstanceCullingLoadBalancer.cpp b/Engine/Source/Runtime/Renderer/Private/InstanceCulling/InstanceCullingLoadBalancer.cpp
new file mode 100644
--- /dev/null
+++ b/Engine/Source/Runtime/Renderer/Private/InstanceCulling/InstanceCullingLoadBalancer.cpp
@@ -0,0 +1,3 @@
+// Copyright Epic Games, Inc. All Rights Reserved.
+
+#include "InstanceCullingLoadBalancer.h"
diff --git a/Engine/Source/Runtime/Renderer/Private/InstanceCulling/InstanceCullingLoadBalancer.h b/Engine/Source/Runtime/Renderer/Private/InstanceCulling/InstanceCullingLoadBalancer.h
new file mode 100644
--- /dev/null
+++ b/Engine/Source/Runtime/Renderer/Private/InstanceCulling/InstanceCullingLoadBalancer.h
@@ -0,0 +1,229 @@
+// Copyright Epic Games, Inc. All Rights Reserved.
+
+#pragma once
+
+#include "CoreMinimal.h"
+#include "RHI.h"
+#include "RenderGraphBuilder.h"
+#include "RenderGraphResources.h"
+#include "RenderGraphUtils.h"
+
+
+namespace InstanceCullingImplementationDetails
+{
+	// Helpers to make default ERDGInitialDataFlags::NoCopy for scene rendering allocator which is known to live longer than the RenderGraph
+	template <typename AllocatorType>
+	struct AllocatorTypeRDGInitialDataFlags
+	{
+		static constexpr ERDGInitialDataFlags Flags = ERDGInitialDataFlags::None;
+	};
+
+	template <>
+	struct AllocatorTypeRDGInitialDataFlags<SceneRenderingAllocator>
+	{
+		static constexpr ERDGInitialDataFlags Flags = ERDGInitialDataFlags::NoCopy;
+	};
+};
+
+/*
+ * Helper to build the needed data to run per-instance operation on the GPU in a balanced way
+ */
+template <typename InAllocatorType = FDefaultAllocator>
+class TInstanceCullingLoadBalancer
+{
+public:
+	using AllocatorType = InAllocatorType;
+	static constexpr ERDGInitialDataFlags DefaultRDGInitialDataFlags = InstanceCullingImplementationDetails::AllocatorTypeRDGInitialDataFlags<AllocatorType>::Flags;
+
+	static constexpr uint32 ThreadGroupSize = 64U;
+
+	// Number of bits needed for prefix sum storage
+	static constexpr uint32 PrefixBits = 6U;//ILog2Const(uint32(ThreadGroupSize));
+	static_assert((1U << PrefixBits) == ThreadGroupSize, "ThreadGroupSize and PrefixBits must be kept in sync");
+	static constexpr uint32 PrefixBitMask = (1U << PrefixBits) - 1U;
+
+	static constexpr uint32 NumInstancesItemBits = PrefixBits + 1U;
+	static constexpr uint32 NumInstancesItemMask = (1U << NumInstancesItemBits) - 1U;
+
+	/*
+	 * Publish constants to a shader implementing a kernel using the load balancer.
+	 * Call from ModifyCompilationEnvironment
+	 */
+	static void SetShaderDefines(FShaderCompilerEnvironment& OutEnvironment)
+	{
+		OutEnvironment.SetDefine(TEXT("NUM_THREADS_PER_GROUP"), ThreadGroupSize);
+		OutEnvironment.SetDefine(TEXT("NUM_INSTANCES_ITEM_BITS"), NumInstancesItemBits);
+		OutEnvironment.SetDefine(TEXT("NUM_INSTANCES_ITEM_MASK"), NumInstancesItemMask);
+		OutEnvironment.SetDefine(TEXT("PREFIX_BITS"), PrefixBits);
+		OutEnvironment.SetDefine(TEXT("PREFIX_BIT_MASK"), PrefixBitMask);
+	}
+
+	void ReserveStorage(int32 NumBatches, int32 NumItems)
+	{
+		Batches.Empty(NumBatches);
+		Items.Empty(NumItems);
+	}
+
+	/**
+	 * Add a span of instances to be processed.
+	 */
+	void Add(uint32 InstanceDataOffset, uint32 NumInstanceDataEntries, uint32 Payload)
+	{
+		uint32 InstancesAdded = 0;
+		while (InstancesAdded < NumInstanceDataEntries)
+		{
+			uint32 MaxInstancesThisBatch = ThreadGroupSize - CurrentBatchPrefixSum;
+
+			if (MaxInstancesThisBatch > 0)
+			{
+				const uint32 NumInstancesThisItem = FMath::Min(MaxInstancesThisBatch, NumInstanceDataEntries - InstancesAdded);
+
+				Items.Add(PackItem(InstanceDataOffset + InstancesAdded, NumInstancesThisItem, Payload, CurrentBatchPrefixSum));
+				if (CurrentBatchNumItems * PrefixBits < sizeof(CurrentBatchPackedPrefixSum) * 8U)
+				{
+					CurrentBatchPackedPrefixSum |= CurrentBatchPackedPrefixSum << (PrefixBits * CurrentBatchNumItems);
+				}
+				CurrentBatchNumItems += 1U;
+				InstancesAdded += NumInstancesThisItem;
+				CurrentBatchPrefixSum += NumInstancesThisItem;
+
+			}
+
+			// Flush batch if it is not possible to add any more items (for one of the reasons)
+			if (MaxInstancesThisBatch <= 0U || CurrentBatchPrefixSum >= ThreadGroupSize)
+			{
+				Batches.Add(PackBatch(CurrentBatchFirstItem, CurrentBatchNumItems));
+				CurrentBatchFirstItem = uint32(Items.Num());
+				CurrentBatchPrefixSum = 0u;
+				CurrentBatchNumItems = 0U;
+				CurrentBatchPackedPrefixSum = 0U;
+			}
+		}
+	}
+
+	bool IsEmpty() const
+	{
+		return Items.IsEmpty();
+	}
+
+	struct FPackedBatch
+	{
+		uint32 FirstItem_NumItems;
+	};
+
+	FPackedBatch PackBatch(uint32 FirstItem, uint32 NumItems)
+	{
+		checkSlow(NumItems < (1U << NumInstancesItemBits));
+		checkSlow(FirstItem < (1U << (32U - NumInstancesItemBits)));
+
+		return FPackedBatch{ (FirstItem << NumInstancesItemBits) | (NumItems & NumInstancesItemMask) };
+	}
+
+	struct FPackedItem
+	{
+		// packed 32-NumInstancesItemBits:NumInstancesItemBits - need one more bit for the case where one item has ThreadGroupSize work to do
+		uint32 InstanceDataOffset_NumInstances;
+		// packed 32-PrefixBits:PrefixBits
+		uint32 Payload_BatchPrefixOffset;
+	};
+	FPackedItem PackItem(uint32 InstanceDataOffset, uint32 NumInstances, uint32 Payload, uint32 BatchPrefixSum)
+	{
+		checkSlow(NumInstances < (1U << NumInstancesItemBits));
+		checkSlow(InstanceDataOffset < (1U << (32U - NumInstancesItemBits)));
+		checkSlow(BatchPrefixSum < (1U << PrefixBits));
+		checkSlow(Payload < (1U << (32U - PrefixBits)));
+
+		return FPackedItem
+		{
+			(InstanceDataOffset << NumInstancesItemBits) | (NumInstances & NumInstancesItemMask),
+			(Payload << PrefixBits) | (BatchPrefixSum & PrefixBitMask)
+		};
+	}
+
+	BEGIN_SHADER_PARAMETER_STRUCT(FShaderParameters, )
+		SHADER_PARAMETER_RDG_BUFFER_SRV(StructuredBuffer< FPackedBatch >, BatchBuffer)
+		SHADER_PARAMETER_RDG_BUFFER_SRV(StructuredBuffer< FPackedItem >, ItemBuffer)
+		SHADER_PARAMETER(uint32, NumBatches)
+		SHADER_PARAMETER(uint32, NumItems)
+	END_SHADER_PARAMETER_STRUCT()
+
+	struct FGPUData
+	{
+		int32 NumBatches = 0;
+		int32 NumItems = 0;
+		FRDGBufferRef BatchBuffer = nullptr;
+		FRDGBufferRef ItemBuffer = nullptr;
+
+		void GetShaderParameters(FRDGBuilder& GraphBuilder, FShaderParameters& ShaderParameters)
+		{
+			ShaderParameters.BatchBuffer = GraphBuilder.CreateSRV(BatchBuffer);
+			ShaderParameters.ItemBuffer = GraphBuilder.CreateSRV(ItemBuffer);
+			ShaderParameters.NumBatches = NumBatches;
+			ShaderParameters.NumItems = NumItems;
+		}
+	};
+
+	FGPUData Upload(FRDGBuilder& GraphBuilder, ERDGInitialDataFlags RDGInitialDataFlags = DefaultRDGInitialDataFlags)
+	{
+		FinalizeBatches();
+
+		FGPUData Result;
+		Result.BatchBuffer = CreateStructuredBuffer(GraphBuilder, TEXT("InstanceCullingLoadBalancer.Batches"), Batches, RDGInitialDataFlags);
+		Result.ItemBuffer = CreateStructuredBuffer(GraphBuilder, TEXT("InstanceCullingLoadBalancer.Items"), Items, RDGInitialDataFlags);
+		Result.NumBatches = Batches.Num();
+		Result.NumItems = Items.Num();
+
+		return Result;
+	}
+
+	/**
+	 * Call when finished adding work items to the balancer to flush any in-progress batches.
+	 */
+	void FinalizeBatches()
+	{
+		if (CurrentBatchNumItems != 0)
+		{
+			Batches.Add(PackBatch(CurrentBatchFirstItem, CurrentBatchNumItems));
+			CurrentBatchNumItems = 0;
+		}
+	}
+
+	/**
+	 * Returns a 3D group count that is large enough to generate one group per batch using FComputeShaderUtils::GetGroupCountWrapped.
+	 * Use GetUnWrappedDispatchGroupId in the shader to retrieve the linear index.
+	 */
+	FIntVector GetWrappedCsGroupCount() const
+	{
+		return FComputeShaderUtils::GetGroupCountWrapped(Batches.Num());
+	}
+
+	const TArray<FPackedBatch, AllocatorType> &GetBatches() const
+	{
+		check(CurrentBatchNumItems == 0);
+		return Batches;
+	};
+
+	const TArray<FPackedItem, AllocatorType> &GetItems() const
+	{
+		check(CurrentBatchNumItems == 0);
+		return Items;
+	}
+
+	template <typename AllocatorType>
+	void AppendData(const TInstanceCullingLoadBalancer<AllocatorType> &Other)
+	{
+		check(Other.CurrentBatchNumItems == 0);
+
+		Batches.Append(Other.Batches);
+		Items.Append(Other.Items);
+	}
+
+protected:
+	TArray<FPackedBatch, AllocatorType>  Batches;
+	TArray<FPackedItem, AllocatorType>  Items;
+
+	uint32 CurrentBatchPrefixSum = 0u;
+	uint32 CurrentBatchNumItems = 0U;
+	uint32 CurrentBatchPackedPrefixSum = 0U;
+	uint32 CurrentBatchFirstItem = 0U;
+};
