Added basic gaussian blur to the scene color buffer in order to approximate rough refraction.
Separated from bloom for now.
This should be resolution independent (will need more test down the line)
Maximum filter sample count is capped.

diff --git a/Engine/Shaders/Private/DistortApplyScreenPS.usf b/Engine/Shaders/Private/DistortApplyScreenPS.usf
--- a/Engine/Shaders/Private/DistortApplyScreenPS.usf
+++ b/Engine/Shaders/Private/DistortApplyScreenPS.usf
@@ -61,11 +61,12 @@ void Main(
 	half RoughnessScatter = Texture2DSample(RoughnessScatterTexture, RoughnessScatterSampler, TexCoord.xy) * 10.0f;
 #endif
 
-	float2 NewBufferUV = ((TexCoord.xy + DistBufferUVOffset) * View.BufferSizeAndInvSize.xy - View.ViewRectMin.xy) * View.BufferSizeAndInvSize.zw;
+	const float2 NoDistortionUVs = TexCoord.xy * View.BufferSizeAndInvSize.xy * View.ViewSizeAndInvSize.zw;
+	float2 NewBufferUV = NoDistortionUVs + DistBufferUVOffset;
 
 	FLATTEN if (any(NewBufferUV < 0.0f) || any(NewBufferUV > 1.0f))
 	{
-		NewBufferUV = (TexCoord.xy * View.BufferSizeAndInvSize.xy - View.ViewRectMin.xy) * View.BufferSizeAndInvSize.zw;
+		NewBufferUV = NoDistortionUVs;
 	}
 
 #if USE_MSAA
diff --git a/Engine/Shaders/Private/DistortFiltering.usf b/Engine/Shaders/Private/DistortFiltering.usf
--- a/Engine/Shaders/Private/DistortFiltering.usf
+++ b/Engine/Shaders/Private/DistortFiltering.usf
@@ -7,6 +7,24 @@
 #include "Common.ush"
 
 
+
+#ifdef COPYSCENECOLORTEXTUREPS
+
+Texture2D SceneColorTexture;
+
+void CopySceneColorTexturePS(FScreenVertexOutput Input, out float4 OutColor : SV_Target0)
+{
+	float2 TextureUV = (Input.UV * View.ViewSizeAndInvSize.xy + View.ViewRectMin.xy) * View.BufferSizeAndInvSize.zw;
+	float4 LinearColor = Texture2DSample(SceneColorTexture, GlobalPointClampedSampler, TextureUV);
+	OutColor = float4(LinearColor.rgb, 0);
+}
+
+#endif // COPYSCENECOLORPS
+
+
+
+#ifdef DOWNSAMPLECOLORCS
+
 uint SrcMipIndex;
 int2 SrcMipResolution;
 int2 DstMipResolution;
@@ -14,8 +32,6 @@ SamplerState SourceSampler;
 Texture2D<float4> SourceTexture;
 RWTexture2D<float4> OutTextureMipColor;
 
-
-/** output distortion offsets as color so they can be accumulated (via blending) */
 [numthreads(THREADGROUP_SIZE, THREADGROUP_SIZE, 1)]
 void DownsampleColorCS(uint3 ThreadId : SV_DispatchThreadID)
 {
@@ -28,5 +44,44 @@ void DownsampleColorCS(uint3 ThreadId : SV_DispatchThreadID)
 	float4 OutColor = SourceTexture.SampleLevel(SourceSampler, ScaledUVs, SrcMipIndex);
 
 	OutTextureMipColor[uint2(ThreadId.xy)] = OutColor;
+}
+
+#endif // DOWNSAMPLECOLORCS
+
 
+
+#ifdef FILTERCOLORCS
+
+uint SampleCount;
+uint SrcMipIndex;
+int2 MipResolution;
+float4 BlurDirection;
+SamplerState SourceSampler;
+Texture2D<float4> SourceTexture;
+RWTexture2D<float4> OutTextureMipColor;
+float4 SampleOffsetsWeights[128];
+
+[numthreads(THREADGROUP_SIZE, THREADGROUP_SIZE, 1)]
+void FilterColorCS(uint3 ThreadId : SV_DispatchThreadID)
+{
+	if (any(ThreadId.xy >= uint2(MipResolution)))
+	{
+		return;
+	}
+	float2 ScaledUVs = (float2(ThreadId.xy) + 0.5f) / float2(MipResolution);
+
+	float4 OutColor = 0;
+	
+	for (uint i = 0; i < SampleCount; ++i)
+	{
+		float2 SampleOffsetWeight = SampleOffsetsWeights[i].xy;
+
+		OutColor += SourceTexture.SampleLevel(SourceSampler, ScaledUVs + SampleOffsetWeight.xx * BlurDirection.xy, SrcMipIndex) * SampleOffsetWeight.y;
+	}
+
+	OutTextureMipColor[uint2(ThreadId.xy)] = OutColor;
 }
+
+#endif // FILTERCOLORCS
+
+
diff --git a/Engine/Source/Runtime/Renderer/Private/DistortionRendering.cpp b/Engine/Source/Runtime/Renderer/Private/DistortionRendering.cpp
--- a/Engine/Source/Runtime/Renderer/Private/DistortionRendering.cpp
+++ b/Engine/Source/Runtime/Renderer/Private/DistortionRendering.cpp
@@ -14,6 +14,7 @@
 #include "ScreenPass.h"
 #include "MeshPassProcessor.inl"
 #include "Strata/Strata.h"
+#include "ScreenRendering.h"
 
 DECLARE_GPU_STAT(Distortion);
 
@@ -25,6 +26,25 @@ static TAutoConsoleVariable<int32> CVarDisableDistortion(
 	TEXT("Prevents distortion effects from rendering.  Saves a full-screen framebuffer's worth of memory."),
 	ECVF_Default);
 
+static TAutoConsoleVariable<int32> CVarRoughRefraction(
+	TEXT("r.RoughRefraction"),
+	1,
+	TEXT("Prevent rough refraction from happening, i.e. blurring of the background."),
+	ECVF_Scalability | ECVF_RenderThreadSafe);
+
+static TAutoConsoleVariable<float> CVarRoughRefractionBlurScale(
+	TEXT("r.RoughRefraction.BlurScale"),
+	0.2f,
+	TEXT("Global scale the background blur amount after it is mapped form the surface back roughness/scattering amount."),
+	ECVF_RenderThreadSafe);
+
+static TAutoConsoleVariable<float> CVarRoughRefractionCenterWeight(
+	TEXT("r.RoughRefraction.CenterWeight"),
+	0.0f,
+	TEXT("The weight of the center sample. Value greater than 0 means the sharp image is more and more visible."),
+	ECVF_RenderThreadSafe);
+
+
 IMPLEMENT_STATIC_UNIFORM_BUFFER_STRUCT(FDistortionPassUniformParameters, "DistortionPass", SceneTextures);
 
 int32 FSceneRenderer::GetRefractionQuality(const FSceneViewFamily& ViewFamily)
@@ -70,7 +90,7 @@ TRDGUniformBufferRef<FDistortionPassUniformParameters> CreateDistortionPassUnifo
 
 static bool GetUseRoughRefraction()
 {
-	return Strata::IsStrataEnabled();
+	return Strata::IsStrataEnabled() && CVarRoughRefraction.GetValueOnRenderThread() > 0;
 }
 
 class FDistortionScreenPS : public FGlobalShader
@@ -205,6 +225,53 @@ BEGIN_SHADER_PARAMETER_STRUCT(FDistortionPassParameters, )
 	RENDER_TARGET_BINDING_SLOTS()
 END_SHADER_PARAMETER_STRUCT()
 
+//////////////////////////////////////////////////////////////////////////
+
+class FCopySceneColorTexturePS : public FGlobalShader
+{
+public:
+	DECLARE_GLOBAL_SHADER(FCopySceneColorTexturePS);
+	SHADER_USE_PARAMETER_STRUCT(FCopySceneColorTexturePS, FGlobalShader);
+
+	BEGIN_SHADER_PARAMETER_STRUCT(FParameters, )
+		SHADER_PARAMETER_STRUCT_REF(FViewUniformShaderParameters, View)
+		SHADER_PARAMETER_RDG_TEXTURE(Texture2D, SceneColorTexture)
+		RENDER_TARGET_BINDING_SLOTS()
+	END_SHADER_PARAMETER_STRUCT()
+
+	static bool ShouldCompilePermutation(const FGlobalShaderPermutationParameters& Parameters)
+	{
+		return true;
+	}
+	static void ModifyCompilationEnvironment(const FGlobalShaderPermutationParameters& Parameters, FShaderCompilerEnvironment& OutEnvironment)
+	{
+		FGlobalShader::ModifyCompilationEnvironment(Parameters, OutEnvironment);
+		OutEnvironment.SetDefine(TEXT("COPYSCENECOLORTEXTUREPS"), 1);
+	}
+};
+
+IMPLEMENT_GLOBAL_SHADER(FCopySceneColorTexturePS, "/Engine/Private/DistortFiltering.usf", "CopySceneColorTexturePS", SF_Pixel);
+
+static void AddCopySceneColorPass(FRDGBuilder& GraphBuilder, const FViewInfo& View, FRDGTextureRef SceneColorTexture, FRDGTextureRef SceneColorCopyTexture)
+{
+	const FScreenPassTextureViewport Viewport(SceneColorCopyTexture, View.ViewRect);
+
+	TShaderMapRef<FScreenVS> VertexShader(View.ShaderMap);
+	TShaderMapRef<FCopySceneColorTexturePS> PixelShader(View.ShaderMap);
+
+	auto* PassParameters = GraphBuilder.AllocParameters<FCopySceneColorTexturePS::FParameters>();
+	PassParameters->View = View.ViewUniformBuffer;
+	PassParameters->SceneColorTexture = SceneColorTexture;
+	PassParameters->RenderTargets[0] = FRenderTargetBinding(SceneColorCopyTexture, ERenderTargetLoadAction::ENoAction);
+
+	const FScreenPassTextureViewport InputViewport(SceneColorTexture, View.ViewRect);
+	const FScreenPassTextureViewport OutputViewport(SceneColorCopyTexture);
+
+	AddDrawScreenPass(GraphBuilder, {}, View, InputViewport, OutputViewport, VertexShader, PixelShader, PassParameters);
+}
+
+//////////////////////////////////////////////////////////////////////////
+
 class FDownsampleSceneColorCS : public FGlobalShader
 {
 	DECLARE_GLOBAL_SHADER(FDownsampleSceneColorCS);
@@ -229,11 +296,170 @@ public:
 	{
 		FGlobalShader::ModifyCompilationEnvironment(Parameters, OutEnvironment);
 		OutEnvironment.SetDefine(TEXT("THREADGROUP_SIZE"), ThreadGroupSize);
-		OutEnvironment.SetDefine(TEXT("USE_COMPUTE"), 1);
+		OutEnvironment.SetDefine(TEXT("DOWNSAMPLECOLORCS"), 1);
 	}
 };
 IMPLEMENT_GLOBAL_SHADER(FDownsampleSceneColorCS, "/Engine/Private/DistortFiltering.usf", "DownsampleColorCS", SF_Compute);
 
+static void AddDownsampleSceneColorPass(FRDGBuilder& GraphBuilder, FDownsampleSceneColorCS::FParameters* PassParameters, const FViewInfo& View)
+{
+	FIntVector NumGroups = FIntVector::DivideAndRoundUp(
+		FIntVector(PassParameters->DstMipResolution.X, PassParameters->DstMipResolution.Y, 1), 
+		FIntVector(FDownsampleSceneColorCS::ThreadGroupSize, FDownsampleSceneColorCS::ThreadGroupSize, 1));
+
+	FDownsampleSceneColorCS::FPermutationDomain PermutationVector;
+	TShaderMapRef<FDownsampleSceneColorCS> ComputeShader(View.ShaderMap, PermutationVector);
+
+	// Dispatch with GenerateMips: reading from a slice through SRV and writing into lower mip through UAV.
+	ClearUnusedGraphResources(ComputeShader, PassParameters);
+	GraphBuilder.AddPass(
+		Forward<FRDGEventName>(RDG_EVENT_NAME("DistoMipGen")),
+		PassParameters,
+		ERDGPassFlags::Compute,
+		[PassParameters, ComputeShader, NumGroups](FRHICommandList& RHICmdList)
+		{
+			FComputeShaderUtils::Dispatch(RHICmdList, ComputeShader, *PassParameters, NumGroups);
+		});
+}
+
+//////////////////////////////////////////////////////////////////////////
+
+// STARTA_TODO make common with the function in PostProcessWeightedSampleSum.cpp
+class FGaussianFiltering
+{
+public:
+	// Evaluates an unnormalized normal distribution PDF around 0 at given X with Variance.
+	static float NormalDistributionUnscaled(float X, float Sigma, float CrossCenterWeight)
+	{
+		const float DX = FMath::Abs(X);
+
+		const float ClampedOneMinusDX = FMath::Max(0.0f, 1.0f - DX);
+
+		// Tweak the gaussian shape e.g. "r.Bloom.Cross 3.5"
+		if (CrossCenterWeight > 1.0f)
+		{
+			return FMath::Pow(ClampedOneMinusDX, CrossCenterWeight);
+		}
+		else
+		{
+			// Constant is tweaked give a similar look to UE before we fixed the scale bug (Some content tweaking might be needed).
+			// The value defines how much of the Gaussian clipped by the sample window.
+			// r.Filter.SizeScale allows to tweak that for performance/quality.
+			const float LegacyCompatibilityConstant = -16.7f;
+
+			const float Gaussian = FMath::Exp(LegacyCompatibilityConstant * FMath::Square(DX / Sigma));
+
+			return FMath::Lerp(Gaussian, ClampedOneMinusDX, CrossCenterWeight);
+		}
+	}
+
+	static float GetClampedKernelRadius(uint32 SampleCountMax, float KernelRadius)
+	{
+		return FMath::Clamp<float>(KernelRadius, DELTA, SampleCountMax - 1);
+	}
+
+	static int GetIntegerKernelRadius(uint32 SampleCountMax, float KernelRadius)
+	{
+		// Smallest radius will be 1.
+		return FMath::Min<int32>(FMath::CeilToInt(GetClampedKernelRadius(SampleCountMax, KernelRadius)), SampleCountMax - 1);
+	}
+
+	static uint32 Compute1DGaussianFilterKernel(FVector2D* OutOffsetAndWeight, uint32 SampleCountMax, float KernelRadius, float CrossCenterWeight, float FilterSizeScale)
+	{
+		const float ClampedKernelRadius = GetClampedKernelRadius(SampleCountMax, KernelRadius);
+
+		const int32 IntegerKernelRadius = GetIntegerKernelRadius(SampleCountMax, KernelRadius * FilterSizeScale);
+
+		uint32 SampleCount = 0;
+		float WeightSum = 0.0f;
+
+		for (int32 SampleIndex = -IntegerKernelRadius; SampleIndex <= IntegerKernelRadius; SampleIndex += 2)
+		{
+			float Weight0 = NormalDistributionUnscaled(SampleIndex, ClampedKernelRadius, CrossCenterWeight);
+			float Weight1 = 0.0f;
+
+			// We use the bilinear filter optimization for gaussian blur. However, we don't want to bias the
+			// last sample off the edge of the filter kernel, so the very last tap just is on the pixel center.
+			if (SampleIndex != IntegerKernelRadius)
+			{
+				Weight1 = NormalDistributionUnscaled(SampleIndex + 1, ClampedKernelRadius, CrossCenterWeight);
+			}
+
+			const float TotalWeight = Weight0 + Weight1;
+			OutOffsetAndWeight[SampleCount].X = SampleIndex + (Weight1 / TotalWeight);
+			OutOffsetAndWeight[SampleCount].Y = TotalWeight;
+			WeightSum += TotalWeight;
+			SampleCount++;
+		}
+
+		// Normalize blur weights.
+		const float WeightSumInverse = 1.0f / WeightSum;
+		for (uint32 SampleIndex = 0; SampleIndex < SampleCount; ++SampleIndex)
+		{
+			OutOffsetAndWeight[SampleIndex].Y *= WeightSumInverse;
+		}
+
+		return SampleCount;
+	}
+};
+
+// If this is update, please update DistortFiltering.usf
+#define MAX_FILTER_SAMPLE_COUNT	128
+
+class FFilterSceneColorCS : public FGlobalShader
+{
+	DECLARE_GLOBAL_SHADER(FFilterSceneColorCS);
+	SHADER_USE_PARAMETER_STRUCT(FFilterSceneColorCS, FGlobalShader);
+
+	static const uint32 ThreadGroupSize = 8;
+
+	using FPermutationDomain = TShaderPermutationDomain<>;
+
+	BEGIN_SHADER_PARAMETER_STRUCT(FParameters, )
+		SHADER_PARAMETER(uint32, SampleCount)
+		SHADER_PARAMETER(uint32, SrcMipIndex)
+		SHADER_PARAMETER(FIntPoint, MipResolution)
+		SHADER_PARAMETER(FVector4, BlurDirection)
+		SHADER_PARAMETER_RDG_TEXTURE_SRV(Texture2D, SourceTexture)
+		SHADER_PARAMETER_SAMPLER(SamplerState, SourceSampler)
+		SHADER_PARAMETER_RDG_TEXTURE_UAV(RWTexture2D, OutTextureMipColor)
+		SHADER_PARAMETER_ARRAY(FVector4, SampleOffsetsWeights, [MAX_FILTER_SAMPLE_COUNT])
+	END_SHADER_PARAMETER_STRUCT()
+
+public:
+	static bool ShouldCompilePermutation(const FGlobalShaderPermutationParameters& Parameters) { return GetMaxSupportedFeatureLevel(Parameters.Platform) >= ERHIFeatureLevel::SM5; }
+	static void ModifyCompilationEnvironment(const FGlobalShaderPermutationParameters& Parameters, FShaderCompilerEnvironment& OutEnvironment)
+	{
+		FGlobalShader::ModifyCompilationEnvironment(Parameters, OutEnvironment);
+		OutEnvironment.SetDefine(TEXT("THREADGROUP_SIZE"), ThreadGroupSize);
+		OutEnvironment.SetDefine(TEXT("FILTERCOLORCS"), 1);
+	}
+};
+IMPLEMENT_GLOBAL_SHADER(FFilterSceneColorCS, "/Engine/Private/DistortFiltering.usf", "FilterColorCS", SF_Compute);
+
+static void AddFilterSceneColorPass(FRDGBuilder& GraphBuilder, FFilterSceneColorCS::FParameters* PassParameters, const FViewInfo& View)
+{
+	FIntVector NumGroups = FIntVector::DivideAndRoundUp(
+		FIntVector(PassParameters->MipResolution.X, PassParameters->MipResolution.Y, 1),
+		FIntVector(FFilterSceneColorCS::ThreadGroupSize, FFilterSceneColorCS::ThreadGroupSize, 1));
+
+	FFilterSceneColorCS::FPermutationDomain PermutationVector;
+	TShaderMapRef<FFilterSceneColorCS> ComputeShader(View.ShaderMap, PermutationVector);
+
+	// Dispatch with GenerateMips: reading from a slice through SRV and writing into lower mip through UAV.
+	ClearUnusedGraphResources(ComputeShader, PassParameters);
+ 	GraphBuilder.AddPass(
+		Forward<FRDGEventName>(RDG_EVENT_NAME("FilterMipGen")),
+		PassParameters,
+		ERDGPassFlags::Compute,
+		[PassParameters, ComputeShader, NumGroups](FRHICommandList& RHICmdList)
+		{
+			FComputeShaderUtils::Dispatch(RHICmdList, ComputeShader, *PassParameters, NumGroups);
+		});
+}
+
+//////////////////////////////////////////////////////////////////////////
+
 void FDeferredShadingSceneRenderer::RenderDistortion(FRDGBuilder& GraphBuilder, FRDGTextureRef SceneColorTexture, FRDGTextureRef SceneDepthTexture)
 {
 	check(SceneDepthTexture);
@@ -253,6 +479,7 @@ void FDeferredShadingSceneRenderer::RenderDistortion(FRDGBuilder& GraphBuilder,
 
 	FRDGTextureRef DistortionTexture = nullptr;
 	FRDGTextureRef RoughnessScatterTexture = nullptr;
+	FRDGTextureRef TempSceneColorMipchainTexture = nullptr;
 	FGlobalShaderMap* ShaderMap = GetGlobalShaderMap(FeatureLevel);
 
 	// Create a texture of the scene color with MIP levels, unfiltered
@@ -262,11 +489,9 @@ void FDeferredShadingSceneRenderer::RenderDistortion(FRDGBuilder& GraphBuilder,
 	const bool bUseRoughRefraction = GetUseRoughRefraction();
 	if(bUseRoughRefraction)
 	{
-		// We keep the same format, for AddCopyTexturePass to not complain for now. Do something special to rely on PF_FloatR11G11B10 only.
-		EPixelFormat SceneColorFormat = SceneColorTexture->Desc.Format;	
-
 		for (int32 ViewIndex = 0, Num = Views.Num(); ViewIndex < Num; ++ViewIndex)
 		{
+			RDG_EVENT_SCOPE(GraphBuilder, "Rough Refraction View%d", ViewIndex);
 			// Create the mip chain for the scene color in a separated texture
 
 			const FViewInfo& View = Views[ViewIndex];
@@ -275,27 +500,24 @@ void FDeferredShadingSceneRenderer::RenderDistortion(FRDGBuilder& GraphBuilder,
 			// We do not use max and do not add 1 to the result, this to not go down to the lowest mip level as this is not required.
 			const int32 MipCount = FMath::Max((int32)1, (int32)FMath::CeilLogTwo(FMath::Min(SceneColorMip0Resolution.X, SceneColorMip0Resolution.Y)));
 
-			FRDGTextureDesc SceneColorMipchainDesc = FRDGTextureDesc::Create2D(SceneColorMip0Resolution, SceneColorFormat, FClearValueBinding::None,
+			FRDGTextureDesc SceneColorMipchainDesc = FRDGTextureDesc::Create2D(SceneColorMip0Resolution, PF_FloatR11G11B10, FClearValueBinding::None,
 				TexCreate_TargetArraySlicesIndependently | TexCreate_ShaderResource | TexCreate_UAV | TexCreate_RenderTargetable, MipCount);
 
 			FRDGTextureRef SceneColorMipchainTexture = GraphBuilder.CreateTexture(SceneColorMipchainDesc, TEXT("SceneColorMipchain"));
 			ViewsSceneColorMipchain[ViewIndex] = SceneColorMipchainTexture;
 
+			if (!TempSceneColorMipchainTexture || TempSceneColorMipchainTexture->Desc.Extent != SceneColorMipchainTexture->Desc.Extent)
+			{
+				// Only allocate a new temporary if it needs to have another resolution.
+				TempSceneColorMipchainTexture = GraphBuilder.CreateTexture(SceneColorMipchainDesc, TEXT("TempSceneColorMipchain"));
+			}
+
 			// Copy scene color into the first mip level
 
-			FRHICopyTextureInfo CopyInfo;
-			CopyInfo.SourcePosition.X = View.ViewRect.Min.X;
-			CopyInfo.SourcePosition.Y = View.ViewRect.Min.Y;
-			CopyInfo.Size.X = SceneColorMip0Resolution.X;
-			CopyInfo.Size.Y = SceneColorMip0Resolution.Y;
-			CopyInfo.Size.Z = 1;
-			CopyInfo.NumMips = 1;
-			AddCopyTexturePass(GraphBuilder, SceneColorTexture, SceneColorMipchainTexture, CopyInfo);
+			AddCopySceneColorPass(GraphBuilder, View, SceneColorTexture, SceneColorMipchainTexture);
 
 			// Now render the mip chain
-
-			FDownsampleSceneColorCS::FPermutationDomain PermutationVector;
-			TShaderMapRef<FDownsampleSceneColorCS> ComputeShader(GetGlobalShaderMap(FeatureLevel), PermutationVector);
+			// STRATA_TODO we could optimize that pass by doing one pass with a tile of 16x16 writing out the 8x8, 4x4, 2x2 and 1x1 down sampled output
 
 			FIntPoint SrcMipResolution = SceneColorMip0Resolution;
 			FIntPoint DstMipResolution = SceneColorMip0Resolution / 2;
@@ -305,27 +527,80 @@ void FDeferredShadingSceneRenderer::RenderDistortion(FRDGBuilder& GraphBuilder,
 				PassParameters->SrcMipIndex = DstMipIndex - 1;
 				PassParameters->SrcMipResolution = SrcMipResolution;
 				PassParameters->DstMipResolution = DstMipResolution;
-				PassParameters->SourceSampler = TStaticSamplerState<SF_Bilinear>::GetRHI();	// To average 4 texel with one sample
+				PassParameters->SourceSampler = TStaticSamplerState<SF_Bilinear>::GetRHI();
 
 				PassParameters->SourceTexture = GraphBuilder.CreateSRV(FRDGTextureSRVDesc::CreateForMipLevel(SceneColorMipchainTexture, DstMipIndex - 1));
 				PassParameters->OutTextureMipColor = GraphBuilder.CreateUAV(FRDGTextureUAVDesc(SceneColorMipchainTexture, DstMipIndex));
 
-				FIntVector NumGroups = FIntVector::DivideAndRoundUp(FIntVector(DstMipResolution.X, DstMipResolution.Y, 1), FIntVector(FDownsampleSceneColorCS::ThreadGroupSize, FDownsampleSceneColorCS::ThreadGroupSize, 1));
-
-				// Dispatch with GenerateMips: reading from a slice through SRV and writing into lower mip through UAV.
-				ClearUnusedGraphResources(ComputeShader, PassParameters);
-				GraphBuilder.AddPass(
-					Forward<FRDGEventName>(RDG_EVENT_NAME("DistoMipGen")),
-					PassParameters,
-					ERDGPassFlags::Compute | ERDGPassFlags::NeverCull,
-					[PassParameters, ComputeShader, NumGroups](FRHICommandList& RHICmdList)
-					{
-						FComputeShaderUtils::Dispatch(RHICmdList, ComputeShader, *PassParameters, NumGroups);
-					});
+				AddDownsampleSceneColorPass(GraphBuilder, PassParameters, View);
 
 				SrcMipResolution = DstMipResolution;
 				DstMipResolution = DstMipResolution / 2;
 			}
+
+			// Now the horizontal blur
+			// STATA_TODO: check that compute overlap is working as all horizontal filtering steps can happen in parallel
+
+			const float KernelRadius = 16.0f; // The maximum sample count we will run. This is scaled down by FilterSizeScale.
+			const float FilterSizeScale = CVarRoughRefractionBlurScale.GetValueOnRenderThread();
+			const float CrossCenterWeight = CVarRoughRefractionCenterWeight.GetValueOnRenderThread(); // This must be 0 for the sharp image to not affect the output, the blurring dominates.
+
+			FIntPoint MipResolution = SceneColorMip0Resolution / 2;
+			for (int32 MipIndex = 1; MipIndex < MipCount; MipIndex++)
+			{
+				FVector2D OffsetAndWeight[MAX_FILTER_SAMPLE_COUNT];
+				uint32 SampleCount = FGaussianFiltering::Compute1DGaussianFilterKernel(OffsetAndWeight, MAX_FILTER_SAMPLE_COUNT, KernelRadius, CrossCenterWeight, FilterSizeScale);
+				const FVector2D InverseFilterTextureExtent(1.0f / static_cast<float>(MipResolution.X), 1.0f / static_cast<float>(MipResolution.Y));
+
+				FFilterSceneColorCS::FParameters* PassParameters = GraphBuilder.AllocParameters<FFilterSceneColorCS::FParameters>();
+				PassParameters->SrcMipIndex = MipIndex;
+				PassParameters->MipResolution = MipResolution;
+				PassParameters->SourceSampler = TStaticSamplerState<SF_Bilinear>::GetRHI();
+
+				PassParameters->SourceTexture = GraphBuilder.CreateSRV(FRDGTextureSRVDesc::CreateForMipLevel(SceneColorMipchainTexture, MipIndex));
+				PassParameters->OutTextureMipColor = GraphBuilder.CreateUAV(FRDGTextureUAVDesc(TempSceneColorMipchainTexture, MipIndex));
+
+				PassParameters->SampleCount = SampleCount;
+				PassParameters->BlurDirection = FVector4(1.0f, 0.0f, 0.0f, 0.0f);
+				for (uint32 i = 0; i < SampleCount; ++i)
+				{
+					PassParameters->SampleOffsetsWeights[i] = FVector4(InverseFilterTextureExtent.X * OffsetAndWeight[i].X, OffsetAndWeight[i].Y);
+				}
+
+				AddFilterSceneColorPass(GraphBuilder, PassParameters, View);
+
+				MipResolution = MipResolution / 2;
+			}
+
+			// Now the vertical blur
+			// STATA_TODO: check that compute overlap is working as all vertical filtering steps can happen in parallel
+
+			MipResolution = SceneColorMip0Resolution / 2;
+			for (int32 MipIndex = 1; MipIndex < MipCount; MipIndex++)
+			{
+				FVector2D OffsetAndWeight[MAX_FILTER_SAMPLE_COUNT];
+				uint32 SampleCount = FGaussianFiltering::Compute1DGaussianFilterKernel(OffsetAndWeight, MAX_FILTER_SAMPLE_COUNT, KernelRadius, CrossCenterWeight, FilterSizeScale);
+				const FVector2D InverseFilterTextureExtent(1.0f / static_cast<float>(MipResolution.X), 1.0f / static_cast<float>(MipResolution.Y));
+
+				FFilterSceneColorCS::FParameters* PassParameters = GraphBuilder.AllocParameters<FFilterSceneColorCS::FParameters>();
+				PassParameters->SrcMipIndex = MipIndex;
+				PassParameters->MipResolution = MipResolution;
+				PassParameters->SourceSampler = TStaticSamplerState<SF_Bilinear>::GetRHI();
+
+				PassParameters->SourceTexture = GraphBuilder.CreateSRV(FRDGTextureSRVDesc::CreateForMipLevel(TempSceneColorMipchainTexture, MipIndex));
+				PassParameters->OutTextureMipColor = GraphBuilder.CreateUAV(FRDGTextureUAVDesc(SceneColorMipchainTexture, MipIndex));
+
+				PassParameters->SampleCount = SampleCount;
+				PassParameters->BlurDirection = FVector4(0.0f, 1.0f, 0.0f, 0.0f);
+				for (uint32 i = 0; i < SampleCount; ++i)
+				{
+					PassParameters->SampleOffsetsWeights[i] = FVector4(InverseFilterTextureExtent.Y * OffsetAndWeight[i].X, OffsetAndWeight[i].Y);
+				}
+
+				AddFilterSceneColorPass(GraphBuilder, PassParameters, View);
+
+				MipResolution = MipResolution / 2;
+			}
 		}
 	}
 
