Fixed compilation error on GLES platofrms

diff --git a/Engine/Shaders/Private/DistortFiltering.usf b/Engine/Shaders/Private/DistortFiltering.usf
--- a/Engine/Shaders/Private/DistortFiltering.usf
+++ b/Engine/Shaders/Private/DistortFiltering.usf
@@ -11,11 +11,12 @@
 #ifdef COPYSCENECOLORTEXTUREPS
 
 Texture2D SceneColorTexture;
+SamplerState SceneColorSampler;
 
 void CopySceneColorTexturePS(FScreenVertexOutput Input, out float4 OutColor : SV_Target0)
 {
 	float2 TextureUV = (Input.UV * View.ViewSizeAndInvSize.xy + View.ViewRectMin.xy) * View.BufferSizeAndInvSize.zw;
-	float4 LinearColor = Texture2DSample(SceneColorTexture, GlobalPointClampedSampler, TextureUV);
+	float4 LinearColor = Texture2DSample(SceneColorTexture, SceneColorSampler, TextureUV);
 	OutColor = float4(LinearColor.rgb, 0);
 }
 
diff --git a/Engine/Source/Runtime/Renderer/Private/DistortionRendering.cpp b/Engine/Source/Runtime/Renderer/Private/DistortionRendering.cpp
--- a/Engine/Source/Runtime/Renderer/Private/DistortionRendering.cpp
+++ b/Engine/Source/Runtime/Renderer/Private/DistortionRendering.cpp
@@ -236,6 +236,7 @@ public:
 	BEGIN_SHADER_PARAMETER_STRUCT(FParameters, )
 		SHADER_PARAMETER_STRUCT_REF(FViewUniformShaderParameters, View)
 		SHADER_PARAMETER_RDG_TEXTURE(Texture2D, SceneColorTexture)
+		SHADER_PARAMETER_SAMPLER(SamplerState, SceneColorSampler)
 		RENDER_TARGET_BINDING_SLOTS()
 	END_SHADER_PARAMETER_STRUCT()
 
@@ -262,6 +263,7 @@ static void AddCopySceneColorPass(FRDGBuilder& GraphBuilder, const FViewInfo& Vi
 	auto* PassParameters = GraphBuilder.AllocParameters<FCopySceneColorTexturePS::FParameters>();
 	PassParameters->View = View.ViewUniformBuffer;
 	PassParameters->SceneColorTexture = SceneColorTexture;
+	PassParameters->SceneColorSampler = TStaticSamplerState<SF_Point>::GetRHI();
 	PassParameters->RenderTargets[0] = FRenderTargetBinding(SceneColorCopyTexture, ERenderTargetLoadAction::ENoAction);
 
 	const FScreenPassTextureViewport InputViewport(SceneColorTexture, View.ViewRect);
