Fix crash when removing virtual texture callbacks no callbacks are found.

diff --git a/Engine/Source/Runtime/Renderer/Private/VT/VirtualTextureProducer.cpp b/Engine/Source/Runtime/Renderer/Private/VT/VirtualTextureProducer.cpp
--- a/Engine/Source/Runtime/Renderer/Private/VT/VirtualTextureProducer.cpp
+++ b/Engine/Source/Runtime/Renderer/Private/VT/VirtualTextureProducer.cpp
@@ -187,25 +187,29 @@ uint32 FVirtualTextureProducerCollection::RemoveAllCallbacks(const void* Baton)
 	check(Baton);
 
 	uint32 NumRemoved = 0u;
-	for (int32 CallbackIndex : CallbacksMap.FindChecked(Baton))
+	TArray<uint32>* CallbackIndices = CallbacksMap.Find(Baton);
+	if (CallbackIndices != nullptr)
 	{
-		FCallbackEntry& Callback = Callbacks[CallbackIndex];
-		check(Callback.Baton == Baton);
-		check(Callback.DestroyedFunction);
-		Callback.DestroyedFunction = nullptr;
-		Callback.Baton = nullptr;
-		Callback.OwnerHandle = FVirtualTextureProducerHandle();
-
-		// If callback is already pending, we can't move it back to free list, or we risk corrupting the pending list while it's being iterated
-		// Setting DestroyedFunction to nullptr here will ensure callback is no longer invoked, and it will be moved to free list later when it's removed from pending list
-		if (!Callback.bPending)
+		for (int32 CallbackIndex : *CallbackIndices)
 		{
-			Callback.PackedFlags = 0u;
-			ReleaseCallback(CallbackIndex);
+			FCallbackEntry& Callback = Callbacks[CallbackIndex];
+			check(Callback.Baton == Baton);
+			check(Callback.DestroyedFunction);
+			Callback.DestroyedFunction = nullptr;
+			Callback.Baton = nullptr;
+			Callback.OwnerHandle = FVirtualTextureProducerHandle();
+
+			// If callback is already pending, we can't move it back to free list, or we risk corrupting the pending list while it's being iterated
+			// Setting DestroyedFunction to nullptr here will ensure callback is no longer invoked, and it will be moved to free list later when it's removed from pending list
+			if (!Callback.bPending)
+			{
+				Callback.PackedFlags = 0u;
+				ReleaseCallback(CallbackIndex);
+			}
+			++NumRemoved;
 		}
-		++NumRemoved;
+		CallbacksMap.Remove(Baton);
 	}
-	CallbacksMap.Remove(Baton);
 	return NumRemoved;
 }
 
