Parallel occlusion memory stomp fixes.
Since pointers to items inside InsertPrimitiveOcclusionHistory are stored, we cannot resize that array. To calculate the number of items that will be added, we have to count up all the visible items and any subqueries they might have. This incurs an iteration of the visibility list but hasn't incurred a drastic increase in time as measured in a few replays.

diff --git a/Engine/Source/Runtime/Renderer/Private/SceneVisibility.cpp b/Engine/Source/Runtime/Renderer/Private/SceneVisibility.cpp
--- a/Engine/Source/Runtime/Renderer/Private/SceneVisibility.cpp
+++ b/Engine/Source/Runtime/Renderer/Private/SceneVisibility.cpp
@@ -281,7 +281,7 @@ DECLARE_CYCLE_STAT(TEXT("After Occlusion Readback"), STAT_CLMM_AfterOcclusionRea
  * @param View - The view for which to update.
  * @param bVisible - Whether the primitive should be visible in the view.
  */
-static void UpdatePrimitiveFadingStateHelper(FPrimitiveFadingState& FadingState, FViewInfo& View, bool bVisible)
+static void UpdatePrimitiveFadingStateHelper(FPrimitiveFadingState& FadingState, const FViewInfo& View, bool bVisible)
 {
 	if (FadingState.bValid)
 	{
@@ -830,8 +830,6 @@ static void FetchVisibilityForPrimitives_Range(FVisForPrimParams& Params, FGloba
 	const bool bSubmitQueries		= Params.bSubmitQueries;
 	const bool bHZBOcclusion		= Params.bHZBOcclusion;
 
-	const TBitArray<>& PrimitivesAlwaysVisible = Scene->PrimitivesAlwaysVisible;
-
 	const float PrimitiveProbablyVisibleTime = GEngine->PrimitiveProbablyVisibleTime;
 
 	FSceneViewState* ViewState = (FSceneViewState*)View.State;
@@ -872,6 +870,7 @@ static void FetchVisibilityForPrimitives_Range(FVisForPrimParams& Params, FGloba
 	const FVector ViewOrigin = View.ViewMatrices.GetViewOrigin();
 
 	const int32 ReserveAmount = NumToProcess;
+	int32 NumQueriesToReserve = NumToProcess;
 	if (!bSingleThreaded)
 	{		
 		check(InsertPrimitiveOcclusionHistory);
@@ -879,8 +878,44 @@ static void FetchVisibilityForPrimitives_Range(FVisForPrimParams& Params, FGloba
 		check(HZBBoundsToAdd);
 		check(QueriesToAdd);
 
+		// We need to calculuate the actual number of queries to reserve since the pointers to InsertPrimitiveOcclusionHistory need to be preserved.
+		if (GAllowSubPrimitiveQueries && !View.bDisableQuerySubmissions)
+		{
+			NumQueriesToReserve = 0;
+			int32 NumProcessed = 0;
+#if BALANCE_LOAD
+			for (FSceneSetBitIterator BitIt(View.PrimitiveVisibilityMap, StartIndex); BitIt && (NumProcessed < NumToProcess); ++BitIt, ++NumProcessed)
+#else
+			for (TBitArray<SceneRenderingBitArrayAllocator>::FIterator BitIt(View.PrimitiveVisibilityMap, StartIndex); BitIt && (NumProcessed < NumToProcess); ++BitIt, ++NumProcessed)
+#endif
+			{
+#if !BALANCE_LOAD		
+				if (!View.PrimitiveVisibilityMap.AccessCorrespondingBit(BitIt))
+				{
+					continue;
+				}
+#endif
+
+				const uint8 OcclusionFlags = Scene->PrimitiveOcclusionFlags[BitIt.GetIndex()];
+
+				if (Scene->PrimitivesAlwaysVisible[BitIt.GetIndex()] || (OcclusionFlags & EOcclusionFlags::CanBeOccluded) == 0)
+				{
+					continue;
+				}
+
+				if ((OcclusionFlags & EOcclusionFlags::HasSubprimitiveQueries) != 0)
+				{
+					NumQueriesToReserve += Scene->Primitives[BitIt.GetIndex()]->Proxy->GetOcclusionQueries(&View)->Num();
+				}
+				else
+				{
+					NumQueriesToReserve++;
+				}
+			}
+		}
+
 		//avoid doing reallocs as much as possible.  Unlikely to make an entry per processed element.		
-		InsertPrimitiveOcclusionHistory->Reserve(ReserveAmount);
+		InsertPrimitiveOcclusionHistory->Reserve(NumQueriesToReserve);
 		QueriesToRelease->Reserve(ReserveAmount);
 		HZBBoundsToAdd->Reserve(ReserveAmount);
 		QueriesToAdd->Reserve(ReserveAmount);
@@ -901,10 +936,6 @@ static void FetchVisibilityForPrimitives_Range(FVisForPrimParams& Params, FGloba
 		for (TBitArray<SceneRenderingBitArrayAllocator>::FIterator BitIt(View.PrimitiveVisibilityMap, StartIndex); BitIt && (NumProcessed < NumToProcess); ++BitIt, ++NumProcessed)
 #endif
 		{
-			const bool bAlwaysVisible = PrimitivesAlwaysVisible[BitIt.GetIndex()];
-			uint8 OcclusionFlags = Scene->PrimitiveOcclusionFlags[BitIt.GetIndex()];
-			bool bCanBeOccluded = !bAlwaysVisible && (OcclusionFlags & EOcclusionFlags::CanBeOccluded) != 0;
-
 #if !BALANCE_LOAD		
 			if (!View.PrimitiveVisibilityMap.AccessCorrespondingBit(BitIt))
 			{
@@ -912,7 +943,9 @@ static void FetchVisibilityForPrimitives_Range(FVisForPrimParams& Params, FGloba
 			}
 #endif
 
-			if (!bCanBeOccluded)
+			const uint8 OcclusionFlags = Scene->PrimitiveOcclusionFlags[BitIt.GetIndex()];
+
+			if (Scene->PrimitivesAlwaysVisible[BitIt.GetIndex()] || (OcclusionFlags & EOcclusionFlags::CanBeOccluded) == 0)
 			{
 				View.PrimitiveDefinitelyUnoccludedMap.AccessCorrespondingBit(BitIt) = true;
 				continue;
@@ -922,6 +955,7 @@ static void FetchVisibilityForPrimitives_Range(FVisForPrimParams& Params, FGloba
 			const bool bCanAllocPrimHistory = bSingleThreaded || InsertPrimitiveOcclusionHistory->Num() < InsertPrimitiveOcclusionHistory->Max();
 
 #if WITH_EDITOR
+			bool bCanBeOccluded = true;
 			if (GIsEditor)
 			{
 				if (Scene->PrimitivesSelected[BitIt.GetIndex()])
@@ -930,6 +964,8 @@ static void FetchVisibilityForPrimitives_Range(FVisForPrimParams& Params, FGloba
 					bCanBeOccluded = false;
 				}
 			}
+#else
+			constexpr bool bCanBeOccluded = true;
 #endif
 
 			int32 NumSubQueries = 1;
@@ -939,7 +975,7 @@ static void FetchVisibilityForPrimitives_Range(FVisForPrimParams& Params, FGloba
 			check(Params.SubIsOccluded);
 			TArray<bool>& SubIsOccluded = *Params.SubIsOccluded;
 			int32 SubIsOccludedStart = SubIsOccluded.Num();
-			if ((OcclusionFlags & EOcclusionFlags::HasSubprimitiveQueries) && GAllowSubPrimitiveQueries && !View.bDisableQuerySubmissions && !bAlwaysVisible)
+			if ((OcclusionFlags & EOcclusionFlags::HasSubprimitiveQueries) && GAllowSubPrimitiveQueries && !View.bDisableQuerySubmissions)
 			{
 				FPrimitiveSceneProxy* Proxy = Scene->Primitives[BitIt.GetIndex()]->Proxy;
 				SubBounds = Proxy->GetOcclusionQueries(&View);
@@ -987,7 +1023,7 @@ static void FetchVisibilityForPrimitives_Range(FVisForPrimParams& Params, FGloba
 
 					// Flag the primitive's occlusion state as indefinite, which will force it to be queried this frame.
 					// The exception is if the primitive isn't occludable, in which case we know that it's definitely unoccluded.
-					bOcclusionStateIsDefinite = bCanBeOccluded ? false : true;
+					bOcclusionStateIsDefinite = !bCanBeOccluded;
 				}
 				else
 				{
@@ -1230,7 +1266,6 @@ static void FetchVisibilityForPrimitives_Range(FVisForPrimParams& Params, FGloba
 									}
 									else
 									{
-										check(GRHIMaximumReccommendedOustandingOcclusionQueries < MAX_int32); // it would be fairly easy to set up this path to optimize when there are a limited number, but it hasn't been done yet
 										QueriesToAdd->Emplace(PrimitiveOcclusionHistory, BoundOrigin, BoundExtent, bGroupedQuery);
 									}
 								}
@@ -1303,7 +1338,7 @@ static void FetchVisibilityForPrimitives_Range(FVisForPrimParams& Params, FGloba
 	}
 	check(NumTotalDefUnoccluded == View.PrimitiveDefinitelyUnoccludedMap.Num());
 	check(NumTotalPrims == View.PrimitiveVisibilityMap.Num());
-	check(!InsertPrimitiveOcclusionHistory || InsertPrimitiveOcclusionHistory->Num() <= ReserveAmount);
+	check(!InsertPrimitiveOcclusionHistory || InsertPrimitiveOcclusionHistory->Num() <= NumQueriesToReserve);
 	Params.NumOccludedPrims = NumOccludedPrimitives;	
 }
 
