Pass ownership of FHZBOcclusionTester::Primitives to RDG pass and let it release when the pass ends.  Fixes crashes due to thread race conditions when FHZBOcclusionTester kept ownership.

diff --git a/Engine/Source/Runtime/Renderer/Private/SceneOcclusion.cpp b/Engine/Source/Runtime/Renderer/Private/SceneOcclusion.cpp
--- a/Engine/Source/Runtime/Renderer/Private/SceneOcclusion.cpp
+++ b/Engine/Source/Runtime/Renderer/Private/SceneOcclusion.cpp
@@ -1010,7 +1010,7 @@ void FHZBOcclusionTester::Submit(FRDGBuilder& GraphBuilder, const FViewInfo& Vie
 			RDG_EVENT_NAME("UpdateTextures"),
 			PassParameters,
 			ERDGPassFlags::Copy,
-			[this, BoundsCenterTexture, BoundsExtentTexture](FRHICommandListImmediate& RHICmdList)
+			[Primitives = MoveTemp(Primitives), BoundsCenterTexture, BoundsExtentTexture](FRHICommandListImmediate& RHICmdList)
 		{
 			// Update in blocks to avoid large update
 			const uint32 BlockSize = 8;
@@ -1055,8 +1055,6 @@ void FHZBOcclusionTester::Submit(FRDGBuilder& GraphBuilder, const FViewInfo& Vie
 				RHIUpdateTexture2D((FRHITexture2D*)BoundsCenterTexture->GetRHI(), 0, Region, BlockStride, (uint8*)CenterBuffer);
 				RHIUpdateTexture2D((FRHITexture2D*)BoundsExtentTexture->GetRHI(), 0, Region, BlockStride, (uint8*)ExtentBuffer);
 			}
-
-			Primitives.Empty();
 		});
 	}
 
