Made the roughness buffer be higher precision as this is needed for smoother rough refraction transitions.

diff --git a/Engine/Source/Runtime/Renderer/Private/DistortionRendering.cpp b/Engine/Source/Runtime/Renderer/Private/DistortionRendering.cpp
--- a/Engine/Source/Runtime/Renderer/Private/DistortionRendering.cpp
+++ b/Engine/Source/Runtime/Renderer/Private/DistortionRendering.cpp
@@ -637,7 +637,7 @@ void FDeferredShadingSceneRenderer::RenderDistortion(FRDGBuilder& GraphBuilder,
 			RoughnessScatterTexture = GraphBuilder.CreateTexture(
 				FRDGTextureDesc::Create2D(
 					SceneDepthTexture->Desc.Extent,
-					PF_G8,
+					PF_R16F,
 					FClearValueBinding::Transparent,
 					GFastVRamConfig.Distortion | TexCreate_RenderTargetable | TexCreate_ShaderResource,
 					1,
