Strata rough refraction
- use computeutils to dispatch rdg compute
- added rdg GPU events
- tryclip is now allowed when final roughness is low

diff --git a/Engine/Shaders/Private/DistortAccumulatePS.usf b/Engine/Shaders/Private/DistortAccumulatePS.usf
--- a/Engine/Shaders/Private/DistortAccumulatePS.usf
+++ b/Engine/Shaders/Private/DistortAccumulatePS.usf
@@ -203,6 +203,11 @@ void Main(
 		OutThroughput *= LayerThroughput;
 	}
 
+	// Not write out roughness
+	OutRough = saturate(StrataLobVarianceToRoughness(WiLob.Sigma));	// STRATA_TODO accumulate variance and convert back to roughness in the apply hader
+	// Only allow clipping if roughness is close to 0.
+	const bool TryToClip = OutRough > 0.01f;
+
 	// Finally, negate the direction to make sure the new incombing lob is the refracted lob we want to take into account
 	WiLob.Mu = -WiLob.Mu;
 
@@ -211,7 +216,6 @@ void Main(
 	Normal = LobWorldDir;
 
 	// Compute UV distortion
-	const bool TryToClip = false; // We cannot clip because the roughness information is always needed.
 	float2 BufferUVDistortion = ComputeBufferUVDistortion(
 		MaterialParameters, PixelMaterialInputs, ResolvedView,
 		Normal, DistortionParams, ScreenUV, Refraction, TryToClip, EyeIndex);
@@ -219,8 +223,6 @@ void Main(
 	// Sample depth at distortion offset
 	float2 DistortBufferUV = ScreenUV + BufferUVDistortion;
 
-	OutRough = saturate(StrataLobVarianceToRoughness(WiLob.Sigma));	// STRATA_TODO accumulate variance and convert back to roughness in the apply hader
-
 
 #else // STRATA_ENABLED
 
diff --git a/Engine/Source/Runtime/Renderer/Private/DistortionRendering.cpp b/Engine/Source/Runtime/Renderer/Private/DistortionRendering.cpp
--- a/Engine/Source/Runtime/Renderer/Private/DistortionRendering.cpp
+++ b/Engine/Source/Runtime/Renderer/Private/DistortionRendering.cpp
@@ -314,14 +314,7 @@ static void AddDownsampleSceneColorPass(FRDGBuilder& GraphBuilder, FDownsampleSc
 
 	// Dispatch with GenerateMips: reading from a slice through SRV and writing into lower mip through UAV.
 	ClearUnusedGraphResources(ComputeShader, PassParameters);
-	GraphBuilder.AddPass(
-		Forward<FRDGEventName>(RDG_EVENT_NAME("DistoMipGen")),
-		PassParameters,
-		ERDGPassFlags::Compute,
-		[PassParameters, ComputeShader, NumGroups](FRHICommandList& RHICmdList)
-		{
-			FComputeShaderUtils::Dispatch(RHICmdList, ComputeShader, *PassParameters, NumGroups);
-		});
+	FComputeShaderUtils::AddPass(GraphBuilder, RDG_EVENT_NAME("DistoMipGen"), ComputeShader, PassParameters, NumGroups);
 }
 
 //////////////////////////////////////////////////////////////////////////
@@ -450,14 +443,7 @@ static void AddFilterSceneColorPass(FRDGBuilder& GraphBuilder, FFilterSceneColor
 
 	// Dispatch with GenerateMips: reading from a slice through SRV and writing into lower mip through UAV.
 	ClearUnusedGraphResources(ComputeShader, PassParameters);
- 	GraphBuilder.AddPass(
-		Forward<FRDGEventName>(RDG_EVENT_NAME("FilterMipGen")),
-		PassParameters,
-		ERDGPassFlags::Compute,
-		[PassParameters, ComputeShader, NumGroups](FRHICommandList& RHICmdList)
-		{
-			FComputeShaderUtils::Dispatch(RHICmdList, ComputeShader, *PassParameters, NumGroups);
-		});
+	FComputeShaderUtils::AddPass(GraphBuilder, RDG_EVENT_NAME("FilterMipGen"), ComputeShader, PassParameters, NumGroups);
 }
 
 //////////////////////////////////////////////////////////////////////////
@@ -494,7 +480,6 @@ void FDeferredShadingSceneRenderer::RenderDistortion(FRDGBuilder& GraphBuilder,
 		for (int32 ViewIndex = 0, Num = Views.Num(); ViewIndex < Num; ++ViewIndex)
 		{
 			RDG_EVENT_SCOPE(GraphBuilder, "Rough Refraction View%d", ViewIndex);
-			// Create the mip chain for the scene color in a separated texture
 
 			const FViewInfo& View = Views[ViewIndex];
 			const FIntPoint SceneColorMip0Resolution = View.ViewRect.Size();
@@ -516,28 +501,35 @@ void FDeferredShadingSceneRenderer::RenderDistortion(FRDGBuilder& GraphBuilder,
 
 			// Copy scene color into the first mip level
 
-			AddCopySceneColorPass(GraphBuilder, View, SceneColorTexture, SceneColorMipchainTexture);
+			{
+				RDG_EVENT_SCOPE(GraphBuilder, "CopySceneColor");
+				AddCopySceneColorPass(GraphBuilder, View, SceneColorTexture, SceneColorMipchainTexture);
+			}
 
 			// Now render the mip chain
 			// STRATA_TODO we could optimize that pass by doing one pass with a tile of 16x16 writing out the 8x8, 4x4, 2x2 and 1x1 down sampled output
 
-			FIntPoint SrcMipResolution = SceneColorMip0Resolution;
-			FIntPoint DstMipResolution = SceneColorMip0Resolution / 2;
-			for (int32 DstMipIndex = 1; DstMipIndex < MipCount; DstMipIndex++)
 			{
-				FDownsampleSceneColorCS::FParameters* PassParameters = GraphBuilder.AllocParameters<FDownsampleSceneColorCS::FParameters>();
-				PassParameters->SrcMipIndex = DstMipIndex - 1;
-				PassParameters->SrcMipResolution = SrcMipResolution;
-				PassParameters->DstMipResolution = DstMipResolution;
-				PassParameters->SourceSampler = TStaticSamplerState<SF_Bilinear>::GetRHI();
+				RDG_EVENT_SCOPE(GraphBuilder, "SceneColorMipChain");
 
-				PassParameters->SourceTexture = GraphBuilder.CreateSRV(FRDGTextureSRVDesc::CreateForMipLevel(SceneColorMipchainTexture, DstMipIndex - 1));
-				PassParameters->OutTextureMipColor = GraphBuilder.CreateUAV(FRDGTextureUAVDesc(SceneColorMipchainTexture, DstMipIndex));
+				FIntPoint SrcMipResolution = SceneColorMip0Resolution;
+				FIntPoint DstMipResolution = SceneColorMip0Resolution / 2;
+				for (int32 DstMipIndex = 1; DstMipIndex < MipCount; DstMipIndex++)
+				{
+					FDownsampleSceneColorCS::FParameters* PassParameters = GraphBuilder.AllocParameters<FDownsampleSceneColorCS::FParameters>();
+					PassParameters->SrcMipIndex = DstMipIndex - 1;
+					PassParameters->SrcMipResolution = SrcMipResolution;
+					PassParameters->DstMipResolution = DstMipResolution;
+					PassParameters->SourceSampler = TStaticSamplerState<SF_Bilinear>::GetRHI();
 
-				AddDownsampleSceneColorPass(GraphBuilder, PassParameters, View);
+					PassParameters->SourceTexture = GraphBuilder.CreateSRV(FRDGTextureSRVDesc::CreateForMipLevel(SceneColorMipchainTexture, DstMipIndex - 1));
+					PassParameters->OutTextureMipColor = GraphBuilder.CreateUAV(FRDGTextureUAVDesc(SceneColorMipchainTexture, DstMipIndex));
 
-				SrcMipResolution = DstMipResolution;
-				DstMipResolution = DstMipResolution / 2;
+					AddDownsampleSceneColorPass(GraphBuilder, PassParameters, View);
+
+					SrcMipResolution = DstMipResolution;
+					DstMipResolution = DstMipResolution / 2;
+				}
 			}
 
 			// Now the horizontal blur
@@ -547,61 +539,69 @@ void FDeferredShadingSceneRenderer::RenderDistortion(FRDGBuilder& GraphBuilder,
 			const float FilterSizeScale = CVarRoughRefractionBlurScale.GetValueOnRenderThread();
 			const float CrossCenterWeight = CVarRoughRefractionCenterWeight.GetValueOnRenderThread(); // This must be 0 for the sharp image to not affect the output, the blurring dominates.
 
-			FIntPoint MipResolution = SceneColorMip0Resolution / 2;
-			for (int32 MipIndex = 1; MipIndex < MipCount; MipIndex++)
 			{
-				FVector2D OffsetAndWeight[MAX_FILTER_SAMPLE_COUNT];
-				uint32 SampleCount = FGaussianFiltering::Compute1DGaussianFilterKernel(OffsetAndWeight, MAX_FILTER_SAMPLE_COUNT, KernelRadius, CrossCenterWeight, FilterSizeScale);
-				const FVector2D InverseFilterTextureExtent(1.0f / static_cast<float>(MipResolution.X), 1.0f / static_cast<float>(MipResolution.Y));
+				RDG_EVENT_SCOPE(GraphBuilder, "SceneColorMipHBlur");
 
-				FFilterSceneColorCS::FParameters* PassParameters = GraphBuilder.AllocParameters<FFilterSceneColorCS::FParameters>();
-				PassParameters->SrcMipIndex = MipIndex;
-				PassParameters->MipResolution = MipResolution;
-				PassParameters->SourceSampler = TStaticSamplerState<SF_Bilinear>::GetRHI();
+				FIntPoint MipResolution = SceneColorMip0Resolution / 2;
+				for (int32 MipIndex = 1; MipIndex < MipCount; MipIndex++)
+				{
+					FVector2D OffsetAndWeight[MAX_FILTER_SAMPLE_COUNT];
+					uint32 SampleCount = FGaussianFiltering::Compute1DGaussianFilterKernel(OffsetAndWeight, MAX_FILTER_SAMPLE_COUNT, KernelRadius, CrossCenterWeight, FilterSizeScale);
+					const FVector2D InverseFilterTextureExtent(1.0f / static_cast<float>(MipResolution.X), 1.0f / static_cast<float>(MipResolution.Y));
 
-				PassParameters->SourceTexture = GraphBuilder.CreateSRV(FRDGTextureSRVDesc::CreateForMipLevel(SceneColorMipchainTexture, MipIndex));
-				PassParameters->OutTextureMipColor = GraphBuilder.CreateUAV(FRDGTextureUAVDesc(TempSceneColorMipchainTexture, MipIndex));
+					FFilterSceneColorCS::FParameters* PassParameters = GraphBuilder.AllocParameters<FFilterSceneColorCS::FParameters>();
+					PassParameters->SrcMipIndex = MipIndex;
+					PassParameters->MipResolution = MipResolution;
+					PassParameters->SourceSampler = TStaticSamplerState<SF_Bilinear>::GetRHI();
 
-				PassParameters->SampleCount = SampleCount;
-				PassParameters->BlurDirection = FVector4(1.0f, 0.0f, 0.0f, 0.0f);
-				for (uint32 i = 0; i < SampleCount; ++i)
-				{
-					PassParameters->SampleOffsetsWeights[i] = FVector4(InverseFilterTextureExtent.X * OffsetAndWeight[i].X, OffsetAndWeight[i].Y);
-				}
+					PassParameters->SourceTexture = GraphBuilder.CreateSRV(FRDGTextureSRVDesc::CreateForMipLevel(SceneColorMipchainTexture, MipIndex));
+					PassParameters->OutTextureMipColor = GraphBuilder.CreateUAV(FRDGTextureUAVDesc(TempSceneColorMipchainTexture, MipIndex));
+
+					PassParameters->SampleCount = SampleCount;
+					PassParameters->BlurDirection = FVector4(1.0f, 0.0f, 0.0f, 0.0f);
+					for (uint32 i = 0; i < SampleCount; ++i)
+					{
+						PassParameters->SampleOffsetsWeights[i] = FVector4(InverseFilterTextureExtent.X * OffsetAndWeight[i].X, OffsetAndWeight[i].Y);
+					}
 
-				AddFilterSceneColorPass(GraphBuilder, PassParameters, View);
+					AddFilterSceneColorPass(GraphBuilder, PassParameters, View);
 
-				MipResolution = MipResolution / 2;
+					MipResolution = MipResolution / 2;
+				}
 			}
 
 			// Now the vertical blur
 			// STATA_TODO: check that compute overlap is working as all vertical filtering steps can happen in parallel
 
-			MipResolution = SceneColorMip0Resolution / 2;
-			for (int32 MipIndex = 1; MipIndex < MipCount; MipIndex++)
 			{
-				FVector2D OffsetAndWeight[MAX_FILTER_SAMPLE_COUNT];
-				uint32 SampleCount = FGaussianFiltering::Compute1DGaussianFilterKernel(OffsetAndWeight, MAX_FILTER_SAMPLE_COUNT, KernelRadius, CrossCenterWeight, FilterSizeScale);
-				const FVector2D InverseFilterTextureExtent(1.0f / static_cast<float>(MipResolution.X), 1.0f / static_cast<float>(MipResolution.Y));
+				RDG_EVENT_SCOPE(GraphBuilder, "SceneColorMipVBlur");
+
+				FIntPoint MipResolution = SceneColorMip0Resolution / 2;
+				for (int32 MipIndex = 1; MipIndex < MipCount; MipIndex++)
+				{
+					FVector2D OffsetAndWeight[MAX_FILTER_SAMPLE_COUNT];
+					uint32 SampleCount = FGaussianFiltering::Compute1DGaussianFilterKernel(OffsetAndWeight, MAX_FILTER_SAMPLE_COUNT, KernelRadius, CrossCenterWeight, FilterSizeScale);
+					const FVector2D InverseFilterTextureExtent(1.0f / static_cast<float>(MipResolution.X), 1.0f / static_cast<float>(MipResolution.Y));
 
-				FFilterSceneColorCS::FParameters* PassParameters = GraphBuilder.AllocParameters<FFilterSceneColorCS::FParameters>();
-				PassParameters->SrcMipIndex = MipIndex;
-				PassParameters->MipResolution = MipResolution;
-				PassParameters->SourceSampler = TStaticSamplerState<SF_Bilinear>::GetRHI();
+					FFilterSceneColorCS::FParameters* PassParameters = GraphBuilder.AllocParameters<FFilterSceneColorCS::FParameters>();
+					PassParameters->SrcMipIndex = MipIndex;
+					PassParameters->MipResolution = MipResolution;
+					PassParameters->SourceSampler = TStaticSamplerState<SF_Bilinear>::GetRHI();
 
-				PassParameters->SourceTexture = GraphBuilder.CreateSRV(FRDGTextureSRVDesc::CreateForMipLevel(TempSceneColorMipchainTexture, MipIndex));
-				PassParameters->OutTextureMipColor = GraphBuilder.CreateUAV(FRDGTextureUAVDesc(SceneColorMipchainTexture, MipIndex));
+					PassParameters->SourceTexture = GraphBuilder.CreateSRV(FRDGTextureSRVDesc::CreateForMipLevel(TempSceneColorMipchainTexture, MipIndex));
+					PassParameters->OutTextureMipColor = GraphBuilder.CreateUAV(FRDGTextureUAVDesc(SceneColorMipchainTexture, MipIndex));
 
-				PassParameters->SampleCount = SampleCount;
-				PassParameters->BlurDirection = FVector4(0.0f, 1.0f, 0.0f, 0.0f);
-				for (uint32 i = 0; i < SampleCount; ++i)
-				{
-					PassParameters->SampleOffsetsWeights[i] = FVector4(InverseFilterTextureExtent.Y * OffsetAndWeight[i].X, OffsetAndWeight[i].Y);
-				}
+					PassParameters->SampleCount = SampleCount;
+					PassParameters->BlurDirection = FVector4(0.0f, 1.0f, 0.0f, 0.0f);
+					for (uint32 i = 0; i < SampleCount; ++i)
+					{
+						PassParameters->SampleOffsetsWeights[i] = FVector4(InverseFilterTextureExtent.Y * OffsetAndWeight[i].X, OffsetAndWeight[i].Y);
+					}
 
-				AddFilterSceneColorPass(GraphBuilder, PassParameters, View);
+					AddFilterSceneColorPass(GraphBuilder, PassParameters, View);
 
-				MipResolution = MipResolution / 2;
+					MipResolution = MipResolution / 2;
+				}
 			}
 		}
 	}
