Starta rough refraction
- updated cvar name
- added high quality offset for smoother rough refraaction
- updated TryToClip logic...

diff --git a/Engine/Shaders/Private/DistortAccumulatePS.usf b/Engine/Shaders/Private/DistortAccumulatePS.usf
--- a/Engine/Shaders/Private/DistortAccumulatePS.usf
+++ b/Engine/Shaders/Private/DistortAccumulatePS.usf
@@ -206,7 +206,7 @@ void Main(
 	// Not write out roughness
 	OutRough = saturate(StrataLobVarianceToRoughness(WiLob.Sigma));	// STRATA_TODO accumulate variance and convert back to roughness in the apply hader
 	// Only allow clipping if roughness is close to 0.
-	const bool TryToClip = OutRough > 0.01f;
+	const bool TryToClip = OutRough < 0.01f;
 
 	// Finally, negate the direction to make sure the new incombing lob is the refracted lob we want to take into account
 	WiLob.Mu = -WiLob.Mu;
diff --git a/Engine/Source/Runtime/Renderer/Private/DistortionRendering.cpp b/Engine/Source/Runtime/Renderer/Private/DistortionRendering.cpp
--- a/Engine/Source/Runtime/Renderer/Private/DistortionRendering.cpp
+++ b/Engine/Source/Runtime/Renderer/Private/DistortionRendering.cpp
@@ -26,24 +26,30 @@ static TAutoConsoleVariable<int32> CVarDisableDistortion(
 	TEXT("Prevents distortion effects from rendering.  Saves a full-screen framebuffer's worth of memory."),
 	ECVF_Default);
 
-static TAutoConsoleVariable<int32> CVarRoughRefraction(
-	TEXT("r.RoughRefraction"),
+static TAutoConsoleVariable<int32> CVarRefractionBlur(
+	TEXT("r.Refraction.Blur"),
 	1,
 	TEXT("Prevent rough refraction from happening, i.e. blurring of the background."),
 	ECVF_Scalability | ECVF_RenderThreadSafe);
 
-static TAutoConsoleVariable<float> CVarRoughRefractionBlurScale(
-	TEXT("r.RoughRefraction.BlurScale"),
+static TAutoConsoleVariable<float> CVarRefractionBlurScale(
+	TEXT("r.Refraction.BlurScale"),
 	0.2f,
 	TEXT("Global scale the background blur amount after it is mapped form the surface back roughness/scattering amount."),
 	ECVF_RenderThreadSafe);
 
-static TAutoConsoleVariable<float> CVarRoughRefractionCenterWeight(
-	TEXT("r.RoughRefraction.CenterWeight"),
+static TAutoConsoleVariable<float> CVarRefractionBlurCenterWeight(
+	TEXT("r.Refraction.BlurCenterWeight"),
 	0.0f,
 	TEXT("The weight of the center sample. Value greater than 0 means the sharp image is more and more visible."),
 	ECVF_RenderThreadSafe);
 
+static TAutoConsoleVariable<int32> CVarRefractionOffsetQuality(
+	TEXT("r.Refraction.OffsetQuality"),
+	0,
+	TEXT("When enabled, the offset buffer is made float for higher quality. This is important to maintain the softness of the blurred scene buffer."),
+	ECVF_Scalability | ECVF_RenderThreadSafe);
+
 
 IMPLEMENT_STATIC_UNIFORM_BUFFER_STRUCT(FDistortionPassUniformParameters, "DistortionPass", SceneTextures);
 
@@ -90,7 +96,7 @@ TRDGUniformBufferRef<FDistortionPassUniformParameters> CreateDistortionPassUnifo
 
 static bool GetUseRoughRefraction()
 {
-	return Strata::IsStrataEnabled() && CVarRoughRefraction.GetValueOnRenderThread() > 0;
+	return Strata::IsStrataEnabled() && CVarRefractionBlur.GetValueOnRenderThread() > 0;
 }
 
 class FDistortionScreenPS : public FGlobalShader
@@ -536,8 +542,8 @@ void FDeferredShadingSceneRenderer::RenderDistortion(FRDGBuilder& GraphBuilder,
 			// STATA_TODO: check that compute overlap is working as all horizontal filtering steps can happen in parallel
 
 			const float KernelRadius = 16.0f; // The maximum sample count we will run. This is scaled down by FilterSizeScale.
-			const float FilterSizeScale = CVarRoughRefractionBlurScale.GetValueOnRenderThread();
-			const float CrossCenterWeight = CVarRoughRefractionCenterWeight.GetValueOnRenderThread(); // This must be 0 for the sharp image to not affect the output, the blurring dominates.
+			const float FilterSizeScale = CVarRefractionBlurScale.GetValueOnRenderThread();
+			const float CrossCenterWeight = CVarRefractionBlurCenterWeight.GetValueOnRenderThread(); // This must be 0 for the sharp image to not affect the output, the blurring dominates.
 
 			{
 				RDG_EVENT_SCOPE(GraphBuilder, "SceneColorMipHBlur");
@@ -624,7 +630,7 @@ void FDeferredShadingSceneRenderer::RenderDistortion(FRDGBuilder& GraphBuilder,
 		DistortionTexture = GraphBuilder.CreateTexture(
 			FRDGTextureDesc::Create2D(
 				SceneDepthTexture->Desc.Extent,
-				PF_B8G8R8A8,
+				CVarRefractionOffsetQuality.GetValueOnRenderThread() > 0 ? PF_FloatRGBA : PF_B8G8R8A8,
 				FClearValueBinding::Transparent,
 				GFastVRamConfig.Distortion | TexCreate_RenderTargetable | TexCreate_ShaderResource,
 				1,
