Added support for animated width to GroomCache.

diff --git a/Engine/Plugins/Importers/AlembicHairImporter/Source/AlembicHairTranslator/Private/AlembicHairTranslator.cpp b/Engine/Plugins/Importers/AlembicHairImporter/Source/AlembicHairTranslator/Private/AlembicHairTranslator.cpp
--- a/Engine/Plugins/Importers/AlembicHairImporter/Source/AlembicHairTranslator/Private/AlembicHairTranslator.cpp
+++ b/Engine/Plugins/Importers/AlembicHairImporter/Source/AlembicHairTranslator/Private/AlembicHairTranslator.cpp
@@ -453,7 +453,7 @@ static void ParseObject(const Alembic::Abc::IObject& InObject, int32 FrameIndex,
 			{
 				AnimInfo->Attributes |= EGroomCacheAttributes::Position;
 			}
-			else if (bIsWidthAnimated)
+			if (bIsWidthAnimated)
 			{
 				AnimInfo->Attributes |= EGroomCacheAttributes::Width;
 			}
diff --git a/Engine/Plugins/Runtime/HairStrands/Source/HairStrandsCore/Private/GroomBuilder.cpp b/Engine/Plugins/Runtime/HairStrands/Source/HairStrandsCore/Private/GroomBuilder.cpp
--- a/Engine/Plugins/Runtime/HairStrands/Source/HairStrandsCore/Private/GroomBuilder.cpp
+++ b/Engine/Plugins/Runtime/HairStrands/Source/HairStrandsCore/Private/GroomBuilder.cpp
@@ -193,10 +193,19 @@ namespace HairStrandsBuilder
 
 				for (uint32 PointIndex = 0; PointIndex < StrandCount; ++PointIndex, ++RadiusIterator, ++CoordUIterator)
 				{
-					*CoordUIterator /= *LengthIterator;
-					*RadiusIterator /= Curves.MaxRadius;
+					if (*LengthIterator > 0.0f)
+					{
+						*CoordUIterator /= *LengthIterator;
+					}
+					if (Curves.MaxRadius > 0.0f)
+					{
+						*RadiusIterator /= Curves.MaxRadius;
+					}
+				}
+				if (Curves.MaxLength > 0.0f)
+				{
+					*LengthIterator /= Curves.MaxLength;
 				}
-				*LengthIterator /= Curves.MaxLength;
 			}
 
 			if (bComputeRootUV)
diff --git a/Engine/Plugins/Runtime/HairStrands/Source/HairStrandsCore/Private/GroomComponent.cpp b/Engine/Plugins/Runtime/HairStrands/Source/HairStrandsCore/Private/GroomComponent.cpp
--- a/Engine/Plugins/Runtime/HairStrands/Source/HairStrandsCore/Private/GroomComponent.cpp
+++ b/Engine/Plugins/Runtime/HairStrands/Source/HairStrandsCore/Private/GroomComponent.cpp
@@ -1037,6 +1037,12 @@ public:
 			FVector InterpolatedCenter = FMath::Lerp(CurrentGroupData.BoundingBox.GetCenter(), NextGroupData.BoundingBox.GetCenter(), InterpolationFactor);
 			InterpolatedGroupData.BoundingBox = CurrentGroupData.BoundingBox.MoveTo(InterpolatedCenter) + NextGroupData.BoundingBox.MoveTo(InterpolatedCenter);
 
+			const bool bHasRadiusData = InterpolatedGroupData.VertexData.PointsRadius.Num() > 0;
+			if (bHasRadiusData)
+			{
+				InterpolatedGroupData.StrandData.MaxRadius = FMath::Lerp(CurrentGroupData.StrandData.MaxRadius, NextGroupData.StrandData.MaxRadius, InterpolationFactor);
+			}
+
 			// Parallel batched interpolation
 			const int32 BatchSize = 1024;
 			const int32 BatchCount = (NumVertices + BatchSize - 1) / BatchSize;
@@ -1052,6 +1058,14 @@ public:
 					const FVector& NextPosition = NextGroupData.VertexData.PointsPosition[VertexIndex];
 
 					InterpolatedGroupData.VertexData.PointsPosition[VertexIndex] = FMath::Lerp(CurrentPosition, NextPosition, InterpolationFactor);
+
+					if (bHasRadiusData)
+					{
+						const float CurrentRadius = CurrentGroupData.VertexData.PointsRadius[VertexIndex];
+						const float NextRadius = NextGroupData.VertexData.PointsRadius[VertexIndex];
+
+						InterpolatedGroupData.VertexData.PointsRadius[VertexIndex] = FMath::Lerp(CurrentRadius, NextRadius, InterpolationFactor);
+					}
 				}
 			});
 		}
diff --git a/Engine/Plugins/Runtime/HairStrands/Source/HairStrandsCore/Private/HairStrandsRendering.cpp b/Engine/Plugins/Runtime/HairStrands/Source/HairStrandsCore/Private/HairStrandsRendering.cpp
--- a/Engine/Plugins/Runtime/HairStrands/Source/HairStrandsCore/Private/HairStrandsRendering.cpp
+++ b/Engine/Plugins/Runtime/HairStrands/Source/HairStrandsCore/Private/HairStrandsRendering.cpp
@@ -23,6 +23,7 @@
 #include "GroomAsset.h" 
 #include "GroomManager.h"
 #include "GroomInstance.h"
+#include "SystemTextures.h"
 
 static float GHairRaytracingRadiusScale = 0;
 static FAutoConsoleVariableRef CVarHairRaytracingRadiusScale(TEXT("r.HairStrands.RaytracingRadiusScale"), GHairRaytracingRadiusScale, TEXT("Override the per instance scale factor for raytracing hair strands geometry (0: disabled, >0:enabled)"));
@@ -191,7 +192,9 @@ class FGroomCacheUpdatePassCS : public FGlobalShader
 	BEGIN_SHADER_PARAMETER_STRUCT(FParameters, )
 		SHADER_PARAMETER(uint32, ElementCount)
 		SHADER_PARAMETER(uint32, DispatchCountX)
+		SHADER_PARAMETER(int, bHasRadiusData)
 		SHADER_PARAMETER_RDG_BUFFER_SRV(Buffer, InAnimatedBuffer)
+		SHADER_PARAMETER_RDG_BUFFER_SRV(Buffer, InRadiusBuffer)
 		SHADER_PARAMETER_RDG_BUFFER_SRV(Buffer, InRestPoseBuffer)
 		SHADER_PARAMETER_RDG_BUFFER_SRV(Buffer, InDeformedOffsetBuffer)
 		SHADER_PARAMETER_RDG_BUFFER_UAV(RWBuffer, OutDeformedBuffer)
@@ -218,6 +221,8 @@ static void AddGroomCacheUpdatePass(
 	FRDGBufferUAVRef OutBuffer
 	)
 {
+	TRACE_CPUPROFILER_EVENT_SCOPE(AddGroomCacheUpdatePass);
+
 	if (ElementCount == 0) return;
 
 	const uint32 GroupSize = 64;
@@ -231,7 +236,7 @@ static void AddGroomCacheUpdatePass(
 	{
 		// Deformation are upload into a Buffer<float> as the original position are float3 which is is both
 		// 1) incompatible with structure buffer alignment (128bits), and 2) incompatible with vertex buffer 
-		// as R32G32B32_FLOAT format is not well supported for SRV accross HW.
+		// as R32G32B32_FLOAT format is not well supported for SRV across HW.
 		// So instead the positions are uploaded into vertex buffer Buffer<float>
 		VertexBuffer = CreateVertexBuffer(
 			GraphBuilder,
@@ -254,6 +259,28 @@ static void AddGroomCacheUpdatePass(
 	Parameters->InDeformedOffsetBuffer = InDeformedOffsetBuffer;
 	Parameters->OutDeformedBuffer = OutBuffer;
 
+	const uint32 RadiusDataCount = GroomCacheData.VertexData.PointsRadius.Num();
+	const bool bHasRadiusData = RadiusDataCount > 0;
+	FRDGBufferRef RadiusBuffer = nullptr;
+	if (bHasRadiusData)
+	{
+		const uint32 RadiusDataSizeInBytes = sizeof(float) * RadiusDataCount;
+		RadiusBuffer = CreateVertexBuffer(
+			GraphBuilder,
+			TEXT("GroomCache_RadiusBuffer"),
+			FRDGBufferDesc::CreateBufferDesc(sizeof(float), GroomCacheData.VertexData.PointsRadius.Num()),
+			GroomCacheData.VertexData.PointsRadius.GetData(),
+			RadiusDataSizeInBytes,
+			ERDGInitialDataFlags::None);
+
+		Parameters->InRadiusBuffer = GraphBuilder.CreateSRV(RadiusBuffer, PF_R32_FLOAT);
+	}
+	else
+	{
+		Parameters->InRadiusBuffer = GraphBuilder.CreateSRV(GSystemTextures.GetDefaultBuffer(GraphBuilder, 4), PF_R32_FLOAT);
+	}
+	Parameters->bHasRadiusData = bHasRadiusData ? 1 : 0;
+
 	TShaderMapRef<FGroomCacheUpdatePassCS> ComputeShader(ShaderMap);
 	FComputeShaderUtils::AddPass(
 		GraphBuilder,
@@ -1713,6 +1740,11 @@ void ComputeHairStrandsInterpolation(
 						RegisterAsSRV(GraphBuilder, Instance->Strands.DeformedResource->GetPositionOffsetBuffer(FHairStrandsDeformedResource::EFrameType::Current)),
 						Strands_DeformedPosition.UAV
 						);
+
+					if (GroomCacheGroupData.VertexData.PointsRadius.Num() > 0)
+					{
+						Instance->Strands.Modifier.HairWidth = GroomCacheGroupData.StrandData.MaxRadius * 2.f;
+					}
 				}
 
 				// 2.2 Update tangent data based on the deformed positions
diff --git a/Engine/Plugins/Runtime/HairStrands/Source/HairStrandsCore/Public/GroomCacheData.h b/Engine/Plugins/Runtime/HairStrands/Source/HairStrandsCore/Public/GroomCacheData.h
--- a/Engine/Plugins/Runtime/HairStrands/Source/HairStrandsCore/Public/GroomCacheData.h
+++ b/Engine/Plugins/Runtime/HairStrands/Source/HairStrandsCore/Public/GroomCacheData.h
@@ -13,6 +13,12 @@ enum class EGroomCacheAttributes : uint8
 	Position = 1,
 	Width = 1 << 1,
 	Color = 1 << 2,
+
+	// For display names
+	PositionWidth = (Position | Width) UMETA(DisplayName = "Position & Width"),
+	PositionColor = (Position | Color) UMETA(DisplayName = "Position & Color"),
+	WidthColor = (Position | Color) UMETA(DisplayName = "Width & Color"),
+	PositionWidthColor = (Position | Width | Color) UMETA(DisplayName = "Position, Width, Color")
 };
 
 ENUM_CLASS_FLAGS(EGroomCacheAttributes)
diff --git a/Engine/Shaders/Private/HairStrands/HairStrandsInterpolation.usf b/Engine/Shaders/Private/HairStrands/HairStrandsInterpolation.usf
--- a/Engine/Shaders/Private/HairStrands/HairStrandsInterpolation.usf
+++ b/Engine/Shaders/Private/HairStrands/HairStrandsInterpolation.usf
@@ -755,7 +755,9 @@ void MainCS(uint3 DispatchThreadId : SV_DispatchThreadID)
 
 uint ElementCount;
 uint DispatchCountX;
+uint bHasRadiusData;
 
+Buffer<float> InRadiusBuffer;
 Buffer<float> InAnimatedBuffer;
 Buffer<uint4> InRestPoseBuffer;
 Buffer<float4> InDeformedOffsetBuffer;
@@ -780,7 +782,16 @@ void MainCS(uint3 DispatchThreadId : SV_DispatchThreadID)
 	uint4 PackedRestPosition = InRestPoseBuffer[Index];
 	float3 HairPositionOffset = InDeformedOffsetBuffer[0];
 	float3 WorldVertexPosition = InAnimatedPosition - HairPositionOffset;
-	OutDeformedBuffer[Index] = uint4(f32tof16(WorldVertexPosition) , PackedRestPosition.w);
+
+	FHairControlPoint CP = UnpackHairControlPoint(PackedRestPosition, float3(0, 0, 0), 1, 1, 1);
+	CP.Position = WorldVertexPosition;
+
+	if (bHasRadiusData == 1)
+	{
+		CP.WorldRadius = InRadiusBuffer[Index];
+	}
+
+	OutDeformedBuffer[Index] = PackHairControlPoint(CP, float3(0, 0, 0), 1);
 }
 
 #endif // SHADER_GROOMCACHE_UPDATE
diff --git a/Engine/Shaders/Private/HairStrands/HairStrandsVertexFactoryCommon.ush b/Engine/Shaders/Private/HairStrands/HairStrandsVertexFactoryCommon.ush
--- a/Engine/Shaders/Private/HairStrands/HairStrandsVertexFactoryCommon.ush
+++ b/Engine/Shaders/Private/HairStrands/HairStrandsVertexFactoryCommon.ush
@@ -17,6 +17,17 @@ struct FHairControlPoint
 	uint   Type;
 };
 
+uint4 PackHairControlPoint(
+	FHairControlPoint CP,
+	float3 InVF_PositionOffset,
+	float InVF_Radius)
+{
+	uint PackedAttribute =  (uint(clamp(CP.UCoord * 0xFFu, 0u, 0xFFu)) << 8) |
+							(uint(clamp((CP.WorldRadius / InVF_Radius) * 63u, 0u, 63u)) << 2u) |
+							(CP.Type & 0x3);
+	return uint4(f32tof16(CP.Position - InVF_PositionOffset), PackedAttribute);
+}
+
 FHairControlPoint UnpackHairControlPoint(
 	uint4 InPackedData,
 	float3 InVF_PositionOffset,
diff --git a/Engine/Source/Runtime/Renderer/Private/SystemTextures.h b/Engine/Source/Runtime/Renderer/Private/SystemTextures.h
--- a/Engine/Source/Runtime/Renderer/Private/SystemTextures.h
+++ b/Engine/Source/Runtime/Renderer/Private/SystemTextures.h
@@ -148,52 +148,52 @@ public:
 	TRefCountPtr<IPooledRenderTarget> AsciiTexture;
 
 	// Create simple default texture
-	FRDGTextureRef GetWhiteDummy(FRDGBuilder& GraphBuilder) const;
-	FRDGTextureRef GetBlackDummy(FRDGBuilder& GraphBuilder) const;
-	FRDGTextureRef GetZeroUIntDummy(FRDGBuilder& GraphBuilder) const;
-	FRDGTextureRef GetZeroUShort4Dummy(FRDGBuilder& GraphBuilder) const;
-	FRDGTextureRef GetBlackAlphaOneDummy(FRDGBuilder& GraphBuilder) const;
-	FRDGTextureRef GetMaxFP16Depth(FRDGBuilder& GraphBuilder) const;
-	FRDGTextureRef GetDepthDummy(FRDGBuilder& GraphBuilder) const;
-	FRDGTextureRef GetStencilDummy(FRDGBuilder& GraphBuilder) const;
-	FRDGTextureRef GetGreenDummy(FRDGBuilder& GraphBuilder) const;
-	FRDGTextureRef GetDefaultNormal8Bit(FRDGBuilder& GraphBuilder) const;
-	FRDGTextureRef GetMidGreyDummy(FRDGBuilder& GraphBuilder) const;
-	FRDGTextureRef GetVolumetricBlackDummy(FRDGBuilder& GraphBuilder) const;
+	FRDGTextureRef RENDERER_API GetWhiteDummy(FRDGBuilder& GraphBuilder) const;
+	FRDGTextureRef RENDERER_API GetBlackDummy(FRDGBuilder& GraphBuilder) const;
+	FRDGTextureRef RENDERER_API GetZeroUIntDummy(FRDGBuilder& GraphBuilder) const;
+	FRDGTextureRef RENDERER_API GetZeroUShort4Dummy(FRDGBuilder& GraphBuilder) const;
+	FRDGTextureRef RENDERER_API GetBlackAlphaOneDummy(FRDGBuilder& GraphBuilder) const;
+	FRDGTextureRef RENDERER_API GetMaxFP16Depth(FRDGBuilder& GraphBuilder) const;
+	FRDGTextureRef RENDERER_API GetDepthDummy(FRDGBuilder& GraphBuilder) const;
+	FRDGTextureRef RENDERER_API GetStencilDummy(FRDGBuilder& GraphBuilder) const;
+	FRDGTextureRef RENDERER_API GetGreenDummy(FRDGBuilder& GraphBuilder) const;
+	FRDGTextureRef RENDERER_API GetDefaultNormal8Bit(FRDGBuilder& GraphBuilder) const;
+	FRDGTextureRef RENDERER_API GetMidGreyDummy(FRDGBuilder& GraphBuilder) const;
+	FRDGTextureRef RENDERER_API GetVolumetricBlackDummy(FRDGBuilder& GraphBuilder) const;
 
 	// Create default 2D texture (1x1) with specific format and initialize value 
-	FRDGTextureRef GetDefaultTexture2D(FRDGBuilder& GraphBuilder, EPixelFormat Format, float Value);
-	FRDGTextureRef GetDefaultTexture2D(FRDGBuilder& GraphBuilder, EPixelFormat Format, uint32 Value);
-	FRDGTextureRef GetDefaultTexture2D(FRDGBuilder& GraphBuilder, EPixelFormat Format, const FVector& Value);
-	FRDGTextureRef GetDefaultTexture2D(FRDGBuilder& GraphBuilder, EPixelFormat Format, const FVector4& Value);
-	FRDGTextureRef GetDefaultTexture2D(FRDGBuilder& GraphBuilder, EPixelFormat Format, const FUintVector4& Value);
-	FRDGTextureRef GetDefaultTexture2D(FRDGBuilder& GraphBuilder, EPixelFormat Format, const FClearValueBinding& Value);
+	FRDGTextureRef RENDERER_API GetDefaultTexture2D(FRDGBuilder& GraphBuilder, EPixelFormat Format, float Value);
+	FRDGTextureRef RENDERER_API GetDefaultTexture2D(FRDGBuilder& GraphBuilder, EPixelFormat Format, uint32 Value);
+	FRDGTextureRef RENDERER_API GetDefaultTexture2D(FRDGBuilder& GraphBuilder, EPixelFormat Format, const FVector& Value);
+	FRDGTextureRef RENDERER_API GetDefaultTexture2D(FRDGBuilder& GraphBuilder, EPixelFormat Format, const FVector4& Value);
+	FRDGTextureRef RENDERER_API GetDefaultTexture2D(FRDGBuilder& GraphBuilder, EPixelFormat Format, const FUintVector4& Value);
+	FRDGTextureRef RENDERER_API GetDefaultTexture2D(FRDGBuilder& GraphBuilder, EPixelFormat Format, const FClearValueBinding& Value);
 
 	// Create default 2D/3D/Cube/Array texture (1x1) with specific format and initialize value 
-	FRDGTextureRef GetDefaultTexture(FRDGBuilder& GraphBuilder, ETextureDimension Dimension, EPixelFormat Format, float Value);
-	FRDGTextureRef GetDefaultTexture(FRDGBuilder& GraphBuilder, ETextureDimension Dimension, EPixelFormat Format, uint32 Value);
-	FRDGTextureRef GetDefaultTexture(FRDGBuilder& GraphBuilder, ETextureDimension Dimension, EPixelFormat Format, const FVector2D& Value);
-	FRDGTextureRef GetDefaultTexture(FRDGBuilder& GraphBuilder, ETextureDimension Dimension, EPixelFormat Format, const FIntPoint& Value);
-	FRDGTextureRef GetDefaultTexture(FRDGBuilder& GraphBuilder, ETextureDimension Dimension, EPixelFormat Format, const FVector& Value);
-	FRDGTextureRef GetDefaultTexture(FRDGBuilder& GraphBuilder, ETextureDimension Dimension, EPixelFormat Format, const FVector4& Value);
-	FRDGTextureRef GetDefaultTexture(FRDGBuilder& GraphBuilder, ETextureDimension Dimension, EPixelFormat Format, const FUintVector4& Value);
-	FRDGTextureRef GetDefaultTexture(FRDGBuilder& GraphBuilder, ETextureDimension Dimension, EPixelFormat Format, const FClearValueBinding& Value);
+	FRDGTextureRef RENDERER_API GetDefaultTexture(FRDGBuilder& GraphBuilder, ETextureDimension Dimension, EPixelFormat Format, float Value);
+	FRDGTextureRef RENDERER_API GetDefaultTexture(FRDGBuilder& GraphBuilder, ETextureDimension Dimension, EPixelFormat Format, uint32 Value);
+	FRDGTextureRef RENDERER_API GetDefaultTexture(FRDGBuilder& GraphBuilder, ETextureDimension Dimension, EPixelFormat Format, const FVector2D& Value);
+	FRDGTextureRef RENDERER_API GetDefaultTexture(FRDGBuilder& GraphBuilder, ETextureDimension Dimension, EPixelFormat Format, const FIntPoint& Value);
+	FRDGTextureRef RENDERER_API GetDefaultTexture(FRDGBuilder& GraphBuilder, ETextureDimension Dimension, EPixelFormat Format, const FVector& Value);
+	FRDGTextureRef RENDERER_API GetDefaultTexture(FRDGBuilder& GraphBuilder, ETextureDimension Dimension, EPixelFormat Format, const FVector4& Value);
+	FRDGTextureRef RENDERER_API GetDefaultTexture(FRDGBuilder& GraphBuilder, ETextureDimension Dimension, EPixelFormat Format, const FUintVector4& Value);
+	FRDGTextureRef RENDERER_API GetDefaultTexture(FRDGBuilder& GraphBuilder, ETextureDimension Dimension, EPixelFormat Format, const FClearValueBinding& Value);
 
 	// Create default buffer initialize to zero.
-	FRDGBufferRef GetDefaultBuffer(FRDGBuilder& GraphBuilder, uint32 NumBytePerElement);
-	FRDGBufferRef GetDefaultStructuredBuffer(FRDGBuilder& GraphBuilder, uint32 NumBytePerElement);
+	FRDGBufferRef RENDERER_API GetDefaultBuffer(FRDGBuilder& GraphBuilder, uint32 NumBytePerElement);
+	FRDGBufferRef RENDERER_API GetDefaultStructuredBuffer(FRDGBuilder& GraphBuilder, uint32 NumBytePerElement);
 
 	// Create a default buffer initialized with a reference element.
-	FRDGBufferRef GetDefaultBuffer(FRDGBuilder& GraphBuilder, uint32 NumBytePerElement, float Value);
-	FRDGBufferRef GetDefaultBuffer(FRDGBuilder& GraphBuilder, uint32 NumBytePerElement, uint32 Value);
-	FRDGBufferRef GetDefaultBuffer(FRDGBuilder& GraphBuilder, uint32 NumBytePerElement, const FVector& Value);
-	FRDGBufferRef GetDefaultBuffer(FRDGBuilder& GraphBuilder, uint32 NumBytePerElement, const FVector4& Value);
-	FRDGBufferRef GetDefaultBuffer(FRDGBuilder& GraphBuilder, uint32 NumBytePerElement, const FUintVector4& Value);
-	FRDGBufferRef GetDefaultStructuredBuffer(FRDGBuilder& GraphBuilder, uint32 NumBytePerElement, float Value);
-	FRDGBufferRef GetDefaultStructuredBuffer(FRDGBuilder& GraphBuilder, uint32 NumBytePerElement, uint32 Value);
-	FRDGBufferRef GetDefaultStructuredBuffer(FRDGBuilder& GraphBuilder, uint32 NumBytePerElement, const FVector& Value);
-	FRDGBufferRef GetDefaultStructuredBuffer(FRDGBuilder& GraphBuilder, uint32 NumBytePerElement, const FVector4& Value);
-	FRDGBufferRef GetDefaultStructuredBuffer(FRDGBuilder& GraphBuilder, uint32 NumBytePerElement, const FUintVector4& Value);
+	FRDGBufferRef RENDERER_API GetDefaultBuffer(FRDGBuilder& GraphBuilder, uint32 NumBytePerElement, float Value);
+	FRDGBufferRef RENDERER_API GetDefaultBuffer(FRDGBuilder& GraphBuilder, uint32 NumBytePerElement, uint32 Value);
+	FRDGBufferRef RENDERER_API GetDefaultBuffer(FRDGBuilder& GraphBuilder, uint32 NumBytePerElement, const FVector& Value);
+	FRDGBufferRef RENDERER_API GetDefaultBuffer(FRDGBuilder& GraphBuilder, uint32 NumBytePerElement, const FVector4& Value);
+	FRDGBufferRef RENDERER_API GetDefaultBuffer(FRDGBuilder& GraphBuilder, uint32 NumBytePerElement, const FUintVector4& Value);
+	FRDGBufferRef RENDERER_API GetDefaultStructuredBuffer(FRDGBuilder& GraphBuilder, uint32 NumBytePerElement, float Value);
+	FRDGBufferRef RENDERER_API GetDefaultStructuredBuffer(FRDGBuilder& GraphBuilder, uint32 NumBytePerElement, uint32 Value);
+	FRDGBufferRef RENDERER_API GetDefaultStructuredBuffer(FRDGBuilder& GraphBuilder, uint32 NumBytePerElement, const FVector& Value);
+	FRDGBufferRef RENDERER_API GetDefaultStructuredBuffer(FRDGBuilder& GraphBuilder, uint32 NumBytePerElement, const FVector4& Value);
+	FRDGBufferRef RENDERER_API GetDefaultStructuredBuffer(FRDGBuilder& GraphBuilder, uint32 NumBytePerElement, const FUintVector4& Value);
 
 protected:
 	/** Maximum feature level that the textures have been initialized up to */
