Fix invalid voxel resources if no hair instance within a macro groups requires voxelization.

In this case dummy voxel resources are allocated and used.

diff --git a/Engine/Source/Runtime/Renderer/Private/HairStrands/HairStrandsVoxelization.cpp b/Engine/Source/Runtime/Renderer/Private/HairStrands/HairStrandsVoxelization.cpp
--- a/Engine/Source/Runtime/Renderer/Private/HairStrands/HairStrandsVoxelization.cpp
+++ b/Engine/Source/Runtime/Renderer/Private/HairStrands/HairStrandsVoxelization.cpp
@@ -1596,12 +1596,33 @@ void VoxelizeHairStrands(
 	FHairStrandsVoxelResources& VirtualVoxelResources = View.HairStrandsViewData.VirtualVoxelResources;
 	FHairStrandsMacroGroupResources& MacroGroupResources = View.HairStrandsViewData.MacroGroupResources;
 
+	// Simple early out to check if voxelization is needed
 	if (!IsHairStrandsVoxelizationEnable() || MacroGroupDatas.Num() == 0)
 	{
 		VirtualVoxelResources = AllocateDummyVirtualVoxelResources(GraphBuilder, View, MacroGroupDatas);
 		return;
 	}
 
+	// Detailed early out to check if voxelization is needed
+	bool bHasValidElementToVoxelize = false;
+	for (const FHairStrandsMacroGroupData& MacroGroup : MacroGroupDatas)
+	{
+		for (const FHairStrandsMacroGroupData::PrimitiveInfo& PrimitiveInfo : MacroGroup.PrimitivesInfos)
+		{
+			if (PrimitiveInfo.PublicDataPtr->DoesSupportVoxelization())
+			{
+				bHasValidElementToVoxelize = true;
+				break;
+			}
+		}
+	}
+
+	if (!bHasValidElementToVoxelize)
+	{
+		VirtualVoxelResources = AllocateDummyVirtualVoxelResources(GraphBuilder, View, MacroGroupDatas);
+		return;
+	}
+
 	DECLARE_GPU_STAT(HairStrandsVoxelization);
 	RDG_EVENT_SCOPE(GraphBuilder, "HairStrandsVoxelization");
 	RDG_GPU_STAT_SCOPE(GraphBuilder, HairStrandsVoxelization);
