Fix shadow flickering when groom assets not supporting voxelization are visible.

diff --git a/Engine/Shaders/Private/HairStrands/HairStrandsVoxelPageAllocation.usf b/Engine/Shaders/Private/HairStrands/HairStrandsVoxelPageAllocation.usf
--- a/Engine/Shaders/Private/HairStrands/HairStrandsVoxelPageAllocation.usf
+++ b/Engine/Shaders/Private/HairStrands/HairStrandsVoxelPageAllocation.usf
@@ -13,6 +13,7 @@ uint				TotalPageIndexCount; 	// This is the max page index count;
 uint				PageResolution; 		// Resolution of a page
 uint				MacroGroupCount;
 uint				IndirectDispatchGroupSize;
+uint				bDoesMacroGroupSupportVoxelization;
 
 // For testing parity with CPU version
 float3				CPU_MinAABB;
@@ -39,8 +40,9 @@ void AllocatePageIndex(uint2 DispatchThreadId : SV_DispatchThreadID)
 	bool bIsValid = MacroGroupId < MacroGroupCount;
 	if (bIsValid)
 	{
+		const bool bSupportVoxelization = (bDoesMacroGroupSupportVoxelization >> MacroGroupId) & 0x1;
 		Bound = ReadHairAABB(MacroGroupId, MacroGroupAABBBuffer);
-		if (any(Bound.Min > Bound.Max))
+		if (any(Bound.Min > Bound.Max) || !bSupportVoxelization)
 		{
 			Bound.Min = 0;
 			Bound.Max = 0;
diff --git a/Engine/Shaders/Private/HairStrands/HairStrandsVoxelPageCommon.ush b/Engine/Shaders/Private/HairStrands/HairStrandsVoxelPageCommon.ush
--- a/Engine/Shaders/Private/HairStrands/HairStrandsVoxelPageCommon.ush
+++ b/Engine/Shaders/Private/HairStrands/HairStrandsVoxelPageCommon.ush
@@ -80,6 +80,7 @@ struct FVirtualVoxelNodeDesc
 	uint3	PageIndexResolution;
 	uint3   VirtualResolution;
 	uint    PageIndexOffset;
+	bool    bIsValid;
 };
 
 FVirtualVoxelNodeDesc UnpackVoxelNode(FPackedVirtualVoxelNodeDesc In, uint InPageResolution)
@@ -90,6 +91,7 @@ FVirtualVoxelNodeDesc UnpackVoxelNode(FPackedVirtualVoxelNodeDesc In, uint InPag
 	Out.PageIndexResolution = uint3(In.PackedPageIndexResolution & 0xFF, (In.PackedPageIndexResolution >> 8) & 0xFF, (In.PackedPageIndexResolution >> 16) & 0xFF);
 	Out.VirtualResolution = Out.PageIndexResolution * InPageResolution;
 	Out.PageIndexOffset = In.PageIndexOffset;
+	Out.bIsValid = all(Out.PageIndexResolution != 0);
 	return Out;
 }
 
diff --git a/Engine/Shaders/Private/HairStrands/HairStrandsVoxelPageTraversal.ush b/Engine/Shaders/Private/HairStrands/HairStrandsVoxelPageTraversal.ush
--- a/Engine/Shaders/Private/HairStrands/HairStrandsVoxelPageTraversal.ush
+++ b/Engine/Shaders/Private/HairStrands/HairStrandsVoxelPageTraversal.ush
@@ -170,6 +170,10 @@ FHairTraversalResult ComputeHairCountVirtualVoxel(
 	const float MipLevelDistanceScale = 1.0f / InCommonDesc.VoxelWorldSize;
 
 	FHairTraversalResult Out = InitHairTraversalResult();
+	if (!InNodeDesc.bIsValid)
+	{
+		return Out;
+	}
 
 	int3 CurrentPageIndexCoord = -1;
 	bool bIsPageValid = false;
diff --git a/Engine/Source/Runtime/Renderer/Private/HairStrands/HairStrandsCluster.cpp b/Engine/Source/Runtime/Renderer/Private/HairStrands/HairStrandsCluster.cpp
--- a/Engine/Source/Runtime/Renderer/Private/HairStrands/HairStrandsCluster.cpp
+++ b/Engine/Source/Runtime/Renderer/Private/HairStrands/HairStrandsCluster.cpp
@@ -152,6 +152,11 @@ static void InternalUpdateMacroGroup(FHairStrandsMacroGroupData& MacroGroup, int
 		PrimitiveInfo.GroupIndex = HairGroupPublicData->GetGroupIndex();
 		PrimitiveInfo.PublicDataPtr = HairGroupPublicData;
 		check(PrimitiveInfo.GroupIndex < 32); // Sanity check
+
+		if (HairGroupPublicData->DoesSupportVoxelization())
+		{
+			MacroGroup.bSupportVoxelization = true;
+		}
 	}
 }
 
diff --git a/Engine/Source/Runtime/Renderer/Private/HairStrands/HairStrandsData.h b/Engine/Source/Runtime/Renderer/Private/HairStrands/HairStrandsData.h
--- a/Engine/Source/Runtime/Renderer/Private/HairStrands/HairStrandsData.h
+++ b/Engine/Source/Runtime/Renderer/Private/HairStrands/HairStrandsData.h
@@ -279,6 +279,7 @@ struct FHairStrandsMacroGroupData
 	FIntRect ScreenRect;
 	uint32 MacroGroupId;
 
+	bool bSupportVoxelization = false; // true if at least one of the Primitive requires voxelization
 	bool bNeedScatterSceneLighting = false;
 };
 
diff --git a/Engine/Source/Runtime/Renderer/Private/HairStrands/HairStrandsVoxelization.cpp b/Engine/Source/Runtime/Renderer/Private/HairStrands/HairStrandsVoxelization.cpp
--- a/Engine/Source/Runtime/Renderer/Private/HairStrands/HairStrandsVoxelization.cpp
+++ b/Engine/Source/Runtime/Renderer/Private/HairStrands/HairStrandsVoxelization.cpp
@@ -219,6 +219,7 @@ class FVoxelAllocatePageIndexCS : public FGlobalShader
 		SHADER_PARAMETER(uint32, PageResolution)
 		SHADER_PARAMETER(uint32, MacroGroupCount)
 		SHADER_PARAMETER(uint32, IndirectDispatchGroupSize)
+		SHADER_PARAMETER(uint32, bDoesMacroGroupSupportVoxelization)
 
 		SHADER_PARAMETER(FVector3f, CPU_MinAABB)
 		SHADER_PARAMETER(FVector3f, CPU_MaxAABB)
@@ -599,11 +600,21 @@ static void AddAllocateVoxelPagesPass(
 	// Allocate page index for all instance group
 //	if (bIsGPUDriven)
 	{
+		uint32 bDoesMacroGroupSupportVoxelization = 0;
+		for (const FHairStrandsMacroGroupData& MacroGroup : MacroGroupDatas)
+		{
+			if (MacroGroup.bSupportVoxelization)
+			{
+				bDoesMacroGroupSupportVoxelization |= 1u << MacroGroup.MacroGroupId;
+			}
+		}
+
 		FVoxelAllocatePageIndexCS::FParameters* Parameters = GraphBuilder.AllocParameters<FVoxelAllocatePageIndexCS::FParameters>();
 		Parameters->PageWorldSize = PageWorldSize;
 		Parameters->TotalPageIndexCount = OutTotalPageIndexCount;
 		Parameters->PageResolution = PageResolution;
 		Parameters->MacroGroupCount = MacroGroupCount;
+		Parameters->bDoesMacroGroupSupportVoxelization = bDoesMacroGroupSupportVoxelization;
 		Parameters->MacroGroupAABBBuffer = GraphBuilder.CreateUAV(MacroGroupResources.MacroGroupAABBsBuffer, PF_R32_SINT);
 		Parameters->IndirectDispatchGroupSize = GroupSize; // This is the GroupSize used for FVoxelAllocateVoxelPageCS
 		Parameters->OutPageIndexResolutionAndOffsetBuffer = GraphBuilder.CreateUAV(PageIndexResolutionBuffer, PF_R32G32B32A32_UINT);
@@ -636,7 +647,10 @@ static void AddAllocateVoxelPagesPass(
 		for (const FHairStrandsMacroGroupData::PrimitiveInfo& PrimitiveInfo : MacroGroup.PrimitivesInfos)
 		{
 			FHairGroupPublicData* HairGroupData = PrimitiveInfo.PublicDataPtr;
-			TotalClusterCount += HairGroupData->GetClusterCount();
+			if (HairGroupData->DoesSupportVoxelization())
+			{
+				TotalClusterCount += HairGroupData->GetClusterCount();
+			}
 		}
 	}
 
@@ -669,6 +683,7 @@ static void AddAllocateVoxelPagesPass(
 			{
 				FHairGroupPublicData* HairGroupData = PrimitiveInfo.PublicDataPtr;
 
+				// Among the macro group, some prmitive might not support/require voxelization. So skip them.
 				if (!HairGroupData->DoesSupportVoxelization())
 					continue;
 
@@ -1239,6 +1254,8 @@ static void AddVirtualVoxelizationComputeRasterPass(
 	if (!bIsGPUDriven)
 		return;
 
+	ensure(MacroGroup.bSupportVoxelization);
+
 	if (ViewInfo)
 	{
 		const FHairStrandsMacroGroupData::TPrimitiveInfos& PrimitiveSceneInfos = MacroGroup.PrimitivesInfos;
@@ -1607,13 +1624,10 @@ void VoxelizeHairStrands(
 	bool bHasValidElementToVoxelize = false;
 	for (const FHairStrandsMacroGroupData& MacroGroup : MacroGroupDatas)
 	{
-		for (const FHairStrandsMacroGroupData::PrimitiveInfo& PrimitiveInfo : MacroGroup.PrimitivesInfos)
+		if (MacroGroup.bSupportVoxelization)
 		{
-			if (PrimitiveInfo.PublicDataPtr->DoesSupportVoxelization())
-			{
-				bHasValidElementToVoxelize = true;
-				break;
-			}
+			bHasValidElementToVoxelize = true;
+			break;
 		}
 	}
 
@@ -1636,14 +1650,20 @@ void VoxelizeHairStrands(
 
 		for (FHairStrandsMacroGroupData& MacroGroup : MacroGroupDatas)
 		{
-			AddVirtualVoxelizationRasterPass(GraphBuilder, Scene, &View, VirtualVoxelResources, InstanceCullingManager, MacroGroup);
+			if (MacroGroup.bSupportVoxelization)
+			{
+				AddVirtualVoxelizationRasterPass(GraphBuilder, Scene, &View, VirtualVoxelResources, InstanceCullingManager, MacroGroup);
+			}
 		}
 
 		if (GHairVoxelInjectOpaqueDepthEnable > 0)
 		{
 			for (FHairStrandsMacroGroupData& MacroGroup : MacroGroupDatas)
 			{
-				AddVirtualVoxelInjectOpaquePass(GraphBuilder, View, VirtualVoxelResources, MacroGroup);
+				if (MacroGroup.bSupportVoxelization)
+				{
+					AddVirtualVoxelInjectOpaquePass(GraphBuilder, View, VirtualVoxelResources, MacroGroup);
+				}
 			}
 		}
 
