Make FFT Bloom work with -rhivalidation (UE-116006).

- Works with both Fits and !FitsInGroupShared (multipass) modes.

diff --git a/Engine/Source/Runtime/Renderer/Private/GPUFastFourierTransform.cpp b/Engine/Source/Runtime/Renderer/Private/GPUFastFourierTransform.cpp
--- a/Engine/Source/Runtime/Renderer/Private/GPUFastFourierTransform.cpp
+++ b/Engine/Source/Runtime/Renderer/Private/GPUFastFourierTransform.cpp
@@ -1692,8 +1692,6 @@ bool GPUFFT::ComplexFFTImage1D::MultiPass(FGPUFFTShaderContext& Context, const F
 
 		for (uint32 Ns = 2; Ns < TransformLength; Ns *= 2)
 		{
-			// Make it safe to read from the buffer we just wrote to.
-			RHICmdList.Transition(FRHITransitionInfo(Targets.DstTarget().UAV, ERHIAccess::Unknown, ERHIAccess::UAVCompute));
 			Targets.Swap();
 
 			auto HasValidTargets = [&Targets, &DstExtent]()->bool
@@ -1707,6 +1705,11 @@ bool GPUFFT::ComplexFFTImage1D::MultiPass(FGPUFFTShaderContext& Context, const F
 
 			checkf(HasValidTargets(), TEXT("FFT: Allocated Buffers too small."));
 
+			RHICmdList.Transition(
+				{
+					FRHITransitionInfo(Targets.SrcTarget().ShaderResourceTexture, ERHIAccess::UAVMask, ERHIAccess::SRVMask),
+					FRHITransitionInfo(Targets.DstTarget().UAV, ERHIAccess::SRVMask, ERHIAccess::UAVMask)
+				});
 			DispatchComplexFFTPassCS(Context, FFTDesc, Ns, Targets.SrcTarget().ShaderResourceTexture, XFormWindow, Targets.DstTarget().UAV);
 		}
 
@@ -1724,15 +1727,17 @@ bool GPUFFT::ComplexFFTImage1D::MultiPass(FGPUFFTShaderContext& Context, const F
 	{
 		// Re-order the data so we can do a pass of group-shared transforms
 		DispatchReorderFFTPassCS(Context, FFTDesc, Window, SrcTexture, XFormWindow, Targets.DstTarget().UAV);
-		RHICmdList.Transition(FRHITransitionInfo(Targets.DstTarget().UAV, ERHIAccess::Unknown, ERHIAccess::UAVCompute));
 		Targets.Swap();
 
+		RHICmdList.Transition(
+			{
+				FRHITransitionInfo(Targets.SrcTarget().ShaderResourceTexture, ERHIAccess::UAVMask, ERHIAccess::SRVMask),
+				FRHITransitionInfo(Targets.DstTarget().UAV, ERHIAccess::SRVMask, ERHIAccess::UAVMask)
+			});
 		DispatchGSSubComplexFFTPassCS(Context, FFTDesc, Targets.SrcTarget().ShaderResourceTexture, XFormWindow, Targets.DstTarget().UAV);
 
 		for (uint32 Ns = FGroupShardSubFFTPassCS::SubPassLength(); Ns < TransformLength; Ns *= 2)
 		{
-			// Make it safe to read from the buffer we just wrote to.
-			RHICmdList.Transition(FRHITransitionInfo(Targets.DstTarget().UAV, ERHIAccess::Unknown, ERHIAccess::UAVCompute));
 			Targets.Swap();
 
 			auto HasValidTargets = [&Targets, &DstExtent]()->bool
@@ -1746,6 +1751,11 @@ bool GPUFFT::ComplexFFTImage1D::MultiPass(FGPUFFTShaderContext& Context, const F
 
 			checkf(HasValidTargets(), TEXT("FFT: Allocated Buffers too small."));
 
+			RHICmdList.Transition(
+				{
+					FRHITransitionInfo(Targets.SrcTarget().ShaderResourceTexture, ERHIAccess::UAVMask, ERHIAccess::SRVMask),
+					FRHITransitionInfo(Targets.DstTarget().UAV, ERHIAccess::SRVMask, ERHIAccess::UAVMask)
+				});
 			DispatchComplexFFTPassCS(Context, FFTDesc, Ns, Targets.SrcTarget().ShaderResourceTexture, XFormWindow, Targets.DstTarget().UAV);
 
 		}
@@ -1810,7 +1820,11 @@ bool GPUFFT::TwoForOneRealFFTImage1D::MultiPass(FGPUFFTShaderContext& Context, c
 			// Copy data into DstBuffer
 			CopyImage2D(Context, SrcWindow, SrcTexture, SrcWindow, DstBuffer.UAV, PreFilter);
 
-			Context.GetRHICmdList().Transition(FRHITransitionInfo(DstBuffer.UAV, ERHIAccess::Unknown, ERHIAccess::UAVCompute));
+			Context.GetRHICmdList().Transition(
+				{
+					FRHITransitionInfo(DstBuffer.ShaderResourceTexture, ERHIAccess::UAVMask, ERHIAccess::SRVMask),
+					FRHITransitionInfo(TmpBuffer.UAV, ERHIAccess::SRVMask, ERHIAccess::UAVMask)
+				});
 
 			// Transform as 2 sets of complex data, putting the result in the DstBuffer.  This performs multiple dispatches.
 			// Transform into DstBuffer.
@@ -1827,19 +1841,26 @@ bool GPUFFT::TwoForOneRealFFTImage1D::MultiPass(FGPUFFTShaderContext& Context, c
 			SwapContents(DstBuffer, TmpBuffer);
 		}
 
-		Context.GetRHICmdList().Transition(FRHITransitionInfo(TmpBuffer.UAV, ERHIAccess::Unknown, ERHIAccess::UAVCompute));
-
+		Context.GetRHICmdList().Transition(
+			{
+				FRHITransitionInfo(TmpBuffer.ShaderResourceTexture, ERHIAccess::UAVMask, ERHIAccess::SRVMask),
+				FRHITransitionInfo(DstBuffer.UAV, ERHIAccess::SRVMask, ERHIAccess::UAVMask)
+			});
 		// Unpack the complex transform into transform of real data
 		DispatchPackTwoForOneFFTPassCS(Context, FFTDesc, TmpBuffer.ShaderResourceTexture, DstBuffer.UAV);
 		
 	}
 	else  // Inverse transform.
 	{
+		Context.GetRHICmdList().Transition(FRHITransitionInfo(DstBuffer.UAV, ERHIAccess::SRVMask, ERHIAccess::UAVMask));
 		// Pack the 4 transforms of real data as 2 transforms of complex data 
 		DispatchPackTwoForOneFFTPassCS(Context, FFTDesc, SrcTexture, DstBuffer.UAV);
 
-		Context.GetRHICmdList().Transition(FRHITransitionInfo(DstBuffer.UAV, ERHIAccess::Unknown, ERHIAccess::UAVCompute));
-
+		Context.GetRHICmdList().Transition(
+			{
+				FRHITransitionInfo(DstBuffer.ShaderResourceTexture, ERHIAccess::UAVMask, ERHIAccess::SRVMask),
+				FRHITransitionInfo(TmpBuffer.UAV, ERHIAccess::SRVMask, ERHIAccess::UAVMask)
+			});
 		// Transform as complex data
 		SuccessValue = SuccessValue &&
 			ComplexFFTImage1D::MultiPass(Context, FFTDesc, SrcWindow, DstBuffer.ShaderResourceTexture, TmpBuffer/*result*/, DstBuffer/*tmp*/);
@@ -1886,6 +1907,7 @@ bool GPUFFT::FFTImage2D(FGPUFFTShaderContext& Context, const FIntPoint& Frequenc
 
 	// Two-for-one transform: SrcTexture fills the TmpBuffer
 
+	Context.GetRHICmdList().Transition(FRHITransitionInfo(TmpBuffer.UAV, ERHIAccess::Unknown, ERHIAccess::UAVMask));
 	if (FitsInGroupSharedMemory(TwoForOneFFTDesc))
 	{
 		// Horizontal FFT: input -> tmp buffer 0 
@@ -1901,11 +1923,13 @@ bool GPUFFT::FFTImage2D(FGPUFFTShaderContext& Context, const FIntPoint& Frequenc
 		SuccessValue = SuccessValue &&
 			TwoForOneRealFFTImage1D::MultiPass(Context, TwoForOneFFTDesc, ROIRect, SrcTexture, TmpBuffer /*dst*/, DmyBuffer/*tmp*/);
 		//->dst
-
 	}
 
-	Context.GetRHICmdList().Transition(FRHITransitionInfo(DstBuffer.UAV, ERHIAccess::Unknown, ERHIAccess::UAVCompute));
-	Context.GetRHICmdList().Transition(FRHITransitionInfo(TmpBuffer.UAV, ERHIAccess::Unknown, ERHIAccess::UAVCompute));
+	Context.GetRHICmdList().Transition(
+		{
+			FRHITransitionInfo(DstBuffer.UAV, ERHIAccess::Unknown, ERHIAccess::UAVMask),
+			FRHITransitionInfo(TmpBuffer.ShaderResourceTexture, ERHIAccess::UAVMask, ERHIAccess::SRVMask)
+		});
 
 	// Complex transform in the other direction: TmpBuffer fills DstBuffer
 	if (FitsInGroupSharedMemory(ComplexFFTDesc))
@@ -1961,7 +1985,7 @@ bool GPUFFT::ConvolutionWithTextureImage1D::MultiPass(FGPUFFTShaderContext& Cont
 {
 	bool SuccessValue = true;
 
-	
+
 	// Frequency Space size.
 	const FIntPoint TargetExtent = FFTDesc.TransformExtent();
 	const FIntRect  TargetRect(FIntPoint(0, 0), TargetExtent);
@@ -1969,12 +1993,15 @@ bool GPUFFT::ConvolutionWithTextureImage1D::MultiPass(FGPUFFTShaderContext& Cont
 	// Forward transform: Results in DstBuffer
 	SuccessValue = SuccessValue &&
 		ComplexFFTImage1D::MultiPass(Context, FFTDesc, SrcWindow, SrcTexture, DstBuffer, TmpBuffer);
-	Context.GetRHICmdList().Transition(FRHITransitionInfo(DstBuffer.UAV, ERHIAccess::Unknown, ERHIAccess::UAVCompute));
 
+	Context.GetRHICmdList().Transition(
+		{
+			FRHITransitionInfo(DstBuffer.ShaderResourceTexture, ERHIAccess::UAVMask, ERHIAccess::SRVMask),
+			FRHITransitionInfo(TmpBuffer.UAV, ERHIAccess::SRVMask, ERHIAccess::UAVMask)
+		});
 
 	// Convolution: Results in TmpBuffer
 	DispatchComplexMultiplyImagesCS(Context, FFTDesc.IsHorizontal(), TargetRect, DstBuffer.ShaderResourceTexture, TransformedKernel,  TmpBuffer.UAV);
-	Context.GetRHICmdList().Transition(FRHITransitionInfo(TmpBuffer.UAV, ERHIAccess::Unknown, ERHIAccess::UAVCompute));
 
 	// Inverse Transform: Results in DstBuffer
 	
@@ -1982,10 +2009,19 @@ bool GPUFFT::ConvolutionWithTextureImage1D::MultiPass(FGPUFFTShaderContext& Cont
 	InvFFTDesc.XFormType = GetInverseOfXForm(FFTDesc.XFormType);
 	// testing.  Verify that the forward transform works.
 	
+	Context.GetRHICmdList().Transition(
+		{
+			FRHITransitionInfo(TmpBuffer.ShaderResourceTexture, ERHIAccess::UAVMask, ERHIAccess::SRVMask),
+			FRHITransitionInfo(DstBuffer.UAV, ERHIAccess::SRVMask, ERHIAccess::UAVMask)
+		});
 	SuccessValue = SuccessValue &&
 		ComplexFFTImage1D::MultiPass(Context, InvFFTDesc, TargetRect, TmpBuffer.ShaderResourceTexture, DstBuffer, TmpBuffer);
 	
-	Context.GetRHICmdList().Transition(FRHITransitionInfo(DstBuffer.UAV, ERHIAccess::Unknown, ERHIAccess::UAVCompute));
+	Context.GetRHICmdList().Transition(
+		{
+			FRHITransitionInfo(DstBuffer.ShaderResourceTexture, ERHIAccess::UAVMask, ERHIAccess::SRVMask),
+			FRHITransitionInfo(TmpBuffer.UAV, ERHIAccess::SRVMask, ERHIAccess::UAVMask)
+		});
 
 	// Copy back to the correct sized sub-window
 	DispatchCopyWindowCS(Context, TargetRect, DstBuffer.ShaderResourceTexture, SrcWindow, TmpBuffer.UAV);
@@ -2080,6 +2116,7 @@ bool GPUFFT::ConvolutionWithTextureImage2D(FGPUFFTShaderContext& Context, const
 	// ---- Two For One Transform --- 
 
 	// Horizontal FFT: input -> tmp buffer 0 
+	RHICmdList.Transition(FRHITransitionInfo(TmpBuffer0.UAV, ERHIAccess::Unknown, ERHIAccess::UAVMask));
 	if (FitsInGroupSharedMemory(TwoForOneFFTDesc))
 	{
 		SuccessValue = SuccessValue &&
@@ -2091,10 +2128,14 @@ bool GPUFFT::ConvolutionWithTextureImage2D(FGPUFFTShaderContext& Context, const
 		SuccessValue = SuccessValue &&
 			TwoForOneRealFFTImage1D::MultiPass(Context, TwoForOneFFTDesc, ROIRect, SrcTexture, TmpBuffer0, TmpBuffer1, PreFilter);
 	}
-	RHICmdList.Transition(FRHITransitionInfo(TmpBuffer0.UAV, ERHIAccess::Unknown, ERHIAccess::UAVCompute));
-	
+
 	// ---- 1 D Convolution --- 
 
+	RHICmdList.Transition(
+		{
+			FRHITransitionInfo(TmpBuffer0.ShaderResourceTexture, ERHIAccess::UAVMask, ERHIAccess::SRVMask),
+			FRHITransitionInfo(TmpBuffer1.UAV, ERHIAccess::Unknown, ERHIAccess::UAVMask)
+		});
 	if (FitsInGroupSharedMemory(ConvolutionFFTDesc))
 	{
 		// Vertical FFT / Filter / Vertical InvFFT : tmp buffer 0 -> tmp buffer 1
@@ -2107,11 +2148,15 @@ bool GPUFFT::ConvolutionWithTextureImage2D(FGPUFFTShaderContext& Context, const
 		SuccessValue = SuccessValue &&
 			ConvolutionWithTextureImage1D::MultiPass(Context, ConvolutionFFTDesc, TransformedKernel, TwoForOneOutputRect, TmpBuffer0.ShaderResourceTexture, TmpBuffer1, TmpBuffer0);
 	}
-	RHICmdList.Transition(FRHITransitionInfo(TmpBuffer1.UAV, ERHIAccess::Unknown, ERHIAccess::UAVCompute));
 
 
 	// ---- Inverse Two For One ---
 
+	RHICmdList.Transition(
+		{
+			FRHITransitionInfo(TmpBuffer1.ShaderResourceTexture, ERHIAccess::UAVMask, ERHIAccess::SRVMask),
+			FRHITransitionInfo(ResultUAV, ERHIAccess::Unknown, ERHIAccess::UAVMask)
+		});
 	if (FitsInGroupSharedMemory(TwoForOneIvnFFTDesc))
 	{
 		// Horizontal InvFFT: tmp buffer 1 -> result
@@ -2124,8 +2169,8 @@ bool GPUFFT::ConvolutionWithTextureImage2D(FGPUFFTShaderContext& Context, const
 		// TmpBuffer1 -> TmpBuffer0
 		SuccessValue = SuccessValue &&
 			TwoForOneRealFFTImage1D::MultiPass(Context, TwoForOneIvnFFTDesc, TwoForOneOutputRect, TmpBuffer1.ShaderResourceTexture, TmpBuffer0, TmpBuffer1);
-		
-		RHICmdList.Transition(FRHITransitionInfo(TmpBuffer0.UAV, ERHIAccess::Unknown, ERHIAccess::SRVCompute));
+
+		RHICmdList.Transition(FRHITransitionInfo(TmpBuffer0.ShaderResourceTexture, ERHIAccess::UAVMask, ERHIAccess::SRVMask));
 
 		DispatchCopyWindowCS(Context, TwoForOneOutputRect, TmpBuffer0.ShaderResourceTexture, ROIRect, ResultUAV);
 	}
diff --git a/Engine/Source/Runtime/Renderer/Private/PostProcess/PostProcessFFTBloom.cpp b/Engine/Source/Runtime/Renderer/Private/PostProcess/PostProcessFFTBloom.cpp
--- a/Engine/Source/Runtime/Renderer/Private/PostProcess/PostProcessFFTBloom.cpp
+++ b/Engine/Source/Runtime/Renderer/Private/PostProcess/PostProcessFFTBloom.cpp
@@ -705,7 +705,8 @@ FSceneRenderTargetItem* InitDomainAndGetKernel(FRHICommandList& RHICmdList, cons
 		ResizeAndCenterTexture(RHICmdList, View, PhysicalSpaceKernelTextureRef, ImageSize, CenterUV, Intermediates.KernelSupportScale,
 			Intermediates.FrequencySize, SpectralKernelRTItem.UAV, PaddedFrequencySize, Intermediates.bHalfResolutionFFT);
 
-		RHICmdList.Transition(FRHITransitionInfo(SpectralKernelRTItem.UAV, ERHIAccess::Unknown, ERHIAccess::UAVCompute));
+		// ResizeAndCenterTexture transitioned SpectralKernelRTItem to UAVCompute
+		RHICmdList.Transition(FRHITransitionInfo(SpectralKernelRTItem.ShaderResourceTexture, ERHIAccess::UAVCompute, ERHIAccess::SRVMask));
 
 		// Two Dimensional FFT of the physical space kernel.  
 		// Input: SpectralRTItem holds the physical space kernel, on return it will be the spectral space 
