Fixed static analyzer warnings.
Courtasy of Mickael.Gilabert

diff --git a/Engine/Source/Runtime/Renderer/Private/SceneCaptureRendering.cpp b/Engine/Source/Runtime/Renderer/Private/SceneCaptureRendering.cpp
--- a/Engine/Source/Runtime/Renderer/Private/SceneCaptureRendering.cpp
+++ b/Engine/Source/Runtime/Renderer/Private/SceneCaptureRendering.cpp
@@ -445,6 +445,15 @@ static void BuildOrthoMatrix(FIntPoint InRenderTargetSize, float InOrthoWidth, i
 		return;
 	}
 
+#if DO_CHECK
+	check(InNumXTiles != 0 && InNumYTiles != 0);
+	if (InNumXTiles == 0 || InNumYTiles == 0)
+	{
+		OutProjectionMatrix = FMatrix(EForceInit::ForceInitToZero);
+		return;
+	}
+#endif
+
 	const float XTileDividerRcp = 1.0f / float(InNumXTiles);
 	const float YTileDividerRcp = 1.0f / float(InNumYTiles);
 
