Protect extra GPU Scene instance data arrays from out of bounds access if data flags are incorrectly set

diff --git a/Engine/Source/Runtime/Renderer/Private/GPUScene.cpp b/Engine/Source/Runtime/Renderer/Private/GPUScene.cpp
--- a/Engine/Source/Runtime/Renderer/Private/GPUScene.cpp
+++ b/Engine/Source/Runtime/Renderer/Private/GPUScene.cpp
@@ -1133,9 +1133,9 @@ void FGPUScene::UploadGeneral(FRHICommandListImmediate& RHICmdList, FScene *Scen
 										const FPrimitiveInstance& SceneData = UploadInfo.PrimitiveInstances[InstanceIndex];
 
 										// TODO: Temporary
-										const FRenderTransform& PrevLocalToPrimitive = SceneData.Flags & INSTANCE_SCENE_DATA_FLAG_HAS_DYNAMIC_DATA ? UploadInfo.InstanceDynamicData[InstanceIndex].PrevLocalToPrimitive : FRenderTransform::Identity;
-										FVector4 LightMapShadowMapUVBias = SceneData.Flags & INSTANCE_SCENE_DATA_FLAG_HAS_LIGHTSHADOW_UV_BIAS ? UploadInfo.InstanceLightShadowUVBias[InstanceIndex] : FVector4(ForceInitToZero);
-										float RandomID = SceneData.Flags & INSTANCE_SCENE_DATA_FLAG_HAS_RANDOM ? UploadInfo.InstanceRandomID[InstanceIndex] : 0.0f;
+										const FRenderTransform& PrevLocalToPrimitive = UploadInfo.InstanceDynamicData.Num() == UploadInfo.PrimitiveInstances.Num() ? UploadInfo.InstanceDynamicData[InstanceIndex].PrevLocalToPrimitive : FRenderTransform::Identity;
+										FVector4 LightMapShadowMapUVBias = UploadInfo.InstanceLightShadowUVBias.Num() == UploadInfo.PrimitiveInstances.Num() ? UploadInfo.InstanceLightShadowUVBias[InstanceIndex] : FVector4(ForceInitToZero);
+										float RandomID = UploadInfo.InstanceRandomID.Num() == UploadInfo.PrimitiveInstances.Num() ? UploadInfo.InstanceRandomID[InstanceIndex] : 0.0f;
 
 										// TODO: Temporary hack!
 										float CustomDataFloat0 = 0.0f;
