Change hair composition to use hair strands uniform buffer.

This CL also adds hair emissive texture to its uniform buffer.

diff --git a/Engine/Shaders/Private/HairStrands/HairStrandsVisibilityComposeSubPixelPS.usf b/Engine/Shaders/Private/HairStrands/HairStrandsVisibilityComposeSubPixelPS.usf
--- a/Engine/Shaders/Private/HairStrands/HairStrandsVisibilityComposeSubPixelPS.usf
+++ b/Engine/Shaders/Private/HairStrands/HairStrandsVisibilityComposeSubPixelPS.usf
@@ -1,5 +1,6 @@
 // Copyright Epic Games, Inc. All Rights Reserved.
 
+#include "HairStrandsVisibilityCommonStruct.ush" 
 #include "../Common.ush"
 #include "../CommonViewUniformBuffer.ush"
 
@@ -89,13 +90,8 @@ void FastResolvePS(in FScreenVertexOutput Input)
 
 int2 OutputResolution;
 uint bComposeDofDepth;
-uint bEmissiveEnable;
 
-Texture2D<float4>	HairEmissiveTexture;
 Texture2D<float>	HairDOFDepthTexture;
-Texture2D<uint>		HairSampleCount;
-Texture2D<uint4>	HairCategorizationTexture;
-Texture2D<uint>		HairVisibilityNodeOffsetAndCount;
 Texture2D<float4>	HairLightingSampleBuffer;
 
 float4 EvaluateVolumetric(float3 WorldPosition)
@@ -146,18 +142,18 @@ void ComposeSamplePS(
 	OutColor = 0.0f;
 	const uint3 PixelCoord = uint3(floor(Input.Position.xy), 0);
 
-	FNodeDesc NodeDesc = DecodeNodeDesc(HairVisibilityNodeOffsetAndCount.Load(PixelCoord));
+	FNodeDesc NodeDesc = DecodeNodeDesc(HairStrands.HairSampleOffset.Load(PixelCoord));
 	if (NodeDesc.Count == 0)
 	{
 		discard;
 	}
 
-	const FCategorizationData CatData = DecodeCategorizationData(HairCategorizationTexture.Load(PixelCoord));
+	const FCategorizationData CatData = DecodeCategorizationData(HairStrands.HairCategorizationTexture.Load(PixelCoord));
 	const float PixelCoverage = min(CatData.PixelCoverage, 1);
 	if (PixelCoverage == 0)
 		discard;
 
-	const uint TotalNodeCount = HairSampleCount.Load(uint3(0, 0, 0));
+	const uint TotalNodeCount = HairStrands.HairSampleCount.Load(uint3(0, 0, 0));
 	const uint Resolution = uint(ceil(sqrt(float(TotalNodeCount))));
 
 	const float3 ClosetPointWorldPosition = SvPositionToWorld(float4(Input.Position.xy, CatData.ClosestDepth, 1.0));;
@@ -170,7 +166,7 @@ void ComposeSamplePS(
 		const uint LocalOffset = NodeDesc.Offset + SampleIt;
 		const uint2 Coord = uint2(LocalOffset % Resolution, LocalOffset / Resolution);
 		const float4 LightingSample = HairLightingSampleBuffer.Load(uint3(Coord,0));
-		const float4 EmissiveSample = HairEmissiveTexture.Load(uint3(Coord, 0));
+		const float4 EmissiveSample = HairStrands.HairEmissiveTexture.Load(uint3(Coord, 0));
 		LocalAccColor += LightingSample.rgb + EmissiveSample.rgb;
 	}
 	OutColor.rgb = LocalAccColor * Volumetric.a + Volumetric.rgb;
diff --git a/Engine/Source/Runtime/Renderer/Private/HairStrands/HairStrandsComposition.cpp b/Engine/Source/Runtime/Renderer/Private/HairStrands/HairStrandsComposition.cpp
--- a/Engine/Source/Runtime/Renderer/Private/HairStrands/HairStrandsComposition.cpp
+++ b/Engine/Source/Runtime/Renderer/Private/HairStrands/HairStrandsComposition.cpp
@@ -170,11 +170,7 @@ class FHairVisibilityComposeSamplePS : public FGlobalShader
 		SHADER_PARAMETER_STRUCT_INCLUDE(FHairStrandsTilePassVS::FParameters, TileData)
 		SHADER_PARAMETER(FIntPoint, OutputResolution)
 		SHADER_PARAMETER(uint32, bComposeDofDepth)
-		SHADER_PARAMETER(uint32, bEmissiveEnable)
-		SHADER_PARAMETER_RDG_TEXTURE(Texture2D, HairEmissiveTexture)
-		SHADER_PARAMETER_RDG_TEXTURE(Texture2D, HairSampleCount)
-		SHADER_PARAMETER_RDG_TEXTURE(Texture2D, HairCategorizationTexture)
-		SHADER_PARAMETER_RDG_TEXTURE(Texture2D, HairVisibilityNodeOffsetAndCount)
+		SHADER_PARAMETER_RDG_UNIFORM_BUFFER(FHairStrandsViewUniformParameters, HairStrands)
 		SHADER_PARAMETER_RDG_TEXTURE(Texture2D, HairLightingSampleBuffer)
 		SHADER_PARAMETER_RDG_TEXTURE(Texture2D, HairDOFDepthTexture)
 		SHADER_PARAMETER_RDG_UNIFORM_BUFFER(FFogUniformParameters, FogStruct)
@@ -207,16 +203,12 @@ static void AddHairVisibilityComposeSamplePass(
 	TRDGUniformBufferRef<FFogUniformParameters> FogBuffer = CreateFogUniformBuffer(GraphBuilder, View);
 
 	FHairVisibilityComposeSamplePS::FParameters* Parameters = GraphBuilder.AllocParameters<FHairVisibilityComposeSamplePS::FParameters>();
-	Parameters->HairSampleCount = VisibilityData.NodeCount;
 	Parameters->bComposeDofDepth = bDOFEnable ? 1 : 0;
-	Parameters->HairCategorizationTexture = CategorizationTexture;
-	Parameters->HairVisibilityNodeOffsetAndCount = VisibilityData.NodeIndex;
 	Parameters->HairLightingSampleBuffer = VisibilityData.SampleLightingBuffer;
 	Parameters->HairDOFDepthTexture = bDOFEnable ? HairDOFDepthTexture : GSystemTextures.GetBlackDummy(GraphBuilder);
 	Parameters->OutputResolution = OutColorTexture->Desc.Extent;
 	Parameters->ViewUniformBuffer = View.ViewUniformBuffer;
-	Parameters->bEmissiveEnable = VisibilityData.EmissiveTexture != nullptr ? 1 : 0;
-	Parameters->HairEmissiveTexture = VisibilityData.EmissiveTexture != nullptr ? VisibilityData.EmissiveTexture : GSystemTextures.GetBlackDummy(GraphBuilder);
+	Parameters->HairStrands = View.HairStrandsViewData.UniformBuffer;
 	Parameters->FogStruct = FogBuffer;
 	Parameters->RenderTargets[0] = FRenderTargetBinding(OutColorTexture, ERenderTargetLoadAction::ELoad);
 	Parameters->RenderTargets.DepthStencil = FDepthStencilBinding(OutDepthTexture, ERenderTargetLoadAction::ELoad, ERenderTargetLoadAction::ENoAction, FExclusiveDepthStencil::DepthWrite_StencilNop);
diff --git a/Engine/Source/Runtime/Renderer/Private/HairStrands/HairStrandsData.h b/Engine/Source/Runtime/Renderer/Private/HairStrands/HairStrandsData.h
--- a/Engine/Source/Runtime/Renderer/Private/HairStrands/HairStrandsData.h
+++ b/Engine/Source/Runtime/Renderer/Private/HairStrands/HairStrandsData.h
@@ -24,12 +24,14 @@ BEGIN_GLOBAL_SHADER_PARAMETER_STRUCT(FHairStrandsViewUniformParameters, )
 	SHADER_PARAMETER(float, HairDualScatteringRoughnessOverride)						// Override the roughness used for dual scattering (for hack/test purpose only)
 	SHADER_PARAMETER(FIntPoint, HairSampleViewportResolution)							// Maximum viewport resolution of the sample space
 	SHADER_PARAMETER(uint32, bHairTileValid)											// True if tile data are valid
+	SHADER_PARAMETER(uint32, bHasEmissive)												// True if emissive is enabled
 	SHADER_PARAMETER_RDG_TEXTURE(Texture2D<uint4>, HairCategorizationTexture)			// Categorization texture aggregating hair info in screen space (closest depth, coverage, ...)
 	SHADER_PARAMETER_RDG_TEXTURE(Texture2D, HairOnlyDepthTexture)						// Depth texture containing only hair depth (not strongly typed for legacy shader code reason)
 	SHADER_PARAMETER_RDG_TEXTURE(Texture2D, HairOnlyDepthClosestHZBTexture)				// HZB closest depth texture containing only hair depth (not strongly typed for legacy shader code reason)
 	SHADER_PARAMETER_RDG_TEXTURE(Texture2D, HairOnlyDepthFurthestHZBTexture)			// HZB furthest depth texture containing only hair depth (not strongly typed for legacy shader code reason)
 	SHADER_PARAMETER_RDG_TEXTURE(Texture2D<uint>, HairSampleOffset)						// Offset & count, for accessing pixel's samples, based on screen pixel position
 	SHADER_PARAMETER_RDG_TEXTURE(Texture2D<uint>, HairSampleCount)						// Total count of hair sample, in sample space
+	SHADER_PARAMETER_RDG_TEXTURE(Texture2D<float4>, HairEmissiveTexture)				// Emissive texture
 	SHADER_PARAMETER_RDG_BUFFER_SRV(StructuredBuffer<FPackedHairSample>, HairSampleData)// Sample data (coverage, tangent, base color, ...), in sample space // HAIRSTRANDS_TODO: change this to be a uint4 so that we don't have to include the type for generated contant buffer
 	SHADER_PARAMETER_RDG_BUFFER_SRV(Buffer<uint2>, HairSampleCoords)					// Screen pixel coordinate of each sample, in sample space
 	SHADER_PARAMETER_RDG_BUFFER_SRV(Buffer<uint2>, HairTileData)						// Tile coords (RG16F)
diff --git a/Engine/Source/Runtime/Renderer/Private/HairStrands/HairStrandsRendering.cpp b/Engine/Source/Runtime/Renderer/Private/HairStrands/HairStrandsRendering.cpp
--- a/Engine/Source/Runtime/Renderer/Private/HairStrands/HairStrandsRendering.cpp
+++ b/Engine/Source/Runtime/Renderer/Private/HairStrands/HairStrandsRendering.cpp
@@ -22,6 +22,18 @@ static TRDGUniformBufferRef<FHairStrandsViewUniformParameters> InternalCreateHai
 		Parameters->HairSampleCoords = GraphBuilder.CreateSRV(In->NodeCoord, FHairStrandsVisibilityData::NodeCoordFormat);
 		Parameters->HairSampleCount = In->NodeCount;
 		Parameters->HairSampleViewportResolution = In->SampleLightingViewportResolution;
+
+		if (In->EmissiveTexture)
+		{
+			Parameters->HairEmissiveTexture = In->EmissiveTexture;
+			Parameters->bHasEmissive = 1u;
+		}
+		else
+		{
+			Parameters->HairEmissiveTexture = GSystemTextures.GetBlackDummy(GraphBuilder);
+			Parameters->bHasEmissive = 0u;
+			
+		}
 		if (In->TileData.IsValid())
 		{
 			Parameters->HairTileData = In->TileData.TileDataSRV;
@@ -69,6 +81,8 @@ static TRDGUniformBufferRef<FHairStrandsViewUniformParameters> InternalCreateHai
 		Parameters->HairSampleCoords = DummyBufferRG16SRV;
 		Parameters->HairSampleData	 = GraphBuilder.CreateSRV(DummyNodeBuffer);
 		Parameters->HairSampleViewportResolution = FIntPoint(0, 0);
+		Parameters->HairEmissiveTexture = BlackTexture;
+		Parameters->bHasEmissive = 0u;
 
 		Parameters->HairTileData = DummyBufferR32SRV;
 		Parameters->HairTileCount = DummyBufferRG16SRV;
