Replaces r.PostProcessAAQuality with r.FXAA.Quality & r.TemporalAA.Quality

Now only r.AntiAliasingMethod prevails to selects the anti-aliasing method explicitly for the deferred shading renderer. Each AA technic have it's own cvars to control their quality independently

FXAA now have its own r.FXAA.Quality. The FXAA shader permutations where mis configured with r.PostProcessAAQuality=1 & 2 both mapping to the FXAA_PC_CONSOLE=1 that is now r.FXAA.Quality=0. Instead r.FXAA.Quality now offer more mid-quality settings for FXAA_PC=1 with r.FXAA.Quality=1 & 2 & 3. Backward comaptible migration is as followed:
- r.AntiAliasingMethod=1 r.PostProcessAAQuality=0 -> r.FXAA.Quality=0
- r.AntiAliasingMethod=1 r.PostProcessAAQuality=1 -> r.FXAA.Quality=0
- r.AntiAliasingMethod=1 r.PostProcessAAQuality=2 -> r.FXAA.Quality=0
- r.AntiAliasingMethod=1 r.PostProcessAAQuality=3 -> r.FXAA.Quality=3 but with some neighborhood search setting changed in FXAAShader.usf
- r.AntiAliasingMethod=1 r.PostProcessAAQuality=4 -> r.FXAA.Quality=4
- r.AntiAliasingMethod=1 r.PostProcessAAQuality=5 -> r.FXAA.Quality=5

TAA now have its own r.TemporalAA.Quality. Backward compatible migration is as followed:
- r.AntiAliasingMethod=2 r.PostProcessAAQuality=0 -> r.AntiAliasingMethod=0
- r.AntiAliasingMethod=2 r.PostProcessAAQuality=1 -> r.AntiAliasingMethod=1 r.FXAA.Quality=0
- r.AntiAliasingMethod=2 r.PostProcessAAQuality=2 -> r.AntiAliasingMethod=1 r.FXAA.Quality=0
- r.AntiAliasingMethod=2 r.PostProcessAAQuality=3 r.TemporalAAUpsampleFiltered=0 -> r.TemporalAA.Quality=0
- r.AntiAliasingMethod=2 r.PostProcessAAQuality=3 r.TemporalAAUpsampleFiltered=1 -> r.TemporalAA.Quality=1
- r.AntiAliasingMethod=2 r.PostProcessAAQuality=4 -> r.TemporalAA.Quality=2

MSAA r.AntiAliasingMethod=3 & TSR r.AntiAliasingMethod=4 remains unchanged.

sg.AntiAliasingQuality now maps to different r.TemporalAA.Quality or r.FXAA.Quality

Automated tests on base CL: https://horde.devtools.epicgames.com/job/60d5b8410123b700014f9db5
Automated tests on change CL: https://horde.devtools.epicgames.com/job/60d47cde57b302000114bebf

diff --git a/Engine/Config/BaseDeviceProfiles.ini b/Engine/Config/BaseDeviceProfiles.ini
--- a/Engine/Config/BaseDeviceProfiles.ini
+++ b/Engine/Config/BaseDeviceProfiles.ini
@@ -867,7 +867,7 @@ BaseProfileName=
 +CVars=r.LightShaftQuality=0
 +CVars=r.RefractionQuality=0
 +CVars=r.ShadowQuality=2
-+CVars=r.PostProcessAAQuality=0
++CVars=r.AntiAliasingMethod=0
 +CVars=r.Mobile.ForceDepthResolve=1
 
 [LuminClient DeviceProfile]
diff --git a/Engine/Config/BaseScalability.ini b/Engine/Config/BaseScalability.ini
--- a/Engine/Config/BaseScalability.ini
+++ b/Engine/Config/BaseScalability.ini
@@ -35,23 +35,28 @@ PerfIndexThresholds_ShadingQuality="GPU 18 42 115"
 PerfIndexValues_ResolutionQuality="50 71 87 100 100"
 
 [AntiAliasingQuality@0]
-r.PostProcessAAQuality=0
+r.FXAA.Quality=0
+r.TemporalAA.Quality=0
 r.TSR.RejectionAntiAliasingQuality=0
 
 [AntiAliasingQuality@1]
-r.PostProcessAAQuality=2
+r.FXAA.Quality=1
+r.TemporalAA.Quality=1
 r.TSR.RejectionAntiAliasingQuality=0
 
 [AntiAliasingQuality@2]
-r.PostProcessAAQuality=3
+r.FXAA.Quality=3
+r.TemporalAA.Quality=1
 r.TSR.RejectionAntiAliasingQuality=1
 
 [AntiAliasingQuality@3]
-r.PostProcessAAQuality=4
+r.FXAA.Quality=4
+r.TemporalAA.Quality=2
 r.TSR.RejectionAntiAliasingQuality=1
 
 [AntiAliasingQuality@Cine]
-r.PostProcessAAQuality=6
+r.FXAA.Quality=5
+r.TemporalAA.Quality=2
 r.TSR.RejectionAntiAliasingQuality=2
 
 
diff --git a/Engine/Shaders/Private/FXAAShader.usf b/Engine/Shaders/Private/FXAAShader.usf
--- a/Engine/Shaders/Private/FXAAShader.usf
+++ b/Engine/Shaders/Private/FXAAShader.usf
@@ -17,13 +17,13 @@ SamplerState Input_Sampler;
 	#define FXAA_QUALITY__PRESET 10
 	#define FXAA_PC_CONSOLE 1
 #elif FXAA_PRESET == 1
-	#define FXAA_QUALITY__PRESET 15
-	#define FXAA_PC_CONSOLE 1
+	#define FXAA_QUALITY__PRESET 10
+	#define FXAA_PC 1
 #elif FXAA_PRESET == 2
-	#define FXAA_QUALITY__PRESET 20
-	#define FXAA_PC_CONSOLE 1
+	#define FXAA_QUALITY__PRESET 13
+	#define FXAA_PC 1
 #elif FXAA_PRESET == 3
-	#define FXAA_QUALITY__PRESET 20
+	#define FXAA_QUALITY__PRESET 15
 	#define FXAA_PC 1
 #elif FXAA_PRESET == 4
 	#define FXAA_QUALITY__PRESET 29
diff --git a/Engine/Shaders/Private/TemporalAA.usf b/Engine/Shaders/Private/TemporalAA.usf
--- a/Engine/Shaders/Private/TemporalAA.usf
+++ b/Engine/Shaders/Private/TemporalAA.usf
@@ -54,6 +54,10 @@
 	// Min max samples that are within distance from output pixel. 
 	#define HISTORY_CLAMPING_BOX_SAMPLE_DISTANCE 2
 
+/** Qualities */
+	#define TAA_QUALITY_LOW 0
+	#define TAA_QUALITY_MEDIUM 1
+	#define TAA_QUALITY_HIGH 2
 
 //------------------------------------------------------- CONFIGS
 
@@ -71,13 +75,24 @@
 	#define AA_BICUBIC 1
 	#define AA_CROSS 2
 	#define AA_DYNAMIC 1
-	#define AA_FILTERED 1
 	#define AA_MANUALLY_CLAMP_HISTORY_UV 1
 	#define AA_TONE 1
 	#define AA_YCOCG 1
 
-	#if !TAA_FAST
+	#if TAA_QUALITY == TAA_QUALITY_LOW
+		#define AA_FILTERED 0
+		#define AA_DYNAMIC_ANTIGHOST 0
+
+	#elif TAA_QUALITY == TAA_QUALITY_MEDIUM
+		#define AA_FILTERED 1
+		#define AA_DYNAMIC_ANTIGHOST 0
+
+	#elif TAA_QUALITY == TAA_QUALITY_HIGH
+		#define AA_FILTERED 1
 		#define AA_DYNAMIC_ANTIGHOST 1
+
+	#else
+		#error Unknown TAA quality
 	#endif
 
 	#if COMPUTESHADER
@@ -113,7 +128,6 @@
 	#define AA_BICUBIC 1
 	#define AA_CROSS 1
 	#define AA_DYNAMIC 1
-	#define AA_FILTERED TAA_UPSAMPLE_FILTERED
 	#define AA_MANUALLY_CLAMP_HISTORY_UV 1
 	#define AA_TONE 1
 	#if SWITCH_PROFILE || SWITCH_PROFILE_FORWARD
@@ -124,24 +138,31 @@
 	
 	#define AA_UPSAMPLE 1
 	#define AA_UPSAMPLE_ADAPTIVE_FILTERING 1
-	
-	#if TAA_FAST
-		// When doing TAA downsampling, still need the 9 samples.
-		#if AA_DOWNSAMPLE
-			#define AA_SAMPLES 9
+
+	#if TAA_QUALITY == TAA_QUALITY_LOW
+		#define AA_FILTERED 0
+		#if SWITCH_PROFILE || SWITCH_PROFILE_FORWARD || AA_MOBILE_CONFIG
+			#define AA_SAMPLES 5
 		#else
-			#if SWITCH_PROFILE || SWITCH_PROFILE_FORWARD || AA_MOBILE_CONFIG
-				#define AA_SAMPLES 5
-			#else
-				#define AA_SAMPLES 6
-			#endif
+			#define AA_SAMPLES 6
 		#endif
 
-	#else
+	#elif TAA_QUALITY == TAA_QUALITY_MEDIUM
+		#define AA_FILTERED 1
+		#if SWITCH_PROFILE || SWITCH_PROFILE_FORWARD || AA_MOBILE_CONFIG
+			#define AA_SAMPLES 5
+		#else
+			#define AA_SAMPLES 6
+		#endif
+
+	#elif TAA_QUALITY == TAA_QUALITY_HIGH
 		#define AA_HISTORY_CLAMPING_BOX (HISTORY_CLAMPING_BOX_SAMPLE_DISTANCE)
+		#define AA_FILTERED 1
 		#define AA_DYNAMIC_ANTIGHOST 1
 		#define AA_SAMPLES 9
 
+	#else
+		#error Unknown TAA quality
 	#endif
 
 	// Do not use LDS caching for screen percentage > 100% or < 50%.
@@ -179,14 +200,20 @@
 	#else
 		#define AA_UPSAMPLE 0
 	#endif
-	
-	#if !TAA_FAST
+
+	#if TAA_QUALITY == TAA_QUALITY_MEDIUM
+		#if AA_UPSAMPLE
+			#define AA_BICUBIC 1
+			#define AA_FILTERED 1
+		#endif
+
+	#elif TAA_QUALITY == TAA_QUALITY_HIGH
 		#define AA_BICUBIC 1
 		#define AA_FILTERED 1
 		#define AA_YCOCG 1
-	#elif AA_UPSAMPLE
-		#define AA_BICUBIC 1
-		#define AA_FILTERED 1
+
+	#else
+		#error Unknown TAA quality
 	#endif
 
 	#if TAA_SCREEN_PERCENTAGE_RANGE != 2 && AA_UPSAMPLE
diff --git a/Engine/Source/Developer/FunctionalTesting/Private/AutomationBlueprintFunctionLibrary.cpp b/Engine/Source/Developer/FunctionalTesting/Private/AutomationBlueprintFunctionLibrary.cpp
--- a/Engine/Source/Developer/FunctionalTesting/Private/AutomationBlueprintFunctionLibrary.cpp
+++ b/Engine/Source/Developer/FunctionalTesting/Private/AutomationBlueprintFunctionLibrary.cpp
@@ -235,7 +235,6 @@ FAutomationTestScreenshotEnvSetup::FAutomationTestScreenshotEnvSetup()
 	: DefaultFeature_AntiAliasing(TEXT("r.AntiAliasingMethod"))
 	, DefaultFeature_AutoExposure(TEXT("r.DefaultFeature.AutoExposure"))
 	, DefaultFeature_MotionBlur(TEXT("r.DefaultFeature.MotionBlur"))
-	, PostProcessAAQuality(TEXT("r.PostProcessAAQuality"))
 	, MotionBlurQuality(TEXT("r.MotionBlurQuality"))
 	, ScreenSpaceReflectionQuality(TEXT("r.SSR.Quality"))
 	, EyeAdaptationQuality(TEXT("r.EyeAdaptationQuality"))
@@ -261,7 +260,6 @@ void FAutomationTestScreenshotEnvSetup::Setup(UWorld* InWorld, FAutomationScreen
 		DefaultFeature_AntiAliasing.Set(0);
 		DefaultFeature_AutoExposure.Set(0);
 		DefaultFeature_MotionBlur.Set(0);
-		PostProcessAAQuality.Set(0);
 		MotionBlurQuality.Set(0);
 		ScreenSpaceReflectionQuality.Set(0);
 		ContactShadows.Set(0);
@@ -307,7 +305,6 @@ void FAutomationTestScreenshotEnvSetup::Restore()
 	DefaultFeature_AntiAliasing.Restore();
 	DefaultFeature_AutoExposure.Restore();
 	DefaultFeature_MotionBlur.Restore();
-	PostProcessAAQuality.Restore();
 	MotionBlurQuality.Restore();
 	ScreenSpaceReflectionQuality.Restore();
 	EyeAdaptationQuality.Restore();
diff --git a/Engine/Source/Developer/FunctionalTesting/Public/AutomationBlueprintFunctionLibrary.h b/Engine/Source/Developer/FunctionalTesting/Public/AutomationBlueprintFunctionLibrary.h
--- a/Engine/Source/Developer/FunctionalTesting/Public/AutomationBlueprintFunctionLibrary.h
+++ b/Engine/Source/Developer/FunctionalTesting/Public/AutomationBlueprintFunctionLibrary.h
@@ -218,7 +218,6 @@ private:
 	FConsoleVariableSwapperTempl<int32> DefaultFeature_AntiAliasing;
 	FConsoleVariableSwapperTempl<int32> DefaultFeature_AutoExposure;
 	FConsoleVariableSwapperTempl<int32> DefaultFeature_MotionBlur;
-	FConsoleVariableSwapperTempl<int32> PostProcessAAQuality;
 	FConsoleVariableSwapperTempl<int32> MotionBlurQuality;
 	FConsoleVariableSwapperTempl<int32> ScreenSpaceReflectionQuality;
 	FConsoleVariableSwapperTempl<int32> EyeAdaptationQuality;
diff --git a/Engine/Source/Runtime/Core/Private/HAL/ConsoleManager.cpp b/Engine/Source/Runtime/Core/Private/HAL/ConsoleManager.cpp
--- a/Engine/Source/Runtime/Core/Private/HAL/ConsoleManager.cpp
+++ b/Engine/Source/Runtime/Core/Private/HAL/ConsoleManager.cpp
@@ -2297,13 +2297,6 @@ static TAutoConsoleVariable<int32> CVarMotionBlurQuality(
 	TEXT(" 0:off, 1:low, 2:medium, 3:high (default), 4: very high"),
 	ECVF_Scalability | ECVF_RenderThreadSafe);
 
-static TAutoConsoleVariable<int32> CVarPostProcessAAQuality(
-	TEXT("r.PostProcessAAQuality"),
-	4,
-	TEXT("Defines the postprocess anti aliasing method which allows to adjust for quality or performance.\n")
-	TEXT(" 0:off, 1:very low (faster FXAA), 2:low (FXAA), 3:medium (faster TemporalAA), 4:high (default TemporalAA), 5:very high, 6:max"),
-	ECVF_Scalability | ECVF_RenderThreadSafe);
-
 static TAutoConsoleVariable<int32> CVarFullscreenMode(
 	TEXT("r.FullScreenMode"),
 	1,
diff --git a/Engine/Source/Runtime/Engine/Private/SceneUtils.cpp b/Engine/Source/Runtime/Engine/Private/SceneUtils.cpp
--- a/Engine/Source/Runtime/Engine/Private/SceneUtils.cpp
+++ b/Engine/Source/Runtime/Engine/Private/SceneUtils.cpp
@@ -67,8 +67,6 @@ ENGINE_API EAntiAliasingMethod GetDefaultAntiAliasingMethod(const FStaticFeature
 
 	static auto* MSAACountCVar = IConsoleManager::Get().FindTConsoleVariableDataInt(TEXT("r.MSAACount"));
 
-	static auto* PostProcessAAQualityCVar = IConsoleManager::Get().FindTConsoleVariableDataInt(TEXT("r.PostProcessAAQuality"));
-
 	EShaderPlatform ShaderPlatform = GetFeatureLevelShaderPlatform(InFeatureLevel);
 
 	if (AntiAliasingMethod == EAntiAliasingMethod::AAM_None)
@@ -77,25 +75,11 @@ ENGINE_API EAntiAliasingMethod GetDefaultAntiAliasingMethod(const FStaticFeature
 	}
 	else if (AntiAliasingMethod == EAntiAliasingMethod::AAM_FXAA)
 	{
-		int32 Quality = PostProcessAAQualityCVar->GetValueOnAnyThread();
-
-		if (Quality <= 0)
-		{
-			AntiAliasingMethod = EAntiAliasingMethod::AAM_None;
-		}
+		// NOP
 	}
 	else if (AntiAliasingMethod == EAntiAliasingMethod::AAM_TemporalAA)
 	{
-		int32 Quality = PostProcessAAQualityCVar->GetValueOnAnyThread();
-
-		if (Quality <= 0)
-		{
-			AntiAliasingMethod = EAntiAliasingMethod::AAM_None;
-		}
-		else if (Quality == 1 || Quality == 2)
-		{
-			AntiAliasingMethod = AAM_FXAA;
-		}
+		// NOP
 	}
 	else if (AntiAliasingMethod == EAntiAliasingMethod::AAM_MSAA)
 	{
diff --git a/Engine/Source/Runtime/Renderer/Private/PostProcess/DiaphragmDOF.cpp b/Engine/Source/Runtime/Renderer/Private/PostProcess/DiaphragmDOF.cpp
--- a/Engine/Source/Runtime/Renderer/Private/PostProcess/DiaphragmDOF.cpp
+++ b/Engine/Source/Runtime/Renderer/Private/PostProcess/DiaphragmDOF.cpp
@@ -1426,7 +1426,7 @@ FRDGTextureRef DiaphragmDOF::AddPasses(
 		TAAParameters.SetupViewRect(View, PrefilteringResolutionDivisor);
 		TAAParameters.TopLeftCornerViewRects();
 
-		TAAParameters.bUseFast = CVarDOFTemporalAAQuality.GetValueOnRenderThread() == 0;
+		TAAParameters.Quality = CVarDOFTemporalAAQuality.GetValueOnRenderThread() == 0 ? ETAAQuality::Medium : ETAAQuality::High;
 	}
 
 	// Size of the view in GatherColorSetup.
diff --git a/Engine/Source/Runtime/Renderer/Private/PostProcess/PostProcessAA.cpp b/Engine/Source/Runtime/Renderer/Private/PostProcess/PostProcessAA.cpp
--- a/Engine/Source/Runtime/Renderer/Private/PostProcess/PostProcessAA.cpp
+++ b/Engine/Source/Runtime/Renderer/Private/PostProcess/PostProcessAA.cpp
@@ -3,6 +3,17 @@
 #include "PostProcess/PostProcessAA.h"
 #include "PostProcess/PostProcessing.h"
 
+TAutoConsoleVariable<int32> CVarFXAAQuality(
+	TEXT("r.FXAA.Quality"), 4,
+	TEXT("Selects the quality permutation of FXAA.\n")
+	TEXT(" 0: Console\n")
+	TEXT(" 1: PC medium-dither 3-sample\n")
+	TEXT(" 2: PC medium-dither 5-sample\n")
+	TEXT(" 3: PC medium-dither 8-sample\n")
+	TEXT(" 4: PC low-dither 12-sample (Default)\n")
+	TEXT(" 5: PC extrem quality 12-samples"),
+	ECVF_Scalability | ECVF_RenderThreadSafe);
+
 BEGIN_SHADER_PARAMETER_STRUCT(FFXAAParameters, )
 	SHADER_PARAMETER_STRUCT(FScreenPassTextureInput, Input)
 	SHADER_PARAMETER(FVector4, fxaaConsoleRcpFrameOpt)
@@ -53,9 +64,7 @@ IMPLEMENT_GLOBAL_SHADER(FFXAAPS, "/Engine/Private/FXAAShader.usf", "FxaaPS", SF_
 
 EFXAAQuality GetFXAAQuality()
 {
-	const EPostProcessAAQuality PostProcessAAQuality = GetPostProcessAAQuality();
-	static_assert(uint32(EPostProcessAAQuality::MAX) == uint32(EFXAAQuality::MAX), "FXAA quality levels don't match post process AA quality levels. Can't trivially convert.");
-	return static_cast<EFXAAQuality>(PostProcessAAQuality);
+	return EFXAAQuality(FMath::Clamp(CVarFXAAQuality.GetValueOnRenderThread(), 0, 5));
 }
 
 FScreenPassTexture AddFXAAPass(FRDGBuilder& GraphBuilder, const FViewInfo& View, const FFXAAInputs& Inputs)
diff --git a/Engine/Source/Runtime/Renderer/Private/PostProcess/PostProcessing.cpp b/Engine/Source/Runtime/Renderer/Private/PostProcess/PostProcessing.cpp
--- a/Engine/Source/Runtime/Renderer/Private/PostProcess/PostProcessing.cpp
+++ b/Engine/Source/Runtime/Renderer/Private/PostProcess/PostProcessing.cpp
@@ -144,13 +144,6 @@ bool IsPostProcessingWithAlphaChannelSupported()
 	return CVarPostProcessingPropagateAlpha.GetValueOnAnyThread() != 0;
 }
 
-EPostProcessAAQuality GetPostProcessAAQuality()
-{
-	static const auto CVar = IConsoleManager::Get().FindTConsoleVariableDataInt(TEXT("r.PostProcessAAQuality"));
-
-	return static_cast<EPostProcessAAQuality>(FMath::Clamp(CVar->GetValueOnAnyThread(), 0, static_cast<int32>(EPostProcessAAQuality::MAX) - 1));
-}
-
 class FComposeSeparateTranslucencyPS : public FGlobalShader
 {
 	DECLARE_GLOBAL_SHADER(FComposeSeparateTranslucencyPS);
diff --git a/Engine/Source/Runtime/Renderer/Private/PostProcess/PostProcessing.h b/Engine/Source/Runtime/Renderer/Private/PostProcess/PostProcessing.h
--- a/Engine/Source/Runtime/Renderer/Private/PostProcess/PostProcessing.h
+++ b/Engine/Source/Runtime/Renderer/Private/PostProcess/PostProcessing.h
@@ -14,24 +14,6 @@ namespace Nanite
 	struct FRasterResults;
 }
 
-enum class EPostProcessAAQuality : uint32
-{
-	Disabled,
-	// Faster FXAA
-	VeryLow,
-	// FXAA
-	Low,
-	// Faster Temporal AA
-	Medium,
-	// Temporal AA
-	High,
-	VeryHigh,
-	MAX
-};
-
-// Returns the quality of post process anti-aliasing defined by CVar.
-EPostProcessAAQuality GetPostProcessAAQuality();
-
 // Returns whether the full post process pipeline is enabled. Otherwise, the minimal set of operations are performed.
 bool IsPostProcessingEnabled(const FViewInfo& View);
 
diff --git a/Engine/Source/Runtime/Renderer/Private/PostProcess/TemporalAA.cpp b/Engine/Source/Runtime/Renderer/Private/PostProcess/TemporalAA.cpp
--- a/Engine/Source/Runtime/Renderer/Private/PostProcess/TemporalAA.cpp
+++ b/Engine/Source/Runtime/Renderer/Private/PostProcess/TemporalAA.cpp
@@ -41,10 +41,12 @@ TAutoConsoleVariable<float> CVarTemporalAACurrentFrameWeight(
 	TEXT("Weight of current frame's contribution to the history.  Low values cause blurriness and ghosting, high values fail to hide jittering."),
 	ECVF_Scalability | ECVF_RenderThreadSafe);
 
-TAutoConsoleVariable<int32> CVarTemporalAAUpsampleFiltered(
-	TEXT("r.TemporalAAUpsampleFiltered"),
-	1,
-	TEXT("Use filtering to fetch color history during TamporalAA upsampling (see AA_FILTERED define in TAA shader). Disabling this makes TAAU faster, but lower quality. "),
+TAutoConsoleVariable<int32> CVarTemporalAAQuality(
+	TEXT("r.TemporalAA.Quality"), 2,
+	TEXT("Quality of the main Temporal AA pass.\n")
+	TEXT(" 0: Disable input filtering;\n")
+	TEXT(" 1: Enable input filtering;\n")
+	TEXT(" 2: Enable input filtering, enable mobility based anti-ghosting (Default)"),
 	ECVF_Scalability | ECVF_RenderThreadSafe);
 
 TAutoConsoleVariable<float> CVarTemporalAAHistorySP(
@@ -85,17 +87,15 @@ class FTemporalAACS : public FGlobalShader
 	SHADER_USE_PARAMETER_STRUCT(FTemporalAACS, FGlobalShader);
 
 	class FTAAPassConfigDim : SHADER_PERMUTATION_ENUM_CLASS("TAA_PASS_CONFIG", ETAAPassConfig);
-	class FTAAFastDim : SHADER_PERMUTATION_BOOL("TAA_FAST");
+	class FTAAQualityDim : SHADER_PERMUTATION_ENUM_CLASS("TAA_QUALITY", ETAAQuality);
 	class FTAAResponsiveDim : SHADER_PERMUTATION_BOOL("TAA_RESPONSIVE");
 	class FTAAScreenPercentageDim : SHADER_PERMUTATION_INT("TAA_SCREEN_PERCENTAGE_RANGE", 4);
-	class FTAAUpsampleFilteredDim : SHADER_PERMUTATION_BOOL("TAA_UPSAMPLE_FILTERED");
 	class FTAADownsampleDim : SHADER_PERMUTATION_BOOL("TAA_DOWNSAMPLE");
 
 	using FPermutationDomain = TShaderPermutationDomain<
 		FTAAPassConfigDim,
-		FTAAFastDim,
+		FTAAQualityDim,
 		FTAAScreenPercentageDim,
-		FTAAUpsampleFilteredDim,
 		FTAADownsampleDim>;
 
 	BEGIN_SHADER_PARAMETER_STRUCT(FParameters, )
@@ -156,10 +156,49 @@ class FTemporalAACS : public FGlobalShader
 		SHADER_PARAMETER_RDG_TEXTURE_UAV(RWTexture2D, DebugOutput)
 	END_SHADER_PARAMETER_STRUCT()
 
+	static FPermutationDomain RemapPermutation(FPermutationDomain PermutationVector)
+	{
+		if (PermutationVector.Get<FTAAPassConfigDim>() == ETAAPassConfig::Main ||
+			PermutationVector.Get<FTAAPassConfigDim>() == ETAAPassConfig::MainUpsampling)
+		{
+			// No point downsampling if not using the faster quality permutation already.
+			if (PermutationVector.Get<FTAAQualityDim>() == ETAAQuality::High)
+			{
+				PermutationVector.Set<FTAADownsampleDim>(false);
+			}
+		}
+		else if (
+			PermutationVector.Get<FTAAPassConfigDim>() == ETAAPassConfig::DiaphragmDOF ||
+			PermutationVector.Get<FTAAPassConfigDim>() == ETAAPassConfig::DiaphragmDOFUpsampling)
+		{
+			// DOF pass allow only quality 1 and 2
+			PermutationVector.Set<FTAAQualityDim>(ETAAQuality(FMath::Max(int(PermutationVector.Get<FTAAQualityDim>()), int(ETAAQuality::Medium))));
+
+			// Only the Main and Main Upsampling can downsample the output.
+			PermutationVector.Set<FTAADownsampleDim>(false);
+		}
+		else
+		{
+			// Only the Main and Main Upsampling have quality options 0, 1, 2.
+			PermutationVector.Set<FTAAQualityDim>(ETAAQuality::High);
+
+			// Only the Main and Main Upsampling can downsample the output.
+			PermutationVector.Set<FTAADownsampleDim>(false);
+		}
+
+		return PermutationVector;
+	}
+
 	static bool ShouldCompilePermutation(const FGlobalShaderPermutationParameters& Parameters)
 	{
 		FPermutationDomain PermutationVector(Parameters.PermutationId);
 
+		// Don't compile the shader permutation if gets remaped at runtime.
+		if (PermutationVector != RemapPermutation(PermutationVector))
+		{
+			return false;
+		}
+
 		// Screen percentage dimension is only for upsampling permutation.
 		if (!IsTAAUpsamplingConfig(PermutationVector.Get<FTAAPassConfigDim>()) &&
 			PermutationVector.Get<FTAAScreenPercentageDim>() != 0)
@@ -174,32 +213,6 @@ class FTemporalAACS : public FGlobalShader
 			{
 				return false;
 			}
-
-			// No point disabling filtering.
-			if (!PermutationVector.Get<FTAAUpsampleFilteredDim>())
-			{
-				return false;
-			}
-
-			// No point doing a fast permutation since it is PC only.
-			if (PermutationVector.Get<FTAAFastDim>())
-			{
-				return false;
-			}
-		}
-
-		// No point disabling filtering if not using the fast permutation already.
-		if (!PermutationVector.Get<FTAAUpsampleFilteredDim>() &&
-			!PermutationVector.Get<FTAAFastDim>())
-		{
-			return false;
-		}
-
-		// No point downsampling if not using the fast permutation already.
-		if (PermutationVector.Get<FTAADownsampleDim>() &&
-			!PermutationVector.Get<FTAAFastDim>())
-		{
-			return false;
 		}
 
 		// Screen percentage range 3 is only for super sampling.
@@ -208,28 +221,6 @@ class FTemporalAACS : public FGlobalShader
 		{
 			return false;
 		}
-
-		// Fast dimensions is only for Main and Diaphragm DOF.
-		if (PermutationVector.Get<FTAAFastDim>() &&
-			!IsMainTAAConfig(PermutationVector.Get<FTAAPassConfigDim>()) &&
-			!IsDOFTAAConfig(PermutationVector.Get<FTAAPassConfigDim>()))
-		{
-			return false;
-		}
-
-		// Non filtering option is only for upsampling.
-		if (!PermutationVector.Get<FTAAUpsampleFilteredDim>() &&
-			PermutationVector.Get<FTAAPassConfigDim>() != ETAAPassConfig::MainUpsampling)
-		{
-			return false;
-		}
-
-		// TAA_DOWNSAMPLE is only only for Main and MainUpsampling configs.
-		if (PermutationVector.Get<FTAADownsampleDim>() &&
-			!IsMainTAAConfig(PermutationVector.Get<FTAAPassConfigDim>()))
-		{
-			return false;
-		}
 		
 		//Only Main and MainUpsampling config without DownSample permutations are supported on mobile platform.
 		return SupportsGen4TAA(Parameters.Platform) && (!IsMobilePlatform(Parameters.Platform) || ((PermutationVector.Get<FTAAPassConfigDim>() == ETAAPassConfig::Main || PermutationVector.Get<FTAAPassConfigDim>() == ETAAPassConfig::MainUpsampling) && !PermutationVector.Get<FTAADownsampleDim>()));
@@ -340,17 +331,24 @@ const TCHAR* const kTAAOutputNames[] = {
 };
 
 const TCHAR* const kTAAPassNames[] = {
-	TEXT("Gen4 Main"),
-	TEXT("Gen4 MainUpsampling"),
-	TEXT("Gen4 MainSuperSampling"),
+	TEXT("Main"),
+	TEXT("MainUpsampling"),
+	TEXT("MainSuperSampling"),
 	TEXT("ScreenSpaceReflections"),
 	TEXT("LightShaft"),
 	TEXT("DOF"),
 	TEXT("DOFUpsampling"),
 };
 
+const TCHAR* const kTAAQualityNames[] = {
+	TEXT("Low"),
+	TEXT("Medium"),
+	TEXT("High"),
+};
+
 static_assert(UE_ARRAY_COUNT(kTAAOutputNames) == int32(ETAAPassConfig::MAX), "Missing TAA output name.");
 static_assert(UE_ARRAY_COUNT(kTAAPassNames) == int32(ETAAPassConfig::MAX), "Missing TAA pass name.");
+static_assert(UE_ARRAY_COUNT(kTAAQualityNames) == int32(ETAAQuality::MAX), "Missing TAA quality name.");
 } //! namespace
 
 bool IsTemporalAASceneDownsampleAllowed(const FViewInfo& View)
@@ -464,7 +462,7 @@ FTAAOutputs AddTemporalAAPass(
 
 	{
 		EPixelFormat HistoryPixelFormat = PF_FloatRGBA;
-		if (bIsMainPass && Inputs.bUseFast && !bSupportsAlpha && CVarTAAR11G11B10History.GetValueOnRenderThread())
+		if (bIsMainPass && Inputs.Quality != ETAAQuality::High && !bSupportsAlpha && CVarTAAR11G11B10History.GetValueOnRenderThread())
 		{
 			HistoryPixelFormat = PF_FloatR11G11B10;
 		}
@@ -524,15 +522,11 @@ FTAAOutputs AddTemporalAAPass(
 	{
 		FTemporalAACS::FPermutationDomain PermutationVector;
 		PermutationVector.Set<FTemporalAACS::FTAAPassConfigDim>(Inputs.Pass);
-		PermutationVector.Set<FTemporalAACS::FTAAFastDim>(Inputs.bUseFast);
+		PermutationVector.Set<FTemporalAACS::FTAAQualityDim>(Inputs.Quality);
 		PermutationVector.Set<FTemporalAACS::FTAADownsampleDim>(Inputs.bDownsample);
-		PermutationVector.Set<FTemporalAACS::FTAAUpsampleFilteredDim>(true);
 
 		if (IsTAAUpsamplingConfig(Inputs.Pass))
 		{
-			const bool bUpsampleFiltered = CVarTemporalAAUpsampleFiltered.GetValueOnRenderThread() != 0 || Inputs.Pass != ETAAPassConfig::MainUpsampling;
-			PermutationVector.Set<FTemporalAACS::FTAAUpsampleFilteredDim>(bUpsampleFiltered);
-
 			// If screen percentage > 100% on X or Y axes, then use screen percentage range = 2 shader permutation to disable LDS caching.
 			if (SrcRect.Width() > DestRect.Width() ||
 				SrcRect.Height() > DestRect.Height())
@@ -554,6 +548,8 @@ FTAAOutputs AddTemporalAAPass(
 			}
 		}
 
+		PermutationVector = FTemporalAACS::RemapPermutation(PermutationVector);
+
 		FTemporalAACS::FParameters* PassParameters = GraphBuilder.AllocParameters<FTemporalAACS::FParameters>();
 
 		// Setups common shader parameters
@@ -754,8 +750,9 @@ FTAAOutputs AddTemporalAAPass(
 
 		FComputeShaderUtils::AddPass(
 			GraphBuilder,
-			RDG_EVENT_NAME("TAA Gen4 %s%s %dx%d -> %dx%d",
-				PassName, Inputs.bUseFast ? TEXT(" Fast") : TEXT(""),
+			RDG_EVENT_NAME("TAA(%s Quality=%s) %dx%d -> %dx%d",
+				PassName,
+				kTAAQualityNames[int32(PermutationVector.Get<FTemporalAACS::FTAAQualityDim>())],
 				PracticableSrcRect.Width(), PracticableSrcRect.Height(),
 				PracticableDestRect.Width(), PracticableDestRect.Height()),
 			ComputeShader,
@@ -801,9 +798,7 @@ static void AddGen4MainTemporalAAPasses(
 
 	TAAParameters.SetupViewRect(View);
 
-	const EPostProcessAAQuality LowQualityTemporalAA = EPostProcessAAQuality::Medium;
-
-	TAAParameters.bUseFast = GetPostProcessAAQuality() == LowQualityTemporalAA;
+	TAAParameters.Quality = ETAAQuality(FMath::Clamp(CVarTemporalAAQuality.GetValueOnRenderThread(), 0, 2));
 
 	const FIntRect SecondaryViewRect = TAAParameters.OutputViewRect;
 
@@ -819,7 +814,7 @@ static void AddGen4MainTemporalAAPasses(
 			TAAParameters.OutputViewRect.Height() * HistoryUpscaleFactor);
 
 		TAAParameters.Pass = ETAAPassConfig::MainSuperSampling;
-		TAAParameters.bUseFast = false;
+		TAAParameters.Quality = ETAAQuality::High;
 
 		TAAParameters.OutputViewRect.Min.X = 0;
 		TAAParameters.OutputViewRect.Min.Y = 0;
@@ -828,7 +823,7 @@ static void AddGen4MainTemporalAAPasses(
 
 	TAAParameters.DownsampleOverrideFormat = PassInputs.DownsampleOverrideFormat;
 
-	TAAParameters.bDownsample = PassInputs.bAllowDownsampleSceneColor && TAAParameters.bUseFast;
+	TAAParameters.bDownsample = PassInputs.bAllowDownsampleSceneColor && TAAParameters.Quality != ETAAQuality::High;
 
 	TAAParameters.SceneDepthTexture = PassInputs.SceneDepthTexture;
 	TAAParameters.SceneVelocityTexture = PassInputs.SceneVelocityTexture;
diff --git a/Engine/Source/Runtime/Renderer/Private/PostProcess/TemporalAA.h b/Engine/Source/Runtime/Renderer/Private/PostProcess/TemporalAA.h
--- a/Engine/Source/Runtime/Renderer/Private/PostProcess/TemporalAA.h
+++ b/Engine/Source/Runtime/Renderer/Private/PostProcess/TemporalAA.h
@@ -77,6 +77,14 @@ struct FTAAOutputs
 	FRDGTexture* DownsampledSceneColor = nullptr;
 };
 
+/** Quality of TAA. */
+enum class ETAAQuality : uint8
+{
+	Low,
+	Medium,
+	High,
+	MAX
+};
 
 /** Configuration of TAA. */
 struct FTAAPassParameters
@@ -85,7 +93,7 @@ struct FTAAPassParameters
 	ETAAPassConfig Pass = ETAAPassConfig::Main;
 
 	// Whether to use the faster shader permutation.
-	bool bUseFast = false;
+	ETAAQuality Quality = ETAAQuality::High;
 
 	// Whether output texture should be render targetable.
 	bool bOutputRenderTargetable = false;
