Replaces old grain jitter with a proper FilmGrain using tiled repeated texture.

Waiting on ITH-152309 to compress memory footprint of Marcie_Grain texture that is used by default.

diff --git a/Engine/Config/Android/AndroidEngine.ini b/Engine/Config/Android/AndroidEngine.ini
--- a/Engine/Config/Android/AndroidEngine.ini
+++ b/Engine/Config/Android/AndroidEngine.ini
@@ -2,6 +2,9 @@
 ; DefaultBloomKernel is not used on mobile
 DefaultBloomKernelTextureName=/Engine/EngineResources/DefaultTexture.DefaultTexture
 
+; Default film grain is not used on mobile
+DefaultFilmGrainTextureName=""
+
 [PlatformCrypto]
 PlatformRequiresDataCrypto=True
 PakSigningRequired=False
diff --git a/Engine/Config/BaseEngine.ini b/Engine/Config/BaseEngine.ini
--- a/Engine/Config/BaseEngine.ini
+++ b/Engine/Config/BaseEngine.ini
@@ -105,6 +105,7 @@ DefaultBSPVertexTextureName=/Engine/EditorResources/BSPVertex.BSPVertex
 HighFrequencyNoiseTextureName=/Engine/EngineMaterials/Good64x64TilingNoiseHighFreq.Good64x64TilingNoiseHighFreq
 DefaultBokehTextureName=/Engine/EngineMaterials/DefaultBokeh.DefaultBokeh
 DefaultBloomKernelTextureName=/Engine/EngineMaterials/DefaultBloomKernel.DefaultBloomKernel
+DefaultFilmGrainTextureName=/Engine/EngineMaterials/Marcie_Grain_1024_M5_555_A0.Marcie_Grain_1024_M5_555_A0
 GeomMaterialName=/Engine/EngineDebugMaterials/GeomMaterial.GeomMaterial
 DebugMeshMaterialName=/Engine/EngineDebugMaterials/DebugMeshMaterial.DebugMeshMaterial
 EmissiveMeshMaterialName=/Engine/EngineMaterials/EmissiveMeshMaterial.EmissiveMeshMaterial
diff --git a/Engine/Config/IOS/IOSEngine.ini b/Engine/Config/IOS/IOSEngine.ini
--- a/Engine/Config/IOS/IOSEngine.ini
+++ b/Engine/Config/IOS/IOSEngine.ini
@@ -2,6 +2,9 @@
 ; DefaultBloomKernel is not used on mobile
 DefaultBloomKernelTextureName=/Engine/EngineResources/DefaultTexture.DefaultTexture
 
+; Default film grain is not used on mobile
+DefaultFilmGrainTextureName=""
+
 [/Script/Engine.GarbageCollectionSettings]
 gc.MaxObjectsInGame=131072
 
diff --git a/Engine/Shaders/Private/PostProcessTonemap.usf b/Engine/Shaders/Private/PostProcessTonemap.usf
--- a/Engine/Shaders/Private/PostProcessTonemap.usf
+++ b/Engine/Shaders/Private/PostProcessTonemap.usf
@@ -56,9 +56,12 @@ float4 ColorScale1;
 // from the postprocess settings, x:VignetteIntensity, y:SharpenDiv6
 float4 TonemapperParams;
 
-// Fine film grain
+// Film grain
 float3 GrainRandomFull;
-float3 GrainScaleBiasJitter;
+float FilmGrainMul;
+float FilmGrainAdd;
+Texture2D<half3> FilmGrainTexture;
+float4 ScreenPosToFilmGrainTextureUV;
 
 float EditorNITLevel;
 
@@ -342,7 +345,7 @@ half3 LookupSceneColor(float2 SceneUV, int2 PixelOffset)
 
 float4 SampleSceneColor(float2 SceneUV)
 {
-	#if USE_GRAIN_JITTER
+	#if USE_COLOR_FRINGE
 		SceneUV = clamp(SceneUV, Color_UVViewportBilinearMin, Color_UVViewportBilinearMax);
 	#endif
 
@@ -384,19 +387,10 @@ float4 TonemapCommonPS(
 	}
 #endif
 
-	#if USE_GRAIN_JITTER || USE_GRAIN_INTENSITY || USE_GRAIN_QUANTIZATION
-		half Grain = GrainFromUV(GrainUV.zw);
-	#endif
 	float2 SceneUV = UV.xy;
-	#if USE_GRAIN_JITTER
-		SceneUV = lerp(UV.xy, GrainUV.xy, (1.0 - Grain*Grain) * GrainScaleBiasJitter.z);
-	#endif
 
 	#if USE_COLOR_FRINGE
 		float2 SceneUVJitter = float2(0.0, 0.0);
-		#if USE_GRAIN_JITTER
-			SceneUVJitter = SceneUV.xy - UV.xy;
-		#endif
 		
 	        float2 CAScale = ChromaticAberrationParams.rg;
 	        float StartOffset = ChromaticAberrationParams.z;
@@ -567,17 +561,24 @@ float4 TonemapCommonPS(
 			LinearColor.rgb *= ComputeVignetteMask( ExposureScaleVignette.yz, TonemapperParams.x );
 		#endif
 
-		#if USE_GRAIN_INTENSITY
-			// Needs to go before tonemapping.
-			half GrainMul = Grain * GrainScaleBiasJitter.x + GrainScaleBiasJitter.y;
-			LinearColor.rgb *= GrainMul;
+		#if USE_FILM_GRAIN
+		{
+			float2 FilmGrainUV = ApplyScreenTransform(ScreenPos, ScreenPosToFilmGrainTextureUV);
+
+			half3 RawGrain = FilmGrainTexture.SampleLevel(GlobalBilinearWrappedSampler, FilmGrainUV, 0);
+			half3 FinalGrain = RawGrain * half(FilmGrainMul) + half(FilmGrainAdd);
+
+			LinearColor.rgb *= FinalGrain;
+		}
 		#endif
 
-			half3 OutDeviceColor = ColorLookupTable(LinearColor);
+		half3 OutDeviceColor = ColorLookupTable(LinearColor);
 
 		half LuminanceForPostProcessAA  = dot(OutDeviceColor, half3(0.299f, 0.587f, 0.114f));
 
 		#if USE_GRAIN_QUANTIZATION
+			half Grain = GrainFromUV(GrainUV.zw);
+
 			#if DIM_OUTPUT_DEVICE == TONEMAPPER_OUTPUT_NoToneCurve
 				#error LinearNoToneCurve device requires no Grain Quantization
 			#endif
diff --git a/Engine/Source/Runtime/Engine/Classes/Engine/Engine.h b/Engine/Source/Runtime/Engine/Classes/Engine/Engine.h
--- a/Engine/Source/Runtime/Engine/Classes/Engine/Engine.h
+++ b/Engine/Source/Runtime/Engine/Classes/Engine/Engine.h
@@ -840,6 +840,14 @@ public:
 	/** Path of the texture used to bloom when using FFT, mimics characteristic bloom produced in a camera from a signle bright source */
 	UPROPERTY(globalconfig)
 	FSoftObjectPath DefaultBloomKernelTextureName;
+	
+	/** Texture used to film grain by default. */
+	UPROPERTY()
+	TObjectPtr<class UTexture2D> DefaultFilmGrainTexture;
+
+	/** Path of the texture used by film grain by default. */
+	UPROPERTY(globalconfig)
+	FSoftObjectPath DefaultFilmGrainTextureName;
 
 	/** The material used to render wireframe meshes. */
 	UPROPERTY()
@@ -1728,6 +1736,9 @@ public:
 	*/
 	void LoadDefaultBloomTexture();
 
+	/** Delay loading this texture until it is needed by the renderer. */
+	void LoadDefaultFilmGrainTexture();
+
 private:
 	#if WITH_DYNAMIC_RESOLUTION
 		/** Last dynamic resolution event. */
diff --git a/Engine/Source/Runtime/Engine/Classes/Engine/Scene.h b/Engine/Source/Runtime/Engine/Classes/Engine/Scene.h
--- a/Engine/Source/Runtime/Engine/Classes/Engine/Scene.h
+++ b/Engine/Source/Runtime/Engine/Classes/Engine/Scene.h
@@ -927,11 +927,20 @@ struct FPostProcessSettings
 	UPROPERTY(EditAnywhere, BlueprintReadWrite, Category=Overrides, meta=(PinHiddenByDefault, InlineEditConditionToggle))
 	uint8 bOverride_VignetteIntensity:1;
 
+	UPROPERTY()
+	uint8 bOverride_GrainIntensity_DEPRECATED:1;
+
+	UPROPERTY()
+	uint8 bOverride_GrainJitter_DEPRECATED:1;
+	
 	UPROPERTY(EditAnywhere, BlueprintReadWrite, Category=Overrides, meta=(PinHiddenByDefault, InlineEditConditionToggle))
-	uint8 bOverride_GrainIntensity:1;
+	uint8 bOverride_FilmGrainIntensity:1;
 
 	UPROPERTY(EditAnywhere, BlueprintReadWrite, Category=Overrides, meta=(PinHiddenByDefault, InlineEditConditionToggle))
-	uint8 bOverride_GrainJitter:1;
+	uint8 bOverride_FilmGrainTexelSize:1;
+
+	UPROPERTY(EditAnywhere, BlueprintReadWrite, Category = Overrides, meta = (PinHiddenByDefault, InlineEditConditionToggle))
+	uint8 bOverride_FilmGrainTexture : 1;
 
 	UPROPERTY(EditAnywhere, BlueprintReadWrite, Category=Overrides, meta=(PinHiddenByDefault, InlineEditConditionToggle))
 	uint8 bOverride_AmbientOcclusionIntensity:1;
@@ -1696,13 +1705,31 @@ struct FPostProcessSettings
 	UPROPERTY(interp, BlueprintReadWrite, Category="Lens|Image Effects", meta=(UIMin = "0.0", UIMax = "1.0", editcondition = "bOverride_VignetteIntensity"))
 	float VignetteIntensity;
 
-	/** 0..1 grain jitter */
-	UPROPERTY(interp, BlueprintReadWrite, Category = "Lens|Image Effects", meta=(UIMin = "0.0", UIMax = "1.0", editcondition = "bOverride_GrainJitter"))
-	float GrainJitter;
+	UPROPERTY()
+	float GrainJitter_DEPRECATED;
+
+	UPROPERTY()
+	float GrainIntensity_DEPRECATED;
+
+	/** 0..1 Film grain intensity to apply. LinearSceneColor *= lerp(1.0, DecodedFilmGrainTexture, FilmGrainIntensity) */
+	UPROPERTY(interp, BlueprintReadWrite, Category = "Film|Grain", meta = (UIMin = "0.0", UIMax = "1.0", editcondition = "bOverride_FilmGrainIntensity"))
+	float FilmGrainIntensity;
+
+	/** Size of texel of FilmGrainTexture on screen. */
+	UPROPERTY(interp, BlueprintReadWrite, Category = "Film|Grain", meta = (UIMin = "0.0", UIMax = "4.0", editcondition = "bOverride_FilmGrainTexelSize"))
+	float FilmGrainTexelSize;
+
+	/** Defines film grain texture to use. */
+	UPROPERTY(EditAnywhere, BlueprintReadWrite, Category = "Film|Grain", meta = (editcondition = "bOverride_FilmGrainTexture"))
+	TObjectPtr<class UTexture2D> FilmGrainTexture;
+
+	/** Multiply the FilmGrainTexture. DecodedFilmGrainTexture = FilmGrainTexture * FilmGrainDecodeMultiply + FilmGrainDecodeAdd */
+	UPROPERTY(EditAnywhere, BlueprintReadWrite, Category = "Film|Grain", meta = (editcondition = "bOverride_FilmGrainTexture"))
+	float FilmGrainDecodeMultiply;
 
-	/** 0..1 grain intensity */
-	UPROPERTY(interp, BlueprintReadWrite, Category="Lens|Image Effects", meta=(UIMin = "0.0", UIMax = "1.0", editcondition = "bOverride_GrainIntensity"))
-	float GrainIntensity;
+	/** Add to the FilmGrainTexture. DecodedFilmGrainTexture = FilmGrainTexture * FilmGrainDecodeMultiply + FilmGrainDecodeAdd */
+	UPROPERTY(EditAnywhere, BlueprintReadWrite, Category = "Film|Grain", meta = (editcondition = "bOverride_FilmGrainTexture"))
+	float FilmGrainDecodeAdd;
 
 	/** 0..1 0=off/no ambient occlusion .. 1=strong ambient occlusion, defines how much it affects the non direct lighting after base pass */
 	UPROPERTY(interp, BlueprintReadWrite, Category="Rendering Features|Ambient Occlusion", meta=(ClampMin = "0.0", ClampMax = "1.0", editcondition = "bOverride_AmbientOcclusionIntensity", DisplayName = "Intensity"))
diff --git a/Engine/Source/Runtime/Engine/Private/Scene.cpp b/Engine/Source/Runtime/Engine/Private/Scene.cpp
--- a/Engine/Source/Runtime/Engine/Private/Scene.cpp
+++ b/Engine/Source/Runtime/Engine/Private/Scene.cpp
@@ -349,6 +349,12 @@ static void DoPostProcessSettingsSanityCheck()
 		RenameMap.Add(TEXT("FilmShadowTint"), TEXT(""));
 		RenameMap.Add(TEXT("FilmShadowTintBlend"), TEXT(""));
 		RenameMap.Add(TEXT("FilmShadowTintAmount"), TEXT(""));
+		// Film Grain are ignored
+		RenameMap.Add(TEXT("FilmGrainIntensity"), TEXT(""));
+		RenameMap.Add(TEXT("FilmGrainTexelSize"), TEXT(""));
+		RenameMap.Add(TEXT("FilmGrainTexture"), TEXT(""));
+		RenameMap.Add(TEXT("FilmGrainDecodeMultiply"), TEXT(""));
+		RenameMap.Add(TEXT("FilmGrainDecodeAdd"), TEXT(""));
 		VerifyPostProcessingProperties(TEXT("Film"),
 			TArray<const UStruct*>({FFilmStockSettings::StaticStruct()}),
 			RenameMap);
@@ -493,8 +499,12 @@ FPostProcessSettings::FPostProcessSettings()
 	LensFlareBokehSize = 3.0f;
 	LensFlareThreshold = 8.0f;
 	VignetteIntensity = 0.4f;
-	GrainIntensity = 0.0f;
-	GrainJitter = 0.0f;
+	GrainIntensity_DEPRECATED = 0.0f;
+	GrainJitter_DEPRECATED = 0.0f;
+	FilmGrainIntensity = 0.0f;
+	FilmGrainTexelSize = 1.0f;
+	FilmGrainDecodeMultiply = 1.0f;
+	FilmGrainDecodeAdd = 0.0f;
 	// next value might get overwritten by r.DefaultFeature.AmbientOcclusion
 	AmbientOcclusionIntensity = .5f;
 	// next value might get overwritten by r.DefaultFeature.AmbientOcclusionStaticFraction
@@ -694,8 +704,11 @@ FPostProcessSettings::FPostProcessSettings(const FPostProcessSettings& Settings)
 	, bOverride_LensFlareBokehShape(Settings.bOverride_LensFlareBokehShape)
 	, bOverride_LensFlareThreshold(Settings.bOverride_LensFlareThreshold)
 	, bOverride_VignetteIntensity(Settings.bOverride_VignetteIntensity)
-	, bOverride_GrainIntensity(Settings.bOverride_GrainIntensity)
-	, bOverride_GrainJitter(Settings.bOverride_GrainJitter)
+	, bOverride_GrainIntensity_DEPRECATED(Settings.bOverride_GrainIntensity_DEPRECATED)
+	, bOverride_GrainJitter_DEPRECATED(Settings.bOverride_GrainJitter_DEPRECATED)
+	, bOverride_FilmGrainIntensity(Settings.bOverride_FilmGrainIntensity)
+	, bOverride_FilmGrainTexelSize(Settings.bOverride_FilmGrainTexelSize)
+	, bOverride_FilmGrainTexture(Settings.bOverride_FilmGrainTexture)
 	, bOverride_AmbientOcclusionIntensity(Settings.bOverride_AmbientOcclusionIntensity)
 	, bOverride_AmbientOcclusionStaticFraction(Settings.bOverride_AmbientOcclusionStaticFraction)
 	, bOverride_AmbientOcclusionRadius(Settings.bOverride_AmbientOcclusionRadius)
@@ -886,8 +899,13 @@ FPostProcessSettings::FPostProcessSettings(const FPostProcessSettings& Settings)
 	, LensFlareThreshold(Settings.LensFlareThreshold)
 	, LensFlareBokehShape(Settings.LensFlareBokehShape)
 	, VignetteIntensity(Settings.VignetteIntensity)
-	, GrainJitter(Settings.GrainJitter)
-	, GrainIntensity(Settings.GrainIntensity)
+	, GrainJitter_DEPRECATED(Settings.GrainJitter_DEPRECATED)
+	, GrainIntensity_DEPRECATED(Settings.GrainIntensity_DEPRECATED)
+	, FilmGrainIntensity(Settings.FilmGrainIntensity)
+	, FilmGrainTexelSize(Settings.FilmGrainTexelSize)
+	, FilmGrainTexture(Settings.FilmGrainTexture)
+	, FilmGrainDecodeMultiply(Settings.FilmGrainDecodeMultiply)
+	, FilmGrainDecodeAdd(Settings.FilmGrainDecodeAdd)
 	, AmbientOcclusionIntensity(Settings.AmbientOcclusionIntensity)
 	, AmbientOcclusionStaticFraction(Settings.AmbientOcclusionStaticFraction)
 	, AmbientOcclusionRadius(Settings.AmbientOcclusionRadius)
diff --git a/Engine/Source/Runtime/Engine/Private/SceneView.cpp b/Engine/Source/Runtime/Engine/Private/SceneView.cpp
--- a/Engine/Source/Runtime/Engine/Private/SceneView.cpp
+++ b/Engine/Source/Runtime/Engine/Private/SceneView.cpp
@@ -1436,8 +1436,8 @@ void FSceneView::OverridePostProcessSettings(const FPostProcessSettings& Src, fl
 		LERP_PP(LensFlareBokehSize);
 		LERP_PP(LensFlareThreshold);
 		LERP_PP(VignetteIntensity);
-		LERP_PP(GrainIntensity);
-		LERP_PP(GrainJitter);
+		LERP_PP(FilmGrainIntensity);
+		LERP_PP(FilmGrainTexelSize);
 		LERP_PP(AmbientOcclusionIntensity);
 		LERP_PP(AmbientOcclusionStaticFraction);
 		LERP_PP(AmbientOcclusionRadius);
@@ -1660,6 +1660,14 @@ void FSceneView::OverridePostProcessSettings(const FPostProcessSettings& Src, fl
 			Dest.BloomConvolutionTexture = Src.BloomConvolutionTexture;
 		}
 
+		// actual texture cannot be blended but the film grain intensity can be blended
+		IF_PP(FilmGrainTexture)
+		{
+			Dest.FilmGrainTexture = Src.FilmGrainTexture;
+			Dest.FilmGrainDecodeMultiply = Src.FilmGrainDecodeMultiply;
+			Dest.FilmGrainDecodeAdd = Src.FilmGrainDecodeAdd;
+		}
+
 		// A continuous blending of this value would result trashing the pre-convolved bloom kernel cache.
 		IF_PP(BloomConvolutionBufferScale)
 		{
@@ -1887,12 +1895,7 @@ void FSceneView::EndFinalPostprocessSettings(const FSceneViewInitOptions& ViewIn
 
 		if(Quality < 4)
 		{
-			FinalPostProcessSettings.GrainIntensity = 0;
-		}
-
-		if(Quality < 5)
-		{
-			FinalPostProcessSettings.GrainJitter = 0;
+			FinalPostProcessSettings.FilmGrainIntensity = 0;
 		}
 	}
 
@@ -1919,8 +1922,7 @@ void FSceneView::EndFinalPostprocessSettings(const FSceneViewInitOptions& ViewIn
 
 	if(!Family->EngineShowFlags.Grain)
 	{
-		FinalPostProcessSettings.GrainIntensity = 0;
-		FinalPostProcessSettings.GrainJitter = 0;
+		FinalPostProcessSettings.FilmGrainIntensity = 0.0f;
 	}
 
 	if(!Family->EngineShowFlags.CameraImperfections)
diff --git a/Engine/Source/Runtime/Engine/Private/UnrealEngine.cpp b/Engine/Source/Runtime/Engine/Private/UnrealEngine.cpp
--- a/Engine/Source/Runtime/Engine/Private/UnrealEngine.cpp
+++ b/Engine/Source/Runtime/Engine/Private/UnrealEngine.cpp
@@ -2951,6 +2951,14 @@ void UEngine::LoadDefaultBloomTexture()
 	}
 }
 
+void UEngine::LoadDefaultFilmGrainTexture()
+{
+	if (DefaultFilmGrainTexture == nullptr && DefaultFilmGrainTextureName.IsValid())
+	{
+		LoadEngineTexture(DefaultFilmGrainTexture, *DefaultFilmGrainTextureName.ToString());
+	}
+}
+
 void UEngine::InitializePortalServices()
 {
 	IMessagingRpcModule* MessagingRpcModule = nullptr;
diff --git a/Engine/Source/Runtime/Renderer/Private/PostProcess/PostProcessTonemap.cpp b/Engine/Source/Runtime/Renderer/Private/PostProcess/PostProcessTonemap.cpp
--- a/Engine/Source/Runtime/Renderer/Private/PostProcess/PostProcessTonemap.cpp
+++ b/Engine/Source/Runtime/Renderer/Private/PostProcess/PostProcessTonemap.cpp
@@ -15,6 +15,8 @@
 #include "PostProcess/PostProcessing.h"
 #include "ClearQuad.h"
 #include "PipelineStateCache.h"
+#include "Rendering/Texture2DResource.h"
+#include "Math/Halton.h"
 
 namespace
 {
@@ -87,6 +89,11 @@ TAutoConsoleVariable<float> CVarGamma(
 	TEXT("Gamma on output"),
 	ECVF_RenderThreadSafe);
 
+TAutoConsoleVariable<int32> CVarFilmGrainSequenceLength(
+	TEXT("r.FilmGrain.SequenceLength"), 97,
+	TEXT("Length of the random sequence for film grain (preferably a prime number, default=97)."),
+	ECVF_RenderThreadSafe);
+
 const int32 GTonemapComputeTileSizeX = 8;
 const int32 GTonemapComputeTileSizeY = 8;
 
@@ -95,10 +102,9 @@ namespace TonemapperPermutation
 // Shared permutation dimensions between deferred and mobile renderer.
 class FTonemapperBloomDim          : SHADER_PERMUTATION_BOOL("USE_BLOOM");
 class FTonemapperGammaOnlyDim      : SHADER_PERMUTATION_BOOL("USE_GAMMA_ONLY");
-class FTonemapperGrainIntensityDim : SHADER_PERMUTATION_BOOL("USE_GRAIN_INTENSITY");
 class FTonemapperVignetteDim       : SHADER_PERMUTATION_BOOL("USE_VIGNETTE");
 class FTonemapperSharpenDim        : SHADER_PERMUTATION_BOOL("USE_SHARPEN");
-class FTonemapperGrainJitterDim    : SHADER_PERMUTATION_BOOL("USE_GRAIN_JITTER");
+class FTonemapperFilmGrainDim      : SHADER_PERMUTATION_BOOL("USE_FILM_GRAIN");
 class FTonemapperSwitchAxis        : SHADER_PERMUTATION_BOOL("NEEDTOSWITCHVERTICLEAXIS");
 class FTonemapperMsaaDim           : SHADER_PERMUTATION_BOOL("METAL_MSAA_HDR_DECODE");
 class FTonemapperEyeAdaptationDim  : SHADER_PERMUTATION_BOOL("EYEADAPTATION_EXPOSURE_FIX");
@@ -107,10 +113,9 @@ class FTonemapperUseFXAA           : SHADER_PERMUTATION_BOOL("USE_FXAA");
 using FCommonDomain = TShaderPermutationDomain<
 	FTonemapperBloomDim,
 	FTonemapperGammaOnlyDim,
-	FTonemapperGrainIntensityDim,
 	FTonemapperVignetteDim,
 	FTonemapperSharpenDim,
-	FTonemapperGrainJitterDim,
+	FTonemapperFilmGrainDim,
 	FTonemapperSwitchAxis,
 	FTonemapperMsaaDim,
 	FTonemapperUseFXAA>;
@@ -138,10 +143,9 @@ bool ShouldCompileCommonPermutation(const FGlobalShaderPermutationParameters& Pa
 	if (PermutationVector.Get<FTonemapperGammaOnlyDim>())
 	{
 		return !PermutationVector.Get<FTonemapperBloomDim>() &&
-			!PermutationVector.Get<FTonemapperGrainIntensityDim>() &&
 			!PermutationVector.Get<FTonemapperVignetteDim>() &&
 			!PermutationVector.Get<FTonemapperSharpenDim>() &&
-			!PermutationVector.Get<FTonemapperGrainJitterDim>() &&
+			!PermutationVector.Get<FTonemapperFilmGrainDim>() &&
 			!PermutationVector.Get<FTonemapperMsaaDim>();
 	}
 	return true;
@@ -164,10 +168,9 @@ FCommonDomain BuildCommonPermutationDomain(const FViewInfo& View, bool bGammaOnl
 	}
 
 	const FPostProcessSettings& Settings = View.FinalPostProcessSettings;
-	PermutationVector.Set<FTonemapperGrainIntensityDim>(Settings.GrainIntensity > 0.0f);
 	PermutationVector.Set<FTonemapperVignetteDim>(Settings.VignetteIntensity > 0.0f);
 	PermutationVector.Set<FTonemapperBloomDim>(Settings.BloomIntensity > 0.0);
-	PermutationVector.Set<FTonemapperGrainJitterDim>(Settings.GrainJitter > 0.0f);
+	PermutationVector.Set<FTonemapperFilmGrainDim>(View.FilmGrainTexture != nullptr);
 	PermutationVector.Set<FTonemapperSharpenDim>(CVarTonemapperSharpen.GetValueOnRenderThread() > 0.0f);	
 	PermutationVector.Set<FTonemapperSwitchAxis>(bSwitchVerticalAxis);
 	PermutationVector.Set<FTonemapperMsaaDim>(bMetalMSAAHDRDecode);
@@ -201,13 +204,11 @@ FDesktopDomain RemapPermutation(FDesktopDomain PermutationVector, ERHIFeatureLev
 
 	// Grain jitter or intensity looks bad anyway.
 	bool bFallbackToSlowest = false;
-	bFallbackToSlowest = bFallbackToSlowest || CommonPermutationVector.Get<FTonemapperGrainIntensityDim>();
-	bFallbackToSlowest = bFallbackToSlowest || CommonPermutationVector.Get<FTonemapperGrainJitterDim>();
+	bFallbackToSlowest = bFallbackToSlowest || CommonPermutationVector.Get<FTonemapperFilmGrainDim>();
 
 	if (bFallbackToSlowest)
 	{
-		CommonPermutationVector.Set<FTonemapperGrainIntensityDim>(true);
-		CommonPermutationVector.Set<FTonemapperGrainJitterDim>(true);
+		CommonPermutationVector.Set<FTonemapperFilmGrainDim>(true);
 		CommonPermutationVector.Set<FTonemapperSharpenDim>(true);
 
 		PermutationVector.Set<FTonemapperColorFringeDim>(true);
@@ -258,15 +259,6 @@ bool ShouldCompileDesktopPermutation(const FGlobalShaderPermutationParameters& P
 } // namespace TonemapperPermutation
 } // namespace
 
-void GrainPostSettings(FVector3f* RESTRICT const Constant, const FPostProcessSettings* RESTRICT const Settings)
-{
-	float GrainJitter = Settings->GrainJitter;
-	float GrainIntensity = Settings->GrainIntensity;
-	Constant->X = GrainIntensity;
-	Constant->Y = 1.0f + (-0.5f * GrainIntensity);
-	Constant->Z = GrainJitter;
-}
-
 FTonemapperOutputDeviceParameters GetTonemapperOutputDeviceParameters(const FSceneViewFamily& Family)
 {
 	static TConsoleVariableData<int32>* CVarOutputGamut = IConsoleManager::Get().FindTConsoleVariableDataInt(TEXT("r.HDR.Display.ColorGamut"));
@@ -320,30 +312,12 @@ FTonemapperOutputDeviceParameters GetTonemapperOutputDeviceParameters(const FSce
 
 BEGIN_SHADER_PARAMETER_STRUCT(FFilmGrainParameters, )
 	SHADER_PARAMETER(FVector3f, GrainRandomFull)
-	SHADER_PARAMETER(FVector3f, GrainScaleBiasJitter)
+	SHADER_PARAMETER(float, FilmGrainMul)
+	SHADER_PARAMETER(float, FilmGrainAdd)
+	SHADER_PARAMETER_TEXTURE(Texture2D, FilmGrainTexture)
+	SHADER_PARAMETER(FScreenTransform, ScreenPosToFilmGrainTextureUV)
 END_SHADER_PARAMETER_STRUCT()
 
-FFilmGrainParameters GetFilmGrainParameters(const FViewInfo& View)
-{
-	FVector3f GrainRandomFullValue;
-	{
-		uint8 FrameIndexMod8 = 0;
-		if (View.State)
-		{
-			FrameIndexMod8 = View.ViewState->GetFrameIndex(8);
-		}
-		GrainRandomFromFrame(&GrainRandomFullValue, FrameIndexMod8);
-	}
-
-	FVector3f GrainScaleBiasJitter;
-	GrainPostSettings(&GrainScaleBiasJitter, &View.FinalPostProcessSettings);
-
-	FFilmGrainParameters Parameters;
-	Parameters.GrainRandomFull = GrainRandomFullValue;
-	Parameters.GrainScaleBiasJitter = GrainScaleBiasJitter;
-	return Parameters;
-}
-
 BEGIN_SHADER_PARAMETER_STRUCT(FTonemapParameters, )
 	SHADER_PARAMETER_STRUCT_REF(FViewUniformShaderParameters, View)
 	SHADER_PARAMETER_STRUCT_INCLUDE(FFilmGrainParameters, FilmGrain)
@@ -575,7 +549,47 @@ FScreenPassTexture AddTonemapPass(FRDGBuilder& GraphBuilder, const FViewInfo& Vi
 
 	FTonemapParameters CommonParameters;
 	CommonParameters.View = View.ViewUniformBuffer;
-	CommonParameters.FilmGrain = GetFilmGrainParameters(View);
+
+	{
+		uint8 FrameIndexMod8 = 0;
+		if (View.State)
+		{
+			FrameIndexMod8 = View.ViewState->GetFrameIndex(8);
+		}
+		GrainRandomFromFrame(&CommonParameters.FilmGrain.GrainRandomFull, FrameIndexMod8);
+	}
+
+	if (View.FilmGrainTexture)
+	{
+		check(View.FilmGrainTexture);
+
+		// (FilmGrainTexture * FilmGrainDecodeMultiply + FilmGrainDecodeAdd) * View.FinalPostProcessSettings.FilmGrainIntensity + 1.0f - View.FinalPostProcessSettings.FilmGrainIntensity
+		CommonParameters.FilmGrain.FilmGrainMul = View.FinalPostProcessSettings.FilmGrainDecodeMultiply * View.FinalPostProcessSettings.FilmGrainIntensity;
+		CommonParameters.FilmGrain.FilmGrainAdd = View.FinalPostProcessSettings.FilmGrainDecodeAdd * View.FinalPostProcessSettings.FilmGrainIntensity + 1.0f - View.FinalPostProcessSettings.FilmGrainIntensity;
+		CommonParameters.FilmGrain.FilmGrainTexture = View.FilmGrainTexture->GetTexture2DRHI();
+
+		int32 RandomSequenceLength = CVarFilmGrainSequenceLength.GetValueOnRenderThread();
+		int32 RandomSequenceIndex = (View.ViewState ? View.ViewState->FrameIndex : 0) % RandomSequenceLength;
+
+		FVector2D RandomGrainTextureUVOffset;
+		RandomGrainTextureUVOffset.X = Halton(RandomSequenceIndex + 1, 2);
+		RandomGrainTextureUVOffset.Y = Halton(RandomSequenceIndex + 1, 3);
+
+		FVector2D TextureSize(View.FilmGrainTexture->GetSizeX(), View.FilmGrainTexture->GetSizeY());
+		FVector2D OutputSizeF(OutputViewport.Rect.Width(), OutputViewport.Rect.Height());
+
+		CommonParameters.FilmGrain.ScreenPosToFilmGrainTextureUV = 
+			FScreenTransform::ScreenPosToViewportUV *
+			((OutputSizeF / TextureSize) * (1.0f / View.FinalPostProcessSettings.FilmGrainTexelSize)) + RandomGrainTextureUVOffset;
+	}
+	else
+	{
+		CommonParameters.FilmGrain.FilmGrainMul = 0.0f;
+		CommonParameters.FilmGrain.FilmGrainAdd = 1.0f;
+		CommonParameters.FilmGrain.FilmGrainTexture = GWhiteTexture->TextureRHI;
+		CommonParameters.FilmGrain.ScreenPosToFilmGrainTextureUV = FScreenTransform::ScreenPosToViewportUV;
+	}
+
 	CommonParameters.OutputDevice = GetTonemapperOutputDeviceParameters(ViewFamily);
 	CommonParameters.Color = GetScreenPassTextureViewportParameters(SceneColorViewport);
 	if (Inputs.Bloom.Texture)
diff --git a/Engine/Source/Runtime/Renderer/Private/SceneRendering.cpp b/Engine/Source/Runtime/Renderer/Private/SceneRendering.cpp
--- a/Engine/Source/Runtime/Renderer/Private/SceneRendering.cpp
+++ b/Engine/Source/Runtime/Renderer/Private/SceneRendering.cpp
@@ -2307,6 +2307,65 @@ FSceneRenderer::FSceneRenderer(const FSceneViewFamily* InViewFamily,FHitProxyCon
 			ViewInfo->bCameraCut = true;
 		}
 		#endif
+
+		// Handle the film grain texture
+		if (ViewInfo->FinalPostProcessSettings.FilmGrainIntensity > 0.0f)
+		{
+			GEngine->LoadDefaultFilmGrainTexture();
+
+			UTexture2D* FilmGrainTexture = GEngine->DefaultFilmGrainTexture;
+			if (ViewInfo->FinalPostProcessSettings.FilmGrainTexture)
+			{
+				FilmGrainTexture = ViewInfo->FinalPostProcessSettings.FilmGrainTexture;
+			}
+
+			bool bIsValid = FilmGrainTexture != nullptr;
+
+			if (bIsValid)
+			{
+				const int32 CinematicTextureGroups = 0;
+				const float Seconds = 5.0f;
+				FilmGrainTexture->SetForceMipLevelsToBeResident(Seconds, CinematicTextureGroups);
+			}
+
+			const uint32 FramesPerWarning = 15;
+
+			if (bIsValid && FilmGrainTexture->IsFullyStreamedIn() == false)
+			{
+				if ((ViewInfo->Family->FrameNumber % FramesPerWarning) == 0)
+				{
+					UE_LOG(LogRenderer, Warning, TEXT("The film grain texture is not streamed in."));
+				}
+
+				bIsValid = false;
+			}
+
+			if (bIsValid && FilmGrainTexture->bHasStreamingUpdatePending == true)
+			{
+				if ((ViewInfo->Family->FrameNumber % FramesPerWarning) == 0)
+				{
+					UE_LOG(LogRenderer, Warning, TEXT("The film grain texture has pending update."));
+				}
+
+				bIsValid = false;
+			}
+
+			if (bIsValid)
+			{
+				const FTextureResource* TextureResource = FilmGrainTexture->GetResource();
+				if (TextureResource)
+				{
+					ViewInfo->FilmGrainTexture = TextureResource->GetTexture2DResource();
+				}
+
+				if (FilmGrainTexture == GEngine->DefaultFilmGrainTexture)
+				{
+					// Decode Texture2D'/Engine/EngineMaterials/Marcie_Grain_1024_M5_555_A0.Marcie_Grain_1024_M5_555_A0'
+					ViewInfo->FinalPostProcessSettings.FilmGrainDecodeMultiply = (1.0f / 0.18f);
+					ViewInfo->FinalPostProcessSettings.FilmGrainDecodeAdd = 0.0f;
+				}
+			}
+		}
 	}
 
 	// Catches inconsistency one engine show flags for screen percentage and whether it is supported or not.
diff --git a/Engine/Source/Runtime/Renderer/Private/SceneRendering.h b/Engine/Source/Runtime/Renderer/Private/SceneRendering.h
--- a/Engine/Source/Runtime/Renderer/Private/SceneRendering.h
+++ b/Engine/Source/Runtime/Renderer/Private/SceneRendering.h
@@ -1413,6 +1413,8 @@ public:
 	/** Custom visibility query for view */
 	ICustomVisibilityQuery* CustomVisibilityQuery;
 
+	const FTexture2DResource* FilmGrainTexture = nullptr;
+
 	TArray<FPrimitiveSceneInfo*, SceneRenderingAllocator> IndirectShadowPrimitives;
 
 	/** Only one of the resources(TextureBuffer or Texture2D) will be used depending on the Mobile.UseGPUSceneTexture cvar */
