Fixes the desktop tonemaper for Android

diff --git a/Engine/Shaders/Private/PostProcessTonemap.usf b/Engine/Shaders/Private/PostProcessTonemap.usf
--- a/Engine/Shaders/Private/PostProcessTonemap.usf
+++ b/Engine/Shaders/Private/PostProcessTonemap.usf
@@ -61,6 +61,7 @@ float3 GrainRandomFull;
 float FilmGrainMul;
 float FilmGrainAdd;
 Texture2D<half3> FilmGrainTexture;
+SamplerState FilmGrainSampler;
 float4 ScreenPosToFilmGrainTextureUV;
 
 float EditorNITLevel;
@@ -565,7 +566,7 @@ float4 TonemapCommonPS(
 		{
 			float2 FilmGrainUV = ApplyScreenTransform(ScreenPos, ScreenPosToFilmGrainTextureUV);
 
-			half3 RawGrain = FilmGrainTexture.SampleLevel(GlobalBilinearWrappedSampler, FilmGrainUV, 0);
+			half3 RawGrain = FilmGrainTexture.SampleLevel(FilmGrainSampler, FilmGrainUV, 0);
 			half3 FinalGrain = RawGrain * half(FilmGrainMul) + half(FilmGrainAdd);
 
 			LinearColor.rgb *= FinalGrain;
diff --git a/Engine/Source/Runtime/Renderer/Private/PostProcess/PostProcessTonemap.cpp b/Engine/Source/Runtime/Renderer/Private/PostProcess/PostProcessTonemap.cpp
--- a/Engine/Source/Runtime/Renderer/Private/PostProcess/PostProcessTonemap.cpp
+++ b/Engine/Source/Runtime/Renderer/Private/PostProcess/PostProcessTonemap.cpp
@@ -315,6 +315,7 @@ BEGIN_SHADER_PARAMETER_STRUCT(FFilmGrainParameters, )
 	SHADER_PARAMETER(float, FilmGrainMul)
 	SHADER_PARAMETER(float, FilmGrainAdd)
 	SHADER_PARAMETER_TEXTURE(Texture2D, FilmGrainTexture)
+	SHADER_PARAMETER_SAMPLER(SamplerState, FilmGrainSampler)
 	SHADER_PARAMETER(FScreenTransform, ScreenPosToFilmGrainTextureUV)
 END_SHADER_PARAMETER_STRUCT()
 
@@ -567,6 +568,7 @@ FScreenPassTexture AddTonemapPass(FRDGBuilder& GraphBuilder, const FViewInfo& Vi
 		CommonParameters.FilmGrain.FilmGrainMul = View.FinalPostProcessSettings.FilmGrainDecodeMultiply * View.FinalPostProcessSettings.FilmGrainIntensity;
 		CommonParameters.FilmGrain.FilmGrainAdd = View.FinalPostProcessSettings.FilmGrainDecodeAdd * View.FinalPostProcessSettings.FilmGrainIntensity + 1.0f - View.FinalPostProcessSettings.FilmGrainIntensity;
 		CommonParameters.FilmGrain.FilmGrainTexture = View.FilmGrainTexture->GetTexture2DRHI();
+		CommonParameters.FilmGrain.FilmGrainSampler = TStaticSamplerState<SF_Bilinear, AM_Wrap, AM_Wrap>::GetRHI();
 
 		int32 RandomSequenceLength = CVarFilmGrainSequenceLength.GetValueOnRenderThread();
 		int32 RandomSequenceIndex = (View.ViewState ? View.ViewState->FrameIndex : 0) % RandomSequenceLength;
@@ -587,6 +589,7 @@ FScreenPassTexture AddTonemapPass(FRDGBuilder& GraphBuilder, const FViewInfo& Vi
 		CommonParameters.FilmGrain.FilmGrainMul = 0.0f;
 		CommonParameters.FilmGrain.FilmGrainAdd = 1.0f;
 		CommonParameters.FilmGrain.FilmGrainTexture = GWhiteTexture->TextureRHI;
+		CommonParameters.FilmGrain.FilmGrainSampler = TStaticSamplerState<SF_Bilinear, AM_Wrap, AM_Wrap>::GetRHI();
 		CommonParameters.FilmGrain.ScreenPosToFilmGrainTextureUV = FScreenTransform::ScreenPosToViewportUV;
 	}
 
