Remove experimental hair scatter pass.

diff --git a/Engine/Shaders/Private/HairStrands/HairScatter.usf b/Engine/Shaders/Private/HairStrands/HairScatter.usf
deleted file mode 100644
diff --git a/Engine/Shaders/Private/HairStrands/HairScatterCompose.usf b/Engine/Shaders/Private/HairStrands/HairScatterCompose.usf
deleted file mode 100644
diff --git a/Engine/Source/Runtime/Renderer/Private/HairStrands/HairStrandsComposition.cpp b/Engine/Source/Runtime/Renderer/Private/HairStrands/HairStrandsComposition.cpp
--- a/Engine/Source/Runtime/Renderer/Private/HairStrands/HairStrandsComposition.cpp
+++ b/Engine/Source/Runtime/Renderer/Private/HairStrands/HairStrandsComposition.cpp
@@ -9,7 +9,6 @@
 #include "RenderGraphUtils.h"
 #include "PostProcessing.h"
 #include "HairStrandsRendering.h"
-#include "HairStrandsScatter.h"
 #include "HairStrandsTile.h"
 #include "FogRendering.h"
 
diff --git a/Engine/Source/Runtime/Renderer/Private/HairStrands/HairStrandsScatter.cpp b/Engine/Source/Runtime/Renderer/Private/HairStrands/HairStrandsScatter.cpp
deleted file mode 100644
diff --git a/Engine/Source/Runtime/Renderer/Private/HairStrands/HairStrandsScatter.h b/Engine/Source/Runtime/Renderer/Private/HairStrands/HairStrandsScatter.h
deleted file mode 100644
