Remove experimental hair gbuffer dither pass.

diff --git a/Engine/Plugins/Runtime/HairStrands/Source/HairStrandsCore/Private/GroomComponent.cpp b/Engine/Plugins/Runtime/HairStrands/Source/HairStrandsCore/Private/GroomComponent.cpp
--- a/Engine/Plugins/Runtime/HairStrands/Source/HairStrandsCore/Private/GroomComponent.cpp
+++ b/Engine/Plugins/Runtime/HairStrands/Source/HairStrandsCore/Private/GroomComponent.cpp
@@ -816,7 +816,7 @@ public:
 		FMeshBatch& Mesh = Collector.AllocateMesh();
 
 		const bool bUseCardsOrMeshes = GeometryType == EHairGeometryType::Cards || GeometryType == EHairGeometryType::Meshes;
-		Mesh.CastShadow = !IsHairStrandsComplexLightingEnabled() || bUseCardsOrMeshes;
+		Mesh.CastShadow = bUseCardsOrMeshes;
 		Mesh.bUseForMaterial  = true;
 		Mesh.bUseForDepthPass = bUseCardsOrMeshes;
 		Mesh.SegmentIndex = 0;
diff --git a/Engine/Shaders/Private/HairStrands/HairStrandsMaterialGBufferPS.usf b/Engine/Shaders/Private/HairStrands/HairStrandsMaterialGBufferPS.usf
deleted file mode 100644
diff --git a/Engine/Shaders/Private/HairStrands/HairStrandsMaterialGBufferVS.usf b/Engine/Shaders/Private/HairStrands/HairStrandsMaterialGBufferVS.usf
deleted file mode 100644
diff --git a/Engine/Source/Runtime/Renderer/Private/HairStrands/HairStrandsVisibility.cpp b/Engine/Source/Runtime/Renderer/Private/HairStrands/HairStrandsVisibility.cpp
--- a/Engine/Source/Runtime/Renderer/Private/HairStrands/HairStrandsVisibility.cpp
+++ b/Engine/Source/Runtime/Renderer/Private/HairStrands/HairStrandsVisibility.cpp
@@ -80,9 +80,6 @@ static FAutoConsoleVariableRef CVarHairStrandsSortHairSampleByDepth(TEXT("r.Hair
 static int32 GHairStrandsHairCountToTransmittance = 0;
 static FAutoConsoleVariableRef CVarHairStrandsHairCountToTransmittance(TEXT("r.HairStrands.Visibility.UseCoverageMappping"), GHairStrandsHairCountToTransmittance, TEXT("Use hair count to coverage transfer function."));
 
-static int32 GHairStrandsVisibility_UseFastPath = 0;
-static FAutoConsoleVariableRef CVarHairStrandsVisibility_UseFastPath(TEXT("r.HairStrands.Visibility.UseFastPath"), GHairStrandsVisibility_UseFastPath, TEXT("Use fast path writing hair data into Gbuffer."));
-
 static int32 GHairStrandsVisibility_OutputEmissiveData = 0;
 static FAutoConsoleVariableRef CVarHairStrandsVisibility_OutputEmissiveData(TEXT("r.HairStrands.Visibility.Emissive"), GHairStrandsVisibility_OutputEmissiveData, TEXT("Enable emissive data during the material pass."));
 
@@ -289,11 +286,6 @@ static bool IsCompatibleWithHairVisibility(const FMeshMaterialShaderPermutationP
 	return IsCompatibleWithHairStrands(Parameters.Platform, Parameters.MaterialParameters);
 }
 
-bool IsHairStrandsComplexLightingEnabled()
-{
-	return GHairStrandsVisibility_UseFastPath == 0 || GetMeanSamplePerPixel() > 1 || GetHairVisibilityRenderMode() == HairVisibilityRenderMode_PPLL;
-}
-
 float GetHairWriteVelocityCoverageThreshold()
 {
 	return FMath::Clamp(GHairStrandsWriteVelocityCoverageThreshold, 0.f, 1.f);
@@ -394,409 +386,6 @@ static FRDGTextureRef AddClearLightSamplePass(
 	return Output;
 }
 
-/////////////////////////////////////////////////////////////////////////////////////////
-// Material pass which write value directly into the GBuffer. This is a fast (low-quality) 
-// pass, which dither raster result into GBuffer, and works only with 1spp
-BEGIN_GLOBAL_SHADER_PARAMETER_STRUCT(FMaterialGBufferPassUniformParameters, )
-	SHADER_PARAMETER(FIntPoint, MaxResolution)
-	SHADER_PARAMETER(uint32, InputType)
-	SHADER_PARAMETER(float, CoverageThreshold)
-	SHADER_PARAMETER_RDG_TEXTURE(Texture2D<float>, InTransmittanceTexture)
-	SHADER_PARAMETER_RDG_TEXTURE(Texture2D<uint>,  InRasterOutputVisibilityTexture)
-	SHADER_PARAMETER_RDG_TEXTURE(Texture2D<uint>,  InMSAAIDTexture)
-	SHADER_PARAMETER_RDG_TEXTURE(Texture2D<float>, InMSAADepthTexture)
-END_GLOBAL_SHADER_PARAMETER_STRUCT()
-IMPLEMENT_STATIC_UNIFORM_BUFFER_STRUCT(FMaterialGBufferPassUniformParameters, "MaterialGBufferPassParameters", SceneTextures);
-
-class FHairMaterialGBufferVS : public FMeshMaterialShader
-{
-	DECLARE_SHADER_TYPE(FHairMaterialGBufferVS, MeshMaterial);
-
-protected:
-	FHairMaterialGBufferVS(const ShaderMetaType::CompiledShaderInitializerType& Initializer)
-		: FMeshMaterialShader(Initializer)
-	{
-		ERHIFeatureLevel::Type FeatureLevel = GetMaxSupportedFeatureLevel((EShaderPlatform)Initializer.Target.Platform);
-		check(FSceneInterface::GetShadingPath(FeatureLevel) != EShadingPath::Mobile);
-	}
-
-	FHairMaterialGBufferVS() {}
-
-	static bool ShouldCompilePermutation(const FMeshMaterialShaderPermutationParameters& Parameters)
-	{
-		const bool bIsCompatible = IsCompatibleWithHairVisibility(Parameters) && Parameters.VertexFactoryType->GetFName() == FName(TEXT("FHairStrandsVertexFactory"));
-		return bIsCompatible;
-	}
-
-	static void ModifyCompilationEnvironment(const FMaterialShaderPermutationParameters& Parameters, FShaderCompilerEnvironment& OutEnvironment)
-	{
-		FMaterialShader::ModifyCompilationEnvironment(Parameters, OutEnvironment);
-	}
-};
-IMPLEMENT_MATERIAL_SHADER_TYPE(, FHairMaterialGBufferVS, TEXT("/Engine/Private/HairStrands/HairStrandsMaterialGBufferVS.usf"), TEXT("Main"), SF_Vertex);
-
-/////////////////////////////////////////////////////////////////////////////////////////
-
-/////////////////////////////////////////////////////////////////////////////////////////
-
-class FHairMaterialGBufferShaderElementData : public FMeshMaterialShaderElementData
-{
-public:
-	FHairMaterialGBufferShaderElementData(int32 MacroGroupId, int32 MaterialId, int32 PrimitiveId, uint32 LightChannelMask) : MaterialPass_MacroGroupId(MacroGroupId), MaterialPass_MaterialId(MaterialId), MaterialPass_PrimitiveId(PrimitiveId), MaterialPass_LightChannelMask(LightChannelMask) { }
-	uint32 MaterialPass_MacroGroupId;
-	uint32 MaterialPass_MaterialId;
-	uint32 MaterialPass_PrimitiveId;
-	uint32 MaterialPass_LightChannelMask;
-};
-
-class FHairMaterialGBufferPS : public FMeshMaterialShader
-{
-	DECLARE_SHADER_TYPE(FHairMaterialGBufferPS, MeshMaterial);
-
-public:
-	FHairMaterialGBufferPS(const ShaderMetaType::CompiledShaderInitializerType& Initializer)
-		: FMeshMaterialShader(Initializer)
-	{
-		ERHIFeatureLevel::Type FeatureLevel = GetMaxSupportedFeatureLevel((EShaderPlatform)Initializer.Target.Platform);
-		check(FSceneInterface::GetShadingPath(FeatureLevel) != EShadingPath::Mobile);
-		MaterialPass_MacroGroupId.Bind(Initializer.ParameterMap, TEXT("MaterialPass_MacroGroupId"));
-		MaterialPass_MaterialId.Bind(Initializer.ParameterMap, TEXT("MaterialPass_MaterialId"));
-		MaterialPass_PrimitiveId.Bind(Initializer.ParameterMap, TEXT("MaterialPass_PrimitiveId"));
-		MaterialPass_LightChannelMask.Bind(Initializer.ParameterMap, TEXT("MaterialPass_LightChannelMask"));
-	}
-
-	FHairMaterialGBufferPS() {}
-
-	static bool ShouldCompilePermutation(const FMeshMaterialShaderPermutationParameters& Parameters)
-	{
-		const bool bIsCompatible = IsCompatibleWithHairStrands(Parameters.Platform, Parameters.MaterialParameters) && Parameters.VertexFactoryType->GetFName() == FName(TEXT("FHairStrandsVertexFactory"));
-		return bIsCompatible;
-	}
-
-	static void ModifyCompilationEnvironment(const FMaterialShaderPermutationParameters& Parameters, FShaderCompilerEnvironment& OutEnvironment)
-	{
-		FMaterialShader::ModifyCompilationEnvironment(Parameters, OutEnvironment);
-		//OutEnvironment.SetDefine(TEXT("HAIR_MATERIAL_DEBUG_OUTPUT"), bPlatformRequireRenderTarget ? 1 : 0);
-	}
-
-	void GetShaderBindings(
-		const FScene* Scene,
-		ERHIFeatureLevel::Type FeatureLevel,
-		const FPrimitiveSceneProxy* PrimitiveSceneProxy,
-		const FMaterialRenderProxy& MaterialRenderProxy,
-		const FMaterial& Material,
-		const FMeshPassProcessorRenderState& DrawRenderState,
-		const FHairMaterialGBufferShaderElementData& ShaderElementData,
-		FMeshDrawSingleShaderBindings& ShaderBindings) const
-	{
-		FMeshMaterialShader::GetShaderBindings(Scene, FeatureLevel, PrimitiveSceneProxy, MaterialRenderProxy, Material, DrawRenderState, ShaderElementData, ShaderBindings);
-		ShaderBindings.Add(MaterialPass_MacroGroupId, ShaderElementData.MaterialPass_MacroGroupId);
-		ShaderBindings.Add(MaterialPass_MaterialId, ShaderElementData.MaterialPass_MaterialId);
-		ShaderBindings.Add(MaterialPass_PrimitiveId, ShaderElementData.MaterialPass_PrimitiveId);
-		ShaderBindings.Add(MaterialPass_LightChannelMask, ShaderElementData.MaterialPass_LightChannelMask);
-	}
-
-private:
-	LAYOUT_FIELD(FShaderParameter, MaterialPass_MacroGroupId);
-	LAYOUT_FIELD(FShaderParameter, MaterialPass_MaterialId);
-	LAYOUT_FIELD(FShaderParameter, MaterialPass_PrimitiveId);
-	LAYOUT_FIELD(FShaderParameter, MaterialPass_LightChannelMask);
-};
-IMPLEMENT_MATERIAL_SHADER_TYPE(, FHairMaterialGBufferPS, TEXT("/Engine/Private/HairStrands/HairStrandsMaterialGBufferPS.usf"), TEXT("Main"), SF_Pixel);
-/////////////////////////////////////////////////////////////////////////////////////////
-
-BEGIN_SHADER_PARAMETER_STRUCT(FVisibilityMaterialGBufferPassParameters, )
-	SHADER_PARAMETER_STRUCT_REF(FViewUniformShaderParameters, View)
-	SHADER_PARAMETER_STRUCT_INCLUDE(FInstanceCullingDrawParams, InstanceCullingDrawParams)
-	SHADER_PARAMETER_RDG_UNIFORM_BUFFER(FMaterialGBufferPassUniformParameters, UniformBuffer)
-	RENDER_TARGET_BINDING_SLOTS()
-END_SHADER_PARAMETER_STRUCT()
-
-/////////////////////////////////////////////////////////////////////////////////////////
-
-class FHairMaterialGBufferProcessor : public FMeshPassProcessor
-{
-public:
-	FHairMaterialGBufferProcessor(
-		const FScene* Scene,
-		const FSceneView* InViewIfDynamicMeshCommand,
-		const FMeshPassProcessorRenderState& InPassDrawRenderState,
-		FDynamicPassMeshDrawListContext* InDrawListContext);
-
-	virtual void AddMeshBatch(const FMeshBatch& RESTRICT MeshBatch, uint64 BatchElementMask, const FPrimitiveSceneProxy* RESTRICT PrimitiveSceneProxy, int32 StaticMeshId = -1) override final;
-	void AddMeshBatch(const FMeshBatch& RESTRICT MeshBatch, uint64 BatchElementMask, const FPrimitiveSceneProxy* RESTRICT PrimitiveSceneProxy, int32 StaticMeshId, int32 MacroGroupId, int32 HairMaterialId);
-
-private:
-	bool TryAddMeshBatch(
-		const FMeshBatch& RESTRICT MeshBatch,
-		uint64 BatchElementMask,
-		const FPrimitiveSceneProxy* RESTRICT PrimitiveSceneProxy,
-		int32 StaticMeshId,
-		int32 MacroGroupId,
-		int32 HairMaterialId,
-		const FMaterialRenderProxy& MaterialRenderProxy,
-		const FMaterial& Material);
-
-	bool Process(
-		const FMeshBatch& MeshBatch,
-		uint64 BatchElementMask,
-		const FPrimitiveSceneProxy* RESTRICT PrimitiveSceneProxy,
-		int32 StaticMeshId,
-		const FMaterialRenderProxy& RESTRICT MaterialRenderProxy,
-		const FMaterial& RESTRICT MaterialResource,
-		const int32 MacroGroupId,
-		const int32 HairMaterialId,
-		const int32 HairPrimitiveId,
-		const uint32 HairPrimitiveLightChannelMask);
-
-	FMeshPassProcessorRenderState PassDrawRenderState;
-};
-
-void FHairMaterialGBufferProcessor::AddMeshBatch(const FMeshBatch& RESTRICT MeshBatch, uint64 BatchElementMask, const FPrimitiveSceneProxy* RESTRICT PrimitiveSceneProxy, int32 StaticMeshId)
-{
-	AddMeshBatch(MeshBatch, BatchElementMask, PrimitiveSceneProxy, StaticMeshId, 0, 0);
-}
-
-void FHairMaterialGBufferProcessor::AddMeshBatch(const FMeshBatch& RESTRICT MeshBatch, uint64 BatchElementMask, const FPrimitiveSceneProxy* RESTRICT PrimitiveSceneProxy, int32 StaticMeshId, int32 MacroGroupId, int32 HairMaterialId)
-{
-	const FMaterialRenderProxy* MaterialRenderProxy = MeshBatch.MaterialRenderProxy;
-	while (MaterialRenderProxy)
-	{
-		const FMaterial* Material = MaterialRenderProxy->GetMaterialNoFallback(FeatureLevel);
-		if (Material)
-		{
-			if (TryAddMeshBatch(MeshBatch, BatchElementMask, PrimitiveSceneProxy, StaticMeshId, MacroGroupId, HairMaterialId, *MaterialRenderProxy, *Material))
-			{
-				break;
-			}
-		}
-
-		MaterialRenderProxy = MaterialRenderProxy->GetFallback(FeatureLevel);
-	}
-}
-
-bool FHairMaterialGBufferProcessor::TryAddMeshBatch(
-	const FMeshBatch& RESTRICT MeshBatch,
-	uint64 BatchElementMask,
-	const FPrimitiveSceneProxy* RESTRICT PrimitiveSceneProxy,
-	int32 StaticMeshId,
-	int32 MacroGroupId,
-	int32 HairMaterialId,
-	const FMaterialRenderProxy& MaterialRenderProxy,
-	const FMaterial& Material)
-{
-	static const FVertexFactoryType* CompatibleVF = FVertexFactoryType::GetVFByName(TEXT("FHairStrandsVertexFactory"));
-
-	// Determine the mesh's material and blend mode.
-	const bool bIsCompatible = IsCompatibleWithHairStrands(&Material, FeatureLevel);
-	const bool bIsHairStrandsFactory = MeshBatch.VertexFactory->GetType()->GetHashedName() == CompatibleVF->GetHashedName();
-	const bool bShouldRender = (!PrimitiveSceneProxy && MeshBatch.Elements.Num() > 0) || (PrimitiveSceneProxy && PrimitiveSceneProxy->ShouldRenderInMainPass());
-
-	if (bIsCompatible
-		&& bIsHairStrandsFactory
-		&& bShouldRender
-		&& ShouldIncludeDomainInMeshPass(Material.GetMaterialDomain()))
-	{
-		// For the mesh patch to be rendered a single triangle triangle to spawn the necessary amount of thread
-		FMeshBatch MeshBatchCopy = MeshBatch;
-		for (uint32 ElementIt = 0, ElementCount = uint32(MeshBatch.Elements.Num()); ElementIt < ElementCount; ++ElementIt)
-		{
-			MeshBatchCopy.Elements[ElementIt].FirstIndex = 0;
-			MeshBatchCopy.Elements[ElementIt].NumPrimitives = 1;
-			MeshBatchCopy.Elements[ElementIt].NumInstances = 1;
-			MeshBatchCopy.Elements[ElementIt].IndirectArgsBuffer = nullptr;
-			MeshBatchCopy.Elements[ElementIt].IndirectArgsOffset = 0;
-		}
-
-		FPrimitiveSceneInfo* SceneInfo = PrimitiveSceneProxy ? PrimitiveSceneProxy->GetPrimitiveSceneInfo() : nullptr;
-		uint32 LightChannelMask = PrimitiveSceneProxy ? PrimitiveSceneProxy->GetLightingChannelMask() : 0;
-		FMeshDrawCommandPrimitiveIdInfo IdInfo = GetDrawCommandPrimitiveId(SceneInfo, MeshBatch.Elements[0]);
-
-		return Process(MeshBatchCopy, BatchElementMask, PrimitiveSceneProxy, StaticMeshId, MaterialRenderProxy, Material, MacroGroupId, HairMaterialId, IdInfo.DrawPrimitiveId, LightChannelMask);
-	}
-
-	return true;
-}
-
-bool FHairMaterialGBufferProcessor::Process(
-	const FMeshBatch& MeshBatch,
-	uint64 BatchElementMask,
-	const FPrimitiveSceneProxy* RESTRICT PrimitiveSceneProxy,
-	int32 StaticMeshId,
-	const FMaterialRenderProxy& RESTRICT MaterialRenderProxy,
-	const FMaterial& RESTRICT MaterialResource,
-	const int32 MacroGroupId,
-	const int32 HairMaterialId,
-	const int32 HairPrimitiveId,
-	const uint32 HairPrimitiveLightChannelMask)
-{
-	const FVertexFactory* VertexFactory = MeshBatch.VertexFactory;
-
-	TMeshProcessorShaders<
-		FHairMaterialGBufferVS,
-		FHairMaterialGBufferPS> PassShaders;
-	{
-		FMaterialShaderTypes ShaderTypes;
-		ShaderTypes.AddShaderType<FHairMaterialGBufferVS>();
-		ShaderTypes.AddShaderType<FHairMaterialGBufferPS>();
-
-		FVertexFactoryType* VertexFactoryType = VertexFactory->GetType();
-
-		FMaterialShaders Shaders;
-		if (!MaterialResource.TryGetShaders(ShaderTypes, VertexFactoryType, Shaders))
-		{
-			return false;
-		}
-
-		Shaders.TryGetVertexShader(PassShaders.VertexShader);
-		Shaders.TryGetPixelShader(PassShaders.PixelShader);
-	}
-
-	FMeshPassProcessorRenderState DrawRenderState(PassDrawRenderState);
-	FHairMaterialGBufferShaderElementData ShaderElementData(MacroGroupId, HairMaterialId, HairPrimitiveId, HairPrimitiveLightChannelMask);
-	ShaderElementData.InitializeMeshMaterialData(ViewIfDynamicMeshCommand, PrimitiveSceneProxy, MeshBatch, StaticMeshId, false);
-
-	BuildMeshDrawCommands(
-		MeshBatch,
-		BatchElementMask,
-		PrimitiveSceneProxy,
-		MaterialRenderProxy,
-		MaterialResource,
-		DrawRenderState,
-		PassShaders,
-		ERasterizerFillMode::FM_Solid,
-		ERasterizerCullMode::CM_CCW,
-		FMeshDrawCommandSortKey::Default,
-		EMeshPassFeatures::Default,
-		ShaderElementData);
-
-	return true;
-}
-
-FHairMaterialGBufferProcessor::FHairMaterialGBufferProcessor(
-	const FScene* Scene,
-	const FSceneView* InViewIfDynamicMeshCommand,
-	const FMeshPassProcessorRenderState& InPassDrawRenderState,
-	FDynamicPassMeshDrawListContext* InDrawListContext)
-	: FMeshPassProcessor(Scene, Scene->GetFeatureLevel(), InViewIfDynamicMeshCommand, InDrawListContext)
-	, PassDrawRenderState(InPassDrawRenderState)
-{
-}
-
-///////////////////////////////////////////////////////////////////////////////////////////////////
-
-static void AddHairMaterialGBufferPass(
-	FRDGBuilder& GraphBuilder,
-	const FScene* Scene,
-	const FViewInfo* ViewInfo,
-	FInstanceCullingManager& InstanceCullingManager,
-	const FHairStrandsMacroGroupDatas& MacroGroupDatas,
-
-	FRDGTextureRef InTransmittanceTexture,
-	FRDGTextureRef InRasterOutputVisibilityTexture,
-	FRDGTextureRef InMSAAIDTexture,
-	FRDGTextureRef InMSAADepthTexture,
-
-	FRDGTextureRef OutBufferATexture,
-	FRDGTextureRef OutBufferBTexture,
-	FRDGTextureRef OutBufferCTexture,
-	FRDGTextureRef OutBufferDTexture,
-	FRDGTextureRef OutBufferETexture,
-
-	FRDGTextureRef OutColorTexture,
-	FRDGTextureRef OutDepthTexture,
-	FRDGTextureRef OutVelocityTexture)
-{
-	const FRDGSystemTextures& SystemTextures = FRDGSystemTextures::Get(GraphBuilder);
-
-	// Add resources reference to the pass parameters, in order to get the resource lifetime extended to this pass
-	FVisibilityMaterialGBufferPassParameters* PassParameters = GraphBuilder.AllocParameters<FVisibilityMaterialGBufferPassParameters>();
-
-	const FIntPoint Resolution = OutBufferATexture->Desc.Extent;
-
-	{
-		FMaterialGBufferPassUniformParameters* UniformParameters = GraphBuilder.AllocParameters<FMaterialGBufferPassUniformParameters>();
-		UniformParameters->InputType = InRasterOutputVisibilityTexture != nullptr ? 1 : 0;
-		UniformParameters->MaxResolution = Resolution;
-		UniformParameters->CoverageThreshold = FMath::Clamp(GHairStrandsFullCoverageThreshold, 0.1f, 1.f);
-		UniformParameters->InTransmittanceTexture = InTransmittanceTexture;
-		UniformParameters->InMSAAIDTexture = SystemTextures.Black;
-		UniformParameters->InMSAADepthTexture = SystemTextures.Black;
-		UniformParameters->InRasterOutputVisibilityTexture = SystemTextures.Black;
-		if (UniformParameters->InputType == 0)
-		{
-			check(InMSAAIDTexture);
-			check(InMSAADepthTexture);
-			UniformParameters->InMSAAIDTexture = InMSAAIDTexture;
-			UniformParameters->InMSAADepthTexture = InMSAADepthTexture;
-		}
-		else if (UniformParameters->InputType == 1)
-		{
-			check(InRasterOutputVisibilityTexture);
-			UniformParameters->InRasterOutputVisibilityTexture = InRasterOutputVisibilityTexture;
-		}
-
-		PassParameters->UniformBuffer = GraphBuilder.CreateUniformBuffer(UniformParameters);
-	}
-
-	{
-		const bool bEnableMSAA = false;
-		SetUpViewHairRenderInfo(*ViewInfo, bEnableMSAA, ViewInfo->CachedViewUniformShaderParameters->HairRenderInfo, ViewInfo->CachedViewUniformShaderParameters->HairRenderInfoBits, ViewInfo->CachedViewUniformShaderParameters->HairComponents);
-		PassParameters->View = TUniformBufferRef<FViewUniformShaderParameters>::CreateUniformBufferImmediate(*ViewInfo->CachedViewUniformShaderParameters, UniformBuffer_SingleFrame);
-	}
-
-	// If there is velocity texture, we recreate a dummy one
-	bool bIsVelocityDummy = false;
-	if (OutVelocityTexture == nullptr)
-	{
-		FRDGTextureDesc VelocityDesc = FVelocityRendering::GetRenderTargetDesc(ViewInfo->GetShaderPlatform(), OutDepthTexture->Desc.Extent);
-		OutVelocityTexture = GraphBuilder.CreateTexture(VelocityDesc, TEXT("Hair.DummyVelocity"));
-		bIsVelocityDummy = true;
-	}
-	bool bIsBufferEDummy = false;
-	if (OutBufferETexture == nullptr)
-	{
-		FRDGTextureDesc Desc = FVelocityRendering::GetRenderTargetDesc(ViewInfo->GetShaderPlatform(), OutBufferDTexture->Desc.Extent);
-		OutBufferETexture = GraphBuilder.CreateTexture(Desc, TEXT("Hair.DummyGBufferE"));
-		bIsBufferEDummy = true;
-	}
-	PassParameters->RenderTargets[0] = FRenderTargetBinding(OutBufferATexture, ERenderTargetLoadAction::ELoad, 0);
-	PassParameters->RenderTargets[1] = FRenderTargetBinding(OutBufferBTexture, ERenderTargetLoadAction::ELoad, 0);
-	PassParameters->RenderTargets[2] = FRenderTargetBinding(OutBufferCTexture, ERenderTargetLoadAction::ELoad, 0);
-	PassParameters->RenderTargets[3] = FRenderTargetBinding(OutBufferDTexture, ERenderTargetLoadAction::ELoad, 0);
-	PassParameters->RenderTargets[4] = FRenderTargetBinding(OutBufferETexture, bIsBufferEDummy ? ERenderTargetLoadAction::ENoAction : ERenderTargetLoadAction::ELoad, 0);
-	PassParameters->RenderTargets[5] = FRenderTargetBinding(OutVelocityTexture, bIsVelocityDummy ? ERenderTargetLoadAction::ENoAction : ERenderTargetLoadAction::ELoad, 0);
-	PassParameters->RenderTargets[6] = FRenderTargetBinding(OutColorTexture, ERenderTargetLoadAction::ELoad, 0);
-	PassParameters->RenderTargets.DepthStencil = FDepthStencilBinding(
-		OutDepthTexture,
-		ERenderTargetLoadAction::ELoad,
-		ERenderTargetLoadAction::ELoad,
-		FExclusiveDepthStencil::DepthWrite_StencilNop);
-
-	AddSimpleMeshPass(GraphBuilder, PassParameters, Scene, *ViewInfo, &InstanceCullingManager, RDG_EVENT_NAME("HairStrandsMaterialGBufferPass"), FIntRect(0, 0, Resolution.X, Resolution.Y),
-		[PassParameters, Scene = Scene, ViewInfo, &MacroGroupDatas, Resolution](FDynamicPassMeshDrawListContext* ShadowContext)
-		{
-			FMeshPassProcessorRenderState DrawRenderState;
-
-			{
-				DrawRenderState.SetBlendState(TStaticBlendState<>::GetRHI());
-				DrawRenderState.SetDepthStencilState(TStaticDepthStencilState <true, CF_Always> ::GetRHI());
-
-				FHairMaterialGBufferProcessor MeshProcessor(Scene, ViewInfo, DrawRenderState, ShadowContext);
-
-				for (const FHairStrandsMacroGroupData& MacroGroupData : MacroGroupDatas)
-				{
-					for (const FHairStrandsMacroGroupData::PrimitiveInfo& PrimitiveInfo : MacroGroupData.PrimitivesInfos)
-					{
-						const FMeshBatch& MeshBatch = *PrimitiveInfo.Mesh;
-						const uint64 BatchElementMask = ~0ull;
-						MeshProcessor.AddMeshBatch(MeshBatch, BatchElementMask, PrimitiveInfo.PrimitiveSceneProxy, -1, MacroGroupData.MacroGroupId, PrimitiveInfo.MaterialId);
-					}
-				}
-			}
-		});
-}
-
 /////////////////////////////////////////////////////////////////////////////////////////
 
 class FHairMaterialVS : public FMeshMaterialShader
@@ -3358,8 +2947,6 @@ void RenderHairStrandsVisibilityBuffer(
 					VisibilityData.ViewHairCountUintTexture = ViewTransmittance.HairCountTextureUint;
 				}
 
-				const bool bUseComplexPath = IsHairStrandsComplexLightingEnabled();
-				if (bUseComplexPath)
 				{
 					{
 						FRDGBufferRef CompactNodeCoord;
@@ -3464,32 +3051,6 @@ void RenderHairStrandsVisibilityBuffer(
 						VisibilityData.HairOnlyDepthClosestHZBTexture,
 						VisibilityData.HairOnlyDepthFurthestHZBTexture);
 				}
-				else
-				{
-					AddHairMaterialGBufferPass(
-						GraphBuilder,
-						Scene,
-						&View,
-						InstanceCullingManager,
-						MacroGroupDatas,
-
-						ViewTransmittance.TransmittanceTexture,
-						RasterOutput.VisibilityTexture0,
-						nullptr,
-						nullptr,
-
-						SceneGBufferATexture,
-						SceneGBufferBTexture,
-						SceneGBufferCTexture,
-						SceneGBufferDTexture,
-						SceneGBufferETexture,
-
-						SceneColorTexture,
-						SceneDepthTexture,
-						SceneVelocityTexture);
-
-					ViewTransmittance.HairCountTexture = nullptr;
-				}
 			}
 			else if (RenderMode == HairVisibilityRenderMode_MSAA_Visibility)
 			{
@@ -3548,8 +3109,6 @@ void RenderHairStrandsVisibilityBuffer(
 				VisibilityData.MaxSampleCount = MsaaVisibilityResources.IdTexture->Desc.NumSamples;
 				VisibilityData.HairOnlyDepthTexture = HairOnlyDepthTexture;
 				
-				const bool bUseComplexPath = IsHairStrandsComplexLightingEnabled();
-				if (bUseComplexPath)
 				{
 					FHairVisibilityPrimitiveIdCompactionCS::FParameters* PassParameters = GraphBuilder.AllocParameters<FHairVisibilityPrimitiveIdCompactionCS::FParameters>();
 					PassParameters->MSAA_DepthTexture = MsaaVisibilityResources.DepthTexture;
@@ -3684,32 +3243,6 @@ void RenderHairStrandsVisibilityBuffer(
 						VisibilityData.HairOnlyDepthClosestHZBTexture,
 						VisibilityData.HairOnlyDepthFurthestHZBTexture);
 				}
-				else
-				{
-					AddHairMaterialGBufferPass(
-						GraphBuilder,
-						Scene,
-						&View,
-						InstanceCullingManager,
-						MacroGroupDatas,
-
-						ViewTransmittance.TransmittanceTexture,
-						nullptr,
-						MsaaVisibilityResources.IdTexture,
-						MsaaVisibilityResources.DepthTexture,
-
-						SceneGBufferATexture,
-						SceneGBufferBTexture,
-						SceneGBufferCTexture,
-						SceneGBufferDTexture,
-						SceneGBufferETexture,
-
-						SceneColorTexture,
-						SceneDepthTexture,
-						SceneVelocityTexture);
-
-						ViewTransmittance.HairCountTexture = nullptr;
-				}
 			}
 			else if (RenderMode == HairVisibilityRenderMode_PPLL)
 			{
diff --git a/Engine/Source/Runtime/Renderer/Public/HairStrandsInterface.h b/Engine/Source/Runtime/Renderer/Public/HairStrandsInterface.h
--- a/Engine/Source/Runtime/Renderer/Public/HairStrandsInterface.h
+++ b/Engine/Source/Runtime/Renderer/Public/HairStrandsInterface.h
@@ -399,9 +399,6 @@ RENDERER_API FHairCullInfo GetHairStrandsCullInfo();
 
 RENDERER_API bool IsHairRayTracingEnabled();
 
-// Return true if the hair should be rendered using the sub-pixel lighting path, false if the regular gbuffer lighting path should be used
-RENDERER_API bool IsHairStrandsComplexLightingEnabled();
-
 // Return true if hair simulation is enabled.
 RENDERER_API bool IsHairStrandsSimulationEnable();
 
