Use half res only when using GTAO on mobile renderer

diff --git a/Engine/Source/Runtime/Renderer/Private/PostProcess/PostProcessAmbientOcclusionMobile.cpp b/Engine/Source/Runtime/Renderer/Private/PostProcess/PostProcessAmbientOcclusionMobile.cpp
--- a/Engine/Source/Runtime/Renderer/Private/PostProcess/PostProcessAmbientOcclusionMobile.cpp
+++ b/Engine/Source/Runtime/Renderer/Private/PostProcess/PostProcessAmbientOcclusionMobile.cpp
@@ -430,7 +430,9 @@ IMPLEMENT_GLOBAL_SHADER(FGTAOMobile_SpatialFilterPS, "/Engine/Private/PostProces
 
 FRDGTextureRef CreateMobileScreenSpaceAOTexture(FRDGBuilder& GraphBuilder, FIntPoint SceneTextureExtent)
 {
-	const uint32 DownsampleFactor = 2;
+	bool bGTAO = (CVarMobileAmbientOcclusionTechnique.GetValueOnRenderThread() == 0);
+	const uint32 DownsampleFactor = bGTAO ? 2 : 1;
+
 	const FIntPoint Extent = FIntPoint::DivideAndRoundUp(SceneTextureExtent, DownsampleFactor);
 	const bool bUsePixelShader = CVarMobileAmbientOcclusionShaderType.GetValueOnRenderThread() == 2;
 	return GraphBuilder.CreateTexture(
