Fix crash using async SSAO.

diff --git a/Engine/Source/Runtime/Renderer/Private/CompositionLighting/PostProcessAmbientOcclusion.cpp b/Engine/Source/Runtime/Renderer/Private/CompositionLighting/PostProcessAmbientOcclusion.cpp
--- a/Engine/Source/Runtime/Renderer/Private/CompositionLighting/PostProcessAmbientOcclusion.cpp
+++ b/Engine/Source/Runtime/Renderer/Private/CompositionLighting/PostProcessAmbientOcclusion.cpp
@@ -615,7 +615,7 @@ public:
 };
 IMPLEMENT_GLOBAL_SHADER(FAmbientOcclusionSmoothCS, "/Engine/Private/PostProcessAmbientOcclusion.usf", "MainSSAOSmoothCS", SF_Compute);
 
-FScreenPassTexture AddAmbientOcclusionSmoothPass(
+void AddAmbientOcclusionSmoothPass(
 	FRDGBuilder& GraphBuilder,
 	const FViewInfo& View,
 	ESSAOType SSAOType,
@@ -646,8 +646,6 @@ FScreenPassTexture AddAmbientOcclusionSmoothPass(
 		ComputeShader,
 		PassParameters,
 		FComputeShaderUtils::GetGroupCount(OutputViewport.Rect.Size(), 8));
-
-	return MoveTemp(Output);
 }
 
 //----------------------------------------------------------------------------------------------------------------------
@@ -744,7 +742,41 @@ public:
 };
 IMPLEMENT_GLOBAL_SHADER(FAmbientOcclusionCS, "/Engine/Private/PostProcessAmbientOcclusion.usf", "MainCS", SF_Compute);
 
-FScreenPassTexture AddAmbientOcclusionPass(
+FScreenPassRenderTarget CreateAmbientOcclusionOutputTarget(
+	FRDGBuilder& GraphBuilder,
+	FRDGTextureDesc OutputDesc,
+	FIntRect ViewRect,
+	ESSAOType AOType,
+	EPixelFormat IntermediateFormatOverride)
+{
+	const bool bUsingUAVOutput = (AOType == ESSAOType::ECS || AOType == ESSAOType::EAsyncCS);
+
+	OutputDesc.Reset();
+	OutputDesc.ClearValue = FClearValueBinding::None;
+	OutputDesc.Flags &= ~TexCreate_DepthStencilTargetable;
+
+	if (bUsingUAVOutput)
+	{
+		// UAV allowed format
+		OutputDesc.Format = PF_FloatRGBA;
+		OutputDesc.Flags |= TexCreate_UAV;
+	}
+	else
+	{
+		// R:AmbientOcclusion, GBA:used for normal
+		OutputDesc.Format = PF_B8G8R8A8;
+		OutputDesc.Flags |= TexCreate_RenderTargetable;
+	}
+
+	if (IntermediateFormatOverride != PF_Unknown)
+	{
+		OutputDesc.Format = IntermediateFormatOverride;
+	}
+
+	return FScreenPassRenderTarget(GraphBuilder.CreateTexture(OutputDesc, TEXT("AmbientOcclusion")), ViewRect, ERenderTargetLoadAction::ENoAction);
+}
+
+void AddAmbientOcclusionPass(
 	FRDGBuilder& GraphBuilder,
 	const FViewInfo& View,
 	const FSSAOCommonParameters& CommonParameters,
@@ -752,46 +784,12 @@ FScreenPassTexture AddAmbientOcclusionPass(
 	const FScreenPassTexture& NormalsTexture,
 	const FScreenPassTexture& DownsampledAO,
 	const FScreenPassTexture& HZBInput,
-	FScreenPassRenderTarget SuggestedOutput,
+	FScreenPassRenderTarget Output,
 	ESSAOType AOType,
-	bool bAOSetupAsInput,
-	EPixelFormat IntermediateFormatOverride)
+	bool bAOSetupAsInput)
 {
 	RDG_GPU_STAT_SCOPE(GraphBuilder, SSAO);
 
-	FScreenPassRenderTarget Output = SuggestedOutput;
-	if (!Output.IsValid())
-	{
-		check(SetupTexture.IsValid());
-
-		const bool bUsingUAVOutput = (AOType == ESSAOType::ECS || AOType == ESSAOType::EAsyncCS);
-
-		FRDGTextureDesc OutputDesc = SetupTexture.Texture->Desc;
-		OutputDesc.Reset();
-		OutputDesc.ClearValue = FClearValueBinding::None;
-		OutputDesc.Flags &= ~TexCreate_DepthStencilTargetable;
-		if (bUsingUAVOutput)
-		{
-			// UAV allowed format
-			OutputDesc.Format = PF_FloatRGBA;
-			OutputDesc.Flags |= TexCreate_UAV;
-		}
-		else
-		{
-			// R:AmbientOcclusion, GBA:used for normal
-			OutputDesc.Format = PF_B8G8R8A8;
-			OutputDesc.Flags |= TexCreate_RenderTargetable;
-		}
-		if (IntermediateFormatOverride != PF_Unknown)
-		{
-			OutputDesc.Format = IntermediateFormatOverride;
-		}
-
-		Output.Texture = GraphBuilder.CreateTexture(OutputDesc, TEXT("AmbientOcclusion"));
-		Output.ViewRect = SetupTexture.ViewRect;
-		Output.LoadAction = ERenderTargetLoadAction::ENoAction;
-	}
-
 	// No setup texture falls back to a depth scene texture fetch.
 	const FScreenPassTextureViewport InputViewport = SetupTexture.IsValid()
 		? FScreenPassTextureViewport(SetupTexture)
@@ -993,8 +991,6 @@ FScreenPassTexture AddAmbientOcclusionPass(
 			}
 		});
 	}
-
-	return MoveTemp(Output);
 }
 
 FScreenPassTexture AddAmbientOcclusionStepPass(
@@ -1006,7 +1002,8 @@ FScreenPassTexture AddAmbientOcclusionStepPass(
 	const FScreenPassTexture& DownsampledAO,
 	const FScreenPassTexture& HZBInput)
 {
-	return AddAmbientOcclusionPass(
+	FScreenPassRenderTarget Output = CreateAmbientOcclusionOutputTarget(GraphBuilder, SetupTexture.Texture->Desc, SetupTexture.ViewRect, CommonParameters.DownscaleType, PF_Unknown);
+	AddAmbientOcclusionPass(
 		GraphBuilder,
 		View,
 		CommonParameters,
@@ -1014,10 +1011,10 @@ FScreenPassTexture AddAmbientOcclusionStepPass(
 		NormalsTexture,
 		DownsampledAO,
 		HZBInput,
-		FScreenPassRenderTarget(),
+		Output,
 		CommonParameters.DownscaleType,
-		true,
-		PF_Unknown);
+		true);
+	return MoveTemp(Output);
 }
 
 FScreenPassTexture AddAmbientOcclusionFinalPass(
@@ -1030,32 +1027,35 @@ FScreenPassTexture AddAmbientOcclusionFinalPass(
 	const FScreenPassTexture& HZBInput,
 	FScreenPassRenderTarget FinalOutput)
 {
-	FScreenPassTexture CurrentOutput =
-		AddAmbientOcclusionPass(
-			GraphBuilder,
-			View,
-			CommonParameters,
-			SetupTexture,
-			NormalsTexture,
-			DownsampledAO,
-			HZBInput,
-			CommonParameters.bNeedSmoothingPass ? FScreenPassRenderTarget() : FinalOutput,
-			CommonParameters.FullscreenType,
-			false,
-			PF_G8);
+	FScreenPassRenderTarget CurrentOutput = CommonParameters.bNeedSmoothingPass ?
+		CreateAmbientOcclusionOutputTarget(GraphBuilder, FinalOutput.Texture->Desc, FinalOutput.ViewRect, CommonParameters.FullscreenType, PF_G8):
+		FinalOutput;
+
+	AddAmbientOcclusionPass(
+		GraphBuilder,
+		View,
+		CommonParameters,
+		SetupTexture,
+		NormalsTexture,
+		DownsampledAO,
+		HZBInput,
+		CurrentOutput,
+		CommonParameters.FullscreenType,
+		false);
 
 	if (CommonParameters.bNeedSmoothingPass)
 	{
-		CurrentOutput =
-			AddAmbientOcclusionSmoothPass(
+		AddAmbientOcclusionSmoothPass(
 				GraphBuilder,
 				View,
 				CommonParameters.FullscreenType,
 				CurrentOutput,
 				FinalOutput);
+
+		CurrentOutput = FinalOutput;
 	}
 
-	return CurrentOutput;
+	return MoveTemp(CurrentOutput);
 }
 
 //----------------------------------------------------------------------------------------------------------------------
