Fix invalid hair tile viewport setup.

diff --git a/Engine/Shaders/Private/HairStrands/HairStrandsVisibilityTile.usf b/Engine/Shaders/Private/HairStrands/HairStrandsVisibilityTile.usf
--- a/Engine/Shaders/Private/HairStrands/HairStrandsVisibilityTile.usf
+++ b/Engine/Shaders/Private/HairStrands/HairStrandsVisibilityTile.usf
@@ -10,7 +10,7 @@
 ////////////////////////////////////////////////////////////////////////////////////////////////////////////
 
 #if SHADER_TILE_GENERATION
-int2 ViewResolution;
+int2 BufferResolution;
 Texture2D<uint4> InputTexture;
 RWBuffer<uint> TileCountBuffer;
 RWBuffer<uint2> TileDataBuffer;
@@ -24,7 +24,7 @@ void TileMainCS(uint2 DispatchThreadId : SV_DispatchThreadID, uint LinearIndex :
 	const uint2 PixelCoord = DispatchThreadId;
 
 	bool bHasHairData = false;
-	if (all(PixelCoord < uint2(ViewResolution)))
+	if (all(PixelCoord < uint2(BufferResolution)))
 	{
 		const FCategorizationData CategorisationData = DecodeCategorizationData(InputTexture.Load(uint3(PixelCoord, 0)));
 		bHasHairData = CategorisationData.PixelCoverage > 0.0f;
@@ -142,7 +142,8 @@ void MainPS(
 ////////////////////////////////////////////////////////////////////////////////////////////////////////////
 
 #if SHADER_TILE_VS
-int2 TileOutputResolution;
+int2 ViewMin;
+float2 ViewInvSize;
 Buffer<uint2> TileDataBuffer;
 
 void MainVS(
@@ -156,7 +157,7 @@ void MainVS(
 	uint2 TileVertex = TileCoord * TILE_SIZE;
 	TileVertex.x += InVertexId == 1 || InVertexId == 2 || InVertexId == 4 ? TILE_SIZE : 0;
 	TileVertex.y += InVertexId == 2 || InVertexId == 4 || InVertexId == 5 ? TILE_SIZE : 0;
-	Out.UV = float2(TileVertex) / float2(TileOutputResolution);
+	Out.UV = float2(TileVertex - ViewMin) * ViewInvSize;
 	Out.Position = float4(Out.UV * float2(2.0f, -2.0f) + float2(-1.0, 1.0f), 0.5f, 1.0f);
 }
 
diff --git a/Engine/Source/Runtime/Renderer/Private/HairStrands/HairStrandsComposition.cpp b/Engine/Source/Runtime/Renderer/Private/HairStrands/HairStrandsComposition.cpp
--- a/Engine/Source/Runtime/Renderer/Private/HairStrands/HairStrandsComposition.cpp
+++ b/Engine/Source/Runtime/Renderer/Private/HairStrands/HairStrandsComposition.cpp
@@ -62,7 +62,7 @@ void InternalCommonDrawPass(
 	const bool bUseTile = TileData.IsValid();
 	if (TileData.IsValid())
 	{
-		PassParamters->TileData = GetHairStrandsTileParameters(TileData);
+		PassParamters->TileData = GetHairStrandsTileParameters(View, TileData);
 	}
 	GraphBuilder.AddPass(
 		Forward<FRDGEventName>(EventName),
@@ -134,6 +134,8 @@ void InternalCommonDrawPass(
 		}
 
 		SetShaderParameters(RHICmdList, PixelShader, PixelShader.GetPixelShader(), *PassParamters);
+		RHICmdList.SetViewport(Viewport.Min.X, Viewport.Min.Y, 0.0f, Viewport.Max.X, Viewport.Max.Y, 1.0f);
+
 		if (bUseTile)
 		{
 			SetShaderParameters(RHICmdList, TileVertexShader, TileVertexShader.GetVertexShader(), ParametersVS);
@@ -142,7 +144,6 @@ void InternalCommonDrawPass(
 		}
 		else
 		{
-			RHICmdList.SetViewport(Viewport.Min.X, Viewport.Min.Y, 0.0f, Viewport.Max.X, Viewport.Max.Y, 1.0f);
 			DrawRectangle(
 				RHICmdList,
 				0, 0,
diff --git a/Engine/Source/Runtime/Renderer/Private/HairStrands/HairStrandsData.h b/Engine/Source/Runtime/Renderer/Private/HairStrands/HairStrandsData.h
--- a/Engine/Source/Runtime/Renderer/Private/HairStrands/HairStrandsData.h
+++ b/Engine/Source/Runtime/Renderer/Private/HairStrands/HairStrandsData.h
@@ -45,7 +45,7 @@ struct FHairStrandsTiles
 {
 	enum class ETileType : uint8 { Hair, Other };
 
-	FIntPoint			Resolution = FIntPoint(0, 0);
+	FIntPoint			BufferResolution = FIntPoint(0, 0);
 	static const uint32 GroupSize = 64;
 	static const uint32	TileSize = 8;
 	static const uint32	TilePerThread_GroupSize = 64;
diff --git a/Engine/Source/Runtime/Renderer/Private/HairStrands/HairStrandsTile.cpp b/Engine/Source/Runtime/Renderer/Private/HairStrands/HairStrandsTile.cpp
--- a/Engine/Source/Runtime/Renderer/Private/HairStrands/HairStrandsTile.cpp
+++ b/Engine/Source/Runtime/Renderer/Private/HairStrands/HairStrandsTile.cpp
@@ -16,18 +16,14 @@
 
 ////////////////////////////////////////////////////////////////////////////////////////////////////////////////////////////////////////////////
 
-FHairStrandsTilePassVS::FParameters GetHairStrandsTileParameters(const FHairStrandsTiles& In, FHairStrandsTiles::ETileType TileType)
-{
-	return GetHairStrandsTileParameters(In, In.Resolution, TileType);
-}
-
-FHairStrandsTilePassVS::FParameters GetHairStrandsTileParameters(const FHairStrandsTiles& In, const FIntPoint& InOutputResolution, FHairStrandsTiles::ETileType TileType)
+FHairStrandsTilePassVS::FParameters GetHairStrandsTileParameters(const FViewInfo& InView, const FHairStrandsTiles& InTile, FHairStrandsTiles::ETileType TileType)
 {
 	FHairStrandsTilePassVS::FParameters Out;
-	Out.bRectPrimitive			= In.bRectPrimitive ? 1 : 0;
-	Out.TileOutputResolution	= InOutputResolution;
-	Out.TileDataBuffer			= TileType == FHairStrandsTiles::ETileType::Hair ? In.TileDataSRV : In.TileClearSRV;
-	Out.TileIndirectBuffer		= TileType == FHairStrandsTiles::ETileType::Hair ? In.TileIndirectDrawBuffer : In.TileClearIndirectDrawBuffer;
+	Out.bRectPrimitive			= InTile.bRectPrimitive ? 1 : 0;
+	Out.ViewMin					= InView.ViewRect.Min;
+	Out.ViewInvSize				= FVector2D(1.f / InView.ViewRect.Width(), 1.f / InView.ViewRect.Height());
+	Out.TileDataBuffer			= TileType == FHairStrandsTiles::ETileType::Hair ? InTile.TileDataSRV : InTile.TileClearSRV;
+	Out.TileIndirectBuffer		= TileType == FHairStrandsTiles::ETileType::Hair ? InTile.TileIndirectDrawBuffer : InTile.TileClearIndirectDrawBuffer;
 	return Out;
 }
 
@@ -86,7 +82,7 @@ void AddHairStrandsCopyArgsTilesPass(
 
 	FComputeShaderUtils::AddPass(
 		GraphBuilder,
-		RDG_EVENT_NAME("HairStrandsTileCopyArgs"),
+		RDG_EVENT_NAME("HairStrands::TileCopyArgs"),
 		ComputeShader,
 		PassParameters,
 		FIntVector(1, 1, 1));
@@ -104,7 +100,7 @@ class FHairStrandsTileGenerationPassCS : public FGlobalShader
 
 	BEGIN_SHADER_PARAMETER_STRUCT(FParameters, )
 		SHADER_PARAMETER_STRUCT_REF(FViewUniformShaderParameters, ViewUniformBuffer)
-		SHADER_PARAMETER(FIntPoint, ViewResolution)
+		SHADER_PARAMETER(FIntPoint, BufferResolution)
 		SHADER_PARAMETER_RDG_TEXTURE(Texture2D<uint4>, InputTexture)
 		SHADER_PARAMETER_RDG_BUFFER_UAV(RWBuffer, TileCountBuffer)
 		SHADER_PARAMETER_RDG_BUFFER_UAV(RWBuffer, TileDataBuffer)
@@ -136,7 +132,7 @@ FHairStrandsTiles AddHairStrandsGenerateTilesPass(
 	const FIntPoint InputResolution = InputTexture->Desc.Extent;
 	Out.TileCountXY = FIntPoint(FMath::CeilToInt(InputResolution.X / float(FHairStrandsTiles::TileSize)), FMath::CeilToInt(InputResolution.Y / float(FHairStrandsTiles::TileSize)));
 	Out.TileCount = Out.TileCountXY.X * Out.TileCountXY.Y;
-	Out.Resolution = InputResolution;
+	Out.BufferResolution = InputResolution;
 	Out.bRectPrimitive = GRHISupportsRectTopology;
 	Out.TileCountBuffer = GraphBuilder.CreateBuffer(FRDGBufferDesc::CreateBufferDesc(4, 2), TEXT("Hair.TileCountBuffer"));
 	Out.TileIndirectDrawBuffer = GraphBuilder.CreateBuffer(FRDGBufferDesc::CreateIndirectDesc<FRHIDrawIndirectParameters>(1), TEXT("Hair.TileIndirectDrawBuffer"));
@@ -152,7 +148,7 @@ FHairStrandsTiles AddHairStrandsGenerateTilesPass(
 
 	TShaderMapRef<FHairStrandsTileGenerationPassCS> ComputeShader(View.ShaderMap);
 	FHairStrandsTileGenerationPassCS::FParameters* PassParameters = GraphBuilder.AllocParameters<FHairStrandsTileGenerationPassCS::FParameters>();
-	PassParameters->ViewResolution = View.ViewRect.Size();
+	PassParameters->BufferResolution = InputResolution;//View.ViewRect.Size();
 	PassParameters->InputTexture = InputTexture;
 	PassParameters->TileDataBuffer = GraphBuilder.CreateUAV(Out.TileDataBuffer, PF_R16G16_UINT);
 	PassParameters->TileClearBuffer = GraphBuilder.CreateUAV(Out.TileClearBuffer, PF_R16G16_UINT);
@@ -160,7 +156,7 @@ FHairStrandsTiles AddHairStrandsGenerateTilesPass(
 
 	FComputeShaderUtils::AddPass(
 		GraphBuilder,
-		RDG_EVENT_NAME("HairStrandsTileClassification"),
+		RDG_EVENT_NAME("HairStrands::TileClassification"),
 		ComputeShader,
 		PassParameters,
 		FComputeShaderUtils::GetGroupCount(InputResolution, FHairStrandsTiles::TileSize));
@@ -227,11 +223,10 @@ void AddHairStrandsDebugTilePass(
 	const FRDGTextureRef& ColorTexture,
 	const FHairStrandsTiles& TileData)
 {	
-	const FIntPoint OutputResolution = View.ViewRect.Size();
+	const FIntRect Viewport = View.ViewRect;
 	
 	FHairStrandsTileDebugPassPS::FParameters* PassParameters = GraphBuilder.AllocParameters<FHairStrandsTileDebugPassPS::FParameters>();
-	PassParameters->TileParameters = GetHairStrandsTileParameters(TileData, OutputResolution);
-	PassParameters->OutputResolution = OutputResolution; // Use the Original Buffer for 
+	PassParameters->TileParameters = GetHairStrandsTileParameters(View, TileData);
 
 	TShaderMapRef<FHairStrandsTilePassVS> VertexShader(View.ShaderMap);
 	TShaderMapRef<FHairStrandsTileDebugPassPS> PixelShader(View.ShaderMap);
@@ -239,10 +234,10 @@ void AddHairStrandsDebugTilePass(
 	PassParameters->RenderTargets[0] = FRenderTargetBinding(ColorTexture, ERenderTargetLoadAction::ELoad);
 	
 	GraphBuilder.AddPass(
-		RDG_EVENT_NAME("HairStrandsTileDebugPass"),
+		RDG_EVENT_NAME("HairStrands::TileDebugPass"),
 		PassParameters,
 		ERDGPassFlags::Raster,
-		[PassParameters, VertexShader, PixelShader, OutputResolution](FRHICommandList& RHICmdList)
+		[PassParameters, VertexShader, PixelShader, Viewport](FRHICommandList& RHICmdList)
 		{
 			FHairStrandsTilePassVS::FParameters ParametersVS = PassParameters->TileParameters;
 
@@ -260,7 +255,7 @@ void AddHairStrandsDebugTilePass(
 			SetShaderParameters(RHICmdList, VertexShader, VertexShader.GetVertexShader(), ParametersVS);
 			SetShaderParameters(RHICmdList, PixelShader, PixelShader.GetPixelShader(), *PassParameters);
 
-			RHICmdList.SetViewport(0, 0, 0.0f, OutputResolution.X, OutputResolution.Y, 1.0f);
+			RHICmdList.SetViewport(Viewport.Min.X, Viewport.Min.Y, 0.0f, Viewport.Max.X, Viewport.Max.Y, 1.0f);
 			RHICmdList.SetStreamSource(0, nullptr, 0);
 			RHICmdList.DrawPrimitiveIndirect(PassParameters->TileParameters.TileIndirectBuffer->GetRHI(), 0);
 		});
diff --git a/Engine/Source/Runtime/Renderer/Private/HairStrands/HairStrandsTile.h b/Engine/Source/Runtime/Renderer/Private/HairStrands/HairStrandsTile.h
--- a/Engine/Source/Runtime/Renderer/Private/HairStrands/HairStrandsTile.h
+++ b/Engine/Source/Runtime/Renderer/Private/HairStrands/HairStrandsTile.h
@@ -22,7 +22,8 @@ class FHairStrandsTilePassVS : public FGlobalShader
 
 	BEGIN_SHADER_PARAMETER_STRUCT(FParameters, )
 		SHADER_PARAMETER(int32, bRectPrimitive)
-		SHADER_PARAMETER(FIntPoint, TileOutputResolution)
+		SHADER_PARAMETER(FIntPoint, ViewMin)
+		SHADER_PARAMETER(FVector2D, ViewInvSize)
 		SHADER_PARAMETER_RDG_BUFFER_SRV(Buffer, TileDataBuffer)
 		RDG_BUFFER_ACCESS(TileIndirectBuffer, ERHIAccess::IndirectArgs)
 	END_SHADER_PARAMETER_STRUCT()
@@ -31,8 +32,7 @@ class FHairStrandsTilePassVS : public FGlobalShader
 	static void ModifyCompilationEnvironment(const FGlobalShaderPermutationParameters& Parameters, FShaderCompilerEnvironment& OutEnvironment);
 };
 
-FHairStrandsTilePassVS::FParameters GetHairStrandsTileParameters(const FHairStrandsTiles& In, FHairStrandsTiles::ETileType TileType = FHairStrandsTiles::ETileType::Hair);
-FHairStrandsTilePassVS::FParameters GetHairStrandsTileParameters(const FHairStrandsTiles& In, const FIntPoint& InOutputResolution, FHairStrandsTiles::ETileType TileType = FHairStrandsTiles::ETileType::Hair);
+FHairStrandsTilePassVS::FParameters GetHairStrandsTileParameters(const FViewInfo& InView, const FHairStrandsTiles& In, FHairStrandsTiles::ETileType TileType = FHairStrandsTiles::ETileType::Hair);
 
 FHairStrandsTiles AddHairStrandsGenerateTilesPass(
 	FRDGBuilder& GraphBuilder,
diff --git a/Engine/Source/Runtime/Renderer/Private/HairStrands/HairStrandsVisibility.cpp b/Engine/Source/Runtime/Renderer/Private/HairStrands/HairStrandsVisibility.cpp
--- a/Engine/Source/Runtime/Renderer/Private/HairStrands/HairStrandsVisibility.cpp
+++ b/Engine/Source/Runtime/Renderer/Private/HairStrands/HairStrandsVisibility.cpp
@@ -2308,6 +2308,7 @@ class FHairVisibilityDepthPS : public FGlobalShader
 
 	BEGIN_SHADER_PARAMETER_STRUCT(FParameters, )
 		SHADER_PARAMETER_STRUCT_INCLUDE(FHairStrandsTilePassVS::FParameters, TileData)
+		SHADER_PARAMETER_STRUCT_REF(FViewUniformShaderParameters, ViewUniformBuffer)
 		SHADER_PARAMETER(uint32, bClear)
 		SHADER_PARAMETER_RDG_TEXTURE(Texture2D, CategorisationTexture)
 		RENDER_TARGET_BINDING_SLOTS()
@@ -2350,8 +2351,9 @@ static void AddHairVisibilityCommonPatchPass(
 {
 	FHairVisibilityDepthPS::FParameters* Parameters = GraphBuilder.AllocParameters<FHairVisibilityDepthPS::FParameters>();
 	Parameters->bClear = PassType == EHairPatchPass::DepthClear ? 1u : 0u;
-	Parameters->TileData = GetHairStrandsTileParameters(TileData, PassType == EHairPatchPass::DepthClear ?  FHairStrandsTiles::ETileType::Other : FHairStrandsTiles::ETileType::Hair);
+	Parameters->TileData = GetHairStrandsTileParameters(View, TileData, PassType == EHairPatchPass::DepthClear ?  FHairStrandsTiles::ETileType::Other : FHairStrandsTiles::ETileType::Hair);
 	Parameters->CategorisationTexture = CategorisationTexture;
+	Parameters->ViewUniformBuffer = View.ViewUniformBuffer;
 	Parameters->RenderTargets.DepthStencil = FDepthStencilBinding(
 		OutDepthTexture,
 		ERenderTargetLoadAction::ELoad,
@@ -2395,13 +2397,12 @@ static void AddHairVisibilityCommonPatchPass(
 		GraphicsPSOInit.PrimitiveType = Parameters->TileData.bRectPrimitive > 0 ? PT_RectList : PT_TriangleList;
 		SetGraphicsPipelineState(RHICmdList, GraphicsPSOInit);
 
-		RHICmdList.SetViewport(Viewport.Min.X, Viewport.Min.Y, 0.0f, Viewport.Max.X, Viewport.Max.Y, 1.0f);
 		SetShaderParameters(RHICmdList, PixelShader, PixelShader.GetPixelShader(), *Parameters);
 
+		RHICmdList.SetViewport(Viewport.Min.X, Viewport.Min.Y, 0.0f, Viewport.Max.X, Viewport.Max.Y, 1.0f);
 		if (bUseTile)
 		{
 			SetShaderParameters(RHICmdList, TileVertexShader, TileVertexShader.GetVertexShader(), ParametersVS);
-//			RHICmdList.SetViewport(0, 0, 0.0f, Resolution.X, Resolution.Y, 1.0f);
 			RHICmdList.SetStreamSource(0, nullptr, 0);
 			RHICmdList.DrawPrimitiveIndirect(Parameters->TileData.TileIndirectBuffer->GetRHI(), 0);
 		}
