Change hair pass to use hair uniform buffer.

diff --git a/Engine/Shaders/Private/HairStrands/HairStrandsDebug.usf b/Engine/Shaders/Private/HairStrands/HairStrandsDebug.usf
--- a/Engine/Shaders/Private/HairStrands/HairStrandsDebug.usf
+++ b/Engine/Shaders/Private/HairStrands/HairStrandsDebug.usf
@@ -1,25 +1,18 @@
 // Copyright Epic Games, Inc. All Rights Reserved.
 
+#include "HairStrandsVisibilityCommonStruct.ush"
 #include "../Common.ush"
-#include "HairStrandsVisibilityCommon.ush"
 #include "../ColorMap.ush"
+#include "HairStrandsVisibilityCommon.ush"
 
 float2 OutputResolution;
 uint FastResolveMask;
 uint DebugMode;
 int SampleIndex;
-uint MaxSampleCount;
 
-Texture2D<uint4> CategorizationTexture;
-Texture2D<float> ViewTransmittanceTexture;
 Texture2D<uint2> DepthStencilTexture;
-Texture2D<uint> NodeIndex;
-StructuredBuffer<FPackedHairSample> NodeData;
 SamplerState LinearSampler;
 
-Texture2D<uint> HairCountUintTexture;
-Texture2D<float> HairCountTexture;
-
 #define DEBUG_MODE_SAMPLE_COUNT 0 
 #define DEBUG_MODE_COVERAGE_TYPE 1
 #define DEBUG_MODE_TAA_RESOLVE_TYPE 2
@@ -38,35 +31,24 @@ void MainPS(
 	const float2 UV = Input.UV;
 	const float2 PixelPos = UV * OutputResolution;
 
-	//if (DebugMode == DEBUG_MODE_SAMPLE_COUNT)
-	//{
-	//	const float HairCount = HairCountUintTexture.Load(uint3(PixelPos, 0)) / 1000.f;
-	//	//const float HairCount = HairCountTexture.Load(uint3(PixelPos, 0));
-	//
-	//	if (HairCount==0)
-	//		discard;
-	//	OutColor = float4(float3(1, 0, 0), 1);
-	//
-	//	OutColor = float4(GetHSVDebugColor(HairCount / float(32)), 1);
-	//
-	//	return;
-	//}
-
-	const FCategorizationData CatData = DecodeCategorizationData(CategorizationTexture.Load(uint3(PixelPos,0)));
-	const bool bIsValid = CatData.PixelCoverage > 0 || DebugMode == DEBUG_MODE_TILE;
+	const FCategorizationData CatData = DecodeCategorizationData(HairStrands.HairCategorizationTexture.Load(uint3(PixelPos,0)));
+	const float HairPixelCoverage = CatData.PixelCoverage;
+	const bool bIsValid = HairPixelCoverage > 0 || DebugMode == DEBUG_MODE_TILE;
 	if (!bIsValid)
 		discard;
 
+	const FNodeDesc NodeDesc = DecodeNodeDesc(HairStrands.HairSampleOffset.Load(uint3(PixelPos, 0)));
+
 	float3 Color = 0;
 	float Alpha = 1;
 	if (DebugMode == DEBUG_MODE_SAMPLE_COUNT)
 	{
-		Color = GetHSVDebugColor(CatData.SampleCount/float(MaxSampleCount));
+		Color = GetHSVDebugColor(NodeDesc.Count/float(HairStrands.MaxSamplePerPixelCount));
 	}
 
 	if (DebugMode == DEBUG_MODE_COVERAGE_TYPE)
 	{
-		const bool bIsFullyCovered = CatData.PixelCoverage >= 1;
+		const bool bIsFullyCovered = HairPixelCoverage >= 1;
 		Color = bIsFullyCovered ? float3(0.25,0.85,0.25) : float3(0.85,0.25,0.25);
 	}
 
@@ -79,7 +61,7 @@ void MainPS(
 
 	if (DebugMode == DEBUG_MODE_COVERAGE)
 	{
-		Color = GetHSVDebugColor(saturate(CatData.PixelCoverage));
+		Color = GetHSVDebugColor(saturate(HairPixelCoverage));
 	}
 
 	const bool bIsDebugModeMaterial =
@@ -89,7 +71,6 @@ void MainPS(
 		DebugMode == DEBUG_MODE_SPECULAR || 
 		DebugMode == DEBUG_MODE_TANGENT;
 
-	const FNodeDesc NodeDesc = DecodeNodeDesc(NodeIndex.Load(uint3(PixelPos, 0)));
 	if (NodeDesc.Count > 0 && bIsDebugModeMaterial)
 	{
 		const int LocalSampleIndex = SampleIndex; 
@@ -104,7 +85,7 @@ void MainPS(
 		for (uint NodeIt = 0; NodeIt < NodeDesc.Count; ++NodeIt)
 		{
 			const uint NodeOffset = NodeDesc.Offset + NodeIt;
-			const FPackedHairSample Data = NodeData[NodeOffset];
+			const FPackedHairSample Data = HairStrands.HairSampleData[NodeOffset];
 			const FHairSample Sample = UnpackHairSample(Data);
 
 			if (bAverageAllSamples)
diff --git a/Engine/Shaders/Private/HairStrands/HairStrandsDebugPrint.usf b/Engine/Shaders/Private/HairStrands/HairStrandsDebugPrint.usf
--- a/Engine/Shaders/Private/HairStrands/HairStrandsDebugPrint.usf
+++ b/Engine/Shaders/Private/HairStrands/HairStrandsDebugPrint.usf
@@ -1,5 +1,7 @@
 // Copyright Epic Games, Inc. All Rights Reserved.
 
+#include "HairStrandsVisibilityCommonStruct.ush"
+
 #include "../Common.ush"
 #include "../ShaderPrintCommon.ush"
 #include "../ShaderDrawDebug.ush"
@@ -10,17 +12,14 @@ int2 PixelCoord;
 int2 MaxResolution;
 uint FastResolveMask;
 uint HairMacroGroupCount;
-uint MaxSampleCount;
 uint2 GroupSize; 
 uint HairVisibilityNodeGroupSize;
 
 Texture2D<float> HairCountTexture;
 Texture2D<uint> HairCountUintTexture;
-Texture2D<uint4> CategorizationTexture;
+
 Texture2D<uint2> DepthStencilTexture;
 SamplerState LinearSampler;
-Texture2D<uint> HairVisibilityNodeOffsetAndCount;
-StructuredBuffer<FPackedHairSample> HairVisibilityNodeData;
 Buffer<uint> HairVisibilityIndirectArgsBuffer;
 
 Buffer<int> HairMacroGroupAABBBuffer;
@@ -46,7 +45,6 @@ void MainCS(uint GroupIndex : SV_GroupIndex, uint3 DispatchThreadId : SV_Dispatc
 			AddAABB(Bound.Min, Bound.Max, InstanceBoundColor);
 		}
 
-
 		const uint DispatchX = HairVisibilityIndirectArgsBuffer[0];
 		const uint DispatchY = HairVisibilityIndirectArgsBuffer[1];
 		const uint TotalSampleCount = DispatchX * DispatchY * HairVisibilityNodeGroupSize;
@@ -57,6 +55,7 @@ void MainCS(uint GroupIndex : SV_GroupIndex, uint3 DispatchThreadId : SV_Dispatc
 		float2 Pos = OriginalPos;
 		Pos = ShaderPrintText(Pos, TEXT("Total Sample Count "), FontCyan);
 		Pos = ShaderPrint(Pos, TotalSampleCount, FontCyan);
+		Pos = ShaderPrint(Pos, HairStrands.HairSampleCount.Load(uint3(0,0,0)), FontCyan);
 		AddNewLine(OriginalPos, Pos);
 
 		Pos = ShaderPrintText(Pos, TEXT("Equivalent Res     "), FontCyan);
@@ -74,8 +73,9 @@ void MainCS(uint GroupIndex : SV_GroupIndex, uint3 DispatchThreadId : SV_Dispatc
 
 	ShaderPrintFilter(true);
 
-	const FCategorizationData CatData = DecodeCategorizationData(CategorizationTexture.Load(uint3(PixelPos,0)));
-	const bool bIsValid = CatData.PixelCoverage > 0;
+	const FCategorizationData CatData = DecodeCategorizationData(HairStrands.HairCategorizationTexture.Load(uint3(PixelPos,0)));
+	const float HairPixelCoverage = CatData.PixelCoverage;
+	const bool bIsValid = HairPixelCoverage > 0;
 	if (!bIsValid)
 		return;
 
@@ -87,11 +87,11 @@ void MainCS(uint GroupIndex : SV_GroupIndex, uint3 DispatchThreadId : SV_Dispatc
 	float2 Pos = OriginalPos;
 
 	uint Total8BitCoverage = 0;
-	const FNodeDesc NodeDesc = DecodeNodeDesc(HairVisibilityNodeOffsetAndCount.Load(uint3(PixelCoord, 0)));
+	const FNodeDesc NodeDesc = DecodeNodeDesc(HairStrands.HairSampleOffset.Load(uint3(PixelCoord, 0)));
 	for (uint SampleIt0 = 0; SampleIt0 < NodeDesc.Count; SampleIt0++)
 	{
 		const uint LocalOffset = NodeDesc.Offset + SampleIt0;
-		const FPackedHairSample NodeData = HairVisibilityNodeData[LocalOffset];
+		const FPackedHairSample NodeData = HairStrands.HairSampleData[LocalOffset];
 		const FHairSample Sample = UnpackHairSample(NodeData);
 		Total8BitCoverage += Sample.Coverage8bit;
 	}
@@ -103,12 +103,12 @@ void MainCS(uint GroupIndex : SV_GroupIndex, uint3 DispatchThreadId : SV_Dispatc
 
 	// Max sample
 	Pos = ShaderPrintText(Pos, TEXT("Max Sample      "));
-	Pos = ShaderPrint(Pos, MaxSampleCount);
+	Pos = ShaderPrint(Pos, HairStrands.MaxSamplePerPixelCount);
 	AddNewLine(OriginalPos, Pos);
 
 	// Unique sample
 	Pos = ShaderPrintText(Pos, TEXT("Sample          "), FontOrange);
-	Pos = ShaderPrint(Pos, CatData.SampleCount, FontOrange);
+	Pos = ShaderPrint(Pos, NodeDesc.Count, FontOrange);
 	AddNewLine(OriginalPos, Pos);
 
 	// Accurate coverage (if coverage pass is activated)
@@ -123,13 +123,7 @@ void MainCS(uint GroupIndex : SV_GroupIndex, uint3 DispatchThreadId : SV_Dispatc
 	
 	// Accurate coverage (if coverage pass is activated)
 	Pos = ShaderPrintText(Pos, TEXT("Coverage        "), FontOrange);
-	Pos = ShaderPrint(Pos, CatData.PixelCoverage, FontOrange);
-	AddNewLine(OriginalPos, Pos);
-
-	// 8 bit Coverage (cateogory texture)
-	Pos = ShaderPrintText(Pos, TEXT("Coverage 8bits  "));
-	Pos = ShaderPrint(Pos, CatData.TotalCoverage8bit);
-	Pos = ShaderPrint(Pos, CatData.TotalCoverage8bit / 255.f);
+	Pos = ShaderPrint(Pos, HairPixelCoverage, FontOrange);
 	AddNewLine(OriginalPos, Pos);
 
 	// 8 bit Coverage (summed)
@@ -191,8 +185,7 @@ void MainCS(uint GroupIndex : SV_GroupIndex, uint3 DispatchThreadId : SV_Dispatc
 	for (uint SampleIt = 0; SampleIt < NodeDesc.Count; SampleIt++)
 	{
 		const uint LocalOffset = NodeDesc.Offset + SampleIt;
-		const FPackedHairSample NodeData = HairVisibilityNodeData[LocalOffset];		
-		const FHairSample Sample = UnpackHairSample(NodeData);
+		const FHairSample Sample = UnpackHairSample(HairStrands.HairSampleData[LocalOffset]);
 
 		float2 StartSamplePos = ColumnPos + float2(150.f + 80.f * SampleIt, 0.f) / float2(MaxResolution);
 		Pos = StartSamplePos;
diff --git a/Engine/Shaders/Private/HairStrands/HairStrandsEnvironmentAO.usf b/Engine/Shaders/Private/HairStrands/HairStrandsEnvironmentAO.usf
--- a/Engine/Shaders/Private/HairStrands/HairStrandsEnvironmentAO.usf
+++ b/Engine/Shaders/Private/HairStrands/HairStrandsEnvironmentAO.usf
@@ -1,5 +1,6 @@
 // Copyright Epic Games, Inc. All Rights Reserved.
   
+#include "HairStrandsVisibilityCommonStruct.ush"
 #include "HairStrandsVoxelPageCommonStruct.ush"
 
 #include "../Common.ush"
@@ -150,15 +151,15 @@ void MainPS(
 	bool bIsValid = false;
 	if (bHalfRes)
 	{
-		float Coverage0 = DecodeCategorizationData(HairCategorizationTexture.Load(uint3(PixelCoord              , 0))).PixelCoverage;
-		float Coverage1 = DecodeCategorizationData(HairCategorizationTexture.Load(uint3(PixelCoord + uint2(1, 0), 0))).PixelCoverage;
-		float Coverage2 = DecodeCategorizationData(HairCategorizationTexture.Load(uint3(PixelCoord + uint2(0, 1), 0))).PixelCoverage;
-		float Coverage3 = DecodeCategorizationData(HairCategorizationTexture.Load(uint3(PixelCoord + uint2(1, 1), 0))).PixelCoverage;
+		float Coverage0 = DecodeCategorizationData(HairStrands.HairCategorizationTexture.Load(uint3(PixelCoord              , 0))).PixelCoverage;
+		float Coverage1 = DecodeCategorizationData(HairStrands.HairCategorizationTexture.Load(uint3(PixelCoord + uint2(1, 0), 0))).PixelCoverage;
+		float Coverage2 = DecodeCategorizationData(HairStrands.HairCategorizationTexture.Load(uint3(PixelCoord + uint2(0, 1), 0))).PixelCoverage;
+		float Coverage3 = DecodeCategorizationData(HairStrands.HairCategorizationTexture.Load(uint3(PixelCoord + uint2(1, 1), 0))).PixelCoverage;
 		bIsValid = Coverage0 < 1 || Coverage1 < 1 || Coverage2 < 1 || Coverage3 < 1 || GBuffer.ShadingModelID != SHADINGMODELID_UNLIT;
 	}
 	else
 	{
-		bIsValid = DecodeCategorizationData(HairCategorizationTexture.Load(uint3(PixelCoord, 0))).PixelCoverage < 1;
+		bIsValid = DecodeCategorizationData(HairStrands.HairCategorizationTexture.Load(uint3(PixelCoord, 0))).PixelCoverage < 1;
 	}
 
 	if (bIsValid)
diff --git a/Engine/Shaders/Private/HairStrands/HairStrandsGBufferWrite.usf b/Engine/Shaders/Private/HairStrands/HairStrandsGBufferWrite.usf
--- a/Engine/Shaders/Private/HairStrands/HairStrandsGBufferWrite.usf
+++ b/Engine/Shaders/Private/HairStrands/HairStrandsGBufferWrite.usf
@@ -1,5 +1,6 @@
 // Copyright Epic Games, Inc. All Rights Reserved.
 
+#include "HairStrandsVisibilityCommonStruct.ush"
 #include "../Common.ush"
 
 // Note, this doesn't work after refactoring, undoing the workaround.
@@ -15,9 +16,6 @@
 #include "HairStrandsVisibilityCommon.ush"
 
 uint bWriteDummyData;
-Texture2D<uint4> CategorizationTexture;
-Texture2D<uint>	 NodeIndex;
-StructuredBuffer<FPackedHairSample> NodeData;
 
 void MainPS(
 	in float2 UV : TEXCOORD
@@ -43,19 +41,12 @@ void MainPS(
 #endif
 
 	const uint3 PixelCoord = uint3(floor(SvPosition.xy), 0);
-	FNodeDesc NodeDesc = DecodeNodeDesc(NodeIndex.Load(PixelCoord));
+	const FNodeDesc NodeDesc = DecodeNodeDesc(HairStrands.HairSampleOffset.Load(PixelCoord));
 	if (NodeDesc.Count == 0)
 	{
 		discard;
 	}
 
-	const FCategorizationData CatData = DecodeCategorizationData(CategorizationTexture.Load(PixelCoord));
-	const float PixelCoverage = min(CatData.PixelCoverage, 1);
-	if (PixelCoverage == 0)
-	{
-		discard;
-	}
-
 	if (bWriteDummyData > 0)
 	{
 		const float3 Tangent = float3(0, 0, 1);
@@ -75,7 +66,7 @@ void MainPS(
 		for (uint NodeIt = 0; NodeIt < NodeDesc.Count; ++NodeIt)
 		{
 			const uint NodeOffset = NodeDesc.Offset + NodeIt;
-			const FPackedHairSample Data = NodeData[NodeOffset];
+			const FPackedHairSample Data = HairStrands.HairSampleData[NodeOffset];
 			const FHairSample Sample = UnpackHairSample(Data);
 
 			Samples_Tangent		+= Sample.Tangent;
@@ -88,7 +79,7 @@ void MainPS(
 				Samples_LightChannelMask = Sample.LightChannelMask;
 			}
 		}
-		Samples_Depth = CatData.ClosestDepth;
+		Samples_Depth = HairStrands.HairOnlyDepthTexture.Load(PixelCoord).x;
 
 		const float InvSampleCount = 1.f / max(1u, NodeDesc.Count);
 		Samples_Tangent		 = normalize(Samples_Tangent);
diff --git a/Engine/Shaders/Private/HairStrands/HairStrandsVisibilityComposeSubPixelPS.usf b/Engine/Shaders/Private/HairStrands/HairStrandsVisibilityComposeSubPixelPS.usf
--- a/Engine/Shaders/Private/HairStrands/HairStrandsVisibilityComposeSubPixelPS.usf
+++ b/Engine/Shaders/Private/HairStrands/HairStrandsVisibilityComposeSubPixelPS.usf
@@ -11,7 +11,6 @@
 #if SHADER_SUBCOLOR
 
 Texture2D<float4> SubPixelColorTexture;
-Texture2D<uint4> CategorisationTexture;
 
 void SubColorPS(
 	in FScreenVertexOutput Input,
@@ -25,8 +24,7 @@ void SubColorPS(
 	//if (OutColor.a == 0)
 	//	discard;
 
-	const FCategorizationData CategorisationData = DecodeCategorizationData(CategorisationTexture.Load(uint3(Input.Position.xy, 0)));
-	OutDepth = CategorisationData.ClosestDepth;
+	OutDepth = HairStrands.HairOnlyDepthTexture.Load(uint3(Input.Position.xy, 0));
 }
 
 #endif // SHADER_SUBCOLOR
@@ -142,7 +140,7 @@ void ComposeSamplePS(
 	OutColor = 0.0f;
 	const uint3 PixelCoord = uint3(floor(Input.Position.xy), 0);
 
-	FNodeDesc NodeDesc = DecodeNodeDesc(HairStrands.HairSampleOffset.Load(PixelCoord));
+	const FNodeDesc NodeDesc = DecodeNodeDesc(HairStrands.HairSampleOffset.Load(PixelCoord));
 	if (NodeDesc.Count == 0)
 	{
 		discard;
@@ -156,6 +154,7 @@ void ComposeSamplePS(
 	const uint TotalNodeCount = HairStrands.HairSampleCount.Load(uint3(0, 0, 0));
 	const uint Resolution = uint(ceil(sqrt(float(TotalNodeCount))));
 
+	const float HairDeviceZ = HairStrands.HairOnlyDepthTexture.Load(PixelCoord);
 	const float3 ClosetPointWorldPosition = SvPositionToWorld(float4(Input.Position.xy, CatData.ClosestDepth, 1.0));;
 	const float4 Volumetric = EvaluateVolumetric(ClosetPointWorldPosition);
 
@@ -173,7 +172,7 @@ void ComposeSamplePS(
 
 	OutColor.rgb *= PixelCoverage;
 	OutColor.a = PixelCoverage;
-	OutDepth = CatData.ClosestDepth;
+	OutDepth = HairDeviceZ;
 	if (bComposeDofDepth)
 	{
 		OutDepth = HairDOFDepthTexture.Load(PixelCoord);
@@ -189,7 +188,6 @@ void ComposeSamplePS(
 
 // Included from DeferredLightingCommon now
 //Texture2D<float>	SceneDepthTexture;
-Texture2D<uint4>	HairCategorizationTexture;
 
 void DOFDepthPS(
 	in FScreenVertexOutput Input,
@@ -198,7 +196,7 @@ void DOFDepthPS(
 	OutTexture = 0.0f;
 	const uint3 PixelCoord = uint3(floor(Input.Position.xy), 0);
 
-	const FCategorizationData CatData = DecodeCategorizationData(HairCategorizationTexture.Load(PixelCoord));
+	const FCategorizationData CatData = DecodeCategorizationData(HairStrands.HairCategorizationTexture.Load(PixelCoord));
 	const float PixelCoverage = min(CatData.PixelCoverage, 1);
 	if (PixelCoverage == 0)
 	{
@@ -222,7 +220,7 @@ void DOFDepthPS(
 	float HairDepth = 0;
 	float HairCoC = 0; // Radius in pixel
 	{
-		HairDeviceZ = CatData.ClosestDepth;
+		HairDeviceZ = HairStrands.HairOnlyDepthTexture.Load(PixelCoord).x;
 		HairDepth = ConvertFromDeviceZ(HairDeviceZ);
 		HairDeviceZ = 1 - HairDeviceZ;
 		HairCoC = DepthToCoc(HairDepth) * HalfResToFullRes;
diff --git a/Engine/Source/Runtime/Renderer/Private/HairStrands/HairStrandsComposition.cpp b/Engine/Source/Runtime/Renderer/Private/HairStrands/HairStrandsComposition.cpp
--- a/Engine/Source/Runtime/Renderer/Private/HairStrands/HairStrandsComposition.cpp
+++ b/Engine/Source/Runtime/Renderer/Private/HairStrands/HairStrandsComposition.cpp
@@ -237,12 +237,10 @@ class FHairDOFDepthPS : public FGlobalShader
 	BEGIN_SHADER_PARAMETER_STRUCT(FParameters, )
 		SHADER_PARAMETER_STRUCT_REF(FViewUniformShaderParameters, ViewUniformBuffer)
 		SHADER_PARAMETER_STRUCT_INCLUDE(FHairStrandsTilePassVS::FParameters, TileData)
-		SHADER_PARAMETER_RDG_TEXTURE(Texture2D, HairSampleCount)
-		SHADER_PARAMETER_RDG_TEXTURE(Texture2D, HairCategorizationTexture)
-		SHADER_PARAMETER_RDG_TEXTURE(Texture2D, HairVisibilityNodeOffsetAndCount)
 		SHADER_PARAMETER_RDG_TEXTURE(Texture2D, HairLightingSampleBuffer)
 		SHADER_PARAMETER_RDG_TEXTURE(Texture2D, SceneColorTexture)
 		SHADER_PARAMETER_RDG_TEXTURE(Texture2D, SceneDepthTexture)
+		SHADER_PARAMETER_RDG_UNIFORM_BUFFER(FHairStrandsViewUniformParameters, HairStrands)
 		RENDER_TARGET_BINDING_SLOTS()
 	END_SHADER_PARAMETER_STRUCT()
 
@@ -276,9 +274,7 @@ static FRDGTextureRef AddHairDOFDepthPass(
 	}
 
 	FHairDOFDepthPS::FParameters* Parameters = GraphBuilder.AllocParameters<FHairDOFDepthPS::FParameters>();
-	Parameters->HairSampleCount = VisibilityData.NodeCount;
-	Parameters->HairCategorizationTexture = CategorizationTexture;
-	Parameters->HairVisibilityNodeOffsetAndCount = VisibilityData.NodeIndex;
+	Parameters->HairStrands = View.HairStrandsViewData.UniformBuffer;
 	Parameters->HairLightingSampleBuffer = VisibilityData.SampleLightingBuffer;
 	Parameters->SceneColorTexture = InColorTexture;
 	Parameters->SceneDepthTexture = InDepthTexture;
@@ -452,9 +448,7 @@ class FHairVisibilityGBufferWritePS : public FGlobalShader
 	BEGIN_SHADER_PARAMETER_STRUCT(FParameters, )
 		SHADER_PARAMETER_STRUCT_INCLUDE(FHairStrandsTilePassVS::FParameters, TileData)
 		SHADER_PARAMETER(uint32, bWriteDummyData)
-		SHADER_PARAMETER_RDG_TEXTURE(Texture2D, CategorizationTexture)
-		SHADER_PARAMETER_RDG_TEXTURE(Texture2D, NodeIndex)
-		SHADER_PARAMETER_RDG_BUFFER_SRV(StructuredBuffer, NodeData)
+		SHADER_PARAMETER_RDG_UNIFORM_BUFFER(FHairStrandsViewUniformParameters, HairStrands)
 		RENDER_TARGET_BINDING_SLOTS()
 	END_SHADER_PARAMETER_STRUCT()
 
@@ -483,9 +477,6 @@ static void AddHairVisibilityGBufferWritePass(
 	FRDGBuilder& GraphBuilder,
 	const FViewInfo& View,
 	const bool bWriteDummyData,
-	const FRDGTextureRef& CategorizationTexture,
-	const FRDGTextureRef& NodeIndex,
-	const FRDGBufferRef& NodeData,
 	const FHairStrandsTiles& TileData,
 	FRDGTextureRef OutGBufferATexture,
 	FRDGTextureRef OutGBufferBTexture,
@@ -509,9 +500,7 @@ static void AddHairVisibilityGBufferWritePass(
 
 	FHairVisibilityGBufferWritePS::FParameters* Parameters = GraphBuilder.AllocParameters<FHairVisibilityGBufferWritePS::FParameters>();
 	Parameters->bWriteDummyData = bWriteDummyData ? 1 : 0;
-	Parameters->CategorizationTexture = CategorizationTexture;
-	Parameters->NodeIndex = NodeIndex;
-	Parameters->NodeData = GraphBuilder.CreateSRV(NodeData);
+	Parameters->HairStrands = View.HairStrandsViewData.UniformBuffer;
 	Parameters->RenderTargets[0] = FRenderTargetBinding(OutGBufferATexture, ERenderTargetLoadAction::ELoad);
 	Parameters->RenderTargets[1] = FRenderTargetBinding(OutGBufferBTexture, ERenderTargetLoadAction::ELoad);	
 	if (bWriteFullGBuffer)
@@ -629,9 +618,6 @@ static void InternalRenderHairComposition(
 							GraphBuilder,
 							View,
 							bWriteDummyData,
-							VisibilityData.CategorizationTexture,
-							VisibilityData.NodeIndex,
-							VisibilityData.NodeData,
 							VisibilityData.TileData,
 							GBufferATexture,
 							GBufferBTexture,
@@ -646,9 +632,6 @@ static void InternalRenderHairComposition(
 							GraphBuilder,
 							View,
 							bWriteDummyData,
-							VisibilityData.CategorizationTexture,
-							VisibilityData.NodeIndex,
-							VisibilityData.NodeData,
 							VisibilityData.TileData,
 							GBufferATexture,
 							GBufferBTexture,
diff --git a/Engine/Source/Runtime/Renderer/Private/HairStrands/HairStrandsData.h b/Engine/Source/Runtime/Renderer/Private/HairStrands/HairStrandsData.h
--- a/Engine/Source/Runtime/Renderer/Private/HairStrands/HairStrandsData.h
+++ b/Engine/Source/Runtime/Renderer/Private/HairStrands/HairStrandsData.h
@@ -19,8 +19,8 @@ struct FMeshBatch;
 // HairStrands uniform buffer
 
 BEGIN_GLOBAL_SHADER_PARAMETER_STRUCT(FHairStrandsViewUniformParameters, )	
-//	SHADER_PARAMETER(uint32, SampleCount)
 	SHADER_PARAMETER(FIntPoint, HairTileCountXY)										// Tile count in X/Y
+	SHADER_PARAMETER(uint32, MaxSamplePerPixelCount)									// Max number of sample per pixel
 	SHADER_PARAMETER(float, HairDualScatteringRoughnessOverride)						// Override the roughness used for dual scattering (for hack/test purpose only)
 	SHADER_PARAMETER(FIntPoint, HairSampleViewportResolution)							// Maximum viewport resolution of the sample space
 	SHADER_PARAMETER(uint32, bHairTileValid)											// True if tile data are valid
@@ -85,8 +85,8 @@ struct FHairStrandsVisibilityData
 
 	FRDGTextureRef LightChannelMaskTexture = nullptr;
 
-	uint32			MaxSampleCount = 8;
-	uint32			MaxNodeCount = 0;
+	uint32			MaxSampleCount = 8; // Sample count per pixel
+	uint32			MaxNodeCount = 0;	// Total sample count
 	FRDGTextureRef	NodeCount = nullptr;
 	FRDGTextureRef	NodeIndex = nullptr;
 	FRDGBufferRef	NodeData = nullptr;
diff --git a/Engine/Source/Runtime/Renderer/Private/HairStrands/HairStrandsDebug.cpp b/Engine/Source/Runtime/Renderer/Private/HairStrands/HairStrandsDebug.cpp
--- a/Engine/Source/Runtime/Renderer/Private/HairStrands/HairStrandsDebug.cpp
+++ b/Engine/Source/Runtime/Renderer/Private/HairStrands/HairStrandsDebug.cpp
@@ -249,13 +249,9 @@ class FHairDebugPrintCS : public FGlobalShader
 		SHADER_PARAMETER(FIntPoint, MaxResolution)
 		SHADER_PARAMETER(uint32, FastResolveMask)
 		SHADER_PARAMETER(uint32, HairMacroGroupCount)
-		SHADER_PARAMETER(uint32, MaxSampleCount)
 		SHADER_PARAMETER(uint32, HairVisibilityNodeGroupSize)
 		SHADER_PARAMETER_RDG_TEXTURE(Texture2D, HairCountTexture)
 		SHADER_PARAMETER_RDG_TEXTURE(Texture2D, HairCountUintTexture)
-		SHADER_PARAMETER_RDG_TEXTURE(Texture2D, CategorizationTexture)
-		SHADER_PARAMETER_RDG_TEXTURE(Texture2D, HairVisibilityNodeOffsetAndCount)
-		SHADER_PARAMETER_RDG_BUFFER_SRV(StructuredBuffer, HairVisibilityNodeData)
 		SHADER_PARAMETER_RDG_BUFFER_SRV(Buffer<uint>, HairVisibilityIndirectArgsBuffer)
 		SHADER_PARAMETER_RDG_BUFFER_SRV(Buffer, HairMacroGroupAABBBuffer)
 		SHADER_PARAMETER_RDG_TEXTURE_SRV(Texture2D, DepthStencilTexture)
@@ -263,6 +259,7 @@ class FHairDebugPrintCS : public FGlobalShader
 		SHADER_PARAMETER_STRUCT_REF(FViewUniformShaderParameters, ViewUniformBuffer)
 		SHADER_PARAMETER_STRUCT_INCLUDE(ShaderPrint::FShaderParameters, ShaderPrintUniformBuffer)
 		SHADER_PARAMETER_STRUCT_INCLUDE(ShaderDrawDebug::FShaderParameters, ShaderDrawUniformBuffer)
+		SHADER_PARAMETER_RDG_UNIFORM_BUFFER(FHairStrandsViewUniformParameters, HairStrands)
 	END_SHADER_PARAMETER_STRUCT()
 
 public:
@@ -285,7 +282,7 @@ static void AddDebugHairPrintPass(
 	const FHairStrandsMacroGroupResources& MacroGroupResources,
 	FRDGTextureSRVRef InDepthStencilTexture)
 {
-	if (!VisibilityData.CategorizationTexture || !VisibilityData.NodeIndex || !VisibilityData.NodeData || !InDepthStencilTexture || !ShaderDrawDebug::IsEnabled(*View)) return;
+	if (!View || !View->HairStrandsViewData.UniformBuffer || !InDepthStencilTexture || !ShaderDrawDebug::IsEnabled(*View)) return;
 
 	FRDGTextureRef ViewHairCountTexture = VisibilityData.ViewHairCountTexture ? VisibilityData.ViewHairCountTexture : GSystemTextures.GetBlackDummy(GraphBuilder);
 	FRDGTextureRef ViewHairCountUintTexture = VisibilityData.ViewHairCountUintTexture ? VisibilityData.ViewHairCountUintTexture : GSystemTextures.GetBlackDummy(GraphBuilder);
@@ -298,17 +295,14 @@ static void AddDebugHairPrintPass(
 	Parameters->MaxResolution = VisibilityData.CategorizationTexture ? VisibilityData.CategorizationTexture->Desc.Extent : FIntPoint(0,0);
 	Parameters->PixelCoord = View->CursorPos;
 	Parameters->FastResolveMask = STENCIL_TEMPORAL_RESPONSIVE_AA_MASK;
-	Parameters->CategorizationTexture = VisibilityData.CategorizationTexture;
 	Parameters->HairCountTexture = ViewHairCountTexture;
 	Parameters->HairCountUintTexture = ViewHairCountUintTexture;
-	Parameters->HairVisibilityNodeData = GraphBuilder.CreateSRV(VisibilityData.NodeData);
-	Parameters->HairVisibilityNodeOffsetAndCount = VisibilityData.NodeIndex;
 	Parameters->HairVisibilityIndirectArgsBuffer = GraphBuilder.CreateSRV(VisibilityData.NodeIndirectArg, PF_R32_UINT);
 	Parameters->HairVisibilityNodeGroupSize = VisibilityData.NodeGroupSize;
 	Parameters->DepthStencilTexture = InDepthStencilTexture;
 	Parameters->LinearSampler = TStaticSamplerState<SF_Bilinear, AM_Clamp, AM_Clamp, AM_Clamp>::GetRHI();
 	Parameters->HairMacroGroupCount = MacroGroupResources.MacroGroupCount;
-	Parameters->MaxSampleCount = VisibilityData.MaxSampleCount;
+	Parameters->HairStrands = View->HairStrandsViewData.UniformBuffer;
 	Parameters->HairMacroGroupAABBBuffer = GraphBuilder.CreateSRV(MacroGroupResources.MacroGroupAABBsBuffer, PF_R32_SINT);
 	ShaderPrint::SetParameters(GraphBuilder, *View, Parameters->ShaderPrintUniformBuffer);
 	ShaderDrawDebug::SetParameters(GraphBuilder, View->ShaderDrawData, Parameters->ShaderDrawUniformBuffer);
@@ -338,12 +332,10 @@ class FHairDebugPS : public FGlobalShader
 		SHADER_PARAMETER(uint32, MaxSampleCount)
 		SHADER_PARAMETER_RDG_TEXTURE(Texture2D, HairCountTexture)
 		SHADER_PARAMETER_RDG_TEXTURE(Texture2D, HairCountUintTexture)
-		SHADER_PARAMETER_RDG_TEXTURE(Texture2D, CategorizationTexture)
-		SHADER_PARAMETER_RDG_TEXTURE(Texture2D, NodeIndex)
-		SHADER_PARAMETER_RDG_BUFFER_SRV(StructuredBuffer, NodeData)
 		SHADER_PARAMETER_RDG_TEXTURE_SRV(Texture2D, DepthStencilTexture)
 		SHADER_PARAMETER_SAMPLER(SamplerState, LinearSampler)
 		SHADER_PARAMETER_STRUCT_REF(FViewUniformShaderParameters, ViewUniformBuffer)
+		SHADER_PARAMETER_RDG_UNIFORM_BUFFER(FHairStrandsViewUniformParameters, HairStrands)
 		RENDER_TARGET_BINDING_SLOTS()
 	END_SHADER_PARAMETER_STRUCT()
 
@@ -375,9 +367,6 @@ static void AddDebugHairPass(
 	if (!VisibilityData.CategorizationTexture || !VisibilityData.NodeIndex || !VisibilityData.NodeData) return;
 	if (InDebugMode == EHairDebugMode::TAAResolveType && !InDepthStencilTexture) return;
 
-	FRDGTextureRef HairCountTexture = VisibilityData.ViewHairCountTexture ? VisibilityData.ViewHairCountTexture : GSystemTextures.GetBlackDummy(GraphBuilder);
-	FRDGTextureRef HairCountUintTexture = VisibilityData.ViewHairCountUintTexture ? VisibilityData.ViewHairCountUintTexture : GSystemTextures.GetBlackDummy(GraphBuilder); 
-
 	const FIntRect Viewport = View->ViewRect;
 	const FIntPoint Resolution(Viewport.Width(), Viewport.Height());
 
@@ -399,16 +388,11 @@ static void AddDebugHairPass(
 	Parameters->ViewUniformBuffer = View->ViewUniformBuffer;
 	Parameters->OutputResolution = Resolution;
 	Parameters->FastResolveMask = STENCIL_TEMPORAL_RESPONSIVE_AA_MASK;
-	Parameters->CategorizationTexture = VisibilityData.CategorizationTexture;
-	Parameters->HairCountTexture = HairCountTexture;
-	Parameters->HairCountUintTexture = HairCountUintTexture;
-	Parameters->NodeIndex = VisibilityData.NodeIndex;
-	Parameters->NodeData = GraphBuilder.CreateSRV(VisibilityData.NodeData);
-		Parameters->DepthStencilTexture = InDepthStencilTexture;
+	Parameters->HairStrands = View->HairStrandsViewData.UniformBuffer;
+	Parameters->DepthStencilTexture = InDepthStencilTexture;
 	Parameters->LinearSampler = TStaticSamplerState<SF_Bilinear, AM_Clamp, AM_Clamp, AM_Clamp>::GetRHI();
 	Parameters->DebugMode = InternalDebugMode;
 	Parameters->SampleIndex = GHairStrandsDebugSampleIndex;
-	Parameters->MaxSampleCount = VisibilityData.MaxSampleCount;
 	Parameters->RenderTargets[0] = FRenderTargetBinding(OutTarget, ERenderTargetLoadAction::ELoad, 0);
 	TShaderMapRef<FPostProcessVS> VertexShader(View->ShaderMap);
 
diff --git a/Engine/Source/Runtime/Renderer/Private/HairStrands/HairStrandsEnvironment.cpp b/Engine/Source/Runtime/Renderer/Private/HairStrands/HairStrandsEnvironment.cpp
--- a/Engine/Source/Runtime/Renderer/Private/HairStrands/HairStrandsEnvironment.cpp
+++ b/Engine/Source/Runtime/Renderer/Private/HairStrands/HairStrandsEnvironment.cpp
@@ -111,12 +111,9 @@ class FHairEnvironmentAO : public FGlobalShader
 		SHADER_PARAMETER(FVector2D, Output_InvResolution)
 
 		SHADER_PARAMETER_STRUCT_INCLUDE(FSceneTextureParameters, SceneTextures)
-
-		SHADER_PARAMETER_RDG_TEXTURE(Texture2D, HairCategorizationTexture)
 		SHADER_PARAMETER_STRUCT_REF(FViewUniformShaderParameters, ViewUniformBuffer)
-
 		SHADER_PARAMETER_RDG_UNIFORM_BUFFER(FVirtualVoxelParameters, VirtualVoxel)
-
+		SHADER_PARAMETER_RDG_UNIFORM_BUFFER(FHairStrandsViewUniformParameters, HairStrands)
 		SHADER_PARAMETER_STRUCT_INCLUDE(ShaderDrawDebug::FShaderParameters, ShaderDrawParameters)
 
 		RENDER_TARGET_BINDING_SLOTS()
@@ -128,7 +125,6 @@ IMPLEMENT_GLOBAL_SHADER(FHairEnvironmentAO, "/Engine/Private/HairStrands/HairStr
 static void AddHairStrandsEnvironmentAOPass(
 	FRDGBuilder& GraphBuilder,
 	const FViewInfo& View,
-	const FHairStrandsVisibilityData& VisibilityResources,
 	const FHairStrandsVoxelResources& VoxelResources,
 	const FHairStrandsMacroGroupData& MacroGroupData,
 	FRDGTextureRef Output)
@@ -147,7 +143,7 @@ static void AddHairStrandsEnvironmentAOPass(
 	PassParameters->VirtualVoxel = VoxelResources.UniformBuffer;
 
 	PassParameters->ViewUniformBuffer = View.ViewUniformBuffer;
-	PassParameters->HairCategorizationTexture = VisibilityResources.CategorizationTexture;
+	PassParameters->HairStrands = HairStrands::BindHairStrandsViewUniformParameters(View);
 	const FFinalPostProcessSettings& Settings = View.FinalPostProcessSettings;
 	PassParameters->AO_Power = Settings.AmbientOcclusionPower;
 	PassParameters->AO_Intensity = Settings.AmbientOcclusionIntensity;
@@ -543,6 +539,6 @@ void RenderHairStrandsAmbientOcclusion(
 
 	for (const FHairStrandsMacroGroupData& MacroGroupData : View.HairStrandsViewData.MacroGroupDatas)
 	{
-		AddHairStrandsEnvironmentAOPass(GraphBuilder, View, VisibilityData, VoxelResources, MacroGroupData, InAOTexture);
+		AddHairStrandsEnvironmentAOPass(GraphBuilder, View, VoxelResources, MacroGroupData, InAOTexture);
 	}
 }
diff --git a/Engine/Source/Runtime/Renderer/Private/HairStrands/HairStrandsRendering.cpp b/Engine/Source/Runtime/Renderer/Private/HairStrands/HairStrandsRendering.cpp
--- a/Engine/Source/Runtime/Renderer/Private/HairStrands/HairStrandsRendering.cpp
+++ b/Engine/Source/Runtime/Renderer/Private/HairStrands/HairStrandsRendering.cpp
@@ -22,6 +22,7 @@ static TRDGUniformBufferRef<FHairStrandsViewUniformParameters> InternalCreateHai
 		Parameters->HairSampleCoords = GraphBuilder.CreateSRV(In->NodeCoord, FHairStrandsVisibilityData::NodeCoordFormat);
 		Parameters->HairSampleCount = In->NodeCount;
 		Parameters->HairSampleViewportResolution = In->SampleLightingViewportResolution;
+		Parameters->MaxSamplePerPixelCount = In->MaxSampleCount;
 
 		if (In->EmissiveTexture)
 		{
@@ -83,6 +84,7 @@ static TRDGUniformBufferRef<FHairStrandsViewUniformParameters> InternalCreateHai
 		Parameters->HairSampleViewportResolution = FIntPoint(0, 0);
 		Parameters->HairEmissiveTexture = BlackTexture;
 		Parameters->bHasEmissive = 0u;
+		Parameters->MaxSamplePerPixelCount = 0u;
 
 		Parameters->HairTileData = DummyBufferR32SRV;
 		Parameters->HairTileCount = DummyBufferRG16SRV;
