Change hair categorization texture to contain only hair coverage data, which shrink its size.

diff --git a/Engine/Shaders/Private/HairStrands/HairStrandsDebug.usf b/Engine/Shaders/Private/HairStrands/HairStrandsDebug.usf
--- a/Engine/Shaders/Private/HairStrands/HairStrandsDebug.usf
+++ b/Engine/Shaders/Private/HairStrands/HairStrandsDebug.usf
@@ -31,8 +31,7 @@ void MainPS(
 	const float2 UV = Input.UV;
 	const float2 PixelPos = UV * OutputResolution;
 
-	const FCategorizationData CatData = DecodeCategorizationData(HairStrands.HairCategorizationTexture.Load(uint3(PixelPos,0)));
-	const float HairPixelCoverage = CatData.PixelCoverage;
+	const float HairPixelCoverage = HairStrands.HairCoverageTexture.Load(uint3(PixelPos, 0));
 	const bool bIsValid = HairPixelCoverage > 0 || DebugMode == DEBUG_MODE_TILE;
 	if (!bIsValid)
 		discard;
diff --git a/Engine/Shaders/Private/HairStrands/HairStrandsDebugPrint.usf b/Engine/Shaders/Private/HairStrands/HairStrandsDebugPrint.usf
--- a/Engine/Shaders/Private/HairStrands/HairStrandsDebugPrint.usf
+++ b/Engine/Shaders/Private/HairStrands/HairStrandsDebugPrint.usf
@@ -73,8 +73,7 @@ void MainCS(uint GroupIndex : SV_GroupIndex, uint3 DispatchThreadId : SV_Dispatc
 
 	ShaderPrintFilter(true);
 
-	const FCategorizationData CatData = DecodeCategorizationData(HairStrands.HairCategorizationTexture.Load(uint3(PixelPos,0)));
-	const float HairPixelCoverage = CatData.PixelCoverage;
+	const float HairPixelCoverage = HairStrands.HairCoverageTexture.Load(uint3(PixelPos, 0));
 	const bool bIsValid = HairPixelCoverage > 0;
 	if (!bIsValid)
 		return;
diff --git a/Engine/Shaders/Private/HairStrands/HairStrandsDeepShadowMask.usf b/Engine/Shaders/Private/HairStrands/HairStrandsDeepShadowMask.usf
--- a/Engine/Shaders/Private/HairStrands/HairStrandsDeepShadowMask.usf
+++ b/Engine/Shaders/Private/HairStrands/HairStrandsDeepShadowMask.usf
@@ -343,8 +343,7 @@ void MainPS(
 	const float2 UV = Input.UV;
 	const uint2 PixelCoord = floor(Input.Position.xy);
 	const float SceneDepth = ConvertFromDeviceZ(SceneDepthTexture.Load(uint3(PixelCoord, 0)).x);
-	const FCategorizationData Categorization = DecodeCategorizationData(HairStrands.HairCategorizationTexture.Load(uint3(PixelCoord, 0)));
-	const bool bIsHairPixel = Categorization.PixelCoverage >= 1;
+	const bool bIsHairPixel = HairStrands.HairCoverageTexture.Load(uint3(PixelCoord, 0)) >= 1;
 
 	const float3 WorldPosition		= TotWorldPosition(UV, SceneDepth);
 	const float4 LightSpacePosition = ToLightPosition(WorldPosition);
@@ -447,8 +446,8 @@ void MainPS(
 	const uint2 PixelCoord = floor(Input.Position.xy);
 	const float SceneDepth = ConvertFromDeviceZ(SceneDepthTexture.Load(uint3(PixelCoord, 0)).x);
 
-	const FCategorizationData Categorization = DecodeCategorizationData(HairStrands.HairCategorizationTexture.Load(uint3(PixelCoord, 0)));
-	const bool bIsHairPixel = Categorization.PixelCoverage >= 1;
+	const float HairPixelCoverage = HairStrands.HairCoverageTexture.Load(uint3(PixelCoord, 0));
+	const bool bIsHairPixel = HairPixelCoverage >= 1;
 
 	float3 WorldPosition = TotWorldPosition(UV, SceneDepth);
 
@@ -500,7 +499,7 @@ void MainPS(
 		TraversalSettings.Random = Random;
 		TraversalSettings.PixelRadius = SceneDepth * VirtualVoxel.HairCoveragePixelRadiusAtDepth1;
 		
-		if (bIsValid && RayMarchMaskTexture.Load(uint3(PixelCoord, 0)).r > 0.0f &&  Categorization.PixelCoverage < CoverageThreshold)
+		if (bIsValid && RayMarchMaskTexture.Load(uint3(PixelCoord, 0)).r > 0.0f && HairPixelCoverage < CoverageThreshold)
 		{
 			Result = ComputeHairCountVirtualVoxel(
 				WorldPosition,
diff --git a/Engine/Shaders/Private/HairStrands/HairStrandsEnvironmentAO.usf b/Engine/Shaders/Private/HairStrands/HairStrandsEnvironmentAO.usf
--- a/Engine/Shaders/Private/HairStrands/HairStrandsEnvironmentAO.usf
+++ b/Engine/Shaders/Private/HairStrands/HairStrandsEnvironmentAO.usf
@@ -151,15 +151,15 @@ void MainPS(
 	bool bIsValid = false;
 	if (bHalfRes)
 	{
-		float Coverage0 = DecodeCategorizationData(HairStrands.HairCategorizationTexture.Load(uint3(PixelCoord              , 0))).PixelCoverage;
-		float Coverage1 = DecodeCategorizationData(HairStrands.HairCategorizationTexture.Load(uint3(PixelCoord + uint2(1, 0), 0))).PixelCoverage;
-		float Coverage2 = DecodeCategorizationData(HairStrands.HairCategorizationTexture.Load(uint3(PixelCoord + uint2(0, 1), 0))).PixelCoverage;
-		float Coverage3 = DecodeCategorizationData(HairStrands.HairCategorizationTexture.Load(uint3(PixelCoord + uint2(1, 1), 0))).PixelCoverage;
+		const float Coverage0 = HairStrands.HairCoverageTexture.Load(uint3(PixelCoord              , 0));
+		const float Coverage1 = HairStrands.HairCoverageTexture.Load(uint3(PixelCoord + uint2(1, 0), 0));
+		const float Coverage2 = HairStrands.HairCoverageTexture.Load(uint3(PixelCoord + uint2(0, 1), 0));
+		const float Coverage3 = HairStrands.HairCoverageTexture.Load(uint3(PixelCoord + uint2(1, 1), 0));
 		bIsValid = Coverage0 < 1 || Coverage1 < 1 || Coverage2 < 1 || Coverage3 < 1 || GBuffer.ShadingModelID != SHADINGMODELID_UNLIT;
 	}
 	else
 	{
-		bIsValid = DecodeCategorizationData(HairStrands.HairCategorizationTexture.Load(uint3(PixelCoord, 0))).PixelCoverage < 1;
+		bIsValid = HairStrands.HairCoverageTexture.Load(uint3(PixelCoord, 0)) < 1;
 	}
 
 	if (bIsValid)
diff --git a/Engine/Shaders/Private/HairStrands/HairStrandsEnvironmentLighting.usf b/Engine/Shaders/Private/HairStrands/HairStrandsEnvironmentLighting.usf
--- a/Engine/Shaders/Private/HairStrands/HairStrandsEnvironmentLighting.usf
+++ b/Engine/Shaders/Private/HairStrands/HairStrandsEnvironmentLighting.usf
@@ -712,8 +712,7 @@ float3 ComputeThinTransmission(FGBufferData GBuffer, float2 ScreenPosition, floa
 	const float3 V = -CameraToPixel;
 	const float3 T = GBuffer.WorldNormal;
 
-	const FCategorizationData CatData = DecodeCategorizationData(HairStrands.HairCategorizationTexture.Load(uint3(PixelPosition, 0)));
-	const float PixelCoverage = min(CatData.PixelCoverage, 1);
+	const float PixelCoverage = min(HairStrands.HairCoverageTexture.Load(uint3(PixelPosition, 0)), 1);
 	
 	// The scene color is attenuated propertionnaly to the hair coverage. In order to retrieve the actual lighting buffer 
 	// value, we divide by the hair coverage
diff --git a/Engine/Shaders/Private/HairStrands/HairStrandsVelocity.usf b/Engine/Shaders/Private/HairStrands/HairStrandsVelocity.usf
--- a/Engine/Shaders/Private/HairStrands/HairStrandsVelocity.usf
+++ b/Engine/Shaders/Private/HairStrands/HairStrandsVelocity.usf
@@ -20,7 +20,7 @@ float CoverageThreshold;
 Texture2D<uint> NodeIndex;
 Buffer<float2> NodeVelocity;
 StructuredBuffer<FNodeVis> NodeVis; 
-Texture2D<uint4> CategorizationTexture;
+Texture2D<float> CoverageTexture;
 
 #define VELOCITY_TYPE_NONE 0
 #define VELOCITY_TYPE_AVG 1  
@@ -65,8 +65,8 @@ void MainCS(uint GroupIndex : SV_GroupIndex, uint3 DispatchThreadId : SV_Dispatc
 		float MaxVelocityMagnitude2 = 0;
 		#endif		
 
-		const FCategorizationData CatData = DecodeCategorizationData(CategorizationTexture.Load(uint3(PixelCoord,0)));
-		if (CatData.PixelCoverage >= CoverageThreshold)
+		const float HairPixelCoverage = CoverageTexture.Load(uint3(PixelCoord,0));
+		if (HairPixelCoverage >= CoverageThreshold)
 		{
 			for (uint NodeIt = 0; NodeIt < NodeDesc.Count; ++NodeIt)
 			{
diff --git a/Engine/Shaders/Private/HairStrands/HairStrandsVisibilityCommon.ush b/Engine/Shaders/Private/HairStrands/HairStrandsVisibilityCommon.ush
--- a/Engine/Shaders/Private/HairStrands/HairStrandsVisibilityCommon.ush
+++ b/Engine/Shaders/Private/HairStrands/HairStrandsVisibilityCommon.ush
@@ -224,38 +224,9 @@ uint GetPackedHairSampleCoverage8bit(in FPackedHairSample In)
 }
 
 ////////////////////////////////////////////////////////////////////////////////////////////////////////////////////////////////
-// Hair categorization
+// Hair velocity
 
-// Rename FHairCatetorizationData
-struct FCategorizationData
-{
-	uint  TotalCoverage8bit;
-	uint  SampleCount;
-	float PixelCoverage;
-	float ClosestDepth;
-};
-
-uint4 EncodeCategorizationData(FCategorizationData Data)
-{
-	const uint ClosestDepthUint = asuint(Data.ClosestDepth);
-	const uint ClosestDepthLow  = ClosestDepthUint & 0xFFFF;
-	const uint ClosestDepthHigh	= (ClosestDepthUint >> 16) & 0xFFFF;
-	const uint ComponentX		= min(Data.TotalCoverage8bit, uint(0xFF)) | ((Data.SampleCount & 0xFF)<<8);
-	return uint4(ComponentX, f32tof16(Data.PixelCoverage), ClosestDepthLow, ClosestDepthHigh);
-}
-
-FCategorizationData DecodeCategorizationData(uint4 Data)
-{
-	FCategorizationData Out;
-	Out.TotalCoverage8bit		= Data.x & 0xFF;
-	Out.SampleCount				=(Data.x>>8) & 0xFF;
-	Out.PixelCoverage			= f16tof32(Data.y);
-	const uint ClosestDepthLow	= Data.z & 0xFFFF;
-	const uint ClosestDepthHigh	=(Data.w & 0xFFFF) << 16;
-	Out.ClosestDepth			= asfloat(ClosestDepthHigh | ClosestDepthLow);
-
-	return Out;
-}
+float3 DecodeVelocityFromTexture(float4 EncodedV);
 
 // Return true if the underlying pixel need to run with TAA fast resolve. This is the case when 
 // the velocity is larger than a certain threshold.
diff --git a/Engine/Shaders/Private/HairStrands/HairStrandsVisibilityCompaction.usf b/Engine/Shaders/Private/HairStrands/HairStrandsVisibilityCompaction.usf
--- a/Engine/Shaders/Private/HairStrands/HairStrandsVisibilityCompaction.usf
+++ b/Engine/Shaders/Private/HairStrands/HairStrandsVisibilityCompaction.usf
@@ -32,7 +32,7 @@ RWTexture2D<uint> OutCompactNodeIndex;
 	RWStructuredBuffer<FPackedHairSample> OutCompactNodeData;
 #endif
 RWBuffer<uint2> OutCompactNodeCoord;
-RWTexture2D<uint4> OutCategorizationTexture;
+RWTexture2D<float> OutCoverageTexture;
 
 #ifndef PERMUTATION_VELOCITY
 	#error Velocity permutation needs to be defined
@@ -75,14 +75,14 @@ void MainCS(uint GroupIndex : SV_GroupIndex, uint3 DispatchThreadId : SV_Dispatc
 	uint2 PixelCoord = DispatchThreadId.xy + ResolutionOffset;
 	const uint FirstNodeIndex = PPLLNodeIndex[PixelCoord];
 
-	FCategorizationData Categorization = (FCategorizationData)0;
+	float ClosestDepth = 0; // Inverse-Z
+	float PixelCoverage = 0;
 
 #if PERMUTATION_VELOCITY != VELOCITY_TYPE_NONE
 	float2 AverageVelocity = 0;
 	uint ValidVelocityCount = 0;
 
 	float2 ClosestVelocity = 0;
-	float ClosestDepth = 0;
 
 	float2 MaxVelocity = 0;
 	float MaxVelocityMagnitude2 = 0;
@@ -202,13 +202,12 @@ void MainCS(uint GroupIndex : SV_GroupIndex, uint3 DispatchThreadId : SV_Dispatc
 			uint Coverage16bit = 0;
 			FPPLLNodeData PPLLNode = PPLLNodeData[NodePPLLIndex];
 
-			// Coverage8bit is a weight normalising to 1 the contribution of all the compacted samples. Because later it is weighted by Categorization.PixelCoverage.
+			// Coverage8bit is a weight normalising to 1 the contribution of all the compacted samples. Because later it is weighted by PixelCoverage.
 			// Patch the coverage on the out node
 			const float PatchedCoverage8bit = To8bitCoverage(SortedCoverage[j] / float(ValidPixelSampleTotalCoverage));
 			FPackedHairSample OutNodeData = ConvertToPackedHairSample(PPLLNode, PatchedCoverage8bit);
 
-			Categorization.TotalCoverage8bit += PatchedCoverage8bit;
-			Categorization.ClosestDepth = max(Categorization.ClosestDepth, OutNodeData.Depth); // Inverse-Z
+			ClosestDepth = max(ClosestDepth, OutNodeData.Depth); // Inverse-Z
 
 			// TODOAGREGATE if last, Create FPackedHairSample from aggregated  data
 			OutCompactNodeData[OffsetInStorage + j] = OutNodeData;
@@ -244,7 +243,7 @@ void MainCS(uint GroupIndex : SV_GroupIndex, uint3 DispatchThreadId : SV_Dispatc
 		#endif		
 		}
 
-		Categorization.PixelCoverage = TransmittanceToCoverage(TotalTransmittance, CoverageThreshold);
+		PixelCoverage = TransmittanceToCoverage(TotalTransmittance, CoverageThreshold);
 
 	#if PERMUTATION_VELOCITY != VELOCITY_TYPE_NONE
 		float2 VelocityOutput = 0;
@@ -269,8 +268,7 @@ void MainCS(uint GroupIndex : SV_GroupIndex, uint3 DispatchThreadId : SV_Dispatc
 	NodeDesc.Offset = OffsetInStorage;
 	OutCompactNodeIndex[PixelCoord] = EncodeNodeDesc(NodeDesc);
 
-	Categorization.SampleCount = PixelRenderNodecount;
-	OutCategorizationTexture[PixelCoord] = EncodeCategorizationData(Categorization);
+	OutCoverageTexture[PixelCoord] = PixelCoverage;
 }
 
 #else // PERMUTATION_PPLL or PERMUTATION_VISIBILITY
@@ -468,7 +466,7 @@ void MainCS(uint GroupIndex : SV_GroupIndex, uint3 DispatchThreadId : SV_Dispatc
 	GroupMemoryBarrierWithGroupSync();
 
 	// Allocate node space
-	FCategorizationData Categorization = (FCategorizationData)0;
+	float PixelCoverage = 0;
 	if (NodeDesc.Count > 0)
 	{
 		NodeDesc.Offset += AllocationNodeOffset;
@@ -476,19 +474,19 @@ void MainCS(uint GroupIndex : SV_GroupIndex, uint3 DispatchThreadId : SV_Dispatc
 		// Store final sort node data 
 		if (NodeDesc.Offset + NodeDesc.Count < MaxNodeCount)
 		{
+			float ClosestDepth = 0; // Inverse-Z
 			for (uint OutIndex = 0; OutIndex < NodeDesc.Count; ++OutIndex)
 			{
 				const FHairVisSample VisSample = UnpackHairVisSample(Samples[OutIndex].x);
 
-				// VisibilityData.Coverage8bit is a weight normalising to 1 the contribution of all the compacted samples. Because later it is weighted by Categorization.PixelCoverage.
+				// VisibilityData.Coverage8bit is a weight normalising to 1 the contribution of all the compacted samples. Because later it is weighted by PixelCoverage.
 				FNodeVis OutNodeVis;
 				OutNodeVis.PrimitiveId	= VisSample.PrimitiveId;
 				OutNodeVis.Depth		= asfloat(Samples[OutIndex].w); // MSAA_DepthTexture.Load(PixelCoord, Samples[OutIndex].z);
 				OutNodeVis.Coverage8bit = To8bitCoverage(Samples[OutIndex].y / float(SampleDesc.ValidSampleCount));
 				OutNodeVis.MaterialId	= VisSample.MaterialId;
 
-				Categorization.TotalCoverage8bit += OutNodeVis.Coverage8bit;
-				Categorization.ClosestDepth = max(Categorization.ClosestDepth, OutNodeVis.Depth); // Inverse-Z
+				ClosestDepth = max(ClosestDepth, OutNodeVis.Depth); // Inverse-Z
 
 				const uint StoreIndex = NodeDesc.Offset + OutIndex;
 				OutCompactNodeData[StoreIndex] = OutNodeVis;
@@ -496,18 +494,17 @@ void MainCS(uint GroupIndex : SV_GroupIndex, uint3 DispatchThreadId : SV_Dispatc
 			}
 
 			NodeDesc.Count = NodeDesc.Count;
-			Categorization.SampleCount = NodeDesc.Count;
 
 		#if PERMUTATION_VIEWTRANSMITTANCE == 1
-			Categorization.PixelCoverage = TransmittanceToCoverage(ViewTransmittanceTexture.Load(uint3(PixelCoord, 0)), CoverageThreshold);
+			PixelCoverage = TransmittanceToCoverage(ViewTransmittanceTexture.Load(uint3(PixelCoord, 0)), CoverageThreshold);
 		#else
-			Categorization.PixelCoverage = SampleDesc.ValidSampleCount / float(PERMUTATION_MSAACOUNT);
+			PixelCoverage = SampleDesc.ValidSampleCount / float(PERMUTATION_MSAACOUNT);
 		#endif
 		}
 	}
 
 	OutCompactNodeIndex[PixelCoord] = EncodeNodeDesc(NodeDesc);
-	OutCategorizationTexture[PixelCoord] = EncodeCategorizationData(Categorization);
+	OutCoverageTexture[PixelCoord] = PixelCoverage;
 }
 
 #endif // PERMUTATION_PPLL or PERMUTATION_VISIBILITY
diff --git a/Engine/Shaders/Private/HairStrands/HairStrandsVisibilityCompactionComputeRaster.usf b/Engine/Shaders/Private/HairStrands/HairStrandsVisibilityCompactionComputeRaster.usf
--- a/Engine/Shaders/Private/HairStrands/HairStrandsVisibilityCompactionComputeRaster.usf
+++ b/Engine/Shaders/Private/HairStrands/HairStrandsVisibilityCompactionComputeRaster.usf
@@ -40,7 +40,7 @@ RWTexture2D<uint> 				OutCompactNodeCounter;
 RWTexture2D<uint> 				OutCompactNodeIndex;
 RWStructuredBuffer<FNodeVis> 	OutCompactNodeData;
 RWBuffer<uint2>					OutCompactNodeCoord;
-RWTexture2D<uint4> 				OutCategorizationTexture;
+RWTexture2D<float> 				OutCoverageTexture;
 
 groupshared uint AllocationNodeCount;
 groupshared uint AllocationNodeOffset;
@@ -177,7 +177,7 @@ void MainCS(uint GroupIndex : SV_GroupIndex, uint3 DispatchThreadId : SV_Dispatc
 	GroupMemoryBarrierWithGroupSync();
 
 	// Allocate node space
-	FCategorizationData Categorization = (FCategorizationData)0;
+	float PixelCoverage = 0;
 	if (NodeDesc.Count > 0)
 	{
 		NodeDesc.Offset += AllocationNodeOffset;
@@ -194,21 +194,15 @@ void MainCS(uint GroupIndex : SV_GroupIndex, uint3 DispatchThreadId : SV_Dispatc
 				OutNodeVis.Depth		= Vis_Depth_Cov[OutIndex].x;
 				OutNodeVis.Coverage8bit = To8bitCoverage(Vis_Depth_Cov[OutIndex].y / TotalCoverage); 
 
-				Categorization.TotalCoverage8bit += OutNodeVis.Coverage8bit;
-				Categorization.ClosestDepth = max(Categorization.ClosestDepth, OutNodeVis.Depth); // Inverse-Z
-
 				const uint StoreIndex = NodeDesc.Offset + OutIndex;
 				OutCompactNodeData[StoreIndex] = OutNodeVis;
 				OutCompactNodeCoord[StoreIndex] = PixelCoord;
 			}
 
-			NodeDesc.Count = NodeDesc.Count;
-			Categorization.SampleCount = NodeDesc.Count;
-
-			Categorization.PixelCoverage = TransmittanceToCoverage(ViewTransmittanceTexture.Load(uint3(PixelCoord, 0)), CoverageThreshold);
+			PixelCoverage = TransmittanceToCoverage(ViewTransmittanceTexture.Load(uint3(PixelCoord, 0)), CoverageThreshold);
 		}
 	}
 
 	OutCompactNodeIndex[PixelCoord] = EncodeNodeDesc(NodeDesc);
-	OutCategorizationTexture[PixelCoord] = EncodeCategorizationData(Categorization);
+	OutCoverageTexture[PixelCoord] = PixelCoverage;
 }
diff --git a/Engine/Shaders/Private/HairStrands/HairStrandsVisibilityComposeSubPixelPS.usf b/Engine/Shaders/Private/HairStrands/HairStrandsVisibilityComposeSubPixelPS.usf
--- a/Engine/Shaders/Private/HairStrands/HairStrandsVisibilityComposeSubPixelPS.usf
+++ b/Engine/Shaders/Private/HairStrands/HairStrandsVisibilityComposeSubPixelPS.usf
@@ -146,8 +146,7 @@ void ComposeSamplePS(
 		discard;
 	}
 
-	const FCategorizationData CatData = DecodeCategorizationData(HairStrands.HairCategorizationTexture.Load(PixelCoord));
-	const float PixelCoverage = min(CatData.PixelCoverage, 1);
+	const float PixelCoverage = min(HairStrands.HairCoverageTexture.Load(PixelCoord), 1);
 	if (PixelCoverage == 0)
 		discard;
 
@@ -155,7 +154,7 @@ void ComposeSamplePS(
 	const uint Resolution = uint(ceil(sqrt(float(TotalNodeCount))));
 
 	const float HairDeviceZ = HairStrands.HairOnlyDepthTexture.Load(PixelCoord);
-	const float3 ClosetPointWorldPosition = SvPositionToWorld(float4(Input.Position.xy, CatData.ClosestDepth, 1.0));;
+	const float3 ClosetPointWorldPosition = SvPositionToWorld(float4(Input.Position.xy, HairDeviceZ, 1.0));
 	const float4 Volumetric = EvaluateVolumetric(ClosetPointWorldPosition);
 
 	float3 LocalAccColor = 0;
@@ -196,8 +195,7 @@ void DOFDepthPS(
 	OutTexture = 0.0f;
 	const uint3 PixelCoord = uint3(floor(Input.Position.xy), 0);
 
-	const FCategorizationData CatData = DecodeCategorizationData(HairStrands.HairCategorizationTexture.Load(PixelCoord));
-	const float PixelCoverage = min(CatData.PixelCoverage, 1);
+	const float PixelCoverage = min(HairStrands.HairCoverageTexture.Load(PixelCoord), 1);
 	if (PixelCoverage == 0)
 	{
 		return;
diff --git a/Engine/Shaders/Private/HairStrands/HairStrandsVisibilityDepthPS.usf b/Engine/Shaders/Private/HairStrands/HairStrandsVisibilityDepthPS.usf
--- a/Engine/Shaders/Private/HairStrands/HairStrandsVisibilityDepthPS.usf
+++ b/Engine/Shaders/Private/HairStrands/HairStrandsVisibilityDepthPS.usf
@@ -4,10 +4,28 @@
 
 #include "../SceneTexturesCommon.ush"
 #include "../DeferredShadingCommon.ush"
+#include "HairStrandsVisibilityCommonStruct.ush"
 #include "HairStrandsVisibilityCommon.ush"
 
-Texture2D<uint4> CategorisationTexture;
+// Manual binding of the sample data as the uniform buffer is not created at this point
+Texture2D<float>					CoverageTexture;
+Texture2D<uint>						HairSampleOffset;
+StructuredBuffer<FPackedHairSample> HairSampleData;
 
+float ExtractClosestDeviceZ(uint3 PixelCoord)
+{
+	float Out = 0; // Inverse-Z
+
+	const FNodeDesc NodeDesc = DecodeNodeDesc(HairSampleOffset.Load(PixelCoord));
+	LOOP
+	for (uint SampleIt = 0; SampleIt < NodeDesc.Count; SampleIt++)
+	{
+		const uint SampleIndex = NodeDesc.Offset + SampleIt;
+		const FHairSample Sample = UnpackHairSample(HairSampleData[SampleIndex]);
+		Out = max(Out, Sample.Depth); // Inverse-Z
+	}
+	return Out;
+}
 ////////////////////////////////////////////////////////////////////////////////////////////////////////////////////////////////////
 
 #if PERMUTATION_OUTPUT_TYPE == 0
@@ -19,8 +37,9 @@ void MainPS(
 	out float4 OutColor		: SV_Target2
 )
 {
-	const FCategorizationData CategorisationData = DecodeCategorizationData(CategorisationTexture.Load(uint3(Input.Position.xy, 0)));
-	if (CategorisationData.PixelCoverage < 1.0f)
+	const uint3 PixelCoord = uint3(Input.Position.xy, 0);
+	const float Coverage = CoverageTexture.Load(PixelCoord);
+	if (Coverage < 1.0f)
 	{
 		discard;
 	}
@@ -28,7 +47,8 @@ void MainPS(
 	// Add a depth bias in order to avoid precision issue during the composition pass, which would reject certain samples.
 	// This depth is only used for screen passes like AO/SSGI/... so no need to be very accurate
 	const float DepthBias = 0.001f; 
-	SvDepth = CategorisationData.ClosestDepth-DepthBias;
+	const float DeviceZ = ExtractClosestDeviceZ(PixelCoord);
+	SvDepth = DeviceZ - DepthBias;
 	OutGBufferB = float4(0, 0, 0, 0); // EncodeShadingModelIdAndSelectiveOutputMask(SHADINGMODELID_HAIR, 0)
 	OutGBufferC = 0;
 	OutColor = 0;
@@ -46,11 +66,7 @@ void MainPS(
 	SvDepth = 0; // Inverse-Z
 	if (bClear == 0)
 	{
-		const FCategorizationData CategorisationData = DecodeCategorizationData(CategorisationTexture.Load(uint3(Input.Position.xy, 0)));
-		if (CategorisationData.PixelCoverage > 0.0f)
-		{
-			SvDepth = CategorisationData.ClosestDepth;
-		}
+		SvDepth = ExtractClosestDeviceZ(uint3(Input.Position.xy, 0));
 	}
 }
 #endif
\ No newline at end of file
diff --git a/Engine/Shaders/Private/HairStrands/HairStrandsVisibilityTile.usf b/Engine/Shaders/Private/HairStrands/HairStrandsVisibilityTile.usf
--- a/Engine/Shaders/Private/HairStrands/HairStrandsVisibilityTile.usf
+++ b/Engine/Shaders/Private/HairStrands/HairStrandsVisibilityTile.usf
@@ -11,7 +11,7 @@
 
 #if SHADER_TILE_GENERATION
 int2 BufferResolution;
-Texture2D<uint4> InputTexture;
+Texture2D<float> InputTexture;
 RWBuffer<uint> TileCountBuffer;
 RWBuffer<uint2> TileDataBuffer;
 RWBuffer<uint2> TileClearBuffer;
@@ -26,8 +26,7 @@ void TileMainCS(uint2 DispatchThreadId : SV_DispatchThreadID, uint LinearIndex :
 	bool bHasHairData = false;
 	if (all(PixelCoord < uint2(BufferResolution)))
 	{
-		const FCategorizationData CategorisationData = DecodeCategorizationData(InputTexture.Load(uint3(PixelCoord, 0)));
-		bHasHairData = CategorisationData.PixelCoverage > 0.0f;
+		bHasHairData = InputTexture.Load(uint3(PixelCoord, 0)) > 0.0f;
 	}
 	s_HasTileHair[LinearIndex] = bHasHairData;
 
diff --git a/Engine/Shaders/Private/HairStrands/HairStrandsVisibilityTransmittanceDepthPS.usf b/Engine/Shaders/Private/HairStrands/HairStrandsVisibilityTransmittanceDepthPS.usf
--- a/Engine/Shaders/Private/HairStrands/HairStrandsVisibilityTransmittanceDepthPS.usf
+++ b/Engine/Shaders/Private/HairStrands/HairStrandsVisibilityTransmittanceDepthPS.usf
@@ -5,7 +5,7 @@
 #include "HairStrandsVisibilityCommon.ush"
 
 Texture2D<float> SceneDepthTexture;
-Texture2D<uint4> CategorizationTexture;
+Texture2D<float> HairOnlyDepthTexture;
 
 float DistanceThreshold;
 void MainPS(
@@ -17,12 +17,14 @@ void MainPS(
 #endif
 )
 {
+#if 1
+	OutColor = 0;
+#else
 	const uint3 PixelCoord = uint3(Input.Position.xy, 0);
 	const float FarDepth = 10e-06;
 
-	const float SceneDepth = ConvertFromDeviceZ(SceneDepthTexture.Load(PixelCoord));
-	const FCategorizationData CatData = DecodeCategorizationData(CategorizationTexture.Load(PixelCoord));
-	const float SampleDepth = ConvertFromDeviceZ(CatData.ClosestDepth);
+	const float SceneDepth  = ConvertFromDeviceZ(SceneDepthTexture.Load(PixelCoord));
+	const float SampleDepth = ConvertFromDeviceZ(HairOnlyDepthTexture.Load(PixelCoord));
 
 	if (abs(SampleDepth-SceneDepth) > DistanceThreshold)
 	{
@@ -35,4 +37,6 @@ void MainPS(
 #elif PERMUTATION_OUTPUT_FORMAT == 1
 	OutColor = float2(LargeHairCount, 1);
 #endif
+
+#endif
 }
diff --git a/Engine/Source/Runtime/Renderer/Private/HairStrands/HairStrandsComposition.cpp b/Engine/Source/Runtime/Renderer/Private/HairStrands/HairStrandsComposition.cpp
--- a/Engine/Source/Runtime/Renderer/Private/HairStrands/HairStrandsComposition.cpp
+++ b/Engine/Source/Runtime/Renderer/Private/HairStrands/HairStrandsComposition.cpp
@@ -554,7 +554,7 @@ static void InternalRenderHairComposition(
 			{
 				const FHairStrandsVisibilityData& VisibilityData = View.HairStrandsViewData.VisibilityData;
 
-				if (!VisibilityData.CategorizationTexture)
+				if (!VisibilityData.CoverageTexture)
 				{
 					return; // Automatically skip for any view not rendering hair
 				}
@@ -567,7 +567,7 @@ static void InternalRenderHairComposition(
 						GraphBuilder,
 						View,
 						VisibilityData,
-						VisibilityData.CategorizationTexture,
+						VisibilityData.CoverageTexture,
 						SceneColorTexture,
 						SceneDepthTexture);
 				}
@@ -576,7 +576,7 @@ static void InternalRenderHairComposition(
 					GraphBuilder,
 					View,
 					VisibilityData,
-					VisibilityData.CategorizationTexture,
+					VisibilityData.CoverageTexture,
 					DOFDepth,
 					SceneColorTexture,
 					SceneDepthTexture);
diff --git a/Engine/Source/Runtime/Renderer/Private/HairStrands/HairStrandsData.h b/Engine/Source/Runtime/Renderer/Private/HairStrands/HairStrandsData.h
--- a/Engine/Source/Runtime/Renderer/Private/HairStrands/HairStrandsData.h
+++ b/Engine/Source/Runtime/Renderer/Private/HairStrands/HairStrandsData.h
@@ -25,7 +25,7 @@ BEGIN_GLOBAL_SHADER_PARAMETER_STRUCT(FHairStrandsViewUniformParameters, )
 	SHADER_PARAMETER(FIntPoint, HairSampleViewportResolution)							// Maximum viewport resolution of the sample space
 	SHADER_PARAMETER(uint32, bHairTileValid)											// True if tile data are valid
 	SHADER_PARAMETER(uint32, bHasEmissive)												// True if emissive is enabled
-	SHADER_PARAMETER_RDG_TEXTURE(Texture2D<uint4>, HairCategorizationTexture)			// Categorization texture aggregating hair info in screen space (closest depth, coverage, ...)
+	SHADER_PARAMETER_RDG_TEXTURE(Texture2D<float>, HairCoverageTexture)					// Hair pixel's coverage
 	SHADER_PARAMETER_RDG_TEXTURE(Texture2D, HairOnlyDepthTexture)						// Depth texture containing only hair depth (not strongly typed for legacy shader code reason)
 	SHADER_PARAMETER_RDG_TEXTURE(Texture2D, HairOnlyDepthClosestHZBTexture)				// HZB closest depth texture containing only hair depth (not strongly typed for legacy shader code reason)
 	SHADER_PARAMETER_RDG_TEXTURE(Texture2D, HairOnlyDepthFurthestHZBTexture)			// HZB furthest depth texture containing only hair depth (not strongly typed for legacy shader code reason)
@@ -74,7 +74,7 @@ struct FHairStrandsVisibilityData
 {
 	FRDGTextureRef VelocityTexture = nullptr;
 	FRDGTextureRef ResolveMaskTexture = nullptr;
-	FRDGTextureRef CategorizationTexture = nullptr;
+	FRDGTextureRef CoverageTexture = nullptr;
 	FRDGTextureRef ViewHairCountTexture = nullptr;
 	FRDGTextureRef ViewHairCountUintTexture = nullptr;
 	FRDGTextureRef EmissiveTexture = nullptr;
@@ -97,6 +97,7 @@ struct FHairStrandsVisibilityData
 	FHairStrandsTiles TileData;
 
 	const static EPixelFormat NodeCoordFormat = PF_R16G16_UINT;
+	const static EPixelFormat CoverageFormat = PF_R16F;
 
 	// Hair lighting is accumulated within this buffer
 	// Allocated conservatively
diff --git a/Engine/Source/Runtime/Renderer/Private/HairStrands/HairStrandsDebug.cpp b/Engine/Source/Runtime/Renderer/Private/HairStrands/HairStrandsDebug.cpp
--- a/Engine/Source/Runtime/Renderer/Private/HairStrands/HairStrandsDebug.cpp
+++ b/Engine/Source/Runtime/Renderer/Private/HairStrands/HairStrandsDebug.cpp
@@ -292,7 +292,7 @@ static void AddDebugHairPrintPass(
 
 	FHairDebugPrintCS::FParameters* Parameters = GraphBuilder.AllocParameters<FHairDebugPrintCS::FParameters>();
 	Parameters->GroupSize = GetVendorOptimalGroupSize2D();
-	Parameters->MaxResolution = VisibilityData.CategorizationTexture ? VisibilityData.CategorizationTexture->Desc.Extent : FIntPoint(0,0);
+	Parameters->MaxResolution = VisibilityData.CoverageTexture ? VisibilityData.CoverageTexture->Desc.Extent : FIntPoint(0,0);
 	Parameters->PixelCoord = View->CursorPos;
 	Parameters->FastResolveMask = STENCIL_TEMPORAL_RESPONSIVE_AA_MASK;
 	Parameters->HairCountTexture = ViewHairCountTexture;
@@ -364,7 +364,7 @@ static void AddDebugHairPass(
 		InDebugMode == EHairDebugMode::MaterialSpecular ||
 		InDebugMode == EHairDebugMode::MaterialTangent);
 
-	if (!VisibilityData.CategorizationTexture || !VisibilityData.NodeIndex || !VisibilityData.NodeData) return;
+	if (!VisibilityData.CoverageTexture || !VisibilityData.NodeIndex || !VisibilityData.NodeData) return;
 	if (InDebugMode == EHairDebugMode::TAAResolveType && !InDepthStencilTexture) return;
 
 	const FIntRect Viewport = View->ViewRect;
diff --git a/Engine/Source/Runtime/Renderer/Private/HairStrands/HairStrandsEnvironment.cpp b/Engine/Source/Runtime/Renderer/Private/HairStrands/HairStrandsEnvironment.cpp
--- a/Engine/Source/Runtime/Renderer/Private/HairStrands/HairStrandsEnvironment.cpp
+++ b/Engine/Source/Runtime/Renderer/Private/HairStrands/HairStrandsEnvironment.cpp
@@ -499,7 +499,7 @@ static void InternalRenderHairStrandsEnvironmentLighting(
 
 	const FHairStrandsVisibilityData& VisibilityData = View.HairStrandsViewData.VisibilityData;
 	const FHairStrandsVoxelResources& VoxelResources = View.HairStrandsViewData.VirtualVoxelResources;
-	const bool bRenderHairLighting = VoxelResources.IsValid() && VisibilityData.CategorizationTexture;
+	const bool bRenderHairLighting = VoxelResources.IsValid() && VisibilityData.CoverageTexture;
 	if (!bRenderHairLighting)
 	{
 		return;
@@ -534,7 +534,7 @@ void RenderHairStrandsAmbientOcclusion(
 
 	const FHairStrandsVisibilityData& VisibilityData = View.HairStrandsViewData.VisibilityData;
 	const FHairStrandsVoxelResources& VoxelResources = View.HairStrandsViewData.VirtualVoxelResources;
-	const bool bRenderHairLighting = VoxelResources.IsValid() && VisibilityData.CategorizationTexture;
+	const bool bRenderHairLighting = VoxelResources.IsValid() && VisibilityData.CoverageTexture;
 	check(bRenderHairLighting);
 
 	for (const FHairStrandsMacroGroupData& MacroGroupData : View.HairStrandsViewData.MacroGroupDatas)
diff --git a/Engine/Source/Runtime/Renderer/Private/HairStrands/HairStrandsRendering.cpp b/Engine/Source/Runtime/Renderer/Private/HairStrands/HairStrandsRendering.cpp
--- a/Engine/Source/Runtime/Renderer/Private/HairStrands/HairStrandsRendering.cpp
+++ b/Engine/Source/Runtime/Renderer/Private/HairStrands/HairStrandsRendering.cpp
@@ -11,9 +11,9 @@ static TRDGUniformBufferRef<FHairStrandsViewUniformParameters> InternalCreateHai
 {
 	FHairStrandsViewUniformParameters* Parameters = GraphBuilder.AllocParameters<FHairStrandsViewUniformParameters>();
 	Parameters->HairDualScatteringRoughnessOverride = GetHairDualScatteringRoughnessOverride();
-	if (In && In->CategorizationTexture)
+	if (In && In->CoverageTexture)
 	{
-		Parameters->HairCategorizationTexture = In->CategorizationTexture;
+		Parameters->HairCoverageTexture = In->CoverageTexture;
 		Parameters->HairOnlyDepthTexture = In->HairOnlyDepthTexture;
 		Parameters->HairOnlyDepthClosestHZBTexture = In->HairOnlyDepthClosestHZBTexture;
 		Parameters->HairOnlyDepthFurthestHZBTexture = In->HairOnlyDepthFurthestHZBTexture;
@@ -76,7 +76,7 @@ static TRDGUniformBufferRef<FHairStrandsViewUniformParameters> InternalCreateHai
 		Parameters->HairOnlyDepthTexture = FarDepth;
 		Parameters->HairOnlyDepthFurthestHZBTexture = BlackTexture;
 		Parameters->HairOnlyDepthClosestHZBTexture = Parameters->HairOnlyDepthFurthestHZBTexture;
-		Parameters->HairCategorizationTexture = ZeroRGBA16_UINT;
+		Parameters->HairCoverageTexture = BlackTexture;
 		Parameters->HairSampleCount = ZeroR32_UINT;
 		Parameters->HairSampleOffset = ZeroR32_UINT;
 		Parameters->HairSampleCoords = DummyBufferRG16SRV;
@@ -175,7 +175,7 @@ void RenderHairBasePass(
 			}
 		}
 		
-		if (View.HairStrandsViewData.VisibilityData.CategorizationTexture)
+		if (View.HairStrandsViewData.VisibilityData.CoverageTexture)
 		{			
 			View.HairStrandsViewData.UniformBuffer = InternalCreateHairStrandsViewUniformBuffer(GraphBuilder, &View.HairStrandsViewData.VisibilityData);
 			View.HairStrandsViewData.bIsValid = true;
diff --git a/Engine/Source/Runtime/Renderer/Private/HairStrands/HairStrandsTransmittance.cpp b/Engine/Source/Runtime/Renderer/Private/HairStrands/HairStrandsTransmittance.cpp
--- a/Engine/Source/Runtime/Renderer/Private/HairStrands/HairStrandsTransmittance.cpp
+++ b/Engine/Source/Runtime/Renderer/Private/HairStrands/HairStrandsTransmittance.cpp
@@ -966,7 +966,7 @@ void RenderHairStrandsShadowMask(
 	{
 		if (HairStrands::HasViewHairStrandsData(View))
 		{
-			check(View.HairStrandsViewData.VisibilityData.CategorizationTexture);
+			check(View.HairStrandsViewData.VisibilityData.CoverageTexture);
 			InternalRenderHairStrandsShadowMask(
 				GraphBuilder, 
 				View, 
diff --git a/Engine/Source/Runtime/Renderer/Private/HairStrands/HairStrandsVisibility.cpp b/Engine/Source/Runtime/Renderer/Private/HairStrands/HairStrandsVisibility.cpp
--- a/Engine/Source/Runtime/Renderer/Private/HairStrands/HairStrandsVisibility.cpp
+++ b/Engine/Source/Runtime/Renderer/Private/HairStrands/HairStrandsVisibility.cpp
@@ -872,7 +872,7 @@ class FHairVelocityCS: public FGlobalShader
 		SHADER_PARAMETER(FIntPoint, ResolutionOffset)
 		SHADER_PARAMETER(float, VelocityThreshold)
 		SHADER_PARAMETER(float, CoverageThreshold)
-		SHADER_PARAMETER_RDG_TEXTURE(Texture2D, CategorizationTexture)
+		SHADER_PARAMETER_RDG_TEXTURE(Texture2D, CoverageTexture)
 		SHADER_PARAMETER_RDG_TEXTURE(Texture2D, NodeIndex)
 		SHADER_PARAMETER_RDG_BUFFER_SRV(Buffer, NodeVelocity)
 		SHADER_PARAMETER_RDG_BUFFER_SRV(StructuredBuffer<FNodeVis>, NodeVis)
@@ -898,7 +898,7 @@ static void AddHairVelocityPass(
 	FRDGBuilder& GraphBuilder,
 	const FViewInfo& View,
 	const FHairStrandsMacroGroupDatas& MacroGroupDatas,
-	FRDGTextureRef& CategorizationTexture,
+	FRDGTextureRef& CoverageTexture,
 	FRDGTextureRef& NodeIndex,
 	FRDGBufferRef& NodeVis,
 	FRDGBufferRef& NodeVelocity,
@@ -931,7 +931,7 @@ static void AddHairVelocityPass(
 	PassParameters->NodeIndex = NodeIndex;
 	PassParameters->NodeVis = GraphBuilder.CreateSRV(NodeVis);
 	PassParameters->NodeVelocity = GraphBuilder.CreateSRV(NodeVelocity, FMaterialPassOutput::VelocityFormat);
-	PassParameters->CategorizationTexture = CategorizationTexture;
+	PassParameters->CoverageTexture = CoverageTexture;
 	PassParameters->OutVelocityTexture = GraphBuilder.CreateUAV(OutVelocityTexture);
 	PassParameters->OutResolveMaskTexture = GraphBuilder.CreateUAV(OutResolveMaskTexture);
 
@@ -1558,7 +1558,7 @@ class FHairVisibilityPrimitiveIdCompactionCS : public FGlobalShader
 
 		SHADER_PARAMETER_RDG_TEXTURE_UAV(RWTexture2D, OutCompactNodeCounter)
 		SHADER_PARAMETER_RDG_TEXTURE_UAV(RWTexture2D, OutCompactNodeIndex)
-		SHADER_PARAMETER_RDG_TEXTURE_UAV(RWTexture2D, OutCategorizationTexture)
+		SHADER_PARAMETER_RDG_TEXTURE_UAV(RWTexture2D, OutCoverageTexture)
 		SHADER_PARAMETER_RDG_BUFFER_UAV(RWStructuredBuffer, OutCompactNodeData)
 		SHADER_PARAMETER_RDG_BUFFER_UAV(RWBuffer, OutCompactNodeCoord)
 		SHADER_PARAMETER_RDG_TEXTURE_UAV(RWTexture2D, OutVelocityTexture)
@@ -1607,7 +1607,7 @@ static void AddHairVisibilityPrimitiveIdCompactionPass(
 	FRDGTextureRef& OutCompactNodeIndex,
 	FRDGBufferRef& OutCompactNodeData,
 	FRDGBufferRef& OutCompactNodeCoord,
-	FRDGTextureRef& OutCategorizationTexture,
+	FRDGTextureRef& OutCoverageTexture,
 	FRDGTextureRef& OutVelocityTexture,
 	FRDGBufferRef& OutIndirectArgsBuffer,
 	uint32& OutMaxRenderNodeCount)
@@ -1639,14 +1639,14 @@ static void AddHairVisibilityPrimitiveIdCompactionPass(
 	}
 
 	{
-		FRDGTextureDesc Desc = FRDGTextureDesc::Create2D(Resolution, PF_R16G16B16A16_UINT, FClearValueBinding::None, TexCreate_UAV | TexCreate_ShaderResource);
-		OutCategorizationTexture = GraphBuilder.CreateTexture(Desc, TEXT("Hair.CategorizationTexture"));
+		FRDGTextureDesc Desc = FRDGTextureDesc::Create2D(Resolution, FHairStrandsVisibilityData::CoverageFormat, FClearValueBinding::None, TexCreate_UAV | TexCreate_ShaderResource);
+		OutCoverageTexture = GraphBuilder.CreateTexture(Desc, TEXT("Hair.CoverageTexture"));
 	}
 
 	const uint32 ClearValues[4] = { 0u,0u,0u,0u };
 	AddClearUAVPass(GraphBuilder, GraphBuilder.CreateUAV(OutCompactCounter), ClearValues);
 	AddClearUAVPass(GraphBuilder, GraphBuilder.CreateUAV(OutCompactNodeIndex), ClearValues);
-	AddClearUAVPass(GraphBuilder, GraphBuilder.CreateUAV(OutCategorizationTexture), ClearValues);
+	AddClearUAVPass(GraphBuilder, GraphBuilder.CreateUAV(OutCoverageTexture), 0.f);
 
 	// Adapt the buffer allocation based on the bounding box of the hair macro groups. This allows to reduce the overall allocation size
 	const FIntRect HairRect = ComputeVisibleHairStrandsMacroGroupsRect(View.ViewRect, MacroGroupDatas);
@@ -1688,7 +1688,7 @@ static void AddHairVisibilityPrimitiveIdCompactionPass(
 	PassParameters->OutCompactNodeIndex = GraphBuilder.CreateUAV(OutCompactNodeIndex);
 	PassParameters->OutCompactNodeData = GraphBuilder.CreateUAV(OutCompactNodeData);
 	PassParameters->OutCompactNodeCoord = GraphBuilder.CreateUAV(OutCompactNodeCoord, FHairStrandsVisibilityData::NodeCoordFormat);
-	PassParameters->OutCategorizationTexture = GraphBuilder.CreateUAV(OutCategorizationTexture);
+	PassParameters->OutCoverageTexture = GraphBuilder.CreateUAV(OutCoverageTexture);
 
 	if (bWriteOutVelocity)
 	{
@@ -1742,7 +1742,7 @@ class FHairVisibilityCompactionComputeRasterCS : public FGlobalShader
 
 		SHADER_PARAMETER_RDG_TEXTURE_UAV(RWTexture2D, OutCompactNodeCounter)
 		SHADER_PARAMETER_RDG_TEXTURE_UAV(RWTexture2D, OutCompactNodeIndex)
-		SHADER_PARAMETER_RDG_TEXTURE_UAV(RWTexture2D, OutCategorizationTexture)
+		SHADER_PARAMETER_RDG_TEXTURE_UAV(RWTexture2D, OutCoverageTexture)
 		SHADER_PARAMETER_RDG_BUFFER_UAV(RWStructuredBuffer, OutCompactNodeData)
 		SHADER_PARAMETER_RDG_BUFFER_UAV(RWBuffer, OutCompactNodeCoord)
 
@@ -1776,7 +1776,7 @@ static void AddHairVisibilityCompactionComputeRasterPass(
 	FRDGTextureRef& OutCompactNodeIndex,
 	FRDGBufferRef&  OutCompactNodeData,
 	FRDGBufferRef&  OutCompactNodeCoord,
-	FRDGTextureRef& OutCategorizationTexture,
+	FRDGTextureRef& OutCoverageTexture,
 	FRDGTextureRef& OutVelocityTexture,
 	FRDGBufferRef&  OutIndirectArgsBuffer,
 	uint32& OutMaxRenderNodeCount)
@@ -1794,14 +1794,14 @@ static void AddHairVisibilityCompactionComputeRasterPass(
 	}
 
 	{
-		FRDGTextureDesc Desc = FRDGTextureDesc::Create2D(Resolution, PF_R16G16B16A16_UINT, FClearValueBinding::None, TexCreate_UAV | TexCreate_ShaderResource);
-		OutCategorizationTexture = GraphBuilder.CreateTexture(Desc, TEXT("Hair.CategorizationTexture"));
+		FRDGTextureDesc Desc = FRDGTextureDesc::Create2D(Resolution, FHairStrandsVisibilityData::CoverageFormat, FClearValueBinding::None, TexCreate_UAV | TexCreate_ShaderResource);
+		OutCoverageTexture = GraphBuilder.CreateTexture(Desc, TEXT("Hair.CoverageTexture"));
 	}
 
 	const uint32 ClearValues[4] = { 0u,0u,0u,0u };
 	AddClearUAVPass(GraphBuilder, GraphBuilder.CreateUAV(OutCompactCounter), ClearValues);
 	AddClearUAVPass(GraphBuilder, GraphBuilder.CreateUAV(OutCompactNodeIndex), ClearValues);
-	AddClearUAVPass(GraphBuilder, GraphBuilder.CreateUAV(OutCategorizationTexture), ClearValues);
+	AddClearUAVPass(GraphBuilder, GraphBuilder.CreateUAV(OutCoverageTexture), 0.f);
 
 	// Select render node count according to current mode
 	const uint32 MSAASampleCount = GetHairVisibilityRenderMode() == HairVisibilityRenderMode_MSAA_Visibility ? GetMaxSamplePerPixel() : 1;
@@ -1827,7 +1827,7 @@ static void AddHairVisibilityCompactionComputeRasterPass(
 	PassParameters->OutCompactNodeIndex		= GraphBuilder.CreateUAV(OutCompactNodeIndex);
 	PassParameters->OutCompactNodeData		= GraphBuilder.CreateUAV(OutCompactNodeData);
 	PassParameters->OutCompactNodeCoord		= GraphBuilder.CreateUAV(OutCompactNodeCoord, FHairStrandsVisibilityData::NodeCoordFormat);
-	PassParameters->OutCategorizationTexture= GraphBuilder.CreateUAV(OutCategorizationTexture);
+	PassParameters->OutCoverageTexture		= GraphBuilder.CreateUAV(OutCoverageTexture);
 
 	const FIntPoint GroupSize = GetVendorOptimalGroupSize2D();
 	FHairVisibilityCompactionComputeRasterCS::FPermutationDomain PermutationVector;
@@ -2210,7 +2210,7 @@ class FHairViewTransmittanceDepthPS : public FGlobalShader
 	BEGIN_SHADER_PARAMETER_STRUCT(FParameters, )
 		SHADER_PARAMETER(float, DistanceThreshold)
 		SHADER_PARAMETER_RDG_TEXTURE(Texture2D, SceneDepthTexture)
-		SHADER_PARAMETER_RDG_TEXTURE(Texture2D, CategorizationTexture)
+		SHADER_PARAMETER_RDG_TEXTURE(Texture2D, CoverageTexture)
 		SHADER_PARAMETER_STRUCT_REF(FViewUniformShaderParameters, ViewUniformBuffer)
 		RENDER_TARGET_BINDING_SLOTS()
 	END_SHADER_PARAMETER_STRUCT()
@@ -2242,13 +2242,13 @@ IMPLEMENT_GLOBAL_SHADER(FHairViewTransmittanceDepthPS, "/Engine/Private/HairStra
 static void AddHairViewTransmittanceDepthPass(
 	FRDGBuilder& GraphBuilder,
 	const FViewInfo& View,
-	const FRDGTextureRef& CategorizationTexture,
+	const FRDGTextureRef& CoverageTexture,
 	const FRDGTextureRef& SceneDepthTexture,
 	FRDGTextureRef& HairCountTexture)
 {
 	FHairViewTransmittanceDepthPS::FParameters* Parameters = GraphBuilder.AllocParameters<FHairViewTransmittanceDepthPS::FParameters>();
 	Parameters->DistanceThreshold = FMath::Max(1.f, GHairStrandsViewHairCountDepthDistanceThreshold);
-	Parameters->CategorizationTexture = CategorizationTexture;
+	Parameters->CoverageTexture = CoverageTexture;
 	Parameters->SceneDepthTexture = SceneDepthTexture;
 	Parameters->ViewUniformBuffer = View.ViewUniformBuffer;
 	Parameters->RenderTargets[0] = FRenderTargetBinding(HairCountTexture, ERenderTargetLoadAction::ELoad);
@@ -2310,7 +2310,9 @@ class FHairVisibilityDepthPS : public FGlobalShader
 		SHADER_PARAMETER_STRUCT_INCLUDE(FHairStrandsTilePassVS::FParameters, TileData)
 		SHADER_PARAMETER_STRUCT_REF(FViewUniformShaderParameters, ViewUniformBuffer)
 		SHADER_PARAMETER(uint32, bClear)
-		SHADER_PARAMETER_RDG_TEXTURE(Texture2D, CategorisationTexture)
+		SHADER_PARAMETER_RDG_TEXTURE(Texture2D<float>, CoverageTexture)
+		SHADER_PARAMETER_RDG_TEXTURE(Texture2D<uint>, HairSampleOffset)
+		SHADER_PARAMETER_RDG_BUFFER_SRV(StructuredBuffer<FPackedHairSample>, HairSampleData)
 		RENDER_TARGET_BINDING_SLOTS()
 	END_SHADER_PARAMETER_STRUCT()
 
@@ -2343,7 +2345,9 @@ static void AddHairVisibilityCommonPatchPass(
 	const FViewInfo& View,
 	const FHairStrandsTiles& TileData,
 	const EHairPatchPass PassType,
-	const FRDGTextureRef& CategorisationTexture,
+	const FRDGTextureRef& CoverageTexture,
+	const FRDGTextureRef& HairSampleOffset,
+	const FRDGBufferRef& HairSampleData,
 	FRDGTextureRef OutGBufferBTexture,
 	FRDGTextureRef OutGBufferCTexture,
 	FRDGTextureRef OutColorTexture,
@@ -2352,8 +2356,11 @@ static void AddHairVisibilityCommonPatchPass(
 	FHairVisibilityDepthPS::FParameters* Parameters = GraphBuilder.AllocParameters<FHairVisibilityDepthPS::FParameters>();
 	Parameters->bClear = PassType == EHairPatchPass::DepthClear ? 1u : 0u;
 	Parameters->TileData = GetHairStrandsTileParameters(View, TileData, PassType == EHairPatchPass::DepthClear ?  FHairStrandsTiles::ETileType::Other : FHairStrandsTiles::ETileType::Hair);
-	Parameters->CategorisationTexture = CategorisationTexture;
 	Parameters->ViewUniformBuffer = View.ViewUniformBuffer;
+	Parameters->CoverageTexture = CoverageTexture;
+	Parameters->HairSampleOffset = HairSampleOffset;
+	Parameters->HairSampleData = GraphBuilder.CreateSRV(HairSampleData);
+
 	Parameters->RenderTargets.DepthStencil = FDepthStencilBinding(
 		OutDepthTexture,
 		ERenderTargetLoadAction::ELoad,
@@ -2427,7 +2434,9 @@ static void AddHairVisibilityColorAndDepthPatchPass(
 	FRDGBuilder& GraphBuilder,
 	const FViewInfo& View,
 	const FHairStrandsTiles& TileData,
-	const FRDGTextureRef& CategorisationTexture,
+	const FRDGTextureRef& CoverageTexture,
+	const FRDGTextureRef& HairSampleOffset,
+	const FRDGBufferRef& HairSampleData,
 	FRDGTextureRef& OutGBufferBTexture,
 	FRDGTextureRef& OutGBufferCTexture,
 	FRDGTextureRef& OutColorTexture,
@@ -2443,7 +2452,9 @@ static void AddHairVisibilityColorAndDepthPatchPass(
 		View,
 		TileData,
 		EHairPatchPass::GBufferPatch,
-		CategorisationTexture,
+		CoverageTexture,
+		HairSampleOffset,
+		HairSampleData,
 		OutGBufferBTexture,
 		OutGBufferCTexture,
 		OutColorTexture,
@@ -2454,7 +2465,9 @@ static void AddHairOnlyDepthPass(
 	FRDGBuilder& GraphBuilder,
 	const FViewInfo& View,
 	const FHairStrandsTiles& TileData,
-	const FRDGTextureRef& CategorisationTexture,
+	const FRDGTextureRef& CoverageTexture,
+	const FRDGTextureRef& HairSampleOffset,
+	const FRDGBufferRef& HairSampleData,
 	FRDGTextureRef& OutDepthTexture)
 {
 	if (!OutDepthTexture)
@@ -2471,7 +2484,9 @@ static void AddHairOnlyDepthPass(
 			View,
 			TileData,
 			EHairPatchPass::DepthClear,
-			CategorisationTexture,
+			CoverageTexture,
+			HairSampleOffset,
+			HairSampleData,
 			nullptr,
 			nullptr,
 			nullptr,
@@ -2484,7 +2499,9 @@ static void AddHairOnlyDepthPass(
 		View,
 		TileData,
 		EHairPatchPass::DepthPatch,
-		CategorisationTexture,
+		CoverageTexture,
+		HairSampleOffset,
+		HairSampleData,
 		nullptr,
 		nullptr,
 		nullptr,
@@ -2916,7 +2933,7 @@ void RenderHairStrandsVisibilityBuffer(
 			check(RenderMode == HairVisibilityRenderMode_MSAA_Visibility || RenderMode == HairVisibilityRenderMode_PPLL || RenderMode == HairVisibilityRenderMode_ComputeRaster);
 
 			FRDGTextureRef HairOnlyDepthTexture = GraphBuilder.CreateTexture(SceneDepthTexture->Desc, TEXT("Hair.HairOnlyDepthTexture"));
-			FRDGTextureRef CategorizationTexture = nullptr;
+			FRDGTextureRef CoverageTexture = nullptr;
 			FRDGTextureRef CompactNodeIndex = nullptr;
 			FRDGBufferRef  CompactNodeData = nullptr;
 			FRDGTextureRef NodeCounter = nullptr;
@@ -2964,15 +2981,15 @@ void RenderHairStrandsVisibilityBuffer(
 							CompactNodeIndex,
 							CompactNodeData,
 							CompactNodeCoord,
-							CategorizationTexture,
+							CoverageTexture,
 							SceneVelocityTexture,
 							IndirectArgsBuffer,
 							VisibilityData.MaxNodeCount);
 
 						// Generate Tile data
-						if (CategorizationTexture && bGenerateTile)
+						if (CoverageTexture && bGenerateTile)
 						{
-							VisibilityData.TileData = AddHairStrandsGenerateTilesPass(GraphBuilder, View, CategorizationTexture);
+							VisibilityData.TileData = AddHairStrandsGenerateTilesPass(GraphBuilder, View, CoverageTexture);
 						}
 
 						// Evaluate material based on the visiblity pass result
@@ -2996,7 +3013,7 @@ void RenderHairStrandsVisibilityBuffer(
 							GraphBuilder,
 							View,
 							MacroGroupDatas,
-							CategorizationTexture,
+							CoverageTexture,
 							CompactNodeIndex,
 							CompactNodeData,
 							PassOutput.NodeVelocity,
@@ -3011,7 +3028,7 @@ void RenderHairStrandsVisibilityBuffer(
 
 						VisibilityData.SampleLightingBuffer = SampleLightingBuffer;
 						VisibilityData.NodeIndex = CompactNodeIndex;
-						VisibilityData.CategorizationTexture = CategorizationTexture;
+						VisibilityData.CoverageTexture = CoverageTexture;
 						VisibilityData.HairOnlyDepthTexture = HairOnlyDepthTexture;
 						VisibilityData.NodeData = CompactNodeData;
 						VisibilityData.NodeCoord = CompactNodeCoord;
@@ -3031,7 +3048,9 @@ void RenderHairStrandsVisibilityBuffer(
 							GraphBuilder,
 							View,
 							VisibilityData.TileData,
-							CategorizationTexture,
+							CoverageTexture,
+							CompactNodeIndex,
+							CompactNodeData,
 							SceneGBufferBTexture,
 							SceneGBufferCTexture,
 							SceneColorTexture,
@@ -3042,7 +3061,9 @@ void RenderHairStrandsVisibilityBuffer(
 						GraphBuilder,
 						View,
 						VisibilityData.TileData,
-						CategorizationTexture,
+						CoverageTexture,
+						CompactNodeIndex,
+						CompactNodeData,
 						HairOnlyDepthTexture);
 
 					AddHairOnlyHZBPass(
@@ -3131,15 +3152,15 @@ void RenderHairStrandsVisibilityBuffer(
 						CompactNodeIndex,
 						CompactNodeData,
 						CompactNodeCoord,
-						CategorizationTexture,
+						CoverageTexture,
 						SceneVelocityTexture,
 						IndirectArgsBuffer,
 						VisibilityData.MaxNodeCount);
 
 					// Generate Tile data
-					if (CategorizationTexture && bGenerateTile)
+					if (CoverageTexture && bGenerateTile)
 					{
-						VisibilityData.TileData = AddHairStrandsGenerateTilesPass(GraphBuilder, View, CategorizationTexture);
+						VisibilityData.TileData = AddHairStrandsGenerateTilesPass(GraphBuilder, View, CoverageTexture);
 					}
 
 					{
@@ -3166,7 +3187,7 @@ void RenderHairStrandsVisibilityBuffer(
 							GraphBuilder,
 							View,
 							MacroGroupDatas,
-							CategorizationTexture,
+							CoverageTexture,
 							CompactNodeIndex,
 							CompactNodeData,
 							PassOutput.NodeVelocity,
@@ -3192,7 +3213,7 @@ void RenderHairStrandsVisibilityBuffer(
 
 					 VisibilityData.SampleLightingBuffer	= SampleLightingBuffer;
 					 VisibilityData.NodeIndex				= CompactNodeIndex;
-					 VisibilityData.CategorizationTexture	= CategorizationTexture;
+					 VisibilityData.CoverageTexture	= CoverageTexture;
 					 VisibilityData.HairOnlyDepthTexture	= HairOnlyDepthTexture;
 					 VisibilityData.NodeData				= CompactNodeData;
 					 VisibilityData.NodeCoord				= CompactNodeCoord;
@@ -3207,7 +3228,7 @@ void RenderHairStrandsVisibilityBuffer(
 						AddHairViewTransmittanceDepthPass(
 							GraphBuilder,
 							View,
-							CategorizationTexture,
+							CoverageTexture,
 							SceneDepthTexture,
 							ViewTransmittance.HairCountTexture);
 						VisibilityData.ViewHairCountTexture = ViewTransmittance.HairCountTexture;
@@ -3223,7 +3244,9 @@ void RenderHairStrandsVisibilityBuffer(
 							GraphBuilder,
 							View,
 							VisibilityData.TileData,
-							CategorizationTexture,
+							CoverageTexture,
+							CompactNodeIndex,
+							CompactNodeData,
 							SceneGBufferBTexture,
 							SceneGBufferCTexture,
 							SceneColorTexture,
@@ -3234,7 +3257,9 @@ void RenderHairStrandsVisibilityBuffer(
 						GraphBuilder,
 						View,
 						VisibilityData.TileData,
-						CategorizationTexture,
+						CoverageTexture,
+						CompactNodeIndex,
+						CompactNodeData,
 						HairOnlyDepthTexture);
 
 					AddHairOnlyHZBPass(
@@ -3281,14 +3306,14 @@ void RenderHairStrandsVisibilityBuffer(
 						CompactNodeIndex,
 						CompactNodeData,
 						CompactNodeCoord,
-						CategorizationTexture,
+						CoverageTexture,
 						SceneVelocityTexture,
 						IndirectArgsBuffer,
 						VisibilityData.MaxNodeCount);
 
 					VisibilityData.MaxSampleCount = GetMaxSamplePerPixel();
 					VisibilityData.NodeIndex = CompactNodeIndex;
-					VisibilityData.CategorizationTexture = CategorizationTexture;
+					VisibilityData.CoverageTexture = CoverageTexture;
 					VisibilityData.HairOnlyDepthTexture = HairOnlyDepthTexture;
 					VisibilityData.NodeData = CompactNodeData;
 					VisibilityData.NodeCoord = CompactNodeCoord;
@@ -3297,9 +3322,9 @@ void RenderHairStrandsVisibilityBuffer(
 				}
 
 				// Generate Tile data
-				if (CategorizationTexture && bGenerateTile)
+				if (CoverageTexture && bGenerateTile)
 				{
-					VisibilityData.TileData = AddHairStrandsGenerateTilesPass(GraphBuilder, View, CategorizationTexture);
+					VisibilityData.TileData = AddHairStrandsGenerateTilesPass(GraphBuilder, View, CoverageTexture);
 				}
 
 				if (bRunColorAndDepthPatching)
@@ -3308,7 +3333,9 @@ void RenderHairStrandsVisibilityBuffer(
 						GraphBuilder,
 						View,
 						VisibilityData.TileData,
-						CategorizationTexture,
+						CoverageTexture,
+						CompactNodeIndex,
+						CompactNodeData,
 						SceneGBufferBTexture,
 						SceneGBufferCTexture,
 						SceneColorTexture,
@@ -3319,7 +3346,9 @@ void RenderHairStrandsVisibilityBuffer(
 					GraphBuilder,
 					View,
 					VisibilityData.TileData,
-					CategorizationTexture,
+					CoverageTexture,
+					CompactNodeIndex,
+					CompactNodeData,
 					HairOnlyDepthTexture);
 
 				// Allocate buffer for storing all the light samples
@@ -3339,7 +3368,7 @@ void RenderHairStrandsVisibilityBuffer(
 			}
 
 		#if RHI_RAYTRACING
-			if (IsRayTracingEnabled() && CategorizationTexture)
+			if (IsRayTracingEnabled() && CoverageTexture)
 			{
 				FRDGTextureRef LightingChannelMaskTexture = AddHairLightChannelMaskPass(
 					GraphBuilder,
