Remove legacy hair velocity/composition shaders

diff --git a/Engine/Shaders/Private/HairStrands/HairStrandsVisibilityComposeSubPixelPS.usf b/Engine/Shaders/Private/HairStrands/HairStrandsVisibilityComposeSubPixelPS.usf
--- a/Engine/Shaders/Private/HairStrands/HairStrandsVisibilityComposeSubPixelPS.usf
+++ b/Engine/Shaders/Private/HairStrands/HairStrandsVisibilityComposeSubPixelPS.usf
@@ -8,62 +8,8 @@
 #include "../DeferredShadingCommon.ush"
 #include "HairStrandsVisibilityCommon.ush"   
 
-#if SHADER_SUBCOLOR
-
-Texture2D<float4> SubPixelColorTexture;
-
-void SubColorPS(
-	in FScreenVertexOutput Input,
-	out float4 OutColor  : SV_Target0, 
-	out float  OutDepth : SV_DEPTH)
-{
-	const uint2 PixelCoord = floor(Input.Position.xy);
-	OutColor = SubPixelColorTexture.Load(uint3(PixelCoord, 0));
-
-	// Early out writing to not write stencil value for full coverred pixels (or not covered at all)
-	//if (OutColor.a == 0)
-	//	discard;
-
-	OutDepth = HairStrands.HairOnlyDepthTexture.Load(uint3(Input.Position.xy, 0));
-}
-
-#endif // SHADER_SUBCOLOR
-
 ////////////////////////////////////////////////////////////////////////////////////////////////
 
-#if SHADER_FASTRESOLVE_MSAA
-
-#ifndef PERMUTATION_MSAACOUNT 
-#error PERMUTATION_MSAACOUNT is not defined
-#endif
-
-float VelocityThreshold;
-Texture2DMS<float2, PERMUTATION_MSAACOUNT> HairVisibilityVelocityTexture;
-
-void FastResolvePS(in FScreenVertexOutput Input)
-{
-	const uint2 PixelCoord = floor(Input.Position.xy);
-
-	bool bNeedFastResolve = false;
-	for (uint SampleIt = 0; SampleIt < PERMUTATION_MSAACOUNT; ++SampleIt)
-	{
-		const float2 EncodedVelocity = HairVisibilityVelocityTexture.Load(PixelCoord, SampleIt);
-		const float2 Velocity = DecodeVelocityFromTexture(float4(EncodedVelocity, 0.0, 0.0)).xy;
-		const float VelocityMagnitude = sqrt(dot(Velocity, Velocity));
-		if (EncodedVelocity.x > 0 && VelocityMagnitude > VelocityThreshold)
-		{
-			bNeedFastResolve = true;
-			break;
-		}
-	}
-
-	if (!bNeedFastResolve)
-		discard;
-}
-
-#endif // SHADER_FASTRESOLVE_MSAA
-
-
 #if SHADER_FASTRESOLVE_MASK
 
 Texture2D<uint> ResolveMaskTexture;
diff --git a/Engine/Source/Runtime/Renderer/Private/HairStrands/HairStrandsComposition.cpp b/Engine/Source/Runtime/Renderer/Private/HairStrands/HairStrandsComposition.cpp
--- a/Engine/Source/Runtime/Renderer/Private/HairStrands/HairStrandsComposition.cpp
+++ b/Engine/Source/Runtime/Renderer/Private/HairStrands/HairStrandsComposition.cpp
@@ -298,82 +298,6 @@ static FRDGTextureRef AddHairDOFDepthPass(
 
 ///////////////////////////////////////////////////////////////////////////////////////////////////
 
-class FHairVisibilityFastResolvePS : public FGlobalShader
-{
-	DECLARE_GLOBAL_SHADER(FHairVisibilityFastResolvePS);
-	SHADER_USE_PARAMETER_STRUCT(FHairVisibilityFastResolvePS, FGlobalShader);
-
-	class FMSAACount : SHADER_PERMUTATION_SPARSE_INT("PERMUTATION_MSAACOUNT", 2, 4, 8);
-	using FPermutationDomain = TShaderPermutationDomain<FMSAACount>;
-
-	BEGIN_SHADER_PARAMETER_STRUCT(FParameters, )
-		SHADER_PARAMETER_STRUCT_INCLUDE(FHairStrandsTilePassVS::FParameters, TileData)
-		SHADER_PARAMETER(float, VelocityThreshold)
-		SHADER_PARAMETER_RDG_TEXTURE(Texture2D, HairVisibilityVelocityTexture)
-		RENDER_TARGET_BINDING_SLOTS()
-	END_SHADER_PARAMETER_STRUCT()
-
-public:
-	static bool ShouldCompilePermutation(const FGlobalShaderPermutationParameters& Parameters) { return IsHairStrandsSupported(EHairStrandsShaderType::Strands, Parameters.Platform); }
-	static void ModifyCompilationEnvironment(const FGlobalShaderPermutationParameters& Parameters, FShaderCompilerEnvironment& OutEnvironment)
-	{
-		FGlobalShader::ModifyCompilationEnvironment(Parameters, OutEnvironment);
-		OutEnvironment.SetDefine(TEXT("SHADER_FASTRESOLVE_MSAA"), 1);
-		OutEnvironment.SetRenderTargetOutputFormat(0, PF_R8G8B8A8);
-	}
-};
-
-IMPLEMENT_GLOBAL_SHADER(FHairVisibilityFastResolvePS, "/Engine/Private/HairStrands/HairStrandsVisibilityComposeSubPixelPS.usf", "FastResolvePS", SF_Pixel);
-
-static void AddHairVisibilityFastResolveMSAAPass(
-	FRDGBuilder& GraphBuilder,
-	const FViewInfo& View,
-	const FRDGTextureRef& HairVisibilityVelocityTexture,
-	const FHairStrandsTiles& TileData,
-	FRDGTextureRef& OutDepthTexture)
-{
-	const FIntPoint Resolution = OutDepthTexture->Desc.Extent;
-	FRDGTextureRef DummyTexture;
-	{
-		FRDGTextureDesc Desc;
-		Desc.Extent = Resolution;
-		Desc.Format = PF_R8G8B8A8;
-		Desc.Flags = TexCreate_RenderTargetable | TexCreate_ShaderResource;
-		Desc.ClearValue = FClearValueBinding(0);
-		DummyTexture = GraphBuilder.CreateTexture(Desc, TEXT("Hair.DummyTexture"));
-	}
-
-	FVector2D PixelVelocity(1.f / (Resolution.X * 2), 1.f / (Resolution.Y * 2));
-	const float VelocityThreshold = FMath::Clamp(GHairFastResolveVelocityThreshold, 0, 512) * FMath::Min(PixelVelocity.X, PixelVelocity.Y);
-
-	FHairVisibilityFastResolvePS::FParameters* Parameters = GraphBuilder.AllocParameters<FHairVisibilityFastResolvePS::FParameters>();
-	Parameters->HairVisibilityVelocityTexture = HairVisibilityVelocityTexture;
-	Parameters->VelocityThreshold = GetHairFastResolveVelocityThreshold(Resolution);
-	Parameters->RenderTargets[0] = FRenderTargetBinding(DummyTexture, ERenderTargetLoadAction::ENoAction);
-	Parameters->RenderTargets.DepthStencil = FDepthStencilBinding(
-		OutDepthTexture,
-		ERenderTargetLoadAction::ELoad,
-		ERenderTargetLoadAction::ELoad,
-		FExclusiveDepthStencil::DepthNop_StencilWrite);
-
-	const uint32 MSAASampleCount = HairVisibilityVelocityTexture->Desc.NumSamples;
-	check(MSAASampleCount == 4 || MSAASampleCount == 8);
-	FHairVisibilityFastResolvePS::FPermutationDomain PermutationVector;
-	PermutationVector.Set<FHairVisibilityFastResolvePS::FMSAACount>(MSAASampleCount == 4 ? 4 : 8);
-
-	TShaderMapRef<FHairVisibilityFastResolvePS> PixelShader(View.ShaderMap, PermutationVector);
-	InternalCommonDrawPass(
-		GraphBuilder,
-		RDG_EVENT_NAME("HairStrandsVisibilityMarkTAAFastResolve"),
-		View,
-		Resolution,
-		EHairStrandsCommonPassType::TAAFastResolve,
-		false,
-		TileData,
-		PixelShader,
-		Parameters);
-}
-
 class FHairVisibilityFastResolveMaskPS : public FGlobalShader
 {
 	DECLARE_GLOBAL_SHADER(FHairVisibilityFastResolveMaskPS);
@@ -581,17 +505,7 @@ static void InternalRenderHairComposition(
 					SceneColorTexture,
 					SceneDepthTexture);
 
-				
-				if (VisibilityData.VelocityTexture)
-				{
-					AddHairVisibilityFastResolveMSAAPass(
-						GraphBuilder,
-						View,
-						VisibilityData.VelocityTexture,
-						VisibilityData.TileData,
-						SceneDepthTexture);
-				}
-				else if (VisibilityData.ResolveMaskTexture)
+				if (VisibilityData.ResolveMaskTexture)
 				{
 					AddHairVisibilityFastResolveMaskPass(
 						GraphBuilder,
