Change hair resolve texture pixel format from R32 to R8 to save memory.

diff --git a/Engine/Source/Runtime/Renderer/Private/HairStrands/HairStrandsVisibility.cpp b/Engine/Source/Runtime/Renderer/Private/HairStrands/HairStrandsVisibility.cpp
--- a/Engine/Source/Runtime/Renderer/Private/HairStrands/HairStrandsVisibility.cpp
+++ b/Engine/Source/Runtime/Renderer/Private/HairStrands/HairStrandsVisibility.cpp
@@ -911,7 +911,7 @@ static void AddHairVelocityPass(
 
 	const FIntPoint Resolution = OutVelocityTexture->Desc.Extent;
 	{
-		FRDGTextureDesc Desc = FRDGTextureDesc::Create2D(Resolution, PF_R32_UINT, FClearValueBinding::None, TexCreate_UAV);
+		FRDGTextureDesc Desc = FRDGTextureDesc::Create2D(Resolution, PF_R8_UINT, FClearValueBinding::None, TexCreate_UAV);
 		OutResolveMaskTexture = GraphBuilder.CreateTexture(Desc, TEXT("Hair.VelocityResolveMaskTexture"));
 	}
 
