Fix AdaptiveVirtualTexture not processing first allocation request each frame.
This led to the VT not allocating and was particularly obvious when there is only one request in the queue.

diff --git a/Engine/Source/Runtime/Renderer/Private/VT/AdaptiveVirtualTexture.cpp b/Engine/Source/Runtime/Renderer/Private/VT/AdaptiveVirtualTexture.cpp
--- a/Engine/Source/Runtime/Renderer/Private/VT/AdaptiveVirtualTexture.cpp
+++ b/Engine/Source/Runtime/Renderer/Private/VT/AdaptiveVirtualTexture.cpp
@@ -412,7 +412,7 @@ void FAdaptiveVirtualTexture::QueuePackedAllocationRequests(uint32 const* InRequ
 	for (uint32 RequestIndex = 0; RequestIndex < InNumRequests; ++RequestIndex)
 	{
 		// Skip duplicates.
-		if (RequestIndex > 0 && InRequests[RequestIndex] != InRequests[RequestIndex - 1])
+		if (RequestIndex == 0 || InRequests[RequestIndex] != InRequests[RequestIndex - 1])
 		{
 			FPackedAdaptiveAllocationRequest Request;
 			Request.PackedValue = InRequests[RequestIndex];
