Change hair light sample texture to 11.11.10 to save memory.

diff --git a/Engine/Source/Runtime/Renderer/Private/HairStrands/HairStrandsVisibility.cpp b/Engine/Source/Runtime/Renderer/Private/HairStrands/HairStrandsVisibility.cpp
--- a/Engine/Source/Runtime/Renderer/Private/HairStrands/HairStrandsVisibility.cpp
+++ b/Engine/Source/Runtime/Renderer/Private/HairStrands/HairStrandsVisibility.cpp
@@ -89,6 +89,8 @@ static FAutoConsoleVariableRef CVarHairStrandsDebugPPLL(TEXT("r.HairStrands.Visi
 static int32 GHairStrandsTile = 1;
 static FAutoConsoleVariableRef CVarHairStrandsTile(TEXT("r.HairStrands.Tile"), GHairStrandsTile, TEXT("Enable tile generation & usage for hair strands."));
 
+static int32 GHairStrandsLightSampleFormat = 1;
+static FAutoConsoleVariableRef CVarHairStrandsLightSampleFormat(TEXT("r.HairStrands.LightSampleFormat"), GHairStrandsLightSampleFormat, TEXT("Define the format used for storing the lighting of hair samples (0: RGBA-16bits, 1: RGB-11.11.10bits)"));
 
 /////////////////////////////////////////////////////////////////////////////////////////
 
@@ -316,6 +318,9 @@ class FHairLightSampleClearPS : public FGlobalShader
 	DECLARE_GLOBAL_SHADER(FHairLightSampleClearPS);
 	SHADER_USE_PARAMETER_STRUCT(FHairLightSampleClearPS, FGlobalShader)
 
+	class FOutputFormat : SHADER_PERMUTATION_INT("PERMUTATION_OUTPUT_FORMAT", 2);
+	using FPermutationDomain = TShaderPermutationDomain<FOutputFormat>;
+
 	BEGIN_SHADER_PARAMETER_STRUCT(FParameters, )
 		SHADER_PARAMETER(FIntPoint, MaxViewportResolution)
 		SHADER_PARAMETER_RDG_TEXTURE(Texture2D, HairNodeCountTexture)
@@ -327,7 +332,16 @@ class FHairLightSampleClearPS : public FGlobalShader
 	{
 		FGlobalShader::ModifyCompilationEnvironment(Parameters, OutEnvironment);
 		OutEnvironment.SetDefine(TEXT("SHADER_CLEAR"), 1);
-		OutEnvironment.SetRenderTargetOutputFormat(0, PF_FloatRGBA);
+
+		FPermutationDomain PermutationVector(Parameters.PermutationId);
+		if (PermutationVector.Get<FOutputFormat>() == 0)
+		{
+			OutEnvironment.SetRenderTargetOutputFormat(0, PF_FloatRGBA);
+		}
+		else if (PermutationVector.Get<FOutputFormat>() == 1)
+		{
+			OutEnvironment.SetRenderTargetOutputFormat(0, PF_FloatR11G11B10);
+		}
 	}
 };
 
@@ -340,22 +354,31 @@ static FRDGTextureRef AddClearLightSamplePass(
 	const uint32 MaxNodeCount,
 	const FRDGTextureRef NodeCounter)
 {	
+	EPixelFormat Format = PF_FloatRGBA;
+	if (GHairStrandsLightSampleFormat > 0 && GPixelFormats[PF_FloatR11G11B10].Supported)
+	{
+		Format = PF_FloatR11G11B10;
+	}
+
 	const uint32 SampleTextureResolution = FMath::CeilToInt(FMath::Sqrt(static_cast<float>(MaxNodeCount)));
-	FRDGTextureDesc Desc = FRDGTextureDesc::Create2D(FIntPoint(SampleTextureResolution, SampleTextureResolution), PF_FloatRGBA, FClearValueBinding::Black, TexCreate_UAV | TexCreate_ShaderResource | TexCreate_RenderTargetable);
+	FRDGTextureDesc Desc = FRDGTextureDesc::Create2D(FIntPoint(SampleTextureResolution, SampleTextureResolution), Format, FClearValueBinding::Black, TexCreate_UAV | TexCreate_ShaderResource | TexCreate_RenderTargetable);
 	FRDGTextureRef Output = GraphBuilder.CreateTexture(Desc, TEXT("Hair.LightSample"));
 
 	FHairLightSampleClearPS::FParameters* ParametersPS = GraphBuilder.AllocParameters<FHairLightSampleClearPS::FParameters>();
 	ParametersPS->MaxViewportResolution = Desc.Extent;
 	ParametersPS->HairNodeCountTexture = NodeCounter;
 	
+	FHairLightSampleClearPS::FPermutationDomain PermutationVector;
+	PermutationVector.Set<FHairLightSampleClearPS::FOutputFormat>(Format == PF_FloatR11G11B10 ? 1 : 0);
+
 	const FIntPoint ViewportResolution = Desc.Extent;
 	TShaderMapRef<FHairLightSampleClearVS> VertexShader(View->ShaderMap);
-	TShaderMapRef<FHairLightSampleClearPS> PixelShader(View->ShaderMap);
+	TShaderMapRef<FHairLightSampleClearPS> PixelShader(View->ShaderMap, PermutationVector);
 
 	ParametersPS->RenderTargets[0] = FRenderTargetBinding(Output, ERenderTargetLoadAction::ENoAction);
 
 	GraphBuilder.AddPass(
-		RDG_EVENT_NAME("HairLightSampleClearPS"),
+		RDG_EVENT_NAME("HairStrands::LightSampleClearPS"),
 		ParametersPS,
 		ERDGPassFlags::Raster,
 		[ParametersPS, VertexShader, PixelShader, ViewportResolution](FRHICommandList& RHICmdList)
