Change hair visibility depth texture format to save memory.

This change the format from PF_DepthStencil (8bytes) to PF_D24 (3bytes).

diff --git a/Engine/Source/Runtime/Renderer/Private/HairStrands/HairStrandsVisibility.cpp b/Engine/Source/Runtime/Renderer/Private/HairStrands/HairStrandsVisibility.cpp
--- a/Engine/Source/Runtime/Renderer/Private/HairStrands/HairStrandsVisibility.cpp
+++ b/Engine/Source/Runtime/Renderer/Private/HairStrands/HairStrandsVisibility.cpp
@@ -1073,7 +1073,7 @@ TRDGUniformBufferRef<FVisibilityPassUniformParameters> CreatePassDummyTextures(F
 	FVisibilityPassUniformParameters* UniformParameters = GraphBuilder.AllocParameters<FVisibilityPassUniformParameters>();
 
 	FRDGTextureDesc Desc = FRDGTextureDesc::Create2D(FIntPoint(1,1), PF_R32_UINT, FClearValueBinding::None, TexCreate_UAV | TexCreate_ShaderResource);
-	UniformParameters->PPLLCounter		= GraphBuilder.CreateUAV(GraphBuilder.CreateTexture(Desc, TEXT("Hair.VisibilityPPLLNodeIndex")));
+	UniformParameters->PPLLCounter		= GraphBuilder.CreateUAV(GraphBuilder.CreateTexture(Desc, TEXT("Hair.VisibilityPPLLNodeCounter")));
 	UniformParameters->PPLLNodeIndex	= GraphBuilder.CreateUAV(GraphBuilder.CreateTexture(Desc, TEXT("Hair.VisibilityPPLLNodeIndex")));
 	UniformParameters->PPLLNodeData		= GraphBuilder.CreateUAV(GraphBuilder.CreateBuffer(FRDGBufferDesc::CreateStructuredDesc(sizeof(PPLLNodeData), 1), TEXT("Hair.DummyPPLLNodeData")));
 
@@ -1899,7 +1899,7 @@ static FRDGTextureRef AddHairVisibilityFillOpaqueDepth(
 	{
 		check(GetHairVisibilityRenderMode() == HairVisibilityRenderMode_MSAA_Visibility);
 
-		FRDGTextureDesc Desc = FRDGTextureDesc::Create2D(Resolution, PF_DepthStencil, FClearValueBinding::DepthFar, TexCreate_DepthStencilTargetable | TexCreate_ShaderResource, 1, GetMaxSamplePerPixel());
+		FRDGTextureDesc Desc = FRDGTextureDesc::Create2D(Resolution, PF_D24, FClearValueBinding::DepthFar, TexCreate_DepthStencilTargetable | TexCreate_ShaderResource, 1, GetMaxSamplePerPixel());
 		OutVisibilityDepthTexture = GraphBuilder.CreateTexture(Desc, TEXT("Hair.VisibilityDepthTexture"));
 	}
 
@@ -2101,7 +2101,7 @@ static void AddHairVisibilityMSAAPass(
 	{
 		{
 			FRDGTextureDesc Desc = FRDGTextureDesc::Create2D(Resolution, PF_R32_UINT, FClearValueBinding(EClearBinding::ENoneBound), TexCreate_NoFastClear | TexCreate_RenderTargetable | TexCreate_ShaderResource, 1, MSAASampleCount);
-			OutVisibilityIdTexture = GraphBuilder.CreateTexture(Desc, TEXT("HairVisibilityIDTexture"));
+			OutVisibilityIdTexture = GraphBuilder.CreateTexture(Desc, TEXT("Hair.VisibilityIDTexture"));
 		}
 
 		AddClearGraphicPass(GraphBuilder, RDG_EVENT_NAME("HairStrandsClearVisibilityMSAAIdTexture"), ViewInfo, 0xFFFFFFFF, OutVisibilityIdTexture);
