Prevent crash when selecting both forward rendering and hardware ray tracing.

Lumen hardware ray-tracing shaders were prepared for execution, but are compiled out when forward rendering is selected. Deferred ray-tracing shaders are similarly excluded when forward rendering is selected.

diff --git a/Engine/Source/Runtime/Renderer/Private/DeferredShadingRenderer.cpp b/Engine/Source/Runtime/Renderer/Private/DeferredShadingRenderer.cpp
--- a/Engine/Source/Runtime/Renderer/Private/DeferredShadingRenderer.cpp
+++ b/Engine/Source/Runtime/Renderer/Private/DeferredShadingRenderer.cpp
@@ -1319,19 +1319,26 @@ bool FDeferredShadingSceneRenderer::SetupRayTracingPipelineStates(FRHICommandLis
 
 		// These other cases do potentially depend on the camera position since they are
 		// driven by FinalPostProcessSettings, which is why we need to merge them across views
-		for (const FViewInfo& View : Views)
-		{
-			PrepareRayTracingReflections(View, *Scene, RayGenShaders);
-			PrepareSingleLayerWaterRayTracingReflections(View, *Scene, RayGenShaders);
-			PrepareRayTracingShadows(View, RayGenShaders);
-			PrepareRayTracingAmbientOcclusion(View, RayGenShaders);
-			PrepareRayTracingSkyLight(View, RayGenShaders);
-			PrepareRayTracingGlobalIllumination(View, RayGenShaders);
-			PrepareRayTracingTranslucency(View, RayGenShaders);
-			PrepareLumenHardwareRayTracingScreenProbeGather(View, RayGenShaders);
-			PrepareLumenHardwareRayTracingRadianceCache(View, RayGenShaders);
-			PrepareLumenHardwareRayTracingReflections(View, RayGenShaders);
-			PrepareLumenHardwareRayTracingVisualize(View, RayGenShaders);
+		if (!IsForwardShadingEnabled(ShaderPlatform))
+		{
+			for (const FViewInfo& View : Views)
+			{
+				PrepareRayTracingReflections(View, *Scene, RayGenShaders);
+				PrepareSingleLayerWaterRayTracingReflections(View, *Scene, RayGenShaders);
+				PrepareRayTracingShadows(View, RayGenShaders);
+				PrepareRayTracingAmbientOcclusion(View, RayGenShaders);
+				PrepareRayTracingSkyLight(View, RayGenShaders);
+				PrepareRayTracingGlobalIllumination(View, RayGenShaders);
+				PrepareRayTracingTranslucency(View, RayGenShaders);
+
+				if (DoesPlatformSupportLumenGI(ShaderPlatform))
+				{
+					PrepareLumenHardwareRayTracingScreenProbeGather(View, RayGenShaders);
+					PrepareLumenHardwareRayTracingRadianceCache(View, RayGenShaders);
+					PrepareLumenHardwareRayTracingReflections(View, RayGenShaders);
+					PrepareLumenHardwareRayTracingVisualize(View, RayGenShaders);
+				}
+			}
 		}
 		DeduplicateRayGenerationShaders(RayGenShaders);
 	}
@@ -1629,15 +1636,21 @@ void FDeferredShadingSceneRenderer::WaitForRayTracingScene(FRDGBuilder& GraphBui
 			if (!bIsPathTracing)
 			{
 				TArray<FRHIRayTracingShader*> DeferredMaterialRayGenShaders;
-				for (const FViewInfo& View : Views)
+				if (!IsForwardShadingEnabled(ShaderPlatform))
 				{
-					PrepareRayTracingReflectionsDeferredMaterial(View, *Scene, DeferredMaterialRayGenShaders);
-					PrepareRayTracingDeferredReflectionsDeferredMaterial(View, *Scene, DeferredMaterialRayGenShaders);
-					PrepareRayTracingGlobalIlluminationDeferredMaterial(View, DeferredMaterialRayGenShaders);
-					PrepareLumenHardwareRayTracingReflectionsDeferredMaterial(View, DeferredMaterialRayGenShaders);
-					PrepareLumenHardwareRayTracingRadianceCacheDeferredMaterial(View, DeferredMaterialRayGenShaders);
-					PrepareLumenHardwareRayTracingScreenProbeGatherDeferredMaterial(View, DeferredMaterialRayGenShaders);
-					PrepareLumenHardwareRayTracingVisualizeDeferredMaterial(View, DeferredMaterialRayGenShaders);
+					for (const FViewInfo& View : Views)
+					{
+						PrepareRayTracingReflectionsDeferredMaterial(View, *Scene, DeferredMaterialRayGenShaders);
+						PrepareRayTracingDeferredReflectionsDeferredMaterial(View, *Scene, DeferredMaterialRayGenShaders);
+						PrepareRayTracingGlobalIlluminationDeferredMaterial(View, DeferredMaterialRayGenShaders);
+						if (DoesPlatformSupportLumenGI(ShaderPlatform))
+						{
+							PrepareLumenHardwareRayTracingReflectionsDeferredMaterial(View, DeferredMaterialRayGenShaders);
+							PrepareLumenHardwareRayTracingRadianceCacheDeferredMaterial(View, DeferredMaterialRayGenShaders);
+							PrepareLumenHardwareRayTracingScreenProbeGatherDeferredMaterial(View, DeferredMaterialRayGenShaders);
+							PrepareLumenHardwareRayTracingVisualizeDeferredMaterial(View, DeferredMaterialRayGenShaders);
+						}
+					}
 				}
 				DeduplicateRayGenerationShaders(DeferredMaterialRayGenShaders);
 
@@ -1648,13 +1661,16 @@ void FDeferredShadingSceneRenderer::WaitForRayTracingScene(FRDGBuilder& GraphBui
 
 				// Add Lumen hardware ray tracing materials
 				TArray<FRHIRayTracingShader*> LumenHardwareRayTracingRayGenShaders;
-				for (const FViewInfo& View : Views)
+				if (DoesPlatformSupportLumenGI(ShaderPlatform))
 				{
-					PrepareLumenHardwareRayTracingVisualizeLumenMaterial(View, LumenHardwareRayTracingRayGenShaders);
-					PrepareLumenHardwareRayTracingRadianceCacheLumenMaterial(View, LumenHardwareRayTracingRayGenShaders);
-					PrepareLumenHardwareRayTracingReflectionsLumenMaterial(View, LumenHardwareRayTracingRayGenShaders);
-					PrepareLumenHardwareRayTracingScreenProbeGatherLumenMaterial(View, LumenHardwareRayTracingRayGenShaders);
-					PrepareLumenHardwareRayTracingDirectLightingLumenMaterial(View, LumenHardwareRayTracingRayGenShaders);
+					for (const FViewInfo& View : Views)
+					{
+						PrepareLumenHardwareRayTracingVisualizeLumenMaterial(View, LumenHardwareRayTracingRayGenShaders);
+						PrepareLumenHardwareRayTracingRadianceCacheLumenMaterial(View, LumenHardwareRayTracingRayGenShaders);
+						PrepareLumenHardwareRayTracingReflectionsLumenMaterial(View, LumenHardwareRayTracingRayGenShaders);
+						PrepareLumenHardwareRayTracingScreenProbeGatherLumenMaterial(View, LumenHardwareRayTracingRayGenShaders);
+						PrepareLumenHardwareRayTracingDirectLightingLumenMaterial(View, LumenHardwareRayTracingRayGenShaders);
+					}
 				}
 				DeduplicateRayGenerationShaders(DeferredMaterialRayGenShaders);
 
