Fix crash when creating UAV on null resources.

This can happen when ShaderDrawDebug is not supported/disable on certain platform.

diff --git a/Engine/Source/Runtime/Renderer/Private/ShaderDebug.cpp b/Engine/Source/Runtime/Renderer/Private/ShaderDebug.cpp
--- a/Engine/Source/Runtime/Renderer/Private/ShaderDebug.cpp
+++ b/Engine/Source/Runtime/Renderer/Private/ShaderDebug.cpp
@@ -313,6 +313,10 @@ namespace ShaderDrawDebug
 
 	void SetParameters(FRDGBuilder& GraphBuilder, const FShaderDrawDebugData& Data, FShaderParameters& OutParameters)
 	{
+		// Early out if debug rendering is not enabled/supported
+		if (Data.Buffer == nullptr)
+			return;
+
 		FRDGBufferRef DataBuffer = Data.Buffer;
 		FRDGBufferRef IndirectBuffer = Data.IndirectBuffer;
 
