Make sure only "Translucent" decal blend mode gets translated to "Additive" when emmisive input is connected on mobile platfroms

diff --git a/Engine/Source/Runtime/Renderer/Private/DecalRenderingCommon.cpp b/Engine/Source/Runtime/Renderer/Private/DecalRenderingCommon.cpp
--- a/Engine/Source/Runtime/Renderer/Private/DecalRenderingCommon.cpp
+++ b/Engine/Source/Runtime/Renderer/Private/DecalRenderingCommon.cpp
@@ -24,10 +24,6 @@ namespace DecalRendering
 		{
 			Desc.BlendMode = BLEND_Translucent;
 		}
-		if (bIsMobilePlatform && !bIsMobileDeferredPlatform && Desc.bWriteEmissive)
-		{
-			Desc.BlendMode = BLEND_Translucent;
-		}
 
 		// Enforce platform output limitations.
 		if (bIsMobilePlatform && !bIsMobileDeferredPlatform)
@@ -738,27 +734,27 @@ namespace DecalRendering
 	{
 		if (DecalRenderStage == EDecalRenderStage::Mobile)
 		{
-			if (DecalBlendDesc.bWriteEmissive)
+			if (DecalBlendDesc.BlendMode == BLEND_Translucent)
 			{
-				// Treat blend as emissive
-				return TStaticBlendState<CW_RGB, BO_Add, BF_SourceAlpha, BF_One>::GetRHI();
-			}
-			else
-			{
-				// Treat blend as non-emissive
-				if (DecalBlendDesc.BlendMode == BLEND_Translucent)
+				if (DecalBlendDesc.bWriteEmissive)
 				{
-					return TStaticBlendState<CW_RGB, BO_Add, BF_SourceAlpha, BF_InverseSourceAlpha>::GetRHI();
+					// Treat blend as emissive
+					return TStaticBlendState<CW_RGB, BO_Add, BF_SourceAlpha, BF_One>::GetRHI();
 				}
-				else if (DecalBlendDesc.BlendMode == BLEND_AlphaComposite)
+				else
 				{
-					return TStaticBlendState<CW_RGB, BO_Add, BF_One, BF_InverseSourceAlpha>::GetRHI();
-				}
-				else if (DecalBlendDesc.BlendMode == BLEND_Modulate)
-				{
-					return TStaticBlendState<CW_RGB, BO_Add, BF_DestColor, BF_InverseSourceAlpha>::GetRHI();
+					// Treat blend as non-emissive
+					return TStaticBlendState<CW_RGB, BO_Add, BF_SourceAlpha, BF_InverseSourceAlpha>::GetRHI();
 				}
 			}
+			else if (DecalBlendDesc.BlendMode == BLEND_AlphaComposite)
+			{
+				return TStaticBlendState<CW_RGB, BO_Add, BF_One, BF_InverseSourceAlpha>::GetRHI();
+			}
+			else if (DecalBlendDesc.BlendMode == BLEND_Modulate)
+			{
+				return TStaticBlendState<CW_RGB, BO_Add, BF_DestColor, BF_InverseSourceAlpha>::GetRHI();
+			}
 		}
 		else
 		{
