Fix deadlock during WaitForMeshPassSetupTask

#rnx

diff --git a/Engine/Source/Runtime/Renderer/Private/MeshDrawCommands.cpp b/Engine/Source/Runtime/Renderer/Private/MeshDrawCommands.cpp
--- a/Engine/Source/Runtime/Renderer/Private/MeshDrawCommands.cpp
+++ b/Engine/Source/Runtime/Renderer/Private/MeshDrawCommands.cpp
@@ -1411,7 +1411,7 @@ void FParallelMeshDrawCommandPass::DispatchDraw(FParallelCommandListSet* Paralle
 	{
 		QUICK_SCOPE_CYCLE_COUNTER(STAT_MeshPassDrawImmediate);
 
-		WaitForMeshPassSetupTask(EWaitThread::Task);
+		WaitForMeshPassSetupTask(IsInActualRenderingThread() ? EWaitThread::Render : EWaitThread::Task);
 
 		if (TaskContext.bUseGPUScene)
 		{
