Fixes SSR when reading from TSR history

diff --git a/Engine/Source/Runtime/Renderer/Private/ScreenSpaceRayTracing.cpp b/Engine/Source/Runtime/Renderer/Private/ScreenSpaceRayTracing.cpp
--- a/Engine/Source/Runtime/Renderer/Private/ScreenSpaceRayTracing.cpp
+++ b/Engine/Source/Runtime/Renderer/Private/ScreenSpaceRayTracing.cpp
@@ -846,8 +846,8 @@ FPrevSceneColorMip ReducePrevSceneColorMip(
 			PassParameters->PrevSceneColorSampler = TStaticSamplerState<SF_Bilinear>::GetRHI();
 
 			BufferSize = PassParameters->PrevSceneColor->Desc.Extent;
-			ViewportOffset = View.PrevViewInfo.TemporalAAHistory.ViewportRect.Min;
-			ViewportExtent = View.PrevViewInfo.TemporalAAHistory.ViewportRect.Size();
+			ViewportOffset = View.PrevViewInfo.TSRHistory.OutputViewportRect.Min;
+			ViewportExtent = View.PrevViewInfo.TSRHistory.OutputViewportRect.Size();
 		}
 		else
 		{
@@ -1052,6 +1052,10 @@ void RenderScreenSpaceReflections(
 		{
 			InputColor = GraphBuilder.RegisterExternalTexture(View.PrevViewInfo.HalfResTemporalAAHistory);
 		}
+		else if (View.PrevViewInfo.TSRHistory.IsValid())
+		{
+			InputColor = GraphBuilder.RegisterExternalTexture(View.PrevViewInfo.TSRHistory.LowFrequency);
+		}
 		else if (View.PrevViewInfo.TemporalAAHistory.IsValid())
 		{
 			InputColor = GraphBuilder.RegisterExternalTexture(View.PrevViewInfo.TemporalAAHistory.RT[0]);
@@ -1163,7 +1167,15 @@ void RenderScreenSpaceReflections(
 			FIntPoint ViewportExtent = View.ViewRect.Size();
 			FIntPoint BufferSize = SceneTextures.SceneDepthTexture->Desc.Extent;
 
-			if (View.PrevViewInfo.TemporalAAHistory.IsValid())
+			if (View.PrevViewInfo.TSRHistory.IsValid())
+			{
+				ViewportOffset = View.PrevViewInfo.TSRHistory.OutputViewportRect.Min;
+				ViewportExtent = View.PrevViewInfo.TSRHistory.OutputViewportRect.Size();
+				BufferSize = View.PrevViewInfo.TSRHistory.LowFrequency->GetDesc().Extent;
+				ensure(ViewportExtent.X > 0 && ViewportExtent.Y > 0);
+				ensure(BufferSize.X > 0 && BufferSize.Y > 0);
+			}
+			else if (View.PrevViewInfo.TemporalAAHistory.IsValid())
 			{
 				ViewportOffset = View.PrevViewInfo.TemporalAAHistory.ViewportRect.Min;
 				ViewportExtent = View.PrevViewInfo.TemporalAAHistory.ViewportRect.Size();
