Change visibility data for hair to reduce memory footprint.

diff --git a/Engine/Shaders/Private/HairStrands/HairStrandsMaterialPS.usf b/Engine/Shaders/Private/HairStrands/HairStrandsMaterialPS.usf
--- a/Engine/Shaders/Private/HairStrands/HairStrandsMaterialPS.usf
+++ b/Engine/Shaders/Private/HairStrands/HairStrandsMaterialPS.usf
@@ -20,7 +20,7 @@ void Main(
 	float3 Emissive = 0;
 	if (SampleIndex < TotalNodeCount)
 	{
-		const FNodeVis InNode = MaterialPassParameters.NodeVis[SampleIndex];
+		const FHairVis InNode = UnpackHairVis(MaterialPassParameters.NodeVis[SampleIndex]);
 		const float2 SamplePixelCoord = float2(MaterialPassParameters.NodeCoord[SampleIndex]) + float2(0.5f, 0.5f);
 
 		float SampleDepth = InNode.Depth;
diff --git a/Engine/Shaders/Private/HairStrands/HairStrandsVelocity.usf b/Engine/Shaders/Private/HairStrands/HairStrandsVelocity.usf
--- a/Engine/Shaders/Private/HairStrands/HairStrandsVelocity.usf
+++ b/Engine/Shaders/Private/HairStrands/HairStrandsVelocity.usf
@@ -19,7 +19,7 @@ float CoverageThreshold;
 
 Texture2D<uint> NodeIndex;
 Buffer<float2> NodeVelocity;
-StructuredBuffer<FNodeVis> NodeVis; 
+StructuredBuffer<FPackedHairVis> NodeVis; 
 Texture2D<float> CoverageTexture;
 
 #define VELOCITY_TYPE_NONE 0
@@ -78,7 +78,7 @@ void MainCS(uint GroupIndex : SV_GroupIndex, uint3 DispatchThreadId : SV_Dispatc
 				#endif			
 
 				#if PERMUTATION_VELOCITY == VELOCITY_TYPE_CLOSEST
-				const float NodeDepth = NodeVis[SampleIndex].Depth;
+				const float NodeDepth = UnpackHairVis(NodeVis[SampleIndex]).Depth;
 				if (NodeDepth > ClosestDepth) // Inverse-Z
 				{
 					ClosestEncodedVelocity = EncodedVelocity;
diff --git a/Engine/Shaders/Private/HairStrands/HairStrandsVisibilityCommon.ush b/Engine/Shaders/Private/HairStrands/HairStrandsVisibilityCommon.ush
--- a/Engine/Shaders/Private/HairStrands/HairStrandsVisibilityCommon.ush
+++ b/Engine/Shaders/Private/HairStrands/HairStrandsVisibilityCommon.ush
@@ -109,28 +109,8 @@ float2 UnpackVelocity(uint PackedVelocity)
 }
 
 ////////////////////////////////////////////////////////////////////////////////////////////////////////////////////////////////
-// Hair samples
-
-// Footprint: 128bits | 16bytes
-// Option to reduce the footprint: 
-// * we don't use PrimitiveID at the moment which is encoded on 3bytes
-// * Deferred the evaluation of Coverage & material properties into some other pass?
-struct FHairSample
-{
-	float	Depth;
-	float3	Tangent;
-	uint	Coverage8bit;
-	uint	PrimitiveId;
-	uint	MacroGroupId;
-	float3	BaseColor;
-	float	Roughness;
-	float	Specular;
-	float	Backlit;
-	uint	LightChannelMask;
-	bool	bScatterSceneLighting;
-	float3  Emissive;
-};
-
+// Node Desc
+// 
 // * Max 67,108,864 total nodes (This is 32 nodes per pixel at 1080p if all pixel were covered)
 // * Max 63 nodes per pixel
 // 26bits for offset | 6 bits for count (max 63 nodes)
@@ -154,6 +134,106 @@ FNodeDesc DecodeNodeDesc(uint In)
 	return Out;
 }
 
+////////////////////////////////////////////////////////////////////////////////////////////////////////////////////////////////
+// Hair Vis
+
+uint PackHairVisDepthCoverage(float InDepth, float InCoverage)
+{
+	return (uint(InDepth * 0x00FFFFFF) << 8) | uint(0xFF * saturate(InCoverage));
+}
+
+uint PackHairVisDepthCoverage(float InDepth, uint InCoverage8bits)
+{
+	return (uint(InDepth * 0x00FFFFFF) << 8) | min(0xFF, InCoverage8bits);
+}
+
+uint PackHairVisPrimitiveMaterialId(uint PrimitiveId, uint MaterialId)
+{
+	return (PrimitiveId & 0x00FFFFFF) | ((MaterialId & 0xFF) << 24);
+}
+
+uint UnpackHairVisPrimitiveId(uint RawVis)
+{
+	return RawVis & 0x00FFFFFF;
+}
+
+uint UnpackHairVisMaterialId(uint RawVis)
+{
+	return (RawVis >> 24) & 0xFF;
+}
+
+float UnpackHairVisDepth(uint InPacked)
+{
+	return float(InPacked >> 8) / float(0x00FFFFFF);
+}
+
+float UnpackHairVisCoverage(uint InPacked)
+{
+	return float(InPacked & 0xFF) / float(0xFF);
+}
+
+uint UnpackHairVisCoverage8bits(uint InPacked)
+{
+	return InPacked & 0xFF;
+}
+
+uint GetInvalidHairPrimitiveId()
+{
+	return 0xFFFFFFFF;
+}
+
+struct FHairVis
+{
+	float	Depth;
+	float   Coverage;	  // Coverage in float point, but with 8-bits precision. For convenience.
+	uint	Coverage8bit;
+	uint	PrimitiveId;
+	uint	MaterialId;
+};
+
+FPackedHairVis PackHairVis(const FHairVis In)
+{
+	FPackedHairVis Out;
+	Out.Depth_Coverage8bit	   = PackHairVisDepthCoverage(In.Depth, In.Coverage8bit);
+	Out.PrimitiveID_MaterialID = PackHairVisPrimitiveMaterialId(In.PrimitiveId, In.MaterialId);
+	return Out;
+}
+
+FHairVis UnpackHairVis(const FPackedHairVis In)
+{
+	FHairVis Out;
+	Out.Depth		 = UnpackHairVisDepth(In.Depth_Coverage8bit);
+	Out.Coverage8bit = UnpackHairVisCoverage8bits(In.Depth_Coverage8bit);
+	Out.Coverage	 = float(Out.Coverage8bit) / float(0xFF);
+	Out.PrimitiveId  = UnpackHairVisPrimitiveId(In.PrimitiveID_MaterialID);
+	Out.MaterialId	 = UnpackHairVisMaterialId(In.PrimitiveID_MaterialID);
+
+	return Out;
+}
+
+////////////////////////////////////////////////////////////////////////////////////////////////////////////////////////////////
+// Hair sample
+
+// Footprint: 128bits | 16bytes
+// Option to reduce the footprint: 
+// * we don't use PrimitiveID at the moment which is encoded on 3bytes
+// * Deferred the evaluation of Coverage & material properties into some other pass?
+struct FHairSample
+{
+	float	Depth;
+	float3	Tangent;
+	uint	Coverage8bit;
+	uint	PrimitiveId;
+	uint	MacroGroupId;
+	float3	BaseColor;
+	float	Roughness;
+	float	Specular;
+	float	Backlit;
+	uint	LightChannelMask;
+	bool	bScatterSceneLighting;
+	float3  Emissive;
+};
+
 FPackedHairSample PackHairSample(const FHairSample In)
 {
 	FPackedHairSample Out;
@@ -358,57 +438,7 @@ FPackedHairSample ConvertToPackedHairSample(FPPLLNodeData In, uint InNewCoverage
 }
 
  ////////////////////////////////////////////////////////////////////////////////////////////////////////////////////////////////
- // Hair visibility buffer 
-
-// Renamed into FHairNodeVisFat
-struct FNodeVis // 128 bit aligments. TODO: Move the coverage computation into the material evaluation for saving 64 bits? 
-{
-	float Depth;		// 32 bits
-	uint  PrimitiveId;	// 32 bits
-	uint  MaterialId;	//  8 bits : this could be packed into a the 3 high bits of PrimitiveId: up to 8 material on screen
-	uint  Coverage8bit;	//  8 bits
-};
-
-struct FHairVisSample
-{
-	uint PrimitiveId;
-	uint MaterialId;
-};
-
-uint GetInvalidHairPrimitiveId()
-{
-return 0xFFFFFFFF;
-}
-
-uint PackHairVisSample(uint PrimitiveId, uint MaterialId)
-{
-	return (PrimitiveId & 0x00FFFFFF) | ((MaterialId & 0xFF) << 24);
-}
-
-FHairVisSample UnpackHairVisSample(uint RawVis)
-{
-	FHairVisSample Out;
-	Out.PrimitiveId = RawVis & 0x00FFFFFF;
-	Out.MaterialId = (RawVis >> 24) & 0xFF;
-	return Out;
-}
- 
-// Pack depth & coverage within a uint32
-uint PackHairVisSampleDepth(float InDepth, float InCoverage)
-{
-	//return asuint(Depth);
-	return (uint(InDepth * 0x00FFFFFF) << 8) | uint(0xFF * saturate(InCoverage));
-}
-
-// Unpack depth & coverage from a uint32
-float2 UnpackHairVisSampleDepth(uint InPacked)
-{
-	//return asfloat(InPacked);
-	float OutDepth = float(InPacked >> 8) / float(0x00FFFFFF);
-	float OutCoverage = float(InPacked & 0xFF) / float(0xFF);
-
-	return float2(OutDepth, OutCoverage);
-}
+ // Coverage and hair count
 
 uint PackCoverageAndHairCount(float InCoverage)
 {
diff --git a/Engine/Shaders/Private/HairStrands/HairStrandsVisibilityCommonStruct.ush b/Engine/Shaders/Private/HairStrands/HairStrandsVisibilityCommonStruct.ush
--- a/Engine/Shaders/Private/HairStrands/HairStrandsVisibilityCommonStruct.ush
+++ b/Engine/Shaders/Private/HairStrands/HairStrandsVisibilityCommonStruct.ush
@@ -22,3 +22,9 @@ struct FPackedHairSample
 	uint  BaseColor_Roughness;
 	uint  Specular_LightChannels_Backlit;
 };
+
+struct FPackedHairVis
+{
+	uint Depth_Coverage8bit;
+	uint PrimitiveID_MaterialID;
+};
\ No newline at end of file
diff --git a/Engine/Shaders/Private/HairStrands/HairStrandsVisibilityCompaction.usf b/Engine/Shaders/Private/HairStrands/HairStrandsVisibilityCompaction.usf
--- a/Engine/Shaders/Private/HairStrands/HairStrandsVisibilityCompaction.usf
+++ b/Engine/Shaders/Private/HairStrands/HairStrandsVisibilityCompaction.usf
@@ -3,6 +3,7 @@
 #pragma once
 
 #include "../Common.ush"
+#include "HairStrandsVisibilityCommonStruct.ush" 
 #include "HairStrandsVisibilityCommon.ush" 
 
 #if PERMUTATION_GROUPSIZE ==32
@@ -27,7 +28,7 @@ uint VelocityType;
 RWTexture2D<uint> OutCompactNodeCounter;
 RWTexture2D<uint> OutCompactNodeIndex;
 #if PERMUTATION_PPLL == 0
-	RWStructuredBuffer<FNodeVis> OutCompactNodeData;
+	RWStructuredBuffer<FPackedHairVis> OutCompactNodeVis;
 #else
 	RWStructuredBuffer<FPackedHairSample> OutCompactNodeData;
 #endif
@@ -477,19 +478,17 @@ void MainCS(uint GroupIndex : SV_GroupIndex, uint3 DispatchThreadId : SV_Dispatc
 			float ClosestDepth = 0; // Inverse-Z
 			for (uint OutIndex = 0; OutIndex < NodeDesc.Count; ++OutIndex)
 			{
-				const FHairVisSample VisSample = UnpackHairVisSample(Samples[OutIndex].x);
-
 				// VisibilityData.Coverage8bit is a weight normalising to 1 the contribution of all the compacted samples. Because later it is weighted by PixelCoverage.
-				FNodeVis OutNodeVis;
-				OutNodeVis.PrimitiveId	= VisSample.PrimitiveId;
+				FHairVis OutNodeVis;
+				OutNodeVis.PrimitiveId	= UnpackHairVisPrimitiveId(Samples[OutIndex].x);
 				OutNodeVis.Depth		= asfloat(Samples[OutIndex].w); // MSAA_DepthTexture.Load(PixelCoord, Samples[OutIndex].z);
 				OutNodeVis.Coverage8bit = To8bitCoverage(Samples[OutIndex].y / float(SampleDesc.ValidSampleCount));
-				OutNodeVis.MaterialId	= VisSample.MaterialId;
+				OutNodeVis.MaterialId	= UnpackHairVisMaterialId(Samples[OutIndex].x);
 
 				ClosestDepth = max(ClosestDepth, OutNodeVis.Depth); // Inverse-Z
 
 				const uint StoreIndex = NodeDesc.Offset + OutIndex;
-				OutCompactNodeData[StoreIndex] = OutNodeVis;
+				OutCompactNodeVis[StoreIndex] = PackHairVis(OutNodeVis);
 				OutCompactNodeCoord[StoreIndex] = PixelCoord;
 			}
 
diff --git a/Engine/Shaders/Private/HairStrands/HairStrandsVisibilityCompactionComputeRaster.usf b/Engine/Shaders/Private/HairStrands/HairStrandsVisibilityCompactionComputeRaster.usf
--- a/Engine/Shaders/Private/HairStrands/HairStrandsVisibilityCompactionComputeRaster.usf
+++ b/Engine/Shaders/Private/HairStrands/HairStrandsVisibilityCompactionComputeRaster.usf
@@ -3,6 +3,7 @@
 #pragma once
 
 #include "../Common.ush"
+#include "HairStrandsVisibilityCommonStruct.ush" 
 #include "HairStrandsVisibilityCommon.ush" 
 
 #if PERMUTATION_GROUPSIZE == 0
@@ -36,11 +37,11 @@ uint MaxNodeCount;
 uint SamplerPerPixel;
 float CoverageThreshold;
 
-RWTexture2D<uint> 				OutCompactNodeCounter;
-RWTexture2D<uint> 				OutCompactNodeIndex;
-RWStructuredBuffer<FNodeVis> 	OutCompactNodeData;
-RWBuffer<uint2>					OutCompactNodeCoord;
-RWTexture2D<float> 				OutCoverageTexture;
+RWTexture2D<uint> 					OutCompactNodeCounter;
+RWTexture2D<uint> 					OutCompactNodeIndex;
+RWStructuredBuffer<FPackedHairVis> 	OutCompactNodeVis;
+RWBuffer<uint2>						OutCompactNodeCoord;
+RWTexture2D<float> 					OutCoverageTexture;
 
 groupshared uint AllocationNodeCount;
 groupshared uint AllocationNodeOffset;
@@ -99,64 +100,68 @@ void MainCS(uint GroupIndex : SV_GroupIndex, uint3 DispatchThreadId : SV_Dispatc
 	const float DepthThres = OpaqueDepth;
 	if (VisData0.y)
 	{
-		const FHairVisSample VisSample	= UnpackHairVisSample(VisData0.x);
-		const float2 Depth_Cov			= UnpackHairVisSampleDepth(VisData0.y);
+		const uint VisSamplePrimitiveId = UnpackHairVisPrimitiveId(VisData0.x);
+		const uint VisSampleMaterialId  = UnpackHairVisMaterialId(VisData0.x);
+		const float2 Depth_Cov = float2(UnpackHairVisDepth(VisData0.y), UnpackHairVisCoverage(VisData0.y));
 
-		Vis_PrimId_MatId[UniqueSampleCount]	= uint2(VisSample.PrimitiveId, VisSample.MaterialId);
+		Vis_PrimId_MatId[UniqueSampleCount]	= uint2(VisSamplePrimitiveId, VisSampleMaterialId);
 		Vis_Depth_Cov[UniqueSampleCount]	= Depth_Cov;
 
-		if (Depth_Cov.x > DepthThres && VisSample.PrimitiveId != CurrentPrimitiveId)
+		if (Depth_Cov.x > DepthThres && VisSamplePrimitiveId != CurrentPrimitiveId)
 		{
 			TotalCoverage += Vis_Depth_Cov[UniqueSampleCount].y;
-			CurrentPrimitiveId = VisSample.PrimitiveId;
+			CurrentPrimitiveId = VisSamplePrimitiveId;
 			UniqueSampleCount++;
 		}
 	}
 
 	if (VisData1.y && SamplerPerPixel > 1)
 	{
-		const FHairVisSample VisSample	= UnpackHairVisSample(VisData1.x);
-		const float2 Depth_Cov			= UnpackHairVisSampleDepth(VisData1.y);
+		const uint VisSamplePrimitiveId = UnpackHairVisPrimitiveId(VisData1.x);
+		const uint VisSampleMaterialId  = UnpackHairVisMaterialId(VisData1.x);
+		const float2 Depth_Cov			= float2(UnpackHairVisDepth(VisData1.y), UnpackHairVisCoverage(VisData1.y));
 
-		Vis_PrimId_MatId[UniqueSampleCount]	= uint2(VisSample.PrimitiveId, VisSample.MaterialId);
+		Vis_PrimId_MatId[UniqueSampleCount]	= uint2(VisSamplePrimitiveId, VisSampleMaterialId);
 		Vis_Depth_Cov[UniqueSampleCount]	= Depth_Cov;
 		
-		if (Depth_Cov.x > DepthThres && VisSample.PrimitiveId != CurrentPrimitiveId)
+		if (Depth_Cov.x > DepthThres && VisSamplePrimitiveId != CurrentPrimitiveId)
 		{
 			TotalCoverage += Vis_Depth_Cov[UniqueSampleCount].y;
-			CurrentPrimitiveId = VisSample.PrimitiveId;
+			CurrentPrimitiveId = VisSamplePrimitiveId;
 			UniqueSampleCount++;
 		}
 	}
 
 	if (VisData2.y && SamplerPerPixel > 2)
 	{
-		const FHairVisSample VisSample	= UnpackHairVisSample(VisData2.x);
-		const float2 Depth_Cov			= UnpackHairVisSampleDepth(VisData2.y);
+		const uint VisSamplePrimitiveId = UnpackHairVisPrimitiveId(VisData2.x);
+		const uint VisSampleMaterialId  = UnpackHairVisMaterialId(VisData2.x);
+		const float2 Depth_Cov			= float2(UnpackHairVisDepth(VisData2.y), UnpackHairVisCoverage(VisData2.y));
 
-		Vis_PrimId_MatId[UniqueSampleCount]	= uint2(VisSample.PrimitiveId, VisSample.MaterialId);
+		Vis_PrimId_MatId[UniqueSampleCount]	= uint2(VisSamplePrimitiveId, VisSampleMaterialId);
 		Vis_Depth_Cov[UniqueSampleCount]	= Depth_Cov;
 		
-		if (Depth_Cov.x > DepthThres && VisSample.PrimitiveId != CurrentPrimitiveId)
+		if (Depth_Cov.x > DepthThres && VisSamplePrimitiveId != CurrentPrimitiveId)
 		{
 			TotalCoverage += Vis_Depth_Cov[UniqueSampleCount].y;
-			CurrentPrimitiveId = VisSample.PrimitiveId;
+			CurrentPrimitiveId = VisSamplePrimitiveId;
 			UniqueSampleCount++;
 		}
 	}
 
 	if (VisData3.y && SamplerPerPixel > 3)
 	{
-		const FHairVisSample VisSample	= UnpackHairVisSample(VisData3.x);
-		const float2 Depth_Cov			= UnpackHairVisSampleDepth(VisData3.y);
+		const uint VisSamplePrimitiveId = UnpackHairVisPrimitiveId(VisData3.x);
+		const uint VisSampleMaterialId  = UnpackHairVisMaterialId(VisData3.x);
+		const float2 Depth_Cov			= float2(UnpackHairVisDepth(VisData3.y), UnpackHairVisCoverage(VisData3.y));
 
-		Vis_PrimId_MatId[UniqueSampleCount]	 = uint2(VisSample.PrimitiveId, VisSample.MaterialId);
+		Vis_PrimId_MatId[UniqueSampleCount]	 = uint2(VisSamplePrimitiveId, VisSampleMaterialId);
 		Vis_Depth_Cov[UniqueSampleCount]	 = Depth_Cov;
 		
-		if (Depth_Cov.x > DepthThres && VisSample.PrimitiveId != CurrentPrimitiveId)
+		if (Depth_Cov.x > DepthThres && VisSamplePrimitiveId != CurrentPrimitiveId)
 		{
 			TotalCoverage += Vis_Depth_Cov[UniqueSampleCount].y;
-			CurrentPrimitiveId = VisSample.PrimitiveId;
+			CurrentPrimitiveId = VisSamplePrimitiveId;
 			UniqueSampleCount++;
 		}
 	}
@@ -188,14 +193,14 @@ void MainCS(uint GroupIndex : SV_GroupIndex, uint3 DispatchThreadId : SV_Dispatc
 			for (uint OutIndex = 0; OutIndex < NodeDesc.Count; ++OutIndex)
 			{
 				// VisibilityData.Coverage8bit is a weight normalising to 1 the contribution of all the compacted samples. Because later it is weighted by Categorization.PixelCoverage.
-				FNodeVis OutNodeVis;
+				FHairVis OutNodeVis;
 				OutNodeVis.PrimitiveId	= Vis_PrimId_MatId[OutIndex].x;
 				OutNodeVis.MaterialId	= Vis_PrimId_MatId[OutIndex].y;
 				OutNodeVis.Depth		= Vis_Depth_Cov[OutIndex].x;
 				OutNodeVis.Coverage8bit = To8bitCoverage(Vis_Depth_Cov[OutIndex].y / TotalCoverage); 
 
 				const uint StoreIndex = NodeDesc.Offset + OutIndex;
-				OutCompactNodeData[StoreIndex] = OutNodeVis;
+				OutCompactNodeVis[StoreIndex] = PackHairVis(OutNodeVis);
 				OutCompactNodeCoord[StoreIndex] = PixelCoord;
 			}
 
diff --git a/Engine/Shaders/Private/HairStrands/HairStrandsVisibilityPS.usf b/Engine/Shaders/Private/HairStrands/HairStrandsVisibilityPS.usf
--- a/Engine/Shaders/Private/HairStrands/HairStrandsVisibilityPS.usf
+++ b/Engine/Shaders/Private/HairStrands/HairStrandsVisibilityPS.usf
@@ -191,7 +191,7 @@ void MainVisibility(
 
 #if HAIR_RENDER_MODE == RENDER_MODE_MSAA_VISIBILITY
   #if HAIR_STRAND_MESH_FACTORY
-	OutColor0 = PackHairVisSample(HairPrimitiveId, HairVisibilityPass_HairMaterialId);
+	OutColor0 = PackHairVisPrimitiveMaterialId(HairPrimitiveId, HairVisibilityPass_HairMaterialId);
   #else
 	OutColor0 = 0;
   #endif
diff --git a/Engine/Shaders/Private/HairStrands/HairStrandsVisibilityRasterCompute.usf b/Engine/Shaders/Private/HairStrands/HairStrandsVisibilityRasterCompute.usf
--- a/Engine/Shaders/Private/HairStrands/HairStrandsVisibilityRasterCompute.usf
+++ b/Engine/Shaders/Private/HairStrands/HairStrandsVisibilityRasterCompute.usf
@@ -170,8 +170,8 @@ void InternalRasterize(int2 Coords, float Ratio, float Depth, uint PrimitiveId,
 		InterlockedAdd(OutHairCountTexture[uint2(Coords)], PackedViewTransmittance);
 
 		uint2 PackedVisibility;
-		PackedVisibility.x = PackHairVisSample(PrimitiveId, HairMaterialId);
-		PackedVisibility.y = PackHairVisSampleDepth(Depth, Ratio * Coverage);
+		PackedVisibility.x = PackHairVisPrimitiveMaterialId(PrimitiveId, HairMaterialId);
+		PackedVisibility.y = PackHairVisDepthCoverage(Depth, Ratio * Coverage);
 
 		PackedType In = ToPackedType(PackedVisibility);
 		PackedType Out = 0;
diff --git a/Engine/Source/Runtime/Renderer/Private/HairStrands/HairStrandsVisibility.cpp b/Engine/Source/Runtime/Renderer/Private/HairStrands/HairStrandsVisibility.cpp
--- a/Engine/Source/Runtime/Renderer/Private/HairStrands/HairStrandsVisibility.cpp
+++ b/Engine/Source/Runtime/Renderer/Private/HairStrands/HairStrandsVisibility.cpp
@@ -105,13 +105,11 @@ namespace HairStrandsVisibilityInternal
 		uint32 Specular;
 	};
 
-	// 128 bit alignment
+	// 64 bit alignment
 	struct NodeVis
 	{
-		uint32 Depth;
+		uint32 Depth_Coverage;
 		uint32 PrimitiveId_MacroGroupId;
-		uint32 Coverage_MacroGroupId_Pad;
-		uint32 Pad;
 	};
 }
 
@@ -701,7 +699,7 @@ BEGIN_GLOBAL_SHADER_PARAMETER_STRUCT(FVisibilityMaterialPassUniformParameters, )
 	SHADER_PARAMETER_RDG_TEXTURE(Texture2D<uint>, NodeIndex)
 	SHADER_PARAMETER_RDG_TEXTURE(Texture2D<uint>, TotalNodeCounter)
 	SHADER_PARAMETER_RDG_BUFFER_SRV(Buffer<uint2>, NodeCoord)
-	SHADER_PARAMETER_RDG_BUFFER_SRV(StructuredBuffer<FNodeVis>, NodeVis)
+	SHADER_PARAMETER_RDG_BUFFER_SRV(StructuredBuffer<FPackedHairVis>, NodeVis)
 	SHADER_PARAMETER_RDG_BUFFER_SRV(Buffer<uint>, IndirectArgs)
 	SHADER_PARAMETER_RDG_BUFFER_UAV(RWStructuredBuffer<FPackedHairSample>, OutNodeData)
 	SHADER_PARAMETER_RDG_BUFFER_UAV(RWBuffer<float2>, OutNodeVelocity)
@@ -898,7 +896,7 @@ class FHairVelocityCS: public FGlobalShader
 		SHADER_PARAMETER_RDG_TEXTURE(Texture2D, CoverageTexture)
 		SHADER_PARAMETER_RDG_TEXTURE(Texture2D, NodeIndex)
 		SHADER_PARAMETER_RDG_BUFFER_SRV(Buffer, NodeVelocity)
-		SHADER_PARAMETER_RDG_BUFFER_SRV(StructuredBuffer<FNodeVis>, NodeVis)
+		SHADER_PARAMETER_RDG_BUFFER_SRV(StructuredBuffer<FPackedHairVis>, NodeVis)
 		SHADER_PARAMETER_RDG_TEXTURE_UAV(RWTexture2D, OutVelocityTexture)
 		SHADER_PARAMETER_RDG_TEXTURE_UAV(RWTexture2D, OutResolveMaskTexture)
 		SHADER_PARAMETER_STRUCT_REF(FViewUniformShaderParameters, ViewUniformBuffer)
@@ -1582,6 +1580,7 @@ class FHairVisibilityPrimitiveIdCompactionCS : public FGlobalShader
 		SHADER_PARAMETER_RDG_TEXTURE_UAV(RWTexture2D, OutCompactNodeCounter)
 		SHADER_PARAMETER_RDG_TEXTURE_UAV(RWTexture2D, OutCompactNodeIndex)
 		SHADER_PARAMETER_RDG_TEXTURE_UAV(RWTexture2D, OutCoverageTexture)
+		SHADER_PARAMETER_RDG_BUFFER_UAV(RWStructuredBuffer, OutCompactNodeVis)
 		SHADER_PARAMETER_RDG_BUFFER_UAV(RWStructuredBuffer, OutCompactNodeData)
 		SHADER_PARAMETER_RDG_BUFFER_UAV(RWBuffer, OutCompactNodeCoord)
 		SHADER_PARAMETER_RDG_TEXTURE_UAV(RWTexture2D, OutVelocityTexture)
@@ -1628,7 +1627,7 @@ static void AddHairVisibilityPrimitiveIdCompactionPass(
 	FHairVisibilityPrimitiveIdCompactionCS::FParameters* PassParameters,
 	FRDGTextureRef& OutCompactCounter,
 	FRDGTextureRef& OutCompactNodeIndex,
-	FRDGBufferRef& OutCompactNodeData,
+	FRDGBufferRef& OutCompactNodeVis, // Or OutCompactNodeData for PPLL
 	FRDGBufferRef& OutCompactNodeCoord,
 	FRDGTextureRef& OutCoverageTexture,
 	FRDGTextureRef& OutVelocityTexture,
@@ -1680,11 +1679,19 @@ static void AddHairVisibilityPrimitiveIdCompactionPass(
 	const uint32 PPLLMaxRenderNodePerPixel = GetMaxSamplePerPixel();
 	const uint32 MaxRenderNodeCount = GetTotalSampleCountForAllocation(EffectiveResolution);
 
-	OutCompactNodeData = GraphBuilder.CreateBuffer(FRDGBufferDesc::CreateStructuredDesc(sizeof(HairStrandsVisibilityInternal::NodeVis), MaxRenderNodeCount), TEXT("Hair.VisibilityPrimitiveIdCompactNodeData"));
+	if (bUsePPLL)
+	{
+		// PPLL output directly the node data
+		OutCompactNodeVis = GraphBuilder.CreateBuffer(FRDGBufferDesc::CreateStructuredDesc(sizeof(HairStrandsVisibilityInternal::NodeData), MaxRenderNodeCount), TEXT("Hair.VisibilityNodeData"));
+	}
+	else
+	{
+		OutCompactNodeVis = GraphBuilder.CreateBuffer(FRDGBufferDesc::CreateStructuredDesc(sizeof(HairStrandsVisibilityInternal::NodeVis), MaxRenderNodeCount), TEXT("Hair.VisibilityNodeVis"));
+	}
 
 	{
-		// Pixel coord of the node. Stored as 2*uint16, packed into a single uint32
-		OutCompactNodeCoord = GraphBuilder.CreateBuffer(FRDGBufferDesc::CreateBufferDesc(sizeof(uint32), MaxRenderNodeCount), TEXT("Hair.VisibilityPrimitiveIdCompactNodeCoord"));
+		// Pixel coord of the node. Stored as 2*R16_UINT
+		OutCompactNodeCoord = GraphBuilder.CreateBuffer(FRDGBufferDesc::CreateBufferDesc(sizeof(uint32), MaxRenderNodeCount), TEXT("Hair.VisibilityNodeCoord"));
 	}
 
 	const bool bWriteOutVelocity = OutVelocityTexture != nullptr;
@@ -1709,7 +1716,10 @@ static void AddHairVisibilityPrimitiveIdCompactionPass(
 	PassParameters->ViewUniformBuffer = View.ViewUniformBuffer;
 	PassParameters->OutCompactNodeCounter = GraphBuilder.CreateUAV(OutCompactCounter);
 	PassParameters->OutCompactNodeIndex = GraphBuilder.CreateUAV(OutCompactNodeIndex);
-	PassParameters->OutCompactNodeData = GraphBuilder.CreateUAV(OutCompactNodeData);
+	if (bUsePPLL)
+		PassParameters->OutCompactNodeData = GraphBuilder.CreateUAV(OutCompactNodeVis);
+	else
+		PassParameters->OutCompactNodeVis = GraphBuilder.CreateUAV(OutCompactNodeVis);
 	PassParameters->OutCompactNodeCoord = GraphBuilder.CreateUAV(OutCompactNodeCoord, FHairStrandsVisibilityData::NodeCoordFormat);
 	PassParameters->OutCoverageTexture = GraphBuilder.CreateUAV(OutCoverageTexture);
 
@@ -1766,7 +1776,7 @@ class FHairVisibilityCompactionComputeRasterCS : public FGlobalShader
 		SHADER_PARAMETER_RDG_TEXTURE_UAV(RWTexture2D, OutCompactNodeCounter)
 		SHADER_PARAMETER_RDG_TEXTURE_UAV(RWTexture2D, OutCompactNodeIndex)
 		SHADER_PARAMETER_RDG_TEXTURE_UAV(RWTexture2D, OutCoverageTexture)
-		SHADER_PARAMETER_RDG_BUFFER_UAV(RWStructuredBuffer, OutCompactNodeData)
+		SHADER_PARAMETER_RDG_BUFFER_UAV(RWStructuredBuffer, OutCompactNodeVis)
 		SHADER_PARAMETER_RDG_BUFFER_UAV(RWBuffer, OutCompactNodeCoord)
 
 		SHADER_PARAMETER_RDG_UNIFORM_BUFFER(FSceneTextureUniformParameters, SceneTexturesStruct)
@@ -1797,7 +1807,7 @@ static void AddHairVisibilityCompactionComputeRasterPass(
 	FRDGTextureRef& InTransmittanceTexture,
 	FRDGTextureRef& OutCompactCounter,
 	FRDGTextureRef& OutCompactNodeIndex,
-	FRDGBufferRef&  OutCompactNodeData,
+	FRDGBufferRef&  OutCompactNodeVis,
 	FRDGBufferRef&  OutCompactNodeCoord,
 	FRDGTextureRef& OutCoverageTexture,
 	FRDGTextureRef& OutVelocityTexture,
@@ -1830,7 +1840,7 @@ static void AddHairVisibilityCompactionComputeRasterPass(
 	const uint32 MSAASampleCount = GetHairVisibilityRenderMode() == HairVisibilityRenderMode_MSAA_Visibility ? GetMaxSamplePerPixel() : 1;
 	const uint32 PPLLMaxRenderNodePerPixel = GetMaxSamplePerPixel();
 	const uint32 MaxRenderNodeCount = GetTotalSampleCountForAllocation(Resolution);
-	OutCompactNodeData = GraphBuilder.CreateBuffer(FRDGBufferDesc::CreateStructuredDesc(sizeof(HairStrandsVisibilityInternal::NodeVis), MaxRenderNodeCount), TEXT("Hair.VisibilityPrimitiveIdCompactNodeData"));
+	OutCompactNodeVis = GraphBuilder.CreateBuffer(FRDGBufferDesc::CreateStructuredDesc(sizeof(HairStrandsVisibilityInternal::NodeVis), MaxRenderNodeCount), TEXT("Hair.VisibilityPrimitiveIdCompactNodeData"));
 	OutCompactNodeCoord = GraphBuilder.CreateBuffer(FRDGBufferDesc::CreateBufferDesc(sizeof(uint32), MaxRenderNodeCount), TEXT("Hair.VisibilityPrimitiveIdCompactNodeCoord"));
 
 	FRDGTextureRef DefaultTexture = GSystemTextures.GetBlackDummy(GraphBuilder);
@@ -1848,7 +1858,7 @@ static void AddHairVisibilityCompactionComputeRasterPass(
 	PassParameters->SceneTexturesStruct		= CreateSceneTextureUniformBuffer(GraphBuilder, View.FeatureLevel);
 	PassParameters->OutCompactNodeCounter	= GraphBuilder.CreateUAV(OutCompactCounter);
 	PassParameters->OutCompactNodeIndex		= GraphBuilder.CreateUAV(OutCompactNodeIndex);
-	PassParameters->OutCompactNodeData		= GraphBuilder.CreateUAV(OutCompactNodeData);
+	PassParameters->OutCompactNodeVis		= GraphBuilder.CreateUAV(OutCompactNodeVis);
 	PassParameters->OutCompactNodeCoord		= GraphBuilder.CreateUAV(OutCompactNodeCoord, FHairStrandsVisibilityData::NodeCoordFormat);
 	PassParameters->OutCoverageTexture		= GraphBuilder.CreateUAV(OutCoverageTexture);
 
@@ -2959,6 +2969,7 @@ void RenderHairStrandsVisibilityBuffer(
 			FRDGTextureRef CoverageTexture = nullptr;
 			FRDGTextureRef CompactNodeIndex = nullptr;
 			FRDGBufferRef  CompactNodeData = nullptr;
+			FRDGBufferRef  CompactNodeVis = nullptr;
 			FRDGTextureRef NodeCounter = nullptr;
 
 			if (RenderMode != HairVisibilityRenderMode_ComputeRaster)
@@ -3002,7 +3013,7 @@ void RenderHairStrandsVisibilityBuffer(
 							ViewTransmittance.TransmittanceTexture,
 							NodeCounter,
 							CompactNodeIndex,
-							CompactNodeData,
+							CompactNodeVis,
 							CompactNodeCoord,
 							CoverageTexture,
 							SceneVelocityTexture,
@@ -3026,7 +3037,7 @@ void RenderHairStrandsVisibilityBuffer(
 							InstanceCullingManager,
 							VisibilityData.NodeGroupSize,
 							CompactNodeIndex,
-							CompactNodeData,
+							CompactNodeVis,
 							CompactNodeCoord,
 							NodeCounter,
 							IndirectArgsBuffer);
@@ -3038,7 +3049,7 @@ void RenderHairStrandsVisibilityBuffer(
 							MacroGroupDatas,
 							CoverageTexture,
 							CompactNodeIndex,
-							CompactNodeData,
+							CompactNodeVis,
 							PassOutput.NodeVelocity,
 							SceneVelocityTexture,
 							ResolveMaskTexture);
@@ -3173,7 +3184,7 @@ void RenderHairStrandsVisibilityBuffer(
 						PassParameters,
 						NodeCounter,
 						CompactNodeIndex,
-						CompactNodeData,
+						CompactNodeVis,
 						CompactNodeCoord,
 						CoverageTexture,
 						SceneVelocityTexture,
@@ -3200,7 +3211,7 @@ void RenderHairStrandsVisibilityBuffer(
 							InstanceCullingManager,
 							VisibilityData.NodeGroupSize,
 							CompactNodeIndex,
-							CompactNodeData,
+							CompactNodeVis,
 							CompactNodeCoord,
 							NodeCounter,
 							IndirectArgsBuffer);
@@ -3212,7 +3223,7 @@ void RenderHairStrandsVisibilityBuffer(
 							MacroGroupDatas,
 							CoverageTexture,
 							CompactNodeIndex,
-							CompactNodeData,
+							CompactNodeVis,
 							PassOutput.NodeVelocity,
 							SceneVelocityTexture,
 							ResolveMaskTexture);
@@ -3236,7 +3247,7 @@ void RenderHairStrandsVisibilityBuffer(
 
 					 VisibilityData.SampleLightingBuffer	= SampleLightingBuffer;
 					 VisibilityData.NodeIndex				= CompactNodeIndex;
-					 VisibilityData.CoverageTexture	= CoverageTexture;
+					 VisibilityData.CoverageTexture			= CoverageTexture;
 					 VisibilityData.HairOnlyDepthTexture	= HairOnlyDepthTexture;
 					 VisibilityData.NodeData				= CompactNodeData;
 					 VisibilityData.NodeCoord				= CompactNodeCoord;
