* Fix hair tile indirect dispatch group computation.
* Add hair tile debug print.

diff --git a/Engine/Shaders/Private/HairStrands/HairStrandsVisibilityTile.usf b/Engine/Shaders/Private/HairStrands/HairStrandsVisibilityTile.usf
--- a/Engine/Shaders/Private/HairStrands/HairStrandsVisibilityTile.usf
+++ b/Engine/Shaders/Private/HairStrands/HairStrandsVisibilityTile.usf
@@ -1,6 +1,7 @@
 // Copyright Epic Games, Inc. All Rights Reserved.
 #pragma once
 
+#include "HairStrandsVisibilityCommonStruct.ush"
 #include "../Common.ush"
 #include "HairStrandsVisibilityCommon.ush"
 
@@ -100,7 +101,7 @@ void MainCS(uint2 DispatchThreadId : SV_DispatchThreadID)
 		TileIndirectDrawBuffer[2] = 1;
 
 		// Indirect dispatch
-		TileIndirectDispatchBuffer[0] = TileCount % TileCountXY.x;
+		TileIndirectDispatchBuffer[0] = TileCount >= TileCountXY.x ? TileCountXY.x : TileCount;
 		TileIndirectDispatchBuffer[1] = (TileCount+TileCountXY.x-1) / TileCountXY.x;
 		TileIndirectDispatchBuffer[2] = 1;
 
@@ -140,6 +141,61 @@ void MainPS(
 
 ////////////////////////////////////////////////////////////////////////////////////////////////////////////
 
+#if SHADER_TILE_DEBUG_PRINT
+#include "../ShaderPrintCommon.ush"
+
+int2 MaxResolution;
+uint TileGroupSize;
+uint TileSize;
+uint TileCount;
+int2 TileCountXY;
+uint bRectPrimitive;
+
+void AddNewLine(float2 OriginalPos, inout float2 Pos)
+{
+	Pos = ShaderPrintNewline(Pos); Pos.x = OriginalPos.x;
+}
+
+[numthreads(1, 1, 1)]
+void MainCS(uint GroupIndex : SV_GroupIndex, uint3 DispatchThreadId : SV_DispatchThreadID)
+{
+	if (any(DispatchThreadId != 0)) return;
+
+	// Pixel coord
+	float2 OriginalPos = float2(50, 50) / float2(MaxResolution);
+	float2 Pos = OriginalPos;
+
+	Pos = ShaderPrintText(Pos, TEXT("Tile Count XY: "), FontWhite);
+	Pos = ShaderPrint(Pos, TileCountXY, FontEmerald);
+	AddNewLine(OriginalPos, Pos);
+
+	const uint ActiveTileCount = HairStrands.HairTileCount[0];
+	Pos = ShaderPrintText(Pos, TEXT("Tile Count   : "), FontWhite);
+	Pos = ShaderPrint(Pos, ActiveTileCount, FontYellow);
+	Pos = ShaderPrint(Pos, 100.f * float(ActiveTileCount)/float(TileCount), FontOrange);
+	Pos = ShaderPrintText(Pos, TEXT("%"), FontOrange);
+	AddNewLine(OriginalPos, Pos);
+
+	Pos = ShaderPrintText(Pos, TEXT("Tile Size    : "), FontWhite);
+	Pos = ShaderPrint(Pos, TileSize, FontYellow);
+	AddNewLine(OriginalPos, Pos);
+
+	Pos = ShaderPrintText(Pos, TEXT("Tile Threads : "), FontWhite);
+	Pos = ShaderPrint(Pos, TileGroupSize, FontYellow);
+	AddNewLine(OriginalPos, Pos);
+
+	Pos = ShaderPrintText(Pos, TEXT("Primitive    : "), FontWhite);
+	if (bRectPrimitive)
+		Pos = ShaderPrintText(Pos, TEXT("Rect."), FontOrange);
+	else
+		Pos = ShaderPrintText(Pos, TEXT("Triangle"), FontOrange);
+	AddNewLine(OriginalPos, Pos);
+}
+
+#endif //SHADER_TILE_DEBUG_PRINT
+
+////////////////////////////////////////////////////////////////////////////////////////////////////////////
+
 #if SHADER_TILE_VS
 int2 ViewMin;
 float2 ViewInvSize;
diff --git a/Engine/Source/Runtime/Renderer/Private/HairStrands/HairStrandsTile.cpp b/Engine/Source/Runtime/Renderer/Private/HairStrands/HairStrandsTile.cpp
--- a/Engine/Source/Runtime/Renderer/Private/HairStrands/HairStrandsTile.cpp
+++ b/Engine/Source/Runtime/Renderer/Private/HairStrands/HairStrandsTile.cpp
@@ -13,6 +13,8 @@
 #include "PostProcessing.h"
 #include "MeshPassProcessor.inl"
 #include "ScenePrivate.h"
+#include "ShaderPrintParameters.h"
+#include "ShaderPrint.h"
 
 ////////////////////////////////////////////////////////////////////////////////////////////////////////////////////////////////////////////////
 
@@ -101,7 +103,7 @@ class FHairStrandsTileGenerationPassCS : public FGlobalShader
 	BEGIN_SHADER_PARAMETER_STRUCT(FParameters, )
 		SHADER_PARAMETER_STRUCT_REF(FViewUniformShaderParameters, ViewUniformBuffer)
 		SHADER_PARAMETER(FIntPoint, BufferResolution)
-		SHADER_PARAMETER_RDG_TEXTURE(Texture2D<uint4>, InputTexture)
+		SHADER_PARAMETER_RDG_TEXTURE(Texture2D<float>, InputTexture)
 		SHADER_PARAMETER_RDG_BUFFER_UAV(RWBuffer, TileCountBuffer)
 		SHADER_PARAMETER_RDG_BUFFER_UAV(RWBuffer, TileDataBuffer)
 		SHADER_PARAMETER_RDG_BUFFER_UAV(RWBuffer, TileClearBuffer)
@@ -185,6 +187,41 @@ void FHairStrandsTilePassVS::ModifyCompilationEnvironment(const FGlobalShaderPer
 
 IMPLEMENT_GLOBAL_SHADER(FHairStrandsTilePassVS, "/Engine/Private/HairStrands/HairStrandsVisibilityTile.usf", "MainVS", SF_Vertex);
 
+
+////////////////////////////////////////////////////////////////////////////////////////////////////////////////////////////////////////////////
+
+class FHairStrandsTileDebugPrintPassCS : public FGlobalShader
+{
+	DECLARE_GLOBAL_SHADER(FHairStrandsTileDebugPrintPassCS);
+	SHADER_USE_PARAMETER_STRUCT(FHairStrandsTileDebugPrintPassCS, FGlobalShader);
+
+	using FPermutationDomain = TShaderPermutationDomain<>;
+
+	BEGIN_SHADER_PARAMETER_STRUCT(FParameters, )
+		SHADER_PARAMETER(FIntPoint, MaxResolution)
+		SHADER_PARAMETER(uint32, TileGroupSize)
+		SHADER_PARAMETER(uint32, TileSize)
+		SHADER_PARAMETER(uint32, TileCount)
+		SHADER_PARAMETER(FIntPoint, TileCountXY)
+		SHADER_PARAMETER(uint32, bRectPrimitive)
+		SHADER_PARAMETER_STRUCT_INCLUDE(ShaderPrint::FShaderParameters, ShaderPrintUniformBuffer)
+		SHADER_PARAMETER_RDG_UNIFORM_BUFFER(FHairStrandsViewUniformParameters, HairStrands)
+	END_SHADER_PARAMETER_STRUCT()
+
+		static bool ShouldCompilePermutation(const FGlobalShaderPermutationParameters& Parameters)
+	{
+		return GetMaxSupportedFeatureLevel(Parameters.Platform) >= ERHIFeatureLevel::SM5; //TODO
+	}
+
+	static void ModifyCompilationEnvironment(const FGlobalShaderPermutationParameters& Parameters, FShaderCompilerEnvironment& OutEnvironment)
+	{
+		FGlobalShader::ModifyCompilationEnvironment(Parameters, OutEnvironment);
+		OutEnvironment.SetDefine(TEXT("SHADER_TILE_DEBUG_PRINT"), 1);
+	}
+};
+
+IMPLEMENT_GLOBAL_SHADER(FHairStrandsTileDebugPrintPassCS, "/Engine/Private/HairStrands/HairStrandsVisibilityTile.usf", "MainCS", SF_Compute);
+
 ////////////////////////////////////////////////////////////////////////////////////////////////////////////////////////////////////////////////
 
 class FHairStrandsTileDebugPassPS : public FGlobalShader
@@ -225,38 +262,62 @@ void AddHairStrandsDebugTilePass(
 {	
 	const FIntRect Viewport = View.ViewRect;
 	
-	FHairStrandsTileDebugPassPS::FParameters* PassParameters = GraphBuilder.AllocParameters<FHairStrandsTileDebugPassPS::FParameters>();
-	PassParameters->TileParameters = GetHairStrandsTileParameters(View, TileData);
+	{
+		FHairStrandsTileDebugPassPS::FParameters* PassParameters = GraphBuilder.AllocParameters<FHairStrandsTileDebugPassPS::FParameters>();
+		PassParameters->TileParameters = GetHairStrandsTileParameters(View, TileData);
 
-	TShaderMapRef<FHairStrandsTilePassVS> VertexShader(View.ShaderMap);
-	TShaderMapRef<FHairStrandsTileDebugPassPS> PixelShader(View.ShaderMap);
+		TShaderMapRef<FHairStrandsTilePassVS> VertexShader(View.ShaderMap);
+		TShaderMapRef<FHairStrandsTileDebugPassPS> PixelShader(View.ShaderMap);
 
-	PassParameters->RenderTargets[0] = FRenderTargetBinding(ColorTexture, ERenderTargetLoadAction::ELoad);
+		PassParameters->RenderTargets[0] = FRenderTargetBinding(ColorTexture, ERenderTargetLoadAction::ELoad);
 	
-	GraphBuilder.AddPass(
-		RDG_EVENT_NAME("HairStrands::TileDebugPass"),
-		PassParameters,
-		ERDGPassFlags::Raster,
-		[PassParameters, VertexShader, PixelShader, Viewport](FRHICommandList& RHICmdList)
-		{
-			FHairStrandsTilePassVS::FParameters ParametersVS = PassParameters->TileParameters;
-
-			FGraphicsPipelineStateInitializer GraphicsPSOInit;
-			RHICmdList.ApplyCachedRenderTargets(GraphicsPSOInit);
-			GraphicsPSOInit.BlendState = TStaticBlendState<CW_RGBA, BO_Add, BF_One, BF_One, BO_Max, BF_SourceAlpha, BF_DestAlpha>::GetRHI();
-			GraphicsPSOInit.RasterizerState = TStaticRasterizerState<>::GetRHI();
-			GraphicsPSOInit.DepthStencilState = TStaticDepthStencilState<false, CF_Always>::GetRHI();
-
-			GraphicsPSOInit.BoundShaderState.VertexDeclarationRHI = GFilterVertexDeclaration.VertexDeclarationRHI;
-			GraphicsPSOInit.BoundShaderState.VertexShaderRHI = VertexShader.GetVertexShader();
-			GraphicsPSOInit.BoundShaderState.PixelShaderRHI = PixelShader.GetPixelShader();
-			GraphicsPSOInit.PrimitiveType = PassParameters->TileParameters.bRectPrimitive > 0 ? PT_RectList : PT_TriangleList;
-			SetGraphicsPipelineState(RHICmdList, GraphicsPSOInit);
-			SetShaderParameters(RHICmdList, VertexShader, VertexShader.GetVertexShader(), ParametersVS);
-			SetShaderParameters(RHICmdList, PixelShader, PixelShader.GetPixelShader(), *PassParameters);
-
-			RHICmdList.SetViewport(Viewport.Min.X, Viewport.Min.Y, 0.0f, Viewport.Max.X, Viewport.Max.Y, 1.0f);
-			RHICmdList.SetStreamSource(0, nullptr, 0);
-			RHICmdList.DrawPrimitiveIndirect(PassParameters->TileParameters.TileIndirectBuffer->GetRHI(), 0);
-		});
+		GraphBuilder.AddPass(
+			RDG_EVENT_NAME("HairStrands::TileDebugPass"),
+			PassParameters,
+			ERDGPassFlags::Raster,
+			[PassParameters, VertexShader, PixelShader, Viewport](FRHICommandList& RHICmdList)
+			{
+				FHairStrandsTilePassVS::FParameters ParametersVS = PassParameters->TileParameters;
+
+				FGraphicsPipelineStateInitializer GraphicsPSOInit;
+				RHICmdList.ApplyCachedRenderTargets(GraphicsPSOInit);
+				GraphicsPSOInit.BlendState = TStaticBlendState<CW_RGBA, BO_Add, BF_One, BF_One, BO_Max, BF_SourceAlpha, BF_DestAlpha>::GetRHI();
+				GraphicsPSOInit.RasterizerState = TStaticRasterizerState<>::GetRHI();
+				GraphicsPSOInit.DepthStencilState = TStaticDepthStencilState<false, CF_Always>::GetRHI();
+
+				GraphicsPSOInit.BoundShaderState.VertexDeclarationRHI = GFilterVertexDeclaration.VertexDeclarationRHI;
+				GraphicsPSOInit.BoundShaderState.VertexShaderRHI = VertexShader.GetVertexShader();
+				GraphicsPSOInit.BoundShaderState.PixelShaderRHI = PixelShader.GetPixelShader();
+				GraphicsPSOInit.PrimitiveType = PassParameters->TileParameters.bRectPrimitive > 0 ? PT_RectList : PT_TriangleList;
+				SetGraphicsPipelineState(RHICmdList, GraphicsPSOInit);
+				SetShaderParameters(RHICmdList, VertexShader, VertexShader.GetVertexShader(), ParametersVS);
+				SetShaderParameters(RHICmdList, PixelShader, PixelShader.GetPixelShader(), *PassParameters);
+
+				RHICmdList.SetViewport(Viewport.Min.X, Viewport.Min.Y, 0.0f, Viewport.Max.X, Viewport.Max.Y, 1.0f);
+				RHICmdList.SetStreamSource(0, nullptr, 0);
+				RHICmdList.DrawPrimitiveIndirect(PassParameters->TileParameters.TileIndirectBuffer->GetRHI(), 0);
+			});
+	}
+
+	if (View.HairStrandsViewData.UniformBuffer || ShaderPrint::IsEnabled(View))
+	{
+		FHairStrandsTileDebugPrintPassCS::FParameters* Parameters = GraphBuilder.AllocParameters<FHairStrandsTileDebugPrintPassCS::FParameters>();
+		Parameters->MaxResolution = FIntPoint(Viewport.Width(), Viewport.Height());
+		Parameters->TileGroupSize = TileData.GroupSize;
+		Parameters->TileSize = TileData.TileSize;
+		Parameters->TileCount = TileData.TileCount;
+		Parameters->TileCountXY = TileData.TileCountXY;
+		Parameters->bRectPrimitive = TileData.bRectPrimitive ? 1u : 0u;
+		Parameters->HairStrands = View.HairStrandsViewData.UniformBuffer;
+		ShaderPrint::SetParameters(GraphBuilder, View, Parameters->ShaderPrintUniformBuffer);
+
+		TShaderMapRef<FHairStrandsTileDebugPrintPassCS> ComputeShader(View.ShaderMap);
+		ClearUnusedGraphResources(ComputeShader, Parameters);
+		FComputeShaderUtils::AddPass(
+			GraphBuilder,
+			RDG_EVENT_NAME("HairStrands::TileDebugPrint"),
+			ComputeShader,
+			Parameters,
+			FIntVector(1, 1, 1));
+	}
 }
\ No newline at end of file
