add RGB555A1 texture format support(not supported on Mac)
useful for RVT on mobile platforms, to save the cost of runtime compression

diff --git a/Engine/Source/Runtime/Apple/AGXRHI/Private/AGXRHI.cpp b/Engine/Source/Runtime/Apple/AGXRHI/Private/AGXRHI.cpp
--- a/Engine/Source/Runtime/Apple/AGXRHI/Private/AGXRHI.cpp
+++ b/Engine/Source/Runtime/Apple/AGXRHI/Private/AGXRHI.cpp
@@ -666,6 +666,7 @@ FAGXDynamicRHI::FAGXDynamicRHI(ERHIFeatureLevel::Type RequestedFeatureLevel)
 	GAGXBufferFormats[PF_R16G16B16A16_UINT    ] = { mtlpp::PixelFormat::RGBA16Uint, (uint8)EMetalBufferFormat::RGBA16Uint };
 	GAGXBufferFormats[PF_R16G16B16A16_SINT    ] = { mtlpp::PixelFormat::RGBA16Sint, (uint8)EMetalBufferFormat::RGBA16Sint };
 	GAGXBufferFormats[PF_R5G6B5_UNORM         ] = { mtlpp::PixelFormat::Invalid, (uint8)EMetalBufferFormat::R5G6B5Unorm };
+	GAGXBufferFormats[PF_B5G5R5A1_UNORM       ] = { mtlpp::PixelFormat::Invalid, (uint8)EMetalBufferFormat::B5G5R5A1Unorm };
 	GAGXBufferFormats[PF_R8G8B8A8             ] = { mtlpp::PixelFormat::RGBA8Unorm, (uint8)EMetalBufferFormat::RGBA8Unorm };
 	GAGXBufferFormats[PF_A8R8G8B8				] = { mtlpp::PixelFormat::RGBA8Unorm, (uint8)EMetalBufferFormat::RGBA8Unorm }; // mtlpp::PixelFormat::BGRA8Unorm/EMetalBufferFormat::BGRA8Unorm,  < We don't support this as a vertex-format so we have code to swizzle in the shader
 	GAGXBufferFormats[PF_BC4					] = { mtlpp::PixelFormat::Invalid, (uint8)EMetalBufferFormat::Unknown };
@@ -789,6 +790,9 @@ FAGXDynamicRHI::FAGXDynamicRHI(ERHIFeatureLevel::Type RequestedFeatureLevel)
 		
 	GPixelFormats[PF_BC5				].PlatformFormat	= (uint32)mtlpp::PixelFormat::Invalid;
 	GPixelFormats[PF_R5G6B5_UNORM		].PlatformFormat	= (uint32)mtlpp::PixelFormat::B5G6R5Unorm;
+	GPixelFormats[PF_R5G6B5_UNORM       ].Supported         = true;
+	GPixelFormats[PF_B5G5R5A1_UNORM     ].PlatformFormat    = (uint32)mtlpp::PixelFormat::BGR5A1Unorm;
+	GPixelFormats[PF_B5G5R5A1_UNORM     ].Supported         = true;
 #else
     GPixelFormats[PF_DXT1				].PlatformFormat	= (uint32)mtlpp::PixelFormat::BC1_RGBA;
     GPixelFormats[PF_DXT3				].PlatformFormat	= (uint32)mtlpp::PixelFormat::BC2_RGBA;
@@ -849,6 +853,7 @@ FAGXDynamicRHI::FAGXDynamicRHI(ERHIFeatureLevel::Type RequestedFeatureLevel)
 	GPixelFormats[PF_BC7				].Supported			= true;
 	GPixelFormats[PF_BC7				].PlatformFormat	= (uint32)mtlpp::PixelFormat::BC7_RGBAUnorm;
 	GPixelFormats[PF_R5G6B5_UNORM		].PlatformFormat	= (uint32)mtlpp::PixelFormat::Invalid;
+	GPixelFormats[PF_B5G5R5A1_UNORM		].PlatformFormat	= (uint32)mtlpp::PixelFormat::Invalid;
 #endif
 	GPixelFormats[PF_UYVY				].PlatformFormat	= (uint32)mtlpp::PixelFormat::Invalid;
 	GPixelFormats[PF_FloatRGBA			].PlatformFormat	= (uint32)mtlpp::PixelFormat::RGBA16Float;
diff --git a/Engine/Source/Runtime/Apple/Common/Public/MetalShaderResources.h b/Engine/Source/Runtime/Apple/Common/Public/MetalShaderResources.h
--- a/Engine/Source/Runtime/Apple/Common/Public/MetalShaderResources.h
+++ b/Engine/Source/Runtime/Apple/Common/Public/MetalShaderResources.h
@@ -94,8 +94,9 @@ enum class EMetalBufferFormat : uint8
 	RG11B10Half 			=51,
 	
 	R5G6B5Unorm         	=52,
-	
-	Max						=53
+	B5G5R5A1Unorm           =53,
+
+	Max						=54
 };
 
 struct FMetalShaderResourceTable : public FBaseShaderResourceTable
diff --git a/Engine/Source/Runtime/Apple/MetalRHI/Private/MetalRHI.cpp b/Engine/Source/Runtime/Apple/MetalRHI/Private/MetalRHI.cpp
--- a/Engine/Source/Runtime/Apple/MetalRHI/Private/MetalRHI.cpp
+++ b/Engine/Source/Runtime/Apple/MetalRHI/Private/MetalRHI.cpp
@@ -645,6 +645,7 @@ FMetalDynamicRHI::FMetalDynamicRHI(ERHIFeatureLevel::Type RequestedFeatureLevel)
 	GMetalBufferFormats[PF_R16G16B16A16_UINT    ] = { mtlpp::PixelFormat::RGBA16Uint, (uint8)EMetalBufferFormat::RGBA16Uint };
 	GMetalBufferFormats[PF_R16G16B16A16_SINT    ] = { mtlpp::PixelFormat::RGBA16Sint, (uint8)EMetalBufferFormat::RGBA16Sint };
 	GMetalBufferFormats[PF_R5G6B5_UNORM         ] = { mtlpp::PixelFormat::Invalid, (uint8)EMetalBufferFormat::R5G6B5Unorm };
+	GMetalBufferFormats[PF_B5G5R5A1_UNORM       ] = { mtlpp::PixelFormat::Invalid, (uint8)EMetalBufferFormat::B5G5R5A1Unorm };
 	GMetalBufferFormats[PF_R8G8B8A8             ] = { mtlpp::PixelFormat::RGBA8Unorm, (uint8)EMetalBufferFormat::RGBA8Unorm };
 	GMetalBufferFormats[PF_A8R8G8B8				] = { mtlpp::PixelFormat::RGBA8Unorm, (uint8)EMetalBufferFormat::RGBA8Unorm }; // mtlpp::PixelFormat::BGRA8Unorm/EMetalBufferFormat::BGRA8Unorm,  < We don't support this as a vertex-format so we have code to swizzle in the shader
 	GMetalBufferFormats[PF_BC4					] = { mtlpp::PixelFormat::Invalid, (uint8)EMetalBufferFormat::Unknown };
@@ -691,7 +692,7 @@ FMetalDynamicRHI::FMetalDynamicRHI(ERHIFeatureLevel::Type RequestedFeatureLevel)
 	GPixelFormats[PF_R32G32B32A32_UINT	].PlatformFormat	= (uint32)mtlpp::PixelFormat::RGBA32Uint;
 	GPixelFormats[PF_R16G16_UINT		].PlatformFormat	= (uint32)mtlpp::PixelFormat::RG16Uint;
 	GPixelFormats[PF_R32G32_UINT		].PlatformFormat	= (uint32)mtlpp::PixelFormat::RG32Uint;
-		
+
 #if PLATFORM_IOS
     GPixelFormats[PF_DXT1				].PlatformFormat	= (uint32)mtlpp::PixelFormat::Invalid;
 	GPixelFormats[PF_DXT1				].Supported			= false;
@@ -768,6 +769,9 @@ FMetalDynamicRHI::FMetalDynamicRHI(ERHIFeatureLevel::Type RequestedFeatureLevel)
 		
 	GPixelFormats[PF_BC5				].PlatformFormat	= (uint32)mtlpp::PixelFormat::Invalid;
 	GPixelFormats[PF_R5G6B5_UNORM		].PlatformFormat	= (uint32)mtlpp::PixelFormat::B5G6R5Unorm;
+	GPixelFormats[PF_R5G6B5_UNORM       ].Supported         = true;
+	GPixelFormats[PF_B5G5R5A1_UNORM     ].PlatformFormat    = (uint32)mtlpp::PixelFormat::BGR5A1Unorm;
+	GPixelFormats[PF_B5G5R5A1_UNORM     ].Supported         = true;
 #else
     GPixelFormats[PF_DXT1				].PlatformFormat	= (uint32)mtlpp::PixelFormat::BC1_RGBA;
     GPixelFormats[PF_DXT3				].PlatformFormat	= (uint32)mtlpp::PixelFormat::BC2_RGBA;
@@ -828,6 +832,7 @@ FMetalDynamicRHI::FMetalDynamicRHI(ERHIFeatureLevel::Type RequestedFeatureLevel)
 	GPixelFormats[PF_BC7				].Supported			= true;
 	GPixelFormats[PF_BC7				].PlatformFormat	= (uint32)mtlpp::PixelFormat::BC7_RGBAUnorm;
 	GPixelFormats[PF_R5G6B5_UNORM		].PlatformFormat	= (uint32)mtlpp::PixelFormat::Invalid;
+	GPixelFormats[PF_B5G5R5A1_UNORM		].PlatformFormat	= (uint32)mtlpp::PixelFormat::Invalid;
 #endif
 	GPixelFormats[PF_UYVY				].PlatformFormat	= (uint32)mtlpp::PixelFormat::Invalid;
 	GPixelFormats[PF_FloatRGBA			].PlatformFormat	= (uint32)mtlpp::PixelFormat::RGBA16Float;
diff --git a/Engine/Source/Runtime/Core/Public/PixelFormat.h b/Engine/Source/Runtime/Core/Public/PixelFormat.h
--- a/Engine/Source/Runtime/Core/Public/PixelFormat.h
+++ b/Engine/Source/Runtime/Core/Public/PixelFormat.h
@@ -84,7 +84,8 @@ enum EPixelFormat
 	PF_ETC2_R11_EAC			=69,
 	PF_ETC2_RG11_EAC		=70,
 	PF_R8		            =71,
-	PF_MAX					=72,
+	PF_B5G5R5A1_UNORM       =72,
+	PF_MAX					=73,
 };
 #define FOREACH_ENUM_EPIXELFORMAT(op) \
 	op(PF_Unknown) \
@@ -158,8 +159,8 @@ enum EPixelFormat
 	op(PF_R32G32_UINT) \
 	op(PF_ETC2_R11_EAC) \
 	op(PF_ETC2_RG11_EAC) \
-	op(PF_R8) 
-
+	op(PF_R8) \
+	op(PF_B5G5R5A1_UNORM)
 // Defines which channel is valid for each pixel format
 enum class EPixelFormatChannelFlags : uint8
 {
diff --git a/Engine/Source/Runtime/CoreUObject/Public/UObject/NoExportTypes.h b/Engine/Source/Runtime/CoreUObject/Public/UObject/NoExportTypes.h
--- a/Engine/Source/Runtime/CoreUObject/Public/UObject/NoExportTypes.h
+++ b/Engine/Source/Runtime/CoreUObject/Public/UObject/NoExportTypes.h
@@ -246,6 +246,8 @@ enum EPixelFormat
 	PF_R32G32_UINT,
 	PF_ETC2_R11_EAC,
 	PF_ETC2_RG11_EAC,
+	PF_R8,
+	PF_B5G5R5A1_UNORM,
 	PF_MAX,
 };
 
diff --git a/Engine/Source/Runtime/D3D12RHI/Private/D3D12RHI.cpp b/Engine/Source/Runtime/D3D12RHI/Private/D3D12RHI.cpp
--- a/Engine/Source/Runtime/D3D12RHI/Private/D3D12RHI.cpp
+++ b/Engine/Source/Runtime/D3D12RHI/Private/D3D12RHI.cpp
@@ -152,6 +152,9 @@ FD3D12DynamicRHI::FD3D12DynamicRHI(const TArray<TSharedPtr<FD3D12Adapter>>& Chos
 	GPixelFormats[PF_R16G16B16A16_SINT].PlatformFormat = DXGI_FORMAT_R16G16B16A16_SINT;
 
 	GPixelFormats[PF_R5G6B5_UNORM	].PlatformFormat = DXGI_FORMAT_B5G6R5_UNORM;
+	GPixelFormats[PF_R5G6B5_UNORM	].Supported = true;
+	GPixelFormats[PF_B5G5R5A1_UNORM ].PlatformFormat = DXGI_FORMAT_B5G5R5A1_UNORM;
+	GPixelFormats[PF_B5G5R5A1_UNORM ].Supported = true;
 	GPixelFormats[PF_R8G8B8A8		].PlatformFormat = DXGI_FORMAT_R8G8B8A8_TYPELESS;
 	GPixelFormats[PF_R8G8B8A8_UINT	].PlatformFormat = DXGI_FORMAT_R8G8B8A8_UINT;
 	GPixelFormats[PF_R8G8B8A8_SNORM	].PlatformFormat = DXGI_FORMAT_R8G8B8A8_SNORM;
diff --git a/Engine/Source/Runtime/OpenGLDrv/Private/OpenGLDevice.cpp b/Engine/Source/Runtime/OpenGLDrv/Private/OpenGLDevice.cpp
--- a/Engine/Source/Runtime/OpenGLDrv/Private/OpenGLDevice.cpp
+++ b/Engine/Source/Runtime/OpenGLDrv/Private/OpenGLDevice.cpp
@@ -926,7 +926,8 @@ static void InitRHICapabilitiesForGL()
 	SetupTextureFormat( PF_R16G16_UINT,			FOpenGLTextureFormat( GL_RG16UI,				GL_RG16UI,				GL_RG_INTEGER,		GL_UNSIGNED_SHORT,				false,			false));
 	SetupTextureFormat( PF_R8,					FOpenGLTextureFormat( GL_R8,					GL_R8,					GL_RED,				GL_UNSIGNED_BYTE,				false,			false));
 
-	SetupTextureFormat( PF_R5G6B5_UNORM,		FOpenGLTextureFormat( ));
+	SetupTextureFormat( PF_R5G6B5_UNORM,        FOpenGLTextureFormat( GL_RGB565,                GL_RGB565,              GL_RGB,             GL_UNSIGNED_SHORT_5_6_5,        false,          false));
+	SetupTextureFormat( PF_B5G5R5A1_UNORM,		FOpenGLTextureFormat( GL_RGB5_A1,               GL_RGB5_A1,             GL_RGBA,            GL_UNSIGNED_SHORT_5_5_5_1,      false,          false));
 
 	SetupTextureFormat( PF_G16,					FOpenGLTextureFormat( GL_R16,					GL_R16,					GL_RED,				GL_UNSIGNED_SHORT,					false,	false));
 	SetupTextureFormat( PF_R32_FLOAT,			FOpenGLTextureFormat( GL_R32F,					GL_R32F,				GL_RED,				GL_FLOAT,							false,	false));
diff --git a/Engine/Source/Runtime/RHI/Private/RHI.cpp b/Engine/Source/Runtime/RHI/Private/RHI.cpp
--- a/Engine/Source/Runtime/RHI/Private/RHI.cpp
+++ b/Engine/Source/Runtime/RHI/Private/RHI.cpp
@@ -2413,7 +2413,7 @@ FPixelFormatInfo    GPixelFormats[PF_MAX] =
 	FPixelFormatInfo(PF_R16_SINT,           TEXT("R16_SINT"),              1,          1,          1,          2,          1,              1),
 	FPixelFormatInfo(PF_R16G16B16A16_UINT,  TEXT("R16G16B16A16_UINT"),     1,          1,          1,          8,          4,              1),
 	FPixelFormatInfo(PF_R16G16B16A16_SINT,  TEXT("R16G16B16A16_SINT"),     1,          1,          1,          8,          4,              1),
-	FPixelFormatInfo(PF_R5G6B5_UNORM,       TEXT("R5G6B5_UNORM"),          1,          1,          1,          2,          3,              1),
+	FPixelFormatInfo(PF_R5G6B5_UNORM,       TEXT("R5G6B5_UNORM"),          1,          1,          1,          2,          3,              0),
 	FPixelFormatInfo(PF_R8G8B8A8,           TEXT("R8G8B8A8"),              1,          1,          1,          4,          4,              1),
 	FPixelFormatInfo(PF_A8R8G8B8,           TEXT("A8R8G8B8"),              1,          1,          1,          4,          4,              1),
 	FPixelFormatInfo(PF_BC4,                TEXT("BC4"),                   4,          4,          1,          8,          1,              1),
@@ -2459,6 +2459,8 @@ FPixelFormatInfo    GPixelFormats[PF_MAX] =
 	FPixelFormatInfo(PF_ETC2_R11_EAC,       TEXT("PF_ETC2_R11_EAC"),       4,          4,          1,          8,          1,              0),
 	FPixelFormatInfo(PF_ETC2_RG11_EAC,      TEXT("PF_ETC2_RG11_EAC"),      4,          4,          1,          16,         2,              0),
 	FPixelFormatInfo(PF_R8,                 TEXT("R8"),                    1,          1,          1,          1,          1,              1),
+    FPixelFormatInfo(PF_B5G5R5A1_UNORM,     TEXT("B5G5R5A1_UNORM"),        1,          1,          1,          2,          4,              0),
+
 };
 
 static struct FValidatePixelFormats
diff --git a/Engine/Source/Runtime/RenderCore/Private/RenderUtils.cpp b/Engine/Source/Runtime/RenderCore/Private/RenderUtils.cpp
--- a/Engine/Source/Runtime/RenderCore/Private/RenderUtils.cpp
+++ b/Engine/Source/Runtime/RenderCore/Private/RenderUtils.cpp
@@ -927,6 +927,7 @@ EPixelFormatChannelFlags GetPixelFormatValidChannels(EPixelFormat InPixelFormat)
 		EPixelFormatChannelFlags::R,		// PF_ETC2_R11_EAC
 		EPixelFormatChannelFlags::RG,		// PF_ETC2_RG11_EAC
 		EPixelFormatChannelFlags::R,		// PF_R8
+		EPixelFormatChannelFlags::RGBA,		// PF_B5G5R5A1_UNORM
 	};
 	static_assert(UE_ARRAY_COUNT(PixelFormatToChannelFlags) == (uint8)PF_MAX, "Missing pixel format");
 	return (InPixelFormat < PF_MAX) ? PixelFormatToChannelFlags[(uint8)InPixelFormat] : EPixelFormatChannelFlags::None;
diff --git a/Engine/Source/Runtime/Renderer/Private/SystemTextures.cpp b/Engine/Source/Runtime/Renderer/Private/SystemTextures.cpp
--- a/Engine/Source/Runtime/Renderer/Private/SystemTextures.cpp
+++ b/Engine/Source/Runtime/Renderer/Private/SystemTextures.cpp
@@ -955,7 +955,9 @@ enum class EDefaultInputType
 	SNorm,
 	UNorm10,
 	UNorm11,
-	UNorm2
+	UNorm2,
+	UNorm5,
+	UNorm1
 };
 
 // Convert input type into the final type. This function manages UNorm/SNorm type by assuming if the input if float, its value is normalized in [0..1].
@@ -974,7 +976,8 @@ template<>   int8 ConvertInputFormat<float,   int8, EDefaultInputType::SNorm>(co
 template<> uint32 ConvertInputFormat<float, uint32, EDefaultInputType::UNorm10>(const float& In) { return FMath::Clamp(In, 0.f, 1.f) * 1024u; }
 template<> uint32 ConvertInputFormat<float, uint32, EDefaultInputType::UNorm11>(const float& In) { return FMath::Clamp(In, 0.f, 1.f) * 2048u; }
 template<> uint32 ConvertInputFormat<float, uint32, EDefaultInputType::UNorm2> (const float& In) { return FMath::Clamp(In, 0.f, 1.f) * 3u; }
-
+template<> uint32 ConvertInputFormat<float, uint32, EDefaultInputType::UNorm1> (const float& In) { return uint32(In > 0.5f);}
+template<> uint32 ConvertInputFormat<float, uint32, EDefaultInputType::UNorm5>(const float& In) { return FMath::Clamp(In, 0.f, 1.f) * 31u; }
 // 4 components conversion with swizzling
 template <EDefaultInputType InputFormatType, typename TInType, typename TOutType, uint32 SwizzleX, uint32 SwizzleY, uint32 SwizzleZ, uint32 SwizzleW>
 void FormatData(const TInType& In, uint8* Out, uint32& OutByteCount)
@@ -1040,11 +1043,24 @@ void FormatData1010102(const TInType& In, uint8* Out, uint32& OutByteCount)
 	OutByteCount = 4;
 }
 
+template <typename TInType>
+void FormatData5551(const TInType& In, uint8* Out, uint32& OutByteCount)
+{
+	uint16* OutTyped = (uint16*)Out;
+	*OutTyped =
+		(
+		((31u & ConvertInputFormat<typename TFormatConversionTraits<TInType>::Type, uint32, EDefaultInputType::UNorm5>(In[0])) << 1) |
+		((31u & ConvertInputFormat<typename TFormatConversionTraits<TInType>::Type, uint32, EDefaultInputType::UNorm5>(In[1])) << 6) |
+		((31u & ConvertInputFormat<typename TFormatConversionTraits<TInType>::Type, uint32, EDefaultInputType::UNorm5>(In[2])) << 11) |
+		((1u & ConvertInputFormat<typename TFormatConversionTraits<TInType>::Type, uint32, EDefaultInputType::UNorm1>(In[3]))));
+	OutByteCount = 2;
+}
+
 template<typename TInType>
 void InitializeData(const TInType& InData, EPixelFormat InFormat, uint8* OutData, uint32& OutByteCount)
 {
 	// If a new format is added insure that it is either supported here, or at least flagged as not supported
-	static_assert(PF_MAX == 72);
+	static_assert(PF_MAX == 73);
 
 	switch (InFormat)
 	{
@@ -1097,6 +1113,7 @@ void InitializeData(const TInType& InData, EPixelFormat InFormat, uint8* OutData
 		case PF_FloatRGB:				{ FormatData111110<TInType>	(InData, OutData, OutByteCount); } break;
 		case PF_A2B10G10R10:			{ FormatData1010102<TInType>(InData, OutData, OutByteCount); } break;
 		case PF_FloatR11G11B10:			{ FormatData111110<TInType>	(InData, OutData, OutByteCount); } break;
+		case PF_B5G5R5A1_UNORM:         { FormatData5551<TInType>(InData, OutData, OutByteCount); } break;
 			return;
 
 		// Not supported
diff --git a/Engine/Source/Runtime/VulkanRHI/Private/VulkanDevice.cpp b/Engine/Source/Runtime/VulkanRHI/Private/VulkanDevice.cpp
--- a/Engine/Source/Runtime/VulkanRHI/Private/VulkanDevice.cpp
+++ b/Engine/Source/Runtime/VulkanRHI/Private/VulkanDevice.cpp
@@ -629,6 +629,9 @@ void FVulkanDevice::SetupFormats()
 	MapFormatSupport(PF_R5G6B5_UNORM, VK_FORMAT_R5G6B5_UNORM_PACK16);
 	SetComponentMapping(PF_R5G6B5_UNORM, VK_COMPONENT_SWIZZLE_R, VK_COMPONENT_SWIZZLE_G, VK_COMPONENT_SWIZZLE_B, VK_COMPONENT_SWIZZLE_A);
 
+	MapFormatSupport(PF_B5G5R5A1_UNORM, VK_FORMAT_A1R5G5B5_UNORM_PACK16);
+	SetComponentMapping(PF_B5G5R5A1_UNORM, VK_COMPONENT_SWIZZLE_R, VK_COMPONENT_SWIZZLE_G, VK_COMPONENT_SWIZZLE_B, VK_COMPONENT_SWIZZLE_A);
+
 	MapFormatSupport(PF_R8G8B8A8, VK_FORMAT_R8G8B8A8_UNORM);
 	SetComponentMapping(PF_R8G8B8A8, VK_COMPONENT_SWIZZLE_R, VK_COMPONENT_SWIZZLE_G, VK_COMPONENT_SWIZZLE_B, VK_COMPONENT_SWIZZLE_A);
 
diff --git a/Engine/Source/Runtime/Windows/D3D11RHI/Private/D3D11Device.cpp b/Engine/Source/Runtime/Windows/D3D11RHI/Private/D3D11Device.cpp
--- a/Engine/Source/Runtime/Windows/D3D11RHI/Private/D3D11Device.cpp
+++ b/Engine/Source/Runtime/Windows/D3D11RHI/Private/D3D11Device.cpp
@@ -174,6 +174,9 @@ FD3D11DynamicRHI::FD3D11DynamicRHI(IDXGIFactory1* InDXGIFactory1,D3D_FEATURE_LEV
 	GPixelFormats[ PF_R16G16B16A16_SINT].PlatformFormat = DXGI_FORMAT_R16G16B16A16_SINT;
 
 	GPixelFormats[ PF_R5G6B5_UNORM	].PlatformFormat	= DXGI_FORMAT_B5G6R5_UNORM;
+	GPixelFormats[ PF_R5G6B5_UNORM  ].Supported         = true;
+	GPixelFormats[ PF_B5G5R5A1_UNORM].PlatformFormat    = DXGI_FORMAT_B5G5R5A1_UNORM;
+	GPixelFormats[ PF_B5G5R5A1_UNORM].Supported         = true;
 	GPixelFormats[ PF_R8G8B8A8		].PlatformFormat	= DXGI_FORMAT_R8G8B8A8_TYPELESS;
 	GPixelFormats[ PF_R8G8B8A8_UINT	].PlatformFormat	= DXGI_FORMAT_R8G8B8A8_UINT;
 	GPixelFormats[ PF_R8G8B8A8_SNORM].PlatformFormat	= DXGI_FORMAT_R8G8B8A8_SNORM;
diff --git a/Engine/Source/Runtime/Windows/D3D11RHI/Private/D3D11RenderTarget.cpp b/Engine/Source/Runtime/Windows/D3D11RHI/Private/D3D11RenderTarget.cpp
--- a/Engine/Source/Runtime/Windows/D3D11RHI/Private/D3D11RenderTarget.cpp
+++ b/Engine/Source/Runtime/Windows/D3D11RHI/Private/D3D11RenderTarget.cpp
@@ -353,6 +353,8 @@ static uint32 ComputeBytesPerPixel(DXGI_FORMAT Format)
 		case DXGI_FORMAT_R16_UINT:
 		case DXGI_FORMAT_R16_SNORM:
 		case DXGI_FORMAT_R16_SINT:
+		case DXGI_FORMAT_B5G6R5_UNORM:
+		case DXGI_FORMAT_B5G5R5A1_UNORM:
 			BytesPerPixel = 2;
 			break;
 		case DXGI_FORMAT_B8G8R8A8_TYPELESS:
