Move hair tile pass generation before hair compaction.
Change hair sample compaction to use tile data.

This ensure that compaction run only on necessary pixels and remove artifacts when adpating memory allocation based on CPU bounds screen size.

diff --git a/Engine/Shaders/Private/HairStrands/HairStrandsVisibilityCompaction.usf b/Engine/Shaders/Private/HairStrands/HairStrandsVisibilityCompaction.usf
--- a/Engine/Shaders/Private/HairStrands/HairStrandsVisibilityCompaction.usf
+++ b/Engine/Shaders/Private/HairStrands/HairStrandsVisibilityCompaction.usf
@@ -3,7 +3,7 @@
 #pragma once
 
 #include "../Common.ush"
-#include "HairStrandsVisibilityCommonStruct.ush" 
+#include "HairStrandsVisibilityCommonStruct.ush"
 #include "HairStrandsVisibilityCommon.ush" 
 
 #if PERMUTATION_GROUPSIZE ==32
@@ -420,22 +420,40 @@ float  GetDepthTexture(uint2 PixelCoord, uint SampleIt) { return MSAA_DepthTextu
 
 #endif // PERMUTATION_MSAACOUNT>1
 
+#if PERMUTATION_TILE
+int2 TileCountXY;
+uint TileSize;
+Buffer<uint> TileCountBuffer;
+Buffer<uint2> TileDataBuffer;
+#endif // PERMUTATION_TILE
 
-
-#if PERMUTATION_VIEWTRANSMITTANCE == 1
 Texture2D<float> ViewTransmittanceTexture;
-#endif
 
 [numthreads(TILE_PIXEL_SIZE_X, TILE_PIXEL_SIZE_Y, 1)]
-void MainCS(uint GroupIndex : SV_GroupIndex, uint3 DispatchThreadId : SV_DispatchThreadID)
+void MainCS(
+	uint  GroupIndex		: SV_GroupIndex, 
+	uint3 GroupId			: SV_GroupID, 
+	uint3 GroupThreadId		: SV_GroupThreadID, 
+	uint3 DispatchThreadId	: SV_DispatchThreadID)
 {
 	if (GroupIndex == 0)
 	{
 		AllocationNodeCount = 0;
 		AllocationNodeOffset = 0;
 	}
-
+#if PERMUTATION_TILE
+	const uint TileCount = TileCountBuffer[0];
+	const uint LinearIndex = GroupId.x + GroupId.y * TileCountXY.x;
+	if (LinearIndex >= TileCount)
+	{
+		return;
+	}
+	const uint2 TileCoord = TileDataBuffer[LinearIndex];
+	uint2 PixelCoord = TileCoord * TileSize + GroupThreadId;
+#else
 	uint2 PixelCoord = DispatchThreadId.xy + ResolutionOffset;
+#endif // PERMUTATION_TILE
+
 	if (PixelCoord.x >= uint(OutputResolution.x) || PixelCoord.y >= uint(OutputResolution.y))
 		PixelCoord = uint2(0, 0);
 
@@ -493,12 +511,7 @@ void MainCS(uint GroupIndex : SV_GroupIndex, uint3 DispatchThreadId : SV_Dispatc
 			}
 
 			NodeDesc.Count = NodeDesc.Count;
-
-		#if PERMUTATION_VIEWTRANSMITTANCE == 1
 			PixelCoverage = TransmittanceToCoverage(ViewTransmittanceTexture.Load(uint3(PixelCoord, 0)), CoverageThreshold);
-		#else
-			PixelCoverage = SampleDesc.ValidSampleCount / float(PERMUTATION_MSAACOUNT);
-		#endif
 		}
 	}
 
diff --git a/Engine/Shaders/Private/HairStrands/HairStrandsVisibilityCompactionComputeRaster.usf b/Engine/Shaders/Private/HairStrands/HairStrandsVisibilityCompactionComputeRaster.usf
--- a/Engine/Shaders/Private/HairStrands/HairStrandsVisibilityCompactionComputeRaster.usf
+++ b/Engine/Shaders/Private/HairStrands/HairStrandsVisibilityCompactionComputeRaster.usf
@@ -66,8 +66,19 @@ Texture2D<PackedType> VisibilityTexture3;
 Texture2D<uint>  VisibilityHairCount;
 Texture2D<float> ViewTransmittanceTexture;
 
+#if PERMUTATION_TILE
+int2 TileCountXY;
+uint TileSize;
+Buffer<uint> TileCountBuffer;
+Buffer<uint2> TileDataBuffer;
+#endif // PERMUTATION_TILE
+
 [numthreads(TILE_PIXEL_SIZE_X, TILE_PIXEL_SIZE_Y, 1)]
-void MainCS(uint GroupIndex : SV_GroupIndex, uint3 DispatchThreadId : SV_DispatchThreadID)
+void MainCS(
+	uint  GroupIndex : SV_GroupIndex,
+	uint3 GroupId : SV_GroupID,
+	uint3 GroupThreadId : SV_GroupThreadID,
+	uint3 DispatchThreadId : SV_DispatchThreadID)
 {
 	if (GroupIndex == 0)
 	{
@@ -75,7 +86,19 @@ void MainCS(uint GroupIndex : SV_GroupIndex, uint3 DispatchThreadId : SV_Dispatc
 		AllocationNodeOffset = 0;
 	}
 
+#if PERMUTATION_TILE
+	const uint TileCount = TileCountBuffer[0];
+	const uint LinearIndex = GroupId.x + GroupId.y * TileCountXY.x;
+	if (LinearIndex >= TileCount)
+	{
+		return;
+	}
+	const uint2 TileCoord = TileDataBuffer[LinearIndex];
+	uint2 PixelCoord = TileCoord * TileSize + GroupThreadId;
+#else
 	uint2 PixelCoord = DispatchThreadId.xy;// +ResolutionOffset;
+#endif // PERMUTATION_TILE
+
 	if (PixelCoord.x >= uint(OutputResolution.x) || PixelCoord.y >= uint(OutputResolution.y))
 		PixelCoord = uint2(0, 0);
 
diff --git a/Engine/Shaders/Private/HairStrands/HairStrandsVisibilityTile.usf b/Engine/Shaders/Private/HairStrands/HairStrandsVisibilityTile.usf
--- a/Engine/Shaders/Private/HairStrands/HairStrandsVisibilityTile.usf
+++ b/Engine/Shaders/Private/HairStrands/HairStrandsVisibilityTile.usf
@@ -12,7 +12,9 @@
 
 #if SHADER_TILE_GENERATION
 int2 BufferResolution;
-Texture2D<float> InputTexture;
+uint bUintTexture;
+Texture2D<float> InputFloatTexture;
+Texture2D<uint> InputUintTexture;
 RWBuffer<uint> TileCountBuffer;
 RWBuffer<uint2> TileDataBuffer;
 RWBuffer<uint2> TileClearBuffer;
@@ -27,7 +29,14 @@ void TileMainCS(uint2 DispatchThreadId : SV_DispatchThreadID, uint LinearIndex :
 	bool bHasHairData = false;
 	if (all(PixelCoord < uint2(BufferResolution)))
 	{
-		bHasHairData = InputTexture.Load(uint3(PixelCoord, 0)) > 0.0f;
+		if (bUintTexture)
+		{
+			bHasHairData = InputUintTexture.Load(uint3(PixelCoord, 0)) > 0;
+		}
+		else
+		{
+			bHasHairData = InputFloatTexture.Load(uint3(PixelCoord, 0)) < 1.0f;
+		}
 	}
 	s_HasTileHair[LinearIndex] = bHasHairData;
 
diff --git a/Engine/Source/Runtime/Renderer/Private/HairStrands/HairStrandsTile.cpp b/Engine/Source/Runtime/Renderer/Private/HairStrands/HairStrandsTile.cpp
--- a/Engine/Source/Runtime/Renderer/Private/HairStrands/HairStrandsTile.cpp
+++ b/Engine/Source/Runtime/Renderer/Private/HairStrands/HairStrandsTile.cpp
@@ -103,7 +103,9 @@ class FHairStrandsTileGenerationPassCS : public FGlobalShader
 	BEGIN_SHADER_PARAMETER_STRUCT(FParameters, )
 		SHADER_PARAMETER_STRUCT_REF(FViewUniformShaderParameters, ViewUniformBuffer)
 		SHADER_PARAMETER(FIntPoint, BufferResolution)
-		SHADER_PARAMETER_RDG_TEXTURE(Texture2D<float>, InputTexture)
+		SHADER_PARAMETER(uint32, bUintTexture)
+		SHADER_PARAMETER_RDG_TEXTURE(Texture2D<float>, InputFloatTexture)
+		SHADER_PARAMETER_RDG_TEXTURE(Texture2D<uint>,  InputUintTexture)
 		SHADER_PARAMETER_RDG_BUFFER_UAV(RWBuffer, TileCountBuffer)
 		SHADER_PARAMETER_RDG_BUFFER_UAV(RWBuffer, TileDataBuffer)
 		SHADER_PARAMETER_RDG_BUFFER_UAV(RWBuffer, TileClearBuffer)
@@ -148,10 +150,13 @@ FHairStrandsTiles AddHairStrandsGenerateTilesPass(
 	FRDGBufferUAVRef TileCountUAV = GraphBuilder.CreateUAV(Out.TileCountBuffer, PF_R32_UINT);
 	AddClearUAVPass(GraphBuilder, TileCountUAV, 0u);
 
+	const bool bUintTexture = InputTexture->Desc.Format == PF_R32_UINT;
 	TShaderMapRef<FHairStrandsTileGenerationPassCS> ComputeShader(View.ShaderMap);
 	FHairStrandsTileGenerationPassCS::FParameters* PassParameters = GraphBuilder.AllocParameters<FHairStrandsTileGenerationPassCS::FParameters>();
 	PassParameters->BufferResolution = InputResolution;//View.ViewRect.Size();
-	PassParameters->InputTexture = InputTexture;
+	PassParameters->bUintTexture = bUintTexture ? 1u : 0u;
+	PassParameters->InputFloatTexture = InputTexture;
+	PassParameters->InputUintTexture  = InputTexture;
 	PassParameters->TileDataBuffer = GraphBuilder.CreateUAV(Out.TileDataBuffer, PF_R16G16_UINT);
 	PassParameters->TileClearBuffer = GraphBuilder.CreateUAV(Out.TileClearBuffer, PF_R16G16_UINT);
 	PassParameters->TileCountBuffer = TileCountUAV;
diff --git a/Engine/Source/Runtime/Renderer/Private/HairStrands/HairStrandsVisibility.cpp b/Engine/Source/Runtime/Renderer/Private/HairStrands/HairStrandsVisibility.cpp
--- a/Engine/Source/Runtime/Renderer/Private/HairStrands/HairStrandsVisibility.cpp
+++ b/Engine/Source/Runtime/Renderer/Private/HairStrands/HairStrandsVisibility.cpp
@@ -21,12 +21,6 @@ DECLARE_GPU_STAT(HairStrandsVisibility);
 
 /////////////////////////////////////////////////////////////////////////////////////////
 
-static int32 GHairStrandsViewTransmittancePassEnable = 1;
-static FAutoConsoleVariableRef CVarHairStrandsCoveragePassEnable(TEXT("r.HairStrands.ViewTransmittancePass"), GHairStrandsViewTransmittancePassEnable, TEXT("Enable accurate transmittance pass for better rendering of small scale hair strand."));
-
-static int32 GHairStrandsMaterialCompactionEnable = 0;
-static FAutoConsoleVariableRef CVarHairStrandsMaterialCompactionEnable(TEXT("r.HairStrands.MaterialCompaction"), GHairStrandsMaterialCompactionEnable, TEXT("Enable extra compaction based on material properties in order to reduce sample per pixel and improve performance."));
-
 static float GHairStrandsMaterialCompactionDepthThreshold = 1.f;
 static float GHairStrandsMaterialCompactionTangentThreshold = 10.f;
 static FAutoConsoleVariableRef CVarHairStrandsMaterialCompactionDepthThreshold(TEXT("r.HairStrands.MaterialCompaction.DepthThreshold"), GHairStrandsMaterialCompactionDepthThreshold, TEXT("Compaction threshold for depth value for material compaction (in centimeters). Default 1 cm."));
@@ -1548,11 +1542,10 @@ class FHairVisibilityPrimitiveIdCompactionCS : public FGlobalShader
 
 	class FGroupSize	: SHADER_PERMUTATION_SPARSE_INT("PERMUTATION_GROUPSIZE", 32, 64);
 	class FVelocity		: SHADER_PERMUTATION_INT("PERMUTATION_VELOCITY", 2);
-	class FViewTransmittance : SHADER_PERMUTATION_INT("PERMUTATION_VIEWTRANSMITTANCE", 2);
-	class FMaterial 	: SHADER_PERMUTATION_INT("PERMUTATION_MATERIAL_COMPACTION", 2);
+	class FTile			: SHADER_PERMUTATION_BOOL("PERMUTATION_TILE");
 	class FPPLL 		: SHADER_PERMUTATION_SPARSE_INT("PERMUTATION_PPLL", 0, 8, 16, 32); // See GetPPLLMaxRenderNodePerPixel
 	class FMSAACount 	: SHADER_PERMUTATION_SPARSE_INT("PERMUTATION_MSAACOUNT", 1, 2, 4, 8);
-	using FPermutationDomain = TShaderPermutationDomain<FGroupSize, FVelocity, FViewTransmittance, FMaterial, FPPLL, FMSAACount>;
+	using FPermutationDomain = TShaderPermutationDomain<FGroupSize, FVelocity, FTile, FPPLL, FMSAACount>;
 
 	BEGIN_SHADER_PARAMETER_STRUCT(FParameters, )
 		SHADER_PARAMETER(FIntPoint, OutputResolution)
@@ -1564,6 +1557,9 @@ class FHairVisibilityPrimitiveIdCompactionCS : public FGlobalShader
 		SHADER_PARAMETER(float, CoverageThreshold)
 		SHADER_PARAMETER(uint32, VelocityType)
 
+		SHADER_PARAMETER(FIntPoint, TileCountXY)
+		SHADER_PARAMETER(uint32, TileSize)
+
 		// Available for the MSAA path
 		SHADER_PARAMETER_RDG_TEXTURE(Texture2D, MSAA_DepthTexture)
 		SHADER_PARAMETER_RDG_TEXTURE(Texture2D, MSAA_IDTexture)
@@ -1584,9 +1580,13 @@ class FHairVisibilityPrimitiveIdCompactionCS : public FGlobalShader
 		SHADER_PARAMETER_RDG_BUFFER_UAV(RWStructuredBuffer, OutCompactNodeData)
 		SHADER_PARAMETER_RDG_BUFFER_UAV(RWBuffer, OutCompactNodeCoord)
 		SHADER_PARAMETER_RDG_TEXTURE_UAV(RWTexture2D, OutVelocityTexture)
+		
+		SHADER_PARAMETER_RDG_BUFFER_SRV(Buffer<uint2>, TileDataBuffer)						// Tile coords (RG16)
+		SHADER_PARAMETER_RDG_BUFFER_SRV(Buffer<uint>, TileCountBuffer)						// Tile total count (actual number of tiles)
 
 		SHADER_PARAMETER_STRUCT_REF(FViewUniformShaderParameters, ViewUniformBuffer)
 		SHADER_PARAMETER_RDG_UNIFORM_BUFFER(FSceneTextureUniformParameters, SceneTexturesStruct)
+		RDG_BUFFER_ACCESS(IndirectBufferArgs, ERHIAccess::IndirectArgs)
 	END_SHADER_PARAMETER_STRUCT()
 
 public:
@@ -1595,8 +1595,7 @@ public:
 	{
 		if (PermutationVector.Get<FPPLL>() > 0)
 		{
-			PermutationVector.Set<FViewTransmittance>(0);
-			PermutationVector.Set<FMSAACount>(4);
+			PermutationVector.Set<FMSAACount>(1);
 		}
 		return PermutationVector;
 	}
@@ -1604,11 +1603,7 @@ public:
 	static bool ShouldCompilePermutation(const FGlobalShaderPermutationParameters& Parameters)
 	{
 		FPermutationDomain PermutationVector(Parameters.PermutationId);
-		if (PermutationVector.Get<FPPLL>() > 0 && PermutationVector.Get<FViewTransmittance>() > 0)
-		{
-			return false;
-		}
-		if (PermutationVector.Get<FPPLL>() > 0 && PermutationVector.Get<FMSAACount>() == 8)
+		if (PermutationVector.Get<FPPLL>() > 0 && PermutationVector.Get<FMSAACount>() != 1)
 		{
 			return false;
 		}
@@ -1624,6 +1619,7 @@ static void AddHairVisibilityPrimitiveIdCompactionPass(
 	const FViewInfo& View,
 	const FHairStrandsMacroGroupDatas& MacroGroupDatas,
 	const uint32 NodeGroupSize,
+	const FHairStrandsTiles& TileData,
 	FHairVisibilityPrimitiveIdCompactionCS::FParameters* PassParameters,
 	FRDGTextureRef& OutCompactCounter,
 	FRDGTextureRef& OutCompactNodeIndex,
@@ -1672,12 +1668,13 @@ static void AddHairVisibilityPrimitiveIdCompactionPass(
 
 	// Adapt the buffer allocation based on the bounding box of the hair macro groups. This allows to reduce the overall allocation size
 	const FIntRect HairRect = ComputeVisibleHairStrandsMacroGroupsRect(View.ViewRect, MacroGroupDatas);
-	const FIntPoint EffectiveResolution(HairRect.Width(), HairRect.Height());
+	const FIntPoint EffectiveResolution = bUsePPLL ? FIntPoint(View.ViewRect.Width(), View.ViewRect.Height()) : FIntPoint(HairRect.Width(), HairRect.Height());
 
 	// Select render node count according to current mode
 	const uint32 MSAASampleCount = GetHairVisibilityRenderMode() == HairVisibilityRenderMode_MSAA_Visibility ? GetMaxSamplePerPixel() : 1;
 	const uint32 PPLLMaxRenderNodePerPixel = GetMaxSamplePerPixel();
 	const uint32 MaxRenderNodeCount = GetTotalSampleCountForAllocation(EffectiveResolution);
+	const bool bUseTile = TileData.IsValid();
 
 	if (bUsePPLL)
 	{
@@ -1697,14 +1694,14 @@ static void AddHairVisibilityPrimitiveIdCompactionPass(
 	const bool bWriteOutVelocity = OutVelocityTexture != nullptr;
 	const uint32 VelocityPermutation = bWriteOutVelocity ? FMath::Clamp(GHairVelocityType + 1, 0, 3) : 0;
 	FHairVisibilityPrimitiveIdCompactionCS::FPermutationDomain PermutationVector;
-	PermutationVector.Set<FHairVisibilityPrimitiveIdCompactionCS::FGroupSize>(GetVendorOptimalGroupSize1D());
+	PermutationVector.Set<FHairVisibilityPrimitiveIdCompactionCS::FGroupSize>(bUseTile ? FHairStrandsTiles::GroupSize : GetVendorOptimalGroupSize1D());
 	PermutationVector.Set<FHairVisibilityPrimitiveIdCompactionCS::FVelocity>(VelocityPermutation > 0 ? 1 : 0);
-	PermutationVector.Set<FHairVisibilityPrimitiveIdCompactionCS::FViewTransmittance>(PassParameters->ViewTransmittanceTexture ? 1 : 0);
-	PermutationVector.Set<FHairVisibilityPrimitiveIdCompactionCS::FMaterial>(GHairStrandsMaterialCompactionEnable ? 1 : 0);
+	PermutationVector.Set<FHairVisibilityPrimitiveIdCompactionCS::FTile>(bUseTile);
 	PermutationVector.Set<FHairVisibilityPrimitiveIdCompactionCS::FPPLL>(bUsePPLL ? PPLLMaxRenderNodePerPixel : 0);
 	PermutationVector.Set<FHairVisibilityPrimitiveIdCompactionCS::FMSAACount>(MSAASampleCount);
 	PermutationVector = FHairVisibilityPrimitiveIdCompactionCS::RemapPermutation(PermutationVector);
 
+	PassParameters->ResolutionOffset = FIntPoint(0,0);
 	PassParameters->OutputResolution = Resolution;
 	PassParameters->VelocityType = VelocityPermutation;
 	PassParameters->MaxNodeCount = MaxRenderNodeCount;
@@ -1727,26 +1724,39 @@ static void AddHairVisibilityPrimitiveIdCompactionPass(
 	{
 		PassParameters->OutVelocityTexture = GraphBuilder.CreateUAV(OutVelocityTexture);
 	}
- 
-	FIntRect TotalRect = ComputeVisibleHairStrandsMacroGroupsRect(View.ViewRect, MacroGroupDatas);
-
-	// Snap the rect onto thread group boundary
-	const FIntPoint GroupSize = GetVendorOptimalGroupSize2D();
-	TotalRect.Min.X = FMath::FloorToInt(float(TotalRect.Min.X) / float(GroupSize.X)) * GroupSize.X;
-	TotalRect.Min.Y = FMath::FloorToInt(float(TotalRect.Min.Y) / float(GroupSize.Y)) * GroupSize.Y;
-	TotalRect.Max.X = FMath::CeilToInt(float(TotalRect.Max.X) / float(GroupSize.X)) * GroupSize.X;
-	TotalRect.Max.Y = FMath::CeilToInt(float(TotalRect.Max.Y) / float(GroupSize.Y)) * GroupSize.Y;
-
-	FIntPoint RectResolution(TotalRect.Width(), TotalRect.Height());
-	PassParameters->ResolutionOffset = FIntPoint(TotalRect.Min.X, TotalRect.Min.Y);
 
+	if (bUseTile)
+	{
+		PassParameters->TileCountXY			= TileData.TileCountXY;
+		PassParameters->TileSize			= FHairStrandsTiles::TileSize;
+		PassParameters->TileCountBuffer		= GraphBuilder.CreateSRV(TileData.TileCountBuffer, PF_R32_UINT);
+		PassParameters->TileDataBuffer		= TileData.TileDataSRV;;
+		PassParameters->IndirectBufferArgs	= TileData.TileIndirectDispatchBuffer;
+	}
+ 
 	TShaderMapRef<FHairVisibilityPrimitiveIdCompactionCS> ComputeShader(View.ShaderMap, PermutationVector);
-	FComputeShaderUtils::AddPass(
-		GraphBuilder,
-		RDG_EVENT_NAME("HairStrandsVisibilityCompaction"),
-		ComputeShader,
-		PassParameters,
-		FComputeShaderUtils::GetGroupCount(RectResolution, GroupSize));
+	if (bUseTile)
+	{
+		FComputeShaderUtils::AddPass(
+			GraphBuilder,
+			RDG_EVENT_NAME("HairStrandsVisibilityCompaction(Tile)"),
+			ComputeShader,
+			PassParameters,
+			TileData.TileIndirectDispatchBuffer,
+			0);
+	}
+	else
+	{
+		const FIntPoint GroupSize = GetVendorOptimalGroupSize2D();
+		const FIntPoint RectResolution(View.ViewRect.Width(), View.ViewRect.Height());
+
+		FComputeShaderUtils::AddPass(
+			GraphBuilder,
+			RDG_EVENT_NAME("HairStrandsVisibilityCompaction(Screen)"),
+			ComputeShader,
+			PassParameters,
+			FComputeShaderUtils::GetGroupCount(RectResolution, GroupSize));
+	}
 
 	OutIndirectArgsBuffer = AddCopyIndirectArgPass(GraphBuilder, &View, NodeGroupSize, 1, OutCompactCounter);
 	OutMaxRenderNodeCount = MaxRenderNodeCount;
@@ -1759,7 +1769,8 @@ class FHairVisibilityCompactionComputeRasterCS : public FGlobalShader
 	SHADER_USE_PARAMETER_STRUCT(FHairVisibilityCompactionComputeRasterCS, FGlobalShader);
 
 	class FGroupSize : SHADER_PERMUTATION_SPARSE_INT("PERMUTATION_GROUPSIZE", 32, 64);
-	using FPermutationDomain = TShaderPermutationDomain<FGroupSize>;
+	class FTile : SHADER_PERMUTATION_BOOL("PERMUTATION_TILE");
+	using FPermutationDomain = TShaderPermutationDomain<FGroupSize, FTile>;
 
 	BEGIN_SHADER_PARAMETER_STRUCT(FParameters, )
 		SHADER_PARAMETER(FIntPoint, OutputResolution)
@@ -1767,6 +1778,9 @@ class FHairVisibilityCompactionComputeRasterCS : public FGlobalShader
 		SHADER_PARAMETER(uint32, SamplerPerPixel)
 		SHADER_PARAMETER(float, CoverageThreshold)
 
+		SHADER_PARAMETER(FIntPoint, TileCountXY)
+		SHADER_PARAMETER(uint32, TileSize)
+
 		SHADER_PARAMETER_RDG_TEXTURE(Texture2D<UlongType>, VisibilityTexture0)
 		SHADER_PARAMETER_RDG_TEXTURE(Texture2D<UlongType>, VisibilityTexture1)
 		SHADER_PARAMETER_RDG_TEXTURE(Texture2D<UlongType>, VisibilityTexture2)
@@ -1779,8 +1793,12 @@ class FHairVisibilityCompactionComputeRasterCS : public FGlobalShader
 		SHADER_PARAMETER_RDG_BUFFER_UAV(RWStructuredBuffer, OutCompactNodeVis)
 		SHADER_PARAMETER_RDG_BUFFER_UAV(RWBuffer, OutCompactNodeCoord)
 
+		SHADER_PARAMETER_RDG_BUFFER_SRV(Buffer<uint2>, TileDataBuffer)						// Tile coords (RG16)
+		SHADER_PARAMETER_RDG_BUFFER_SRV(Buffer<uint>, TileCountBuffer)						// Tile total count (actual number of tiles)
+
 		SHADER_PARAMETER_RDG_UNIFORM_BUFFER(FSceneTextureUniformParameters, SceneTexturesStruct)
 		SHADER_PARAMETER_STRUCT_REF(FViewUniformShaderParameters, ViewUniformBuffer)
+		RDG_BUFFER_ACCESS(IndirectBufferArgs, ERHIAccess::IndirectArgs)
 	END_SHADER_PARAMETER_STRUCT()
 
 public:
@@ -1804,6 +1822,7 @@ static void AddHairVisibilityCompactionComputeRasterPass(
 	const uint32 NodeGroupSize,
 	const uint32 SamplerPerPixel,
 	const FRasterComputeOutput& RasterComputeData,
+	const FHairStrandsTiles& TileData,
 	FRDGTextureRef& InTransmittanceTexture,
 	FRDGTextureRef& OutCompactCounter,
 	FRDGTextureRef& OutCompactNodeIndex,
@@ -1837,6 +1856,7 @@ static void AddHairVisibilityCompactionComputeRasterPass(
 	AddClearUAVPass(GraphBuilder, GraphBuilder.CreateUAV(OutCoverageTexture), 0.f);
 
 	// Select render node count according to current mode
+	const bool bUseTile = TileData.IsValid();
 	const uint32 MSAASampleCount = GetHairVisibilityRenderMode() == HairVisibilityRenderMode_MSAA_Visibility ? GetMaxSamplePerPixel() : 1;
 	const uint32 PPLLMaxRenderNodePerPixel = GetMaxSamplePerPixel();
 	const uint32 MaxRenderNodeCount = GetTotalSampleCountForAllocation(Resolution);
@@ -1862,16 +1882,40 @@ static void AddHairVisibilityCompactionComputeRasterPass(
 	PassParameters->OutCompactNodeCoord		= GraphBuilder.CreateUAV(OutCompactNodeCoord, FHairStrandsVisibilityData::NodeCoordFormat);
 	PassParameters->OutCoverageTexture		= GraphBuilder.CreateUAV(OutCoverageTexture);
 
+	if (bUseTile)
+	{
+		PassParameters->TileCountXY = TileData.TileCountXY;
+		PassParameters->TileSize = FHairStrandsTiles::TileSize;
+		PassParameters->TileCountBuffer = GraphBuilder.CreateSRV(TileData.TileCountBuffer, PF_R32_UINT);
+		PassParameters->TileDataBuffer = TileData.TileDataSRV;;
+		PassParameters->IndirectBufferArgs = TileData.TileIndirectDispatchBuffer;
+	}
+
 	const FIntPoint GroupSize = GetVendorOptimalGroupSize2D();
 	FHairVisibilityCompactionComputeRasterCS::FPermutationDomain PermutationVector;
-	PermutationVector.Set<FHairVisibilityCompactionComputeRasterCS::FGroupSize>(GetVendorOptimalGroupSize1D());
+	PermutationVector.Set<FHairVisibilityCompactionComputeRasterCS::FGroupSize>(bUseTile ? FHairStrandsTiles::GroupSize : GetVendorOptimalGroupSize1D());
+	PermutationVector.Set<FHairVisibilityCompactionComputeRasterCS::FTile>(bUseTile);
 	TShaderMapRef<FHairVisibilityCompactionComputeRasterCS> ComputeShader(View.ShaderMap, PermutationVector);
-	FComputeShaderUtils::AddPass(
-		GraphBuilder,
-		RDG_EVENT_NAME("HairStrandsVisibilityCompaction"),
-		ComputeShader,
-		PassParameters,
-		FComputeShaderUtils::GetGroupCount(Resolution, GroupSize));
+
+	if (bUseTile)
+	{
+		FComputeShaderUtils::AddPass(
+			GraphBuilder,
+			RDG_EVENT_NAME("HairStrands::VisibilityCompaction(Tile)"),
+			ComputeShader,
+			PassParameters,
+			TileData.TileIndirectDispatchBuffer,
+			0);
+	}
+	else
+	{
+		FComputeShaderUtils::AddPass(
+			GraphBuilder,
+			RDG_EVENT_NAME("HairStrands::VisibilityCompaction(Screen)"),
+			ComputeShader,
+			PassParameters,
+			FComputeShaderUtils::GetGroupCount(Resolution, GroupSize));
+	}
 
 	OutIndirectArgsBuffer = AddCopyIndirectArgPass(GraphBuilder, &View, NodeGroupSize, 1, OutCompactCounter);
 	OutMaxRenderNodeCount = MaxRenderNodeCount;
@@ -2999,6 +3043,12 @@ void RenderHairStrandsVisibilityBuffer(
 					VisibilityData.ViewHairCountUintTexture = ViewTransmittance.HairCountTextureUint;
 				}
 
+				// Generate Tile data
+				if (ViewTransmittance.TransmittanceTexture && bGenerateTile)
+				{
+					VisibilityData.TileData = AddHairStrandsGenerateTilesPass(GraphBuilder, View, ViewTransmittance.TransmittanceTexture);
+				}
+
 				{
 					{
 						FRDGBufferRef CompactNodeCoord;
@@ -3010,7 +3060,8 @@ void RenderHairStrandsVisibilityBuffer(
 							VisibilityData.NodeGroupSize,
 							VisibilityData.MaxSampleCount,
 							RasterOutput,
-							ViewTransmittance.TransmittanceTexture,
+							VisibilityData.TileData,
+							ViewTransmittance.TransmittanceTexture, // TODO tile
 							NodeCounter,
 							CompactNodeIndex,
 							CompactNodeVis,
@@ -3020,11 +3071,6 @@ void RenderHairStrandsVisibilityBuffer(
 							IndirectArgsBuffer,
 							VisibilityData.MaxNodeCount);
 
-						// Generate Tile data
-						if (CoverageTexture && bGenerateTile)
-						{
-							VisibilityData.TileData = AddHairStrandsGenerateTilesPass(GraphBuilder, View, CoverageTexture);
-						}
 
 						// Evaluate material based on the visiblity pass result
 						// Output both complete sample data + per-sample velocity
@@ -3112,7 +3158,6 @@ void RenderHairStrandsVisibilityBuffer(
 			{
 				// Run the view transmittance pass if needed (not in PPLL mode that is already a high quality render path)
 				FHairPrimaryTransmittance ViewTransmittance;
-				if (GHairStrandsViewTransmittancePassEnable > 0 && RenderMode != HairVisibilityRenderMode_PPLL)
 				{
 					// Note: Hair count is required for the sky lighting at the moment as it is used for the TT term
 					// TT sampling is disable in hair sky lighting integrator 0. So the GetHairStrandsSkyLightingEnable() check is no longer needed
@@ -3138,6 +3183,12 @@ void RenderHairStrandsVisibilityBuffer(
 
 				}
 
+				// Generate Tile data
+				if (ViewTransmittance.TransmittanceTexture && bGenerateTile)
+				{
+					VisibilityData.TileData = AddHairStrandsGenerateTilesPass(GraphBuilder, View, ViewTransmittance.TransmittanceTexture);
+				}
+
 				struct FRDGMsaaVisibilityResources
 				{
 					FRDGTextureRef DepthTexture;
@@ -3181,6 +3232,7 @@ void RenderHairStrandsVisibilityBuffer(
 						View,
 						MacroGroupDatas,
 						VisibilityData.NodeGroupSize,
+						VisibilityData.TileData,
 						PassParameters,
 						NodeCounter,
 						CompactNodeIndex,
@@ -3191,11 +3243,6 @@ void RenderHairStrandsVisibilityBuffer(
 						IndirectArgsBuffer,
 						VisibilityData.MaxNodeCount);
 
-					// Generate Tile data
-					if (CoverageTexture && bGenerateTile)
-					{
-						VisibilityData.TileData = AddHairStrandsGenerateTilesPass(GraphBuilder, View, CoverageTexture);
-					}
 
 					{
 						const bool bUpdateSampleCoverage = GHairStrandsSortHairSampleByDepth > 0;
@@ -3319,6 +3366,12 @@ void RenderHairStrandsVisibilityBuffer(
 				// Linked list generation pass
 				AddHairVisibilityPPLLPass(GraphBuilder, Scene, &View, MacroGroupDatas, Resolution, InstanceCullingManager, ViewZDepthTexture, PPLLNodeCounterTexture, PPLLNodeIndexTexture, PPLLNodeDataBuffer);
 
+				// Generate Tile data
+				if (PPLLNodeIndexTexture && bGenerateTile)
+				{
+					VisibilityData.TileData = AddHairStrandsGenerateTilesPass(GraphBuilder, View, PPLLNodeIndexTexture); 
+				}
+
 				// Linked list sorting pass and compaction into common representation
 				{
 					FHairVisibilityPrimitiveIdCompactionCS::FParameters* PassParameters = GraphBuilder.AllocParameters<FHairVisibilityPrimitiveIdCompactionCS::FParameters>();
@@ -3335,6 +3388,7 @@ void RenderHairStrandsVisibilityBuffer(
 						View,
 						MacroGroupDatas,
 						VisibilityData.NodeGroupSize,
+						VisibilityData.TileData,
 						PassParameters,
 						NodeCounter,
 						CompactNodeIndex,
@@ -3355,11 +3409,6 @@ void RenderHairStrandsVisibilityBuffer(
 					VisibilityData.NodeCount = NodeCounter;
 				}
 
-				// Generate Tile data
-				if (CoverageTexture && bGenerateTile)
-				{
-					VisibilityData.TileData = AddHairStrandsGenerateTilesPass(GraphBuilder, View, CoverageTexture);
-				}
 
 				if (bRunColorAndDepthPatching)
 				{
