Remove Opacity Pass MobileSceneCapture
Remove MobileSceneCapture, move everything to the shared SceneCaptureRendering file
Remove MobileSceneCapture.usf

diff --git a/Engine/Shaders/Private/MobileSceneCapture.usf b/Engine/Shaders/Private/MobileSceneCapture.usf
deleted file mode 100644
diff --git a/Engine/Shaders/Private/SceneTexturesCommon.ush b/Engine/Shaders/Private/SceneTexturesCommon.ush
--- a/Engine/Shaders/Private/SceneTexturesCommon.ush
+++ b/Engine/Shaders/Private/SceneTexturesCommon.ush
@@ -111,6 +111,15 @@ float4 GatherSceneDepth(float2 UV, float2 InvBufferSize)
 	#define	VULKAN_SUBPASS_DEPTHFETCH	VULKAN_PROFILE
 #endif
 
+float3 CalcSceneColor(float2 ScreenUV)
+{
+#if SCENE_TEXTURES_DISABLED
+	return float3(0.0f, 0.0f, 0.0f);
+#else
+	return Texture2DSampleLevel(MobileSceneTextures.SceneColorTexture, MobileSceneTextures.SceneColorTextureSampler, ScreenUV, 0).rgb;
+#endif
+}
+
 /** return all channels of the scene lighting texture */
 float4 CalcFullSceneColor(float2 ScreenUV)
 {
diff --git a/Engine/Source/Runtime/Renderer/Private/DeferredShadingRenderer.cpp b/Engine/Source/Runtime/Renderer/Private/DeferredShadingRenderer.cpp
--- a/Engine/Source/Runtime/Renderer/Private/DeferredShadingRenderer.cpp
+++ b/Engine/Source/Runtime/Renderer/Private/DeferredShadingRenderer.cpp
@@ -65,6 +65,7 @@
 #include "InstanceCulling/InstanceCullingManager.h"
 #include "ProfilingDebugging/CpuProfilerTrace.h"
 #include "Engine/SubsurfaceProfile.h"
+#include "SceneCaptureRendering.h"
 
 extern int32 GNaniteShowStats;
 
@@ -3117,7 +3118,7 @@ void FDeferredShadingSceneRenderer::Render(FRDGBuilder& GraphBuilder)
 
 	FRDGTextureRef ViewFamilyTexture = TryCreateViewFamilyTexture(GraphBuilder, ViewFamily);
 
-	CopySceneCaptureComponentToTarget(GraphBuilder, SceneTextures.UniformBuffer, ViewFamilyTexture);
+	CopySceneCaptureComponentToTarget(GraphBuilder, SceneTextures, ViewFamilyTexture, ViewFamily, Views);
 
 	for (int32 ViewIndex = 0; ViewIndex < Views.Num(); ++ViewIndex)
 	{
diff --git a/Engine/Source/Runtime/Renderer/Private/DeferredShadingRenderer.h b/Engine/Source/Runtime/Renderer/Private/DeferredShadingRenderer.h
--- a/Engine/Source/Runtime/Renderer/Private/DeferredShadingRenderer.h
+++ b/Engine/Source/Runtime/Renderer/Private/DeferredShadingRenderer.h
@@ -811,11 +811,6 @@ private:
 
 	bool IsNaniteEnabled() const;
 
-	void CopySceneCaptureComponentToTarget(
-		FRDGBuilder& GraphBuilder,
-		TRDGUniformBufferRef<FSceneTextureUniformParameters> SceneTexturesUniformBuffer,
-		FRDGTextureRef ViewFamilyTexture);
-
 	void SetupImaginaryReflectionTextureParameters(
 		FRDGBuilder& GraphBuilder,
 		const FViewInfo& View,
diff --git a/Engine/Source/Runtime/Renderer/Private/MeshDrawCommands.cpp b/Engine/Source/Runtime/Renderer/Private/MeshDrawCommands.cpp
--- a/Engine/Source/Runtime/Renderer/Private/MeshDrawCommands.cpp
+++ b/Engine/Source/Runtime/Renderer/Private/MeshDrawCommands.cpp
@@ -1072,7 +1072,6 @@ void FParallelMeshDrawCommandPass::DispatchPassSetup(
 		case EMeshPass::TranslucencyAfterDOFModulate: TaskContext.TranslucencyPass = ETranslucencyPass::TPT_TranslucencyAfterDOFModulate; break;
 		case EMeshPass::TranslucencyAfterMotionBlur: TaskContext.TranslucencyPass = ETranslucencyPass::TPT_TranslucencyAfterMotionBlur; break;
 		case EMeshPass::TranslucencyAll: TaskContext.TranslucencyPass = ETranslucencyPass::TPT_AllTranslucency; break;
-		case EMeshPass::MobileInverseOpacity: TaskContext.TranslucencyPass = ETranslucencyPass::TPT_StandardTranslucency; break;
 	}
 
 	FMemory::Memswap(&TaskContext.MeshDrawCommands, &InOutMeshDrawCommands, sizeof(InOutMeshDrawCommands));
diff --git a/Engine/Source/Runtime/Renderer/Private/MobileSceneCaptureRendering.cpp b/Engine/Source/Runtime/Renderer/Private/MobileSceneCaptureRendering.cpp
deleted file mode 100644
diff --git a/Engine/Source/Runtime/Renderer/Private/MobileSceneCaptureRendering.h b/Engine/Source/Runtime/Renderer/Private/MobileSceneCaptureRendering.h
deleted file mode 100644
diff --git a/Engine/Source/Runtime/Renderer/Private/MobileTranslucentRendering.cpp b/Engine/Source/Runtime/Renderer/Private/MobileTranslucentRendering.cpp
--- a/Engine/Source/Runtime/Renderer/Private/MobileTranslucentRendering.cpp
+++ b/Engine/Source/Runtime/Renderer/Private/MobileTranslucentRendering.cpp
@@ -45,60 +45,3 @@ void FMobileSceneRenderer::RenderTranslucency(FRHICommandListImmediate& RHICmdLi
 		View.ParallelMeshDrawCommandPasses[StandardTranslucencyMeshPass].DispatchDraw(nullptr, RHICmdList, &MeshPassInstanceCullingDrawParams[StandardTranslucencyMeshPass]);
 	}
 }
-
-BEGIN_SHADER_PARAMETER_STRUCT(FMobileBasePassParameters, )
-	SHADER_PARAMETER_STRUCT_INCLUDE(FViewShaderParameters, View)
-	SHADER_PARAMETER_RDG_UNIFORM_BUFFER(FMobileBasePassUniformParameters, MobileBasePass)
-	RENDER_TARGET_BINDING_SLOTS()
-END_SHADER_PARAMETER_STRUCT()
-
-void FMobileSceneRenderer::RenderInverseOpacity(FRDGBuilder& GraphBuilder, const FViewInfo& View)
-{
-	View.BeginRenderView();
-	UpdateDirectionalLightUniformBuffers(GraphBuilder, View);
-
-	const FSceneTextures& SceneTextures = FSceneTextures::Get(GraphBuilder);
-	EMobileSceneTextureSetupMode SetupMode = EMobileSceneTextureSetupMode::CustomDepth;
-
-	auto* InverseOpacityParameters = GraphBuilder.AllocParameters<FMobileBasePassParameters>();
-	InverseOpacityParameters->View = View.GetShaderParameters();
-	InverseOpacityParameters->MobileBasePass = CreateMobileBasePassUniformBuffer(GraphBuilder, View, EMobileBasePass::Translucent, SetupMode);
-	InverseOpacityParameters->RenderTargets[0] = FRenderTargetBinding(SceneTextures.Color.Target, SceneTextures.Color.Resolve, ERenderTargetLoadAction::EClear);
-	InverseOpacityParameters->RenderTargets.DepthStencil = FDepthStencilBinding(SceneTextures.Depth.Target, ERenderTargetLoadAction::EClear, FExclusiveDepthStencil::DepthWrite_StencilWrite);
-	// Opacity could fetch depth as we use exactly the same shaders as in base pass
-	InverseOpacityParameters->RenderTargets.SubpassHint = ESubpassHint::DepthReadSubpass;
-
-	GraphBuilder.AddPass(RDG_EVENT_NAME("InverseOpacityPass"), InverseOpacityParameters, ERDGPassFlags::Raster,
-		[this, &View](FRHICommandListImmediate& RHICmdList)
-	{
-		// Mobile multi-view is not side by side stereo
-		const FViewInfo& TranslucentViewport = (View.bIsMobileMultiViewEnabled) ? Views[0] : View;
-		RHICmdList.SetViewport(TranslucentViewport.ViewRect.Min.X, TranslucentViewport.ViewRect.Min.Y, 0.0f, TranslucentViewport.ViewRect.Max.X, TranslucentViewport.ViewRect.Max.Y, 1.0f);
-
-		// Default clear value for a SceneColor is (0,0,0,0), after this pass will blend inverse opacity into final render target with an 1-SrcAlpha op
-		// to make this blending work untouched pixels must have alpha = 1
-		DrawClearQuad(RHICmdList, FLinearColor(0, 0, 0, 1));
-		// GPUCULL_TODO: View.ParallelMeshDrawCommandPasses[EMeshPass::MobileInverseOpacity].BuildRenderingCommands(GraphBuilder, Scene->GPUScene);
-
-		RHICmdList.NextSubpass();
-		if (ShouldRenderTranslucency(ETranslucencyPass::TPT_AllTranslucency) && View.ShouldRenderView())
-		{
-			View.ParallelMeshDrawCommandPasses[EMeshPass::MobileInverseOpacity].DispatchDraw(nullptr, RHICmdList);
-		}
-	});
-}
-
-// This pass is registered only when we render to scene capture, see UpdateSceneCaptureContentMobile_RenderThread()
-FMeshPassProcessor* CreateMobileInverseOpacityPassProcessor(const FScene* Scene, const FSceneView* InViewIfDynamicMeshCommand, FMeshPassDrawListContext* InDrawListContext)
-{
-	FMeshPassProcessorRenderState PassDrawRenderState;
-	PassDrawRenderState.SetDepthStencilState(TStaticDepthStencilState<false, CF_Always>::GetRHI());
-	PassDrawRenderState.SetBlendState(TStaticBlendState<CW_ALPHA, BO_Add, BF_Zero, BF_One, BO_Add, BF_Zero, BF_InverseSourceAlpha>::GetRHI());
-	PassDrawRenderState.SetDepthStencilAccess(FExclusiveDepthStencil::DepthRead_StencilRead);
-
-	const FMobileBasePassMeshProcessor::EFlags Flags = 
-		FMobileBasePassMeshProcessor::EFlags::CanUseDepthStencil | 
-		FMobileBasePassMeshProcessor::EFlags::ForcePassDrawRenderState;
-
-	return new(FMemStack::Get()) FMobileBasePassMeshProcessor(Scene, Scene->GetFeatureLevel(), InViewIfDynamicMeshCommand, PassDrawRenderState, InDrawListContext, Flags, ETranslucencyPass::TPT_AllTranslucency);
-}
\ No newline at end of file
diff --git a/Engine/Source/Runtime/Renderer/Private/SceneCaptureRendering.cpp b/Engine/Source/Runtime/Renderer/Private/SceneCaptureRendering.cpp
--- a/Engine/Source/Runtime/Renderer/Private/SceneCaptureRendering.cpp
+++ b/Engine/Source/Runtime/Renderer/Private/SceneCaptureRendering.cpp
@@ -3,8 +3,7 @@
 /*=============================================================================
 	
 =============================================================================*/
-
-#include "CoreMinimal.h"
+#include "SceneCaptureRendering.h"
 #include "Containers/ArrayView.h"
 #include "Misc/MemStack.h"
 #include "EngineDefines.h"
@@ -36,7 +35,6 @@
 #include "ScenePrivate.h"
 #include "PostProcess/SceneFilterRendering.h"
 #include "ScreenRendering.h"
-#include "MobileSceneCaptureRendering.h"
 #include "ClearQuad.h"
 #include "PipelineStateCache.h"
 #include "RendererModule.h"
@@ -53,7 +51,7 @@ public:
 
 	BEGIN_SHADER_PARAMETER_STRUCT(FParameters, )
 		SHADER_PARAMETER_STRUCT_REF(FViewUniformShaderParameters, View)
-		SHADER_PARAMETER_RDG_UNIFORM_BUFFER(FSceneTextureUniformParameters, SceneTextures)
+		SHADER_PARAMETER_STRUCT_INCLUDE(FSceneTextureShaderParameters, SceneTextures)
 		RENDER_TARGET_BINDING_SLOTS()
 	END_SHADER_PARAMETER_STRUCT()
 
@@ -72,7 +70,7 @@ public:
 	class FSourceModeDimension : SHADER_PERMUTATION_ENUM_CLASS("SOURCE_MODE", ESourceMode);
 	using FPermutationDomain = TShaderPermutationDomain<FSourceModeDimension>;
 
-	static FPermutationDomain GetPermutationVector(ESceneCaptureSource CaptureSource)
+	static FPermutationDomain GetPermutationVector(ESceneCaptureSource CaptureSource, bool bIsMobilePlatform)
 	{
 		ESourceMode SourceMode = ESourceMode::MAX;
 		switch (CaptureSource)
@@ -101,14 +99,22 @@ public:
 		default:
 			checkf(false, TEXT("SceneCaptureSource not implemented."));
 		}
+
+		if (bIsMobilePlatform && (SourceMode == ESourceMode::Normal || SourceMode == ESourceMode::BaseColor))
+		{
+			SourceMode = ESourceMode::ColorAndOpacity;
+		}
 		FPermutationDomain PermutationVector;
 		PermutationVector.Set<FSourceModeDimension>(SourceMode);
 		return PermutationVector;
 	}
 
 	static bool ShouldCompilePermutation(const FGlobalShaderPermutationParameters& Parameters) 
-	{ 
-		return IsFeatureLevelSupported(Parameters.Platform, ERHIFeatureLevel::SM5);
+	{
+		FPermutationDomain PermutationVector(Parameters.PermutationId);
+
+		auto SourceModeDim = PermutationVector.Get<FSourceModeDimension>();
+		return !IsMobilePlatform(Parameters.Platform) || (SourceModeDim != ESourceMode::Normal && SourceModeDim != ESourceMode::BaseColor);
 	}
 
 	static void ModifyCompilationEnvironment(const FGlobalShaderPermutationParameters& Parameters, FShaderCompilerEnvironment& OutEnvironment)
@@ -157,10 +163,13 @@ static bool CaptureNeedsSceneColor(ESceneCaptureSource CaptureSource)
 
 static TFunction<void(FRHICommandList& RHICmdList)> CopyCaptureToTargetSetViewportFn;
 
-void FDeferredShadingSceneRenderer::CopySceneCaptureComponentToTarget(
+void CopySceneCaptureComponentToTarget(
 	FRDGBuilder& GraphBuilder,
-	TRDGUniformBufferRef<FSceneTextureUniformParameters> SceneTexturesUniformBuffer,
-	FRDGTextureRef ViewFamilyTexture)
+	const FMinimalSceneTextures& SceneTextures,
+	FRDGTextureRef ViewFamilyTexture,
+	const FSceneViewFamily& ViewFamily,
+	const TArray<FViewInfo>& Views,
+	bool bNeedsFlippedRenderTarget)
 {
 	ESceneCaptureSource SceneCaptureSource = ViewFamily.SceneCaptureSource;
 
@@ -192,15 +201,15 @@ void FDeferredShadingSceneRenderer::CopySceneCaptureComponentToTarget(
 			GraphicsPSOInit.BlendState = TStaticBlendState<>::GetRHI();
 		}
 
-		const FSceneCapturePS::FPermutationDomain PixelPermutationVector = FSceneCapturePS::GetPermutationVector(SceneCaptureSource);
+		const FSceneCapturePS::FPermutationDomain PixelPermutationVector = FSceneCapturePS::GetPermutationVector(SceneCaptureSource, IsMobilePlatform(ViewFamily.GetShaderPlatform()));
 
 		for (int32 ViewIndex = 0; ViewIndex < Views.Num(); ViewIndex++)
 		{
-			FViewInfo& View = Views[ViewIndex];
+			const FViewInfo& View = Views[ViewIndex];
 
 			FSceneCapturePS::FParameters* PassParameters = GraphBuilder.AllocParameters<FSceneCapturePS::FParameters>();
 			PassParameters->View = View.ViewUniformBuffer;
-			PassParameters->SceneTextures = SceneTexturesUniformBuffer;
+			PassParameters->SceneTextures = SceneTextures.GetSceneTextureShaderParameters(ViewFamily.GetFeatureLevel());
 			PassParameters->RenderTargets[0] = FRenderTargetBinding(ViewFamilyTexture, ERenderTargetLoadAction::ENoAction);
 
 			TShaderMapRef<FScreenVS> VertexShader(View.ShaderMap);
@@ -215,7 +224,7 @@ void FDeferredShadingSceneRenderer::CopySceneCaptureComponentToTarget(
 				RDG_EVENT_NAME("View(%d)", ViewIndex),
 				PassParameters,
 				ERDGPassFlags::Raster,
-				[PassParameters, GraphicsPSOInit, VertexShader, PixelShader, &View] (FRHICommandList& RHICmdList)
+				[PassParameters, GraphicsPSOInit, VertexShader, PixelShader, &View, bNeedsFlippedRenderTarget] (FRHICommandList& RHICmdList)
 			{
 				FGraphicsPipelineStateInitializer LocalGraphicsPSOInit = GraphicsPSOInit;
 				RHICmdList.ApplyCachedRenderTargets(LocalGraphicsPSOInit);
@@ -224,16 +233,34 @@ void FDeferredShadingSceneRenderer::CopySceneCaptureComponentToTarget(
 				
 				CopyCaptureToTargetSetViewportFn(RHICmdList);
 
-				DrawRectangle(
-					RHICmdList,
-					View.ViewRect.Min.X, View.ViewRect.Min.Y,
-					View.ViewRect.Width(), View.ViewRect.Height(),
-					View.ViewRect.Min.X, View.ViewRect.Min.Y,
-					View.ViewRect.Width(), View.ViewRect.Height(),
-					View.UnconstrainedViewRect.Size(),
-					GetSceneTextureExtent(),
-					VertexShader,
-					EDRF_UseTriangleOptimization);
+				if (bNeedsFlippedRenderTarget)
+				{
+					CopyCaptureToTargetSetViewportFn(RHICmdList);
+					DrawRectangle(
+						RHICmdList,
+						View.ViewRect.Min.X, View.ViewRect.Min.Y,
+						View.ViewRect.Width(), View.ViewRect.Height(),
+						View.ViewRect.Min.X, View.ViewRect.Height() - View.ViewRect.Min.Y,
+						View.ViewRect.Width(), -View.ViewRect.Height(),
+						View.UnconstrainedViewRect.Size(),
+						GetSceneTextureExtent(),
+						VertexShader,
+						EDRF_UseTriangleOptimization);
+				}
+				else
+				{
+					CopyCaptureToTargetSetViewportFn(RHICmdList);
+					DrawRectangle(
+						RHICmdList,
+						View.ViewRect.Min.X, View.ViewRect.Min.Y,
+						View.ViewRect.Width(), View.ViewRect.Height(),
+						View.ViewRect.Min.X, View.ViewRect.Min.Y,
+						View.ViewRect.Width(), View.ViewRect.Height(),
+						View.UnconstrainedViewRect.Size(),
+						GetSceneTextureExtent(),
+						VertexShader,
+						EDRF_UseTriangleOptimization);
+				}
 			});
 		}
 	}
@@ -311,6 +338,174 @@ static void UpdateSceneCaptureContentDeferred_RenderThread(
 	SceneRenderer->RenderThreadEnd(RHICmdList);
 }
 
+void UpdateSceneCaptureContentMobile_RenderThread(
+	FRHICommandListImmediate& RHICmdList,
+	FSceneRenderer* SceneRenderer,
+	FRenderTarget* RenderTarget,
+	FTexture* RenderTargetTexture,
+	const FString& EventName,
+	const FResolveParams& ResolveParams,
+	bool bGenerateMips,
+	const FGenerateMipsParams& GenerateMipsParams,
+	bool bDisableFlipCopyGLES)
+{
+	SceneRenderer->RenderThreadBegin(RHICmdList);
+
+	// update any resources that needed a deferred update
+	FDeferredUpdateResource::UpdateResources(RHICmdList);
+	bool bUseSceneTextures = SceneRenderer->ViewFamily.SceneCaptureSource != SCS_FinalColorLDR &&
+		SceneRenderer->ViewFamily.SceneCaptureSource != SCS_FinalColorHDR;
+
+	{
+#if WANTS_DRAW_MESH_EVENTS
+		SCOPED_DRAW_EVENTF(RHICmdList, SceneCaptureMobile, TEXT("SceneCaptureMobile %s"), *EventName);
+		FRDGBuilder GraphBuilder(RHICmdList, RDG_EVENT_NAME("SceneCaptureMobile %s", *EventName));
+#else
+		SCOPED_DRAW_EVENT(RHICmdList, UpdateSceneCaptureContentMobile_RenderThread);
+		FRDGBuilder GraphBuilder(RHICmdList, RDG_EVENT_NAME("SceneCaptureMobile"));
+#endif
+
+		FViewInfo& View = SceneRenderer->Views[0];
+
+		const bool bIsMobileHDR = IsMobileHDR();
+		const bool bRHINeedsFlip = RHINeedsToSwitchVerticalAxis(GMaxRHIShaderPlatform) && !bDisableFlipCopyGLES;
+		// note that GLES code will flip the image when:
+		//	bIsMobileHDR && SceneCaptureSource == SCS_FinalColorLDR (flip performed during post processing)
+		//	!bIsMobileHDR (rendering is flipped by vertex shader)
+		// they need flipping again so it is correct for texture addressing.
+		const bool bNeedsFlippedCopy = (!bIsMobileHDR || !bUseSceneTextures) && bRHINeedsFlip;
+		const bool bNeedsFlippedFinalColor = bNeedsFlippedCopy && !bUseSceneTextures;
+
+		// Intermediate render target that will need to be flipped (needed on !IsMobileHDR())
+		FRDGTextureRef FlippedOutputTexture{};
+
+		const FRenderTarget* Target = SceneRenderer->ViewFamily.RenderTarget;
+		if (bNeedsFlippedFinalColor)
+		{
+			// We need to use an intermediate render target since the result will be flipped
+			auto& RenderTargetRHI = Target->GetRenderTargetTexture();
+			FRDGTextureDesc Desc(FRDGTextureDesc::Create2D(
+				Target->GetSizeXY(),
+				RenderTargetRHI.GetReference()->GetFormat(),
+				RenderTargetRHI.GetReference()->GetClearBinding(),
+				TexCreate_RenderTargetable));
+			FlippedOutputTexture = GraphBuilder.CreateTexture(Desc, TEXT("SceneCaptureFlipped"));
+		}
+
+		// We don't support screen percentage in scene capture.
+		FIntRect ViewRect = View.UnscaledViewRect;
+		FIntRect UnconstrainedViewRect = View.UnconstrainedViewRect;
+
+		if (bNeedsFlippedFinalColor)
+		{
+			AddClearRenderTargetPass(GraphBuilder, FlippedOutputTexture, FLinearColor::Black, ViewRect);
+		}
+
+		if (ResolveParams.DestRect.IsValid())
+		{
+			CopyCaptureToTargetSetViewportFn = [ResolveParams, bNeedsFlippedFinalColor, &ViewRect, FlippedOutputTexture](FRHICommandList& RHICmdList)
+			{
+				RHICmdList.SetScissorRect(false, 0, 0, 0, 0);
+
+				if (bNeedsFlippedFinalColor)
+				{
+					FResolveRect DestRect = ResolveParams.DestRect;
+					int32 TileYID = DestRect.Y1 / ViewRect.Height();
+					int32 TileYCount = (FlippedOutputTexture->Desc.GetSize().Y / ViewRect.Height()) - 1;
+					DestRect.Y1 = (TileYCount - TileYID) * ViewRect.Height();
+					DestRect.Y2 = DestRect.Y1 + ViewRect.Height();
+					RHICmdList.SetViewport
+					(
+						float(DestRect.X1),
+						float(DestRect.Y1),
+						0.0f,
+						float(DestRect.X2),
+						float(DestRect.Y2),
+						1.0f
+					);
+				}
+				else
+				{
+					RHICmdList.SetViewport
+					(
+						float(ResolveParams.DestRect.X1),
+						float(ResolveParams.DestRect.Y1),
+						0.0f,
+						float(ResolveParams.DestRect.X2),
+						float(ResolveParams.DestRect.Y2),
+						1.0f
+					);
+				}
+			};
+		}
+		else
+		{
+			CopyCaptureToTargetSetViewportFn = [](FRHICommandList& RHICmdList) {};
+		}
+
+		// Render the scene normally
+		{
+			RDG_RHI_EVENT_SCOPE(GraphBuilder, RenderScene);
+
+			if (bNeedsFlippedFinalColor)
+			{
+				// Helper class to allow setting render target
+				struct FRenderTargetOverride : public FRenderTarget
+				{
+					FRenderTargetOverride(const FRenderTarget* TargetIn, FRHITexture2D* OverrideTexture)
+					{
+						RenderTargetTextureRHI = OverrideTexture;
+						OriginalTarget = TargetIn;
+					}
+
+					virtual FIntPoint GetSizeXY() const override { return FIntPoint(RenderTargetTextureRHI->GetSizeX(), RenderTargetTextureRHI->GetSizeY()); }
+					virtual float GetDisplayGamma() const override { return OriginalTarget->GetDisplayGamma(); }
+
+					FTexture2DRHIRef GetTextureParamRef() { return RenderTargetTextureRHI; }
+					const FRenderTarget* OriginalTarget;
+				};
+
+				// Hijack the render target
+				FRHITexture2D* FlippedOutputTextureRHI = GraphBuilder.ConvertToExternalTexture(FlippedOutputTexture)->GetTargetableRHI()->GetTexture2D();
+				SceneRenderer->ViewFamily.RenderTarget = GraphBuilder.AllocObject<FRenderTargetOverride>(Target, FlippedOutputTextureRHI); //-V506
+			}
+
+			SceneRenderer->Render(GraphBuilder);
+
+			if (bNeedsFlippedFinalColor)
+			{
+				// And restore it
+				SceneRenderer->ViewFamily.RenderTarget = Target;
+			}
+		}
+
+		FRDGTextureRef OutputTexture = RegisterExternalTexture(GraphBuilder, Target->GetRenderTargetTexture(), TEXT("OutputTexture"));
+		const FMinimalSceneTextures& SceneTextures = FSceneTextures::Get(GraphBuilder);
+
+		const FIntPoint TargetSize(UnconstrainedViewRect.Width(), UnconstrainedViewRect.Height());
+		{
+			// We need to flip this texture upside down (since we depended on tonemapping to fix this on the hdr path)
+			RDG_EVENT_SCOPE(GraphBuilder, "CaptureSceneColor");
+			CopySceneCaptureComponentToTarget(
+				GraphBuilder,
+				SceneTextures,
+				OutputTexture,
+				SceneRenderer->ViewFamily,
+				SceneRenderer->Views,
+				bNeedsFlippedFinalColor);
+		}
+
+		if (bGenerateMips)
+		{
+			FGenerateMips::Execute(GraphBuilder, OutputTexture, GenerateMipsParams);
+		}
+
+		GraphBuilder.Execute();
+	}
+
+	SceneRenderer->RenderThreadEnd(RHICmdList);
+}
+
 static void ODSCapture_RenderThread(
 	FRDGBuilder& GraphBuilder,
 	FRDGTextureRef LeftEyeTexture,
diff --git a/Engine/Source/Runtime/Renderer/Private/SceneCaptureRendering.h b/Engine/Source/Runtime/Renderer/Private/SceneCaptureRendering.h
new file mode 100644
--- /dev/null
+++ b/Engine/Source/Runtime/Renderer/Private/SceneCaptureRendering.h
@@ -0,0 +1,19 @@
+// Copyright Epic Games, Inc. All Rights Reserved.
+
+/*=============================================================================
+	SceneCaptureRendering.h: SceneCaptureRendering definitions.
+=============================================================================*/
+
+#pragma once
+
+#include "CoreMinimal.h"
+#include "RenderGraphDefinitions.h"
+#include "SceneTextures.h"
+
+void CopySceneCaptureComponentToTarget(
+	FRDGBuilder& GraphBuilder,
+	const FMinimalSceneTextures& SceneTextures,
+	FRDGTextureRef ViewFamilyTexture,
+	const FSceneViewFamily& ViewFamily,
+	const TArray<FViewInfo>& Views,
+	bool bNeedsFlippedRenderTarget = false);
\ No newline at end of file
diff --git a/Engine/Source/Runtime/Renderer/Private/SceneRendering.h b/Engine/Source/Runtime/Renderer/Private/SceneRendering.h
--- a/Engine/Source/Runtime/Renderer/Private/SceneRendering.h
+++ b/Engine/Source/Runtime/Renderer/Private/SceneRendering.h
@@ -2247,7 +2247,6 @@ public:
 
 	virtual bool ShouldRenderPrePass() const override;
 
-	void RenderInverseOpacity(FRDGBuilder& GraphBuilder, const FViewInfo& View);
 
 protected:
 	/** Finds the visible dynamic shadows for each view. */
diff --git a/Engine/Source/Runtime/Renderer/Private/SceneTextures.cpp b/Engine/Source/Runtime/Renderer/Private/SceneTextures.cpp
--- a/Engine/Source/Runtime/Renderer/Private/SceneTextures.cpp
+++ b/Engine/Source/Runtime/Renderer/Private/SceneTextures.cpp
@@ -107,11 +107,11 @@ static EPixelFormat GetGBufferFFormat()
 	return NormalGBufferFormat;
 }
 
-static EPixelFormat GetMobileSceneColorFormat(EShaderPlatform ShaderPlatform)
+static EPixelFormat GetMobileSceneColorFormat(const FSceneView& View)
 {
 	const EPixelFormat DefaultLowPrecisionFormat = IHeadMountedDisplayModule::IsAvailable() && IHeadMountedDisplayModule::Get().IsStandaloneStereoOnlyDevice()
 		? PF_R8G8B8A8 : PF_B8G8R8A8;
-	const EPixelFormat DefaultPrecisionFormat = IsMobilePropagateAlphaEnabled(ShaderPlatform) ? PF_FloatRGBA : PF_FloatR11G11B10;
+	const EPixelFormat DefaultPrecisionFormat = IsMobilePropagateAlphaEnabled(View.GetShaderPlatform()) || View.bIsSceneCapture? PF_FloatRGBA : PF_FloatR11G11B10;
 
 	EPixelFormat DefaultColorFormat = (!IsMobileHDR() || !GSupportsRenderTargetFormat_PF_FloatRGBA) ? DefaultLowPrecisionFormat : DefaultPrecisionFormat;
 
@@ -535,7 +535,7 @@ FSceneTexturesConfig FSceneTexturesConfig::Create(const FSceneViewFamily& ViewFa
 
 	case EShadingPath::Mobile:
 	{
-		Config.ColorFormat = GetMobileSceneColorFormat(Config.ShaderPlatform);
+		Config.ColorFormat = GetMobileSceneColorFormat(*(ViewFamily.Views[0]));
 
 		// On mobile the scene depth is calculated from the alpha component of the scene color
 		// Use FarPlane for alpha to ensure un-rendered pixels have max depth...
diff --git a/Engine/Source/Runtime/Renderer/Private/SceneVisibility.cpp b/Engine/Source/Runtime/Renderer/Private/SceneVisibility.cpp
--- a/Engine/Source/Runtime/Renderer/Private/SceneVisibility.cpp
+++ b/Engine/Source/Runtime/Renderer/Private/SceneVisibility.cpp
@@ -2428,11 +2428,6 @@ struct FRelevancePacket
 							{
 								DrawCommandPacket.AddCommandsForMesh(PrimitiveIndex, PrimitiveSceneInfo, StaticMeshRelevance, StaticMesh, Scene, bCanCache, EMeshPass::Distortion);
 							}
-
-							if (ShadingPath == EShadingPath::Mobile && View.bIsSceneCapture)
-							{
-								DrawCommandPacket.AddCommandsForMesh(PrimitiveIndex, PrimitiveSceneInfo, StaticMeshRelevance, StaticMesh, Scene, bCanCache, EMeshPass::MobileInverseOpacity);
-							}
 						}
 
 #if WITH_EDITOR
@@ -2888,12 +2883,6 @@ void ComputeDynamicMeshRelevance(EShadingPath ShadingPath, bool bAddLightmapDens
 			PassMask.Set(EMeshPass::Distortion);
 			View.NumVisibleDynamicMeshElements[EMeshPass::Distortion] += NumElements;
 		}
-
-		if (ShadingPath == EShadingPath::Mobile && View.bIsSceneCapture)
-		{
-			PassMask.Set(EMeshPass::MobileInverseOpacity);
-			View.NumVisibleDynamicMeshElements[EMeshPass::MobileInverseOpacity] += NumElements;
-		}
 	}
 
 #if WITH_EDITOR
diff --git a/Engine/Source/Runtime/Renderer/Public/MeshPassProcessor.h b/Engine/Source/Runtime/Renderer/Public/MeshPassProcessor.h
--- a/Engine/Source/Runtime/Renderer/Public/MeshPassProcessor.h
+++ b/Engine/Source/Runtime/Renderer/Public/MeshPassProcessor.h
@@ -44,7 +44,6 @@ namespace EMeshPass
 		DebugViewMode, /** Any of EDebugViewShaderMode */
 		CustomDepth,
 		MobileBasePassCSM,  /** Mobile base pass with CSM shading enabled */
-		MobileInverseOpacity,  /** Mobile specific scene capture, Non-cached */
 		VirtualTexture,
 		LumenCardCapture,
 		DitheredLODFadingOutMaskPass, /** A mini depth pass used to mark pixels with dithered LOD fading out. Currently only used by ray tracing shadows. */
@@ -85,7 +84,6 @@ inline const TCHAR* GetMeshPassName(EMeshPass::Type MeshPass)
 	case EMeshPass::DebugViewMode: return TEXT("DebugViewMode");
 	case EMeshPass::CustomDepth: return TEXT("CustomDepth");
 	case EMeshPass::MobileBasePassCSM: return TEXT("MobileBasePassCSM");
-	case EMeshPass::MobileInverseOpacity: return TEXT("MobileInverseOpacity");
 	case EMeshPass::VirtualTexture: return TEXT("VirtualTexture");
 	case EMeshPass::LumenCardCapture: return TEXT("LumenCardCapture");
 	case EMeshPass::DitheredLODFadingOutMaskPass: return TEXT("DitheredLODFadingOutMaskPass");
@@ -98,9 +96,9 @@ inline const TCHAR* GetMeshPassName(EMeshPass::Type MeshPass)
 	}
 
 #if WITH_EDITOR
-	static_assert(EMeshPass::Num == 23 + 4, "Need to update switch(MeshPass) after changing EMeshPass");
+	static_assert(EMeshPass::Num == 22 + 4, "Need to update switch(MeshPass) after changing EMeshPass");
 #else
-	static_assert(EMeshPass::Num == 23, "Need to update switch(MeshPass) after changing EMeshPass");
+	static_assert(EMeshPass::Num == 22, "Need to update switch(MeshPass) after changing EMeshPass");
 #endif
 
 	checkf(0, TEXT("Missing case for EMeshPass %u"), (uint32)MeshPass);
