ASan: Avoid reading the stack out of bounds when clearing the buffer by simply repeating the clear value if the buffer size is greater than the value's size.

#rnx

diff --git a/Engine/Source/Runtime/Renderer/Private/SystemTextures.cpp b/Engine/Source/Runtime/Renderer/Private/SystemTextures.cpp
--- a/Engine/Source/Runtime/Renderer/Private/SystemTextures.cpp
+++ b/Engine/Source/Runtime/Renderer/Private/SystemTextures.cpp
@@ -1043,7 +1043,7 @@ void FormatData1010102(const TInType& In, uint8* Out, uint32& OutByteCount)
 	OutByteCount = 4;
 }
 
-template <typename TInType>
+template<typename TInType>
 void FormatData5551(const TInType& In, uint8* Out, uint32& OutByteCount)
 {
 	uint16* OutTyped = (uint16*)Out;
@@ -1387,10 +1387,20 @@ FRDGBufferRef GetInternalDefaultBuffer(
 		RHIBuffer = RHICreateVertexBuffer(BufferSize, BUF_Static | BUF_ShaderResource, CreateInfo);
 	}
 
-	void* DestPtr = GraphBuilder.RHICmdList.LockBuffer(RHIBuffer, 0, BufferSize, RLM_WriteOnly);
+	uint8* DestPtr = static_cast<uint8*>(GraphBuilder.RHICmdList.LockBuffer(RHIBuffer, 0, BufferSize, RLM_WriteOnly));
 	if (Value)
 	{
-		FMemory::Memcpy(DestPtr, Value, BufferSize);
+		const uint8 *EndPtr = DestPtr + BufferSize; 
+		for (uint32 Offset = 0; Offset < (BufferSize / sizeof(TClearValue)); Offset++)
+		{
+			FMemory::Memcpy(DestPtr, Value, sizeof(TClearValue));
+			DestPtr += sizeof(TClearValue);
+		}
+		if (DestPtr < EndPtr)
+		{
+			// Zero out the remainder. Byte-splitting the init value is undefined.
+			FMemory::Memzero(DestPtr, static_cast<SIZE_T>(EndPtr - DestPtr));
+		}
 	}
 	else
 	{
