Fix misleading comment

diff --git a/Engine/Source/Runtime/Renderer/Private/SystemTextures.cpp b/Engine/Source/Runtime/Renderer/Private/SystemTextures.cpp
--- a/Engine/Source/Runtime/Renderer/Private/SystemTextures.cpp
+++ b/Engine/Source/Runtime/Renderer/Private/SystemTextures.cpp
@@ -582,7 +582,7 @@ void FSystemTextures::InitializeFeatureLevelDependentTextures(FRHICommandListImm
 
 		} // end Create the PerlinNoise3D texture
 
-		// GTAO Randomization texture	
+		// LTC Textures	(used by Rect Lights)
 		{
 			{
 				FRHIResourceCreateInfo CreateInfo(TEXT("LTCMat"));
