Minor cleanup of shadow passes focused on removing unnecessary command list and feature level arguments.

diff --git a/Engine/Source/Runtime/Renderer/Private/DeferredShadingRenderer.cpp b/Engine/Source/Runtime/Renderer/Private/DeferredShadingRenderer.cpp
--- a/Engine/Source/Runtime/Renderer/Private/DeferredShadingRenderer.cpp
+++ b/Engine/Source/Runtime/Renderer/Private/DeferredShadingRenderer.cpp
@@ -1934,10 +1934,10 @@ void FDeferredShadingSceneRenderer::Render(FRDGBuilder& GraphBuilder)
 
 	FInstanceCullingManager& InstanceCullingManager = *GraphBuilder.AllocObject<FInstanceCullingManager>(Scene->GPUScene.IsEnabled());
 
-	bool bDoInitViewAftersPrepass = false;
+	bool bDoInitViewsAfterPrepass = false;
 	{
 		RDG_GPU_STAT_SCOPE(GraphBuilder, VisibilityCommands);
-		bDoInitViewAftersPrepass = InitViews(GraphBuilder, SceneTexturesConfig, BasePassDepthStencilAccess, ILCTaskData, InstanceCullingManager);
+		bDoInitViewsAfterPrepass = InitViews(GraphBuilder, SceneTexturesConfig, BasePassDepthStencilAccess, ILCTaskData, InstanceCullingManager);
 	}
 
 	// Compute & commit the final state of the entire dependency topology of the renderer.
@@ -2014,7 +2014,7 @@ void FDeferredShadingSceneRenderer::Render(FRDGBuilder& GraphBuilder)
 			InstanceCullingManager.BeginDeferredCulling(GraphBuilder, Scene->GPUScene);
 		}
 
-		if (!bDoInitViewAftersPrepass)
+		if (!bDoInitViewsAfterPrepass)
 		{
 			PrepareDistanceFieldScene(GraphBuilder, false);
 		}
@@ -2097,7 +2097,7 @@ void FDeferredShadingSceneRenderer::Render(FRDGBuilder& GraphBuilder)
 		DynamicVertexBufferForInitViews.Commit();
 		DynamicReadBufferForInitViews.Commit();
 
-		if (!bDoInitViewAftersPrepass)
+		if (!bDoInitViewsAfterPrepass)
 		{
 			DynamicVertexBufferForInitShadows.Commit();
 			DynamicIndexBufferForInitShadows.Commit();
@@ -2209,7 +2209,7 @@ void FDeferredShadingSceneRenderer::Render(FRDGBuilder& GraphBuilder)
 		}
 	}
 
-	if (bDoInitViewAftersPrepass)
+	if (bDoInitViewsAfterPrepass)
 	{
 		GraphBuilder.Drain();
 
diff --git a/Engine/Source/Runtime/Renderer/Private/SceneOcclusion.cpp b/Engine/Source/Runtime/Renderer/Private/SceneOcclusion.cpp
--- a/Engine/Source/Runtime/Renderer/Private/SceneOcclusion.cpp
+++ b/Engine/Source/Runtime/Renderer/Private/SceneOcclusion.cpp
@@ -264,7 +264,7 @@ void FSceneViewState::TrimOcclusionHistory(float CurrentTime, float MinHistoryTi
 	}
 }
 
-bool FSceneViewState::IsShadowOccluded(FRHICommandListImmediate& RHICmdList, FSceneViewState::FProjectedShadowKey ShadowKey, int32 NumBufferedFrames) const
+bool FSceneViewState::IsShadowOccluded(FSceneViewState::FProjectedShadowKey ShadowKey, int32 NumBufferedFrames) const
 {
 	// Find the shadow's occlusion query from the previous frame.
 	// Get the oldest occlusion query	
@@ -277,7 +277,7 @@ bool FSceneViewState::IsShadowOccluded(FRHICommandListImmediate& RHICmdList, FSc
 	// Only block on the query if not running SLI
 	const bool bWaitOnQuery = GNumAlternateFrameRenderingGroups == 1;
 
-	if (Query && RHICmdList.GetRenderQueryResult(Query->GetQuery(), NumSamples, bWaitOnQuery))
+	if (Query && RHIGetRenderQueryResult(Query->GetQuery(), NumSamples, bWaitOnQuery))
 	{
 		// If the shadow's occlusion query didn't have any pixels visible the previous frame, it's occluded.
 		return NumSamples == 0;
diff --git a/Engine/Source/Runtime/Renderer/Private/ScenePrivate.h b/Engine/Source/Runtime/Renderer/Private/ScenePrivate.h
--- a/Engine/Source/Runtime/Renderer/Private/ScenePrivate.h
+++ b/Engine/Source/Runtime/Renderer/Private/ScenePrivate.h
@@ -1223,7 +1223,7 @@ public:
 	 * @param Primitive - The shadow subject.
 	 * @param Light - The shadow source.
 	 */
-	bool IsShadowOccluded(FRHICommandListImmediate& RHICmdList, FSceneViewState::FProjectedShadowKey ShadowKey, int32 NumBufferedFrames) const;
+	bool IsShadowOccluded(FSceneViewState::FProjectedShadowKey ShadowKey, int32 NumBufferedFrames) const;
 
 	/**
 	* Retrieve a single-pixel render targets with intra-frame state for use in eye adaptation post processing.
diff --git a/Engine/Source/Runtime/Renderer/Private/SceneRendering.h b/Engine/Source/Runtime/Renderer/Private/SceneRendering.h
--- a/Engine/Source/Runtime/Renderer/Private/SceneRendering.h
+++ b/Engine/Source/Runtime/Renderer/Private/SceneRendering.h
@@ -2019,7 +2019,6 @@ protected:
 
 	/** Creates a per object projected shadow for the given interaction. */
 	void CreatePerObjectProjectedShadow(
-		FRHICommandListImmediate& RHICmdList,
 		FLightPrimitiveInteraction* Interaction,
 		bool bCreateTranslucentObjectShadow,
 		bool bCreateInsetObjectShadow,
@@ -2028,7 +2027,6 @@ protected:
 
 	/** Creates shadows for the given interaction. */
 	void SetupInteractionShadows(
-		FRHICommandListImmediate& RHICmdList,
 		FLightPrimitiveInteraction* Interaction,
 		FVisibleLightInfo& VisibleLightInfo,
 		bool bStaticSceneOnly,
@@ -2087,7 +2085,7 @@ protected:
 	static void GatherSimpleLights(const FSceneViewFamily& ViewFamily, const TArray<FViewInfo>& Views, FSimpleLightArray& SimpleLights);
 
 	/** Calculates projected shadow visibility. */
-	void InitProjectedShadowVisibility(FRHICommandListImmediate& RHICmdList);	
+	void InitProjectedShadowVisibility();
 
 	/** Gathers dynamic mesh elements for all shadows. */
 	void GatherShadowDynamicMeshElements(FGlobalDynamicIndexBuffer& DynamicIndexBuffer, FGlobalDynamicVertexBuffer& DynamicVertexBuffer, FGlobalDynamicReadBuffer& DynamicReadBuffer, FInstanceCullingManager& InstanceCullingManager);
diff --git a/Engine/Source/Runtime/Renderer/Private/ShadowDepthRendering.cpp b/Engine/Source/Runtime/Renderer/Private/ShadowDepthRendering.cpp
--- a/Engine/Source/Runtime/Renderer/Private/ShadowDepthRendering.cpp
+++ b/Engine/Source/Runtime/Renderer/Private/ShadowDepthRendering.cpp
@@ -1229,7 +1229,7 @@ void FProjectedShadowInfo::RenderDepth(
 	}
 }
 
-void FProjectedShadowInfo::ModifyViewForShadow(FRHICommandList& RHICmdList, FViewInfo* FoundView) const
+void FProjectedShadowInfo::ModifyViewForShadow(FViewInfo* FoundView) const
 {
 	FIntRect OriginalViewRect = FoundView->ViewRect;
 	FoundView->ViewRect = GetOuterViewRect();
@@ -1286,7 +1286,7 @@ FViewInfo* FProjectedShadowInfo::FindViewForShadow(FSceneRenderer* SceneRenderer
 	return FoundView;
 }
 
-void FProjectedShadowInfo::SetupShadowDepthView(FRHICommandListImmediate& RHICmdList, FSceneRenderer* SceneRenderer)
+void FProjectedShadowInfo::SetupShadowDepthView(FSceneRenderer* SceneRenderer)
 {
 	FViewInfo* FoundView = FindViewForShadow(SceneRenderer);
 	check(FoundView && IsInRenderingThread());
@@ -1294,7 +1294,7 @@ void FProjectedShadowInfo::SetupShadowDepthView(FRHICommandListImmediate& RHICmd
 	// We are starting a new collection of dynamic primitives for the shadow views.
 	DepthPassView->DynamicPrimitiveCollector = FGPUScenePrimitiveCollector(&SceneRenderer->GetGPUSceneDynamicContext());
 
-	ModifyViewForShadow(RHICmdList, DepthPassView);
+	ModifyViewForShadow(DepthPassView);
 	ShadowDepthView = DepthPassView;
 }
 
diff --git a/Engine/Source/Runtime/Renderer/Private/ShadowRendering.h b/Engine/Source/Runtime/Renderer/Private/ShadowRendering.h
--- a/Engine/Source/Runtime/Renderer/Private/ShadowRendering.h
+++ b/Engine/Source/Runtime/Renderer/Private/ShadowRendering.h
@@ -577,7 +577,7 @@ public:
 	/**
 	 * Adds a primitive to the shadow's subject list.
 	 */
-	void AddSubjectPrimitive(FPrimitiveSceneInfo* PrimitiveSceneInfo, TArray<FViewInfo>* ViewArray, ERHIFeatureLevel::Type FeatureLevel, bool bRecordShadowSubjectForMobileShading);
+	void AddSubjectPrimitive(FPrimitiveSceneInfo* PrimitiveSceneInfo, TArray<FViewInfo>* ViewArray, bool bRecordShadowSubjectForMobileShading);
 
 	uint64 AddSubjectPrimitive_AnyThread(
 		const FPrimitiveSceneInfoCompact& PrimitiveSceneInfoCompact,
@@ -591,7 +591,6 @@ public:
 	void FinalizeAddSubjectPrimitive(
 		struct FAddSubjectPrimitiveOp const& Op,
 		TArray<FViewInfo>* ViewArray,
-		ERHIFeatureLevel::Type FeatureLevel,
 		struct FFinalizeAddSubjectPrimitiveContext& Context);
 
 	/**
@@ -685,7 +684,7 @@ public:
 	const FPrimitiveSceneInfo* GetParentSceneInfo() const { return ParentSceneInfo; }
 
 	/** Creates a new view from the pool and caches it in ShadowDepthView for depth rendering. */
-	void SetupShadowDepthView(FRHICommandListImmediate& RHICmdList, FSceneRenderer* SceneRenderer);
+	void SetupShadowDepthView(FSceneRenderer* SceneRenderer);
 
 	FShadowDepthType GetShadowDepthType() const 
 	{
@@ -770,7 +769,7 @@ private:
 	/**
 	* Modifies the passed in view for this shadow
 	*/
-	void ModifyViewForShadow(FRHICommandList& RHICmdList, FViewInfo* FoundView) const;
+	void ModifyViewForShadow(FViewInfo* FoundView) const;
 
 	friend class FVirtualShadowMapArray;
 	void BeginRenderView(FRDGBuilder& GraphBuilder, FScene* Scene);
diff --git a/Engine/Source/Runtime/Renderer/Private/ShadowSetup.cpp b/Engine/Source/Runtime/Renderer/Private/ShadowSetup.cpp
--- a/Engine/Source/Runtime/Renderer/Private/ShadowSetup.cpp
+++ b/Engine/Source/Runtime/Renderer/Private/ShadowSetup.cpp
@@ -1560,7 +1560,7 @@ bool FProjectedShadowInfo::ShouldDrawStaticMeshes_AnyThread(
 	return bDrawingStaticMeshes;
 }
 
-void FProjectedShadowInfo::AddSubjectPrimitive(FPrimitiveSceneInfo* PrimitiveSceneInfo, TArray<FViewInfo>* ViewArray, ERHIFeatureLevel::Type FeatureLevel, bool bShouldRecordShadowSubjectsForMobileShading)
+void FProjectedShadowInfo::AddSubjectPrimitive(FPrimitiveSceneInfo* PrimitiveSceneInfo, TArray<FViewInfo>* ViewArray, bool bShouldRecordShadowSubjectsForMobileShading)
 {
 	// Ray traced shadows use the GPU managed distance field object buffers, no CPU culling should be used
 	check(!bRayTracedDistanceField);
@@ -1571,12 +1571,12 @@ void FProjectedShadowInfo::AddSubjectPrimitive(FPrimitiveSceneInfo* PrimitiveSce
 	{
 		const FPrimitiveSceneProxy* Proxy = PrimitiveSceneInfo->Proxy;
 
-		TArray<FViewInfo*, TInlineAllocator<1> > Views;
+		TArrayView<FViewInfo> Views;
 		const bool bWholeSceneDirectionalShadow = IsWholeSceneDirectionalShadow();
 
 		if (bWholeSceneDirectionalShadow)
 		{
-			Views.Add(DependentView);
+			Views = TArrayView<FViewInfo>(DependentView, 1);
 		}
 		else
 		{
@@ -1586,22 +1586,17 @@ void FProjectedShadowInfo::AddSubjectPrimitive(FPrimitiveSceneInfo* PrimitiveSce
 				CascadeSettings.ShadowSplitIndex,
 				bDirectionalLight ? TEXT("true") : TEXT("false"));
 
-			for (int32 ViewIndex = 0; ViewIndex < ViewArray->Num(); ViewIndex++)
-			{
-				Views.Add(&(*ViewArray)[ViewIndex]);
-			}
+			Views = TArrayView<FViewInfo>(*ViewArray);
 		}
 
 		bool bOpaque = false;
 		bool bTranslucentRelevance = false;
 		bool bShadowRelevance = false;
 
-		uint32 ViewMask = 0;
 		int32 PrimitiveId = PrimitiveSceneInfo->GetIndex();
 
-		for (int32 ViewIndex = 0, Num = Views.Num(); ViewIndex < Num; ViewIndex++)
+		for (FViewInfo& CurrentView : Views)
 		{
-			FViewInfo& CurrentView = *Views[ViewIndex];
 			FPrimitiveViewRelevance& ViewRelevance = CurrentView.PrimitiveViewRelevanceMap[PrimitiveId];
 
 			if (!ViewRelevance.bInitializedThisFrame)
@@ -1641,8 +1636,6 @@ void FProjectedShadowInfo::AddSubjectPrimitive(FPrimitiveSceneInfo* PrimitiveSce
 				// Compute the subject primitive's view relevance since it wasn't cached
 				// Update the main view's PrimitiveViewRelevanceMap
 				ViewRelevance = PrimitiveSceneInfo->Proxy->GetViewRelevance(&CurrentView);
-
-				ViewMask |= (1 << ViewIndex);
 			}
 
 			bOpaque |= ViewRelevance.bOpaque || ViewRelevance.bMasked;
@@ -1653,15 +1646,15 @@ void FProjectedShadowInfo::AddSubjectPrimitive(FPrimitiveSceneInfo* PrimitiveSce
 		if (bShadowRelevance)
 		{
 			// Update the primitive component's last render time. Allows the component to update when using bCastWhenHidden.
-			const float CurrentWorldTime = Views[0]->Family->CurrentWorldTime;
+			const float CurrentWorldTime = Views[0].Family->CurrentWorldTime;
 			PrimitiveSceneInfo->UpdateComponentLastRenderTime(CurrentWorldTime, /*bUpdateLastRenderTimeOnScreen=*/false);
 
 			if (PrimitiveSceneInfo->NeedsUniformBufferUpdate())
 			{
-				for (int32 ViewIndex = 0; ViewIndex < Views.Num(); ++ViewIndex)
+				for (const FViewInfo& CurrentView : Views)
 				{
 					// Main view visible primitives are processed on parallel tasks, updating uniform buffer them here will cause a race condition.
-					check(!Views[ViewIndex]->PrimitiveVisibilityMap[PrimitiveSceneInfo->GetIndex()]);
+					check(!CurrentView.PrimitiveVisibilityMap[PrimitiveSceneInfo->GetIndex()]);
 				}
 
 				PrimitiveSceneInfo->ConditionalUpdateUniformBuffer(FRHICommandListExecutor::GetImmediateCommandList());
@@ -1682,10 +1675,8 @@ void FProjectedShadowInfo::AddSubjectPrimitive(FPrimitiveSceneInfo* PrimitiveSce
 
 			if (PrimitiveSceneInfo->StaticMeshes.Num() > 0)
 			{
-				for (int32 ViewIndex = 0, ViewCount = Views.Num(); ViewIndex < ViewCount; ViewIndex++)
+				for (FViewInfo& CurrentView : Views)
 				{
-					FViewInfo& CurrentView = *Views[ViewIndex];
-
 					const float DistanceSquared = ( Bounds.Origin - CurrentView.ShadowViewMatrices.GetViewOrigin() ).SizeSquared();
 
 					if (bWholeSceneShadow && CacheMode != SDCM_StaticPrimitivesOnly)
@@ -1963,7 +1954,6 @@ void FProjectedShadowInfo::PresizeSubjectPrimitiveArrays(const FAddSubjectPrimit
 void FProjectedShadowInfo::FinalizeAddSubjectPrimitive(
 	const FAddSubjectPrimitiveOp& Op,
 	TArray<FViewInfo>* ViewArray,
-	ERHIFeatureLevel::Type FeatureLevel,
 	FFinalizeAddSubjectPrimitiveContext& Context)
 {
 	FPrimitiveSceneInfo* PrimitiveSceneInfo = Op.PrimitiveSceneInfo;
@@ -1989,7 +1979,7 @@ void FProjectedShadowInfo::FinalizeAddSubjectPrimitive(
 
 	if (Result.bAddOnRenderThread)
 	{
-		AddSubjectPrimitive(PrimitiveSceneInfo, ViewArray, FeatureLevel, Result.bShouldRecordShadowSubjectsForMobile);
+		AddSubjectPrimitive(PrimitiveSceneInfo, ViewArray, Result.bShouldRecordShadowSubjectsForMobile);
 		return;
 	}
 
@@ -2588,7 +2578,6 @@ bool ShouldCreateObjectShadowForStationaryLight(const FLightSceneInfo* LightScen
 }
 
 void FSceneRenderer::SetupInteractionShadows(
-	FRHICommandListImmediate& RHICmdList,
 	FLightPrimitiveInteraction* Interaction, 
 	FVisibleLightInfo& VisibleLightInfo, 
 	bool bStaticSceneOnly,
@@ -2624,13 +2613,12 @@ void FSceneRenderer::SetupInteractionShadows(
 			&& (bCreateTranslucentObjectShadow || bCreateInsetObjectShadow || bCreateObjectShadowForStationaryLight))
 		{
 			// Create projected shadow infos
-			CreatePerObjectProjectedShadow(RHICmdList, Interaction, bCreateTranslucentObjectShadow, bCreateInsetObjectShadow || bCreateObjectShadowForStationaryLight, ViewDependentWholeSceneShadows, PreShadows);
+			CreatePerObjectProjectedShadow(Interaction, bCreateTranslucentObjectShadow, bCreateInsetObjectShadow || bCreateObjectShadowForStationaryLight, ViewDependentWholeSceneShadows, PreShadows);
 		}
 	}
 }
 
 void FSceneRenderer::CreatePerObjectProjectedShadow(
-	FRHICommandListImmediate& RHICmdList,
 	FLightPrimitiveInteraction* Interaction, 
 	bool bCreateTranslucentObjectShadow, 
 	bool bCreateOpaqueObjectShadow,
@@ -2676,7 +2664,7 @@ void FSceneRenderer::CreatePerObjectProjectedShadow(
 			!bCreateOpaqueObjectShadow ||
 			(
 				!View.bIgnoreExistingQueries &&	View.State &&
-				((FSceneViewState*)View.State)->IsShadowOccluded(RHICmdList, OpaqueKey, NumBufferedFrames)
+				((FSceneViewState*)View.State)->IsShadowOccluded(OpaqueKey, NumBufferedFrames)
 			);
 
 		const FSceneViewState::FProjectedShadowKey TranslucentKey(PrimitiveSceneInfo->PrimitiveComponentId, LightSceneInfo->Proxy->GetLightComponent(), INDEX_NONE, true);
@@ -2685,7 +2673,7 @@ void FSceneRenderer::CreatePerObjectProjectedShadow(
 			!bCreateTranslucentObjectShadow ||
 			(
 				!View.bIgnoreExistingQueries && View.State &&
-				((FSceneViewState*)View.State)->IsShadowOccluded(RHICmdList, TranslucentKey, NumBufferedFrames)
+				((FSceneViewState*)View.State)->IsShadowOccluded(TranslucentKey, NumBufferedFrames)
 			);
 
 		// if subject doesn't render in the main pass, it's never considered visible
@@ -2901,7 +2889,7 @@ void FSceneRenderer::CreatePerObjectProjectedShadow(
 						for (int32 ChildIndex = 0, ChildCount = ShadowGroupPrimitives.Num(); ChildIndex < ChildCount; ChildIndex++)
 						{
 							FPrimitiveSceneInfo* ShadowChild = ShadowGroupPrimitives[ChildIndex];
-							ProjectedShadowInfo->AddSubjectPrimitive(ShadowChild, &Views, FeatureLevel, false);
+							ProjectedShadowInfo->AddSubjectPrimitive(ShadowChild, &Views, false);
 						}
 					}
 					else if (bShadowIsPotentiallyVisibleNextFrame)
@@ -2943,7 +2931,7 @@ void FSceneRenderer::CreatePerObjectProjectedShadow(
 						for (int32 ChildIndex = 0, ChildCount = ShadowGroupPrimitives.Num(); ChildIndex < ChildCount; ChildIndex++)
 						{
 							FPrimitiveSceneInfo* ShadowChild = ShadowGroupPrimitives[ChildIndex];
-							ProjectedShadowInfo->AddSubjectPrimitive(ShadowChild, &Views, FeatureLevel, false);
+							ProjectedShadowInfo->AddSubjectPrimitive(ShadowChild, &Views, false);
 						}
 					}
 					else if (bShadowIsPotentiallyVisibleNextFrame)
@@ -3661,7 +3649,7 @@ void FSceneRenderer::CreateWholeSceneProjectedShadow(
 								FBoxSphereBounds const& Bounds = Interaction->GetPrimitiveSceneInfo()->Proxy->GetBounds();
 								if (IntersectsConvexHulls(LightViewFrustumConvexHulls, Bounds))
 								{
-									ProjectedShadowInfo->AddSubjectPrimitive(Interaction->GetPrimitiveSceneInfo(), &Views, FeatureLevel, false);
+									ProjectedShadowInfo->AddSubjectPrimitive(Interaction->GetPrimitiveSceneInfo(), &Views, false);
 								}
 							}
 						}
@@ -3787,7 +3775,7 @@ void FSceneRenderer::CreateWholeSceneProjectedShadow(
 	}
 }
 
-void FSceneRenderer::InitProjectedShadowVisibility(FRHICommandListImmediate& RHICmdList)
+void FSceneRenderer::InitProjectedShadowVisibility()
 {
 	SCOPE_CYCLE_COUNTER(STAT_InitProjectedShadowVisibility);
 	int32 NumBufferedFrames = FOcclusionQueryHelpers::GetNumBufferedFrames(FeatureLevel);
@@ -3860,7 +3848,6 @@ void FSceneRenderer::InitProjectedShadowVisibility(FRHICommandListImmediate& RHI
 						// Check if the shadow is occluded.
 						bShadowIsOccluded =
 							((FSceneViewState*)View.State)->IsShadowOccluded(
-							RHICmdList,
 							FSceneViewState::FProjectedShadowKey(ProjectedShadowInfo),
 							NumBufferedFrames
 							);
@@ -4325,7 +4312,7 @@ struct FGatherShadowPrimitivesPacket
 
 			for (int32 PrimitiveIndex = 0; PrimitiveIndex < SubjectPrimitives.Num(); PrimitiveIndex++)
 			{
-				ProjectedShadowInfo->FinalizeAddSubjectPrimitive(SubjectPrimitives[PrimitiveIndex], &Views, FeatureLevel, Context);
+				ProjectedShadowInfo->FinalizeAddSubjectPrimitive(SubjectPrimitives[PrimitiveIndex], &Views, Context);
 			}
 		}
 
@@ -4352,7 +4339,7 @@ struct FGatherShadowPrimitivesPacket
 
 			for (int32 PrimitiveIndex = 0; PrimitiveIndex < SubjectPrimitives.Num(); PrimitiveIndex++)
 			{
-				ProjectedShadowInfo->FinalizeAddSubjectPrimitive(SubjectPrimitives[PrimitiveIndex], nullptr, FeatureLevel, Context);
+				ProjectedShadowInfo->FinalizeAddSubjectPrimitive(SubjectPrimitives[PrimitiveIndex], nullptr, Context);
 			}
 
 			if (ProjectedShadowInfo->CacheMode == SDCM_StaticPrimitivesOnly && ProjectedShadowInfo->IsWholeSceneDirectionalShadow())
@@ -4872,7 +4859,7 @@ void FSceneRenderer::AllocateShadowDepthTargets(FRHICommandListImmediate& RHICmd
 				{
 					if (ProjectedShadowInfo->HasVirtualShadowMap() || ProjectedShadowInfo->VirtualShadowMapClipmap.IsValid())
 					{
-						ProjectedShadowInfo->SetupShadowDepthView(RHICmdList, this);
+						ProjectedShadowInfo->SetupShadowDepthView(this);
 						SortedShadowsForShadowDepthPass.VirtualShadowMapShadows.Add(ProjectedShadowInfo);
 					}
 					else if (ProjectedShadowInfo->bPreShadow && ProjectedShadowInfo->bAllocatedInPreshadowCache)
@@ -4972,7 +4959,7 @@ void FSceneRenderer::AllocateShadowDepthTargets(FRHICommandListImmediate& RHICmd
 
 			// Note: adding preshadows whose depths are cached so that GatherDynamicMeshElements
 			// will still happen, which is necessary for preshadow receiver stenciling
-			ProjectedShadowInfo->SetupShadowDepthView(RHICmdList, this);
+			ProjectedShadowInfo->SetupShadowDepthView(this);
 			SortedShadowsForShadowDepthPass.PreshadowCache.Shadows.Add(ProjectedShadowInfo);
 		}
 	}
@@ -5124,7 +5111,7 @@ void FSceneRenderer::AllocateAtlasedShadowDepthTargets(
 			if (ProjectedShadowInfo->bAllocated)
 			{
 				ProjectedShadowInfo->RenderTargets.CopyReferencesFromRenderTargets(ShadowMapAtlas.RenderTargets);
-				ProjectedShadowInfo->SetupShadowDepthView(RHICmdList, this);
+				ProjectedShadowInfo->SetupShadowDepthView(this);
 				ShadowMapAtlas.Shadows.Add(ProjectedShadowInfo);
 			}
 		}
@@ -5154,7 +5141,7 @@ void FSceneRenderer::AllocateCachedShadowDepthTargets(FRHICommandListImmediate&
 		ProjectedShadowInfo->bAllocated = true;
 		ProjectedShadowInfo->RenderTargets.CopyReferencesFromRenderTargets(ShadowMap.RenderTargets);
 
-		ProjectedShadowInfo->SetupShadowDepthView(RHICmdList, this);
+		ProjectedShadowInfo->SetupShadowDepthView(this);
 		ShadowMap.Shadows.Add(ProjectedShadowInfo);
 	}
 }
@@ -5235,7 +5222,7 @@ void FSceneRenderer::AllocateCSMDepthTargets(
 			{
 				check(ProjectedShadowInfo->bAllocated);
 				ProjectedShadowInfo->RenderTargets.CopyReferencesFromRenderTargets(ShadowMapAtlas.RenderTargets);
-				ProjectedShadowInfo->SetupShadowDepthView(RHICmdList, this);
+				ProjectedShadowInfo->SetupShadowDepthView(this);
 				ShadowMapAtlas.Shadows.Add(ProjectedShadowInfo);
 			}
 		}
@@ -5279,7 +5266,7 @@ void FSceneRenderer::AllocateOnePassPointLightDepthTargets(FRHICommandListImmedi
 				ProjectedShadowInfo->bAllocated = true;
 				ProjectedShadowInfo->RenderTargets.CopyReferencesFromRenderTargets(ShadowMapCubemap.RenderTargets);
 				
-				ProjectedShadowInfo->SetupShadowDepthView(RHICmdList, this);
+				ProjectedShadowInfo->SetupShadowDepthView(this);
 				ShadowMapCubemap.Shadows.Add(ProjectedShadowInfo);
 			}
 		}
@@ -5365,7 +5352,7 @@ void FSceneRenderer::AllocateTranslucentShadowDepthTargets(FRHICommandListImmedi
 			}
 
 			ProjectedShadowInfo->RenderTargets.CopyReferencesFromRenderTargets(ShadowMapAtlas.RenderTargets);
-			ProjectedShadowInfo->SetupShadowDepthView(RHICmdList, this);
+			ProjectedShadowInfo->SetupShadowDepthView(this);
 			ShadowMapAtlas.Shadows.Add(ProjectedShadowInfo);
 		}
 	}
@@ -5502,39 +5489,39 @@ void FSceneRenderer::InitDynamicShadows(FRHICommandListImmediate& RHICmdList, FG
 									// Mobile uses the scene's MobileDirectionalLights only for whole scene shadows.
 									&& (LightSceneInfo == Scene->MobileDirectionalLights[0] || LightSceneInfo == Scene->MobileDirectionalLights[1] || LightSceneInfo == Scene->MobileDirectionalLights[2])))
 							{
-							AddViewDependentWholeSceneShadowsForView(ViewDependentWholeSceneShadows, ViewDependentWholeSceneShadowsThatNeedCulling, VisibleLightInfo, *LightSceneInfo, NumCSMCachesUpdatedThisFrame);
+								AddViewDependentWholeSceneShadowsForView(ViewDependentWholeSceneShadows, ViewDependentWholeSceneShadowsThatNeedCulling, VisibleLightInfo, *LightSceneInfo, NumCSMCachesUpdatedThisFrame);
 							}
 
-						// Disable per-object shadows when non-nanite VSMs are enabled
-						if( !bProjectEnableNonNaniteVirtualShadowMaps && (!bMobile || (LightSceneInfo->Proxy->CastsModulatedShadows() && !LightSceneInfo->Proxy->UseCSMForDynamicObjects())))
+							// Disable per-object shadows when non-nanite VSMs are enabled
+							if (!bProjectEnableNonNaniteVirtualShadowMaps && (!bMobile || (LightSceneInfo->Proxy->CastsModulatedShadows() && !LightSceneInfo->Proxy->UseCSMForDynamicObjects())))
 							{
-							const TArray<FLightPrimitiveInteraction*>* InteractionShadowPrimitives = LightSceneInfo->GetInteractionShadowPrimitives();
+								const TArray<FLightPrimitiveInteraction*>* InteractionShadowPrimitives = LightSceneInfo->GetInteractionShadowPrimitives();
 
 								if (InteractionShadowPrimitives)
 								{
 									const int32 NumPrims = InteractionShadowPrimitives->Num();
 									for (int32 Idx = 0; Idx < NumPrims; ++Idx)
 									{
-										SetupInteractionShadows(RHICmdList, (*InteractionShadowPrimitives)[Idx], VisibleLightInfo, bStaticSceneOnly, ViewDependentWholeSceneShadows, PreShadows);
+										SetupInteractionShadows((*InteractionShadowPrimitives)[Idx], VisibleLightInfo, bStaticSceneOnly, ViewDependentWholeSceneShadows, PreShadows);
 									}
 								}
 								else
 								{
 									// Look for individual primitives with a dynamic shadow.
-								for (FLightPrimitiveInteraction* Interaction = LightSceneInfo->GetDynamicInteractionOftenMovingPrimitiveList();
+									for (FLightPrimitiveInteraction* Interaction = LightSceneInfo->GetDynamicInteractionOftenMovingPrimitiveList();
 										Interaction;
 										Interaction = Interaction->GetNextPrimitive()
 										)
 									{
-										SetupInteractionShadows(RHICmdList, Interaction, VisibleLightInfo, bStaticSceneOnly, ViewDependentWholeSceneShadows, PreShadows);
+										SetupInteractionShadows(Interaction, VisibleLightInfo, bStaticSceneOnly, ViewDependentWholeSceneShadows, PreShadows);
 									}
 
-								for (FLightPrimitiveInteraction* Interaction = LightSceneInfo->GetDynamicInteractionStaticPrimitiveList();
+									for (FLightPrimitiveInteraction* Interaction = LightSceneInfo->GetDynamicInteractionStaticPrimitiveList();
 										Interaction;
 										Interaction = Interaction->GetNextPrimitive()
 										)
 									{
-										SetupInteractionShadows(RHICmdList, Interaction, VisibleLightInfo, bStaticSceneOnly, ViewDependentWholeSceneShadows, PreShadows);
+										SetupInteractionShadows(Interaction, VisibleLightInfo, bStaticSceneOnly, ViewDependentWholeSceneShadows, PreShadows);
 									}
 								}
 							}
@@ -5548,7 +5535,7 @@ void FSceneRenderer::InitDynamicShadows(FRHICommandListImmediate& RHICmdList, FG
 		CSV_CUSTOM_STAT_GLOBAL(ShadowCacheUsageMB, (float(Scene->GetCachedWholeSceneShadowMapsSize()) / 1024) / 1024, ECsvCustomStatOp::Set);
 
 		// Calculate visibility of the projected shadows.
-		InitProjectedShadowVisibility(RHICmdList);
+		InitProjectedShadowVisibility();
 	}
 
 	//SetupFallbackShadowMaps();
@@ -5620,7 +5607,7 @@ void FSceneRenderer::AllocateMobileCSMAndSpotLightShadowDepthTargets(FRHICommand
 				if (ProjectedShadowInfo->bAllocated)
 				{
 					ProjectedShadowInfo->RenderTargets.CopyReferencesFromRenderTargets(ShadowMapAtlas.RenderTargets);
-					ProjectedShadowInfo->SetupShadowDepthView(RHICmdList, this);
+					ProjectedShadowInfo->SetupShadowDepthView(this);
 					ShadowMapAtlas.Shadows.Add(ProjectedShadowInfo);
 
 					if (bResolutionChanged)
