1. Use G8 format for mobile GTAO on Vulkan.
2. Enable mobile AO and PPR on highend mobile devices

#8178

diff --git a/Engine/Config/Android/AndroidScalability.ini b/Engine/Config/Android/AndroidScalability.ini
--- a/Engine/Config/Android/AndroidScalability.ini
+++ b/Engine/Config/Android/AndroidScalability.ini
@@ -98,12 +98,16 @@ r.Tonemapper.Quality=3
 r.BloomQuality=1
 r.DepthOfFieldQuality=1
 r.LightShaftQuality=1
+r.Mobile.AmbientOcclusionQuality=1
+r.Mobile.PixelProjectedReflectionQuality=1
 
 [PostProcessQuality@3]
 r.Tonemapper.Quality=3
 r.BloomQuality=1
 r.DepthOfFieldQuality=1
 r.LightShaftQuality=1
+r.Mobile.AmbientOcclusionQuality=1
+r.Mobile.PixelProjectedReflectionQuality=1
 
 ; ------------------------------------------------------------------------------
 ; Effects settings
diff --git a/Engine/Config/BaseDeviceProfiles.ini b/Engine/Config/BaseDeviceProfiles.ini
--- a/Engine/Config/BaseDeviceProfiles.ini
+++ b/Engine/Config/BaseDeviceProfiles.ini
@@ -405,12 +405,16 @@ DeviceType=IOS
 BaseProfileName=iPhone7
 +CVars=ios.PhysicalScreenDensity=326
 +CVars=sg.ShadowQuality=3
++CVars=r.Mobile.AmbientOcclusionQuality=1
++CVars=r.Mobile.PixelProjectedReflectionQuality=1
 
 [iPhone8Plus DeviceProfile]
 DeviceType=IOS
 BaseProfileName=iPhone7Plus
 +CVars=ios.PhysicalScreenDensity=401
 +CVars=sg.ShadowQuality=3
++CVars=r.Mobile.AmbientOcclusionQuality=1
++CVars=r.Mobile.PixelProjectedReflectionQuality=1
 
 [iPhoneX DeviceProfile]
 DeviceType=IOS
diff --git a/Engine/Config/IOS/IOSScalability.ini b/Engine/Config/IOS/IOSScalability.ini
--- a/Engine/Config/IOS/IOSScalability.ini
+++ b/Engine/Config/IOS/IOSScalability.ini
@@ -98,12 +98,16 @@ r.Tonemapper.Quality=3
 r.BloomQuality=1
 r.DepthOfFieldQuality=0
 r.LightShaftQuality=0
+r.Mobile.AmbientOcclusionQuality=0
+r.Mobile.PixelProjectedReflectionQuality=0
 
 [PostProcessQuality@3]
 r.Tonemapper.Quality=3
 r.BloomQuality=1
 r.DepthOfFieldQuality=1
 r.LightShaftQuality=1
+r.Mobile.AmbientOcclusionQuality=1
+r.Mobile.PixelProjectedReflectionQuality=1
 
 ; ------------------------------------------------------------------------------
 ; Effects settings
diff --git a/Engine/Source/Runtime/Renderer/Private/PostProcess/PostProcessAmbientOcclusionMobile.cpp b/Engine/Source/Runtime/Renderer/Private/PostProcess/PostProcessAmbientOcclusionMobile.cpp
--- a/Engine/Source/Runtime/Renderer/Private/PostProcess/PostProcessAmbientOcclusionMobile.cpp
+++ b/Engine/Source/Runtime/Renderer/Private/PostProcess/PostProcessAmbientOcclusionMobile.cpp
@@ -428,15 +428,26 @@ public:
 
 IMPLEMENT_GLOBAL_SHADER(FGTAOMobile_SpatialFilterPS, "/Engine/Private/PostProcessAmbientOcclusionMobile.usf", "GTAOSpatialFilterPS", SF_Pixel);
 
-FRDGTextureRef CreateMobileScreenSpaceAOTexture(FRDGBuilder& GraphBuilder, FIntPoint SceneTextureExtent)
+FRDGTextureRef CreateMobileScreenSpaceAOTexture(FRDGBuilder& GraphBuilder, const FSceneTexturesConfig& Config)
 {
 	bool bGTAO = (CVarMobileAmbientOcclusionTechnique.GetValueOnRenderThread() == 0);
 	const uint32 DownsampleFactor = bGTAO ? 2 : 1;
 
-	const FIntPoint Extent = FIntPoint::DivideAndRoundUp(SceneTextureExtent, DownsampleFactor);
-	const bool bUsePixelShader = CVarMobileAmbientOcclusionShaderType.GetValueOnRenderThread() == 2;
+	const FIntPoint Extent = FIntPoint::DivideAndRoundUp(Config.Extent, DownsampleFactor);
+
+	EPixelFormat Format = PF_G8;
+
+	// G8 isn't supported as UAV on Android OpenGLES, fall back to RGBA8 for compute shader usage.
+	if (bGTAO
+		&& IsOpenGLPlatform(Config.ShaderPlatform)
+		&& ((GetMaxWorkGroupInvocations() >= 1024 && CVarMobileAmbientOcclusionShaderType.GetValueOnRenderThread() == 0) || CVarMobileAmbientOcclusionShaderType.GetValueOnRenderThread() == 1)
+		)
+	{
+		Format = PF_R8G8B8A8;
+	}
+
 	return GraphBuilder.CreateTexture(
-		FRDGTextureDesc::Create2D(Extent, bUsePixelShader ? PF_G8 : PF_R8G8B8A8, FClearValueBinding::Black, TexCreate_ShaderResource | TexCreate_RenderTargetable | TexCreate_UAV),
+		FRDGTextureDesc::Create2D(Extent, Format, FClearValueBinding::Black, TexCreate_ShaderResource | TexCreate_RenderTargetable | TexCreate_UAV),
 		TEXT("ScreenSpaceAO"));
 }
 
diff --git a/Engine/Source/Runtime/Renderer/Private/PostProcess/PostProcessAmbientOcclusionMobile.h b/Engine/Source/Runtime/Renderer/Private/PostProcess/PostProcessAmbientOcclusionMobile.h
--- a/Engine/Source/Runtime/Renderer/Private/PostProcess/PostProcessAmbientOcclusionMobile.h
+++ b/Engine/Source/Runtime/Renderer/Private/PostProcess/PostProcessAmbientOcclusionMobile.h
@@ -8,7 +8,7 @@
 #include "CoreMinimal.h"
 #include "RendererInterface.h"
 
-FRDGTextureRef CreateMobileScreenSpaceAOTexture(FRDGBuilder& GraphBuilder, FIntPoint Extent);
+FRDGTextureRef CreateMobileScreenSpaceAOTexture(FRDGBuilder& GraphBuilder, const struct FSceneTexturesConfig& Config);
 
 bool IsMobileAmbientOcclusionEnabled(EShaderPlatform ShaderPlatform);
 bool IsUsingMobileAmbientOcclusion(EShaderPlatform ShaderPlatform);
\ No newline at end of file
diff --git a/Engine/Source/Runtime/Renderer/Private/SceneTextures.cpp b/Engine/Source/Runtime/Renderer/Private/SceneTextures.cpp
--- a/Engine/Source/Runtime/Renderer/Private/SceneTextures.cpp
+++ b/Engine/Source/Runtime/Renderer/Private/SceneTextures.cpp
@@ -719,7 +719,7 @@ FSceneTextures& FSceneTextures::Create(FRDGBuilder& GraphBuilder, const FSceneTe
 	else
 	{
 		// Mobile Screen Space Ambient Occlusion
-		SceneTextures.ScreenSpaceAO = CreateMobileScreenSpaceAOTexture(GraphBuilder, Config.Extent);
+		SceneTextures.ScreenSpaceAO = CreateMobileScreenSpaceAOTexture(GraphBuilder, Config);
 
 		if (Config.MobilePixelProjectedReflectionExtent != FIntPoint::ZeroValue)
 		{
