Remove unused skylight shaders

diff --git a/Engine/Shaders/Private/PathTracing/BuildSkyLightColumnCdfComputeShader.usf b/Engine/Shaders/Private/PathTracing/BuildSkyLightColumnCdfComputeShader.usf
deleted file mode 100644
diff --git a/Engine/Shaders/Private/PathTracing/BuildSkyLightCubeFaceCdfComputeShader.usf b/Engine/Shaders/Private/PathTracing/BuildSkyLightCubeFaceCdfComputeShader.usf
deleted file mode 100644
diff --git a/Engine/Shaders/Private/PathTracing/BuildSkyLightRowCdfComputeShader.usf b/Engine/Shaders/Private/PathTracing/BuildSkyLightRowCdfComputeShader.usf
deleted file mode 100644
diff --git a/Engine/Shaders/Private/PathTracing/VisualizeMipTreePixelShader.usf b/Engine/Shaders/Private/PathTracing/VisualizeMipTreePixelShader.usf
deleted file mode 100644
diff --git a/Engine/Shaders/Private/PathTracing/VisualizeSkyLightCdfPixelShader.usf b/Engine/Shaders/Private/PathTracing/VisualizeSkyLightCdfPixelShader.usf
deleted file mode 100644
diff --git a/Engine/Source/Runtime/Renderer/Private/RayTracing/RayTracingSkyLight.h b/Engine/Source/Runtime/Renderer/Private/RayTracing/RayTracingSkyLight.h
--- a/Engine/Source/Runtime/Renderer/Private/RayTracing/RayTracingSkyLight.h
+++ b/Engine/Source/Runtime/Renderer/Private/RayTracing/RayTracingSkyLight.h
@@ -42,6 +42,4 @@ bool SetupSkyLightParameters(
 
 void SetupSkyLightVisibilityRaysParameters(FRDGBuilder& GraphBuilder, const FViewInfo& View, FSkyLightVisibilityRaysData* OutSkyLightVisibilityRaysData);
 
-extern RENDERER_API void BuildSkyLightCdfs(FRHICommandListImmediate& RHICmdList, class FSkyLightSceneProxy* SkyLight);
-
 #endif // RHI_RAYTRACING
