Fixes static analysis in tonemapper

diff --git a/Engine/Source/Runtime/Renderer/Private/PostProcess/PostProcessTonemap.cpp b/Engine/Source/Runtime/Renderer/Private/PostProcess/PostProcessTonemap.cpp
--- a/Engine/Source/Runtime/Renderer/Private/PostProcess/PostProcessTonemap.cpp
+++ b/Engine/Source/Runtime/Renderer/Private/PostProcess/PostProcessTonemap.cpp
@@ -558,7 +558,7 @@ FScreenPassTexture AddTonemapPass(FRDGBuilder& GraphBuilder, const FViewInfo& Vi
 
 	{
 		uint8 FrameIndexMod8 = 0;
-		if (View.State)
+		if (View.ViewState)
 		{
 			FrameIndexMod8 = View.ViewState->GetFrameIndex(8);
 		}
