Implements r.TSR.Velocity.WeightClampingPixelSpeed

diff --git a/Engine/Shaders/Private/TemporalSuperResolution/TSRUpdateHistory.usf b/Engine/Shaders/Private/TemporalSuperResolution/TSRUpdateHistory.usf
--- a/Engine/Shaders/Private/TemporalSuperResolution/TSRUpdateHistory.usf
+++ b/Engine/Shaders/Private/TemporalSuperResolution/TSRUpdateHistory.usf
@@ -63,6 +63,7 @@ float4 HistoryPixelPosToScreenPos;
 float4 HistoryPixelPosToPPCo;
 float3 HistoryQuantizationError;
 float MinTranslucencyRejection;
+float InvWeightClampingPixelSpeed;
 uint ResponsiveStencilMask;
 
 Texture2D<taa_half3> PrevHistory_LowFrequency;
@@ -716,7 +717,7 @@ void MainCS(
 		taa_half2 RejectionConfidentEnough = taa_half(1); // saturate(RejectionValidity * MAX_SAMPLE_COUNT - 3.0);
 
 		// Clamp the validity due to motion to maintain better sharpness.
-		taa_half2 ClampedPrevHistoryValidity = min(PrevHistoryValidity, taa_half(1.0) - (taa_half(1.0) - PredictionOnlyValidity) * saturate(OutputPixelVelocity * rcp(taa_half(40.0))));
+		taa_half2 ClampedPrevHistoryValidity = min(PrevHistoryValidity, taa_half(1.0) - (taa_half(1.0) - PredictionOnlyValidity) * saturate(OutputPixelVelocity * taa_half(InvWeightClampingPixelSpeed)));
 
 		// Compute the newly rejected validity
 		taa_half2 RejectedValidity = (
diff --git a/Engine/Source/Runtime/Renderer/Private/PostProcess/TemporalSuperResolution.cpp b/Engine/Source/Runtime/Renderer/Private/PostProcess/TemporalSuperResolution.cpp
--- a/Engine/Source/Runtime/Renderer/Private/PostProcess/TemporalSuperResolution.cpp
+++ b/Engine/Source/Runtime/Renderer/Private/PostProcess/TemporalSuperResolution.cpp
@@ -65,6 +65,12 @@ TAutoConsoleVariable<int32> CVarTSREnableResponiveAA(
 	TEXT("Whether the responsive AA should keep history fully clamped."),
 	ECVF_RenderThreadSafe);
 
+TAutoConsoleVariable<float> CVarTSRWeightClampingPixelSpeed(
+	TEXT("r.TSR.Velocity.WeightClampingPixelSpeed"), 1.0f,
+	TEXT("Defines the pixel velocity at which the the high frequencies of the history get's their contributing weight clamped. ")
+	TEXT("Smallest reduce blur in movement (Default = 1.0f)."),
+	ECVF_RenderThreadSafe);
+
 #if COMPILE_TSR_DEBUG_PASSES
 
 TAutoConsoleVariable<int32> CVarTSRSetupDebugPasses(
@@ -446,6 +452,7 @@ class FTSRUpdateHistoryCS : public FTSRShader
 		SHADER_PARAMETER(FScreenTransform, HistoryPixelPosToPPCo)
 		SHADER_PARAMETER(FVector3f, HistoryQuantizationError)
 		SHADER_PARAMETER(float, MinTranslucencyRejection)
+		SHADER_PARAMETER(float, InvWeightClampingPixelSpeed)
 		SHADER_PARAMETER(int32, ResponsiveStencilMask)
 
 		SHADER_PARAMETER_STRUCT_INCLUDE(FTSRPrevHistoryParameters, PrevHistoryParameters)
@@ -1291,6 +1298,7 @@ void AddTemporalSuperResolutionPasses(
 		PassParameters->HistoryPixelPosToPPCo = HistoryPixelPosToViewportUV * CommonParameters.InputInfo.ViewportSize + CommonParameters.InputJitter + CommonParameters.InputPixelPosMin;
 		PassParameters->HistoryQuantizationError = ComputePixelFormatQuantizationError(History.LowFrequency->Desc.Format);
 		PassParameters->MinTranslucencyRejection = TranslucencyRejectionTexture == nullptr ? 1.0 : 0.0;
+		PassParameters->InvWeightClampingPixelSpeed = 1.0f / CVarTSRWeightClampingPixelSpeed.GetValueOnRenderThread();
 		PassParameters->ResponsiveStencilMask = CVarTSREnableResponiveAA.GetValueOnRenderThread() ? (STENCIL_TEMPORAL_RESPONSIVE_AA_MASK) : 0;
 
 		PassParameters->PrevHistoryParameters = PrevHistoryParameters;
