Remove unused/no longer working VSM paths

diff --git a/Engine/Shaders/Private/ShadowDepthVertexShader.usf b/Engine/Shaders/Private/ShadowDepthVertexShader.usf
--- a/Engine/Shaders/Private/ShadowDepthVertexShader.usf
+++ b/Engine/Shaders/Private/ShadowDepthVertexShader.usf
@@ -128,7 +128,6 @@ void ScaleBiasClipToPhysicalSmPage(FNaniteView NaniteView, inout float4 Position
 	float2 vUV = vClip * float2(0.5, -0.5) + 0.5 * PositionClip.w;
 	float2 vPixels = vUV * ( uint(VSM_VIRTUAL_MAX_RESOLUTION_XY) >> NaniteView.TargetMipLevel );
 
-	if( PassStruct.bAtomicWrites )
 	{
 		//PointClip.xy = NaniteView.ClipSpaceScaleOffset.xy * PointClip.xy + NaniteView.ClipSpaceScaleOffset.zw * PointClip.w;
 
@@ -136,26 +135,7 @@ void ScaleBiasClipToPhysicalSmPage(FNaniteView NaniteView, inout float4 Position
 		PositionClip.xy *= Scale;
 		PositionClip.xy += ( Scale * float2(1,-1) + float2(-1,1) ) * PositionClip.w;
 	}
-	else
-	{
-		uint2 pPage = ShadowDecodePageTable( PassStruct.VirtualSmPageTable[ CalcPageOffset( NaniteView.TargetLayerIndex, NaniteView.TargetMipLevel, vPage ) ] ).PageIndex;
-
-		float2 pPixels = vPixels + ( (float2)pPage - (float2)vPage ) * VSM_PAGE_SIZE * PositionClip.w;
-		float2 pUV = pPixels * VirtualShadowMap.RecPhysicalPoolSize.xy;
-		float2 pClip = pUV * float2(2, -2) + float2(-1, 1) * PositionClip.w;
 
-		PositionClip.xy = pClip;
-	}
-
-	if( PassStruct.bInstancePerPage )
-	{
-		float2 MinClip =  vPixels - ( vPage + 0 ) * VSM_PAGE_SIZE * PositionClip.w;
-		float2 MaxClip = -vPixels + ( vPage + 1 ) * VSM_PAGE_SIZE * PositionClip.w;
-
-		ClipPlanesInOut.xy = MinClip;
-		ClipPlanesInOut.zw = MaxClip;
-	}
-	else
 	{
 		// Clip to allocated page rect
 		uint4 PageRect = PassStruct.PageRectBounds[ NaniteView.TargetLayerIndex * VSM_MAX_MIP_LEVELS + NaniteView.TargetMipLevel ];
diff --git a/Engine/Shaders/Private/VirtualShadowMaps/BuildPerPageDrawCommands.usf b/Engine/Shaders/Private/VirtualShadowMaps/BuildPerPageDrawCommands.usf
--- a/Engine/Shaders/Private/VirtualShadowMaps/BuildPerPageDrawCommands.usf
+++ b/Engine/Shaders/Private/VirtualShadowMaps/BuildPerPageDrawCommands.usf
@@ -44,7 +44,6 @@ uint FirstPrimaryView;
 uint NumPrimaryViews;
 RWStructuredBuffer<uint> VisibleInstanceCountBufferOut;
 uint InstanceSceneDataSOAStride;
-uint bInstancePerPage;
 uint DynamicInstanceIdOffset;
 uint DynamicInstanceIdMax;
 
@@ -140,33 +139,7 @@ void CullPerPageDrawCommandsCs(uint3 GroupId : SV_GroupID, int GroupThreadIndex
 					RectPages.zw = min(RectPages.zw, AllocatedBounds.zw);
 					if (all(RectPages.zw >= RectPages.xy))
 					{
-						if (bInstancePerPage)
 						{
-							// Count commands for each overlapped page
-							uint PageTableLevelOffset = CalcPageTableLevelOffset(VirtualShadowMapId, MipLevel);
-							for (uint y = RectPages.y; y <= RectPages.w; y += 1U)
-							{
-								for (uint x = RectPages.x; x <= RectPages.z; x += 1U)
-								{
-									uint PageFlagOffset = PageTableLevelOffset + CalcPageOffsetInLevel(MipLevel, uint2(x, y));
-									uint PageFlag = PageFlags[PageFlagOffset];
-
-									if ((PageFlag & PageFlagMask) != 0)
-									{
-										++ThreadTotalForAllViews;
-										WriteCmd(uint2(x,y), MipViewId, InstanceId, IndirectArgIndex);
-
-										if (bHasMoved)
-										{
-											OutDynamicCasterFlags[PageFlagOffset] = 1;
-										}
-									}
-								}
-							}
-						}
-						else
-						{
-							// bInstancePerPage just one per rect
 							++ThreadTotalForAllViews;
 							WriteCmd(RectPages.xy, MipViewId, InstanceId, IndirectArgIndex);
 							if (bHasMoved)
diff --git a/Engine/Shaders/Private/VirtualShadowMaps/PageManagement.usf b/Engine/Shaders/Private/VirtualShadowMaps/PageManagement.usf
--- a/Engine/Shaders/Private/VirtualShadowMaps/PageManagement.usf
+++ b/Engine/Shaders/Private/VirtualShadowMaps/PageManagement.usf
@@ -844,66 +844,6 @@ void UpdatePhysicalPageMetaDataClearState(uint PageIndex : SV_DispatchThreadID)
 }
 
 
-
-Texture2D< float > PrevPhysicalPagePoolHw;
-
-
-uint bRectPrimitive;
-RWBuffer<uint> ClearPhysicalPageDrawArgs;
-
-[numthreads(1, 1, 1)]
-void SetupClearHwPhysicalPagesDrawArgs(uint3 DTID : SV_DispatchThreadID)
-{
-	uint PageCount = NumAllocatedPhysicalPages[0];
-	// Set up indirect draw args for likewise clearing the HW physical page args
-	ClearPhysicalPageDrawArgs[0] = bRectPrimitive > 0 ? 4 : 6;
-	ClearPhysicalPageDrawArgs[1] = PageCount;
-	ClearPhysicalPageDrawArgs[2] = 0;
-	ClearPhysicalPageDrawArgs[3] = 0;
-}
-
-
-void InitializePhysicalPagesVS(
-	in uint InVertexId : SV_VertexID,
-	in uint InInstanceId : SV_InstanceID,
-	nointerpolation out uint PageIndex : TEXCOORD0,
-	out float4 OutPosition : SV_POSITION)
-{
-	PageIndex = InInstanceId;
-	const uint PageX = PageIndex & VirtualShadowMap.PhysicalPageRowMask;
-	const uint PageY = PageIndex >> VirtualShadowMap.PhysicalPageRowShift;
-
-	uint2 TileVertex = uint2(PageX, PageY) * VSM_PAGE_SIZE;
-	TileVertex.x += InVertexId == 1 || InVertexId == 2 || InVertexId == 4 ? VSM_PAGE_SIZE : 0;
-	TileVertex.y += InVertexId == 2 || InVertexId == 4 || InVertexId == 5 ? VSM_PAGE_SIZE : 0;
-
-	OutPosition = float4(float2(TileVertex) * VirtualShadowMap.RecPhysicalPoolSize.xy * float2(2.0f, -2.0f) + float2(-1.0, 1.0f), 0.5f, 1.0f);
-}
-
-void InitializePhysicalPagesPS(uint PageIndex : TEXCOORD0, float4 SVPos : SV_POSITION, out float OutDepth : SV_Depth)
-{
-	const uint PageX = PageIndex & VirtualShadowMap.PhysicalPageRowMask;
-	const uint PageY = PageIndex >> VirtualShadowMap.PhysicalPageRowShift;
-
-	// figure out the relative position within the page
-	uint2 OutPixelPos = uint2(SVPos.xy);
-	// Position within page
-	uint2 PagePixelPos = OutPixelPos - uint2(PageX, PageY) * VSM_PAGE_SIZE;
-	float ClearValue = 0.0f;
-
-#if HAS_CACHE_DATA
-	// Load data from cache if available
-	FCachedPageInfo CachedPageInfo = CachedPageInfos[PageIndex];
-	if (CachedPageInfo.PhysPageAddress.x != VSM_INVALID_PHYSICAL_PAGE_ADDRESS)
-	{
-		const uint2 CachedPos = CachedPageInfo.PhysPageAddress * VSM_PAGE_SIZE + PagePixelPos;
-		ClearValue = max(0.0f, PrevPhysicalPagePoolHw.Load(int3(CachedPos, 0)));
-	}
-#endif // HAS_CACHE_DATA
-
-	OutDepth = ClearValue;
-}
-
 RWStructuredBuffer<FPhysicalPageMetaData> InOutPhysicalPageMetaData;
 StructuredBuffer<uint> VirtualShadowMapFlags;
 
@@ -979,206 +919,3 @@ void InitIndirectArgs1D(
 	IndirectDispatchArgsOut[ 1 ] = 1;
 	IndirectDispatchArgsOut[ 2 ] = 1;
 }
-
-
-
-
-// 1. Calc page rects. Dispacth size {VSM_LEVEL0_DIM_PAGES_XY,VSM_LEVEL0_DIM_PAGES_XY,NumSms}
-[numthreads(VSM_DEFAULT_CS_GROUP_XY, VSM_DEFAULT_CS_GROUP_XY, 1)]
-void CalculatePageRects(uint3 DispatchThreadID : SV_DispatchThreadID)
-{
-	// One group per each mip 0
-	uint ShadowMapID = DispatchThreadID.z;
-	uint2 vPage = DispatchThreadID.xy;
-
-	UNROLL
-	for (uint MipLevel = 0; MipLevel < VSM_MAX_MIP_LEVELS; ++MipLevel)
-	{
-		// Early out for out of bounds 
-		if (any(vPage >= uint2(uint(VSM_LEVEL0_DIM_PAGES_XY) >> MipLevel, uint(VSM_LEVEL0_DIM_PAGES_XY) >> MipLevel)))
-		{
-			return;
-		}
-
-		uint PageOffset = CalcPageOffset(ShadowMapID, MipLevel, vPage);
-		uint Flag = PageFlags[PageOffset];
-		if (Flag)
-		{
-			// Compute the min/max rect of requested pages
-			uint PageBoundIndex = ShadowMapID * VSM_MAX_MIP_LEVELS + MipLevel;
-			InterlockedMin(PageRectBoundsOut[PageBoundIndex].x, vPage.x);
-			InterlockedMin(PageRectBoundsOut[PageBoundIndex].y, vPage.y);
-			InterlockedMax(PageRectBoundsOut[PageBoundIndex].z, vPage.x);
-			InterlockedMax(PageRectBoundsOut[PageBoundIndex].w, vPage.y);
-		}
-		OutPageTable[PageOffset] = VSM_PHYSICAL_PAGE_INVALID;
-
-	}
-}
-
-
-bool TestRectOverlap(uint4 A, uint4 B)
-{
-	return all(A.xy <= B.zw) && all(A.zw >= B.xy);
-}
-
-#define MAX_PAGE_RECTS_BUFFER 1024
-
-// Store compacted range of index+area to skip all empty ones at least
-groupshared uint2 RectsToProcess[MAX_PAGE_RECTS_BUFFER];
-groupshared uint NumRectsToProcessGroupShared;
-
-groupshared uint MaxAreaAndIndex;
-groupshared uint NextFreeOffsetX;
-
-groupshared uint4 AllocatedRectsGroupShared[MAX_PAGE_RECTS_BUFFER];
-groupshared uint NumAllocatedRectsGroupShared;
-
-RWStructuredBuffer<uint4> AllocatedPageRectBoundsOut;
-
-// 2. Sort page rects (just find max)
-// 3. Allocate page rects in phys space. (serial)
-[numthreads(VSM_DEFAULT_CS_GROUP_X, 1, 1)]
-void AllocatePagesUsingRects(uint GroupThreadId : SV_GroupThreadID, uint2 DispatchThreadID : SV_DispatchThreadID)
-{
-	if (GroupThreadId == 0)
-	{
-		NumRectsToProcessGroupShared = 0;
-		MaxAreaAndIndex = 0;
-		NextFreeOffsetX = 0;
-		NumAllocatedRectsGroupShared = 0;
-	}
-	GroupMemoryBarrierWithGroupSync();
-
-	uint TotalMips = VirtualShadowMap.NumShadowMaps * VSM_MAX_MIP_LEVELS;
-	// Load all rects that are non-empty
-	for (uint Index = GroupThreadId; Index < TotalMips; Index += VSM_DEFAULT_CS_GROUP_X)
-	{
-		uint4 Rect = PageRectBounds[Index];
-		AllocatedPageRectBoundsOut[Index] = uint4(VSM_LEVEL0_DIM_PAGES_XY, VSM_LEVEL0_DIM_PAGES_XY, 0, 0);
-		if (all(Rect.zw >= Rect.xy))
-		{
-			uint WriteOffset = 0;
-			uint Area = (Rect.z - Rect.x + 1) * (Rect.w - Rect.y + 1);
-			InterlockedAdd(NumRectsToProcessGroupShared, 1, WriteOffset);
-			if (WriteOffset < MAX_PAGE_RECTS_BUFFER)
-			{
-				RectsToProcess[WriteOffset] = uint2(Index, Area);
-			}
-		}
-	}
-	GroupMemoryBarrierWithGroupSync();
-	// Protect against overflow (but should also detect and somehow fix)
-	uint NumRectsToProcess = min(MAX_PAGE_RECTS_BUFFER, NumRectsToProcessGroupShared);
-	// Loop until all the rects are processed, each iteration finds the max, and then processes that.
-	// Note: this is extremely stupid! If we had a simple to use sorting function for group shared then we could make it less extra stupid.
-	for (uint IterIndex = 0; IterIndex < NumRectsToProcess; ++IterIndex)
-	{
-		// Group stride loop to find max index using atomics
-		for (uint Index = GroupThreadId; Index < NumRectsToProcess; Index += VSM_DEFAULT_CS_GROUP_X)
-		{
-			uint2 Item = RectsToProcess[Index];
-			// Pack area and item index (not source rect index)
-			uint AreaAndIndex = (Item.y << 12U) | Index;
-			InterlockedMax(MaxAreaAndIndex, AreaAndIndex);
-		}
-		GroupMemoryBarrierWithGroupSync();
-
-		// Find item that won the bidding
-		uint ItemIndex = MaxAreaAndIndex & ((1U << 12U) - 1U);
-		// Get info
-		uint2 Item = RectsToProcess[ItemIndex];
-		// Reset stuff and clear the winning rect so it won't win again
-		GroupMemoryBarrierWithGroupSync();
-		if (GroupThreadId == 0)
-		{
-			RectsToProcess[ItemIndex] = uint2(0U, 0U);
-			MaxAreaAndIndex = 0U;
-		}
-		GroupMemoryBarrierWithGroupSync();
-		// Now we have the next largest item, must find allocation
-
-		// Allocation Algo
-		// Scan each possible position across each row, 
-		// test overlap against all existing allocations and skip horizontally
-		// when we run out of row, move to next and repeat
-		uint4 Rect = PageRectBounds[Item.x];
-		uint ShadowMapID = Item.x / VSM_MAX_MIP_LEVELS;
-		uint MipLevel = Item.x % VSM_MAX_MIP_LEVELS;
-		uint2 RectSize = Rect.zw - Rect.xy;
-		// Scan each row where it could fit
-		for (uint RowInd = 0; RowInd < VirtualShadowMap.PhysicalPoolSizePages.y - RectSize.y; ++RowInd)
-		{
-			// Current probe position, test against all rects using group stride loop
-			uint2 TestPos = uint2(0U, RowInd);
-			uint4 ProbeRect = uint4(TestPos.xy, TestPos.xy + RectSize.xy);
-			bool bFoundFree = false;
-			
-			// Until we run off the edge 
-			while (ProbeRect.z < uint(VirtualShadowMap.PhysicalPoolSizePages.x))
-			{
-				NextFreeOffsetX = 0;
-				GroupMemoryBarrierWithGroupSync();
-				for (uint Index = GroupThreadId; Index < NumAllocatedRectsGroupShared; Index += VSM_DEFAULT_CS_GROUP_X)
-				{
-					uint4 AllocatedRect = AllocatedRectsGroupShared[Index];
-					if (TestRectOverlap(AllocatedRect, ProbeRect))
-					{
-						// Find free location to the right of all rects on this row
-						InterlockedMax(NextFreeOffsetX, AllocatedRect.z + 1);
-					}
-				}
-				GroupMemoryBarrierWithGroupSync();
-				// If we did not overlap anything that modified the rect, then we found a free space.
-				if (NextFreeOffsetX == 0U)
-				{
-					bFoundFree = true;
-					break;
-				}
-				else
-				{
-					// Bump probe rect for next iteration
-					ProbeRect.xz = uint2(NextFreeOffsetX, NextFreeOffsetX + RectSize.x);
-				}
-				// need to sync before next iteration also or else we may read the NextFreeOffsetX out of order
-				GroupMemoryBarrierWithGroupSync();
-			}
-
-			if (bFoundFree)
-			{
-				if (GroupThreadId == 0)
-				{
-					AllocatedRectsGroupShared[NumAllocatedRectsGroupShared] = ProbeRect;
-					NumAllocatedRectsGroupShared += 1;
-					AllocatedPageRectBoundsOut[Item.x] = ProbeRect;
-				}
-					// 4. Allocate pages to rects (1:1)
-				for (uint PY = 0U; PY <= RectSize.y; ++PY)
-				{
-					// Group stride loop in X
-					for (uint PX = GroupThreadId; PX <= RectSize.x; PX += VSM_DEFAULT_CS_GROUP_X)
-					{
-						// Virtual page
-						uint2 vPage = Rect.xy + uint2(PX, PY);
-						uint2 pPage = ProbeRect.xy + uint2(PX, PY);
-
-						FCachedPageInfo CachedPageInfo;
-
-						CachedPageInfo.PhysPageAddress = uint2(VSM_INVALID_PHYSICAL_PAGE_ADDRESS, VSM_INVALID_PHYSICAL_PAGE_ADDRESS);
-						// Store data that is used later in the process (for each physical page).
-						OutCachedPageInfos[pPage.y * VirtualShadowMap.PhysicalPoolSizePages.x + pPage.x] = CachedPageInfo;
-
-						uint PageOffset = CalcPageOffset(ShadowMapID, MipLevel, vPage);
-						// Request rendering of pages that are not cached (AKA VSM_INVALID_PHYSICAL_PAGE_ADDRESS in cache).
-						// TODO: make sure caching machinery is use in this path also
-						OutPageFlags[PageOffset] = VSM_INVALID_FLAG | VSM_ALLOCATED_FLAG;
-						OutPageTable[PageOffset] = ShadowEncodePageTable(pPage);
-					}
-				}
-
-				break;
-			}
-		}
-		GroupMemoryBarrierWithGroupSync();
-	}
-}
diff --git a/Engine/Shaders/Private/VirtualShadowMaps/ProjectionCommon.ush b/Engine/Shaders/Private/VirtualShadowMaps/ProjectionCommon.ush
--- a/Engine/Shaders/Private/VirtualShadowMaps/ProjectionCommon.ush
+++ b/Engine/Shaders/Private/VirtualShadowMaps/ProjectionCommon.ush
@@ -53,20 +53,9 @@ FVirtualShadowMapSample InitVirtualShadowMapSample()
 	return Result;
 }
 
-#ifndef TWO_PHYSICAL_TEXTURES
-#define TWO_PHYSICAL_TEXTURES 1
-#endif
-
 float SampleVirtualShadowMapPhysicalDepth(uint2 PhysicalTexelAddress)
 {
-#if ENABLE_NON_NANITE_VSM && TWO_PHYSICAL_TEXTURES
-	return max(
-		asfloat(VirtualShadowMap.PhysicalPagePool.Load(int3(PhysicalTexelAddress, 0))),
-		VirtualShadowMap.PhysicalPagePoolHw.Load(int3(PhysicalTexelAddress, 0))
-	);
-#else
 	return asfloat(VirtualShadowMap.PhysicalPagePool.Load(int3(PhysicalTexelAddress, 0)));
-#endif
 }
 
 FVirtualShadowMapSample SampleVirtualShadowMapLevel(int VirtualShadowMapId, float2 ShadowMapUV, uint MipLevel)
diff --git a/Engine/Source/Runtime/Renderer/Private/ShadowDepthRendering.cpp b/Engine/Source/Runtime/Renderer/Private/ShadowDepthRendering.cpp
--- a/Engine/Source/Runtime/Renderer/Private/ShadowDepthRendering.cpp
+++ b/Engine/Source/Runtime/Renderer/Private/ShadowDepthRendering.cpp
@@ -124,7 +124,6 @@ static FAutoConsoleVariableRef CVarShadowShadowUseGS(
 	TEXT("Use geometry shaders to render cube map shadows."),
 	ECVF_RenderThreadSafe);
 
-extern int32 GVirtualShadowMapAtomicWrites;
 extern int32 GNaniteShowStats;
 extern int32 GEnableNonNaniteVSM;
 
@@ -174,8 +173,6 @@ void SetupShadowDepthPassUniformBuffer(
 	}
 
 	ShadowDepthPassParameters.bRenderToVirtualShadowMap = false;
-	ShadowDepthPassParameters.bInstancePerPage = false;
-	ShadowDepthPassParameters.bAtomicWrites = false;
 	ShadowDepthPassParameters.VirtualSmPageTable = GraphBuilder.CreateSRV(GSystemTextures.GetDefaultStructuredBuffer(GraphBuilder, sizeof(uint32)));
 	ShadowDepthPassParameters.PackedNaniteViews = GraphBuilder.CreateSRV(GSystemTextures.GetDefaultStructuredBuffer(GraphBuilder, sizeof(Nanite::FPackedView)));
 	ShadowDepthPassParameters.PageRectBounds = GraphBuilder.CreateSRV(GSystemTextures.GetDefaultStructuredBuffer(GraphBuilder, sizeof(FIntVector4)));
@@ -594,7 +591,6 @@ bool GetShadowDepthPassShaders(
 	bool bPositionOnlyVS,
 	bool bUsePerspectiveCorrectShadowDepths,
 	bool bVirtualShadowMap,
-	bool bAtomicWrites,
 	TShaderRef<FShadowDepthVS>& VertexShader,
 	TShaderRef<FShadowDepthBasePS>& PixelShader,
 	TShaderRef<FOnePassPointShadowDepthGS>& GeometryShader)
@@ -670,10 +666,10 @@ bool GetShadowDepthPassShaders(
 	}
 
 	// Pixel shaders
-	const bool bNullPixelShader = Material.WritesEveryPixel(true) && !bUsePerspectiveCorrectShadowDepths && !bAtomicWrites && VertexFactory->SupportsNullPixelShader();
+	const bool bNullPixelShader = Material.WritesEveryPixel(true) && !bUsePerspectiveCorrectShadowDepths && !bVirtualShadowMap && VertexFactory->SupportsNullPixelShader();
 	if (!bNullPixelShader)
 	{
-		if (bVirtualShadowMap && bAtomicWrites)
+		if (bVirtualShadowMap)
 		{
 			ShaderTypes.AddShaderType<TShadowDepthPS<PixelShadowDepth_VirtualShadowMap>>();
 		}
@@ -766,7 +762,7 @@ void SetStateForShadowDepth(bool bOnePassPointLightShadow, bool bDirectionalLigh
 	// Disable color writes
 	DrawRenderState.SetBlendState(TStaticBlendState<CW_NONE>::GetRHI());
 
-	if( InMeshPassTargetType == EMeshPass::VSMShadowDepth && GVirtualShadowMapAtomicWrites )
+	if( InMeshPassTargetType == EMeshPass::VSMShadowDepth )
 	{
 		DrawRenderState.SetDepthStencilState(TStaticDepthStencilState<false, CF_Always>::GetRHI());
 	}
@@ -1674,8 +1670,6 @@ void FSceneRenderer::RenderShadowDepthMapAtlases(FRDGBuilder& GraphBuilder)
 	ResourceAccessFinalizer.Finalize(GraphBuilder);
 }
 
-extern TAutoConsoleVariable<int32> CVarAllocatePagesUsingRects;
-
 void FSceneRenderer::RenderShadowDepthMaps(FRDGBuilder& GraphBuilder, FInstanceCullingManager &InstanceCullingManager)
 {
 	CSV_SCOPED_TIMING_STAT_EXCLUSIVE(RenderShadows);
@@ -1696,7 +1690,6 @@ void FSceneRenderer::RenderShadowDepthMaps(FRDGBuilder& GraphBuilder, FInstanceC
 	const bool bHasVSMShadows = SortedShadowsForShadowDepthPass.VirtualShadowMapShadows.Num() > 0;
 	const bool bHasVSMClipMaps = SortedShadowsForShadowDepthPass.VirtualShadowMapClipmaps.Num() > 0;
 	const bool bNaniteEnabled = UseNanite(ShaderPlatform) && ViewFamily.EngineShowFlags.NaniteMeshes;
-	const bool bAllocatePageRectAtlas = CVarAllocatePagesUsingRects.GetValueOnRenderThread() != 0;
 
 	if (bNaniteEnabled && (bHasVSMShadows || bHasVSMClipMaps))
 	{
@@ -1711,25 +1704,17 @@ void FSceneRenderer::RenderShadowDepthMaps(FRDGBuilder& GraphBuilder, FInstanceC
 			const FIntPoint VirtualShadowSize = VirtualShadowMapArray.GetPhysicalPoolSize();
 			const FIntRect VirtualShadowViewRect = FIntRect(0, 0, VirtualShadowSize.X, VirtualShadowSize.Y);
 
-			VirtualShadowMapArray.PhysicalPagePoolRDG = GraphBuilder.CreateTexture( 
-				FRDGTextureDesc::Create2D(VirtualShadowSize, PF_R32_UINT, FClearValueBinding::None, TexCreate_ShaderResource | TexCreate_UAV),
-				TEXT("Shadow.Virtual.PhysicalTexture") );
+			check(VirtualShadowMapArray.PhysicalPagePoolRDG != nullptr);
 
 			Nanite::FRasterContext RasterContext = Nanite::InitRasterContext(
 				GraphBuilder,
 				FeatureLevel,
 				VirtualShadowSize,
 				Nanite::EOutputBufferMode::DepthOnly,
-				bAllocatePageRectAtlas,	// Clear entire texture
+				false,	// Clear entire texture
 				nullptr, 0,
 				VirtualShadowMapArray.PhysicalPagePoolRDG );
 
-			if( !bAllocatePageRectAtlas )
-			{
-				// Clear only mapped pages
-				VirtualShadowMapArray.ClearPhysicalMemory(GraphBuilder, RasterContext.DepthBuffer);
-			}
-
 			const bool bUpdateStreaming = CVarNaniteShadowsUpdateStreaming.GetValueOnRenderThread() != 0;
 
 			auto FilterAndRenderVirtualShadowMaps = [
@@ -1942,8 +1927,6 @@ void FSceneRenderer::RenderShadowDepthMaps(FRDGBuilder& GraphBuilder, FInstanceC
 		VirtualShadowMapArray.RenderVirtualShadowMapsHw(GraphBuilder, SortedShadowsForShadowDepthPass.VirtualShadowMapShadows, *Scene);
 	}
 
-	VirtualShadowMapArray.SetupProjectionParameters(GraphBuilder);
-
 	// Render non-VSM shadows
 	FSceneRenderer::RenderShadowDepthMapAtlases(GraphBuilder);
 
@@ -2185,12 +2168,10 @@ bool FShadowDepthPassMeshProcessor::Process(
 	bool bOnePassPointLightShadow = ShadowDepthType.bOnePassPointLightShadow;
 	
 	bool bVirtualShadowMap = MeshPassTargetType == EMeshPass::VSMShadowDepth;
-	bool bAtomicWrites = false;
 	if (bVirtualShadowMap)
 	{
 		bUsePerspectiveCorrectShadowDepths = false;
 		bOnePassPointLightShadow = false;
-		bAtomicWrites = GVirtualShadowMapAtomicWrites != 0;
 	}
 
 	if (!GetShadowDepthPassShaders(
@@ -2202,7 +2183,6 @@ bool FShadowDepthPassMeshProcessor::Process(
 		bUsePositionOnlyVS,
 		bUsePerspectiveCorrectShadowDepths,
 		bVirtualShadowMap,
-		bAtomicWrites,
 		ShadowDepthPassShaders.VertexShader,
 		ShadowDepthPassShaders.PixelShader,
 		ShadowDepthPassShaders.GeometryShader))
diff --git a/Engine/Source/Runtime/Renderer/Private/ShadowRendering.h b/Engine/Source/Runtime/Renderer/Private/ShadowRendering.h
--- a/Engine/Source/Runtime/Renderer/Private/ShadowRendering.h
+++ b/Engine/Source/Runtime/Renderer/Private/ShadowRendering.h
@@ -198,10 +198,7 @@ BEGIN_GLOBAL_SHADER_PARAMETER_STRUCT(FShadowDepthPassUniformParameters,)
 	SHADER_PARAMETER(float, bClampToNearPlane)
 	SHADER_PARAMETER_ARRAY(FMatrix44f, ShadowViewProjectionMatrices, [6])
 	SHADER_PARAMETER_ARRAY(FMatrix44f, ShadowViewMatrices, [6])
-	// GPUCULL_TODO: ?
 	SHADER_PARAMETER(int, bRenderToVirtualShadowMap)
-	SHADER_PARAMETER(int, bInstancePerPage)
-	SHADER_PARAMETER(int, bAtomicWrites)
 
 	SHADER_PARAMETER_RDG_BUFFER_SRV(StructuredBuffer< uint >, VirtualSmPageTable)
 	SHADER_PARAMETER_RDG_BUFFER_SRV(StructuredBuffer< FPackedNaniteView >, PackedNaniteViews)
diff --git a/Engine/Source/Runtime/Renderer/Private/VirtualShadowMaps/VirtualShadowMapArray.cpp b/Engine/Source/Runtime/Renderer/Private/VirtualShadowMaps/VirtualShadowMapArray.cpp
--- a/Engine/Source/Runtime/Renderer/Private/VirtualShadowMaps/VirtualShadowMapArray.cpp
+++ b/Engine/Source/Runtime/Renderer/Private/VirtualShadowMaps/VirtualShadowMapArray.cpp
@@ -89,13 +89,6 @@ static TAutoConsoleVariable<float> CVarPageDilationBorderSize(
 	ECVF_RenderThreadSafe
 );
 
-TAutoConsoleVariable<int32> CVarAllocatePagesUsingRects(
-	TEXT("r.Shadow.Virtual.AllocatePagesUsingRects"),
-	0,
-	TEXT("If set to 1 then pages are allocated in a contigious block to support the bounding rectangle of allocated virtual pages leading to a trivial mapping but higher memory overhead."),
-	ECVF_RenderThreadSafe
-);
-
 TAutoConsoleVariable<int32> CVarMarkPixelPages(
 	TEXT("r.Shadow.Virtual.MarkPixelPages"),
 	1,
@@ -119,19 +112,6 @@ TAutoConsoleVariable<int32> CVarMarkCoarsePagesLocal(
 	ECVF_RenderThreadSafe
 );
 
-int32 GVirtualShadowMapAtomicWrites = 1;
-FAutoConsoleVariableRef CVarAtomicWrites(
-	TEXT("r.Shadow.Virtual.AtomicWrites"),
-	GVirtualShadowMapAtomicWrites,
-	TEXT("Use per pixel page table routing and atomic writes instead of an instance per page."),
-	FConsoleVariableDelegate::CreateLambda([](IConsoleVariable* InVariable)
-	{
-		// Needed because the depth state changes with method (so cached draw commands must be re-created) see SetStateForShadowDepth
-		FGlobalComponentRecreateRenderStateContext Context;
-	}),
-	ECVF_RenderThreadSafe
-);
-
 static TAutoConsoleVariable<int32> CVarShowClipmapStats(
 	TEXT("r.Shadow.Virtual.ShowClipmapStats"),
 	-1,
@@ -294,7 +274,6 @@ void FVirtualShadowMapArray::Initialize(FRDGBuilder& GraphBuilder, FVirtualShado
 	UniformParameters.ProjectionData = GraphBuilder.CreateSRV(GSystemTextures.GetDefaultStructuredBuffer(GraphBuilder, sizeof(FVirtualShadowMapProjectionShaderData)));
 
 	UniformParameters.PhysicalPagePool = GSystemTextures.GetZeroUIntDummy(GraphBuilder);
-	UniformParameters.PhysicalPagePoolHw = GSystemTextures.GetBlackDummy(GraphBuilder);
 }
 
 FVirtualShadowMapArray::~FVirtualShadowMapArray()
@@ -712,59 +691,6 @@ class FInitPageRectBoundsCS : public FVirtualPageManagementShader
 IMPLEMENT_GLOBAL_SHADER(FInitPageRectBoundsCS, "/Engine/Private/VirtualShadowMaps/PageManagement.usf", "InitPageRectBounds", SF_Compute);
 
 
-
-
-
-class FCalculatePageRectsCS : public FVirtualPageManagementShader
-{
-	DECLARE_GLOBAL_SHADER(FCalculatePageRectsCS);
-	SHADER_USE_PARAMETER_STRUCT(FCalculatePageRectsCS, FVirtualPageManagementShader)
-
-		BEGIN_SHADER_PARAMETER_STRUCT(FParameters, )
-		SHADER_PARAMETER_RDG_UNIFORM_BUFFER(FVirtualShadowMapUniformParameters, VirtualShadowMap)
-
-		SHADER_PARAMETER_RDG_BUFFER_SRV(StructuredBuffer< uint >, PageFlags)
-		SHADER_PARAMETER_RDG_BUFFER_UAV(RWStructuredBuffer<FIntVector4>, PageRectBoundsOut)
-		SHADER_PARAMETER_RDG_BUFFER_UAV(RWStructuredBuffer< uint >, OutPageTable)
-	END_SHADER_PARAMETER_STRUCT()
-};
-IMPLEMENT_GLOBAL_SHADER(FCalculatePageRectsCS, "/Engine/Private/VirtualShadowMaps/PageManagement.usf", "CalculatePageRects", SF_Compute);
-
-
-class FAllocatePagesUsingRectsCS : public FVirtualPageManagementShader
-{
-	DECLARE_GLOBAL_SHADER(FAllocatePagesUsingRectsCS);
-	SHADER_USE_PARAMETER_STRUCT(FAllocatePagesUsingRectsCS, FVirtualPageManagementShader)
-
-	//class FHasCacheDataDim : SHADER_PERMUTATION_BOOL("HAS_CACHE_DATA");
-	//class FGenerateStatsDim : SHADER_PERMUTATION_BOOL("VSM_GENERATE_STATS");
-
-	//using FPermutationDomain = TShaderPermutationDomain<FHasCacheDataDim, FGenerateStatsDim>;
-
-	static void ModifyCompilationEnvironment(const FGlobalShaderPermutationParameters& Parameters, FShaderCompilerEnvironment& OutEnvironment)
-	{
-		FVirtualPageManagementShader::ModifyCompilationEnvironment(Parameters, OutEnvironment);
-
-		// Ensure these are off for now.
-		OutEnvironment.SetDefine(TEXT("HAS_CACHE_DATA"), 0);
-		OutEnvironment.SetDefine(TEXT("VSM_GENERATE_STATS"), 0);
-	}
-
-	BEGIN_SHADER_PARAMETER_STRUCT(FParameters, )
-		SHADER_PARAMETER_RDG_UNIFORM_BUFFER(FVirtualShadowMapUniformParameters, VirtualShadowMap)
-		SHADER_PARAMETER_RDG_BUFFER_UAV(RWStructuredBuffer< uint >, AllocatedPagesOffset)
-		SHADER_PARAMETER_RDG_BUFFER_UAV(RWStructuredBuffer< uint >, OutPageTable)
-		SHADER_PARAMETER_RDG_BUFFER_UAV(RWStructuredBuffer< uint >, OutStatsBuffer)
-		SHADER_PARAMETER_RDG_BUFFER_UAV(RWStructuredBuffer< FCachedPageInfo >, OutCachedPageInfos)
-		SHADER_PARAMETER_RDG_BUFFER_UAV(RWStructuredBuffer< uint >, OutPageFlags)
-		SHADER_PARAMETER_RDG_BUFFER_SRV(StructuredBuffer< FIntVector4 >, PageRectBounds)
-		SHADER_PARAMETER_RDG_BUFFER_UAV(RWStructuredBuffer<FIntVector4>, AllocatedPageRectBoundsOut)
-
-		//SHADER_PARAMETER_STRUCT_INCLUDE(FCacheDataParameters, CacheDataParameters)
-	END_SHADER_PARAMETER_STRUCT()
-};
-IMPLEMENT_GLOBAL_SHADER(FAllocatePagesUsingRectsCS, "/Engine/Private/VirtualShadowMaps/PageManagement.usf", "AllocatePagesUsingRects", SF_Compute);
-
 static FString GetLightNameForDebug(const FLightSceneProxy* Proxy)
 {
 	// TODO: ActorLabel is editor-only... need a better solution that goes right into Proxy->GetComponentName()
@@ -1123,51 +1049,6 @@ void FVirtualShadowMapArray::BuildPageAllocations(
 			);
 		}
 		
-		const bool bAllocatePageRectAtlas = CVarAllocatePagesUsingRects.GetValueOnRenderThread() != 0;
-		if (bAllocatePageRectAtlas)
-		{
-			// 1. Build page rects from requested pages
-			{
-				FCalculatePageRectsCS::FParameters* PassParameters = GraphBuilder.AllocParameters< FCalculatePageRectsCS::FParameters >();
-				PassParameters->VirtualShadowMap = GetUniformBuffer(GraphBuilder);
-				PassParameters->PageFlags = GraphBuilder.CreateSRV(PageRequestFlagsRDG);
-				PassParameters->PageRectBoundsOut = GraphBuilder.CreateUAV(PageRectBoundsRDG);
-				PassParameters->OutPageTable = GraphBuilder.CreateUAV(PageTableRDG);
-				AddClearUAVPass(GraphBuilder, GraphBuilder.CreateUAV(PageFlagsRDG), 0);
-
-				auto ComputeShader = Views[0].ShaderMap->GetShader< FCalculatePageRectsCS >();
-				FComputeShaderUtils::AddPass(
-					GraphBuilder,
-					RDG_EVENT_NAME("CalculatePageRequestRects"),
-					ComputeShader,
-					PassParameters,
-					FComputeShaderUtils::GetGroupCount(FIntVector(FVirtualShadowMap::Level0DimPagesXY, FVirtualShadowMap::Level0DimPagesXY, ShadowMaps.Num()), FIntVector(FCalculatePageRectsCS::DefaultCSGroupXY, FCalculatePageRectsCS::DefaultCSGroupXY, 1U))
-				);
-			}
-
-
-			// 2. Allocate and all that
-			{
-				FAllocatePagesUsingRectsCS::FParameters* PassParameters = GraphBuilder.AllocParameters< FAllocatePagesUsingRectsCS::FParameters >();
-				PassParameters->VirtualShadowMap = GetUniformBuffer(GraphBuilder);
-				PassParameters->PageRectBounds = GraphBuilder.CreateSRV(PageRectBoundsRDG);
-
-				PassParameters->OutPageTable = GraphBuilder.CreateUAV(PageTableRDG);
-				PassParameters->OutCachedPageInfos = GraphBuilder.CreateUAV(CachedPageInfosRDG);
-				PassParameters->OutPageFlags = GraphBuilder.CreateUAV(PageFlagsRDG);
-				PassParameters->AllocatedPageRectBoundsOut = GraphBuilder.CreateUAV(AllocatedPageRectBoundsRDG);
-
-				auto ComputeShader = Views[0].ShaderMap->GetShader< FAllocatePagesUsingRectsCS >();
-				FComputeShaderUtils::AddPass(
-					GraphBuilder,
-					RDG_EVENT_NAME("AllocatePagesUsingRects"),
-					ComputeShader,
-					PassParameters,
-					FIntVector(1U, 1U, 1U)
-				);
-			}
-		}
-		else
 		{
 			// Note: does not actually need mip0 so can be trimmed down a bit
 			FRDGBufferRef CoverageSummary = GraphBuilder.CreateBuffer( FRDGBufferDesc::CreateStructuredDesc( 4, NumPageFlags ), TEXT("Shadow.Virtual.CoverageSummary") );
@@ -1245,6 +1126,13 @@ void FVirtualShadowMapArray::BuildPageAllocations(
 			);
 		}
 
+		// Allocate and clear physical page pool
+		PhysicalPagePoolRDG = GraphBuilder.CreateTexture(
+			FRDGTextureDesc::Create2D(GetPhysicalPoolSize(), PF_R32_UINT, FClearValueBinding::None, TexCreate_ShaderResource | TexCreate_UAV),
+			TEXT("Shadow.Virtual.PhysicalPagePool") );
+		ClearPhysicalMemory(GraphBuilder, PhysicalPagePoolRDG);
+
+		UniformParameters.PhysicalPagePool = PhysicalPagePoolRDG;
 		UniformParameters.PageTable = GraphBuilder.CreateSRV(PageTableRDG);
 	}
 
@@ -1254,12 +1142,6 @@ void FVirtualShadowMapArray::BuildPageAllocations(
 #endif
 }
 
-void FVirtualShadowMapArray::SetupProjectionParameters(FRDGBuilder& GraphBuilder)
-{
-	UniformParameters.PhysicalPagePool = PhysicalPagePoolRDG ? PhysicalPagePoolRDG : GSystemTextures.GetZeroUIntDummy(GraphBuilder);
-	UniformParameters.PhysicalPagePoolHw = PhysicalPagePoolHw ? PhysicalPagePoolHw : GSystemTextures.GetBlackDummy(GraphBuilder);
-}
-
 class FDebugVisualizeVirtualSmCS : public FVirtualPageManagementShader
 {
 	DECLARE_GLOBAL_SHADER(FDebugVisualizeVirtualSmCS);
@@ -1573,7 +1455,6 @@ public:
 		SHADER_PARAMETER_RDG_BUFFER_UAV(RWBuffer<uint>, DrawIndirectArgsBufferOut)
 		SHADER_PARAMETER_RDG_BUFFER_UAV(RWStructuredBuffer<uint>, VisibleInstanceCountBufferOut)
 		SHADER_PARAMETER_RDG_BUFFER_UAV(RWStructuredBuffer< uint >, OutDynamicCasterFlags)
-		SHADER_PARAMETER(int32, bInstancePerPage)
 	END_SHADER_PARAMETER_STRUCT()
 };
 IMPLEMENT_GLOBAL_SHADER(FCullPerPageDrawCommandsCs, "/Engine/Private/VirtualShadowMaps/BuildPerPageDrawCommands.usf", "CullPerPageDrawCommandsCs", SF_Compute);
@@ -1665,27 +1546,9 @@ void FVirtualShadowMapArray::RenderVirtualShadowMapsHw(FRDGBuilder& GraphBuilder
 
 	RDG_EVENT_SCOPE(GraphBuilder, "RenderVirtualShadowMaps(Non-Nanite)");
 
-	const bool bAllocatePageRectAtlas = CVarAllocatePagesUsingRects.GetValueOnRenderThread() != 0;
-	const bool bAtomicWrites = GVirtualShadowMapAtomicWrites != 0;
-
 	FGPUScene& GPUScene = Scene.GPUScene;
 	// Create that pass struct and whatever I guess.
 
-	// track whether the page pool was cleared
-	bool bWasCleared = bAtomicWrites;
-
-	if( !bAtomicWrites )
-	{
-		// Create physical page pool
-		PhysicalPagePoolHw = GraphBuilder.CreateTexture(FRDGTextureDesc::Create2D(GetPhysicalPoolSize(), PF_DepthStencil, FClearValueBinding::DepthFar, TexCreate_DepthStencilTargetable | TexCreate_ShaderResource), TEXT("Shadow.Virtual.PhysicalPagePoolHw"));
-
-		if( !bAllocatePageRectAtlas )
-		{
-			AddInitializePhysicalPagesHwPass(GraphBuilder);
-			bWasCleared = true;
-		}
-	}
-
 	TArray<uint32, SceneRenderingAllocator> VirtualShadowMapFlags;
 	VirtualShadowMapFlags.AddZeroed(ShadowMaps.Num());
 
@@ -1833,7 +1696,6 @@ void FVirtualShadowMapArray::RenderVirtualShadowMapsHw(FRDGBuilder& GraphBuilder
 				PassParameters->VisibleInstancesOut = GraphBuilder.CreateUAV(VisibleInstancesRdg);
 				PassParameters->VisibleInstanceCountBufferOut = GraphBuilder.CreateUAV(VisibleInstanceWriteOffsetRDG);
 				PassParameters->OutDynamicCasterFlags = GraphBuilder.CreateUAV(DynamicCasterPageFlagsRDG);
-				PassParameters->bInstancePerPage = !bAllocatePageRectAtlas && !bAtomicWrites;
 
 				FCullPerPageDrawCommandsCs::FPermutationDomain PermutationVector;
 				PermutationVector.Set< FCullPerPageDrawCommandsCs::FNearClipDim >( !Clipmap.IsValid() );
@@ -1899,22 +1761,11 @@ void FVirtualShadowMapArray::RenderVirtualShadowMapsHw(FRDGBuilder& GraphBuilder
 				);
 			}
 
-			if( PhysicalPagePoolRDG == nullptr )
-			{
-				PhysicalPagePoolRDG = GraphBuilder.CreateTexture( FRDGTextureDesc::Create2D( GetPhysicalPoolSize(), PF_R32_UINT, FClearValueBinding::None, TexCreate_ShaderResource | TexCreate_UAV ), TEXT("Shadow.Virtual.PhysicalPagePool") );
-
-				ClearPhysicalMemory(GraphBuilder, PhysicalPagePoolRDG);
-			}
+			check(PhysicalPagePoolRDG != nullptr);
 
 			FVirtualShadowDepthPassParameters* PassParameters = GraphBuilder.AllocParameters<FVirtualShadowDepthPassParameters>();
 			PassParameters->View = ShadowDepthView->ViewUniformBuffer;
 
-			if( !bAtomicWrites )
-			{
-				PassParameters->RenderTargets.DepthStencil = FDepthStencilBinding(PhysicalPagePoolHw, bWasCleared ? ERenderTargetLoadAction::ELoad : ERenderTargetLoadAction::EClear, FExclusiveDepthStencil::DepthWrite_StencilNop);
-				bWasCleared = true;
-			}
-
 			FShadowDepthPassUniformParameters* ShadowDepthPassParameters = GraphBuilder.AllocParameters<FShadowDepthPassUniformParameters>();
 
 			SetupSceneTextureUniformParameters(GraphBuilder, GMaxRHIFeatureLevel, ESceneTextureSetupMode::None, ShadowDepthPassParameters->SceneTextures);
@@ -1926,8 +1777,6 @@ void FVirtualShadowMapArray::RenderVirtualShadowMapsHw(FRDGBuilder& GraphBuilder
 			ShadowDepthPassParameters->ViewMatrix = FMatrix::Identity;
 			ShadowDepthPassParameters->ShadowParams = FVector4(0.0f, 0.0f, 0.0f, 1.0f);
 			ShadowDepthPassParameters->bRenderToVirtualShadowMap = true;
-			ShadowDepthPassParameters->bInstancePerPage = !bAllocatePageRectAtlas && !bAtomicWrites;
-			ShadowDepthPassParameters->bAtomicWrites = bAtomicWrites;
 			
 			ShadowDepthPassParameters->VirtualSmPageTable	= GraphBuilder.CreateSRV( PageTableRDG );
 			ShadowDepthPassParameters->PackedNaniteViews	= GraphBuilder.CreateSRV( VirtualShadowViewsRDG );
@@ -1951,7 +1800,7 @@ void FVirtualShadowMapArray::RenderVirtualShadowMapsHw(FRDGBuilder& GraphBuilder
 			InstanceCullingResult.GetDrawParameters(PassParameters->InstanceCullingDrawParams);
 
 			FIntRect ViewRect;
-			ViewRect.Max = bAtomicWrites ? FVirtualShadowMap::VirtualMaxResolutionXY : GetPhysicalPoolSize();
+			ViewRect.Max = FVirtualShadowMap::VirtualMaxResolutionXY;
 
 			GraphBuilder.AddPass(
 				RDG_EVENT_NAME("RenderVirtualShadowMapsHw"),
@@ -2000,176 +1849,4 @@ void FVirtualShadowMapArray::RenderVirtualShadowMapsHw(FRDGBuilder& GraphBuilder
 	}
 
 	MarkPhysicalPagesRendered(GraphBuilder, VirtualShadowMapFlags);
-
-	if (!bWasCleared)
-	{
-		AddClearDepthStencilPass(GraphBuilder, PhysicalPagePoolHw);
-	}
-}
-
-class FUpdatePhysicalPageMetaDataClearStateCS : public FVirtualPageManagementShader
-{
-	DECLARE_GLOBAL_SHADER(FUpdatePhysicalPageMetaDataClearStateCS);
-	SHADER_USE_PARAMETER_STRUCT(FUpdatePhysicalPageMetaDataClearStateCS, FVirtualPageManagementShader)
-
-	class FHasCacheDataDim : SHADER_PERMUTATION_BOOL("HAS_CACHE_DATA");
-	using FPermutationDomain = TShaderPermutationDomain<FHasCacheDataDim>;
-
-	BEGIN_SHADER_PARAMETER_STRUCT(FParameters, )
-		SHADER_PARAMETER_RDG_UNIFORM_BUFFER(FVirtualShadowMapUniformParameters, VirtualShadowMap)
-		SHADER_PARAMETER_RDG_BUFFER_SRV(StructuredBuffer< FCachedPageInfo >, CachedPageInfos)
-		SHADER_PARAMETER_RDG_BUFFER_SRV(StructuredBuffer< FPhysicalPageMetaData >, PrevPhysicalPageMetaData)
-		SHADER_PARAMETER_RDG_BUFFER_SRV(StructuredBuffer< uint >, NumAllocatedPhysicalPages)
-		SHADER_PARAMETER_RDG_BUFFER_UAV(RWStructuredBuffer< FPhysicalPageMetaData >, PhysicalPageMetaDataOut)
-	END_SHADER_PARAMETER_STRUCT()
-};
-IMPLEMENT_GLOBAL_SHADER(FUpdatePhysicalPageMetaDataClearStateCS, "/Engine/Private/VirtualShadowMaps/PageManagement.usf", "UpdatePhysicalPageMetaDataClearState", SF_Compute);
-
-
-class FSetupClearHwPhysicalPagesDrawArgsCS : public FVirtualPageManagementShader
-{
-	DECLARE_GLOBAL_SHADER(FSetupClearHwPhysicalPagesDrawArgsCS);
-	SHADER_USE_PARAMETER_STRUCT(FSetupClearHwPhysicalPagesDrawArgsCS, FVirtualPageManagementShader)
-
-	BEGIN_SHADER_PARAMETER_STRUCT(FParameters, )
-		SHADER_PARAMETER_RDG_BUFFER_SRV(StructuredBuffer< uint >, NumAllocatedPhysicalPages)
-		SHADER_PARAMETER_RDG_BUFFER_UAV(RWBuffer< uint >, ClearPhysicalPageDrawArgs)
-		SHADER_PARAMETER(int32, bRectPrimitive)
-	END_SHADER_PARAMETER_STRUCT()
-};
-IMPLEMENT_GLOBAL_SHADER(FSetupClearHwPhysicalPagesDrawArgsCS, "/Engine/Private/VirtualShadowMaps/PageManagement.usf", "SetupClearHwPhysicalPagesDrawArgs", SF_Compute);
-
-
-class FInitializePhysicalPagesVS : public FVirtualPageManagementShader
-{
-	DECLARE_GLOBAL_SHADER(FInitializePhysicalPagesVS);
-	SHADER_USE_PARAMETER_STRUCT(FInitializePhysicalPagesVS, FVirtualPageManagementShader);
-
-	BEGIN_SHADER_PARAMETER_STRUCT(FParameters, )
-		SHADER_PARAMETER_RDG_UNIFORM_BUFFER(FVirtualShadowMapUniformParameters, VirtualShadowMap)
-	END_SHADER_PARAMETER_STRUCT()
-};
-
-class FInitializePhysicalPagesPS : public FVirtualPageManagementShader
-{
-	DECLARE_GLOBAL_SHADER(FInitializePhysicalPagesPS);
-	SHADER_USE_PARAMETER_STRUCT(FInitializePhysicalPagesPS, FVirtualPageManagementShader);
-	
-	class FHasCacheDataDim : SHADER_PERMUTATION_BOOL("HAS_CACHE_DATA");
-	using FPermutationDomain = TShaderPermutationDomain<FHasCacheDataDim>;
-
-	BEGIN_SHADER_PARAMETER_STRUCT(FParameters, )
-		SHADER_PARAMETER_RDG_UNIFORM_BUFFER(FVirtualShadowMapUniformParameters, VirtualShadowMap)
-		SHADER_PARAMETER_RDG_TEXTURE(Texture2D< float >, PrevPhysicalPagePoolHw)
-		SHADER_PARAMETER_RDG_BUFFER_SRV(StructuredBuffer< FCachedPageInfo >, CachedPageInfos)
-		RDG_BUFFER_ACCESS(IndirectArgsBuffer, ERHIAccess::IndirectArgs)
-		SHADER_PARAMETER(int32, bCacheDataAvailable)
-		SHADER_PARAMETER(int32, bRectPrimitive)
-		RENDER_TARGET_BINDING_SLOTS()
-	END_SHADER_PARAMETER_STRUCT()
-};
-
-IMPLEMENT_GLOBAL_SHADER(FInitializePhysicalPagesVS, "/Engine/Private/VirtualShadowMaps/PageManagement.usf", "InitializePhysicalPagesVS", SF_Vertex);
-IMPLEMENT_GLOBAL_SHADER(FInitializePhysicalPagesPS, "/Engine/Private/VirtualShadowMaps/PageManagement.usf", "InitializePhysicalPagesPS", SF_Pixel);
-
-void FVirtualShadowMapArray::AddInitializePhysicalPagesHwPass(FRDGBuilder& GraphBuilder)
-{
-	FRDGBufferRef IndirectArgsBuffer = GraphBuilder.CreateBuffer(FRDGBufferDesc::CreateIndirectDesc<FRHIDrawIndirectParameters>(), TEXT("Shadow.Virtual.InitializePhysicalPagesArgs"));
-	{
-		FSetupClearHwPhysicalPagesDrawArgsCS::FParameters* PassParameters = GraphBuilder.AllocParameters<FSetupClearHwPhysicalPagesDrawArgsCS::FParameters>();
-		PassParameters->NumAllocatedPhysicalPages = GraphBuilder.CreateSRV(AllocatedPagesOffsetRDG);
-		PassParameters->ClearPhysicalPageDrawArgs = GraphBuilder.CreateUAV(IndirectArgsBuffer, PF_R32_UINT);
-		PassParameters->bRectPrimitive = GRHISupportsRectTopology ? 1 : 0;
-
-		auto ComputeShader = GetGlobalShaderMap(GMaxRHIFeatureLevel)->GetShader<FSetupClearHwPhysicalPagesDrawArgsCS>();
-		FComputeShaderUtils::AddPass(
-			GraphBuilder,
-			RDG_EVENT_NAME("InitClearPhysicalPagesArgs"),
-			ComputeShader,
-			PassParameters,
-			FIntVector(1, 1, 1)
-		);
-	}
-
-	{
-		FInitializePhysicalPagesPS::FParameters* ParametersPS = GraphBuilder.AllocParameters<FInitializePhysicalPagesPS::FParameters>();
-		ParametersPS->VirtualShadowMap = GetUniformBuffer(GraphBuilder);
-		ParametersPS->IndirectArgsBuffer = IndirectArgsBuffer;
-		ParametersPS->bRectPrimitive = GRHISupportsRectTopology ? 1 : 0;
-
-		bool bCacheDataAvailable = CacheManager && CacheManager->PrevBuffers.PhysicalPagePoolHw;
-		if (bCacheDataAvailable)
-		{
-			ParametersPS->PrevPhysicalPagePoolHw = GraphBuilder.RegisterExternalTexture(CacheManager->PrevBuffers.PhysicalPagePoolHw);
-			ParametersPS->CachedPageInfos = GraphBuilder.CreateSRV(CachedPageInfosRDG);
-		}
-
-
-		FInitializePhysicalPagesPS::FPermutationDomain PermutationVector;
-		PermutationVector.Set<FInitializePhysicalPagesPS::FHasCacheDataDim>(bCacheDataAvailable);
-
-		TShaderMapRef<FInitializePhysicalPagesVS> VertexShader(GetGlobalShaderMap(GMaxRHIFeatureLevel));
-		TShaderMapRef<FInitializePhysicalPagesPS> PixelShader(GetGlobalShaderMap(GMaxRHIFeatureLevel), PermutationVector);
-
-		ParametersPS->RenderTargets.DepthStencil = FDepthStencilBinding(PhysicalPagePoolHw, ERenderTargetLoadAction::ENoAction, FExclusiveDepthStencil::DepthWrite_StencilNop);
-
-		GraphBuilder.AddPass(
-			RDG_EVENT_NAME("InitializePhysicalPages"),
-			ParametersPS,
-			ERDGPassFlags::Raster,
-			[ParametersPS, VertexShader, PixelShader, PhysicalPagePoolSize = GetPhysicalPoolSize()](FRHICommandList& RHICmdList)
-		{
-			FInitializePhysicalPagesVS::FParameters ParametersVS;
-			ParametersVS.VirtualShadowMap = ParametersPS->VirtualShadowMap;
-
-			FGraphicsPipelineStateInitializer GraphicsPSOInit;
-			RHICmdList.ApplyCachedRenderTargets(GraphicsPSOInit);
-			GraphicsPSOInit.BlendState = TStaticBlendState<CW_NONE>::GetRHI();
-			GraphicsPSOInit.RasterizerState = TStaticRasterizerState<FM_Solid, CM_None>::GetRHI();
-			GraphicsPSOInit.DepthStencilState = TStaticDepthStencilState<true, CF_Always>::GetRHI();
-			GraphicsPSOInit.BoundShaderState.VertexDeclarationRHI = GFilterVertexDeclaration.VertexDeclarationRHI;
-			GraphicsPSOInit.BoundShaderState.VertexShaderRHI = VertexShader.GetVertexShader();
-			GraphicsPSOInit.BoundShaderState.PixelShaderRHI = PixelShader.GetPixelShader();
-			GraphicsPSOInit.PrimitiveType = ParametersPS->bRectPrimitive > 0 ? PT_RectList : PT_TriangleList;
-			SetGraphicsPipelineState(RHICmdList, GraphicsPSOInit);
-			SetShaderParameters(RHICmdList, VertexShader, VertexShader.GetVertexShader(), ParametersVS);
-			SetShaderParameters(RHICmdList, PixelShader, PixelShader.GetPixelShader(), *ParametersPS);
-
-			RHICmdList.SetViewport(
-				0.0f,
-				0.0f,
-				0.0f,
-				PhysicalPagePoolSize.X,
-				PhysicalPagePoolSize.Y,
-				1.0f
-			);
-			RHICmdList.SetStreamSource(0, nullptr, 0);
-			RHICmdList.DrawPrimitiveIndirect(ParametersPS->IndirectArgsBuffer->GetRHI(), 0);
-		});
-	}
-	{
-		FUpdatePhysicalPageMetaDataClearStateCS::FParameters* PassParameters = GraphBuilder.AllocParameters<FUpdatePhysicalPageMetaDataClearStateCS::FParameters>();
-		PassParameters->NumAllocatedPhysicalPages = GraphBuilder.CreateSRV(AllocatedPagesOffsetRDG);
-		PassParameters->PhysicalPageMetaDataOut = GraphBuilder.CreateUAV(PhysicalPageMetaDataRDG);
-		PassParameters->CachedPageInfos = GraphBuilder.CreateSRV(CachedPageInfosRDG);
-		PassParameters->VirtualShadowMap = GetUniformBuffer(GraphBuilder);
-
-		bool bCacheDataAvailable = CacheManager && CacheManager->PrevBuffers.PhysicalPageMetaData;
-		if (bCacheDataAvailable)
-		{
-			PassParameters->PrevPhysicalPageMetaData = GraphBuilder.CreateSRV(GraphBuilder.RegisterExternalBuffer(CacheManager->PrevBuffers.PhysicalPageMetaData));
-		}
-
-		FUpdatePhysicalPageMetaDataClearStateCS::FPermutationDomain PermutationVector;
-		PermutationVector.Set<FUpdatePhysicalPageMetaDataClearStateCS::FHasCacheDataDim>(bCacheDataAvailable);
-
-		auto ComputeShader = GetGlobalShaderMap(GMaxRHIFeatureLevel)->GetShader<FUpdatePhysicalPageMetaDataClearStateCS>(PermutationVector);
-		FComputeShaderUtils::AddPass(
-			GraphBuilder,
-			RDG_EVENT_NAME("UpdatePhysicalPageMetaDataClearState"),
-			ComputeShader,
-			PassParameters,
-			FComputeShaderUtils::GetGroupCount(UniformParameters.MaxPhysicalPages, 64U)
-		);
-	}
 }
diff --git a/Engine/Source/Runtime/Renderer/Private/VirtualShadowMaps/VirtualShadowMapArray.h b/Engine/Source/Runtime/Renderer/Private/VirtualShadowMaps/VirtualShadowMapArray.h
--- a/Engine/Source/Runtime/Renderer/Private/VirtualShadowMaps/VirtualShadowMapArray.h
+++ b/Engine/Source/Runtime/Renderer/Private/VirtualShadowMaps/VirtualShadowMapArray.h
@@ -125,7 +125,6 @@ BEGIN_GLOBAL_SHADER_PARAMETER_STRUCT(FVirtualShadowMapUniformParameters, )
 	SHADER_PARAMETER_RDG_BUFFER_SRV(ByteAddressBuffer, ProjectionData)
 	SHADER_PARAMETER_RDG_BUFFER_SRV(StructuredBuffer<uint>, PageTable)
 	SHADER_PARAMETER_RDG_TEXTURE(Texture2D<uint>, PhysicalPagePool)
-	SHADER_PARAMETER_RDG_TEXTURE(Texture2D<float>, PhysicalPagePoolHw)
 END_GLOBAL_SHADER_PARAMETER_STRUCT()
 
 BEGIN_SHADER_PARAMETER_STRUCT(FVirtualShadowMapSamplingParameters, )
@@ -174,9 +173,6 @@ public:
 
 	static void SetShaderDefines(FShaderCompilerEnvironment& OutEnvironment);
 
-	// Call after creating physical page pools and rendering
-	void SetupProjectionParameters(FRDGBuilder& GraphBuilder);
-
 	void ClearPhysicalMemory(FRDGBuilder& GraphBuilder, FRDGTextureRef& PhysicalTexture);
 	void MarkPhysicalPagesRendered(FRDGBuilder& GraphBuilder, const TArray<uint32, SceneRenderingAllocator> &VirtualShadowMapFlags);
 
@@ -202,8 +198,6 @@ public:
 	 */
 	void RenderVirtualShadowMapsHw(FRDGBuilder& GraphBuilder, const TArray<FProjectedShadowInfo*, SceneRenderingAllocator>& VirtualSmMeshCommandPasses, FScene& Scene);
 
-	void AddInitializePhysicalPagesHwPass(FRDGBuilder& GraphBuilder);
-
 	// Draw debug info into render target 'VSMDebug' of screen-size, the mode is controlled by 'r.Shadow.Virtual.DebugVisualize'.
 	void RenderDebugInfo(FRDGBuilder& GraphBuilder);
 	// 
@@ -216,7 +210,7 @@ public:
 	// This data becomes valid after the shadow depths pass if VSMs are enabled
 	FVirtualShadowMapSamplingParameters GetSamplingParameters(FRDGBuilder& GraphBuilder) const;
 
-	bool HasAnyShadowData() const { return PhysicalPagePoolRDG != nullptr || PhysicalPagePoolHw != nullptr;  }
+	bool HasAnyShadowData() const { return PhysicalPagePoolRDG != nullptr;  }
 
 	void GetPageTableParameters(FRDGBuilder& GraphBuilder, FVirtualShadowMapPageTableParameters& OutParameters);
 
@@ -234,9 +228,6 @@ public:
 	FRDGBufferRef PageTableRDG = nullptr;
 	// Large physical texture of depth format, say 4096^2 or whatever we think is enough texels to go around
 	FRDGTextureRef PhysicalPagePoolRDG = nullptr;
-	// page pool for HW rasterized shadow data. 
-	// Mirrors the regular one, but uses a shadow depth target format
-	FRDGTextureRef PhysicalPagePoolHw = nullptr;
 		
 	// Buffer that stores flags (uints) marking each page that needs to be rendered and cache status, for all virtual shadow maps.
 	// Flag values defined in PageAccessCommon.ush: VSM_ALLOCATED_FLAG | VSM_INVALID_FLAG
diff --git a/Engine/Source/Runtime/Renderer/Private/VirtualShadowMaps/VirtualShadowMapCacheManager.cpp b/Engine/Source/Runtime/Renderer/Private/VirtualShadowMaps/VirtualShadowMapCacheManager.cpp
--- a/Engine/Source/Runtime/Renderer/Private/VirtualShadowMaps/VirtualShadowMapCacheManager.cpp
+++ b/Engine/Source/Runtime/Renderer/Private/VirtualShadowMaps/VirtualShadowMapCacheManager.cpp
@@ -189,15 +189,6 @@ void FVirtualShadowMapArrayCacheManager::ExtractFrameData(bool bEnableCaching, F
 		
 			GraphBuilder.QueueTextureExtraction(VirtualShadowMapArray.PhysicalPagePoolRDG, &PrevBuffers.PhysicalPagePool);
 
-			if( VirtualShadowMapArray.PhysicalPagePoolHw )
-			{
-				GraphBuilder.QueueTextureExtraction(VirtualShadowMapArray.PhysicalPagePoolHw, &PrevBuffers.PhysicalPagePoolHw);
-			}
-			else
-			{
-				PrevBuffers.PhysicalPagePoolHw = TRefCountPtr<IPooledRenderTarget>();
-			}
-
 			GraphBuilder.QueueBufferExtraction(VirtualShadowMapArray.PhysicalPageMetaDataRDG, &PrevBuffers.PhysicalPageMetaData);
 			GraphBuilder.QueueBufferExtraction(VirtualShadowMapArray.DynamicCasterPageFlagsRDG, &PrevBuffers.DynamicCasterPageFlags);
 			GraphBuilder.QueueBufferExtraction(VirtualShadowMapArray.ShadowMapProjectionDataRDG, &PrevBuffers.ShadowMapProjectionDataBuffer);
@@ -223,7 +214,6 @@ void FVirtualShadowMapArrayCacheManager::ExtractFrameData(bool bEnableCaching, F
 	PrevUniformParameters.ProjectionData = nullptr;
 	PrevUniformParameters.PageTable = nullptr;
 	PrevUniformParameters.PhysicalPagePool = nullptr;
-	PrevUniformParameters.PhysicalPagePoolHw = nullptr;
 
 	FRDGBufferRef AccumulatedStatsBufferRDG = nullptr;
 
@@ -346,7 +336,7 @@ bool FVirtualShadowMapArrayCacheManager::IsValid()
 	return CVarCacheVirtualSMs.GetValueOnRenderThread() != 0
 		&& PrevBuffers.PageTable
 		&& PrevBuffers.PageFlags
-		&& (PrevBuffers.PhysicalPagePool || PrevBuffers.PhysicalPagePoolHw)
+		&& PrevBuffers.PhysicalPagePool
 		&& PrevBuffers.PhysicalPageMetaData
 		&& PrevBuffers.DynamicCasterPageFlags;
 }
@@ -491,7 +481,6 @@ void FVirtualShadowMapArrayCacheManager::ProcessInstanceRangeInvalidation(FRDGBu
 	PrevUniformParameters.PageTable = RegExtCreateSrv(PrevBuffers.PageTable, TEXT("Shadow.Virtual.PrevPageTable"));
 	// Unused in this path
 	PrevUniformParameters.PhysicalPagePool = GSystemTextures.GetZeroUIntDummy(GraphBuilder);
-	PrevUniformParameters.PhysicalPagePoolHw = GSystemTextures.GetBlackDummy(GraphBuilder);
 
 	FRDGBufferRef InstanceRangesSmallRDG = !InstanceRangesSmall.IsEmpty() ? CreateStructuredBuffer(GraphBuilder, TEXT("Shadow.Virtual.InstanceRangesSmall"), InstanceRangesSmall) : nullptr;
 	FRDGBufferRef InstanceRangesLargeRDG = !InstanceRangesLarge.IsEmpty() ? CreateStructuredBuffer(GraphBuilder, TEXT("Shadow.Virtual.InstanceRangesSmall"), InstanceRangesLarge) : nullptr;
diff --git a/Engine/Source/Runtime/Renderer/Private/VirtualShadowMaps/VirtualShadowMapCacheManager.h b/Engine/Source/Runtime/Renderer/Private/VirtualShadowMaps/VirtualShadowMapCacheManager.h
--- a/Engine/Source/Runtime/Renderer/Private/VirtualShadowMaps/VirtualShadowMapCacheManager.h
+++ b/Engine/Source/Runtime/Renderer/Private/VirtualShadowMaps/VirtualShadowMapCacheManager.h
@@ -69,7 +69,6 @@ struct FVirtualShadowMapArrayFrameData
 	TRefCountPtr<FRDGPooledBuffer>		DynamicCasterPageFlags;
 
 	TRefCountPtr<IPooledRenderTarget>	PhysicalPagePool;
-	TRefCountPtr<IPooledRenderTarget>	PhysicalPagePoolHw;
 	TRefCountPtr<FRDGPooledBuffer>		PhysicalPageMetaData;
 
 	TRefCountPtr<IPooledRenderTarget>			HZBPhysical;
diff --git a/Engine/Source/Runtime/Renderer/Private/VirtualShadowMaps/VirtualShadowMapProjection.cpp b/Engine/Source/Runtime/Renderer/Private/VirtualShadowMaps/VirtualShadowMapProjection.cpp
--- a/Engine/Source/Runtime/Renderer/Private/VirtualShadowMaps/VirtualShadowMapProjection.cpp
+++ b/Engine/Source/Runtime/Renderer/Private/VirtualShadowMaps/VirtualShadowMapProjection.cpp
@@ -168,7 +168,6 @@ class FVirtualShadowMapProjectionCS : public FGlobalShader
 	
 	class FDirectionalLightDim		: SHADER_PERMUTATION_BOOL("DIRECTIONAL_LIGHT");
 	class FSMRTAdaptiveRayCountDim	: SHADER_PERMUTATION_BOOL("SMRT_ADAPTIVE_RAY_COUNT");
-	class FTwoPhysicalTexturesDim	: SHADER_PERMUTATION_BOOL("TWO_PHYSICAL_TEXTURES");
 	class FOnePassProjectionDim		: SHADER_PERMUTATION_BOOL("ONE_PASS_PROJECTION");
 	class FHairStrandsDim			: SHADER_PERMUTATION_BOOL("HAS_HAIR_STRANDS");
 	class FDebugOutputDim			: SHADER_PERMUTATION_BOOL("DEBUG_OUTPUT");
@@ -177,7 +176,6 @@ class FVirtualShadowMapProjectionCS : public FGlobalShader
 		FDirectionalLightDim,
 		FOnePassProjectionDim,
 		FSMRTAdaptiveRayCountDim,
-		FTwoPhysicalTexturesDim,
 		FHairStrandsDim,
 		FDebugOutputDim>;
 
@@ -248,8 +246,6 @@ static float GetNormalBiasForShader()
 	return CVarNormalBias.GetValueOnRenderThread() / 1000.0f;
 }
 
-extern int32 GVirtualShadowMapAtomicWrites;
-
 static void RenderVirtualShadowMapProjectionCommon(
 	FRDGBuilder& GraphBuilder,
 	const FMinimalSceneTextures& SceneTextures,
@@ -329,7 +325,6 @@ static void RenderVirtualShadowMapProjectionCommon(
 	PermutationVector.Set< FVirtualShadowMapProjectionCS::FDirectionalLightDim >( bDirectionalLight );
 	PermutationVector.Set< FVirtualShadowMapProjectionCS::FOnePassProjectionDim >( bOnePassProjection );
 	PermutationVector.Set< FVirtualShadowMapProjectionCS::FSMRTAdaptiveRayCountDim >( bAdaptiveRayCount );
-	PermutationVector.Set< FVirtualShadowMapProjectionCS::FTwoPhysicalTexturesDim >( GVirtualShadowMapAtomicWrites == 0 );
 	PermutationVector.Set< FVirtualShadowMapProjectionCS::FHairStrandsDim >( bHasHairStrandsData ? 1 : 0 );
 	PermutationVector.Set< FVirtualShadowMapProjectionCS::FDebugOutputDim >( bDebugOutput );
 
