Cleaned up message when skylight needs to be recaptured

diff --git a/Engine/Source/Runtime/Renderer/Private/SceneRendering.cpp b/Engine/Source/Runtime/Renderer/Private/SceneRendering.cpp
--- a/Engine/Source/Runtime/Renderer/Private/SceneRendering.cpp
+++ b/Engine/Source/Runtime/Renderer/Private/SceneRendering.cpp
@@ -3253,7 +3253,7 @@ void FSceneRenderer::RenderFinish(FRDGBuilder& GraphBuilder, FRDGTextureRef View
 						{
 							const FLinearColor OrangeColor = FColor::Orange;
 
-							FString String = TEXT("Sky Light waiting on [");
+							FString String = TEXT("Sky Light waiting on ");
 							bool bAddComma = false;
 							if (SkyLight->bCubemapSkyLightWaitingForCubeMapTexture)
 							{
@@ -3269,13 +3269,9 @@ void FSceneRenderer::RenderFinish(FRDGBuilder& GraphBuilder, FRDGTextureRef View
 							if (SkyLight->bCaptureSkyLightWaitingForMeshesOrTextures)
 							{
 								String += bAddComma ? TEXT(", ") : TEXT("");
-								String += TEXT("Meshes, Textures] for final capture");
+								String += TEXT("Meshes, Textures");
 							}
-							String += TEXT("] for final capture.");
-							Canvas.DrawShadowedText(10, Y, FText::FromString(String), GetStatsFont(), OrangeColor);
-							Y += 14;
-
-							String = FString::Printf(TEXT("Seconds to next incomplete capture = %.0f"), FMath::CeilToFloat(SkyLight->SecondsToNextIncompleteCapture));
+							String += TEXT(" for final capture.");
 							Canvas.DrawShadowedText(10, Y, FText::FromString(String), GetStatsFont(), OrangeColor);
 							Y += 14;
 						}
