Fixed: Crash on mobile while rendering custom depth primitive that requires scene texture access

diff --git a/Engine/Source/Runtime/Renderer/Private/MobileShadingRenderer.cpp b/Engine/Source/Runtime/Renderer/Private/MobileShadingRenderer.cpp
--- a/Engine/Source/Runtime/Renderer/Private/MobileShadingRenderer.cpp
+++ b/Engine/Source/Runtime/Renderer/Private/MobileShadingRenderer.cpp
@@ -773,6 +773,9 @@ void FMobileSceneRenderer::Render(FRDGBuilder& GraphBuilder)
 	// bShouldRenderCustomDepth has been initialized in InitViews on mobile platform
 	if (bShouldRenderCustomDepth)
 	{
+		SceneTextures.MobileSetupMode = EMobileSceneTextureSetupMode::None;
+		SceneTextures.MobileUniformBuffer = CreateMobileSceneTextureUniformBuffer(GraphBuilder, SceneTextures.MobileSetupMode);
+
 		RenderCustomDepthPass(GraphBuilder, SceneTextures.CustomDepth, SceneTextures.GetSceneTextureShaderParameters(FeatureLevel));
 	}
 
