Fixing build break

diff --git a/Engine/Source/Runtime/Renderer/Private/SceneRendering.cpp b/Engine/Source/Runtime/Renderer/Private/SceneRendering.cpp
--- a/Engine/Source/Runtime/Renderer/Private/SceneRendering.cpp
+++ b/Engine/Source/Runtime/Renderer/Private/SceneRendering.cpp
@@ -3784,8 +3784,7 @@ static void RenderViewFamily_RenderThread(FRHICommandListImmediate& RHICmdList,
 				RDG_EVENT_NAME("SceneRenderer_%s(ViewFamily=%s)",
 					ViewFamily.EngineShowFlags.HitProxies ? TEXT("RenderHitProxies") : TEXT("Render"),
 					ViewFamily.bResolveScene ? TEXT("Primary") : TEXT("Auxiliary")
-				),
-				ERDGBuilderFlags::AllowParallelExecute
+				)
 			);
 
 			if (ViewFamily.EngineShowFlags.HitProxies)
