Add dummy vertex buffer for whole scene directional shadows.

- This prevents some D3D debug warnings about zero-element input layouts, and provides Metal with a valid binding to match the vertex descriptor.

#rnx

diff --git a/Engine/Source/Runtime/Renderer/Private/ShadowRendering.cpp b/Engine/Source/Runtime/Renderer/Private/ShadowRendering.cpp
--- a/Engine/Source/Runtime/Renderer/Private/ShadowRendering.cpp
+++ b/Engine/Source/Runtime/Renderer/Private/ShadowRendering.cpp
@@ -190,6 +190,43 @@ static TAutoConsoleVariable<float> CVarShadowMaxSlopeScaleDepthBias(
 
 FForwardScreenSpaceShadowMaskTextureMobileOutputs GScreenSpaceShadowMaskTextureMobileOutputs;
 
+/**
+ * A dummy vertex buffer to bind when rendering whole scene shadows. This
+ * prevents some D3D debug warnings about zero-element input layouts, but is not
+ * strictly required.  It is, however, required for Metal.
+ */
+class FDummyWholeSceneDirectionalShadowStencilVertexBuffer : public FVertexBuffer
+{
+public:
+
+	virtual void InitRHI() override
+	{
+		FRHIResourceCreateInfo CreateInfo(TEXT("FDummyWholeSceneDirectionalShadowStencilVertexBuffer"));
+		VertexBufferRHI = RHICreateBuffer(sizeof(FVector4) * 12, BUF_Static | BUF_VertexBuffer, 0, ERHIAccess::VertexOrIndexBuffer, CreateInfo);
+		FVector4* DummyContents = (FVector4*)RHILockBuffer(VertexBufferRHI, 0, sizeof(FVector4) * 12, RLM_WriteOnly);
+
+		// Far Plane
+		DummyContents[0] = FVector4( 1,  1,  1 /* StencilFar */);
+		DummyContents[1] = FVector4(-1,  1,  1);
+		DummyContents[2] = FVector4( 1, -1,  1);
+		DummyContents[3] = FVector4( 1, -1,  1);
+		DummyContents[4] = FVector4(-1,  1,  1);
+		DummyContents[5] = FVector4(-1, -1,  1);
+
+		// Near Plane
+		DummyContents[6]  = FVector4(-1,  1, -1 /* StencilNear */);
+		DummyContents[7]  = FVector4( 1,  1, -1);
+		DummyContents[8]  = FVector4(-1, -1, -1);
+		DummyContents[9]  = FVector4(-1, -1, -1);
+		DummyContents[10] = FVector4( 1,  1, -1);
+		DummyContents[11] = FVector4( 1, -1, -1);
+
+		RHIUnlockBuffer(VertexBufferRHI);
+	}
+};
+
+TGlobalResource<FDummyWholeSceneDirectionalShadowStencilVertexBuffer> GDummyWholeSceneDirectionalShadowStencilVertexBuffer;
+
 ///////////////////////////////////////////////////////////////////////////////////////////////////
 // Hair
 static TAutoConsoleVariable<int32> CVarHairStrandsCullPerObjectShadowCaster(
@@ -927,7 +964,7 @@ void FProjectedShadowInfo::SetupProjectionStencilMask(
 
 		SetShaderParameters(RHICmdList, VertexShader, VertexShader.GetVertexShader(), Parameters);
 
-		RHICmdList.SetStreamSource(0, nullptr, 0);
+		RHICmdList.SetStreamSource(0, GDummyWholeSceneDirectionalShadowStencilVertexBuffer.VertexBufferRHI, 0);
 		RHICmdList.DrawPrimitive(0, (CascadeSettings.ShadowSplitIndex > 0) ? 2 : 1, 1);
 	}
 	// Not a preshadow, mask the projection to any pixels inside the frustum.
