Unified OculusVR FFR options and Engine->Rendering VRS options under Rendering. Removed options from OculusVR plugin.
Implemented VRS support for adjustable centerpoints and added a foveation level option to match Oculus' "High Top".

diff --git a/Engine/Plugins/Runtime/Oculus/OculusVR/Source/OculusHMD/Private/OculusHMD.cpp b/Engine/Plugins/Runtime/Oculus/OculusVR/Source/OculusHMD/Private/OculusHMD.cpp
--- a/Engine/Plugins/Runtime/Oculus/OculusVR/Source/OculusHMD/Private/OculusHMD.cpp
+++ b/Engine/Plugins/Runtime/Oculus/OculusVR/Source/OculusHMD/Private/OculusHMD.cpp
@@ -29,6 +29,7 @@
 #include "DynamicResolutionState.h"
 #include "DynamicResolutionProxy.h"
 #include "OculusHMDRuntimeSettings.h"
+#include "Engine/RendererSettings.h"
 
 #if PLATFORM_ANDROID
 #include "Android/AndroidJNI.h"
@@ -3576,13 +3577,20 @@ namespace OculusHMD
 		Settings->Flags.bHQDistortion = HMDSettings->bHQDistortion;
 		Settings->Flags.bChromaAbCorrectionEnabled = HMDSettings->bChromaCorrection;
 		Settings->Flags.bRecenterHMDWithController = HMDSettings->bRecenterHMDWithController;
-		Settings->FFRLevel = HMDSettings->FFRLevel;
-		Settings->FFRDynamic = HMDSettings->FFRDynamic;
 		Settings->CPULevel = HMDSettings->CPULevel;
 		Settings->GPULevel = HMDSettings->GPULevel;
 		Settings->PixelDensityMin = HMDSettings->PixelDensityMin;
 		Settings->PixelDensityMax = HMDSettings->PixelDensityMax;
 		Settings->bPhaseSync = HMDSettings->bPhaseSync;
+
+		// Set FFR level and dynamic from rendering settings
+		UEnum* FFREnum = StaticEnum<EFixedFoveatedRenderingLevel>();
+		int32 VRSValue = IConsoleManager::Get().FindConsoleVariable(TEXT("vr.VRS.HMDFixedFoveationLevel"))->GetInt();
+		if (FFREnum->IsValidEnumValue(VRSValue))
+		{
+			Settings->FFRLevel = static_cast<EFixedFoveatedRenderingLevel>(VRSValue);
+		}
+		Settings->FFRDynamic = IConsoleManager::Get().FindConsoleVariable(TEXT("vr.VRS.HMDFixedFoveationDynamic"))->GetBool();
 	}
 
 	/// @endcond
diff --git a/Engine/Plugins/Runtime/Oculus/OculusVR/Source/OculusHMD/Private/OculusHMDRuntimeSettings.cpp b/Engine/Plugins/Runtime/Oculus/OculusVR/Source/OculusHMD/Private/OculusHMDRuntimeSettings.cpp
--- a/Engine/Plugins/Runtime/Oculus/OculusVR/Source/OculusHMD/Private/OculusHMDRuntimeSettings.cpp
+++ b/Engine/Plugins/Runtime/Oculus/OculusVR/Source/OculusHMD/Private/OculusHMDRuntimeSettings.cpp
@@ -18,8 +18,6 @@ UOculusHMDRuntimeSettings::UOculusHMDRuntimeSettings(const FObjectInitializer& O
 	bCompositesDepth = DefaultSettings.Flags.bCompositeDepth;
 	bHQDistortion = DefaultSettings.Flags.bHQDistortion;
 	bChromaCorrection = DefaultSettings.Flags.bChromaAbCorrectionEnabled;
-	FFRLevel = DefaultSettings.FFRLevel;
-	FFRDynamic = DefaultSettings.FFRDynamic;
 	CPULevel = DefaultSettings.CPULevel;
 	GPULevel = DefaultSettings.GPULevel;
 	PixelDensityMin = DefaultSettings.PixelDensityMin;
@@ -38,8 +36,6 @@ UOculusHMDRuntimeSettings::UOculusHMDRuntimeSettings(const FObjectInitializer& O
 	bCompositesDepth = false;
 	bHQDistortion = false;
 	bChromaCorrection = false;
-	FFRLevel = EFixedFoveatedRenderingLevel::EFixedFoveatedRenderingLevel_Off;
-	FFRDynamic = false;
 	CPULevel = 2;
 	GPULevel = 3;
 	PixelDensityMin = 0.5f;
diff --git a/Engine/Plugins/Runtime/Oculus/OculusVR/Source/OculusHMD/Public/OculusHMDRuntimeSettings.h b/Engine/Plugins/Runtime/Oculus/OculusVR/Source/OculusHMD/Public/OculusHMDRuntimeSettings.h
--- a/Engine/Plugins/Runtime/Oculus/OculusVR/Source/OculusHMD/Public/OculusHMDRuntimeSettings.h
+++ b/Engine/Plugins/Runtime/Oculus/OculusVR/Source/OculusHMD/Public/OculusHMDRuntimeSettings.h
@@ -62,14 +62,6 @@ public:
 	UPROPERTY(config, EditAnywhere, Category = Mobile)
 	int GPULevel;
 
-	/** Default Fixed Foveated Rendering level for Oculus Quest */
-	UPROPERTY(config, EditAnywhere, Category = Mobile)
-	EFixedFoveatedRenderingLevel FFRLevel;
-
-	/** Whether FFR adjusts its level dynamically or not */
-	UPROPERTY(config, EditAnywhere, Category = Mobile)
-	bool FFRDynamic;
-
 	/** Compensates in the compositor for chromatic aberration, at a higher GPU cost but without the color fringes on the sides of the lenses */
 	UPROPERTY(config, EditAnywhere, Category = Mobile)
 	bool bChromaCorrection;
diff --git a/Engine/Source/Runtime/Engine/Classes/Engine/RendererSettings.h b/Engine/Source/Runtime/Engine/Classes/Engine/RendererSettings.h
--- a/Engine/Source/Runtime/Engine/Classes/Engine/RendererSettings.h
+++ b/Engine/Source/Runtime/Engine/Classes/Engine/RendererSettings.h
@@ -152,6 +152,7 @@ namespace EFixedFoveationLevels
 		Low = 1 UMETA(DisplayName = "Low"),
 		Medium = 2 UMETA(DisplayName = "Medium"),
 		High = 3 UMETA(DisplayName = "High"),
+		HighTop = 4 UMETA(DisplayName = "High Top"),
 	};
 }
 
@@ -476,9 +477,14 @@ class ENGINE_API URendererSettings : public UDeveloperSettings
 	FVector TranslucentSortAxis;
 
 	UPROPERTY(config, EditAnywhere, Category = VR, meta = (
-		ConsoleVariable = "vr.VRS.HMDFixedFoveationLevel", DisplayName = "HMD Fixed Foveation Level (experimental)",
+		ConsoleVariable = "vr.VRS.HMDFixedFoveationLevel", DisplayName = "HMD Fixed Foveation Level (Experimental)",
 		ToolTip = "Set the level of fixed-foveation to apply when generating the Variable Rate Shading attachment. This feature is currently experimental.\nThis can yield some fairly significant performance benefits on GPUs that support Tier 2 VRS.\nLower settings will result in almost no discernible artifacting on most HMDs; higher settings will show some artifacts towards the edges of the view."))
-		TEnumAsByte<EFixedFoveationLevels::Type> HMDFixedFoveationLevel;
+	TEnumAsByte<EFixedFoveationLevels::Type> HMDFixedFoveationLevel;
+
+	UPROPERTY(config, EditAnywhere, Category = VR, meta = (
+		ConsoleVariable = "vr.VRS.HMDFixedFoveationDynamic", DisplayName = "Dynamic Fixed Foveation (Oculus Quest/Quest 2 Only)",
+		ToolTip = "Allows fixed foveation level to adjust dynamically based on GPU utilization.\nLevel will range between none at the minimum, and the currently selected foveation level at the maximum."))
+	uint32 bHMDFixedFoveationDynamic:1;
 
 	UPROPERTY(config, EditAnywhere, Category=Postprocessing, meta=(
 		ConsoleVariable="r.CustomDepth",DisplayName="Custom Depth-Stencil Pass",
diff --git a/Engine/Source/Runtime/Renderer/Private/VariableRateShadingImageManager.cpp b/Engine/Source/Runtime/Renderer/Private/VariableRateShadingImageManager.cpp
--- a/Engine/Source/Runtime/Renderer/Private/VariableRateShadingImageManager.cpp
+++ b/Engine/Source/Runtime/Renderer/Private/VariableRateShadingImageManager.cpp
@@ -37,8 +37,16 @@ static TAutoConsoleVariable<int> CVarHMDFixedFoveationLevel(
 	TEXT(" 0: Disabled (default);\n")
 	TEXT(" 1: Low;\n")
 	TEXT(" 2: Medium;\n")
-	TEXT(" 3: High;\n"),
-	// @todo: 4 - Dynamic (adjusts based on framerate)
+	TEXT(" 3: High;\n")
+	TEXT(" 4: High Top;\n"),
+	ECVF_RenderThreadSafe);
+
+static TAutoConsoleVariable<int32> CVarHMDFixedFoveationDynamic(
+	TEXT("vr.VRS.HMDFixedFoveationDynamic"),
+	0,
+	TEXT("Whether fixed-foveation level should adjust based on GPU utilization (Currently OculusVR only)\n")
+	TEXT(" 0: Disabled (default);\n")
+	TEXT(" 1: Enabled\n"),
 	ECVF_RenderThreadSafe);
 
 TGlobalResource<FVariableRateShadingImageManager> GVRSImageManager;
@@ -108,7 +116,10 @@ struct FVRSImageGenerationParameters
 	float HMDFixedFoveationFullRateCutoff = 1.0f;
 	float HMDFixedFoveationHalfRateCutoff = 1.0f;
 	float HMDEyeTrackedFoveationFullRateCutoff = 1.0f;
-	float HMDEYeTrackedFoveationHalfRateCutoff = 1.0f;
+	float HMDEyeTrackedFoveationHalfRateCutoff = 1.0f;
+
+	float HMDFixedFoveationCenterX = 0.5f;
+	float HMDFixedFoveationCenterY = 0.5f;
 
 	bool bGenerateFixedFoveation = false;
 	bool bGenerateEyeTrackedFoveation = false;
@@ -129,8 +140,10 @@ uint64 FVariableRateShadingImageManager::CalculateVRSImageHash(const FVRSImageGe
 		Hash = HashCombine(Hash, GetTypeHash(GenerationParamsIn.HMDFixedFoveationHalfRateCutoff));
 		Hash = HashCombine(Hash, GetTypeHash(GenerationParamsIn.HMDEyeTrackedFoveationOrigin));
 		Hash = HashCombine(Hash, GetTypeHash(GenerationParamsIn.HMDEyeTrackedFoveationFullRateCutoff));
-		Hash = HashCombine(Hash, GetTypeHash(GenerationParamsIn.HMDEYeTrackedFoveationHalfRateCutoff));
+		Hash = HashCombine(Hash, GetTypeHash(GenerationParamsIn.HMDEyeTrackedFoveationHalfRateCutoff));
 		Hash = HashCombine(Hash, GetTypeHash(GenerationParamsIn.bInstancedStereo));
+		Hash = HashCombine(Hash, GetTypeHash(GenerationParamsIn.HMDFixedFoveationCenterX));
+		Hash = HashCombine(Hash, GetTypeHash(GenerationParamsIn.HMDFixedFoveationCenterY));
 	}
 
 	return Hash;
@@ -282,7 +295,7 @@ TRefCountPtr<IPooledRenderTarget> FVariableRateShadingImageManager::RenderShadin
 	PassParameters->FixedFoveationFullRateCutoffSquared = VRSImageGenParamsIn.HMDFixedFoveationFullRateCutoff * VRSImageGenParamsIn.HMDFixedFoveationFullRateCutoff;
 	PassParameters->FixedFoveationHalfRateCutoffSquared = VRSImageGenParamsIn.HMDFixedFoveationHalfRateCutoff * VRSImageGenParamsIn.HMDFixedFoveationHalfRateCutoff;
 
-	PassParameters->LeftEyeCenterPixelXY = FVector2D(AttachmentSize.X / 2, AttachmentSize.Y / 2);
+	PassParameters->LeftEyeCenterPixelXY = FVector2D(AttachmentSize.X * VRSImageGenParamsIn.HMDFixedFoveationCenterX, AttachmentSize.Y * VRSImageGenParamsIn.HMDFixedFoveationCenterY);
 	PassParameters->RightEyeCenterPixelXY = PassParameters->LeftEyeCenterPixelXY;
 
 	// If instanced (side-by-side) stereo, there's two "center" points, so adjust both eyes
@@ -380,21 +393,27 @@ void FVariableRateShadingImageManager::Tick()
 
 void FVariableRateShadingImageManager::UpdateFixedFoveationParameters(FVRSImageGenerationParameters& VRSImageGenParamsInOut)
 {
-	static const float FIXED_FOVEATION_FULL_RATE_CUTOFFS[] = { 1.0f, 0.7f, 0.50f, 0.35f };
-	static const float FIXED_FOVEATION_HALF_RATE_CUTOFFS[] = { 1.0f, 0.9f, 0.75f, 0.45f };
 
-	int Level = FMath::Clamp(CVarHMDFixedFoveationLevel.GetValueOnAnyThread(), 0, 3);
+	// Numbers here are currently pretty arbitrary... should eventually be chosen based on headset specific characteristics
+	static const float FIXED_FOVEATION_FULL_RATE_CUTOFFS[] = { 1.0f, 0.7f, 0.50f, 0.35f, 0.35f };
+	static const float FIXED_FOVEATION_HALF_RATE_CUTOFFS[] = { 1.0f, 0.9f, 0.75f, 0.55f, 0.55f };
+	static const float FIXED_FOVEATION_CENTER_X[] = { 0.5f, 0.5f, 0.5f, 0.5f, 0.5f };
+	static const float FIXED_FOVEATION_CENTER_Y[] = { 0.5f, 0.5f, 0.5f, 0.5f, 0.42f };
+	
+	int Level = FMath::Clamp(CVarHMDFixedFoveationLevel.GetValueOnAnyThread(), 0, 4);
 
 	VRSImageGenParamsInOut.bGenerateFixedFoveation = Level ? true : false;
 	VRSImageGenParamsInOut.HMDFixedFoveationFullRateCutoff = FIXED_FOVEATION_FULL_RATE_CUTOFFS[Level];
 	VRSImageGenParamsInOut.HMDFixedFoveationHalfRateCutoff = FIXED_FOVEATION_HALF_RATE_CUTOFFS[Level];
+	VRSImageGenParamsInOut.HMDFixedFoveationCenterX = FIXED_FOVEATION_CENTER_X[Level];
+	VRSImageGenParamsInOut.HMDFixedFoveationCenterY = FIXED_FOVEATION_CENTER_Y[Level];
 }
 
 void FVariableRateShadingImageManager::UpdateEyeTrackedFoveationParameters(FVRSImageGenerationParameters& VRSImageGenParamsInOut, const FSceneViewFamily& ViewFamily)
 {
 	VRSImageGenParamsInOut.bGenerateEyeTrackedFoveation = false;
 	VRSImageGenParamsInOut.HMDEyeTrackedFoveationFullRateCutoff = 1.0f;
-	VRSImageGenParamsInOut.HMDEYeTrackedFoveationHalfRateCutoff = 1.0f;
+	VRSImageGenParamsInOut.HMDEyeTrackedFoveationHalfRateCutoff = 1.0f;
 
 	auto EyeTracker = GEngine->EyeTrackingDevice;
 	if (!EyeTracker.IsValid())
