Implements ISpatialUpscaler

diff --git a/Engine/Source/Runtime/Engine/Private/SceneView.cpp b/Engine/Source/Runtime/Engine/Private/SceneView.cpp
--- a/Engine/Source/Runtime/Engine/Private/SceneView.cpp
+++ b/Engine/Source/Runtime/Engine/Private/SceneView.cpp
@@ -2589,7 +2589,9 @@ FSceneViewFamily::FSceneViewFamily(const ConstructionValues& CVS)
 	SecondaryViewFraction(1.0f),
 	SecondaryScreenPercentageMethod(ESecondaryScreenPercentageMethod::LowerPixelDensitySimulation),
 	ScreenPercentageInterface(nullptr),
-	TemporalUpscalerInterface(nullptr)
+	TemporalUpscalerInterface(nullptr),
+	PrimarySpatialUpscalerInterface(nullptr),
+	SecondarySpatialUpscalerInterface(nullptr)
 {
 	// If we do not pass a valid scene pointer then SetWorldTimes must be called to initialized with valid times.
 	ensure(CVS.bTimesSet);
@@ -2679,6 +2681,18 @@ FSceneViewFamily::~FSceneViewFamily()
 	{
 		delete ScreenPercentageInterface;
 	}
+
+	if (PrimarySpatialUpscalerInterface)
+	{
+		// ISpatialUpscaler* is only defined in renderer's private header because no backward compatibility is provided between major version.
+		delete reinterpret_cast<ISceneViewFamilyExtention*>(PrimarySpatialUpscalerInterface);
+	}
+
+	if (SecondarySpatialUpscalerInterface)
+	{
+		// ISpatialUpscaler* is only defined in renderer's private header because no backward compatibility is provided between major version.
+		delete reinterpret_cast<ISceneViewFamilyExtention*>(SecondarySpatialUpscalerInterface);
+	}
 }
 
 ERHIFeatureLevel::Type FSceneViewFamily::GetFeatureLevel() const
diff --git a/Engine/Source/Runtime/Engine/Public/SceneView.h b/Engine/Source/Runtime/Engine/Public/SceneView.h
--- a/Engine/Source/Runtime/Engine/Public/SceneView.h
+++ b/Engine/Source/Runtime/Engine/Public/SceneView.h
@@ -27,6 +27,7 @@ class ISceneViewExtension;
 class FSceneViewFamily;
 class FVolumetricFogViewResources;
 class FIESLightProfileResource;
+class ISpatialUpscaler;
 class ITemporalUpscaler;
 
 // Projection data for a FSceneView
@@ -1406,13 +1407,29 @@ struct FDisplayInternalsData
 //////////////////////////////////////////////////////////////////////////
 
 
+/**
+ * Generic plugin extension that have a lifetime of the FSceneViewFamily
+ */
+class ENGINE_API ISceneViewFamilyExtention
+{
+protected:
+	/** 
+	 * Called by the destructor of the view family.
+	 * Can be called on game or rendering thread.
+	 */
+	virtual ~ISceneViewFamilyExtention() {};
+
+	friend class FSceneViewFamily;
+};
+
+
 /*
  * Game thread and render thread interface that takes care of a FSceneViewFamily's screen percentage.
  *
  * The renderer reserves the right to delete and replace the view family's screen percentage interface
  * for testing purposes with the r.Test.OverrideScreenPercentageInterface CVar.
  */
-class ENGINE_API ISceneViewFamilyScreenPercentage
+class ENGINE_API ISceneViewFamilyScreenPercentage : ISceneViewFamilyExtention
 {
 public:
 	// Sets the minimal and max screen percentage.
@@ -1435,12 +1452,6 @@ public:
 #endif
 
 protected:
-	/** 
-	 * Called by the destructor of the view family.
-	 * Can be called on game or rendering thread.
-	 */
-	virtual ~ISceneViewFamilyScreenPercentage() {};
-
 	/** 
 	 * Method to know the maximum value that can be returned by GetPrimaryResolutionFraction_RenderThread().
 	 * Can be called on game or rendering thread. This should return >= 1 if screen percentage show flag is disabled.
@@ -1763,6 +1774,8 @@ public:
 	{
 		check(ScreenPercentageInterface == nullptr);
 		check(TemporalUpscalerInterface == nullptr);
+		check(PrimarySpatialUpscalerInterface == nullptr);
+		check(SecondarySpatialUpscalerInterface == nullptr);
 	}
 
 
@@ -1778,11 +1791,37 @@ public:
 		return TemporalUpscalerInterface;
 	}
 
+	FORCEINLINE void SetPrimarySpatialUpscalerInterface(ISpatialUpscaler* InSpatialUpscalerInterface)
+	{
+		check(InSpatialUpscalerInterface);
+		checkf(PrimarySpatialUpscalerInterface == nullptr, TEXT("View family already had a primary spatial upscaler assigned."));
+		PrimarySpatialUpscalerInterface = InSpatialUpscalerInterface;
+	}
+
+	FORCEINLINE const ISpatialUpscaler* GetPrimarySpatialUpscalerInterface() const
+	{
+		return PrimarySpatialUpscalerInterface;
+	}
+
+	FORCEINLINE void SetSecondarySpatialUpscalerInterface(ISpatialUpscaler* InSpatialUpscalerInterface)
+	{
+		check(InSpatialUpscalerInterface);
+		checkf(SecondarySpatialUpscalerInterface == nullptr, TEXT("View family already had a secondary spatial upscaler assigned."));
+		SecondarySpatialUpscalerInterface = InSpatialUpscalerInterface;
+	}
+
+	FORCEINLINE const ISpatialUpscaler* GetSecondarySpatialUpscalerInterface() const
+	{
+		return SecondarySpatialUpscalerInterface;
+	}
+
 private:
 	/** Interface to handle screen percentage of the views of the family. */
 	ISceneViewFamilyScreenPercentage* ScreenPercentageInterface;
 
 	const ITemporalUpscaler* TemporalUpscalerInterface;
+	ISpatialUpscaler* PrimarySpatialUpscalerInterface;
+	ISpatialUpscaler* SecondarySpatialUpscalerInterface;
 
 	/** whether the translucency are allowed to render after DOF, if not they will be rendered in standard translucency. */
 	bool bAllowTranslucencyAfterDOF;
diff --git a/Engine/Source/Runtime/Renderer/Private/PostProcess/PostProcessUpscale.cpp b/Engine/Source/Runtime/Renderer/Private/PostProcess/PostProcessUpscale.cpp
--- a/Engine/Source/Runtime/Renderer/Private/PostProcess/PostProcessUpscale.cpp
+++ b/Engine/Source/Runtime/Renderer/Private/PostProcess/PostProcessUpscale.cpp
@@ -166,10 +166,16 @@ EUpscaleMethod GetUpscaleMethod()
 	return static_cast<EUpscaleMethod>(FMath::Clamp(Value, 0, static_cast<int32>(EUpscaleMethod::Gaussian)));
 }
 
-FScreenPassTexture AddUpscalePass(FRDGBuilder& GraphBuilder, const FViewInfo& View, const FUpscaleInputs& Inputs)
+// static
+FScreenPassTexture ISpatialUpscaler::AddDefaultUpscalePass(
+	FRDGBuilder& GraphBuilder,
+	const FViewInfo& View,
+	const FInputs& Inputs,
+	EUpscaleMethod Method,
+	FPaniniProjectionConfig PaniniConfig)
 {
 	check(Inputs.SceneColor.IsValid());
-	check(Inputs.Method != EUpscaleMethod::MAX);
+	check(Method != EUpscaleMethod::MAX);
 	check(Inputs.Stage != EUpscaleStage::MAX);
 
 	FScreenPassRenderTarget Output = Inputs.OverrideOutput;
@@ -201,7 +207,6 @@ FScreenPassTexture AddUpscalePass(FRDGBuilder& GraphBuilder, const FViewInfo& Vi
 	const FScreenPassTextureViewport InputViewport(Inputs.SceneColor);
 	const FScreenPassTextureViewport OutputViewport(Output);
 
-	FPaniniProjectionConfig PaniniConfig = Inputs.PaniniConfig;
 	PaniniConfig.Sanitize();
 
 	const bool bUsePaniniProjection = PaniniConfig.IsEnabled();
@@ -220,7 +225,7 @@ FScreenPassTexture AddUpscalePass(FRDGBuilder& GraphBuilder, const FViewInfo& Vi
 	PassParameters->View = View.ViewUniformBuffer;
 
 	FUpscalePS::FPermutationDomain PixelPermutationVector;
-	PixelPermutationVector.Set<FUpscalePS::FMethodDimension>(Inputs.Method);
+	PixelPermutationVector.Set<FUpscalePS::FMethodDimension>(Method);
 	TShaderMapRef<FUpscalePS> PixelShader(View.ShaderMap, PixelPermutationVector);
 
 	const TCHAR* const StageNames[] = { TEXT("PrimaryToSecondary"), TEXT("PrimaryToOutput"), TEXT("SecondaryToOutput") };
@@ -228,7 +233,7 @@ FScreenPassTexture AddUpscalePass(FRDGBuilder& GraphBuilder, const FViewInfo& Vi
 	const TCHAR* StageName = StageNames[static_cast<uint32>(Inputs.Stage)];
 
 	GraphBuilder.AddPass(
-		RDG_EVENT_NAME("Upscale(%s Method=%d) %dx%d -> %dx%d", StageName, int32(Inputs.Method), Inputs.SceneColor.ViewRect.Width(), Inputs.SceneColor.ViewRect.Height(), Output.ViewRect.Width(), Output.ViewRect.Height()),
+		RDG_EVENT_NAME("Upscale(%s Method=%d) %dx%d -> %dx%d", StageName, int32(Method), Inputs.SceneColor.ViewRect.Width(), Inputs.SceneColor.ViewRect.Height(), Output.ViewRect.Width(), Output.ViewRect.Height()),
 		PassParameters,
 		ERDGPassFlags::Raster,
 		[&View, bUsePaniniProjection, PixelShader, PassParameters, InputViewport, OutputViewport](FRHICommandList& RHICmdList)
diff --git a/Engine/Source/Runtime/Renderer/Private/PostProcess/PostProcessUpscale.h b/Engine/Source/Runtime/Renderer/Private/PostProcess/PostProcessUpscale.h
--- a/Engine/Source/Runtime/Renderer/Private/PostProcess/PostProcessUpscale.h
+++ b/Engine/Source/Runtime/Renderer/Private/PostProcess/PostProcessUpscale.h
@@ -60,22 +60,36 @@ enum class EUpscaleStage
 	MAX
 };
 
-struct FUpscaleInputs
+/** Interface for custom spatial upscaling algorithm meant to be set on the FSceneViewFamily by ISceneViewExtension::BeginRenderViewFamily(). */
+class RENDERER_API ISpatialUpscaler : public ISceneViewFamilyExtention
 {
-	// [Optional] Render to the specified output. If invalid, a new texture is created and returned.
-	FScreenPassRenderTarget OverrideOutput;
+public:
+	struct FInputs
+	{
+		// [Optional] Render to the specified output. If invalid, a new texture is created and returned.
+		FScreenPassRenderTarget OverrideOutput;
 
-	// [Required] The input scene color and view rect.
-	FScreenPassTexture SceneColor;
+		// [Required] The input scene color and view rect.
+		FScreenPassTexture SceneColor;
 
-	// [Required] The method to use when upscaling.
-	EUpscaleMethod Method = EUpscaleMethod::MAX;
+		// Whether this is a secondary upscale to the final view family target.
+		EUpscaleStage Stage = EUpscaleStage::MAX;
+	};
 
-	// [Optional] A configuration used to control Panini projection. Disabled in the default state.
-	FPaniniProjectionConfig PaniniConfig;
+	virtual const TCHAR* GetDebugName() const = 0;
 
-	// Whether this is a secondary upscale to the final view family target.
-	EUpscaleStage Stage = EUpscaleStage::MAX;
-};
+	/** Create a new ISpatialUpscaler interface for a new view family. */
+	virtual ISpatialUpscaler* Fork_GameThread(const class FSceneViewFamily& ViewFamily) const = 0;
 
-FScreenPassTexture AddUpscalePass(FRDGBuilder& GraphBuilder, const FViewInfo& View, const FUpscaleInputs& Inputs);
\ No newline at end of file
+	virtual FScreenPassTexture AddPasses(
+		FRDGBuilder& GraphBuilder,
+		const FViewInfo& View,
+		const FInputs& PassInputs) const = 0;
+
+	static FScreenPassTexture AddDefaultUpscalePass(
+		FRDGBuilder& GraphBuilder,
+		const FViewInfo& View,
+		const FInputs& PassInputs,
+		EUpscaleMethod Method,
+		FPaniniProjectionConfig PaniniConfig);
+};
diff --git a/Engine/Source/Runtime/Renderer/Private/PostProcess/PostProcessing.cpp b/Engine/Source/Runtime/Renderer/Private/PostProcess/PostProcessing.cpp
--- a/Engine/Source/Runtime/Renderer/Private/PostProcess/PostProcessing.cpp
+++ b/Engine/Source/Runtime/Renderer/Private/PostProcess/PostProcessing.cpp
@@ -473,7 +473,7 @@ void AddPostProcessingPasses(FRDGBuilder& GraphBuilder, const FViewInfo& View, c
 	PassSequence.SetEnabled(EPass::HMDDistortion, EngineShowFlags.StereoRendering && EngineShowFlags.HMDDistortion);
 	PassSequence.SetEnabled(EPass::HighResolutionScreenshotMask, IsHighResolutionScreenshotMaskEnabled(View));
 	PassSequence.SetEnabled(EPass::PrimaryUpscale, PaniniConfig.IsEnabled() || (View.PrimaryScreenPercentageMethod == EPrimaryScreenPercentageMethod::SpatialUpscale && PrimaryViewRect.Size() != View.GetSecondaryViewRectSize()));
-	PassSequence.SetEnabled(EPass::SecondaryUpscale, View.RequiresSecondaryUpscale());
+	PassSequence.SetEnabled(EPass::SecondaryUpscale, View.RequiresSecondaryUpscale() || View.Family->GetSecondarySpatialUpscalerInterface() != nullptr);
 
 	const auto GetPostProcessMaterialInputs = [&](FScreenPassTexture InSceneColor)
 	{ 
@@ -1177,27 +1177,66 @@ void AddPostProcessingPasses(FRDGBuilder& GraphBuilder, const FViewInfo& View, c
 
 	if (PassSequence.IsEnabled(EPass::PrimaryUpscale))
 	{
-		FUpscaleInputs PassInputs;
+		ISpatialUpscaler::FInputs PassInputs;
 		PassSequence.AcceptOverrideIfLastPass(EPass::PrimaryUpscale, PassInputs.OverrideOutput);
 		PassInputs.SceneColor = SceneColor;
-		PassInputs.Method = GetUpscaleMethod();
 		PassInputs.Stage = PassSequence.IsEnabled(EPass::SecondaryUpscale) ? EUpscaleStage::PrimaryToSecondary : EUpscaleStage::PrimaryToOutput;
 
-		// Panini projection is handled by the primary upscale pass.
-		PassInputs.PaniniConfig = PaniniConfig;
+		if (const ISpatialUpscaler* CustomUpscaler = View.Family->GetPrimarySpatialUpscalerInterface())
+		{
+			RDG_EVENT_SCOPE(
+				GraphBuilder,
+				"ThirdParty PrimaryUpscale %s %dx%d -> %dx%d",
+				CustomUpscaler->GetDebugName(),
+				SceneColor.ViewRect.Width(), SceneColor.ViewRect.Height(),
+				View.GetSecondaryViewRectSize().X, View.GetSecondaryViewRectSize().Y);
+
+			SceneColor = CustomUpscaler->AddPasses(GraphBuilder, View, PassInputs);
 
-		SceneColor = AddUpscalePass(GraphBuilder, View, PassInputs);
+			if (PassSequence.IsLastPass(EPass::PrimaryUpscale))
+			{
+				check(SceneColor == ViewFamilyOutput);
+			}
+			else
+			{
+				check(SceneColor.ViewRect.Size() == View.GetSecondaryViewRectSize());
+			}
+		}
+		else
+		{
+			EUpscaleMethod Method = GetUpscaleMethod();
+
+			SceneColor = ISpatialUpscaler::AddDefaultUpscalePass(GraphBuilder, View, PassInputs, Method, PaniniConfig);
+		}
 	}
 
 	if (PassSequence.IsEnabled(EPass::SecondaryUpscale))
 	{
-		FUpscaleInputs PassInputs;
+		ISpatialUpscaler::FInputs PassInputs;
 		PassSequence.AcceptOverrideIfLastPass(EPass::SecondaryUpscale, PassInputs.OverrideOutput);
 		PassInputs.SceneColor = SceneColor;
-		PassInputs.Method = View.Family->SecondaryScreenPercentageMethod == ESecondaryScreenPercentageMethod::LowerPixelDensitySimulation ? EUpscaleMethod::SmoothStep : EUpscaleMethod::Nearest;
 		PassInputs.Stage = EUpscaleStage::SecondaryToOutput;
+		
+		if (const ISpatialUpscaler* CustomUpscaler = View.Family->GetSecondarySpatialUpscalerInterface())
+		{
+			RDG_EVENT_SCOPE(
+				GraphBuilder,
+				"ThirdParty SecondaryUpscale %s %dx%d -> %dx%d",
+				CustomUpscaler->GetDebugName(),
+				SceneColor.ViewRect.Width(), SceneColor.ViewRect.Height(),
+				View.UnscaledViewRect.Width(), View.UnscaledViewRect.Height());
 
-		SceneColor = AddUpscalePass(GraphBuilder, View, PassInputs);
+			SceneColor = CustomUpscaler->AddPasses(GraphBuilder, View, PassInputs);
+			check(SceneColor == ViewFamilyOutput);
+		}
+		else
+		{
+			EUpscaleMethod Method = View.Family->SecondaryScreenPercentageMethod == ESecondaryScreenPercentageMethod::LowerPixelDensitySimulation
+				? EUpscaleMethod::SmoothStep
+				: EUpscaleMethod::Nearest;
+
+			SceneColor = ISpatialUpscaler::AddDefaultUpscalePass(GraphBuilder, View, PassInputs, Method, FPaniniProjectionConfig());
+		}
 	}
 }
 
@@ -1250,7 +1289,7 @@ void AddDebugViewPostProcessingPasses(FRDGBuilder& GraphBuilder, const FViewInfo
 	PassSequence.SetEnabled(EPass::TonemapAfter, bTonemapAfter);
 	PassSequence.SetEnabled(EPass::SelectionOutline, GIsEditor);
 	PassSequence.SetEnabled(EPass::PrimaryUpscale, View.ViewRect.Size() != View.GetSecondaryViewRectSize() && View.PrimaryScreenPercentageMethod != EPrimaryScreenPercentageMethod::TemporalUpscale);
-	PassSequence.SetEnabled(EPass::SecondaryUpscale, View.RequiresSecondaryUpscale());
+	PassSequence.SetEnabled(EPass::SecondaryUpscale, View.RequiresSecondaryUpscale() || View.Family->GetSecondarySpatialUpscalerInterface() != nullptr);
 	PassSequence.Finalize();
 
 	if (bTonemapBefore)
@@ -1369,24 +1408,66 @@ void AddDebugViewPostProcessingPasses(FRDGBuilder& GraphBuilder, const FViewInfo
 
 	if (PassSequence.IsEnabled(EPass::PrimaryUpscale))
 	{
-		FUpscaleInputs PassInputs;
+		ISpatialUpscaler::FInputs PassInputs;
 		PassSequence.AcceptOverrideIfLastPass(EPass::PrimaryUpscale, PassInputs.OverrideOutput);
 		PassInputs.SceneColor = SceneColor;
-		PassInputs.Method = GetUpscaleMethod();
 		PassInputs.Stage = PassSequence.IsEnabled(EPass::SecondaryUpscale) ? EUpscaleStage::PrimaryToSecondary : EUpscaleStage::PrimaryToOutput;
 
-		SceneColor = AddUpscalePass(GraphBuilder, View, PassInputs);
+		if (const ISpatialUpscaler* CustomUpscaler = View.Family->GetPrimarySpatialUpscalerInterface())
+		{
+			RDG_EVENT_SCOPE(
+				GraphBuilder,
+				"ThirdParty PrimaryUpscale %s %dx%d -> %dx%d",
+				CustomUpscaler->GetDebugName(),
+				SceneColor.ViewRect.Width(), SceneColor.ViewRect.Height(),
+				View.GetSecondaryViewRectSize().X, View.GetSecondaryViewRectSize().Y);
+
+			SceneColor = CustomUpscaler->AddPasses(GraphBuilder, View, PassInputs);
+
+			if (PassSequence.IsLastPass(EPass::PrimaryUpscale))
+			{
+				check(SceneColor == ViewFamilyOutput);
+			}
+			else
+			{
+				check(SceneColor.ViewRect.Size() == View.GetSecondaryViewRectSize());
+			}
+		}
+		else
+		{
+			EUpscaleMethod Method = GetUpscaleMethod();
+
+			SceneColor = ISpatialUpscaler::AddDefaultUpscalePass(GraphBuilder, View, PassInputs, Method, FPaniniProjectionConfig());
+		}
 	}
 
 	if (PassSequence.IsEnabled(EPass::SecondaryUpscale))
 	{
-		FUpscaleInputs PassInputs;
+		ISpatialUpscaler::FInputs PassInputs;
 		PassSequence.AcceptOverrideIfLastPass(EPass::SecondaryUpscale, PassInputs.OverrideOutput);
 		PassInputs.SceneColor = SceneColor;
-		PassInputs.Method = View.Family->SecondaryScreenPercentageMethod == ESecondaryScreenPercentageMethod::LowerPixelDensitySimulation ? EUpscaleMethod::SmoothStep : EUpscaleMethod::Nearest;
 		PassInputs.Stage = EUpscaleStage::SecondaryToOutput;
 
-		SceneColor = AddUpscalePass(GraphBuilder, View, PassInputs);
+		if (const ISpatialUpscaler* CustomUpscaler = View.Family->GetSecondarySpatialUpscalerInterface())
+		{
+			RDG_EVENT_SCOPE(
+				GraphBuilder,
+				"ThirdParty SecondaryUpscale %s %dx%d -> %dx%d",
+				CustomUpscaler->GetDebugName(),
+				SceneColor.ViewRect.Width(), SceneColor.ViewRect.Height(),
+				View.UnscaledViewRect.Width(), View.UnscaledViewRect.Height());
+
+			SceneColor = CustomUpscaler->AddPasses(GraphBuilder, View, PassInputs);
+			check(SceneColor == ViewFamilyOutput);
+		}
+		else
+		{
+			EUpscaleMethod Method = View.Family->SecondaryScreenPercentageMethod == ESecondaryScreenPercentageMethod::LowerPixelDensitySimulation
+				? EUpscaleMethod::SmoothStep
+				: EUpscaleMethod::Nearest;
+
+			SceneColor = ISpatialUpscaler::AddDefaultUpscalePass(GraphBuilder, View, PassInputs, Method, FPaniniProjectionConfig());
+		}
 	}
 }
 
@@ -2153,15 +2234,36 @@ void AddMobilePostProcessingPasses(FRDGBuilder& GraphBuilder, FScene* Scene, con
 	// Apply ScreenPercentage
 	if (PassSequence.IsEnabled(EPass::PrimaryUpscale))
 	{
-		FUpscaleInputs PassInputs;
+		ISpatialUpscaler::FInputs PassInputs;
 		PassSequence.AcceptOverrideIfLastPass(EPass::PrimaryUpscale, PassInputs.OverrideOutput);
-		PassInputs.Method = EUpscaleMethod::Bilinear;
 		PassInputs.Stage = EUpscaleStage::PrimaryToOutput;
 		PassInputs.SceneColor = SceneColor;
-		PassInputs.PaniniConfig = PaniniConfig;
 		PassInputs.OverrideOutput.LoadAction = View.IsFirstInFamily() ? ERenderTargetLoadAction::EClear : ERenderTargetLoadAction::ELoad;
 
-		SceneColor = AddUpscalePass(GraphBuilder, View, PassInputs);
+		if (const ISpatialUpscaler* CustomUpscaler = View.Family->GetPrimarySpatialUpscalerInterface())
+		{
+			RDG_EVENT_SCOPE(
+				GraphBuilder,
+				"ThirdParty PrimaryUpscale %s %dx%d -> %dx%d",
+				CustomUpscaler->GetDebugName(),
+				SceneColor.ViewRect.Width(), SceneColor.ViewRect.Height(),
+				View.UnscaledViewRect.Width(), View.UnscaledViewRect.Height());
+
+			SceneColor = CustomUpscaler->AddPasses(GraphBuilder, View, PassInputs);
+
+			if (PassSequence.IsLastPass(EPass::PrimaryUpscale))
+			{
+				check(SceneColor == ViewFamilyOutput);
+			}
+			else
+			{
+				check(SceneColor.ViewRect.Size() == View.UnscaledViewRect.Size());
+			}
+		}
+		else
+		{
+			SceneColor = ISpatialUpscaler::AddDefaultUpscalePass(GraphBuilder, View, PassInputs, EUpscaleMethod::Bilinear, PaniniConfig);
+		}
 	}
 
 	if (PassSequence.IsEnabled(EPass::Visualize))
diff --git a/Engine/Source/Runtime/Renderer/Private/SceneRendering.cpp b/Engine/Source/Runtime/Renderer/Private/SceneRendering.cpp
--- a/Engine/Source/Runtime/Renderer/Private/SceneRendering.cpp
+++ b/Engine/Source/Runtime/Renderer/Private/SceneRendering.cpp
@@ -24,6 +24,7 @@
 #include "PostProcess/PostProcessEyeAdaptation.h"
 #include "PostProcess/PostProcessSubsurface.h"
 #include "PostProcess/TemporalAA.h"
+#include "PostProcess/PostProcessUpscale.h"
 #include "PostProcess/PostProcessing.h"
 #include "CompositionLighting/CompositionLighting.h"
 #include "LegacyScreenPercentageDriver.h"
@@ -2367,11 +2368,25 @@ FSceneRenderer::FSceneRenderer(const FSceneViewFamily* InViewFamily,FHitProxyCon
 	ensureMsgf(!(ViewFamily.EngineShowFlags.ScreenPercentage && !ViewFamily.SupportsScreenPercentage()),
 		TEXT("Screen percentage is not supported, but show flag was incorectly set to true."));
 
-	// Fork the view family.
+	// Fork the plugin interfaces of the view family.
 	{
-		check(InViewFamily->ScreenPercentageInterface);
-		ViewFamily.ScreenPercentageInterface = nullptr;
-		ViewFamily.SetScreenPercentageInterface(InViewFamily->ScreenPercentageInterface->Fork_GameThread(ViewFamily));
+		{
+			check(InViewFamily->ScreenPercentageInterface);
+			ViewFamily.ScreenPercentageInterface = nullptr;
+			ViewFamily.SetScreenPercentageInterface(InViewFamily->ScreenPercentageInterface->Fork_GameThread(ViewFamily));
+		}
+
+		if (ViewFamily.PrimarySpatialUpscalerInterface)
+		{
+			ViewFamily.PrimarySpatialUpscalerInterface = nullptr;
+			ViewFamily.SetPrimarySpatialUpscalerInterface(InViewFamily->PrimarySpatialUpscalerInterface->Fork_GameThread(ViewFamily));
+		}
+
+		if (ViewFamily.SecondarySpatialUpscalerInterface)
+		{
+			ViewFamily.SecondarySpatialUpscalerInterface = nullptr;
+			ViewFamily.SetSecondarySpatialUpscalerInterface(InViewFamily->SecondarySpatialUpscalerInterface->Fork_GameThread(ViewFamily));
+		}
 	}
 
 	#if !UE_BUILD_SHIPPING
@@ -3997,6 +4012,10 @@ void FRendererModule::BeginRenderingViewFamily(FCanvas* Canvas, FSceneViewFamily
 	}
 #endif // !(UE_BUILD_SHIPPING || UE_BUILD_TEST)
 
+	// Force the spatial upscaler to be set no earlier than ISceneViewExtension::BeginRenderViewFamily();
+	check(ViewFamily->GetPrimarySpatialUpscalerInterface() == nullptr);
+	check(ViewFamily->GetSecondarySpatialUpscalerInterface() == nullptr);
+
 	for (int ViewExt = 0; ViewExt < ViewFamily->ViewExtensions.Num(); ViewExt++)
 	{
 		ViewFamily->ViewExtensions[ViewExt]->BeginRenderViewFamily(*ViewFamily);
