Fix for bilinear sampling of light shaft texture in exponential fog draw step sampling outside the valid region due to dynamic resolution

diff --git a/Engine/Shaders/Private/HeightFogPixelShader.usf b/Engine/Shaders/Private/HeightFogPixelShader.usf
--- a/Engine/Shaders/Private/HeightFogPixelShader.usf
+++ b/Engine/Shaders/Private/HeightFogPixelShader.usf
@@ -16,6 +16,7 @@ float bOnlyOnRenderedOpaque;
 float bUseLinearDepthTexture;
 
 float4 LinearDepthTextureMinMaxUV;
+float4 OcclusionTextureMinMaxUV;
 
 Texture2D OcclusionTexture;
 SamplerState OcclusionSampler;
@@ -64,7 +65,8 @@ void ExponentialPixelMain(
 	float4 HeightFogInscatteringAndOpacity = CalculateHeightFog(WorldPositionRelativeToCamera);
 
 	float4 FogInscatteringAndOpacity = CombineVolumetricFog(HeightFogInscatteringAndOpacity, VolumeUV, 0);
-	float LightShaftMask = Texture2DSample(OcclusionTexture, OcclusionSampler, TexCoord).x;
+	float2 OcclusionTexCoord = clamp(TexCoord, OcclusionTextureMinMaxUV.xy, OcclusionTextureMinMaxUV.zw);
+	float LightShaftMask = Texture2DSample(OcclusionTexture, OcclusionSampler, OcclusionTexCoord).x;
 	FogInscatteringAndOpacity.rgb *= LightShaftMask;
 
 	FogInscatteringAndOpacity.rgb *= View.PreExposure;
diff --git a/Engine/Source/Runtime/Renderer/Private/DeferredShadingRenderer.cpp b/Engine/Source/Runtime/Renderer/Private/DeferredShadingRenderer.cpp
--- a/Engine/Source/Runtime/Renderer/Private/DeferredShadingRenderer.cpp
+++ b/Engine/Source/Runtime/Renderer/Private/DeferredShadingRenderer.cpp
@@ -2927,7 +2927,8 @@ void FDeferredShadingSceneRenderer::Render(FRDGBuilder& GraphBuilder)
 	{
 		RDG_CSV_STAT_EXCLUSIVE_SCOPE(GraphBuilder, RenderFog);
 		SCOPE_CYCLE_COUNTER(STAT_FDeferredShadingSceneRenderer_RenderFog);
-		RenderFog(GraphBuilder, SceneTextures, LightShaftOcclusionTexture);
+		const FIntPoint LightShaftOcclusionTextureValidExtent = GetDownscaledExtent(SceneTextures.Config.Extent, GetLightShaftDownsampleFactor());
+		RenderFog(GraphBuilder, SceneTextures, LightShaftOcclusionTexture, LightShaftOcclusionTextureValidExtent);
 	}
 
 	// After the height fog, Draw volumetric clouds (having fog applied on them already) when using per pixel tracing,
diff --git a/Engine/Source/Runtime/Renderer/Private/DeferredShadingRenderer.h b/Engine/Source/Runtime/Renderer/Private/DeferredShadingRenderer.h
--- a/Engine/Source/Runtime/Renderer/Private/DeferredShadingRenderer.h
+++ b/Engine/Source/Runtime/Renderer/Private/DeferredShadingRenderer.h
@@ -419,7 +419,8 @@ private:
 	void RenderFog(
 		FRDGBuilder& GraphBuilder,
 		const FMinimalSceneTextures& SceneTextures,
-		FRDGTextureRef LightShaftOcclusionTexture);
+		FRDGTextureRef LightShaftOcclusionTexture,
+		const FIntPoint& LightShaftOcclusionTextureExtent);
 
 	void RenderUnderWaterFog(
 		FRDGBuilder& GraphBuilder,
diff --git a/Engine/Source/Runtime/Renderer/Private/FogRendering.cpp b/Engine/Source/Runtime/Renderer/Private/FogRendering.cpp
--- a/Engine/Source/Runtime/Renderer/Private/FogRendering.cpp
+++ b/Engine/Source/Runtime/Renderer/Private/FogRendering.cpp
@@ -59,6 +59,7 @@ struct FHeightFogRenderingParameters
 	FIntRect ViewRect;
 	float LinearDepthReadScale = 1.0f;
 	FVector4 LinearDepthMinMaxUV;
+	FVector4 LightShaftOcclusionMinMaxUV;
 };
 
 void SetupFogUniformParameters(FRDGBuilder& GraphBuilder, const FViewInfo& View, FFogUniformParameters& OutParameters)
@@ -204,6 +205,7 @@ public:
 		bUseLinearDepthTexture.Bind(Initializer.ParameterMap, TEXT("bUseLinearDepthTexture"));
 		UpsampleJitterMultiplier.Bind(Initializer.ParameterMap, TEXT("UpsampleJitterMultiplier"));
 		LinearDepthTextureMinMaxUV.Bind(Initializer.ParameterMap, TEXT("LinearDepthTextureMinMaxUV"));
+		LightShaftOcclusionMinMaxUV.Bind(Initializer.ParameterMap, TEXT("OcclusionTextureMinMaxUV"));
 	}
 
 	void SetParameters(FRHICommandList& RHICmdList, const FViewInfo& View, const FHeightFogRenderingParameters& Params, FRHIUniformBuffer* FogUniformBuffer)
@@ -241,6 +243,7 @@ public:
 		SetShaderValue(RHICmdList, RHICmdList.GetBoundPixelShader(), LinearDepthTextureMinMaxUV, Params.LinearDepthMinMaxUV);
 		extern int32 GVolumetricFogGridPixelSize;
 		SetShaderValue(RHICmdList, RHICmdList.GetBoundPixelShader(), UpsampleJitterMultiplier, CVarUpsampleJitterMultiplier.GetValueOnRenderThread() * GVolumetricFogGridPixelSize);
+		SetShaderValue(RHICmdList, RHICmdList.GetBoundPixelShader(), LightShaftOcclusionMinMaxUV, Params.LightShaftOcclusionMinMaxUV);
 	}
 
 private:
@@ -252,6 +255,7 @@ private:
 	LAYOUT_FIELD(FShaderParameter, bUseLinearDepthTexture);
 	LAYOUT_FIELD(FShaderParameter, UpsampleJitterMultiplier);
 	LAYOUT_FIELD(FShaderParameter, LinearDepthTextureMinMaxUV);
+	LAYOUT_FIELD(FShaderParameter, LightShaftOcclusionMinMaxUV);
 };
 
 IMPLEMENT_SHADER_TYPE(template<>,TExponentialHeightFogPS<EHeightFogFeature::HeightFog>,TEXT("/Engine/Private/HeightFogPixelShader.usf"), TEXT("ExponentialPixelMain"),SF_Pixel)
@@ -491,7 +495,7 @@ END_SHADER_PARAMETER_STRUCT()
 void FDeferredShadingSceneRenderer::RenderFog(
 	FRDGBuilder& GraphBuilder,
 	const FMinimalSceneTextures& SceneTextures,
-	FRDGTextureRef LightShaftOcclusionTexture)
+	FRDGTextureRef LightShaftOcclusionTexture, const FIntPoint& LightShaftOcclusionTextureExtent)
 {
 	if (Scene->ExponentialFogs.Num() > 0 
 		// Fog must be done in the base pass for MSAA to work
@@ -519,13 +523,23 @@ void FDeferredShadingSceneRenderer::RenderFog(
 				PassParameters->RenderTargets[0] = FRenderTargetBinding(SceneTextures.Color.Target, ERenderTargetLoadAction::ELoad);
 				PassParameters->RenderTargets.DepthStencil = FDepthStencilBinding(SceneTextures.Depth.Target, ERenderTargetLoadAction::ELoad, ERenderTargetLoadAction::ELoad, FExclusiveDepthStencil::DepthRead_StencilWrite);
 
-				GraphBuilder.AddPass(RDG_EVENT_NAME("Fog"), PassParameters, ERDGPassFlags::Raster, [this, &View, bShouldRenderVolumetricFog, LightShaftOcclusionTexture, FogUniformBuffer](FRHICommandList& RHICmdList)
+				GraphBuilder.AddPass(RDG_EVENT_NAME("Fog"), PassParameters, ERDGPassFlags::Raster, [this, &View, bShouldRenderVolumetricFog, LightShaftOcclusionTexture, LightShaftOcclusionTextureExtent, FogUniformBuffer](FRHICommandList& RHICmdList)
 				{
 					FHeightFogRenderingParameters Parameters;
 					Parameters.ViewRect = View.ViewRect;
 					if (LightShaftOcclusionTexture)
 					{
 						Parameters.LightShaftOcclusionRHI = LightShaftOcclusionTexture->GetRHI();
+
+						{
+							const float InvLightShaftOcclusionTextureExtentX = 1.0f / LightShaftOcclusionTextureExtent.X;
+							const float InvLightShaftOcclusionTextureExtentY = 1.0f / LightShaftOcclusionTextureExtent.Y;
+							Parameters.LightShaftOcclusionMinMaxUV = FVector4(
+								0.5f * InvLightShaftOcclusionTextureExtentX,
+								0.5f * InvLightShaftOcclusionTextureExtentY,
+								(LightShaftOcclusionTextureExtent.X - 0.5f) * InvLightShaftOcclusionTextureExtentX,
+								(LightShaftOcclusionTextureExtent.Y - 0.5f) * InvLightShaftOcclusionTextureExtentY);
+						}
 					}
 					RenderViewFog(RHICmdList, View, bShouldRenderVolumetricFog, Parameters, FogUniformBuffer->GetRHI());
 				});
diff --git a/Engine/Source/Runtime/Renderer/Private/LightShaftRendering.h b/Engine/Source/Runtime/Renderer/Private/LightShaftRendering.h
--- a/Engine/Source/Runtime/Renderer/Private/LightShaftRendering.h
+++ b/Engine/Source/Runtime/Renderer/Private/LightShaftRendering.h
@@ -24,4 +24,7 @@ extern FMobileLightShaftInfo GetMobileLightShaftInfo(const FViewInfo& View, cons
 extern bool ShouldRenderLightShafts(const FSceneViewFamily& ViewFamily);
 
 // Returns whether the light proxy is eligible for light shaft rendering. Assumes light shafts are enabled.
-extern bool ShouldRenderLightShaftsForLight(const FViewInfo& View, const FLightSceneProxy& LightSceneProxy);
\ No newline at end of file
+extern bool ShouldRenderLightShaftsForLight(const FViewInfo& View, const FLightSceneProxy& LightSceneProxy);
+
+// Returns the current downsample factor for the light shaft render target.
+extern int32 GetLightShaftDownsampleFactor();
\ No newline at end of file
