Convert MaxFP16Depth to use lock / unlock to avoid RDG tracking cost.

diff --git a/Engine/Source/Runtime/Renderer/Private/SystemTextures.cpp b/Engine/Source/Runtime/Renderer/Private/SystemTextures.cpp
--- a/Engine/Source/Runtime/Renderer/Private/SystemTextures.cpp
+++ b/Engine/Source/Runtime/Renderer/Private/SystemTextures.cpp
@@ -175,16 +175,10 @@ void FSystemTextures::InitializeCommonTextures(FRHICommandListImmediate& RHICmdL
 
 	if (GPixelFormats[PF_FloatRGBA].Supported)
 	{
-		FPooledRenderTargetDesc Desc(FPooledRenderTargetDesc::Create2DDesc(FIntPoint(1, 1), PF_FloatRGBA, FClearValueBinding(FLinearColor(65500.0f, 65500.0f, 65500.0f, 65500.0f)), TexCreate_HideInVisualizeTexture, TexCreate_RenderTargetable | TexCreate_NoFastClear | TexCreate_ShaderResource, false));
-		Desc.AutoWritable = false;
-		GRenderTargetPool.FindFreeElement(RHICmdList, Desc, MaxFP16Depth, TEXT("MaxFP16Depth"), ERenderTargetTransience::NonTransient);
-
-		RHICmdList.Transition(FRHITransitionInfo(MaxFP16Depth->GetRenderTargetItem().TargetableTexture, ERHIAccess::SRVMask, ERHIAccess::RTV));
-
-		FRHIRenderPassInfo RPInfo(MaxFP16Depth->GetRenderTargetItem().TargetableTexture, ERenderTargetActions::Clear_Store);
-		RHICmdList.BeginRenderPass(RPInfo, TEXT("MaxFP16Depth"));
-		RHICmdList.EndRenderPass();
-		RHICmdList.CopyToResolveTarget(MaxFP16Depth->GetRenderTargetItem().TargetableTexture, MaxFP16Depth->GetRenderTargetItem().ShaderResourceTexture, FResolveParams());
+		FRHIResourceCreateInfo CreateInfo(TEXT("MaxFP16Depth"));
+		FTexture2DRHIRef Texture2D = RHICreateTexture2D(1, 1, PF_FloatRGBA, 1, 1, TexCreate_ShaderResource, CreateInfo);
+		SetDummyTextureData<FFloat16Color>(Texture2D, FFloat16Color(FLinearColor(65500.0f, 65500.0f, 65500.0f, 65500.0f)));
+		MaxFP16Depth = CreateRenderTarget(Texture2D, CreateInfo.DebugName);
 	}
 
 	{
