Move ConeTraceGlobal pass to use the non-immediate command list.

diff --git a/Engine/Source/Runtime/Renderer/Private/DistanceFieldScreenGridLighting.cpp b/Engine/Source/Runtime/Renderer/Private/DistanceFieldScreenGridLighting.cpp
--- a/Engine/Source/Runtime/Renderer/Private/DistanceFieldScreenGridLighting.cpp
+++ b/Engine/Source/Runtime/Renderer/Private/DistanceFieldScreenGridLighting.cpp
@@ -455,7 +455,7 @@ void FDeferredShadingSceneRenderer::RenderDistanceFieldAOScreenGrid(
 			RDG_EVENT_NAME("ConeTraceGlobal"),
 			PassParameters,
 			ERDGPassFlags::Compute,
-			[PassParameters, ComputeShader, &View, Parameters, DistanceFieldNormal, TileListGroupSize, bUseObjectDistanceField](FRHICommandListImmediate& RHICmdList)
+			[PassParameters, ComputeShader, &View, Parameters, DistanceFieldNormal, TileListGroupSize, bUseObjectDistanceField](FRHICommandList& RHICmdList)
 			{
 				const uint32 GroupSizeX = FMath::DivideAndRoundUp(View.ViewRect.Size().X / GAODownsampleFactor / GConeTraceDownsampleFactor, GConeTraceGlobalDFTileSize);
 				const uint32 GroupSizeY = FMath::DivideAndRoundUp(View.ViewRect.Size().Y / GAODownsampleFactor / GConeTraceDownsampleFactor, GConeTraceGlobalDFTileSize);
@@ -495,7 +495,7 @@ void FDeferredShadingSceneRenderer::RenderDistanceFieldAOScreenGrid(
 			RDG_EVENT_NAME("ConeTraceObjects"),
 			PassParameters,
 			ERDGPassFlags::Compute,
-			[PassParameters, ComputeShader, &View, Parameters, DistanceFieldNormal, bUseGlobalDistanceField, ObjectTilesIndirectArguments](FRHICommandListImmediate& RHICmdList)
+			[PassParameters, ComputeShader, &View, Parameters, DistanceFieldNormal, bUseGlobalDistanceField, ObjectTilesIndirectArguments](FRHICommandList& RHICmdList)
 			{
 				RHICmdList.SetComputeShader(ComputeShader.GetComputeShader());
 
