Fix CSM not resolving on switch (mobile renderer)

diff --git a/Engine/Source/Runtime/Renderer/Private/ShadowSetup.cpp b/Engine/Source/Runtime/Renderer/Private/ShadowSetup.cpp
--- a/Engine/Source/Runtime/Renderer/Private/ShadowSetup.cpp
+++ b/Engine/Source/Runtime/Renderer/Private/ShadowSetup.cpp
@@ -4834,8 +4834,8 @@ void FSceneRenderer::AllocateShadowDepthTargets(FRHICommandListImmediate& RHICmd
 				}
 
 				bool bNeedsProjection = ProjectedShadowInfo->CacheMode != SDCM_StaticPrimitivesOnly
-					// Filter out everything but PerObjectOpaqueShadows & Distance field shadows for ES31
-					&& (!bMobile || ProjectedShadowInfo->bPerObjectOpaqueShadow || ProjectedShadowInfo->bRayTracedDistanceField);
+					//// Filter out everything but PerObjectOpaqueShadows & Distance field shadows for ES31
+					&& (!bMobile || ProjectedShadowInfo->bPerObjectOpaqueShadow || ProjectedShadowInfo->bRayTracedDistanceField || ProjectedShadowInfo->bWholeSceneShadow);
 
 				if (bNeedsProjection)
 				{
