Permanently enable r.DoInitViewsLightingAfterPrepass.

diff --git a/Engine/Source/Runtime/Renderer/Private/DeferredShadingRenderer.cpp b/Engine/Source/Runtime/Renderer/Private/DeferredShadingRenderer.cpp
--- a/Engine/Source/Runtime/Renderer/Private/DeferredShadingRenderer.cpp
+++ b/Engine/Source/Runtime/Renderer/Private/DeferredShadingRenderer.cpp
@@ -1934,10 +1934,9 @@ void FDeferredShadingSceneRenderer::Render(FRDGBuilder& GraphBuilder)
 
 	FInstanceCullingManager& InstanceCullingManager = *GraphBuilder.AllocObject<FInstanceCullingManager>(Scene->GPUScene.IsEnabled());
 
-	bool bDoInitViewsAfterPrepass = false;
 	{
 		RDG_GPU_STAT_SCOPE(GraphBuilder, VisibilityCommands);
-		bDoInitViewsAfterPrepass = InitViews(GraphBuilder, SceneTexturesConfig, BasePassDepthStencilAccess, ILCTaskData, InstanceCullingManager);
+		InitViews(GraphBuilder, SceneTexturesConfig, BasePassDepthStencilAccess, ILCTaskData, InstanceCullingManager);
 	}
 
 	// Compute & commit the final state of the entire dependency topology of the renderer.
@@ -2014,11 +2013,6 @@ void FDeferredShadingSceneRenderer::Render(FRDGBuilder& GraphBuilder)
 			InstanceCullingManager.BeginDeferredCulling(GraphBuilder, Scene->GPUScene);
 		}
 
-		if (!bDoInitViewsAfterPrepass)
-		{
-			PrepareDistanceFieldScene(GraphBuilder, false);
-		}
-
 		if (Views.Num() > 0)
 		{
 			FViewInfo& View = Views[0];
@@ -2096,13 +2090,6 @@ void FDeferredShadingSceneRenderer::Render(FRDGBuilder& GraphBuilder)
 		DynamicIndexBufferForInitViews.Commit();
 		DynamicVertexBufferForInitViews.Commit();
 		DynamicReadBufferForInitViews.Commit();
-
-		if (!bDoInitViewsAfterPrepass)
-		{
-			DynamicVertexBufferForInitShadows.Commit();
-			DynamicIndexBufferForInitShadows.Commit();
-			DynamicReadBufferForInitShadows.Commit();
-		}
 	}
 
 	// Strata initialisation is always run even when not enabled.
@@ -2209,7 +2196,6 @@ void FDeferredShadingSceneRenderer::Render(FRDGBuilder& GraphBuilder)
 		}
 	}
 
-	if (bDoInitViewsAfterPrepass)
 	{
 		GraphBuilder.Drain();
 
diff --git a/Engine/Source/Runtime/Renderer/Private/DeferredShadingRenderer.h b/Engine/Source/Runtime/Renderer/Private/DeferredShadingRenderer.h
--- a/Engine/Source/Runtime/Renderer/Private/DeferredShadingRenderer.h
+++ b/Engine/Source/Runtime/Renderer/Private/DeferredShadingRenderer.h
@@ -377,7 +377,7 @@ private:
 	void PreVisibilityFrameSetup(FRDGBuilder& GraphBuilder, const FSceneTexturesConfig& SceneTexturesConfig);
 
 	/** Determines which primitives are visible for each view. */
-	bool InitViews(FRDGBuilder& GraphBuilder, const FSceneTexturesConfig& SceneTexturesConfig, FExclusiveDepthStencil::Type BasePassDepthStencilAccess, struct FILCUpdatePrimTaskData& ILCTaskData, FInstanceCullingManager& InstanceCullingManager);
+	void InitViews(FRDGBuilder& GraphBuilder, const FSceneTexturesConfig& SceneTexturesConfig, FExclusiveDepthStencil::Type BasePassDepthStencilAccess, struct FILCUpdatePrimTaskData& ILCTaskData, FInstanceCullingManager& InstanceCullingManager);
 
 	void InitViewsPossiblyAfterPrepass(FRDGBuilder& GraphBuilder, struct FILCUpdatePrimTaskData& ILCTaskData, FInstanceCullingManager& InstanceCullingManager);
 	void BeginUpdateLumenSceneTasks(FRDGBuilder& GraphBuilder);
diff --git a/Engine/Source/Runtime/Renderer/Private/SceneHitProxyRendering.cpp b/Engine/Source/Runtime/Renderer/Private/SceneHitProxyRendering.cpp
--- a/Engine/Source/Runtime/Renderer/Private/SceneHitProxyRendering.cpp
+++ b/Engine/Source/Runtime/Renderer/Private/SceneHitProxyRendering.cpp
@@ -596,10 +596,8 @@ void FDeferredShadingSceneRenderer::RenderHitProxies(FRDGBuilder& GraphBuilder)
 	// Find the visible primitives.
 	{
 		FILCUpdatePrimTaskData ILCTaskData;
-		if (InitViews(GraphBuilder, SceneTexturesConfig, FExclusiveDepthStencil::DepthWrite_StencilWrite, ILCTaskData, InstanceCullingManager))
-		{
-			InitViewsPossiblyAfterPrepass(GraphBuilder, ILCTaskData, InstanceCullingManager);
-		}
+		InitViews(GraphBuilder, SceneTexturesConfig, FExclusiveDepthStencil::DepthWrite_StencilWrite, ILCTaskData, InstanceCullingManager);
+		InitViewsPossiblyAfterPrepass(GraphBuilder, ILCTaskData, InstanceCullingManager);
 	}
 
 	extern TSet<IPersistentViewUniformBufferExtension*> PersistentViewUniformBufferExtensions;
diff --git a/Engine/Source/Runtime/Renderer/Private/SceneVisibility.cpp b/Engine/Source/Runtime/Renderer/Private/SceneVisibility.cpp
--- a/Engine/Source/Runtime/Renderer/Private/SceneVisibility.cpp
+++ b/Engine/Source/Runtime/Renderer/Private/SceneVisibility.cpp
@@ -151,14 +151,6 @@ static FAutoConsoleVariableRef CVarILCUpdatePrimitivesTask(
 	ECVF_RenderThreadSafe
 	);
 
-static int32 GDoInitViewsLightingAfterPrepass = 0;
-static FAutoConsoleVariableRef CVarDoInitViewsLightingAfterPrepass(
-	TEXT("r.DoInitViewsLightingAfterPrepass"),
-	GDoInitViewsLightingAfterPrepass,
-	TEXT("Delays the lighting part of InitViews until after the prepass. This improves the threading throughput and gets the prepass to the GPU ASAP. Experimental options; has an unknown race."),
-	ECVF_RenderThreadSafe
-	);
-
 static int32 GFramesNotOcclusionTestedToExpandBBoxes = 5;
 static FAutoConsoleVariableRef CVarFramesNotOcclusionTestedToExpandBBoxes(
 	TEXT("r.GFramesNotOcclusionTestedToExpandBBoxes"),
@@ -4533,7 +4525,7 @@ void FDeferredShadingSceneRenderer::PreVisibilityFrameSetup(FRDGBuilder& GraphBu
  * Initialize scene's views.
  * Check visibility, build visible mesh commands, etc.
  */
-bool FDeferredShadingSceneRenderer::InitViews(FRDGBuilder& GraphBuilder, const FSceneTexturesConfig& SceneTexturesConfig, FExclusiveDepthStencil::Type BasePassDepthStencilAccess, struct FILCUpdatePrimTaskData& ILCTaskData, FInstanceCullingManager& InstanceCullingManager)
+void FDeferredShadingSceneRenderer::InitViews(FRDGBuilder& GraphBuilder, const FSceneTexturesConfig& SceneTexturesConfig, FExclusiveDepthStencil::Type BasePassDepthStencilAccess, struct FILCUpdatePrimTaskData& ILCTaskData, FInstanceCullingManager& InstanceCullingManager)
 {
 	SCOPED_NAMED_EVENT(FDeferredShadingSceneRenderer_InitViews, FColor::Emerald);
 	SCOPE_CYCLE_COUNTER(STAT_InitViewsTime);
@@ -4598,13 +4590,6 @@ bool FDeferredShadingSceneRenderer::InitViews(FRDGBuilder& GraphBuilder, const F
 		AverageViewPosition += View.ViewMatrices.GetViewOrigin() / Views.Num();
 	}
 
-	bool bDoInitViewAftersPrepass = !!GDoInitViewsLightingAfterPrepass;
-
-	if (!bDoInitViewAftersPrepass)
-	{
-		InitViewsPossiblyAfterPrepass(GraphBuilder, ILCTaskData, InstanceCullingManager);
-	}
-
 	{
 		QUICK_SCOPE_CYCLE_COUNTER(STAT_InitViews_InitRHIResources);
 		// initialize per-view uniform buffer.
@@ -4651,8 +4636,6 @@ bool FDeferredShadingSceneRenderer::InitViews(FRDGBuilder& GraphBuilder, const F
 		QUICK_SCOPE_CYCLE_COUNTER(STAT_InitViews_OnStartRender);
 		OnStartRender(RHICmdList);
 	}
-
-	return bDoInitViewAftersPrepass;
 }
 
 void FSceneRenderer::SetupSceneReflectionCaptureBuffer(FRHICommandListImmediate& RHICmdList)
