Implements r.Bloom.HalfResolutionFFT=2 for quarter-res FFT Bloom

Original author: Bennett Bernardoni

diff --git a/Engine/Config/Linux/DataDrivenPlatformInfo.ini b/Engine/Config/Linux/DataDrivenPlatformInfo.ini
--- a/Engine/Config/Linux/DataDrivenPlatformInfo.ini
+++ b/Engine/Config/Linux/DataDrivenPlatformInfo.ini
@@ -119,6 +119,7 @@ bRequiresDisableForwardLocalLights = true
 
 bCompileSignalProcessingPipeline = true
 bSupportsGen5TemporalAA=true
+bSupportsFFTBloom=true
 
 ; //#todo-vulkanddpi
 bSupportsPerPixelDBufferMask = false
diff --git a/Engine/Config/Mac/DataDrivenPlatformInfo.ini b/Engine/Config/Mac/DataDrivenPlatformInfo.ini
--- a/Engine/Config/Mac/DataDrivenPlatformInfo.ini
+++ b/Engine/Config/Mac/DataDrivenPlatformInfo.ini
@@ -26,3 +26,4 @@ PlatformGroupName=Desktop
 [ShaderPlatform METAL_SM5]
 bSupportsByteBufferComputeShaders=true
 bSupportsGen5TemporalAA=true
+bSupportsFFTBloom=true
diff --git a/Engine/Config/Windows/DataDrivenPlatformInfo.ini b/Engine/Config/Windows/DataDrivenPlatformInfo.ini
--- a/Engine/Config/Windows/DataDrivenPlatformInfo.ini
+++ b/Engine/Config/Windows/DataDrivenPlatformInfo.ini
@@ -51,6 +51,7 @@ bSupportsVariableRateShading=true
 bSupportsGPUSkinCache=true
 bSupportsGen5TemporalAA=true
 bSupportsHairStrandGeometry=true
+bSupportsFFTBloom=true
 
 [PreviewPlatform PCD3D_SM5]
 PlatformName=PCD3D_SM5
@@ -81,6 +82,7 @@ bSupportsGen5TemporalAA=true
 bSupportsMeshShaders=true
 FriendlyName=LOCTEXT("FriendlyShaderPlatformName_PCD3D_SM6", "Shader Model 6")
 bSupportsHairStrandGeometry=true
+bSupportsFFTBloom=true
 
 [PreviewPlatform PCD3D_SM6]
 PlatformName=PCD3D_SM6
@@ -197,6 +199,7 @@ bSupportsVariableRateShading = false
 NumberOfComputeThreads = 64
 
 bSupportsHairStrandGeometry=true
+bSupportsFFTBloom=true
 
 [ShaderPlatform PCD3D_ES3_1]
 Language=D3D
diff --git a/Engine/Source/Runtime/RHI/Private/RHI.cpp b/Engine/Source/Runtime/RHI/Private/RHI.cpp
--- a/Engine/Source/Runtime/RHI/Private/RHI.cpp
+++ b/Engine/Source/Runtime/RHI/Private/RHI.cpp
@@ -2266,6 +2266,7 @@ void FGenericDataDrivenShaderPlatformInfo::ParseDataDrivenShaderInfo(const FConf
 	GET_SECTION_BOOL_HELPER(bRequiresReverseCullingOnMobile);
 	GET_SECTION_BOOL_HELPER(bOverrideFMaterial_NeedsGBufferEnabled);
 	GET_SECTION_BOOL_HELPER(bSupportsMobileDistanceField);
+	GET_SECTION_BOOL_HELPER(bSupportsFFTBloom);
 #undef GET_SECTION_BOOL_HELPER
 #undef GET_SECTION_INT_HELPER
 
diff --git a/Engine/Source/Runtime/RHI/Public/RHIDefinitions.h b/Engine/Source/Runtime/RHI/Public/RHIDefinitions.h
--- a/Engine/Source/Runtime/RHI/Public/RHIDefinitions.h
+++ b/Engine/Source/Runtime/RHI/Public/RHIDefinitions.h
@@ -385,6 +385,7 @@ class RHI_API FGenericDataDrivenShaderPlatformInfo
 	uint32 bRequiresReverseCullingOnMobile : 1;
 	uint32 bOverrideFMaterial_NeedsGBufferEnabled : 1;
 	uint32 bSupportsMobileDistanceField : 1;
+	uint32 bSupportsFFTBloom : 1;
 
 		
 #if WITH_EDITOR
@@ -757,6 +758,11 @@ public:
 		return Infos[Platform].bSupportsMobileDistanceField;
 	}
 
+	static FORCEINLINE_DEBUGGABLE const bool GetSupportsFFTBloom(const FStaticShaderPlatform Platform)
+	{
+		return Infos[Platform].bSupportsFFTBloom;
+	}
+
 #if WITH_EDITOR
 	static FORCEINLINE_DEBUGGABLE FText GetFriendlyName(const FStaticShaderPlatform Platform)
 	{
diff --git a/Engine/Source/Runtime/Renderer/Private/GPUFastFourierTransform.cpp b/Engine/Source/Runtime/Renderer/Private/GPUFastFourierTransform.cpp
--- a/Engine/Source/Runtime/Renderer/Private/GPUFastFourierTransform.cpp
+++ b/Engine/Source/Runtime/Renderer/Private/GPUFastFourierTransform.cpp
@@ -6,6 +6,11 @@
 #include "RenderTargetPool.h" // used for on-the-fly acquisition of 
 #include "ShaderCompilerCore.h"
 
+bool ShouldCompileGPUFFT(EShaderPlatform Platform)
+{
+	return FDataDrivenShaderPlatformInfo::GetSupportsFFTBloom(Platform);
+}
+
 uint32 GPUFFT::MaxScanLineLength() 
 {
 	return 4096;
@@ -165,7 +170,21 @@ namespace GPUFFT
 namespace GPUFFT
 {
 
-	class FReorderFFTPassCS : public FGlobalShader
+	class FFFTShader : public FGlobalShader
+	{
+	public:
+		static bool ShouldCompilePermutation(const FGlobalShaderPermutationParameters& Parameters)
+		{
+			return ShouldCompileGPUFFT(Parameters.Platform);
+		}
+
+		FFFTShader() = default;
+		FFFTShader(const CompiledShaderInitializerType& Initializer)
+			: FGlobalShader(Initializer)
+		{}
+	};
+
+	class FReorderFFTPassCS : public FFFTShader
 	{
 		// NB:  the following is actually "public:" 
 		// due to text in the DECLARE_SHADER_TYPE Macro 
@@ -179,7 +198,7 @@ namespace GPUFFT
 		FReorderFFTPassCS() {};
 
 		FReorderFFTPassCS(const ShaderMetaType::CompiledShaderInitializerType& Initializer)
-			: FGlobalShader(Initializer)
+			: FFFTShader(Initializer)
 		{
 			using GPUFFTComputeShaderUtils::FComputeParameterBinder;
 
@@ -197,12 +216,6 @@ namespace GPUFFT
 		static const TCHAR* GetSourceFilename() { return TEXT("/Engine/Private/GPUFastFourierTransform.usf"); }
 		static const TCHAR* GetFunctionName() { return TEXT("ReorderFFTPassCS"); }
 
-		static bool ShouldCompilePermutation(const FGlobalShaderPermutationParameters& Parameters)
-		{
-			// @todo MetalMRT: Metal MRT can't cope with the threadgroup storage requirements for these shaders right now
-			return IsFeatureLevelSupported(Parameters.Platform, ERHIFeatureLevel::SM5) && !IsMetalMRTPlatform(Parameters.Platform) && IsPCPlatform(Parameters.Platform);
-		}
-
 		static void ModifyCompilationEnvironment(const FGlobalShaderPermutationParameters& Parameters, FShaderCompilerEnvironment& OutEnvironment)
 		{
 			FGlobalShader::ModifyCompilationEnvironment(Parameters, OutEnvironment);
@@ -254,7 +267,7 @@ namespace GPUFFT
 	};
 
 
-	class FGroupShardSubFFTPassCS : public FGlobalShader
+	class FGroupShardSubFFTPassCS : public FFFTShader
 	{
 		// NB:  the following is actually "public:" 
 		// due to text in the DECLARE_SHADER_TYPE Macro 
@@ -268,7 +281,7 @@ namespace GPUFFT
 		FGroupShardSubFFTPassCS() {};
 
 		FGroupShardSubFFTPassCS(const ShaderMetaType::CompiledShaderInitializerType& Initializer)
-			: FGlobalShader(Initializer)
+			: FFFTShader(Initializer)
 		{
 			using GPUFFTComputeShaderUtils::FComputeParameterBinder;
 
@@ -288,12 +301,6 @@ namespace GPUFFT
 		static uint32 SubPassLength() { return 2048;}
 		static uint32 Radix() { return 2; }
 
-		static bool ShouldCompilePermutation(const FGlobalShaderPermutationParameters& Parameters)
-		{
-			// @todo MetalMRT: Metal MRT can't cope with the threadgroup storage requirements for these shaders right now
-			return IsFeatureLevelSupported(Parameters.Platform, ERHIFeatureLevel::SM5) && !IsMetalMRTPlatform(Parameters.Platform) && IsPCPlatform(Parameters.Platform);
-		}
-
 		static void ModifyCompilationEnvironment(const FGlobalShaderPermutationParameters& Parameters, FShaderCompilerEnvironment& OutEnvironment)
 		{
 			FGlobalShader::ModifyCompilationEnvironment(Parameters, OutEnvironment);
@@ -344,7 +351,7 @@ namespace GPUFFT
 		LAYOUT_FIELD(FShaderParameter, NumSubRegions);
 	};
 
-	class FComplexFFTPassCS : public FGlobalShader
+	class FComplexFFTPassCS : public FFFTShader
 	{
 	// NB:  the following is actually "public:" 
 	// due to text in the DECLARE_SHADER_TYPE Macro 
@@ -358,7 +365,7 @@ namespace GPUFFT
 		FComplexFFTPassCS() {};
 
 		FComplexFFTPassCS(const ShaderMetaType::CompiledShaderInitializerType& Initializer)
-			: FGlobalShader(Initializer)
+			: FFFTShader(Initializer)
 		{
 			using GPUFFTComputeShaderUtils::FComputeParameterBinder;
 
@@ -376,12 +383,6 @@ namespace GPUFFT
 		static const TCHAR* GetSourceFilename() { return TEXT("/Engine/Private/GPUFastFourierTransform.usf"); }
 		static const TCHAR* GetFunctionName()   { return TEXT("ComplexFFTPassCS"); }
 
-		static bool ShouldCompilePermutation(const FGlobalShaderPermutationParameters& Parameters)
-		{
-			// @todo MetalMRT: Metal MRT can't cope with the threadgroup storage requirements for these shaders right now
-			return IsFeatureLevelSupported(Parameters.Platform, ERHIFeatureLevel::SM5) && !IsMetalMRTPlatform(Parameters.Platform) && IsPCPlatform(Parameters.Platform);
-		}
-
 		static void ModifyCompilationEnvironment(const FGlobalShaderPermutationParameters& Parameters, FShaderCompilerEnvironment& OutEnvironment)
 		{
 			FGlobalShader::ModifyCompilationEnvironment(Parameters, OutEnvironment);
@@ -433,7 +434,7 @@ namespace GPUFFT
 	};
 
 
-	class FPackTwoForOneFFTPassCS : public FGlobalShader
+	class FPackTwoForOneFFTPassCS : public FFFTShader
 	{
 		// NB:  the following is actually "public:" 
 		// due to text in the DECLARE_SHADER_TYPE Macro 
@@ -447,7 +448,7 @@ namespace GPUFFT
 		FPackTwoForOneFFTPassCS() {};
 
 		FPackTwoForOneFFTPassCS(const ShaderMetaType::CompiledShaderInitializerType& Initializer)
-			: FGlobalShader(Initializer)
+			: FFFTShader(Initializer)
 		{
 			using GPUFFTComputeShaderUtils::FComputeParameterBinder;
 
@@ -462,12 +463,6 @@ namespace GPUFFT
 		static const TCHAR* GetSourceFilename() { return TEXT("/Engine/Private/GPUFastFourierTransform.usf"); }
 		static const TCHAR* GetFunctionName() { return TEXT("PackTwoForOneFFTPassCS"); }
 
-		static bool ShouldCompilePermutation(const FGlobalShaderPermutationParameters& Parameters)
-		{
-			// @todo MetalMRT: Metal MRT can't cope with the threadgroup storage requirements for these shaders right now
-			return IsFeatureLevelSupported(Parameters.Platform, ERHIFeatureLevel::SM5) && !IsMetalMRTPlatform(Parameters.Platform) && IsPCPlatform(Parameters.Platform);
-		}
-
 		static void ModifyCompilationEnvironment(const FGlobalShaderPermutationParameters& Parameters, FShaderCompilerEnvironment& OutEnvironment)
 		{
 			FGlobalShader::ModifyCompilationEnvironment(Parameters, OutEnvironment);
@@ -506,7 +501,7 @@ namespace GPUFFT
 	};
 
 
-	class FCopyWindowCS : public FGlobalShader
+	class FCopyWindowCS : public FFFTShader
 	{
 		// NB:  the following is actually "public:" 
 		// due to text in the DECLARE_SHADER_TYPE Macro 
@@ -517,7 +512,7 @@ namespace GPUFFT
 		FCopyWindowCS() {};
 
 		FCopyWindowCS(const ShaderMetaType::CompiledShaderInitializerType& Initializer)
-			: FGlobalShader(Initializer)
+			: FFFTShader(Initializer)
 		{
 			using GPUFFTComputeShaderUtils::FComputeParameterBinder;
 
@@ -536,15 +531,9 @@ namespace GPUFFT
 		static uint32 XThreadCount() { return 1; }
 		static uint32 YThreadCount() { return 32; }
 
-		static bool ShouldCompilePermutation(const FGlobalShaderPermutationParameters& Parameters)
-		{
-			// @todo MetalMRT: Metal MRT can't cope with the threadgroup storage requirements for these shaders right now
-			return IsFeatureLevelSupported(Parameters.Platform, ERHIFeatureLevel::SM5) && !IsMetalMRTPlatform(Parameters.Platform) && IsPCPlatform(Parameters.Platform);
-		}
-
 		static void ModifyCompilationEnvironment(const FGlobalShaderPermutationParameters& Parameters, FShaderCompilerEnvironment& OutEnvironment)
 		{
-			FGlobalShader::ModifyCompilationEnvironment(Parameters, OutEnvironment);
+			FFFTShader::ModifyCompilationEnvironment(Parameters, OutEnvironment);
 			OutEnvironment.SetDefine(TEXT("INCLUDE_COPY_WINDOW"), 1);
 			OutEnvironment.SetDefine(TEXT("X_THREAD_COUNT"), XThreadCount());
 			OutEnvironment.SetDefine(TEXT("Y_THREAD_COUNT"), YThreadCount());
@@ -580,7 +569,7 @@ namespace GPUFFT
 		LAYOUT_FIELD(FShaderParameter, PreFilter);
 	};
 
-	class FComplexMultiplyImagesCS : public FGlobalShader
+	class FComplexMultiplyImagesCS : public FFFTShader
 	{
 		// NB:  the following is actually "public:" 
 		// due to text in the DECLARE_SHADER_TYPE Macro 
@@ -591,7 +580,7 @@ namespace GPUFFT
 		FComplexMultiplyImagesCS() {};
 
 		FComplexMultiplyImagesCS(const ShaderMetaType::CompiledShaderInitializerType& Initializer)
-			: FGlobalShader(Initializer)
+			: FFFTShader(Initializer)
 		{
 			using GPUFFTComputeShaderUtils::FComputeParameterBinder;
 
@@ -607,13 +596,6 @@ namespace GPUFFT
 		static const TCHAR* GetSourceFilename() { return TEXT("/Engine/Private/GPUFastFourierTransform.usf"); }
 		static const TCHAR* GetFunctionName() { return TEXT("ComplexMultiplyImagesCS"); }
 
-
-		static bool ShouldCompilePermutation(const FGlobalShaderPermutationParameters& Parameters)
-		{
-			// @todo MetalMRT: Metal MRT can't cope with the threadgroup storage requirements for these shaders right now
-			return IsFeatureLevelSupported(Parameters.Platform, ERHIFeatureLevel::SM5) && !IsMetalMRTPlatform(Parameters.Platform) && IsPCPlatform(Parameters.Platform);
-		}
-
 		static void ModifyCompilationEnvironment(const FGlobalShaderPermutationParameters& Parameters, FShaderCompilerEnvironment& OutEnvironment)
 		{
 			FGlobalShader::ModifyCompilationEnvironment(Parameters, OutEnvironment);
@@ -653,7 +635,7 @@ namespace GPUFFT
 		LAYOUT_FIELD(FShaderParameter, DataLayout);
 	};
 
-	class FGSComplexTransformBaseCS : public FGlobalShader
+	class FGSComplexTransformBaseCS : public FFFTShader
 	{
 		DECLARE_INLINE_TYPE_LAYOUT(FGSComplexTransformBaseCS, NonVirtual);
 	public:
@@ -664,7 +646,7 @@ namespace GPUFFT
 		FGSComplexTransformBaseCS() {};
 
 		FGSComplexTransformBaseCS(const ShaderMetaType::CompiledShaderInitializerType& Initializer)
-			: FGlobalShader(Initializer)
+			: FFFTShader(Initializer)
 		{
 			using GPUFFTComputeShaderUtils::FComputeParameterBinder;
 		
@@ -737,13 +719,6 @@ namespace GPUFFT
 		static const TCHAR* GetSourceFilename() { return TEXT("/Engine/Private/GPUFastFourierTransform.usf"); }
 		static const TCHAR* GetFunctionName()   { return TEXT("GroupSharedComplexFFTCS"); }
 
-		static bool ShouldCompilePermutation(const FGlobalShaderPermutationParameters& Parameters)
-		{
-			// @todo MetalMRT: Metal MRT can't cope with the threadgroup storage requirements for these shaders right now
-			return IsFeatureLevelSupported(Parameters.Platform, ERHIFeatureLevel::SM5) && !IsMetalMRTPlatform(Parameters.Platform) && IsPCPlatform(Parameters.Platform);
-		}
-
-
 		static void ModifyCompilationEnvironment(const FGlobalShaderPermutationParameters& Parameters, FShaderCompilerEnvironment& OutEnvironment)
 		{
 			FGlobalShader::ModifyCompilationEnvironment(Parameters, OutEnvironment);
@@ -778,13 +753,6 @@ namespace GPUFFT
 		static const TCHAR* GetSourceFilename() { return TEXT("/Engine/Private/GPUFastFourierTransform.usf"); }
 		static const TCHAR* GetFunctionName()   { return TEXT("GroupSharedTwoForOneFFTCS"); }
 
-		static bool ShouldCompilePermutation(const FGlobalShaderPermutationParameters& Parameters)
-		{
-			// @todo MetalMRT: Metal MRT can't cope with the threadgroup storage requirements for these shaders right now
-			return IsFeatureLevelSupported(Parameters.Platform, ERHIFeatureLevel::SM5) && !IsMetalMRTPlatform(Parameters.Platform) && IsPCPlatform(Parameters.Platform);
-		}
-
-
 		static void ModifyCompilationEnvironment(const FGlobalShaderPermutationParameters& Parameters, FShaderCompilerEnvironment& OutEnvironment)
 		{
 			FGlobalShader::ModifyCompilationEnvironment(Parameters, OutEnvironment);
@@ -809,7 +777,7 @@ namespace GPUFFT
 	/**
 	* Base class used for 1d convolution.  
 	*/
-	class FGSConvolutionBaseCS : public FGlobalShader
+	class FGSConvolutionBaseCS : public FFFTShader
 	{
 		DECLARE_INLINE_TYPE_LAYOUT(FGSConvolutionBaseCS, NonVirtual);
 	public:
@@ -817,7 +785,7 @@ namespace GPUFFT
 		FGSConvolutionBaseCS() {};
 
 		FGSConvolutionBaseCS(const ShaderMetaType::CompiledShaderInitializerType& Initializer)
-			: FGlobalShader(Initializer)
+			: FFFTShader(Initializer)
 		{
 			using GPUFFTComputeShaderUtils::FComputeParameterBinder;
 
@@ -930,15 +898,6 @@ namespace GPUFFT
 		// due to text in the DECLARE_SHADER_TYPE Macro 
 		DECLARE_SHADER_TYPE(TGSConvolutionWithTexturerCS, Global);
 
-
-
-		static bool ShouldCompilePermutation(const FGlobalShaderPermutationParameters& Parameters)
-		{
-			// @todo MetalMRT: Metal MRT can't cope with the threadgroup storage requirements for these shaders right now
-			return IsFeatureLevelSupported(Parameters.Platform, ERHIFeatureLevel::SM5) && !IsMetalMRTPlatform(Parameters.Platform) && IsPCPlatform(Parameters.Platform);
-		}
-
-
 		static void ModifyCompilationEnvironment(const FGlobalShaderPermutationParameters& Parameters, FShaderCompilerEnvironment& OutEnvironment)
 		{
 			FGlobalShader::ModifyCompilationEnvironment(Parameters, OutEnvironment);
diff --git a/Engine/Source/Runtime/Renderer/Private/PostProcess/PostProcessBloomSetup.cpp b/Engine/Source/Runtime/Renderer/Private/PostProcess/PostProcessBloomSetup.cpp
--- a/Engine/Source/Runtime/Renderer/Private/PostProcess/PostProcessBloomSetup.cpp
+++ b/Engine/Source/Runtime/Renderer/Private/PostProcess/PostProcessBloomSetup.cpp
@@ -199,7 +199,15 @@ FBloomOutputs AddBloomPass(FRDGBuilder& GraphBuilder, const FViewInfo& View, con
 		if (bFFTBloomEnabled)
 		{
 			FScreenPassTexture FullResolution = Inputs.SceneColor;
-			FScreenPassTexture HalfResolution = Inputs.SceneDownsampleChain->GetFirstTexture();
+			FScreenPassTexture HalfResolution;
+			if (IsFFTBloomQuarterResolutionEnabled())
+			{
+				HalfResolution = Inputs.SceneDownsampleChain->GetTexture(1);
+			}
+			else
+			{
+				HalfResolution = Inputs.SceneDownsampleChain->GetFirstTexture();
+			}
 
 			FFFTBloomInputs PassInputs;
 			PassInputs.FullResolutionTexture = FullResolution.Texture;
diff --git a/Engine/Source/Runtime/Renderer/Private/PostProcess/PostProcessFFTBloom.cpp b/Engine/Source/Runtime/Renderer/Private/PostProcess/PostProcessFFTBloom.cpp
--- a/Engine/Source/Runtime/Renderer/Private/PostProcess/PostProcessFFTBloom.cpp
+++ b/Engine/Source/Runtime/Renderer/Private/PostProcess/PostProcessFFTBloom.cpp
@@ -11,16 +11,16 @@ TAutoConsoleVariable<int32> CVarHalfResFFTBloom(
 	0,
 	TEXT("Experimental half-resolution FFT Bloom convolution. \n")
 	TEXT(" 0: Standard full resolution convolution bloom.")
-	TEXT(" 1: Half-resolution convolution that excludes the center of the kernel.\n"),
+	TEXT(" 1: Half-resolution convolution that excludes the center of the kernel.\n")
+	TEXT(" 2: Quarter-resolution convolution that excludes the center of the kernel.\n"),
 	ECVF_Scalability | ECVF_RenderThreadSafe);
 
 static bool DoesPlatformSupportFFTBloom(EShaderPlatform Platform)
 {
-	// @todo MetalMRT: Metal MRT can't cope with the threadgroup storage requirements for these shaders right now
-	return IsFeatureLevelSupported(Platform, ERHIFeatureLevel::SM5) && !IsMetalMRTPlatform(Platform) && IsPCPlatform(Platform);
+	return FDataDrivenShaderPlatformInfo::GetSupportsFFTBloom(Platform);
 }
 
-class FFFTShader : public FGlobalShader
+class FFFTBloomShader : public FGlobalShader
 {
 public:
 	// Determine the number of threads used per scanline when writing the physical space kernel
@@ -31,17 +31,17 @@ public:
 		return DoesPlatformSupportFFTBloom(Parameters.Platform);
 	}
 
-	FFFTShader() = default;
-	FFFTShader(const CompiledShaderInitializerType& Initializer)
+	FFFTBloomShader() = default;
+	FFFTBloomShader(const CompiledShaderInitializerType& Initializer)
 		: FGlobalShader(Initializer)
 	{}
 };
 
-class FResizeAndCenterTextureCS : public FFFTShader
+class FResizeAndCenterTextureCS : public FFFTBloomShader
 {
 public:
 	DECLARE_GLOBAL_SHADER(FResizeAndCenterTextureCS);
-	SHADER_USE_PARAMETER_STRUCT(FResizeAndCenterTextureCS, FFFTShader);
+	SHADER_USE_PARAMETER_STRUCT(FResizeAndCenterTextureCS, FFFTBloomShader);
 
 	BEGIN_SHADER_PARAMETER_STRUCT(FParameters, )
 		SHADER_PARAMETER_TEXTURE(Texture2D, SrcTexture)
@@ -55,7 +55,7 @@ public:
 
 	static void ModifyCompilationEnvironment(const FGlobalShaderPermutationParameters& Parameters, FShaderCompilerEnvironment& OutEnvironment)
 	{
-		FFFTShader::ModifyCompilationEnvironment(Parameters, OutEnvironment);
+		FFFTBloomShader::ModifyCompilationEnvironment(Parameters, OutEnvironment);
 		OutEnvironment.SetDefine(TEXT("INCLUDE_RESIZE_AND_CENTER"), 1);
 		OutEnvironment.SetDefine(TEXT("THREADS_PER_GROUP"), ThreadsPerGroup);
 	}
@@ -63,11 +63,11 @@ public:
 
 IMPLEMENT_GLOBAL_SHADER(FResizeAndCenterTextureCS, "/Engine/Private/PostProcessFFTBloom.usf", "ResizeAndCenterTextureCS", SF_Compute);
 
-class FCaptureKernelWeightsCS : public FFFTShader
+class FCaptureKernelWeightsCS : public FFFTBloomShader
 {
 public:
 	DECLARE_GLOBAL_SHADER(FCaptureKernelWeightsCS);
-	SHADER_USE_PARAMETER_STRUCT(FCaptureKernelWeightsCS, FFFTShader);
+	SHADER_USE_PARAMETER_STRUCT(FCaptureKernelWeightsCS, FFFTBloomShader);
 
 	BEGIN_SHADER_PARAMETER_STRUCT(FParameters, )
 		SHADER_PARAMETER_TEXTURE(Texture2D, HalfResSrcTexture)
@@ -80,18 +80,18 @@ public:
 
 	static void ModifyCompilationEnvironment(const FGlobalShaderPermutationParameters& Parameters, FShaderCompilerEnvironment& OutEnvironment)
 	{
-		FFFTShader::ModifyCompilationEnvironment(Parameters, OutEnvironment);
+		FFFTBloomShader::ModifyCompilationEnvironment(Parameters, OutEnvironment);
 		OutEnvironment.SetDefine(TEXT("INCLUDE_CAPTURE_KERNEL_WEIGHTS"), 1);
 	}
 };
 
 IMPLEMENT_GLOBAL_SHADER(FCaptureKernelWeightsCS, "/Engine/Private/PostProcessFFTBloom.usf", "CaptureKernelWeightsCS", SF_Compute);
 
-class FBlendLowResCS : public FFFTShader
+class FBlendLowResCS : public FFFTBloomShader
 {
 public:
 	DECLARE_GLOBAL_SHADER(FBlendLowResCS);
-	SHADER_USE_PARAMETER_STRUCT(FBlendLowResCS, FFFTShader);
+	SHADER_USE_PARAMETER_STRUCT(FBlendLowResCS, FFFTBloomShader);
 
 	BEGIN_SHADER_PARAMETER_STRUCT(FParameters, )
 		SHADER_PARAMETER_TEXTURE(Texture2D, SrcTexture)
@@ -106,7 +106,7 @@ public:
 
 	static void ModifyCompilationEnvironment(const FGlobalShaderPermutationParameters& Parameters, FShaderCompilerEnvironment& OutEnvironment)
 	{
-		FFFTShader::ModifyCompilationEnvironment(Parameters, OutEnvironment);
+		FFFTBloomShader::ModifyCompilationEnvironment(Parameters, OutEnvironment);
 		OutEnvironment.SetDefine(TEXT("INCLUDE_BLEND_LOW_RES"), 1);
 		OutEnvironment.SetDefine(TEXT("THREADS_PER_GROUP"), ThreadsPerGroup);
 	}
@@ -114,11 +114,11 @@ public:
 
 IMPLEMENT_GLOBAL_SHADER(FBlendLowResCS, "/Engine/Private/PostProcessFFTBloom.usf", "BlendLowResCS", SF_Compute);
 
-class FPassThroughCS : public FFFTShader
+class FPassThroughCS : public FFFTBloomShader
 {
 public:
 	DECLARE_GLOBAL_SHADER(FPassThroughCS);
-	SHADER_USE_PARAMETER_STRUCT(FPassThroughCS, FFFTShader);
+	SHADER_USE_PARAMETER_STRUCT(FPassThroughCS, FFFTBloomShader);
 
 	BEGIN_SHADER_PARAMETER_STRUCT(FParameters, )
 		SHADER_PARAMETER_TEXTURE(Texture2D, SrcTexture)
@@ -129,7 +129,7 @@ public:
 
 	static void ModifyCompilationEnvironment(const FGlobalShaderPermutationParameters& Parameters, FShaderCompilerEnvironment& OutEnvironment)
 	{
-		FFFTShader::ModifyCompilationEnvironment(Parameters, OutEnvironment);
+		FFFTBloomShader::ModifyCompilationEnvironment(Parameters, OutEnvironment);
 		OutEnvironment.SetDefine(TEXT("INCLUDE_PASSTHROUGH"), 1);
 		OutEnvironment.SetDefine(TEXT("THREADS_PER_GROUP"), ThreadsPerGroup);
 	}
@@ -140,7 +140,12 @@ IMPLEMENT_GLOBAL_SHADER(FPassThroughCS, "/Engine/Private/PostProcessFFTBloom.usf
 
 bool IsFFTBloomHalfResolutionEnabled()
 {
-	return CVarHalfResFFTBloom.GetValueOnRenderThread() == 1;
+	return CVarHalfResFFTBloom.GetValueOnRenderThread() != 0;
+}
+
+bool IsFFTBloomQuarterResolutionEnabled()
+{
+	return CVarHalfResFFTBloom.GetValueOnRenderThread() == 2;
 }
 
 bool IsFFTBloomPhysicalKernelReady(const FViewInfo& View)
diff --git a/Engine/Source/Runtime/Renderer/Private/PostProcess/PostProcessFFTBloom.h b/Engine/Source/Runtime/Renderer/Private/PostProcess/PostProcessFFTBloom.h
--- a/Engine/Source/Runtime/Renderer/Private/PostProcess/PostProcessFFTBloom.h
+++ b/Engine/Source/Runtime/Renderer/Private/PostProcess/PostProcessFFTBloom.h
@@ -17,4 +17,6 @@ struct FFFTBloomInputs
 	FIntRect HalfResolutionViewRect;
 };
 
+bool IsFFTBloomQuarterResolutionEnabled();
+
 FRDGTextureRef AddFFTBloomPass(FRDGBuilder& GraphBuilder, const FViewInfo& View, const FFFTBloomInputs& Inputs);
\ No newline at end of file
