Updated CapsuleShadowing to RDG.
- Resource states/transitions handled by RDG
- Shader bindings using SHADER_PARAMETER_STRUCT.
Converted data and intermediate buffers to StructuredBuffers.
Removed a lot of unecessary structs from DistanceFiedLightingShared.h

diff --git a/Engine/Shaders/Private/CapsuleShadowShaders.usf b/Engine/Shaders/Private/CapsuleShadowShaders.usf
--- a/Engine/Shaders/Private/CapsuleShadowShaders.usf
+++ b/Engine/Shaders/Private/CapsuleShadowShaders.usf
@@ -29,6 +29,13 @@
 #define LIGHT_SOURCE_FROM_CAPSULE 1
 #define LIGHT_SOURCE_FROM_RECEIVER 2
 
+struct FCapsuleShape
+{
+	float3 Center;
+	float Radius;
+	float3 Orientation;
+	float Length;
+};
 
 // SUPPORT_CAPSULE_SHAPES
 
@@ -40,7 +47,7 @@
 	groupshared uint IntersectingShapeIndices[MAX_INTERSECTING_SHAPES * 2];
 
 	uint NumShadowCapsules;
-	Buffer<float4> ShadowCapsuleShapes;
+	StructuredBuffer<FCapsuleShape> ShadowCapsuleShapes;
 
 
 // SUPPORT_MESH_DISTANCE_FIELDS
@@ -53,11 +60,11 @@
 	groupshared uint IntersectingMeshDistanceFieldIndices[MAX_INTERSECTING_DISTANCE_FIELDS * 2];
 
 	uint NumMeshDistanceFieldCasters;
-	Buffer<uint> MeshDistanceFieldCasterIndices;
+	StructuredBuffer<uint> MeshDistanceFieldCasterIndices;
 
 
 #if LIGHT_SOURCE_MODE == LIGHT_SOURCE_FROM_CAPSULE
-Buffer<float4> LightDirectionData;
+StructuredBuffer<float4> LightDirectionData;
 
 void DecodeLightDirectionW(float LightDirectionW, out float LightAngle, out float MinVisibility)
 {
@@ -87,7 +94,7 @@ float CapsuleSkyAngleScale;
 float CapsuleMinSkyAngle;
 uint NumLightDirectionData;
 
-RWBuffer<float4> RWComputedLightDirectionData;
+RWStructuredBuffer<float4> RWComputedLightDirectionData;
 
 [numthreads(THREADGROUP_SIZEX, 1, 1)]
 void ComputeLightDirectionFromVolumetricLightmapCS(
@@ -440,10 +447,11 @@ uint CullCapsuleShapesToTile(
 			float ConeAngleSin1 = ConeAngleSin0;
 
 		#endif
+		
+		FCapsuleShape CapsuleShape = ShadowCapsuleShapes[ShapeIndex];
 
-		float4 SphereCenterAndRadius = ShadowCapsuleShapes[ShapeIndex * 2];
-		float3 TransformedSphereCenter = SphereCenterAndRadius.xyz;
-		float TransformedSphereRadius = SphereCenterAndRadius.w;
+		float3 TransformedSphereCenter = CapsuleShape.Center;
+		float TransformedSphereRadius = CapsuleShape.Radius;
 		float3 TransformedTileBoundingSphereCenter0 = TileCullingData0.BoundingSphere.xyz;
 		float3 TransformedTileBoundingSphereCenter1 = TileCullingData1.BoundingSphere.xyz;
 
@@ -453,30 +461,26 @@ uint CullCapsuleShapesToTile(
 		#endif
 
 		#if USE_SCALED_SPHERE_INTERSECTION
-			
-			float4 CapsuleCenterAndRadius = SphereCenterAndRadius;
-			float4 CapsuleOrientationAndLength = ShadowCapsuleShapes[ShapeIndex * 2 + 1];
-
 			float3 CapsuleSpaceX;
 			float3 CapsuleSpaceY;
-			float3 CapsuleSpaceZ = CapsuleOrientationAndLength.xyz;
+			float3 CapsuleSpaceZ = CapsuleShape.Orientation;
 			GenerateCoordinateSystem(CapsuleSpaceZ, CapsuleSpaceX, CapsuleSpaceY);
 
 			// Scale required along the capsule's axis to turn it into a sphere (assuming it was originally a scaled sphere instead of a capsule)
-			float CapsuleZScale = CapsuleCenterAndRadius.w / (.5f * CapsuleOrientationAndLength.w + CapsuleCenterAndRadius.w);
+			float CapsuleZScale = CapsuleShape.Radius / (.5f * CapsuleShape.Length + CapsuleShape.Radius);
 			CapsuleSpaceZ *= CapsuleZScale;
 
 			// The capsule is centered at 0 in the scaled sphere space
 			TransformedSphereCenter = 0;
 
 			// After scaling along the capsule axis it will become a sphere with the original radius
-			TransformedSphereRadius = SphereCenterAndRadius.w;
+			TransformedSphereRadius = CapsuleShape.Radius;
 
 			// Transform the sphere center and cone axis into the scaled sphere space
-			float3 CapsuleCenterToTileCenter0 = TileCullingData0.BoundingSphere.xyz - CapsuleCenterAndRadius.xyz;
+			float3 CapsuleCenterToTileCenter0 = TileCullingData0.BoundingSphere.xyz - CapsuleShape.Center;
 			TransformedTileBoundingSphereCenter0 = float3(dot(CapsuleCenterToTileCenter0, CapsuleSpaceX), dot(CapsuleCenterToTileCenter0, CapsuleSpaceY), dot(CapsuleCenterToTileCenter0, CapsuleSpaceZ));
 			
-			float3 CapsuleCenterToTileCenter1 = TileCullingData1.BoundingSphere.xyz - CapsuleCenterAndRadius.xyz;
+			float3 CapsuleCenterToTileCenter1 = TileCullingData1.BoundingSphere.xyz - CapsuleShape.Center;
 			TransformedTileBoundingSphereCenter1 = float3(dot(CapsuleCenterToTileCenter1, CapsuleSpaceX), dot(CapsuleCenterToTileCenter1, CapsuleSpaceY), dot(CapsuleCenterToTileCenter1, CapsuleSpaceZ));
 
 			#if LIGHT_SOURCE_MODE != LIGHT_SOURCE_FROM_RECEIVER
@@ -484,13 +488,9 @@ uint CullCapsuleShapesToTile(
 				TransformedConeAxis0 = normalize(float3(dot(ConeAxis0, CapsuleSpaceX), dot(ConeAxis0, CapsuleSpaceY), dot(ConeAxis0, CapsuleSpaceZ)));
 				TransformedConeAxis1 = normalize(float3(dot(ConeAxis1, CapsuleSpaceX), dot(ConeAxis1, CapsuleSpaceY), dot(ConeAxis1, CapsuleSpaceZ)));
 			#endif
-
 		#else
-
-			float CapsuleLength = ShadowCapsuleShapes[ShapeIndex * 2 + 1].w;
-
 			// Add half capsule length to bounding sphere
-			TransformedSphereRadius = SphereCenterAndRadius.w + .5f * CapsuleLength;
+			TransformedSphereRadius = CapsuleShape.Radius + .5f * CapsuleShape.Length;
 		#endif
 
 		BRANCH
@@ -609,30 +609,29 @@ float ShadowConeTraceAgainstCulledCapsuleShapes(
 			UnitRayDirection = float3(0, 0, 1);
 			LightAngle = .3f;
 		#endif
+			
+		FCapsuleShape CapsuleShape = ShadowCapsuleShapes[ObjectIndex];
 
-		float4 CapsuleCenterAndRadius = ShadowCapsuleShapes[ObjectIndex * 2];
-		ApplyFadeToCapsuleRadius(CapsuleCenterAndRadius.w, UnitRayDirection);
-
-		float4 CapsuleOrientationAndLength = ShadowCapsuleShapes[ObjectIndex * 2 + 1];
+		ApplyFadeToCapsuleRadius(CapsuleShape.Radius, UnitRayDirection);
 
 		float DistanceToShadowSphere;
 		float3 UnitVectorToShadowSphere;
 		float3 UnitRayDirectionInCorrectSpace = UnitRayDirection;
 
 		BRANCH
-		if (CapsuleOrientationAndLength.w > 0)
+		if (CapsuleShape.Length > 0)
 		{
 #if USE_SCALED_SPHERE_INTERSECTION
 
 			float3 CapsuleSpaceX;
 			float3 CapsuleSpaceY;
-			float3 CapsuleSpaceZ = CapsuleOrientationAndLength.xyz;
+			float3 CapsuleSpaceZ = CapsuleShape.Orientation;
 			GenerateCoordinateSystem(CapsuleSpaceZ, CapsuleSpaceX, CapsuleSpaceY);
 
-			float CapsuleZScale = CapsuleCenterAndRadius.w / (.5f * CapsuleOrientationAndLength.w + CapsuleCenterAndRadius.w);
+			float CapsuleZScale = CapsuleShape.Radius / (.5f * CapsuleShape.Length + CapsuleShape.Radius);
 			CapsuleSpaceZ *= CapsuleZScale;
 
-			float3 CapsuleCenterToRayStart = WorldRayStart - CapsuleCenterAndRadius.xyz;
+			float3 CapsuleCenterToRayStart = WorldRayStart - CapsuleShape.Center;
 			float3 CapsuleSpaceRayStart = float3(dot(CapsuleCenterToRayStart, CapsuleSpaceX), dot(CapsuleCenterToRayStart, CapsuleSpaceY), dot(CapsuleCenterToRayStart, CapsuleSpaceZ));
 
 			float3 CapsuleSpaceRayDirection = float3(dot(UnitRayDirection, CapsuleSpaceX), dot(UnitRayDirection, CapsuleSpaceY), dot(UnitRayDirection, CapsuleSpaceZ));
@@ -641,10 +640,10 @@ float ShadowConeTraceAgainstCulledCapsuleShapes(
 			UnitVectorToShadowSphere = -CapsuleSpaceRayStart / DistanceToShadowSphere;
 			UnitRayDirectionInCorrectSpace = normalize(CapsuleSpaceRayDirection);
 #else
-			float3 VectorToCapsuleCenter = CapsuleCenterAndRadius.xyz - WorldRayStart;
+			float3 VectorToCapsuleCenter = CapsuleShape.Center - WorldRayStart;
 
 			// Closest point on line segment to ray
-			float3 L01 = CapsuleOrientationAndLength.xyz * CapsuleOrientationAndLength.w;
+			float3 L01 = CapsuleShape.Orientation * CapsuleShape.Length;
 			float3 L0 = VectorToCapsuleCenter - 0.5 * L01;
 			float3 L1 = VectorToCapsuleCenter + 0.5 * L01;
 
@@ -653,7 +652,7 @@ float ShadowConeTraceAgainstCulledCapsuleShapes(
 			// Vector that spans L01 perpendicular to the ray
 			float3 PerpendicularSpanningVector = CapsuleOrientationProjectedOntoRay * UnitRayDirection - L01;
 			// Length of PerpendicularSpanningVector using the right triangle formed by L01 and UnitRayDirection * CapsuleOrientationProjectedOntoRay
-			float PerpendicularDistance = Square(CapsuleOrientationAndLength.w) - CapsuleOrientationProjectedOntoRay * CapsuleOrientationProjectedOntoRay;
+			float PerpendicularDistance = Square(CapsuleShape.Length) - CapsuleOrientationProjectedOntoRay * CapsuleOrientationProjectedOntoRay;
 			// Project the vector to a capsule endpoint onto the perpendicular spanning vector, normalized
 			float t = saturate(dot(L0, PerpendicularSpanningVector) / PerpendicularDistance);
 			// Compute the vector to the shadow sphere which best approximates the capsule's shadowing 
@@ -671,17 +670,17 @@ float ShadowConeTraceAgainstCulledCapsuleShapes(
 		}
 		else
 		{
-			DistanceToShadowSphere = length(CapsuleCenterAndRadius.xyz - WorldRayStart);
-			UnitVectorToShadowSphere = (CapsuleCenterAndRadius.xyz - WorldRayStart) / DistanceToShadowSphere;
+			DistanceToShadowSphere = length(CapsuleShape.Center - WorldRayStart);
+			UnitVectorToShadowSphere = (CapsuleShape.Center - WorldRayStart) / DistanceToShadowSphere;
 		}
 
 		float AngleBetween = acosFast(dot(UnitVectorToShadowSphere, UnitRayDirectionInCorrectSpace));
-		float IntersectionArea = SphericalCapIntersectionAreaFast(LightAngle, atanFastPos(CapsuleCenterAndRadius.w / DistanceToShadowSphere), AngleBetween);
+		float IntersectionArea = SphericalCapIntersectionAreaFast(LightAngle, atanFastPos(CapsuleShape.Radius / DistanceToShadowSphere), AngleBetween);
 		#if POINT_LIGHT
 			// SphericalCapIntersectionAreaFast does not take the relative distance of the two sphere caps into account, which can cause shadows to be cast on the opposite direction
 			// Here we compare DistanceToShadowSphere and LightVectorLength to determine whether a shadow should be cast
 			// To prevent discontinuity, we use the ratio of the distance difference to the capsule's radius as a smooth factor
-			IntersectionArea = lerp(IntersectionArea, 0, saturate((DistanceToShadowSphere - LightVectorLength + CapsuleCenterAndRadius.w) / CapsuleCenterAndRadius.w));
+			IntersectionArea = lerp(IntersectionArea, 0, saturate((DistanceToShadowSphere - LightVectorLength + CapsuleShape.Radius) / CapsuleShape.Radius));
 		#endif
 		float ConeConeIntersection = 1 - saturate(IntersectionArea / AreaOfLight);
 		float DistanceFadeAlpha = saturate(DistanceToShadowSphere * InvMaxOcclusionDistance * 3 - 2);
@@ -974,7 +973,7 @@ RWTexture2D<float2> RWShadowFactors;
 float MaxOcclusionDistance;
 
 uint2 TileDimensions;
-RWBuffer<uint> RWTileIntersectionCounts;
+RWStructuredBuffer<uint> RWTileIntersectionCounts;
 
 [numthreads(THREADGROUP_SIZEX, THREADGROUP_SIZEY, 1)]
 void CapsuleShadowingCS(
@@ -1134,7 +1133,7 @@ void CapsuleShadowingCS(
 #endif
 }
 
-Buffer<uint> TileIntersectionCounts;
+StructuredBuffer<uint> TileIntersectionCounts;
 
 // Size of a tile in NDC
 float2 TileSize;
diff --git a/Engine/Source/Runtime/Renderer/Private/CapsuleShadowRendering.cpp b/Engine/Source/Runtime/Renderer/Private/CapsuleShadowRendering.cpp
--- a/Engine/Source/Runtime/Renderer/Private/CapsuleShadowRendering.cpp
+++ b/Engine/Source/Runtime/Renderer/Private/CapsuleShadowRendering.cpp
@@ -132,8 +132,8 @@ public:
 		SHADER_PARAMETER(float, CapsuleIndirectConeAngle)
 		SHADER_PARAMETER(float, CapsuleSkyAngleScale)
 		SHADER_PARAMETER(float, CapsuleMinSkyAngle)
-		SHADER_PARAMETER_SRV(Buffer<float4>, LightDirectionData)
-		SHADER_PARAMETER_UAV(RWBuffer<float4>, RWComputedLightDirectionData)
+		SHADER_PARAMETER_RDG_BUFFER_SRV(StructuredBuffer<float4>, LightDirectionData)
+		SHADER_PARAMETER_RDG_BUFFER_UAV(RWStructuredBuffer<float4>, RWComputedLightDirectionData)
 	END_SHADER_PARAMETER_STRUCT()
 
 	static bool ShouldCompilePermutation(const FGlobalShaderPermutationParameters& Parameters)
@@ -195,6 +195,45 @@ class FCapsuleShadowingCS : public FGlobalShader
 {
 	DECLARE_GLOBAL_SHADER(FCapsuleShadowingCS);
 public:
+	SHADER_USE_PARAMETER_STRUCT(FCapsuleShadowingCS, FGlobalShader);
+
+	BEGIN_SHADER_PARAMETER_STRUCT(FParameters, )
+		SHADER_PARAMETER_RDG_UNIFORM_BUFFER(FSceneTextureUniformParameters, SceneTextures)
+		SHADER_PARAMETER_STRUCT_REF(FViewUniformShaderParameters, View)
+
+		SHADER_PARAMETER_RDG_TEXTURE_UAV(RWTexture2D, RWShadowFactors)
+		SHADER_PARAMETER_RDG_TEXTURE_UAV(RWTexture2D, RWBentNormalTexture)
+		SHADER_PARAMETER_RDG_TEXTURE(Texture2D, ReceiverBentNormalTexture)
+		SHADER_PARAMETER(FIntPoint, TileDimensions)
+		SHADER_PARAMETER(FVector2D, NumGroups)
+
+		SHADER_PARAMETER(FVector4, LightPositionAndInvRadius)
+		SHADER_PARAMETER(FVector3f, LightDirection)
+		SHADER_PARAMETER(float, LightSourceRadius)
+		SHADER_PARAMETER(FVector3f, LightAngleAndNormalThreshold)
+		SHADER_PARAMETER(float, RayStartOffsetDepthScale)
+
+		SHADER_PARAMETER(FIntRect, ScissorRectMinAndSize)
+
+		SHADER_PARAMETER(float, IndirectCapsuleSelfShadowingIntensity)
+		SHADER_PARAMETER(uint32, DownsampleFactor)
+		SHADER_PARAMETER(float, MaxOcclusionDistance)
+
+		SHADER_PARAMETER(FVector2D, CosFadeStartAngle)
+
+		SHADER_PARAMETER(uint32, NumShadowCapsules)
+		SHADER_PARAMETER_RDG_BUFFER_SRV(StructuredBuffer<FCapsuleShape>, ShadowCapsuleShapes)
+
+		SHADER_PARAMETER(uint32, NumMeshDistanceFieldCasters)
+		SHADER_PARAMETER_RDG_BUFFER_SRV(StructuredBuffer<uint>, MeshDistanceFieldCasterIndices)
+
+		SHADER_PARAMETER_RDG_BUFFER_SRV(StructuredBuffer<float4>, LightDirectionData)
+
+		SHADER_PARAMETER_STRUCT_INCLUDE(FDistanceFieldObjectBufferParameters, DFObjectBufferParameters)
+		SHADER_PARAMETER_STRUCT_INCLUDE(FDistanceFieldAtlasParameters, DFAtlasParameters)
+
+		SHADER_PARAMETER_RDG_BUFFER_UAV(RWStructuredBuffer<uint>, RWTileIntersectionCounts)
+	END_SHADER_PARAMETER_STRUCT()
 
 	class FShapeShadow : SHADER_PERMUTATION_ENUM_CLASS("SHAPE_SHADOW", ECapsuleShadowingType);
 	class FIndirectPrimitiveType : SHADER_PERMUTATION_ENUM_CLASS("INDIRECT_PRIMITIVE_TYPE", EIndirectShadowingPrimitiveTypes);
@@ -253,206 +292,6 @@ public:
 			OutEnvironment.SetDefine(TEXT("SUPPORT_MESH_DISTANCE_FIELDS"), 1);
 		}
 	}
-
-	/** Default constructor. */
-	FCapsuleShadowingCS() {}
-
-	/** Initialization constructor. */
-	FCapsuleShadowingCS(const ShaderMetaType::CompiledShaderInitializerType& Initializer)
-		: FGlobalShader(Initializer)
-	{
-		ShadowFactors.Bind(Initializer.ParameterMap, TEXT("ShadowFactors"));
-		TileIntersectionCounts.Bind(Initializer.ParameterMap, TEXT("TileIntersectionCounts"));
-		TileDimensions.Bind(Initializer.ParameterMap, TEXT("TileDimensions"));
-		BentNormalTexture.Bind(Initializer.ParameterMap, TEXT("BentNormalTexture"));
-		ReceiverBentNormalTexture.Bind(Initializer.ParameterMap, TEXT("ReceiverBentNormalTexture"));
-		NumGroups.Bind(Initializer.ParameterMap, TEXT("NumGroups"));
-		LightDirection.Bind(Initializer.ParameterMap, TEXT("LightDirection"));
-		LightSourceRadius.Bind(Initializer.ParameterMap, TEXT("LightSourceRadius"));
-		RayStartOffsetDepthScale.Bind(Initializer.ParameterMap, TEXT("RayStartOffsetDepthScale"));
-		LightPositionAndInvRadius.Bind(Initializer.ParameterMap, TEXT("LightPositionAndInvRadius"));
-		LightAngleAndNormalThreshold.Bind(Initializer.ParameterMap, TEXT("LightAngleAndNormalThreshold"));
-		ScissorRectMinAndSize.Bind(Initializer.ParameterMap, TEXT("ScissorRectMinAndSize"));
-		DownsampleFactor.Bind(Initializer.ParameterMap, TEXT("DownsampleFactor"));
-		NumShadowCapsules.Bind(Initializer.ParameterMap, TEXT("NumShadowCapsules"));
-		ShadowCapsuleShapes.Bind(Initializer.ParameterMap, TEXT("ShadowCapsuleShapes"));
-		NumMeshDistanceFieldCasters.Bind(Initializer.ParameterMap, TEXT("NumMeshDistanceFieldCasters"));
-		MeshDistanceFieldCasterIndices.Bind(Initializer.ParameterMap, TEXT("MeshDistanceFieldCasterIndices"));
-		MaxOcclusionDistance.Bind(Initializer.ParameterMap, TEXT("MaxOcclusionDistance"));
-		CosFadeStartAngle.Bind(Initializer.ParameterMap, TEXT("CosFadeStartAngle"));
-		LightDirectionData.Bind(Initializer.ParameterMap, TEXT("LightDirectionData"));
-		IndirectCapsuleSelfShadowingIntensity.Bind(Initializer.ParameterMap, TEXT("IndirectCapsuleSelfShadowingIntensity"));
-		DistanceFieldObjectParameters.Bind(Initializer.ParameterMap);
-		DistanceFieldCulledObjectParameters.Bind(Initializer.ParameterMap);
-	}
-
-	void SetParameters(
-		FRHIComputeCommandList& RHICmdList, 
-		FScene* Scene,
-		const FSceneView& View, 
-		const FLightSceneInfo* LightSceneInfo,
-		FRHITexture* OutputTexture,
-		FRHIUnorderedAccessView* OutputTextureUAV, 
-		FIntPoint TileDimensionsValue,
-		const FRWBuffer* TileIntersectionCountsBuffer,
-		FVector2D NumGroupsValue,
-		float MaxOcclusionDistanceValue,
-		const FIntRect& ScissorRect,
-		int32 DownsampleFactorValue,
-		int32 NumShadowCapsulesValue,
-		FRHIShaderResourceView* ShadowCapsuleShapesSRV,
-		int32 NumMeshDistanceFieldCastersValue,
-		FRHIShaderResourceView* MeshDistanceFieldCasterIndicesSRV,
-		FRHIShaderResourceView* LightDirectionDataSRV,
-		FRHITexture* ReceiverBentNormalTextureValue,
-		ECapsuleShadowingType ShadowingType)
-	{
-		FRHIComputeShader* ShaderRHI = RHICmdList.GetBoundComputeShader();
-
-		FGlobalShader::SetParameters<FViewUniformShaderParameters>(RHICmdList, ShaderRHI, View.ViewUniformBuffer);
-
-		if (TileIntersectionCountsBuffer)
-		{
-			RHICmdList.Transition(FRHITransitionInfo(TileIntersectionCountsBuffer->UAV, ERHIAccess::Unknown, ERHIAccess::UAVCompute));
-		}
-
-		if (ShadowingType == ECapsuleShadowingType::MovableSkylightTiledCulling)
-		{
-			check(!ShadowFactors.IsBound());
-			BentNormalTexture.SetTexture(RHICmdList, ShaderRHI, OutputTexture, OutputTextureUAV);
-		}
-		else
-		{
-			check(!BentNormalTexture.IsBound());
-			ShadowFactors.SetTexture(RHICmdList, ShaderRHI, OutputTexture, OutputTextureUAV);
-		}
-		
-		if (TileIntersectionCountsBuffer)
-		{
-			TileIntersectionCounts.SetBuffer(RHICmdList, ShaderRHI, *TileIntersectionCountsBuffer);
-		}
-		else
-		{
-			check(!TileIntersectionCounts.IsBound());
-		}
-
-		SetShaderValue(RHICmdList, ShaderRHI, TileDimensions, TileDimensionsValue);
-
-		if (ShadowingType == ECapsuleShadowingType::MovableSkylightTiledCulling)
-		{
-			check(ReceiverBentNormalTextureValue);
-			SetTextureParameter(RHICmdList, ShaderRHI, ReceiverBentNormalTexture, ReceiverBentNormalTextureValue);
-		}
-		else
-		{
-			check(!ReceiverBentNormalTexture.IsBound());
-		}
-
-		SetShaderValue(RHICmdList, ShaderRHI, NumGroups, NumGroupsValue);
-
-		if (LightSceneInfo)
-		{
-			check(ShadowingType == ECapsuleShadowingType::DirectionalLightTiledCulling || ShadowingType == ECapsuleShadowingType::PointLightTiledCulling);
-			
-			const FLightSceneProxy& LightProxy = *LightSceneInfo->Proxy;
-
-			FLightShaderParameters LightParameters;
-			LightProxy.GetLightShaderParameters(LightParameters);
-
-			SetShaderValue(RHICmdList, ShaderRHI, LightDirection, LightParameters.Direction);
-			FVector4 LightPositionAndInvRadiusValue(LightParameters.Position, LightParameters.InvRadius);
-			SetShaderValue(RHICmdList, ShaderRHI, LightPositionAndInvRadius, LightPositionAndInvRadiusValue);
-			// Default light source radius of 0 gives poor results
-			SetShaderValue(RHICmdList, ShaderRHI, LightSourceRadius, LightParameters.SourceRadius == 0 ? 20 : FMath::Clamp(LightParameters.SourceRadius, .001f, 1.0f / (4 * LightParameters.InvRadius)));
-
-			SetShaderValue(RHICmdList, ShaderRHI, RayStartOffsetDepthScale, LightProxy.GetRayStartOffsetDepthScale());
-
-			const float LightSourceAngle = FMath::Clamp(LightProxy.GetLightSourceAngle() * 5, 1.0f, 30.0f) * PI / 180.0f;
-			const FVector3f LightAngleAndNormalThresholdValue(LightSourceAngle, FMath::Cos(PI / 2 + LightSourceAngle), LightProxy.GetTraceDistance());
-			SetShaderValue(RHICmdList, ShaderRHI, LightAngleAndNormalThreshold, LightAngleAndNormalThresholdValue);
-		}
-		else
-		{
-			check(ShadowingType == ECapsuleShadowingType::IndirectTiledCulling || ShadowingType == ECapsuleShadowingType::MovableSkylightTiledCulling || ShadowingType == ECapsuleShadowingType::MovableSkylightTiledCullingGatherFromReceiverBentNormal);
-			check(!LightDirection.IsBound() && !LightPositionAndInvRadius.IsBound());
-		}
-
-		SetShaderValue(RHICmdList, ShaderRHI, ScissorRectMinAndSize, FIntRect(ScissorRect.Min, ScissorRect.Size()));
-		SetShaderValue(RHICmdList, ShaderRHI, DownsampleFactor, DownsampleFactorValue);
-
-		SetShaderValue(RHICmdList, ShaderRHI, NumShadowCapsules, NumShadowCapsulesValue);
-		SetSRVParameter(RHICmdList, ShaderRHI, ShadowCapsuleShapes, ShadowCapsuleShapesSRV);
-
-		SetShaderValue(RHICmdList, ShaderRHI, NumMeshDistanceFieldCasters, NumMeshDistanceFieldCastersValue);
-		SetSRVParameter(RHICmdList, ShaderRHI, MeshDistanceFieldCasterIndices, MeshDistanceFieldCasterIndicesSRV);
-
-		SetShaderValue(RHICmdList, ShaderRHI, MaxOcclusionDistance, MaxOcclusionDistanceValue);
-		const float CosFadeStartAngleValue = FMath::Cos(GCapsuleShadowFadeAngleFromVertical);
-		const FVector2D CosFadeStartAngleVector(CosFadeStartAngleValue, 1.0f / (1.0f - CosFadeStartAngleValue));
-		SetShaderValue(RHICmdList, ShaderRHI, CosFadeStartAngle, CosFadeStartAngleVector);
-		SetSRVParameter(RHICmdList, ShaderRHI, LightDirectionData, LightDirectionDataSRV);
-
-		float IndirectCapsuleSelfShadowingIntensityValue = Scene->DynamicIndirectShadowsSelfShadowingIntensity;
-		SetShaderValue(RHICmdList, ShaderRHI, IndirectCapsuleSelfShadowingIntensity, IndirectCapsuleSelfShadowingIntensityValue);
-
-		if (Scene->DistanceFieldSceneData.GetCurrentObjectBuffers())
-		{
-			DistanceFieldObjectParameters.Set(
-				RHICmdList,
-				ShaderRHI,
-				*Scene->DistanceFieldSceneData.GetCurrentObjectBuffers(),
-				Scene->DistanceFieldSceneData.NumObjectsInBuffer);
-
-			TDistanceFieldCulledObjectBuffers<DFPT_SignedDistanceField> Dummy;
-			DistanceFieldCulledObjectParameters.Set(
-				RHICmdList,
-				ShaderRHI,
-				Dummy,
-				Scene->DistanceFieldSceneData);
-		}
-		else
-		{
-			check(!DistanceFieldObjectParameters.AnyBound());
-		}
-	}
-
-	void UnsetParameters(FRHIComputeCommandList& RHICmdList, const FRWBuffer* TileIntersectionCountsBuffer)
-	{
-		ShadowFactors.UnsetUAV(RHICmdList, RHICmdList.GetBoundComputeShader());
-		BentNormalTexture.UnsetUAV(RHICmdList, RHICmdList.GetBoundComputeShader());
-		TileIntersectionCounts.UnsetUAV(RHICmdList, RHICmdList.GetBoundComputeShader());
-
-		if (TileIntersectionCountsBuffer)
-		{
-			RHICmdList.Transition(FRHITransitionInfo(TileIntersectionCountsBuffer->UAV, ERHIAccess::Unknown, ERHIAccess::SRVMask));
-		}
-	}
-
-private:
-
-	LAYOUT_FIELD(FRWShaderParameter, ShadowFactors);
-	LAYOUT_FIELD(FRWShaderParameter, TileIntersectionCounts);
-	LAYOUT_FIELD(FShaderParameter, TileDimensions);
-	LAYOUT_FIELD(FRWShaderParameter, BentNormalTexture);
-	LAYOUT_FIELD(FShaderResourceParameter, ReceiverBentNormalTexture);
-	LAYOUT_FIELD(FShaderParameter, NumGroups);
-	LAYOUT_FIELD(FShaderParameter, LightDirection);
-	LAYOUT_FIELD(FShaderParameter, LightPositionAndInvRadius);
-	LAYOUT_FIELD(FShaderParameter, LightSourceRadius);
-	LAYOUT_FIELD(FShaderParameter, RayStartOffsetDepthScale);
-	LAYOUT_FIELD(FShaderParameter, LightAngleAndNormalThreshold);
-	LAYOUT_FIELD(FShaderParameter, ScissorRectMinAndSize);
-	LAYOUT_FIELD(FShaderParameter, DownsampleFactor);
-	LAYOUT_FIELD(FShaderParameter, NumShadowCapsules);
-	LAYOUT_FIELD(FShaderResourceParameter, ShadowCapsuleShapes);
-	LAYOUT_FIELD(FShaderParameter, NumMeshDistanceFieldCasters);
-	LAYOUT_FIELD(FShaderResourceParameter, MeshDistanceFieldCasterIndices);
-	LAYOUT_FIELD(FShaderParameter, MaxOcclusionDistance);
-	LAYOUT_FIELD(FShaderParameter, CosFadeStartAngle);
-	LAYOUT_FIELD(FShaderResourceParameter, LightDirectionData);
-	LAYOUT_FIELD(FShaderParameter, IndirectCapsuleSelfShadowingIntensity);
-	LAYOUT_FIELD((TDistanceFieldObjectBufferParameters<DFPT_SignedDistanceField>), DistanceFieldObjectParameters);
-	LAYOUT_FIELD((TDistanceFieldCulledObjectBufferParameters<DFPT_SignedDistanceField>), DistanceFieldCulledObjectParameters);
 };
 
 IMPLEMENT_GLOBAL_SHADER(FCapsuleShadowingCS, "/Engine/Private/CapsuleShadowShaders.usf", "CapsuleShadowingCS", SF_Compute);
@@ -461,6 +300,15 @@ class FCapsuleShadowingUpsampleVS : public FGlobalShader
 {
 	DECLARE_SHADER_TYPE(FCapsuleShadowingUpsampleVS, Global);
 public:
+	SHADER_USE_PARAMETER_STRUCT(FCapsuleShadowingUpsampleVS, FGlobalShader);
+
+	BEGIN_SHADER_PARAMETER_STRUCT(FParameters, )
+		SHADER_PARAMETER_STRUCT_REF(FViewUniformShaderParameters, View)
+		SHADER_PARAMETER_RDG_BUFFER_SRV(StructuredBuffer<uint>, TileIntersectionCounts)
+		SHADER_PARAMETER(FIntPoint, TileDimensions)
+		SHADER_PARAMETER(FVector2D, TileSize)
+		SHADER_PARAMETER(FIntRect, ScissorRectMinAndSize)
+	END_SHADER_PARAMETER_STRUCT()
 
 	static bool ShouldCompilePermutation(const FGlobalShaderPermutationParameters& Parameters)
 	{
@@ -471,39 +319,6 @@ public:
 	{
 		OutEnvironment.SetDefine(TEXT("TILES_PER_INSTANCE"), NumTileQuadsInBuffer);
 	}
-
-	/** Default constructor. */
-	FCapsuleShadowingUpsampleVS() {}
-
-	/** Initialization constructor. */
-	FCapsuleShadowingUpsampleVS(const ShaderMetaType::CompiledShaderInitializerType& Initializer)
-		: FGlobalShader(Initializer)
-	{
-		TileDimensions.Bind(Initializer.ParameterMap,TEXT("TileDimensions"));
-		TileSize.Bind(Initializer.ParameterMap,TEXT("TileSize"));
-		ScissorRectMinAndSize.Bind(Initializer.ParameterMap,TEXT("ScissorRectMinAndSize"));
-		TileIntersectionCounts.Bind(Initializer.ParameterMap,TEXT("TileIntersectionCounts"));
-	}
-
-	void SetParameters(FRHICommandList& RHICmdList, const FSceneView& View, FIntPoint TileDimensionsValue, const FIntRect& ScissorRect, const FRWBuffer& TileIntersectionCountsBuffer)
-	{
-		FRHIVertexShader* ShaderRHI = RHICmdList.GetBoundVertexShader();
-		FGlobalShader::SetParameters<FViewUniformShaderParameters>(RHICmdList, ShaderRHI, View.ViewUniformBuffer);
-
-		SetShaderValue(RHICmdList, ShaderRHI, TileDimensions, TileDimensionsValue);
-		SetShaderValue(RHICmdList, ShaderRHI, TileSize, FVector2D(
-			GShadowShapeTileSize * GetCapsuleShadowDownsampleFactor(), 
-			GShadowShapeTileSize * GetCapsuleShadowDownsampleFactor()));
-		SetShaderValue(RHICmdList, ShaderRHI, ScissorRectMinAndSize, FIntRect(ScissorRect.Min, ScissorRect.Size()));
-		SetSRVParameter(RHICmdList, ShaderRHI, TileIntersectionCounts, TileIntersectionCountsBuffer.SRV);
-	}
-
-private:
-
-	LAYOUT_FIELD(FShaderParameter, TileDimensions);
-	LAYOUT_FIELD(FShaderParameter, TileSize);
-	LAYOUT_FIELD(FShaderParameter, ScissorRectMinAndSize);
-	LAYOUT_FIELD(FShaderResourceParameter, TileIntersectionCounts);
 };
 
 IMPLEMENT_GLOBAL_SHADER(FCapsuleShadowingUpsampleVS, "/Engine/Private/CapsuleShadowShaders.usf", "CapsuleShadowingUpsampleVS", SF_Vertex);
@@ -512,6 +327,15 @@ class FCapsuleShadowingUpsamplePS : public FGlobalShader
 {
 	DECLARE_GLOBAL_SHADER(FCapsuleShadowingUpsamplePS);
 public:
+	SHADER_USE_PARAMETER_STRUCT(FCapsuleShadowingUpsamplePS, FGlobalShader);
+
+	BEGIN_SHADER_PARAMETER_STRUCT(FParameters, )
+		SHADER_PARAMETER_STRUCT_REF(FViewUniformShaderParameters, View)
+		SHADER_PARAMETER_RDG_TEXTURE(Texture2D, ShadowFactorsTexture)
+		SHADER_PARAMETER_SAMPLER(SamplerState, ShadowFactorsSampler)
+		SHADER_PARAMETER(FIntRect, ScissorRectMinAndSize)
+		SHADER_PARAMETER(float, OutputtingToLightAttenuation)
+	END_SHADER_PARAMETER_STRUCT()
 
 	class FUpsampleRequired : SHADER_PERMUTATION_BOOL("UPSAMPLE_REQUIRED");
 	class FApplySSAO : SHADER_PERMUTATION_BOOL("APPLY_TO_SSAO");
@@ -526,66 +350,119 @@ public:
 	{
 		OutEnvironment.SetDefine(TEXT("DOWNSAMPLE_FACTOR"), 2);
 	}
+};
+
+IMPLEMENT_GLOBAL_SHADER(FCapsuleShadowingUpsamplePS, "/Engine/Private/CapsuleShadowShaders.usf", "CapsuleShadowingUpsamplePS", SF_Pixel);
+
+BEGIN_SHADER_PARAMETER_STRUCT(FUpsampleCapsuleShadowParameters, )
+	SHADER_PARAMETER_STRUCT_INCLUDE(FCapsuleShadowingUpsampleVS::FParameters, VS)
+	SHADER_PARAMETER_STRUCT_INCLUDE(FCapsuleShadowingUpsamplePS::FParameters, PS)
+	SHADER_PARAMETER_RDG_UNIFORM_BUFFER(FSceneTextureUniformParameters, SceneTextures)
+	RENDER_TARGET_BINDING_SLOTS()
+END_SHADER_PARAMETER_STRUCT()
+
+void SetupCapsuleShadowingParameters(
+	FCapsuleShadowingCS::FParameters& Parameters,
+	ECapsuleShadowingType ShadowingType,
+	FRDGTextureUAVRef OutputUAV,
+	FIntPoint TileDimensions,
+	FRDGTextureRef ReceiverBentNormalTexture,
+	FVector2D NumGroups,
+	const FLightSceneInfo* LightSceneInfo,
+	const FIntRect& ScissorRect,
+	int32 DownsampleFactor,
+	float MaxOcclusionDistance,
+	const FScene* Scene,
+	const FSceneView& View,
+	TRDGUniformBufferRef<FSceneTextureUniformParameters> SceneTexturesUniformBuffer,
+
+	uint32 NumShadowCapsules,
+	FRDGBufferSRVRef ShadowCapsuleShapesBuffer,
+	uint32 NumMeshDistanceFieldCasters,
+	FRDGBufferSRVRef MeshDistanceFieldCasterIndicesSRV,
+
+	FRDGBufferSRVRef LightDirectionDataSRV,
+
+	FRDGBufferUAVRef CapsuleTileIntersectionCountsUAV
+)
+{
+	Parameters.SceneTextures = SceneTexturesUniformBuffer;
+	Parameters.View = View.ViewUniformBuffer;
+
+	if (ShadowingType == ECapsuleShadowingType::MovableSkylightTiledCulling)
+	{
+		Parameters.RWShadowFactors = nullptr;
+		Parameters.RWBentNormalTexture = OutputUAV;
+	}
+	else
+	{
+		Parameters.RWShadowFactors = OutputUAV;
+		Parameters.RWBentNormalTexture = nullptr;
+	}
 
-	/** Default constructor. */
-	FCapsuleShadowingUpsamplePS() {}
+	Parameters.RWTileIntersectionCounts = CapsuleTileIntersectionCountsUAV;
 
-	/** Initialization constructor. */
-	FCapsuleShadowingUpsamplePS(const ShaderMetaType::CompiledShaderInitializerType& Initializer)
-		: FGlobalShader(Initializer)
+	Parameters.TileDimensions = TileDimensions;
+
+	if (ShadowingType == ECapsuleShadowingType::MovableSkylightTiledCulling)
+	{
+		check(ReceiverBentNormalTexture);
+		Parameters.ReceiverBentNormalTexture = ReceiverBentNormalTexture;
+	}
+	else
 	{
-		ShadowFactorsTexture.Bind(Initializer.ParameterMap,TEXT("ShadowFactorsTexture"));
-		ShadowFactorsSampler.Bind(Initializer.ParameterMap,TEXT("ShadowFactorsSampler"));
-		ScissorRectMinAndSize.Bind(Initializer.ParameterMap,TEXT("ScissorRectMinAndSize"));
-		OutputtingToLightAttenuation.Bind(Initializer.ParameterMap,TEXT("OutputtingToLightAttenuation"));
+		Parameters.ReceiverBentNormalTexture = nullptr;
 	}
 
-	void SetParameters(FRHICommandList& RHICmdList, const FSceneView& View, const FIntRect& ScissorRect, FRHITexture* ShadowFactorsTextureRHI, bool bOutputtingToLightAttenuation)
+	Parameters.NumGroups = NumGroups;
+
+	if (LightSceneInfo)
 	{
-		FRHIPixelShader* ShaderRHI = RHICmdList.GetBoundPixelShader();
+		check(ShadowingType == ECapsuleShadowingType::DirectionalLightTiledCulling || ShadowingType == ECapsuleShadowingType::PointLightTiledCulling);
 
-		FGlobalShader::SetParameters<FViewUniformShaderParameters>(RHICmdList, ShaderRHI, View.ViewUniformBuffer);
+		const FLightSceneProxy& LightProxy = *LightSceneInfo->Proxy;
 
-		SetTextureParameter(RHICmdList, ShaderRHI, ShadowFactorsTexture, ShadowFactorsSampler, TStaticSamplerState<SF_Bilinear>::GetRHI(), ShadowFactorsTextureRHI);
-	
-		SetShaderValue(RHICmdList, ShaderRHI, ScissorRectMinAndSize, FIntRect(ScissorRect.Min, ScissorRect.Size()));
-		SetShaderValue(RHICmdList, ShaderRHI, OutputtingToLightAttenuation, bOutputtingToLightAttenuation ? 1.0f : 0.0f);
-	}
+		FLightShaderParameters LightParameters;
+		LightProxy.GetLightShaderParameters(LightParameters);
 
-private:
-	LAYOUT_FIELD(FShaderResourceParameter, ShadowFactorsTexture);
-	LAYOUT_FIELD(FShaderResourceParameter, ShadowFactorsSampler);
-	LAYOUT_FIELD(FShaderParameter, ScissorRectMinAndSize);
-	LAYOUT_FIELD(FShaderParameter, OutputtingToLightAttenuation);
-};
+		Parameters.LightDirection = LightParameters.Direction;
 
-IMPLEMENT_GLOBAL_SHADER(FCapsuleShadowingUpsamplePS, "/Engine/Private/CapsuleShadowShaders.usf", "CapsuleShadowingUpsamplePS", SF_Pixel);
+		FVector4 LightPositionAndInvRadius(LightParameters.Position, LightParameters.InvRadius);
+		Parameters.LightPositionAndInvRadius = LightPositionAndInvRadius;
 
-void AllocateCapsuleTileIntersectionCountsBuffer(FIntPoint GroupSize, FSceneViewState* ViewState)
-{
-	EPixelFormat CapsuleTileIntersectionCountsBufferFormat = PF_R32_UINT;
+		// Default light source radius of 0 gives poor results
+		Parameters.LightSourceRadius = LightParameters.SourceRadius == 0 ? 20 : FMath::Clamp(LightParameters.SourceRadius, .001f, 1.0f / (4 * LightParameters.InvRadius));
+
+		Parameters.RayStartOffsetDepthScale = LightProxy.GetRayStartOffsetDepthScale();
 
-	if (!IsValidRef(ViewState->CapsuleTileIntersectionCountsBuffer.Buffer) 
-		|| (int32)ViewState->CapsuleTileIntersectionCountsBuffer.NumBytes < GroupSize.X * GroupSize.Y * GPixelFormats[CapsuleTileIntersectionCountsBufferFormat].BlockBytes)
+		const float LightSourceAngle = FMath::Clamp(LightProxy.GetLightSourceAngle() * 5, 1.0f, 30.0f) * PI / 180.0f;
+		const FVector3f LightAngleAndNormalThreshold(LightSourceAngle, FMath::Cos(PI / 2 + LightSourceAngle), LightProxy.GetTraceDistance());
+		Parameters.LightAngleAndNormalThreshold = LightAngleAndNormalThreshold;
+	}
+	else
 	{
-		ViewState->CapsuleTileIntersectionCountsBuffer.Release();
-		ViewState->CapsuleTileIntersectionCountsBuffer.Initialize(TEXT("CapsuleTileIntersectionCountsBuffer"), GPixelFormats[CapsuleTileIntersectionCountsBufferFormat].BlockBytes, GroupSize.X * GroupSize.Y, CapsuleTileIntersectionCountsBufferFormat);
+		check(ShadowingType == ECapsuleShadowingType::IndirectTiledCulling || ShadowingType == ECapsuleShadowingType::MovableSkylightTiledCulling || ShadowingType == ECapsuleShadowingType::MovableSkylightTiledCullingGatherFromReceiverBentNormal);
 	}
-}
 
-// TODO(RDG) Move these into the shader FParameters.
-BEGIN_SHADER_PARAMETER_STRUCT(FTiledCapsuleShadowParameters, )
-	SHADER_PARAMETER_RDG_UNIFORM_BUFFER(FSceneTextureUniformParameters, SceneTextures)
-	SHADER_PARAMETER_RDG_TEXTURE_UAV(RWTexture2D, RayTracedShadowsUAV)
-	RDG_TEXTURE_ACCESS(RayTracedShadows, ERHIAccess::UAVCompute)
-END_SHADER_PARAMETER_STRUCT()
+	Parameters.ScissorRectMinAndSize = FIntRect(ScissorRect.Min, ScissorRect.Size());
+	Parameters.DownsampleFactor = DownsampleFactor;
 
-// TODO(RDG) Move these into the shader FParameters.
-BEGIN_SHADER_PARAMETER_STRUCT(FUpsampleCapsuleShadowParameters, )
-	SHADER_PARAMETER_RDG_UNIFORM_BUFFER(FSceneTextureUniformParameters, SceneTextures)
-	RDG_TEXTURE_ACCESS(RayTracedShadows, ERHIAccess::SRVGraphics)
-	RENDER_TARGET_BINDING_SLOTS()
-END_SHADER_PARAMETER_STRUCT()
+	Parameters.IndirectCapsuleSelfShadowingIntensity = Scene->DynamicIndirectShadowsSelfShadowingIntensity;
+	Parameters.MaxOcclusionDistance = MaxOcclusionDistance;
+
+	Parameters.NumShadowCapsules = NumShadowCapsules;
+	Parameters.ShadowCapsuleShapes = ShadowCapsuleShapesBuffer;
+
+	Parameters.NumMeshDistanceFieldCasters = NumMeshDistanceFieldCasters;
+	Parameters.MeshDistanceFieldCasterIndices = MeshDistanceFieldCasterIndicesSRV;
+	Parameters.LightDirectionData = LightDirectionDataSRV;
+
+	const float CosFadeStartAngleValue = FMath::Cos(GCapsuleShadowFadeAngleFromVertical);
+	Parameters.CosFadeStartAngle = FVector2D(CosFadeStartAngleValue, 1.0f / (1.0f - CosFadeStartAngleValue));
+
+	Parameters.DFObjectBufferParameters = DistanceField::SetupObjectBufferParameters(Scene->DistanceFieldSceneData);
+	Parameters.DFAtlasParameters = DistanceField::SetupAtlasParameters(Scene->DistanceFieldSceneData);
+}
 
 bool FDeferredShadingSceneRenderer::RenderCapsuleDirectShadows(
 	FRDGBuilder& GraphBuilder,
@@ -675,89 +552,89 @@ bool FDeferredShadingSceneRenderer::RenderCapsuleDirectShadows(
 				FMath::DivideAndRoundUp(ScissorRect.Size().X / GetCapsuleShadowDownsampleFactor(), GShadowShapeTileSize),
 				FMath::DivideAndRoundUp(ScissorRect.Size().Y / GetCapsuleShadowDownsampleFactor(), GShadowShapeTileSize));
 
-			AllocateCapsuleTileIntersectionCountsBuffer(GroupSize, View.ViewState);
-			int32 NumCapsuleShapeData = CapsuleShapeData.Num();
-			AddPass(GraphBuilder, RDG_EVENT_NAME("UploadCapsuleShapes"), [&View, &LightSceneInfo, CapsuleShapeData = MoveTemp(CapsuleShapeData)](FRHICommandListImmediate& RHICmdList)
-			{
-				static_assert(sizeof(FCapsuleShape) == sizeof(FVector4) * 2, "FCapsuleShape has padding");
-				const int32 DataSize = CapsuleShapeData.Num() * CapsuleShapeData.GetTypeSize();
-
-				if (!IsValidRef(LightSceneInfo.ShadowCapsuleShapesVertexBuffer) || (int32)LightSceneInfo.ShadowCapsuleShapesVertexBuffer->GetSize() < DataSize)
-				{
-					LightSceneInfo.ShadowCapsuleShapesVertexBuffer.SafeRelease();
-					LightSceneInfo.ShadowCapsuleShapesSRV.SafeRelease();
-					FRHIResourceCreateInfo CreateInfo(TEXT("LightSceneInfo"));
-					LightSceneInfo.ShadowCapsuleShapesVertexBuffer = RHICreateVertexBuffer(DataSize, BUF_Volatile | BUF_ShaderResource, CreateInfo);
-					LightSceneInfo.ShadowCapsuleShapesSRV = RHICreateShaderResourceView(LightSceneInfo.ShadowCapsuleShapesVertexBuffer, sizeof(FVector4), PF_A32B32G32R32F);
-				}
+			FRDGBufferRef CapsuleTileIntersectionCountsBuffer = GraphBuilder.CreateBuffer(FRDGBufferDesc::CreateStructuredDesc(sizeof(uint32), GroupSize.X * GroupSize.Y), TEXT("CapsuleTileIntersectionCountsBuffer"));
+			FRDGBufferUAVRef CapsuleTileIntersectionCountsUAV = GraphBuilder.CreateUAV(CapsuleTileIntersectionCountsBuffer);
+			FRDGBufferSRVRef CapsuleTileIntersectionCountsSRV = GraphBuilder.CreateSRV(CapsuleTileIntersectionCountsBuffer);
 
-				void* CapsuleShapeLockedData = RHILockBuffer(LightSceneInfo.ShadowCapsuleShapesVertexBuffer, 0, DataSize, RLM_WriteOnly);
-				FPlatformMemory::Memcpy(CapsuleShapeLockedData, CapsuleShapeData.GetData(), DataSize);
-				RHIUnlockBuffer(LightSceneInfo.ShadowCapsuleShapesVertexBuffer);
+			AddClearUAVPass(GraphBuilder, CapsuleTileIntersectionCountsUAV, 0);
 
-				RHICmdList.Transition(FRHITransitionInfo(View.ViewState->CapsuleTileIntersectionCountsBuffer.UAV, ERHIAccess::Unknown, ERHIAccess::UAVCompute));
-				RHICmdList.ClearUAVUint(View.ViewState->CapsuleTileIntersectionCountsBuffer.UAV, FUintVector4(0, 0, 0, 0));
-			});
+			// Upload capsule shape data
+			// TODO: Use FRDGUploadData
+			FRDGBufferRef ShadowCapsuleShapesBuffer = CreateStructuredBuffer(GraphBuilder, TEXT("ShadowCapsuleShapesBuffer"), CapsuleShapeData);
 
 			{
-				FTiledCapsuleShadowParameters* PassParameters = GraphBuilder.AllocParameters<FTiledCapsuleShadowParameters>();
-				PassParameters->RayTracedShadows = RayTracedShadowsRT;
-				PassParameters->RayTracedShadowsUAV = GraphBuilder.CreateUAV(RayTracedShadowsRT);
-				PassParameters->SceneTextures = SceneTexturesUniformBuffer;
-
 				ECapsuleShadowingType ShadowingType = bDirectionalLight ? ECapsuleShadowingType::DirectionalLightTiledCulling : ECapsuleShadowingType::PointLightTiledCulling;
 
+				auto* PassParameters = GraphBuilder.AllocParameters<FCapsuleShadowingCS::FParameters>();
+
+				SetupCapsuleShadowingParameters(
+					*PassParameters,
+					ShadowingType,
+					GraphBuilder.CreateUAV(RayTracedShadowsRT),
+					GroupSize,
+					nullptr,
+					FVector2D(GroupSize.X, GroupSize.Y),
+					&LightSceneInfo,
+					ScissorRect,
+					GetCapsuleShadowDownsampleFactor(),
+					GCapsuleMaxDirectOcclusionDistance,
+					Scene,
+					View,
+					SceneTexturesUniformBuffer,
+					CapsuleShapeData.Num(),
+					GraphBuilder.CreateSRV(ShadowCapsuleShapesBuffer),
+					0,
+					nullptr,
+					nullptr,
+					CapsuleTileIntersectionCountsUAV
+				);
+
 				FCapsuleShadowingCS::FPermutationDomain PermutationVector;
 				PermutationVector.Set<FCapsuleShadowingCS::FShapeShadow>(ShadowingType);
 				PermutationVector.Set<FCapsuleShadowingCS::FIndirectPrimitiveType>(EIndirectShadowingPrimitiveTypes::CapsuleShapes);
 				auto ComputeShader = View.ShaderMap->GetShader<FCapsuleShadowingCS>(PermutationVector);
 
-				GraphBuilder.AddPass(
+				FComputeShaderUtils::AddPass(
+					GraphBuilder,
 					RDG_EVENT_NAME("TiledCapsuleShadowing"),
-					PassParameters,
 					ERDGPassFlags::Compute,
-					[ComputeShader, this, &View, &LightSceneInfo, RayTracedShadows = PassParameters->RayTracedShadows.GetTexture(), RayTracedShadowsUAV = PassParameters->RayTracedShadowsUAV, GroupSize, ScissorRect, NumCapsuleShapeData, ShadowingType](FRHIComputeCommandList& RHICmdList)
-					{
-						FRHIComputeShader* ShaderRHI = ComputeShader.GetComputeShader();
-						RHICmdList.SetComputeShader(ShaderRHI);
-
-						ComputeShader->SetParameters(
-							RHICmdList,
-							Scene,
-							View,
-							&LightSceneInfo,
-							RayTracedShadows->GetRHI(),
-							RayTracedShadowsUAV->GetRHI(),
-							GroupSize,
-							&View.ViewState->CapsuleTileIntersectionCountsBuffer,
-							FVector2D(GroupSize.X, GroupSize.Y),
-							GCapsuleMaxDirectOcclusionDistance,
-							ScissorRect,
-							GetCapsuleShadowDownsampleFactor(),
-							NumCapsuleShapeData,
-							LightSceneInfo.ShadowCapsuleShapesSRV.GetReference(),
-							0,
-							nullptr,
-							nullptr,
-							nullptr,
-							ShadowingType);
-
-						DispatchComputeShader(RHICmdList, ComputeShader.GetShader(), GroupSize.X, GroupSize.Y, 1);
-						ComputeShader->UnsetParameters(RHICmdList, &View.ViewState->CapsuleTileIntersectionCountsBuffer);
-					});
+					ComputeShader,
+					PassParameters,
+					FIntVector(GroupSize.X, GroupSize.Y, 1));
 			}
 
 			{
+				auto VertexShader = View.ShaderMap->GetShader<FCapsuleShadowingUpsampleVS>();
+
+				FCapsuleShadowingUpsamplePS::FPermutationDomain PermutationVector;
+				PermutationVector.Set<FCapsuleShadowingUpsamplePS::FUpsampleRequired>(!GCapsuleShadowsFullResolution);
+				PermutationVector.Set<FCapsuleShadowingUpsamplePS::FApplySSAO>(false);
+				auto PixelShader = View.ShaderMap->GetShader<FCapsuleShadowingUpsamplePS>(PermutationVector);
+
 				FUpsampleCapsuleShadowParameters* PassParameters = GraphBuilder.AllocParameters<FUpsampleCapsuleShadowParameters>();
 				PassParameters->RenderTargets[0] = FRenderTargetBinding(ScreenShadowMaskTexture, ERenderTargetLoadAction::ELoad);
-				PassParameters->RayTracedShadows = RayTracedShadowsRT;
 				PassParameters->SceneTextures = SceneTexturesUniformBuffer;
 
+				PassParameters->VS.View = GetShaderBinding(View.ViewUniformBuffer);
+				PassParameters->VS.TileIntersectionCounts = CapsuleTileIntersectionCountsSRV;
+				PassParameters->VS.TileDimensions = GroupSize;
+				PassParameters->VS.TileSize = FVector2D(GShadowShapeTileSize * GetCapsuleShadowDownsampleFactor(), GShadowShapeTileSize * GetCapsuleShadowDownsampleFactor());
+				PassParameters->VS.ScissorRectMinAndSize = FIntRect(ScissorRect.Min, ScissorRect.Size());
+
+				PassParameters->PS.View = GetShaderBinding(View.ViewUniformBuffer);
+				PassParameters->PS.ShadowFactorsTexture = RayTracedShadowsRT;
+				PassParameters->PS.ShadowFactorsSampler = TStaticSamplerState<SF_Bilinear>::GetRHI();
+				PassParameters->PS.ScissorRectMinAndSize = FIntRect(ScissorRect.Min, ScissorRect.Size());
+				PassParameters->PS.OutputtingToLightAttenuation = 1.0f;
+
+				ClearUnusedGraphResources(VertexShader, &PassParameters->VS);
+				ClearUnusedGraphResources(PixelShader, &PassParameters->PS);
+
 				GraphBuilder.AddPass(
 					RDG_EVENT_NAME("UpsampleCapsuleShadow %dx%d", ScissorRect.Width(), ScissorRect.Height()),
 					PassParameters,
 					ERDGPassFlags::Raster,
-					[this, &View, &LightSceneInfo, RayTracedShadowsRT, GroupSize, ScissorRect, bProjectingForForwardShading](FRHICommandList& RHICmdList)
+					[PassParameters, VertexShader, PixelShader, &View, &LightSceneInfo, RayTracedShadowsRT, GroupSize, ScissorRect, bProjectingForForwardShading](FRHICommandList& RHICmdList)
 				{
 					FGraphicsPipelineStateInitializer GraphicsPSOInit;
 					RHICmdList.ApplyCachedRenderTargets(GraphicsPSOInit);
@@ -773,13 +650,6 @@ bool FDeferredShadingSceneRenderer::RenderCapsuleDirectShadows(
 						bProjectingForForwardShading,
 						false);
 
-					auto VertexShader = View.ShaderMap->GetShader<FCapsuleShadowingUpsampleVS>();
-
-					FCapsuleShadowingUpsamplePS::FPermutationDomain PermutationVector;
-					PermutationVector.Set<FCapsuleShadowingUpsamplePS::FUpsampleRequired>(!GCapsuleShadowsFullResolution);
-					PermutationVector.Set<FCapsuleShadowingUpsamplePS::FApplySSAO>(false);
-					auto PixelShader = View.ShaderMap->GetShader<FCapsuleShadowingUpsamplePS>(PermutationVector);
-
 					GraphicsPSOInit.PrimitiveType = PT_TriangleList;
 					GraphicsPSOInit.BoundShaderState.VertexDeclarationRHI = GTileVertexDeclaration.VertexDeclarationRHI;
 					GraphicsPSOInit.BoundShaderState.VertexShaderRHI = VertexShader.GetVertexShader();
@@ -787,8 +657,8 @@ bool FDeferredShadingSceneRenderer::RenderCapsuleDirectShadows(
 
 					SetGraphicsPipelineState(RHICmdList, GraphicsPSOInit);
 
-					VertexShader->SetParameters(RHICmdList, View, GroupSize, ScissorRect, View.ViewState->CapsuleTileIntersectionCountsBuffer);
-					PixelShader->SetParameters(RHICmdList, View, ScissorRect, RayTracedShadowsRT->GetRHI(), true);
+					SetShaderParameters(RHICmdList, VertexShader, VertexShader.GetVertexShader(), PassParameters->VS);
+					SetShaderParameters(RHICmdList, PixelShader, PixelShader.GetPixelShader(), PassParameters->PS);
 
 					RHICmdList.SetStreamSource(0, GTileTexCoordVertexBuffer.VertexBufferRHI, 0);
 					RHICmdList.DrawIndexedPrimitive(GTileIndexBuffer.IndexBufferRHI, 
@@ -859,14 +729,26 @@ void FDeferredShadingSceneRenderer::CreateIndirectCapsuleShadows()
 	}
 }
 
-void FDeferredShadingSceneRenderer::SetupIndirectCapsuleShadows(
-	FRDGBuilder& GraphBuilder, 
+struct IndirectCapsuleShadowsResources
+{
+	FRDGBufferSRVRef IndirectShadowCapsuleShapesSRV;
+	FRDGBufferSRVRef IndirectShadowMeshDistanceFieldCasterIndicesSRV;
+	FRDGBufferSRVRef IndirectShadowLightDirectionSRV;
+};
+
+static IndirectCapsuleShadowsResources CreateIndirectCapsuleShadowsResources(
+	FRDGBuilder& GraphBuilder,
+	const FScene* Scene,
 	const FViewInfo& View, 
 	int32& NumCapsuleShapes, 
 	int32& NumMeshesWithCapsules, 
-	int32& NumMeshDistanceFieldCasters,
-	FRHIShaderResourceView*& IndirectShadowLightDirectionSRV) const
+	int32& NumMeshDistanceFieldCasters)
 {
+	IndirectCapsuleShadowsResources Output;
+	Output.IndirectShadowCapsuleShapesSRV = nullptr;
+	Output.IndirectShadowMeshDistanceFieldCasterIndicesSRV = nullptr;
+	Output.IndirectShadowLightDirectionSRV = nullptr;
+
 	const float CosFadeStartAngle = FMath::Cos(GCapsuleShadowFadeAngleFromVertical);
 	const FSkyLightSceneProxy* SkyLight = Scene ? Scene->SkyLight : NULL;
 
@@ -879,7 +761,6 @@ void FDeferredShadingSceneRenderer::SetupIndirectCapsuleShadows(
 	MeshDistanceFieldCasterIndices.Reset();
 	CapsuleLightSourceData.Reset();
 	DistanceFieldCasterLightSourceData.Reset();
-	IndirectShadowLightDirectionSRV = NULL;
 
 	const bool bComputeLightDataFromVolumetricLightmapOrGpuSkyEnvMapIrradiance = Scene && (Scene->VolumetricLightmapSceneData.HasData() || (Scene->SkyLight && Scene->SkyLight->bRealTimeCaptureEnabled));
 
@@ -985,81 +866,49 @@ void FDeferredShadingSceneRenderer::SetupIndirectCapsuleShadows(
 
 	if (CapsuleShapeData.Num() > 0 || MeshDistanceFieldCasterIndices.Num() > 0)
 	{
-		static_assert(sizeof(FCapsuleShape) == sizeof(FVector4)* 2, "FCapsuleShape has padding");
 		if (CapsuleShapeData.Num() > 0)
 		{
-			const int32 DataSize = CapsuleShapeData.Num() * CapsuleShapeData.GetTypeSize();
-
-			if (!IsValidRef(View.ViewState->IndirectShadowCapsuleShapesVertexBuffer) || (int32)View.ViewState->IndirectShadowCapsuleShapesVertexBuffer->GetSize() < DataSize)
-			{
-				View.ViewState->IndirectShadowCapsuleShapesVertexBuffer.SafeRelease();
-				View.ViewState->IndirectShadowCapsuleShapesSRV.SafeRelease();
-				FRHIResourceCreateInfo CreateInfo(TEXT("IndirectShadowCapsuleShapesVertexBuffer"));
-				View.ViewState->IndirectShadowCapsuleShapesVertexBuffer = RHICreateVertexBuffer(DataSize, BUF_Volatile | BUF_ShaderResource, CreateInfo);
-				View.ViewState->IndirectShadowCapsuleShapesSRV = RHICreateShaderResourceView(View.ViewState->IndirectShadowCapsuleShapesVertexBuffer, sizeof(FVector4), PF_A32B32G32R32F);
-			}
-
-			void* CapsuleShapeLockedData = RHILockBuffer(View.ViewState->IndirectShadowCapsuleShapesVertexBuffer, 0, DataSize, RLM_WriteOnly);
-			FPlatformMemory::Memcpy(CapsuleShapeLockedData, CapsuleShapeData.GetData(), DataSize);
-			RHIUnlockBuffer(View.ViewState->IndirectShadowCapsuleShapesVertexBuffer);
+			// Upload capsule shape data
+			// TODO: Use FRDGUploadData
+			FRDGBufferRef IndirectShadowCapsuleShapesVertexBuffer = CreateStructuredBuffer(GraphBuilder, TEXT("IndirectShadowCapsuleShapesVertexBuffer"), CapsuleShapeData);
+			Output.IndirectShadowCapsuleShapesSRV = GraphBuilder.CreateSRV(IndirectShadowCapsuleShapesVertexBuffer);
 		}
 
 		if (MeshDistanceFieldCasterIndices.Num() > 0)
 		{
-			const int32 DataSize = MeshDistanceFieldCasterIndices.Num() * MeshDistanceFieldCasterIndices.GetTypeSize();
-
-			if (!IsValidRef(View.ViewState->IndirectShadowMeshDistanceFieldCasterIndicesVertexBuffer) || (int32)View.ViewState->IndirectShadowMeshDistanceFieldCasterIndicesVertexBuffer->GetSize() < DataSize)
-			{
-				View.ViewState->IndirectShadowMeshDistanceFieldCasterIndicesVertexBuffer.SafeRelease();
-				View.ViewState->IndirectShadowMeshDistanceFieldCasterIndicesSRV.SafeRelease();
-				FRHIResourceCreateInfo CreateInfo(TEXT("IndirectShadowMeshDistanceFieldCasterIndicesVertexBuffer"));
-				View.ViewState->IndirectShadowMeshDistanceFieldCasterIndicesVertexBuffer = RHICreateVertexBuffer(DataSize, BUF_Volatile | BUF_ShaderResource, CreateInfo);
-				View.ViewState->IndirectShadowMeshDistanceFieldCasterIndicesSRV = RHICreateShaderResourceView(View.ViewState->IndirectShadowMeshDistanceFieldCasterIndicesVertexBuffer, sizeof(uint32), PF_R32_UINT);
-			}
-
-			void* LockedData = RHILockBuffer(View.ViewState->IndirectShadowMeshDistanceFieldCasterIndicesVertexBuffer, 0, DataSize, RLM_WriteOnly);
-			FPlatformMemory::Memcpy(LockedData, MeshDistanceFieldCasterIndices.GetData(), DataSize);
-			RHIUnlockBuffer(View.ViewState->IndirectShadowMeshDistanceFieldCasterIndicesVertexBuffer);
+			// Upload mesh distance field caster indices
+			// TODO: Use FRDGUploadData
+			FRDGBufferRef IndirectShadowMeshDistanceFieldCasterIndicesVertexBuffer = CreateStructuredBuffer(GraphBuilder, TEXT("IndirectShadowMeshDistanceFieldCasterIndicesVertexBuffer"), MeshDistanceFieldCasterIndices);
+			Output.IndirectShadowMeshDistanceFieldCasterIndicesSRV = GraphBuilder.CreateSRV(IndirectShadowMeshDistanceFieldCasterIndicesVertexBuffer);
 		}
 
-		EPixelFormat LightDirectionDataFormat = PF_A32B32G32R32F;
+		const int32 NumLightDataElements = CapsuleLightSourceData.Num() + DistanceFieldCasterLightSourceData.Num();
 
+		FRDGBufferRef IndirectShadowLightDirectionVertexBuffer = nullptr;
+		FRDGBufferSRVRef IndirectShadowLightDirectionSRV = nullptr;
 		{
-			size_t CapsuleLightSourceDataSize = CapsuleLightSourceData.Num() * CapsuleLightSourceData.GetTypeSize();
-			const int32 DataSize = CapsuleLightSourceDataSize + DistanceFieldCasterLightSourceData.Num() * DistanceFieldCasterLightSourceData.GetTypeSize();
-			check(DataSize > 0);
-
-			if (!IsValidRef(View.ViewState->IndirectShadowLightDirectionVertexBuffer) || (int32)View.ViewState->IndirectShadowLightDirectionVertexBuffer->GetSize() < DataSize)
-			{
-				View.ViewState->IndirectShadowLightDirectionVertexBuffer.SafeRelease();
-				View.ViewState->IndirectShadowLightDirectionSRV.SafeRelease();
-				FRHIResourceCreateInfo CreateInfo(TEXT("IndirectShadowLightDirectionVertexBuffer"));
-				View.ViewState->IndirectShadowLightDirectionVertexBuffer = RHICreateVertexBuffer(DataSize, BUF_Volatile | BUF_ShaderResource, CreateInfo);
-				View.ViewState->IndirectShadowLightDirectionSRV = RHICreateShaderResourceView(View.ViewState->IndirectShadowLightDirectionVertexBuffer, sizeof(FVector4), LightDirectionDataFormat);
-			}
-
-			FVector4* LightDirectionLockedData = (FVector4*)RHILockBuffer(View.ViewState->IndirectShadowLightDirectionVertexBuffer, 0, DataSize, RLM_WriteOnly);
-			FPlatformMemory::Memcpy(LightDirectionLockedData, CapsuleLightSourceData.GetData(), CapsuleLightSourceDataSize);
+			// Upload light data
 			// Light data for distance fields is placed after capsule light data
 			// This packing behavior must match GetLightDirectionData
-			FPlatformMemory::Memcpy((char*)LightDirectionLockedData + CapsuleLightSourceDataSize, DistanceFieldCasterLightSourceData.GetData(), DistanceFieldCasterLightSourceData.Num() * DistanceFieldCasterLightSourceData.GetTypeSize());
-			RHIUnlockBuffer(View.ViewState->IndirectShadowLightDirectionVertexBuffer);
 
-			IndirectShadowLightDirectionSRV = View.ViewState->IndirectShadowLightDirectionSRV;
+			const size_t CapsuleLightSourceDataSize = CapsuleLightSourceData.Num() * CapsuleLightSourceData.GetTypeSize();
+			const size_t DistanceFieldCasterLightSourceDataSize = DistanceFieldCasterLightSourceData.Num() * DistanceFieldCasterLightSourceData.GetTypeSize();
+
+			FRDGUploadData<FVector4> TempData(GraphBuilder, NumLightDataElements);
+			FPlatformMemory::Memcpy(TempData.GetData(), CapsuleLightSourceData.GetData(), CapsuleLightSourceDataSize);
+			FPlatformMemory::Memcpy((char*)TempData.GetData() + CapsuleLightSourceDataSize, DistanceFieldCasterLightSourceData.GetData(), DistanceFieldCasterLightSourceDataSize);
+
+			IndirectShadowLightDirectionVertexBuffer = CreateStructuredBuffer(GraphBuilder, TEXT("IndirectShadowLightDirectionVertexBuffer"), TempData);
+
+			IndirectShadowLightDirectionSRV = GraphBuilder.CreateSRV(IndirectShadowLightDirectionVertexBuffer);
+			Output.IndirectShadowLightDirectionSRV = IndirectShadowLightDirectionSRV;
 		}
 
 		if (bComputeLightDataFromVolumetricLightmapOrGpuSkyEnvMapIrradiance)
 		{
-			int32 NumLightDataElements = CapsuleLightSourceData.Num() + DistanceFieldCasterLightSourceData.Num();
-
-			if (!IsValidRef(View.ViewState->IndirectShadowVolumetricLightmapDerivedLightDirection.Buffer) 
-				|| (int32)View.ViewState->IndirectShadowVolumetricLightmapDerivedLightDirection.NumBytes != View.ViewState->IndirectShadowLightDirectionVertexBuffer->GetSize())
-			{
-				View.ViewState->IndirectShadowVolumetricLightmapDerivedLightDirection.Release();
-				View.ViewState->IndirectShadowVolumetricLightmapDerivedLightDirection.Initialize(TEXT("IndirectShadowVolumetricLightmapDerivedLightDirection"), GPixelFormats[LightDirectionDataFormat].BlockBytes, NumLightDataElements, LightDirectionDataFormat);
-			}
-
-			IndirectShadowLightDirectionSRV = View.ViewState->IndirectShadowVolumetricLightmapDerivedLightDirection.SRV;
+			FRDGBufferRef IndirectShadowVolumetricLightmapDerivedLightDirection = GraphBuilder.CreateBuffer(FRDGBufferDesc::CreateStructuredDesc(sizeof(FVector4), NumLightDataElements), TEXT("IndirectShadowVolumetricLightmapDerivedLightDirection"));
+			FRDGBufferUAVRef ComputedLightDirectionUAV = GraphBuilder.CreateUAV(IndirectShadowVolumetricLightmapDerivedLightDirection);
+			Output.IndirectShadowLightDirectionSRV = GraphBuilder.CreateSRV(IndirectShadowVolumetricLightmapDerivedLightDirection);
 
 			TShaderMapRef<FComputeLightDirectionFromVolumetricLightmapCS> ComputeShader(View.ShaderMap);
 
@@ -1068,13 +917,6 @@ void FDeferredShadingSceneRenderer::SetupIndirectCapsuleShadows(
 
 			const int32 GroupSize = FMath::DivideAndRoundUp(NumLightDataElements, GComputeLightDirectionFromVolumetricLightmapGroupSize);
 
-			FRWBuffer& ComputedLightDirectionData = View.ViewState->IndirectShadowVolumetricLightmapDerivedLightDirection;
-
-			AddPass(GraphBuilder, RDG_EVENT_NAME("TransitionLightDirectionUAV"), [&ComputedLightDirectionData](FRHIComputeCommandList& RHICmdList)
-			{
-				RHICmdList.Transition(FRHITransitionInfo(ComputedLightDirectionData.UAV, ERHIAccess::Unknown, ERHIAccess::UAVCompute));
-			});
-
 			FComputeLightDirectionFromVolumetricLightmapCS::FParameters* PassParameters = GraphBuilder.AllocParameters<FComputeLightDirectionFromVolumetricLightmapCS::FParameters>();
 			PassParameters->View = View.ViewUniformBuffer;
 			PassParameters->NumLightDirectionData = NumLightDataElements;
@@ -1082,26 +924,23 @@ void FDeferredShadingSceneRenderer::SetupIndirectCapsuleShadows(
 			PassParameters->CapsuleIndirectConeAngle = GCapsuleSkyAngleScale;
 			PassParameters->CapsuleSkyAngleScale = GCapsuleSkyAngleScale;
 			PassParameters->CapsuleMinSkyAngle = GCapsuleMinSkyAngle;
-			PassParameters->RWComputedLightDirectionData = ComputedLightDirectionData.UAV;
-			PassParameters->LightDirectionData = View.ViewState->IndirectShadowLightDirectionSRV;
+			PassParameters->RWComputedLightDirectionData = ComputedLightDirectionUAV;
+			PassParameters->LightDirectionData = IndirectShadowLightDirectionSRV;
 
 			FComputeShaderUtils::AddPass(
 				GraphBuilder,
 				RDG_EVENT_NAME("LightDirectionFromVolumetricLightmap"),
-				ERDGPassFlags::Compute | ERDGPassFlags::NeverCull,
+				ERDGPassFlags::Compute,
 				ComputeShader,
 				PassParameters,
 				FIntVector(GroupSize, 1, 1));
-
-			AddPass(GraphBuilder, RDG_EVENT_NAME("TransitionLightDirectionUAV"), [&ComputedLightDirectionData](FRHIComputeCommandList& RHICmdList)
-			{
-				RHICmdList.Transition(FRHITransitionInfo(ComputedLightDirectionData.UAV, ERHIAccess::UAVCompute, ERHIAccess::SRVMask));
-			});
 		}
 	}
 
 	NumCapsuleShapes = CapsuleShapeData.Num();
 	NumMeshDistanceFieldCasters = MeshDistanceFieldCasterIndices.Num();
+
+	return Output;
 }
 
 void FDeferredShadingSceneRenderer::RenderIndirectCapsuleShadows(FRDGBuilder& GraphBuilder, const FSceneTextures& SceneTextures) const
@@ -1170,15 +1009,14 @@ void FDeferredShadingSceneRenderer::RenderIndirectCapsuleShadows(FRDGBuilder& Gr
 			int32 NumCapsuleShapes = 0;
 			int32 NumMeshesWithCapsules = 0;
 			int32 NumMeshDistanceFieldCasters = 0;
-			FRHIShaderResourceView* IndirectShadowLightDirectionSRV = NULL;
-			SetupIndirectCapsuleShadows(GraphBuilder, View, NumCapsuleShapes, NumMeshesWithCapsules, NumMeshDistanceFieldCasters, IndirectShadowLightDirectionSRV);
+			IndirectCapsuleShadowsResources Resources = CreateIndirectCapsuleShadowsResources(GraphBuilder, Scene, View, NumCapsuleShapes, NumMeshesWithCapsules, NumMeshDistanceFieldCasters);
 
 			if (NumCapsuleShapes == 0 && NumMeshDistanceFieldCasters == 0)
 			{
 				continue;
 			}
 
-			check(IndirectShadowLightDirectionSRV);
+			check(Resources.IndirectShadowLightDirectionSRV);
 
 			const FIntRect ScissorRect = View.ViewRect;
 
@@ -1186,19 +1024,36 @@ void FDeferredShadingSceneRenderer::RenderIndirectCapsuleShadows(FRDGBuilder& Gr
 				FMath::DivideAndRoundUp(ScissorRect.Size().X / GetCapsuleShadowDownsampleFactor(), GShadowShapeTileSize),
 				FMath::DivideAndRoundUp(ScissorRect.Size().Y / GetCapsuleShadowDownsampleFactor(), GShadowShapeTileSize));
 
-			AllocateCapsuleTileIntersectionCountsBuffer(GroupSize, View.ViewState);
+			FRDGBufferRef CapsuleTileIntersectionCountsBuffer = GraphBuilder.CreateBuffer(FRDGBufferDesc::CreateStructuredDesc(sizeof(uint32), GroupSize.X * GroupSize.Y), TEXT("CapsuleTileIntersectionCountsBuffer"));
+			FRDGBufferUAVRef CapsuleTileIntersectionCountsUAV = GraphBuilder.CreateUAV(CapsuleTileIntersectionCountsBuffer);
+			FRDGBufferSRVRef CapsuleTileIntersectionCountsSRV = GraphBuilder.CreateSRV(CapsuleTileIntersectionCountsBuffer);
 
-			AddPass(GraphBuilder, RDG_EVENT_NAME("TransitionTileIntersectionCountsUAV"), [&View](FRHIComputeCommandList& RHICmdList)
-			{
-				RHICmdList.Transition(FRHITransitionInfo(View.ViewState->CapsuleTileIntersectionCountsBuffer.UAV, ERHIAccess::Unknown, ERHIAccess::UAVCompute));
-				RHICmdList.ClearUAVUint(View.ViewState->CapsuleTileIntersectionCountsBuffer.UAV, FUintVector4(0, 0, 0, 0));
-			});
+			AddClearUAVPass(GraphBuilder, CapsuleTileIntersectionCountsUAV, 0);
 
 			{
-				FTiledCapsuleShadowParameters* PassParameters = GraphBuilder.AllocParameters<FTiledCapsuleShadowParameters>();
-				PassParameters->RayTracedShadows = RayTracedShadowsRT;
-				PassParameters->RayTracedShadowsUAV = GraphBuilder.CreateUAV(RayTracedShadowsRT);
-				PassParameters->SceneTextures = SceneTextures.UniformBuffer;
+				auto* PassParameters = GraphBuilder.AllocParameters<FCapsuleShadowingCS::FParameters>();
+
+				SetupCapsuleShadowingParameters(
+					*PassParameters,
+					ECapsuleShadowingType::IndirectTiledCulling,
+					GraphBuilder.CreateUAV(RayTracedShadowsRT),
+					GroupSize,
+					nullptr,
+					FVector2D(GroupSize.X, GroupSize.Y),
+					nullptr,
+					ScissorRect,
+					GetCapsuleShadowDownsampleFactor(),
+					GCapsuleMaxIndirectOcclusionDistance,
+					Scene,
+					View,
+					SceneTextures.UniformBuffer,
+					NumCapsuleShapes,
+					Resources.IndirectShadowCapsuleShapesSRV,
+					NumMeshDistanceFieldCasters,
+					Resources.IndirectShadowMeshDistanceFieldCasterIndicesSRV,
+					Resources.IndirectShadowLightDirectionSRV,
+					CapsuleTileIntersectionCountsUAV
+				);
 
 				EIndirectShadowingPrimitiveTypes PrimitiveTypes;
 
@@ -1221,57 +1076,52 @@ void FDeferredShadingSceneRenderer::RenderIndirectCapsuleShadows(FRDGBuilder& Gr
 				PermutationVector.Set<FCapsuleShadowingCS::FIndirectPrimitiveType>(PrimitiveTypes);
 				auto ComputeShader = View.ShaderMap->GetShader<FCapsuleShadowingCS>(PermutationVector);
 
-				GraphBuilder.AddPass(
+				FComputeShaderUtils::AddPass(
+					GraphBuilder,
 					RDG_EVENT_NAME("TiledCapsuleShadowing %u capsules among %u meshes", NumCapsuleShapes, NumMeshesWithCapsules),
-					PassParameters,
 					ERDGPassFlags::Compute,
-					[ComputeShader, this, &View, RayTracedShadows = PassParameters->RayTracedShadows.GetTexture(), RayTracedShadowsUAV = PassParameters->RayTracedShadowsUAV, IndirectShadowLightDirectionSRV, GroupSize, ScissorRect, NumCapsuleShapes, NumMeshDistanceFieldCasters](FRHIComputeCommandList& RHICmdList)
-					{
-						FRHIComputeShader* ShaderRHI = ComputeShader.GetComputeShader();
-						RHICmdList.SetComputeShader(ShaderRHI);
-
-						ComputeShader->SetParameters(
-							RHICmdList,
-							Scene,
-							View,
-							nullptr,
-							RayTracedShadows->GetRHI(),
-							RayTracedShadowsUAV->GetRHI(),
-							GroupSize,
-							&View.ViewState->CapsuleTileIntersectionCountsBuffer,
-							FVector2D(GroupSize.X, GroupSize.Y),
-							GCapsuleMaxIndirectOcclusionDistance,
-							ScissorRect,
-							GetCapsuleShadowDownsampleFactor(),
-							NumCapsuleShapes,
-							View.ViewState->IndirectShadowCapsuleShapesSRV ? View.ViewState->IndirectShadowCapsuleShapesSRV.GetReference() : nullptr,
-							NumMeshDistanceFieldCasters,
-							View.ViewState->IndirectShadowMeshDistanceFieldCasterIndicesSRV ? View.ViewState->IndirectShadowMeshDistanceFieldCasterIndicesSRV.GetReference() : nullptr,
-							IndirectShadowLightDirectionSRV,
-							nullptr,
-							ECapsuleShadowingType::IndirectTiledCulling);
-
-						DispatchComputeShader(RHICmdList, ComputeShader, GroupSize.X, GroupSize.Y, 1);
-						ComputeShader->UnsetParameters(RHICmdList, &View.ViewState->CapsuleTileIntersectionCountsBuffer);
-					});
+					ComputeShader,
+					PassParameters,
+					FIntVector(GroupSize.X, GroupSize.Y, 1));
 			}
 
 			{
 				const int32 RenderTargetCount = RenderTargets.Num();
 
+				auto VertexShader = View.ShaderMap->GetShader<FCapsuleShadowingUpsampleVS>();
+
+				FCapsuleShadowingUpsamplePS::FPermutationDomain PermutationVector;
+				PermutationVector.Set<FCapsuleShadowingUpsamplePS::FUpsampleRequired>(!GCapsuleShadowsFullResolution);
+				PermutationVector.Set<FCapsuleShadowingUpsamplePS::FApplySSAO>(RenderTargetCount > 1);
+				auto PixelShader = View.ShaderMap->GetShader<FCapsuleShadowingUpsamplePS>(PermutationVector);
+
 				FUpsampleCapsuleShadowParameters* PassParameters = GraphBuilder.AllocParameters<FUpsampleCapsuleShadowParameters>();
 				for (int32 Index = 0; Index < RenderTargetCount; ++Index)
 				{
 					PassParameters->RenderTargets[Index] = FRenderTargetBinding(RenderTargets[Index], ERenderTargetLoadAction::ELoad);
 				}
-				PassParameters->RayTracedShadows = RayTracedShadowsRT;
 				PassParameters->SceneTextures = SceneTextures.UniformBuffer;
 
+				PassParameters->VS.View = GetShaderBinding(View.ViewUniformBuffer);
+				PassParameters->VS.TileIntersectionCounts = CapsuleTileIntersectionCountsSRV;
+				PassParameters->VS.TileDimensions = GroupSize;
+				PassParameters->VS.TileSize = FVector2D(GShadowShapeTileSize * GetCapsuleShadowDownsampleFactor(), GShadowShapeTileSize * GetCapsuleShadowDownsampleFactor());
+				PassParameters->VS.ScissorRectMinAndSize = FIntRect(ScissorRect.Min, ScissorRect.Size());
+
+				PassParameters->PS.View = GetShaderBinding(View.ViewUniformBuffer);
+				PassParameters->PS.ShadowFactorsTexture = RayTracedShadowsRT;
+				PassParameters->PS.ShadowFactorsSampler = TStaticSamplerState<SF_Bilinear>::GetRHI();
+				PassParameters->PS.ScissorRectMinAndSize = FIntRect(ScissorRect.Min, ScissorRect.Size());
+				PassParameters->PS.OutputtingToLightAttenuation = 0.0f;
+
+				ClearUnusedGraphResources(VertexShader, &PassParameters->VS);
+				ClearUnusedGraphResources(PixelShader, &PassParameters->PS);
+
 				GraphBuilder.AddPass(
 					RDG_EVENT_NAME("UpsampleCapsuleShadow %dx%d", ScissorRect.Width(), ScissorRect.Height()),
 					PassParameters,
 					ERDGPassFlags::Raster,
-					[this, &View, RayTracedShadowsRT, RenderTargetCount, GroupSize, ScissorRect](FRHICommandList& RHICmdList)
+					[PassParameters, VertexShader, PixelShader, &View, RayTracedShadowsRT, RenderTargetCount, GroupSize, ScissorRect](FRHICommandList& RHICmdList)
 				{
 					FGraphicsPipelineStateInitializer GraphicsPSOInit;
 					RHICmdList.ApplyCachedRenderTargets(GraphicsPSOInit);
@@ -1293,21 +1143,15 @@ void FDeferredShadingSceneRenderer::RenderIndirectCapsuleShadows(FRDGBuilder& Gr
 						GraphicsPSOInit.BlendState = TStaticBlendState<CW_RGB, BO_Add, BF_DestColor, BF_Zero>::GetRHI();
 					}
 
-					auto VertexShader = View.ShaderMap->GetShader<FCapsuleShadowingUpsampleVS>();
-
-					FCapsuleShadowingUpsamplePS::FPermutationDomain PermutationVector;
-					PermutationVector.Set<FCapsuleShadowingUpsamplePS::FUpsampleRequired>(!GCapsuleShadowsFullResolution);
-					PermutationVector.Set<FCapsuleShadowingUpsamplePS::FApplySSAO>(RenderTargetCount > 1);
-					auto PixelShader = View.ShaderMap->GetShader<FCapsuleShadowingUpsamplePS>(PermutationVector);
-
 					GraphicsPSOInit.BoundShaderState.VertexDeclarationRHI = GTileVertexDeclaration.VertexDeclarationRHI;
 					GraphicsPSOInit.BoundShaderState.VertexShaderRHI = VertexShader.GetVertexShader();
 					GraphicsPSOInit.PrimitiveType = PT_TriangleList;
 					GraphicsPSOInit.BoundShaderState.PixelShaderRHI = PixelShader.GetPixelShader();
 
 					SetGraphicsPipelineState(RHICmdList, GraphicsPSOInit);
-					VertexShader->SetParameters(RHICmdList, View, GroupSize, ScissorRect, View.ViewState->CapsuleTileIntersectionCountsBuffer);
-					PixelShader->SetParameters(RHICmdList, View, ScissorRect, RayTracedShadowsRT->GetRHI(), false);
+
+					SetShaderParameters(RHICmdList, VertexShader, VertexShader.GetVertexShader(), PassParameters->VS);
+					SetShaderParameters(RHICmdList, PixelShader, PixelShader.GetPixelShader(), PassParameters->PS);
 
 					RHICmdList.SetStreamSource(0, GTileTexCoordVertexBuffer.VertexBufferRHI, 0);
 					RHICmdList.DrawIndexedPrimitive(GTileIndexBuffer.IndexBufferRHI,
@@ -1352,13 +1196,6 @@ bool FSceneRenderer::ShouldPrepareForDFInsetIndirectShadow() const
 	return bSceneHasInsetDFPrimitives && SupportsCapsuleIndirectShadows(FeatureLevel, GShaderPlatformForFeatureLevel[FeatureLevel]) && ViewFamily.EngineShowFlags.CapsuleShadows;
 }
 
-BEGIN_SHADER_PARAMETER_STRUCT(FCapsuleShadowsForMovableSkylightParameters, )
-	SHADER_PARAMETER_RDG_UNIFORM_BUFFER(FSceneTextureUniformParameters, SceneTextures)
-	RDG_TEXTURE_ACCESS(BentNormalInput, ERHIAccess::SRVCompute)
-	RDG_TEXTURE_ACCESS(BentNormalOutput, ERHIAccess::UAVCompute)
-	SHADER_PARAMETER_RDG_TEXTURE_UAV(RWTexture2D, BentNormalOutputUAV)
-END_SHADER_PARAMETER_STRUCT()
-
 void FDeferredShadingSceneRenderer::RenderCapsuleShadowsForMovableSkylight(
 	FRDGBuilder& GraphBuilder,
 	TRDGUniformBufferRef<FSceneTextureUniformParameters> SceneTexturesUniformBuffer,
@@ -1399,17 +1236,16 @@ void FDeferredShadingSceneRenderer::RenderCapsuleShadowsForMovableSkylight(
 					int32 NumCapsuleShapes = 0;
 					int32 NumMeshesWithCapsules = 0;
 					int32 NumMeshDistanceFieldCasters = 0;
-					FRHIShaderResourceView* IndirectShadowLightDirectionSRV = NULL;
-					SetupIndirectCapsuleShadows(GraphBuilder, View, NumCapsuleShapes, NumMeshesWithCapsules, NumMeshDistanceFieldCasters, IndirectShadowLightDirectionSRV);
+					IndirectCapsuleShadowsResources Resources = CreateIndirectCapsuleShadowsResources(GraphBuilder, Scene, View, NumCapsuleShapes, NumMeshesWithCapsules, NumMeshDistanceFieldCasters);
 
 					// Don't render indirect occlusion from mesh distance fields when operating on a movable skylight,
 					// DFAO is responsible for indirect occlusion from meshes with distance fields on a movable skylight.
 					// A single mesh should only provide indirect occlusion for a given lighting component in one way.
 					NumMeshDistanceFieldCasters = 0;
 
-					if (NumCapsuleShapes > 0 || NumMeshDistanceFieldCasters > 0)
+					if (NumCapsuleShapes > 0)
 					{
-						check(IndirectShadowLightDirectionSRV);
+						check(Resources.IndirectShadowLightDirectionSRV);
 
 						FIntRect ScissorRect = View.ViewRect;
 
@@ -1417,67 +1253,42 @@ void FDeferredShadingSceneRenderer::RenderCapsuleShadowsForMovableSkylight(
 							uint32 GroupSizeX = FMath::DivideAndRoundUp(ScissorRect.Size().X / GAODownsampleFactor, GShadowShapeTileSize);
 							uint32 GroupSizeY = FMath::DivideAndRoundUp(ScissorRect.Size().Y / GAODownsampleFactor, GShadowShapeTileSize);
 
-							auto* PassParameters = GraphBuilder.AllocParameters< FCapsuleShadowsForMovableSkylightParameters>();
-							PassParameters->BentNormalInput = BentNormalOutput;
-							PassParameters->BentNormalOutput = NewBentNormal;
-							PassParameters->BentNormalOutputUAV = GraphBuilder.CreateUAV(NewBentNormal);
-							PassParameters->SceneTextures = SceneTexturesUniformBuffer;
-
-							EIndirectShadowingPrimitiveTypes PrimitiveTypes;
-
-							if (NumCapsuleShapes > 0 && NumMeshDistanceFieldCasters > 0)
-							{
-								PrimitiveTypes = EIndirectShadowingPrimitiveTypes::CapsuleShapesAndMeshDistanceFields;
-							}
-							else if (NumCapsuleShapes > 0)
-							{
-								PrimitiveTypes = EIndirectShadowingPrimitiveTypes::CapsuleShapes;
-							}
-							else
-							{
-								check(NumMeshDistanceFieldCasters > 0);
-								PrimitiveTypes = EIndirectShadowingPrimitiveTypes::MeshDistanceFields;
-							}
+							auto* PassParameters = GraphBuilder.AllocParameters<FCapsuleShadowingCS::FParameters>();
+
+							SetupCapsuleShadowingParameters(
+								*PassParameters,
+								ECapsuleShadowingType::MovableSkylightTiledCulling,
+								GraphBuilder.CreateUAV(NewBentNormal),
+								FIntPoint(GroupSizeX, GroupSizeY),
+								BentNormalOutput,
+								FVector2D(GroupSizeX, GroupSizeY),
+								nullptr,
+								ScissorRect,
+								GAODownsampleFactor,
+								GCapsuleMaxIndirectOcclusionDistance,
+								Scene,
+								View,
+								SceneTexturesUniformBuffer,
+								NumCapsuleShapes,
+								Resources.IndirectShadowCapsuleShapesSRV,
+								NumMeshDistanceFieldCasters,
+								Resources.IndirectShadowMeshDistanceFieldCasterIndicesSRV,
+								Resources.IndirectShadowLightDirectionSRV,
+								nullptr
+							);
 
 							FCapsuleShadowingCS::FPermutationDomain PermutationVector;
 							PermutationVector.Set<FCapsuleShadowingCS::FShapeShadow>(ECapsuleShadowingType::MovableSkylightTiledCulling);
-							PermutationVector.Set<FCapsuleShadowingCS::FIndirectPrimitiveType>(PrimitiveTypes);
+							PermutationVector.Set<FCapsuleShadowingCS::FIndirectPrimitiveType>(EIndirectShadowingPrimitiveTypes::CapsuleShapes);
 							auto ComputeShader = View.ShaderMap->GetShader<FCapsuleShadowingCS>(PermutationVector);
 
-							GraphBuilder.AddPass(
+							FComputeShaderUtils::AddPass(
+								GraphBuilder,
 								RDG_EVENT_NAME("TiledCapsuleShadowing % u capsules among % u meshes", NumCapsuleShapes, NumMeshesWithCapsules),
-								PassParameters,
 								ERDGPassFlags::Compute,
-								[PassParameters, ComputeShader, this, &View, GroupSizeX, GroupSizeY, NumCapsuleShapes, NumMeshDistanceFieldCasters, ScissorRect, IndirectShadowLightDirectionSRV]
-							(FRHICommandList& RHICmdList)
-								{
-									FRHIComputeShader* ShaderRHI = ComputeShader.GetComputeShader();
-									RHICmdList.SetComputeShader(ShaderRHI);
-
-									ComputeShader->SetParameters(
-										RHICmdList,
-										Scene,
-										View,
-										NULL,
-										PassParameters->BentNormalOutput->GetRHI(),
-										PassParameters->BentNormalOutputUAV->GetRHI(),
-										FIntPoint(GroupSizeX, GroupSizeY),
-										NULL,
-										FVector2D(GroupSizeX, GroupSizeY),
-										GCapsuleMaxIndirectOcclusionDistance,
-										ScissorRect,
-										GAODownsampleFactor,
-										NumCapsuleShapes,
-										View.ViewState->IndirectShadowCapsuleShapesSRV ? View.ViewState->IndirectShadowCapsuleShapesSRV.GetReference() : NULL,
-										NumMeshDistanceFieldCasters,
-										View.ViewState->IndirectShadowMeshDistanceFieldCasterIndicesSRV ? View.ViewState->IndirectShadowMeshDistanceFieldCasterIndicesSRV.GetReference() : NULL,
-										IndirectShadowLightDirectionSRV,
-										PassParameters->BentNormalInput->GetRHI(),
-										ECapsuleShadowingType::MovableSkylightTiledCulling);
-
-									DispatchComputeShader(RHICmdList, ComputeShader.GetShader(), GroupSizeX, GroupSizeY, 1);
-									ComputeShader->UnsetParameters(RHICmdList, nullptr);
-								});
+								ComputeShader,
+								PassParameters,
+								FIntVector(GroupSizeX, GroupSizeY, 1));
 						}
 
 						// Replace the pipeline output with our output that has capsule shadows applied
diff --git a/Engine/Source/Runtime/Renderer/Private/DeferredShadingRenderer.h b/Engine/Source/Runtime/Renderer/Private/DeferredShadingRenderer.h
--- a/Engine/Source/Runtime/Renderer/Private/DeferredShadingRenderer.h
+++ b/Engine/Source/Runtime/Renderer/Private/DeferredShadingRenderer.h
@@ -650,15 +650,6 @@ private:
 		TArrayView<const FProjectedShadowInfo* const> CapsuleShadows,
 		bool bProjectingForForwardShading) const;
 
-	/** Sets up ViewState buffers for rendering capsule shadows. */
-	void SetupIndirectCapsuleShadows(
-		FRDGBuilder& GraphBuilder, 
-		const FViewInfo& View, 
-		int32& NumCapsuleShapes, 
-		int32& NumMeshesWithCapsules, 
-		int32& NumMeshDistanceFieldCasters,
-		FRHIShaderResourceView*& IndirectShadowLightDirectionSRV) const;
-
 	/** Renders indirect shadows from capsules modulated onto scene color. */
 	void RenderIndirectCapsuleShadows(FRDGBuilder& GraphBuilder, const FSceneTextures& SceneTextures) const;
 
diff --git a/Engine/Source/Runtime/Renderer/Private/DistanceFieldLightingShared.h b/Engine/Source/Runtime/Renderer/Private/DistanceFieldLightingShared.h
--- a/Engine/Source/Runtime/Renderer/Private/DistanceFieldLightingShared.h
+++ b/Engine/Source/Runtime/Renderer/Private/DistanceFieldLightingShared.h
@@ -127,175 +127,6 @@ namespace DistanceField
 	void BuildUpdateTrackingBounds(const TArray<FViewInfo>& Views, DistanceField::FUpdateTrackingBounds& UpdateTrackingBounds);
 };
 
-template <EDistanceFieldPrimitiveType PrimitiveType>
-class TDistanceFieldObjectBufferParameters
-{
-	DECLARE_INLINE_TYPE_LAYOUT(TDistanceFieldObjectBufferParameters, NonVirtual);
-public:
-	void Bind(const FShaderParameterMap& ParameterMap)
-	{
-		SceneObjectBounds.Bind(ParameterMap, TEXT("SceneObjectBounds"));
-		SceneObjectData.Bind(ParameterMap, TEXT("SceneObjectData"));
-		NumSceneObjects.Bind(ParameterMap, TEXT("NumSceneObjects"));
-	}
-
-	template<typename TParamRef>
-	void Set(
-		FRHIComputeCommandList& RHICmdList,
-		const TParamRef& ShaderRHI,
-		const TDistanceFieldObjectBuffers<PrimitiveType>& ObjectBuffers,
-		int32 NumObjectsValue,
-		bool bBarrier = false)
-	{
-		if (bBarrier)
-		{
-			FRHITransitionInfo UAVTransitions[2];
-			UAVTransitions[0] = FRHITransitionInfo(ObjectBuffers.Bounds.UAV, ERHIAccess::Unknown, ERHIAccess::UAVCompute);
-			UAVTransitions[1] = FRHITransitionInfo(ObjectBuffers.Data.UAV, ERHIAccess::Unknown, ERHIAccess::UAVCompute);
-			RHICmdList.Transition(MakeArrayView(UAVTransitions, UE_ARRAY_COUNT(UAVTransitions)));
-		}
-
-		SceneObjectBounds.SetBuffer(RHICmdList, ShaderRHI, ObjectBuffers.Bounds);
-		SceneObjectData.SetBuffer(RHICmdList, ShaderRHI, ObjectBuffers.Data);
-		SetShaderValue(RHICmdList, ShaderRHI, NumSceneObjects, NumObjectsValue);
-	}
-
-	template<typename TParamRef>
-	void UnsetParameters(FRHICommandList& RHICmdList, const TParamRef& ShaderRHI, const TDistanceFieldObjectBuffers<PrimitiveType>& ObjectBuffers, bool bBarrier = false)
-	{
-		SceneObjectBounds.UnsetUAV(RHICmdList, ShaderRHI);
-		SceneObjectData.UnsetUAV(RHICmdList, ShaderRHI);
-
-		if (bBarrier)
-		{
-			FRHITransitionInfo SRVTransitions[2];
-			SRVTransitions[0] = FRHITransitionInfo(ObjectBuffers.Bounds.UAV, ERHIAccess::Unknown, ERHIAccess::SRVMask);
-			SRVTransitions[1] = FRHITransitionInfo(ObjectBuffers.Data.UAV, ERHIAccess::Unknown, ERHIAccess::SRVMask);
-			RHICmdList.Transition(MakeArrayView(SRVTransitions, UE_ARRAY_COUNT(SRVTransitions)));
-		}
-	}
-
-	friend FArchive& operator<<(FArchive& Ar, TDistanceFieldObjectBufferParameters& P)
-	{
-		Ar << P.SceneObjectBounds << P.SceneObjectData << P.NumSceneObjects;
-		return Ar;
-	}
-
-	bool AnyBound() const
-	{
-		return SceneObjectBounds.IsBound() || SceneObjectData.IsBound() || NumSceneObjects.IsBound();
-	}
-
-private:
-	
-	LAYOUT_FIELD(FRWShaderParameter, SceneObjectBounds)
-	LAYOUT_FIELD(FRWShaderParameter, SceneObjectData)
-	LAYOUT_FIELD(FShaderParameter, NumSceneObjects)
-};
-
-template <EDistanceFieldPrimitiveType PrimitiveType>
-class TDistanceFieldCulledObjectBuffers
-{
-public:
-
-	static int32 ObjectDataStride;
-	static int32 ObjectBoxBoundsStride;
-
-	bool bWantBoxBounds;
-	int32 MaxObjects;
-
-	FRWBuffer ObjectIndirectArguments;
-	FRWBuffer ObjectIndirectDispatch;
-	FRWBufferStructured Bounds;
-	FRWBufferStructured Data;
-	FRWBufferStructured BoxBounds;
-
-	TDistanceFieldCulledObjectBuffers()
-	{
-		MaxObjects = 0;
-		bWantBoxBounds = false;
-	}
-
-	void Initialize()
-	{
-		if (MaxObjects > 0)
-		{
-			const TCHAR* ObjectIndirectArgumentsDebugName;
-			const TCHAR* ObjectIndirectDispatchDebugName;
-			const TCHAR* BoundsDebugName;
-			const TCHAR* DataDebugName;
-			const TCHAR* BoxBoundsDebugName;
-			uint32 BoundsNumElements;
-
-			if (PrimitiveType == DFPT_HeightField)
-			{
-				ObjectIndirectArgumentsDebugName = TEXT("FHeightFieldCulledObjectBuffers_ObjectIndirectArguments");
-				ObjectIndirectDispatchDebugName = TEXT("FHeightFieldCulledObjectBuffers_ObjectIndirectDispatch");
-				BoundsDebugName = TEXT("FHeightFieldCulledObjectBuffers_Bounds");
-				DataDebugName = TEXT("FHeightFieldCulledObjectBuffers_Data");
-				BoxBoundsDebugName = TEXT("FHeightFieldCulledObjectBuffers_BoxBounds");
-				BoundsNumElements = MaxObjects * 2;
-			}
-			else
-			{
-				check(PrimitiveType == DFPT_SignedDistanceField);
-				ObjectIndirectArgumentsDebugName = TEXT("FDistanceFieldCulledObjectBuffers_ObjectIndirectArguments");
-				ObjectIndirectDispatchDebugName = TEXT("FDistanceFieldCulledObjectBuffers_ObjectIndirectDispatch");
-				BoundsDebugName = TEXT("FDistanceFieldCulledObjectBuffers_Bounds");
-				DataDebugName = TEXT("FDistanceFieldCulledObjectBuffers_Data");
-				BoxBoundsDebugName = TEXT("FDistanceFieldCulledObjectBuffers_BoxBounds");
-				BoundsNumElements = MaxObjects;
-			}
-
-			const uint32 FastVRamFlag = GFastVRamConfig.DistanceFieldCulledObjectBuffers | ( IsTransientResourceBufferAliasingEnabled() ? BUF_Transient : BUF_None );
-
-			ObjectIndirectArguments.Initialize(ObjectIndirectArgumentsDebugName, sizeof(uint32), 5, PF_R32_UINT, BUF_Static | BUF_DrawIndirect);
-			ObjectIndirectDispatch.Initialize(ObjectIndirectDispatchDebugName, sizeof(uint32), 3, PF_R32_UINT, BUF_Static | BUF_DrawIndirect);
-			Bounds.Initialize(BoundsDebugName, sizeof(FVector4), BoundsNumElements, BUF_Static | FastVRamFlag);
-			Data.Initialize(DataDebugName, sizeof(FVector4), MaxObjects * ObjectDataStride, BUF_Static | FastVRamFlag);
-
-			if (bWantBoxBounds)
-			{
-				BoxBounds.Initialize(BoxBoundsDebugName, sizeof(FVector4), MaxObjects * ObjectBoxBoundsStride, BUF_Static | FastVRamFlag);
-			}
-		}
-	}
-
-	void AcquireTransientResource()
-	{
-		Bounds.AcquireTransientResource();
-		Data.AcquireTransientResource();
-		if (bWantBoxBounds)
-		{
-			BoxBounds.AcquireTransientResource();
-		}
-	}
-
-	void DiscardTransientResource()
-	{
-		Bounds.DiscardTransientResource();
-		Data.DiscardTransientResource();
-		if (bWantBoxBounds)
-		{
-			BoxBounds.DiscardTransientResource();
-		}
-	}
-
-	void Release()
-	{
-		ObjectIndirectArguments.Release();
-		ObjectIndirectDispatch.Release();
-		Bounds.Release();
-		Data.Release();
-		BoxBounds.Release();
-	}
-
-	size_t GetSizeBytes() const
-	{
-		return ObjectIndirectArguments.NumBytes + ObjectIndirectDispatch.NumBytes + Bounds.NumBytes + Data.NumBytes + BoxBounds.NumBytes;
-	}
-};
-
 BEGIN_SHADER_PARAMETER_STRUCT(FDistanceFieldCulledObjectBufferParameters, )
 	SHADER_PARAMETER_RDG_BUFFER_UAV(RWBuffer<uint>, RWObjectIndirectArguments)
 	SHADER_PARAMETER_RDG_BUFFER_UAV(RWBuffer<uint>, RWCulledObjectBounds)
@@ -315,157 +146,6 @@ extern void AllocateDistanceFieldCulledObjectBuffers(
 	FRDGBufferRef& OutObjectIndirectArguments,
 	FDistanceFieldCulledObjectBufferParameters& OutParameters);
 
-template <EDistanceFieldPrimitiveType PrimitiveType>
-class TDistanceFieldCulledObjectBufferParameters
-{
-	DECLARE_INLINE_TYPE_LAYOUT(TDistanceFieldCulledObjectBufferParameters, NonVirtual);
-public:
-	void Bind(const FShaderParameterMap& ParameterMap)
-	{
-		ObjectIndirectArguments.Bind(ParameterMap, TEXT("ObjectIndirectArguments"));
-		CulledObjectBounds.Bind(ParameterMap, TEXT("CulledObjectBounds"));
-		CulledObjectData.Bind(ParameterMap, TEXT("CulledObjectData"));
-		CulledObjectBoxBounds.Bind(ParameterMap, TEXT("CulledObjectBoxBounds"));
-		HFVisibilityTexture.Bind(ParameterMap, TEXT("HFVisibilityTexture"));
-		SceneDistanceFieldAssetData.Bind(ParameterMap, TEXT("SceneDistanceFieldAssetData"));
-		DistanceFieldIndirectionTable.Bind(ParameterMap, TEXT("DistanceFieldIndirectionTable"));
-		DistanceFieldBrickTexture.Bind(ParameterMap, TEXT("DistanceFieldBrickTexture"));
-		DistanceFieldSampler.Bind(ParameterMap, TEXT("DistanceFieldSampler"));
-		DistanceFieldBrickSize.Bind(ParameterMap, TEXT("DistanceFieldBrickSize"));
-		DistanceFieldUniqueDataBrickSize.Bind(ParameterMap, TEXT("DistanceFieldUniqueDataBrickSize"));
-		DistanceFieldBrickAtlasSizeInBricks.Bind(ParameterMap, TEXT("DistanceFieldBrickAtlasSizeInBricks"));
-		DistanceFieldBrickAtlasMask.Bind(ParameterMap, TEXT("DistanceFieldBrickAtlasMask"));
-		DistanceFieldBrickAtlasSizeLog2.Bind(ParameterMap, TEXT("DistanceFieldBrickAtlasSizeLog2"));
-		DistanceFieldBrickAtlasTexelSize.Bind(ParameterMap, TEXT("DistanceFieldBrickAtlasTexelSize"));
-	}
-
-	template<typename TShaderRHI, typename TRHICommandList>
-	void Set(
-		TRHICommandList& RHICmdList,
-		TShaderRHI* ShaderRHI,
-		const TDistanceFieldCulledObjectBuffers<PrimitiveType>& ObjectBuffers,
-		const FDistanceFieldSceneData& DistanceFieldSceneData,
-		FRHITexture* HFVisibilityAtlas = nullptr)
-	{
-		ObjectIndirectArguments.SetBuffer(RHICmdList, ShaderRHI, ObjectBuffers.ObjectIndirectArguments);
-		CulledObjectBounds.SetBuffer(RHICmdList, ShaderRHI, ObjectBuffers.Bounds);
-		CulledObjectData.SetBuffer(RHICmdList, ShaderRHI, ObjectBuffers.Data);
-
-		if (CulledObjectBoxBounds.IsBound())
-		{
-			check(ObjectBuffers.bWantBoxBounds);
-			CulledObjectBoxBounds.SetBuffer(RHICmdList, ShaderRHI, ObjectBuffers.BoxBounds);
-		}
-
-		FDistanceFieldAtlasParameters AtlasParameters = DistanceField::SetupAtlasParameters(DistanceFieldSceneData);
-
-		SetSRVParameter(RHICmdList, ShaderRHI, SceneDistanceFieldAssetData, AtlasParameters.SceneDistanceFieldAssetData);
-		SetSRVParameter(RHICmdList, ShaderRHI, DistanceFieldIndirectionTable, AtlasParameters.DistanceFieldIndirectionTable);
-
-		SetTextureParameter(
-			RHICmdList,
-			ShaderRHI,
-			DistanceFieldBrickTexture,
-			DistanceFieldSampler,
-			TStaticSamplerState<SF_Bilinear,AM_Clamp,AM_Clamp,AM_Clamp>::GetRHI(),
-			DistanceFieldSceneData.DistanceFieldBrickVolumeTexture->GetRenderTargetItem().ShaderResourceTexture
-			);
-
-		SetShaderValue(RHICmdList, ShaderRHI, DistanceFieldBrickSize, AtlasParameters.DistanceFieldBrickSize);
-		SetShaderValue(RHICmdList, ShaderRHI, DistanceFieldUniqueDataBrickSize, AtlasParameters.DistanceFieldUniqueDataBrickSize);
-		SetShaderValue(RHICmdList, ShaderRHI, DistanceFieldBrickAtlasSizeInBricks, AtlasParameters.DistanceFieldBrickAtlasSizeInBricks);
-		SetShaderValue(RHICmdList, ShaderRHI, DistanceFieldBrickAtlasMask, AtlasParameters.DistanceFieldBrickAtlasMask);
-		SetShaderValue(RHICmdList, ShaderRHI, DistanceFieldBrickAtlasSizeLog2, AtlasParameters.DistanceFieldBrickAtlasSizeLog2);
-		SetShaderValue(RHICmdList, ShaderRHI, DistanceFieldBrickAtlasTexelSize, AtlasParameters.DistanceFieldBrickAtlasTexelSize);
-
-		if (HFVisibilityTexture.IsBound())
-		{
-			check(HFVisibilityAtlas);
-			SetTextureParameter(RHICmdList, ShaderRHI, HFVisibilityTexture, HFVisibilityAtlas);
-		}
-	}
-
-	template<typename TRHIShader, typename TRHICommandList>
-	void UnsetParameters(TRHICommandList& RHICmdList, TRHIShader* ShaderRHI)
-	{
-		ObjectIndirectArguments.UnsetUAV(RHICmdList, ShaderRHI);
-		CulledObjectBounds.UnsetUAV(RHICmdList, ShaderRHI);
-		CulledObjectData.UnsetUAV(RHICmdList, ShaderRHI);
-		CulledObjectBoxBounds.UnsetUAV(RHICmdList, ShaderRHI);
-	}
-
-	void GetUAVs(const TDistanceFieldCulledObjectBuffers<PrimitiveType>& ObjectBuffers, TArray<FRHIUnorderedAccessView*>& UAVs)
-	{
-		uint32 MaxIndex = 0;
-		MaxIndex = FMath::Max(MaxIndex, ObjectIndirectArguments.GetUAVIndex());
-		MaxIndex = FMath::Max(MaxIndex, CulledObjectBounds.GetUAVIndex());
-		MaxIndex = FMath::Max(MaxIndex, CulledObjectData.GetUAVIndex());
-		MaxIndex = FMath::Max(MaxIndex, CulledObjectBoxBounds.GetUAVIndex());
-
-		UAVs.AddZeroed(MaxIndex + 1);
-
-		if (ObjectIndirectArguments.IsUAVBound())
-		{
-			UAVs[ObjectIndirectArguments.GetUAVIndex()] = ObjectBuffers.ObjectIndirectArguments.UAV;
-		}
-
-		if (CulledObjectBounds.IsUAVBound())
-		{
-			UAVs[CulledObjectBounds.GetUAVIndex()] = ObjectBuffers.Bounds.UAV;
-		}
-
-		if (CulledObjectData.IsUAVBound())
-		{
-			UAVs[CulledObjectData.GetUAVIndex()] = ObjectBuffers.Data.UAV;
-		}
-
-		if (CulledObjectBoxBounds.IsUAVBound())
-		{
-			UAVs[CulledObjectBoxBounds.GetUAVIndex()] = ObjectBuffers.BoxBounds.UAV;
-		}
-
-		check(UAVs.Num() > 0);
-	}
-
-	friend FArchive& operator<<(FArchive& Ar, TDistanceFieldCulledObjectBufferParameters<PrimitiveType>& P)
-	{
-		Ar << P.ObjectIndirectArguments;
-		Ar << P.CulledObjectBounds;
-		Ar << P.CulledObjectData;
-		Ar << P.CulledObjectBoxBounds;
-		Ar << P.HFVisibilityTexture;
-		Ar << P.SceneDistanceFieldAssetData;
-		Ar << P.DistanceFieldIndirectionTable;
-		Ar << P.DistanceFieldBrickTexture;
-		Ar << P.DistanceFieldSampler;
-		Ar << P.DistanceFieldBrickSize;
-		Ar << P.DistanceFieldUniqueDataBrickSize;
-		Ar << P.DistanceFieldBrickAtlasSizeInBricks;
-		Ar << P.DistanceFieldBrickAtlasMask;
-		Ar << P.DistanceFieldBrickAtlasSizeLog2;
-		Ar << P.DistanceFieldBrickAtlasTexelSize;
-		return Ar;
-	}
-
-private:
-	
-	LAYOUT_FIELD(FRWShaderParameter, ObjectIndirectArguments)
-	LAYOUT_FIELD(FRWShaderParameter, CulledObjectBounds)
-	LAYOUT_FIELD(FRWShaderParameter, CulledObjectData)
-	LAYOUT_FIELD(FRWShaderParameter, CulledObjectBoxBounds)
-	LAYOUT_FIELD(FShaderResourceParameter, HFVisibilityTexture);
-	LAYOUT_FIELD(FShaderResourceParameter, SceneDistanceFieldAssetData)
-	LAYOUT_FIELD(FShaderResourceParameter, DistanceFieldIndirectionTable)
-	LAYOUT_FIELD(FShaderResourceParameter, DistanceFieldBrickTexture)
-	LAYOUT_FIELD(FShaderResourceParameter, DistanceFieldSampler)
-	LAYOUT_FIELD(FShaderParameter, DistanceFieldBrickSize);
-	LAYOUT_FIELD(FShaderParameter, DistanceFieldUniqueDataBrickSize);
-	LAYOUT_FIELD(FShaderParameter, DistanceFieldBrickAtlasSizeInBricks);
-	LAYOUT_FIELD(FShaderParameter, DistanceFieldBrickAtlasMask);
-	LAYOUT_FIELD(FShaderParameter, DistanceFieldBrickAtlasSizeLog2);
-	LAYOUT_FIELD(FShaderParameter, DistanceFieldBrickAtlasTexelSize);
-};
-
 class FCPUUpdatedBuffer
 {
 public:
@@ -539,25 +219,3 @@ extern void CullDistanceFieldObjectsForLight(
 	FLightTileIntersectionParameters& LightTileIntersectionParameters);
 
 extern bool SupportsDistanceFieldAO(ERHIFeatureLevel::Type FeatureLevel, EShaderPlatform ShaderPlatform);
-
-template <bool bIsAType> struct TSelector {};
-
-template <>
-struct TSelector<true>
-{
-	template <typename AType, typename BType>
-	AType& operator()(AType& A, BType& B)
-	{
-		return A;
-	}
-};
-
-template <>
-struct TSelector<false>
-{
-	template <typename AType, typename BType>
-	BType& operator()(AType& A, BType& B)
-	{
-		return B;
-	}
-};
\ No newline at end of file
diff --git a/Engine/Source/Runtime/Renderer/Private/DistanceFieldObjectManagement.cpp b/Engine/Source/Runtime/Renderer/Private/DistanceFieldObjectManagement.cpp
--- a/Engine/Source/Runtime/Renderer/Private/DistanceFieldObjectManagement.cpp
+++ b/Engine/Source/Runtime/Renderer/Private/DistanceFieldObjectManagement.cpp
@@ -89,10 +89,6 @@ template<> int32 TDistanceFieldObjectBuffers<DFPT_SignedDistanceField>::ObjectDa
 template<> int32 TDistanceFieldObjectBuffers<DFPT_SignedDistanceField>::ObjectBoundsStride = GDistanceFieldObjectBoundsStride;
 template<> int32 TDistanceFieldObjectBuffers<DFPT_HeightField>::ObjectDataStride = GHeightFieldObjectDataStride;
 template<> int32 TDistanceFieldObjectBuffers<DFPT_HeightField>::ObjectBoundsStride = GHeightFieldObjectBoundsStride;
-template<> int32 TDistanceFieldCulledObjectBuffers<DFPT_SignedDistanceField>::ObjectDataStride = GDistanceFieldCulledObjectDataStride;
-template<> int32 TDistanceFieldCulledObjectBuffers<DFPT_SignedDistanceField>::ObjectBoxBoundsStride = GDistanceFieldCulledObjectBoxBoundsStride;
-template<> int32 TDistanceFieldCulledObjectBuffers<DFPT_HeightField>::ObjectDataStride = GHeightFieldCulledObjectDataStride;
-template<> int32 TDistanceFieldCulledObjectBuffers<DFPT_HeightField>::ObjectBoxBoundsStride = GHeightFieldCulledObjectBoxBoundsStride;
 
 template <EDistanceFieldPrimitiveType PrimitiveType>
 void TDistanceFieldObjectBuffers<PrimitiveType>::Initialize()
diff --git a/Engine/Source/Runtime/Renderer/Private/DistanceFieldShadowing.cpp b/Engine/Source/Runtime/Renderer/Private/DistanceFieldShadowing.cpp
--- a/Engine/Source/Runtime/Renderer/Private/DistanceFieldShadowing.cpp
+++ b/Engine/Source/Runtime/Renderer/Private/DistanceFieldShadowing.cpp
@@ -451,14 +451,14 @@ void AllocateDistanceFieldCulledObjectBuffers(
 	if (PrimitiveType == DFPT_SignedDistanceField)
 	{
 		NumBoundsElementsScale = 1;
-		ObjectDataStride = TDistanceFieldCulledObjectBuffers<DFPT_SignedDistanceField>::ObjectDataStride;
-		ObjectBoxBoundsStride = TDistanceFieldCulledObjectBuffers<DFPT_SignedDistanceField>::ObjectBoxBoundsStride;
+		ObjectDataStride = GDistanceFieldCulledObjectDataStride;
+		ObjectBoxBoundsStride = GDistanceFieldCulledObjectBoxBoundsStride;
 	}
 	else
 	{
 		NumBoundsElementsScale = 2;
-		ObjectDataStride = TDistanceFieldCulledObjectBuffers<DFPT_HeightField>::ObjectDataStride;
-		ObjectBoxBoundsStride = TDistanceFieldCulledObjectBuffers<DFPT_HeightField>::ObjectBoxBoundsStride;
+		ObjectDataStride = GHeightFieldCulledObjectDataStride;
+		ObjectBoxBoundsStride = GHeightFieldCulledObjectBoxBoundsStride;
 	}
 
 	FRDGBufferRef Bounds = GraphBuilder.CreateBuffer(FRDGBufferDesc::CreateStructuredDesc(sizeof(FVector4), MaxObjects * NumBoundsElementsScale), TEXT("FDistanceFieldCulledObjectBuffers_Bounds"));
@@ -827,10 +827,7 @@ void FProjectedShadowInfo::BeginRenderRayTracedDistanceFieldProjection(
 
 			const FMatrix WorldToShadowValue = FTranslationMatrix(PreShadowTranslation) * TranslatedWorldToClipInnerMatrix;
 
-			FDistanceFieldObjectBufferParameters ObjectBufferParameters;
-			ObjectBufferParameters.SceneObjectBounds = Scene->DistanceFieldSceneData.GetCurrentObjectBuffers()->Bounds.SRV;
-			ObjectBufferParameters.SceneObjectData = Scene->DistanceFieldSceneData.GetCurrentObjectBuffers()->Data.SRV;
-			ObjectBufferParameters.NumSceneObjects = Scene->DistanceFieldSceneData.NumObjectsInBuffer;
+			FDistanceFieldObjectBufferParameters ObjectBufferParameters = DistanceField::SetupObjectBufferParameters(Scene->DistanceFieldSceneData);
 
 			FLightTileIntersectionParameters LightTileIntersectionParameters;
 			FDistanceFieldCulledObjectBufferParameters CulledObjectBufferParameters;
diff --git a/Engine/Source/Runtime/Renderer/Private/LightSceneInfo.cpp b/Engine/Source/Runtime/Renderer/Private/LightSceneInfo.cpp
--- a/Engine/Source/Runtime/Renderer/Private/LightSceneInfo.cpp
+++ b/Engine/Source/Runtime/Renderer/Private/LightSceneInfo.cpp
@@ -243,12 +243,6 @@ FLightPrimitiveInteraction* FLightSceneInfo::GetDynamicInteractionStaticPrimitiv
 	return DynamicInteractionStaticPrimitiveList;
 }
 
-void FLightSceneInfo::ReleaseRHI()
-{
-	ShadowCapsuleShapesVertexBuffer.SafeRelease();
-	ShadowCapsuleShapesSRV.SafeRelease();
-}
-
 void FLightSceneInfo::ConditionalUpdateMobileMovablePointLightUniformBuffer(const FSceneRenderer* SceneRenderer)
 {
 	QUICK_SCOPE_CYCLE_COUNTER(STAT_FLightSceneProxy_UpdateMobileMovablePointLightUniformBuffer);
diff --git a/Engine/Source/Runtime/Renderer/Private/LightSceneInfo.h b/Engine/Source/Runtime/Renderer/Private/LightSceneInfo.h
--- a/Engine/Source/Runtime/Renderer/Private/LightSceneInfo.h
+++ b/Engine/Source/Runtime/Renderer/Private/LightSceneInfo.h
@@ -197,9 +197,6 @@ public:
 	/** The identifier for the primitive in Scene->PrimitiveOctree. */
 	FOctreeElementId2 OctreeId;
 
-	mutable FBufferRHIRef ShadowCapsuleShapesVertexBuffer;
-	mutable FShaderResourceViewRHIRef ShadowCapsuleShapesSRV;
-
 	/** Persistent shadow state used for HZB occlusion culling. */
 	TMap<FPersistentShadowStateKey, FPersistentShadowState> PrevPersistentShadows;
 	TMap<FPersistentShadowStateKey, FPersistentShadowState> PersistentShadows;
@@ -343,9 +340,6 @@ public:
 		return (uint32)LightSceneInfo->Id;
 	}
 
-	// FRenderResource interface.
-	virtual void ReleaseRHI();
-
 	// Update the mobile movable point light uniform buffer before it is used for mobile base pass rendering.
 	void ConditionalUpdateMobileMovablePointLightUniformBuffer(const class FSceneRenderer* SceneRenderer);
 
diff --git a/Engine/Source/Runtime/Renderer/Private/ScenePrivate.h b/Engine/Source/Runtime/Renderer/Private/ScenePrivate.h
--- a/Engine/Source/Runtime/Renderer/Private/ScenePrivate.h
+++ b/Engine/Source/Runtime/Renderer/Private/ScenePrivate.h
@@ -1060,15 +1060,6 @@ public:
 	FRadianceCacheState RadianceCacheState;
 	FRadianceCacheState RadiosityRadianceCacheState;
 
-	FBufferRHIRef IndirectShadowCapsuleShapesVertexBuffer;
-	FShaderResourceViewRHIRef IndirectShadowCapsuleShapesSRV;
-	FBufferRHIRef IndirectShadowMeshDistanceFieldCasterIndicesVertexBuffer;
-	FShaderResourceViewRHIRef IndirectShadowMeshDistanceFieldCasterIndicesSRV;
-	FBufferRHIRef IndirectShadowLightDirectionVertexBuffer;
-	FShaderResourceViewRHIRef IndirectShadowLightDirectionSRV;
-	FRWBuffer IndirectShadowVolumetricLightmapDerivedLightDirection;
-	FRWBuffer CapsuleTileIntersectionCountsBuffer;
-
 	/** Timestamp queries around separate translucency, used for auto-downsampling. */
 	FRenderQueryPoolRHIRef TimerQueryPool;
 	FLatentGPUTimer TranslucencyTimer;
