Makes r.TSR.Translucency.SeparateTemporalAccumulation=1 production ready

diff --git a/Engine/Shaders/Private/TemporalSuperResolution/TSRCompareTranslucency.usf b/Engine/Shaders/Private/TemporalSuperResolution/TSRCompareTranslucency.usf
--- a/Engine/Shaders/Private/TemporalSuperResolution/TSRCompareTranslucency.usf
+++ b/Engine/Shaders/Private/TemporalSuperResolution/TSRCompareTranslucency.usf
@@ -49,6 +49,8 @@ float2 PrevTranslucencyInfo_UVViewportBilinearMin;
 float2 PrevTranslucencyInfo_UVViewportBilinearMax;
 float PrevTranslucencyPreExposureCorrection;
 
+float TranslucencyHighlightLuminance;
+
 FScreenTransform InputPixelPosToScreenPos;
 FScreenTransform ScreenPosToPrevTranslucencyTextureUV;
 
@@ -66,6 +68,27 @@ taa_half4 TransformToComparisonColorSpace(taa_half4 Color)
 	return taa_half4(RGBToYCoCg(Color.rgb), Color.a);
 }
 
+taa_half MeasureTranslucencyRejectionFactor(
+	taa_half4 PrevTranslucency,
+	taa_half4 Translucency,
+	taa_half4 TranslucencyMin,
+	taa_half4 TranslucencyMax)
+{
+	taa_half IsHighlight = saturate(PrevTranslucency.x - taa_half(TranslucencyHighlightLuminance));
+
+	taa_half4 ClampedPrevTranslucency = clamp(PrevTranslucency, TranslucencyMin, TranslucencyMax);
+	
+	taa_half BackbufferQuantizationError = 0.0; //MeasureBackbufferLDRQuantizationErrorFromHDRLuma(max(PrevTranslucency[0], Translucency[0]));
+
+	taa_half Rejection = MeasureRejectionFactor(
+		PrevTranslucency.xyz, ClampedPrevTranslucency.xyz,
+		Translucency.xyz, TranslucencyMin.xyz, TranslucencyMax.xyz,
+		BackbufferQuantizationError);
+
+	Rejection = lerp(1.0, Rejection, IsHighlight);
+
+	return Rejection;
+}
 
 //------------------------------------------------------- ENTRY POINT
 
@@ -126,10 +149,18 @@ void MainCS(
 			}
 			else
 			{
+				#if CONFIG_KERNEL
+					Translucency += TranslucencyColor;
+				#endif
+
 				TranslucencyMin = min(TranslucencyMin, TranslucencyColor);
 				TranslucencyMax = max(TranslucencyMax, TranslucencyColor);
 			}
 		}
+
+		#if CONFIG_KERNEL
+			Translucency *= rcp(taa_half(CONFIG_SAMPLES));
+		#endif
 	}
 
 	// Sample previous frame's translucency
@@ -176,11 +207,19 @@ void MainCS(
 			}
 			else
 			{
+				#if CONFIG_KERNEL
+					PrevTranslucency += PrevTranslucencyColor;
+				#endif
+
 				PrevTranslucencyMin = min(PrevTranslucencyMin, PrevTranslucencyColor);
 				PrevTranslucencyMax = max(PrevTranslucencyMax, PrevTranslucencyColor);
 			}
 		}
 
+		#if CONFIG_KERNEL
+			PrevTranslucency *= rcp(taa_half(CONFIG_SAMPLES));
+		#endif
+
 		PrevTranslucency *= PrevTranslucencyPreExposureCorrection;
 		PrevTranslucencyMin *= PrevTranslucencyPreExposureCorrection;
 		PrevTranslucencyMax *= PrevTranslucencyPreExposureCorrection;
@@ -189,20 +228,19 @@ void MainCS(
 	taa_half Rejection = 1.0;
 	#if 1
 	{
-		taa_half4 ClampedPrevTranslucency = clamp(PrevTranslucency, TranslucencyMin, TranslucencyMax);
-		taa_half4 ClampedTranslucency = clamp(Translucency, PrevTranslucencyMin, PrevTranslucencyMax);
-
 		// TODO: reject with alpha too? Might not need since translucency is premultiplied RGBA 
-		taa_half ForwardRejection = MeasureRejectionFactor(
-			PrevTranslucency.xyz, ClampedPrevTranslucency.xyz,
-			Translucency.xyz, TranslucencyMin.xyz, TranslucencyMax.xyz,
-			MeasureBackbufferLDRQuantizationErrorFromHDRLuma(max(PrevTranslucency[0], ClampedPrevTranslucency[0])));
-		
-		taa_half BackwardRejection = MeasureRejectionFactor(
-			Translucency.xyz, ClampedTranslucency.xyz,
-			PrevTranslucency.xyz, PrevTranslucencyMin.xyz, PrevTranslucencyMax.xyz,
-			MeasureBackbufferLDRQuantizationErrorFromHDRLuma(max(Translucency[0], ClampedTranslucency[0])));
+		taa_half ForwardRejection = MeasureTranslucencyRejectionFactor(
+			PrevTranslucency,
+			Translucency,
+			TranslucencyMin,
+			TranslucencyMax);
 		
+		taa_half BackwardRejection = MeasureTranslucencyRejectionFactor(
+			Translucency,
+			PrevTranslucency,
+			PrevTranslucencyMin,
+			PrevTranslucencyMax);
+
 		Rejection = min(ForwardRejection, BackwardRejection);
 	}
 	#endif
diff --git a/Engine/Shaders/Private/TemporalSuperResolution/TSRDebugHistory.usf b/Engine/Shaders/Private/TemporalSuperResolution/TSRDebugHistory.usf
--- a/Engine/Shaders/Private/TemporalSuperResolution/TSRDebugHistory.usf
+++ b/Engine/Shaders/Private/TemporalSuperResolution/TSRDebugHistory.usf
@@ -99,7 +99,7 @@ void MainCS(
 		Debug.b = SubpixelId;
 	}
 
-	#if DEBUG_OUTPUT
+	#if DEBUG_OUTPUT && 0
 	if (all(DispatchThreadId < HistoryInfo_ViewportMax * SUB_PIXEL_GRID_SIZE))
 	{
 		DebugOutput[DispatchThreadId] = Debug;
diff --git a/Engine/Shaders/Private/TemporalSuperResolution/TSRUpdateHistory.usf b/Engine/Shaders/Private/TemporalSuperResolution/TSRUpdateHistory.usf
--- a/Engine/Shaders/Private/TemporalSuperResolution/TSRUpdateHistory.usf
+++ b/Engine/Shaders/Private/TemporalSuperResolution/TSRUpdateHistory.usf
@@ -25,11 +25,13 @@
 #if DIM_SEPARATE_TRANSLUCENCY == 1
 	#define CONFIG_ACCUMULATE_TRANSLUCENCY_SEPARATLY 1
 	#define CONFIG_VFX_SURVIVAL 0
-	#define CONFIG_RESPONSIVE_STENCIL 0
+	#define CONFIG_RESPONSIVE_STENCIL 1
+	#define CONFIG_REJECT_TRAMSLUCENCY_ON_RESPONSIVE_AA_ONLY 1
 #else
 	#define CONFIG_ACCUMULATE_TRANSLUCENCY_SEPARATLY 0
 	#define CONFIG_VFX_SURVIVAL 1
 	#define CONFIG_RESPONSIVE_STENCIL 1
+	#define CONFIG_REJECT_TRAMSLUCENCY_ON_RESPONSIVE_AA_ONLY 0
 #endif
 
 #if CONFIG_COMPILE_FP16
@@ -69,13 +71,15 @@ uint ResponsiveStencilMask;
 Texture2D<taa_half3> PrevHistory_LowFrequency;
 Texture2D<taa_half3> PrevHistory_HighFrequency;
 Texture2D<taa_half2> PrevHistory_Metadata;
-Texture2D<taa_half4> PrevHistory_Translucency;
+Texture2D<taa_half3> PrevHistory_Translucency;
+Texture2D<taa_half>  PrevHistory_TranslucencyAlpha;
 Texture2D<taa_subpixel_details> PrevHistory_SubpixelDetails;
 
 RWTexture2D<taa_half3> HistoryOutput_LowFrequency;
 RWTexture2D<taa_half3> HistoryOutput_HighFrequency;
 RWTexture2D<taa_half2> HistoryOutput_Metadata;
-RWTexture2D<taa_half4> HistoryOutput_Translucency;
+RWTexture2D<taa_half3> HistoryOutput_Translucency;
+RWTexture2D<taa_half>  HistoryOutput_TranslucencyAlpha;
 RWTexture2D<taa_subpixel_details> HistoryOutput_SubpixelDetails;
 
 RWTexture2D<taa_half3> SceneColorOutput;
@@ -95,6 +99,14 @@ groupshared taa_half4 SharedArray2[GROUP_SHARED_SIZE];
 groupshared taa_half4 SharedArray3[GROUP_SHARED_SIZE];
 groupshared taa_half4 SharedArray4[GROUP_SHARED_SIZE];
 groupshared taa_half4 SharedArray5[GROUP_SHARED_SIZE];
+
+groupshared taa_half4 SharedArray6[GROUP_SHARED_SIZE];
+groupshared taa_half4 SharedArray7[GROUP_SHARED_SIZE];
+groupshared taa_half4 SharedArray8[GROUP_SHARED_SIZE];
+groupshared taa_half4 SharedArray9[GROUP_SHARED_SIZE];
+groupshared taa_half4 SharedArray10[GROUP_SHARED_SIZE];
+groupshared taa_half4 SharedArray11[GROUP_SHARED_SIZE];
+
 groupshared taa_subpixel_payload2 SharedParallaxFactorBits[GROUP_SHARED_SIZE];
 
 
@@ -279,9 +291,6 @@ void MainCS(
 					taa_short2x2(PPCo),
 					InputPixelPosMin, InputPixelPosMax);
 			}
-			
-			Debug[0] = abs(TexelOffset[0]);
-			Debug[1] = abs(TexelOffset[1]);
 		}
 		#endif
 
@@ -313,6 +322,10 @@ void MainCS(
 		#if CONFIG_RESPONSIVE_STENCIL
 		{
 			bIsResponsiveAAPixel = (SceneStencilRef & ResponsiveStencilMask) != 0;
+
+			#if CONFIG_REJECT_TRAMSLUCENCY_ON_RESPONSIVE_AA_ONLY
+				TranslucencyRejection = max(TranslucencyRejection, bIsResponsiveAAPixel ? taa_half(0.0) : taa_half(1.0));
+			#endif
 		}
 		#endif
 
@@ -454,6 +467,16 @@ void MainCS(
 			SharedArray3[LocalGroupThreadIndex] = taa_half4(InputMinColor[2], InputPixelAlignement);
 			SharedArray4[LocalGroupThreadIndex] = taa_half4(InputMaxColor[0], InputMaxColor[1]);
 			SharedArray5[LocalGroupThreadIndex] = taa_half4(InputMaxColor[2], OutputPixelVelocity);
+
+			#if CONFIG_ACCUMULATE_TRANSLUCENCY_SEPARATLY
+				SharedArray6[LocalGroupThreadIndex] = taa_half4(FilteredTranslucencyColor[0], FilteredTranslucencyColor[1]);
+				SharedArray7[LocalGroupThreadIndex] = taa_half4(FilteredTranslucencyColor[2], FilteredTranslucencyColor[3]);
+				SharedArray8[LocalGroupThreadIndex] = taa_half4(TranslucencyMinColor[0], TranslucencyMinColor[1]);
+				SharedArray9[LocalGroupThreadIndex] = taa_half4(TranslucencyMinColor[2], TranslucencyMinColor[3]);
+				SharedArray10[LocalGroupThreadIndex] = taa_half4(TranslucencyMaxColor[0], TranslucencyMaxColor[1]);
+				SharedArray11[LocalGroupThreadIndex] = taa_half4(TranslucencyMaxColor[2], TranslucencyMaxColor[3]);
+			#endif
+
 			SharedParallaxFactorBits[LocalGroupThreadIndex] = ParallaxFactorBits;
 		#else
 			SharedArray0[LocalGroupThreadIndex] = taa_half4(dpv_lo(FilteredInputColor), dpv_lo(LowFrequencyRejection));
@@ -517,8 +540,12 @@ void MainCS(
 						PrevHistory_Metadata.SampleLevel(GlobalBilinearClampedSampler, SampleUV0, 0),
 						PrevHistory_Metadata.SampleLevel(GlobalBilinearClampedSampler, SampleUV1, 0));
 					taa_half4x2 SampleTranslucency = dpv_interleave_registers(
-						PrevHistory_Translucency.SampleLevel(GlobalBilinearClampedSampler, SampleUV0, 0),
-						PrevHistory_Translucency.SampleLevel(GlobalBilinearClampedSampler, SampleUV1, 0));
+						taa_half4(
+							PrevHistory_Translucency.SampleLevel(GlobalBilinearClampedSampler, SampleUV0, 0),
+							PrevHistory_TranslucencyAlpha.SampleLevel(GlobalBilinearClampedSampler, SampleUV0, 0)),
+						taa_half4(
+							PrevHistory_Translucency.SampleLevel(GlobalBilinearClampedSampler, SampleUV1, 0),
+							PrevHistory_TranslucencyAlpha.SampleLevel(GlobalBilinearClampedSampler, SampleUV1, 0)));
 
 					taa_half2 InvWeight = SafeRcp(SampleMetadata[1]);
 
@@ -604,6 +631,15 @@ void MainCS(
 			taa_half4 RawLDS3 = SharedArray3[LocalGroupThreadIndex];
 			taa_half4 RawLDS4 = SharedArray4[LocalGroupThreadIndex];
 			taa_half4 RawLDS5 = SharedArray5[LocalGroupThreadIndex];
+			
+			#if CONFIG_ACCUMULATE_TRANSLUCENCY_SEPARATLY
+				taa_half4 RawLDS6 = SharedArray6[LocalGroupThreadIndex];
+				taa_half4 RawLDS7 = SharedArray7[LocalGroupThreadIndex];
+				taa_half4 RawLDS8 = SharedArray8[LocalGroupThreadIndex];
+				taa_half4 RawLDS9 = SharedArray9[LocalGroupThreadIndex];
+				taa_half4 RawLDS10 = SharedArray10[LocalGroupThreadIndex];
+				taa_half4 RawLDS11 = SharedArray11[LocalGroupThreadIndex];
+			#endif
 
 			ParallaxFactorBits = SharedParallaxFactorBits[LocalGroupThreadIndex];
 
@@ -614,6 +650,12 @@ void MainCS(
 			LowFrequencyRejection = RawLDS1.ba;
 			InputPixelAlignement = RawLDS3.ba;
 			OutputPixelVelocity = RawLDS5.ba;
+			
+			#if CONFIG_ACCUMULATE_TRANSLUCENCY_SEPARATLY
+				FilteredTranslucencyColor = taa_half4x2(RawLDS6.rg, RawLDS6.ba, RawLDS7.rg, RawLDS7.ba);
+				TranslucencyMinColor = taa_half4x2(RawLDS8.rg, RawLDS8.ba, RawLDS9.rg, RawLDS9.ba);
+				TranslucencyMaxColor = taa_half4x2(RawLDS10.rg, RawLDS10.ba, RawLDS11.rg, RawLDS11.ba);
+			#endif
 		#else
 			taa_half4 RawLDS0 = SharedArray0[LocalGroupThreadIndex];
 			taa_half4 RawLDS1 = SharedArray1[LocalGroupThreadIndex];
@@ -731,17 +773,19 @@ void MainCS(
 			clamp(RejectedValidity + DesiredCurrentContribution, taa_half(0.0), PredictionOnlyValidity) +
 			clamp(RejectedValidity + DesiredCurrentContribution * PrevHistoryRejectionWeight * RejectionConfidentEnough - PredictionOnlyValidity, 0.0, 1.0 - PredictionOnlyValidity));
 
-		FLATTEN
-		if (bIsResponsiveAAPixel[0])
-		{
-			OutputValidity[0] = taa_half(0.0);
-		}
+		#if !CONFIG_ACCUMULATE_TRANSLUCENCY_SEPARATLY
+			FLATTEN
+			if (bIsResponsiveAAPixel[0])
+			{
+				OutputValidity[0] = taa_half(0.0);
+			}
 
-		FLATTEN
-		if (bIsResponsiveAAPixel[1])
-		{
-			OutputValidity[1] = taa_half(0.0);
-		}
+			FLATTEN
+			if (bIsResponsiveAAPixel[1])
+			{
+				OutputValidity[1] = taa_half(0.0);
+			}
+		#endif
 		
 		taa_half2 InvPrevHistoryValidity = SafeRcp(PrevHistoryValidity);
 
@@ -867,9 +911,11 @@ void MainCS(
 			}
 			#endif
 
-			// Responsive forces 1/4 of new frame.
-			BlendFinal[0] = bIsResponsiveAAPixel[0] ? taa_half(1.0/4.0) : BlendFinal[0];
-			BlendFinal[1] = bIsResponsiveAAPixel[1] ? taa_half(1.0/4.0) : BlendFinal[1];
+			#if !CONFIG_ACCUMULATE_TRANSLUCENCY_SEPARATLY
+				// Responsive forces 1/4 of new frame.
+				BlendFinal[0] = bIsResponsiveAAPixel[0] ? taa_half(1.0/4.0) : BlendFinal[0];
+				BlendFinal[1] = bIsResponsiveAAPixel[1] ? taa_half(1.0/4.0) : BlendFinal[1];
+			#endif
 
 			// Completely discard the history
 			{
@@ -899,24 +945,68 @@ void MainCS(
 	taa_half4x2 FinalTranslucencyColor;
 	#if CONFIG_ACCUMULATE_TRANSLUCENCY_SEPARATLY
 	{
+		taa_half2 TranslucencyLumaMin = Luma4(GetColorChannel(TranslucencyMinColor));
+		taa_half2 TranslucencyLumaMax = Luma4(GetColorChannel(TranslucencyMaxColor));
+
+		taa_half2 TranslucencyLumaHistory  = Luma4(GetColorChannel(PrevHistoryTranslucency));
+		taa_half2 TranslucencyLumaFiltered = Luma4(GetColorChannel(FilteredTranslucencyColor));
+
 		taa_half4x2 ClampedPrevHistoryTranslucency = clamp(PrevHistoryTranslucency, TranslucencyMinColor, TranslucencyMaxColor);
 		
 		taa_half2 HighFrequencyTranslucencyRejection = MeasureRejectionFactor(
 			GetColorChannel(PrevHistoryTranslucency), GetColorChannel(ClampedPrevHistoryTranslucency),
 			GetColorChannel(FilteredTranslucencyColor), GetColorChannel(TranslucencyMinColor), GetColorChannel(TranslucencyMaxColor),
 			taa_half(0.0));
-		
-		taa_half2 BlendFinal = max(taa_half(0.2) * InputPixelAlignement, taa_half(1.0) - TranslucencyRejection * HighFrequencyTranslucencyRejection);
 
+		taa_half2 BlendFinal = taa_half(0.2) * InputPixelAlignement;
+
+		// Antiflicker
+		#if 1
+		{
+			taa_half2 DistToClamp = min( abs(TranslucencyLumaHistory - TranslucencyLumaMin), abs(TranslucencyLumaHistory - TranslucencyLumaMax) ) / max3( TranslucencyLumaHistory, TranslucencyLumaFiltered, taa_half(1e-4) );
+			BlendFinal *= taa_half(0.2) + taa_half(0.8) * saturate(taa_half(0.5) * DistToClamp);
+		}
+		#endif
+			
+		// Make sure to have at least some small contribution
+		#if 0
+		{
+			BlendFinal = max( BlendFinal, saturate( taa_half(0.01) * TranslucencyLumaHistory / abs( TranslucencyLumaFiltered - TranslucencyLumaHistory ) ) );
+		}
+		#endif
+		
+		BlendFinal = max(BlendFinal, (taa_half(1.0) - TranslucencyRejection));
+		
 		FLATTEN
 		if (bCameraCut)
 		{
 			BlendFinal = taa_half(1.0);
 		}
+		
+		#if 0
+		{
+			taa_half2 OneMinusBlendFinal = taa_half(1.0) - BlendFinal;
+		
+			FinalTranslucencyColor = dpv_scale(ClampedPrevHistoryTranslucency, OneMinusBlendFinal) + dpv_scale(FilteredTranslucencyColor, BlendFinal);
+		}
+		#else
+		{
+			taa_half2 FilterWeight = HdrWeight4(GetColorChannel(FilteredTranslucencyColor));
+			taa_half2 ClampedHistoryWeight = HdrWeight4(GetColorChannel(ClampedPrevHistoryTranslucency));
 
-		taa_half2 OneMinusBlendFinal = taa_half(1.0) - BlendFinal;
+			
+			taa_half2 BlendA = saturate(taa_half(1.0) - 1.05 * BlendFinal) * ClampedHistoryWeight;
+			taa_half2 BlendB = BlendFinal * FilterWeight;
+			taa_half2 RcpBlend = SafeRcp(BlendA + BlendB);
+			BlendA *= RcpBlend;
+			BlendB *= RcpBlend;
+			taa_half2x2 Weights =  taa_half2x2(BlendA, BlendB);
 
-		FinalTranslucencyColor = dpv_scale(PrevHistoryTranslucency, OneMinusBlendFinal) + dpv_scale(FilteredTranslucencyColor, BlendFinal);
+			//taa_half2x2 Weights = WeightedLerpFactors(ClampedHistoryWeight, FilterWeight, BlendFinal);
+
+			FinalTranslucencyColor = dpv_scale(ClampedPrevHistoryTranslucency, Weights[0]) + dpv_scale(FilteredTranslucencyColor, Weights[1]);
+		}
+		#endif
 	}
 	#endif
 
@@ -1070,7 +1160,8 @@ void MainCS(
 			HistoryOutput_HighFrequency[dpv_lo(LocalHistoryPixelPos)] = dpv_lo(FinalHighFrequencyColor);
 			HistoryOutput_Metadata[dpv_lo(LocalHistoryPixelPos)]      = taa_half2(dpv_lo(FinalLowFrequencyWeight), dpv_lo(FinalHistoryValidity));
 			#if CONFIG_ACCUMULATE_TRANSLUCENCY_SEPARATLY
-				HistoryOutput_Translucency[dpv_lo(LocalHistoryPixelPos)] = dpv_lo(FinalTranslucencyColor);
+				HistoryOutput_Translucency[dpv_lo(LocalHistoryPixelPos)] = dpv_lo(FinalTranslucencyColor).rgb;
+				HistoryOutput_TranslucencyAlpha[dpv_lo(LocalHistoryPixelPos)] = dpv_lo(FinalTranslucencyColor).a;
 			#endif
 			HistoryOutput_SubpixelDetails[dpv_lo(LocalHistoryPixelPos)] = dpv_lo(FinalSubpixelDetails);
 
@@ -1079,7 +1170,8 @@ void MainCS(
 				HistoryOutput_HighFrequency[dpv_hi(LocalHistoryPixelPos)] = dpv_hi(FinalHighFrequencyColor);
 				HistoryOutput_Metadata[dpv_hi(LocalHistoryPixelPos)]      = taa_half2(dpv_hi(FinalLowFrequencyWeight), dpv_hi(FinalHistoryValidity));
 				#if CONFIG_ACCUMULATE_TRANSLUCENCY_SEPARATLY
-					HistoryOutput_Translucency[dpv_hi(LocalHistoryPixelPos)] = dpv_hi(FinalTranslucencyColor);
+					HistoryOutput_Translucency[dpv_hi(LocalHistoryPixelPos)] = dpv_hi(FinalTranslucencyColor).rgb;
+					HistoryOutput_TranslucencyAlpha[dpv_hi(LocalHistoryPixelPos)] = dpv_hi(FinalTranslucencyColor).a;
 				#endif
 				HistoryOutput_SubpixelDetails[dpv_hi(LocalHistoryPixelPos)] = dpv_hi(FinalSubpixelDetails);
 			#endif
diff --git a/Engine/Source/Runtime/Renderer/Private/PostProcess/PostProcessing.cpp b/Engine/Source/Runtime/Renderer/Private/PostProcess/PostProcessing.cpp
--- a/Engine/Source/Runtime/Renderer/Private/PostProcess/PostProcessing.cpp
+++ b/Engine/Source/Runtime/Renderer/Private/PostProcess/PostProcessing.cpp
@@ -149,7 +149,7 @@ class FComposeSeparateTranslucencyPS : public FGlobalShader
 	DECLARE_GLOBAL_SHADER(FComposeSeparateTranslucencyPS);
 	SHADER_USE_PARAMETER_STRUCT(FComposeSeparateTranslucencyPS, FGlobalShader);
 
-	class FNearestDepthNeighborUpsampling : SHADER_PERMUTATION_INT("PERMUTATION_NEARESTDEPTHNEIGHBOR", 2);
+	class FNearestDepthNeighborUpsampling : SHADER_PERMUTATION_BOOL("PERMUTATION_NEARESTDEPTHNEIGHBOR");
 	using FPermutationDomain = TShaderPermutationDomain<FNearestDepthNeighborUpsampling>;
 
 	BEGIN_SHADER_PARAMETER_STRUCT(FParameters, )
@@ -181,37 +181,81 @@ IMPLEMENT_GLOBAL_SHADER(FComposeSeparateTranslucencyPS, "/Engine/Private/Compose
 
 extern bool GetUseTranslucencyNearestDepthNeighborUpsample(float DownsampleScale);
 
-FRDGTextureRef AddPostDOFTranslucencyCompositionPass(FRDGBuilder& GraphBuilder, const FViewInfo& View, const FScreenPassTexture& SceneColor, const FScreenPassTexture& SceneDepth, const FSeparateTranslucencyTextures& SeparateTranslucencyTextures)
+FRDGTextureRef AddTranslucencyCompositionPass(
+	FRDGBuilder& GraphBuilder,
+	const FViewInfo& View,
+	const FScreenPassTexture& SceneColor,
+	const FScreenPassTexture& SceneDepth,
+	const FSeparateTranslucencyTextures& SeparateTranslucencyTextures,
+	bool bPostMotionBlur,
+	bool bApplyModulateOnly)
 {
-	// if nothing is rendered into the separate translucency, then just return the existing Scenecolor
-	if (!SeparateTranslucencyTextures.IsColorValid() && !SeparateTranslucencyTextures.IsColorModulateValid())
+	FRDGTextureRef SeparateModulationTexture;
+	FRDGTextureRef SeparateTranslucencyTexture;
+
+	if (bPostMotionBlur)
 	{
-		return SceneColor.Texture;
-	}
+		check(!bApplyModulateOnly);
 
-	FRDGTextureDesc OutputDesc = SceneColor.Texture->Desc;
-	OutputDesc.Reset();
+		// if nothing is rendered into the separate translucency, then just return the existing Scenecolor
+		if (!SeparateTranslucencyTextures.IsPostMotionBlurColorValid())
+		{
+			return SceneColor.Texture;
+		}
 
-	if (OutputDesc.Format == PF_FloatRGBA && !IsPostProcessingWithAlphaChannelSupported())
+		SeparateTranslucencyTexture = SeparateTranslucencyTextures.GetPostMotionBlurColorForRead(GraphBuilder);
+		SeparateModulationTexture = GraphBuilder.RegisterExternalTexture(GSystemTextures.WhiteDummy);
+	}
+	else
 	{
-		OutputDesc.Format = PF_FloatRGB;
+		// if nothing is rendered into the separate translucency, then just return the existing Scenecolor
+		if (!SeparateTranslucencyTextures.IsColorValid() && !SeparateTranslucencyTextures.IsColorModulateValid())
+		{
+			return SceneColor.Texture;
+		}
+
+		SeparateTranslucencyTexture = SeparateTranslucencyTextures.GetColorForRead(GraphBuilder);
+		SeparateModulationTexture = SeparateTranslucencyTextures.GetColorModulateForRead(GraphBuilder);
+
+		if (bApplyModulateOnly)
+		{
+			if (!SeparateTranslucencyTextures.IsColorModulateValid())
+			{
+				return SceneColor.Texture;
+			}
+
+			SeparateTranslucencyTexture = GraphBuilder.RegisterExternalTexture(GSystemTextures.BlackAlphaOneDummy);
+		}
 	}
+	
+	FRDGTextureRef NewSceneColor;
+	{
+		FRDGTextureDesc OutputDesc = SceneColor.Texture->Desc;
+		OutputDesc.Reset();
 
-	FRDGTextureRef NewSceneColor = GraphBuilder.CreateTexture(OutputDesc, TEXT("PostDOFTranslucency.SceneColor"));
-	FRDGTextureRef SeparateTranslucency = SeparateTranslucencyTextures.GetColorForRead(GraphBuilder);	
+		if (OutputDesc.Format == PF_FloatRGBA && !IsPostProcessingWithAlphaChannelSupported())
+		{
+			OutputDesc.Format = PF_FloatRGB;
+		}
 
-	const FIntRect SceneColorRect = SceneColor.ViewRect;
-	const FVector2D SceneColorSize(SceneColorRect.Size());
+		NewSceneColor = GraphBuilder.CreateTexture(
+			OutputDesc,
+			bPostMotionBlur ? TEXT("PostMotionBlurTranslucency.SceneColor") : TEXT("PostDOFTranslucency.SceneColor"));
+	}
+
+	const FVector2D SceneColorSize(SceneColor.ViewRect.Size());
 	const FVector2D SceneColorSizeInv = FVector2D(1.0f, 1.0f) / SceneColorSize;	
-	const FVector2D SceneColorExtent(OutputDesc.Extent);
+	const FVector2D SceneColorExtent(NewSceneColor->Desc.Extent);
 	const FVector2D SceneColorExtentInv = FVector2D(1.0f, 1.0f) / SceneColorExtent;
 	
 	const FIntRect SeparateTranslucencyRect = SeparateTranslucencyTextures.GetDimensions().GetViewport(View.ViewRect).Rect;
 	const FVector2D SeparateTranslucencySize(SeparateTranslucencyRect.Size());
-	const FVector2D SeparateTranslucencyExtent(SeparateTranslucency->Desc.Extent);
+	const FVector2D SeparateTranslucencyExtent((bApplyModulateOnly ? SeparateModulationTexture : SeparateTranslucencyTexture)->Desc.Extent);
 	const FVector2D SeparateTranslucencyExtentInv = FVector2D(1.0f, 1.0f) / SeparateTranslucencyExtent;
 	
 	const bool bScaleSeparateTranslucency = SeparateTranslucencySize != SceneColorSize;
+	const float DownsampleScale = SeparateTranslucencySize.X / SceneColorExtent.X;
+	const bool DepthUpscampling = !bPostMotionBlur && GetUseTranslucencyNearestDepthNeighborUpsample(DownsampleScale) ? 1 : 0;
 
 	FComposeSeparateTranslucencyPS::FParameters* PassParameters = GraphBuilder.AllocParameters<FComposeSeparateTranslucencyPS::FParameters>();
 	PassParameters->ScreenPosToSceneColorUV = SceneColorSizeInv * SceneColorSize * SceneColorExtentInv;
@@ -220,97 +264,37 @@ FRDGTextureRef AddPostDOFTranslucencyCompositionPass(FRDGBuilder& GraphBuilder,
 	PassParameters->SeparateTranslucencyExtentInverse = SeparateTranslucencyExtentInv;
 	PassParameters->SceneColor = SceneColor.Texture;
 	PassParameters->SceneColorSampler = TStaticSamplerState<SF_Point>::GetRHI();
-	PassParameters->SeparateTranslucency = SeparateTranslucency;
+	PassParameters->SeparateTranslucency = SeparateTranslucencyTexture;
 	PassParameters->SeparateTranslucencySampler = bScaleSeparateTranslucency ? TStaticSamplerState<SF_Bilinear>::GetRHI() : TStaticSamplerState<SF_Point>::GetRHI();
-	PassParameters->SeparateModulation = SeparateTranslucencyTextures.GetColorModulateForRead(GraphBuilder);
+	PassParameters->SeparateModulation = SeparateModulationTexture;
 	PassParameters->SeparateModulationSampler = bScaleSeparateTranslucency ? TStaticSamplerState<SF_Bilinear>::GetRHI() : TStaticSamplerState<SF_Point>::GetRHI();
 	PassParameters->ViewUniformBuffer = View.ViewUniformBuffer;
 	PassParameters->RenderTargets[0] = FRenderTargetBinding(NewSceneColor, ERenderTargetLoadAction::ENoAction);
 
-	PassParameters->LowResDepthTexture = SeparateTranslucencyTextures.GetDepthForRead(GraphBuilder);
-	PassParameters->LowResDepthSampler = TStaticSamplerState<SF_Point>::GetRHI();
-	PassParameters->FullResDepthTexture = SceneDepth.Texture;
-	PassParameters->FullResDepthSampler = TStaticSamplerState<SF_Point>::GetRHI();
-
-	FComposeSeparateTranslucencyPS::FPermutationDomain PermutationVector;
-	const float DownsampleScale = SeparateTranslucencySize.X / SceneColorExtent.X;
-	PermutationVector.Set<FComposeSeparateTranslucencyPS::FNearestDepthNeighborUpsampling>(GetUseTranslucencyNearestDepthNeighborUpsample(DownsampleScale) ? 1 : 0);
-
-	TShaderMapRef<FComposeSeparateTranslucencyPS> PixelShader(View.ShaderMap, PermutationVector);
-	FPixelShaderUtils::AddFullscreenPass(
-		GraphBuilder,
-		View.ShaderMap,
-		RDG_EVENT_NAME(
-			"ComposePostDOFTranslucency%s %dx%d",
-			bScaleSeparateTranslucency ? TEXT(" Rescale") : TEXT(""),
-			SceneColorRect.Width(), SceneColorRect.Height()),
-		PixelShader,
-		PassParameters,
-		SceneColorRect);
-
-	return NewSceneColor;
-}
-
-FRDGTextureRef AddPostMotionBlurTranslucencyCompositionPass(FRDGBuilder& GraphBuilder, const FViewInfo& View, const FScreenPassTexture& SceneColor, const FSeparateTranslucencyTextures& SeparateTranslucencyTextures)
-{
-	// if nothing is rendered into the separate translucency, then just return the existing Scenecolor
-	if (!SeparateTranslucencyTextures.IsPostMotionBlurColorValid())
-	{
-		return SceneColor.Texture;
-	}
-
-	FRDGTextureDesc OutputDesc = SceneColor.Texture->Desc;
-	OutputDesc.Reset();
-
-	if (OutputDesc.Format == PF_FloatRGBA && !IsPostProcessingWithAlphaChannelSupported())
+	if (DepthUpscampling)
 	{
-		OutputDesc.Format = PF_FloatRGB;
+		PassParameters->LowResDepthTexture = SeparateTranslucencyTextures.GetDepthForRead(GraphBuilder);
+		PassParameters->LowResDepthSampler = TStaticSamplerState<SF_Point>::GetRHI();
+		PassParameters->FullResDepthTexture = SceneDepth.Texture;
+		PassParameters->FullResDepthSampler = TStaticSamplerState<SF_Point>::GetRHI();
 	}
 
-	FRDGTextureRef NewSceneColor = GraphBuilder.CreateTexture(OutputDesc, TEXT("PostMotionBlurTranslucency.SceneColor"));
-	FRDGTextureRef SeparateTranslucency = SeparateTranslucencyTextures.GetPostMotionBlurColorForRead(GraphBuilder);
-
-	const FIntRect SceneColorRect = SceneColor.ViewRect;
-	const FVector2D SceneColorSize(SceneColorRect.Size());
-	const FVector2D SceneColorSizeInv = FVector2D(1.0f, 1.0f) / SceneColorSize;
-	const FVector2D SceneColorExtent(OutputDesc.Extent);
-	const FVector2D SceneColorExtentInv = FVector2D(1.0f, 1.0f) / SceneColorExtent;
-
-	const FIntRect SeparateTranslucencyRect = SeparateTranslucencyTextures.GetDimensions().GetViewport(View.ViewRect).Rect;
-	const FVector2D SeparateTranslucencySize(SeparateTranslucencyRect.Size());
-	const FVector2D SeparateTranslucencyExtent(SeparateTranslucency->Desc.Extent);
-	const FVector2D SeparateTranslucencyExtentInv = FVector2D(1.0f, 1.0f) / SeparateTranslucencyExtent;
-
-	const bool bScaleSeparateTranslucency = SeparateTranslucencySize != SceneColorSize;
-
-	FComposeSeparateTranslucencyPS::FParameters* PassParameters = GraphBuilder.AllocParameters<FComposeSeparateTranslucencyPS::FParameters>();
-	PassParameters->ScreenPosToSceneColorUV = SceneColorSizeInv * SceneColorSize * SceneColorExtentInv;
-	PassParameters->ScreenPosToSeparateTranslucencyUV = SceneColorSizeInv * SeparateTranslucencySize * SeparateTranslucencyExtentInv;
-	PassParameters->SeparateTranslucencyUVMax = (SeparateTranslucencySize - FVector2D(0.5f, 0.5f)) * SeparateTranslucencyExtentInv;
-	PassParameters->SeparateTranslucencyExtentInverse = SeparateTranslucencyExtentInv;
-	PassParameters->SceneColor = SceneColor.Texture;
-	PassParameters->SceneColorSampler = TStaticSamplerState<SF_Point>::GetRHI();
-	PassParameters->SeparateTranslucency = SeparateTranslucency;
-	PassParameters->SeparateTranslucencySampler = bScaleSeparateTranslucency ? TStaticSamplerState<SF_Bilinear>::GetRHI() : TStaticSamplerState<SF_Point>::GetRHI();
-	PassParameters->SeparateModulation = GraphBuilder.RegisterExternalTexture(GSystemTextures.WhiteDummy);
-	PassParameters->SeparateModulationSampler = TStaticSamplerState<SF_Point>::GetRHI();
-	PassParameters->ViewUniformBuffer = View.ViewUniformBuffer;
-	PassParameters->RenderTargets[0] = FRenderTargetBinding(NewSceneColor, ERenderTargetLoadAction::ENoAction);
-
 	FComposeSeparateTranslucencyPS::FPermutationDomain PermutationVector;
-	PermutationVector.Set<FComposeSeparateTranslucencyPS::FNearestDepthNeighborUpsampling>(false);
+	PermutationVector.Set<FComposeSeparateTranslucencyPS::FNearestDepthNeighborUpsampling>(DepthUpscampling);
 
 	TShaderMapRef<FComposeSeparateTranslucencyPS> PixelShader(View.ShaderMap, PermutationVector);
 	FPixelShaderUtils::AddFullscreenPass(
 		GraphBuilder,
 		View.ShaderMap,
 		RDG_EVENT_NAME(
-			"ComposePostMotionBlurTranslucency%s %dx%d",
+			"ComposeTranslucency(%s%s%s) %dx%d",
+			bPostMotionBlur ? TEXT("PostMB") : TEXT("PostDOF"),
+			bApplyModulateOnly ? TEXT(" ModulateOnly") : TEXT(""),
 			bScaleSeparateTranslucency ? TEXT(" Rescale") : TEXT(""),
-			SceneColorRect.Width(), SceneColorRect.Height()),
+			SceneColor.ViewRect.Width(), SceneColor.ViewRect.Height()),
 		PixelShader,
 		PassParameters,
-		SceneColorRect);
+		SceneColor.ViewRect);
 
 	return NewSceneColor;
 }
@@ -590,14 +574,14 @@ void AddPostProcessingPasses(FRDGBuilder& GraphBuilder, const FViewInfo& View, c
 
 		// Diaphragm Depth of Field
 		{
-			FRDGTextureRef LocalSceneColorTexture = SceneColor.Texture;
+			FRDGTextureRef InputSceneColorTexture = SceneColor.Texture;
 
 			if (bDepthOfFieldEnabled)
 			{
 				FSeparateTranslucencyDimensions DummyTranslucencyDimensions;
 				FSeparateTranslucencyTextures DummyTranslucency(DummyTranslucencyDimensions);
 
-				LocalSceneColorTexture = DiaphragmDOF::AddPasses(
+				SceneColor.Texture = DiaphragmDOF::AddPasses(
 					GraphBuilder,
 					SceneTextureParameters,
 					View,
@@ -606,13 +590,18 @@ void AddPostProcessingPasses(FRDGBuilder& GraphBuilder, const FViewInfo& View, c
 			}
 
 			// DOF passes were not added, therefore need to compose Separate translucency manually.
-			if (LocalSceneColorTexture == SceneColor.Texture && !bComposeSeparateTranslucencyInTSR)
+			if (SceneColor.Texture == InputSceneColorTexture || bComposeSeparateTranslucencyInTSR)
 			{
-				LocalSceneColorTexture = AddPostDOFTranslucencyCompositionPass(GraphBuilder, View, SceneColor, SceneDepth, *Inputs.SeparateTranslucencyTextures);
+				SceneColor.Texture = AddTranslucencyCompositionPass(
+					GraphBuilder,
+					View,
+					SceneColor,
+					SceneDepth,
+					*Inputs.SeparateTranslucencyTextures,
+					/* bPostMotionBlur = */ false,
+					/* bApplyModulateOnly = */ bComposeSeparateTranslucencyInTSR);
 			}
 
-			SceneColor.Texture = LocalSceneColorTexture;
-
 			if (GetHairStrandsComposition() == EHairStrandsCompositionType::AfterSeparateTranslucent)
 			{
 				RenderHairComposition(GraphBuilder, View, SceneColor.Texture, SceneDepth.Texture);
@@ -736,7 +725,15 @@ void AddPostProcessingPasses(FRDGBuilder& GraphBuilder, const FViewInfo& View, c
 		else
 		{
 			// Compose Post-MotionBlur translucency
-			SceneColor.Texture = AddPostMotionBlurTranslucencyCompositionPass(GraphBuilder, View, SceneColor, *Inputs.SeparateTranslucencyTextures);
+
+			SceneColor.Texture = AddTranslucencyCompositionPass(
+				GraphBuilder,
+				View,
+				SceneColor,
+				SceneDepth,
+				*Inputs.SeparateTranslucencyTextures,
+				/* bPostMotionBlur = */ true,
+				/* bApplyModulateOnly = */ false);
 		}
 
 		SceneColor = AddAfterPass(EPass::MotionBlur, SceneColor);
@@ -958,8 +955,22 @@ void AddPostProcessingPasses(FRDGBuilder& GraphBuilder, const FViewInfo& View, c
 		PassSequence.Finalize();
 
 		// Compose separate translucency passes
-		SceneColor.Texture = AddPostDOFTranslucencyCompositionPass(GraphBuilder, View, SceneColor, SceneDepth, *Inputs.SeparateTranslucencyTextures);
-		SceneColor.Texture = AddPostMotionBlurTranslucencyCompositionPass(GraphBuilder, View, SceneColor, *Inputs.SeparateTranslucencyTextures);
+		SceneColor.Texture = AddTranslucencyCompositionPass(
+			GraphBuilder,
+			View,
+			SceneColor,
+			SceneDepth,
+			*Inputs.SeparateTranslucencyTextures,
+			/* bPostMotionBlur = */ false,
+			/* bApplyModulateOnly = */ false);
+		SceneColor.Texture = AddTranslucencyCompositionPass(
+			GraphBuilder,
+			View,
+			SceneColor,
+			SceneDepth,
+			*Inputs.SeparateTranslucencyTextures,
+			/* bPostMotionBlur = */ true,
+			/* bApplyModulateOnly = */ false);
 
 		SceneColorBeforeTonemap = SceneColor;
 
diff --git a/Engine/Source/Runtime/Renderer/Private/PostProcess/TemporalSuperResolution.cpp b/Engine/Source/Runtime/Renderer/Private/PostProcess/TemporalSuperResolution.cpp
--- a/Engine/Source/Runtime/Renderer/Private/PostProcess/TemporalSuperResolution.cpp
+++ b/Engine/Source/Runtime/Renderer/Private/PostProcess/TemporalSuperResolution.cpp
@@ -50,6 +50,11 @@ TAutoConsoleVariable<int32> CVarTSREnableAntiInterference(
 	TEXT("Enable heuristic to detect geometric interference between input pixel grid alignement and structured geometry."),
 	ECVF_RenderThreadSafe);
 
+TAutoConsoleVariable<float> CVarTSRTranslucencyHighlightLuminance(
+	TEXT("r.TSR.Translucency.HighlightLuminance"), -1.0f,
+	TEXT("Sets the liminance at which translucency is considered an highlights (default=-1.0)."),
+	ECVF_RenderThreadSafe);
+
 TAutoConsoleVariable<int32> CVarTSRTranslucencyPreviousFrameRejection(
 	TEXT("r.TSR.Translucency.PreviousFrameRejection"), 0,
 	TEXT("Enable heuristic to reject Separate translucency based on previous frame translucency."),
@@ -102,6 +107,7 @@ BEGIN_SHADER_PARAMETER_STRUCT(FTSRHistoryTextures, )
 	SHADER_PARAMETER_RDG_TEXTURE(Texture2D, Metadata)
 	SHADER_PARAMETER_RDG_TEXTURE(Texture2D, SubpixelDetails)
 	SHADER_PARAMETER_RDG_TEXTURE(Texture2D, Translucency)
+	SHADER_PARAMETER_RDG_TEXTURE(Texture2D, TranslucencyAlpha)
 END_SHADER_PARAMETER_STRUCT()
 
 BEGIN_SHADER_PARAMETER_STRUCT(FTSRHistoryUAVs, )
@@ -110,6 +116,7 @@ BEGIN_SHADER_PARAMETER_STRUCT(FTSRHistoryUAVs, )
 	SHADER_PARAMETER_RDG_TEXTURE_UAV(RWTexture2D, Metadata)
 	SHADER_PARAMETER_RDG_TEXTURE_UAV(RWTexture2D, SubpixelDetails)
 	SHADER_PARAMETER_RDG_TEXTURE_UAV(RWTexture2D, Translucency)
+	SHADER_PARAMETER_RDG_TEXTURE_UAV(RWTexture2D, TranslucencyAlpha)
 END_SHADER_PARAMETER_STRUCT()
 
 BEGIN_SHADER_PARAMETER_STRUCT(FTSRPrevHistoryParameters, )
@@ -128,6 +135,7 @@ FTSRHistoryUAVs CreateUAVs(FRDGBuilder& GraphBuilder, const FTSRHistoryTextures&
 	if (Textures.Translucency)
 	{
 		UAVs.Translucency = GraphBuilder.CreateUAV(Textures.Translucency);
+		UAVs.TranslucencyAlpha = GraphBuilder.CreateUAV(Textures.TranslucencyAlpha);
 	}
 	return UAVs;
 }
@@ -235,6 +243,7 @@ class FTSRCompareTranslucencyCS : public FTSRShader
 		SHADER_PARAMETER_STRUCT(FScreenPassTextureViewportParameters, TranslucencyInfo)
 		SHADER_PARAMETER_STRUCT(FScreenPassTextureViewportParameters, PrevTranslucencyInfo)
 		SHADER_PARAMETER(float, PrevTranslucencyPreExposureCorrection)
+		SHADER_PARAMETER(float, TranslucencyHighlightLuminance)
 
 		SHADER_PARAMETER(FScreenTransform, InputPixelPosToScreenPos)
 		SHADER_PARAMETER(FScreenTransform, ScreenPosToPrevTranslucencyTextureUV)
@@ -734,6 +743,7 @@ void AddTemporalSuperResolutionPasses(
 		DummyHistory.HighFrequency = BlackDummy;
 		DummyHistory.Metadata = BlackDummy;
 		DummyHistory.Translucency = BlackDummy;
+		DummyHistory.TranslucencyAlpha = WhiteDummy;
 
 		DummyHistory.SubpixelDetails = BlackUintDummy;
 	}
@@ -746,7 +756,8 @@ void AddTemporalSuperResolutionPasses(
 		PrevHistory.LowFrequency = GraphBuilder.RegisterExternalTexture(InputHistory.LowFrequency);
 		PrevHistory.HighFrequency = GraphBuilder.RegisterExternalTexture(InputHistory.HighFrequency);
 		PrevHistory.Metadata = GraphBuilder.RegisterExternalTexture(InputHistory.Metadata);
-		PrevHistory.Translucency = InputHistory.Translucency.IsValid() ? GraphBuilder.RegisterExternalTexture(InputHistory.Translucency) : BlackDummy;
+		PrevHistory.Translucency = InputHistory.Translucency.IsValid() ? GraphBuilder.RegisterExternalTexture(InputHistory.Translucency) : DummyHistory.Translucency;
+		PrevHistory.TranslucencyAlpha = InputHistory.TranslucencyAlpha.IsValid() ? GraphBuilder.RegisterExternalTexture(InputHistory.TranslucencyAlpha) : DummyHistory.TranslucencyAlpha;
 
 		// Register non-filterable history
 		PrevHistory.SubpixelDetails = GraphBuilder.RegisterExternalTexture(InputHistory.SubpixelDetails);
@@ -800,8 +811,11 @@ void AddTemporalSuperResolutionPasses(
 
 		if (bAccumulateSeparateTranslucency)
 		{
-			Desc.Format = PF_FloatRGBA;
+			Desc.Format = (CVarTSRR11G11B10History.GetValueOnRenderThread() != 0) ? PF_FloatR11G11B10 : PF_FloatRGBA;
 			History.Translucency = GraphBuilder.CreateTexture(Desc, TEXT("TSR.History.Translucency"));
+
+			Desc.Format = PF_R8;
+			History.TranslucencyAlpha = GraphBuilder.CreateTexture(Desc, TEXT("TSR.History.TranslucencyAlpha"));
 		}
 
 		Desc.Format = PF_R16_UINT;
@@ -957,6 +971,7 @@ void AddTemporalSuperResolutionPasses(
 			InputExtent, InputRect));
 		PassParameters->PrevTranslucencyInfo = GetScreenPassTextureViewportParameters(PrevTranslucencyViewport);
 		PassParameters->PrevTranslucencyPreExposureCorrection = PrevHistoryParameters.HistoryPreExposureCorrection;
+		PassParameters->TranslucencyHighlightLuminance = CVarTSRTranslucencyHighlightLuminance.GetValueOnRenderThread();
 
 		PassParameters->InputPixelPosToScreenPos = (FScreenTransform::Identity + 0.5) * CommonParameters.InputInfo.ViewportSizeInverse * FScreenTransform::ViewportUVToScreenPos;
 		PassParameters->ScreenPosToPrevTranslucencyTextureUV = FScreenTransform::ChangeTextureBasisFromTo(
@@ -1369,6 +1384,7 @@ void AddTemporalSuperResolutionPasses(
 		if (bAccumulateSeparateTranslucency)
 		{
 			GraphBuilder.QueueTextureExtraction(History.Translucency, &OutputHistory.Translucency);
+			GraphBuilder.QueueTextureExtraction(History.TranslucencyAlpha, &OutputHistory.TranslucencyAlpha);
 		}
 
 		// Extract non-filterable history
diff --git a/Engine/Source/Runtime/Renderer/Private/SceneRendering.h b/Engine/Source/Runtime/Renderer/Private/SceneRendering.h
--- a/Engine/Source/Runtime/Renderer/Private/SceneRendering.h
+++ b/Engine/Source/Runtime/Renderer/Private/SceneRendering.h
@@ -892,6 +892,7 @@ struct FTSRHistory
 	TRefCountPtr<IPooledRenderTarget> HighFrequency;
 	TRefCountPtr<IPooledRenderTarget> Metadata;
 	TRefCountPtr<IPooledRenderTarget> Translucency;
+	TRefCountPtr<IPooledRenderTarget> TranslucencyAlpha;
 
 	// Non-filterable output resolution
 	TRefCountPtr<IPooledRenderTarget> SubpixelDetails;
