Add a Mobile friendly RVT MaterialType BaseColor_Normal_Roughness with basecolor in RT0,  NormalXY and roughness in RT1, NormalZ is considered always a positive number( which is the situation for local space normal)

Add A  flag bUseLowQualityCompression to BaseColor_Normal_Roughness, which fakes compression by 16bit format(RGB565)

diff --git a/Engine/Plugins/FX/Niagara/Source/Niagara/Private/NiagaraDataInterfaceLandscape.cpp b/Engine/Plugins/FX/Niagara/Source/Niagara/Private/NiagaraDataInterfaceLandscape.cpp
--- a/Engine/Plugins/FX/Niagara/Source/Niagara/Private/NiagaraDataInterfaceLandscape.cpp
+++ b/Engine/Plugins/FX/Niagara/Source/Niagara/Private/NiagaraDataInterfaceLandscape.cpp
@@ -991,7 +991,7 @@ void UNiagaraDataInterfaceLandscape::ApplyLandscape(const FNiagaraSystemInstance
 						InstanceData.HeightVirtualTextureIndex = TextureIt;
 					}
 					break;
-
+				case ERuntimeVirtualTextureMaterialType::BaseColor_Normal_Roughness:
 				case ERuntimeVirtualTextureMaterialType::BaseColor_Normal_Specular:
 				case ERuntimeVirtualTextureMaterialType::BaseColor_Normal_Specular_YCoCg:
 				case ERuntimeVirtualTextureMaterialType::BaseColor_Normal_Specular_Mask_YCoCg:
@@ -1091,6 +1091,7 @@ private:
 		None = 0,
 		BC3BC3,
 		BC5BC1,
+		B5G6R5
 	};
 
 public:
@@ -1221,6 +1222,15 @@ public:
 
 		switch (ProxyData.NormalVirtualTextureMode)
 		{
+			case ERuntimeVirtualTextureMaterialType::BaseColor_Normal_Roughness:
+				PhysicalTextureSrv[0] = NormalAllocatedTexture->GetPhysicalTextureSRV(0, false);
+				NormalAllocatedTexture->GetPackedUniform(&NormalVirtualTextureUniforms[0], 0);
+
+				PhysicalTextureSrv[1] = NormalAllocatedTexture->GetPhysicalTextureSRV(1, false);
+				NormalAllocatedTexture->GetPackedUniform(&NormalVirtualTextureUniforms[1], 1);
+
+				SetShaderValue(RHICmdList, ComputeShaderRHI, NormalVirtualTextureUnpackModeParam, ENormalUnpackType::B5G6R5);
+				break;
 			case ERuntimeVirtualTextureMaterialType::BaseColor_Normal_Specular:
 				PhysicalTextureSrv[0] = NormalAllocatedTexture->GetPhysicalTextureSRV(0, false);
 				NormalAllocatedTexture->GetPackedUniform(&NormalVirtualTextureUniforms[0], 0);
diff --git a/Engine/Shaders/Private/VirtualTextureCommon.ush b/Engine/Shaders/Private/VirtualTextureCommon.ush
--- a/Engine/Shaders/Private/VirtualTextureCommon.ush
+++ b/Engine/Shaders/Private/VirtualTextureCommon.ush
@@ -4,7 +4,7 @@
 
 #include "Common.ush"
 #include "Random.ush"
-
+#include "GammaCorrectionCommon.ush"
 #ifndef VT_DISABLE_VIEW_UNIFORM_BUFFER
 #define VT_DISABLE_VIEW_UNIFORM_BUFFER 0
 #endif
@@ -727,3 +727,16 @@ float VirtualTextureUnpackHeight(in float4 PackedValue, in float2 UnpackHeightSc
 {
 	return PackedValue.r * UnpackHeightScaleBias.x + UnpackHeightScaleBias.y;
 }
+
+
+float3 VirtualTextureUnpackNormalBGR565(in float4 PackedValue)
+{
+	// sign of normal z is considered always positive to save channels
+	return VirtualTextureUnpackNormal(PackedValue.xz, 1.0);
+}
+
+/** Unpack SRGB Color */
+float3 VirtualTextureUnpackBaseColorSRGB(in float4 PackedValue)
+{
+	return sRGBToLinear(PackedValue.rgb);
+}
diff --git a/Engine/Shaders/Private/VirtualTextureCompress.usf b/Engine/Shaders/Private/VirtualTextureCompress.usf
--- a/Engine/Shaders/Private/VirtualTextureCompress.usf
+++ b/Engine/Shaders/Private/VirtualTextureCompress.usf
@@ -135,6 +135,36 @@ void CompressBaseColorNormalSpecularCS(
 
 }
 
+/** Compress base color, normal, roughness to BC1 & BC3. */
+[numthreads(8, 8, 1)]
+void CompressBaseColorNormalRoughnessCS(
+	uint3 ThreadId : SV_DispatchThreadID
+)
+{
+	uint2 SamplePos = ThreadId.xy * 4;
+	if (any(SamplePos >= DestRect.zw))
+		return;
+
+	float2 TexelUVSize = 1.f / float2(DestRect.zw);
+	float2 SampleUV = (float2(SamplePos) + 0.5f) * TexelUVSize;
+
+	float3 BlockBaseColor[16];
+	ReadBlockRGB(RenderTexture0, TextureSampler0, SampleUV, TexelUVSize, BlockBaseColor);
+
+	float3 BlockNormalXRougnnessY[16];
+	const float BlockNormalZ[16]={ 1.0, 1.0, 1.0, 1.0, 1.0, 1.0, 1.0, 1.0, 
+								   1.0, 1.0, 1.0, 1.0, 1.0, 1.0, 1.0, 1.0};
+	ReadBlockRGB(RenderTexture1, TextureSampler2, SampleUV, TexelUVSize, BlockNormalXRougnnessY);
+#if ETC_PROFILE
+	OutCompressTexture0_u2[ThreadId.xy + DestRect.xy] = PackRG32(CompressBlock_ETC2_RGB(BlockBaseColor));
+	OutCompressTexture1[ThreadId.xy + DestRect.xy] = CompressBlock_ETC2_RGBA(BlockNormalXRougnnessY, BlockNormalZ);
+#else
+	OutCompressTexture0_u2[ThreadId.xy + DestRect.xy] = PackRG32(CompressBC1Block_SRGB(BlockBaseColor));
+	OutCompressTexture1[ThreadId.xy + DestRect.xy] = CompressBC3Block(BlockNormalXRougnnessY, BlockNormalZ);
+#endif
+
+}
+
 /** Copy path used when we disable texture compression, because we need to keep the same final channel layout. */
 void CopyBaseColorNormalSpecularPS(
 	in float4 InPosition : SV_POSITION, 
diff --git a/Engine/Shaders/Private/VirtualTextureMaterial.usf b/Engine/Shaders/Private/VirtualTextureMaterial.usf
--- a/Engine/Shaders/Private/VirtualTextureMaterial.usf
+++ b/Engine/Shaders/Private/VirtualTextureMaterial.usf
@@ -3,6 +3,7 @@
 // Material shader for runtime virtual texture
 
 #include "Common.ush"
+#include "GammaCorrectionCommon.ush"
 #include "/Engine/Generated/Material.ush"
 #include "/Engine/Generated/VertexFactory.ush"
 
@@ -91,6 +92,12 @@ void FPixelShaderInOut_MainPS(
 
 #if defined(OUT_BASECOLOR)
 	Out.MRT[0] = float4(BaseColor, 1.f) * Opacity;
+#elif defined(OUT_BASECOLOR_NORMAL_ROUGHNESS)
+	float3 PackedNormal = PackNormal(Normal);
+	//store basecolor in srgb space to improve image quality
+	Out.MRT[0] = float4(LinearToSrgb(BaseColor), 1.0) * Opacity;
+	//not enough channels for sign of normal  z
+	Out.MRT[1] = float4(PackedNormal.x, Roughness, PackedNormal.y, 1.0) * Opacity;
 #elif defined(OUT_BASECOLOR_NORMAL_SPECULAR)
 	float3 PackedNormal = PackNormal(Normal);
 	Out.MRT[0] = float4(BaseColor, 1.f) * Opacity;
diff --git a/Engine/Source/Editor/VirtualTexturingEditor/Private/RuntimeVirtualTextureBuildStreamingMips.cpp b/Engine/Source/Editor/VirtualTexturingEditor/Private/RuntimeVirtualTextureBuildStreamingMips.cpp
--- a/Engine/Source/Editor/VirtualTexturingEditor/Private/RuntimeVirtualTextureBuildStreamingMips.cpp
+++ b/Engine/Source/Editor/VirtualTexturingEditor/Private/RuntimeVirtualTextureBuildStreamingMips.cpp
@@ -30,7 +30,8 @@ namespace
 
 			for (int32 Layer = 0; Layer < NumLayers; ++Layer)
 			{
-				check(InLayerFormats[Layer] == PF_G16 || InLayerFormats[Layer] == PF_B8G8R8A8 || InLayerFormats[Layer] == PF_DXT1 || InLayerFormats[Layer] == PF_DXT5 || InLayerFormats[Layer] == PF_BC5);
+				check(InLayerFormats[Layer] == PF_G16 || InLayerFormats[Layer] == PF_B8G8R8A8 || InLayerFormats[Layer] == PF_DXT1 || InLayerFormats[Layer] == PF_DXT5 || InLayerFormats[Layer] == PF_BC5
+					|| InLayerFormats[Layer] == PF_R5G6B5_UNORM || InLayerFormats[Layer] == PF_B5G5R5A1_UNORM);
 				LayerFormats[Layer] = InLayerFormats[Layer] == PF_G16 ? PF_G16 : PF_B8G8R8A8;
 				LayerOffsets[Layer] = TotalSizeBytes;
 				TotalSizeBytes += CalculateImageBytes(InTileSize, InTileSize, 0, LayerFormats[Layer]) * InNumTilesX * InNumTilesY;
diff --git a/Engine/Source/Runtime/Engine/Classes/VT/RuntimeVirtualTexture.h b/Engine/Source/Runtime/Engine/Classes/VT/RuntimeVirtualTexture.h
--- a/Engine/Source/Runtime/Engine/Classes/VT/RuntimeVirtualTexture.h
+++ b/Engine/Source/Runtime/Engine/Classes/VT/RuntimeVirtualTexture.h
@@ -37,9 +37,15 @@ protected:
 	ERuntimeVirtualTextureMaterialType MaterialType = ERuntimeVirtualTextureMaterialType::BaseColor_Normal_Specular;
 
 	/** Enable storing the virtual texture in GPU supported compression formats. Using uncompressed is only recommended for debugging and quality comparisons. */
-	UPROPERTY(EditAnywhere, BlueprintReadOnly, Category = Layout, meta = (DisplayName = "Enable BC texture compression"))
+	UPROPERTY(EditAnywhere, BlueprintReadOnly, Category = Layout, meta = (DisplayName = "Enable texture compression"))
 	bool bCompressTextures = true;
 
+	/**
+	* Use low quality textures (RGB565/RGB555A1) to replace runtime compression
+	*/
+	UPROPERTY(EditAnywhere, BlueprintReadOnly, Category = Performance, meta = (DisplayName = "Use Low Quality Compression", editcondition = "MaterialType == ERuntimeVirtualTextureMaterialType::BaseColor_Normal_Roughness && bCompressTextures == true"))
+	bool bUseLowQualityCompression = false;
+
 	/** Enable clear before rendering a page of the virtual texture. Disabling this can be an optimization if you know that the texture will always be fully covered by rendering.  */
 	UPROPERTY(EditAnywhere, AdvancedDisplay, Category = Layout, meta = (DisplayName = "Enable clear before render"))
 	bool bClearTextures = true;
@@ -115,6 +121,7 @@ public:
 	/** Public getter for virtual texture removed low mips */
 	int32 GetRemoveLowMips() const { return FMath::Clamp(RemoveLowMips, 0, 5); }
 
+	bool GetLQCompression() const { return bUseLowQualityCompression; }
 	/** Public getter for texture LOD Group */
 	TEnumAsByte<enum TextureGroup> GetLODGroup() const { return LODGroup; }
 
diff --git a/Engine/Source/Runtime/Engine/Private/Components/RuntimeVirtualTextureComponent.cpp b/Engine/Source/Runtime/Engine/Private/Components/RuntimeVirtualTextureComponent.cpp
--- a/Engine/Source/Runtime/Engine/Private/Components/RuntimeVirtualTextureComponent.cpp
+++ b/Engine/Source/Runtime/Engine/Private/Components/RuntimeVirtualTextureComponent.cpp
@@ -135,6 +135,7 @@ uint64 URuntimeVirtualTextureComponent::CalculateStreamingTextureSettingsHash()
 			uint32 SinglePhysicalSpace : 1;
 			uint32 EnableCompressCrunch : 1;
 			uint32 ContinuousUpdate : 1;
+			uint32 bUseLowQualityCompression : 1;
 		};
 	};
 
@@ -148,6 +149,7 @@ uint64 URuntimeVirtualTextureComponent::CalculateStreamingTextureSettingsHash()
 	Settings.ContinuousUpdate = (uint32)VirtualTexture->GetContinuousUpdate();
 	Settings.SinglePhysicalSpace = (uint32)VirtualTexture->GetSinglePhysicalSpace();
 	Settings.EnableCompressCrunch = (uint32)bEnableCompressCrunch;
+	Settings.bUseLowQualityCompression = (uint32)VirtualTexture->GetLQCompression();
 
 	return Settings.PackedValue;
 }
@@ -204,6 +206,15 @@ private:
 	TArray<URuntimeVirtualTexture*> VirtualTextures;
 };
 
+static void GetLayerFormatSettings(FTextureFormatSettings& OutFormatSettings, EPixelFormat LayerFormat, bool IsLayerYCoCg, bool IsLayerSRGB, bool IsLayerLQCompression)
+{
+	OutFormatSettings.CompressionSettings = IsLayerLQCompression? TC_LQ : (LayerFormat == PF_BC5 ? TC_Normalmap : TC_Default);
+	OutFormatSettings.CompressionNone = LayerFormat == PF_B8G8R8A8 || LayerFormat == PF_G16;
+	OutFormatSettings.CompressionNoAlpha = LayerFormat == PF_DXT1 || LayerFormat == PF_BC5 || LayerFormat == PF_R5G6B5_UNORM;
+	OutFormatSettings.CompressionForceAlpha = LayerFormat == PF_DXT5;
+	OutFormatSettings.CompressionYCoCg = IsLayerYCoCg;
+	OutFormatSettings.SRGB = IsLayerSRGB;
+}
 void URuntimeVirtualTextureComponent::InitializeStreamingTexture(uint32 InSizeX, uint32 InSizeY, uint8* InData)
 {
 	// We need an existing StreamingTexture object to update.
@@ -225,17 +236,12 @@ void URuntimeVirtualTextureComponent::InitializeStreamingTexture(uint32 InSizeX,
 		BuildDesc.LayerFormats.AddDefaulted(BuildDesc.LayerCount);
 		BuildDesc.LayerFormatSettings.AddDefaulted(BuildDesc.LayerCount);
 
-		for (int32 Layer = 0; Layer < BuildDesc.LayerCount; Layer++)
+		for (int32 Layer = 0; Layer < BuildDesc.LayerCount; Layer++) 
 		{
 			const EPixelFormat LayerFormat = VirtualTexture->GetLayerFormat(Layer);
 			BuildDesc.LayerFormats[Layer] = LayerFormat == PF_G16 ? TSF_G16 : TSF_BGRA8;
-
-			BuildDesc.LayerFormatSettings[Layer].CompressionSettings = LayerFormat == PF_BC5 ? TC_Normalmap : TC_Default;
-			BuildDesc.LayerFormatSettings[Layer].CompressionNone = LayerFormat == PF_B8G8R8A8 || LayerFormat == PF_G16;
-			BuildDesc.LayerFormatSettings[Layer].CompressionNoAlpha = LayerFormat == PF_DXT1 || LayerFormat == PF_BC5;
-			BuildDesc.LayerFormatSettings[Layer].CompressionForceAlpha = LayerFormat == PF_DXT5;
-			BuildDesc.LayerFormatSettings[Layer].CompressionYCoCg = VirtualTexture->IsLayerYCoCg(Layer);
-			BuildDesc.LayerFormatSettings[Layer].SRGB = VirtualTexture->IsLayerSRGB(Layer);
+			bool IsLayerLQCompression = (VirtualTexture->GetMaterialType() == ERuntimeVirtualTextureMaterialType::BaseColor_Normal_Roughness && VirtualTexture->GetLQCompression() && LayerFormat != PF_B8G8R8A8);
+			GetLayerFormatSettings(BuildDesc.LayerFormatSettings[Layer], LayerFormat, VirtualTexture->IsLayerYCoCg(Layer), VirtualTexture->IsLayerSRGB(Layer), IsLayerLQCompression);
 		}
 
 		BuildDesc.BuildHash = CalculateStreamingTextureSettingsHash();
diff --git a/Engine/Source/Runtime/Engine/Private/Materials/HLSLMaterialTranslator.cpp b/Engine/Source/Runtime/Engine/Private/Materials/HLSLMaterialTranslator.cpp
--- a/Engine/Source/Runtime/Engine/Private/Materials/HLSLMaterialTranslator.cpp
+++ b/Engine/Source/Runtime/Engine/Private/Materials/HLSLMaterialTranslator.cpp
@@ -6422,7 +6422,16 @@ int32 FHLSLMaterialTranslator::VirtualTextureUnpack(int32 CodeIndex0, int32 Code
 		FString	SampleCode(TEXT("VirtualTextureUnpackHeight(%s, %s)"));
 		return CodeIndex0 == INDEX_NONE ? INDEX_NONE : AddCodeChunk(MCT_Float, *SampleCode, *GetParameterCode(CodeIndex0), *GetParameterCode(P0));
 	}
-
+	else if (UnpackType == EVirtualTextureUnpackType::NormalBGR565)
+	{
+		FString	SampleCode(TEXT("VirtualTextureUnpackNormalBGR565(%s)"));
+		return CodeIndex0 == INDEX_NONE || CodeIndex1 == INDEX_NONE ? INDEX_NONE : AddCodeChunk(MCT_Float3, *SampleCode, *GetParameterCode(CodeIndex1));
+	}
+	if (UnpackType == EVirtualTextureUnpackType::BaseColorSRGB)
+	{
+		FString	SampleCode(TEXT("VirtualTextureUnpackBaseColorSRGB(%s)"));
+		return CodeIndex0 == INDEX_NONE ? INDEX_NONE : AddCodeChunk(MCT_Float3, *SampleCode, *GetParameterCode(CodeIndex0));
+	}
 	return CodeIndex0;
 }
 
diff --git a/Engine/Source/Runtime/Engine/Private/Materials/MaterialExpressions.cpp b/Engine/Source/Runtime/Engine/Private/Materials/MaterialExpressions.cpp
--- a/Engine/Source/Runtime/Engine/Private/Materials/MaterialExpressions.cpp
+++ b/Engine/Source/Runtime/Engine/Private/Materials/MaterialExpressions.cpp
@@ -2398,6 +2398,7 @@ int32 UMaterialExpressionRuntimeVirtualTextureSample::Compile(class FMaterialCom
 
 	bool bIsBaseColorValid = false;
 	bool bIsSpecularValid = false;
+	bool bIsRoughnessValid = false;
 	bool bIsNormalValid = false;
 	bool bIsWorldHeightValid = false;
 	bool bIsMaskValid = false;
@@ -2405,9 +2406,10 @@ int32 UMaterialExpressionRuntimeVirtualTextureSample::Compile(class FMaterialCom
 	switch (MaterialType)
 	{
 	case ERuntimeVirtualTextureMaterialType::BaseColor: bIsBaseColorValid = true; break;
-	case ERuntimeVirtualTextureMaterialType::BaseColor_Normal_Specular: bIsBaseColorValid = bIsNormalValid = bIsSpecularValid = true; break;
-	case ERuntimeVirtualTextureMaterialType::BaseColor_Normal_Specular_YCoCg: bIsBaseColorValid = bIsNormalValid = bIsSpecularValid = true; break;
-	case ERuntimeVirtualTextureMaterialType::BaseColor_Normal_Specular_Mask_YCoCg: bIsBaseColorValid = bIsNormalValid = bIsSpecularValid = bIsMaskValid = true; break;
+	case ERuntimeVirtualTextureMaterialType::BaseColor_Normal_Roughness: bIsBaseColorValid = bIsNormalValid = bIsRoughnessValid = true; break;
+	case ERuntimeVirtualTextureMaterialType::BaseColor_Normal_Specular: bIsRoughnessValid = bIsBaseColorValid = bIsNormalValid = bIsSpecularValid = true; break;
+	case ERuntimeVirtualTextureMaterialType::BaseColor_Normal_Specular_YCoCg: bIsRoughnessValid = bIsBaseColorValid = bIsNormalValid = bIsSpecularValid = true; break;
+	case ERuntimeVirtualTextureMaterialType::BaseColor_Normal_Specular_Mask_YCoCg: bIsRoughnessValid = bIsBaseColorValid = bIsNormalValid = bIsSpecularValid = bIsMaskValid = true; break;
 	case ERuntimeVirtualTextureMaterialType::WorldHeight: bIsWorldHeightValid = true; break;
 	}
 
@@ -2420,6 +2422,7 @@ int32 UMaterialExpressionRuntimeVirtualTextureSample::Compile(class FMaterialCom
 			{
 			case ERuntimeVirtualTextureMaterialType::BaseColor_Normal_Specular_YCoCg:
 			case ERuntimeVirtualTextureMaterialType::BaseColor_Normal_Specular_Mask_YCoCg: UnpackType = EVirtualTextureUnpackType::BaseColorYCoCg; break;
+			case ERuntimeVirtualTextureMaterialType::BaseColor_Normal_Roughness: UnpackType = EVirtualTextureUnpackType::BaseColorSRGB; break;
 			default: UnpackTarget = 0; UnpackMask = 0x7; break;
 			}
 		}
@@ -2444,10 +2447,11 @@ int32 UMaterialExpressionRuntimeVirtualTextureSample::Compile(class FMaterialCom
 		}
 		break;
 	case 2:
-		if ((bIsParameter || bIsVirtualTextureValid) && bIsSpecularValid)
+		if ((bIsParameter || bIsVirtualTextureValid) && bIsRoughnessValid)
 		{
 			switch (MaterialType)
 			{
+			case ERuntimeVirtualTextureMaterialType::BaseColor_Normal_Roughness: UnpackTarget = 1; UnpackMask = 0x2; break;
 			case ERuntimeVirtualTextureMaterialType::BaseColor_Normal_Specular: UnpackTarget = 1; UnpackMask = 0x2; break;
 			case ERuntimeVirtualTextureMaterialType::BaseColor_Normal_Specular_YCoCg:
 			case ERuntimeVirtualTextureMaterialType::BaseColor_Normal_Specular_Mask_YCoCg: UnpackTarget = 2; UnpackMask = 0x2; break;
@@ -2463,6 +2467,7 @@ int32 UMaterialExpressionRuntimeVirtualTextureSample::Compile(class FMaterialCom
 		{
 			switch (MaterialType)
 			{
+			case ERuntimeVirtualTextureMaterialType::BaseColor_Normal_Roughness: UnpackType = EVirtualTextureUnpackType::NormalBGR565; break;
 			case ERuntimeVirtualTextureMaterialType::BaseColor_Normal_Specular: UnpackType = EVirtualTextureUnpackType::NormalBC3BC3; break;
 			case ERuntimeVirtualTextureMaterialType::BaseColor_Normal_Specular_YCoCg:
 			case ERuntimeVirtualTextureMaterialType::BaseColor_Normal_Specular_Mask_YCoCg: UnpackType = EVirtualTextureUnpackType::NormalBC5BC1; break;
diff --git a/Engine/Source/Runtime/Engine/Private/ShaderCompiler/ShaderGenerationUtil.cpp b/Engine/Source/Runtime/Engine/Private/ShaderCompiler/ShaderGenerationUtil.cpp
--- a/Engine/Source/Runtime/Engine/Private/ShaderCompiler/ShaderGenerationUtil.cpp
+++ b/Engine/Source/Runtime/Engine/Private/ShaderCompiler/ShaderGenerationUtil.cpp
@@ -208,6 +208,7 @@ void FShaderCompileUtilities::ApplyFetchEnvironment(FShaderMaterialPropertyDefin
 	FETCH_COMPILE_BOOL(MATERIALDOMAIN_UI);
 
 	FETCH_COMPILE_BOOL(OUT_BASECOLOR);
+	FETCH_COMPILE_BOOL(OUT_BASECOLOR_NORMAL_ROUGHNESS);
 	FETCH_COMPILE_BOOL(OUT_BASECOLOR_NORMAL_SPECULAR);
 	FETCH_COMPILE_BOOL(OUT_WORLDHEIGHT);
 
@@ -1845,6 +1846,11 @@ void FShaderCompileUtilities::ApplyDerivedDefines(FShaderCompilerEnvironment& Ou
 		{
 			bTargetUsage[0] = 1;
 		}
+		else if (MaterialDefines.OUT_BASECOLOR_NORMAL_ROUGHNESS)
+		{
+			bTargetUsage[0] = 1;
+			bTargetUsage[1] = 1;
+		}
 		else if (MaterialDefines.OUT_BASECOLOR_NORMAL_SPECULAR)
 		{
 			bTargetUsage[0] = 1;
diff --git a/Engine/Source/Runtime/Engine/Private/VT/RuntimeVirtualTexture.cpp b/Engine/Source/Runtime/Engine/Private/VT/RuntimeVirtualTexture.cpp
--- a/Engine/Source/Runtime/Engine/Private/VT/RuntimeVirtualTexture.cpp
+++ b/Engine/Source/Runtime/Engine/Private/VT/RuntimeVirtualTexture.cpp
@@ -346,6 +346,7 @@ int32 URuntimeVirtualTexture::GetLayerCount(ERuntimeVirtualTextureMaterialType I
 	case ERuntimeVirtualTextureMaterialType::WorldHeight:
 		return 1;
 	case ERuntimeVirtualTextureMaterialType::BaseColor_Normal_Specular:
+	case ERuntimeVirtualTextureMaterialType::BaseColor_Normal_Roughness:
 		return 2;
 	case ERuntimeVirtualTextureMaterialType::BaseColor_Normal_Specular_YCoCg:
 	case ERuntimeVirtualTextureMaterialType::BaseColor_Normal_Specular_Mask_YCoCg:
@@ -391,6 +392,14 @@ static EPixelFormat PlatformCompressedRVTFormat(EPixelFormat Format)
 	return Format;
 }
 
+static EPixelFormat PlatformLQCompressedFormat(bool bRequireAlpha)
+{
+	const EPixelFormat LQFormat = bRequireAlpha ? PF_B5G5R5A1_UNORM : PF_R5G6B5_UNORM;
+	const EPixelFormat HQFormat = bRequireAlpha ? PF_DXT5 : PF_DXT1;
+	bool bLQFormatSupported = GPixelFormats[PF_B5G5R5A1_UNORM].Supported && GPixelFormats[PF_R5G6B5_UNORM].Supported;
+	return bLQFormatSupported? LQFormat : HQFormat;
+}
+
 EPixelFormat URuntimeVirtualTexture::GetLayerFormat(int32 LayerIndex) const
 {
 	if (LayerIndex == 0)
@@ -399,6 +408,8 @@ EPixelFormat URuntimeVirtualTexture::GetLayerFormat(int32 LayerIndex) const
 		{
 		case ERuntimeVirtualTextureMaterialType::BaseColor:
 			return bCompressTextures ? PlatformCompressedRVTFormat(PF_DXT1) : PF_B8G8R8A8;
+		case ERuntimeVirtualTextureMaterialType::BaseColor_Normal_Roughness:
+			return bCompressTextures ? (bUseLowQualityCompression? PlatformLQCompressedFormat(false) : PlatformCompressedRVTFormat(PF_DXT1)) : PF_B8G8R8A8;
 		case ERuntimeVirtualTextureMaterialType::BaseColor_Normal_Specular:
 		case ERuntimeVirtualTextureMaterialType::BaseColor_Normal_Specular_YCoCg:
 		case ERuntimeVirtualTextureMaterialType::BaseColor_Normal_Specular_Mask_YCoCg:
@@ -413,6 +424,8 @@ EPixelFormat URuntimeVirtualTexture::GetLayerFormat(int32 LayerIndex) const
 	{
 		switch (MaterialType)
 		{
+		case ERuntimeVirtualTextureMaterialType::BaseColor_Normal_Roughness:
+			return bCompressTextures ? (bUseLowQualityCompression ? PlatformLQCompressedFormat(false) : PlatformCompressedRVTFormat(PF_DXT5)) : PF_B8G8R8A8;
 		case ERuntimeVirtualTextureMaterialType::BaseColor_Normal_Specular:
 			return bCompressTextures ? PlatformCompressedRVTFormat(PF_DXT5) : PF_B8G8R8A8;
 		case ERuntimeVirtualTextureMaterialType::BaseColor_Normal_Specular_YCoCg:
@@ -451,6 +464,7 @@ bool URuntimeVirtualTexture::IsLayerSRGB(int32 LayerIndex) const
 	case ERuntimeVirtualTextureMaterialType::BaseColor_Normal_Specular_YCoCg:
 	case ERuntimeVirtualTextureMaterialType::BaseColor_Normal_Specular_Mask_YCoCg:
 	case ERuntimeVirtualTextureMaterialType::WorldHeight:
+	case ERuntimeVirtualTextureMaterialType::BaseColor_Normal_Roughness:
 		return false;
 	default:
 		break;
diff --git a/Engine/Source/Runtime/Engine/Public/MaterialCompiler.h b/Engine/Source/Runtime/Engine/Public/MaterialCompiler.h
--- a/Engine/Source/Runtime/Engine/Public/MaterialCompiler.h
+++ b/Engine/Source/Runtime/Engine/Public/MaterialCompiler.h
@@ -38,6 +38,8 @@ enum class EVirtualTextureUnpackType
 	NormalBC3BC3,
 	NormalBC5BC1,
 	HeightR16,
+	NormalBGR565,
+	BaseColorSRGB
 };
 
 /** What type of compiler is this? Used by material expressions that select input based on compile context */
diff --git a/Engine/Source/Runtime/Engine/Public/VT/RuntimeVirtualTextureEnum.h b/Engine/Source/Runtime/Engine/Public/VT/RuntimeVirtualTextureEnum.h
--- a/Engine/Source/Runtime/Engine/Public/VT/RuntimeVirtualTextureEnum.h
+++ b/Engine/Source/Runtime/Engine/Public/VT/RuntimeVirtualTextureEnum.h
@@ -36,11 +36,11 @@ enum class ERuntimeVirtualTextureMaterialType : uint8
 {
 	BaseColor UMETA(DisplayName = "Base Color"),
 	BaseColor_Normal_DEPRECATED UMETA(Hidden),
+	BaseColor_Normal_Roughness UMETA(DisplayName = "Base Color, Normal, Roughness", ToolTip = "Local space Normal. Requires less memory than 'Base Color, Normal, Roughness, Specular'. Supports LQ compression."),
 	BaseColor_Normal_Specular UMETA(DisplayName = "Base Color, Normal, Roughness, Specular"),
 	BaseColor_Normal_Specular_YCoCg UMETA(DisplayName = "YCoCg Base Color, Normal, Roughness, Specular", ToolTip = "Base Color is stored in YCoCg space. This requires more memory but may provide better quality."),
 	BaseColor_Normal_Specular_Mask_YCoCg UMETA(DisplayName = "YCoCg Base Color, Normal, Roughness, Specular, Mask", ToolTip="Base Color is stored in YCoCg space. This requires more memory but may provide better quality."),
 	WorldHeight UMETA(DisplayName = "World Height"),
-
 	Count UMETA(Hidden),
 };
 
diff --git a/Engine/Source/Runtime/Landscape/Private/LandscapeRender.cpp b/Engine/Source/Runtime/Landscape/Private/LandscapeRender.cpp
--- a/Engine/Source/Runtime/Landscape/Private/LandscapeRender.cpp
+++ b/Engine/Source/Runtime/Landscape/Private/LandscapeRender.cpp
@@ -3549,10 +3549,12 @@ public:
 			FName(TEXT("TVirtualTextureVSBaseColor")),
 			FName(TEXT("TVirtualTextureVSBaseColorNormal")),
 			FName(TEXT("TVirtualTextureVSBaseColorNormalSpecular")),
+			FName(TEXT("TVirtualTextureVSBaseColorNormalRoughness")),
 			FName(TEXT("TVirtualTextureVSWorldHeight")),
 			FName(TEXT("TVirtualTexturePSBaseColor")),
 			FName(TEXT("TVirtualTexturePSBaseColorNormal")),
 			FName(TEXT("TVirtualTexturePSBaseColorNormalSpecular")),
+			FName(TEXT("TVirtualTexturePSBaseColorNormalRoughness")),
 			FName(TEXT("TVirtualTexturePSWorldHeight")),
 		};
 		return AllowedShaderTypes;
@@ -3772,10 +3774,12 @@ public:
 			FName(TEXT("TVirtualTextureVSBaseColor")),
 			FName(TEXT("TVirtualTextureVSBaseColorNormal")),
 			FName(TEXT("TVirtualTextureVSBaseColorNormalSpecular")),
+			FName(TEXT("TVirtualTextureVSBaseColorNormalRoughness")),
 			FName(TEXT("TVirtualTextureVSWorldHeight")),
 			FName(TEXT("TVirtualTexturePSBaseColor")),
 			FName(TEXT("TVirtualTexturePSBaseColorNormal")),
 			FName(TEXT("TVirtualTexturePSBaseColorNormalSpecular")),
+			FName(TEXT("TVirtualTexturePSBaseColorNormalRoughness")),
 			FName(TEXT("TVirtualTexturePSWorldHeight")),
 		};
 		return ShaderTypes;
diff --git a/Engine/Source/Runtime/RenderCore/Private/ShaderMaterialDerivedHelpers.cpp b/Engine/Source/Runtime/RenderCore/Private/ShaderMaterialDerivedHelpers.cpp
--- a/Engine/Source/Runtime/RenderCore/Private/ShaderMaterialDerivedHelpers.cpp
+++ b/Engine/Source/Runtime/RenderCore/Private/ShaderMaterialDerivedHelpers.cpp
@@ -109,6 +109,11 @@ FShaderMaterialDerivedDefines RENDERCORE_API CalculateDerivedMaterialParameters(
 			{
 				Dst.PIXELSHADEROUTPUT_MRT0 = 1;
 			}
+			else if (Mat.OUT_BASECOLOR_NORMAL_ROUGHNESS)
+			{
+				Dst.PIXELSHADEROUTPUT_MRT0 = 1;
+				Dst.PIXELSHADEROUTPUT_MRT1 = 1;
+			}
 			else if (Mat.OUT_BASECOLOR_NORMAL_SPECULAR)
 			{
 				Dst.PIXELSHADEROUTPUT_MRT0 = 1;
diff --git a/Engine/Source/Runtime/RenderCore/Public/ShaderMaterial.h b/Engine/Source/Runtime/RenderCore/Public/ShaderMaterial.h
--- a/Engine/Source/Runtime/RenderCore/Public/ShaderMaterial.h
+++ b/Engine/Source/Runtime/RenderCore/Public/ShaderMaterial.h
@@ -164,6 +164,7 @@ struct FShaderMaterialPropertyDefines
 	uint8 MATERIALDOMAIN_UI : 1;
 
 	uint8 OUT_BASECOLOR : 1;
+	uint8 OUT_BASECOLOR_NORMAL_ROUGHNESS : 1;
 	uint8 OUT_BASECOLOR_NORMAL_SPECULAR : 1;
 	uint8 OUT_WORLDHEIGHT : 1;
 
diff --git a/Engine/Source/Runtime/Renderer/Private/VT/RuntimeVirtualTextureProducer.cpp b/Engine/Source/Runtime/Renderer/Private/VT/RuntimeVirtualTextureProducer.cpp
--- a/Engine/Source/Runtime/Renderer/Private/VT/RuntimeVirtualTextureProducer.cpp
+++ b/Engine/Source/Runtime/Renderer/Private/VT/RuntimeVirtualTextureProducer.cpp
@@ -81,6 +81,7 @@ void FRuntimeVirtualTextureFinalizer::Finalize(FRDGBuilder& GraphBuilder)
 	{
 		RenderPageBatchDesc.Targets[LayerIndex].Texture = Tiles[0].Targets[LayerIndex].TextureRHI != nullptr ? Tiles[0].Targets[LayerIndex].TextureRHI->GetTexture2D() : nullptr;
 		RenderPageBatchDesc.Targets[LayerIndex].UAV = Tiles[0].Targets[LayerIndex].UnorderedAccessViewRHI;
+		RenderPageBatchDesc.Targets[LayerIndex].PooledRenderTarget = Tiles[0].Targets[LayerIndex].PooledRenderTarget;
 	}
 
 	int32 BatchSize = 0;
@@ -128,6 +129,7 @@ void FRuntimeVirtualTextureFinalizer::Finalize(FRDGBuilder& GraphBuilder)
 			{
 				RenderPageBatchDesc.Targets[LayerIndex].Texture = Tiles[0].Targets[LayerIndex].TextureRHI != nullptr ? Tiles[0].Targets[LayerIndex].TextureRHI->GetTexture2D() : nullptr;
 				RenderPageBatchDesc.Targets[LayerIndex].UAV = Tiles[0].Targets[LayerIndex].UnorderedAccessViewRHI;
+				RenderPageBatchDesc.Targets[LayerIndex].PooledRenderTarget = Tiles[0].Targets[LayerIndex].PooledRenderTarget;
 			}
 		}
 	}
diff --git a/Engine/Source/Runtime/Renderer/Private/VT/RuntimeVirtualTextureRender.cpp b/Engine/Source/Runtime/Renderer/Private/VT/RuntimeVirtualTextureRender.cpp
--- a/Engine/Source/Runtime/Renderer/Private/VT/RuntimeVirtualTextureRender.cpp
+++ b/Engine/Source/Runtime/Renderer/Private/VT/RuntimeVirtualTextureRender.cpp
@@ -239,6 +239,90 @@ namespace RuntimeVirtualTexture
 		}
 	};
 
+	/** Specialization for ERuntimeVirtualTextureMaterialType::BaseColor_Normal_Roughness */
+	class FMaterialPolicy_BaseColorNormalRoughness
+	{
+	private:
+		/** Compile time helper to build blend state from the connected output attribute mask. */
+		static constexpr EColorWriteMask GetColorMaskFromAttributeMask(uint8 AttributeMask, uint8 RenderTargetIndex)
+		{
+			// Color mask in the output render targets for each of the relevant attributes in ERuntimeVirtualTextureAttributeType
+			const EColorWriteMask AttributeMasks[][2] = {
+				{ CW_RGBA, CW_NONE}, // BaseColor
+				{ CW_NONE, EColorWriteMask(CW_RED| CW_BLUE | CW_ALPHA)}, // Normal
+				{ CW_NONE, EColorWriteMask(CW_GREEN | CW_ALPHA)}, // Roughness
+				{ CW_NONE, CW_NONE}, // Specular
+				{ CW_NONE, CW_NONE}, // Mask
+			};
+
+			// Combine the color masks for this AttributeMask
+			EColorWriteMask ColorWriteMask = CW_NONE;
+			for (int32 i = 0; i < 5; ++i)
+			{
+				if (AttributeMask & (1 << i))
+				{
+					ColorWriteMask = EColorWriteMask(ColorWriteMask | AttributeMasks[i][RenderTargetIndex]);
+				}
+			}
+			return ColorWriteMask;
+		}
+		/** Helper to convert the connected output attribute mask to a blend state with a color mask for these attributes. */
+		template< uint32 AttributeMask >
+		static FRHIBlendState* TGetBlendStateFromAttributeMask()
+		{
+			return TStaticBlendState<
+				GetColorMaskFromAttributeMask(AttributeMask, 0), BO_Add, BF_One, BF_InverseSourceAlpha, BO_Add, BF_Zero, BF_One,
+				// normal XY is stored in R and B channels, and the Sign of Z is considered always positive
+				GetColorMaskFromAttributeMask(AttributeMask, 1), BO_Add, BF_One, BF_InverseSourceAlpha, BO_Add, BF_Zero, BF_One>::GetRHI();
+		}
+		/** Runtime conversion of attribute mask to static blend state. */
+		static FRHIBlendState* GetBlendStateImpl(uint8 AttributeMask)
+		{
+			// We have 5 relevant bits in the attribute mask. Any more and this would get painful...
+			switch (AttributeMask & 0x1f)
+			{
+			case 1: return TGetBlendStateFromAttributeMask<1>();
+			case 2: return TGetBlendStateFromAttributeMask<2>();
+			case 3: return TGetBlendStateFromAttributeMask<3>();
+			case 4: return TGetBlendStateFromAttributeMask<4>();
+			case 5: return TGetBlendStateFromAttributeMask<5>();
+			case 6: return TGetBlendStateFromAttributeMask<6>();
+			case 7: return TGetBlendStateFromAttributeMask<7>();
+			case 8: return TGetBlendStateFromAttributeMask<8>();
+			case 9: return TGetBlendStateFromAttributeMask<9>();
+			case 10: return TGetBlendStateFromAttributeMask<10>();
+			case 11: return TGetBlendStateFromAttributeMask<11>();
+			case 12: return TGetBlendStateFromAttributeMask<12>();
+			case 13: return TGetBlendStateFromAttributeMask<13>();
+			case 14: return TGetBlendStateFromAttributeMask<14>();
+			case 15: return TGetBlendStateFromAttributeMask<15>();
+			case 16: return TGetBlendStateFromAttributeMask<16>();
+			case 17: return TGetBlendStateFromAttributeMask<17>();
+			case 18: return TGetBlendStateFromAttributeMask<18>();
+			case 19: return TGetBlendStateFromAttributeMask<19>();
+			case 21: return TGetBlendStateFromAttributeMask<21>();
+			case 22: return TGetBlendStateFromAttributeMask<22>();
+			case 23: return TGetBlendStateFromAttributeMask<23>();
+			case 24: return TGetBlendStateFromAttributeMask<24>();
+			case 25: return TGetBlendStateFromAttributeMask<25>();
+			case 26: return TGetBlendStateFromAttributeMask<26>();
+			case 27: return TGetBlendStateFromAttributeMask<27>();
+			case 28: return TGetBlendStateFromAttributeMask<28>();
+			case 29: return TGetBlendStateFromAttributeMask<29>();
+			case 30: return TGetBlendStateFromAttributeMask<30>();
+			default: return TGetBlendStateFromAttributeMask<31>();
+			}
+		}
+	public:
+		static void ModifyCompilationEnvironment(FShaderCompilerEnvironment& OutEnvironment)
+		{
+			OutEnvironment.SetDefine(TEXT("OUT_BASECOLOR_NORMAL_ROUGHNESS"), 1);
+		}
+		static FRHIBlendState* GetBlendState(uint8 OutputAttributeMask)
+		{
+			return GetBlendStateImpl(OutputAttributeMask);
+		}
+	};
 	/** Specialization for ERuntimeVirtualTextureMaterialType::WorldHeight */
 	class FMaterialPolicy_WorldHeight
 	{
@@ -309,7 +393,7 @@ namespace RuntimeVirtualTexture
 	IMPLEMENT_VIRTUALTEXTURE_SHADER_TYPE(FMaterialPolicy_BaseColor, BaseColor);
 	IMPLEMENT_VIRTUALTEXTURE_SHADER_TYPE(FMaterialPolicy_BaseColorNormalSpecular, BaseColorNormalSpecular);
 	IMPLEMENT_VIRTUALTEXTURE_SHADER_TYPE(FMaterialPolicy_WorldHeight, WorldHeight);
-
+	IMPLEMENT_VIRTUALTEXTURE_SHADER_TYPE(FMaterialPolicy_BaseColorNormalRoughness, BaseColorNormalRoughness);
 	/** Mesh processor for rendering static meshes to the virtual texture */
 	class FRuntimeVirtualTextureMeshProcessor : public FMeshPassProcessor
 	{
@@ -338,6 +422,9 @@ namespace RuntimeVirtualTexture
 				case ERuntimeVirtualTextureMaterialType::BaseColor:
 					return Process<FMaterialPolicy_BaseColor>(MeshBatch, BatchElementMask, StaticMeshId, OutputAttributeMask, PrimitiveSceneProxy, *MaterialRenderProxy, *Material);
 					break;
+				case ERuntimeVirtualTextureMaterialType::BaseColor_Normal_Roughness:
+					return Process<FMaterialPolicy_BaseColorNormalRoughness>(MeshBatch, BatchElementMask, StaticMeshId, OutputAttributeMask, PrimitiveSceneProxy, *MaterialRenderProxy, *Material);
+					break;
 				case ERuntimeVirtualTextureMaterialType::BaseColor_Normal_Specular:
 				case ERuntimeVirtualTextureMaterialType::BaseColor_Normal_Specular_YCoCg:
 				case ERuntimeVirtualTextureMaterialType::BaseColor_Normal_Specular_Mask_YCoCg:
@@ -618,6 +705,7 @@ namespace RuntimeVirtualTexture
 
 	IMPLEMENT_SHADER_TYPE(template<>, FShader_VirtualTextureCompress_CS< ERuntimeVirtualTextureMaterialType::BaseColor >, TEXT("/Engine/Private/VirtualTextureCompress.usf"), TEXT("CompressBaseColorCS"), SF_Compute);
 	IMPLEMENT_SHADER_TYPE(template<>, FShader_VirtualTextureCompress_CS< ERuntimeVirtualTextureMaterialType::BaseColor_Normal_Specular >, TEXT("/Engine/Private/VirtualTextureCompress.usf"), TEXT("CompressBaseColorNormalSpecularCS"), SF_Compute);
+	IMPLEMENT_SHADER_TYPE(template<>, FShader_VirtualTextureCompress_CS< ERuntimeVirtualTextureMaterialType::BaseColor_Normal_Roughness >, TEXT("/Engine/Private/VirtualTextureCompress.usf"), TEXT("CompressBaseColorNormalRoughnessCS"), SF_Compute);
 	IMPLEMENT_SHADER_TYPE(template<>, FShader_VirtualTextureCompress_CS< ERuntimeVirtualTextureMaterialType::BaseColor_Normal_Specular_YCoCg >, TEXT("/Engine/Private/VirtualTextureCompress.usf"), TEXT("CompressBaseColorNormalSpecularYCoCgCS"), SF_Compute);
 	IMPLEMENT_SHADER_TYPE(template<>, FShader_VirtualTextureCompress_CS< ERuntimeVirtualTextureMaterialType::BaseColor_Normal_Specular_Mask_YCoCg >, TEXT("/Engine/Private/VirtualTextureCompress.usf"), TEXT("CompressBaseColorNormalSpecularMaskYCoCgCS"), SF_Compute);
 
@@ -649,6 +737,9 @@ namespace RuntimeVirtualTexture
 		case ERuntimeVirtualTextureMaterialType::BaseColor_Normal_Specular:
 			AddCompressPass<ERuntimeVirtualTextureMaterialType::BaseColor_Normal_Specular>(GraphBuilder, FeatureLevel, Parameters, GroupCount);
 			break;
+		case ERuntimeVirtualTextureMaterialType::BaseColor_Normal_Roughness:
+			AddCompressPass<ERuntimeVirtualTextureMaterialType::BaseColor_Normal_Roughness>(GraphBuilder, FeatureLevel, Parameters, GroupCount);
+			break;
 		case ERuntimeVirtualTextureMaterialType::BaseColor_Normal_Specular_YCoCg:
 			AddCompressPass<ERuntimeVirtualTextureMaterialType::BaseColor_Normal_Specular_YCoCg>(GraphBuilder, FeatureLevel, Parameters, GroupCount);
 			break;
@@ -829,13 +920,14 @@ namespace RuntimeVirtualTexture
 			bCopyThumbnailPass = bRenderPass && bIsThumbnails;
 			EPixelFormat OutputFormat = (OutputTexture0 != nullptr ? OutputTexture0->GetFormat() : PF_Unknown);
 			bool bCompressedFormat = GPixelFormats[OutputFormat].BlockSizeX == 4 && GPixelFormats[OutputFormat].BlockSizeY == 4;
+			bool bLQFormat = (OutputTexture0->GetFormat() == PF_R5G6B5_UNORM);
 			bCompressPass = bRenderPass && !bCopyThumbnailPass && bCompressedFormat;
 			bCopyPass = bRenderPass && !bCopyThumbnailPass && !bCompressPass && (MaterialType == ERuntimeVirtualTextureMaterialType::BaseColor_Normal_Specular || MaterialType == ERuntimeVirtualTextureMaterialType::BaseColor_Normal_Specular_YCoCg || MaterialType == ERuntimeVirtualTextureMaterialType::BaseColor_Normal_Specular_Mask_YCoCg);
 			// Not all mobile RHIs support sRGB texture views/aliasing, use only linear targets on mobile
 			ETextureCreateFlags VT_SRGB = FeatureLevel > ERHIFeatureLevel::ES3_1 ? TexCreate_SRGB : TexCreate_None;
 			// GLES does not support writing to 2-channel images, instead pack RG32 into RGBA16 on all GL
 			EPixelFormat Compressed64BitFormat = IsOpenGLPlatform(GMaxRHIShaderPlatform) ? PF_R16G16B16A16_UINT : PF_R32G32_UINT;
-		
+
 			switch (MaterialType)
 			{
 			case ERuntimeVirtualTextureMaterialType::BaseColor:
@@ -852,6 +944,23 @@ namespace RuntimeVirtualTexture
 					OutputAlias0 = CopyTexture0 = GraphBuilder.CreateTexture(FRDGTextureDesc::Create2D(TextureSize, PF_B8G8R8A8, FClearValueBinding::Black, VT_SRGB | TexCreate_RenderTargetable | TexCreate_ShaderResource), TEXT("CopyTexture0"));
 				}
 				break;
+			case ERuntimeVirtualTextureMaterialType::BaseColor_Normal_Roughness:
+				if (bRenderPass)
+				{
+					OutputAlias0 = RenderTexture0 = GraphBuilder.CreateTexture(FRDGTextureDesc::Create2D(TextureSize, bLQFormat ? PF_R5G6B5_UNORM : PF_B8G8R8A8, FClearValueBinding::Black, TexCreate_RenderTargetable | TexCreate_ShaderResource), TEXT("RenderTexture0"));
+					OutputAlias1 = RenderTexture1 = GraphBuilder.CreateTexture(FRDGTextureDesc::Create2D(TextureSize, bLQFormat ? PF_R5G6B5_UNORM : PF_B8G8R8A8, FClearValueBinding::Black, TexCreate_RenderTargetable | TexCreate_ShaderResource), TEXT("RenderTexture1"));
+				}
+				if (bCompressPass)
+				{
+					OutputAlias0 = CompressTexture0_u2 = GraphBuilder.CreateTexture(FRDGTextureDesc::Create2D(TextureSize / 4, Compressed64BitFormat, FClearValueBinding::None, TexCreate_UAV), TEXT("CompressTexture0"));
+					OutputAlias1 = CompressTexture1 = GraphBuilder.CreateTexture(FRDGTextureDesc::Create2D(TextureSize / 4, PF_R32G32B32A32_UINT, FClearValueBinding::None, TexCreate_UAV), TEXT("CompressTexture1"));
+				}
+				if (bCopyThumbnailPass)
+				{
+					OutputAlias0 = CopyTexture0 = GraphBuilder.CreateTexture(FRDGTextureDesc::Create2D(TextureSize, bLQFormat ? PF_R5G6B5_UNORM : PF_B8G8R8A8, FClearValueBinding::Black, TexCreate_RenderTargetable | TexCreate_ShaderResource), TEXT("CopyTexture0"));
+					OutputAlias1 = nullptr;
+				}
+				break;
 			case ERuntimeVirtualTextureMaterialType::BaseColor_Normal_Specular:
 				if (bRenderPass)
 				{
@@ -979,12 +1088,15 @@ namespace RuntimeVirtualTexture
 		bool bIsThumbnails,
 		FRHITexture2D* OutputTexture0,
 		FRHIUnorderedAccessView* OutputUAV0,
+		IPooledRenderTarget* OutputTarget0,
 		FBox2D const& DestBox0,
 		FRHITexture2D* OutputTexture1,
 		FRHIUnorderedAccessView* OutputUAV1,
+		IPooledRenderTarget* OutputTarget1,
 		FBox2D const& DestBox1,
 		FRHITexture2D* OutputTexture2, 
 		FRHIUnorderedAccessView* OutputUAV2,
+		IPooledRenderTarget* OutputTarget2,
 		FBox2D const& DestBox2,
 		FTransform const& UVToWorld,
 		FBox const& WorldBounds,
@@ -1064,12 +1176,12 @@ namespace RuntimeVirtualTexture
 			PassParameters->RenderTargets[0] = GraphSetup.RenderTexture0 ? FRenderTargetBinding(GraphSetup.RenderTexture0, LoadAction) : FRenderTargetBinding();
 			PassParameters->RenderTargets[1] = GraphSetup.RenderTexture1 ? FRenderTargetBinding(GraphSetup.RenderTexture1, LoadAction) : FRenderTargetBinding();
 			PassParameters->RenderTargets[2] = GraphSetup.RenderTexture2 ? FRenderTargetBinding(GraphSetup.RenderTexture2, LoadAction) : FRenderTargetBinding();
-
+    		
 			AddSimpleMeshPass(GraphBuilder, PassParameters, Scene, *View, nullptr, RDG_EVENT_NAME("VirtualTextureDraw"), View->ViewRect,
-			[Scene, View, MaterialType, RuntimeVirtualTextureMask, vLevel, MaxLevel](FDynamicPassMeshDrawListContext* DynamicMeshPassContext)
-			{
-				GatherMeshesToDraw(DynamicMeshPassContext, Scene, View, MaterialType, RuntimeVirtualTextureMask, vLevel, MaxLevel);
-			});
+	    	[Scene, View, MaterialType, RuntimeVirtualTextureMask, vLevel, MaxLevel](FDynamicPassMeshDrawListContext* DynamicMeshPassContext)
+	    	{
+	    		GatherMeshesToDraw(DynamicMeshPassContext, Scene, View, MaterialType, RuntimeVirtualTextureMask, vLevel, MaxLevel);
+	    	});
 		}
 
 		// Compression Pass
@@ -1148,9 +1260,9 @@ namespace RuntimeVirtualTexture
 				InDesc.MaterialType,
 				InDesc.bClearTextures,
 				InDesc.bIsThumbnails,
-				InDesc.Targets[0].Texture, InDesc.Targets[0].UAV, PageDesc.DestBox[0],
-				InDesc.Targets[1].Texture, InDesc.Targets[1].UAV, PageDesc.DestBox[1],
-				InDesc.Targets[2].Texture, InDesc.Targets[2].UAV, PageDesc.DestBox[2],
+				InDesc.Targets[0].Texture, InDesc.Targets[0].UAV, InDesc.Targets[0].PooledRenderTarget, PageDesc.DestBox[0],
+				InDesc.Targets[1].Texture, InDesc.Targets[1].UAV, InDesc.Targets[1].PooledRenderTarget, PageDesc.DestBox[1],
+				InDesc.Targets[2].Texture, InDesc.Targets[2].UAV, InDesc.Targets[2].PooledRenderTarget, PageDesc.DestBox[2],
 				InDesc.UVToWorld,
 				InDesc.WorldBounds,
 				PageDesc.UVRange,
diff --git a/Engine/Source/Runtime/Renderer/Private/VT/VirtualTexturePhysicalSpace.cpp b/Engine/Source/Runtime/Renderer/Private/VT/VirtualTexturePhysicalSpace.cpp
--- a/Engine/Source/Runtime/Renderer/Private/VT/VirtualTexturePhysicalSpace.cpp
+++ b/Engine/Source/Runtime/Renderer/Private/VT/VirtualTexturePhysicalSpace.cpp
@@ -137,15 +137,16 @@ void FVirtualTexturePhysicalSpace::InitRHI()
 {
 	FRHICommandListImmediate& RHICmdList = FRHICommandListExecutor::GetImmediateCommandList();
 
-	// Not all mobile RHIs support sRGB texture views/aliasing, use only linear targets on mobile
-	ETextureCreateFlags VT_SRGB = GetFeatureLevel() > ERHIFeatureLevel::ES3_1 ? TexCreate_SRGB : TexCreate_None;
-
 	for (int32 Layer = 0; Layer < Description.NumLayers; ++Layer)
 	{
 		const EPixelFormat FormatSRV = RemapVirtualTexturePhysicalSpaceFormat(Description.Format[Layer]);
 		const EPixelFormat FormatUAV = GetUnorderedAccessViewFormat(FormatSRV);
 		const bool bCreateAliasedUAV = (FormatUAV != PF_Unknown) && (FormatUAV != FormatSRV);
 		
+		const bool b16BitFormat = (FormatSRV == EPixelFormat::PF_R5G6B5_UNORM) || (FormatSRV == EPixelFormat::PF_B5G5R5A1_UNORM);
+		// Not all mobile RHIs support sRGB texture views/aliasing, use only linear targets on mobile, and the 16bit format transfer the color space manually in shader
+		ETextureCreateFlags VT_SRGB = (b16BitFormat || GetFeatureLevel() <= ERHIFeatureLevel::ES3_1) ? TexCreate_None : TexCreate_SRGB;
+
 		// Allocate physical texture from the render target pool
 		const uint32 TextureSize = GetTextureSize();
 		FPooledRenderTargetDesc Desc = FPooledRenderTargetDesc::Create2DDesc(
diff --git a/Engine/Source/Runtime/Renderer/Public/VT/RuntimeVirtualTextureRender.h b/Engine/Source/Runtime/Renderer/Public/VT/RuntimeVirtualTextureRender.h
--- a/Engine/Source/Runtime/Renderer/Public/VT/RuntimeVirtualTextureRender.h
+++ b/Engine/Source/Runtime/Renderer/Public/VT/RuntimeVirtualTextureRender.h
@@ -35,6 +35,8 @@ namespace RuntimeVirtualTexture
 		FRHITexture2D* Texture = nullptr;
 		/** Unordered access view of physical texture to render to. If this exists we can render to it directly instead of using RHICopyTexture(). */
 		FRHIUnorderedAccessView* UAV = nullptr;
+
+		struct IPooledRenderTarget* PooledRenderTarget = nullptr;
 	};
 
 	/** A single page description. Multiple of these can be placed in a single FRenderPageBatchDesc batch description. */
