Constrain selection outline pass to run on the render thread.

diff --git a/Engine/Source/Runtime/Renderer/Private/PostProcess/PostProcessSelectionOutline.cpp b/Engine/Source/Runtime/Renderer/Private/PostProcess/PostProcessSelectionOutline.cpp
--- a/Engine/Source/Runtime/Renderer/Private/PostProcess/PostProcessSelectionOutline.cpp
+++ b/Engine/Source/Runtime/Renderer/Private/PostProcess/PostProcessSelectionOutline.cpp
@@ -104,7 +104,7 @@ FScreenPassTexture AddSelectionOutlinePass(FRDGBuilder& GraphBuilder, const FVie
 			RDG_EVENT_NAME("OutlineDepth %dx%d", SceneColorViewport.Rect.Width(), SceneColorViewport.Rect.Height()),
 			PassParameters,
 			ERDGPassFlags::Raster,
-			[&View, SceneColorViewport, DepthStencilTexture, NaniteRasterResults, PassParameters, bNaniteEnabled](FRHICommandList& RHICmdList)
+			[&View, SceneColorViewport, DepthStencilTexture, NaniteRasterResults, PassParameters, bNaniteEnabled](FRHICommandListImmediate& RHICmdList)
 		{
 			RHICmdList.SetViewport(SceneColorViewport.Rect.Min.X, SceneColorViewport.Rect.Min.Y, 0.0f, SceneColorViewport.Rect.Max.X, SceneColorViewport.Rect.Max.Y, 1.0f);
 
