Fix typo when initializing default textures with 11.11.10 or 10.10.10.2 format.

diff --git a/Engine/Source/Runtime/Renderer/Private/SystemTextures.cpp b/Engine/Source/Runtime/Renderer/Private/SystemTextures.cpp
--- a/Engine/Source/Runtime/Renderer/Private/SystemTextures.cpp
+++ b/Engine/Source/Runtime/Renderer/Private/SystemTextures.cpp
@@ -1019,9 +1019,9 @@ void FormatData111110(const TInType& In, uint8* Out, uint32& OutByteCount)
 {
 	uint32* OutTyped = (uint32*)Out;
 	*OutTyped = 
-		 (2048u & ConvertInputFormat<typename TFormatConversionTraits<TInType>::Type, uint32, EDefaultInputType::UNorm11>(In[0]))     |
-		((2048u & ConvertInputFormat<typename TFormatConversionTraits<TInType>::Type, uint32, EDefaultInputType::UNorm11>(In[1]))<<11)|
-		((1024u & ConvertInputFormat<typename TFormatConversionTraits<TInType>::Type, uint32, EDefaultInputType::UNorm10>(In[2]))<<22);
+		 (2047u & ConvertInputFormat<typename TFormatConversionTraits<TInType>::Type, uint32, EDefaultInputType::UNorm11>(In[0]))     |
+		((2047u & ConvertInputFormat<typename TFormatConversionTraits<TInType>::Type, uint32, EDefaultInputType::UNorm11>(In[1]))<<11)|
+		((1023u & ConvertInputFormat<typename TFormatConversionTraits<TInType>::Type, uint32, EDefaultInputType::UNorm10>(In[2]))<<22);
 	OutByteCount = 4;
 }
 
@@ -1030,9 +1030,9 @@ void FormatData1010102(const TInType& In, uint8* Out, uint32& OutByteCount)
 {
 	uint32* OutTyped = (uint32*)Out;
 	*OutTyped =
-		 (1024u & ConvertInputFormat<typename TFormatConversionTraits<TInType>::Type, uint32, EDefaultInputType::UNorm10>(In[0]))        |
-		((1024u & ConvertInputFormat<typename TFormatConversionTraits<TInType>::Type, uint32, EDefaultInputType::UNorm10>(In[1])) << 10) |
-		((1024u & ConvertInputFormat<typename TFormatConversionTraits<TInType>::Type, uint32, EDefaultInputType::UNorm10>(In[2])) << 20) |
+		 (1023u & ConvertInputFormat<typename TFormatConversionTraits<TInType>::Type, uint32, EDefaultInputType::UNorm10>(In[0]))        |
+		((1023u & ConvertInputFormat<typename TFormatConversionTraits<TInType>::Type, uint32, EDefaultInputType::UNorm10>(In[1])) << 10) |
+		((1023u & ConvertInputFormat<typename TFormatConversionTraits<TInType>::Type, uint32, EDefaultInputType::UNorm10>(In[2])) << 20) |
 		((   3u & ConvertInputFormat<typename TFormatConversionTraits<TInType>::Type, uint32, EDefaultInputType::UNorm2> (In[3])) << 30);
 	OutByteCount = 4;
 }
