Switches r.Histogram.UseAtomic=0 by default until issues are norrowed down

diff --git a/Engine/Source/Runtime/Renderer/Private/PostProcess/PostProcessHistogram.cpp b/Engine/Source/Runtime/Renderer/Private/PostProcess/PostProcessHistogram.cpp
--- a/Engine/Source/Runtime/Renderer/Private/PostProcess/PostProcessHistogram.cpp
+++ b/Engine/Source/Runtime/Renderer/Private/PostProcess/PostProcessHistogram.cpp
@@ -9,7 +9,7 @@
 #include "ShaderCompilerCore.h"
 
 TAutoConsoleVariable<int32> CVarUseAtomicHistogram(
-	TEXT("r.Histogram.UseAtomic"), 1,
+	TEXT("r.Histogram.UseAtomic"), 0,
 	TEXT("Uses atomic to speed up the generation of the histogram."),
 	ECVF_RenderThreadSafe);
 
