Move Nanite base pass rendering into the base pass render loop to fix a few rendering order issues

diff --git a/Engine/Source/Runtime/Renderer/Private/BasePassRendering.cpp b/Engine/Source/Runtime/Renderer/Private/BasePassRendering.cpp
--- a/Engine/Source/Runtime/Renderer/Private/BasePassRendering.cpp
+++ b/Engine/Source/Runtime/Renderer/Private/BasePassRendering.cpp
@@ -700,7 +700,9 @@ void FDeferredShadingSceneRenderer::RenderBasePass(
 	const FDBufferTextures& DBufferTextures,
 	FExclusiveDepthStencil::Type BasePassDepthStencilAccess,
 	FRDGTextureRef ForwardShadowMaskTexture,
-	FInstanceCullingManager& InstanceCullingManager)
+	FInstanceCullingManager& InstanceCullingManager,
+	bool bNaniteEnabled,
+	const TArrayView<Nanite::FRasterResults>& NaniteRasterResults)
 {
 	TRACE_CPUPROFILER_EVENT_SCOPE(FDeferredShadingSceneRenderer::RenderBasePass);
 
@@ -871,7 +873,7 @@ void FDeferredShadingSceneRenderer::RenderBasePass(
 	}
 
 	GraphBuilder.SetCommandListStat(GET_STATID(STAT_CLM_BasePass));
-	RenderBasePassInternal(GraphBuilder, BasePassRenderTargets, BasePassDepthStencilAccess, ForwardBasePassTextures, DBufferTextures, SceneTextures.QuadOverdraw, bDoParallelBasePass, bRenderLightmapDensity, InstanceCullingManager);
+	RenderBasePassInternal(GraphBuilder, SceneTextures, BasePassRenderTargets, BasePassDepthStencilAccess, ForwardBasePassTextures, DBufferTextures, bDoParallelBasePass, bRenderLightmapDensity, InstanceCullingManager, bNaniteEnabled, NaniteRasterResults);
 	GraphBuilder.SetCommandListStat(GET_STATID(STAT_CLM_AfterBasePass));
 
 	if (ViewFamily.ViewExtensions.Num() > 0)
@@ -1049,14 +1051,16 @@ static void RenderEditorPrimitives(
 
 void FDeferredShadingSceneRenderer::RenderBasePassInternal(
 	FRDGBuilder& GraphBuilder,
+	const FSceneTextures& SceneTextures,
 	const FRenderTargetBindingSlots& BasePassRenderTargets,
 	FExclusiveDepthStencil::Type BasePassDepthStencilAccess,
 	const FForwardBasePassTextures& ForwardBasePassTextures,
 	const FDBufferTextures& DBufferTextures,
-	FRDGTextureRef QuadOverdrawTexture,
 	bool bParallelBasePass,
 	bool bRenderLightmapDensity,
-	FInstanceCullingManager& InstanceCullingManager)
+	FInstanceCullingManager& InstanceCullingManager,
+	bool bNaniteEnabled,
+	const TArrayView<Nanite::FRasterResults>& NaniteRasterResults)
 {
 	RDG_CSV_STAT_EXCLUSIVE_SCOPE(GraphBuilder, RenderBasePass);
 	SCOPED_NAMED_EVENT(FDeferredShadingSceneRenderer_RenderBasePass, FColor::Emerald);
@@ -1069,7 +1073,7 @@ void FDeferredShadingSceneRenderer::RenderBasePassInternal(
 	else if (ViewFamily.UseDebugViewPS())
 	{
 		// Override the base pass with one of the debug view shader mode (see EDebugViewShaderMode) if required.
-		RenderDebugViewMode(GraphBuilder, Views, QuadOverdrawTexture, BasePassRenderTargets);
+		RenderDebugViewMode(GraphBuilder, Views, SceneTextures.QuadOverdraw, BasePassRenderTargets);
 	}
 	else
 	{
@@ -1077,6 +1081,49 @@ void FDeferredShadingSceneRenderer::RenderBasePassInternal(
 		RDG_EVENT_SCOPE(GraphBuilder, "BasePass");
 		RDG_GPU_STAT_SCOPE(GraphBuilder, Basepass);
 
+		// Nanite materials do not currently support most debug view modes.
+		const bool bShouldDrawNaniteBasePass = bNaniteEnabled && !ViewFamily.EngineShowFlags.ShaderComplexity && !ViewFamily.EngineShowFlags.Wireframe;
+		const bool bNeedsPrePass = ShouldRenderPrePass();
+		auto RenderNaniteBasePass = [&](FViewInfo& View, int32 ViewIndex)
+		{
+			Nanite::FRasterResults& RasterResults = NaniteRasterResults[ViewIndex];
+
+			if (!bNeedsPrePass)
+			{
+				// Emit velocity with depth if not writing it in base pass.
+				FRDGTexture* VelocityBuffer = !IsUsingBasePassVelocity(ShaderPlatform) ? SceneTextures.Velocity : nullptr;
+
+				const bool bEmitStencilMask = NANITE_MATERIAL_STENCIL != 0;
+
+				Nanite::EmitDepthTargets(
+					GraphBuilder,
+					*Scene,
+					View,
+					RasterResults.SOAStrides,
+					RasterResults.VisibleClustersSWHW,
+					RasterResults.ViewsBuffer,
+					SceneTextures.Depth.Target,
+					RasterResults.VisBuffer64,
+					VelocityBuffer,
+					RasterResults.MaterialDepth,
+					RasterResults.NaniteMask,
+					bNeedsPrePass,
+					bEmitStencilMask
+				);
+			}
+
+			Nanite::DrawBasePass(
+				GraphBuilder,
+				View.NaniteMaterialPassCommands,
+				*this,
+				SceneTextures,
+				DBufferTextures,
+				*Scene,
+				View,
+				RasterResults
+			);
+		};
+
 		if (bParallelBasePass)
 		{
 			RDG_WAIT_FOR_TASKS_CONDITIONAL(GraphBuilder, IsBasePassWaitForTasksEnabled());
@@ -1111,6 +1158,11 @@ void FDeferredShadingSceneRenderer::RenderBasePassInternal(
 						FRDGParallelCommandListSet ParallelCommandListSet(RHICmdList, GET_STATID(STAT_CLP_BasePass), *this, View, FParallelCommandListBindings(PassParameters));
 						View.ParallelMeshDrawCommandPasses[EMeshPass::BasePass].DispatchDraw(&ParallelCommandListSet, RHICmdList, &PassParameters->InstanceCullingDrawParams);
 					});
+				}				
+
+				if (bShouldDrawNaniteBasePass)
+				{
+					RenderNaniteBasePass(View, ViewIndex);
 				}
 
 				RenderEditorPrimitives(GraphBuilder, PassParameters, View, DrawRenderState, InstanceCullingManager);
@@ -1171,6 +1223,11 @@ void FDeferredShadingSceneRenderer::RenderBasePassInternal(
 					});
 				}
 
+				if (bShouldDrawNaniteBasePass)
+				{
+					RenderNaniteBasePass(View, ViewIndex);
+				}
+
 				RenderEditorPrimitives(GraphBuilder, PassParameters, View, DrawRenderState, InstanceCullingManager);
 
 				if (bShouldRenderView && View.Family->EngineShowFlags.Atmosphere)
diff --git a/Engine/Source/Runtime/Renderer/Private/DeferredShadingRenderer.cpp b/Engine/Source/Runtime/Renderer/Private/DeferredShadingRenderer.cpp
--- a/Engine/Source/Runtime/Renderer/Private/DeferredShadingRenderer.cpp
+++ b/Engine/Source/Runtime/Renderer/Private/DeferredShadingRenderer.cpp
@@ -1905,13 +1905,6 @@ void FDeferredShadingSceneRenderer::Render(FRDGBuilder& GraphBuilder)
 	// NOTE: Must be done after  system texture initialization
 	VirtualShadowMapArray.Initialize(GraphBuilder, Scene->VirtualShadowMapArrayCacheManager, UseVirtualShadowMaps(ShaderPlatform, FeatureLevel));
 
-	// Nanite materials do not currently support most debug view modes.
-	const bool bShouldApplyNaniteMaterials
-		 = !ViewFamily.EngineShowFlags.ShaderComplexity
-		&& !ViewFamily.UseDebugViewPS()
-		&& !ViewFamily.EngineShowFlags.Wireframe
-		&& !ViewFamily.EngineShowFlags.LightMapDensity;
-
 	// if DDM_AllOpaqueNoVelocity was used, then velocity should have already been rendered as well
 	const bool bIsEarlyDepthComplete = (DepthPass.EarlyZPassMode == DDM_AllOpaque || DepthPass.EarlyZPassMode == DDM_AllOpaqueNoVelocity);
 
@@ -2541,52 +2534,8 @@ void FDeferredShadingSceneRenderer::Render(FRDGBuilder& GraphBuilder)
 #endif
 
 	{
-		RenderBasePass(GraphBuilder, SceneTextures, DBufferTextures, BasePassDepthStencilAccess, ForwardScreenSpaceShadowMaskTexture, InstanceCullingManager);
-		GraphBuilder.AddDispatchHint();
-
-		if (bNaniteEnabled && bShouldApplyNaniteMaterials)
-		{
-			for (int32 ViewIndex = 0; ViewIndex < Views.Num(); ++ViewIndex)
-			{
-				FViewInfo& View = Views[ViewIndex];
-				Nanite::FRasterResults& RasterResults = NaniteRasterResults[ViewIndex];
-
-				if (!bNeedsPrePass)
-				{
-					// Emit velocity with depth if not writing it in base pass.
-					FRDGTexture* VelocityBuffer = !IsUsingBasePassVelocity(ShaderPlatform) ? SceneTextures.Velocity : nullptr;
-
-					const bool bEmitStencilMask = NANITE_MATERIAL_STENCIL != 0;
-
-					Nanite::EmitDepthTargets(
-						GraphBuilder,
-						*Scene,
-						View,
-						RasterResults.SOAStrides,
-						RasterResults.VisibleClustersSWHW,
-						RasterResults.ViewsBuffer,
-						SceneTextures.Depth.Target,
-						RasterResults.VisBuffer64,
-						VelocityBuffer,
-						RasterResults.MaterialDepth,
-						RasterResults.NaniteMask,
-						bNeedsPrePass,
-						bEmitStencilMask
-					);
-				}
-
-				Nanite::DrawBasePass(
-					GraphBuilder,
-					View.NaniteMaterialPassCommands,
-					*this,
-					SceneTextures,
-					DBufferTextures,
-					*Scene,
-					View,
-					RasterResults
-				);
-			}
-		}
+		RenderBasePass(GraphBuilder, SceneTextures, DBufferTextures, BasePassDepthStencilAccess, ForwardScreenSpaceShadowMaskTexture, InstanceCullingManager, bNaniteEnabled, NaniteRasterResults);
+		GraphBuilder.AddDispatchHint();	
 
 		if (!bAllowReadOnlyDepthBasePass)
 		{
diff --git a/Engine/Source/Runtime/Renderer/Private/DeferredShadingRenderer.h b/Engine/Source/Runtime/Renderer/Private/DeferredShadingRenderer.h
--- a/Engine/Source/Runtime/Renderer/Private/DeferredShadingRenderer.h
+++ b/Engine/Source/Runtime/Renderer/Private/DeferredShadingRenderer.h
@@ -205,18 +205,22 @@ public:
 		const FDBufferTextures& DBufferTextures,
 		FExclusiveDepthStencil::Type BasePassDepthStencilAccess,
 		FRDGTextureRef ForwardShadowMaskTexture,
-		FInstanceCullingManager& InstanceCullingManager);
+		FInstanceCullingManager& InstanceCullingManager,
+		bool bNaniteEnabled,
+		const TArrayView<Nanite::FRasterResults>& NaniteRasterResults);
 
 	void RenderBasePassInternal(
 		FRDGBuilder& GraphBuilder,
+		const FSceneTextures& SceneTextures,
 		const FRenderTargetBindingSlots& BasePassRenderTargets,
 		FExclusiveDepthStencil::Type BasePassDepthStencilAccess,
 		const FForwardBasePassTextures& ForwardBasePassTextures,
 		const FDBufferTextures& DBufferTextures,
-		FRDGTextureRef QuadOverdrawTexture,
 		bool bParallelBasePass,
 		bool bRenderLightmapDensity,
-		FInstanceCullingManager& InstanceCullingManager);
+		FInstanceCullingManager& InstanceCullingManager,
+		bool bNaniteEnabled,
+		const TArrayView<Nanite::FRasterResults>& NaniteRasterResults);
 
 	bool ShouldRenderAnisotropyPass() const;
 
