Disable parallel RDG execution for mobile renderer.

#swarm https://p4-swarm.epicgames.net/reviews/16988240/v1/
#pf 6101d11dff10730001d2ede8

diff --git a/Engine/Source/Runtime/Renderer/Private/PlanarReflectionRendering.cpp b/Engine/Source/Runtime/Renderer/Private/PlanarReflectionRendering.cpp
--- a/Engine/Source/Runtime/Renderer/Private/PlanarReflectionRendering.cpp
+++ b/Engine/Source/Runtime/Renderer/Private/PlanarReflectionRendering.cpp
@@ -346,8 +346,11 @@ static void UpdatePlanarReflectionContents_RenderThread(
 #else
 				SCOPED_DRAW_EVENT(RHICmdList, UpdatePlanarReflectionContent_RenderThread);
 #endif
+				const ERDGBuilderFlags BuilderFlags = SceneRenderer->Scene->GetShadingPath() == EShadingPath::Mobile
+					? ERDGBuilderFlags::None
+					: ERDGBuilderFlags::AllowParallelExecute;
 
-				FRDGBuilder GraphBuilder(RHICmdList, RDG_EVENT_NAME("PlanarReflection"), ERDGBuilderFlags::AllowParallelExecute);
+				FRDGBuilder GraphBuilder(RHICmdList, RDG_EVENT_NAME("PlanarReflection"), BuilderFlags);
 
 				// Reflection view late update
 				if (SceneRenderer->Views.Num() > 1)
diff --git a/Engine/Source/Runtime/Renderer/Private/ReflectionEnvironmentCapture.cpp b/Engine/Source/Runtime/Renderer/Private/ReflectionEnvironmentCapture.cpp
--- a/Engine/Source/Runtime/Renderer/Private/ReflectionEnvironmentCapture.cpp
+++ b/Engine/Source/Runtime/Renderer/Private/ReflectionEnvironmentCapture.cpp
@@ -723,7 +723,11 @@ void CaptureSceneToScratchCubemap(FRHICommandListImmediate& RHICmdList, FSceneRe
 	{
 		SCOPED_DRAW_EVENT(RHICmdList, CubeMapCapture);
 
-		FRDGBuilder GraphBuilder(RHICmdList, RDG_EVENT_NAME("CubeMapCapture"), ERDGBuilderFlags::AllowParallelExecute);
+		const ERDGBuilderFlags BuilderFlags = SceneRenderer->Scene->GetShadingPath() == EShadingPath::Mobile
+			? ERDGBuilderFlags::None
+			: ERDGBuilderFlags::AllowParallelExecute;
+
+		FRDGBuilder GraphBuilder(RHICmdList, RDG_EVENT_NAME("CubeMapCapture"), BuilderFlags);
 
 		// Render the scene normally for one face of the cubemap
 		SceneRenderer->Render(GraphBuilder);
diff --git a/Engine/Source/Runtime/Renderer/Private/SceneCaptureRendering.cpp b/Engine/Source/Runtime/Renderer/Private/SceneCaptureRendering.cpp
--- a/Engine/Source/Runtime/Renderer/Private/SceneCaptureRendering.cpp
+++ b/Engine/Source/Runtime/Renderer/Private/SceneCaptureRendering.cpp
@@ -283,12 +283,16 @@ static void UpdateSceneCaptureContentDeferred_RenderThread(
 	// update any resources that needed a deferred update
 	FDeferredUpdateResource::UpdateResources(RHICmdList);
 	{
+		const ERDGBuilderFlags BuilderFlags = SceneRenderer->Scene->GetShadingPath() == EShadingPath::Mobile
+			? ERDGBuilderFlags::None
+			: ERDGBuilderFlags::AllowParallelExecute;
+
 #if WANTS_DRAW_MESH_EVENTS
 		SCOPED_DRAW_EVENTF(RHICmdList, SceneCapture, TEXT("SceneCapture %s"), *EventName);
-		FRDGBuilder GraphBuilder(RHICmdList, RDG_EVENT_NAME("SceneCapture %s", *EventName), ERDGBuilderFlags::AllowParallelExecute);
+		FRDGBuilder GraphBuilder(RHICmdList, RDG_EVENT_NAME("SceneCapture %s", *EventName), BuilderFlags);
 #else
 		SCOPED_DRAW_EVENT(RHICmdList, UpdateSceneCaptureContent_RenderThread);
-		FRDGBuilder GraphBuilder(RHICmdList, RDG_EVENT_NAME("SceneCapture"), ERDGBuilderFlags::AllowParallelExecute);
+		FRDGBuilder GraphBuilder(RHICmdList, RDG_EVENT_NAME("SceneCapture"), BuilderFlags);
 #endif
 
 		FRDGTextureRef TargetTexture = RegisterExternalTexture(GraphBuilder, RenderTarget->GetRenderTargetTexture(), TEXT("SceneCaptureTarget"));
diff --git a/Engine/Source/Runtime/Renderer/Private/SceneRendering.cpp b/Engine/Source/Runtime/Renderer/Private/SceneRendering.cpp
--- a/Engine/Source/Runtime/Renderer/Private/SceneRendering.cpp
+++ b/Engine/Source/Runtime/Renderer/Private/SceneRendering.cpp
@@ -3794,13 +3794,17 @@ static void RenderViewFamily_RenderThread(FRHICommandListImmediate& RHICmdList,
 		SCOPE_CYCLE_COUNTER(STAT_TotalSceneRenderingTime);
 
 		{
+			const ERDGBuilderFlags BuilderFlags = SceneRenderer->Scene->GetShadingPath() == EShadingPath::Mobile
+				? ERDGBuilderFlags::None
+				: ERDGBuilderFlags::AllowParallelExecute;
+
 			FRDGBuilder GraphBuilder(
 				RHICmdList,
 				RDG_EVENT_NAME("SceneRenderer_%s(ViewFamily=%s)",
 					ViewFamily.EngineShowFlags.HitProxies ? TEXT("RenderHitProxies") : TEXT("Render"),
 					ViewFamily.bResolveScene ? TEXT("Primary") : TEXT("Auxiliary")
 				),
-				ERDGBuilderFlags::AllowParallelExecute
+				BuilderFlags
 			);
 
 			if (ViewFamily.EngineShowFlags.HitProxies)
