Drop reference to physical page pool in contexts with VSM caching disabled
Ensures that thumbnail scene renderers don't hang on to copies of extra page pools

diff --git a/Engine/Source/Runtime/Renderer/Private/VirtualShadowMaps/VirtualShadowMapCacheManager.cpp b/Engine/Source/Runtime/Renderer/Private/VirtualShadowMaps/VirtualShadowMapCacheManager.cpp
--- a/Engine/Source/Runtime/Renderer/Private/VirtualShadowMaps/VirtualShadowMapCacheManager.cpp
+++ b/Engine/Source/Runtime/Renderer/Private/VirtualShadowMaps/VirtualShadowMapCacheManager.cpp
@@ -121,7 +121,7 @@ TRefCountPtr<IPooledRenderTarget> FVirtualShadowMapArrayCacheManager::SetPhysica
 		GRenderTargetPool.FindFreeElement(GraphBuilder.RHICmdList, Desc2D, PhysicalPagePool, TEXT("Shadow.Virtual.PhysicalPagePool"));
 
 		Invalidate();
-		UE_LOG(LogRenderer, Display, TEXT("Recreating Shadow.Virtual.PhysicalPagePool. This will also drop any cached pages."));
+		//UE_LOG(LogRenderer, Display, TEXT("Recreating Shadow.Virtual.PhysicalPagePool. This will also drop any cached pages."));
 	}
 
 	return PhysicalPagePool;
@@ -240,12 +240,16 @@ void FVirtualShadowMapArrayCacheManager::ExtractFrameData(
 		{
 			GraphBuilder.QueueBufferExtraction(VirtualShadowMapArray.PageTableRDG, &PrevBuffers.PageTable);
 		}
+
+		CacheEntries.Reset();
 	}
 	else
 	{
-		PrevCacheEntries.Empty();
+		// We drop the physical page pool here as well to ensure that it disappears in the case where
+		// thumbnail rendering or similar creates multiple FSceneRenderers that never get deleted.
+		// Caching is disabled on these contexts intentionally to avoid these issues.
+		FreePhysicalPool();
 	}
-	CacheEntries.Reset();
 
 	// Drop any temp references embedded in the uniform parameters this frame.
 	// We'll reestablish them when we reimport the extracted resources next frame
