Rewrites FFT bloom to RDG

diff --git a/Engine/Shaders/Private/PostProcessFFTBloom.usf b/Engine/Shaders/Private/PostProcessFFTBloom.usf
--- a/Engine/Shaders/Private/PostProcessFFTBloom.usf
+++ b/Engine/Shaders/Private/PostProcessFFTBloom.usf
@@ -280,56 +280,3 @@ void CaptureKernelWeightsCS(uint3 GroupID : SV_GroupID, uint3 GroupThreadID : SV
 }
 
 #endif //#ifdef INCLUDE_CAPTURE_KERNEL_WEIGHTS
-
-
-
-#ifdef INCLUDE_PASSTHROUGH
-
-#define SRCTYPE  float4
-#define DSTTYPE  float4
-
-#define Texture2DType    Texture2D<SRCTYPE>
-#define RWTexture2DType  RWTexture2D<DSTTYPE>
-
-// Input SRV: 
-Texture2DType SrcTexture;
-
-// Output: 
-RWTexture2DType DstTexture;
-
-uint4 DstRect;
-uint4 SrcRect;
-[numthreads(THREADS_PER_GROUP, 1, 1)]
-void PassThroughCS(uint3 GroupID : SV_GroupID, uint3 GroupThreadID : SV_GroupThreadID )
-{
-	// The pixel that this thread 'owns'
-
-	uint2 DstMin = DstRect.xy;
-	uint2 PixelId = uint2(THREADS_PER_GROUP * GroupID.x + GroupThreadID.x, GroupID.y) + DstMin;
-	
-	// If the actual scan line length was not divisible by THREADS_PER_GROUP we will have a few extra threads that don't
-	// own any pixels.  This will be 'false' for most thread groups.
-
-
-	if ( !(PixelId.x  < DstRect.z )  ) return;
-	if ( !(PixelId.y  < DstRect.w )  ) return;
-
-	
-
-	// Check that this pixel is in the src window - otherwise copy black
-	bool bInSrcRect = true;
-	bInSrcRect = bInSrcRect && (PixelId.x < (SrcRect.z + 1)) && (PixelId.y < (SrcRect.w + 1));
-	bInSrcRect = bInSrcRect && !(PixelId.x < SrcRect.x ) && !(PixelId.y < SrcRect.y);
-
-	float4 ResultValue = float4(0.f, 0.f, 0.f, 0.f);
-	if (bInSrcRect)
-	{
-		ResultValue = SrcTexture[PixelId];
-	}
-
-	DstTexture[PixelId] = ResultValue;	
-}
-
-#endif  // INCLUDE_PASSTHROUGH
-
-
diff --git a/Engine/Source/Runtime/Engine/Private/SceneView.cpp b/Engine/Source/Runtime/Engine/Private/SceneView.cpp
--- a/Engine/Source/Runtime/Engine/Private/SceneView.cpp
+++ b/Engine/Source/Runtime/Engine/Private/SceneView.cpp
@@ -2131,12 +2131,6 @@ void FSceneView::EndFinalPostprocessSettings(const FSceneViewInitOptions& ViewIn
 #endif // WITH_EDITOR
 
 	check(VerifyMembersChecks());
-
-	// If we are using FFT bloom and we don't have a convolution texture load the default one.
-	if (FinalPostProcessSettings.BloomMethod == EBloomMethod::BM_FFT && FinalPostProcessSettings.BloomConvolutionTexture == nullptr)
-	{
-		GEngine->LoadDefaultBloomTexture();
-	}
 }
 
 void FSceneView::ConfigureBufferVisualizationSettings()
diff --git a/Engine/Source/Runtime/Renderer/Private/GPUFastFourierTransform.cpp b/Engine/Source/Runtime/Renderer/Private/GPUFastFourierTransform.cpp
--- a/Engine/Source/Runtime/Renderer/Private/GPUFastFourierTransform.cpp
+++ b/Engine/Source/Runtime/Renderer/Private/GPUFastFourierTransform.cpp
@@ -3,7 +3,7 @@
 #include "GPUFastFourierTransform.h"
 #include "SceneUtils.h"
 #include "GlobalShader.h"
-#include "RenderTargetPool.h" // used for on-the-fly acquisition of 
+#include "RenderGraph.h" 
 #include "ShaderCompilerCore.h"
 
 bool ShouldCompileGPUFFT(EShaderPlatform Platform)
@@ -16,16 +16,6 @@ uint32 GPUFFT::MaxScanLineLength()
 	return 4096;
 }
 
-
- FString GPUFFT::FFT_TypeName(const FFT_XFORM_TYPE& XFormType)
-{
-	if (XFormType == FFT_XFORM_TYPE::FORWARD_HORIZONTAL) return FString(TEXT("Forward Horizontal"));
-	if (XFormType == FFT_XFORM_TYPE::INVERSE_HORIZONTAL) return FString(TEXT("Inverse Horizontal"));
-	if (XFormType == FFT_XFORM_TYPE::FORWARD_VERTICAL)   return FString(TEXT("Forward Vertical"));
-	if (XFormType == FFT_XFORM_TYPE::INVERSE_VERTICAL)   return FString(TEXT("Inverse Vertical"));
-	return FString(TEXT("Error"));
-}
-
 bool GPUFFT::IsHorizontal(const FFT_XFORM_TYPE& XFormType)
 {
 	return (XFormType == FFT_XFORM_TYPE::FORWARD_HORIZONTAL || XFormType == FFT_XFORM_TYPE::INVERSE_HORIZONTAL);
@@ -68,9 +58,14 @@ bool GPUFFT::FFTDescription::IsForward() const
 	return GPUFFT::IsForward(XFormType);
 }
 
-FString GPUFFT::FFTDescription::FFT_TypeName() const
+const TCHAR* GPUFFT::FFTDescription::FFT_TypeName() const
 {
-	return GPUFFT::FFT_TypeName(XFormType);
+	if (XFormType == FFT_XFORM_TYPE::FORWARD_HORIZONTAL) return TEXT("Forward Horizontal");
+	if (XFormType == FFT_XFORM_TYPE::INVERSE_HORIZONTAL) return TEXT("Inverse Horizontal");
+	if (XFormType == FFT_XFORM_TYPE::FORWARD_VERTICAL)   return TEXT("Forward Vertical");
+	if (XFormType == FFT_XFORM_TYPE::INVERSE_VERTICAL)   return TEXT("Inverse Vertical");
+	unimplemented();
+	return TEXT("Error");
 }
 
 
@@ -136,1502 +131,824 @@ namespace GPUFFT
 	}
 
 
-	/**
-	* Double buffer to manage RenderTargets during multi-pass FFTs.
-	*/
-	class  FDoubleBufferTargets
-	{
-	public:
-		FDoubleBufferTargets(FSceneRenderTargetItem& InitialSrc, FSceneRenderTargetItem& InitialDst) :
-				SrcIdx(0)
-		{
-			Targets[0] = &InitialSrc;
-			Targets[1] = &InitialDst;
-		}
-
-		void Swap() { SrcIdx = 1 - SrcIdx; }
-
-		// Return the index of the current Src target.  If it is 0 than this is the original src, otherwise
-		// it is the original dst.
-		const uint32&  GetSrcIdx() const { return SrcIdx; }
-
-		// Access to the render targets with a return by reference for clear ownership semantics
-		const FSceneRenderTargetItem& SrcTarget() const { return *Targets[SrcIdx]; }
-		FSceneRenderTargetItem& DstTarget() { return *Targets[1 - SrcIdx]; }
-
-	private:
-		uint32 SrcIdx;
-		FSceneRenderTargetItem* Targets[2];
-	};
-}
-
-
-
-namespace GPUFFT
-{
-
-	class FFFTShader : public FGlobalShader
-	{
-	public:
-		static bool ShouldCompilePermutation(const FGlobalShaderPermutationParameters& Parameters)
-		{
-			return ShouldCompileGPUFFT(Parameters.Platform);
-		}
-
-		FFFTShader() = default;
-		FFFTShader(const CompiledShaderInitializerType& Initializer)
-			: FGlobalShader(Initializer)
-		{}
-	};
-
-	class FReorderFFTPassCS : public FFFTShader
-	{
-		// NB:  the following is actually "public:" 
-		// due to text in the DECLARE_SHADER_TYPE Macro 
-		DECLARE_SHADER_TYPE(FReorderFFTPassCS, Global);
-
-	public:
-
-
-		typedef FFT_XFORM_TYPE   FFT_XFORM_TYPE;
-
-		FReorderFFTPassCS() {};
-
-		FReorderFFTPassCS(const ShaderMetaType::CompiledShaderInitializerType& Initializer)
-			: FFFTShader(Initializer)
-		{
-			using GPUFFTComputeShaderUtils::FComputeParameterBinder;
-
-			FComputeParameterBinder Binder(Initializer.ParameterMap);
-			Binder(SrcROTexture,     TEXT("SrcSRV"))
-				(DstRWTexture,       TEXT("DstUAV"))
-				(TransformType,      TEXT("TransformType"))
-				(SrcRect,            TEXT("SrcRect"))
-				(DstRect,            TEXT("DstRect"))
-				(LogTransformLength, TEXT("LogTwoLength"))
-				(BitCount,           TEXT("BitCount"));
-		}
-
-		// Used by IMPLEMENT_SHADER_TYPE2
-		static const TCHAR* GetSourceFilename() { return TEXT("/Engine/Private/GPUFastFourierTransform.usf"); }
-		static const TCHAR* GetFunctionName() { return TEXT("ReorderFFTPassCS"); }
-
-		static void ModifyCompilationEnvironment(const FGlobalShaderPermutationParameters& Parameters, FShaderCompilerEnvironment& OutEnvironment)
-		{
-			FGlobalShader::ModifyCompilationEnvironment(Parameters, OutEnvironment);
-			OutEnvironment.SetDefine(TEXT("INCLUDE_REORDER_FFT_PASS"), 1);
-		}
-
-		void SetCSParamters(FRHICommandList& RHICmdList,
-			const FFT_XFORM_TYPE& XFormType,
-			const FTextureRHIRef& SrcTexture, const FIntRect& SrcRectValue, const FIntRect& DstRectValue, const uint32 TransformLength, const uint32 PowTwoSubLengthCount, const bool bScrubNaNs)
-		{
-			using GPUFFTComputeShaderUtils::FComputeParamterValueSetter;
-
-			FRHIComputeShader* ShaderRHI = RHICmdList.GetBoundComputeShader();
-
-			// Set up the input.  We have to do this explicitly because the FFT dispatches multiple compute shaders and manages their input/output.
-			FComputeParamterValueSetter ParamSetter(RHICmdList, ShaderRHI);
-
-			ParamSetter(SrcROTexture, SrcTexture);
-
-			// Translate the transform type. 
-			uint32 TransformTypeValue = GPUFFT::BitEncode(XFormType);
-			if (bScrubNaNs)
-			{
-				TransformTypeValue |= 4;
-			}
-		
-			const uint32 BitCountValue =    BitSize(PowTwoSubLengthCount) - 1;
-			const uint32 LogTwoTransformLength = BitSize(TransformLength) - 1;
-
-			ParamSetter(TransformType, TransformTypeValue)
-				(SrcRect, SrcRectValue)
-				(DstRect, DstRectValue)
-				(LogTransformLength, LogTwoTransformLength)
-				(BitCount, BitCountValue);
-		}
-
-		// Method for use with the FScopedUAVBind
-		FShaderResourceParameter& DestinationResourceParameter() { return DstRWTexture; }
-
-	public:
-
-		LAYOUT_FIELD(FShaderResourceParameter, SrcROTexture);
-		LAYOUT_FIELD(FShaderResourceParameter, DstRWTexture);
-		LAYOUT_FIELD(FShaderParameter, TransformType);
-		LAYOUT_FIELD(FShaderParameter, SrcRect);
-		LAYOUT_FIELD(FShaderParameter, DstRect);
-		LAYOUT_FIELD(FShaderParameter, LogTransformLength);
-		LAYOUT_FIELD(FShaderParameter, BitCount);
-	};
-
-
-	class FGroupShardSubFFTPassCS : public FFFTShader
-	{
-		// NB:  the following is actually "public:" 
-		// due to text in the DECLARE_SHADER_TYPE Macro 
-		DECLARE_SHADER_TYPE(FGroupShardSubFFTPassCS, Global);
-
-	public:
-
-
-		typedef FFT_XFORM_TYPE   FFT_XFORM_TYPE;
-
-		FGroupShardSubFFTPassCS() {};
-
-		FGroupShardSubFFTPassCS(const ShaderMetaType::CompiledShaderInitializerType& Initializer)
-			: FFFTShader(Initializer)
-		{
-			using GPUFFTComputeShaderUtils::FComputeParameterBinder;
-
-			FComputeParameterBinder Binder(Initializer.ParameterMap);
-			Binder(SrcROTexture,  TEXT("SrcTexture"))
-				(DstRWTexture,    TEXT("DstTexture"))
-				(TransformType,   TEXT("TransformType"))
-				(SrcRect,         TEXT("SrcWindow"))
-				(TransformLength, TEXT("TransformLength"))
-				(NumSubRegions,   TEXT("NumSubRegions"));
-		}
-
-		// Used by IMPLEMENT_SHADER_TYPE2
-		static const TCHAR* GetSourceFilename() { return TEXT("/Engine/Private/GPUFastFourierTransform.usf"); }
-		static const TCHAR* GetFunctionName()   { return TEXT("GroupSharedSubComplexFFTCS"); }
-		
-		static uint32 SubPassLength() { return 2048;}
-		static uint32 Radix() { return 2; }
-
-		static void ModifyCompilationEnvironment(const FGlobalShaderPermutationParameters& Parameters, FShaderCompilerEnvironment& OutEnvironment)
-		{
-			FGlobalShader::ModifyCompilationEnvironment(Parameters, OutEnvironment);
-			OutEnvironment.SetDefine(TEXT("INCLUDE_GROUP_SHARED_SUB_COMPLEX_FFT"), 1);
-			OutEnvironment.SetDefine(TEXT("SCAN_LINE_LENGTH"), FGroupShardSubFFTPassCS::SubPassLength());
-			OutEnvironment.SetDefine(TEXT("RADIX"), FGroupShardSubFFTPassCS::Radix());
-		
-
-		}
-
-		void SetCSParamters(FRHICommandList& RHICmdList,
-			const FFT_XFORM_TYPE& XFormType,
-			const uint32 TransformLengthValue,
-			const FIntRect& WindowValue,
-			const FTextureRHIRef& SrcTexture,
-			const uint32 SubRegionCount)
-		{
-
-
-			using GPUFFTComputeShaderUtils::FComputeParamterValueSetter;
-
-			FRHIComputeShader* ShaderRHI = RHICmdList.GetBoundComputeShader();
-
-			// Set up the input.  We have to do this explicitly because the FFT dispatches multiple compute shaders and manages their input/output.
-			FComputeParamterValueSetter ParamSetter(RHICmdList, ShaderRHI);
-
-			ParamSetter(SrcROTexture, SrcTexture);
-
-			// Translate the transform type. 
-			uint32 TransformTypeValue = GPUFFT::BitEncode(XFormType);
-
-			ParamSetter(TransformType, TransformTypeValue)
-				(SrcRect, WindowValue)
-				(TransformLength, TransformLengthValue)
-				(NumSubRegions, SubRegionCount);
-		}
-
-		// Method for use with the FScopedUAVBind
-		FShaderResourceParameter& DestinationResourceParameter() { return DstRWTexture; }
-
-	public:
-
-		LAYOUT_FIELD(FShaderResourceParameter, SrcROTexture);
-		LAYOUT_FIELD(FShaderResourceParameter, DstRWTexture);
-		LAYOUT_FIELD(FShaderParameter, TransformType);
-		LAYOUT_FIELD(FShaderParameter, SrcRect);
-		LAYOUT_FIELD(FShaderParameter, TransformLength);
-		LAYOUT_FIELD(FShaderParameter, NumSubRegions);
-	};
-
-	class FComplexFFTPassCS : public FFFTShader
-	{
-	// NB:  the following is actually "public:" 
-	// due to text in the DECLARE_SHADER_TYPE Macro 
-	DECLARE_SHADER_TYPE(FComplexFFTPassCS, Global);
-	
-	public:
-
-
-		typedef FFT_XFORM_TYPE   FFT_XFORM_TYPE;
-
-		FComplexFFTPassCS() {};
-
-		FComplexFFTPassCS(const ShaderMetaType::CompiledShaderInitializerType& Initializer)
-			: FFFTShader(Initializer)
-		{
-			using GPUFFTComputeShaderUtils::FComputeParameterBinder;
-
-			FComputeParameterBinder Binder(Initializer.ParameterMap);
-			Binder(SrcROTexture, TEXT("SrcSRV"))
-				(DstRWTexture, TEXT("DstUAV"))
-				(TransformType, TEXT("TransformType"))
-				(SrcRect, TEXT("SrcRect"))
-				(DstRect, TEXT("DstRect"))
-				(BitCount, TEXT("BitCount"))
-				(PowTwoLength, TEXT("PowTwoLength"));
-		}
-
-		// Used by IMPLEMENT_SHADER_TYPE2
-		static const TCHAR* GetSourceFilename() { return TEXT("/Engine/Private/GPUFastFourierTransform.usf"); }
-		static const TCHAR* GetFunctionName()   { return TEXT("ComplexFFTPassCS"); }
-
-		static void ModifyCompilationEnvironment(const FGlobalShaderPermutationParameters& Parameters, FShaderCompilerEnvironment& OutEnvironment)
-		{
-			FGlobalShader::ModifyCompilationEnvironment(Parameters, OutEnvironment);
-			OutEnvironment.SetDefine(TEXT("INCLUDE_COMPLEX_FFT_PASS"), 1);
-		}
-
-		void SetCSParamters(FRHICommandList& RHICmdList,
-			const FFT_XFORM_TYPE& XFormType,
-			const FTextureRHIRef& SrcTexture, const FIntRect& SrcRectValue, const FIntRect& DstRectValue, const uint32 TransformLength, const uint32 PassLength, const bool bScrubNaNs)
-		{
-			const uint32 BitCountValue = BitSize(TransformLength);
-			const uint32 PowTwo        = PassLength;  // The pass number should be log(2, PassLength)
-
-			using GPUFFTComputeShaderUtils::FComputeParamterValueSetter;
-
-			FRHIComputeShader* ShaderRHI = RHICmdList.GetBoundComputeShader();
-
-			// Set up the input.  We have to do this explicitly because the FFT dispatches multiple compute shaders and manages their input/output.
-			FComputeParamterValueSetter ParamSetter(RHICmdList, ShaderRHI);
-
-			ParamSetter(SrcROTexture, SrcTexture);
-
-			// Translate the transform type. 
-			uint32 TransformTypeValue = GPUFFT::BitEncode(XFormType);
-			if (bScrubNaNs)
-			{
-				TransformTypeValue |= 4;
-			}
-
-			ParamSetter(TransformType, TransformTypeValue)
-				(SrcRect, SrcRectValue)
-				(DstRect, DstRectValue)
-				(BitCount, BitCountValue)
-				(PowTwoLength, PowTwo);
-		}
-
-		// Method for use with the FScopedUAVBind
-		FShaderResourceParameter& DestinationResourceParameter() { return DstRWTexture; }
-
-	public:
-
-		LAYOUT_FIELD(FShaderResourceParameter, SrcROTexture);
-		LAYOUT_FIELD(FShaderResourceParameter, DstRWTexture);
-		LAYOUT_FIELD(FShaderParameter, TransformType);
-		LAYOUT_FIELD(FShaderParameter, SrcRect);
-		LAYOUT_FIELD(FShaderParameter, DstRect);
-		LAYOUT_FIELD(FShaderParameter, BitCount);
-		LAYOUT_FIELD(FShaderParameter, PowTwoLength);
-	};
-
-
-	class FPackTwoForOneFFTPassCS : public FFFTShader
-	{
-		// NB:  the following is actually "public:" 
-		// due to text in the DECLARE_SHADER_TYPE Macro 
-		DECLARE_SHADER_TYPE(FPackTwoForOneFFTPassCS, Global);
-
-	public:
-
-
-		typedef FFT_XFORM_TYPE   FFT_XFORM_TYPE;
-
-		FPackTwoForOneFFTPassCS() {};
-
-		FPackTwoForOneFFTPassCS(const ShaderMetaType::CompiledShaderInitializerType& Initializer)
-			: FFFTShader(Initializer)
-		{
-			using GPUFFTComputeShaderUtils::FComputeParameterBinder;
-
-			FComputeParameterBinder Binder(Initializer.ParameterMap);
-			Binder(SrcROTexture,  TEXT("SrcSRV"))
-					(DstRWTexture,  TEXT("DstUAV"))
-					(TransformType, TEXT("TransformType"))
-					(DstRect,       TEXT("DstRect"));
-		}
-
-		// Used by IMPLEMENT_SHADER_TYPE2
-		static const TCHAR* GetSourceFilename() { return TEXT("/Engine/Private/GPUFastFourierTransform.usf"); }
-		static const TCHAR* GetFunctionName() { return TEXT("PackTwoForOneFFTPassCS"); }
-
-		static void ModifyCompilationEnvironment(const FGlobalShaderPermutationParameters& Parameters, FShaderCompilerEnvironment& OutEnvironment)
-		{
-			FGlobalShader::ModifyCompilationEnvironment(Parameters, OutEnvironment);
-			OutEnvironment.SetDefine(TEXT("INCLUDE_PACK_TWOFORONE_FFT_PASS"), 1);
-		}
-
-		void SetCSParamters(FRHICommandList& RHICmdList,
-			const FFT_XFORM_TYPE& XFormType,
-			const FTextureRHIRef& SrcTexture, const FIntRect& DstRectValue)
-		{
-			using GPUFFTComputeShaderUtils::FComputeParamterValueSetter;
-
-			FRHIComputeShader* ShaderRHI = RHICmdList.GetBoundComputeShader();
-
-			// Set up the input.  We have to do this explicitly because the FFT dispatches multiple compute shaders and manages their input/output.
-			FComputeParamterValueSetter ParamSetter(RHICmdList, ShaderRHI);
-
-			ParamSetter(SrcROTexture, SrcTexture);
-
-			// Translate the transform type. 
-			uint32 TransformTypeValue = GPUFFT::BitEncode(XFormType);
-
-			ParamSetter(TransformType, TransformTypeValue)
-				(DstRect, DstRectValue);
-		}
-
-		// Method for use with the FScopedUAVBind
-		FShaderResourceParameter& DestinationResourceParameter() { return DstRWTexture; }
-
-	public:
-
-		LAYOUT_FIELD(FShaderResourceParameter, SrcROTexture);
-		LAYOUT_FIELD(FShaderResourceParameter, DstRWTexture);
-		LAYOUT_FIELD(FShaderParameter, TransformType);
-		LAYOUT_FIELD(FShaderParameter, DstRect);
-	};
-
-
-	class FCopyWindowCS : public FFFTShader
-	{
-		// NB:  the following is actually "public:" 
-		// due to text in the DECLARE_SHADER_TYPE Macro 
-		DECLARE_SHADER_TYPE(FCopyWindowCS, Global);
-
-	public:
-
-		FCopyWindowCS() {};
-
-		FCopyWindowCS(const ShaderMetaType::CompiledShaderInitializerType& Initializer)
-			: FFFTShader(Initializer)
-		{
-			using GPUFFTComputeShaderUtils::FComputeParameterBinder;
-
-			FComputeParameterBinder Binder(Initializer.ParameterMap);
-			Binder(SrcROTexture, TEXT("SrcSRV"))
-				(DstRWTexture, TEXT("DstUAV"))
-				(SrcRect, TEXT("SrcRect"))
-				(DstRect, TEXT("DstRect"))
-				(PreFilter, TEXT("BrightPixelGain"));
-		}
-
-		// Used by IMPLEMENT_SHADER_TYPE2
-		static const TCHAR* GetSourceFilename() { return TEXT("/Engine/Private/GPUFastFourierTransform.usf"); }
-		static const TCHAR* GetFunctionName()   { return TEXT("CopyWindowCS"); }
-
-		static uint32 XThreadCount() { return 1; }
-		static uint32 YThreadCount() { return 32; }
-
-		static void ModifyCompilationEnvironment(const FGlobalShaderPermutationParameters& Parameters, FShaderCompilerEnvironment& OutEnvironment)
-		{
-			FFFTShader::ModifyCompilationEnvironment(Parameters, OutEnvironment);
-			OutEnvironment.SetDefine(TEXT("INCLUDE_COPY_WINDOW"), 1);
-			OutEnvironment.SetDefine(TEXT("X_THREAD_COUNT"), XThreadCount());
-			OutEnvironment.SetDefine(TEXT("Y_THREAD_COUNT"), YThreadCount());
-		}
-
-		void SetCSParamters(FRHICommandList& RHICmdList,
-			const FIntRect& SrcRectValue,
-			const FTextureRHIRef& SrcTexture, const FIntRect& DstRectValue,
-			const FPreFilter& PreFilterValue)
-		{
-			using GPUFFTComputeShaderUtils::FComputeParamterValueSetter;
-
-			FRHIComputeShader* ShaderRHI = RHICmdList.GetBoundComputeShader();
-
-			// Set up the input.  We have to do this explicitly because the FFT dispatches multiple compute shaders and manages their input/output.
-			FComputeParamterValueSetter ParamSetter(RHICmdList, ShaderRHI);
-
-			ParamSetter(SrcROTexture, SrcTexture);
-			ParamSetter(SrcRect, SrcRectValue)
-				(DstRect, DstRectValue)
-				(PreFilter, PreFilterValue);
-		}
-
-		// Method for use with the FScopedUAVBind
-		FShaderResourceParameter& DestinationResourceParameter() { return DstRWTexture; }
-
-	public:
-
-		LAYOUT_FIELD(FShaderResourceParameter, SrcROTexture);
-		LAYOUT_FIELD(FShaderResourceParameter, DstRWTexture);
-		LAYOUT_FIELD(FShaderParameter, SrcRect);
-		LAYOUT_FIELD(FShaderParameter, DstRect);
-		LAYOUT_FIELD(FShaderParameter, PreFilter);
-	};
-
-	class FComplexMultiplyImagesCS : public FFFTShader
-	{
-		// NB:  the following is actually "public:" 
-		// due to text in the DECLARE_SHADER_TYPE Macro 
-		DECLARE_SHADER_TYPE(FComplexMultiplyImagesCS, Global);
-
-	public:
-
-		FComplexMultiplyImagesCS() {};
-
-		FComplexMultiplyImagesCS(const ShaderMetaType::CompiledShaderInitializerType& Initializer)
-			: FFFTShader(Initializer)
-		{
-			using GPUFFTComputeShaderUtils::FComputeParameterBinder;
-
-			FComputeParameterBinder Binder(Initializer.ParameterMap);
-			Binder(SrcROTexture, TEXT("SrcSRV"))
-				(KnlROTexture, TEXT("KnlSRV"))
-				(DstRWTexture, TEXT("DstUAV"))
-				(SrcRect, TEXT("SrcRect"))
-				(DataLayout, TEXT("DataLayout"));
-		}
-
-		// Used by IMPLEMENT_SHADER_TYPE2
-		static const TCHAR* GetSourceFilename() { return TEXT("/Engine/Private/GPUFastFourierTransform.usf"); }
-		static const TCHAR* GetFunctionName() { return TEXT("ComplexMultiplyImagesCS"); }
-
-		static void ModifyCompilationEnvironment(const FGlobalShaderPermutationParameters& Parameters, FShaderCompilerEnvironment& OutEnvironment)
-		{
-			FGlobalShader::ModifyCompilationEnvironment(Parameters, OutEnvironment);
-			OutEnvironment.SetDefine(TEXT("INCLUDE_COMPLEX_MULTIPLY_IMAGES"), 1);
-		}
-
-		void SetCSParamters(FRHICommandList& RHICmdList,
-			const bool bHorizontalScanlines,
-			const FIntRect& SrcRectValue,
-			const FTextureRHIRef& SrcTexture, 
-			const FTextureRHIRef& KnlTexture)
-		{
-			using GPUFFTComputeShaderUtils::FComputeParamterValueSetter;
-
-			FRHIComputeShader* ShaderRHI = RHICmdList.GetBoundComputeShader();
-
-			const uint32 DataLayoutValue = (bHorizontalScanlines) ? 1 : 0;
-			// Set up the input.  We have to do this explicitly because the FFT dispatches multiple compute shaders and manages their input/output.
-			FComputeParamterValueSetter ParamSetter(RHICmdList, ShaderRHI);
-
-			ParamSetter(SrcROTexture, SrcTexture)
-				(KnlROTexture, KnlTexture)
-				(SrcRect, SrcRectValue)
-				(DataLayout, DataLayoutValue);
-
-		}
-
-		// Method for use with the FScopedUAVBind
-		FShaderResourceParameter& DestinationResourceParameter() { return DstRWTexture; }
-
-	public:
-
-		LAYOUT_FIELD(FShaderResourceParameter, SrcROTexture);
-		LAYOUT_FIELD(FShaderResourceParameter, KnlROTexture);
-		LAYOUT_FIELD(FShaderResourceParameter, DstRWTexture);
-		LAYOUT_FIELD(FShaderParameter, SrcRect);
-		LAYOUT_FIELD(FShaderParameter, DataLayout);
-	};
-
-	class FGSComplexTransformBaseCS : public FFFTShader
-	{
-		DECLARE_INLINE_TYPE_LAYOUT(FGSComplexTransformBaseCS, NonVirtual);
-	public:
-		
-
-		typedef FFT_XFORM_TYPE   FFT_XFORM_TYPE;
-
-		FGSComplexTransformBaseCS() {};
-
-		FGSComplexTransformBaseCS(const ShaderMetaType::CompiledShaderInitializerType& Initializer)
-			: FFFTShader(Initializer)
-		{
-			using GPUFFTComputeShaderUtils::FComputeParameterBinder;
-		
-			FComputeParameterBinder Binder(Initializer.ParameterMap);
-			Binder(SrcROTexture,  TEXT("SrcTexture"))
-				(DstRWTexture,    TEXT("DstTexture"))
-				(TransformType,   TEXT("TransformType"))
-				(SrcRectMin,      TEXT("SrcRectMin"))
-				(SrcRectMax,      TEXT("SrcRectMax"))
-				(DstExtent,       TEXT("DstExtent"))
-				(DstRect,         TEXT("DstRect"))
-				(BrightPixelGain, TEXT("BrightPixelGain"));
-		}
-
-
-		void SetCSParamters(FRHICommandList& RHICmdList,
-			const FFT_XFORM_TYPE& XFormType,
-			const FTextureRHIRef& SrcTexture, const FIntRect& SrcRect, const FIntRect& DstRectValue,
-			const GPUFFT::FPreFilter& PreFilterParameters = GPUFFT::FPreFilter(TNumericLimits<float>::Max(), TNumericLimits<float>::Lowest(), 0.f))
-		{
-			using GPUFFTComputeShaderUtils::FComputeParamterValueSetter;
-
-			FRHIComputeShader* ShaderRHI = RHICmdList.GetBoundComputeShader();
-
-			// Set up the input.  We have to do this explicitly because the FFT dispatches multiple compute shaders and manages their input/output.
-			FComputeParamterValueSetter ParamSetter(RHICmdList, ShaderRHI);
-
-			ParamSetter(SrcROTexture, SrcTexture);
-
-			// Translate the transform type. 
-			uint32 TransformTypeValue = GPUFFT::BitEncode(XFormType);
-
-			// We have a valid prefilter if the min is less than the max
-			if (PreFilterParameters.Component(0) < PreFilterParameters.Component(1))
-			{
-				// Encode a bool to turn on the pre-filter.
-				TransformTypeValue |= 4;
-			}
-
-			ParamSetter(TransformType, TransformTypeValue)
-				(SrcRectMin,      SrcRect.Min)
-				(SrcRectMax,      SrcRect.Max)
-				(DstRect,         DstRectValue)
-				(DstExtent,       DstRectValue.Size())
-				(BrightPixelGain, PreFilterParameters);
-		}
-
-		// Method for use with the FScopedUAVBind
-		FShaderResourceParameter& DestinationResourceParameter() { return DstRWTexture; }
-
-	public:
-		LAYOUT_FIELD(FShaderResourceParameter, SrcROTexture);
-		LAYOUT_FIELD(FShaderResourceParameter, DstRWTexture);
-		LAYOUT_FIELD(FShaderParameter, TransformType);
-		LAYOUT_FIELD(FShaderParameter, SrcRectMin);
-		LAYOUT_FIELD(FShaderParameter, SrcRectMax);
-		LAYOUT_FIELD(FShaderParameter, DstExtent);
-		LAYOUT_FIELD(FShaderParameter, DstRect);
-		LAYOUT_FIELD(FShaderParameter, BrightPixelGain);
-	};
-
-	template <int PowRadixSignalLength>
-	class TGSComplexTransformCS : public FGSComplexTransformBaseCS
-	{
-		// NB:  the following is actually "public:" 
-		// due to text in the DECLARE_SHADER_TYPE Macro 
-		DECLARE_SHADER_TYPE(TGSComplexTransformCS, Global);
-
-		// Used by IMPLEMENT_SHADER_TYPE2
-		static const TCHAR* GetSourceFilename() { return TEXT("/Engine/Private/GPUFastFourierTransform.usf"); }
-		static const TCHAR* GetFunctionName()   { return TEXT("GroupSharedComplexFFTCS"); }
-
-		static void ModifyCompilationEnvironment(const FGlobalShaderPermutationParameters& Parameters, FShaderCompilerEnvironment& OutEnvironment)
-		{
-			FGlobalShader::ModifyCompilationEnvironment(Parameters, OutEnvironment);
-			OutEnvironment.SetDefine(TEXT("INCLUDE_GROUP_SHARED_COMPLEX_FFT"), 1);
-			OutEnvironment.SetDefine(TEXT("SCAN_LINE_LENGTH"), PowRadixSignalLength);
-			if (PowRadixSignalLength > 8)
-			{
-				OutEnvironment.SetDefine(TEXT("MIXED_RADIX"), 1);
-			}
-		}
-
-		/** Default constructor **/
-		TGSComplexTransformCS() {};
-
-	public:
-
-		TGSComplexTransformCS(const ShaderMetaType::CompiledShaderInitializerType& Initializer)
-			: FGSComplexTransformBaseCS(Initializer)
-		{}
-	};
-
-
-
-	template < int PowRadixSignalLength>
-	class TGSTwoForOneTransformCS : public FGSComplexTransformBaseCS
-	{
-		// NB:  the following is actually "public:" 
-		// due to text in the DECLARE_SHADER_TYPE Macro 
-		DECLARE_SHADER_TYPE(TGSTwoForOneTransformCS, Global);
-
-		// Used by IMPLEMENT_SHADER_TYPE2
-		static const TCHAR* GetSourceFilename() { return TEXT("/Engine/Private/GPUFastFourierTransform.usf"); }
-		static const TCHAR* GetFunctionName()   { return TEXT("GroupSharedTwoForOneFFTCS"); }
-
-		static void ModifyCompilationEnvironment(const FGlobalShaderPermutationParameters& Parameters, FShaderCompilerEnvironment& OutEnvironment)
-		{
-			FGlobalShader::ModifyCompilationEnvironment(Parameters, OutEnvironment);
-			OutEnvironment.SetDefine(TEXT("INCLUDE_GROUP_SHARED_TWO_FOR_ONE_FFT"), 1);
-			OutEnvironment.SetDefine(TEXT("SCAN_LINE_LENGTH"), PowRadixSignalLength);
-			if (PowRadixSignalLength > 8)
-			{
-				OutEnvironment.SetDefine(TEXT("MIXED_RADIX"), 1);
-			}
-		}
-
-		/** Default constructor **/
-		TGSTwoForOneTransformCS() {};
-
-	public:
-
-		TGSTwoForOneTransformCS(const ShaderMetaType::CompiledShaderInitializerType& Initializer)
-			: FGSComplexTransformBaseCS(Initializer)
-		{}
-	};
-
-	/**
-	* Base class used for 1d convolution.  
-	*/
-	class FGSConvolutionBaseCS : public FFFTShader
-	{
-		DECLARE_INLINE_TYPE_LAYOUT(FGSConvolutionBaseCS, NonVirtual);
-	public:
-
-		FGSConvolutionBaseCS() {};
-
-		FGSConvolutionBaseCS(const ShaderMetaType::CompiledShaderInitializerType& Initializer)
-			: FFFTShader(Initializer)
-		{
-			using GPUFFTComputeShaderUtils::FComputeParameterBinder;
-
-			FComputeParameterBinder Binder(Initializer.ParameterMap);
-			Binder(SrcROTexture, TEXT("SrcTexture"))
-				(DstRWTexture,   TEXT("DstTexture"))
-				(SrcRectMin,     TEXT("SrcRectMin"))
-				(SrcRectMax,     TEXT("SrcRectMax"))
-				(DstExtent,      TEXT("DstExtent"))
-				(TransformType,  TEXT("TransformType"));
-		}
-
-
-		// @todo - template this on KernelType.  Have KernelType know how to do its own SetShaderValue
-		// Note: We pass in the 2d spectral size (transform size) to help the filter
-		void SetCSParamters(FRHICommandList& RHICmdList, const FFT_XFORM_TYPE& XFormType,
-			const FTextureRHIRef& SrcTexture, const FIntRect& SrcRect, const FIntPoint& DstExtentValue)
-		{
-
-			using GPUFFTComputeShaderUtils::FComputeParamterValueSetter;
-
-			// Translate the transform type. 
-			uint32 TransformTypeValue = GPUFFT::BitEncode(XFormType);
-
-			const bool bUseAlpha = true;
-
-			if (bUseAlpha)
-			{
-				TransformTypeValue |= 8;
-			}
-
-			FRHIComputeShader* ShaderRHI = RHICmdList.GetBoundComputeShader();
-
-			// Set up the input.  We have to do this explicitly because the FFT dispatches multiple compute shaders and manages their input/output.
-
-			FComputeParamterValueSetter ParamSetter(RHICmdList, ShaderRHI);
-			ParamSetter(SrcROTexture, SrcTexture); // set the texture
-
-													// set all the other parameters
-			ParamSetter(SrcRectMin, SrcRect.Min)
-				(SrcRectMax,        SrcRect.Max)
-				(DstExtent,         DstExtentValue)
-				(TransformType,     TransformTypeValue);
-		}
-
-		// Method for use with the FScopedUAVBind
-		FShaderResourceParameter& DestinationResourceParameter() { return DstRWTexture; }
-
-	public:
-		LAYOUT_FIELD(FShaderResourceParameter, SrcROTexture);
-		LAYOUT_FIELD(FShaderResourceParameter, DstRWTexture);
-		LAYOUT_FIELD(FShaderParameter, SrcRectMin);
-		LAYOUT_FIELD(FShaderParameter, SrcRectMax);
-		LAYOUT_FIELD(FShaderParameter, DstExtent);
-		LAYOUT_FIELD(FShaderParameter, TransformType);
-	};
-
-	/**
-	* Common class used by the shader permutations of the convolution with texture
-	*/
-	class FGSConvolutionWithTextureKernelBaseCS : public FGSConvolutionBaseCS
-	{
-		DECLARE_INLINE_TYPE_LAYOUT(FGSConvolutionWithTextureKernelBaseCS, NonVirtual);
-	public:
-
-		FGSConvolutionWithTextureKernelBaseCS() {};
-
-		FGSConvolutionWithTextureKernelBaseCS(const ShaderMetaType::CompiledShaderInitializerType& Initializer)
-			: FGSConvolutionBaseCS(Initializer)
-		{
-			GPUFFTComputeShaderUtils::FComputeParameterBinder Binder(Initializer.ParameterMap);
-
-			Binder(FilterSrcROTexture, TEXT("FilterTexture"));
-		}
-
-		// Used by IMPLEMENT_SHADER_TYPE2
-		static const TCHAR* GetSourceFilename() { return TEXT("/Engine/Private/GPUFastFourierTransform.usf"); }
-		static const TCHAR* GetFunctionName()   { return TEXT("GSConvolutionWithTextureCS"); }
-	
-
-
-		// @todo - template this on KernelType.  Have KernelType know how to do its own SetShaderValue
-		// Note: We pass in the 2d spectral size (transform size) to help the filter
-		void SetCSParamters(FRHICommandList& RHICmdList,
-			const FFT_XFORM_TYPE& XFormType,
-			const FTextureRHIRef& FilterSrcTexture,
-			const FTextureRHIRef& SrcTexture, const FIntRect& SrcRect, const FIntPoint& DstExtentValue)
-		{
-
-			FGSConvolutionBaseCS::SetCSParamters(RHICmdList, XFormType, SrcTexture, SrcRect, DstExtentValue);
-
-			// additional source input for sampling the spectral texture
-			FRHIComputeShader* ShaderRHI = RHICmdList.GetBoundComputeShader();
-			GPUFFTComputeShaderUtils::FComputeParamterValueSetter ParamSetter(RHICmdList, ShaderRHI);
-			// set the texture
-			ParamSetter(FilterSrcROTexture, FilterSrcTexture);
-
-		}
-
-	public:
-		LAYOUT_FIELD(FShaderResourceParameter, FilterSrcROTexture);
-	};
-
-
-	template <int PowRadixSignalLength>
-	class TGSConvolutionWithTexturerCS : public FGSConvolutionWithTextureKernelBaseCS
-	{
-
-		// NB:  the following is actually "public:" 
-		// due to text in the DECLARE_SHADER_TYPE Macro 
-		DECLARE_SHADER_TYPE(TGSConvolutionWithTexturerCS, Global);
-
-		static void ModifyCompilationEnvironment(const FGlobalShaderPermutationParameters& Parameters, FShaderCompilerEnvironment& OutEnvironment)
-		{
-			FGlobalShader::ModifyCompilationEnvironment(Parameters, OutEnvironment);
-			OutEnvironment.SetDefine(TEXT("INCLUDE_GROUP_SHARED_CONVOLUTION_WITH_TEXTURE"), 1);
-			OutEnvironment.SetDefine(TEXT("SCAN_LINE_LENGTH"), PowRadixSignalLength);
-			if (PowRadixSignalLength > 8)
-			{
-				OutEnvironment.SetDefine(TEXT("MIXED_RADIX"), 1);
-			}
-		}
-
-		/** Default constructor **/
-		TGSConvolutionWithTexturerCS() {};
-
-	public:
-
-		TGSConvolutionWithTexturerCS(const ShaderMetaType::CompiledShaderInitializerType& Initializer)
-			: FGSConvolutionWithTextureKernelBaseCS(Initializer)
-		{}
-	};
-
-
-	void SwapContents(FSceneRenderTargetItem& TmpBuffer, FSceneRenderTargetItem& DstBuffer)
+	void SwapContents(FRDGTextureRef& TmpBuffer, FRDGTextureRef& DstBuffer)
 	{
 		// Swap the pointers
-		FSceneRenderTargetItem TmpDst = DstBuffer;
+		FRDGTextureRef TmpDst = DstBuffer;
 		DstBuffer = TmpBuffer;
 		TmpBuffer = TmpDst;
 	}
-} // end namespace
-
-IMPLEMENT_SHADER_TYPE3(GPUFFT::FReorderFFTPassCS, SF_Compute);
-IMPLEMENT_SHADER_TYPE3(GPUFFT::FGroupShardSubFFTPassCS, SF_Compute);
-IMPLEMENT_SHADER_TYPE3(GPUFFT::FComplexFFTPassCS, SF_Compute);
-IMPLEMENT_SHADER_TYPE3(GPUFFT::FPackTwoForOneFFTPassCS, SF_Compute);
-IMPLEMENT_SHADER_TYPE3(GPUFFT::FCopyWindowCS, SF_Compute);
-IMPLEMENT_SHADER_TYPE3(GPUFFT::FComplexMultiplyImagesCS, SF_Compute);
-
-#define GROUPSHARED_COMPLEX_TRANSFORM(_Length) \
-typedef GPUFFT::TGSComplexTransformCS< _Length> FGSComplexTransformCSLength##_Length; \
-IMPLEMENT_SHADER_TYPE2(FGSComplexTransformCSLength##_Length, SF_Compute);
-
-GROUPSHARED_COMPLEX_TRANSFORM(2)  GROUPSHARED_COMPLEX_TRANSFORM(16)  GROUPSHARED_COMPLEX_TRANSFORM(128)   GROUPSHARED_COMPLEX_TRANSFORM(1024)
-
-GROUPSHARED_COMPLEX_TRANSFORM(4)  GROUPSHARED_COMPLEX_TRANSFORM(32)  GROUPSHARED_COMPLEX_TRANSFORM(256)   GROUPSHARED_COMPLEX_TRANSFORM(2048)  
-
-GROUPSHARED_COMPLEX_TRANSFORM(8)  GROUPSHARED_COMPLEX_TRANSFORM(64)  GROUPSHARED_COMPLEX_TRANSFORM(512)   GROUPSHARED_COMPLEX_TRANSFORM(4096)
-// NB: FFTBLOCK(8192, false/true) won't work because the max number of threads in a group 1024 is less than the requested 8192 / 2 
-
-#undef GROUPSHARED_COMPLEX_TRANSFORM
-
-
-
-#define GROUPSHARED_TWO_FOR_ONE_TRANSFORM(_Length) \
-typedef GPUFFT::TGSTwoForOneTransformCS<_Length> FGSTwoForOneTransformCSLength##_Length; \
-IMPLEMENT_SHADER_TYPE2(FGSTwoForOneTransformCSLength##_Length, SF_Compute);
-
-GROUPSHARED_TWO_FOR_ONE_TRANSFORM(2)  GROUPSHARED_TWO_FOR_ONE_TRANSFORM(16)  GROUPSHARED_TWO_FOR_ONE_TRANSFORM(128)   GROUPSHARED_TWO_FOR_ONE_TRANSFORM(1024)
-
-GROUPSHARED_TWO_FOR_ONE_TRANSFORM(4)  GROUPSHARED_TWO_FOR_ONE_TRANSFORM(32)  GROUPSHARED_TWO_FOR_ONE_TRANSFORM(256)   GROUPSHARED_TWO_FOR_ONE_TRANSFORM(2048)  
- 
-GROUPSHARED_TWO_FOR_ONE_TRANSFORM(8)  GROUPSHARED_TWO_FOR_ONE_TRANSFORM(64)  GROUPSHARED_TWO_FOR_ONE_TRANSFORM(512)   GROUPSHARED_TWO_FOR_ONE_TRANSFORM(4096)  
-
-// NB: FFTBLOCK(8192, false/true) won't work because the max number of threads in a group 1024 is less than the requested 8192 / 2 
-
-#undef GROUPSHARED_TWO_FOR_ONE_TRANSFORM
-
-
-#define GROUPSHARED_CONVOLUTION_WTEXTURE(_Length) \
-typedef GPUFFT::TGSConvolutionWithTexturerCS<_Length >  FGSConvolutionWithTextureCSLength##_Length; \
-IMPLEMENT_SHADER_TYPE2(FGSConvolutionWithTextureCSLength##_Length,  SF_Compute); 
-
-GROUPSHARED_CONVOLUTION_WTEXTURE(2)  GROUPSHARED_CONVOLUTION_WTEXTURE(16)  GROUPSHARED_CONVOLUTION_WTEXTURE(128)   GROUPSHARED_CONVOLUTION_WTEXTURE(1024)
-
-GROUPSHARED_CONVOLUTION_WTEXTURE(4)  GROUPSHARED_CONVOLUTION_WTEXTURE(32)  GROUPSHARED_CONVOLUTION_WTEXTURE(256)   GROUPSHARED_CONVOLUTION_WTEXTURE(2048)
-
-GROUPSHARED_CONVOLUTION_WTEXTURE(8)  GROUPSHARED_CONVOLUTION_WTEXTURE(64)  GROUPSHARED_CONVOLUTION_WTEXTURE(512)   GROUPSHARED_CONVOLUTION_WTEXTURE(4096)
-//GROUPSHARED_CONVOLUTION_WTEXTURE(2, 8192)
-// NB: FFTBLOCK(2, 4196, false/true) won't work because the max number of threads in a group 1024 is less than the requested 4196 / 2 
-
-
-#undef GROUPSHARED_CONVOLUTION_WTEXTURE
-
+}
 
-namespace 
+namespace
 {
-	using namespace GPUFFT;
-
-	TShaderRef<FCopyWindowCS> GetCopyWindowCS(const FGPUFFTShaderContext::ShaderMapType& ShaderMap)
-	{
-		return  ShaderMap.GetShader<FCopyWindowCS>();
-	}
-
-	TShaderRef<FComplexMultiplyImagesCS> GetComplexMultiplyImagesCS(const FGPUFFTShaderContext::ShaderMapType& ShaderMap)
-	{
-		return  ShaderMap.GetShader<FComplexMultiplyImagesCS>();
-	}
 
-	TShaderRef<FGroupShardSubFFTPassCS> GetGroupSharedSubFFTPassCS(const FGPUFFTShaderContext::ShaderMapType& ShaderMap, const uint32 TransformLength)
-	{
-		return ShaderMap.GetShader<FGroupShardSubFFTPassCS>();
-	}
+class FPowRadixSignalLengthDim : SHADER_PERMUTATION_SPARSE_INT("SCAN_LINE_LENGTH", 2, 4, 8, 16, 32, 64, 128, 256, 512, 1024, 2048, 4096);
 
-	TShaderRef<FReorderFFTPassCS> GetReorderFFTPassCS(const FGPUFFTShaderContext::ShaderMapType& ShaderMap)
-	{
-		return ShaderMap.GetShader<FReorderFFTPassCS>();
-	}
-	TShaderRef<FPackTwoForOneFFTPassCS> GetPackTwoForOneFFTPassCS(const FGPUFFTShaderContext::ShaderMapType& ShaderMap, const uint32 TransformLength)
-	{
-		return  ShaderMap.GetShader<FPackTwoForOneFFTPassCS>();
-	}
-
-	TShaderRef<FComplexFFTPassCS> GetComplexFFTPassCS(const FGPUFFTShaderContext::ShaderMapType& ShaderMap, const uint32 TransformLength)
-	{
-		return ShaderMap.GetShader<FComplexFFTPassCS>();
-	}
-
-	// Shader Permutation picker.
-	TShaderRef<FGSComplexTransformBaseCS> GetComplexFFTCS(const FGPUFFTShaderContext::ShaderMapType& ShaderMap, const uint32 TransformLength)
+class FFFTShader : public FGlobalShader
+{
+public:
+	static bool ShouldCompilePermutation(const FGlobalShaderPermutationParameters& Parameters)
 	{
-		TShaderRef<FGSComplexTransformBaseCS> Result;
-
-	#define GET_COMPLEX_SHADER(_LENGTH)  ShaderMap.GetShader<TGSComplexTransformCS<_LENGTH>>();
-
-		switch (TransformLength)
-		{
-		case 2:     Result = GET_COMPLEX_SHADER(2);    break;
-		case 4:     Result = GET_COMPLEX_SHADER(4);    break;
-		case 8:     Result = GET_COMPLEX_SHADER(8);    break;
-		case 16:    Result = GET_COMPLEX_SHADER(16);   break;
-		case 32:    Result = GET_COMPLEX_SHADER(32);   break;
-		case 64:    Result = GET_COMPLEX_SHADER(64);   break;
-		case 128:   Result = GET_COMPLEX_SHADER(128);  break;
-		case 256:   Result = GET_COMPLEX_SHADER(256);  break;
-		case 512:   Result = GET_COMPLEX_SHADER(512);  break;
-		case 1024:  Result = GET_COMPLEX_SHADER(1024); break;
-		case 2048:  Result = GET_COMPLEX_SHADER(2048); break;
-		case 4096:  Result = GET_COMPLEX_SHADER(4096); break;
-		//case 8192:  Result = GET_COMPLEX_SHADER(2, 8192); break;
-
-		default:
-			ensureMsgf(0, TEXT("The FFT block height is not supported"));
-			break;
-		}
-
-	#undef GET_COMPLEX_SHADER
-
-		return Result;
+		return ShouldCompileGPUFFT(Parameters.Platform);
 	}
 
-	// Shader Permutation picker.
-	TShaderRef<FGSComplexTransformBaseCS> GetTwoForOneFFTCS(const FGPUFFTShaderContext::ShaderMapType& ShaderMap, const uint32 TransformLength)
-	{
-		TShaderRef<FGSComplexTransformBaseCS> Result;
-
-	#define GET_TWOFORONE_SHADER(_LENGTH)  ShaderMap.GetShader<TGSTwoForOneTransformCS<_LENGTH>>();
-
-		switch (TransformLength)
-		{
-		case 2:     Result = GET_TWOFORONE_SHADER(2);    break;
-		case 4:     Result = GET_TWOFORONE_SHADER(4);    break;
-		case 8:     Result = GET_TWOFORONE_SHADER(8);    break;
-		case 16:    Result = GET_TWOFORONE_SHADER(16);   break;
-		case 32:    Result = GET_TWOFORONE_SHADER(32);   break;
-		case 64:    Result = GET_TWOFORONE_SHADER(64);   break;
-		case 128:   Result = GET_TWOFORONE_SHADER(128);  break;
-		case 256:   Result = GET_TWOFORONE_SHADER(256);  break;
-		case 512:   Result = GET_TWOFORONE_SHADER(512);  break;
-		case 1024:  Result = GET_TWOFORONE_SHADER(1024); break;
-		case 2048:  Result = GET_TWOFORONE_SHADER(2048); break;
-		case 4096:  Result = GET_TWOFORONE_SHADER(4096); break;
-		//case 8192:  Result = GET_TWOFORONE_SHADER(2, 8192); break;
-
-		default:
-			ensureMsgf(0, TEXT("The FFT block height is not supported"));
-			break;
-		}
+	FFFTShader() = default;
+	FFFTShader(const CompiledShaderInitializerType& Initializer)
+		: FGlobalShader(Initializer)
+	{}
+};
 
-	#undef GET_TWOFORONE_SHADER
+class FReorderFFTPassCS : public FFFTShader
+{
+	DECLARE_GLOBAL_SHADER(FReorderFFTPassCS);
+	SHADER_USE_PARAMETER_STRUCT(FReorderFFTPassCS, FFFTShader);
 
-		return Result;
-	}
+	BEGIN_SHADER_PARAMETER_STRUCT(FParameters, )
+		SHADER_PARAMETER(FIntRect, SrcRect)
+		SHADER_PARAMETER(FIntRect, DstRect)
+		SHADER_PARAMETER(int32, TransformType)
+		SHADER_PARAMETER(int32, LogTwoLength)
+		SHADER_PARAMETER(int32, BitCount)
 
+		SHADER_PARAMETER_RDG_TEXTURE(Texture2D, SrcSRV)
+		SHADER_PARAMETER_RDG_TEXTURE_UAV(RWTexture2D, DstUAV)
+	END_SHADER_PARAMETER_STRUCT()
 
-	TShaderRef<FGSConvolutionWithTextureKernelBaseCS> GetConvolutionWithTextureKernelCS(const FGPUFFTShaderContext::ShaderMapType& ShaderMap, const uint32 TransformLength)
+	static void ModifyCompilationEnvironment(const FGlobalShaderPermutationParameters& Parameters, FShaderCompilerEnvironment& OutEnvironment)
 	{
-		TShaderRef<FGSConvolutionWithTextureKernelBaseCS> Result;
-
-		// Get the spectral filter.
-	#define GET_GROUP_SHARED_TEXTURE_FILTER(_LENGTH)  ShaderMap.GetShader< TGSConvolutionWithTexturerCS<_LENGTH> >();
-		switch (TransformLength)
-		{
-		case 2:    Result = GET_GROUP_SHARED_TEXTURE_FILTER(2);    break;
-		case 4:    Result = GET_GROUP_SHARED_TEXTURE_FILTER(4);    break;
-		case 8:    Result = GET_GROUP_SHARED_TEXTURE_FILTER(8);    break;
-		case 16:   Result = GET_GROUP_SHARED_TEXTURE_FILTER(16);   break;
-		case 32:   Result = GET_GROUP_SHARED_TEXTURE_FILTER(32);   break;
-		case 64:   Result = GET_GROUP_SHARED_TEXTURE_FILTER(64);   break;
-		case 128:  Result = GET_GROUP_SHARED_TEXTURE_FILTER(128);  break;
-		case 256:  Result = GET_GROUP_SHARED_TEXTURE_FILTER(256);  break;
-		case 512:  Result = GET_GROUP_SHARED_TEXTURE_FILTER(512);  break;
-		case 1024: Result = GET_GROUP_SHARED_TEXTURE_FILTER(1024); break;
-		case 2048: Result = GET_GROUP_SHARED_TEXTURE_FILTER(2048); break;
-		case 4096: Result = GET_GROUP_SHARED_TEXTURE_FILTER(4096); break;
-		//case 8192: Result = GET_GROUP_SHARED_TEXTURE_FILTER(2, 8192); break;
-		default:
-			ensureMsgf(0, TEXT("The FFT block height is not supported"));
-			break;
-		}
-	#undef GET_GROUP_SHARED_TEXTURE_FILTER
-
-		return Result;
-
+		FGlobalShader::ModifyCompilationEnvironment(Parameters, OutEnvironment);
+		OutEnvironment.SetDefine(TEXT("INCLUDE_REORDER_FFT_PASS"), 1);
 	}
-	/**
-	* Single Pass that Copies a sub-region of a buffer and potentially 
-	* boosts the intensity of select pixels.
-	* 
-	* Assumes the dst buffer is large enough to hold the result.
-	* The Src float4 data is interpreted as a pair of independent complex numbers.
-	* The Knl float4 data is interpreted as a pair of independent complex numbers
-	*
-	* @param Context    - container for RHI and ShanderMap
-	* @param SrcWindow  - The region of interest to copy.
-	* @param SrcTexture - SRV the source image buffer.
-	* @param DstWindow  - The target location for the images.
-	*
-	* @param DstUAV     - UAV, the destination buffer that will hold the result of the single pass
-	* @param Prefilter  - Optional filter to boost selected pixels.
-	*/
-	void DispatchCopyWindowCS(FGPUFFTShaderContext& Context,
-		const FIntRect& SrcWindow, const FTextureRHIRef& SrcTexture,
-		const FIntRect& DstWindow, FUnorderedAccessViewRHIRef DstUAV,
-		const FPreFilter& PreFilter = FPreFilter(TNumericLimits<float>::Max(), TNumericLimits<float>::Lowest(), 0.f))
-	{
-		using namespace GPUFFTComputeShaderUtils;
-		// The size of the dst
-		const FIntPoint DstExtent = DstWindow.Size();
-
-		const uint32 XThreadCount = FCopyWindowCS::XThreadCount();
-		const uint32 YThreadCount = FCopyWindowCS::YThreadCount();
-
-		// Number of thread groups
-		const uint32 XGroups = ( DstExtent.X / XThreadCount ) + ( (DstExtent.X % XThreadCount == 0) ? 0 : 1 ); //-V::547, 1063
-		const uint32 YGroups = ( DstExtent.Y / YThreadCount ) + ( (DstExtent.Y % YThreadCount == 0) ? 0 : 1 ); //-V547
+};
 
-		const FGPUFFTShaderContext::ShaderMapType& ShaderMap = Context.GetShaderMap();
-		FRHICommandList& RHICmdList = Context.GetRHICmdList();
 
-		SCOPED_DRAW_EVENTF(RHICmdList, CopyWindowCS, TEXT("FFT Multipass: Copy Subwindow"));
-
-		// Get pointer to the shader
-		TShaderRef<FCopyWindowCS> ComputeShader = GetCopyWindowCS(ShaderMap);
-		RHICmdList.SetComputeShader(ComputeShader.GetComputeShader());
+class FGroupShardSubFFTPassCS : public FFFTShader
+{
+	DECLARE_GLOBAL_SHADER(FGroupShardSubFFTPassCS);
+	SHADER_USE_PARAMETER_STRUCT(FGroupShardSubFFTPassCS, FFFTShader);
 
-		// Bind output
-		FScopedUAVBind ScopedBind = FScopedUAVBind::BindOutput(RHICmdList, ComputeShader, DstUAV);
+	BEGIN_SHADER_PARAMETER_STRUCT(FParameters, )
+		SHADER_PARAMETER(FIntRect, SrcWindow)
+		SHADER_PARAMETER(int32, TransformType)
+		SHADER_PARAMETER(int32, TransformLength)
+		SHADER_PARAMETER(int32, NumSubRegions)
 
-		ComputeShader->SetCSParamters(RHICmdList, SrcWindow, SrcTexture,  DstWindow, PreFilter);
+		SHADER_PARAMETER_RDG_TEXTURE(Texture2D, SrcTexture)
+		SHADER_PARAMETER_RDG_TEXTURE_UAV(RWTexture2D, DstTexture)
+	END_SHADER_PARAMETER_STRUCT()
 
-		// Dispatch with a single thread-group for each "column" in the result where the transform direction is the "row" direction.
-		RHICmdList.DispatchComputeShader(XGroups, YGroups, 1);
-	}
+	static uint32 SubPassLength() { return 2048; }
+	static uint32 Radix() { return 2; }
 
-	/**
-	* 
-	* Single Pass that computes the Frequency Space convolution of two buffers
-	* that have already been transformed into frequency space.
-	*
-	* This is really means the complex product of two buffers divided
-	* by the correct values to "normalize" the effect of the kernel buffer.
-	*
-	* In this case, each float4 is viewed as a pair of complex numbers,
-	* and the product of float4 Src, Knl is computed
-	* as float4(ComplexMult(Src.xy, Knl.xy) / Na, ComplexMult(Src.zw, Knl.zw)/Nb)
-	* where Na and Nb are related to the sums of the weights in the kernel buffer. 
-	*
-	* Assumes the dst buffer is large enough to hold the result.
-	* The Src float4 data is interpreted as a pair of independent complex numbers.
-	* The Knl float4 data is interpreted as a pair of independent complex numbers  
-	*
-	* @param Context    - container for RHI and ShanderMap
-	* @param bHorizontalFirst - Describes the layout of the transformed data.
-	*                           bHorizontalFirst == true for data that was transformed as 
-	*                           Vertical::ComplexFFT following Horizontal::TwoForOneRealFFT
-	*                           bHorizontalFirst == false for data that was transformed as
-	*                           Horizontal::ComplexFFT following Vertical::TwoForOneRealFFT
-	* @param SrcTexture - SRV the source image buffer.
-	* @param KnlTexture - SRV the kernel image buffer.
-	*
-	* @param DstUAV     - UAV, the destination buffer that will hold the result of the single pass
-	*
-	*/
-
-	void DispatchComplexMultiplyImagesCS(FGPUFFTShaderContext& Context,
-		const bool bHorizontalScanlines,
-		const FIntRect& SrcWindow, const FTextureRHIRef& SrcTexture,
-		const FTextureRHIRef& KnlTexure,
-		FUnorderedAccessViewRHIRef DstUAV)
+	static void ModifyCompilationEnvironment(const FGlobalShaderPermutationParameters& Parameters, FShaderCompilerEnvironment& OutEnvironment)
 	{
-		using namespace GPUFFTComputeShaderUtils;
-		// The size of the dst
-		const FIntPoint DstExtent = SrcWindow.Size();
-
-		const FGPUFFTShaderContext::ShaderMapType& ShaderMap = Context.GetShaderMap();
-		FRHICommandList& RHICmdList = Context.GetRHICmdList();
-
-		SCOPED_DRAW_EVENTF(RHICmdList, ComplexMultiplyImagesCS, TEXT("FFT Multipass: Convolution in freq-space"));
-
-		// Get pointer to the shader
-		TShaderRef<FComplexMultiplyImagesCS> ComputeShader = GetComplexMultiplyImagesCS(ShaderMap);
-		RHICmdList.SetComputeShader(ComputeShader.GetComputeShader());
-
-		// Bind output
-		FScopedUAVBind ScopedBind = FScopedUAVBind::BindOutput(RHICmdList, ComputeShader, DstUAV);
-
-		ComputeShader->SetCSParamters(RHICmdList, bHorizontalScanlines, SrcWindow, SrcTexture, KnlTexure);
-		
-		// Align the scanlines in the direction of the first transform.
-		const uint32 NumScanLines = (bHorizontalScanlines)  ? DstExtent.Y : DstExtent.X;
-		const uint32 SignalLength = (!bHorizontalScanlines) ? DstExtent.Y : DstExtent.X;
-	
-		// Dispatch with a single thread-group for each "column."
-		// Configured this way because the entire column will share the same normalization values (retrieved  from the kernel)
-		//RHICmdList.DispatchComputeShader(1, ThreadGroupsPerScanline, NumScanLines);
-		RHICmdList.DispatchComputeShader(1, 1, NumScanLines);
+		FGlobalShader::ModifyCompilationEnvironment(Parameters, OutEnvironment);
+		OutEnvironment.SetDefine(TEXT("INCLUDE_GROUP_SHARED_SUB_COMPLEX_FFT"), 1);
+		OutEnvironment.SetDefine(TEXT("SCAN_LINE_LENGTH"), FGroupShardSubFFTPassCS::SubPassLength());
+		OutEnvironment.SetDefine(TEXT("RADIX"), FGroupShardSubFFTPassCS::Radix());
 	}
+};
 
+class FComplexFFTPassCS : public FFFTShader
+{
+	DECLARE_GLOBAL_SHADER(FComplexFFTPassCS);
+	SHADER_USE_PARAMETER_STRUCT(FComplexFFTPassCS, FFFTShader);
 
-	/**
-	* Single Pass that separates or merges the transform of four real signals from 
-	* viewed as the transform of two complex  signals.  
-	*
-	* Assumes the dst buffer is large enough to hold the result.
-	* The Src float4 data is interprets  as a pair of independent complex numbers.
-	*
-	* @param Context    - container for RHI and ShanderMap
-	* @param FFTDesc    - Metadata that describes the underlying complex FFT
-	* @param SrcTexture - SRV the source image buffer
-	* @param DstUAV     - UAV, the destination buffer that will hold the result of the single pass
-	*
-	* The data in the SrcTexture and DstUAV is aligned with (0,0).
-	* FFTDesc.IsHorizontal() indicates the transform direction in the buffer that needs to be spit/merged
-	* FFTDesc.IsForward() indicates spit (true) vs merge (false).
-	*/
-	void DispatchPackTwoForOneFFTPassCS(FGPUFFTShaderContext& Context, const FFTDescription& FFTDesc,
-		const FTextureRHIRef& SrcTexture, FUnorderedAccessViewRHIRef DstUAV)
-	{
-
-		using namespace GPUFFTComputeShaderUtils;
-
-		const FGPUFFTShaderContext::ShaderMapType& ShaderMap = Context.GetShaderMap();
-		FRHICommandList& RHICmdList = Context.GetRHICmdList();
-	
-		const uint32 TransformLength = FFTDesc.SignalLength;
-		const FString TransformName = FFTDesc.FFT_TypeName();
-
-		// A real signal of length 'TransformLenght' requires  only TransformLength / 2 + 1  complex coefficients, 
-		const uint32 RealTransformLength = ( TransformLength / 2 ) + 1;
+	BEGIN_SHADER_PARAMETER_STRUCT(FParameters, )
+		SHADER_PARAMETER(FIntRect, SrcRect)
+		SHADER_PARAMETER(FIntRect, DstRect)
+		SHADER_PARAMETER(int32, TransformType)
+		SHADER_PARAMETER(int32, PowTwoLength)
+		SHADER_PARAMETER(int32, BitCount)
 
-		// The splitting into two real signals (isForward) or joinging back into a single signal
-		const uint32 ResultingLength = (FFTDesc.IsForward()) ? 2 * RealTransformLength : TransformLength;
+		SHADER_PARAMETER_RDG_TEXTURE(Texture2D, SrcSRV)
+		SHADER_PARAMETER_RDG_TEXTURE_UAV(RWTexture2D, DstUAV)
+	END_SHADER_PARAMETER_STRUCT()
 
-		SCOPED_DRAW_EVENTF(RHICmdList, PackTwoForOneFFTPass, TEXT("FFT Multipass: TwoForOne Combine/split result of %s of size %d"), *TransformName, TransformLength);
+	static void ModifyCompilationEnvironment(const FGlobalShaderPermutationParameters& Parameters, FShaderCompilerEnvironment& OutEnvironment)
+	{
+		FGlobalShader::ModifyCompilationEnvironment(Parameters, OutEnvironment);
+		OutEnvironment.SetDefine(TEXT("INCLUDE_COMPLEX_FFT_PASS"), 1);
+	}
+};
 
-		FIntPoint DstExtent = FFTDesc.TransformExtent();
-		if (FFTDesc.IsHorizontal())
-		{
-			DstExtent.X = ResultingLength;
-		}	
-		else
-		{
-			DstExtent.Y = ResultingLength;
-		}
 
-		// Get pointer to the shader
-		TShaderRef<FPackTwoForOneFFTPassCS> ComputeShader = GetPackTwoForOneFFTPassCS(ShaderMap, TransformLength);
-		RHICmdList.SetComputeShader(ComputeShader.GetComputeShader());
+class FPackTwoForOneFFTPassCS : public FFFTShader
+{
+	DECLARE_GLOBAL_SHADER(FPackTwoForOneFFTPassCS);
+	SHADER_USE_PARAMETER_STRUCT(FPackTwoForOneFFTPassCS, FFFTShader);
 
-		// Bind output
-		FScopedUAVBind ScopedBind = FScopedUAVBind::BindOutput(RHICmdList, ComputeShader, DstUAV);
+	BEGIN_SHADER_PARAMETER_STRUCT(FParameters, )
+		SHADER_PARAMETER(FIntRect, DstRect)
+		SHADER_PARAMETER(int32, TransformType)
 
-		ComputeShader->SetCSParamters(RHICmdList, FFTDesc.XFormType, SrcTexture, FIntRect(FIntPoint(0, 0), DstExtent));
+		SHADER_PARAMETER_RDG_TEXTURE(Texture2D, SrcSRV)
+		SHADER_PARAMETER_RDG_TEXTURE_UAV(RWTexture2D, DstUAV)
+	END_SHADER_PARAMETER_STRUCT()
 
-		// Dispatch with a single thread-group for each "column" in the result where the transform direction is the "row" direction.
-		RHICmdList.DispatchComputeShader(1, 1, RealTransformLength);
+	static void ModifyCompilationEnvironment(const FGlobalShaderPermutationParameters& Parameters, FShaderCompilerEnvironment& OutEnvironment)
+	{
+		FGlobalShader::ModifyCompilationEnvironment(Parameters, OutEnvironment);
+		OutEnvironment.SetDefine(TEXT("INCLUDE_PACK_TWOFORONE_FFT_PASS"), 1);
 	}
+};
 
-	/**
-	* Single Pass of a multi-pass complex FFT.
-	* Assumes the dst buffer is large enough to hold the result.
-	* The Src float4 data is interpt as a pair of independent complex numbers.
-	*
-	* @param Context    - container for RHI and ShanderMap
-	* @param FFTDesc    - Metadata that describes the underlying complex FFT
-	* @param PassStage  - The Depth at which this FFT pass lives. 
-	* @param SrcTexture - SRV the source image buffer
-	* @param SrcRct     - The region in the Src buffer where the image to transform lives.
-	* @param DstUAV     - UAV, the destination buffer that will hold the result of the single pass
-	*
-	*/
-	void DispatchComplexFFTPassCS(FGPUFFTShaderContext& Context, const FFTDescription& FFTDesc,
-		const uint32 PassLength,
-		const FTextureRHIRef& SrcTexture, const FIntRect& SrcWindow,
-		FUnorderedAccessViewRHIRef DstUAV,
-		const bool bScrubNaNs = false)
-	{
-		// Using multiple radix two passes.
-		const uint32 Radix = 2;
-		using namespace GPUFFTComputeShaderUtils;
+class FCopyWindowCS : public FFFTShader
+{
+	DECLARE_GLOBAL_SHADER(FCopyWindowCS);
+	SHADER_USE_PARAMETER_STRUCT(FCopyWindowCS, FFFTShader);
 
-		const FGPUFFTShaderContext::ShaderMapType& ShaderMap = Context.GetShaderMap();
-		FRHICommandList& RHICmdList = Context.GetRHICmdList();
+	BEGIN_SHADER_PARAMETER_STRUCT(FParameters, )
+		SHADER_PARAMETER(FIntRect, SrcRect)
+		SHADER_PARAMETER(FIntRect, DstRect)
+		SHADER_PARAMETER(FVector3f, BrightPixelGain)
 
-		const uint32 TransformLength = FFTDesc.SignalLength;
-		const FString TransformName = FFTDesc.FFT_TypeName();
+		SHADER_PARAMETER_RDG_TEXTURE(Texture2D, SrcSRV)
+		SHADER_PARAMETER_RDG_TEXTURE_UAV(RWTexture2D, DstUAV)
+	END_SHADER_PARAMETER_STRUCT()
 
+	static FIntPoint ThreadCount() { return FIntPoint(1, 32); }
 
-		SCOPED_DRAW_EVENTF(RHICmdList, ComplexFFTPass, TEXT("FFT Multipass: Pass %d of Complex %s of size %d"), PassLength, *TransformName, TransformLength);
-		
-		FIntPoint DstExtent = FFTDesc.TransformExtent();
-		
-		// Get a base pointer to the shader
-		TShaderRef<FComplexFFTPassCS> ComputeShader = GetComplexFFTPassCS(ShaderMap, TransformLength);
-		RHICmdList.SetComputeShader(ComputeShader.GetComputeShader());
+	static void ModifyCompilationEnvironment(const FGlobalShaderPermutationParameters& Parameters, FShaderCompilerEnvironment& OutEnvironment)
+	{
+		FFFTShader::ModifyCompilationEnvironment(Parameters, OutEnvironment);
+		OutEnvironment.SetDefine(TEXT("INCLUDE_COPY_WINDOW"), 1);
+		OutEnvironment.SetDefine(TEXT("X_THREAD_COUNT"), ThreadCount().X);
+		OutEnvironment.SetDefine(TEXT("Y_THREAD_COUNT"), ThreadCount().Y);
+	}
+};
 
-		// Bind output
-		FScopedUAVBind ScopedBind = FScopedUAVBind::BindOutput(RHICmdList, ComputeShader, DstUAV);
+class FComplexMultiplyImagesCS : public FFFTShader
+{
+	DECLARE_GLOBAL_SHADER(FComplexMultiplyImagesCS);
+	SHADER_USE_PARAMETER_STRUCT(FComplexMultiplyImagesCS, FFFTShader);
 
-		ComputeShader->SetCSParamters(RHICmdList, FFTDesc.XFormType, SrcTexture, SrcWindow, FIntRect(FIntPoint(0, 0), DstExtent), TransformLength, PassLength, bScrubNaNs);
+	BEGIN_SHADER_PARAMETER_STRUCT(FParameters, )
+		SHADER_PARAMETER(FIntRect, SrcRect)
+		SHADER_PARAMETER(int32, DataLayout)
 
-		RHICmdList.DispatchComputeShader(1, 1, TransformLength / Radix);
-	}
+		SHADER_PARAMETER_RDG_TEXTURE(Texture2D, SrcSRV)
+		SHADER_PARAMETER_RDG_TEXTURE(Texture2D, KnlSRV)
+		SHADER_PARAMETER_RDG_TEXTURE_UAV(RWTexture2D, DstUAV)
+	END_SHADER_PARAMETER_STRUCT()
 
-	/**
-	* Single Pass of a multi-pass complex FFT that just reorders data for
-	* a group shared subpass to consume.
-	*
-	* Assumes the dst buffer is large enough to hold the result.
-	* The Src float4 data is interpt as a pair of independent complex numbers.
-	*
-	* @param Context    - container for RHI and ShanderMap
-	* @param FFTDesc    - Metadata that describes the underlying complex FFT
-	* @param SrcRct     - The region in the Src buffer where the image to transform lives.
-	* @param SrcTexture - SRV the source image buffer
-	* @param DstRect    - The target Window.
-	* @param DstUAV     - UAV, the destination buffer that will hold the result of the single pass
-	*
-	*/
-	void DispatchReorderFFTPassCS(FGPUFFTShaderContext& Context, const FFTDescription& FFTDesc,
-		const FIntRect& SrcWindow, const FTextureRHIRef& SrcTexture,
-		const FIntRect& DstWindow, FUnorderedAccessViewRHIRef DstUAV,
-		const bool bScrubNaNs = false)
+	static void ModifyCompilationEnvironment(const FGlobalShaderPermutationParameters& Parameters, FShaderCompilerEnvironment& OutEnvironment)
 	{
-		// Using multiple radix two passes.
-		const uint32 Radix = 2;
-		using namespace GPUFFTComputeShaderUtils;
+		FGlobalShader::ModifyCompilationEnvironment(Parameters, OutEnvironment);
+		OutEnvironment.SetDefine(TEXT("INCLUDE_COMPLEX_MULTIPLY_IMAGES"), 1);
+	}
+};
 
-		const FGPUFFTShaderContext::ShaderMapType& ShaderMap = Context.GetShaderMap();
-		FRHICommandList& RHICmdList = Context.GetRHICmdList();
+class FGSComplexTransformCS : public FFFTShader
+{
+	DECLARE_GLOBAL_SHADER(FGSComplexTransformCS);
+	SHADER_USE_PARAMETER_STRUCT(FGSComplexTransformCS, FFFTShader);
 
-		const uint32 TransformLength = FFTDesc.SignalLength;
-		const FString TransformName = FFTDesc.FFT_TypeName();
+	using FPermutationDomain = TShaderPermutationDomain<FPowRadixSignalLengthDim>;
 
-		// Breaks the data into the correct number of sub-transforms for the later group-shared pass.
-		const uint32 SubLength = TransformLength / FGroupShardSubFFTPassCS::SubPassLength();
+	BEGIN_SHADER_PARAMETER_STRUCT(FParameters, )
+		SHADER_PARAMETER(FIntPoint, SrcRectMin)
+		SHADER_PARAMETER(FIntPoint, SrcRectMax)
+		SHADER_PARAMETER(FIntPoint, DstExtent)
+		SHADER_PARAMETER(FIntRect, DstRect)
+		SHADER_PARAMETER(FVector3f, BrightPixelGain)
+		SHADER_PARAMETER(int32, TransformType)
 
-		SCOPED_DRAW_EVENTF(RHICmdList, ReorderFFTPass, TEXT("FFT Multipass: Complex %s Reorder pass of size %d"), *TransformName, TransformLength);
+		SHADER_PARAMETER_RDG_TEXTURE(Texture2D, SrcTexture)
+		SHADER_PARAMETER_RDG_TEXTURE_UAV(RWTexture2D, DstTexture)
+	END_SHADER_PARAMETER_STRUCT()
 
-		FIntPoint DstExtent = FFTDesc.TransformExtent();
+	static void ModifyCompilationEnvironment(const FGlobalShaderPermutationParameters& Parameters, FShaderCompilerEnvironment& OutEnvironment)
+	{
+		FPermutationDomain PermutationVector(Parameters.PermutationId);
+		int32 PowRadixSignalLength = PermutationVector.Get<FPowRadixSignalLengthDim>();
 
-		// Get a base pointer to the shader
-		TShaderRef<FReorderFFTPassCS> ComputeShader = GetReorderFFTPassCS(ShaderMap);
-		RHICmdList.SetComputeShader(ComputeShader.GetComputeShader());
+		FFFTShader::ModifyCompilationEnvironment(Parameters, OutEnvironment);
+		OutEnvironment.SetDefine(TEXT("INCLUDE_GROUP_SHARED_COMPLEX_FFT"), 1);
+		if (PowRadixSignalLength > 8)
+		{
+			OutEnvironment.SetDefine(TEXT("MIXED_RADIX"), 1);
+		}
+	}
+};
 
-		// Bind output
-		FScopedUAVBind ScopedBind = FScopedUAVBind::BindOutput(RHICmdList, ComputeShader, DstUAV);
+class FGSTwoForOneTransformCS : public FFFTShader
+{
+	DECLARE_GLOBAL_SHADER(FGSTwoForOneTransformCS);
+	SHADER_USE_PARAMETER_STRUCT(FGSTwoForOneTransformCS, FFFTShader);
 
-		ComputeShader->SetCSParamters(RHICmdList, FFTDesc.XFormType, SrcTexture, SrcWindow, DstWindow, TransformLength, SubLength, bScrubNaNs);
-			
+	using FPermutationDomain = TShaderPermutationDomain<FPowRadixSignalLengthDim>;
 
-		RHICmdList.DispatchComputeShader(1, 1, TransformLength / Radix);
-	}
+	BEGIN_SHADER_PARAMETER_STRUCT(FParameters, )
+		SHADER_PARAMETER(FIntPoint, SrcRectMin)
+		SHADER_PARAMETER(FIntPoint, SrcRectMax)
+		SHADER_PARAMETER(FIntPoint, DstExtent)
+		SHADER_PARAMETER(FIntRect, DstRect)
+		SHADER_PARAMETER(FVector3f, BrightPixelGain)
+		SHADER_PARAMETER(int32, TransformType)
 
-	/**
-	* A Group Shared Single Pass of a multi-pass complex FFT.
-	* Assumes the dst buffer is large enough to hold the result.
-	* The Src float4 data is interpt as a pair of independent complex numbers.
-	*
-	* @param Context    - container for RHI and ShanderMap
-	* @param FFTDesc    - Metadata that describes the underlying complex FFT
-	* @param SrcTexture - SRV the source image buffer
-	* @param SrcRct     - The region in the Src buffer where the image to transform lives.
-	* @param DstUAV     - UAV, the destination buffer that will hold the result of the single pass
-	*
-	*/
-	void DispatchGSSubComplexFFTPassCS(FGPUFFTShaderContext& Context, const FFTDescription& FFTDesc,
-		const FTextureRHIRef& SrcTexture, const FIntRect& SrcWindow,
-		FUnorderedAccessViewRHIRef DstUAV)
+		SHADER_PARAMETER_RDG_TEXTURE(Texture2D, SrcTexture)
+		SHADER_PARAMETER_RDG_TEXTURE_UAV(RWTexture2D, DstTexture)
+	END_SHADER_PARAMETER_STRUCT()
+
+	static void ModifyCompilationEnvironment(const FGlobalShaderPermutationParameters& Parameters, FShaderCompilerEnvironment& OutEnvironment)
 	{
-	
-		using namespace GPUFFTComputeShaderUtils;
+		FPermutationDomain PermutationVector(Parameters.PermutationId);
+		int32 PowRadixSignalLength = PermutationVector.Get<FPowRadixSignalLengthDim>();
 
-		const FGPUFFTShaderContext::ShaderMapType& ShaderMap = Context.GetShaderMap();
-		FRHICommandList& RHICmdList = Context.GetRHICmdList();
+		FFFTShader::ModifyCompilationEnvironment(Parameters, OutEnvironment);
+		OutEnvironment.SetDefine(TEXT("INCLUDE_GROUP_SHARED_TWO_FOR_ONE_FFT"), 1);
+		if (PowRadixSignalLength > 8)
+		{
+			OutEnvironment.SetDefine(TEXT("MIXED_RADIX"), 1);
+		}
+	}
+};
 
-		const uint32 TransformLength = FFTDesc.SignalLength;
-		const FString TransformName = FFTDesc.FFT_TypeName();
+class FGSConvolutionWithTextureCS : public FFFTShader
+{
+	DECLARE_GLOBAL_SHADER(FGSConvolutionWithTextureCS);
+	SHADER_USE_PARAMETER_STRUCT(FGSConvolutionWithTextureCS, FFFTShader);
 
-		const uint32 NumSubRegions = TransformLength / FGroupShardSubFFTPassCS::SubPassLength();
+	using FPermutationDomain = TShaderPermutationDomain<FPowRadixSignalLengthDim>;
 
-		
+	BEGIN_SHADER_PARAMETER_STRUCT(FParameters, )
+		SHADER_PARAMETER(FIntPoint, SrcRectMin)
+		SHADER_PARAMETER(FIntPoint, SrcRectMax)
+		SHADER_PARAMETER(FIntPoint, DstExtent)
+		SHADER_PARAMETER(int32, TransformType)
 
-		SCOPED_DRAW_EVENTF(RHICmdList, GSSubComplexFFTPass, TEXT("FFT Multipass: %d GS Subpasses Complex %s of size %d"), 
-			NumSubRegions, *TransformName, FGroupShardSubFFTPassCS::SubPassLength());
+		SHADER_PARAMETER_RDG_TEXTURE(Texture2D, SrcTexture)
+		SHADER_PARAMETER_RDG_TEXTURE(Texture2D, FilterTexture)
+		SHADER_PARAMETER_RDG_TEXTURE_UAV(RWTexture2D, DstTexture)
+	END_SHADER_PARAMETER_STRUCT()
 
-		// The window on which a single transform acts.
-		FIntRect SubPassWindow = SrcWindow;
-		if (FFTDesc.IsHorizontal())
-		{
-			SubPassWindow.Max.X = SubPassWindow.Min.X + FGroupShardSubFFTPassCS::SubPassLength();
-		}
-		else
+	static void ModifyCompilationEnvironment(const FGlobalShaderPermutationParameters& Parameters, FShaderCompilerEnvironment& OutEnvironment)
+	{
+		FPermutationDomain PermutationVector(Parameters.PermutationId);
+		int32 PowRadixSignalLength = PermutationVector.Get<FPowRadixSignalLengthDim>();
+
+		FFFTShader::ModifyCompilationEnvironment(Parameters, OutEnvironment);
+		OutEnvironment.SetDefine(TEXT("INCLUDE_GROUP_SHARED_CONVOLUTION_WITH_TEXTURE"), 1);
+		if (PowRadixSignalLength > 8)
 		{
-			SubPassWindow.Max.Y = SubPassWindow.Min.Y + FGroupShardSubFFTPassCS::SubPassLength();
+			OutEnvironment.SetDefine(TEXT("MIXED_RADIX"), 1);
 		}
+	}
+};
+
+
+IMPLEMENT_GLOBAL_SHADER(FReorderFFTPassCS,           "/Engine/Private/GPUFastFourierTransform.usf", "ReorderFFTPassCS", SF_Compute);
+IMPLEMENT_GLOBAL_SHADER(FGroupShardSubFFTPassCS,     "/Engine/Private/GPUFastFourierTransform.usf", "GroupSharedSubComplexFFTCS", SF_Compute);
+IMPLEMENT_GLOBAL_SHADER(FComplexFFTPassCS,           "/Engine/Private/GPUFastFourierTransform.usf", "ComplexFFTPassCS", SF_Compute);
+IMPLEMENT_GLOBAL_SHADER(FPackTwoForOneFFTPassCS,     "/Engine/Private/GPUFastFourierTransform.usf", "PackTwoForOneFFTPassCS", SF_Compute);
+IMPLEMENT_GLOBAL_SHADER(FCopyWindowCS,               "/Engine/Private/GPUFastFourierTransform.usf", "CopyWindowCS", SF_Compute);
+IMPLEMENT_GLOBAL_SHADER(FComplexMultiplyImagesCS,    "/Engine/Private/GPUFastFourierTransform.usf", "ComplexMultiplyImagesCS", SF_Compute);
+IMPLEMENT_GLOBAL_SHADER(FGSComplexTransformCS,       "/Engine/Private/GPUFastFourierTransform.usf", "GroupSharedComplexFFTCS", SF_Compute);
+IMPLEMENT_GLOBAL_SHADER(FGSTwoForOneTransformCS,     "/Engine/Private/GPUFastFourierTransform.usf", "GroupSharedTwoForOneFFTCS", SF_Compute);
+IMPLEMENT_GLOBAL_SHADER(FGSConvolutionWithTextureCS, "/Engine/Private/GPUFastFourierTransform.usf", "GSConvolutionWithTextureCS", SF_Compute);
+
+
+/**
+ * Single Pass of a multi-pass complex FFT that just reorders data for
+ * a group shared subpass to consume.
+ *
+ * Assumes the dst buffer is large enough to hold the result.
+ * The Src float4 data is interpt as a pair of independent complex numbers.
+ *
+ * @param FFTDesc    - Metadata that describes the underlying complex FFT
+ * @param SrcRct     - The region in the Src buffer where the image to transform lives.
+ * @param SrcTexture - SRV the source image buffer
+ * @param DstRect    - The target Window.
+ * @param DstUAV     - UAV, the destination buffer that will hold the result of the single pass
+ *
+ */
+void DispatchReorderFFTPassCS(
+	FRDGBuilder& GraphBuilder,
+	const FGlobalShaderMap* ShaderMap,
+	const GPUFFT::FFTDescription& FFTDesc,
+	FRDGTextureRef SrcTexture, const FIntRect& SrcWindow,
+	FRDGTextureRef DstTexture, const FIntRect& DstWindow,
+	const bool bScrubNaNs = false)
+{
+	// Using multiple radix two passes.
+	const uint32 Radix = 2;
 
+	uint32 TransformLength = FFTDesc.SignalLength;
 
-		FIntPoint DstExtent = FFTDesc.TransformExtent();
+	// Breaks the data into the correct number of sub-transforms for the later group-shared pass.
+	uint32 PowTwoSubLengthCount = TransformLength / FGroupShardSubFFTPassCS::SubPassLength();
 
-		// Get a base pointer to the shader
-		TShaderRef<FGroupShardSubFFTPassCS> ComputeShader = GetGroupSharedSubFFTPassCS(ShaderMap, TransformLength);
-		RHICmdList.SetComputeShader(ComputeShader.GetComputeShader());
+	FReorderFFTPassCS::FParameters* PassParameters = GraphBuilder.AllocParameters<FReorderFFTPassCS::FParameters>();
+	PassParameters->SrcRect = SrcWindow;
+	PassParameters->DstRect = DstWindow;
+	{
+		PassParameters->TransformType = GPUFFT::BitEncode(FFTDesc.XFormType);
+		if (bScrubNaNs)
+		{
+			PassParameters->TransformType |= 4;
+		}
+	}
+	PassParameters->LogTwoLength = GPUFFT::BitSize(FFTDesc.SignalLength) - 1;
+	PassParameters->BitCount = GPUFFT::BitSize(PowTwoSubLengthCount) - 1;
+
+	PassParameters->SrcSRV = SrcTexture;
+	PassParameters->DstUAV = GraphBuilder.CreateUAV(DstTexture);
+
+	TShaderMapRef<FReorderFFTPassCS> ComputeShader(ShaderMap);
+	FComputeShaderUtils::AddPass(
+		GraphBuilder,
+		RDG_EVENT_NAME("FFT Multipass: Complex %s Reorder pass of size %d", FFTDesc.FFT_TypeName(), TransformLength),
+		ComputeShader,
+		PassParameters,
+		FIntVector(1, 1, TransformLength / Radix));
+}
 
-		// Bind output
-		FScopedUAVBind ScopedBind = FScopedUAVBind::BindOutput(RHICmdList, ComputeShader, DstUAV);
 
-		ComputeShader->SetCSParamters(RHICmdList, FFTDesc.XFormType, TransformLength, SubPassWindow, SrcTexture, NumSubRegions);
-		// The number of signals to transform simultaneously (i.e. number of scan lines)
+/**
+* A Group Shared Single Pass of a multi-pass complex FFT.
+* Assumes the dst buffer is large enough to hold the result.
+* The Src float4 data is interpt as a pair of independent complex numbers.
+*
+* @param FFTDesc    - Metadata that describes the underlying complex FFT
+* @param SrcTexture - SRV the source image buffer
+* @param SrcRct     - The region in the Src buffer where the image to transform lives.
+* @param DstUAV     - UAV, the destination buffer that will hold the result of the single pass
+*
+*/
+void DispatchGSSubComplexFFTPassCS(
+	FRDGBuilder& GraphBuilder,
+	const FGlobalShaderMap* ShaderMap,
+	const GPUFFT::FFTDescription& FFTDesc,
+	FRDGTextureRef SrcTexture, const FIntRect& SrcWindow,
+	FRDGTextureRef DstTexture)
+{
+	const uint32 TransformLength = FFTDesc.SignalLength;
 
-		const uint32 NumSignals = FFTDesc.IsHorizontal() ? SubPassWindow.Size().Y : SubPassWindow.Size().X;
+	const uint32 NumSubRegions = TransformLength / FGroupShardSubFFTPassCS::SubPassLength();
 
-		RHICmdList.DispatchComputeShader(1, 1, NumSignals);
+	// The window on which a single transform acts.
+	FIntRect SubPassWindow = SrcWindow;
+	if (FFTDesc.IsHorizontal())
+	{
+		SubPassWindow.Max.X = SubPassWindow.Min.X + FGroupShardSubFFTPassCS::SubPassLength();
+	}
+	else
+	{
+		SubPassWindow.Max.Y = SubPassWindow.Min.Y + FGroupShardSubFFTPassCS::SubPassLength();
 	}
 
+	const uint32 NumSignals = FFTDesc.IsHorizontal() ? SubPassWindow.Size().Y : SubPassWindow.Size().X;
+
+	FGroupShardSubFFTPassCS::FParameters* PassParameters = GraphBuilder.AllocParameters<FGroupShardSubFFTPassCS::FParameters>();
+	PassParameters->SrcWindow = SrcWindow;
+	PassParameters->TransformType = GPUFFT::BitEncode(FFTDesc.XFormType);
+	PassParameters->TransformLength = TransformLength;
+	PassParameters->NumSubRegions = NumSubRegions;
+
+	PassParameters->SrcTexture = SrcTexture;
+	PassParameters->DstTexture = GraphBuilder.CreateUAV(DstTexture);
+
+	TShaderMapRef<FGroupShardSubFFTPassCS> ComputeShader(ShaderMap);
+	FComputeShaderUtils::AddPass(
+		GraphBuilder,
+		RDG_EVENT_NAME("FFT Multipass: %d GS Subpasses Complex %s of size %d",
+			NumSubRegions, FFTDesc.FFT_TypeName(), FGroupShardSubFFTPassCS::SubPassLength()),
+		ComputeShader,
+		PassParameters,
+		FIntVector(1, 1, NumSignals));
+}
 
-	/**
-	* Complex 1D FFT of two independent complex signals.  
-	* Assumes the dst buffer is large enough to hold the result.
-	* The Src float4 data is interpt as a pair of independent complex numbers. 
-	*
-	* @param Context    - container for RHI and ShanderMap
-	* @param FFTDesc    - Metadata that describes the underlying complex FFT
-	* @param SrcTexture - SRV the source image buffer
-	* @param SrcRct     - The region in the Src buffer where the image to transform lives.
-	* @param DstUAV     - UAV, the destination buffer that will hold the result of the 1d complex FFT
-	*
-	*/
-	void DispatchGSComplexFFTCS(FGPUFFTShaderContext& Context, const FFTDescription& FFTDesc,
-		const FTextureRHIRef& SrcTexture, const FIntRect& SrcRect, 
-		FUnorderedAccessViewRHIRef DstUAV)
-	{
-		using namespace GPUFFTComputeShaderUtils;
+/**
+* Single Pass of a multi-pass complex FFT.
+* Assumes the dst buffer is large enough to hold the result.
+* The Src float4 data is interpt as a pair of independent complex numbers.
+*
+* @param FFTDesc    - Metadata that describes the underlying complex FFT
+* @param PassStage  - The Depth at which this FFT pass lives. 
+* @param SrcTexture - SRV the source image buffer
+* @param SrcRct     - The region in the Src buffer where the image to transform lives.
+* @param DstUAV     - UAV, the destination buffer that will hold the result of the single pass
+*
+*/
+void DispatchComplexFFTPassCS(
+	FRDGBuilder& GraphBuilder,
+	const FGlobalShaderMap* ShaderMap,
+	const GPUFFT::FFTDescription& FFTDesc,
+	const uint32 PassLength,
+	FRDGTextureRef SrcTexture, const FIntRect& SrcWindow,
+	FRDGTextureRef DstTexture,
+	const bool bScrubNaNs = false)
+{
+	// Using multiple radix two passes.
+	const uint32 Radix = 2;
 
-		const FGPUFFTShaderContext::ShaderMapType& ShaderMap = Context.GetShaderMap();
-		FRHICommandList& RHICmdList                 = Context.GetRHICmdList();
+	const uint32 TransformLength = FFTDesc.SignalLength;
 
-		const uint32 TransformLength = FFTDesc.SignalLength;
-		const FString TransformName  = FFTDesc.FFT_TypeName();
-		const FIntPoint DstExtent    = FFTDesc.TransformExtent();
+	// Translate the transform type. 
+	uint32 TransformTypeValue = GPUFFT::BitEncode(FFTDesc.XFormType);
+	if (bScrubNaNs)
+	{
+		TransformTypeValue |= 4;
+	}
 
-		SCOPED_DRAW_EVENTF(RHICmdList, ComplexFFTImage, TEXT("FFT: Complex %s of size %d"), *TransformName, TransformLength);
+	FComplexFFTPassCS::FParameters* PassParameters = GraphBuilder.AllocParameters<FComplexFFTPassCS::FParameters>();
+	PassParameters->SrcRect = SrcWindow;
+	PassParameters->DstRect = FIntRect(FIntPoint(0, 0), FFTDesc.TransformExtent());
+	PassParameters->TransformType = TransformTypeValue;
+	PassParameters->PowTwoLength = PassLength;
+	PassParameters->BitCount = GPUFFT::BitSize(TransformLength);
+
+	PassParameters->SrcSRV = SrcTexture;
+	PassParameters->DstUAV = GraphBuilder.CreateUAV(DstTexture);
+
+	TShaderMapRef<FComplexFFTPassCS> ComputeShader(ShaderMap);
+	FComputeShaderUtils::AddPass(
+		GraphBuilder,
+		RDG_EVENT_NAME("FFT Multipass: Pass %d of Complex %s of size %d", PassLength, FFTDesc.FFT_TypeName(), TransformLength),
+		ComputeShader,
+		PassParameters,
+		FIntVector(1, 1, TransformLength / Radix));
+}
 
+/**
+* Single Pass that separates or merges the transform of four real signals from 
+* viewed as the transform of two complex  signals.  
+*
+* Assumes the dst buffer is large enough to hold the result.
+* The Src float4 data is interprets  as a pair of independent complex numbers.
+*
+* @param FFTDesc    - Metadata that describes the underlying complex FFT
+* @param SrcTexture - SRV the source image buffer
+* @param DstUAV     - UAV, the destination buffer that will hold the result of the single pass
+*
+* The data in the SrcTexture and DstUAV is aligned with (0,0).
+* FFTDesc.IsHorizontal() indicates the transform direction in the buffer that needs to be spit/merged
+* FFTDesc.IsForward() indicates spit (true) vs merge (false).
+*/
+void DispatchPackTwoForOneFFTPassCS(
+	FRDGBuilder& GraphBuilder,
+	const FGlobalShaderMap* ShaderMap,
+	const GPUFFT::FFTDescription& FFTDesc,
+	FRDGTextureRef SrcTexture,
+	FRDGTextureRef DstTexture)
+{
+	const uint32 TransformLength = FFTDesc.SignalLength;
 
-		// Get a base pointer to the shader
-		TShaderRef<FGSComplexTransformBaseCS> ComputeShader = GetComplexFFTCS(ShaderMap, TransformLength);
-		RHICmdList.SetComputeShader(ComputeShader.GetComputeShader());
+	// A real signal of length 'TransformLenght' requires  only TransformLength / 2 + 1  complex coefficients, 
+	const uint32 RealTransformLength = ( TransformLength / 2 ) + 1;
 
-		// Bind output
-		FScopedUAVBind ScopedBind = FScopedUAVBind::BindOutput(RHICmdList, ComputeShader, DstUAV);
+	// The splitting into two real signals (isForward) or joinging back into a single signal
+	const uint32 ResultingLength = (FFTDesc.IsForward()) ? 2 * RealTransformLength : TransformLength;
 
-		ComputeShader->SetCSParamters(RHICmdList, FFTDesc.XFormType, SrcTexture, SrcRect, FIntRect(FIntPoint(0, 0), DstExtent));
+	FIntPoint DstExtent = FFTDesc.TransformExtent();
+	if (FFTDesc.IsHorizontal())
+	{
+		DstExtent.X = ResultingLength;
+	}	
+	else
+	{
+		DstExtent.Y = ResultingLength;
+	}
 
-		// The number of signals to transform simultaneously (i.e. number of scan lines)
+	FPackTwoForOneFFTPassCS::FParameters* PassParameters = GraphBuilder.AllocParameters<FPackTwoForOneFFTPassCS::FParameters>();
+	PassParameters->DstRect = FIntRect(FIntPoint(0, 0), DstExtent);
+	PassParameters->TransformType = GPUFFT::BitEncode(FFTDesc.XFormType);
 
-		const uint32 NumSignals = FFTDesc.IsHorizontal() ? SrcRect.Size().Y : SrcRect.Size().X;
+	PassParameters->SrcSRV = SrcTexture;
+	PassParameters->DstUAV = GraphBuilder.CreateUAV(DstTexture);
 
-		RHICmdList.DispatchComputeShader(1, 1, NumSignals);
-	}
+	TShaderMapRef<FPackTwoForOneFFTPassCS> ComputeShader(ShaderMap);
+	FComputeShaderUtils::AddPass(
+		GraphBuilder,
+		RDG_EVENT_NAME("FFT Multipass: TwoForOne Combine/split result of %s of size %d", FFTDesc.FFT_TypeName(), TransformLength),
+		ComputeShader,
+		PassParameters,
+		FIntVector(1, 1, RealTransformLength));
+}
 
-	/**
-	* Real 1D FFT of four independent real signals.   
-	* Assumes the dst buffer is large enough to hold the result.
-	* The Src float4 data is interptd as 4 independent real numbers.
-	* The Dst float4 data will be two complex numbers.
-	*
-	* @param Context    - container for RHI and ShanderMap
-	* @param FFTDesc    - Metadata that describes the underlying complex FFT
-	* @param SrcTexture - SRV the source image buffer
-	* @param SrcRct     - The region in the Src buffer where the image to transform lives.
-	* @param DstUAV     - UAV, the destination buffer that will hold the result of the 1d complex FFT
-	* @param DstRect    - Where to write the tranform data in the Dst buffer
-	* @param PreFilter  - Used to boost the intensity of already bright pixels.
-	*
-	*/
-	void DispatchGSTwoForOneFFTCS(FGPUFFTShaderContext& Context, const FFTDescription& FFTDesc,
-		const FTextureRHIRef& SrcTexture, const FIntRect& SrcRect,
-		FUnorderedAccessViewRHIRef DstUAV, const FIntRect& DstRect,
-		const FPreFilter& PreFilter)
-	{
-		using namespace GPUFFTComputeShaderUtils;
+/**
+ * Single Pass that Copies a sub-region of a buffer and potentially 
+ * boosts the intensity of select pixels.
+ * 
+ * Assumes the dst buffer is large enough to hold the result.
+ * The Src float4 data is interpreted as a pair of independent complex numbers.
+ * The Knl float4 data is interpreted as a pair of independent complex numbers
+ *
+ * @param SrcWindow  - The region of interest to copy.
+ * @param SrcTexture - SRV the source image buffer.
+ * @param DstWindow  - The target location for the images.
+ *
+ * @param DstUAV     - UAV, the destination buffer that will hold the result of the single pass
+ * @param Prefilter  - Optional filter to boost selected pixels.
+ */
+void DispatchCopyWindowCS(
+	FRDGBuilder& GraphBuilder,
+	const FGlobalShaderMap* ShaderMap,
+	FRDGTextureRef SrcTexture, const FIntRect& SrcWindow,
+	FRDGTextureRef DstTexture, const FIntRect& DstWindow,
+	const GPUFFT::FPreFilter& PreFilter = GPUFFT::FPreFilter(TNumericLimits<float>::Max(), TNumericLimits<float>::Lowest(), 0.f))
+{
+	FCopyWindowCS::FParameters* PassParameters = GraphBuilder.AllocParameters<FCopyWindowCS::FParameters>();
+	PassParameters->SrcRect = SrcWindow;
+	PassParameters->DstRect = DstWindow;
+	PassParameters->BrightPixelGain = PreFilter;
+
+	PassParameters->SrcSRV = SrcTexture;
+	PassParameters->DstUAV = GraphBuilder.CreateUAV(DstTexture);
+
+	TShaderMapRef<FCopyWindowCS> ComputeShader(ShaderMap);
+	FComputeShaderUtils::AddPass(
+		GraphBuilder,
+		RDG_EVENT_NAME("FFT Multipass: Copy Subwindow"),
+		ComputeShader,
+		PassParameters,
+		FComputeShaderUtils::GetGroupCount(DstWindow.Size(), FCopyWindowCS::ThreadCount()));
+}
 
-		const FGPUFFTShaderContext::ShaderMapType& ShaderMap = Context.GetShaderMap();
-		FRHICommandList& RHICmdList                 = Context.GetRHICmdList();
+/**
+ * 
+ * Single Pass that computes the Frequency Space convolution of two buffers
+ * that have already been transformed into frequency space.
+ *
+ * This is really means the complex product of two buffers divided
+ * by the correct values to "normalize" the effect of the kernel buffer.
+ *
+ * In this case, each float4 is viewed as a pair of complex numbers,
+ * and the product of float4 Src, Knl is computed
+ * as float4(ComplexMult(Src.xy, Knl.xy) / Na, ComplexMult(Src.zw, Knl.zw)/Nb)
+ * where Na and Nb are related to the sums of the weights in the kernel buffer. 
+ *
+ * Assumes the dst buffer is large enough to hold the result.
+ * The Src float4 data is interpreted as a pair of independent complex numbers.
+ * The Knl float4 data is interpreted as a pair of independent complex numbers  
+ *
+ * @param bHorizontalFirst - Describes the layout of the transformed data.
+ *                           bHorizontalFirst == true for data that was transformed as 
+ *                           Vertical::ComplexFFT following Horizontal::TwoForOneRealFFT
+ *                           bHorizontalFirst == false for data that was transformed as
+ *                           Horizontal::ComplexFFT following Vertical::TwoForOneRealFFT
+ * @param SrcTexture - SRV the source image buffer.
+ * @param KnlTexture - SRV the kernel image buffer.
+ *
+ * @param DstUAV     - UAV, the destination buffer that will hold the result of the single pass
+ *
+ */
+void DispatchComplexMultiplyImagesCS(
+	FRDGBuilder& GraphBuilder,
+	const FGlobalShaderMap* ShaderMap,
+	const bool bHorizontalScanlines,
+	FRDGTextureRef SrcTexture, const FIntRect& SrcWindow,
+	FRDGTextureRef KnlTexure,
+	FRDGTextureRef DstTexture)
+{
+	// The size of the dst
+	const FIntPoint DstExtent = SrcWindow.Size();
 
-		const uint32 TransformLength = FFTDesc.SignalLength;
-		const FString TransformName  = FFTDesc.FFT_TypeName();
+	// Align the scanlines in the direction of the first transform.
+	const uint32 NumScanLines = (bHorizontalScanlines)  ? DstExtent.Y : DstExtent.X;
+	const uint32 SignalLength = (!bHorizontalScanlines) ? DstExtent.Y : DstExtent.X;
 	
+	FComplexMultiplyImagesCS::FParameters* PassParameters = GraphBuilder.AllocParameters<FComplexMultiplyImagesCS::FParameters>();
+	PassParameters->SrcRect = SrcWindow;
+	PassParameters->DataLayout = (bHorizontalScanlines) ? 1 : 0;
+
+	PassParameters->SrcSRV = SrcTexture;
+	PassParameters->KnlSRV = KnlTexure;
+	PassParameters->DstUAV = GraphBuilder.CreateUAV(DstTexture);
+
+	TShaderMapRef<FComplexMultiplyImagesCS> ComputeShader(ShaderMap);
+	FComputeShaderUtils::AddPass(
+		GraphBuilder,
+		RDG_EVENT_NAME("FFT Multipass: Convolution in freq-space"),
+		ComputeShader,
+		PassParameters,
+		FIntVector(1, 1, NumScanLines));
+}
 
+/**
+ * Complex 1D FFT of two independent complex signals.
+ * Assumes the dst buffer is large enough to hold the result.
+ * The Src float4 data is interpt as a pair of independent complex numbers.
+ *
+ * @param Context    - container for RHI and ShanderMap
+ * @param FFTDesc    - Metadata that describes the underlying complex FFT
+ * @param SrcTexture - SRV the source image buffer
+ * @param SrcRct     - The region in the Src buffer where the image to transform lives.
+ * @param DstUAV     - UAV, the destination buffer that will hold the result of the 1d complex FFT
+ *
+ */
+void DispatchGSComplexFFTCS(
+	FRDGBuilder& GraphBuilder,
+	const FGlobalShaderMap* ShaderMap,
+	const GPUFFT::FFTDescription& FFTDesc,
+	FRDGTextureRef SrcTexture, const FIntRect& SrcRect,
+	FRDGTextureRef DstTexture)
+{
+	const uint32 TransformLength = FFTDesc.SignalLength;
+	const FIntPoint DstExtent    = FFTDesc.TransformExtent();
 
-		SCOPED_DRAW_EVENTF(RHICmdList, FRCPassFFT, TEXT("FFT: Two-For-One %s of size %d of buffer %d x %d"), *TransformName, TransformLength, SrcRect.Size().X, SrcRect.Size().Y);
-
-		// Get a base pointer to the shader
-		TShaderRef<FGSComplexTransformBaseCS> ComputeShader = GetTwoForOneFFTCS(ShaderMap, TransformLength);
-		RHICmdList.SetComputeShader(ComputeShader.GetComputeShader());
+	// The number of signals to transform simultaneously (i.e. number of scan lines)
+	const uint32 NumSignals = FFTDesc.IsHorizontal() ? SrcRect.Size().Y : SrcRect.Size().X;
 
+	FGSComplexTransformCS::FParameters* PassParameters = GraphBuilder.AllocParameters<FGSComplexTransformCS::FParameters>();
+	PassParameters->SrcRectMin = SrcRect.Min;
+	PassParameters->SrcRectMax = SrcRect.Max;
+	PassParameters->DstExtent = DstExtent;
+	PassParameters->DstRect = FIntRect(FIntPoint(0, 0), DstExtent);
+	PassParameters->BrightPixelGain = GPUFFT::FPreFilter(TNumericLimits<float>::Max(), TNumericLimits<float>::Lowest(), 0.f);
+	{
+		// Translate the transform type. 
+		PassParameters->TransformType = GPUFFT::BitEncode(FFTDesc.XFormType);
 
-		FScopedUAVBind ScopedBind = FScopedUAVBind::BindOutput(RHICmdList, ComputeShader, DstUAV);
-		ComputeShader->SetCSParamters(RHICmdList, FFTDesc.XFormType, SrcTexture, SrcRect, DstRect, PreFilter);
+		// We have a valid prefilter if the min is less than the max
+		if (PassParameters->BrightPixelGain.Component(0) < PassParameters->BrightPixelGain.Component(1))
+		{
+			// Encode a bool to turn on the pre-filter.
+			PassParameters->TransformType |= 4;
+		}
+	}
 
+	PassParameters->SrcTexture = SrcTexture;
+	PassParameters->DstTexture = GraphBuilder.CreateUAV(DstTexture);
 
-		// The number of signals to transform simultaneously (i.e. number of scan lines)
-		const uint32 NumScanLines = (FFTDesc.IsHorizontal()) ? SrcRect.Size().Y : SrcRect.Size().X;
+	FGSComplexTransformCS::FPermutationDomain PermutationVector;
+	PermutationVector.Set<FPowRadixSignalLengthDim>(TransformLength);
 
-		RHICmdList.DispatchComputeShader(1, 1, NumScanLines);
+	TShaderMapRef<FGSComplexTransformCS> ComputeShader(ShaderMap, PermutationVector);
+	FComputeShaderUtils::AddPass(
+		GraphBuilder,
+		RDG_EVENT_NAME("FFT: Complex %s of size %d", FFTDesc.FFT_TypeName(), TransformLength),
+		ComputeShader,
+		PassParameters,
+		FIntVector(1, 1, NumSignals));
+}
 
-	}
 
-	/**
-	* Complex 1D FFT followed by multiplication in with kernel and inverse transform.
-	*
-	* @param Context    - container for RHI and ShanderMap
-	* @param FFTDesc    - Metadata that describes the underlying complex FFT
-	* @param PreTransformedKernel - Pre-transformed kernel used in the convolution.
-	* @param SrcTexture - SRV the source image buffer
-	* @param SrcRct     - The region in the Src buffer where the image to transform lives.
-	* @param DstUAV     - UAV, the destination buffer that will hold the result of the 1d complex FFT
-	* @param DstRect    - Where to write the tranform data in the Dst buffer
-	*
-	*/
-	void DispatchGSConvolutionWithTextureCS(FGPUFFTShaderContext& Context,
-		const FFTDescription& FFTDesc,
-		const FTextureRHIRef& PreTransformedKernel,
-		const FTextureRHIRef& SrcTexture, 
-		const FIntRect& SrcRect,
-		FUnorderedAccessViewRHIRef DstUAV)
+/**
+* Real 1D FFT of four independent real signals.   
+* Assumes the dst buffer is large enough to hold the result.
+* The Src float4 data is interptd as 4 independent real numbers.
+* The Dst float4 data will be two complex numbers.
+*
+* @param Context    - container for RHI and ShanderMap
+* @param FFTDesc    - Metadata that describes the underlying complex FFT
+* @param SrcTexture - SRV the source image buffer
+* @param SrcRct     - The region in the Src buffer where the image to transform lives.
+* @param DstUAV     - UAV, the destination buffer that will hold the result of the 1d complex FFT
+* @param DstRect    - Where to write the tranform data in the Dst buffer
+* @param PreFilter  - Used to boost the intensity of already bright pixels.
+*
+*/
+void DispatchGSTwoForOneFFTCS(
+	FRDGBuilder& GraphBuilder,
+	const FGlobalShaderMap* ShaderMap,
+	const GPUFFT::FFTDescription& FFTDesc,
+	FRDGTextureRef SrcTexture, const FIntRect& SrcRect,
+	FRDGTextureRef DstTexture, const FIntRect& DstRect,
+	const GPUFFT::FPreFilter& PreFilter)
+{
+	const uint32 TransformLength = FFTDesc.SignalLength;
+	
+	// The number of signals to transform simultaneously (i.e. number of scan lines)
+	const uint32 NumScanLines = (FFTDesc.IsHorizontal()) ? SrcRect.Size().Y : SrcRect.Size().X;
+
+	FGSTwoForOneTransformCS::FParameters* PassParameters = GraphBuilder.AllocParameters<FGSTwoForOneTransformCS::FParameters>();
+	PassParameters->SrcRectMin = SrcRect.Min;
+	PassParameters->SrcRectMax = SrcRect.Max;
+	PassParameters->DstExtent = DstRect.Size();
+	PassParameters->DstRect = DstRect;
+	PassParameters->BrightPixelGain = PreFilter;
 	{
-		using GPUFFTComputeShaderUtils::FScopedUAVBind;
-
-		const FGPUFFTShaderContext::ShaderMapType& ShaderMap = Context.GetShaderMap();
-		FRHICommandList& RHICmdList                 = Context.GetRHICmdList();
+		// Translate the transform type. 
+		PassParameters->TransformType = GPUFFT::BitEncode(FFTDesc.XFormType);
 
-		// Transform particulars
+		// We have a valid prefilter if the min is less than the max
+		if (PassParameters->BrightPixelGain.Component(0) < PassParameters->BrightPixelGain.Component(1))
+		{
+			// Encode a bool to turn on the pre-filter.
+			PassParameters->TransformType |= 4;
+		}
+	}
 
-		const uint32 SignalLength   = FFTDesc.SignalLength;
-		const FString XFormDirName  = FFTDesc.FFT_TypeName();
-		const bool bIsHornizontal   = FFTDesc.IsHorizontal();
+	PassParameters->SrcTexture = SrcTexture;
+	PassParameters->DstTexture = GraphBuilder.CreateUAV(DstTexture);
 
-		const FIntPoint& SrcRectSize = SrcRect.Size();
-		// The number of signals to transform simultaneously (i.e. number of scan lines)
-		// NB: This may be different from the FFTDesc.NumScanlines.
-		const uint32 NumSignals = (bIsHornizontal) ? SrcRectSize.Y : SrcRectSize.X;
+	FGSTwoForOneTransformCS::FPermutationDomain PermutationVector;
+	PermutationVector.Set<FPowRadixSignalLengthDim>(TransformLength);
 
-		
-		SCOPED_DRAW_EVENTF(RHICmdList, FRCPassFFTBloom, TEXT("FFT: Apply %s Transform, Multiply Texture, and InverseTransform size %d of buffer %d x %d"), *XFormDirName, SignalLength, SrcRectSize.X, SrcRectSize.Y);
+	TShaderMapRef<FGSTwoForOneTransformCS> ComputeShader(ShaderMap, PermutationVector);
+	FComputeShaderUtils::AddPass(
+		GraphBuilder,
+		RDG_EVENT_NAME("FFT: Two-For-One %s of size %d of buffer %d x %d", FFTDesc.FFT_TypeName(), TransformLength, SrcRect.Size().X, SrcRect.Size().Y),
+		ComputeShader,
+		PassParameters,
+		FIntVector(1, 1, NumScanLines));
+}
 
-		// Get a base pointer to the shader
-		TShaderRef<FGSConvolutionWithTextureKernelBaseCS> ComputeShader = GetConvolutionWithTextureKernelCS(ShaderMap, SignalLength);
-		RHICmdList.SetComputeShader(ComputeShader.GetComputeShader());
-		
-		FScopedUAVBind ScopedBind = FScopedUAVBind::BindOutput(RHICmdList, ComputeShader, DstUAV);
 
-		ComputeShader->SetCSParamters(RHICmdList, FFTDesc.XFormType, PreTransformedKernel, SrcTexture, SrcRect, SrcRect.Size());
+/**
+ * Complex 1D FFT followed by multiplication in with kernel and inverse transform.
+ *
+ * @param Context    - container for RHI and ShanderMap
+ * @param FFTDesc    - Metadata that describes the underlying complex FFT
+ * @param PreTransformedKernel - Pre-transformed kernel used in the convolution.
+ * @param SrcTexture - SRV the source image buffer
+ * @param SrcRct     - The region in the Src buffer where the image to transform lives.
+ * @param DstUAV     - UAV, the destination buffer that will hold the result of the 1d complex FFT
+ * @param DstRect    - Where to write the tranform data in the Dst buffer
+ *
+ */
+void DispatchGSConvolutionWithTextureCS(
+	FRDGBuilder& GraphBuilder,
+	const FGlobalShaderMap* ShaderMap,
+	const GPUFFT::FFTDescription& FFTDesc,
+	FRDGTextureRef SrcTexture, const FIntRect& SrcRect,
+	FRDGTextureRef PreTransformedKernel,
+	FRDGTextureRef DstTexture)
+{
+	const uint32 SignalLength   = FFTDesc.SignalLength;
+	const bool bIsHornizontal   = FFTDesc.IsHorizontal();
+
+	const FIntPoint& SrcRectSize = SrcRect.Size();
+	// The number of signals to transform simultaneously (i.e. number of scan lines)
+	// NB: This may be different from the FFTDesc.NumScanlines.
+	const uint32 NumSignals = (bIsHornizontal) ? SrcRectSize.Y : SrcRectSize.X;
+
+	FGSConvolutionWithTextureCS::FParameters* PassParameters = GraphBuilder.AllocParameters<FGSConvolutionWithTextureCS::FParameters>();
+	PassParameters->SrcRectMin = SrcRect.Min;
+	PassParameters->SrcRectMax = SrcRect.Max;
+	PassParameters->DstExtent = SrcRect.Size();
+	{
+		PassParameters->TransformType = GPUFFT::BitEncode(FFTDesc.XFormType);
 
-		RHICmdList.DispatchComputeShader(1, 1, NumSignals);
+		const bool bUseAlpha = true;
+		if (bUseAlpha)
+		{
+			PassParameters->TransformType |= 8;
+		}
 	}
-} // end anonymous namespace 
 
+	PassParameters->SrcTexture = SrcTexture;
+	PassParameters->FilterTexture = PreTransformedKernel;
+	PassParameters->DstTexture = GraphBuilder.CreateUAV(DstTexture);
 
-void GPUFFT::CopyImage2D(FGPUFFTShaderContext& Context,
-	const FIntRect& SrcWindow, const FTextureRHIRef& SrcTexture,
-	const FIntRect& DstWindow, FUnorderedAccessViewRHIRef DstUAV,
-	const FPreFilter& PreFilter)
-{
-	DispatchCopyWindowCS(Context, SrcWindow, SrcTexture, DstWindow, DstUAV, PreFilter);
+	FGSConvolutionWithTextureCS::FPermutationDomain PermutationVector;
+	PermutationVector.Set<FPowRadixSignalLengthDim>(SignalLength);
+
+	TShaderMapRef<FGSConvolutionWithTextureCS> ComputeShader(ShaderMap, PermutationVector);
+	FComputeShaderUtils::AddPass(
+		GraphBuilder,
+		RDG_EVENT_NAME("FFT: Apply %s Transform, Multiply Texture, and InverseTransform size %d of buffer %d x %d", FFTDesc.FFT_TypeName(), SignalLength, SrcRectSize.X, SrcRectSize.Y),
+		ComputeShader,
+		PassParameters,
+		FIntVector(1, 1, NumSignals));
 }
 
-void GPUFFT::ComplexFFTImage1D::Requirements(const FFTDescription& FFTDesc, FIntPoint& MinBufferSize, bool& bUseMultipass)
+bool FitsInGroupSharedMemory(const uint32& Length)
 {
-	MinBufferSize = FFTDesc.TransformExtent();
-	bUseMultipass = !FitsInGroupSharedMemory(FFTDesc);
+	const bool bFitsInGroupSharedMemory = !(GPUFFT::MaxScanLineLength() < Length);
+	return bFitsInGroupSharedMemory;
 }
 
-
-bool GPUFFT::ComplexFFTImage1D::GroupShared(FGPUFFTShaderContext& Context, const FFTDescription& FFTDesc,
-	const FIntRect& SrcWindow, const FTextureRHIRef& SrcTexture,  FUnorderedAccessViewRHIRef DstUAV)
+bool FitsInGroupSharedMemory(const GPUFFT::FFTDescription& FFTDesc)
 {
-	bool SuccessValue = true;
+	return FitsInGroupSharedMemory(FFTDesc.SignalLength);
+}
 
-	check(FMath::IsPowerOfTwo(FFTDesc.SignalLength));
 
-	if (FitsInGroupSharedMemory(FFTDesc))
-	{
-		DispatchGSComplexFFTCS(Context, FFTDesc, SrcTexture, SrcWindow, DstUAV);
-	}
-	else
-	{
-		SuccessValue = false;
-		// @todo
-		ensureMsgf(0, TEXT("The FFT size is too large for group shared memory"));
-		// Do forward expensive transform
-	}
+} // namespace
 
-	return SuccessValue;
-}
 
+void GPUFFT::CopyImage2D(
+	FRDGBuilder& GraphBuilder,
+	const FGlobalShaderMap* ShaderMap,
+	FRDGTextureRef SrcTexture, const FIntRect& SrcWindow,
+	FRDGTextureRef DstTexture, const FIntRect& DstWindow,
+	const FPreFilter& PreFilter)
+{
+	DispatchCopyWindowCS(GraphBuilder, ShaderMap, SrcTexture, SrcWindow, DstTexture, DstWindow, PreFilter);
+}
 
-bool GPUFFT::ComplexFFTImage1D::MultiPass(FGPUFFTShaderContext& Context, const FFTDescription& FFTDesc,
-	const FIntRect& Window, const FTextureRHIRef& SrcTexture,  FSceneRenderTargetItem& DstBuffer,
-	FSceneRenderTargetItem& TmpBuffer,
+void GPUFFT::ComplexFFTImage1D::MultiPass(
+	FRDGBuilder& GraphBuilder,
+	const FGlobalShaderMap* ShaderMap,
+	const FFTDescription& FFTDesc,
+	FRDGTextureRef SrcTexture, const FIntRect& SrcWindow,
+	FRDGTextureRef DstTexture,
 	const bool bScrubNaNs)
 {
-	bool SuccessValue = true;
+	if (FitsInGroupSharedMemory(FFTDesc))
+	{
+		return DispatchGSComplexFFTCS(
+			GraphBuilder, ShaderMap,
+			FFTDesc,
+			SrcTexture, SrcWindow,
+			DstTexture);
+	}
 
+	RDG_EVENT_SCOPE(GraphBuilder, "ComplexFFTImage1D");
 
 	const uint32 TransformLength = FFTDesc.SignalLength;
 	const FFT_XFORM_TYPE XFormType = FFTDesc.XFormType;
@@ -1640,24 +957,16 @@ bool GPUFFT::ComplexFFTImage1D::MultiPass(FGPUFFTShaderContext& Context, const F
 
 	check(FMath::IsPowerOfTwo(TransformLength));
 
-	// Command list
-	FRHICommandList& RHICmdList = Context.GetRHICmdList();
-
 	// The number of iterations required.
 	const uint32 Log2TransformLength = BitSize(TransformLength) - 1;
 
-	// Double buffer to manage the Dst and Tmp buffers.
-	GPUFFT::FDoubleBufferTargets Targets(TmpBuffer, DstBuffer);
-
-	
-
 	FIntPoint DstExtent = FFTDesc.TransformExtent();
 
 	const FIntRect  XFormWindow(FIntPoint(0, 0), DstExtent);
 
 	// Testing code branch: Breaks the transform into Log2(TransformLength) number of passes.
 	// this is the slowest algorithm, and uses no group-shared storage.
-	if (0) 
+	#if 0
 	{
 		DispatchComplexFFTPassCS(Context, FFTDesc, 1, SrcTexture, Window, Targets.DstTarget().UAV, bScrubNaNs);
 
@@ -1676,11 +985,6 @@ bool GPUFFT::ComplexFFTImage1D::MultiPass(FGPUFFTShaderContext& Context, const F
 
 			checkf(HasValidTargets(), TEXT("FFT: Allocated Buffers too small."));
 
-			RHICmdList.Transition(
-				{
-					FRHITransitionInfo(Targets.SrcTarget().ShaderResourceTexture, ERHIAccess::UAVMask, ERHIAccess::SRVMask),
-					FRHITransitionInfo(Targets.DstTarget().UAV, ERHIAccess::SRVMask, ERHIAccess::UAVMask)
-				});
 			DispatchComplexFFTPassCS(Context, FFTDesc, Ns, Targets.SrcTarget().ShaderResourceTexture, XFormWindow, Targets.DstTarget().UAV);
 		}
 
@@ -1691,168 +995,134 @@ bool GPUFFT::ComplexFFTImage1D::MultiPass(FGPUFFTShaderContext& Context, const F
 			SwapContents(TmpBuffer, DstBuffer);
 		}
 	}
-	else
-		// Reorder, followed by a High-level group-shared pass followed by Log2(TransformLength / SubPassLength() ) simple passes.
-		// In total 2 + Log2(TransformLength / SubPassLength() )  passes.   This will be on the order of 3 or 4 passes 
-		// compared with 12 or more ..
+	#endif
+	// Reorder, followed by a High-level group-shared pass followed by Log2(TransformLength / SubPassLength() ) simple passes.
+	// In total 2 + Log2(TransformLength / SubPassLength() )  passes.   This will be on the order of 3 or 4 passes 
+	// compared with 12 or more ..
 	{
-		// Re-order the data so we can do a pass of group-shared transforms
-		DispatchReorderFFTPassCS(Context, FFTDesc, Window, SrcTexture, XFormWindow, Targets.DstTarget().UAV);
-		Targets.Swap();
-
-		RHICmdList.Transition(
-			{
-				FRHITransitionInfo(Targets.SrcTarget().ShaderResourceTexture, ERHIAccess::UAVMask, ERHIAccess::SRVMask),
-				FRHITransitionInfo(Targets.DstTarget().UAV, ERHIAccess::SRVMask, ERHIAccess::UAVMask)
-			});
-		DispatchGSSubComplexFFTPassCS(Context, FFTDesc, Targets.SrcTarget().ShaderResourceTexture, XFormWindow, Targets.DstTarget().UAV);
+		FRDGTextureDesc SpectralDesc = FFTDesc.IsForward() ? DstTexture->Desc : SrcTexture->Desc;
 
+		// Re-order the data so we can do a pass of group-shared transforms
+		FRDGTextureRef SpectralTexture = GraphBuilder.CreateTexture(SpectralDesc, TEXT("FFT.Spectral"));
+		DispatchReorderFFTPassCS(
+			GraphBuilder, ShaderMap, FFTDesc,
+			/* SrcTexture = */ SrcTexture, SrcWindow,
+			/* DstTexture = */ SpectralTexture, XFormWindow);
+
+		FRDGTextureRef PrevSpectralTexture = GraphBuilder.CreateTexture(SpectralDesc, TEXT("FFT.Spectral"));
+		DispatchGSSubComplexFFTPassCS(
+			GraphBuilder, ShaderMap, FFTDesc,
+			/* SrcTexture = */ SpectralTexture, XFormWindow,
+			/* DstTexture = */ PrevSpectralTexture);
+		
 		for (uint32 Ns = FGroupShardSubFFTPassCS::SubPassLength(); Ns < TransformLength; Ns *= 2)
 		{
-			Targets.Swap();
-
-			auto HasValidTargets = [&Targets, &DstExtent]()->bool
+			FRDGTextureRef OpDstTexture;
+			if (Ns * 2 < TransformLength)
 			{
-				// Verify that the buffers being used are big enough.  Note that we are checking the "src" buffer, but due
-				// to the double buffering we will end up testing both buffers.
-				FIntPoint SrcBufferSize = Targets.SrcTarget().ShaderResourceTexture->GetTexture2D()->GetSizeXY();
-				bool Fits = !(SrcBufferSize.X < DstExtent.X) && !(SrcBufferSize.Y < DstExtent.Y);
-				return Fits;
-			};
+				FIntPoint SrcBufferSize = PrevSpectralTexture->Desc.Extent;
+				checkf(!(SrcBufferSize.X < DstExtent.X) && !(SrcBufferSize.Y < DstExtent.Y), TEXT("FFT: Allocated Buffers too small."));
 
-			checkf(HasValidTargets(), TEXT("FFT: Allocated Buffers too small."));
+				OpDstTexture = GraphBuilder.CreateTexture(DstTexture->Desc, TEXT("FFT.Spectral"));
+			}
+			else
+			{
+				OpDstTexture = DstTexture;
+			}
 
-			RHICmdList.Transition(
-				{
-					FRHITransitionInfo(Targets.SrcTarget().ShaderResourceTexture, ERHIAccess::UAVMask, ERHIAccess::SRVMask),
-					FRHITransitionInfo(Targets.DstTarget().UAV, ERHIAccess::SRVMask, ERHIAccess::UAVMask)
-				});
-			DispatchComplexFFTPassCS(Context, FFTDesc, Ns, Targets.SrcTarget().ShaderResourceTexture, XFormWindow, Targets.DstTarget().UAV);
+			DispatchComplexFFTPassCS(
+				GraphBuilder, ShaderMap, FFTDesc,
+				Ns,
+				/* SrcTexture = */ PrevSpectralTexture, XFormWindow,
+				/* DstTexture = */ OpDstTexture);
 
+			PrevSpectralTexture = OpDstTexture;
 		}
 
-		if (Targets.GetSrcIdx() != 0)
-		{
-			SwapContents(TmpBuffer, DstBuffer);
-		}
+		check(PrevSpectralTexture == DstTexture);
 	}
-
-	return SuccessValue;
 }
 
-
-void GPUFFT::TwoForOneRealFFTImage1D::Requirements(const FFTDescription& FFTDesc, FIntPoint& MinBufferSize, bool& bUseMultipass)
-{
-	FFTDescription TmpDesc = FFTDesc;
-	// The two-for-one produces a result that has two additional elements in the transform direction.
-	TmpDesc.SignalLength += 2;
-
-	MinBufferSize = TmpDesc.TransformExtent();
-	bUseMultipass = !FitsInGroupSharedMemory(FFTDesc);
-}
-
-bool GPUFFT::TwoForOneRealFFTImage1D::GroupShared(FGPUFFTShaderContext& Context, const FFTDescription& FFTDesc,
-	const FIntRect& SrcWindow, const FTextureRHIRef& SrcTexture,
-	const FIntRect& DstWindow, FUnorderedAccessViewRHIRef DstUAV,
+void GPUFFT::TwoForOneRealFFTImage1D::MultiPass(
+	FRDGBuilder& GraphBuilder,
+	const FGlobalShaderMap* ShaderMap,
+	const FFTDescription& FFTDesc,
+	FRDGTextureRef SrcTexture, const FIntRect& SrcWindow,
+	FRDGTextureRef DstTexture, const FIntRect& DstWindow,
 	const FPreFilter& PreFilter)
 {
-	bool SuccessValue = true;
-				
 	if (FitsInGroupSharedMemory(FFTDesc))
 	{
-		DispatchGSTwoForOneFFTCS(Context, FFTDesc, SrcTexture, SrcWindow, DstUAV, DstWindow, PreFilter);
-	}
-	else
-	{
-		SuccessValue = false;
-		// @todo
-		ensureMsgf(0, TEXT("The FFT size is too large for group shared memory"));
-		// Do forward expensive transform
+		DispatchGSTwoForOneFFTCS(GraphBuilder, ShaderMap, FFTDesc, SrcTexture, SrcWindow, DstTexture, DstWindow, PreFilter);
 	}
-
-	return SuccessValue;
-}
-
-
-bool GPUFFT::TwoForOneRealFFTImage1D::MultiPass(FGPUFFTShaderContext& Context, const FFTDescription& FFTDesc,
-	const FIntRect& SrcWindow, const FTextureRHIRef& SrcTexture,
-	FSceneRenderTargetItem& DstBuffer,
-	FSceneRenderTargetItem& TmpBuffer,
-	const FPreFilter& PreFilter)
-{
-	
-	bool SuccessValue = true;
-	if (FFTDesc.IsForward())
+	else if (FFTDesc.IsForward())
 	{
+		RDG_EVENT_SCOPE(GraphBuilder, "TwoForOneRealFFTImage1D(Forward)");
 
-		// Only filter on the forward transform.
+		// Filter only on forward transform through image copy pass.
 		if (IsActive(PreFilter))
 		{
-			// Copy data into DstBuffer
-			CopyImage2D(Context, SrcWindow, SrcTexture, SrcWindow, DstBuffer.UAV, PreFilter);
-
-			Context.GetRHICmdList().Transition(
-				{
-					FRHITransitionInfo(DstBuffer.ShaderResourceTexture, ERHIAccess::UAVMask, ERHIAccess::SRVMask),
-					FRHITransitionInfo(TmpBuffer.UAV, ERHIAccess::SRVMask, ERHIAccess::UAVMask)
-				});
-
-			// Transform as 2 sets of complex data, putting the result in the DstBuffer.  This performs multiple dispatches.
-			// Transform into DstBuffer.
-			SuccessValue = SuccessValue &&
-				ComplexFFTImage1D::MultiPass(Context, FFTDesc, SrcWindow, DstBuffer.ShaderResourceTexture, TmpBuffer /*result*/, DstBuffer/*tmp*/, true /*scrub nans*/);
-		}
-		else
-		{
-			// Transform as 2 sets of complex data, putting the result in the DstBuffer.  This performs multiple dispatches.
+			FRDGTextureDesc Desc = FRDGTextureDesc::Create2D(
+				SrcTexture->Desc.Extent,
+				SrcTexture->Desc.Format,
+				FClearValueBinding::None,
+				TexCreate_ShaderResource | TexCreate_UAV);
 
-			SuccessValue = SuccessValue &&
-				ComplexFFTImage1D::MultiPass(Context, FFTDesc, SrcWindow, SrcTexture, DstBuffer /*result*/, TmpBuffer/*tmp*/, true /*scrub nans*/);
-			
-			SwapContents(DstBuffer, TmpBuffer);
+			FRDGTextureRef PreFilteredTexture = GraphBuilder.CreateTexture(Desc, TEXT("FFT.PreFilter"));
+
+			CopyImage2D(
+				GraphBuilder, ShaderMap,
+				/* SrcTexture = */ SrcTexture, SrcWindow,
+				/* DstTexture = */ PreFilteredTexture, SrcWindow,
+				PreFilter);
+
+			SrcTexture = PreFilteredTexture;
 		}
 
-		Context.GetRHICmdList().Transition(
-			{
-				FRHITransitionInfo(TmpBuffer.ShaderResourceTexture, ERHIAccess::UAVMask, ERHIAccess::SRVMask),
-				FRHITransitionInfo(DstBuffer.UAV, ERHIAccess::SRVMask, ERHIAccess::UAVMask)
-			});
-		// Unpack the complex transform into transform of real data
-		DispatchPackTwoForOneFFTPassCS(Context, FFTDesc, TmpBuffer.ShaderResourceTexture, DstBuffer.UAV);
+		FRDGTextureRef SpectralTexture = GraphBuilder.CreateTexture(DstTexture->Desc, TEXT("FFT.Spectral"));
+		ComplexFFTImage1D::MultiPass(
+			GraphBuilder, ShaderMap, FFTDesc,
+			/* SrcTexture = */ SrcTexture, SrcWindow,
+			/* DstTexture = */ SpectralTexture,
+			/* bScrubNaNs = */ true);
 		
+		// Unpack the complex transform into transform of real data
+		ensure(DstWindow.Min == FIntPoint::ZeroValue);
+		DispatchPackTwoForOneFFTPassCS(
+			GraphBuilder, ShaderMap, FFTDesc,
+			/* SrcTexture = */ SpectralTexture,
+			/* DstTexture = */ DstTexture);
 	}
 	else  // Inverse transform.
 	{
-		Context.GetRHICmdList().Transition(FRHITransitionInfo(DstBuffer.UAV, ERHIAccess::SRVMask, ERHIAccess::UAVMask));
+		RDG_EVENT_SCOPE(GraphBuilder, "TwoForOneRealFFTImage1D(Inverse)");
+
 		// Pack the 4 transforms of real data as 2 transforms of complex data 
-		DispatchPackTwoForOneFFTPassCS(Context, FFTDesc, SrcTexture, DstBuffer.UAV);
+		FRDGTextureRef SpectralTexture = GraphBuilder.CreateTexture(SrcTexture->Desc, TEXT("FFT.Spectral"));
+		ensure(SrcWindow.Min == FIntPoint::ZeroValue);
+		DispatchPackTwoForOneFFTPassCS(
+			GraphBuilder, ShaderMap, FFTDesc,
+			/* SrcTexture = */ SrcTexture,
+			/* DstTexture = */ SpectralTexture);
 
-		Context.GetRHICmdList().Transition(
-			{
-				FRHITransitionInfo(DstBuffer.ShaderResourceTexture, ERHIAccess::UAVMask, ERHIAccess::SRVMask),
-				FRHITransitionInfo(TmpBuffer.UAV, ERHIAccess::SRVMask, ERHIAccess::UAVMask)
-			});
 		// Transform as complex data
-		SuccessValue = SuccessValue &&
-			ComplexFFTImage1D::MultiPass(Context, FFTDesc, SrcWindow, DstBuffer.ShaderResourceTexture, TmpBuffer/*result*/, DstBuffer/*tmp*/);
-
-		SwapContents(TmpBuffer, DstBuffer);
+		ensure(DstWindow.Min == FIntPoint::ZeroValue);
+		ComplexFFTImage1D::MultiPass(
+			GraphBuilder, ShaderMap, FFTDesc,
+			/* SrcTexture = */ SpectralTexture, SrcWindow,
+			/* DstTexture = */ DstTexture);
 	}
-
-	return SuccessValue;
 }
 
-
-
-bool GPUFFT::FFTImage2D(FGPUFFTShaderContext& Context, const FIntPoint& FrequencySize, bool bHorizontalFirst,
-	const FIntRect& ROIRect, const FTextureRHIRef& SrcTexture,
-	FSceneRenderTargetItem& DstBuffer,
-	FSceneRenderTargetItem& TmpBuffer)
+void GPUFFT::FFTImage2D(
+	FRDGBuilder& GraphBuilder,
+	const FGlobalShaderMap* ShaderMap,
+	const FIntPoint& FrequencySize, bool bHorizontalFirst,
+	FRDGTextureRef SrcTexture, const FIntRect& ROIRect,
+	FRDGTextureRef DstTexture)
 {
 	using GPUFFT::FFTDescription;
 
-	FRHICommandList& RHICmdList = Context.GetRHICmdList();
-
 	// This will be for the actual output.
 	const FIntRect FFTResultRect = ROIRect;
 	const FIntPoint ImageSpaceExtent = ROIRect.Size();
@@ -1872,137 +1142,69 @@ bool GPUFFT::FFTImage2D(FGPUFFTShaderContext& Context, const FIntPoint& Frequenc
 
 	const FIntRect TmpRect = FIntRect(FIntPoint(0, 0), TmpExent);
 
-
-	// Perform the transforms and convolutions.
-	bool SuccessValue = true;
-
 	// Two-for-one transform: SrcTexture fills the TmpBuffer
-
-	Context.GetRHICmdList().Transition(FRHITransitionInfo(TmpBuffer.UAV, ERHIAccess::Unknown, ERHIAccess::UAVMask));
-	if (FitsInGroupSharedMemory(TwoForOneFFTDesc))
-	{
-		// Horizontal FFT: input -> tmp buffer 0 
-		SuccessValue = SuccessValue &&
-			TwoForOneRealFFTImage1D::GroupShared(Context, TwoForOneFFTDesc, ROIRect, SrcTexture, TmpRect, TmpBuffer.UAV);
-		//->dst
-	}
-	else
-	{
-		// use the dst buffer as a dummy buffer
-		FSceneRenderTargetItem& DmyBuffer = DstBuffer;
-		// Horizontal FFT: input -> TmpBuffer
-		SuccessValue = SuccessValue &&
-			TwoForOneRealFFTImage1D::MultiPass(Context, TwoForOneFFTDesc, ROIRect, SrcTexture, TmpBuffer /*dst*/, DmyBuffer/*tmp*/);
-		//->dst
-	}
-
-	Context.GetRHICmdList().Transition(
-		{
-			FRHITransitionInfo(DstBuffer.UAV, ERHIAccess::Unknown, ERHIAccess::UAVMask),
-			FRHITransitionInfo(TmpBuffer.ShaderResourceTexture, ERHIAccess::UAVMask, ERHIAccess::SRVMask)
-		});
-
-	// Complex transform in the other direction: TmpBuffer fills DstBuffer
-	if (FitsInGroupSharedMemory(ComplexFFTDesc))
-	{
-		SuccessValue = SuccessValue &&
-			ComplexFFTImage1D::GroupShared(Context, ComplexFFTDesc, TmpRect, TmpBuffer.ShaderResourceTexture, DstBuffer.UAV);
-	}
-	else
-	{
-		SuccessValue = SuccessValue &&
-			ComplexFFTImage1D::MultiPass(Context, ComplexFFTDesc, TmpRect, TmpBuffer.ShaderResourceTexture, DstBuffer, TmpBuffer);
-	}
-
-	return SuccessValue;
-
+	FRDGTextureRef TempTexture = GraphBuilder.CreateTexture(DstTexture->Desc, DstTexture->Name);
+	TwoForOneRealFFTImage1D::MultiPass(
+		GraphBuilder, ShaderMap,
+		TwoForOneFFTDesc,
+		SrcTexture, ROIRect,
+		TempTexture, TmpRect);
+
+	ComplexFFTImage1D::MultiPass(
+		GraphBuilder, ShaderMap,
+		ComplexFFTDesc,
+		TempTexture, TmpRect,
+		DstTexture);
 }
 
-void GPUFFT::ConvolutionWithTextureImage1D::Requirements(const FFTDescription& FFTDesc, FIntPoint& MinBufferSize, bool& bUseMultipass)
+void GPUFFT::ConvolutionWithTextureImage1D::MultiPass(
+	FRDGBuilder& GraphBuilder,
+	const FGlobalShaderMap* ShaderMap,
+	const FFTDescription& FFTDesc,
+	FRDGTextureRef SrcTexture, const FIntRect& SrcWindow,
+	FRDGTextureRef TransformedKernel,
+	FRDGTextureRef DstTexture)
 {
-	bUseMultipass = !FitsInGroupSharedMemory(FFTDesc);
-	MinBufferSize = FFTDesc.TransformExtent();
-}
-
-bool GPUFFT::ConvolutionWithTextureImage1D::GroupShared(FGPUFFTShaderContext& Context, const FFTDescription& FFTDesc,
-	const FTextureRHIRef& TransformedKernel,
-	const FIntRect& SrcWindow,
-	const FTextureRHIRef& SrcTexture,
-	FUnorderedAccessViewRHIRef DstUAV)
-{
-
-	bool SuccessValue = true;
-
 	if (FitsInGroupSharedMemory(FFTDesc))
 	{
-		DispatchGSConvolutionWithTextureCS(Context, FFTDesc, TransformedKernel, SrcTexture, SrcWindow, DstUAV);
+		DispatchGSConvolutionWithTextureCS(
+			GraphBuilder, ShaderMap, FFTDesc,
+			SrcTexture, SrcWindow, TransformedKernel, DstTexture);
 	}
 	else
 	{
-		SuccessValue = false;
-		// @todo
-		ensureMsgf(0, TEXT("The FFT size is too large for group shared memory"));
+		RDG_EVENT_SCOPE(GraphBuilder, "ConvolutionWithTextureImage1D");
+
+		// Frequency Space size.
+		const FIntPoint TargetExtent = FFTDesc.TransformExtent();
+		const FIntRect  TargetRect(FIntPoint(0, 0), TargetExtent);
+
+		// Forward transform: Results in DstBuffer
+		FRDGTextureRef SpectralTexture0 = GraphBuilder.CreateTexture(SrcTexture->Desc, TEXT("FFT.Spectral"));
+		ComplexFFTImage1D::MultiPass(
+			GraphBuilder, ShaderMap, FFTDesc,
+			/* SrcTexture = */ SrcTexture, SrcWindow,
+			/* DstTexture = */ SpectralTexture0);
+
+		// Apply convolution kernel.
+		FRDGTextureRef SpectralTexture1 = GraphBuilder.CreateTexture(SrcTexture->Desc, TEXT("FFT.Spectral"));
+		DispatchComplexMultiplyImagesCS(
+			GraphBuilder, ShaderMap, FFTDesc.IsHorizontal(),
+			/* SrcTexture = */ SpectralTexture0, TargetRect,
+			TransformedKernel, 
+			/* DstTexture = */ SpectralTexture1);
+
+		// Backward transform.
+		FFTDescription InvFFTDesc = FFTDesc;
+		InvFFTDesc.XFormType = GetInverseOfXForm(FFTDesc.XFormType);
+
+		ComplexFFTImage1D::MultiPass(
+			GraphBuilder, ShaderMap, InvFFTDesc,
+			/* SrcTexture = */ SpectralTexture1, TargetRect,
+			/* DstTexture = */ DstTexture);
 	}
-
-	return SuccessValue;
-}
-
-bool GPUFFT::ConvolutionWithTextureImage1D::MultiPass(FGPUFFTShaderContext& Context, const FFTDescription& FFTDesc,
-	const FTextureRHIRef& TransformedKernel,
-	const FIntRect& SrcWindow,
-	const FTextureRHIRef& SrcTexture,
-	FSceneRenderTargetItem& DstBuffer,
-	FSceneRenderTargetItem& TmpBuffer)
-{
-	bool SuccessValue = true;
-
-
-	// Frequency Space size.
-	const FIntPoint TargetExtent = FFTDesc.TransformExtent();
-	const FIntRect  TargetRect(FIntPoint(0, 0), TargetExtent);
-
-	// Forward transform: Results in DstBuffer
-	SuccessValue = SuccessValue &&
-		ComplexFFTImage1D::MultiPass(Context, FFTDesc, SrcWindow, SrcTexture, DstBuffer, TmpBuffer);
-
-	Context.GetRHICmdList().Transition(
-		{
-			FRHITransitionInfo(DstBuffer.ShaderResourceTexture, ERHIAccess::UAVMask, ERHIAccess::SRVMask),
-			FRHITransitionInfo(TmpBuffer.UAV, ERHIAccess::SRVMask, ERHIAccess::UAVMask)
-		});
-
-	// Convolution: Results in TmpBuffer
-	DispatchComplexMultiplyImagesCS(Context, FFTDesc.IsHorizontal(), TargetRect, DstBuffer.ShaderResourceTexture, TransformedKernel,  TmpBuffer.UAV);
-
-	// Inverse Transform: Results in DstBuffer
-	
-	FFTDescription InvFFTDesc = FFTDesc;
-	InvFFTDesc.XFormType = GetInverseOfXForm(FFTDesc.XFormType);
-	// testing.  Verify that the forward transform works.
-	
-	Context.GetRHICmdList().Transition(
-		{
-			FRHITransitionInfo(TmpBuffer.ShaderResourceTexture, ERHIAccess::UAVMask, ERHIAccess::SRVMask),
-			FRHITransitionInfo(DstBuffer.UAV, ERHIAccess::SRVMask, ERHIAccess::UAVMask)
-		});
-	SuccessValue = SuccessValue &&
-		ComplexFFTImage1D::MultiPass(Context, InvFFTDesc, TargetRect, TmpBuffer.ShaderResourceTexture, DstBuffer, TmpBuffer);
-	
-	Context.GetRHICmdList().Transition(
-		{
-			FRHITransitionInfo(DstBuffer.ShaderResourceTexture, ERHIAccess::UAVMask, ERHIAccess::SRVMask),
-			FRHITransitionInfo(TmpBuffer.UAV, ERHIAccess::SRVMask, ERHIAccess::UAVMask)
-		});
-
-	// Copy back to the correct sized sub-window
-	DispatchCopyWindowCS(Context, TargetRect, DstBuffer.ShaderResourceTexture, SrcWindow, TmpBuffer.UAV);
-
-	SwapContents(TmpBuffer, DstBuffer);
-
-	return SuccessValue;
 }
 
-
 FIntPoint GPUFFT::Convolution2DBufferSize(const FIntPoint& FrequencySize, const bool bHorizontalFirst, const FIntPoint& ROIExtent)
 {
 	using GPUFFT::FFTDescription;
@@ -2037,25 +1239,33 @@ FIntPoint GPUFFT::Convolution2DBufferSize(const FIntPoint& FrequencySize, const
 	return BufferSize;
 }
 
-bool GPUFFT::ConvolutionWithTextureImage2D(FGPUFFTShaderContext& Context, const FIntPoint& FrequencySize, bool bHorizontalFirst,
-	const FTextureRHIRef& TransformedKernel,
-	const FIntRect& ROIRect, const FTextureRHIRef& SrcTexture,
-	FUnorderedAccessViewRHIRef ResultUAV,
-	FSceneRenderTargetItem& TmpBuffer0,
-	FSceneRenderTargetItem& TmpBuffer1,
-	const FPreFilter& PreFilter)
+FRDGTextureDesc CreateFrequencyDesc(const FIntPoint& FrequencySize, const bool bHorizontalFirst, const FIntPoint& SrcExtent)
 {
+	const FIntPoint FrequencyExtent = GPUFFT::Convolution2DBufferSize(
+		FrequencySize, bHorizontalFirst, SrcExtent);
+
+	return FRDGTextureDesc::Create2D(
+		FrequencyExtent,
+		GPUFFT::PixelFormat(),
+		FClearValueBinding::None,
+		TexCreate_ShaderResource | TexCreate_UAV);
+}
 
-	using GPUFFT::FFTDescription;
-
-	FRHICommandList& RHICmdList = Context.GetRHICmdList();
+void GPUFFT::ConvolutionWithTextureImage2D(
+	FRDGBuilder& GraphBuilder,
+	const FGlobalShaderMap* ShaderMap,
+	const FIntPoint& FrequencySize, bool bHorizontalFirst,
+	FRDGTextureRef TransformedKernel,
+	FRDGTextureRef SrcTexture, const FIntRect& ROIRect,
+	FRDGTextureRef DstTexture, const FIntRect& DstRect,
+	const FPreFilter& PreFilter)
+{
+	FRDGTextureDesc SpectralDesc = CreateFrequencyDesc(FrequencySize, bHorizontalFirst, ROIRect.Size());
 
 	// This will be for the actual output.
 	const FIntRect FFTResultRect = ROIRect;
 	const FIntPoint ROISize = ROIRect.Size();
 
-	//DoubleBufferTargets TmpBuffers(Context, TmpExent, PixelFormat);
-
 	// Set up the transform descriptions
 	FFTDescription TwoForOneFFTDesc;
 	TwoForOneFFTDesc.XFormType = bHorizontalFirst ? FFT_XFORM_TYPE::FORWARD_HORIZONTAL : FFT_XFORM_TYPE::FORWARD_VERTICAL;
@@ -2063,11 +1273,11 @@ bool GPUFFT::ConvolutionWithTextureImage2D(FGPUFFTShaderContext& Context, const
 	TwoForOneFFTDesc.NumScanLines = (bHorizontalFirst) ? ROISize.Y : ROISize.X;
 
 	// The output has two more elements
-	FIntRect TwoForOneOutputRect;
+	FIntRect SpectralOutputRect;
 	{
 		FFTDescription TmpDesc = TwoForOneFFTDesc;
 		TmpDesc.SignalLength += 2;
-		TwoForOneOutputRect = FIntRect(FIntPoint(0, 0), TmpDesc.TransformExtent());
+		SpectralOutputRect = FIntRect(FIntPoint(0, 0), TmpDesc.TransformExtent());
 	}
 
 	FFTDescription ConvolutionFFTDesc;
@@ -2080,72 +1290,32 @@ bool GPUFFT::ConvolutionWithTextureImage2D(FGPUFFTShaderContext& Context, const
 	
 	
 
-	// Perform the transforms and convolutions.
-	bool SuccessValue = true;
-
-
 	// ---- Two For One Transform --- 
 
-	// Horizontal FFT: input -> tmp buffer 0 
-	RHICmdList.Transition(FRHITransitionInfo(TmpBuffer0.UAV, ERHIAccess::Unknown, ERHIAccess::UAVMask));
-	if (FitsInGroupSharedMemory(TwoForOneFFTDesc))
-	{
-		SuccessValue = SuccessValue &&
-			TwoForOneRealFFTImage1D::GroupShared(Context, TwoForOneFFTDesc, ROIRect, SrcTexture, TwoForOneOutputRect, TmpBuffer0.UAV, PreFilter);
-		//->dst
-	}
-	else
-	{
-		SuccessValue = SuccessValue &&
-			TwoForOneRealFFTImage1D::MultiPass(Context, TwoForOneFFTDesc, ROIRect, SrcTexture, TmpBuffer0, TmpBuffer1, PreFilter);
-	}
+	FRDGTextureRef TwoForOneFFTOutput = GraphBuilder.CreateTexture(SpectralDesc, TEXT("FFT.TwoForOne"));
+	TwoForOneRealFFTImage1D::MultiPass(
+		GraphBuilder, ShaderMap, TwoForOneFFTDesc,
+		SrcTexture, ROIRect,
+		TwoForOneFFTOutput, SpectralOutputRect,
+		PreFilter);
 
 	// ---- 1 D Convolution --- 
 
-	RHICmdList.Transition(
-		{
-			FRHITransitionInfo(TmpBuffer0.ShaderResourceTexture, ERHIAccess::UAVMask, ERHIAccess::SRVMask),
-			FRHITransitionInfo(TmpBuffer1.UAV, ERHIAccess::Unknown, ERHIAccess::UAVMask)
-		});
-	if (FitsInGroupSharedMemory(ConvolutionFFTDesc))
-	{
-		// Vertical FFT / Filter / Vertical InvFFT : tmp buffer 0 -> tmp buffer 1
-		SuccessValue = SuccessValue &&
-			ConvolutionWithTextureImage1D::GroupShared(Context, ConvolutionFFTDesc, TransformedKernel, TwoForOneOutputRect, TmpBuffer0.ShaderResourceTexture, TmpBuffer1.UAV);
-		//->src
-	}
-	else
-	{
-		SuccessValue = SuccessValue &&
-			ConvolutionWithTextureImage1D::MultiPass(Context, ConvolutionFFTDesc, TransformedKernel, TwoForOneOutputRect, TmpBuffer0.ShaderResourceTexture, TmpBuffer1, TmpBuffer0);
-	}
+	FRDGTextureRef ConvolutionOutput = GraphBuilder.CreateTexture(SpectralDesc, TEXT("FFT.Convolution"));
+	ConvolutionWithTextureImage1D::MultiPass(
+		GraphBuilder, ShaderMap,
+		ConvolutionFFTDesc,
+		TwoForOneFFTOutput, SpectralOutputRect,
+		TransformedKernel,
+		ConvolutionOutput);
 
 
 	// ---- Inverse Two For One ---
 
-	RHICmdList.Transition(
-		{
-			FRHITransitionInfo(TmpBuffer1.ShaderResourceTexture, ERHIAccess::UAVMask, ERHIAccess::SRVMask),
-			FRHITransitionInfo(ResultUAV, ERHIAccess::Unknown, ERHIAccess::UAVMask)
-		});
-	if (FitsInGroupSharedMemory(TwoForOneIvnFFTDesc))
-	{
-		// Horizontal InvFFT: tmp buffer 1 -> result
-		SuccessValue = SuccessValue &&
-			TwoForOneRealFFTImage1D::GroupShared(Context, TwoForOneIvnFFTDesc, TwoForOneOutputRect, TmpBuffer1.ShaderResourceTexture, ROIRect, ResultUAV);
-		//->dst
-	}
-	else
-	{
-		// TmpBuffer1 -> TmpBuffer0
-		SuccessValue = SuccessValue &&
-			TwoForOneRealFFTImage1D::MultiPass(Context, TwoForOneIvnFFTDesc, TwoForOneOutputRect, TmpBuffer1.ShaderResourceTexture, TmpBuffer0, TmpBuffer1);
-
-		RHICmdList.Transition(FRHITransitionInfo(TmpBuffer0.ShaderResourceTexture, ERHIAccess::UAVMask, ERHIAccess::SRVMask));
-
-		DispatchCopyWindowCS(Context, TwoForOneOutputRect, TmpBuffer0.ShaderResourceTexture, ROIRect, ResultUAV);
-	}
-	
-	return SuccessValue;
+	TwoForOneRealFFTImage1D::MultiPass(
+		GraphBuilder, ShaderMap,
+		TwoForOneIvnFFTDesc,
+		ConvolutionOutput, SpectralOutputRect,
+		DstTexture, DstRect,
+		PreFilter);
 }
-
diff --git a/Engine/Source/Runtime/Renderer/Private/GPUFastFourierTransform.h b/Engine/Source/Runtime/Renderer/Private/GPUFastFourierTransform.h
--- a/Engine/Source/Runtime/Renderer/Private/GPUFastFourierTransform.h
+++ b/Engine/Source/Runtime/Renderer/Private/GPUFastFourierTransform.h
@@ -37,12 +37,7 @@ namespace GPUFFT
 	{
 		return  static_cast<FFT_XFORM_TYPE> ((static_cast<int> (XForm) + 2) % 4);
 	}
-	/**
-	* Pretty name for the associated transform type'
-	* @return - Pretty name for given XFormType
-	*/
-	FString FFT_TypeName(const FFT_XFORM_TYPE& XFormType);
-	
+
 	/**
 	* The direction of the associated transform type
 	*/
@@ -66,7 +61,7 @@ namespace GPUFFT
 		FIntPoint TransformExtent() const;
 		bool IsHorizontal() const;
 		bool IsForward() const;
-		FString FFT_TypeName() const;
+		const TCHAR* FFT_TypeName() const;
 		
 		// member data public
 		FFT_XFORM_TYPE  XFormType = FFT_XFORM_TYPE::FORWARD_HORIZONTAL;
@@ -74,30 +69,6 @@ namespace GPUFFT
 		uint32 NumScanLines = 0;
 	};
 
-
-
-	/**
-	* Context to hold the Hardware Interface and the Shader Map.
-	*/
-	class FGPUFFTShaderContext
-	{
-	public:
-
-		typedef FGlobalShaderMap  ShaderMapType;
-
-		FGPUFFTShaderContext(FRHICommandList& CmdList, const FGPUFFTShaderContext::ShaderMapType& Map)
-			: RHICmdList(&CmdList), ShaderMap(&Map)
-		{}
-
-		FRHICommandList& GetRHICmdList() { return *RHICmdList; }
-		const ShaderMapType& GetShaderMap() const { return *ShaderMap; }
-
-	private:
-		FGPUFFTShaderContext();
-		FRHICommandList* RHICmdList;
-		const ShaderMapType* ShaderMap;
-	};
-
 	/**
 	* The largest power of two length scan line that can
 	* be FFT'd with group shared memory.  Larger sizes require
@@ -105,23 +76,6 @@ namespace GPUFFT
 	*/
 	uint32 MaxScanLineLength();// { return 4096; }
 
-	/**
-	* Compare a signal length with the max that fits in group shared memory
-	*/
-	static bool FitsInGroupSharedMemory(const uint32& Length)
-	{
-		const bool bFitsInGroupSharedMemory = !(MaxScanLineLength() < Length);
-		return bFitsInGroupSharedMemory;
-	}
-
-	/**
-	* Verify the given FFT will fit in group shared memory.
-	*/
-	static bool FitsInGroupSharedMemory(const  FFTDescription& FFTDesc)
-	{
-		return FitsInGroupSharedMemory(FFTDesc.SignalLength);
-	}
-
 	/**
 	*  The pixel format required for transforming rgba buffer.
 	*/
@@ -145,42 +99,6 @@ namespace GPUFFT
 
 	struct ComplexFFTImage1D
 	{
-		/**
-		* The requirements of the complex 1d FFT. 
-		*
-		* @param FFTDesc         - Indicates the number of frequencies and scan lines as well as direction.
-		* @param MinBufferSize   - On return: The size of Dst buffer required. If multipass is required
-		*                          this size also applies to the tmpBuffer
-		* @param bUseGroupShared - On return: indicates if the transform requires the multipass 
-		*                          or if we can use the much faster group shared variant
-		*/
-		static void Requirements(const FFTDescription& FFTDesc, FIntPoint& MinBufferSize, bool& bUseMultipass);
-
-		/**
-		* One Dimensional Fast Fourier Transform of two signals in a two dimensional buffer.
-		* The each scanline of the float4 buffer is interpreted as two complex signals (x+iy and z + iw).
-		*
-		* On input a scanline (lenght N) contains 2 interlaced signals float4(r,g,b,a) = {z_n, q_n}.
-		* where z_n = r_n + i g_n  and q_n = b_n + i a_n
-		*
-		* On output a scanline will be N  elements long with the following data layout
-		* Entries k = [0,     N-1] hold  float4(Z_k, Q_k)
-		* where the array {Z_k} is the transform of the signal {z_n}
-		*
-		* @param Context    - RHI and ShadersMap
-		* @param FFTDesc    - Defines the scanline direction and count and transform type
-		* @param SrcWindow  - Location of the data in the SrcTexture to be transformed. Zero values pad
-		*                     the data to fit the  FFTDesc TrasformExtent.
-		* @param SrcTexture - The two dimensional float4 input buffer to transform
-		* @param DstUAV     - The buffer to be populated by the result of the transform.
-		*
-		* @return           - True on success.  Will only fail if the buffer is too large for supported transform method
-		*
-		* NB: This function does not transition resources.  That must be done by the calling scope.
-		*/
-		static bool GroupShared(FGPUFFTShaderContext& Context, const FFTDescription& FFTDesc,
-			const FIntRect& SrcWindow, const FTextureRHIRef& SrcTexture, FUnorderedAccessViewRHIRef DstUAV);
-
 		/**
 		* One Dimensional Fast Fourier Transform of two signals in a two dimensional buffer.
 		* The each scanline of the float4 buffer is interpreted as two complex signals (x+iy and z + iw).
@@ -206,9 +124,13 @@ namespace GPUFFT
 		*
 		* NB: This function does not transition resources.  That must be done by the calling scope.
 		*/
-		static bool MultiPass(FGPUFFTShaderContext& Context, const FFTDescription& FFTDesc,
-			const FIntRect& SrcWindow, const FTextureRHIRef& SrcTexture, FSceneRenderTargetItem& DstBuffer,
-			FSceneRenderTargetItem& TmpBuffer, const bool bScrubNaNs = false);
+		static void MultiPass(
+			FRDGBuilder& GraphBuilder,
+			const FGlobalShaderMap* ShaderMap,
+			const FFTDescription& FFTDesc,
+			FRDGTextureRef SrcTexture, const FIntRect& SrcWindow,
+			FRDGTextureRef DstTexture,
+			const bool bScrubNaNs = false);
 
 	};
 	
@@ -216,62 +138,6 @@ namespace GPUFFT
 
 	struct TwoForOneRealFFTImage1D
 	{
-		/**
-		* The requirements of the real two-for-one 1d FFT.
-		*
-		* @param FFTDesc         - Indicates the number of frequencies and scan lines as well as direction.
-		*                          Due to real-transform symmetries, the input frequency count is the number of 
-		*                          frequencies used for the complex transform (of two real singals togehter, e.g. z = r + i * g).
-		* @param MinBufferSize   - On return: The size of Dst buffer required to hold the full result. If multipass is required
-		*                          this size also applies to the tmpBuffer
-		* @param bUseGroupShared - On return: indicates if the transform requires the multipass
-		*                          or if we can use the much faster group shared variant.
-		* 
-		* NB: When doing the inverse transform in group-shared memory, a smaller buffer maybe used if the full result is not desired.
-		* for example, if only the result in a sub-window is required.
-		*/
-		static void Requirements(const FFTDescription& FFTDesc, FIntPoint& MinBufferSize, bool& bUseMultipass);
-
-
-		/**
-		* One Dimensional Fast Fourier Transform of four real signals (rgba) in a two dimensional buffer.
-		* The each scanline of the float4 buffer is interpreted four distinct real signals.
-		*
-		* On input a scanline (lenght N) contains 4 interlaced signals {r_n, g_n, b_n, a_n}.
-		* On output a scanline will be N + 2 elements long with the following data layout
-		* Entries i = [0,     N/2] hold  float4(R_i, B_i).
-		* Entries i = [N/2+1, N-1] hold  float4(G_i, A_i)
-		* Entries i = N            hold  float4(G_N/2, A_N/2)
-		* Entry   i = N + 1        hold  float4(G_0, A_0)
-		*
-		* Here R_k, G_k ,B_k, A_k are complex (float2) values representing the transforms of the real signals r,g,b,a
-		*
-		* Note: Due to the fact that r,g,b,a are real signals, R_0, B_0, G_0, A_0 are real values.
-		* and will equal the sum of the signals , i.e. R_0 = sum(r_i, i=0..N-1).
-		*
-		* Note if {x_n} is a real signal of length N, then the transfromed signal has the symmetry X_k = ComplexConjugate(X_{N-k})
-		*
-		* @param Context    - RHI and ShadersMap
-		* @param FFTDesc    - Defines the scanline direction and power-of-two length and transform type
-		* @param SrcWindow  - Location of the data in the SrcTexture to be transformed. Zero values pad
-		*                     the data to fit the  FFTDesc TrasformExtent.
-		* @param SrcTexture - The two dimensional float4 input buffer to transform
-		* @param DstWindow  - The region in the DstUAV to write the result of the transform.
-		* @param DstUAV     - The buffer to be populated by the result of the transform.
-		* @param PreFilter  - Values use to boost the intensity of any pixels.  Default value deactivates prefilter
-		*
-		* @return           - True on success.  Will only fail if the buffer is too large for supported transform method.
-		*
-		* NB: On a forward transform, the DstWindow must accommodate Context.SignalLength + 2 in the transform direction.
-		*
-		* NB: This function does not transition resources.  That must be done by the calling scope
-		*/
-		static bool GroupShared(FGPUFFTShaderContext& Context, const FFTDescription& FFTDesc,
-			const FIntRect& SrcWindow, const FTextureRHIRef& SrcTexture,
-			const FIntRect& DstWindow, FUnorderedAccessViewRHIRef DstUAV,
-			const FPreFilter& PreFilter = FPreFilter(TNumericLimits<float>::Max(), TNumericLimits<float>::Lowest(), 0.f));
-
-
 		/**
 		* One Dimensional Fast Fourier Transform of four real signals (rgba) in a two dimensional buffer.
 		* The each scanline of the float4 buffer is interpreted four distinct real signals.
@@ -307,71 +173,18 @@ namespace GPUFFT
 		*               On inverse transform the DstBuffer must accommodate Context.SignalLength in the transform direction.
 		* NB: This function does not transition resources.  That must be done by the calling scope
 		*/
-		static bool MultiPass(FGPUFFTShaderContext& Context, const FFTDescription& FFTDesc,
-			const FIntRect& SrcWindow, const FTextureRHIRef& SrcTexture,
-			FSceneRenderTargetItem& DstBuffer,
-			FSceneRenderTargetItem& TmpBuffer,
+		static void MultiPass(
+			FRDGBuilder& GraphBuilder,
+			const FGlobalShaderMap* ShaderMap,
+			const FFTDescription& FFTDesc,
+			FRDGTextureRef SrcTexture, const FIntRect& SrcWindow,
+			FRDGTextureRef DstTexture, const FIntRect& DstWindow,
 			const FPreFilter& PreFilter = FPreFilter(TNumericLimits<float>::Max(), TNumericLimits<float>::Lowest(), 0.f));
 
 	};
 
 	struct ConvolutionWithTextureImage1D
 	{
-
-		/**
-		* The requirements of the 1d convolution.
-		*
-		* @param FFTDesc         - Indicates the number of frequencies and scan lines as well as direction.
-		*                         
-		* @param MinBufferSize   - On return: The size of Dst buffer required. If multipass is required
-		*                          this size also applies to the tmpBuffer
-		* @param bUseGroupShared - On return: indicates if the transform requires the multipass
-		*                          or if we can use the much faster group shared variant
-		* NB: When doing the inverse transform in group-shared memory, a smaller buffer maybe used if the full result is not desired.
-		* for example, if only the result in a sub-window is required.
-		*/
-		static void Requirements(const FFTDescription& FFTDesc, FIntPoint& MinBufferSize, bool& bUseMultipass);
-
-		/**
-		* One Dimensional convolution against a pre-convolved texture: does two complex signals at once.
-		*
-		* The each scanline of the float4 buffer is interpreted as two complex signals (x+iy and z + iw).
-		*
-		* ------------------Breakdown of the convolution------------------------------------
-		* On input:  Each scanline (length N) of the SrcTexture contains 2 interlaced signals
-		* float4(a,b,c,d) = {z_n, q_n}  where z_n = a_n + i b_n  and q_n = c_n + i d_n
-		*
-		*            Each scanline (length N) of the TransformedKernel contains 2 interlaced, pre-transformed
-		*            signals  float4( TK1, TK2)
-		*
-		* This function transforms the Src scanline to give two interlaced frequency signals float4(Z_k, Q_k)
-		*
-		* The convolution is affected as a multiplication in frequency space:
-		*       float4( ComplexMult(Z_k, TK1_k), ComplexMult(Q_k, TK2_k))
-		*
-		* And inverse transform is applied to bring the result back to the input space.
-		* ---------------------------------------------------------------------------------
-		*
-		*
-		* @param Context    - RHI and ShadersMap
-		* @param FFTDesc    - Defines the scanline direction and count and transform type
-		* @param TransformedKernel - A texture that contains a transformed convolution kernel.
-		* @param SrcWindow  - Location of the data in the SrcTexture to be transformed. Zero values pad
-		*                     the data to fit the  FFTDesc TrasformExtent.
-		* @param SrcTexture - The two dimensional float4 input buffer to transform
-		* @param DstUAV     - The buffer to be populated by the result of the transform.
-		*
-		* @return           - True on success.  Will only fail if the buffer is too large for supported transform method
-		*
-		* NB: This function does not transition resources.  That must be done by the calling scope.
-		*/
-		static bool GroupShared(FGPUFFTShaderContext& Context, const FFTDescription& FFTDesc,
-			const FTextureRHIRef& TransformedKernel,
-			const FIntRect& SrcWindow,
-			const FTextureRHIRef& SrcTexture,
-			FUnorderedAccessViewRHIRef DstUAV);
-
-
 		/**
 		* One Dimensional convolution against a pre-convolved texture: does two complex signals at once.
 		*
@@ -406,12 +219,13 @@ namespace GPUFFT
 		*
 		* NB: This function does not transition resources.  That must be done by the calling scope.
 		*/
-		static bool MultiPass(FGPUFFTShaderContext& Context, const FFTDescription& FFTDesc,
-			const FTextureRHIRef& TransformedKernel,
-			const FIntRect& SrcWindow,
-			const FTextureRHIRef& SrcTexture,
-			FSceneRenderTargetItem& DstBuffer,
-			FSceneRenderTargetItem& TmpBuffer);
+		static void MultiPass(
+			FRDGBuilder& GraphBuilder,
+			const FGlobalShaderMap* ShaderMap,
+			const FFTDescription& FFTDesc,
+			FRDGTextureRef SrcTexture, const FIntRect& SrcWindow,
+			FRDGTextureRef TransformedKernel,
+			FRDGTextureRef DstTexture);
 	};
 
 
@@ -431,10 +245,12 @@ namespace GPUFFT
 	* NB: This function does not transition resources on the Src or Target.  That must be done by the calling scope.
 	*
 	*/
-	bool FFTImage2D(FGPUFFTShaderContext& Context, const FIntPoint& FrequencySize, bool bHorizontalFirst,
-		const FIntRect& ROIRect/*region of interest*/, const FTextureRHIRef& SrcTexture,
-		FSceneRenderTargetItem& ResultBuffer,
-		FSceneRenderTargetItem& TmpBuffer0);
+	void FFTImage2D(
+		FRDGBuilder& GraphBuilder,
+		const FGlobalShaderMap* ShaderMap,
+		const FIntPoint& FrequencySize, bool bHorizontalFirst,
+		FRDGTextureRef SrcTexture, const FIntRect& ROIRect,
+		FRDGTextureRef ResultTexture);
 
 
 
@@ -462,13 +278,14 @@ namespace GPUFFT
 	*
 	* The Src and Dst textures can be re-used as the tmp buffers (Dst->TmpBuffer0, Src->TmpBuffer1) provided they are large enough.
 	*/
-	bool ConvolutionWithTextureImage2D(FGPUFFTShaderContext& Context, const FIntPoint& FrequencySize, bool bHorizontalFirst,
-		const FTextureRHIRef& TransformedKernel,
-		const FIntRect& ROIRect/*region of interest*/,
-		const FTextureRHIRef& SrcTexture, 
-		FUnorderedAccessViewRHIRef ResultUAV,
-		FSceneRenderTargetItem& TmpBuffer0,
-		FSceneRenderTargetItem& TmpBuffer1,
+	void ConvolutionWithTextureImage2D(
+		FRDGBuilder& GraphBuilder,
+		const FGlobalShaderMap* ShaderMap,
+		const FIntPoint& FrequencySize,
+		bool bHorizontalFirst,
+		FRDGTextureRef TransformedKernel,
+		FRDGTextureRef SrcTexture, const FIntRect& ROIRect,
+		FRDGTextureRef DstTexture, const FIntRect& DstRect,
 		const FPreFilter& PreFilter);
 
 	FIntPoint Convolution2DBufferSize(const FIntPoint& FrequencySize, const bool bHorizontalFirst, const FIntPoint& SrcExtent);
@@ -489,118 +306,10 @@ namespace GPUFFT
 
 	* NB: This function does not transition resources.  That must be done by the calling scope.
 	*/
-	void CopyImage2D(FGPUFFTShaderContext& Context,
-		const FIntRect& SrcWindow, const FTextureRHIRef& SrcTexture,
-		const FIntRect& DstWindow, FUnorderedAccessViewRHIRef DstUAV,
+	void CopyImage2D(
+		FRDGBuilder& GraphBuilder,
+		const FGlobalShaderMap* ShaderMap,
+		FRDGTextureRef SrcTexture, const FIntRect& SrcWindow,
+		FRDGTextureRef DstTexture, const FIntRect& DstWindow,
 		const FPreFilter& PreFilter = FPreFilter(TNumericLimits<float>::Max(), TNumericLimits<float>::Lowest(), 0.f));
 };
-
-// Collection of utilities that make interfacing with compute shaders a little more nice
-namespace GPUFFTComputeShaderUtils
-{
-
-
-	class FComputeParamterValueSetter
-	{
-	public:
-		FComputeParamterValueSetter(FRHICommandList& CmdList, FRHIComputeShader* InShaderRHI) :
-			RHICmdList(&CmdList), ShaderRHI(InShaderRHI) {}
-
-
-		template < typename TValue >
-		inline FComputeParamterValueSetter& operator()(FShaderParameter& Parameter, const TValue& Value)
-		{
-			SetShaderValue(*RHICmdList, ShaderRHI, Parameter, Value);
-			return *this;
-		}
-
-		template <typename TValue>
-		inline FComputeParamterValueSetter& operator()(FShaderParameter& Parameter, const TValue* Value, const uint32 NumElements)
-		{
-			SetShaderValueArray(*RHICmdList, ShaderRHI, Parameter, Value, NumElements);
-			return *this;
-		}
-
-		inline FComputeParamterValueSetter& operator()(FShaderResourceParameter& TextureParameter, FRHITexture* TextureRHI)
-		{
-			if (TextureParameter.IsBound())
-			{
-				SetTextureParameter(*RHICmdList, ShaderRHI, TextureParameter, TextureRHI);
-			}
-			return *this;
-		}
-
-
-		template <ESamplerFilter Filter, ESamplerAddressMode AddressMode>
-		inline FComputeParamterValueSetter& Set(FShaderResourceParameter& TextureParameter, FShaderResourceParameter& SamplerParameter, FRHITexture* TextureRHI)
-		{
-			SetTextureParameter(*RHICmdList, ShaderRHI, TextureParameter, SamplerParameter,
-				TStaticSamplerState<Filter, AddressMode, AddressMode, AddressMode>::GetRHI(), TextureRHI);
-			return *this;
-		}
-
-	private:
-		FRHICommandList* RHICmdList;
-		FRHIComputeShader* ShaderRHI;
-	};
-
-	class FComputeParameterBinder
-	{
-	public:
-		FComputeParameterBinder(const FShaderParameterMap& ParameterMap) : Map(&ParameterMap) {}
-
-		template <typename ParameterType>
-		inline const FComputeParameterBinder& operator()(ParameterType& Parameter, const TCHAR* Name) const
-		{
-			Parameter.Bind(*Map, Name);
-			return *this;
-		}
-
-		template <typename ParameterType>
-		inline const FComputeParameterBinder& operator()(ParameterType& Parameter) const
-		{
-			Parameter.Bind(*Map);
-			return *this;
-		}
-	private:
-		const FShaderParameterMap* Map;
-	};
-
-	class FScopedUAVBind
-	{
-	public:
-
-		// Factory method.
-		template <typename ComputeShaderT>
-		inline static FScopedUAVBind BindOutput(FRHICommandList& CmdList, const TShaderRef<ComputeShaderT>& Shader, const FUnorderedAccessViewRHIRef& BufferUAV)
-		{
-			FRHIComputeShader* ShaderRHI = Shader.GetComputeShader();
-			const uint32 BaseIndex = Shader->DestinationResourceParameter().GetBaseIndex();
-			return FScopedUAVBind(CmdList, ShaderRHI, BaseIndex, BufferUAV);
-		}
-
-
-		// destructor unbinds the UAV
-		~FScopedUAVBind()
-		{
-			RHICmdList->SetUAVParameter(ComputeShader, Index, NULL);
-		}
-
-
-	private:
-		FScopedUAVBind();
-
-		FScopedUAVBind(FRHICommandList& CmdList, FRHIComputeShader* ShaderRHI, const uint32 BaseIndex, FUnorderedAccessViewRHIRef BufferUAV)
-			:ComputeShader(ShaderRHI), RHICmdList(&CmdList), Index(BaseIndex)
-		{
-			RHICmdList->SetUAVParameter(ComputeShader, Index, BufferUAV);
-		}
-
-
-		FRHIComputeShader* ComputeShader;
-		FRHICommandList* RHICmdList;
-		uint32 Index;
-	};
-
-
-}; // end namespace GPUFFTComputeShaderUtils
diff --git a/Engine/Source/Runtime/Renderer/Private/PostProcess/PostProcessFFTBloom.cpp b/Engine/Source/Runtime/Renderer/Private/PostProcess/PostProcessFFTBloom.cpp
--- a/Engine/Source/Runtime/Renderer/Private/PostProcess/PostProcessFFTBloom.cpp
+++ b/Engine/Source/Runtime/Renderer/Private/PostProcess/PostProcessFFTBloom.cpp
@@ -3,6 +3,7 @@
 #include "PostProcess/PostProcessFFTBloom.h"
 #include "GPUFastFourierTransform.h"
 #include "RendererModule.h"
+#include "Rendering/Texture2DResource.h"
 
 namespace
 {
@@ -44,13 +45,15 @@ public:
 	SHADER_USE_PARAMETER_STRUCT(FResizeAndCenterTextureCS, FFFTBloomShader);
 
 	BEGIN_SHADER_PARAMETER_STRUCT(FParameters, )
-		SHADER_PARAMETER_TEXTURE(Texture2D, SrcTexture)
-		SHADER_PARAMETER_SAMPLER(SamplerState, SrcSampler)
-		SHADER_PARAMETER_UAV(RWTexture2D<float4>, DstTexture)
 		SHADER_PARAMETER(FIntPoint, DstExtent)
 		SHADER_PARAMETER(FIntPoint, ImageExtent)
 		SHADER_PARAMETER(FVector4, KernelCenterAndScale)
 		SHADER_PARAMETER(FIntPoint, DstBufferExtent)
+
+		SHADER_PARAMETER_TEXTURE(Texture2D, SrcTexture)
+		SHADER_PARAMETER_SAMPLER(SamplerState, SrcSampler)
+
+		SHADER_PARAMETER_RDG_TEXTURE_UAV(RWTexture2D<float4>, DstTexture)
 	END_SHADER_PARAMETER_STRUCT()
 
 	static void ModifyCompilationEnvironment(const FGlobalShaderPermutationParameters& Parameters, FShaderCompilerEnvironment& OutEnvironment)
@@ -61,8 +64,6 @@ public:
 	}
 };
 
-IMPLEMENT_GLOBAL_SHADER(FResizeAndCenterTextureCS, "/Engine/Private/PostProcessFFTBloom.usf", "ResizeAndCenterTextureCS", SF_Compute);
-
 class FCaptureKernelWeightsCS : public FFFTBloomShader
 {
 public:
@@ -70,12 +71,14 @@ public:
 	SHADER_USE_PARAMETER_STRUCT(FCaptureKernelWeightsCS, FFFTBloomShader);
 
 	BEGIN_SHADER_PARAMETER_STRUCT(FParameters, )
-		SHADER_PARAMETER_TEXTURE(Texture2D, HalfResSrcTexture)
-		SHADER_PARAMETER_TEXTURE(Texture2D, PhysicalSrcTexture)
-		SHADER_PARAMETER_SAMPLER(SamplerState, PhysicalSrcSampler)
-		SHADER_PARAMETER_UAV(RWTexture2D<float4>, DstTexture)
 		SHADER_PARAMETER(FIntPoint, HalfResSumLocation)
 		SHADER_PARAMETER(FVector2D, UVCenter)
+
+		SHADER_PARAMETER_SAMPLER(SamplerState, PhysicalSrcSampler)
+		SHADER_PARAMETER_RDG_TEXTURE(Texture2D, HalfResSrcTexture)
+		SHADER_PARAMETER_TEXTURE(Texture2D, PhysicalSrcTexture)
+
+		SHADER_PARAMETER_RDG_TEXTURE_UAV(RWTexture2D<float4>, DstTexture)
 	END_SHADER_PARAMETER_STRUCT()
 
 	static void ModifyCompilationEnvironment(const FGlobalShaderPermutationParameters& Parameters, FShaderCompilerEnvironment& OutEnvironment)
@@ -85,8 +88,6 @@ public:
 	}
 };
 
-IMPLEMENT_GLOBAL_SHADER(FCaptureKernelWeightsCS, "/Engine/Private/PostProcessFFTBloom.usf", "CaptureKernelWeightsCS", SF_Compute);
-
 class FBlendLowResCS : public FFFTBloomShader
 {
 public:
@@ -94,48 +95,30 @@ public:
 	SHADER_USE_PARAMETER_STRUCT(FBlendLowResCS, FFFTBloomShader);
 
 	BEGIN_SHADER_PARAMETER_STRUCT(FParameters, )
-		SHADER_PARAMETER_TEXTURE(Texture2D, SrcTexture)
-		SHADER_PARAMETER_TEXTURE(Texture2D, HalfResSrcTexture)
-		SHADER_PARAMETER_SAMPLER(SamplerState, HalfResSrcSampler)
-		SHADER_PARAMETER_TEXTURE(Texture2D, CenterWeightTexture)
-		SHADER_PARAMETER_UAV(RWTexture2D<float4>, DstTexture)
 		SHADER_PARAMETER(FIntRect, DstRect)
 		SHADER_PARAMETER(FIntRect, HalfRect)
 		SHADER_PARAMETER(FIntPoint, HalfBufferSize)
-	END_SHADER_PARAMETER_STRUCT()
 
-	static void ModifyCompilationEnvironment(const FGlobalShaderPermutationParameters& Parameters, FShaderCompilerEnvironment& OutEnvironment)
-	{
-		FFFTBloomShader::ModifyCompilationEnvironment(Parameters, OutEnvironment);
-		OutEnvironment.SetDefine(TEXT("INCLUDE_BLEND_LOW_RES"), 1);
-		OutEnvironment.SetDefine(TEXT("THREADS_PER_GROUP"), ThreadsPerGroup);
-	}
-};
-
-IMPLEMENT_GLOBAL_SHADER(FBlendLowResCS, "/Engine/Private/PostProcessFFTBloom.usf", "BlendLowResCS", SF_Compute);
-
-class FPassThroughCS : public FFFTBloomShader
-{
-public:
-	DECLARE_GLOBAL_SHADER(FPassThroughCS);
-	SHADER_USE_PARAMETER_STRUCT(FPassThroughCS, FFFTBloomShader);
+		SHADER_PARAMETER_RDG_TEXTURE(Texture2D, SrcTexture)
+		SHADER_PARAMETER_RDG_TEXTURE(Texture2D, HalfResSrcTexture)
+		SHADER_PARAMETER_SAMPLER(SamplerState, HalfResSrcSampler)
+		SHADER_PARAMETER_RDG_TEXTURE(Texture2D, CenterWeightTexture)
 
-	BEGIN_SHADER_PARAMETER_STRUCT(FParameters, )
-		SHADER_PARAMETER_TEXTURE(Texture2D, SrcTexture)
-		SHADER_PARAMETER_UAV(RWTexture2D<float4>, DstTexture)
-		SHADER_PARAMETER(FIntRect, SrcRect)
-		SHADER_PARAMETER(FIntRect, DstRect)
+		SHADER_PARAMETER_RDG_TEXTURE_UAV(RWTexture2D<float4>, DstTexture)
 	END_SHADER_PARAMETER_STRUCT()
 
 	static void ModifyCompilationEnvironment(const FGlobalShaderPermutationParameters& Parameters, FShaderCompilerEnvironment& OutEnvironment)
 	{
 		FFFTBloomShader::ModifyCompilationEnvironment(Parameters, OutEnvironment);
-		OutEnvironment.SetDefine(TEXT("INCLUDE_PASSTHROUGH"), 1);
+		OutEnvironment.SetDefine(TEXT("INCLUDE_BLEND_LOW_RES"), 1);
 		OutEnvironment.SetDefine(TEXT("THREADS_PER_GROUP"), ThreadsPerGroup);
 	}
 };
 
-IMPLEMENT_GLOBAL_SHADER(FPassThroughCS, "/Engine/Private/PostProcessFFTBloom.usf", "PassThroughCS", SF_Compute);
+IMPLEMENT_GLOBAL_SHADER(FResizeAndCenterTextureCS, "/Engine/Private/PostProcessFFTBloom.usf", "ResizeAndCenterTextureCS", SF_Compute);
+IMPLEMENT_GLOBAL_SHADER(FCaptureKernelWeightsCS,   "/Engine/Private/PostProcessFFTBloom.usf", "CaptureKernelWeightsCS", SF_Compute);
+IMPLEMENT_GLOBAL_SHADER(FBlendLowResCS,            "/Engine/Private/PostProcessFFTBloom.usf", "BlendLowResCS", SF_Compute);
+
 } //! namespace
 
 bool IsFFTBloomHalfResolutionEnabled()
@@ -148,51 +131,6 @@ bool IsFFTBloomQuarterResolutionEnabled()
 	return CVarHalfResFFTBloom.GetValueOnRenderThread() == 2;
 }
 
-bool IsFFTBloomPhysicalKernelReady(const FViewInfo& View)
-{
-	UTexture2D* BloomConvolutionTexture = View.FinalPostProcessSettings.BloomConvolutionTexture;
-
-	// Fall back to the default bloom texture if provided.
-	if (BloomConvolutionTexture == nullptr)
-	{
-		BloomConvolutionTexture = GEngine->DefaultBloomKernelTexture;
-	}
-
-	bool bValidSetup = (BloomConvolutionTexture != nullptr && BloomConvolutionTexture->GetResource() != nullptr);
-
-	// The bloom convolution kernel needs to be resident to avoid visual artifacts.
-	if (bValidSetup)
-	{
-		const int32 CinematicTextureGroups = 0;
-		const float Seconds = 5.0f;
-		BloomConvolutionTexture->SetForceMipLevelsToBeResident(Seconds, CinematicTextureGroups);
-	}
-
-	const uint32 FramesPerWarning = 15;
-
-	if (bValidSetup && BloomConvolutionTexture->IsFullyStreamedIn() == false)
-	{
-		if ((View.Family->FrameNumber % FramesPerWarning) == 0)
-		{
-			UE_LOG(LogRenderer, Warning, TEXT("The Physical Kernel Texture not fully streamed in."));
-		}
-	}
-
-	bValidSetup = bValidSetup && (BloomConvolutionTexture->IsFullyStreamedIn() == true);
-
-	if (bValidSetup && BloomConvolutionTexture->bHasStreamingUpdatePending == true)
-	{
-		if ((View.Family->FrameNumber % FramesPerWarning) == 0)
-		{
-			UE_LOG(LogRenderer, Warning, TEXT("The Physical Kernel Texture has pending update."));
-		}
-	}
-
-	bValidSetup = bValidSetup && (BloomConvolutionTexture->bHasStreamingUpdatePending == false);
-
-	return bValidSetup;
-}
-
 bool IsFFTBloomEnabled(const FViewInfo& View)
 {
 	const bool bOldMetalNoFFT = IsMetalPlatform(View.GetShaderPlatform()) && (RHIGetShaderLanguageVersion(View.GetShaderPlatform()) < 4) && IsPCPlatform(View.GetShaderPlatform());
@@ -208,7 +146,7 @@ bool IsFFTBloomEnabled(const FViewInfo& View)
 
 	if (bUseFFTBloom && !bOldMetalNoFFT)
 	{
-		return IsFFTBloomPhysicalKernelReady(View);
+		return View.FFTBloomKernelTexture != nullptr;
 	}
 	else
 	{
@@ -233,95 +171,83 @@ bool IsFFTBloomEnabled(const FViewInfo& View)
 * @param bForceCenterZero -  is true only for the experimental 1/2 res version, part of conserving energy
 */
 void ResizeAndCenterTexture(
-	FRHICommandList& RHICmdList,
+	FRDGBuilder& GraphBuilder,
 	const FViewInfo& View,
 	const FTextureRHIRef& SrcTexture,
 	const FIntPoint& SrcImageSize,
 	const FVector2D& SrcImageCenterUV,
 	const float ResizeScale,
 	const FIntPoint& TargetSize,
-	FUnorderedAccessViewRHIRef& DstUAV,
+	FRDGTextureRef DstTexture,
 	const FIntPoint& DstBufferSize,
 	const bool bForceCenterZero)
 {
-	SCOPED_DRAW_EVENTF(RHICmdList, FRCPassFFTBloom, TEXT("FFT: Pre-process the space kernel to %d by %d"), TargetSize.X, TargetSize.Y);
-
 	// Clamp the image center
 	FVector2D ClampedImageCenterUV = SrcImageCenterUV;
 	ClampedImageCenterUV.X = FMath::Clamp(SrcImageCenterUV.X, 0.f, 1.f);
 	ClampedImageCenterUV.Y = FMath::Clamp(SrcImageCenterUV.Y, 0.f, 1.f);
 
-	check(DstUAV);
-	RHICmdList.Transition(FRHITransitionInfo(DstUAV, ERHIAccess::Unknown, ERHIAccess::UAVCompute));
-
-	{
-		TShaderMapRef<FResizeAndCenterTextureCS> ComputeShader(View.ShaderMap);
-
-		float CenterScale = bForceCenterZero ? 0.f : 1.f;
-		const FLinearColor KernelCenterAndScale(ClampedImageCenterUV.X, ClampedImageCenterUV.Y, ResizeScale, CenterScale);
-
-		FResizeAndCenterTextureCS::FParameters PassParameters;
-		PassParameters.SrcTexture = SrcTexture;
-		PassParameters.SrcSampler = TStaticSamplerState<SF_Bilinear, AM_Wrap, AM_Wrap, AM_Wrap>::GetRHI();
-		PassParameters.DstTexture = DstUAV;
-		PassParameters.DstExtent = TargetSize;
-		PassParameters.ImageExtent = SrcImageSize;
-		PassParameters.KernelCenterAndScale = KernelCenterAndScale;
-		PassParameters.DstBufferExtent = DstBufferSize;
-
-		// Use multiple threads per scan line to insure memory coalescing during the write
-		const int32 ThreadsPerGroup = FResizeAndCenterTextureCS::ThreadsPerGroup;
-		const int32 ThreadsGroupsPerScanLine = (DstBufferSize.X % ThreadsPerGroup == 0) ? DstBufferSize.X / ThreadsPerGroup : DstBufferSize.X / ThreadsPerGroup + 1;
-		FComputeShaderUtils::Dispatch(RHICmdList, ComputeShader, PassParameters, FIntVector(ThreadsGroupsPerScanLine, DstBufferSize.Y, 1));
-	}
+	check(DstTexture);
+
+	float CenterScale = bForceCenterZero ? 0.f : 1.f;
+	const FLinearColor KernelCenterAndScale(ClampedImageCenterUV.X, ClampedImageCenterUV.Y, ResizeScale, CenterScale);
+
+	FResizeAndCenterTextureCS::FParameters* PassParameters = GraphBuilder.AllocParameters<FResizeAndCenterTextureCS::FParameters>();
+	PassParameters->DstExtent = TargetSize;
+	PassParameters->ImageExtent = SrcImageSize;
+	PassParameters->KernelCenterAndScale = KernelCenterAndScale;
+	PassParameters->DstBufferExtent = DstBufferSize;
+
+	PassParameters->SrcTexture = SrcTexture;
+	PassParameters->SrcSampler = TStaticSamplerState<SF_Bilinear, AM_Wrap, AM_Wrap, AM_Wrap>::GetRHI();
+	PassParameters->DstTexture = GraphBuilder.CreateUAV(DstTexture);
+
+	// Use multiple threads per scan line to insure memory coalescing during the write
+	const int32 ThreadsPerGroup = FResizeAndCenterTextureCS::ThreadsPerGroup;
+	const int32 ThreadsGroupsPerScanLine = (DstBufferSize.X % ThreadsPerGroup == 0) ? DstBufferSize.X / ThreadsPerGroup : DstBufferSize.X / ThreadsPerGroup + 1;
+	
+	TShaderMapRef<FResizeAndCenterTextureCS> ComputeShader(View.ShaderMap);
+	FComputeShaderUtils::AddPass(
+		GraphBuilder,
+		RDG_EVENT_NAME("FFTBloom PreProcessKernel %dx%d", TargetSize.X, TargetSize.Y),
+		ComputeShader,
+		PassParameters, FIntVector(ThreadsGroupsPerScanLine, DstBufferSize.Y, 1));
 }
 
-/**
-* Used by experimental energy conserving 1/2 resolution version of the bloom.
-* Captures the sum of the kernel weights represented by the 1/2 res kernel and
-* the Center weight from the physical space kernel.
-*
-* @param Context            - container for RHI and ShaderMap
-* @param HalfResKernel      - SRV for the pre-transformed 1/2 res kernel
-* @param HalfResSumLocation - The location to sample in the pre-transformed kernel to find the sum of the physical space kernel weights
-* @param PhysicalKernel     - SRV for the original physical space kernel
-* @param CenterUV           - Where to sample the Physical Kernel for the center weight
-* @param CenterWeightRT     - 2x1 float4 buffer that on return will hold result:
-*     At  (0,0) the center weight of physical kernel, and (1,0) the sum of the 1/2res kernel weights
-*/
-void CaptureKernelWeight(
-	FRHICommandList& RHICmdList,
+FRDGTextureRef CaptureKernelWeight(
+	FRDGBuilder& GraphBuilder,
 	const FViewInfo& View,
-	const FTextureRHIRef& HalfResKernel,
+	FRDGTextureRef HalfResKernel,
 	const FIntPoint& HalfResSumLocation,
 	const FTextureRHIRef& PhysicalKernel,
-	const FVector2D& CenterUV,
-	TRefCountPtr<IPooledRenderTarget>& CenterWeightRT)
+	const FVector2D& CenterUV)
 {
-	SCOPED_DRAW_EVENTF(RHICmdList, FRCPassFFTBloom, TEXT("FFT: Capture Kernel Weights"));
-
-	FSceneRenderTargetItem& DstTargetItem = CenterWeightRT->GetRenderTargetItem();
-
-	check(DstTargetItem.UAV);
-	RHICmdList.Transition(FRHITransitionInfo(DstTargetItem.UAV, ERHIAccess::Unknown, ERHIAccess::UAVCompute));
-
-	{
-		TShaderMapRef<FCaptureKernelWeightsCS> ComputeShader(View.ShaderMap);
-
-		FCaptureKernelWeightsCS::FParameters PassParameters;
-		PassParameters.HalfResSrcTexture = HalfResKernel;
-		PassParameters.PhysicalSrcTexture = PhysicalKernel;
-		PassParameters.PhysicalSrcSampler = TStaticSamplerState<SF_Bilinear, AM_Wrap, AM_Wrap, AM_Wrap>::GetRHI();
-		PassParameters.DstTexture = DstTargetItem.UAV;
-		PassParameters.HalfResSumLocation = HalfResSumLocation;
-		PassParameters.UVCenter = CenterUV;
-
-		FComputeShaderUtils::Dispatch(RHICmdList, ComputeShader, PassParameters, FIntVector(1, 1, 1));
-	}
-
-	RHICmdList.Transition(FRHITransitionInfo(DstTargetItem.UAV, ERHIAccess::UAVCompute, ERHIAccess::SRVMask));
-
-	ensureMsgf(DstTargetItem.TargetableTexture == DstTargetItem.ShaderResourceTexture, TEXT("%s should be resolved to a separate SRV"), *DstTargetItem.TargetableTexture->GetName().ToString());
+	FRDGTextureDesc CenterWeightDesc = FRDGTextureDesc::Create2D(
+		FIntPoint(2, 1),
+		GPUFFT::PixelFormat(),
+		FClearValueBinding::None,
+		TexCreate_ShaderResource | TexCreate_UAV);
+
+	// Resize the buffer to hold the transformed kernel
+	FRDGTextureRef CenterWeightTexture = GraphBuilder.CreateTexture(CenterWeightDesc, TEXT("Bloom.FFT.KernelCenterWeight"));
+
+	FCaptureKernelWeightsCS::FParameters* PassParameters = GraphBuilder.AllocParameters<FCaptureKernelWeightsCS::FParameters>();
+	PassParameters->HalfResSumLocation = HalfResSumLocation;
+	PassParameters->UVCenter = CenterUV;
+	PassParameters->PhysicalSrcSampler = TStaticSamplerState<SF_Bilinear, AM_Wrap, AM_Wrap, AM_Wrap>::GetRHI();
+	PassParameters->HalfResSrcTexture = HalfResKernel;
+	PassParameters->PhysicalSrcTexture = PhysicalKernel;
+	PassParameters->DstTexture = GraphBuilder.CreateUAV(CenterWeightTexture);
+
+	TShaderMapRef<FCaptureKernelWeightsCS> ComputeShader(View.ShaderMap);
+	FComputeShaderUtils::AddPass(
+		GraphBuilder,
+		RDG_EVENT_NAME("FFTBloom CaptureKernelWeights"),
+		ComputeShader,
+		PassParameters, 
+		FIntVector(1, 1, 1));
+
+	return CenterWeightTexture;
 }
 
 /**
@@ -342,137 +268,87 @@ void CaptureKernelWeight(
 * @param DstUAV                - Destination buffer that will hold the result.
 */
 void BlendLowRes(
-	FRHICommandList& RHICmdList,
+	FRDGBuilder& GraphBuilder,
 	const FViewInfo& View,
-	const FTextureRHIRef& FullResImage,
+	FRDGTextureRef FullResImage,
 	const FIntRect& FullResImageRect,
-	const FTextureRHIRef& HalfResConvolvedImage,
+	FRDGTextureRef HalfResConvolvedImage,
 	const FIntRect& HalfResRect,
 	const FIntPoint& HalfBufferSize,
-	const FTextureRHIRef& CenterWeightTexutre,
-	FUnorderedAccessViewRHIRef DstUAV)
-{
-	SCOPED_DRAW_EVENTF(RHICmdList, FRCPassFFTBloom, TEXT("FFT: Post-process upres and blend"));
-
-	// set destination
-	check(DstUAV);
-	RHICmdList.Transition(FRHITransitionInfo(DstUAV, ERHIAccess::Unknown, ERHIAccess::UAVCompute));
-
-	{
-		TShaderMapRef<FBlendLowResCS> ComputeShader(View.ShaderMap);
-
-		FBlendLowResCS::FParameters PassParameters;
-		PassParameters.SrcTexture = FullResImage;
-		PassParameters.DstTexture = DstUAV;
-		PassParameters.HalfResSrcTexture = HalfResConvolvedImage;
-		PassParameters.HalfResSrcSampler = TStaticSamplerState<SF_Bilinear, AM_Wrap, AM_Wrap, AM_Wrap>::GetRHI();
-		PassParameters.CenterWeightTexture = CenterWeightTexutre;
-		PassParameters.DstRect = FullResImageRect;
-		PassParameters.HalfRect = HalfResRect;
-		PassParameters.HalfBufferSize = HalfBufferSize;
-
-		const FIntPoint TargetExtent = FullResImageRect.Size();
-
-		// Use multiple threads per scan line to insure memory coalescing during the write
-		const int32 ThreadsPerGroup = ComputeShader->ThreadsPerGroup;
-		const int32 ThreadsGroupsPerScanLine = (TargetExtent.X % ThreadsPerGroup == 0) ? TargetExtent.X / ThreadsPerGroup : TargetExtent.X / ThreadsPerGroup + 1;
-
-		FComputeShaderUtils::Dispatch(RHICmdList, ComputeShader, PassParameters, FIntVector(ThreadsGroupsPerScanLine, TargetExtent.Y, 1));
-	}
-
-	RHICmdList.Transition(FRHITransitionInfo(DstUAV, ERHIAccess::UAVCompute, ERHIAccess::SRVMask));
-}
-
-/**
-* Used to copy the input image in the event that it is too large to bloom (i.e. doesn't fit in the FFT group shared memory)
-*
-* @param SrcTargetItem  - The SrcBuffer to be copied.
-* @param SrcRect        - The region in the SrcBuffer to copy
-* @param DstUAV         - The target buffer for the copy
-* @param DstRect        - The location and region in the target buffer for the copy
-*/
-void CopyImageRect(
-	FRHICommandList& RHICmdList,
-	const FViewInfo& View,
-	const FSceneRenderTargetItem& SrcTargetItem,
-	const FIntRect& SrcRect,
-	FUnorderedAccessViewRHIRef& DstUAV,
-	const FIntRect& DstRect)
+	FRDGTextureRef CenterWeightTexutre,
+	FRDGTextureRef DstTexture)
 {
-	SCOPED_DRAW_EVENTF(RHICmdList, FRCPassFFTBloom, TEXT("FFT: passthrough "));
-
-	// set destination
-	check(DstUAV);
-	RHICmdList.Transition(FRHITransitionInfo(DstUAV, ERHIAccess::Unknown, ERHIAccess::UAVCompute));
-
-	{
-		TShaderMapRef<FPassThroughCS> ComputeShader(View.ShaderMap);
-
-		FPassThroughCS::FParameters PassParameters;
-		PassParameters.SrcTexture = SrcTargetItem.ShaderResourceTexture;
-		PassParameters.DstTexture = DstUAV;
-		PassParameters.SrcRect = SrcRect;
-		PassParameters.DstRect = DstRect;
-
-		const FIntPoint DstRectSize = DstRect.Size();
-
-		// Use multiple threads per scan line to insure memory coalescing during the write
-		const int32 ThreadsPerGroup = ComputeShader->ThreadsPerGroup;
-		const int32 ThreadsGroupsPerScanLine = (DstRectSize.X % ThreadsPerGroup == 0) ? DstRectSize.X / ThreadsPerGroup : DstRectSize.X / ThreadsPerGroup + 1;
-
-		FComputeShaderUtils::Dispatch(RHICmdList, ComputeShader, PassParameters, FIntVector(ThreadsGroupsPerScanLine, DstRectSize.Y, 1));
-	}
+	FBlendLowResCS::FParameters* PassParameters = GraphBuilder.AllocParameters<FBlendLowResCS::FParameters>();
+	PassParameters->DstRect = FullResImageRect;
+	PassParameters->HalfRect = HalfResRect;
+	PassParameters->HalfBufferSize = HalfBufferSize;
+
+	PassParameters->SrcTexture = FullResImage;
+	PassParameters->HalfResSrcTexture = HalfResConvolvedImage;
+	PassParameters->HalfResSrcSampler = TStaticSamplerState<SF_Bilinear, AM_Wrap, AM_Wrap, AM_Wrap>::GetRHI();
+	PassParameters->CenterWeightTexture = CenterWeightTexutre;
+
+	PassParameters->DstTexture = GraphBuilder.CreateUAV(DstTexture);
+
+	TShaderMapRef<FBlendLowResCS> ComputeShader(View.ShaderMap);
+
+	const FIntPoint TargetSize = FullResImageRect.Size();
+
+	// Use multiple threads per scan line to insure memory coalescing during the write
+	const int32 ThreadsPerGroup = ComputeShader->ThreadsPerGroup;
+	const int32 ThreadsGroupsPerScanLine = (TargetSize.X % ThreadsPerGroup == 0) ? TargetSize.X / ThreadsPerGroup : TargetSize.X / ThreadsPerGroup + 1;
+
+	FComputeShaderUtils::AddPass(
+		GraphBuilder,
+		RDG_EVENT_NAME("FFTBloom Upscale&Blend %dx%d -> %dx%d",
+			HalfResRect.Width(),
+			HalfResRect.Height(),
+			FullResImageRect.Width(),
+			FullResImageRect.Height()),
+		ComputeShader,
+		PassParameters,
+		FIntVector(ThreadsGroupsPerScanLine, TargetSize.Y, 1));
 }
 
-bool TransformKernelFFT(
-	FRHICommandList& RHICmdList,
+FRDGTextureRef TransformKernelFFT(
+	FRDGBuilder& GraphBuilder,
 	const FViewInfo& View,
-	FSceneRenderTargetItem& KernelTargetItem,
+	FRDGTextureRef ResizedKernel,
 	bool bDoHorizontalFirst,
 	FIntPoint FrequencySize)
 {
-	GPUFFT::FGPUFFTShaderContext FFTContext(RHICmdList, *View.ShaderMap);
-
-	// Create the tmp buffer
-
 	// Our frequency storage layout adds two elements to the first transform direction. 
 	const FIntPoint FrequencyPadding = (bDoHorizontalFirst) ? FIntPoint(2, 0) : FIntPoint(0, 2);
 	const FIntPoint PaddedFrequencySize = FrequencySize + FrequencyPadding;
 
 	// Should read / write to PF_G16R16F or PF_G32R32F (float2 formats)
 	// Need to set the render target description before we "request surface"
-	const EPixelFormat PixelFormat = GPUFFT::PixelFormat();
-	FPooledRenderTargetDesc Desc = FPooledRenderTargetDesc::Create2DDesc(PaddedFrequencySize, PixelFormat,
-		FClearValueBinding::None, TexCreate_None, TexCreate_RenderTargetable | TexCreate_UAV, false);
-
-	// Temp buffer used at intermediate buffer when transforming the world space kernel 
-	TRefCountPtr<IPooledRenderTarget> TmpRT;
-	GRenderTargetPool.FindFreeElement(RHICmdList, Desc, TmpRT, TEXT("FFT Tmp Kernel Buffer"));
-
-	FIntRect SrcRect(FIntPoint(0, 0), FrequencySize);
-	const FTextureRHIRef& SrcImage = KernelTargetItem.ShaderResourceTexture;
-	FSceneRenderTargetItem& ResultBuffer = KernelTargetItem;
+	FRDGTextureRef SpectralKernel;
+	{
 
-	bool SuccessValue = GPUFFT::FFTImage2D(FFTContext, FrequencySize, bDoHorizontalFirst, SrcRect, SrcImage, ResultBuffer, TmpRT->GetRenderTargetItem());
+		FRDGTextureDesc Desc = FRDGTextureDesc::Create2D(
+			PaddedFrequencySize,
+			GPUFFT::PixelFormat(),
+			FClearValueBinding::None,
+			TexCreate_ShaderResource | TexCreate_UAV);
 
-	// Transition resource
-	RHICmdList.Transition(FRHITransitionInfo(ResultBuffer.UAV, ERHIAccess::Unknown, ERHIAccess::SRVMask));
+		SpectralKernel = GraphBuilder.CreateTexture(Desc, TEXT("Bloom.FFT.SpectralKernel"));
+	}
 
-	return SuccessValue;
+	FIntRect SrcRect(FIntPoint(0, 0), FrequencySize);
+	GPUFFT::FFTImage2D(
+		GraphBuilder,
+		View.ShaderMap,
+		FrequencySize, bDoHorizontalFirst,
+		ResizedKernel, SrcRect,
+		SpectralKernel);
+
+	return SpectralKernel;
 }
 
 struct FFFTBloomIntermediates
 {
-	FRHITexture* FullResolutionTarget = nullptr;
-	FIntRect FullResolutionViewRect;
-	FIntPoint FullResolutionSize;
-
-	FRHITexture* HalfResolutionTarget = nullptr;
-	FIntRect HalfResolutionViewRect;
-	FIntPoint HalfResolutionSize;
-
-	FRHITexture* InputTarget = nullptr;
-	FRHIUnorderedAccessView* OutputUAV = nullptr;
+	FRDGTextureRef InputTexture;
 
 	// The size of the input buffer.
 	FIntPoint InputBufferSize;
@@ -493,14 +369,11 @@ struct FFFTBloomIntermediates
 	// The order of the two-dimensional transform.  This implicitly defines
 	// the data layout in transform space for both the kernel and image transform
 	bool bDoHorizontalFirst = false;
-
-	bool bHalfResolutionFFT = false;
 };
 
 FFFTBloomIntermediates GetFFTBloomIntermediates(
 	const FViewInfo& View,
-	const FFFTBloomInputs& Inputs,
-	FRDGTextureUAVRef OutputUAV)
+	const FFFTBloomInputs& Inputs)
 {
 	FSceneViewState* ViewState = (FSceneViewState*)View.State;
 	check(ViewState);
@@ -535,30 +408,19 @@ FFFTBloomIntermediates GetFFTBloomIntermediates(
 	const bool bHalfResolutionFFT = IsFFTBloomHalfResolutionEnabled();
 
 	FFFTBloomIntermediates Intermediates;
-	Intermediates.FullResolutionTarget = Inputs.FullResolutionTexture->GetRHI();
-	Intermediates.FullResolutionViewRect = Inputs.FullResolutionViewRect;
-	Intermediates.FullResolutionSize = Intermediates.FullResolutionViewRect.Size();
-	Intermediates.bHalfResolutionFFT = bHalfResolutionFFT;
-
 	if (bHalfResolutionFFT)
 	{
-		Intermediates.HalfResolutionTarget = Inputs.HalfResolutionTexture->GetRHI();
-		Intermediates.HalfResolutionViewRect = Inputs.HalfResolutionViewRect;
-		Intermediates.HalfResolutionSize = Inputs.HalfResolutionViewRect.Size();
-
-		Intermediates.InputTarget = Intermediates.HalfResolutionTarget;
-		Intermediates.InputBufferSize = Intermediates.HalfResolutionSize;
-		Intermediates.ImageRect = Intermediates.HalfResolutionViewRect;
+		Intermediates.InputTexture = Inputs.HalfResolutionTexture;
+		Intermediates.InputBufferSize = Inputs.HalfResolutionViewRect.Size();
+		Intermediates.ImageRect = Inputs.HalfResolutionViewRect;
 	}
 	else
 	{
-		Intermediates.InputTarget = Intermediates.FullResolutionTarget;
-		Intermediates.InputBufferSize = Intermediates.FullResolutionSize;
-		Intermediates.ImageRect = Intermediates.FullResolutionViewRect;
+		Intermediates.InputTexture = Inputs.FullResolutionTexture;
+		Intermediates.InputBufferSize = Inputs.FullResolutionViewRect.Size();
+		Intermediates.ImageRect = Inputs.FullResolutionViewRect;
 	}
 
-	Intermediates.OutputUAV = OutputUAV->GetRHI();
-
 	Intermediates.KernelSupportScale = KernelSupportScale;
 	Intermediates.KernelSupportScaleClamp = KernelSupportScaleClamp;
 
@@ -595,50 +457,20 @@ FFFTBloomIntermediates GetFFTBloomIntermediates(
 	return Intermediates;
 }
 
-void ConvolveWithKernel(
-	FRHICommandList& RHICmdList,
+void InitDomainAndGetKernel(
+	FRDGBuilder& GraphBuilder,
 	const FViewInfo& View,
 	const FFFTBloomIntermediates& Intermediates,
-	const FTextureRHIRef& SpectralKernelTexture,
-	const FLinearColor& Tint,
-	FUnorderedAccessViewRHIRef ResultUAV)
+	bool bForceCenterZero,
+	FRDGTextureRef* OutCachedSpectralKernel,
+	FRDGTextureRef* OutCenterWeightTexture)
 {
-	GPUFFT::FGPUFFTShaderContext FFTContext(RHICmdList, *View.ShaderMap);
-
-	// Get Tmp buffers required for the Convolution
-
-	TRefCountPtr<IPooledRenderTarget> TmpTargets[2];
-
-	const FIntPoint TmpExtent = GPUFFT::Convolution2DBufferSize(Intermediates.FrequencySize, Intermediates.bDoHorizontalFirst, Intermediates.ImageRect.Size());
-	//(bDoHorizontalFirst) ? FIntPoint(FrequencySize.X + 2, ImageRect.Size().Y) : FIntPoint(ImageRect.Size().X, FrequencySize.Y + 2);
-	FPooledRenderTargetDesc Desc =
-		FPooledRenderTargetDesc::Create2DDesc(TmpExtent, GPUFFT::PixelFormat(), FClearValueBinding::None, TexCreate_None, TexCreate_RenderTargetable | TexCreate_UAV, false);
-
-	GRenderTargetPool.FindFreeElement(RHICmdList, Desc, TmpTargets[0], TEXT("Tmp FFT Buffer A"));
-	GRenderTargetPool.FindFreeElement(RHICmdList, Desc, TmpTargets[1], TEXT("Tmp FFT Buffer B"));
-
-	// Get the source
-	FRHITexture* InputTexture = Intermediates.InputTarget;
-
-	GPUFFT::ConvolutionWithTextureImage2D(FFTContext, Intermediates.FrequencySize, Intermediates.bDoHorizontalFirst,
-		SpectralKernelTexture,
-		Intermediates.ImageRect/*region of interest*/,
-		InputTexture,
-		ResultUAV,
-		TmpTargets[0]->GetRenderTargetItem(),
-		TmpTargets[1]->GetRenderTargetItem(),
-		Intermediates.PreFilter);
-
-	RHICmdList.Transition(FRHITransitionInfo(ResultUAV, ERHIAccess::Unknown, ERHIAccess::UAVCompute));
-}
-
-FSceneRenderTargetItem* InitDomainAndGetKernel(FRHICommandList& RHICmdList, const FViewInfo& View, const FFFTBloomIntermediates& Intermediates)
-{
-	FSceneViewState* ViewState = (FSceneViewState*)View.State;
+	FSceneViewState* ViewState = View.ViewState;
 
 	const auto& PPSettings = View.FinalPostProcessSettings;
 
 	UTexture2D* BloomConvolutionTexture = PPSettings.BloomConvolutionTexture;
+	const FTexture2DResource* BloomConvolutionTextureResource = View.FFTBloomKernelTexture;
 
 	// Fall back to the default bloom texture if provided.
 	if (BloomConvolutionTexture == nullptr)
@@ -647,7 +479,7 @@ FSceneRenderTargetItem* InitDomainAndGetKernel(FRHICommandList& RHICmdList, cons
 	}
 
 	// This should exist if we called IsFFTBloomPhysicalKernelReady.
-	check(BloomConvolutionTexture && BloomConvolutionTexture->GetResource() && BloomConvolutionTexture->GetResource()->TextureRHI);
+	check(BloomConvolutionTexture && BloomConvolutionTextureResource && BloomConvolutionTextureResource->TextureRHI);
 
 	const float BloomConvolutionSize = PPSettings.BloomConvolutionSize;
 	const FVector2D CenterUV = PPSettings.BloomConvolutionCenterUV;
@@ -658,103 +490,100 @@ FSceneRenderTargetItem* InitDomainAndGetKernel(FRHICommandList& RHICmdList, cons
 
 	// Should read / write to PF_G16R16F or PF_G32R32F (float2 formats)
 	// Need to set the render target description before we "request surface"
-	const EPixelFormat PixelFormat = GPUFFT::PixelFormat();
-	const FPooledRenderTargetDesc TransformDesc = FPooledRenderTargetDesc::Create2DDesc(PaddedFrequencySize, PixelFormat,
-		FClearValueBinding::None, TexCreate_None, TexCreate_RenderTargetable | TexCreate_UAV, false);
-
-	auto& FFTKernel = ViewState->BloomFFTKernel;
-	// Get the FFT kernel from the view state (note, this has already been transformed).
-	TRefCountPtr<IPooledRenderTarget>& TransformedKernelRT = FFTKernel.Spectral;
-	const UTexture2D* CachedKernelPhysical = FFTKernel.Physical;
-	const float       CachedKernelScale = FFTKernel.Scale;
-	const FVector2D&  CachedKernelCenterUV = FFTKernel.CenterUV;
-	const FIntPoint&  CachedImageSize = FFTKernel.ImageSize;
+	const FRDGTextureDesc TransformDesc = FRDGTextureDesc::Create2D(
+		PaddedFrequencySize,
+		GPUFFT::PixelFormat(),
+		FClearValueBinding::None,
+		TexCreate_ShaderResource | TexCreate_UAV);
 
-	const FIntPoint ImageSize = Intermediates.ImageRect.Size();
+	FRDGTextureRef CachedSpectralKernel = nullptr;
+	FRDGTextureRef CenterWeightTexture = nullptr;
+	if (ViewState && ViewState->BloomFFTKernel.Spectral)
+	{
+		auto& FFTKernel = ViewState->BloomFFTKernel;
 
-	bool bCachedKernelIsDirty = true;
+		FRDGTextureRef PrevCachedSpectralKernel = GraphBuilder.RegisterExternalTexture(FFTKernel.Spectral);
 
-	if (TransformedKernelRT)
-	{
-		FPooledRenderTarget* TransformedTexture = (FPooledRenderTarget*)TransformedKernelRT.GetReference();
+		const bool bSameTexture = (FFTKernel.Physical == static_cast<const UTexture2D*>(BloomConvolutionTexture));
+		const bool bSameSpectralBuffer = TransformDesc.ClearValue == PrevCachedSpectralKernel->Desc.ClearValue
+			&& TransformDesc.Flags == PrevCachedSpectralKernel->Desc.Flags
+			&& TransformDesc.Format == PrevCachedSpectralKernel->Desc.Format
+			&& TransformDesc.Extent == PrevCachedSpectralKernel->Desc.Extent;
 
-		const bool bSameTexture = (CachedKernelPhysical == static_cast<const UTexture2D*>(BloomConvolutionTexture));
-		const bool bSameSpectralBuffer = TransformedTexture->GetDesc().Compare(TransformDesc, true /*exact match*/);
-		const bool bSameKernelSize = FMath::IsNearlyEqual(CachedKernelScale, BloomConvolutionSize, float(1.e-6) /*tol*/);
-		const bool bSameImageSize = (ImageSize == CachedImageSize);
-		const bool bSameKernelCenterUV = CachedKernelCenterUV.Equals(CenterUV, float(1.e-6) /*tol*/);
-		const bool bSameMipLevel = bSameTexture && FFTKernel.PhysicalMipLevel == BloomConvolutionTexture->GetResource()->GetCurrentMipCount();
+		const bool bSameKernelSize = FMath::IsNearlyEqual(FFTKernel.Scale, BloomConvolutionSize, float(1.e-6) /*tol*/);
+		const bool bSameImageSize = (Intermediates.ImageRect.Size() == FFTKernel.ImageSize);
+		const bool bSameKernelCenterUV = FFTKernel.CenterUV.Equals(CenterUV, float(1.e-6) /*tol*/);
+		const bool bSameMipLevel = bSameTexture && FFTKernel.PhysicalMipLevel == BloomConvolutionTextureResource->GetCurrentMipCount();
 
 		if (bSameTexture && bSameSpectralBuffer && bSameKernelSize && bSameImageSize && bSameKernelCenterUV && bSameMipLevel)
 		{
-			bCachedKernelIsDirty = false;
+			CachedSpectralKernel = PrevCachedSpectralKernel;
+
+			if (bForceCenterZero)
+			{
+				CenterWeightTexture = GraphBuilder.RegisterExternalTexture(FFTKernel.CenterWeight);
+			}
 		}
 	}
 
 	// Re-transform the kernel if needed.
-	if (bCachedKernelIsDirty)
+	if (!CachedSpectralKernel)
 	{
-		// Resize the buffer to hold the transformed kernel
-		GRenderTargetPool.FindFreeElement(RHICmdList, TransformDesc, TransformedKernelRT, TEXT("FFTKernel"));
-
-		// NB: SpectralKernelRTItem is member data
-		FSceneRenderTargetItem& SpectralKernelRTItem = TransformedKernelRT->GetRenderTargetItem();
-		FUnorderedAccessViewRHIRef SpectralKernelUAV = SpectralKernelRTItem.UAV;
+		RDG_EVENT_SCOPE(GraphBuilder, "InitDomainAndGetKernel");
 
 		// Sample the physical space kernel into the resized buffer
-
-		FTextureRHIRef& PhysicalSpaceKernelTextureRef = BloomConvolutionTexture->GetResource()->TextureRHI;
+		const FTextureRHIRef& PhysicalSpaceKernelTextureRef = BloomConvolutionTextureResource->TextureRHI;
 
 		// Rescale the physical space kernel ( and omit the center if this is a 1/2 resolution fft, it will be added later)
-
-		ResizeAndCenterTexture(RHICmdList, View, PhysicalSpaceKernelTextureRef, ImageSize, CenterUV, Intermediates.KernelSupportScale,
-			Intermediates.FrequencySize, SpectralKernelRTItem.UAV, PaddedFrequencySize, Intermediates.bHalfResolutionFFT);
-
-		// ResizeAndCenterTexture transitioned SpectralKernelRTItem to UAVCompute
-		RHICmdList.Transition(FRHITransitionInfo(SpectralKernelRTItem.ShaderResourceTexture, ERHIAccess::UAVCompute, ERHIAccess::SRVMask));
+		FRDGTextureRef ResizedKernel = GraphBuilder.CreateTexture(TransformDesc, TEXT("Bloom.FFT.ResizedKernel"));
+		ResizeAndCenterTexture(
+			GraphBuilder,
+			View, PhysicalSpaceKernelTextureRef, Intermediates.ImageRect.Size(), CenterUV, Intermediates.KernelSupportScale,
+			Intermediates.FrequencySize, ResizedKernel, PaddedFrequencySize, bForceCenterZero);
 
 		// Two Dimensional FFT of the physical space kernel.  
-		// Input: SpectralRTItem holds the physical space kernel, on return it will be the spectral space 
+		CachedSpectralKernel = TransformKernelFFT(
+			GraphBuilder, View,
+			ResizedKernel,
+			Intermediates.bDoHorizontalFirst,
+			Intermediates.FrequencySize);
 
-		TransformKernelFFT(RHICmdList, View, SpectralKernelRTItem, Intermediates.bDoHorizontalFirst, Intermediates.FrequencySize);
-
-		if (Intermediates.bHalfResolutionFFT)
+		if (bForceCenterZero)
 		{
-			TRefCountPtr<IPooledRenderTarget>& CenterWeightRT = FFTKernel.CenterWeight;
-
-			const FPooledRenderTargetDesc CenterWeightDesc = FPooledRenderTargetDesc::Create2DDesc(FIntPoint(2, 1), PixelFormat,
-				FClearValueBinding::None, TexCreate_None, TexCreate_RenderTargetable | TexCreate_UAV, false);
-
-			// Resize the buffer to hold the transformed kernel
-			GRenderTargetPool.FindFreeElement(RHICmdList, CenterWeightDesc, CenterWeightRT, TEXT("FFTKernelCenterWeight"));
-
-			const FTextureRHIRef& HalfResKernelTextureRef = SpectralKernelRTItem.ShaderResourceTexture;
-
-			const FIntPoint& HalfResKernelExtent = PaddedFrequencySize;
-
-			const FIntPoint HalfResSumLocation = (Intermediates.bDoHorizontalFirst) ? FIntPoint(HalfResKernelExtent.X, 0) : FIntPoint(0, HalfResKernelExtent.Y);
 			// Capture the missing center weight from the kernel and the sum of the existing weights.
-			CaptureKernelWeight(RHICmdList, View, HalfResKernelTextureRef, HalfResKernelExtent, PhysicalSpaceKernelTextureRef, CenterUV, CenterWeightRT);
+			CenterWeightTexture = CaptureKernelWeight(
+				GraphBuilder,
+				View,
+				CachedSpectralKernel,
+				PaddedFrequencySize,
+				PhysicalSpaceKernelTextureRef,
+				CenterUV);
 		}
 
 		// Update the data on the ViewState
-		ViewState->BloomFFTKernel.Scale = BloomConvolutionSize;
-		ViewState->BloomFFTKernel.ImageSize = ImageSize;
-		ViewState->BloomFFTKernel.Physical = BloomConvolutionTexture;
-		ViewState->BloomFFTKernel.CenterUV = CenterUV;
-		ViewState->BloomFFTKernel.PhysicalMipLevel = BloomConvolutionTexture->GetResource()->GetCurrentMipCount();
+		if (ViewState)
+		{
+			ViewState->BloomFFTKernel.Scale = BloomConvolutionSize;
+			ViewState->BloomFFTKernel.ImageSize = Intermediates.ImageRect.Size();
+			ViewState->BloomFFTKernel.Physical = BloomConvolutionTexture;
+			ViewState->BloomFFTKernel.CenterUV = CenterUV;
+			ViewState->BloomFFTKernel.PhysicalMipLevel = BloomConvolutionTextureResource->GetCurrentMipCount();
+			{
+				ViewState->BloomFFTKernel.Spectral.SafeRelease();
+				GraphBuilder.QueueTextureExtraction(CachedSpectralKernel, &ViewState->BloomFFTKernel.Spectral);
+			}
+			{
+				ViewState->BloomFFTKernel.CenterWeight.SafeRelease();
+				if (CenterWeightTexture)
+					GraphBuilder.QueueTextureExtraction(CenterWeightTexture, &ViewState->BloomFFTKernel.CenterWeight);
+			}
+		}
 	}
 
-	// Return pointer to the transformed kernel.
-	return &(TransformedKernelRT->GetRenderTargetItem());
+	*OutCachedSpectralKernel = CachedSpectralKernel;
+	*OutCenterWeightTexture = CenterWeightTexture;
 }
 
-BEGIN_SHADER_PARAMETER_STRUCT(FFFTPassParameters, )
-	SHADER_PARAMETER_RDG_TEXTURE(Texture2D<float4>, FullResolution)
-	SHADER_PARAMETER_RDG_TEXTURE(Texture2D<float4>, HalfResolution)
-	SHADER_PARAMETER_RDG_TEXTURE_UAV(RWTexture2D<float4>, Output)
-END_SHADER_PARAMETER_STRUCT()
-
 FRDGTextureRef AddFFTBloomPass(FRDGBuilder& GraphBuilder, const FViewInfo& View, const FFFTBloomInputs& Inputs)
 {
 	check(Inputs.FullResolutionTexture);
@@ -762,75 +591,96 @@ FRDGTextureRef AddFFTBloomPass(FRDGBuilder& GraphBuilder, const FViewInfo& View,
 	check(Inputs.HalfResolutionTexture);
 	check(!Inputs.HalfResolutionViewRect.IsEmpty());
 
-	const FRDGTextureDesc OutputDesc = FRDGTextureDesc::Create2D(
-		Inputs.FullResolutionTexture->Desc.Extent,
-		Inputs.FullResolutionTexture->Desc.Format,
-		FClearValueBinding::None,
-		TexCreate_ShaderResource | TexCreate_RenderTargetable | TexCreate_UAV);
-
-	FRDGTextureRef OutputTexture = GraphBuilder.CreateTexture(OutputDesc, TEXT("FFTBloom"));
-	FRDGTextureUAVRef OutputUAV = GraphBuilder.CreateUAV(OutputTexture);
-
-	FFFTPassParameters* PassParameters = GraphBuilder.AllocParameters<FFFTPassParameters>();
-	PassParameters->FullResolution = Inputs.FullResolutionTexture;
-	PassParameters->HalfResolution = Inputs.HalfResolutionTexture;
-	PassParameters->Output = OutputUAV;
+	const bool bHalfResolutionFFT = IsFFTBloomHalfResolutionEnabled();
 
-	GraphBuilder.AddPass(
-		RDG_EVENT_NAME("FFTBloom"),
-		PassParameters,
-		ERDGPassFlags::Compute,
-		[&View, Inputs, OutputUAV](FRHICommandList& RHICmdList)
-	{
-		FFFTBloomIntermediates Intermediates = GetFFTBloomIntermediates(View, Inputs, OutputUAV);
+	FFFTBloomIntermediates Intermediates = GetFFTBloomIntermediates(View, Inputs);
+	
+	// Init the domain data update the cached kernel if needed.
+	FRDGTextureRef SpectralKernelTexture;
+	FRDGTextureRef CenterWeightTexture;
+	InitDomainAndGetKernel(
+		GraphBuilder, View, Intermediates,
+		/* bForceCenterZero = */ bHalfResolutionFFT,
+		/* out */ &SpectralKernelTexture,
+		/* out */ &CenterWeightTexture);
 
-		// Init the domain data update the cached kernel if needed.
-		FSceneRenderTargetItem* SpectralKernelRTItem = InitDomainAndGetKernel(RHICmdList, View, Intermediates);
+	const FLinearColor Tint(1, 1, 1, 1);
 
-		// Do the convolution with the kernel
-		const FTextureRHIRef& SpectralKernelTexture = SpectralKernelRTItem->ShaderResourceTexture;
+	RDG_EVENT_SCOPE(GraphBuilder, "FFTBloom %dx%d", Intermediates.ImageRect.Width(), Intermediates.ImageRect.Height());
 
-		const FLinearColor Tint(1, 1, 1, 1);
+	if (bHalfResolutionFFT)
+	{
+		// Get a half-resolution destination buffer.
+		TRefCountPtr<IPooledRenderTarget> HalfResConvolutionResult;
 
-		if (Intermediates.bHalfResolutionFFT)
+		FRDGTextureRef FFTConvolutionTexture;
 		{
-			// Get a half-resolution destination buffer.
-			TRefCountPtr<IPooledRenderTarget> HalfResConvolutionResult;
-
-			const EPixelFormat PixelFormat = GPUFFT::PixelFormat();
-
-			const FPooledRenderTargetDesc HalfResFFTDesc = FPooledRenderTargetDesc::Create2DDesc(Intermediates.InputBufferSize, PixelFormat,
-				FClearValueBinding::None, TexCreate_None, TexCreate_RenderTargetable | TexCreate_UAV, false);
-
-			GRenderTargetPool.FindFreeElement(RHICmdList, HalfResFFTDesc, HalfResConvolutionResult, TEXT("HalfRes FFT Result"));
-			FSceneRenderTargetItem& HalfResConvolutionRTItem = HalfResConvolutionResult->GetRenderTargetItem();
-
-			// The FFT result buffer is also half res.
+			const FRDGTextureDesc Desc = FRDGTextureDesc::Create2D(
+				Intermediates.InputBufferSize,
+				GPUFFT::PixelFormat(),
+				FClearValueBinding::None,
+				TexCreate_ShaderResource | TexCreate_UAV);
 
-			ConvolveWithKernel(RHICmdList, View, Intermediates, SpectralKernelTexture, Tint, HalfResConvolutionRTItem.UAV);
+			FFTConvolutionTexture = GraphBuilder.CreateTexture(Desc, TEXT("Bloom.FFT.SceneColor"));
+		}
 
-			// The blend weighting parameters from the View State
+		GPUFFT::ConvolutionWithTextureImage2D(
+			GraphBuilder,
+			View.ShaderMap,
+			Intermediates.FrequencySize,
+			Intermediates.bDoHorizontalFirst,
+			SpectralKernelTexture,
+			Intermediates.InputTexture, Intermediates.ImageRect,
+			FFTConvolutionTexture, Intermediates.ImageRect,
+			Intermediates.PreFilter);
 
-			FSceneViewState* ViewState = (FSceneViewState*)View.State;
-			auto& FFTKernel = ViewState->BloomFFTKernel;
+		// Blend with  alpha * SrcBuffer + betta * BloomedBuffer  where alpha = Weights[0], beta = Weights[1]
+		const FIntPoint HalfResBufferSize = Intermediates.InputBufferSize;
 
-			const FTextureRHIRef& CenterWeightTexture = FFTKernel.CenterWeight->GetRenderTargetItem().ShaderResourceTexture;
+		FRDGTextureRef OutputSceneColorTexture;
+		{
+			const FRDGTextureDesc Desc = FRDGTextureDesc::Create2D(
+				Inputs.FullResolutionTexture->Desc.Extent,
+				Inputs.FullResolutionTexture->Desc.Format,
+				FClearValueBinding::None,
+				TexCreate_ShaderResource | TexCreate_UAV);
 
-			// Get full resolution source
-			FRHITexture* FullResResourceTexture = Intermediates.FullResolutionTarget;
+			OutputSceneColorTexture = GraphBuilder.CreateTexture(Desc, TEXT("Bloom.FFT.SceneColor"));
+		}
 
-			// Blend with  alpha * SrcBuffer + betta * BloomedBuffer  where alpha = Weights[0], beta = Weights[1]
-			const FIntPoint HalfResBufferSize = Intermediates.InputBufferSize;
+		BlendLowRes(
+			GraphBuilder,
+			View,
+			Inputs.FullResolutionTexture, Inputs.FullResolutionViewRect,
+			FFTConvolutionTexture, Intermediates.ImageRect,
+			Intermediates.InputBufferSize,
+			CenterWeightTexture,
+			OutputSceneColorTexture);
 
-			BlendLowRes(RHICmdList, View, FullResResourceTexture, Intermediates.FullResolutionViewRect, HalfResConvolutionRTItem.ShaderResourceTexture, Intermediates.ImageRect, HalfResBufferSize, CenterWeightTexture, Intermediates.OutputUAV);
-		}
-		else
-		{
-			// Do Convolution directly into the output buffer
-			// NB: In this case there is only one input, and the output has matching resolution
-			ConvolveWithKernel(RHICmdList, View, Intermediates, SpectralKernelTexture, Tint, Intermediates.OutputUAV);
-		}
-	});
+		return OutputSceneColorTexture;
+	}
+	else
+	{
+		FRDGTextureDesc OutputSceneColorDesc = FRDGTextureDesc::Create2D(
+			Inputs.FullResolutionTexture->Desc.Extent,
+			Inputs.FullResolutionTexture->Desc.Format,
+			FClearValueBinding::None,
+			TexCreate_ShaderResource | TexCreate_UAV);
+		
+		FRDGTextureRef OutputSceneColorTexture = GraphBuilder.CreateTexture(OutputSceneColorDesc, TEXT("Bloom.FFT.SceneColor"));
+
+		GPUFFT::ConvolutionWithTextureImage2D(
+			GraphBuilder,
+			View.ShaderMap,
+			Intermediates.FrequencySize,
+			Intermediates.bDoHorizontalFirst,
+			SpectralKernelTexture,
+			Intermediates.InputTexture, Intermediates.ImageRect,
+			OutputSceneColorTexture, Intermediates.ImageRect,
+			Intermediates.PreFilter);
+
+		return OutputSceneColorTexture;
+	}
 
-	return OutputTexture;
+	return Inputs.FullResolutionTexture;
 }
diff --git a/Engine/Source/Runtime/Renderer/Private/SceneRendering.cpp b/Engine/Source/Runtime/Renderer/Private/SceneRendering.cpp
--- a/Engine/Source/Runtime/Renderer/Private/SceneRendering.cpp
+++ b/Engine/Source/Runtime/Renderer/Private/SceneRendering.cpp
@@ -2230,6 +2230,44 @@ void FSortedShadowMaps::Release()
 	PreshadowCache.RenderTargets.Release();
 }
 
+static bool PreparePostProcessSettingTextureForRenderer(const FViewInfo& View, UTexture2D* Texture2D, const TCHAR* TextureUsageName)
+{
+	check(IsInGameThread());
+
+	bool bIsValid = Texture2D != nullptr;
+
+	if (bIsValid)
+	{
+		const int32 CinematicTextureGroups = 0;
+		const float Seconds = 5.0f;
+		Texture2D->SetForceMipLevelsToBeResident(Seconds, CinematicTextureGroups);
+	}
+
+	const uint32 FramesPerWarning = 15;
+
+	if (bIsValid && Texture2D->IsFullyStreamedIn() == false)
+	{
+		if ((View.Family->FrameNumber % FramesPerWarning) == 0)
+		{
+			UE_LOG(LogRenderer, Warning, TEXT("The %s texture is not streamed in."), TextureUsageName);
+		}
+
+		bIsValid = false;
+	}
+
+	if (bIsValid && Texture2D->bHasStreamingUpdatePending == true)
+	{
+		if ((View.Family->FrameNumber % FramesPerWarning) == 0)
+		{
+			UE_LOG(LogRenderer, Warning, TEXT("The %s texture has pending update."), TextureUsageName);
+		}
+
+		bIsValid = false;
+	}
+
+	return bIsValid;
+};
+
 template <typename T>
 inline T* CheckPointer(T* Ptr)
 {
@@ -2304,48 +2342,52 @@ FSceneRenderer::FSceneRenderer(const FSceneViewFamily* InViewFamily,FHitProxyCon
 		}
 		#endif
 
-		// Handle the film grain texture
-		if (ViewInfo->FinalPostProcessSettings.FilmGrainIntensity > 0.0f)
+		// Handle the FFT bloom kernel textire
+		if (ViewInfo->FinalPostProcessSettings.BloomMethod == EBloomMethod::BM_FFT && ViewInfo->ViewState != nullptr)
 		{
-			GEngine->LoadDefaultFilmGrainTexture();
-
-			UTexture2D* FilmGrainTexture = GEngine->DefaultFilmGrainTexture;
-			if (ViewInfo->FinalPostProcessSettings.FilmGrainTexture)
-			{
-				FilmGrainTexture = ViewInfo->FinalPostProcessSettings.FilmGrainTexture;
-			}
-			  
-			bool bIsValid = FilmGrainTexture != nullptr;
+			GEngine->LoadDefaultBloomTexture();
 
-			if (bIsValid)
+			UTexture2D* BloomConvolutionTexture = ViewInfo->FinalPostProcessSettings.BloomConvolutionTexture;
+			if (BloomConvolutionTexture == nullptr)
 			{
-				const int32 CinematicTextureGroups = 0;
-				const float Seconds = 5.0f;
-				FilmGrainTexture->SetForceMipLevelsToBeResident(Seconds, CinematicTextureGroups);
+				BloomConvolutionTexture = GEngine->DefaultBloomKernelTexture;
 			}
 
-			const uint32 FramesPerWarning = 15;
+			bool bIsValid = PreparePostProcessSettingTextureForRenderer(*ViewInfo, BloomConvolutionTexture, TEXT("convolution bloom"));
 
-			if (bIsValid && FilmGrainTexture->IsFullyStreamedIn() == false)
+			if (bIsValid)
 			{
-				if ((ViewInfo->Family->FrameNumber % FramesPerWarning) == 0)
+				const FTextureResource* TextureResource = BloomConvolutionTexture->GetResource();
+				if (TextureResource)
 				{
-					UE_LOG(LogRenderer, Warning, TEXT("The film grain texture is not streamed in."));
+					ViewInfo->FFTBloomKernelTexture = TextureResource->GetTexture2DResource();
+					ViewInfo->FinalPostProcessSettings.BloomConvolutionTexture = BloomConvolutionTexture;
+				}
+				else
+				{
+					ViewInfo->FinalPostProcessSettings.BloomConvolutionTexture = nullptr;
 				}
 
-				bIsValid = false;
-			}
-
-			if (bIsValid && FilmGrainTexture->bHasStreamingUpdatePending == true)
-			{
-				if ((ViewInfo->Family->FrameNumber % FramesPerWarning) == 0)
+				if (BloomConvolutionTexture == GEngine->DefaultBloomKernelTexture)
 				{
-					UE_LOG(LogRenderer, Warning, TEXT("The film grain texture has pending update."));
+
 				}
+			}
+		}
 
-				bIsValid = false;
+		// Handle the film grain texture
+		if (ViewInfo->FinalPostProcessSettings.FilmGrainIntensity > 0.0f)
+		{
+			GEngine->LoadDefaultFilmGrainTexture();
+
+			UTexture2D* FilmGrainTexture = GEngine->DefaultFilmGrainTexture;
+			if (ViewInfo->FinalPostProcessSettings.FilmGrainTexture)
+			{
+				FilmGrainTexture = ViewInfo->FinalPostProcessSettings.FilmGrainTexture;
 			}
 
+			bool bIsValid = PreparePostProcessSettingTextureForRenderer(*ViewInfo, FilmGrainTexture, TEXT("film grain"));
+			 
 			if (bIsValid)
 			{
 				const FTextureResource* TextureResource = FilmGrainTexture->GetResource();
diff --git a/Engine/Source/Runtime/Renderer/Private/SceneRendering.h b/Engine/Source/Runtime/Renderer/Private/SceneRendering.h
--- a/Engine/Source/Runtime/Renderer/Private/SceneRendering.h
+++ b/Engine/Source/Runtime/Renderer/Private/SceneRendering.h
@@ -1423,6 +1423,7 @@ public:
 	/** Custom visibility query for view */
 	ICustomVisibilityQuery* CustomVisibilityQuery;
 
+	const FTexture2DResource* FFTBloomKernelTexture = nullptr;
 	const FTexture2DResource* FilmGrainTexture = nullptr;
 
 	TArray<FPrimitiveSceneInfo*, SceneRenderingAllocator> IndirectShadowPrimitives;
