Fix a bug of mobile Pixel Projection Reflection.

diff --git a/Engine/Source/Runtime/Renderer/Private/MobileBasePassRendering.cpp b/Engine/Source/Runtime/Renderer/Private/MobileBasePassRendering.cpp
--- a/Engine/Source/Runtime/Renderer/Private/MobileBasePassRendering.cpp
+++ b/Engine/Source/Runtime/Renderer/Private/MobileBasePassRendering.cpp
@@ -189,6 +189,11 @@ void SetupMobileBasePassUniformParameters(
 		BasePassParameters.PlanarReflection.PlanarReflectionTexture = View.PrevViewInfo.MobilePixelProjectedReflection->GetRenderTargetItem().ShaderResourceTexture;
 		BasePassParameters.PlanarReflection.PlanarReflectionSampler = TStaticSamplerState<SF_Bilinear, AM_Clamp, AM_Clamp, AM_Clamp>::GetRHI();
 	}
+	else if (BasePassParameters.PlanarReflection.PlanarReflectionTexture == nullptr)
+	{
+		BasePassParameters.PlanarReflection.PlanarReflectionTexture = GBlackTexture->TextureRHI;
+		BasePassParameters.PlanarReflection.PlanarReflectionSampler = GBlackTexture->SamplerStateRHI;
+	}
 
 	const FRDGSystemTextures& SystemTextures = FRDGSystemTextures::Get(GraphBuilder);
 
