Fix crash with deep shadow rendering when enabling large world coordinates.

The memcpy on the pass parameters was stomping the previously set variable, due to the missmatch side. This CL changes this memcopy into a loop assignation, so that the correct casting operator is used.

diff --git a/Engine/Source/Runtime/Renderer/Private/HairStrands/HairStrandsTransmittance.cpp b/Engine/Source/Runtime/Renderer/Private/HairStrands/HairStrandsTransmittance.cpp
--- a/Engine/Source/Runtime/Renderer/Private/HairStrands/HairStrandsTransmittance.cpp
+++ b/Engine/Source/Runtime/Renderer/Private/HairStrands/HairStrandsTransmittance.cpp
@@ -389,9 +389,12 @@ static FRDGBufferRef AddHairStrandsDeepShadowTransmittanceMaskPass(
 	Parameters->DeepShadow_WorldToLightTransformBuffer = Params.DeepShadow_WorldToLightTransformBuffer;
 	Parameters->RayMarchMaskTexture = ScreenShadowMaskSubPixelTexture ? ScreenShadowMaskSubPixelTexture : GraphBuilder.RegisterExternalTexture(GSystemTextures.WhiteDummy);
 
-	memcpy(&(Parameters->DeepShadow_AtlasSlotOffsets_AtlasSlotIndex[0]), Params.DeepShadow_AtlasSlotOffsets_AtlasSlotIndex, sizeof(FIntVector4) * FHairStrandsDeepShadowData::MaxMacroGroupCount);
-	memcpy(&(Parameters->DeepShadow_CPUWorldToLightTransforms[0]), Params.DeepShadow_CPUWorldToLightTransforms, sizeof(FMatrix) * FHairStrandsDeepShadowData::MaxMacroGroupCount);
-
+	for (uint32 SlotIndex=0;SlotIndex< FHairStrandsDeepShadowData::MaxMacroGroupCount;++SlotIndex)
+	{
+		Parameters->DeepShadow_AtlasSlotOffsets_AtlasSlotIndex[SlotIndex] = Params.DeepShadow_AtlasSlotOffsets_AtlasSlotIndex[SlotIndex];
+		Parameters->DeepShadow_CPUWorldToLightTransforms[SlotIndex] = Params.DeepShadow_CPUWorldToLightTransforms[SlotIndex];
+	}
+	
 	check(NodeGroupSize == 64 || NodeGroupSize == 32);
 	FHairStrandsDeepShadowTransmittanceMaskCS::FPermutationDomain PermutationVector;
 	PermutationVector.Set<FHairStrandsDeepShadowTransmittanceMaskCS::FTransmittanceGroupSize>(NodeGroupSize);
