Add clear passes to fix capture replay issues on certain consoles.

diff --git a/Engine/Shaders/Private/InstanceCulling/BuildInstanceDrawCommands.usf b/Engine/Shaders/Private/InstanceCulling/BuildInstanceDrawCommands.usf
--- a/Engine/Shaders/Private/InstanceCulling/BuildInstanceDrawCommands.usf
+++ b/Engine/Shaders/Private/InstanceCulling/BuildInstanceDrawCommands.usf
@@ -13,6 +13,9 @@ uint InstanceSceneDataSOAStride;
 uint GPUSceneNumInstances;
 uint GPUSceneNumPrimitives;
 
+#ifndef BATCH_PROCESSING_MODE_NUM
+	#define BATCH_PROCESSING_MODE_NUM 1
+#endif // BATCH_PROCESSING_MODE_NUM
 
 StructuredBuffer<FDrawCommandDesc> DrawCommandDescs;
 Buffer<uint> InstanceIdOffsetBuffer;
@@ -205,3 +208,18 @@ void InstanceCullBuildInstanceIdBufferCS(uint3 GroupId : SV_GroupID, int GroupTh
 	}
 #endif // STEREO_CULLING_MODE
 }
+
+
+uint NumIndirectArgs;
+/**
+ */
+[numthreads(NUM_THREADS_PER_GROUP, 1, 1)]
+void ClearIndirectArgInstanceCountCS(uint3 GroupId : SV_GroupID, int GroupThreadIndex : SV_GroupIndex)
+{
+	uint IndirectArgIndex = GetUnWrappedDispatchThreadId(GroupId, GroupThreadIndex, NUM_THREADS_PER_GROUP);
+
+	if (IndirectArgIndex < NumIndirectArgs)
+	{
+		DrawIndirectArgsBufferOut[IndirectArgIndex * INDIRECT_ARGS_NUM_WORDS + 1] = 0U;
+	}
+}
\ No newline at end of file
diff --git a/Engine/Source/Runtime/Renderer/Private/InstanceCulling/InstanceCullingContext.cpp b/Engine/Source/Runtime/Renderer/Private/InstanceCulling/InstanceCullingContext.cpp
--- a/Engine/Source/Runtime/Renderer/Private/InstanceCulling/InstanceCullingContext.cpp
+++ b/Engine/Source/Runtime/Renderer/Private/InstanceCulling/InstanceCullingContext.cpp
@@ -294,6 +294,9 @@ void FInstanceCullingContext::BuildRenderingCommands(
 	FRDGBufferRef DrawIndirectArgsRDG = GraphBuilder.CreateBuffer(FRDGBufferDesc::CreateIndirectDesc(IndirectArgsNumWords * IndirectArgs.Num()), TEXT("InstanceCulling.DrawIndirectArgsBuffer"));
 	GraphBuilder.QueueBufferUpload(DrawIndirectArgsRDG, IndirectArgs.GetData(), IndirectArgs.GetTypeSize() * IndirectArgs.Num());
 
+	// Note: we redundantly clear the instance counts here as there is some issue with replays on certain consoles.
+	AddClearIndirectArgInstanceCountPass(GraphBuilder, ShaderMap, DrawIndirectArgsRDG);
+
 	// not using structured buffer as we have to get at it as a vertex buffer 
 	FRDGBufferRef InstanceIdOffsetBufferRDG = GraphBuilder.CreateBuffer(FRDGBufferDesc::CreateBufferDesc(sizeof(uint32), InstanceIdOffsets.Num()), TEXT("InstanceCulling.InstanceIdOffsetBuffer"));
 	GraphBuilder.QueueBufferUpload(InstanceIdOffsetBufferRDG, InstanceIdOffsets.GetData(), InstanceIdOffsets.GetTypeSize() * InstanceIdOffsets.Num());
@@ -533,6 +536,13 @@ void FInstanceCullingContext::BuildRenderingCommandsDeferred(
 
 	BatchedCullingScratch.DrawIndirectArgsBuffer = GraphBuilder.CreateBuffer(FRDGBufferDesc::CreateIndirectDesc(), TEXT("InstanceCulling.DrawIndirectArgsBuffer"), INST_CULL_CALLBACK(IndirectArgsNumWords * MergedContext.IndirectArgs.Num()));
 	GraphBuilder.QueueBufferUpload(BatchedCullingScratch.DrawIndirectArgsBuffer, INST_CULL_CALLBACK(MergedContext.IndirectArgs.GetData()), INST_CULL_CALLBACK(GetArrayDataSize(MergedContext.IndirectArgs)));
+
+	const ERHIFeatureLevel::Type FeatureLevel = GPUScene.GetFeatureLevel();
+	FGlobalShaderMap* ShaderMap = GetGlobalShaderMap(FeatureLevel);
+
+	// Note: we redundantly clear the instance counts here as there is some issue with replays on certain consoles.
+	AddClearIndirectArgInstanceCountPass(GraphBuilder, ShaderMap, BatchedCullingScratch.DrawIndirectArgsBuffer);
+
 	// not using structured buffer as we want/have to get at it as a vertex buffer 
 	BatchedCullingScratch.InstanceIdOffsetBuffer = GraphBuilder.CreateBuffer(FRDGBufferDesc::CreateBufferDesc(sizeof(uint32), 1), TEXT("InstanceCulling.InstanceIdOffsetBuffer"), INST_CULL_CALLBACK(MergedContext.InstanceIdOffsets.Num()));
 	GraphBuilder.QueueBufferUpload(BatchedCullingScratch.InstanceIdOffsetBuffer, INST_CULL_CALLBACK(MergedContext.InstanceIdOffsets.GetData()), INST_CULL_CALLBACK(MergedContext.InstanceIdOffsets.GetTypeSize() * MergedContext.InstanceIdOffsets.Num()));
@@ -592,7 +602,7 @@ void FInstanceCullingContext::BuildRenderingCommandsDeferred(
 		PermutationVector.Set<FBuildInstanceIdBufferAndCommandsFromPrimitiveIdsCs::FBatchedDim>(true);
 		PermutationVector.Set<FBuildInstanceIdBufferAndCommandsFromPrimitiveIdsCs::FCullInstancesDim>(bCullInstances && EBatchProcessingMode(Mode) != EBatchProcessingMode::UnCulled);
 
-		auto ComputeShader = GetGlobalShaderMap(GMaxRHIFeatureLevel)->GetShader<FBuildInstanceIdBufferAndCommandsFromPrimitiveIdsCs>(PermutationVector);
+		auto ComputeShader = ShaderMap->GetShader<FBuildInstanceIdBufferAndCommandsFromPrimitiveIdsCs>(PermutationVector);
 
 		FComputeShaderUtils::AddPass(
 			GraphBuilder,
@@ -622,6 +632,54 @@ bool FInstanceCullingContext::AllowBatchedBuildRenderingCommands(const FGPUScene
 
 
 
+class FClearIndirectArgInstanceCountCs : public FGlobalShader
+{
+	DECLARE_GLOBAL_SHADER(FClearIndirectArgInstanceCountCs);
+	SHADER_USE_PARAMETER_STRUCT(FClearIndirectArgInstanceCountCs, FGlobalShader)
+
+public:
+	static constexpr int32 NumThreadsPerGroup = 64;
+
+	static bool ShouldCompilePermutation(const FGlobalShaderPermutationParameters& Parameters)
+	{
+		return UseGPUScene(Parameters.Platform);
+	}
+
+	static void ModifyCompilationEnvironment(const FGlobalShaderPermutationParameters& Parameters, FShaderCompilerEnvironment& OutEnvironment)
+	{
+		FGlobalShader::ModifyCompilationEnvironment(Parameters, OutEnvironment);
+		FInstanceProcessingGPULoadBalancer::SetShaderDefines(OutEnvironment);
+
+		OutEnvironment.SetDefine(TEXT("INDIRECT_ARGS_NUM_WORDS"), FInstanceCullingContext::IndirectArgsNumWords);
+		OutEnvironment.SetDefine(TEXT("NUM_THREADS_PER_GROUP"), NumThreadsPerGroup);
+	}
+
+	BEGIN_SHADER_PARAMETER_STRUCT(FParameters, )
+		SHADER_PARAMETER_RDG_BUFFER_UAV(RWBuffer<uint>, DrawIndirectArgsBufferOut)
+		SHADER_PARAMETER(uint32, NumIndirectArgs)
+	END_SHADER_PARAMETER_STRUCT()
+};
+IMPLEMENT_GLOBAL_SHADER(FClearIndirectArgInstanceCountCs, "/Engine/Private/InstanceCulling/BuildInstanceDrawCommands.usf", "ClearIndirectArgInstanceCountCS", SF_Compute);
+
+
+void FInstanceCullingContext::AddClearIndirectArgInstanceCountPass(FRDGBuilder& GraphBuilder, FGlobalShaderMap* ShaderMap, FRDGBufferRef DrawIndirectArgsBuffer)
+{
+	FClearIndirectArgInstanceCountCs::FParameters* PassParameters = GraphBuilder.AllocParameters<FClearIndirectArgInstanceCountCs::FParameters>();
+	// Upload data etc
+	PassParameters->DrawIndirectArgsBufferOut = GraphBuilder.CreateUAV(DrawIndirectArgsBuffer, PF_R32_UINT);
+	PassParameters->NumIndirectArgs = DrawIndirectArgsBuffer->Desc.NumElements / FInstanceCullingContext::IndirectArgsNumWords;
+
+	auto ComputeShader = ShaderMap->GetShader<FClearIndirectArgInstanceCountCs>();
+
+	FComputeShaderUtils::AddPass(
+		GraphBuilder,
+		RDG_EVENT_NAME("ClearIndirectArgInstanceCount"),
+		ComputeShader,
+		PassParameters,
+		FComputeShaderUtils::GetGroupCountWrapped(PassParameters->NumIndirectArgs, FClearIndirectArgInstanceCountCs::NumThreadsPerGroup)
+	);
+}
+
 /**
  * Allocate indirect arg slots for all meshes to use instancing,
  * add commands that populate the indirect calls and index & id buffers, and
diff --git a/Engine/Source/Runtime/Renderer/Private/VirtualShadowMaps/VirtualShadowMapArray.cpp b/Engine/Source/Runtime/Renderer/Private/VirtualShadowMaps/VirtualShadowMapArray.cpp
--- a/Engine/Source/Runtime/Renderer/Private/VirtualShadowMaps/VirtualShadowMapArray.cpp
+++ b/Engine/Source/Runtime/Renderer/Private/VirtualShadowMaps/VirtualShadowMapArray.cpp
@@ -1626,18 +1626,18 @@ void FVirtualShadowMapArray::RenderVirtualShadowMapsHw(FRDGBuilder& GraphBuilder
 
 			int32 NumIndirectArgs = InstanceCullingContext->IndirectArgs.Num();
 
-			TArray<int32, SceneRenderingAllocator> DrawCommandInstanceCountTmp;
-			DrawCommandInstanceCountTmp.AddZeroed(NumIndirectArgs);
-			FRDGBufferRef TmpInstanceIdOffsetBufferRDG = CreateStructuredBuffer(GraphBuilder, TEXT("Shadow.Virtual.TmpInstanceIdOffsetBuffer"), DrawCommandInstanceCountTmp);
+			FRDGBufferRef TmpInstanceIdOffsetBufferRDG = CreateStructuredBuffer(GraphBuilder, TEXT("Shadow.Virtual.TmpInstanceIdOffsetBuffer"), sizeof(uint32), NumIndirectArgs, nullptr, 0);
 
 			// TODO: This is both not right, and also over conservative when running with the atomic path
 			const uint32 MaxNumInstancesPerPass = InstanceCullingContext->TotalInstances * 64u;
 			FRDGBufferRef VisibleInstancesRdg = CreateStructuredBuffer(GraphBuilder, TEXT("Shadow.Virtual.VisibleInstances"), sizeof(FVisibleInstanceCmd), MaxNumInstancesPerPass, nullptr, 0);
 
-			TArray<uint32, TInlineAllocator<1> > NullArray;
-			NullArray.AddZeroed(1);
-			FRDGBufferRef VisibleInstanceWriteOffsetRDG = CreateStructuredBuffer(GraphBuilder, TEXT("Shadow.Virtual.VisibleInstanceWriteOffset"), NullArray);
-			FRDGBufferRef OutputOffsetBufferRDG = CreateStructuredBuffer(GraphBuilder, TEXT("Shadow.Virtual.OutputOffsetBuffer"), NullArray);
+			FRDGBufferRef VisibleInstanceWriteOffsetRDG = CreateStructuredBuffer(GraphBuilder, TEXT("Shadow.Virtual.VisibleInstanceWriteOffset"), sizeof(uint32),  1, nullptr, 0);
+			FRDGBufferRef OutputOffsetBufferRDG = CreateStructuredBuffer(GraphBuilder, TEXT("Shadow.Virtual.OutputOffsetBuffer"), sizeof(uint32), 1, nullptr, 0);
+
+			AddClearUAVPass(GraphBuilder, GraphBuilder.CreateUAV(VisibleInstanceWriteOffsetRDG), 0);
+			AddClearUAVPass(GraphBuilder, GraphBuilder.CreateUAV(OutputOffsetBufferRDG), 0);
+
 			FRDGBufferRef VirtualShadowViewsRDG = CreateStructuredBuffer(GraphBuilder, TEXT("Shadow.Virtual.VirtualShadowViews"), VirtualShadowViews);
 
 			const auto& IndirectArgs = InstanceCullingContext->IndirectArgs;
@@ -1648,6 +1648,11 @@ void FVirtualShadowMapArray::RenderVirtualShadowMapsHw(FRDGBuilder& GraphBuilder
 			FRDGBufferRef DrawIndirectArgsRDG = GraphBuilder.CreateBuffer(FRDGBufferDesc::CreateIndirectDesc(FInstanceCullingContext::IndirectArgsNumWords * IndirectArgs.Num()), TEXT("Shadow.Virtual.DrawIndirectArgsBuffer"));
 			GraphBuilder.QueueBufferUpload(DrawIndirectArgsRDG, IndirectArgs.GetData(), IndirectArgs.GetTypeSize() * IndirectArgs.Num());
 
+			FGlobalShaderMap* ShaderMap = GetGlobalShaderMap(GMaxRHIFeatureLevel);
+
+			// Note: we redundantly clear the instance counts here as there is some issue with replays on certain consoles.
+			FInstanceCullingContext::AddClearIndirectArgInstanceCountPass(GraphBuilder, ShaderMap, DrawIndirectArgsRDG);
+
 			// not using structured buffer as we have to get at it as a vertex buffer 
 			FRDGBufferRef InstanceIdOffsetBufferRDG = GraphBuilder.CreateBuffer(FRDGBufferDesc::CreateBufferDesc(sizeof(uint32), InstanceIdOffsets.Num()), TEXT("Shadow.Virtual.InstanceIdOffsetBuffer"));
 
@@ -1684,7 +1689,7 @@ void FVirtualShadowMapArray::RenderVirtualShadowMapsHw(FRDGBuilder& GraphBuilder
 				PermutationVector.Set< FCullPerPageDrawCommandsCs::FNearClipDim >( !Clipmap.IsValid() );
 				PermutationVector.Set< FCullPerPageDrawCommandsCs::FLoopOverViewsDim >( Clipmap.IsValid() );
 
-				auto ComputeShader = GetGlobalShaderMap(GMaxRHIFeatureLevel)->GetShader<FCullPerPageDrawCommandsCs>( PermutationVector );
+				auto ComputeShader = ShaderMap->GetShader<FCullPerPageDrawCommandsCs>( PermutationVector );
 
 				FComputeShaderUtils::AddPass(
 					GraphBuilder,
@@ -1698,7 +1703,8 @@ void FVirtualShadowMapArray::RenderVirtualShadowMapsHw(FRDGBuilder& GraphBuilder
 			{
 				FAllocateCommandInstanceOutputSpaceCs::FParameters* PassParameters = GraphBuilder.AllocParameters<FAllocateCommandInstanceOutputSpaceCs::FParameters>();
 
-				FRDGBufferRef InstanceIdOutOffsetBufferRDG = CreateStructuredBuffer(GraphBuilder, TEXT("InstanceCulling.OutputOffsetBufferOut"), NullArray);
+				FRDGBufferRef InstanceIdOutOffsetBufferRDG = CreateStructuredBuffer(GraphBuilder, TEXT("InstanceCulling.OutputOffsetBufferOut"), sizeof(uint32), 1, nullptr, 0);
+				AddClearUAVPass(GraphBuilder, GraphBuilder.CreateUAV(InstanceIdOutOffsetBufferRDG), 0);
 
 				PassParameters->NumIndirectArgs = NumIndirectArgs;
 				PassParameters->InstanceIdOffsetBufferOut = GraphBuilder.CreateUAV(InstanceIdOffsetBufferRDG, PF_R32_UINT);;
@@ -1706,7 +1712,7 @@ void FVirtualShadowMapArray::RenderVirtualShadowMapsHw(FRDGBuilder& GraphBuilder
 				PassParameters->TmpInstanceIdOffsetBufferOut = GraphBuilder.CreateUAV(TmpInstanceIdOffsetBufferRDG);
 				PassParameters->DrawIndirectArgsBuffer = GraphBuilder.CreateSRV(DrawIndirectArgsRDG, PF_R32_UINT);
 
-				auto ComputeShader = GetGlobalShaderMap(GMaxRHIFeatureLevel)->GetShader<FAllocateCommandInstanceOutputSpaceCs>();
+				auto ComputeShader = ShaderMap->GetShader<FAllocateCommandInstanceOutputSpaceCs>();
 
 				FComputeShaderUtils::AddPass(
 					GraphBuilder,
@@ -1732,7 +1738,7 @@ void FVirtualShadowMapArray::RenderVirtualShadowMapsHw(FRDGBuilder& GraphBuilder
 				PassParameters->VisibleInstanceCountBuffer = GraphBuilder.CreateSRV(VisibleInstanceWriteOffsetRDG);
 				PassParameters->IndirectArgs = OutputPassIndirectArgs;
 
-				auto ComputeShader = GetGlobalShaderMap(GMaxRHIFeatureLevel)->GetShader<FOutputCommandInstanceListsCs>();
+				auto ComputeShader = ShaderMap->GetShader<FOutputCommandInstanceListsCs>();
 
 				FComputeShaderUtils::AddPass(
 					GraphBuilder,
diff --git a/Engine/Source/Runtime/Renderer/Public/InstanceCulling/InstanceCullingContext.h b/Engine/Source/Runtime/Renderer/Public/InstanceCulling/InstanceCullingContext.h
--- a/Engine/Source/Runtime/Renderer/Public/InstanceCulling/InstanceCullingContext.h
+++ b/Engine/Source/Runtime/Renderer/Public/InstanceCulling/InstanceCullingContext.h
@@ -132,6 +132,12 @@ public:
 
 	static bool AllowBatchedBuildRenderingCommands(const FGPUScene& GPUScene);
 
+	/**
+	 * Helper function to add a pass to zero the instance count in the indirect args.
+	 */
+	static void AddClearIndirectArgInstanceCountPass(FRDGBuilder& GraphBuilder, FGlobalShaderMap* ShaderMap, FRDGBufferRef DrawIndirectArgsBuffer);
+
+
 	void SetupDrawCommands(
 		FMeshCommandOneFrameArray& VisibleMeshDrawCommandsInOut,
 		bool bCompactIdenticalCommands,
