Disable Parallel RDG for Mobile Platforms from the RenderGraph

diff --git a/Engine/Source/Runtime/RenderCore/Private/RenderGraphPrivate.h b/Engine/Source/Runtime/RenderCore/Private/RenderGraphPrivate.h
--- a/Engine/Source/Runtime/RenderCore/Private/RenderGraphPrivate.h
+++ b/Engine/Source/Runtime/RenderCore/Private/RenderGraphPrivate.h
@@ -203,7 +203,7 @@ FORCEINLINE bool IsParallelExecuteEnabled()
 		&& !IsImmediateMode()
 		&& !GRDGDebug
 		&& !GRDGTransitionLog
-		&& !IsOpenGLPlatform(GMaxRHIShaderPlatform);
+		&& !IsMobilePlatform(GMaxRHIShaderPlatform);
 }
 
 template <typename ResourceRegistryType, typename FunctionType>
diff --git a/Engine/Source/Runtime/Renderer/Private/PlanarReflectionRendering.cpp b/Engine/Source/Runtime/Renderer/Private/PlanarReflectionRendering.cpp
--- a/Engine/Source/Runtime/Renderer/Private/PlanarReflectionRendering.cpp
+++ b/Engine/Source/Runtime/Renderer/Private/PlanarReflectionRendering.cpp
@@ -346,11 +346,7 @@ static void UpdatePlanarReflectionContents_RenderThread(
 #else
 				SCOPED_DRAW_EVENT(RHICmdList, UpdatePlanarReflectionContent_RenderThread);
 #endif
-				const ERDGBuilderFlags BuilderFlags = SceneRenderer->Scene->GetShadingPath() == EShadingPath::Mobile
-					? ERDGBuilderFlags::None
-					: ERDGBuilderFlags::AllowParallelExecute;
-
-				FRDGBuilder GraphBuilder(RHICmdList, RDG_EVENT_NAME("PlanarReflection"), BuilderFlags);
+				FRDGBuilder GraphBuilder(RHICmdList, RDG_EVENT_NAME("PlanarReflection"), ERDGBuilderFlags::AllowParallelExecute);
 
 				// Reflection view late update
 				if (SceneRenderer->Views.Num() > 1)
diff --git a/Engine/Source/Runtime/Renderer/Private/ReflectionEnvironmentCapture.cpp b/Engine/Source/Runtime/Renderer/Private/ReflectionEnvironmentCapture.cpp
--- a/Engine/Source/Runtime/Renderer/Private/ReflectionEnvironmentCapture.cpp
+++ b/Engine/Source/Runtime/Renderer/Private/ReflectionEnvironmentCapture.cpp
@@ -723,11 +723,7 @@ void CaptureSceneToScratchCubemap(FRHICommandListImmediate& RHICmdList, FSceneRe
 	{
 		SCOPED_DRAW_EVENT(RHICmdList, CubeMapCapture);
 
-		const ERDGBuilderFlags BuilderFlags = SceneRenderer->Scene->GetShadingPath() == EShadingPath::Mobile
-			? ERDGBuilderFlags::None
-			: ERDGBuilderFlags::AllowParallelExecute;
-
-		FRDGBuilder GraphBuilder(RHICmdList, RDG_EVENT_NAME("CubeMapCapture"), BuilderFlags);
+		FRDGBuilder GraphBuilder(RHICmdList, RDG_EVENT_NAME("CubeMapCapture"), ERDGBuilderFlags::AllowParallelExecute);
 
 		// Render the scene normally for one face of the cubemap
 		SceneRenderer->Render(GraphBuilder);
diff --git a/Engine/Source/Runtime/Renderer/Private/SceneCaptureRendering.cpp b/Engine/Source/Runtime/Renderer/Private/SceneCaptureRendering.cpp
--- a/Engine/Source/Runtime/Renderer/Private/SceneCaptureRendering.cpp
+++ b/Engine/Source/Runtime/Renderer/Private/SceneCaptureRendering.cpp
@@ -283,16 +283,12 @@ static void UpdateSceneCaptureContentDeferred_RenderThread(
 	// update any resources that needed a deferred update
 	FDeferredUpdateResource::UpdateResources(RHICmdList);
 	{
-		const ERDGBuilderFlags BuilderFlags = SceneRenderer->Scene->GetShadingPath() == EShadingPath::Mobile
-			? ERDGBuilderFlags::None
-			: ERDGBuilderFlags::AllowParallelExecute;
-
 #if WANTS_DRAW_MESH_EVENTS
 		SCOPED_DRAW_EVENTF(RHICmdList, SceneCapture, TEXT("SceneCapture %s"), *EventName);
-		FRDGBuilder GraphBuilder(RHICmdList, RDG_EVENT_NAME("SceneCapture %s", *EventName), BuilderFlags);
+		FRDGBuilder GraphBuilder(RHICmdList, RDG_EVENT_NAME("SceneCapture %s", *EventName), ERDGBuilderFlags::AllowParallelExecute);
 #else
 		SCOPED_DRAW_EVENT(RHICmdList, UpdateSceneCaptureContent_RenderThread);
-		FRDGBuilder GraphBuilder(RHICmdList, RDG_EVENT_NAME("SceneCapture"), BuilderFlags);
+		FRDGBuilder GraphBuilder(RHICmdList, RDG_EVENT_NAME("SceneCapture"), ERDGBuilderFlags::AllowParallelExecute);
 #endif
 
 		FRDGTextureRef TargetTexture = RegisterExternalTexture(GraphBuilder, RenderTarget->GetRenderTargetTexture(), TEXT("SceneCaptureTarget"));
diff --git a/Engine/Source/Runtime/Renderer/Private/SceneRendering.cpp b/Engine/Source/Runtime/Renderer/Private/SceneRendering.cpp
--- a/Engine/Source/Runtime/Renderer/Private/SceneRendering.cpp
+++ b/Engine/Source/Runtime/Renderer/Private/SceneRendering.cpp
@@ -3836,17 +3836,13 @@ static void RenderViewFamily_RenderThread(FRHICommandListImmediate& RHICmdList,
 		SCOPE_CYCLE_COUNTER(STAT_TotalSceneRenderingTime);
 
 		{
-			const ERDGBuilderFlags BuilderFlags = SceneRenderer->Scene->GetShadingPath() == EShadingPath::Mobile
-				? ERDGBuilderFlags::None
-				: ERDGBuilderFlags::AllowParallelExecute;
-
 			FRDGBuilder GraphBuilder(
 				RHICmdList,
 				RDG_EVENT_NAME("SceneRenderer_%s(ViewFamily=%s)",
 					ViewFamily.EngineShowFlags.HitProxies ? TEXT("RenderHitProxies") : TEXT("Render"),
 					ViewFamily.bResolveScene ? TEXT("Primary") : TEXT("Auxiliary")
 				),
-				BuilderFlags
+				ERDGBuilderFlags::AllowParallelExecute
 			);
 
 			if (ViewFamily.EngineShowFlags.HitProxies)
