Removing dead code path in mesh draw command task that would break RDG parallelism if enabled.

diff --git a/Engine/Source/Runtime/Renderer/Private/MeshDrawCommands.cpp b/Engine/Source/Runtime/Renderer/Private/MeshDrawCommands.cpp
--- a/Engine/Source/Runtime/Renderer/Private/MeshDrawCommands.cpp
+++ b/Engine/Source/Runtime/Renderer/Private/MeshDrawCommands.cpp
@@ -1089,7 +1089,6 @@ void FParallelMeshDrawCommandPass::DispatchPassSetup(
 	if (MaxNumDraws > 0)
 	{
 		// Preallocate resources on rendering thread based on MaxNumDraws.
-		bPrimitiveIdBufferDataOwnedByRHIThread = false;
 		TaskContext.PrimitiveIdBufferDataSize = TaskContext.InstanceFactor * MaxNumDraws * sizeof(int32);
 		TaskContext.PrimitiveIdBufferData = FMemory::Malloc(TaskContext.PrimitiveIdBufferDataSize);
 #if DO_GUARD_SLOW
@@ -1173,24 +1172,10 @@ void FParallelMeshDrawCommandPass::WaitForTasksAndEmpty(EWaitThread WaitThread)
 
 	if (MaxNumDraws > 0)
 	{
-		if (bPrimitiveIdBufferDataOwnedByRHIThread)
-		{
-			FRHICommandListExecutor::GetImmediateCommandList().EnqueueLambda([PrimitiveIdVertexBufferPoolEntry = PrimitiveIdVertexBufferPoolEntry](FRHICommandListImmediate&) {
-				GPrimitiveIdVertexBufferPool.ReturnToFreeList(PrimitiveIdVertexBufferPoolEntry);
-			});
-		}
-		else
-		{
-			GPrimitiveIdVertexBufferPool.ReturnToFreeList(PrimitiveIdVertexBufferPoolEntry);
-		}
-	}
-
-	if (!bPrimitiveIdBufferDataOwnedByRHIThread)
-	{
-		FMemory::Free(TaskContext.PrimitiveIdBufferData);
+		GPrimitiveIdVertexBufferPool.ReturnToFreeList(PrimitiveIdVertexBufferPoolEntry);
 	}
 
-	bPrimitiveIdBufferDataOwnedByRHIThread = false;
+	FMemory::Free(TaskContext.PrimitiveIdBufferData);
 	MaxNumDraws = 0;
 	PassNameForStats.Empty();
 
diff --git a/Engine/Source/Runtime/Renderer/Private/MeshDrawCommands.h b/Engine/Source/Runtime/Renderer/Private/MeshDrawCommands.h
--- a/Engine/Source/Runtime/Renderer/Private/MeshDrawCommands.h
+++ b/Engine/Source/Runtime/Renderer/Private/MeshDrawCommands.h
@@ -145,8 +145,7 @@ public:
 	};
 
 	FParallelMeshDrawCommandPass()
-		: bPrimitiveIdBufferDataOwnedByRHIThread(false)
-		, bHasInstanceCullingDrawParameters(false)
+		: bHasInstanceCullingDrawParameters(false)
 		, MaxNumDraws(0)
 	{
 	}
@@ -217,8 +216,6 @@ private:
 	FGraphEventRef TaskEventRef;
 	FString PassNameForStats;
 
-	// If TaskContext::PrimitiveIdBufferData will be released by RHI Thread.
-	mutable bool bPrimitiveIdBufferDataOwnedByRHIThread;
 	bool bHasInstanceCullingDrawParameters;
 
 	// Maximum number of draws for this pass. Used to prealocate resources on rendering thread. 
