Disables FFT bloom if still use the default texture proxy, and invalidate caches the FFT bloom kernel based on the RHITexture*

diff --git a/Engine/Source/Runtime/Renderer/Private/PostProcess/PostProcessFFTBloom.cpp b/Engine/Source/Runtime/Renderer/Private/PostProcess/PostProcessFFTBloom.cpp
--- a/Engine/Source/Runtime/Renderer/Private/PostProcess/PostProcessFFTBloom.cpp
+++ b/Engine/Source/Runtime/Renderer/Private/PostProcess/PostProcessFFTBloom.cpp
@@ -469,17 +469,11 @@ void InitDomainAndGetKernel(
 
 	const auto& PPSettings = View.FinalPostProcessSettings;
 
-	UTexture2D* BloomConvolutionTexture = PPSettings.BloomConvolutionTexture;
 	const FTexture2DResource* BloomConvolutionTextureResource = View.FFTBloomKernelTexture;
-
-	// Fall back to the default bloom texture if provided.
-	if (BloomConvolutionTexture == nullptr)
-	{
-		BloomConvolutionTexture = GEngine->DefaultBloomKernelTexture;
-	}
+	const FTextureRHIRef& PhysicalSpaceKernelTextureRef = BloomConvolutionTextureResource->TextureRHI;
 
 	// This should exist if we called IsFFTBloomPhysicalKernelReady.
-	check(BloomConvolutionTexture && BloomConvolutionTextureResource && BloomConvolutionTextureResource->TextureRHI);
+	check(BloomConvolutionTextureResource && PhysicalSpaceKernelTextureRef);
 
 	const float BloomConvolutionSize = PPSettings.BloomConvolutionSize;
 	const FVector2D CenterUV = PPSettings.BloomConvolutionCenterUV;
@@ -504,7 +498,7 @@ void InitDomainAndGetKernel(
 
 		FRDGTextureRef PrevCachedSpectralKernel = GraphBuilder.RegisterExternalTexture(FFTKernel.Spectral);
 
-		const bool bSameTexture = (FFTKernel.Physical == static_cast<const UTexture2D*>(BloomConvolutionTexture));
+		const bool bSameTexture = (FFTKernel.PhysicalRHI == PhysicalSpaceKernelTextureRef);
 		const bool bSameSpectralBuffer = TransformDesc.ClearValue == PrevCachedSpectralKernel->Desc.ClearValue
 			&& TransformDesc.Flags == PrevCachedSpectralKernel->Desc.Flags
 			&& TransformDesc.Format == PrevCachedSpectralKernel->Desc.Format
@@ -531,9 +525,6 @@ void InitDomainAndGetKernel(
 	{
 		RDG_EVENT_SCOPE(GraphBuilder, "InitDomainAndGetKernel");
 
-		// Sample the physical space kernel into the resized buffer
-		const FTextureRHIRef& PhysicalSpaceKernelTextureRef = BloomConvolutionTextureResource->TextureRHI;
-
 		// Rescale the physical space kernel ( and omit the center if this is a 1/2 resolution fft, it will be added later)
 		FRDGTextureRef ResizedKernel = GraphBuilder.CreateTexture(TransformDesc, TEXT("Bloom.FFT.ResizedKernel"));
 		ResizeAndCenterTexture(
@@ -565,7 +556,8 @@ void InitDomainAndGetKernel(
 		{
 			ViewState->BloomFFTKernel.Scale = BloomConvolutionSize;
 			ViewState->BloomFFTKernel.ImageSize = Intermediates.ImageRect.Size();
-			ViewState->BloomFFTKernel.Physical = BloomConvolutionTexture;
+			ViewState->BloomFFTKernel.Physical = View.FinalPostProcessSettings.BloomConvolutionTexture;
+			ViewState->BloomFFTKernel.PhysicalRHI = PhysicalSpaceKernelTextureRef;
 			ViewState->BloomFFTKernel.CenterUV = CenterUV;
 			ViewState->BloomFFTKernel.PhysicalMipLevel = BloomConvolutionTextureResource->GetCurrentMipCount();
 			{
diff --git a/Engine/Source/Runtime/Renderer/Private/ScenePrivate.h b/Engine/Source/Runtime/Renderer/Private/ScenePrivate.h
--- a/Engine/Source/Runtime/Renderer/Private/ScenePrivate.h
+++ b/Engine/Source/Runtime/Renderer/Private/ScenePrivate.h
@@ -970,8 +970,9 @@ public:
 		TRefCountPtr<IPooledRenderTarget> Spectral;
 		TRefCountPtr<IPooledRenderTarget> CenterWeight; // a 1-pixel buffer that holds blend weights for half-resolution fft.
 
-														// The physical space source texture
-		UTexture2D* Physical = NULL;
+		// The physical space source texture
+		UTexture2D* Physical = nullptr;
+		FRHITexture* PhysicalRHI = nullptr;
 
 		// The Scale * 100 = percentage of the image space that the physical kernel represents.
 		// e.g. Scale = 1 indicates that the physical kernel image occupies the same size 
diff --git a/Engine/Source/Runtime/Renderer/Private/SceneRendering.cpp b/Engine/Source/Runtime/Renderer/Private/SceneRendering.cpp
--- a/Engine/Source/Runtime/Renderer/Private/SceneRendering.cpp
+++ b/Engine/Source/Runtime/Renderer/Private/SceneRendering.cpp
@@ -2245,7 +2245,7 @@ static bool PreparePostProcessSettingTextureForRenderer(const FViewInfo& View, U
 
 	const uint32 FramesPerWarning = 15;
 
-	if (bIsValid && Texture2D->IsFullyStreamedIn() == false)
+	if (bIsValid && (Texture2D->IsFullyStreamedIn() == false || Texture2D->HasPendingInitOrStreaming()))
 	{
 		if ((View.Family->FrameNumber % FramesPerWarning) == 0)
 		{
@@ -2265,6 +2265,20 @@ static bool PreparePostProcessSettingTextureForRenderer(const FViewInfo& View, U
 		bIsValid = false;
 	}
 
+#if WITH_EDITOR
+	if (bIsValid && Texture2D->IsDefaultTexture())
+#else
+	if (bIsValid && (!Texture2D->GetResource() || Texture2D->GetResource()->IsProxy()))
+#endif
+	{
+		if ((View.Family->FrameNumber % FramesPerWarning) == 0)
+		{
+			UE_LOG(LogRenderer, Warning, TEXT("The %s texture is still using the default texture proxy."), TextureUsageName);
+		}
+
+		bIsValid = false;
+	}
+
 	return bIsValid;
 };
 
@@ -2345,11 +2359,11 @@ FSceneRenderer::FSceneRenderer(const FSceneViewFamily* InViewFamily,FHitProxyCon
 		// Handle the FFT bloom kernel textire
 		if (ViewInfo->FinalPostProcessSettings.BloomMethod == EBloomMethod::BM_FFT && ViewInfo->ViewState != nullptr)
 		{
-			GEngine->LoadDefaultBloomTexture();
-
 			UTexture2D* BloomConvolutionTexture = ViewInfo->FinalPostProcessSettings.BloomConvolutionTexture;
 			if (BloomConvolutionTexture == nullptr)
 			{
+				GEngine->LoadDefaultBloomTexture();
+
 				BloomConvolutionTexture = GEngine->DefaultBloomKernelTexture;
 			}
 
@@ -2367,11 +2381,6 @@ FSceneRenderer::FSceneRenderer(const FSceneViewFamily* InViewFamily,FHitProxyCon
 				{
 					ViewInfo->FinalPostProcessSettings.BloomConvolutionTexture = nullptr;
 				}
-
-				if (BloomConvolutionTexture == GEngine->DefaultBloomKernelTexture)
-				{
-
-				}
 			}
 		}
 
