Added support for desktop dynamic VRS. Fixed foveation level adjusts based on GPU utilization.

diff --git a/Engine/Source/Runtime/Engine/Classes/Engine/RendererSettings.h b/Engine/Source/Runtime/Engine/Classes/Engine/RendererSettings.h
--- a/Engine/Source/Runtime/Engine/Classes/Engine/RendererSettings.h
+++ b/Engine/Source/Runtime/Engine/Classes/Engine/RendererSettings.h
@@ -482,7 +482,7 @@ class ENGINE_API URendererSettings : public UDeveloperSettings
 	TEnumAsByte<EFixedFoveationLevels::Type> HMDFixedFoveationLevel;
 
 	UPROPERTY(config, EditAnywhere, Category = VR, meta = (
-		ConsoleVariable = "vr.VRS.HMDFixedFoveationDynamic", DisplayName = "Dynamic Fixed Foveation (Oculus Quest/Quest 2 Only)",
+		ConsoleVariable = "vr.VRS.HMDFixedFoveationDynamic", DisplayName = "Dynamic Fixed Foveation",
 		ToolTip = "Allows fixed foveation level to adjust dynamically based on GPU utilization.\nLevel will range between none at the minimum, and the currently selected foveation level at the maximum."))
 	uint32 bHMDFixedFoveationDynamic:1;
 
diff --git a/Engine/Source/Runtime/Renderer/Private/VariableRateShadingImageManager.cpp b/Engine/Source/Runtime/Renderer/Private/VariableRateShadingImageManager.cpp
--- a/Engine/Source/Runtime/Renderer/Private/VariableRateShadingImageManager.cpp
+++ b/Engine/Source/Runtime/Renderer/Private/VariableRateShadingImageManager.cpp
@@ -44,7 +44,7 @@ static TAutoConsoleVariable<int> CVarHMDFixedFoveationLevel(
 static TAutoConsoleVariable<int32> CVarHMDFixedFoveationDynamic(
 	TEXT("vr.VRS.HMDFixedFoveationDynamic"),
 	0,
-	TEXT("Whether fixed-foveation level should adjust based on GPU utilization (Currently OculusVR only)\n")
+	TEXT("Whether fixed-foveation level should adjust based on GPU utilization\n")
 	TEXT(" 0: Disabled (default);\n")
 	TEXT(" 1: Enabled\n"),
 	ECVF_RenderThreadSafe);
@@ -394,19 +394,26 @@ void FVariableRateShadingImageManager::Tick()
 void FVariableRateShadingImageManager::UpdateFixedFoveationParameters(FVRSImageGenerationParameters& VRSImageGenParamsInOut)
 {
 
-	// Numbers here are currently pretty arbitrary... should eventually be chosen based on headset specific characteristics
-	static const float FIXED_FOVEATION_FULL_RATE_CUTOFFS[] = { 1.0f, 0.7f, 0.50f, 0.35f, 0.35f };
-	static const float FIXED_FOVEATION_HALF_RATE_CUTOFFS[] = { 1.0f, 0.9f, 0.75f, 0.55f, 0.55f };
-	static const float FIXED_FOVEATION_CENTER_X[] = { 0.5f, 0.5f, 0.5f, 0.5f, 0.5f };
-	static const float FIXED_FOVEATION_CENTER_Y[] = { 0.5f, 0.5f, 0.5f, 0.5f, 0.42f };
-	
-	int Level = FMath::Clamp(CVarHMDFixedFoveationLevel.GetValueOnAnyThread(), 0, 4);
+	// VRS level parameters - pretty arbitrary right now, later should depend on device characteristics
+	static const float kFixedFoveationFullRateCutoffs[] = { 1.0f, 0.7f, 0.50f, 0.35f, 0.35f };
+	static const float kFixedFoveationHalfRateCutofffs[] = { 1.0f, 0.9f, 0.75f, 0.55f, 0.55f };
+	static const float kFixedFoveationCenterX[] = { 0.5f, 0.5f, 0.5f, 0.5f, 0.5f };
+	static const float kFixedFoveationCenterY[] = { 0.5f, 0.5f, 0.5f, 0.5f, 0.42f };
+
 
-	VRSImageGenParamsInOut.bGenerateFixedFoveation = Level ? true : false;
-	VRSImageGenParamsInOut.HMDFixedFoveationFullRateCutoff = FIXED_FOVEATION_FULL_RATE_CUTOFFS[Level];
-	VRSImageGenParamsInOut.HMDFixedFoveationHalfRateCutoff = FIXED_FOVEATION_HALF_RATE_CUTOFFS[Level];
-	VRSImageGenParamsInOut.HMDFixedFoveationCenterX = FIXED_FOVEATION_CENTER_X[Level];
-	VRSImageGenParamsInOut.HMDFixedFoveationCenterY = FIXED_FOVEATION_CENTER_Y[Level];
+	const int VRSMaxLevel = FMath::Clamp(CVarHMDFixedFoveationLevel.GetValueOnAnyThread(), 0, 4);
+	float VRSAmount = 1.0f;
+	
+	if (CVarHMDFixedFoveationDynamic.GetValueOnAnyThread() && VRSMaxLevel > 0)
+	{
+		VRSAmount = GetDynamicVRSAmount();
+	}
+	
+	VRSImageGenParamsInOut.bGenerateFixedFoveation = (VRSMaxLevel > 0 && VRSAmount > 0.0f);
+	VRSImageGenParamsInOut.HMDFixedFoveationFullRateCutoff = FMath::Lerp(kFixedFoveationFullRateCutoffs[0], kFixedFoveationFullRateCutoffs[VRSMaxLevel], VRSAmount);
+	VRSImageGenParamsInOut.HMDFixedFoveationHalfRateCutoff = FMath::Lerp(kFixedFoveationHalfRateCutofffs[0], kFixedFoveationHalfRateCutofffs[VRSMaxLevel], VRSAmount);
+	VRSImageGenParamsInOut.HMDFixedFoveationCenterX = FMath::Lerp(kFixedFoveationCenterX[0], kFixedFoveationCenterX[VRSMaxLevel], VRSAmount);
+	VRSImageGenParamsInOut.HMDFixedFoveationCenterY = FMath::Lerp(kFixedFoveationCenterY[0], kFixedFoveationCenterY[VRSMaxLevel], VRSAmount);
 }
 
 void FVariableRateShadingImageManager::UpdateEyeTrackedFoveationParameters(FVRSImageGenerationParameters& VRSImageGenParamsInOut, const FSceneViewFamily& ViewFamily)
@@ -429,3 +436,53 @@ void FVariableRateShadingImageManager::UpdateEyeTrackedFoveationParameters(FVRSI
 
 	// @todo:
 }
+
+float FVariableRateShadingImageManager::GetDynamicVRSAmount()
+{
+	const float		kUpdateIncrement		= 0.1f;
+	const int		kNumFramesToAverage		= 10; 
+	const double	kRoundingErrorMargin	= 0.5; // Add a small margin (.5 ms) to avoid oscillation
+	const double	kDesktopMaxAllowedFrameTime	= 12.50; // 80 fps
+	const double	kMobileMaxAllowedFrameTime	= 16.66; // 60 fps
+
+	if (GFrameNumber != DynamicVRSData.LastUpdateFrame)
+	{
+		DynamicVRSData.LastUpdateFrame = GFrameNumber;
+
+		const double GPUFrameTime = FPlatformTime::ToMilliseconds(RHIGetGPUFrameCycles());
+		
+		// Update sum for rolling average
+		if (DynamicVRSData.NumFramesStored < kNumFramesToAverage)
+		{
+			DynamicVRSData.SumBusyTime += GPUFrameTime;
+			DynamicVRSData.NumFramesStored++;
+		}
+		else if (DynamicVRSData.NumFramesStored == kNumFramesToAverage)
+		{
+			DynamicVRSData.SumBusyTime -= DynamicVRSData.SumBusyTime / kNumFramesToAverage;
+			DynamicVRSData.SumBusyTime += GPUFrameTime;
+		}
+
+		// If rolling average has sufficient samples, check if update is necessary
+		if (DynamicVRSData.NumFramesStored == kNumFramesToAverage)
+		{
+			const double AverageBusyTime = DynamicVRSData.SumBusyTime / kNumFramesToAverage;
+			const double TargetBusyTime = IsMobilePlatform(GShaderPlatformForFeatureLevel[GMaxRHIFeatureLevel]) ? kMobileMaxAllowedFrameTime : kDesktopMaxAllowedFrameTime;
+			
+			if (AverageBusyTime > TargetBusyTime + kRoundingErrorMargin && DynamicVRSData.VRSAmount < 1.0f)
+			{
+				DynamicVRSData.VRSAmount = FMath::Clamp(DynamicVRSData.VRSAmount + kUpdateIncrement, 0.0f, 1.0f);
+				DynamicVRSData.SumBusyTime = 0.0;
+				DynamicVRSData.NumFramesStored = 0;
+			}
+			else if (AverageBusyTime < TargetBusyTime - kRoundingErrorMargin && DynamicVRSData.VRSAmount > 0.0f)
+			{
+				DynamicVRSData.VRSAmount = FMath::Clamp(DynamicVRSData.VRSAmount - kUpdateIncrement, 0.0f, 1.0f);
+				DynamicVRSData.SumBusyTime = 0.0;
+				DynamicVRSData.NumFramesStored = 0;
+			}
+		}
+	}
+
+	return DynamicVRSData.VRSAmount;
+}
diff --git a/Engine/Source/Runtime/Renderer/Public/VariableRateShadingImageManager.h b/Engine/Source/Runtime/Renderer/Public/VariableRateShadingImageManager.h
--- a/Engine/Source/Runtime/Renderer/Public/VariableRateShadingImageManager.h
+++ b/Engine/Source/Runtime/Renderer/Public/VariableRateShadingImageManager.h
@@ -42,6 +42,14 @@ enum class EVRSGenerationFlags : uint32;
 
 struct FVRSImageGenerationParameters;
 
+struct FDynamicVRSData
+{
+	float	VRSAmount = 1.0f;
+	double	SumBusyTime = 0.0;
+	int		NumFramesStored = 0;
+	uint32	LastUpdateFrame = 0;
+};
+
 class FVariableRateShadingImageManager : public FRenderResource
 {
 public:
@@ -64,6 +72,7 @@ protected:
 	uint64 CalculateVRSImageHash(const FVRSImageGenerationParameters& VRSImageGenParamsIn, EVRSGenerationFlags ViewFlags) const;
 	void UpdateFixedFoveationParameters(FVRSImageGenerationParameters& VRSImageGenParamsInOut);
 	void UpdateEyeTrackedFoveationParameters(FVRSImageGenerationParameters& VRSImageGenParamsInOut, const FSceneViewFamily& ViewFamily);
+	float GetDynamicVRSAmount();
 
 	/** Per frame tick/update. */
 	void Tick();
@@ -89,6 +98,7 @@ private:
 	TRefCountPtr<IPooledRenderTarget> MobileHMDFixedFoveationOverrideImage;
 
 	uint64 LastFrameTick;
+	FDynamicVRSData DynamicVRSData;
 };
 
 RENDERER_API extern TGlobalResource<FVariableRateShadingImageManager> GVRSImageManager;
