Add stat scope for virtual texture finalize and page table update phases.

diff --git a/Engine/Source/Runtime/Renderer/Private/VT/VirtualTextureSystem.cpp b/Engine/Source/Runtime/Renderer/Private/VT/VirtualTextureSystem.cpp
--- a/Engine/Source/Runtime/Renderer/Private/VT/VirtualTextureSystem.cpp
+++ b/Engine/Source/Runtime/Renderer/Private/VT/VirtualTextureSystem.cpp
@@ -2252,6 +2252,8 @@ void FVirtualTextureSystem::SubmitRequests(FRDGBuilder& GraphBuilder, ERHIFeatur
 	// Finalize requests
 	{
 		SCOPE_CYCLE_COUNTER(STAT_ProcessRequests_Finalize);
+		RDG_EVENT_SCOPE(GraphBuilder, "VirtualTextureFinalizeRequests");
+		
 		for (IVirtualTextureFinalizer* VTFinalizer : Finalizers)
 		{
 			VTFinalizer->Finalize(GraphBuilder);
@@ -2262,6 +2264,8 @@ void FVirtualTextureSystem::SubmitRequests(FRDGBuilder& GraphBuilder, ERHIFeatur
 	// Update page tables
 	{
 		SCOPE_CYCLE_COUNTER(STAT_PageTableUpdates);
+		RDG_EVENT_SCOPE(GraphBuilder, "VirtualTexturePageTableUpdates");
+
 		for (uint32 ID = 0; ID < MaxSpaces; ID++)
 		{
 			if (Spaces[ID])
