Add r.vt.DumpPoolUsage.
This outputs the physical tile footprint of each VT producer.
Useful for finding producers that are using up most of the physical cache.

diff --git a/Engine/Source/Runtime/Renderer/Private/VT/TexturePagePool.cpp b/Engine/Source/Runtime/Renderer/Private/VT/TexturePagePool.cpp
--- a/Engine/Source/Runtime/Renderer/Private/VT/TexturePagePool.cpp
+++ b/Engine/Source/Runtime/Renderer/Private/VT/TexturePagePool.cpp
@@ -331,6 +331,19 @@ uint32 FTexturePagePool::GetNumVisiblePages(uint32 Frame) const
 	return Count;
 }
 
+void FTexturePagePool::CollectProducerCounts(TMap<uint32, uint32>& OutProducerCountMap) const
+{
+	for (uint32 i = 0; i < NumPages; ++i)
+	{
+		const uint32 PackedProducerHandle = Pages[i].PackedProducerHandle;
+		if (PackedProducerHandle != 0u)
+		{
+			OutProducerCountMap.FindOrAdd(PackedProducerHandle) += 1;
+		}
+	}
+}
+
+
 void FTexturePagePool::MapPage(FVirtualTextureSpace* Space, FVirtualTexturePhysicalSpace* PhysicalSpace, uint8 PageTableLayerIndex, uint8 MaxLevel, uint8 vLogSize, uint32 vAddress, uint8 Local_vLevel, uint16 pAddress)
 {
 	check(pAddress < NumPages);
diff --git a/Engine/Source/Runtime/Renderer/Private/VT/TexturePagePool.h b/Engine/Source/Runtime/Renderer/Private/VT/TexturePagePool.h
--- a/Engine/Source/Runtime/Renderer/Private/VT/TexturePagePool.h
+++ b/Engine/Source/Runtime/Renderer/Private/VT/TexturePagePool.h
@@ -130,6 +130,11 @@ public:
 	 */
 	uint32		GetNumVisiblePages(uint32 Frame) const;
 
+	/**
+	 * Returns a map between producer handle and number of tiles that the producer uses in the pool.
+	 */
+	void		CollectProducerCounts(TMap<uint32, uint32>& OutProducerCountMap) const;
+
 	/**
 	* Map the physical address to a specific virtual address.
 	*/
diff --git a/Engine/Source/Runtime/Renderer/Private/VT/VirtualTextureSystem.cpp b/Engine/Source/Runtime/Renderer/Private/VT/VirtualTextureSystem.cpp
--- a/Engine/Source/Runtime/Renderer/Private/VT/VirtualTextureSystem.cpp
+++ b/Engine/Source/Runtime/Renderer/Private/VT/VirtualTextureSystem.cpp
@@ -238,10 +238,12 @@ FVirtualTextureSystem::FVirtualTextureSystem()
 	, bFlushCaches(false)
 	, FlushCachesCommand(TEXT("r.VT.Flush"), TEXT("Flush all the physical caches in the VT system."),
 		FConsoleCommandDelegate::CreateRaw(this, &FVirtualTextureSystem::FlushCachesFromConsole))
-	, DumpCommand(TEXT("r.VT.Dump"), TEXT("Lot a whole lot of info on the VT system state."),
+	, DumpCommand(TEXT("r.VT.Dump"), TEXT("Dump a whole lot of info on the VT system state."),
 		FConsoleCommandDelegate::CreateRaw(this, &FVirtualTextureSystem::DumpFromConsole))
-	, ListPhysicalPools(TEXT("r.VT.ListPhysicalPools"), TEXT("Lot a whole lot of info on the VT system state."),
+	, ListPhysicalPools(TEXT("r.VT.ListPhysicalPools"), TEXT("Dump a whole lot of info on the VT system state."),
 		FConsoleCommandDelegate::CreateRaw(this, &FVirtualTextureSystem::ListPhysicalPoolsFromConsole))
+	, DumpPoolUsageCommand(TEXT("r.VT.DumpPoolUsage"), TEXT("Dump detailed info about VT pool usage."),
+		FConsoleCommandDelegate::CreateRaw(this, &FVirtualTextureSystem::DumpPoolUsageFromConsole))
 #if WITH_EDITOR
 	, SaveAllocatorImages(TEXT("r.VT.SaveAllocatorImages"), TEXT("Save images showing allocator usage."),
 		FConsoleCommandDelegate::CreateRaw(this, &FVirtualTextureSystem::SaveAllocatorImagesFromConsole))
@@ -420,6 +422,48 @@ void FVirtualTextureSystem::ListPhysicalPoolsFromConsole()
 	UE_LOG(LogConsoleResponse, Display, TEXT("TotalPhysicalMemory: %fMB"), (double)TotalPhysicalMemory / 1024.0 / 1024.0);
 }
 
+void FVirtualTextureSystem::DumpPoolUsageFromConsole()
+{
+	for (int32 i = 0; i < PhysicalSpaces.Num(); ++i)
+	{
+		if (PhysicalSpaces[i])
+		{
+			const FVirtualTexturePhysicalSpace& PhysicalSpace = *PhysicalSpaces[i];
+			const FVTPhysicalSpaceDescription& Desc = PhysicalSpace.GetDescription();
+			const FTexturePagePool& PagePool = PhysicalSpace.GetPagePool();
+
+			UE_LOG(LogConsoleResponse, Display, TEXT("PhysicaPool: [%i] %s (%ix%i):"), i, *PhysicalSpace.GetFormatString(), Desc.TileSize, Desc.TileSize);
+
+			TMap<uint32, uint32> ProducerCountMap;
+			PagePool.CollectProducerCounts(ProducerCountMap);
+
+			TSet< TPair<uint32, uint32> > SortedProducerCounts;
+			for (TPair<uint32, uint32> ProducerCount : ProducerCountMap)
+			{
+				// Filter out producers that only have locked page mapped.
+				// In future we can add other filters here (count, age etc), or we can dump more info and process off line.
+				if (ProducerCount.Value > 1)
+				{
+					SortedProducerCounts.Add(ProducerCount);
+				}
+			}
+			SortedProducerCounts.Sort([](TPair<uint32, uint32> const& LHS, TPair<uint32, uint32> const& RHS) { return LHS.Value > RHS.Value; });
+
+			for (TPair<uint32, uint32> ProducerCount : SortedProducerCounts)
+			{
+				const uint32 PackedProducerHandle = ProducerCount.Key;
+				const uint32 Count = ProducerCount.Value;
+
+				FVirtualTextureProducer* Producer = Producers.FindProducer(FVirtualTextureProducerHandle(PackedProducerHandle));
+				if (Producer != nullptr)
+				{
+					UE_LOG(LogConsoleResponse, Display, TEXT("   %s %d"), *Producer->GetName().ToString(), Count);
+				}
+			}
+		}
+	}
+}
+
 #if !UE_BUILD_SHIPPING
 void FVirtualTextureSystem::GetOnScreenMessages(FCoreDelegates::FSeverityMessageMap& OutMessages)
 {
diff --git a/Engine/Source/Runtime/Renderer/Private/VT/VirtualTextureSystem.h b/Engine/Source/Runtime/Renderer/Private/VT/VirtualTextureSystem.h
--- a/Engine/Source/Runtime/Renderer/Private/VT/VirtualTextureSystem.h
+++ b/Engine/Source/Runtime/Renderer/Private/VT/VirtualTextureSystem.h
@@ -140,6 +140,9 @@ private:
 	void ListPhysicalPoolsFromConsole();
 	FAutoConsoleCommand ListPhysicalPools;
 
+	void DumpPoolUsageFromConsole();
+	FAutoConsoleCommand DumpPoolUsageCommand;
+
 #if WITH_EDITOR
 	void SaveAllocatorImagesFromConsole();
 	FAutoConsoleCommand SaveAllocatorImages;
