Applies FFT bloom to scene color in tonemaper pass

diff --git a/Engine/Shaders/Private/Bloom/BloomClampKernel.usf b/Engine/Shaders/Private/Bloom/BloomClampKernel.usf
new file mode 100644
--- /dev/null
+++ b/Engine/Shaders/Private/Bloom/BloomClampKernel.usf
@@ -0,0 +1,32 @@
+// Copyright Epic Games, Inc. All Rights Reserved.
+
+#include "BloomCommon.ush"
+
+
+//------------------------------------------------------- CONFIG
+
+#define TILE_SIZE 8
+
+
+//------------------------------------------------------- PARAMETERS
+
+Texture2D KernelSpatialTexture;
+StructuredBuffer<FBloomKernelInfo> KernelConstantsBuffer;
+RWTexture2D<float4> ClampedKernelSpatialOutput;
+
+
+//------------------------------------------------------- ENTRY POINT
+
+[numthreads(TILE_SIZE, TILE_SIZE, 1)]
+void MainCS(
+	uint2 GroupId : SV_GroupID,
+	uint GroupThreadIndex : SV_GroupIndex)
+{
+	FBloomKernelInfo OriginalKernelInfo = KernelConstantsBuffer[0];
+
+	uint2 PixelCoord = TILE_SIZE * GroupId + uint2(GroupThreadIndex % TILE_SIZE, GroupThreadIndex / TILE_SIZE);
+
+	float4 KernelTexel = KernelSpatialTexture[PixelCoord];
+
+	ClampedKernelSpatialOutput[PixelCoord] = clamp(KernelTexel, 0.0, OriginalKernelInfo.MaxScatterDispersionEnergy);
+}
diff --git a/Engine/Shaders/Private/Bloom/BloomCommon.ush b/Engine/Shaders/Private/Bloom/BloomCommon.ush
new file mode 100644
--- /dev/null
+++ b/Engine/Shaders/Private/Bloom/BloomCommon.ush
@@ -0,0 +1,36 @@
+// Copyright Epic Games, Inc. All Rights Reserved.
+
+#include "../Common.ush"
+
+
+//------------------------------------------------------- COMPILER CONFIG
+
+// Generate vector truncation warnings to errors.
+#pragma warning(error: 3206)
+
+
+//------------------------------------------------------- STRUCTURES
+
+struct FBloomKernelInfo
+{
+	// Total energy at the center of the kernel.
+	float4 CenterEnergy;
+	
+	// Total energy of the bloom.
+	float4 ScatterDispersionEnergy;
+
+	// Highest energy per pixel that should clamp with to remove the kernel center
+	float4 MaxScatterDispersionEnergy;
+
+	// Pixel coordinate of the center of the kernel
+	uint2 CenterPixelCoord;
+
+	// Some padding to make vulkan happy
+	uint2 _Padding;
+};
+
+struct FBloomApplyInfo
+{
+	float4 SceneColorMultiply;
+	float4 BloomMultiply;
+};
diff --git a/Engine/Shaders/Private/Bloom/BloomFinalizeApplyConstants.usf b/Engine/Shaders/Private/Bloom/BloomFinalizeApplyConstants.usf
new file mode 100644
--- /dev/null
+++ b/Engine/Shaders/Private/Bloom/BloomFinalizeApplyConstants.usf
@@ -0,0 +1,41 @@
+// Copyright Epic Games, Inc. All Rights Reserved.
+
+#include "BloomCommon.ush"
+
+
+//------------------------------------------------------- PARAMETERS
+
+float ScatterDispersionIntensity;
+
+StructuredBuffer<FBloomKernelInfo> KernelConstantsBuffer;
+RWStructuredBuffer<FBloomApplyInfo> BloomApplyConstantsOutput;
+
+
+//------------------------------------------------------- ENTRY POINT
+
+[numthreads(1, 1, 1)]
+void MainCS(
+	uint2 GroupId : SV_GroupID,
+	uint GroupThreadIndex : SV_GroupIndex)
+{
+	FBloomKernelInfo OriginalKernelInfo = KernelConstantsBuffer[0];
+
+	// Compute the tint of the kernel
+	float4 Tint;
+	{
+		float3 RefTotalEnergy = OriginalKernelInfo.CenterEnergy.rgb + OriginalKernelInfo.ScatterDispersionEnergy.rgb;
+
+		float MaxEnergy = max3(RefTotalEnergy.r, RefTotalEnergy.g, RefTotalEnergy.b);
+
+		Tint = float4(saturate(RefTotalEnergy / MaxEnergy), 1.0);
+	}
+
+	float4 OriginalTotalEnergy = OriginalKernelInfo.CenterEnergy + OriginalKernelInfo.ScatterDispersionEnergy;
+	
+	float4 FinalScatterDispersionEnergy = OriginalKernelInfo.ScatterDispersionEnergy * ScatterDispersionIntensity;
+	
+	float4 FinalCenterEnergy = OriginalTotalEnergy - FinalScatterDispersionEnergy;
+
+	BloomApplyConstantsOutput[0].SceneColorMultiply = Tint * saturate(FinalCenterEnergy / OriginalTotalEnergy);
+	BloomApplyConstantsOutput[0].BloomMultiply = Tint * saturate(FinalScatterDispersionEnergy / OriginalTotalEnergy);
+}
diff --git a/Engine/Shaders/Private/Bloom/BloomFindKernelCenter.usf b/Engine/Shaders/Private/Bloom/BloomFindKernelCenter.usf
new file mode 100644
--- /dev/null
+++ b/Engine/Shaders/Private/Bloom/BloomFindKernelCenter.usf
@@ -0,0 +1,111 @@
+// Copyright Epic Games, Inc. All Rights Reserved.
+
+#include "BloomCommon.ush"
+
+
+//------------------------------------------------------- CONSTANTS
+
+#define TILE_SIZE 8
+
+
+//------------------------------------------------------- PARAMETERS
+
+/*
+uint
+uint
+float
+uint
+* */
+uint2 KernelSpatialTextureSize;
+Texture2D KernelSpatialTexture;
+
+RWCoherentStructuredBuffer(uint) KernelCenterCoordOutput;
+
+
+//------------------------------------------------------- LDS
+
+groupshared uint MaxLuminance;
+groupshared uint2 MaxLuminancePixelPos;
+
+
+//------------------------------------------------------- ENTRY POINT
+
+[numthreads(TILE_SIZE, TILE_SIZE, 1)]
+void MainCS(
+	uint2 GroupId : SV_GroupID,
+	uint GroupThreadIndex : SV_GroupIndex)
+{
+	// Init LDS
+	if (GroupThreadIndex == 0)
+	{
+		MaxLuminance = 0;
+	}
+
+	GroupMemoryBarrierWithGroupSync();
+
+	uint2 InputPixelPos = TILE_SIZE * GroupId + uint2(GroupThreadIndex % TILE_SIZE, GroupThreadIndex / TILE_SIZE);
+
+	float4 TexelColor = KernelSpatialTexture[InputPixelPos];
+
+	bool bValidPixel = all(InputPixelPos < KernelSpatialTextureSize);
+
+	if (!bValidPixel)
+	{
+		TexelColor = 0.0;
+	}
+
+	// Locate the position of the brightest pixel
+	{
+		float TexelLuminance = Luminance(TexelColor.rgb);
+
+		InterlockedMax(MaxLuminance, asuint(TexelLuminance));
+
+		GroupMemoryBarrierWithGroupSync();
+
+		if (MaxLuminance == asuint(TexelLuminance))
+		{
+			MaxLuminancePixelPos = InputPixelPos;
+		}
+	}
+	
+	GroupMemoryBarrierWithGroupSync();
+
+	BRANCH
+	if (GroupThreadIndex == 0 && MaxLuminance > 0)
+	{
+		bool bDone = false;
+		while (!bDone)
+		{
+			if (MaxLuminance > KernelCenterCoordOutput[2])
+			{
+				uint PrevLock;
+				InterlockedCompareExchange(
+					KernelCenterCoordOutput[3],
+					uint(0),
+					uint(1),
+					/* out */ PrevLock);
+
+				if (PrevLock == 0)
+				{
+					AllMemoryBarrier();
+
+					if (MaxLuminance > KernelCenterCoordOutput[2])
+					{
+						KernelCenterCoordOutput[0] = MaxLuminancePixelPos.x;
+						KernelCenterCoordOutput[1] = MaxLuminancePixelPos.y;
+						KernelCenterCoordOutput[2] = MaxLuminance;
+					}
+
+					AllMemoryBarrier();
+
+					InterlockedExchange(KernelCenterCoordOutput[3], 0, PrevLock);
+					bDone = true;
+				}
+			}
+			else
+			{
+				bDone = true;
+			}
+		}
+	}
+}
diff --git a/Engine/Shaders/Private/Bloom/BloomPackKernelConstants.usf b/Engine/Shaders/Private/Bloom/BloomPackKernelConstants.usf
new file mode 100644
--- /dev/null
+++ b/Engine/Shaders/Private/Bloom/BloomPackKernelConstants.usf
@@ -0,0 +1,27 @@
+// Copyright Epic Games, Inc. All Rights Reserved.
+
+#include "BloomCommon.ush"
+
+
+//------------------------------------------------------- PARAMETERS
+
+StructuredBuffer<uint> KernelCenterCoordBuffer;
+StructuredBuffer<float4> KernelCenterEnergyBuffer;
+StructuredBuffer<float4> MaxScatterDispersionBuffer;
+Texture2D ScatterDispersionTexture;
+
+RWStructuredBuffer<FBloomKernelInfo> KernelConstantsOutput;
+
+
+//------------------------------------------------------- ENTRY POINT
+
+[numthreads(1, 1, 1)]
+void MainCS(
+	uint2 GroupId : SV_GroupID,
+	uint GroupThreadIndex : SV_GroupIndex)
+{
+	KernelConstantsOutput[0].CenterEnergy = KernelCenterEnergyBuffer[0];
+	KernelConstantsOutput[0].ScatterDispersionEnergy = ScatterDispersionTexture[uint2(0, 0)];
+	KernelConstantsOutput[0].MaxScatterDispersionEnergy = MaxScatterDispersionBuffer[0];
+	KernelConstantsOutput[0].CenterPixelCoord = uint2(KernelCenterCoordBuffer[0], KernelCenterCoordBuffer[1]);
+}
diff --git a/Engine/Shaders/Private/Bloom/BloomReduceKernelSurvey.usf b/Engine/Shaders/Private/Bloom/BloomReduceKernelSurvey.usf
new file mode 100644
--- /dev/null
+++ b/Engine/Shaders/Private/Bloom/BloomReduceKernelSurvey.usf
@@ -0,0 +1,89 @@
+// Copyright Epic Games, Inc. All Rights Reserved.
+
+#include "BloomCommon.ush"
+
+
+//------------------------------------------------------- CONFIG
+
+#define TILE_SIZE 64
+
+
+//------------------------------------------------------- PARAMETERS
+
+uint SurveyReduceOp;
+uint SurveyGroupCount;
+RWStructuredBuffer<float4> SurveyOutput;
+
+
+//------------------------------------------------------- LDS
+
+groupshared float4 SharedColor[TILE_SIZE];
+
+
+//------------------------------------------------------- ENTRY POINT
+
+[numthreads(TILE_SIZE, 1, 1)]
+void MainCS(uint GroupThreadIndex : SV_GroupIndex)
+{
+	float4 Color = 0.0;
+	{
+		uint LoadCount = (SurveyGroupCount + TILE_SIZE - 1) / TILE_SIZE;
+
+		LOOP
+		for (uint i = 0; i < LoadCount; i++)
+		{
+			float4 LoadedColor = SurveyOutput[GroupThreadIndex + i * TILE_SIZE];
+
+			if (SurveyReduceOp == 0)
+			{
+				Color = max(Color, LoadedColor);
+			}
+			else
+			{
+				Color += LoadedColor;
+			}
+		}
+	}
+
+	SharedColor[GroupThreadIndex] = Color;
+
+	GroupMemoryBarrierWithGroupSync();
+	
+	// Reduction
+	{
+		UNROLL
+		for (uint i = 32; i > 1; i /= 2)
+		{
+			
+			if (SurveyReduceOp == 0)
+			{
+				Color = max(Color, SharedColor[GroupThreadIndex ^ i]);
+			}
+			else
+			{
+				Color += SharedColor[GroupThreadIndex ^ i];
+			}
+			
+			GroupMemoryBarrierWithGroupSync();
+
+			SharedColor[GroupThreadIndex] = Color;
+			
+			GroupMemoryBarrierWithGroupSync();
+		}
+		
+		if (SurveyReduceOp == 0)
+		{
+			Color = max(Color, SharedColor[GroupThreadIndex ^ 0x1]);
+		}
+		else
+		{
+			Color += SharedColor[GroupThreadIndex ^ 0x1];
+		}
+	}
+
+	BRANCH
+	if (GroupThreadIndex == 0)
+	{
+		SurveyOutput[0] = Color;
+	}
+}
diff --git a/Engine/Shaders/Private/Bloom/BloomResizeKernel.usf b/Engine/Shaders/Private/Bloom/BloomResizeKernel.usf
new file mode 100644
--- /dev/null
+++ b/Engine/Shaders/Private/Bloom/BloomResizeKernel.usf
@@ -0,0 +1,160 @@
+// Copyright Epic Games, Inc. All Rights Reserved.
+
+#include "BloomCommon.ush"
+
+
+//------------------------------------------------------- CONFIG
+
+#define TILE_SIZE 8
+
+
+//------------------------------------------------------- PARAMETERS
+
+uint2 DstExtent;
+uint2 ImageExtent;
+float2 KernelSpatialTextureInvSize;
+uint2 DstBufferExtent;
+float KernelSupportScale;
+
+StructuredBuffer<FBloomKernelInfo> KernelConstantsBuffer;
+
+Texture2D SrcTexture;
+SamplerState SrcSampler;
+
+RWTexture2D<float4> DstTexture;
+
+
+//------------------------------------------------------- ENTRY POINT
+
+[numthreads(TILE_SIZE, TILE_SIZE, 1)]
+void MainCS(
+	uint2 GroupId : SV_GroupID,
+	uint GroupThreadIndex : SV_GroupIndex)
+{
+	FBloomKernelInfo OriginalKernelInfo = KernelConstantsBuffer[0];
+
+	// f(x) = T(x)  0 < x < k
+    //      = 0    x > k
+	// g(x) = f(L * frac(x/L))   - periodic extention of f with period L
+	// g_c(x) = f(L * frac( (x+c)/ L)  - periodic extension of offset f with period L
+	//
+	// 
+	// The pixel that this thread 'owns'
+
+	uint2 PixelId = TILE_SIZE * GroupId + uint2(GroupThreadIndex % TILE_SIZE, GroupThreadIndex / TILE_SIZE);
+	
+	// If the actual scan line length was not divisible by THREADS_PER_GROUP we will have a few extra threads that don't
+	// own any pixels.  This will be 'false' for most thread groups.
+
+	if (! (PixelId.x  < DstBufferExtent.x )) return;
+	if (! (PixelId.y  < DstBufferExtent.y )) return;
+
+	float2 DstSize = float2(DstExtent.x, DstExtent.y);
+
+	// Where to find the center of the kernel, and the relative size of the kernel texture wrt the dst buffer.
+	//float2 UVCenter    = KernelCenter;
+	float2 UVCenter    = (float2(OriginalKernelInfo.CenterPixelCoord) + 0.5) * KernelSpatialTextureInvSize;
+	float  ResizeScale = KernelSupportScale;
+
+	// The minor axis of the target buffer
+	// NB: I should change this to the minor axis of the viewing window..  The target buffer is bigger and some power of two. 
+		
+	uint MajorAxisPixelCount = max(ImageExtent.y, ImageExtent.x);
+
+	// The number of kernels that could fit along the minor axis = ResizeScale.
+
+	// The length of the kernel is pixels in the Dst Buffer.  Here the kernel is assumed
+	// to come from a square texture.
+	
+	float KernelSizeInDstPixels = max( float(MajorAxisPixelCount) * (ResizeScale), 1.f);
+	
+	// We have a simple frac trick for making the periodic image, but it fails
+	// in the following case.  Then we just have to do it by hand.
+	bool bKernelTooBigForFrac = !(float(DstBufferExtent.y) > UVCenter.y * KernelSizeInDstPixels);
+	bKernelTooBigForFrac = bKernelTooBigForFrac || !(float(DstBufferExtent.x) > UVCenter.x * KernelSizeInDstPixels);
+
+
+	float4 ResultColor = float4(0.f, 0.f, 0.f, 0.f);
+
+	if ( bKernelTooBigForFrac )
+	{
+
+		// distance from each corner
+
+		float2 UL = float2(PixelId.x, PixelId.y) / KernelSizeInDstPixels;
+		float2 UR = float2(DstBufferExtent.x - PixelId.x, PixelId.y) / KernelSizeInDstPixels;
+		float2 LL = float2(PixelId.x, DstBufferExtent.y - PixelId.y) / KernelSizeInDstPixels;
+		float2 LR = float2(DstBufferExtent.x - PixelId.x, DstBufferExtent.y - PixelId.y) / KernelSizeInDstPixels;
+
+		float4 TmpColor = float4(0.f, 0.f ,0.f ,0.f);
+		if (UL.x < 0.5 && UL.y < 0.5)
+		{
+			float4 SampledColor = SrcTexture.SampleLevel(SrcSampler, UL + UVCenter, 0);
+			TmpColor.rgb += SampledColor.rgb;
+			TmpColor.a = max(TmpColor.a, SampledColor.a);
+		}
+		if (UR.x < 0.5 && UR.y < 0.5)
+		{
+			float2 SamplePoint = float2(UVCenter.x - UR.x, UVCenter.y + UR.y);
+			float4 SampledColor = SrcTexture.SampleLevel(SrcSampler,  SamplePoint, 0);
+			TmpColor.rgb += SampledColor.rgb;
+			TmpColor.a = max(TmpColor.a, SampledColor.a);
+		}
+		if (LL.x < 0.5 && LL.y < 0.5)
+		{
+			float2 SamplePoint = float2(UVCenter.x + LL.x, UVCenter.y - LL.y);
+			float4 SampledColor = SrcTexture.SampleLevel(SrcSampler,  SamplePoint, 0);
+			TmpColor.rgb += SampledColor.rgb;
+			TmpColor.a = max(TmpColor.a, SampledColor.a);
+		}
+		if (LR.x < 0.5 && LR.y < 0.5)
+		{
+			float2 SamplePoint = float2(UVCenter.x - LR.x, UVCenter.y - LR.y);
+			float4 SampledColor = SrcTexture.SampleLevel(SrcSampler,  SamplePoint, 0);
+			TmpColor.rgb += SampledColor.rgb;
+			TmpColor.a = max(TmpColor.a, SampledColor.a);
+		}
+		ResultColor = TmpColor;
+	}
+	else 
+	{
+		// Current Pixel in KernelUV space
+
+		float2 PixelInKernelUV = float2(float(PixelId.x) / KernelSizeInDstPixels, float(PixelId.y) / KernelSizeInDstPixels);
+
+		// The size of the Dst Buffer in terms of KernelSizeInDstPixels.  Note the Dst buffer is bigger than the image extent.
+	
+		float2 DstSizeInKernelUnits = float2( DstSize.x / KernelSizeInDstPixels, DstSize.y / KernelSizeInDstPixels);
+	
+		float2 InvDstSizeInKernelUnits = float2(1.f / DstSizeInKernelUnits.x, 1.f / DstSizeInKernelUnits.y );
+
+	
+	
+		// Offset by 1/2, 1/2 so we start in the middle of the kernel texture
+	 
+		float2 OffsetPixelInKernelUV = PixelInKernelUV + UVCenter;
+
+		// make periodic with period DstSizeInKernelUV
+
+		OffsetPixelInKernelUV *= InvDstSizeInKernelUnits;
+		OffsetPixelInKernelUV = frac(OffsetPixelInKernelUV); 
+		OffsetPixelInKernelUV *= DstSizeInKernelUnits; 
+
+		
+
+		if (OffsetPixelInKernelUV.x < 1.f && OffsetPixelInKernelUV.y < 1.f)
+		{
+			float4 SampledColor = SrcTexture.SampleLevel(SrcSampler, OffsetPixelInKernelUV, 0);
+			ResultColor.rgba += SampledColor; 
+		}
+	}
+
+	// Force the alpha kernel to be a delta function at the pixel center.
+#if 0	
+	ResultColor.w = (PixelId.x == 0 && PixelId.y == 0) ? 1.f : 0.f;
+	
+#endif
+
+	DstTexture[PixelId] =  ResultColor;
+
+}
diff --git a/Engine/Shaders/Private/Bloom/BloomSumScatterDispersionEnergy.usf b/Engine/Shaders/Private/Bloom/BloomSumScatterDispersionEnergy.usf
new file mode 100644
--- /dev/null
+++ b/Engine/Shaders/Private/Bloom/BloomSumScatterDispersionEnergy.usf
@@ -0,0 +1,80 @@
+// Copyright Epic Games, Inc. All Rights Reserved.
+
+#include "BloomCommon.ush"
+
+
+//------------------------------------------------------- CONSTANTS
+
+#define TILE_SIZE 8
+
+
+//------------------------------------------------------- PARAMETERS
+
+uint PassId;
+uint2 ScatterDispersionTextureSize;
+Texture2D ScatterDispersionTexture;
+StructuredBuffer<float4> MaxScatterDispersionBuffer;
+
+RWTexture2D<float4> ScatterDispersionOutput;
+
+
+//------------------------------------------------------- LDS
+
+groupshared float4 SharedColor[TILE_SIZE * TILE_SIZE];
+
+
+//------------------------------------------------------- ENTRY POINT
+
+[numthreads(TILE_SIZE, TILE_SIZE, 1)]
+void MainCS(
+	uint2 GroupId : SV_GroupID,
+	uint GroupThreadIndex : SV_GroupIndex)
+{
+	float4 MaxScatterDispersion = MaxScatterDispersionBuffer[0];
+	
+	uint2 InputPixelPos = TILE_SIZE * GroupId + uint2(GroupThreadIndex % TILE_SIZE, GroupThreadIndex / TILE_SIZE);
+
+	// Init LDS
+	float4 TexelColor = ScatterDispersionTexture[InputPixelPos];
+	{
+		bool bValidPixel = all(InputPixelPos < ScatterDispersionTextureSize);
+
+		if (PassId == 0)
+		{
+			// Clamp the kernel to highest value found near the kernel center.
+			TexelColor = clamp(TexelColor, 0.0, MaxScatterDispersion);
+		}
+
+		if (!bValidPixel)
+		{
+			TexelColor = 0.0;
+		}
+
+		SharedColor[GroupThreadIndex] = TexelColor;
+	}
+
+	GroupMemoryBarrierWithGroupSync();
+
+	// Reduction
+	{
+		UNROLL
+		for (uint i = 32; i > 1; i /= 2)
+		{
+			TexelColor += SharedColor[GroupThreadIndex ^ i];
+
+			GroupMemoryBarrierWithGroupSync();
+
+			SharedColor[GroupThreadIndex] = TexelColor;
+			
+			GroupMemoryBarrierWithGroupSync();
+		}
+		
+		TexelColor += SharedColor[GroupThreadIndex ^ 0x1];
+	}
+
+	// Output
+	if (GroupThreadIndex == 0)
+	{
+		ScatterDispersionOutput[GroupId] = TexelColor;
+	}
+}
diff --git a/Engine/Shaders/Private/Bloom/BloomSurveyKernelCenterEnergy.usf b/Engine/Shaders/Private/Bloom/BloomSurveyKernelCenterEnergy.usf
new file mode 100644
--- /dev/null
+++ b/Engine/Shaders/Private/Bloom/BloomSurveyKernelCenterEnergy.usf
@@ -0,0 +1,88 @@
+// Copyright Epic Games, Inc. All Rights Reserved.
+
+#include "BloomCommon.ush"
+
+
+//------------------------------------------------------- CONFIG
+
+#define TILE_SIZE 8
+
+
+//------------------------------------------------------- PARAMETERS
+
+float ViewTexelRadiusInKernelTexels;
+uint SurveyGroupGridSize;
+
+
+uint2 KernelSpatialTextureSize;
+Texture2D KernelSpatialTexture;
+StructuredBuffer<uint> KernelCenterCoordBuffer;
+StructuredBuffer<float4> MaxScatterDispersionBuffer;
+
+RWStructuredBuffer<float4> SurveyOutput;
+RWTexture2D<float4> DebugOutput;
+
+
+//------------------------------------------------------- LDS
+
+groupshared float4 SharedColor[TILE_SIZE * TILE_SIZE];
+
+
+//------------------------------------------------------- ENTRY POINT
+
+[numthreads(TILE_SIZE, TILE_SIZE, 1)]
+void MainCS(
+	uint2 GroupId : SV_GroupID,
+	uint GroupThreadIndex : SV_GroupIndex)
+{
+	float4 Debug = 0.0;
+
+	uint LinearGroupId = dot(GroupId, uint2(1, SurveyGroupGridSize));
+	int2 PixelOffset = int2(GroupThreadIndex % TILE_SIZE, GroupThreadIndex / TILE_SIZE) + TILE_SIZE * (GroupId - SurveyGroupGridSize / 2);
+
+	// Write LDS
+	float4 TexelColor;
+	{
+		uint2 KernelCenterPos = uint2(KernelCenterCoordBuffer[0], KernelCenterCoordBuffer[1]);
+		
+		int2 PixelOffset = int2(GroupThreadIndex % TILE_SIZE, GroupThreadIndex / TILE_SIZE) + TILE_SIZE * (GroupId - SurveyGroupGridSize / 2);
+
+		uint2 InputPixelPos = (KernelCenterPos + PixelOffset) % KernelSpatialTextureSize;
+
+		float IsSurveyZone = saturate(ViewTexelRadiusInKernelTexels + 0.5 - max(abs(PixelOffset.x), abs(PixelOffset.y)));
+		//Debug.r = IsSurveyZone;
+
+		TexelColor = max(KernelSpatialTexture[InputPixelPos] - MaxScatterDispersionBuffer[0], 0.0) * IsSurveyZone;
+		//Debug = TexelColor;
+
+		SharedColor[GroupThreadIndex] = TexelColor;
+	}
+
+	GroupMemoryBarrierWithGroupSync();
+
+	// Reduction
+	{
+		UNROLL
+		for (uint i = 32; i > 1; i /= 2)
+		{
+			TexelColor += SharedColor[GroupThreadIndex ^ i];
+
+			GroupMemoryBarrierWithGroupSync();
+
+			SharedColor[GroupThreadIndex] = TexelColor;
+			
+			GroupMemoryBarrierWithGroupSync();
+		}
+		
+		TexelColor += SharedColor[GroupThreadIndex ^ 0x1];
+	}
+
+	BRANCH
+	if (GroupThreadIndex == 0)
+	{
+		SurveyOutput[LinearGroupId] = TexelColor;
+		Debug = TexelColor;
+	}
+	
+	//DebugOutput[PixelOffset + TILE_SIZE * (SurveyGroupGridSize / 2)] = Debug;
+}
diff --git a/Engine/Shaders/Private/Bloom/BloomSurveyMaxScatterDispersion.usf b/Engine/Shaders/Private/Bloom/BloomSurveyMaxScatterDispersion.usf
new file mode 100644
--- /dev/null
+++ b/Engine/Shaders/Private/Bloom/BloomSurveyMaxScatterDispersion.usf
@@ -0,0 +1,92 @@
+// Copyright Epic Games, Inc. All Rights Reserved.
+
+#include "BloomCommon.ush"
+
+
+//------------------------------------------------------- CONFIG
+
+#define TILE_SIZE 8
+
+
+//------------------------------------------------------- PARAMETERS
+
+float ViewTexelRadiusInKernelTexels;
+uint SurveyGroupGridSize;
+uint SurveyType;
+
+
+uint2 KernelSpatialTextureSize;
+Texture2D KernelSpatialTexture;
+StructuredBuffer<uint> KernelCenterCoordBuffer;
+
+RWStructuredBuffer<float4> SurveyOutput;
+RWTexture2D<float4> DebugOutput;
+
+
+//------------------------------------------------------- LDS
+
+groupshared uint4 SharedMaxNeighborColor;
+
+
+//------------------------------------------------------- ENTRY POINT
+
+[numthreads(TILE_SIZE, TILE_SIZE, 1)]
+void MainCS(
+	uint2 GroupId : SV_GroupID,
+	uint GroupThreadIndex : SV_GroupIndex)
+{
+	float4 Debug = 0.0;
+
+	uint LinearGroupId = dot(GroupId, uint2(1, SurveyGroupGridSize));
+	int2 PixelOffset = int2(GroupThreadIndex % TILE_SIZE, GroupThreadIndex / TILE_SIZE) + TILE_SIZE * (GroupId - SurveyGroupGridSize / 2);
+
+	// Init LDS
+	if (GroupThreadIndex == 0)
+	{
+		SharedMaxNeighborColor = 0;
+	}
+	
+	GroupMemoryBarrierWithGroupSync();
+
+	{
+		uint2 KernelCenterPos = uint2(KernelCenterCoordBuffer[0], KernelCenterCoordBuffer[1]);
+
+		uint2 InputPixelPos = (KernelCenterPos + PixelOffset) % KernelSpatialTextureSize;
+		
+		//float IsCenterZone = saturate(ViewTexelRadiusInKernelTexels + 0.5 - max(abs(PixelOffset.x), abs(PixelOffset.y)));
+		//float IsEdgeOrCenterZone = saturate(ViewTexelRadiusInKernelTexels + 3.5 - max(abs(PixelOffset.x), abs(PixelOffset.y)));
+				
+		float IsCenterZone = saturate(2.0 * ViewTexelRadiusInKernelTexels + 1.0 - max(abs(PixelOffset.x), abs(PixelOffset.y)));
+		float IsEdgeOrCenterZone = saturate(2.0 * ViewTexelRadiusInKernelTexels + 4.0 - max(abs(PixelOffset.x), abs(PixelOffset.y)));
+
+		float IsSurveyZone = saturate(1.0 - IsCenterZone) * IsEdgeOrCenterZone;
+		//Debug.r = IsSurveyZone;
+
+		float4 TexelColor = KernelSpatialTexture[InputPixelPos] * IsSurveyZone;
+		//Debug = TexelColor;
+
+		TexelColor.rgb = max3(TexelColor.r, TexelColor.g, TexelColor.b).xxx;
+
+		InterlockedMax(SharedMaxNeighborColor.r, asuint(TexelColor.r));
+		InterlockedMax(SharedMaxNeighborColor.g, asuint(TexelColor.g));
+		InterlockedMax(SharedMaxNeighborColor.b, asuint(TexelColor.b));
+		InterlockedMax(SharedMaxNeighborColor.a, asuint(TexelColor.a));
+	}
+
+	GroupMemoryBarrierWithGroupSync();
+	
+	if (GroupThreadIndex == 0)
+	{
+		float4 MaxNeighborColor = float4(
+			asfloat(SharedMaxNeighborColor.r),
+			asfloat(SharedMaxNeighborColor.g),
+			asfloat(SharedMaxNeighborColor.b),
+			asfloat(SharedMaxNeighborColor.a));
+
+		SurveyOutput[LinearGroupId] = MaxNeighborColor;
+		
+		//Debug = MaxNeighborColor;
+	}
+	
+	//DebugOutput[PixelOffset + TILE_SIZE * (SurveyGroupGridSize / 2)] = Debug;
+}
diff --git a/Engine/Shaders/Private/GPUFastFourierTransform.usf b/Engine/Shaders/Private/GPUFastFourierTransform.usf
--- a/Engine/Shaders/Private/GPUFastFourierTransform.usf
+++ b/Engine/Shaders/Private/GPUFastFourierTransform.usf
@@ -429,10 +429,7 @@ void GroupSharedTwoForOneFFTCS(uint3 GroupID : SV_GroupID, uint3 GroupThreadID :
 //              (e.g. Vertical ComplexFFT ( Horizontal TwoForOneFFT(Kernel))) 
 //
 //      
-//              And needed helper functions           
-//
-//              NB: This code affects a scale preserving 'normalization' of the kernel, by using the integral of the largest channel
-//                  of the physical space kernel.    
+//              And needed helper functions
 // ---------------------------------------------------------------------------------------------------------------------------------------
 
 // Input SRV: 
@@ -628,30 +625,27 @@ void GSConvolutionWithTextureCS(uint3 GroupID : SV_GroupID, uint3 GroupThreadID
 		GetKernelSum(FilterTexture, bIsHorizontal, NumScanlines, Norm);
 		// redSum = Norm[0].x;  greenSum = Norm[0].y; blueSum = Norm[1].x; alphaSum = Norm[1].y
 		
-		
-
-		// We normalize to preserve the relative sizes of the r,g,b kernels.
-		float NormMax = max(Norm[0].x, max(Norm[0].y, Norm[1].x));
-		float AlphaSum = (bUseAlpha) ? Norm[1].y : 1.f;
-
-
 		// Normalize R & G
 		{
+			//float Normal = NormMax;
+			float Normal = (bIsGAGroup) ? Norm[0].y :  Norm[0].x;
+
 			for (uint r = 0; r < RADIX; ++r)
 			{
 				// This is the R or G channel
-				LocalBuffer[0][r] /= NormMax; 
+				LocalBuffer[0][r] /= Normal; 
 			}
 		}
 		
-		float AorBNorm = (bIsGAGroup) ? AlphaSum : NormMax; 
-
 		// Normalize B & A
 		{
+			//float Normal = AorBNorm;
+			float Normal = (bIsGAGroup) ? Norm[1].y :  Norm[1].x;
+
 			for (uint r = 0; r < RADIX; ++r)
 			{
 				// This is the B or A channel
-				LocalBuffer[1][r] /= AorBNorm; 
+				LocalBuffer[1][r] /= Normal; 
 			}
 		}
 	}
@@ -1387,14 +1381,8 @@ void ComplexMultiplyImagesCS(uint3 GroupID : SV_GroupID, uint3 GroupThreadID : S
 
 		// redSum = Norm[0].x;  greenSum = Norm[0].y; blueSum = Norm[1].x; alphaSum = Norm[1].y
 
-		// We normalize to preserve the relative sizes of the r,g,b kernels.
-		float InvNormMax = 1.f / max(Norm[0].x, max(Norm[0].y, Norm[1].x));
-		
-		float InvAlphaSum = (bUseAlpha) ? 1.f / Norm[1].y : 1.f;
-		
-		float BorANorm = (bIsGAGroup) ? InvAlphaSum : InvNormMax;
-
-		InvNorm = float4(InvNormMax, InvNormMax, BorANorm, BorANorm);
+		// We normalize each RGBA channel independently.
+		InvNorm = float4(((bIsGAGroup) ? Norm[0].y : Norm[0].x).xx, ((bIsGAGroup) ? Norm[1].y : Norm[1].x).xx);
 	}
 
 
diff --git a/Engine/Shaders/Private/PostProcessFFTBloom.usf b/Engine/Shaders/Private/PostProcessFFTBloom.usf
deleted file mode 100644
diff --git a/Engine/Shaders/Private/PostProcessTonemap.usf b/Engine/Shaders/Private/PostProcessTonemap.usf
--- a/Engine/Shaders/Private/PostProcessTonemap.usf
+++ b/Engine/Shaders/Private/PostProcessTonemap.usf
@@ -6,12 +6,20 @@
 
 #define EYE_ADAPTATION_LOOSE_PARAMETERS 1
 
+#ifndef SUPPORTS_BLOOM_APPLY_PARAMETERS
+	#define SUPPORTS_BLOOM_APPLY_PARAMETERS 0
+#endif
+
+
 #include "Common.ush"
 #include "PostProcessCommon.ush"
 #include "TonemapCommon.ush"
 #include "EyeAdaptationCommon.ush"
 #include "PixelQuadMessagePassing.ush"
 #include "ScreenPass.ush"
+#if SUPPORTS_BLOOM_APPLY_PARAMETERS
+	#include "Bloom/BloomCommon.ush"
+#endif
 
 /** Values of DIM_OUTPUT_DEVICE (matching C++'s ETonemapperOutputDevice) */
 #define TONEMAPPER_OUTPUT_sRGB 0
@@ -37,16 +45,24 @@
 #define DEBUG_TONEMAPPER 0
 
 SCREEN_PASS_TEXTURE_VIEWPORT(Color)
-SCREEN_PASS_TEXTURE_VIEWPORT(Bloom)
 SCREEN_PASS_TEXTURE_VIEWPORT(Output)
-FScreenTransform ColorToBloom;
 
 Texture2D ColorTexture;
 SamplerState ColorSampler;
 
+
+FScreenTransform ColorToBloom;
+float2 BloomUVViewportBilinearMin;
+float2 BloomUVViewportBilinearMax;
+float4 BloomMultiply;
 Texture2D BloomTexture;
 SamplerState BloomSampler;
 
+#if SUPPORTS_BLOOM_APPLY_PARAMETERS
+	StructuredBuffer<FBloomApplyInfo> BloomApplyParameters;
+#endif
+
+
 // xyz:SceneColorTint.rgb, w:unused
 float4 ColorScale0;
 
@@ -536,7 +552,7 @@ float4 TonemapCommonPS(
 		OutColor.rgb = pow(SceneColor.rgb, InverseGamma.x);
 
 	#else 
-
+	{
 		#if METAL_MSAA_HDR_DECODE
 			// look for PreTonemapMSAA 
 			SceneColor.rgb *= rcp(SceneColor.r*(-0.299) + SceneColor.g*(-0.587) + SceneColor.b*(-0.114) + 1.0);
@@ -547,13 +563,21 @@ float4 TonemapCommonPS(
 		half3 LinearColor = SceneColor.rgb * ColorScale0.rgb;
 
 		#if USE_BLOOM
-			float2 BloomUV = ApplyScreenTransform(UV, ColorToBloom);
-			BloomUV = clamp(BloomUV, Bloom_UVViewportBilinearMin, Bloom_UVViewportBilinearMax);
+		{
+			float2 BloomUV;
 			#if ES3_1_PROFILE
+			{
 				BloomUV = FullViewUV.xy;
+			}
+			#else
+			{
+				BloomUV = ApplyScreenTransform(UV, ColorToBloom);
+				BloomUV = clamp(BloomUV, BloomUVViewportBilinearMin, BloomUVViewportBilinearMax);
+			}
 			#endif
+
 			float4 CombinedBloom = Texture2DSample(BloomTexture, BloomSampler, BloomUV);
-			CombinedBloom.rgb *= View.OneOverPreExposure; // Standard bloom path, convolution path has bloom already in scene color.
+			CombinedBloom.rgb *= View.OneOverPreExposure;
 
 			//float2 DirtViewportUV = (SvPosition.xy - float2(Output_ViewportMin)) * Output_ViewportSizeInverse; 
 
@@ -573,16 +597,36 @@ float4 TonemapCommonPS(
 #endif //DEBUG_TONEMAPPER
 
 			#if FEATURE_LEVEL == FEATURE_LEVEL_ES3_1
+			{
 				// Support sunshaft and vignette for mobile, and we have applied the BloomIntensity and the BloomDirtMask at the sun merge pass.
 				LinearColor = LinearColor.rgb * CombinedBloom.a + CombinedBloom.rgb;
+			}
 			#else
-			
+			{
 				float2 DirtLensUV = ConvertScreenViewportSpaceToLensViewportSpace(ScreenPos) * float2(1.0f, -1.0f);
 
 				float3 BloomDirtMaskColor = Texture2DSample(BloomDirtMaskTexture, BloomDirtMaskSampler, DirtLensUV * .5f + .5f).rgb * BloomDirtMaskTint.rgb;
-
-				LinearColor += CombinedBloom.rgb * (ColorScale1.rgb + BloomDirtMaskColor); 
+				
+				#if SUPPORTS_BLOOM_APPLY_PARAMETERS
+				{
+					FBloomApplyInfo BloomApplyInfo = BloomApplyParameters[0];
+
+					LinearColor = (
+						LinearColor.rgb * BloomApplyInfo.SceneColorMultiply.rgb +
+						CombinedBloom.rgb * (BloomApplyInfo.BloomMultiply.rgb + BloomDirtMaskColor));
+					
+					SceneColor.a = (
+						SceneColor.a * BloomApplyInfo.SceneColorMultiply.a +
+						CombinedBloom.a * BloomApplyInfo.BloomMultiply.a);
+				}
+				#else
+				{
+					LinearColor += CombinedBloom.rgb * (BloomMultiply.rgb + BloomDirtMaskColor);
+				}
+				#endif
+			}
 			#endif
+		}
 		#endif
 
 		#if NO_EYEADAPTATION_EXPOSURE_FIX == 1
@@ -649,7 +693,7 @@ float4 TonemapCommonPS(
 		#elif DIM_OUTPUT_DEVICE == TONEMAPPER_OUTPUT_WithToneCurve
 			OutColor.xyz = OutDeviceColor;
 		#endif
-
+	}
 	#endif
 	
 	#if (FEATURE_LEVEL <= FEATURE_LEVEL_ES3_1 && (USE_FXAA == 0 || MOBILE_PROPAGATE_ALPHA > 0)) || POST_PROCESS_ALPHA == 2 || (POST_PROCESS_ALPHA == 1 && DIM_OUTPUT_DEVICE == TONEMAPPER_OUTPUT_NoToneCurve)  || (POST_PROCESS_ALPHA == 1 && DIM_OUTPUT_DEVICE == TONEMAPPER_OUTPUT_WithToneCurve)
diff --git a/Engine/Source/Runtime/Engine/Classes/Engine/Scene.h b/Engine/Source/Runtime/Engine/Classes/Engine/Scene.h
--- a/Engine/Source/Runtime/Engine/Classes/Engine/Scene.h
+++ b/Engine/Source/Runtime/Engine/Classes/Engine/Scene.h
@@ -371,6 +371,10 @@ struct FConvolutionBloomSettings
 	UPROPERTY(EditAnywhere, BlueprintReadWrite, Category = "Lens|Bloom", meta = (DisplayName = "Convolution Kernel"))
 	TObjectPtr<class UTexture2D> Texture;
 
+	/** Intensity multiplier on the scatter dispersion energy of the kernel. 1.0 means exactly use the same energy as the kernel scatter dispersion. */
+	UPROPERTY(Interp, BlueprintReadWrite, Category = "Lens|Bloom", AdvancedDisplay, meta = (ClampMin = "0.0", UIMax = "1.0", DisplayName = "Scatter Dispersion"))
+	float ScatterDispersion;
+
 	/** Relative size of the convolution kernel image compared to the minor axis of the viewport  */
 	UPROPERTY(Interp, BlueprintReadWrite, Category = "Lens|Bloom", AdvancedDisplay, meta = (ClampMin = "0.0", UIMax = "1.0", DisplayName = "Convolution Scale"))
 	float Size;
@@ -827,6 +831,9 @@ struct FPostProcessSettings
 
 	UPROPERTY(EditAnywhere, BlueprintReadWrite, Category=Overrides, meta=(PinHiddenByDefault, InlineEditConditionToggle))
 	uint8 bOverride_BloomConvolutionTexture : 1;
+	
+	UPROPERTY(EditAnywhere, BlueprintReadWrite, Category = Overrides, meta = (PinHiddenByDefault, InlineEditConditionToggle))
+	uint8 bOverride_BloomConvolutionScatterDispersion : 1;
 
 	UPROPERTY(EditAnywhere, BlueprintReadWrite, Category = Overrides, meta = (PinHiddenByDefault, InlineEditConditionToggle))
 	uint8 bOverride_BloomConvolutionSize : 1;
@@ -1438,6 +1445,10 @@ struct FPostProcessSettings
 	UPROPERTY(interp, BlueprintReadWrite, Category="Lens|Bloom", AdvancedDisplay, meta=(editcondition = "bOverride_Bloom6Tint", DisplayName = "#6 Tint", HideAlphaChannel))
 	FLinearColor Bloom6Tint;
 
+	/** Intensity multiplier on the scatter dispersion energy of the kernel. 1.0 means exactly use the same energy as the kernel scatter dispersion. */
+	UPROPERTY(interp, BlueprintReadWrite, Category = "Lens|Bloom", meta = (ClampMin = "0.0", UIMax = "20.0", editcondition = "bOverride_BloomConvolutionScatterDispersion", DisplayName = "Convolution Scatter Dispersion"))
+	float BloomConvolutionScatterDispersion;
+
 	/** Relative size of the convolution kernel image compared to the minor axis of the viewport  */
 	UPROPERTY(interp, BlueprintReadWrite, Category = "Lens|Bloom", AdvancedDisplay, meta = (ClampMin = "0.0", UIMax = "1.0", editcondition = "bOverride_BloomConvolutionSize", DisplayName = "Convolution Scale"))
 	float BloomConvolutionSize;
diff --git a/Engine/Source/Runtime/Engine/Private/PostProcessVolume.cpp b/Engine/Source/Runtime/Engine/Private/PostProcessVolume.cpp
--- a/Engine/Source/Runtime/Engine/Private/PostProcessVolume.cpp
+++ b/Engine/Source/Runtime/Engine/Private/PostProcessVolume.cpp
@@ -190,6 +190,7 @@ bool APostProcessVolume::CanEditChange(const FProperty* InProperty) const
 			// Parameters that are only of use with the bloom texture based fft
 			if (PropertyName == GET_MEMBER_NAME_STRING_CHECKED(FPostProcessSettings, BloomConvolutionTexture)      ||
 				PropertyName == GET_MEMBER_NAME_STRING_CHECKED(FPostProcessSettings, BloomConvolutionSize)         ||
+				PropertyName == GET_MEMBER_NAME_STRING_CHECKED(FPostProcessSettings, BloomConvolutionScatterDispersion) ||
 				PropertyName == GET_MEMBER_NAME_STRING_CHECKED(FPostProcessSettings, BloomConvolutionCenterUV)     ||
 				PropertyName == GET_MEMBER_NAME_STRING_CHECKED(FPostProcessSettings, BloomConvolutionPreFilterMin) ||
 				PropertyName == GET_MEMBER_NAME_STRING_CHECKED(FPostProcessSettings, BloomConvolutionPreFilterMax) ||
diff --git a/Engine/Source/Runtime/Engine/Private/Scene.cpp b/Engine/Source/Runtime/Engine/Private/Scene.cpp
--- a/Engine/Source/Runtime/Engine/Private/Scene.cpp
+++ b/Engine/Source/Runtime/Engine/Private/Scene.cpp
@@ -117,6 +117,7 @@ void FConvolutionBloomSettings::ExportToPostProcessSettings(FPostProcessSettings
 {
 	OutPostProcessSettings->bOverride_BloomConvolutionTexture = true;
 	OutPostProcessSettings->bOverride_BloomConvolutionSize = true;
+	OutPostProcessSettings->bOverride_BloomConvolutionScatterDispersion = true;
 	OutPostProcessSettings->bOverride_BloomConvolutionCenterUV = true;
 	OutPostProcessSettings->bOverride_BloomConvolutionPreFilterMin = true;
 	OutPostProcessSettings->bOverride_BloomConvolutionPreFilterMax = true;
@@ -125,6 +126,7 @@ void FConvolutionBloomSettings::ExportToPostProcessSettings(FPostProcessSettings
 
 	OutPostProcessSettings->BloomConvolutionTexture = Texture;
 	OutPostProcessSettings->BloomConvolutionSize = Size;
+	OutPostProcessSettings->BloomConvolutionScatterDispersion = ScatterDispersion;
 	OutPostProcessSettings->BloomConvolutionCenterUV = CenterUV;
 	OutPostProcessSettings->BloomConvolutionPreFilterMin = PreFilterMin;
 	OutPostProcessSettings->BloomConvolutionPreFilterMax = PreFilterMax;
@@ -305,6 +307,7 @@ static void DoPostProcessSettingsSanityCheck()
 
 		RenameMap.Add(TEXT("BloomConvolutionTexture"), TEXT(""));
 		RenameMap.Add(TEXT("BloomConvolutionSize"), TEXT(""));
+		RenameMap.Add(TEXT("BloomConvolutionScatterDispersion"), TEXT(""));
 		RenameMap.Add(TEXT("BloomConvolutionCenterUV"), TEXT(""));
 		RenameMap.Add(TEXT("BloomConvolutionPreFilterMin"), TEXT(""));
 		RenameMap.Add(TEXT("BloomConvolutionPreFilterMax"), TEXT(""));
@@ -449,6 +452,7 @@ FPostProcessSettings::FPostProcessSettings()
 	Bloom5Size = 30.0f;
 	Bloom6Tint = FLinearColor(0.061f, 0.061f, 0.061f);
 	Bloom6Size = 64.0f;
+	BloomConvolutionScatterDispersion = 1.f;
 	BloomConvolutionSize = 1.f;
 	BloomConvolutionCenterUV = FVector2D(0.5f, 0.5f);
 #if WITH_EDITORONLY_DATA
@@ -684,6 +688,7 @@ FPostProcessSettings::FPostProcessSettings(const FPostProcessSettings& Settings)
 	, bOverride_Bloom6Size(Settings.bOverride_Bloom6Size)
 	, bOverride_BloomSizeScale(Settings.bOverride_BloomSizeScale)
 	, bOverride_BloomConvolutionTexture(Settings.bOverride_BloomConvolutionTexture)
+	, bOverride_BloomConvolutionScatterDispersion(Settings.bOverride_BloomConvolutionScatterDispersion)
 	, bOverride_BloomConvolutionSize(Settings.bOverride_BloomConvolutionSize)
 	, bOverride_BloomConvolutionCenterUV(Settings.bOverride_BloomConvolutionCenterUV)
 	//, bOverride_BloomConvolutionPreFilter_DEPRECATED(Settings.bOverride_BloomConvolutionPreFilter_DEPRECATED)
@@ -862,6 +867,7 @@ FPostProcessSettings::FPostProcessSettings(const FPostProcessSettings& Settings)
 	, Bloom4Tint(Settings.Bloom4Tint)
 	, Bloom5Tint(Settings.Bloom5Tint)
 	, Bloom6Tint(Settings.Bloom6Tint)
+	, BloomConvolutionScatterDispersion(Settings.BloomConvolutionScatterDispersion)
 	, BloomConvolutionSize(Settings.BloomConvolutionSize)
 	, BloomConvolutionTexture(Settings.BloomConvolutionTexture)
 	, BloomConvolutionCenterUV(Settings.BloomConvolutionCenterUV)
diff --git a/Engine/Source/Runtime/Engine/Private/SceneView.cpp b/Engine/Source/Runtime/Engine/Private/SceneView.cpp
--- a/Engine/Source/Runtime/Engine/Private/SceneView.cpp
+++ b/Engine/Source/Runtime/Engine/Private/SceneView.cpp
@@ -1413,6 +1413,7 @@ void FSceneView::OverridePostProcessSettings(const FPostProcessSettings& Src, fl
 		LERP_PP(Bloom6Size);
 		LERP_PP(BloomDirtMaskIntensity);
 		LERP_PP(BloomDirtMaskTint);
+		LERP_PP(BloomConvolutionScatterDispersion);
 		LERP_PP(BloomConvolutionSize);
 		LERP_PP(BloomConvolutionCenterUV);
 		LERP_PP(BloomConvolutionPreFilterMin);
diff --git a/Engine/Source/Runtime/Renderer/Private/PostProcess/PostProcessBloomSetup.cpp b/Engine/Source/Runtime/Renderer/Private/PostProcess/PostProcessBloomSetup.cpp
--- a/Engine/Source/Runtime/Renderer/Private/PostProcess/PostProcessBloomSetup.cpp
+++ b/Engine/Source/Runtime/Renderer/Private/PostProcess/PostProcessBloomSetup.cpp
@@ -109,6 +109,11 @@ public:
 IMPLEMENT_GLOBAL_SHADER(FBloomSetupCS, "/Engine/Private/PostProcessBloom.usf", "BloomSetupCS", SF_Compute);
 } //!namespace
 
+bool FBloomOutputs::SupportsApplyParametersBuffer(EShaderPlatform Platform)
+{
+	return FDataDrivenShaderPlatformInfo::GetSupportsFFTBloom(Platform);
+}
+
 FScreenPassTexture AddBloomSetupPass(FRDGBuilder& GraphBuilder, const FViewInfo& View, const FBloomSetupInputs& Inputs)
 {
 	check(Inputs.SceneColor.IsValid());
@@ -180,112 +185,96 @@ static_assert(
 	static_cast<uint32>(EBloomQuality::MAX) == FSceneDownsampleChain::StageCount,
 	"The total number of stages in the scene downsample chain and the number of bloom quality levels must match.");
 
-FBloomOutputs AddBloomPass(FRDGBuilder& GraphBuilder, const FViewInfo& View, const FBloomInputs& Inputs)
+FBloomOutputs AddGaussianBloomPasses(FRDGBuilder& GraphBuilder, const FViewInfo& View, const FSceneDownsampleChain* SceneDownsampleChain)
 {
-	check(Inputs.SceneColor.IsValid());
-	check(Inputs.SceneDownsampleChain);
+	check(SceneDownsampleChain);
+	check(!IsFFTBloomEnabled(View));
 
 	const FFinalPostProcessSettings& Settings = View.FinalPostProcessSettings;
 
 	const EBloomQuality BloomQuality = GetBloomQuality();
 
-	FScreenPassTexture SceneColor = Inputs.SceneColor;
-	FScreenPassTexture Bloom;
-
+	FBloomOutputs PassOutputs;
 	if (BloomQuality != EBloomQuality::Disabled)
 	{
-		const bool bFFTBloomEnabled = IsFFTBloomEnabled(View);
+		RDG_EVENT_SCOPE(GraphBuilder, "Bloom");
 
-		if (bFFTBloomEnabled)
-		{
-			FScreenPassTexture FullResolution = Inputs.SceneColor;
-			FScreenPassTexture HalfResolution;
-			if (IsFFTBloomQuarterResolutionEnabled())
-			{
-				HalfResolution = Inputs.SceneDownsampleChain->GetTexture(1);
-			}
-			else
-			{
-				HalfResolution = Inputs.SceneDownsampleChain->GetFirstTexture();
-			}
+		const float CrossBloom = CVarBloomCross.GetValueOnRenderThread();
 
-			FFFTBloomInputs PassInputs;
-			PassInputs.FullResolutionTexture = FullResolution.Texture;
-			PassInputs.FullResolutionViewRect = FullResolution.ViewRect;
-			PassInputs.HalfResolutionTexture = HalfResolution.Texture;
-			PassInputs.HalfResolutionViewRect = HalfResolution.ViewRect;
+		const FVector2D CrossCenterWeight(FMath::Max(CrossBloom, 0.0f), FMath::Abs(CrossBloom));
 
-			SceneColor.Texture = AddFFTBloomPass(GraphBuilder, View, PassInputs);
-		}
-		else
+		check(BloomQuality != EBloomQuality::Disabled);
+		const uint32 BloomQualityIndex = static_cast<uint32>(BloomQuality);
+		const uint32 BloomQualityCountMax = static_cast<uint32>(EBloomQuality::MAX);
+
+		struct FBloomStage
 		{
-			RDG_EVENT_SCOPE(GraphBuilder, "Bloom");
+			const float Size;
+			const FLinearColor& Tint;
+		};
 
-			const float CrossBloom = CVarBloomCross.GetValueOnRenderThread();
+		FBloomStage BloomStages[] =
+		{
+			{ Settings.Bloom6Size, Settings.Bloom6Tint },
+			{ Settings.Bloom5Size, Settings.Bloom5Tint },
+			{ Settings.Bloom4Size, Settings.Bloom4Tint },
+			{ Settings.Bloom3Size, Settings.Bloom3Tint },
+			{ Settings.Bloom2Size, Settings.Bloom2Tint },
+			{ Settings.Bloom1Size, Settings.Bloom1Tint }
+		};
+
+		const uint32 BloomQualityToSceneDownsampleStage[] =
+		{
+			static_cast<uint32>(-1), // Disabled (sentinel entry to preserve indices)
+			3, // Q1
+			3, // Q2
+			4, // Q3
+			5, // Q4
+			6  // Q5
+		};
 
-			const FVector2D CrossCenterWeight(FMath::Max(CrossBloom, 0.0f), FMath::Abs(CrossBloom));
+		static_assert(UE_ARRAY_COUNT(BloomStages) == BloomQualityCountMax, "Array must be one less than the number of bloom quality entries.");
+		static_assert(UE_ARRAY_COUNT(BloomQualityToSceneDownsampleStage) == BloomQualityCountMax, "Array must be one less than the number of bloom quality entries.");
 
-			check(BloomQuality != EBloomQuality::Disabled);
-			const uint32 BloomQualityIndex = static_cast<uint32>(BloomQuality);
-			const uint32 BloomQualityCountMax = static_cast<uint32>(EBloomQuality::MAX);
+		// Use bloom quality to select the number of downsample stages to use for bloom.
+		const uint32 BloomStageCount = BloomQualityToSceneDownsampleStage[BloomQualityIndex];
 
-			struct FBloomStage
-			{
-				const float Size;
-				const FLinearColor& Tint;
-			};
+		const float TintScale = 1.0f / BloomQualityCountMax;
 
-			FBloomStage BloomStages[] =
-			{
-				{ Settings.Bloom6Size, Settings.Bloom6Tint },
-				{ Settings.Bloom5Size, Settings.Bloom5Tint },
-				{ Settings.Bloom4Size, Settings.Bloom4Tint },
-				{ Settings.Bloom3Size, Settings.Bloom3Tint },
-				{ Settings.Bloom2Size, Settings.Bloom2Tint },
-				{ Settings.Bloom1Size, Settings.Bloom1Tint }
-			};
-
-			const uint32 BloomQualityToSceneDownsampleStage[] =
-			{
-				static_cast<uint32>(-1), // Disabled (sentinel entry to preserve indices)
-				3, // Q1
-				3, // Q2
-				4, // Q3
-				5, // Q4
-				6  // Q5
-			};
+		for (uint32 StageIndex = 0, SourceIndex = BloomQualityCountMax - 1; StageIndex < BloomStageCount; ++StageIndex, --SourceIndex)
+		{
+			const FBloomStage& BloomStage = BloomStages[StageIndex];
 
-			static_assert(UE_ARRAY_COUNT(BloomStages) == BloomQualityCountMax, "Array must be one less than the number of bloom quality entries.");
-			static_assert(UE_ARRAY_COUNT(BloomQualityToSceneDownsampleStage) == BloomQualityCountMax, "Array must be one less than the number of bloom quality entries.");
+			if (BloomStage.Size > SMALL_NUMBER)
+			{
+				FGaussianBlurInputs PassInputs;
+				PassInputs.NameX = TEXT("BloomX");
+				PassInputs.NameY = TEXT("BloomY");
+				PassInputs.Filter = SceneDownsampleChain->GetTexture(SourceIndex);
+				PassInputs.Additive = PassOutputs.Bloom;
+				PassInputs.CrossCenterWeight = CrossCenterWeight;
+				PassInputs.KernelSizePercent = BloomStage.Size * Settings.BloomSizeScale;
+				PassInputs.TintColor = BloomStage.Tint * TintScale;
+
+				PassOutputs.Bloom = AddGaussianBlurPass(GraphBuilder, View, PassInputs);
+			}
+		}
 
-			// Use bloom quality to select the number of downsample stages to use for bloom.
-			const uint32 BloomStageCount = BloomQualityToSceneDownsampleStage[BloomQualityIndex];
+		if (FBloomOutputs::SupportsApplyParametersBuffer(View.GetShaderPlatform()))
+		{
+			FBloomOutputs::FApplyInfo ApplyParameters;
+			ApplyParameters.SceneColorMultiply = FLinearColor::White;
+			ApplyParameters.BloomMultiply = FLinearColor::White * View.FinalPostProcessSettings.BloomIntensity;;
 
-			const float TintScale = 1.0f / BloomQualityCountMax;
+			TArray<FBloomOutputs::FApplyInfo, TInlineAllocator<1> > ApplyParametersArray;
+			ApplyParametersArray.Add(ApplyParameters);
 
-			for (uint32 StageIndex = 0, SourceIndex = BloomQualityCountMax - 1; StageIndex < BloomStageCount; ++StageIndex, --SourceIndex)
-			{
-				const FBloomStage& BloomStage = BloomStages[StageIndex];
-
-				if (BloomStage.Size > SMALL_NUMBER)
-				{
-					FGaussianBlurInputs PassInputs;
-					PassInputs.NameX = TEXT("BloomX");
-					PassInputs.NameY = TEXT("BloomY");
-					PassInputs.Filter = Inputs.SceneDownsampleChain->GetTexture(SourceIndex);
-					PassInputs.Additive = Bloom;
-					PassInputs.CrossCenterWeight = CrossCenterWeight;
-					PassInputs.KernelSizePercent = BloomStage.Size * Settings.BloomSizeScale;
-					PassInputs.TintColor = BloomStage.Tint * TintScale;
-
-					Bloom = AddGaussianBlurPass(GraphBuilder, View, PassInputs);
-				}
-			}
+			PassOutputs.ApplyParameters = CreateStructuredBuffer(
+				GraphBuilder,
+				TEXT("Bloom.ApplyParameters"),
+				ApplyParametersArray);
 		}
 	}
 
-	FBloomOutputs PassOutputs;
-	PassOutputs.SceneColor = SceneColor;
-	PassOutputs.Bloom = Bloom;
 	return PassOutputs;
 }
diff --git a/Engine/Source/Runtime/Renderer/Private/PostProcess/PostProcessBloomSetup.h b/Engine/Source/Runtime/Renderer/Private/PostProcess/PostProcessBloomSetup.h
--- a/Engine/Source/Runtime/Renderer/Private/PostProcess/PostProcessBloomSetup.h
+++ b/Engine/Source/Runtime/Renderer/Private/PostProcess/PostProcessBloomSetup.h
@@ -19,20 +19,26 @@ enum class EBloomQuality : uint32
 
 EBloomQuality GetBloomQuality();
 
-struct FBloomInputs
-{
-	FScreenPassTexture SceneColor;
-
-	const FSceneDownsampleChain* SceneDownsampleChain = nullptr;
-};
-
 struct FBloomOutputs
 {
-	FScreenPassTexture SceneColor;
+	// Bloom texture to apply to the scene colors.
 	FScreenPassTexture Bloom;
+
+	// Parameters to apply the bloom texture to scene color.
+	FRDGBufferRef ApplyParameters = nullptr;
+
+	// Structure of ApplyParameters.
+	struct FApplyInfo
+	{
+		FLinearColor SceneColorMultiply;
+		FLinearColor BloomMultiply;
+	};
+
+	// Returns whether ApplyParameters is supported by the tonemapper.
+	static bool SupportsApplyParametersBuffer(EShaderPlatform Platform);
 };
 
-FBloomOutputs AddBloomPass(FRDGBuilder& GraphBuilder, const FViewInfo& View, const FBloomInputs& Inputs);
+FBloomOutputs AddGaussianBloomPasses(FRDGBuilder& GraphBuilder, const FViewInfo& View, const FSceneDownsampleChain* SceneDownsampleChain);
 
 struct FBloomSetupInputs
 {
diff --git a/Engine/Source/Runtime/Renderer/Private/PostProcess/PostProcessFFTBloom.cpp b/Engine/Source/Runtime/Renderer/Private/PostProcess/PostProcessFFTBloom.cpp
--- a/Engine/Source/Runtime/Renderer/Private/PostProcess/PostProcessFFTBloom.cpp
+++ b/Engine/Source/Runtime/Renderer/Private/PostProcess/PostProcessFFTBloom.cpp
@@ -1,19 +1,24 @@
 // Copyright Epic Games, Inc. All Rights Reserved.
 
 #include "PostProcess/PostProcessFFTBloom.h"
+#include "PostProcess/PostProcessBloomSetup.h"
 #include "GPUFastFourierTransform.h"
 #include "RendererModule.h"
 #include "Rendering/Texture2DResource.h"
 
 namespace
 {
+TAutoConsoleVariable<int32> CVarBloomCacheKernel(
+	TEXT("r.Bloom.CacheKernel"), 1,
+	TEXT("Whether to cache the kernel in spectral domain."),
+	ECVF_RenderThreadSafe);
+
 TAutoConsoleVariable<int32> CVarHalfResFFTBloom(
-	TEXT("r.Bloom.HalfResolutionFFT"),
-	0,
+	TEXT("r.Bloom.HalfResolutionFFT"), 0,
 	TEXT("Experimental half-resolution FFT Bloom convolution. \n")
-	TEXT(" 0: Standard full resolution convolution bloom.")
-	TEXT(" 1: Half-resolution convolution that excludes the center of the kernel.\n")
-	TEXT(" 2: Quarter-resolution convolution that excludes the center of the kernel.\n"),
+	TEXT(" 0: Standard full resolution convolution bloom;")
+	TEXT(" 1: Half-resolution convolution;\n")
+	TEXT(" 2: Quarter-resolution convolution.\n"),
 	ECVF_Scalability | ECVF_RenderThreadSafe);
 
 static bool DoesPlatformSupportFFTBloom(EShaderPlatform Platform)
@@ -24,9 +29,6 @@ static bool DoesPlatformSupportFFTBloom(EShaderPlatform Platform)
 class FFFTBloomShader : public FGlobalShader
 {
 public:
-	// Determine the number of threads used per scanline when writing the physical space kernel
-	static const uint32 ThreadsPerGroup = 32;
-
 	static bool ShouldCompilePermutation(const FGlobalShaderPermutationParameters& Parameters)
 	{
 		return DoesPlatformSupportFFTBloom(Parameters.Platform);
@@ -38,86 +40,154 @@ public:
 	{}
 };
 
-class FResizeAndCenterTextureCS : public FFFTBloomShader
+class FBloomFindKernelCenterCS : public FFFTBloomShader
 {
 public:
-	DECLARE_GLOBAL_SHADER(FResizeAndCenterTextureCS);
-	SHADER_USE_PARAMETER_STRUCT(FResizeAndCenterTextureCS, FFFTBloomShader);
+	DECLARE_GLOBAL_SHADER(FBloomFindKernelCenterCS);
+	SHADER_USE_PARAMETER_STRUCT(FBloomFindKernelCenterCS, FFFTBloomShader);
 
 	BEGIN_SHADER_PARAMETER_STRUCT(FParameters, )
-		SHADER_PARAMETER(FIntPoint, DstExtent)
-		SHADER_PARAMETER(FIntPoint, ImageExtent)
-		SHADER_PARAMETER(FVector4, KernelCenterAndScale)
-		SHADER_PARAMETER(FIntPoint, DstBufferExtent)
+		SHADER_PARAMETER(FIntPoint, KernelSpatialTextureSize)
+		SHADER_PARAMETER_RDG_TEXTURE(Texture2D, KernelSpatialTexture)
+		SHADER_PARAMETER_RDG_BUFFER_UAV(RWStructuredBuffer<uint>, KernelCenterCoordOutput)
+	END_SHADER_PARAMETER_STRUCT()
+};
 
-		SHADER_PARAMETER_TEXTURE(Texture2D, SrcTexture)
-		SHADER_PARAMETER_SAMPLER(SamplerState, SrcSampler)
+class FBloomSurveyMaxScatterDispersionCS : public FFFTBloomShader
+{
+public:
+	DECLARE_GLOBAL_SHADER(FBloomSurveyMaxScatterDispersionCS);
+	SHADER_USE_PARAMETER_STRUCT(FBloomSurveyMaxScatterDispersionCS, FFFTBloomShader);
 
-		SHADER_PARAMETER_RDG_TEXTURE_UAV(RWTexture2D<float4>, DstTexture)
+	BEGIN_SHADER_PARAMETER_STRUCT(FParameters, )
+		SHADER_PARAMETER(float, ViewTexelRadiusInKernelTexels)
+		SHADER_PARAMETER(int32, SurveyGroupGridSize)
+		SHADER_PARAMETER(FIntPoint, KernelSpatialTextureSize)
+		SHADER_PARAMETER_RDG_TEXTURE(Texture2D, KernelSpatialTexture)
+		SHADER_PARAMETER_RDG_BUFFER_SRV(StructuredBuffer<uint>, KernelCenterCoordBuffer)
+		SHADER_PARAMETER_RDG_BUFFER_UAV(RWStructuredBuffer<uint>, SurveyOutput)
+		SHADER_PARAMETER_RDG_TEXTURE_UAV(RWTexture2D<float4>, DebugOutput)
 	END_SHADER_PARAMETER_STRUCT()
+};
 
-	static void ModifyCompilationEnvironment(const FGlobalShaderPermutationParameters& Parameters, FShaderCompilerEnvironment& OutEnvironment)
-	{
-		FFFTBloomShader::ModifyCompilationEnvironment(Parameters, OutEnvironment);
-		OutEnvironment.SetDefine(TEXT("INCLUDE_RESIZE_AND_CENTER"), 1);
-		OutEnvironment.SetDefine(TEXT("THREADS_PER_GROUP"), ThreadsPerGroup);
-	}
+class FBloomSurveyKernelCenterEnergyCS : public FFFTBloomShader
+{
+public:
+	DECLARE_GLOBAL_SHADER(FBloomSurveyKernelCenterEnergyCS);
+	SHADER_USE_PARAMETER_STRUCT(FBloomSurveyKernelCenterEnergyCS, FFFTBloomShader);
+
+	BEGIN_SHADER_PARAMETER_STRUCT(FParameters, )
+		SHADER_PARAMETER(float, ViewTexelRadiusInKernelTexels)
+		SHADER_PARAMETER(int32, SurveyGroupGridSize)
+		SHADER_PARAMETER(FIntPoint, KernelSpatialTextureSize)
+		SHADER_PARAMETER_RDG_TEXTURE(Texture2D, KernelSpatialTexture)
+		SHADER_PARAMETER_RDG_BUFFER_SRV(StructuredBuffer<uint>, KernelCenterCoordBuffer)
+		SHADER_PARAMETER_RDG_BUFFER_SRV(StructuredBuffer<uint>, MaxScatterDispersionBuffer)
+		SHADER_PARAMETER_RDG_BUFFER_UAV(RWStructuredBuffer<uint>, SurveyOutput)
+		SHADER_PARAMETER_RDG_TEXTURE_UAV(RWTexture2D<float4>, DebugOutput)
+	END_SHADER_PARAMETER_STRUCT()
 };
 
-class FCaptureKernelWeightsCS : public FFFTBloomShader
+class FBloomReduceKernelSurveyCS : public FFFTBloomShader
 {
 public:
-	DECLARE_GLOBAL_SHADER(FCaptureKernelWeightsCS);
-	SHADER_USE_PARAMETER_STRUCT(FCaptureKernelWeightsCS, FFFTBloomShader);
+	DECLARE_GLOBAL_SHADER(FBloomReduceKernelSurveyCS);
+	SHADER_USE_PARAMETER_STRUCT(FBloomReduceKernelSurveyCS, FFFTBloomShader);
 
 	BEGIN_SHADER_PARAMETER_STRUCT(FParameters, )
-		SHADER_PARAMETER(FIntPoint, HalfResSumLocation)
-		SHADER_PARAMETER(FVector2D, UVCenter)
+		SHADER_PARAMETER(int32, SurveyReduceOp)
+		SHADER_PARAMETER(int32, SurveyGroupCount)
+		SHADER_PARAMETER_RDG_BUFFER_UAV(RWStructuredBuffer<uint>, SurveyOutput)
+	END_SHADER_PARAMETER_STRUCT()
+};
 
-		SHADER_PARAMETER_SAMPLER(SamplerState, PhysicalSrcSampler)
-		SHADER_PARAMETER_RDG_TEXTURE(Texture2D, HalfResSrcTexture)
-		SHADER_PARAMETER_TEXTURE(Texture2D, PhysicalSrcTexture)
+class FBloomSumScatterDispersionEnergyCS : public FFFTBloomShader
+{
+public:
+	DECLARE_GLOBAL_SHADER(FBloomSumScatterDispersionEnergyCS);
+	SHADER_USE_PARAMETER_STRUCT(FBloomSumScatterDispersionEnergyCS, FFFTBloomShader);
 
-		SHADER_PARAMETER_RDG_TEXTURE_UAV(RWTexture2D<float4>, DstTexture)
+	BEGIN_SHADER_PARAMETER_STRUCT(FParameters, )
+		SHADER_PARAMETER(int32, PassId)
+		SHADER_PARAMETER(FIntPoint, ScatterDispersionTextureSize)
+		SHADER_PARAMETER_RDG_TEXTURE(Texture2D, ScatterDispersionTexture)
+		SHADER_PARAMETER_RDG_BUFFER_SRV(StructuredBuffer<uint>, MaxScatterDispersionBuffer)
+		SHADER_PARAMETER_RDG_TEXTURE_UAV(RWTexture2D<float4>, ScatterDispersionOutput)
 	END_SHADER_PARAMETER_STRUCT()
+};
 
-	static void ModifyCompilationEnvironment(const FGlobalShaderPermutationParameters& Parameters, FShaderCompilerEnvironment& OutEnvironment)
-	{
-		FFFTBloomShader::ModifyCompilationEnvironment(Parameters, OutEnvironment);
-		OutEnvironment.SetDefine(TEXT("INCLUDE_CAPTURE_KERNEL_WEIGHTS"), 1);
-	}
+class FBloomPackKernelConstantsCS : public FFFTBloomShader
+{
+public:
+	DECLARE_GLOBAL_SHADER(FBloomPackKernelConstantsCS);
+	SHADER_USE_PARAMETER_STRUCT(FBloomPackKernelConstantsCS, FFFTBloomShader);
+
+	BEGIN_SHADER_PARAMETER_STRUCT(FParameters, )
+		SHADER_PARAMETER_RDG_BUFFER_SRV(StructuredBuffer<uint>, KernelCenterCoordBuffer)
+		SHADER_PARAMETER_RDG_BUFFER_SRV(StructuredBuffer<uint>, KernelCenterEnergyBuffer)
+		SHADER_PARAMETER_RDG_BUFFER_SRV(StructuredBuffer<uint>, MaxScatterDispersionBuffer)
+		SHADER_PARAMETER_RDG_TEXTURE(StructuredBuffer<uint>, ScatterDispersionTexture)
+		SHADER_PARAMETER_RDG_BUFFER_UAV(RWStructuredBuffer<uint>, KernelConstantsOutput)
+	END_SHADER_PARAMETER_STRUCT()
+};
+
+class FBloomClampKernelCS : public FFFTBloomShader
+{
+public:
+	DECLARE_GLOBAL_SHADER(FBloomClampKernelCS);
+	SHADER_USE_PARAMETER_STRUCT(FBloomClampKernelCS, FFFTBloomShader);
+
+	BEGIN_SHADER_PARAMETER_STRUCT(FParameters, )
+		SHADER_PARAMETER_RDG_TEXTURE(Texture2D, KernelSpatialTexture)
+		SHADER_PARAMETER_RDG_BUFFER_SRV(StructuredBuffer<uint>, KernelConstantsBuffer)
+		SHADER_PARAMETER_RDG_TEXTURE_UAV(RWTexture2D<float4>, ClampedKernelSpatialOutput)
+	END_SHADER_PARAMETER_STRUCT()
 };
 
-class FBlendLowResCS : public FFFTBloomShader
+class FBloomResizeKernelCS : public FFFTBloomShader
 {
 public:
-	DECLARE_GLOBAL_SHADER(FBlendLowResCS);
-	SHADER_USE_PARAMETER_STRUCT(FBlendLowResCS, FFFTBloomShader);
+	DECLARE_GLOBAL_SHADER(FBloomResizeKernelCS);
+	SHADER_USE_PARAMETER_STRUCT(FBloomResizeKernelCS, FFFTBloomShader);
 
 	BEGIN_SHADER_PARAMETER_STRUCT(FParameters, )
-		SHADER_PARAMETER(FIntRect, DstRect)
-		SHADER_PARAMETER(FIntRect, HalfRect)
-		SHADER_PARAMETER(FIntPoint, HalfBufferSize)
+		SHADER_PARAMETER(FIntPoint, DstExtent)
+		SHADER_PARAMETER(FIntPoint, ImageExtent)
+		SHADER_PARAMETER(FVector2D, KernelSpatialTextureInvSize)
+		SHADER_PARAMETER(FIntPoint, DstBufferExtent)
+		SHADER_PARAMETER(float, KernelSupportScale)
+
+		SHADER_PARAMETER_RDG_BUFFER_SRV(StructuredBuffer<uint>, KernelConstantsBuffer)
 
 		SHADER_PARAMETER_RDG_TEXTURE(Texture2D, SrcTexture)
-		SHADER_PARAMETER_RDG_TEXTURE(Texture2D, HalfResSrcTexture)
-		SHADER_PARAMETER_SAMPLER(SamplerState, HalfResSrcSampler)
-		SHADER_PARAMETER_RDG_TEXTURE(Texture2D, CenterWeightTexture)
+		SHADER_PARAMETER_SAMPLER(SamplerState, SrcSampler)
 
 		SHADER_PARAMETER_RDG_TEXTURE_UAV(RWTexture2D<float4>, DstTexture)
 	END_SHADER_PARAMETER_STRUCT()
+};
 
-	static void ModifyCompilationEnvironment(const FGlobalShaderPermutationParameters& Parameters, FShaderCompilerEnvironment& OutEnvironment)
-	{
-		FFFTBloomShader::ModifyCompilationEnvironment(Parameters, OutEnvironment);
-		OutEnvironment.SetDefine(TEXT("INCLUDE_BLEND_LOW_RES"), 1);
-		OutEnvironment.SetDefine(TEXT("THREADS_PER_GROUP"), ThreadsPerGroup);
-	}
+class FBloomFinalizeApplyConstantsCS : public FFFTBloomShader
+{
+public:
+	DECLARE_GLOBAL_SHADER(FBloomFinalizeApplyConstantsCS);
+	SHADER_USE_PARAMETER_STRUCT(FBloomFinalizeApplyConstantsCS, FFFTBloomShader);
+
+	BEGIN_SHADER_PARAMETER_STRUCT(FParameters, )
+		SHADER_PARAMETER(float, ScatterDispersionIntensity)
+		SHADER_PARAMETER_RDG_BUFFER_SRV(StructuredBuffer<uint>, KernelConstantsBuffer)
+		SHADER_PARAMETER_RDG_BUFFER_UAV(RWStructuredBuffer<uint>, BloomApplyConstantsOutput)
+	END_SHADER_PARAMETER_STRUCT()
 };
 
-IMPLEMENT_GLOBAL_SHADER(FResizeAndCenterTextureCS, "/Engine/Private/PostProcessFFTBloom.usf", "ResizeAndCenterTextureCS", SF_Compute);
-IMPLEMENT_GLOBAL_SHADER(FCaptureKernelWeightsCS,   "/Engine/Private/PostProcessFFTBloom.usf", "CaptureKernelWeightsCS", SF_Compute);
-IMPLEMENT_GLOBAL_SHADER(FBlendLowResCS,            "/Engine/Private/PostProcessFFTBloom.usf", "BlendLowResCS", SF_Compute);
+IMPLEMENT_GLOBAL_SHADER(FBloomFindKernelCenterCS,           "/Engine/Private/Bloom/BloomFindKernelCenter.usf", "MainCS", SF_Compute);
+IMPLEMENT_GLOBAL_SHADER(FBloomSurveyMaxScatterDispersionCS, "/Engine/Private/Bloom/BloomSurveyMaxScatterDispersion.usf", "MainCS", SF_Compute);
+IMPLEMENT_GLOBAL_SHADER(FBloomSurveyKernelCenterEnergyCS,   "/Engine/Private/Bloom/BloomSurveyKernelCenterEnergy.usf", "MainCS", SF_Compute);
+IMPLEMENT_GLOBAL_SHADER(FBloomReduceKernelSurveyCS,         "/Engine/Private/Bloom/BloomReduceKernelSurvey.usf", "MainCS", SF_Compute);
+IMPLEMENT_GLOBAL_SHADER(FBloomSumScatterDispersionEnergyCS, "/Engine/Private/Bloom/BloomSumScatterDispersionEnergy.usf", "MainCS", SF_Compute);
+IMPLEMENT_GLOBAL_SHADER(FBloomPackKernelConstantsCS,        "/Engine/Private/Bloom/BloomPackKernelConstants.usf", "MainCS", SF_Compute);
+IMPLEMENT_GLOBAL_SHADER(FBloomClampKernelCS,                "/Engine/Private/Bloom/BloomClampKernel.usf", "MainCS", SF_Compute);
+IMPLEMENT_GLOBAL_SHADER(FBloomResizeKernelCS,               "/Engine/Private/Bloom/BloomResizeKernel.usf", "MainCS", SF_Compute);
+IMPLEMENT_GLOBAL_SHADER(FBloomFinalizeApplyConstantsCS,     "/Engine/Private/Bloom/BloomFinalizeApplyConstants.usf", "MainCS", SF_Compute);
 
 } //! namespace
 
@@ -154,162 +224,6 @@ bool IsFFTBloomEnabled(const FViewInfo& View)
 	}
 }
 
-/**
-* Used to resample the physical space kernel into the correct sized buffer with the 
-* correct periodicity and center.
-*
-* Resizes the image, moves the center to to 0,0 and applies periodicity 
-* across the full TargetSize (periods TargetSize.x & TargetSize.y)
-*
-* @param SrcTexture    - SRV for the physical space kernel supplied by user
-* @param SrcImageSize  - The extent of the src image
-* @param SrcImageCenterUV - The location of the center in src image (e.g. where the kernel center really is).
-* @param ResizeScale    - Affective size of the physical space kernel in units of the ImageExtent.x 
-* @param TargetSize    - Size of the image produced. 
-* @param DstUAV        - Holds the result
-* @param DstBufferSize - Size of DstBuffer
-* @param bForceCenterZero -  is true only for the experimental 1/2 res version, part of conserving energy
-*/
-void ResizeAndCenterTexture(
-	FRDGBuilder& GraphBuilder,
-	const FViewInfo& View,
-	const FTextureRHIRef& SrcTexture,
-	const FIntPoint& SrcImageSize,
-	const FVector2D& SrcImageCenterUV,
-	const float ResizeScale,
-	const FIntPoint& TargetSize,
-	FRDGTextureRef DstTexture,
-	const FIntPoint& DstBufferSize,
-	const bool bForceCenterZero)
-{
-	// Clamp the image center
-	FVector2D ClampedImageCenterUV = SrcImageCenterUV;
-	ClampedImageCenterUV.X = FMath::Clamp(SrcImageCenterUV.X, 0.f, 1.f);
-	ClampedImageCenterUV.Y = FMath::Clamp(SrcImageCenterUV.Y, 0.f, 1.f);
-
-	check(DstTexture);
-
-	float CenterScale = bForceCenterZero ? 0.f : 1.f;
-	const FLinearColor KernelCenterAndScale(ClampedImageCenterUV.X, ClampedImageCenterUV.Y, ResizeScale, CenterScale);
-
-	FResizeAndCenterTextureCS::FParameters* PassParameters = GraphBuilder.AllocParameters<FResizeAndCenterTextureCS::FParameters>();
-	PassParameters->DstExtent = TargetSize;
-	PassParameters->ImageExtent = SrcImageSize;
-	PassParameters->KernelCenterAndScale = KernelCenterAndScale;
-	PassParameters->DstBufferExtent = DstBufferSize;
-
-	PassParameters->SrcTexture = SrcTexture;
-	PassParameters->SrcSampler = TStaticSamplerState<SF_Bilinear, AM_Wrap, AM_Wrap, AM_Wrap>::GetRHI();
-	PassParameters->DstTexture = GraphBuilder.CreateUAV(DstTexture);
-
-	// Use multiple threads per scan line to insure memory coalescing during the write
-	const int32 ThreadsPerGroup = FResizeAndCenterTextureCS::ThreadsPerGroup;
-	const int32 ThreadsGroupsPerScanLine = (DstBufferSize.X % ThreadsPerGroup == 0) ? DstBufferSize.X / ThreadsPerGroup : DstBufferSize.X / ThreadsPerGroup + 1;
-	
-	TShaderMapRef<FResizeAndCenterTextureCS> ComputeShader(View.ShaderMap);
-	FComputeShaderUtils::AddPass(
-		GraphBuilder,
-		RDG_EVENT_NAME("FFTBloom PreProcessKernel %dx%d", TargetSize.X, TargetSize.Y),
-		ComputeShader,
-		PassParameters, FIntVector(ThreadsGroupsPerScanLine, DstBufferSize.Y, 1));
-}
-
-FRDGTextureRef CaptureKernelWeight(
-	FRDGBuilder& GraphBuilder,
-	const FViewInfo& View,
-	FRDGTextureRef HalfResKernel,
-	const FIntPoint& HalfResSumLocation,
-	const FTextureRHIRef& PhysicalKernel,
-	const FVector2D& CenterUV)
-{
-	FRDGTextureDesc CenterWeightDesc = FRDGTextureDesc::Create2D(
-		FIntPoint(2, 1),
-		GPUFFT::PixelFormat(),
-		FClearValueBinding::None,
-		TexCreate_ShaderResource | TexCreate_UAV);
-
-	// Resize the buffer to hold the transformed kernel
-	FRDGTextureRef CenterWeightTexture = GraphBuilder.CreateTexture(CenterWeightDesc, TEXT("Bloom.FFT.KernelCenterWeight"));
-
-	FCaptureKernelWeightsCS::FParameters* PassParameters = GraphBuilder.AllocParameters<FCaptureKernelWeightsCS::FParameters>();
-	PassParameters->HalfResSumLocation = HalfResSumLocation;
-	PassParameters->UVCenter = CenterUV;
-	PassParameters->PhysicalSrcSampler = TStaticSamplerState<SF_Bilinear, AM_Wrap, AM_Wrap, AM_Wrap>::GetRHI();
-	PassParameters->HalfResSrcTexture = HalfResKernel;
-	PassParameters->PhysicalSrcTexture = PhysicalKernel;
-	PassParameters->DstTexture = GraphBuilder.CreateUAV(CenterWeightTexture);
-
-	TShaderMapRef<FCaptureKernelWeightsCS> ComputeShader(View.ShaderMap);
-	FComputeShaderUtils::AddPass(
-		GraphBuilder,
-		RDG_EVENT_NAME("FFTBloom CaptureKernelWeights"),
-		ComputeShader,
-		PassParameters, 
-		FIntVector(1, 1, 1));
-
-	return CenterWeightTexture;
-}
-
-/**
-* Used by energy conserving 1/2 resolution version of the bloom.
-* This blends the results of the low resolution bloom with the full resolution image 
-* in an energy conserving manner.  Assumes the 1/2-res bloom is done with a kernel that
-* is missing the center pixel (i.e. the self-gather contribution), and this missing contribution
-* is supplied by the full-res image.
-*
-* @param Context  - container for RHI and ShaderMap
-* @param FullResImage - Unbloomed full-resolution source image
-* @param FullResImageRect      - Region in FullResImage and DstUAV where the image lives
-* @param HaflResConvolvedImage - A 1/2 res image that has been convolved with the bloom kernel minus center.
-* @param HalfResRect           - Location of image in the HalfResConvolvedImage buffer
-* @param HalfBufferSize        - Full size of the 1/2 Res buffer.
-* @param CenterWeightTexture   - Texture that holds the weight between the kernel center and sum of 1/2res kernel weights.
-*                                 needed to correctly composite the 1/2 res bloomed result with the full-res image.
-* @param DstUAV                - Destination buffer that will hold the result.
-*/
-void BlendLowRes(
-	FRDGBuilder& GraphBuilder,
-	const FViewInfo& View,
-	FRDGTextureRef FullResImage,
-	const FIntRect& FullResImageRect,
-	FRDGTextureRef HalfResConvolvedImage,
-	const FIntRect& HalfResRect,
-	const FIntPoint& HalfBufferSize,
-	FRDGTextureRef CenterWeightTexutre,
-	FRDGTextureRef DstTexture)
-{
-	FBlendLowResCS::FParameters* PassParameters = GraphBuilder.AllocParameters<FBlendLowResCS::FParameters>();
-	PassParameters->DstRect = FullResImageRect;
-	PassParameters->HalfRect = HalfResRect;
-	PassParameters->HalfBufferSize = HalfBufferSize;
-
-	PassParameters->SrcTexture = FullResImage;
-	PassParameters->HalfResSrcTexture = HalfResConvolvedImage;
-	PassParameters->HalfResSrcSampler = TStaticSamplerState<SF_Bilinear, AM_Wrap, AM_Wrap, AM_Wrap>::GetRHI();
-	PassParameters->CenterWeightTexture = CenterWeightTexutre;
-
-	PassParameters->DstTexture = GraphBuilder.CreateUAV(DstTexture);
-
-	TShaderMapRef<FBlendLowResCS> ComputeShader(View.ShaderMap);
-
-	const FIntPoint TargetSize = FullResImageRect.Size();
-
-	// Use multiple threads per scan line to insure memory coalescing during the write
-	const int32 ThreadsPerGroup = ComputeShader->ThreadsPerGroup;
-	const int32 ThreadsGroupsPerScanLine = (TargetSize.X % ThreadsPerGroup == 0) ? TargetSize.X / ThreadsPerGroup : TargetSize.X / ThreadsPerGroup + 1;
-
-	FComputeShaderUtils::AddPass(
-		GraphBuilder,
-		RDG_EVENT_NAME("FFTBloom Upscale&Blend %dx%d -> %dx%d",
-			HalfResRect.Width(),
-			HalfResRect.Height(),
-			FullResImageRect.Width(),
-			FullResImageRect.Height()),
-		ComputeShader,
-		PassParameters,
-		FIntVector(ThreadsGroupsPerScanLine, TargetSize.Y, 1));
-}
-
 FRDGTextureRef TransformKernelFFT(
 	FRDGBuilder& GraphBuilder,
 	const FViewInfo& View,
@@ -457,13 +371,175 @@ FFFTBloomIntermediates GetFFTBloomIntermediates(
 	return Intermediates;
 }
 
+
+struct FKernelAnalysisResult
+{
+	FRDGBufferRef MaxScatterDispersionBuffer;
+	FRDGBufferRef KernelCenterEnergyBuffer;
+	FRDGTextureRef ScatterDispersionTexture;
+};
+
+FKernelAnalysisResult AnalysesKernelAtDensity(
+	FRDGBuilder& GraphBuilder,
+	const FGlobalShaderMap* ShaderMap,
+	const FFFTBloomIntermediates& Intermediates,
+	FRDGTextureRef SpatialKernelTexture,
+	FRDGBufferRef KernelCenterCoordBuffer,
+	float ViewTexelDiameterInKernelTexels)
+{
+	// Find information about the kernel's energy distribution
+	FRDGBufferRef MaxScatterDispersionBuffer;
+	FRDGBufferRef KernelCenterEnergyBuffer;
+	{
+		const int32 SurveyTileSize = 8;
+
+		float ViewTexelRadiusInKernelTexels = ViewTexelDiameterInKernelTexels * 0.5f;
+
+		int32 SurveyGroupGridSize = 2 * FMath::DivideAndRoundUp(FMath::CeilToInt(ViewTexelRadiusInKernelTexels) + 4, SurveyTileSize);
+		int32 SurveyGroupCount = SurveyGroupGridSize * SurveyGroupGridSize;
+
+		FRDGTextureUAVRef DebugTextureUAV;
+		{
+			FRDGTextureDesc DebugDesc = FRDGTextureDesc::Create2D(
+				FIntPoint(SurveyGroupGridSize * 8, SurveyGroupGridSize * 8),
+				PF_FloatRGBA,
+				FClearValueBinding::None,
+				/* InFlags = */ TexCreate_ShaderResource | TexCreate_UAV);
+
+			FRDGTextureRef DebugTexture = GraphBuilder.CreateTexture(DebugDesc, TEXT("Debug.Bloom.Survey"));
+
+			DebugTextureUAV = GraphBuilder.CreateUAV(DebugTexture);
+		}
+
+		RDG_EVENT_SCOPE(GraphBuilder, "FFTBloom SurveyKernel(TexelDiameter=%f)", ViewTexelDiameterInKernelTexels);
+
+		// Reduce a survey buffer to a single value.
+		auto ReduceSurveyBuffer = [&](FRDGBufferRef SurveyBuffer, int32 Op)
+		{
+			FBloomReduceKernelSurveyCS::FParameters* PassParameters = GraphBuilder.AllocParameters<FBloomReduceKernelSurveyCS::FParameters>();
+			PassParameters->SurveyReduceOp = Op;
+			PassParameters->SurveyGroupCount = SurveyGroupCount;
+			PassParameters->SurveyOutput = GraphBuilder.CreateUAV(SurveyBuffer);
+
+			TShaderMapRef<FBloomReduceKernelSurveyCS> ComputeShader(ShaderMap);
+			FComputeShaderUtils::AddPass(
+				GraphBuilder,
+				RDG_EVENT_NAME("FFTBloom ReduceKernelSurvey(Op=%d) %d", Op, SurveyGroupCount),
+				ComputeShader,
+				PassParameters,
+				FComputeShaderUtils::GetGroupCount(SurveyGroupCount, 64));
+		};
+
+		// Find the max scatter dispersion to use around the footprint of the view pixel in the kernel.
+		{
+			MaxScatterDispersionBuffer = GraphBuilder.CreateBuffer(
+				FRDGBufferDesc::CreateStructuredDesc(sizeof(FLinearColor), /* NumElements = */ SurveyGroupCount),
+				TEXT("Bloom.FFT.MaxScatterDispersion"));
+
+			FBloomSurveyMaxScatterDispersionCS::FParameters* PassParameters = GraphBuilder.AllocParameters<FBloomSurveyMaxScatterDispersionCS::FParameters>();
+			PassParameters->ViewTexelRadiusInKernelTexels = ViewTexelRadiusInKernelTexels;
+			PassParameters->SurveyGroupGridSize = SurveyGroupGridSize;
+			PassParameters->KernelSpatialTextureSize = SpatialKernelTexture->Desc.Extent;
+			PassParameters->KernelSpatialTexture = SpatialKernelTexture;
+			PassParameters->KernelCenterCoordBuffer = GraphBuilder.CreateSRV(KernelCenterCoordBuffer);
+			PassParameters->SurveyOutput = GraphBuilder.CreateUAV(MaxScatterDispersionBuffer);
+			PassParameters->DebugOutput = DebugTextureUAV;
+
+			TShaderMapRef<FBloomSurveyMaxScatterDispersionCS> ComputeShader(ShaderMap);
+			FComputeShaderUtils::AddPass(
+				GraphBuilder,
+				RDG_EVENT_NAME("FFTBloom SurveyMaxScatterDispersion"),
+				ComputeShader,
+				PassParameters,
+				FIntVector(SurveyGroupGridSize, SurveyGroupGridSize, 1));
+
+			ReduceSurveyBuffer(MaxScatterDispersionBuffer, /* Op = */ 0);
+		}
+
+		// Find the amount of energy at the center in the footprint of the view pixel in the kernel.
+		{
+			KernelCenterEnergyBuffer = GraphBuilder.CreateBuffer(
+				FRDGBufferDesc::CreateStructuredDesc(sizeof(FLinearColor), /* NumElements = */ SurveyGroupCount),
+				TEXT("Bloom.FFT.KernelCenterEnergy"));
+
+			FBloomSurveyKernelCenterEnergyCS::FParameters* PassParameters = GraphBuilder.AllocParameters<FBloomSurveyKernelCenterEnergyCS::FParameters>();
+			PassParameters->ViewTexelRadiusInKernelTexels = ViewTexelRadiusInKernelTexels;
+			PassParameters->SurveyGroupGridSize = SurveyGroupGridSize;
+			PassParameters->KernelSpatialTextureSize = SpatialKernelTexture->Desc.Extent;
+			PassParameters->KernelSpatialTexture = SpatialKernelTexture;
+			PassParameters->KernelCenterCoordBuffer = GraphBuilder.CreateSRV(KernelCenterCoordBuffer);
+			PassParameters->MaxScatterDispersionBuffer = GraphBuilder.CreateSRV(MaxScatterDispersionBuffer);
+			PassParameters->SurveyOutput = GraphBuilder.CreateUAV(KernelCenterEnergyBuffer);
+			PassParameters->DebugOutput = DebugTextureUAV;
+
+			TShaderMapRef<FBloomSurveyKernelCenterEnergyCS> ComputeShader(ShaderMap);
+			FComputeShaderUtils::AddPass(
+				GraphBuilder,
+				RDG_EVENT_NAME("FFTBloom SurveyKernelCenterEnergy"),
+				ComputeShader,
+				PassParameters,
+				FIntVector(SurveyGroupGridSize, SurveyGroupGridSize, 1));
+
+			ReduceSurveyBuffer(KernelCenterEnergyBuffer, /* Op = */ 1);
+		}
+	}
+
+	// Find out the total energy of the kernel - center.
+	// Implemented in such away to prioritize numerical error rather than speed.
+	FRDGTextureRef ScatterDispersionTexture;
+	{
+		RDG_EVENT_SCOPE(GraphBuilder, "FFTBloom SumScatterDispersionEnergy %dx%d",
+			SpatialKernelTexture->Desc.Extent.X, SpatialKernelTexture->Desc.Extent.Y);
+
+		ScatterDispersionTexture = SpatialKernelTexture;
+
+		for (int32 PassId = 0; ScatterDispersionTexture->Desc.Extent.X > 1 && ScatterDispersionTexture->Desc.Extent.Y > 1; PassId++)
+		{
+			FRDGTextureRef NewScatterDispersionTexture;
+			{
+				FRDGTextureDesc Desc = FRDGTextureDesc::Create2D(
+					FIntPoint::DivideAndRoundUp(ScatterDispersionTexture->Desc.Extent, 8),
+					PF_A32B32G32R32F,
+					FClearValueBinding::None,
+					TexCreate_ShaderResource | TexCreate_UAV);
+
+				NewScatterDispersionTexture = GraphBuilder.CreateTexture(Desc, TEXT("Bloom.FFT.KernelIntensity"));
+			}
+
+			FBloomSumScatterDispersionEnergyCS::FParameters* PassParameters = GraphBuilder.AllocParameters<FBloomSumScatterDispersionEnergyCS::FParameters>();
+			PassParameters->PassId = PassId;
+			PassParameters->ScatterDispersionTextureSize = ScatterDispersionTexture->Desc.Extent;
+			PassParameters->ScatterDispersionTexture = ScatterDispersionTexture;
+			PassParameters->MaxScatterDispersionBuffer = GraphBuilder.CreateSRV(MaxScatterDispersionBuffer);
+			PassParameters->ScatterDispersionOutput = GraphBuilder.CreateUAV(NewScatterDispersionTexture);
+
+			TShaderMapRef<FBloomSumScatterDispersionEnergyCS> ComputeShader(ShaderMap);
+			FComputeShaderUtils::AddPass(
+				GraphBuilder,
+				RDG_EVENT_NAME("FFTBloom SumScatterDispersionEnergy %dx%d -> %dx%d",
+					ScatterDispersionTexture->Desc.Extent.X, ScatterDispersionTexture->Desc.Extent.Y,
+					NewScatterDispersionTexture->Desc.Extent.X, NewScatterDispersionTexture->Desc.Extent.Y),
+				ComputeShader,
+				PassParameters,
+				FIntVector(NewScatterDispersionTexture->Desc.Extent.X, NewScatterDispersionTexture->Desc.Extent.Y, 1));
+
+			ScatterDispersionTexture = NewScatterDispersionTexture;
+		}
+	}
+
+	FKernelAnalysisResult Result;
+	Result.MaxScatterDispersionBuffer = MaxScatterDispersionBuffer;
+	Result.KernelCenterEnergyBuffer = KernelCenterEnergyBuffer;
+	Result.ScatterDispersionTexture = ScatterDispersionTexture;
+	return Result;
+}
+
 void InitDomainAndGetKernel(
 	FRDGBuilder& GraphBuilder,
 	const FViewInfo& View,
 	const FFFTBloomIntermediates& Intermediates,
-	bool bForceCenterZero,
-	FRDGTextureRef* OutCachedSpectralKernel,
-	FRDGTextureRef* OutCenterWeightTexture)
+	FRDGTextureRef* OutSpectralKernelTexture,
+	FRDGBufferRef* OutKernelConstantsBuffer)
 {
 	FSceneViewState* ViewState = View.ViewState;
 
@@ -476,7 +552,7 @@ void InitDomainAndGetKernel(
 	check(BloomConvolutionTextureResource && PhysicalSpaceKernelTextureRef);
 
 	const float BloomConvolutionSize = PPSettings.BloomConvolutionSize;
-	const FVector2D CenterUV = PPSettings.BloomConvolutionCenterUV;
+	const FVector2D CenterUV = PPSettings.BloomConvolutionCenterUV; // TODO: remove
 
 	// Our frequency storage layout adds two elements to the first transform direction. 
 	const FIntPoint FrequencyPadding = (Intermediates.bDoHorizontalFirst) ? FIntPoint(2, 0) : FIntPoint(0, 2);
@@ -490,9 +566,7 @@ void InitDomainAndGetKernel(
 		FClearValueBinding::None,
 		TexCreate_ShaderResource | TexCreate_UAV);
 
-	FRDGTextureRef CachedSpectralKernel = nullptr;
-	FRDGTextureRef CenterWeightTexture = nullptr;
-	if (ViewState && ViewState->BloomFFTKernel.Spectral)
+	if (ViewState && ViewState->BloomFFTKernel.Spectral && CVarBloomCacheKernel.GetValueOnRenderThread() != 0)
 	{
 		auto& FFTKernel = ViewState->BloomFFTKernel;
 
@@ -506,77 +580,306 @@ void InitDomainAndGetKernel(
 
 		const bool bSameKernelSize = FMath::IsNearlyEqual(FFTKernel.Scale, BloomConvolutionSize, float(1.e-6) /*tol*/);
 		const bool bSameImageSize = (Intermediates.ImageRect.Size() == FFTKernel.ImageSize);
-		const bool bSameKernelCenterUV = FFTKernel.CenterUV.Equals(CenterUV, float(1.e-6) /*tol*/);
 		const bool bSameMipLevel = bSameTexture && FFTKernel.PhysicalMipLevel == BloomConvolutionTextureResource->GetCurrentMipCount();
 
-		if (bSameTexture && bSameSpectralBuffer && bSameKernelSize && bSameImageSize && bSameKernelCenterUV && bSameMipLevel)
+		if (bSameTexture && bSameSpectralBuffer && bSameKernelSize && bSameImageSize && bSameMipLevel)
 		{
-			CachedSpectralKernel = PrevCachedSpectralKernel;
-
-			if (bForceCenterZero)
-			{
-				CenterWeightTexture = GraphBuilder.RegisterExternalTexture(FFTKernel.CenterWeight);
-			}
+			*OutSpectralKernelTexture = PrevCachedSpectralKernel;
+			*OutKernelConstantsBuffer = GraphBuilder.RegisterExternalBuffer(FFTKernel.ConstantsBuffer);
+			return;
 		}
 	}
 
 	// Re-transform the kernel if needed.
-	if (!CachedSpectralKernel)
+	RDG_EVENT_SCOPE(GraphBuilder, "InitBloomKernel");
+
+
+	FRDGTextureRef SpatialKernelTexture = RegisterExternalTexture(GraphBuilder, PhysicalSpaceKernelTextureRef, TEXT("Bloom.FFT.OriginalKernel"));
+
+
+	// Find the center of the kernel
+	FRDGBufferRef KernelCenterCoordBuffer;
 	{
-		RDG_EVENT_SCOPE(GraphBuilder, "InitDomainAndGetKernel");
+		KernelCenterCoordBuffer = GraphBuilder.CreateBuffer(
+			FRDGBufferDesc::CreateStructuredDesc(sizeof(int32_t), /* NumElements = */ 4),
+			TEXT("Bloom.FFT.KernelCenterCoord"));
 
-		// Rescale the physical space kernel ( and omit the center if this is a 1/2 resolution fft, it will be added later)
-		FRDGTextureRef ResizedKernel = GraphBuilder.CreateTexture(TransformDesc, TEXT("Bloom.FFT.ResizedKernel"));
-		ResizeAndCenterTexture(
+		FBloomFindKernelCenterCS::FParameters* PassParameters = GraphBuilder.AllocParameters<FBloomFindKernelCenterCS::FParameters>();
+		PassParameters->KernelSpatialTextureSize = SpatialKernelTexture->Desc.Extent;
+		PassParameters->KernelSpatialTexture = SpatialKernelTexture;
+		PassParameters->KernelCenterCoordOutput = GraphBuilder.CreateUAV(KernelCenterCoordBuffer);
+
+		AddClearUAVPass(GraphBuilder, PassParameters->KernelCenterCoordOutput, /* Value = */ 0);
+
+		TShaderMapRef<FBloomFindKernelCenterCS> ComputeShader(View.ShaderMap);
+		FComputeShaderUtils::AddPass(
 			GraphBuilder,
-			View, PhysicalSpaceKernelTextureRef, Intermediates.ImageRect.Size(), CenterUV, Intermediates.KernelSupportScale,
-			Intermediates.FrequencySize, ResizedKernel, PaddedFrequencySize, bForceCenterZero);
+			RDG_EVENT_NAME("FFTBloom FindKernelCenter %dx%d", SpatialKernelTexture->Desc.Extent.X, SpatialKernelTexture->Desc.Extent.Y),
+			ComputeShader,
+			PassParameters, 
+			FComputeShaderUtils::GetGroupCount(SpatialKernelTexture->Desc.Extent, 8));
+	}
 
-		// Two Dimensional FFT of the physical space kernel.  
-		CachedSpectralKernel = TransformKernelFFT(
-			GraphBuilder, View,
-			ResizedKernel,
-			Intermediates.bDoHorizontalFirst,
-			Intermediates.FrequencySize);
+	// Find information about the kernel's energy distribution
+	FRDGBufferRef MaxScatterDispersionBuffer;
+	FRDGBufferRef KernelCenterEnergyBuffer;
+	{
+		const int32 SurveyTileSize = 8;
+
+		float KernelSizeInDstPixels = FMath::Max(float(Intermediates.ImageRect.Width()) * Intermediates.KernelSupportScale, 1.0f);
+
+		// Diameter of a view texel in the kernel.
+		float ViewTexelDiameterInKernelTexels = FMath::Max(SpatialKernelTexture->Desc.Extent.X / KernelSizeInDstPixels, 1.0f);
+
+		float ViewTexelRadiusInKernelTexels = ViewTexelDiameterInKernelTexels * 0.5f;
+
+		int32 SurveyGroupGridSize = 2 * FMath::DivideAndRoundUp(FMath::CeilToInt(ViewTexelRadiusInKernelTexels) + 4, SurveyTileSize);
+		int32 SurveyGroupCount = SurveyGroupGridSize * SurveyGroupGridSize;
+
+		FRDGTextureUAVRef DebugTextureUAV;
+		{
+			FRDGTextureDesc DebugDesc = FRDGTextureDesc::Create2D(
+				FIntPoint(SurveyGroupGridSize * 8, SurveyGroupGridSize * 8),
+				PF_FloatRGBA,
+				FClearValueBinding::None,
+				/* InFlags = */ TexCreate_ShaderResource | TexCreate_UAV);
+
+			FRDGTextureRef DebugTexture = GraphBuilder.CreateTexture(DebugDesc, TEXT("Debug.Bloom.Survey"));
+
+			DebugTextureUAV = GraphBuilder.CreateUAV(DebugTexture);
+		}
+
+		RDG_EVENT_SCOPE(GraphBuilder, "FFTBloom SurveyKernel(TexelDiameter=%f)", ViewTexelDiameterInKernelTexels);
+
+		// Reduce a survey buffer to a single value.
+		auto ReduceSurveyBuffer = [&](FRDGBufferRef SurveyBuffer, int32 Op)
+		{
+			FBloomReduceKernelSurveyCS::FParameters* PassParameters = GraphBuilder.AllocParameters<FBloomReduceKernelSurveyCS::FParameters>();
+			PassParameters->SurveyReduceOp = Op;
+			PassParameters->SurveyGroupCount = SurveyGroupCount;
+			PassParameters->SurveyOutput = GraphBuilder.CreateUAV(SurveyBuffer);
 
-		if (bForceCenterZero)
+			TShaderMapRef<FBloomReduceKernelSurveyCS> ComputeShader(View.ShaderMap);
+			FComputeShaderUtils::AddPass(
+				GraphBuilder,
+				RDG_EVENT_NAME("FFTBloom ReduceKernelSurvey(Op=%d) %d", Op, SurveyGroupCount),
+				ComputeShader,
+				PassParameters,
+				FComputeShaderUtils::GetGroupCount(SurveyGroupCount, 64));
+		};
+
+		// Find the max scatter dispersion to use around the footprint of the view pixel in the kernel.
+		{
+			MaxScatterDispersionBuffer = GraphBuilder.CreateBuffer(
+				FRDGBufferDesc::CreateStructuredDesc(sizeof(FLinearColor), /* NumElements = */ SurveyGroupCount),
+				TEXT("Bloom.FFT.MaxScatterDispersion"));
+
+			FBloomSurveyMaxScatterDispersionCS::FParameters* PassParameters = GraphBuilder.AllocParameters<FBloomSurveyMaxScatterDispersionCS::FParameters>();
+			PassParameters->ViewTexelRadiusInKernelTexels = ViewTexelRadiusInKernelTexels;
+			PassParameters->SurveyGroupGridSize = SurveyGroupGridSize;
+			PassParameters->KernelSpatialTextureSize = SpatialKernelTexture->Desc.Extent;
+			PassParameters->KernelSpatialTexture = SpatialKernelTexture;
+			PassParameters->KernelCenterCoordBuffer = GraphBuilder.CreateSRV(KernelCenterCoordBuffer);
+			PassParameters->SurveyOutput = GraphBuilder.CreateUAV(MaxScatterDispersionBuffer);
+			PassParameters->DebugOutput = DebugTextureUAV;
+
+			TShaderMapRef<FBloomSurveyMaxScatterDispersionCS> ComputeShader(View.ShaderMap);
+			FComputeShaderUtils::AddPass(
+				GraphBuilder,
+				RDG_EVENT_NAME("FFTBloom SurveyMaxScatterDispersion"),
+				ComputeShader,
+				PassParameters,
+				FIntVector(SurveyGroupGridSize, SurveyGroupGridSize, 1));
+
+			ReduceSurveyBuffer(MaxScatterDispersionBuffer, /* Op = */ 0);
+		}
+
+		// Find the amount of energy at the center in the footprint of the view pixel in the kernel.
 		{
-			// Capture the missing center weight from the kernel and the sum of the existing weights.
-			CenterWeightTexture = CaptureKernelWeight(
+			KernelCenterEnergyBuffer = GraphBuilder.CreateBuffer(
+				FRDGBufferDesc::CreateStructuredDesc(sizeof(FLinearColor), /* NumElements = */ SurveyGroupCount),
+				TEXT("Bloom.FFT.KernelCenterEnergy"));
+
+			FBloomSurveyKernelCenterEnergyCS::FParameters* PassParameters = GraphBuilder.AllocParameters<FBloomSurveyKernelCenterEnergyCS::FParameters>();
+			PassParameters->ViewTexelRadiusInKernelTexels = ViewTexelRadiusInKernelTexels;
+			PassParameters->SurveyGroupGridSize = SurveyGroupGridSize;
+			PassParameters->KernelSpatialTextureSize = SpatialKernelTexture->Desc.Extent;
+			PassParameters->KernelSpatialTexture = SpatialKernelTexture;
+			PassParameters->KernelCenterCoordBuffer = GraphBuilder.CreateSRV(KernelCenterCoordBuffer);
+			PassParameters->MaxScatterDispersionBuffer = GraphBuilder.CreateSRV(MaxScatterDispersionBuffer);
+			PassParameters->SurveyOutput = GraphBuilder.CreateUAV(KernelCenterEnergyBuffer);
+			PassParameters->DebugOutput = DebugTextureUAV;
+
+			TShaderMapRef<FBloomSurveyKernelCenterEnergyCS> ComputeShader(View.ShaderMap);
+			FComputeShaderUtils::AddPass(
 				GraphBuilder,
-				View,
-				CachedSpectralKernel,
-				PaddedFrequencySize,
-				PhysicalSpaceKernelTextureRef,
-				CenterUV);
+				RDG_EVENT_NAME("FFTBloom SurveyKernelCenterEnergy"),
+				ComputeShader,
+				PassParameters,
+				FIntVector(SurveyGroupGridSize, SurveyGroupGridSize, 1));
+
+			ReduceSurveyBuffer(KernelCenterEnergyBuffer, /* Op = */ 1);
 		}
+	}
+
+	// Find out the total energy of the kernel - center.
+	// Implemented in such away to prioritize numerical error rather than speed.
+	FRDGTextureRef ScatterDispersionTexture;
+	{
+		RDG_EVENT_SCOPE(GraphBuilder, "FFTBloom SumScatterDispersionEnergy %dx%d",
+			SpatialKernelTexture->Desc.Extent.X, SpatialKernelTexture->Desc.Extent.Y);
 
-		// Update the data on the ViewState
-		if (ViewState)
+		ScatterDispersionTexture = SpatialKernelTexture;
+
+		for (int32 PassId = 0; ScatterDispersionTexture->Desc.Extent.X > 1 && ScatterDispersionTexture->Desc.Extent.Y > 1; PassId++)
 		{
-			ViewState->BloomFFTKernel.Scale = BloomConvolutionSize;
-			ViewState->BloomFFTKernel.ImageSize = Intermediates.ImageRect.Size();
-			ViewState->BloomFFTKernel.Physical = View.FinalPostProcessSettings.BloomConvolutionTexture;
-			ViewState->BloomFFTKernel.PhysicalRHI = PhysicalSpaceKernelTextureRef;
-			ViewState->BloomFFTKernel.CenterUV = CenterUV;
-			ViewState->BloomFFTKernel.PhysicalMipLevel = BloomConvolutionTextureResource->GetCurrentMipCount();
+			FRDGTextureRef NewScatterDispersionTexture;
 			{
-				ViewState->BloomFFTKernel.Spectral.SafeRelease();
-				GraphBuilder.QueueTextureExtraction(CachedSpectralKernel, &ViewState->BloomFFTKernel.Spectral);
+				FRDGTextureDesc Desc = FRDGTextureDesc::Create2D(
+					FIntPoint::DivideAndRoundUp(ScatterDispersionTexture->Desc.Extent, 8),
+					PF_A32B32G32R32F,
+					FClearValueBinding::None,
+					TexCreate_ShaderResource | TexCreate_UAV);
+
+				NewScatterDispersionTexture = GraphBuilder.CreateTexture(Desc, TEXT("Bloom.FFT.KernelIntensity"));
 			}
+
+			FBloomSumScatterDispersionEnergyCS::FParameters* PassParameters = GraphBuilder.AllocParameters<FBloomSumScatterDispersionEnergyCS::FParameters>();
+			PassParameters->PassId = PassId;
+			PassParameters->ScatterDispersionTextureSize = ScatterDispersionTexture->Desc.Extent;
+			PassParameters->ScatterDispersionTexture = ScatterDispersionTexture;
+			PassParameters->MaxScatterDispersionBuffer = GraphBuilder.CreateSRV(MaxScatterDispersionBuffer);
+			PassParameters->ScatterDispersionOutput = GraphBuilder.CreateUAV(NewScatterDispersionTexture);
+
+			TShaderMapRef<FBloomSumScatterDispersionEnergyCS> ComputeShader(View.ShaderMap);
+			FComputeShaderUtils::AddPass(
+				GraphBuilder,
+				RDG_EVENT_NAME("FFTBloom SumScatterDispersionEnergy %dx%d -> %dx%d",
+					ScatterDispersionTexture->Desc.Extent.X, ScatterDispersionTexture->Desc.Extent.Y,
+					NewScatterDispersionTexture->Desc.Extent.X, NewScatterDispersionTexture->Desc.Extent.Y),
+				ComputeShader,
+				PassParameters,
+				FIntVector(NewScatterDispersionTexture->Desc.Extent.X, NewScatterDispersionTexture->Desc.Extent.Y, 1));
+
+			ScatterDispersionTexture = NewScatterDispersionTexture;
+		}
+	}
+
+	// Packs all the kernel information into the same buffer.
+	FRDGBufferRef KernelConstantsBuffer = nullptr;
+	{
+		KernelConstantsBuffer = GraphBuilder.CreateBuffer(
+			FRDGBufferDesc::CreateStructuredDesc(sizeof(float) * 16, /* NumElements = */ 1),
+			TEXT("Bloom.FFT.KernelConstants"));
+
+		FBloomPackKernelConstantsCS::FParameters* PassParameters = GraphBuilder.AllocParameters<FBloomPackKernelConstantsCS::FParameters>();
+		PassParameters->KernelCenterCoordBuffer = GraphBuilder.CreateSRV(KernelCenterCoordBuffer);
+		PassParameters->KernelCenterEnergyBuffer   = GraphBuilder.CreateSRV(KernelCenterEnergyBuffer);
+		PassParameters->MaxScatterDispersionBuffer = GraphBuilder.CreateSRV(MaxScatterDispersionBuffer);
+		PassParameters->ScatterDispersionTexture   = ScatterDispersionTexture;
+		PassParameters->KernelConstantsOutput = GraphBuilder.CreateUAV(KernelConstantsBuffer);
+
+		TShaderMapRef<FBloomPackKernelConstantsCS> ComputeShader(View.ShaderMap);
+		FComputeShaderUtils::AddPass(
+			GraphBuilder,
+			RDG_EVENT_NAME("FFTBloom PackKernelConstants"),
+			ComputeShader,
+			PassParameters,
+			FIntVector(1, 1, 1));
+	}
+
+	// Preprocess the original kernel for Fourier transformation
+	FRDGTextureRef ResizedKernel;
+	{
+		RDG_EVENT_SCOPE(GraphBuilder, "FFTBloom PreprocessKernel");
+
+		// Clamp the kernel to avoid highlight contamination in the resize.
+		FRDGTextureRef ClampedKernelTexture;
+		{
 			{
-				ViewState->BloomFFTKernel.CenterWeight.SafeRelease();
-				if (CenterWeightTexture)
-					GraphBuilder.QueueTextureExtraction(CenterWeightTexture, &ViewState->BloomFFTKernel.CenterWeight);
+				FRDGTextureDesc Desc = FRDGTextureDesc::Create2D(
+					SpatialKernelTexture->Desc.Extent,
+					PF_FloatRGBA,
+					FClearValueBinding::None,
+					TexCreate_ShaderResource | TexCreate_UAV);
+
+				ClampedKernelTexture = GraphBuilder.CreateTexture(Desc, TEXT("Bloom.FFT.ClampedKernel"));
 			}
+
+			FBloomClampKernelCS::FParameters* PassParameters = GraphBuilder.AllocParameters<FBloomClampKernelCS::FParameters>();
+			PassParameters->KernelSpatialTexture = SpatialKernelTexture;
+			PassParameters->KernelConstantsBuffer = GraphBuilder.CreateSRV(KernelConstantsBuffer);
+			PassParameters->ClampedKernelSpatialOutput = GraphBuilder.CreateUAV(ClampedKernelTexture);
+
+			TShaderMapRef<FBloomClampKernelCS> ComputeShader(View.ShaderMap);
+			FComputeShaderUtils::AddPass(
+				GraphBuilder,
+				RDG_EVENT_NAME("FFTBloom ClampedKernel %dx%d",
+					SpatialKernelTexture->Desc.Extent.X,
+					SpatialKernelTexture->Desc.Extent.Y),
+				ComputeShader,
+				PassParameters,
+				FComputeShaderUtils::GetGroupCount(SpatialKernelTexture->Desc.Extent, 8));
+		}
+
+		// Final resize the kernel for Fourier transformation
+		{
+			ResizedKernel = GraphBuilder.CreateTexture(TransformDesc, TEXT("Bloom.FFT.ResizedKernel"));
+
+			FBloomResizeKernelCS::FParameters* PassParameters = GraphBuilder.AllocParameters<FBloomResizeKernelCS::FParameters>();
+			PassParameters->DstExtent = Intermediates.FrequencySize;
+			PassParameters->ImageExtent = Intermediates.ImageRect.Size();
+			PassParameters->KernelSpatialTextureInvSize.X = 1.0f / float(SpatialKernelTexture->Desc.Extent.X);
+			PassParameters->KernelSpatialTextureInvSize.Y = 1.0f / float(SpatialKernelTexture->Desc.Extent.Y);
+			PassParameters->DstBufferExtent = Intermediates.FrequencySize;
+			PassParameters->KernelSupportScale = Intermediates.KernelSupportScale;
+
+			PassParameters->KernelConstantsBuffer = GraphBuilder.CreateSRV(KernelConstantsBuffer);
+			PassParameters->SrcTexture = ClampedKernelTexture;
+			PassParameters->SrcSampler = TStaticSamplerState<SF_Bilinear, AM_Wrap, AM_Wrap, AM_Wrap>::GetRHI();
+		
+			PassParameters->DstTexture = GraphBuilder.CreateUAV(ResizedKernel);
+
+			TShaderMapRef<FBloomResizeKernelCS> ComputeShader(View.ShaderMap);
+			FComputeShaderUtils::AddPass(
+				GraphBuilder,
+				RDG_EVENT_NAME("FFTBloom PreProcessKernel %dx%d", Intermediates.FrequencySize.X, Intermediates.FrequencySize.Y),
+				ComputeShader,
+				PassParameters,
+				FComputeShaderUtils::GetGroupCount(PaddedFrequencySize, 8));
 		}
 	}
 
-	*OutCachedSpectralKernel = CachedSpectralKernel;
-	*OutCenterWeightTexture = CenterWeightTexture;
+	// Two Dimensional FFT of the physical space kernel.  
+	FRDGTextureRef SpectralKernelTexture = TransformKernelFFT(
+		GraphBuilder, View,
+		ResizedKernel,
+		Intermediates.bDoHorizontalFirst,
+		Intermediates.FrequencySize);
+
+	// Update the data on the ViewState
+	if (ViewState && CVarBloomCacheKernel.GetValueOnRenderThread() != 0)
+	{
+		ViewState->BloomFFTKernel.Scale = BloomConvolutionSize;
+		ViewState->BloomFFTKernel.ImageSize = Intermediates.ImageRect.Size();
+		ViewState->BloomFFTKernel.Physical = View.FinalPostProcessSettings.BloomConvolutionTexture;
+		ViewState->BloomFFTKernel.PhysicalRHI = PhysicalSpaceKernelTextureRef;
+		ViewState->BloomFFTKernel.PhysicalMipLevel = BloomConvolutionTextureResource->GetCurrentMipCount();
+		{
+			ViewState->BloomFFTKernel.Spectral.SafeRelease();
+			GraphBuilder.QueueTextureExtraction(SpectralKernelTexture, &ViewState->BloomFFTKernel.Spectral);
+		}
+		{
+			ViewState->BloomFFTKernel.ConstantsBuffer.SafeRelease();
+			GraphBuilder.QueueBufferExtraction(KernelConstantsBuffer, &ViewState->BloomFFTKernel.ConstantsBuffer);
+		}
+	}
+
+	*OutSpectralKernelTexture = SpectralKernelTexture;
+	*OutKernelConstantsBuffer = KernelConstantsBuffer;
 }
 
-FRDGTextureRef AddFFTBloomPass(FRDGBuilder& GraphBuilder, const FViewInfo& View, const FFFTBloomInputs& Inputs)
+FBloomOutputs AddFFTBloomPass(FRDGBuilder& GraphBuilder, const FViewInfo& View, const FFFTBloomInputs& Inputs)
 {
 	check(Inputs.FullResolutionTexture);
 	check(!Inputs.FullResolutionViewRect.IsEmpty());
@@ -586,93 +889,59 @@ FRDGTextureRef AddFFTBloomPass(FRDGBuilder& GraphBuilder, const FViewInfo& View,
 	const bool bHalfResolutionFFT = IsFFTBloomHalfResolutionEnabled();
 
 	FFFTBloomIntermediates Intermediates = GetFFTBloomIntermediates(View, Inputs);
-	
+
+	RDG_EVENT_SCOPE(GraphBuilder, "FFTBloom %dx%d", Intermediates.ImageRect.Width(), Intermediates.ImageRect.Height());
+
 	// Init the domain data update the cached kernel if needed.
 	FRDGTextureRef SpectralKernelTexture;
-	FRDGTextureRef CenterWeightTexture;
+	FRDGBufferRef KernelConstantsBuffer;
 	InitDomainAndGetKernel(
 		GraphBuilder, View, Intermediates,
-		/* bForceCenterZero = */ bHalfResolutionFFT,
 		/* out */ &SpectralKernelTexture,
-		/* out */ &CenterWeightTexture);
+		/* out */ &KernelConstantsBuffer);
 
-	const FLinearColor Tint(1, 1, 1, 1);
+	FBloomOutputs BloomOutput;
 
-	RDG_EVENT_SCOPE(GraphBuilder, "FFTBloom %dx%d", Intermediates.ImageRect.Width(), Intermediates.ImageRect.Height());
-
-	if (bHalfResolutionFFT)
+	// Generate the apply constant buffer for the tone-maping pass.
 	{
-		// Get a half-resolution destination buffer.
-		TRefCountPtr<IPooledRenderTarget> HalfResConvolutionResult;
+		check(FBloomOutputs::SupportsApplyParametersBuffer(View.GetShaderPlatform()));
 
-		FRDGTextureRef FFTConvolutionTexture;
-		{
-			const FRDGTextureDesc Desc = FRDGTextureDesc::Create2D(
-				Intermediates.InputBufferSize,
-				GPUFFT::PixelFormat(),
-				FClearValueBinding::None,
-				TexCreate_ShaderResource | TexCreate_UAV);
-
-			FFTConvolutionTexture = GraphBuilder.CreateTexture(Desc, TEXT("Bloom.FFT.SceneColor"));
-		}
-
-		GPUFFT::ConvolutionWithTextureImage2D(
-			GraphBuilder,
-			View.ShaderMap,
-			Intermediates.FrequencySize,
-			Intermediates.bDoHorizontalFirst,
-			SpectralKernelTexture,
-			Intermediates.InputTexture, Intermediates.ImageRect,
-			FFTConvolutionTexture, Intermediates.ImageRect,
-			Intermediates.PreFilter);
-
-		// Blend with  alpha * SrcBuffer + betta * BloomedBuffer  where alpha = Weights[0], beta = Weights[1]
-		const FIntPoint HalfResBufferSize = Intermediates.InputBufferSize;
-
-		FRDGTextureRef OutputSceneColorTexture;
-		{
-			const FRDGTextureDesc Desc = FRDGTextureDesc::Create2D(
-				Inputs.FullResolutionTexture->Desc.Extent,
-				Inputs.FullResolutionTexture->Desc.Format,
-				FClearValueBinding::None,
-				TexCreate_ShaderResource | TexCreate_UAV);
+		BloomOutput.ApplyParameters = GraphBuilder.CreateBuffer(
+			FRDGBufferDesc::CreateStructuredDesc(sizeof(FBloomOutputs::FApplyInfo), /* NumElements = */ 1),
+			TEXT("Bloom.FFT.KernelConstants"));
 
-			OutputSceneColorTexture = GraphBuilder.CreateTexture(Desc, TEXT("Bloom.FFT.SceneColor"));
-		}
+		FBloomFinalizeApplyConstantsCS::FParameters* PassParameters = GraphBuilder.AllocParameters<FBloomFinalizeApplyConstantsCS::FParameters>();
+		PassParameters->ScatterDispersionIntensity = View.FinalPostProcessSettings.BloomConvolutionScatterDispersion;
+		PassParameters->KernelConstantsBuffer = GraphBuilder.CreateSRV(KernelConstantsBuffer);
+		PassParameters->BloomApplyConstantsOutput = GraphBuilder.CreateUAV(BloomOutput.ApplyParameters);
 
-		BlendLowRes(
+		TShaderMapRef<FBloomFinalizeApplyConstantsCS> ComputeShader(View.ShaderMap);
+		FComputeShaderUtils::AddPass(
 			GraphBuilder,
-			View,
-			Inputs.FullResolutionTexture, Inputs.FullResolutionViewRect,
-			FFTConvolutionTexture, Intermediates.ImageRect,
-			Intermediates.InputBufferSize,
-			CenterWeightTexture,
-			OutputSceneColorTexture);
-
-		return OutputSceneColorTexture;
+			RDG_EVENT_NAME("FFTBloom FinalizeApplyConstants(ScatterDispersion=%f)", PassParameters->ScatterDispersionIntensity),
+			ComputeShader,
+			PassParameters,
+			FIntVector(1, 1, 1));
 	}
-	else
-	{
-		FRDGTextureDesc OutputSceneColorDesc = FRDGTextureDesc::Create2D(
-			Inputs.FullResolutionTexture->Desc.Extent,
-			Inputs.FullResolutionTexture->Desc.Format,
-			FClearValueBinding::None,
-			TexCreate_ShaderResource | TexCreate_UAV);
-		
-		FRDGTextureRef OutputSceneColorTexture = GraphBuilder.CreateTexture(OutputSceneColorDesc, TEXT("Bloom.FFT.SceneColor"));
 
-		GPUFFT::ConvolutionWithTextureImage2D(
-			GraphBuilder,
-			View.ShaderMap,
-			Intermediates.FrequencySize,
-			Intermediates.bDoHorizontalFirst,
-			SpectralKernelTexture,
-			Intermediates.InputTexture, Intermediates.ImageRect,
-			OutputSceneColorTexture, Intermediates.ImageRect,
-			Intermediates.PreFilter);
-
-		return OutputSceneColorTexture;
-	}
+	FRDGTextureDesc OutputSceneColorDesc = FRDGTextureDesc::Create2D(
+		Intermediates.InputTexture->Desc.Extent,
+		Intermediates.InputTexture->Desc.Format,
+		FClearValueBinding::None,
+		TexCreate_ShaderResource | TexCreate_UAV);
+		
+	BloomOutput.Bloom.Texture = GraphBuilder.CreateTexture(OutputSceneColorDesc, TEXT("Bloom.FFT.SceneColor"));
+	BloomOutput.Bloom.ViewRect = Intermediates.ImageRect;
 
-	return Inputs.FullResolutionTexture;
+	GPUFFT::ConvolutionWithTextureImage2D(
+		GraphBuilder,
+		View.ShaderMap,
+		Intermediates.FrequencySize,
+		Intermediates.bDoHorizontalFirst,
+		SpectralKernelTexture,
+		Intermediates.InputTexture, Intermediates.ImageRect,
+		BloomOutput.Bloom.Texture, BloomOutput.Bloom.ViewRect,
+		Intermediates.PreFilter);
+
+	return BloomOutput;
 }
diff --git a/Engine/Source/Runtime/Renderer/Private/PostProcess/PostProcessFFTBloom.h b/Engine/Source/Runtime/Renderer/Private/PostProcess/PostProcessFFTBloom.h
--- a/Engine/Source/Runtime/Renderer/Private/PostProcess/PostProcessFFTBloom.h
+++ b/Engine/Source/Runtime/Renderer/Private/PostProcess/PostProcessFFTBloom.h
@@ -5,6 +5,8 @@
 #include "RenderGraph.h"
 #include "ScenePrivate.h"
 
+struct FBloomOutputs;
+
 // Returns whether FFT bloom is enabled for the view.
 bool IsFFTBloomEnabled(const FViewInfo& View);
 
@@ -19,4 +21,4 @@ struct FFFTBloomInputs
 
 bool IsFFTBloomQuarterResolutionEnabled();
 
-FRDGTextureRef AddFFTBloomPass(FRDGBuilder& GraphBuilder, const FViewInfo& View, const FFFTBloomInputs& Inputs);
\ No newline at end of file
+FBloomOutputs AddFFTBloomPass(FRDGBuilder& GraphBuilder, const FViewInfo& View, const FFFTBloomInputs& Inputs);
diff --git a/Engine/Source/Runtime/Renderer/Private/PostProcess/PostProcessTonemap.cpp b/Engine/Source/Runtime/Renderer/Private/PostProcess/PostProcessTonemap.cpp
--- a/Engine/Source/Runtime/Renderer/Private/PostProcess/PostProcessTonemap.cpp
+++ b/Engine/Source/Runtime/Renderer/Private/PostProcess/PostProcessTonemap.cpp
@@ -329,21 +329,27 @@ BEGIN_SHADER_PARAMETER_STRUCT(FTonemapParameters, )
 	SHADER_PARAMETER_STRUCT_INCLUDE(FFilmGrainParameters, FilmGrain)
 	SHADER_PARAMETER_STRUCT_INCLUDE(FTonemapperOutputDeviceParameters, OutputDevice)
 	SHADER_PARAMETER_STRUCT(FScreenPassTextureViewportParameters, Color)
-	SHADER_PARAMETER_STRUCT(FScreenPassTextureViewportParameters, Bloom)
 	SHADER_PARAMETER_STRUCT(FScreenPassTextureViewportParameters, Output)
 	SHADER_PARAMETER_RDG_TEXTURE(Texture2D, ColorTexture)
+
+	// Bloom texture
+	SHADER_PARAMETER(FScreenTransform, ColorToBloom)
+	SHADER_PARAMETER(FVector2D, BloomUVViewportBilinearMin)
+	SHADER_PARAMETER(FVector2D, BloomUVViewportBilinearMax)
+	SHADER_PARAMETER(FVector4, BloomMultiply)
 	SHADER_PARAMETER_RDG_TEXTURE(Texture2D, BloomTexture)
+	SHADER_PARAMETER_SAMPLER(SamplerState, BloomSampler)
+	SHADER_PARAMETER_RDG_BUFFER_SRV(SamplerState, BloomApplyParameters)
+
 	// SM5 and above use Texture2D for EyeAdaptationTexture
 	SHADER_PARAMETER_RDG_TEXTURE(Texture2D, EyeAdaptationTexture)
 	SHADER_PARAMETER_RDG_TEXTURE(Texture2D, ColorGradingLUT)
 	SHADER_PARAMETER_TEXTURE(Texture2D, BloomDirtMaskTexture)
 	SHADER_PARAMETER_SAMPLER(SamplerState, ColorSampler)
-	SHADER_PARAMETER_SAMPLER(SamplerState, BloomSampler)
 	SHADER_PARAMETER_SAMPLER(SamplerState, ColorGradingLUTSampler)
 	SHADER_PARAMETER_SAMPLER(SamplerState, BloomDirtMaskSampler)
-	SHADER_PARAMETER(FScreenTransform, ColorToBloom)
+
 	SHADER_PARAMETER(FVector4, ColorScale0)
-	SHADER_PARAMETER(FVector4, ColorScale1)
 	SHADER_PARAMETER(FVector4, BloomDirtMaskTint)
 	SHADER_PARAMETER(FVector4, ChromaticAberrationParams)
 	SHADER_PARAMETER(FVector4, TonemapperParams)
@@ -408,6 +414,8 @@ public:
 	{
 		const int UseVolumeLut = PipelineVolumeTextureLUTSupportGuaranteedAtRuntime(Parameters.Platform) ? 1 : 0;
 		OutEnvironment.SetDefine(TEXT("USE_VOLUME_LUT"), UseVolumeLut);
+
+		OutEnvironment.SetDefine(TEXT("SUPPORTS_BLOOM_APPLY_PARAMETERS"), FBloomOutputs::SupportsApplyParametersBuffer(Parameters.Platform) ? 1 : 0);
 	}
 };
 
@@ -447,6 +455,8 @@ public:
 
 		const int UseVolumeLut = PipelineVolumeTextureLUTSupportGuaranteedAtRuntime(Parameters.Platform) ? 1 : 0;
 		OutEnvironment.SetDefine(TEXT("USE_VOLUME_LUT"), UseVolumeLut);
+
+		OutEnvironment.SetDefine(TEXT("SUPPORTS_BLOOM_APPLY_PARAMETERS"), FBloomOutputs::SupportsApplyParametersBuffer(Parameters.Platform) ? 1 : 0);
 	}
 };
 
@@ -611,24 +621,15 @@ FScreenPassTexture AddTonemapPass(FRDGBuilder& GraphBuilder, const FViewInfo& Vi
 
 	CommonParameters.OutputDevice = GetTonemapperOutputDeviceParameters(ViewFamily);
 	CommonParameters.Color = GetScreenPassTextureViewportParameters(SceneColorViewport);
-	if (Inputs.Bloom.Texture)
-	{
-		const FScreenPassTextureViewport BloomViewport(Inputs.Bloom);
-		CommonParameters.Bloom = GetScreenPassTextureViewportParameters(BloomViewport);
-		CommonParameters.ColorToBloom = FScreenTransform::ChangeTextureUVCoordinateFromTo(SceneColorViewport, BloomViewport);
-	}
 	CommonParameters.Output = GetScreenPassTextureViewportParameters(OutputViewport);
 	CommonParameters.ColorTexture = Inputs.SceneColor.Texture;
-	CommonParameters.BloomTexture = Inputs.Bloom.Texture;
 	CommonParameters.EyeAdaptationTexture = Inputs.EyeAdaptationTexture;
 	CommonParameters.ColorGradingLUT = Inputs.ColorGradingTexture;
 	CommonParameters.BloomDirtMaskTexture = BloomDirtMaskTexture;
 	CommonParameters.ColorSampler = BilinearClampSampler;
-	CommonParameters.BloomSampler = BilinearClampSampler;
 	CommonParameters.ColorGradingLUTSampler = BilinearClampSampler;
 	CommonParameters.BloomDirtMaskSampler = BilinearClampSampler;
 	CommonParameters.ColorScale0 = PostProcessSettings.SceneColorTint;
-	CommonParameters.ColorScale1 = FLinearColor::White * PostProcessSettings.BloomIntensity;
 	CommonParameters.BloomDirtMaskTint = PostProcessSettings.BloomDirtMaskTint * PostProcessSettings.BloomDirtMaskIntensity;
 	CommonParameters.ChromaticAberrationParams = ChromaticAberrationParams;
 	CommonParameters.TonemapperParams = FVector4(PostProcessSettings.VignetteIntensity, SharpenDiv6, 0.0f, 0.0f);
@@ -638,6 +639,54 @@ FScreenPassTexture AddTonemapPass(FRDGBuilder& GraphBuilder, const FViewInfo& Vi
 	CommonParameters.bOutputInHDR = ViewFamily.bIsHDR;
 	CommonParameters.LensPrincipalPointOffsetScale = View.LensPrincipalPointOffsetScale;
 
+	// Bloom parameters
+	{
+		if (Inputs.Bloom.Bloom.Texture)
+		{
+			const FScreenPassTextureViewport BloomViewport(Inputs.Bloom.Bloom);
+			CommonParameters.ColorToBloom = FScreenTransform::ChangeTextureUVCoordinateFromTo(SceneColorViewport, BloomViewport);
+			CommonParameters.BloomUVViewportBilinearMin = GetScreenPassTextureViewportParameters(BloomViewport).UVViewportBilinearMin;
+			CommonParameters.BloomUVViewportBilinearMax = GetScreenPassTextureViewportParameters(BloomViewport).UVViewportBilinearMax;
+			CommonParameters.BloomMultiply = FLinearColor::White * PostProcessSettings.BloomIntensity;
+			CommonParameters.BloomTexture = Inputs.Bloom.Bloom.Texture;
+			CommonParameters.BloomSampler = BilinearClampSampler;
+		}
+		else
+		{
+			CommonParameters.ColorToBloom = FScreenTransform::Identity;
+			CommonParameters.BloomUVViewportBilinearMin = FVector2D::ZeroVector;
+			CommonParameters.BloomUVViewportBilinearMax = FVector2D::UnitVector;
+			CommonParameters.BloomMultiply = FLinearColor::Transparent;
+			CommonParameters.BloomTexture = GSystemTextures.GetBlackDummy(GraphBuilder);
+			CommonParameters.BloomSampler = BilinearClampSampler;
+		}
+
+		if (!FBloomOutputs::SupportsApplyParametersBuffer(View.GetShaderPlatform()))
+		{
+			check(Inputs.Bloom.ApplyParameters == nullptr);
+		}
+		else if (Inputs.Bloom.ApplyParameters)
+		{
+			CommonParameters.BloomApplyParameters = GraphBuilder.CreateSRV(Inputs.Bloom.ApplyParameters);
+		}
+		else
+		{
+			FBloomOutputs::FApplyInfo ApplyParameters;
+			ApplyParameters.SceneColorMultiply = FLinearColor::White;
+			ApplyParameters.BloomMultiply = FLinearColor(CommonParameters.BloomMultiply);
+
+			TArray<FBloomOutputs::FApplyInfo, TInlineAllocator<1> > ApplyParametersArray;
+			ApplyParametersArray.Add(ApplyParameters);
+
+			FRDGBufferRef ApplyParametersBuffer = CreateStructuredBuffer(
+				GraphBuilder,
+				TEXT("Bloom.ApplyParameters"),
+				ApplyParametersArray);
+
+			CommonParameters.BloomApplyParameters = GraphBuilder.CreateSRV(ApplyParametersBuffer);
+		}
+	}
+
 	// forward transformation from shader:
 	//return LensPrincipalPointOffsetScale.xy + UV * LensPrincipalPointOffsetScale.zw;
 
diff --git a/Engine/Source/Runtime/Renderer/Private/PostProcess/PostProcessTonemap.h b/Engine/Source/Runtime/Renderer/Private/PostProcess/PostProcessTonemap.h
--- a/Engine/Source/Runtime/Renderer/Private/PostProcess/PostProcessTonemap.h
+++ b/Engine/Source/Runtime/Renderer/Private/PostProcess/PostProcessTonemap.h
@@ -3,6 +3,7 @@
 #pragma once
 
 #include "PostProcess/PostProcessEyeAdaptation.h"
+#include "PostProcess/PostProcessBloomSetup.h"
 #include "ScreenPass.h"
 #include "OverridePassSequence.h"
 #include "Math/Halton.h"
@@ -47,7 +48,7 @@ struct FTonemapInputs
 	FScreenPassTexture SceneColor;
 
 	// [Required] Filtered bloom texture to composite with tonemapped scene color. This should be transparent black for no bloom.
-	FScreenPassTexture Bloom;
+	FBloomOutputs Bloom;
 
 	// [Required] Color grading texture used to remap colors.
 	FRDGTextureRef ColorGradingTexture = nullptr;
diff --git a/Engine/Source/Runtime/Renderer/Private/PostProcess/PostProcessing.cpp b/Engine/Source/Runtime/Renderer/Private/PostProcess/PostProcessing.cpp
--- a/Engine/Source/Runtime/Renderer/Private/PostProcess/PostProcessing.cpp
+++ b/Engine/Source/Runtime/Renderer/Private/PostProcess/PostProcessing.cpp
@@ -772,47 +772,45 @@ void AddPostProcessingPasses(FRDGBuilder& GraphBuilder, const FViewInfo& View, c
 			GraphBuilder.QueueTextureExtraction(DownsampledSceneColor.Texture, &View.ViewState->PrevFrameViewInfo.HalfResTemporalAAHistory);
 		}
 
-		FSceneDownsampleChain SceneDownsampleChain;
-
 		if (bHistogramEnabled)
 		{
 			HistogramTexture = AddHistogramPass(GraphBuilder, View, EyeAdaptationParameters, DownsampledSceneColor, LastEyeAdaptationTexture);
 		}
 
-		if (bEyeAdaptationEnabled)
-		{
-			const bool bBasicEyeAdaptationEnabled = bEyeAdaptationEnabled && (AutoExposureMethod == EAutoExposureMethod::AEM_Basic);
+		const bool bFFTBloomEnabled = IsFFTBloomEnabled(View);
+		const bool bBasicEyeAdaptationEnabled = bEyeAdaptationEnabled && (AutoExposureMethod == EAutoExposureMethod::AEM_Basic);
 
-			if (bBasicEyeAdaptationEnabled)
-			{
-				const bool bLogLumaInAlpha = true;
-				SceneDownsampleChain.Init(GraphBuilder, View, EyeAdaptationParameters, DownsampledSceneColor, DownsampleQuality, bLogLumaInAlpha);
+		FSceneDownsampleChain SceneDownsampleChain;
+		if (bBasicEyeAdaptationEnabled || (bBloomEnabled && bFFTBloomEnabled))
+		{
+			const bool bLogLumaInAlpha = true;
+			SceneDownsampleChain.Init(GraphBuilder, View, EyeAdaptationParameters, DownsampledSceneColor, DownsampleQuality, bLogLumaInAlpha);
+		}
 
-				// Use the alpha channel in the last downsample (smallest) to compute eye adaptations values.
-				EyeAdaptationTexture = AddBasicEyeAdaptationPass(GraphBuilder, View, EyeAdaptationParameters, SceneDownsampleChain.GetLastTexture(), LastEyeAdaptationTexture);
-			}
-			// Add histogram eye adaptation pass even if no histogram exists to support the manual clamping mode.
-			else
-			{
-				EyeAdaptationTexture = AddHistogramEyeAdaptationPass(GraphBuilder, View, EyeAdaptationParameters, HistogramTexture);
-			}
+		if (bBasicEyeAdaptationEnabled)
+		{
+			// Use the alpha channel in the last downsample (smallest) to compute eye adaptations values.
+			EyeAdaptationTexture = AddBasicEyeAdaptationPass(GraphBuilder, View, EyeAdaptationParameters, SceneDownsampleChain.GetLastTexture(), LastEyeAdaptationTexture);
+		}
+		// Add histogram eye adaptation pass even if no histogram exists to support the manual clamping mode.
+		else if (bEyeAdaptationEnabled)
+		{
+			EyeAdaptationTexture = AddHistogramEyeAdaptationPass(GraphBuilder, View, EyeAdaptationParameters, HistogramTexture);
 		}
 
-		FScreenPassTexture Bloom;
+		FBloomOutputs Bloom;
 
 		if (bBloomEnabled)
 		{
 			FSceneDownsampleChain BloomDownsampleChain;
 
-			FBloomInputs PassInputs;
-			PassInputs.SceneColor = SceneColor;
-
 			const bool bBloomThresholdEnabled = View.FinalPostProcessSettings.BloomThreshold > -1.0f;
 
-			// Reuse the main scene downsample chain if a threshold isn't required for bloom.
+			// Reuse the main scene downsample chain if a threshold isn't required for gaussian bloom.
+			const FSceneDownsampleChain* GaussianBloomSceneDownsampleChain;
 			if (SceneDownsampleChain.IsInitialized() && !bBloomThresholdEnabled)
 			{
-				PassInputs.SceneDownsampleChain = &SceneDownsampleChain;
+				GaussianBloomSceneDownsampleChain = &SceneDownsampleChain;
 			}
 			else
 			{
@@ -833,28 +831,43 @@ void AddPostProcessingPasses(FRDGBuilder& GraphBuilder, const FViewInfo& View, c
 				const bool bLogLumaInAlpha = false;
 				BloomDownsampleChain.Init(GraphBuilder, View, EyeAdaptationParameters, DownsampleInput, DownsampleQuality, bLogLumaInAlpha);
 
-				PassInputs.SceneDownsampleChain = &BloomDownsampleChain;
+				GaussianBloomSceneDownsampleChain = &BloomDownsampleChain;
 			}
 
-			FBloomOutputs PassOutputs = AddBloomPass(GraphBuilder, View, PassInputs);
-			SceneColor = PassOutputs.SceneColor;
-			Bloom = PassOutputs.Bloom;
+			if (bFFTBloomEnabled)
+			{
+				FScreenPassTexture HalfResolution;
+				if (IsFFTBloomQuarterResolutionEnabled())
+				{
+					HalfResolution = SceneDownsampleChain.GetTexture(1);
+				}
+				else
+				{
+					HalfResolution = SceneDownsampleChain.GetFirstTexture();
+				}
 
-			FScreenPassTexture LensFlares = AddLensFlaresPass(GraphBuilder, View, Bloom, *PassInputs.SceneDownsampleChain);
+				FFFTBloomInputs PassInputs;
+				PassInputs.FullResolutionTexture = SceneColor.Texture;
+				PassInputs.FullResolutionViewRect = SceneColor.ViewRect;
+				PassInputs.HalfResolutionTexture = HalfResolution.Texture;
+				PassInputs.HalfResolutionViewRect = HalfResolution.ViewRect;
+
+				Bloom = AddFFTBloomPass(GraphBuilder, View, PassInputs);
+			}
+			else
+			{
+				Bloom = AddGaussianBloomPasses(GraphBuilder, View, GaussianBloomSceneDownsampleChain);
+			}
+
+			FScreenPassTexture LensFlares = AddLensFlaresPass(GraphBuilder, View, Bloom.Bloom, *GaussianBloomSceneDownsampleChain);
 
 			if (LensFlares.IsValid())
 			{
 				// Lens flares are composited with bloom.
-				Bloom = LensFlares;
+				Bloom.Bloom = LensFlares;
 			}
 		}
 
-		// Tonemapper needs a valid bloom target, even if it's black.
-		if (!Bloom.IsValid())
-		{
-			Bloom = BlackDummy;
-		}
-
 		SceneColorBeforeTonemap = SceneColor;
 
 		if (PassSequence.IsEnabled(EPass::Tonemap))
@@ -869,7 +882,7 @@ void AddPostProcessingPasses(FRDGBuilder& GraphBuilder, const FViewInfo& View, c
 				PassSequence.AcceptOverrideIfLastPass(EPass::Tonemap, PassInputs.OverrideOutput);
 				PassInputs.SetInput(EPostProcessMaterialInput::SceneColor, SceneColor);
 				PassInputs.SetInput(EPostProcessMaterialInput::SeparateTranslucency, SeparateTranslucency);
-				PassInputs.SetInput(EPostProcessMaterialInput::CombinedBloom, Bloom);
+				PassInputs.SetInput(EPostProcessMaterialInput::CombinedBloom, Bloom.Bloom);
 				PassInputs.SceneTextures = GetSceneTextureShaderParameters(Inputs.SceneTextures);
 				PassInputs.CustomDepthTexture = CustomDepth.Texture;
 
@@ -2176,7 +2189,7 @@ void AddMobilePostProcessingPasses(FRDGBuilder& GraphBuilder, FScene* Scene, con
 		}
 			
 		TonemapperInputs.SceneColor = SceneColor;
-		TonemapperInputs.Bloom = BloomOutput;
+		TonemapperInputs.Bloom.Bloom = BloomOutput;
 		TonemapperInputs.EyeAdaptationTexture = nullptr;
 		TonemapperInputs.ColorGradingTexture = ColorGradingTexture;
 		TonemapperInputs.bWriteAlphaChannel = View.AntiAliasingMethod == AAM_FXAA || IsPostProcessingWithAlphaChannelSupported() || bUseMobileDof || IsMobilePropagateAlphaEnabled(View.GetShaderPlatform());
diff --git a/Engine/Source/Runtime/Renderer/Private/ScenePrivate.h b/Engine/Source/Runtime/Renderer/Private/ScenePrivate.h
--- a/Engine/Source/Runtime/Renderer/Private/ScenePrivate.h
+++ b/Engine/Source/Runtime/Renderer/Private/ScenePrivate.h
@@ -963,12 +963,16 @@ public:
 	// Pre-computed filter in spectral (i.e. FFT) domain along with data to determine if we need to up date it
 	struct {
 		/// @cond DOXYGEN_WARNINGS
-		void SafeRelease() { Spectral.SafeRelease(); CenterWeight.SafeRelease(); }
+		void SafeRelease()
+		{
+			Spectral.SafeRelease();
+			ConstantsBuffer.SafeRelease();
+		}
 		/// @endcond
 
 		// The 2d fourier transform of the physical space texture.
 		TRefCountPtr<IPooledRenderTarget> Spectral;
-		TRefCountPtr<IPooledRenderTarget> CenterWeight; // a 1-pixel buffer that holds blend weights for half-resolution fft.
+		TRefCountPtr<FRDGPooledBuffer> ConstantsBuffer;
 
 		// The physical space source texture
 		UTexture2D* Physical = nullptr;
@@ -982,8 +986,6 @@ public:
 		// The size of the viewport for which the spectral kernel was calculated. 
 		FIntPoint ImageSize;
 
-		FVector2D CenterUV;
-
 		// Mip level of the physical space source texture used when caching the spectral space texture.
 		uint32 PhysicalMipLevel;
 	} BloomFFTKernel;
