Removes some dead code in PostProcessFFTBloom.cpp

diff --git a/Engine/Source/Runtime/Renderer/Private/PostProcess/PostProcessFFTBloom.cpp b/Engine/Source/Runtime/Renderer/Private/PostProcess/PostProcessFFTBloom.cpp
--- a/Engine/Source/Runtime/Renderer/Private/PostProcess/PostProcessFFTBloom.cpp
+++ b/Engine/Source/Runtime/Renderer/Private/PostProcess/PostProcessFFTBloom.cpp
@@ -371,169 +371,6 @@ FFFTBloomIntermediates GetFFTBloomIntermediates(
 	return Intermediates;
 }
 
-
-struct FKernelAnalysisResult
-{
-	FRDGBufferRef MaxScatterDispersionBuffer;
-	FRDGBufferRef KernelCenterEnergyBuffer;
-	FRDGTextureRef ScatterDispersionTexture;
-};
-
-FKernelAnalysisResult AnalysesKernelAtDensity(
-	FRDGBuilder& GraphBuilder,
-	const FGlobalShaderMap* ShaderMap,
-	const FFFTBloomIntermediates& Intermediates,
-	FRDGTextureRef SpatialKernelTexture,
-	FRDGBufferRef KernelCenterCoordBuffer,
-	float ViewTexelDiameterInKernelTexels)
-{
-	// Find information about the kernel's energy distribution
-	FRDGBufferRef MaxScatterDispersionBuffer;
-	FRDGBufferRef KernelCenterEnergyBuffer;
-	{
-		const int32 SurveyTileSize = 8;
-
-		float ViewTexelRadiusInKernelTexels = ViewTexelDiameterInKernelTexels * 0.5f;
-
-		int32 SurveyGroupGridSize = 2 * FMath::DivideAndRoundUp(FMath::CeilToInt(ViewTexelRadiusInKernelTexels) + 4, SurveyTileSize);
-		int32 SurveyGroupCount = SurveyGroupGridSize * SurveyGroupGridSize;
-
-		FRDGTextureUAVRef DebugTextureUAV;
-		{
-			FRDGTextureDesc DebugDesc = FRDGTextureDesc::Create2D(
-				FIntPoint(SurveyGroupGridSize * 8, SurveyGroupGridSize * 8),
-				PF_FloatRGBA,
-				FClearValueBinding::None,
-				/* InFlags = */ TexCreate_ShaderResource | TexCreate_UAV);
-
-			FRDGTextureRef DebugTexture = GraphBuilder.CreateTexture(DebugDesc, TEXT("Debug.Bloom.Survey"));
-
-			DebugTextureUAV = GraphBuilder.CreateUAV(DebugTexture);
-		}
-
-		RDG_EVENT_SCOPE(GraphBuilder, "FFTBloom SurveyKernel(TexelDiameter=%f)", ViewTexelDiameterInKernelTexels);
-
-		// Reduce a survey buffer to a single value.
-		auto ReduceSurveyBuffer = [&](FRDGBufferRef SurveyBuffer, int32 Op)
-		{
-			FBloomReduceKernelSurveyCS::FParameters* PassParameters = GraphBuilder.AllocParameters<FBloomReduceKernelSurveyCS::FParameters>();
-			PassParameters->SurveyReduceOp = Op;
-			PassParameters->SurveyGroupCount = SurveyGroupCount;
-			PassParameters->SurveyOutput = GraphBuilder.CreateUAV(SurveyBuffer);
-
-			TShaderMapRef<FBloomReduceKernelSurveyCS> ComputeShader(ShaderMap);
-			FComputeShaderUtils::AddPass(
-				GraphBuilder,
-				RDG_EVENT_NAME("FFTBloom ReduceKernelSurvey(Op=%d) %d", Op, SurveyGroupCount),
-				ComputeShader,
-				PassParameters,
-				FComputeShaderUtils::GetGroupCount(SurveyGroupCount, 64));
-		};
-
-		// Find the max scatter dispersion to use around the footprint of the view pixel in the kernel.
-		{
-			MaxScatterDispersionBuffer = GraphBuilder.CreateBuffer(
-				FRDGBufferDesc::CreateStructuredDesc(sizeof(FLinearColor), /* NumElements = */ SurveyGroupCount),
-				TEXT("Bloom.FFT.MaxScatterDispersion"));
-
-			FBloomSurveyMaxScatterDispersionCS::FParameters* PassParameters = GraphBuilder.AllocParameters<FBloomSurveyMaxScatterDispersionCS::FParameters>();
-			PassParameters->ViewTexelRadiusInKernelTexels = ViewTexelRadiusInKernelTexels;
-			PassParameters->SurveyGroupGridSize = SurveyGroupGridSize;
-			PassParameters->KernelSpatialTextureSize = SpatialKernelTexture->Desc.Extent;
-			PassParameters->KernelSpatialTexture = SpatialKernelTexture;
-			PassParameters->KernelCenterCoordBuffer = GraphBuilder.CreateSRV(KernelCenterCoordBuffer);
-			PassParameters->SurveyOutput = GraphBuilder.CreateUAV(MaxScatterDispersionBuffer);
-			PassParameters->DebugOutput = DebugTextureUAV;
-
-			TShaderMapRef<FBloomSurveyMaxScatterDispersionCS> ComputeShader(ShaderMap);
-			FComputeShaderUtils::AddPass(
-				GraphBuilder,
-				RDG_EVENT_NAME("FFTBloom SurveyMaxScatterDispersion"),
-				ComputeShader,
-				PassParameters,
-				FIntVector(SurveyGroupGridSize, SurveyGroupGridSize, 1));
-
-			ReduceSurveyBuffer(MaxScatterDispersionBuffer, /* Op = */ 0);
-		}
-
-		// Find the amount of energy at the center in the footprint of the view pixel in the kernel.
-		{
-			KernelCenterEnergyBuffer = GraphBuilder.CreateBuffer(
-				FRDGBufferDesc::CreateStructuredDesc(sizeof(FLinearColor), /* NumElements = */ SurveyGroupCount),
-				TEXT("Bloom.FFT.KernelCenterEnergy"));
-
-			FBloomSurveyKernelCenterEnergyCS::FParameters* PassParameters = GraphBuilder.AllocParameters<FBloomSurveyKernelCenterEnergyCS::FParameters>();
-			PassParameters->ViewTexelRadiusInKernelTexels = ViewTexelRadiusInKernelTexels;
-			PassParameters->SurveyGroupGridSize = SurveyGroupGridSize;
-			PassParameters->KernelSpatialTextureSize = SpatialKernelTexture->Desc.Extent;
-			PassParameters->KernelSpatialTexture = SpatialKernelTexture;
-			PassParameters->KernelCenterCoordBuffer = GraphBuilder.CreateSRV(KernelCenterCoordBuffer);
-			PassParameters->MaxScatterDispersionBuffer = GraphBuilder.CreateSRV(MaxScatterDispersionBuffer);
-			PassParameters->SurveyOutput = GraphBuilder.CreateUAV(KernelCenterEnergyBuffer);
-			PassParameters->DebugOutput = DebugTextureUAV;
-
-			TShaderMapRef<FBloomSurveyKernelCenterEnergyCS> ComputeShader(ShaderMap);
-			FComputeShaderUtils::AddPass(
-				GraphBuilder,
-				RDG_EVENT_NAME("FFTBloom SurveyKernelCenterEnergy"),
-				ComputeShader,
-				PassParameters,
-				FIntVector(SurveyGroupGridSize, SurveyGroupGridSize, 1));
-
-			ReduceSurveyBuffer(KernelCenterEnergyBuffer, /* Op = */ 1);
-		}
-	}
-
-	// Find out the total energy of the kernel - center.
-	// Implemented in such away to prioritize numerical error rather than speed.
-	FRDGTextureRef ScatterDispersionTexture;
-	{
-		RDG_EVENT_SCOPE(GraphBuilder, "FFTBloom SumScatterDispersionEnergy %dx%d",
-			SpatialKernelTexture->Desc.Extent.X, SpatialKernelTexture->Desc.Extent.Y);
-
-		ScatterDispersionTexture = SpatialKernelTexture;
-
-		for (int32 PassId = 0; ScatterDispersionTexture->Desc.Extent.X > 1 && ScatterDispersionTexture->Desc.Extent.Y > 1; PassId++)
-		{
-			FRDGTextureRef NewScatterDispersionTexture;
-			{
-				FRDGTextureDesc Desc = FRDGTextureDesc::Create2D(
-					FIntPoint::DivideAndRoundUp(ScatterDispersionTexture->Desc.Extent, 8),
-					PF_A32B32G32R32F,
-					FClearValueBinding::None,
-					TexCreate_ShaderResource | TexCreate_UAV);
-
-				NewScatterDispersionTexture = GraphBuilder.CreateTexture(Desc, TEXT("Bloom.FFT.KernelIntensity"));
-			}
-
-			FBloomSumScatterDispersionEnergyCS::FParameters* PassParameters = GraphBuilder.AllocParameters<FBloomSumScatterDispersionEnergyCS::FParameters>();
-			PassParameters->PassId = PassId;
-			PassParameters->ScatterDispersionTextureSize = ScatterDispersionTexture->Desc.Extent;
-			PassParameters->ScatterDispersionTexture = ScatterDispersionTexture;
-			PassParameters->MaxScatterDispersionBuffer = GraphBuilder.CreateSRV(MaxScatterDispersionBuffer);
-			PassParameters->ScatterDispersionOutput = GraphBuilder.CreateUAV(NewScatterDispersionTexture);
-
-			TShaderMapRef<FBloomSumScatterDispersionEnergyCS> ComputeShader(ShaderMap);
-			FComputeShaderUtils::AddPass(
-				GraphBuilder,
-				RDG_EVENT_NAME("FFTBloom SumScatterDispersionEnergy %dx%d -> %dx%d",
-					ScatterDispersionTexture->Desc.Extent.X, ScatterDispersionTexture->Desc.Extent.Y,
-					NewScatterDispersionTexture->Desc.Extent.X, NewScatterDispersionTexture->Desc.Extent.Y),
-				ComputeShader,
-				PassParameters,
-				FIntVector(NewScatterDispersionTexture->Desc.Extent.X, NewScatterDispersionTexture->Desc.Extent.Y, 1));
-
-			ScatterDispersionTexture = NewScatterDispersionTexture;
-		}
-	}
-
-	FKernelAnalysisResult Result;
-	Result.MaxScatterDispersionBuffer = MaxScatterDispersionBuffer;
-	Result.KernelCenterEnergyBuffer = KernelCenterEnergyBuffer;
-	Result.ScatterDispersionTexture = ScatterDispersionTexture;
-	return Result;
-}
-
 void InitDomainAndGetKernel(
 	FRDGBuilder& GraphBuilder,
 	const FViewInfo& View,
