Fixes warnings in FFT bloom

diff --git a/Engine/Source/Runtime/Engine/Classes/Engine/Scene.h b/Engine/Source/Runtime/Engine/Classes/Engine/Scene.h
--- a/Engine/Source/Runtime/Engine/Classes/Engine/Scene.h
+++ b/Engine/Source/Runtime/Engine/Classes/Engine/Scene.h
@@ -403,6 +403,7 @@ struct FConvolutionBloomSettings
 	FConvolutionBloomSettings()
 	{
 		Texture = nullptr;
+		ScatterDispersion = 1.0f;
 		Size = 1.f;
 		CenterUV = FVector2D(0.5f, 0.5f);
 		PreFilterMin = 7.f;
diff --git a/Engine/Source/Runtime/Renderer/Private/PostProcess/PostProcessFFTBloom.cpp b/Engine/Source/Runtime/Renderer/Private/PostProcess/PostProcessFFTBloom.cpp
--- a/Engine/Source/Runtime/Renderer/Private/PostProcess/PostProcessFFTBloom.cpp
+++ b/Engine/Source/Runtime/Renderer/Private/PostProcess/PostProcessFFTBloom.cpp
@@ -126,7 +126,7 @@ public:
 		SHADER_PARAMETER_RDG_BUFFER_SRV(StructuredBuffer<uint>, KernelCenterCoordBuffer)
 		SHADER_PARAMETER_RDG_BUFFER_SRV(StructuredBuffer<uint>, KernelCenterEnergyBuffer)
 		SHADER_PARAMETER_RDG_BUFFER_SRV(StructuredBuffer<uint>, MaxScatterDispersionBuffer)
-		SHADER_PARAMETER_RDG_TEXTURE(StructuredBuffer<uint>, ScatterDispersionTexture)
+		SHADER_PARAMETER_RDG_TEXTURE(Texture2D, ScatterDispersionTexture)
 		SHADER_PARAMETER_RDG_BUFFER_UAV(RWStructuredBuffer<uint>, KernelConstantsOutput)
 	END_SHADER_PARAMETER_STRUCT()
 };
diff --git a/Engine/Source/Runtime/Renderer/Private/PostProcess/PostProcessTonemap.cpp b/Engine/Source/Runtime/Renderer/Private/PostProcess/PostProcessTonemap.cpp
--- a/Engine/Source/Runtime/Renderer/Private/PostProcess/PostProcessTonemap.cpp
+++ b/Engine/Source/Runtime/Renderer/Private/PostProcess/PostProcessTonemap.cpp
@@ -339,7 +339,7 @@ BEGIN_SHADER_PARAMETER_STRUCT(FTonemapParameters, )
 	SHADER_PARAMETER(FVector4, BloomMultiply)
 	SHADER_PARAMETER_RDG_TEXTURE(Texture2D, BloomTexture)
 	SHADER_PARAMETER_SAMPLER(SamplerState, BloomSampler)
-	SHADER_PARAMETER_RDG_BUFFER_SRV(SamplerState, BloomApplyParameters)
+	SHADER_PARAMETER_RDG_BUFFER_SRV(StructuredBuffer<float4>, BloomApplyParameters)
 
 	// SM5 and above use Texture2D for EyeAdaptationTexture
 	SHADER_PARAMETER_RDG_TEXTURE(Texture2D, EyeAdaptationTexture)
