Fix viewport for real-time reflection capture render.

diff --git a/Engine/Source/Runtime/Renderer/Private/ReflectionEnvironmentRealTimeCapture.cpp b/Engine/Source/Runtime/Renderer/Private/ReflectionEnvironmentRealTimeCapture.cpp
--- a/Engine/Source/Runtime/Renderer/Private/ReflectionEnvironmentRealTimeCapture.cpp
+++ b/Engine/Source/Runtime/Renderer/Private/ReflectionEnvironmentRealTimeCapture.cpp
@@ -584,7 +584,7 @@ void FScene::AllocateAndCaptureFrameSkyEnvMap(
 								PassParameters->RenderTargets.DepthStencil = FDepthStencilBinding(CubeDepthTexture, ERenderTargetLoadAction::EClear, FExclusiveDepthStencil::DepthWrite_StencilNop);
 							}
 
-							AddSimpleMeshPass(GraphBuilder, PassParameters, Scene, MainView, &InstanceCullingManager, RDG_EVENT_NAME("CaptureSkyMeshReflection"), MainView.ViewRect,
+							AddSimpleMeshPass(GraphBuilder, PassParameters, Scene, MainView, &InstanceCullingManager, RDG_EVENT_NAME("CaptureSkyMeshReflection"), SkyRC.Viewport,
 							[&MainView, &CubeViewUniformBuffer, bUseDepthBuffer, Scene](FDynamicPassMeshDrawListContext* DynamicMeshPassContext)
 							{
 								FMeshPassProcessorRenderState DrawRenderState;
