Fix render-thread starvation under high load by better distributing batches among threads
  - ParallelFor is already doing batching internally if Unbalanced is not used, resulting in uneven distribution

#rnx

diff --git a/Engine/Source/Runtime/Renderer/Private/PrimitiveSceneInfo.cpp b/Engine/Source/Runtime/Renderer/Private/PrimitiveSceneInfo.cpp
--- a/Engine/Source/Runtime/Renderer/Private/PrimitiveSceneInfo.cpp
+++ b/Engine/Source/Runtime/Renderer/Private/PrimitiveSceneInfo.cpp
@@ -410,7 +410,7 @@ void FPrimitiveSceneInfo::CacheMeshDrawCommands(FRHICommandListImmediate& RHICmd
 
 	if (FApp::ShouldUseThreadingForPerformance())
 	{
-		ParallelForTemplate(NumBatches, DoWorkLambda, EParallelForFlags::PumpRenderingThread);
+		ParallelForTemplate(NumBatches, DoWorkLambda, EParallelForFlags::PumpRenderingThread | EParallelForFlags::Unbalanced);
 	}
 	else
 	{
