Exported more methods of FMeshDrawShaderBindings so that it can be used in pluggins

diff --git a/Engine/Source/Runtime/Renderer/Public/MeshPassProcessor.h b/Engine/Source/Runtime/Renderer/Public/MeshPassProcessor.h
--- a/Engine/Source/Runtime/Renderer/Public/MeshPassProcessor.h
+++ b/Engine/Source/Runtime/Renderer/Public/MeshPassProcessor.h
@@ -767,12 +767,12 @@ public:
 	}
 
 	/** Allocates space for the bindings of all shaders. */
-	void Initialize(FMeshProcessorShaders Shaders);
+	RENDERER_API void Initialize(FMeshProcessorShaders Shaders);
 
 	/** Called once binding setup is complete. */
-	void Finalize(const FMeshProcessorShaders* ShadersForDebugging);
+	RENDERER_API void Finalize(const FMeshProcessorShaders* ShadersForDebugging);
 
-	inline FMeshDrawSingleShaderBindings GetSingleShaderBindings(EShaderFrequency Frequency, int32& DataOffset)
+	FORCEINLINE FMeshDrawSingleShaderBindings GetSingleShaderBindings(EShaderFrequency Frequency, int32& DataOffset)
 	{
 		int FrequencyIndex = FPlatformMath::CountBits(ShaderFrequencyBits & ((1 << (Frequency + 1)) - 1)) - 1;
 
@@ -796,8 +796,8 @@ public:
 	}
 
 	/** Set shader bindings on the commandlist, filtered by state cache. */
-	void SetOnCommandList(FRHICommandList& RHICmdList, FBoundShaderStateInput Shaders, class FShaderBindingState* StateCacheShaderBindings) const;
-	void SetOnCommandList(FRHIComputeCommandList& RHICmdList, FRHIComputeShader* Shader, class FShaderBindingState* StateCacheShaderBindings = nullptr) const;
+	RENDERER_API void SetOnCommandList(FRHICommandList& RHICmdList, FBoundShaderStateInput Shaders, class FShaderBindingState* StateCacheShaderBindings) const;
+	RENDERER_API void SetOnCommandList(FRHIComputeCommandList& RHICmdList, FRHIComputeShader* Shader, class FShaderBindingState* StateCacheShaderBindings = nullptr) const;
 
 #if RHI_RAYTRACING
 	RENDERER_API void SetRayTracingShaderBindingsForHitGroup(FRayTracingLocalShaderBindingWriter* BindingWriter, uint32 InstanceIndex, uint32 SegmentIndex, uint32 HitGroupIndex, uint32 ShaderSlot) const;
