warning fix

diff --git a/Engine/Source/Runtime/Renderer/Private/VT/RuntimeVirtualTextureRender.cpp b/Engine/Source/Runtime/Renderer/Private/VT/RuntimeVirtualTextureRender.cpp
--- a/Engine/Source/Runtime/Renderer/Private/VT/RuntimeVirtualTextureRender.cpp
+++ b/Engine/Source/Runtime/Renderer/Private/VT/RuntimeVirtualTextureRender.cpp
@@ -920,7 +920,7 @@ namespace RuntimeVirtualTexture
 			bCopyThumbnailPass = bRenderPass && bIsThumbnails;
 			EPixelFormat OutputFormat = (OutputTexture0 != nullptr ? OutputTexture0->GetFormat() : PF_Unknown);
 			bool bCompressedFormat = GPixelFormats[OutputFormat].BlockSizeX == 4 && GPixelFormats[OutputFormat].BlockSizeY == 4;
-			bool bLQFormat = (OutputTexture0->GetFormat() == PF_R5G6B5_UNORM);
+			bool bLQFormat = (OutputTexture0 && OutputTexture0->GetFormat() == PF_R5G6B5_UNORM);
 			bCompressPass = bRenderPass && !bCopyThumbnailPass && bCompressedFormat;
 			bCopyPass = bRenderPass && !bCopyThumbnailPass && !bCompressPass && (MaterialType == ERuntimeVirtualTextureMaterialType::BaseColor_Normal_Specular || MaterialType == ERuntimeVirtualTextureMaterialType::BaseColor_Normal_Specular_YCoCg || MaterialType == ERuntimeVirtualTextureMaterialType::BaseColor_Normal_Specular_Mask_YCoCg);
 			// Not all mobile RHIs support sRGB texture views/aliasing, use only linear targets on mobile
