Removes film grain decode parameters in FPostProcessSettings

diff --git a/Engine/Shaders/Private/PostProcessTonemap.usf b/Engine/Shaders/Private/PostProcessTonemap.usf
--- a/Engine/Shaders/Private/PostProcessTonemap.usf
+++ b/Engine/Shaders/Private/PostProcessTonemap.usf
@@ -76,8 +76,6 @@ float4 TonemapperParams;
 float3 GrainRandomFull;
 
 // Film grain
-float FilmGrainMul;
-float FilmGrainAdd;
 float FilmGrainIntensityShadows;
 float FilmGrainIntensityMidtones;
 float FilmGrainIntensityHighlights;
@@ -86,6 +84,9 @@ float FilmGrainHighlightsMin;
 Texture2D<half3> FilmGrainTexture;
 SamplerState FilmGrainSampler;
 float4 ScreenPosToFilmGrainTextureUV;
+#if SUPPORTS_FILM_GRAIN
+	StructuredBuffer<float4> FilmGrainTextureConstants;
+#endif // SUPPORTS_FILM_GRAIN
 
 float EditorNITLevel;
 
@@ -641,12 +642,14 @@ float4 TonemapCommonPS(
 
 		#if USE_FILM_GRAIN
 		{
+			float3 FilmGrainDecodeMultiply = FilmGrainTextureConstants[0].rgb;
+
 			half FilmGrainIntensity = ComputeFilmGrainIntensity(LinearColor);
 
 			float2 FilmGrainUV = ApplyScreenTransform(ScreenPos, ScreenPosToFilmGrainTextureUV);
 
 			half3 RawGrain = FilmGrainTexture.SampleLevel(FilmGrainSampler, FilmGrainUV, 0);
-			half3 FinalGrain = RawGrain * half(FilmGrainMul) + half(FilmGrainAdd);
+			half3 FinalGrain = RawGrain * half3(FilmGrainDecodeMultiply);
 
 			LinearColor.rgb *= lerp(1.0, FinalGrain, FilmGrainIntensity);
 		}
diff --git a/Engine/Shaders/Private/PostProcessing/FilmGrainPackConstants.usf b/Engine/Shaders/Private/PostProcessing/FilmGrainPackConstants.usf
new file mode 100644
--- /dev/null
+++ b/Engine/Shaders/Private/PostProcessing/FilmGrainPackConstants.usf
@@ -0,0 +1,27 @@
+// Copyright Epic Games, Inc. All Rights Reserved.
+
+#include "../Common.ush"
+
+
+//------------------------------------------------------- PARAMETERS
+
+uint2 OriginalFilmGrainTextureSize;
+Texture2D ReducedFilmGrainTexture;
+RWStructuredBuffer<float4> FilmGrainConstantsOutput;
+
+
+//------------------------------------------------------- ENTRY POINT
+
+[numthreads(1, 1, 1)]
+void MainCS(
+	uint2 GroupId : SV_GroupID,
+	uint GroupThreadIndex : SV_GroupIndex)
+{
+	float4 ReducedFilmGrainEnergy = ReducedFilmGrainTexture[uint2(0, 0)];
+
+	float FilmGrainPixelCount = float(OriginalFilmGrainTextureSize.x * OriginalFilmGrainTextureSize.y);
+
+	float4 FilmGrainDecodeMultiply = FilmGrainPixelCount.xxxx / ReducedFilmGrainEnergy;
+
+	FilmGrainConstantsOutput[0] = FilmGrainDecodeMultiply;
+}
diff --git a/Engine/Shaders/Private/PostProcessing/FilmGrainReduce.usf b/Engine/Shaders/Private/PostProcessing/FilmGrainReduce.usf
new file mode 100644
--- /dev/null
+++ b/Engine/Shaders/Private/PostProcessing/FilmGrainReduce.usf
@@ -0,0 +1,70 @@
+// Copyright Epic Games, Inc. All Rights Reserved.
+
+#include "../Common.ush"
+
+
+//------------------------------------------------------- CONSTANTS
+
+#define TILE_SIZE 8
+
+
+//------------------------------------------------------- PARAMETERS
+
+uint2 FilmGrainTextureSize;
+Texture2D FilmGrainTexture;
+
+RWTexture2D<float4> Output;
+
+
+//------------------------------------------------------- LDS
+
+groupshared float4 SharedColor[TILE_SIZE * TILE_SIZE];
+
+
+//------------------------------------------------------- ENTRY POINT
+
+[numthreads(TILE_SIZE, TILE_SIZE, 1)]
+void MainCS(
+	uint2 GroupId : SV_GroupID,
+	uint GroupThreadIndex : SV_GroupIndex)
+{
+	uint2 InputPixelPos = TILE_SIZE * GroupId + uint2(GroupThreadIndex % TILE_SIZE, GroupThreadIndex / TILE_SIZE);
+
+	// Init LDS
+	float4 TexelColor = FilmGrainTexture[InputPixelPos];
+	{
+		bool bValidPixel = all(InputPixelPos < FilmGrainTextureSize);
+
+		if (!bValidPixel)
+		{
+			TexelColor = 0.0;
+		}
+
+		SharedColor[GroupThreadIndex] = TexelColor;
+	}
+
+	GroupMemoryBarrierWithGroupSync();
+
+	// Reduction
+	{
+		UNROLL
+		for (uint i = 32; i > 1; i /= 2)
+		{
+			TexelColor += SharedColor[GroupThreadIndex ^ i];
+
+			GroupMemoryBarrierWithGroupSync();
+
+			SharedColor[GroupThreadIndex] = TexelColor;
+			
+			GroupMemoryBarrierWithGroupSync();
+		}
+		
+		TexelColor += SharedColor[GroupThreadIndex ^ 0x1];
+	}
+
+	// Output
+	if (GroupThreadIndex == 0)
+	{
+		Output[GroupId] = TexelColor;
+	}
+}
diff --git a/Engine/Source/Runtime/Engine/Classes/Engine/Scene.h b/Engine/Source/Runtime/Engine/Classes/Engine/Scene.h
--- a/Engine/Source/Runtime/Engine/Classes/Engine/Scene.h
+++ b/Engine/Source/Runtime/Engine/Classes/Engine/Scene.h
@@ -1765,14 +1765,6 @@ struct FPostProcessSettings
 	UPROPERTY(EditAnywhere, BlueprintReadWrite, Category = "Film Grain", meta = (editcondition = "bOverride_FilmGrainTexture"))
 	TObjectPtr<class UTexture2D> FilmGrainTexture;
 
-	/** Multiply the FilmGrainTexture. DecodedFilmGrainTexture = FilmGrainTexture * FilmGrainDecodeMultiply + FilmGrainDecodeAdd */
-	UPROPERTY(EditAnywhere, BlueprintReadWrite, Category = "Film Grain", meta = (editcondition = "bOverride_FilmGrainTexture"))
-	float FilmGrainDecodeMultiply;
-
-	/** Add to the FilmGrainTexture. DecodedFilmGrainTexture = FilmGrainTexture * FilmGrainDecodeMultiply + FilmGrainDecodeAdd */
-	UPROPERTY(EditAnywhere, BlueprintReadWrite, Category = "Film Grain", meta = (editcondition = "bOverride_FilmGrainTexture"))
-	float FilmGrainDecodeAdd;
-
 	/** 0..1 0=off/no ambient occlusion .. 1=strong ambient occlusion, defines how much it affects the non direct lighting after base pass */
 	UPROPERTY(interp, BlueprintReadWrite, Category="Rendering Features|Ambient Occlusion", meta=(ClampMin = "0.0", ClampMax = "1.0", editcondition = "bOverride_AmbientOcclusionIntensity", DisplayName = "Intensity"))
 	float AmbientOcclusionIntensity;
diff --git a/Engine/Source/Runtime/Engine/Private/Scene.cpp b/Engine/Source/Runtime/Engine/Private/Scene.cpp
--- a/Engine/Source/Runtime/Engine/Private/Scene.cpp
+++ b/Engine/Source/Runtime/Engine/Private/Scene.cpp
@@ -361,8 +361,6 @@ static void DoPostProcessSettingsSanityCheck()
 		RenameMap.Add(TEXT("FilmGrainHighlightsMin"), TEXT(""));
 		RenameMap.Add(TEXT("FilmGrainTexelSize"), TEXT(""));
 		RenameMap.Add(TEXT("FilmGrainTexture"), TEXT(""));
-		RenameMap.Add(TEXT("FilmGrainDecodeMultiply"), TEXT(""));
-		RenameMap.Add(TEXT("FilmGrainDecodeAdd"), TEXT(""));
 		VerifyPostProcessingProperties(TEXT("Film"),
 			TArray<const UStruct*>({FFilmStockSettings::StaticStruct()}),
 			RenameMap);
@@ -519,8 +517,6 @@ FPostProcessSettings::FPostProcessSettings()
 	FilmGrainShadowsMax = 0.09f;
 	FilmGrainHighlightsMin = 0.5f;
 	FilmGrainTexelSize = 1.0f;
-	FilmGrainDecodeMultiply = 1.0f;
-	FilmGrainDecodeAdd = 0.0f;
 
 	// next value might get overwritten by r.DefaultFeature.AmbientOcclusion
 	AmbientOcclusionIntensity = .5f;
@@ -933,8 +929,6 @@ FPostProcessSettings::FPostProcessSettings(const FPostProcessSettings& Settings)
 	, FilmGrainHighlightsMin(Settings.FilmGrainHighlightsMin)
 	, FilmGrainTexelSize(Settings.FilmGrainTexelSize)
 	, FilmGrainTexture(Settings.FilmGrainTexture)
-	, FilmGrainDecodeMultiply(Settings.FilmGrainDecodeMultiply)
-	, FilmGrainDecodeAdd(Settings.FilmGrainDecodeAdd)
 	, AmbientOcclusionIntensity(Settings.AmbientOcclusionIntensity)
 	, AmbientOcclusionStaticFraction(Settings.AmbientOcclusionStaticFraction)
 	, AmbientOcclusionRadius(Settings.AmbientOcclusionRadius)
diff --git a/Engine/Source/Runtime/Engine/Private/SceneView.cpp b/Engine/Source/Runtime/Engine/Private/SceneView.cpp
--- a/Engine/Source/Runtime/Engine/Private/SceneView.cpp
+++ b/Engine/Source/Runtime/Engine/Private/SceneView.cpp
@@ -1669,8 +1669,6 @@ void FSceneView::OverridePostProcessSettings(const FPostProcessSettings& Src, fl
 		IF_PP(FilmGrainTexture)
 		{
 			Dest.FilmGrainTexture = Src.FilmGrainTexture;
-			Dest.FilmGrainDecodeMultiply = Src.FilmGrainDecodeMultiply;
-			Dest.FilmGrainDecodeAdd = Src.FilmGrainDecodeAdd;
 		}
 
 		// A continuous blending of this value would result trashing the pre-convolved bloom kernel cache.
diff --git a/Engine/Source/Runtime/Renderer/Private/PostProcess/PostProcessTonemap.cpp b/Engine/Source/Runtime/Renderer/Private/PostProcess/PostProcessTonemap.cpp
--- a/Engine/Source/Runtime/Renderer/Private/PostProcess/PostProcessTonemap.cpp
+++ b/Engine/Source/Runtime/Renderer/Private/PostProcess/PostProcessTonemap.cpp
@@ -18,6 +18,11 @@
 #include "Rendering/Texture2DResource.h"
 #include "Math/Halton.h"
 
+bool SupportsFilmGrain(EShaderPlatform Platform)
+{
+	return IsFeatureLevelSupported(Platform, ERHIFeatureLevel::SM5);
+}
+
 namespace
 {
 TAutoConsoleVariable<float> CVarTonemapperSharpen(
@@ -94,6 +99,11 @@ TAutoConsoleVariable<int32> CVarFilmGrainSequenceLength(
 	TEXT("Length of the random sequence for film grain (preferably a prime number, default=97)."),
 	ECVF_RenderThreadSafe);
 
+TAutoConsoleVariable<int32> CVarFilmGrainCacheTextureConstants(
+	TEXT("r.FilmGrain.CacheTextureConstants"), 1,
+	TEXT("Wether the constants related to the film grain should be cached."),
+	ECVF_RenderThreadSafe);
+
 const int32 GTonemapComputeTileSizeX = 8;
 const int32 GTonemapComputeTileSizeY = 8;
 
@@ -229,8 +239,14 @@ FDesktopDomain RemapPermutation(FDesktopDomain PermutationVector, ERHIFeatureLev
 		PermutationVector.Set<FTonemapperGrainQuantizationDim>(false);
 	else
 		PermutationVector.Set<FTonemapperGrainQuantizationDim>(true);
-	
+
+	if (FeatureLevel < ERHIFeatureLevel::SM5)
+	{
+		CommonPermutationVector.Set<FTonemapperFilmGrainDim>(false);
+	}
+
 	PermutationVector.Set<FCommonDomain>(CommonPermutationVector);
+
 	return PermutationVector;
 }
 
@@ -312,16 +328,15 @@ FTonemapperOutputDeviceParameters GetTonemapperOutputDeviceParameters(const FSce
 
 BEGIN_SHADER_PARAMETER_STRUCT(FFilmGrainParameters, )
 	SHADER_PARAMETER(FVector3f, GrainRandomFull)
-	SHADER_PARAMETER(float, FilmGrainMul)
-	SHADER_PARAMETER(float, FilmGrainAdd)
 	SHADER_PARAMETER(float, FilmGrainIntensityShadows)
 	SHADER_PARAMETER(float, FilmGrainIntensityMidtones)
 	SHADER_PARAMETER(float, FilmGrainIntensityHighlights)
 	SHADER_PARAMETER(float, FilmGrainShadowsMax)
 	SHADER_PARAMETER(float, FilmGrainHighlightsMin)
-	SHADER_PARAMETER_TEXTURE(Texture2D, FilmGrainTexture)
+	SHADER_PARAMETER_RDG_TEXTURE(Texture2D, FilmGrainTexture)
 	SHADER_PARAMETER_SAMPLER(SamplerState, FilmGrainSampler)
 	SHADER_PARAMETER(FScreenTransform, ScreenPosToFilmGrainTextureUV)
+	SHADER_PARAMETER_RDG_BUFFER_SRV(StructuredBuffer<float4>, FilmGrainTextureConstants)
 END_SHADER_PARAMETER_STRUCT()
 
 BEGIN_SHADER_PARAMETER_STRUCT(FTonemapParameters, )
@@ -363,6 +378,42 @@ BEGIN_SHADER_PARAMETER_STRUCT(FTonemapParameters, )
 	SHADER_PARAMETER_RDG_BUFFER_SRV(Buffer<float4>, EyeAdaptationBuffer)
 END_SHADER_PARAMETER_STRUCT()
 
+class FFilmGrainReduceCS : public FGlobalShader
+{
+public:
+	DECLARE_GLOBAL_SHADER(FFilmGrainReduceCS);
+	SHADER_USE_PARAMETER_STRUCT(FFilmGrainReduceCS, FGlobalShader);
+
+	BEGIN_SHADER_PARAMETER_STRUCT(FParameters, )
+		SHADER_PARAMETER(FIntPoint, FilmGrainTextureSize)
+		SHADER_PARAMETER_RDG_TEXTURE(Texture2D, FilmGrainTexture)
+		SHADER_PARAMETER_RDG_TEXTURE_UAV(RWTexture2D<float4>, Output)
+	END_SHADER_PARAMETER_STRUCT()
+	
+	static bool ShouldCompilePermutation(const FGlobalShaderPermutationParameters& Parameters)
+	{
+		return SupportsFilmGrain(Parameters.Platform);
+	}
+};
+
+class FFilmGrainPackConstantsCS : public FGlobalShader
+{
+public:
+	DECLARE_GLOBAL_SHADER(FFilmGrainPackConstantsCS);
+	SHADER_USE_PARAMETER_STRUCT(FFilmGrainPackConstantsCS, FGlobalShader);
+
+	BEGIN_SHADER_PARAMETER_STRUCT(FParameters, )
+		SHADER_PARAMETER(FIntPoint, OriginalFilmGrainTextureSize)
+		SHADER_PARAMETER_RDG_TEXTURE(Texture2D, ReducedFilmGrainTexture)
+		SHADER_PARAMETER_RDG_BUFFER_UAV(RWTexture2D<float4>, FilmGrainConstantsOutput)
+	END_SHADER_PARAMETER_STRUCT()
+	
+	static bool ShouldCompilePermutation(const FGlobalShaderPermutationParameters& Parameters)
+	{
+		return SupportsFilmGrain(Parameters.Platform);
+	}
+};
+
 class FTonemapVS : public FGlobalShader
 {
 public:
@@ -386,8 +437,6 @@ public:
 	}
 };
 
-IMPLEMENT_GLOBAL_SHADER(FTonemapVS, "/Engine/Private/PostProcessTonemap.usf", "MainVS", SF_Vertex);
-
 class FTonemapPS : public FGlobalShader
 {
 public:
@@ -416,11 +465,10 @@ public:
 		OutEnvironment.SetDefine(TEXT("USE_VOLUME_LUT"), UseVolumeLut);
 
 		OutEnvironment.SetDefine(TEXT("SUPPORTS_BLOOM_APPLY_PARAMETERS"), FBloomOutputs::SupportsApplyParametersBuffer(Parameters.Platform) ? 1 : 0);
+		OutEnvironment.SetDefine(TEXT("SUPPORTS_FILM_GRAIN"), SupportsFilmGrain(Parameters.Platform) ? 1 : 0);
 	}
 };
 
-IMPLEMENT_GLOBAL_SHADER(FTonemapPS, "/Engine/Private/PostProcessTonemap.usf", "MainPS", SF_Pixel);
-
 class FTonemapCS : public FGlobalShader
 {
 public:
@@ -457,11 +505,77 @@ public:
 		OutEnvironment.SetDefine(TEXT("USE_VOLUME_LUT"), UseVolumeLut);
 
 		OutEnvironment.SetDefine(TEXT("SUPPORTS_BLOOM_APPLY_PARAMETERS"), FBloomOutputs::SupportsApplyParametersBuffer(Parameters.Platform) ? 1 : 0);
+		OutEnvironment.SetDefine(TEXT("SUPPORTS_FILM_GRAIN"), SupportsFilmGrain(Parameters.Platform) ? 1 : 0);
 	}
 };
 
+IMPLEMENT_GLOBAL_SHADER(FFilmGrainReduceCS,        "/Engine/Private/PostProcessing/FilmGrainReduce.usf", "MainCS", SF_Compute);
+IMPLEMENT_GLOBAL_SHADER(FFilmGrainPackConstantsCS, "/Engine/Private/PostProcessing/FilmGrainPackConstants.usf", "MainCS", SF_Compute);
+IMPLEMENT_GLOBAL_SHADER(FTonemapVS, "/Engine/Private/PostProcessTonemap.usf", "MainVS", SF_Vertex);
+IMPLEMENT_GLOBAL_SHADER(FTonemapPS, "/Engine/Private/PostProcessTonemap.usf", "MainPS", SF_Pixel);
 IMPLEMENT_GLOBAL_SHADER(FTonemapCS, "/Engine/Private/PostProcessTonemap.usf", "MainCS", SF_Compute);
 
+static
+FRDGBufferRef BuildFilmGrainConstants(FRDGBuilder& GraphBuilder, const FViewInfo& View, FRDGTextureRef FilmGrainTexture)
+{
+	RDG_EVENT_SCOPE(GraphBuilder, "FilmGrain BuildTextureConstants");
+
+	FRDGTextureRef ReducedFilmGrainTexture = FilmGrainTexture;
+
+	for (int32 PassId = 0; ReducedFilmGrainTexture->Desc.Extent.X > 1 && ReducedFilmGrainTexture->Desc.Extent.Y > 1; PassId++)
+	{
+		FRDGTextureRef NewReducedFilmGrainTexture;
+		{
+			FRDGTextureDesc Desc = FRDGTextureDesc::Create2D(
+				FIntPoint::DivideAndRoundUp(ReducedFilmGrainTexture->Desc.Extent, 8),
+				PF_A32B32G32R32F,
+				FClearValueBinding::None,
+				TexCreate_ShaderResource | TexCreate_UAV);
+
+			NewReducedFilmGrainTexture = GraphBuilder.CreateTexture(Desc, TEXT("FilmGrain.Reduce"));
+		}
+
+		FFilmGrainReduceCS::FParameters* PassParameters = GraphBuilder.AllocParameters<FFilmGrainReduceCS::FParameters>();
+		PassParameters->FilmGrainTextureSize = ReducedFilmGrainTexture->Desc.Extent;
+		PassParameters->FilmGrainTexture = ReducedFilmGrainTexture;
+		PassParameters->Output = GraphBuilder.CreateUAV(NewReducedFilmGrainTexture);
+
+		TShaderMapRef<FFilmGrainReduceCS> ComputeShader(View.ShaderMap);
+		FComputeShaderUtils::AddPass(
+			GraphBuilder,
+			RDG_EVENT_NAME("FilmGrain Reduce %dx%d -> %dx%d",
+				ReducedFilmGrainTexture->Desc.Extent.X, ReducedFilmGrainTexture->Desc.Extent.Y,
+				NewReducedFilmGrainTexture->Desc.Extent.X, NewReducedFilmGrainTexture->Desc.Extent.Y),
+			ComputeShader,
+			PassParameters,
+			FIntVector(NewReducedFilmGrainTexture->Desc.Extent.X, NewReducedFilmGrainTexture->Desc.Extent.Y, 1));
+
+		ReducedFilmGrainTexture = NewReducedFilmGrainTexture;
+	}
+
+	FRDGBufferRef FilmGrainConstantsBuffer;
+	{
+		FilmGrainConstantsBuffer = GraphBuilder.CreateBuffer(
+			FRDGBufferDesc::CreateStructuredDesc(sizeof(FLinearColor), /* NumElements = */ 1),
+			TEXT("FilmGrain.TextureConstants"));
+
+		FFilmGrainPackConstantsCS::FParameters* PassParameters = GraphBuilder.AllocParameters<FFilmGrainPackConstantsCS::FParameters>();
+		PassParameters->OriginalFilmGrainTextureSize = FilmGrainTexture->Desc.Extent;
+		PassParameters->ReducedFilmGrainTexture = ReducedFilmGrainTexture;
+		PassParameters->FilmGrainConstantsOutput = GraphBuilder.CreateUAV(FilmGrainConstantsBuffer);
+
+		TShaderMapRef<FFilmGrainPackConstantsCS> ComputeShader(View.ShaderMap);
+		FComputeShaderUtils::AddPass(
+			GraphBuilder,
+			RDG_EVENT_NAME("FilmGrain PackConstants"),
+			ComputeShader,
+			PassParameters,
+			FIntVector(1, 1, 1));
+	}
+
+	return FilmGrainConstantsBuffer;
+}
+
 FScreenPassTexture AddTonemapPass(FRDGBuilder& GraphBuilder, const FViewInfo& View, const FTonemapInputs& Inputs)
 {
 	if (!Inputs.bGammaOnly)
@@ -577,18 +691,37 @@ FScreenPassTexture AddTonemapPass(FRDGBuilder& GraphBuilder, const FViewInfo& Vi
 
 	if (View.FilmGrainTexture)
 	{
-		check(View.FilmGrainTexture);
+		check(SupportsFilmGrain(View.GetShaderPlatform()));
+
+		FRHITexture* FilmGrainTextureRHI = View.FilmGrainTexture->GetTexture2DRHI();
+		FRDGTextureRef FilmGrainTexture = RegisterExternalTexture(GraphBuilder, FilmGrainTextureRHI, TEXT("FilmGrain.OriginalTexture"));
+
+
+		FRDGBufferRef FilmGrainConstantsBuffer;
+		if (View.ViewState && View.ViewState->FilmGrainCache.TextureRHI == FilmGrainTextureRHI && CVarFilmGrainCacheTextureConstants.GetValueOnRenderThread() != 0)
+		{
+			FilmGrainConstantsBuffer = GraphBuilder.RegisterExternalBuffer(View.ViewState->FilmGrainCache.ConstantsBuffer, TEXT("FilmGrain.TextureConstants"));
+		}
+		else
+		{
+			FilmGrainConstantsBuffer = BuildFilmGrainConstants(GraphBuilder, View, FilmGrainTexture);
+
+			if (View.ViewState)
+			{
+				View.ViewState->FilmGrainCache.Texture = View.FinalPostProcessSettings.FilmGrainTexture;
+				View.ViewState->FilmGrainCache.TextureRHI = FilmGrainTextureRHI;
+				GraphBuilder.QueueBufferExtraction(FilmGrainConstantsBuffer, &View.ViewState->FilmGrainCache.ConstantsBuffer);
+			}
+		}
 
 		// (FilmGrainTexture * FilmGrainDecodeMultiply + FilmGrainDecodeAdd)
-		CommonParameters.FilmGrain.FilmGrainMul = View.FinalPostProcessSettings.FilmGrainDecodeMultiply;
-		CommonParameters.FilmGrain.FilmGrainAdd = View.FinalPostProcessSettings.FilmGrainDecodeAdd;
 		CommonParameters.FilmGrain.FilmGrainIntensityShadows    = View.FinalPostProcessSettings.FilmGrainIntensity * View.FinalPostProcessSettings.FilmGrainIntensityShadows;
 		CommonParameters.FilmGrain.FilmGrainIntensityMidtones   = View.FinalPostProcessSettings.FilmGrainIntensity * View.FinalPostProcessSettings.FilmGrainIntensityMidtones;
 		CommonParameters.FilmGrain.FilmGrainIntensityHighlights = View.FinalPostProcessSettings.FilmGrainIntensity * View.FinalPostProcessSettings.FilmGrainIntensityHighlights;
 		CommonParameters.FilmGrain.FilmGrainShadowsMax = View.FinalPostProcessSettings.FilmGrainShadowsMax;
 		CommonParameters.FilmGrain.FilmGrainHighlightsMin = View.FinalPostProcessSettings.FilmGrainHighlightsMin;
 
-		CommonParameters.FilmGrain.FilmGrainTexture = View.FilmGrainTexture->GetTexture2DRHI();
+		CommonParameters.FilmGrain.FilmGrainTexture = FilmGrainTexture;
 		CommonParameters.FilmGrain.FilmGrainSampler = TStaticSamplerState<SF_Bilinear, AM_Wrap, AM_Wrap>::GetRHI();
 
 		int32 RandomSequenceLength = CVarFilmGrainSequenceLength.GetValueOnRenderThread();
@@ -604,19 +737,38 @@ FScreenPassTexture AddTonemapPass(FRDGBuilder& GraphBuilder, const FViewInfo& Vi
 		CommonParameters.FilmGrain.ScreenPosToFilmGrainTextureUV = 
 			FScreenTransform::ScreenPosToViewportUV *
 			((OutputSizeF / TextureSize) * (1.0f / View.FinalPostProcessSettings.FilmGrainTexelSize)) + RandomGrainTextureUVOffset;
+
+		CommonParameters.FilmGrain.FilmGrainTextureConstants = GraphBuilder.CreateSRV(FilmGrainConstantsBuffer);
 	}
 	else
 	{
-		CommonParameters.FilmGrain.FilmGrainMul = 0.0f;
-		CommonParameters.FilmGrain.FilmGrainAdd = 1.0f;
+		// Release the film grain cache
+		if (View.ViewState)
+		{
+			View.ViewState->FilmGrainCache.SafeRelease();
+		}
+
 		CommonParameters.FilmGrain.FilmGrainIntensityShadows = 0.0f;
 		CommonParameters.FilmGrain.FilmGrainIntensityMidtones = 0.0f;
 		CommonParameters.FilmGrain.FilmGrainIntensityHighlights = 0.0f;
 		CommonParameters.FilmGrain.FilmGrainShadowsMax = 0.09f;
 		CommonParameters.FilmGrain.FilmGrainHighlightsMin = 0.5f;
-		CommonParameters.FilmGrain.FilmGrainTexture = GWhiteTexture->TextureRHI;
+		CommonParameters.FilmGrain.FilmGrainTexture = GSystemTextures.GetWhiteDummy(GraphBuilder);
 		CommonParameters.FilmGrain.FilmGrainSampler = TStaticSamplerState<SF_Bilinear, AM_Wrap, AM_Wrap>::GetRHI();
 		CommonParameters.FilmGrain.ScreenPosToFilmGrainTextureUV = FScreenTransform::ScreenPosToViewportUV;
+
+		// Set a dummy texture constants
+		if (SupportsFilmGrain(View.GetShaderPlatform()))
+		{
+			FRDGBufferRef DummyFilmGrainTextureConstants = GSystemTextures.GetDefaultStructuredBuffer(
+				GraphBuilder, sizeof(FVector4), FVector4(FLinearColor::White));
+
+			CommonParameters.FilmGrain.FilmGrainTextureConstants = GraphBuilder.CreateSRV(DummyFilmGrainTextureConstants);
+		}
+		else
+		{
+			CommonParameters.FilmGrain.FilmGrainTextureConstants = nullptr;
+		}
 	}
 
 	CommonParameters.OutputDevice = GetTonemapperOutputDeviceParameters(ViewFamily);
diff --git a/Engine/Source/Runtime/Renderer/Private/PostProcess/PostProcessTonemap.h b/Engine/Source/Runtime/Renderer/Private/PostProcess/PostProcessTonemap.h
--- a/Engine/Source/Runtime/Renderer/Private/PostProcess/PostProcessTonemap.h
+++ b/Engine/Source/Runtime/Renderer/Private/PostProcess/PostProcessTonemap.h
@@ -8,6 +8,9 @@
 #include "OverridePassSequence.h"
 #include "Math/Halton.h"
 
+
+bool SupportsFilmGrain(EShaderPlatform Platform);
+
 // You must update values in PostProcessTonemap.usf when changing this enum.
 enum class ETonemapperOutputDevice
 {
@@ -75,35 +78,3 @@ struct FTonemapInputs
 };
 
 FScreenPassTexture AddTonemapPass(FRDGBuilder& GraphBuilder, const FViewInfo& View, const FTonemapInputs& Inputs);
-
-struct FMobileTonemapperInputs
-{
-	// [Optional] Render to the specified output. If invalid, a new texture is created and returned.
-	FScreenPassRenderTarget OverrideOutput;
-
-	// [Required] HDR scene color to tonemap.
-	FScreenPassTexture SceneColor;
-
-	// [Required] Filtered bloom texture to composite with tonemapped scene color. This should be transparent black for no bloom.
-	FScreenPassTexture BloomOutput;
-
-	FScreenPassTexture DofOutput;
-
-	FScreenPassTexture SunShaftAndDof;
-
-	FRDGBufferRef EyeAdaptationBuffer = nullptr;
-
-	// [Raster Only, Mobile] Flips the image vertically on output.
-	bool bFlipYAxis = false;
-
-	// Whether to leave the final output in HDR.
-	bool bOutputInHDR = false;
-
-	bool bMetalMSAAHDRDecode = false;
-
-	bool bUseEyeAdaptation = false;
-
-	bool bSRGBAwareTarget = false;
-};
-
-FScreenPassTexture AddMobileTonemapperPass(FRDGBuilder& GraphBuilder, const FViewInfo& View, const FMobileTonemapperInputs& Inputs);
diff --git a/Engine/Source/Runtime/Renderer/Private/ScenePrivate.h b/Engine/Source/Runtime/Renderer/Private/ScenePrivate.h
--- a/Engine/Source/Runtime/Renderer/Private/ScenePrivate.h
+++ b/Engine/Source/Runtime/Renderer/Private/ScenePrivate.h
@@ -965,6 +965,8 @@ public:
 		/// @cond DOXYGEN_WARNINGS
 		void SafeRelease()
 		{
+			Physical = nullptr;
+			PhysicalRHI = nullptr;
 			Spectral.SafeRelease();
 			ConstantsBuffer.SafeRelease();
 		}
@@ -990,6 +992,23 @@ public:
 		uint32 PhysicalMipLevel;
 	} BloomFFTKernel;
 
+	// Film grain
+	struct
+	{
+		/// @cond DOXYGEN_WARNINGS
+		void SafeRelease()
+		{
+			Texture = nullptr;
+			TextureRHI = nullptr;
+			ConstantsBuffer.SafeRelease();
+		}
+		/// @endcond
+
+		UTexture2D* Texture = nullptr;
+		FRHITexture* TextureRHI = nullptr;
+		TRefCountPtr<FRDGPooledBuffer> ConstantsBuffer;
+	} FilmGrainCache;
+
 	// Cached material texture samplers
 	float MaterialTextureCachedMipBias;
 	FSamplerStateRHIRef MaterialTextureBilinearWrapedSamplerCache;
@@ -1376,6 +1395,11 @@ public:
 		{
 			Collector.AddReferencedObject(BloomFFTKernel.Physical);
 		}
+
+		if (FilmGrainCache.Texture)
+		{
+			Collector.AddReferencedObject(FilmGrainCache.Texture);
+		}
 	}
 
 	// needed for GetReusableMID()
diff --git a/Engine/Source/Runtime/Renderer/Private/SceneRendering.cpp b/Engine/Source/Runtime/Renderer/Private/SceneRendering.cpp
--- a/Engine/Source/Runtime/Renderer/Private/SceneRendering.cpp
+++ b/Engine/Source/Runtime/Renderer/Private/SceneRendering.cpp
@@ -26,6 +26,7 @@
 #include "PostProcess/TemporalAA.h"
 #include "PostProcess/PostProcessUpscale.h"
 #include "PostProcess/PostProcessing.h"
+#include "PostProcess/PostProcessTonemap.h"
 #include "CompositionLighting/CompositionLighting.h"
 #include "LegacyScreenPercentageDriver.h"
 #include "SceneViewExtension.h"
@@ -337,6 +338,11 @@ static FAutoConsoleVariableRef CVarParallelCmdListInheritBreadcrumbs(
 	ECVF_ReadOnly
 );
 
+static TAutoConsoleVariable<int32> CVarFilmGrain(
+	TEXT("r.FilmGrain"), 1,
+	TEXT("Whether to enable film grain."),
+	ECVF_RenderThreadSafe);
+
 #if !UE_BUILD_SHIPPING
 
 static TAutoConsoleVariable<int32> CVarTestInternalViewRectOffset(
@@ -447,7 +453,6 @@ bool IsStaticLightingAllowed()
 	return CVar->GetValueOnRenderThread() != 0;
 }
 
-
 TSharedPtr<FVirtualShadowMapClipmap> FVisibleLightInfo::FindShadowClipmapForView(const FViewInfo* View) const
 {
 	for (const auto& Clipmap : VirtualShadowMapClipmaps)
@@ -2385,14 +2390,16 @@ FSceneRenderer::FSceneRenderer(const FSceneViewFamily* InViewFamily,FHitProxyCon
 		}
 
 		// Handle the film grain texture
-		if (ViewInfo->FinalPostProcessSettings.FilmGrainIntensity > 0.0f)
+		if (ViewInfo->FinalPostProcessSettings.FilmGrainIntensity > 0.0f &&
+			ViewFamily.EngineShowFlags.Grain &&
+			CVarFilmGrain.GetValueOnGameThread() != 0 &&
+			SupportsFilmGrain(ViewFamily.GetShaderPlatform()))
 		{
-			GEngine->LoadDefaultFilmGrainTexture();
-
-			UTexture2D* FilmGrainTexture = GEngine->DefaultFilmGrainTexture;
-			if (ViewInfo->FinalPostProcessSettings.FilmGrainTexture)
+			UTexture2D* FilmGrainTexture = ViewInfo->FinalPostProcessSettings.FilmGrainTexture;
+			if (FilmGrainTexture == nullptr)
 			{
-				FilmGrainTexture = ViewInfo->FinalPostProcessSettings.FilmGrainTexture;
+				GEngine->LoadDefaultFilmGrainTexture();
+				FilmGrainTexture = GEngine->DefaultFilmGrainTexture;
 			}
 
 			bool bIsValid = PreparePostProcessSettingTextureForRenderer(*ViewInfo, FilmGrainTexture, TEXT("film grain"));
@@ -2404,13 +2411,6 @@ FSceneRenderer::FSceneRenderer(const FSceneViewFamily* InViewFamily,FHitProxyCon
 				{
 					ViewInfo->FilmGrainTexture = TextureResource->GetTexture2DResource();
 				}
-
-				if (FilmGrainTexture == GEngine->DefaultFilmGrainTexture)
-				{
-					// Decode Texture2D'/Engine/EngineResources/FilmGrains/Marcie_Grain_v3_128_M2_000.Marcie_Grain_v3_128_M2_000'
-					ViewInfo->FinalPostProcessSettings.FilmGrainDecodeMultiply = 2.000f;
-					ViewInfo->FinalPostProcessSettings.FilmGrainDecodeAdd = 0.0f;
-				}
 			}
 		}
 	}
