Fixes warning in FilmGrain

diff --git a/Engine/Source/Runtime/Renderer/Private/PostProcess/PostProcessTonemap.cpp b/Engine/Source/Runtime/Renderer/Private/PostProcess/PostProcessTonemap.cpp
--- a/Engine/Source/Runtime/Renderer/Private/PostProcess/PostProcessTonemap.cpp
+++ b/Engine/Source/Runtime/Renderer/Private/PostProcess/PostProcessTonemap.cpp
@@ -405,7 +405,7 @@ public:
 	BEGIN_SHADER_PARAMETER_STRUCT(FParameters, )
 		SHADER_PARAMETER(FIntPoint, OriginalFilmGrainTextureSize)
 		SHADER_PARAMETER_RDG_TEXTURE(Texture2D, ReducedFilmGrainTexture)
-		SHADER_PARAMETER_RDG_BUFFER_UAV(RWTexture2D<float4>, FilmGrainConstantsOutput)
+		SHADER_PARAMETER_RDG_BUFFER_UAV(RWStructuredBuffer<float4>, FilmGrainConstantsOutput)
 	END_SHADER_PARAMETER_STRUCT()
 	
 	static bool ShouldCompilePermutation(const FGlobalShaderPermutationParameters& Parameters)
