Static analysis fix

diff --git a/Engine/Source/Runtime/Renderer/Private/VT/RuntimeVirtualTextureRender.cpp b/Engine/Source/Runtime/Renderer/Private/VT/RuntimeVirtualTextureRender.cpp
--- a/Engine/Source/Runtime/Renderer/Private/VT/RuntimeVirtualTextureRender.cpp
+++ b/Engine/Source/Runtime/Renderer/Private/VT/RuntimeVirtualTextureRender.cpp
@@ -919,14 +919,15 @@ namespace RuntimeVirtualTexture
 			bRenderPass = OutputTexture0 != nullptr;
 			bCopyThumbnailPass = bRenderPass && bIsThumbnails;
 			EPixelFormat OutputFormat = (OutputTexture0 != nullptr ? OutputTexture0->GetFormat() : PF_Unknown);
-			bool bCompressedFormat = GPixelFormats[OutputFormat].BlockSizeX == 4 && GPixelFormats[OutputFormat].BlockSizeY == 4;
-			bool bLQFormat = (OutputTexture0 && OutputTexture0->GetFormat() == PF_R5G6B5_UNORM);
+			const bool bCompressedFormat = GPixelFormats[OutputFormat].BlockSizeX == 4 && GPixelFormats[OutputFormat].BlockSizeY == 4;
+			const bool bLQFormat = OutputTexture0 != nullptr && OutputTexture0->GetFormat() == PF_R5G6B5_UNORM;
 			bCompressPass = bRenderPass && !bCopyThumbnailPass && bCompressedFormat;
 			bCopyPass = bRenderPass && !bCopyThumbnailPass && !bCompressPass && (MaterialType == ERuntimeVirtualTextureMaterialType::BaseColor_Normal_Specular || MaterialType == ERuntimeVirtualTextureMaterialType::BaseColor_Normal_Specular_YCoCg || MaterialType == ERuntimeVirtualTextureMaterialType::BaseColor_Normal_Specular_Mask_YCoCg);
 			// Not all mobile RHIs support sRGB texture views/aliasing, use only linear targets on mobile
-			ETextureCreateFlags VT_SRGB = FeatureLevel > ERHIFeatureLevel::ES3_1 ? TexCreate_SRGB : TexCreate_None;
+			const ETextureCreateFlags VT_SRGB = FeatureLevel > ERHIFeatureLevel::ES3_1 ? TexCreate_SRGB : TexCreate_None;
 			// GLES does not support writing to 2-channel images, instead pack RG32 into RGBA16 on all GL
-			EPixelFormat Compressed64BitFormat = IsOpenGLPlatform(GMaxRHIShaderPlatform) ? PF_R16G16B16A16_UINT : PF_R32G32_UINT;
+			const EPixelFormat Compressed64BitFormat = IsOpenGLPlatform(GMaxRHIShaderPlatform) ? PF_R16G16B16A16_UINT : PF_R32G32_UINT;
+			const EPixelFormat Compressed128BitFormat = PF_R32G32B32A32_UINT;
 
 			switch (MaterialType)
 			{
@@ -953,7 +954,7 @@ namespace RuntimeVirtualTexture
 				if (bCompressPass)
 				{
 					OutputAlias0 = CompressTexture0_u2 = GraphBuilder.CreateTexture(FRDGTextureDesc::Create2D(TextureSize / 4, Compressed64BitFormat, FClearValueBinding::None, TexCreate_UAV), TEXT("CompressTexture0"));
-					OutputAlias1 = CompressTexture1 = GraphBuilder.CreateTexture(FRDGTextureDesc::Create2D(TextureSize / 4, PF_R32G32B32A32_UINT, FClearValueBinding::None, TexCreate_UAV), TEXT("CompressTexture1"));
+					OutputAlias1 = CompressTexture1 = GraphBuilder.CreateTexture(FRDGTextureDesc::Create2D(TextureSize / 4, Compressed128BitFormat, FClearValueBinding::None, TexCreate_UAV), TEXT("CompressTexture1"));
 				}
 				if (bCopyThumbnailPass)
 				{
@@ -970,8 +971,8 @@ namespace RuntimeVirtualTexture
 				}
 				if (bCompressPass)
 				{
-					OutputAlias0 = CompressTexture0_u4 = GraphBuilder.CreateTexture(FRDGTextureDesc::Create2D(TextureSize / 4, PF_R32G32B32A32_UINT, FClearValueBinding::None, TexCreate_UAV), TEXT("CompressTexture0"));
-					OutputAlias1 = CompressTexture1 = GraphBuilder.CreateTexture(FRDGTextureDesc::Create2D(TextureSize / 4, PF_R32G32B32A32_UINT, FClearValueBinding::None, TexCreate_UAV), TEXT("CompressTexture1"));
+					OutputAlias0 = CompressTexture0_u4 = GraphBuilder.CreateTexture(FRDGTextureDesc::Create2D(TextureSize / 4, Compressed128BitFormat, FClearValueBinding::None, TexCreate_UAV), TEXT("CompressTexture0"));
+					OutputAlias1 = CompressTexture1 = GraphBuilder.CreateTexture(FRDGTextureDesc::Create2D(TextureSize / 4, Compressed128BitFormat, FClearValueBinding::None, TexCreate_UAV), TEXT("CompressTexture1"));
 				}
 				if (bCopyPass)
 				{
@@ -992,8 +993,8 @@ namespace RuntimeVirtualTexture
 				}
 				if (bCompressPass)
 				{
-					OutputAlias0 = CompressTexture0_u4 = GraphBuilder.CreateTexture(FRDGTextureDesc::Create2D(TextureSize / 4, PF_R32G32B32A32_UINT, FClearValueBinding::None, TexCreate_UAV), TEXT("CompressTexture0"));
-					OutputAlias1 = CompressTexture1 = GraphBuilder.CreateTexture(FRDGTextureDesc::Create2D(TextureSize / 4, PF_R32G32B32A32_UINT, FClearValueBinding::None, TexCreate_UAV), TEXT("CompressTexture1"));
+					OutputAlias0 = CompressTexture0_u4 = GraphBuilder.CreateTexture(FRDGTextureDesc::Create2D(TextureSize / 4, Compressed128BitFormat, FClearValueBinding::None, TexCreate_UAV), TEXT("CompressTexture0"));
+					OutputAlias1 = CompressTexture1 = GraphBuilder.CreateTexture(FRDGTextureDesc::Create2D(TextureSize / 4, Compressed128BitFormat, FClearValueBinding::None, TexCreate_UAV), TEXT("CompressTexture1"));
 					OutputAlias2 = CompressTexture2_u2 = GraphBuilder.CreateTexture(FRDGTextureDesc::Create2D(TextureSize / 4, Compressed64BitFormat, FClearValueBinding::None, TexCreate_UAV), TEXT("CompressTexture2"));
 				}
 				if (bCopyPass)
@@ -1016,9 +1017,9 @@ namespace RuntimeVirtualTexture
 				}
 				if (bCompressPass)
 				{
-					OutputAlias0 = CompressTexture0_u4 = GraphBuilder.CreateTexture(FRDGTextureDesc::Create2D(TextureSize / 4, PF_R32G32B32A32_UINT, FClearValueBinding::None, TexCreate_UAV), TEXT("CompressTexture0"));
-					OutputAlias1 = CompressTexture1 = GraphBuilder.CreateTexture(FRDGTextureDesc::Create2D(TextureSize / 4, PF_R32G32B32A32_UINT, FClearValueBinding::None, TexCreate_UAV), TEXT("CompressTexture1"));
-					OutputAlias2 = CompressTexture2_u4 = GraphBuilder.CreateTexture(FRDGTextureDesc::Create2D(TextureSize / 4, PF_R32G32B32A32_UINT, FClearValueBinding::None, TexCreate_UAV), TEXT("CompressTexture2"));
+					OutputAlias0 = CompressTexture0_u4 = GraphBuilder.CreateTexture(FRDGTextureDesc::Create2D(TextureSize / 4, Compressed128BitFormat, FClearValueBinding::None, TexCreate_UAV), TEXT("CompressTexture0"));
+					OutputAlias1 = CompressTexture1 = GraphBuilder.CreateTexture(FRDGTextureDesc::Create2D(TextureSize / 4, Compressed128BitFormat, FClearValueBinding::None, TexCreate_UAV), TEXT("CompressTexture1"));
+					OutputAlias2 = CompressTexture2_u4 = GraphBuilder.CreateTexture(FRDGTextureDesc::Create2D(TextureSize / 4, Compressed128BitFormat, FClearValueBinding::None, TexCreate_UAV), TEXT("CompressTexture2"));
 				}
 				if (bCopyPass)
 				{
@@ -1231,15 +1232,15 @@ namespace RuntimeVirtualTexture
 		}
 
 		// Copy to Output for each output texture
-		if (GraphSetup.OutputAlias0 != nullptr)
+		if (GraphSetup.OutputAlias0 != nullptr && OutputTexture0 != nullptr)
 		{
 			AddCopyToOutputPass(GraphBuilder, GraphSetup.OutputAlias0, OutputTexture0, DestBox0);
 		}
-		if (GraphSetup.OutputAlias1 != nullptr)
+		if (GraphSetup.OutputAlias1 != nullptr && OutputTexture1 != nullptr)
 		{
 			AddCopyToOutputPass(GraphBuilder, GraphSetup.OutputAlias1, OutputTexture1, DestBox1);
 		}
-		if (GraphSetup.OutputAlias2 != nullptr)
+		if (GraphSetup.OutputAlias2 != nullptr && OutputTexture2 != nullptr)
 		{
 			AddCopyToOutputPass(GraphBuilder, GraphSetup.OutputAlias2, OutputTexture2, DestBox2);
 		}
