Make GPU-Scene instance culling batching restartable and restart after each drain.
- Also make view data upload deferred to fix shadow culling now that init views is done after prepass.
- Fix deferred instance count clear pass (to use correct deferred argument count).

diff --git a/Engine/Source/Runtime/Renderer/Private/DeferredShadingRenderer.cpp b/Engine/Source/Runtime/Renderer/Private/DeferredShadingRenderer.cpp
--- a/Engine/Source/Runtime/Renderer/Private/DeferredShadingRenderer.cpp
+++ b/Engine/Source/Runtime/Renderer/Private/DeferredShadingRenderer.cpp
@@ -2189,8 +2189,22 @@ void FDeferredShadingSceneRenderer::Render(FRDGBuilder& GraphBuilder)
 		}
 	}
 
+	/**
+	 * Helper to perform RDG drain (essentially Execute() for the queued passes so far), and do any tasks that need doing at the start of a new 
+	 * section of RDG setup (currently, restarting deferred culling).
+	 */
+	auto ExecuteDrainPoint = [this, &GraphBuilder, &InstanceCullingManager]()
 	{
 		GraphBuilder.Drain();
+		if (GraphBuilder.IsDrainEnabled())
+		{
+			// This sets up a fresh deferred context that is used by all subsequent culling passes since the previous one is 'used up' by Drain.
+			InstanceCullingManager.BeginDeferredCulling(GraphBuilder, Scene->GPUScene);
+		}
+	};
+
+	{
+		ExecuteDrainPoint();
 
 		{
 			RDG_RHI_GPU_STAT_SCOPE(GraphBuilder, VisibilityCommands);
@@ -2381,7 +2395,7 @@ void FDeferredShadingSceneRenderer::Render(FRDGBuilder& GraphBuilder)
 			CompositionLightingAsyncResults = CompositionLighting::ProcessAsync(GraphBuilder, Views, SceneTextures);
 		}
 
-		GraphBuilder.Drain();
+		ExecuteDrainPoint();
 	};
 
 	// Early occlusion queries
diff --git a/Engine/Source/Runtime/Renderer/Private/InstanceCulling/InstanceCullingContext.cpp b/Engine/Source/Runtime/Renderer/Private/InstanceCulling/InstanceCullingContext.cpp
--- a/Engine/Source/Runtime/Renderer/Private/InstanceCulling/InstanceCullingContext.cpp
+++ b/Engine/Source/Runtime/Renderer/Private/InstanceCulling/InstanceCullingContext.cpp
@@ -22,12 +22,6 @@ static TAutoConsoleVariable<int32> CVarCullInstances(
 IMPLEMENT_STATIC_UNIFORM_BUFFER_SLOT(InstanceCullingUbSlot);
 IMPLEMENT_STATIC_UNIFORM_BUFFER_STRUCT(FInstanceCullingGlobalUniforms, "InstanceCulling", InstanceCullingUbSlot);
 
-static int32 GAllowBatchedBuildRenderingCommands = 1;
-static FAutoConsoleVariableRef CVarAllowBatchedBuildRenderingCommands(
-	TEXT("r.InstanceCulling.AllowBatchedBuildRenderingCommands"),
-	GAllowBatchedBuildRenderingCommands,
-	TEXT("Whether to allow batching BuildRenderingCommands for GPU instance culling"),
-	ECVF_RenderThreadSafe);
 
 
 static const TCHAR* BatchProcessingModeStr[] =
@@ -225,6 +219,57 @@ const TRDGUniformBufferRef<FInstanceCullingGlobalUniforms> FInstanceCullingConte
 	return GraphBuilder.CreateUniformBuffer(InstanceCullingGlobalUniforms);
 }
 
+
+class FInstanceCullingDeferredContext
+{
+public:
+	FInstanceCullingDeferredContext(FInstanceCullingManager* InInstanceCullingManager = nullptr) : InstanceCullingManager(InInstanceCullingManager) {}
+
+	FInstanceCullingManager* InstanceCullingManager;
+
+	FRDGBufferRef DrawIndirectArgsBuffer = nullptr;
+	FRDGBufferRef InstanceIdOffsetBuffer = nullptr;
+	TRDGUniformBufferRef<FInstanceCullingGlobalUniforms> UniformBuffer = nullptr;
+
+	/** Batches of GPU instance culling input data. */
+	TArray<FInstanceCullingContext::FBatchItem, SceneRenderingAllocator> Batches;
+
+	TArray<int32, SceneRenderingAllocator> ViewIds;
+	//TArray<FMeshDrawCommandInfo, SceneRenderingAllocator> MeshDrawCommandInfos;
+	TArray<FRHIDrawIndexedIndirectParameters, SceneRenderingAllocator> IndirectArgs;
+	TArray<FInstanceCullingContext::FDrawCommandDesc, SceneRenderingAllocator> DrawCommandDescs;
+	TArray<uint32, SceneRenderingAllocator> InstanceIdOffsets;
+
+	FInstanceCullingContext::LoadBalancerArray LoadBalancers = FInstanceCullingContext::LoadBalancerArray(InPlace, nullptr);
+	TStaticArray<TArray<uint32, SceneRenderingAllocator>, static_cast<uint32>(EBatchProcessingMode::Num)> BatchInds;
+	TArray<FInstanceCullingContext::FContextBatchInfo, SceneRenderingAllocator> BatchInfos;
+
+
+	// Counters to sum up all sizes to facilitate pre-sizing
+	uint32 InstanceIdBufferSize = 0U;
+	TStaticArray<int32, uint32(EBatchProcessingMode::Num)> TotalBatches = TStaticArray<int32, uint32(EBatchProcessingMode::Num)>(InPlace, 0);
+	TStaticArray<int32, uint32(EBatchProcessingMode::Num)> TotalItems = TStaticArray<int32, uint32(EBatchProcessingMode::Num)>(InPlace, 0);
+	int32 TotalIndirectArgs = 0;
+	int32 TotalViewIds = 0;
+
+	bool bProcessed = false;
+
+	void ProcessBatched(TStaticArray<FBuildInstanceIdBufferAndCommandsFromPrimitiveIdsCs::FParameters*, static_cast<uint32>(EBatchProcessingMode::Num)> PassParameters);
+
+
+	~FInstanceCullingDeferredContext()
+	{
+		for (auto& LoadBalancer : LoadBalancers)
+		{
+			if (LoadBalancer != nullptr)
+			{
+				delete LoadBalancer;
+			}
+		}
+	}
+};
+
+
 void FInstanceCullingContext::BuildRenderingCommands(
 	FRDGBuilder& GraphBuilder,
 	const FGPUScene& GPUScene,
@@ -247,24 +292,28 @@ void FInstanceCullingContext::BuildRenderingCommands(
 	const uint32 InstanceIdBufferSize = TotalInstances * ViewIds.Num();
 	if (InstanceCullingDrawParams && InstanceCullingManager && InstanceCullingManager->IsDeferredCullingActive())
 	{
-		auto& Scratch = InstanceCullingManager->BatchedCullingScratch;
-		auto& MergedContext = Scratch.MergedContext;
-		Results.DrawIndirectArgsBuffer = Scratch.DrawIndirectArgsBuffer;
-		Results.InstanceIdOffsetBuffer = Scratch.InstanceIdOffsetBuffer;
-		Results.UniformBuffer = Scratch.UniformBuffer;
-		Scratch.Batches.Add(FBatchItem{ this, InstanceCullingDrawParams, DynamicInstanceIdOffset, DynamicInstanceIdNum });
-
+		FInstanceCullingDeferredContext *DeferredContext = InstanceCullingManager->DeferredContext;
 
-		// Accumulate the totals so the deferred processing can pre-size the arrays
-		for (uint32 Mode = 0U; Mode < uint32(EBatchProcessingMode::Num); ++Mode)
+		// If this is true, then RDG Execute or Drain has been called, and no further contexts can be deferred. 
+		if (!DeferredContext->bProcessed)
 		{
-			LoadBalancers[Mode]->FinalizeBatches();
-			MergedContext.TotalBatches[Mode] += LoadBalancers[Mode]->GetBatches().Max();
-			MergedContext.TotalItems[Mode] += LoadBalancers[Mode]->GetItems().Max();
+			Results.DrawIndirectArgsBuffer = DeferredContext->DrawIndirectArgsBuffer;
+			Results.InstanceIdOffsetBuffer = DeferredContext->InstanceIdOffsetBuffer;
+			Results.UniformBuffer = DeferredContext->UniformBuffer;
+			DeferredContext->Batches.Add(FBatchItem{ this, InstanceCullingDrawParams, DynamicInstanceIdOffset, DynamicInstanceIdNum });
+
+
+			// Accumulate the totals so the deferred processing can pre-size the arrays
+			for (uint32 Mode = 0U; Mode < uint32(EBatchProcessingMode::Num); ++Mode)
+			{
+				LoadBalancers[Mode]->FinalizeBatches();
+				DeferredContext->TotalBatches[Mode] += LoadBalancers[Mode]->GetBatches().Max();
+				DeferredContext->TotalItems[Mode] += LoadBalancers[Mode]->GetItems().Max();
+			}
+			DeferredContext->TotalIndirectArgs += IndirectArgs.Num();
+			DeferredContext->TotalViewIds += ViewIds.Num();
+			DeferredContext->InstanceIdBufferSize += InstanceIdBufferSize;
 		}
-		MergedContext.TotalIndirectArgs += IndirectArgs.Num();
-		MergedContext.TotalViewIds += ViewIds.Num();
-		MergedContext.InstanceIdBufferSize += InstanceIdBufferSize;
 		return;
 	}
 
@@ -273,7 +322,7 @@ void FInstanceCullingContext::BuildRenderingCommands(
 	// If there is no manager, then there is no data on culling, so set flag to skip that and ignore buffers.
 	const bool bCullInstances = InstanceCullingManager != nullptr && CVarCullInstances.GetValueOnRenderThread() != 0;
 
-	RDG_EVENT_SCOPE(GraphBuilder, "BuildRenderingCommandsDeferred(Culling=%s)", bCullInstances ? TEXT("On") : TEXT("Off"));
+	RDG_EVENT_SCOPE(GraphBuilder, "BuildRenderingCommands(Culling=%s)", bCullInstances ? TEXT("On") : TEXT("Off"));
 
 	const ERHIFeatureLevel::Type FeatureLevel = GPUScene.GetFeatureLevel();
 	FGlobalShaderMap* ShaderMap = GetGlobalShaderMap(FeatureLevel);
@@ -371,104 +420,99 @@ void FInstanceCullingContext::BuildRenderingCommands(
 	Results.UniformBuffer = GraphBuilder.CreateUniformBuffer(UniformParameters);
 }
 
-/** Helper function to merge GPU instance culling input batches queued throughout a frame. */
-static void MergeInstanceCullingContext(
-	const TArray<FInstanceCullingContext::FBatchItem, SceneRenderingAllocator>& Batches,
-	TStaticArray<FBuildInstanceIdBufferAndCommandsFromPrimitiveIdsCs::FParameters*, static_cast<uint32>(EBatchProcessingMode::Num)> PassParameters,
-	FInstanceCullingContext::FMergedContext& MergedContext)
+
+
+void FInstanceCullingDeferredContext::ProcessBatched(TStaticArray<FBuildInstanceIdBufferAndCommandsFromPrimitiveIdsCs::FParameters*, static_cast<uint32>(EBatchProcessingMode::Num)> PassParameters)
 {
-	if (!MergedContext.bInitialized)
+	if (bProcessed)
 	{
-		TStaticArray<int32, uint32(EBatchProcessingMode::Num)> TotalBatches = TStaticArray<int32, uint32(EBatchProcessingMode::Num)>(InPlace, 0);
-		TStaticArray<int32, uint32(EBatchProcessingMode::Num)> TotalItems = TStaticArray<int32, uint32(EBatchProcessingMode::Num)>(InPlace, 0);
-		int32 TotalIndirectArgs = 0;
-		int32 TotalViewIds = 0;
-
+		return;
+	}
 
-		MergedContext.bInitialized = true;
-		for (uint32 Mode = 0U; Mode < uint32(EBatchProcessingMode::Num); ++Mode)
-		{
-			MergedContext.LoadBalancers[Mode]= new FInstanceProcessingGPULoadBalancer;
-			MergedContext.LoadBalancers[Mode]->ReserveStorage(MergedContext.TotalBatches[Mode], MergedContext.TotalItems[Mode]);
-		}
-		// Pre-size all arrays
-		MergedContext.IndirectArgs.Empty(MergedContext.TotalIndirectArgs);
-		MergedContext.DrawCommandDescs.Empty(MergedContext.TotalIndirectArgs);
-		MergedContext.InstanceIdOffsets.Empty(MergedContext.TotalIndirectArgs);
-		MergedContext.ViewIds.Empty(MergedContext.TotalViewIds);
-
-		MergedContext.BatchInfos.AddDefaulted(Batches.Num());
-		uint32 InstanceIdBufferOffset = 0U;
-		uint32 InstanceDataByteOffset = 0U;
-
-		// Index that maps from each command to the corresponding batch - maybe not the utmost efficiency
-		for (int32 BatchIndex = 0; BatchIndex < Batches.Num(); ++BatchIndex)
-		{
-			const FInstanceCullingContext::FBatchItem& BatchItem = Batches[BatchIndex];
-			const FInstanceCullingContext& InstanceCullingContext = *BatchItem.Context;
+	bProcessed = true;
+	for (uint32 Mode = 0U; Mode < uint32(EBatchProcessingMode::Num); ++Mode)
+	{
+		LoadBalancers[Mode]= new FInstanceProcessingGPULoadBalancer;
+		LoadBalancers[Mode]->ReserveStorage(TotalBatches[Mode], TotalItems[Mode]);
+	}
+	// Pre-size all arrays
+	IndirectArgs.Empty(TotalIndirectArgs);
+	DrawCommandDescs.Empty(TotalIndirectArgs);
+	InstanceIdOffsets.Empty(TotalIndirectArgs);
+	ViewIds.Empty(TotalViewIds);
+
+	BatchInfos.AddDefaulted(Batches.Num());
+	uint32 InstanceIdBufferOffset = 0U;
+	uint32 InstanceDataByteOffset = 0U;
+
+	// Index that maps from each command to the corresponding batch - maybe not the utmost efficiency
+	for (int32 BatchIndex = 0; BatchIndex < Batches.Num(); ++BatchIndex)
+	{
+		const FInstanceCullingContext::FBatchItem& BatchItem = Batches[BatchIndex];
+		const FInstanceCullingContext& InstanceCullingContext = *BatchItem.Context;
 
-			FInstanceCullingContext::FContextBatchInfo& BatchInfo = MergedContext.BatchInfos[BatchIndex];
+		FInstanceCullingContext::FContextBatchInfo& BatchInfo = BatchInfos[BatchIndex];
 
-			BatchInfo.IndirectArgsOffset = MergedContext.IndirectArgs.Num();
-			//BatchInfo.NumIndirectArgs = InstanceCullingContext.IndirectArgs.Num();
-			MergedContext.IndirectArgs.Append(InstanceCullingContext.IndirectArgs);
+		BatchInfo.IndirectArgsOffset = IndirectArgs.Num();
+		//BatchInfo.NumIndirectArgs = InstanceCullingContext.IndirectArgs.Num();
+		IndirectArgs.Append(InstanceCullingContext.IndirectArgs);
 
-			check(InstanceCullingContext.DrawCommandDescs.Num() == InstanceCullingContext.IndirectArgs.Num());
-			MergedContext.DrawCommandDescs.Append(InstanceCullingContext.DrawCommandDescs);
+		check(InstanceCullingContext.DrawCommandDescs.Num() == InstanceCullingContext.IndirectArgs.Num());
+		DrawCommandDescs.Append(InstanceCullingContext.DrawCommandDescs);
 
-			check(InstanceCullingContext.InstanceIdOffsets.Num() == InstanceCullingContext.IndirectArgs.Num());
-			MergedContext.InstanceIdOffsets.AddDefaulted(InstanceCullingContext.InstanceIdOffsets.Num());
-			// TODO: perform offset on GPU
-			// MergedContext.InstanceIdOffsets.Append(InstanceCullingContext.InstanceIdOffsets);
-			for (int32 Index = 0; Index < InstanceCullingContext.InstanceIdOffsets.Num(); ++Index)
-			{
-				MergedContext.InstanceIdOffsets[BatchInfo.IndirectArgsOffset + Index] = InstanceCullingContext.InstanceIdOffsets[Index] + InstanceIdBufferOffset;
-			}
+		check(InstanceCullingContext.InstanceIdOffsets.Num() == InstanceCullingContext.IndirectArgs.Num());
+		InstanceIdOffsets.AddDefaulted(InstanceCullingContext.InstanceIdOffsets.Num());
+		// TODO: perform offset on GPU
+		// InstanceIdOffsets.Append(InstanceCullingContext.InstanceIdOffsets);
+		for (int32 Index = 0; Index < InstanceCullingContext.InstanceIdOffsets.Num(); ++Index)
+		{
+			InstanceIdOffsets[BatchInfo.IndirectArgsOffset + Index] = InstanceCullingContext.InstanceIdOffsets[Index] + InstanceIdBufferOffset;
+		}
 
-			BatchInfo.ViewIdsOffset = MergedContext.ViewIds.Num();
-			BatchInfo.NumViewIds = InstanceCullingContext.ViewIds.Num();
-			MergedContext.ViewIds.Append(InstanceCullingContext.ViewIds);
+		BatchInfo.ViewIdsOffset = ViewIds.Num();
+		BatchInfo.NumViewIds = InstanceCullingContext.ViewIds.Num();
+		ViewIds.Append(InstanceCullingContext.ViewIds);
 
-			BatchInfo.DynamicInstanceIdOffset = BatchItem.DynamicInstanceIdOffset;
-			BatchInfo.DynamicInstanceIdMax = BatchItem.DynamicInstanceIdOffset + BatchItem.DynamicInstanceIdNum;
+		BatchInfo.DynamicInstanceIdOffset = BatchItem.DynamicInstanceIdOffset;
+		BatchInfo.DynamicInstanceIdMax = BatchItem.DynamicInstanceIdOffset + BatchItem.DynamicInstanceIdNum;
 
-			for (uint32 Mode = 0U; Mode < uint32(EBatchProcessingMode::Num); ++Mode)
-			{
-				int32 StartIndex = MergedContext.BatchInds[Mode].Num();
-				FInstanceProcessingGPULoadBalancer* MergedLoadBalancer = MergedContext.LoadBalancers[Mode];
+		for (uint32 Mode = 0U; Mode < uint32(EBatchProcessingMode::Num); ++Mode)
+		{
+			int32 StartIndex = BatchInds[Mode].Num();
+			FInstanceProcessingGPULoadBalancer* MergedLoadBalancer = LoadBalancers[Mode];
 				
-				BatchInfo.ItemDataOffset[Mode] = MergedLoadBalancer->GetItems().Num();
-				FInstanceProcessingGPULoadBalancer* LoadBalancer = InstanceCullingContext.LoadBalancers[Mode];
-				LoadBalancer->FinalizeBatches();
+			BatchInfo.ItemDataOffset[Mode] = MergedLoadBalancer->GetItems().Num();
+			FInstanceProcessingGPULoadBalancer* LoadBalancer = InstanceCullingContext.LoadBalancers[Mode];
+			LoadBalancer->FinalizeBatches();
 
-				// UnCulled bucket is used for a single instance mode
-				check(EBatchProcessingMode(Mode) != EBatchProcessingMode::UnCulled || LoadBalancer->HasSingleInstanceItemsOnly());
+			// UnCulled bucket is used for a single instance mode
+			check(EBatchProcessingMode(Mode) != EBatchProcessingMode::UnCulled || LoadBalancer->HasSingleInstanceItemsOnly());
 
-				MergedContext.BatchInds[Mode].AddDefaulted(LoadBalancer->GetBatches().Num());
+			BatchInds[Mode].AddDefaulted(LoadBalancer->GetBatches().Num());
 
-				MergedLoadBalancer->AppendData(*LoadBalancer);
-				for (int32 Index = StartIndex; Index < MergedContext.BatchInds[Mode].Num(); ++Index)
-				{
-					MergedContext.BatchInds[Mode][Index] = BatchIndex;
-				}
+			MergedLoadBalancer->AppendData(*LoadBalancer);
+			for (int32 Index = StartIndex; Index < BatchInds[Mode].Num(); ++Index)
+			{
+				BatchInds[Mode][Index] = BatchIndex;
 			}
+		}
 
-			FInstanceCullingDrawParams& Result = *BatchItem.Result;
-			Result.InstanceDataByteOffset = InstanceDataByteOffset;
-			Result.IndirectArgsByteOffset = BatchInfo.IndirectArgsOffset * FInstanceCullingContext::IndirectArgsNumWords * sizeof(uint32);
+		FInstanceCullingDrawParams& Result = *BatchItem.Result;
+		Result.InstanceDataByteOffset = InstanceDataByteOffset;
+		Result.IndirectArgsByteOffset = BatchInfo.IndirectArgsOffset * FInstanceCullingContext::IndirectArgsNumWords * sizeof(uint32);
 
-			BatchInfo.InstanceDataWriteOffset = InstanceIdBufferOffset;
-			InstanceIdBufferOffset += InstanceCullingContext.TotalInstances * InstanceCullingContext.ViewIds.Num();
-			InstanceDataByteOffset += InstanceCullingContext.InstanceIdOffsets.Num() * sizeof(uint32);
-		}
+		BatchInfo.InstanceDataWriteOffset = InstanceIdBufferOffset;
+		InstanceIdBufferOffset += InstanceCullingContext.TotalInstances * InstanceCullingContext.ViewIds.Num();
+		InstanceDataByteOffset += InstanceCullingContext.InstanceIdOffsets.Num() * sizeof(uint32);
+	}
 
-		// Finalize culling pass parameters
-		for (uint32 Mode = 0U; Mode < uint32(EBatchProcessingMode::Num); ++Mode)
-		{
-			PassParameters[Mode]->NumViewIds = MergedContext.ViewIds.Num();
-			PassParameters[Mode]->LoadBalancerParameters.NumBatches = MergedContext.LoadBalancers[Mode]->GetBatches().Num();
-			PassParameters[Mode]->LoadBalancerParameters.NumItems = MergedContext.LoadBalancers[Mode]->GetItems().Num();
-		}
+	// Finalize culling pass parameters
+	for (uint32 Mode = 0U; Mode < uint32(EBatchProcessingMode::Num); ++Mode)
+	{
+		PassParameters[Mode]->NumViewIds = ViewIds.Num();
+		PassParameters[Mode]->LoadBalancerParameters.NumBatches = LoadBalancers[Mode]->GetBatches().Num();
+		PassParameters[Mode]->LoadBalancerParameters.NumItems = LoadBalancers[Mode]->GetItems().Num();
+		PassParameters[Mode]->NumCullingViews = InstanceCullingManager->GetCullingViews().Num();
 	}
 }
 
@@ -483,90 +527,81 @@ FORCEINLINE int32 GetArrayDataSize(const TArray<DataType, AllocatorType>& Array)
 {
 	return Array.GetTypeSize() * Array.Num();
 }
-void FInstanceCullingContext::BuildRenderingCommandsDeferred(
+
+FInstanceCullingDeferredContext *FInstanceCullingContext::CreateDeferredContext(
 	FRDGBuilder& GraphBuilder,
 	FGPUScene& GPUScene,
-	FInstanceCullingManager& InstanceCullingManager)
+	FInstanceCullingManager* InstanceCullingManager)
 {
-
+	check(InstanceCullingManager != nullptr);
 #define INST_CULL_CALLBACK_MODE(CustomCode) \
-	[&Batches, PassParameters, &MergedContext, Mode]() \
+	[PassParameters, DeferredContext, Mode]() \
 	{ \
-		MergeInstanceCullingContext(Batches, PassParameters, MergedContext); \
+		DeferredContext->ProcessBatched(PassParameters); \
 		return CustomCode; \
 	}
 
 #define INST_CULL_CALLBACK(CustomCode) \
-	[&Batches, PassParameters, &MergedContext]() \
+	[PassParameters, DeferredContext]() \
 	{ \
-		MergeInstanceCullingContext(Batches, PassParameters, MergedContext); \
+		DeferredContext->ProcessBatched(PassParameters); \
 		return CustomCode; \
 	}
 
 #define INST_CULL_CREATE_STRUCT_BUFF_ARGS(ArrayName) \
 	GraphBuilder, \
 	TEXT("InstanceCulling.") TEXT(#ArrayName), \
-	MergedContext.ArrayName.GetTypeSize(), \
-	INST_CULL_CALLBACK(MergedContext.ArrayName.Num()), \
-	INST_CULL_CALLBACK(MergedContext.ArrayName.GetData()), \
-	INST_CULL_CALLBACK(MergedContext.ArrayName.Num() * MergedContext.ArrayName.GetTypeSize())
+	DeferredContext->ArrayName.GetTypeSize(), \
+	INST_CULL_CALLBACK(DeferredContext->ArrayName.Num()), \
+	INST_CULL_CALLBACK(DeferredContext->ArrayName.GetData()), \
+	INST_CULL_CALLBACK(DeferredContext->ArrayName.Num() * DeferredContext->ArrayName.GetTypeSize())
 
 #define INST_CULL_CREATE_STRUCT_BUFF_ARGS_MODE(ArrayName) \
 	GraphBuilder, \
 	TEXT("InstanceCulling.") TEXT(#ArrayName), \
-	MergedContext.ArrayName[Mode].GetTypeSize(), \
-	INST_CULL_CALLBACK_MODE(MergedContext.ArrayName[Mode].Num()), \
-	INST_CULL_CALLBACK_MODE(MergedContext.ArrayName[Mode].GetData()), \
-	INST_CULL_CALLBACK_MODE(MergedContext.ArrayName[Mode].Num() * MergedContext.ArrayName[Mode].GetTypeSize())
-
-	// Cannot defer pass execution in immediate mode.
-	// TODO: support batching when r.RDG.Drain is enabled
-	if (!AllowBatchedBuildRenderingCommands(GPUScene))
-	{
-		return;
-	}
+	DeferredContext->ArrayName[Mode].GetTypeSize(), \
+	INST_CULL_CALLBACK_MODE(DeferredContext->ArrayName[Mode].Num()), \
+	INST_CULL_CALLBACK_MODE(DeferredContext->ArrayName[Mode].GetData()), \
+	INST_CULL_CALLBACK_MODE(DeferredContext->ArrayName[Mode].Num() * DeferredContext->ArrayName[Mode].GetTypeSize())
 
-	// If there are no instances, there can be no work to perform later.
-	if (GPUScene.GetNumInstances() == 0 || InstanceCullingManager.CullingIntermediate.NumViews == 0)
-	{
-		return;
-	}
+	FInstanceCullingDeferredContext* DeferredContext = GraphBuilder.AllocObject<FInstanceCullingDeferredContext>(InstanceCullingManager);
 
 	const bool bCullInstances = CVarCullInstances.GetValueOnRenderThread() != 0;
-
 	RDG_EVENT_SCOPE(GraphBuilder, "BuildRenderingCommandsDeferred(Culling=%s)", bCullInstances ? TEXT("On") : TEXT("Off"));
 
-	FBatchedInstanceCullingScratchSpace& BatchedCullingScratch = InstanceCullingManager.BatchedCullingScratch;
-
-	checkf(!BatchedCullingScratch.bBatchingActive, TEXT("Batched GPU instance culling has already been enabled."));
-	BatchedCullingScratch.bBatchingActive = true;
-
-	const TArray<FBatchItem, SceneRenderingAllocator>& Batches = BatchedCullingScratch.Batches;
-	FInstanceCullingContext::FMergedContext& MergedContext = BatchedCullingScratch.MergedContext;
 	TStaticArray<FBuildInstanceIdBufferAndCommandsFromPrimitiveIdsCs::FParameters*, static_cast<uint32>(EBatchProcessingMode::Num)> PassParameters;
 	for (uint32 Mode = 0U; Mode < uint32(EBatchProcessingMode::Num); ++Mode)
 	{
 		PassParameters[Mode] = GraphBuilder.AllocParameters<FBuildInstanceIdBufferAndCommandsFromPrimitiveIdsCs::FParameters>();
 	}
+
 	FBuildInstanceIdBufferAndCommandsFromPrimitiveIdsCs::FParameters PassParametersTmp;
 
+	// Defer upload of views as these may be added during setup of passes.
+	FRDGBufferRef ViewsRDG = CreateStructuredBuffer(GraphBuilder, 
+		TEXT("InstanceCulling.Views"), InstanceCullingManager->CullingViews.GetTypeSize(), 
+		INST_CULL_CALLBACK(DeferredContext->InstanceCullingManager->CullingViews.Num()),
+		INST_CULL_CALLBACK(DeferredContext->InstanceCullingManager->CullingViews.GetData()),
+		INST_CULL_CALLBACK(GetArrayDataSize(DeferredContext->InstanceCullingManager->CullingViews))
+	);
+
 	FRDGBufferRef DrawCommandDescsRDG = CreateStructuredBuffer(INST_CULL_CREATE_STRUCT_BUFF_ARGS(DrawCommandDescs));
 	FRDGBufferRef ViewIdsRDG = CreateStructuredBuffer(INST_CULL_CREATE_STRUCT_BUFF_ARGS(ViewIds));
 	FRDGBufferRef BatchInfosRDG = CreateStructuredBuffer(INST_CULL_CREATE_STRUCT_BUFF_ARGS(BatchInfos));
-	FRDGBufferRef InstanceIdsBuffer = GraphBuilder.CreateBuffer(FRDGBufferDesc::CreateStructuredDesc(sizeof(uint32), 1), TEXT("InstanceCulling.InstanceIdsBuffer"), INST_CULL_CALLBACK(MergedContext.InstanceIdBufferSize));
+	FRDGBufferRef InstanceIdsBuffer = GraphBuilder.CreateBuffer(FRDGBufferDesc::CreateStructuredDesc(sizeof(uint32), 1), TEXT("InstanceCulling.InstanceIdsBuffer"), INST_CULL_CALLBACK(DeferredContext->InstanceIdBufferSize));
 
-	BatchedCullingScratch.DrawIndirectArgsBuffer = GraphBuilder.CreateBuffer(FRDGBufferDesc::CreateIndirectDesc(), TEXT("InstanceCulling.DrawIndirectArgsBuffer"), INST_CULL_CALLBACK(IndirectArgsNumWords * MergedContext.IndirectArgs.Num()));
-	GraphBuilder.QueueBufferUpload(BatchedCullingScratch.DrawIndirectArgsBuffer, INST_CULL_CALLBACK(MergedContext.IndirectArgs.GetData()), INST_CULL_CALLBACK(GetArrayDataSize(MergedContext.IndirectArgs)));
+	DeferredContext->DrawIndirectArgsBuffer = GraphBuilder.CreateBuffer(FRDGBufferDesc::CreateIndirectDesc(), TEXT("InstanceCulling.DrawIndirectArgsBuffer"), INST_CULL_CALLBACK(IndirectArgsNumWords * DeferredContext->IndirectArgs.Num()));
+	GraphBuilder.QueueBufferUpload(DeferredContext->DrawIndirectArgsBuffer, INST_CULL_CALLBACK(DeferredContext->IndirectArgs.GetData()), INST_CULL_CALLBACK(GetArrayDataSize(DeferredContext->IndirectArgs)));
 
 	const ERHIFeatureLevel::Type FeatureLevel = GPUScene.GetFeatureLevel();
 	FGlobalShaderMap* ShaderMap = GetGlobalShaderMap(FeatureLevel);
 
 	// Note: we redundantly clear the instance counts here as there is some issue with replays on certain consoles.
-	AddClearIndirectArgInstanceCountPass(GraphBuilder, ShaderMap, BatchedCullingScratch.DrawIndirectArgsBuffer);
+	AddClearIndirectArgInstanceCountPass(GraphBuilder, ShaderMap, DeferredContext->DrawIndirectArgsBuffer, INST_CULL_CALLBACK(DeferredContext->IndirectArgs.Num()));
 
 	// not using structured buffer as we want/have to get at it as a vertex buffer 
-	BatchedCullingScratch.InstanceIdOffsetBuffer = GraphBuilder.CreateBuffer(FRDGBufferDesc::CreateBufferDesc(sizeof(uint32), 1), TEXT("InstanceCulling.InstanceIdOffsetBuffer"), INST_CULL_CALLBACK(MergedContext.InstanceIdOffsets.Num()));
-	GraphBuilder.QueueBufferUpload(BatchedCullingScratch.InstanceIdOffsetBuffer, INST_CULL_CALLBACK(MergedContext.InstanceIdOffsets.GetData()), INST_CULL_CALLBACK(MergedContext.InstanceIdOffsets.GetTypeSize() * MergedContext.InstanceIdOffsets.Num()));
+	DeferredContext->InstanceIdOffsetBuffer = GraphBuilder.CreateBuffer(FRDGBufferDesc::CreateBufferDesc(sizeof(uint32), 1), TEXT("InstanceCulling.InstanceIdOffsetBuffer"), INST_CULL_CALLBACK(DeferredContext->InstanceIdOffsets.Num()));
+	GraphBuilder.QueueBufferUpload(DeferredContext->InstanceIdOffsetBuffer, INST_CULL_CALLBACK(DeferredContext->InstanceIdOffsets.GetData()), INST_CULL_CALLBACK(DeferredContext->InstanceIdOffsets.GetTypeSize() * DeferredContext->InstanceIdOffsets.Num()));
 
 	// Because the view uniforms are not set up by the time this runs
 	// PassParameters->View = View.ViewUniformBuffer;
@@ -582,12 +617,14 @@ void FInstanceCullingContext::BuildRenderingCommandsDeferred(
 	PassParametersTmp.BatchInfos = GraphBuilder.CreateSRV(BatchInfosRDG);
 	PassParametersTmp.ViewIds = GraphBuilder.CreateSRV(ViewIdsRDG);
 	PassParametersTmp.InstanceIdsBufferOut = GraphBuilder.CreateUAV(InstanceIdsBuffer, ERDGUnorderedAccessViewFlags::SkipBarrier);
-	PassParametersTmp.DrawIndirectArgsBufferOut = GraphBuilder.CreateUAV(BatchedCullingScratch.DrawIndirectArgsBuffer, PF_R32_UINT, ERDGUnorderedAccessViewFlags::SkipBarrier);
-	PassParametersTmp.InstanceIdOffsetBuffer = GraphBuilder.CreateSRV(BatchedCullingScratch.InstanceIdOffsetBuffer, PF_R32_UINT);
+	PassParametersTmp.DrawIndirectArgsBufferOut = GraphBuilder.CreateUAV(DeferredContext->DrawIndirectArgsBuffer, PF_R32_UINT, ERDGUnorderedAccessViewFlags::SkipBarrier);
+	PassParametersTmp.InstanceIdOffsetBuffer = GraphBuilder.CreateSRV(DeferredContext->InstanceIdOffsetBuffer, PF_R32_UINT);
+	
 	if (bCullInstances)
 	{
-		PassParametersTmp.InViews = GraphBuilder.CreateSRV(InstanceCullingManager.CullingIntermediate.CullingViews);
-		PassParametersTmp.NumCullingViews = InstanceCullingManager.CullingIntermediate.NumViews;
+		PassParametersTmp.InViews = GraphBuilder.CreateSRV(ViewsRDG);
+		// Note: not the final number
+		PassParametersTmp.NumCullingViews = InstanceCullingManager->CullingIntermediate.NumViews;
 	}
 
 	for (uint32 Mode = 0U; Mode < uint32(EBatchProcessingMode::Num); ++Mode)
@@ -603,17 +640,17 @@ void FInstanceCullingContext::BuildRenderingCommandsDeferred(
 			GraphBuilder,
 			TEXT("InstanceCullingLoadBalancer.Batches"),
 			sizeof(FInstanceProcessingGPULoadBalancer::FPackedBatch),
-			INST_CULL_CALLBACK_MODE(MergedContext.LoadBalancers[Mode]->GetBatches().Num()),
-			INST_CULL_CALLBACK_MODE(MergedContext.LoadBalancers[Mode]->GetBatches().GetData()),
-			INST_CULL_CALLBACK_MODE(GetArrayDataSize(MergedContext.LoadBalancers[Mode]->GetBatches())));
+			INST_CULL_CALLBACK_MODE(DeferredContext->LoadBalancers[Mode]->GetBatches().Num()),
+			INST_CULL_CALLBACK_MODE(DeferredContext->LoadBalancers[Mode]->GetBatches().GetData()),
+			INST_CULL_CALLBACK_MODE(GetArrayDataSize(DeferredContext->LoadBalancers[Mode]->GetBatches())));
 
 		FRDGBufferRef ItemBuffer = CreateStructuredBuffer(
 			GraphBuilder,
 			TEXT("InstanceCullingLoadBalancer.Items"),
 			sizeof(FInstanceProcessingGPULoadBalancer::FPackedItem),
-			INST_CULL_CALLBACK_MODE(MergedContext.LoadBalancers[Mode]->GetItems().Num()),
-			INST_CULL_CALLBACK_MODE(MergedContext.LoadBalancers[Mode]->GetItems().GetData()),
-			INST_CULL_CALLBACK_MODE(GetArrayDataSize(MergedContext.LoadBalancers[Mode]->GetItems())));
+			INST_CULL_CALLBACK_MODE(DeferredContext->LoadBalancers[Mode]->GetItems().Num()),
+			INST_CULL_CALLBACK_MODE(DeferredContext->LoadBalancers[Mode]->GetItems().GetData()),
+			INST_CULL_CALLBACK_MODE(GetArrayDataSize(DeferredContext->LoadBalancers[Mode]->GetItems())));
 
 		PassParameters[Mode]->LoadBalancerParameters.BatchBuffer = GraphBuilder.CreateSRV(BatchBuffer);
 		PassParameters[Mode]->LoadBalancerParameters.ItemBuffer = GraphBuilder.CreateSRV(ItemBuffer);
@@ -631,25 +668,22 @@ void FInstanceCullingContext::BuildRenderingCommandsDeferred(
 			RDG_EVENT_NAME("CullInstances(%s)", BatchProcessingModeStr[Mode]),
 			ComputeShader,
 			PassParameters[Mode],
-			INST_CULL_CALLBACK_MODE(MergedContext.LoadBalancers[Mode]->GetWrappedCsGroupCount()));
+			INST_CULL_CALLBACK_MODE(DeferredContext->LoadBalancers[Mode]->GetWrappedCsGroupCount()));
 	}
 
 	FInstanceCullingGlobalUniforms* UniformParameters = GraphBuilder.AllocParameters<FInstanceCullingGlobalUniforms>();
 	UniformParameters->InstanceIdsBuffer = GraphBuilder.CreateSRV(InstanceIdsBuffer);
 	UniformParameters->PageInfoBuffer = GraphBuilder.CreateSRV(InstanceIdsBuffer);
 	UniformParameters->BufferCapacity = 0U; // TODO: this is not used at the moment, but is intended for range checks so would have been good.
-	BatchedCullingScratch.UniformBuffer = GraphBuilder.CreateUniformBuffer(UniformParameters);
+	DeferredContext->UniformBuffer = GraphBuilder.CreateUniformBuffer(UniformParameters);
 
 
 #undef INST_CULL_CREATE_STRUCT_BUFF_ARGS
 #undef INST_CULL_CALLBACK
 #undef INST_CULL_CALLBACK_MODE
 #undef INST_CULL_CREATE_STRUCT_BUFF_ARGS_MODE
-}
 
-bool FInstanceCullingContext::AllowBatchedBuildRenderingCommands(const FGPUScene& GPUScene)
-{
-	return GPUScene.IsEnabled() && !!GAllowBatchedBuildRenderingCommands && !FRDGBuilder::IsImmediateMode() && !FRDGBuilder::IsDrainEnabled();
+	return DeferredContext;
 }
 
 
@@ -684,7 +718,7 @@ public:
 IMPLEMENT_GLOBAL_SHADER(FClearIndirectArgInstanceCountCs, "/Engine/Private/InstanceCulling/BuildInstanceDrawCommands.usf", "ClearIndirectArgInstanceCountCS", SF_Compute);
 
 
-void FInstanceCullingContext::AddClearIndirectArgInstanceCountPass(FRDGBuilder& GraphBuilder, FGlobalShaderMap* ShaderMap, FRDGBufferRef DrawIndirectArgsBuffer)
+void FInstanceCullingContext::AddClearIndirectArgInstanceCountPass(FRDGBuilder& GraphBuilder, FGlobalShaderMap* ShaderMap, FRDGBufferRef DrawIndirectArgsBuffer, TFunction<int32()> NumIndirectArgsCallback)
 {
 	FClearIndirectArgInstanceCountCs::FParameters* PassParameters = GraphBuilder.AllocParameters<FClearIndirectArgInstanceCountCs::FParameters>();
 	// Upload data etc
@@ -693,13 +727,38 @@ void FInstanceCullingContext::AddClearIndirectArgInstanceCountPass(FRDGBuilder&
 
 	auto ComputeShader = ShaderMap->GetShader<FClearIndirectArgInstanceCountCs>();
 
-	FComputeShaderUtils::AddPass(
-		GraphBuilder,
-		RDG_EVENT_NAME("ClearIndirectArgInstanceCount"),
-		ComputeShader,
-		PassParameters,
-		FComputeShaderUtils::GetGroupCountWrapped(PassParameters->NumIndirectArgs, FClearIndirectArgInstanceCountCs::NumThreadsPerGroup)
-	);
+	if (NumIndirectArgsCallback)
+	{
+		const FShaderParametersMetadata* ParametersMetadata = FClearIndirectArgInstanceCountCs::FParameters::FTypeInfo::GetStructMetadata();
+		ClearUnusedGraphResources(ComputeShader, ParametersMetadata, PassParameters);
+
+		GraphBuilder.AddPass(
+			RDG_EVENT_NAME("ClearIndirectArgInstanceCount"),
+			ParametersMetadata,
+			PassParameters,
+			ERDGPassFlags::Compute,
+			[ParametersMetadata, PassParameters, ComputeShader, NumIndirectArgsCallback = MoveTemp(NumIndirectArgsCallback)](FRHIComputeCommandList& RHICmdList)
+		{
+			int32 NumIndirectArgs = NumIndirectArgsCallback();
+			PassParameters->NumIndirectArgs = NumIndirectArgs;
+			const FIntVector GroupCount = FComputeShaderUtils::GetGroupCountWrapped(NumIndirectArgs, FClearIndirectArgInstanceCountCs::NumThreadsPerGroup);
+			if (GroupCount.X > 0 && GroupCount.Y > 0 && GroupCount.Z > 0)
+			{
+				FComputeShaderUtils::ValidateGroupCount(GroupCount);
+				FComputeShaderUtils::Dispatch(RHICmdList, ComputeShader, ParametersMetadata, *PassParameters, GroupCount);
+			}
+		});
+	}
+	else
+	{
+		FComputeShaderUtils::AddPass(
+			GraphBuilder,
+			RDG_EVENT_NAME("ClearIndirectArgInstanceCount"),
+			ComputeShader,
+			PassParameters,
+			FComputeShaderUtils::GetGroupCountWrapped(PassParameters->NumIndirectArgs, FClearIndirectArgInstanceCountCs::NumThreadsPerGroup)
+		);
+	}
 }
 
 /**
@@ -774,7 +833,7 @@ void FInstanceCullingContext::SetupDrawCommands(
 			if (DrawCmd.bUseIndirect == 0)
 			{
 				DrawCmd.IndirectArgsOffsetOrNumInstances += 1;
-			}
+		}
 		}
 		else
 		{
diff --git a/Engine/Source/Runtime/Renderer/Private/InstanceCulling/InstanceCullingManager.cpp b/Engine/Source/Runtime/Renderer/Private/InstanceCulling/InstanceCullingManager.cpp
--- a/Engine/Source/Runtime/Renderer/Private/InstanceCulling/InstanceCullingManager.cpp
+++ b/Engine/Source/Runtime/Renderer/Private/InstanceCulling/InstanceCullingManager.cpp
@@ -14,15 +14,16 @@
 
 #include "InstanceCulling/InstanceCullingContext.h"
 
+static int32 GAllowBatchedBuildRenderingCommands = 1;
+static FAutoConsoleVariableRef CVarAllowBatchedBuildRenderingCommands(
+	TEXT("r.InstanceCulling.AllowBatchedBuildRenderingCommands"),
+	GAllowBatchedBuildRenderingCommands,
+	TEXT("Whether to allow batching BuildRenderingCommands for GPU instance culling"),
+	ECVF_RenderThreadSafe);
+
+
 FInstanceCullingManager::~FInstanceCullingManager()
 {
-	for (auto& LoadBalancer : BatchedCullingScratch.MergedContext.LoadBalancers)
-	{
-		if (LoadBalancer != nullptr)
-		{
-			delete LoadBalancer;
-		}
-	}
 }
 
 int32 FInstanceCullingManager::RegisterView(const FViewInfo& ViewInfo)
@@ -65,7 +66,26 @@ void FInstanceCullingManager::CullInstances(FRDGBuilder& GraphBuilder, FGPUScene
 	CullingIntermediate.DummyUniformBuffer = FInstanceCullingContext::CreateDummyInstanceCullingUniformBuffer(GraphBuilder);
 }
 
+
+bool FInstanceCullingManager::AllowBatchedBuildRenderingCommands(const FGPUScene& GPUScene)
+{
+	return GPUScene.IsEnabled() && !!GAllowBatchedBuildRenderingCommands && !FRDGBuilder::IsImmediateMode();
+}
+
+
 void FInstanceCullingManager::BeginDeferredCulling(FRDGBuilder& GraphBuilder, FGPUScene& GPUScene)
 {
-	FInstanceCullingContext::BuildRenderingCommandsDeferred(GraphBuilder, GPUScene, *this);
+	// Cannot defer pass execution in immediate mode.
+	if (!AllowBatchedBuildRenderingCommands(GPUScene))
+	{
+		return;
+	}
+
+	// If there are no instances, there can be no work to perform later.
+	if (GPUScene.GetNumInstances() == 0 || CullingIntermediate.NumViews == 0)
+	{
+		return;
+	}
+
+	DeferredContext = FInstanceCullingContext::CreateDeferredContext(GraphBuilder, GPUScene, this);
 }
diff --git a/Engine/Source/Runtime/Renderer/Private/InstanceCulling/InstanceCullingManager.h b/Engine/Source/Runtime/Renderer/Private/InstanceCulling/InstanceCullingManager.h
--- a/Engine/Source/Runtime/Renderer/Private/InstanceCulling/InstanceCullingManager.h
+++ b/Engine/Source/Runtime/Renderer/Private/InstanceCulling/InstanceCullingManager.h
@@ -69,21 +69,7 @@ struct FInstanceCullingResult
 	}
 };
 
-struct FBatchedInstanceCullingScratchSpace
-{
-	/** Whether we defer and batch GPU instance culling work throughout a frame. */
-	bool bBatchingActive = false;
-
-	FRDGBufferRef DrawIndirectArgsBuffer = nullptr;
-	FRDGBufferRef InstanceIdOffsetBuffer = nullptr;
-	TRDGUniformBufferRef<FInstanceCullingGlobalUniforms> UniformBuffer = nullptr;
-
-	/** GPU instance culling input data merged from multiple batches throughout a frame. */
-	FInstanceCullingContext::FMergedContext MergedContext;
-
-	/** Batches of GPU instance culling input data. */
-	TArray<FInstanceCullingContext::FBatchItem, SceneRenderingAllocator> Batches;
-};
+class FInstanceCullingDeferredContext;
 
 /**
  * Manages allocation of indirect arguments and culling jobs for all instanced draws (use the GPU Scene culling).
@@ -115,21 +101,31 @@ public:
 	}
 	const TRDGUniformBufferRef<FInstanceCullingGlobalUniforms> GetDummyInstanceCullingUniformBuffer() const { return CullingIntermediate.DummyUniformBuffer; }
 	
-	/** Starts batching GPU instance culling work items if possible. */
+	static bool AllowBatchedBuildRenderingCommands(const FGPUScene& GPUScene);
+
+
+	/**
+	 * Add a deferred, batched, gpu culling pass. Each batch represents a BuildRenderingCommands call from a mesh pass.
+	 * Batches are collected as we walk through the main render setup and call BuildRenderingCommands, and are processed when RDG Execute or Drain is called.
+	 * This implicitly ends the deferred context, so if Drain is used, it should be paired with a new call to BeginDeferredCulling.
+	 */
 	void BeginDeferredCulling(FRDGBuilder& GraphBuilder, FGPUScene& GPUScene);
 
+
 	/** Whether we are actively batching GPU instance culling work. */
-	bool IsDeferredCullingActive() const { return BatchedCullingScratch.bBatchingActive; }
+	bool IsDeferredCullingActive() const { return DeferredContext != nullptr; }
 
 	// Populated by CullInstances, used when performing final culling & rendering 
 	FInstanceCullingIntermediate CullingIntermediate;
 
+	const TArray<Nanite::FPackedView>& GetCullingViews() { return CullingViews; }
 
 private:
-	friend class FInstanceCullingContext;
 
+	friend class FInstanceCullingContext;
+	
 	// Polulated by FInstanceCullingContext::BuildRenderingCommandsDeferred, used to hold instance culling related data that needs to be passed around
-	FBatchedInstanceCullingScratchSpace BatchedCullingScratch;
+	FInstanceCullingDeferredContext *DeferredContext = nullptr;
 
 	FInstanceCullingManager() = delete;
 	FInstanceCullingManager(FInstanceCullingManager &) = delete;
diff --git a/Engine/Source/Runtime/Renderer/Public/InstanceCulling/InstanceCullingContext.h b/Engine/Source/Runtime/Renderer/Public/InstanceCulling/InstanceCullingContext.h
--- a/Engine/Source/Runtime/Renderer/Public/InstanceCulling/InstanceCullingContext.h
+++ b/Engine/Source/Runtime/Renderer/Public/InstanceCulling/InstanceCullingContext.h
@@ -13,9 +13,7 @@ class FInstanceCullingManager;
 class FInstanceCullingDrawParams;
 class FScene;
 class FGPUScenePrimitiveCollector;
-
-
-#define GPUCULL_ENABLE_CONCAT_ON_UPLOAD 0
+class FInstanceCullingDeferredContext;
 
 
 BEGIN_GLOBAL_SHADER_PARAMETER_STRUCT(FInstanceCullingGlobalUniforms, RENDERER_API)
@@ -125,19 +123,18 @@ public:
 
 	/**
 	 * Add a batched BuildRenderingCommands pass. Each batch represents a BuildRenderingCommands call from a mesh pass.
-	 * Batches are gradually added as we walk through the main render function.
+	 * Batches are collected as we walk through the main render setup and are executed when RDG Execute or Drain is called.
+	 * This implicitly ends the deferred context, so if Drain is used, it should be paired with a new call to BeginDeferredCulling.
 	 */
-	static void BuildRenderingCommandsDeferred(
+	static FInstanceCullingDeferredContext *CreateDeferredContext(
 		FRDGBuilder& GraphBuilder,
 		FGPUScene& GPUScene,
-		FInstanceCullingManager& InstanceCullingManager);
-
-	static bool AllowBatchedBuildRenderingCommands(const FGPUScene& GPUScene);
+		FInstanceCullingManager* InstanceCullingManager);
 
 	/**
 	 * Helper function to add a pass to zero the instance count in the indirect args.
 	 */
-	static void AddClearIndirectArgInstanceCountPass(FRDGBuilder& GraphBuilder, FGlobalShaderMap* ShaderMap, FRDGBufferRef DrawIndirectArgsBuffer);
+	static void AddClearIndirectArgInstanceCountPass(FRDGBuilder& GraphBuilder, FGlobalShaderMap* ShaderMap, FRDGBufferRef DrawIndirectArgsBuffer, TFunction<int32()> NumIndirectArgsCallback = TFunction<int32()>());
 
 
 	void SetupDrawCommands(
@@ -207,33 +204,6 @@ public:
 
 	// Set of specialized batches that collect items with different properties each context may have only a subset.
 	//TStaticArray<FBatchProcessor, EBatchProcessingMode::Num> Batches;
-	struct FMergedContext
-	{
-		bool bInitialized = false;
-#if GPUCULL_ENABLE_CONCAT_ON_UPLOAD
-		TArray<TArrayView<const int32>, SceneRenderingAllocator> ViewIds;
-		TArray<TArrayView<const FRHIDrawIndexedIndirectParameters>, SceneRenderingAllocator> IndirectArgs;
-		TArray<TArrayView<const FDrawCommandDesc>, SceneRenderingAllocator> DrawCommandDescs;
-		TArray<TArrayView<const uint32>, SceneRenderingAllocator> InstanceIdOffsets;
-#else // !GPUCULL_ENABLE_CONCAT_ON_UPLOAD
-		TArray<int32, SceneRenderingAllocator> ViewIds;
-		//TArray<FMeshDrawCommandInfo, SceneRenderingAllocator> MeshDrawCommandInfos;
-		TArray<FRHIDrawIndexedIndirectParameters, SceneRenderingAllocator> IndirectArgs;
-		TArray<FDrawCommandDesc, SceneRenderingAllocator> DrawCommandDescs;
-		TArray<uint32, SceneRenderingAllocator> InstanceIdOffsets;
-#endif // GPUCULL_ENABLE_CONCAT_ON_UPLOAD
-		LoadBalancerArray LoadBalancers = LoadBalancerArray(InPlace, nullptr);
-		TStaticArray<TArray<uint32, SceneRenderingAllocator>, static_cast<uint32>(EBatchProcessingMode::Num)> BatchInds;
-		TArray<FInstanceCullingContext::FContextBatchInfo, SceneRenderingAllocator> BatchInfos;
-
-
-		// Counters to sum up all sizes to facilitate pre-sizing
-		uint32 InstanceIdBufferSize = 0U;
-		TStaticArray<int32, uint32(EBatchProcessingMode::Num)> TotalBatches = TStaticArray<int32, uint32(EBatchProcessingMode::Num)>(InPlace, 0);
-		TStaticArray<int32, uint32(EBatchProcessingMode::Num)> TotalItems = TStaticArray<int32, uint32(EBatchProcessingMode::Num)>(InPlace, 0);
-		int32 TotalIndirectArgs = 0;
-		int32 TotalViewIds = 0;
-	};
 
 	// Processing mode to use for single-instance primitives, default to skip culling, as this is already done on CPU. 
 	EBatchProcessingMode SingleInstanceProcessingMode = EBatchProcessingMode::UnCulled;
