Moving gathering of shadow primitives to tasks that start before rendering prepass.

diff --git a/Engine/Source/Runtime/Renderer/Private/DeferredShadingRenderer.cpp b/Engine/Source/Runtime/Renderer/Private/DeferredShadingRenderer.cpp
--- a/Engine/Source/Runtime/Renderer/Private/DeferredShadingRenderer.cpp
+++ b/Engine/Source/Runtime/Renderer/Private/DeferredShadingRenderer.cpp
@@ -2222,7 +2222,7 @@ void FDeferredShadingSceneRenderer::Render(FRDGBuilder& GraphBuilder)
 
 		{
 			RDG_RHI_GPU_STAT_SCOPE(GraphBuilder, VisibilityCommands);
-			InitViewsPossiblyAfterPrepass(GraphBuilder, ILCTaskData, InstanceCullingManager);
+			InitViewsAfterPrepass(GraphBuilder, ILCTaskData, InstanceCullingManager);
 		}
 
 		{
diff --git a/Engine/Source/Runtime/Renderer/Private/DeferredShadingRenderer.h b/Engine/Source/Runtime/Renderer/Private/DeferredShadingRenderer.h
--- a/Engine/Source/Runtime/Renderer/Private/DeferredShadingRenderer.h
+++ b/Engine/Source/Runtime/Renderer/Private/DeferredShadingRenderer.h
@@ -383,7 +383,8 @@ private:
 	/** Determines which primitives are visible for each view. */
 	void InitViews(FRDGBuilder& GraphBuilder, const FSceneTexturesConfig& SceneTexturesConfig, FExclusiveDepthStencil::Type BasePassDepthStencilAccess, struct FILCUpdatePrimTaskData& ILCTaskData, FInstanceCullingManager& InstanceCullingManager);
 
-	void InitViewsPossiblyAfterPrepass(FRDGBuilder& GraphBuilder, struct FILCUpdatePrimTaskData& ILCTaskData, FInstanceCullingManager& InstanceCullingManager);
+	void InitViewsBeforePrepass(FRDGBuilder& GraphBuilder, FInstanceCullingManager& InstanceCullingManager);
+	void InitViewsAfterPrepass(FRDGBuilder& GraphBuilder, struct FILCUpdatePrimTaskData& ILCTaskData, FInstanceCullingManager& InstanceCullingManager);
 	void BeginUpdateLumenSceneTasks(FRDGBuilder& GraphBuilder);
 	void UpdateLumenScene(FRDGBuilder& GraphBuilder);
 	void RenderLumenSceneLighting(FRDGBuilder& GraphBuilder, FViewInfo& View);
@@ -996,6 +997,8 @@ private:
 
 	/** Set to true if lights were injected into the light grid (this controlled by somewhat complex logic, this flag is used to cross-check). */
 	bool bAreLightsInLightGrid;
+
+	FDynamicShadowsTaskData* CurrentDynamicShadowsTaskData{};
 };
 
 DECLARE_CYCLE_STAT_EXTERN(TEXT("PrePass"), STAT_CLM_PrePass, STATGROUP_CommandListMarkers, );
\ No newline at end of file
diff --git a/Engine/Source/Runtime/Renderer/Private/SceneHitProxyRendering.cpp b/Engine/Source/Runtime/Renderer/Private/SceneHitProxyRendering.cpp
--- a/Engine/Source/Runtime/Renderer/Private/SceneHitProxyRendering.cpp
+++ b/Engine/Source/Runtime/Renderer/Private/SceneHitProxyRendering.cpp
@@ -597,7 +597,7 @@ void FDeferredShadingSceneRenderer::RenderHitProxies(FRDGBuilder& GraphBuilder)
 	{
 		FILCUpdatePrimTaskData ILCTaskData;
 		InitViews(GraphBuilder, SceneTexturesConfig, FExclusiveDepthStencil::DepthWrite_StencilWrite, ILCTaskData, InstanceCullingManager);
-		InitViewsPossiblyAfterPrepass(GraphBuilder, ILCTaskData, InstanceCullingManager);
+		InitViewsAfterPrepass(GraphBuilder, ILCTaskData, InstanceCullingManager);
 	}
 
 	extern TSet<IPersistentViewUniformBufferExtension*> PersistentViewUniformBufferExtensions;
diff --git a/Engine/Source/Runtime/Renderer/Private/SceneRendering.h b/Engine/Source/Runtime/Renderer/Private/SceneRendering.h
--- a/Engine/Source/Runtime/Renderer/Private/SceneRendering.h
+++ b/Engine/Source/Runtime/Renderer/Private/SceneRendering.h
@@ -61,6 +61,7 @@ struct FSceneTexturesConfig;
 struct FMinimalSceneTextures;
 struct FSceneTextures;
 struct FCustomDepthTextures;
+struct FDynamicShadowsTaskData;
 
 namespace DistanceField
 {
@@ -1985,6 +1986,8 @@ protected:
 
 	// Shared functionality between all scene renderers
 
+	FDynamicShadowsTaskData* BeginInitDynamicShadows(bool bRunningEarly);
+	void FinishInitDynamicShadows(FRHICommandListImmediate& RHICmdList, FDynamicShadowsTaskData* TaskData, FGlobalDynamicIndexBuffer& DynamicIndexBuffer, FGlobalDynamicVertexBuffer& DynamicVertexBuffer, FGlobalDynamicReadBuffer& DynamicReadBuffer, FInstanceCullingManager& InstanceCullingManager);
 	void InitDynamicShadows(FRHICommandListImmediate& RHICmdList, FGlobalDynamicIndexBuffer& DynamicIndexBuffer, FGlobalDynamicVertexBuffer& DynamicVertexBuffer, FGlobalDynamicReadBuffer& DynamicReadBuffer, FInstanceCullingManager& InstanceCullingManager);
 
 	void SetupMeshPass(FViewInfo& View, FExclusiveDepthStencil::Type BasePassDepthStencilAccess, FViewCommands& ViewCommands, FInstanceCullingManager& InstanceCullingManager);
@@ -2066,10 +2069,8 @@ protected:
 	bool CheckForProjectedShadows(const FLightSceneInfo* LightSceneInfo) const;
 
 	/** Gathers the list of primitives used to draw various shadow types */
-	void GatherShadowPrimitives(
-		const TArray<FProjectedShadowInfo*, SceneRenderingAllocator>& PreShadows,
-		const TArray<FProjectedShadowInfo*, SceneRenderingAllocator>& ViewDependentWholeSceneShadows,
-		bool bReflectionCaptureScene);
+	void BeginGatherShadowPrimitives(FDynamicShadowsTaskData* TaskData);
+	void FinishGatherShadowPrimitives(FDynamicShadowsTaskData* TaskData);
 
 	void RenderShadowDepthMaps(FRDGBuilder& GraphBuilder, FInstanceCullingManager& InstanceCullingManager);
 	void RenderShadowDepthMapAtlases(FRDGBuilder& GraphBuilder);
diff --git a/Engine/Source/Runtime/Renderer/Private/SceneVisibility.cpp b/Engine/Source/Runtime/Renderer/Private/SceneVisibility.cpp
--- a/Engine/Source/Runtime/Renderer/Private/SceneVisibility.cpp
+++ b/Engine/Source/Runtime/Renderer/Private/SceneVisibility.cpp
@@ -151,6 +151,14 @@ static FAutoConsoleVariableRef CVarILCUpdatePrimitivesTask(
 	ECVF_RenderThreadSafe
 	);
 
+static int32 GEarlyInitDynamicShadows = 1;
+static FAutoConsoleVariableRef CVarEarlyInitDynamicShadows(
+	TEXT("r.EarlyInitDynamicShadows"),
+	GEarlyInitDynamicShadows,
+	TEXT("Starts shadow culling tasks earlier in the frame."),
+	ECVF_RenderThreadSafe
+);
+
 static int32 GFramesNotOcclusionTestedToExpandBBoxes = 5;
 static FAutoConsoleVariableRef CVarFramesNotOcclusionTestedToExpandBBoxes(
 	TEXT("r.GFramesNotOcclusionTestedToExpandBBoxes"),
@@ -4581,6 +4589,8 @@ void FDeferredShadingSceneRenderer::InitViews(FRDGBuilder& GraphBuilder, const F
 		AverageViewPosition += View.ViewMatrices.GetViewOrigin() / Views.Num();
 	}
 
+	InitViewsBeforePrepass(GraphBuilder, InstanceCullingManager);
+
 	{
 		QUICK_SCOPE_CYCLE_COUNTER(STAT_InitViews_InitRHIResources);
 		// initialize per-view uniform buffer.
@@ -4672,9 +4682,24 @@ void FSceneRenderer::SetupSceneReflectionCaptureBuffer(FRHICommandListImmediate&
 	}
 }
 
-void FDeferredShadingSceneRenderer::InitViewsPossiblyAfterPrepass(FRDGBuilder& GraphBuilder, struct FILCUpdatePrimTaskData& ILCTaskData, FInstanceCullingManager& InstanceCullingManager)
+void FDeferredShadingSceneRenderer::InitViewsBeforePrepass(FRDGBuilder& GraphBuilder, FInstanceCullingManager& InstanceCullingManager)
 {
-	SCOPED_NAMED_EVENT(FDeferredShadingSceneRenderer_InitViewsPossiblyAfterPrepass, FColor::Emerald);
+	const bool bHasRayTracedOverlay = HasRayTracedOverlay(ViewFamily);
+
+	if (GEarlyInitDynamicShadows &&
+		CurrentDynamicShadowsTaskData == nullptr &&
+		ViewFamily.EngineShowFlags.DynamicShadows
+		&& !IsSimpleForwardShadingEnabled(ShaderPlatform)
+		&& !ViewFamily.EngineShowFlags.HitProxies
+		&& !bHasRayTracedOverlay)
+	{
+		CurrentDynamicShadowsTaskData = BeginInitDynamicShadows(true);
+	}
+}
+
+void FDeferredShadingSceneRenderer::InitViewsAfterPrepass(FRDGBuilder& GraphBuilder, struct FILCUpdatePrimTaskData& ILCTaskData, FInstanceCullingManager& InstanceCullingManager)
+{
+	SCOPED_NAMED_EVENT(FDeferredShadingSceneRenderer_InitViewsAfterPrepass, FColor::Emerald);
 	SCOPE_CYCLE_COUNTER(STAT_InitViewsPossiblyAfterPrepass);
 
 	FRHICommandListImmediate& RHICmdList = GraphBuilder.RHICmdList;
@@ -4687,7 +4712,15 @@ void FDeferredShadingSceneRenderer::InitViewsPossiblyAfterPrepass(FRDGBuilder& G
 		&& !bHasRayTracedOverlay)
 	{
 		// Setup dynamic shadows.
-		InitDynamicShadows(RHICmdList, DynamicIndexBufferForInitShadows, DynamicVertexBufferForInitShadows, DynamicReadBufferForInitShadows, InstanceCullingManager);
+		if (CurrentDynamicShadowsTaskData)
+		{
+			FinishInitDynamicShadows(RHICmdList, CurrentDynamicShadowsTaskData, DynamicIndexBufferForInitShadows, DynamicVertexBufferForInitShadows, DynamicReadBufferForInitShadows, InstanceCullingManager);
+			CurrentDynamicShadowsTaskData = nullptr;
+		}
+		else
+		{
+			InitDynamicShadows(RHICmdList, DynamicIndexBufferForInitShadows, DynamicVertexBufferForInitShadows, DynamicReadBufferForInitShadows, InstanceCullingManager);
+		}
 
 		RHICmdList.ImmediateFlush(EImmediateFlushType::DispatchToRHIThread);
 	}
diff --git a/Engine/Source/Runtime/Renderer/Private/ShadowRendering.h b/Engine/Source/Runtime/Renderer/Private/ShadowRendering.h
--- a/Engine/Source/Runtime/Renderer/Private/ShadowRendering.h
+++ b/Engine/Source/Runtime/Renderer/Private/ShadowRendering.h
@@ -574,11 +574,11 @@ public:
 	/**
 	 * Adds a primitive to the shadow's subject list.
 	 */
-	void AddSubjectPrimitive(FPrimitiveSceneInfo* PrimitiveSceneInfo, TArray<FViewInfo>* ViewArray, bool bRecordShadowSubjectForMobileShading);
+	void AddSubjectPrimitive(FPrimitiveSceneInfo* PrimitiveSceneInfo, TArrayView<FViewInfo> ViewArray, bool bRecordShadowSubjectForMobileShading);
 
 	uint64 AddSubjectPrimitive_AnyThread(
 		const FPrimitiveSceneInfoCompact& PrimitiveSceneInfoCompact,
-		TArray<FViewInfo>* ViewArray,
+		TArrayView<FViewInfo> ViewArray,
 		ERHIFeatureLevel::Type FeatureLevel,
 		struct FAddSubjectPrimitiveStats& OutStats,
 		struct FAddSubjectPrimitiveOverflowedIndices& OverflowBuffer) const;
@@ -587,7 +587,7 @@ public:
 
 	void FinalizeAddSubjectPrimitive(
 		struct FAddSubjectPrimitiveOp const& Op,
-		TArray<FViewInfo>* ViewArray,
+		TArrayView<FViewInfo> ViewArray,
 		struct FFinalizeAddSubjectPrimitiveContext& Context);
 
 	/**
diff --git a/Engine/Source/Runtime/Renderer/Private/ShadowSetup.cpp b/Engine/Source/Runtime/Renderer/Private/ShadowSetup.cpp
--- a/Engine/Source/Runtime/Renderer/Private/ShadowSetup.cpp
+++ b/Engine/Source/Runtime/Renderer/Private/ShadowSetup.cpp
@@ -1560,7 +1560,7 @@ bool FProjectedShadowInfo::ShouldDrawStaticMeshes_AnyThread(
 	return bDrawingStaticMeshes;
 }
 
-void FProjectedShadowInfo::AddSubjectPrimitive(FPrimitiveSceneInfo* PrimitiveSceneInfo, TArray<FViewInfo>* ViewArray, bool bShouldRecordShadowSubjectsForMobileShading)
+void FProjectedShadowInfo::AddSubjectPrimitive(FPrimitiveSceneInfo* PrimitiveSceneInfo, TArrayView<FViewInfo> ViewArray, bool bShouldRecordShadowSubjectsForMobileShading)
 {
 	// Ray traced shadows use the GPU managed distance field object buffers, no CPU culling should be used
 	check(!bRayTracedDistanceField);
@@ -1580,13 +1580,13 @@ void FProjectedShadowInfo::AddSubjectPrimitive(FPrimitiveSceneInfo* PrimitiveSce
 		}
 		else
 		{
-			checkf(ViewArray,
+			checkf(ViewArray.Num(),
 				TEXT("bWholeSceneShadow=%d, CascadeSettings.ShadowSplitIndex=%d, bDirectionalLight=%s"),
 				bWholeSceneShadow ? TEXT("true") : TEXT("false"),
 				CascadeSettings.ShadowSplitIndex,
 				bDirectionalLight ? TEXT("true") : TEXT("false"));
 
-			Views = TArrayView<FViewInfo>(*ViewArray);
+			Views = ViewArray;
 		}
 
 		bool bOpaque = false;
@@ -1753,7 +1753,7 @@ bool FProjectedShadowInfo::TestPrimitiveFarCascadeConditions(bool bPrimitiveCast
 
 uint64 FProjectedShadowInfo::AddSubjectPrimitive_AnyThread(
 	const FPrimitiveSceneInfoCompact& PrimitiveSceneInfoCompact,
-	TArray<FViewInfo>* ViewArray,
+	TArrayView<FViewInfo> ViewArray,
 	ERHIFeatureLevel::Type FeatureLevel,
 	FAddSubjectPrimitiveStats& OutStats,
 	FAddSubjectPrimitiveOverflowedIndices& OverflowBuffer) const
@@ -1792,20 +1792,20 @@ uint64 FProjectedShadowInfo::AddSubjectPrimitive_AnyThread(
 		}
 		else
 		{
-			checkf(ViewArray,
+			checkf(ViewArray.Num(),
 				TEXT("bWholeSceneShadow=%d, CascadeSettings.ShadowSplitIndex=%d, bDirectionalLight=%s"),
 				bWholeSceneShadow ? TEXT("true") : TEXT("false"),
 				CascadeSettings.ShadowSplitIndex,
 				bDirectionalLight ? TEXT("true") : TEXT("false"));
 
-			if (ViewArray->Num() > 1)
+			if (ViewArray.Num() > 1)
 			{
 				Result.bAddOnRenderThread = true;
 				++OutStats.NumDeferredPrimitives;
 				return Result.Qword;
 			}
 
-			CurrentView = &(*ViewArray)[0];
+			CurrentView = &ViewArray[0];
 		}
 
 		bool bOpaque = false;
@@ -1953,7 +1953,7 @@ void FProjectedShadowInfo::PresizeSubjectPrimitiveArrays(const FAddSubjectPrimit
 
 void FProjectedShadowInfo::FinalizeAddSubjectPrimitive(
 	const FAddSubjectPrimitiveOp& Op,
-	TArray<FViewInfo>* ViewArray,
+	TArrayView<FViewInfo> ViewArray,
 	FFinalizeAddSubjectPrimitiveContext& Context)
 {
 	FPrimitiveSceneInfo* PrimitiveSceneInfo = Op.PrimitiveSceneInfo;
@@ -1985,7 +1985,7 @@ void FProjectedShadowInfo::FinalizeAddSubjectPrimitive(
 
 	if (Result.bNeedPrimitiveFadingStateUpdate)
 	{
-		FViewInfo& View = IsWholeSceneDirectionalShadow() ? *DependentView : (*ViewArray)[0];
+		FViewInfo& View = IsWholeSceneDirectionalShadow() ? *DependentView : ViewArray[0];
 		if (View.UpdatePrimitiveFadingState(PrimitiveSceneInfo, Result.bFadingIn))
 		{
 			if (Result.bOverflowed)
@@ -2889,7 +2889,7 @@ void FSceneRenderer::CreatePerObjectProjectedShadow(
 						for (int32 ChildIndex = 0, ChildCount = ShadowGroupPrimitives.Num(); ChildIndex < ChildCount; ChildIndex++)
 						{
 							FPrimitiveSceneInfo* ShadowChild = ShadowGroupPrimitives[ChildIndex];
-							ProjectedShadowInfo->AddSubjectPrimitive(ShadowChild, &Views, false);
+							ProjectedShadowInfo->AddSubjectPrimitive(ShadowChild, Views, false);
 						}
 					}
 					else if (bShadowIsPotentiallyVisibleNextFrame)
@@ -2931,7 +2931,7 @@ void FSceneRenderer::CreatePerObjectProjectedShadow(
 						for (int32 ChildIndex = 0, ChildCount = ShadowGroupPrimitives.Num(); ChildIndex < ChildCount; ChildIndex++)
 						{
 							FPrimitiveSceneInfo* ShadowChild = ShadowGroupPrimitives[ChildIndex];
-							ProjectedShadowInfo->AddSubjectPrimitive(ShadowChild, &Views, false);
+							ProjectedShadowInfo->AddSubjectPrimitive(ShadowChild, Views, false);
 						}
 					}
 					else if (bShadowIsPotentiallyVisibleNextFrame)
@@ -3649,7 +3649,7 @@ void FSceneRenderer::CreateWholeSceneProjectedShadow(
 								FBoxSphereBounds const& Bounds = Interaction->GetPrimitiveSceneInfo()->Proxy->GetBounds();
 								if (IntersectsConvexHulls(LightViewFrustumConvexHulls, Bounds))
 								{
-									ProjectedShadowInfo->AddSubjectPrimitive(Interaction->GetPrimitiveSceneInfo(), &Views, false);
+									ProjectedShadowInfo->AddSubjectPrimitive(Interaction->GetPrimitiveSceneInfo(), Views, false);
 								}
 							}
 						}
@@ -4025,59 +4025,98 @@ void FSceneRenderer::GatherShadowDynamicMeshElements(FGlobalDynamicIndexBuffer&
 	}
 }
 
+inline bool AreAnyViewsStaticSceneOnly(TArrayView<FViewInfo> Views)
+{
+	for (const FViewInfo& View : Views)
+	{
+		if (View.bStaticSceneOnly)
+		{
+			return true;
+		}
+	}
+	return false;
+}
+
 typedef TArray<FAddSubjectPrimitiveOp> FShadowSubjectPrimitives;
-typedef TArray<FAddSubjectPrimitiveStats, TInlineAllocator<4, SceneRenderingAllocator>> FPerShadowGatherStats;
-typedef TArray<FAddSubjectPrimitiveOverflowedIndices, SceneRenderingAllocator> FPerShadowOverflowedIndices;
+typedef TArray<FAddSubjectPrimitiveStats> FPerShadowGatherStats;
+
+// Common setup and working data for all GatherShadowPrimitives Tasks
+struct FDynamicShadowsTaskData
+{
+	// Common data read from all points
+	const FScene* Scene;
+	TArrayView<FViewInfo> Views;
+	const ERHIFeatureLevel::Type FeatureLevel;
+	const EShaderPlatform ShaderPlatform;
+	const bool bStaticSceneOnly;
+	const bool bRunningEarly;
+	const bool bMultithreaded;
+
+	// Generated from RenderThread
+	TArray<FProjectedShadowInfo*, SceneRenderingAllocator> PreShadows;
+	TArray<FProjectedShadowInfo*, SceneRenderingAllocator> ViewDependentWholeSceneShadows;
+
+	// Written from task
+	TArray<struct FGatherShadowPrimitivesPacket*> Packets;
+	FPerShadowGatherStats GatherStats;
+
+	// Used by RenderThread
+	FGraphEventRef TaskEvent;
+
+	FDynamicShadowsTaskData(const FScene* InScene, TArray<FViewInfo>& InViews, bool bInRunningEarly)
+		: Scene(InScene)
+		, Views(InViews)
+		, FeatureLevel(InScene->GetFeatureLevel())
+		, ShaderPlatform(GShaderPlatformForFeatureLevel[FeatureLevel])
+		, bStaticSceneOnly(AreAnyViewsStaticSceneOnly(InViews))
+		, bRunningEarly(bInRunningEarly)
+		, bMultithreaded(FApp::ShouldUseThreadingForPerformance() || FForkProcessHelper::IsForkedMultithreadInstance())
+	{
+	}
+};
+
+FAutoConsoleTaskPriority CPrio_GatherShadowPrimitives(
+	TEXT("TaskGraph.TaskPriorities.GatherShadowPrimitives"),
+	TEXT("Task and thread priority for GatherShadowPrimitives tasks."),
+	ENamedThreads::HighThreadPriority,
+	ENamedThreads::NormalTaskPriority,
+	ENamedThreads::HighTaskPriority
+);
 
 struct FGatherShadowPrimitivesPacket
 {
 	// Inputs
-	const FScene* Scene;
-	TArray<FViewInfo>& Views;
-	const FScenePrimitiveOctree::FNodeIndex NodeIndex;
-	int32 StartPrimitiveIndex;
-	int32 NumPrimitives;
-	const TArray<FProjectedShadowInfo*, SceneRenderingAllocator>& PreShadows;
-	const TArray<FProjectedShadowInfo*, SceneRenderingAllocator>& ViewDependentWholeSceneShadows;
-	ERHIFeatureLevel::Type FeatureLevel;
-	bool bStaticSceneOnly;
+	FScenePrimitiveOctree::FNodeIndex NodeIndex = INDEX_NONE;
+	int32 StartPrimitiveIndex{ 0 };
+	int32 NumPrimitives{ 0 };
 
 	// Scratch
 	FPerShadowGatherStats ViewDependentWholeSceneShadowStats;
-	FPerShadowOverflowedIndices PreShadowOverflowedIndices;
-	FPerShadowOverflowedIndices ViewDependentWholeSceneShadowOverflowedIndices;
-	TArray<FShadowSubjectPrimitives, SceneRenderingAllocator> PreShadowSubjectPrimitives;
-	TArray<FShadowSubjectPrimitives, SceneRenderingAllocator> ViewDependentWholeSceneShadowSubjectPrimitives;
 
-	// Outputs
-	FPerShadowGatherStats& GlobalStats;
+	// Output
+	TArray<FAddSubjectPrimitiveOverflowedIndices> PreShadowOverflowedIndices;
+	TArray<FAddSubjectPrimitiveOverflowedIndices> ViewDependentWholeSceneShadowOverflowedIndices;
+	TArray<FShadowSubjectPrimitives> PreShadowSubjectPrimitives;
+	TArray<FShadowSubjectPrimitives> ViewDependentWholeSceneShadowSubjectPrimitives;
 
 	FGatherShadowPrimitivesPacket(
-		const FScene* InScene,
-		TArray<FViewInfo>& InViews,
-		const FScenePrimitiveOctree::FNodeIndex InNodeIndex,
+		FScenePrimitiveOctree::FNodeIndex InNodeIndex,
 		int32 InStartPrimitiveIndex,
-		int32 InNumPrimitives,
-		const TArray<FProjectedShadowInfo*, SceneRenderingAllocator>& InPreShadows,
-		const TArray<FProjectedShadowInfo*, SceneRenderingAllocator>& InViewDependentWholeSceneShadows,
-		ERHIFeatureLevel::Type InFeatureLevel,
-		bool bInStaticSceneOnly,
-		FPerShadowGatherStats& OutGlobalStats) 
-		: Scene(InScene)
-		, Views(InViews)
-		, NodeIndex(InNodeIndex)
+		int32 InNumPrimitives)
+		: NodeIndex(InNodeIndex)
 		, StartPrimitiveIndex(InStartPrimitiveIndex)
 		, NumPrimitives(InNumPrimitives)
-		, PreShadows(InPreShadows)
-		, ViewDependentWholeSceneShadows(InViewDependentWholeSceneShadows)
-		, FeatureLevel(InFeatureLevel)
-		, bStaticSceneOnly(bInStaticSceneOnly)
-		, GlobalStats(OutGlobalStats)
 	{
-		const int32 NumPreShadows = PreShadows.Num();
-		const int32 NumVDWSShadows = ViewDependentWholeSceneShadows.Num();
+	}
 
-		check(GlobalStats.Num() == NumVDWSShadows);
+	void AnyThreadTask(FDynamicShadowsTaskData& TaskData)
+	{
+		FOptionalTaskTagScope Scope(ETaskTag::EParallelRenderingThread);
+
+		const int32 NumPreShadows = TaskData.PreShadows.Num();
+		const int32 NumVDWSShadows = TaskData.ViewDependentWholeSceneShadows.Num();
+
+		check(TaskData.GatherStats.Num() == NumVDWSShadows);
 		ViewDependentWholeSceneShadowStats.Empty(NumVDWSShadows);
 		ViewDependentWholeSceneShadowStats.AddDefaulted(NumVDWSShadows);
 
@@ -4092,22 +4131,17 @@ struct FGatherShadowPrimitivesPacket
 
 		ViewDependentWholeSceneShadowSubjectPrimitives.Empty(NumVDWSShadows);
 		ViewDependentWholeSceneShadowSubjectPrimitives.AddDefaulted(NumVDWSShadows);
-	}
-
-	void AnyThreadTask()
-	{
-		FOptionalTaskTagScope Scope(ETaskTag::EParallelRenderingThread);
 
 		if (NodeIndex != INDEX_NONE)
 		{
 			// Check all the primitives in this octree node.
-			for (const FPrimitiveSceneInfoCompact& PrimitiveSceneInfoCompact : Scene->PrimitiveOctree.GetElementsForNode(NodeIndex))
+			for (const FPrimitiveSceneInfoCompact& PrimitiveSceneInfoCompact : TaskData.Scene->PrimitiveOctree.GetElementsForNode(NodeIndex))
 			{
 				// Nanite has its own culling
 				if (PrimitiveSceneInfoCompact.PrimitiveFlagsCompact.bCastDynamicShadow &&
 					(!PrimitiveSceneInfoCompact.PrimitiveFlagsCompact.bIsNaniteMesh || GSkipCullingNaniteMeshes == 0))
 				{
-					FilterPrimitiveForShadows(PrimitiveSceneInfoCompact);
+					FilterPrimitiveForShadows(TaskData, PrimitiveSceneInfoCompact);
 				}
 			}
 		}
@@ -4118,30 +4152,30 @@ struct FGatherShadowPrimitivesPacket
 			// Check primitives in this packet's range
 			for (int32 PrimitiveIndex = StartPrimitiveIndex; PrimitiveIndex < StartPrimitiveIndex + NumPrimitives; PrimitiveIndex++)
 			{
-				const FPrimitiveFlagsCompact PrimitiveFlagsCompact = Scene->PrimitiveFlagsCompact[PrimitiveIndex];
+				const FPrimitiveFlagsCompact PrimitiveFlagsCompact = TaskData.Scene->PrimitiveFlagsCompact[PrimitiveIndex];
 
 				// Nanite has its own culling
 				if (PrimitiveFlagsCompact.bCastDynamicShadow &&
 					(!PrimitiveFlagsCompact.bIsNaniteMesh || GSkipCullingNaniteMeshes == 0))
 				{
-					FPrimitiveSceneInfo* PrimitiveSceneInfo = Scene->Primitives[PrimitiveIndex];
+					FPrimitiveSceneInfo* PrimitiveSceneInfo = TaskData.Scene->Primitives[PrimitiveIndex];
 					const FPrimitiveSceneInfoCompact PrimitiveSceneInfoCompact(PrimitiveSceneInfo);
 
-					FilterPrimitiveForShadows(PrimitiveSceneInfoCompact);
+					FilterPrimitiveForShadows(TaskData, PrimitiveSceneInfoCompact);
 				}
 			}
 		}
 
-		const int32 NumStatsToMerge = ViewDependentWholeSceneShadowStats.Num();
-		for (int32 StatIdx = 0; StatIdx < NumStatsToMerge; ++StatIdx)
+		for (int32 StatIndex=0, StatNum = ViewDependentWholeSceneShadowStats.Num(); StatIndex < StatNum; StatIndex++)
 		{
-			GlobalStats[StatIdx].InterlockedAdd(ViewDependentWholeSceneShadowStats[StatIdx]);
+			TaskData.GatherStats[StatIndex].InterlockedAdd(ViewDependentWholeSceneShadowStats[StatIndex]);
 		}
+		ViewDependentWholeSceneShadowStats.Empty();
 	}
 
-	bool DoesPrimitiveCastInsetShadow(const FPrimitiveSceneInfo* PrimitiveSceneInfo, const FPrimitiveSceneProxy* PrimitiveProxy) const
+	bool DoesPrimitiveCastInsetShadow(FDynamicShadowsTaskData& TaskData, const FPrimitiveSceneInfo* PrimitiveSceneInfo, const FPrimitiveSceneProxy* PrimitiveProxy) const
 	{
-		if (UseNonNaniteVirtualShadowMaps(GShaderPlatformForFeatureLevel[FeatureLevel], FeatureLevel))
+		if (UseNonNaniteVirtualShadowMaps(TaskData.ShaderPlatform, TaskData.FeatureLevel))
 		{
 			return false;
 		}
@@ -4158,7 +4192,7 @@ struct FGatherShadowPrimitivesPacket
 		}
 	}
 
-	void FilterPrimitiveForShadows(const FPrimitiveSceneInfoCompact& PrimitiveSceneInfoCompact)
+	void FilterPrimitiveForShadows(FDynamicShadowsTaskData& TaskData, const FPrimitiveSceneInfoCompact& PrimitiveSceneInfoCompact)
 	{
 		const FPrimitiveFlagsCompact& PrimitiveFlagsCompact = PrimitiveSceneInfoCompact.PrimitiveFlagsCompact;
 		const FBoxSphereBounds& PrimitiveBounds = PrimitiveSceneInfoCompact.Bounds;
@@ -4167,11 +4201,11 @@ struct FGatherShadowPrimitivesPacket
 
 		// Check if the primitive is a subject for any of the preshadows.
 		// Only allow preshadows from lightmapped primitives that cast both dynamic and static shadows.
-		if (PreShadows.Num() && PrimitiveFlagsCompact.bCastStaticShadow && PrimitiveFlagsCompact.bStaticLighting)
+		if (TaskData.PreShadows.Num() && PrimitiveFlagsCompact.bCastStaticShadow && PrimitiveFlagsCompact.bStaticLighting)
 		{
-			for (int32 ShadowIndex = 0, Num = PreShadows.Num(); ShadowIndex < Num; ShadowIndex++)
+			for (int32 ShadowIndex = 0, Num = TaskData.PreShadows.Num(); ShadowIndex < Num; ShadowIndex++)
 			{
-				FProjectedShadowInfo* RESTRICT ProjectedShadowInfo = PreShadows[ShadowIndex];
+				FProjectedShadowInfo* RESTRICT ProjectedShadowInfo = TaskData.PreShadows[ShadowIndex];
 
 				// Note: Culling based on the primitive's bounds BEFORE dereferencing PrimitiveSceneInfo / PrimitiveProxy
 				// Check if this primitive is in the shadow's frustum.
@@ -4181,9 +4215,9 @@ struct FGatherShadowPrimitivesPacket
 				{
 					FAddSubjectPrimitiveResult Result;
 					FAddSubjectPrimitiveStats UnusedStats;
-					Result.Qword = ProjectedShadowInfo->AddSubjectPrimitive_AnyThread(PrimitiveSceneInfoCompact, &Views, FeatureLevel, UnusedStats, PreShadowOverflowedIndices[ShadowIndex]);
+					Result.Qword = ProjectedShadowInfo->AddSubjectPrimitive_AnyThread(PrimitiveSceneInfoCompact, TaskData.Views, TaskData.FeatureLevel, UnusedStats, PreShadowOverflowedIndices[ShadowIndex]);
 
-					if (!!Result.Qword)
+					if (Result.Qword != 0)
 					{
 						FShadowSubjectPrimitives& SubjectPrimitives = PreShadowSubjectPrimitives[ShadowIndex];
 						FAddSubjectPrimitiveOp& Op = SubjectPrimitives[SubjectPrimitives.AddUninitialized()];
@@ -4194,9 +4228,9 @@ struct FGatherShadowPrimitivesPacket
 			}
 		}
 
-		for (int32 ShadowIndex = 0, Num = ViewDependentWholeSceneShadows.Num();ShadowIndex < Num;ShadowIndex++)
+		for (int32 ShadowIndex = 0, Num = TaskData.ViewDependentWholeSceneShadows.Num();ShadowIndex < Num;ShadowIndex++)
 		{
-			const FProjectedShadowInfo* RESTRICT ProjectedShadowInfo = ViewDependentWholeSceneShadows[ShadowIndex];
+			const FProjectedShadowInfo* RESTRICT ProjectedShadowInfo = TaskData.ViewDependentWholeSceneShadows[ShadowIndex];
 			const FLightSceneInfo& RESTRICT LightSceneInfo = ProjectedShadowInfo->GetLightSceneInfo();
 			const FLightSceneProxy& RESTRICT LightProxy = *LightSceneInfo.Proxy;
 
@@ -4233,7 +4267,7 @@ struct FGatherShadowPrimitivesPacket
 
 				if (ProjectedShadowInfo->CacheMode == SDCM_MovablePrimitivesOnly)
 				{
-					const FCachedShadowMapData& CachedShadowMapData = Scene->GetCachedShadowMapDataRef(LightSceneInfo.Id, ProjectedShadowInfo->CascadeSettings.ShadowSplitIndex);
+					const FCachedShadowMapData& CachedShadowMapData = TaskData.Scene->GetCachedShadowMapDataRef(LightSceneInfo.Id, ProjectedShadowInfo->CascadeSettings.ShadowSplitIndex);
 
 					// if the mesh is dynamic or a new added mesh, it should cast csm.
 					bShadowRelevance = PrimitiveProxy->IsMeshShapeOftenMoving() || !CachedShadowMapData.StaticShadowSubjectMap[PrimitiveSceneInfo->GetIndex()];
@@ -4244,7 +4278,7 @@ struct FGatherShadowPrimitivesPacket
 				}
 				else if (ProjectedShadowInfo->CacheMode == SDCM_CSMScrolling)
 				{
-					const FCachedShadowMapData& CachedShadowMapData = Scene->GetCachedShadowMapDataRef(LightSceneInfo.Id, ProjectedShadowInfo->CascadeSettings.ShadowSplitIndex);
+					const FCachedShadowMapData& CachedShadowMapData = TaskData.Scene->GetCachedShadowMapDataRef(LightSceneInfo.Id, ProjectedShadowInfo->CascadeSettings.ShadowSplitIndex);
 
 					// if the mesh is dynamic or a new added mesh, it should cast csm.
 					bShadowRelevance = PrimitiveProxy->IsMeshShapeOftenMoving() || !CachedShadowMapData.StaticShadowSubjectMap[PrimitiveSceneInfo->GetIndex()];
@@ -4267,23 +4301,23 @@ struct FGatherShadowPrimitivesPacket
 					// Since lights with static shadowing still create per-object shadows for primitives without static shadowing.
 					&& (!LightProxy.HasStaticLighting() || (!LightSceneInfo.IsPrecomputedLightingValid() || LightProxy.UseCSMForDynamicObjects()))
 					// Exclude primitives that will create their own per-object shadow, except when rendering RSMs
-					&& (!DoesPrimitiveCastInsetShadow(PrimitiveSceneInfo, PrimitiveProxy))
+					&& (!DoesPrimitiveCastInsetShadow(TaskData, PrimitiveSceneInfo, PrimitiveProxy))
 					// Exclude primitives that will create a per-object shadow from a stationary light
 					&& !ShouldCreateObjectShadowForStationaryLight(&LightSceneInfo, PrimitiveProxy, true)
 					// Only render shadows from objects that use static lighting during a reflection capture, since the reflection capture doesn't update at runtime
-					&& (!bStaticSceneOnly || PrimitiveProxy->HasStaticLighting())
+					&& (!TaskData.bStaticSceneOnly || PrimitiveProxy->HasStaticLighting())
 					// Render dynamic lit objects if CSMForDynamicObjects is enabled.
 					&& (!LightProxy.UseCSMForDynamicObjects() || !PrimitiveProxy->HasStaticLighting()))
 				{
 					FAddSubjectPrimitiveResult Result;
 					Result.Qword = ProjectedShadowInfo->AddSubjectPrimitive_AnyThread(
 						PrimitiveSceneInfoCompact,
-						nullptr,
-						FeatureLevel,
+						TArrayView<FViewInfo>(),
+						TaskData.FeatureLevel,
 						ViewDependentWholeSceneShadowStats[ShadowIndex],
 						ViewDependentWholeSceneShadowOverflowedIndices[ShadowIndex]);
 
-					if (!!Result.Qword)
+					if (Result.Qword != 0)
 					{
 						FShadowSubjectPrimitives& SubjectPrimitives = ViewDependentWholeSceneShadowSubjectPrimitives[ShadowIndex];
 						if (!SubjectPrimitives.Num())
@@ -4299,27 +4333,25 @@ struct FGatherShadowPrimitivesPacket
 		}
 	}
 
-	void RenderThreadFinalize()
+	void RenderThreadFinalize(FDynamicShadowsTaskData& TaskData)
 	{
 		for (int32 ShadowIndex = 0; ShadowIndex < PreShadowSubjectPrimitives.Num(); ShadowIndex++)
 		{
-			FProjectedShadowInfo* ProjectedShadowInfo = PreShadows[ShadowIndex];
-			const FShadowSubjectPrimitives& SubjectPrimitives = PreShadowSubjectPrimitives[ShadowIndex];
+			FProjectedShadowInfo* ProjectedShadowInfo = TaskData.PreShadows[ShadowIndex];
 			const FAddSubjectPrimitiveOverflowedIndices& OverflowBuffer = PreShadowOverflowedIndices[ShadowIndex];
 			FFinalizeAddSubjectPrimitiveContext Context;
 			Context.OverflowedMDCIndices = OverflowBuffer.MDCIndices.GetData();
 			Context.OverflowedMeshIndices = OverflowBuffer.MeshIndices.GetData();
 
-			for (int32 PrimitiveIndex = 0; PrimitiveIndex < SubjectPrimitives.Num(); PrimitiveIndex++)
+			for (const FAddSubjectPrimitiveOp& PrimitiveOp : PreShadowSubjectPrimitives[ShadowIndex])
 			{
-				ProjectedShadowInfo->FinalizeAddSubjectPrimitive(SubjectPrimitives[PrimitiveIndex], &Views, Context);
+				ProjectedShadowInfo->FinalizeAddSubjectPrimitive(PrimitiveOp, TaskData.Views, Context);
 			}
 		}
 
 		for (int32 ShadowIndex = 0; ShadowIndex < ViewDependentWholeSceneShadowSubjectPrimitives.Num(); ShadowIndex++)
 		{
-			FProjectedShadowInfo* ProjectedShadowInfo = ViewDependentWholeSceneShadows[ShadowIndex];
-			const FShadowSubjectPrimitives& SubjectPrimitives = ViewDependentWholeSceneShadowSubjectPrimitives[ShadowIndex];
+			FProjectedShadowInfo* ProjectedShadowInfo = TaskData.ViewDependentWholeSceneShadows[ShadowIndex];
 			const FAddSubjectPrimitiveOverflowedIndices& OverflowBuffer = ViewDependentWholeSceneShadowOverflowedIndices[ShadowIndex];
 			FFinalizeAddSubjectPrimitiveContext Context;
 			Context.OverflowedMDCIndices = OverflowBuffer.MDCIndices.GetData();
@@ -4329,17 +4361,17 @@ struct FGatherShadowPrimitivesPacket
 
 			if (ProjectedShadowInfo->CacheMode != SDCM_Uncached && ProjectedShadowInfo->IsWholeSceneDirectionalShadow())
 			{
-				CachedShadowMapData = const_cast<FCachedShadowMapData*>(Scene->GetCachedShadowMapData(ProjectedShadowInfo->GetLightSceneInfo().Id, ProjectedShadowInfo->CascadeSettings.ShadowSplitIndex));
-				
+				CachedShadowMapData = const_cast<FCachedShadowMapData*>(TaskData.Scene->GetCachedShadowMapData(ProjectedShadowInfo->GetLightSceneInfo().Id, ProjectedShadowInfo->CascadeSettings.ShadowSplitIndex));
+
 				if (ProjectedShadowInfo->CacheMode == SDCM_MovablePrimitivesOnly || ProjectedShadowInfo->CacheMode == SDCM_CSMScrolling)
 				{
 					CachedShadowMapData->LastFrameExtraStaticShadowSubjects = 0;
 				}
 			}
 
-			for (int32 PrimitiveIndex = 0; PrimitiveIndex < SubjectPrimitives.Num(); PrimitiveIndex++)
+			for (const FAddSubjectPrimitiveOp& PrimitiveOp : ViewDependentWholeSceneShadowSubjectPrimitives[ShadowIndex])
 			{
-				ProjectedShadowInfo->FinalizeAddSubjectPrimitive(SubjectPrimitives[PrimitiveIndex], nullptr, Context);
+				ProjectedShadowInfo->FinalizeAddSubjectPrimitive(PrimitiveOp, TArrayView<FViewInfo>(), Context);
 			}
 
 			if (ProjectedShadowInfo->CacheMode == SDCM_StaticPrimitivesOnly && ProjectedShadowInfo->IsWholeSceneDirectionalShadow())
@@ -4351,136 +4383,213 @@ struct FGatherShadowPrimitivesPacket
 	}
 };
 
-void FSceneRenderer::GatherShadowPrimitives(
-	const TArray<FProjectedShadowInfo*,SceneRenderingAllocator>& PreShadows,
-	const TArray<FProjectedShadowInfo*,SceneRenderingAllocator>& ViewDependentWholeSceneShadows,
-	bool bStaticSceneOnly
-	)
+struct FGatherShadowPrimitivesPrepareTask
 {
-	SCOPE_CYCLE_COUNTER(STAT_GatherShadowPrimitivesTime);
+	FDynamicShadowsTaskData& TaskData;
+
+	FGatherShadowPrimitivesPrepareTask(FDynamicShadowsTaskData& InTaskData)
+		: TaskData(InTaskData)
+	{
+	}
+
+	static FORCEINLINE TStatId GetStatId()
+	{
+		return GET_STATID(STAT_ParallelForTask);
+	}
+
+	static ENamedThreads::Type GetDesiredThread()
+	{
+		return CPrio_GatherShadowPrimitives.Get();
+	}
+
+	static ESubsequentsMode::Type GetSubsequentsMode()
+	{
+		return ESubsequentsMode::TrackSubsequents;
+	}
+
+	void DoTask(ENamedThreads::Type CurrentThread, const FGraphEventRef& MyCompletionGraphEvent)
+	{
+		SCOPED_NAMED_EVENT_TEXT("FGatherShadowPrimitivesPrepareTask", FColor::Green);
+
+		AnyThreadTask();
+
+		// Kick off any child tasks if we added them
 
-	if (PreShadows.Num() || ViewDependentWholeSceneShadows.Num())
+		const int32 NumPackets = TaskData.Packets.Num();
+		const int32 NumTasks = FMath::Min<int32>(FTaskGraphInterface::Get().GetNumWorkerThreads(), NumPackets);
+		const int32 PacketsPerTask = NumPackets / NumTasks;
+
+		for (int32 TaskIndex = 0; TaskIndex < NumTasks; TaskIndex++)
+		{
+			FDynamicShadowsTaskData* pTaskData = &TaskData;
+			const int32 FirstPacketToProcess = TaskIndex * PacketsPerTask;
+			const int32 NumPacketsToProcess = PacketsPerTask;
+
+			FGraphEventRef NewTask = FFunctionGraphTask::CreateAndDispatchWhenReady([pTaskData, FirstPacketToProcess, NumPacketsToProcess]()
+				{
+					SCOPED_NAMED_EVENT_TEXT("FGatherShadowPrimitivesTask", FColor::Green);
+					for (int32 PacketIndex = 0; PacketIndex < NumPacketsToProcess; PacketIndex++)
+					{
+						if (int32 PacketToProcess = FirstPacketToProcess + PacketIndex; PacketToProcess < pTaskData->Packets.Num())
+						{
+							pTaskData->Packets[PacketToProcess]->AnyThreadTask(*pTaskData);
+						}
+					}
+				},
+				GET_STATID(STAT_ParallelForTask),
+				nullptr,
+				CPrio_GatherShadowPrimitives.Get());
+
+			MyCompletionGraphEvent->DontCompleteUntil(NewTask);
+		}
+	}
+
+	void AddSubTask(FScenePrimitiveOctree::FNodeIndex NodeIndex, int32 StartPrimitiveIndex, int32 NumPrimitives)
 	{
-		TArray<FGatherShadowPrimitivesPacket*,SceneRenderingAllocator> Packets;
-		FPerShadowGatherStats GatherStats;
+		if (TaskData.bMultithreaded)
+		{
+			TaskData.Packets.Add(new FGatherShadowPrimitivesPacket(NodeIndex, StartPrimitiveIndex, NumPrimitives));
+		}
+		else
+		{
+			TaskData.Packets.Add(new(FMemStack::Get()) FGatherShadowPrimitivesPacket(NodeIndex, StartPrimitiveIndex, NumPrimitives));
+		}
+	}
 
-		GatherStats.AddDefaulted(ViewDependentWholeSceneShadows.Num());
+	void AnyThreadTask()
+	{
+		TaskData.GatherStats.AddDefaulted(TaskData.ViewDependentWholeSceneShadows.Num());
 
 		if (GUseOctreeForShadowCulling)
 		{
 			QUICK_SCOPE_CYCLE_COUNTER(STAT_ShadowSceneOctreeTraversal);
 
-			Packets.Reserve(100);
+			TaskData.Packets.Reserve(100);
 
 			// Find primitives that are in a shadow frustum in the octree.
-			Scene->PrimitiveOctree.FindNodesWithPredicate([&PreShadows, &ViewDependentWholeSceneShadows](const FBoxCenterAndExtent& NodeBounds)
-			{
-				// Check that the child node is in the frustum for at least one shadow.
-
-				// Check for subjects of preshadows.
-				for (int32 ShadowIndex = 0, Num = PreShadows.Num(); ShadowIndex < Num; ShadowIndex++)
+			TaskData.Scene->PrimitiveOctree.FindNodesWithPredicate(
+				[this](const FBoxCenterAndExtent& NodeBounds)
 				{
-					FProjectedShadowInfo* ProjectedShadowInfo = PreShadows[ShadowIndex];
-
-					check(ProjectedShadowInfo->CasterOuterFrustum.PermutedPlanes.Num());
-					// Check if this primitive is in the shadow's frustum.
-					if (ProjectedShadowInfo->CasterOuterFrustum.IntersectBox(
-						NodeBounds.Center + ProjectedShadowInfo->PreShadowTranslation,
-						NodeBounds.Extent
-					))
-					return true;
-				}
+					// Check that the child node is in the frustum for at least one shadow.
 
-				for (int32 ShadowIndex = 0, Num = ViewDependentWholeSceneShadows.Num(); ShadowIndex < Num; ShadowIndex++)
-				{
-					FProjectedShadowInfo* ProjectedShadowInfo = ViewDependentWholeSceneShadows[ShadowIndex];
-
-					//check(ProjectedShadowInfo->CasterFrustum.PermutedPlanes.Num());
-					// Check if this primitive is in the shadow's frustum.
-					if (ProjectedShadowInfo->CasterOuterFrustum.IntersectBox(
-						NodeBounds.Center + ProjectedShadowInfo->PreShadowTranslation,
-						NodeBounds.Extent
-					))
-					return true;
-				}
+					// Check for subjects of preshadows.
+					for (FProjectedShadowInfo* ProjectedShadowInfo : TaskData.PreShadows)
+					{
+						check(ProjectedShadowInfo->CasterOuterFrustum.PermutedPlanes.Num());
+						// Check if this primitive is in the shadow's frustum.
+						if (ProjectedShadowInfo->CasterOuterFrustum.IntersectBox(
+							NodeBounds.Center + ProjectedShadowInfo->PreShadowTranslation,
+							NodeBounds.Extent))
+						{
+							return true;
+						}
+					}
 
-				// If the child node was in the frustum of at least one preshadow, push it on
-				// the iterator's pending node stack.
-				return false;
-			},
-			[this, &PreShadows, &ViewDependentWholeSceneShadows, bStaticSceneOnly, &GatherStats, &Packets](FScenePrimitiveOctree::FNodeIndex NodeIndex)
-			{
-				if (Scene->PrimitiveOctree.GetElementsForNode(NodeIndex).Num() > 0)
+					for (FProjectedShadowInfo* ProjectedShadowInfo : TaskData.ViewDependentWholeSceneShadows)
+					{
+						//check(ProjectedShadowInfo->CasterFrustum.PermutedPlanes.Num());
+						// Check if this primitive is in the shadow's frustum.
+						if (ProjectedShadowInfo->CasterOuterFrustum.IntersectBox(
+							NodeBounds.Center + ProjectedShadowInfo->PreShadowTranslation,
+							NodeBounds.Extent))
+						{
+							return true;
+						}
+					}
+
+					// If the child node was in the frustum of at least one preshadow, push it on
+					// the iterator's pending node stack.
+					return false;
+				},
+				[this](FScenePrimitiveOctree::FNodeIndex NodeIndex)
 				{
-					FGatherShadowPrimitivesPacket* Packet = new(FMemStack::Get()) FGatherShadowPrimitivesPacket(
-						Scene,
-						Views,
-						NodeIndex,
-						0,
-						0,
-						PreShadows,
-						ViewDependentWholeSceneShadows,
-						FeatureLevel,
-						bStaticSceneOnly,
-						GatherStats);
-					Packets.Add(Packet);
-				}
-			});
+					if (TaskData.Scene->PrimitiveOctree.GetElementsForNode(NodeIndex).Num() > 0)
+					{
+						AddSubTask(NodeIndex, 0, 0);
+					}
+				});
 		}
 		else
 		{
 			const int32 PacketSize = CVarParallelGatherNumPrimitivesPerPacket.GetValueOnRenderThread();
-			const int32 NumPackets = FMath::DivideAndRoundUp(Scene->Primitives.Num(), PacketSize);
-			
-			Packets.Reserve(NumPackets);
+			const int32 NumPackets = FMath::DivideAndRoundUp(TaskData.Scene->Primitives.Num(), PacketSize);
+
+			TaskData.Packets.Reserve(NumPackets);
 
 			for (int32 PacketIndex = 0; PacketIndex < NumPackets; PacketIndex++)
 			{
 				const int32 StartPrimitiveIndex = PacketIndex * PacketSize;
-				const int32 NumPrimitives = FMath::Min(PacketSize, Scene->Primitives.Num() - StartPrimitiveIndex);
-				FGatherShadowPrimitivesPacket* Packet = new(FMemStack::Get()) FGatherShadowPrimitivesPacket(
-					Scene,
-					Views,
-					INDEX_NONE,
-					StartPrimitiveIndex,
-					NumPrimitives,
-					PreShadows,
-					ViewDependentWholeSceneShadows,
-					FeatureLevel,
-					bStaticSceneOnly,
-					GatherStats);
-				Packets.Add(Packet);
+				const int32 NumPrimitives = FMath::Min(PacketSize, TaskData.Scene->Primitives.Num() - StartPrimitiveIndex);
+
+				AddSubTask(INDEX_NONE, StartPrimitiveIndex, NumPrimitives);
 			}
 		}
-			
-		{
-			QUICK_SCOPE_CYCLE_COUNTER(STAT_FilterPrimitivesForShadows);
+	}
+};
+
+void FSceneRenderer::BeginGatherShadowPrimitives(FDynamicShadowsTaskData* TaskData)
+{
+	SCOPE_CYCLE_COUNTER(STAT_GatherShadowPrimitivesTime);
 
-			ParallelFor(Packets.Num(), 
-				[&Packets](int32 Index)
+	if (TaskData && (TaskData->PreShadows.Num() || TaskData->ViewDependentWholeSceneShadows.Num()))
+	{
+		if (!TaskData->bMultithreaded)
+		{
+			FGatherShadowPrimitivesPrepareTask(*TaskData).AnyThreadTask();
+			for (FGatherShadowPrimitivesPacket* Packet : TaskData->Packets)
+			{
+				Packet->AnyThreadTask(*TaskData);
+			}
+		}
+		else if (!TaskData->bRunningEarly)
+		{
+			FGatherShadowPrimitivesPrepareTask(*TaskData).AnyThreadTask();
+			ParallelFor(TaskData->Packets.Num(),
+				[TaskData](int32 Index)
 				{
-					Packets[Index]->AnyThreadTask();
+					TaskData->Packets[Index]->AnyThreadTask(*TaskData);
 				},
-				!(FApp::ShouldUseThreadingForPerformance() && CVarParallelGatherShadowPrimitives.GetValueOnRenderThread() > 0)
-			);
+				!(FApp::ShouldUseThreadingForPerformance() && CVarParallelGatherShadowPrimitives.GetValueOnRenderThread() > 0));
+		}
+		else
+		{
+			TaskData->TaskEvent = TGraphTask<FGatherShadowPrimitivesPrepareTask>::CreateTask().ConstructAndDispatchWhenReady(*TaskData);
 		}
+	}
+}
 
+void FSceneRenderer::FinishGatherShadowPrimitives(FDynamicShadowsTaskData* TaskData)
+{
+	if (TaskData)
+	{
+		SCOPED_NAMED_EVENT_TEXT("FSceneRenderer::FinishGatherShadowPrimitives", FColor::Green);
+		QUICK_SCOPE_CYCLE_COUNTER(STAT_RenderThreadFinalize);
+
+		if (TaskData->TaskEvent)
 		{
-			QUICK_SCOPE_CYCLE_COUNTER(STAT_RenderThreadFinalize);
+			FTaskGraphInterface::Get().WaitUntilTaskCompletes(TaskData->TaskEvent, ENamedThreads::GetRenderThread_Local());
+		}
+
+		for (int32 ShadowIdx = 0, Num = TaskData->ViewDependentWholeSceneShadows.Num(); ShadowIdx < Num; ++ShadowIdx)
+		{
+			TaskData->ViewDependentWholeSceneShadows[ShadowIdx]->PresizeSubjectPrimitiveArrays(TaskData->GatherStats[ShadowIdx]);
+		}
 
-			for (int32 ShadowIdx = 0, Num = ViewDependentWholeSceneShadows.Num(); ShadowIdx < Num; ++ShadowIdx)
+		for (FGatherShadowPrimitivesPacket* Packet : TaskData->Packets)
+		{
+			Packet->RenderThreadFinalize(*TaskData);
+			if (TaskData->bMultithreaded)
 			{
-				ViewDependentWholeSceneShadows[ShadowIdx]->PresizeSubjectPrimitiveArrays(GatherStats[ShadowIdx]);
+				delete Packet;
 			}
-
-			for (int32 PacketIndex = 0; PacketIndex < Packets.Num(); PacketIndex++)
+			else
 			{
-				FGatherShadowPrimitivesPacket* Packet = Packets[PacketIndex];
-				Packet->RenderThreadFinalize();
 				// Class was allocated on the memstack which does not call destructors
 				Packet->~FGatherShadowPrimitivesPacket();
 			}
 		}
+
+		delete TaskData;
 	}
 }
 
@@ -5358,22 +5467,14 @@ void FSceneRenderer::AllocateTranslucentShadowDepthTargets(FRHICommandListImmedi
 	}
 }
 
-void FSceneRenderer::InitDynamicShadows(FRHICommandListImmediate& RHICmdList, FGlobalDynamicIndexBuffer& DynamicIndexBuffer, FGlobalDynamicVertexBuffer& DynamicVertexBuffer, FGlobalDynamicReadBuffer& DynamicReadBuffer, FInstanceCullingManager& InstanceCullingManager)
+FDynamicShadowsTaskData* FSceneRenderer::BeginInitDynamicShadows(bool bRunningEarly)
 {
 	SCOPE_CYCLE_COUNTER(STAT_DynamicShadowSetupTime);
 	CSV_SCOPED_TIMING_STAT_EXCLUSIVE(InitViews_Shadows);
-	SCOPED_NAMED_EVENT(FSceneRenderer_InitDynamicShadows, FColor::Magenta);
+	SCOPED_NAMED_EVENT_TEXT("FSceneRenderer::BeginInitDynamicShadows", FColor::Magenta);
 
 	const bool bMobile = FeatureLevel < ERHIFeatureLevel::SM5;
 
-	bool bStaticSceneOnly = false;
-
-	for (int32 ViewIndex = 0; ViewIndex < Views.Num(); ViewIndex++)
-	{
-		FViewInfo& View = Views[ViewIndex];
-		bStaticSceneOnly = bStaticSceneOnly || View.bStaticSceneOnly;
-	}
-
 	const bool bProjectEnablePointLightShadows = Scene->ReadOnlyCVARCache.bEnablePointLightShadows && !bMobile; // Point light shadow is unsupported on mobile for now.
 	const bool bProjectEnableMovableDirectionLightShadows = !bMobile || Scene->ReadOnlyCVARCache.bMobileAllowMovableDirectionalLights;
 	const bool bProjectEnableMovableSpotLightShadows = !bMobile || Scene->ReadOnlyCVARCache.bMobileEnableMovableSpotlightsShadow;
@@ -5383,9 +5484,12 @@ void FSceneRenderer::InitDynamicShadows(FRHICommandListImmediate& RHICmdList, FG
 	uint32 NumSpotShadowCachesUpdatedThisFrame = 0;
 	uint32 NumCSMCachesUpdatedThisFrame = 0;
 
-	TArray<FProjectedShadowInfo*,SceneRenderingAllocator> PreShadows;
-	TArray<FProjectedShadowInfo*,SceneRenderingAllocator> ViewDependentWholeSceneShadows;
-	TArray<FProjectedShadowInfo*,SceneRenderingAllocator> ViewDependentWholeSceneShadowsThatNeedCulling;
+	FDynamicShadowsTaskData* DynamicShadowTaskData = new FDynamicShadowsTaskData(Scene, Views, bRunningEarly);
+
+	const bool bStaticSceneOnly = DynamicShadowTaskData->bStaticSceneOnly;
+	TArray<FProjectedShadowInfo*, SceneRenderingAllocator>& PreShadows = DynamicShadowTaskData->PreShadows;
+	TArray<FProjectedShadowInfo*, SceneRenderingAllocator> ViewDependentWholeSceneShadows;
+	TArray<FProjectedShadowInfo*, SceneRenderingAllocator>& ViewDependentWholeSceneShadowsThatNeedCulling = DynamicShadowTaskData->ViewDependentWholeSceneShadows;
 	{
 		SCOPE_CYCLE_COUNTER(STAT_InitDynamicShadowsTime);
 		CSV_SCOPED_TIMING_STAT_EXCLUSIVE(ShadowInitDynamic);
@@ -5544,7 +5648,20 @@ void FSceneRenderer::InitDynamicShadows(FRHICommandListImmediate& RHICmdList, FG
 	UpdatePreshadowCache();
 
 	// Gathers the list of primitives used to draw various shadow types
-	GatherShadowPrimitives(PreShadows, ViewDependentWholeSceneShadowsThatNeedCulling, bStaticSceneOnly);
+	BeginGatherShadowPrimitives(DynamicShadowTaskData);
+
+	return DynamicShadowTaskData;
+}
+
+
+void FSceneRenderer::FinishInitDynamicShadows(FRHICommandListImmediate& RHICmdList, FDynamicShadowsTaskData* TaskData, FGlobalDynamicIndexBuffer& DynamicIndexBuffer, FGlobalDynamicVertexBuffer& DynamicVertexBuffer, FGlobalDynamicReadBuffer& DynamicReadBuffer, FInstanceCullingManager& InstanceCullingManager)
+{
+	SCOPED_NAMED_EVENT_TEXT("FSceneRenderer::FinishInitDynamicShadows", FColor::Magenta);
+	SCOPE_CYCLE_COUNTER(STAT_InitDynamicShadowsTime);
+	SCOPE_CYCLE_COUNTER(STAT_DynamicShadowSetupTime);
+	CSV_SCOPED_TIMING_STAT_EXCLUSIVE(ShadowInitDynamic);
+
+	FinishGatherShadowPrimitives(TaskData);
 
 	AllocateShadowDepthTargets(RHICmdList);
 
@@ -5552,6 +5669,12 @@ void FSceneRenderer::InitDynamicShadows(FRHICommandListImmediate& RHICmdList, FG
 	GatherShadowDynamicMeshElements(DynamicIndexBuffer, DynamicVertexBuffer, DynamicReadBuffer, InstanceCullingManager);
 }
 
+void FSceneRenderer::InitDynamicShadows(FRHICommandListImmediate& RHICmdList, FGlobalDynamicIndexBuffer& DynamicIndexBuffer, FGlobalDynamicVertexBuffer& DynamicVertexBuffer, FGlobalDynamicReadBuffer& DynamicReadBuffer, FInstanceCullingManager& InstanceCullingManager)
+{
+	FDynamicShadowsTaskData* TaskData = BeginInitDynamicShadows(false);
+	FinishInitDynamicShadows(RHICmdList, TaskData, DynamicIndexBuffer, DynamicVertexBuffer, DynamicReadBuffer, InstanceCullingManager);
+}
+
 void FSceneRenderer::AllocateMobileCSMAndSpotLightShadowDepthTargets(FRHICommandListImmediate& RHICmdList, const TArray<FProjectedShadowInfo*, SceneRenderingAllocator>& MobileCSMAndSpotLightShadows)
 {
 	if (MobileCSMAndSpotLightShadows.Num() > 0)
