Changing WaitForTasksAndClearSnapshots Task to add the required tasks as prerequisites instead of actually waiting on them.

diff --git a/Engine/Source/Runtime/Renderer/Private/MeshDrawCommands.cpp b/Engine/Source/Runtime/Renderer/Private/MeshDrawCommands.cpp
--- a/Engine/Source/Runtime/Renderer/Private/MeshDrawCommands.cpp
+++ b/Engine/Source/Runtime/Renderer/Private/MeshDrawCommands.cpp
@@ -782,7 +782,7 @@ public:
 	void AnyThreadTask()
 	{
 		FOptionalTaskTagScope Scope(ETaskTag::EParallelRenderingThread);
-		TRACE_CPUPROFILER_EVENT_SCOPE(MeshDrawCommandPassSetupTask);
+		SCOPED_NAMED_EVENT(MeshDrawCommandPassSetupTask, FColor::Magenta);
 		// Mobile base pass is a special case, as final lists is created from two mesh passes based on CSM visibility.
 		const bool bMobileShadingBasePass = Context.ShadingPath == EShadingPath::Mobile && Context.PassType == EMeshPass::BasePass;
 		// On SM5 Mobile platform, still want the same sorting
@@ -1143,7 +1143,7 @@ bool FParallelMeshDrawCommandPass::IsOnDemandShaderCreationEnabled()
 
 void FParallelMeshDrawCommandPass::WaitForMeshPassSetupTask(EWaitThread WaitThread) const
 {
-	if (TaskEventRef.IsValid())
+	if (TaskEventRef.IsValid() && WaitThread != EWaitThread::TaskAlreadyWaited)
 	{
 		// Need to wait on GetRenderThread_Local, as mesh pass setup task can wait on rendering thread inside InitResourceFromPossiblyParallelRendering().
 		QUICK_SCOPE_CYCLE_COUNTER(STAT_WaitForMeshPassSetupTask);
@@ -1241,7 +1241,7 @@ public:
 	void DoTask(ENamedThreads::Type CurrentThread, const FGraphEventRef& MyCompletionGraphEvent)
 	{
 		FOptionalTaskTagScope Scope(ETaskTag::EParallelRenderingThread);
-		TRACE_CPUPROFILER_EVENT_SCOPE(DrawVisibleMeshCommandsAnyThreadTask);
+		SCOPED_NAMED_EVENT_TEXT("DrawVisibleMeshCommandsAnyThreadTask", FColor::Magenta);
 		checkSlow(RHICmdList.IsInsideRenderPass());
 
 		// FDrawVisibleMeshCommandsAnyThreadTasks must only run on RT if RHISupportsMultithreadedShaderCreation is not supported!
diff --git a/Engine/Source/Runtime/Renderer/Private/MeshDrawCommands.h b/Engine/Source/Runtime/Renderer/Private/MeshDrawCommands.h
--- a/Engine/Source/Runtime/Renderer/Private/MeshDrawCommands.h
+++ b/Engine/Source/Runtime/Renderer/Private/MeshDrawCommands.h
@@ -141,7 +141,8 @@ public:
 	enum class EWaitThread
 	{
 		Render,
-		Task
+		Task,
+		TaskAlreadyWaited,
 	};
 
 	FParallelMeshDrawCommandPass()
@@ -209,6 +210,7 @@ public:
 	static bool IsOnDemandShaderCreationEnabled();
 
 	FInstanceCullingContext* GetInstanceCullingContext() { return &TaskContext.InstanceCullingContext; }
+	const FGraphEventRef& GetTaskEvent() const { return TaskEventRef; }
 
 private:
 	FPrimitiveIdVertexBufferPoolEntry PrimitiveIdVertexBufferPoolEntry;
diff --git a/Engine/Source/Runtime/Renderer/Private/SceneRendering.cpp b/Engine/Source/Runtime/Renderer/Private/SceneRendering.cpp
--- a/Engine/Source/Runtime/Renderer/Private/SceneRendering.cpp
+++ b/Engine/Source/Runtime/Renderer/Private/SceneRendering.cpp
@@ -3760,15 +3760,29 @@ void FSceneRenderer::RenderThreadEnd(FRHICommandListImmediate& RHICmdList)
 
 		if (SceneRenderCleanUpMode == ESceneRenderCleanUpMode::DeferredAndAsync)
 		{
-			FGraphEventArray& WaitOutstandingTasks = RHICmdList.GetRenderThreadTaskArray();
+			FGraphEventArray& CommandListTasks = RHICmdList.GetRenderThreadTaskArray();
 
-			GSceneRenderCleanUpState.Task = FFunctionGraphTask::CreateAndDispatchWhenReady([this]
+			FGraphEventArray Prerequisites;
+			Prerequisites.Append(CommandListTasks);
+
+			for (FParallelMeshDrawCommandPass* DispatchedShadowDepthPass : DispatchedShadowDepthPasses)
 			{
-				WaitForTasksAndClearSnapshots(FParallelMeshDrawCommandPass::EWaitThread::Task);
+				Prerequisites.Add(DispatchedShadowDepthPass->GetTaskEvent());
+			}
+			for (const FViewInfo& View : Views)
+			{
+				for (const FParallelMeshDrawCommandPass& Pass : View.ParallelMeshDrawCommandPasses)
+				{
+					Prerequisites.Add(Pass.GetTaskEvent());
+				}
+			}
 
-			}, TStatId(), &WaitOutstandingTasks);
+			GSceneRenderCleanUpState.Task = FFunctionGraphTask::CreateAndDispatchWhenReady([this]
+			{
+				WaitForTasksAndClearSnapshots(FParallelMeshDrawCommandPass::EWaitThread::TaskAlreadyWaited);
+			}, TStatId(), & Prerequisites);
 
-			WaitOutstandingTasks.Empty();
+			CommandListTasks.Empty();
 		}
 	}
 }
@@ -3796,7 +3810,7 @@ void FSceneRenderer::CleanUp(FRHICommandListImmediate& RHICmdList)
 
 void FSceneRenderer::WaitForTasksAndClearSnapshots(FParallelMeshDrawCommandPass::EWaitThread WaitThread)
 {
-	TRACE_CPUPROFILER_EVENT_SCOPE(FSceneRenderer::WaitForTasksAndClearSnapshots);
+	SCOPED_NAMED_EVENT_TEXT("FSceneRenderer::WaitForTasksAndClearSnapshots", FColor::Red);
 
 	// Wait for all dispatched shadow mesh draw tasks.
 	for (int32 PassIndex = 0; PassIndex < DispatchedShadowDepthPasses.Num(); ++PassIndex)
