Fix CIS compilation error.

diff --git a/Engine/Source/Runtime/Renderer/Private/HairStrands/HairStrandsComposition.cpp b/Engine/Source/Runtime/Renderer/Private/HairStrands/HairStrandsComposition.cpp
--- a/Engine/Source/Runtime/Renderer/Private/HairStrands/HairStrandsComposition.cpp
+++ b/Engine/Source/Runtime/Renderer/Private/HairStrands/HairStrandsComposition.cpp
@@ -216,7 +216,7 @@ static void AddHairVisibilityComposeSamplePass(
 	Parameters->RenderTargets[0] = FRenderTargetBinding(OutColorTexture, ERenderTargetLoadAction::ELoad);
 	Parameters->RenderTargets.DepthStencil = FDepthStencilBinding(OutDepthTexture, ERenderTargetLoadAction::ELoad, ERenderTargetLoadAction::ENoAction, FExclusiveDepthStencil::DepthWrite_StencilNop);
 
-	const bool bDebugComposition = View.Family->EngineShowFlags.LODColoration > 0;
+	const bool bDebugComposition = View.Family->EngineShowFlags.LODColoration;
 	FHairVisibilityComposeSamplePS::FPermutationDomain PermutationVector;
 	PermutationVector.Set<FHairVisibilityComposeSamplePS::FDebug>(bDebugComposition);
 	TShaderMapRef<FHairVisibilityComposeSamplePS> PixelShader(View.ShaderMap, PermutationVector);
