Turns on r.TSR.Translucency.SeparateTemporalAccumulation by default

diff --git a/Engine/Source/Runtime/Renderer/Private/PostProcess/TemporalSuperResolution.cpp b/Engine/Source/Runtime/Renderer/Private/PostProcess/TemporalSuperResolution.cpp
--- a/Engine/Source/Runtime/Renderer/Private/PostProcess/TemporalSuperResolution.cpp
+++ b/Engine/Source/Runtime/Renderer/Private/PostProcess/TemporalSuperResolution.cpp
@@ -61,8 +61,8 @@ TAutoConsoleVariable<int32> CVarTSRTranslucencyPreviousFrameRejection(
 	ECVF_RenderThreadSafe);
 
 TAutoConsoleVariable<int32> CVarTSRTranslucencySeparateTemporalAccumulation(
-	TEXT("r.TSR.Translucency.SeparateTemporalAccumulation"), 0,
-	TEXT("Accumulates separate translucency separatly."),
+	TEXT("r.TSR.Translucency.SeparateTemporalAccumulation"), 1,
+	TEXT("Accumulates separate translucency separatly (enabled by default)."),
 	ECVF_RenderThreadSafe);
 
 TAutoConsoleVariable<int32> CVarTSREnableResponiveAA(
