Film grain change for pete

diff --git a/Engine/Source/Runtime/Renderer/Private/PostProcess/PostProcessTonemap.cpp b/Engine/Source/Runtime/Renderer/Private/PostProcess/PostProcessTonemap.cpp
--- a/Engine/Source/Runtime/Renderer/Private/PostProcess/PostProcessTonemap.cpp
+++ b/Engine/Source/Runtime/Renderer/Private/PostProcess/PostProcessTonemap.cpp
@@ -721,7 +721,7 @@ FScreenPassTexture AddTonemapPass(FRDGBuilder& GraphBuilder, const FViewInfo& Vi
 		RandomGrainTextureUVOffset.Y = Halton(RandomSequenceIndex + 1, 3);
 
 		FVector2D TextureSize(View.FilmGrainTexture->GetSizeX(), View.FilmGrainTexture->GetSizeY());
-		FVector2D OutputSizeF(OutputViewport.Rect.Width(), OutputViewport.Rect.Height());
+		FVector2D OutputSizeF(1920.0f, 1920.0f * float(OutputViewport.Rect.Height()) / float(OutputViewport.Rect.Width()));
 
 		CommonParameters.FilmGrain.ScreenPosToFilmGrainTextureUV = 
 			FScreenTransform::ScreenPosToViewportUV *
