Fix stat name for PrepareDistanceFieldScene

#pf 610c535b26ea630001a87682

diff --git a/Engine/Source/Runtime/Renderer/Private/DistanceFieldObjectManagement.cpp b/Engine/Source/Runtime/Renderer/Private/DistanceFieldObjectManagement.cpp
--- a/Engine/Source/Runtime/Renderer/Private/DistanceFieldObjectManagement.cpp
+++ b/Engine/Source/Runtime/Renderer/Private/DistanceFieldObjectManagement.cpp
@@ -978,7 +978,7 @@ void FSceneRenderer::UpdateGlobalHeightFieldObjectBuffers(FRDGBuilder& GraphBuil
 
 void FSceneRenderer::PrepareDistanceFieldScene(FRDGBuilder& GraphBuilder, bool bSplitDispatch)
 {
-	CSV_SCOPED_TIMING_STAT_EXCLUSIVE(RenderDFAO);
+	CSV_SCOPED_TIMING_STAT_EXCLUSIVE(PrepareDistanceFieldScene);
 	TRACE_CPUPROFILER_EVENT_SCOPE(FSceneRenderer::PrepareDistanceFieldScene);
 	QUICK_SCOPE_CYCLE_COUNTER(STAT_PrepareDistanceFieldScene);
 	LLM_SCOPE_BYTAG(DistanceFields);
