Fixes responsive AA stencil bit through distortion rendering

diff --git a/Engine/Source/Runtime/Renderer/Private/DistortionRendering.cpp b/Engine/Source/Runtime/Renderer/Private/DistortionRendering.cpp
--- a/Engine/Source/Runtime/Renderer/Private/DistortionRendering.cpp
+++ b/Engine/Source/Runtime/Renderer/Private/DistortionRendering.cpp
@@ -478,7 +478,7 @@ void FDeferredShadingSceneRenderer::RenderDistortion(FRDGBuilder& GraphBuilder,
 	RDG_GPU_STAT_SCOPE(GraphBuilder, Distortion);
 
 	const FDepthStencilBinding StencilReadBinding(SceneDepthTexture, ERenderTargetLoadAction::ELoad, ERenderTargetLoadAction::ELoad, FExclusiveDepthStencil::DepthRead_StencilRead);
-	FDepthStencilBinding StencilWriteBinding(SceneDepthTexture, ERenderTargetLoadAction::ELoad, ERenderTargetLoadAction::EClear, FExclusiveDepthStencil::DepthRead_StencilWrite);
+	FDepthStencilBinding StencilWriteBinding(SceneDepthTexture, ERenderTargetLoadAction::ELoad, ERenderTargetLoadAction::ELoad, FExclusiveDepthStencil::DepthRead_StencilWrite);
 
 	FRDGTextureRef DistortionTexture = nullptr;
 	FRDGTextureRef RoughnessScatterTexture = nullptr;
@@ -700,9 +700,7 @@ void FDeferredShadingSceneRenderer::RenderDistortion(FRDGBuilder& GraphBuilder,
 			});
 
 			LoadAction = ERenderTargetLoadAction::ELoad;
-			StencilWriteBinding.SetStencilLoadAction(ERenderTargetLoadAction::ELoad);
 		}
-		StencilWriteBinding.SetStencilLoadAction(ERenderTargetLoadAction::ELoad);
 	}
 
 	FRDGTextureDesc DistortedSceneColorDesc = SceneColorTexture->Desc;
