Integrate SDF to mobile deferred renderer

diff --git a/Engine/Shaders/Private/MobileDeferredShading.usf b/Engine/Shaders/Private/MobileDeferredShading.usf
--- a/Engine/Shaders/Private/MobileDeferredShading.usf
+++ b/Engine/Shaders/Private/MobileDeferredShading.usf
@@ -271,8 +271,8 @@ void MobileDirectLightPS(
 	FMobileLightAccumulator MobileLightAccumulator = (FMobileLightAccumulator)0;
 	// Check movable light param to determine if we should be using precomputed shadows
 	half Shadow = LightFunctionParameters2.z > 0.0f ? 1.0f : GBuffer.PrecomputedShadowFactors.r;
-
-#if APPLY_CSM  
+	
+#if APPLY_CSM && !ENABLE_DISTANCE_FIELD
 	// TODO: we don't need to fully compute ScreenPosition here
 	float  ShadowPositionZ = 0;
 	float4 ScreenPosition = SvPositionToScreenPosition(float4(SvPosition.xyz,SceneDepth));
@@ -280,6 +280,10 @@ void MobileDirectLightPS(
 	Shadow = min(ShadowMap, Shadow);
 #endif
 
+#if ENABLE_DISTANCE_FIELD
+	Shadow = Texture2DSample(MobileBasePass.ScreenSpaceShadowMaskTexture, MobileBasePass.ScreenSpaceShadowMaskSampler, SvPositionToBufferUV(SvPosition)).x;
+#endif
+
 	FMobileShadingModelContext ShadingModelContext = (FMobileShadingModelContext)0;
 	InitShadingModelContext(ShadingModelContext, GBuffer, CameraVector);
 	
diff --git a/Engine/Source/Runtime/Renderer/Private/MobileBasePass.cpp b/Engine/Source/Runtime/Renderer/Private/MobileBasePass.cpp
--- a/Engine/Source/Runtime/Renderer/Private/MobileBasePass.cpp
+++ b/Engine/Source/Runtime/Renderer/Private/MobileBasePass.cpp
@@ -353,7 +353,7 @@ ELightMapPolicyType MobileBasePass::SelectMeshLightmapPolicy(
 	return SelectedLightmapPolicy;
 }
 
-void MobileBasePass::SetOpaqueRenderState(FMeshPassProcessorRenderState& DrawRenderState, const FPrimitiveSceneProxy* PrimitiveSceneProxy, const FMaterial& Material, bool bEnableReceiveDecalOutput, bool bUsesDeferredShading)
+void MobileBasePass::SetOpaqueRenderState(FMeshPassProcessorRenderState& DrawRenderState, const FPrimitiveSceneProxy* PrimitiveSceneProxy, const FMaterial& Material, bool bEnableReceiveDecalOutput, bool bUsesDeferredShading, bool bDisableDepthWrite)
 {
 	uint8 StencilValue = 0;
 	if (bEnableReceiveDecalOutput)
@@ -371,18 +371,36 @@ void MobileBasePass::SetOpaqueRenderState(FMeshPassProcessorRenderState& DrawRen
 		
 	if (bEnableReceiveDecalOutput || bUsesDeferredShading)
 	{
-		DrawRenderState.SetDepthStencilState(TStaticDepthStencilState<
+		if (bDisableDepthWrite)
+		{
+			DrawRenderState.SetDepthStencilState(TStaticDepthStencilState<
+				false, CF_Equal,
+				true, CF_Always, SO_Keep, SO_Keep, SO_Replace,
+				false, CF_Always, SO_Keep, SO_Keep, SO_Keep,
+				// don't use masking as it has significant performance hit on Mali GPUs (T860MP2)
+				0x00, 0xff >::GetRHI());
+		}
+		else
+		{
+			DrawRenderState.SetDepthStencilState(TStaticDepthStencilState<
 				true, CF_DepthNearOrEqual,
 				true, CF_Always, SO_Keep, SO_Keep, SO_Replace,
 				false, CF_Always, SO_Keep, SO_Keep, SO_Keep,
 				// don't use masking as it has significant performance hit on Mali GPUs (T860MP2)
 				0x00, 0xff >::GetRHI());
+		}
+		
 
 		DrawRenderState.SetStencilRef(StencilValue); 
 	}
 	else
 	{
 		// default depth state should be already set
+		// disable DepthWrite if Prepass enabled
+		if (bDisableDepthWrite)
+		{
+			DrawRenderState.SetDepthStencilState(TStaticDepthStencilState<false, CF_Equal>::GetRHI());
+		}
 	}
 
 	if (Material.GetBlendMode() == BLEND_Masked && Material.IsUsingAlphaToCoverage())
@@ -760,14 +778,11 @@ bool FMobileBasePassMeshProcessor::Process(
 		{
 			MobileBasePass::SetTranslucentRenderState(DrawRenderState, MaterialResource);
 		}
-		else if((MeshBatch.bUseForDepthPass && Scene->EarlyZPassMode == DDM_AllOpaque) || bMaskedInEarlyPass)
-		{
-			DrawRenderState.SetDepthStencilState(TStaticDepthStencilState<false, CF_Equal>::GetRHI());
-		}
 		else
 		{
 			const bool bEnableReceiveDecalOutput = ((Flags & EFlags::CanUseDepthStencil) == EFlags::CanUseDepthStencil);
-			MobileBasePass::SetOpaqueRenderState(DrawRenderState, PrimitiveSceneProxy, MaterialResource, bEnableReceiveDecalOutput && IsMobileHDR(), bUsesDeferredShading);
+			const bool bDisableDepthWrite = (MeshBatch.bUseForDepthPass && Scene->EarlyZPassMode == DDM_AllOpaque) || bMaskedInEarlyPass;
+			MobileBasePass::SetOpaqueRenderState(DrawRenderState, PrimitiveSceneProxy, MaterialResource, bEnableReceiveDecalOutput && IsMobileHDR(), bUsesDeferredShading, bDisableDepthWrite);
 		}
 	}
 
diff --git a/Engine/Source/Runtime/Renderer/Private/MobileBasePassRendering.h b/Engine/Source/Runtime/Renderer/Private/MobileBasePassRendering.h
--- a/Engine/Source/Runtime/Renderer/Private/MobileBasePassRendering.h
+++ b/Engine/Source/Runtime/Renderer/Private/MobileBasePassRendering.h
@@ -353,7 +353,7 @@ namespace MobileBasePass
 
 	bool StaticCanReceiveCSM(const FLightSceneInfo* LightSceneInfo, const FPrimitiveSceneProxy* PrimitiveSceneProxy);
 
-	void SetOpaqueRenderState(FMeshPassProcessorRenderState& DrawRenderState, const FPrimitiveSceneProxy* PrimitiveSceneProxy, const FMaterial& Material, bool bEnableReceiveDecalOutput, bool bUsesDeferredShading);
+	void SetOpaqueRenderState(FMeshPassProcessorRenderState& DrawRenderState, const FPrimitiveSceneProxy* PrimitiveSceneProxy, const FMaterial& Material, bool bEnableReceiveDecalOutput, bool bUsesDeferredShading, bool bDisableDepthWrite);
 	void SetTranslucentRenderState(FMeshPassProcessorRenderState& DrawRenderState, const FMaterial& Material);
 
 	bool StationarySkyLightHasBeenApplied(const FScene* Scene, ELightMapPolicyType LightMapPolicyType);
diff --git a/Engine/Source/Runtime/Renderer/Private/MobileDeferredShadingPass.cpp b/Engine/Source/Runtime/Renderer/Private/MobileDeferredShadingPass.cpp
--- a/Engine/Source/Runtime/Renderer/Private/MobileDeferredShadingPass.cpp
+++ b/Engine/Source/Runtime/Renderer/Private/MobileDeferredShadingPass.cpp
@@ -52,6 +52,8 @@ class FMobileDirectLightFunctionPS : public FMaterialShader
 		SHADER_PARAMETER(FVector3f, LightFunctionParameters2)
 		SHADER_PARAMETER_TEXTURE(Texture2D, PreIntegratedGF)
 		SHADER_PARAMETER_SAMPLER(SamplerState, PreIntegratedGFSampler)
+		SHADER_PARAMETER_TEXTURE(Texture2D, ScreenSpaceShadowMaskTexture)
+		SHADER_PARAMETER_SAMPLER(SamplerState, ScreenSpaceShadowMaskSampler)
 	END_SHADER_PARAMETER_STRUCT()
 
 	static void ModifyCompilationEnvironment(const FMaterialShaderPermutationParameters& Parameters, FShaderCompilerEnvironment& OutEnvironment)
@@ -61,6 +63,7 @@ class FMobileDirectLightFunctionPS : public FMaterialShader
 		OutEnvironment.SetDefine(TEXT(PREPROCESSOR_TO_STRING(MAX_MOBILE_SHADOWCASCADES)), GetMobileMaxShadowCascades());
 		OutEnvironment.SetDefine(TEXT("SUPPORTS_TEXTURECUBE_ARRAY"), 1);
 		OutEnvironment.SetDefine(TEXT("USE_LIGHT_FUNCTION"), Parameters.MaterialParameters.bIsDefaultMaterial ? 0 : 1);
+		OutEnvironment.SetDefine(TEXT("ENABLE_DISTANCE_FIELD"), IsMobileDistanceFieldEnabled(Parameters.Platform));
 	}
 
 	static FPermutationDomain RemapPermutationVector(FPermutationDomain PermutationVector)
@@ -267,6 +270,19 @@ static void RenderDirectLight(FRHICommandListImmediate& RHICmdList, const FScene
 	PassParameters.PreIntegratedGF = GSystemTextures.PreintegratedGF->GetRenderTargetItem().ShaderResourceTexture;
 	PassParameters.PreIntegratedGFSampler = TStaticSamplerState<SF_Bilinear, AM_Clamp, AM_Clamp, AM_Clamp>::GetRHI();
 
+	bool bRequiresDistanceFieldShadowingPass = IsMobileDistanceFieldShadowingEnabled(View.GetShaderPlatform());
+
+	if (bRequiresDistanceFieldShadowingPass && GScreenSpaceShadowMaskTextureMobileOutputs.ScreenSpaceShadowMaskTextureMobile.IsValid())
+	{
+		PassParameters.ScreenSpaceShadowMaskTexture = GScreenSpaceShadowMaskTextureMobileOutputs.ScreenSpaceShadowMaskTextureMobile->GetRenderTargetItem().ShaderResourceTexture;
+		PassParameters.ScreenSpaceShadowMaskSampler = TStaticSamplerState<SF_Point, AM_Clamp, AM_Clamp, AM_Clamp>::GetRHI();
+	}
+	else
+	{
+		PassParameters.ScreenSpaceShadowMaskTexture = GSystemTextures.WhiteDummy->GetRenderTargetItem().ShaderResourceTexture;
+		PassParameters.ScreenSpaceShadowMaskSampler = TStaticSamplerState<SF_Point, AM_Clamp, AM_Clamp, AM_Clamp>::GetRHI();
+	}
+	
 	if (DirectionalLight)
 	{
 		const bool bUseMovableLight = DirectionalLight && !DirectionalLight->Proxy->HasStaticShadowing();
diff --git a/Engine/Source/Runtime/Renderer/Private/MobileShadingRenderer.cpp b/Engine/Source/Runtime/Renderer/Private/MobileShadingRenderer.cpp
--- a/Engine/Source/Runtime/Renderer/Private/MobileShadingRenderer.cpp
+++ b/Engine/Source/Runtime/Renderer/Private/MobileShadingRenderer.cpp
@@ -425,8 +425,7 @@ void FMobileSceneRenderer::InitViews(FRDGBuilder& GraphBuilder, FSceneTexturesCo
 		&& !Views[0].bIsPlanarReflection
 		&& !ViewFamily.EngineShowFlags.HitProxies
 		&& !ViewFamily.EngineShowFlags.VisualizeLightCulling
-		&& !ViewFamily.UseDebugViewPS()
-		&& !bDeferredShading;
+		&& !ViewFamily.UseDebugViewPS();
 
 	bRequiresDistanceFieldShadowingPass = bRequiresDistanceField && IsMobileDistanceFieldShadowingEnabled(ShaderPlatform);
 
