Fix TSR separated translucency output permutation when seperate translucency inputs is invalid.

diff --git a/Engine/Source/Runtime/Renderer/Private/PostProcess/TemporalSuperResolution.cpp b/Engine/Source/Runtime/Renderer/Private/PostProcess/TemporalSuperResolution.cpp
--- a/Engine/Source/Runtime/Renderer/Private/PostProcess/TemporalSuperResolution.cpp
+++ b/Engine/Source/Runtime/Renderer/Private/PostProcess/TemporalSuperResolution.cpp
@@ -545,9 +545,11 @@ void AddTemporalSuperResolutionPasses(
 
 	bool bEnableInterferenceHeuristic = CVarTSREnableAntiInterference.GetValueOnRenderThread() != 0;
 
-	bool bRejectSeparateTranslucency = PassInputs.SeparateTranslucencyTextures != nullptr && CVarTSRTranslucencyPreviousFrameRejection.GetValueOnRenderThread() != 0;
+	bool bIsSeperateTranslucyTexturesValid = PassInputs.SeparateTranslucencyTextures != nullptr && PassInputs.SeparateTranslucencyTextures->IsColorValid();
 
-	bool bAccumulateSeparateTranslucency = PassInputs.SeparateTranslucencyTextures != nullptr && CVarTSRTranslucencySeparateTemporalAccumulation.GetValueOnRenderThread() != 0;
+	bool bRejectSeparateTranslucency = bIsSeperateTranslucyTexturesValid && CVarTSRTranslucencyPreviousFrameRejection.GetValueOnRenderThread() != 0;
+
+	bool bAccumulateSeparateTranslucency = bIsSeperateTranslucyTexturesValid && CVarTSRTranslucencySeparateTemporalAccumulation.GetValueOnRenderThread() != 0;
 
 	int32 RejectionAntiAliasingQuality = FMath::Clamp(CVarTSRRejectionAntiAliasingQuality.GetValueOnRenderThread(), 0, 2);
 
