Disabled VT pool notifications by default.
Enable with r.VT.Residency.Notify
Default VT residency mip map on for default pools.
This is overridable per project and per pool.

diff --git a/Engine/Config/BaseEngine.ini b/Engine/Config/BaseEngine.ini
--- a/Engine/Config/BaseEngine.ini
+++ b/Engine/Config/BaseEngine.ini
@@ -2804,11 +2804,15 @@ bAllowDelayedMessaging=false
 DefaultChaosSolverActorClass=/Script/ChaosSolverEngine.ChaosSolverActor
 
 [/Script/Engine.VirtualTexturePoolConfig]
-; Configure VT physical memory pools. Usefull commands to set these
+; Configure VT physical memory pools. Useful commands to set these:
 ; "stat virtualtexturing" show dynamic use of the VT system including the cache loads
 ; "r.VT.ListPhysicalPool" list in-depth details on the allocated physical pools
-; The upper allocation limit for a pool is caluclated by matching a size and format from the list below or falling back to the default
-DefaultSizeInMegabyte=64
+; "r.VT.DumpPoolUsage" list in-depth details on asset usage of physical pools
+; The default pool size and behavior is set here. It can be overridden in project config files.
++Pools=(SizeInMegabyte=64, bAllowSizeScale=False, bEnableResidencyMipMapBias=True)
+; Specific pools can be configured per tile size or format in project config files. Examples commented below:
+;+Pools=(Formats=(PF_DXT1), SizeInMegabyte=36, bAllowSizeScale=True, ScalabilityGroup=1, bEnableResidencyMipMapBias=False)
+;+Pools=(Formats=(PF_DXT5, PF_DXT5), SizeInMegabyte=100)
 
 [PlatformPaths]
 UnrealBuildTool=Engine/Binaries/DotNET/UnrealBuildTool/UnrealBuildTool.exe
diff --git a/Engine/Source/Runtime/Renderer/Private/VT/VirtualTextureSystem.cpp b/Engine/Source/Runtime/Renderer/Private/VT/VirtualTextureSystem.cpp
--- a/Engine/Source/Runtime/Renderer/Private/VT/VirtualTextureSystem.cpp
+++ b/Engine/Source/Runtime/Renderer/Private/VT/VirtualTextureSystem.cpp
@@ -122,6 +122,12 @@ static TAutoConsoleVariable<int32> CVarVTResidencyShow(
 	TEXT("Show on screen HUD for virtual texture physical pool residency"),
 	ECVF_Default
 );
+static TAutoConsoleVariable<int32> CVarVTResidencyNotify(
+	TEXT("r.VT.Residency.Notify"),
+	0,
+	TEXT("Show on screen notifications for virtual texture physical pool residency"),
+	ECVF_Default
+);
 
 static FORCEINLINE uint32 EncodePage(uint32 ID, uint32 vLevel, uint32 vTileX, uint32 vTileY)
 {
@@ -441,8 +447,9 @@ void FVirtualTextureSystem::DumpPoolUsageFromConsole()
 			for (TPair<uint32, uint32> ProducerCount : ProducerCountMap)
 			{
 				// Filter out producers that only have locked page mapped.
+				// Keep all producers for non compressed formats (assuming that we want to avoid these as much as possible).
 				// In future we can add other filters here (count, age etc), or we can dump more info and process off line.
-				if (ProducerCount.Value > 1)
+				if (ProducerCount.Value > 1 || GPixelFormats[PhysicalSpace.GetFormat(0)].BlockSizeX == 1)
 				{
 					SortedProducerCounts.Add(ProducerCount);
 				}
@@ -464,14 +471,6 @@ void FVirtualTextureSystem::DumpPoolUsageFromConsole()
 	}
 }
 
-#if !UE_BUILD_SHIPPING
-void FVirtualTextureSystem::GetOnScreenMessages(FCoreDelegates::FSeverityMessageMap& OutMessages)
-{
-	FScopeLock Locker(&OnScreenMessageLock);
-	OutMessages.Append(OnScreenMessages);
-}
-#endif
-
 #if WITH_EDITOR
 void FVirtualTextureSystem::SaveAllocatorImagesFromConsole()
 {
@@ -1344,6 +1343,10 @@ void FVirtualTextureSystem::Update(FRDGBuilder& GraphBuilder, ERHIFeatureLevel::
 
 	UpdateResidencyTracking();
 
+#if !UE_BUILD_SHIPPING
+	UpdateNotifications();
+#endif
+
 	ReleasePendingSpaces();
 }
 
@@ -2215,28 +2218,6 @@ void FVirtualTextureSystem::SubmitRequests(FRDGBuilder& GraphBuilder, ERHIFeatur
 		INC_DWORD_STAT_BY(STAT_NumStacksRequested, RequestList->GetNumLoadRequests());
 		INC_DWORD_STAT_BY(STAT_NumStacksProduced, NumStacksProduced);
 		INC_DWORD_STAT_BY(STAT_NumPageAllocateFails, NumPageAllocateFails);
-
-#if !UE_BUILD_SHIPPING
-		{
-			FScopeLock Locker(&OnScreenMessageLock);
-			OnScreenMessages.Reset();
-
-			for (int32 SpaceIndex = 0u; SpaceIndex < PhysicalSpaces.Num(); ++SpaceIndex)
-			{
-				const FVirtualTexturePhysicalSpace* PhysicalSpace = PhysicalSpaces[SpaceIndex];
-				// Once a pool fails an allocation, keep displaying the message for a minimum number of frames, to avoid the message quickly flickering on/off
-				if (PhysicalSpace != nullptr && Frame <= PhysicalSpace->GetLastFrameOversubscribed() + 60u)
-				{
-					FString const& FormatString = PhysicalSpace->GetFormatString();
-					const float MipBias = PhysicalSpace->GetResidencyMipMapBias();
-
-					OnScreenMessages.Add(
-						FCoreDelegates::EOnScreenMessageSeverity::Warning, 
-						FText::Format(LOCTEXT("VTOversubscribed", "VT Pool [{0}] is oversubscribed. Setting MipBias {1}"), FText::FromString(FormatString), FText::AsNumber(MipBias)));
-				}
-			}
-		}
-#endif // !UE_BUILD_SHIPPING
 	}
 
 	{
@@ -2360,9 +2341,45 @@ float FVirtualTextureSystem::GetGlobalMipBias() const
 	return UTexture2D::GetGlobalMipMapLODBias() + MaxResidencyMipMapBias;
 }
 
+#if !UE_BUILD_SHIPPING
+
+void FVirtualTextureSystem::GetOnScreenMessages(FCoreDelegates::FSeverityMessageMap& OutMessages)
+{
+	FScopeLock Locker(&OnScreenMessageLock);
+	OutMessages.Append(OnScreenMessages);
+}
+
+void FVirtualTextureSystem::UpdateNotifications()
+{
+	FScopeLock Locker(&OnScreenMessageLock);
+	OnScreenMessages.Reset();
+	UpdateResidencyNotifications();
+}
+
+void FVirtualTextureSystem::UpdateResidencyNotifications()
+{
+	if (CVarVTResidencyNotify.GetValueOnRenderThread() == 0)
+	{
+		return;
+	}
+
+	for (int32 SpaceIndex = 0u; SpaceIndex < PhysicalSpaces.Num(); ++SpaceIndex)
+	{
+		const FVirtualTexturePhysicalSpace* PhysicalSpace = PhysicalSpaces[SpaceIndex];
+		if (PhysicalSpace != nullptr && Frame <= PhysicalSpace->GetLastFrameOversubscribed() + 60u)
+		{
+			FString const& FormatString = PhysicalSpace->GetFormatString();
+			const float MipBias = PhysicalSpace->GetResidencyMipMapBias();
+
+			OnScreenMessages.Add(
+				FCoreDelegates::EOnScreenMessageSeverity::Warning,
+				FText::Format(LOCTEXT("VTOversubscribed", "VT Pool [{0}] is oversubscribed. Setting MipBias {1}"), FText::FromString(FormatString), FText::AsNumber(MipBias)));
+		}
+	}
+}
+
 void FVirtualTextureSystem::DrawResidencyHud(UCanvas* InCanvas, APlayerController* InController)
 {
-#if !UE_BUILD_SHIPPING
 	if (CVarVTResidencyShow.GetValueOnGameThread() == 0)
 	{
 		return;
@@ -2412,7 +2429,7 @@ void FVirtualTextureSystem::DrawResidencyHud(UCanvas* InCanvas, APlayerControlle
 			}
 
 			PhysicalSpace->DrawResidencyGraph(InCanvas->Canvas, CanvasPosition);
-		
+
 			GraphIndex++;
 		}
 	}
@@ -2421,8 +2438,8 @@ void FVirtualTextureSystem::DrawResidencyHud(UCanvas* InCanvas, APlayerControlle
 	InCanvas->Canvas->DrawShadowedString(25, 320, TEXT("MipMap Bias"), GEngine->GetSmallFont(), FLinearColor(0.1f, 0.8f, 0.1f));
 	InCanvas->Canvas->DrawShadowedString(125, 320, TEXT("Page Residency"), GEngine->GetSmallFont(), FLinearColor(0.8f, 0.1f, 0.1f));
 	InCanvas->Canvas->DrawShadowedString(225, 320, TEXT("LockedPage Residency"), GEngine->GetSmallFont(), FLinearColor(0.8f, 0.8f, 0.1f));
-
-#endif
 }
 
+#endif // !UE_BUILD_SHIPPING
+
 #undef LOCTEXT_NAMESPACE
diff --git a/Engine/Source/Runtime/Renderer/Private/VT/VirtualTextureSystem.h b/Engine/Source/Runtime/Renderer/Private/VT/VirtualTextureSystem.h
--- a/Engine/Source/Runtime/Renderer/Private/VT/VirtualTextureSystem.h
+++ b/Engine/Source/Runtime/Renderer/Private/VT/VirtualTextureSystem.h
@@ -107,13 +107,6 @@ private:
 
 	void UpdateResidencyTracking() const;
 
-#if !UE_BUILD_SHIPPING
-	void GetOnScreenMessages(FCoreDelegates::FSeverityMessageMap& OutMessages);
-
-	FCoreDelegates::FSeverityMessageMap OnScreenMessages;
-	FCriticalSection OnScreenMessageLock;
-#endif
-
 	uint32	Frame;
 
 	static const uint32 MaxNumTasks = 16;
@@ -159,8 +152,17 @@ private:
 	TArray<FVirtualTextureLocalTile> TransientCollectedPages;
 	TArray<IVirtualTextureFinalizer*> Finalizers;
 
+#if !UE_BUILD_SHIPPING
+	void GetOnScreenMessages(FCoreDelegates::FSeverityMessageMap& OutMessages);
+	FCoreDelegates::FSeverityMessageMap OnScreenMessages;
+	FCriticalSection OnScreenMessageLock;
+
+	void UpdateNotifications();
+	void UpdateResidencyNotifications();
+
 	void DrawResidencyHud(class UCanvas*, class APlayerController*);
 	FDelegateHandle	DrawResidencyHudDelegateHandle;
+#endif
 };
 
 
