Make Shadow Depth Rendering check Nanite::GStreamingManager.HasResourceEntries()

diff --git a/Engine/Source/Runtime/Renderer/Private/ShadowDepthRendering.cpp b/Engine/Source/Runtime/Renderer/Private/ShadowDepthRendering.cpp
--- a/Engine/Source/Runtime/Renderer/Private/ShadowDepthRendering.cpp
+++ b/Engine/Source/Runtime/Renderer/Private/ShadowDepthRendering.cpp
@@ -1692,7 +1692,10 @@ void FSceneRenderer::RenderShadowDepthMaps(FRDGBuilder& GraphBuilder, FInstanceC
 
 	const bool bHasVSMShadows = SortedShadowsForShadowDepthPass.VirtualShadowMapShadows.Num() > 0;
 	const bool bHasVSMClipMaps = SortedShadowsForShadowDepthPass.VirtualShadowMapClipmaps.Num() > 0;
-	const bool bNaniteEnabled = UseNanite(ShaderPlatform) && ViewFamily.EngineShowFlags.NaniteMeshes;
+	const bool bNaniteEnabled = 
+		UseNanite(ShaderPlatform) &&
+		ViewFamily.EngineShowFlags.NaniteMeshes &&
+		Nanite::GStreamingManager.HasResourceEntries();
 
 	if (bNaniteEnabled && (bHasVSMShadows || bHasVSMClipMaps))
 	{
