Only update the FVirtualTextureSystem during DrawTileMesh when the material being rendered actually samples from a virtual texture. This was introducing lags in processing VT feedbacks, particularly evident when rendering the water brush which uses lots of DrawTileMesh internally

diff --git a/Engine/Source/Runtime/Renderer/Private/Renderer.cpp b/Engine/Source/Runtime/Renderer/Private/Renderer.cpp
--- a/Engine/Source/Runtime/Renderer/Private/Renderer.cpp
+++ b/Engine/Source/Runtime/Renderer/Private/Renderer.cpp
@@ -214,12 +214,14 @@ void FRendererModule::DrawTileMesh(FCanvasRenderContext& RenderContext, FMeshPas
 		}
 
 		//get the blend mode of the material
-		const EBlendMode MaterialBlendMode = Mesh.MaterialRenderProxy->GetIncompleteMaterialWithFallback(FeatureLevel).GetBlendMode();
+		const FMaterial& MeshMaterial = Mesh.MaterialRenderProxy->GetIncompleteMaterialWithFallback(FeatureLevel);
+		const EBlendMode MaterialBlendMode = MeshMaterial.GetBlendMode();
 
 		FRDGBuilder& GraphBuilder = RenderContext.GraphBuilder;
 
 		const bool bUseVirtualTexturing = UseVirtualTexturing(FeatureLevel);
-		if (bUseVirtualTexturing)
+		// Materials sampling VTs need FVirtualTextureSystem to be updated before being rendered :
+		if (bUseVirtualTexturing && !MeshMaterial.GetUniformVirtualTextureExpressions().IsEmpty())
 		{
 			RDG_GPU_STAT_SCOPE(GraphBuilder, VirtualTextureUpdate);
 			FVirtualTextureSystem::Get().AllocateResources(GraphBuilder, FeatureLevel);
