Fixed Skylight not rendering correctly when r.Lumen.TranslucencyVolume.Enabled = 0

diff --git a/Engine/Source/Runtime/Renderer/Private/TranslucentRendering.cpp b/Engine/Source/Runtime/Renderer/Private/TranslucentRendering.cpp
--- a/Engine/Source/Runtime/Renderer/Private/TranslucentRendering.cpp
+++ b/Engine/Source/Runtime/Renderer/Private/TranslucentRendering.cpp
@@ -856,6 +856,11 @@ TRDGUniformBufferRef<FTranslucentBasePassUniformParameters> CreateTranslucentBas
 	BasePassParameters.TranslucencyLightingVolume = GetTranslucencyLightingVolumeParameters(GraphBuilder, TranslucencyLightingVolumeTextures, ViewIndex);
 	BasePassParameters.LumenParameters = GetLumenTranslucencyLightingParameters(GraphBuilder, View.LumenTranslucencyGIVolume);
 
+	const bool bLumenGIHandlingSkylight = View.FinalPostProcessSettings.DynamicGlobalIlluminationMethod == EDynamicGlobalIlluminationMethod::Lumen
+		&& BasePassParameters.LumenParameters.TranslucencyGIGridSize.Z > 0;
+
+	BasePassParameters.Shared.UseBasePassSkylight = bLumenGIHandlingSkylight ? 0 : 1;
+
 	BasePassParameters.SceneColorCopyTexture = SystemTextures.Black;
 	BasePassParameters.SceneColorCopySampler = TStaticSamplerState<SF_Bilinear, AM_Clamp, AM_Clamp, AM_Clamp>::GetRHI();
 
