Use StructuredBuffer for InstanceIds to avoid size limitations of texel buffer on MALI

diff --git a/Engine/Shaders/Private/InstanceCulling/BuildInstanceDrawCommands.usf b/Engine/Shaders/Private/InstanceCulling/BuildInstanceDrawCommands.usf
--- a/Engine/Shaders/Private/InstanceCulling/BuildInstanceDrawCommands.usf
+++ b/Engine/Shaders/Private/InstanceCulling/BuildInstanceDrawCommands.usf
@@ -63,7 +63,7 @@ int bDrawOnlyVSMInvalidatingGeometry;
 
 
 RWStructuredBuffer<uint> InstanceIdsBufferOut;
-RWBuffer<float4> InstanceIdsBufferOutMobile;
+RWStructuredBuffer<float4> InstanceIdsBufferOutMobile;
 RWBuffer<uint> DrawIndirectArgsBufferOut;
 
 // this is just to avoid compiling packing code for unrelated shaders
diff --git a/Engine/Source/Runtime/Renderer/Private/InstanceCulling/InstanceCullingContext.cpp b/Engine/Source/Runtime/Renderer/Private/InstanceCulling/InstanceCullingContext.cpp
--- a/Engine/Source/Runtime/Renderer/Private/InstanceCulling/InstanceCullingContext.cpp
+++ b/Engine/Source/Runtime/Renderer/Private/InstanceCulling/InstanceCullingContext.cpp
@@ -223,7 +223,7 @@ public:
 		SHADER_PARAMETER_RDG_BUFFER_SRV(Buffer<uint>, InstanceIdOffsetBuffer)
 
 		SHADER_PARAMETER_RDG_BUFFER_UAV(RWStructuredBuffer<uint>, InstanceIdsBufferOut)
-		SHADER_PARAMETER_RDG_BUFFER_UAV(RWBuffer<float4>, InstanceIdsBufferOutMobile)
+		SHADER_PARAMETER_RDG_BUFFER_UAV(RWStructuredBuffer<float4>, InstanceIdsBufferOutMobile)
 		SHADER_PARAMETER_RDG_BUFFER_UAV(RWStructuredBuffer<uint>, DrawCommandIdsBufferOut)
 
 		SHADER_PARAMETER_RDG_BUFFER_UAV(RWBuffer<uint>, DrawIndirectArgsBufferOut)
@@ -363,20 +363,10 @@ void FInstanceCullingContext::BuildRenderingCommands(
 	const bool bUseDebugMode = bDrawOnlyVSMInvalidatingGeometry;
 
 	FRDGBufferRef ViewIdsBuffer = CreateStructuredBuffer(GraphBuilder, TEXT("InstanceCulling.ViewIds"), ViewIds);
-	FRDGBufferRef InstanceIdsBuffer = nullptr;
-	FRDGBufferUAVRef InstanceIdsBufferUAV = nullptr;
+	
 	const uint32 InstanceIdBufferStride = GetInstanceIdBufferStride(FeatureLevel);
-	if (FeatureLevel == ERHIFeatureLevel::ES3_1)
-	{
-		// This buffer will be used as per-instance vertex buffer. AFAIK only DX11 does not allow structured buffer to be used as a vertex buffer 
-		InstanceIdsBuffer = GraphBuilder.CreateBuffer(FRDGBufferDesc::CreateBufferDesc(InstanceIdBufferStride, InstanceIdBufferSize), TEXT("InstanceCulling.InstanceIdsBuffer"));
-		InstanceIdsBufferUAV = GraphBuilder.CreateUAV(InstanceIdsBuffer, PF_A32B32G32R32F, ERDGUnorderedAccessViewFlags::SkipBarrier);
-	}
-	else
-	{
-		InstanceIdsBuffer = GraphBuilder.CreateBuffer(FRDGBufferDesc::CreateStructuredDesc(InstanceIdBufferStride, InstanceIdBufferSize), TEXT("InstanceCulling.InstanceIdsBuffer"));
-		InstanceIdsBufferUAV = GraphBuilder.CreateUAV(InstanceIdsBuffer, ERDGUnorderedAccessViewFlags::SkipBarrier);
-	}
+	FRDGBufferRef InstanceIdsBuffer = GraphBuilder.CreateBuffer(FRDGBufferDesc::CreateStructuredDesc(InstanceIdBufferStride, InstanceIdBufferSize), TEXT("InstanceCulling.InstanceIdsBuffer"));
+	FRDGBufferUAVRef InstanceIdsBufferUAV = GraphBuilder.CreateUAV(InstanceIdsBuffer, ERDGUnorderedAccessViewFlags::SkipBarrier);
 
 	FBuildInstanceIdBufferAndCommandsFromPrimitiveIdsCs::FParameters PassParametersTmp;
 
@@ -663,20 +653,15 @@ FInstanceCullingDeferredContext *FInstanceCullingContext::CreateDeferredContext(
 	FRDGBufferRef InstanceIdOffsetBuffer = GraphBuilder.CreateBuffer(FRDGBufferDesc::CreateBufferDesc(sizeof(uint32), 1), TEXT("InstanceCulling.InstanceIdOffsetBuffer"), INST_CULL_CALLBACK(DeferredContext->InstanceIdOffsets.Num()));
 	GraphBuilder.QueueBufferUpload(InstanceIdOffsetBuffer, INST_CULL_CALLBACK(DeferredContext->InstanceIdOffsets.GetData()), INST_CULL_CALLBACK(DeferredContext->InstanceIdOffsets.GetTypeSize() * DeferredContext->InstanceIdOffsets.Num()));
 
-	FRDGBufferRef InstanceIdsBuffer = nullptr;
-	FRDGBufferUAVRef InstanceIdsBufferUAV = nullptr;
 	const uint32 InstanceIdBufferStride = GetInstanceIdBufferStride(FeatureLevel);
+	FRDGBufferRef InstanceIdsBuffer = GraphBuilder.CreateBuffer(FRDGBufferDesc::CreateStructuredDesc(InstanceIdBufferStride, 1), TEXT("InstanceCulling.InstanceIdsBuffer"), INST_CULL_CALLBACK(DeferredContext->InstanceIdBufferSize));
+	FRDGBufferUAVRef InstanceIdsBufferUAV = GraphBuilder.CreateUAV(InstanceIdsBuffer, ERDGUnorderedAccessViewFlags::SkipBarrier);
 	if (FeatureLevel == ERHIFeatureLevel::ES3_1)
 	{
-		// This buffer will be used as per-instance vertex buffer. AFAIK only DX11 does not allow structured buffer to be used as a vertex buffer 
-		InstanceIdsBuffer = GraphBuilder.CreateBuffer(FRDGBufferDesc::CreateBufferDesc(InstanceIdBufferStride, 1), TEXT("InstanceCulling.InstanceIdsBuffer"), INST_CULL_CALLBACK(DeferredContext->InstanceIdBufferSize));
-		InstanceIdsBufferUAV = GraphBuilder.CreateUAV(InstanceIdsBuffer, PF_A32B32G32R32F, ERDGUnorderedAccessViewFlags::SkipBarrier);
 		DeferredContext->InstanceDataBuffer = InstanceIdsBuffer;
 	}
 	else
 	{
-		InstanceIdsBuffer = GraphBuilder.CreateBuffer(FRDGBufferDesc::CreateStructuredDesc(InstanceIdBufferStride, 1), TEXT("InstanceCulling.InstanceIdsBuffer"), INST_CULL_CALLBACK(DeferredContext->InstanceIdBufferSize));
-		InstanceIdsBufferUAV = GraphBuilder.CreateUAV(InstanceIdsBuffer, ERDGUnorderedAccessViewFlags::SkipBarrier);
 		DeferredContext->InstanceDataBuffer = InstanceIdOffsetBuffer;
 	}
 
