Fill StencilValue correctly in the DepthPass when enabled
Revert BasePass changes

diff --git a/Engine/Source/Runtime/Renderer/Private/DepthRendering.cpp b/Engine/Source/Runtime/Renderer/Private/DepthRendering.cpp
--- a/Engine/Source/Runtime/Renderer/Private/DepthRendering.cpp
+++ b/Engine/Source/Runtime/Renderer/Private/DepthRendering.cpp
@@ -708,7 +708,7 @@ FMeshDrawCommandSortKey CalculateDepthPassMeshStaticSortKey(EBlendMode BlendMode
 	return SortKey;
 }
 
-void SetMobileDepthPassRenderState(const FPrimitiveSceneProxy* RESTRICT PrimitiveSceneProxy, FMeshPassProcessorRenderState& DrawRenderState)
+void SetMobileDepthPassRenderState(const FPrimitiveSceneProxy* RESTRICT PrimitiveSceneProxy, FMeshPassProcessorRenderState& DrawRenderState, const FMaterial& RESTRICT MaterialResource, bool bUsesDeferredShading)
 {
 	DrawRenderState.SetDepthStencilState(TStaticDepthStencilState<
 		true, CF_DepthNearOrEqual,
@@ -716,9 +716,19 @@ void SetMobileDepthPassRenderState(const FPrimitiveSceneProxy* RESTRICT Primitiv
 		false, CF_Always, SO_Keep, SO_Keep, SO_Keep,
 		// don't use masking as it has significant performance hit on Mali GPUs (T860MP2)
 		0x00, 0xff >::GetRHI());
+
 	uint8 StencilValue = 0;
+	
 	uint8 ReceiveDecals = (PrimitiveSceneProxy && !PrimitiveSceneProxy->ReceivesDecals() ? 0x01 : 0x00);
 	StencilValue |= GET_STENCIL_BIT_MASK(RECEIVE_DECAL, ReceiveDecals);
+
+	if (bUsesDeferredShading)
+	{
+		// store into [1-3] bits
+		uint8 ShadingModel = MaterialResource.GetShadingModels().IsLit() ? MSM_DefaultLit : MSM_Unlit;
+		StencilValue |= GET_STENCIL_MOBILE_SM_MASK(ShadingModel);
+	}
+
 	DrawRenderState.SetStencilRef(StencilValue);
 }
 
@@ -763,7 +773,7 @@ bool FDepthPassMeshProcessor::Process(
 	// Use StencilMask for DecalOutput on mobile
 	if (FeatureLevel == ERHIFeatureLevel::ES3_1)
 	{
-		SetMobileDepthPassRenderState(PrimitiveSceneProxy, DrawRenderState);
+		SetMobileDepthPassRenderState(PrimitiveSceneProxy, DrawRenderState, MaterialResource, IsMobileDeferredShadingEnabled(GetFeatureLevelShaderPlatform(FeatureLevel)));
 	}
 
 	FDepthOnlyShaderElementData ShaderElementData(0.0f);
diff --git a/Engine/Source/Runtime/Renderer/Private/MobileBasePass.cpp b/Engine/Source/Runtime/Renderer/Private/MobileBasePass.cpp
--- a/Engine/Source/Runtime/Renderer/Private/MobileBasePass.cpp
+++ b/Engine/Source/Runtime/Renderer/Private/MobileBasePass.cpp
@@ -353,7 +353,7 @@ ELightMapPolicyType MobileBasePass::SelectMeshLightmapPolicy(
 	return SelectedLightmapPolicy;
 }
 
-void MobileBasePass::SetOpaqueRenderState(FMeshPassProcessorRenderState& DrawRenderState, const FPrimitiveSceneProxy* PrimitiveSceneProxy, const FMaterial& Material, bool bEnableReceiveDecalOutput, bool bUsesDeferredShading, bool bDisableDepthWrite)
+void MobileBasePass::SetOpaqueRenderState(FMeshPassProcessorRenderState& DrawRenderState, const FPrimitiveSceneProxy* PrimitiveSceneProxy, const FMaterial& Material, bool bEnableReceiveDecalOutput, bool bUsesDeferredShading)
 {
 	uint8 StencilValue = 0;
 	if (bEnableReceiveDecalOutput)
@@ -371,36 +371,18 @@ void MobileBasePass::SetOpaqueRenderState(FMeshPassProcessorRenderState& DrawRen
 		
 	if (bEnableReceiveDecalOutput || bUsesDeferredShading)
 	{
-		if (bDisableDepthWrite)
-		{
-			DrawRenderState.SetDepthStencilState(TStaticDepthStencilState<
-				false, CF_Equal,
-				true, CF_Always, SO_Keep, SO_Keep, SO_Replace,
-				false, CF_Always, SO_Keep, SO_Keep, SO_Keep,
-				// don't use masking as it has significant performance hit on Mali GPUs (T860MP2)
-				0x00, 0xff >::GetRHI());
-		}
-		else
-		{
-			DrawRenderState.SetDepthStencilState(TStaticDepthStencilState<
+		DrawRenderState.SetDepthStencilState(TStaticDepthStencilState<
 				true, CF_DepthNearOrEqual,
 				true, CF_Always, SO_Keep, SO_Keep, SO_Replace,
 				false, CF_Always, SO_Keep, SO_Keep, SO_Keep,
 				// don't use masking as it has significant performance hit on Mali GPUs (T860MP2)
 				0x00, 0xff >::GetRHI());
-		}
-		
 
 		DrawRenderState.SetStencilRef(StencilValue); 
 	}
 	else
 	{
 		// default depth state should be already set
-		// disable DepthWrite if Prepass enabled
-		if (bDisableDepthWrite)
-		{
-			DrawRenderState.SetDepthStencilState(TStaticDepthStencilState<false, CF_Equal>::GetRHI());
-		}
 	}
 
 	if (Material.GetBlendMode() == BLEND_Masked && Material.IsUsingAlphaToCoverage())
@@ -778,11 +760,14 @@ bool FMobileBasePassMeshProcessor::Process(
 		{
 			MobileBasePass::SetTranslucentRenderState(DrawRenderState, MaterialResource);
 		}
+		else if((MeshBatch.bUseForDepthPass && Scene->EarlyZPassMode == DDM_AllOpaque) || bMaskedInEarlyPass)
+		{
+			DrawRenderState.SetDepthStencilState(TStaticDepthStencilState<false, CF_Equal>::GetRHI());
+		}
 		else
 		{
 			const bool bEnableReceiveDecalOutput = ((Flags & EFlags::CanUseDepthStencil) == EFlags::CanUseDepthStencil);
-			const bool bDisableDepthWrite = (MeshBatch.bUseForDepthPass && Scene->EarlyZPassMode == DDM_AllOpaque) || bMaskedInEarlyPass;
-			MobileBasePass::SetOpaqueRenderState(DrawRenderState, PrimitiveSceneProxy, MaterialResource, bEnableReceiveDecalOutput && IsMobileHDR(), bUsesDeferredShading, bDisableDepthWrite);
+			MobileBasePass::SetOpaqueRenderState(DrawRenderState, PrimitiveSceneProxy, MaterialResource, bEnableReceiveDecalOutput && IsMobileHDR(), bUsesDeferredShading);
 		}
 	}
 
diff --git a/Engine/Source/Runtime/Renderer/Private/MobileBasePassRendering.h b/Engine/Source/Runtime/Renderer/Private/MobileBasePassRendering.h
--- a/Engine/Source/Runtime/Renderer/Private/MobileBasePassRendering.h
+++ b/Engine/Source/Runtime/Renderer/Private/MobileBasePassRendering.h
@@ -353,7 +353,7 @@ namespace MobileBasePass
 
 	bool StaticCanReceiveCSM(const FLightSceneInfo* LightSceneInfo, const FPrimitiveSceneProxy* PrimitiveSceneProxy);
 
-	void SetOpaqueRenderState(FMeshPassProcessorRenderState& DrawRenderState, const FPrimitiveSceneProxy* PrimitiveSceneProxy, const FMaterial& Material, bool bEnableReceiveDecalOutput, bool bUsesDeferredShading, bool bDisableDepthWrite);
+	void SetOpaqueRenderState(FMeshPassProcessorRenderState& DrawRenderState, const FPrimitiveSceneProxy* PrimitiveSceneProxy, const FMaterial& Material, bool bEnableReceiveDecalOutput, bool bUsesDeferredShading);
 	void SetTranslucentRenderState(FMeshPassProcessorRenderState& DrawRenderState, const FMaterial& Material);
 
 	bool StationarySkyLightHasBeenApplied(const FScene* Scene, ELightMapPolicyType LightMapPolicyType);
