Disables translucency depth upscaling in DOFRecombine pass when r.SeparateTranslucencyScreenPercentage=100

diff --git a/Engine/Shaders/Private/DiaphragmDOF/DOFRecombine.usf b/Engine/Shaders/Private/DiaphragmDOF/DOFRecombine.usf
--- a/Engine/Shaders/Private/DiaphragmDOF/DOFRecombine.usf
+++ b/Engine/Shaders/Private/DiaphragmDOF/DOFRecombine.usf
@@ -140,6 +140,7 @@ float2 DepthBlurParameters;
 float2 DOFBufferUVMax;
 
 float4 SeparateTranslucencyBilinearUVMinMax;
+uint SeparateTranslucencyUpscaling;
 
 Texture2D BokehLUT;
 
@@ -1085,9 +1086,19 @@ void RecombineMainCS(
 	// Compose separate translucency.
 	if (1)
 	{
-		float2 PixelPos = InputParameters.ViewportUV * View.ViewSizeAndInvSize.xy + View.ViewRectMin.xy;
-		float2 BufferUV = PixelPos * View.BufferSizeAndInvSize.zw;
-		NearestDepthNeighborUpsamplingResult UpsampleResult = NearestDepthNeighborUpsampling(PixelPos, BufferUV, SeparateTranslucencyTextureLowResExtentInverse);
+		NearestDepthNeighborUpsamplingResult UpsampleResult;
+		if (SeparateTranslucencyUpscaling == 0)
+		{
+			UpsampleResult.bUsePointSampler = true;
+			UpsampleResult.UV = InputParameters.SceneBufferUV;
+		}
+		else
+		{
+			float2 PixelPos = (View.ViewRectMin.xy + DispatchThreadId) + 0.5;
+			UpsampleResult = NearestDepthNeighborUpsampling(PixelPos, InputParameters.SceneBufferUV, SeparateTranslucencyTextureLowResExtentInverse);
+		}
+		
+		UpsampleResult.UV = clamp(UpsampleResult.UV, SeparateTranslucencyBilinearUVMinMax.xy, SeparateTranslucencyBilinearUVMinMax.zw);
 
 		float4 SeparateTranslucencyColor = 0;
 		float4 SeparateTranslucencyModulateColor = 0;
diff --git a/Engine/Source/Runtime/Renderer/Private/PostProcess/DiaphragmDOF.cpp b/Engine/Source/Runtime/Renderer/Private/PostProcess/DiaphragmDOF.cpp
--- a/Engine/Source/Runtime/Renderer/Private/PostProcess/DiaphragmDOF.cpp
+++ b/Engine/Source/Runtime/Renderer/Private/PostProcess/DiaphragmDOF.cpp
@@ -1259,6 +1259,7 @@ class FDiaphragmDOFRecombineCS : public FDiaphragmDOFShader
 		SHADER_PARAMETER(FVector2D, TemporalJitterPixels)
 		SHADER_PARAMETER(FVector2D, DOFBufferUVMax)
 		SHADER_PARAMETER(FVector4, SeparateTranslucencyBilinearUVMinMax)
+		SHADER_PARAMETER(int32, SeparateTranslucencyUpscaling)
 		
 		SHADER_PARAMETER_RDG_TEXTURE(Texture2D, BokehLUT)
 
@@ -2632,11 +2633,12 @@ FRDGTextureRef DiaphragmDOF::AddPasses(
 		PassParameters->DOFBufferUVMax = FVector2D(
 			(GatheringViewSize.X - 0.5f) / float(RefBufferSize.X),
 			(GatheringViewSize.Y - 0.5f) / float(RefBufferSize.Y));
-		
+
 		PassParameters->SeparateTranslucencyBilinearUVMinMax.X = (SeparateTranslucencyRect.Min.X + 0.5f) / float(SeparateTranslucencyDimensions.Extent.X);
 		PassParameters->SeparateTranslucencyBilinearUVMinMax.Y = (SeparateTranslucencyRect.Min.Y + 0.5f) / float(SeparateTranslucencyDimensions.Extent.Y);
 		PassParameters->SeparateTranslucencyBilinearUVMinMax.Z = (SeparateTranslucencyRect.Max.X - 0.5f) / float(SeparateTranslucencyDimensions.Extent.X);
 		PassParameters->SeparateTranslucencyBilinearUVMinMax.W = (SeparateTranslucencyRect.Max.Y - 0.5f) / float(SeparateTranslucencyDimensions.Extent.Y);
+		PassParameters->SeparateTranslucencyUpscaling = SeparateTranslucencyRect.Size() != FullResViewSize ? 1 : 0;
 
 		PassParameters->SceneColorInput = FullResGatherInputTextures.SceneColor;
 		PassParameters->SceneDepthTexture = SceneTextures.SceneDepthTexture;
