Added far field culling using r.RayTracing.FarField.MaxTraceDistance

diff --git a/Engine/Source/Runtime/Renderer/Private/DeferredShadingRenderer.cpp b/Engine/Source/Runtime/Renderer/Private/DeferredShadingRenderer.cpp
--- a/Engine/Source/Runtime/Renderer/Private/DeferredShadingRenderer.cpp
+++ b/Engine/Source/Runtime/Renderer/Private/DeferredShadingRenderer.cpp
@@ -210,11 +210,14 @@ static FAutoConsoleVariableRef CVarRayTracingDebugForceOpaque(
 	TEXT("Forces all ray tracing geometry instances to be opaque, effectively disabling any-hit shaders. This is useful for debugging and profiling. (default = 0)")
 );
 
-static int32 GRayTracingFarField = 0;
-static FAutoConsoleVariableRef CVarRaytracingFarField(
-	TEXT("r.RayTracing.FarField"),
-	GRayTracingFarField,
-	TEXT("Enable/Disable far field ray tracing.\n"),
+static TAutoConsoleVariable<int32> CVarRayTracingFarField(
+	TEXT("r.RayTracing.FarField"), 0,
+	TEXT("Enable/Disable far field ray tracing."),
+	ECVF_RenderThreadSafe);
+
+static TAutoConsoleVariable<float> CVarRayTracingFarFieldMaxTraceDistance(
+	TEXT("r.RayTracing.FarField.MaxTraceDistance"), 1.0e6f,
+	TEXT("Maximum hit-distance for far-field ray tracing (Default = 1.0e6)."),
 	ECVF_RenderThreadSafe);
 
 #if !UE_BUILD_SHIPPING
@@ -262,6 +265,16 @@ float GetRayTracingCullingRadius()
 	return CVarRayTracingCullingRadius.GetValueOnRenderThread();
 }
 
+int32 GetRayTracingFarField()
+{
+	return CVarRayTracingFarField.GetValueOnRenderThread();
+}
+
+float GetRayTracingFarFieldMaxTraceDistance()
+{
+	return CVarRayTracingFarFieldMaxTraceDistance.GetValueOnRenderThread();
+}
+
 #endif // RHI_RAYTRACING
 
 DECLARE_CYCLE_STAT(TEXT("InitViews Intentional Stall"), STAT_InitViews_Intentional_Stall, STATGROUP_InitViews);
@@ -600,12 +613,14 @@ bool FDeferredShadingSceneRenderer::GatherRayTracingWorldInstancesForView(FRHICo
 		int32 BroadIndex = 0;
 		const int32 CullInRayTracing = GetRayTracingCulling();
 		const float CullingRadius = GetRayTracingCullingRadius();
+		const float FarFieldCullingRadius = GetRayTracingFarFieldMaxTraceDistance();
 		const float CullAngleThreshold = CVarRayTracingCullingAngle.GetValueOnRenderThread();
 		const float AngleThresholdRatio = FMath::Tan(FMath::Min(89.99f, CullAngleThreshold) * PI / 180.0f);
 		const FVector ViewOrigin = View.ViewMatrices.GetViewOrigin();
 		const FVector ViewDirection = View.GetViewDirection();
 		const bool bCullAllObjects = CullInRayTracing == 2 || CullInRayTracing == 3;
 		const bool bCullByRadiusOrDistance = CullInRayTracing == 3;
+		const bool bIsRayTracingFarField = GetRayTracingFarField() > 0;
 
 		for (int PrimitiveIndex = 0; PrimitiveIndex < Scene->PrimitiveSceneProxies.Num(); PrimitiveIndex++)
 		{
@@ -630,13 +645,13 @@ bool FDeferredShadingSceneRenderer::GatherRayTracingWorldInstancesForView(FRHICo
 			const FPrimitiveSceneInfo* SceneInfo = Scene->Primitives[PrimitiveIndex];
 
 			// Skip far field if not enabled
-			const bool bIsRayTracingFarField = EnumHasAnyFlags(Scene->PrimitiveRayTracingFlags[PrimitiveIndex], ERayTracingPrimitiveFlags::FarField);
-			if (!GRayTracingFarField && bIsRayTracingFarField)
+			const bool bIsFarFieldPrimitive = EnumHasAnyFlags(Scene->PrimitiveRayTracingFlags[PrimitiveIndex], ERayTracingPrimitiveFlags::FarField);
+			if (!bIsRayTracingFarField && bIsFarFieldPrimitive)
 			{
 				continue;
 			}
 
-			if (CullInRayTracing > 0 && !bIsRayTracingFarField)
+			if (CullInRayTracing > 0)
 			{
 				const FBoxSphereBounds ObjectBounds = Scene->PrimitiveBounds[PrimitiveIndex].BoxSphereBounds;
 				const float ObjectRadius = ObjectBounds.SphereRadius;
@@ -648,18 +663,29 @@ bool FDeferredShadingSceneRenderer::GatherRayTracingWorldInstancesForView(FRHICo
 				if (bConsiderCulling)
 				{
 					const float CameraToObjectCenterLength = CameraToObjectCenter.Size();
-					const bool bIsFarEnoughToCull = CameraToObjectCenterLength > (CullingRadius + ObjectRadius);
-
-					// Cull by solid angle: check the radius of bounding sphere against angle threshold
-					const bool bAngleIsSmallEnoughToCull = ObjectRadius / CameraToObjectCenterLength < AngleThresholdRatio;
-
-					if (bCullByRadiusOrDistance && (bIsFarEnoughToCull || bAngleIsSmallEnoughToCull))
+					
+					if (bIsFarFieldPrimitive)
 					{
-						continue;
+						if (CameraToObjectCenterLength > (FarFieldCullingRadius + ObjectRadius))
+						{
+							continue;
+						}
 					}
-					else if (bIsFarEnoughToCull && bAngleIsSmallEnoughToCull)
+					else
 					{
-						continue;
+						const bool bIsFarEnoughToCull = CameraToObjectCenterLength > (CullingRadius + ObjectRadius);
+
+						// Cull by solid angle: check the radius of bounding sphere against angle threshold
+						const bool bAngleIsSmallEnoughToCull = ObjectRadius / CameraToObjectCenterLength < AngleThresholdRatio;
+
+						if (bCullByRadiusOrDistance && (bIsFarEnoughToCull || bAngleIsSmallEnoughToCull))
+						{
+							continue;
+						}
+						else if (bIsFarEnoughToCull && bAngleIsSmallEnoughToCull)
+						{
+							continue;
+						}
 					}
 				}
 			}
