Fixed missing instancing stats for a mobile basepass

diff --git a/Engine/Source/Runtime/Renderer/Private/MobileShadingRenderer.cpp b/Engine/Source/Runtime/Renderer/Private/MobileShadingRenderer.cpp
--- a/Engine/Source/Runtime/Renderer/Private/MobileShadingRenderer.cpp
+++ b/Engine/Source/Runtime/Renderer/Private/MobileShadingRenderer.cpp
@@ -331,6 +331,11 @@ void FMobileSceneRenderer::SetupMobileBasePassAfterShadowInit(FExclusiveDepthSte
 
 		// Run sorting on BasePass, as it's ignored inside FSceneRenderer::SetupMeshPass, so it can be done after shadow init on mobile.
 		FParallelMeshDrawCommandPass& Pass = View.ParallelMeshDrawCommandPasses[EMeshPass::BasePass];
+		if (ShouldDumpMeshDrawCommandInstancingStats())
+		{
+			Pass.SetDumpInstancingStats(GetMeshPassName(EMeshPass::BasePass));
+		}
+
 		Pass.DispatchPassSetup(
 			Scene,
 			View,
