Allow multiple physical pools to be created for the same virtual texture format.
Producers will be randomly distributed between physical pools.
The maximum physical pool tile size is set by r.VT.SplitPhysicalPoolSize.
Setting a size of 64 will force all pools to be divided so that their page tables fit in 16 bits.
This makes it possible to use large physical pools without doubling the cost of page table memory.

diff --git a/Engine/Source/Runtime/Engine/Private/VT/VirtualTexturePoolConfig.cpp b/Engine/Source/Runtime/Engine/Private/VT/VirtualTexturePoolConfig.cpp
--- a/Engine/Source/Runtime/Engine/Private/VT/VirtualTexturePoolConfig.cpp
+++ b/Engine/Source/Runtime/Engine/Private/VT/VirtualTexturePoolConfig.cpp
@@ -2,8 +2,6 @@
 
 #include "VT/VirtualTexturePoolConfig.h"
 
-#include "VT/VirtualTextureScalability.h"
-
 UVirtualTexturePoolConfig::UVirtualTexturePoolConfig(const FObjectInitializer& ObjectInitializer)
 	: Super(ObjectInitializer)
 {
@@ -33,16 +31,6 @@ void UVirtualTexturePoolConfig::FindPoolConfig(TEnumAsByte<EPixelFormat> const*
 			if (bAllFormatsMatch)
 			{
 				OutConfig = Config;
-				const float Scale = Config.bAllowSizeScale ? VirtualTextureScalability::GetPoolSizeScale(Config.ScalabilityGroup) : 1.f;
-				OutConfig.SizeInMegabyte = (int32)(Scale * (float)OutConfig.SizeInMegabyte);
-				if (Scale < 1.f && OutConfig.MinScaledSizeInMegabyte > 0)
-				{
-					OutConfig.SizeInMegabyte = FMath::Max(OutConfig.SizeInMegabyte, OutConfig.MinScaledSizeInMegabyte);
-				}
-				if (Scale > 1.f && OutConfig.MaxScaledSizeInMegabyte > 0)
-				{
-					OutConfig.SizeInMegabyte = FMath::Min(OutConfig.SizeInMegabyte, OutConfig.MaxScaledSizeInMegabyte);
-				}
 				return;
 			}
 		}
diff --git a/Engine/Source/Runtime/Engine/Private/VT/VirtualTextureScalability.cpp b/Engine/Source/Runtime/Engine/Private/VT/VirtualTextureScalability.cpp
--- a/Engine/Source/Runtime/Engine/Private/VT/VirtualTextureScalability.cpp
+++ b/Engine/Source/Runtime/Engine/Private/VT/VirtualTextureScalability.cpp
@@ -121,6 +121,15 @@ namespace VirtualTextureScalability
 		ECVF_Scalability
 	);
 
+	static TAutoConsoleVariable<int32> CVarVTSplitPhysicalPoolSize(
+		TEXT("r.VT.SplitPhysicalPoolSize"),
+		0,
+		TEXT("Create multiple physical pools per format to keep pools at this maximum size in tiles.")
+		TEXT("A value of 64 tiles will force 16bit page tables. This can be a page table memory optimization for large physical pools.")
+		TEXT("Defaults to 0 (off)."),
+		ECVF_RenderThreadSafe
+	);
+
 	static TAutoConsoleVariable<int32> CVarVTEnableAnisotropy(
 		TEXT("r.VT.AnisotropicFiltering"),
 		0,
@@ -255,8 +264,20 @@ namespace VirtualTextureScalability
 		return GroupIndex < NumScalabilityGroups ? GTileCountBiases[GroupIndex] : 0;
 	}
 
+	int32 GetSplitPhysicalPoolSize()
+	{
+		return FMath::Max(CVarVTSplitPhysicalPoolSize.GetValueOnRenderThread(), 0);
+	}
+
 	uint32 GetPageFreeThreshold()
 	{
 		return FMath::Max(CVarVTPageFreeThreshold.GetValueOnRenderThread(), 0);
 	}
+	
+	uint32 GetPhysicalPoolSettingsHash()
+	{
+		uint32 Hash = GetTypeHash(GPoolSizeScales);
+		Hash = HashCombine(Hash, GetSplitPhysicalPoolSize());
+		return Hash;
+	}
 }
diff --git a/Engine/Source/Runtime/Engine/Public/VT/VirtualTextureScalability.h b/Engine/Source/Runtime/Engine/Public/VT/VirtualTextureScalability.h
--- a/Engine/Source/Runtime/Engine/Public/VT/VirtualTextureScalability.h
+++ b/Engine/Source/Runtime/Engine/Public/VT/VirtualTextureScalability.h
@@ -18,7 +18,6 @@ namespace VirtualTextureScalability
 	ENGINE_API int32 GetRuntimeVirtualTextureSizeBias();
 	/** Is HW Anisotropic filtering enabled for VT */
 	ENGINE_API bool IsAnisotropicFilteringEnabled();
-
 	/**
 	 * Get maximum anisotropy when virtual texture sampling. 
 	 * This is also clamped per virtual texture according to the tile border size.
@@ -28,6 +27,11 @@ namespace VirtualTextureScalability
 	ENGINE_API float GetPoolSizeScale(uint32 GroupIndex);
 	/** Get resolution bias for runtime virtual textures. */
 	ENGINE_API int32 GetRuntimeVirtualTextureSizeBias(uint32 GroupIndex);
+	/** Get maximum size in tiles for physical pools before we split them. */
+	ENGINE_API int32 GetSplitPhysicalPoolSize();
 	/** Get the number of frames a page must be unused, before it's considered free */
 	ENGINE_API uint32 GetPageFreeThreshold();
+
+	/** Get a unique hash of all state affecting physical pool allocation.*/
+	ENGINE_API uint32 GetPhysicalPoolSettingsHash();
 }
diff --git a/Engine/Source/Runtime/Renderer/Private/VT/VirtualTexturePhysicalSpace.cpp b/Engine/Source/Runtime/Renderer/Private/VT/VirtualTexturePhysicalSpace.cpp
--- a/Engine/Source/Runtime/Renderer/Private/VT/VirtualTexturePhysicalSpace.cpp
+++ b/Engine/Source/Runtime/Renderer/Private/VT/VirtualTexturePhysicalSpace.cpp
@@ -46,12 +46,12 @@ static TAutoConsoleVariable<float> CVarVTResidencyAdjustmentRate(
 	ECVF_RenderThreadSafe);
 
 
-FVirtualTexturePhysicalSpace::FVirtualTexturePhysicalSpace(const FVTPhysicalSpaceDescription& InDesc, uint16 InID)
+FVirtualTexturePhysicalSpace::FVirtualTexturePhysicalSpace(const FVTPhysicalSpaceDescription& InDesc, uint16 InID, int32 InTileWidthHeight, bool bInEnableResidencyMipMapBias)
 	: Description(InDesc)
+	, TextureSizeInTiles(InTileWidthHeight)
 	, NumRefs(0u)
 	, ID(InID)
-	, bGpuTextureLimit(false)
-	, bEnableResidencyMipMapBias(true)
+	, bEnableResidencyMipMapBias(bInEnableResidencyMipMapBias)
 	, ResidencyMipMapBias(0.0f)
 	, LastFrameOversubscribed(0)
 #if !UE_BUILD_SHIPPING
@@ -61,33 +61,6 @@ FVirtualTexturePhysicalSpace::FVirtualTexturePhysicalSpace(const FVTPhysicalSpac
 	, HistoryIndex(0)
 #endif
 {
-	// Find matching physical pool
-	FVirtualTextureSpacePoolConfig Config;
-	UVirtualTexturePoolConfig const* PoolConfig = GetDefault<UVirtualTexturePoolConfig>();
-	PoolConfig->FindPoolConfig(InDesc.Format, InDesc.NumLayers, InDesc.TileSize, Config);
-	const uint32 PoolSizeInBytes = Config.SizeInMegabyte * 1024u * 1024u;
-
-	const FPixelFormatInfo& FormatInfo = GPixelFormats[InDesc.Format[0]];
-	check(InDesc.TileSize % FormatInfo.BlockSizeX == 0);
-	check(InDesc.TileSize % FormatInfo.BlockSizeY == 0);
-	SIZE_T TileSizeBytes = 0;
-	for (int32 Layer = 0; Layer < InDesc.NumLayers; ++Layer)
-	{
-		TileSizeBytes += CalculateImageBytes(InDesc.TileSize, InDesc.TileSize, 0, InDesc.Format[Layer]);
-	}
-	const uint32 MaxTiles = FMath::Max((uint32)(PoolSizeInBytes / TileSizeBytes), 1u);
-	TextureSizeInTiles = FMath::FloorToInt(FMath::Sqrt((float)MaxTiles));
-	
-	if (TextureSizeInTiles * InDesc.TileSize > GetMax2DTextureDimension())
-	{
-		// A good option to support extremely large caches would be to allow additional slices in an array here for caches...
-		// Just try to use the maximum texture size for now
-		TextureSizeInTiles = GetMax2DTextureDimension() / InDesc.TileSize;
-		bGpuTextureLimit = true;
-	}
-
-	bEnableResidencyMipMapBias = Config.bEnableResidencyMipMapBias;
-
 	Pool.Initialize(GetNumTiles());
 
 	// Initialize this resource FeatureLevel, so it gets re-created on FeatureLevel changes
@@ -249,12 +222,14 @@ void FVirtualTexturePhysicalSpace::UpdateResidencyTracking(uint32 Frame)
 #endif
 }
 
-void FVirtualTexturePhysicalSpace::DrawResidencyGraph(FCanvas* Canvas, FBox2D CanvasPosition)
+void FVirtualTexturePhysicalSpace::DrawResidencyGraph(FCanvas* Canvas, FBox2D CanvasPosition, bool bDrawKey)
 {
 	// Note that this is called on game thread and reads the history values written on the render thread.
 	// But it should be safe and any race condition will only lead to a slightly incorrect graph.
 
 #if !UE_BUILD_SHIPPING
+	const int32 BorderSize = 10;
+
 	const FLinearColor BackgroundColor(0.0f, 0.0f, 0.0f, 0.7f);
 	const FLinearColor GraphBorderColor(0.1f, 0.1f, 0.1f);
 	const FLinearColor GraphResidencyColor(0.8f, 0.1f, 0.1f);
@@ -265,11 +240,16 @@ void FVirtualTexturePhysicalSpace::DrawResidencyGraph(FCanvas* Canvas, FBox2D Ca
 	BackgroundTile.BlendMode = SE_BLEND_AlphaBlend;
 	Canvas->DrawItem(BackgroundTile);
 
-	const int32 BorderSize = 10;
-
-	FString Title = FString::Printf(TEXT("%s (%dMB)"), *FormatString, GetSizeInBytes() / (1024 * 1024));
+	const FString Title = FString::Printf(TEXT("%s (%dPages, %dMB)"), *FormatString, Pool.GetNumPages(), GetSizeInBytes() / (1024 * 1024));
 	Canvas->DrawShadowedString(CanvasPosition.Min.X + BorderSize, CanvasPosition.Min.Y, *Title, GEngine->GetSmallFont(), FLinearColor::White);
 
+	if (bDrawKey)
+	{
+		Canvas->DrawShadowedString(CanvasPosition.Min.X, CanvasPosition.Max.Y + 10, TEXT("Page Residency"), GEngine->GetSmallFont(), GraphResidencyColor);
+		Canvas->DrawShadowedString(CanvasPosition.Min.X + 100, CanvasPosition.Max.Y + 10, TEXT("MipMap Bias"), GEngine->GetSmallFont(), GraphMipMapBiasColor);
+		Canvas->DrawShadowedString(CanvasPosition.Min.X + 180, CanvasPosition.Max.Y + 10, TEXT("LockedPage Residency"), GEngine->GetSmallFont(), GraphLockedPageColor);
+	}
+
 	CanvasPosition.Min += FVector2D(BorderSize, BorderSize);
 	CanvasPosition.Max -= FVector2D(BorderSize, BorderSize);
 
diff --git a/Engine/Source/Runtime/Renderer/Private/VT/VirtualTexturePhysicalSpace.h b/Engine/Source/Runtime/Renderer/Private/VT/VirtualTexturePhysicalSpace.h
--- a/Engine/Source/Runtime/Renderer/Private/VT/VirtualTexturePhysicalSpace.h
+++ b/Engine/Source/Runtime/Renderer/Private/VT/VirtualTexturePhysicalSpace.h
@@ -42,10 +42,23 @@ inline bool operator!=(const FVTPhysicalSpaceDescription& Lhs, const FVTPhysical
 	return !operator==(Lhs, Rhs);
 }
 
+inline uint32 GetTypeHash(const FVTPhysicalSpaceDescription& Desc)
+{
+	uint32 Hash = GetTypeHash(Desc.TileSize);
+	Hash = HashCombine(Hash, GetTypeHash(Desc.Dimensions));
+	Hash = HashCombine(Hash, GetTypeHash(Desc.NumLayers));
+	Hash = HashCombine(Hash, GetTypeHash(Desc.bContinuousUpdate));
+	for (int32 Layer = 0; Layer < Desc.NumLayers; ++Layer)
+	{
+		Hash = HashCombine(Hash, GetTypeHash(Desc.Format[Layer]));
+	}
+	return Hash;
+}
+
 class FVirtualTexturePhysicalSpace final : public FRenderResource
 {
 public:
-	FVirtualTexturePhysicalSpace(const FVTPhysicalSpaceDescription& InDesc, uint16 InID);
+	FVirtualTexturePhysicalSpace(const FVTPhysicalSpaceDescription& InDesc, uint16 InID, int32 InTileWidthHeight, bool bInEnableResidencyMipMapBias);
 	virtual ~FVirtualTexturePhysicalSpace();
 
 	inline const FVTPhysicalSpaceDescription& GetDescription() const { return Description; }
@@ -102,7 +115,7 @@ public:
 	/** Get frame at which residency tracking last saw over-subscription. */
 	inline uint32 GetLastFrameOversubscribed() const { return LastFrameOversubscribed; }
 	/** Draw residency graph on screen. */
-	void DrawResidencyGraph(class FCanvas* Canvas, FBox2D CanvasPosition);
+	void DrawResidencyGraph(class FCanvas* Canvas, FBox2D CanvasPosition, bool bDrawKey);
 
 private:
 	FVTPhysicalSpaceDescription Description;
@@ -117,8 +130,6 @@ private:
 	uint32 NumRefs;
 	uint16 ID;
 	
-	bool bGpuTextureLimit; // True if the physical size was limited by the maximum GPU texture size
-	
 	bool bEnableResidencyMipMapBias;
 	float ResidencyMipMapBias;
 	uint32 LastFrameOversubscribed;
diff --git a/Engine/Source/Runtime/Renderer/Private/VT/VirtualTextureSystem.cpp b/Engine/Source/Runtime/Renderer/Private/VT/VirtualTextureSystem.cpp
--- a/Engine/Source/Runtime/Renderer/Private/VT/VirtualTextureSystem.cpp
+++ b/Engine/Source/Runtime/Renderer/Private/VT/VirtualTextureSystem.cpp
@@ -19,6 +19,7 @@
 #include "VT/UniqueRequestList.h"
 #include "VT/VirtualTextureFeedback.h"
 #include "VT/VirtualTexturePhysicalSpace.h"
+#include "VT/VirtualTexturePoolConfig.h"
 #include "VT/VirtualTextureScalability.h"
 #include "VT/VirtualTextureSpace.h"
 
@@ -758,19 +759,132 @@ void FVirtualTextureSystem::ReleaseSpace(FVirtualTextureSpace* Space)
 	}
 }
 
+/** 
+ * Description of additional settings to use when initializing a physical space. 
+ * This is filed by GetPoolInitDescription() based on the requested space description and the current config settings.
+ */
+struct FVTPhysicalSpaceInitDescription
+{
+	int32 TileWidthHeight = 0;
+	int32 PoolCount = 1;
+	bool bEnableResidencyMipBias = false;
+};
+
+void GetPoolInitDescription(FVTPhysicalSpaceDescription const& InDesc, FVTPhysicalSpaceInitDescription& OutInitDescription)
+{
+	// Find matching config from ini file
+	FVirtualTextureSpacePoolConfig Config;
+	UVirtualTexturePoolConfig const* PoolConfig = GetDefault<UVirtualTexturePoolConfig>();
+	PoolConfig->FindPoolConfig(InDesc.Format, InDesc.NumLayers, InDesc.TileSize, Config);
+	int32 SizeInMegabyte = Config.SizeInMegabyte;
+
+	// Adjust found config for scaling.
+	const float Scale = Config.bAllowSizeScale ? VirtualTextureScalability::GetPoolSizeScale(Config.ScalabilityGroup) : 1.f;
+	SizeInMegabyte = (int32)(Scale * (float)SizeInMegabyte);
+	if (Scale < 1.f && Config.MinScaledSizeInMegabyte > 0)
+	{
+		SizeInMegabyte = FMath::Max(SizeInMegabyte, Config.MinScaledSizeInMegabyte);
+	}
+	if (Scale > 1.f && Config.MaxScaledSizeInMegabyte > 0)
+	{
+		SizeInMegabyte = FMath::Min(SizeInMegabyte, Config.MaxScaledSizeInMegabyte);
+	}
+	const int32 PoolSizeInBytes = SizeInMegabyte * 1024u * 1024u;
+
+	// Get size of a single tile.
+	const FPixelFormatInfo& FormatInfo = GPixelFormats[InDesc.Format[0]];
+	check(InDesc.TileSize % FormatInfo.BlockSizeX == 0);
+	check(InDesc.TileSize % FormatInfo.BlockSizeY == 0);
+	SIZE_T TileSizeBytes = 0;
+	for (int32 Layer = 0; Layer < InDesc.NumLayers; ++Layer)
+	{
+		TileSizeBytes += CalculateImageBytes(InDesc.TileSize, InDesc.TileSize, 0, InDesc.Format[Layer]);
+	}
+
+	// Calculate final size in tiles.
+	// Loop to find matching pool count if necessary.
+	int32 TileWidthHeight = 0;
+	int32 PoolCount = 1;
+	while (1)
+	{
+		const uint32 MaxTiles = FMath::Max((uint32)(PoolSizeInBytes / (PoolCount * TileSizeBytes)), 1u);
+
+		TileWidthHeight = FMath::FloorToInt(FMath::Sqrt((float)MaxTiles));
+
+		if (TileWidthHeight * InDesc.TileSize > GetMax2DTextureDimension())
+		{
+			// A good option to support extremely large caches would be to allow additional slices in an array here for caches...
+			// Just try to use the maximum texture size for now
+			TileWidthHeight = GetMax2DTextureDimension() / InDesc.TileSize;
+			break;
+		}
+
+		const int32 SplitPhysicalPoolSize = VirtualTextureScalability::GetSplitPhysicalPoolSize();
+		if (SplitPhysicalPoolSize <= 0 || TileWidthHeight <= SplitPhysicalPoolSize)
+		{
+			break;
+		}
+
+		PoolCount++;
+	}
+
+	OutInitDescription.TileWidthHeight = TileWidthHeight;
+	OutInitDescription.PoolCount = PoolCount;
+	OutInitDescription.bEnableResidencyMipBias = Config.bEnableResidencyMipMapBias;
+}
+
+/** Cached version of GetPoolInitDescription() to avoid regularly repeating the heavy work in that function. */
+void GetPoolInitDescription_Cached(FVTPhysicalSpaceDescription const& InDesc, FVTPhysicalSpaceInitDescription& OutInitDescription)
+{
+	check(IsInRenderingThread());
+	static TMap<FVTPhysicalSpaceDescription, FVTPhysicalSpaceInitDescription> Map;
+
+	// Invalidate the cache if any relevant CVar settings change.
+	static uint32 PhysicalPoolSettingsHash = VirtualTextureScalability::GetPhysicalPoolSettingsHash();
+	if (PhysicalPoolSettingsHash != VirtualTextureScalability::GetPhysicalPoolSettingsHash())
+	{
+		Map.Reset();
+	}
+
+	FVTPhysicalSpaceInitDescription* InitDescriptionPtr = Map.Find(InDesc);
+	if (InitDescriptionPtr == nullptr)
+	{
+		GetPoolInitDescription(InDesc, OutInitDescription);
+		Map.Add(InDesc, OutInitDescription);
+	}
+	else
+	{
+		OutInitDescription = *InitDescriptionPtr;
+	}
+}
+
 FVirtualTexturePhysicalSpace* FVirtualTextureSystem::AcquirePhysicalSpace(const FVTPhysicalSpaceDescription& InDesc)
 {
 	LLM_SCOPE(ELLMTag::VirtualTextureSystem);
 
+	// Find matching pools.
+	// We support multiple matching pools to allow for 16bit page table memory optimization.
+	TArray<int32, TInlineAllocator<8>> Matching;
 	for (int32 i = 0; i < PhysicalSpaces.Num(); ++i)
 	{
 		FVirtualTexturePhysicalSpace* PhysicalSpace = PhysicalSpaces[i];
 		if (PhysicalSpace && PhysicalSpace->GetDescription() == InDesc)
 		{
-			return PhysicalSpace;
+			Matching.Add(i);
 		}
 	}
 
+	FVTPhysicalSpaceInitDescription InitDescription;
+	GetPoolInitDescription_Cached(InDesc, InitDescription);
+
+	if (InitDescription.PoolCount <= Matching.Num())
+	{
+		// Randomly select from any pools that exist.
+		int32 RandomIndex = FMath::RandHelper(Matching.Num());
+		return PhysicalSpaces[Matching[RandomIndex]];
+	}
+	
+	// Not reached maximum matching pool count yet so create a new pool.
 	uint32 ID = PhysicalSpaces.Num();
 	check(ID <= 0x0fff);
 
@@ -788,11 +902,12 @@ FVirtualTexturePhysicalSpace* FVirtualTextureSystem::AcquirePhysicalSpace(const
 		PhysicalSpaces.AddZeroed();
 	}
 
-	FVirtualTexturePhysicalSpace* PhysicalSpace = new FVirtualTexturePhysicalSpace(InDesc, ID);
+	FVirtualTexturePhysicalSpace* PhysicalSpace = new FVirtualTexturePhysicalSpace(InDesc, ID, InitDescription.TileWidthHeight, InitDescription.bEnableResidencyMipBias);
 	PhysicalSpaces[ID] = PhysicalSpace;
 
 	INC_MEMORY_STAT_BY(STAT_TotalPhysicalMemory, PhysicalSpace->GetSizeInBytes());
 	BeginInitResource(PhysicalSpace);
+
 	return PhysicalSpace;
 }
 
@@ -2400,11 +2515,11 @@ void FVirtualTextureSystem::DrawResidencyHud(UCanvas* InCanvas, APlayerControlle
 		return;
 	}
 
-	const FIntPoint GraphSize(400, 200);
+	const FIntPoint GraphSize(250, 125);
 	const FIntPoint BorderSize(25, 25);
 	const FIntPoint GraphWithBorderSize = GraphSize + BorderSize * 2;
 
-	const FIntPoint CanvasSize = InCanvas->Canvas->GetParentCanvasSize();
+	const FIntPoint CanvasSize = FIntPoint(InCanvas->ClipX, InCanvas->ClipY);
 	const int32 NumGraphsInRow = FMath::Max(CanvasSize.X / GraphWithBorderSize.X, 1);
 	const int32 CanvasOffsetY = 90;
 
@@ -2428,16 +2543,12 @@ void FVirtualTextureSystem::DrawResidencyHud(UCanvas* InCanvas, APlayerControlle
 				break;
 			}
 
-			PhysicalSpace->DrawResidencyGraph(InCanvas->Canvas, CanvasPosition);
+			const bool bDrawKey = GraphIndex == 0;
+			PhysicalSpace->DrawResidencyGraph(InCanvas->Canvas, CanvasPosition, bDrawKey);
 
 			GraphIndex++;
 		}
 	}
-
-	// Key with same colors as used by FVirtualTexturePhysicalSpace::DrawResidencyGraph()
-	InCanvas->Canvas->DrawShadowedString(25, 320, TEXT("MipMap Bias"), GEngine->GetSmallFont(), FLinearColor(0.1f, 0.8f, 0.1f));
-	InCanvas->Canvas->DrawShadowedString(125, 320, TEXT("Page Residency"), GEngine->GetSmallFont(), FLinearColor(0.8f, 0.1f, 0.1f));
-	InCanvas->Canvas->DrawShadowedString(225, 320, TEXT("LockedPage Residency"), GEngine->GetSmallFont(), FLinearColor(0.8f, 0.8f, 0.1f));
 }
 
 #endif // !UE_BUILD_SHIPPING
