Change hair strands emissive rendering to be content driven.

This allows to remove the hair emissive CVar and avoid any cumbersome workflow. The hair material pass run now in two pass on platform that don't require an RT during PS pass. A first pass for emissive, and a second pass for non-emissive. This allows to keep non-emissive cost low, and have the flexibility of emissive.

diff --git a/Engine/Shaders/Private/HairStrands/HairStrandsMaterialPS.usf b/Engine/Shaders/Private/HairStrands/HairStrandsMaterialPS.usf
--- a/Engine/Shaders/Private/HairStrands/HairStrandsMaterialPS.usf
+++ b/Engine/Shaders/Private/HairStrands/HairStrandsMaterialPS.usf
@@ -7,7 +7,7 @@ void Main(
 	, in uint2 Resolution : DISPATCH_RESOLUTION
 	, in uint TotalNodeCount : DISPATCH_NODE_COUNT
 	, in uint InPrimitiveId : PRIMITIVE_ID
-	#if HAIR_MATERIAL_DEBUG_OR_EMISSIVE_OUTPUT
+	#if HAIR_MATERIAL_EMISSIVE_OUTPUT
 	, out float4 OutColor0 : SV_Target0
 	#endif
 	)
@@ -50,15 +50,8 @@ void Main(
 		}
 	}
 
-	#if HAIR_MATERIAL_DEBUG_OR_EMISSIVE_OUTPUT
-	if (MaterialPassParameters.bOutputEmissive > 0)
-	{
-		OutColor0 = float4(Emissive,1);
-	}
-	else
-	{	
-		OutColor0 = bIsValid ? float4(0, 1, 0, 1) : float4(1, 0, 0, 1);
-	}
+	#if HAIR_MATERIAL_EMISSIVE_OUTPUT
+	OutColor0 = float4(Emissive, 1);
 	#endif
 }
 
diff --git a/Engine/Shaders/Private/HairStrands/HairStrandsVisibilityComposeSubPixelPS.usf b/Engine/Shaders/Private/HairStrands/HairStrandsVisibilityComposeSubPixelPS.usf
--- a/Engine/Shaders/Private/HairStrands/HairStrandsVisibilityComposeSubPixelPS.usf
+++ b/Engine/Shaders/Private/HairStrands/HairStrandsVisibilityComposeSubPixelPS.usf
@@ -107,9 +107,7 @@ void ComposeSamplePS(
 	{
 		const uint LocalOffset = NodeDesc.Offset + SampleIt;
 		const uint2 Coord = uint2(LocalOffset % Resolution, LocalOffset / Resolution);
-		const float4 LightingSample = HairLightingSampleBuffer.Load(uint3(Coord,0));
-		const float4 EmissiveSample = HairStrands.HairEmissiveTexture.Load(uint3(Coord, 0));
-		LocalAccColor += LightingSample.rgb + EmissiveSample.rgb;
+		LocalAccColor += HairLightingSampleBuffer.Load(uint3(Coord, 0));
 	}
 	OutColor.rgb = LocalAccColor * Volumetric.a + Volumetric.rgb;
 
diff --git a/Engine/Source/Runtime/Renderer/Private/ClusteredDeferredShadingPass.cpp b/Engine/Source/Runtime/Renderer/Private/ClusteredDeferredShadingPass.cpp
--- a/Engine/Source/Runtime/Renderer/Private/ClusteredDeferredShadingPass.cpp
+++ b/Engine/Source/Runtime/Renderer/Private/ClusteredDeferredShadingPass.cpp
@@ -177,7 +177,7 @@ static void InternalAddClusteredDeferredShadingPass(
 	PassParameters->RenderTargets[0] = FRenderTargetBinding(SceneTextures.Color.Target, ERenderTargetLoadAction::ELoad);
 	if (bHairStrands)
 	{
-		PassParameters->RenderTargets[0] = FRenderTargetBinding(View.HairStrandsViewData.VisibilityData.SampleLightingBuffer, ERenderTargetLoadAction::ELoad);
+		PassParameters->RenderTargets[0] = FRenderTargetBinding(View.HairStrandsViewData.VisibilityData.SampleLightingTexture, ERenderTargetLoadAction::ELoad);
 	}
 
 	GraphBuilder.AddPass(
diff --git a/Engine/Source/Runtime/Renderer/Private/HairStrands/HairStrandsComposition.cpp b/Engine/Source/Runtime/Renderer/Private/HairStrands/HairStrandsComposition.cpp
--- a/Engine/Source/Runtime/Renderer/Private/HairStrands/HairStrandsComposition.cpp
+++ b/Engine/Source/Runtime/Renderer/Private/HairStrands/HairStrandsComposition.cpp
@@ -200,14 +200,14 @@ static void AddHairVisibilityComposeSamplePass(
 	FRDGTextureRef& OutColorTexture,
 	FRDGTextureRef& OutDepthTexture)
 {
-	check(VisibilityData.SampleLightingBuffer);
+	check(VisibilityData.SampleLightingTexture);
 	const bool bDOFEnable = HairDOFDepthTexture != nullptr ? 1 : 0;
 
 	TRDGUniformBufferRef<FFogUniformParameters> FogBuffer = CreateFogUniformBuffer(GraphBuilder, View);
 
 	FHairVisibilityComposeSamplePS::FParameters* Parameters = GraphBuilder.AllocParameters<FHairVisibilityComposeSamplePS::FParameters>();
 	Parameters->bComposeDofDepth = bDOFEnable ? 1 : 0;
-	Parameters->HairLightingSampleBuffer = VisibilityData.SampleLightingBuffer;
+	Parameters->HairLightingSampleBuffer = VisibilityData.SampleLightingTexture;
 	Parameters->HairDOFDepthTexture = bDOFEnable ? HairDOFDepthTexture : GSystemTextures.GetBlackDummy(GraphBuilder);
 	Parameters->OutputResolution = OutColorTexture->Desc.Extent;
 	Parameters->ViewUniformBuffer = View.ViewUniformBuffer;
@@ -270,7 +270,7 @@ static FRDGTextureRef AddHairDOFDepthPass(
 	const FRDGTextureRef& InColorTexture,
 	const FRDGTextureRef& InDepthTexture)
 {
-	check(VisibilityData.SampleLightingBuffer);
+	check(VisibilityData.SampleLightingTexture);
 	FIntPoint OutputResolution = InColorTexture->Desc.Extent;
 
 	FRDGTextureRef OutDOFDepthTexture = nullptr;
@@ -281,7 +281,7 @@ static FRDGTextureRef AddHairDOFDepthPass(
 
 	FHairDOFDepthPS::FParameters* Parameters = GraphBuilder.AllocParameters<FHairDOFDepthPS::FParameters>();
 	Parameters->HairStrands = View.HairStrandsViewData.UniformBuffer;
-	Parameters->HairLightingSampleBuffer = VisibilityData.SampleLightingBuffer;
+	Parameters->HairLightingSampleBuffer = VisibilityData.SampleLightingTexture;
 	Parameters->SceneColorTexture = InColorTexture;
 	Parameters->SceneDepthTexture = InDepthTexture;
 	Parameters->ViewUniformBuffer = View.ViewUniformBuffer;
diff --git a/Engine/Source/Runtime/Renderer/Private/HairStrands/HairStrandsData.h b/Engine/Source/Runtime/Renderer/Private/HairStrands/HairStrandsData.h
--- a/Engine/Source/Runtime/Renderer/Private/HairStrands/HairStrandsData.h
+++ b/Engine/Source/Runtime/Renderer/Private/HairStrands/HairStrandsData.h
@@ -24,14 +24,12 @@ BEGIN_GLOBAL_SHADER_PARAMETER_STRUCT(FHairStrandsViewUniformParameters, )
 	SHADER_PARAMETER(float, HairDualScatteringRoughnessOverride)						// Override the roughness used for dual scattering (for hack/test purpose only)
 	SHADER_PARAMETER(FIntPoint, HairSampleViewportResolution)							// Maximum viewport resolution of the sample space
 	SHADER_PARAMETER(uint32, bHairTileValid)											// True if tile data are valid
-	SHADER_PARAMETER(uint32, bHasEmissive)												// True if emissive is enabled
 	SHADER_PARAMETER_RDG_TEXTURE(Texture2D<float>, HairCoverageTexture)					// Hair pixel's coverage
 	SHADER_PARAMETER_RDG_TEXTURE(Texture2D, HairOnlyDepthTexture)						// Depth texture containing only hair depth (not strongly typed for legacy shader code reason)
 	SHADER_PARAMETER_RDG_TEXTURE(Texture2D, HairOnlyDepthClosestHZBTexture)				// HZB closest depth texture containing only hair depth (not strongly typed for legacy shader code reason)
 	SHADER_PARAMETER_RDG_TEXTURE(Texture2D, HairOnlyDepthFurthestHZBTexture)			// HZB furthest depth texture containing only hair depth (not strongly typed for legacy shader code reason)
 	SHADER_PARAMETER_RDG_TEXTURE(Texture2D<uint>, HairSampleOffset)						// Offset & count, for accessing pixel's samples, based on screen pixel position
 	SHADER_PARAMETER_RDG_TEXTURE(Texture2D<uint>, HairSampleCount)						// Total count of hair sample, in sample space
-	SHADER_PARAMETER_RDG_TEXTURE(Texture2D<float4>, HairEmissiveTexture)				// Emissive texture
 	SHADER_PARAMETER_RDG_BUFFER_SRV(StructuredBuffer<FPackedHairSample>, HairSampleData)// Sample data (coverage, tangent, base color, ...), in sample space // HAIRSTRANDS_TODO: change this to be a uint4 so that we don't have to include the type for generated contant buffer
 	SHADER_PARAMETER_RDG_BUFFER_SRV(Buffer<uint2>, HairSampleCoords)					// Screen pixel coordinate of each sample, in sample space
 	SHADER_PARAMETER_RDG_BUFFER_SRV(Buffer<uint2>, HairTileData)						// Tile coords (RG16F)
@@ -77,7 +75,6 @@ struct FHairStrandsVisibilityData
 	FRDGTextureRef CoverageTexture = nullptr;
 	FRDGTextureRef ViewHairCountTexture = nullptr;
 	FRDGTextureRef ViewHairCountUintTexture = nullptr;
-	FRDGTextureRef EmissiveTexture = nullptr;
 	FRDGTextureRef HairOnlyDepthTexture = nullptr;
 
 	FRDGTextureRef HairOnlyDepthClosestHZBTexture = nullptr;
@@ -103,7 +100,7 @@ struct FHairStrandsVisibilityData
 	// Allocated conservatively
 	// User indirect dispatch for accumulating contribution
 	FIntPoint SampleLightingViewportResolution;
-	FRDGTextureRef SampleLightingBuffer = nullptr;
+	FRDGTextureRef SampleLightingTexture = nullptr;
 };
 
 ////////////////////////////////////////////////////////////////////////////////////
diff --git a/Engine/Source/Runtime/Renderer/Private/HairStrands/HairStrandsEnvironment.cpp b/Engine/Source/Runtime/Renderer/Private/HairStrands/HairStrandsEnvironment.cpp
--- a/Engine/Source/Runtime/Renderer/Private/HairStrands/HairStrandsEnvironment.cpp
+++ b/Engine/Source/Runtime/Renderer/Private/HairStrands/HairStrandsEnvironment.cpp
@@ -417,8 +417,8 @@ static void AddHairStrandsEnvironmentLightingPassPS(
 	TShaderMapRef<FHairEnvironmentLightingVS> VertexShader(View.ShaderMap);
 	TShaderMapRef<FHairEnvironmentLightingPS> PixelShader(View.ShaderMap, PermutationVector);
 
-	check(VisibilityData.SampleLightingBuffer);
-	ParametersPS->RenderTargets[0] = FRenderTargetBinding(VisibilityData.SampleLightingBuffer, ERenderTargetLoadAction::ELoad);
+	check(VisibilityData.SampleLightingTexture);
+	ParametersPS->RenderTargets[0] = FRenderTargetBinding(VisibilityData.SampleLightingTexture, ERenderTargetLoadAction::ELoad);
 
 	GraphBuilder.AddPass(
 		RDG_EVENT_NAME("HairEnvLightingPS(%s)", LightingType == EHairLightingSourceType::SceneColor ? TEXT("SceneScatter") : (LightingType == EHairLightingSourceType::Lumen ? TEXT("Lumen") : TEXT("ReflectionProbe")) ),
diff --git a/Engine/Source/Runtime/Renderer/Private/HairStrands/HairStrandsRendering.cpp b/Engine/Source/Runtime/Renderer/Private/HairStrands/HairStrandsRendering.cpp
--- a/Engine/Source/Runtime/Renderer/Private/HairStrands/HairStrandsRendering.cpp
+++ b/Engine/Source/Runtime/Renderer/Private/HairStrands/HairStrandsRendering.cpp
@@ -24,17 +24,6 @@ static TRDGUniformBufferRef<FHairStrandsViewUniformParameters> InternalCreateHai
 		Parameters->HairSampleViewportResolution = In->SampleLightingViewportResolution;
 		Parameters->MaxSamplePerPixelCount = In->MaxSampleCount;
 
-		if (In->EmissiveTexture)
-		{
-			Parameters->HairEmissiveTexture = In->EmissiveTexture;
-			Parameters->bHasEmissive = 1u;
-		}
-		else
-		{
-			Parameters->HairEmissiveTexture = GSystemTextures.GetBlackDummy(GraphBuilder);
-			Parameters->bHasEmissive = 0u;
-			
-		}
 		if (In->TileData.IsValid())
 		{
 			Parameters->HairTileData = In->TileData.TileDataSRV;
@@ -82,8 +71,6 @@ static TRDGUniformBufferRef<FHairStrandsViewUniformParameters> InternalCreateHai
 		Parameters->HairSampleCoords = DummyBufferRG16SRV;
 		Parameters->HairSampleData	 = GraphBuilder.CreateSRV(DummyNodeBuffer);
 		Parameters->HairSampleViewportResolution = FIntPoint(0, 0);
-		Parameters->HairEmissiveTexture = BlackTexture;
-		Parameters->bHasEmissive = 0u;
 		Parameters->MaxSamplePerPixelCount = 0u;
 
 		Parameters->HairTileData = DummyBufferR32SRV;
diff --git a/Engine/Source/Runtime/Renderer/Private/HairStrands/HairStrandsVisibility.cpp b/Engine/Source/Runtime/Renderer/Private/HairStrands/HairStrandsVisibility.cpp
--- a/Engine/Source/Runtime/Renderer/Private/HairStrands/HairStrandsVisibility.cpp
+++ b/Engine/Source/Runtime/Renderer/Private/HairStrands/HairStrandsVisibility.cpp
@@ -74,9 +74,6 @@ static FAutoConsoleVariableRef CVarHairStrandsSortHairSampleByDepth(TEXT("r.Hair
 static int32 GHairStrandsHairCountToTransmittance = 0;
 static FAutoConsoleVariableRef CVarHairStrandsHairCountToTransmittance(TEXT("r.HairStrands.Visibility.UseCoverageMappping"), GHairStrandsHairCountToTransmittance, TEXT("Use hair count to coverage transfer function."));
 
-static int32 GHairStrandsVisibility_OutputEmissiveData = 0;
-static FAutoConsoleVariableRef CVarHairStrandsVisibility_OutputEmissiveData(TEXT("r.HairStrands.Visibility.Emissive"), GHairStrandsVisibility_OutputEmissiveData, TEXT("Enable emissive data during the material pass."));
-
 static int32 GHairStrandsDebugPPLL = 0;
 static FAutoConsoleVariableRef CVarHairStrandsDebugPPLL(TEXT("r.HairStrands.Visibility.PPLL.Debug"), GHairStrandsDebugPPLL, TEXT("Draw debug per pixel light list rendering."));
 
@@ -319,6 +316,16 @@ class FHairLightSampleClearPS : public FGlobalShader
 		RENDER_TARGET_BINDING_SLOTS()
 	END_SHADER_PARAMETER_STRUCT()
 
+	static EPixelFormat GetHairLightSampleFormat()
+	{
+		EPixelFormat Format = PF_FloatRGBA;
+		if (GHairStrandsLightSampleFormat > 0 && GPixelFormats[PF_FloatR11G11B10].Supported)
+		{
+			Format = PF_FloatR11G11B10;
+		}
+		return Format;
+	}
+
 	static bool ShouldCompilePermutation(const FGlobalShaderPermutationParameters& Parameters) { return IsHairStrandsSupported(EHairStrandsShaderType::Strands, Parameters.Platform); }
 	static void ModifyCompilationEnvironment(const FGlobalShaderPermutationParameters& Parameters, FShaderCompilerEnvironment& OutEnvironment)
 	{
@@ -346,11 +353,7 @@ static FRDGTextureRef AddClearLightSamplePass(
 	const uint32 MaxNodeCount,
 	const FRDGTextureRef NodeCounter)
 {	
-	EPixelFormat Format = PF_FloatRGBA;
-	if (GHairStrandsLightSampleFormat > 0 && GPixelFormats[PF_FloatR11G11B10].Supported)
-	{
-		Format = PF_FloatR11G11B10;
-	}
+	const EPixelFormat Format = FHairLightSampleClearPS::GetHairLightSampleFormat();
 
 	const uint32 SampleTextureResolution = FMath::CeilToInt(FMath::Sqrt(static_cast<float>(MaxNodeCount)));
 	FRDGTextureDesc Desc = FRDGTextureDesc::Create2D(FIntPoint(SampleTextureResolution, SampleTextureResolution), Format, FClearValueBinding::Black, TexCreate_UAV | TexCreate_ShaderResource | TexCreate_RenderTargetable);
@@ -432,12 +435,6 @@ IMPLEMENT_MATERIAL_SHADER_TYPE(, FHairMaterialVS, TEXT("/Engine/Private/HairStra
 
 /////////////////////////////////////////////////////////////////////////////////////////
 
-
-static bool IsHairStrandsEmissiveEnable()
-{
-	return GHairStrandsVisibility_OutputEmissiveData > 0;
-}
-
 class FHairMaterialShaderElementData : public FMeshMaterialShaderElementData
 {
 public:
@@ -482,8 +479,12 @@ public:
 	{
 		FMaterialShader::ModifyCompilationEnvironment(Parameters, OutEnvironment);
 		const bool bPlatformRequireRenderTarget = IsPlatformRequiringRenderTargetForMaterialPass(Parameters.Platform);
-		OutEnvironment.SetDefine(TEXT("HAIR_MATERIAL_DEBUG_OR_EMISSIVE_OUTPUT"), (IsHairStrandsEmissiveEnable() || bPlatformRequireRenderTarget) ? 1 : 0);
+		const bool bHasEmissiveConnected = Parameters.MaterialParameters.bHasEmissiveColorConnected;
+		OutEnvironment.SetDefine(TEXT("HAIR_MATERIAL_EMISSIVE_OUTPUT"), (bHasEmissiveConnected || bPlatformRequireRenderTarget) ? 1 : 0);
 		OutEnvironment.SetDefine(TEXT("HAIRSTRANDS_HAS_NORMAL_CONNECTED"), Parameters.MaterialParameters.bHasNormalConnected ? 1 : 0);
+
+		const EPixelFormat Format = FHairLightSampleClearPS::GetHairLightSampleFormat();
+		OutEnvironment.SetRenderTargetOutputFormat(0, Format);
 	}
 
 	void GetShaderBindings(
@@ -689,7 +690,6 @@ BEGIN_GLOBAL_SHADER_PARAMETER_STRUCT(FVisibilityMaterialPassUniformParameters, )
 	SHADER_PARAMETER(uint32, MaxSampleCount)
 	SHADER_PARAMETER(uint32, NodeGroupSize)
 	SHADER_PARAMETER(uint32, bUpdateSampleCoverage)
-	SHADER_PARAMETER(uint32, bOutputEmissive)
 	SHADER_PARAMETER_RDG_TEXTURE(Texture2D<uint>, NodeIndex)
 	SHADER_PARAMETER_RDG_TEXTURE(Texture2D<uint>, TotalNodeCounter)
 	SHADER_PARAMETER_RDG_BUFFER_SRV(Buffer<uint2>, NodeCoord)
@@ -764,7 +764,7 @@ struct FMaterialPassOutput
 	static const EPixelFormat VelocityFormat = PF_G16R16;
 	FRDGBufferRef NodeData = nullptr;
 	FRDGBufferRef NodeVelocity = nullptr;
-	FRDGTextureRef EmissiveTexture = nullptr;
+	FRDGTextureRef SampleLightingTexture = nullptr;
 };
 
 static FMaterialPassOutput AddHairMaterialPass(
@@ -787,66 +787,84 @@ static FMaterialPassOutput AddHairMaterialPass(
 	const uint32 MaxNodeCount = CompactNodeVis->Desc.NumElements;
 
 	FMaterialPassOutput Output;
-	Output.NodeData		= GraphBuilder.CreateBuffer(FRDGBufferDesc::CreateStructuredDesc(sizeof(HairStrandsVisibilityInternal::NodeData), MaxNodeCount), TEXT("Hair.CompactNodeData"));
-	Output.NodeVelocity = GraphBuilder.CreateBuffer(FRDGBufferDesc::CreateBufferDesc(4, CompactNodeVis->Desc.NumElements), TEXT("Hair.CompactNodeVelocity"));
+	Output.NodeData				 = GraphBuilder.CreateBuffer(FRDGBufferDesc::CreateStructuredDesc(sizeof(HairStrandsVisibilityInternal::NodeData), MaxNodeCount), TEXT("Hair.CompactNodeData"));
+	Output.NodeVelocity			 = GraphBuilder.CreateBuffer(FRDGBufferDesc::CreateBufferDesc(4, CompactNodeVis->Desc.NumElements), TEXT("Hair.CompactNodeVelocity"));
+	Output.SampleLightingTexture = AddClearLightSamplePass(GraphBuilder, ViewInfo, MaxNodeCount, CompactNodeCounter);
 
 	const uint32 ResolutionDim = FMath::CeilToInt(FMath::Sqrt(static_cast<float>(MaxNodeCount)));
 	const FIntPoint Resolution(ResolutionDim, ResolutionDim);
 
-	const bool bOutputEmissive = IsHairStrandsEmissiveEnable();
-	const bool bIsPlatformRequireRenderTarget = IsPlatformRequiringRenderTargetForMaterialPass(Scene->GetShaderPlatform()) || GRHIRequiresRenderTargetForPixelShaderUAVs;
-
-	// Add resources reference to the pass parameters, in order to get the resource lifetime extended to this pass
-	FVisibilityMaterialPassParameters* PassParameters = GraphBuilder.AllocParameters<FVisibilityMaterialPassParameters>();
-
+	enum class EHairMaterialPassFilter : uint8
 	{
-		FVisibilityMaterialPassUniformParameters* UniformParameters = GraphBuilder.AllocParameters<FVisibilityMaterialPassUniformParameters>();
-
-		UniformParameters->bUpdateSampleCoverage = bUpdateSampleCoverage ? 1 : 0;
-		UniformParameters->bOutputEmissive = bOutputEmissive ? 1 : 0;
-		UniformParameters->MaxResolution = Resolution;
-		UniformParameters->NodeGroupSize = NodeGroupSize;
-		UniformParameters->MaxSampleCount = MaxNodeCount;
-		UniformParameters->TotalNodeCounter = CompactNodeCounter;
-		UniformParameters->NodeIndex = CompactNodeIndex;
-		UniformParameters->NodeVis = GraphBuilder.CreateSRV(CompactNodeVis);
-		UniformParameters->NodeCoord = GraphBuilder.CreateSRV(CompactNodeCoord, FHairStrandsVisibilityData::NodeCoordFormat);
-		UniformParameters->IndirectArgs = GraphBuilder.CreateSRV(IndirectArgBuffer);
-		UniformParameters->OutNodeData = GraphBuilder.CreateUAV(FRDGBufferUAVDesc(Output.NodeData));
-		UniformParameters->OutNodeVelocity = GraphBuilder.CreateUAV(FRDGBufferUAVDesc(Output.NodeVelocity, FMaterialPassOutput::VelocityFormat));
+		All,
+		EmissiveOnly,
+		NonEmissiveOnly
+	};
 
-		PassParameters->UniformBuffer = GraphBuilder.CreateUniformBuffer(UniformParameters);
-	}
+	const ERHIFeatureLevel::Type FeatureLevel = ViewInfo->FeatureLevel;
 
+	// Find among the mesh batch, if any of them emit emissive data
+	bool bHasEmissiveMaterial = false;
+	for (const FHairStrandsMacroGroupData& MacroGroupData : MacroGroupDatas)
 	{
-		const bool bEnableMSAA = false;
-		SetUpViewHairRenderInfo(*ViewInfo, bEnableMSAA, ViewInfo->CachedViewUniformShaderParameters->HairRenderInfo, ViewInfo->CachedViewUniformShaderParameters->HairRenderInfoBits, ViewInfo->CachedViewUniformShaderParameters->HairComponents);
-		PassParameters->View = TUniformBufferRef<FViewUniformShaderParameters>::CreateUniformBufferImmediate(*ViewInfo->CachedViewUniformShaderParameters, UniformBuffer_SingleFrame);
+		for (const FHairStrandsMacroGroupData::PrimitiveInfo& PrimitiveInfo : MacroGroupData.PrimitivesInfos)
+		{
+			const FMeshBatch& MeshBatch = *PrimitiveInfo.Mesh;
+			if (MeshBatch.MaterialRenderProxy->GetIncompleteMaterialWithFallback(FeatureLevel).HasEmissiveColorConnected())
+			{
+				bHasEmissiveMaterial = true;
+				break;
+			}
+		}
 	}
 
-	// For debug purpose only
-	if (bOutputEmissive)
+	// Generic material pass dispatch
+	auto MaterialPass = [&](FRDGTextureRef RenderTarget, EHairMaterialPassFilter Filter)
 	{
-		FRDGTextureDesc OutputDesc = FRDGTextureDesc::Create2D(Resolution, PF_FloatRGBA, FClearValueBinding::Black, TexCreate_RenderTargetable);
-		Output.EmissiveTexture = GraphBuilder.CreateTexture(OutputDesc, TEXT("Hair.MaterialEmissiveOutput"));
-		PassParameters->RenderTargets[0] = FRenderTargetBinding(Output.EmissiveTexture, ERenderTargetLoadAction::EClear, 0);
-		
-	}
-	else if (bIsPlatformRequireRenderTarget)
-	{
-		FRDGTextureDesc OutputDesc = FRDGTextureDesc::Create2D(Resolution, PF_FloatRGBA, FClearValueBinding::Black, TexCreate_RenderTargetable);
-		FRDGTextureRef OutDummyTexture0 = GraphBuilder.CreateTexture(OutputDesc, TEXT("Hair.MaterialDummyOutput"));
-		PassParameters->RenderTargets[0] = FRenderTargetBinding(OutDummyTexture0, ERenderTargetLoadAction::EClear, 0);
-	}
-	
+		// Add resources reference to the pass parameters, in order to get the resource lifetime extended to this pass
+		FVisibilityMaterialPassParameters* PassParameters = GraphBuilder.AllocParameters<FVisibilityMaterialPassParameters>();
 
-	AddSimpleMeshPass(GraphBuilder, PassParameters, Scene, *ViewInfo, &InstanceCullingManager, RDG_EVENT_NAME("HairStrandsMaterialPass"), FIntRect(0, 0, Resolution.X, Resolution.Y),
-		[PassParameters, Scene = Scene, ViewInfo, &MacroGroupDatas, MaxNodeCount, Resolution, NodeGroupSize, bUpdateSampleCoverage, bOutputEmissive](FDynamicPassMeshDrawListContext* ShadowContext)
-	{
-		FMeshPassProcessorRenderState DrawRenderState;
+		{
+			FVisibilityMaterialPassUniformParameters* UniformParameters = GraphBuilder.AllocParameters<FVisibilityMaterialPassUniformParameters>();
+
+			UniformParameters->bUpdateSampleCoverage = bUpdateSampleCoverage ? 1 : 0;
+			UniformParameters->MaxResolution = Resolution;
+			UniformParameters->NodeGroupSize = NodeGroupSize;
+			UniformParameters->MaxSampleCount = MaxNodeCount;
+			UniformParameters->TotalNodeCounter = CompactNodeCounter;
+			UniformParameters->NodeIndex = CompactNodeIndex;
+			UniformParameters->NodeVis = GraphBuilder.CreateSRV(CompactNodeVis);
+			UniformParameters->NodeCoord = GraphBuilder.CreateSRV(CompactNodeCoord, FHairStrandsVisibilityData::NodeCoordFormat);
+			UniformParameters->IndirectArgs = GraphBuilder.CreateSRV(IndirectArgBuffer);
+			UniformParameters->OutNodeData = GraphBuilder.CreateUAV(FRDGBufferUAVDesc(Output.NodeData));
+			UniformParameters->OutNodeVelocity = GraphBuilder.CreateUAV(FRDGBufferUAVDesc(Output.NodeVelocity, FMaterialPassOutput::VelocityFormat));
+
+			PassParameters->UniformBuffer = GraphBuilder.CreateUniformBuffer(UniformParameters);
+		}
 
 		{
-			if (bOutputEmissive)
+			const bool bEnableMSAA = false;
+			SetUpViewHairRenderInfo(*ViewInfo, bEnableMSAA, ViewInfo->CachedViewUniformShaderParameters->HairRenderInfo, ViewInfo->CachedViewUniformShaderParameters->HairRenderInfoBits, ViewInfo->CachedViewUniformShaderParameters->HairComponents);
+			PassParameters->View = TUniformBufferRef<FViewUniformShaderParameters>::CreateUniformBufferImmediate(*ViewInfo->CachedViewUniformShaderParameters, UniformBuffer_SingleFrame);
+		}
+
+		if (RenderTarget)
+		{
+			PassParameters->RenderTargets[0] = FRenderTargetBinding(RenderTarget, ERenderTargetLoadAction::EClear, 0);
+		}
+
+		AddSimpleMeshPass(
+			GraphBuilder, 
+			PassParameters, 
+			Scene, 
+			*ViewInfo, 
+			&InstanceCullingManager, 
+			RDG_EVENT_NAME("HairStrandsMaterialPass(Emissive=%s)", Filter == EHairMaterialPassFilter::All ? TEXT("On/Off") : (Filter == EHairMaterialPassFilter::EmissiveOnly ? TEXT("On") : TEXT("Off"))),
+			FIntRect(0, 0, Resolution.X, Resolution.Y),
+		[PassParameters, Scene = Scene, ViewInfo, &MacroGroupDatas, MaxNodeCount, Resolution, NodeGroupSize, bUpdateSampleCoverage, Filter, FeatureLevel](FDynamicPassMeshDrawListContext* ShadowContext)
+		{
+			FMeshPassProcessorRenderState DrawRenderState;
+			if (Filter == EHairMaterialPassFilter::All || Filter == EHairMaterialPassFilter::EmissiveOnly)
 			{
 				DrawRenderState.SetBlendState(TStaticBlendState<CW_RGBA, BO_Add, BF_One, BF_One, BO_Add, BF_One, BF_Zero>::GetRHI());
 			}
@@ -863,11 +881,42 @@ static FMaterialPassOutput AddHairMaterialPass(
 				{
 					const FMeshBatch& MeshBatch = *PrimitiveInfo.Mesh;
 					const uint64 BatchElementMask = ~0ull;
-					MeshProcessor.AddMeshBatch(MeshBatch, BatchElementMask, PrimitiveInfo.PrimitiveSceneProxy, -1, MacroGroupData.MacroGroupId, PrimitiveInfo.MaterialId);
+					bool bIsCompatible = true;
+					if (Filter != EHairMaterialPassFilter::All)
+					{
+						const bool bHasEmissive = MeshBatch.MaterialRenderProxy->GetIncompleteMaterialWithFallback(FeatureLevel).HasEmissiveColorConnected();
+						bIsCompatible = (bHasEmissive && Filter == EHairMaterialPassFilter::EmissiveOnly) || (!bHasEmissive && Filter == EHairMaterialPassFilter::NonEmissiveOnly);
+					}
+
+					if (bIsCompatible)
+					{
+						MeshProcessor.AddMeshBatch(MeshBatch, BatchElementMask, PrimitiveInfo.PrimitiveSceneProxy, -1, MacroGroupData.MacroGroupId, PrimitiveInfo.MaterialId);
+					}
 				}
 			}
-		}
-	});
+		});
+	};
+
+
+	const bool bIsPlatformRequireRenderTarget = IsPlatformRequiringRenderTargetForMaterialPass(Scene->GetShaderPlatform()) || GRHIRequiresRenderTargetForPixelShaderUAVs;
+
+	// Output:
+	// 1. Single pass: when the platform require an RT as output, render both emissive & non-emissive in a single pass
+	// 2. Two passes : one pass for emissive material with an RT, one pass for regular/non-emissive material without an RT
+	// 3. Single pass: when there is no emissive material, and platform does not require an RT
+	if (bIsPlatformRequireRenderTarget)
+	{
+		MaterialPass(Output.SampleLightingTexture, EHairMaterialPassFilter::All);
+	}
+	else if (bHasEmissiveMaterial)
+	{
+		MaterialPass(Output.SampleLightingTexture, EHairMaterialPassFilter::EmissiveOnly);
+		MaterialPass(nullptr, EHairMaterialPassFilter::NonEmissiveOnly);
+	}
+	else
+	{
+		MaterialPass(nullptr, EHairMaterialPassFilter::NonEmissiveOnly);
+	}
 
 	return Output;
 }
@@ -3102,11 +3151,8 @@ void RenderHairStrandsVisibilityBuffer(
 
 						CompactNodeData = PassOutput.NodeData;
 
-						// Allocate buffer for storing all the light samples
-						FRDGTextureRef SampleLightingBuffer = AddClearLightSamplePass(GraphBuilder, &View, VisibilityData.MaxNodeCount, NodeCounter);
-						VisibilityData.SampleLightingViewportResolution = SampleLightingBuffer->Desc.Extent;
-
-						VisibilityData.SampleLightingBuffer = SampleLightingBuffer;
+						VisibilityData.SampleLightingViewportResolution = PassOutput.SampleLightingTexture->Desc.Extent;
+						VisibilityData.SampleLightingTexture = PassOutput.SampleLightingTexture;
 						VisibilityData.NodeIndex = CompactNodeIndex;
 						VisibilityData.CoverageTexture = CoverageTexture;
 						VisibilityData.HairOnlyDepthTexture = HairOnlyDepthTexture;
@@ -3115,7 +3161,6 @@ void RenderHairStrandsVisibilityBuffer(
 						VisibilityData.NodeIndirectArg = IndirectArgsBuffer;
 						VisibilityData.NodeCount = NodeCounter;
 						VisibilityData.ResolveMaskTexture = ResolveMaskTexture;	
-						VisibilityData.EmissiveTexture = PassOutput.EmissiveTexture;
 					}
 
 					// For fully covered pixels, write: 
@@ -3225,7 +3270,6 @@ void RenderHairStrandsVisibilityBuffer(
 					FRDGBufferRef CompactNodeCoord;
 					FRDGBufferRef IndirectArgsBuffer;
 					FRDGTextureRef ResolveMaskTexture = nullptr;
-					FRDGTextureRef EmissiveTexture = nullptr;
 					AddHairVisibilityPrimitiveIdCompactionPass(
 						false, // bUsePPLL
 						GraphBuilder,
@@ -3285,23 +3329,19 @@ void RenderHairStrandsVisibilityBuffer(
 						}
 
 						CompactNodeData = PassOutput.NodeData;
-						EmissiveTexture = PassOutput.EmissiveTexture;
+
+						VisibilityData.SampleLightingViewportResolution = PassOutput .SampleLightingTexture->Desc.Extent;
+						VisibilityData.SampleLightingTexture			= PassOutput.SampleLightingTexture;
 					}
 
-					// Allocate buffer for storing all the light samples
-					FRDGTextureRef SampleLightingBuffer = AddClearLightSamplePass(GraphBuilder, &View, VisibilityData.MaxNodeCount, NodeCounter);
-					VisibilityData.SampleLightingViewportResolution = SampleLightingBuffer->Desc.Extent;
-
-					 VisibilityData.SampleLightingBuffer	= SampleLightingBuffer;
-					 VisibilityData.NodeIndex				= CompactNodeIndex;
-					 VisibilityData.CoverageTexture			= CoverageTexture;
-					 VisibilityData.HairOnlyDepthTexture	= HairOnlyDepthTexture;
-					 VisibilityData.NodeData				= CompactNodeData;
-					 VisibilityData.NodeCoord				= CompactNodeCoord;
-					 VisibilityData.NodeIndirectArg			= IndirectArgsBuffer;
-					 VisibilityData.NodeCount				= NodeCounter;
-					 VisibilityData.ResolveMaskTexture		= ResolveMaskTexture;
-					 VisibilityData.EmissiveTexture			= EmissiveTexture;
+					VisibilityData.NodeIndex			= CompactNodeIndex;
+					VisibilityData.CoverageTexture		= CoverageTexture;
+					VisibilityData.HairOnlyDepthTexture	= HairOnlyDepthTexture;
+					VisibilityData.NodeData				= CompactNodeData;
+					VisibilityData.NodeCoord			= CompactNodeCoord;
+					VisibilityData.NodeIndirectArg		= IndirectArgsBuffer;
+					VisibilityData.NodeCount			= NodeCounter;
+					VisibilityData.ResolveMaskTexture	= ResolveMaskTexture;
 
 					// View transmittance depth test needs to happen before the scene depth is patched with the hair depth (for fully-covered-by-hair pixels)
 					if (ViewTransmittance.HairCountTexture)
@@ -3435,9 +3475,8 @@ void RenderHairStrandsVisibilityBuffer(
 					HairOnlyDepthTexture);
 
 				// Allocate buffer for storing all the light samples
-				FRDGTextureRef SampleLightingBuffer = AddClearLightSamplePass(GraphBuilder, &View, VisibilityData.MaxNodeCount, NodeCounter);
-				VisibilityData.SampleLightingViewportResolution = SampleLightingBuffer->Desc.Extent;
-				VisibilityData.SampleLightingBuffer = SampleLightingBuffer;
+				VisibilityData.SampleLightingTexture = AddClearLightSamplePass(GraphBuilder, &View, VisibilityData.MaxNodeCount, NodeCounter);
+				VisibilityData.SampleLightingViewportResolution = VisibilityData.SampleLightingTexture->Desc.Extent;
 
 			#if WITH_EDITOR
 				// Extract texture for debug visualization
diff --git a/Engine/Source/Runtime/Renderer/Private/LightRendering.cpp b/Engine/Source/Runtime/Renderer/Private/LightRendering.cpp
--- a/Engine/Source/Runtime/Renderer/Private/LightRendering.cpp
+++ b/Engine/Source/Runtime/Renderer/Private/LightRendering.cpp
@@ -2499,7 +2499,7 @@ void FDeferredShadingSceneRenderer::RenderLightForHair(
 		}
 
 		const FHairStrandsVisibilityData& HairVisibilityData = View.HairStrandsViewData.VisibilityData;
-		if (!HairVisibilityData.SampleLightingBuffer)
+		if (!HairVisibilityData.SampleLightingTexture)
 		{
 			return;
 		}
@@ -2507,7 +2507,7 @@ void FDeferredShadingSceneRenderer::RenderLightForHair(
 		FRenderLightForHairParameters* PassParameters = GraphBuilder.AllocParameters<FRenderLightForHairParameters>();
 		GetRenderLightParameters(
 			View,
-			HairVisibilityData.SampleLightingBuffer, 
+			HairVisibilityData.SampleLightingTexture, 
 			nullptr, 
 			SceneTexturesUniformBuffer, 
 			HairStrands::BindHairStrandsViewUniformParameters(View), 
