Fix GPU-Scene instance culling occlusion when nanite is not used.

diff --git a/Engine/Source/Runtime/Renderer/Private/DeferredShadingRenderer.cpp b/Engine/Source/Runtime/Renderer/Private/DeferredShadingRenderer.cpp
--- a/Engine/Source/Runtime/Renderer/Private/DeferredShadingRenderer.cpp
+++ b/Engine/Source/Runtime/Renderer/Private/DeferredShadingRenderer.cpp
@@ -429,7 +429,7 @@ bool FDeferredShadingSceneRenderer::RenderHzb(FRDGBuilder& GraphBuilder, FRDGTex
 				// Extract furthest HZB texture.
 				if (View.ViewState)
 				{
-					if (IsNaniteEnabled())
+					if (IsNaniteEnabled() || FInstanceCullingContext::IsOcclusionCullingEnabled())
 					{
 						GraphBuilder.QueueTextureExtraction(FurthestHZBTexture, &View.ViewState->PrevFrameViewInfo.HZB);
 					}
diff --git a/Engine/Source/Runtime/Renderer/Private/InstanceCulling/InstanceCullingContext.cpp b/Engine/Source/Runtime/Renderer/Private/InstanceCulling/InstanceCullingContext.cpp
--- a/Engine/Source/Runtime/Renderer/Private/InstanceCulling/InstanceCullingContext.cpp
+++ b/Engine/Source/Runtime/Renderer/Private/InstanceCulling/InstanceCullingContext.cpp
@@ -87,6 +87,13 @@ FInstanceCullingContext::FInstanceCullingContext(ERHIFeatureLevel::Type InFeatur
 {
 }
 
+
+bool FInstanceCullingContext::IsOcclusionCullingEnabled()
+{
+	return GOcclusionCullInstances != 0;
+}
+
+
 FInstanceCullingContext::~FInstanceCullingContext()
 {
 	for (auto& LoadBalancer : LoadBalancers)
@@ -339,7 +346,7 @@ void FInstanceCullingContext::BuildRenderingCommands(
 		return;
 	}
 
-	const bool bOcclusionCullInstances = PrevHZB.IsValid() && GOcclusionCullInstances > 0;
+	const bool bOcclusionCullInstances = PrevHZB.IsValid() && IsOcclusionCullingEnabled();
 	const uint32 InstanceIdBufferSize = TotalInstances * ViewIds.Num();
 	if (InstanceCullingDrawParams && InstanceCullingManager && InstanceCullingManager->IsDeferredCullingActive())
 	{
@@ -601,7 +608,7 @@ void FInstanceCullingDeferredContext::ProcessBatched(TStaticArray<FBuildInstance
 		PassParameters[Mode]->LoadBalancerParameters.NumItems = LoadBalancers[Mode]->GetItems().Num();
 		PassParameters[Mode]->NumCullingViews = InstanceCullingManager->GetCullingViews().Num();
 
-		const bool bOcclusionCullInstances = PrevHZB != nullptr && GOcclusionCullInstances > 0;
+		const bool bOcclusionCullInstances = PrevHZB != nullptr && FInstanceCullingContext::IsOcclusionCullingEnabled();
 		if (bOcclusionCullInstances)
 		{
 			PassParameters[Mode]->HZBTexture = PrevHZB;
@@ -767,7 +774,7 @@ FInstanceCullingDeferredContext *FInstanceCullingContext::CreateDeferredContext(
 		PassParameters[Mode]->LoadBalancerParameters.ItemBuffer = GraphBuilder.CreateSRV(ItemBuffer);
 		PassParameters[Mode]->CurrentBatchProcessingMode = Mode;
 
-		const bool bOcclusionCullInstances = GOcclusionCullInstances > 0;
+		const bool bOcclusionCullInstances = FInstanceCullingContext::IsOcclusionCullingEnabled();
 		if (bOcclusionCullInstances)
 		{
 			// Fill with a placeholder as AddPass expects HZBTexture to be valid. ProcessBatched will fill with real HZB textures.
diff --git a/Engine/Source/Runtime/Renderer/Private/SceneRendering.cpp b/Engine/Source/Runtime/Renderer/Private/SceneRendering.cpp
--- a/Engine/Source/Runtime/Renderer/Private/SceneRendering.cpp
+++ b/Engine/Source/Runtime/Renderer/Private/SceneRendering.cpp
@@ -3527,14 +3527,10 @@ void FSceneRenderer::SetupMeshPass(FViewInfo& View, FExclusiveDepthStencil::Type
 			}
 			const bool bDrawOnlyVSMInvalidatingGeometry = ViewFamily.EngineShowFlags.DrawOnlyVSMInvalidatingGeo != 0;
 
-			FDepthPassInfo DepthPassInfo = GetDepthPassInfo(Scene);
-			const bool bIsEarlyDepthComplete = (DepthPassInfo.EarlyZPassMode == DDM_AllOpaque || DepthPassInfo.EarlyZPassMode == DDM_AllOpaqueNoVelocity);
-			TRefCountPtr<IPooledRenderTarget> PrevHZB = !bIsEarlyDepthComplete ? View.PrevViewInfo.NaniteHZB : View.PrevViewInfo.HZB;
-
 			Pass.DispatchPassSetup(
 				Scene,
 				View,
-				FInstanceCullingContext(FeatureLevel, &InstanceCullingManager, ViewIds, PrevHZB, InstanceCullingMode, bDrawOnlyVSMInvalidatingGeometry),
+				FInstanceCullingContext(FeatureLevel, &InstanceCullingManager, ViewIds, View.PrevViewInfo.HZB, InstanceCullingMode, bDrawOnlyVSMInvalidatingGeometry),
 				PassType,
 				BasePassDepthStencilAccess,
 				MeshPassProcessor,
diff --git a/Engine/Source/Runtime/Renderer/Public/InstanceCulling/InstanceCullingContext.h b/Engine/Source/Runtime/Renderer/Public/InstanceCulling/InstanceCullingContext.h
--- a/Engine/Source/Runtime/Renderer/Public/InstanceCulling/InstanceCullingContext.h
+++ b/Engine/Source/Runtime/Renderer/Public/InstanceCulling/InstanceCullingContext.h
@@ -81,6 +81,8 @@ public:
 
 	static RENDERER_API const TRDGUniformBufferRef<FInstanceCullingGlobalUniforms> CreateDummyInstanceCullingUniformBuffer(FRDGBuilder& GraphBuilder);
 
+	static bool IsOcclusionCullingEnabled();
+
 	struct FBatchItem
 	{
 		const FInstanceCullingContext* Context = nullptr;
