Fix missing clear when HairStrandsVelocity happens to be the first to use the SceneVelocity (repro in QAGame->TM-HairAndFur, D3D and Vulkan).

diff --git a/Engine/Source/Runtime/Renderer/Private/HairStrands/HairStrandsVisibility.cpp b/Engine/Source/Runtime/Renderer/Private/HairStrands/HairStrandsVisibility.cpp
--- a/Engine/Source/Runtime/Renderer/Private/HairStrands/HairStrandsVisibility.cpp
+++ b/Engine/Source/Runtime/Renderer/Private/HairStrands/HairStrandsVisibility.cpp
@@ -973,6 +973,11 @@ static void AddHairVelocityPass(
 	if (!bWriteOutVelocity)
 		return;
 
+	if (!HasBeenProduced(OutVelocityTexture))
+	{
+		AddClearUAVPass(GraphBuilder, GraphBuilder.CreateUAV(OutVelocityTexture), 0.f);
+	}
+
 	const FIntPoint Resolution = OutVelocityTexture->Desc.Extent;
 	{
 		FRDGTextureDesc Desc = FRDGTextureDesc::Create2D(Resolution, PF_R8_UINT, FClearValueBinding::None, TexCreate_UAV);
