MarkTAAFastResolve's stencil pass used DepthNop which could potentially throw away the depth (depends on platform).  (easy to spot in Vulkan with Renderdoc, where targets are filled with garbage to simulate the DONT_CARE)

diff --git a/Engine/Source/Runtime/Renderer/Private/HairStrands/HairStrandsComposition.cpp b/Engine/Source/Runtime/Renderer/Private/HairStrands/HairStrandsComposition.cpp
--- a/Engine/Source/Runtime/Renderer/Private/HairStrands/HairStrandsComposition.cpp
+++ b/Engine/Source/Runtime/Renderer/Private/HairStrands/HairStrandsComposition.cpp
@@ -348,7 +348,7 @@ static void AddHairVisibilityFastResolveMaskPass(
 		OutDepthTexture,
 		ERenderTargetLoadAction::ELoad,
 		ERenderTargetLoadAction::ELoad,
-		FExclusiveDepthStencil::DepthNop_StencilWrite);
+		FExclusiveDepthStencil::DepthRead_StencilWrite);
 
 	TShaderMapRef<FHairVisibilityFastResolveMaskPS> PixelShader(View.ShaderMap);
 	InternalCommonDrawPass(
