Fix a bug of depth isn't stored with fullprepass enabled.

#p4v-preflight-copy 17164862

diff --git a/Engine/Source/Runtime/Renderer/Private/MobileShadingRenderer.cpp b/Engine/Source/Runtime/Renderer/Private/MobileShadingRenderer.cpp
--- a/Engine/Source/Runtime/Renderer/Private/MobileShadingRenderer.cpp
+++ b/Engine/Source/Runtime/Renderer/Private/MobileShadingRenderer.cpp
@@ -454,7 +454,8 @@ void FMobileSceneRenderer::InitViews(FRDGBuilder& GraphBuilder, FSceneTexturesCo
 		Views[0].bIsReflectionCapture ||
 		(bDeferredShading && bPostProcessUsesSceneDepth) ||
 		bShouldRenderVelocities ||
-		bRequireSeparateViewPass;
+		bRequireSeparateViewPass ||
+		bIsFullDepthPrepassEnabled;
 	// never keep MSAA depth
 	bKeepDepthContent = (NumMSAASamples > 1 ? false : bKeepDepthContent);
 
