Fixed editor crash on switching to a mobile preview if scene uses RVT. Due to incorrect View FeatureLevel

diff --git a/Engine/Source/Runtime/Renderer/Private/VT/RuntimeVirtualTextureRender.cpp b/Engine/Source/Runtime/Renderer/Private/VT/RuntimeVirtualTextureRender.cpp
--- a/Engine/Source/Runtime/Renderer/Private/VT/RuntimeVirtualTextureRender.cpp
+++ b/Engine/Source/Runtime/Renderer/Private/VT/RuntimeVirtualTextureRender.cpp
@@ -1111,7 +1111,7 @@ namespace RuntimeVirtualTexture
 		// Initialize a temporary view required for the material render pass
 		//todo[vt]: Some of this, such as ViewRotationMatrix, can be computed once in the Finalizer and passed down.
 		//todo[vt]: Have specific shader variations and setup for different output texture configs
-		FSceneViewFamily::ConstructionValues ViewFamilyInit(nullptr, nullptr, FEngineShowFlags(ESFIM_Game));
+		FSceneViewFamily::ConstructionValues ViewFamilyInit(nullptr, Scene, FEngineShowFlags(ESFIM_Game));
 		ViewFamilyInit.SetWorldTimes(0.0f, 0.0f, 0.0f);
 		FSceneViewFamily& ViewFamily = *GraphBuilder.AllocObject<FSceneViewFamily>(ViewFamilyInit);
 
