Fixed crash when trying to ray trace shadows when ray tracing was disabled. Fallout of 17181924
#

diff --git a/Engine/Source/Runtime/Renderer/Private/LightRendering.cpp b/Engine/Source/Runtime/Renderer/Private/LightRendering.cpp
--- a/Engine/Source/Runtime/Renderer/Private/LightRendering.cpp
+++ b/Engine/Source/Runtime/Renderer/Private/LightRendering.cpp
@@ -164,13 +164,15 @@ bool ShouldRenderRayTracingShadows()
 bool ShouldRenderRayTracingShadowsForLight(const FLightSceneProxy& LightProxy)
 {
 	return (LightProxy.CastsRaytracedShadow() == ECastRayTracedShadow::Enabled || (ShouldRenderRayTracingShadows() && LightProxy.CastsRaytracedShadow() == ECastRayTracedShadow::UseProjectSetting))
-		&& ShouldRenderRayTracingShadowsForLightType((ELightComponentType)LightProxy.GetLightType());
+		&& ShouldRenderRayTracingShadowsForLightType((ELightComponentType)LightProxy.GetLightType())
+		&& IsRayTracingEnabled();
 }
 
 bool ShouldRenderRayTracingShadowsForLight(const FLightSceneInfoCompact& LightInfo)
 {
 	return (LightInfo.CastRaytracedShadow == ECastRayTracedShadow::Enabled || (ShouldRenderRayTracingShadows() && LightInfo.CastRaytracedShadow == ECastRayTracedShadow::UseProjectSetting))
-		&& ShouldRenderRayTracingShadowsForLightType((ELightComponentType)LightInfo.LightType);
+		&& ShouldRenderRayTracingShadowsForLightType((ELightComponentType)LightInfo.LightType)
+		&& IsRayTracingEnabled();
 }
 #endif // RHI_RAYTRACING
 
