Fixing divide by zero by not checking the array size first.

diff --git a/Engine/Source/Runtime/Renderer/Private/ShadowSetup.cpp b/Engine/Source/Runtime/Renderer/Private/ShadowSetup.cpp
--- a/Engine/Source/Runtime/Renderer/Private/ShadowSetup.cpp
+++ b/Engine/Source/Runtime/Renderer/Private/ShadowSetup.cpp
@@ -4427,32 +4427,34 @@ struct FGatherShadowPrimitivesPrepareTask
 
 		// Kick off any child tasks if we added them
 
-		const int32 NumPackets = TaskData.Packets.Num();
-		const int32 NumTasks = FMath::Min<int32>(FTaskGraphInterface::Get().GetNumWorkerThreads(), NumPackets);
-		const int32 PacketsPerTask = NumPackets / NumTasks;
-
-		for (int32 TaskIndex = 0; TaskIndex < NumTasks; TaskIndex++)
+		if (const int32 NumPackets = TaskData.Packets.Num(); NumPackets > 0)
 		{
-			FDynamicShadowsTaskData* pTaskData = &TaskData;
-			const int32 FirstPacketToProcess = TaskIndex * PacketsPerTask;
-			const int32 NumPacketsToProcess = PacketsPerTask;
+			const int32 NumTasks = FMath::Min<int32>(FTaskGraphInterface::Get().GetNumWorkerThreads(), NumPackets);
+			const int32 PacketsPerTask = NumPackets / NumTasks;
 
-			FGraphEventRef NewTask = FFunctionGraphTask::CreateAndDispatchWhenReady([pTaskData, FirstPacketToProcess, NumPacketsToProcess]()
-				{
-					SCOPED_NAMED_EVENT_TEXT("FGatherShadowPrimitivesTask", FColor::Green);
-					for (int32 PacketIndex = 0; PacketIndex < NumPacketsToProcess; PacketIndex++)
+			for (int32 TaskIndex = 0; TaskIndex < NumTasks; TaskIndex++)
+			{
+				FDynamicShadowsTaskData* pTaskData = &TaskData;
+				const int32 FirstPacketToProcess = TaskIndex * PacketsPerTask;
+				const int32 NumPacketsToProcess = PacketsPerTask;
+
+				FGraphEventRef NewTask = FFunctionGraphTask::CreateAndDispatchWhenReady([pTaskData, FirstPacketToProcess, NumPacketsToProcess]()
 					{
-						if (int32 PacketToProcess = FirstPacketToProcess + PacketIndex; PacketToProcess < pTaskData->Packets.Num())
+						SCOPED_NAMED_EVENT_TEXT("FGatherShadowPrimitivesTask", FColor::Green);
+						for (int32 PacketIndex = 0; PacketIndex < NumPacketsToProcess; PacketIndex++)
 						{
-							pTaskData->Packets[PacketToProcess]->AnyThreadTask(*pTaskData);
+							if (int32 PacketToProcess = FirstPacketToProcess + PacketIndex; PacketToProcess < pTaskData->Packets.Num())
+							{
+								pTaskData->Packets[PacketToProcess]->AnyThreadTask(*pTaskData);
+							}
 						}
-					}
-				},
-				GET_STATID(STAT_ParallelForTask),
-				nullptr,
-				CPrio_GatherShadowPrimitives.Get());
+					},
+					GET_STATID(STAT_ParallelForTask),
+						nullptr,
+						CPrio_GatherShadowPrimitives.Get());
 
-			MyCompletionGraphEvent->DontCompleteUntil(NewTask);
+				MyCompletionGraphEvent->DontCompleteUntil(NewTask);
+			}
 		}
 	}
 
