Wait on mesh draw command setup tasks immediately when deferring scene render cleanup, as setup tasks require state to remain consistent within the RenderViewFamily_RenderThread render command.

diff --git a/Engine/Source/Runtime/Renderer/Private/SceneRendering.cpp b/Engine/Source/Runtime/Renderer/Private/SceneRendering.cpp
--- a/Engine/Source/Runtime/Renderer/Private/SceneRendering.cpp
+++ b/Engine/Source/Runtime/Renderer/Private/SceneRendering.cpp
@@ -3784,29 +3784,43 @@ void FSceneRenderer::RenderThreadEnd(FRHICommandListImmediate& RHICmdList)
 
 		if (SceneRenderCleanUpMode == ESceneRenderCleanUpMode::DeferredAndAsync)
 		{
-			FGraphEventArray& CommandListTasks = RHICmdList.GetRenderThreadTaskArray();
+			// Wait on all setup tasks now to ensure that no additional render commands are enqueued which
+			// might mess with render state, since setup tasks are working with high-level render objects.
+			{
+				FGraphEventArray SetupTasks;
 
-			FGraphEventArray Prerequisites;
-			Prerequisites.Append(CommandListTasks);
+				for (FParallelMeshDrawCommandPass* DispatchedShadowDepthPass : DispatchedShadowDepthPasses)
+				{
+					if (DispatchedShadowDepthPass->GetTaskEvent())
+					{
+						SetupTasks.Add(DispatchedShadowDepthPass->GetTaskEvent());
+					}
+				}
 
-			for (FParallelMeshDrawCommandPass* DispatchedShadowDepthPass : DispatchedShadowDepthPasses)
-			{
-				Prerequisites.Add(DispatchedShadowDepthPass->GetTaskEvent());
-			}
-			for (const FViewInfo& View : Views)
-			{
-				for (const FParallelMeshDrawCommandPass& Pass : View.ParallelMeshDrawCommandPasses)
+				for (const FViewInfo& View : Views)
 				{
-					Prerequisites.Add(Pass.GetTaskEvent());
+					for (const FParallelMeshDrawCommandPass& Pass : View.ParallelMeshDrawCommandPasses)
+					{
+						if (Pass.GetTaskEvent())
+						{
+							SetupTasks.Add(Pass.GetTaskEvent());
+						}
+					}
+				}
+
+				if (!SetupTasks.IsEmpty())
+				{
+					FTaskGraphInterface::Get().WaitUntilTasksComplete(SetupTasks, ENamedThreads::GetRenderThread_Local());
 				}
 			}
 
+			FGraphEventArray CommandListTasks = MoveTemp(RHICmdList.GetRenderThreadTaskArray());
+
 			GSceneRenderCleanUpState.Task = FFunctionGraphTask::CreateAndDispatchWhenReady([this]
 			{
 				WaitForTasksAndClearSnapshots(FParallelMeshDrawCommandPass::EWaitThread::TaskAlreadyWaited);
-			}, TStatId(), & Prerequisites);
 
-			CommandListTasks.Empty();
+			}, TStatId(), &CommandListTasks);
 		}
 	}
 }
