Ray Tracing Debug Visualization: Add Triangles Mode

diff --git a/Engine/Shaders/Private/RayTracing/RayTracingDebug.usf b/Engine/Shaders/Private/RayTracing/RayTracingDebug.usf
--- a/Engine/Shaders/Private/RayTracing/RayTracingDebug.usf
+++ b/Engine/Shaders/Private/RayTracing/RayTracingDebug.usf
@@ -158,11 +158,12 @@ float4 DebugRayTracingMaterial(RayDesc Ray)
 struct FRayTracingDebugPayload : FMinimalPayload
 {
 	uint InstanceHash;
+	uint TriangleHash;
 };
 
 float4 DebugRayTracingInstance(RayDesc Ray)
 {
-	float4 Result = (float4)0;
+	float4 Result = float4(0, 0, 0, 1);
 
 	FRayTracingDebugPayload Payload = (FRayTracingDebugPayload)0;
 	Payload.SetMiss();
@@ -179,8 +180,16 @@ float4 DebugRayTracingInstance(RayDesc Ray)
 
 	if (Payload.IsHit())
 	{
-		Result.rgb = IntToColor(Payload.InstanceHash);
-		Result.a = 1;
+		switch (VisualizationMode)
+		{
+		default:
+		case RAY_TRACING_DEBUG_VIZ_INSTANCES:
+			Result.rgb = IntToColor(Payload.InstanceHash);
+			break;
+		case RAY_TRACING_DEBUG_VIZ_TRIANGLES:
+			Result.rgb = IntToColor(Payload.TriangleHash);
+			break;
+		}
 	}
 
 	return Result;
@@ -196,7 +205,7 @@ RAY_TRACING_ENTRY_RAYGEN(RayTracingDebugMainRGS)
 
 	float4 Result = (float4)0;
 
-	if (VisualizationMode == RAY_TRACING_DEBUG_VIZ_INSTANCES)
+	if (VisualizationMode == RAY_TRACING_DEBUG_VIZ_INSTANCES || VisualizationMode == RAY_TRACING_DEBUG_VIZ_TRIANGLES)
 	{
 		Result = DebugRayTracingInstance(Ray);
 	}
@@ -235,4 +244,5 @@ RAY_TRACING_ENTRY_CLOSEST_HIT(RayTracingDebugMainCHS,
 			Payload.InstanceHash += MurmurMix(asuint(Mat[i][j]));
 		}
 	}
+	Payload.TriangleHash = MurmurMix(Payload.InstanceHash + PrimitiveIndex());
 }
diff --git a/Engine/Shaders/Shared/RaytracingDebugDefinitions.h b/Engine/Shaders/Shared/RaytracingDebugDefinitions.h
--- a/Engine/Shaders/Shared/RaytracingDebugDefinitions.h
+++ b/Engine/Shaders/Shared/RaytracingDebugDefinitions.h
@@ -31,4 +31,5 @@
 #define RAY_TRACING_DEBUG_VIZ_ANISOTROPY				21 
 #define RAY_TRACING_DEBUG_VIZ_INSTANCES					22
 #define RAY_TRACING_DEBUG_VIZ_PERFORMANCE				23 
+#define RAY_TRACING_DEBUG_VIZ_TRIANGLES					24
 
diff --git a/Engine/Source/Editor/UnrealEd/Private/RayTracingDebugVisualizationMenuCommands.cpp b/Engine/Source/Editor/UnrealEd/Private/RayTracingDebugVisualizationMenuCommands.cpp
--- a/Engine/Source/Editor/UnrealEd/Private/RayTracingDebugVisualizationMenuCommands.cpp
+++ b/Engine/Source/Editor/UnrealEd/Private/RayTracingDebugVisualizationMenuCommands.cpp
@@ -60,6 +60,8 @@ void FRayTracingDebugVisualizationMenuCommands::CreateRayTracingDebugVisualizati
 		RayTracingDebugModeNames.Add(LOCTEXT("Performance", "Performance"));
 	}
 
+	RayTracingDebugModeNames.Add(LOCTEXT("Triangles", "Triangles"));
+
 	for ( int32 RayTracingDebugIndex = 0; RayTracingDebugIndex < RayTracingDebugModeNames.Num(); ++RayTracingDebugIndex)
 	{
 		const FText CommandNameText = RayTracingDebugModeNames[RayTracingDebugIndex];
diff --git a/Engine/Source/Runtime/Renderer/Private/RayTracing/RayTracingDebug.cpp b/Engine/Source/Runtime/Renderer/Private/RayTracing/RayTracingDebug.cpp
--- a/Engine/Source/Runtime/Renderer/Private/RayTracing/RayTracingDebug.cpp
+++ b/Engine/Source/Runtime/Renderer/Private/RayTracing/RayTracingDebug.cpp
@@ -83,6 +83,10 @@ public:
 };
 IMPLEMENT_SHADER_TYPE(, FRayTracingDebugCHS, TEXT("/Engine/Private/RayTracing/RayTracingDebug.usf"), TEXT("RayTracingDebugMainCHS"), SF_RayHitGroup);
 
+static bool RequiresRayTracingDebugCHS(uint32 DebugVisualizationMode)
+{
+	return DebugVisualizationMode == RAY_TRACING_DEBUG_VIZ_INSTANCES || DebugVisualizationMode == RAY_TRACING_DEBUG_VIZ_TRIANGLES;
+}
 
 void FDeferredShadingSceneRenderer::PrepareRayTracingDebug(const FSceneViewFamily& ViewFamily, TArray<FRHIRayTracingShader*>& OutRayGenShaders)
 {
@@ -123,6 +127,7 @@ void FDeferredShadingSceneRenderer::RenderRayTracingDebug(FRDGBuilder& GraphBuil
 		RayTracingDebugVisualizationModes.Emplace(FName(*LOCTEXT("Anisotropy", "Anisotropy").ToString()),										RAY_TRACING_DEBUG_VIZ_ANISOTROPY);
 		RayTracingDebugVisualizationModes.Emplace(FName(*LOCTEXT("Instances", "Instances").ToString()),											RAY_TRACING_DEBUG_VIZ_INSTANCES);
 		RayTracingDebugVisualizationModes.Emplace(FName(*LOCTEXT("Performance", "Performance").ToString()),										RAY_TRACING_DEBUG_VIZ_PERFORMANCE);
+		RayTracingDebugVisualizationModes.Emplace(FName(*LOCTEXT("Triangles", "Triangles").ToString()),											RAY_TRACING_DEBUG_VIZ_TRIANGLES);
 	}
 
 	uint32 DebugVisualizationMode;
@@ -167,7 +172,7 @@ void FDeferredShadingSceneRenderer::RenderRayTracingDebug(FRDGBuilder& GraphBuil
 
 	FRayTracingPipelineState* Pipeline = View.RayTracingMaterialPipeline;
 
-	if (DebugVisualizationMode == RAY_TRACING_DEBUG_VIZ_INSTANCES)
+	if (RequiresRayTracingDebugCHS(DebugVisualizationMode))
 	{
 		FRayTracingPipelineStateInitializer Initializer;
 
