Fix bug that SSAO doesn't work on mobile.

diff --git a/Engine/Source/Runtime/Renderer/Private/MobileBasePassRendering.cpp b/Engine/Source/Runtime/Renderer/Private/MobileBasePassRendering.cpp
--- a/Engine/Source/Runtime/Renderer/Private/MobileBasePassRendering.cpp
+++ b/Engine/Source/Runtime/Renderer/Private/MobileBasePassRendering.cpp
@@ -18,6 +18,7 @@
 #include "FramePro/FrameProProfiler.h"
 #include "PostProcess/PostProcessPixelProjectedReflectionMobile.h"
 #include "Engine/SubsurfaceProfile.h"
+#include "SceneTextures.h"
 
 // Changing this causes a full shader recompile
 static TAutoConsoleVariable<int32> CVarMobileDisableVertexFog(
@@ -205,11 +206,8 @@ void SetupMobileBasePassUniformParameters(
 	FRDGTextureRef AmbientOcclusionTexture = SystemTextures.White;
 	if (BasePass == EMobileBasePass::Opaque)
 	{
-		FRDGTextureRef LastFrameScreenSpaceAO = TryRegisterExternalTexture(GraphBuilder, View.PrevViewInfo.MobileAmbientOcclusion); 
-		if (LastFrameScreenSpaceAO)
-		{
-			AmbientOcclusionTexture = LastFrameScreenSpaceAO;
-		}
+		const FSceneTextures& SceneTextures = FSceneTextures::Get(GraphBuilder);
+		AmbientOcclusionTexture = SceneTextures.ScreenSpaceAO;
 	}
 
 	BasePassParameters.AmbientOcclusionTexture = AmbientOcclusionTexture;
diff --git a/Engine/Source/Runtime/Renderer/Private/MobileShadingRenderer.cpp b/Engine/Source/Runtime/Renderer/Private/MobileShadingRenderer.cpp
--- a/Engine/Source/Runtime/Renderer/Private/MobileShadingRenderer.cpp
+++ b/Engine/Source/Runtime/Renderer/Private/MobileShadingRenderer.cpp
@@ -792,6 +792,8 @@ void FMobileSceneRenderer::Render(FRDGBuilder& GraphBuilder)
 	}
 
 	FRDGTextureRef ViewFamilyTexture = TryCreateViewFamilyTexture(GraphBuilder, ViewFamily);
+
+	const FRDGSystemTextures& SystemTextures = FRDGSystemTextures::Get(GraphBuilder);
 	
 	if (bIsFullDepthPrepassEnabled)
 	{
@@ -815,6 +817,10 @@ void FMobileSceneRenderer::Render(FRDGBuilder& GraphBuilder)
 		{
 			RenderAmbientOcclusion(GraphBuilder, SceneTextures.Depth.Resolve, SceneTextures.ScreenSpaceAO);
 		}
+		else
+		{
+			SceneTextures.ScreenSpaceAO = SystemTextures.White;
+		}
 	}
 
 	if (bDeferredShading)
diff --git a/Engine/Source/Runtime/Renderer/Private/PostProcess/PostProcessAmbientOcclusionMobile.cpp b/Engine/Source/Runtime/Renderer/Private/PostProcess/PostProcessAmbientOcclusionMobile.cpp
--- a/Engine/Source/Runtime/Renderer/Private/PostProcess/PostProcessAmbientOcclusionMobile.cpp
+++ b/Engine/Source/Runtime/Renderer/Private/PostProcess/PostProcessAmbientOcclusionMobile.cpp
@@ -960,7 +960,7 @@ static void AddMobileAmbientOcclusionPass(
 static void RenderSSAO(FRDGBuilder& GraphBuilder, FRDGTextureRef SceneDepthTexture, FRDGTextureRef AmbientOcclusionTexture, const TArray<FViewInfo>& Views)
 {
 	
-	TRDGUniformBufferRef<FMobileSceneTextureUniformParameters> SceneTexturesUniformBufferRDG = CreateMobileSceneTextureUniformBuffer(GraphBuilder, EMobileSceneTextureSetupMode::SceneColor);
+	TRDGUniformBufferRef<FMobileSceneTextureUniformParameters> SceneTexturesUniformBufferRDG = CreateMobileSceneTextureUniformBuffer(GraphBuilder, EMobileSceneTextureSetupMode::SceneDepth);
 	
 	for (int32 ViewIndex = 0; ViewIndex < Views.Num(); ViewIndex++)
 	{
