Primitives with any section with a material with the IsSky flag enabled are now excluded from the ray tracing world

diff --git a/Engine/Source/Runtime/Renderer/Private/DeferredShadingRenderer.cpp b/Engine/Source/Runtime/Renderer/Private/DeferredShadingRenderer.cpp
--- a/Engine/Source/Runtime/Renderer/Private/DeferredShadingRenderer.cpp
+++ b/Engine/Source/Runtime/Renderer/Private/DeferredShadingRenderer.cpp
@@ -590,6 +590,7 @@ bool FDeferredShadingSceneRenderer::GatherRayTracingWorldInstancesForView(FRHICo
 		bool bAnySegmentsCastShadow = false;
 		bool bAnySegmentsDecal = false;
 		bool bTwoSided = false;
+		bool bIsSky = false;
 
 		uint64 InstancingKey() const
 		{
@@ -599,6 +600,7 @@ bool FDeferredShadingSceneRenderer::GatherRayTracingWorldInstancesForView(FRHICo
 			Key ^= bAnySegmentsCastShadow ? 0x1ull << 41 : 0x0;
 			Key ^= bAnySegmentsDecal ? 0x1ull << 42 : 0x0;
 			Key ^= bTwoSided ? 0x1ull << 43 : 0x0;
+			Key ^= bIsSky ? 0x1ull << 44 : 0x0;
 			return Key ^ reinterpret_cast<uint64>(RayTracingGeometryRHI);
 		}
 	};
@@ -837,6 +839,7 @@ bool FDeferredShadingSceneRenderer::GatherRayTracingWorldInstancesForView(FRHICo
 									RelevantPrimitive.bAnySegmentsCastShadow |= RayTracingMeshCommand.bCastRayTracedShadows;
 									RelevantPrimitive.bAnySegmentsDecal |= RayTracingMeshCommand.bDecal;
 									RelevantPrimitive.bTwoSided |= RayTracingMeshCommand.bTwoSided;
+									RelevantPrimitive.bIsSky |= RayTracingMeshCommand.bIsSky;
 								}
 								else
 								{
@@ -1229,7 +1232,8 @@ bool FDeferredShadingSceneRenderer::GatherRayTracingWorldInstancesForView(FRHICo
 					continue; // skip dynamic primitives and other 
 				}
 
-				if (GRayTracingExcludeDecals && RelevantPrimitive.bAnySegmentsDecal)
+				if ((GRayTracingExcludeDecals && RelevantPrimitive.bAnySegmentsDecal)
+					|| RelevantPrimitive.bIsSky)
 				{
 					continue;
 				}
diff --git a/Engine/Source/Runtime/Renderer/Private/RayTracing/RayTracingMaterialHitShaders.h b/Engine/Source/Runtime/Renderer/Private/RayTracing/RayTracingMaterialHitShaders.h
--- a/Engine/Source/Runtime/Renderer/Private/RayTracing/RayTracingMaterialHitShaders.h
+++ b/Engine/Source/Runtime/Renderer/Private/RayTracing/RayTracingMaterialHitShaders.h
@@ -69,6 +69,7 @@ protected:
 		SharedCommand.bCastRayTracedShadows = MeshBatch.CastRayTracedShadow && MaterialResource.CastsRayTracedShadows();
 		SharedCommand.bOpaque = MaterialResource.GetBlendMode() == EBlendMode::BLEND_Opaque;
 		SharedCommand.bDecal = MaterialResource.GetMaterialDomain() == EMaterialDomain::MD_DeferredDecal;
+		SharedCommand.bIsSky = MaterialResource.IsSky();
 		SharedCommand.bTwoSided = MaterialResource.IsTwoSided();
 
 		FVertexInputStreamArray VertexStreams;
diff --git a/Engine/Source/Runtime/Renderer/Public/MeshPassProcessor.h b/Engine/Source/Runtime/Renderer/Public/MeshPassProcessor.h
--- a/Engine/Source/Runtime/Renderer/Public/MeshPassProcessor.h
+++ b/Engine/Source/Runtime/Renderer/Public/MeshPassProcessor.h
@@ -1961,6 +1961,7 @@ public:
 	bool bCastRayTracedShadows = true;
 	bool bOpaque = true;
 	bool bDecal = false;
+	bool bIsSky = false;
 	bool bIsTranslucent = false;
 	bool bTwoSided = false;
 
