Fix a bug of getting FSceneTextures without initialization.

#p4v-preflight-copy 17246618

diff --git a/Engine/Source/Runtime/Renderer/Private/MobileBasePassRendering.cpp b/Engine/Source/Runtime/Renderer/Private/MobileBasePassRendering.cpp
--- a/Engine/Source/Runtime/Renderer/Private/MobileBasePassRendering.cpp
+++ b/Engine/Source/Runtime/Renderer/Private/MobileBasePassRendering.cpp
@@ -177,6 +177,7 @@ void SetupMobileBasePassUniformParameters(
 	const FViewInfo& View, 
 	EMobileBasePass BasePass,
 	EMobileSceneTextureSetupMode SetupMode,
+	const FMobileBasePassTextures& MobileBasePassTextures,
 	FMobileBasePassUniformParameters& BasePassParameters)
 {
 	SetupFogUniformParameters(GraphBuilder, View, BasePassParameters.Fog);
@@ -204,10 +205,9 @@ void SetupMobileBasePassUniformParameters(
 	BasePassParameters.EyeAdaptationBuffer = GraphBuilder.CreateSRV(GetEyeAdaptationBuffer(GraphBuilder, View), PF_A32B32G32R32F);
 
 	FRDGTextureRef AmbientOcclusionTexture = SystemTextures.White;
-	if (BasePass == EMobileBasePass::Opaque)
+	if (BasePass == EMobileBasePass::Opaque && MobileBasePassTextures.ScreenSpaceAO != nullptr)
 	{
-		const FSceneTextures& SceneTextures = FSceneTextures::Get(GraphBuilder);
-		AmbientOcclusionTexture = SceneTextures.ScreenSpaceAO;
+		AmbientOcclusionTexture = MobileBasePassTextures.ScreenSpaceAO;
 	}
 
 	BasePassParameters.AmbientOcclusionTexture = AmbientOcclusionTexture;
@@ -243,10 +243,11 @@ TRDGUniformBufferRef<FMobileBasePassUniformParameters> CreateMobileBasePassUnifo
 	FRDGBuilder& GraphBuilder,
 	const FViewInfo& View,
 	EMobileBasePass BasePass,
-	EMobileSceneTextureSetupMode SetupMode)
+	EMobileSceneTextureSetupMode SetupMode,
+	const FMobileBasePassTextures& MobileBasePassTextures)
 {
 	FMobileBasePassUniformParameters* BasePassParameters = GraphBuilder.AllocParameters<FMobileBasePassUniformParameters>();
-	SetupMobileBasePassUniformParameters(GraphBuilder, View, BasePass, SetupMode, *BasePassParameters);
+	SetupMobileBasePassUniformParameters(GraphBuilder, View, BasePass, SetupMode, MobileBasePassTextures, *BasePassParameters);
 	return GraphBuilder.CreateUniformBuffer(BasePassParameters);
 }
 
diff --git a/Engine/Source/Runtime/Renderer/Private/MobileBasePassRendering.h b/Engine/Source/Runtime/Renderer/Private/MobileBasePassRendering.h
--- a/Engine/Source/Runtime/Renderer/Private/MobileBasePassRendering.h
+++ b/Engine/Source/Runtime/Renderer/Private/MobileBasePassRendering.h
@@ -25,6 +25,11 @@
 #include "SkyAtmosphereRendering.h"
 #include "RenderUtils.h"
 
+struct FMobileBasePassTextures
+{
+	FRDGTextureRef ScreenSpaceAO = nullptr;
+};
+
 BEGIN_GLOBAL_SHADER_PARAMETER_STRUCT(FMobileBasePassUniformParameters, )
 	SHADER_PARAMETER(float, AmbientOcclusionStaticFraction)
 	SHADER_PARAMETER(FVector4, SSProfilesTextureSizeAndInvSize)
@@ -54,13 +59,15 @@ extern void SetupMobileBasePassUniformParameters(
 	FRDGBuilder& GraphBuilder,
 	const FViewInfo& View,
 	EMobileBasePass BasePass,
+	const FMobileBasePassTextures& MobileBasePassTextures,
 	FMobileBasePassUniformParameters& BasePassParameters);
 
 extern TRDGUniformBufferRef<FMobileBasePassUniformParameters> CreateMobileBasePassUniformBuffer(
 	FRDGBuilder& GraphBuilder,
 	const FViewInfo& View,
 	EMobileBasePass BasePass,
-	EMobileSceneTextureSetupMode SetupMode);
+	EMobileSceneTextureSetupMode SetupMode,
+	const FMobileBasePassTextures& MobileBasePassTextures = {});
 
 extern void SetupMobileDirectionalLightUniformParameters(
 	const FScene& Scene,
diff --git a/Engine/Source/Runtime/Renderer/Private/MobileShadingRenderer.cpp b/Engine/Source/Runtime/Renderer/Private/MobileShadingRenderer.cpp
--- a/Engine/Source/Runtime/Renderer/Private/MobileShadingRenderer.cpp
+++ b/Engine/Source/Runtime/Renderer/Private/MobileShadingRenderer.cpp
@@ -994,10 +994,13 @@ void FMobileSceneRenderer::RenderForward(FRDGBuilder& GraphBuilder, FRDGTextureR
 
 		UpdateDirectionalLightUniformBuffers(GraphBuilder, View);
 
+		FMobileBasePassTextures MobileBasePassTextures{};
+		MobileBasePassTextures.ScreenSpaceAO = SceneTextures.ScreenSpaceAO;
+
 		EMobileSceneTextureSetupMode SetupMode = EMobileSceneTextureSetupMode::CustomDepth;
 		FMobileRenderPassParameters* PassParameters = GraphBuilder.AllocParameters<FMobileRenderPassParameters>();
 		PassParameters->View = View.GetShaderParameters();
-		PassParameters->MobileBasePass = CreateMobileBasePassUniformBuffer(GraphBuilder, View, EMobileBasePass::Opaque, SetupMode);
+		PassParameters->MobileBasePass = CreateMobileBasePassUniformBuffer(GraphBuilder, View, EMobileBasePass::Opaque, SetupMode, MobileBasePassTextures);
 		PassParameters->RenderTargets = BasePassRenderTargets;
 		#if WITH_DEBUG_VIEW_MODES
 		if (ViewFamily.UseDebugViewPS())
@@ -1279,10 +1282,13 @@ void FMobileSceneRenderer::RenderDeferred(FRDGBuilder& GraphBuilder, const FSort
 
 		UpdateDirectionalLightUniformBuffers(GraphBuilder, View);
 
+		FMobileBasePassTextures MobileBasePassTextures{};
+		MobileBasePassTextures.ScreenSpaceAO = SceneTextures.ScreenSpaceAO;
+
 		EMobileSceneTextureSetupMode SetupMode = EMobileSceneTextureSetupMode::CustomDepth;
 		auto* PassParameters = GraphBuilder.AllocParameters<FMobileRenderPassParameters>();
 		PassParameters->View = View.GetShaderParameters();
-		PassParameters->MobileBasePass = CreateMobileBasePassUniformBuffer(GraphBuilder, View, EMobileBasePass::Opaque, SetupMode);
+		PassParameters->MobileBasePass = CreateMobileBasePassUniformBuffer(GraphBuilder, View, EMobileBasePass::Opaque, SetupMode, MobileBasePassTextures);
 		PassParameters->RenderTargets = BasePassRenderTargets;
 		#if WITH_DEBUG_VIEW_MODES
 		if (ViewFamily.UseDebugViewPS())
