Added option r.RayTracing.ExcludeTransparent to exclude all transparent objects from ray tracing (disabled by default)

diff --git a/Engine/Source/Runtime/Renderer/Private/DeferredShadingRenderer.cpp b/Engine/Source/Runtime/Renderer/Private/DeferredShadingRenderer.cpp
--- a/Engine/Source/Runtime/Renderer/Private/DeferredShadingRenderer.cpp
+++ b/Engine/Source/Runtime/Renderer/Private/DeferredShadingRenderer.cpp
@@ -140,6 +140,15 @@ static FAutoConsoleVariableRef CRayTracingExcludeDecals(
 	TEXT(" 1: Decals excluded from the ray tracing BVH"),
 	ECVF_RenderThreadSafe);
 
+static int32 GRayTracingExcludeTransparent = 0;
+static FAutoConsoleVariableRef CRayTracingExcludeTransparent(
+	TEXT("r.RayTracing.ExcludeTransparent"),
+	GRayTracingExcludeTransparent,
+	TEXT("A toggle that modifies the inclusion of transparent objects (with one or more non-opaque sections) in the ray tracing BVH.\n")
+	TEXT(" 0: Transparent objects included in the ray tracing BVH (default)\n")
+	TEXT(" 1: Transparent objects excluded from the ray tracing BVH"),
+	ECVF_RenderThreadSafe);
+
 static TAutoConsoleVariable<int32> CVarRayTracingAsyncBuild(
 	TEXT("r.RayTracing.AsyncBuild"),
 	0,
@@ -1233,6 +1242,7 @@ bool FDeferredShadingSceneRenderer::GatherRayTracingWorldInstancesForView(FRHICo
 				}
 
 				if ((GRayTracingExcludeDecals && RelevantPrimitive.bAnySegmentsDecal)
+					|| (GRayTracingExcludeTransparent && !RelevantPrimitive.bAllSegmentsOpaque)
 					|| RelevantPrimitive.bIsSky)
 				{
 					continue;
