Update mobile movable point light uniformbuffer if LightFadeFactor changed.

#p4v-preflight-copy 17113013

diff --git a/Engine/Source/Runtime/Engine/Private/Components/LightComponent.cpp b/Engine/Source/Runtime/Engine/Private/Components/LightComponent.cpp
--- a/Engine/Source/Runtime/Engine/Private/Components/LightComponent.cpp
+++ b/Engine/Source/Runtime/Engine/Private/Components/LightComponent.cpp
@@ -359,10 +359,6 @@ FLightSceneProxy::FLightSceneProxy(const ULightComponent* InLightComponent)
 	, ShadowAmount(1.0f)
 	, SamplesPerPixel(1)
 	, DeepShadowLayerDistribution(InLightComponent->DeepShadowLayerDistribution)
-	, bMobileMovablePointLightUniformBufferNeedsUpdate(false)
-	, bMobileMovablePointLightShouldBeRender(false)
-	, bMobileMovablePointLightShouldCastShadow(false)
-	, MobileMovablePointLightShadowmapMinMax()
 {
 	check(SceneInterface);
 
diff --git a/Engine/Source/Runtime/Engine/Private/LocalLightSceneProxy.cpp b/Engine/Source/Runtime/Engine/Private/LocalLightSceneProxy.cpp
--- a/Engine/Source/Runtime/Engine/Private/LocalLightSceneProxy.cpp
+++ b/Engine/Source/Runtime/Engine/Private/LocalLightSceneProxy.cpp
@@ -141,6 +141,21 @@ void FLocalLightSceneProxy::UpdateRadius(float ComponentRadius)
 
 	// Min to avoid div by 0 (NaN in InvRadius)
 	InvRadius = 1.0f / FMath::Max(0.00001f, ComponentRadius);
+}
+
+bool FLocalLightSceneProxy::IsLocalLight() const
+{
+	return true;
+}
 
-	bMobileMovablePointLightUniformBufferNeedsUpdate = true;
+void FLocalLightSceneProxy::UpdateMobileMovableLocalLightUniformBuffer(const FMobileMovablePointLightUniformShaderParameters& MobileMovableLocalLightUniformShaderParameters)
+{
+	if (MobileMovableLocalLightUniformBuffer == nullptr)
+	{
+		MobileMovableLocalLightUniformBuffer = TUniformBufferRef<FMobileMovablePointLightUniformShaderParameters>::CreateUniformBufferImmediate(MobileMovableLocalLightUniformShaderParameters, UniformBuffer_MultiFrame);
+	}
+	else
+	{
+		MobileMovableLocalLightUniformBuffer.UpdateUniformBufferImmediate(MobileMovableLocalLightUniformShaderParameters);
+	}
 }
\ No newline at end of file
diff --git a/Engine/Source/Runtime/Engine/Public/LocalLightSceneProxy.h b/Engine/Source/Runtime/Engine/Public/LocalLightSceneProxy.h
--- a/Engine/Source/Runtime/Engine/Public/LocalLightSceneProxy.h
+++ b/Engine/Source/Runtime/Engine/Public/LocalLightSceneProxy.h
@@ -53,6 +53,15 @@ public:
 
 	virtual bool GetPerObjectProjectedShadowInitializer(const FBoxSphereBounds& SubjectBounds, class FPerObjectProjectedShadowInitializer& OutInitializer) const override;
 
+	virtual bool IsLocalLight() const override;
+
+	FORCEINLINE FRHIUniformBuffer* GetMobileMovableLocalLightUniformBufferRHI() const
+	{
+		return MobileMovableLocalLightUniformBuffer.GetReference();
+	}
+
+	virtual void UpdateMobileMovableLocalLightUniformBuffer(const FMobileMovablePointLightUniformShaderParameters& MobileMovableLocalLightUniformShaderParameters);
+
 protected:
 
 	/** Updates the light scene info's radius from the component. */
@@ -60,4 +69,7 @@ protected:
 
 	float MaxDrawDistance;
 	float FadeRange;
+
+	/** The movable local light's uniform buffer for mobile. */
+	TUniformBufferRef<FMobileMovablePointLightUniformShaderParameters> MobileMovableLocalLightUniformBuffer;
 };
diff --git a/Engine/Source/Runtime/Engine/Public/SceneManagement.h b/Engine/Source/Runtime/Engine/Public/SceneManagement.h
--- a/Engine/Source/Runtime/Engine/Public/SceneManagement.h
+++ b/Engine/Source/Runtime/Engine/Public/SceneManagement.h
@@ -1286,42 +1286,20 @@ BEGIN_GLOBAL_SHADER_PARAMETER_STRUCT(FMobileMovablePointLightUniformShaderParame
 	SHADER_PARAMETER(FMatrix44f, SpotLightShadowWorldToShadowMatrix)
 END_GLOBAL_SHADER_PARAMETER_STRUCT()
 
-/** Initializes the movable point light uniform shader parameters. */
-FORCEINLINE FMobileMovablePointLightUniformShaderParameters GetMovablePointLightUniformShaderParameters(
-	const FVector4& LightPositionAndInvRadius,
-	const FVector4& LightColorAndFalloffExponent,
-	const FVector4& SpotLightDirectionAndSpecularScale,
-	const FVector4& SpotLightAnglesAndSoftTransitionScaleAndLightShadowType,
-	const FVector4& SpotLightShadowSharpenAndShadowFadeFraction,
-	const FVector4& SpotLightShadowmapMinMax,
-	const FMatrix& SpotLightShadowWorldToShadowMatrix
-)
+FORCEINLINE FMobileMovablePointLightUniformShaderParameters GetDummyMovablePointLightUniformShaderParameters()
 {
 	FMobileMovablePointLightUniformShaderParameters Result;
-	Result.LightPositionAndInvRadius = LightPositionAndInvRadius;
-	Result.LightColorAndFalloffExponent = LightColorAndFalloffExponent;
-	Result.SpotLightDirectionAndSpecularScale = SpotLightDirectionAndSpecularScale;
-	Result.SpotLightAnglesAndSoftTransitionScaleAndLightShadowType = SpotLightAnglesAndSoftTransitionScaleAndLightShadowType;
-	Result.SpotLightShadowSharpenAndShadowFadeFraction = SpotLightShadowSharpenAndShadowFadeFraction;
-	Result.SpotLightShadowmapMinMax = SpotLightShadowmapMinMax;
-	Result.SpotLightShadowWorldToShadowMatrix = (FMatrix44f)SpotLightShadowWorldToShadowMatrix;
+	Result.LightPositionAndInvRadius = FVector4();
+	Result.LightColorAndFalloffExponent = FVector4();
+	Result.SpotLightDirectionAndSpecularScale = FVector4();
+	Result.SpotLightAnglesAndSoftTransitionScaleAndLightShadowType = FVector4();
+	Result.SpotLightShadowSharpenAndShadowFadeFraction = FVector4();
+	Result.SpotLightShadowmapMinMax = FVector4();
+	Result.SpotLightShadowWorldToShadowMatrix = (FMatrix44f)FMatrix();
 
 	return Result;
 }
 
-FORCEINLINE FMobileMovablePointLightUniformShaderParameters GetDummyMovablePointLightUniformShaderParameters()
-{
-	return GetMovablePointLightUniformShaderParameters(
-		FVector4(),
-		FVector4(),
-		FVector4(),
-		FVector4(),
-		FVector4(),
-		FVector4(),
-		FMatrix()
-	);
-}
-
 /**
  * Dummy mobile movable point light uniform buffer.
  */
@@ -1377,6 +1355,7 @@ public:
 	virtual float GetSourceRadius() const { return 0.0f; }
 	virtual bool IsInverseSquared() const { return true; }
 	virtual bool IsRectLight() const { return false; }
+	virtual bool  IsLocalLight() const { return false; }
 	virtual bool HasSourceTexture() const { return false; }
 	virtual float GetLightSourceAngle() const { return 0.0f; }
 	virtual float GetShadowSourceAngleFactor() const { return 1.0f; }
@@ -1582,16 +1561,6 @@ public:
 	virtual FLinearColor GetCloudScatteredLuminanceScale() const { return FLinearColor::White; }
 	virtual bool GetUsePerPixelAtmosphereTransmittance() const { return false; }
 
-	FORCEINLINE void SetMobileMovablePointLightUniformBufferNeedsUpdate(bool bInMobileMovablePointLightUniformBufferNeedsUpdate)
-	{
-		bMobileMovablePointLightUniformBufferNeedsUpdate = bInMobileMovablePointLightUniformBufferNeedsUpdate;
-	}
-
-	FORCEINLINE FRHIUniformBuffer* GetMobileMovablePointLightUniformBufferRHI() const
-	{
-		return MobileMovablePointLightUniformBuffer.GetReference();
-	}
-
 protected:
 
 	friend class FScene;
@@ -1776,21 +1745,6 @@ protected:
 	/** Deep shadow layer distribution. */
 	float DeepShadowLayerDistribution;
 
-	/** If this is TRUE, the light's mobile movable point light uniform buffer needs to be updated before it can be used for mobile base pass rendering. */
-	bool bMobileMovablePointLightUniformBufferNeedsUpdate;
-
-	/** Cached ShouldBeRender for mobile, since if the ShouldBeRender is changed we have to update the movable point lights uniform buffer. */
-	bool bMobileMovablePointLightShouldBeRender;
-
-	/** Cached DynamicShadows show flag for mobile, since if the show flag is changed we have to update the movable point lights uniform buffer. */
-	bool bMobileMovablePointLightShouldCastShadow;
-
-	/** Cached the spotlight shadow map min and max value for mobile, since if the value is changed we have to update the movable point lights uniform buffer. */
-	FVector4 MobileMovablePointLightShadowmapMinMax;
-
-	/** The movable point light's uniform buffer for mobile. */
-	TUniformBufferRef<FMobileMovablePointLightUniformShaderParameters> MobileMovablePointLightUniformBuffer;
-
 	/**
 	 * Updates the light proxy's cached transforms.
 	 * @param InLightToWorld - The new light-to-world transform.
diff --git a/Engine/Source/Runtime/Renderer/Private/LightSceneInfo.cpp b/Engine/Source/Runtime/Renderer/Private/LightSceneInfo.cpp
--- a/Engine/Source/Runtime/Renderer/Private/LightSceneInfo.cpp
+++ b/Engine/Source/Runtime/Renderer/Private/LightSceneInfo.cpp
@@ -9,6 +9,7 @@
 #include "SceneCore.h"
 #include "ScenePrivate.h"
 #include "DistanceFieldLightingShared.h"
+#include "LocalLightSceneProxy.h"
 
 int32 GWholeSceneShadowUnbuiltInteractionThreshold = 500;
 static FAutoConsoleVariableRef CVarWholeSceneShadowUnbuiltInteractionThreshold(
@@ -120,8 +121,9 @@ void FLightSceneInfo::AddToScene()
 
 			if (bIsValidLightTypeMobile)
 			{
-				Proxy->MobileMovablePointLightUniformBuffer = TUniformBufferRef<FMobileMovablePointLightUniformShaderParameters>::CreateUniformBufferImmediate(GetDummyMovablePointLightUniformShaderParameters(), UniformBuffer_MultiFrame);
-				Proxy->bMobileMovablePointLightUniformBufferNeedsUpdate = true;
+				check(Proxy->IsLocalLight());
+				FLocalLightSceneProxy* LocalLightSceneProxy = (FLocalLightSceneProxy*)Proxy;
+				LocalLightSceneProxy->UpdateMobileMovableLocalLightUniformBuffer(GetDummyMovablePointLightUniformShaderParameters());
 			}
 		}
 	}
@@ -243,7 +245,7 @@ FLightPrimitiveInteraction* FLightSceneInfo::GetDynamicInteractionStaticPrimitiv
 	return DynamicInteractionStaticPrimitiveList;
 }
 
-void FLightSceneInfo::ConditionalUpdateMobileMovablePointLightUniformBuffer(const FSceneRenderer* SceneRenderer)
+void FLightSceneInfo::UpdateMobileMovableLocalLightUniformBuffer(const FSceneRenderer* SceneRenderer)
 {
 	QUICK_SCOPE_CYCLE_COUNTER(STAT_FLightSceneProxy_UpdateMobileMovablePointLightUniformBuffer);
 
@@ -258,13 +260,8 @@ void FLightSceneInfo::ConditionalUpdateMobileMovablePointLightUniformBuffer(cons
 
 	checkSlow(MobileNumDynamicPointLights > 0 && SceneRenderer);
 
-	FVector4 LightPositionAndInvRadius;
-	FVector4 LightColorAndFalloffExponent;
-	FVector4 SpotLightDirectionAndSpecularScale;
-	FVector4 SpotLightAnglesAndSoftTransitionScaleAndLightShadowType;
-	FVector4 SpotLightShadowSharpenAndShadowFadeFraction;
-	FVector4 SpotLightShadowmapMinMax;
-	FMatrix SpotLightWorldToShadowMatrix;
+	check(Proxy->IsLocalLight());
+	FLocalLightSceneProxy* LocalLightSceneProxy = (FLocalLightSceneProxy*)Proxy;
 
 	bool bShouldBeRender = false;
 	bool bShouldCastShadow = false;
@@ -322,6 +319,8 @@ void FLightSceneInfo::ConditionalUpdateMobileMovablePointLightUniformBuffer(cons
 		float SoftTransitionScale = 0.0f;
 		float ShadowFadeFraction = 0.0f;
 
+		FMobileMovablePointLightUniformShaderParameters MobileMovablePointLightUniformShaderParameters;
+
 		if (bShouldCastShadow)
 		{
 			FProjectedShadowInfo *ProjectedShadowInfo = SceneRenderer->VisibleLightInfos[Id].AllProjectedShadows.Last();
@@ -337,45 +336,20 @@ void FLightSceneInfo::ConditionalUpdateMobileMovablePointLightUniformBuffer(cons
 				ShadowFadeFraction = FMath::Max(ShadowFadeFraction, ProjectedShadowInfo->FadeAlphas[ViewIndex]);
 			}
 
-			SpotLightShadowSharpenAndShadowFadeFraction = FVector4(Proxy->GetShadowSharpen() * 7.0f + 1.0f, ShadowFadeFraction, ProjectedShadowInfo->GetShaderReceiverDepthBias(), 0.0f);
-			SpotLightWorldToShadowMatrix = ProjectedShadowInfo->GetWorldToShadowMatrix(SpotLightShadowmapMinMax);
+			MobileMovablePointLightUniformShaderParameters.SpotLightShadowSharpenAndShadowFadeFraction = FVector4(Proxy->GetShadowSharpen() * 7.0f + 1.0f, ShadowFadeFraction, ProjectedShadowInfo->GetShaderReceiverDepthBias(), 0.0f);
+			MobileMovablePointLightUniformShaderParameters.SpotLightShadowWorldToShadowMatrix = ProjectedShadowInfo->GetWorldToShadowMatrix(MobileMovablePointLightUniformShaderParameters.SpotLightShadowmapMinMax);
 		}
 		
-		LightPositionAndInvRadius = FVector4(LightParameters.Position, LightParameters.InvRadius);
-		LightColorAndFalloffExponent = FVector4(LightParameters.Color, LightParameters.FalloffExponent);
-		SpotLightDirectionAndSpecularScale = FVector4(LightParameters.Direction.X, LightParameters.Direction.Y, LightParameters.Direction.Z, Proxy->GetSpecularScale());
-		SpotLightAnglesAndSoftTransitionScaleAndLightShadowType = FVector4(LightParameters.SpotAngles.X, LightParameters.SpotAngles.Y, SoftTransitionScale, LightShadowType);
-	}
-
-	if (bShouldBeRender != Proxy->bMobileMovablePointLightShouldBeRender || 
-		bShouldCastShadow != Proxy->bMobileMovablePointLightShouldCastShadow ||
-		SpotLightShadowmapMinMax != Proxy->MobileMovablePointLightShadowmapMinMax)
-	{
-		Proxy->bMobileMovablePointLightUniformBufferNeedsUpdate = true;
+		MobileMovablePointLightUniformShaderParameters.LightPositionAndInvRadius = FVector4(LightParameters.Position, LightParameters.InvRadius);
+		MobileMovablePointLightUniformShaderParameters.LightColorAndFalloffExponent = FVector4(LightParameters.Color, LightParameters.FalloffExponent);
+		MobileMovablePointLightUniformShaderParameters.SpotLightDirectionAndSpecularScale = FVector4(LightParameters.Direction.X, LightParameters.Direction.Y, LightParameters.Direction.Z, Proxy->GetSpecularScale());
+		MobileMovablePointLightUniformShaderParameters.SpotLightAnglesAndSoftTransitionScaleAndLightShadowType = FVector4(LightParameters.SpotAngles.X, LightParameters.SpotAngles.Y, SoftTransitionScale, LightShadowType);
 
-		Proxy->bMobileMovablePointLightShouldBeRender = bShouldBeRender;
-
-		Proxy->bMobileMovablePointLightShouldCastShadow = bShouldCastShadow;
-
-		Proxy->MobileMovablePointLightShadowmapMinMax = SpotLightShadowmapMinMax;
+		LocalLightSceneProxy->UpdateMobileMovableLocalLightUniformBuffer(MobileMovablePointLightUniformShaderParameters);
 	}
-
-	if (Proxy->bMobileMovablePointLightUniformBufferNeedsUpdate)
+	else
 	{
-		const FMobileMovablePointLightUniformShaderParameters MobileMovablePointLightUniformShaderParameters =
-			GetMovablePointLightUniformShaderParameters(
-				LightPositionAndInvRadius,
-				LightColorAndFalloffExponent,
-				SpotLightDirectionAndSpecularScale,
-				SpotLightAnglesAndSoftTransitionScaleAndLightShadowType,
-				SpotLightShadowSharpenAndShadowFadeFraction,
-				SpotLightShadowmapMinMax,
-				SpotLightWorldToShadowMatrix
-			);
-
-		Proxy->MobileMovablePointLightUniformBuffer.UpdateUniformBufferImmediate(MobileMovablePointLightUniformShaderParameters);
-
-		Proxy->bMobileMovablePointLightUniformBufferNeedsUpdate = false;
+		LocalLightSceneProxy->UpdateMobileMovableLocalLightUniformBuffer(GetDummyMovablePointLightUniformShaderParameters());
 	}
 }
 
diff --git a/Engine/Source/Runtime/Renderer/Private/LightSceneInfo.h b/Engine/Source/Runtime/Renderer/Private/LightSceneInfo.h
--- a/Engine/Source/Runtime/Renderer/Private/LightSceneInfo.h
+++ b/Engine/Source/Runtime/Renderer/Private/LightSceneInfo.h
@@ -340,8 +340,8 @@ public:
 		return (uint32)LightSceneInfo->Id;
 	}
 
-	// Update the mobile movable point light uniform buffer before it is used for mobile base pass rendering.
-	void ConditionalUpdateMobileMovablePointLightUniformBuffer(const class FSceneRenderer* SceneRenderer);
+	// Update the mobile movable local light uniform buffer before it is used for mobile base pass rendering.
+	void UpdateMobileMovableLocalLightUniformBuffer(const class FSceneRenderer* SceneRenderer);
 
 	bool ShouldRecordShadowSubjectsForMobile() const;
 };
diff --git a/Engine/Source/Runtime/Renderer/Private/MobileBasePassRendering.cpp b/Engine/Source/Runtime/Renderer/Private/MobileBasePassRendering.cpp
--- a/Engine/Source/Runtime/Renderer/Private/MobileBasePassRendering.cpp
+++ b/Engine/Source/Runtime/Renderer/Private/MobileBasePassRendering.cpp
@@ -18,7 +18,7 @@
 #include "FramePro/FrameProProfiler.h"
 #include "PostProcess/PostProcessPixelProjectedReflectionMobile.h"
 #include "Engine/SubsurfaceProfile.h"
-#include "SceneTextures.h"
+#include "LocalLightSceneProxy.h"
 
 // Changing this causes a full shader recompile
 static TAutoConsoleVariable<int32> CVarMobileDisableVertexFog(
@@ -164,7 +164,8 @@ FMobileBasePassMovableLightInfo::FMobileBasePassMovableLightInfo(const FPrimitiv
 				
 			if (bIsValidLightType && LightProxy->IsMovable() && (LightProxy->GetLightingChannelMask() & InSceneProxy->GetLightingChannelMask()) != 0)
 			{
-				MovablePointLightUniformBuffer[NumMovablePointLights] = LightProxy->GetMobileMovablePointLightUniformBufferRHI();
+				check(LightProxy->IsLocalLight());
+				MovablePointLightUniformBuffer[NumMovablePointLights] = ((FLocalLightSceneProxy*)LightProxy)->GetMobileMovableLocalLightUniformBufferRHI();
 
 				NumMovablePointLights++;
 			}
diff --git a/Engine/Source/Runtime/Renderer/Private/MobileShadingRenderer.cpp b/Engine/Source/Runtime/Renderer/Private/MobileShadingRenderer.cpp
--- a/Engine/Source/Runtime/Renderer/Private/MobileShadingRenderer.cpp
+++ b/Engine/Source/Runtime/Renderer/Private/MobileShadingRenderer.cpp
@@ -518,8 +518,8 @@ void FMobileSceneRenderer::InitViews(FRDGBuilder& GraphBuilder, FSceneTexturesCo
 		PrepareViewVisibilityLists();
 	}
 
-	/** Before SetupMobileBasePassAfterShadowInit, we need to update the uniform buffer and shadow info for all movable point lights.*/
-	UpdateMovablePointLightUniformBufferAndShadowInfo();
+	/** Before SetupMobileBasePassAfterShadowInit, we need to update the uniform buffer and shadow info for all movable local lights.*/
+	UpdateMovableLocalLightUniformBufferAndShadowInfo();
 
 	SetupMobileBasePassAfterShadowInit(BasePassDepthStencilAccess, ViewCommandsPerView, InstanceCullingManager);
 
@@ -1627,7 +1627,7 @@ void FMobileSceneRenderer::PreTonemapMSAA(FRHICommandListImmediate& RHICmdList,
 }
 
 /** Before SetupMobileBasePassAfterShadowInit, we need to update the uniform buffer and shadow info for all movable point lights.*/
-void FMobileSceneRenderer::UpdateMovablePointLightUniformBufferAndShadowInfo()
+void FMobileSceneRenderer::UpdateMovableLocalLightUniformBufferAndShadowInfo()
 {
 	static auto* MobileNumDynamicPointLightsCVar = IConsoleManager::Get().FindTConsoleVariableDataInt(TEXT("r.MobileNumDynamicPointLights"));
 	const int32 MobileNumDynamicPointLights = MobileNumDynamicPointLightsCVar->GetValueOnRenderThread();
@@ -1660,7 +1660,7 @@ void FMobileSceneRenderer::UpdateMovablePointLightUniformBufferAndShadowInfo()
 
 			if (bIsValidLightType && LightProxy->IsMovable())
 			{
-				LightSceneInfo->ConditionalUpdateMobileMovablePointLightUniformBuffer(this);
+				LightSceneInfo->UpdateMobileMovableLocalLightUniformBuffer(this);
 
 				bool bDynamicShadows = bShouldDynamicShadows
 					&& LightType == LightType_Spot
diff --git a/Engine/Source/Runtime/Renderer/Private/RendererScene.cpp b/Engine/Source/Runtime/Renderer/Private/RendererScene.cpp
--- a/Engine/Source/Runtime/Renderer/Private/RendererScene.cpp
+++ b/Engine/Source/Runtime/Renderer/Private/RendererScene.cpp
@@ -1033,7 +1033,6 @@ FScene::FScene(UWorld* InWorld, bool bInRequiresHitProxies, bool bInIsEditorScen
 ,	NumUnbuiltReflectionCaptures(0)
 ,	NumMobileStaticAndCSMLights_RenderThread(0)
 ,	NumMobileMovableDirectionalLights_RenderThread(0)
-,	MobileWholeSceneShadowAtlasSize(0, 0)
 ,	GPUSkinCache(nullptr)
 ,	SceneLODHierarchy(this)
 ,	RuntimeVirtualTexturePrimitiveHideMaskEditor(0)
@@ -2879,8 +2878,6 @@ void FScene::UpdateLightColorAndBrightness(ULightComponent* Light)
 					{
 						Scene->Lights[ LightSceneInfo->Id ].Color = NewParameters.NewColor;
 					}
-
-					LightSceneInfo->Proxy->SetMobileMovablePointLightUniformBufferNeedsUpdate(true);
 				}
 			});
 	}
diff --git a/Engine/Source/Runtime/Renderer/Private/ScenePrivate.h b/Engine/Source/Runtime/Renderer/Private/ScenePrivate.h
--- a/Engine/Source/Runtime/Renderer/Private/ScenePrivate.h
+++ b/Engine/Source/Runtime/Renderer/Private/ScenePrivate.h
@@ -2897,9 +2897,6 @@ public:
 	int32 NumMobileStaticAndCSMLights_RenderThread;
 	int32 NumMobileMovableDirectionalLights_RenderThread;
 
-	/** Cached shadow atlas size for mobile, since if the resolution is changed we have to update the movable point lights uniform buffer. */
-	FIntPoint MobileWholeSceneShadowAtlasSize;
-
 	FSceneVelocityData VelocityData;
 
 	/** GPU Skinning cache, if enabled */
diff --git a/Engine/Source/Runtime/Renderer/Private/SceneRendering.h b/Engine/Source/Runtime/Renderer/Private/SceneRendering.h
--- a/Engine/Source/Runtime/Renderer/Private/SceneRendering.h
+++ b/Engine/Source/Runtime/Renderer/Private/SceneRendering.h
@@ -2327,8 +2327,8 @@ protected:
 
 	void RenderPixelProjectedReflection(FRDGBuilder& GraphBuilder, FRDGTextureRef SceneColorTexture, FRDGTextureRef SceneDepthTexture, FRDGTextureRef PixelProjectedReflectionTexture, const FPlanarReflectionSceneProxy* PlanarReflectionSceneProxy);
 
-	/** Before SetupMobileBasePassAfterShadowInit, we need to update the uniform buffer and shadow info for all movable point lights.*/
-	void UpdateMovablePointLightUniformBufferAndShadowInfo();
+	/** Before SetupMobileBasePassAfterShadowInit, we need to update the uniform buffer and shadow info for all movable local lights.*/
+	void UpdateMovableLocalLightUniformBufferAndShadowInfo();
 
 	void RenderMobileShadowProjections(FRDGBuilder& GraphBuilder, FRDGTextureRef SceneDepthTexture);
 private:
diff --git a/Engine/Source/Runtime/Renderer/Private/ShadowSetup.cpp b/Engine/Source/Runtime/Renderer/Private/ShadowSetup.cpp
--- a/Engine/Source/Runtime/Renderer/Private/ShadowSetup.cpp
+++ b/Engine/Source/Runtime/Renderer/Private/ShadowSetup.cpp
@@ -5735,10 +5735,6 @@ void FSceneRenderer::AllocateMobileCSMAndSpotLightShadowDepthTargets(FRHICommand
 				WholeSceneAtlasSize.Y = 1 << FMath::CeilLogTwo(WholeSceneAtlasSize.Y);
 			}
 
-			bool bResolutionChanged = Scene->MobileWholeSceneShadowAtlasSize != WholeSceneAtlasSize;
-
-			Scene->MobileWholeSceneShadowAtlasSize = WholeSceneAtlasSize;
-
 			FPooledRenderTargetDesc WholeSceneShadowMapDesc2D(FPooledRenderTargetDesc::Create2DDesc(WholeSceneAtlasSize, PF_ShadowDepth, FClearValueBinding::DepthOne, TexCreate_None, TexCreate_DepthStencilTargetable | TexCreate_ShaderResource, false));
 			WholeSceneShadowMapDesc2D.Flags |= GFastVRamConfig.ShadowCSM;
 			GRenderTargetPool.FindFreeElement(RHICmdList, WholeSceneShadowMapDesc2D, ShadowMapAtlas.RenderTargets.DepthTarget, TEXT("MobileCSMAndSpotLightShadowmap"));
@@ -5752,12 +5748,6 @@ void FSceneRenderer::AllocateMobileCSMAndSpotLightShadowDepthTargets(FRHICommand
 					ProjectedShadowInfo->RenderTargets.CopyReferencesFromRenderTargets(ShadowMapAtlas.RenderTargets);
 					ProjectedShadowInfo->SetupShadowDepthView(this);
 					ShadowMapAtlas.Shadows.Add(ProjectedShadowInfo);
-
-					if (bResolutionChanged)
-					{
-						FLightSceneInfo* LightSceneInfo = (FLightSceneInfo*)(&ProjectedShadowInfo->GetLightSceneInfo());
-						LightSceneInfo->Proxy->SetMobileMovablePointLightUniformBufferNeedsUpdate(true);
-					}
 				}
 			}
 		}
