Bind a dummy with proper format in the "unused" param to prevent tripping Vulkan validation with a "requires UINT......bound descriptor is FLOAT" error.

diff --git a/Engine/Source/Runtime/Renderer/Private/HairStrands/HairStrandsTile.cpp b/Engine/Source/Runtime/Renderer/Private/HairStrands/HairStrandsTile.cpp
--- a/Engine/Source/Runtime/Renderer/Private/HairStrands/HairStrandsTile.cpp
+++ b/Engine/Source/Runtime/Renderer/Private/HairStrands/HairStrandsTile.cpp
@@ -155,8 +155,8 @@ FHairStrandsTiles AddHairStrandsGenerateTilesPass(
 	FHairStrandsTileGenerationPassCS::FParameters* PassParameters = GraphBuilder.AllocParameters<FHairStrandsTileGenerationPassCS::FParameters>();
 	PassParameters->BufferResolution = InputResolution;//View.ViewRect.Size();
 	PassParameters->bUintTexture = bUintTexture ? 1u : 0u;
-	PassParameters->InputFloatTexture = InputTexture;
-	PassParameters->InputUintTexture  = InputTexture;
+	PassParameters->InputFloatTexture = bUintTexture ? GSystemTextures.GetBlackDummy(GraphBuilder) : InputTexture;
+	PassParameters->InputUintTexture  = bUintTexture ? InputTexture : GSystemTextures.GetZeroUIntDummy(GraphBuilder);
 	PassParameters->TileDataBuffer = GraphBuilder.CreateUAV(Out.TileDataBuffer, PF_R16G16_UINT);
 	PassParameters->TileClearBuffer = GraphBuilder.CreateUAV(Out.TileClearBuffer, PF_R16G16_UINT);
 	PassParameters->TileCountBuffer = TileCountUAV;
