Set r.RayTracing.Shadows.AcceptFirstHit=True by default in the engine, improving RT shadow performance by ~15%

Allows shadow rays to terminate at the first hit point rather than closest.
This produces a wrong closest hit distance used by denoiser, however in practice the first hit point is good enough.

diff --git a/Engine/Source/Runtime/Renderer/Private/RayTracing/RayTracingShadows.cpp b/Engine/Source/Runtime/Renderer/Private/RayTracing/RayTracingShadows.cpp
--- a/Engine/Source/Runtime/Renderer/Private/RayTracing/RayTracingShadows.cpp
+++ b/Engine/Source/Runtime/Renderer/Private/RayTracing/RayTracingShadows.cpp
@@ -86,7 +86,7 @@ static TAutoConsoleVariable<int32> CVarRayTracingShadowsLODTransitionEnd(
 
 static TAutoConsoleVariable<int32> CVarRayTracingShadowsAcceptFirstHit(
 	TEXT("r.RayTracing.Shadows.AcceptFirstHit"),
-	0,
+	1,
 	TEXT("Whether to allow shadow rays to terminate early, on first intersected primitive. This may result in worse denoising quality in some cases. (default = 0)"),
 	ECVF_RenderThreadSafe
 );
