Exposure Post Processing: Made it so that if either AutoExposureSpeedUp and AutoExposureSpeedDown parameters are negative, then the speeds are considered invalid and we force ourselves to hit the target exposure in PostProcessEyeAdaptation.cpp

diff --git a/Engine/Source/Runtime/Renderer/Private/PostProcess/PostProcessEyeAdaptation.cpp b/Engine/Source/Runtime/Renderer/Private/PostProcess/PostProcessEyeAdaptation.cpp
--- a/Engine/Source/Runtime/Renderer/Private/PostProcess/PostProcessEyeAdaptation.cpp
+++ b/Engine/Source/Runtime/Renderer/Private/PostProcess/PostProcessEyeAdaptation.cpp
@@ -373,11 +373,12 @@ FEyeAdaptationParameters GetEyeAdaptationParameters(const FViewInfo& View, ERHIF
 	const float MaxAverageLuminance = MaxWhitePointLuminance * kMiddleGrey;
 
 	const bool bValidRange = View.FinalPostProcessSettings.AutoExposureMinBrightness < View.FinalPostProcessSettings.AutoExposureMaxBrightness;
+	const bool bValidSpeeds = View.FinalPostProcessSettings.AutoExposureSpeedDown >= 0.f && View.FinalPostProcessSettings.AutoExposureSpeedUp >= 0.f;
 
 	// if it is a camera cut we force the exposure to go all the way to the target exposure without blending.
 	// if it is manual mode, we also force the exposure to hit the target, which matters for HDR Visualization
 	// if we don't have a valid range (AutoExposureMinBrightness == AutoExposureMaxBrightness) then force it like Manual as well.
-	const float ForceTarget = (View.bCameraCut || AutoExposureMethod == EAutoExposureMethod::AEM_Manual || !bValidRange) ? 1.0f : 0.0f;
+	const float ForceTarget = (View.bCameraCut || AutoExposureMethod == EAutoExposureMethod::AEM_Manual || !bValidRange || !bValidSpeeds) ? 1.0f : 0.0f;
 
 	FEyeAdaptationParameters Parameters;
 	Parameters.ExposureLowPercent = ExposureLowPercent;
