Fix a RDG validation caused by setting mobile AO texture.

#p4v-preflight-copy 17281346

diff --git a/Engine/Source/Runtime/Renderer/Private/MobileShadingRenderer.cpp b/Engine/Source/Runtime/Renderer/Private/MobileShadingRenderer.cpp
--- a/Engine/Source/Runtime/Renderer/Private/MobileShadingRenderer.cpp
+++ b/Engine/Source/Runtime/Renderer/Private/MobileShadingRenderer.cpp
@@ -817,10 +817,6 @@ void FMobileSceneRenderer::Render(FRDGBuilder& GraphBuilder)
 		{
 			RenderAmbientOcclusion(GraphBuilder, SceneTextures.Depth.Resolve, SceneTextures.ScreenSpaceAO);
 		}
-		else
-		{
-			SceneTextures.ScreenSpaceAO = SystemTextures.White;
-		}
 	}
 
 	if (bDeferredShading)
@@ -995,7 +991,7 @@ void FMobileSceneRenderer::RenderForward(FRDGBuilder& GraphBuilder, FRDGTextureR
 		UpdateDirectionalLightUniformBuffers(GraphBuilder, View);
 
 		FMobileBasePassTextures MobileBasePassTextures{};
-		MobileBasePassTextures.ScreenSpaceAO = SceneTextures.ScreenSpaceAO;
+		MobileBasePassTextures.ScreenSpaceAO = bRequiresAmbientOcclusionPass ? SceneTextures.ScreenSpaceAO : SystemTextures.White;
 
 		EMobileSceneTextureSetupMode SetupMode = EMobileSceneTextureSetupMode::CustomDepth;
 		FMobileRenderPassParameters* PassParameters = GraphBuilder.AllocParameters<FMobileRenderPassParameters>();
@@ -1283,7 +1279,7 @@ void FMobileSceneRenderer::RenderDeferred(FRDGBuilder& GraphBuilder, const FSort
 		UpdateDirectionalLightUniformBuffers(GraphBuilder, View);
 
 		FMobileBasePassTextures MobileBasePassTextures{};
-		MobileBasePassTextures.ScreenSpaceAO = SceneTextures.ScreenSpaceAO;
+		MobileBasePassTextures.ScreenSpaceAO = bRequiresAmbientOcclusionPass ? SceneTextures.ScreenSpaceAO : SystemTextures.White;
 
 		EMobileSceneTextureSetupMode SetupMode = EMobileSceneTextureSetupMode::CustomDepth;
 		auto* PassParameters = GraphBuilder.AllocParameters<FMobileRenderPassParameters>();
