Add debug rendering of instance bounds causing virtual shadow map invaidtion
- Change ShaderDrawDebug to make it possible to render without a view (store active views internally).

diff --git a/Engine/Plugins/Runtime/HairStrands/Source/HairStrandsCore/Private/GroomDebug.cpp b/Engine/Plugins/Runtime/HairStrands/Source/HairStrandsCore/Private/GroomDebug.cpp
--- a/Engine/Plugins/Runtime/HairStrands/Source/HairStrandsCore/Private/GroomDebug.cpp
+++ b/Engine/Plugins/Runtime/HairStrands/Source/HairStrandsCore/Private/GroomDebug.cpp
@@ -684,10 +684,7 @@ static void AddDrawDebugCardsGuidesPass(
 	}
 
 	const uint32 MaxCount = 128000;
-	if (ShaderDrawDebug::GetMaxElementCount() < MaxCount)
-	{
-		ShaderDrawDebug::SetMaxElementCount(MaxCount);
-	}
+	ShaderDrawDebug::RequestSpaceForElements(MaxCount);
 
 	if (Instance->HairGroupPublicData->VFInput.GeometryType != EHairGeometryType::Cards || ShaderDrawData == nullptr)
 	{
diff --git a/Engine/Shaders/Private/VirtualShadowMaps/CacheManagement.usf b/Engine/Shaders/Private/VirtualShadowMaps/CacheManagement.usf
--- a/Engine/Shaders/Private/VirtualShadowMaps/CacheManagement.usf
+++ b/Engine/Shaders/Private/VirtualShadowMaps/CacheManagement.usf
@@ -9,6 +9,11 @@
 //#include "NaniteDataDecode.ush"
  #include "../Nanite/HZBCull.ush"
 
+#if ENABLE_DEBUG_MODE
+#include "../ShaderDrawDebug.ush"
+#include "../ColorMap.ush"
+#endif
+
 struct FInstanceSceneDataRange
 {
 	int InstanceSceneDataOffset;
@@ -20,6 +25,10 @@ uint NumRemovedItems;
 RWStructuredBuffer<uint> OutDynamicCasterFlags;
 uint InstanceSceneDataSOAStride;
 
+#if ENABLE_DEBUG_MODE
+uint bDrawBounds;
+#endif
+
 /**
  * Each thread loops over a range on instances loaded from a buffer. The instance bounds are projected to all cached virtual shadow map address space 
  * and any overlapped pages are marked as invalid.
@@ -49,7 +58,10 @@ void VirtualSmInvalidateInstancePagesCS(
 
 		// TODO: Clean up hardcoded flag field test.
 		bool bCastShadows = (GetPrimitiveData(InstanceSceneData.PrimitiveId).Flags & 1U) != 0U;
-		
+#if ENABLE_DEBUG_MODE
+		uint PageInvalidationCount = 0U;
+#endif
+
 		// TODO: test the flag on the instance instead once it is updated correctly InstanceSceneData.CastShadows
 		if (bCastShadows)
 		{
@@ -108,10 +120,13 @@ void VirtualSmInvalidateInstancePagesCS(
 											{
 												// Somewhere
 												OutDynamicCasterFlags[PageFlagOffset] = 1;
+#if ENABLE_DEBUG_MODE
+												++PageInvalidationCount;
+#endif
+
 											}
 										}
 									}
-
 								}
 							}
 						}
@@ -119,5 +134,12 @@ void VirtualSmInvalidateInstancePagesCS(
 				}
 			}
 		}
+#if ENABLE_DEBUG_MODE
+		if (bDrawBounds && PageInvalidationCount > 0U)
+		{
+			float3 Color = float3(0.3f, 0.3f, 0.3f) + ColorMapTurbo(min(1.0f, float(PageInvalidationCount) / 100.0f)) * 0.7f;
+			AddOBB(InstanceSceneData.LocalBoundsCenter - InstanceSceneData.LocalBoundsExtent, InstanceSceneData.LocalBoundsCenter + InstanceSceneData.LocalBoundsExtent, float4(Color, 1.0f), InstanceSceneData.LocalToWorld);
+		}
+#endif
 	}
 }
diff --git a/Engine/Source/Runtime/Renderer/Private/RendererScene.cpp b/Engine/Source/Runtime/Renderer/Private/RendererScene.cpp
--- a/Engine/Source/Runtime/Renderer/Private/RendererScene.cpp
+++ b/Engine/Source/Runtime/Renderer/Private/RendererScene.cpp
@@ -1937,10 +1937,7 @@ void FScene::ShowPhysicsField()
 		{
 			ShaderDrawDebug::SetEnabled(true);
 		}
-		if (ShaderDrawDebug::GetMaxElementCount() < 128000)
-		{
-			ShaderDrawDebug::SetMaxElementCount(128000);
-		}
+		ShaderDrawDebug::RequestSpaceForElements(128000);
 	}
 }
 
diff --git a/Engine/Source/Runtime/Renderer/Private/ShaderDebug.cpp b/Engine/Source/Runtime/Renderer/Private/ShaderDebug.cpp
--- a/Engine/Source/Runtime/Renderer/Private/ShaderDebug.cpp
+++ b/Engine/Source/Runtime/Renderer/Private/ShaderDebug.cpp
@@ -19,7 +19,7 @@ namespace ShaderDrawDebug
 		TEXT("ShaderDrawDebug debugging toggle.\n"),
 		ECVF_Cheat | ECVF_RenderThreadSafe);
 
-	static int32 GShaderDrawDebug_MaxElementCount;
+	static int32 GShaderDrawDebug_MaxElementCount = 1;
 	static FAutoConsoleVariableRef CVarShaderDrawMaxElementCount(
 		TEXT("r.ShaderDrawDebug.MaxElementCount"),
 		GShaderDrawDebug_MaxElementCount,
@@ -32,6 +32,10 @@ namespace ShaderDrawDebug
 		TEXT("Lock the shader draw buffer.\n"),
 		ECVF_Cheat | ECVF_RenderThreadSafe);
 
+	static FViewInfo* GDefaultView = nullptr;
+	static uint32 GElementRequestCount = 0;
+
+
 	bool IsEnabled()
 	{
 		return GShaderDrawDebug_Enable > 0;
@@ -54,7 +58,7 @@ namespace ShaderDrawDebug
 
 	void SetMaxElementCount(uint32 MaxCount)
 	{
-		GShaderDrawDebug_MaxElementCount = FMath::Max(1024, int32(MaxCount));
+		GShaderDrawDebug_MaxElementCount = FMath::Max3(1024, GShaderDrawDebug_MaxElementCount, int32(MaxCount));
 	}
 
 	uint32 GetMaxElementCount()
@@ -62,6 +66,12 @@ namespace ShaderDrawDebug
 		return uint32(FMath::Max(1, GShaderDrawDebug_MaxElementCount));
 	}
 
+
+	void RequestSpaceForElements(uint32 MaxElementCount)
+	{
+		GElementRequestCount += MaxElementCount;
+	}
+
 	bool IsEnabled(const FViewInfo& View)
 	{
 		return IsEnabled() && IsSupported(View.GetShaderPlatform());
@@ -238,6 +248,10 @@ namespace ShaderDrawDebug
 			return;
 		}
 
+		// Update max element count from the request count and reset
+		GShaderDrawDebug_MaxElementCount = FMath::Max(GShaderDrawDebug_MaxElementCount, int32(GElementRequestCount));
+		GElementRequestCount = 0;
+
 		FSceneViewState* ViewState = View.ViewState;
 
 		const bool bLockBufferThisFrame = IsShaderDrawLocked() && ViewState && !ViewState->ShaderDrawDebugStateData.bIsLocked;
@@ -280,6 +294,13 @@ namespace ShaderDrawDebug
 			ViewState->ShaderDrawDebugStateData.IndirectBuffer = nullptr;
 			ViewState->ShaderDrawDebugStateData.bIsLocked = false;
 		}
+
+		// Invalid to call begin twice for the same view.
+		check(GDefaultView != &View);
+		if (GDefaultView == nullptr)
+		{
+			GDefaultView = &View;
+		}
 	}
 
 	void DrawView(FRDGBuilder& GraphBuilder, const FViewInfo& View, FRDGTextureRef OutputTexture, FRDGTextureRef DepthTexture)
@@ -309,6 +330,8 @@ namespace ShaderDrawDebug
 		{
 			return;
 		}
+
+		GDefaultView = nullptr;
 	}
 
 	void SetParameters(FRDGBuilder& GraphBuilder, const FShaderDrawDebugData& Data, FShaderParameters& OutParameters)
@@ -325,4 +348,18 @@ namespace ShaderDrawDebug
 		OutParameters.OutShaderDrawPrimitive = GraphBuilder.CreateUAV(DataBuffer);
 		OutParameters.OutputShaderDrawIndirect = GraphBuilder.CreateUAV(IndirectBuffer);
 	}
+
+	// Returns true if the default view exists and has shader debug rendering enabled (this needs to be checked before using a permutation that requires the shader draw parameters)
+	bool IsDefaultViewEnabled()
+	{
+		return GDefaultView != nullptr && IsEnabled(*GDefaultView);
+	}
+
+	void SetParameters(FRDGBuilder& GraphBuilder, FShaderParameters& OutParameters)
+	{
+		if (GDefaultView != nullptr)
+		{
+			SetParameters(GraphBuilder, GDefaultView->ShaderDrawData, OutParameters);
+		}
+	}
 }
diff --git a/Engine/Source/Runtime/Renderer/Private/VirtualShadowMaps/VirtualShadowMapCacheManager.cpp b/Engine/Source/Runtime/Renderer/Private/VirtualShadowMaps/VirtualShadowMapCacheManager.cpp
--- a/Engine/Source/Runtime/Renderer/Private/VirtualShadowMaps/VirtualShadowMapCacheManager.cpp
+++ b/Engine/Source/Runtime/Renderer/Private/VirtualShadowMaps/VirtualShadowMapCacheManager.cpp
@@ -11,6 +11,7 @@
 #include "HAL/FileManager.h"
 
 #include "PrimitiveSceneInfo.h"
+#include "ShaderDebug.h"
 
 static TAutoConsoleVariable<int32> CVarAccumulateStats(
 	TEXT("r.Shadow.Virtual.AccumulateStats"),
@@ -26,6 +27,12 @@ static TAutoConsoleVariable<int32> CVarCacheVirtualSMs(
 	ECVF_RenderThreadSafe
 );
 
+static TAutoConsoleVariable<int32> CVarDrawInvalidatingBounds(
+	TEXT("r.Shadow.Virtual.DrawInvalidatingBounds"),
+	0,
+	TEXT("Turn on debug render cache invalidating instance bounds."),
+	ECVF_RenderThreadSafe
+);
 
 void FVirtualShadowMapCacheEntry::UpdateClipmap(
 	int32 VirtualShadowMapId,
@@ -273,7 +280,10 @@ void FVirtualShadowMapArrayCacheManager::ExtractStats(FRDGBuilder& GraphBuilder,
 	// Note: stats accumulation thing is here because it needs to persist over frames.
 	if (!AccumulatedStatsBuffer.IsValid())
 	{
-		AccumulatedStatsBufferRDG = GraphBuilder.CreateBuffer(FRDGBufferDesc::CreateBufferDesc(4, 1 + FVirtualShadowMapArray::NumStats * MaxStatFrames), TEXT("Shadow.Virtual.AccumulatedStatsBuffer"));	// TODO: Can't be a structured buffer as EnqueueCopy is only defined for vertex buffers
+		FRDGBufferDesc Desc = FRDGBufferDesc::CreateBufferDesc(4, 1 + FVirtualShadowMapArray::NumStats * MaxStatFrames);
+		Desc.Usage = EBufferUsageFlags(Desc.Usage | BUF_SourceCopy);
+
+		AccumulatedStatsBufferRDG = GraphBuilder.CreateBuffer(Desc, TEXT("Shadow.Virtual.AccumulatedStatsBuffer"));	// TODO: Can't be a structured buffer as EnqueueCopy is only defined for vertex buffers
 		AddClearUAVPass(GraphBuilder, GraphBuilder.CreateUAV(AccumulatedStatsBufferRDG, PF_R32_UINT), 0);
 		AccumulatedStatsBuffer = GraphBuilder.ConvertToExternalBuffer(AccumulatedStatsBufferRDG);
 	}
@@ -404,6 +414,8 @@ void FVirtualShadowMapArrayCacheManager::ProcessRemovedPrimives(FRDGBuilder& Gra
 {
 	if (CVarCacheVirtualSMs.GetValueOnRenderThread() != 0 && RemovedPrimitiveSceneInfos.Num() > 0 && PrevBuffers.DynamicCasterPageFlags.IsValid())
 	{
+		RDG_EVENT_SCOPE(GraphBuilder, "Shadow.Virtual.ProcessRemovedPrimives [%d]", RemovedPrimitiveSceneInfos.Num());
+
 		TArray<FInstanceSceneDataRange, SceneRenderingAllocator> InstanceRangesLarge;
 		TArray<FInstanceSceneDataRange, SceneRenderingAllocator> InstanceRangesSmall;
 
@@ -432,6 +444,8 @@ void FVirtualShadowMapArrayCacheManager::ProcessPrimitivesToUpdate(FRDGBuilder&
 	const FGPUScene& GPUScene = Scene.GPUScene;
 	if (IsValid() && GPUScene.PrimitivesToUpdate.Num() > 0)
 	{
+		RDG_EVENT_SCOPE(GraphBuilder, "Shadow.Virtual.ProcessPrimitivesToUpdate [%d]", GPUScene.PrimitivesToUpdate.Num());
+
 		// TODO: As a slight CPU optimization just pass primitive ID list and use instance ranges stored in GPU scene
 		TArray<FInstanceSceneDataRange, SceneRenderingAllocator> InstanceRangesLarge;
 		TArray<FInstanceSceneDataRange, SceneRenderingAllocator> InstanceRangesSmall;
@@ -445,7 +459,13 @@ void FVirtualShadowMapArrayCacheManager::ProcessPrimitivesToUpdate(FRDGBuilder&
 
 				// SKIP if marked for Add, because this means it has no previous location to invalidate.
 				// SKIP if transform has not changed, as this means no invalidation needs to take place.
-				if (!EnumHasAnyFlags(PrimitiveDirtyState, EPrimitiveDirtyState::Added) && EnumHasAnyFlags(PrimitiveDirtyState, EPrimitiveDirtyState::ChangedTransform))
+#if UE_BUILD_SHIPPING
+				const bool bDoInvalidation = !EnumHasAnyFlags(PrimitiveDirtyState, EPrimitiveDirtyState::Added) && EnumHasAnyFlags(PrimitiveDirtyState, EPrimitiveDirtyState::ChangedTransform);
+#else
+				// Don't skip added if invadiation debug is enabled, since we want to render those bounds also 
+				const bool bDoInvalidation = (CVarDrawInvalidatingBounds.GetValueOnRenderThread() != 0 || !EnumHasAnyFlags(PrimitiveDirtyState, EPrimitiveDirtyState::Added)) && EnumHasAnyFlags(PrimitiveDirtyState, EPrimitiveDirtyState::ChangedTransform);
+#endif
+				if (bDoInvalidation)
 				{
 					const FPrimitiveSceneInfo* PrimitiveSceneInfo = Scene.Primitives[PrimitiveId];
 					if (PrimitiveSceneInfo->GetInstanceSceneDataOffset() != INDEX_NONE)
@@ -475,12 +495,16 @@ class FVirtualSmInvalidateInstancePagesCS : public FGlobalShader
 	DECLARE_GLOBAL_SHADER(FVirtualSmInvalidateInstancePagesCS);
 	SHADER_USE_PARAMETER_STRUCT(FVirtualSmInvalidateInstancePagesCS, FGlobalShader)
 
+	class FDebugDim : SHADER_PERMUTATION_BOOL("ENABLE_DEBUG_MODE");
 	class FLargeSmallDim : SHADER_PERMUTATION_BOOL("PROCESS_LARGE_INSTANCE_COUNT_RANGES");
-	using FPermutationDomain = TShaderPermutationDomain<FLargeSmallDim>;
+	using FPermutationDomain = TShaderPermutationDomain<FDebugDim,FLargeSmallDim>;
 
 public:
 	BEGIN_SHADER_PARAMETER_STRUCT(FParameters, )
 		SHADER_PARAMETER_RDG_UNIFORM_BUFFER(FVirtualShadowMapUniformParameters, VirtualShadowMap)
+		SHADER_PARAMETER_STRUCT_INCLUDE(ShaderDrawDebug::FShaderParameters, ShaderDrawUniformBuffer)
+		SHADER_PARAMETER(uint32, bDrawBounds)
+
 		SHADER_PARAMETER_RDG_BUFFER_SRV(StructuredBuffer< FInstanceSceneDataRange >, InstanceSceneRanges)
 		SHADER_PARAMETER(uint32, NumRemovedItems)
 		SHADER_PARAMETER_RDG_BUFFER_SRV(StructuredBuffer< uint >, PageFlags)
@@ -525,11 +549,29 @@ TRDGUniformBufferRef<FVirtualShadowMapUniformParameters> FVirtualShadowMapArrayC
 
 void FVirtualShadowMapArrayCacheManager::ProcessInstanceRangeInvalidation(FRDGBuilder& GraphBuilder, const TArray<FInstanceSceneDataRange, SceneRenderingAllocator>& InstanceRangesLarge, const TArray<FInstanceSceneDataRange, SceneRenderingAllocator>& InstanceRangesSmall, const FGPUScene& GPUScene)
 {
+	if (InstanceRangesSmall.IsEmpty() && InstanceRangesLarge.IsEmpty())
+	{
+		return;
+	}
+
+
 	auto RegExtCreateSrv = [&GraphBuilder](const TRefCountPtr<FRDGPooledBuffer>& Buffer, const TCHAR* Name) -> FRDGBufferSRVRef
 	{
 		return GraphBuilder.CreateSRV(GraphBuilder.RegisterExternalBuffer(Buffer, Name));
 	};
 
+	const bool bDrawBounds = CVarDrawInvalidatingBounds.GetValueOnRenderThread() != 0;
+
+
+	if (bDrawBounds)
+	{
+		ShaderDrawDebug::SetEnabled(true);
+		ShaderDrawDebug::RequestSpaceForElements((InstanceRangesSmall.Num() + InstanceRangesLarge.Num()) * 12);
+	}
+
+	// Note: this disables the whole debug permutation since the parameters must be bound.
+	const bool bUseDebugPermutation = bDrawBounds && ShaderDrawDebug::IsDefaultViewEnabled();
+
 	// Update references in our last frame uniform buffer with reimported resources for this frame
 	PrevUniformParameters.ProjectionData = RegExtCreateSrv(PrevBuffers.ShadowMapProjectionDataBuffer, TEXT("Shadow.Virtual.PrevProjectionData"));
 	PrevUniformParameters.PageTable = RegExtCreateSrv(PrevBuffers.PageTable, TEXT("Shadow.Virtual.PrevPageTable"));
@@ -539,31 +581,40 @@ void FVirtualShadowMapArrayCacheManager::ProcessInstanceRangeInvalidation(FRDGBu
 	FRDGBufferRef InstanceRangesSmallRDG = !InstanceRangesSmall.IsEmpty() ? CreateStructuredBuffer(GraphBuilder, TEXT("Shadow.Virtual.InstanceRangesSmall"), InstanceRangesSmall) : nullptr;
 	FRDGBufferRef InstanceRangesLargeRDG = !InstanceRangesLarge.IsEmpty() ? CreateStructuredBuffer(GraphBuilder, TEXT("Shadow.Virtual.InstanceRangesSmall"), InstanceRangesLarge) : nullptr;
 
+
+	FVirtualSmInvalidateInstancePagesCS::FParameters PassParametersTmp;
+	PassParametersTmp.VirtualShadowMap = GetPreviousUniformBuffer(GraphBuilder);
+
+	PassParametersTmp.PageFlags = RegExtCreateSrv(PrevBuffers.PageFlags, TEXT("Shadow.Virtual.PrevPageFlags"));
+	PassParametersTmp.HPageFlags = RegExtCreateSrv(PrevBuffers.HPageFlags, TEXT("Shadow.Virtual.PrevHPageFlags"));
+	PassParametersTmp.PageRectBounds = RegExtCreateSrv(PrevBuffers.PageRectBounds, TEXT("Shadow.Virtual.PrevPageRectBounds"));
+
+	FRDGBufferRef DynamicCasterFlagsRDG = GraphBuilder.RegisterExternalBuffer(PrevBuffers.DynamicCasterPageFlags, TEXT("Shadow.Virtual.PrevDynamicCasterFlags"));
+	PassParametersTmp.OutDynamicCasterFlags = GraphBuilder.CreateUAV(DynamicCasterFlagsRDG);
+
+	PassParametersTmp.GPUSceneInstanceSceneData = GPUScene.InstanceSceneDataBuffer.SRV;
+	PassParametersTmp.GPUScenePrimitiveSceneData = GPUScene.PrimitiveBuffer.SRV;
+	PassParametersTmp.GPUSceneInstancePayloadData = GPUScene.InstancePayloadDataBuffer.SRV;
+	PassParametersTmp.GPUSceneFrameNumber = GPUScene.GetSceneFrameNumber();
+	PassParametersTmp.InstanceSceneDataSOAStride = GPUScene.InstanceSceneDataSOAStride;
+	PassParametersTmp.bDrawBounds = bDrawBounds;
+
+	if (bUseDebugPermutation)
+	{
+		ShaderDrawDebug::SetParameters(GraphBuilder, PassParametersTmp.ShaderDrawUniformBuffer);
+	}
+
+	FVirtualSmInvalidateInstancePagesCS::FPermutationDomain PermutationVector;
+	PermutationVector.Set<FVirtualSmInvalidateInstancePagesCS::FDebugDim>(bUseDebugPermutation);
 	if (InstanceRangesSmall.Num())
 	{
 		RDG_EVENT_SCOPE(GraphBuilder, "ProcessInstanceRangeInvalidation [%d small-ranges]", InstanceRangesSmall.Num());
 
 		FVirtualSmInvalidateInstancePagesCS::FParameters* PassParameters = GraphBuilder.AllocParameters<FVirtualSmInvalidateInstancePagesCS::FParameters>();
-
-		PassParameters->VirtualShadowMap = GetPreviousUniformBuffer(GraphBuilder);
-		
+		*PassParameters = PassParametersTmp;
 		PassParameters->InstanceSceneRanges = GraphBuilder.CreateSRV(InstanceRangesSmallRDG);
 		PassParameters->NumRemovedItems = InstanceRangesSmall.Num();
 
-		PassParameters->PageFlags = RegExtCreateSrv(PrevBuffers.PageFlags, TEXT("Shadow.Virtual.PrevPageFlags"));
-		PassParameters->HPageFlags = RegExtCreateSrv(PrevBuffers.HPageFlags, TEXT("Shadow.Virtual.PrevHPageFlags"));
-		PassParameters->PageRectBounds = RegExtCreateSrv(PrevBuffers.PageRectBounds, TEXT("Shadow.Virtual.PrevPageRectBounds"));
-
-		FRDGBufferRef DynamicCasterFlagsRDG = GraphBuilder.RegisterExternalBuffer(PrevBuffers.DynamicCasterPageFlags, TEXT("Shadow.Virtual.PrevDynamicCasterFlags"));
-		PassParameters->OutDynamicCasterFlags = GraphBuilder.CreateUAV(DynamicCasterFlagsRDG);
-
-		PassParameters->GPUSceneInstanceSceneData = GPUScene.InstanceSceneDataBuffer.SRV;
-		PassParameters->GPUSceneInstancePayloadData = GPUScene.InstancePayloadDataBuffer.SRV;
-		PassParameters->GPUScenePrimitiveSceneData = GPUScene.PrimitiveBuffer.SRV;
-		PassParameters->GPUSceneFrameNumber = GPUScene.GetSceneFrameNumber();
-		PassParameters->InstanceSceneDataSOAStride = GPUScene.InstanceSceneDataSOAStride;
-
-		FVirtualSmInvalidateInstancePagesCS::FPermutationDomain PermutationVector;
 		PermutationVector.Set<FVirtualSmInvalidateInstancePagesCS::FLargeSmallDim>(0);
 
 		auto ComputeShader = GetGlobalShaderMap(GMaxRHIFeatureLevel)->GetShader<FVirtualSmInvalidateInstancePagesCS>(PermutationVector);
@@ -581,25 +632,10 @@ void FVirtualShadowMapArrayCacheManager::ProcessInstanceRangeInvalidation(FRDGBu
 		RDG_EVENT_SCOPE(GraphBuilder, "ProcessInstanceRangeInvalidation [%d large-ranges]", InstanceRangesLarge.Num());
 
 		FVirtualSmInvalidateInstancePagesCS::FParameters* PassParameters = GraphBuilder.AllocParameters<FVirtualSmInvalidateInstancePagesCS::FParameters>();
-
-		PassParameters->VirtualShadowMap = GetPreviousUniformBuffer(GraphBuilder);
+		*PassParameters = PassParametersTmp;
 		PassParameters->InstanceSceneRanges = GraphBuilder.CreateSRV(InstanceRangesLargeRDG);
 		PassParameters->NumRemovedItems = InstanceRangesLarge.Num();
 
-		PassParameters->PageFlags = RegExtCreateSrv(PrevBuffers.PageFlags, TEXT("Shadow.Virtual.PrevPageFlags"));
-		PassParameters->HPageFlags = RegExtCreateSrv(PrevBuffers.HPageFlags, TEXT("Shadow.Virtual.PrevHPageFlags"));
-		PassParameters->PageRectBounds = RegExtCreateSrv(PrevBuffers.PageRectBounds, TEXT("Shadow.Virtual.PrevPageRectBounds"));
-
-		FRDGBufferRef DynamicCasterFlagsRDG = GraphBuilder.RegisterExternalBuffer(PrevBuffers.DynamicCasterPageFlags, TEXT("Shadow.Virtual.PrevDynamicCasterFlags"));
-		PassParameters->OutDynamicCasterFlags = GraphBuilder.CreateUAV(DynamicCasterFlagsRDG);
-
-		PassParameters->GPUSceneInstanceSceneData = GPUScene.InstanceSceneDataBuffer.SRV;
-		PassParameters->GPUSceneInstancePayloadData = GPUScene.InstancePayloadDataBuffer.SRV;
-		PassParameters->GPUScenePrimitiveSceneData = GPUScene.PrimitiveBuffer.SRV;
-		PassParameters->GPUSceneFrameNumber = GPUScene.GetSceneFrameNumber();
-		PassParameters->InstanceSceneDataSOAStride = GPUScene.InstanceSceneDataSOAStride;
-
-		FVirtualSmInvalidateInstancePagesCS::FPermutationDomain PermutationVector;
 		PermutationVector.Set<FVirtualSmInvalidateInstancePagesCS::FLargeSmallDim>(1);
 
 		auto ComputeShader = GetGlobalShaderMap(GMaxRHIFeatureLevel)->GetShader<FVirtualSmInvalidateInstancePagesCS>(PermutationVector);
diff --git a/Engine/Source/Runtime/Renderer/Public/ShaderDebug.h b/Engine/Source/Runtime/Renderer/Public/ShaderDebug.h
--- a/Engine/Source/Runtime/Renderer/Public/ShaderDebug.h
+++ b/Engine/Source/Runtime/Renderer/Public/ShaderDebug.h
@@ -28,6 +28,13 @@ namespace ShaderDrawDebug
 	RENDERER_API void SetMaxElementCount(uint32 MaxCount);
 	RENDERER_API uint32 GetMaxElementCount();
 
+	/**
+	 * Call to ensure enough space for some number of elements, is added cumulatively each frame, to make it possible for several systems to request a certain number independently.
+	 * Is used to grow the max element count for subsequent frames (as the allocation happens early in the frame). 
+	 * @param The number of elements requested, an element corresponds to a line, so a cube, for example, needs 12 elements.
+	 */
+	RENDERER_API void RequestSpaceForElements(uint32 MaxElementCount);
+
 	// Call this to know if a view can render this debug information
 	bool IsEnabled(const FViewInfo& View);
 
@@ -48,4 +55,9 @@ namespace ShaderDrawDebug
 
 	// Call this to fill the FShaderDrawParameters
 	RENDERER_API void SetParameters(FRDGBuilder& GraphBuilder, const FShaderDrawDebugData& Data, FShaderParameters& OutParameters);
+
+	// Returns true if the default view exists and has shader debug rendering enabled (this needs to be checked before using a permutation that requires the shader draw parameters)
+	bool IsDefaultViewEnabled();
+	// Call this to fill the FShaderDrawParameters using the default view (the first one for which Begin was called in case of stereo or similar)
+	RENDERER_API void SetParameters(FRDGBuilder& GraphBuilder, FShaderParameters& OutParameters);
 }
