Fix some LWC related fallout in mobile rendering

diff --git a/Engine/Source/Runtime/Renderer/Private/MobileBasePass.cpp b/Engine/Source/Runtime/Renderer/Private/MobileBasePass.cpp
--- a/Engine/Source/Runtime/Renderer/Private/MobileBasePass.cpp
+++ b/Engine/Source/Runtime/Renderer/Private/MobileBasePass.cpp
@@ -511,7 +511,7 @@ void TMobileBasePassPSPolicyParamType<FUniformLightMapPolicy>::GetShaderBindings
 			FVector4 CapturePositions[MaxNumReflections] = { FVector4(0, 0, 0, 0), FVector4(0, 0, 0, 0), FVector4(0, 0, 0, 0) };
 			FVector4 ReflectionParams(0.0f, 0.0f, 0.0f, 0.0f);
 			FVector4 ReflectanceMaxValueRGBMParams(0.0f, 0.0f, 0.0f, 0.0f);
-			FMatrix CaptureBoxTransformArray[MaxNumReflections] = { FMatrix(EForceInit::ForceInitToZero), FMatrix(EForceInit::ForceInitToZero), FMatrix(EForceInit::ForceInitToZero) };
+			FMatrix44f CaptureBoxTransformArray[MaxNumReflections] = { FMatrix44f(EForceInit::ForceInitToZero), FMatrix44f(EForceInit::ForceInitToZero), FMatrix44f(EForceInit::ForceInitToZero) };
 			FVector4 CaptureBoxScalesArray[MaxNumReflections] = { FVector4(EForceInit::ForceInitToZero), FVector4(EForceInit::ForceInitToZero), FVector4(EForceInit::ForceInitToZero) };
 			FPrimitiveSceneInfo* PrimitiveSceneInfo = PrimitiveSceneProxy ? PrimitiveSceneProxy->GetPrimitiveSceneInfo() : nullptr;
 			if (PrimitiveSceneInfo)
