Revert the changes from CL#17258939, since it causes regression on mobile.

diff --git a/Engine/Source/Runtime/Engine/Private/Components/LightComponent.cpp b/Engine/Source/Runtime/Engine/Private/Components/LightComponent.cpp
--- a/Engine/Source/Runtime/Engine/Private/Components/LightComponent.cpp
+++ b/Engine/Source/Runtime/Engine/Private/Components/LightComponent.cpp
@@ -359,6 +359,10 @@ FLightSceneProxy::FLightSceneProxy(const ULightComponent* InLightComponent)
 	, ShadowAmount(1.0f)
 	, SamplesPerPixel(1)
 	, DeepShadowLayerDistribution(InLightComponent->DeepShadowLayerDistribution)
+	, bMobileMovablePointLightUniformBufferNeedsUpdate(false)
+	, bMobileMovablePointLightShouldBeRender(false)
+	, bMobileMovablePointLightShouldCastShadow(false)
+	, MobileMovablePointLightShadowmapMinMax()
 {
 	check(SceneInterface);
 
diff --git a/Engine/Source/Runtime/Engine/Private/LocalLightSceneProxy.cpp b/Engine/Source/Runtime/Engine/Private/LocalLightSceneProxy.cpp
--- a/Engine/Source/Runtime/Engine/Private/LocalLightSceneProxy.cpp
+++ b/Engine/Source/Runtime/Engine/Private/LocalLightSceneProxy.cpp
@@ -141,21 +141,6 @@ void FLocalLightSceneProxy::UpdateRadius(float ComponentRadius)
 
 	// Min to avoid div by 0 (NaN in InvRadius)
 	InvRadius = 1.0f / FMath::Max(0.00001f, ComponentRadius);
-}
-
-bool FLocalLightSceneProxy::IsLocalLight() const
-{
-	return true;
-}
 
-void FLocalLightSceneProxy::UpdateMobileMovableLocalLightUniformBuffer(const FMobileMovablePointLightUniformShaderParameters& MobileMovableLocalLightUniformShaderParameters)
-{
-	if (MobileMovableLocalLightUniformBuffer == nullptr)
-	{
-		MobileMovableLocalLightUniformBuffer = TUniformBufferRef<FMobileMovablePointLightUniformShaderParameters>::CreateUniformBufferImmediate(MobileMovableLocalLightUniformShaderParameters, UniformBuffer_MultiFrame);
-	}
-	else
-	{
-		MobileMovableLocalLightUniformBuffer.UpdateUniformBufferImmediate(MobileMovableLocalLightUniformShaderParameters);
-	}
+	bMobileMovablePointLightUniformBufferNeedsUpdate = true;
 }
\ No newline at end of file
diff --git a/Engine/Source/Runtime/Engine/Public/LocalLightSceneProxy.h b/Engine/Source/Runtime/Engine/Public/LocalLightSceneProxy.h
--- a/Engine/Source/Runtime/Engine/Public/LocalLightSceneProxy.h
+++ b/Engine/Source/Runtime/Engine/Public/LocalLightSceneProxy.h
@@ -53,15 +53,6 @@ public:
 
 	virtual bool GetPerObjectProjectedShadowInitializer(const FBoxSphereBounds& SubjectBounds, class FPerObjectProjectedShadowInitializer& OutInitializer) const override;
 
-	virtual bool IsLocalLight() const override;
-
-	FORCEINLINE FRHIUniformBuffer* GetMobileMovableLocalLightUniformBufferRHI() const
-	{
-		return MobileMovableLocalLightUniformBuffer.GetReference();
-	}
-
-	virtual void UpdateMobileMovableLocalLightUniformBuffer(const FMobileMovablePointLightUniformShaderParameters& MobileMovableLocalLightUniformShaderParameters);
-
 protected:
 
 	/** Updates the light scene info's radius from the component. */
@@ -69,7 +60,4 @@ protected:
 
 	float MaxDrawDistance;
 	float FadeRange;
-
-	/** The movable local light's uniform buffer for mobile. */
-	TUniformBufferRef<FMobileMovablePointLightUniformShaderParameters> MobileMovableLocalLightUniformBuffer;
 };
diff --git a/Engine/Source/Runtime/Engine/Public/SceneManagement.h b/Engine/Source/Runtime/Engine/Public/SceneManagement.h
--- a/Engine/Source/Runtime/Engine/Public/SceneManagement.h
+++ b/Engine/Source/Runtime/Engine/Public/SceneManagement.h
@@ -1286,20 +1286,42 @@ BEGIN_GLOBAL_SHADER_PARAMETER_STRUCT(FMobileMovablePointLightUniformShaderParame
 	SHADER_PARAMETER(FMatrix44f, SpotLightShadowWorldToShadowMatrix)
 END_GLOBAL_SHADER_PARAMETER_STRUCT()
 
-FORCEINLINE FMobileMovablePointLightUniformShaderParameters GetDummyMovablePointLightUniformShaderParameters()
+/** Initializes the movable point light uniform shader parameters. */
+FORCEINLINE FMobileMovablePointLightUniformShaderParameters GetMovablePointLightUniformShaderParameters(
+	const FVector4& LightPositionAndInvRadius,
+	const FVector4& LightColorAndFalloffExponent,
+	const FVector4& SpotLightDirectionAndSpecularScale,
+	const FVector4& SpotLightAnglesAndSoftTransitionScaleAndLightShadowType,
+	const FVector4& SpotLightShadowSharpenAndShadowFadeFraction,
+	const FVector4& SpotLightShadowmapMinMax,
+	const FMatrix& SpotLightShadowWorldToShadowMatrix
+)
 {
 	FMobileMovablePointLightUniformShaderParameters Result;
-	Result.LightPositionAndInvRadius = FVector4();
-	Result.LightColorAndFalloffExponent = FVector4();
-	Result.SpotLightDirectionAndSpecularScale = FVector4();
-	Result.SpotLightAnglesAndSoftTransitionScaleAndLightShadowType = FVector4();
-	Result.SpotLightShadowSharpenAndShadowFadeFraction = FVector4();
-	Result.SpotLightShadowmapMinMax = FVector4();
-	Result.SpotLightShadowWorldToShadowMatrix = (FMatrix44f)FMatrix();
+	Result.LightPositionAndInvRadius = LightPositionAndInvRadius;
+	Result.LightColorAndFalloffExponent = LightColorAndFalloffExponent;
+	Result.SpotLightDirectionAndSpecularScale = SpotLightDirectionAndSpecularScale;
+	Result.SpotLightAnglesAndSoftTransitionScaleAndLightShadowType = SpotLightAnglesAndSoftTransitionScaleAndLightShadowType;
+	Result.SpotLightShadowSharpenAndShadowFadeFraction = SpotLightShadowSharpenAndShadowFadeFraction;
+	Result.SpotLightShadowmapMinMax = SpotLightShadowmapMinMax;
+	Result.SpotLightShadowWorldToShadowMatrix = (FMatrix44f)SpotLightShadowWorldToShadowMatrix;
 
 	return Result;
 }
 
+FORCEINLINE FMobileMovablePointLightUniformShaderParameters GetDummyMovablePointLightUniformShaderParameters()
+{
+	return GetMovablePointLightUniformShaderParameters(
+		FVector4(),
+		FVector4(),
+		FVector4(),
+		FVector4(),
+		FVector4(),
+		FVector4(),
+		FMatrix()
+	);
+}
+
 /**
  * Dummy mobile movable point light uniform buffer.
  */
@@ -1355,7 +1377,6 @@ public:
 	virtual float GetSourceRadius() const { return 0.0f; }
 	virtual bool IsInverseSquared() const { return true; }
 	virtual bool IsRectLight() const { return false; }
-	virtual bool  IsLocalLight() const { return false; }
 	virtual bool HasSourceTexture() const { return false; }
 	virtual float GetLightSourceAngle() const { return 0.0f; }
 	virtual float GetShadowSourceAngleFactor() const { return 1.0f; }
@@ -1561,6 +1582,16 @@ public:
 	virtual FLinearColor GetCloudScatteredLuminanceScale() const { return FLinearColor::White; }
 	virtual bool GetUsePerPixelAtmosphereTransmittance() const { return false; }
 
+	FORCEINLINE void SetMobileMovablePointLightUniformBufferNeedsUpdate(bool bInMobileMovablePointLightUniformBufferNeedsUpdate)
+	{
+		bMobileMovablePointLightUniformBufferNeedsUpdate = bInMobileMovablePointLightUniformBufferNeedsUpdate;
+	}
+
+	FORCEINLINE FRHIUniformBuffer* GetMobileMovablePointLightUniformBufferRHI() const
+	{
+		return MobileMovablePointLightUniformBuffer.GetReference();
+	}
+
 protected:
 
 	friend class FScene;
@@ -1745,6 +1776,21 @@ protected:
 	/** Deep shadow layer distribution. */
 	float DeepShadowLayerDistribution;
 
+	/** If this is TRUE, the light's mobile movable point light uniform buffer needs to be updated before it can be used for mobile base pass rendering. */
+	bool bMobileMovablePointLightUniformBufferNeedsUpdate;
+
+	/** Cached ShouldBeRender for mobile, since if the ShouldBeRender is changed we have to update the movable point lights uniform buffer. */
+	bool bMobileMovablePointLightShouldBeRender;
+
+	/** Cached DynamicShadows show flag for mobile, since if the show flag is changed we have to update the movable point lights uniform buffer. */
+	bool bMobileMovablePointLightShouldCastShadow;
+
+	/** Cached the spotlight shadow map min and max value for mobile, since if the value is changed we have to update the movable point lights uniform buffer. */
+	FVector4 MobileMovablePointLightShadowmapMinMax;
+
+	/** The movable point light's uniform buffer for mobile. */
+	TUniformBufferRef<FMobileMovablePointLightUniformShaderParameters> MobileMovablePointLightUniformBuffer;
+
 	/**
 	 * Updates the light proxy's cached transforms.
 	 * @param InLightToWorld - The new light-to-world transform.
diff --git a/Engine/Source/Runtime/Renderer/Private/LightSceneInfo.cpp b/Engine/Source/Runtime/Renderer/Private/LightSceneInfo.cpp
--- a/Engine/Source/Runtime/Renderer/Private/LightSceneInfo.cpp
+++ b/Engine/Source/Runtime/Renderer/Private/LightSceneInfo.cpp
@@ -9,7 +9,6 @@
 #include "SceneCore.h"
 #include "ScenePrivate.h"
 #include "DistanceFieldLightingShared.h"
-#include "LocalLightSceneProxy.h"
 
 int32 GWholeSceneShadowUnbuiltInteractionThreshold = 500;
 static FAutoConsoleVariableRef CVarWholeSceneShadowUnbuiltInteractionThreshold(
@@ -121,9 +120,8 @@ void FLightSceneInfo::AddToScene()
 
 			if (bIsValidLightTypeMobile)
 			{
-				check(Proxy->IsLocalLight());
-				FLocalLightSceneProxy* LocalLightSceneProxy = (FLocalLightSceneProxy*)Proxy;
-				LocalLightSceneProxy->UpdateMobileMovableLocalLightUniformBuffer(GetDummyMovablePointLightUniformShaderParameters());
+				Proxy->MobileMovablePointLightUniformBuffer = TUniformBufferRef<FMobileMovablePointLightUniformShaderParameters>::CreateUniformBufferImmediate(GetDummyMovablePointLightUniformShaderParameters(), UniformBuffer_MultiFrame);
+				Proxy->bMobileMovablePointLightUniformBufferNeedsUpdate = true;
 			}
 		}
 	}
@@ -245,7 +243,7 @@ FLightPrimitiveInteraction* FLightSceneInfo::GetDynamicInteractionStaticPrimitiv
 	return DynamicInteractionStaticPrimitiveList;
 }
 
-void FLightSceneInfo::UpdateMobileMovableLocalLightUniformBuffer(const FSceneRenderer* SceneRenderer)
+void FLightSceneInfo::ConditionalUpdateMobileMovablePointLightUniformBuffer(const FSceneRenderer* SceneRenderer)
 {
 	QUICK_SCOPE_CYCLE_COUNTER(STAT_FLightSceneProxy_UpdateMobileMovablePointLightUniformBuffer);
 
@@ -260,8 +258,13 @@ void FLightSceneInfo::UpdateMobileMovableLocalLightUniformBuffer(const FSceneRen
 
 	checkSlow(MobileNumDynamicPointLights > 0 && SceneRenderer);
 
-	check(Proxy->IsLocalLight());
-	FLocalLightSceneProxy* LocalLightSceneProxy = (FLocalLightSceneProxy*)Proxy;
+	FVector4 LightPositionAndInvRadius;
+	FVector4 LightColorAndFalloffExponent;
+	FVector4 SpotLightDirectionAndSpecularScale;
+	FVector4 SpotLightAnglesAndSoftTransitionScaleAndLightShadowType;
+	FVector4 SpotLightShadowSharpenAndShadowFadeFraction;
+	FVector4 SpotLightShadowmapMinMax;
+	FMatrix SpotLightWorldToShadowMatrix;
 
 	bool bShouldBeRender = false;
 	bool bShouldCastShadow = false;
@@ -319,37 +322,60 @@ void FLightSceneInfo::UpdateMobileMovableLocalLightUniformBuffer(const FSceneRen
 		float SoftTransitionScale = 0.0f;
 		float ShadowFadeFraction = 0.0f;
 
-		FMobileMovablePointLightUniformShaderParameters MobileMovablePointLightUniformShaderParameters;
-
 		if (bShouldCastShadow)
 		{
-			FProjectedShadowInfo *ProjectedShadowInfo = SceneRenderer->VisibleLightInfos[Id].AllProjectedShadows.Last();
+			FProjectedShadowInfo* ProjectedShadowInfo = SceneRenderer->VisibleLightInfos[Id].AllProjectedShadows.Last();
 			checkSlow(ProjectedShadowInfo && ProjectedShadowInfo->CacheMode != SDCM_StaticPrimitivesOnly);
-			
+
 			const float TransitionSize = ProjectedShadowInfo->ComputeTransitionSize();
 			checkSlow(TransitionSize > 0.0f);
 
 			SoftTransitionScale = 1.0f / TransitionSize;
-			
+
 			for (int32 ViewIndex = 0; ViewIndex < SceneRenderer->Views.Num(); ++ViewIndex)
 			{
 				ShadowFadeFraction = FMath::Max(ShadowFadeFraction, ProjectedShadowInfo->FadeAlphas[ViewIndex]);
 			}
 
-			MobileMovablePointLightUniformShaderParameters.SpotLightShadowSharpenAndShadowFadeFraction = FVector4(Proxy->GetShadowSharpen() * 7.0f + 1.0f, ShadowFadeFraction, ProjectedShadowInfo->GetShaderReceiverDepthBias(), 0.0f);
-			MobileMovablePointLightUniformShaderParameters.SpotLightShadowWorldToShadowMatrix = ProjectedShadowInfo->GetWorldToShadowMatrix(MobileMovablePointLightUniformShaderParameters.SpotLightShadowmapMinMax);
+			SpotLightShadowSharpenAndShadowFadeFraction = FVector4(Proxy->GetShadowSharpen() * 7.0f + 1.0f, ShadowFadeFraction, ProjectedShadowInfo->GetShaderReceiverDepthBias(), 0.0f);
+			SpotLightWorldToShadowMatrix = ProjectedShadowInfo->GetWorldToShadowMatrix(SpotLightShadowmapMinMax);
 		}
-		
-		MobileMovablePointLightUniformShaderParameters.LightPositionAndInvRadius = FVector4(LightParameters.Position, LightParameters.InvRadius);
-		MobileMovablePointLightUniformShaderParameters.LightColorAndFalloffExponent = FVector4(LightParameters.Color, LightParameters.FalloffExponent);
-		MobileMovablePointLightUniformShaderParameters.SpotLightDirectionAndSpecularScale = FVector4(LightParameters.Direction.X, LightParameters.Direction.Y, LightParameters.Direction.Z, Proxy->GetSpecularScale());
-		MobileMovablePointLightUniformShaderParameters.SpotLightAnglesAndSoftTransitionScaleAndLightShadowType = FVector4(LightParameters.SpotAngles.X, LightParameters.SpotAngles.Y, SoftTransitionScale, LightShadowType);
 
-		LocalLightSceneProxy->UpdateMobileMovableLocalLightUniformBuffer(MobileMovablePointLightUniformShaderParameters);
+		LightPositionAndInvRadius = FVector4(LightParameters.Position, LightParameters.InvRadius);
+		LightColorAndFalloffExponent = FVector4(LightParameters.Color, LightParameters.FalloffExponent);
+		SpotLightDirectionAndSpecularScale = FVector4(LightParameters.Direction.X, LightParameters.Direction.Y, LightParameters.Direction.Z, Proxy->GetSpecularScale());
+		SpotLightAnglesAndSoftTransitionScaleAndLightShadowType = FVector4(LightParameters.SpotAngles.X, LightParameters.SpotAngles.Y, SoftTransitionScale, LightShadowType);
 	}
-	else
+
+	if (bShouldBeRender != Proxy->bMobileMovablePointLightShouldBeRender ||
+		bShouldCastShadow != Proxy->bMobileMovablePointLightShouldCastShadow ||
+		SpotLightShadowmapMinMax != Proxy->MobileMovablePointLightShadowmapMinMax)
+	{
+		Proxy->bMobileMovablePointLightUniformBufferNeedsUpdate = true;
+
+		Proxy->bMobileMovablePointLightShouldBeRender = bShouldBeRender;
+
+		Proxy->bMobileMovablePointLightShouldCastShadow = bShouldCastShadow;
+
+		Proxy->MobileMovablePointLightShadowmapMinMax = SpotLightShadowmapMinMax;
+	}
+
+	if (Proxy->bMobileMovablePointLightUniformBufferNeedsUpdate)
 	{
-		LocalLightSceneProxy->UpdateMobileMovableLocalLightUniformBuffer(GetDummyMovablePointLightUniformShaderParameters());
+		const FMobileMovablePointLightUniformShaderParameters MobileMovablePointLightUniformShaderParameters =
+			GetMovablePointLightUniformShaderParameters(
+				LightPositionAndInvRadius,
+				LightColorAndFalloffExponent,
+				SpotLightDirectionAndSpecularScale,
+				SpotLightAnglesAndSoftTransitionScaleAndLightShadowType,
+				SpotLightShadowSharpenAndShadowFadeFraction,
+				SpotLightShadowmapMinMax,
+				SpotLightWorldToShadowMatrix
+			);
+
+		Proxy->MobileMovablePointLightUniformBuffer.UpdateUniformBufferImmediate(MobileMovablePointLightUniformShaderParameters);
+
+		Proxy->bMobileMovablePointLightUniformBufferNeedsUpdate = false;
 	}
 }
 
diff --git a/Engine/Source/Runtime/Renderer/Private/LightSceneInfo.h b/Engine/Source/Runtime/Renderer/Private/LightSceneInfo.h
--- a/Engine/Source/Runtime/Renderer/Private/LightSceneInfo.h
+++ b/Engine/Source/Runtime/Renderer/Private/LightSceneInfo.h
@@ -340,8 +340,8 @@ public:
 		return (uint32)LightSceneInfo->Id;
 	}
 
-	// Update the mobile movable local light uniform buffer before it is used for mobile base pass rendering.
-	void UpdateMobileMovableLocalLightUniformBuffer(const class FSceneRenderer* SceneRenderer);
+	// Update the mobile movable point light uniform buffer before it is used for mobile base pass rendering.
+	void ConditionalUpdateMobileMovablePointLightUniformBuffer(const class FSceneRenderer* SceneRenderer);
 
 	bool ShouldRecordShadowSubjectsForMobile() const;
 };
diff --git a/Engine/Source/Runtime/Renderer/Private/MobileBasePassRendering.cpp b/Engine/Source/Runtime/Renderer/Private/MobileBasePassRendering.cpp
--- a/Engine/Source/Runtime/Renderer/Private/MobileBasePassRendering.cpp
+++ b/Engine/Source/Runtime/Renderer/Private/MobileBasePassRendering.cpp
@@ -18,7 +18,6 @@
 #include "FramePro/FrameProProfiler.h"
 #include "PostProcess/PostProcessPixelProjectedReflectionMobile.h"
 #include "Engine/SubsurfaceProfile.h"
-#include "LocalLightSceneProxy.h"
 
 // Changing this causes a full shader recompile
 static TAutoConsoleVariable<int32> CVarMobileDisableVertexFog(
@@ -164,8 +163,7 @@ FMobileBasePassMovableLightInfo::FMobileBasePassMovableLightInfo(const FPrimitiv
 				
 			if (bIsValidLightType && LightProxy->IsMovable() && (LightProxy->GetLightingChannelMask() & InSceneProxy->GetLightingChannelMask()) != 0)
 			{
-				check(LightProxy->IsLocalLight());
-				MovablePointLightUniformBuffer[NumMovablePointLights] = ((FLocalLightSceneProxy*)LightProxy)->GetMobileMovableLocalLightUniformBufferRHI();
+				MovablePointLightUniformBuffer[NumMovablePointLights] = LightProxy->GetMobileMovablePointLightUniformBufferRHI();
 
 				NumMovablePointLights++;
 			}
diff --git a/Engine/Source/Runtime/Renderer/Private/MobileShadingRenderer.cpp b/Engine/Source/Runtime/Renderer/Private/MobileShadingRenderer.cpp
--- a/Engine/Source/Runtime/Renderer/Private/MobileShadingRenderer.cpp
+++ b/Engine/Source/Runtime/Renderer/Private/MobileShadingRenderer.cpp
@@ -518,8 +518,8 @@ void FMobileSceneRenderer::InitViews(FRDGBuilder& GraphBuilder, FSceneTexturesCo
 		PrepareViewVisibilityLists();
 	}
 
-	/** Before SetupMobileBasePassAfterShadowInit, we need to update the uniform buffer and shadow info for all movable local lights.*/
-	UpdateMovableLocalLightUniformBufferAndShadowInfo();
+	/** Before SetupMobileBasePassAfterShadowInit, we need to update the uniform buffer and shadow info for all movable point lights.*/
+	UpdateMovablePointLightUniformBufferAndShadowInfo();
 
 	SetupMobileBasePassAfterShadowInit(BasePassDepthStencilAccess, ViewCommandsPerView, InstanceCullingManager);
 
@@ -1623,7 +1623,7 @@ void FMobileSceneRenderer::PreTonemapMSAA(FRHICommandListImmediate& RHICmdList,
 }
 
 /** Before SetupMobileBasePassAfterShadowInit, we need to update the uniform buffer and shadow info for all movable point lights.*/
-void FMobileSceneRenderer::UpdateMovableLocalLightUniformBufferAndShadowInfo()
+void FMobileSceneRenderer::UpdateMovablePointLightUniformBufferAndShadowInfo()
 {
 	static auto* MobileNumDynamicPointLightsCVar = IConsoleManager::Get().FindTConsoleVariableDataInt(TEXT("r.MobileNumDynamicPointLights"));
 	const int32 MobileNumDynamicPointLights = MobileNumDynamicPointLightsCVar->GetValueOnRenderThread();
@@ -1656,7 +1656,7 @@ void FMobileSceneRenderer::UpdateMovableLocalLightUniformBufferAndShadowInfo()
 
 			if (bIsValidLightType && LightProxy->IsMovable())
 			{
-				LightSceneInfo->UpdateMobileMovableLocalLightUniformBuffer(this);
+				LightSceneInfo->ConditionalUpdateMobileMovablePointLightUniformBuffer(this);
 
 				bool bDynamicShadows = bShouldDynamicShadows
 					&& LightType == LightType_Spot
diff --git a/Engine/Source/Runtime/Renderer/Private/RendererScene.cpp b/Engine/Source/Runtime/Renderer/Private/RendererScene.cpp
--- a/Engine/Source/Runtime/Renderer/Private/RendererScene.cpp
+++ b/Engine/Source/Runtime/Renderer/Private/RendererScene.cpp
@@ -1033,6 +1033,7 @@ FScene::FScene(UWorld* InWorld, bool bInRequiresHitProxies, bool bInIsEditorScen
 ,	NumUnbuiltReflectionCaptures(0)
 ,	NumMobileStaticAndCSMLights_RenderThread(0)
 ,	NumMobileMovableDirectionalLights_RenderThread(0)
+,	MobileWholeSceneShadowAtlasSize(0, 0)
 ,	GPUSkinCache(nullptr)
 ,	SceneLODHierarchy(this)
 ,	RuntimeVirtualTexturePrimitiveHideMaskEditor(0)
@@ -2875,6 +2876,8 @@ void FScene::UpdateLightColorAndBrightness(ULightComponent* Light)
 					{
 						Scene->Lights[ LightSceneInfo->Id ].Color = NewParameters.NewColor;
 					}
+
+					LightSceneInfo->Proxy->SetMobileMovablePointLightUniformBufferNeedsUpdate(true);
 				}
 			});
 	}
diff --git a/Engine/Source/Runtime/Renderer/Private/ScenePrivate.h b/Engine/Source/Runtime/Renderer/Private/ScenePrivate.h
--- a/Engine/Source/Runtime/Renderer/Private/ScenePrivate.h
+++ b/Engine/Source/Runtime/Renderer/Private/ScenePrivate.h
@@ -2897,6 +2897,9 @@ public:
 	int32 NumMobileStaticAndCSMLights_RenderThread;
 	int32 NumMobileMovableDirectionalLights_RenderThread;
 
+	/** Cached shadow atlas size for mobile, since if the resolution is changed we have to update the movable point lights uniform buffer. */
+	FIntPoint MobileWholeSceneShadowAtlasSize;
+
 	FSceneVelocityData VelocityData;
 
 	/** GPU Skinning cache, if enabled */
diff --git a/Engine/Source/Runtime/Renderer/Private/SceneRendering.h b/Engine/Source/Runtime/Renderer/Private/SceneRendering.h
--- a/Engine/Source/Runtime/Renderer/Private/SceneRendering.h
+++ b/Engine/Source/Runtime/Renderer/Private/SceneRendering.h
@@ -2328,8 +2328,8 @@ protected:
 
 	void RenderPixelProjectedReflection(FRDGBuilder& GraphBuilder, FRDGTextureRef SceneColorTexture, FRDGTextureRef SceneDepthTexture, FRDGTextureRef PixelProjectedReflectionTexture, const FPlanarReflectionSceneProxy* PlanarReflectionSceneProxy);
 
-	/** Before SetupMobileBasePassAfterShadowInit, we need to update the uniform buffer and shadow info for all movable local lights.*/
-	void UpdateMovableLocalLightUniformBufferAndShadowInfo();
+	/** Before SetupMobileBasePassAfterShadowInit, we need to update the uniform buffer and shadow info for all movable point lights.*/
+	void UpdateMovablePointLightUniformBufferAndShadowInfo();
 
 	void RenderMobileShadowProjections(FRDGBuilder& GraphBuilder, FRDGTextureRef SceneDepthTexture);
 private:
diff --git a/Engine/Source/Runtime/Renderer/Private/ShadowSetup.cpp b/Engine/Source/Runtime/Renderer/Private/ShadowSetup.cpp
--- a/Engine/Source/Runtime/Renderer/Private/ShadowSetup.cpp
+++ b/Engine/Source/Runtime/Renderer/Private/ShadowSetup.cpp
@@ -5735,6 +5735,10 @@ void FSceneRenderer::AllocateMobileCSMAndSpotLightShadowDepthTargets(FRHICommand
 				WholeSceneAtlasSize.Y = 1 << FMath::CeilLogTwo(WholeSceneAtlasSize.Y);
 			}
 
+			bool bResolutionChanged = Scene->MobileWholeSceneShadowAtlasSize != WholeSceneAtlasSize;
+
+			Scene->MobileWholeSceneShadowAtlasSize = WholeSceneAtlasSize;
+
 			FPooledRenderTargetDesc WholeSceneShadowMapDesc2D(FPooledRenderTargetDesc::Create2DDesc(WholeSceneAtlasSize, PF_ShadowDepth, FClearValueBinding::DepthOne, TexCreate_None, TexCreate_DepthStencilTargetable | TexCreate_ShaderResource, false));
 			WholeSceneShadowMapDesc2D.Flags |= GFastVRamConfig.ShadowCSM;
 			GRenderTargetPool.FindFreeElement(RHICmdList, WholeSceneShadowMapDesc2D, ShadowMapAtlas.RenderTargets.DepthTarget, TEXT("MobileCSMAndSpotLightShadowmap"));
@@ -5748,6 +5752,12 @@ void FSceneRenderer::AllocateMobileCSMAndSpotLightShadowDepthTargets(FRHICommand
 					ProjectedShadowInfo->RenderTargets.CopyReferencesFromRenderTargets(ShadowMapAtlas.RenderTargets);
 					ProjectedShadowInfo->SetupShadowDepthView(this);
 					ShadowMapAtlas.Shadows.Add(ProjectedShadowInfo);
+
+					if (bResolutionChanged)
+					{
+						FLightSceneInfo* LightSceneInfo = (FLightSceneInfo*)(&ProjectedShadowInfo->GetLightSceneInfo());
+						LightSceneInfo->Proxy->SetMobileMovablePointLightUniformBufferNeedsUpdate(true);
+					}
 				}
 			}
 		}
