Fixed broken separate translucency with split screen and off-center viewports.

diff --git a/Engine/Shaders/Private/ComposeSeparateTranslucency.usf b/Engine/Shaders/Private/ComposeSeparateTranslucency.usf
--- a/Engine/Shaders/Private/ComposeSeparateTranslucency.usf
+++ b/Engine/Shaders/Private/ComposeSeparateTranslucency.usf
@@ -1,6 +1,7 @@
 // Copyright Epic Games, Inc. All Rights Reserved.
 
 #include "Common.ush"
+#include "ScreenPass.ush"
 
 Texture2D SceneColor;
 SamplerState SceneColorSampler;
@@ -11,6 +12,10 @@ SamplerState SeparateTranslucencySampler;
 Texture2D SeparateModulation;
 SamplerState SeparateModulationSampler;
 
+SCREEN_PASS_TEXTURE_VIEWPORT(Color)
+SCREEN_PASS_TEXTURE_VIEWPORT(Translucency)
+FScreenTransform ColorToTranslucency;
+
 float2 ScreenPosToSceneColorUV;
 float2 ScreenPosToSeparateTranslucencyUV;
 float2 SeparateTranslucencyUVMax;
@@ -26,14 +31,14 @@ void MainPS(
 	float4 SvPosition : SV_POSITION,
 	out float4 OutColor : SV_Target0)
 {
-	float2 BufferUV = SvPosition.xy * ScreenPosToSceneColorUV;
-	float2 SeparateTranslucencyBufferUV = min(SvPosition.xy * ScreenPosToSeparateTranslucencyUV, SeparateTranslucencyUVMax);
+	float2 ColorUV = SvPosition.xy * Color_ExtentInverse.xy;
+	float2 SeparateTranslucencyUV = clamp(ApplyScreenTransform(ColorUV, ColorToTranslucency), Translucency_UVViewportBilinearMin, Translucency_UVViewportBilinearMax);
 
-	float4 SceneColorSample = SceneColor.SampleLevel(SceneColorSampler, BufferUV, 0);
+	float4 SceneColorSample = SceneColor.SampleLevel(SceneColorSampler, ColorUV, 0);
 
 #if PERMUTATION_NEARESTDEPTHNEIGHBOR 
 
-	NearestDepthNeighborUpsamplingResult UpsampleResult = NearestDepthNeighborUpsampling(SvPosition.xy, BufferUV, SeparateTranslucencyExtentInverse);
+	NearestDepthNeighborUpsamplingResult UpsampleResult = NearestDepthNeighborUpsampling(SvPosition.xy, ColorUV, Translucency_ExtentInverse);
 
 	float4 SeparateTranslucencySample = 0;
 	float4 SeparateModulationSample = 0;
@@ -50,8 +55,8 @@ void MainPS(
 
 #else // PERMUTATION_NEARESTDEPTHNEIGHBOR
 	
-	float4 SeparateTranslucencySample = SeparateTranslucency.SampleLevel(SeparateTranslucencySampler, SeparateTranslucencyBufferUV, 0);
-	float4 SeparateModulationSample   = SeparateModulation.SampleLevel(	 SeparateModulationSampler, SeparateTranslucencyBufferUV, 0);
+	float4 SeparateTranslucencySample = SeparateTranslucency.SampleLevel(SeparateTranslucencySampler, SeparateTranslucencyUV, 0);
+	float4 SeparateModulationSample   = SeparateModulation.SampleLevel(	 SeparateModulationSampler, SeparateTranslucencyUV, 0);
 
 #endif // PERMUTATION_NEARESTDEPTHNEIGHBOR
 
diff --git a/Engine/Source/Runtime/Renderer/Private/PostProcess/PostProcessing.cpp b/Engine/Source/Runtime/Renderer/Private/PostProcess/PostProcessing.cpp
--- a/Engine/Source/Runtime/Renderer/Private/PostProcess/PostProcessing.cpp
+++ b/Engine/Source/Runtime/Renderer/Private/PostProcess/PostProcessing.cpp
@@ -153,10 +153,9 @@ class FComposeSeparateTranslucencyPS : public FGlobalShader
 	using FPermutationDomain = TShaderPermutationDomain<FNearestDepthNeighborUpsampling>;
 
 	BEGIN_SHADER_PARAMETER_STRUCT(FParameters, )
-		SHADER_PARAMETER(FVector2D, ScreenPosToSceneColorUV)
-		SHADER_PARAMETER(FVector2D, ScreenPosToSeparateTranslucencyUV)
-		SHADER_PARAMETER(FVector2D, SeparateTranslucencyUVMax)
-		SHADER_PARAMETER(FVector2D, SeparateTranslucencyExtentInverse)
+		SHADER_PARAMETER_STRUCT(FScreenPassTextureViewportParameters, Color)
+		SHADER_PARAMETER_STRUCT(FScreenPassTextureViewportParameters, Translucency)
+		SHADER_PARAMETER(FScreenTransform, ColorToTranslucency)
 		SHADER_PARAMETER_RDG_TEXTURE(Texture2D, SceneColor)
 		SHADER_PARAMETER_SAMPLER(SamplerState, SceneColorSampler)
 		SHADER_PARAMETER_RDG_TEXTURE(Texture2D, SeparateTranslucency)
@@ -252,16 +251,19 @@ FRDGTextureRef AddTranslucencyCompositionPass(
 	const FVector2D SeparateTranslucencySize(SeparateTranslucencyRect.Size());
 	const FVector2D SeparateTranslucencyExtent((bApplyModulateOnly ? SeparateModulationTexture : SeparateTranslucencyTexture)->Desc.Extent);
 	const FVector2D SeparateTranslucencyExtentInv = FVector2D(1.0f, 1.0f) / SeparateTranslucencyExtent;
-	
-	const bool bScaleSeparateTranslucency = SeparateTranslucencySize != SceneColorSize;
-	const float DownsampleScale = SeparateTranslucencySize.X / SceneColorExtent.X;
+
+	const FScreenPassTextureViewport SceneColorViewport(SceneColor);
+	const FScreenPassTextureViewport TranslucencyViewport = SeparateTranslucencyTextures.GetDimensions().GetViewport(View.ViewRect);
+
+	const float DownsampleScale = SeparateTranslucencyTextures.GetDimensions().Scale;
+
+	const bool bScaleSeparateTranslucency = DownsampleScale != 1.0f;
 	const bool DepthUpscampling = !bPostMotionBlur && GetUseTranslucencyNearestDepthNeighborUpsample(DownsampleScale) ? 1 : 0;
 
 	FComposeSeparateTranslucencyPS::FParameters* PassParameters = GraphBuilder.AllocParameters<FComposeSeparateTranslucencyPS::FParameters>();
-	PassParameters->ScreenPosToSceneColorUV = SceneColorSizeInv * SceneColorSize * SceneColorExtentInv;
-	PassParameters->ScreenPosToSeparateTranslucencyUV = SceneColorSizeInv * SeparateTranslucencySize * SeparateTranslucencyExtentInv;
-	PassParameters->SeparateTranslucencyUVMax = (SeparateTranslucencySize - FVector2D(0.5f, 0.5f)) * SeparateTranslucencyExtentInv;
-	PassParameters->SeparateTranslucencyExtentInverse = SeparateTranslucencyExtentInv;
+	PassParameters->Color = GetScreenPassTextureViewportParameters(SceneColorViewport);
+	PassParameters->Translucency = GetScreenPassTextureViewportParameters(TranslucencyViewport);
+	PassParameters->ColorToTranslucency = FScreenTransform::ChangeTextureUVCoordinateFromTo(SceneColorViewport, TranslucencyViewport);
 	PassParameters->SceneColor = SceneColor.Texture;
 	PassParameters->SceneColorSampler = TStaticSamplerState<SF_Point>::GetRHI();
 	PassParameters->SeparateTranslucency = SeparateTranslucencyTexture;
