Fixed ambient occlusion not working with multiple viewports. Refactored logic to be done in one place.

diff --git a/Engine/Source/Runtime/Renderer/Private/CompositionLighting/CompositionLighting.cpp b/Engine/Source/Runtime/Renderer/Private/CompositionLighting/CompositionLighting.cpp
--- a/Engine/Source/Runtime/Renderer/Private/CompositionLighting/CompositionLighting.cpp
+++ b/Engine/Source/Runtime/Renderer/Private/CompositionLighting/CompositionLighting.cpp
@@ -485,195 +485,228 @@ static FScreenPassTexture AddPostProcessingAmbientOcclusion(
 	return FinalOutput;
 }
 
-namespace CompositionLighting
+FCompositionLighting::FCompositionLighting(TArrayView<const FViewInfo> InViews, const FSceneTextures& InSceneTextures, TUniqueFunction<bool(int32)> RequestSSAOFunction)
+	: Views(InViews)
+	, ViewFamily(*InViews[0].Family)
+	, SceneTextures(InSceneTextures)
+	, bEnableDBuffer(IsDBufferEnabled(ViewFamily, SceneTextures.Config.ShaderPlatform))
+	, bEnableDecals(ViewFamily.EngineShowFlags.Decals && !ViewFamily.EngineShowFlags.VisualizeLightCulling)
 {
+	ViewAOConfigs.SetNum(Views.Num());
 
-void ProcessBeforeBasePass(
-	FRDGBuilder& GraphBuilder,
-	TArrayView<const FViewInfo> Views,
-	const FSceneTextures& SceneTextures,
-	FDBufferTextures& DBufferTextures)
+	for (int32 Index = 0; Index < Views.Num(); ++Index)
+	{
+		ViewAOConfigs[Index].bRequested = RequestSSAOFunction(Index);
+	}
+}
+
+void FCompositionLighting::TryInit()
 {
-	check(Views.Num());
+	if (bInitialized)
+	{
+		return;
+	}
 
-	const EShaderPlatform ShaderPlatform = Views[0].GetShaderPlatform();
-	const bool bDBuffer = IsDBufferEnabled(*Views[0].Family, ShaderPlatform);
+	const bool bForwardShading = IsForwardShadingEnabled(SceneTextures.Config.ShaderPlatform);
 
-	const auto ProcessView = [&](const FViewInfo& View, uint32 SSAOLevels)
+	for (int32 ViewIndex = 0; ViewIndex < Views.Num(); ++ViewIndex)
 	{
-		const bool bNeedSSAO = SSAOLevels && FSSAOHelper::GetGTAOPassType(View, SSAOLevels) == EGTAOType::EOff;
+		const FViewInfo& View = Views[ViewIndex];
+
+		FAOConfig& ViewConfig = ViewAOConfigs[ViewIndex];
 
-		// so that the passes can register themselves to the graph
-		if (bDBuffer || bNeedSSAO)
+		if (!ViewConfig.bRequested)
 		{
-			RDG_EVENT_SCOPE(GraphBuilder, "CompositionBeforeBasePass");
-			RDG_GPU_STAT_SCOPE(GraphBuilder, CompositionBeforeBasePass);
-			View.BeginRenderView();
+			continue;
+		}
 
-			// decals are before AmbientOcclusion so the decal can output a normal that AO is affected by
-			if (bDBuffer)
-			{
-				FDeferredDecalPassTextures DecalPassTextures = GetDeferredDecalPassTextures(GraphBuilder, View, SceneTextures, &DBufferTextures);
-				AddDeferredDecalPass(GraphBuilder, View, DecalPassTextures, EDecalRenderStage::BeforeBasePass);
-			}
+		ViewConfig.Levels = FSSAOHelper::ComputeAmbientOcclusionPassCount(View);
 
-			if (bNeedSSAO)
-			{
-				FSSAOCommonParameters Parameters = GetSSAOCommonParameters(GraphBuilder, View, SceneTextures.UniformBuffer, SSAOLevels, false);
-				FScreenPassRenderTarget FinalTarget = FScreenPassRenderTarget(SceneTextures.ScreenSpaceAO, View.ViewRect, ERenderTargetLoadAction::ENoAction);
+		if (!bForwardShading)
+		{
+			ViewConfig.GTAOType = FSSAOHelper::GetGTAOPassType(View, ViewConfig.Levels);
+		}
 
-				AddPostProcessingAmbientOcclusion(
-					GraphBuilder,
-					View,
-					Parameters,
-					FinalTarget);
-			}
+		if (ViewConfig.GTAOType == EGTAOType::EOff && ViewConfig.Levels > 0)
+		{
+			ViewConfig.bSSAOAsync = FSSAOHelper::IsAmbientOcclusionAsyncCompute(View, ViewConfig.Levels);
+
+			ViewConfig.SSAOLocation = View.HZB != nullptr && (ViewConfig.bSSAOAsync || bForwardShading)
+				? ESSAOLocation::BeforeBasePass
+				: ESSAOLocation::AfterBasePass;
 		}
-	};
+	}
+
+	bInitialized = true;
+}
+
+void FCompositionLighting::ProcessBeforeBasePass(FRDGBuilder& GraphBuilder, FDBufferTextures& DBufferTextures)
+{
+	if (HasRayTracedOverlay(ViewFamily))
+	{
+		return;
+	}
+
+	TryInit();
+
+	RDG_EVENT_SCOPE(GraphBuilder, "CompositionBeforeBasePass");
+	RDG_GPU_STAT_SCOPE(GraphBuilder, CompositionBeforeBasePass);
 
 	for (int32 ViewIndex = 0; ViewIndex < Views.Num(); ++ViewIndex)
 	{
 		const FViewInfo& View = Views[ViewIndex];
+		const FAOConfig& ViewConfig = ViewAOConfigs[ViewIndex];
+
+		const bool bEnableSSAO = ViewConfig.SSAOLocation == ESSAOLocation::BeforeBasePass;
+
+		if (!bEnableDBuffer && !bEnableSSAO)
+		{
+			continue;
+		}
+
 		RDG_GPU_MASK_SCOPE(GraphBuilder, View.GPUMask);
 		RDG_EVENT_SCOPE_CONDITIONAL(GraphBuilder, Views.Num() > 1, "View%d", ViewIndex);
 
-		uint32 SSAOLevels = FSSAOHelper::ComputeAmbientOcclusionPassCount(View);
+		View.BeginRenderView();
 
-		// Disable SSAO if no valid HZB exists, or if we are deferred and non-async (this is handled in ProcessAfterBasePass instead).
-		if (!View.HZB || (!IsForwardShadingEnabled(View.GetShaderPlatform()) && !FSSAOHelper::IsAmbientOcclusionAsyncCompute(View, SSAOLevels)))
+		// decals are before AmbientOcclusion so the decal can output a normal that AO is affected by
+		if (bEnableDBuffer)
 		{
-			SSAOLevels = 0;
+			FDeferredDecalPassTextures DecalPassTextures = GetDeferredDecalPassTextures(GraphBuilder, View, SceneTextures, &DBufferTextures);
+			AddDeferredDecalPass(GraphBuilder, View, DecalPassTextures, EDecalRenderStage::BeforeBasePass);
 		}
 
-		ProcessView(View, SSAOLevels);
+		if (bEnableSSAO)
+		{
+			FSSAOCommonParameters Parameters = GetSSAOCommonParameters(GraphBuilder, View, SceneTextures.UniformBuffer, ViewConfig.Levels, false);
+			FScreenPassRenderTarget FinalTarget = FScreenPassRenderTarget(SceneTextures.ScreenSpaceAO, View.ViewRect, ERenderTargetLoadAction::ENoAction);
+
+			AddPostProcessingAmbientOcclusion(
+				GraphBuilder,
+				View,
+				Parameters,
+				FinalTarget);
+		}
 	}
 }
 
-void ProcessAfterBasePass(
-	FRDGBuilder& GraphBuilder,
-	const FViewInfo& View,
-	const FSceneTextures& SceneTextures,
-	const FAsyncResults& AsyncResults,
-	bool bEnableSSAO)
+void FCompositionLighting::ProcessAfterBasePass(FRDGBuilder& GraphBuilder)
 {
-	const FSceneViewFamily& ViewFamily = *View.Family;
 	if (HasRayTracedOverlay(ViewFamily))
 	{
 		return;
 	}
 
-	RDG_GPU_MASK_SCOPE(GraphBuilder, View.GPUMask);
+	check(bInitialized);
+
 	RDG_EVENT_SCOPE(GraphBuilder, "LightCompositionTasks_PreLighting");
 	RDG_GPU_STAT_SCOPE(GraphBuilder, CompositionPreLighting);
 
-	View.BeginRenderView();
+	for (int32 ViewIndex = 0; ViewIndex < Views.Num(); ++ViewIndex)
+	{
+		const FViewInfo& View = Views[ViewIndex];
+		const FAOConfig& ViewConfig = ViewAOConfigs[ViewIndex];
 
-	// decal are distracting when looking at LightCulling.
-	const bool bDoDecal = ViewFamily.EngineShowFlags.Decals && !ViewFamily.EngineShowFlags.VisualizeLightCulling;
-	const bool bDBuffer = IsUsingDBuffers(View.GetShaderPlatform());
+		RDG_GPU_MASK_SCOPE(GraphBuilder, View.GPUMask);
+		RDG_EVENT_SCOPE_CONDITIONAL(GraphBuilder, Views.Num() > 1, "View%d", ViewIndex);
 
-	FDeferredDecalPassTextures DecalPassTextures = GetDeferredDecalPassTextures(GraphBuilder, View, SceneTextures, nullptr);
+		View.BeginRenderView();
 
-	if (bDoDecal && !bDBuffer && IsUsingGBuffers(View.GetShaderPlatform()))
-	{
-		// We can disable this pass if using DBuffer decals
-		// Decals are before AmbientOcclusion so the decal can output a normal that AO is affected by
-		AddDeferredDecalPass(GraphBuilder, View, DecalPassTextures, EDecalRenderStage::BeforeLighting);
-	}
+		FDeferredDecalPassTextures DecalPassTextures = GetDeferredDecalPassTextures(GraphBuilder, View, SceneTextures, nullptr);
 
-	if (bDoDecal && !IsSimpleForwardShadingEnabled(View.GetShaderPlatform()))
-	{
-		// DBuffer decals with emissive component
-		AddDeferredDecalPass(GraphBuilder, View, DecalPassTextures, EDecalRenderStage::Emissive);
-	}
+		if (bEnableDecals && !bEnableDBuffer && IsUsingGBuffers(SceneTextures.Config.ShaderPlatform))
+		{
+			// We can disable this pass if using DBuffer decals
+			// Decals are before AmbientOcclusion so the decal can output a normal that AO is affected by
+			AddDeferredDecalPass(GraphBuilder, View, DecalPassTextures, EDecalRenderStage::BeforeLighting);
+		}
 
-	// Forward shading SSAO is applied before the base pass using only the depth buffer.
-	if (!IsForwardShadingEnabled(View.GetShaderPlatform()))
-	{
-		const uint32 SSAOLevels = FSSAOHelper::ComputeAmbientOcclusionPassCount(View);
-		if (SSAOLevels && bEnableSSAO)
+		if (bEnableDecals && !IsSimpleForwardShadingEnabled(View.GetShaderPlatform()))
 		{
-			const bool bScreenSpaceAOIsProduced = SceneTextures.ScreenSpaceAO->HasBeenProduced();
-			FScreenPassRenderTarget FinalTarget = FScreenPassRenderTarget(SceneTextures.ScreenSpaceAO, View.ViewRect, bScreenSpaceAOIsProduced ? ERenderTargetLoadAction::ELoad : ERenderTargetLoadAction::ENoAction);
+			// DBuffer decals with emissive component
+			AddDeferredDecalPass(GraphBuilder, View, DecalPassTextures, EDecalRenderStage::Emissive);
+		}
 
-			FScreenPassTexture AmbientOcclusion = bScreenSpaceAOIsProduced ? FinalTarget : FScreenPassTexture();
+		// Forward shading SSAO is applied before the base pass using only the depth buffer.
+		if (!IsForwardShadingEnabled(View.GetShaderPlatform()))
+		{
+			if (ViewConfig.Levels > 0)
+			{
+				const bool bScreenSpaceAOIsProduced = SceneTextures.ScreenSpaceAO->HasBeenProduced();
+				FScreenPassRenderTarget FinalTarget = FScreenPassRenderTarget(SceneTextures.ScreenSpaceAO, View.ViewRect, bScreenSpaceAOIsProduced ? ERenderTargetLoadAction::ELoad : ERenderTargetLoadAction::ENoAction);
 
-			const EGTAOType GTAOType = FSSAOHelper::GetGTAOPassType(View, SSAOLevels);
+				FScreenPassTexture AmbientOcclusion = bScreenSpaceAOIsProduced ? FinalTarget : FScreenPassTexture();
 
-			// If doing the Split GTAO method then we need to do the second part here.
-			if (GTAOType == EGTAOType::EAsyncHorizonSearch || GTAOType == EGTAOType::EAsyncCombinedSpatial)
-			{
-				check(AsyncResults.HorizonsTexture);
+				// If doing the Split GTAO method then we need to do the second part here.
+				if (ViewConfig.GTAOType == EGTAOType::EAsyncHorizonSearch || ViewConfig.GTAOType == EGTAOType::EAsyncCombinedSpatial)
+				{
+					check(HorizonsTexture);
 
-				FGTAOCommonParameters Parameters = GetGTAOCommonParameters(GraphBuilder, View, SceneTextures.UniformBuffer, GTAOType);
+					FGTAOCommonParameters Parameters = GetGTAOCommonParameters(GraphBuilder, View, SceneTextures.UniformBuffer, ViewConfig.GTAOType);
 
-				FScreenPassTexture GTAOHorizons(AsyncResults.HorizonsTexture, Parameters.DownsampledViewRect);
-				AmbientOcclusion = AddPostProcessingGTAOPostAsync(GraphBuilder, View, Parameters, GTAOHorizons, FinalTarget);
+					FScreenPassTexture GTAOHorizons(HorizonsTexture, Parameters.DownsampledViewRect);
+					AmbientOcclusion = AddPostProcessingGTAOPostAsync(GraphBuilder, View, Parameters, GTAOHorizons, FinalTarget);
 
-				ensureMsgf(
-					DecalRendering::BuildVisibleDecalList(*(FScene*)View.Family->Scene, View, EDecalRenderStage::AmbientOcclusion, nullptr) == false,
-					TEXT("Ambient occlusion decals are not supported with Async compute SSAO."));
-			}
-			else
-			{
-				if (GTAOType == EGTAOType::ENonAsync)
-				{
-					FGTAOCommonParameters Parameters = GetGTAOCommonParameters(GraphBuilder, View, SceneTextures.UniformBuffer, GTAOType);
-					AmbientOcclusion = AddPostProcessingGTAOAllPasses(GraphBuilder, View, Parameters, FinalTarget);
+					ensureMsgf(
+						DecalRendering::BuildVisibleDecalList(*(FScene*)View.Family->Scene, View, EDecalRenderStage::AmbientOcclusion, nullptr) == false,
+						TEXT("Ambient occlusion decals are not supported with Async compute SSAO."));
 				}
-				else if (!AmbientOcclusion.IsValid())
+				else
 				{
-					FSSAOCommonParameters Parameters = GetSSAOCommonParameters(GraphBuilder, View, SceneTextures.UniformBuffer, SSAOLevels, true);
-					AmbientOcclusion = AddPostProcessingAmbientOcclusion(GraphBuilder, View, Parameters, FinalTarget);
-				}
-
-				if (bDoDecal)
-				{
-					DecalPassTextures.ScreenSpaceAO = AmbientOcclusion.Texture;
-					AddDeferredDecalPass(GraphBuilder, View, DecalPassTextures, EDecalRenderStage::AmbientOcclusion);
+					if (ViewConfig.GTAOType == EGTAOType::ENonAsync)
+					{
+						FGTAOCommonParameters Parameters = GetGTAOCommonParameters(GraphBuilder, View, SceneTextures.UniformBuffer, ViewConfig.GTAOType);
+						AmbientOcclusion = AddPostProcessingGTAOAllPasses(GraphBuilder, View, Parameters, FinalTarget);
+					}
+					else if (ViewConfig.SSAOLocation == ESSAOLocation::AfterBasePass)
+					{
+						FSSAOCommonParameters Parameters = GetSSAOCommonParameters(GraphBuilder, View, SceneTextures.UniformBuffer, ViewConfig.Levels, true);
+						AmbientOcclusion = AddPostProcessingAmbientOcclusion(GraphBuilder, View, Parameters, FinalTarget);
+					}
+
+					if (bEnableDecals)
+					{
+						DecalPassTextures.ScreenSpaceAO = AmbientOcclusion.Texture;
+						AddDeferredDecalPass(GraphBuilder, View, DecalPassTextures, EDecalRenderStage::AmbientOcclusion);
+					}
 				}
 			}
 		}
 	}
 }
 
-bool CanProcessAsync(TArrayView<const FViewInfo> Views)
+void FCompositionLighting::ProcessAfterOcclusion(FRDGBuilder& GraphBuilder)
 {
-	for (const FViewInfo& View : Views)
+	if (HasRayTracedOverlay(ViewFamily))
 	{
-		uint32 Levels = FSSAOHelper::ComputeAmbientOcclusionPassCount(View);
-		if (!FSSAOHelper::IsAmbientOcclusionAsyncCompute(View, Levels) || HasRayTracedOverlay(*View.Family))
-		{
-			return false;
-		}
+		return;
 	}
-	return true;
-}
 
-FAsyncResults ProcessAsync(FRDGBuilder& GraphBuilder, TArrayView<const FViewInfo> Views, const FMinimalSceneTextures& SceneTextures)
-{
-	FAsyncResults AsyncResults;
+	TryInit();
 
 	RDG_ASYNC_COMPUTE_BUDGET_SCOPE(GraphBuilder, FSSAOHelper::GetAmbientOcclusionAsyncComputeBudget());
 
-	for (const FViewInfo& View : Views)
+	for (int32 ViewIndex = 0; ViewIndex < Views.Num(); ++ViewIndex)
 	{
-		const uint32 Levels = FSSAOHelper::ComputeAmbientOcclusionPassCount(View);
-		const EGTAOType GTAOType = FSSAOHelper::GetGTAOPassType(View, Levels);
+		const FViewInfo& View = Views[ViewIndex];
+		const FAOConfig& ViewConfig = ViewAOConfigs[ViewIndex];
 
-		if (GTAOType == EGTAOType::EAsyncCombinedSpatial || GTAOType == EGTAOType::EAsyncHorizonSearch)
+		if (ViewConfig.GTAOType == EGTAOType::EAsyncCombinedSpatial || ViewConfig.GTAOType == EGTAOType::EAsyncHorizonSearch)
 		{
 			RDG_GPU_MASK_SCOPE(GraphBuilder, View.GPUMask);
 
-			FGTAOCommonParameters CommonParameters = GetGTAOCommonParameters(GraphBuilder, View, SceneTextures.UniformBuffer, GTAOType);
+			FGTAOCommonParameters CommonParameters = GetGTAOCommonParameters(GraphBuilder, View, SceneTextures.UniformBuffer, ViewConfig.GTAOType);
+
+			const ERenderTargetLoadAction LoadAction = View.DecayLoadAction(ERenderTargetLoadAction::ENoAction);
 
+			if (!HorizonsTexture)
 			{
 				const FIntPoint HorizonTextureSize = FIntPoint::DivideAndRoundUp(SceneTextures.Config.Extent, CommonParameters.DownscaleFactor);
 				const FRDGTextureDesc Desc = FRDGTextureDesc::Create2D(HorizonTextureSize, PF_R8G8, FClearValueBinding::White, TexCreate_UAV | TexCreate_RenderTargetable);
-				AsyncResults.HorizonsTexture = GraphBuilder.CreateTexture(Desc, TEXT("ScreenSpaceGTAOHorizons"));
+				HorizonsTexture = GraphBuilder.CreateTexture(Desc, TEXT("ScreenSpaceGTAOHorizons"));
 			}
 
-			FScreenPassRenderTarget GTAOHorizons(AsyncResults.HorizonsTexture, CommonParameters.DownsampledViewRect, ERenderTargetLoadAction::ENoAction);
+			FScreenPassRenderTarget GTAOHorizons(HorizonsTexture, CommonParameters.DownsampledViewRect, LoadAction);
 
 			AddPostProcessingGTAOAsyncPasses(
 				GraphBuilder,
@@ -682,8 +715,4 @@ FAsyncResults ProcessAsync(FRDGBuilder& GraphBuilder, TArrayView<const FViewInfo
 				GTAOHorizons);
 		}
 	}
-
-	return AsyncResults;
-}
-
-} //! namespace CompositionLighting
\ No newline at end of file
+}
\ No newline at end of file
diff --git a/Engine/Source/Runtime/Renderer/Private/CompositionLighting/CompositionLighting.h b/Engine/Source/Runtime/Renderer/Private/CompositionLighting/CompositionLighting.h
--- a/Engine/Source/Runtime/Renderer/Private/CompositionLighting/CompositionLighting.h
+++ b/Engine/Source/Runtime/Renderer/Private/CompositionLighting/CompositionLighting.h
@@ -9,35 +9,51 @@
 #include "CoreMinimal.h"
 #include "SceneRendering.h"
 #include "PostProcessDeferredDecals.h"
-
-class FPersistentUniformBuffers;
+#include "PostProcessAmbientOcclusion.h"
 
 /**
  * The center for all screen space processing activities (e.g. G-buffer manipulation, lighting).
  */
-namespace CompositionLighting
+class FCompositionLighting
 {
-	struct FAsyncResults
-	{
-		FRDGTextureRef HorizonsTexture = nullptr;
-	};
+public:
+	FCompositionLighting(TArrayView<const FViewInfo> InViews, const FSceneTextures& InSceneTextures, TUniqueFunction<bool(int32)> RequestSSAOFunction);
+
+	void ProcessAfterOcclusion(FRDGBuilder& GraphBuilder);
+
+	void ProcessBeforeBasePass(FRDGBuilder& GraphBuilder, FDBufferTextures& DBufferTextures);
+
+	void ProcessAfterBasePass(FRDGBuilder& GraphBuilder);
 
-	bool CanProcessAsync(TArrayView<const FViewInfo> Views);
+private:
+	void TryInit();
 
-	void ProcessBeforeBasePass(
-		FRDGBuilder& GraphBuilder,
-		TArrayView<const FViewInfo> Views,
-		const FSceneTextures& SceneTextures,
-		FDBufferTextures& DBufferTextures);
+	const TArrayView<const FViewInfo> Views;
+	const FSceneViewFamily& ViewFamily;
+	const FSceneTextures& SceneTextures;
 
-	void ProcessAfterBasePass(
-		FRDGBuilder& GraphBuilder,
-		const FViewInfo& View,
-		const FSceneTextures& SceneTextures,
-		const FAsyncResults& AsyncResults,
-		bool bEnableSSAO);
+	const bool bEnableDBuffer;
+	const bool bEnableDecals;
+
+	enum class ESSAOLocation
+	{
+		None,
+		BeforeBasePass,
+		AfterBasePass
+	};
+
+	struct FAOConfig
+	{
+		uint32 Levels = 0;
+		EGTAOType GTAOType = EGTAOType::EOff;
+		ESSAOLocation SSAOLocation = ESSAOLocation::None;
+		bool bSSAOAsync = false;
+		bool bRequested = false;
+	};
 
-	FAsyncResults ProcessAsync(FRDGBuilder& GraphBuilder, TArrayView<const FViewInfo> Views, const FMinimalSceneTextures& SceneTextures);
+	TArray<FAOConfig, TInlineAllocator<8>> ViewAOConfigs;
+	FRDGTextureRef HorizonsTexture = nullptr;
+	bool bInitialized = false;
 };
 
 extern bool ShouldRenderScreenSpaceAmbientOcclusion(const FViewInfo& View);
\ No newline at end of file
diff --git a/Engine/Source/Runtime/Renderer/Private/CompositionLighting/PostProcessAmbientOcclusion.h b/Engine/Source/Runtime/Renderer/Private/CompositionLighting/PostProcessAmbientOcclusion.h
--- a/Engine/Source/Runtime/Renderer/Private/CompositionLighting/PostProcessAmbientOcclusion.h
+++ b/Engine/Source/Runtime/Renderer/Private/CompositionLighting/PostProcessAmbientOcclusion.h
@@ -105,16 +105,8 @@ public:
 
 // SSAO
 
-/** TODO: The RHI version of the uniform buffer is used on all compute shaders because RDG does not yet support reads
- *  from multiple pipes (e.g. reading depth from both async compute and graphics). Using the RDG version of the uniform
- *  buffer ends up joining the async work back to graphics on the depth read. This temporary solution just uses the RHI
- *  uniform buffer to avoid registering the textures with RDG. This only works because depth is required to be
- *  read-only by previous passes. This can be removed once RDG supports multi-pipe transitions.
- */
-
 struct FSSAOCommonParameters
 {
-	TUniformBufferRef<FSceneTextureUniformParameters> SceneTexturesUniformBufferRHI;
 	TRDGUniformBufferRef<FSceneTextureUniformParameters> SceneTexturesUniformBuffer = nullptr;
 	FScreenPassTextureViewport SceneTexturesViewport;
 
diff --git a/Engine/Source/Runtime/Renderer/Private/DeferredShadingRenderer.cpp b/Engine/Source/Runtime/Renderer/Private/DeferredShadingRenderer.cpp
--- a/Engine/Source/Runtime/Renderer/Private/DeferredShadingRenderer.cpp
+++ b/Engine/Source/Runtime/Renderer/Private/DeferredShadingRenderer.cpp
@@ -2453,16 +2453,16 @@ void FDeferredShadingSceneRenderer::Render(FRDGBuilder& GraphBuilder)
 		}
 	};
 
-	CompositionLighting::FAsyncResults CompositionLightingAsyncResults;
+	FCompositionLighting CompositionLighting(Views, SceneTextures, [this] (int32 ViewIndex)
+	{
+		return ViewPipelineStates[ViewIndex]->AmbientOcclusionMethod == EAmbientOcclusionMethod::SSAO;
+	});
 
 	const auto RenderOcclusionLambda = [&]()
 	{
 		RenderOcclusion(GraphBuilder, SceneTextures, bIsOcclusionTesting);
 
-		if (CompositionLighting::CanProcessAsync(Views))
-		{
-			CompositionLightingAsyncResults = CompositionLighting::ProcessAsync(GraphBuilder, Views, SceneTextures);
-		}
+		CompositionLighting.ProcessAfterOcclusion(GraphBuilder);
 
 		ExecuteDrainPoint();
 	};
@@ -2593,7 +2593,7 @@ void FDeferredShadingSceneRenderer::Render(FRDGBuilder& GraphBuilder)
 	{
 		CSV_SCOPED_TIMING_STAT_EXCLUSIVE(DeferredShadingSceneRenderer_DBuffer);
 		SCOPE_CYCLE_COUNTER(STAT_FDeferredShadingSceneRenderer_DBuffer);
-		CompositionLighting::ProcessBeforeBasePass(GraphBuilder, Views, SceneTextures, DBufferTextures);
+		CompositionLighting.ProcessBeforeBasePass(GraphBuilder, DBufferTextures);
 	}
 	
 	if (IsForwardShadingEnabled(ShaderPlatform) && bAllowStaticLighting)
@@ -2789,15 +2789,7 @@ void FDeferredShadingSceneRenderer::Render(FRDGBuilder& GraphBuilder)
 			AddResolveSceneDepthPass(GraphBuilder, Views, SceneTextures.Depth);
 		}
 
-		for (int32 ViewIndex = 0; ViewIndex < Views.Num(); ViewIndex++)
-		{
-			const FViewInfo& View = Views[ViewIndex];
-			TPipelineState<FPerViewPipelineState>& ViewPipelineState = ViewPipelineStates[ViewIndex];
-			RDG_EVENT_SCOPE_CONDITIONAL(GraphBuilder, Views.Num() > 1, "View%d", ViewIndex);
-
-			bool bEnableSSAO = ViewPipelineState->AmbientOcclusionMethod == EAmbientOcclusionMethod::SSAO;
-			CompositionLighting::ProcessAfterBasePass(GraphBuilder, View, SceneTextures, CompositionLightingAsyncResults, bEnableSSAO);
-		}
+		CompositionLighting.ProcessAfterBasePass(GraphBuilder);
 	}
 
 	// Rebuild scene textures to include velocity, custom depth, and SSAO.
