Sync deferred clean up tasks when performing UpdateScenePrimitiveInfos.

diff --git a/Engine/Source/Runtime/Renderer/Private/RendererScene.cpp b/Engine/Source/Runtime/Renderer/Private/RendererScene.cpp
--- a/Engine/Source/Runtime/Renderer/Private/RendererScene.cpp
+++ b/Engine/Source/Runtime/Renderer/Private/RendererScene.cpp
@@ -3953,6 +3953,8 @@ void FScene::UpdateAllPrimitiveSceneInfos(FRDGBuilder& GraphBuilder, bool bAsync
 
 	check(IsInRenderingThread());
 
+	FSceneRenderer::WaitForCleanUpTasks(GraphBuilder.RHICmdList);
+
 	RDG_EVENT_SCOPE(GraphBuilder, "UpdateAllPrimitiveSceneInfos");
 
 	TArray<FPrimitiveSceneInfo*> RemovedLocalPrimitiveSceneInfos(RemovedPrimitiveSceneInfos.Array());
diff --git a/Engine/Source/Runtime/Renderer/Private/SceneRendering.cpp b/Engine/Source/Runtime/Renderer/Private/SceneRendering.cpp
--- a/Engine/Source/Runtime/Renderer/Private/SceneRendering.cpp
+++ b/Engine/Source/Runtime/Renderer/Private/SceneRendering.cpp
@@ -3745,7 +3745,7 @@ static void DeleteSceneRenderer(FRHICommandListImmediate& RHICmdList, FSceneRend
 	FGraphicsMinimalPipelineStateId::ResetLocalPipelineIdTableSize();
 }
 
-static void ReleaseSceneRenderer(FRHICommandListImmediate& RHICmdList, FSceneRenderer* SceneRenderer, FMemMark* MemStackMark)
+static void ReleaseSceneRenderer(FRHICommandListImmediate& RHICmdList, FSceneRenderer* SceneRenderer)
 {
 	{
 		QUICK_SCOPE_CYCLE_COUNTER(STAT_DeleteSceneRenderer_WaitForTasks);
@@ -3753,7 +3753,6 @@ static void ReleaseSceneRenderer(FRHICommandListImmediate& RHICmdList, FSceneRen
 	}
 
 	SceneRenderer->WaitForTasksAndClearSnapshots(FParallelMeshDrawCommandPass::EWaitThread::Render);
-	DeleteSceneRenderer(RHICmdList, SceneRenderer, MemStackMark);
 }
 
 void FSceneRenderer::RenderThreadBegin(FRHICommandListImmediate& RHICmdList)
@@ -3769,6 +3768,7 @@ struct FSceneRenderCleanUpState
 	FMemMark* MemStackMark{};
 	FGraphEventRef Task;
 	ESceneRenderCleanUpMode CompletionMode = ESceneRenderCleanUpMode::Immediate;
+	bool bWaitForTasksComplete = false;
 };
 
 static FSceneRenderCleanUpState GSceneRenderCleanUpState;
@@ -3783,7 +3783,8 @@ void FSceneRenderer::RenderThreadEnd(FRHICommandListImmediate& RHICmdList)
 
 	if (GSceneRenderCleanUpState.CompletionMode == ESceneRenderCleanUpMode::Immediate)
 	{
-		ReleaseSceneRenderer(RHICmdList, this, MemStackMark);
+		ReleaseSceneRenderer(RHICmdList, this);
+		DeleteSceneRenderer(RHICmdList, this, MemStackMark);
 	}
 	else
 	{
@@ -3837,7 +3838,31 @@ void FSceneRenderer::RenderThreadEnd(FRHICommandListImmediate& RHICmdList)
 
 void FSceneRenderer::CleanUp(FRHICommandListImmediate& RHICmdList)
 {
-	if (GSceneRenderCleanUpState.CompletionMode == ESceneRenderCleanUpMode::Immediate)
+	if (GSceneRenderCleanUpState.CompletionMode == ESceneRenderCleanUpMode::Immediate || !GSceneRenderCleanUpState.Renderer)
+	{
+		return;
+	}
+
+	if (!GSceneRenderCleanUpState.bWaitForTasksComplete)
+	{
+		switch (GSceneRenderCleanUpState.CompletionMode)
+		{
+		case ESceneRenderCleanUpMode::Deferred:
+			ReleaseSceneRenderer(RHICmdList, GSceneRenderCleanUpState.Renderer);
+			break;
+		case ESceneRenderCleanUpMode::DeferredAndAsync:
+			GSceneRenderCleanUpState.Task->Wait(ENamedThreads::GetRenderThread_Local());
+			break;
+		}
+	}
+
+	DeleteSceneRenderer(RHICmdList, GSceneRenderCleanUpState.Renderer, GSceneRenderCleanUpState.MemStackMark);
+	GSceneRenderCleanUpState = {};
+}
+
+void FSceneRenderer::WaitForCleanUpTasks(FRHICommandListImmediate& RHICmdList)
+{
+	if (GSceneRenderCleanUpState.CompletionMode == ESceneRenderCleanUpMode::Immediate || !GSceneRenderCleanUpState.Renderer || GSceneRenderCleanUpState.bWaitForTasksComplete)
 	{
 		return;
 	}
@@ -3845,15 +3870,15 @@ void FSceneRenderer::CleanUp(FRHICommandListImmediate& RHICmdList)
 	switch (GSceneRenderCleanUpState.CompletionMode)
 	{
 	case ESceneRenderCleanUpMode::Deferred:
-		ReleaseSceneRenderer(RHICmdList, GSceneRenderCleanUpState.Renderer, GSceneRenderCleanUpState.MemStackMark);
+		ReleaseSceneRenderer(RHICmdList, GSceneRenderCleanUpState.Renderer);
 		break;
 	case ESceneRenderCleanUpMode::DeferredAndAsync:
 		GSceneRenderCleanUpState.Task->Wait(ENamedThreads::GetRenderThread_Local());
-		DeleteSceneRenderer(RHICmdList, GSceneRenderCleanUpState.Renderer, GSceneRenderCleanUpState.MemStackMark);
+		GSceneRenderCleanUpState.Task = nullptr;
 		break;
 	}
 
-	GSceneRenderCleanUpState = {};
+	GSceneRenderCleanUpState.bWaitForTasksComplete = true;
 }
 
 void FSceneRenderer::WaitForTasksAndClearSnapshots(FParallelMeshDrawCommandPass::EWaitThread WaitThread)
diff --git a/Engine/Source/Runtime/Renderer/Private/SceneRendering.h b/Engine/Source/Runtime/Renderer/Private/SceneRendering.h
--- a/Engine/Source/Runtime/Renderer/Private/SceneRendering.h
+++ b/Engine/Source/Runtime/Renderer/Private/SceneRendering.h
@@ -1915,6 +1915,9 @@ public:
 	/** Called to release any deallocations that were deferred until the next render. */
 	static void CleanUp(FRHICommandListImmediate& RHICmdList);
 
+	/** Called to wait for dependent clean up tasks, but CleanUp must still be called later. */
+	static void WaitForCleanUpTasks(FRHICommandListImmediate& RHICmdList);
+
 	/** Apply the ResolutionFraction on ViewSize, taking into account renderer's requirements. */
 	static FIntPoint ApplyResolutionFraction(
 		const FSceneViewFamily& ViewFamily, const FIntPoint& UnscaledViewSize, float ResolutionFraction);
