Fix how wave operations are detected by the path tracer

Should use the provided shader platform when deciding which permutation to compile instead of relying on the global GRHISupportsWaveOperations

Enable wave operations for PCD3D_SM6 (part of ShaderModel 6.0)

diff --git a/Engine/Config/Windows/DataDrivenPlatformInfo.ini b/Engine/Config/Windows/DataDrivenPlatformInfo.ini
--- a/Engine/Config/Windows/DataDrivenPlatformInfo.ini
+++ b/Engine/Config/Windows/DataDrivenPlatformInfo.ini
@@ -73,6 +73,7 @@ bSupportsByteBufferComputeShaders=true
 bSupportsLumenGI=true
 bSupportsDistanceFields=true
 bSupportsSSDIndirect=true
+bSupportsWaveOperations=true
 ; TODO: SM6.6
 bRequiresVendorExtensionsForAtomics=true
 bSupportsAnisotropicMaterials=true
diff --git a/Engine/Source/Runtime/Renderer/Private/PathTracing.cpp b/Engine/Source/Runtime/Renderer/Private/PathTracing.cpp
--- a/Engine/Source/Runtime/Renderer/Private/PathTracing.cpp
+++ b/Engine/Source/Runtime/Renderer/Private/PathTracing.cpp
@@ -442,7 +442,7 @@ class FPathTracingRG : public FGlobalShader
 	{
 		FPermutationDomain PermutationVector(Parameters.PermutationId);
 		const bool bUseCompaction = PermutationVector.Get<FEnableCompactionDim>();
-		if (bUseCompaction && !GRHISupportsWaveOperations)
+		if (bUseCompaction && !RHISupportsWaveOperations(Parameters.Platform))
 		{
 			// compaction requires wave operations
 			return false;
@@ -1240,6 +1240,11 @@ void FDeferredShadingSceneRenderer::PreparePathTracing(const FSceneViewFamily& V
 		{
 			FPathTracingRG::FPermutationDomain PermutationVector;
 			PermutationVector.Set<FPathTracingRG::FEnableCompactionDim>(EnableCompaction != 0);
+			FGlobalShaderPermutationParameters Parameters(FPathTracingRG::StaticGetTypeLayout().Name, ViewFamily.GetShaderPlatform(), PermutationVector.ToDimensionValueId());
+			if (!FPathTracingRG::ShouldCompilePermutation(Parameters))
+			{
+				continue;
+			}
 			auto RayGenShader = GetGlobalShaderMap(ViewFamily.GetShaderPlatform())->GetShader<FPathTracingRG>(PermutationVector);
 			OutRayGenShaders.Add(RayGenShader.GetRayTracingShader());
 		}
@@ -1400,7 +1405,7 @@ void FDeferredShadingSceneRenderer::RenderPathTracing(
 		const int32 DispatchSizeY = FMath::DivideAndRoundUp(View.ViewRect.Size().Y, PATHTRACER_COHERENT_TILE_SIZE) * PATHTRACER_COHERENT_TILE_SIZE;
 
 		// should we use path compaction?
-		const bool bUseCompaction = CVarPathTracingCompaction.GetValueOnRenderThread() != 0 && GRHISupportsWaveOperations;
+		const bool bUseCompaction = CVarPathTracingCompaction.GetValueOnRenderThread() != 0 && GRHISupportsWaveOperations && RHISupportsWaveOperations(View.GetShaderPlatform());
 		FRDGBuffer* ActivePaths[2] = {};
 		FRDGBuffer* NumActivePaths = nullptr;
 		FRDGBuffer* PathStateData = nullptr;
