Change visiblity buffer clear and scene depth filling to use tile data.

diff --git a/Engine/Source/Runtime/Renderer/Private/HairStrands/HairStrandsVisibility.cpp b/Engine/Source/Runtime/Renderer/Private/HairStrands/HairStrandsVisibility.cpp
--- a/Engine/Source/Runtime/Renderer/Private/HairStrands/HairStrandsVisibility.cpp
+++ b/Engine/Source/Runtime/Renderer/Private/HairStrands/HairStrandsVisibility.cpp
@@ -751,7 +751,7 @@ static FRDGBufferRef AddUpdateSampleCoveragePass(
 		1);
 	FComputeShaderUtils::AddPass(
 		GraphBuilder,
-		RDG_EVENT_NAME("HairStrandsVisbilityUpdateCoverage"),
+		RDG_EVENT_NAME("HairStrands::VisbilityUpdateCoverage"),
 		ComputeShader,
 		Parameters,
 		DispatchCount);
@@ -859,7 +859,7 @@ static FMaterialPassOutput AddHairMaterialPass(
 			Scene, 
 			*ViewInfo, 
 			&InstanceCullingManager, 
-			RDG_EVENT_NAME("HairStrandsMaterialPass(Emissive=%s)", Filter == EHairMaterialPassFilter::All ? TEXT("On/Off") : (Filter == EHairMaterialPassFilter::EmissiveOnly ? TEXT("On") : TEXT("Off"))),
+			RDG_EVENT_NAME("HairStrands::MaterialPass(Emissive=%s)", Filter == EHairMaterialPassFilter::All ? TEXT("On/Off") : (Filter == EHairMaterialPassFilter::EmissiveOnly ? TEXT("On") : TEXT("Off"))),
 			FIntRect(0, 0, Resolution.X, Resolution.Y),
 		[PassParameters, Scene = Scene, ViewInfo, &MacroGroupDatas, MaxNodeCount, Resolution, NodeGroupSize, bUpdateSampleCoverage, Filter, FeatureLevel](FDynamicPassMeshDrawListContext* ShadowContext)
 		{
@@ -1023,7 +1023,7 @@ static void AddHairVelocityPass(
 	TShaderMapRef<FHairVelocityCS> ComputeShader(View.ShaderMap, PermutationVector);
 	FComputeShaderUtils::AddPass(
 		GraphBuilder,
-		RDG_EVENT_NAME("HairStrandsVelocity"),
+		RDG_EVENT_NAME("HairStrands::Velocity"),
 		ComputeShader,
 		PassParameters,
 		FComputeShaderUtils::GetGroupCount(RectResolution, GroupSize));
@@ -1078,7 +1078,7 @@ static FRDGTextureRef AddHairLightChannelMaskPass(
 	TShaderMapRef<FHairLightChannelMaskCS> ComputeShader(View.ShaderMap, PermutationVector);
 	FComputeShaderUtils::AddPass(
 		GraphBuilder,
-		RDG_EVENT_NAME("HairLightChannelMask"),
+		RDG_EVENT_NAME("HairStrands::LightChannelMask"),
 		ComputeShader,
 		PassParameters,
 		FComputeShaderUtils::GetGroupCount(Resolution, GroupSize));
@@ -1453,6 +1453,7 @@ class FClearUIntGraphicPS : public FGlobalShader
 
 	BEGIN_SHADER_PARAMETER_STRUCT(FParameters, )
 		SHADER_PARAMETER(uint32, ClearValue)
+		SHADER_PARAMETER_STRUCT_INCLUDE(FHairStrandsTilePassVS::FParameters, TileData)
 		RENDER_TARGET_BINDING_SLOTS()
 	END_SHADER_PARAMETER_STRUCT()
 
@@ -1475,18 +1476,20 @@ public:
 
 IMPLEMENT_GLOBAL_SHADER(FClearUIntGraphicPS, "/Engine/Private/HairStrands/HairStrandsVisibilityClearPS.usf", "ClearPS", SF_Pixel);
 
-// Opaque mask
 static void AddClearGraphicPass(
 	FRDGBuilder& GraphBuilder,
 	FRDGEventName&& PassName,
 	const FViewInfo* View,
 	const uint32 ClearValue,
+	const FHairStrandsTiles& TileData,
 	FRDGTextureRef& OutTarget)
 {
 	check(OutTarget);
+	const bool bUseTile = TileData.IsValid();
 
 	FClearUIntGraphicPS::FParameters* Parameters = GraphBuilder.AllocParameters<FClearUIntGraphicPS::FParameters>();
 	Parameters->ClearValue = ClearValue;
+	Parameters->TileData = GetHairStrandsTileParameters(*View, TileData, FHairStrandsTiles::ETileType::Hair);
 	Parameters->RenderTargets[0] = FRenderTargetBinding(OutTarget, ERenderTargetLoadAction::ENoAction, 0);
 
 	FClearUIntGraphicPS::FPermutationDomain PermutationVector;
@@ -1499,43 +1502,54 @@ static void AddClearGraphicPass(
 		PermutationVector.Set<FClearUIntGraphicPS::FOutputFormat>(1);
 	}
 
-	TShaderMapRef<FPostProcessVS> VertexShader(View->ShaderMap);
+	TShaderMapRef<FPostProcessVS> ScreenVertexShader(View->ShaderMap);
+	TShaderMapRef<FHairStrandsTilePassVS> TileVertexShader(View->ShaderMap);
 	TShaderMapRef<FClearUIntGraphicPS> PixelShader(View->ShaderMap, PermutationVector);
-	const FIntRect Viewport = FIntRect(FIntPoint(0, 0), OutTarget->Desc.Extent);// View->ViewRect;
+	const FIntRect Viewport = bUseTile ? View->ViewRect : FIntRect(FIntPoint(0, 0), OutTarget->Desc.Extent);
 	const FIntPoint Resolution = OutTarget->Desc.Extent;
 
-	ClearUnusedGraphResources(PixelShader, Parameters);
+	//ClearUnusedGraphResources(PixelShader, Parameters);
 
 	GraphBuilder.AddPass(
 		Forward<FRDGEventName>(PassName),
 		Parameters,
 		ERDGPassFlags::Raster,
-		[Parameters, VertexShader, PixelShader, Viewport, Resolution](FRHICommandList& RHICmdList)
+		[Parameters, ScreenVertexShader, TileVertexShader, PixelShader, Viewport, Resolution, bUseTile](FRHICommandList& RHICmdList)
 	{
+		FHairStrandsTilePassVS::FParameters ParametersVS = Parameters->TileData;
+
 		FGraphicsPipelineStateInitializer GraphicsPSOInit;
 		RHICmdList.ApplyCachedRenderTargets(GraphicsPSOInit);
 		GraphicsPSOInit.BlendState = TStaticBlendState<CW_RGBA, BO_Add, BF_One, BF_Zero, BO_Add, BF_One, BF_Zero>::GetRHI();
 		GraphicsPSOInit.RasterizerState = TStaticRasterizerState<>::GetRHI();
 		GraphicsPSOInit.DepthStencilState = TStaticDepthStencilState<false, CF_Always>::GetRHI();
 		GraphicsPSOInit.BoundShaderState.VertexDeclarationRHI = GFilterVertexDeclaration.VertexDeclarationRHI;
-		GraphicsPSOInit.BoundShaderState.VertexShaderRHI = VertexShader.GetVertexShader();
+		GraphicsPSOInit.BoundShaderState.VertexShaderRHI = bUseTile ? TileVertexShader.GetVertexShader() : ScreenVertexShader.GetVertexShader();
 		GraphicsPSOInit.BoundShaderState.PixelShaderRHI = PixelShader.GetPixelShader();
-		GraphicsPSOInit.PrimitiveType = PT_TriangleList;
+		GraphicsPSOInit.PrimitiveType = Parameters->TileData.bRectPrimitive > 0 ? PT_RectList : PT_TriangleList;
 		SetGraphicsPipelineState(RHICmdList, GraphicsPSOInit);
 
 		RHICmdList.SetViewport(Viewport.Min.X, Viewport.Min.Y, 0.0f, Viewport.Max.X, Viewport.Max.Y, 1.0f);
 		SetShaderParameters(RHICmdList, PixelShader, PixelShader.GetPixelShader(), *Parameters);
-
-		DrawRectangle(
-			RHICmdList,
-			0, 0,
-			Viewport.Width(),Viewport.Height(),
-			Viewport.Min.X,  Viewport.Min.Y,
-			Viewport.Width(),Viewport.Height(),
-			Viewport.Size(),
-			Resolution,
-			VertexShader,
-			EDRF_UseTriangleOptimization);
+		if (bUseTile)
+		{
+			SetShaderParameters(RHICmdList, TileVertexShader, TileVertexShader.GetVertexShader(), ParametersVS);
+			RHICmdList.SetStreamSource(0, nullptr, 0);
+			RHICmdList.DrawPrimitiveIndirect(Parameters->TileData.TileIndirectBuffer->GetRHI(), 0);
+		}
+		else
+		{
+			DrawRectangle(
+				RHICmdList,
+				0, 0,
+				Viewport.Width(), Viewport.Height(),
+				Viewport.Min.X, Viewport.Min.Y,
+				Viewport.Width(), Viewport.Height(),
+				Viewport.Size(),
+				Resolution,
+				ScreenVertexShader,
+				EDRF_UseTriangleOptimization);
+		}
 	});
 }
 
@@ -1580,7 +1594,7 @@ static FRDGBufferRef AddCopyIndirectArgPass(
 
 	FComputeShaderUtils::AddPass(
 		GraphBuilder,
-		RDG_EVENT_NAME("HairStrandsVisbilityCopyIndirectArgs"),
+		RDG_EVENT_NAME("HairStrands::VisbilityCopyIndirectArgs"),
 		ComputeShader,
 		Parameters,
 		FIntVector(1,1,1));
@@ -1793,7 +1807,7 @@ static void AddHairVisibilityPrimitiveIdCompactionPass(
 	{
 		FComputeShaderUtils::AddPass(
 			GraphBuilder,
-			RDG_EVENT_NAME("HairStrandsVisibilityCompaction(Tile)"),
+			RDG_EVENT_NAME("HairStrands::VisibilityCompaction(Tile)"),
 			ComputeShader,
 			PassParameters,
 			TileData.TileIndirectDispatchBuffer,
@@ -1806,7 +1820,7 @@ static void AddHairVisibilityPrimitiveIdCompactionPass(
 
 		FComputeShaderUtils::AddPass(
 			GraphBuilder,
-			RDG_EVENT_NAME("HairStrandsVisibilityCompaction(Screen)"),
+			RDG_EVENT_NAME("HairStrands::VisibilityCompaction(Screen)"),
 			ComputeShader,
 			PassParameters,
 			FComputeShaderUtils::GetGroupCount(RectResolution, GroupSize));
@@ -1981,11 +1995,14 @@ class FHairVisibilityFillOpaqueDepthPS : public FGlobalShader
 	DECLARE_GLOBAL_SHADER(FHairVisibilityFillOpaqueDepthPS);
 	SHADER_USE_PARAMETER_STRUCT(FHairVisibilityFillOpaqueDepthPS, FGlobalShader);
 
+	class FTile : SHADER_PERMUTATION_BOOL("PERMUTATION_TILE");
+	using FPermutationDomain = TShaderPermutationDomain<FTile>;
+
 	BEGIN_SHADER_PARAMETER_STRUCT(FParameters, )
 		SHADER_PARAMETER_RDG_TEXTURE(Texture2D, SceneDepthTexture)
 		SHADER_PARAMETER_RDG_TEXTURE(Texture2D, VisibilityDepthTexture)
 		SHADER_PARAMETER_RDG_TEXTURE(Texture2D, VisibilityIDTexture)
-
+		SHADER_PARAMETER_STRUCT_INCLUDE(FHairStrandsTilePassVS::FParameters, TileData)
 		SHADER_PARAMETER_STRUCT_REF(FViewUniformShaderParameters, ViewUniformBuffer)
 		RENDER_TARGET_BINDING_SLOTS()
 	END_SHADER_PARAMETER_STRUCT()
@@ -2001,17 +2018,16 @@ static FRDGTextureRef AddHairVisibilityFillOpaqueDepth(
 	const FViewInfo& View,
 	const FIntPoint& Resolution,
 	const FHairStrandsMacroGroupDatas& MacroGroupDatas,
+	const FHairStrandsTiles& TileData,
 	const FRDGTextureRef& SceneDepthTexture)
 {
-	FRDGTextureRef OutVisibilityDepthTexture;
-	{
-		check(GetHairVisibilityRenderMode() == HairVisibilityRenderMode_MSAA_Visibility);
+	check(GetHairVisibilityRenderMode() == HairVisibilityRenderMode_MSAA_Visibility);
 
-		FRDGTextureDesc Desc = FRDGTextureDesc::Create2D(Resolution, PF_D24, FClearValueBinding::DepthFar, TexCreate_DepthStencilTargetable | TexCreate_ShaderResource, 1, GetMaxSamplePerPixel());
-		OutVisibilityDepthTexture = GraphBuilder.CreateTexture(Desc, TEXT("Hair.VisibilityDepthTexture"));
-	}
+	const bool bUseTile = TileData.IsValid();
+	FRDGTextureRef OutVisibilityDepthTexture = GraphBuilder.CreateTexture(FRDGTextureDesc::Create2D(Resolution, PF_D24, FClearValueBinding::DepthFar, TexCreate_DepthStencilTargetable | TexCreate_ShaderResource, 1, GetMaxSamplePerPixel()), TEXT("Hair.VisibilityDepthTexture"));
 
 	FHairVisibilityFillOpaqueDepthPS::FParameters* Parameters = GraphBuilder.AllocParameters<FHairVisibilityFillOpaqueDepthPS::FParameters>();
+	Parameters->TileData = GetHairStrandsTileParameters(View, TileData, FHairStrandsTiles::ETileType::Hair);
 	Parameters->ViewUniformBuffer = View.ViewUniformBuffer;
 	Parameters->SceneDepthTexture = SceneDepthTexture;
 	Parameters->RenderTargets.DepthStencil = FDepthStencilBinding(
@@ -2020,29 +2036,65 @@ static FRDGTextureRef AddHairVisibilityFillOpaqueDepth(
 		ERenderTargetLoadAction::ENoAction,
 		FExclusiveDepthStencil::DepthWrite_StencilNop);
 
-	TShaderMapRef<FPostProcessVS> VertexShader(View.ShaderMap);
-	TShaderMapRef<FHairVisibilityFillOpaqueDepthPS> PixelShader(View.ShaderMap);
-	const FGlobalShaderMap* GlobalShaderMap = View.ShaderMap;
+	FHairVisibilityFillOpaqueDepthPS::FPermutationDomain PermutationVector;
+	PermutationVector.Set<FHairVisibilityFillOpaqueDepthPS::FTile>(bUseTile);
+	TShaderMapRef<FHairVisibilityFillOpaqueDepthPS> PixelShader(View.ShaderMap, PermutationVector);
+	
 	const FIntRect Viewport = View.ViewRect;
-
-	TArray<FIntRect> MacroGroupRects;
-	if (IsHairStrandsViewRectOptimEnable())
+	if (bUseTile)
 	{
-		for (const FHairStrandsMacroGroupData& MacroGroupData : MacroGroupDatas)
-		{
-			MacroGroupRects.Add(MacroGroupData.ScreenRect);
-		}
+		TShaderMapRef<FHairStrandsTilePassVS> TileVertexShader(View.ShaderMap);
+		//ClearUnusedGraphResources(PixelShader, Parameters);
+
+		GraphBuilder.AddPass(
+			RDG_EVENT_NAME("HairStrands::FillVisibilityDepth(Tile)"),
+			Parameters,
+			ERDGPassFlags::Raster,
+			[Parameters, TileVertexShader, PixelShader, Viewport](FRHICommandList& RHICmdList)
+			{
+				FHairStrandsTilePassVS::FParameters ParametersVS = Parameters->TileData;
+
+				FGraphicsPipelineStateInitializer GraphicsPSOInit;
+				RHICmdList.ApplyCachedRenderTargets(GraphicsPSOInit);
+				GraphicsPSOInit.BlendState = TStaticBlendState<CW_RGBA, BO_Add, BF_One, BF_Zero, BO_Add, BF_One, BF_Zero>::GetRHI();
+				GraphicsPSOInit.RasterizerState = TStaticRasterizerState<>::GetRHI();
+				GraphicsPSOInit.DepthStencilState = TStaticDepthStencilState<true, CF_DepthNearOrEqual>::GetRHI();
+
+				GraphicsPSOInit.BoundShaderState.VertexDeclarationRHI = GFilterVertexDeclaration.VertexDeclarationRHI;
+				GraphicsPSOInit.BoundShaderState.VertexShaderRHI = TileVertexShader.GetVertexShader();
+				GraphicsPSOInit.BoundShaderState.PixelShaderRHI = PixelShader.GetPixelShader();
+				GraphicsPSOInit.PrimitiveType = Parameters->TileData.bRectPrimitive > 0 ? PT_RectList : PT_TriangleList;
+				SetGraphicsPipelineState(RHICmdList, GraphicsPSOInit);
+
+				SetShaderParameters(RHICmdList, PixelShader, PixelShader.GetPixelShader(), *Parameters);
+
+				RHICmdList.SetViewport(Viewport.Min.X, Viewport.Min.Y, 0.0f, Viewport.Max.X, Viewport.Max.Y, 1.0f);
+				SetShaderParameters(RHICmdList, TileVertexShader, TileVertexShader.GetVertexShader(), ParametersVS);
+				RHICmdList.SetStreamSource(0, nullptr, 0);
+				RHICmdList.DrawPrimitiveIndirect(Parameters->TileData.TileIndirectBuffer->GetRHI(), 0);
+			});
+
 	}
 	else
 	{
-		MacroGroupRects.Add(Viewport);
-	}
+		TShaderMapRef<FPostProcessVS> VertexShader(View.ShaderMap);
+		TArray<FIntRect> MacroGroupRects;
+		if (IsHairStrandsViewRectOptimEnable())
+		{
+			for (const FHairStrandsMacroGroupData& MacroGroupData : MacroGroupDatas)
+			{
+				MacroGroupRects.Add(MacroGroupData.ScreenRect);
+			}
+		}
+		else
+		{
+			MacroGroupRects.Add(Viewport);
+		}
 
-	{
 		ClearUnusedGraphResources(PixelShader, Parameters);
 
 		GraphBuilder.AddPass(
-			RDG_EVENT_NAME("HairStrandsVisibilityFillOpaqueDepth"),
+			RDG_EVENT_NAME("HairStrands::FillVisibilityDepth(View)"),
 			Parameters,
 			ERDGPassFlags::Raster,
 			[Parameters, VertexShader, PixelShader, MacroGroupRects, Viewport, Resolution](FRHICommandList& RHICmdList)
@@ -2123,10 +2175,10 @@ static void AddHairVisibilityCommonPass(
 	{
 		switch (RenderMode)
 		{
-		case HairVisibilityRenderMode_PPLL:						return RDG_EVENT_NAME("HairStrandsVisibilityPPLLPass");
-		case HairVisibilityRenderMode_MSAA_Visibility:			return RDG_EVENT_NAME("HairStrandsVisibilityMSAAVisPass");
-		case HairVisibilityRenderMode_Transmittance:			return RDG_EVENT_NAME("HairStrandsTransmittancePass");
-		case HairVisibilityRenderMode_TransmittanceAndHairCount:return RDG_EVENT_NAME("HairStrandsTransmittanceAndHairCountPass");
+		case HairVisibilityRenderMode_PPLL:						return RDG_EVENT_NAME("HairStrands::VisibilityPPLLPass");
+		case HairVisibilityRenderMode_MSAA_Visibility:			return RDG_EVENT_NAME("HairStrands::VisibilityMSAAVisPass");
+		case HairVisibilityRenderMode_Transmittance:			return RDG_EVENT_NAME("HairStrands::TransmittancePass");
+		case HairVisibilityRenderMode_TransmittanceAndHairCount:return RDG_EVENT_NAME("HairStrands::TransmittanceAndHairCountPass");
 		default:												return RDG_EVENT_NAME("Noname");
 		}
 	};
@@ -2201,6 +2253,7 @@ static void AddHairVisibilityMSAAPass(
 	const FViewInfo* ViewInfo,
 	const FHairStrandsMacroGroupDatas& MacroGroupDatas,
 	const FIntPoint& Resolution,
+	const FHairStrandsTiles& TileData,
 	FInstanceCullingManager& InstanceCullingManager,
 	FRDGTextureRef& OutVisibilityIdTexture,
 	FRDGTextureRef& OutVisibilityDepthTexture)
@@ -2212,7 +2265,7 @@ static void AddHairVisibilityMSAAPass(
 			OutVisibilityIdTexture = GraphBuilder.CreateTexture(Desc, TEXT("Hair.VisibilityIDTexture"));
 		}
 
-		AddClearGraphicPass(GraphBuilder, RDG_EVENT_NAME("HairStrandsClearVisibilityMSAAIdTexture"), ViewInfo, 0xFFFFFFFF, OutVisibilityIdTexture);
+		AddClearGraphicPass(GraphBuilder, RDG_EVENT_NAME("HairStrands::ClearVisibilityMSAAIdTexture(%s)", TileData.IsValid()? TEXT("Tile") : TEXT("Screen")), ViewInfo, 0xFFFFFFFF, TileData, OutVisibilityIdTexture);
 
 		FVisibilityPassParameters* PassParameters = GraphBuilder.AllocParameters<FVisibilityPassParameters>();
 		PassParameters->UniformBuffer = CreatePassDummyTextures(GraphBuilder);
@@ -2395,7 +2448,7 @@ static void AddHairViewTransmittanceDepthPass(
 	ClearUnusedGraphResources(PixelShader, Parameters);
 
 	GraphBuilder.AddPass(
-		RDG_EVENT_NAME("HairStrandsViewTransmittanceDepth"),
+		RDG_EVENT_NAME("HairStrands::ViewTransmittanceDepth"),
 		Parameters,
 		ERDGPassFlags::Raster,
 		[Parameters, VertexShader, PixelShader, Viewport, Resolution](FRHICommandList& RHICmdList)
@@ -2516,7 +2569,7 @@ static void AddHairVisibilityCommonPatchPass(
 	const bool bUseTile = TileData.IsValid();
 
 	GraphBuilder.AddPass(
-		RDG_EVENT_NAME("%s(%s)", (PassType == EHairPatchPass::DepthClear || PassType == EHairPatchPass::DepthPatch) ? TEXT("HairStrandsVisibilityPatchHairOnlyDepth") : TEXT("HairStrandsVisibilityWriteColorAndDepth"), bUseTile ? TEXT("Tile") : TEXT("Screen")),
+		RDG_EVENT_NAME("%s(%s)", (PassType == EHairPatchPass::DepthClear || PassType == EHairPatchPass::DepthPatch) ? TEXT("HairStrands::VisibilityPatchHairOnlyDepth") : TEXT("HairStrands::VisibilityWriteColorAndDepth"), bUseTile ? TEXT("Tile") : TEXT("Screen")),
 		Parameters,
 		ERDGPassFlags::Raster,
 		[Parameters, ScreenVertexShader, TileVertexShader, PixelShader, Viewport, Resolution, bUseTile](FRHICommandList& RHICmdList)
@@ -3250,6 +3303,7 @@ void RenderHairStrandsVisibilityBuffer(
 					View,
 					Resolution,
 					MacroGroupDatas,
+					VisibilityData.TileData,
 					SceneDepthTexture);
 
 				AddHairVisibilityMSAAPass(
@@ -3258,6 +3312,7 @@ void RenderHairStrandsVisibilityBuffer(
 					&View,
 					MacroGroupDatas,
 					Resolution,
+					VisibilityData.TileData,
 					InstanceCullingManager,
 					MsaaVisibilityResources.IdTexture,
 					MsaaVisibilityResources.DepthTexture);
