Fixed velocity texture being cleared for each view rather than once for all views.

diff --git a/Engine/Source/Runtime/Renderer/Private/VelocityRendering.cpp b/Engine/Source/Runtime/Renderer/Private/VelocityRendering.cpp
--- a/Engine/Source/Runtime/Renderer/Private/VelocityRendering.cpp
+++ b/Engine/Source/Runtime/Renderer/Private/VelocityRendering.cpp
@@ -272,6 +272,8 @@ void FSceneRenderer::RenderVelocities(
 				VelocityLoadAction = ERenderTargetLoadAction::ELoad;
 			}
 
+			VelocityLoadAction = View.DecayLoadAction(VelocityLoadAction);
+
 			if (!bHasAnyDraw)
 			{
 				continue;
