Fixed invalid order of GTAO input textures.

diff --git a/Engine/Source/Runtime/Renderer/Private/CompositionLighting/CompositionLighting.cpp b/Engine/Source/Runtime/Renderer/Private/CompositionLighting/CompositionLighting.cpp
--- a/Engine/Source/Runtime/Renderer/Private/CompositionLighting/CompositionLighting.cpp
+++ b/Engine/Source/Runtime/Renderer/Private/CompositionLighting/CompositionLighting.cpp
@@ -259,8 +259,8 @@ static FScreenPassTexture AddPostProcessingGTAOAllPasses(
 					GraphBuilder,
 					View,
 					CommonParameters,
-					CommonParameters.SceneDepth,
-					CurrentOutput);
+					HorizonSearchOutputs.Color,
+					CommonParameters.SceneDepth);
 		}
 
 		if (bTemporalPass)
