Fixed typo that incorrectly overrode the instance scene data ptr with the instance payload data ptr

diff --git a/Engine/Source/Runtime/Renderer/Private/SceneRendering.cpp b/Engine/Source/Runtime/Renderer/Private/SceneRendering.cpp
--- a/Engine/Source/Runtime/Renderer/Private/SceneRendering.cpp
+++ b/Engine/Source/Runtime/Renderer/Private/SceneRendering.cpp
@@ -1842,7 +1842,7 @@ void FViewInfo::SetupUniformBufferParameters(
 
 		if (InstancePayloadDataOverrideSRV)
 		{
-			ViewUniformShaderParameters.InstanceSceneData = InstancePayloadDataOverrideSRV;
+			ViewUniformShaderParameters.InstancePayloadData = InstancePayloadDataOverrideSRV;
 		}
 		else if (Scene && Scene->GPUScene.InstancePayloadDataBuffer.SRV)
 		{
