Fix distance field shadows when LWC is enabled.
- Incorrect reinterpret cast to float.

diff --git a/Engine/Source/Runtime/Renderer/Private/DistanceFieldObjectManagement.cpp b/Engine/Source/Runtime/Renderer/Private/DistanceFieldObjectManagement.cpp
--- a/Engine/Source/Runtime/Renderer/Private/DistanceFieldObjectManagement.cpp
+++ b/Engine/Source/Runtime/Renderer/Private/DistanceFieldObjectManagement.cpp
@@ -745,7 +745,7 @@ void FDistanceFieldSceneData::UpdateDistanceFieldObjectBuffers(
 									Vector4.W = *(const float*)&GPUSceneInstanceIndex;
 									UploadObjectData[4] = Vector4;
 
-									FMatrix VolumeToWorldT = VolumeToWorld.GetTransposed();
+									const FMatrix44f VolumeToWorldT = VolumeToWorld.GetTransposed();
 									UploadObjectData[5] = *(FVector4*)&VolumeToWorldT.M[0];
 									UploadObjectData[6] = *(FVector4*)&VolumeToWorldT.M[1];
 									UploadObjectData[7] = *(FVector4*)&VolumeToWorldT.M[2];
diff --git a/Engine/Source/Runtime/Renderer/Private/GlobalDistanceField.cpp b/Engine/Source/Runtime/Renderer/Private/GlobalDistanceField.cpp
--- a/Engine/Source/Runtime/Renderer/Private/GlobalDistanceField.cpp
+++ b/Engine/Source/Runtime/Renderer/Private/GlobalDistanceField.cpp
@@ -111,7 +111,7 @@ int32 GAODrawGlobalDistanceFieldModifiedPrimitives = 0;
 FAutoConsoleVariableRef CVarAODrawGlobalDistanceFieldModifiedPrimitives(
 	TEXT("r.AOGlobalDistanceFieldDrawModifiedPrimitives"),
 	GAODrawGlobalDistanceFieldModifiedPrimitives,
-	TEXT("Whether to lodrawg primitive modifications (add, remove, updatetransform) that caused an update of the global distance field.\n")
+	TEXT("Whether to draw primitive modifications (add, remove, updatetransform) that caused an update of the global distance field.\n")
 	TEXT("This can be useful for tracking down why updating the global distance field is always costing a lot, since it should be mostly cached."),
 	ECVF_Scalability | ECVF_RenderThreadSafe
 	);
diff --git a/Engine/Source/Runtime/Renderer/Private/HeightfieldLighting.cpp b/Engine/Source/Runtime/Renderer/Private/HeightfieldLighting.cpp
--- a/Engine/Source/Runtime/Renderer/Private/HeightfieldLighting.cpp
+++ b/Engine/Source/Runtime/Renderer/Private/HeightfieldLighting.cpp
@@ -634,8 +634,8 @@ void FillHeightfieldDescriptionData(const TArray<FHeightfieldComponentDescriptio
 		HeightfieldDescriptionData.Add(FVector4(Description.HeightfieldRect.Size().X, Description.HeightfieldRect.Size().Y, 1.f / Description.HeightfieldRect.Size().X, 1.f / Description.HeightfieldRect.Size().Y));
 		HeightfieldDescriptionData.Add(FVector4(InvLightingAtlasSize.X, InvLightingAtlasSize.Y, 0.f, 0.f));
 
-		const FMatrix LocalToWorldT = Description.LocalToWorld.GetTransposed();
-		const FMatrix WorldToLocalT = Description.LocalToWorld.Inverse().GetTransposed();
+		const FMatrix44f LocalToWorldT = Description.LocalToWorld.GetTransposed();
+		const FMatrix44f WorldToLocalT = Description.LocalToWorld.Inverse().GetTransposed();
 
 		HeightfieldDescriptionData.Add(*(FVector4*)&WorldToLocalT.M[0]);
 		HeightfieldDescriptionData.Add(*(FVector4*)&WorldToLocalT.M[1]);
