Decouple shadow depth rendering from bOcclusionBeforeBasePass.
Remove early shadow depth rendering.

diff --git a/Engine/Source/Runtime/Renderer/Private/DeferredShadingRenderer.cpp b/Engine/Source/Runtime/Renderer/Private/DeferredShadingRenderer.cpp
--- a/Engine/Source/Runtime/Renderer/Private/DeferredShadingRenderer.cpp
+++ b/Engine/Source/Runtime/Renderer/Private/DeferredShadingRenderer.cpp
@@ -1787,22 +1787,6 @@ void FDeferredShadingSceneRenderer::WaitForRayTracingScene(FRDGBuilder& GraphBui
 	});
 }
 
-enum class ERayTracingWorldUpdatesDispatchPoint
-{
-	BeforeLumenSceneLighting,
-	OverlapWithBasePass
-};
-
-ERayTracingWorldUpdatesDispatchPoint GetRayTracingWorldUpdatesDispatchPoint(bool bOcclusionBeforeBasePass)
-{
-	if (bOcclusionBeforeBasePass && Lumen::UseHardwareRayTracedSceneLighting())
-	{
-		return ERayTracingWorldUpdatesDispatchPoint::BeforeLumenSceneLighting;
-	}
-
-	return ERayTracingWorldUpdatesDispatchPoint::OverlapWithBasePass;
-}
-
 #endif // RHI_RAYTRACING
 
 static TAutoConsoleVariable<float> CVarStallInitViews(
@@ -2442,17 +2426,6 @@ void FDeferredShadingSceneRenderer::Render(FRDGBuilder& GraphBuilder)
 	CSV_CUSTOM_STAT(LightCount, ShadowOff, float(SortedLightSet.AttenuationLightStart), ECsvCustomStatOp::Set);
 	CSV_CUSTOM_STAT(LightCount, ShadowOn, float(SortedLightSet.SortedLights.Num()) - float(SortedLightSet.AttenuationLightStart), ECsvCustomStatOp::Set);
 
-	// Local helper function to perform virtual shadow map allocation, which can occur early, or late.
-	const auto AllocateVirtualShadowMaps = [&](bool bPostBasePass)
-	{
-		if (VirtualShadowMapArray.IsEnabled())
-		{
-			ensureMsgf(AreLightsInLightGrid(), TEXT("Virtual shadow map setup requires local lights to be injected into the light grid (this may be caused by 'r.LightCulling.Quality=0')."));
-			// ensure(ShadowMapSetupDone)
-			VirtualShadowMapArray.BuildPageAllocations(GraphBuilder, SceneTextures, Views, SortedLightSet, VisibleLightInfos, NaniteRasterResults, bPostBasePass);
-		}
-	};
-
 	FCompositionLighting CompositionLighting(Views, SceneTextures, [this] (int32 ViewIndex)
 	{
 		return ViewPipelineStates[ViewIndex]->AmbientOcclusionMethod == EAmbientOcclusionMethod::SSAO;
@@ -2468,10 +2441,9 @@ void FDeferredShadingSceneRenderer::Render(FRDGBuilder& GraphBuilder)
 	};
 
 	// Early occlusion queries
-	const bool bOcclusionBeforeBasePass = !bNaniteEnabled && !bAnyLumenEnabled && !bHairEnable && ((DepthPass.EarlyZPassMode == EDepthDrawingMode::DDM_AllOccluders) || bIsEarlyDepthComplete);
+	const bool bOcclusionBeforeBasePass = ((DepthPass.EarlyZPassMode == EDepthDrawingMode::DDM_AllOccluders) || bIsEarlyDepthComplete);
 
 #if RHI_RAYTRACING
-	ERayTracingWorldUpdatesDispatchPoint RayTracingWorldUpdatesDispatchPoint = GetRayTracingWorldUpdatesDispatchPoint(bOcclusionBeforeBasePass);
 	bool bRayTracingSceneReady = false;
 #endif
 
@@ -2484,16 +2456,6 @@ void FDeferredShadingSceneRenderer::Render(FRDGBuilder& GraphBuilder)
 
 	BeginAsyncDistanceFieldShadowProjections(GraphBuilder, SceneTextures);
 
-	// Early Shadow depth rendering
-	if (!bHasRayTracedOverlay && bOcclusionBeforeBasePass)
-	{
-		const bool bAfterBasePass = false;
-		AllocateVirtualShadowMaps(bAfterBasePass);
-
-		RenderShadowDepthMaps(GraphBuilder, InstanceCullingManager);
-	}
-	// End early Shadow depth rendering
-
 	const bool bShouldRenderSkyAtmosphere = ShouldRenderSkyAtmosphere(Scene, ViewFamily.EngineShowFlags);
 	const bool bShouldRenderVolumetricCloudBase = ShouldRenderVolumetricCloud(Scene, ViewFamily.EngineShowFlags);
 	const bool bShouldRenderVolumetricCloud = bShouldRenderVolumetricCloudBase && !ViewFamily.EngineShowFlags.VisualizeVolumetricCloudConservativeDensity;
@@ -2509,7 +2471,8 @@ void FDeferredShadingSceneRenderer::Render(FRDGBuilder& GraphBuilder)
 
 	InitVolumetricCloudsForViews(GraphBuilder, bShouldRenderVolumetricCloudBase, InstanceCullingManager);
 
-	// Generate sky LUTs once all shadow map has been evaluated (for volumetric light shafts). Requires bOcclusionBeforeBasePass.
+	// Generate sky LUTs
+	// TODO: Valid shadow maps (for volumetric light shafts) have not yet been generated at this point in the frame. Need to resolve dependency ordering!
 	// This also must happen before the BasePass for Sky material to be able to sample valid LUTs.
 	if (bShouldRenderSkyAtmosphere)
 	{
@@ -2538,30 +2501,6 @@ void FDeferredShadingSceneRenderer::Render(FRDGBuilder& GraphBuilder)
 
 	UpdateLumenScene(GraphBuilder);
 
-	if (bOcclusionBeforeBasePass)
-	{
-#if RHI_RAYTRACING
-		if (RayTracingWorldUpdatesDispatchPoint == ERayTracingWorldUpdatesDispatchPoint::BeforeLumenSceneLighting)
-		{
-			DispatchRayTracingWorldUpdates(GraphBuilder);
-		}
-
-		// Lumen scene lighting requires ray tracing scene to be ready if HWRT shadows are desired
-		if (Lumen::UseHardwareRayTracedSceneLighting())
-		{
-			WaitForRayTracingScene(GraphBuilder);
-			bRayTracingSceneReady = true;
-		}
-#endif // RHI_RAYTRACING
-
-		{
-			LLM_SCOPE_BYTAG(Lumen);
-			RenderLumenSceneLighting(GraphBuilder, Views[0]);
-		}
-
-		ComputeVolumetricFog(GraphBuilder, SceneTextures);
-	}
-
 	FRDGTextureRef HalfResolutionDepthCheckerboardMinMaxTexture = nullptr;
 
 	// Kick off async compute cloud eraly if all depth has been written in the prepass
@@ -2579,6 +2518,10 @@ void FDeferredShadingSceneRenderer::Render(FRDGBuilder& GraphBuilder)
 	FRDGTextureRef ForwardScreenSpaceShadowMaskHairTexture = nullptr;
 	if (IsForwardShadingEnabled(ShaderPlatform))
 	{
+		// With forward shading we need to render shadow maps early
+		ensureMsgf(!VirtualShadowMapArray.IsEnabled(), TEXT("Virtual shadow maps are not supported in the forward shading path"));
+		RenderShadowDepthMaps(GraphBuilder, InstanceCullingManager);
+
 		if (bHairEnable)
 		{
 			RenderHairPrePass(GraphBuilder, Scene, Views, InstanceCullingManager);
@@ -2609,11 +2552,8 @@ void FDeferredShadingSceneRenderer::Render(FRDGBuilder& GraphBuilder)
 	}
 
 #if RHI_RAYTRACING
-	if (RayTracingWorldUpdatesDispatchPoint == ERayTracingWorldUpdatesDispatchPoint::OverlapWithBasePass)
-	{
-		// Async AS builds can potentially overlap with BasePass
-		DispatchRayTracingWorldUpdates(GraphBuilder);
-	}
+	// Async AS builds can potentially overlap with BasePass
+	DispatchRayTracingWorldUpdates(GraphBuilder);
 #endif
 
 	{
@@ -2685,13 +2625,21 @@ void FDeferredShadingSceneRenderer::Render(FRDGBuilder& GraphBuilder)
 		RenderHairBasePass(GraphBuilder, Scene, SceneTextures, Views, InstanceCullingManager);
 	}
 
-	// Shadow and fog after base pass
-	if (!bHasRayTracedOverlay && !bOcclusionBeforeBasePass)
+	// Shadows, lumen and fog after base pass
+	if (!bHasRayTracedOverlay)
 	{
-		const bool bAfterBasePass = true;
-		AllocateVirtualShadowMaps(bAfterBasePass);
+		// If forward shading is enabled, we rendered shadow maps earlier already
+		if (!IsForwardShadingEnabled(ShaderPlatform))
+		{
+			if (VirtualShadowMapArray.IsEnabled())
+			{
+				ensureMsgf(AreLightsInLightGrid(), TEXT("Virtual shadow map setup requires local lights to be injected into the light grid (this may be caused by 'r.LightCulling.Quality=0')."));
+				VirtualShadowMapArray.BuildPageAllocations(GraphBuilder, SceneTextures, Views, SortedLightSet, VisibleLightInfos, NaniteRasterResults, true);
+			}
 
-		RenderShadowDepthMaps(GraphBuilder, InstanceCullingManager);
+			RenderShadowDepthMaps(GraphBuilder, InstanceCullingManager);
+		}
+		CheckShadowDepthRenderCompleted();
 
 #if RHI_RAYTRACING
 		// Lumen scene lighting requires ray tracing scene to be ready if HWRT shadows are desired
diff --git a/Engine/Source/Runtime/Renderer/Private/LightGridInjection.cpp b/Engine/Source/Runtime/Renderer/Private/LightGridInjection.cpp
--- a/Engine/Source/Runtime/Renderer/Private/LightGridInjection.cpp
+++ b/Engine/Source/Runtime/Renderer/Private/LightGridInjection.cpp
@@ -820,6 +820,8 @@ void FDeferredShadingSceneRenderer::RenderForwardShadowProjections(
 	FRDGTextureRef& OutForwardScreenSpaceShadowMask,
 	FRDGTextureRef& OutForwardScreenSpaceShadowMaskSubPixel)
 {
+	CheckShadowDepthRenderCompleted();
+
 	const bool bIsHairEnable = HairStrands::HasViewHairStrandsData(Views);
 	bool bScreenShadowMaskNeeded = false;
 
diff --git a/Engine/Source/Runtime/Renderer/Private/SceneRendering.h b/Engine/Source/Runtime/Renderer/Private/SceneRendering.h
--- a/Engine/Source/Runtime/Renderer/Private/SceneRendering.h
+++ b/Engine/Source/Runtime/Renderer/Private/SceneRendering.h
@@ -2199,6 +2199,12 @@ protected:
 
 protected:
 	FGPUSceneDynamicContext GPUSceneDynamicContext;
+
+	void CheckShadowDepthRenderCompleted() const
+	{
+		checkf(bShadowDepthRenderCompleted, TEXT("Shadow depth rendering was not done before shadow projections, this will cause severe shadow artifacts and indicates an engine bug (pass ordering)"));
+	}
+
 private:
 	void ComputeFamilySize();
 
diff --git a/Engine/Source/Runtime/Renderer/Private/ShadowDepthRendering.cpp b/Engine/Source/Runtime/Renderer/Private/ShadowDepthRendering.cpp
--- a/Engine/Source/Runtime/Renderer/Private/ShadowDepthRendering.cpp
+++ b/Engine/Source/Runtime/Renderer/Private/ShadowDepthRendering.cpp
@@ -1675,6 +1675,8 @@ void FSceneRenderer::RenderShadowDepthMapAtlases(FRDGBuilder& GraphBuilder)
 
 void FSceneRenderer::RenderShadowDepthMaps(FRDGBuilder& GraphBuilder, FInstanceCullingManager &InstanceCullingManager)
 {
+	ensureMsgf(!bShadowDepthRenderCompleted, TEXT("RenderShadowDepthMaps called twice in the same frame"));
+
 	CSV_SCOPED_TIMING_STAT_EXCLUSIVE(RenderShadows);
 
 	TRACE_CPUPROFILER_EVENT_SCOPE(FSceneRenderer::RenderShadowDepthMaps);
diff --git a/Engine/Source/Runtime/Renderer/Private/ShadowRendering.cpp b/Engine/Source/Runtime/Renderer/Private/ShadowRendering.cpp
--- a/Engine/Source/Runtime/Renderer/Private/ShadowRendering.cpp
+++ b/Engine/Source/Runtime/Renderer/Private/ShadowRendering.cpp
@@ -1920,7 +1920,7 @@ void FSceneRenderer::RenderShadowProjections(
 	bool bProjectingForForwardShading,
 	bool bMobileModulatedProjections)
 {
-	ensureMsgf(bShadowDepthRenderCompleted, TEXT("Shadow depth rendering was not done before shadow projections, this will cause severe shadow artifacts and indicates an engine bug (pass ordering)"));
+	CheckShadowDepthRenderCompleted();
 
 	for (int32 ViewIndex = 0; ViewIndex < Views.Num(); ViewIndex++)
 	{
@@ -1986,6 +1986,8 @@ void FSceneRenderer::RenderShadowProjections(
 	const FLightSceneInfo* LightSceneInfo,
 	bool bProjectingForForwardShading)
 {
+	CheckShadowDepthRenderCompleted();
+
 	const FVisibleLightInfo& VisibleLightInfo = VisibleLightInfos[LightSceneInfo->Id];
 	const FLightSceneProxy* LightSceneProxy = LightSceneInfo->Proxy;
 
@@ -2214,6 +2216,8 @@ void FDeferredShadingSceneRenderer::RenderDeferredShadowProjections(
 	FRDGTextureRef ScreenShadowMaskSubPixelTexture,
 	bool& bInjectedTranslucentVolume)
 {
+	CheckShadowDepthRenderCompleted();
+
 	SCOPED_NAMED_EVENT(FDeferredShadingSceneRenderer_RenderShadowProjections, FColor::Emerald);
 	SCOPE_CYCLE_COUNTER(STAT_ProjectedShadowDrawTime);
 	RDG_EVENT_SCOPE(GraphBuilder, "ShadowProjectionOnOpaque");
