Merge Lighting channel generation with other buffer pass when possible and use hair tile.

diff --git a/Engine/Shaders/Private/HairStrands/HairStrandsLightChannelMask.usf b/Engine/Shaders/Private/HairStrands/HairStrandsLightChannelMask.usf
deleted file mode 100644
diff --git a/Engine/Shaders/Private/HairStrands/HairStrandsVisibilityDepthPS.usf b/Engine/Shaders/Private/HairStrands/HairStrandsVisibilityDepthPS.usf
--- a/Engine/Shaders/Private/HairStrands/HairStrandsVisibilityDepthPS.usf
+++ b/Engine/Shaders/Private/HairStrands/HairStrandsVisibilityDepthPS.usf
@@ -5,14 +5,31 @@
 #include "../DeferredShadingCommon.ush"
 #include "HairStrandsVisibilityCommon.ush"
 
+#define OUTPUT_DEPTH 0
+#define OUTPUT_GBUFFER 1 
+#define OUTPUT_LIGHTCHANNEL 2
+#define OUTPUT_GBUFFER_LIGHTCHANNEL 3
+
 // Manual binding of the sample data as the uniform buffer is not created at this point
 Texture2D<float>					CoverageTexture;
 Texture2D<uint>						HairSampleOffset;
 StructuredBuffer<FPackedHairSample> HairSampleData;
 
-float ExtractClosestDeviceZ(uint3 PixelCoord)
+#if PERMUTATION_OUTPUT_TYPE == OUTPUT_GBUFFER_LIGHTCHANNEL || PERMUTATION_OUTPUT_TYPE == OUTPUT_LIGHTCHANNEL
+RWTexture2D<uint> OutLightChannelMaskTexture;
+#endif
+
+////////////////////////////////////////////////////////////////////////////////////////////////////////////////////////////////////
+
+struct FHairSampleData
 {
-	float Out = 0; // Inverse-Z
+	float DeviceZ;
+	uint LightChannelMask;
+};
+
+FHairSampleData ExtractSampleData(uint3 PixelCoord)
+{
+	FHairSampleData Out = (FHairSampleData)0; // Inverse-Z
 
 	const FNodeDesc NodeDesc = DecodeNodeDesc(HairSampleOffset.Load(PixelCoord));
 	LOOP
@@ -20,13 +37,15 @@ float ExtractClosestDeviceZ(uint3 PixelCoord)
 	{
 		const uint SampleIndex = NodeDesc.Offset + SampleIt;
 		const FHairSample Sample = UnpackHairSample(HairSampleData[SampleIndex]);
-		Out = max(Out, Sample.Depth); // Inverse-Z
+		Out.DeviceZ = max(Out.DeviceZ, Sample.Depth); // Inverse-Z
+		Out.LightChannelMask |= Sample.LightChannelMask;
 	}
 	return Out;
 }
-////////////////////////////////////////////////////////////////////////////////////////////////////////////////////////////////////
 
-#if PERMUTATION_OUTPUT_TYPE == 0
+////////////////////////////////////////////////////////////////////////////////////////////////////////////////////////////////////
+// Gbuffer patching (+ LightChannel mask)
+#if PERMUTATION_OUTPUT_TYPE == OUTPUT_GBUFFER || PERMUTATION_OUTPUT_TYPE == OUTPUT_GBUFFER_LIGHTCHANNEL
 void MainPS(
 	in FScreenVertexOutput Input,
 	out float SvDepth		: SV_DEPTH,
@@ -37,25 +56,47 @@ void MainPS(
 {
 	const uint3 PixelCoord = uint3(Input.Position.xy, 0);
 	const float Coverage = CoverageTexture.Load(PixelCoord);
+
+	#if PERMUTATION_OUTPUT_TYPE == OUTPUT_GBUFFER
 	if (Coverage < 1.0f)
 	{
 		discard;
 	}
+	#endif
 
 	// Add a depth bias in order to avoid precision issue during the composition pass, which would reject certain samples.
 	// This depth is only used for screen passes like AO/SSGI/... so no need to be very accurate
 	const float DepthBias = 0.001f; 
-	const float DeviceZ = ExtractClosestDeviceZ(PixelCoord);
-	SvDepth = DeviceZ - DepthBias;
+	const FHairSampleData SampleData = ExtractSampleData(PixelCoord);
+	SvDepth = SampleData.DeviceZ - DepthBias;
 	OutGBufferB = float4(0, 0, 0, 0); // EncodeShadingModelIdAndSelectiveOutputMask(SHADINGMODELID_HAIR, 0)
 	OutGBufferC = 0;
 	OutColor = 0;
+
+	#if PERMUTATION_OUTPUT_TYPE == OUTPUT_GBUFFER_LIGHTCHANNEL
+	OutLightChannelMaskTexture[PixelCoord.xy] = SampleData.LightChannelMask;
+	if (Coverage < 1.0f)
+	{
+		discard;
+	}
+	#endif
 }
 #endif
 
 ////////////////////////////////////////////////////////////////////////////////////////////////////////////////////////////////////
+// LightChannel mask only
+#if PERMUTATION_OUTPUT_TYPE == OUTPUT_LIGHTCHANNEL
+void MainPS(
+	in FScreenVertexOutput Input)
+{
+	const uint3 PixelCoord = uint3(Input.Position.xy, 0);
+	OutLightChannelMaskTexture[PixelCoord.xy] = ExtractSampleData(uint3(PixelCoord, 0)).LightChannelMask;
+}
+#endif
 
-#if PERMUTATION_OUTPUT_TYPE == 1
+////////////////////////////////////////////////////////////////////////////////////////////////////////////////////////////////////
+// Depth only
+#if PERMUTATION_OUTPUT_TYPE == OUTPUT_DEPTH
 uint bClear;
 void MainPS(
 	in FScreenVertexOutput Input,
@@ -64,7 +105,7 @@ void MainPS(
 	SvDepth = 0; // Inverse-Z
 	if (bClear == 0)
 	{
-		SvDepth = ExtractClosestDeviceZ(uint3(Input.Position.xy, 0));
+		SvDepth = ExtractSampleData(uint3(Input.Position.xy, 0)).DeviceZ;
 	}
 }
 #endif
\ No newline at end of file
diff --git a/Engine/Source/Runtime/Renderer/Private/HairStrands/HairStrandsVisibility.cpp b/Engine/Source/Runtime/Renderer/Private/HairStrands/HairStrandsVisibility.cpp
--- a/Engine/Source/Runtime/Renderer/Private/HairStrands/HairStrandsVisibility.cpp
+++ b/Engine/Source/Runtime/Renderer/Private/HairStrands/HairStrandsVisibility.cpp
@@ -1029,63 +1029,6 @@ static void AddHairVelocityPass(
 		FComputeShaderUtils::GetGroupCount(RectResolution, GroupSize));
 }
 
-///////////////////////////////////////////////////////////////////////////////////////////////////
-class FHairLightChannelMaskCS : public FGlobalShader
-{
-	DECLARE_GLOBAL_SHADER(FHairLightChannelMaskCS);
-	SHADER_USE_PARAMETER_STRUCT(FHairLightChannelMaskCS, FGlobalShader);
-
-	class FGroupSize : SHADER_PERMUTATION_SPARSE_INT("PERMUTATION_GROUP_SIZE", 32, 64);
-	using FPermutationDomain = TShaderPermutationDomain<FGroupSize>;
-
-	BEGIN_SHADER_PARAMETER_STRUCT(FParameters, )
-		SHADER_PARAMETER(FIntPoint, OutputResolution)
-		SHADER_PARAMETER_RDG_BUFFER_SRV(StructuredBuffer, NodeData)
-		SHADER_PARAMETER_RDG_TEXTURE(Texture2D, NodeOffsetAndCount)
-		SHADER_PARAMETER_RDG_TEXTURE_UAV(RWTexture2D, OutLightChannelMaskTexture)
-	END_SHADER_PARAMETER_STRUCT()
-		
-public:
-
-	static bool ShouldCompilePermutation(const FGlobalShaderPermutationParameters& Parameters)
-	{
-		return IsHairStrandsSupported(EHairStrandsShaderType::Strands, Parameters.Platform);
-	}
-};
-
-IMPLEMENT_GLOBAL_SHADER(FHairLightChannelMaskCS, "/Engine/Private/HairStrands/HairStrandsLightChannelMask.usf", "MainCS", SF_Compute);
-
-static FRDGTextureRef AddHairLightChannelMaskPass(
-	FRDGBuilder& GraphBuilder,
-	const FViewInfo& View,
-	const FIntPoint Resolution,
-	FRDGBufferRef NodeData,
-	FRDGTextureRef NodeOffsetAndCount)
-{
-	FRDGTextureDesc Desc = FRDGTextureDesc::Create2D(Resolution, PF_R32_UINT, FClearValueBinding::None, TexCreate_UAV | TexCreate_ShaderResource);
-	FRDGTextureRef OutLightChannelMaskTexture = GraphBuilder.CreateTexture(Desc, TEXT("Hair.LightChannelMask"));
-
-	FHairLightChannelMaskCS::FPermutationDomain PermutationVector;
-	PermutationVector.Set<FHairLightChannelMaskCS::FGroupSize>(GetVendorOptimalGroupSize1D());
-
-	FHairLightChannelMaskCS::FParameters* PassParameters = GraphBuilder.AllocParameters<FHairLightChannelMaskCS::FParameters>();
-	PassParameters->OutputResolution = Resolution;
-	PassParameters->NodeData = GraphBuilder.CreateSRV(NodeData);
-	PassParameters->NodeOffsetAndCount = NodeOffsetAndCount;
-	PassParameters->OutLightChannelMaskTexture = GraphBuilder.CreateUAV(OutLightChannelMaskTexture);
-	
-	const FIntPoint GroupSize = GetVendorOptimalGroupSize2D();
-	TShaderMapRef<FHairLightChannelMaskCS> ComputeShader(View.ShaderMap, PermutationVector);
-	FComputeShaderUtils::AddPass(
-		GraphBuilder,
-		RDG_EVENT_NAME("HairStrands::LightChannelMask"),
-		ComputeShader,
-		PassParameters,
-		FComputeShaderUtils::GetGroupCount(Resolution, GroupSize));
-
-	return OutLightChannelMaskTexture;
-}
-
 /////////////////////////////////////////////////////////////////////////////////////////
 BEGIN_GLOBAL_SHADER_PARAMETER_STRUCT(FVisibilityPassUniformParameters, )
 	SHADER_PARAMETER(uint32, MaxPPLLNodeCount)
@@ -2131,7 +2074,10 @@ static FRDGTextureRef AddHairVisibilityFillOpaqueDepth(
 	}
 
 	// Ensure HTile is valid after manually feeding the scene depth value
-	AddResummarizeHTilePass(GraphBuilder, OutVisibilityDepthTexture);
+	if (GRHISupportsResummarizeHTile)
+	{
+		AddResummarizeHTilePass(GraphBuilder, OutVisibilityDepthTexture);
+	}
 
 	return OutVisibilityDepthTexture;
 }
@@ -2489,7 +2435,7 @@ class FHairVisibilityDepthPS : public FGlobalShader
 	DECLARE_GLOBAL_SHADER(FHairVisibilityDepthPS);
 	SHADER_USE_PARAMETER_STRUCT(FHairVisibilityDepthPS, FGlobalShader);
 
-	class FOutputType : SHADER_PERMUTATION_INT("PERMUTATION_OUTPUT_TYPE", 2);
+	class FOutputType : SHADER_PERMUTATION_INT("PERMUTATION_OUTPUT_TYPE", 4);
 	using FPermutationDomain = TShaderPermutationDomain<FOutputType>;
 
 	BEGIN_SHADER_PARAMETER_STRUCT(FParameters, )
@@ -2499,6 +2445,7 @@ class FHairVisibilityDepthPS : public FGlobalShader
 		SHADER_PARAMETER_RDG_TEXTURE(Texture2D<float>, CoverageTexture)
 		SHADER_PARAMETER_RDG_TEXTURE(Texture2D<uint>, HairSampleOffset)
 		SHADER_PARAMETER_RDG_BUFFER_SRV(StructuredBuffer<FPackedHairSample>, HairSampleData)
+		SHADER_PARAMETER_RDG_TEXTURE_UAV(RWTexture2D<uint>, OutLightChannelMaskTexture)
 		RENDER_TARGET_BINDING_SLOTS()
 	END_SHADER_PARAMETER_STRUCT()
 
@@ -2519,40 +2466,48 @@ public:
 
 IMPLEMENT_GLOBAL_SHADER(FHairVisibilityDepthPS, "/Engine/Private/HairStrands/HairStrandsVisibilityDepthPS.usf", "MainPS", SF_Pixel);
 
-enum class EHairPatchPass
+enum class EHairAuxilaryPassType
 {
 	GBufferPatch,
+	GBufferPatch_LightChannelMask,
+	LightChannelMask,
 	DepthPatch,
 	DepthClear
 };
 
-static void AddHairVisibilityCommonPatchPass(
+static void AddHairAuxilaryPass(
 	FRDGBuilder& GraphBuilder,
 	const FViewInfo& View,
 	const FHairStrandsTiles& TileData,
-	const EHairPatchPass PassType,
+	const EHairAuxilaryPassType PassType,
 	const FRDGTextureRef& CoverageTexture,
 	const FRDGTextureRef& HairSampleOffset,
 	const FRDGBufferRef& HairSampleData,
 	FRDGTextureRef OutGBufferBTexture,
 	FRDGTextureRef OutGBufferCTexture,
 	FRDGTextureRef OutColorTexture,
-	FRDGTextureRef OutDepthTexture)
+	FRDGTextureRef OutDepthTexture,
+	FRDGTextureRef OutLightChannelMaskTexture)
 {
 	FHairVisibilityDepthPS::FParameters* Parameters = GraphBuilder.AllocParameters<FHairVisibilityDepthPS::FParameters>();
-	Parameters->bClear = PassType == EHairPatchPass::DepthClear ? 1u : 0u;
-	Parameters->TileData = GetHairStrandsTileParameters(View, TileData, PassType == EHairPatchPass::DepthClear ?  FHairStrandsTiles::ETileType::Other : FHairStrandsTiles::ETileType::Hair);
+	Parameters->bClear = PassType == EHairAuxilaryPassType::DepthClear ? 1u : 0u;
+	Parameters->TileData = GetHairStrandsTileParameters(View, TileData, PassType == EHairAuxilaryPassType::DepthClear ?  FHairStrandsTiles::ETileType::Other : FHairStrandsTiles::ETileType::Hair);
 	Parameters->ViewUniformBuffer = View.ViewUniformBuffer;
 	Parameters->CoverageTexture = CoverageTexture;
 	Parameters->HairSampleOffset = HairSampleOffset;
 	Parameters->HairSampleData = GraphBuilder.CreateSRV(HairSampleData);
 
-	Parameters->RenderTargets.DepthStencil = FDepthStencilBinding(
-		OutDepthTexture,
-		ERenderTargetLoadAction::ELoad,
-		ERenderTargetLoadAction::ELoad,
-		FExclusiveDepthStencil::DepthWrite_StencilNop);
-	if (PassType == EHairPatchPass::GBufferPatch)
+	const bool bDepthTested = PassType != EHairAuxilaryPassType::LightChannelMask;
+	if (bDepthTested)
+	{
+		Parameters->RenderTargets.DepthStencil = FDepthStencilBinding(
+			OutDepthTexture,
+			ERenderTargetLoadAction::ELoad,
+			ERenderTargetLoadAction::ELoad,
+			FExclusiveDepthStencil::DepthWrite_StencilNop);
+	}
+
+	if (PassType == EHairAuxilaryPassType::GBufferPatch || PassType == EHairAuxilaryPassType::GBufferPatch_LightChannelMask)
 	{
 		check(OutGBufferBTexture && OutGBufferCTexture && OutColorTexture);
 		Parameters->RenderTargets[0] = FRenderTargetBinding(OutGBufferBTexture, ERenderTargetLoadAction::ELoad);
@@ -2560,21 +2515,39 @@ static void AddHairVisibilityCommonPatchPass(
 		Parameters->RenderTargets[2] = FRenderTargetBinding(OutColorTexture, ERenderTargetLoadAction::ELoad);
 	}
 
+	if (PassType == EHairAuxilaryPassType::GBufferPatch_LightChannelMask || PassType == EHairAuxilaryPassType::LightChannelMask)
+	{
+		check(OutLightChannelMaskTexture);
+		Parameters->OutLightChannelMaskTexture = GraphBuilder.CreateUAV(OutLightChannelMaskTexture);
+	}
+
+	uint32 OutputType = 0;
+	const TCHAR* Method = nullptr;
+	switch (PassType)
+	{
+		case EHairAuxilaryPassType::DepthPatch:						OutputType = 0; Method = TEXT("HairOnlyDepth"); break;
+		case EHairAuxilaryPassType::DepthClear:						OutputType = 0; Method = TEXT("HairOnlyDepth:Clear"); break;
+		case EHairAuxilaryPassType::GBufferPatch:					OutputType = 1; Method = TEXT("GBuffer"); break;
+		case EHairAuxilaryPassType::LightChannelMask:				OutputType = 2; Method = TEXT("LightChannel"); break;
+		case EHairAuxilaryPassType::GBufferPatch_LightChannelMask:	OutputType = 3; Method = TEXT("GBuffer, LightChannel"); break;
+		default: 													OutputType = 0; Method = TEXT("Unknown"); break;
+	};
+
 	TShaderMapRef<FPostProcessVS> ScreenVertexShader(View.ShaderMap);
 	TShaderMapRef<FHairStrandsTilePassVS> TileVertexShader(View.ShaderMap);
 
 	FHairVisibilityDepthPS::FPermutationDomain PermutationVector;
-	PermutationVector.Set<FHairVisibilityDepthPS::FOutputType>( (PassType == EHairPatchPass::DepthClear || PassType == EHairPatchPass::DepthPatch) ? 1 : 0);
+	PermutationVector.Set<FHairVisibilityDepthPS::FOutputType>(OutputType);
 	TShaderMapRef<FHairVisibilityDepthPS> PixelShader(View.ShaderMap, PermutationVector);
 	const FIntRect Viewport = View.ViewRect;
 	const FIntPoint Resolution = OutDepthTexture->Desc.Extent;
 	const bool bUseTile = TileData.IsValid();
 
 	GraphBuilder.AddPass(
-		RDG_EVENT_NAME("%s(%s)", (PassType == EHairPatchPass::DepthClear || PassType == EHairPatchPass::DepthPatch) ? TEXT("HairStrands::VisibilityPatchHairOnlyDepth") : TEXT("HairStrands::VisibilityWriteColorAndDepth"), bUseTile ? TEXT("Tile") : TEXT("Screen")),
+		RDG_EVENT_NAME("ComputeHairAuxilary(%s)(%s)", Method, bUseTile ? TEXT("Tile") : TEXT("Screen")),
 		Parameters,
 		ERDGPassFlags::Raster,
-		[Parameters, ScreenVertexShader, TileVertexShader, PixelShader, Viewport, Resolution, bUseTile](FRHICommandList& RHICmdList)
+		[Parameters, ScreenVertexShader, TileVertexShader, PixelShader, Viewport, Resolution, bUseTile, bDepthTested](FRHICommandList& RHICmdList)
 	{
 		FHairStrandsTilePassVS::FParameters ParametersVS = Parameters->TileData;
 
@@ -2582,7 +2555,9 @@ static void AddHairVisibilityCommonPatchPass(
 		RHICmdList.ApplyCachedRenderTargets(GraphicsPSOInit);
 		GraphicsPSOInit.BlendState = TStaticBlendState<CW_RGBA, BO_Add, BF_One, BF_Zero, BO_Add, BF_One, BF_Zero>::GetRHI();
 		GraphicsPSOInit.RasterizerState = TStaticRasterizerState<>::GetRHI();
-		GraphicsPSOInit.DepthStencilState = TStaticDepthStencilState<true, CF_Always>::GetRHI();
+		GraphicsPSOInit.DepthStencilState = bDepthTested ? 
+			TStaticDepthStencilState<true, CF_Always>::GetRHI() : 
+			TStaticDepthStencilState<false, CF_Always>::GetRHI();
 
 		GraphicsPSOInit.BoundShaderState.VertexDeclarationRHI = GFilterVertexDeclaration.VertexDeclarationRHI;
 		GraphicsPSOInit.BoundShaderState.VertexShaderRHI = bUseTile ? TileVertexShader.GetVertexShader() : ScreenVertexShader.GetVertexShader();
@@ -2615,8 +2590,42 @@ static void AddHairVisibilityCommonPatchPass(
 	});
 }
 
+#if RHI_RAYTRACING
+static FRDGTextureRef CreateLigthtChannelMaskTexture(FRDGBuilder& GraphBuilder, const FIntPoint& Resolution)
+{
+	return GraphBuilder.CreateTexture(FRDGTextureDesc::Create2D(Resolution, PF_R32_UINT, FClearValueBinding::None, TexCreate_UAV | TexCreate_ShaderResource), TEXT("Hair.LightChannelMask"));
+}
 
-static void AddHairVisibilityColorAndDepthPatchPass(
+static FRDGTextureRef AddHairLightChannelMaskPass(
+	FRDGBuilder& GraphBuilder,
+	const FViewInfo& View,
+	const FHairStrandsTiles& TileData,
+	const FRDGTextureRef& CoverageTexture,
+	const FRDGTextureRef& HairSampleOffset,
+	const FRDGBufferRef& HairSampleData)
+{
+	check(IsRayTracingEnabled());
+	FRDGTextureRef OutLightChannelMask = CreateLigthtChannelMaskTexture(GraphBuilder, View.ViewRect.Size());
+
+	AddHairAuxilaryPass(
+		GraphBuilder,
+		View,
+		TileData,
+		EHairAuxilaryPassType::LightChannelMask,
+		CoverageTexture,
+		HairSampleOffset,
+		HairSampleData,
+		nullptr,
+		nullptr,
+		nullptr,
+		nullptr,
+		OutLightChannelMask);
+	return OutLightChannelMask;
+}
+#endif
+
+
+static void AddHairGbufferPatchPass(
 	FRDGBuilder& GraphBuilder,
 	const FViewInfo& View,
 	const FHairStrandsTiles& TileData,
@@ -2626,25 +2635,37 @@ static void AddHairVisibilityColorAndDepthPatchPass(
 	FRDGTextureRef& OutGBufferBTexture,
 	FRDGTextureRef& OutGBufferCTexture,
 	FRDGTextureRef& OutColorTexture,
-	FRDGTextureRef& OutDepthTexture)
+	FRDGTextureRef& OutDepthTexture,
+	FRDGTextureRef& OutLightChannelMask)
 {
 	if (!OutGBufferBTexture || !OutGBufferCTexture || !OutColorTexture || !OutDepthTexture)
 	{
 		return;
 	}
 
-	AddHairVisibilityCommonPatchPass(
+#if RHI_RAYTRACING
+	const bool bLightingChannel = IsRayTracingEnabled() && OutLightChannelMask == nullptr;
+	if (bLightingChannel)
+	{
+		OutLightChannelMask = CreateLigthtChannelMaskTexture(GraphBuilder, View.ViewRect.Size());
+	}
+#else
+	const bool bLightingChannel = false;
+#endif
+
+	AddHairAuxilaryPass(
 		GraphBuilder,
 		View,
 		TileData,
-		EHairPatchPass::GBufferPatch,
+		bLightingChannel ? EHairAuxilaryPassType::GBufferPatch_LightChannelMask : EHairAuxilaryPassType::GBufferPatch,
 		CoverageTexture,
 		HairSampleOffset,
 		HairSampleData,
 		OutGBufferBTexture,
 		OutGBufferCTexture,
 		OutColorTexture,
-		OutDepthTexture);
+		OutDepthTexture,
+		OutLightChannelMask);
 }
 
 static void AddHairOnlyDepthPass(
@@ -2665,33 +2686,35 @@ static void AddHairOnlyDepthPass(
 	// If tile data are not available, then the clearly is done prior to that.
 	if (TileData.IsValid())
 	{
-		AddHairVisibilityCommonPatchPass(
+		AddHairAuxilaryPass(
 			GraphBuilder,
 			View,
 			TileData,
-			EHairPatchPass::DepthClear,
+			EHairAuxilaryPassType::DepthClear,
 			CoverageTexture,
 			HairSampleOffset,
 			HairSampleData,
 			nullptr,
 			nullptr,
 			nullptr,
-			OutDepthTexture);
+			OutDepthTexture,
+			nullptr);
 	}
 
 	// Depth value
-	AddHairVisibilityCommonPatchPass(
+	AddHairAuxilaryPass(
 		GraphBuilder,
 		View,
 		TileData,
-		EHairPatchPass::DepthPatch,
+		EHairAuxilaryPassType::DepthPatch,
 		CoverageTexture,
 		HairSampleOffset,
 		HairSampleData,
 		nullptr,
 		nullptr,
 		nullptr,
-		OutDepthTexture);
+		OutDepthTexture,
+		nullptr);
 }
 
 static void AddHairOnlyHZBPass(
@@ -3229,7 +3252,7 @@ void RenderHairStrandsVisibilityBuffer(
 					// * unlit shading model ID 
 					if (bRunColorAndDepthPatching)
 					{
-						AddHairVisibilityColorAndDepthPatchPass(
+						AddHairGbufferPatchPass(
 							GraphBuilder,
 							View,
 							VisibilityData.TileData,
@@ -3239,7 +3262,8 @@ void RenderHairStrandsVisibilityBuffer(
 							SceneGBufferBTexture,
 							SceneGBufferCTexture,
 							SceneColorTexture,
-							SceneDepthTexture);
+							SceneDepthTexture,
+							VisibilityData.LightChannelMaskTexture);
 					}
 
 					AddHairOnlyDepthPass(
@@ -3423,7 +3447,7 @@ void RenderHairStrandsVisibilityBuffer(
 					// * unlit shading model ID 
 					if (bRunColorAndDepthPatching)
 					{
-						AddHairVisibilityColorAndDepthPatchPass(
+						AddHairGbufferPatchPass(
 							GraphBuilder,
 							View,
 							VisibilityData.TileData,
@@ -3433,7 +3457,8 @@ void RenderHairStrandsVisibilityBuffer(
 							SceneGBufferBTexture,
 							SceneGBufferCTexture,
 							SceneColorTexture,
-							SceneDepthTexture);
+							SceneDepthTexture,
+							VisibilityData.LightChannelMaskTexture);
 					}
 
 					AddHairOnlyDepthPass(
@@ -3514,7 +3539,7 @@ void RenderHairStrandsVisibilityBuffer(
 
 				if (bRunColorAndDepthPatching)
 				{
-					AddHairVisibilityColorAndDepthPatchPass(
+					AddHairGbufferPatchPass(
 						GraphBuilder,
 						View,
 						VisibilityData.TileData,
@@ -3524,7 +3549,8 @@ void RenderHairStrandsVisibilityBuffer(
 						SceneGBufferBTexture,
 						SceneGBufferCTexture,
 						SceneColorTexture,
-						SceneDepthTexture);
+						SceneDepthTexture,
+						VisibilityData.LightChannelMaskTexture);
 				}
 
 				AddHairOnlyDepthPass(
@@ -3552,15 +3578,15 @@ void RenderHairStrandsVisibilityBuffer(
 			}
 
 		#if RHI_RAYTRACING
-			if (IsRayTracingEnabled() && CoverageTexture)
+			if (IsRayTracingEnabled() && VisibilityData.LightChannelMaskTexture == nullptr)
 			{
-				FRDGTextureRef LightingChannelMaskTexture = AddHairLightChannelMaskPass(
+				VisibilityData.LightChannelMaskTexture = AddHairLightChannelMaskPass(
 					GraphBuilder,
 					View,
-					Resolution,
-					CompactNodeData,
-					CompactNodeIndex);
-				VisibilityData.LightChannelMaskTexture = LightingChannelMaskTexture;
+					VisibilityData.TileData,
+					CoverageTexture,
+					CompactNodeIndex,
+					CompactNodeData);
 			}
 		#endif
 		}
