Use store action to avoid RHI validation errors.

#rnx

diff --git a/Engine/Source/Runtime/Renderer/Private/SingleLayerWaterRendering.cpp b/Engine/Source/Runtime/Renderer/Private/SingleLayerWaterRendering.cpp
--- a/Engine/Source/Runtime/Renderer/Private/SingleLayerWaterRendering.cpp
+++ b/Engine/Source/Runtime/Renderer/Private/SingleLayerWaterRendering.cpp
@@ -538,7 +538,7 @@ void FDeferredShadingSceneRenderer::BeginRenderingWaterGBuffer(FRHICommandList&
 	FRHIRenderPassInfo RPInfo;
 	SceneContext.FillGBufferRenderPassInfo(ERenderTargetLoadAction::ELoad, RPInfo, VelocityRTIndex, TangentRTIndex);
 	// Set a dummy Scene color RT to avoid gbuffer to stomp HDR scene color we want to blend over
-	RPInfo.ColorRenderTargets[0].Action = MakeRenderTargetActions(ERenderTargetLoadAction::ENoAction, ERenderTargetStoreAction::ENoAction);
+	RPInfo.ColorRenderTargets[0].Action = MakeRenderTargetActions(ERenderTargetLoadAction::ENoAction, ERenderTargetStoreAction::EStore);
 	RPInfo.ColorRenderTargets[0].RenderTarget = SceneContext.GetSceneColorSurface();
 
 	// Stencil always has to be store or certain VK drivers will leave the attachment in an undefined state.
