Fix crash happening when a large number of hair macro group are visible, beyond current expected limit, causing some data to be invalid/uninitialized.

diff --git a/Engine/Shaders/Private/HairStrands/HairStrandsDeepTransmittanceMask.usf b/Engine/Shaders/Private/HairStrands/HairStrandsDeepTransmittanceMask.usf
--- a/Engine/Shaders/Private/HairStrands/HairStrandsDeepTransmittanceMask.usf
+++ b/Engine/Shaders/Private/HairStrands/HairStrandsDeepTransmittanceMask.usf
@@ -16,9 +16,12 @@
 	#include "../ShaderPrintCommon.ush"
 #endif //DEBUG_ENABLE
 
+#ifndef MAX_HAIR_MACROGROUP_COUNT
+#error MAX_HAIR_MACROGROUP_COUNT needs to be defined
+#endif
+
 // Need to match FDeepShadowData::MaxMacroGroupCount in HairRendering.h
-#if SHADER_TRANSMITTANCE_DEEPSHADOW || SHADER_TRANSMITTANCE_VOXEL
-	#define MAX_HAIR_MACROGROUP_COUNT 16 
+#if SHADER_TRANSMITTANCE_DEEPSHADOW || SHADER_TRANSMITTANCE_VOXEL 
 	#if PERMUTATION_GROUP_SIZE == 64
 		#define TRANSMITTANCE_GROUP_SIZE 64
 	#elif PERMUTATION_GROUP_SIZE == 32
diff --git a/Engine/Shaders/Private/HairStrands/HairStrandsVisibilityCommon.ush b/Engine/Shaders/Private/HairStrands/HairStrandsVisibilityCommon.ush
--- a/Engine/Shaders/Private/HairStrands/HairStrandsVisibilityCommon.ush
+++ b/Engine/Shaders/Private/HairStrands/HairStrandsVisibilityCommon.ush
@@ -13,6 +13,8 @@
 
 #define HAIR_VISIBILITY_GROUP_COUNT_WIDTH 64
 
+#define MAX_HAIR_MACROGROUP_COUNT 16
+
 #include "HairStrandsVisibilityCommonStruct.ush"
 
  ////////////////////////////////////////////////////////////////////////////////////////////////////////////////////////////////
@@ -253,7 +255,7 @@ FPackedHairSample PackHairSample(const FHairSample In)
 	// PrimitiveId is on 28bits | MacroGroupId is on 4bits
 	Out.PrimitiveID_MacroGroupID = 
 		((In.PrimitiveId  & 0x0FFFFFFF)) |
-		((In.MacroGroupId & 0xF) << 28);
+		((In.MacroGroupId & 0xF) << 28); // If this encoding change, ensure it conforms to MAX_HAIR_MACROGROUP_COUNT
 
 	Out.Depth					= In.Depth;																			// 32bits float
 	Out.BaseColor_Roughness		= float4ToUint(float4(sqrt(In.BaseColor), In.Roughness));							// 32bits uint
@@ -279,7 +281,7 @@ FHairSample UnpackHairSample(const FPackedHairSample In)
 	Out.Coverage8bit = (In.Tangent_Coverage8bit >> 24) & 0xFF;
 
 	Out.PrimitiveId  = In.PrimitiveID_MacroGroupID & 0x0FFFFFFF;
-	Out.MacroGroupId = (In.PrimitiveID_MacroGroupID>>28) & 0xF;
+	Out.MacroGroupId = (In.PrimitiveID_MacroGroupID>>28) & 0xF; // If this encoding change, ensure it conforms to MAX_HAIR_MACROGROUP_COUNT
 
 	const float4 BaseColorAndRoughness = UintToFloat4(In.BaseColor_Roughness);
 	Out.BaseColor			= BaseColorAndRoughness.xyz * BaseColorAndRoughness.xyz;
diff --git a/Engine/Shaders/Private/HairStrands/HairStrandsVoxelPageAllocation.usf b/Engine/Shaders/Private/HairStrands/HairStrandsVoxelPageAllocation.usf
--- a/Engine/Shaders/Private/HairStrands/HairStrandsVoxelPageAllocation.usf
+++ b/Engine/Shaders/Private/HairStrands/HairStrandsVoxelPageAllocation.usf
@@ -4,6 +4,7 @@
 #include "../Matrices.ush"
 #include "HairStrandsVoxelPageCommon.ush"
 #include "HairStrandsAABBCommon.ush"
+#include "HairStrandsVisibilityCommon.ush"
 
 #if SHADER_ALLOCATEPAGEINDEX 
 
@@ -25,12 +26,16 @@ RWBuffer<uint4>								OutPageIndexResolutionAndOffsetBuffer;
 RWBuffer<uint>								OutPageIndexAllocationIndirectBufferArgs;
 RWStructuredBuffer<FVoxelizationViewInfo>	OutVoxelizationViewInfoBuffer;
 
+#ifndef MAX_HAIR_MACROGROUP_COUNT
+#error MAX_HAIR_MACROGROUP_COUNT needs to be defined
+#endif
+
 #define INVALID_OFFSET 0xFFFFFFFF
-#define MAX_INSTANCE_GROUP_COUNT 32
-groupshared uint PageIndexOffsets[MAX_INSTANCE_GROUP_COUNT];
+
+groupshared uint PageIndexOffsets[MAX_HAIR_MACROGROUP_COUNT];
 
 // This code assume we have less than 32 macro group (which fit into a single CU/SM)
-[numthreads(MAX_INSTANCE_GROUP_COUNT, 1, 1)]
+[numthreads(MAX_HAIR_MACROGROUP_COUNT, 1, 1)]
 void AllocatePageIndex(uint2 DispatchThreadId : SV_DispatchThreadID)
 {
 	const uint MacroGroupId = DispatchThreadId.x;
diff --git a/Engine/Source/Runtime/Renderer/Private/HairStrands/HairStrandsCluster.cpp b/Engine/Source/Runtime/Renderer/Private/HairStrands/HairStrandsCluster.cpp
--- a/Engine/Source/Runtime/Renderer/Private/HairStrands/HairStrandsCluster.cpp
+++ b/Engine/Source/Runtime/Renderer/Private/HairStrands/HairStrandsCluster.cpp
@@ -115,7 +115,7 @@ static void AddHairMacroGroupAABBPass(
 
 static bool DoesGroupExists(uint32 ResourceId, uint32 GroupIndex, const FHairStrandsMacroGroupData::TPrimitiveInfos& PrimitivesGroups)
 {
-	// Simple linear search as the expected number of groups is supposed to be low (<10)
+	// Simple linear search as the expected number of groups is supposed to be low (<16, see FHairStrandsMacroGroupData::MaxMacroGroupCount)
 	for (const FHairStrandsMacroGroupData::PrimitiveInfo& Group : PrimitivesGroups)
 	{
 		if (Group.GroupIndex == GroupIndex && Group.ResourceId == ResourceId)
@@ -151,7 +151,6 @@ static void InternalUpdateMacroGroup(FHairStrandsMacroGroupData& MacroGroup, int
 		PrimitiveInfo.ResourceId = reinterpret_cast<uint64>(MeshBatchAndRelevance->Mesh->Elements[0].UserData);
 		PrimitiveInfo.GroupIndex = HairGroupPublicData->GetGroupIndex();
 		PrimitiveInfo.PublicDataPtr = HairGroupPublicData;
-		check(PrimitiveInfo.GroupIndex < 32); // Sanity check
 
 		if (HairGroupPublicData->DoesSupportVoxelization())
 		{
@@ -173,7 +172,6 @@ void CreateHairStrandsMacroGroups(
 	static const FVertexFactoryType* CompatibleVF = FVertexFactoryType::GetVFByName(TEXT("FHairStrandsVertexFactory"));
 
 	TArray<FHairStrandsMacroGroupData, SceneRenderingAllocator>& MacroGroups = View.HairStrandsViewData.MacroGroupDatas;
-	FHairStrandsMacroGroupResources& MacroGroupResources = View.HairStrandsViewData.MacroGroupResources;
 
 	int32 MaterialId = 0;
 
@@ -188,9 +186,17 @@ void CreateHairStrandsMacroGroups(
 		const FBoxSphereBounds& PrimitiveBounds = Proxy->GetBounds();
 
 		bool bFound = false;
+		float MinDistance = FLT_MAX;
+		uint32 ClosestMacroGroupId = ~0u;
 		for (FHairStrandsMacroGroupData& MacroGroup : MacroGroups)
 		{
-			const bool bIntersect = FBoxSphereBounds::SpheresIntersect(MacroGroup.Bounds, PrimitiveBounds);
+			const FSphere MacroSphere = MacroGroup.Bounds.GetSphere();
+			const FSphere PrimSphere = PrimitiveBounds.GetSphere();
+
+			const float DistCenters = (MacroSphere.Center - PrimSphere.Center).Size();
+			const float AccumRadius = FMath::Max(0.f, MacroSphere.W + PrimSphere.W);
+			const bool bIntersect = DistCenters <= AccumRadius;
+			
 			if (bIntersect)
 			{
 				MacroGroup.Bounds = Union(MacroGroup.Bounds, PrimitiveBounds);
@@ -198,15 +204,34 @@ void CreateHairStrandsMacroGroups(
 				bFound = true;
 				break;
 			}
+
+			const float MacroToPrimDistance = DistCenters - AccumRadius;
+			if (MacroToPrimDistance < MinDistance)
+			{
+				MinDistance = MacroToPrimDistance;
+				ClosestMacroGroupId = MacroGroup.MacroGroupId;
+			}
 		}
 
 		if (!bFound)
 		{
-			FHairStrandsMacroGroupData MacroGroup;
-			MacroGroup.MacroGroupId = MacroGroupId++;
-			InternalUpdateMacroGroup(MacroGroup, MaterialId, MeshBatchAndRelevance);
-			MacroGroup.Bounds = PrimitiveBounds;
-			MacroGroups.Add(MacroGroup);
+			// If we have reached the max number of macro group (MAX_HAIR_MACROGROUP_COUNT), then merge the current one with the closest one.
+			if (MacroGroups.Num() == FHairStrandsMacroGroupData::MaxMacroGroupCount)
+			{
+				check(ClosestMacroGroupId != ~0u);
+				FHairStrandsMacroGroupData& MacroGroup = MacroGroups[ClosestMacroGroupId];
+				check(MacroGroup.MacroGroupId == ClosestMacroGroupId);
+				MacroGroup.Bounds = Union(MacroGroup.Bounds, PrimitiveBounds);
+				InternalUpdateMacroGroup(MacroGroup, MaterialId, MeshBatchAndRelevance);
+			}
+			else
+			{
+				FHairStrandsMacroGroupData MacroGroup;
+				MacroGroup.MacroGroupId = MacroGroupId++;
+				InternalUpdateMacroGroup(MacroGroup, MaterialId, MeshBatchAndRelevance);
+				MacroGroup.Bounds = PrimitiveBounds;
+				MacroGroups.Add(MacroGroup);
+			}
 		}
 	};
 
@@ -219,8 +244,11 @@ void CreateHairStrandsMacroGroups(
 	{
 		MacroGroup.ScreenRect = ComputeProjectedScreenRect(MacroGroup.Bounds.GetBox(), View);
 	}
+	// Sanity check
+	check(MacroGroups.Num() <= FHairStrandsMacroGroupData::MaxMacroGroupCount);
 
 	// Build hair macro group AABBB
+	FHairStrandsMacroGroupResources& MacroGroupResources = View.HairStrandsViewData.MacroGroupResources;
 	const uint32 MacroGroupCount = MacroGroups.Num();
 	if (MacroGroupCount > 0)
 	{
diff --git a/Engine/Source/Runtime/Renderer/Private/HairStrands/HairStrandsData.h b/Engine/Source/Runtime/Renderer/Private/HairStrands/HairStrandsData.h
--- a/Engine/Source/Runtime/Renderer/Private/HairStrands/HairStrandsData.h
+++ b/Engine/Source/Runtime/Renderer/Private/HairStrands/HairStrandsData.h
@@ -259,6 +259,8 @@ class FHairGroupPublicData;
 /// Hair macro group infos
 struct FHairStrandsMacroGroupData
 {
+	static const uint32 MaxMacroGroupCount = 16u; // Needs to be kept consistent with MAX_HAIR_MACROGROUP_COUNT (HairStrandsVisibilityCommon.ush)
+
 	// List of primitive/mesh batch within an instance group
 	struct PrimitiveInfo
 	{
diff --git a/Engine/Source/Runtime/Renderer/Private/HairStrands/HairStrandsVoxelization.cpp b/Engine/Source/Runtime/Renderer/Private/HairStrands/HairStrandsVoxelization.cpp
--- a/Engine/Source/Runtime/Renderer/Private/HairStrands/HairStrandsVoxelization.cpp
+++ b/Engine/Source/Runtime/Renderer/Private/HairStrands/HairStrandsVoxelization.cpp
@@ -200,7 +200,7 @@ static void AddVirtualVoxelInjectOpaquePass(
 
 	FComputeShaderUtils::AddPass(
 		GraphBuilder, 
-		RDG_EVENT_NAME("HairStrandsInjectOpaqueDepthInVoxel"), 
+		RDG_EVENT_NAME("HairStrands::InjectOpaqueDepthInVoxel"), 
 		ComputeShader, 
 		Parameters, 
 		VoxelResources.IndirectArgsBuffer,
@@ -629,11 +629,11 @@ static void AddAllocateVoxelPagesPass(
 		}
 
 		// Currently support only 32 instance group at max
-		check(Parameters->MacroGroupCount < 32);
+		check(Parameters->MacroGroupCount <= FHairStrandsMacroGroupData::MaxMacroGroupCount);
 		TShaderMapRef<FVoxelAllocatePageIndexCS> ComputeShader(View.ShaderMap);
 		FComputeShaderUtils::AddPass(
 			GraphBuilder,
-			RDG_EVENT_NAME("HairStrandsAllocatePageIndex"),
+			RDG_EVENT_NAME("HairStrands::AllocatePageIndex"),
 			ComputeShader,
 			Parameters,
 			FIntVector(1,1,1));
@@ -721,7 +721,7 @@ static void AddAllocateVoxelPagesPass(
 				TShaderMapRef<FVoxelMarkValidPageIndex_PrepareCS> ComputeShader(View.ShaderMap, PermutationVector);
 				FComputeShaderUtils::AddPass(
 					GraphBuilder,
-					RDG_EVENT_NAME("HairStrandsMarkValidPageIndex_Prepare(%s,%s)", bUseClusterAABB ? TEXT("Cluster") : TEXT("Group"), bForceDirectPageAllocation ? TEXT("Direct") : TEXT("Indirect")),
+					RDG_EVENT_NAME("HairStrands::MarkValidPageIndex_Prepare(%s,%s)", bUseClusterAABB ? TEXT("Cluster") : TEXT("Group"), bForceDirectPageAllocation ? TEXT("Direct") : TEXT("Indirect")),
 					ComputeShader,
 					Parameters,
 					DispatchCount);
@@ -742,7 +742,7 @@ static void AddAllocateVoxelPagesPass(
 
 					const FIntVector DispatchCount(1, 1, 1);
 					TShaderMapRef<FVoxelMarkValidPageIndex_IndirectArgsCS> ComputeShader(View.ShaderMap);
-					FComputeShaderUtils::AddPass(GraphBuilder, RDG_EVENT_NAME("HairStrandsMarkValidPageIndex_IndirectArgs"), ComputeShader, Parameters, DispatchCount);
+					FComputeShaderUtils::AddPass(GraphBuilder, RDG_EVENT_NAME("HairStrands::MarkValidPageIndex_IndirectArgs"), ComputeShader, Parameters, DispatchCount);
 				}
 
 				// Scatter
@@ -758,7 +758,7 @@ static void AddAllocateVoxelPagesPass(
 					Parameters->OutValidPageIndexBuffer = PageIndexBufferUAV;
 
 					TShaderMapRef<FVoxelMarkValidPageIndex_ScatterCS> ComputeShader(View.ShaderMap);
-					FComputeShaderUtils::AddPass(GraphBuilder, RDG_EVENT_NAME("HairStrandsMarkValidPageIndex_Scatter"), ComputeShader, Parameters, ScatterIndirectArgsBuffer, 0);
+					FComputeShaderUtils::AddPass(GraphBuilder, RDG_EVENT_NAME("HairStrands::MarkValidPageIndex_Scatter"), ComputeShader, Parameters, ScatterIndirectArgsBuffer, 0);
 				}
 			}
 		}
@@ -794,7 +794,7 @@ static void AddAllocateVoxelPagesPass(
 				TShaderMapRef<FVoxelMarkValidPageIndexCS> ComputeShader(View.ShaderMap, PermutationVector);
 				FComputeShaderUtils::AddPass(
 					GraphBuilder,
-					RDG_EVENT_NAME("HairStrandsMarkValidPageIndex"),
+					RDG_EVENT_NAME("HairStrands::MarkValidPageIndex"),
 					ComputeShader,
 					Parameters,
 					DispatchCount);
@@ -826,7 +826,7 @@ static void AddAllocateVoxelPagesPass(
 
 			const FIntVector DispatchCount(1, 1, 1);
 			TShaderMapRef<FVoxelAddNodeDescCS> ComputeShader(View.ShaderMap, PermutationVector);
-			FComputeShaderUtils::AddPass(GraphBuilder, RDG_EVENT_NAME("HairStrandsAddNodeDesc"), ComputeShader, Parameters, DispatchCount);
+			FComputeShaderUtils::AddPass(GraphBuilder, RDG_EVENT_NAME("HairStrands::AddNodeDesc"), ComputeShader, Parameters, DispatchCount);
 		}
 
 		// Allocate pages
@@ -853,7 +853,7 @@ static void AddAllocateVoxelPagesPass(
 
 				const uint32 ArgsOffset = sizeof(uint32) * 3 * MacroGroup.MacroGroupId;
 
-				FComputeShaderUtils::AddPass(GraphBuilder, RDG_EVENT_NAME("HairStrandsAllocateVoxelPage"), ComputeShader, Parameters, 
+				FComputeShaderUtils::AddPass(GraphBuilder, RDG_EVENT_NAME("HairStrands::AllocateVoxelPage"), ComputeShader, Parameters, 
 					PageIndexAllocationIndirectBufferArgs,
 					ArgsOffset);
 			}
@@ -861,7 +861,7 @@ static void AddAllocateVoxelPagesPass(
 			{
 				const FIntVector DispatchCount((CPUAllocationDesc.PageIndexCount + GroupSize - 1) / GroupSize, 1, 1);
 				check(DispatchCount.X < 65535);
-				FComputeShaderUtils::AddPass(GraphBuilder, RDG_EVENT_NAME("HairStrandsAllocateVoxelPage"), ComputeShader, Parameters, DispatchCount);
+				FComputeShaderUtils::AddPass(GraphBuilder, RDG_EVENT_NAME("HairStrands::AllocateVoxelPage"), ComputeShader, Parameters, DispatchCount);
 			}
 		}
 
@@ -879,7 +879,7 @@ static void AddAllocateVoxelPagesPass(
 
 			const FIntVector DispatchCount(1, 1, 1);
 			TShaderMapRef<FVoxelAddIndirectBufferCS> ComputeShader(View.ShaderMap);
-			FComputeShaderUtils::AddPass(GraphBuilder, RDG_EVENT_NAME("HairStrandsBuildVoxelIndirectArgs"), ComputeShader, Parameters, DispatchCount);
+			FComputeShaderUtils::AddPass(GraphBuilder, RDG_EVENT_NAME("HairStrands::BuildVoxelIndirectArgs"), ComputeShader, Parameters, DispatchCount);
 		}
 	}
 
@@ -1172,7 +1172,7 @@ static FRDGBufferRef IndirectVoxelPageClear(
 		TShaderMapRef<FVoxelIndPageClearBufferGenCS> ComputeShader(ViewInfo.ShaderMap);
 		FComputeShaderUtils::AddPass(
 			GraphBuilder,
-			RDG_EVENT_NAME("HairStrandsVoxelGenIndBufferClearCS"),
+			RDG_EVENT_NAME("HairStrands::VoxelGenIndBufferClearCS"),
 			ComputeShader,
 			Parameters,
 			FIntVector(1,1,1));
@@ -1188,7 +1188,7 @@ static FRDGBufferRef IndirectVoxelPageClear(
 		TShaderMapRef<FVoxelIndPageClearCS> ComputeShader(ViewInfo.ShaderMap);
 		FComputeShaderUtils::AddPass(
 			GraphBuilder,
-			RDG_EVENT_NAME("HairStrandsVoxelIndPageClearCS"),
+			RDG_EVENT_NAME("HairStrands::VoxelIndPageClearCS"),
 			ComputeShader,
 			Parameters,
 			ClearIndArgsBuffer,
@@ -1318,13 +1318,13 @@ static void AddVirtualVoxelizationComputeRasterPass(
 					PassParameters->HairStrandsVF_CullingRadiusScaleBuffer	= RegisterAsSRV(GraphBuilder, HairGroupPublicData->GetCulledVertexRadiusScaleBuffer());
 					PassParameters->IndirectBufferArgs = CullingIndirectBuffer.Buffer;
 
-					FComputeShaderUtils::AddPass(GraphBuilder, RDG_EVENT_NAME("HairStrandsVoxelComputeRaster(culling=on)"), ComputeShader_CullingOn, PassParameters, CullingIndirectBuffer.Buffer, 0);
+					FComputeShaderUtils::AddPass(GraphBuilder, RDG_EVENT_NAME("HairStrands::VoxelComputeRaster(culling=on)"), ComputeShader_CullingOn, PassParameters, CullingIndirectBuffer.Buffer, 0);
 				}
 				else
 				{
 					const FIntVector DispatchCount = ComputeDispatchCount(PassParameters->HairStrandsVF_VertexCount, GroupSize);
 					PassParameters->DispatchCountX = DispatchCount.X;
-					FComputeShaderUtils::AddPass(GraphBuilder, RDG_EVENT_NAME("HairStrandsVoxelComputeRaster(culling=off)"), ComputeShader_CullingOff, PassParameters, DispatchCount);
+					FComputeShaderUtils::AddPass(GraphBuilder, RDG_EVENT_NAME("HairStrands::VoxelComputeRaster(culling=off)"), ComputeShader_CullingOff, PassParameters, DispatchCount);
 				}
 			}
 		}
@@ -1537,7 +1537,7 @@ static void AddVirtualVoxelGenerateMipPass(
 		Parameters->OutIndirectArgs		= GraphBuilder.CreateUAV(MipIndirectArgs);
 
 		TShaderMapRef<FVirtualVoxelIndirectArgMipCS> ComputeShader(View.ShaderMap);
-		FComputeShaderUtils::AddPass(GraphBuilder, RDG_EVENT_NAME("HairStrandsBuildVoxelMipIndirectArgs"), ComputeShader, Parameters, FIntVector(1, 1, 1));
+		FComputeShaderUtils::AddPass(GraphBuilder, RDG_EVENT_NAME("HairStrands::BuildVoxelMipIndirectArgs"), ComputeShader, Parameters, FIntVector(1, 1, 1));
 	}
 
 	// Generate MIP level (in one go for all allocated pages)
@@ -1558,7 +1558,7 @@ static void AddVirtualVoxelGenerateMipPass(
 		TShaderMapRef<FVirtualVoxelGenerateMipCS> ComputeShader(View.ShaderMap);
 		ClearUnusedGraphResources(ComputeShader, Parameters);
 		GraphBuilder.AddPass(
-			RDG_EVENT_NAME("HairStrandsComputeVoxelMip"),
+			RDG_EVENT_NAME("HairStrands::ComputeVoxelMip"),
 			Parameters,
 			ERDGPassFlags::Compute,
 			[Parameters, ComputeShader](FRHICommandList& RHICmdList)
@@ -1592,7 +1592,7 @@ static void AddVirtualVoxelGenerateMipPass(
 		TShaderMapRef<FVirtualVoxelPatchPageIndexWithMipDataCS> ComputeShader(View.ShaderMap);
 		ClearUnusedGraphResources(ComputeShader, Parameters);
 		GraphBuilder.AddPass(
-			RDG_EVENT_NAME("HairStrandsPatchPageIndexWithMip"),
+			RDG_EVENT_NAME("HairStrands::PatchPageIndexWithMip"),
 			Parameters,
 			ERDGPassFlags::Compute,
 			[Parameters, ComputeShader](FRHICommandList& RHICmdList)
