Pass FMatrix by const reference.

diff --git a/Engine/Source/Runtime/Renderer/Private/PrimitiveSceneInfo.cpp b/Engine/Source/Runtime/Renderer/Private/PrimitiveSceneInfo.cpp
--- a/Engine/Source/Runtime/Renderer/Private/PrimitiveSceneInfo.cpp
+++ b/Engine/Source/Runtime/Renderer/Private/PrimitiveSceneInfo.cpp
@@ -870,7 +870,7 @@ void FPrimitiveSceneInfo::RemoveCachedRayTracingPrimitives()
 	}
 }
 
-void FPrimitiveSceneInfo::UpdateCachedRayTracingInstanceTransforms(FMatrix NewPrimitiveLocalToWorld)
+void FPrimitiveSceneInfo::UpdateCachedRayTracingInstanceTransforms(const FMatrix& NewPrimitiveLocalToWorld)
 {
 	TRACE_CPUPROFILER_EVENT_SCOPE(UpdateCachedRayTracingInstanceTransforms);
 	for (int32 Index = 0; Index < CachedRayTracingInstanceLocalTransforms.Num(); Index++)
diff --git a/Engine/Source/Runtime/Renderer/Public/PrimitiveSceneInfo.h b/Engine/Source/Runtime/Renderer/Public/PrimitiveSceneInfo.h
--- a/Engine/Source/Runtime/Renderer/Public/PrimitiveSceneInfo.h
+++ b/Engine/Source/Runtime/Renderer/Public/PrimitiveSceneInfo.h
@@ -571,7 +571,7 @@ public:
 	/** Mark the runtime virtual textures covered by this primitive as dirty. */
 	void FlushRuntimeVirtualTexture();
 
-#if RHI_RAYTRACING	
+#if RHI_RAYTRACING
 	static void UpdateCachedRaytracingData(FScene* Scene, const TArrayView<FPrimitiveSceneInfo*>& SceneInfos);
 	RENDERER_API FRHIRayTracingGeometry* GetStaticRayTracingGeometryInstance(int LodLevel) const;
 #endif
@@ -659,7 +659,7 @@ private:
 	void RemoveCachedRayTracingPrimitives();
 
 	/** Updates cached transforms in CachedRayTracingInstance */
-	void UpdateCachedRayTracingInstanceTransforms(FMatrix NewPrimitiveLocalToWorld);
+	void UpdateCachedRayTracingInstanceTransforms(const FMatrix& NewPrimitiveLocalToWorld);
 
 	/** Updates cached ray tracing instances. Utility closesly mirrors CacheRayTracingPrimitives(..) */
 	static void UpdateCachedRayTracingInstances(FScene* Scene, const TArrayView<FPrimitiveSceneInfo*>& SceneInfos);
