If a mesh is not going to fade out, the visibillity would not get set properly, causing a disconnect from when the object would disappear and the shadow would disappear.  This change makes sure that visibility, disctance culled, and fading bits are always set correctly.  Also, this change make sure that when the distance is set to FLT_MAX, the calculations don't get messed up by squaring the distance, and adding the fade distance.

diff --git a/Engine/Source/Runtime/Renderer/Private/SceneVisibility.cpp b/Engine/Source/Runtime/Renderer/Private/SceneVisibility.cpp
--- a/Engine/Source/Runtime/Renderer/Private/SceneVisibility.cpp
+++ b/Engine/Source/Runtime/Renderer/Private/SceneVisibility.cpp
@@ -483,6 +483,33 @@ static FAutoConsoleVariableRef CVarFrustumCullNumWordsPerTask(
 	ECVF_Default
 	);
 
+// Returns true if the frustum and bounds intersect
+template<bool UseCustomCulling, bool bAlsoUseSphereTest, bool bUseFastIntersect>
+FORCEINLINE static bool IsPrimitiveVisible(FViewInfo& View, const FPlane* PermutedPlanePtr, const FPrimitiveBounds& RESTRICT Bounds, int32 VisibilityId)
+{
+	// The custom culling and sphere culling are additional tests, meaning that if they pass, the
+	// remaining culling tests will still be performed.  If any of the tests fail, then the primitive
+	// is culled, and the remaining tests do not need be performed
+
+	if (UseCustomCulling && !View.CustomVisibilityQuery->IsVisible(VisibilityId, FBoxSphereBounds(Bounds.BoxSphereBounds.Origin, Bounds.BoxSphereBounds.BoxExtent, Bounds.BoxSphereBounds.SphereRadius)))
+	{
+		return false;
+	}
+
+	if (bAlsoUseSphereTest && !View.ViewFrustum.IntersectSphere(Bounds.BoxSphereBounds.Origin, Bounds.BoxSphereBounds.SphereRadius))
+	{
+		return false;
+	}
+
+	if (bUseFastIntersect)
+	{
+		return IntersectBox8Plane(Bounds.BoxSphereBounds.Origin, Bounds.BoxSphereBounds.BoxExtent, PermutedPlanePtr);
+	}
+	else
+	{
+		return View.ViewFrustum.IntersectBox(Bounds.BoxSphereBounds.Origin, Bounds.BoxSphereBounds.BoxExtent);
+	}
+}
 
 template<bool UseCustomCulling, bool bAlsoUseSphereTest, bool bUseFastIntersect>
 static int32 FrustumCull(const FScene* RESTRICT Scene, FViewInfo& View)
@@ -512,10 +539,11 @@ static int32 FrustumCull(const FScene* RESTRICT Scene, FViewInfo& View)
 			SCOPED_NAMED_EVENT(SceneVisibility_FrustumCull, FColor::Red);
 
 			QUICK_SCOPE_CYCLE_COUNTER(STAT_FrustumCull_Loop);
+			bool bDisableLODFade = GDisableLODFade || View.bDisableDistanceBasedFadeTransitions;
 			const FPlane* PermutedPlanePtr = View.ViewFrustum.PermutedPlanes.GetData();
 			const int32 BitArrayNumInner = View.PrimitiveVisibilityMap.Num();
 			FVector ViewOriginForDistanceCulling = View.ViewMatrices.GetViewOrigin();
-			float FadeRadius = GDisableLODFade ? 0.0f : GDistanceFadeMaxTravel;
+			float FadeRadius = bDisableLODFade ? 0.0f : GDistanceFadeMaxTravel;
 			uint8 CustomVisibilityFlags = EOcclusionFlags::CanBeOccluded | EOcclusionFlags::HasPrecomputedVisibility;
 
 			uint32 NumPrimitivesCulledForTask = 0;
@@ -543,13 +571,16 @@ static int32 FrustumCull(const FScene* RESTRICT Scene, FViewInfo& View)
 					}
 
 					// Preserve infinite draw distance
-					float MaxDrawDistance = Bounds.MaxCullDistance < FLT_MAX ? Bounds.MaxCullDistance * MaxDrawDistanceScale : FLT_MAX; 
+					bool bHasMaxDrawDistance = Bounds.MaxCullDistance < FLT_MAX;
+					float MaxDrawDistance =  Bounds.MaxCullDistance * MaxDrawDistanceScale;
 					float MinDrawDistanceSq = Bounds.MinDrawDistanceSq;
 
 					// If cull distance is disabled, always show the primitive (except foliage)
-					if (View.Family->EngineShowFlags.DistanceCulledPrimitives && !Scene->Primitives[Index]->Proxy->IsDetailMesh())
+					if (   View.Family->EngineShowFlags.DistanceCulledPrimitives 
+						&& !Scene->Primitives[Index]->Proxy->IsDetailMesh()) // Proxy call is intentionally behind the DistancedCulledPrimitives check to prevent an expensive memory read
 					{
 						MaxDrawDistance = FLT_MAX;
+						bHasMaxDrawDistance = false;
 					}
 
 					// Fading HLODs and their children must be visible, objects hidden by HLODs can be culled
@@ -557,60 +588,55 @@ static int32 FrustumCull(const FScene* RESTRICT Scene, FViewInfo& View)
 					{
 						if (HLODState->IsNodeForcedVisible(Index))
 						{
+							bHasMaxDrawDistance = false;
 							MaxDrawDistance = FLT_MAX;
 							MinDrawDistanceSq = 0.f;
 						}
 						else if (HLODState->IsNodeForcedHidden(Index))
 						{
+							bHasMaxDrawDistance = true;
 							MaxDrawDistance = 0.f;
 						}
 					}
 
-					const bool bDistanceCulled = DistanceSquared > FMath::Square(MaxDrawDistance + FadeRadius) || (DistanceSquared < MinDrawDistanceSq);
-
-					// Store distane culled primitives so it can correctly culled when collecting RT primitives
-					if (bDistanceCulled)
+					// Always test the fade in distance.  If a primitive was set to always draw, it may need to be faded in.
+					if (bHasMaxDrawDistance)
 					{
-						DistanceCulledBits |= Mask;
+						float MaxFadeDistanceSquared = FMath::Square(MaxDrawDistance + FadeRadius);
+						float MinFadeDistanceSquared = FMath::Square(MaxDrawDistance - FadeRadius);
+						if (  (DistanceSquared < MaxFadeDistanceSquared && DistanceSquared > MinFadeDistanceSquared)
+							&& Scene->Primitives[Index]->Proxy->IsUsingDistanceCullFade())  // Proxy call is intentionally behind the fade check to prevent an expensive memory read
+						{
+							FadingBits |= Mask;
+						}
 					}
 
-					// Handle primitives that are always visible.
-					if (Scene->PrimitivesAlwaysVisible[Index])
+					bool bIsVisible = true;
+					
+					// Handle primitives that are not always visible.
+					if (!Scene->PrimitivesAlwaysVisible[Index])
 					{
-						VisBits |= Mask;
-						if (DistanceSquared > FMath::Square(MaxDrawDistance - FadeRadius) && Scene->Primitives[Index]->Proxy->IsUsingDistanceCullFade())
+						// Check for distance culling first
+						const bool bFarDistanceCulled = bHasMaxDrawDistance && (DistanceSquared > FMath::Square(MaxDrawDistance));
+						const bool bNearDistanceCulled = (DistanceSquared < MinDrawDistanceSq);
+						bool bDistanceCulled = bNearDistanceCulled || bFarDistanceCulled;
+
+						if (bDistanceCulled)
 						{
-							FadingBits |= Mask;
+							DistanceCulledBits |= Mask;
 						}
+
+						bIsVisible = !bDistanceCulled && IsPrimitiveVisible<UseCustomCulling, bAlsoUseSphereTest, bUseFastIntersect>(View, PermutedPlanePtr, Bounds, VisibilityId);
 					}
-					else if (bDistanceCulled ||
-						(UseCustomCulling && !View.CustomVisibilityQuery->IsVisible(VisibilityId, FBoxSphereBounds(Bounds.BoxSphereBounds.Origin, Bounds.BoxSphereBounds.BoxExtent, Bounds.BoxSphereBounds.SphereRadius))) ||
-						(bAlsoUseSphereTest && View.ViewFrustum.IntersectSphere(Bounds.BoxSphereBounds.Origin, Bounds.BoxSphereBounds.SphereRadius) == false) ||
-						(bUseFastIntersect ? IntersectBox8Plane(Bounds.BoxSphereBounds.Origin, Bounds.BoxSphereBounds.BoxExtent, PermutedPlanePtr) : View.ViewFrustum.IntersectBox(Bounds.BoxSphereBounds.Origin, Bounds.BoxSphereBounds.BoxExtent)) == false)
+
+					if (bIsVisible)
 					{
-						STAT(++NumPrimitivesCulledForTask);
+						// The primitive is visible!
+						VisBits |= Mask;
 					}
 					else
 					{
-						if (DistanceSquared > FMath::Square(MaxDrawDistance))
-						{
-							if (Scene->Primitives[Index]->Proxy->IsUsingDistanceCullFade())
-							{
-								FadingBits |= Mask;
-							}
-						}
-						else
-						{
-							// The primitive is visible!
-							VisBits |= Mask;
-							if (DistanceSquared > FMath::Square(MaxDrawDistance - FadeRadius))
-							{
-								if (Scene->Primitives[Index]->Proxy->IsUsingDistanceCullFade())
-								{
-									FadingBits |= Mask;
-								}
-							}
-						}
+						STAT(++NumPrimitivesCulledForTask);
 					}
 				}
 				if (FadingBits)
@@ -679,6 +705,7 @@ static void UpdatePrimitiveFading(const FScene* Scene, FViewInfo& View)
 				{
 					// If the primitive is fading out make sure it remains visible.
 					View.PrimitiveVisibilityMap.AccessCorrespondingBit(BitIt) = true;
+					View.DistanceCullingPrimitiveMap.AccessCorrespondingBit(BitIt) = false;
 				}
 				View.PrimitiveFadeUniformBuffers[BitIt.GetIndex()] = UniformBuffer;
 				View.PrimitiveFadeUniformBufferMap[BitIt.GetIndex()] = UniformBuffer != nullptr;
