Fixed: Vulkan validation warning about a missing VK_BUFFER_USAGE_VERTEX_BUFFER_BIT for an instance data buffers

diff --git a/Engine/Source/Runtime/Renderer/Private/InstanceCulling/InstanceCullingContext.cpp b/Engine/Source/Runtime/Renderer/Private/InstanceCulling/InstanceCullingContext.cpp
--- a/Engine/Source/Runtime/Renderer/Private/InstanceCulling/InstanceCullingContext.cpp
+++ b/Engine/Source/Runtime/Renderer/Private/InstanceCulling/InstanceCullingContext.cpp
@@ -327,6 +327,17 @@ public:
 	}
 };
 
+static FRDGBufferDesc CreateInstanceIdsBufferDesc(ERHIFeatureLevel::Type FeatureLevel, uint32 BufferSize)
+{
+	FRDGBufferDesc BufferDesc = FRDGBufferDesc::CreateStructuredDesc(FInstanceCullingContext::GetInstanceIdBufferStride(FeatureLevel), BufferSize);
+	if (FeatureLevel == ERHIFeatureLevel::ES3_1)
+	{
+		// Mobile writes to this buffer from compute and then uses as a vertex input
+		// D3D does not allow a storage buffer with vertex buffer usage
+		BufferDesc.Usage|= (IsD3DPlatform(GMaxRHIShaderPlatform) ? BUF_None : BUF_VertexBuffer);
+	}
+	return BufferDesc;
+}
 
 void FInstanceCullingContext::BuildRenderingCommands(
 	FRDGBuilder& GraphBuilder,
@@ -401,8 +412,7 @@ void FInstanceCullingContext::BuildRenderingCommands(
 
 	FRDGBufferRef ViewIdsBuffer = CreateStructuredBuffer(GraphBuilder, TEXT("InstanceCulling.ViewIds"), ViewIds);
 	
-	const uint32 InstanceIdBufferStride = GetInstanceIdBufferStride(FeatureLevel);
-	FRDGBufferRef InstanceIdsBuffer = GraphBuilder.CreateBuffer(FRDGBufferDesc::CreateStructuredDesc(InstanceIdBufferStride, InstanceIdBufferSize), TEXT("InstanceCulling.InstanceIdsBuffer"));
+	FRDGBufferRef InstanceIdsBuffer = GraphBuilder.CreateBuffer(CreateInstanceIdsBufferDesc(FeatureLevel, InstanceIdBufferSize), TEXT("InstanceCulling.InstanceIdsBuffer"));
 	FRDGBufferUAVRef InstanceIdsBufferUAV = GraphBuilder.CreateUAV(InstanceIdsBuffer, ERDGUnorderedAccessViewFlags::SkipBarrier);
 
 	FBuildInstanceIdBufferAndCommandsFromPrimitiveIdsCs::FParameters PassParametersTmp;
@@ -707,8 +717,7 @@ FInstanceCullingDeferredContext *FInstanceCullingContext::CreateDeferredContext(
 	FRDGBufferRef InstanceIdOffsetBuffer = GraphBuilder.CreateBuffer(FRDGBufferDesc::CreateBufferDesc(sizeof(uint32), 1), TEXT("InstanceCulling.InstanceIdOffsetBuffer"), INST_CULL_CALLBACK(DeferredContext->InstanceIdOffsets.Num()));
 	GraphBuilder.QueueBufferUpload(InstanceIdOffsetBuffer, INST_CULL_CALLBACK(DeferredContext->InstanceIdOffsets.GetData()), INST_CULL_CALLBACK(DeferredContext->InstanceIdOffsets.GetTypeSize() * DeferredContext->InstanceIdOffsets.Num()));
 
-	const uint32 InstanceIdBufferStride = GetInstanceIdBufferStride(FeatureLevel);
-	FRDGBufferRef InstanceIdsBuffer = GraphBuilder.CreateBuffer(FRDGBufferDesc::CreateStructuredDesc(InstanceIdBufferStride, 1), TEXT("InstanceCulling.InstanceIdsBuffer"), INST_CULL_CALLBACK(DeferredContext->InstanceIdBufferSize));
+	FRDGBufferRef InstanceIdsBuffer = GraphBuilder.CreateBuffer(CreateInstanceIdsBufferDesc(FeatureLevel, 1), TEXT("InstanceCulling.InstanceIdsBuffer"), INST_CULL_CALLBACK(DeferredContext->InstanceIdBufferSize));
 	FRDGBufferUAVRef InstanceIdsBufferUAV = GraphBuilder.CreateUAV(InstanceIdsBuffer, ERDGUnorderedAccessViewFlags::SkipBarrier);
 	if (FeatureLevel == ERHIFeatureLevel::ES3_1)
 	{
