Fix invalidating bounds debug rendering (request the correct number of debug elements for all instances)
- also remove special case for added as this was incorrect.

diff --git a/Engine/Source/Runtime/Renderer/Private/VirtualShadowMaps/VirtualShadowMapCacheManager.cpp b/Engine/Source/Runtime/Renderer/Private/VirtualShadowMaps/VirtualShadowMapCacheManager.cpp
--- a/Engine/Source/Runtime/Renderer/Private/VirtualShadowMaps/VirtualShadowMapCacheManager.cpp
+++ b/Engine/Source/Runtime/Renderer/Private/VirtualShadowMaps/VirtualShadowMapCacheManager.cpp
@@ -419,12 +419,14 @@ void FVirtualShadowMapArrayCacheManager::ProcessRemovedPrimives(FRDGBuilder& Gra
 		TArray<FInstanceSceneDataRange, SceneRenderingAllocator> InstanceRangesLarge;
 		TArray<FInstanceSceneDataRange, SceneRenderingAllocator> InstanceRangesSmall;
 
+		int32 TotalInstanceCount = 0;
 		for (const FPrimitiveSceneInfo* PrimitiveSceneInfo : RemovedPrimitiveSceneInfos)
 		{
 			if (PrimitiveSceneInfo->GetInstanceSceneDataOffset() != INDEX_NONE)
 			{
 				const int32 NumInstanceSceneDataEntries = PrimitiveSceneInfo->GetNumInstanceSceneDataEntries();
-				if (NumInstanceSceneDataEntries >= 8u)
+				TotalInstanceCount += NumInstanceSceneDataEntries;
+				if (NumInstanceSceneDataEntries >= 8)
 				{
 					InstanceRangesLarge.Add(FInstanceSceneDataRange{ PrimitiveSceneInfo->GetInstanceSceneDataOffset(), NumInstanceSceneDataEntries });
 				}
@@ -434,7 +436,7 @@ void FVirtualShadowMapArrayCacheManager::ProcessRemovedPrimives(FRDGBuilder& Gra
 				}
 			}
 		}
-		ProcessInstanceRangeInvalidation(GraphBuilder, InstanceRangesLarge, InstanceRangesSmall, GPUScene);
+		ProcessInstanceRangeInvalidation(GraphBuilder, InstanceRangesLarge, InstanceRangesSmall, TotalInstanceCount, GPUScene);
 	}
 }
 
@@ -450,6 +452,7 @@ void FVirtualShadowMapArrayCacheManager::ProcessPrimitivesToUpdate(FRDGBuilder&
 		TArray<FInstanceSceneDataRange, SceneRenderingAllocator> InstanceRangesLarge;
 		TArray<FInstanceSceneDataRange, SceneRenderingAllocator> InstanceRangesSmall;
 
+		int32 TotalInstanceCount = 0;
 		for (const int32 PrimitiveId : GPUScene.PrimitivesToUpdate)
 		{
 			// There may possibly be IDs that are out of range if they were marked for update and then removed.
@@ -459,19 +462,15 @@ void FVirtualShadowMapArrayCacheManager::ProcessPrimitivesToUpdate(FRDGBuilder&
 
 				// SKIP if marked for Add, because this means it has no previous location to invalidate.
 				// SKIP if transform has not changed, as this means no invalidation needs to take place.
-#if UE_BUILD_SHIPPING
 				const bool bDoInvalidation = !EnumHasAnyFlags(PrimitiveDirtyState, EPrimitiveDirtyState::Added) && EnumHasAnyFlags(PrimitiveDirtyState, EPrimitiveDirtyState::ChangedTransform);
-#else
-				// Don't skip added if invadiation debug is enabled, since we want to render those bounds also 
-				const bool bDoInvalidation = (CVarDrawInvalidatingBounds.GetValueOnRenderThread() != 0 || !EnumHasAnyFlags(PrimitiveDirtyState, EPrimitiveDirtyState::Added)) && EnumHasAnyFlags(PrimitiveDirtyState, EPrimitiveDirtyState::ChangedTransform);
-#endif
 				if (bDoInvalidation)
 				{
 					const FPrimitiveSceneInfo* PrimitiveSceneInfo = Scene.Primitives[PrimitiveId];
 					if (PrimitiveSceneInfo->GetInstanceSceneDataOffset() != INDEX_NONE)
 					{
 						int32 NumInstanceSceneDataEntries = PrimitiveSceneInfo->GetNumInstanceSceneDataEntries();
-						if (NumInstanceSceneDataEntries >= 8u)
+						TotalInstanceCount += NumInstanceSceneDataEntries;
+						if (NumInstanceSceneDataEntries >= 8)
 						{
 							InstanceRangesLarge.Add(FInstanceSceneDataRange{ PrimitiveSceneInfo->GetInstanceSceneDataOffset(), NumInstanceSceneDataEntries });
 						}
@@ -483,7 +482,7 @@ void FVirtualShadowMapArrayCacheManager::ProcessPrimitivesToUpdate(FRDGBuilder&
 				}
 			}
 		}
-		ProcessInstanceRangeInvalidation(GraphBuilder, InstanceRangesLarge, InstanceRangesSmall, GPUScene);
+		ProcessInstanceRangeInvalidation(GraphBuilder, InstanceRangesLarge, InstanceRangesSmall, TotalInstanceCount, GPUScene);
 	}
 }
 
@@ -547,7 +546,7 @@ TRDGUniformBufferRef<FVirtualShadowMapUniformParameters> FVirtualShadowMapArrayC
 	return GraphBuilder.CreateUniformBuffer(VersionedParameters);
 }
 
-void FVirtualShadowMapArrayCacheManager::ProcessInstanceRangeInvalidation(FRDGBuilder& GraphBuilder, const TArray<FInstanceSceneDataRange, SceneRenderingAllocator>& InstanceRangesLarge, const TArray<FInstanceSceneDataRange, SceneRenderingAllocator>& InstanceRangesSmall, const FGPUScene& GPUScene)
+void FVirtualShadowMapArrayCacheManager::ProcessInstanceRangeInvalidation(FRDGBuilder& GraphBuilder, const TArray<FInstanceSceneDataRange, SceneRenderingAllocator>& InstanceRangesLarge, const TArray<FInstanceSceneDataRange, SceneRenderingAllocator>& InstanceRangesSmall, int32 TotalInstanceCount, const FGPUScene& GPUScene)
 {
 	if (InstanceRangesSmall.IsEmpty() && InstanceRangesLarge.IsEmpty())
 	{
@@ -566,7 +565,7 @@ void FVirtualShadowMapArrayCacheManager::ProcessInstanceRangeInvalidation(FRDGBu
 	if (bDrawBounds)
 	{
 		ShaderDrawDebug::SetEnabled(true);
-		ShaderDrawDebug::RequestSpaceForElements((InstanceRangesSmall.Num() + InstanceRangesLarge.Num()) * 12);
+		ShaderDrawDebug::RequestSpaceForElements(TotalInstanceCount * 12);
 	}
 
 	// Note: this disables the whole debug permutation since the parameters must be bound.
diff --git a/Engine/Source/Runtime/Renderer/Private/VirtualShadowMaps/VirtualShadowMapCacheManager.h b/Engine/Source/Runtime/Renderer/Private/VirtualShadowMaps/VirtualShadowMapCacheManager.h
--- a/Engine/Source/Runtime/Renderer/Private/VirtualShadowMaps/VirtualShadowMapCacheManager.h
+++ b/Engine/Source/Runtime/Renderer/Private/VirtualShadowMaps/VirtualShadowMapCacheManager.h
@@ -133,10 +133,9 @@ private:
 
 	void ProcessInstanceRangeInvalidation(
 		FRDGBuilder& GraphBuilder,
-		const TArray<FInstanceSceneDataRange,
-		SceneRenderingAllocator>& InstanceRangesLarge,
-		const TArray<FInstanceSceneDataRange,
-		SceneRenderingAllocator>& InstanceRangesSmall,
+		const TArray<FInstanceSceneDataRange, SceneRenderingAllocator>& InstanceRangesLarge,
+		const TArray<FInstanceSceneDataRange, SceneRenderingAllocator>& InstanceRangesSmall,
+		int32 TotalInstanceCount,
 		const FGPUScene& GPUScene);
 
 	void ExtractStats(FRDGBuilder& GraphBuilder, FVirtualShadowMapArray &VirtualShadowMapArray);
