Resurects moire heuristic

diff --git a/Engine/Shaders/Private/TemporalSuperResolution/TSRCommon.ush b/Engine/Shaders/Private/TemporalSuperResolution/TSRCommon.ush
--- a/Engine/Shaders/Private/TemporalSuperResolution/TSRCommon.ush
+++ b/Engine/Shaders/Private/TemporalSuperResolution/TSRCommon.ush
@@ -280,7 +280,7 @@ FScreenTransform ScreenPosToPrevHistoryBufferUV;
 float HistoryPreExposureCorrection;
 
 uint bCameraCut;
-uint bEnableInterferenceHeuristic;
+float MoireInvExposure;
 
 #if DEBUG_OUTPUT
 	RWTexture2D<float4> DebugOutput;
@@ -893,6 +893,8 @@ taa_half UncompressParallaxFactor(taa_subpixel_payload CompressedParallaxFactor)
 
 taa_half LumaToInterferenceSeed(taa_half Luma)
 {
+	Luma *= taa_half(MoireInvExposure);
+
 	return Luma * HdrWeightY(Luma);
 }
 
diff --git a/Engine/Shaders/Private/TemporalSuperResolution/TSRCompareHistory.usf b/Engine/Shaders/Private/TemporalSuperResolution/TSRCompareHistory.usf
--- a/Engine/Shaders/Private/TemporalSuperResolution/TSRCompareHistory.usf
+++ b/Engine/Shaders/Private/TemporalSuperResolution/TSRCompareHistory.usf
@@ -15,6 +15,7 @@
 Texture2D<taa_half> ParallaxRejectionMaskTexture;
 Texture2D<taa_half3> FilteredInputTexture;
 Texture2D<taa_half3> FilteredPredictionSceneColorTexture;
+Texture2D<taa_half4> FilteredInterferenceSeedTexture;
 Texture2D<taa_half> InterferenceWeightTexture;
 
 RWTexture2D<taa_half> HistoryRejectionOutput;
@@ -52,6 +53,15 @@ void MainCS(
 	taa_half3 InputCenterColor;
 	taa_half ParallaxRejectionMask;
 	taa_half InterferenceWeight;
+	
+	taa_half ClosestInterferenceSeedMin;
+	taa_half ClosestInterferenceSeedMax;
+	taa_half IntegratedInterferenceSeed;
+	taa_half IntegratedInterferenceSeedMin;
+	taa_half IntegratedInterferenceSeedMax;
+	taa_half MinInterferenceSeed;
+	taa_half MaxInterferenceSeed;
+
 	{
 		InputCenterColor = TransformColorForClampingBox(FilteredInputTexture[InputPixelPos]);
 		ParallaxRejectionMask = ParallaxRejectionMaskTexture[InputPixelPos];
@@ -60,8 +70,20 @@ void MainCS(
 		InputMaxColor = InputCenterColor;
 		
 		InterferenceWeight = InterferenceWeightTexture[InputPixelPos];
+		
+		taa_half4 InterferenceSeed = FilteredInterferenceSeedTexture[InputPixelPos];
+		
+		ClosestInterferenceSeedMin = InterferenceSeed.x;
+		ClosestInterferenceSeedMax = InterferenceSeed.x; 
+		IntegratedInterferenceSeed = InterferenceSeed.y;
+		IntegratedInterferenceSeedMin = InterferenceSeed.y;
+		IntegratedInterferenceSeedMax = InterferenceSeed.y;
+		MinInterferenceSeed = InterferenceSeed.z;
+		MaxInterferenceSeed = InterferenceSeed.w;
 	}
-
+	
+	taa_half NeigborMaxInterferenceSeedGradient = 0.0;
+	
 	UNROLL
 	for (uint i = 0; i < 9; i++)
 	{
@@ -76,10 +98,67 @@ void MainCS(
 			InputPixelPosMin, InputPixelPosMax);
 
 		taa_half3 Color = TransformColorForClampingBox(FilteredInputTexture[SampleInputPixelPos]);
+		taa_half4 InterferenceSeed = FilteredInterferenceSeedTexture[SampleInputPixelPos];
 
 		InputMinColor = min(InputMinColor, Color);
 		InputMaxColor = max(InputMaxColor, Color);
+		
+		taa_half InterferenceSeedGradient = InterferenceSeed.w - InterferenceSeed.z;
+		
+		ClosestInterferenceSeedMin = min(ClosestInterferenceSeedMin, InterferenceSeed.x);
+		ClosestInterferenceSeedMax = max(ClosestInterferenceSeedMax, InterferenceSeed.x); 
+		
+		// TODO(TSR): this is InterferenceSeed.w - InterferenceSeed.z
+		IntegratedInterferenceSeedMin = min(IntegratedInterferenceSeedMin, InterferenceSeed.y);
+		IntegratedInterferenceSeedMax = max(IntegratedInterferenceSeedMax, InterferenceSeed.y);
+		
+		NeigborMaxInterferenceSeedGradient = max(NeigborMaxInterferenceSeedGradient, InterferenceSeedGradient); 
+	}
+	
+	bool bIsInterference = false;
+	#if DIM_MOIRE_HEURISTIC
+	{
+		// Error due to the low bitdepth of the inteference seeds in the history.
+		const taa_half InterferenceSeedEncodingError = rcp(taa_half(SUPERRES_INTERFERENCE_BITMASK));
+
+		taa_half InputSeedMin = LumaToInterferenceSeed(InputMinColor.x);
+		taa_half InputSeedMax = LumaToInterferenceSeed(InputMaxColor.x);
+
+		// Measure the low frequency gradien in the seed. Needs to happen both on current frame input to avoid ghosting on
+		// current frame changes with high contrast geometric edges, and previous frames to avoid ghosting on area that had 
+		// high contrast geometric edges but no longer had in current frame due to shading change.
+		taa_half InterferenceSeedGradient = max(InputSeedMax - InputSeedMin, IntegratedInterferenceSeedMax - IntegratedInterferenceSeedMin);
+		//taa_half InterferenceSeedGradient = InputSeedMax - InputSeedMin;
+
+		// Detect if there is any inteference based on any subpixel detail range from the lowest LDR to highest LDR.
+		bool bRejectInterferenceSeedWithGradient = (
+			1 && (MinInterferenceSeed < ClosestInterferenceSeedMin - InterferenceSeedEncodingError - InterferenceSeedGradient) ||
+			1 && (MaxInterferenceSeed > ClosestInterferenceSeedMax + InterferenceSeedEncodingError + InterferenceSeedGradient));
+			
+		bool bRejectInterferenceSeed = (
+			(MinInterferenceSeed < ClosestInterferenceSeedMin - InterferenceSeedEncodingError) ||
+			(MaxInterferenceSeed > ClosestInterferenceSeedMax + InterferenceSeedEncodingError));
+
+		// Detect interference solely based on the inteference seed, but need to take into account low frequency
+		// gradien that there may be in the neighborhood to avoid ghosting on high contrast geometric edges.
+		//bIsInterference = bRejectInterferenceSeed;
+		bIsInterference = bRejectInterferenceSeedWithGradient;
+
+		Debug.r = bIsInterference ? 1.0 : 0.0;
+		Debug.g = InputSeedMax - InputSeedMin;
+		Debug.b = ClosestInterferenceSeedMax;
+		Debug.a = MaxInterferenceSeed;
+		
+		Debug.b = ClosestInterferenceSeedMax;
+		Debug.a = MinInterferenceSeed;
+
+		//FLATTEN
+		//if (bOffScreen)
+		//{
+		//	bIsInterference = false;
+		//}
 	}
+	#endif
 
 	taa_half3 PrevColor = TransformColorForClampingBox(FilteredPredictionSceneColorTexture[InputPixelPos]);
 
@@ -93,15 +172,22 @@ void MainCS(
 			InputCenterColor, InputMinColor, InputMaxColor,
 			MeasureBackbufferLDRQuantizationErrorFromHDRLuma(max(PrevColor[0], ClampedPrevColor[0])));
 		
-		Debug.rgb = abs(YCoCgToRGB(PrevColor - ClampedPrevColor)) / max3(YCoCgToRGB(InputMaxColor), YCoCgToRGB(InputCenterColor), 0.01);
+		//Debug.rgb = abs(YCoCgToRGB(PrevColor - ClampedPrevColor)) / max3(YCoCgToRGB(InputMaxColor), YCoCgToRGB(InputCenterColor), 0.01);
 	}
 	#endif
 	
-	FLATTEN
-	if (bEnableInterferenceHeuristic)
+	#if DIM_MOIRE_HEURISTIC
 	{
+		if (bIsInterference)
+		{
+			InterferenceWeight = 0.0;
+		}
+		 
 		Rejection = saturate(Rejection + 1.0 - InterferenceWeight);
+		
+		Debug.r = InterferenceWeight;
 	}
+	#endif
 
 	if (ParallaxRejectionMask < taa_half(PARALLAX_REJECTION_MASK_THRESHOLD))
 	{
diff --git a/Engine/Shaders/Private/TemporalSuperResolution/TSRDecimateHistory.usf b/Engine/Shaders/Private/TemporalSuperResolution/TSRDecimateHistory.usf
--- a/Engine/Shaders/Private/TemporalSuperResolution/TSRDecimateHistory.usf
+++ b/Engine/Shaders/Private/TemporalSuperResolution/TSRDecimateHistory.usf
@@ -11,8 +11,6 @@
 
 #define CONFIG_OUTPUT_HALF_RES (DIM_OUTPUT_HALF_RES)
 
-#define CONFIG_OUTPUT_INTERFERENCE_SEEDS 0
-
 
 //------------------------------------------------------- PARAMETERS
 
@@ -32,6 +30,9 @@ Texture2D<taa_half2> PrevHistory_Metadata;
 Texture2D<taa_half4> PrevHistory_Translucency;
 Texture2D<taa_subpixel_details> PrevHistory_SubpixelDetails;
 
+Texture2D<half> PrevHistory_LowResTextures_0;
+Texture2D<uint> PrevHistory_SuperResTextures_0;
+
 RWTexture2D<taa_half3> HalfResSceneColorOutput;
 RWTexture2D<taa_half3> HalfResPredictionSceneColorOutput;
 RWTexture2D<taa_half> HalfResParallaxRejectionMaskOutput;
@@ -169,7 +170,7 @@ void MainCS(
 		
 		float2 PrevInputBufferUV = (InputInfo_ScreenPosToViewportScale * PrevOutputScreenPos + InputInfo_ScreenPosToViewportBias) * InputInfo_ExtentInverse;
 
-		#if CONFIG_OUTPUT_INTERFERENCE_SEEDS
+		#if DIM_MOIRE_HEURISTIC
 		{
 			const uint KernelSize = 4;
 
@@ -177,6 +178,9 @@ void MainCS(
 
 			uint ClosestPayload = ~0;
 			uint ClosestAge = 0;
+			taa_half ClosestDistanceToSeed = 0.2;
+
+			Debug.r = InputInterferenceSeed;
 
 			// Sample the super res history and find an aproximative age of the input pixel.
 			UNROLL
@@ -201,11 +205,14 @@ void MainCS(
 
 					taa_half DistanceToSeed = abs(InputInterferenceSeed - UnpackedSubpixel.InterferenceSeed);
 
+					// UnpackedSubpixel.Age >= ClosestAge && 
+
 					FLATTEN
-					if (DistanceToSeed < 0.1 && UnpackedSubpixel.Age > ClosestAge)
+					if (DistanceToSeed < ClosestDistanceToSeed)
 					{
 						ClosestPayload = SubpixelPayload;
 						ClosestAge = UnpackedSubpixel.Age;
+						ClosestDistanceToSeed = DistanceToSeed;
 					}
 				}
 			}
@@ -239,20 +246,23 @@ void MainCS(
 			// average of min and max that turns out interference free.
 			IntegratedInterferenceSeed = (MinInterferenceSeed + MaxInterferenceSeed) * 0.5;
 		}
-		#endif
+		#endif // DIM_MOIRE_HEURISTIC
 
-		//taa_half PrevUninterferenceSeed = PrevHistory_LowResTextures_0.SampleLevel(GlobalBilinearClampedSampler, PrevInputBufferUV, 0);
-		//taa_half PrevUninterferenceLuma = PrevUninterferenceSeed * HdrWeightInvY(PrevUninterferenceSeed);
+		taa_half PrevUninterferenceSeed = PrevHistory_LowResTextures_0.SampleLevel(GlobalBilinearClampedSampler, PrevInputBufferUV, 0);
+		taa_half PrevUninterferenceLuma = PrevUninterferenceSeed * HdrWeightInvY(PrevUninterferenceSeed);
 
-		#if !CONFIG_OUTPUT_INTERFERENCE_SEEDS
+		PrevUninterferenceLuma /= MoireInvExposure;
+
+		#if !DIM_MOIRE_HEURISTIC || 2
 		{
 			DecimatedColor = PrevHistory_LowFrequency.SampleLevel(GlobalBilinearClampedSampler, KernelHistoryBufferUV, 0);
 			
 			// Cancels out the flickering Luma component with the interference free luma from the inteference seed of the history.
-			//if (bEnableInterferenceHeuristic)
-			//{
-			//	DecimatedColor *= PrevUninterferenceLuma * SafeRcp(Luma4(DecimatedColor));
-			//}
+			#if DIM_MOIRE_HEURISTIC && 1
+			{
+				DecimatedColor *= PrevUninterferenceLuma * SafeRcp(Luma4(DecimatedColor));
+			}
+			#endif
 		}
 		#else
 		{
@@ -271,10 +281,11 @@ void MainCS(
 				taa_half3 HistorySample = PrevHistory_LowFrequency.SampleLevel(GlobalBilinearClampedSampler, SampleBufferUV, 0);
 
 				// Cancels out the flickering Luma component with the interference free luma from the inteference seed of the history.
-				if (bEnableInterferenceHeuristic)
+				#if DIM_MOIRE_HEURISTIC
 				{
 					HistorySample *= PrevUninterferenceLuma * SafeRcp(Luma4(HistorySample));
 				}
+				#endif
 
 				DecimatedColor += HistorySample;
 			}
@@ -404,7 +415,7 @@ void MainCS(
 		
 		ParallaxRejectionMaskOutput[InputPixelPos] = ParallaxRejectionMask;
 
-		#if CONFIG_OUTPUT_INTERFERENCE_SEEDS
+		#if DIM_MOIRE_HEURISTIC
 		{
 			InterferenceSeedOutput[InputPixelPos] = taa_half4(ClosestInterferenceSeed, IntegratedInterferenceSeed, MinInterferenceSeed, MaxInterferenceSeed);
 		}
diff --git a/Engine/Shaders/Private/TemporalSuperResolution/TSRFilterFrequencies.usf b/Engine/Shaders/Private/TemporalSuperResolution/TSRFilterFrequencies.usf
--- a/Engine/Shaders/Private/TemporalSuperResolution/TSRFilterFrequencies.usf
+++ b/Engine/Shaders/Private/TemporalSuperResolution/TSRFilterFrequencies.usf
@@ -35,10 +35,12 @@ float3 OutputQuantizationError;
 Texture2D<taa_half3> InputTexture;
 Texture2D<taa_half3> PredictionSceneColorTexture;
 Texture2D<taa_half> ParallaxRejectionMaskTexture;
+Texture2D<taa_half4> InterferenceSeedTexture;
 Texture2D<taa_half> InterferenceWeightTexture;
 
 RWTexture2D<taa_half3> FilteredInputOutput;
 RWTexture2D<taa_half3> FilteredPredictionSceneColorOutput;
+RWTexture2D<taa_half4> FilteredInterferenceSeedOutput;
 RWTexture2D<taa_half>  InputSceneColorLdrLumaOutput;
 
 
@@ -82,6 +84,10 @@ void MainCS(
 	taa_half FilteredInputWeight;
 	taa_half3 FilteredHistory;
 	taa_half FilteredHistoryWeight;
+	taa_half4 FilteredInterferenceSeed;
+	taa_half FilteredInterferenceSeedWeight;
+
+	taa_half CenterSampleInterferenceWeight;
 
 	{
 		taa_half SampleSpatialWeight = 1.0;
@@ -90,7 +96,10 @@ void MainCS(
 		taa_half3 HistorySample = PredictionSceneColorTexture[InputPixelPos];
 		taa_half HistorySampleValidity = ParallaxRejectionMaskToHistoryValidity(
 			ParallaxRejectionMaskTexture[InputPixelPos]);
+		taa_half4 InterferenceSeedSample = InterferenceSeedTexture[InputPixelPos];
 		
+		CenterSampleInterferenceWeight = InterferenceWeightTexture[InputPixelPos];
+
 		taa_half InputWeight = ComputeColorWeight(InputSample);
 		taa_half HistoryWeight = HistorySampleValidity * ComputeColorWeight(HistorySample);
 
@@ -101,6 +110,13 @@ void MainCS(
 		
 		FilteredHistory       = (SampleSpatialWeight * HistoryWeight) * HistorySample;
 		FilteredHistoryWeight = (SampleSpatialWeight * HistoryWeight);
+		
+		#if DIM_MOIRE_HEURISTIC
+		{
+			FilteredInterferenceSeed       = (SampleSpatialWeight * HistorySampleValidity) * InterferenceSeedSample;
+			FilteredInterferenceSeedWeight = (SampleSpatialWeight * HistorySampleValidity);
+		}
+		#endif
 	}
 	
 	uint StartSampleId;
@@ -178,19 +194,22 @@ void MainCS(
 		taa_half3 HistorySample = PredictionSceneColorTexture[SampleInputPixelPos];
 		taa_half HistorySampleValidity = ParallaxRejectionMaskToHistoryValidity(
 			ParallaxRejectionMaskTexture[SampleInputPixelPos]);
-		
-		taa_half SampleInterferenceWeight = InterferenceWeightTexture[SampleInputPixelPos];
-		
+		taa_half4 InterferenceSeedSample = InterferenceSeedTexture[SampleInputPixelPos];
+
+		#if DIM_MOIRE_HEURISTIC
+			taa_half SampleInterferenceWeight = InterferenceWeightTexture[SampleInputPixelPos];
+			//SampleInterferenceWeight = 0.0; // TODO: need to dectect interference due to this pattern
+			//SampleInterferenceWeight = 1.0;
+
+			SampleInterferenceWeight = min(SampleInterferenceWeight, CenterSampleInterferenceWeight);
+		#else
+			taa_half SampleInterferenceWeight = taa_half(1.0);
+		#endif
+
 		#if CONFIG_SCENE_COLOR_OVERFLOW
 			InputSample = min(InputSample, LargestSceneColorRGB);
 		#endif
 
-		FLATTEN
-		if (!bEnableInterferenceHeuristic)
-		{
-			SampleInterferenceWeight = 1.0;
-		}
-
 		taa_half InputWeight = (HistorySampleValidity * SampleInterferenceWeight) * ComputeColorWeight(InputSample);
 		taa_half HistoryWeight = (HistorySampleValidity * SampleInterferenceWeight) * ComputeColorWeight(HistorySample);
 
@@ -199,11 +218,24 @@ void MainCS(
 		
 		FilteredHistory       += (SampleSpatialWeight * HistoryWeight) * HistorySample;
 		FilteredHistoryWeight += (SampleSpatialWeight * HistoryWeight);
+		
+		#if DIM_MOIRE_HEURISTIC
+		{
+			FilteredInterferenceSeed       += (SampleSpatialWeight * HistorySampleValidity * SampleInterferenceWeight) * InterferenceSeedSample;
+			FilteredInterferenceSeedWeight += (SampleSpatialWeight * HistorySampleValidity * SampleInterferenceWeight);
+		}
+		#endif
 	}
 
 	{
 		FilteredInput *= SafeRcp(FilteredInputWeight);
 		FilteredHistory *= SafeRcp(FilteredHistoryWeight);
+		
+		#if DIM_MOIRE_HEURISTIC
+		{
+			FilteredInterferenceSeed *= SafeRcp(FilteredInterferenceSeedWeight);
+		}
+		#endif
 	}
 	
 	// Compute the LDR luma used for spatial anti-aliasing's computation of the TexelOffset
@@ -252,6 +284,12 @@ void MainCS(
 
 	FilteredInputOutput[OutputPixelPos] = FilteredInput;
 	FilteredPredictionSceneColorOutput[OutputPixelPos] = FilteredHistory;
+	
+	#if DIM_MOIRE_HEURISTIC
+	{
+		FilteredInterferenceSeedOutput[OutputPixelPos] = FilteredInterferenceSeed;
+	}
+	#endif
 
 	#if DIM_OUTPUT_ANTI_ALIASING_LUMA
 	{
diff --git a/Engine/Shaders/Private/TemporalSuperResolution/TSRUpdateHistory.usf b/Engine/Shaders/Private/TemporalSuperResolution/TSRUpdateHistory.usf
--- a/Engine/Shaders/Private/TemporalSuperResolution/TSRUpdateHistory.usf
+++ b/Engine/Shaders/Private/TemporalSuperResolution/TSRUpdateHistory.usf
@@ -53,6 +53,7 @@ Texture2D<taa_half3> InputSceneColorTexture;
 Texture2D<uint2> InputSceneStencilTexture;
 Texture2D<taa_half4> InputSceneTranslucencyTexture;
 
+Texture2D<taa_half> InterferenceWeightTexture;
 Texture2D<taa_half> HistoryRejectionTexture;
 Texture2D<taa_half> TranslucencyRejectionTexture;
 
@@ -233,6 +234,7 @@ void MainCS(
 	taa_half2 LowFrequencyRejection = taa_half(1.0).xx;
 	taa_half2 TranslucencyRejection = taa_half(1.0).xx;
 	taa_half2 OutputPixelVelocity = taa_half(0.0).xx;
+	taa_half2 MoireWeight = taa_half(1.0).xx;
 	taa_subpixel_payload2 ParallaxFactorBits = taa_subpixel_payload(0x0).xx;
 	bool2 bIsResponsiveAAPixel = false;
 
@@ -314,6 +316,12 @@ void MainCS(
 			SceneStencilRef |= dpv_interleave_registers(
 				InputSceneStencilTexture.Load(int3(dpv_lo(SampleInputPixelPos), 0)) STENCIL_COMPONENT_SWIZZLE, 
 				InputSceneStencilTexture.Load(int3(dpv_hi(SampleInputPixelPos), 0)) STENCIL_COMPONENT_SWIZZLE);
+			
+			#if DIM_MOIRE_HEURISTIC
+				MoireWeight = min(MoireWeight, dpv_interleave_registers(
+					InterferenceWeightTexture[dpv_lo(SampleInputPixelPos)],
+					InterferenceWeightTexture[dpv_hi(SampleInputPixelPos)]));
+			#endif
 		}
 
 		TranslucencyRejection = max(TranslucencyRejection, taa_half(MinTranslucencyRejection));
@@ -695,6 +703,10 @@ void MainCS(
 			PrevHighFrequencyYCoCg, InputMinColor, InputMaxColor,
 			MeasureBackbufferLDRQuantizationErrorFromHDRLuma(max(PrevYCoCg[0], ClampedPrevYCoCg[0])));
 		
+		#if DIM_MOIRE_HEURISTIC
+			HighFrequencyRejection = max(HighFrequencyRejection, 1.0 - MoireWeight);
+		#endif
+
 		//PrevHighFrequencyColor = dpv_scale(PrevHighFrequencyColor, HighFrequencyRejection);
 		PrevHighFrequencyColor = dpv_scale(clamp(NormalizedPrevHighFrequencyColor, PrevHistoryMommentMin, PrevHistoryMommentMax), HighFrequencyRejection * PrevHistoryValidity);
 		PrevHistoryValidity *= HighFrequencyRejection;
diff --git a/Engine/Shaders/Private/TemporalSuperResolution/TSRUpdateSuperResHistory.usf b/Engine/Shaders/Private/TemporalSuperResolution/TSRUpdateSuperResHistory.usf
--- a/Engine/Shaders/Private/TemporalSuperResolution/TSRUpdateSuperResHistory.usf
+++ b/Engine/Shaders/Private/TemporalSuperResolution/TSRUpdateSuperResHistory.usf
@@ -135,7 +135,7 @@ void MainCS(
 				BestKernelToSampleDistSquare = KernelToSampleDistSquare;
 			}
 
-			bool bIsCloseEnough = abs(UnpackedSubpixel.InterferenceSeed - ClosestInterferenceSeed) < 0.1;
+			bool bIsCloseEnough = abs(UnpackedSubpixel.InterferenceSeed - ClosestInterferenceSeed) < 0.2;
 
 			if (bIsCloseEnough)
 			{
@@ -160,7 +160,7 @@ void MainCS(
 
 		SubpixelHistory.TexelToSample = clamp(CurrentTexelToSample, SUPERRES_OFFSET_MIN, SUPERRES_OFFSET_MAX);
 		SubpixelHistory.Age = bOffScreen ? 0 : OldestAge;
-		SubpixelHistory.InterferenceSeed = saturate(ClosestInterferenceSeed + (E.x - 0.5) * rcp(taa_half(SUPERRES_INTERFERENCE_BITMASK)));
+		SubpixelHistory.InterferenceSeed = saturate(ClosestInterferenceSeed + (E.x) * rcp(taa_half(SUPERRES_INTERFERENCE_BITMASK)));
 	}
 
 	#if 1
diff --git a/Engine/Source/Runtime/Renderer/Private/PostProcess/TemporalSuperResolution.cpp b/Engine/Source/Runtime/Renderer/Private/PostProcess/TemporalSuperResolution.cpp
--- a/Engine/Source/Runtime/Renderer/Private/PostProcess/TemporalSuperResolution.cpp
+++ b/Engine/Source/Runtime/Renderer/Private/PostProcess/TemporalSuperResolution.cpp
@@ -45,9 +45,15 @@ TAutoConsoleVariable<int32> CVarTSRRejectionAntiAliasingQuality(
 	TEXT("Controls the quality of spatial anti-aliasing on history rejection (default=1)."),
 	ECVF_Scalability | ECVF_RenderThreadSafe);
 
-TAutoConsoleVariable<int32> CVarTSREnableAntiInterference(
-	TEXT("r.TSR.AntiInterference"), 0,
-	TEXT("Enable heuristic to detect geometric interference between input pixel grid alignement and structured geometry."),
+TAutoConsoleVariable<int32> CVarTSRMoireEnable(
+	TEXT("r.TSR.Moire"), 0,
+	TEXT("Enable heuristic to detect moire pattern in the input scene color caused by geometric interference between input ")
+	TEXT("pixel grid alignement and structured geometry."),
+	ECVF_RenderThreadSafe);
+
+TAutoConsoleVariable<float> CVarTSRMoireSeedExposure(
+	TEXT("r.TSR.Moire.SeedExposure"), 0.125f,
+	TEXT("Exposure of the interference seed. Lower value will allow to detect interference more towards the shadows."),
 	ECVF_RenderThreadSafe);
 
 TAutoConsoleVariable<float> CVarTSRTranslucencyHighlightLuminance(
@@ -96,7 +102,7 @@ BEGIN_SHADER_PARAMETER_STRUCT(FTSRCommonParameters, )
 
 	SHADER_PARAMETER(FVector2D, InputJitter)
 	SHADER_PARAMETER(int32, bCameraCut)
-	SHADER_PARAMETER(int32, bEnableInterferenceHeuristic)
+	SHADER_PARAMETER(float, MoireInvExposure)
 
 	SHADER_PARAMETER_STRUCT_REF(FViewUniformShaderParameters, ViewUniformBuffer)
 END_SHADER_PARAMETER_STRUCT()
@@ -108,6 +114,9 @@ BEGIN_SHADER_PARAMETER_STRUCT(FTSRHistoryTextures, )
 	SHADER_PARAMETER_RDG_TEXTURE(Texture2D, SubpixelDetails)
 	SHADER_PARAMETER_RDG_TEXTURE(Texture2D, Translucency)
 	SHADER_PARAMETER_RDG_TEXTURE(Texture2D, TranslucencyAlpha)
+	
+	SHADER_PARAMETER_RDG_TEXTURE_ARRAY(Texture2D, LowResTextures, [1])
+	SHADER_PARAMETER_RDG_TEXTURE_ARRAY(Texture2D, SuperResTextures, [2])
 END_SHADER_PARAMETER_STRUCT()
 
 BEGIN_SHADER_PARAMETER_STRUCT(FTSRHistoryUAVs, )
@@ -117,6 +126,9 @@ BEGIN_SHADER_PARAMETER_STRUCT(FTSRHistoryUAVs, )
 	SHADER_PARAMETER_RDG_TEXTURE_UAV(RWTexture2D, SubpixelDetails)
 	SHADER_PARAMETER_RDG_TEXTURE_UAV(RWTexture2D, Translucency)
 	SHADER_PARAMETER_RDG_TEXTURE_UAV(RWTexture2D, TranslucencyAlpha)
+
+	SHADER_PARAMETER_RDG_TEXTURE_UAV_ARRAY(RWTexture2D, LowResTextures, [1])
+	SHADER_PARAMETER_RDG_TEXTURE_UAV_ARRAY(RWTexture2D, SuperResTextures, [2])
 END_SHADER_PARAMETER_STRUCT()
 
 BEGIN_SHADER_PARAMETER_STRUCT(FTSRPrevHistoryParameters, )
@@ -137,9 +149,25 @@ FTSRHistoryUAVs CreateUAVs(FRDGBuilder& GraphBuilder, const FTSRHistoryTextures&
 		UAVs.Translucency = GraphBuilder.CreateUAV(Textures.Translucency);
 		UAVs.TranslucencyAlpha = GraphBuilder.CreateUAV(Textures.TranslucencyAlpha);
 	}
+
+	if (Textures.LowResTextures[0])
+	{
+		UAVs.LowResTextures[0] = GraphBuilder.CreateUAV(Textures.LowResTextures[0]);
+	}
+	if (Textures.SuperResTextures[0])
+	{
+		UAVs.SuperResTextures[0] = GraphBuilder.CreateUAV(Textures.SuperResTextures[0]);
+	}
+	if (Textures.SuperResTextures[1])
+	{
+		UAVs.SuperResTextures[1] = GraphBuilder.CreateUAV(Textures.SuperResTextures[1]);
+	}
+
 	return UAVs;
 }
 
+class FTSRMoireHeuristic : SHADER_PERMUTATION_BOOL("DIM_MOIRE_HEURISTIC");
+
 class FTSRShader : public FGlobalShader
 {
 public:
@@ -205,7 +233,7 @@ class FTSRDecimateHistoryCS : public FTSRShader
 
 	class FOutputHalfRes : SHADER_PERMUTATION_BOOL("DIM_OUTPUT_HALF_RES");
 
-	using FPermutationDomain = TShaderPermutationDomain<FOutputHalfRes>;
+	using FPermutationDomain = TShaderPermutationDomain<FOutputHalfRes, FTSRMoireHeuristic>;
 
 	BEGIN_SHADER_PARAMETER_STRUCT(FParameters, )
 		SHADER_PARAMETER_STRUCT_INCLUDE(FTSRCommonParameters, CommonParameters)
@@ -279,11 +307,6 @@ class FTSRDetectInterferenceCS : public FTSRShader
 		SHADER_PARAMETER_RDG_TEXTURE_UAV(RWTexture2D, InterferenceWeightOutput)
 		SHADER_PARAMETER_RDG_TEXTURE_UAV(RWTexture2D, DebugOutput)
 	END_SHADER_PARAMETER_STRUCT()
-
-	static bool ShouldCompilePermutation(const FGlobalShaderPermutationParameters& Parameters)
-	{
-		return false;
-	}
 }; // class FTSRDetectInterferenceCS
 
 class FTSRFilterFrequenciesCS : public FTSRShader
@@ -293,7 +316,7 @@ class FTSRFilterFrequenciesCS : public FTSRShader
 
 	class FOutputAALumaDim : SHADER_PERMUTATION_BOOL("DIM_OUTPUT_ANTI_ALIASING_LUMA");
 
-	using FPermutationDomain = TShaderPermutationDomain<FOutputAALumaDim>;
+	using FPermutationDomain = TShaderPermutationDomain<FOutputAALumaDim, FTSRMoireHeuristic>;
 
 	BEGIN_SHADER_PARAMETER_STRUCT(FParameters, )
 		SHADER_PARAMETER_STRUCT_INCLUDE(FTSRCommonParameters, CommonParameters)
@@ -302,10 +325,12 @@ class FTSRFilterFrequenciesCS : public FTSRShader
 		SHADER_PARAMETER_RDG_TEXTURE(Texture2D, InputTexture)
 		SHADER_PARAMETER_RDG_TEXTURE(Texture2D, PredictionSceneColorTexture)
 		SHADER_PARAMETER_RDG_TEXTURE(Texture2D, ParallaxRejectionMaskTexture)
+		SHADER_PARAMETER_RDG_TEXTURE(Texture2D, InterferenceSeedTexture)
 		SHADER_PARAMETER_RDG_TEXTURE(Texture2D, InterferenceWeightTexture)
 
 		SHADER_PARAMETER_RDG_TEXTURE_UAV(RWTexture2D, FilteredInputOutput)
 		SHADER_PARAMETER_RDG_TEXTURE_UAV(RWTexture2D, FilteredPredictionSceneColorOutput)
+		SHADER_PARAMETER_RDG_TEXTURE_UAV(RWTexture2D, FilteredInterferenceSeedOutput)
 		SHADER_PARAMETER_RDG_TEXTURE_UAV(RWTexture2D, InputSceneColorLdrLumaOutput)
 		SHADER_PARAMETER_RDG_TEXTURE_UAV(RWTexture2D, DebugOutput)
 	END_SHADER_PARAMETER_STRUCT()
@@ -317,13 +342,14 @@ class FTSRCompareHistoryCS : public FTSRShader
 	SHADER_USE_PARAMETER_STRUCT(FTSRCompareHistoryCS, FTSRShader);
 
 	class FOutputHalfRes : SHADER_PERMUTATION_BOOL("DIM_OUTPUT_HALF_RES");
-	using FPermutationDomain = TShaderPermutationDomain<FOutputHalfRes>;
+	using FPermutationDomain = TShaderPermutationDomain<FOutputHalfRes, FTSRMoireHeuristic>;
 
 	BEGIN_SHADER_PARAMETER_STRUCT(FParameters, )
 		SHADER_PARAMETER_STRUCT_INCLUDE(FTSRCommonParameters, CommonParameters)
 		SHADER_PARAMETER_RDG_TEXTURE(Texture2D, ParallaxRejectionMaskTexture)
 		SHADER_PARAMETER_RDG_TEXTURE(Texture2D, FilteredInputTexture)
 		SHADER_PARAMETER_RDG_TEXTURE(Texture2D, FilteredPredictionSceneColorTexture)
+		SHADER_PARAMETER_RDG_TEXTURE(Texture2D, FilteredInterferenceSeedTexture)
 		SHADER_PARAMETER_RDG_TEXTURE(Texture2D, InterferenceWeightTexture)
 		SHADER_PARAMETER_RDG_TEXTURE_UAV(RWTexture2D, HistoryRejectionOutput)
 		SHADER_PARAMETER_RDG_TEXTURE_UAV(RWTexture2D, DebugOutput)
@@ -427,11 +453,6 @@ class FTSRUpdateSuperResHistoryCS : public FTSRShader
 		SHADER_PARAMETER_STRUCT(FTSRHistoryUAVs, HistoryOutput)
 		SHADER_PARAMETER_RDG_TEXTURE_UAV(RWTexture2D, DebugOutput)
 	END_SHADER_PARAMETER_STRUCT()
-
-	static bool ShouldCompilePermutation(const FGlobalShaderPermutationParameters& Parameters)
-	{
-		return false;
-	}
 }; // class FTSRUpdateSuperResHistoryCS
 
 class FTSRUpdateHistoryCS : public FTSRShader
@@ -442,7 +463,7 @@ class FTSRUpdateHistoryCS : public FTSRShader
 	class FRejectionAADim : SHADER_PERMUTATION_BOOL("DIM_REJECTION_ANTI_ALIASING");
 	class FSeparateTranslucencyDim : SHADER_PERMUTATION_BOOL("DIM_SEPARATE_TRANSLUCENCY");
 
-	using FPermutationDomain = TShaderPermutationDomain<FRejectionAADim, FSeparateTranslucencyDim>;
+	using FPermutationDomain = TShaderPermutationDomain<FRejectionAADim, FSeparateTranslucencyDim, FTSRMoireHeuristic>;
 
 	BEGIN_SHADER_PARAMETER_STRUCT(FParameters, )
 		SHADER_PARAMETER_STRUCT_INCLUDE(FTSRCommonParameters, CommonParameters)
@@ -450,6 +471,7 @@ class FTSRUpdateHistoryCS : public FTSRShader
 		SHADER_PARAMETER_RDG_TEXTURE_SRV(Texture2D, InputSceneStencilTexture)
 		SHADER_PARAMETER_RDG_TEXTURE(Texture2D, InputSceneTranslucencyTexture)
 
+		SHADER_PARAMETER_RDG_TEXTURE(Texture2D, InterferenceWeightTexture)
 		SHADER_PARAMETER_RDG_TEXTURE(Texture2D, HistoryRejectionTexture)
 		SHADER_PARAMETER_RDG_TEXTURE(Texture2D, TranslucencyRejectionTexture)
 
@@ -544,7 +566,7 @@ void AddTemporalSuperResolutionPasses(
 
 	bool bHalfResLowFrequency = CVarTSRHalfResShadingRejection.GetValueOnRenderThread() != 0;
 
-	bool bEnableInterferenceHeuristic = CVarTSREnableAntiInterference.GetValueOnRenderThread() != 0;
+	bool bEnableMoireHeuristic = CVarTSRMoireEnable.GetValueOnRenderThread() != 0;
 
 	bool bIsSeperateTranslucyTexturesValid = PassInputs.SeparateTranslucencyTextures != nullptr && PassInputs.SeparateTranslucencyTextures->IsColorValid();
 
@@ -644,7 +666,7 @@ void AddTemporalSuperResolutionPasses(
 
 		CommonParameters.InputJitter = View.TemporalJitterPixels;
 		CommonParameters.bCameraCut = bCameraCut;
-		CommonParameters.bEnableInterferenceHeuristic = bEnableInterferenceHeuristic;
+		CommonParameters.MoireInvExposure = 1.0f / CVarTSRMoireSeedExposure.GetValueOnRenderThread();
 		CommonParameters.ViewUniformBuffer = View.ViewUniformBuffer;
 	}
 
@@ -751,6 +773,11 @@ void AddTemporalSuperResolutionPasses(
 		DummyHistory.TranslucencyAlpha = WhiteDummy;
 
 		DummyHistory.SubpixelDetails = BlackUintDummy;
+
+		DummyHistory.LowResTextures[0] = BlackDummy;
+
+		DummyHistory.SuperResTextures[0] = BlackUintDummy;
+		DummyHistory.SuperResTextures[1] = BlackUintDummy;
 	}
 
 	// Setup the previous frame history
@@ -766,6 +793,17 @@ void AddTemporalSuperResolutionPasses(
 
 		// Register non-filterable history
 		PrevHistory.SubpixelDetails = GraphBuilder.RegisterExternalTexture(InputHistory.SubpixelDetails);
+
+		PrevHistory.LowResTextures[0] = InputHistory.LowResTextures[0].IsValid()
+			? GraphBuilder.RegisterExternalTexture(InputHistory.LowResTextures[0])
+			: DummyHistory.LowResTextures[0];
+
+		PrevHistory.SuperResTextures[0] = InputHistory.SuperResTextures[0].IsValid()
+			? GraphBuilder.RegisterExternalTexture(InputHistory.SuperResTextures[0])
+			: DummyHistory.SuperResTextures[0];
+		PrevHistory.SuperResTextures[1] = InputHistory.SuperResTextures[1].IsValid()
+			? GraphBuilder.RegisterExternalTexture(InputHistory.SuperResTextures[1])
+			: DummyHistory.SuperResTextures[1];
 	}
 	else
 	{
@@ -791,6 +829,7 @@ void AddTemporalSuperResolutionPasses(
 
 	// Create new history.
 	FTSRHistoryTextures History;
+	if (bEnableMoireHeuristic)
 	{
 		FRDGTextureDesc LowResDesc = FRDGTextureDesc::Create2D(
 			InputExtent,
@@ -798,7 +837,7 @@ void AddTemporalSuperResolutionPasses(
 			FClearValueBinding::None,
 			TexCreate_ShaderResource | TexCreate_UAV);
 
-		//History.LowResTextures[0] = GraphBuilder.CreateTexture(LowResDesc, TEXT("TSR.History.LowResMetadata[0]"));
+		History.LowResTextures[0] = GraphBuilder.CreateTexture(LowResDesc, TEXT("TSR.History.LowResMetadata[0]"));
 	}
 
 	{
@@ -827,6 +866,7 @@ void AddTemporalSuperResolutionPasses(
 		History.SubpixelDetails = GraphBuilder.CreateTexture(Desc, TEXT("TSR.History.SubpixelInfo"));
 	}
 
+	if (bEnableMoireHeuristic)
 	{
 		FRDGTextureDesc SuperResDesc = FRDGTextureDesc::Create2D(
 			HistoryExtent * 2,
@@ -834,10 +874,10 @@ void AddTemporalSuperResolutionPasses(
 			FClearValueBinding::None,
 			TexCreate_ShaderResource | TexCreate_UAV);
 
-		//History.SuperResTextures[0] = GraphBuilder.CreateTexture(SuperResDesc, TEXT("TSR.History.SuperResMetadata[0]"));
+		History.SuperResTextures[0] = GraphBuilder.CreateTexture(SuperResDesc, TEXT("TSR.History.SuperResMetadata[0]"));
 		//
-		//SuperResDesc.Format = PF_R8_UINT;
-		//History.SuperResTextures[1] = GraphBuilder.CreateTexture(SuperResDesc, TEXT("TSR.History.SuperResMetadata[1]"));
+		SuperResDesc.Format = PF_R8_UINT;
+		History.SuperResTextures[1] = GraphBuilder.CreateTexture(SuperResDesc, TEXT("TSR.History.SuperResMetadata[1]"));
 	}
 
 	// Decimate input to flicker at same frequency as input.
@@ -859,7 +899,7 @@ void AddTemporalSuperResolutionPasses(
 
 			// TODO(TSR): can compress to the history seed's 4bit per pixel
 			Desc.Format = PF_R8G8B8A8;
-			InterferenceSeedTexture = GraphBuilder.CreateTexture(Desc, TEXT("TSR.Interference.Seed"));
+			InterferenceSeedTexture = GraphBuilder.CreateTexture(Desc, TEXT("TSR.Moire.Seed"));
 		}
 
 		FTSRDecimateHistoryCS::FParameters* PassParameters = GraphBuilder.AllocParameters<FTSRDecimateHistoryCS::FParameters>();
@@ -928,6 +968,7 @@ void AddTemporalSuperResolutionPasses(
 
 		FTSRDecimateHistoryCS::FPermutationDomain PermutationVector;
 		PermutationVector.Set<FTSRDecimateHistoryCS::FOutputHalfRes>(bHalfResLowFrequency);
+		PermutationVector.Set<FTSRMoireHeuristic>(bEnableMoireHeuristic);
 
 		TShaderMapRef<FTSRDecimateHistoryCS> ComputeShader(View.ShaderMap, PermutationVector);
 		FComputeShaderUtils::AddPass(
@@ -1004,7 +1045,7 @@ void AddTemporalSuperResolutionPasses(
 	// an interference is possible.
 	// TODO(TSR): Could sample the interference seed in the DilateVelocity and detect interference in the decimate.
 	FRDGTextureRef InterferenceWeightTexture;
-	if (bEnableInterferenceHeuristic)
+	if (bEnableMoireHeuristic)
 	{
 		{
 			// TODO(TSR): Compress to 1bit per pixel
@@ -1014,7 +1055,7 @@ void AddTemporalSuperResolutionPasses(
 				FClearValueBinding::None,
 				/* InFlags = */ TexCreate_ShaderResource | TexCreate_UAV);
 
-			InterferenceWeightTexture = GraphBuilder.CreateTexture(Desc, TEXT("TSR.Interference.Weight"));
+			InterferenceWeightTexture = GraphBuilder.CreateTexture(Desc, TEXT("TSR.Moire.Weight"));
 		}
 
 		FTSRDetectInterferenceCS::FParameters* PassParameters = GraphBuilder.AllocParameters<FTSRDetectInterferenceCS::FParameters>();
@@ -1043,7 +1084,6 @@ void AddTemporalSuperResolutionPasses(
 	}
 	else
 	{
-		// TODO(TSR): Shader permutation.
 		InterferenceWeightTexture = WhiteDummy;
 	}
 
@@ -1054,6 +1094,7 @@ void AddTemporalSuperResolutionPasses(
 		// Filter out the high frquencies
 		FRDGTextureRef FilteredInputTexture;
 		FRDGTextureRef FilteredPredictionSceneColorTexture;
+		FRDGTextureRef FilteredInterferenceSeedTexture = nullptr;
 		{
 			{
 				FRDGTextureDesc Desc = FRDGTextureDesc::Create2D(
@@ -1077,6 +1118,16 @@ void AddTemporalSuperResolutionPasses(
 				InputSceneColorLdrLumaTexture = GraphBuilder.CreateTexture(Desc, TEXT("TSR.SceneColorLdrLuma"));
 			}
 
+			{
+				FRDGTextureDesc Desc = FRDGTextureDesc::Create2D(
+					LowFrequencyExtent,
+					PF_R8G8B8A8,
+					FClearValueBinding::None,
+					/* InFlags = */ TexCreate_ShaderResource | TexCreate_UAV);
+
+				FilteredInterferenceSeedTexture = GraphBuilder.CreateTexture(Desc, TEXT("TSR.Moire.FilteredSeeds"));
+			}
+
 			FTSRFilterFrequenciesCS::FParameters* PassParameters = GraphBuilder.AllocParameters<FTSRFilterFrequenciesCS::FParameters>();
 			PassParameters->CommonParameters = CommonParameters;
 			PassParameters->OutputQuantizationError = ComputePixelFormatQuantizationError(FilteredInputTexture->Desc.Format);
@@ -1092,16 +1143,19 @@ void AddTemporalSuperResolutionPasses(
 				PassParameters->InputTexture = PassInputs.SceneColorTexture;
 				PassParameters->PredictionSceneColorTexture = PredictionSceneColorTexture;
 				PassParameters->ParallaxRejectionMaskTexture = ParallaxRejectionMaskTexture;
+				PassParameters->InterferenceSeedTexture = InterferenceSeedTexture;
 				PassParameters->InterferenceWeightTexture = InterferenceWeightTexture;
 			}
 
 			PassParameters->FilteredInputOutput = GraphBuilder.CreateUAV(FilteredInputTexture);
 			PassParameters->FilteredPredictionSceneColorOutput = GraphBuilder.CreateUAV(FilteredPredictionSceneColorTexture);
+			PassParameters->FilteredInterferenceSeedOutput = GraphBuilder.CreateUAV(FilteredInterferenceSeedTexture);
 			PassParameters->InputSceneColorLdrLumaOutput = InputSceneColorLdrLumaTexture ? GraphBuilder.CreateUAV(InputSceneColorLdrLumaTexture) : nullptr;
 			PassParameters->DebugOutput = CreateDebugUAV(LowFrequencyExtent, TEXT("Debug.TSR.FilterFrequencies"));
 
 			FTSRFilterFrequenciesCS::FPermutationDomain PermutationVector;
 			PermutationVector.Set<FTSRFilterFrequenciesCS::FOutputAALumaDim>(RejectionAntiAliasingQuality > 0);
+			PermutationVector.Set<FTSRMoireHeuristic>(bEnableMoireHeuristic);
 
 			TShaderMapRef<FTSRFilterFrequenciesCS> ComputeShader(View.ShaderMap, PermutationVector);
 			FComputeShaderUtils::AddPass(
@@ -1133,6 +1187,7 @@ void AddTemporalSuperResolutionPasses(
 			PassParameters->ParallaxRejectionMaskTexture = ParallaxRejectionMaskTexture;
 			PassParameters->FilteredInputTexture = FilteredInputTexture;
 			PassParameters->FilteredPredictionSceneColorTexture = FilteredPredictionSceneColorTexture;
+			PassParameters->FilteredInterferenceSeedTexture = FilteredInterferenceSeedTexture;
 			PassParameters->InterferenceWeightTexture = InterferenceWeightTexture;
 
 			PassParameters->HistoryRejectionOutput = GraphBuilder.CreateUAV(HistoryRejectionTexture);
@@ -1140,6 +1195,7 @@ void AddTemporalSuperResolutionPasses(
 
 			FTSRCompareHistoryCS::FPermutationDomain PermutationVector;
 			PermutationVector.Set<FTSRCompareHistoryCS::FOutputHalfRes>(bOutputHalfRes);
+			PermutationVector.Set<FTSRMoireHeuristic>(bEnableMoireHeuristic);
 
 			TShaderMapRef<FTSRCompareHistoryCS> ComputeShader(View.ShaderMap, PermutationVector);
 			FComputeShaderUtils::AddPass(
@@ -1263,7 +1319,7 @@ void AddTemporalSuperResolutionPasses(
 			FComputeShaderUtils::GetGroupCount(RejectionRect.Size(), 8));
 	}
 
-	if (bEnableInterferenceHeuristic)
+	if (bEnableMoireHeuristic)
 	{
 		FTSRUpdateSuperResHistoryCS::FParameters* PassParameters = GraphBuilder.AllocParameters<FTSRUpdateSuperResHistoryCS::FParameters>();
 		PassParameters->CommonParameters = CommonParameters;
@@ -1305,6 +1361,7 @@ void AddTemporalSuperResolutionPasses(
 		PassParameters->InputSceneStencilTexture = GraphBuilder.CreateSRV(
 			FRDGTextureSRVDesc::CreateWithPixelFormat(PassInputs.SceneDepthTexture, PF_X24_G8));
 		PassParameters->InputSceneTranslucencyTexture = SeparateTranslucencyTexture;
+		PassParameters->InterferenceWeightTexture = InterferenceWeightTexture;
 		PassParameters->HistoryRejectionTexture = DilatedHistoryRejectionTexture;
 		PassParameters->TranslucencyRejectionTexture = TranslucencyRejectionTexture ? TranslucencyRejectionTexture : BlackDummy;
 
@@ -1331,6 +1388,7 @@ void AddTemporalSuperResolutionPasses(
 		FTSRUpdateHistoryCS::FPermutationDomain PermutationVector;
 		PermutationVector.Set<FTSRUpdateHistoryCS::FRejectionAADim>(RejectionAntiAliasingQuality > 0);
 		PermutationVector.Set<FTSRUpdateHistoryCS::FSeparateTranslucencyDim>(bAccumulateSeparateTranslucency);
+		PermutationVector.Set<FTSRMoireHeuristic>(bEnableMoireHeuristic);
 
 		TShaderMapRef<FTSRUpdateHistoryCS> ComputeShader(View.ShaderMap, PermutationVector);
 		FComputeShaderUtils::AddPass(
@@ -1402,6 +1460,13 @@ void AddTemporalSuperResolutionPasses(
 				SeparateTranslucencyTexture, &View.ViewState->PrevFrameViewInfo.SeparateTranslucency);
 		}
 
+		if (bEnableMoireHeuristic)
+		{
+			GraphBuilder.QueueTextureExtraction(History.LowResTextures[0], &OutputHistory.LowResTextures[0]);
+			GraphBuilder.QueueTextureExtraction(History.SuperResTextures[0], &OutputHistory.SuperResTextures[0]);
+			GraphBuilder.QueueTextureExtraction(History.SuperResTextures[1], &OutputHistory.SuperResTextures[1]);
+		}
+
 		// Extract the output for next frame SSR so that separate translucency shows up in SSR.
 		if (bAccumulateSeparateTranslucency || HistorySize != OutputRect.Size())
 		{
diff --git a/Engine/Source/Runtime/Renderer/Private/SceneRendering.h b/Engine/Source/Runtime/Renderer/Private/SceneRendering.h
--- a/Engine/Source/Runtime/Renderer/Private/SceneRendering.h
+++ b/Engine/Source/Runtime/Renderer/Private/SceneRendering.h
@@ -903,6 +903,10 @@ struct FTSRHistory
 	// Non-filterable output resolution
 	TRefCountPtr<IPooledRenderTarget> SubpixelDetails;
 
+	// Super resolution informaiton
+	TStaticArray<TRefCountPtr<IPooledRenderTarget>, 1> LowResTextures;
+	TStaticArray<TRefCountPtr<IPooledRenderTarget>, 2> SuperResTextures;
+
 	FIntRect OutputViewportRect;
 
 
