Fixed broken instance culling with non-batched mode on mobile

diff --git a/Engine/Source/Runtime/Renderer/Private/MobileShadingRenderer.cpp b/Engine/Source/Runtime/Renderer/Private/MobileShadingRenderer.cpp
--- a/Engine/Source/Runtime/Renderer/Private/MobileShadingRenderer.cpp
+++ b/Engine/Source/Runtime/Renderer/Private/MobileShadingRenderer.cpp
@@ -168,12 +168,12 @@ static void RenderOpaqueFX(
 	}
 }
 
-
 BEGIN_SHADER_PARAMETER_STRUCT(FMobileRenderPassParameters, )
 	SHADER_PARAMETER_STRUCT_INCLUDE(FViewShaderParameters, View)
 	SHADER_PARAMETER_RDG_UNIFORM_BUFFER(FMobileBasePassUniformParameters, MobileBasePass)
 	SHADER_PARAMETER_RDG_UNIFORM_BUFFER(FDebugViewModePassUniformParameters, DebugViewMode)
-	SHADER_PARAMETER_STRUCT_INCLUDE(FInstanceCullingDrawParams, InstanceCullingDrawParams)
+	RDG_BUFFER_ACCESS_ARRAY(DrawIndirectArgsBuffers)
+	RDG_BUFFER_ACCESS_ARRAY(InstanceIdOffsetBuffers)
 	RENDER_TARGET_BINDING_SLOTS()
 END_SHADER_PARAMETER_STRUCT()
 
@@ -601,6 +601,23 @@ void FMobileSceneRenderer::InitViews(FRDGBuilder& GraphBuilder, FSceneTexturesCo
 	OnStartRender(RHICmdList);
 }
 
+static void BuildMeshPassInstanceCullingDrawParams(FRDGBuilder& GraphBuilder, const FGPUScene& GPUScene, FParallelMeshDrawCommandPass& MeshPass, FMobileRenderPassParameters& PassParameters, FInstanceCullingDrawParams& InstanceCullingDrawParams)
+{
+	// This manual resource handling is need since we can't include multiple FInstanceCullingDrawParams into RDG pass parameters,
+	// because FInstanceCullingDrawParams includes a global UB
+	MeshPass.BuildRenderingCommands(GraphBuilder, GPUScene, InstanceCullingDrawParams);
+	
+	if (InstanceCullingDrawParams.DrawIndirectArgsBuffer)
+	{
+		PassParameters.DrawIndirectArgsBuffers.Emplace(InstanceCullingDrawParams.DrawIndirectArgsBuffer, ERHIAccess::IndirectArgs);
+	}
+	
+	if (InstanceCullingDrawParams.InstanceIdOffsetBuffer)
+	{
+		PassParameters.InstanceIdOffsetBuffers.Emplace(InstanceCullingDrawParams.InstanceIdOffsetBuffer, ERHIAccess::VertexOrIndexBuffer);
+	}
+}
+
 /*
 * Renders the Full Depth Prepass
 */
@@ -624,7 +641,10 @@ void FMobileSceneRenderer::RenderFullDepthPrepass(FRDGBuilder& GraphBuilder, FSc
 		PassParameters->RenderTargets = BasePassRenderTargets;
 		PassParameters->View = View.GetShaderParameters();
 
-		View.ParallelMeshDrawCommandPasses[EMeshPass::DepthPass].BuildRenderingCommands(GraphBuilder, Scene->GPUScene, PassParameters->InstanceCullingDrawParams);
+		if (Scene->GPUScene.IsEnabled())
+		{
+			BuildMeshPassInstanceCullingDrawParams(GraphBuilder, Scene->GPUScene, View.ParallelMeshDrawCommandPasses[EMeshPass::DepthPass], *PassParameters, MeshPassInstanceCullingDrawParams[EMeshPass::DepthPass]);
+		}
 
 		GraphBuilder.AddPass(
 			RDG_EVENT_NAME("FullDepthPrepass"),
@@ -899,18 +919,23 @@ void FMobileSceneRenderer::Render(FRDGBuilder& GraphBuilder)
 
 void FMobileSceneRenderer::BuildInstanceCullingDrawParams(FRDGBuilder& GraphBuilder, FViewInfo& View, FMobileRenderPassParameters* PassParameters)
 {
-	View.ParallelMeshDrawCommandPasses[EMeshPass::DepthPass].BuildRenderingCommands(GraphBuilder, Scene->GPUScene, MeshPassInstanceCullingDrawParams[EMeshPass::DepthPass]);
-	View.ParallelMeshDrawCommandPasses[EMeshPass::BasePass].BuildRenderingCommands(GraphBuilder, Scene->GPUScene, MeshPassInstanceCullingDrawParams[EMeshPass::BasePass]);
-	View.ParallelMeshDrawCommandPasses[EMeshPass::SkyPass].BuildRenderingCommands(GraphBuilder, Scene->GPUScene, MeshPassInstanceCullingDrawParams[EMeshPass::SkyPass]);
-	View.ParallelMeshDrawCommandPasses[StandardTranslucencyMeshPass].BuildRenderingCommands(GraphBuilder, Scene->GPUScene, MeshPassInstanceCullingDrawParams[StandardTranslucencyMeshPass]);
-	if (ViewFamily.UseDebugViewPS())
+	if (Scene->GPUScene.IsEnabled())
 	{
-		View.ParallelMeshDrawCommandPasses[EMeshPass::DebugViewMode].BuildRenderingCommands(GraphBuilder, Scene->GPUScene, MeshPassInstanceCullingDrawParams[EMeshPass::DebugViewMode]);
-	}
+		const EMeshPass::Type BaseMeshPasses[] = 
+		{
+			EMeshPass::DepthPass,
+			EMeshPass::BasePass,
+			EMeshPass::SkyPass,
+			StandardTranslucencyMeshPass,
+			EMeshPass::DebugViewMode
+		};
 
-	// these resources should be shared among all mesh passes
-	// we assign it a pass parameters to make sure RDG does resource transtions
-	PassParameters->InstanceCullingDrawParams = MeshPassInstanceCullingDrawParams[EMeshPass::BasePass];
+		for (int32 MeshPassIdx = 0; MeshPassIdx < UE_ARRAY_COUNT(BaseMeshPasses); ++MeshPassIdx)
+		{
+			const EMeshPass::Type PassType = BaseMeshPasses[MeshPassIdx];
+			BuildMeshPassInstanceCullingDrawParams(GraphBuilder, Scene->GPUScene, View.ParallelMeshDrawCommandPasses[PassType], *PassParameters, MeshPassInstanceCullingDrawParams[PassType]);
+		}
+	}
 }
 
 void FMobileSceneRenderer::RenderForward(FRDGBuilder& GraphBuilder, FRDGTextureRef ViewFamilyTexture, FSceneTextures& SceneTextures)
