Local Exposure using Bilaterial Grid based Contrast Reduction.
- Image exposure is locally adjusted by decomposing luminance of the frame into a base layer and a detail layer.
- Base layer contrast can be reduced while preserving detail.
- Support blending with gaussian blurred luminance as proposed in Advances in Real-Time Rendering in Games 2021.
- Controls in PostProcessVolume.
- r.LocalExposure cvar to control whether local exposure is supported.
- Visualization mode in "Show->Visualize->Local Expsoure". Multiple modes controllable using r.LocalExposure.VisualizeDebugMode.
- Bilateral grid generation not yet optimized. Optimization should include increasing depth of grid to match histogram size so that global histogram can be generated from grid.
- Added option to use mirror address mode in PostProcessWeightedSampleSum, since black border adds noticeable vignetting to the blurred luminance.

diff --git a/Engine/Shaders/Private/Histogram.usf b/Engine/Shaders/Private/Histogram.usf
--- a/Engine/Shaders/Private/Histogram.usf
+++ b/Engine/Shaders/Private/Histogram.usf
@@ -150,7 +150,7 @@ void CalculateBucketsAndWeights(in uint2 TexelPos, out float Weight0, out float
 	float4 SceneColor = InputTexture.Load(int3(TexelPos, 0));
 	SceneColor.xyz *= View.OneOverPreExposure;
 
-	float LuminanceVal = max(dot(SceneColor.xyz, float3(1.0f, 1.0f, 1.0f)/3.0f), EyeAdaptation_LuminanceMin);
+	float LuminanceVal = CalculateEyeAdaptationLuminance(SceneColor.xyz);
 	float LogLuminance = ComputeHistogramPositionFromLuminance(LuminanceVal);
 	// Map the normalized histogram position into texels.
 	float fBucket = saturate(LogLuminance) * (HISTOGRAM_SIZE-1);
diff --git a/Engine/Shaders/Private/PostProcessEyeAdaptation.usf b/Engine/Shaders/Private/PostProcessEyeAdaptation.usf
--- a/Engine/Shaders/Private/PostProcessEyeAdaptation.usf
+++ b/Engine/Shaders/Private/PostProcessEyeAdaptation.usf
@@ -147,7 +147,7 @@ void BasicEyeAdaptationSetupPS(
 	OutColor = Texture2DSample(ColorTexture, ColorSampler, UV);
 
 	// Use max to ensure intensity is never zero (so the following log is well behaved)
-	const float Intensity = max(dot(OutColor.xyz, float3(1.0f, 1.0f, 1.0f)/3.0f), EyeAdaptation_LuminanceMin);
+	const float Intensity = CalculateEyeAdaptationLuminance(OutColor.xyz);
 	const float LogIntensity = clamp(log2(Intensity),-10.0f,20.0f);
 
 	// Store log intensity in the alpha channel: scale to 0,1 range.
diff --git a/Engine/Shaders/Private/PostProcessHistogram.usf b/Engine/Shaders/Private/PostProcessHistogram.usf
--- a/Engine/Shaders/Private/PostProcessHistogram.usf
+++ b/Engine/Shaders/Private/PostProcessHistogram.usf
@@ -12,8 +12,13 @@ SCREEN_PASS_TEXTURE_VIEWPORT(Input)
 
 Texture2D InputTexture;
 
+#if BILATERAL_GRID
+// Output bilateral grid texture (UAV)
+RWTexture3D<float2> BilateralGridRWTexture;
+#else
 // Output histogram texture (UAV)
 RWTexture2D<float4> HistogramRWTexture;
+#endif
 
 // Number of thread groups in the dispatch
 uint2 ThreadGroupCount;
@@ -21,6 +26,10 @@ uint2 ThreadGroupCount;
 // THREADGROUP_SIZEX*THREADGROUP_SIZEY histograms of the size HISTOGRAM_SIZE
 groupshared float SharedHistogram[HISTOGRAM_SIZE][THREADGROUP_SIZEX][THREADGROUP_SIZEY];
 
+#if BILATERAL_GRID
+groupshared float SharedLuminanceSum[HISTOGRAM_SIZE][THREADGROUP_SIZEX][THREADGROUP_SIZEY];
+#endif
+
 [numthreads(THREADGROUP_SIZEX, THREADGROUP_SIZEY, 1)]
 void MainCS(
 	uint3 GroupId : SV_GroupID,
@@ -33,6 +42,9 @@ void MainCS(
 	UNROLL for (uint i = 0; i < HISTOGRAM_SIZE; ++i)
 	{
 		SharedHistogram[i][GroupThreadId.x][GroupThreadId.y] = 0.0f;
+#if BILATERAL_GRID
+		SharedLuminanceSum[i][GroupThreadId.x][GroupThreadId.y] = 0.0f;
+#endif
 	}
 	
 	GroupMemoryBarrierWithGroupSync();
@@ -59,14 +71,16 @@ void MainCS(
 
 				SceneColor.xyz *= View.OneOverPreExposure;
 
-				float LuminanceVal = max(dot(SceneColor.xyz, float3(1.0f, 1.0f, 1.0f)/3.0f), EyeAdaptation_LuminanceMin);
+				float LuminanceVal = CalculateEyeAdaptationLuminance(SceneColor.xyz);
+
+				float LogLuminance = log2(LuminanceVal);
+				float LogLuminanceHist = ComputeHistogramPositionFromLogLuminance(LogLuminance);
 
-				float LogLuminance = ComputeHistogramPositionFromLuminance(LuminanceVal);
 				float2 ScreenUV = (TexelPos.xy - Input_ViewportMin) * InvViewSize;
 				float ScreenWeight = AdaptationWeightTexture(ScreenUV);
 
 				// Map the normalized histogram position into texels.
-				float fBucket = saturate(LogLuminance) * (HISTOGRAM_SIZE-1);// * 0.9999f;
+				float fBucket = saturate(LogLuminanceHist) * (HISTOGRAM_SIZE-1);// * 0.9999f;
 
 				// Find two discrete buckets that straddle the continuous histogram position.
 				uint Bucket0 = (uint)(fBucket);
@@ -79,6 +93,7 @@ void MainCS(
 				float Weight1 = frac(fBucket);
 				float Weight0 = 1.0f - Weight1;
 
+#if !BILATERAL_GRID
 				// When EyeAdaptation_BlackHistogramBucketInfluence=.0, we will ignore the last bucket. The main use
 				// case is so the black background pixels in the editor have no effect. But if we have cases where
 				// pixel values can actually be black, we want to set EyeAdaptation_LastHistogramBucketInfluence=1.0.
@@ -88,9 +103,18 @@ void MainCS(
 					Weight0 *= EyeAdaptation_BlackHistogramBucketInfluence;
 				}
 
+				Weight0 *= ScreenWeight;
+				Weight1 *= ScreenWeight;
+#endif
+
 				// Accumulate the weight to the nearby history buckets.
-				SharedHistogram[Bucket0][GroupThreadId.x][GroupThreadId.y] += Weight0 * ScreenWeight;
-				SharedHistogram[Bucket1][GroupThreadId.x][GroupThreadId.y] += Weight1 * ScreenWeight;
+				SharedHistogram[Bucket0][GroupThreadId.x][GroupThreadId.y] += Weight0;
+				SharedHistogram[Bucket1][GroupThreadId.x][GroupThreadId.y] += Weight1;
+
+#if BILATERAL_GRID
+				SharedLuminanceSum[Bucket0][GroupThreadId.x][GroupThreadId.y] += LogLuminance * Weight0;
+				SharedLuminanceSum[Bucket1][GroupThreadId.x][GroupThreadId.y] += LogLuminance * Weight1;
+#endif
 			}
 		}
 	}
@@ -101,6 +125,9 @@ void MainCS(
 	if (GroupIndex < HISTOGRAM_SIZE / 4) 
 	{
 		float4 Sum = 0;
+#if BILATERAL_GRID
+		float4 SumLuminance = 0;
+#endif
 
 		LOOP for (uint y = 0; y < THREADGROUP_SIZEY; ++y)
 		{
@@ -111,9 +138,23 @@ void MainCS(
 					SharedHistogram[GroupIndex * 4 + 1][x][y],
 					SharedHistogram[GroupIndex * 4 + 2][x][y],
 					SharedHistogram[GroupIndex * 4 + 3][x][y]);
+
+#if BILATERAL_GRID
+				SumLuminance += float4(
+					SharedLuminanceSum[GroupIndex * 4 + 0][x][y],
+					SharedLuminanceSum[GroupIndex * 4 + 1][x][y],
+					SharedLuminanceSum[GroupIndex * 4 + 2][x][y],
+					SharedLuminanceSum[GroupIndex * 4 + 3][x][y]);
+#endif
 			}
 		}
 
+#if BILATERAL_GRID
+		BilateralGridRWTexture[uint3(GroupId.xy, GroupIndex * 4 + 0)] = float2(SumLuminance.x, Sum.x);
+		BilateralGridRWTexture[uint3(GroupId.xy, GroupIndex * 4 + 1)] = float2(SumLuminance.y, Sum.y);
+		BilateralGridRWTexture[uint3(GroupId.xy, GroupIndex * 4 + 2)] = float2(SumLuminance.z, Sum.z);
+		BilateralGridRWTexture[uint3(GroupId.xy, GroupIndex * 4 + 3)] = float2(SumLuminance.w, Sum.w);
+#else
 		float2 MaxExtent = Input_ViewportSize;
 		float Area = MaxExtent.x * MaxExtent.y;
 	
@@ -122,5 +163,6 @@ void MainCS(
 
 		// Output texture with one histogram per line, x and y unwrapped into all the lines
 		HistogramRWTexture[uint2(GroupIndex, GroupId.x + GroupId.y * ThreadGroupCount.x)] = Sum * NormalizeFactor;
+#endif
 	}
 }
diff --git a/Engine/Shaders/Private/PostProcessHistogramCommon.ush b/Engine/Shaders/Private/PostProcessHistogramCommon.ush
--- a/Engine/Shaders/Private/PostProcessHistogramCommon.ush
+++ b/Engine/Shaders/Private/PostProcessHistogramCommon.ush
@@ -16,6 +16,10 @@ float EyeAdaptation_ExposureSpeedDown;
 float EyeAdaptation_HistogramScale;
 float EyeAdaptation_HistogramBias;
 float EyeAdaptation_LuminanceMin;
+float EyeAdaptation_LocalExposureContrastReduction;
+float EyeAdaptation_LocalExposureDetailStrength;
+float EyeAdaptation_LocalExposureBlurredLuminanceBlend;
+float EyeAdaptation_LocalExposureMiddleGreyExposureCompensation;
 float EyeAdaptation_BlackHistogramBucketInfluence;
 float EyeAdaptation_GreyMult;
 float EyeAdaptation_ExponentialUpM;
@@ -27,32 +31,47 @@ int EyeAdaptation_VisualizeDebugType;
 Texture2D EyeAdaptation_MeterMaskTexture;
 SamplerState EyeAdaptation_MeterMaskSampler;
 
+float CalculateEyeAdaptationLuminance(float3 Color)
+{
+	return max(dot(Color, float3(1.0f, 1.0f, 1.0f) / 3.0f), EyeAdaptation_LuminanceMin);
+}
+
+// inverse of ComputeLogLuminanceFromHistogramPosition
+// @param LogLuminance
+// @return HistogramPosition 0..1
+float ComputeHistogramPositionFromLogLuminance(float LogLuminance)
+{
+	return LogLuminance * EyeAdaptation_HistogramScale + EyeAdaptation_HistogramBias;
+}
+
 // inverse of ComputeLuminanceFromHistogramPosition
-// is executed more often than ComputeLuminanceFromHistogramPosition()
 // @param Luminance
 // @return HistogramPosition 0..1
 float ComputeHistogramPositionFromLuminance(float Luminance)
 {
-	return log2(Luminance) * EyeAdaptation_HistogramScale + EyeAdaptation_HistogramBias;
+	return ComputeHistogramPositionFromLogLuminance(log2(Luminance));
 }
 
-// inverse of ComputeHistogramPositionFromLuminance()
-// is not as often executed as ComputeHistogramPositionFromLuminance()
+// inverse of ComputeHistogramPositionFromLogLuminance()
 // @param HistogramPosition 0..1
-// @return Luminance
-float ComputeLuminanceFromHistogramPosition(float HistogramPosition)
+// @return LogLuminance
+float ComputeLogLuminanceFromHistogramPosition(float HistogramPosition)
 {
-	return exp2((HistogramPosition - EyeAdaptation_HistogramBias) / EyeAdaptation_HistogramScale);
+	return ((HistogramPosition - EyeAdaptation_HistogramBias) / EyeAdaptation_HistogramScale);
 }
 
-
-float ComputeLogLuminanceFromHistogramPosition(float HistogramPosition)
+// inverse of ComputeHistogramPositionFromLuminance()
+// @param HistogramPosition 0..1
+// @return Luminance
+float ComputeLuminanceFromHistogramPosition(float HistogramPosition)
 {
-	return ((HistogramPosition - EyeAdaptation_HistogramBias) / EyeAdaptation_HistogramScale);
+	return exp2(ComputeLogLuminanceFromHistogramPosition(HistogramPosition));
 }
 
+#ifndef HISTOGRAM_SIZE
+	#define HISTOGRAM_SIZE 64
+#endif
 
-#define HISTOGRAM_SIZE 64
 #define HISTOGRAM_TEXEL_SIZE (HISTOGRAM_SIZE / 4)
 
 float4 ComputeARGBStripeMaskInt(uint x)
@@ -182,4 +201,37 @@ float ComputeEyeAdaptation(float OldExposure, float TargetExposure, float FrameT
 	const float AdjustedExposure = lerp(AdaptedExposure,TargetExposure,EyeAdaptation_ForceTarget);
 	
 	return AdjustedExposure;
-}
\ No newline at end of file
+}
+
+#define BILATERAL_GRID_DEPTH 32
+
+float CalculateBaseLogLuminance(float BilateralLum, float BlurredLum, float BlurredLumBlend, float ExposureScale)
+{
+	return lerp(BilateralLum, BlurredLum, BlurredLumBlend) + log2(ExposureScale);
+}
+
+float CalculateBaseLogLuminance(float LogLum, float BlurredLumBlend, float ExposureScale, float2 UV, Texture3D LumBilateralGrid, Texture2D BlurredLogLum, SamplerState LumBilateralGridSampler, SamplerState BlurredLogLumSampler)
+{
+	float LumW = (ComputeHistogramPositionFromLogLuminance(LogLum) * (BILATERAL_GRID_DEPTH - 1) + 0.5f) / BILATERAL_GRID_DEPTH;
+	float2 BilateralGridLum = Texture3DSample(LumBilateralGrid, LumBilateralGridSampler, float3(UV, LumW)).xy;
+	float BilateralLum = BilateralGridLum.x / BilateralGridLum.y;
+
+	float BlurredLum = Texture2DSample(BlurredLogLum, BlurredLogLumSampler, UV).r;
+
+	if (BilateralGridLum.y < 0.001)
+	{
+		// fallback to blurred luminance if bilateral grid doesn't have data
+		// this can happen since grid is populated using half resolution image
+		BilateralLum = BlurredLum;
+	}
+
+	return CalculateBaseLogLuminance(BilateralLum, BlurredLum, BlurredLumBlend, ExposureScale);
+}
+
+float CalculateLocalExposure(float LogLum, float BaseLogLum, float LogMiddleGrey, float ContrastReduction, float DetailStrength)
+{
+	float DetailLogLum = LogLum - BaseLogLum;
+
+	float LogLocalLum = LogMiddleGrey + (BaseLogLum - LogMiddleGrey) * ContrastReduction + DetailLogLum * DetailStrength;
+	return exp2(LogLocalLum - LogLum);
+}
diff --git a/Engine/Shaders/Private/PostProcessLocalExposure.usf b/Engine/Shaders/Private/PostProcessLocalExposure.usf
new file mode 100644
--- /dev/null
+++ b/Engine/Shaders/Private/PostProcessLocalExposure.usf
@@ -0,0 +1,18 @@
+// Copyright Epic Games, Inc. All Rights Reserved.
+
+#include "Common.ush"
+#include "ScreenPass.ush"
+#include "PostProcessHistogramCommon.ush"
+
+SCREEN_PASS_TEXTURE_VIEWPORT(Input)
+Texture2D InputTexture;
+
+RWTexture2D<float> Output2DFRW;
+
+[numthreads(THREADGROUP_SIZEX, THREADGROUP_SIZEY, 1)]
+void SetupLogLuminanceCS(uint3 DispatchThreadId : SV_DispatchThreadID)
+{
+	float3 SceneColor = InputTexture.Load(uint3(Input_ViewportMin + DispatchThreadId.xy, 0)).rgb * View.OneOverPreExposure;
+	float LuminanceVal = CalculateEyeAdaptationLuminance(SceneColor);
+	Output2DFRW[DispatchThreadId.xy] = log2(LuminanceVal);
+}
diff --git a/Engine/Shaders/Private/PostProcessTonemap.usf b/Engine/Shaders/Private/PostProcessTonemap.usf
--- a/Engine/Shaders/Private/PostProcessTonemap.usf
+++ b/Engine/Shaders/Private/PostProcessTonemap.usf
@@ -15,6 +15,7 @@
 #include "PostProcessCommon.ush"
 #include "TonemapCommon.ush"
 #include "EyeAdaptationCommon.ush"
+#include "PostProcessHistogramCommon.ush"
 #include "PixelQuadMessagePassing.ush"
 #include "ScreenPass.ush"
 #if SUPPORTS_BLOOM_APPLY_PARAMETERS
@@ -62,6 +63,10 @@ SamplerState BloomSampler;
 	StructuredBuffer<FBloomApplyInfo> BloomApplyParameters;
 #endif
 
+Texture3D LumBilateralGrid;
+Texture2D BlurredLogLum;
+SamplerState LumBilateralGridSampler;
+SamplerState BlurredLogLumSampler;
 
 // xyz:SceneColorTint.rgb, w:unused
 float4 ColorScale0;
@@ -227,7 +232,7 @@ float4 ChromaticAberrationParams;
 void TonemapCommonVS(
 	in	float4 Position,
 	in	float2 TexCoord,
-	out	float3 OutExposureScaleVignette,
+	out	float4 OutExposureScaleVignette,
 	out	float4 OutGrainUV,
 	out	float2 OutScreenPos,
 	out	float2 OutFullViewUV
@@ -244,22 +249,25 @@ void TonemapCommonVS(
 #if EYEADAPTATION_EXPOSURE_FIX == 1
 	#if FEATURE_LEVEL >= FEATURE_LEVEL_SM5
 		// texture can be GWhiteTexture which is 1x1. It's important we don't read outside bounds.
-		OutExposureScaleVignette.x = EyeAdaptationTexture.Load(int3(0, 0, 0)).x;
+		OutExposureScaleVignette.xy = EyeAdaptationTexture.Load(int3(0, 0, 0)).xw;
 	#else
-		OutExposureScaleVignette.x = EyeAdaptationBuffer[0].x;
+		OutExposureScaleVignette.xy = EyeAdaptationBuffer[0].xw;
 	#endif
 #else
 	// If eye adaptation is disabled (e.g. fixed exposure level ) or  not available.  
-	OutExposureScaleVignette.x = DefaultEyeExposure;
+	OutExposureScaleVignette.xy = DefaultEyeExposure;
 #endif
 
+	// Middle grey lum value adjusted by exposure compensation
+	OutExposureScaleVignette.y = log2(0.18 * OutExposureScaleVignette.y * EyeAdaptation_LocalExposureMiddleGreyExposureCompensation);
+
 	const float AspectRatio = Output_ViewportSize.y * Output_ViewportSizeInverse.x;
 
 	float2 ColorViewportPos = UVToScreenPos(TexCoord, Color_Extent);
 
 	// Scale vignette to always be a circle with consistent corner intensity.
 	float2 LensViewportPos = ConvertScreenViewportSpaceToLensViewportSpace(ColorViewportPos);
-	OutExposureScaleVignette.yz = VignetteSpace(LensViewportPos, AspectRatio);
+	OutExposureScaleVignette.zw = VignetteSpace(LensViewportPos, AspectRatio);
 
 	// Grain
 	OutGrainUV.xy = TexCoord + Color_ExtentInverse * float2(-0.5,0.5);
@@ -274,7 +282,7 @@ void MainVS(
 	in 					float4 InPosition 				: ATTRIBUTE0,
 	in 					float2 InTexCoord 				: ATTRIBUTE1,
 	out noperspective 	float2 OutTexCoord 				: TEXCOORD0,
-	out noperspective 	float3 OutExposureScaleVignette : TEXCOORD1,
+	out noperspective 	float4 OutExposureScaleVignette : TEXCOORD1,
 	out noperspective 	float4 OutGrainUV 				: TEXCOORD2,
 	out noperspective 	float2 OutScreenPos 			: TEXCOORD3,
 	out noperspective 	float2 OutFullViewUV 			: TEXCOORD4,
@@ -405,7 +413,7 @@ float4 SampleSceneColor(float2 SceneUV)
 
 float4 TonemapCommonPS(
 	float2 UV,
-	float3 ExposureScaleVignette,
+	float4 ExposureScaleVignette,
 	float4 GrainUV,
 	float2 ScreenPos, // [-1, 1]x[-1, 1]
 	float2 FullViewUV,
@@ -470,6 +478,18 @@ float4 TonemapCommonPS(
 
 	float ExposureScale = ExposureScaleVignette.x;
 
+#if USE_LOCAL_EXPOSURE
+	float LuminanceVal = CalculateEyeAdaptationLuminance(SceneColor.rgb);
+	float LogLuminance = log2(LuminanceVal);
+	float MiddleGreyLumValue = ExposureScaleVignette.y;
+
+	float BaseLogLum = CalculateBaseLogLuminance(LogLuminance, EyeAdaptation_LocalExposureBlurredLuminanceBlend, ExposureScale, FullViewUV, LumBilateralGrid, BlurredLogLum, LumBilateralGridSampler, BlurredLogLumSampler);
+	float LocalExposure = CalculateLocalExposure(LogLuminance + log2(ExposureScale), BaseLogLum, MiddleGreyLumValue, EyeAdaptation_LocalExposureContrastReduction, EyeAdaptation_LocalExposureDetailStrength);
+
+	// Apply before bloom
+	SceneColor.rgb *= LocalExposure;
+#endif
+
 #if DEBUG_TONEMAPPER
 	if (all(DebugTile == int2(-1, 1)))
 	{
@@ -637,7 +657,7 @@ float4 TonemapCommonPS(
 		LinearColor *= ExposureScale;
 
 		#if USE_VIGNETTE
-			LinearColor.rgb *= ComputeVignetteMask( ExposureScaleVignette.yz, TonemapperParams.x );
+			LinearColor.rgb *= ComputeVignetteMask( ExposureScaleVignette.zw, TonemapperParams.x );
 		#endif
 
 		#if USE_FILM_GRAIN
@@ -738,7 +758,7 @@ float4 TonemapCommonPS(
 // pixel shader entry point
 void MainPS(
 	in noperspective float2 UV : TEXCOORD0,
-	in noperspective float3 InExposureScaleVignette : TEXCOORD1,
+	in noperspective float4 InExposureScaleVignette : TEXCOORD1,
 	in noperspective float4 GrainUV : TEXCOORD2,
 	in noperspective float2 ScreenPos : TEXCOORD3,
 	in noperspective float2 FullViewUV : TEXCOORD4,
@@ -764,7 +784,7 @@ void MainCS(uint2 DispatchThreadId : SV_DispatchThreadID)
 		return;
 	}
 
-	float3 ExposureScaleVignette;
+	float4 ExposureScaleVignette;
 	float4 GrainUV;
 	float2 FullViewUV, ScreenPos;	
 	TonemapCommonVS(InScreenPos, UV, ExposureScaleVignette, GrainUV, ScreenPos, FullViewUV);
diff --git a/Engine/Shaders/Private/PostProcessVisualizeLocalExposure.usf b/Engine/Shaders/Private/PostProcessVisualizeLocalExposure.usf
new file mode 100644
--- /dev/null
+++ b/Engine/Shaders/Private/PostProcessVisualizeLocalExposure.usf
@@ -0,0 +1,57 @@
+// Copyright Epic Games, Inc. All Rights Reserved.
+
+#include "Common.ush"
+#include "ScreenPass.ush"
+#include "PostProcessCommon.ush"
+#include "PostProcessHistogramCommon.ush"
+#include "DeferredShadingCommon.ush"
+#include "TonemapCommon.ush"
+
+Texture2D EyeAdaptationTexture;
+
+Texture2D HDRSceneColorTexture;
+
+Texture3D LumBilateralGrid;
+Texture2D BlurredLogLum;
+
+SamplerState TextureSampler;
+
+uint DebugMode;
+
+SCREEN_PASS_TEXTURE_VIEWPORT(Input)
+SCREEN_PASS_TEXTURE_VIEWPORT(Output)
+
+float4 MainPS(noperspective float4 UVAndScreenPos : TEXCOORD0, float4 SvPosition : SV_POSITION) : SV_Target0
+{
+	float2 UV = UVAndScreenPos.xy;
+	int2 PixelPos = (int2)SvPosition.xy;
+
+	float2 ExposureScaleMiddleGrey = EyeAdaptationTexture.Load(int3(0, 0, 0)).xw;
+
+	float4 SceneColor = Texture2DSample(HDRSceneColorTexture, TextureSampler, UV) * View.OneOverPreExposure;
+
+	float LuminanceVal = CalculateEyeAdaptationLuminance(SceneColor.rgb);
+	float LogLuminance = log2(LuminanceVal);
+	float MiddleGreyLumValue = log2(0.18 * ExposureScaleMiddleGrey.y * EyeAdaptation_LocalExposureMiddleGreyExposureCompensation);
+
+	float BaseLogLum = CalculateBaseLogLuminance(LogLuminance, EyeAdaptation_LocalExposureBlurredLuminanceBlend, ExposureScaleMiddleGrey.x, UV, LumBilateralGrid, BlurredLogLum, TextureSampler, TextureSampler);
+	float LocalExposure = CalculateLocalExposure(LogLuminance + log2(ExposureScaleMiddleGrey.x), BaseLogLum, MiddleGreyLumValue, EyeAdaptation_LocalExposureContrastReduction, EyeAdaptation_LocalExposureDetailStrength);
+	float DetailLogLum = (LogLuminance + log2(ExposureScaleMiddleGrey.x)) - BaseLogLum;
+
+	float3 OutValue;
+
+	if (DebugMode == 0)
+	{
+		OutValue = saturate(abs(log2(LocalExposure) / 2)) * lerp(float3(1.0f, 0.0f, 0.0f), float3(0.0f, 1.0f, 0.0f), step(LocalExposure, 1));
+	}
+	else if (DebugMode == 1)
+	{
+		OutValue = exp2(BaseLogLum);
+	}
+	else if (DebugMode == 2)
+	{
+		OutValue = exp2(DetailLogLum);
+	}
+
+	return float4(OutValue, 1.0f);
+}
\ No newline at end of file
diff --git a/Engine/Source/Editor/DetailCustomizations/Private/SceneCaptureDetails.cpp b/Engine/Source/Editor/DetailCustomizations/Private/SceneCaptureDetails.cpp
--- a/Engine/Source/Editor/DetailCustomizations/Private/SceneCaptureDetails.cpp
+++ b/Engine/Source/Editor/DetailCustomizations/Private/SceneCaptureDetails.cpp
@@ -99,6 +99,7 @@ void FSceneCaptureDetails::CustomizeDetails( IDetailLayoutBuilder& DetailLayout
 	ShowFlagsToAllowForCaptures.Add(FEngineShowFlags::EShowFlag::SF_TemporalAA);
 	ShowFlagsToAllowForCaptures.Add(FEngineShowFlags::EShowFlag::SF_MotionBlur);
 	ShowFlagsToAllowForCaptures.Add(FEngineShowFlags::EShowFlag::SF_Bloom);
+	ShowFlagsToAllowForCaptures.Add(FEngineShowFlags::EShowFlag::SF_LocalExposure);
 	ShowFlagsToAllowForCaptures.Add(FEngineShowFlags::EShowFlag::SF_EyeAdaptation);
 	ShowFlagsToAllowForCaptures.Add(FEngineShowFlags::EShowFlag::SF_Game);
 	ShowFlagsToAllowForCaptures.Add(FEngineShowFlags::EShowFlag::SF_ToneCurve); 
diff --git a/Engine/Source/Runtime/Core/Private/HAL/ConsoleManager.cpp b/Engine/Source/Runtime/Core/Private/HAL/ConsoleManager.cpp
--- a/Engine/Source/Runtime/Core/Private/HAL/ConsoleManager.cpp
+++ b/Engine/Source/Runtime/Core/Private/HAL/ConsoleManager.cpp
@@ -2246,6 +2246,12 @@ static TAutoConsoleVariable<int32> CVarSetClearSceneMethod(
 	TEXT(" 2: Quad at max z"),
 	ECVF_RenderThreadSafe);
 
+static TAutoConsoleVariable<int32> CVarLocalExposure(
+	TEXT("r.LocalExposure"),
+	1,
+	TEXT("Whether to support local exposure"),
+	ECVF_Scalability | ECVF_RenderThreadSafe);
+
 static TAutoConsoleVariable<int32> CVarBloomQuality(
 	TEXT("r.BloomQuality"),
 	5,
@@ -2262,7 +2268,7 @@ static TAutoConsoleVariable<int32> CVarSceneColorFringeQuality(
 	TEXT("r.SceneColorFringeQuality"),
 	1,
 	TEXT(" 0: off but best for performance\n")
-	TEXT(" 1: 3 texture samples (default)n"),
+	TEXT(" 1: 3 texture samples (default)"),
 	ECVF_Scalability | ECVF_RenderThreadSafe);
 
 
diff --git a/Engine/Source/Runtime/Engine/Classes/Engine/Scene.h b/Engine/Source/Runtime/Engine/Classes/Engine/Scene.h
--- a/Engine/Source/Runtime/Engine/Classes/Engine/Scene.h
+++ b/Engine/Source/Runtime/Engine/Classes/Engine/Scene.h
@@ -926,6 +926,21 @@ struct FPostProcessSettings
 	UPROPERTY(EditAnywhere, BlueprintReadWrite, Category=Overrides, meta=(PinHiddenByDefault, InlineEditConditionToggle))
 	uint8 bOverride_HistogramLogMax:1;
 
+	UPROPERTY(EditAnywhere, BlueprintReadWrite, Category=Overrides, meta=(PinHiddenByDefault, InlineEditConditionToggle))
+	uint8 bOverride_LocalExposureContrastReduction:1;
+
+	UPROPERTY(EditAnywhere, BlueprintReadWrite, Category=Overrides, meta=(PinHiddenByDefault, InlineEditConditionToggle))
+	uint8 bOverride_LocalExposureDetailStrength:1;
+
+	UPROPERTY(EditAnywhere, BlueprintReadWrite, Category=Overrides, meta=(PinHiddenByDefault, InlineEditConditionToggle))
+	uint8 bOverride_LocalExposureBlurredLuminanceBlend:1;
+
+	UPROPERTY(EditAnywhere, BlueprintReadWrite, Category=Overrides, meta=(PinHiddenByDefault, InlineEditConditionToggle))
+	uint8 bOverride_LocalExposureBlurredLuminanceKernelSizePercent:1;
+
+	UPROPERTY(EditAnywhere, BlueprintReadWrite, Category=Overrides, meta=(PinHiddenByDefault, InlineEditConditionToggle))
+	uint8 bOverride_LocalExposureMiddleGreyBias:1;
+
 	UPROPERTY(EditAnywhere, BlueprintReadWrite, Category=Overrides, meta=(PinHiddenByDefault, InlineEditConditionToggle))
 	uint8 bOverride_LensFlareIntensity:1;
 
@@ -1723,6 +1738,45 @@ struct FPostProcessSettings
 	UPROPERTY()
 	float AutoExposureCalibrationConstant_DEPRECATED;
 
+	/** 
+	 * Local Exposure decomposes luminance of the frame into a base layer and a detail layer.
+	 * Contrast of the base layer is reduced based on this value.
+	 * Value less than 1 will enable local exposure.
+	 * Good values are usually in the range 0.6 .. 1.0.
+	*/
+	UPROPERTY(interp, BlueprintReadWrite, Category="Lens|Local Exposure", meta=(UIMin = "0.0", UIMax = "1.0", editcondition = "bOverride_LocalExposureContrastReduction", DisplayName = "Contrast Reduction"))
+	float LocalExposureContrastReduction;
+
+	/**
+	 * Local Exposure decomposes luminance of the frame into a base layer and a detail layer.
+	 * Value different than 1 will enable local exposure.
+	 * This value should be set to 1 in most cases.
+	*/
+	UPROPERTY(interp, BlueprintReadWrite, Category="Lens|Local Exposure", meta=(UIMin = "0.0", UIMax = "4.0", editcondition = "bOverride_LocalExposureDetailStrength", DisplayName = "Detail Strength"))
+	float LocalExposureDetailStrength;
+
+	/**
+	 * Local Exposure decomposes luminance of the frame into a base layer and a detail layer.
+	 * Blend between bilateral filtered and blurred luminance as the base layer.
+	 * Blurred luminance helps preserve image appearance and specular highlights, and reduce ringing.
+	 * Good values are usually in the range 0.4 .. 0.6
+	*/
+	UPROPERTY(interp, BlueprintReadWrite, Category="Lens|Local Exposure", meta=(UIMin = "0.0", UIMax = "1.0", editcondition = "bOverride_LocalExposureBlurredLuminanceBlend", DisplayName = "Blurred Luminance Blend"))
+	float LocalExposureBlurredLuminanceBlend;
+
+	/**
+	 * Kernel size (percentage of screen) used to blur frame luminance.
+	*/
+	UPROPERTY(interp, BlueprintReadWrite, Category="Lens|Local Exposure", meta=(UIMin = "0.0", UIMax = "100.0", editcondition = "bOverride_LocalExposureBlurredLuminanceKernelSizePercent", DisplayName = "Blurred Luminance Kernel Size Percent"))
+	float LocalExposureBlurredLuminanceKernelSizePercent;
+
+	/**
+	 * Logarithmic adjustment for the local exposure middle grey.
+	 * 0: no adjustment, -1:2x darker, -2:4x darker, 1:2x brighter, 2:4x brighter, ...
+	 */
+	UPROPERTY(interp, BlueprintReadWrite, Category = "Lens|Local Exposure", meta = (UIMin = "-15.0", UIMax = "15.0", editcondition = "bOverride_LocalExposureMiddleGreyBias", DisplayName = "Middle Grey Bias"))
+	float LocalExposureMiddleGreyBias;
+
 	/** Brightness scale of the image cased lens flares (linear) */
 	UPROPERTY(interp, BlueprintReadWrite, Category="Lens|Lens Flares", meta=(UIMin = "0.0", UIMax = "16.0", editcondition = "bOverride_LensFlareIntensity", DisplayName = "Intensity"))
 	float LensFlareIntensity;
diff --git a/Engine/Source/Runtime/Engine/Private/Scene.cpp b/Engine/Source/Runtime/Engine/Private/Scene.cpp
--- a/Engine/Source/Runtime/Engine/Private/Scene.cpp
+++ b/Engine/Source/Runtime/Engine/Private/Scene.cpp
@@ -501,6 +501,12 @@ FPostProcessSettings::FPostProcessSettings()
 
 	AutoExposureApplyPhysicalCameraExposure = 1;
 
+	LocalExposureContrastReduction = 1.0f;
+	LocalExposureDetailStrength = 1.0f;
+	LocalExposureBlurredLuminanceBlend = 0.6f;
+	LocalExposureBlurredLuminanceKernelSizePercent = 50.0f;
+	LocalExposureMiddleGreyBias = 0.0f;
+
 	// next value might get overwritten by r.DefaultFeature.LensFlare
 	LensFlareIntensity = 1.0f;
 	LensFlareTint = FLinearColor(1.0f, 1.0f, 1.0f);
@@ -713,6 +719,11 @@ FPostProcessSettings::FPostProcessSettings(const FPostProcessSettings& Settings)
 	, bOverride_AutoExposureApplyPhysicalCameraExposure(Settings.bOverride_AutoExposureApplyPhysicalCameraExposure)
 	, bOverride_HistogramLogMin(Settings.bOverride_HistogramLogMin)
 	, bOverride_HistogramLogMax(Settings.bOverride_HistogramLogMax)
+	, bOverride_LocalExposureContrastReduction(Settings.bOverride_LocalExposureContrastReduction)
+	, bOverride_LocalExposureDetailStrength(Settings.bOverride_LocalExposureDetailStrength)
+	, bOverride_LocalExposureBlurredLuminanceBlend(Settings.bOverride_LocalExposureBlurredLuminanceBlend)
+	, bOverride_LocalExposureBlurredLuminanceKernelSizePercent(Settings.bOverride_LocalExposureBlurredLuminanceKernelSizePercent)
+	, bOverride_LocalExposureMiddleGreyBias(Settings.bOverride_LocalExposureMiddleGreyBias)
 	, bOverride_LensFlareIntensity(Settings.bOverride_LensFlareIntensity)
 	, bOverride_LensFlareTint(Settings.bOverride_LensFlareTint)
 	, bOverride_LensFlareTints(Settings.bOverride_LensFlareTints)
@@ -918,6 +929,11 @@ FPostProcessSettings::FPostProcessSettings(const FPostProcessSettings& Settings)
 	, HistogramLogMin(Settings.HistogramLogMin)
 	, HistogramLogMax(Settings.HistogramLogMax)
 	, AutoExposureCalibrationConstant_DEPRECATED(Settings.AutoExposureCalibrationConstant_DEPRECATED)
+	, LocalExposureContrastReduction(Settings.LocalExposureContrastReduction)
+	, LocalExposureDetailStrength(Settings.LocalExposureDetailStrength)
+	, LocalExposureBlurredLuminanceBlend(Settings.LocalExposureBlurredLuminanceBlend)
+	, LocalExposureBlurredLuminanceKernelSizePercent(Settings.LocalExposureBlurredLuminanceKernelSizePercent)
+	, LocalExposureMiddleGreyBias(Settings.LocalExposureMiddleGreyBias)
 	, LensFlareIntensity(Settings.LensFlareIntensity)
 	, LensFlareTint(Settings.LensFlareTint)
 	, LensFlareBokehSize(Settings.LensFlareBokehSize)
diff --git a/Engine/Source/Runtime/Engine/Private/SceneView.cpp b/Engine/Source/Runtime/Engine/Private/SceneView.cpp
--- a/Engine/Source/Runtime/Engine/Private/SceneView.cpp
+++ b/Engine/Source/Runtime/Engine/Private/SceneView.cpp
@@ -1436,6 +1436,11 @@ void FSceneView::OverridePostProcessSettings(const FPostProcessSettings& Src, fl
 		LERP_PP(AutoExposureBias);
 		LERP_PP(HistogramLogMin);
 		LERP_PP(HistogramLogMax);
+		LERP_PP(LocalExposureContrastReduction);
+		LERP_PP(LocalExposureDetailStrength);
+		LERP_PP(LocalExposureBlurredLuminanceBlend);
+		LERP_PP(LocalExposureBlurredLuminanceKernelSizePercent);
+		LERP_PP(LocalExposureMiddleGreyBias);
 		LERP_PP(LensFlareIntensity);
 		LERP_PP(LensFlareTint);
 		LERP_PP(LensFlareBokehSize);
@@ -1865,6 +1870,18 @@ void FSceneView::EndFinalPostprocessSettings(const FSceneViewInitOptions& ViewIn
 		}
 	}
 
+	{
+		static const auto LocalExposureCVar = IConsoleManager::Get().FindTConsoleVariableDataInt(TEXT("r.LocalExposure"));
+
+		int Value = LocalExposureCVar->GetValueOnGameThread();
+
+		if (Value <= 0)
+		{
+			FinalPostProcessSettings.LocalExposureContrastReduction = 1.0f;
+			FinalPostProcessSettings.LocalExposureDetailStrength = 1.0f;
+		}
+	}
+
 	{
 		static const auto BloomQualityCVar = IConsoleManager::Get().FindTConsoleVariableDataInt(TEXT("r.BloomQuality"));
 
@@ -1881,6 +1898,12 @@ void FSceneView::EndFinalPostprocessSettings(const FSceneViewInitOptions& ViewIn
 		FinalPostProcessSettings.BloomIntensity = 0.0f;
 	}
 
+	if (!Family->EngineShowFlags.LocalExposure)
+	{
+		FinalPostProcessSettings.LocalExposureContrastReduction = 1.0f;
+		FinalPostProcessSettings.LocalExposureDetailStrength = 1.0f;
+	}
+
 	// scale down tone mapper shader permutation
 	{
 		int32 Quality = CVarTonemapperQuality.GetValueOnGameThread();
diff --git a/Engine/Source/Runtime/Engine/Private/ShowFlags.cpp b/Engine/Source/Runtime/Engine/Private/ShowFlags.cpp
--- a/Engine/Source/Runtime/Engine/Private/ShowFlags.cpp
+++ b/Engine/Source/Runtime/Engine/Private/ShowFlags.cpp
@@ -540,6 +540,7 @@ void EngineShowFlagOverride(EShowFlagInitMode ShowFlagInitMode, EViewModeIndex V
 		{
 			EngineShowFlags.SetRayTracingDebug(true);
 			EngineShowFlags.SetVisualizeHDR(false);
+			EngineShowFlags.SetVisualizeLocalExposure(false);
 			EngineShowFlags.SetVisualizeMotionBlur(false);
 			EngineShowFlags.SetDepthOfField(false);
 			EngineShowFlags.SetPostProcessMaterial(false);
diff --git a/Engine/Source/Runtime/Engine/Public/ShowFlags.h b/Engine/Source/Runtime/Engine/Public/ShowFlags.h
--- a/Engine/Source/Runtime/Engine/Public/ShowFlags.h
+++ b/Engine/Source/Runtime/Engine/Public/ShowFlags.h
@@ -365,6 +365,7 @@ private:
 
 		// The following code sets what should be off by default.
 		SetVisualizeHDR(false);
+		SetVisualizeLocalExposure(false);
 		SetVisualizeShadingModels(false);
 		SetOverrideDiffuseAndSpecular(false);
 		SetLightingOnlyOverride(false);
diff --git a/Engine/Source/Runtime/Engine/Public/ShowFlagsValues.inl b/Engine/Source/Runtime/Engine/Public/ShowFlagsValues.inl
--- a/Engine/Source/Runtime/Engine/Public/ShowFlagsValues.inl
+++ b/Engine/Source/Runtime/Engine/Public/ShowFlagsValues.inl
@@ -22,6 +22,8 @@
 SHOWFLAG_ALWAYS_ACCESSIBLE(PostProcessing, SFG_Hidden, NSLOCTEXT("UnrealEd", "PostProcessingSF", "Post-processing"))
 /** Bloom, for now SHOWFLAG_ALWAYS_ACCESSIBLE because it's exposed in SceneCapture */
 SHOWFLAG_ALWAYS_ACCESSIBLE(Bloom, SFG_PostProcess, NSLOCTEXT("UnrealEd", "BloomSF", "Bloom"))
+/** Local Exposure, for now SHOWFLAG_ALWAYS_ACCESSIBLE because it's exposed in SceneCapture */
+SHOWFLAG_ALWAYS_ACCESSIBLE(LocalExposure, SFG_PostProcess, NSLOCTEXT("UnrealEd", "LocalExposureSF", "Local Exposure"))
 /** HDR->LDR conversion is done through a tone mapper (otherwise linear mapping is used) */
 SHOWFLAG_FIXED_IN_SHIPPING(1, Tonemapper, SFG_PostProcess, NSLOCTEXT("UnrealEd", "TonemapperSF", "Tonemapper"))
 /** Any Anti-aliasing e.g. FXAA, Temporal AA, for now SHOWFLAG_ALWAYS_ACCESSIBLE because it's exposed in SceneCapture */
@@ -34,6 +36,8 @@ SHOWFLAG_ALWAYS_ACCESSIBLE(AmbientCubemap, SFG_LightingFeatures, NSLOCTEXT("Unre
 SHOWFLAG_ALWAYS_ACCESSIBLE(EyeAdaptation, SFG_PostProcess, NSLOCTEXT("UnrealEd", "EyeAdaptationSF", "Eye Adaptation"))
 /** Display a histogram of the scene HDR color */
 SHOWFLAG_FIXED_IN_SHIPPING(0, VisualizeHDR, SFG_Visualize, NSLOCTEXT("UnrealEd", "VisualizeHDRSF", "HDR (Eye Adaptation)"))
+/** Helper to tweak local expsoure settings */
+SHOWFLAG_FIXED_IN_SHIPPING(0, VisualizeLocalExposure, SFG_Visualize, NSLOCTEXT("UnrealEd", "VisualizeLocalExposureSF", "Local Exposure"))
 /** Image based lens flares (Simulate artifact of reflections within a camera system) */
 SHOWFLAG_FIXED_IN_SHIPPING(1, LensFlares, SFG_PostProcess, NSLOCTEXT("UnrealEd", "LensFlaresSF", "Lens Flares"))
 /** show indirect lighting component, for now SHOWFLAG_ALWAYS_ACCESSIBLE because it's needed by r.GBuffer */
diff --git a/Engine/Source/Runtime/Renderer/Private/PostProcess/PostProcessEyeAdaptation.cpp b/Engine/Source/Runtime/Renderer/Private/PostProcess/PostProcessEyeAdaptation.cpp
--- a/Engine/Source/Runtime/Renderer/Private/PostProcess/PostProcessEyeAdaptation.cpp
+++ b/Engine/Source/Runtime/Renderer/Private/PostProcess/PostProcessEyeAdaptation.cpp
@@ -278,6 +278,14 @@ FEyeAdaptationParameters GetEyeAdaptationParameters(const FViewInfo& View, ERHIF
 	float ExposureCompensationCurve = GetAutoExposureCompensationFromCurve(View);
 	const float BlackHistogramBucketInfluence = CVarEyeAdaptationBlackHistogramBucketInfluence.GetValueOnRenderThread();
 
+	float LocalExposureMiddleGreyExposureCompensation = FMath::Pow(2.0f, View.FinalPostProcessSettings.LocalExposureMiddleGreyBias);
+
+	if (AutoExposureMethod == EAutoExposureMethod::AEM_Manual)
+	{
+		// when using manual exposure cancel exposure compensation setting and curve from middle grey used by local exposure.
+		LocalExposureMiddleGreyExposureCompensation /= (ExposureCompensationSettings * ExposureCompensationCurve);
+	}
+
 	const float kMiddleGrey = 0.18f;
 
 	// AEM_Histogram and AEM_Basic adjust their ExposureCompensation to middle grey (0.18). AEM_Manual ExposureCompensation is already calibrated to 1.0.
@@ -393,6 +401,10 @@ FEyeAdaptationParameters GetEyeAdaptationParameters(const FViewInfo& View, ERHIF
 	Parameters.HistogramScale = HistogramScale;
 	Parameters.HistogramBias = HistogramBias;
 	Parameters.LuminanceMin = LuminanceMin;
+	Parameters.LocalExposureContrastReduction = Settings.LocalExposureContrastReduction;
+	Parameters.LocalExposureDetailStrength = Settings.LocalExposureDetailStrength;
+	Parameters.LocalExposureBlurredLuminanceBlend = Settings.LocalExposureBlurredLuminanceBlend;
+	Parameters.LocalExposureMiddleGreyExposureCompensation = LocalExposureMiddleGreyExposureCompensation;
 	Parameters.BlackHistogramBucketInfluence = BlackHistogramBucketInfluence; // no calibration constant because it is now baked into ExposureCompensation
 	Parameters.GreyMult = GreyMult;
 	Parameters.ExponentialDownM = ExponentialDownM;
diff --git a/Engine/Source/Runtime/Renderer/Private/PostProcess/PostProcessEyeAdaptation.h b/Engine/Source/Runtime/Renderer/Private/PostProcess/PostProcessEyeAdaptation.h
--- a/Engine/Source/Runtime/Renderer/Private/PostProcess/PostProcessEyeAdaptation.h
+++ b/Engine/Source/Runtime/Renderer/Private/PostProcess/PostProcessEyeAdaptation.h
@@ -65,6 +65,10 @@ BEGIN_SHADER_PARAMETER_STRUCT(FEyeAdaptationParameters, )
 	SHADER_PARAMETER(float, HistogramScale)
 	SHADER_PARAMETER(float, HistogramBias)
 	SHADER_PARAMETER(float, LuminanceMin)
+	SHADER_PARAMETER(float, LocalExposureContrastReduction)
+	SHADER_PARAMETER(float, LocalExposureDetailStrength)
+	SHADER_PARAMETER(float, LocalExposureBlurredLuminanceBlend)
+	SHADER_PARAMETER(float, LocalExposureMiddleGreyExposureCompensation)
 	SHADER_PARAMETER(float, BlackHistogramBucketInfluence)
 	SHADER_PARAMETER(float, GreyMult)
 	SHADER_PARAMETER(float, ExponentialUpM)
diff --git a/Engine/Source/Runtime/Renderer/Private/PostProcess/PostProcessHistogram.cpp b/Engine/Source/Runtime/Renderer/Private/PostProcess/PostProcessHistogram.cpp
--- a/Engine/Source/Runtime/Renderer/Private/PostProcess/PostProcessHistogram.cpp
+++ b/Engine/Source/Runtime/Renderer/Private/PostProcess/PostProcessHistogram.cpp
@@ -16,15 +16,13 @@ TAutoConsoleVariable<int32> CVarUseAtomicHistogram(
 namespace
 {
 
-class FHistogramCS : public FGlobalShader
+template <uint32 ThreadGroupSizeX, uint32 ThreadGroupSizeY, uint32 HistogramSize, bool BilateralGrid>
+class THistogramCS : public FGlobalShader
 {
 public:
 	// Changing these numbers requires Histogram.usf to be recompiled.
-	static const uint32 ThreadGroupSizeX = 8;
-	static const uint32 ThreadGroupSizeY = 4;
 	static const uint32 LoopCountX = 8;
 	static const uint32 LoopCountY = 8;
-	static const uint32 HistogramSize = 64;
 
 	// /4 as we store 4 buckets in one ARGB texel.
 	static const uint32 HistogramTexelCount = HistogramSize / 4;
@@ -32,8 +30,8 @@ public:
 	// The number of texels on each axis processed by a single thread group.
 	static const FIntPoint TexelsPerThreadGroup;
 
-	DECLARE_GLOBAL_SHADER(FHistogramCS);
-	SHADER_USE_PARAMETER_STRUCT(FHistogramCS, FGlobalShader);
+	DECLARE_GLOBAL_SHADER(THistogramCS);
+	SHADER_USE_PARAMETER_STRUCT(THistogramCS, FGlobalShader);
 
 	BEGIN_SHADER_PARAMETER_STRUCT(FParameters, )
 		SHADER_PARAMETER_STRUCT_REF(FViewUniformShaderParameters, View)
@@ -41,6 +39,7 @@ public:
 		SHADER_PARAMETER_STRUCT(FEyeAdaptationParameters, EyeAdaptation)
 		SHADER_PARAMETER_RDG_TEXTURE(Texture2D, InputTexture)
 		SHADER_PARAMETER_RDG_TEXTURE_UAV(RWTexture2D<float4>, HistogramRWTexture)
+		SHADER_PARAMETER_RDG_TEXTURE_UAV(RWTexture3D<float2>, BilateralGridRWTexture)
 		SHADER_PARAMETER(FIntPoint, ThreadGroupCount)
 	END_SHADER_PARAMETER_STRUCT()
 
@@ -57,6 +56,7 @@ public:
 		OutEnvironment.SetDefine(TEXT("LOOP_SIZEX"), LoopCountX);
 		OutEnvironment.SetDefine(TEXT("LOOP_SIZEY"), LoopCountY);
 		OutEnvironment.SetDefine(TEXT("HISTOGRAM_SIZE"), HistogramSize);
+		OutEnvironment.SetDefine(TEXT("BILATERAL_GRID"), BilateralGrid);
 		OutEnvironment.CompilerFlags.Add( CFLAG_StandardOptimization );
 	}
 
@@ -174,9 +174,17 @@ public:
 	}
 };
 
-const FIntPoint FHistogramCS::TexelsPerThreadGroup(ThreadGroupSizeX* LoopCountX, ThreadGroupSizeY* LoopCountY);
+typedef THistogramCS<8, 4, 64, false> FHistogramCS;
+IMPLEMENT_SHADER_TYPE(template<>, FHistogramCS, TEXT("/Engine/Private/PostProcessHistogram.usf"), TEXT("MainCS"), SF_Compute);
+template<>
+const FIntPoint FHistogramCS::TexelsPerThreadGroup(8 * LoopCountX, 4 * LoopCountY);
+
+static const uint32 BilateralGridDepth = 32;
+typedef THistogramCS<8, 8, BilateralGridDepth, true> FHistogramBilateralGridCS;
+IMPLEMENT_SHADER_TYPE(template<>, FHistogramBilateralGridCS, TEXT("/Engine/Private/PostProcessHistogram.usf"), TEXT("MainCS"), SF_Compute);
+template<>
+const FIntPoint FHistogramBilateralGridCS::TexelsPerThreadGroup(8 * LoopCountX, 8 * LoopCountY);
 
-IMPLEMENT_GLOBAL_SHADER(FHistogramCS, "/Engine/Private/PostProcessHistogram.usf", "MainCS", SF_Compute);
 IMPLEMENT_GLOBAL_SHADER(FHistogramReducePS, "/Engine/Private/PostProcessHistogramReduce.usf", "MainPS", SF_Pixel);
 
 IMPLEMENT_GLOBAL_SHADER(FHistogramAtomicCS,        "/Engine/Private/Histogram.usf", "MainAtomicCS", SF_Compute);
@@ -404,7 +412,52 @@ FRDGTextureRef AddHistogramPass(
 	}
 }
 
-FIntPoint GetHistogramTexelsPerGroup()
+FRDGTextureRef AddLocalExposurePass(
+	FRDGBuilder& GraphBuilder,
+	const FViewInfo& View,
+	const FEyeAdaptationParameters& EyeAdaptationParameters,
+	FScreenPassTexture SceneColor)
 {
-	return FHistogramCS::TexelsPerThreadGroup;
+	check(SceneColor.IsValid());
+
+	const FIntPoint TileSize(64, 64);
+	const FIntPoint ThreadGroupSize(8, 8);
+	const FIntPoint NumTiles = FIntPoint::DivideAndRoundUp(SceneColor.ViewRect.Size(), FHistogramBilateralGridCS::TexelsPerThreadGroup);
+	const FIntPoint ThreadGroupCount = FIntPoint::DivideAndRoundUp(SceneColor.ViewRect.Size(), FHistogramBilateralGridCS::TexelsPerThreadGroup);
+
+	FRDGTextureRef LocalExposureTexture = nullptr;
+
+	RDG_EVENT_SCOPE(GraphBuilder, "LocalExposure");
+
+	{
+		const FIntVector TextureExtent = FIntVector(NumTiles.X, NumTiles.Y, BilateralGridDepth);
+
+		const FRDGTextureDesc TextureDesc = FRDGTextureDesc::Create3D(
+			TextureExtent,
+			PF_G32R32F,
+			FClearValueBinding::None,
+			TexCreate_UAV | TexCreate_ShaderResource);
+
+		LocalExposureTexture = GraphBuilder.CreateTexture(TextureDesc, TEXT("LocalExposure"));
+
+		auto* PassParameters = GraphBuilder.AllocParameters<FHistogramBilateralGridCS::FParameters>();
+		PassParameters->View = View.ViewUniformBuffer;
+		PassParameters->EyeAdaptation = EyeAdaptationParameters;
+		PassParameters->Input = GetScreenPassTextureViewportParameters(FScreenPassTextureViewport(SceneColor));
+		PassParameters->InputTexture = SceneColor.Texture;
+		PassParameters->BilateralGridRWTexture = GraphBuilder.CreateUAV(LocalExposureTexture);
+		PassParameters->ThreadGroupCount = ThreadGroupCount;
+
+		auto ComputeShader = View.ShaderMap->GetShader<FHistogramBilateralGridCS>();
+
+		FComputeShaderUtils::AddPass(
+			GraphBuilder,
+			RDG_EVENT_NAME("FLocalExposure %dx%d (CS)", SceneColor.ViewRect.Width(), SceneColor.ViewRect.Height()),
+			ERDGPassFlags::Compute | ERDGPassFlags::NeverCull,
+			ComputeShader,
+			PassParameters,
+			FIntVector(ThreadGroupCount.X, ThreadGroupCount.Y, 1));
+	}
+
+	return LocalExposureTexture;
 }
\ No newline at end of file
diff --git a/Engine/Source/Runtime/Renderer/Private/PostProcess/PostProcessHistogram.h b/Engine/Source/Runtime/Renderer/Private/PostProcess/PostProcessHistogram.h
--- a/Engine/Source/Runtime/Renderer/Private/PostProcess/PostProcessHistogram.h
+++ b/Engine/Source/Runtime/Renderer/Private/PostProcess/PostProcessHistogram.h
@@ -11,4 +11,10 @@ FRDGTextureRef AddHistogramPass(
 	const FViewInfo& View,
 	const FEyeAdaptationParameters& EyeAdaptationParameters,
 	FScreenPassTexture SceneColor,
-	FRDGTextureRef EyeAdaptationTexture);
\ No newline at end of file
+	FRDGTextureRef EyeAdaptationTexture);
+
+FRDGTextureRef AddLocalExposurePass(
+	FRDGBuilder& GraphBuilder,
+	const FViewInfo& View,
+	const FEyeAdaptationParameters& EyeAdaptationParameters,
+	FScreenPassTexture SceneColor);
\ No newline at end of file
diff --git a/Engine/Source/Runtime/Renderer/Private/PostProcess/PostProcessLocalExposure.cpp b/Engine/Source/Runtime/Renderer/Private/PostProcess/PostProcessLocalExposure.cpp
new file mode 100644
--- /dev/null
+++ b/Engine/Source/Runtime/Renderer/Private/PostProcess/PostProcessLocalExposure.cpp
@@ -0,0 +1,105 @@
+// Copyright Epic Games, Inc. All Rights Reserved.
+
+/*=============================================================================
+	PostProcessLocalExposure.cpp: Post processing local exposure implementation.
+=============================================================================*/
+
+#include "PostProcess/PostProcessLocalExposure.h"
+#include "PostProcess/PostProcessEyeAdaptation.h"
+#include "PostProcess/PostProcessWeightedSampleSum.h"
+#include "ShaderCompilerCore.h"
+
+namespace
+{
+
+class FSetupLogLuminanceCS : public FGlobalShader
+{
+public:
+	// Changing these numbers requires LocalExposure.usf to be recompiled
+	static const uint32 ThreadGroupSizeX = 8;
+	static const uint32 ThreadGroupSizeY = 8;
+
+	DECLARE_GLOBAL_SHADER(FSetupLogLuminanceCS);
+	SHADER_USE_PARAMETER_STRUCT(FSetupLogLuminanceCS, FGlobalShader);
+
+	BEGIN_SHADER_PARAMETER_STRUCT(FParameters, )
+		SHADER_PARAMETER_STRUCT_REF(FViewUniformShaderParameters, View)
+		SHADER_PARAMETER_STRUCT(FEyeAdaptationParameters, EyeAdaptation)
+		SHADER_PARAMETER_STRUCT(FScreenPassTextureViewportParameters, Input)
+		SHADER_PARAMETER_RDG_TEXTURE(Texture2D, InputTexture)
+		SHADER_PARAMETER_RDG_TEXTURE_UAV(RWTexture2D<float>, Output2DFRW)
+	END_SHADER_PARAMETER_STRUCT()
+
+	static bool ShouldCompilePermutation(const FGlobalShaderPermutationParameters& Parameters)
+	{
+		return IsFeatureLevelSupported(Parameters.Platform, ERHIFeatureLevel::SM5);
+	}
+
+	static void ModifyCompilationEnvironment(const FGlobalShaderPermutationParameters& Parameters, FShaderCompilerEnvironment& OutEnvironment)
+	{
+		FGlobalShader::ModifyCompilationEnvironment(Parameters, OutEnvironment);
+		OutEnvironment.SetDefine(TEXT("THREADGROUP_SIZEX"), ThreadGroupSizeX);
+		OutEnvironment.SetDefine(TEXT("THREADGROUP_SIZEY"), ThreadGroupSizeY);
+		OutEnvironment.CompilerFlags.Add(CFLAG_StandardOptimization);
+	}
+};
+
+IMPLEMENT_GLOBAL_SHADER(FSetupLogLuminanceCS, "/Engine/Private/PostProcessLocalExposure.usf", "SetupLogLuminanceCS", SF_Compute);
+
+} //! namespace
+
+FRDGTextureRef AddLocalExposureBlurredLogLuminancePass(
+	FRDGBuilder& GraphBuilder,
+	const FViewInfo& View,
+	const FEyeAdaptationParameters& EyeAdaptationParameters,
+	FScreenPassTexture InputTexture)
+{
+	check(InputTexture.IsValid());
+
+	RDG_EVENT_SCOPE(GraphBuilder, "LocalExposure - Blurred Luminance");
+
+	FRDGTextureRef GaussianLumSetupTexture;
+
+	// Copy log luminance to temporary texture
+	{
+		const FRDGTextureDesc TextureDesc = FRDGTextureDesc::Create2D(
+			InputTexture.ViewRect.Size(),
+			PF_R16F,
+			FClearValueBinding::None,
+			TexCreate_UAV | TexCreate_ShaderResource);
+
+		GaussianLumSetupTexture = GraphBuilder.CreateTexture(TextureDesc, TEXT("GaussianLumSetupTexture"));
+
+		auto* PassParameters = GraphBuilder.AllocParameters<FSetupLogLuminanceCS::FParameters>();
+		PassParameters->View = View.ViewUniformBuffer;
+		PassParameters->EyeAdaptation = EyeAdaptationParameters;
+		PassParameters->Input = GetScreenPassTextureViewportParameters(FScreenPassTextureViewport(InputTexture));
+		PassParameters->InputTexture = InputTexture.Texture;
+		PassParameters->Output2DFRW = GraphBuilder.CreateUAV(GaussianLumSetupTexture);
+
+		FComputeShaderUtils::AddPass(
+			GraphBuilder,
+			RDG_EVENT_NAME("SetupLogLuminance %dx%d", GaussianLumSetupTexture->Desc.Extent.X, GaussianLumSetupTexture->Desc.Extent.Y),
+			ERDGPassFlags::Compute,
+			View.ShaderMap->GetShader<FSetupLogLuminanceCS>(),
+			PassParameters,
+			FComputeShaderUtils::GetGroupCount(GaussianLumSetupTexture->Desc.Extent, FIntPoint(FSetupLogLuminanceCS::ThreadGroupSizeX, FSetupLogLuminanceCS::ThreadGroupSizeY)));
+	}
+
+	FRDGTextureRef GaussianTexture;
+
+	{
+		FGaussianBlurInputs GaussianBlurInputs;
+		GaussianBlurInputs.NameX = TEXT("LocalExposureGaussianX");
+		GaussianBlurInputs.NameY = TEXT("LocalExposureGaussianY");
+		GaussianBlurInputs.Filter = FScreenPassTexture(GaussianLumSetupTexture, InputTexture.ViewRect);
+		GaussianBlurInputs.TintColor = FLinearColor::White;
+		GaussianBlurInputs.CrossCenterWeight = FVector2D::ZeroVector;
+		GaussianBlurInputs.KernelSizePercent = View.FinalPostProcessSettings.LocalExposureBlurredLuminanceKernelSizePercent;
+		GaussianBlurInputs.UseMirrorAddressMode = true;
+
+		GaussianTexture = AddGaussianBlurPass(GraphBuilder, View, GaussianBlurInputs).Texture;
+	}
+
+	return GaussianTexture;
+}
\ No newline at end of file
diff --git a/Engine/Source/Runtime/Renderer/Private/PostProcess/PostProcessLocalExposure.h b/Engine/Source/Runtime/Renderer/Private/PostProcess/PostProcessLocalExposure.h
new file mode 100644
--- /dev/null
+++ b/Engine/Source/Runtime/Renderer/Private/PostProcess/PostProcessLocalExposure.h
@@ -0,0 +1,13 @@
+// Copyright Epic Games, Inc. All Rights Reserved.
+
+#pragma once
+
+#include "ScreenPass.h"
+
+class FEyeAdaptationParameters;
+
+FRDGTextureRef AddLocalExposureBlurredLogLuminancePass(
+	FRDGBuilder& GraphBuilder,
+	const FViewInfo& View,
+	const FEyeAdaptationParameters& EyeAdaptationParameters,
+	FScreenPassTexture InputTexture);
diff --git a/Engine/Source/Runtime/Renderer/Private/PostProcess/PostProcessTonemap.cpp b/Engine/Source/Runtime/Renderer/Private/PostProcess/PostProcessTonemap.cpp
--- a/Engine/Source/Runtime/Renderer/Private/PostProcess/PostProcessTonemap.cpp
+++ b/Engine/Source/Runtime/Renderer/Private/PostProcess/PostProcessTonemap.cpp
@@ -112,6 +112,7 @@ namespace TonemapperPermutation
 // Shared permutation dimensions between deferred and mobile renderer.
 class FTonemapperBloomDim          : SHADER_PERMUTATION_BOOL("USE_BLOOM");
 class FTonemapperGammaOnlyDim      : SHADER_PERMUTATION_BOOL("USE_GAMMA_ONLY");
+class FTonemapperLocalExposureDim : SHADER_PERMUTATION_BOOL("USE_LOCAL_EXPOSURE");
 class FTonemapperVignetteDim       : SHADER_PERMUTATION_BOOL("USE_VIGNETTE");
 class FTonemapperSharpenDim        : SHADER_PERMUTATION_BOOL("USE_SHARPEN");
 class FTonemapperFilmGrainDim      : SHADER_PERMUTATION_BOOL("USE_FILM_GRAIN");
@@ -122,6 +123,7 @@ class FTonemapperEyeAdaptationDim  : SHADER_PERMUTATION_BOOL("EYEADAPTATION_EXPO
 using FCommonDomain = TShaderPermutationDomain<
 	FTonemapperBloomDim,
 	FTonemapperGammaOnlyDim,
+	FTonemapperLocalExposureDim,
 	FTonemapperVignetteDim,
 	FTonemapperSharpenDim,
 	FTonemapperFilmGrainDim,
@@ -146,6 +148,7 @@ bool ShouldCompileCommonPermutation(const FGlobalShaderPermutationParameters& Pa
 	if (PermutationVector.Get<FTonemapperGammaOnlyDim>())
 	{
 		return !PermutationVector.Get<FTonemapperBloomDim>() &&
+			!PermutationVector.Get<FTonemapperLocalExposureDim>() &&
 			!PermutationVector.Get<FTonemapperVignetteDim>() &&
 			!PermutationVector.Get<FTonemapperSharpenDim>() &&
 			!PermutationVector.Get<FTonemapperFilmGrainDim>() &&
@@ -155,7 +158,7 @@ bool ShouldCompileCommonPermutation(const FGlobalShaderPermutationParameters& Pa
 }
 
 // Common conversion of engine settings into.
-FCommonDomain BuildCommonPermutationDomain(const FViewInfo& View, bool bGammaOnly, bool bSwitchVerticalAxis, bool bMetalMSAAHDRDecode)
+FCommonDomain BuildCommonPermutationDomain(const FViewInfo& View, bool bGammaOnly, bool bLocalExposure, bool bSwitchVerticalAxis, bool bMetalMSAAHDRDecode)
 {
 	const FSceneViewFamily* Family = View.Family;
 
@@ -173,6 +176,7 @@ FCommonDomain BuildCommonPermutationDomain(const FViewInfo& View, bool bGammaOnl
 	const FPostProcessSettings& Settings = View.FinalPostProcessSettings;
 	PermutationVector.Set<FTonemapperVignetteDim>(Settings.VignetteIntensity > 0.0f);
 	PermutationVector.Set<FTonemapperBloomDim>(Settings.BloomIntensity > 0.0);
+	PermutationVector.Set<FTonemapperLocalExposureDim>(bLocalExposure);
 	PermutationVector.Set<FTonemapperFilmGrainDim>(View.FilmGrainTexture != nullptr);
 	PermutationVector.Set<FTonemapperSharpenDim>(CVarTonemapperSharpen.GetValueOnRenderThread() > 0.0f);	
 	PermutationVector.Set<FTonemapperSwitchAxis>(bSwitchVerticalAxis);
@@ -334,6 +338,7 @@ BEGIN_SHADER_PARAMETER_STRUCT(FTonemapParameters, )
 	SHADER_PARAMETER_STRUCT_INCLUDE(FTonemapperOutputDeviceParameters, OutputDevice)
 	SHADER_PARAMETER_STRUCT(FScreenPassTextureViewportParameters, Color)
 	SHADER_PARAMETER_STRUCT(FScreenPassTextureViewportParameters, Output)
+	SHADER_PARAMETER_STRUCT(FEyeAdaptationParameters, EyeAdaptation)
 	SHADER_PARAMETER_RDG_TEXTURE(Texture2D, ColorTexture)
 
 	// Bloom texture
@@ -345,6 +350,11 @@ BEGIN_SHADER_PARAMETER_STRUCT(FTonemapParameters, )
 	SHADER_PARAMETER_SAMPLER(SamplerState, BloomSampler)
 	SHADER_PARAMETER_RDG_BUFFER_SRV(StructuredBuffer<float4>, BloomApplyParameters)
 
+	SHADER_PARAMETER_RDG_TEXTURE(Texture3D, LumBilateralGrid)
+	SHADER_PARAMETER_RDG_TEXTURE(Texture2D, BlurredLogLum)
+	SHADER_PARAMETER_SAMPLER(SamplerState, LumBilateralGridSampler)
+	SHADER_PARAMETER_SAMPLER(SamplerState, BlurredLogLumSampler)
+
 	// SM5 and above use Texture2D for EyeAdaptationTexture
 	SHADER_PARAMETER_RDG_TEXTURE(Texture2D, EyeAdaptationTexture)
 	SHADER_PARAMETER_RDG_TEXTURE(Texture2D, ColorGradingLUT)
@@ -764,7 +774,12 @@ FScreenPassTexture AddTonemapPass(FRDGBuilder& GraphBuilder, const FViewInfo& Vi
 	CommonParameters.Color = GetScreenPassTextureViewportParameters(SceneColorViewport);
 	CommonParameters.Output = GetScreenPassTextureViewportParameters(OutputViewport);
 	CommonParameters.ColorTexture = Inputs.SceneColor.Texture;
+	CommonParameters.LumBilateralGrid = Inputs.LocalExposureTexture;
+	CommonParameters.BlurredLogLum = Inputs.BlurredLogLuminanceTexture;
+	CommonParameters.LumBilateralGridSampler = BilinearClampSampler;
+	CommonParameters.BlurredLogLumSampler = BilinearClampSampler;
 	CommonParameters.EyeAdaptationTexture = Inputs.EyeAdaptationTexture;
+	CommonParameters.EyeAdaptation = *Inputs.EyeAdaptationParameters;
 	CommonParameters.ColorGradingLUT = Inputs.ColorGradingTexture;
 	CommonParameters.BloomDirtMaskTexture = BloomDirtMaskTexture;
 	CommonParameters.ColorSampler = BilinearClampSampler;
@@ -847,7 +862,7 @@ FScreenPassTexture AddTonemapPass(FRDGBuilder& GraphBuilder, const FViewInfo& Vi
 	TonemapperPermutation::FDesktopDomain DesktopPermutationVector;
 
 	{
-		TonemapperPermutation::FCommonDomain CommonDomain = TonemapperPermutation::BuildCommonPermutationDomain(View, Inputs.bGammaOnly, Inputs.bFlipYAxis, Inputs.bMetalMSAAHDRDecode);
+		TonemapperPermutation::FCommonDomain CommonDomain = TonemapperPermutation::BuildCommonPermutationDomain(View, Inputs.bGammaOnly, Inputs.LocalExposureTexture != nullptr, Inputs.bFlipYAxis, Inputs.bMetalMSAAHDRDecode);
 		DesktopPermutationVector.Set<TonemapperPermutation::FCommonDomain>(CommonDomain);
 
 		if (!CommonDomain.Get<TonemapperPermutation::FTonemapperGammaOnlyDim>())
diff --git a/Engine/Source/Runtime/Renderer/Private/PostProcess/PostProcessTonemap.h b/Engine/Source/Runtime/Renderer/Private/PostProcess/PostProcessTonemap.h
--- a/Engine/Source/Runtime/Renderer/Private/PostProcess/PostProcessTonemap.h
+++ b/Engine/Source/Runtime/Renderer/Private/PostProcess/PostProcessTonemap.h
@@ -53,6 +53,15 @@ struct FTonemapInputs
 	// [Required] Filtered bloom texture to composite with tonemapped scene color. This should be transparent black for no bloom.
 	FBloomOutputs Bloom;
 
+	// [Optional] Luminance bilateral grid. If this is null, local exposure is disabled.
+	FRDGTextureRef LocalExposureTexture = nullptr;
+
+	// [Optional] Blurred luminance texture used to calculate local exposure.
+	FRDGTextureRef BlurredLogLuminanceTexture = nullptr;
+
+	// [Required] Eye adaptation parameters.
+	const FEyeAdaptationParameters* EyeAdaptationParameters = nullptr;
+
 	// [Required] Color grading texture used to remap colors.
 	FRDGTextureRef ColorGradingTexture = nullptr;
 
diff --git a/Engine/Source/Runtime/Renderer/Private/PostProcess/PostProcessVisualizeLocalExposure.cpp b/Engine/Source/Runtime/Renderer/Private/PostProcess/PostProcessVisualizeLocalExposure.cpp
new file mode 100644
--- /dev/null
+++ b/Engine/Source/Runtime/Renderer/Private/PostProcess/PostProcessVisualizeLocalExposure.cpp
@@ -0,0 +1,91 @@
+// Copyright Epic Games, Inc. All Rights Reserved.
+
+#include "PostProcess/PostProcessVisualizeLocalExposure.h"
+#include "PostProcess/PostProcessTonemap.h"
+#include "UnrealEngine.h"
+
+TAutoConsoleVariable<int> CVarLocalExposureVisualizeDebugMode(
+	TEXT("r.LocalExposure.VisualizeDebugMode"),
+	0,
+	TEXT("When enabling Show->Visualize->Local Exposure is enabled, this flag controls which mode to use.\n")
+	TEXT("    0: Local Exposure\n")
+	TEXT("    1: Base Luminance\n")
+	TEXT("    2: Detail Luminance\n"),
+	ECVF_RenderThreadSafe);
+
+class FVisualizeLocalExposurePS : public FGlobalShader
+{
+public:
+	DECLARE_GLOBAL_SHADER(FVisualizeLocalExposurePS);
+	SHADER_USE_PARAMETER_STRUCT(FVisualizeLocalExposurePS, FGlobalShader);
+
+	BEGIN_SHADER_PARAMETER_STRUCT(FParameters, )
+		SHADER_PARAMETER_STRUCT_REF(FViewUniformShaderParameters, View)
+		SHADER_PARAMETER_STRUCT(FEyeAdaptationParameters, EyeAdaptation)
+		SHADER_PARAMETER_STRUCT(FScreenPassTextureViewportParameters, Input)
+		SHADER_PARAMETER_STRUCT(FScreenPassTextureViewportParameters, Output)
+
+		SHADER_PARAMETER_RDG_TEXTURE(Texture2D, HDRSceneColorTexture)
+		SHADER_PARAMETER_RDG_TEXTURE(Texture2D, EyeAdaptationTexture)
+		SHADER_PARAMETER_RDG_TEXTURE(Texture3D, LumBilateralGrid)
+		SHADER_PARAMETER_RDG_TEXTURE(Texture2D, BlurredLogLum)
+
+		SHADER_PARAMETER_SAMPLER(SamplerState, TextureSampler)
+
+		SHADER_PARAMETER(uint32, DebugMode)
+
+		RENDER_TARGET_BINDING_SLOTS()
+	END_SHADER_PARAMETER_STRUCT()
+
+	static bool ShouldCompilePermutation(const FGlobalShaderPermutationParameters& Parameters)
+	{
+		return IsFeatureLevelSupported(Parameters.Platform, ERHIFeatureLevel::SM5);
+	}
+};
+
+IMPLEMENT_GLOBAL_SHADER(FVisualizeLocalExposurePS, "/Engine/Private/PostProcessVisualizeLocalExposure.usf", "MainPS", SF_Pixel);
+
+FScreenPassTexture AddVisualizeLocalExposurePass(FRDGBuilder& GraphBuilder, const FViewInfo& View, const FVisualizeLocalExposureInputs& Inputs)
+{
+	check(Inputs.SceneColor.IsValid());
+	check(Inputs.HDRSceneColor.IsValid());
+	check(Inputs.LumBilateralGridTexture);
+	check(Inputs.BlurredLumTexture);
+	check(Inputs.EyeAdaptationTexture);
+	check(Inputs.EyeAdaptationParameters);
+
+	FScreenPassRenderTarget Output = Inputs.OverrideOutput;
+
+	if (!Output.IsValid())
+	{
+		Output = FScreenPassRenderTarget::CreateFromInput(GraphBuilder, Inputs.SceneColor, View.GetOverwriteLoadAction(), TEXT("VisualizeLocalExposure"));
+	}
+
+	const FScreenPassTextureViewport InputViewport(Inputs.SceneColor);
+	const FScreenPassTextureViewport OutputViewport(Output);
+
+	FRHISamplerState* BilinearClampSampler = TStaticSamplerState<SF_Bilinear, AM_Clamp, AM_Clamp, AM_Clamp>::GetRHI();
+
+	const FPostProcessSettings& Settings = View.FinalPostProcessSettings;
+
+	auto PassParameters = GraphBuilder.AllocParameters<FVisualizeLocalExposurePS::FParameters>();
+	PassParameters->RenderTargets[0] = Output.GetRenderTargetBinding();
+	PassParameters->View = View.ViewUniformBuffer;
+	PassParameters->EyeAdaptation = *Inputs.EyeAdaptationParameters;
+	PassParameters->Input = GetScreenPassTextureViewportParameters(InputViewport);
+	PassParameters->Output = GetScreenPassTextureViewportParameters(OutputViewport);
+	PassParameters->HDRSceneColorTexture = Inputs.HDRSceneColor.Texture;
+	PassParameters->EyeAdaptationTexture = Inputs.EyeAdaptationTexture;
+	PassParameters->LumBilateralGrid = Inputs.LumBilateralGridTexture;
+	PassParameters->BlurredLogLum = Inputs.BlurredLumTexture;
+	PassParameters->TextureSampler = BilinearClampSampler;
+	PassParameters->DebugMode = CVarLocalExposureVisualizeDebugMode.GetValueOnRenderThread();
+
+	TShaderMapRef<FVisualizeLocalExposurePS> PixelShader(View.ShaderMap);
+
+	RDG_EVENT_SCOPE(GraphBuilder, "VisualizeLocalExposure");
+
+	AddDrawScreenPass(GraphBuilder, RDG_EVENT_NAME("Visualizer"), View, OutputViewport, InputViewport, PixelShader, PassParameters);
+
+	return MoveTemp(Output);
+}
\ No newline at end of file
diff --git a/Engine/Source/Runtime/Renderer/Private/PostProcess/PostProcessVisualizeLocalExposure.h b/Engine/Source/Runtime/Renderer/Private/PostProcess/PostProcessVisualizeLocalExposure.h
new file mode 100644
--- /dev/null
+++ b/Engine/Source/Runtime/Renderer/Private/PostProcess/PostProcessVisualizeLocalExposure.h
@@ -0,0 +1,25 @@
+// Copyright Epic Games, Inc. All Rights Reserved.
+
+#pragma once
+
+#include "OverridePassSequence.h"
+
+class FEyeAdaptationParameters;
+
+struct FVisualizeLocalExposureInputs
+{
+	// [Optional] Render to the specified output. If invalid, a new texture is created and returned.
+	FScreenPassRenderTarget OverrideOutput;
+
+	FScreenPassTexture SceneColor;
+	FScreenPassTexture HDRSceneColor;
+
+	FRDGTextureRef EyeAdaptationTexture = nullptr;
+
+	FRDGTextureRef LumBilateralGridTexture = nullptr;
+	FRDGTextureRef BlurredLumTexture = nullptr;
+
+	const FEyeAdaptationParameters* EyeAdaptationParameters = nullptr;
+};
+
+FScreenPassTexture AddVisualizeLocalExposurePass(FRDGBuilder& GraphBuilder, const FViewInfo& View, const FVisualizeLocalExposureInputs& Inputs);
\ No newline at end of file
diff --git a/Engine/Source/Runtime/Renderer/Private/PostProcess/PostProcessWeightedSampleSum.cpp b/Engine/Source/Runtime/Renderer/Private/PostProcess/PostProcessWeightedSampleSum.cpp
--- a/Engine/Source/Runtime/Renderer/Private/PostProcess/PostProcessWeightedSampleSum.cpp
+++ b/Engine/Source/Runtime/Renderer/Private/PostProcess/PostProcessWeightedSampleSum.cpp
@@ -343,7 +343,8 @@ FScreenPassTexture AddGaussianBlurPass(
 	FScreenPassTexture Filter,
 	FScreenPassTexture Additive,
 	TArrayView<const FVector2D> SampleOffsets,
-	TArrayView<const FLinearColor> SampleWeights)
+	TArrayView<const FLinearColor> SampleWeights,
+	bool UseMirrorAddressMode)
 {
 	check(!OutputViewport.IsEmpty());
 	check(Filter.IsValid());
@@ -360,7 +361,16 @@ FScreenPassTexture AddGaussianBlurPass(
 
 	const uint32 StaticSampleCount = GetStaticSampleCount(SampleCount);
 
-	FRHISamplerState* SamplerState = TStaticSamplerState<SF_Bilinear, AM_Border, AM_Border, AM_Clamp>::GetRHI();
+	FRHISamplerState* SamplerState;
+
+	if (UseMirrorAddressMode)
+	{
+		SamplerState = TStaticSamplerState<SF_Bilinear, AM_Mirror, AM_Mirror, AM_Clamp>::GetRHI();
+	}
+	else
+	{
+		SamplerState = TStaticSamplerState<SF_Bilinear, AM_Border, AM_Border, AM_Clamp>::GetRHI();
+	}
 
 	const FScreenPassTextureInput FilterInput = GetScreenPassTextureInput(Filter, SamplerState);
 
@@ -507,7 +517,8 @@ FScreenPassTexture AddGaussianBlurPass(
 			Inputs.Filter,
 			Additive,
 			TArrayView<const FVector2D>(SampleOffsets, SampleCount),
-			TArrayView<const FLinearColor>(SampleWeights, SampleCount));
+			TArrayView<const FLinearColor>(SampleWeights, SampleCount),
+			Inputs.UseMirrorAddressMode);
 	}
 
 	// Vertical Pass
@@ -541,6 +552,7 @@ FScreenPassTexture AddGaussianBlurPass(
 			HorizontalOutput,
 			Inputs.Additive,
 			TArrayView<const FVector2D>(SampleOffsets, SampleCount),
-			TArrayView<const FLinearColor>(SampleWeights, SampleCount));
+			TArrayView<const FLinearColor>(SampleWeights, SampleCount),
+			Inputs.UseMirrorAddressMode);
 	}
 }
\ No newline at end of file
diff --git a/Engine/Source/Runtime/Renderer/Private/PostProcess/PostProcessWeightedSampleSum.h b/Engine/Source/Runtime/Renderer/Private/PostProcess/PostProcessWeightedSampleSum.h
--- a/Engine/Source/Runtime/Renderer/Private/PostProcess/PostProcessWeightedSampleSum.h
+++ b/Engine/Source/Runtime/Renderer/Private/PostProcess/PostProcessWeightedSampleSum.h
@@ -24,6 +24,8 @@ struct FGaussianBlurInputs
 
 	// The filter kernel size in percentage of the screen.
 	float KernelSizePercent = 0.0f;
+
+	bool UseMirrorAddressMode = false;
 };
 
 using FGaussianBlurOutputs = FScreenPassTexture;
diff --git a/Engine/Source/Runtime/Renderer/Private/PostProcess/PostProcessing.cpp b/Engine/Source/Runtime/Renderer/Private/PostProcess/PostProcessing.cpp
--- a/Engine/Source/Runtime/Renderer/Private/PostProcess/PostProcessing.cpp
+++ b/Engine/Source/Runtime/Renderer/Private/PostProcess/PostProcessing.cpp
@@ -12,7 +12,9 @@
 #include "PostProcess/PostProcessMobile.h"
 #include "PostProcess/PostProcessDownsample.h"
 #include "PostProcess/PostProcessHistogram.h"
+#include "PostProcess/PostProcessLocalExposure.h"
 #include "PostProcess/PostProcessVisualizeHDR.h"
+#include "PostProcess/PostProcessVisualizeLocalExposure.h"
 #include "PostProcess/VisualizeShadingModels.h"
 #include "PostProcess/PostProcessSelectionOutline.h"
 #include "PostProcess/PostProcessVisualizeLevelInstance.h"
@@ -357,6 +359,9 @@ void AddPostProcessingPasses(FRDGBuilder& GraphBuilder, const FViewInfo& View, c
 	// Histogram defaults to black because the histogram eye adaptation pass is used for the manual metering mode.
 	FRDGTextureRef HistogramTexture = BlackDummy.Texture;
 
+	FRDGTextureRef LocalExposureTexture = nullptr;
+	FRDGTextureRef LocalExposureBlurredLogLumTexture = BlackDummy.Texture;
+
 	const FEngineShowFlags& EngineShowFlags = View.Family->EngineShowFlags;
 	const bool bVisualizeHDR = EngineShowFlags.VisualizeHDR;
 	const bool bViewFamilyOutputInHDR = GRHISupportsHDROutput && IsHDREnabled();
@@ -384,6 +389,7 @@ void AddPostProcessingPasses(FRDGBuilder& GraphBuilder, const FViewInfo& View, c
 		VisualizeSubsurface,
 		VisualizeGBufferOverview,
 		VisualizeHDR,
+		VisualizeLocalExposure,
 		PixelInspector,
 		HMDDistortion,
 		HighResolutionScreenshotMask,
@@ -424,6 +430,7 @@ void AddPostProcessingPasses(FRDGBuilder& GraphBuilder, const FViewInfo& View, c
 		TEXT("VisualizeSubsurface"),
 		TEXT("VisualizeGBufferOverview"),
 		TEXT("VisualizeHDR"),
+		TEXT("VisualizeLocalExposure"),
 		TEXT("PixelInspector"),
 		TEXT("HMDDistortion"),
 		TEXT("HighResolutionScreenshotMask"),
@@ -531,6 +538,11 @@ void AddPostProcessingPasses(FRDGBuilder& GraphBuilder, const FViewInfo& View, c
 			// Skip if we don't have any exposure range to generate (eye adaptation will clamp).
 			View.FinalPostProcessSettings.AutoExposureMinBrightness < View.FinalPostProcessSettings.AutoExposureMaxBrightness);
 
+		const bool bLocalExposureEnabled =
+			EngineShowFlags.VisualizeLocalExposure || 
+			View.FinalPostProcessSettings.LocalExposureContrastReduction < 1.0 ||
+			!FMath::IsNearlyEqual(View.FinalPostProcessSettings.LocalExposureDetailStrength, 1.0f);
+
 		const bool bBloomEnabled = View.FinalPostProcessSettings.BloomIntensity > 0.0f;
 
 
@@ -547,6 +559,7 @@ void AddPostProcessingPasses(FRDGBuilder& GraphBuilder, const FViewInfo& View, c
 		PassSequence.SetEnabled(EPass::FXAA, AntiAliasingMethod == AAM_FXAA);
 		PassSequence.SetEnabled(EPass::PostProcessMaterialAfterTonemapping, PostProcessMaterialAfterTonemappingChain.Num() != 0);
 		PassSequence.SetEnabled(EPass::VisualizeDepthOfField, bVisualizeDepthOfField);
+		PassSequence.SetEnabled(EPass::VisualizeLocalExposure, EngineShowFlags.VisualizeLocalExposure);
 
 		for (int32 ViewExt = 0; ViewExt < View.Family->ViewExtensions.Num(); ++ViewExt)
 		{
@@ -774,6 +787,11 @@ void AddPostProcessingPasses(FRDGBuilder& GraphBuilder, const FViewInfo& View, c
 			GraphBuilder.QueueTextureExtraction(DownsampledSceneColor.Texture, &View.ViewState->PrevFrameViewInfo.HalfResTemporalAAHistory);
 		}
 
+		if (bLocalExposureEnabled)
+		{
+			LocalExposureTexture = AddLocalExposurePass(GraphBuilder, View, EyeAdaptationParameters, DownsampledSceneColor);
+		}
+
 		if (bHistogramEnabled)
 		{
 			HistogramTexture = AddHistogramPass(GraphBuilder, View, EyeAdaptationParameters, DownsampledSceneColor, LastEyeAdaptationTexture);
@@ -781,14 +799,20 @@ void AddPostProcessingPasses(FRDGBuilder& GraphBuilder, const FViewInfo& View, c
 
 		const bool bFFTBloomEnabled = IsFFTBloomEnabled(View);
 		const bool bBasicEyeAdaptationEnabled = bEyeAdaptationEnabled && (AutoExposureMethod == EAutoExposureMethod::AEM_Basic);
+		const bool bLocalExposureBlurredLum = bLocalExposureEnabled && View.FinalPostProcessSettings.LocalExposureBlurredLuminanceBlend > 0.0f;
 
 		FSceneDownsampleChain SceneDownsampleChain;
-		if (bBasicEyeAdaptationEnabled || (bBloomEnabled && bFFTBloomEnabled))
+		if (bBasicEyeAdaptationEnabled || (bBloomEnabled && bFFTBloomEnabled) || bLocalExposureBlurredLum)
 		{
 			const bool bLogLumaInAlpha = true;
 			SceneDownsampleChain.Init(GraphBuilder, View, EyeAdaptationParameters, DownsampledSceneColor, DownsampleQuality, bLogLumaInAlpha);
 		}
 
+		if (bLocalExposureBlurredLum)
+		{
+			LocalExposureBlurredLogLumTexture = AddLocalExposureBlurredLogLuminancePass(GraphBuilder, View, EyeAdaptationParameters, SceneDownsampleChain.GetTexture(4));
+		}
+
 		if (bBasicEyeAdaptationEnabled)
 		{
 			// Use the alpha channel in the last downsample (smallest) to compute eye adaptations values.
@@ -913,6 +937,9 @@ void AddPostProcessingPasses(FRDGBuilder& GraphBuilder, const FViewInfo& View, c
 				PassSequence.AcceptOverrideIfLastPass(EPass::Tonemap, PassInputs.OverrideOutput);
 				PassInputs.SceneColor = SceneColor;
 				PassInputs.Bloom = Bloom;
+				PassInputs.LocalExposureTexture = LocalExposureTexture;
+				PassInputs.BlurredLogLuminanceTexture = LocalExposureBlurredLogLumTexture;
+				PassInputs.EyeAdaptationParameters = &EyeAdaptationParameters;
 				PassInputs.EyeAdaptationTexture = EyeAdaptationTexture;
 				PassInputs.ColorGradingTexture = ColorGradingTexture;
 				PassInputs.bWriteAlphaChannel = AntiAliasingMethod == AAM_FXAA || IsPostProcessingWithAlphaChannelSupported();
@@ -970,6 +997,7 @@ void AddPostProcessingPasses(FRDGBuilder& GraphBuilder, const FViewInfo& View, c
 		PassSequence.SetEnabled(EPass::FXAA, false);
 		PassSequence.SetEnabled(EPass::PostProcessMaterialAfterTonemapping, false);
 		PassSequence.SetEnabled(EPass::VisualizeDepthOfField, false);
+		PassSequence.SetEnabled(EPass::VisualizeLocalExposure, false);
 		PassSequence.Finalize();
 
 		// Compose separate translucency passes
@@ -997,6 +1025,7 @@ void AddPostProcessingPasses(FRDGBuilder& GraphBuilder, const FViewInfo& View, c
 			FTonemapInputs PassInputs;
 			PassSequence.AcceptOverrideIfLastPass(EPass::Tonemap, PassInputs.OverrideOutput);
 			PassInputs.SceneColor = SceneColor;
+			PassInputs.EyeAdaptationParameters = &EyeAdaptationParameters;
 			PassInputs.EyeAdaptationTexture = EyeAdaptationTexture;
 			PassInputs.bOutputInHDR = bViewFamilyOutputInHDR;
 			PassInputs.bGammaOnly = true;
@@ -1141,6 +1170,20 @@ void AddPostProcessingPasses(FRDGBuilder& GraphBuilder, const FViewInfo& View, c
 		SceneColor = AddVisualizeHDRPass(GraphBuilder, View, PassInputs);
 	}
 
+	if (PassSequence.IsEnabled(EPass::VisualizeLocalExposure))
+	{
+		FVisualizeLocalExposureInputs PassInputs;
+		PassSequence.AcceptOverrideIfLastPass(EPass::VisualizeLocalExposure, PassInputs.OverrideOutput);
+		PassInputs.SceneColor = SceneColor;
+		PassInputs.HDRSceneColor = SceneColorBeforeTonemap;
+		PassInputs.LumBilateralGridTexture = LocalExposureTexture;
+		PassInputs.BlurredLumTexture = LocalExposureBlurredLogLumTexture;
+		PassInputs.EyeAdaptationTexture = EyeAdaptationTexture;
+		PassInputs.EyeAdaptationParameters = &EyeAdaptationParameters;
+
+		SceneColor = AddVisualizeLocalExposurePass(GraphBuilder, View, PassInputs);
+	}
+
 #if WITH_EDITOR
 	if (PassSequence.IsEnabled(EPass::PixelInspector))
 	{
@@ -1412,13 +1455,16 @@ void AddDebugViewPostProcessingPasses(FRDGBuilder& GraphBuilder, const FViewInfo
 
 	if (PassSequence.IsEnabled(EPass::TonemapAfter))
 	{
+		const FEyeAdaptationParameters EyeAdaptationParameters = GetEyeAdaptationParameters(View, ERHIFeatureLevel::SM5);
+
 		FTonemapInputs PassInputs;
 		PassSequence.AcceptOverrideIfLastPass(EPass::TonemapAfter, PassInputs.OverrideOutput);
 		PassInputs.SceneColor = SceneColor;
 		PassInputs.bOutputInHDR = bViewFamilyOutputInHDR;
 		PassInputs.bGammaOnly = true;
 		// Do eye adaptation in ray tracing debug modes to match raster buffer visualization modes
-		PassInputs.EyeAdaptationTexture = GetEyeAdaptationTexture(GraphBuilder, View);;
+		PassInputs.EyeAdaptationParameters = &EyeAdaptationParameters;
+		PassInputs.EyeAdaptationTexture = GetEyeAdaptationTexture(GraphBuilder, View);
 
 		SceneColor = AddTonemapPass(GraphBuilder, View, PassInputs);
 	}
@@ -2192,6 +2238,7 @@ void AddMobilePostProcessingPasses(FRDGBuilder& GraphBuilder, FScene* Scene, con
 			
 		TonemapperInputs.SceneColor = SceneColor;
 		TonemapperInputs.Bloom.Bloom = BloomOutput;
+		TonemapperInputs.EyeAdaptationParameters = &EyeAdaptationParameters;
 		TonemapperInputs.EyeAdaptationTexture = nullptr;
 		TonemapperInputs.ColorGradingTexture = ColorGradingTexture;
 		TonemapperInputs.bWriteAlphaChannel = View.AntiAliasingMethod == AAM_FXAA || IsPostProcessingWithAlphaChannelSupported() || bUseMobileDof || IsMobilePropagateAlphaEnabled(View.GetShaderPlatform());
diff --git a/Engine/Source/Runtime/Renderer/Private/SceneRendering.cpp b/Engine/Source/Runtime/Renderer/Private/SceneRendering.cpp
--- a/Engine/Source/Runtime/Renderer/Private/SceneRendering.cpp
+++ b/Engine/Source/Runtime/Renderer/Private/SceneRendering.cpp
@@ -3068,6 +3068,8 @@ void FSceneRenderer::RenderFinish(FRDGBuilder& GraphBuilder, FRDGTextureRef View
 		bool bLumenEnabledButDisabledForTheProject = false;
 		bool bNaniteEnabledButDisabledInProject = false;
 
+		bool bLocalExposureEnabledOnAnyView = false;
+
 		for (int32 ViewIndex = 0;ViewIndex < Views.Num();ViewIndex++)
 		{	
 			FViewInfo& View = Views[ViewIndex];
@@ -3080,8 +3082,13 @@ void FSceneRenderer::RenderFinish(FRDGBuilder& GraphBuilder, FRDGTextureRef View
 				|| (!ShouldRenderLumenReflections(View) && ShouldRenderLumenReflections(View, /*bSkipTracingDataCheck*/ false, /*bSkipProjectCheck*/ true));
 
 			bNaniteEnabledButDisabledInProject = bNaniteEnabledButDisabledInProject || (WouldRenderNanite(Scene, View, /*bCheckForAtomicSupport*/ false, /*bCheckForProjectSetting*/ false) && !WouldRenderNanite(Scene, View, /*bCheckForAtomicSupport*/ false, /*bCheckForProjectSetting*/ true));
+
+			if (IsPostProcessingEnabled(View) && (View.FinalPostProcessSettings.LocalExposureContrastReduction < 1.0f || !FMath::IsNearlyEqual(View.FinalPostProcessSettings.LocalExposureDetailStrength, 1.0f)))
+				bLocalExposureEnabledOnAnyView = true;
 		}
 
+		const bool bShowLocalExposureDisabledWarning = ViewFamily.EngineShowFlags.VisualizeLocalExposure && !bLocalExposureEnabledOnAnyView;
+
 		bool bNaniteEnabledButNoAtomics = false;
 		if (!GRHISupportsAtomicUInt64)
 		{
@@ -3137,7 +3144,8 @@ void FSceneRenderer::RenderFinish(FRDGBuilder& GraphBuilder, FRDGTextureRef View
 		const bool bAnyWarning = bShowPrecomputedVisibilityWarning || bShowDemotedLocalMemoryWarning || bShowGlobalClipPlaneWarning || bShowSkylightWarning || bShowPointLightWarning
 			|| bShowDFAODisabledWarning || bShowShadowedLightOverflowWarning || bShowMobileDynamicCSMWarning || bShowMobileLowQualityLightmapWarning || bShowMobileMovableDirectionalLightWarning
 			|| bMobileShowVertexFogWarning || bMobileMissingSkyMaterial || bShowSkinCacheOOM || bSingleLayerWaterWarning || bShowDFDisabledWarning || bShowNoSkyAtmosphereComponentWarning || bFxDebugDraw 
-			|| bLumenEnabledButHasNoDataForTracing || bLumenEnabledButDisabledForTheProject || bNaniteEnabledButNoAtomics || bNaniteEnabledButDisabledInProject || bRealTimeSkyCaptureButNothingToCapture || bShowWaitingSkylight;
+			|| bLumenEnabledButHasNoDataForTracing || bLumenEnabledButDisabledForTheProject || bNaniteEnabledButNoAtomics || bNaniteEnabledButDisabledInProject || bRealTimeSkyCaptureButNothingToCapture || bShowWaitingSkylight
+			|| bShowLocalExposureDisabledWarning;
 
 		for(int32 ViewIndex = 0;ViewIndex < Views.Num();ViewIndex++)
 		{	
@@ -3159,7 +3167,7 @@ void FSceneRenderer::RenderFinish(FRDGBuilder& GraphBuilder, FRDGTextureRef View
 						bLocked, bShowPrecomputedVisibilityWarning, bShowDemotedLocalMemoryWarning, bShowGlobalClipPlaneWarning, bShowDFAODisabledWarning, bShowDFDisabledWarning,
 						bViewParentOrFrozen, bShowSkylightWarning, bShowPointLightWarning, bShowShadowedLightOverflowWarning,
 						bShowMobileLowQualityLightmapWarning, bShowMobileMovableDirectionalLightWarning, bShowMobileDynamicCSMWarning, bMobileShowVertexFogWarning, bMobileMissingSkyMaterial, 
-						bShowSkinCacheOOM, bSingleLayerWaterWarning, bShowNoSkyAtmosphereComponentWarning, bFxDebugDraw, FXInterface, 
+						bShowSkinCacheOOM, bSingleLayerWaterWarning, bShowNoSkyAtmosphereComponentWarning, bFxDebugDraw, FXInterface, bShowLocalExposureDisabledWarning,
 						bLumenEnabledButHasNoDataForTracing, bLumenEnabledButDisabledForTheProject, bNaniteEnabledButNoAtomics, bNaniteEnabledButDisabledInProject, bNaniteRequireAtomics, bRealTimeSkyCaptureButNothingToCapture, bShowWaitingSkylight]
 						(FCanvas& Canvas)
 					{
@@ -3284,6 +3292,12 @@ void FSceneRenderer::RenderFinish(FRDGBuilder& GraphBuilder, FRDGTextureRef View
 							Canvas.DrawShadowedText(10, Y, FText::FromString(String), GetStatsFont(), FLinearColor(1.0, 0.05, 0.05, 1.0));
 							Y += 14;
 						}
+						if (bShowLocalExposureDisabledWarning)
+						{
+							static const FText Message = NSLOCTEXT("Renderer", "LocalExposureDisabled", "Local Exposure is disabled.");
+							Canvas.DrawShadowedText(10, Y, Message, GetStatsFont(), FLinearColor(1.0, 0.05, 0.05, 1.0));
+							Y += 14;
+						}
 
 						if (bLocked)
 						{
