Fix shader paameters issue with LWC types.

diff --git a/Engine/Source/Runtime/Renderer/Private/HairStrands/HairStrandsTile.cpp b/Engine/Source/Runtime/Renderer/Private/HairStrands/HairStrandsTile.cpp
--- a/Engine/Source/Runtime/Renderer/Private/HairStrands/HairStrandsTile.cpp
+++ b/Engine/Source/Runtime/Renderer/Private/HairStrands/HairStrandsTile.cpp
@@ -275,7 +275,7 @@ class FHairStrandsTileDebugPassPS : public FGlobalShader
 	BEGIN_SHADER_PARAMETER_STRUCT(FParameters, )
 		SHADER_PARAMETER_STRUCT_REF(FViewUniformShaderParameters, ViewUniformBuffer)
 		SHADER_PARAMETER_STRUCT_INCLUDE(FHairStrandsTilePassVS::FParameters, TileParameters)
-		SHADER_PARAMETER(FVector, DebugColor)
+		SHADER_PARAMETER(FVector3f, DebugColor)
 		SHADER_PARAMETER(FIntPoint, OutputResolution)
 		RENDER_TARGET_BINDING_SLOTS()
 	END_SHADER_PARAMETER_STRUCT()
@@ -308,13 +308,13 @@ void AddHairStrandsDebugTilePass(
 	{
 		FHairStrandsTileDebugPassPS::FParameters* PassParameters = GraphBuilder.AllocParameters<FHairStrandsTileDebugPassPS::FParameters>();
 		PassParameters->TileParameters = GetHairStrandsTileParameters(View, TileData, TileType);
-		PassParameters->DebugColor = FVector(1.f, 0.f, 1.f);
+		PassParameters->DebugColor = FVector3f(1.f, 0.f, 1.f);
 		switch (TileType)
 		{
-		case FHairStrandsTiles::ETileType::HairAll:		PassParameters->DebugColor = FVector(1.f, 1.f, 0.f); break;
-		case FHairStrandsTiles::ETileType::HairFull:	PassParameters->DebugColor = FVector(1.f, 0.f, 0.f); break;
-		case FHairStrandsTiles::ETileType::HairPartial: PassParameters->DebugColor = FVector(0.f, 1.f, 0.f); break;
-		case FHairStrandsTiles::ETileType::Other:		PassParameters->DebugColor = FVector(0.6f, 0.6f, 0.6f); break;
+		case FHairStrandsTiles::ETileType::HairAll:		PassParameters->DebugColor = FVector3f(1.f, 1.f, 0.f); break;
+		case FHairStrandsTiles::ETileType::HairFull:	PassParameters->DebugColor = FVector3f(1.f, 0.f, 0.f); break;
+		case FHairStrandsTiles::ETileType::HairPartial: PassParameters->DebugColor = FVector3f(0.f, 1.f, 0.f); break;
+		case FHairStrandsTiles::ETileType::Other:		PassParameters->DebugColor = FVector3f(0.6f, 0.6f, 0.6f); break;
 		}
 
 		TShaderMapRef<FHairStrandsTilePassVS> VertexShader(View.ShaderMap);
