Fixed: instance culling assert on mobile when full depth pre-pass is used

diff --git a/Engine/Source/Runtime/Renderer/Private/MobileShadingRenderer.cpp b/Engine/Source/Runtime/Renderer/Private/MobileShadingRenderer.cpp
--- a/Engine/Source/Runtime/Renderer/Private/MobileShadingRenderer.cpp
+++ b/Engine/Source/Runtime/Renderer/Private/MobileShadingRenderer.cpp
@@ -921,9 +921,13 @@ void FMobileSceneRenderer::BuildInstanceCullingDrawParams(FRDGBuilder& GraphBuil
 {
 	if (Scene->GPUScene.IsEnabled())
 	{
+		if (!bIsFullDepthPrepassEnabled)
+		{
+			BuildMeshPassInstanceCullingDrawParams(GraphBuilder, Scene->GPUScene, View.ParallelMeshDrawCommandPasses[EMeshPass::DepthPass], *PassParameters, MeshPassInstanceCullingDrawParams[EMeshPass::DepthPass]);
+		}
+		
 		const EMeshPass::Type BaseMeshPasses[] = 
 		{
-			EMeshPass::DepthPass,
 			EMeshPass::BasePass,
 			EMeshPass::SkyPass,
 			StandardTranslucencyMeshPass,
