Fixes alpha channel through gaussian bloom

diff --git a/Engine/Source/Runtime/Renderer/Private/PostProcess/PostProcessBloomSetup.cpp b/Engine/Source/Runtime/Renderer/Private/PostProcess/PostProcessBloomSetup.cpp
--- a/Engine/Source/Runtime/Renderer/Private/PostProcess/PostProcessBloomSetup.cpp
+++ b/Engine/Source/Runtime/Renderer/Private/PostProcess/PostProcessBloomSetup.cpp
@@ -264,7 +264,8 @@ FBloomOutputs AddGaussianBloomPasses(FRDGBuilder& GraphBuilder, const FViewInfo&
 		{
 			FBloomOutputs::FApplyInfo ApplyParameters;
 			ApplyParameters.SceneColorMultiply = FLinearColor::White;
-			ApplyParameters.BloomMultiply = FLinearColor::White * View.FinalPostProcessSettings.BloomIntensity;;
+			ApplyParameters.BloomMultiply = FLinearColor::White * View.FinalPostProcessSettings.BloomIntensity;
+			ApplyParameters.BloomMultiply.A = 0.0f;
 
 			TArray<FBloomOutputs::FApplyInfo, TInlineAllocator<1> > ApplyParametersArray;
 			ApplyParametersArray.Add(ApplyParameters);
