Fixed crash running RayTracing Basic test with Large World Coordinates enabled

diff --git a/Engine/Source/Runtime/Renderer/Private/Tests/RayTracingTestbed.cpp b/Engine/Source/Runtime/Renderer/Private/Tests/RayTracingTestbed.cpp
--- a/Engine/Source/Runtime/Renderer/Private/Tests/RayTracingTestbed.cpp
+++ b/Engine/Source/Runtime/Renderer/Private/Tests/RayTracingTestbed.cpp
@@ -28,11 +28,11 @@ bool RunRayTracingTestbed_RenderThread(const FString& Parameters)
 	FBufferRHIRef VertexBuffer;
 
 	{
-		TResourceArray<FVector> PositionData;
+		TResourceArray<FVector3f> PositionData;
 		PositionData.SetNumUninitialized(3);
-		PositionData[0] = FVector( 1, -1, 0);
-		PositionData[1] = FVector( 1,  1, 0);
-		PositionData[2] = FVector(-1, -1, 0);
+		PositionData[0] = FVector3f( 1, -1, 0);
+		PositionData[1] = FVector3f( 1,  1, 0);
+		PositionData[2] = FVector3f(-1, -1, 0);
 
 		FRHIResourceCreateInfo CreateInfo(TEXT("RayTracingTestbedVB"));
 		CreateInfo.ResourceArray = &PositionData;
