Add HZB test to prevent VSM invalidations from occluded instances.

diff --git a/Engine/Shaders/Private/VirtualShadowMaps/CacheManagement.usf b/Engine/Shaders/Private/VirtualShadowMaps/CacheManagement.usf
--- a/Engine/Shaders/Private/VirtualShadowMaps/CacheManagement.usf
+++ b/Engine/Shaders/Private/VirtualShadowMaps/CacheManagement.usf
@@ -106,12 +106,16 @@ void VirtualSmInvalidateInstancePagesCS(
 						for (int MipLevel = 0; MipLevel < NumMipLevels; ++MipLevel)
 						{
 							int ViewDim = int(uint(VSM_VIRTUAL_MAX_RESOLUTION_XY) >> MipLevel);
-							FScreenRect Rect = GetScreenRect(int4(0, 0, ViewDim, ViewDim), Cull, 2);
+							FScreenRect Rect = GetScreenRect(int4(0, 0, ViewDim, ViewDim), Cull, 4);
 							// Use Hierarchical mip test to speed up (allows skipping invalidating areas that don't have any flags anyway)
 							if (OverlapsAnyValidPage(ProjectionData.VirtualShadowMapId, MipLevel, float2(0.0f, 0.0f), Rect, VSM_ALLOCATED_FLAG))
 							{
 								uint4 RectPages = uint4(Rect.Pixels) >> VSM_LOG2_PAGE_SIZE;
 
+#if USE_HZB_OCCLUSION
+								FPageTestScreenRect HZBTestRect = SetupPageHZBRect(Rect);
+#endif // USE_HZB_OCCLUSION
+
 								// Clip to actually allocated pages
 								// TODO: move this to be done as part of or before the overlap test?
 								uint4 AllocatedBounds = PageRectBounds[ProjectionData.VirtualShadowMapId * VSM_MAX_MIP_LEVELS + MipLevel];
@@ -130,7 +134,12 @@ void VirtualSmInvalidateInstancePagesCS(
 
 											if (PageFlag != 0)
 											{
-												// Somewhere
+#if USE_HZB_OCCLUSION
+												if (!IsPageVisibleHZB(uint2(x, y), PageFlagOffset, HZBTestRect))
+												{
+													continue;
+												}
+#endif // USE_HZB_OCCLUSION
 												OutDynamicCasterFlags[PageFlagOffset] = 1;
 #if ENABLE_DEBUG_MODE
 												++PageInvalidationCount;
diff --git a/Engine/Shaders/Private/VirtualShadowMaps/PageOverlap.ush b/Engine/Shaders/Private/VirtualShadowMaps/PageOverlap.ush
--- a/Engine/Shaders/Private/VirtualShadowMaps/PageOverlap.ush
+++ b/Engine/Shaders/Private/VirtualShadowMaps/PageOverlap.ush
@@ -126,3 +126,56 @@ bool IsVisibleHZB( uint ShadowMapID, uint MipLevel, FScreenRect Rect )
 
 	return false;
 }
+
+/**
+ * Wrapper type to make misuse slightly harder.
+ */
+struct FPageTestScreenRect
+{
+	FScreenRect ScreenRect;
+	uint HZBLevelPageSizeShift;
+	int HZBLevelPageSizeInclusive;
+};
+
+
+/**
+ * Set up a screen rect and pre-computed data for testing pages against HZB, this assumes a 4x4-HZB FScreenRect
+ * as input. The resulting rect has been clamped to the mip level where a page is 4x4 texels, as higher mips are meaningless.
+ */
+FPageTestScreenRect SetupPageHZBRect(FScreenRect ScreenRect)
+{
+	FPageTestScreenRect Result;
+	Result.ScreenRect = ScreenRect;
+	// Clamp to level where a page is 4x4 (HZB mip 0 is half-size)
+	if (Result.ScreenRect.HZBLevel > VSM_LOG2_PAGE_SIZE - 3U)
+	{
+		// Adjust HZB texel rect to match new mip level, this will be too large, but is clipped below.
+		Result.ScreenRect.HZBTexels = int4(Result.ScreenRect.Pixels.xy, max(Result.ScreenRect.Pixels.xy, Result.ScreenRect.Pixels.zw)) >> (VSM_LOG2_PAGE_SIZE - 2U);
+		Result.ScreenRect.HZBLevel = VSM_LOG2_PAGE_SIZE - 3U;
+	}
+	Result.HZBLevelPageSizeShift = VSM_LOG2_PAGE_SIZE - 1U - Result.ScreenRect.HZBLevel;
+	Result.HZBLevelPageSizeInclusive = (1 << Result.HZBLevelPageSizeShift) - 1;
+	return Result;
+}
+
+
+bool IsPageVisibleHZB(uint2 vPage, uint PageFlagOffset, FPageTestScreenRect PageTestScreenRect)
+{
+	FShadowPhysicalPage pPage = ShadowDecodePageTable(ShadowHZBPageTable[PageFlagOffset]);
+
+	if (pPage.bThisLODValid)
+	{
+		FScreenRect HZBTestRect = PageTestScreenRect.ScreenRect;
+
+		// Move to page local (in mip level) space and clamp rect to page size.
+		HZBTestRect.HZBTexels -= (vPage << PageTestScreenRect.HZBLevelPageSizeShift).xyxy;
+		HZBTestRect.HZBTexels = clamp(HZBTestRect.HZBTexels, 0, PageTestScreenRect.HZBLevelPageSizeInclusive);
+		// Translate to physical address space
+		HZBTestRect.HZBTexels += (pPage.PhysicalAddress << PageTestScreenRect.HZBLevelPageSizeShift).xyxy;
+
+		// Skip invalidation if the instance is occluded
+		return IsVisibleHZB(HZBTestRect, true);
+	}
+
+	return true;
+}
\ No newline at end of file
diff --git a/Engine/Source/Runtime/Renderer/Private/VirtualShadowMaps/VirtualShadowMapCacheManager.cpp b/Engine/Source/Runtime/Renderer/Private/VirtualShadowMaps/VirtualShadowMapCacheManager.cpp
--- a/Engine/Source/Runtime/Renderer/Private/VirtualShadowMaps/VirtualShadowMapCacheManager.cpp
+++ b/Engine/Source/Runtime/Renderer/Private/VirtualShadowMaps/VirtualShadowMapCacheManager.cpp
@@ -34,6 +34,12 @@ static TAutoConsoleVariable<int32> CVarDrawInvalidatingBounds(
 	ECVF_RenderThreadSafe
 );
 
+static TAutoConsoleVariable<int32> CVarCacheVsmUseHzb(
+	TEXT("r.Shadow.Virtual.Cache.InvalidateUseHZB"),
+	1,
+	TEXT("Enables testing HZB for Virtual Shadow Map invalidations."),
+	ECVF_RenderThreadSafe);
+
 void FVirtualShadowMapCacheEntry::UpdateClipmap(
 	int32 VirtualShadowMapId,
 	const FMatrix &WorldToLight,
@@ -513,7 +519,8 @@ class FVirtualSmInvalidateInstancePagesCS : public FGlobalShader
 
 	class FDebugDim : SHADER_PERMUTATION_BOOL("ENABLE_DEBUG_MODE");
 	class FInputKindDim : SHADER_PERMUTATION_INT("INPUT_KIND", EInputKind_Num);
-	using FPermutationDomain = TShaderPermutationDomain<FDebugDim, FInputKindDim>;
+	class FUseHzbDim : SHADER_PERMUTATION_BOOL("USE_HZB_OCCLUSION");
+	using FPermutationDomain = TShaderPermutationDomain<FUseHzbDim, FDebugDim, FInputKindDim>;
 
 public:
 	BEGIN_SHADER_PARAMETER_STRUCT(FParameters, )
@@ -534,6 +541,11 @@ public:
 		SHADER_PARAMETER(uint32, GPUSceneFrameNumber)
 		SHADER_PARAMETER(uint32, InstanceSceneDataSOAStride)
 
+		SHADER_PARAMETER_RDG_BUFFER_SRV(StructuredBuffer< uint >, ShadowHZBPageTable)
+		SHADER_PARAMETER_RDG_TEXTURE(Texture2D, HZBTexture)
+		SHADER_PARAMETER_SAMPLER(SamplerState, HZBSampler)
+		SHADER_PARAMETER( FVector2D,	HZBSize )
+
 		SHADER_PARAMETER_RDG_BUFFER_SRV(StructuredBuffer<uint>, InvalidatingInstances)
 		SHADER_PARAMETER(uint32, NumInvalidatingInstanceSlots)
 		RDG_BUFFER_ACCESS(IndirectArgs, ERHIAccess::IndirectArgs)
@@ -569,7 +581,9 @@ TRDGUniformBufferRef<FVirtualShadowMapUniformParameters> FVirtualShadowMapArrayC
 	return GraphBuilder.CreateUniformBuffer(VersionedParameters);
 }
 
-static bool SetupCommonParameters(FRDGBuilder& GraphBuilder, FVirtualShadowMapArrayCacheManager* CacheManager, int32 TotalInstanceCount, const FGPUScene& GPUScene, FVirtualSmInvalidateInstancePagesCS::FParameters& OutPassParameters)
+static void SetupCommonParameters(FRDGBuilder& GraphBuilder, FVirtualShadowMapArrayCacheManager* CacheManager, int32 TotalInstanceCount, const FGPUScene& GPUScene, 
+	FVirtualSmInvalidateInstancePagesCS::FParameters& OutPassParameters,
+	FVirtualSmInvalidateInstancePagesCS::FPermutationDomain &OutPermutationVector)
 {
 
 	auto RegExtCreateSrv = [&GraphBuilder](const TRefCountPtr<FRDGPooledBuffer>& Buffer, const TCHAR* Name) -> FRDGBufferSRVRef
@@ -618,7 +632,19 @@ static bool SetupCommonParameters(FRDGBuilder& GraphBuilder, FVirtualShadowMapAr
 		ShaderDrawDebug::SetParameters(GraphBuilder, OutPassParameters.ShaderDrawUniformBuffer);
 	}
 
-	return bUseDebugPermutation;
+	const bool bUseHZB = (CVarCacheVsmUseHzb.GetValueOnRenderThread() != 0);
+	const TRefCountPtr<IPooledRenderTarget> PrevHZBPhysical = bUseHZB ? PrevBuffers.HZBPhysical : nullptr;
+	if (PrevHZBPhysical)
+	{
+		// Same, since we are not producing a new frame just yet
+		OutPassParameters.ShadowHZBPageTable = CacheManager->PrevUniformParameters.PageTable;
+		OutPassParameters.HZBTexture = GraphBuilder.RegisterExternalTexture(PrevHZBPhysical);
+		OutPassParameters.HZBSize = PrevHZBPhysical->GetDesc().Extent;
+		OutPassParameters.HZBSampler = TStaticSamplerState< SF_Point, AM_Clamp, AM_Clamp, AM_Clamp >::GetRHI();
+
+	}
+	OutPermutationVector.Set<FVirtualSmInvalidateInstancePagesCS::FDebugDim>(bUseDebugPermutation);
+	OutPermutationVector.Set<FVirtualSmInvalidateInstancePagesCS::FUseHzbDim>(PrevHZBPhysical != nullptr);
 }
 
 void FVirtualShadowMapArrayCacheManager::ProcessInstanceRangeInvalidation(FRDGBuilder& GraphBuilder, const TArray<FInstanceSceneDataRange, SceneRenderingAllocator>& InstanceRangesLarge, const TArray<FInstanceSceneDataRange, SceneRenderingAllocator>& InstanceRangesSmall, int32 TotalInstanceCount, const FGPUScene& GPUScene)
@@ -633,10 +659,9 @@ void FVirtualShadowMapArrayCacheManager::ProcessInstanceRangeInvalidation(FRDGBu
 
 
 	FVirtualSmInvalidateInstancePagesCS::FParameters PassParametersTmp;
-	bool bUseDebugPermutation = SetupCommonParameters(GraphBuilder, this, TotalInstanceCount, GPUScene, PassParametersTmp);
-
 	FVirtualSmInvalidateInstancePagesCS::FPermutationDomain PermutationVector;
-	PermutationVector.Set<FVirtualSmInvalidateInstancePagesCS::FDebugDim>(bUseDebugPermutation);
+	SetupCommonParameters(GraphBuilder, this, TotalInstanceCount, GPUScene, PassParametersTmp, PermutationVector);
+
 	if (InstanceRangesSmall.Num())
 	{
 		RDG_EVENT_SCOPE(GraphBuilder, "ProcessInstanceRangeInvalidation [%d small-ranges]", InstanceRangesSmall.Num());
@@ -691,15 +716,14 @@ void FVirtualShadowMapArrayCacheManager::ProcessGPUInstanceInvalidations(FRDGBui
 		FRDGBufferRef InvalidatingInstancesBufferRDG = GraphBuilder.RegisterExternalBuffer(PrevBuffers.InvalidatingInstancesBuffer, TEXT("Shadow.Virtual.PrevInvalidatingInstancesBuffer"));
 		FRDGBufferRef IndirectArgs = FComputeShaderUtils::AddIndirectArgsSetupCsPass1D(GraphBuilder, InvalidatingInstancesBufferRDG, TEXT("Shadow.Virtual.ProcessGPUInstanceInvalidationsIndirectArgs"), FVirtualSmInvalidateInstancePagesCS::Cs1dGroupSizeX);
 
+		FVirtualSmInvalidateInstancePagesCS::FPermutationDomain PermutationVector;
 		FVirtualSmInvalidateInstancePagesCS::FParameters* PassParameters = GraphBuilder.AllocParameters<FVirtualSmInvalidateInstancePagesCS::FParameters>();
-		bool bUseDebugPermutation = SetupCommonParameters(GraphBuilder, this, 16*1024, GPUScene, *PassParameters);
+		SetupCommonParameters(GraphBuilder, this, 16*1024, GPUScene, *PassParameters, PermutationVector);
 
 		PassParameters->IndirectArgs = IndirectArgs;
 		PassParameters->InvalidatingInstances = GraphBuilder.CreateSRV(InvalidatingInstancesBufferRDG);
 		PassParameters->NumInvalidatingInstanceSlots = PrevBuffers.NumInvalidatingInstanceSlots;
 
-		FVirtualSmInvalidateInstancePagesCS::FPermutationDomain PermutationVector;
-		PermutationVector.Set<FVirtualSmInvalidateInstancePagesCS::FDebugDim>(bUseDebugPermutation);
 		PermutationVector.Set<FVirtualSmInvalidateInstancePagesCS::FInputKindDim>(FVirtualSmInvalidateInstancePagesCS::EInputKind_GPUInstances);
 		auto ComputeShader = GetGlobalShaderMap(GMaxRHIFeatureLevel)->GetShader<FVirtualSmInvalidateInstancePagesCS>(PermutationVector);
 
