Fix a bug of rendering multi view on seperate passes on mobile

#p4v-preflight-copy 17383209

diff --git a/Engine/Shaders/Private/PostProcessAmbientOcclusionMobile.usf b/Engine/Shaders/Private/PostProcessAmbientOcclusionMobile.usf
--- a/Engine/Shaders/Private/PostProcessAmbientOcclusionMobile.usf
+++ b/Engine/Shaders/Private/PostProcessAmbientOcclusionMobile.usf
@@ -136,7 +136,11 @@ float4 ViewRectMin;
 
 half2 TexturePosToBufferUV(int2 TexturePos)
 {
+#if HORIZONSEARCH_INTEGRAL_PIXEL_SHADER || SPATIALFILTER_PIXEL_SHADER
+	half2 BufferUV = (TexturePos + half2(0.5f, 0.5f)) * BufferSizeAndInvSize.zw;
+#else
 	half2 BufferUV = (TexturePos + half2(0.5f, 0.5f) + ViewRectMin.xy) * BufferSizeAndInvSize.zw;
+#endif
 
 	return BufferUV;
 }
@@ -596,7 +600,7 @@ half4 GTAOSpatialFilter(half2 TextureUV, int2 ThreadPos, int2 PixelPos)
 void GTAOSpatialFilter(int2 ThreadPos, int2 PixelPos)
 #endif
 {
-	if (any(PixelPos >= (int2)ViewSizeAndInvSize.xy))
+	if (any(PixelPos >= int2(ViewRectMin.xy + ViewSizeAndInvSize.xy)))
 	{
 #if SPATIALFILTER_PIXEL_SHADER || HORIZONSEARCH_INTEGRAL_PIXEL_SHADER
 		return 1.0f;
@@ -767,7 +771,7 @@ void GTAOSpatialFilter(int2 ThreadPos, int2 PixelPos)
 #if SPATIALFILTER_PIXEL_SHADER || HORIZONSEARCH_INTEGRAL_PIXEL_SHADER
 	return SumAO;
 #else
-	OutTexture[PixelPos] = SumAO;
+	OutTexture[PixelPos + ViewRectMin.xy] = SumAO;
 #endif
 }
 #endif
@@ -904,7 +908,7 @@ void GTAOHorizonSearchIntegralCS(
 
 	half4 EncodeZ = EncodeFloatRGBA(DeviceZ);
 	EncodeZ.a = GTAO;
-	OutTexture[TexturePos] = EncodeZ;
+	OutTexture[TexturePos + ViewRectMin.xy] = EncodeZ;
 }
 #endif
 
diff --git a/Engine/Source/Runtime/Renderer/Private/MobileShadingRenderer.cpp b/Engine/Source/Runtime/Renderer/Private/MobileShadingRenderer.cpp
--- a/Engine/Source/Runtime/Renderer/Private/MobileShadingRenderer.cpp
+++ b/Engine/Source/Runtime/Renderer/Private/MobileShadingRenderer.cpp
@@ -635,6 +635,12 @@ void FMobileSceneRenderer::RenderFullDepthPrepass(FRDGBuilder& GraphBuilder, FSc
 			continue;
 		}
 
+		if (ViewIndex > 0)
+		{
+			BasePassRenderTargets.DepthStencil.SetDepthLoadAction(ERenderTargetLoadAction::ELoad);
+			BasePassRenderTargets.DepthStencil.SetStencilLoadAction(ERenderTargetLoadAction::ELoad);
+		}
+
 		View.BeginRenderView();
 
 		auto* PassParameters = GraphBuilder.AllocParameters<FMobileRenderPassParameters>();
diff --git a/Engine/Source/Runtime/Renderer/Private/PostProcess/PostProcessAmbientOcclusionMobile.cpp b/Engine/Source/Runtime/Renderer/Private/PostProcess/PostProcessAmbientOcclusionMobile.cpp
--- a/Engine/Source/Runtime/Renderer/Private/PostProcess/PostProcessAmbientOcclusionMobile.cpp
+++ b/Engine/Source/Runtime/Renderer/Private/PostProcess/PostProcessAmbientOcclusionMobile.cpp
@@ -523,7 +523,7 @@ static void RenderGTAO(FRDGBuilder& GraphBuilder, FRDGTextureRef SceneDepthTextu
 		{
 			TShaderMapRef<FScreenPassVS> VertexShader(View.ShaderMap);
 
-			FScreenPassRenderTarget HorizonSearchIntegralRT(HorizonSearchIntegralTexture, ViewRect, ERenderTargetLoadAction::EClear);
+			FScreenPassRenderTarget HorizonSearchIntegralRT(HorizonSearchIntegralTexture, ViewRect, ViewIndex > 0 ? ERenderTargetLoadAction::ELoad : ERenderTargetLoadAction::EClear);
 
 			FGTAOMobile_HorizonSearchIntegralPS::FParameters* HorizonSearchIntegralParameters = GraphBuilder.AllocParameters<FGTAOMobile_HorizonSearchIntegralPS::FParameters>();
 			FGTAOMobile_HorizonSearchIntegral::SetupShaderParameters(HorizonSearchIntegralParameters->Common, GraphBuilder, View, ViewRect, DepthBufferSize, BufferSize, FallOffStartEndScaleBias, WorldRadiusAdjSinCosDeltaAngleThickness, SceneDepthTexture);
@@ -571,7 +571,7 @@ static void RenderGTAO(FRDGBuilder& GraphBuilder, FRDGTextureRef SceneDepthTextu
 					EDRF_UseTriangleOptimization);
 			});
 
-			FScreenPassRenderTarget AmbientOcclusionRT(AmbientOcclusionTexture, ViewRect, ERenderTargetLoadAction::EClear);
+			FScreenPassRenderTarget AmbientOcclusionRT(AmbientOcclusionTexture, ViewRect, ViewIndex > 0 ? ERenderTargetLoadAction::ELoad : ERenderTargetLoadAction::EClear);
 
 			FGTAOMobile_SpatialFilterPS::FParameters* SpatialFilterParameters = GraphBuilder.AllocParameters<FGTAOMobile_SpatialFilterPS::FParameters>();
 			FGTAOMobile_SpatialFilter::SetupShaderParameters(SpatialFilterParameters->Common, View, ViewRect, BufferSize, HorizonSearchIntegralTexture);
diff --git a/Engine/Source/Runtime/Renderer/Private/PostProcess/PostProcessPixelProjectedReflectionMobile.cpp b/Engine/Source/Runtime/Renderer/Private/PostProcess/PostProcessPixelProjectedReflectionMobile.cpp
--- a/Engine/Source/Runtime/Renderer/Private/PostProcess/PostProcessPixelProjectedReflectionMobile.cpp
+++ b/Engine/Source/Runtime/Renderer/Private/PostProcess/PostProcessPixelProjectedReflectionMobile.cpp
@@ -287,7 +287,7 @@ void FMobileSceneRenderer::RenderPixelProjectedReflection(FRDGBuilder& GraphBuil
 
 		FPixelProjectedReflectionMobile_ReflectionPassPS::FParameters* PSShaderParameters = GraphBuilder.AllocParameters<FPixelProjectedReflectionMobile_ReflectionPassPS::FParameters>();
 
-		PSShaderParameters->RenderTargets[0] = FRenderTargetBinding(PixelProjectedReflectionTexture, ERenderTargetLoadAction::EClear);
+		PSShaderParameters->RenderTargets[0] = FRenderTargetBinding(PixelProjectedReflectionTexture, ViewIndex > 0 ? ERenderTargetLoadAction::ELoad : ERenderTargetLoadAction::EClear);
 
 		PSShaderParameters->View = View.ViewUniformBuffer;
 		PSShaderParameters->SceneColorTexture = SceneColorTexture;
diff --git a/Engine/Source/Runtime/Renderer/Private/VelocityRendering.cpp b/Engine/Source/Runtime/Renderer/Private/VelocityRendering.cpp
--- a/Engine/Source/Runtime/Renderer/Private/VelocityRendering.cpp
+++ b/Engine/Source/Runtime/Renderer/Private/VelocityRendering.cpp
@@ -291,7 +291,7 @@ void FSceneRenderer::RenderVelocities(
 				ERenderTargetLoadAction::ELoad,
 				ExclusiveDepthStencil);
 
-			PassParameters->RenderTargets[0] = FRenderTargetBinding(SceneTextures.Velocity, VelocityLoadAction);
+			PassParameters->RenderTargets[0] = FRenderTargetBinding(SceneTextures.Velocity, ViewIndex > 0 ? ERenderTargetLoadAction::ELoad : VelocityLoadAction);
 
 			if (bIsParallelVelocity)
 			{
