r.RayTracing.ExcludeTransparent renamed to r.RayTracing.ExcludeTranslucent. When set to one it just excludes objects with translucent materials from the BVH

diff --git a/Engine/Source/Runtime/Renderer/Private/DeferredShadingRenderer.cpp b/Engine/Source/Runtime/Renderer/Private/DeferredShadingRenderer.cpp
--- a/Engine/Source/Runtime/Renderer/Private/DeferredShadingRenderer.cpp
+++ b/Engine/Source/Runtime/Renderer/Private/DeferredShadingRenderer.cpp
@@ -142,13 +142,13 @@ static FAutoConsoleVariableRef CRayTracingExcludeDecals(
 	TEXT(" 1: Decals excluded from the ray tracing BVH"),
 	ECVF_RenderThreadSafe);
 
-static int32 GRayTracingExcludeTransparent = 0;
-static FAutoConsoleVariableRef CRayTracingExcludeTransparent(
-	TEXT("r.RayTracing.ExcludeTransparent"),
-	GRayTracingExcludeTransparent,
-	TEXT("A toggle that modifies the inclusion of transparent objects (with one or more non-opaque sections) in the ray tracing BVH.\n")
-	TEXT(" 0: Transparent objects included in the ray tracing BVH (default)\n")
-	TEXT(" 1: Transparent objects excluded from the ray tracing BVH"),
+static int32 GRayTracingExcludeTranslucent = 0;
+static FAutoConsoleVariableRef CRayTracingExcludeTranslucent(
+	TEXT("r.RayTracing.ExcludeTranslucent"),
+	GRayTracingExcludeTranslucent,
+	TEXT("A toggle that modifies the inclusion of translucent objects in the ray tracing scene.\n")
+	TEXT(" 0: Translucent objects included in the ray tracing scene (default)\n")
+	TEXT(" 1: Translucent objects excluded from the ray tracing scene"),
 	ECVF_RenderThreadSafe);
 
 static TAutoConsoleVariable<int32> CVarRayTracingAsyncBuild(
@@ -582,6 +582,7 @@ bool FDeferredShadingSceneRenderer::GatherRayTracingWorldInstancesForView(FRHICo
 		bool bAnySegmentsDecal = false;
 		bool bTwoSided = false;
 		bool bIsSky = false;
+		bool bAllSegmentsTranslucent = true;
 
 		uint64 InstancingKey() const
 		{
@@ -592,6 +593,7 @@ bool FDeferredShadingSceneRenderer::GatherRayTracingWorldInstancesForView(FRHICo
 			Key ^= bAnySegmentsDecal ? 0x1ull << 42 : 0x0;
 			Key ^= bTwoSided ? 0x1ull << 43 : 0x0;
 			Key ^= bIsSky ? 0x1ull << 44 : 0x0;
+			Key ^= bAllSegmentsTranslucent ? 0x1ull << 45 : 0x0;
 			return Key ^ reinterpret_cast<uint64>(RayTracingGeometryRHI);
 		}
 	};
@@ -849,6 +851,7 @@ bool FDeferredShadingSceneRenderer::GatherRayTracingWorldInstancesForView(FRHICo
 									RelevantPrimitive.bAnySegmentsDecal |= RayTracingMeshCommand.bDecal;
 									RelevantPrimitive.bTwoSided |= RayTracingMeshCommand.bTwoSided;
 									RelevantPrimitive.bIsSky |= RayTracingMeshCommand.bIsSky;
+									RelevantPrimitive.bAllSegmentsTranslucent &= RayTracingMeshCommand.bIsTranslucent;
 								}
 								else
 								{
@@ -1242,7 +1245,7 @@ bool FDeferredShadingSceneRenderer::GatherRayTracingWorldInstancesForView(FRHICo
 				}
 
 				if ((GRayTracingExcludeDecals && RelevantPrimitive.bAnySegmentsDecal)
-					|| (GRayTracingExcludeTransparent && !RelevantPrimitive.bAllSegmentsOpaque)
+					|| (GRayTracingExcludeTranslucent && RelevantPrimitive.bAllSegmentsTranslucent)
 					|| RelevantPrimitive.bIsSky)
 				{
 					continue;
diff --git a/Engine/Source/Runtime/Renderer/Private/RayTracing/RayTracingMaterialHitShaders.h b/Engine/Source/Runtime/Renderer/Private/RayTracing/RayTracingMaterialHitShaders.h
--- a/Engine/Source/Runtime/Renderer/Private/RayTracing/RayTracingMaterialHitShaders.h
+++ b/Engine/Source/Runtime/Renderer/Private/RayTracing/RayTracingMaterialHitShaders.h
@@ -71,6 +71,7 @@ protected:
 		SharedCommand.bDecal = MaterialResource.GetMaterialDomain() == EMaterialDomain::MD_DeferredDecal;
 		SharedCommand.bIsSky = MaterialResource.IsSky();
 		SharedCommand.bTwoSided = MaterialResource.IsTwoSided();
+		SharedCommand.bIsTranslucent = MaterialResource.GetBlendMode() == EBlendMode::BLEND_Translucent;
 
 		FVertexInputStreamArray VertexStreams;
 		VertexFactory->GetStreams(ERHIFeatureLevel::SM5, EVertexInputStreamType::Default, VertexStreams);
