Tidy use of RuntimeVirtualTexture::IsSceneReadyToRender().
Part of fix for low frame rate when selecting an RVT map in content browser.

diff --git a/Engine/Source/Editor/VirtualTexturingEditor/Private/RuntimeVirtualTextureThumbnailRenderer.cpp b/Engine/Source/Editor/VirtualTexturingEditor/Private/RuntimeVirtualTextureThumbnailRenderer.cpp
--- a/Engine/Source/Editor/VirtualTexturingEditor/Private/RuntimeVirtualTextureThumbnailRenderer.cpp
+++ b/Engine/Source/Editor/VirtualTexturingEditor/Private/RuntimeVirtualTextureThumbnailRenderer.cpp
@@ -41,13 +41,9 @@ bool URuntimeVirtualTextureThumbnailRenderer::CanVisualizeAsset(UObject* Object)
 
 	// We need a matching URuntimeVirtualTextureComponent in a Scene to be able to render a thumbnail
 	URuntimeVirtualTextureComponent* RuntimeVirtualTextureComponent = FindComponent(RuntimeVirtualTexture);
-	if (RuntimeVirtualTextureComponent != nullptr)
+	if (RuntimeVirtualTextureComponent != nullptr && RuntimeVirtualTextureComponent->GetScene() != nullptr)
 	{
-		FSceneInterface* Scene = RuntimeVirtualTextureComponent->GetScene();
-		if (Scene != nullptr && RuntimeVirtualTexture::IsSceneReadyToRender(Scene->GetRenderScene()))
-		{
-			return true;
-		}
+		return true;
 	}
 
 	return false;
diff --git a/Engine/Source/Runtime/Renderer/Private/VT/RuntimeVirtualTextureProducer.cpp b/Engine/Source/Runtime/Renderer/Private/VT/RuntimeVirtualTextureProducer.cpp
--- a/Engine/Source/Runtime/Renderer/Private/VT/RuntimeVirtualTextureProducer.cpp
+++ b/Engine/Source/Runtime/Renderer/Private/VT/RuntimeVirtualTextureProducer.cpp
@@ -29,8 +29,7 @@ FRuntimeVirtualTextureFinalizer::FRuntimeVirtualTextureFinalizer(
 
 bool FRuntimeVirtualTextureFinalizer::IsReady()
 {
-	FScene* RenderScene = Scene->GetRenderScene();
-	return RuntimeVirtualTexture::IsSceneReadyToRender(RenderScene) && RenderScene->GPUScene.IsRendering();
+	return RuntimeVirtualTexture::IsSceneReadyToRender(Scene);
 }
 
 void FRuntimeVirtualTextureFinalizer::InitProducer(const FVirtualTextureProducerHandle& ProducerHandle)
diff --git a/Engine/Source/Runtime/Renderer/Private/VT/RuntimeVirtualTextureRender.cpp b/Engine/Source/Runtime/Renderer/Private/VT/RuntimeVirtualTextureRender.cpp
--- a/Engine/Source/Runtime/Renderer/Private/VT/RuntimeVirtualTextureRender.cpp
+++ b/Engine/Source/Runtime/Renderer/Private/VT/RuntimeVirtualTextureRender.cpp
@@ -1072,12 +1072,9 @@ namespace RuntimeVirtualTexture
 	};
 
 
-	bool IsSceneReadyToRender(FScene* Scene)
+	bool IsSceneReadyToRender(FSceneInterface* Scene)
 	{
-		// Test scene is loaded and has been updated once by main rendering passes
-		// This function gets called on the main thread, so accessing scene frame number is not strictly thread safe, but we can probably
-		// assume that frame number is always increasing, and so the test is conservative
-		return Scene != nullptr && Scene->GetRenderScene() != nullptr && Scene->GetRenderScene()->GetFrameNumber() > 1;
+		return Scene != nullptr && Scene->GetRenderScene() != nullptr && Scene->GetRenderScene()->GPUScene.IsRendering();
 	}
 
 	void RenderPage(
diff --git a/Engine/Source/Runtime/Renderer/Public/VT/RuntimeVirtualTextureRender.h b/Engine/Source/Runtime/Renderer/Public/VT/RuntimeVirtualTextureRender.h
--- a/Engine/Source/Runtime/Renderer/Public/VT/RuntimeVirtualTextureRender.h
+++ b/Engine/Source/Runtime/Renderer/Public/VT/RuntimeVirtualTextureRender.h
@@ -12,6 +12,7 @@ class FRDGBuilder;
 class FRDGTexture;
 class FRDGTextureUAV;
 class FScene;
+class FSceneInterface;
 class URuntimeVirtualTextureComponent;
 
 namespace RuntimeVirtualTexture
@@ -81,8 +82,8 @@ namespace RuntimeVirtualTexture
 		FRenderPageDesc PageDescs[EMaxRenderPageBatch];
 	};
 
-	/** Returns true if the FScene is initialized for rendering to runtime virtual textures. Always check this before calling RenderPages(). */
-	RENDERER_API bool IsSceneReadyToRender(FScene* Scene);
+	/** Returns true if the scene is initialized for rendering to runtime virtual textures. Always check this before calling RenderPages(). */
+	RENDERER_API bool IsSceneReadyToRender(FSceneInterface* Scene);
 
 	/** Render a batch of pages for a runtime virtual texture. */
 	RENDERER_API void RenderPages(FRDGBuilder& GraphBuilder, FRenderPageBatchDesc const& InDesc);
