Fix for crash in viewmode LODColoration

diff --git a/Engine/Source/Runtime/Renderer/Private/PostProcess/PostProcessing.cpp b/Engine/Source/Runtime/Renderer/Private/PostProcess/PostProcessing.cpp
--- a/Engine/Source/Runtime/Renderer/Private/PostProcess/PostProcessing.cpp
+++ b/Engine/Source/Runtime/Renderer/Private/PostProcess/PostProcessing.cpp
@@ -1366,10 +1366,13 @@ void AddDebugViewPostProcessingPasses(FRDGBuilder& GraphBuilder, const FViewInfo
 
 	if (bTonemapBefore)
 	{
+		const FEyeAdaptationParameters EyeAdaptationParameters = GetEyeAdaptationParameters(View, ERHIFeatureLevel::SM5);
+
 		FTonemapInputs PassInputs;
 		PassInputs.SceneColor = SceneColor;
 		PassInputs.bOutputInHDR = bViewFamilyOutputInHDR;
 		PassInputs.bGammaOnly = true;
+		PassInputs.EyeAdaptationParameters = &EyeAdaptationParameters;
 
 		SceneColor = AddTonemapPass(GraphBuilder, View, PassInputs);
 	}
