Use a shader constant to control back face culling in ray traced shadow shader instead of permutation.

diff --git a/Engine/Shaders/Private/RayTracing/RayTracingOcclusionRGS.usf b/Engine/Shaders/Private/RayTracing/RayTracingOcclusionRGS.usf
--- a/Engine/Shaders/Private/RayTracing/RayTracingOcclusionRGS.usf
+++ b/Engine/Shaders/Private/RayTracing/RayTracingOcclusionRGS.usf
@@ -1,5 +1,8 @@
 // Copyright Epic Games, Inc. All Rights Reserved.
 
+// Bump to force a rebuild of this shader (when changing related C++ code). Generate unique value on merge conflict.
+#define RAY_TRACING_OCCLUSION_RGS_VERSION 0x8219847B
+
 #include "../Common.ush"
 #include "../MonteCarlo.ush"
 #include "../DeferredShadingCommon.ush"
@@ -42,6 +45,7 @@ uint bTransmissionSamplingDistanceCulling;
 uint TransmissionSamplingTechnique;
 uint RejectionSamplingTrials;
 uint bAcceptFirstHit;
+uint bTwoSidedGeometry;
 
 #define TransmissionSampling_ConstantTrackingInfiniteDomain 0
 #define TransmissionSampling_ConstantTrackingFiniteDomain 1
@@ -418,9 +422,10 @@ FOcclusionResult ComputeOcclusion(
 
 		uint RayFlags = 0;
 
-	#if !ENABLE_TWO_SIDED_GEOMETRY
-		RayFlags |= RAY_FLAG_CULL_BACK_FACING_TRIANGLES;
-	#endif
+		if (bTwoSidedGeometry != 0)
+		{
+			RayFlags |= RAY_FLAG_CULL_BACK_FACING_TRIANGLES;
+		}
 
 	#if USE_HAIR_LIGHTING
 		if (ShadingModelID == SHADINGMODELID_HAIR)
diff --git a/Engine/Source/Runtime/Renderer/Private/RayTracing/RayTracingShadows.cpp b/Engine/Source/Runtime/Renderer/Private/RayTracing/RayTracingShadows.cpp
--- a/Engine/Source/Runtime/Renderer/Private/RayTracing/RayTracingShadows.cpp
+++ b/Engine/Source/Runtime/Renderer/Private/RayTracing/RayTracingShadows.cpp
@@ -104,12 +104,11 @@ class FOcclusionRGS : public FGlobalShader
 
 	class FLightTypeDim : SHADER_PERMUTATION_INT("LIGHT_TYPE", LightType_MAX);
 	class FDenoiserOutputDim : SHADER_PERMUTATION_INT("DIM_DENOISER_OUTPUT", 3);
-	class FEnableTwoSidedGeometryDim : SHADER_PERMUTATION_BOOL("ENABLE_TWO_SIDED_GEOMETRY");
 	class FEnableMultipleSamplesPerPixel : SHADER_PERMUTATION_BOOL("ENABLE_MULTIPLE_SAMPLES_PER_PIXEL");
 	class FHairLighting : SHADER_PERMUTATION_INT("USE_HAIR_LIGHTING", 2);
 	class FEnableTransmissionDim : SHADER_PERMUTATION_INT("ENABLE_TRANSMISSION", 2);
 
-	using FPermutationDomain = TShaderPermutationDomain<FLightTypeDim, FDenoiserOutputDim, FEnableTwoSidedGeometryDim, FHairLighting, FEnableMultipleSamplesPerPixel, FEnableTransmissionDim>;
+	using FPermutationDomain = TShaderPermutationDomain<FLightTypeDim, FDenoiserOutputDim, FHairLighting, FEnableMultipleSamplesPerPixel, FEnableTransmissionDim>;
 
 	static bool ShouldCompilePermutation(const FGlobalShaderPermutationParameters& Parameters)
 	{
@@ -151,6 +150,7 @@ class FOcclusionRGS : public FGlobalShader
 		SHADER_PARAMETER(uint32, TransmissionSamplingTechnique)
 		SHADER_PARAMETER(uint32, RejectionSamplingTrials)
 		SHADER_PARAMETER(uint32, bAcceptFirstHit)
+		SHADER_PARAMETER(uint32, bTwoSidedGeometry)
 
 		SHADER_PARAMETER_STRUCT(FLightShaderParameters, Light)
 		SHADER_PARAMETER_STRUCT_INCLUDE(FSceneTextureParameters, SceneTextures)
@@ -207,7 +207,6 @@ void FDeferredShadingSceneRenderer::PrepareRayTracingShadows(const FViewInfo& Vi
 						FOcclusionRGS::FPermutationDomain PermutationVector;
 						PermutationVector.Set<FOcclusionRGS::FLightTypeDim>(LightType);
 						PermutationVector.Set<FOcclusionRGS::FDenoiserOutputDim>((int32)DenoiserRequirement);
-						PermutationVector.Set<FOcclusionRGS::FEnableTwoSidedGeometryDim>(CVarRayTracingShadowsEnableTwoSidedGeometry.GetValueOnRenderThread() != 0);
 						PermutationVector.Set<FOcclusionRGS::FHairLighting>(HairLighting);
 						PermutationVector.Set<FOcclusionRGS::FEnableMultipleSamplesPerPixel>(MultiSPP != 0);
 						PermutationVector.Set<FOcclusionRGS::FEnableTransmissionDim>(EnableTransmissionDim);
@@ -279,6 +278,7 @@ void FDeferredShadingSceneRenderer::RenderRayTracingShadows(
 		PassParameters->LODTransitionStart = CVarRayTracingShadowsLODTransitionStart.GetValueOnRenderThread();
 		PassParameters->LODTransitionEnd = CVarRayTracingShadowsLODTransitionEnd.GetValueOnRenderThread();
 		PassParameters->bAcceptFirstHit = CVarRayTracingShadowsAcceptFirstHit.GetValueOnRenderThread();
+		PassParameters->bTwoSidedGeometry = EnableRayTracingShadowTwoSidedGeometry() ? 1 : 0;
 		PassParameters->TLAS = View.GetRayTracingSceneViewChecked();
 		PassParameters->ViewUniformBuffer = View.ViewUniformBuffer;
 		PassParameters->SceneTextures = SceneTextures;
@@ -316,7 +316,6 @@ void FDeferredShadingSceneRenderer::RenderRayTracingShadows(
 		{
 			PermutationVector.Set<FOcclusionRGS::FDenoiserOutputDim>(0);
 		}
-		PermutationVector.Set<FOcclusionRGS::FEnableTwoSidedGeometryDim>(CVarRayTracingShadowsEnableTwoSidedGeometry.GetValueOnRenderThread() != 0);
 		PermutationVector.Set<FOcclusionRGS::FHairLighting>(bUseHairLighting? 1 : 0);
 
 		PermutationVector.Set<FOcclusionRGS::FEnableMultipleSamplesPerPixel>(RayTracingConfig.RayCountPerPixel > 1);
