Add basic GPU-Scene instance debug rendering

diff --git a/Engine/Shaders/Private/GPUSceneDebugRender.usf b/Engine/Shaders/Private/GPUSceneDebugRender.usf
new file mode 100644
--- /dev/null
+++ b/Engine/Shaders/Private/GPUSceneDebugRender.usf
@@ -0,0 +1,83 @@
+// Copyright Epic Games, Inc. All Rights Reserved.
+
+#include "Common.ush"
+#include "SceneData.ush"
+#include "ShaderDrawDebug.ush"
+
+int NumInstances;
+int NumScenePrimitives;
+int bDrawAll;
+int bDrawUpdatedOnly;
+
+uint InstanceDataSOAStride;
+int2 InputCoord;
+float3 PickingRayStart;
+float3 PickingRayEnd;
+float DrawRange;
+
+StructuredBuffer<uint> SelectedPrimitiveFlags;
+
+[numthreads(NUM_THREADS_PER_GROUP, 1, 1)]
+void GPUSceneDebugRenderCS(uint InstanceID : SV_DispatchThreadID)
+{
+	if (InstanceID >= NumInstances)
+	{
+		return;
+	}
+
+	FInstanceSceneData InstanceData = GetInstanceSceneData(InstanceID, InstanceDataSOAStride);
+
+	// Early out on invalid ID
+	if (InstanceData.PrimitiveId == 0xFFFFFFFFu)
+	{
+		return;
+	}
+
+	// Only show updated
+	if (bDrawUpdatedOnly && InstanceData.LastUpdateSceneFrameNumber != GetGPUSceneFrameNumber())
+	{
+		return;
+	}
+
+	float3 LocalRayStart = mul(float4(PickingRayStart, 1.0f), InstanceData.WorldToLocal).xyz;
+
+	// Skip if beyond draw range
+	if (DrawRange > 0.0f && length(InstanceData.LocalBoundsCenter - LocalRayStart) > DrawRange)
+	{
+		return;
+	}
+
+	float3 LocalRayEnd = mul(float4(PickingRayEnd, 1.0f), InstanceData.WorldToLocal).xyz;
+	float2 RayResult = LineBoxIntersect(LocalRayStart, LocalRayEnd, InstanceData.LocalBoundsCenter - InstanceData.LocalBoundsExtent, InstanceData.LocalBoundsCenter + InstanceData.LocalBoundsExtent);
+	bool bIsMouseOver = RayResult.x <= RayResult.y;
+	// Only draw selected
+	const uint BitsPerWord = 32U;
+	const uint Mask = 1U << (InstanceData.PrimitiveId % BitsPerWord);
+
+	bool bIsSelected = (InstanceData.PrimitiveId < NumScenePrimitives && (SelectedPrimitiveFlags[InstanceData.PrimitiveId / BitsPerWord] & Mask) != 0U);
+
+	FPrimitiveSceneData PrimitiveData = GetPrimitiveData(InstanceData.PrimitiveId);
+
+	//float4 InstanceBoundColor = InstanceData.PrimitiveId < NumScenePrimitives ? float4(0.3f, 1.0f, 0.5f, 1.0f) : float4(0.3f, 0.3f, 0.3f, 0.5f);
+	float4 InstanceBoundColor = float4(0.3f, 0.3f, 0.3f, 0.5f);
+
+	if (bIsSelected || bDrawAll)
+	{
+		InstanceBoundColor = float4(0.5f, 0.75f, 0.5f, 1.0f);
+	}
+	
+	if (bIsMouseOver)
+	{
+		InstanceBoundColor = float4(0.75f, 1.0f, 0.75f, 1.0f);
+	}
+	
+	if (bDrawUpdatedOnly)
+	{
+		InstanceBoundColor.x = 1.0f;
+		InstanceBoundColor.z = 1.0f;
+	}
+	if (bIsMouseOver || bIsSelected || bDrawAll != 0)
+	{
+		AddOBB(InstanceData.LocalBoundsCenter - InstanceData.LocalBoundsExtent, InstanceData.LocalBoundsCenter + InstanceData.LocalBoundsExtent, InstanceBoundColor, InstanceData.LocalToWorld);
+	}
+}
diff --git a/Engine/Source/Runtime/Renderer/Private/DeferredShadingRenderer.cpp b/Engine/Source/Runtime/Renderer/Private/DeferredShadingRenderer.cpp
--- a/Engine/Source/Runtime/Renderer/Private/DeferredShadingRenderer.cpp
+++ b/Engine/Source/Runtime/Renderer/Private/DeferredShadingRenderer.cpp
@@ -2074,6 +2074,8 @@ void FDeferredShadingSceneRenderer::Render(FRDGBuilder& GraphBuilder)
 			RDG_GPU_MASK_SCOPE(GraphBuilder, View.GPUMask);
 
 			Scene->GPUScene.UploadDynamicPrimitiveShaderDataForView(GraphBuilder, Scene, View);
+
+			Scene->GPUScene.DebugRender(GraphBuilder, *Scene, View);
 		}
 
 		{
diff --git a/Engine/Source/Runtime/Renderer/Private/GPUScene.cpp b/Engine/Source/Runtime/Renderer/Private/GPUScene.cpp
--- a/Engine/Source/Runtime/Renderer/Private/GPUScene.cpp
+++ b/Engine/Source/Runtime/Renderer/Private/GPUScene.cpp
@@ -20,6 +20,7 @@
 #include "NaniteSceneProxy.h"
 #include "HAL/LowLevelMemTracker.h"
 #include "HAL/LowLevelMemStats.h"
+#include "ShaderDebug.h"
 
 int32 GGPUSceneUploadEveryFrame = 0;
 FAutoConsoleVariableRef CVarGPUSceneUploadEveryFrame(
@@ -69,6 +70,25 @@ FAutoConsoleVariableRef CVarGPUSceneInstanceBVH(
 	ECVF_RenderThreadSafe | ECVF_ReadOnly
 );
 
+static TAutoConsoleVariable<int32> CVarGPUSceneDebugMode(
+	TEXT("r.GPUScene.DebugMode"),
+	0,
+	TEXT("Debug Rendering Mode:\n")
+	TEXT("0 - (show nothing, decault)\n")
+	TEXT(" 1 - Draw All\n")
+	TEXT(" 2 - Draw Selected (in the editor)\n")
+	TEXT(" 3 - Draw Updated (updated this frame)\n")
+	TEXT("You can use r.GPUScene.DebugDrawRange to limit the range\n"),
+	ECVF_RenderThreadSafe
+);
+
+static TAutoConsoleVariable<float> CVarGPUSceneDebugDrawRange(
+	TEXT("r.GPUScene.DebugDrawRange"),
+	-1.0f,
+	TEXT("Maximum distance the to draw instance bounds, the default is -1.0 <=> infinite range."),
+	ECVF_RenderThreadSafe
+);
+
 LLM_DECLARE_TAG_API(GPUScene, RENDERER_API);
 DECLARE_LLM_MEMORY_STAT(TEXT("GPUScene"), STAT_GPUSceneLLM, STATGROUP_LLMFULL);
 DECLARE_LLM_MEMORY_STAT(TEXT("GPUScene"), STAT_GPUSceneSummaryLLM, STATGROUP_LLM);
@@ -1554,6 +1574,113 @@ void FGPUScene::FreeInstancePayloadDataSlots(int32 InstancePayloadDataOffset, in
 	}
 }
 
+class FGPUSceneDebugRenderCS : public FGlobalShader
+{
+	DECLARE_GLOBAL_SHADER(FGPUSceneDebugRenderCS);
+	SHADER_USE_PARAMETER_STRUCT(FGPUSceneDebugRenderCS, FGlobalShader);
+
+	BEGIN_SHADER_PARAMETER_STRUCT(FParameters, )
+		SHADER_PARAMETER_STRUCT_INCLUDE(ShaderDrawDebug::FShaderParameters, ShaderDrawUniformBuffer)
+		SHADER_PARAMETER_SRV(StructuredBuffer<float4>, GPUSceneInstanceSceneData)
+		SHADER_PARAMETER_SRV(StructuredBuffer<float4>, GPUScenePrimitiveSceneData)
+		SHADER_PARAMETER(uint32, InstanceDataSOAStride)
+		SHADER_PARAMETER(uint32, GPUSceneFrameNumber)
+		SHADER_PARAMETER(int32, NumInstances)
+		SHADER_PARAMETER(int32, NumScenePrimitives)
+		SHADER_PARAMETER(int32, bDrawAll)
+		SHADER_PARAMETER(int32, bDrawUpdatedOnly)
+		SHADER_PARAMETER_RDG_BUFFER_SRV(StructuredBuffer< uint32 >, SelectedPrimitiveFlags)
+		SHADER_PARAMETER(FVector, PickingRayStart)
+		SHADER_PARAMETER(FVector, PickingRayEnd)
+		SHADER_PARAMETER(float, DrawRange)
+	END_SHADER_PARAMETER_STRUCT()
+
+public:
+	static constexpr uint32 NumThreadsPerGroup = 128U;
+
+	static bool ShouldCompilePermutation(const FGlobalShaderPermutationParameters& Parameters) 
+	{ 
+		return UseGPUScene(Parameters.Platform) && ShaderDrawDebug::IsSupported(Parameters.Platform);;
+	}
+
+	static void ModifyCompilationEnvironment(const FGlobalShaderPermutationParameters& Parameters, FShaderCompilerEnvironment& OutEnvironment)
+	{
+		FGlobalShader::ModifyCompilationEnvironment(Parameters, OutEnvironment);
+
+		OutEnvironment.SetDefine(TEXT("VF_SUPPORTS_PRIMITIVE_SCENE_DATA"), 1);
+		OutEnvironment.SetDefine(TEXT("USE_GLOBAL_GPU_SCENE_DATA"), 1);
+		OutEnvironment.SetDefine(TEXT("NUM_THREADS_PER_GROUP"), NumThreadsPerGroup);
+
+		// Skip optimization for avoiding long compilation time due to large UAV writes
+		OutEnvironment.CompilerFlags.Add(CFLAG_Debug);
+	}
+};
+
+IMPLEMENT_GLOBAL_SHADER(FGPUSceneDebugRenderCS, "/Engine/Private/GPUSceneDebugRender.usf", "GPUSceneDebugRenderCS", SF_Compute);
+
+
+void FGPUScene::DebugRender(FRDGBuilder& GraphBuilder, FScene& Scene, FViewInfo& View)
+{
+	int32 DebugMode = CVarGPUSceneDebugMode.GetValueOnRenderThread();
+	if (DebugMode > 0)
+	{
+		ShaderDrawDebug::SetEnabled(true);
+
+		int32 NumInstances = InstanceSceneDataAllocator.GetMaxSize();
+		if (ShaderDrawDebug::IsEnabled(View) && NumInstances > 0)
+		{
+			// This lags by one frame, so may miss some in one frame, also overallocates since we will cull a lot.
+			ShaderDrawDebug::RequestSpaceForElements(NumInstances * 12);
+
+			FGPUSceneDebugRenderCS::FParameters* PassParameters = GraphBuilder.AllocParameters<FGPUSceneDebugRenderCS::FParameters>();
+
+			TArray<uint32> SelectedPrimitiveFlags;
+			const int32 BitsPerWord = (sizeof(uint32) * 8U);
+			SelectedPrimitiveFlags.Init(0U, FMath::DivideAndRoundUp(Scene.Primitives.Num(), BitsPerWord));
+			for (int32 PrimitiveID = 0; PrimitiveID < Scene.PrimitiveSceneProxies.Num(); ++PrimitiveID)
+			{
+				if (Scene.PrimitiveSceneProxies[PrimitiveID]->IsSelected())
+				{
+					SelectedPrimitiveFlags[PrimitiveID / BitsPerWord] |= 1U << uint32(PrimitiveID % BitsPerWord);
+				}
+			}
+			FRDGBufferRef SelectedPrimitiveFlagsRDG = CreateStructuredBuffer(GraphBuilder, TEXT("GPUScene.Debug.SelectedPrimitiveFlags"), SelectedPrimitiveFlags);
+
+			ShaderDrawDebug::SetParameters(GraphBuilder, View.ShaderDrawData, PassParameters->ShaderDrawUniformBuffer);
+			PassParameters->GPUSceneInstanceSceneData = InstanceSceneDataBuffer.SRV;
+			PassParameters->GPUScenePrimitiveSceneData = PrimitiveBuffer.SRV;
+			PassParameters->InstanceDataSOAStride = InstanceSceneDataSOAStride;
+			PassParameters->GPUSceneFrameNumber = GetSceneFrameNumber();
+			PassParameters->bDrawUpdatedOnly = DebugMode == 3;
+			PassParameters->bDrawAll = DebugMode != 2;
+			PassParameters->NumInstances = NumInstances;
+			PassParameters->SelectedPrimitiveFlags = GraphBuilder.CreateSRV(SelectedPrimitiveFlagsRDG);
+			PassParameters->NumScenePrimitives = NumScenePrimitives;
+			PassParameters->DrawRange = CVarGPUSceneDebugDrawRange.GetValueOnRenderThread();
+
+			FVector PickingRayStart(ForceInit);
+			FVector PickingRayDir(ForceInit);
+			View.DeprojectFVector2D(View.CursorPos, PickingRayStart, PickingRayDir);
+
+			PassParameters->PickingRayStart = PickingRayStart;
+			PassParameters->PickingRayEnd = PickingRayStart + PickingRayDir * WORLD_MAX;
+
+			auto ComputeShader = View.ShaderMap->GetShader<FGPUSceneDebugRenderCS>();
+
+			FComputeShaderUtils::AddPass(
+				GraphBuilder,
+				RDG_EVENT_NAME("GPUSceneDebugRender"),
+				ComputeShader,
+				PassParameters,
+				FComputeShaderUtils::GetGroupCount(NumInstances, FGPUSceneDebugRenderCS::NumThreadsPerGroup)
+			);
+		}
+	}
+}
+
+
+
+
 TRange<int32> FGPUScene::CommitPrimitiveCollector(FGPUScenePrimitiveCollector& PrimitiveCollector)
 {
 	ensure(bInBeginEndBlock);
diff --git a/Engine/Source/Runtime/Renderer/Private/GPUScene.h b/Engine/Source/Runtime/Renderer/Private/GPUScene.h
--- a/Engine/Source/Runtime/Renderer/Private/GPUScene.h
+++ b/Engine/Source/Runtime/Renderer/Private/GPUScene.h
@@ -298,6 +298,12 @@ public:
 	int32 GetNumPrimitives() const { return DynamicPrimitivesOffset; }
 	int32 GetNumLightmapDataItems() const { return LightmapDataAllocator.GetMaxSize(); }
 
+	/**
+	 * Draw GPU-Scene debug info, such as bounding boxes. Call once per view at some point in the frame after GPU scene has been updated fully.
+	 * What is drawn is controlled by the CVar: r.GPUScene.DebugMode. Enabling this cvar causes ShaderDraw to be being active (if supported). 
+	 */
+	void DebugRender(FRDGBuilder& GraphBuilder, FScene& Scene, FViewInfo& View);
+
 	bool bUpdateAllPrimitives;
 
 	/** Indices of primitives that need to be updated in GPU Scene */
