CIS fix for LWC types in GPUScene debug shader parameters

diff --git a/Engine/Source/Runtime/Renderer/Private/GPUScene.cpp b/Engine/Source/Runtime/Renderer/Private/GPUScene.cpp
--- a/Engine/Source/Runtime/Renderer/Private/GPUScene.cpp
+++ b/Engine/Source/Runtime/Renderer/Private/GPUScene.cpp
@@ -1590,8 +1590,8 @@ class FGPUSceneDebugRenderCS : public FGlobalShader
 		SHADER_PARAMETER(int32, bDrawAll)
 		SHADER_PARAMETER(int32, bDrawUpdatedOnly)
 		SHADER_PARAMETER_RDG_BUFFER_SRV(StructuredBuffer< uint32 >, SelectedPrimitiveFlags)
-		SHADER_PARAMETER(FVector, PickingRayStart)
-		SHADER_PARAMETER(FVector, PickingRayEnd)
+		SHADER_PARAMETER(FVector3f, PickingRayStart)
+		SHADER_PARAMETER(FVector3f, PickingRayEnd)
 		SHADER_PARAMETER(float, DrawRange)
 	END_SHADER_PARAMETER_STRUCT()
 
