* Fix shader debug default resources
* Enable shader draw debug on a wider range of platform.

diff --git a/Engine/Shaders/Private/ShaderDrawDebug.usf b/Engine/Shaders/Private/ShaderDrawDebug.usf
--- a/Engine/Shaders/Private/ShaderDrawDebug.usf
+++ b/Engine/Shaders/Private/ShaderDrawDebug.usf
@@ -20,11 +20,11 @@ void ShaderDrawDebugVS(
 	out float4 Position : SV_POSITION,
 	out float4 ColorAlpha : TEXCOORD0)
 {
-	FPackedShaderDrawElement PackedLine = ShaderDrawDebugPrimitive[InstanceId];
-	FShaderDrawElement Line = UnpackShaderElement(PackedLine);
+	FPackedShaderDrawElement PackedItem = ShaderDrawDebugPrimitive[InstanceId];
+	FShaderDrawElement Item = UnpackShaderElement(PackedItem);
 
-	float4 VertexPosition = float4(VertexId == 0 ? Line.Pos0 : Line.Pos1, 1.0f);
-	float4 VertexColor = VertexId == 0 ? Line.Color0 : Line.Color1;
+	float4 VertexPosition = float4(VertexId == 0 ? Item.Pos0 : Item.Pos1, 1.0f);
+	float4 VertexColor = VertexId == 0 ? Item.Color0 : Item.Color1;
 
 	Position = mul(VertexPosition, View.WorldToClip);
 	ColorAlpha = VertexColor;
@@ -52,7 +52,7 @@ void ShaderDrawDebugPS(
 
 #if GPU_DEBUG_RENDERING_CLEAR_CS
 
-RWStructuredBuffer<FPackedShaderDrawElement>	DataBuffer;
+RWStructuredBuffer<FPackedShaderDrawElement> ElementBuffer; // Referenced so that the resource is marked as used
 RWBuffer<uint> IndirectBuffer;
 
 [numthreads(1, 1, 1)]
@@ -63,4 +63,4 @@ void ShaderDrawDebugClearCS(uint3 DispatchThreadId : SV_DispatchThreadID)
 	IndirectBuffer[2] = 0;
 	IndirectBuffer[3] = 0;
 }
-#endif
\ No newline at end of file
+#endif
diff --git a/Engine/Shaders/Private/ShaderDrawDebug.ush b/Engine/Shaders/Private/ShaderDrawDebug.ush
--- a/Engine/Shaders/Private/ShaderDrawDebug.ush
+++ b/Engine/Shaders/Private/ShaderDrawDebug.ush
@@ -67,7 +67,7 @@ RWStructuredBuffer<FPackedShaderDrawElement>OutShaderDrawPrimitive;
 RWBuffer<uint>								OutputShaderDrawIndirect;
 
 int2 ShaderDrawCursorPos;
-uint ShaderDrawMaxElementCount; 
+uint ShaderDrawMaxElementCount;
 
 int2 GetCursorPos()
 {
@@ -76,10 +76,15 @@ int2 GetCursorPos()
 
 void AddLine(float3 Pos0, float3 Pos1, float4 Color0, float4 Color1)
 {
+	if (ShaderDrawMaxElementCount == 0)
+	{
+		return;
+	}
+
 	uint PreviousPrimitiveCount = 0;
 	InterlockedAdd(OutputShaderDrawIndirect[1], 1, PreviousPrimitiveCount);
 	if ((PreviousPrimitiveCount+1) <= ShaderDrawMaxElementCount)
-    {
+	{
 		FShaderDrawElement Element;
 		Element.Pos0 = Pos0;
 		Element.Pos1 = Pos1;
@@ -87,7 +92,7 @@ void AddLine(float3 Pos0, float3 Pos1, float4 Color0, float4 Color1)
 		Element.Color1 = Color1;
 
 		OutShaderDrawPrimitive[PreviousPrimitiveCount] = PackShaderElement(Element);
-    }
+	}
 	else
 	{
 		InterlockedAdd(OutputShaderDrawIndirect[1], -1, PreviousPrimitiveCount);// Make sure we stay in the safe zone for indirect draw
@@ -101,6 +106,11 @@ void AddLine(float3 Pos0, float3 Pos1, float4 Color)
 
 void AddQuad(float3 Pos0, float3 Pos1, float3 Pos2, float3 Pos3, float4 Color)
 {
+	if (ShaderDrawMaxElementCount == 0)
+	{
+		return;
+	}
+
 	uint PreviousPrimitiveCount = 0;
 	InterlockedAdd(OutputShaderDrawIndirect[1], 4, PreviousPrimitiveCount);
 	if ((PreviousPrimitiveCount+4) <= ShaderDrawMaxElementCount)
@@ -128,6 +138,11 @@ void AddQuad(float3 Pos0, float3 Pos1, float3 Pos2, float3 Pos3, float4 Color)
 
 void AddAABB(float3 Min, float3 Max, float4 Color)
 {
+	if (ShaderDrawMaxElementCount == 0)
+	{
+		return;
+	}
+
 	uint PreviousPrimitiveCount = 0;
 	InterlockedAdd(OutputShaderDrawIndirect[1], 12, PreviousPrimitiveCount);
 	if ((PreviousPrimitiveCount + 12) <= ShaderDrawMaxElementCount)
@@ -205,6 +220,11 @@ void AddLineTriangle(float3 P0, float3 P1, float3 P2, float4 Color)
 
 void AddOBB(float3 Min, float3 Max, float4 Color, float4x4 Transform)
 {
+	if (ShaderDrawMaxElementCount == 0)
+	{
+		return;
+	}
+
 	uint PreviousPrimitiveCount = 0;
 	InterlockedAdd(OutputShaderDrawIndirect[1], 12, PreviousPrimitiveCount);
 	if ((PreviousPrimitiveCount + 12) <= ShaderDrawMaxElementCount)
diff --git a/Engine/Source/Runtime/Renderer/Private/ShaderDebug.cpp b/Engine/Source/Runtime/Renderer/Private/ShaderDebug.cpp
--- a/Engine/Source/Runtime/Renderer/Private/ShaderDebug.cpp
+++ b/Engine/Source/Runtime/Renderer/Private/ShaderDebug.cpp
@@ -48,7 +48,7 @@ namespace ShaderDrawDebug
 
 	bool IsSupported(const EShaderPlatform Platform)
 	{
-		return IsFeatureLevelSupported(Platform, ERHIFeatureLevel::SM5) && IsPCPlatform(Platform) && !IsOpenGLPlatform(Platform);
+		return RHISupportsComputeShaders(Platform) && !IsHlslccShaderPlatform(Platform);
 	}
 
 	void SetEnabled(bool bInEnabled)
@@ -61,12 +61,6 @@ namespace ShaderDrawDebug
 		GShaderDrawDebug_MaxElementCount = FMath::Max3(1024, GShaderDrawDebug_MaxElementCount, int32(MaxCount));
 	}
 
-	uint32 GetMaxElementCount()
-	{
-		return uint32(FMath::Max(1, GShaderDrawDebug_MaxElementCount));
-	}
-
-
 	void RequestSpaceForElements(uint32 MaxElementCount)
 	{
 		GElementRequestCount += MaxElementCount;
@@ -95,7 +89,7 @@ namespace ShaderDrawDebug
 		SHADER_USE_PARAMETER_STRUCT(FShaderDrawDebugClearCS, FGlobalShader);
 
 		BEGIN_SHADER_PARAMETER_STRUCT(FParameters, )
-		SHADER_PARAMETER_RDG_BUFFER_UAV(RWStructuredBuffer, DataBuffer)
+		SHADER_PARAMETER_RDG_BUFFER_UAV(RWStructuredBuffer, ElementBuffer)
 		SHADER_PARAMETER_RDG_BUFFER_UAV(RWBuffer, IndirectBuffer)
 		END_SHADER_PARAMETER_STRUCT()
 
@@ -212,7 +206,7 @@ namespace ShaderDrawDebug
 		ClearUnusedGraphResources(VertexShader, &PassParameters->ShaderDrawVSParameters, { IndirectBuffer });
 
 		GraphBuilder.AddPass(
-			RDG_EVENT_NAME("ShaderDrawDebug"),
+			RDG_EVENT_NAME("ShaderDebug::Draw"),
 			PassParameters,
 			ERDGPassFlags::Raster,
 			[VertexShader, PixelShader, PassParameters, IndirectBuffer](FRHICommandList& RHICmdList)
@@ -241,41 +235,48 @@ namespace ShaderDrawDebug
 			});
 	}
 
-	void BeginView(FRDGBuilder& GraphBuilder, FViewInfo& View)
+	static void AddShaderDrawDebugClearPass(FRDGBuilder& GraphBuilder, FViewInfo& View, FRDGBufferRef& DataBuffer, FRDGBufferRef& IndirectBuffer)
 	{
-		if (!IsEnabled(View))
-		{
-			return;
-		}
-
-		// Update max element count from the request count and reset
-		GShaderDrawDebug_MaxElementCount = FMath::Max(GShaderDrawDebug_MaxElementCount, int32(GElementRequestCount));
-		GElementRequestCount = 0;
-
-		FSceneViewState* ViewState = View.ViewState;
-
-		const bool bLockBufferThisFrame = IsShaderDrawLocked() && ViewState && !ViewState->ShaderDrawDebugStateData.bIsLocked;
-		ERDGBufferFlags Flags = bLockBufferThisFrame ? ERDGBufferFlags::MultiFrame : ERDGBufferFlags::None;
-
-		FRDGBufferRef DataBuffer = GraphBuilder.CreateBuffer(FRDGBufferDesc::CreateStructuredDesc(sizeof(FPackedShaderDrawElement), GetMaxElementCount()), TEXT("ShaderDrawDataBuffer"), Flags);
-		FRDGBufferRef IndirectBuffer = GraphBuilder.CreateBuffer(FRDGBufferDesc::CreateIndirectDesc<FRHIDrawIndirectParameters>(1), TEXT("ShaderDrawDataIndirectBuffer"), Flags);
-
 		FShaderDrawDebugClearCS::FParameters* Parameters = GraphBuilder.AllocParameters<FShaderDrawDebugClearCS::FParameters>();
-		Parameters->DataBuffer = GraphBuilder.CreateUAV(DataBuffer);
-		Parameters->IndirectBuffer = GraphBuilder.CreateUAV(IndirectBuffer);
+		Parameters->ElementBuffer	= GraphBuilder.CreateUAV(DataBuffer);
+		Parameters->IndirectBuffer	= GraphBuilder.CreateUAV(IndirectBuffer);
 
 		TShaderMapRef<FShaderDrawDebugClearCS> ComputeShader(View.ShaderMap);
 
 		// Note: we do not call ClearUnusedGraphResources here as we want to for the allocation of DataBuffer
 		GraphBuilder.AddPass(
-			RDG_EVENT_NAME("ShaderDrawClear"),
+			RDG_EVENT_NAME("ShaderDebug::Clear"),
 			Parameters,
 			ERDGPassFlags::Compute,
 			[Parameters, ComputeShader](FRHICommandList& RHICmdList)
 		{
 			FComputeShaderUtils::Dispatch(RHICmdList, ComputeShader, *Parameters, FIntVector(1,1,1));
 		});
+	}
 
+	void BeginView(FRDGBuilder& GraphBuilder, FViewInfo& View)
+	{
+		View.ShaderDrawData = FShaderDrawDebugData();
+		if (!IsEnabled(View))
+		{
+			// Default resources
+			View.ShaderDrawData.Buffer			= GraphBuilder.CreateBuffer(FRDGBufferDesc::CreateStructuredDesc(sizeof(FPackedShaderDrawElement), 1), TEXT("ShaderDraw.DataBuffer(Dummy)"), ERDGBufferFlags::None);
+			View.ShaderDrawData.IndirectBuffer	= GraphBuilder.CreateBuffer(FRDGBufferDesc::CreateIndirectDesc<FRHIDrawIndirectParameters>(1), TEXT("ShaderDraw.IndirectBuffer(Dummy)"), ERDGBufferFlags::None);
+			return;
+		}
+
+		// Update max element count from the request count and reset
+		GShaderDrawDebug_MaxElementCount = FMath::Max3(1, GShaderDrawDebug_MaxElementCount, int32(GElementRequestCount));
+		View.ShaderDrawData.MaxElementCount = GShaderDrawDebug_MaxElementCount;
+		GElementRequestCount = 0;
+
+		FSceneViewState* ViewState = View.ViewState;
+		const bool bLockBufferThisFrame = IsShaderDrawLocked() && ViewState && !ViewState->ShaderDrawDebugStateData.bIsLocked;
+		ERDGBufferFlags Flags = bLockBufferThisFrame ? ERDGBufferFlags::MultiFrame : ERDGBufferFlags::None;
+
+		FRDGBufferRef DataBuffer = GraphBuilder.CreateBuffer(FRDGBufferDesc::CreateStructuredDesc(sizeof(FPackedShaderDrawElement), View.ShaderDrawData.MaxElementCount), TEXT("ShaderDraw.DataBuffer"), Flags);
+		FRDGBufferRef IndirectBuffer = GraphBuilder.CreateBuffer(FRDGBufferDesc::CreateIndirectDesc<FRHIDrawIndirectParameters>(1), TEXT("ShaderDraw.IndirectBuffer"), Flags);
+		AddShaderDrawDebugClearPass(GraphBuilder, View, DataBuffer, IndirectBuffer);
 
 		View.ShaderDrawData.Buffer = DataBuffer;
 		View.ShaderDrawData.IndirectBuffer = IndirectBuffer;
@@ -340,13 +341,10 @@ namespace ShaderDrawDebug
 		if (Data.Buffer == nullptr)
 			return;
 
-		FRDGBufferRef DataBuffer = Data.Buffer;
-		FRDGBufferRef IndirectBuffer = Data.IndirectBuffer;
-
-		OutParameters.ShaderDrawCursorPos = Data.CursorPosition;
-		OutParameters.ShaderDrawMaxElementCount = GetMaxElementCount();
-		OutParameters.OutShaderDrawPrimitive = GraphBuilder.CreateUAV(DataBuffer);
-		OutParameters.OutputShaderDrawIndirect = GraphBuilder.CreateUAV(IndirectBuffer);
+		OutParameters.ShaderDrawCursorPos		= Data.CursorPosition;
+		OutParameters.ShaderDrawMaxElementCount = Data.MaxElementCount;
+		OutParameters.OutShaderDrawPrimitive	= GraphBuilder.CreateUAV(Data.Buffer);
+		OutParameters.OutputShaderDrawIndirect	= GraphBuilder.CreateUAV(Data.IndirectBuffer);
 	}
 
 	// Returns true if the default view exists and has shader debug rendering enabled (this needs to be checked before using a permutation that requires the shader draw parameters)
diff --git a/Engine/Source/Runtime/Renderer/Public/ShaderDebug.h b/Engine/Source/Runtime/Renderer/Public/ShaderDebug.h
--- a/Engine/Source/Runtime/Renderer/Public/ShaderDebug.h
+++ b/Engine/Source/Runtime/Renderer/Public/ShaderDebug.h
@@ -10,9 +10,11 @@ class FViewInfo;
 
 struct FShaderDrawDebugData
 {
-	FIntPoint CursorPosition;
+	uint32 MaxElementCount = 0u;
+	FIntPoint CursorPosition = FIntPoint(-1,-1);
 	FRDGBufferRef Buffer = nullptr;
 	FRDGBufferRef IndirectBuffer = nullptr;
+	bool IsEnabled() const { return MaxElementCount > 0; }
 	bool IsValid() const { return Buffer != nullptr && IndirectBuffer != nullptr; }
 };
 
@@ -26,7 +28,6 @@ namespace ShaderDrawDebug
 
 	RENDERER_API void SetEnabled(bool bEnable);
 	RENDERER_API void SetMaxElementCount(uint32 MaxCount);
-	RENDERER_API uint32 GetMaxElementCount();
 
 	/**
 	 * Call to ensure enough space for some number of elements, is added cumulatively each frame, to make it possible for several systems to request a certain number independently.
