* Fix missing clear on scene velocity texture, when hair strands is the first pass to write to it.
* Change velocity pass to use hair tile data.

diff --git a/Engine/Shaders/Private/HairStrands/HairStrandsTileCommon.ush b/Engine/Shaders/Private/HairStrands/HairStrandsTileCommon.ush
--- a/Engine/Shaders/Private/HairStrands/HairStrandsTileCommon.ush
+++ b/Engine/Shaders/Private/HairStrands/HairStrandsTileCommon.ush
@@ -2,6 +2,9 @@
 
 #pragma once
 
+// Default tile size
+#define HAIR_TILE_SIZE 8
+
 // Tile types
 #define HAIRTILE_HAIR_ALL 0 
 #define HAIRTILE_HAIR_FULL 1 
diff --git a/Engine/Shaders/Private/HairStrands/HairStrandsVelocity.usf b/Engine/Shaders/Private/HairStrands/HairStrandsVelocity.usf
--- a/Engine/Shaders/Private/HairStrands/HairStrandsVelocity.usf
+++ b/Engine/Shaders/Private/HairStrands/HairStrandsVelocity.usf
@@ -4,18 +4,20 @@
 
 #include "../Common.ush"
 #include "HairStrandsVisibilityCommon.ush"
+#include "HairStrandsTileCommon.ush"
 
-#if PERMUTATION_GROUPSIZE == 64
 #define TILE_PIXEL_SIZE_X 8
 #define TILE_PIXEL_SIZE_Y 8
-#elif PERMUTATION_GROUPSIZE == 32
-#define TILE_PIXEL_SIZE_X 8
-#define TILE_PIXEL_SIZE_Y 4
+
+#if (TILE_PIXEL_SIZE_X != HAIR_TILE_SIZE) || (TILE_PIXEL_SIZE_Y != HAIR_TILE_SIZE) 
+#error Tile size needs to match 
 #endif
 
-uint2 ResolutionOffset;
+int2 Resolution;
+int2 ResolutionOffset;
 float VelocityThreshold;
 float CoverageThreshold;
+uint bNeedClear;
 
 Texture2D<uint> NodeIndex;
 Buffer<float2> NodeVelocity;
@@ -24,7 +26,7 @@ Texture2D<float> CoverageTexture;
 
 #define VELOCITY_TYPE_NONE 0
 #define VELOCITY_TYPE_AVG 1  
-#define VELOCITY_TYPE_CLOSEST 2
+#define VELOCITY_TYPE_CLOSEST 2	
 #define VELOCITY_TYPE_MAX 3
 #if PERMUTATION_VELOCITY != VELOCITY_TYPE_NONE
   #if PERMUTATION_OUTPUT_FORMAT == 0
@@ -34,6 +36,13 @@ Texture2D<float> CoverageTexture;
   #endif
 #endif
 
+#if PERMUTATION_TILE
+uint TileSize;
+int2 TileCountXY;
+Buffer<uint> TileCountBuffer;
+Buffer<uint2> TileDataBuffer;
+#endif
+
 RWTexture2D<uint> OutResolveMaskTexture;
 
 #if PERMUTATION_OUTPUT_FORMAT == 0
@@ -43,11 +52,25 @@ float4 ToOutput(float2 In) { return float4(In, 0, 0); }
 #endif
 
 [numthreads(TILE_PIXEL_SIZE_X, TILE_PIXEL_SIZE_Y, 1)]
-void MainCS(uint GroupIndex : SV_GroupIndex, uint3 DispatchThreadId : SV_DispatchThreadID)
+void MainCS(uint3 DispatchThreadId : SV_DispatchThreadID, uint2 InGroupId : SV_GroupID, uint2 InGroupThreadId : SV_GroupThreadID)
 {
+#if PERMUTATION_TILE
+	const uint TileCount = TileCountBuffer[HAIRTILE_HAIR_ALL];
+	const uint TileIndex1D = InGroupId.x + InGroupId.y * TileCountXY.x;
+	if (TileIndex1D >= TileCount)
+	{
+		return;
+	}
+	const uint2 GroupId = TileDataBuffer[TileIndex1D];
+	const uint2 PixelCoord = GroupId * TileSize + InGroupThreadId;
+#else
 	const uint2 PixelCoord = DispatchThreadId.xy + ResolutionOffset;
+#endif
+
 	const FNodeDesc NodeDesc = DecodeNodeDesc(NodeIndex.Load(uint3(PixelCoord, 0)));
 	bool bNeedFastResolve = false;
+	bool bValidVelocity = NodeDesc.Count > 0;
+	float2 OutEncodedVelocity = 0;
 	if (NodeDesc.Count > 0)
 	{
 		// Store final sort node data 
@@ -103,18 +126,24 @@ void MainCS(uint GroupIndex : SV_GroupIndex, uint3 DispatchThreadId : SV_Dispatc
 			}
 
 			#if PERMUTATION_VELOCITY == VELOCITY_TYPE_AVG
-			OutVelocityTexture[PixelCoord] = ToOutput(EncodeVelocityToTexture(float3(AverageVelocity / max(NodeDesc.Count, 1u), 0.0)).xy);
+			OutEncodedVelocity = EncodeVelocityToTexture(float3(AverageVelocity / max(NodeDesc.Count, 1u), 0.0)).xy;
 			#endif		
 
 			#if PERMUTATION_VELOCITY == VELOCITY_TYPE_CLOSEST
-			OutVelocityTexture[PixelCoord] = ToOutput(ClosestEncodedVelocity);
+			OutEncodedVelocity = ClosestEncodedVelocity;
 			#endif		
 
 			#if PERMUTATION_VELOCITY == VELOCITY_TYPE_MAX
-			OutVelocityTexture[PixelCoord] = ToOutput(MaxEncodedVelocity);
+			OutEncodedVelocity = MaxEncodedVelocity;
 			#endif
 		}
 	}
 
+	#if PERMUTATION_VELOCITY != VELOCITY_TYPE_NONE
+	if (bValidVelocity || bNeedClear)
+	{
+		OutVelocityTexture[PixelCoord]  = bValidVelocity ? ToOutput(OutEncodedVelocity) : 0.f;
+	}
+	#endif
 	OutResolveMaskTexture[PixelCoord] = bNeedFastResolve ? 1 : 0;
 }
diff --git a/Engine/Shaders/Private/HairStrands/HairStrandsVisibilityTile.usf b/Engine/Shaders/Private/HairStrands/HairStrandsVisibilityTile.usf
--- a/Engine/Shaders/Private/HairStrands/HairStrandsVisibilityTile.usf
+++ b/Engine/Shaders/Private/HairStrands/HairStrandsVisibilityTile.usf
@@ -5,8 +5,7 @@
 #include "HairStrandsVisibilityCommon.ush"
 #include "HairStrandsTileCommon.ush"
 
-#define TILE_SIZE 8
-#define GROUP_THREAD_COUNT (TILE_SIZE*TILE_SIZE)
+#define GROUP_THREAD_COUNT (HAIR_TILE_SIZE*HAIR_TILE_SIZE)
 
 ////////////////////////////////////////////////////////////////////////////////////////////////////////////
 
@@ -28,7 +27,7 @@ RWBuffer<uint2> TileOtherBuffer;
 groupshared uint2 s_fMask[GROUP_THREAD_COUNT];
 groupshared uint2 s_pMask[GROUP_THREAD_COUNT];
 
-[numthreads(TILE_SIZE, TILE_SIZE, 1)]
+[numthreads(HAIR_TILE_SIZE, HAIR_TILE_SIZE, 1)]
 void TileMainCS(uint2 DispatchThreadId : SV_DispatchThreadID, uint LinearIndex : SV_GroupIndex, uint3 GroupId : SV_GroupID)
 {
 	const uint2 PixelCoord = DispatchThreadId;
@@ -187,7 +186,7 @@ void MainPS(
 	in FScreenVertexOutput Input,
 	out float4 OutColor0 : SV_Target0)
 {
-	const uint2 InTileCoord = uint2(Input.Position.xy) >> 3u; // TILE_SIZE == 8;
+	const uint2 InTileCoord = uint2(Input.Position.xy) >> 3u; // HAIR_TILE_SIZE == 8;
 	const bool bTileX = (InTileCoord.x & 1) == 0;
 	const bool bTileY = (InTileCoord.y & 1) == 0;
 	const bool bChecker = (bTileX && bTileY) || (!bTileX && !bTileY);
@@ -331,11 +330,43 @@ void MainVS(
 	Out = (FScreenVertexOutput)0;
 	const uint2 TileCoord = TileDataBuffer[InInstanceId];
 
-	uint2 TileVertex = TileCoord * TILE_SIZE;
-	TileVertex.x += InVertexId == 1 || InVertexId == 2 || InVertexId == 4 ? TILE_SIZE : 0;
-	TileVertex.y += InVertexId == 2 || InVertexId == 4 || InVertexId == 5 ? TILE_SIZE : 0;
+	uint2 TileVertex = TileCoord * HAIR_TILE_SIZE;
+	TileVertex.x += InVertexId == 1 || InVertexId == 2 || InVertexId == 4 ? HAIR_TILE_SIZE : 0;
+	TileVertex.y += InVertexId == 2 || InVertexId == 4 || InVertexId == 5 ? HAIR_TILE_SIZE : 0;
 	Out.UV = float2(TileVertex - ViewMin) * ViewInvSize;
 	Out.Position = float4(Out.UV * float2(2.0f, -2.0f) + float2(-1.0, 1.0f), 0.5f, 1.0f);
 }
 
 #endif //SHADER_TILE_VS
+
+////////////////////////////////////////////////////////////////////////////////////////////////////////////
+
+#if SHADER_TILE_CLEAR
+
+uint TileSize;
+uint TileType;
+int2 TileCountXY;
+int2 Resolution;
+Buffer<uint>  TileCountBuffer;
+Buffer<uint2> TileDataBuffer;
+
+RWTexture2D<float4> OutTexture;
+
+[numthreads(HAIR_TILE_SIZE, HAIR_TILE_SIZE, 1)]
+void TileMainCS(uint2 InGroupId : SV_GroupID, uint2 InGroupThreadId : SV_GroupThreadID)
+{
+	const uint TileCount   = TileCountBuffer[TileType];
+	const uint TileIndex1D = InGroupId.x + InGroupId.y * TileCountXY.x;
+	if (TileIndex1D >= TileCount)
+	{
+		return;
+	}
+	const uint2 GroupId = TileDataBuffer[TileIndex1D];
+	const uint2 PixelCoord = GroupId * TileSize + InGroupThreadId;
+
+	if (all(PixelCoord < uint2(Resolution)))
+	{
+		OutTexture[PixelCoord] = 0;
+	}
+}
+#endif // SHADER_TILE_CLEAR
\ No newline at end of file
diff --git a/Engine/Source/Runtime/Renderer/Private/HairStrands/HairStrandsTile.cpp b/Engine/Source/Runtime/Renderer/Private/HairStrands/HairStrandsTile.cpp
--- a/Engine/Source/Runtime/Renderer/Private/HairStrands/HairStrandsTile.cpp
+++ b/Engine/Source/Runtime/Renderer/Private/HairStrands/HairStrandsTile.cpp
@@ -378,4 +378,69 @@ void AddHairStrandsDebugTilePass(
 			Parameters,
 			FIntVector(1, 1, 1));
 	}
+}
+
+////////////////////////////////////////////////////////////////////////////////////////////////////////////////////////////////////////////////
+
+class FHairTileClearCS : public FGlobalShader
+{
+	DECLARE_GLOBAL_SHADER(FHairTileClearCS);
+	SHADER_USE_PARAMETER_STRUCT(FHairTileClearCS, FGlobalShader);
+
+	BEGIN_SHADER_PARAMETER_STRUCT(FParameters, )
+		SHADER_PARAMETER(FIntPoint, Resolution)
+		SHADER_PARAMETER(FIntPoint, TileCountXY)
+		SHADER_PARAMETER(uint32, TileSize)
+		SHADER_PARAMETER(uint32, TileType)
+		SHADER_PARAMETER_RDG_BUFFER_SRV(Buffer<uint2>, TileDataBuffer)
+		SHADER_PARAMETER_RDG_BUFFER_SRV(Buffer<uint>, TileCountBuffer)
+		SHADER_PARAMETER_RDG_BUFFER_SRV(Buffer<uint>, HairTileCount)
+		SHADER_PARAMETER_RDG_TEXTURE_UAV(RWTexture2D, OutTexture)
+		RDG_BUFFER_ACCESS(TileIndirectArgs, ERHIAccess::IndirectArgs)
+	END_SHADER_PARAMETER_STRUCT()
+
+public:
+
+	static bool ShouldCompilePermutation(const FGlobalShaderPermutationParameters& Parameters)
+	{
+		return IsHairStrandsSupported(EHairStrandsShaderType::Strands, Parameters.Platform);
+	}
+
+	static void ModifyCompilationEnvironment(const FGlobalShaderPermutationParameters& Parameters, FShaderCompilerEnvironment& OutEnvironment)
+	{
+		FGlobalShader::ModifyCompilationEnvironment(Parameters, OutEnvironment);
+		OutEnvironment.SetDefine(TEXT("SHADER_TILE_CLEAR"), 1);
+	}
+};
+
+IMPLEMENT_GLOBAL_SHADER(FHairTileClearCS, "/Engine/Private/HairStrands/HairStrandsVisibilityTile.usf", "TileMainCS", SF_Compute);
+
+void AddHairStrandsTileClearPass(
+	FRDGBuilder& GraphBuilder,
+	const FViewInfo& View,
+	const FHairStrandsTiles& TileData,
+	FHairStrandsTiles::ETileType TileType,
+	FRDGTextureRef OutTexture)
+{
+	if (!OutTexture || !TileData.IsValid() || TileType == FHairStrandsTiles::ETileType::Count) return;
+	
+
+	FHairTileClearCS::FParameters* PassParameters = GraphBuilder.AllocParameters<FHairTileClearCS::FParameters>();
+	PassParameters->Resolution = OutTexture->Desc.Extent;
+	PassParameters->TileCountXY = TileData.TileCountXY;
+	PassParameters->TileSize = TileData.TileSize;
+	PassParameters->TileType = ToIndex(TileType);
+	PassParameters->TileCountBuffer = TileData.TileCountSRV;
+	PassParameters->TileDataBuffer = TileData.GetTileBufferSRV(TileType);
+	PassParameters->TileIndirectArgs = TileData.TileIndirectDispatchBuffer;
+	PassParameters->OutTexture = GraphBuilder.CreateUAV(OutTexture);
+
+	FHairTileClearCS::FPermutationDomain PermutationVector;
+	TShaderMapRef<FHairTileClearCS> ComputeShader(View.ShaderMap, PermutationVector);
+	FComputeShaderUtils::AddPass(
+		GraphBuilder,
+		RDG_EVENT_NAME("HairStrands::TileClear(%s)", ToString(TileType)),
+		ComputeShader,
+		PassParameters,
+		PassParameters->TileIndirectArgs, TileData.GetIndirectDispatchArgOffset(TileType));
 }
\ No newline at end of file
diff --git a/Engine/Source/Runtime/Renderer/Private/HairStrands/HairStrandsTile.h b/Engine/Source/Runtime/Renderer/Private/HairStrands/HairStrandsTile.h
--- a/Engine/Source/Runtime/Renderer/Private/HairStrands/HairStrandsTile.h
+++ b/Engine/Source/Runtime/Renderer/Private/HairStrands/HairStrandsTile.h
@@ -40,6 +40,13 @@ FHairStrandsTiles AddHairStrandsGenerateTilesPass(
 	const FViewInfo& View,
 	const FRDGTextureRef& InputTexture);
 
+void AddHairStrandsTileClearPass(
+	FRDGBuilder& GraphBuilder,
+	const FViewInfo& View,
+	const FHairStrandsTiles& TileData,
+	FHairStrandsTiles::ETileType TileType,
+	FRDGTextureRef OutTexture);
+
 void AddHairStrandsDebugTilePass(
 	FRDGBuilder& GraphBuilder,
 	const FViewInfo& View,
diff --git a/Engine/Source/Runtime/Renderer/Private/HairStrands/HairStrandsVisibility.cpp b/Engine/Source/Runtime/Renderer/Private/HairStrands/HairStrandsVisibility.cpp
--- a/Engine/Source/Runtime/Renderer/Private/HairStrands/HairStrandsVisibility.cpp
+++ b/Engine/Source/Runtime/Renderer/Private/HairStrands/HairStrandsVisibility.cpp
@@ -931,15 +931,18 @@ class FHairVelocityCS: public FGlobalShader
 	DECLARE_GLOBAL_SHADER(FHairVelocityCS);
 	SHADER_USE_PARAMETER_STRUCT(FHairVelocityCS, FGlobalShader);
 
-	class FGroupSize : SHADER_PERMUTATION_SPARSE_INT("PERMUTATION_GROUPSIZE", 32, 64);
 	class FVelocity : SHADER_PERMUTATION_INT("PERMUTATION_VELOCITY", 4);
 	class FOuputFormat : SHADER_PERMUTATION_INT("PERMUTATION_OUTPUT_FORMAT", 2);
-	using FPermutationDomain = TShaderPermutationDomain<FGroupSize, FVelocity, FOuputFormat>;
+	class FTile : SHADER_PERMUTATION_BOOL("PERMUTATION_TILE");
+	using FPermutationDomain = TShaderPermutationDomain<FVelocity, FOuputFormat, FTile>;
 
 	BEGIN_SHADER_PARAMETER_STRUCT(FParameters, )
+
+		SHADER_PARAMETER(FIntPoint, Resolution)
 		SHADER_PARAMETER(FIntPoint, ResolutionOffset)
 		SHADER_PARAMETER(float, VelocityThreshold)
 		SHADER_PARAMETER(float, CoverageThreshold)
+		SHADER_PARAMETER(uint32, bNeedClear)
 		SHADER_PARAMETER_RDG_TEXTURE(Texture2D, CoverageTexture)
 		SHADER_PARAMETER_RDG_TEXTURE(Texture2D, NodeIndex)
 		SHADER_PARAMETER_RDG_BUFFER_SRV(Buffer, NodeVelocity)
@@ -948,6 +951,13 @@ class FHairVelocityCS: public FGlobalShader
 		SHADER_PARAMETER_RDG_TEXTURE_UAV(RWTexture2D, OutResolveMaskTexture)
 		SHADER_PARAMETER_STRUCT_REF(FViewUniformShaderParameters, ViewUniformBuffer)
 		SHADER_PARAMETER_RDG_UNIFORM_BUFFER(FSceneTextureUniformParameters, SceneTexturesStruct)
+
+		SHADER_PARAMETER(FIntPoint, TileCountXY)
+		SHADER_PARAMETER(uint32, TileSize)
+		SHADER_PARAMETER_RDG_BUFFER_SRV(Buffer<uint2>, TileDataBuffer)
+		SHADER_PARAMETER_RDG_BUFFER_SRV(Buffer<uint>, TileCountBuffer)
+		SHADER_PARAMETER_RDG_BUFFER_SRV(Buffer<uint>, HairTileCount)
+		RDG_BUFFER_ACCESS(TileIndirectArgs, ERHIAccess::IndirectArgs)
 	END_SHADER_PARAMETER_STRUCT()
 
 public:
@@ -966,6 +976,7 @@ static void AddHairVelocityPass(
 	FRDGBuilder& GraphBuilder,
 	const FViewInfo& View,
 	const FHairStrandsMacroGroupDatas& MacroGroupDatas,
+	const FHairStrandsTiles& TileData,
 	FRDGTextureRef& CoverageTexture,
 	FRDGTextureRef& NodeIndex,
 	FRDGBufferRef& NodeVis,
@@ -977,26 +988,35 @@ static void AddHairVelocityPass(
 	if (!bWriteOutVelocity)
 		return;
 
-	if (!HasBeenProduced(OutVelocityTexture))
+	// If velocity texture has not been created by the base-pass, clear it here
+	const bool bNeedClear = !HasBeenProduced(OutVelocityTexture);
+	if (bNeedClear)
 	{
-		AddClearUAVPass(GraphBuilder, GraphBuilder.CreateUAV(OutVelocityTexture), 0.f);
+		if (!TileData.IsValid())
+		{
+			AddClearUAVPass(GraphBuilder, GraphBuilder.CreateUAV(OutVelocityTexture), 0.f);
+		}
+		else
+		{
+			AddHairStrandsTileClearPass(GraphBuilder, View, TileData, FHairStrandsTiles::ETileType::Other, OutVelocityTexture);
+		}
 	}
 
+	const bool bUseTile = TileData.IsValid();
+
 	const FIntPoint Resolution = OutVelocityTexture->Desc.Extent;
-	{
-		FRDGTextureDesc Desc = FRDGTextureDesc::Create2D(Resolution, PF_R8_UINT, FClearValueBinding::None, TexCreate_UAV);
-		OutResolveMaskTexture = GraphBuilder.CreateTexture(Desc, TEXT("Hair.VelocityResolveMaskTexture"));
-	}
+	OutResolveMaskTexture = GraphBuilder.CreateTexture(FRDGTextureDesc::Create2D(Resolution, PF_R8_UINT, FClearValueBinding::None, TexCreate_UAV), TEXT("Hair.VelocityResolveMaskTexture"));
 
 	check(OutVelocityTexture->Desc.Format == PF_G16R16 || OutVelocityTexture->Desc.Format == PF_A16B16G16R16);
 	const bool bTwoChannelsOutput = OutVelocityTexture->Desc.Format == PF_G16R16;
 
 	FHairVelocityCS::FPermutationDomain PermutationVector;
-	PermutationVector.Set<FHairVelocityCS::FGroupSize>(GetVendorOptimalGroupSize1D());
 	PermutationVector.Set<FHairVelocityCS::FVelocity>(bWriteOutVelocity ? FMath::Clamp(GHairVelocityType + 1, 0, 3) : 0);
 	PermutationVector.Set<FHairVelocityCS::FOuputFormat>(bTwoChannelsOutput ? 0 : 1);
+	PermutationVector.Set<FHairVelocityCS::FTile>(bUseTile);
 
 	FHairVelocityCS::FParameters* PassParameters = GraphBuilder.AllocParameters<FHairVelocityCS::FParameters>();
+	PassParameters->bNeedClear = bNeedClear ? 1u : 0u;
 	PassParameters->SceneTexturesStruct = CreateSceneTextureUniformBuffer(GraphBuilder, View.FeatureLevel);
 	PassParameters->ViewUniformBuffer = View.ViewUniformBuffer;
 	PassParameters->VelocityThreshold = GetHairFastResolveVelocityThreshold(Resolution);
@@ -1008,29 +1028,44 @@ static void AddHairVelocityPass(
 	PassParameters->OutVelocityTexture = GraphBuilder.CreateUAV(OutVelocityTexture);
 	PassParameters->OutResolveMaskTexture = GraphBuilder.CreateUAV(OutResolveMaskTexture);
 
-	const FIntPoint GroupSize = GetVendorOptimalGroupSize2D();
-	// We don't use the CPU screen projection for running the velocity pass, as we need to clear the entire 
-	// velocity mask through the UAV write, otherwise the mask will be partially invalid.
-#if 1
-	FIntRect TotalRect = View.ViewRect; 
-#else
-	// Snap the rect onto thread group boundary
-	FIntRect TotalRect = ComputeVisibleHairStrandsMacroGroupsRect(View.ViewRect, MacroGroupDatas);
-	TotalRect.Min.X = FMath::FloorToInt(float(TotalRect.Min.X) / float(GroupSize.X)) * GroupSize.X;
-	TotalRect.Min.Y = FMath::FloorToInt(float(TotalRect.Min.Y) / float(GroupSize.Y)) * GroupSize.Y;
-	TotalRect.Max.X = FMath::CeilToInt(float(TotalRect.Max.X) / float(GroupSize.X)) * GroupSize.X;
-	TotalRect.Max.Y = FMath::CeilToInt(float(TotalRect.Max.Y) / float(GroupSize.Y)) * GroupSize.Y;
-#endif
-	FIntPoint RectResolution(TotalRect.Width(), TotalRect.Height());
-	PassParameters->ResolutionOffset = FIntPoint(TotalRect.Min.X, TotalRect.Min.Y);
+	if (bUseTile)
+	{
+		const FHairStrandsTiles::ETileType TileType = FHairStrandsTiles::ETileType::HairAll;
 
-	TShaderMapRef<FHairVelocityCS> ComputeShader(View.ShaderMap, PermutationVector);
-	FComputeShaderUtils::AddPass(
-		GraphBuilder,
-		RDG_EVENT_NAME("HairStrands::Velocity"),
-		ComputeShader,
-		PassParameters,
-		FComputeShaderUtils::GetGroupCount(RectResolution, GroupSize));
+		PassParameters->ResolutionOffset	= FIntPoint(0,0);
+		PassParameters->Resolution			= Resolution;
+		PassParameters->TileCountXY			= TileData.TileCountXY;
+		PassParameters->TileSize			= TileData.TileSize;
+		PassParameters->TileCountBuffer		= TileData.TileCountSRV;
+		PassParameters->TileDataBuffer		= TileData.GetTileBufferSRV(TileType);
+		PassParameters->TileIndirectArgs	= TileData.TileIndirectDispatchBuffer;
+
+		TShaderMapRef<FHairVelocityCS> ComputeShader(View.ShaderMap, PermutationVector);
+		FComputeShaderUtils::AddPass(
+			GraphBuilder,
+			RDG_EVENT_NAME("HairStrands::Velocity(Tile)"),
+			ComputeShader,
+			PassParameters,
+			PassParameters->TileIndirectArgs, TileData.GetIndirectDispatchArgOffset(TileType));
+	}
+	else
+	{
+		// We don't use the CPU screen projection for running the velocity pass, as we need to clear the entire 
+		// velocity mask through the UAV write, otherwise the mask will be partially invalid.
+		const FIntRect TotalRect = View.ViewRect; 
+		const FIntPoint RectResolution(TotalRect.Width(), TotalRect.Height());
+
+		PassParameters->ResolutionOffset = FIntPoint(TotalRect.Min.X, TotalRect.Min.Y);
+		PassParameters->Resolution		 = RectResolution;
+
+		TShaderMapRef<FHairVelocityCS> ComputeShader(View.ShaderMap, PermutationVector);
+		FComputeShaderUtils::AddPass(
+			GraphBuilder,
+			RDG_EVENT_NAME("HairStrands::Velocity(Screen)"),
+			ComputeShader,
+			PassParameters,
+			FComputeShaderUtils::GetGroupCount(RectResolution, FIntPoint(8, 8)));
+	}
 }
 
 /////////////////////////////////////////////////////////////////////////////////////////
@@ -1643,7 +1678,7 @@ static void AddHairVisibilityPrimitiveIdCompactionPass(
 	FRDGBufferRef& OutCompactNodeVis, // Or OutCompactNodeData for PPLL
 	FRDGBufferRef& OutCompactNodeCoord,
 	FRDGTextureRef& OutCoverageTexture,
-	FRDGTextureRef& OutVelocityTexture,
+	FRDGTextureRef OutVelocityTexture,
 	FRDGBufferRef& OutIndirectArgsBuffer,
 	uint32& OutMaxRenderNodeCount)
 {
@@ -1708,7 +1743,7 @@ static void AddHairVisibilityPrimitiveIdCompactionPass(
 		OutCompactNodeCoord = GraphBuilder.CreateBuffer(FRDGBufferDesc::CreateBufferDesc(sizeof(uint32), MaxRenderNodeCount), TEXT("Hair.VisibilityNodeCoord"));
 	}
 
-	const bool bWriteOutVelocity = OutVelocityTexture != nullptr;
+	const bool bWriteOutVelocity = OutVelocityTexture != nullptr && bUsePPLL; // Velocity write out is only support with PPLL
 	const uint32 VelocityPermutation = bWriteOutVelocity ? FMath::Clamp(GHairVelocityType + 1, 0, 3) : 0;
 	FHairVisibilityPrimitiveIdCompactionCS::FPermutationDomain PermutationVector;
 	PermutationVector.Set<FHairVisibilityPrimitiveIdCompactionCS::FGroupSize>(bUseTile ? FHairStrandsTiles::GroupSize : GetVendorOptimalGroupSize1D());
@@ -1847,7 +1882,6 @@ static void AddHairVisibilityCompactionComputeRasterPass(
 	FRDGBufferRef&  OutCompactNodeVis,
 	FRDGBufferRef&  OutCompactNodeCoord,
 	FRDGTextureRef& OutCoverageTexture,
-	FRDGTextureRef& OutVelocityTexture,
 	FRDGBufferRef&  OutIndirectArgsBuffer,
 	uint32& OutMaxRenderNodeCount)
 {	
@@ -3210,7 +3244,6 @@ void RenderHairStrandsVisibilityBuffer(
 							CompactNodeVis,
 							CompactNodeCoord,
 							CoverageTexture,
-							SceneVelocityTexture,
 							IndirectArgsBuffer,
 							VisibilityData.MaxNodeCount);
 
@@ -3236,6 +3269,7 @@ void RenderHairStrandsVisibilityBuffer(
 							GraphBuilder,
 							View,
 							MacroGroupDatas,
+							VisibilityData.TileData,
 							CoverageTexture,
 							CompactNodeIndex,
 							CompactNodeVis,
@@ -3380,7 +3414,7 @@ void RenderHairStrandsVisibilityBuffer(
 						CompactNodeVis,
 						CompactNodeCoord,
 						CoverageTexture,
-						SceneVelocityTexture,
+						nullptr, // Velocity output is only needed for PPLL
 						IndirectArgsBuffer,
 						VisibilityData.MaxNodeCount);
 
@@ -3409,6 +3443,7 @@ void RenderHairStrandsVisibilityBuffer(
 							GraphBuilder,
 							View,
 							MacroGroupDatas,
+							VisibilityData.TileData,
 							CoverageTexture,
 							CompactNodeIndex,
 							CompactNodeVis,
