Add debug display for hair tangent vector (strands/cards/meshes).

diff --git a/Engine/Shaders/Private/HairStrands/HairStrandsDebugPrint.usf b/Engine/Shaders/Private/HairStrands/HairStrandsDebugPrint.usf
--- a/Engine/Shaders/Private/HairStrands/HairStrandsDebugPrint.usf
+++ b/Engine/Shaders/Private/HairStrands/HairStrandsDebugPrint.usf
@@ -1,11 +1,35 @@
 // Copyright Epic Games, Inc. All Rights Reserved.
 
 #include "../Common.ush"
+#include "../CommonViewUniformBuffer.ush"
+#include "../SceneTextureParameters.ush"
+#include "../DeferredShadingCommon.ush"
 #include "../ShaderPrintCommon.ush"
 #include "../ShaderDrawDebug.ush"
 #include "HairStrandsVisibilityCommon.ush"
 #include "HairStrandsAABBCommon.ush" 
 
+
+////////////////////////////////////////////////////////////////////////////////////////////////
+
+float3 ReconstructWorldPositionFromDepth(float2 UV, float Depth)
+{
+	float2 ScreenPosition = (UV - View.ScreenPositionScaleBias.wz) / View.ScreenPositionScaleBias.xy;
+	float4 HomogeneousWorldPosition = mul(float4(ScreenPosition * Depth, Depth, 1), View.ScreenToWorld);
+	float3 WorldPosition = HomogeneousWorldPosition.xyz / HomogeneousWorldPosition.w;
+
+	return WorldPosition;
+}
+
+void AddNewLine(float2 OriginalPos, inout float2 Pos)
+{
+	Pos = ShaderPrintNewline(Pos); Pos.x = OriginalPos.x;
+}
+
+////////////////////////////////////////////////////////////////////////////////////////////////
+
+#if SHADER_PRINT
+
 int2 PixelCoord;
 int2 MaxResolution;
 uint FastResolveMask;
@@ -16,7 +40,7 @@ uint HairVisibilityNodeGroupSize;
 Texture2D<float> HairCountTexture;
 Texture2D<uint> HairCountUintTexture;
 
-Texture2D<uint2> DepthStencilTexture;
+Texture2D<uint2> StencilTexture;
 SamplerState LinearSampler;
 Buffer<uint> HairVisibilityIndirectArgsBuffer;
 
@@ -25,11 +49,6 @@ Buffer<int> HairMacroGroupAABBBuffer;
 #define MATERIAL_DATA 1
 #define DETAIL_DATA 1
 
-void AddNewLine(float2 OriginalPos, inout float2 Pos)
-{
-	Pos = ShaderPrintNewline(Pos); Pos.x = OriginalPos.x;
-}
-
 [numthreads(1, 1, 1)]
 void MainCS(uint GroupIndex : SV_GroupIndex, uint3 DispatchThreadId : SV_DispatchThreadID)
 {	
@@ -129,7 +148,7 @@ void MainCS(uint GroupIndex : SV_GroupIndex, uint3 DispatchThreadId : SV_Dispatc
 	AddNewLine(OriginalPos, Pos);
 
 	// Resolve type
-	const uint Stencil = DepthStencilTexture.Load(uint3(PixelPos,0)).y;
+	const uint Stencil = StencilTexture.Load(uint3(PixelPos,0)).y;
 	const bool bIsFastResolve = (Stencil & FastResolveMask) != 0;
 	Pos = ShaderPrintText(Pos, TEXT("Resolve         "));
 
@@ -221,3 +240,131 @@ void MainCS(uint GroupIndex : SV_GroupIndex, uint3 DispatchThreadId : SV_Dispatc
 		Pos = ShaderPrint(Pos, Sample.Coverage8bit, FontOrange);						AddNewLine(StartSamplePos, Pos);
 	}
 }
+
+#endif // SHADER_PRINT
+
+////////////////////////////////////////////////////////////////////////////////////////////////
+
+#if SHADER_TANGENT
+uint TileSize;
+int2 TileCount;
+float2 OutputResolution;
+RWTexture2D<float4> OutputTexture;
+
+#define HAIR_TANGENT_INVALID 0
+#define HAIR_TANGENT_STRANDS 1
+#define HAIR_TANGENT_CARDS   2
+
+struct FTangentDesc
+{
+	float3 Tangent;
+	uint Type;
+};
+
+FTangentDesc AddTangentPlot(uint2 PixelCoord, float2 UV)
+{
+	FTangentDesc Out = (FTangentDesc)0;
+
+	// Strands
+	{
+		// Find closest
+		float3 Closest_Tangent = 0;
+		float  Closest_Depth = 0;
+
+		uint Total8BitCoverage = 0;
+		const FNodeDesc NodeDesc = DecodeNodeDesc(HairStrands.HairSampleOffset.Load(uint3(PixelCoord, 0)));
+		for (uint SampleIt0 = 0; SampleIt0 < NodeDesc.Count; SampleIt0++)
+		{
+			const uint LocalOffset = NodeDesc.Offset + SampleIt0;
+			const FPackedHairSample NodeData = HairStrands.HairSampleData[LocalOffset];
+			const FHairSample Sample = UnpackHairSample(NodeData);
+
+			if (Sample.Depth > Closest_Depth)
+			{
+				Closest_Depth = Sample.Depth;
+				Closest_Tangent = Sample.Tangent;
+			}
+		}
+
+		// Draw tangent vector (strands)
+		if (NodeDesc.Count > 0)
+		{
+			const float SceneDepth = ConvertFromDeviceZ(Closest_Depth);
+			const float3 WorldPosition = ReconstructWorldPositionFromDepth(UV, SceneDepth);
+			AddLine(WorldPosition, WorldPosition + Closest_Tangent, ColorRed, ColorYellow);
+			Out.Type = 1;
+			Out.Tangent = Closest_Tangent;
+		}
+	}
+
+	// Cards
+	if (Out.Type == 0)
+	{
+		FScreenSpaceData ScreenSpaceData = GetScreenSpaceData(UV);
+		const float SceneDepth = CalcSceneDepth(UV);
+		const float3 WorldPosition = ReconstructWorldPositionFromDepth(UV, SceneDepth);
+		if (ScreenSpaceData.GBuffer.ShadingModelID == SHADINGMODELID_HAIR)
+		{
+			const float3 Tangent = ScreenSpaceData.GBuffer.WorldNormal;
+			AddLine(WorldPosition, WorldPosition + Tangent, ColorPurple, ColorYellow);
+			Out.Type = 2;
+			Out.Tangent = Tangent;
+		}
+	}
+
+	return Out;
+}
+
+void AddTangentText(float2 UV, float3 Tangent, uint TangentType)
+{
+	const float2 OriginalPos = UV;
+	float2 Pos = OriginalPos;
+	if (TangentType == HAIR_TANGENT_STRANDS)
+		Pos = ShaderPrintText(Pos, TEXT("Strands"), FontYellow);
+	else if (TangentType == HAIR_TANGENT_CARDS)
+		Pos = ShaderPrintText(Pos, TEXT("Cards "), FontYellow);
+	else if (TangentType == HAIR_TANGENT_CARDS)
+		Pos = ShaderPrintText(Pos, TEXT("Invalid"), FontRed);
+
+	AddNewLine(OriginalPos, Pos);
+
+	Pos = ShaderPrint(Pos, Tangent.x, FontRed);		AddNewLine(OriginalPos, Pos);
+	Pos = ShaderPrint(Pos, Tangent.y, FontGreen);	AddNewLine(OriginalPos, Pos);
+	Pos = ShaderPrint(Pos, Tangent.z, FontBlue);	AddNewLine(OriginalPos, Pos);
+}
+
+[numthreads(8, 8, 1)]
+void MainCS(int2 DispatchThreadId : SV_DispatchThreadID)
+{
+	// Cursor plot
+	if (all(DispatchThreadId == 0))
+	{
+		const float2 PixelCoord = GetCursorPos();// DispatchThreadId.xy + 0.5f;
+		const float2 UV = PixelCoord / float2(OutputResolution);
+		const float2 UVText = (PixelCoord + float2(20, 0)) / float2(OutputResolution);
+
+		FTangentDesc Desc = AddTangentPlot(PixelCoord, UV);
+		if (Desc.Type == HAIR_TANGENT_STRANDS)
+		{			
+			AddTangentText(UVText, Desc.Tangent, HAIR_TANGENT_STRANDS);
+		}
+		else if (Desc.Type == HAIR_TANGENT_CARDS)
+		{
+			AddTangentText(UVText, Desc.Tangent, HAIR_TANGENT_CARDS);
+		}
+	}
+
+	// Grid plot
+	if (all(DispatchThreadId.xy < TileCount))
+	{
+		const float2 PixelCoord = DispatchThreadId.xy * TileSize;
+		if (all(PixelCoord < OutputResolution))
+		{
+			const float2 UV = PixelCoord / float2(OutputResolution);
+			AddTangentPlot(PixelCoord, UV);
+		}
+	}
+}
+#endif // SHADER_TANGENT
+
+////////////////////////////////////////////////////////////////////////////////////////////////
\ No newline at end of file
diff --git a/Engine/Source/Runtime/Renderer/Private/HairStrands/HairStrandsDebug.cpp b/Engine/Source/Runtime/Renderer/Private/HairStrands/HairStrandsDebug.cpp
--- a/Engine/Source/Runtime/Renderer/Private/HairStrands/HairStrandsDebug.cpp
+++ b/Engine/Source/Runtime/Renderer/Private/HairStrands/HairStrandsDebug.cpp
@@ -58,6 +58,11 @@ static FAutoConsoleVariableRef CVarHairStrandsDebugMaterialSampleIndex(TEXT("r.H
 static int32 GHairStrandsClusterDebug = 0;
 static FAutoConsoleVariableRef CVarHairStrandsDebugClusterAABB(TEXT("r.HairStrands.Cluster.Debug"), GHairStrandsClusterDebug, TEXT("Draw debug the world bounding box of hair clusters used for culling optimisation (0:off, 1:visible cluster, 2:culled cluster, 3:colored LOD, 4:LOD info)."));
 
+static int32 GHairTangentDebug = 0;
+static int32 GHairTangentDebug_TileSize = 8;
+static FAutoConsoleVariableRef CVarHairTangentDebug(TEXT("r.HairStrands.DebugMode.Tangent"), GHairTangentDebug, TEXT("Draw debug tangent for hair strands and hair cards."));
+static FAutoConsoleVariableRef CVarHairTangentDebug_TileSize(TEXT("r.HairStrands.DebugMode.Tangent.TileSize"), GHairTangentDebug_TileSize, TEXT("Draw debug tangent - Grid size for drawing debug tangent"));
+
 static int32 GHairVirtualVoxel_DrawDebugPage = 0;
 static FAutoConsoleVariableRef CVarHairVirtualVoxel_DrawDebugPage(TEXT("r.HairStrands.Voxelization.Virtual.DrawDebugPage"), GHairVirtualVoxel_DrawDebugPage, TEXT("When voxel debug rendering is enable 1: render the page bounds, instead of the voxel 2: the occupancy within the page (i.e., 8x8x8 brick)"));
 static int32 GHairVirtualVoxel_ForceMipLevel = -1;
@@ -237,7 +242,7 @@ class FHairDebugPrintCS : public FGlobalShader
 		SHADER_PARAMETER_RDG_TEXTURE(Texture2D, HairCountUintTexture)
 		SHADER_PARAMETER_RDG_BUFFER_SRV(Buffer<uint>, HairVisibilityIndirectArgsBuffer)
 		SHADER_PARAMETER_RDG_BUFFER_SRV(Buffer, HairMacroGroupAABBBuffer)
-		SHADER_PARAMETER_RDG_TEXTURE_SRV(Texture2D, DepthStencilTexture)
+		SHADER_PARAMETER_RDG_TEXTURE_SRV(Texture2D, StencilTexture)
 		SHADER_PARAMETER_SAMPLER(SamplerState, LinearSampler)
 		SHADER_PARAMETER_STRUCT_REF(FViewUniformShaderParameters, ViewUniformBuffer)
 		SHADER_PARAMETER_STRUCT_INCLUDE(ShaderPrint::FShaderParameters, ShaderPrintUniformBuffer)
@@ -252,6 +257,7 @@ public:
 		// Skip optimization for avoiding long compilation time due to large UAV writes
 		FGlobalShader::ModifyCompilationEnvironment(Parameters, OutEnvironment);
 		OutEnvironment.CompilerFlags.Add(CFLAG_Debug);
+		OutEnvironment.SetDefine(TEXT("SHADER_PRINT"), 1);
 	}
 };
 
@@ -263,9 +269,9 @@ static void AddDebugHairPrintPass(
 	const EHairDebugMode InDebugMode,
 	const FHairStrandsVisibilityData& VisibilityData,
 	const FHairStrandsMacroGroupResources& MacroGroupResources,
-	FRDGTextureSRVRef InDepthStencilTexture)
+	FRDGTextureSRVRef InStencilTexture)
 {
-	if (!View || !View->HairStrandsViewData.UniformBuffer || !InDepthStencilTexture || !ShaderDrawDebug::IsEnabled(*View)) return;
+	if (!View || !View->HairStrandsViewData.UniformBuffer || !InStencilTexture || !ShaderDrawDebug::IsEnabled(*View)) return;
 
 	FRDGTextureRef ViewHairCountTexture = VisibilityData.ViewHairCountTexture ? VisibilityData.ViewHairCountTexture : GSystemTextures.GetBlackDummy(GraphBuilder);
 	FRDGTextureRef ViewHairCountUintTexture = VisibilityData.ViewHairCountUintTexture ? VisibilityData.ViewHairCountUintTexture : GSystemTextures.GetBlackDummy(GraphBuilder);
@@ -275,6 +281,7 @@ static void AddDebugHairPrintPass(
 
 	FHairDebugPrintCS::FParameters* Parameters = GraphBuilder.AllocParameters<FHairDebugPrintCS::FParameters>();
 	Parameters->GroupSize = GetVendorOptimalGroupSize2D();
+	Parameters->ViewUniformBuffer = View->ViewUniformBuffer;
 	Parameters->MaxResolution = VisibilityData.CoverageTexture ? VisibilityData.CoverageTexture->Desc.Extent : FIntPoint(0,0);
 	Parameters->PixelCoord = View->CursorPos;
 	Parameters->FastResolveMask = STENCIL_TEMPORAL_RESPONSIVE_AA_MASK;
@@ -282,7 +289,7 @@ static void AddDebugHairPrintPass(
 	Parameters->HairCountUintTexture = ViewHairCountUintTexture;
 	Parameters->HairVisibilityIndirectArgsBuffer = GraphBuilder.CreateSRV(VisibilityData.NodeIndirectArg, PF_R32_UINT);
 	Parameters->HairVisibilityNodeGroupSize = VisibilityData.NodeGroupSize;
-	Parameters->DepthStencilTexture = InDepthStencilTexture;
+	Parameters->StencilTexture = InStencilTexture;
 	Parameters->LinearSampler = TStaticSamplerState<SF_Bilinear, AM_Clamp, AM_Clamp, AM_Clamp>::GetRHI();
 	Parameters->HairMacroGroupCount = MacroGroupResources.MacroGroupCount;
 	Parameters->HairStrands = View->HairStrandsViewData.UniformBuffer;
@@ -295,7 +302,7 @@ static void AddDebugHairPrintPass(
 
 	FComputeShaderUtils::AddPass(
 		GraphBuilder,
-		RDG_EVENT_NAME("HairStrandsDebugPrint"),
+		RDG_EVENT_NAME("HairStrands::DebugPrint"),
 		ComputeShader,
 		Parameters,
 		FIntVector(1, 1, 1));
@@ -384,7 +391,7 @@ static void AddDebugHairPass(
 	ClearUnusedGraphResources(PixelShader, Parameters);
 
 	GraphBuilder.AddPass(
-		RDG_EVENT_NAME("HairStrandsDebug"),
+		RDG_EVENT_NAME("HairStrands::DebugMode(%s)", ToString(InDebugMode)),
 		Parameters,
 		ERDGPassFlags::Raster,
 		[Parameters, VertexShader, PixelShader, Viewport, Resolution](FRHICommandList& RHICmdList)
@@ -587,7 +594,7 @@ static void AddDeepShadowInfoPass(
 	TShaderMapRef<FDeepShadowInfoCS> ComputeShader(View.ShaderMap);
 
 	const FIntVector DispatchCount = DispatchCount.DivideAndRoundUp(FIntVector(OutputTexture->Desc.Extent.X, OutputTexture->Desc.Extent.Y, 1), FIntVector(8, 8, 1));
-	FComputeShaderUtils::AddPass(GraphBuilder, RDG_EVENT_NAME("HairStrandsDeepShadowDebugInfo"), ComputeShader, Parameters, DispatchCount);
+	FComputeShaderUtils::AddPass(GraphBuilder, RDG_EVENT_NAME("HairStrands::DeepShadowDebugInfo"), ComputeShader, Parameters, DispatchCount);
 }
 
 ///////////////////////////////////////////////////////////////////////////////////////////////////
@@ -664,8 +671,63 @@ static void AddVoxelPageRaymarchingPass(
 		TShaderMapRef<FVoxelVirtualRaymarchingCS> ComputeShader(View.ShaderMap, PermutationVector);
 
 		const FIntVector DispatchCount = DispatchCount.DivideAndRoundUp(FIntVector(OutputTexture->Desc.Extent.X, OutputTexture->Desc.Extent.Y, 1), FIntVector(8, 8, 1));
-		FComputeShaderUtils::AddPass(GraphBuilder, RDG_EVENT_NAME("HairStrandsVoxelVirtualRaymarching"), ComputeShader, Parameters, DispatchCount);
+		FComputeShaderUtils::AddPass(GraphBuilder, RDG_EVENT_NAME("HairStrands::VoxelVirtualRaymarching"), ComputeShader, Parameters, DispatchCount);
+	}
+}
+///////////////////////////////////////////////////////////////////////////////////////////////////
+
+class FDebugHairTangentCS : public FGlobalShader
+{
+	DECLARE_GLOBAL_SHADER(FDebugHairTangentCS);
+	SHADER_USE_PARAMETER_STRUCT(FDebugHairTangentCS, FGlobalShader);
+
+	BEGIN_SHADER_PARAMETER_STRUCT(FParameters, )
+		SHADER_PARAMETER_RDG_UNIFORM_BUFFER(FSceneTextureUniformParameters, SceneTextures)
+		SHADER_PARAMETER_STRUCT_INCLUDE(ShaderDrawDebug::FShaderParameters, ShaderDraw)
+		SHADER_PARAMETER_STRUCT_INCLUDE(ShaderPrint::FShaderParameters, ShaderPrint)
+		SHADER_PARAMETER(FVector2D, OutputResolution)
+		SHADER_PARAMETER(FIntPoint, TileCount)
+		SHADER_PARAMETER(uint32, TileSize)
+		SHADER_PARAMETER_STRUCT_REF(FViewUniformShaderParameters, ViewUniformBuffer)
+		SHADER_PARAMETER_SAMPLER(SamplerState, BilinearTextureSampler)
+		SHADER_PARAMETER_RDG_UNIFORM_BUFFER(FHairStrandsViewUniformParameters, HairStrands)
+	END_SHADER_PARAMETER_STRUCT()
+
+public:
+	static bool ShouldCompilePermutation(const FGlobalShaderPermutationParameters& Parameters) { return IsHairStrandsSupported(EHairStrandsShaderType::Tool, Parameters.Platform); }
+	static void ModifyCompilationEnvironment(const FGlobalShaderPermutationParameters& Parameters, FShaderCompilerEnvironment& OutEnvironment)
+	{
+		FGlobalShader::ModifyCompilationEnvironment(Parameters, OutEnvironment);
+		OutEnvironment.SetDefine(TEXT("SHADER_TANGENT"), 1);
 	}
+};
+
+IMPLEMENT_GLOBAL_SHADER(FDebugHairTangentCS, "/Engine/Private/HairStrands/HairStrandsDebugPrint.usf", "MainCS", SF_Compute);
+
+static void AddDebugHairTangentPass(
+	FRDGBuilder& GraphBuilder,
+	const FViewInfo& View,
+	const FSceneTextures& SceneTextures,
+	FRDGTextureRef& OutputTexture)
+{
+	ShaderDrawDebug::SetEnabled(true);
+
+	FDebugHairTangentCS::FParameters* Parameters = GraphBuilder.AllocParameters<FDebugHairTangentCS::FParameters>();
+	Parameters->ViewUniformBuffer		= View.ViewUniformBuffer;
+	Parameters->HairStrands				= View.HairStrandsViewData.UniformBuffer;
+	Parameters->OutputResolution		= OutputTexture->Desc.Extent;
+	Parameters->TileSize				= FMath::Clamp(GHairTangentDebug_TileSize, 4, 32);
+	Parameters->TileCount				= FIntPoint(FMath::FloorToInt(Parameters->OutputResolution.X / Parameters->TileSize), FMath::FloorToInt(Parameters->OutputResolution.X / Parameters->TileSize));
+	Parameters->SceneTextures			= SceneTextures.UniformBuffer;
+	Parameters->BilinearTextureSampler	= TStaticSamplerState<SF_Bilinear>::GetRHI();
+	ShaderDrawDebug::SetParameters(GraphBuilder, View.ShaderDrawData, Parameters->ShaderDraw);
+	ShaderPrint::SetParameters(GraphBuilder, View, Parameters->ShaderPrint);
+
+	const FIntVector DispatchCount = DispatchCount.DivideAndRoundUp(FIntVector(OutputTexture->Desc.Extent.X, OutputTexture->Desc.Extent.Y, 1), FIntVector(8, 8, 1));
+	ShaderDrawDebug::RequestSpaceForElements(DispatchCount.X * DispatchCount.Y);
+
+	TShaderMapRef<FDebugHairTangentCS> ComputeShader(View.ShaderMap);
+	FComputeShaderUtils::AddPass(GraphBuilder, RDG_EVENT_NAME("HairStrands::DebugTangentCS"), ComputeShader, Parameters, DispatchCount);
 }
 
 ///////////////////////////////////////////////////////////////////////////////////////////////////
@@ -726,7 +788,7 @@ static void AddPlotBSDFPass(
 	ClearUnusedGraphResources(PixelShader, Parameters);
 
 	GraphBuilder.AddPass(
-		RDG_EVENT_NAME("HairStrandsBsdfPlot"),
+		RDG_EVENT_NAME("HairStrands::BsdfPlot"),
 		Parameters,
 		ERDGPassFlags::Raster,
 		[Parameters, VertexShader, PixelShader, Viewport, Resolution](FRHICommandList& RHICmdList)
@@ -816,7 +878,7 @@ static void AddPlotSamplePass(
 	ClearUnusedGraphResources(PixelShader, Parameters);
 
 	GraphBuilder.AddPass(
-		RDG_EVENT_NAME("HairStrandsSamplePlot"),
+		RDG_EVENT_NAME("HairStrands::SamplePlot"),
 		Parameters,
 		ERDGPassFlags::Raster,
 		[Parameters, VertexShader, PixelShader, Viewport, Resolution](FRHICommandList& RHICmdList)
@@ -1035,6 +1097,12 @@ static void InternalRenderHairStrandsDebugInfo(
 		RunHairStrandsBookmark(GraphBuilder, EHairStrandsBookmark::ProcessDebug, Params);
 	}
 
+	// Display tangent vector for strands/cards/meshes
+	if (GHairTangentDebug > 0)
+	{
+		AddDebugHairTangentPass(GraphBuilder, View, SceneTextures, SceneColorTexture);
+	}
+
 	// Pass this point, all debug rendering concern only hair strands data
 	if (!HairStrands::HasViewHairStrandsData(View))
 	{
@@ -1255,7 +1323,7 @@ static void InternalRenderHairStrandsDebugInfo(
 			FHairVisibilityDebugPPLLCS::FPermutationDomain PermutationVector;
 			TShaderMapRef<FHairVisibilityDebugPPLLCS> ComputeShader(View.ShaderMap, PermutationVector);
 			FIntVector TextureSize = SceneColorTexture->Desc.GetSize(); TextureSize.Z = 1;
-			FComputeShaderUtils::AddPass(GraphBuilder, RDG_EVENT_NAME("HairPPLLDebug"), ComputeShader, PassParameters,
+			FComputeShaderUtils::AddPass(GraphBuilder, RDG_EVENT_NAME("HairStrands::PPLLDebug"), ComputeShader, PassParameters,
 				FIntVector::DivideAndRoundUp(TextureSize, FIntVector(8, 8, 1)));
 		}
 	}
