Reduces noise when history rejecting on noisy input

diff --git a/Engine/Shaders/Private/TemporalSuperResolution/TSRSpatialAntiAliasing.usf b/Engine/Shaders/Private/TemporalSuperResolution/TSRSpatialAntiAliasing.usf
--- a/Engine/Shaders/Private/TemporalSuperResolution/TSRSpatialAntiAliasing.usf
+++ b/Engine/Shaders/Private/TemporalSuperResolution/TSRSpatialAntiAliasing.usf
@@ -9,6 +9,8 @@
 
 #define CONFIG_MIN_LUMA_ALIASING 0.05
 
+#define CONFIG_COMPUTE_NOISE_FILTERING 1
+
 #if DIM_QUALITY_PRESET == 1
 	#define CONFIG_ITERATIONS 3
 
@@ -35,6 +37,7 @@
 Texture2D<taa_half> InputSceneColorLdrLumaTexture;
 
 RWTexture2D<uint> AntiAliasingOutput;
+RWTexture2D<taa_half> NoiseFilteringOutput;
 
 
 //------------------------------------------------------- FUNCTIONS
@@ -204,6 +207,7 @@ void MainCS(
 	//Debug[0] = 1.0;
 
 	taa_half2x2 TexelOffset = 0.0;
+	taa_half2 NoiseFiltering = 0.0;
 	{
 		#if CONFIG_ENABLE_DUAL_PIXEL_VECTORIZATION && 1
 			taa_half InputCache[3 * 4];
@@ -258,6 +262,18 @@ void MainCS(
 		taa_half2 TotalVarV  = abs(0.5 * (InputN  + InputS ) - InputC);
 		taa_half2 TotalVarVW = abs(0.5 * (InputNW + InputSW) - InputW);
 		
+		taa_half2 DiffN = abs(InputN - InputC);
+		taa_half2 DiffS = abs(InputS - InputC);
+		taa_half2 DiffE = abs(InputE - InputC);
+		taa_half2 DiffW = abs(InputW - InputC);
+
+		#if CONFIG_COMPUTE_NOISE_FILTERING
+		{
+			NoiseFiltering = max(NoiseFiltering, taa_half(2.0) * TotalVarH - max(DiffE, DiffW));
+			NoiseFiltering = max(NoiseFiltering, taa_half(2.0) * TotalVarV - max(DiffN, DiffS));
+		}
+		#endif
+
 		bBrowseVerticaly = (TotalVarHN + TotalVarH + TotalVarHS) > (TotalVarVE + TotalVarV + TotalVarVW);
 		BrowseDirection[0] = bBrowseVerticaly ? 0 : 1;
 		BrowseDirection[1] = bBrowseVerticaly ? 1 : 0;
@@ -265,25 +281,25 @@ void MainCS(
 		FLATTEN
 		if (bBrowseVerticaly[0])
 		{
-			EdgeSide[0][0] = abs(InputW[0] - InputC[0]) > abs(InputE[0] - InputC[0]) ? -1 : 1;
-			EdgeLuma[0] = abs(InputW[0] - InputC[0]) > abs(InputE[0] - InputC[0]) ? InputW[0] : InputE[0];
+			EdgeSide[0][0] = DiffW[0] > DiffE[0] ? -1 : 1;
+			EdgeLuma[0] = DiffW[0] > DiffE[0] ? InputW[0] : InputE[0];
 		}
 		else
 		{
-			EdgeSide[1][0] = abs(InputN[0] - InputC[0]) > abs(InputS[0] - InputC[0]) ? -1 : 1;
-			EdgeLuma[0] = abs(InputN[0] - InputC[0]) > abs(InputS[0] - InputC[0]) ? InputN[0] : InputS[0];
+			EdgeSide[1][0] = DiffN[0] > DiffS[0] ? -1 : 1;
+			EdgeLuma[0] = DiffN[0] > DiffS[0] ? InputN[0] : InputS[0];
 		}
 		
 		FLATTEN
 		if (bBrowseVerticaly[1])
 		{
-			EdgeSide[0][1] = abs(InputW[1] - InputC[1]) > abs(InputE[1] - InputC[1]) ? -1 : 1;
-			EdgeLuma[1] = abs(InputW[1] - InputC[1]) > abs(InputE[1] - InputC[1]) ? InputW[1] : InputE[1];
+			EdgeSide[0][1] = DiffW[1] > DiffE[1] ? -1 : 1;
+			EdgeLuma[1] = DiffW[1] > DiffE[1] ? InputW[1] : InputE[1];
 		}
 		else
 		{
-			EdgeSide[1][1] = abs(InputN[1] - InputC[1]) > abs(InputS[1] - InputC[1]) ? -1 : 1;
-			EdgeLuma[1] = abs(InputN[1] - InputC[1]) > abs(InputS[1] - InputC[1]) ? InputN[1] : InputS[1];
+			EdgeSide[1][1] = DiffN[1] > DiffS[1] ? -1 : 1;
+			EdgeLuma[1] = DiffN[1] > DiffS[1] ? InputN[1] : InputS[1];
 		}
 
 		taa_half2 LumaDelta = abs(EdgeLuma - InputC) * 0.5;
@@ -484,10 +500,12 @@ void MainCS(
 		
 		OutputPixelPos = InvalidateOutputPixelPos(OutputPixelPos, InputInfo_ViewportMax);
 		
-		AntiAliasingOutput[dpv_lo(OutputPixelPos)] = dpv_lo(EncodedTexelOffset); // dpv_lo(EncodedTexelOffset);
+		AntiAliasingOutput[dpv_lo(OutputPixelPos)] = dpv_lo(EncodedTexelOffset);
+		NoiseFilteringOutput[dpv_lo(OutputPixelPos)] = dpv_lo(NoiseFiltering);
 		
 		#if CONFIG_ENABLE_DUAL_PIXEL_VECTORIZATION
-			AntiAliasingOutput[dpv_hi(OutputPixelPos)] = dpv_hi(EncodedTexelOffset); // dpv_hi(EncodedTexelOffset);
+			AntiAliasingOutput[dpv_hi(OutputPixelPos)] = dpv_hi(EncodedTexelOffset);
+			NoiseFilteringOutput[dpv_hi(OutputPixelPos)] = dpv_hi(NoiseFiltering);
 		#endif
 
 		#if DEBUG_OUTPUT
diff --git a/Engine/Shaders/Private/TemporalSuperResolution/TSRUpdateHistory.usf b/Engine/Shaders/Private/TemporalSuperResolution/TSRUpdateHistory.usf
--- a/Engine/Shaders/Private/TemporalSuperResolution/TSRUpdateHistory.usf
+++ b/Engine/Shaders/Private/TemporalSuperResolution/TSRUpdateHistory.usf
@@ -61,12 +61,14 @@ Texture2D<float2> DilatedVelocityTexture;
 Texture2D<uint> ParallaxFactorTexture;
 Texture2D<taa_half> ParallaxRejectionMaskTexture;
 Texture2D<uint> AntiAliasingTexture;
+Texture2D<taa_half> NoiseFilteringTexture;
 
 float4 HistoryPixelPosToScreenPos;
 float4 HistoryPixelPosToPPCo;
 float3 HistoryQuantizationError;
 float MinTranslucencyRejection;
 float InvWeightClampingPixelSpeed;
+float InputToHistoryFactor;
 uint ResponsiveStencilMask;
 
 Texture2D<taa_half3> PrevHistory_LowFrequency;
@@ -235,6 +237,7 @@ void MainCS(
 	taa_half2 TranslucencyRejection = taa_half(1.0).xx;
 	taa_half2 OutputPixelVelocity = taa_half(0.0).xx;
 	taa_half2 MoireWeight = taa_half(1.0).xx;
+	taa_half2 NoiseFiltering = taa_half(0.0).xx;
 	taa_subpixel_payload2 ParallaxFactorBits = taa_subpixel_payload(0x0).xx;
 	bool2 bIsResponsiveAAPixel = false;
 
@@ -284,6 +287,10 @@ void MainCS(
 				AntiAliasingTexture[dpv_lo(InputPixelPos)],
 				AntiAliasingTexture[dpv_hi(InputPixelPos)]);
 			
+			NoiseFiltering = dpv_interleave_registers(
+				NoiseFilteringTexture[dpv_lo(InputPixelPos)],
+				NoiseFilteringTexture[dpv_hi(InputPixelPos)]);
+
 			taa_half2x2 TexelOffset = DecodeSpatialAntiAliasingOffset(EncodedInputTexelOffset);
 			
 			{
@@ -403,11 +410,12 @@ void MainCS(
 
 		taa_half2 FilteredInputColorWeight = taa_half(0.0);
 		
-		// TODO: UECON-464 - this causes shader warnings on XSX DXC.  Revisit once we have aged out the problem GDK versions.
-		taa_half InputToHistoryFactor = taa_half(HistoryInfo_ViewportSize.x * InputInfo_ViewportSizeInverse.x);
-		taa_half2 FinalInputToHistoryFactor = taa_half2(
-			bOffScreen.x ? taa_half(1.0) : InputToHistoryFactor,
-			bOffScreen.y ? taa_half(1.0) : InputToHistoryFactor);
+		// Soften the spatial kernel when there is high noise in the input to make parrallax disocclusion of detailed geometry less distracting.
+		taa_half2 OffScreenInputToHistoryFactor = taa_half(1.0) - taa_half(0.5) * NoiseFiltering;
+
+		taa_half2 KernelInputToHistoryFactor = taa_half2(
+			bOffScreen.x ? OffScreenInputToHistoryFactor.x : taa_half(InputToHistoryFactor),
+			bOffScreen.y ? OffScreenInputToHistoryFactor.y : taa_half(InputToHistoryFactor));
 
 		UNROLL_N(CONFIG_SAMPLES)
 		for (uint SampleId = 0; SampleId < CONFIG_SAMPLES; SampleId++)
@@ -432,7 +440,7 @@ void MainCS(
 			#endif
 
 			taa_half2x2 dPP = PixelOffset - dKO;
-			taa_half2 SampleSpatialWeight = ComputeSampleWeigth(FinalInputToHistoryFactor, dPP, /* MinimalContribution = */ float(0.005));
+			taa_half2 SampleSpatialWeight = ComputeSampleWeigth(KernelInputToHistoryFactor, dPP, /* MinimalContribution = */ float(0.005));
 
 			taa_half2 ToneWeight = HdrWeight4(InputColor);
 
@@ -460,7 +468,7 @@ void MainCS(
 		
 		FilteredInputColor = dpv_scale(FilteredInputColor, rcp(FilteredInputColorWeight));
 
-		InputPixelAlignement = ComputeSampleWeigth(InputToHistoryFactor, dKO, /* MinimalContribution = */ float(0.0));
+		InputPixelAlignement = ComputeSampleWeigth(taa_half(InputToHistoryFactor), dKO, /* MinimalContribution = */ float(0.0));
 		
 		#if CONFIG_SCENE_COLOR_OVERFLOW
 		{
diff --git a/Engine/Source/Runtime/Renderer/Private/PostProcess/TemporalSuperResolution.cpp b/Engine/Source/Runtime/Renderer/Private/PostProcess/TemporalSuperResolution.cpp
--- a/Engine/Source/Runtime/Renderer/Private/PostProcess/TemporalSuperResolution.cpp
+++ b/Engine/Source/Runtime/Renderer/Private/PostProcess/TemporalSuperResolution.cpp
@@ -399,6 +399,7 @@ class FTSRSpatialAntiAliasingCS : public FTSRShader
 		SHADER_PARAMETER_RDG_TEXTURE(Texture2D, InputSceneColorTexture)
 		SHADER_PARAMETER_RDG_TEXTURE(Texture2D, InputSceneColorLdrLumaTexture)
 		SHADER_PARAMETER_RDG_TEXTURE_UAV(RWTexture2D, AntiAliasingOutput)
+		SHADER_PARAMETER_RDG_TEXTURE_UAV(RWTexture2D, NoiseFilteringOutput)
 		SHADER_PARAMETER_RDG_TEXTURE_UAV(RWTexture2D, DebugOutput)
 	END_SHADER_PARAMETER_STRUCT()
 
@@ -479,12 +480,14 @@ class FTSRUpdateHistoryCS : public FTSRShader
 		SHADER_PARAMETER_RDG_TEXTURE(Texture2D, ParallaxFactorTexture)
 		SHADER_PARAMETER_RDG_TEXTURE(Texture2D, ParallaxRejectionMaskTexture)
 		SHADER_PARAMETER_RDG_TEXTURE(Texture2D, AntiAliasingTexture)
-
+		SHADER_PARAMETER_RDG_TEXTURE(Texture2D, NoiseFilteringTexture)
+		
 		SHADER_PARAMETER(FScreenTransform, HistoryPixelPosToScreenPos)
 		SHADER_PARAMETER(FScreenTransform, HistoryPixelPosToPPCo)
 		SHADER_PARAMETER(FVector3f, HistoryQuantizationError)
 		SHADER_PARAMETER(float, MinTranslucencyRejection)
 		SHADER_PARAMETER(float, InvWeightClampingPixelSpeed)
+		SHADER_PARAMETER(float, InputToHistoryFactor)
 		SHADER_PARAMETER(int32, ResponsiveStencilMask)
 
 		SHADER_PARAMETER_STRUCT_INCLUDE(FTSRPrevHistoryParameters, PrevHistoryParameters)
@@ -1209,6 +1212,7 @@ void AddTemporalSuperResolutionPasses(
 
 	// Spatial anti-aliasing when doing history rejection.
 	FRDGTextureRef AntiAliasingTexture = nullptr;
+	FRDGTextureRef NoiseFilteringTexture = nullptr;
 	if (RejectionAntiAliasingQuality > 0)
 	{
 		FRDGTextureRef RawAntiAliasingTexture;
@@ -1221,6 +1225,9 @@ void AddTemporalSuperResolutionPasses(
 
 			RawAntiAliasingTexture = GraphBuilder.CreateTexture(Desc, TEXT("TSR.AntiAliasing.Raw"));
 			AntiAliasingTexture = GraphBuilder.CreateTexture(Desc, TEXT("TSR.AntiAliasing.Filtered"));
+
+			Desc.Format = PF_R8;
+			NoiseFilteringTexture = GraphBuilder.CreateTexture(Desc, TEXT("TSR.AntiAliasing.Noise"));
 		}
 
 		{
@@ -1229,6 +1236,7 @@ void AddTemporalSuperResolutionPasses(
 			PassParameters->InputSceneColorTexture = PassInputs.SceneColorTexture;
 			PassParameters->InputSceneColorLdrLumaTexture = InputSceneColorLdrLumaTexture;
 			PassParameters->AntiAliasingOutput = GraphBuilder.CreateUAV(RawAntiAliasingTexture);
+			PassParameters->NoiseFilteringOutput = GraphBuilder.CreateUAV(NoiseFilteringTexture);
 			PassParameters->DebugOutput = CreateDebugUAV(InputExtent, TEXT("Debug.TSR.SpatialAntiAliasing"));
 
 			FTSRSpatialAntiAliasingCS::FPermutationDomain PermutationVector;
@@ -1369,6 +1377,7 @@ void AddTemporalSuperResolutionPasses(
 		PassParameters->ParallaxFactorTexture = ParallaxFactorTexture;
 		PassParameters->ParallaxRejectionMaskTexture = ParallaxRejectionMaskTexture;
 		PassParameters->AntiAliasingTexture = AntiAliasingTexture;
+		PassParameters->NoiseFilteringTexture = NoiseFilteringTexture;
 
 		FScreenTransform HistoryPixelPosToViewportUV = (FScreenTransform::Identity + 0.5f) * CommonParameters.HistoryInfo.ViewportSizeInverse;
 		PassParameters->HistoryPixelPosToScreenPos = HistoryPixelPosToViewportUV * FScreenTransform::ViewportUVToScreenPos;
@@ -1376,6 +1385,7 @@ void AddTemporalSuperResolutionPasses(
 		PassParameters->HistoryQuantizationError = ComputePixelFormatQuantizationError(History.LowFrequency->Desc.Format);
 		PassParameters->MinTranslucencyRejection = TranslucencyRejectionTexture == nullptr ? 1.0 : 0.0;
 		PassParameters->InvWeightClampingPixelSpeed = 1.0f / CVarTSRWeightClampingPixelSpeed.GetValueOnRenderThread();
+		PassParameters->InputToHistoryFactor = float(HistorySize.X) / float(InputRect.Width());
 		PassParameters->ResponsiveStencilMask = CVarTSREnableResponiveAA.GetValueOnRenderThread() ? (STENCIL_TEMPORAL_RESPONSIVE_AA_MASK) : 0;
 
 		PassParameters->PrevHistoryParameters = PrevHistoryParameters;
