Workaround for Editor deadlock in material editor while in Mobile Preview

diff --git a/Engine/Source/Runtime/RenderCore/Private/RenderGraphPrivate.h b/Engine/Source/Runtime/RenderCore/Private/RenderGraphPrivate.h
--- a/Engine/Source/Runtime/RenderCore/Private/RenderGraphPrivate.h
+++ b/Engine/Source/Runtime/RenderCore/Private/RenderGraphPrivate.h
@@ -202,8 +202,7 @@ FORCEINLINE bool IsParallelExecuteEnabled()
 		&& !GRHICommandList.Bypass()
 		&& !IsImmediateMode()
 		&& !GRDGDebug
-		&& !GRDGTransitionLog
-		&& !IsMobilePlatform(GMaxRHIShaderPlatform);
+		&& !GRDGTransitionLog;
 }
 
 template <typename ResourceRegistryType, typename FunctionType>
diff --git a/Engine/Source/Runtime/Renderer/Private/PlanarReflectionRendering.cpp b/Engine/Source/Runtime/Renderer/Private/PlanarReflectionRendering.cpp
--- a/Engine/Source/Runtime/Renderer/Private/PlanarReflectionRendering.cpp
+++ b/Engine/Source/Runtime/Renderer/Private/PlanarReflectionRendering.cpp
@@ -346,7 +346,8 @@ static void UpdatePlanarReflectionContents_RenderThread(
 #else
 				SCOPED_DRAW_EVENT(RHICmdList, UpdatePlanarReflectionContent_RenderThread);
 #endif
-				FRDGBuilder GraphBuilder(RHICmdList, RDG_EVENT_NAME("PlanarReflection"), ERDGBuilderFlags::AllowParallelExecute);
+				const ERHIFeatureLevel::Type FeatureLevel = SceneRenderer->FeatureLevel;
+				FRDGBuilder GraphBuilder(RHICmdList, RDG_EVENT_NAME("PlanarReflection"), FSceneRenderer::GetRDGParalelExecuteFlags(FeatureLevel));
 
 				// Reflection view late update
 				if (SceneRenderer->Views.Num() > 1)
diff --git a/Engine/Source/Runtime/Renderer/Private/ReflectionEnvironmentCapture.cpp b/Engine/Source/Runtime/Renderer/Private/ReflectionEnvironmentCapture.cpp
--- a/Engine/Source/Runtime/Renderer/Private/ReflectionEnvironmentCapture.cpp
+++ b/Engine/Source/Runtime/Renderer/Private/ReflectionEnvironmentCapture.cpp
@@ -718,12 +718,12 @@ void CaptureSceneToScratchCubemap(FRHICommandListImmediate& RHICmdList, FSceneRe
 	FDeferredUpdateResource::UpdateResources(RHICmdList);
 	FMaterialRenderProxy::UpdateDeferredCachedUniformExpressions();
 
-	const auto FeatureLevel = SceneRenderer->FeatureLevel;
+	const ERHIFeatureLevel::Type FeatureLevel = SceneRenderer->FeatureLevel;
 	
 	{
 		SCOPED_DRAW_EVENT(RHICmdList, CubeMapCapture);
 
-		FRDGBuilder GraphBuilder(RHICmdList, RDG_EVENT_NAME("CubeMapCapture"), ERDGBuilderFlags::AllowParallelExecute);
+		FRDGBuilder GraphBuilder(RHICmdList, RDG_EVENT_NAME("CubeMapCapture"), FSceneRenderer::GetRDGParalelExecuteFlags(FeatureLevel));
 
 		// Render the scene normally for one face of the cubemap
 		SceneRenderer->Render(GraphBuilder);
diff --git a/Engine/Source/Runtime/Renderer/Private/SceneCaptureRendering.cpp b/Engine/Source/Runtime/Renderer/Private/SceneCaptureRendering.cpp
--- a/Engine/Source/Runtime/Renderer/Private/SceneCaptureRendering.cpp
+++ b/Engine/Source/Runtime/Renderer/Private/SceneCaptureRendering.cpp
@@ -284,12 +284,14 @@ static void UpdateSceneCaptureContentDeferred_RenderThread(
 	// update any resources that needed a deferred update
 	FDeferredUpdateResource::UpdateResources(RHICmdList);
 	{
+		const ERHIFeatureLevel::Type FeatureLevel = SceneRenderer->FeatureLevel;
+
 #if WANTS_DRAW_MESH_EVENTS
 		SCOPED_DRAW_EVENTF(RHICmdList, SceneCapture, TEXT("SceneCapture %s"), *EventName);
-		FRDGBuilder GraphBuilder(RHICmdList, RDG_EVENT_NAME("SceneCapture %s", *EventName), ERDGBuilderFlags::AllowParallelExecute);
+		FRDGBuilder GraphBuilder(RHICmdList, RDG_EVENT_NAME("SceneCapture %s", *EventName), FSceneRenderer::GetRDGParalelExecuteFlags(FeatureLevel));
 #else
 		SCOPED_DRAW_EVENT(RHICmdList, UpdateSceneCaptureContent_RenderThread);
-		FRDGBuilder GraphBuilder(RHICmdList, RDG_EVENT_NAME("SceneCapture"), ERDGBuilderFlags::AllowParallelExecute);
+		FRDGBuilder GraphBuilder(RHICmdList, RDG_EVENT_NAME("SceneCapture"), FSceneRenderer::GetRDGParalelExecuteFlags(FeatureLevel));
 #endif
 
 		FRDGTextureRef TargetTexture = RegisterExternalTexture(GraphBuilder, RenderTarget->GetRenderTargetTexture(), TEXT("SceneCaptureTarget"));
diff --git a/Engine/Source/Runtime/Renderer/Private/SceneRendering.cpp b/Engine/Source/Runtime/Renderer/Private/SceneRendering.cpp
--- a/Engine/Source/Runtime/Renderer/Private/SceneRendering.cpp
+++ b/Engine/Source/Runtime/Renderer/Private/SceneRendering.cpp
@@ -3951,13 +3951,15 @@ static void RenderViewFamily_RenderThread(FRHICommandListImmediate& RHICmdList,
 		SCOPE_CYCLE_COUNTER(STAT_TotalSceneRenderingTime);
 
 		{
+			const ERHIFeatureLevel::Type FeatureLevel = SceneRenderer->FeatureLevel;
+
 			FRDGBuilder GraphBuilder(
 				RHICmdList,
 				RDG_EVENT_NAME("SceneRenderer_%s(ViewFamily=%s)",
 					ViewFamily.EngineShowFlags.HitProxies ? TEXT("RenderHitProxies") : TEXT("Render"),
 					ViewFamily.bResolveScene ? TEXT("Primary") : TEXT("Auxiliary")
 				),
-				ERDGBuilderFlags::AllowParallelExecute
+				FSceneRenderer::GetRDGParalelExecuteFlags(FeatureLevel)
 			);
 
 			if (ViewFamily.EngineShowFlags.HitProxies)
diff --git a/Engine/Source/Runtime/Renderer/Private/SceneRendering.h b/Engine/Source/Runtime/Renderer/Private/SceneRendering.h
--- a/Engine/Source/Runtime/Renderer/Private/SceneRendering.h
+++ b/Engine/Source/Runtime/Renderer/Private/SceneRendering.h
@@ -1982,6 +1982,13 @@ public:
 	/** Gets a readable light name for use with a draw event. */
 	static void GetLightNameForDrawEvent(const FLightSceneProxy* LightProxy, FString& LightNameWithLevel);
 
+	/** Get the ParallelExecuteFlags depending on FeatureLevel*/
+	static ERDGBuilderFlags GetRDGParalelExecuteFlags(ERHIFeatureLevel::Type FeatureLevel)
+	{
+		return  FeatureLevel == ERHIFeatureLevel::ES3_1
+			? ERDGBuilderFlags::None
+			: ERDGBuilderFlags::AllowParallelExecute;
+	}
 protected:
 
 	/** Size of the family. */
