Ray tracing per instance culling

diff --git a/Engine/Source/Runtime/D3D12RHI/Private/D3D12RayTracing.cpp b/Engine/Source/Runtime/D3D12RHI/Private/D3D12RayTracing.cpp
--- a/Engine/Source/Runtime/D3D12RHI/Private/D3D12RayTracing.cpp
+++ b/Engine/Source/Runtime/D3D12RHI/Private/D3D12RayTracing.cpp
@@ -3633,6 +3633,7 @@ FD3D12RayTracingScene::FD3D12RayTracingScene(FD3D12Adapter* Adapter, const FRayT
 	const int32 NumSceneInstances = Initializer.Instances.Num();
 
 	int32 NumDxrInstances = 0;
+	int32 NumInactiveDxrInstances = 0;
 	int32 NumCopyCommands = 0;
 	NumTotalSegments = 0;
 
@@ -3682,12 +3683,13 @@ FD3D12RayTracingScene::FD3D12RayTracingScene(FD3D12Adapter* Adapter, const FRayT
 	Instances.SetNumUninitialized(NumDxrInstances);
 
 	const EParallelForFlags ParallelForFlags = EParallelForFlags::None; // set ForceSingleThread for testing
-	ParallelFor(NumSceneInstances, [this, &Initializer](int32 InstanceIndex)
+	ParallelFor(NumSceneInstances, [this, &Initializer, &NumInactiveDxrInstances](int32 InstanceIndex)
 	{
 		const FRayTracingGeometryInstance& Instance = Initializer.Instances[InstanceIndex];
 		FD3D12RayTracingGeometry* Geometry = PerInstanceGeometries[InstanceIndex];
 
 		D3D12_RAYTRACING_INSTANCE_DESC InstanceDesc = {};
+		check(InstanceDesc.AccelerationStructure == 0);
 
 		InstanceDesc.InstanceMask = Instance.Mask;
 		InstanceDesc.InstanceContributionToHitGroupIndex = SegmentPrefixSum[InstanceIndex] * ShaderSlotsPerGeometrySegment;
@@ -3704,12 +3706,29 @@ FD3D12RayTracingScene::FD3D12RayTracingScene(FD3D12Adapter* Adapter, const FRayT
 		const bool bCpuInstance = !bGpuInstance;
 
 		uint32 DescIndex = BaseInstancePrefixSum[InstanceIndex];
+
+		int32 NumInactiveDxrInstancesThisSceneInstance = 0;
+
 		for (uint32 TransformIndex = 0; TransformIndex < NumTransforms; ++TransformIndex)
 		{
 			InstanceDesc.InstanceID = bUseUniqueUserData ? Instance.UserData[TransformIndex] : Instance.DefaultUserData;
 
 			if (LIKELY(bCpuInstance))
 			{
+				// Set flag for deactivated instances
+				if (!Instance.ActivationMask.IsEmpty())
+				{
+					if ((Instance.ActivationMask[TransformIndex / 32] & (1 << (TransformIndex % 32))) == 0)
+					{
+						InstanceDesc.AccelerationStructure = 0xFFFFFFFFFFFFFFFF;
+						NumInactiveDxrInstancesThisSceneInstance++;
+					}
+					else
+					{
+						InstanceDesc.AccelerationStructure = 0;
+					}
+				}
+
 				// DXR uses a 3x4 transform matrix in row major layout
 
 				const FMatrix& Transform = Instance.Transforms[TransformIndex];
@@ -3733,6 +3752,11 @@ FD3D12RayTracingScene::FD3D12RayTracingScene(FD3D12Adapter* Adapter, const FRayT
 			Instances[DescIndex] = InstanceDesc;
 			++DescIndex;
 		}
+
+	#if STATS
+		FPlatformAtomics::InterlockedAdd(&NumInactiveDxrInstances, NumInactiveDxrInstancesThisSceneInstance);
+	#endif
+
 	}, ParallelForFlags);
 
 	ERayTracingAccelerationStructureFlags BuildFlags = ERayTracingAccelerationStructureFlags::FastTrace; // #yuriy_todo: pass this in
@@ -3748,6 +3772,8 @@ FD3D12RayTracingScene::FD3D12RayTracingScene(FD3D12Adapter* Adapter, const FRayT
 
 	// Get maximum buffer sizes for all GPUs in the system
 	SizeInfo = RHICalcRayTracingSceneSize(BuildInputs.NumDescs, BuildFlags);
+
+	SET_DWORD_STAT(STAT_RayTracingInstances, NumDxrInstances - NumInactiveDxrInstances);
 };
 
 FD3D12RayTracingScene::~FD3D12RayTracingScene()
@@ -3865,7 +3891,16 @@ void FD3D12RayTracingScene::BuildAccelerationStructure(FD3D12CommandContext& Com
 				const int32 NumTransforms = PerInstanceNumTransforms[InstanceIndex];
 				for (int32 TransformIndex = 0; TransformIndex < NumTransforms; ++TransformIndex)
 				{
-					Instances[DxrInstanceIndex].AccelerationStructure = BlasAddress;
+					// If this instance has been deactivated, set BLAS address to 0
+					if (Instances[DxrInstanceIndex].AccelerationStructure == 0xFFFFFFFFFFFFFFFF)
+					{
+						Instances[DxrInstanceIndex].AccelerationStructure = 0;
+					}
+					else
+					{
+						Instances[DxrInstanceIndex].AccelerationStructure = BlasAddress;
+					}
+
 					DxrInstanceIndex++;
 				}
 
diff --git a/Engine/Source/Runtime/RHI/Private/DynamicRHI.cpp b/Engine/Source/Runtime/RHI/Private/DynamicRHI.cpp
--- a/Engine/Source/Runtime/RHI/Private/DynamicRHI.cpp
+++ b/Engine/Source/Runtime/RHI/Private/DynamicRHI.cpp
@@ -22,8 +22,8 @@
 IMPLEMENT_TYPE_LAYOUT(FRayTracingGeometryInitializer);
 IMPLEMENT_TYPE_LAYOUT(FRayTracingGeometrySegment);
 
-static_assert(sizeof(FRayTracingGeometryInstance) <= 64,
-	"Ray tracing instance descriptor is expected to be no more than 64 bytes, "
+static_assert(sizeof(FRayTracingGeometryInstance) <= 80,
+	"Ray tracing instance descriptor is expected to be no more than 80 bytes, "
 	"as there may be a very large number of them.");
 
 #ifndef PLATFORM_ALLOW_NULL_RHI
diff --git a/Engine/Source/Runtime/RHI/Public/DynamicRHI.h b/Engine/Source/Runtime/RHI/Public/DynamicRHI.h
--- a/Engine/Source/Runtime/RHI/Public/DynamicRHI.h
+++ b/Engine/Source/Runtime/RHI/Public/DynamicRHI.h
@@ -117,6 +117,9 @@ struct FRayTracingGeometryInstance
 	uint32 DefaultUserData = 0;
 	TArrayView<const uint32> UserData;
 
+	// Each geometry copy can have one bit to make it individually deactivated (removed from TLAS while maintaining hit group indexing). Useful for culling.
+	TArrayView<const uint32> ActivationMask;
+
 	// Mask that will be tested against one provided to TraceRay() in shader code.
 	// If binary AND of instance mask with ray mask is zero, then the instance is considered not intersected / invisible.
 	uint8 Mask = 0xFF;
diff --git a/Engine/Source/Runtime/Renderer/Private/DeferredShadingRenderer.cpp b/Engine/Source/Runtime/Renderer/Private/DeferredShadingRenderer.cpp
--- a/Engine/Source/Runtime/Renderer/Private/DeferredShadingRenderer.cpp
+++ b/Engine/Source/Runtime/Renderer/Private/DeferredShadingRenderer.cpp
@@ -68,6 +68,7 @@
 #include "Engine/SubsurfaceProfile.h"
 #include "SceneCaptureRendering.h"
 #include "NaniteSceneProxy.h"
+#include "RayTracing/RayTracingInstanceCulling.h"
 
 extern int32 GNaniteShowStats;
 
@@ -177,28 +178,6 @@ static TAutoConsoleVariable<float> CVarRayTracingDynamicGeometryLastRenderTimeUp
 	5000.0f,
 	TEXT("Dynamic geometries within this distance will have their LastRenderTime updated, so that visibility based ticking (like skeletal mesh) can work when the component is not directly visible in the view (but reflected)."));
 
-static TAutoConsoleVariable<int32> CVarRayTracingCulling(
-	TEXT("r.RayTracing.Culling"),
-	0,
-	TEXT("Enable culling in ray tracing for objects that are behind the camera\n")
-	TEXT(" 0: Culling disabled (default)\n")
-	TEXT(" 1: Culling by distance and solid angle enabled. Only cull objects behind camera.\n")
-	TEXT(" 2: Culling by distance and solid angle enabled. Cull objects in front and behind camera.\n")
-	TEXT(" 3: Culling by distance OR solid angle enabled. Cull objects in front and behind camera."),
-	ECVF_RenderThreadSafe);
-
-static TAutoConsoleVariable<float> CVarRayTracingCullingRadius(
-	TEXT("r.RayTracing.Culling.Radius"),
-	10000.0f,
-	TEXT("Do camera culling for objects behind the camera outside of this radius in ray tracing effects (default = 10000 (100m))"),
-	ECVF_RenderThreadSafe);
-
-static TAutoConsoleVariable<float> CVarRayTracingCullingAngle(
-	TEXT("r.RayTracing.Culling.Angle"),
-	1.0f, 
-	TEXT("Do camera culling for objects behind the camera with a projected angle smaller than this threshold in ray tracing effects (default = 5 degrees )"),
-	ECVF_RenderThreadSafe);
-
 static TAutoConsoleVariable<int32> CVarRayTracingAutoInstance(
 	TEXT("r.RayTracing.AutoInstance"),
 	1,
@@ -255,19 +234,6 @@ namespace Nanite
 	extern void ListStatFilters(FSceneRenderer* SceneRenderer);
 }
 
-#if RHI_RAYTRACING
-
-int32 GetRayTracingCulling()
-{
-	return CVarRayTracingCulling.GetValueOnRenderThread();
-}
-
-float GetRayTracingCullingRadius()
-{
-	return CVarRayTracingCullingRadius.GetValueOnRenderThread();
-}
-#endif // RHI_RAYTRACING
-
 DECLARE_CYCLE_STAT(TEXT("InitViews Intentional Stall"), STAT_InitViews_Intentional_Stall, STATGROUP_InitViews);
 
 DECLARE_CYCLE_STAT(TEXT("DeferredShadingSceneRenderer UpdateDownsampledDepthSurface"), STAT_FDeferredShadingSceneRenderer_UpdateDownsampledDepthSurface, STATGROUP_SceneRendering);
@@ -557,6 +523,8 @@ bool FDeferredShadingSceneRenderer::GatherRayTracingWorldInstancesForView(FRHICo
 		&DynamicVertexBufferForInitViews,
 		&DynamicReadBufferForInitViews);
 
+	View.RayTracingCullingParameters.Init(View);
+
 	FRayTracingMaterialGatheringContext MaterialGatheringContext
 	{
 		Scene,
@@ -609,16 +577,6 @@ bool FDeferredShadingSceneRenderer::GatherRayTracingWorldInstancesForView(FRHICo
 		TRACE_CPUPROFILER_EVENT_SCOPE(GatherRayTracingWorldInstances_RelevantPrimitives);
 
 		int32 BroadIndex = 0;
-		const int32 CullInRayTracing = GetRayTracingCulling();
-		const float CullingRadius = GetRayTracingCullingRadius();
-		const float FarFieldCullingRadius = Lumen::GetFarFieldMaxTraceDistance();
-		const float CullAngleThreshold = CVarRayTracingCullingAngle.GetValueOnRenderThread();
-		const float AngleThresholdRatio = FMath::Tan(FMath::Min(89.99f, CullAngleThreshold) * PI / 180.0f);
-		const FVector ViewOrigin = View.ViewMatrices.GetViewOrigin();
-		const FVector ViewDirection = View.GetViewDirection();
-		const bool bCullAllObjects = CullInRayTracing == 2 || CullInRayTracing == 3;
-		const bool bCullByRadiusOrDistance = CullInRayTracing == 3;
-		const bool bIsRayTracingFarField = Lumen::UseFarField();
 
 		for (int PrimitiveIndex = 0; PrimitiveIndex < Scene->PrimitiveSceneProxies.Num(); PrimitiveIndex++)
 		{
@@ -644,48 +602,14 @@ bool FDeferredShadingSceneRenderer::GatherRayTracingWorldInstancesForView(FRHICo
 
 			// Skip far field if not enabled
 			const bool bIsFarFieldPrimitive = EnumHasAnyFlags(Scene->PrimitiveRayTracingFlags[PrimitiveIndex], ERayTracingPrimitiveFlags::FarField);
-			if (!bIsRayTracingFarField && bIsFarFieldPrimitive)
+			if (!View.RayTracingCullingParameters.bIsRayTracingFarField && bIsFarFieldPrimitive)
 			{
 				continue;
 			}
 
-			if (CullInRayTracing > 0)
+			if (RayTracing::ShouldCullBounds(View.RayTracingCullingParameters, Scene->PrimitiveBounds[PrimitiveIndex].BoxSphereBounds, bIsFarFieldPrimitive))
 			{
-				const FBoxSphereBounds ObjectBounds = Scene->PrimitiveBounds[PrimitiveIndex].BoxSphereBounds;
-				const float ObjectRadius = ObjectBounds.SphereRadius;
-				const FVector ObjectCenter = ObjectBounds.Origin + 0.5*ObjectBounds.BoxExtent;
-				const FVector CameraToObjectCenter = FVector(ObjectCenter - ViewOrigin);
-
-				const bool bConsiderCulling = bCullAllObjects || FVector::DotProduct(ViewDirection, CameraToObjectCenter) < -ObjectRadius;
-
-				if (bConsiderCulling)
-				{
-					const float CameraToObjectCenterLength = CameraToObjectCenter.Size();
-					
-					if (bIsFarFieldPrimitive)
-					{
-						if (CameraToObjectCenterLength > (FarFieldCullingRadius + ObjectRadius))
-						{
-							continue;
-						}
-					}
-					else
-					{
-						const bool bIsFarEnoughToCull = CameraToObjectCenterLength > (CullingRadius + ObjectRadius);
-
-						// Cull by solid angle: check the radius of bounding sphere against angle threshold
-						const bool bAngleIsSmallEnoughToCull = ObjectRadius / CameraToObjectCenterLength < AngleThresholdRatio;
-
-						if (bCullByRadiusOrDistance && (bIsFarEnoughToCull || bAngleIsSmallEnoughToCull))
-						{
-							continue;
-						}
-						else if (bIsFarEnoughToCull && bAngleIsSmallEnoughToCull)
-						{
-							continue;
-						}
-					}
-				}
+				continue;
 			}
 
 			if (!View.State)
@@ -869,8 +793,6 @@ bool FDeferredShadingSceneRenderer::GatherRayTracingWorldInstancesForView(FRHICo
 		}
 	}
 
-	//
-
 	{
 		TRACE_CPUPROFILER_EVENT_SCOPE(GatherRayTracingWorldInstances_DynamicElements);
 
@@ -1206,6 +1128,13 @@ bool FDeferredShadingSceneRenderer::GatherRayTracingWorldInstancesForView(FRHICo
 
 		InstanceBatches.Reserve(RelevantPrimitives.Num());
 
+		TArray<FRayTracingCullPrimitiveInstancesClosure> CullInstancesClosures;
+		if (View.RayTracingCullingParameters.CullInRayTracing > 0 && GetRayTracingCullingPerInstance())
+		{
+			CullInstancesClosures.Reserve(RelevantPrimitives.Num());
+			View.RayTracingPerInstanceCullingTaskList.Reserve(RelevantPrimitives.Num() / 256 + 1);
+		}
+
 		// scan relevant primitives computing hash data to look for duplicate instances
 		for (const FRelevantPrimitive& RelevantPrimitive : RelevantPrimitives)
 		{
@@ -1233,6 +1162,39 @@ bool FDeferredShadingSceneRenderer::GatherRayTracingWorldInstancesForView(FRHICo
 
 				checkf(SceneInfo->CachedRayTracingInstance.GeometryRHI, TEXT("Ray tracing instance must have a valid geometry."));
 				RayTracingScene.Instances.Add(SceneInfo->CachedRayTracingInstance);
+
+				if (View.RayTracingCullingParameters.CullInRayTracing > 0 && GetRayTracingCullingPerInstance() && SceneInfo->CachedRayTracingInstance.NumTransforms > 1)
+				{
+					FRayTracingGeometryInstance& NewInstance = RayTracingScene.Instances.Last();
+
+					const bool bIsFarFieldPrimitive = EnumHasAnyFlags(Flags, ERayTracingPrimitiveFlags::FarField);
+
+					TArrayView<uint32> InstanceActivationMask = RayTracingScene.Allocate<uint32>(FMath::DivideAndRoundUp(NewInstance.NumTransforms, 32u));
+
+					NewInstance.ActivationMask = InstanceActivationMask;
+
+					FRayTracingCullPrimitiveInstancesClosure Closure;
+					Closure.Scene = Scene;
+					Closure.SceneInfo = SceneInfo;
+					Closure.PrimitiveIndex = PrimitiveIndex;
+					Closure.bIsFarFieldPrimitive = bIsFarFieldPrimitive;
+					Closure.CullingParameters = &View.RayTracingCullingParameters;
+					Closure.OutInstanceActivationMask = InstanceActivationMask;
+
+					CullInstancesClosures.Add(MoveTemp(Closure));
+
+					if (CullInstancesClosures.Num() >= 256)
+					{
+						View.RayTracingPerInstanceCullingTaskList.Add(FFunctionGraphTask::CreateAndDispatchWhenReady([CullInstancesClosures = MoveTemp(CullInstancesClosures)]()
+						{
+							for (auto& Closure : CullInstancesClosures)
+							{
+								Closure();
+							}
+						}, TStatId(), nullptr, ENamedThreads::AnyThread));
+					}
+				}
+
 				AddDebugRayTracingInstanceFlags(RayTracingScene.Instances.Last().Flags);
 			}
 			else
@@ -1319,6 +1281,14 @@ bool FDeferredShadingSceneRenderer::GatherRayTracingWorldInstancesForView(FRHICo
 				}
 			}
 		}
+
+		View.RayTracingPerInstanceCullingTaskList.Add(FFunctionGraphTask::CreateAndDispatchWhenReady([CullInstancesClosures = MoveTemp(CullInstancesClosures)]()
+		{
+			for (auto& Closure : CullInstancesClosures)
+			{
+				Closure();
+			}
+		}, TStatId(), nullptr, ENamedThreads::AnyThread));
 	}
 
 	// Inform the coarse mesh streaming manager about all the used streamable render assets in the scene
@@ -1505,6 +1475,9 @@ bool FDeferredShadingSceneRenderer::DispatchRayTracingWorldUpdates(FRDGBuilder&
 		GRayTracingGeometryManager.ForceBuildIfPending(GraphBuilder.RHICmdList, RayTracingScene.GeometriesToBuild);
 	}
 
+	FTaskGraphInterface::Get().WaitUntilTasksComplete(ReferenceView.RayTracingPerInstanceCullingTaskList, ENamedThreads::GetRenderThread_Local());
+	ReferenceView.RayTracingPerInstanceCullingTaskList.Empty();
+
 	RDG_GPU_MASK_SCOPE(GraphBuilder, FRHIGPUMask::All());
 
 	RayTracingScene.BeginCreate(GraphBuilder);
diff --git a/Engine/Source/Runtime/Renderer/Private/PrimitiveSceneInfo.cpp b/Engine/Source/Runtime/Renderer/Private/PrimitiveSceneInfo.cpp
--- a/Engine/Source/Runtime/Renderer/Private/PrimitiveSceneInfo.cpp
+++ b/Engine/Source/Runtime/Renderer/Private/PrimitiveSceneInfo.cpp
@@ -819,6 +819,10 @@ void FPrimitiveSceneInfo::UpdateCachedRayTracingInstance(FPrimitiveSceneInfo* Sc
 		{
 			LocalToWorld = LocalToWorld.ConcatTranslation(Lumen::GetFarFieldReferencePos());
 		}
+
+		SceneInfo->CachedRayTracingInstanceWorldBounds.Empty();
+		SceneInfo->CachedRayTracingInstanceWorldBounds.AddUninitialized(CachedRayTracingInstance.NumTransforms);
+
 		SceneInfo->UpdateCachedRayTracingInstanceTransforms(LocalToWorld);
 		SceneInfo->CachedRayTracingInstance.Transforms = MakeArrayView(SceneInfo->CachedRayTracingInstanceWorldTransforms);
 
@@ -873,9 +877,15 @@ void FPrimitiveSceneInfo::RemoveCachedRayTracingPrimitives()
 void FPrimitiveSceneInfo::UpdateCachedRayTracingInstanceTransforms(const FMatrix& NewPrimitiveLocalToWorld)
 {
 	TRACE_CPUPROFILER_EVENT_SCOPE(UpdateCachedRayTracingInstanceTransforms);
+	
+	SmallestRayTracingInstanceWorldBoundsIndex = 0;
+
 	for (int32 Index = 0; Index < CachedRayTracingInstanceLocalTransforms.Num(); Index++)
 	{
 		CachedRayTracingInstanceWorldTransforms[Index] = CachedRayTracingInstanceLocalTransforms[Index] * NewPrimitiveLocalToWorld;
+		CachedRayTracingInstanceWorldBounds[Index] = Proxy->GetInstanceSceneData()[Index].LocalBounds.TransformBy(CachedRayTracingInstanceLocalTransforms[Index] * NewPrimitiveLocalToWorld).ToBoxSphereBounds();
+		
+		SmallestRayTracingInstanceWorldBoundsIndex = CachedRayTracingInstanceWorldBounds[Index].SphereRadius < CachedRayTracingInstanceWorldBounds[SmallestRayTracingInstanceWorldBoundsIndex].SphereRadius ? Index : SmallestRayTracingInstanceWorldBoundsIndex;
 	}
 }
 #endif
diff --git a/Engine/Source/Runtime/Renderer/Private/RayTracing/RayTracingInstanceCulling.cpp b/Engine/Source/Runtime/Renderer/Private/RayTracing/RayTracingInstanceCulling.cpp
new file mode 100644
--- /dev/null
+++ b/Engine/Source/Runtime/Renderer/Private/RayTracing/RayTracingInstanceCulling.cpp
@@ -0,0 +1,180 @@
+// Copyright Epic Games, Inc.All Rights Reserved.
+
+#include "RayTracingInstanceCulling.h"
+#include "Lumen/Lumen.h"
+
+#if RHI_RAYTRACING
+
+static TAutoConsoleVariable<int32> CVarRayTracingCulling(
+	TEXT("r.RayTracing.Culling"),
+	0,
+	TEXT("Enable culling in ray tracing for objects that are behind the camera\n")
+	TEXT(" 0: Culling disabled (default)\n")
+	TEXT(" 1: Culling by distance and solid angle enabled. Only cull objects behind camera.\n")
+	TEXT(" 2: Culling by distance and solid angle enabled. Cull objects in front and behind camera.\n")
+	TEXT(" 3: Culling by distance OR solid angle enabled. Cull objects in front and behind camera."),
+	ECVF_RenderThreadSafe);
+
+static TAutoConsoleVariable<int32> CVarRayTracingCullingPerInstance(
+	TEXT("r.RayTracing.Culling.PerInstance"),
+	1,
+	TEXT(""),
+	ECVF_RenderThreadSafe);
+
+static TAutoConsoleVariable<float> CVarRayTracingCullingRadius(
+	TEXT("r.RayTracing.Culling.Radius"),
+	10000.0f,
+	TEXT("Do camera culling for objects behind the camera outside of this radius in ray tracing effects (default = 10000 (100m))"),
+	ECVF_RenderThreadSafe);
+
+static TAutoConsoleVariable<float> CVarRayTracingCullingAngle(
+	TEXT("r.RayTracing.Culling.Angle"),
+	1.0f,
+	TEXT("Do camera culling for objects behind the camera with a projected angle smaller than this threshold in ray tracing effects (default = 5 degrees )"),
+	ECVF_RenderThreadSafe);
+
+int32 GetRayTracingCulling()
+{
+	return CVarRayTracingCulling.GetValueOnRenderThread();
+}
+
+float GetRayTracingCullingRadius()
+{
+	return CVarRayTracingCullingRadius.GetValueOnRenderThread();
+}
+
+int32 GetRayTracingCullingPerInstance()
+{
+	return CVarRayTracingCullingPerInstance.GetValueOnRenderThread();
+}
+
+void FRayTracingCullingParameters::Init(FViewInfo& View)
+{
+	CullInRayTracing = CVarRayTracingCulling.GetValueOnRenderThread();
+	CullingRadius = CVarRayTracingCullingRadius.GetValueOnRenderThread();
+	FarFieldCullingRadius = Lumen::GetFarFieldMaxTraceDistance();
+	CullAngleThreshold = CVarRayTracingCullingAngle.GetValueOnRenderThread();
+	AngleThresholdRatio = FMath::Tan(FMath::Min(89.99f, CullAngleThreshold) * PI / 180.0f);
+	ViewOrigin = View.ViewMatrices.GetViewOrigin();
+	ViewDirection = View.GetViewDirection();
+	bCullAllObjects = CullInRayTracing == 2 || CullInRayTracing == 3;
+	bCullByRadiusOrDistance = CullInRayTracing == 3;
+	bIsRayTracingFarField = Lumen::UseFarField();
+}
+
+namespace RayTracing
+{
+
+bool ShouldCullBounds(const FRayTracingCullingParameters& CullingParameters, FBoxSphereBounds ObjectBounds, bool bIsFarFieldPrimitive)
+{
+	if (CullingParameters.CullInRayTracing > 0)
+	{
+		const float ObjectRadius = ObjectBounds.SphereRadius;
+		const FVector ObjectCenter = ObjectBounds.Origin + 0.5 * ObjectBounds.BoxExtent;
+		const FVector CameraToObjectCenter = FVector(ObjectCenter - CullingParameters.ViewOrigin);
+
+		const bool bConsiderCulling = CullingParameters.bCullAllObjects || FVector::DotProduct(CullingParameters.ViewDirection, CameraToObjectCenter) < -ObjectRadius;
+
+		if (bConsiderCulling)
+		{
+			const float CameraToObjectCenterLength = CameraToObjectCenter.Size();
+
+			if (bIsFarFieldPrimitive)
+			{
+				if (CameraToObjectCenterLength > (CullingParameters.FarFieldCullingRadius + ObjectRadius))
+				{
+					return true;
+				}
+			}
+			else
+			{
+				const bool bIsFarEnoughToCull = CameraToObjectCenterLength > (CullingParameters.CullingRadius + ObjectRadius);
+
+				// Cull by solid angle: check the radius of bounding sphere against angle threshold
+				const bool bAngleIsSmallEnoughToCull = ObjectRadius / CameraToObjectCenterLength < CullingParameters.AngleThresholdRatio;
+
+				if (CullingParameters.bCullByRadiusOrDistance && (bIsFarEnoughToCull || bAngleIsSmallEnoughToCull))
+				{
+					return true;
+				}
+				else if (bIsFarEnoughToCull && bAngleIsSmallEnoughToCull)
+				{
+					return true;
+				}
+			}
+		}
+	}
+
+	return false;
+}
+
+bool ShouldSkipPerInstanceCullingForPrimitive(const FRayTracingCullingParameters& CullingParameters, FBoxSphereBounds ObjectBounds, FBoxSphereBounds SmallestInstanceBounds, bool bIsFarFieldPrimitive)
+{
+	bool bSkipCulling = false;
+
+	const float ObjectRadius = ObjectBounds.SphereRadius;
+	const FVector ObjectCenter = ObjectBounds.Origin + 0.5 * ObjectBounds.BoxExtent;
+	const FVector CameraToObjectCenter = FVector(ObjectCenter - CullingParameters.ViewOrigin);
+
+	const FVector CameraToFurthestInstanceCenter = CameraToObjectCenter * (CameraToObjectCenter.Size() + ObjectRadius + SmallestInstanceBounds.SphereRadius) / CameraToObjectCenter.Size();
+
+	const bool bConsiderCulling = CullingParameters.bCullAllObjects || FVector::DotProduct(CullingParameters.ViewDirection, CameraToObjectCenter) < -ObjectRadius;
+
+	if (bConsiderCulling)
+	{
+		const float CameraToObjectCenterLength = CameraToObjectCenter.Size();
+
+		if (bIsFarFieldPrimitive)
+		{
+			if (CameraToObjectCenterLength < (CullingParameters.FarFieldCullingRadius - ObjectRadius))
+			{
+				bSkipCulling = true;
+			}
+		}
+		else
+		{
+			const bool bSkipDistanceCulling = CameraToObjectCenterLength < (CullingParameters.CullingRadius - ObjectRadius);
+
+			// Cull by solid angle: check the radius of bounding sphere against angle threshold
+			const bool bSkipAngleCulling = FMath::IsFinite(SmallestInstanceBounds.SphereRadius / CameraToFurthestInstanceCenter.Size()) && SmallestInstanceBounds.SphereRadius / CameraToFurthestInstanceCenter.Size() >= CullingParameters.AngleThresholdRatio;
+
+			if (CullingParameters.bCullByRadiusOrDistance)
+			{
+				if (bSkipDistanceCulling && bSkipAngleCulling)
+				{
+					bSkipCulling = true;
+				}
+			}
+			else if (bSkipDistanceCulling || bSkipAngleCulling)
+			{
+				bSkipCulling = true;
+			}
+		}
+	}
+	else
+	{
+		bSkipCulling = true;
+	}
+
+	return bSkipCulling;
+}
+
+}
+
+void FRayTracingCullPrimitiveInstancesClosure::operator()() const
+{
+	FMemory::Memset(OutInstanceActivationMask.GetData(), 0xFF, OutInstanceActivationMask.Num() * 4);
+
+	if (!RayTracing::ShouldSkipPerInstanceCullingForPrimitive(*CullingParameters, Scene->PrimitiveBounds[PrimitiveIndex].BoxSphereBounds, SceneInfo->CachedRayTracingInstanceWorldBounds[SceneInfo->SmallestRayTracingInstanceWorldBoundsIndex], bIsFarFieldPrimitive))
+	{
+		for (int32 InstanceIndex = 0; InstanceIndex < SceneInfo->CachedRayTracingInstanceWorldTransforms.Num(); InstanceIndex++)
+		{
+			if (RayTracing::ShouldCullBounds(*CullingParameters, SceneInfo->CachedRayTracingInstanceWorldBounds[InstanceIndex], bIsFarFieldPrimitive))
+			{
+				OutInstanceActivationMask[InstanceIndex / 32] &= ~(1 << (InstanceIndex % 32));
+			}
+		}
+	}
+}
+
+#endif
diff --git a/Engine/Source/Runtime/Renderer/Private/RayTracing/RayTracingInstanceCulling.h b/Engine/Source/Runtime/Renderer/Private/RayTracing/RayTracingInstanceCulling.h
new file mode 100644
--- /dev/null
+++ b/Engine/Source/Runtime/Renderer/Private/RayTracing/RayTracingInstanceCulling.h
@@ -0,0 +1,35 @@
+// Copyright Epic Games, Inc.All Rights Reserved.
+
+#pragma once
+
+#include "SceneRendering.h"
+#include "ScenePrivate.h"
+
+#if RHI_RAYTRACING
+
+int32 GetRayTracingCulling();
+float GetRayTracingCullingRadius();
+int32 GetRayTracingCullingPerInstance();
+
+namespace RayTracing
+{
+
+bool ShouldCullBounds(const FRayTracingCullingParameters& CullingParameters, FBoxSphereBounds ObjectBounds, bool bIsFarFieldPrimitive);
+bool ShouldSkipPerInstanceCullingForPrimitive(const FRayTracingCullingParameters& CullingParameters, FBoxSphereBounds ObjectBounds, FBoxSphereBounds SmallestInstanceBounds, bool bIsFarFieldPrimitive);
+
+}
+
+struct FRayTracingCullPrimitiveInstancesClosure
+{
+	FScene* Scene;
+	int32 PrimitiveIndex;
+	const FPrimitiveSceneInfo* SceneInfo;
+	bool bIsFarFieldPrimitive;
+	TArrayView<uint32> OutInstanceActivationMask;
+
+	const FRayTracingCullingParameters* CullingParameters;
+
+	void operator()() const;
+};
+
+#endif
diff --git a/Engine/Source/Runtime/Renderer/Private/RayTracing/RayTracingScene.cpp b/Engine/Source/Runtime/Renderer/Private/RayTracing/RayTracingScene.cpp
--- a/Engine/Source/Runtime/Renderer/Private/RayTracing/RayTracingScene.cpp
+++ b/Engine/Source/Runtime/Renderer/Private/RayTracing/RayTracingScene.cpp
@@ -34,7 +34,6 @@ FGraphEventRef FRayTracingScene::BeginCreate(FRDGBuilder& GraphBuilder)
 
 		NumNativeInstances += InstanceDesc.NumTransforms;
 	}
-	SET_DWORD_STAT(STAT_RayTracingInstances, NumNativeInstances);
 
 	// Round up number of instances to some multiple to avoid pathological growth reallocations.
 	static constexpr uint32 AllocationGranularity = 8 * 1024;
diff --git a/Engine/Source/Runtime/Renderer/Private/SceneRendering.h b/Engine/Source/Runtime/Renderer/Private/SceneRendering.h
--- a/Engine/Source/Runtime/Renderer/Private/SceneRendering.h
+++ b/Engine/Source/Runtime/Renderer/Private/SceneRendering.h
@@ -1122,6 +1122,23 @@ struct FRayTracingLightData
 	FBufferRHIRef									LightCullVolume;
 	FShaderResourceViewRHIRef						LightCullVolumeSRV;
 };
+
+struct FRayTracingCullingParameters
+{
+	int32 CullInRayTracing;
+	float CullingRadius;
+	float FarFieldCullingRadius;
+	float CullAngleThreshold;
+	float AngleThresholdRatio;
+	FVector ViewOrigin;
+	FVector ViewDirection;
+	bool bCullAllObjects;
+	bool bCullByRadiusOrDistance;
+	bool bIsRayTracingFarField;
+
+	void Init(FViewInfo& View);
+};
+
 #endif
 
 /** A FSceneView with additional state used by the scene renderer. */
@@ -1267,6 +1284,9 @@ public:
 	FRayTracingMeshCommandOneFrameArray VisibleRayTracingMeshCommands;
 	FDynamicRayTracingMeshCommandStorage DynamicRayTracingMeshCommandStorage;
 
+	FRayTracingCullingParameters RayTracingCullingParameters;
+	FGraphEventArray RayTracingPerInstanceCullingTaskList;
+
 	FGraphEventArray AddRayTracingMeshBatchTaskList;
 	TArray<FRayTracingMeshCommandOneFrameArray*, SceneRenderingAllocator> VisibleRayTracingMeshCommandsPerTask;
 	TArray<FDynamicRayTracingMeshCommandStorage*, SceneRenderingAllocator> DynamicRayTracingMeshCommandStoragePerTask;
diff --git a/Engine/Source/Runtime/Renderer/Public/PrimitiveSceneInfo.h b/Engine/Source/Runtime/Renderer/Public/PrimitiveSceneInfo.h
--- a/Engine/Source/Runtime/Renderer/Public/PrimitiveSceneInfo.h
+++ b/Engine/Source/Runtime/Renderer/Public/PrimitiveSceneInfo.h
@@ -401,6 +401,8 @@ public:
 	FRayTracingGeometryInstance CachedRayTracingInstance;
 	TArray<FMatrix> CachedRayTracingInstanceLocalTransforms;
 	TArray<FMatrix> CachedRayTracingInstanceWorldTransforms;
+	TArray<FBoxSphereBounds> CachedRayTracingInstanceWorldBounds;
+	int32 SmallestRayTracingInstanceWorldBoundsIndex;
 #endif
 
 	/** Initialization constructor. */
