Renamed SceneColorMobile / SceneColorDeferred to SceneColor.

diff --git a/Engine/Source/Runtime/Renderer/Private/SceneTextures.cpp b/Engine/Source/Runtime/Renderer/Private/SceneTextures.cpp
--- a/Engine/Source/Runtime/Renderer/Private/SceneTextures.cpp
+++ b/Engine/Source/Runtime/Renderer/Private/SceneTextures.cpp
@@ -657,7 +657,7 @@ FSceneTextures& FMinimalSceneTextures::Create(FRDGBuilder& GraphBuilder, const F
 			Flags |= TexCreate_UAV;
 		}
 
-		const TCHAR* SceneColorName = bIsMobilePlatform ? TEXT("SceneColorMobile") : TEXT("SceneColorDeferred");
+		const TCHAR* SceneColorName = TEXT("SceneColor");
 
 		// Create the scene color.
 		// TODO: Array-size could be values > 2, on upcoming XR devices. This should be part of the config.
