Fix condition for moving depth pre pass draw batches to velocity+depth pass.
Should match the logic for velocity mesh processor which is an AND of velocity conditions not an OR.

diff --git a/Engine/Source/Runtime/Renderer/Private/DepthRendering.cpp b/Engine/Source/Runtime/Renderer/Private/DepthRendering.cpp
--- a/Engine/Source/Runtime/Renderer/Private/DepthRendering.cpp
+++ b/Engine/Source/Runtime/Renderer/Private/DepthRendering.cpp
@@ -855,7 +855,7 @@ void FDepthPassMeshProcessor::AddMeshBatch(const FMeshBatch& RESTRICT MeshBatch,
 	if (bDraw && bRespectUseAsOccluderFlag && !MeshBatch.bUseAsOccluder && EarlyZPassMode < DDM_AllOpaque)
 	{
 		if (PrimitiveSceneProxy)
-	{
+		{
 			// Only render primitives marked as occluders.
 			bDraw = PrimitiveSceneProxy->ShouldUseAsOccluder()
 			// Only render static objects unless movable are requested.
@@ -863,9 +863,9 @@ void FDepthPassMeshProcessor::AddMeshBatch(const FMeshBatch& RESTRICT MeshBatch,
 
 			// Filter dynamic mesh commands by screen size.
 			if (ViewIfDynamicMeshCommand)
-	{
-		extern float GMinScreenRadiusForDepthPrepass;
-		const float LODFactorDistanceSquared = (PrimitiveSceneProxy->GetBounds().Origin - ViewIfDynamicMeshCommand->ViewMatrices.GetViewOrigin()).SizeSquared() * FMath::Square(ViewIfDynamicMeshCommand->LODDistanceFactor);
+			{
+				extern float GMinScreenRadiusForDepthPrepass;
+				const float LODFactorDistanceSquared = (PrimitiveSceneProxy->GetBounds().Origin - ViewIfDynamicMeshCommand->ViewMatrices.GetViewOrigin()).SizeSquared() * FMath::Square(ViewIfDynamicMeshCommand->LODDistanceFactor);
 				bDraw = bDraw && FMath::Square(PrimitiveSceneProxy->GetBounds().SphereRadius) > GMinScreenRadiusForDepthPrepass * GMinScreenRadiusForDepthPrepass * LODFactorDistanceSquared;
 			}
 		}
@@ -875,31 +875,30 @@ void FDepthPassMeshProcessor::AddMeshBatch(const FMeshBatch& RESTRICT MeshBatch,
 		}
 	}
 
-	// if we are skipping movable objects in early Z, which can happen in DDM_AllOpaqueNoVelocity
-	if (EarlyZPassMode == DDM_AllOpaqueNoVelocity && PrimitiveSceneProxy&& ViewIfDynamicMeshCommand)
+	// When using DDM_AllOpaqueNoVelocity we skip objects that will write depth+velocity in the subsequent velocity pass.
+	if (EarlyZPassMode == DDM_AllOpaqueNoVelocity && PrimitiveSceneProxy)
 	{
-		// We should ideally check to see if we are using the FOpaqueVelocityMeshProcessor or FTranslucentVelocityMeshProcessor. But for
-		// the object to get here, it would already be culled if it was translucent. We can safely use the FOpaqueVelocityMeshProcessor.
+		// We should ideally check to see if we this primitive is using the FOpaqueVelocityMeshProcessor or FTranslucentVelocityMeshProcessor. 
+		// But for the object to get here, it would already be culled if it was translucent, so we can assume FOpaqueVelocityMeshProcessor.
+		// This logic needs to match the logic in FOpaqueVelocityMeshProcessor::AddMeshBatch()
+		// todo: Move that logic to a single place.
 
-		// This logic is copy/paste/modified from FOpaqueVelocityMeshProcessor::AddMeshBatch(), but ideally we should clean it up into
-		// a single function that is shared to avoid breakages from code changes.
-		EShaderPlatform ShaderPlatform = ViewIfDynamicMeshCommand->GetShaderPlatform();
+		const EShaderPlatform ShaderPlatform = GetFeatureLevelShaderPlatform(FeatureLevel);
 		if (FOpaqueVelocityMeshProcessor::PrimitiveCanHaveVelocity(ShaderPlatform, PrimitiveSceneProxy))
 		{
-			bDraw = false;
-		}
-
-		if (FOpaqueVelocityMeshProcessor::PrimitiveHasVelocityForFrame(PrimitiveSceneProxy))
-		{
-			bDraw = false;
-		}
-
-		checkSlow(ViewIfDynamicMeshCommand->bIsViewInfo);
-		FViewInfo* ViewInfo = (FViewInfo*)ViewIfDynamicMeshCommand;
+			if (ViewIfDynamicMeshCommand)
+			{
+				if (FOpaqueVelocityMeshProcessor::PrimitiveHasVelocityForFrame(PrimitiveSceneProxy))
+				{
+					checkSlow(ViewIfDynamicMeshCommand->bIsViewInfo);
+					FViewInfo* ViewInfo = (FViewInfo*)ViewIfDynamicMeshCommand;
 
-		if (FOpaqueVelocityMeshProcessor::PrimitiveHasVelocityForView(*ViewInfo, PrimitiveSceneProxy))
-		{
-			bDraw = false;
+					if (FOpaqueVelocityMeshProcessor::PrimitiveHasVelocityForView(*ViewInfo, PrimitiveSceneProxy))
+					{
+						bDraw = false;
+					}
+				}
+			}
 		}
 	}
 
