Implement BxDF energy conservation for the path tracer

Implement basic energy conservation using albedo tables of the GGX and cloth lobes. Reflection and refraction are tabulated independently.

Also provide analytic approximations for the reflection case if a tabulation is not convenient. If a good analytic approximation for the glass case can be found (requires a 3D fit), we should remove the tabulated implementation which was done mainly to provide a baseline.

The albedo table is also used to improve lobe selection probabilities, which improves sampling of the clearcoat case in particular.

Refactored how the path tracer's internal view state data is stored to keep the structure details fully hidden inside PathTracing.cpp

diff --git a/Engine/Shaders/Private/PathTracing/Material/PathTracingClearCoat.ush b/Engine/Shaders/Private/PathTracing/Material/PathTracingClearCoat.ush
--- a/Engine/Shaders/Private/PathTracing/Material/PathTracingClearCoat.ush
+++ b/Engine/Shaders/Private/PathTracing/Material/PathTracingClearCoat.ush
@@ -15,6 +15,7 @@
 
 #include "PathTracingMaterialCommon.ush"
 #include "PathTracingFresnel.ush"
+#include "PathTracingEnergyConservation.ush"
 
 float3 CalcThinTransmission(float NoL, float NoV, FPathTracingPayload Payload)
 {
@@ -37,7 +38,7 @@ float3 CalcThinTransmission(float NoL, float NoV, FPathTracingPayload Payload)
 	return Transmission;
 }
 
-float CalcClearCoatSelectionPdf(FPathTracingPayload Payload)
+float CalcClearCoatSelectionPdf(float CoatEnergy)
 {
 #if EVAL_TOP_LAYER == 0
 	// never select coating if top is disabled
@@ -46,10 +47,8 @@ float CalcClearCoatSelectionPdf(FPathTracingPayload Payload)
 	// always select coating if bottom is disabled
 	return 1.0; 
 #endif
-	// probability of choosing the coat
-	// TODO: this could be improved if we had a directional albedo tabulation
-	float Pdf = Payload.GetClearCoat() * CLEAR_COAT_F0;
-	return Pdf;
+	// probability of choosing the coat (proportional to its directional albedo)
+	return CoatEnergy;
 }
 
 float ClearCoatFresnel(float VoH) {
@@ -70,8 +69,6 @@ FMaterialEval ClearCoat_EvalMaterial(
 	float3 L_World = OutgoingDirection;
 	float3 N_World = Payload.WorldNormal;
 
-	float CoatSelectionPdf = CalcClearCoatSelectionPdf(Payload);
-
 	float2 Alpha = 0;
 	float3x3 Basis = GetGGXBasis(Payload.GetClearCoatRoughness(), Payload.WorldNormal, Alpha);
 
@@ -85,12 +82,16 @@ FMaterialEval ClearCoat_EvalMaterial(
 	float NoH = saturate(H.z);
 	float VoH = saturate(dot(V, H));
 
+	const FBxDFEnergyTerms Coat = ComputeGGXSpecEnergyTerms(Payload.GetClearCoatRoughness(), NoV, CLEAR_COAT_F0);
+
 	// Specular lobe
 	const float2 GGXResult = GGXEvalReflection(L, V, H, Alpha);
-	const float F = Payload.GetClearCoat() * ClearCoatFresnel(VoH);
-	const float3 SpecWeight = F * GGXResult.x;
+	const float F = ClearCoatFresnel(VoH);
+	const float3 SpecWeight = F * GGXResult.x * Coat.W;
 	const float SpecPdf = GGXResult.y;
 
+	const float CoatSelectionPdf = CalcClearCoatSelectionPdf(Coat.E.x);
+
 	AddLobeWithMIS(Result.Weight, Result.Pdf, SpecWeight, SpecPdf, CoatSelectionPdf);
 	
 
@@ -113,14 +114,8 @@ FMaterialEval ClearCoat_EvalMaterial(
 	// Calculate transmission through substrate
 	float3 Transmission = CalcThinTransmission(NoL, NoV, Payload);
 
-	// Apply fresnel through exit interface
-	// TODO: This should actually use a measure of the directional albedo along V and L
-	// TODO: The original implementation tried to refract V and L, but doing this correctly would require a stochastic eval method
-	// NOTE: implementation in ClearCoatBxDF uses VoH for these which again makes the coating math dependent on the (stochastic) choice of H
-	float CC_F1 = ClearCoatFresnel(NoV);
-	float CC_F2 = ClearCoatFresnel(NoL);
-
-	BottomEval.Weight *= lerp(1.0, Transmission * (1 - CC_F1) * (1 - CC_F2), Payload.GetClearCoat());
+	BottomEval.Weight *= lerp(1.0, Transmission, Payload.GetClearCoat());
+	BottomEval.Weight *= 1.0 - Coat.E; // remove energy reflected by clearcoat to maintain energy preservation
 
 	AddLobeWithMIS(Result.Weight, Result.Pdf, BottomEval.Weight, BottomEval.Pdf, 1.0 - CoatSelectionPdf);
 
@@ -136,7 +131,10 @@ FMaterialSample ClearCoat_SampleMaterial(
 	float3 N_World = Payload.WorldNormal;
 	float3 V_World = -RayDirection;
 
-	float CoatSelectionPdf = CalcClearCoatSelectionPdf(Payload);
+	const float NoV = saturate(dot(N_World, V_World));
+	const FBxDFEnergyTerms Coat = ComputeGGXSpecEnergyTerms(Payload.Roughness, NoV, CLEAR_COAT_F0);
+
+	const float CoatSelectionPdf = CalcClearCoatSelectionPdf(Coat.E.x);
 
 	// choose between coat and bottom layer
 	float3 OutDirection = 0;
@@ -146,9 +144,8 @@ FMaterialSample ClearCoat_SampleMaterial(
 		RandSample.x /= CoatSelectionPdf;
 
 		// No anisotropy, just pick an arbitrary basis
-		float Roughness = max(Payload.GetClearCoatRoughness(), GGX_MIN_ROUGHNESS);
-		float2 Alpha = Roughness * Roughness;
-		float3x3 Basis = GetTangentBasis(N_World);
+		float2 Alpha = 0;
+		float3x3 Basis = GetGGXBasis(Payload.GetClearCoatRoughness(), Payload.WorldNormal, Alpha);
 
 		float3 V = -mul(Basis, RayDirection);
 
@@ -163,7 +160,7 @@ FMaterialSample ClearCoat_SampleMaterial(
 		}
 		// reflect and transform
 		OutDirection = mul(L, Basis);
-		OutRoughness = Roughness;
+		OutRoughness = Payload.GetClearCoatRoughness();
 	}
 	else
 	{
diff --git a/Engine/Shaders/Private/PathTracing/Material/PathTracingCloth.ush b/Engine/Shaders/Private/PathTracing/Material/PathTracingCloth.ush
--- a/Engine/Shaders/Private/PathTracing/Material/PathTracingCloth.ush
+++ b/Engine/Shaders/Private/PathTracing/Material/PathTracingCloth.ush
@@ -7,19 +7,7 @@
 #pragma once
 
 #include "PathTracingFresnel.ush"
-
-float ClothLobeSelectionPdf(float3 V, FPathTracingPayload Payload)
-{
-	// TODO: This could be improved if we had a directional albedo tabulation, but this formulation will at least prevent
-	// under-sampling when one lobe is black (like with metals)
-
-	const float3 Diff = Payload.DiffuseColor;
-	const float3 Spec = Payload.SpecularColor;
-	const float3 Fuzz = Payload.CustomData.xyz;
-	const float Cloth = Payload.CustomData.w;
-
-	return LobeSelectionProb(Diff + Cloth * Fuzz, (1.0 - Cloth) * Spec);
-}
+#include "PathTracingEnergyConservation.ush"
 
 FMaterialEval Cloth_EvalMaterial(
 	float3 IncomingDirection,
@@ -46,23 +34,30 @@ FMaterialEval Cloth_EvalMaterial(
 	const float VoH = saturate(dot(V, H));
 	const float NoH = saturate(H.z);
 
+	const float3 FuzzColor = Payload.CustomData.xyz;
+	const float Cloth = Payload.CustomData.w;
+
+	const FBxDFEnergyTerms Spec = ComputeGGXSpecEnergyTerms(Payload.Roughness, NoV, Payload.SpecularColor);
+	const FBxDFEnergyTerms Fuzz = ComputeClothEnergyTerms(Payload.Roughness, NoV, FuzzColor);
+
 	// Probability of picking diffuse lobe vs. specular lobe
-	const float LobeSelectionPdf = ClothLobeSelectionPdf(V_World, Payload);
+	const float3 DiffWeight = (1.0 - lerp(Spec.E, Fuzz.E, Cloth)) * Payload.DiffuseColor;
+	const float3 DiffLobeAlbedo = DiffWeight + Cloth * Fuzz.E;
+	const float3 SpecLobeAlbedo = (1.0 - Cloth) * Spec.E;
 
-	const float Cloth = Payload.CustomData.w;
+	const float LobeSelectionPdf = LobeSelectionProb(DiffLobeAlbedo, SpecLobeAlbedo);
 
 	// Diffuse and Cloth Lobe (see ClothBxDF)
-	const float3 FuzzColor = Payload.CustomData.xyz;
 	const float DCloth = D_InvGGX(Pow4(Payload.Roughness), NoH);
 	const float3 FCloth = F_Schlick(FuzzColor, VoH);
-	const float3 ClothWeight = NoL > 0 && NoV > 0 ? (PI * Cloth * DCloth * Vis_Cloth(NoV, NoL)) * FCloth : 0.0;
+	const float3 ClothWeight = NoL > 0 && NoV > 0 ? (PI * Cloth * DCloth * Vis_Cloth(NoV, NoL)) * FCloth * Fuzz.W : 0.0;
 
-	AddLobeWithMIS(Result.Weight, Result.Pdf, Payload.DiffuseColor + ClothWeight, NoL / PI, LobeSelectionPdf);
+	AddLobeWithMIS(Result.Weight, Result.Pdf, DiffWeight + ClothWeight, NoL / PI, LobeSelectionPdf);
 
 	// Specular lobe
 	const float2 GGXResult = GGXEvalReflection(L, V, H, Alpha);
 	const float3 F = F_Schlick(Payload.SpecularColor, VoH);
-	const float3 SpecWeight = (1 - Cloth) * F * GGXResult.x;
+	const float3 SpecWeight = (1 - Cloth) * F * GGXResult.x * Spec.W;
 	const float SpecPdf = GGXResult.y;
 
 	AddLobeWithMIS(Result.Weight, Result.Pdf, SpecWeight, SpecPdf, 1 - LobeSelectionPdf);
@@ -79,14 +74,29 @@ FMaterialSample Cloth_SampleMaterial(
 	float3 N = Payload.WorldNormal;
 	float3 V = -RayDirection;
 
-	const float LobeSelectionPdf = ClothLobeSelectionPdf(V, Payload);
-
 	float2 Alpha;
 	float3x3 Basis = GetGGXBasis(Payload.Roughness, Payload.Anisotropy, Payload.WorldNormal, Payload.WorldTangent, Alpha);
 
 	V = mul(Basis, V);
 
-	// Randomly choose to sample diffuse, specular or cloth lobe
+	const float NoV = saturate(V.z);
+
+	const float3 FuzzColor = Payload.CustomData.xyz;
+	const float Cloth = Payload.CustomData.w;
+	const FBxDFEnergyTerms Spec = ComputeGGXSpecEnergyTerms(Payload.Roughness, NoV, Payload.SpecularColor);
+	const FBxDFEnergyTerms Fuzz = ComputeClothEnergyTerms(Payload.Roughness, NoV, FuzzColor);
+
+
+	// Probability of picking diffuse lobe vs. specular lobe
+	const float3 DiffWeight = (1.0 - lerp(Spec.E, Fuzz.E, Cloth)) * Payload.DiffuseColor;
+	const float3 DiffLobeAlbedo = DiffWeight + Cloth * Fuzz.E;
+	const float3 SpecLobeAlbedo = (1.0 - Cloth) * Spec.E;
+
+	const float LobeSelectionPdf = LobeSelectionProb(DiffLobeAlbedo, SpecLobeAlbedo);
+
+
+	// Randomly choose to sample diffuse+cloth or specular lobe
+	// The cloth lobe is low enough frequency that sampling with cosine weighting is sufficient
 	float3 L = 0, H = 0;
 	float OutRoughness = 0;
 	if (RandSample.x < LobeSelectionPdf)
@@ -125,26 +135,23 @@ FMaterialSample Cloth_SampleMaterial(
 	// With a valid direction in hand -- now evaluate the BxDF (taking advantage of already computed terms)
 
 	const float NoL = saturate(L.z);
-	const float NoV = saturate(V.z);
 	const float VoH = saturate(dot(V, H));
 	const float NoH = saturate(H.z);
 
-	const float Cloth = Payload.CustomData.w;
-
 	// Diffuse and Cloth Lobe (see ClothBxDF)
-	const float3 FuzzColor = Payload.CustomData.xyz;
 	const float DCloth = D_InvGGX(Pow4(Payload.Roughness), NoH);
 	const float3 FCloth = F_Schlick(FuzzColor, VoH);
-	const float3 ClothWeight = NoL > 0 && NoV > 0 ? (PI * Cloth * DCloth * Vis_Cloth(NoV, NoL)) * FCloth : 0.0;
+	const float3 ClothWeight = NoL > 0 && NoV > 0 ? (PI * Cloth * DCloth * Vis_Cloth(NoV, NoL)) * FCloth * Fuzz.W : 0.0;
 
 	float3 OutWeight = 0;
 	float OutPdf = 0;
-	AddLobeWithMIS(OutWeight, OutPdf, Payload.DiffuseColor + ClothWeight, NoL / PI, LobeSelectionPdf);
+	AddLobeWithMIS(OutWeight, OutPdf, DiffWeight + ClothWeight, NoL / PI, LobeSelectionPdf);
 
 	// Specular lobe
+
 	const float2 GGXResult = GGXEvalReflection(L, V, H, Alpha);
 	const float3 F = F_Schlick(Payload.SpecularColor, VoH);
-	const float3 SpecWeight = (1 - Cloth) * F * GGXResult.x;
+	const float3 SpecWeight = (1 - Cloth) * F * GGXResult.x * Spec.W;
 	const float SpecPdf = GGXResult.y;
 	AddLobeWithMIS(OutWeight, OutPdf, SpecWeight, SpecPdf, 1 - LobeSelectionPdf);
 
diff --git a/Engine/Shaders/Private/PathTracing/Material/PathTracingDefaultLit.ush b/Engine/Shaders/Private/PathTracing/Material/PathTracingDefaultLit.ush
--- a/Engine/Shaders/Private/PathTracing/Material/PathTracingDefaultLit.ush
+++ b/Engine/Shaders/Private/PathTracing/Material/PathTracingDefaultLit.ush
@@ -8,17 +8,7 @@
 
 #include "PathTracingMaterialCommon.ush"
 #include "PathTracingFresnel.ush"
-
-float CalcLobeSelectionPdf(float3 V, FPathTracingPayload Payload)
-{
-	// TODO: This could be improved if we had a directional albedo tabulation, but this formulation will at least prevent
-	// under-sampling when one lobe is black (like with metals)
-
-	float3 Diff = Payload.DiffuseColor;
-	float3 Spec = Payload.SpecularColor;
-
-	return LobeSelectionProb(Diff, Spec);
-}
+#include "PathTracingEnergyConservation.ush"
 
 FMaterialEval DefaultLit_EvalMaterial(
 	float3 IncomingDirection,
@@ -40,19 +30,22 @@ FMaterialEval DefaultLit_EvalMaterial(
 	const float3 L = mul(Basis, L_World);
 	const float3 H = normalize(V + L);
 
+	const float NoV = saturate(V.z);
 	const float NoL = saturate(L.z);
 	const float VoH = saturate(dot(V, H));
 
+	const FBxDFEnergyTerms Spec = ComputeGGXSpecEnergyTerms(Payload.Roughness, NoV, Payload.SpecularColor);
+
 	// Probability of picking diffuse lobe vs. specular lobe
-	const float LobeSelectionPdf = CalcLobeSelectionPdf(V_World, Payload);
+	const float LobeSelectionPdf = LobeSelectionProb(Payload.DiffuseColor * (1 - Spec.E), Spec.E);
 
 	// Diffuse Lobe
-	AddLobeWithMIS(Result.Weight, Result.Pdf, Payload.DiffuseColor, NoL / PI, LobeSelectionPdf);
+	AddLobeWithMIS(Result.Weight, Result.Pdf, (1 - Spec.E) * Payload.DiffuseColor, NoL / PI, LobeSelectionPdf);
 
 	// Specular lobe
 	const float2 GGXResult = GGXEvalReflection(L, V, H, Alpha);
 	const float3 F = F_Schlick(Payload.SpecularColor, VoH);
-	const float3 SpecWeight = F * GGXResult.x;
+	const float3 SpecWeight = F * GGXResult.x * Spec.W;
 	const float SpecPdf = GGXResult.y;
 
 	AddLobeWithMIS(Result.Weight, Result.Pdf, SpecWeight, SpecPdf, 1.0 - LobeSelectionPdf);
@@ -69,13 +62,17 @@ FMaterialSample DefaultLit_SampleMaterial(
 	float3 N = Payload.WorldNormal;
 	float3 V = -RayDirection;
 	
-	float LobeSelectionPdf = CalcLobeSelectionPdf(V, Payload);
-
 	float2 Alpha;
 	float3x3 Basis = GetGGXBasis(Payload.Roughness, Payload.Anisotropy, Payload.WorldNormal, Payload.WorldTangent, Alpha);
 
 	V = mul(Basis, V);
-	
+
+	const float NoV = saturate(V.z);
+	const FBxDFEnergyTerms Spec = ComputeGGXSpecEnergyTerms(Payload.Roughness, NoV, Payload.SpecularColor);
+
+	// Probability of picking diffuse lobe vs. specular lobe
+	const float LobeSelectionPdf = LobeSelectionProb(Payload.DiffuseColor * (1 - Spec.E), Spec.E);
+
 	// Randomly choose to sample diffuse or specular
 	float3 L = 0, H = 0;
 	float OutRoughness = 0;
@@ -117,12 +114,12 @@ FMaterialSample DefaultLit_SampleMaterial(
 	float OutPdf = 0;
 
 	// Diffuse Lobe
-	AddLobeWithMIS(OutWeight, OutPdf, Payload.DiffuseColor, NoL / PI, LobeSelectionPdf);
+	AddLobeWithMIS(OutWeight, OutPdf, (1 - Spec.E) * Payload.DiffuseColor, NoL / PI, LobeSelectionPdf);
 
 	// Specular lobe
 	const float2 GGXResult = GGXEvalReflection(L, V, H, Alpha);
 	const float3 F = F_Schlick(Payload.SpecularColor, VoH);
-	const float3 SpecWeight = F * GGXResult.x;
+	const float3 SpecWeight = F * GGXResult.x * Spec.W;
 	const float SpecPdf = GGXResult.y;
 	AddLobeWithMIS(OutWeight, OutPdf, SpecWeight, SpecPdf, 1.0 - LobeSelectionPdf);
 
diff --git a/Engine/Shaders/Private/PathTracing/Material/PathTracingEnergyConservation.ush b/Engine/Shaders/Private/PathTracing/Material/PathTracingEnergyConservation.ush
new file mode 100644
--- /dev/null
+++ b/Engine/Shaders/Private/PathTracing/Material/PathTracingEnergyConservation.ush
@@ -0,0 +1,113 @@
+// Copyright Epic Games, Inc. All Rights Reserved.
+
+#pragma once
+
+// Energy compensation/preservation for the various BxDFs used in the path tracer. This file provides utility functions to tweak lobe weights
+// 
+// References:
+// [1] "Revisiting physically based shading at Imageworks" - Christopher Kulla & Alejandro Conty
+//   https://blog.selfshadow.com/publications/s2017-shading-course/imageworks/s2017_pbs_imageworks_slides_v2.pdf
+// [2] "Practical multiple scattering compensation for microfacet models" - Emmanuel Turquin
+//   https://blog.selfshadow.com/publications/turquin/ms_comp_final.pdf
+// [3] "A Multiple-Scattering Microfacet Model for Real-Time Image Based Lighting" - Carmelo J. Fdez-Agüera
+//   http://jcgt.org/published/0008/01/03/
+// 
+// We follow the terminology from [1]:
+// Energy compensation: adjust individual BxDFs so that they have unit albedo from all viewing angles, compensating for (for example) the lack of multiple scattering in microfacet BxDFs
+// Energy preservation: adjust the combination of lobes (such as diffuse+specular) to ensure we never create additional energy from any angle
+
+// The suggested implementation in [1] is cumbersome for a GPU implementation because it requires table lookups during sampling to achieve low variance. We instead sacrifice reciprocity and
+// simply divide the BxDF by its (tabulated/fit) directional albedo to enforce conservation of energy. We follow the approach outlined in [2] to account for spec/cloth and glass, and use
+// the split-sum approach from [3] for the diffuse/specular energy preservation case (which generalizes to other lobes as well).
+
+// USE_PATH_TRACING_ENERGY_CONSERVATION controls the implementation details of this file
+//    0: disable conservation, all bsdfs are assumed to already have unit albedo from all viewing angles
+//    1: use a small baked texture to store the directional albedo (and fresnel weighted albedo)
+//    2: use an analytic approximation to directional albedo (doesn't require any textures to be bound and can be faster -- WIP)
+
+#if USE_PATH_TRACING_ENERGY_CONSERVATION == 1
+// when using this method, the shader will need access to the following textures
+Texture2D<float2> GGXSpecEnergyTexture;
+Texture3D<float2> GGXGlassEnergyTexture;
+Texture2D<float2> ClothEnergyTexture;
+SamplerState EnergySampler;
+#endif
+
+
+float2 GGXEnergyLookup(float Roughness, float NoV)
+{
+#if USE_PATH_TRACING_ENERGY_CONSERVATION == 1
+	return GGXSpecEnergyTexture.SampleLevel(EnergySampler, float2(NoV, Roughness), 0);
+#elif USE_PATH_TRACING_ENERGY_CONSERVATION == 2
+	const float r = Roughness;
+	const float c = NoV;
+	const float E = 1.0 - saturate(pow(r, c / r) * ((r * c + 0.0266916) / (0.466495 + c)));
+	const float Ef = Pow5(1 - c) * pow(2.36651 * pow(c, 4.7703 * r) + 0.0387332, r);
+	return float2(E, Ef);
+#else
+	return float2(1, 0);
+#endif
+}
+
+float GGXEnergyLookup(float Roughness, float NoV, float Eta)
+{
+	// Following [2] Eq 18, we simply divide the whole bsdf by its single scattering albedo to conserve energy
+	// This break reciprocity, but is much simpler than the approach outlined in [1] which requires substantially more computation.
+#if USE_PATH_TRACING_ENERGY_CONSERVATION == 1
+	// NOTE: Eta is encoded for [1,3] range only, energy loss will happens above 3.0 - but this keeps more resolution on the visually significant portion of the range
+	float2 E = GGXGlassEnergyTexture.SampleLevel(EnergySampler, float3(NoV, Roughness, max(Eta, rcp(Eta)) * 0.5 - 0.5), 0);
+	return Eta >= 1.0 ? E.x : E.y;
+#else
+	// TODO: find an analytic fit for this case
+	return 1.0;
+#endif
+}
+
+float2 ClothEnergyLookup(float Roughness, float NoV)
+{
+#if USE_PATH_TRACING_ENERGY_CONSERVATION == 1
+	return ClothEnergyTexture.SampleLevel(EnergySampler, float2(NoV, Roughness), 0);
+#elif USE_PATH_TRACING_ENERGY_CONSERVATION == 2
+	// Simple analytical fit (avoids building a table)
+	// NOTE: this fit is not exact near grazing angles for roughness ~0.7
+	const float c = NoV;
+	const float r = Roughness;
+	return float2(
+		(0.526422 / ((-0.227114 + r) * (-0.968835 + r) * ((5.38869 - 20.2835 * c) * r) - (-1.18761 - ((2.58744 - c) * c)))) + 0.0615456,
+		0.0 // TODO: is it worth capturing this? this term is close to 0 (and experiments with the tabulated form show it doesn't have much impact)
+	);
+#else
+	return float2(1, 0);
+#endif
+}
+
+struct FBxDFEnergyTerms {
+	float3 W; // overall weight to scale the lobe BxDF by to ensure energy conservation
+	float3 E; // Directional albedo of the lobe for energy preservation and lobe picking
+};
+
+// Given a split-sum approximation of directional albedo for a BxDF, compute multiple scattering weight and multiple scattering directional albedo
+// while taking into account the fresnel term (assumed to be F_Schlick)
+FBxDFEnergyTerms ComputeFresnelEnergyTerms(float2 E, float3 F0)
+{
+	FBxDFEnergyTerms Result;
+	// [2] Eq 16: this restores the missing energy of the bsdf, while also accounting for the fact that the fresnel term causes some energy to be absorbed
+	// NOTE: using F0 here is an approximation, but for schlick fresnel Favg is almost exactly equal to F0
+	Result.W = 1.0 + F0 * ((1 - E.x) / E.x);
+	// Now estimate the amount of energy reflected off this specular lobe so that we can remove it from underlying BxDF layers (like diffuse)
+	// This relies on the split-sum approximation as in [3] Sec 4.
+	// This term can also be useful to compute the probability of choosing among lobes
+	const float F90 = saturate(50.0 * F0.g); // See F_Schlick implementation
+	Result.E = Result.W * (E.x * F0 + E.y * (F90 - F0));
+	return Result;
+}
+
+FBxDFEnergyTerms ComputeGGXSpecEnergyTerms(float Roughness, float NoV, float3 F0)
+{
+	return ComputeFresnelEnergyTerms(GGXEnergyLookup(Roughness, NoV), F0);
+}
+
+FBxDFEnergyTerms ComputeClothEnergyTerms(float Roughness, float NoV, float3 F0)
+{
+	return ComputeFresnelEnergyTerms(ClothEnergyLookup(Roughness, NoV), F0);
+}
diff --git a/Engine/Shaders/Private/PathTracing/Material/PathTracingGlossy.ush b/Engine/Shaders/Private/PathTracing/Material/PathTracingGlossy.ush
--- a/Engine/Shaders/Private/PathTracing/Material/PathTracingGlossy.ush
+++ b/Engine/Shaders/Private/PathTracing/Material/PathTracingGlossy.ush
@@ -6,6 +6,9 @@ Glossy.usf: Microfacet GGX BRDF sampling functions
 
 #pragma once
 
+#include "../../BRDF.ush"
+#include "../../MonteCarlo.ush"
+
 // be conservative since the GPU may be aggressive with floating point optimizations
 #define GGX_MIN_ROUGHNESS 0.01
 #define GGX_MIN_ANISOTROPY 0.001 // ignore anisotropy below this amount
@@ -60,27 +63,25 @@ float2 GGXEvalReflection(float3 L, float3 V, float3 H, float2 Alpha)
 	return 0;
 }
 
-
 // Compute GGX lobe Weight and Pdf (without Fresnel term) given a set of vectors in local space (Z up)
 float2 GGXEvalRefraction(float3 L, float3 V, float3 H, float2 Alpha, float Eta)
 {
 	const float NoV = saturate(V.z);
 	const float NoL = saturate(-L.z); // NOTE: L should point down
 
-	// This checks that V and L point in opposite sides 
-	if (NoL > 0 && NoV > 0)
+	// This checks that V and L point in opposite sides and that the half-vector is valid
+	if (NoL > 0 && NoV > 0 && H.z > 0)
 	{
 		const float D = D_GGXaniso(Alpha.x, Alpha.y, H.z, H.x, H.y);
 
-		const float VoH = dot(V, H);
-		const float LoH = dot(L, H);
-		const float3 Ht = (Eta * L + V);
+		const float VoH = saturate( dot(V, H));
+		const float LoH = saturate(-dot(L, H));
 
 		// Same as logic above, only the pdf math differs slightly
 		const float LenL = length(float3(L.x * Alpha.x, L.y * Alpha.y, NoL));
 		const float LenV = length(float3(V.x * Alpha.x, V.y * Alpha.y, NoV));
 		const float Weight = NoL * (LenV + NoV) / (NoV * LenL + NoL * LenV);
-		const float Pdf = (abs(LoH * VoH) * (Eta * Eta) * (2.0 * D)) / (dot(Ht, Ht) * (LenV + NoV));
+		const float Pdf = (LoH * VoH * Eta * Eta * 2.0 * D) / (Pow2(VoH - Eta * LoH) * (LenV + NoV));
 
 		return float2(Weight, Pdf);
 	}
diff --git a/Engine/Shaders/Private/PathTracing/Material/PathTracingSolidGlass.ush b/Engine/Shaders/Private/PathTracing/Material/PathTracingSolidGlass.ush
--- a/Engine/Shaders/Private/PathTracing/Material/PathTracingSolidGlass.ush
+++ b/Engine/Shaders/Private/PathTracing/Material/PathTracingSolidGlass.ush
@@ -9,6 +9,7 @@ PathTracingSolidGlass.ush: Microfacet Refractive BSDF for solid glass
 #include "PathTracingMaterialCommon.ush"
 #include "PathTracingFresnel.ush"
 #include "PathTracingGlossy.ush"
+#include "PathTracingEnergyConservation.ush"
 
 FMaterialSample RoughGlass_SampleMaterial(
 	float3 RayDirection,
@@ -22,7 +23,7 @@ FMaterialSample RoughGlass_SampleMaterial(
 
 	// #dxr_todo: Maintain a refraction stack on the path tracing payload
 	float Ior = Payload.Ior;
-	float Eta = bIsEntering ? Ior : 1.0 / Ior;
+	float Eta = bIsEntering ? Ior : rcp(Ior);
 
 	float2 Alpha;
 	float3x3 Basis = GetGGXBasis(Payload.Roughness, Payload.Anisotropy, Payload.WorldNormal, Payload.WorldTangent, Alpha);
@@ -36,6 +37,9 @@ FMaterialSample RoughGlass_SampleMaterial(
 		return NullMaterialSample();
 	}
 
+	// correct for energy loss by scaling the whole BSDF
+	const float E = GGXEnergyLookup(Payload.Roughness, NoV, Eta);
+
 	float3 H = ImportanceSampleVisibleGGX_aniso(UniformSampleDisk(RandSample.xy), Alpha, V).xyz;
 	float3 N = float3(0, 0, 1);
 	float NoH = saturate(H.z);
@@ -48,14 +52,14 @@ FMaterialSample RoughGlass_SampleMaterial(
 		const float2 GGXResult = GGXEvalRefraction(L, V, H, Alpha, Eta);
 
 		// transform to world space
-		return CreateMaterialSample(normalize(mul(L, Basis)), GGXResult.x, F * GGXResult.y, -1.0, Payload.Roughness);
+		return CreateMaterialSample(normalize(mul(L, Basis)), GGXResult.x / E, F * GGXResult.y, -1.0, Payload.Roughness);
 	}
 	else
 	{
 		const float2 GGXResult = GGXEvalReflection(L, V, H, Alpha);
 
 		// transform to world space
-		return CreateMaterialSample(normalize(mul(L, Basis)), GGXResult.x, F * GGXResult.y, 1.0, Payload.Roughness);
+		return CreateMaterialSample(normalize(mul(L, Basis)), GGXResult.x / E, F * GGXResult.y, 1.0, Payload.Roughness);
 	}
 }
 
@@ -73,7 +77,7 @@ FMaterialEval RoughGlass_EvalMaterial(
 
 	// #dxr_todo: Maintain a refraction stack on the path tracing payload
 	float Ior = Payload.Ior;
-	float Eta = bIsEntering ? Ior : 1.0 / Ior;
+	float Eta = bIsEntering ? Ior : rcp(Ior);
 
 	float2 Alpha;
 	float3x3 Basis = GetGGXBasis(Payload.Roughness, Payload.Anisotropy, Payload.WorldNormal, Payload.WorldTangent, Alpha);
@@ -89,6 +93,10 @@ FMaterialEval RoughGlass_EvalMaterial(
 	}
 
 	float NoV = saturate(V.z);
+
+	// correct for energy loss by scaling the whole BSDF
+	const float E = GGXEnergyLookup(Payload.Roughness, NoV, Eta);
+
 	if (L.z > 0)
 	{
 		// reflection side
@@ -97,23 +105,22 @@ FMaterialEval RoughGlass_EvalMaterial(
 		const float F = FresnelReflectance(VoH, Eta);
 
 		const float2 GGXResult = GGXEvalReflection(L, V, H, Alpha);
-		return CreateMaterialEval(GGXResult.x, F * GGXResult.y);
+		return CreateMaterialEval(GGXResult.x / E, F * GGXResult.y);
 	}
 	else if (L.z < 0)
 	{
 		// refraction side
 		float NoL = saturate(-L.z);
 		float3 Ht = -(Eta * L + V);
-		Ht = (Eta < 1.0f) ? -Ht : Ht;
-		float3 Htn = normalize(Ht);
+		Ht = normalize((Eta < 1.0f) ? -Ht : Ht);
 
-		float VoH = dot(V, Htn);
+		float VoH = dot(V, Ht);
 		float F = 1.0f - FresnelReflectance(VoH, Eta);
 		if (F > 0)
 		{
-			const float2 GGXResult = GGXEvalRefraction(L, V, Htn, Alpha, Eta);
+			const float2 GGXResult = GGXEvalRefraction(L, V, Ht, Alpha, Eta);
 
-			return CreateMaterialEval(GGXResult.x, F * GGXResult.y);
+			return CreateMaterialEval(GGXResult.x / E, F * GGXResult.y);
 		}
 	}
 	return NullMaterialEval();
diff --git a/Engine/Shaders/Private/PathTracing/Material/PathTracingSubsurfaceProfile.ush b/Engine/Shaders/Private/PathTracing/Material/PathTracingSubsurfaceProfile.ush
--- a/Engine/Shaders/Private/PathTracing/Material/PathTracingSubsurfaceProfile.ush
+++ b/Engine/Shaders/Private/PathTracing/Material/PathTracingSubsurfaceProfile.ush
@@ -13,16 +13,11 @@
 #include "PathTracingMaterialCommon.ush"
 #include "PathTracingFresnel.ush"
 
-float2 SubsurfaceProfile_LobeCdf(float3 V, FPathTracingPayload Payload)
+float2 SubsurfaceProfile_LobeCdf(float3 DiffuseColor, float3 SpecE, float SpecMix)
 {
-	// TODO: This could be improved if we had a directional albedo tabulation, but this formulation will at least prevent
-	// under-sampling when one lobe is black (like with metals)
+	float3 DiffE = (1 - SpecE) * DiffuseColor;
 
-	float3 Diff = Payload.DiffuseColor;
-	float3 Spec = Payload.SpecularColor;
-
-	float DiffuseProb = LobeSelectionProb(Diff, Spec);
-	float SpecMix = Payload.GetDualRoughnessSpecular().z;
+	float DiffuseProb = LobeSelectionProb(DiffE, SpecE);
 	return float2(DiffuseProb, DiffuseProb + (1 - DiffuseProb) * (1 - SpecMix));
 }
 
@@ -49,21 +44,26 @@ FMaterialEval SubsurfaceProfile_EvalMaterial(
 	const float3 L = mul(Basis, L_World);
 	const float3 H = normalize(V + L);
 
+	const float NoV = saturate(V.z);
+	const FBxDFEnergyTerms Spec0 = ComputeGGXSpecEnergyTerms(DualRoughnessData.x, NoV, Payload.SpecularColor);
+	const FBxDFEnergyTerms Spec1 = ComputeGGXSpecEnergyTerms(DualRoughnessData.y, NoV, Payload.SpecularColor);
+	const float3 SpecE = lerp(Spec0.E, Spec1.E, DualRoughnessData.z);
+	const float2 LobeCdf = SubsurfaceProfile_LobeCdf(Payload.DiffuseColor, SpecE, DualRoughnessData.z);
+	const float3 LobeProb = float3(LobeCdf.x, LobeCdf.y - LobeCdf.x, 1 - LobeCdf.y);
+
+
 	const float NoL = saturate(L.z);
 	const float VoH = saturate(dot(V, H));
 
-	const float2 LobeCdf = SubsurfaceProfile_LobeCdf(V_World, Payload);
-	const float3 LobeProb = float3(LobeCdf.x, LobeCdf.y - LobeCdf.x, 1 - LobeCdf.y);
-
 	// Diffuse Lobe
-	AddLobeWithMIS(Result.Weight, Result.Pdf, Payload.DiffuseColor, NoL / PI, LobeProb.x);
+	AddLobeWithMIS(Result.Weight, Result.Pdf, Payload.DiffuseColor * (1 - SpecE), NoL / PI, LobeProb.x);
 
 	// Specular lobes
 	const float2 GGXResult0 = GGXEvalReflection(L, V, H, Alpha0);
 	const float2 GGXResult1 = GGXEvalReflection(L, V, H, Alpha1);
 	const float3 F = F_Schlick(Payload.SpecularColor, VoH);
-	AddLobeWithMIS(Result.Weight, Result.Pdf, (1 - DualRoughnessData.z) * F * GGXResult0.x, GGXResult0.y, LobeProb.y);
-	AddLobeWithMIS(Result.Weight, Result.Pdf, (    DualRoughnessData.z) * F * GGXResult1.x, GGXResult1.y, LobeProb.z);
+	AddLobeWithMIS(Result.Weight, Result.Pdf, (1 - DualRoughnessData.z) * F * GGXResult0.x * Spec0.W, GGXResult0.y, LobeProb.y);
+	AddLobeWithMIS(Result.Weight, Result.Pdf, (    DualRoughnessData.z) * F * GGXResult1.x * Spec1.W, GGXResult1.y, LobeProb.z);
 
 	return Result;
 }
@@ -77,9 +77,6 @@ FMaterialSample SubsurfaceProfile_SampleMaterial(
 	float3 N = Payload.WorldNormal;
 	float3 V = -RayDirection;
 
-	const float2 LobeCdf = SubsurfaceProfile_LobeCdf(V, Payload);
-	const float3 LobeProb = float3(LobeCdf.x, LobeCdf.y - LobeCdf.x, 1 - LobeCdf.y);
-
 	float3 DualRoughnessData = Payload.GetDualRoughnessSpecular();
 	float2 Alpha0 = 0;
 	float2 Alpha1 = 0;
@@ -88,6 +85,13 @@ FMaterialSample SubsurfaceProfile_SampleMaterial(
 
 	V = mul(Basis, V);
 
+	const float NoV = saturate(V.z);
+	const FBxDFEnergyTerms Spec0 = ComputeGGXSpecEnergyTerms(DualRoughnessData.x, NoV, Payload.SpecularColor);
+	const FBxDFEnergyTerms Spec1 = ComputeGGXSpecEnergyTerms(DualRoughnessData.y, NoV, Payload.SpecularColor);
+	const float3 SpecE = lerp(Spec0.E, Spec1.E, DualRoughnessData.z);
+	const float2 LobeCdf = SubsurfaceProfile_LobeCdf(Payload.DiffuseColor, SpecE, DualRoughnessData.z);
+	const float3 LobeProb = float3(LobeCdf.x, LobeCdf.y - LobeCdf.x, 1 - LobeCdf.y);
+
 	// Randomly choose to sample diffuse or specular
 	float3 L = 0, H = 0;
 	float OutRoughness = 0;
@@ -147,14 +151,14 @@ FMaterialSample SubsurfaceProfile_SampleMaterial(
 	float OutPdf = 0;
 
 	// Diffuse Lobe
-	AddLobeWithMIS(OutWeight, OutPdf, Payload.DiffuseColor, NoL / PI, LobeProb.x);
+	AddLobeWithMIS(OutWeight, OutPdf, Payload.DiffuseColor * (1 - SpecE), NoL / PI, LobeProb.x);
 
 	// Specular lobe
 	const float2 GGXResult0 = GGXEvalReflection(L, V, H, Alpha0);
 	const float2 GGXResult1 = GGXEvalReflection(L, V, H, Alpha1);
 	const float3 F = F_Schlick(Payload.SpecularColor, VoH);
-	AddLobeWithMIS(OutWeight, OutPdf, (1 - DualRoughnessData.z) * F * GGXResult0.x, GGXResult0.y, LobeProb.y);
-	AddLobeWithMIS(OutWeight, OutPdf, (    DualRoughnessData.z) * F * GGXResult1.x, GGXResult1.y, LobeProb.z);
+	AddLobeWithMIS(OutWeight, OutPdf, (1 - DualRoughnessData.z) * F * GGXResult0.x * Spec0.W, GGXResult0.y, LobeProb.y);
+	AddLobeWithMIS(OutWeight, OutPdf, (    DualRoughnessData.z) * F * GGXResult1.x * Spec1.W, GGXResult1.y, LobeProb.z);
 
 	// transform to world space
 	return CreateMaterialSample(normalize(mul(L, Basis)), OutWeight, OutPdf, 1.0, OutRoughness);
diff --git a/Engine/Shaders/Private/PathTracing/Material/PathTracingThinGlass.ush b/Engine/Shaders/Private/PathTracing/Material/PathTracingThinGlass.ush
--- a/Engine/Shaders/Private/PathTracing/Material/PathTracingThinGlass.ush
+++ b/Engine/Shaders/Private/PathTracing/Material/PathTracingThinGlass.ush
@@ -9,7 +9,7 @@ PathTracingThinGlass.ush: Microfacet BSDF for thin glass
 #include "PathTracingMaterialCommon.ush"
 #include "PathTracingFresnel.ush"
 #include "PathTracingGlossy.ush"
-
+#include "PathTracingEnergyConservation.ush"
 
 struct SlabWeights
 {
@@ -87,13 +87,17 @@ FMaterialEval RoughThinGlass_EvalMaterial(
 
 	const float2 GGXResult = GGXEvalReflection(L, V, H, Alpha);
 
+	// GGX reflected energy is split between reflection and refraction equally
+	// so just use the overall GGX energy loss for energy compensation (as in the solid glass case)
+	const float OutWeight = GGXResult.x / GGXEnergyLookup(Payload.Roughness, NoV).x;
+
 	if (IsTransmission)
 	{
-		return CreateMaterialEval(GGXResult.x / (1.0 - ProbR) * Result.Transmitted, GGXResult.y * (1.0 - ProbR));
+		return CreateMaterialEval(OutWeight / (1.0 - ProbR) * Result.Transmitted, GGXResult.y * (1.0 - ProbR));
 	}
 	else
 	{
-		return CreateMaterialEval(GGXResult.x / ProbR * Result.Reflected, GGXResult.y * ProbR);
+		return CreateMaterialEval(OutWeight / ProbR * Result.Reflected, GGXResult.y * ProbR);
 	}
 }
 
@@ -139,7 +143,10 @@ FMaterialSample RoughThinGlass_SampleMaterial(
 	const float ProbR = LobeSelectionProb(Result.Reflected, Result.Transmitted);
 
 	const float2 GGXResult = GGXEvalReflection(L, V, H, Alpha);
-	float3 OutWeight = GGXResult.x;
+
+	// GGX reflected energy is split between reflection and refraction equally
+	// so just use the overall GGX energy loss for energy compensation (as in the solid glass case)
+	float3 OutWeight = GGXResult.x / GGXEnergyLookup(Payload.Roughness, NoV).x;
 	float OutPdf = GGXResult.y;
 
 	// decide between reflection/transmission
diff --git a/Engine/Shaders/Private/PathTracing/Material/PathTracingTwoSidedFoliage.ush b/Engine/Shaders/Private/PathTracing/Material/PathTracingTwoSidedFoliage.ush
--- a/Engine/Shaders/Private/PathTracing/Material/PathTracingTwoSidedFoliage.ush
+++ b/Engine/Shaders/Private/PathTracing/Material/PathTracingTwoSidedFoliage.ush
@@ -7,18 +7,7 @@
 #pragma once
 
 #include "PathTracingDefaultLit.ush"
-
-float CalcTransmissionSelectionPdf(
-	float3 IncomingDirection,
-	FPathTracingPayload Payload
-)
-{
-	// pick back lobe proportionally to its contribution relative to front side
-	float3 Diff = Payload.DiffuseColor;
-	float3 Spec = Payload.SpecularColor;
-	float3 Back = Payload.GetSubsurfaceColor();
-	return LobeSelectionProb(Back, Diff + Spec);
-}
+#include "PathTracingEnergyConservation.ush"
 
 FMaterialSample TwoSidedFoliage_SampleMaterial(
 	float3 RayDirection,
@@ -26,17 +15,23 @@ FMaterialSample TwoSidedFoliage_SampleMaterial(
 	float4 RandSample
 )
 {
-	float TransmissionSelectionPdf = CalcTransmissionSelectionPdf(RayDirection, Payload);
+	const float NoV = saturate(-dot(RayDirection, Payload.WorldNormal));
+	const FBxDFEnergyTerms Spec = ComputeGGXSpecEnergyTerms(Payload.Roughness, NoV, Payload.SpecularColor);
+
+	// pick back lobe proportionally to its contribution relative to front side
+	const float3 Diff = (1 - Spec.E) * Payload.DiffuseColor;
+	const float3 Back = (1 - Spec.E) * Payload.GetSubsurfaceColor();
+	float TransmissionSelectionPdf = LobeSelectionProb(Back, Diff + Spec.E);
+
 	if (RandSample.x < TransmissionSelectionPdf)
 	{
 		RandSample.x /= TransmissionSelectionPdf;
 
 		float3 N = Payload.WorldNormal;
-		float3 SubsurfaceColor = Payload.GetSubsurfaceColor();
 
 		float4 SampledValue = CosineSampleHemisphere(RandSample.yz);
 
-		return CreateMaterialSample(TangentToWorld(SampledValue.xyz, -N), SubsurfaceColor / TransmissionSelectionPdf, TransmissionSelectionPdf * SampledValue.w, -1.0, 1.0);
+		return CreateMaterialSample(TangentToWorld(SampledValue.xyz, -N), Back / TransmissionSelectionPdf, TransmissionSelectionPdf * SampledValue.w, -1.0, 1.0);
 	}
 	else
 	{
@@ -56,7 +51,13 @@ FMaterialEval TwoSidedFoliage_EvalMaterial(
 	FPathTracingPayload Payload
 )
 {
-	float TransmissionSelectionPdf = CalcTransmissionSelectionPdf(IncomingDirection, Payload);
+	const float NoV = saturate(-dot(IncomingDirection, Payload.WorldNormal));
+	const FBxDFEnergyTerms Spec = ComputeGGXSpecEnergyTerms(Payload.Roughness, NoV, Payload.SpecularColor);
+
+	// pick back lobe proportionally to its contribution relative to front side
+	const float3 Diff = (1 - Spec.E) * Payload.DiffuseColor;
+	const float3 Back = (1 - Spec.E) * Payload.GetSubsurfaceColor();
+	float TransmissionSelectionPdf = LobeSelectionProb(Back, Diff + Spec.E);
 
 	float3 N = Payload.WorldNormal;
 	float NoL = dot(N, OutgoingDirection);
@@ -66,10 +67,9 @@ FMaterialEval TwoSidedFoliage_EvalMaterial(
 
 		if (TransmissionSelectionPdf > 0)
 		{
-			float3 SubsurfaceColor = Payload.GetSubsurfaceColor();
 			float LambertPdf = -NoL / PI;
 
-			return CreateMaterialEval(SubsurfaceColor / TransmissionSelectionPdf, TransmissionSelectionPdf * LambertPdf);
+			return CreateMaterialEval(Back / TransmissionSelectionPdf, TransmissionSelectionPdf * LambertPdf);
 		}
 	}
 	else if (TransmissionSelectionPdf < 1)
diff --git a/Engine/Shaders/Private/PathTracing/PathTracingBuildEnergyTable.usf b/Engine/Shaders/Private/PathTracing/PathTracingBuildEnergyTable.usf
new file mode 100644
--- /dev/null
+++ b/Engine/Shaders/Private/PathTracing/PathTracingBuildEnergyTable.usf
@@ -0,0 +1,106 @@
+// Copyright Epic Games, Inc. All Rights Reserved.
+
+#include "../Common.ush"
+#include "/Engine/Shared/PathTracingDefinitions.h"
+#include "Material/PathTracingGlossy.ush"
+#include "Material/PathTracingFresnel.ush"
+#include "Utilities/PathTracingRandomSequence.ush"
+
+// This shader computes the directional albedo (basically a furnace test) of a given BxDF to determine how much energy is "missing"
+// from a given viewing direction and roughness. These tables can later be used to compensate for the missing energy during shading.
+// "Energy" here refers to the directional albedo. Energy loss is simply 1-Energy since we assume the raw BxDF should reflect all
+// incoming light.
+
+// x: directional albedo of pure GGX lobe
+// y: directional albedo of GGX+Schlick fresnel term
+RWTexture2D<float2> OutputTexture2D;
+// x: directional albedo from outside->in, y: directional albedo from inside->out
+RWTexture3D<float2> OutputTexture3D;
+
+[numthreads(THREADGROUPSIZE_X, THREADGROUPSIZE_Y, 1)]
+void PathTracingBuildEnergyTableCS(uint3 DispatchThreadId : SV_DispatchThreadID)
+{
+	float3 Index = (float3(DispatchThreadId) + 0.5) / float(PATHTRACER_ENERGY_TABLE_RESOLUTION);
+
+	// Map from thread ID to the exact roughness/viewing angle we want to compute
+	float NoV = max(Index.x, 1e-8);
+	float Roughness = Index.y;
+#if BUILD_ENERGY_TABLE == 1
+	float EtaIn = 1.0 + 2.0 * Index.z; // only handle Ior in [1,3] range
+	float EtaOut = 1.0 / EtaIn;
+#endif
+
+	float2 Alpha; GetGGXBasis(Roughness, float3(0, 0, 1), Alpha);
+	float3 V = float3(sqrt(1 - NoV * NoV), 0, NoV);
+
+	RandomSequence RandSequence;
+	float2 Result = 0;
+	const int NumSamples = 1 << 14; // use plenty of samples -- this is computed on startup once and at fairly low resolution
+	for (int Sample = 0; Sample < NumSamples; Sample++)
+	{
+		// draw a random sample
+		RandomSequence_Initialize(RandSequence, DispatchThreadId.xy, Sample, DispatchThreadId.z, NumSamples);
+		float3 RandSample = RandomSequence_GenerateSample3D(RandSequence);
+
+		float2 Weight = 0;
+
+#if BUILD_ENERGY_TABLE == 0
+		const float3 H = ImportanceSampleVisibleGGX_aniso(UniformSampleDisk(RandSample.xy), Alpha, V).xyz;
+		const float3 L = reflect(-V, H);
+		if (L.z > 0)
+		{
+			const float2 GGXResult = GGXEvalReflection(L, V, H, Alpha);
+			Weight = float2(GGXResult.x, GGXResult.x * Pow5(1 - dot(V, H)));
+		}
+#elif BUILD_ENERGY_TABLE == 1
+		const float3 H = ImportanceSampleVisibleGGX_aniso(UniformSampleDisk(RandSample.xy), Alpha, V).xyz;
+		{
+			// compute weight for outside -> in case
+			float3 L = 0;
+			float F = 0;
+			if (SampleRefraction(-V, H, EtaIn, RandSample.z, L, F))
+			{
+				Weight.x = GGXEvalRefraction(L, V, H, Alpha, EtaIn).x;
+			}
+			else
+			{
+				Weight.x = GGXEvalReflection(L, V, H, Alpha).x;
+			}
+		}
+		{
+			// compute weight for inside -> out case
+			float3 L = 0;
+			float F = 0;
+			if (SampleRefraction(-V, H, EtaOut, RandSample.z, L, F))
+			{
+				Weight.y = GGXEvalRefraction(L, V, H, Alpha, EtaOut).x;
+			}
+			else
+			{
+				Weight.y = GGXEvalReflection(L, V, H, Alpha).x;
+			}
+		}
+#elif BUILD_ENERGY_TABLE == 2
+		const float3 L = CosineSampleHemisphere(RandSample.xy);
+		const float NoV = V.z;
+		const float NoL = L.z;
+		const float3 H = normalize(V + L);
+		const float NoH = H.z;
+		const float DCloth = D_InvGGX(Alpha.x * Alpha.y, NoH);
+		const float ClothWeight = NoL > 0 && NoV > 0 ? (PI * DCloth * Vis_Cloth(NoV, NoL)) : 0.0;
+
+		Weight = float2(ClothWeight, ClothWeight * Pow5(1 - dot(V, H)));
+#endif
+
+		// average in
+		Result = lerp(Result, Weight, 1.0 / (Sample + 1));
+	}
+
+#if BUILD_ENERGY_TABLE == 0 || BUILD_ENERGY_TABLE == 2
+	OutputTexture2D[DispatchThreadId.xy] = Result;
+#endif
+
+#if BUILD_ENERGY_TABLE == 1
+	OutputTexture3D[DispatchThreadId] = Result;
+#endif
+}
diff --git a/Engine/Shaders/Private/PathTracing/PathTracingCore.ush b/Engine/Shaders/Private/PathTracing/PathTracingCore.ush
--- a/Engine/Shaders/Private/PathTracing/PathTracingCore.ush
+++ b/Engine/Shaders/Private/PathTracing/PathTracingCore.ush
@@ -5,6 +5,7 @@
 #define RANDSEQ_UNROLL_SOBOL PATH_TRACER_USE_COMPACTION // unrolling sobol sampler is best when we can guarantee high occupancy
 
 #define USE_PATH_TRACING_LIGHT_GRID	1
+#define USE_PATH_TRACING_ENERGY_CONSERVATION 1
 
 #include "../Common.ush"
 #include "../PostProcessCommon.ush"
@@ -280,7 +281,7 @@ void ApplyRayBias(inout RayDesc Ray, float HitT, float3 Normal, bool ShadowTermi
 #endif
 }
 
-bool ProcessSubsurfaceRandomWalk(inout FPathTracingPayload Payload, inout float3 PathThroughput, inout RandomSequence RandSequence, uint2 PixelCoord, bool SimplifySSS)
+bool ProcessSubsurfaceRandomWalk(inout FPathTracingPayload Payload, inout float3 PathThroughput, inout RandomSequence RandSequence, float3 RayDirection, uint2 PixelCoord, bool SimplifySSS)
 {
 	if (!Payload.IsSubsurfaceMaterial())
 	{
@@ -305,12 +306,32 @@ bool ProcessSubsurfaceRandomWalk(inout FPathTracingPayload Payload, inout float3
 	SSSRadius = max(SSSRadius, 0.1); // 1mm minimum to make sure random walk can move forward and to keep SigmaT finite
 
 	// decide randomly to evaluate the SSS part of the material, or keep shading the current part
-	float SSSProb = LobeSelectionProb(SSSColor, Payload.DiffuseColor + Payload.SpecularColor);
+	// This also computes the energy conservation to take out the portion of lighting reflected away by the specular layer
+	const float NoV = saturate(dot(Payload.WorldNormal, -RayDirection));
+	float3 SpecE = 0;
+	if (Payload.ShadingModelID == SHADINGMODELID_SUBSURFACE_PROFILE)
+	{
+		// dual spec
+		const float3 DualRoughnessData = Payload.GetDualRoughnessSpecular();
+		const FBxDFEnergyTerms Spec0 = ComputeGGXSpecEnergyTerms(DualRoughnessData.x, NoV, Payload.SpecularColor);
+		const FBxDFEnergyTerms Spec1 = ComputeGGXSpecEnergyTerms(DualRoughnessData.y, NoV, Payload.SpecularColor);
+		SpecE = lerp(Spec0.E, Spec1.E, DualRoughnessData.z);
+	}
+	else
+	{
+		// single spec
+		SpecE = ComputeGGXSpecEnergyTerms(Payload.Roughness, NoV, Payload.SpecularColor).E;
+	}
+
+	const float3 SSSLobeAlbedo = (1 - SpecE) * SSSColor;
+	const float3 DiffLobeAlbedo = (1 - SpecE) * Payload.DiffuseColor;
+	const float3 SpecLobeAlbedo = SpecE;
+	float SSSProb = LobeSelectionProb(SSSLobeAlbedo, DiffLobeAlbedo + SpecLobeAlbedo);
 	float3 RandSample = RandomSequence_GenerateSample3D(RandSequence);
 	if (RandSample.x < SSSProb)
 	{
-		// do random walk
-		PathThroughput *= 1 / SSSProb;
+		// do random walk (and adjust throughput to remove energy reflected away by the specular lobe)
+		PathThroughput *= (1 - SpecE) / SSSProb;
 		RandSample.x /= SSSProb;
 	}
 	else
@@ -675,7 +696,7 @@ bool PathTracingKernel(inout FPathState PathState, int Bounce)
 	}
 
 	// Processing the random walk will (stochastically) move the shading point elsewhere on the surface of the object
-	if (!ProcessSubsurfaceRandomWalk(Payload, PathState.PathThroughput, PathState.RandSequence, PathState.PixelIndex, bSimplifySSS))
+	if (!ProcessSubsurfaceRandomWalk(Payload, PathState.PathThroughput, PathState.RandSequence, PathState.Ray.Direction, PathState.PixelIndex, bSimplifySSS))
 	{
 		// random walk did not terminate at a valid point
 		return false;
diff --git a/Engine/Shaders/Shared/PathTracingDefinitions.h b/Engine/Shaders/Shared/PathTracingDefinitions.h
--- a/Engine/Shaders/Shared/PathTracingDefinitions.h
+++ b/Engine/Shaders/Shared/PathTracingDefinitions.h
@@ -30,3 +30,6 @@
 // Constants for the path tracer light grid
 #define PATHTRACER_LIGHT_GRID_SINGULAR_MASK					0x80000000u
 #define PATHTRACER_LIGHT_GRID_LIGHT_COUNT_MASK				0x7FFFFFFFu
+
+// Constants for the energy conservation texture sizes
+#define PATHTRACER_ENERGY_TABLE_RESOLUTION				32
diff --git a/Engine/Source/Runtime/Renderer/Private/PathTracing.cpp b/Engine/Source/Runtime/Renderer/Private/PathTracing.cpp
--- a/Engine/Source/Runtime/Renderer/Private/PathTracing.cpp
+++ b/Engine/Source/Runtime/Renderer/Private/PathTracing.cpp
@@ -86,7 +86,8 @@ TAutoConsoleVariable<int32> CVarPathTracingVisibleLights(
 	0,
 	TEXT("Should light sources be visible to camera rays? (default = 0 (off))\n")
 	TEXT("0: Hide lights from camera rays (default)\n")
-	TEXT("1: Make lights visible to camera\n"),
+	TEXT("1: Make all lights visible to camera\n")
+	TEXT("2: Make skydome only visible to camera\n"),
 	ECVF_RenderThreadSafe
 );
 
@@ -266,6 +267,26 @@ struct FPathTracingConfig
 	}
 };
 
+struct FPathTracingState {
+	FPathTracingConfig LastConfig;
+	// Textures holding onto the accumulated frame data
+	TRefCountPtr<IPooledRenderTarget> RadianceRT;
+	TRefCountPtr<IPooledRenderTarget> AlbedoRT;
+	TRefCountPtr<IPooledRenderTarget> NormalRT;
+	TRefCountPtr<IPooledRenderTarget> RadianceDenoisedRT;
+
+	TRefCountPtr<IPooledRenderTarget> GGXSpecEnergyTable;
+	TRefCountPtr<IPooledRenderTarget> GGXGlassEnergyTable;
+	TRefCountPtr<IPooledRenderTarget> ClothEnergyTable;
+
+	// Current sample index to be rendered by the path tracer - this gets incremented each time the path tracer accumulates a frame of samples
+	uint32 SampleIndex = 0;
+
+	// Path tracer frame index, not reset on invalidation unlike SampleIndex to avoid
+	// the "screen door" effect and reduce temporal aliasing
+	uint32_t FrameIndex = 0;
+};
+
 // This function prepares the portion of shader arguments that may involve invalidating the path traced state
 static void PrepareShaderArgs(const FViewInfo& View, FPathTracingData& PathTracingData)
 {
@@ -429,6 +450,44 @@ class FPathTracingBuildLightGridCS : public FGlobalShader
 IMPLEMENT_SHADER_TYPE(, FPathTracingBuildLightGridCS, TEXT("/Engine/Private/PathTracing/PathTracingBuildLightGrid.usf"), TEXT("PathTracingBuildLightGridCS"), SF_Compute);
 
 
+class FPathTracingBuildEnergyTableCS : public FGlobalShader
+{
+	DECLARE_GLOBAL_SHADER(FPathTracingBuildEnergyTableCS)
+	SHADER_USE_PARAMETER_STRUCT(FPathTracingBuildEnergyTableCS, FGlobalShader)
+
+	enum class EEnergyTableType {
+		GGXSpecular = 0,
+		GGXGlass    = 1,
+		Cloth       = 2,
+		MAX
+	};
+
+	class FEnergyTableDim : SHADER_PERMUTATION_ENUM_CLASS("BUILD_ENERGY_TABLE", EEnergyTableType);
+
+	using FPermutationDomain = TShaderPermutationDomain<FEnergyTableDim>;
+
+		
+	static bool ShouldCompilePermutation(const FGlobalShaderPermutationParameters& Parameters)
+	{
+		return ShouldCompileRayTracingShadersForProject(Parameters.Platform);
+	}
+
+	static void ModifyCompilationEnvironment(const FGlobalShaderPermutationParameters& Parameters, FShaderCompilerEnvironment& OutEnvironment)
+	{
+		//OutEnvironment.CompilerFlags.Add(CFLAG_WarningsAsErrors);
+		OutEnvironment.CompilerFlags.Add(CFLAG_AllowTypedUAVLoads);
+		OutEnvironment.SetDefine(TEXT("THREADGROUPSIZE_X"), FComputeShaderUtils::kGolden2DGroupSize);
+		OutEnvironment.SetDefine(TEXT("THREADGROUPSIZE_Y"), FComputeShaderUtils::kGolden2DGroupSize);
+	}
+
+	BEGIN_SHADER_PARAMETER_STRUCT(FParameters, )
+		SHADER_PARAMETER_RDG_TEXTURE_UAV(RWTexture2D, OutputTexture2D)
+		SHADER_PARAMETER_RDG_TEXTURE_UAV(RWTexture3D, OutputTexture3D)
+	END_SHADER_PARAMETER_STRUCT()
+};
+
+IMPLEMENT_SHADER_TYPE(, FPathTracingBuildEnergyTableCS, TEXT("/Engine/Private/PathTracing/PathTracingBuildEnergyTable.usf"), TEXT("PathTracingBuildEnergyTableCS") , SF_Compute);
+
 class FPathTracingRG : public FGlobalShader
 {
 	DECLARE_GLOBAL_SHADER(FPathTracingRG)
@@ -490,6 +549,12 @@ class FPathTracingRG : public FGlobalShader
 		// Used by multi-GPU rendering
 		SHADER_PARAMETER(FIntVector, TileOffset)
 
+		// GGX Energy conservation
+		SHADER_PARAMETER_RDG_TEXTURE_SRV(Texture2D<float2>, GGXSpecEnergyTexture)
+		SHADER_PARAMETER_RDG_TEXTURE_SRV(Texture3D<float2>, GGXGlassEnergyTexture)
+		SHADER_PARAMETER_RDG_TEXTURE_SRV(Texture2D<float2>, ClothEnergyTexture)
+		SHADER_PARAMETER_SAMPLER(SamplerState, EnergySampler)
+
 		// extra parameters required for path compacting kernel
 		SHADER_PARAMETER(int, Bounce)
 		SHADER_PARAMETER_RDG_BUFFER_UAV(RWBuffer<uint>, PathStateData)
@@ -957,10 +1022,11 @@ void SetLightParameters(FRDGBuilder& GraphBuilder, FPathTracingRG::FParameters*
 		DestLight.IESTextureSlice = -1;
 		DestLight.BoundMin = FVector(-Inf, -Inf, -Inf);
 		DestLight.BoundMax = FVector( Inf,  Inf,  Inf);
-		if (Scene->SkyLight->bRealTimeCaptureEnabled)
+		if (Scene->SkyLight->bRealTimeCaptureEnabled || CVarPathTracingVisibleLights.GetValueOnRenderThread() == 2)
 		{
 			// When using the realtime capture system, always make the skylight visible
 			// because this is our only way of "seeing" the atmo/clouds at the moment
+			// Also allow seeing just the sky via a cvar for debugging purposes
 			PassParameters->SceneVisibleLightCount = 1;
 		}
 	}
@@ -1190,8 +1256,9 @@ void SetLightParameters(FRDGBuilder& GraphBuilder, FPathTracingRG::FParameters*
 		PassParameters->SceneLights = GraphBuilder.CreateSRV(FRDGBufferSRVDesc(CreateStructuredBuffer(GraphBuilder, TEXT("PathTracer.LightsBuffer"), sizeof(FPathTracingLight), NumCopyLights, Lights, DataSize, ERDGInitialDataFlags::NoCopy)));
 	}
 
-	if (CVarPathTracingVisibleLights.GetValueOnRenderThread() != 0)
+	if (CVarPathTracingVisibleLights.GetValueOnRenderThread() == 1)
 	{
+		// make all lights in the scene visible
 		PassParameters->SceneVisibleLightCount = PassParameters->SceneLightCount;
 	}
 
@@ -1253,24 +1320,25 @@ void FDeferredShadingSceneRenderer::PreparePathTracing(const FSceneViewFamily& V
 
 void FSceneViewState::PathTracingInvalidate()
 {
-	PathTracingRadianceRT.SafeRelease();
-	PathTracingAlbedoRT.SafeRelease();
-	PathTracingNormalRT.SafeRelease();
-	PathTracingRadianceDenoisedRT.SafeRelease();
-	PathTracingSampleIndex = 0;
+	FPathTracingState* State = PathTracingState.Get();
+	if (State)
+	{
+		State->RadianceRT.SafeRelease();
+		State->AlbedoRT.SafeRelease();
+		State->NormalRT.SafeRelease();
+		State->RadianceDenoisedRT.SafeRelease();
+		State->SampleIndex = 0;
+	}
 }
 
 uint32 FSceneViewState::GetPathTracingSampleIndex() const {
-	return PathTracingSampleIndex;
+	const FPathTracingState* State = PathTracingState.Get();
+	return State ? State->SampleIndex : 0;
 }
 
 uint32 FSceneViewState::GetPathTracingSampleCount() const {
-	FPathTracingConfig* LastConfig = PathTracingLastConfig.Get();
-	if (LastConfig)
-	{
-		return LastConfig->PathTracingData.MaxSamples;
-	}
-	return 0;
+	const FPathTracingState* State = PathTracingState.Get();
+	return State ? State->LastConfig.PathTracingData.MaxSamples : 0;
 }
 
 
@@ -1336,13 +1404,15 @@ void FDeferredShadingSceneRenderer::RenderPathTracing(
 	Config.PathTracingData.MaxSamples = MaxSPP;
 
 	bool FirstTime = false;
-	if (!View.ViewState->PathTracingLastConfig.IsValid())
+	if (!View.ViewState->PathTracingState.IsValid())
 	{
-		View.ViewState->PathTracingLastConfig = MakePimpl<FPathTracingConfig>();
+		View.ViewState->PathTracingState = MakePimpl<FPathTracingState>();
 		FirstTime = true; // we just initialized the option state for this view -- don't bother comparing in this case
 	}
+	check(View.ViewState->PathTracingState.IsValid());
+	FPathTracingState* PathTracingState = View.ViewState->PathTracingState.Get();
 
-	if (FirstTime || Config.UseMISCompensation != View.ViewState->PathTracingLastConfig->UseMISCompensation)
+	if (FirstTime || Config.UseMISCompensation != PathTracingState->LastConfig.UseMISCompensation)
 	{
 		// if the mode changes we need to rebuild the importance table
 		Scene->PathTracingSkylightTexture.SafeRelease();
@@ -1351,10 +1421,10 @@ void FDeferredShadingSceneRenderer::RenderPathTracing(
 
 	// If the scene has changed in some way (camera move, object movement, etc ...)
 	// we must invalidate the ViewState to start over from scratch
-	if (FirstTime || Config.IsDifferent(*View.ViewState->PathTracingLastConfig))
+	if (FirstTime || Config.IsDifferent(PathTracingState->LastConfig))
 	{
 		// remember the options we used for next time
-		*View.ViewState->PathTracingLastConfig = Config;
+		PathTracingState->LastConfig = Config;
 		View.ViewState->PathTracingInvalidate();
 	}
 
@@ -1362,26 +1432,102 @@ void FDeferredShadingSceneRenderer::RenderPathTracing(
 	if (Config.LockedSamplingPattern)
 	{
 		// Count samples from 0 for deterministic results
-		Config.PathTracingData.TemporalSeed = View.ViewState->PathTracingSampleIndex;
+		Config.PathTracingData.TemporalSeed = PathTracingState->SampleIndex;
 	}
 	else
 	{
 		// Count samples from an ever-increasing counter to avoid screen-door effect
-		Config.PathTracingData.TemporalSeed = View.ViewState->PathTracingFrameIndex;
+		Config.PathTracingData.TemporalSeed = PathTracingState->FrameIndex;
 	}
-	Config.PathTracingData.Iteration = View.ViewState->PathTracingSampleIndex;
+	Config.PathTracingData.Iteration = PathTracingState->SampleIndex;
 	Config.PathTracingData.BlendFactor = 1.0f / (Config.PathTracingData.Iteration + 1);
 
+	// Prepare energy conservation tables
+	FRDGTexture* GGXSpecEnergyTexture = nullptr;
+	FRDGTexture* GGXGlassEnergyTexture = nullptr;
+	FRDGTexture* ClothEnergyTexture = nullptr;
+	if (PathTracingState->GGXSpecEnergyTable)
+	{
+		GGXSpecEnergyTexture = GraphBuilder.RegisterExternalTexture(PathTracingState->GGXSpecEnergyTable, TEXT("PathTracer.GGXSpecEnergy"));
+		GGXGlassEnergyTexture = GraphBuilder.RegisterExternalTexture(PathTracingState->GGXGlassEnergyTable, TEXT("PathTracer.GGXGlassEnergy"));
+		ClothEnergyTexture = GraphBuilder.RegisterExternalTexture(PathTracingState->ClothEnergyTable, TEXT("PathTracer.ClothEnergy"));
+	}
+	else
+	{
+		// first time through, bake the data we need
+		const int Size = PATHTRACER_ENERGY_TABLE_RESOLUTION;
+		GGXSpecEnergyTexture = GraphBuilder.CreateTexture(FRDGTextureDesc::Create2D(
+			FIntPoint(Size, Size),
+			PF_G32R32F,
+			FClearValueBinding::None,
+			TexCreate_ShaderResource | TexCreate_UAV), TEXT("PathTracer.GGXSpecEnergy"), ERDGTextureFlags::MultiFrame);
+		GGXGlassEnergyTexture = GraphBuilder.CreateTexture(FRDGTextureDesc::Create3D(
+			FIntVector(Size, Size, Size),
+			PF_G32R32F,
+			FClearValueBinding::None,
+			TexCreate_ShaderResource | TexCreate_UAV), TEXT("PathTracer.GGXGlassEnergy"), ERDGTextureFlags::MultiFrame);
+		ClothEnergyTexture = GraphBuilder.CreateTexture(FRDGTextureDesc::Create2D(
+			FIntPoint(Size, Size),
+			PF_G32R32F,
+			FClearValueBinding::None,
+			TexCreate_ShaderResource | TexCreate_UAV), TEXT("PathTracer.GGXSpecEnergy"), ERDGTextureFlags::MultiFrame);
+
+		{
+			FPathTracingBuildEnergyTableCS::FPermutationDomain PermutationVector;
+			PermutationVector.Set<FPathTracingBuildEnergyTableCS::FEnergyTableDim>(FPathTracingBuildEnergyTableCS::EEnergyTableType::GGXSpecular);
+			TShaderMapRef<FPathTracingBuildEnergyTableCS> ComputeShader(View.ShaderMap, PermutationVector);
+			FPathTracingBuildEnergyTableCS::FParameters* PassParameters = GraphBuilder.AllocParameters<FPathTracingBuildEnergyTableCS::FParameters>();
+			PassParameters->OutputTexture2D = GraphBuilder.CreateUAV(FRDGTextureUAVDesc(GGXSpecEnergyTexture, 0));
+			FComputeShaderUtils::AddPass(
+				GraphBuilder,
+				RDG_EVENT_NAME("BuildGGXSpecTable"),
+				ComputeShader,
+				PassParameters,
+				FComputeShaderUtils::GetGroupCount(FIntPoint(Size, Size), FComputeShaderUtils::kGolden2DGroupSize));
+		}
+		{
+			FPathTracingBuildEnergyTableCS::FPermutationDomain PermutationVector;
+			PermutationVector.Set<FPathTracingBuildEnergyTableCS::FEnergyTableDim>(FPathTracingBuildEnergyTableCS::EEnergyTableType::GGXGlass);
+			TShaderMapRef<FPathTracingBuildEnergyTableCS> ComputeShader(View.ShaderMap, PermutationVector);
+			FPathTracingBuildEnergyTableCS::FParameters* PassParameters = GraphBuilder.AllocParameters<FPathTracingBuildEnergyTableCS::FParameters>();
+			PassParameters->OutputTexture3D = GraphBuilder.CreateUAV(FRDGTextureUAVDesc(GGXGlassEnergyTexture, 0));
+			FComputeShaderUtils::AddPass(
+				GraphBuilder,
+				RDG_EVENT_NAME("BuildGGXGlassTable"),
+				ComputeShader,
+				PassParameters,
+				FComputeShaderUtils::GetGroupCount(FIntVector(Size, Size, Size), FIntVector(FComputeShaderUtils::kGolden2DGroupSize, FComputeShaderUtils::kGolden2DGroupSize, 1)));
+		}
+		{
+			FPathTracingBuildEnergyTableCS::FPermutationDomain PermutationVector;
+			PermutationVector.Set<FPathTracingBuildEnergyTableCS::FEnergyTableDim>(FPathTracingBuildEnergyTableCS::EEnergyTableType::Cloth);
+			TShaderMapRef<FPathTracingBuildEnergyTableCS> ComputeShader(View.ShaderMap, PermutationVector);
+			FPathTracingBuildEnergyTableCS::FParameters* PassParameters = GraphBuilder.AllocParameters<FPathTracingBuildEnergyTableCS::FParameters>();
+			PassParameters->OutputTexture2D = GraphBuilder.CreateUAV(FRDGTextureUAVDesc(ClothEnergyTexture, 0));
+			FComputeShaderUtils::AddPass(
+				GraphBuilder,
+				RDG_EVENT_NAME("BuildClothTable"),
+				ComputeShader,
+				PassParameters,
+				FComputeShaderUtils::GetGroupCount(FIntPoint(Size, Size), FComputeShaderUtils::kGolden2DGroupSize));
+		}
+
+		GraphBuilder.QueueTextureExtraction(GGXSpecEnergyTexture , &PathTracingState->GGXSpecEnergyTable );
+		GraphBuilder.QueueTextureExtraction(GGXGlassEnergyTexture, &PathTracingState->GGXGlassEnergyTable);
+		GraphBuilder.QueueTextureExtraction(ClothEnergyTexture   , &PathTracingState->ClothEnergyTable);
+	}
+	
+
 	// Prepare radiance buffer (will be shared with display pass)
 	FRDGTexture* RadianceTexture = nullptr;
 	FRDGTexture* AlbedoTexture = nullptr;
 	FRDGTexture* NormalTexture = nullptr;
-	if (View.ViewState->PathTracingRadianceRT)
+	if (PathTracingState->RadianceRT)
 	{
 		// we already have a valid radiance texture, re-use it
-		RadianceTexture = GraphBuilder.RegisterExternalTexture(View.ViewState->PathTracingRadianceRT, TEXT("PathTracer.Radiance"));
-		AlbedoTexture   = GraphBuilder.RegisterExternalTexture(View.ViewState->PathTracingAlbedoRT, TEXT("PathTracer.Albedo"));
-		NormalTexture   = GraphBuilder.RegisterExternalTexture(View.ViewState->PathTracingNormalRT, TEXT("PathTracer.Normal"));
+		RadianceTexture = GraphBuilder.RegisterExternalTexture(PathTracingState->RadianceRT, TEXT("PathTracer.Radiance"));
+		AlbedoTexture   = GraphBuilder.RegisterExternalTexture(PathTracingState->AlbedoRT, TEXT("PathTracer.Albedo"));
+		NormalTexture   = GraphBuilder.RegisterExternalTexture(PathTracingState->NormalRT, TEXT("PathTracer.Normal"));
 	}
 	else
 	{
@@ -1464,6 +1610,12 @@ void FDeferredShadingSceneRenderer::RenderPathTracing(
 			PassParameters->TileOffset.X = 0;
 			PassParameters->TileOffset.Y = 0;
 
+			// energy compensation tabulated data
+			PassParameters->GGXSpecEnergyTexture = GraphBuilder.CreateSRV(FRDGTextureSRVDesc::Create(GGXSpecEnergyTexture));
+			PassParameters->GGXGlassEnergyTexture = GraphBuilder.CreateSRV(FRDGTextureSRVDesc::Create(GGXGlassEnergyTexture));
+			PassParameters->ClothEnergyTexture = GraphBuilder.CreateSRV(FRDGTextureSRVDesc::Create(ClothEnergyTexture));
+			PassParameters->EnergySampler = TStaticSamplerState<SF_Bilinear, AM_Clamp, AM_Clamp, AM_Clamp>::GetRHI();
+
 			if (bUseCompaction)
 			{
 				PassParameters->Bounce = Bounce;
@@ -1477,8 +1629,8 @@ void FDeferredShadingSceneRenderer::RenderPathTracing(
 			ClearUnusedGraphResources(RayGenShader, PassParameters);
 			GraphBuilder.AddPass(
 				bUseCompaction
-					? RDG_EVENT_NAME("Path Tracer Compute (%d x %d) Sample=%d/%d NumLights=%d Bounce=%d", View.ViewRect.Size().X, View.ViewRect.Size().Y, View.ViewState->PathTracingSampleIndex, MaxSPP, PassParameters->SceneLightCount, Bounce)
-					: RDG_EVENT_NAME("Path Tracer Compute (%d x %d) Sample=%d/%d NumLights=%d", View.ViewRect.Size().X, View.ViewRect.Size().Y, View.ViewState->PathTracingSampleIndex, MaxSPP, PassParameters->SceneLightCount),
+					? RDG_EVENT_NAME("Path Tracer Compute (%d x %d) Sample=%d/%d NumLights=%d Bounce=%d", View.ViewRect.Size().X, View.ViewRect.Size().Y, PathTracingState->SampleIndex, MaxSPP, PassParameters->SceneLightCount, Bounce)
+					: RDG_EVENT_NAME("Path Tracer Compute (%d x %d) Sample=%d/%d NumLights=%d", View.ViewRect.Size().X, View.ViewRect.Size().Y, PathTracingState->SampleIndex, MaxSPP, PassParameters->SceneLightCount),
 				PassParameters,
 				ERDGPassFlags::Compute,
 				[PassParameters, RayGenShader, DispatchSizeX, DispatchSizeY, &View](FRHIRayTracingCommandList& RHICmdList)
@@ -1497,13 +1649,13 @@ void FDeferredShadingSceneRenderer::RenderPathTracing(
 			});
 		}
 		// After we are done, make sure we remember our texture for next time so that we can accumulate samples across frames
-		GraphBuilder.QueueTextureExtraction(RadianceTexture, &View.ViewState->PathTracingRadianceRT);
-		GraphBuilder.QueueTextureExtraction(AlbedoTexture  , &View.ViewState->PathTracingAlbedoRT  );
-		GraphBuilder.QueueTextureExtraction(NormalTexture  , &View.ViewState->PathTracingNormalRT  );
+		GraphBuilder.QueueTextureExtraction(RadianceTexture, &PathTracingState->RadianceRT);
+		GraphBuilder.QueueTextureExtraction(AlbedoTexture  , &PathTracingState->AlbedoRT  );
+		GraphBuilder.QueueTextureExtraction(NormalTexture  , &PathTracingState->NormalRT  );
 
 		// Bump counters for next frame
-		++View.ViewState->PathTracingSampleIndex;
-		++View.ViewState->PathTracingFrameIndex;
+		++PathTracingState->SampleIndex;
+		++PathTracingState->FrameIndex;
 	}
 
 	FRDGTexture* DenoisedRadianceTexture = nullptr;
@@ -1521,13 +1673,13 @@ void FDeferredShadingSceneRenderer::RenderPathTracing(
 		if (!bNeedsMoreRays)
 		{
 			// we aren't currently rendering, run the denoiser if we just turned it on
-			NeedsDenoise |= DenoiserMode != View.ViewState->PathTracingLastConfig->DenoiserMode;
+			NeedsDenoise |= DenoiserMode != PathTracingState->LastConfig.DenoiserMode;
 		}
 
-		if (View.ViewState->PathTracingRadianceDenoisedRT)
+		if (PathTracingState->RadianceDenoisedRT)
 		{
 			// we already have a texture for this
-			DenoisedRadianceTexture = GraphBuilder.RegisterExternalTexture(View.ViewState->PathTracingRadianceDenoisedRT, TEXT("PathTracer.DenoisedRadiance"));
+			DenoisedRadianceTexture = GraphBuilder.RegisterExternalTexture(PathTracingState->RadianceDenoisedRT, TEXT("PathTracer.DenoisedRadiance"));
 		}
 
 		if (NeedsDenoise)
@@ -1559,10 +1711,10 @@ void FDeferredShadingSceneRenderer::RenderPathTracing(
 				}
 			);
 
-			GraphBuilder.QueueTextureExtraction(DenoisedRadianceTexture, &View.ViewState->PathTracingRadianceDenoisedRT);
+			GraphBuilder.QueueTextureExtraction(DenoisedRadianceTexture, &PathTracingState->RadianceDenoisedRT);
 		}
 	}
-	View.ViewState->PathTracingLastConfig->DenoiserMode = DenoiserMode;
+	PathTracingState->LastConfig.DenoiserMode = DenoiserMode;
 
 	// now add a pixel shader pass to display our Radiance buffer
 
diff --git a/Engine/Source/Runtime/Renderer/Private/RendererScene.cpp b/Engine/Source/Runtime/Renderer/Private/RendererScene.cpp
--- a/Engine/Source/Runtime/Renderer/Private/RendererScene.cpp
+++ b/Engine/Source/Runtime/Renderer/Private/RendererScene.cpp
@@ -243,9 +243,6 @@ FSceneViewState::FSceneViewState()
 	bUpdateLastExposure = false;
 
 #if RHI_RAYTRACING
-	PathTracingSampleIndex = 0;
-	PathTracingFrameIndex = 0;
-
 	GatherPointsBuffer = nullptr;
 	GatherPointsResolution = FIntVector(0, 0, 0);
 #endif
diff --git a/Engine/Source/Runtime/Renderer/Private/ScenePrivate.h b/Engine/Source/Runtime/Renderer/Private/ScenePrivate.h
--- a/Engine/Source/Runtime/Renderer/Private/ScenePrivate.h
+++ b/Engine/Source/Runtime/Renderer/Private/ScenePrivate.h
@@ -87,7 +87,7 @@ class FLumenSceneData;
 class FVirtualShadowMapArrayCacheManager;
 class FComputeGraphScheduler;
 struct FHairStrandsInstance;
-struct FPathTracingConfig;
+struct FPathTracingState;
 
 /** Holds information about a single primitive's occlusion. */
 class FPrimitiveOcclusionHistory
@@ -1020,20 +1020,8 @@ public:
 	virtual uint32 GetPathTracingSampleIndex() const override;
 	virtual uint32 GetPathTracingSampleCount() const override;
 
-	// Reference path tracing cached results
-	TRefCountPtr<IPooledRenderTarget> PathTracingRadianceRT;
-	TRefCountPtr<IPooledRenderTarget> PathTracingAlbedoRT;
-	TRefCountPtr<IPooledRenderTarget> PathTracingNormalRT;
-	TRefCountPtr<IPooledRenderTarget> PathTracingRadianceDenoisedRT;
-	// Keeps track of the internal path tracer options relevant to detecting when to restart the path tracer accumulation
-	TPimplPtr<FPathTracingConfig> PathTracingLastConfig;
-
-	// Current sample index to be rendered by the path tracer - this gets incremented each time the path tracer accumulates a frame of samples
-	uint32 PathTracingSampleIndex;
-
-	// Path tracer frame index, not reset on invalidation unlike PathTracingSampleIndex to avoid
-	// the "screen door" effect and reduce temporal aliasing
-	uint32_t PathTracingFrameIndex;
+	// Keeps track of the internal path tracer state
+	TPimplPtr<FPathTracingState> PathTracingState;
 
 	// IES light profiles
 	FIESLightProfileResource IESLightProfileResources;
