PR #8372: Fix Lightmap Density Optimization View's Shader Parameter Size Assertion caused by LWC (Contributed by [REDACTED_PERSON])

diff --git a/Engine/Source/Runtime/Renderer/Private/LightMapDensityRendering.cpp b/Engine/Source/Runtime/Renderer/Private/LightMapDensityRendering.cpp
--- a/Engine/Source/Runtime/Renderer/Private/LightMapDensityRendering.cpp
+++ b/Engine/Source/Runtime/Renderer/Private/LightMapDensityRendering.cpp
@@ -129,7 +129,7 @@ void FLightmapDensityMeshProcessor::Process(
 
 	{		
 		// BuiltLightingAndSelectedFlags informs the shader is lighting is built or not for this primitive
-		ShaderElementData.BuiltLightingAndSelectedFlags = FVector(0.0f, 0.0f, 0.0f);
+		ShaderElementData.BuiltLightingAndSelectedFlags = FVector3f(0.0f, 0.0f, 0.0f);
 		// LightMapResolutionScale is the physical resolution of the lightmap texture
 		ShaderElementData.LightMapResolutionScale = FVector2D(1.0f, 1.0f);
 
diff --git a/Engine/Source/Runtime/Renderer/Private/LightMapDensityRendering.h b/Engine/Source/Runtime/Renderer/Private/LightMapDensityRendering.h
--- a/Engine/Source/Runtime/Renderer/Private/LightMapDensityRendering.h
+++ b/Engine/Source/Runtime/Renderer/Private/LightMapDensityRendering.h
@@ -39,7 +39,7 @@ public:
 
 	typename LightMapPolicyType::ElementDataType LightMapPolicyElementData;
 
-	FVector BuiltLightingAndSelectedFlags;
+	FVector3f BuiltLightingAndSelectedFlags;
 	FVector2D LightMapResolutionScale; 
 	bool bTextureMapped;
 };
