Always translate primitive ID for use by mesh draw commands.
Fixes regression from 17093848 which meant that we sometimes earlied out from TranslatePrimitiveId() for a deferred translate that no longer happens.

diff --git a/Engine/Source/Runtime/Renderer/Private/MeshDrawCommands.cpp b/Engine/Source/Runtime/Renderer/Private/MeshDrawCommands.cpp
--- a/Engine/Source/Runtime/Renderer/Private/MeshDrawCommands.cpp
+++ b/Engine/Source/Runtime/Renderer/Private/MeshDrawCommands.cpp
@@ -364,11 +364,12 @@ void UpdateMobileBasePassMeshSortKeys(
 
 FORCEINLINE int32 TranslatePrimitiveId(int32 DrawPrimitiveIdIn, int32 DynamicPrimitiveIdOffset, int32 DynamicPrimitiveIdMax)
 {
-	// empty range means we defer the translation to later
-	if ((DynamicPrimitiveIdMax - DynamicPrimitiveIdOffset) <= 0)
+	// INDEX_NONE means we defer the translation to later
+	if (DynamicPrimitiveIdOffset == INDEX_NONE)
 	{
 		return DrawPrimitiveIdIn;
 	}
+
 	int32 DrawPrimitiveId = DrawPrimitiveIdIn;
 
 	if ((DrawPrimitiveIdIn & GPrimIDDynamicFlag) != 0)
@@ -378,8 +379,10 @@ FORCEINLINE int32 TranslatePrimitiveId(int32 DrawPrimitiveIdIn, int32 DynamicPri
 		checkSlow(DrawPrimitiveId < DynamicPrimitiveIdMax);
 	}
 
-	// Append flag to mark this as a non-instance data index (which is then treated as a primitive ID in the SceneData.ush loading
-	return DrawPrimitiveId | (1U << 31U);
+	// Append flag to mark this as a non-instance data index.
+	// This value is treated as a primitive ID in the SceneData.ush loading.
+	const uint32 VF_TREAT_INSTANCE_ID_OFFSET_AS_PRIMITIVE_ID_FLAG = 1U << 31U;
+	return DrawPrimitiveId |= VF_TREAT_INSTANCE_ID_OFFSET_AS_PRIMITIVE_ID_FLAG;
 }
 
 /**
@@ -971,9 +974,15 @@ void SortAndMergeDynamicPassMeshDrawCommands(
 				NewPassVisibleMeshDrawCommands.Empty(NumDrawCommands);
 			}
 
-			const TRange<int32> DynamicPrimitiveIdRange = DynamicPrimitiveCollector ? DynamicPrimitiveCollector->GetPrimitiveIdRange() : TRange<int32>::Empty();
-			int32 DynamicPrimitiveIdOffset = DynamicPrimitiveIdRange.GetLowerBoundValue(); 
-			int32 DynamicPrimitiveIdMax =  DynamicPrimitiveIdRange.GetUpperBoundValue();
+			// INDEX_NONE used in TranslatePrimitiveId to defer id translation
+			int32 DynamicPrimitiveIdOffset = INDEX_NONE;
+			int32 DynamicPrimitiveIdMax = 0;
+			if (DynamicPrimitiveCollector)
+			{
+				const TRange<int32> DynamicPrimitiveIdRange = DynamicPrimitiveCollector->GetPrimitiveIdRange();
+				DynamicPrimitiveIdOffset = DynamicPrimitiveIdRange.GetLowerBoundValue();
+				DynamicPrimitiveIdMax = DynamicPrimitiveIdRange.GetUpperBoundValue();
+			}
 
 			const uint32 PrimitiveIdBufferStride = FInstanceCullingContext::GetInstanceIdBufferStride(FeatureLevel);
 			const int32 MaxNumPrimitives = InstanceFactor * NumDrawCommands;
@@ -1333,16 +1342,6 @@ void FParallelMeshDrawCommandPass::WaitForSetupTask()
 {
 	WaitForMeshPassSetupTask();
 }
-/**
- * Helper to upload and translate a primitive ID buffer.
- */
-static void CopyPrimitiveIdBuffer(const int32* RESTRICT SrcData, int32* RESTRICT Data, int32 NumIds, const TRange<int32>& DynamicPrimitiveIdRange)
-{
-	for (int32 Index = 0; Index < NumIds; ++Index)
-	{
-		Data[Index] = TranslatePrimitiveId(SrcData[Index], DynamicPrimitiveIdRange.GetLowerBoundValue(), DynamicPrimitiveIdRange.GetUpperBoundValue());
-	}
-}
 
 void FParallelMeshDrawCommandPass::DispatchDraw(FParallelCommandListSet* ParallelCommandListSet, FRHICommandList& RHICmdList, const FInstanceCullingDrawParams* InstanceCullingDrawParams) const
 {
