Fix a bug of depthstencil exclusive flag after full depth prepass.

#p4v-preflight-copy 17435562

diff --git a/Engine/Source/Runtime/Renderer/Private/MobileShadingRenderer.cpp b/Engine/Source/Runtime/Renderer/Private/MobileShadingRenderer.cpp
--- a/Engine/Source/Runtime/Renderer/Private/MobileShadingRenderer.cpp
+++ b/Engine/Source/Runtime/Renderer/Private/MobileShadingRenderer.cpp
@@ -990,7 +990,7 @@ void FMobileSceneRenderer::RenderForward(FRDGBuilder& GraphBuilder, FRDGTextureR
 	FRenderTargetBindingSlots BasePassRenderTargets;
 	BasePassRenderTargets[0] = FRenderTargetBinding(SceneColor, SceneColorResolve, ERenderTargetLoadAction::EClear);
 	BasePassRenderTargets[1] = FRenderTargetBinding(SceneTextures.DepthAux.Target, SceneTextures.DepthAux.Resolve, ERenderTargetLoadAction::EClear);
-	BasePassRenderTargets.DepthStencil = FDepthStencilBinding(SceneDepth, bIsFullDepthPrepassEnabled ? ERenderTargetLoadAction::ELoad: ERenderTargetLoadAction::EClear, FExclusiveDepthStencil::DepthWrite_StencilWrite);
+	BasePassRenderTargets.DepthStencil = bIsFullDepthPrepassEnabled ? FDepthStencilBinding(SceneDepth, ERenderTargetLoadAction::ELoad, FExclusiveDepthStencil::DepthRead_StencilWrite) : FDepthStencilBinding(SceneDepth, ERenderTargetLoadAction::EClear, FExclusiveDepthStencil::DepthWrite_StencilWrite);
 	BasePassRenderTargets.ShadingRateTexture = (!MainView.bIsSceneCapture && !MainView.bIsReflectionCapture && ShadingRateTarget.IsValid()) ? RegisterExternalTexture(GraphBuilder, ShadingRateTarget->GetRenderTargetItem().ShaderResourceTexture, TEXT("ShadingRateTexture")) : nullptr;
 	BasePassRenderTargets.SubpassHint = ESubpassHint::DepthReadSubpass;
 	if (!bIsFullDepthPrepassEnabled)
@@ -1020,7 +1020,7 @@ void FMobileSceneRenderer::RenderForward(FRDGBuilder& GraphBuilder, FRDGTextureR
 			BasePassRenderTargets[0].SetLoadAction(ERenderTargetLoadAction::ELoad);
 			BasePassRenderTargets.DepthStencil.SetDepthLoadAction(ERenderTargetLoadAction::ELoad);
 			BasePassRenderTargets.DepthStencil.SetStencilLoadAction(ERenderTargetLoadAction::ELoad);
-			BasePassRenderTargets.DepthStencil.SetDepthStencilAccess(FExclusiveDepthStencil::DepthWrite_StencilWrite);
+			BasePassRenderTargets.DepthStencil.SetDepthStencilAccess(bIsFullDepthPrepassEnabled ? FExclusiveDepthStencil::DepthRead_StencilWrite : FExclusiveDepthStencil::DepthWrite_StencilWrite);
 		}
 
 		View.BeginRenderView();
@@ -1285,7 +1285,7 @@ void FMobileSceneRenderer::RenderDeferred(FRDGBuilder& GraphBuilder, const FSort
 	TArrayView<FRDGTextureRef> BasePassTexturesView = MakeArrayView(ColorTargets);
 
 	FRenderTargetBindingSlots BasePassRenderTargets = GetRenderTargetBindings(ERenderTargetLoadAction::ENoAction, BasePassTexturesView);
-	BasePassRenderTargets.DepthStencil = FDepthStencilBinding(SceneTextures.Depth.Target, bIsFullDepthPrepassEnabled ? ERenderTargetLoadAction::ELoad : ERenderTargetLoadAction::EClear, FExclusiveDepthStencil::DepthWrite_StencilWrite);
+	BasePassRenderTargets.DepthStencil = bIsFullDepthPrepassEnabled ? FDepthStencilBinding(SceneTextures.Depth.Target, ERenderTargetLoadAction::ELoad, FExclusiveDepthStencil::DepthRead_StencilWrite) : FDepthStencilBinding(SceneTextures.Depth.Target, ERenderTargetLoadAction::EClear, FExclusiveDepthStencil::DepthWrite_StencilWrite);
 	BasePassRenderTargets.SubpassHint = ESubpassHint::DeferredShadingSubpass;
 	if (!bIsFullDepthPrepassEnabled)
 	{
