Fix mismatching Malloc / delete

#udn 00355521

diff --git a/Engine/Source/Runtime/Renderer/Private/DistanceFieldStreaming.cpp b/Engine/Source/Runtime/Renderer/Private/DistanceFieldStreaming.cpp
--- a/Engine/Source/Runtime/Renderer/Private/DistanceFieldStreaming.cpp
+++ b/Engine/Source/Runtime/Renderer/Private/DistanceFieldStreaming.cpp
@@ -406,7 +406,7 @@ void FDistanceFieldSceneData::AsyncUpdate(FDistanceFieldAsyncUpdateParameters Up
 			check(ReadRequest.Request.Status().IsCompleted());
 		}
 
-		delete ReadRequest.ReadOutputDataPtr;
+		FMemory::Free(ReadRequest.ReadOutputDataPtr);
 	}
 
 	FIoBatch Batch;
