SceneCaptureRendering: Initialized CopyCaptureToTargetSetViewportFn to avoid a crash.

This function gets initialized when there are scene capture components, but it could be called when there aren't any when rendering and CaptureNeedsSceneColor is true, which happens when the CaptureSource is SCS_SceneColorHDR.

diff --git a/Engine/Source/Runtime/Renderer/Private/SceneCaptureRendering.cpp b/Engine/Source/Runtime/Renderer/Private/SceneCaptureRendering.cpp
--- a/Engine/Source/Runtime/Renderer/Private/SceneCaptureRendering.cpp
+++ b/Engine/Source/Runtime/Renderer/Private/SceneCaptureRendering.cpp
@@ -161,7 +161,7 @@ static bool CaptureNeedsSceneColor(ESceneCaptureSource CaptureSource)
 	return CaptureSource != SCS_FinalColorLDR && CaptureSource != SCS_FinalColorHDR && CaptureSource != SCS_FinalToneCurveHDR;
 }
 
-static TFunction<void(FRHICommandList& RHICmdList)> CopyCaptureToTargetSetViewportFn;
+static TFunction<void(FRHICommandList& RHICmdList)> CopyCaptureToTargetSetViewportFn = [](FRHICommandList& RHICmdList) {};
 
 void CopySceneCaptureComponentToTarget(
 	FRDGBuilder& GraphBuilder,
