Skylight CastRayTracingShadows is now initialized to Disabled to fix a problem where skylight was being enabled by accident in old projects with ray tracing shadows enabled.

diff --git a/Engine/Source/Runtime/Engine/Classes/Engine/RendererSettings.h b/Engine/Source/Runtime/Engine/Classes/Engine/RendererSettings.h
--- a/Engine/Source/Runtime/Engine/Classes/Engine/RendererSettings.h
+++ b/Engine/Source/Runtime/Engine/Classes/Engine/RendererSettings.h
@@ -411,6 +411,11 @@ class ENGINE_API URendererSettings : public UDeveloperSettings
 		ToolTip = "Controls whether Ray Traced Shadows are used by default. Lights can still override and force Ray Traced shadows on or off. Requires Hardware Ray Tracing to be enabled."))
 		uint32 bEnableRayTracingShadows : 1;
 
+	UPROPERTY(config, EditAnywhere, Category = HardwareRayTracing, meta = (
+		ConsoleVariable = "r.RayTracing.Skylight", DisplayName = "Ray Traced Skylight",
+		ToolTip = "Controls whether Ray Traced Skylight is used by default. Skylights can still override and force Ray Traced Skylight on or off. Requires Hardware Ray Tracing to be enabled."))
+		uint32 bEnableRayTracingSkylight : 1;
+
 	UPROPERTY(config, EditAnywhere, Category = HardwareRayTracing, meta = (
 		ConsoleVariable = "r.RayTracing.UseTextureLod", DisplayName = "Texture LOD",
 		ToolTip = "Enable automatic texture mip level selection in ray tracing material shaders. Unchecked: highest resolution mip level is used for all texture (default). Checked: texture LOD is approximated based on total ray length, output resolution and texel density at hit point (ray cone method).",
diff --git a/Engine/Source/Runtime/Engine/Private/Components/SkyLightComponent.cpp b/Engine/Source/Runtime/Engine/Private/Components/SkyLightComponent.cpp
--- a/Engine/Source/Runtime/Engine/Private/Components/SkyLightComponent.cpp
+++ b/Engine/Source/Runtime/Engine/Private/Components/SkyLightComponent.cpp
@@ -277,6 +277,7 @@ USkyLightComponent::USkyLightComponent(const FObjectInitializer& ObjectInitializ
 	BlendDestinationAverageBrightness = 1.0f;
 	bCastVolumetricShadow = true;
 	CastRaytracedShadow = ECastRayTracedShadow::UseProjectSetting;
+	bCastRaytracedShadow_DEPRECATED = false;
 	bAffectReflection = true;
 	bAffectGlobalIllumination = true;
 	SamplesPerPixel = 4;
@@ -577,6 +578,11 @@ bool USkyLightComponent::CanEditChange(const FProperty* InProperty) const
 			static const auto CVar = IConsoleManager::Get().FindTConsoleVariableDataInt(TEXT("r.GenerateMeshDistanceFields"));
 			return Mobility == EComponentMobility::Movable && CastShadows && CVar->GetValueOnGameThread() != 0;
 		}
+
+		if (FCString::Strcmp(*PropertyName, TEXT("CastRaytracedShadow")) == 0)
+		{
+			return IsRayTracingEnabled();
+		}
 	}
 
 	return Super::CanEditChange(InProperty);
@@ -1102,7 +1108,6 @@ void USkyLightComponent::Serialize(FArchive& Ar)
 
 	if (Ar.IsLoading() && (Ar.CustomVer(FUE5MainStreamObjectVersion::GUID) < FUE5MainStreamObjectVersion::RayTracedShadowsType))
 	{
-		// Skylight is not driven by r.raytracing.shadows but if it was enabled before, it should be enabled now independently of the project setting
 		CastRaytracedShadow = bCastRaytracedShadow_DEPRECATED == 0 ? ECastRayTracedShadow::Disabled : ECastRayTracedShadow::Enabled;
 	}
 }
diff --git a/Engine/Source/Runtime/Engine/Private/RendererSettings.cpp b/Engine/Source/Runtime/Engine/Private/RendererSettings.cpp
--- a/Engine/Source/Runtime/Engine/Private/RendererSettings.cpp
+++ b/Engine/Source/Runtime/Engine/Private/RendererSettings.cpp
@@ -74,6 +74,7 @@ URendererSettings::URendererSettings(const FObjectInitializer& ObjectInitializer
 	bEnableRayTracing = 0;
 	bUseHardwareRayTracingForLumen = 0;
 	bEnableRayTracingShadows = 0;
+	bEnableRayTracingSkylight = 0;
 	bEnablePathTracing = 0;
 	bEnableRayTracingTextureLOD = 0;
 	MaxSkinBones = FGPUBaseSkinVertexFactory::GHardwareMaxGPUSkinBones;
diff --git a/Engine/Source/Runtime/Renderer/Private/DeferredShadingRenderer.cpp b/Engine/Source/Runtime/Renderer/Private/DeferredShadingRenderer.cpp
--- a/Engine/Source/Runtime/Renderer/Private/DeferredShadingRenderer.cpp
+++ b/Engine/Source/Runtime/Renderer/Private/DeferredShadingRenderer.cpp
@@ -1376,7 +1376,7 @@ bool FDeferredShadingSceneRenderer::SetupRayTracingPipelineStates(FRHICommandLis
 				PrepareSingleLayerWaterRayTracingReflections(View, *Scene, RayGenShaders);
 				PrepareRayTracingShadows(View, RayGenShaders);
 				PrepareRayTracingAmbientOcclusion(View, RayGenShaders);
-				PrepareRayTracingSkyLight(View, RayGenShaders);
+				PrepareRayTracingSkyLight(View, *Scene, RayGenShaders);
 				PrepareRayTracingGlobalIllumination(View, RayGenShaders);
 				PrepareRayTracingTranslucency(View, RayGenShaders);
 
diff --git a/Engine/Source/Runtime/Renderer/Private/DeferredShadingRenderer.h b/Engine/Source/Runtime/Renderer/Private/DeferredShadingRenderer.h
--- a/Engine/Source/Runtime/Renderer/Private/DeferredShadingRenderer.h
+++ b/Engine/Source/Runtime/Renderer/Private/DeferredShadingRenderer.h
@@ -945,7 +945,7 @@ private:
 	static void PrepareSingleLayerWaterRayTracingReflections(const FViewInfo& View, const FScene& Scene, TArray<FRHIRayTracingShader*>& OutRayGenShaders);
 	static void PrepareRayTracingShadows(const FViewInfo& View, TArray<FRHIRayTracingShader*>& OutRayGenShaders);
 	static void PrepareRayTracingAmbientOcclusion(const FViewInfo& View, TArray<FRHIRayTracingShader*>& OutRayGenShaders);
-	static void PrepareRayTracingSkyLight(const FViewInfo& View, TArray<FRHIRayTracingShader*>& OutRayGenShaders);
+	static void PrepareRayTracingSkyLight(const FViewInfo& View, const FScene& Scene, TArray<FRHIRayTracingShader*>& OutRayGenShaders);
 	static void PrepareRayTracingGlobalIllumination(const FViewInfo& View, TArray<FRHIRayTracingShader*>& OutRayGenShaders);
 	static void PrepareRayTracingTranslucency(const FViewInfo& View, TArray<FRHIRayTracingShader*>& OutRayGenShaders);
 	static void PrepareRayTracingDebug(const FSceneViewFamily& ViewFamily, TArray<FRHIRayTracingShader*>& OutRayGenShaders);
diff --git a/Engine/Source/Runtime/Renderer/Private/RayTracing/RaytracingSkylight.cpp b/Engine/Source/Runtime/Renderer/Private/RayTracing/RaytracingSkylight.cpp
--- a/Engine/Source/Runtime/Renderer/Private/RayTracing/RaytracingSkylight.cpp
+++ b/Engine/Source/Runtime/Renderer/Private/RayTracing/RaytracingSkylight.cpp
@@ -7,7 +7,7 @@
 #include "DeferredShadingRenderer.h"
 #include "SceneTextureParameters.h"
 
-static int32 GRayTracingSkyLight = -1;
+static int32 GRayTracingSkyLight = 0;
 
 #if RHI_RAYTRACING
 
@@ -122,11 +122,6 @@ int32 GetSkyLightSamplesPerPixel(const FSkyLightSceneProxy* SkyLightSceneProxy)
 	return GRayTracingSkyLightSamplesPerPixel >= 0 ? GRayTracingSkyLightSamplesPerPixel : FMath::Max(SkyLightSceneProxy->SamplesPerPixel, 2);
 }
 
-static bool IsRayTracingSkyLightAllowed()
-{
-	return ShouldRenderRayTracingEffect(GRayTracingSkyLight != 0);
-}
-
 bool ShouldRenderRayTracingSkyLight(const FSkyLightSceneProxy* SkyLightSceneProxy)
 {
 	if (SkyLightSceneProxy == nullptr)
@@ -134,14 +129,10 @@ bool ShouldRenderRayTracingSkyLight(const FSkyLightSceneProxy* SkyLightSceneProx
 		return false;
 	}
 
-	// SkyLightSceneProxy has a separate r.raytracing.skylight Cvar and does not check r.raytracing.shadows when CastRayTracedShadow == ECastRayTracedShadow::UseProjectSetting
-	bool bRayTracingSkyEnabled = GRayTracingSkyLight  < 0 
-		? SkyLightSceneProxy->CastRayTracedShadow == ECastRayTracedShadow::Enabled
-		: GRayTracingSkyLight != 0;
-
-	bRayTracingSkyEnabled = bRayTracingSkyEnabled && (GetSkyLightSamplesPerPixel(SkyLightSceneProxy) > 0);
+	bool bRayTracingSkyEnabled = (GRayTracingSkyLight  > 0 && SkyLightSceneProxy->CastRayTracedShadow == ECastRayTracedShadow::UseProjectSetting)
+								||  SkyLightSceneProxy->CastRayTracedShadow == ECastRayTracedShadow::Enabled;
 
-	return IsRayTracingSkyLightAllowed() && bRayTracingSkyEnabled;
+	return bRayTracingSkyEnabled && (GetSkyLightSamplesPerPixel(SkyLightSceneProxy) > 0);
 }
 
 IMPLEMENT_GLOBAL_SHADER_PARAMETER_STRUCT(FSkyLightData, "SkyLight");
@@ -261,9 +252,9 @@ class FRayTracingSkyLightRGS : public FGlobalShader
 
 IMPLEMENT_GLOBAL_SHADER(FRayTracingSkyLightRGS, "/Engine/Private/Raytracing/RaytracingSkylightRGS.usf", "SkyLightRGS", SF_RayGen);
 
-void FDeferredShadingSceneRenderer::PrepareRayTracingSkyLight(const FViewInfo& View, TArray<FRHIRayTracingShader*>& OutRayGenShaders)
+void FDeferredShadingSceneRenderer::PrepareRayTracingSkyLight(const FViewInfo& View, const FScene& Scene, TArray<FRHIRayTracingShader*>& OutRayGenShaders)
 {
-	if (!IsRayTracingSkyLightAllowed())
+	if (!ShouldRenderRayTracingSkyLight(Scene.SkyLight))
 	{
 		return;
 	}
