Fix instance shadow culling of dynamic primitives by adding a deferred culling pass after shadow setup.
 - Also correctly handle invalid instances (made possible by the above)
 - Make VSM invalidation also respect invalid instance flag.
 - Refactor FInstanceCullingManager interface to remove defunct CullInstances function.

diff --git a/Engine/Shaders/Private/InstanceCulling/BuildInstanceDrawCommands.usf b/Engine/Shaders/Private/InstanceCulling/BuildInstanceDrawCommands.usf
--- a/Engine/Shaders/Private/InstanceCulling/BuildInstanceDrawCommands.usf
+++ b/Engine/Shaders/Private/InstanceCulling/BuildInstanceDrawCommands.usf
@@ -202,15 +202,17 @@ bool ShouldRenderInstance(FInstanceSceneData InstanceData, FDrawCommandDesc Desc
 bool IsInstanceVisible(FInstanceSceneData InstanceData, uint ViewIdIndex)
 {
 #if CULL_INSTANCES
+	// Should not draw invalid instances, 
+	if (!InstanceData.ValidInstance)
+	{
+		return false;
+	}
+
 
 	// TODO: The test for dot(InstanceData.LocalBoundsExtent, InstanceData.LocalBoundsExtent) <= 0.0f is just a workaround since the FDynamicMeshBuilder::GetMesh
 	//       seems to just set empty bounds (and FLineBatcherSceneProxy pretends everything is at the origin). In the future these should compute reasonable bounds and 
 	//       this should be removed.
-	// NOTE: Flagging invalid as visible. This is a workaround that prevents instances that are uploaded after this kernel runs from being culled. 
-	// When the culling tests are moved into the later kernel we should switch this back to discarding. Batching moved them back up the pipeline so here we are.
-	// This happens for dynamic instances (e.g., procedural mesh) since they are uploaded with the view after the deferred/batched culling. To solve this we could
-	// use similarly batched & deferred upload of all dynamic primitives / instances (also more efficient).
-	if (!InstanceData.ValidInstance || dot(InstanceData.LocalBoundsExtent, InstanceData.LocalBoundsExtent) <= 0.0f)
+	if (dot(InstanceData.LocalBoundsExtent, InstanceData.LocalBoundsExtent) <= 0.0f)
 	{
 		return true;
 	}
@@ -360,4 +362,4 @@ void ClearIndirectArgInstanceCountCS(uint3 GroupId : SV_GroupID, int GroupThread
 	{
 		DrawIndirectArgsBufferOut[IndirectArgIndex * INDIRECT_ARGS_NUM_WORDS + 1] = 0U;
 	}
-}
\ No newline at end of file
+}
diff --git a/Engine/Shaders/Private/VirtualShadowMaps/CacheManagement.usf b/Engine/Shaders/Private/VirtualShadowMaps/CacheManagement.usf
--- a/Engine/Shaders/Private/VirtualShadowMaps/CacheManagement.usf
+++ b/Engine/Shaders/Private/VirtualShadowMaps/CacheManagement.usf
@@ -44,8 +44,6 @@ void VirtualSmInvalidateInstancePagesCS(
 	uint GroupId : SV_GroupID,
 	uint GroupThreadIndex : SV_GroupIndex)
 {
-	bool bIsVisible = false;
-	bool bWasOccluded = false;
 #if INPUT_KIND == INPUT_KIND_LARGE_RANGES
 	// Large instance count ranges use one group per range instead, and use a group-stride loop
 	FInstanceSceneDataRange InstanceSceneDataRange = InstanceSceneRanges[GroupId];
@@ -69,7 +67,8 @@ void VirtualSmInvalidateInstancePagesCS(
 		FInstanceSceneData InstanceSceneData = GetInstanceSceneData(InstanceId, InstanceSceneDataSOAStride);
 
 		// TODO: Clean up hardcoded flag field test.
-		bool bCastShadows = (GetPrimitiveData(InstanceSceneData.PrimitiveId).Flags & 1U) != 0U;
+		bool bCastShadows = InstanceSceneData.ValidInstance 
+			&& (GetPrimitiveData(InstanceSceneData.PrimitiveId).Flags & 1U) != 0U;
 #if ENABLE_DEBUG_MODE
 		uint PageInvalidationCount = 0U;
 #endif
diff --git a/Engine/Source/Runtime/Renderer/Private/DeferredShadingRenderer.cpp b/Engine/Source/Runtime/Renderer/Private/DeferredShadingRenderer.cpp
--- a/Engine/Source/Runtime/Renderer/Private/DeferredShadingRenderer.cpp
+++ b/Engine/Source/Runtime/Renderer/Private/DeferredShadingRenderer.cpp
@@ -1985,7 +1985,7 @@ void FDeferredShadingSceneRenderer::Render(FRDGBuilder& GraphBuilder)
 	// Find the visible primitives.
 	GraphBuilder.RHICmdList.ImmediateFlush(EImmediateFlushType::DispatchToRHIThread);
 
-	FInstanceCullingManager& InstanceCullingManager = *GraphBuilder.AllocObject<FInstanceCullingManager>(Scene->GPUScene.IsEnabled());
+	FInstanceCullingManager& InstanceCullingManager = *GraphBuilder.AllocObject<FInstanceCullingManager>(Scene->GPUScene.IsEnabled(), GraphBuilder);
 
 	{
 		RDG_GPU_STAT_SCOPE(GraphBuilder, VisibilityCommands);
@@ -2051,13 +2051,7 @@ void FDeferredShadingSceneRenderer::Render(FRDGBuilder& GraphBuilder)
 			Scene->GPUScene.DebugRender(GraphBuilder, *Scene, View);
 		}
 
-		{
-			// GPUCULL_TODO: Possibly fold into unpack step
-			InstanceCullingManager.CullInstances(GraphBuilder, Scene->GPUScene);
-
-			// This sets up a deferred context that is used by all subsequent culling passes until it is either flushed by RDG execute or explicitly closed
-			InstanceCullingManager.BeginDeferredCulling(GraphBuilder, Scene->GPUScene);
-		}
+		InstanceCullingManager.BeginDeferredCulling(GraphBuilder, Scene->GPUScene);
 
 		if (Views.Num() > 0)
 		{
diff --git a/Engine/Source/Runtime/Renderer/Private/InstanceCulling/InstanceCullingContext.cpp b/Engine/Source/Runtime/Renderer/Private/InstanceCulling/InstanceCullingContext.cpp
--- a/Engine/Source/Runtime/Renderer/Private/InstanceCulling/InstanceCullingContext.cpp
+++ b/Engine/Source/Runtime/Renderer/Private/InstanceCulling/InstanceCullingContext.cpp
@@ -308,7 +308,7 @@ public:
 
 	// Single Previous frame HZB which is shared among all batched contexts, thus only one is allowed (but the same can be used in multiple passes). (Needs atlas or bindless to expand(.
 	FRDGTextureRef PrevHZB = nullptr;
-
+	int32 NumCullingViews = 0;
 
 	bool bProcessed = false;
 
@@ -391,6 +391,13 @@ void FInstanceCullingContext::BuildRenderingCommands(
 				DeferredContext->TotalBatches[Mode] += LoadBalancers[Mode]->GetBatches().Max();
 				DeferredContext->TotalItems[Mode] += LoadBalancers[Mode]->GetItems().Max();
 			}
+#if DO_CHECK
+			for (int32 ViewId : ViewIds)
+			{
+				checkf(ViewId < DeferredContext->NumCullingViews, TEXT("Attempting to defer a culling context that references a view that has not been uploaded yet."));
+			}
+#endif 
+
 			DeferredContext->TotalIndirectArgs += IndirectArgs.Num();
 			DeferredContext->TotalViewIds += ViewIds.Num();
 			DeferredContext->InstanceIdBufferSize += InstanceIdBufferSize;
@@ -450,6 +457,12 @@ void FInstanceCullingContext::BuildRenderingCommands(
 	PassParametersTmp.NumCullingViews = 0;
 	if (bCullInstances)
 	{
+#if DO_CHECK
+		for (int32 ViewId : ViewIds)
+		{
+			checkf(ViewId < InstanceCullingManager->CullingIntermediate.NumViews, TEXT("Attempting to process a culling context that references a view that has not been uploaded yet."));
+		}
+#endif 
 		PassParametersTmp.InViews = GraphBuilder.CreateSRV(InstanceCullingManager->CullingIntermediate.CullingViews);
 		PassParametersTmp.NumCullingViews = InstanceCullingManager->CullingIntermediate.NumViews;
 	}
@@ -618,7 +631,6 @@ void FInstanceCullingDeferredContext::ProcessBatched(TStaticArray<FBuildInstance
 		PassParameters[Mode]->NumViewIds = ViewIds.Num();
 		PassParameters[Mode]->LoadBalancerParameters.NumBatches = LoadBalancers[Mode]->GetBatches().Num();
 		PassParameters[Mode]->LoadBalancerParameters.NumItems = LoadBalancers[Mode]->GetItems().Num();
-		PassParameters[Mode]->NumCullingViews = InstanceCullingManager->GetCullingViews().Num();
 
 		const bool bOcclusionCullInstances = PrevHZB != nullptr && FInstanceCullingContext::IsOcclusionCullingEnabled();
 		if (bOcclusionCullInstances)
@@ -693,14 +705,6 @@ FInstanceCullingDeferredContext *FInstanceCullingContext::CreateDeferredContext(
 
 	FBuildInstanceIdBufferAndCommandsFromPrimitiveIdsCs::FParameters PassParametersTmp;
 
-	// Defer upload of views as these may be added during setup of passes.
-	FRDGBufferRef ViewsRDG = CreateStructuredBuffer(GraphBuilder, 
-		TEXT("InstanceCulling.Views"), InstanceCullingManager->CullingViews.GetTypeSize(), 
-		INST_CULL_CALLBACK(DeferredContext->InstanceCullingManager->CullingViews.Num()),
-		INST_CULL_CALLBACK(DeferredContext->InstanceCullingManager->CullingViews.GetData()),
-		INST_CULL_CALLBACK(GetArrayDataSize(DeferredContext->InstanceCullingManager->CullingViews))
-	);
-
 	FRDGBufferRef DrawCommandDescsRDG = CreateStructuredBuffer(INST_CULL_CREATE_STRUCT_BUFF_ARGS(DrawCommandDescs));
 	FRDGBufferRef ViewIdsRDG = CreateStructuredBuffer(INST_CULL_CREATE_STRUCT_BUFF_ARGS(ViewIds));
 	FRDGBufferRef BatchInfosRDG = CreateStructuredBuffer(INST_CULL_CREATE_STRUCT_BUFF_ARGS(BatchInfos));
@@ -752,10 +756,11 @@ FInstanceCullingDeferredContext *FInstanceCullingContext::CreateDeferredContext(
 	PassParametersTmp.InstanceIdOffsetBuffer = GraphBuilder.CreateSRV(InstanceIdOffsetBuffer, PF_R32_UINT);
 	if (bCullInstances)
 	{
-		PassParametersTmp.InViews = GraphBuilder.CreateSRV(ViewsRDG);
-		// Note: not the final number
+		PassParametersTmp.InViews = GraphBuilder.CreateSRV(InstanceCullingManager->CullingIntermediate.CullingViews);
 		PassParametersTmp.NumCullingViews = InstanceCullingManager->CullingIntermediate.NumViews;
 	}
+	// Record the number of culling views to be able to check that no views referencing out-of bounds views are queued up
+	DeferredContext->NumCullingViews = InstanceCullingManager->CullingIntermediate.NumViews;
 
 	for (uint32 Mode = 0U; Mode < uint32(EBatchProcessingMode::Num); ++Mode)
 	{
diff --git a/Engine/Source/Runtime/Renderer/Private/InstanceCulling/InstanceCullingManager.cpp b/Engine/Source/Runtime/Renderer/Private/InstanceCulling/InstanceCullingManager.cpp
--- a/Engine/Source/Runtime/Renderer/Private/InstanceCulling/InstanceCullingManager.cpp
+++ b/Engine/Source/Runtime/Renderer/Private/InstanceCulling/InstanceCullingManager.cpp
@@ -21,6 +21,11 @@ static FAutoConsoleVariableRef CVarAllowBatchedBuildRenderingCommands(
 	TEXT("Whether to allow batching BuildRenderingCommands for GPU instance culling"),
 	ECVF_RenderThreadSafe);
 
+FInstanceCullingManager::FInstanceCullingManager(bool bInIsEnabled, FRDGBuilder& GraphBuilder)
+	: bIsEnabled(bInIsEnabled)
+{
+	CullingIntermediate.DummyUniformBuffer = FInstanceCullingContext::CreateDummyInstanceCullingUniformBuffer(GraphBuilder);
+}
 
 FInstanceCullingManager::~FInstanceCullingManager()
 {
@@ -55,27 +60,24 @@ int32 FInstanceCullingManager::RegisterView(const Nanite::FPackedViewParams& Par
 	return CullingViews.Num() - 1;
 }
 
-void FInstanceCullingManager::CullInstances(FRDGBuilder& GraphBuilder, FGPUScene& GPUScene)
+void FInstanceCullingManager::FlushRegisteredViews(FRDGBuilder& GraphBuilder)
 {
-	TRACE_CPUPROFILER_EVENT_SCOPE(FInstanceCullingManager::CullInstances);
-	RDG_EVENT_SCOPE(GraphBuilder, "CullInstances");
-
-
-	CullingIntermediate.CullingViews = CreateStructuredBuffer(GraphBuilder, TEXT("InstanceCulling.CullingViews"), CullingViews);
-	CullingIntermediate.NumViews = CullingViews.Num();
-
-	CullingIntermediate.DummyUniformBuffer = FInstanceCullingContext::CreateDummyInstanceCullingUniformBuffer(GraphBuilder);
+	if (CullingIntermediate.NumViews != CullingViews.Num())
+	{
+		CullingIntermediate.CullingViews = CreateStructuredBuffer(GraphBuilder, TEXT("InstanceCulling.CullingViews"), CullingViews);
+		CullingIntermediate.NumViews = CullingViews.Num();
+	}
 }
 
-
 bool FInstanceCullingManager::AllowBatchedBuildRenderingCommands(const FGPUScene& GPUScene)
 {
 	return GPUScene.IsEnabled() && !!GAllowBatchedBuildRenderingCommands && !FRDGBuilder::IsImmediateMode();
 }
 
-
 void FInstanceCullingManager::BeginDeferredCulling(FRDGBuilder& GraphBuilder, FGPUScene& GPUScene)
 {
+	FlushRegisteredViews(GraphBuilder);
+
 	// Cannot defer pass execution in immediate mode.
 	if (!AllowBatchedBuildRenderingCommands(GPUScene))
 	{
@@ -83,7 +85,7 @@ void FInstanceCullingManager::BeginDeferredCulling(FRDGBuilder& GraphBuilder, FG
 	}
 
 	// If there are no instances, there can be no work to perform later.
-	if (GPUScene.GetNumInstances() == 0 || CullingIntermediate.NumViews == 0)
+	if (GPUScene.GetNumInstances() == 0 || CullingViews.Num() == 0)
 	{
 		return;
 	}
diff --git a/Engine/Source/Runtime/Renderer/Private/InstanceCulling/InstanceCullingManager.h b/Engine/Source/Runtime/Renderer/Private/InstanceCulling/InstanceCullingManager.h
--- a/Engine/Source/Runtime/Renderer/Private/InstanceCulling/InstanceCullingManager.h
+++ b/Engine/Source/Runtime/Renderer/Private/InstanceCulling/InstanceCullingManager.h
@@ -77,6 +77,7 @@ class FInstanceCullingDeferredContext;
 class FInstanceCullingManager
 {
 public:
+	FInstanceCullingManager(bool bInIsEnabled, FRDGBuilder& GraphBuilder);
 	~FInstanceCullingManager();
 
 	bool IsEnabled() const { return bIsEnabled; }
@@ -88,17 +89,11 @@ public:
 	int32 RegisterView(const FViewInfo& ViewInfo);
 
 	/**
-	 * Run:
-	 *   AFTER views have been initialized and registered (including shadow views), 
-	 *   AFTER after GPUScene is updated, but
-	 *   BEFORE rendering commands are submitted.
+	 * Upload all registered views if the number has changed (grown). This must be done before calls to FInstanceCullingContext::BuildRenderingCommands that references
+	 * the views.
 	 */
-	void CullInstances(FRDGBuilder& GraphBuilder, FGPUScene& GPUScene);
+	void FlushRegisteredViews(FRDGBuilder& GraphBuilder);
 
-	FInstanceCullingManager(bool bInIsEnabled)
-	: bIsEnabled(bInIsEnabled)
-	{
-	}
 	const TRDGUniformBufferRef<FInstanceCullingGlobalUniforms> GetDummyInstanceCullingUniformBuffer() const { return CullingIntermediate.DummyUniformBuffer; }
 	
 	static bool AllowBatchedBuildRenderingCommands(const FGPUScene& GPUScene);
@@ -108,6 +103,9 @@ public:
 	 * Add a deferred, batched, gpu culling pass. Each batch represents a BuildRenderingCommands call from a mesh pass.
 	 * Batches are collected as we walk through the main render setup and call BuildRenderingCommands, and are processed when RDG Execute or Drain is called.
 	 * This implicitly ends the deferred context, so if Drain is used, it should be paired with a new call to BeginDeferredCulling.
+	 * Can be called multiple times, and will collect subsequent BuildRenderingCommands. Care must be taken that the views referenced in the build rendering commands
+	 * have been registered before BeginDeferredCulling.
+	 * Calls FlushRegisteredViews that uploads the registered views to the GPU.
 	 */
 	void BeginDeferredCulling(FRDGBuilder& GraphBuilder, FGPUScene& GPUScene);
 
diff --git a/Engine/Source/Runtime/Renderer/Private/MobileShadingRenderer.cpp b/Engine/Source/Runtime/Renderer/Private/MobileShadingRenderer.cpp
--- a/Engine/Source/Runtime/Renderer/Private/MobileShadingRenderer.cpp
+++ b/Engine/Source/Runtime/Renderer/Private/MobileShadingRenderer.cpp
@@ -571,11 +571,8 @@ void FMobileSceneRenderer::InitViews(FRDGBuilder& GraphBuilder, FSceneTexturesCo
 		PrepareDistanceFieldScene(GraphBuilder, false);
 	}
 	
-	{
-		// GPUCULL_TODO: Possibly fold into unpack step
-		InstanceCullingManager.CullInstances(GraphBuilder, Scene->GPUScene);
-		InstanceCullingManager.BeginDeferredCulling(GraphBuilder, Scene->GPUScene);
-	}
+	InstanceCullingManager.BeginDeferredCulling(GraphBuilder, Scene->GPUScene);
+
 	extern TSet<IPersistentViewUniformBufferExtension*> PersistentViewUniformBufferExtensions;
 
 	for (IPersistentViewUniformBufferExtension* Extension : PersistentViewUniformBufferExtensions)
@@ -730,7 +727,7 @@ void FMobileSceneRenderer::Render(FRDGBuilder& GraphBuilder)
 
 	FRDGSystemTextures::Create(GraphBuilder);
 
-	FInstanceCullingManager& InstanceCullingManager = *GraphBuilder.AllocObject<FInstanceCullingManager>(Scene->GPUScene.IsEnabled());
+	FInstanceCullingManager& InstanceCullingManager = *GraphBuilder.AllocObject<FInstanceCullingManager>(Scene->GPUScene.IsEnabled(), GraphBuilder);
 
 	// Important that this uses consistent logic throughout the frame, so evaluate once and pass in the flag from here
 	// NOTE: Must be done after  system texture initialization
diff --git a/Engine/Source/Runtime/Renderer/Private/SceneHitProxyRendering.cpp b/Engine/Source/Runtime/Renderer/Private/SceneHitProxyRendering.cpp
--- a/Engine/Source/Runtime/Renderer/Private/SceneHitProxyRendering.cpp
+++ b/Engine/Source/Runtime/Renderer/Private/SceneHitProxyRendering.cpp
@@ -551,7 +551,7 @@ void FMobileSceneRenderer::RenderHitProxies(FRDGBuilder& GraphBuilder)
 	FRDGTextureRef HitProxyDepthTexture = nullptr;
 	InitHitProxyRender(GraphBuilder, this, SceneTexturesConfig, HitProxyTexture, HitProxyDepthTexture);
 
-	FInstanceCullingManager& InstanceCullingManager = *GraphBuilder.AllocObject<FInstanceCullingManager>(Scene->GPUScene.IsEnabled());
+	FInstanceCullingManager& InstanceCullingManager = *GraphBuilder.AllocObject<FInstanceCullingManager>(Scene->GPUScene.IsEnabled(), GraphBuilder);
 
 	// Find the visible primitives.
 	InitViews(GraphBuilder, SceneTexturesConfig, InstanceCullingManager);
@@ -563,6 +563,8 @@ void FMobileSceneRenderer::RenderHitProxies(FRDGBuilder& GraphBuilder)
 	DynamicVertexBuffer.Commit();
 	DynamicReadBuffer.Commit();
 
+	InstanceCullingManager.FlushRegisteredViews(GraphBuilder);
+
 	TArray<Nanite::FRasterResults, TInlineAllocator<2>> NaniteRasterResults;
 	::DoRenderHitProxies(GraphBuilder, this, HitProxyTexture, HitProxyDepthTexture, NaniteRasterResults, InstanceCullingManager);
 
@@ -591,7 +593,7 @@ void FDeferredShadingSceneRenderer::RenderHitProxies(FRDGBuilder& GraphBuilder)
 
 	const FIntPoint HitProxyTextureSize = HitProxyDepthTexture->Desc.Extent;
 
-	FInstanceCullingManager InstanceCullingManager(Scene->GPUScene.IsEnabled());
+	FInstanceCullingManager& InstanceCullingManager = *GraphBuilder.AllocObject<FInstanceCullingManager>(Scene->GPUScene.IsEnabled(), GraphBuilder);
 
 	// Find the visible primitives.
 	{
@@ -625,7 +627,7 @@ void FDeferredShadingSceneRenderer::RenderHitProxies(FRDGBuilder& GraphBuilder)
 		Scene->GPUScene.UploadDynamicPrimitiveShaderDataForView(GraphBuilder, Scene, Views[ViewIndex]);
 	}
 
-	InstanceCullingManager.CullInstances(GraphBuilder, Scene->GPUScene);
+	InstanceCullingManager.FlushRegisteredViews(GraphBuilder);
 
 	if (bNaniteEnabled)
 	{
diff --git a/Engine/Source/Runtime/Renderer/Private/ShadowDepthRendering.cpp b/Engine/Source/Runtime/Renderer/Private/ShadowDepthRendering.cpp
--- a/Engine/Source/Runtime/Renderer/Private/ShadowDepthRendering.cpp
+++ b/Engine/Source/Runtime/Renderer/Private/ShadowDepthRendering.cpp
@@ -1111,9 +1111,6 @@ void FProjectedShadowInfo::BeginRenderView(FRDGBuilder& GraphBuilder, FScene* Sc
 			}
 		}
 	}
-
-	// This needs to be done for both mobile and deferred
-	Scene->GPUScene.UploadDynamicPrimitiveShaderDataForView(GraphBuilder, Scene, *ShadowDepthView);
 }
 
 static bool IsShadowDepthPassWaitForTasksEnabled()
@@ -1682,6 +1679,49 @@ void FSceneRenderer::RenderShadowDepthMaps(FRDGBuilder& GraphBuilder, FInstanceC
 	RDG_EVENT_SCOPE(GraphBuilder, "ShadowDepths");
 	RDG_GPU_STAT_SCOPE(GraphBuilder, ShadowDepths);
 
+	// Ensure all shadow view dynamic primitives are uploaded before shadow-culling batching pass.
+	// TODO: automate this such that:
+	//  1. we only process views that need it (have dynamic primitives)
+	//  2. it is integrated in the GPU-scene (it already collects the dynamic primives and know about them...)
+	//  3. BUT: we need to touch the views to update the GPUScene buffer references in the FViewInfo
+	//          so need to refactor that into its own binding point, probably. Or something.
+	for (FSortedShadowMapAtlas& ShadowMapAtlas : SortedShadowsForShadowDepthPass.ShadowMapAtlases)
+	{
+		for (FProjectedShadowInfo* ProjectedShadowInfo : ShadowMapAtlas.Shadows)
+		{
+			Scene->GPUScene.UploadDynamicPrimitiveShaderDataForView(GraphBuilder, Scene, *ProjectedShadowInfo->ShadowDepthView);
+		}
+	}
+	for (FSortedShadowMapAtlas& ShadowMap : SortedShadowsForShadowDepthPass.ShadowMapCubemaps)
+	{
+		check(ShadowMap.Shadows.Num() == 1);
+		FProjectedShadowInfo* ProjectedShadowInfo = ShadowMap.Shadows[0];
+		Scene->GPUScene.UploadDynamicPrimitiveShaderDataForView(GraphBuilder, Scene, *ProjectedShadowInfo->ShadowDepthView);
+	}
+	for (FProjectedShadowInfo* ProjectedShadowInfo : SortedShadowsForShadowDepthPass.PreshadowCache.Shadows)
+	{
+		if (!ProjectedShadowInfo->bDepthsCached)
+		{
+			Scene->GPUScene.UploadDynamicPrimitiveShaderDataForView(GraphBuilder, Scene, *ProjectedShadowInfo->ShadowDepthView);
+		}
+	}
+	for (const FSortedShadowMapAtlas& ShadowMapAtlas : SortedShadowsForShadowDepthPass.TranslucencyShadowMapAtlases)
+	{
+		for (FProjectedShadowInfo* ProjectedShadowInfo : ShadowMapAtlas.Shadows)
+		{
+			Scene->GPUScene.UploadDynamicPrimitiveShaderDataForView(GraphBuilder, Scene, *ProjectedShadowInfo->ShadowDepthView);
+		}
+	}
+	for (FProjectedShadowInfo* ProjectedShadowInfo : SortedShadowsForShadowDepthPass.VirtualShadowMapShadows)
+	{
+		Scene->GPUScene.UploadDynamicPrimitiveShaderDataForView(GraphBuilder, Scene, *ProjectedShadowInfo->ShadowDepthView);
+	}
+
+	// Begin new deferred culling bacthing scope to catch shadow render passes, as there can use dynamic primitives that have not been uploaded before 
+	// the previous batching scope. Also flushes the culling views registered during the setup (in InitViewsAfterPrepass) that are referenced in the shadow view
+	// culling.
+	InstanceCullingManager.BeginDeferredCulling(GraphBuilder, Scene->GPUScene);
+
 	// Perform setup work on all GPUs in case any cached shadows are being updated this
 	// frame. We revert to the AllViewsGPUMask for uncached shadows.
 #if WITH_MGPU
