Fixed problem when attempting to render ray tracing skylight with ray tracing disabled

diff --git a/Engine/Source/Runtime/Renderer/Private/RayTracing/RaytracingSkylight.cpp b/Engine/Source/Runtime/Renderer/Private/RayTracing/RaytracingSkylight.cpp
--- a/Engine/Source/Runtime/Renderer/Private/RayTracing/RaytracingSkylight.cpp
+++ b/Engine/Source/Runtime/Renderer/Private/RayTracing/RaytracingSkylight.cpp
@@ -132,7 +132,7 @@ bool ShouldRenderRayTracingSkyLight(const FSkyLightSceneProxy* SkyLightSceneProx
 	bool bRayTracingSkyEnabled = (GRayTracingSkyLight  > 0 && SkyLightSceneProxy->CastRayTracedShadow == ECastRayTracedShadow::UseProjectSetting)
 								||  SkyLightSceneProxy->CastRayTracedShadow == ECastRayTracedShadow::Enabled;
 
-	return bRayTracingSkyEnabled && (GetSkyLightSamplesPerPixel(SkyLightSceneProxy) > 0);
+	return bRayTracingSkyEnabled && ShouldRenderRayTracingEffect(true) && (GetSkyLightSamplesPerPixel(SkyLightSceneProxy) > 0);
 }
 
 IMPLEMENT_GLOBAL_SHADER_PARAMETER_STRUCT(FSkyLightData, "SkyLight");
