Fix texture optimization view modes not working for virtual textures.

diff --git a/Engine/Source/Editor/UnrealEd/Private/EditorViewportCommands.cpp b/Engine/Source/Editor/UnrealEd/Private/EditorViewportCommands.cpp
--- a/Engine/Source/Editor/UnrealEd/Private/EditorViewportCommands.cpp
+++ b/Engine/Source/Editor/UnrealEd/Private/EditorViewportCommands.cpp
@@ -338,17 +338,21 @@ static void AppendMaterialInfoToMenu(const UMaterialInterface* MaterialInterface
 	if (Material)
 	{
 		const FUniformExpressionSet& UniformExpressions = Material->GetUniformExpressions();
-		for (int32 i = 0; i < UniformExpressions.GetNumTextures(EMaterialTextureParameterType::Standard2D); ++i)
+		EMaterialTextureParameterType TextureTypes[] = { EMaterialTextureParameterType::Standard2D, EMaterialTextureParameterType::Virtual };
+		for (EMaterialTextureParameterType TextureType : TextureTypes)
 		{
-			UTexture* Texture = nullptr;
-			UniformExpressions.GetGameThreadTextureValue(EMaterialTextureParameterType::Standard2D, i, MaterialInterface, *Material, Texture, true);
-
-			const UTexture2D* Texture2D = Cast<UTexture2D>(Texture);
-			if (Texture2D)
+			for (int32 i = 0; i < UniformExpressions.GetNumTextures(TextureType); ++i)
 			{
-				const FMaterialTextureParameterInfo& Parameter = UniformExpressions.GetTextureParameter(EMaterialTextureParameterType::Standard2D, i);
-				DataPerTextureIndex.FindOrAdd(Parameter.TextureIndex).AddUnique(FString::Printf(TEXT("%s.%s"), *MaterialInterface->GetName(), *Texture2D->GetName()));
-				DataPerTextureName.FindOrAdd(*Texture2D->GetName()).AddUnique(FString::Printf(TEXT("%s %d : %s"), *MenuName, Parameter.TextureIndex, *MaterialInterface->GetName()));
+				UTexture* Texture = nullptr;
+				UniformExpressions.GetGameThreadTextureValue(TextureType, i, MaterialInterface, *Material, Texture, true);
+
+				const UTexture2D* Texture2D = Cast<UTexture2D>(Texture);
+				if (Texture2D)
+				{
+					const FMaterialTextureParameterInfo& Parameter = UniformExpressions.GetTextureParameter(TextureType, i);
+					DataPerTextureIndex.FindOrAdd(Parameter.TextureIndex).AddUnique(FString::Printf(TEXT("%s.%s"), *MaterialInterface->GetName(), *Texture2D->GetName()));
+					DataPerTextureName.FindOrAdd(*Texture2D->GetName()).AddUnique(FString::Printf(TEXT("%s %d : %s"), *MenuName, Parameter.TextureIndex, *MaterialInterface->GetName()));
+				}
 			}
 		}
 	}
diff --git a/Engine/Source/Runtime/Renderer/Private/RequiredTextureResolutionRendering.cpp b/Engine/Source/Runtime/Renderer/Private/RequiredTextureResolutionRendering.cpp
--- a/Engine/Source/Runtime/Renderer/Private/RequiredTextureResolutionRendering.cpp
+++ b/Engine/Source/Runtime/Renderer/Private/RequiredTextureResolutionRendering.cpp
@@ -33,26 +33,30 @@ void FRequiredTextureResolutionInterface::GetDebugViewModeShaderBindings(
 	int32 TextureResolution = 64;
 	FMaterialRenderContext MaterialContext(&MaterialRenderProxy, Material, nullptr);
 	const FUniformExpressionSet& UniformExpressions = Material.GetUniformExpressions();
+	EMaterialTextureParameterType TextureTypes[] = { EMaterialTextureParameterType::Standard2D, EMaterialTextureParameterType::Virtual };
 	if (ViewModeParam != INDEX_NONE && ViewModeParamName == NAME_None) // If displaying texture per texture indices
 	{
-		AnalysisIndex = ViewModeParam;
-
-		for (int32 ParameterIndex = 0; ParameterIndex < UniformExpressions.GetNumTextures(EMaterialTextureParameterType::Standard2D); ++ParameterIndex)
+		for (EMaterialTextureParameterType TextureType : TextureTypes)
 		{
-			const FMaterialTextureParameterInfo& Parameter = UniformExpressions.GetTextureParameter(EMaterialTextureParameterType::Standard2D, ParameterIndex);
-			if (Parameter.TextureIndex == ViewModeParam)
+			for (int32 ParameterIndex = 0; ParameterIndex < UniformExpressions.GetNumTextures(TextureType); ++ParameterIndex)
 			{
-				const UTexture* Texture = nullptr;
-				UniformExpressions.GetTextureValue(EMaterialTextureParameterType::Standard2D, ParameterIndex, MaterialContext, Material, Texture);
-				if (Texture && Texture->GetResource())
+				const FMaterialTextureParameterInfo& Parameter = UniformExpressions.GetTextureParameter(TextureType, ParameterIndex);
+				if (Parameter.TextureIndex == ViewModeParam)
 				{
-					if (Texture->IsStreamable())
+					const UTexture* Texture = nullptr;
+					UniformExpressions.GetTextureValue(TextureType, ParameterIndex, MaterialContext, Material, Texture);
+					if (Texture && Texture->GetResource())
 					{
-						TextureResolution = 1 << (Texture->GetResource()->GetCurrentMipCount() - 1);
-					}
-					else
-					{
-						TextureResolution = FMath::Max(Texture->GetResource()->GetSizeX(), Texture->GetResource()->GetSizeY());
+						AnalysisIndex = ViewModeParam;
+
+						if (Texture->IsStreamable())
+						{
+							TextureResolution = 1 << (Texture->GetResource()->GetCurrentMipCount() - 1);
+						}
+						else
+						{
+							TextureResolution = FMath::Max(Texture->GetResource()->GetSizeX(), Texture->GetResource()->GetSizeY());
+						}
 					}
 				}
 			}
@@ -60,22 +64,25 @@ void FRequiredTextureResolutionInterface::GetDebugViewModeShaderBindings(
 	}
 	else if (ViewModeParam != INDEX_NONE) // Otherwise show only texture matching the given name
 	{
-		AnalysisIndex = 1024; // Make sure not to find anything by default.
-		for (int32 ParameterIndex = 0; ParameterIndex < UniformExpressions.GetNumTextures(EMaterialTextureParameterType::Standard2D); ++ParameterIndex)
+		for (EMaterialTextureParameterType TextureType : TextureTypes)
 		{
-			const UTexture* Texture = nullptr;
-			UniformExpressions.GetTextureValue(EMaterialTextureParameterType::Standard2D, ParameterIndex, MaterialContext, Material, Texture);
-			if (Texture && Texture->GetResource() && Texture->GetFName() == ViewModeParamName)
+			for (int32 ParameterIndex = 0; ParameterIndex < UniformExpressions.GetNumTextures(TextureType); ++ParameterIndex)
 			{
-				if (Texture->IsStreamable())
+				const UTexture* Texture = nullptr;
+				UniformExpressions.GetTextureValue(TextureType, ParameterIndex, MaterialContext, Material, Texture);
+				if (Texture && Texture->GetResource() && Texture->GetFName() == ViewModeParamName)
 				{
-					const FMaterialTextureParameterInfo& Parameter = UniformExpressions.GetTextureParameter(EMaterialTextureParameterType::Standard2D, ParameterIndex);
+					const FMaterialTextureParameterInfo& Parameter = UniformExpressions.GetTextureParameter(TextureType, ParameterIndex);
 					AnalysisIndex = Parameter.TextureIndex;
-					TextureResolution = 1 << (Texture->GetResource()->GetCurrentMipCount() - 1);
-				}
-				else
-				{
-					TextureResolution = FMath::Max(Texture->GetResource()->GetSizeX(), Texture->GetResource()->GetSizeY());
+
+					if (Texture->IsStreamable())
+					{
+						TextureResolution = 1 << (Texture->GetResource()->GetCurrentMipCount() - 1);
+					}
+					else
+					{
+						TextureResolution = FMath::Max(Texture->GetResource()->GetSizeX(), Texture->GetResource()->GetSizeY());
+					}
 				}
 			}
 		}
