Added an option to disable high-end ray tracing reflection support per platform. This avoids expensive/long shader compilation for effects that will never be needed in practice.

diff --git a/Engine/Config/Windows/DataDrivenPlatformInfo.ini b/Engine/Config/Windows/DataDrivenPlatformInfo.ini
--- a/Engine/Config/Windows/DataDrivenPlatformInfo.ini
+++ b/Engine/Config/Windows/DataDrivenPlatformInfo.ini
@@ -38,6 +38,7 @@ bIsPC=true
 bSupportsRayTracing=true
 bSupportsRayTracingIndirectInstanceData=true
 bSupportsPathTracing=true
+bSupportsHighEndRayTracingReflections=true
 bSupportsWaveOperations=false
 bSupportsGPUScene=true
 bSupportsNanite=true
@@ -67,6 +68,7 @@ bIsPC=true
 bSupportsRayTracing=true
 bSupportsRayTracingIndirectInstanceData=true
 bSupportsPathTracing=true
+bSupportsHighEndRayTracingReflections=true
 bSupportsGPUScene=true
 bSupportsNanite=true
 bSupportsComputeFramework=true
diff --git a/Engine/Source/Runtime/RHI/Private/RHI.cpp b/Engine/Source/Runtime/RHI/Private/RHI.cpp
--- a/Engine/Source/Runtime/RHI/Private/RHI.cpp
+++ b/Engine/Source/Runtime/RHI/Private/RHI.cpp
@@ -2224,6 +2224,7 @@ void FGenericDataDrivenShaderPlatformInfo::ParseDataDrivenShaderInfo(const FConf
 	GET_SECTION_BOOL_HELPER(bSupportsRayTracing);
 	GET_SECTION_BOOL_HELPER(bSupportsRayTracingIndirectInstanceData);
 	GET_SECTION_BOOL_HELPER(bSupportsPathTracing);
+	GET_SECTION_BOOL_HELPER(bSupportsHighEndRayTracingReflections);
 	GET_SECTION_BOOL_HELPER(bSupportsGPUSkinCache);
 	GET_SECTION_BOOL_HELPER(bSupportsByteBufferComputeShaders);
 	GET_SECTION_BOOL_HELPER(bSupportsGPUScene);
diff --git a/Engine/Source/Runtime/RHI/Public/RHIDefinitions.h b/Engine/Source/Runtime/RHI/Public/RHIDefinitions.h
--- a/Engine/Source/Runtime/RHI/Public/RHIDefinitions.h
+++ b/Engine/Source/Runtime/RHI/Public/RHIDefinitions.h
@@ -342,6 +342,7 @@ class RHI_API FGenericDataDrivenShaderPlatformInfo
 	uint32 bSupportsRenderTargetWriteMask: 1;
 	uint32 bSupportsRayTracing: 1;
 	uint32 bSupportsRayTracingIndirectInstanceData : 1; // Whether instance transforms can be copied from the GPU to the TLAS instances buffer
+	uint32 bSupportsHighEndRayTracingReflections : 1; // Whether fully-featured RT reflections can be used on the platform (with multi-bounce, translucency, etc.)
 	uint32 bSupportsPathTracing : 1; // Whether real-time path tracer is supported on this platform (avoids compiling unnecessary shaders)
 	uint32 bSupportsGPUSkinCache: 1;
 	uint32 bSupportsGPUScene : 1;
@@ -550,6 +551,11 @@ public:
 		return Infos[Platform].bSupportsPathTracing;
 	}
 
+	static FORCEINLINE_DEBUGGABLE const bool GetSupportsHighEndRayTracingReflections(const FStaticShaderPlatform Platform)
+	{
+		return Infos[Platform].bSupportsHighEndRayTracingReflections;
+	}
+
 	static FORCEINLINE_DEBUGGABLE const bool GetSupportsGPUSkinCache(const FStaticShaderPlatform Platform)
 	{
 		return Infos[Platform].bSupportsGPUSkinCache;
diff --git a/Engine/Source/Runtime/Renderer/Private/RayTracing/RayTracingReflections.cpp b/Engine/Source/Runtime/Renderer/Private/RayTracing/RayTracingReflections.cpp
--- a/Engine/Source/Runtime/Renderer/Private/RayTracing/RayTracingReflections.cpp
+++ b/Engine/Source/Runtime/Renderer/Private/RayTracing/RayTracingReflections.cpp
@@ -317,7 +317,8 @@ class FRayTracingReflectionsRGS : public FGlobalShader
 
 	static bool ShouldCompilePermutation(const FGlobalShaderPermutationParameters& Parameters)
 	{
-		return ShouldCompileRayTracingShadersForProject(Parameters.Platform);
+		return ShouldCompileRayTracingShadersForProject(Parameters.Platform)
+			&& FDataDrivenShaderPlatformInfo::GetSupportsHighEndRayTracingReflections(Parameters.Platform);
 	}
 
 	static void ModifyCompilationEnvironment(const FGlobalShaderPermutationParameters& Parameters, FShaderCompilerEnvironment& OutEnvironment)
@@ -349,7 +350,8 @@ class FSplitImaginaryReflectionGBufferCS : public FGlobalShader
 
 	static bool ShouldCompilePermutation(const FGlobalShaderPermutationParameters& Parameters)
 	{
-		return ShouldCompileRayTracingShadersForProject(Parameters.Platform);
+		return ShouldCompileRayTracingShadersForProject(Parameters.Platform)
+			&& FDataDrivenShaderPlatformInfo::GetSupportsHighEndRayTracingReflections(Parameters.Platform);
 	}
 
 	static uint32 GetGroupSize()
@@ -414,7 +416,9 @@ static bool ShouldRayTracedReflectionsSortMaterials(const FViewInfo& View)
 
 static bool ShouldRayTracedReflectionsUseSortedDeferredAlgorithm(const FViewInfo& View)
 {
-	return (CVarRayTracingReflectionsExperimentalDeferred.GetValueOnRenderThread() != 0);
+	// Deferred reflections are always used if platform does not support the high-end version
+	return !FDataDrivenShaderPlatformInfo::GetSupportsHighEndRayTracingReflections(View.GetShaderPlatform())
+		|| (CVarRayTracingReflectionsExperimentalDeferred.GetValueOnRenderThread() != 0);
 }
 
 static bool ShouldRayTracedReflectionsRayTraceSkyLightContribution(const FScene& Scene)
@@ -432,7 +436,7 @@ void FDeferredShadingSceneRenderer::PrepareRayTracingReflections(const FViewInfo
 		return;
 	}
 
-	if (CVarRayTracingReflectionsExperimentalDeferred.GetValueOnRenderThread())
+	if (ShouldRayTracedReflectionsUseSortedDeferredAlgorithm(View))
 	{
 		// If deferred reflections technique is used, then we only need to gather its shaders and skip the rest.
 		PrepareRayTracingDeferredReflections(View, Scene, OutRayGenShaders);
@@ -492,10 +496,11 @@ void FDeferredShadingSceneRenderer::PrepareRayTracingReflectionsDeferredMaterial
 		return;
 	}
 
-	const bool bHybridReflections = ShouldRayTracedReflectionsUseHybridReflections();
-	const bool bRayTraceSkyLightContribution = ShouldRayTracedReflectionsRayTraceSkyLightContribution(Scene);
-
+	if (FDataDrivenShaderPlatformInfo::GetSupportsHighEndRayTracingReflections(View.GetShaderPlatform()))
 	{
+		const bool bHybridReflections = ShouldRayTracedReflectionsUseHybridReflections();
+		const bool bRayTraceSkyLightContribution = ShouldRayTracedReflectionsRayTraceSkyLightContribution(Scene);
+
 		FRayTracingReflectionsRGS::FPermutationDomain PermutationVector;
 		PermutationVector.Set<FRayTracingReflectionsRGS::FEnableTwoSidedGeometryForShadowDim>(EnableRayTracingShadowTwoSidedGeometry());
 		PermutationVector.Set<FRayTracingReflectionsRGS::FDeferredMaterialMode>(EDeferredMaterialMode::Gather);
@@ -513,16 +518,19 @@ void FDeferredShadingSceneRenderer::PrepareSingleLayerWaterRayTracingReflections
 		return;
 	}
 
-	const bool bHybridReflections = ShouldRayTracedReflectionsUseHybridReflections();
-	const bool bRayTraceSkyLightContribution = ShouldRayTracedReflectionsRayTraceSkyLightContribution(Scene);
+	if (FDataDrivenShaderPlatformInfo::GetSupportsHighEndRayTracingReflections(View.GetShaderPlatform()))
+	{
+		const bool bHybridReflections = ShouldRayTracedReflectionsUseHybridReflections();
+		const bool bRayTraceSkyLightContribution = ShouldRayTracedReflectionsRayTraceSkyLightContribution(Scene);
 
-	FRayTracingReflectionsRGS::FPermutationDomain PermutationVector;
-	PermutationVector.Set<FRayTracingReflectionsRGS::FEnableTwoSidedGeometryForShadowDim>(EnableRayTracingShadowTwoSidedGeometry());
-	PermutationVector.Set<FRayTracingReflectionsRGS::FDeferredMaterialMode>(EDeferredMaterialMode::None);
-	PermutationVector.Set<FRayTracingReflectionsRGS::FHybrid>(bHybridReflections);
-	PermutationVector.Set<FRayTracingReflectionsRGS::FRayTraceSkyLightContribution>(bRayTraceSkyLightContribution);
-	auto RayGenShader = View.ShaderMap->GetShader<FRayTracingReflectionsRGS>(PermutationVector);
-	OutRayGenShaders.Add(RayGenShader.GetRayTracingShader());
+		FRayTracingReflectionsRGS::FPermutationDomain PermutationVector;
+		PermutationVector.Set<FRayTracingReflectionsRGS::FEnableTwoSidedGeometryForShadowDim>(EnableRayTracingShadowTwoSidedGeometry());
+		PermutationVector.Set<FRayTracingReflectionsRGS::FDeferredMaterialMode>(EDeferredMaterialMode::None);
+		PermutationVector.Set<FRayTracingReflectionsRGS::FHybrid>(bHybridReflections);
+		PermutationVector.Set<FRayTracingReflectionsRGS::FRayTraceSkyLightContribution>(bRayTraceSkyLightContribution);
+		auto RayGenShader = View.ShaderMap->GetShader<FRayTracingReflectionsRGS>(PermutationVector);
+		OutRayGenShaders.Add(RayGenShader.GetRayTracingShader());
+	}
 }
 
 #endif // RHI_RAYTRACING
diff --git a/Engine/Source/Runtime/Renderer/Private/SingleLayerWaterRendering.cpp b/Engine/Source/Runtime/Renderer/Private/SingleLayerWaterRendering.cpp
--- a/Engine/Source/Runtime/Renderer/Private/SingleLayerWaterRendering.cpp
+++ b/Engine/Source/Runtime/Renderer/Private/SingleLayerWaterRendering.cpp
@@ -445,7 +445,9 @@ void FDeferredShadingSceneRenderer::RenderSingleLayerWaterReflections(
 		}
 
  		const bool bEnableSSR = CVarWaterSingleLayerSSR.GetValueOnRenderThread() != 0 && ScreenSpaceRayTracing::ShouldRenderScreenSpaceReflections(View);
-		const bool bEnableRTR = CVarWaterSingleLayerRTR.GetValueOnRenderThread() != 0 && ShouldRenderRayTracingReflections(View);
+		const bool bEnableRTR = CVarWaterSingleLayerRTR.GetValueOnRenderThread() != 0 && ShouldRenderRayTracingReflections(View) 
+			&& FDataDrivenShaderPlatformInfo::GetSupportsHighEndRayTracingReflections(View.GetShaderPlatform()); // Water requires the full RT reflection shader, which may not always be supported
+
 		if (bEnableRTR)
 		{
 			RDG_EVENT_SCOPE(GraphBuilder, "RayTracingWaterReflections");
