Fix crash when editing level instances.

The VisualizeLevelInstance pass needs the view uniform buffer, but didn't set it.

#rnx

diff --git a/Engine/Source/Runtime/Renderer/Private/PostProcess/PostProcessVisualizeLevelInstance.cpp b/Engine/Source/Runtime/Renderer/Private/PostProcess/PostProcessVisualizeLevelInstance.cpp
--- a/Engine/Source/Runtime/Renderer/Private/PostProcess/PostProcessVisualizeLevelInstance.cpp
+++ b/Engine/Source/Runtime/Renderer/Private/PostProcess/PostProcessVisualizeLevelInstance.cpp
@@ -38,7 +38,7 @@ IMPLEMENT_GLOBAL_SHADER(FVisualizeLevelInstancePS, "/Engine/Private/PostProcessV
 
 
 BEGIN_SHADER_PARAMETER_STRUCT(FVisualizeLevelInstancePassPassParameters, )
-	//SHADER_PARAMETER_STRUCT_REF(FViewUniformShaderParameters, View)
+	SHADER_PARAMETER_STRUCT_REF(FViewUniformShaderParameters, View)
 	SHADER_PARAMETER_STRUCT_INCLUDE(FSceneTextureShaderParameters, SceneTextures)
 	SHADER_PARAMETER_STRUCT_INCLUDE(FNaniteVisualizeLevelInstanceParameters, NaniteVisualizeLevelInstanceParameters)
 	SHADER_PARAMETER_STRUCT_INCLUDE(FInstanceCullingDrawParams, InstanceCullingDrawParams)
@@ -87,7 +87,7 @@ FScreenPassTexture AddVisualizeLevelInstancePass(FRDGBuilder& GraphBuilder, cons
 
 		const_cast<FViewInfo&>(View).ParallelMeshDrawCommandPasses[EMeshPass::EditorLevelInstance].BuildRenderingCommands(GraphBuilder, Scene->GPUScene, PassParameters->InstanceCullingDrawParams);
 
-		//PassParameters->View = EditorView->ViewUniformBuffer;
+		PassParameters->View = EditorView->ViewUniformBuffer;
 		PassParameters->SceneTextures = Inputs.SceneTextures;
 		PassParameters->RenderTargets.DepthStencil = FDepthStencilBinding(
 			DepthStencilTexture,
