Improve Distance Field rendering performance
- Reduce ALU, VGPR - move operations to CPU, reorganize asset data buffer.
- Added CSV stats to track atlas and indirection table size.
- DF shadow perf on heavy scenes improved by around 12% (3.75ms vs 4.3ms).

diff --git a/Engine/Shaders/Private/DistanceFieldLightingShared.ush b/Engine/Shaders/Private/DistanceFieldLightingShared.ush
--- a/Engine/Shaders/Private/DistanceFieldLightingShared.ush
+++ b/Engine/Shaders/Private/DistanceFieldLightingShared.ush
@@ -48,7 +48,11 @@ uint3 DistanceFieldBrickAtlasSizeLog2;
 // Size of a voxel in the brick atlas, which is also 1 / AtlasSize.
 float3 DistanceFieldBrickAtlasTexelSize;
 
-// Stores ranges of global brick indices for each asset mip.  Brick indices be INVALID_BRICK_INDEX.  
+float3 DistanceFieldBrickAtlasHalfTexelSize;
+float3 DistanceFieldBrickOffsetToAtlasUVScale;
+float3 DistanceFieldUniqueDataBrickSizeInAtlasTexels;
+
+// Stores ranges of global brick indices for each asset mip. Brick indices be INVALID_BRICK_INDEX.
 ByteAddressBuffer DistanceFieldIndirectionTable;
 
 // Brick atlas texture, used with hardware trilinear filtering
@@ -69,9 +73,9 @@ struct FDFAssetData
 	// Transforms the volume texture encoded distance to volume space 
 	float2 DistanceFieldToVolumeScaleBias;
 
-	// Transforms a volume space position to virtual UV space
-	float3 VolumeToVirtualUVAdd;
-	float3 VolumeToVirtualUVScale;
+	// Transforms a volume space position to virtual indirection space
+	float3 VolumeToIndirectionAdd;
+	float3 VolumeToIndirectionScale;
 };
 
 // Per-scene buffer containing packed Distance Field Asset data
@@ -82,27 +86,34 @@ StructuredBuffer<float4> SceneDistanceFieldAssetData;
 #define DF_MAX_NUM_MIPS 3
 #define DF_ASSET_DATA_STRIDE (DF_ASSET_DATA_MIP_STRIDE * DF_MAX_NUM_MIPS)
 #define INVALID_BRICK_INDEX 0xFFFFFFFF
+#define INDIRECTION_DIMENSION_MASK 0x3FF // Must match DistanceField::MaxIndirectionDimension
 
-// ReversedMipIndex must be in [0, NumMips - 1], where 0 is the lowest resolution mip that is always present
-FDFAssetData LoadDFAssetData(uint AssetIndex, uint ReversedMipIndex)
+FDFAssetData LoadDFAssetData(StructuredBuffer<float4> AssetDataBuffer, uint Offset)
 {
+	uint4 Vector0 = asuint(AssetDataBuffer[Offset + 0]);
+	float4 Vector1 = AssetDataBuffer[Offset + 1];
+	float4 Vector2 = AssetDataBuffer[Offset + 2];
+
 	FDFAssetData Data;
-	uint4 Vector0 = asuint(SceneDistanceFieldAssetData[AssetIndex * DF_ASSET_DATA_STRIDE + ReversedMipIndex * DF_ASSET_DATA_MIP_STRIDE + 0]);
-	// Must match DistanceField::MaxIndirectionDimension
-	Data.IndirectionDimensions.x = Vector0.x & 0x3FF;
-	Data.IndirectionDimensions.y = (Vector0.x >> 10) & 0x3FF;
-	Data.IndirectionDimensions.z = (Vector0.x >> 20) & 0x3FF;
+	Data.IndirectionDimensions.x = Vector0.x & INDIRECTION_DIMENSION_MASK;
+	Data.IndirectionDimensions.y = (Vector0.x >> 10) & INDIRECTION_DIMENSION_MASK;
+	Data.IndirectionDimensions.z = (Vector0.x >> 20) & INDIRECTION_DIMENSION_MASK;
 	Data.NumMips = Vector0.x >> 30;
-	Data.DistanceFieldToVolumeScaleBias = float2(f16tof32(asuint(Vector0.y) & 0xFFFF), f16tof32(asuint(Vector0.y) >> 16));
-	Data.IndirectionTableOffset = Vector0.z;
+	Data.IndirectionTableOffset = Vector0.y;
+
+	Data.DistanceFieldToVolumeScaleBias = float2(Vector1.w, Vector2.w);
+	Data.VolumeToIndirectionScale = Vector1.xyz;
+	Data.VolumeToIndirectionAdd = Vector2.xyz;
 
-	float4 Vector1 = SceneDistanceFieldAssetData[AssetIndex * DF_ASSET_DATA_STRIDE + ReversedMipIndex * DF_ASSET_DATA_MIP_STRIDE + 1];
-	Data.VolumeToVirtualUVScale = Vector1.xyz;
-	float4 Vector2 = SceneDistanceFieldAssetData[AssetIndex * DF_ASSET_DATA_STRIDE + ReversedMipIndex * DF_ASSET_DATA_MIP_STRIDE + 2];
-	Data.VolumeToVirtualUVAdd = Vector2.xyz;
 	return Data;
 }
 
+// ReversedMipIndex must be in [0, NumMips - 1], where 0 is the lowest resolution mip that is always present
+FDFAssetData LoadDFAssetData(uint AssetIndex, uint ReversedMipIndex)
+{
+	return LoadDFAssetData(SceneDistanceFieldAssetData, AssetIndex * DF_ASSET_DATA_STRIDE + ReversedMipIndex * DF_ASSET_DATA_MIP_STRIDE);
+}
+
 // Loads the mip asset data for the highest resolution mip, which is also the narrowest band, so the distance field sampled with this asset data will have clamped distances.  
 // Use DistanceToMeshSurfaceStandalone instead to get both an accurate and unclamped distance field.
 FDFAssetData LoadDFAssetDataHighestResolution(uint AssetIndex)
@@ -438,22 +449,23 @@ float SampleDistanceFieldBrickTexture(float3 UV)
 // This distance is clamped to the band of the current mip
 float SampleSparseMeshSignedDistanceField(float3 SampleVolumePosition, FDFAssetData DFAssetData)
 {
-	// Transform the volume space position into the [0 - 1] virtual UV space of the mesh
-	float3 VolumeUV = SampleVolumePosition * DFAssetData.VolumeToVirtualUVScale + DFAssetData.VolumeToVirtualUVAdd; 
+	// Transform the volume space position into the [0, IndirectionDimensions] virtual indirection space of the mesh
+	float3 IndirectionPos = SampleVolumePosition * DFAssetData.VolumeToIndirectionScale + DFAssetData.VolumeToIndirectionAdd;
 
-	uint3 IndirectionCoord = VolumeUV * DFAssetData.IndirectionDimensions;
+	int3 IndirectionCoord = IndirectionPos;
 	uint IndirectionIndex = (IndirectionCoord.z * DFAssetData.IndirectionDimensions.y + IndirectionCoord.y) * DFAssetData.IndirectionDimensions.x + IndirectionCoord.x;
 	// Fetch the brick index that covers this region of the virtual UV space
 	uint BrickIndex = DistanceFieldIndirectionTable.Load((DFAssetData.IndirectionTableOffset + IndirectionIndex) * 4);
+	bool ValidBrick = BrickIndex != INVALID_BRICK_INDEX;
 
 	float MaxEncodedDistance = DFAssetData.DistanceFieldToVolumeScaleBias.x + DFAssetData.DistanceFieldToVolumeScaleBias.y;
 	float DistanceField = MaxEncodedDistance;
 
 	// If the brick was not stored then it had the maximum encodable distance
-	if (BrickIndex != INVALID_BRICK_INDEX)
+	if (ValidBrick)
 	{
 		// Half voxel border around each brick
-		float3 BrickLocalUV = (VolumeUV * DFAssetData.IndirectionDimensions - IndirectionCoord) * DistanceFieldUniqueDataBrickSize + .5f;
+		float3 BrickLocalUV = IndirectionPos - IndirectionCoord;
 
 		// Decompose into 3d offset
 		float3 BrickOffset = uint3(
@@ -461,7 +473,7 @@ float SampleSparseMeshSignedDistanceField(float3 SampleVolumePosition, FDFAssetD
 			(BrickIndex >> DistanceFieldBrickAtlasSizeLog2.x) & DistanceFieldBrickAtlasMask.y, 
 			BrickIndex >> (DistanceFieldBrickAtlasSizeLog2.x + DistanceFieldBrickAtlasSizeLog2.y));
 
-		float3 AtlasUV = (BrickOffset * DistanceFieldBrickSize + BrickLocalUV) * DistanceFieldBrickAtlasTexelSize;
+		float3 AtlasUV = BrickOffset.xyz * DistanceFieldBrickOffsetToAtlasUVScale + BrickLocalUV * DistanceFieldUniqueDataBrickSizeInAtlasTexels + DistanceFieldBrickAtlasHalfTexelSize;
 		float EncodedDistanceField = SampleDistanceFieldBrickTexture(AtlasUV);
 
 		DistanceField = EncodedDistanceField * DFAssetData.DistanceFieldToVolumeScaleBias.x + DFAssetData.DistanceFieldToVolumeScaleBias.y;
@@ -510,8 +522,8 @@ float DistanceToMeshSurfaceStandalone(float3 SampleVolumePosition, FDFObjectData
 // SampleVolumePosition must be within [-VolumePositionExtent, VolumePositionExtent]
 float3 CalculateMeshSDFGradient(float3 SampleVolumePosition, FDFAssetData DFAssetData)
 {
-	float3 VirtualUVVoxelSize = 1.0f / float3(DFAssetData.IndirectionDimensions * DistanceFieldUniqueDataBrickSize);
-	float3 VolumeSpaceVoxelSize = VirtualUVVoxelSize / DFAssetData.VolumeToVirtualUVScale;
+	float3 VirtualUVVoxelSize = 1.0f / float3(DistanceFieldUniqueDataBrickSize);
+	float3 VolumeSpaceVoxelSize = VirtualUVVoxelSize / DFAssetData.VolumeToIndirectionScale;
 
 	// Mesh distance fields have a 1 voxel border around the mesh bounds and can't bilinear sample further than 0.5 voxel from that border
 	float3 VoxelOffset = 0.5f * VolumeSpaceVoxelSize;
@@ -547,5 +559,3 @@ float SampleHFVisibilityTexture(float2 UV)
 	return Texture2DSampleLevel(HFVisibilityTexture, DistanceFieldSampler, UV, 0).r;
 #endif
 }
-
-float2 HeightFieldAtlasTexelSize;
\ No newline at end of file
diff --git a/Engine/Source/Runtime/Engine/Private/DistanceFieldAtlas.cpp b/Engine/Source/Runtime/Engine/Private/DistanceFieldAtlas.cpp
--- a/Engine/Source/Runtime/Engine/Private/DistanceFieldAtlas.cpp
+++ b/Engine/Source/Runtime/Engine/Private/DistanceFieldAtlas.cpp
@@ -38,8 +38,6 @@
 
 #define LOCTEXT_NAMESPACE "DistanceField"
 
-CSV_DEFINE_CATEGORY(DistanceField, false);
-
 #if ENABLE_COOK_STATS
 namespace DistanceFieldCookStats
 {
diff --git a/Engine/Source/Runtime/Renderer/Private/DistanceFieldLightingShared.h b/Engine/Source/Runtime/Renderer/Private/DistanceFieldLightingShared.h
--- a/Engine/Source/Runtime/Renderer/Private/DistanceFieldLightingShared.h
+++ b/Engine/Source/Runtime/Renderer/Private/DistanceFieldLightingShared.h
@@ -103,12 +103,14 @@ BEGIN_SHADER_PARAMETER_STRUCT(FDistanceFieldAtlasParameters, )
 	SHADER_PARAMETER(FIntVector, DistanceFieldBrickAtlasMask)
 	SHADER_PARAMETER(FIntVector, DistanceFieldBrickAtlasSizeLog2)
 	SHADER_PARAMETER(FVector3f, DistanceFieldBrickAtlasTexelSize)
+	SHADER_PARAMETER(FVector3f, DistanceFieldBrickAtlasHalfTexelSize)
+	SHADER_PARAMETER(FVector3f, DistanceFieldBrickOffsetToAtlasUVScale)
+	SHADER_PARAMETER(FVector3f, DistanceFieldUniqueDataBrickSizeInAtlasTexels)
 END_SHADER_PARAMETER_STRUCT()
 
 BEGIN_SHADER_PARAMETER_STRUCT(FHeightFieldAtlasParameters, )
 	SHADER_PARAMETER_TEXTURE(Texture2D, HeightFieldTexture)
 	SHADER_PARAMETER_TEXTURE(Texture2D, HFVisibilityTexture)
-	SHADER_PARAMETER(FVector2D, HeightFieldAtlasTexelSize)
 END_SHADER_PARAMETER_STRUCT()
 
 namespace DistanceField
diff --git a/Engine/Source/Runtime/Renderer/Private/DistanceFieldObjectManagement.cpp b/Engine/Source/Runtime/Renderer/Private/DistanceFieldObjectManagement.cpp
--- a/Engine/Source/Runtime/Renderer/Private/DistanceFieldObjectManagement.cpp
+++ b/Engine/Source/Runtime/Renderer/Private/DistanceFieldObjectManagement.cpp
@@ -137,6 +137,10 @@ FDistanceFieldAtlasParameters DistanceField::SetupAtlasParameters(const FDistanc
 			FMath::FloorLog2(DistanceFieldSceneData.BrickTextureDimensionsInBricks.Y),
 			FMath::FloorLog2(DistanceFieldSceneData.BrickTextureDimensionsInBricks.Z));
 		SceneParameters.DistanceFieldBrickAtlasTexelSize = FVector(1.0f) / FVector(DistanceFieldSceneData.BrickTextureDimensionsInBricks * DistanceField::BrickSize);
+
+		SceneParameters.DistanceFieldBrickAtlasHalfTexelSize = 0.5f * SceneParameters.DistanceFieldBrickAtlasTexelSize;
+		SceneParameters.DistanceFieldBrickOffsetToAtlasUVScale = SceneParameters.DistanceFieldBrickSize * SceneParameters.DistanceFieldBrickAtlasTexelSize;
+		SceneParameters.DistanceFieldUniqueDataBrickSizeInAtlasTexels = SceneParameters.DistanceFieldUniqueDataBrickSize * SceneParameters.DistanceFieldBrickAtlasTexelSize;
 	}
 	else
 	{
diff --git a/Engine/Source/Runtime/Renderer/Private/DistanceFieldShadowing.cpp b/Engine/Source/Runtime/Renderer/Private/DistanceFieldShadowing.cpp
--- a/Engine/Source/Runtime/Renderer/Private/DistanceFieldShadowing.cpp
+++ b/Engine/Source/Runtime/Renderer/Private/DistanceFieldShadowing.cpp
@@ -709,7 +709,6 @@ void RayTraceShadows(
 	FHeightFieldAtlasParameters HeightFieldAtlasParameters;
 	HeightFieldAtlasParameters.HeightFieldTexture = GHeightFieldTextureAtlas.GetAtlasTexture();
 	HeightFieldAtlasParameters.HFVisibilityTexture = GHFVisibilityTextureAtlas.GetAtlasTexture();
-	HeightFieldAtlasParameters.HeightFieldAtlasTexelSize = FVector2D(1.0f / GHeightFieldTextureAtlas.GetSizeX(), 1.0f / GHeightFieldTextureAtlas.GetSizeY());
 
 	{
 		FDistanceFieldShadowingCS::FParameters* PassParameters = GraphBuilder.AllocParameters<FDistanceFieldShadowingCS::FParameters>();
diff --git a/Engine/Source/Runtime/Renderer/Private/DistanceFieldStreaming.cpp b/Engine/Source/Runtime/Renderer/Private/DistanceFieldStreaming.cpp
--- a/Engine/Source/Runtime/Renderer/Private/DistanceFieldStreaming.cpp
+++ b/Engine/Source/Runtime/Renderer/Private/DistanceFieldStreaming.cpp
@@ -21,6 +21,8 @@
 #include "DistanceFieldAmbientOcclusion.h"
 #include "GlobalDistanceField.h"
 
+CSV_DEFINE_CATEGORY(DistanceField, false);
+
 extern int32 GDFReverseAtlasAllocationOrder;
 
 static TAutoConsoleVariable<int32> CVarBrickAtlasSizeXYInBricks(
@@ -53,16 +55,16 @@ static TAutoConsoleVariable<int32> CVarDebugForceNumMips(
 	TEXT("When set to > 0, overrides the requested number of mips for streaming.  1 = only lowest resolution mip loaded, 3 = all mips loaded.  Mips will still be clamped by available space in the atlas."),
 	ECVF_RenderThreadSafe);
 
-int32 GDistanceFieldAtlasLogStats = 0;
-FAutoConsoleVariableRef CVarDistanceFieldAtlasLogStats(
+static int32 GDistanceFieldAtlasLogStats = 0;
+static FAutoConsoleVariableRef CVarDistanceFieldAtlasLogStats(
 	TEXT("r.DistanceFields.LogAtlasStats"),
 	GDistanceFieldAtlasLogStats,
 	TEXT("Set to 1 to dump atlas stats, set to 2 to dump atlas and SDF asset stats."),
 	ECVF_RenderThreadSafe
 	);
 
-const int32 MaxStreamingRequests = 4095;
-const int32 DistanceFieldBlockAllocatorSizeInBricks = 16;
+static const int32 MaxStreamingRequests = 4095;
+static const int32 DistanceFieldBlockAllocatorSizeInBricks = 16;
 
 class FCopyDistanceFieldAtlasCS : public FGlobalShader
 {
@@ -95,7 +97,6 @@ class FCopyDistanceFieldAtlasCS : public FGlobalShader
 
 IMPLEMENT_GLOBAL_SHADER(FCopyDistanceFieldAtlasCS, "/Engine/Private/DistanceFieldStreaming.usf", "CopyDistanceFieldAtlasCS", SF_Compute);
 
-
 class FScatterUploadDistanceFieldAtlasCS : public FGlobalShader
 {
 	DECLARE_GLOBAL_SHADER(FScatterUploadDistanceFieldAtlasCS)
@@ -789,6 +790,94 @@ void FDistanceFieldSceneData::GenerateStreamingRequests(
 	}
 }
 
+void EncodeAssetData(const FDistanceFieldAssetState& AssetState, const int32 ReversedMipIndex, FVector4* OutAssetData)
+{
+	const FDistanceFieldAssetMipState& MipState = AssetState.ReversedMips[ReversedMipIndex];
+	const int32 MipIndex = AssetState.BuiltData->Mips.Num() - ReversedMipIndex - 1;
+	const FSparseDistanceFieldMip& MipBuiltData = AssetState.BuiltData->Mips[MipIndex];
+	const FVector2D DistanceFieldToVolumeScaleBias = MipBuiltData.DistanceFieldToVolumeScaleBias;
+	const int32 NumMips = AssetState.ReversedMips.Num();
+
+	check(NumMips <= DistanceField::NumMips);
+	check(DistanceField::NumMips < 4);
+	check(MipBuiltData.IndirectionDimensions.X < DistanceField::MaxIndirectionDimension
+		&& MipBuiltData.IndirectionDimensions.Y < DistanceField::MaxIndirectionDimension
+		&& MipBuiltData.IndirectionDimensions.Z < DistanceField::MaxIndirectionDimension);
+
+	uint32 IntVector0[4] =
+	{
+		(uint32)MipBuiltData.IndirectionDimensions.X | (uint32)(MipBuiltData.IndirectionDimensions.Y << 10) | (uint32)(MipBuiltData.IndirectionDimensions.Z << 20) | (uint32)(NumMips << 30),
+		(uint32)MipState.IndirectionTableOffset,
+		0,
+		0
+	};
+
+	// Bypass NaN checks in FVector4 ctors
+	FVector4 FloatVector0;
+	FloatVector0.X = *(const float*)&IntVector0[0];
+	FloatVector0.Y = *(const float*)&IntVector0[1];
+	FloatVector0.Z = *(const float*)&IntVector0[2];
+	FloatVector0.W = *(const float*)&IntVector0[3];
+
+	FVector4 VolumeToIndirectionScale = FVector4(MipBuiltData.VolumeToVirtualUVScale, DistanceFieldToVolumeScaleBias.X);
+	VolumeToIndirectionScale.X *= MipBuiltData.IndirectionDimensions.X;
+	VolumeToIndirectionScale.Y *= MipBuiltData.IndirectionDimensions.Y;
+	VolumeToIndirectionScale.Z *= MipBuiltData.IndirectionDimensions.Z;
+
+	FVector4 VolumeToIndirectionAdd = FVector4(MipBuiltData.VolumeToVirtualUVAdd, DistanceFieldToVolumeScaleBias.Y);
+	VolumeToIndirectionAdd.X *= MipBuiltData.IndirectionDimensions.X;
+	VolumeToIndirectionAdd.Y *= MipBuiltData.IndirectionDimensions.Y;
+	VolumeToIndirectionAdd.Z *= MipBuiltData.IndirectionDimensions.Z;
+
+	OutAssetData[0] = FloatVector0;
+	OutAssetData[1] = VolumeToIndirectionScale;
+	OutAssetData[2] = VolumeToIndirectionAdd;
+}
+
+void FDistanceFieldSceneData::UploadAssetData(
+	FRDGBuilder& GraphBuilder,
+	const TArray<FDistanceFieldAssetMipId>& AssetDataUploads)
+{
+	if (AssetDataUploads.IsEmpty())
+	{
+		return;
+	}
+	
+	AssetDataUploadBuffer.Init(AssetDataUploads.Num(), AssetDataMipStrideFloat4s * sizeof(FVector4), true, TEXT("DistanceFields.DFAssetDataUploadBuffer"));
+
+	for (FDistanceFieldAssetMipId AssetMipUpload : AssetDataUploads)
+	{
+		const int32 ReversedMipIndex = AssetMipUpload.ReversedMipIndex;
+		FVector4* UploadAssetData = (FVector4*)AssetDataUploadBuffer.Add_GetRef(AssetMipUpload.AssetId.AsInteger() * DistanceField::NumMips + ReversedMipIndex);
+
+		if (AssetStateArray.IsValidId(AssetMipUpload.AssetId))
+		{
+			const FDistanceFieldAssetState& AssetState = AssetStateArray[AssetMipUpload.AssetId];
+			EncodeAssetData(AssetState, ReversedMipIndex, UploadAssetData);
+		}
+		else
+		{
+			// Clear invalid entries to zero
+			UploadAssetData[0] = FVector4(0.0f, 0.0f, 0.0f, 0.0f);
+			UploadAssetData[1] = FVector4(0.0f, 0.0f, 0.0f, 0.0f);
+			UploadAssetData[2] = FVector4(0.0f, 0.0f, 0.0f, 0.0f);
+		}
+	}
+
+	AddPass(GraphBuilder, RDG_EVENT_NAME("UploadAssetData"), [this](FRHICommandListImmediate& RHICmdList)
+		{
+			RHICmdList.Transition({
+				FRHITransitionInfo(AssetDataBuffer.UAV, ERHIAccess::Unknown, ERHIAccess::UAVCompute)
+				});
+
+			AssetDataUploadBuffer.ResourceUploadTo(RHICmdList, AssetDataBuffer, false);
+
+			RHICmdList.Transition({
+				FRHITransitionInfo(AssetDataBuffer.UAV, ERHIAccess::UAVCompute, ERHIAccess::SRVMask)
+				});
+		});
+}
+
 void FDistanceFieldSceneData::UpdateDistanceFieldAtlas(
 	FRDGBuilder& GraphBuilder, 
 	const FViewInfo& View,
@@ -875,6 +964,17 @@ void FDistanceFieldSceneData::UpdateDistanceFieldAtlas(
 	const uint32 IndirectionTableSizeBytes = FMath::Max<uint32>(FMath::RoundUpToPowerOfTwo(IndirectionTableAllocator.GetMaxSize()) * sizeof(uint32), 16);
 	ResizeResourceIfNeeded(GraphBuilder.RHICmdList, IndirectionTable, IndirectionTableSizeBytes, TEXT("DistanceFields.DFIndirectionTable"));
 
+	{
+		const FIntVector AtlasDimensions = BrickTextureDimensionsInBricks * DistanceField::BrickSize;
+		const SIZE_T AtlasSizeBytes = AtlasDimensions.X * AtlasDimensions.Y * AtlasDimensions.Z * GPixelFormats[DistanceField::DistanceFieldFormat].BlockBytes;
+		const SIZE_T IndirectionTableBytes = IndirectionTable.NumBytes;
+
+		float AtlasSizeMB = float(AtlasSizeBytes) / (1024.0f * 1024.0f);
+		float IndirectionTableSizeMB = float(IndirectionTableBytes) / (1024.0f * 1024.0f);
+		CSV_CUSTOM_STAT(DistanceField, AtlasMB, AtlasSizeMB, ECsvCustomStatOp::Set);
+		CSV_CUSTOM_STAT(DistanceField, IndirectionTableMB, IndirectionTableSizeMB, ECsvCustomStatOp::Set);
+	}
+
 	{
 		FDistanceFieldAsyncUpdateParameters UpdateParameters;
 		UpdateParameters.DistanceFieldSceneData = this;
@@ -977,72 +1077,8 @@ void FDistanceFieldSceneData::UpdateDistanceFieldAtlas(
 
 		GraphBuilder.FinalizeTextureAccess(DistanceFieldBrickVolumeTextureRDG, ERHIAccess::SRVMask);
 	}
-
-	if (AssetDataUploads.Num() > 0)
-	{
-		AssetDataUploadBuffer.Init(AssetDataUploads.Num(), AssetDataMipStrideFloat4s * sizeof(FVector4), true, TEXT("DistanceFields.DFAssetDataUploadBuffer"));
-
-		for (FDistanceFieldAssetMipId AssetMipUpload : AssetDataUploads)
-		{
-			const int32 ReversedMipIndex = AssetMipUpload.ReversedMipIndex;
-			FVector4* UploadAssetData = (FVector4*)AssetDataUploadBuffer.Add_GetRef(AssetMipUpload.AssetId.AsInteger() * DistanceField::NumMips + ReversedMipIndex);
-
-			if (AssetStateArray.IsValidId(AssetMipUpload.AssetId))
-			{
-				const FDistanceFieldAssetState& AssetState = AssetStateArray[AssetMipUpload.AssetId];
-				const FDistanceFieldAssetMipState& MipState = AssetState.ReversedMips[ReversedMipIndex];
-				const int32 MipIndex = AssetState.BuiltData->Mips.Num() - ReversedMipIndex - 1;
-				const FSparseDistanceFieldMip& MipBuiltData = AssetState.BuiltData->Mips[MipIndex];
-				const FVector2D DistanceFieldToVolumeScaleBias = MipBuiltData.DistanceFieldToVolumeScaleBias;
-				const int32 NumMips = AssetState.ReversedMips.Num();
-
-				check(NumMips <= DistanceField::NumMips);
-				check(DistanceField::NumMips < 4);
-				check(MipBuiltData.IndirectionDimensions.X < DistanceField::MaxIndirectionDimension
-					&& MipBuiltData.IndirectionDimensions.Y < DistanceField::MaxIndirectionDimension
-					&& MipBuiltData.IndirectionDimensions.Z < DistanceField::MaxIndirectionDimension);
-
-				uint32 IntVector0[4] =
-				{
-					(uint32)MipBuiltData.IndirectionDimensions.X | (uint32)(MipBuiltData.IndirectionDimensions.Y << 10) | (uint32)(MipBuiltData.IndirectionDimensions.Z << 20) | (uint32)(NumMips << 30),
-					(uint32)FFloat16(DistanceFieldToVolumeScaleBias.X).Encoded | ((uint32)FFloat16(DistanceFieldToVolumeScaleBias.Y).Encoded << 16),
-					(uint32)MipState.IndirectionTableOffset,
-					0
-				};
-
-				// Bypass NaN checks in FVector4 ctors
-				FVector4 FloatVector0;
-				FloatVector0.X = *(const float*)&IntVector0[0];
-				FloatVector0.Y = *(const float*)&IntVector0[1];
-				FloatVector0.Z = *(const float*)&IntVector0[2];
-				FloatVector0.W = *(const float*)&IntVector0[3];
-
-				UploadAssetData[0] = FloatVector0;
-				UploadAssetData[1] = FVector4(MipBuiltData.VolumeToVirtualUVScale, 0.0f);
-				UploadAssetData[2] = FVector4(MipBuiltData.VolumeToVirtualUVAdd, 0.0f);
-			}
-			else
-			{
-				// Clear invalid entries to zero
-				UploadAssetData[0] = FVector4(0.0f, 0.0f, 0.0f, 0.0f);
-				UploadAssetData[1] = FVector4(0.0f, 0.0f, 0.0f, 0.0f);
-				UploadAssetData[2] = FVector4(0.0f, 0.0f, 0.0f, 0.0f);
-			}
-		}
-
-		AddPass(GraphBuilder, RDG_EVENT_NAME("UploadAssetData"), [this](FRHICommandListImmediate& RHICmdList)
-		{
-			RHICmdList.Transition({
-				FRHITransitionInfo(AssetDataBuffer.UAV, ERHIAccess::Unknown, ERHIAccess::UAVCompute)
-			});
-
-			AssetDataUploadBuffer.ResourceUploadTo(RHICmdList, AssetDataBuffer, false);
-
-			RHICmdList.Transition({
-				FRHITransitionInfo(AssetDataBuffer.UAV, ERHIAccess::UAVCompute, ERHIAccess::SRVMask)
-			});
-		});
-	}
+	
+	UploadAssetData(GraphBuilder, AssetDataUploads);
 
 	GenerateStreamingRequests(GraphBuilder, View, Scene, bLumenEnabled, GlobalShaderMap);
 
diff --git a/Engine/Source/Runtime/Renderer/Private/ScenePrivate.h b/Engine/Source/Runtime/Renderer/Private/ScenePrivate.h
--- a/Engine/Source/Runtime/Renderer/Private/ScenePrivate.h
+++ b/Engine/Source/Runtime/Renderer/Private/ScenePrivate.h
@@ -2051,6 +2051,8 @@ private:
 
 	void ResizeBrickAtlasIfNeeded(FRDGBuilder& GraphBuilder, FGlobalShaderMap* GlobalShaderMap);
 
+	void UploadAssetData(FRDGBuilder& GraphBuilder, const TArray<FDistanceFieldAssetMipId>& AssetDataUploads);
+
 	void AsyncUpdate(FDistanceFieldAsyncUpdateParameters UpdateParameters);
 
 	void GenerateStreamingRequests(
