Fix non-unity compile error.

#rnx

diff --git a/Engine/Source/Runtime/Renderer/Private/RayTracing/RayTracingScene.h b/Engine/Source/Runtime/Renderer/Private/RayTracing/RayTracingScene.h
--- a/Engine/Source/Runtime/Renderer/Private/RayTracing/RayTracingScene.h
+++ b/Engine/Source/Runtime/Renderer/Private/RayTracing/RayTracingScene.h
@@ -12,6 +12,7 @@
 #include "RenderGraphResources.h"
 #include "Misc/MemStack.h"
 #include "Containers/ArrayView.h"
+#include "PrimitiveSceneProxy.h"
 
 class FRHIRayTracingScene;
 class FRHIShaderResourceView;
