Fixed compilation warning. Fixed potential crash on fbx export when ray tracing is supported by not enabled.

diff --git a/Engine/Source/Runtime/Renderer/Private/DeferredShadingRenderer.cpp b/Engine/Source/Runtime/Renderer/Private/DeferredShadingRenderer.cpp
--- a/Engine/Source/Runtime/Renderer/Private/DeferredShadingRenderer.cpp
+++ b/Engine/Source/Runtime/Renderer/Private/DeferredShadingRenderer.cpp
@@ -1433,19 +1433,17 @@ bool FDeferredShadingSceneRenderer::SetupRayTracingPipelineStates(FRHICommandLis
 }
 bool FDeferredShadingSceneRenderer::DispatchRayTracingWorldUpdates(FRDGBuilder& GraphBuilder)
 {
-	if (!IsRayTracingEnabled() || Views.Num() == 0)
-	{
-		return false;
-	}
-
 	bool bAnyRayTracingPassEnabled = false;
 	for (int32 ViewIndex = 0; ViewIndex < Views.Num(); ViewIndex++)
 	{
 		bAnyRayTracingPassEnabled |= AnyRayTracingPassEnabled(Scene, Views[ViewIndex]);
 	}
 
-	if (!bAnyRayTracingPassEnabled)
+	if (!IsRayTracingEnabled() || !bAnyRayTracingPassEnabled || Views.Num() == 0)
 	{
+		// This needs to happen even when ray tracing is not enabled
+		// because importers might batch BVH creation requests that need to be resolved in any case
+		GRayTracingGeometryManager.ProcessBuildRequests(GraphBuilder.RHICmdList);
 		return false;
 	}
 
@@ -3191,7 +3189,7 @@ void FDeferredShadingSceneRenderer::Render(FRDGBuilder& GraphBuilder)
 
 bool AnyRayTracingPassEnabled(const FScene* Scene, const FViewInfo& View)
 {
-	if (!IsRayTracingEnabled())
+	if (!IsRayTracingEnabled() || Scene == nullptr)
 	{
 		return false;
 	}
@@ -3200,7 +3198,7 @@ bool AnyRayTracingPassEnabled(const FScene* Scene, const FViewInfo& View)
 		|| ShouldRenderRayTracingReflections(View)
 		|| ShouldRenderRayTracingGlobalIllumination(View)
 		|| ShouldRenderRayTracingTranslucency(View)
-		|| ShouldRenderRayTracingSkyLight(Scene ? Scene->SkyLight : nullptr)
+		|| ShouldRenderRayTracingSkyLight(Scene->SkyLight)
 		|| ShouldRenderRayTracingShadows()
 		|| Scene->RayTracedLights.Num() > 0
         || Lumen::AnyLumenHardwareRayTracingPassEnabled(Scene, View)
