Implement angular mip gen setting to support legacy path. Code by [REDACTED_PERSON].

#rnx

diff --git a/Engine/Source/Developer/TextureCompressor/Private/TextureCompressorModule.cpp b/Engine/Source/Developer/TextureCompressor/Private/TextureCompressorModule.cpp
--- a/Engine/Source/Developer/TextureCompressor/Private/TextureCompressorModule.cpp
+++ b/Engine/Source/Developer/TextureCompressor/Private/TextureCompressorModule.cpp
@@ -2607,7 +2607,15 @@ private:
 		// Generate any missing mips in the chain.
 		if (NumOutputMips > OutMipChain.Num())
 		{
-			GenerateMipChain(BuildSettings, OutMipChain.Last(), OutMipChain);
+			// Do angular filtering of cubemaps if requested.
+			if (BuildSettings.MipGenSettings == TMGS_Angular)
+			{
+				GenerateAngularFilteredMips(OutMipChain, NumOutputMips, BuildSettings.DiffuseConvolveMipLevel);
+			}
+			else
+			{
+				GenerateMipChain(BuildSettings, OutMipChain.Last(), OutMipChain);
+			}
 		}
 		check(OutMipChain.Num() == NumOutputMips);
 
diff --git a/Engine/Source/Runtime/Engine/Classes/Engine/TextureDefines.h b/Engine/Source/Runtime/Engine/Classes/Engine/TextureDefines.h
--- a/Engine/Source/Runtime/Engine/Classes/Engine/TextureDefines.h
+++ b/Engine/Source/Runtime/Engine/Classes/Engine/TextureDefines.h
@@ -130,6 +130,8 @@ enum TextureMipGenSettings
 	TMGS_Blur5 UMETA(DisplayName = "Blur5"),
 	/** Use the first texel of each 2x2 (or 2x2x2) group. */
 	TMGS_Unfiltered UMETA(DisplayName = "Unfiltered"),
+	/** Introduce significant amount of blur using angular filtering (only applies to cubemaps, useful for ambient lighting). */
+	TMGS_Angular UMETA(DisplayName = "Angular"),
 	TMGS_MAX,
 
 	// Note: These are serialized as as raw values in the texture DDC key, so additional entries
diff --git a/Engine/Source/Runtime/Engine/Private/TextureLODSettings.cpp b/Engine/Source/Runtime/Engine/Private/TextureLODSettings.cpp
--- a/Engine/Source/Runtime/Engine/Private/TextureLODSettings.cpp
+++ b/Engine/Source/Runtime/Engine/Private/TextureLODSettings.cpp
@@ -2,6 +2,7 @@
 
 #include "Engine/TextureLODSettings.h"
 #include "Engine/Texture2D.h"
+#include "Engine/TextureCube.h"
 #include "Interfaces/ITargetPlatform.h"
 #include "Interfaces/ITargetPlatformManagerModule.h"
 #include "VT/VirtualTextureBuiltData.h"
@@ -259,6 +260,13 @@ void UTextureLODSettings::GetMipGenSettings(const UTexture& Texture, TextureMipG
 
 		Setting = LODGroup.MipGenSettings;
 	}
+
+	// angular filtering only applies to cubemaps
+	if (Setting == TMGS_Angular && !Texture.IsA(UTextureCube::StaticClass()))
+	{
+		Setting = TMGS_NoMipmaps;
+	}
+
 	OutMipGenSettings = Setting;
 
 	// ------------
diff --git a/Engine/Source/Runtime/Renderer/Private/SceneRendering.cpp b/Engine/Source/Runtime/Renderer/Private/SceneRendering.cpp
--- a/Engine/Source/Runtime/Renderer/Private/SceneRendering.cpp
+++ b/Engine/Source/Runtime/Renderer/Private/SceneRendering.cpp
@@ -64,6 +64,7 @@
 #include "SystemTextures.h"
 #include "VirtualShadowMaps/VirtualShadowMapClipmap.h"
 #include "Misc/AutomationTest.h"
+#include "Engine/TextureCube.h"
 #if WITH_EDITOR
 #include "Rendering/StaticLightingSystemInterface.h"
 #endif
@@ -3162,7 +3163,24 @@ void FSceneRenderer::RenderFinish(FRDGBuilder& GraphBuilder, FRDGTextureRef View
 				FSceneViewState* ViewState = (FSceneViewState*)View.State;
 				bool bViewParentOrFrozen = ViewState && (ViewState->HasViewParent() || ViewState->bIsFrozen);
 				bool bLocked = View.bIsLocked;
-				if ((GAreScreenMessagesEnabled && !GEngine->bSuppressMapWarnings) && (bViewParentOrFrozen || bLocked || bAnyWarning))
+
+				// display a warning if an ambient cubemap uses non-angular mipmap filtering
+				bool bShowAmbientCubemapMipGenSettingsWarning = false;
+
+#if WITH_EDITORONLY_DATA
+				for (FFinalPostProcessSettings::FCubemapEntry ContributingCubemap : View.FinalPostProcessSettings.ContributingCubemaps)
+				{
+					// platform configuration can't be loaded from the rendering thread, therefore the warning wont be displayed for TMGS_FromTextureGroup settings
+					if (ContributingCubemap.AmbientCubemap &&
+						ContributingCubemap.AmbientCubemap->MipGenSettings != TMGS_FromTextureGroup &&
+						ContributingCubemap.AmbientCubemap->MipGenSettings != TMGS_Angular)
+					{
+						bShowAmbientCubemapMipGenSettingsWarning = true;
+						break;
+					}
+				}
+#endif
+				if ((GAreScreenMessagesEnabled && !GEngine->bSuppressMapWarnings) && (bViewParentOrFrozen || bLocked || bShowAmbientCubemapMipGenSettingsWarning || bAnyWarning))
 				{
 					RDG_EVENT_SCOPE_CONDITIONAL(GraphBuilder, Views.Num() > 1, "View%d", ViewIndex);
 
@@ -3172,7 +3190,7 @@ void FSceneRenderer::RenderFinish(FRDGBuilder& GraphBuilder, FRDGTextureRef View
 						bViewParentOrFrozen, bShowSkylightWarning, bShowPointLightWarning, bShowShadowedLightOverflowWarning,
 						bShowMobileLowQualityLightmapWarning, bShowMobileMovableDirectionalLightWarning, bShowMobileDynamicCSMWarning, bMobileShowVertexFogWarning, bMobileMissingSkyMaterial, 
 						bShowSkinCacheOOM, bSingleLayerWaterWarning, bShowNoSkyAtmosphereComponentWarning, bFxDebugDraw, FXInterface, bShowLocalExposureDisabledWarning,
-						bLumenEnabledButHasNoDataForTracing, bLumenEnabledButDisabledForTheProject, bNaniteEnabledButNoAtomics, bNaniteEnabledButDisabledInProject, bNaniteRequireAtomics, bRealTimeSkyCaptureButNothingToCapture, bShowWaitingSkylight]
+						bLumenEnabledButHasNoDataForTracing, bLumenEnabledButDisabledForTheProject, bNaniteEnabledButNoAtomics, bNaniteEnabledButDisabledInProject, bNaniteRequireAtomics, bRealTimeSkyCaptureButNothingToCapture, bShowWaitingSkylight, bShowAmbientCubemapMipGenSettingsWarning]
 						(FCanvas& Canvas)
 					{
 						// so it can get the screen size
@@ -3370,6 +3388,13 @@ void FSceneRenderer::RenderFinish(FRDGBuilder& GraphBuilder, FRDGTextureRef View
 							Y += 14;
 						}
 
+						if (bShowAmbientCubemapMipGenSettingsWarning)
+						{
+							static const FText Message = NSLOCTEXT("Renderer", "AmbientCubemapMipGenSettings", "Ambient cubemaps should use 'Angular' Mip Gen Settings.");
+							Canvas.DrawShadowedText(10, Y, Message, GetStatsFont(), FLinearColor(1.0, 0.05, 0.05, 1.0));
+							Y += 14;
+						}
+
 #if WITH_EDITOR
 						FSkyLightSceneProxy* SkyLight = Scene->SkyLight;
 						if (bShowWaitingSkylight && SkyLight)
@@ -3864,18 +3889,18 @@ void FSceneRenderer::CleanUp(FRHICommandListImmediate& RHICmdList)
 
 	if (!GSceneRenderCleanUpState.bWaitForTasksComplete)
 	{
-	switch (GSceneRenderCleanUpState.CompletionMode)
-	{
-	case ESceneRenderCleanUpMode::Deferred:
+		switch (GSceneRenderCleanUpState.CompletionMode)
+		{
+		case ESceneRenderCleanUpMode::Deferred:
 			ReleaseSceneRenderer(RHICmdList, GSceneRenderCleanUpState.Renderer);
-		break;
-	case ESceneRenderCleanUpMode::DeferredAndAsync:
-		GSceneRenderCleanUpState.Task->Wait(ENamedThreads::GetRenderThread_Local());
+			break;
+		case ESceneRenderCleanUpMode::DeferredAndAsync:
+			GSceneRenderCleanUpState.Task->Wait(ENamedThreads::GetRenderThread_Local());
 			break;
 		}
 	}
 
-		DeleteSceneRenderer(RHICmdList, GSceneRenderCleanUpState.Renderer, GSceneRenderCleanUpState.MemStackMark);
+	DeleteSceneRenderer(RHICmdList, GSceneRenderCleanUpState.Renderer, GSceneRenderCleanUpState.MemStackMark);
 	GSceneRenderCleanUpState = {};
 }
 
