Hair and fur: Fixed AddHairVisibilityCommonPass() using the wrong ViewRect for AddSimpleMeshPass()
This caused hair and fur to not render at all in the right eye in XR and potentially in other situations that used multiple Views

diff --git a/Engine/Source/Runtime/Renderer/Private/HairStrands/HairStrandsVisibility.cpp b/Engine/Source/Runtime/Renderer/Private/HairStrands/HairStrandsVisibility.cpp
--- a/Engine/Source/Runtime/Renderer/Private/HairStrands/HairStrandsVisibility.cpp
+++ b/Engine/Source/Runtime/Renderer/Private/HairStrands/HairStrandsVisibility.cpp
@@ -2192,7 +2192,7 @@ static void AddHairVisibilityCommonPass(
 		PassParameters->View = ViewInfo->ViewUniformBuffer;
 	}
 
-	AddSimpleMeshPass(GraphBuilder, PassParameters, Scene, *ViewInfo, &InstanceCullingManager, GetPassName(), FIntRect(0, 0, ViewInfo->ViewRect.Width(), ViewInfo->ViewRect.Height()),
+	AddSimpleMeshPass(GraphBuilder, PassParameters, Scene, *ViewInfo, &InstanceCullingManager, GetPassName(), ViewInfo->ViewRect,
 		[PassParameters, Scene = Scene, ViewInfo, &MacroGroupDatas, RenderMode](FDynamicPassMeshDrawListContext* ShadowContext)
 	{
 		check(IsInRenderingThread());
