Add per-object triangle sorting per view for improving translucent ordering artifact.

Review https://p4-swarm.epicgames.net/reviews/17501695/

diff --git a/Engine/Shaders/Private/OIT/OITSorting.usf b/Engine/Shaders/Private/OIT/OITSorting.usf
new file mode 100644
--- /dev/null
+++ b/Engine/Shaders/Private/OIT/OITSorting.usf
@@ -0,0 +1,466 @@
+// Copyright Epic Games, Inc. All Rights Reserved.
+
+#pragma once
+
+#include "../Common.ush"
+
+///////////////////////////////////////////////////////////////////////////////////////////////////
+
+#define PRIMITIVE_TRILIST  0
+#define PRIMITIVE_TRISTRIP 1
+#define ENCODING_TYPE 1
+
+#ifndef SORTING_SLICE_COUNT
+	#error SORTING_SLICE_COUNT needs to be defined 
+#endif
+
+#define BACK_TO_FRONT 0
+#define FRONT_TO_BACK 1
+
+#if ENCODING_TYPE == 1
+	#define FPackedSliceAndOffset uint
+#else
+	#define FPackedSliceAndOffset uint2
+#endif
+
+#if PERMUTATION_DEBUG || SHADER_DEBUG
+#include "../ShaderDrawDebug.ush"
+#include "../ShaderPrintCommon.ush"
+#include "../ColorMap.ush"
+#endif
+
+///////////////////////////////////////////////////////////////////////////////////////////////////
+
+struct FTriIndex
+{
+	uint I0;
+	uint I1;
+	uint I2;
+};
+
+FTriIndex GetTriangleIndices(uint InPrimitiveId, uint InFirstIndex, uint InPrimitiveType)
+{
+	FTriIndex Out = (FTriIndex)0;
+	if (InPrimitiveType == PRIMITIVE_TRILIST)
+	{
+		Out.I0 = InFirstIndex + InPrimitiveId * 3;
+		Out.I1 = Out.I0 + 1;
+		Out.I2 = Out.I0 + 2;
+	}
+	else if (InPrimitiveType == PRIMITIVE_TRISTRIP)
+	{
+		Out.I0 = (InFirstIndex < 3 || InPrimitiveId < 1) ? InPrimitiveId : InPrimitiveId + 2;
+		Out.I1 = Out.I0 + 1;
+		Out.I2 = Out.I0 + 2;
+	}
+
+	return Out;
+}
+
+bool IsValid(FTriIndex InTri, uint InNumIndices)
+{
+	return InTri.I0 < InNumIndices && InTri.I1 < InNumIndices && InTri.I2 < InNumIndices;
+}
+
+FPackedSliceAndOffset PackSliceAndOffset(uint InSliceIndex, uint InLocalOffset)
+{
+#if ENCODING_TYPE
+	return (InLocalOffset & 0xFFFFFF) | ((InSliceIndex&0xFF) << 24);
+#else
+	return uint2(InSliceIndex, InLocalOffset);
+#endif
+}
+
+void UnpackSliceAndOffset(FPackedSliceAndOffset In, inout uint OutSliceIndex, inout uint OutLocalOffset)
+{
+#if ENCODING_TYPE
+	OutSliceIndex  = (In>>24) & 0xFF;
+	OutLocalOffset = In & 0xFFFFFF;
+#else
+	OutSliceIndex  = In.x;
+	OutLocalOffset = In.y;
+#endif
+}
+
+///////////////////////////////////////////////////////////////////////////////////////////////////
+// Debug
+
+struct FOITDebugData
+{
+	uint Index;
+	uint PrimitiveCount;
+	uint Type;
+	uint SizeInBytes;
+};
+
+void AddDebugInfo(RWBuffer<uint> InDebugData, FOITDebugData In)
+{
+	uint Index = 1;
+	InterlockedAdd(InDebugData[0], 1, Index);
+
+	const uint Index3 = 1 + Index * 3;		
+	InDebugData[Index3 + 0] = In.PrimitiveCount;
+	InDebugData[Index3 + 1] = In.Type;
+	InDebugData[Index3 + 2] = In.SizeInBytes;
+}
+
+uint GetDebugInfoCount(Buffer<uint> InDebugData)
+{
+	return InDebugData[0];
+}
+
+FOITDebugData GetDebugInfo(Buffer<uint> InDebugData, uint Index)
+{
+	FOITDebugData Out;
+	if (Index < InDebugData[0])
+	{
+		const uint Index3 = 1 + Index * 3;
+		Out.Index			= Index;
+		Out.PrimitiveCount	= InDebugData[Index3 + 0];
+		Out.Type			= InDebugData[Index3 + 1];
+		Out.SizeInBytes		= InDebugData[Index3 + 2];
+	}
+
+	return Out;
+}
+
+///////////////////////////////////////////////////////////////////////////////////////////////////
+
+#if SHADER_SCAN
+
+#if PERMUTATION_DEBUG
+float4x4 ViewToWorld;
+float3 WorldBound_Min;
+float3 WorldBound_Max;
+float3 ViewBound_Min;
+float3 ViewBound_Max;
+#endif
+
+float4x4 LocalToWorld; // GPU scene buffer instead? How?
+float4x4 WorldToView;
+
+uint SourcePrimitiveType;
+uint NumPrimitives;
+uint NumIndices;
+uint SourceFirstIndex;
+uint SortType;
+uint SortedIndexBufferSizeInByte;
+
+float ViewBoundMaxZ;
+float ViewBoundMinZ;
+
+Buffer<float> PositionBuffer;
+Buffer<uint> IndexBuffer;
+
+RWBuffer<uint> OutIndexBuffer;
+RWBuffer<uint> OutSliceCounterBuffer;
+RWBuffer<FPackedSliceAndOffset> OutPrimitiveSliceBuffer;
+
+RWBuffer<uint> OutDebugData;
+
+groupshared uint s_SliceIndices[SORTING_SLICE_COUNT];
+groupshared uint s_SliceOffsets[SORTING_SLICE_COUNT];
+
+#if SORTING_SLICE_COUNT > PERMUTATION_GROUP_SIZE
+	#error GroupSize (PERMUTATION_GROUP_SIZE) needs to be larger than the number of slices (SORTING_SLICE_COUNT) for having correct counting/allocation
+#endif
+
+[numthreads(PERMUTATION_GROUP_SIZE, 1, 1)]
+void MainCS(uint3 DispatchThreadId : SV_DispatchThreadID, uint GroupThreadId : SV_GroupIndex)
+{
+	uint PrimitiveId = DispatchThreadId.x;
+	uint LocalOffset = 0;
+	uint SliceIndex = 0;
+
+	// 0. Clear local allocation (shared-memory)
+	if (GroupThreadId < SORTING_SLICE_COUNT)
+	{
+		s_SliceIndices[GroupThreadId] = 0;
+	}
+
+	GroupMemoryBarrierWithGroupSync();
+
+	// 1. Compute slice index, and local allocation (shared-memory)
+	const bool bIsValid = PrimitiveId < NumPrimitives;
+	if (bIsValid)
+	{
+		const FTriIndex Tri = GetTriangleIndices(PrimitiveId, SourceFirstIndex, SourcePrimitiveType);
+		const uint I0 = 3 * IndexBuffer[Tri.I0];
+		const uint I1 = 3 * IndexBuffer[Tri.I1];
+		const uint I2 = 3 * IndexBuffer[Tri.I2];
+
+		// Position buffer store data as a float, not float3
+		const float3 P0 = float3(PositionBuffer[I0], PositionBuffer[I0+1], PositionBuffer[I0+2]);
+		const float3 P1 = float3(PositionBuffer[I1], PositionBuffer[I1+1], PositionBuffer[I1+2]);
+		const float3 P2 = float3(PositionBuffer[I2], PositionBuffer[I2+1], PositionBuffer[I2+2]);
+
+		// Compute triangle center in view space
+		const float3 LocalP = (P0 + P1 + P2) / 3;
+		const float3 WorldP = mul(float4(LocalP, 1), LocalToWorld).xyz;
+		const float3 ViewP  = mul(float4(WorldP, 1), WorldToView).xyz;
+
+		const float fSliceIndex = saturate((ViewP.z - ViewBoundMinZ) / (ViewBoundMaxZ - ViewBoundMinZ));
+		SliceIndex = SORTING_SLICE_COUNT * (SortType == BACK_TO_FRONT ? (1.f - fSliceIndex) : fSliceIndex);
+		SliceIndex = clamp(SliceIndex, 0, SORTING_SLICE_COUNT - 1);
+
+		#if PERMUTATION_DEBUG
+		{
+			const float4 TriColor = float4(ColorMapViridis(SliceIndex / float(SORTING_SLICE_COUNT - 1)), 1);
+
+			const float3 WP0 = mul(float4(P0, 1), LocalToWorld).xyz;
+			const float3 WP1 = mul(float4(P1, 1), LocalToWorld).xyz;
+			const float3 WP2 = mul(float4(P2, 1), LocalToWorld).xyz;
+			AddLineTriangle(WP0, WP1, WP2, TriColor);
+		}
+		#endif
+
+		InterlockedAdd(s_SliceIndices[SliceIndex], 1, LocalOffset);
+	}
+
+	#if PERMUTATION_DEBUG
+	if (all(DispatchThreadId == 0))
+	{
+		AddAABB(WorldBound_Min, WorldBound_Max, ColorYellow);
+		AddOBB(ViewBound_Min, ViewBound_Max, ColorCyan, ViewToWorld);
+
+		FOITDebugData DebugData;
+		DebugData.PrimitiveCount= NumPrimitives;
+		DebugData.Type			= SourcePrimitiveType;
+		DebugData.SizeInBytes	= SortedIndexBufferSizeInByte;
+		AddDebugInfo(OutDebugData, DebugData);
+	}
+	#endif
+
+	GroupMemoryBarrierWithGroupSync();
+
+	// 2. Global allocation (global-memory)
+	if (GroupThreadId < SORTING_SLICE_COUNT)
+	{
+		if (s_SliceIndices[GroupThreadId] > 0)
+		{
+			InterlockedAdd(OutSliceCounterBuffer[GroupThreadId], s_SliceIndices[GroupThreadId], s_SliceOffsets[GroupThreadId]);
+		}
+	}
+
+	GroupMemoryBarrierWithGroupSync();
+
+	// 3. Write slice index & relative offset per primitive
+	if (bIsValid)
+	{
+		// Write slice index and the primitive offset so that we don't need to atomic ops during the final write-out pass
+		const uint GlobalRelativeOffset = s_SliceOffsets[SliceIndex] + LocalOffset;
+		OutPrimitiveSliceBuffer[PrimitiveId] = PackSliceAndOffset(SliceIndex, GlobalRelativeOffset);
+	}
+}
+#endif // SHADER_SCAN
+
+///////////////////////////////////////////////////////////////////////////////////////////////////
+
+#if SHADER_ALLOCATE
+Buffer<uint> SliceCounterBuffer;
+RWBuffer<uint> SliceOffsetsBuffer;
+
+groupshared uint s_Counters[SORTING_SLICE_COUNT];
+groupshared uint s_Offsets[SORTING_SLICE_COUNT];
+
+[numthreads(SORTING_SLICE_COUNT, 1, 1)]
+void MainCS(uint3 DispatchThreadId : SV_DispatchThreadID)
+{
+	const uint SliceIndex = DispatchThreadId.x;
+	const bool bIsValid = SliceIndex < SORTING_SLICE_COUNT;
+
+	if (bIsValid)
+	{
+		s_Counters[SliceIndex] = SliceCounterBuffer[SliceIndex];
+	}
+
+	GroupMemoryBarrierWithGroupSync();
+
+	// TODO: parallel prefix sum
+	if (SliceIndex == 0)
+	{
+		uint AccumPrimitiveCount = 0;
+		for (uint SliceIt = 0; SliceIt < SORTING_SLICE_COUNT; ++SliceIt)
+		{
+			s_Offsets[SliceIt] = AccumPrimitiveCount;
+			AccumPrimitiveCount += s_Counters[SliceIt];
+		}
+	}
+
+	GroupMemoryBarrierWithGroupSync();
+
+	if (bIsValid)
+	{
+		SliceOffsetsBuffer[SliceIndex] = s_Offsets[SliceIndex];
+	}
+}
+#endif // SHADER_ALLOCATE
+
+///////////////////////////////////////////////////////////////////////////////////////////////////
+
+#if SHADER_WRITE
+
+uint SourcePrimitiveType;
+uint NumPrimitives;
+uint NumIndices;
+uint SrcFirstIndex;
+uint DstFirstIndex;
+uint SortType;
+
+Buffer<uint>  SliceOffsetsBuffer;
+Buffer<FPackedSliceAndOffset> PrimitiveSliceBuffer;
+
+Buffer<uint> IndexBuffer;
+RWBuffer<uint> OutIndexBuffer;
+
+[numthreads(256, 1, 1)]
+void MainCS(uint3 DispatchThreadId : SV_DispatchThreadID)
+{
+	uint SrcPrimitiveId = DispatchThreadId.x;
+	if (SrcPrimitiveId >= NumPrimitives)
+	{
+		return;
+	}
+
+	uint SliceIndex = 0;
+	uint PrimitiveOffset = 0;
+	UnpackSliceAndOffset(PrimitiveSliceBuffer[SrcPrimitiveId], SliceIndex, PrimitiveOffset);
+
+	uint SlicePrimitiveOffset = 0;
+	if (SliceIndex < SORTING_SLICE_COUNT)
+	{
+		SlicePrimitiveOffset = SliceOffsetsBuffer[SliceIndex];
+	}
+
+	const uint DstPrimitiveId = SlicePrimitiveOffset + PrimitiveOffset;
+
+	const FTriIndex SrcTri = GetTriangleIndices(SrcPrimitiveId, SrcFirstIndex, SourcePrimitiveType);
+	const FTriIndex DstTri = GetTriangleIndices(DstPrimitiveId, DstFirstIndex, PRIMITIVE_TRILIST); // Always write out triangle list, since sorting break strip topology
+
+	if (IsValid(DstTri, NumIndices))
+	{
+		OutIndexBuffer[DstTri.I0] = IndexBuffer[SrcTri.I0];
+		OutIndexBuffer[DstTri.I1] = IndexBuffer[SrcTri.I1];
+		OutIndexBuffer[DstTri.I2] = IndexBuffer[SrcTri.I2];
+	}
+}
+#endif // SHADER_WRITE
+
+
+///////////////////////////////////////////////////////////////////////////////////////////////////
+
+#if SHADER_DEBUG
+
+uint VisibleInstances;
+uint VisiblePrimitives;
+uint VisibleIndexSizeInBytes;
+
+uint AllocatedBuffers;
+uint AllocatedIndexSizeInBytes;
+
+uint UnusedBuffers;
+uint UnusedIndexSizeInBytes;
+
+uint TotalEntries;
+
+Buffer<uint> DebugData;
+
+void AddNewLine(float2 OriginalPos, inout float2 Pos)
+{
+	Pos = ShaderPrintNewline(Pos); Pos.x = OriginalPos.x;
+}
+
+float2 PrintSize(float2 Pos, uint SizeInBytes, FFontColor Color)
+{
+	if (SizeInBytes < 1000)
+	{
+		Pos = ShaderPrint(Pos, SizeInBytes / 1000.f, Color);
+		Pos = ShaderPrintText(Pos, TEXT("bytes"), Color);
+	}
+	else if (SizeInBytes < 1000000)
+	{
+		Pos = ShaderPrint(Pos, SizeInBytes / 1000.f, Color);
+		Pos = ShaderPrintText(Pos, TEXT("Kbytes"), Color);
+	}
+	else
+	{
+		Pos = ShaderPrint(Pos, SizeInBytes / 1000.f, Color);
+		Pos = ShaderPrintText(Pos, TEXT("Mbytes"), Color);
+	}
+	return Pos;
+}
+
+[numthreads(1, 1, 1)]
+void MainCS(uint DispatchThreadId : SV_DispatchThreadID)
+{
+	//float2 OriginalPos = ShaderPrintGetPos(uint2(50, 50)); // TODO: renable this when ported.
+	float2 OriginalPos = uint2(50, 50) / float2(1920,1080); 
+	float2 Pos = OriginalPos;
+
+	// Total NumPrimitive | Total Size 0.4 Mb
+
+	// Visible
+	{
+		Pos = ShaderPrintText(Pos, TEXT("Visible"), FontOrange); AddNewLine(OriginalPos, Pos);
+		Pos = ShaderPrintText(Pos, TEXT("Instances      : "), FontWhite); Pos = ShaderPrint(Pos, VisibleInstances, FontSilver); AddNewLine(OriginalPos, Pos);
+		Pos = ShaderPrintText(Pos, TEXT("Primitives     : "), FontWhite); Pos = ShaderPrint(Pos, VisiblePrimitives, FontSilver); AddNewLine(OriginalPos, Pos);
+		Pos = ShaderPrintText(Pos, TEXT("Size           : "), FontYellow); PrintSize(Pos, VisibleIndexSizeInBytes, FontYellow); AddNewLine(OriginalPos, Pos);
+	}
+
+	AddNewLine(OriginalPos, Pos);
+
+	// Allocated
+	{
+		Pos = ShaderPrintText(Pos, TEXT("Allocated"), FontOrange); AddNewLine(OriginalPos, Pos);
+		Pos = ShaderPrintText(Pos, TEXT("Buffers        : "), FontWhite); Pos = ShaderPrint(Pos, AllocatedBuffers, FontSilver); AddNewLine(OriginalPos, Pos);
+		Pos = ShaderPrintText(Pos, TEXT("Size           : "), FontYellow); PrintSize(Pos, AllocatedIndexSizeInBytes, FontYellow); AddNewLine(OriginalPos, Pos);
+	}
+
+	AddNewLine(OriginalPos, Pos);
+
+	// Unused
+	{
+		Pos = ShaderPrintText(Pos, TEXT("Unused"), FontOrange); AddNewLine(OriginalPos, Pos);
+		Pos = ShaderPrintText(Pos, TEXT("Buffers        : "), FontWhite); Pos = ShaderPrint(Pos, UnusedBuffers, FontSilver); AddNewLine(OriginalPos, Pos);
+		Pos = ShaderPrintText(Pos, TEXT("Size           : "), FontYellow); PrintSize(Pos, UnusedIndexSizeInBytes, FontYellow); AddNewLine(OriginalPos, Pos);
+	}
+
+	AddNewLine(OriginalPos, Pos);
+
+	{
+		Pos = ShaderPrintText(Pos, TEXT("Total"), FontOrange); AddNewLine(OriginalPos, Pos);
+		Pos = ShaderPrintText(Pos, TEXT("Entries        : "), FontWhite); ShaderPrint(Pos, TotalEntries, FontWhite); AddNewLine(OriginalPos, Pos);
+		Pos = ShaderPrintText(Pos, TEXT("Size           : "), FontOrange); PrintSize(Pos, AllocatedIndexSizeInBytes + UnusedIndexSizeInBytes, FontOrange); AddNewLine(OriginalPos, Pos);
+	}
+
+	AddNewLine(OriginalPos, Pos);
+
+	// 1: NumPrimitive | PrimitiveType | 0.4 Mb
+	{
+		Pos = ShaderPrintText(Pos, TEXT("Visible (details)"), FontSilver);
+		AddNewLine(OriginalPos, Pos);
+
+		Pos = ShaderPrintText(Pos, TEXT("Index       "));
+		Pos = ShaderPrintText(Pos, TEXT("Primitive   ")); 
+		Pos = ShaderPrintText(Pos, TEXT("Type  "));
+		Pos = ShaderPrintText(Pos, TEXT("Size"));
+		AddNewLine(OriginalPos, Pos);
+
+		for (uint Index = 0; Index < VisibleInstances; ++Index)
+		{
+			const FOITDebugData Data = GetDebugInfo(DebugData, Index);
+
+			Pos = ShaderPrint(Pos, Data.Index, FontWhite);
+			Pos = ShaderPrint(Pos, Data.PrimitiveCount, FontSilver);
+			if (Data.Type == 0)
+				Pos = ShaderPrintText(Pos, TEXT("List  "), FontSilver);
+			else
+				Pos = ShaderPrintText(Pos, TEXT("Strip "), FontSilver);
+
+			Pos = PrintSize(Pos, Data.SizeInBytes, FontYellow);
+
+			AddNewLine(OriginalPos, Pos);
+		}
+	}
+}
+#endif // SHADER_DEBUG
\ No newline at end of file
diff --git a/Engine/Source/Runtime/Engine/Classes/Components/StaticMeshComponent.h b/Engine/Source/Runtime/Engine/Classes/Components/StaticMeshComponent.h
--- a/Engine/Source/Runtime/Engine/Classes/Components/StaticMeshComponent.h
+++ b/Engine/Source/Runtime/Engine/Classes/Components/StaticMeshComponent.h
@@ -329,6 +329,10 @@ public:
 	uint8 bDisplayNaniteProxyMesh:1;
 #endif
 
+	/** Enable dynamic sort mesh's triangles to remove ordering issue when rendered with a translucent material */
+	UPROPERTY(EditAnywhere, AdvancedDisplay, BlueprintReadOnly, Category = Lighting, meta = (UIMin = "0", UIMax = "1", DisplayName = "Sort Triangles (Experimental)"))
+	uint8 bSortTriangles : 1;
+
 	/**
 	 * Controls whether the static mesh component's backface culling should be reversed
 	 */
diff --git a/Engine/Source/Runtime/Engine/Private/PrimitiveSceneProxy.cpp b/Engine/Source/Runtime/Engine/Private/PrimitiveSceneProxy.cpp
--- a/Engine/Source/Runtime/Engine/Private/PrimitiveSceneProxy.cpp
+++ b/Engine/Source/Runtime/Engine/Private/PrimitiveSceneProxy.cpp
@@ -167,6 +167,7 @@ FPrimitiveSceneProxy::FPrimitiveSceneProxy(const UPrimitiveComponent* InComponen
 ,	bSupportsDistanceFieldRepresentation(false)
 ,	bSupportsMeshCardRepresentation(false)
 ,	bSupportsHeightfieldRepresentation(false)
+,	bSupportsSortedTriangles(false)
 ,	bShouldNotifyOnWorldAddRemove(false)
 ,	bWantsSelectionOutline(true)
 ,	bVerifyUsedMaterials(true)
diff --git a/Engine/Source/Runtime/Engine/Private/StaticMeshRender.cpp b/Engine/Source/Runtime/Engine/Private/StaticMeshRender.cpp
--- a/Engine/Source/Runtime/Engine/Private/StaticMeshRender.cpp
+++ b/Engine/Source/Runtime/Engine/Private/StaticMeshRender.cpp
@@ -364,6 +364,9 @@ FStaticMeshSceneProxy::FStaticMeshSceneProxy(UStaticMeshComponent* InComponent,
 #endif
 
 	AddSpeedTreeWind();
+
+	// Enable dynamic triangle reordering to remove/reduce sorting issue when rendered with a translucent material (i.e., order-independent-transparency)
+	bSupportsSortedTriangles = InComponent->bSortTriangles;
 }
 
 void FStaticMeshSceneProxy::SetEvaluateWorldPositionOffsetInRayTracing(bool NewValue)
@@ -1113,7 +1116,7 @@ void FStaticMeshSceneProxy::DrawStaticElements(FStaticPrimitiveDrawInterface* PD
 			// Draw the static mesh elements.
 			for(int32 SectionIndex = 0; SectionIndex < LODModel.Sections.Num(); SectionIndex++)
 			{
-#if WITH_EDITOR
+			#if WITH_EDITOR
 				if (GIsEditor)
 				{
 					const FLODInfo::FSectionInfo& Section = LODs[LODIndex].Sections[SectionIndex];
diff --git a/Engine/Source/Runtime/Engine/Public/MeshBatch.h b/Engine/Source/Runtime/Engine/Public/MeshBatch.h
--- a/Engine/Source/Runtime/Engine/Public/MeshBatch.h
+++ b/Engine/Source/Runtime/Engine/Public/MeshBatch.h
@@ -42,6 +42,22 @@ enum EPrimitiveIdMode
 // Flag used to mark a primtive ID as dynamic, and thus needing translation (by adding the offset from the dynamic primitive collector).
 static constexpr int32 GPrimIDDynamicFlag = 1 << 31;
 
+struct FMeshBatchElementDynamicIndexBuffer
+{
+	/** The vertex buffer to bind for draw calls. */
+	FIndexBuffer* IndexBuffer = nullptr;
+	/** The offset in to the index buffer (arbitrary limit to 16M indices so that FirstIndex|PrimitiveType fits into 32bits. */
+	uint32 FirstIndex : 24;
+	/** The offset in to the index buffer. */
+	uint32 PrimitiveType : PT_NumBits;
+
+	/** Returns true if the allocation is valid. */
+	FORCEINLINE bool IsValid() const
+	{
+		return IndexBuffer != NULL;
+	}
+};
+
 /**
  * A batch mesh element definition.
  */
@@ -61,6 +77,12 @@ struct FMeshBatchElement
 
 	const FIndexBuffer* IndexBuffer;
 
+	/**
+	 * Store dynamic index buffer
+	 * This is used for objects whose triangles are dynamically sorted for a particular view (i.e., per-object order-independent-transparency)
+	*/
+	FMeshBatchElementDynamicIndexBuffer DynamicIndexBuffer;
+
 	union 
 	{
 		/** If !bIsSplineProxy, Instance runs, where number of runs is specified by NumInstances.  Run structure is [StartInstanceIndex, EndInstanceIndex]. */
diff --git a/Engine/Source/Runtime/Engine/Public/PrimitiveSceneProxy.h b/Engine/Source/Runtime/Engine/Public/PrimitiveSceneProxy.h
--- a/Engine/Source/Runtime/Engine/Public/PrimitiveSceneProxy.h
+++ b/Engine/Source/Runtime/Engine/Public/PrimitiveSceneProxy.h
@@ -635,6 +635,7 @@ public:
 	inline bool SupportsMeshCardRepresentation() const { return bSupportsMeshCardRepresentation; }
 	inline bool SupportsHeightfieldRepresentation() const { return bSupportsHeightfieldRepresentation; }
 	inline bool SupportsInstanceDataBuffer() const { return bSupportsInstanceDataBuffer; }
+	inline bool SupportsSortedTriangles() const { return bSupportsSortedTriangles; }
 	inline bool TreatAsBackgroundForOcclusion() const { return bTreatAsBackgroundForOcclusion; }
 	inline bool ShouldNotifyOnWorldAddRemove() const { return bShouldNotifyOnWorldAddRemove; }
 	inline bool IsForceHidden() const {return bForceHidden;}
@@ -1123,6 +1124,9 @@ protected:
 	/** Whether the primitive implements GetHeightfieldRepresentation() */
 	uint8 bSupportsHeightfieldRepresentation : 1;
 
+	/** Whether the object support triangles when rendered with translucent material */
+	uint8 bSupportsSortedTriangles : 1;
+
 	/** Whether this primitive requires notification when its level is added to the world and made visible for the first time. */
 	uint8 bShouldNotifyOnWorldAddRemove : 1;
 
diff --git a/Engine/Source/Runtime/Renderer/Private/DeferredShadingRenderer.cpp b/Engine/Source/Runtime/Renderer/Private/DeferredShadingRenderer.cpp
--- a/Engine/Source/Runtime/Renderer/Private/DeferredShadingRenderer.cpp
+++ b/Engine/Source/Runtime/Renderer/Private/DeferredShadingRenderer.cpp
@@ -2934,6 +2934,14 @@ void FDeferredShadingSceneRenderer::Render(FRDGBuilder& GraphBuilder)
 			EnumRemoveFlags(TranslucencyViewsToRender, ETranslucencyView::RayTracing);
 		}
 #endif
+		// Sort objects' triangles
+		for (FViewInfo& View : Views)
+		{
+			if (OIT::IsEnabled(View))
+			{
+				OIT::AddSortTrianglesPass(GraphBuilder, View, Scene->OITSceneData, FOITSortingType::BackToFront);
+			}
+		}
 
 		// Render all remaining translucency views.
 		GraphBuilder.SetCommandListStat(GET_STATID(STAT_CLM_Translucency));
diff --git a/Engine/Source/Runtime/Renderer/Private/MeshPassProcessor.cpp b/Engine/Source/Runtime/Renderer/Private/MeshPassProcessor.cpp
--- a/Engine/Source/Runtime/Renderer/Private/MeshPassProcessor.cpp
+++ b/Engine/Source/Runtime/Renderer/Private/MeshPassProcessor.cpp
@@ -772,6 +772,15 @@ void FMeshDrawCommand::SetDrawParametersAndFinalize(
 	NumPrimitives = BatchElement.NumPrimitives;
 	NumInstances = BatchElement.NumInstances;
 
+	// If the mesh batch has a valid dynamic index buffer, use it instead
+	if (BatchElement.DynamicIndexBuffer.IsValid())
+	{
+		check(!BatchElement.DynamicIndexBuffer.IndexBuffer || (BatchElement.DynamicIndexBuffer.IndexBuffer && BatchElement.DynamicIndexBuffer.IndexBuffer->IsInitialized() && BatchElement.DynamicIndexBuffer.IndexBuffer->IndexBufferRHI));
+		IndexBuffer = BatchElement.DynamicIndexBuffer.IndexBuffer ? BatchElement.DynamicIndexBuffer.IndexBuffer->IndexBufferRHI.GetReference() : nullptr;
+		FirstIndex = BatchElement.DynamicIndexBuffer.FirstIndex;
+		PrimitiveType = EPrimitiveType(BatchElement.DynamicIndexBuffer.PrimitiveType);
+	}
+
 	if (NumPrimitives > 0)
 	{
 		VertexParams.BaseVertexIndex = BatchElement.BaseVertexIndex;
diff --git a/Engine/Source/Runtime/Renderer/Private/OIT/OIT.cpp b/Engine/Source/Runtime/Renderer/Private/OIT/OIT.cpp
new file mode 100644
--- /dev/null
+++ b/Engine/Source/Runtime/Renderer/Private/OIT/OIT.cpp
@@ -0,0 +1,682 @@
+// Copyright Epic Games, Inc. All Rights Reserved.
+
+#include "OIT.h"
+#include "Shader.h"
+#include "GlobalShader.h"
+#include "ShaderParameters.h"
+#include "ShaderParameterStruct.h"
+#include "SceneTextureParameters.h"
+#include "SceneRendering.h"
+
+#include "ShaderPrintParameters.h"
+#include "ShaderDebug.h"
+#include "ShaderPrint.h"
+#include "ScreenPass.h"
+
+///////////////////////////////////////////////////////////////////////////////////////////////////
+// Variables
+
+static int32 GOIT_SortObjectTriangles = 1;
+static FAutoConsoleVariableRef CVarOIT_PerObjectSorting(TEXT("r.OIT.SortObjectTriangles"), GOIT_SortObjectTriangles, TEXT("Enable per-instance triangle sorting to avoid invalid triangle ordering (experimental)."));
+
+static int32 GOIT_Debug = 0;
+static FAutoConsoleVariableRef CVarOIT_Debug(TEXT("r.OIT.Debug"), GOIT_Debug, TEXT("Enable per-instance triangle sorting debug rendering."));
+
+static int32 GOIT_Pool = 0;
+static int32 GOIT_PoolReleaseThreshold = 100;
+static FAutoConsoleVariableRef CVarOIT_Pool(TEXT("r.OIT.Pool"), GOIT_Pool, TEXT("Enable index buffer pool allocation which reduce creation/deletion time by re-use buffers."));
+static FAutoConsoleVariableRef CVarOIT_PoolReleaseThreshold(TEXT("r.OIT.Pool.ReleaseFrameThreshold"), GOIT_PoolReleaseThreshold, TEXT("Number of frame after which unused buffer are released."));
+
+///////////////////////////////////////////////////////////////////////////////////////////////////
+// OIT Debug
+
+struct FOITDebugData
+{
+	static const EPixelFormat Format = PF_R32_UINT;
+
+	FRDGBufferRef Buffer = nullptr; // First element is counter, then element are: (NumPrim/Type/Size)
+
+	uint32 VisibleInstances = 0;
+	uint32 VisiblePrimitives = 0;
+	uint32 VisibleIndexSizeInBytes = 0;
+
+	uint32 AllocatedBuffers = 0;
+	uint32 AllocatedIndexSizeInBytes = 0;
+
+	uint32 UnusedBuffers = 0;
+	uint32 UnusedIndexSizeInBytes = 0;
+
+	uint32 TotalEntries = 0;
+
+	bool IsValid() const { return Buffer != nullptr; }
+};
+
+///////////////////////////////////////////////////////////////////////////////////////////////////
+// OIT Index buffer
+
+class FSortedIndexBuffer : public FIndexBuffer
+{
+public:
+	static const uint32 SliceCount = 32;
+
+	FSortedIndexBuffer(uint32 InId, const FBufferRHIRef& InSourceIndexBuffer, uint32 InNumIndices, const TCHAR* InDebugName)
+	: SourceIndexBuffer(InSourceIndexBuffer)
+	, NumIndices(InNumIndices)
+	, Id(InId)
+	, DebugName(InDebugName) { }
+
+	virtual void InitRHI() override
+	{
+		check(SourceIndexBuffer);
+		const uint32 BytesPerElement = SourceIndexBuffer->GetStride();
+		check(BytesPerElement == 2 || BytesPerElement == 4);
+		const EPixelFormat Format = BytesPerElement == 2 ? PF_R16_UINT : PF_R32_UINT;
+
+		FRHIResourceCreateInfo CreateInfo(DebugName);
+		IndexBufferRHI = RHICreateIndexBuffer(BytesPerElement /*Stride*/, NumIndices * BytesPerElement, BUF_UnorderedAccess | BUF_ShaderResource, ERHIAccess::VertexOrIndexBuffer, CreateInfo);
+		SortedIndexUAV = RHICreateUnorderedAccessView(IndexBufferRHI, Format);
+		SourceIndexSRV = RHICreateShaderResourceView(SourceIndexBuffer, BytesPerElement, Format);
+	}
+
+	virtual void ReleaseRHI() override
+	{
+		IndexBufferRHI.SafeRelease();
+		SortedIndexUAV.SafeRelease();
+		SourceIndexSRV.SafeRelease();
+	}
+
+	static constexpr uint32 InvalidId = ~0;
+
+	FBufferRHIRef SourceIndexBuffer = nullptr;
+	uint32 NumIndices = 0;
+	uint32 Id = FSortedIndexBuffer::InvalidId;
+	uint32 LastUsedFrameId = 0;
+	const TCHAR* DebugName = nullptr;
+
+	FShaderResourceViewRHIRef  SourceIndexSRV = nullptr;
+	FUnorderedAccessViewRHIRef SortedIndexUAV = nullptr;
+};
+
+static bool IsOITSupported(EShaderPlatform InShaderPlatform)
+{
+	return RHISupportsComputeShaders(InShaderPlatform);
+}
+
+static void RemoveAllocation(FSortedIndexBuffer* InBuffer)
+{
+	InBuffer->ReleaseResource();
+	delete InBuffer;
+	InBuffer = nullptr;
+}
+
+static void TrimSortedIndexBuffers(TArray<FSortedIndexBuffer*>& FreeBuffers, uint32 FrameId)
+{
+	uint32 FreeCount = FreeBuffers.Num();
+	for (uint32 FreeIt = 0; FreeIt < FreeCount;)
+	{
+		check(FreeBuffers[FreeIt]);
+		const uint32 LastFrameId = FreeBuffers[FreeIt]->LastUsedFrameId;
+		if (LastFrameId != 0)
+		{
+			const int32 ElapsedFrame = FMath::Abs(int32(FrameId) - int32(LastFrameId));
+			if (ElapsedFrame > GOIT_PoolReleaseThreshold)
+			{
+				FSortedIndexBuffer* Buffer = FreeBuffers[FreeIt];
+				RemoveAllocation(Buffer);
+				FreeBuffers[FreeIt] = FreeBuffers[FreeCount - 1];
+				--FreeCount;
+				FreeBuffers.SetNum(FreeCount);
+				continue;
+			}
+		}
+		++FreeIt;
+	}
+}
+
+FSortedTriangleData FOITSceneData::Allocate(const FIndexBuffer* InSource, EPrimitiveType PrimitiveType, uint32 InFirstIndex, uint32 InNumPrimitives)
+{
+	check(IsInRenderingThread());
+	check(InSource && InSource->IndexBufferRHI);
+	check(PrimitiveType == PT_TriangleList || PrimitiveType == PT_TriangleStrip);
+
+	// Find a free slot, or create a new one
+	FSortedTriangleData* Out = nullptr;
+	uint32 FreeSlot = FSortedIndexBuffer::InvalidId;
+	FreeSlots.Dequeue(FreeSlot);
+	if (FreeSlot != FSortedIndexBuffer::InvalidId)
+	{
+		Out = &Allocations[FreeSlot];
+	}
+	else
+	{
+		FreeSlot = Allocations.Num();
+		Out = &Allocations.AddDefaulted_GetRef();
+	}
+
+	// Linear scan if there are some free resource which are large enough
+	const uint32 NumIndices = InNumPrimitives * 3; // Sorted index always has triangle list topology
+	FSortedIndexBuffer* OITIndexBuffer = nullptr;
+	if (GOIT_Pool>0)
+	{
+		for (uint32 FreeIt=0,FreeCount=FreeBuffers.Num(); FreeIt<FreeCount; ++FreeIt)
+		{		
+			FSortedIndexBuffer* FreeBuffer = FreeBuffers[FreeIt];
+			if (FreeBuffer != nullptr && FreeBuffer->NumIndices >= NumIndices && OITIndexBuffer->Id == FSortedIndexBuffer::InvalidId)
+			{			
+				OITIndexBuffer = FreeBuffer;
+				OITIndexBuffer->Id = FreeSlot;
+				FreeBuffers[FreeIt] = FreeBuffers[FreeCount - 1];
+				FreeBuffers.SetNum(FreeCount - 1);
+				break;
+			}
+		}
+	}
+
+	// Otherwise create a new one
+	if (OITIndexBuffer == nullptr)
+	{
+		OITIndexBuffer = new FSortedIndexBuffer(FreeSlot, InSource->IndexBufferRHI, NumIndices, TEXT("OIT::SortedIndexBuffer"));
+		OITIndexBuffer->InitResource();	
+	}
+	Out->NumPrimitives = InNumPrimitives;
+	Out->NumIndices = NumIndices;
+	Out->SourceFirstIndex = InFirstIndex;
+	Out->SortedFirstIndex = 0u;
+	Out->SourcePrimitiveType = PrimitiveType;
+	Out->SortedPrimitiveType = PT_TriangleList;
+	Out->SourceIndexBuffer = InSource;
+	Out->SortedIndexBuffer = OITIndexBuffer;
+	Out->SortedIndexUAV = OITIndexBuffer->SortedIndexUAV;
+	Out->SourceIndexSRV = OITIndexBuffer->SourceIndexSRV;
+
+	return *Out;
+}
+
+void FOITSceneData::Deallocate(FIndexBuffer* InIndexBuffer)
+{
+	if (InIndexBuffer == nullptr)
+	{
+		return;
+	}
+
+	FSortedIndexBuffer* OITIndexBuffer = (FSortedIndexBuffer*)InIndexBuffer;
+	const uint32 Slot = OITIndexBuffer->Id;
+	if (Slot < uint32(Allocations.Num()))
+	{
+		if (GOIT_Pool > 0)
+		{
+			OITIndexBuffer->Id = FSortedIndexBuffer::InvalidId;
+			OITIndexBuffer->LastUsedFrameId = FrameIndex;
+			FreeBuffers.Add(OITIndexBuffer);
+			Allocations[Slot] = FSortedTriangleData();
+		}
+		else
+		{
+			FSortedTriangleData& In = Allocations[Slot];
+			In.SortedIndexUAV = nullptr;
+			In.SourceIndexSRV = nullptr;
+			RemoveAllocation((FSortedIndexBuffer*)In.SortedIndexBuffer);
+			In = FSortedTriangleData();
+		}
+		FreeSlots.Enqueue(Slot);
+	}
+}
+
+static uint32 GetGroupSize()
+{
+	if (IsRHIDeviceNVIDIA())
+	{
+		return 32u;
+	}
+	else  // IsRHIDeviceAMD() and others
+	{
+		return 64u;
+	}
+}
+
+///////////////////////////////////////////////////////////////////////////////////////////////////
+// Sort triangle indices to order them front-to-back or back-to-front
+
+class FOITSortTriangleIndex_ScanCS : public FGlobalShader
+{
+	DECLARE_GLOBAL_SHADER(FOITSortTriangleIndex_ScanCS);
+	SHADER_USE_PARAMETER_STRUCT(FOITSortTriangleIndex_ScanCS, FGlobalShader);
+
+	class FGroupSize : SHADER_PERMUTATION_SPARSE_INT("PERMUTATION_GROUP_SIZE", 32, 64); 
+	class FDebug : SHADER_PERMUTATION_BOOL("PERMUTATION_DEBUG");
+	using FPermutationDomain = TShaderPermutationDomain<FGroupSize, FDebug>;
+
+	BEGIN_SHADER_PARAMETER_STRUCT(FParameters, )
+
+		// For Debug
+		SHADER_PARAMETER_STRUCT_INCLUDE(ShaderDrawDebug::FShaderParameters, ShaderDrawParameters)
+		SHADER_PARAMETER_STRUCT_INCLUDE(ShaderPrint::FShaderParameters, ShaderPrintParameters)
+		SHADER_PARAMETER(FMatrix44f, ViewToWorld)
+		SHADER_PARAMETER(FVector3f, WorldBound_Min)
+		SHADER_PARAMETER(FVector3f, WorldBound_Max)
+		SHADER_PARAMETER(FVector3f, ViewBound_Min)
+		SHADER_PARAMETER(FVector3f, ViewBound_Max)
+
+		SHADER_PARAMETER(FMatrix44f, LocalToWorld)
+		SHADER_PARAMETER(FMatrix44f, WorldToView)
+
+		SHADER_PARAMETER(uint32, SourcePrimitiveType)
+		SHADER_PARAMETER(uint32, NumPrimitives)
+		SHADER_PARAMETER(uint32, NumIndices)
+		SHADER_PARAMETER(uint32, SourceFirstIndex)
+		SHADER_PARAMETER(uint32, SortType)
+		SHADER_PARAMETER(uint32, SortedIndexBufferSizeInByte)
+		
+		SHADER_PARAMETER(float, ViewBoundMinZ)
+		SHADER_PARAMETER(float, ViewBoundMaxZ)
+
+		SHADER_PARAMETER_SRV(Buffer<float>, PositionBuffer)
+		SHADER_PARAMETER_SRV(Buffer<uint>, IndexBuffer)
+		SHADER_PARAMETER_UAV(Buffer<uint>, OutIndexBuffer)
+		SHADER_PARAMETER_RDG_BUFFER_UAV(RWBuffer<uint2>, OutSliceCounterBuffer)
+		SHADER_PARAMETER_RDG_BUFFER_UAV(RWBuffer<uint>, OutPrimitiveSliceBuffer)
+		SHADER_PARAMETER_RDG_BUFFER_UAV(RWBuffer<uint>, OutDebugData)
+
+
+	END_SHADER_PARAMETER_STRUCT()
+public:
+	static bool ShouldCompilePermutation(const FGlobalShaderPermutationParameters& Parameters) { return IsOITSupported(Parameters.Platform); }
+	static void ModifyCompilationEnvironment(const FGlobalShaderPermutationParameters& Parameters, FShaderCompilerEnvironment& OutEnvironment)
+	{
+		FGlobalShader::ModifyCompilationEnvironment(Parameters, OutEnvironment);
+		OutEnvironment.SetDefine(TEXT("SHADER_SCAN"), 1);
+		OutEnvironment.SetDefine(TEXT("SORTING_SLICE_COUNT"), FSortedIndexBuffer::SliceCount);
+	}
+};
+
+IMPLEMENT_GLOBAL_SHADER(FOITSortTriangleIndex_ScanCS, "/Engine/Private/OIT/OITSorting.usf", "MainCS", SF_Compute);
+
+///////////////////////////////////////////////////////////////////////////////////////////////////
+
+class FOITSortTriangleIndex_AllocateCS : public FGlobalShader
+{
+	DECLARE_GLOBAL_SHADER(FOITSortTriangleIndex_AllocateCS);
+	SHADER_USE_PARAMETER_STRUCT(FOITSortTriangleIndex_AllocateCS, FGlobalShader);
+
+	BEGIN_SHADER_PARAMETER_STRUCT(FParameters, )
+		SHADER_PARAMETER_RDG_BUFFER_SRV(Buffer<uint>, SliceCounterBuffer)
+		SHADER_PARAMETER_RDG_BUFFER_UAV(RWBuffer<uint>, SliceOffsetsBuffer)
+	END_SHADER_PARAMETER_STRUCT()
+public:
+	static bool ShouldCompilePermutation(const FGlobalShaderPermutationParameters& Parameters) { return IsOITSupported(Parameters.Platform); }
+	static void ModifyCompilationEnvironment(const FGlobalShaderPermutationParameters& Parameters, FShaderCompilerEnvironment& OutEnvironment)
+	{
+		FGlobalShader::ModifyCompilationEnvironment(Parameters, OutEnvironment);
+		OutEnvironment.SetDefine(TEXT("SHADER_ALLOCATE"), 1);
+		OutEnvironment.SetDefine(TEXT("SORTING_SLICE_COUNT"), FSortedIndexBuffer::SliceCount);
+	}
+};
+
+IMPLEMENT_GLOBAL_SHADER(FOITSortTriangleIndex_AllocateCS, "/Engine/Private/OIT/OITSorting.usf", "MainCS", SF_Compute);
+
+///////////////////////////////////////////////////////////////////////////////////////////////////
+
+class FOITSortTriangleIndex_WriteOutCS : public FGlobalShader
+{
+	DECLARE_GLOBAL_SHADER(FOITSortTriangleIndex_WriteOutCS);
+	SHADER_USE_PARAMETER_STRUCT(FOITSortTriangleIndex_WriteOutCS, FGlobalShader);
+
+	BEGIN_SHADER_PARAMETER_STRUCT(FParameters, )
+		SHADER_PARAMETER(uint32, SourcePrimitiveType)
+		SHADER_PARAMETER(uint32, NumPrimitives)
+		SHADER_PARAMETER(uint32, NumIndices)
+		SHADER_PARAMETER(uint32, SrcFirstIndex)
+		SHADER_PARAMETER(uint32, DstFirstIndex)
+
+		SHADER_PARAMETER_RDG_BUFFER_SRV(Buffer<uint>, SliceOffsetsBuffer)
+		SHADER_PARAMETER_RDG_BUFFER_SRV(Buffer<uint>, PrimitiveSliceBuffer)
+
+		SHADER_PARAMETER_SRV(Buffer<uint>, IndexBuffer)
+		SHADER_PARAMETER_UAV(RWBuffer<uint>, OutIndexBuffer)
+	END_SHADER_PARAMETER_STRUCT()
+public:
+	static bool ShouldCompilePermutation(const FGlobalShaderPermutationParameters& Parameters) { return IsOITSupported(Parameters.Platform); }
+	static void ModifyCompilationEnvironment(const FGlobalShaderPermutationParameters& Parameters, FShaderCompilerEnvironment& OutEnvironment)
+	{
+		FGlobalShader::ModifyCompilationEnvironment(Parameters, OutEnvironment);
+		OutEnvironment.SetDefine(TEXT("SHADER_WRITE"), 1);
+		OutEnvironment.SetDefine(TEXT("SORTING_SLICE_COUNT"), FSortedIndexBuffer::SliceCount);
+	}
+};
+
+IMPLEMENT_GLOBAL_SHADER(FOITSortTriangleIndex_WriteOutCS, "/Engine/Private/OIT/OITSorting.usf", "MainCS", SF_Compute);
+
+///////////////////////////////////////////////////////////////////////////////////////////////////
+
+class FOITSortTriangleIndex_Debug : public FGlobalShader
+{
+	DECLARE_GLOBAL_SHADER(FOITSortTriangleIndex_Debug);
+	SHADER_USE_PARAMETER_STRUCT(FOITSortTriangleIndex_Debug, FGlobalShader);
+
+	BEGIN_SHADER_PARAMETER_STRUCT(FParameters, )
+		SHADER_PARAMETER_STRUCT_INCLUDE(ShaderDrawDebug::FShaderParameters, ShaderDrawParameters)
+		SHADER_PARAMETER_STRUCT_INCLUDE(ShaderPrint::FShaderParameters, ShaderPrintParameters)
+		SHADER_PARAMETER(uint32, VisibleInstances)
+		SHADER_PARAMETER(uint32, VisiblePrimitives)
+		SHADER_PARAMETER(uint32, VisibleIndexSizeInBytes)
+		SHADER_PARAMETER(uint32, AllocatedBuffers)
+		SHADER_PARAMETER(uint32, AllocatedIndexSizeInBytes)
+		SHADER_PARAMETER(uint32, UnusedBuffers)
+		SHADER_PARAMETER(uint32, UnusedIndexSizeInBytes)
+		SHADER_PARAMETER(uint32, TotalEntries)
+		SHADER_PARAMETER_RDG_BUFFER_SRV(Buffer<uint>, DebugData)
+	END_SHADER_PARAMETER_STRUCT()
+public:
+	static bool ShouldCompilePermutation(const FGlobalShaderPermutationParameters& Parameters) { return IsOITSupported(Parameters.Platform); }
+	static void ModifyCompilationEnvironment(const FGlobalShaderPermutationParameters& Parameters, FShaderCompilerEnvironment& OutEnvironment)
+	{
+		FGlobalShader::ModifyCompilationEnvironment(Parameters, OutEnvironment);
+		OutEnvironment.SetDefine(TEXT("SHADER_DEBUG"), 1);
+		OutEnvironment.SetDefine(TEXT("SORTING_SLICE_COUNT"), FSortedIndexBuffer::SliceCount);
+	}
+};
+
+IMPLEMENT_GLOBAL_SHADER(FOITSortTriangleIndex_Debug, "/Engine/Private/OIT/OITSorting.usf", "MainCS", SF_Compute);
+
+///////////////////////////////////////////////////////////////////////////////////////////////////
+
+static void AddOITSortTriangleIndexPass(
+	FRDGBuilder& GraphBuilder,
+	const FViewInfo& View,
+	const FOITSceneData& OITSceneData,
+	const FSortedTrianglesMeshBatch& MeshBatch,
+	FOITSortingType SortType,
+	FOITDebugData& DebugData)
+{
+	static const FVertexFactoryType* CompatibleVF = FVertexFactoryType::GetVFByName(TEXT("FLocalVertexFactory"));
+	const bool bIsCompatible = MeshBatch.Mesh->VertexFactory->GetType()->GetHashedName() == CompatibleVF->GetHashedName();
+
+	// Fat format: PF_R32G32_UINT | Compact format: PF_R32_UINT
+	const EPixelFormat PackedFormat = PF_R32_UINT;
+	const uint32 PackedFormatInBytes = 4;
+
+	const bool bIsValid = 
+		bIsCompatible && 
+		MeshBatch.Mesh != nullptr && 
+		MeshBatch.Mesh->Elements.Num() > 0 && 
+		MeshBatch.Mesh->Elements[0].DynamicIndexBuffer.IndexBuffer != nullptr;
+	if (!bIsValid)
+	{
+		return;
+	}
+
+	const FLocalVertexFactory* VF = (const FLocalVertexFactory*)MeshBatch.Mesh->VertexFactory;
+	if (!VF) { return; }
+	const FShaderResourceViewRHIRef VertexPosition = VF->GetPositionsSRV();
+	if (!VertexPosition) { return; }
+
+	const FSortedIndexBuffer* OITIndexBuffer = (const FSortedIndexBuffer*)MeshBatch.Mesh->Elements[0].DynamicIndexBuffer.IndexBuffer;
+	check(OITIndexBuffer->Id < uint32(OITSceneData.Allocations.Num()));
+	const FSortedTriangleData& Allocation = OITSceneData.Allocations[OITIndexBuffer->Id];
+	
+	check(Allocation.IsValid());
+	check(Allocation.SourcePrimitiveType == PT_TriangleList || Allocation.SourcePrimitiveType == PT_TriangleStrip);
+	check(Allocation.SortedPrimitiveType == PT_TriangleList);
+
+	FRDGBufferRef PrimitiveSliceBuffer = GraphBuilder.CreateBuffer(FRDGBufferDesc::CreateBufferDesc(PackedFormatInBytes, Allocation.NumPrimitives), TEXT("OIT.TriangleSortingSliceIndex"));
+	FRDGBufferRef SliceCounterBuffer = GraphBuilder.CreateBuffer(FRDGBufferDesc::CreateBufferDesc(4, FSortedIndexBuffer::SliceCount), TEXT("OIT.SliceCounters"));
+	FRDGBufferRef SliceOffsetsBuffer = GraphBuilder.CreateBuffer(FRDGBufferDesc::CreateBufferDesc(4, FSortedIndexBuffer::SliceCount), TEXT("OIT.SliceOffsets"));
+	FRDGBufferUAVRef SliceCounterUAV = GraphBuilder.CreateUAV(SliceCounterBuffer, PF_R32_UINT);
+	AddClearUAVPass(GraphBuilder, SliceCounterUAV, 0u);
+
+	const bool bDebugEnable = DebugData.IsValid();
+	FRHIBuffer* SortedIndexBufferRHI = Allocation.SortedIndexBuffer->IndexBufferRHI;
+
+	// 1. Scan the primitive and assign each primitive to a slice
+	{
+		// Compute the primitive Min/Max-Z value in view space. This domain is sliced for sorting
+		const FBoxSphereBounds& Bounds = MeshBatch.Proxy->GetBounds();
+		const FBoxSphereBounds& ViewBounds = Bounds.TransformBy(View.ViewMatrices.GetViewMatrix());
+		const float ViewBoundMinZ = ViewBounds.GetBox().Min.Z;
+		const float ViewBoundMaxZ = ViewBounds.GetBox().Max.Z;
+
+		FOITSortTriangleIndex_ScanCS::FParameters* Parameters = GraphBuilder.AllocParameters<FOITSortTriangleIndex_ScanCS::FParameters>();
+		Parameters->LocalToWorld			= MeshBatch.Proxy->GetLocalToWorld();
+		Parameters->WorldToView				= View.ViewMatrices.GetViewMatrix();
+		Parameters->SourcePrimitiveType		= Allocation.SourcePrimitiveType == PT_TriangleStrip ? 1u : 0u;
+		Parameters->NumPrimitives			= Allocation.NumPrimitives;
+		Parameters->NumIndices				= Allocation.NumIndices;
+		Parameters->ViewBoundMinZ			= ViewBoundMinZ;
+		Parameters->ViewBoundMaxZ			= ViewBoundMaxZ;
+		Parameters->SortType				= SortType == FOITSortingType::BackToFront ? 0 : 1;
+		Parameters->SortedIndexBufferSizeInByte = Allocation.SortedIndexBuffer->IndexBufferRHI->GetSize();
+		Parameters->PositionBuffer			= VertexPosition;
+		Parameters->SourceFirstIndex		= Allocation.SourceFirstIndex;
+		Parameters->IndexBuffer				= Allocation.SourceIndexSRV;
+		Parameters->OutIndexBuffer			= Allocation.SortedIndexUAV;
+		Parameters->OutSliceCounterBuffer	= SliceCounterUAV;
+		Parameters->OutPrimitiveSliceBuffer = GraphBuilder.CreateUAV(PrimitiveSliceBuffer, PackedFormat);
+
+		// Debug
+		if (bDebugEnable)
+		{
+			Parameters->ViewToWorld		= View.ViewMatrices.GetViewMatrix().Inverse();
+			Parameters->WorldBound_Min	= Bounds.GetBox().Min;
+			Parameters->WorldBound_Max	= Bounds.GetBox().Max;
+			Parameters->ViewBound_Min	= ViewBounds.GetBox().Min;
+			Parameters->ViewBound_Max	= ViewBounds.GetBox().Max;
+			if (ShaderDrawDebug::IsEnabled(View))
+			{
+				ShaderDrawDebug::SetParameters(GraphBuilder, View.ShaderDrawData, Parameters->ShaderDrawParameters);
+			}
+			if (ShaderPrint::IsEnabled(View))
+			{
+				ShaderPrint::SetParameters(GraphBuilder, View, Parameters->ShaderPrintParameters);
+			}
+
+			check(DebugData.Buffer);
+
+			++DebugData.VisibleInstances;
+			DebugData.VisiblePrimitives += Parameters->NumPrimitives;
+			DebugData.VisibleIndexSizeInBytes += Allocation.SortedIndexBuffer->IndexBufferRHI->GetSize();
+			Parameters->OutDebugData = GraphBuilder.CreateUAV(DebugData.Buffer, PF_R32_UINT);
+		}
+
+		const uint32 GroupSize = GetGroupSize();
+
+		FOITSortTriangleIndex_ScanCS::FPermutationDomain PermutationVector;
+		PermutationVector.Set<FOITSortTriangleIndex_ScanCS::FGroupSize>(GroupSize);
+		PermutationVector.Set<FOITSortTriangleIndex_ScanCS::FDebug>(bDebugEnable ? 1 : 0);
+		TShaderMapRef<FOITSortTriangleIndex_ScanCS> ComputeShader(View.ShaderMap, PermutationVector);
+		
+		const FIntVector DispatchCount = FIntVector(FMath::CeilToInt(float(Parameters->NumPrimitives) / float(GroupSize)), 1u, 1u);
+		check(DispatchCount.X < GRHIMaxDispatchThreadGroupsPerDimension.X);
+		ClearUnusedGraphResources(ComputeShader, Parameters);
+		GraphBuilder.AddPass(
+			RDG_EVENT_NAME("OIT::SortTriangleIndices(Scan)"),
+			Parameters,
+			ERDGPassFlags::Compute,
+			[Parameters, ComputeShader, DispatchCount, SortedIndexBufferRHI](FRHIComputeCommandList& RHICmdList)
+			{
+				RHICmdList.Transition(FRHITransitionInfo(SortedIndexBufferRHI, ERHIAccess::VertexOrIndexBuffer, ERHIAccess::UAVCompute));
+				FComputeShaderUtils::Dispatch(RHICmdList, ComputeShader, *Parameters, DispatchCount);
+			});
+	}
+
+	// 2. Pre-fix sum onto the slices count to allocate each bucket
+	{
+		FOITSortTriangleIndex_AllocateCS::FParameters* Parameters = GraphBuilder.AllocParameters<FOITSortTriangleIndex_AllocateCS::FParameters>();
+		Parameters->SliceCounterBuffer = GraphBuilder.CreateSRV(SliceCounterBuffer, PF_R32_UINT);
+		Parameters->SliceOffsetsBuffer = GraphBuilder.CreateUAV(SliceOffsetsBuffer, PF_R32_UINT);
+		TShaderMapRef<FOITSortTriangleIndex_AllocateCS> ComputeShader(View.ShaderMap);
+		FComputeShaderUtils::AddPass(
+			GraphBuilder,
+			RDG_EVENT_NAME("OIT::SortTriangleIndices(PrefixedSum)"),
+			ComputeShader,
+			Parameters,
+			FIntVector(1u, 1u, 1u));
+	}
+
+	// 3. Write-out the sorted indices
+	{
+		FOITSortTriangleIndex_WriteOutCS::FParameters* Parameters = GraphBuilder.AllocParameters<FOITSortTriangleIndex_WriteOutCS::FParameters>();
+		Parameters->SourcePrimitiveType		= Allocation.SourcePrimitiveType == PT_TriangleStrip ? 1u : 0u;
+		Parameters->NumPrimitives			= Allocation.NumPrimitives;
+		Parameters->NumIndices				= Allocation.NumIndices;
+		Parameters->SrcFirstIndex			= Allocation.SourceFirstIndex;
+		Parameters->DstFirstIndex			= 0u;
+
+		Parameters->SliceOffsetsBuffer		= GraphBuilder.CreateSRV(SliceOffsetsBuffer, PF_R32_UINT);
+		Parameters->PrimitiveSliceBuffer	= GraphBuilder.CreateSRV(PrimitiveSliceBuffer, PackedFormat);
+
+		Parameters->IndexBuffer				= Allocation.SourceIndexSRV;
+		Parameters->OutIndexBuffer			= Allocation.SortedIndexUAV;
+
+		TShaderMapRef<FOITSortTriangleIndex_WriteOutCS> ComputeShader(View.ShaderMap);
+
+		const uint32 GroupSize = 256;
+		const FIntVector DispatchCount = FIntVector(GroupSize, 1u, 1u);
+		ClearUnusedGraphResources(ComputeShader, Parameters);
+		GraphBuilder.AddPass(
+			RDG_EVENT_NAME("OIT::SortTriangleIndices(Write)"),
+			Parameters,
+			ERDGPassFlags::Compute,
+			[Parameters, ComputeShader, DispatchCount, SortedIndexBufferRHI](FRHIComputeCommandList& RHICmdList)
+			{
+				FComputeShaderUtils::Dispatch(RHICmdList, ComputeShader, *Parameters, DispatchCount);
+
+				RHICmdList.Transition(FRHITransitionInfo(SortedIndexBufferRHI, ERHIAccess::UAVCompute, ERHIAccess::VertexOrIndexBuffer));
+			});
+	}
+
+	// Next todos
+	// * Merge several meshes together (not clear out to do the mapping thread->mesh info)
+	// * Batch Scan/Alloc/Write of several primitive, so that we have better overlapping
+}
+
+static void AddOITDebugPass(
+	FRDGBuilder& GraphBuilder,
+	const FViewInfo& View,
+	const FOITDebugData& DebugData)
+{
+	if (!DebugData.IsValid())
+	{
+		return;
+	}
+
+	FOITSortTriangleIndex_Debug::FParameters* Parameters = GraphBuilder.AllocParameters<FOITSortTriangleIndex_Debug::FParameters>();
+	Parameters->VisibleInstances = DebugData.VisibleInstances;
+	Parameters->VisiblePrimitives = DebugData.VisiblePrimitives;
+	Parameters->VisibleIndexSizeInBytes = DebugData.VisibleIndexSizeInBytes;
+
+	Parameters->UnusedBuffers = DebugData.UnusedBuffers;
+	Parameters->UnusedIndexSizeInBytes = DebugData.UnusedIndexSizeInBytes;
+
+	Parameters->AllocatedBuffers = DebugData.AllocatedBuffers;
+	Parameters->AllocatedIndexSizeInBytes = DebugData.AllocatedIndexSizeInBytes;
+
+	Parameters->TotalEntries = DebugData.TotalEntries;
+
+	Parameters->DebugData = GraphBuilder.CreateSRV(DebugData.Buffer, FOITDebugData::Format);
+	if (ShaderDrawDebug::IsEnabled(View))
+	{
+		ShaderDrawDebug::SetParameters(GraphBuilder, View.ShaderDrawData, Parameters->ShaderDrawParameters);
+	}
+	if (ShaderPrint::IsEnabled(View))
+	{
+		ShaderPrint::SetParameters(GraphBuilder, View, Parameters->ShaderPrintParameters);
+	}
+
+	TShaderMapRef<FOITSortTriangleIndex_Debug> ComputeShader(View.ShaderMap);
+	FComputeShaderUtils::AddPass(
+		GraphBuilder,
+		RDG_EVENT_NAME("OIT::Debug"),
+		ComputeShader,
+		Parameters,
+		FIntVector(1u, 1u, 1u));
+}
+
+///////////////////////////////////////////////////////////////////////////////////////////////////
+
+namespace OIT
+{
+	bool IsEnabled(const FViewInfo& View)
+	{
+		return GOIT_SortObjectTriangles > 0 && IsOITSupported(View.GetShaderPlatform());
+	}
+
+	bool IsEnabled(EShaderPlatform ShaderPlatform)
+	{
+		return GOIT_SortObjectTriangles > 0 && IsOITSupported(ShaderPlatform);
+	}
+
+	bool IsCompatible(const FMeshBatch& InMesh, ERHIFeatureLevel::Type InFeatureLevel)
+	{
+		// Only support local vertex factory at the moment as we need to have direct access to the position
+		static const FVertexFactoryType* CompatibleVF = FVertexFactoryType::GetVFByName(TEXT("FLocalVertexFactory"));
+
+		if (InMesh.IsTranslucent(InFeatureLevel))
+		{
+			return InMesh.VertexFactory->GetType()->GetHashedName() == CompatibleVF->GetHashedName();
+		}
+		return false;
+	}
+
+	void AddSortTrianglesPass(FRDGBuilder& GraphBuilder, const FViewInfo& View, FOITSceneData& OITSceneData, FOITSortingType SortType)
+	{
+		if (!IsEnabled(View))
+		{
+			return;
+		}
+
+		RDG_EVENT_SCOPE(GraphBuilder, "OIT::IndexSorting");
+
+		const bool bDebugEnable = GOIT_Debug > 0;
+		FOITDebugData DebugData;
+		if (bDebugEnable)
+		{
+			const uint32 ValidAllocationCount = OITSceneData.Allocations.Num();
+			DebugData.Buffer = GraphBuilder.CreateBuffer(FRDGBufferDesc::CreateBufferDesc(4u, ValidAllocationCount + 1), TEXT("OIT.DebugData"));
+			AddClearUAVPass(GraphBuilder, GraphBuilder.CreateUAV(DebugData.Buffer, FOITDebugData::Format), 0u);
+
+			// Allocated/used
+			for (const FSortedTriangleData& Allocated : OITSceneData.Allocations)
+			{
+				if (Allocated.IsValid())
+				{
+					++DebugData.AllocatedBuffers;
+					DebugData.AllocatedIndexSizeInBytes += Allocated.SortedIndexBuffer->IndexBufferRHI->GetSize();
+				}
+			}
+
+			// Unused
+			for (const FSortedIndexBuffer* FreeBuffer : OITSceneData.FreeBuffers)
+			{
+				++DebugData.UnusedBuffers;
+				DebugData.UnusedIndexSizeInBytes = FreeBuffer->IndexBufferRHI->GetSize();
+			}
+
+			DebugData.TotalEntries = OITSceneData.Allocations.Num();
+		}
+
+		for (const FSortedTrianglesMeshBatch& MeshBatch : View.SortedTrianglesMeshBatches)		
+		{
+			AddOITSortTriangleIndexPass(GraphBuilder, View, OITSceneData, MeshBatch, SortType, DebugData);
+		}
+
+		if (DebugData.IsValid())
+		{
+			AddOITDebugPass(GraphBuilder, View, DebugData);
+		}
+
+		// Trim unused buffers
+		OITSceneData.FrameIndex = View.Family->FrameNumber;
+		if (GOIT_Pool > 0 && GOIT_PoolReleaseThreshold > 0)
+		{
+			TrimSortedIndexBuffers(OITSceneData.FreeBuffers, OITSceneData.FrameIndex);
+		}
+	}
+
+	void ConvertSortedIndexToDynamicIndex(FSortedTriangleData* In, FMeshBatchElementDynamicIndexBuffer* Out)
+	{
+		check(In && Out);
+		Out->IndexBuffer = In->SortedIndexBuffer;
+		Out->FirstIndex = In->SortedFirstIndex;
+		Out->PrimitiveType = In->SortedPrimitiveType;
+	}
+}
\ No newline at end of file
diff --git a/Engine/Source/Runtime/Renderer/Private/OIT/OIT.h b/Engine/Source/Runtime/Renderer/Private/OIT/OIT.h
new file mode 100644
--- /dev/null
+++ b/Engine/Source/Runtime/Renderer/Private/OIT/OIT.h
@@ -0,0 +1,67 @@
+// Copyright Epic Games, Inc. All Rights Reserved.
+
+#pragma once
+
+#include "CoreMinimal.h"
+#include "RendererInterface.h"
+#include "RenderGraphResources.h"
+#include "RenderResource.h"
+
+class FSortedIndexBuffer;
+struct FMeshBatchElementDynamicIndexBuffer;
+
+struct FSortedTriangleData
+{
+	const FIndexBuffer* SourceIndexBuffer = nullptr;
+	FIndexBuffer* SortedIndexBuffer = nullptr;
+
+	FShaderResourceViewRHIRef  SourceIndexSRV = nullptr;
+	FUnorderedAccessViewRHIRef SortedIndexUAV = nullptr;
+
+	uint32 SortedFirstIndex = 0;
+	uint32 SourceFirstIndex	= 0;
+	uint32 NumPrimitives = 0;
+	uint32 NumIndices = 0;
+
+	EPrimitiveType SourcePrimitiveType = PT_TriangleList;
+	EPrimitiveType SortedPrimitiveType = PT_TriangleList;
+
+	bool IsValid() const { return SortedIndexBuffer != nullptr; }
+};
+
+struct FOITSceneData
+{
+	/* Allocate sorted-triangle data for a instance */
+	FSortedTriangleData Allocate(const FIndexBuffer* InSource, EPrimitiveType PrimitiveType, uint32 InFirstIndex, uint32 InNumPrimitives);
+
+	/* Deallocate sorted-triangle data */
+	void Deallocate(FIndexBuffer* IndexBuffer);
+
+	TArray<FSortedTriangleData> Allocations;
+	TArray<FSortedIndexBuffer*> FreeBuffers;
+	TQueue<uint32> FreeSlots;
+	uint32 FrameIndex = 0;
+};
+
+enum class FOITSortingType
+{
+	FrontToBack,
+	BackToFront,
+};
+
+namespace OIT
+{
+	/* Return true if OIT techniques are enabled/supported */
+	bool IsEnabled(const FViewInfo& View);
+	bool IsEnabled(EShaderPlatform ShaderPlatform);
+
+	/* Return true if the current MeshBatch is compatible with per-instance sorted triangle */
+	bool IsCompatible(const FMeshBatch& Mesh, ERHIFeatureLevel::Type InFeatureLevel);
+
+	/* Sort triangles of all instances whose has the sorted triangle option enabled */
+	void AddSortTrianglesPass(FRDGBuilder& GraphBuilder, const FViewInfo& View, FOITSceneData& OITSceneData, FOITSortingType SortType);
+
+	/* Convert FSortedTriangleData into FMeshBatchElementDynamicIndexBuffer */
+	void ConvertSortedIndexToDynamicIndex(FSortedTriangleData* In, FMeshBatchElementDynamicIndexBuffer* Out);
+}
+ 
diff --git a/Engine/Source/Runtime/Renderer/Private/PrimitiveSceneInfo.cpp b/Engine/Source/Runtime/Renderer/Private/PrimitiveSceneInfo.cpp
--- a/Engine/Source/Runtime/Renderer/Private/PrimitiveSceneInfo.cpp
+++ b/Engine/Source/Runtime/Renderer/Private/PrimitiveSceneInfo.cpp
@@ -912,9 +912,14 @@ void FPrimitiveSceneInfo::AddStaticMeshes(FRHICommandListImmediate& RHICmdList,
 	}
 
 	{
+		const ERHIFeatureLevel::Type FeatureLevel = Scene->GetFeatureLevel();
+
 		SCOPED_NAMED_EVENT(FPrimitiveSceneInfo_AddStaticMeshes_UpdateSceneArrays, FColor::Blue);
 		for (FPrimitiveSceneInfo* SceneInfo : SceneInfos)
 		{
+			// Allocate OIT index buffer where needed
+			const bool bAllocateSortedTriangles = OIT::IsEnabled(GMaxRHIShaderPlatform) && SceneInfo->Proxy->SupportsSortedTriangles();
+
 			for (int32 MeshIndex = 0; MeshIndex < SceneInfo->StaticMeshes.Num(); MeshIndex++)
 			{
 				FStaticMeshBatchRelevance& MeshRelevance = SceneInfo->StaticMeshRelevances[MeshIndex];
@@ -925,6 +930,12 @@ void FPrimitiveSceneInfo::AddStaticMeshes(FRHICommandListImmediate& RHICmdList,
 				Scene->StaticMeshes[SceneArrayAllocation.Index] = &Mesh;
 				Mesh.Id = SceneArrayAllocation.Index;
 				MeshRelevance.Id = SceneArrayAllocation.Index;
+
+				if (bAllocateSortedTriangles && OIT::IsCompatible(Mesh, FeatureLevel))
+				{
+					FSortedTriangleData Allocation = Scene->OITSceneData.Allocate(Mesh.Elements[0].IndexBuffer, EPrimitiveType(Mesh.Type), Mesh.Elements[0].FirstIndex, Mesh.Elements[0].NumPrimitives);
+					OIT::ConvertSortedIndexToDynamicIndex(&Allocation, &Mesh.Elements[0].DynamicIndexBuffer);
+				}
 			}
 		}
 	}
@@ -1339,6 +1350,19 @@ void FPrimitiveSceneInfo::AddToScene(FRHICommandListImmediate& RHICmdList, FScen
 
 void FPrimitiveSceneInfo::RemoveStaticMeshes()
 {
+	// Deallocate potential OIT dynamic index buffer
+	if (OIT::IsEnabled(GMaxRHIShaderPlatform))
+	{
+		for (int32 MeshIndex = 0; MeshIndex < StaticMeshes.Num(); MeshIndex++)
+		{
+			FStaticMeshBatch& Mesh = StaticMeshes[MeshIndex];
+			if (Mesh.Elements.Num() > 0 && Mesh.Elements[0].DynamicIndexBuffer.IsValid())
+			{
+				Scene->OITSceneData.Deallocate(Mesh.Elements[0].DynamicIndexBuffer.IndexBuffer);
+			}
+		}
+	}
+
 	// Remove static meshes from the scene.
 	StaticMeshes.Empty();
 	StaticMeshRelevances.Empty();
diff --git a/Engine/Source/Runtime/Renderer/Private/ScenePrivate.h b/Engine/Source/Runtime/Renderer/Private/ScenePrivate.h
--- a/Engine/Source/Runtime/Renderer/Private/ScenePrivate.h
+++ b/Engine/Source/Runtime/Renderer/Private/ScenePrivate.h
@@ -58,6 +58,7 @@
 #include "VolumetricRenderTargetViewStateData.h"
 #include "GPUScene.h"
 #include "DynamicBVH.h"
+#include "OIT/OIT.h"
 
 /** Factor by which to grow occlusion tests **/
 #define OCCLUSION_SLOP (1.0f)
@@ -2786,6 +2787,9 @@ public:
 	/** The hair strands in the scene. */
 	FHairStrandsSceneData HairStrandsSceneData;
 
+	/** The OIT resources in the scene. */
+	FOITSceneData OITSceneData;
+
 	/** 
 	 * Precomputed lighting volumes in the scene, used for interpolating dynamic object lighting.
 	 * These are typically one per streaming level and they store volume lighting samples computed by Lightmass. 
diff --git a/Engine/Source/Runtime/Renderer/Private/SceneRendering.h b/Engine/Source/Runtime/Renderer/Private/SceneRendering.h
--- a/Engine/Source/Runtime/Renderer/Private/SceneRendering.h
+++ b/Engine/Source/Runtime/Renderer/Private/SceneRendering.h
@@ -836,6 +836,12 @@ struct FSkyMeshBatch
 	bool bVisibleInRealTimeSkyCapture : 1;
 };
 
+struct FSortedTrianglesMeshBatch
+{
+	const FMeshBatch* Mesh = nullptr;
+	const FPrimitiveSceneProxy* Proxy = nullptr;
+};
+
 struct FMeshDecalBatch
 {
 	const FMeshBatch* Mesh;
@@ -1236,6 +1242,9 @@ public:
 	/** Mesh batches with a sky material. */
 	TArray<FSkyMeshBatch, SceneRenderingAllocator> SkyMeshBatches;
 
+	/** Mesh batches with a triangle sorting. */
+	TArray<FSortedTrianglesMeshBatch, SceneRenderingAllocator> SortedTrianglesMeshBatches;
+
 	/** A map from light ID to a boolean visibility value. */
 	TArray<FVisibleLightViewInfo,SceneRenderingAllocator> VisibleLightInfos;
 
diff --git a/Engine/Source/Runtime/Renderer/Private/SceneVisibility.cpp b/Engine/Source/Runtime/Renderer/Private/SceneVisibility.cpp
--- a/Engine/Source/Runtime/Renderer/Private/SceneVisibility.cpp
+++ b/Engine/Source/Runtime/Renderer/Private/SceneVisibility.cpp
@@ -1949,6 +1949,7 @@ struct FRelevancePacket
 	TArray<FMeshDecalBatch> MeshDecalBatches;
 	TArray<FVolumetricMeshBatch> VolumetricMeshBatches;
 	TArray<FSkyMeshBatch> SkyMeshBatches;
+	TArray<FSortedTrianglesMeshBatch> SortedTrianglesMeshBatches;
 	FDrawCommandRelevancePacket DrawCommandPacket;
 
 	struct FPrimitiveLODMask
@@ -2492,6 +2493,14 @@ struct FRelevancePacket
 							BatchAndProxy.bVisibleInRealTimeSkyCapture = PrimitiveSceneInfo->bVisibleInRealTimeSkyCapture;
 						}
 
+						if (ViewRelevance.HasTranslucency() && PrimitiveSceneInfo->Proxy->SupportsSortedTriangles()) // Need to check material as well
+						{
+							SortedTrianglesMeshBatches.AddUninitialized(1);
+							FSortedTrianglesMeshBatch& BatchAndProxy = SortedTrianglesMeshBatches.Last();
+							BatchAndProxy.Mesh = &StaticMesh;
+							BatchAndProxy.Proxy = PrimitiveSceneInfo->Proxy;
+						}
+
 						if (ViewRelevance.bRenderInMainPass && ViewRelevance.bDecal)
 						{
 							MeshDecalBatches.AddUninitialized(1);
@@ -2575,6 +2584,7 @@ struct FRelevancePacket
 		WriteView.MeshDecalBatches.Append(MeshDecalBatches);
 		WriteView.VolumetricMeshBatches.Append(VolumetricMeshBatches);
 		WriteView.SkyMeshBatches.Append(SkyMeshBatches);
+		WriteView.SortedTrianglesMeshBatches.Append(SortedTrianglesMeshBatches);
 
 		for (int32 Index = 0; Index < RecachedReflectionCapturePrimitives.NumPrims; ++Index)
 		{
@@ -2960,6 +2970,14 @@ void ComputeDynamicMeshRelevance(EShadingPath ShadingPath, bool bAddLightmapDens
 		BatchAndProxy.bVisibleInMainPass = ViewRelevance.bRenderInMainPass;
 		BatchAndProxy.bVisibleInRealTimeSkyCapture = PrimitiveSceneInfo->bVisibleInRealTimeSkyCapture;
 	}
+	
+	if (ViewRelevance.HasTranslucency() && PrimitiveSceneInfo->Proxy->SupportsSortedTriangles())
+	{
+		View.SortedTrianglesMeshBatches.AddUninitialized(1);
+		FSortedTrianglesMeshBatch& BatchAndProxy = View.SortedTrianglesMeshBatches.Last();
+		BatchAndProxy.Mesh = MeshBatch.Mesh;
+		BatchAndProxy.Proxy = MeshBatch.PrimitiveSceneProxy;
+	}
 
 	if (ViewRelevance.bRenderInMainPass && ViewRelevance.bDecal)
 	{
