Move ShaderPrint to be accessible by plugins.

diff --git a/Engine/Shaders/Private/ShaderPrintCommon.ush b/Engine/Shaders/Private/ShaderPrintCommon.ush
--- a/Engine/Shaders/Private/ShaderPrintCommon.ush
+++ b/Engine/Shaders/Private/ShaderPrintCommon.ush
@@ -464,6 +464,11 @@ float2 ShaderPrintSetCursorPos(in float2 ScreenPos)
 	return PrevScreenPos;
 }
 
+float2 ShaderPrintGetPos(in uint2 PixelCoord)
+{
+	return float2(PixelCoord) / ShaderPrintUniform.Resolution;
+}
+
 // --------------------------------------------------------------------------
 // Filter helpers
 // Use these to restrict ShaderPrint to a subset of active threads
diff --git a/Engine/Source/Runtime/Renderer/Private/SceneRendering.h b/Engine/Source/Runtime/Renderer/Private/SceneRendering.h
--- a/Engine/Source/Runtime/Renderer/Private/SceneRendering.h
+++ b/Engine/Source/Runtime/Renderer/Private/SceneRendering.h
@@ -31,6 +31,7 @@
 #include "RenderGraphUtils.h"
 #include "MeshDrawCommands.h"
 #include "ShaderDebug.h"
+#include "ShaderPrintParameters.h"
 #include "PostProcess/PostProcessAmbientOcclusionMobile.h"
 #include "Nanite/Nanite.h"
 #include "VirtualShadowMaps/VirtualShadowMapArray.h"
@@ -1474,8 +1475,7 @@ public:
 	FTexture2DRHIRef PrimitiveSceneDataTextureOverrideRHI;
 	FShaderResourceViewRHIRef LightmapSceneDataOverrideSRV;
 
-	FRDGBufferRef ShaderPrintValueBuffer = nullptr;
-
+	FShaderPrintData ShaderPrintData;
 	FShaderDrawDebugData ShaderDrawData;
 	FLumenTranslucencyGIVolume LumenTranslucencyGIVolume;
 
diff --git a/Engine/Source/Runtime/Renderer/Private/ShaderPrint.cpp b/Engine/Source/Runtime/Renderer/Private/ShaderPrint.cpp
--- a/Engine/Source/Runtime/Renderer/Private/ShaderPrint.cpp
+++ b/Engine/Source/Runtime/Renderer/Private/ShaderPrint.cpp
@@ -74,32 +74,28 @@ namespace ShaderPrint
 	typedef TUniformBufferRef<FUniformBufferParameters> FUniformBufferRef;
 
 	// Fill the uniform buffer parameters
-	void SetUniformBufferParameters(FViewInfo const& View, FUniformBufferParameters& OutParameters)
-	{
-		const float FontWidth = (float)FMath::Max(CVarFontSize.GetValueOnRenderThread(), 1) / (float)FMath::Max(View.UnconstrainedViewRect.Size().X, 1);
-		const float FontHeight = (float)FMath::Max(CVarFontSize.GetValueOnRenderThread(), 1) / (float)FMath::Max(View.UnconstrainedViewRect.Size().Y, 1);
-		const float SpaceWidth = (float)FMath::Max(CVarFontSpacingX.GetValueOnRenderThread(), 1) / (float)FMath::Max(View.UnconstrainedViewRect.Size().X, 1);
-		const float SpaceHeight = (float)FMath::Max(CVarFontSpacingY.GetValueOnRenderThread(), 1) / (float)FMath::Max(View.UnconstrainedViewRect.Size().Y, 1);
-
-		OutParameters.FontSize = FVector4f(FontWidth, FontHeight, SpaceWidth + FontWidth, SpaceHeight + FontHeight);
-
-		OutParameters.MaxValueCount = GetMaxValueCount();
-		OutParameters.MaxSymbolCount = GetMaxSymbolCount();
+	// Return a uniform buffer with values filled and with single frame lifetime
+	static FUniformBufferRef CreateUniformBuffer(const FShaderPrintData& Data)
+	{		
+		FUniformBufferParameters Out;
+		Out.FontSize = Data.FontSize;
+		Out.MaxValueCount = Data.MaxValueCount;
+		Out.MaxSymbolCount = Data.MaxSymbolCount;
+		return FUniformBufferRef::CreateUniformBufferImmediate(Out, UniformBuffer_SingleFrame);
 	}
 
-	// Return a uniform buffer with values filled and with single frame lifetime
-	FUniformBufferRef CreateUniformBuffer(FViewInfo const& View)
+	// Fill the FShaderParameters parameters
+	void SetParameters(FRDGBuilder& GraphBuilder, const FViewInfo & View, FShaderParameters& OutParameters)
 	{
-		FUniformBufferParameters Parameters;
-		SetUniformBufferParameters(View, Parameters);
-		return FUniformBufferRef::CreateUniformBufferImmediate(Parameters, UniformBuffer_SingleFrame);
+		OutParameters.UniformBufferParameters = CreateUniformBuffer(View.ShaderPrintData);
+		OutParameters.RWValuesBuffer = GraphBuilder.CreateUAV(View.ShaderPrintData.ShaderPrintValueBuffer);
 	}
 
 	// Fill the FShaderParameters parameters
-	void SetParameters(FRDGBuilder& GraphBuilder, FViewInfo const& View, FShaderParameters& OutParameters)
+	void SetParameters(FRDGBuilder& GraphBuilder, const FShaderPrintData& Data, FShaderParameters& OutParameters)
 	{
-		OutParameters.UniformBufferParameters = CreateUniformBuffer(View);
-		OutParameters.RWValuesBuffer = GraphBuilder.CreateUAV(View.ShaderPrintValueBuffer);
+		OutParameters.UniformBufferParameters = CreateUniformBuffer(Data);
+		OutParameters.RWValuesBuffer = GraphBuilder.CreateUAV(Data.ShaderPrintValueBuffer);
 	}
 
 	// Supported platforms
@@ -270,14 +266,28 @@ namespace ShaderPrint
 	void BeginView(FRDGBuilder& GraphBuilder, FViewInfo& View)
 	{
 		TRACE_CPUPROFILER_EVENT_SCOPE(ShaderPrint::BeginView);
+
+		View.ShaderPrintData = FShaderPrintData();
 		if (!IsSupported(View.GetShaderPlatform()))
 		{
 			return;
 		}
 
+		const FIntPoint Resolution(FMath::Max(View.UnconstrainedViewRect.Size().X, 1), FMath::Max(View.UnconstrainedViewRect.Size().Y, 1));
+
+		const float FontWidth = (float)FMath::Max(CVarFontSize.GetValueOnRenderThread(), 1) / (float)Resolution.X;
+		const float FontHeight = (float)FMath::Max(CVarFontSize.GetValueOnRenderThread(), 1) / (float)Resolution.Y;
+		const float SpaceWidth = (float)FMath::Max(CVarFontSpacingX.GetValueOnRenderThread(), 1) / (float)Resolution.X;
+		const float SpaceHeight = (float)FMath::Max(CVarFontSpacingY.GetValueOnRenderThread(), 1) / (float)Resolution.Y;
+		View.ShaderPrintData.FontSize = FVector4f(FontWidth, FontHeight, SpaceWidth + FontWidth, SpaceHeight + FontHeight);
+		View.ShaderPrintData.Resolution = Resolution;
+		View.ShaderPrintData.MaxValueCount = GetMaxValueCount();
+		View.ShaderPrintData.MaxSymbolCount = GetMaxSymbolCount();
+		View.ShaderPrintData.MaxSymbolCount = GetMaxSymbolCount();
+
 		// Initialize output buffer and store in the view info
 		// Values buffer contains Count + 1 elements. The first element is only used as a counter.
-		View.ShaderPrintValueBuffer = GraphBuilder.CreateBuffer(FRDGBufferDesc::CreateStructuredDesc(sizeof(ShaderPrintItem), GetMaxValueCount() + 1), TEXT("ShaderPrintValueBuffer"));
+		View.ShaderPrintData.ShaderPrintValueBuffer = GraphBuilder.CreateBuffer(FRDGBufferDesc::CreateStructuredDesc(sizeof(ShaderPrintItem), GetMaxValueCount() + 1), TEXT("ShaderPrintValueBuffer"));
 
 		// Early out if system is disabled
 		// Note that we still prepared a minimal ShaderPrintValueBuffer 
@@ -293,7 +303,7 @@ namespace ShaderPrint
 		TShaderMapRef< FShaderInitValueBufferCS > ComputeShader(GlobalShaderMap);
 
 		auto* PassParameters = GraphBuilder.AllocParameters<FShaderInitValueBufferCS::FParameters>();
-		PassParameters->RWValuesBuffer = GraphBuilder.CreateUAV(View.ShaderPrintValueBuffer);
+		PassParameters->RWValuesBuffer = GraphBuilder.CreateUAV(View.ShaderPrintData.ShaderPrintValueBuffer);
 
 		FComputeShaderUtils::AddPass(
 			GraphBuilder,
@@ -316,8 +326,8 @@ namespace ShaderPrint
 		FRDGBufferRef IndirectDrawArgsBuffer = GraphBuilder.CreateBuffer(FRDGBufferDesc::CreateIndirectDesc(5), TEXT("ShaderPrintIndirectDrawArgs"));
 
 		// Non graph managed resources
-		FUniformBufferRef UniformBuffer = CreateUniformBuffer(View);
-		FRDGBufferSRVRef ValuesBuffer = GraphBuilder.CreateSRV(View.ShaderPrintValueBuffer);
+		FUniformBufferRef UniformBuffer = CreateUniformBuffer(View.ShaderPrintData);
+		FRDGBufferSRVRef ValuesBuffer = GraphBuilder.CreateSRV(View.ShaderPrintData.ShaderPrintValueBuffer);
 		FTextureRHIRef FontTexture = GSystemTextures.AsciiTexture->GetRenderTargetItem().ShaderResourceTexture;
 
 		const ERHIFeatureLevel::Type FeatureLevel = View.GetFeatureLevel();
@@ -416,6 +426,6 @@ namespace ShaderPrint
 
 	void EndView(FViewInfo& View)
 	{
-		View.ShaderPrintValueBuffer = nullptr;
+		View.ShaderPrintData = FShaderPrintData();
 	}
 }
diff --git a/Engine/Source/Runtime/Renderer/Private/ShaderPrintParameters.h b/Engine/Source/Runtime/Renderer/Private/ShaderPrintParameters.h
deleted file mode 100644
diff --git a/Engine/Source/Runtime/Renderer/Public/ShaderPrintParameters.h b/Engine/Source/Runtime/Renderer/Public/ShaderPrintParameters.h
new file mode 100644
--- /dev/null
+++ b/Engine/Source/Runtime/Renderer/Public/ShaderPrintParameters.h
@@ -0,0 +1,48 @@
+// Copyright Epic Games, Inc. All Rights Reserved.
+
+// The ShaderPrint system uses a RWBuffer to capture any debug print from a shader.
+// This means that the buffer needs to be bound for the shader you wish to debug.
+// It would be ideal if that was automatic (maybe by having a fixed bind point for the buffer and binding it for the entire view).
+// But for now you need to manually add binding information to your FShader class.
+// To do this either:
+// (i) Use SHADER_PARAMETER_STRUCT_INCLUDE(ShaderPrint::FShaderParameters) in your FShader::FParameters declaration and call SetParameters().
+
+// Also it seems that we can only bind a RWBuffer to compute shaders right now. Fixing this would allow us to use this system from all shader stages.
+
+#pragma once
+
+#include "CoreMinimal.h"
+#include "ShaderParameterMacros.h"
+#include "ShaderParameters.h"
+
+class FViewInfo;
+
+struct FShaderPrintData
+{
+	FVector4 FontSize;
+	FIntPoint Resolution;
+	int32 MaxValueCount = -1;
+	int32 MaxSymbolCount = -1;
+	FRDGBufferRef ShaderPrintValueBuffer = nullptr;
+};
+
+namespace ShaderPrint
+{
+	// ShaderPrint uniform buffer layout
+	BEGIN_GLOBAL_SHADER_PARAMETER_STRUCT(FUniformBufferParameters, RENDERER_API)
+		SHADER_PARAMETER(FVector4f, FontSize)
+		SHADER_PARAMETER(FIntPoint, Resolution)
+		SHADER_PARAMETER(int32, MaxValueCount)
+		SHADER_PARAMETER(int32, MaxSymbolCount)
+	END_GLOBAL_SHADER_PARAMETER_STRUCT()
+
+	// ShaderPrint parameter struct declaration
+	BEGIN_SHADER_PARAMETER_STRUCT(FShaderParameters, )
+		SHADER_PARAMETER_STRUCT_REF(FUniformBufferParameters, UniformBufferParameters)
+		SHADER_PARAMETER_RDG_BUFFER_UAV(RWStructuredBuffer<ShaderPrintItem>, RWValuesBuffer)
+	END_SHADER_PARAMETER_STRUCT()
+
+	// Call this to fill the FShaderParameters
+	RENDERER_API void SetParameters(FRDGBuilder& GraphBuilder, const FViewInfo & View, FShaderParameters& OutParameters);
+	RENDERER_API void SetParameters(FRDGBuilder& GraphBuilder, const FShaderPrintData& Data, FShaderParameters& OutParameters);
+}
