Fix CIS issue.

diff --git a/Engine/Source/Runtime/Renderer/Private/OIT/OIT.cpp b/Engine/Source/Runtime/Renderer/Private/OIT/OIT.cpp
--- a/Engine/Source/Runtime/Renderer/Private/OIT/OIT.cpp
+++ b/Engine/Source/Runtime/Renderer/Private/OIT/OIT.cpp
@@ -161,7 +161,7 @@ FSortedTriangleData FOITSceneData::Allocate(const FIndexBuffer* InSource, EPrimi
 		for (uint32 FreeIt=0,FreeCount=FreeBuffers.Num(); FreeIt<FreeCount; ++FreeIt)
 		{		
 			FSortedIndexBuffer* FreeBuffer = FreeBuffers[FreeIt];
-			if (FreeBuffer != nullptr && FreeBuffer->NumIndices >= NumIndices && OITIndexBuffer->Id == FSortedIndexBuffer::InvalidId)
+			if (FreeBuffer != nullptr && FreeBuffer->NumIndices >= NumIndices && FreeBuffer->Id == FSortedIndexBuffer::InvalidId)
 			{			
 				OITIndexBuffer = FreeBuffer;
 				OITIndexBuffer->Id = FreeSlot;
@@ -389,15 +389,15 @@ static void AddOITSortTriangleIndexPass(
 	FOITDebugData& DebugData)
 {
 	static const FVertexFactoryType* CompatibleVF = FVertexFactoryType::GetVFByName(TEXT("FLocalVertexFactory"));
-	const bool bIsCompatible = MeshBatch.Mesh->VertexFactory->GetType()->GetHashedName() == CompatibleVF->GetHashedName();
 
 	// Fat format: PF_R32G32_UINT | Compact format: PF_R32_UINT
 	const EPixelFormat PackedFormat = PF_R32_UINT;
 	const uint32 PackedFormatInBytes = 4;
 
 	const bool bIsValid = 
-		bIsCompatible && 
 		MeshBatch.Mesh != nullptr && 
+		MeshBatch.Mesh->VertexFactory != nullptr &&
+		MeshBatch.Mesh->VertexFactory->GetType()->GetHashedName() == CompatibleVF->GetHashedName() &&
 		MeshBatch.Mesh->Elements.Num() > 0 && 
 		MeshBatch.Mesh->Elements[0].DynamicIndexBuffer.IndexBuffer != nullptr;
 	if (!bIsValid)
