* Add details to GPU scene debug for selected Items (flags, names, ID, ...).
* Add ShaderPrint character requestion to let other system require more characters space (one frame of latency)
* Fix shader print invalid resolution value

diff --git a/Engine/Shaders/Private/GPUSceneDebugRender.usf b/Engine/Shaders/Private/GPUSceneDebugRender.usf
--- a/Engine/Shaders/Private/GPUSceneDebugRender.usf
+++ b/Engine/Shaders/Private/GPUSceneDebugRender.usf
@@ -3,6 +3,7 @@
 #include "Common.ush"
 #include "SceneData.ush"
 #include "ShaderDrawDebug.ush"
+#include "ShaderPrintCommon.ush"
 
 int NumInstances;
 int NumScenePrimitives;
@@ -15,8 +16,82 @@ float3 PickingRayStart;
 float3 PickingRayEnd;
 float DrawRange;
 
+//////////////////////////////////////////////////////////////////////////////////////////
+#define INVALID_PRIMITIVEID ~0
+
+uint					SelectedNameInfoCount;
+uint					SelectedNameCharacterCount;
+StructuredBuffer<uint2> SelectedPrimitiveNameInfos;
+Buffer<uint>			SelectedPrimitiveNames;
+
+struct FPrimitiveSceneDebugNameInfo
+{
+	uint PrimitiveID;
+	uint Length;
+	uint Offset;
+	uint Pad0;
+	uint Pad1;
+};
+
+FPrimitiveSceneDebugNameInfo UnpackDebugNameInfo(uint2 In)
+{
+	FPrimitiveSceneDebugNameInfo Out = (FPrimitiveSceneDebugNameInfo)0;
+	Out.PrimitiveID = In.x;
+	Out.Offset = (In.y) & 0xFFFF;
+	Out.Length = (In.y>>16) & 0xFF;
+
+	return Out;
+}
+
+FPrimitiveSceneDebugNameInfo FindNameInfo(uint PrimitiveID)
+{
+	FPrimitiveSceneDebugNameInfo Out = (FPrimitiveSceneDebugNameInfo)0;
+	Out.PrimitiveID = INVALID_PRIMITIVEID;
+	for (uint It = 0; It < SelectedNameInfoCount; ++It)
+	{
+		if (SelectedPrimitiveNameInfos[It].x == PrimitiveID)
+		{
+			return UnpackDebugNameInfo(SelectedPrimitiveNameInfos[It]);
+		}
+	}
+	return Out;
+}
+
+float2 PrintInstanceName(float2 InPos, uint PrimitiveID, FFontColor InColor)
+{
+	const FPrimitiveSceneDebugNameInfo Info = FindNameInfo(PrimitiveID);
+	if (Info.PrimitiveID != INVALID_PRIMITIVEID && (Info.Length + Info.Offset) <= SelectedNameCharacterCount)
+	{
+		for (uint It = 0; It < Info.Length; ++It)
+		{
+			const uint Char = SelectedPrimitiveNames[It + Info.Offset];
+			InPos = ShaderPrintSymbol(InPos, Char, InColor);
+		}
+	}
+	return InPos;
+}
+
+//////////////////////////////////////////////////////////////////////////////////////////
+
 StructuredBuffer<uint> SelectedPrimitiveFlags;
+RWBuffer<uint> RWDrawCounter;
 
+void AddNewLine(float2 OriginalPos, inout float2 Pos)
+{
+	Pos = ShaderPrintNewline(Pos); Pos.x = OriginalPos.x;
+}
+
+void AddFlag(inout float2 Pos, bool bEnabled)
+{
+	if (bEnabled)
+	{
+		Pos = ShaderPrintText(Pos, TEXT("On     "), FontGreen);
+	}
+	else
+	{
+		Pos = ShaderPrintText(Pos, TEXT("Off    "), FontRed);
+	}
+}
 [numthreads(NUM_THREADS_PER_GROUP, 1, 1)]
 void GPUSceneDebugRenderCS(uint InstanceID : SV_DispatchThreadID)
 {
@@ -59,25 +134,60 @@ void GPUSceneDebugRenderCS(uint InstanceID : SV_DispatchThreadID)
 	FPrimitiveSceneData PrimitiveData = GetPrimitiveData(InstanceData.PrimitiveId);
 
 	//float4 InstanceBoundColor = InstanceData.PrimitiveId < NumScenePrimitives ? float4(0.3f, 1.0f, 0.5f, 1.0f) : float4(0.3f, 0.3f, 0.3f, 0.5f);
-	float4 InstanceBoundColor = float4(0.3f, 0.3f, 0.3f, 0.5f);
+	float4 InstanceBoundColor = float4(ColorSilver.xyz, 0.5f);
 
-	if (bIsSelected || bDrawAll)
+	if (bDrawAll)
 	{
-		InstanceBoundColor = float4(0.5f, 0.75f, 0.5f, 1.0f);
+		InstanceBoundColor = ColorLightGreen;
 	}
 	
 	if (bIsMouseOver)
 	{
-		InstanceBoundColor = float4(0.75f, 1.0f, 0.75f, 1.0f);
+		InstanceBoundColor = ColorYellow;
 	}
 	
+	if (bIsSelected)
+	{
+		InstanceBoundColor = ColorOrange;
+	}
+
 	if (bDrawUpdatedOnly)
 	{
-		InstanceBoundColor.x = 1.0f;
-		InstanceBoundColor.z = 1.0f;
+		InstanceBoundColor = ColorRed;
 	}
 	if (bIsMouseOver || bIsSelected || bDrawAll != 0)
 	{
 		AddOBB(InstanceData.LocalBoundsCenter - InstanceData.LocalBoundsExtent, InstanceData.LocalBoundsCenter + InstanceData.LocalBoundsExtent, InstanceBoundColor, InstanceData.LocalToWorld);
+		AddReferential(InstanceData.LocalToWorld, 1);
+	}
+
+	uint StartOffsetX = 10;
+	uint StartOffsetY = 30;
+	uint SlotHeight = 8;
+	if (InstanceID == 0 && SelectedNameInfoCount>0)
+	{
+		float2 OriginalPos = ShaderPrintGetPos(uint2(StartOffsetX, StartOffsetY + SlotHeight));
+		float2 Pos = OriginalPos;
+		Pos = ShaderPrintText(Pos, TEXT("Prim.ID     Inst.ID     Shadow Veloc. CusDat DynDat Name"), FontWhite);
+	}
+
+	if (bIsSelected)
+	{	
+		uint SlotIndex = 0;
+		InterlockedAdd(RWDrawCounter[0], 1, SlotIndex);
+
+		float2 OriginalPos = ShaderPrintGetPos(uint2(StartOffsetX, StartOffsetY + 2*SlotHeight + SlotIndex * SlotHeight));
+		float2 Pos = OriginalPos;
+
+		Pos = ShaderPrint(Pos, InstanceData.PrimitiveId, FontYellow);
+		Pos = ShaderPrint(Pos, InstanceID, FontOrange);
+
+		FPrimitiveSceneData PrimitiveData = GetPrimitiveData(InstanceData.PrimitiveId);
+		AddFlag(Pos, 0u != (PrimitiveData.Flags & PRIMITIVE_SCENE_DATA_FLAG_CAST_SHADOWS));
+		AddFlag(Pos, 0u != (PrimitiveData.Flags & PRIMITIVE_SCENE_DATA_FLAG_OUTPUT_VELOCITY));
+		AddFlag(Pos, 0u != (PrimitiveData.Flags & INSTANCE_SCENE_DATA_FLAG_HAS_CUSTOM_DATA));
+		AddFlag(Pos, 0u != (PrimitiveData.Flags & INSTANCE_SCENE_DATA_FLAG_HAS_DYNAMIC_DATA));
+		
+		Pos = PrintInstanceName(Pos, InstanceData.PrimitiveId, FontSilver); AddNewLine(OriginalPos, Pos);
 	}
 }
diff --git a/Engine/Shaders/Private/ShaderDrawDebug.ush b/Engine/Shaders/Private/ShaderDrawDebug.ush
--- a/Engine/Shaders/Private/ShaderDrawDebug.ush
+++ b/Engine/Shaders/Private/ShaderDrawDebug.ush
@@ -20,6 +20,7 @@
 #define ColorTurquoise	float4( 26.f / 255.f, 188.f / 255.f, 156.f / 255.f, 1)
 #define ColorSilver		float4(189.f / 255.f, 195.f / 255.f, 199.f / 255.f, 1)
 #define ColorEmerald	float4( 46.f / 255.f, 204.f / 255.f, 113.f / 255.f, 1)
+#define ColorLightGreen	float4(0.5f, 0.75f, 0.5f, 1.0f)
 
 struct FShaderDrawElement 
 {
diff --git a/Engine/Shaders/Private/ShaderPrintCommon.ush b/Engine/Shaders/Private/ShaderPrintCommon.ush
--- a/Engine/Shaders/Private/ShaderPrintCommon.ush
+++ b/Engine/Shaders/Private/ShaderPrintCommon.ush
@@ -17,7 +17,7 @@ struct FShaderPrintItem
 	float2 ScreenPos; // Position in normalized coordinates
 	int Value;        // Cast to value or symbol
 	int Type;         // SHADER_PRINT_TYPE_* defines how to read Value
-	float3 Color;	  // Color
+	float3 Color;	  // Color 
 };
 
 struct FPackedShaderPrintItem
@@ -466,7 +466,7 @@ float2 ShaderPrintSetCursorPos(in float2 ScreenPos)
 
 float2 ShaderPrintGetPos(in uint2 PixelCoord)
 {
-	return float2(PixelCoord) / ShaderPrintUniform.Resolution;
+	return float2(PixelCoord) / float2(ShaderPrintUniform.Resolution);
 }
 
 // --------------------------------------------------------------------------
diff --git a/Engine/Source/Runtime/Renderer/Private/GPUScene.cpp b/Engine/Source/Runtime/Renderer/Private/GPUScene.cpp
--- a/Engine/Source/Runtime/Renderer/Private/GPUScene.cpp
+++ b/Engine/Source/Runtime/Renderer/Private/GPUScene.cpp
@@ -21,6 +21,7 @@
 #include "HAL/LowLevelMemTracker.h"
 #include "HAL/LowLevelMemStats.h"
 #include "ShaderDebug.h"
+#include "ShaderPrint.h"
 
 int32 GGPUSceneUploadEveryFrame = 0;
 FAutoConsoleVariableRef CVarGPUSceneUploadEveryFrame(
@@ -451,7 +452,7 @@ void FGPUScene::UpdateInternal(FRDGBuilder& GraphBuilder, FScene& Scene)
 	PrimitivesToUpdate.Reset();
 	PrimitiveDirtyState.Init(EPrimitiveDirtyState::None, PrimitiveDirtyState.Num());
 
-	AddPass(GraphBuilder, RDG_EVENT_NAME("GPUSceneUpdate"), 
+	AddPass(GraphBuilder, RDG_EVENT_NAME("GPUScene::Update"), 
 		[this, &Scene, Adapter = MoveTemp(Adapter), BufferState = MoveTemp(BufferState)](FRHICommandListImmediate& RHICmdList)
 	{
 		SCOPED_NAMED_EVENT(STAT_UpdateGPUScene, FColor::Green);
@@ -1416,7 +1417,7 @@ void FGPUScene::UploadDynamicPrimitiveShaderDataForViewInternal(FRDGBuilder& Gra
 		FUploadDataSourceAdapterDynamicPrimitives UploadAdapter(Collector.UploadData->PrimitiveShaderData, UploadIdStart, Collector.UploadData->InstanceSceneDataOffset, SceneFrameNumber);
 		FGPUSceneBufferState BufferState = UpdateBufferState(GraphBuilder, Scene, UploadAdapter);
 
-		AddPass(GraphBuilder, RDG_EVENT_NAME("UploadDynamicPrimitiveShaderDataForView"),
+		AddPass(GraphBuilder, RDG_EVENT_NAME("GPUScene::UploadDynamicPrimitiveShaderDataForView"),
 			[this, Scene, UploadAdapter, BufferState = MoveTemp(BufferState)](FRHICommandListImmediate& RHICmdList)
 		{
 			UploadGeneral<FUploadDataSourceAdapterDynamicPrimitives>(RHICmdList, Scene, UploadAdapter, BufferState);
@@ -1574,6 +1575,14 @@ void FGPUScene::FreeInstancePayloadDataSlots(int32 InstancePayloadDataOffset, in
 	}
 }
 
+struct FPrimitiveSceneDebugNameInfo
+{
+	uint32 PrimitiveID;
+	uint16 Offset;
+	uint8  Length;
+	uint8  Pad0;
+};
+
 class FGPUSceneDebugRenderCS : public FGlobalShader
 {
 	DECLARE_GLOBAL_SHADER(FGPUSceneDebugRenderCS);
@@ -1581,6 +1590,7 @@ class FGPUSceneDebugRenderCS : public FGlobalShader
 
 	BEGIN_SHADER_PARAMETER_STRUCT(FParameters, )
 		SHADER_PARAMETER_STRUCT_INCLUDE(ShaderDrawDebug::FShaderParameters, ShaderDrawUniformBuffer)
+		SHADER_PARAMETER_STRUCT_INCLUDE(ShaderPrint::FShaderParameters, ShaderPrintUniformBuffer)
 		SHADER_PARAMETER_SRV(StructuredBuffer<float4>, GPUSceneInstanceSceneData)
 		SHADER_PARAMETER_SRV(StructuredBuffer<float4>, GPUScenePrimitiveSceneData)
 		SHADER_PARAMETER(uint32, InstanceDataSOAStride)
@@ -1589,10 +1599,15 @@ class FGPUSceneDebugRenderCS : public FGlobalShader
 		SHADER_PARAMETER(int32, NumScenePrimitives)
 		SHADER_PARAMETER(int32, bDrawAll)
 		SHADER_PARAMETER(int32, bDrawUpdatedOnly)
-		SHADER_PARAMETER_RDG_BUFFER_SRV(StructuredBuffer< uint32 >, SelectedPrimitiveFlags)
+		SHADER_PARAMETER(int32, SelectedNameInfoCount)
+		SHADER_PARAMETER(int32, SelectedNameCharacterCount)
+		SHADER_PARAMETER_RDG_BUFFER_SRV(StructuredBuffer<uint>, SelectedPrimitiveFlags)
+		SHADER_PARAMETER_RDG_BUFFER_SRV(StructuredBuffer<uint2>, SelectedPrimitiveNameInfos)
+		SHADER_PARAMETER_RDG_BUFFER_SRV(Buffer<uint8>, SelectedPrimitiveNames)
 		SHADER_PARAMETER(FVector3f, PickingRayStart)
 		SHADER_PARAMETER(FVector3f, PickingRayEnd)
 		SHADER_PARAMETER(float, DrawRange)
+		SHADER_PARAMETER_RDG_BUFFER_UAV(RWBuffer<uint32>, RWDrawCounter)
 	END_SHADER_PARAMETER_STRUCT()
 
 public:
@@ -1625,15 +1640,25 @@ void FGPUScene::DebugRender(FRDGBuilder& GraphBuilder, FScene& Scene, FViewInfo&
 	if (DebugMode > 0)
 	{
 		ShaderDrawDebug::SetEnabled(true);
+		if (!ShaderPrint::IsEnabled(View)) { ShaderPrint::SetEnabled(true); }
 
 		int32 NumInstances = InstanceSceneDataAllocator.GetMaxSize();
-		if (ShaderDrawDebug::IsEnabled(View) && NumInstances > 0)
+		if (ShaderDrawDebug::IsEnabled(View) && ShaderPrint::IsEnabled(View) && NumInstances > 0)
 		{
 			// This lags by one frame, so may miss some in one frame, also overallocates since we will cull a lot.
 			ShaderDrawDebug::RequestSpaceForElements(NumInstances * 12);
 
-			FGPUSceneDebugRenderCS::FParameters* PassParameters = GraphBuilder.AllocParameters<FGPUSceneDebugRenderCS::FParameters>();
+			FRDGBufferRef DrawCounterBuffer = GraphBuilder.CreateBuffer(FRDGBufferDesc::CreateBufferDesc(4, 1), TEXT("GPUScene.DebugCounter"));
+			FRDGBufferUAVRef DrawCounterUAV = GraphBuilder.CreateUAV(DrawCounterBuffer, PF_R32_UINT);
+			AddClearUAVPass(GraphBuilder, DrawCounterUAV, 0u);
+
+			const uint32 MaxPrimitiveNameCount = 128u;
+			check(sizeof(FPrimitiveSceneDebugNameInfo) == 8);
+			TArray<FPrimitiveSceneDebugNameInfo> SelectedNameInfos;
+			TArray<uint8> SelectedNames;
+			SelectedNames.Reserve(MaxPrimitiveNameCount * 30u);
 
+			uint32 SelectedCount = 0;
 			TArray<uint32> SelectedPrimitiveFlags;
 			const int32 BitsPerWord = (sizeof(uint32) * 8U);
 			SelectedPrimitiveFlags.Init(0U, FMath::DivideAndRoundUp(Scene.Primitives.Num(), BitsPerWord));
@@ -1642,11 +1667,49 @@ void FGPUScene::DebugRender(FRDGBuilder& GraphBuilder, FScene& Scene, FViewInfo&
 				if (Scene.PrimitiveSceneProxies[PrimitiveID]->IsSelected())
 				{
 					SelectedPrimitiveFlags[PrimitiveID / BitsPerWord] |= 1U << uint32(PrimitiveID % BitsPerWord);
+
+					// Collect Names
+					if (SelectedNameInfos.Num() < MaxPrimitiveNameCount)
+					{
+						const FString OwnerName = Scene.Primitives[PrimitiveID]->GetFullnameForDebuggingOnly();
+						const uint32 NameOffset = SelectedNames.Num();
+						const uint32 NameLength = OwnerName.Len();
+						for (TCHAR C : OwnerName)
+						{
+							SelectedNames.Add(uint8(C));
+						}
+
+						FPrimitiveSceneDebugNameInfo& NameInfo = SelectedNameInfos.AddDefaulted_GetRef();
+						NameInfo.PrimitiveID= PrimitiveID;
+						NameInfo.Length		= NameLength;
+						NameInfo.Offset		= NameOffset;
+						++SelectedCount;
+					}
 				}
 			}
-			FRDGBufferRef SelectedPrimitiveFlagsRDG = CreateStructuredBuffer(GraphBuilder, TEXT("GPUScene.Debug.SelectedPrimitiveFlags"), SelectedPrimitiveFlags);
 
+			if (SelectedNameInfos.IsEmpty())
+			{
+				FPrimitiveSceneDebugNameInfo& NameInfo = SelectedNameInfos.AddDefaulted_GetRef();
+				NameInfo.PrimitiveID= ~0;
+				NameInfo.Length		= 4;
+				NameInfo.Offset		= 0;
+				SelectedNames.Add(uint8('N'));
+				SelectedNames.Add(uint8('o'));
+				SelectedNames.Add(uint8('n'));
+				SelectedNames.Add(uint8('e'));
+			}
+
+			// Request more characters for printing if needed
+			ShaderPrint::RequestSpaceForCharacters(SelectedNames.Num() + SelectedCount * 48u);
+
+			FRDGBufferRef SelectedPrimitiveNames	 = CreateVertexBuffer(GraphBuilder, TEXT("GPUScene.Debug.SelectedPrimitiveNames"), FRDGBufferDesc::CreateBufferDesc(1, SelectedNames.Num()), SelectedNames.GetData(), SelectedNames.Num());
+			FRDGBufferRef SelectedPrimitiveNameInfos = CreateStructuredBuffer(GraphBuilder, TEXT("GPUScene.Debug.SelectedPrimitiveNameInfos"), SelectedNameInfos);
+			FRDGBufferRef SelectedPrimitiveFlagsRDG  = CreateStructuredBuffer(GraphBuilder, TEXT("GPUScene.Debug.SelectedPrimitiveFlags"), SelectedPrimitiveFlags);
+
+			FGPUSceneDebugRenderCS::FParameters* PassParameters = GraphBuilder.AllocParameters<FGPUSceneDebugRenderCS::FParameters>();
 			ShaderDrawDebug::SetParameters(GraphBuilder, View.ShaderDrawData, PassParameters->ShaderDrawUniformBuffer);
+			ShaderPrint::SetParameters(GraphBuilder, View, PassParameters->ShaderPrintUniformBuffer);
 			PassParameters->GPUSceneInstanceSceneData = InstanceSceneDataBuffer.SRV;
 			PassParameters->GPUScenePrimitiveSceneData = PrimitiveBuffer.SRV;
 			PassParameters->InstanceDataSOAStride = InstanceSceneDataSOAStride;
@@ -1654,9 +1717,14 @@ void FGPUScene::DebugRender(FRDGBuilder& GraphBuilder, FScene& Scene, FViewInfo&
 			PassParameters->bDrawUpdatedOnly = DebugMode == 3;
 			PassParameters->bDrawAll = DebugMode != 2;
 			PassParameters->NumInstances = NumInstances;
+			PassParameters->SelectedNameInfoCount = SelectedCount;
+			PassParameters->SelectedNameCharacterCount = SelectedCount > 0 ? SelectedNames.Num() : 0;
 			PassParameters->SelectedPrimitiveFlags = GraphBuilder.CreateSRV(SelectedPrimitiveFlagsRDG);
+			PassParameters->SelectedPrimitiveNameInfos = GraphBuilder.CreateSRV(SelectedPrimitiveNameInfos);
+			PassParameters->SelectedPrimitiveNames = GraphBuilder.CreateSRV(SelectedPrimitiveNames, PF_R8_UINT);
 			PassParameters->NumScenePrimitives = NumScenePrimitives;
 			PassParameters->DrawRange = CVarGPUSceneDebugDrawRange.GetValueOnRenderThread();
+			PassParameters->RWDrawCounter = DrawCounterUAV;
 
 			FVector PickingRayStart(ForceInit);
 			FVector PickingRayDir(ForceInit);
@@ -1669,7 +1737,7 @@ void FGPUScene::DebugRender(FRDGBuilder& GraphBuilder, FScene& Scene, FViewInfo&
 
 			FComputeShaderUtils::AddPass(
 				GraphBuilder,
-				RDG_EVENT_NAME("GPUSceneDebugRender"),
+				RDG_EVENT_NAME("GPUScene::DebugRender"),
 				ComputeShader,
 				PassParameters,
 				FComputeShaderUtils::GetGroupCount(NumInstances, FGPUSceneDebugRenderCS::NumThreadsPerGroup)
@@ -1730,7 +1798,7 @@ bool FGPUScene::BeginReadWriteAccess(FRDGBuilder& GraphBuilder)
 	if (IsEnabled())
 	{
 		// TODO: Remove this when everything is properly RDG'd
-		AddPass(GraphBuilder, RDG_EVENT_NAME("TransitionInstanceSceneDataBuffer"), [this](FRHICommandList& RHICmdList)
+		AddPass(GraphBuilder, RDG_EVENT_NAME("GPUScene::TransitionInstanceSceneDataBuffer"), [this](FRHICommandList& RHICmdList)
 		{
 			FRHITransitionInfo 	Transitions[2] =
 			{
@@ -1768,7 +1836,7 @@ void FGPUScene::EndReadWriteAccess(FRDGBuilder& GraphBuilder, ERHIAccess FinalAc
 	if (IsEnabled())
 	{
 		// TODO: Remove this when everything is properly RDG'd
-		AddPass(GraphBuilder, RDG_EVENT_NAME("TransitionInstanceSceneDataBuffer"), [this, FinalAccessState](FRHICommandList& RHICmdList)
+		AddPass(GraphBuilder, RDG_EVENT_NAME("GPUScene::TransitionInstanceSceneDataBuffer"), [this, FinalAccessState](FRHICommandList& RHICmdList)
 		{
 			FRHITransitionInfo 	Transitions[2] =
 			{
@@ -1830,7 +1898,7 @@ void FGPUScene::AddUpdatePrimitiveIdsPass(FRDGBuilder& GraphBuilder, FInstanceGP
 
 		FComputeShaderUtils::AddPass(
 			GraphBuilder,
-			RDG_EVENT_NAME("GPUSceneSetInstancePrimitiveIdCS"),
+			RDG_EVENT_NAME("GPUScene::SetInstancePrimitiveIdCS"),
 			ComputeShader,
 			PassParameters,
 			IdOnlyUpdateItems.GetWrappedCsGroupCount()
diff --git a/Engine/Source/Runtime/Renderer/Private/PrimitiveSceneInfo.cpp b/Engine/Source/Runtime/Renderer/Private/PrimitiveSceneInfo.cpp
--- a/Engine/Source/Runtime/Renderer/Private/PrimitiveSceneInfo.cpp
+++ b/Engine/Source/Runtime/Renderer/Private/PrimitiveSceneInfo.cpp
@@ -1933,3 +1933,14 @@ void FPrimitiveSceneInfo::UpdateComponentLastRenderTime(float CurrentWorldTime,
 		*OwnerLastRenderTime = CurrentWorldTime; // Sets OwningActor->LastRenderTime
 	}
 }
+
+FString FPrimitiveSceneInfo::GetFullnameForDebuggingOnly() const
+{
+	// This is not correct to access component from rendering thread, but this is for debugging only
+	if (ComponentForDebuggingOnly)
+	{	
+		return ComponentForDebuggingOnly->GetFullGroupName(false);
+	}
+	return FString(TEXT("Unknown Object"));
+}
+
diff --git a/Engine/Source/Runtime/Renderer/Private/ShaderPrint.cpp b/Engine/Source/Runtime/Renderer/Private/ShaderPrint.cpp
--- a/Engine/Source/Runtime/Renderer/Private/ShaderPrint.cpp
+++ b/Engine/Source/Runtime/Renderer/Private/ShaderPrint.cpp
@@ -23,7 +23,7 @@ namespace ShaderPrint
 
 	static TAutoConsoleVariable<int32> CVarFontSize(
 		TEXT("r.ShaderPrintFontSize"),
-		16,
+		8,
 		TEXT("ShaderPrint font size.\n"),
 		ECVF_Cheat | ECVF_RenderThreadSafe);
 
@@ -45,6 +45,7 @@ namespace ShaderPrint
 		TEXT("ShaderPrint output buffer size.\n"),
 		ECVF_Cheat | ECVF_RenderThreadSafe);
 
+	static uint32 GCharacterRequestCount = 0;
 
 	// Structure used by shader buffers to store values and symbols
 	struct ShaderPrintItem
@@ -56,15 +57,15 @@ namespace ShaderPrint
 
 	// Get value buffer size
 	// Note that if the ShaderPrint system is disabled we still want to bind a minimal buffer
-	int32 GetMaxValueCount()
+	static int32 GetMaxValueCount()
 	{
-		int32 MaxValueCount = FMath::Max(CVarMaxValueCount.GetValueOnRenderThread(), 0);
+		int32 MaxValueCount = FMath::Max(CVarMaxValueCount.GetValueOnRenderThread() + int32(GCharacterRequestCount), 0);
 		return IsEnabled() ? MaxValueCount : 0;
 	}
 
 	// Get symbol buffer size
 	// This is some multiple of the value buffer size to allow for maximum value->symbol expansion
-	int32 GetMaxSymbolCount()
+	static int32 GetMaxSymbolCount()
 	{
 		return GetMaxValueCount() * 12;
 	}
@@ -79,6 +80,7 @@ namespace ShaderPrint
 	{		
 		FUniformBufferParameters Out;
 		Out.FontSize = Data.FontSize;
+		Out.Resolution = Data.Resolution;
 		Out.MaxValueCount = Data.MaxValueCount;
 		Out.MaxSymbolCount = Data.MaxSymbolCount;
 		return FUniformBufferRef::CreateUniformBufferImmediate(Out, UniformBuffer_SingleFrame);
@@ -119,6 +121,11 @@ namespace ShaderPrint
 		CVarMaxValueCount->Set(FMath::Max(256, InMaxCount));
 	}	
 
+	void RequestSpaceForCharacters(uint32 MaxElementCount)
+	{
+		GCharacterRequestCount += MaxElementCount;
+	}
+
 	bool IsEnabled()
 	{
 		return CVarEnable.GetValueOnAnyThread() != 0;
@@ -283,11 +290,11 @@ namespace ShaderPrint
 		View.ShaderPrintData.Resolution = Resolution;
 		View.ShaderPrintData.MaxValueCount = GetMaxValueCount();
 		View.ShaderPrintData.MaxSymbolCount = GetMaxSymbolCount();
-		View.ShaderPrintData.MaxSymbolCount = GetMaxSymbolCount();
+		GCharacterRequestCount = 0;
 
 		// Initialize output buffer and store in the view info
 		// Values buffer contains Count + 1 elements. The first element is only used as a counter.
-		View.ShaderPrintData.ShaderPrintValueBuffer = GraphBuilder.CreateBuffer(FRDGBufferDesc::CreateStructuredDesc(sizeof(ShaderPrintItem), GetMaxValueCount() + 1), TEXT("ShaderPrintValueBuffer"));
+		View.ShaderPrintData.ShaderPrintValueBuffer = GraphBuilder.CreateBuffer(FRDGBufferDesc::CreateStructuredDesc(sizeof(ShaderPrintItem), View.ShaderPrintData.MaxValueCount + 1), TEXT("ShaderPrintValueBuffer"));
 
 		// Early out if system is disabled
 		// Note that we still prepared a minimal ShaderPrintValueBuffer 
diff --git a/Engine/Source/Runtime/Renderer/Public/PrimitiveSceneInfo.h b/Engine/Source/Runtime/Renderer/Public/PrimitiveSceneInfo.h
--- a/Engine/Source/Runtime/Renderer/Public/PrimitiveSceneInfo.h
+++ b/Engine/Source/Runtime/Renderer/Public/PrimitiveSceneInfo.h
@@ -578,6 +578,9 @@ public:
 	RENDERER_API FRHIRayTracingGeometry* GetStaticRayTracingGeometryInstance(int LodLevel) const;
 #endif
 
+	/** Return primitive fullname (for debugging only). */
+	FString GetFullnameForDebuggingOnly() const;
+
 private:
 
 	/** Let FScene have direct access to the Id. */
diff --git a/Engine/Source/Runtime/Renderer/Public/ShaderPrintParameters.h b/Engine/Source/Runtime/Renderer/Public/ShaderPrintParameters.h
--- a/Engine/Source/Runtime/Renderer/Public/ShaderPrintParameters.h
+++ b/Engine/Source/Runtime/Renderer/Public/ShaderPrintParameters.h
@@ -19,7 +19,7 @@ class FViewInfo;
 
 struct FShaderPrintData
 {
-	FVector4 FontSize;
+	FVector4f FontSize;
 	FIntPoint Resolution;
 	int32 MaxValueCount = -1;
 	int32 MaxSymbolCount = -1;
@@ -45,4 +45,11 @@ namespace ShaderPrint
 	// Call this to fill the FShaderParameters
 	RENDERER_API void SetParameters(FRDGBuilder& GraphBuilder, const FViewInfo & View, FShaderParameters& OutParameters);
 	RENDERER_API void SetParameters(FRDGBuilder& GraphBuilder, const FShaderPrintData& Data, FShaderParameters& OutParameters);
+
+	/**
+	 * Call to ensure enough space for some number of characters, is added cumulatively each frame, to make it possible for several systems to request a certain number independently.
+	 * Is used to grow the max element count for subsequent frames (as the allocation happens early in the frame).
+	 * @param The number of elements requested, an element corresponds to a line, so a cube, for example, needs 12 elements.
+	 */
+	RENDERER_API void RequestSpaceForCharacters(uint32 MaxElementCount);
 }
