Fix DistanceField asset data upload when LWC is enabled.

diff --git a/Engine/Source/Runtime/Renderer/Private/DistanceFieldStreaming.cpp b/Engine/Source/Runtime/Renderer/Private/DistanceFieldStreaming.cpp
--- a/Engine/Source/Runtime/Renderer/Private/DistanceFieldStreaming.cpp
+++ b/Engine/Source/Runtime/Renderer/Private/DistanceFieldStreaming.cpp
@@ -843,7 +843,7 @@ void FDistanceFieldSceneData::UploadAssetData(
 		return;
 	}
 	
-	AssetDataUploadBuffer.Init(AssetDataUploads.Num(), AssetDataMipStrideFloat4s * sizeof(FVector4), true, TEXT("DistanceFields.DFAssetDataUploadBuffer"));
+	AssetDataUploadBuffer.Init(AssetDataUploads.Num(), AssetDataMipStrideFloat4s * sizeof(FVector4f), true, TEXT("DistanceFields.DFAssetDataUploadBuffer"));
 
 	for (FDistanceFieldAssetMipId AssetMipUpload : AssetDataUploads)
 	{
