Fix crash when cooking OIT shader for certain platforms.

diff --git a/Engine/Source/Runtime/Renderer/Private/OIT/OIT.cpp b/Engine/Source/Runtime/Renderer/Private/OIT/OIT.cpp
--- a/Engine/Source/Runtime/Renderer/Private/OIT/OIT.cpp
+++ b/Engine/Source/Runtime/Renderer/Private/OIT/OIT.cpp
@@ -99,7 +99,7 @@ public:
 
 static bool IsOITSupported(EShaderPlatform InShaderPlatform)
 {
-	return RHISupportsComputeShaders(InShaderPlatform);
+	return RHISupportsComputeShaders(InShaderPlatform) && !IsMobilePlatform(InShaderPlatform) && !FDataDrivenShaderPlatformInfo::GetIsHlslcc(InShaderPlatform);
 }
 
 static void RemoveAllocation(FSortedIndexBuffer* InBuffer)
