Add the control to turn on/off of SSR on translucency dynamically using showflag and r.ssr.quality.

diff --git a/Engine/Shaders/Private/ForwardLightingCommon.ush b/Engine/Shaders/Private/ForwardLightingCommon.ush
--- a/Engine/Shaders/Private/ForwardLightingCommon.ush
+++ b/Engine/Shaders/Private/ForwardLightingCommon.ush
@@ -439,7 +439,7 @@ half3 GetImageBasedReflectionLighting(FMaterialPixelParameters MaterialParameter
 	}
 
 #if MATERIAL_SSR && !FORWARD_SHADING
-	if( View.CameraCut == 0 )
+	if( View.CameraCut == 0 && TranslucentBasePass.SSRQuality > 0)
 	{
 		float StepOffset = InterleavedGradientNoise( MaterialParameters.SvPosition.xy, View.StateFrameIndexMod8 );
 		StepOffset -= 0.5;
diff --git a/Engine/Source/Runtime/Renderer/Private/BasePassRendering.h b/Engine/Source/Runtime/Renderer/Private/BasePassRendering.h
--- a/Engine/Source/Runtime/Renderer/Private/BasePassRendering.h
+++ b/Engine/Source/Runtime/Renderer/Private/BasePassRendering.h
@@ -102,6 +102,7 @@ BEGIN_GLOBAL_SHADER_PARAMETER_STRUCT(FTranslucentBasePassUniformParameters,)
 	SHADER_PARAMETER(FVector4f, HZBUvFactorAndInvFactor)
 	SHADER_PARAMETER(FVector4f, PrevScreenPositionScaleBias)
 	SHADER_PARAMETER(float, PrevSceneColorPreExposureInv)
+	SHADER_PARAMETER(int32, SSRQuality)
 	SHADER_PARAMETER_RDG_TEXTURE(Texture2D, HZBTexture)
 	SHADER_PARAMETER_SAMPLER(SamplerState, HZBSampler)
 	SHADER_PARAMETER_RDG_TEXTURE(Texture2D, PrevSceneColor)
diff --git a/Engine/Source/Runtime/Renderer/Private/TranslucentRendering.cpp b/Engine/Source/Runtime/Renderer/Private/TranslucentRendering.cpp
--- a/Engine/Source/Runtime/Renderer/Private/TranslucentRendering.cpp
+++ b/Engine/Source/Runtime/Renderer/Private/TranslucentRendering.cpp
@@ -170,6 +170,33 @@ static bool IsSeparateTranslucencyEnabled(ETranslucencyPass::Type TranslucencyPa
 	return false;
 }
 
+static int GetSSRQuality()
+{
+	static IConsoleVariable* CVar = IConsoleManager::Get().FindConsoleVariable(TEXT("r.SSR.Quality"));
+	int SSRQuality = CVar ? (CVar->GetInt()) : 0;
+	return SSRQuality;
+}
+
+static bool ShouldRenderTranslucencyScreenSpaceReflections(const FViewInfo& View)
+{
+	// The screenspace reflection of translucency is not controlled by the postprocessing setting
+	// or the raytracing overlay setting. It needs to be turned on/off dynamically to support
+	// diffuse only
+	if (!View.Family->EngineShowFlags.ScreenSpaceReflections)
+	{
+		return false;
+	}
+
+	int SSRQuality = GetSSRQuality();
+
+	if (SSRQuality <= 0)
+	{
+		return false;
+	}
+
+	return true;
+}
+
 static void AddBeginTranslucencyTimerPass(FRDGBuilder& GraphBuilder, const FViewInfo& View)
 {
 #if STATS
@@ -861,6 +888,7 @@ TRDGUniformBufferRef<FTranslucentBasePassUniformParameters> CreateTranslucentBas
 
 		BasePassParameters.PrevScreenPositionScaleBias = ScreenPosToPixelValue;
 		BasePassParameters.PrevSceneColorPreExposureInv = PrevSceneColorPreExposureInvValue;
+		BasePassParameters.SSRQuality = ShouldRenderTranslucencyScreenSpaceReflections(View) ? GetSSRQuality() : 0;
 	}
 
 	// Translucency Lighting Volume
