Fix sphere geometry vertex buffer to use FVector3f, rather than FVector

diff --git a/Engine/Source/Runtime/Renderer/Private/LightRendering.cpp b/Engine/Source/Runtime/Renderer/Private/LightRendering.cpp
--- a/Engine/Source/Runtime/Renderer/Private/LightRendering.cpp
+++ b/Engine/Source/Runtime/Renderer/Private/LightRendering.cpp
@@ -334,7 +334,7 @@ void StencilingGeometry::DrawCone(FRHICommandList& RHICmdList)
 
 /** The stencil sphere vertex buffer. */
 TGlobalResource<StencilingGeometry::TStencilSphereVertexBuffer<18, 12, FVector4f> > StencilingGeometry::GStencilSphereVertexBuffer;
-TGlobalResource<StencilingGeometry::TStencilSphereVertexBuffer<18, 12, FVector> > StencilingGeometry::GStencilSphereVectorBuffer;
+TGlobalResource<StencilingGeometry::TStencilSphereVertexBuffer<18, 12, FVector3f> > StencilingGeometry::GStencilSphereVectorBuffer;
 
 /** The stencil sphere index buffer. */
 TGlobalResource<StencilingGeometry::TStencilSphereIndexBuffer<18, 12> > StencilingGeometry::GStencilSphereIndexBuffer;
diff --git a/Engine/Source/Runtime/Renderer/Private/LightRendering.h b/Engine/Source/Runtime/Renderer/Private/LightRendering.h
--- a/Engine/Source/Runtime/Renderer/Private/LightRendering.h
+++ b/Engine/Source/Runtime/Renderer/Private/LightRendering.h
@@ -126,6 +126,7 @@ namespace StencilingGeometry
 	class TStencilSphereVertexBuffer : public FVertexBuffer
 	{
 	public:
+		static_assert(std::is_same_v<typename VectorType::FReal, float>, "Must be a float vector type");
 
 		int32 GetNumRings() const
 		{
@@ -357,7 +358,7 @@ namespace StencilingGeometry
 	};
 
 	extern TGlobalResource<TStencilSphereVertexBuffer<18, 12, FVector4f> >	GStencilSphereVertexBuffer;
-	extern TGlobalResource<TStencilSphereVertexBuffer<18, 12, FVector> >	GStencilSphereVectorBuffer;
+	extern TGlobalResource<TStencilSphereVertexBuffer<18, 12, FVector3f> >	GStencilSphereVectorBuffer;
 	extern TGlobalResource<TStencilSphereIndexBuffer<18, 12> >				GStencilSphereIndexBuffer;
 	extern TGlobalResource<TStencilSphereVertexBuffer<4, 4, FVector4f> >		GLowPolyStencilSphereVertexBuffer;
 	extern TGlobalResource<TStencilSphereIndexBuffer<4, 4> >				GLowPolyStencilSphereIndexBuffer;
