Fixed broken CustomDepth rendering for translucent materials on mobile (CL# 17646057 in UE4)

diff --git a/Engine/Shaders/Private/MaterialTemplate.ush b/Engine/Shaders/Private/MaterialTemplate.ush
--- a/Engine/Shaders/Private/MaterialTemplate.ush
+++ b/Engine/Shaders/Private/MaterialTemplate.ush
@@ -1975,8 +1975,20 @@ uint bMetalMSAAHDRDecode;
 
 MaterialFloat4 MobileSceneTextureLookup(inout FMaterialPixelParameters Parameters, int SceneTextureId, float2 UV)
 {
-#if (FEATURE_LEVEL <= FEATURE_LEVEL_ES3_1)
-
+#if SCENE_TEXTURES_DISABLED
+	// When scene textures are disabled, the output is matched to the dummy scene texture defaults.
+	switch(SceneTextureId)
+	{
+	case PPI_SceneDepth:
+	case PPI_CustomDepth:
+		return ConvertFromDeviceZ(0.0f);
+	case PPI_MaterialAO:
+	case PPI_CustomStencil:
+		return 1.0f;
+	default:
+		return 0.0f;
+	}
+#else
 	if (SceneTextureId == PPI_SceneDepth)
 	{
 		#if MOBILE_DEFERRED_SHADING
@@ -2016,7 +2028,7 @@ MaterialFloat4 MobileSceneTextureLookup(inout FMaterialPixelParameters Parameter
 		Stencil = floor(Stencil + 0.5);
 		return MaterialFloat4(Stencil.rrr, 0.f);
 	}
-#endif// FEATURE_LEVEL
+#endif// SCENE_TEXTURES_DISABLED
 
 	return MaterialFloat4(0.0f, 0.0f, 0.0f, 0.0f);
 }
diff --git a/Engine/Source/Runtime/Renderer/Private/SceneTextures.cpp b/Engine/Source/Runtime/Renderer/Private/SceneTextures.cpp
--- a/Engine/Source/Runtime/Renderer/Private/SceneTextures.cpp
+++ b/Engine/Source/Runtime/Renderer/Private/SceneTextures.cpp
@@ -1104,20 +1104,22 @@ void SetupMobileSceneTextureUniformParameters(
 	const FRDGSystemTextures& SystemTextures = FRDGSystemTextures::Get(GraphBuilder);
 
 	SceneTextureParameters.SceneColorTexture = SystemTextures.Black;
-	SceneTextureParameters.SceneColorTextureSampler = TStaticSamplerState<SF_Point, AM_Clamp, AM_Clamp, AM_Clamp>::GetRHI();
+	SceneTextureParameters.SceneColorTextureSampler = TStaticSamplerState<>::GetRHI();
 	SceneTextureParameters.SceneDepthTexture = SystemTextures.DepthDummy;
-	SceneTextureParameters.SceneDepthTextureSampler = TStaticSamplerState<SF_Point, AM_Clamp, AM_Clamp, AM_Clamp>::GetRHI();
-	SceneTextureParameters.CustomDepthTexture = SystemTextures.MaxFP16Depth;
+	SceneTextureParameters.SceneDepthTextureSampler = TStaticSamplerState<>::GetRHI();
+	// CustomDepthTexture is a color texture on mobile, with DeviceZ values
+	SceneTextureParameters.CustomDepthTexture = SystemTextures.Black;
 	SceneTextureParameters.CustomDepthTextureSampler = TStaticSamplerState<>::GetRHI();
 	SceneTextureParameters.MobileCustomStencilTexture = SystemTextures.Black;
-	SceneTextureParameters.MobileCustomStencilTextureSampler = TStaticSamplerState<SF_Point, AM_Clamp, AM_Clamp, AM_Clamp>::GetRHI();
+	SceneTextureParameters.MobileCustomStencilTextureSampler = TStaticSamplerState<>::GetRHI();
 	SceneTextureParameters.SceneVelocityTexture = SystemTextures.Black;
-	SceneTextureParameters.SceneVelocityTextureSampler = TStaticSamplerState<SF_Point, AM_Clamp, AM_Clamp, AM_Clamp>::GetRHI();
+	SceneTextureParameters.SceneVelocityTextureSampler = TStaticSamplerState<>::GetRHI();
 	SceneTextureParameters.GBufferATexture = SystemTextures.Black;
 	SceneTextureParameters.GBufferBTexture = SystemTextures.Black;
 	SceneTextureParameters.GBufferCTexture = SystemTextures.Black;
 	SceneTextureParameters.GBufferDTexture = SystemTextures.Black;
-	SceneTextureParameters.SceneDepthAuxTexture = SystemTextures.MaxFP16Depth;
+	// SceneDepthAuxTexture is a color texture on mobile, with DeviceZ values
+	SceneTextureParameters.SceneDepthAuxTexture = SystemTextures.Black;
 	SceneTextureParameters.GBufferATextureSampler = TStaticSamplerState<>::GetRHI();
 	SceneTextureParameters.GBufferBTextureSampler = TStaticSamplerState<>::GetRHI();
 	SceneTextureParameters.GBufferCTextureSampler = TStaticSamplerState<>::GetRHI();
