Fixed SceneColor alpha on mobile to be consistent with desktop rendering

diff --git a/Engine/Shaders/Private/MobileBasePassPixelShader.usf b/Engine/Shaders/Private/MobileBasePassPixelShader.usf
--- a/Engine/Shaders/Private/MobileBasePassPixelShader.usf
+++ b/Engine/Shaders/Private/MobileBasePassPixelShader.usf
@@ -1047,13 +1047,13 @@ void Main(
 	#else
 		OutColor.rgb = Color * VertexFog.a + VertexFog.rgb;
 
-		#if !MATERIAL_USE_ALPHA_TO_COVERAGE
-			// Scene color alpha is not used yet so we set it to 1
-			OutColor.a = 1.0;
-			
 		#if USE_SCENE_DEPTH_AUX
 			OutSceneDepthAux = SvPosition.z;
 		#endif
+
+		#if !MATERIAL_USE_ALPHA_TO_COVERAGE
+			// Planar reflections and scene captures use scene color alpha to keep track of where content has been rendered, for compositing into a different scene later
+			OutColor.a = 0.0;
 		#else
 			half MaterialOpacityMask = GetMaterialMaskInputRaw(PixelMaterialInputs);
 			OutColor.a = GetMaterialMask(PixelMaterialInputs) / max(abs(ddx(MaterialOpacityMask)) + abs(ddy(MaterialOpacityMask)), 0.0001f) + 0.5f;
diff --git a/Engine/Source/Runtime/Renderer/Private/MobileBasePass.cpp b/Engine/Source/Runtime/Renderer/Private/MobileBasePass.cpp
--- a/Engine/Source/Runtime/Renderer/Private/MobileBasePass.cpp
+++ b/Engine/Source/Runtime/Renderer/Private/MobileBasePass.cpp
@@ -423,11 +423,11 @@ void MobileBasePass::SetTranslucentRenderState(FMeshPassProcessorRenderState& Dr
 			break;
 		case BLEND_Additive:
 			// Add to the existing scene color
-			DrawRenderState.SetBlendState(TStaticBlendState<CW_RGB, BO_Add, BF_One, BF_One, BO_Add, BF_Zero, BF_InverseSourceAlpha>::GetRHI());
+			DrawRenderState.SetBlendState(TStaticBlendState<CW_RGBA, BO_Add, BF_One, BF_One, BO_Add, BF_Zero, BF_InverseSourceAlpha>::GetRHI());
 			break;
 		case BLEND_Modulate:
-			// Modulate with the existing scene color
-			DrawRenderState.SetBlendState(TStaticBlendState<CW_RGBA, BO_Add, BF_DestColor, BF_Zero>::GetRHI());
+			// Modulate with the existing scene color, preserve destination alpha.
+			DrawRenderState.SetBlendState(TStaticBlendState<CW_RGB, BO_Add, BF_DestColor, BF_Zero>::GetRHI());
 			break;
 		case BLEND_AlphaComposite:
 			// Blend with existing scene color. New color is already pre-multiplied by alpha.
@@ -441,7 +441,7 @@ void MobileBasePass::SetTranslucentRenderState(FMeshPassProcessorRenderState& Dr
 			if (Material.GetShadingModels().HasShadingModel(MSM_SingleLayerWater))
 			{
 				// Single layer water is an opaque marerial rendered as translucent on Mobile. We force pre-multiplied alpha to achieve water depth based transmittance.
-				DrawRenderState.SetBlendState(TStaticBlendState<CW_RGB, BO_Add, BF_One, BF_InverseSourceAlpha, BO_Add, BF_Zero, BF_InverseSourceAlpha>::GetRHI());
+				DrawRenderState.SetBlendState(TStaticBlendState<CW_RGBA, BO_Add, BF_One, BF_InverseSourceAlpha, BO_Add, BF_Zero, BF_InverseSourceAlpha>::GetRHI());
 			}
 			else
 			{
diff --git a/Engine/Source/Runtime/Renderer/Private/SceneTextures.cpp b/Engine/Source/Runtime/Renderer/Private/SceneTextures.cpp
--- a/Engine/Source/Runtime/Renderer/Private/SceneTextures.cpp
+++ b/Engine/Source/Runtime/Renderer/Private/SceneTextures.cpp
@@ -107,11 +107,22 @@ static EPixelFormat GetGBufferFFormat()
 	return NormalGBufferFormat;
 }
 
-static EPixelFormat GetMobileSceneColorFormat(const FSceneView& View)
+static EPixelFormat GetMobileSceneColorFormat(const FSceneViewFamily& ViewFamily)
 {
+	bool bRequiresAlphaChannel = IsMobilePropagateAlphaEnabled(ViewFamily.GetShaderPlatform());
+
+	for (int32 ViewIndex = 0; ViewIndex < ViewFamily.Views.Num(); ViewIndex++)
+	{
+		// Planar reflections and scene captures use scene color alpha to keep track of where content has been rendered, for compositing into a different scene later
+		if (ViewFamily.Views[ViewIndex]->bIsPlanarReflection || ViewFamily.Views[ViewIndex]->bIsSceneCapture)
+		{
+			bRequiresAlphaChannel = true;
+		}
+	}
+	
 	const EPixelFormat DefaultLowPrecisionFormat = IHeadMountedDisplayModule::IsAvailable() && IHeadMountedDisplayModule::Get().IsStandaloneStereoOnlyDevice()
 		? PF_R8G8B8A8 : PF_B8G8R8A8;
-	const EPixelFormat DefaultPrecisionFormat = IsMobilePropagateAlphaEnabled(View.GetShaderPlatform()) || View.bIsSceneCapture? PF_FloatRGBA : PF_FloatR11G11B10;
+	const EPixelFormat DefaultPrecisionFormat = bRequiresAlphaChannel ? PF_FloatRGBA : PF_FloatR11G11B10;
 
 	EPixelFormat DefaultColorFormat = (!IsMobileHDR() || !GSupportsRenderTargetFormat_PF_FloatRGBA) ? DefaultLowPrecisionFormat : DefaultPrecisionFormat;
 
@@ -535,12 +546,7 @@ FSceneTexturesConfig FSceneTexturesConfig::Create(const FSceneViewFamily& ViewFa
 
 	case EShadingPath::Mobile:
 	{
-		Config.ColorFormat = GetMobileSceneColorFormat(*(ViewFamily.Views[0]));
-
-		// On mobile the scene depth is calculated from the alpha component of the scene color
-		// Use FarPlane for alpha to ensure un-rendered pixels have max depth...
-		const float DepthFar = (float)ERHIZBuffer::FarPlane;
-		Config.ColorClearValue = FClearValueBinding(FLinearColor(0.f, 0.f, 0.f, DepthFar));
+		Config.ColorFormat = GetMobileSceneColorFormat(ViewFamily);
 		break;
 	}
 
