Fix a non-unity compile error

diff --git a/Engine/Source/Runtime/Renderer/Private/ShadowSetupMobile.cpp b/Engine/Source/Runtime/Renderer/Private/ShadowSetupMobile.cpp
--- a/Engine/Source/Runtime/Renderer/Private/ShadowSetupMobile.cpp
+++ b/Engine/Source/Runtime/Renderer/Private/ShadowSetupMobile.cpp
@@ -16,6 +16,7 @@ ShadowSetupMobile.cpp: Shadow setup implementation for mobile specific features.
 #include "DynamicPrimitiveDrawing.h"
 #include "ScenePrivate.h"
 #include "MeshPassProcessor.h"
+#include "UObject/UObjectIterator.h"
 
 static TAutoConsoleVariable<int32> CVarCsmShaderCullingDebugGfx(
 	TEXT("r.Mobile.Shadow.CSMShaderCullingDebugGfx"),
