Fixes for compile errors in shader draw debug
#rnx

diff --git a/Engine/Source/Runtime/Renderer/Private/PhysicsFieldRendering.cpp b/Engine/Source/Runtime/Renderer/Private/PhysicsFieldRendering.cpp
--- a/Engine/Source/Runtime/Renderer/Private/PhysicsFieldRendering.cpp
+++ b/Engine/Source/Runtime/Renderer/Private/PhysicsFieldRendering.cpp
@@ -59,7 +59,7 @@ class FPhysicsFieldRayMarchingCS : public FGlobalShader
 
 	BEGIN_SHADER_PARAMETER_STRUCT(FParameters, )
 		SHADER_PARAMETER_STRUCT_INCLUDE(FSceneTextureParameters, SceneTextures)
-		SHADER_PARAMETER_STRUCT_INCLUDE(ShaderDrawDebug::FShaderDrawDebugParameters, ShaderDrawParameters)
+		SHADER_PARAMETER_STRUCT_INCLUDE(ShaderDrawDebug::FShaderParameters, ShaderDrawParameters)
 		SHADER_PARAMETER_STRUCT_INCLUDE(ShaderPrint::FShaderParameters, ShaderPrintParameters)
 		SHADER_PARAMETER_STRUCT_REF(FViewUniformShaderParameters, ViewUniformBuffer)
 		SHADER_PARAMETER_SRV(Buffer<float4>, BoundsMin)
@@ -183,7 +183,7 @@ void RenderPhysicsField(
 	{
 		FPhysicsFieldResource* PhysicsFieldResource = (GPhysicsFieldSystemType == 0) ? PhysicsFieldProxy->DebugResource : PhysicsFieldProxy->FieldResource;
 
-		if (Views.Num() > 0 && PhysicsFieldResource && ShaderDrawDebug::IsShaderDrawDebugEnabled() && ShaderPrint::IsEnabled() &&
+		if (Views.Num() > 0 && PhysicsFieldResource && ShaderDrawDebug::IsEnabled() && ShaderPrint::IsEnabled() &&
 			ShaderPrint::IsSupported(Views[0]) && GDynamicRHI->RHIIsTypedUAVLoadSupported(PF_FloatRGBA))
 		{
 			AddPhysicsFieldRayMarchingPass(GraphBuilder, Views[0], PhysicsFieldResource, SceneColorTexture);
diff --git a/Engine/Source/Runtime/Renderer/Private/RendererScene.cpp b/Engine/Source/Runtime/Renderer/Private/RendererScene.cpp
--- a/Engine/Source/Runtime/Renderer/Private/RendererScene.cpp
+++ b/Engine/Source/Runtime/Renderer/Private/RendererScene.cpp
@@ -1933,7 +1933,7 @@ void FScene::ShowPhysicsField()
 			ShaderPrint::SetEnabled(true);
 			ShaderPrint::SetFontSize(8);
 		}
-		if (!ShaderDrawDebug::IsShaderDrawDebugEnabled())
+		if (!ShaderDrawDebug::IsEnabled())
 		{
 			ShaderDrawDebug::SetEnabled(true);
 		}
