Re-add cvar r.Mobile.EarlyZPass

diff --git a/Engine/Source/Runtime/Renderer/Private/RendererScene.cpp b/Engine/Source/Runtime/Renderer/Private/RendererScene.cpp
--- a/Engine/Source/Runtime/Renderer/Private/RendererScene.cpp
+++ b/Engine/Source/Runtime/Renderer/Private/RendererScene.cpp
@@ -108,6 +108,16 @@ TAutoConsoleVariable<int32> CVarEarlyZPass(
 	TEXT("  x: use built in heuristic (default is 3)"),
 	ECVF_Scalability);
 
+static TAutoConsoleVariable<int32> CVarMobileEarlyZPass(
+	TEXT("r.Mobile.EarlyZPass"),
+	0,
+	TEXT("Whether to use a depth only pass to initialize Z culling for the mobile base pass.\n")
+	TEXT("  0: off\n")
+	TEXT("  1: all opaque \n"),
+	ECVF_Scalability
+);
+
+
 static TAutoConsoleVariable<int32> CVarBasePassWriteDepthEvenWithFullPrepass(
 	TEXT("r.BasePassWriteDepthEvenWithFullPrepass"),
 	0,
@@ -691,12 +701,15 @@ static void UpdateEarlyZPassModeCVarSinkFunction()
 {
 	static int32 CachedEarlyZPass = CVarEarlyZPass.GetValueOnGameThread();
 	static int32 CachedBasePassWriteDepthEvenWithFullPrepass = CVarBasePassWriteDepthEvenWithFullPrepass.GetValueOnGameThread();
+	static int32 CachedMobileEarlyZPass = CVarMobileEarlyZPass.GetValueOnGameThread();
 
 	const int32 EarlyZPass = CVarEarlyZPass.GetValueOnGameThread();
 	const int32 BasePassWriteDepthEvenWithFullPrepass = CVarBasePassWriteDepthEvenWithFullPrepass.GetValueOnGameThread();
+	const int32 MobileEarlyZPass = CVarMobileEarlyZPass.GetValueOnGameThread();
 
 	if (EarlyZPass != CachedEarlyZPass
-		|| BasePassWriteDepthEvenWithFullPrepass != CachedBasePassWriteDepthEvenWithFullPrepass)
+		|| BasePassWriteDepthEvenWithFullPrepass != CachedBasePassWriteDepthEvenWithFullPrepass
+		|| MobileEarlyZPass != CachedMobileEarlyZPass)
 	{
 		for (TObjectIterator<UWorld> It; It; ++It)
 		{
@@ -710,6 +723,7 @@ static void UpdateEarlyZPassModeCVarSinkFunction()
 
 		CachedEarlyZPass = EarlyZPass;
 		CachedBasePassWriteDepthEvenWithFullPrepass = BasePassWriteDepthEvenWithFullPrepass;
+		CachedMobileEarlyZPass = MobileEarlyZPass;
 	}
 }
 
@@ -3553,6 +3567,12 @@ void FScene::UpdateEarlyZPassMode()
 		{
 			EarlyZPassMode = DDM_AllOpaque;
 		}
+
+		bool bMobileForceFullDepthPrepass = CVarMobileEarlyZPass.GetValueOnAnyThread() == 1;
+		if (bMobileForceFullDepthPrepass)
+		{
+			EarlyZPassMode = DDM_AllOpaque;
+		}
 	}
 }
 
