Null check eye adaptation meter texture.

diff --git a/Engine/Source/Runtime/Renderer/Private/PostProcess/PostProcessEyeAdaptation.cpp b/Engine/Source/Runtime/Renderer/Private/PostProcess/PostProcessEyeAdaptation.cpp
--- a/Engine/Source/Runtime/Renderer/Private/PostProcess/PostProcessEyeAdaptation.cpp
+++ b/Engine/Source/Runtime/Renderer/Private/PostProcess/PostProcessEyeAdaptation.cpp
@@ -353,10 +353,18 @@ FEyeAdaptationParameters GetEyeAdaptationParameters(const FViewInfo& View, ERHIF
 	// mode does the calculation in pre-exposure space, which is why we need to multiply by View.PreExposure.
 	const float LuminanceMin = (AutoExposureMethod == AEM_Basic) ? 0.0001f : FMath::Exp2(HistogramLogMin);
 
-	//AutoExposureMeterMask
-	const FTextureRHIRef MeterMask = Settings.AutoExposureMeterMask ?
-		Settings.AutoExposureMeterMask->GetResource()->TextureRHI :
-		GWhiteTexture->TextureRHI;
+	FTextureRHIRef MeterMask = nullptr;
+
+	if (Settings.AutoExposureMeterMask &&
+		Settings.AutoExposureMeterMask->GetResource() &&
+		Settings.AutoExposureMeterMask->GetResource()->TextureRHI)
+	{
+		MeterMask = Settings.AutoExposureMeterMask->GetResource()->TextureRHI;
+	}
+	else
+	{
+		MeterMask = GWhiteTexture->TextureRHI;
+	}
 
 	// The distance at which we switch from linear to exponential. I.e. at StartDistance=1.5, when linear is 1.5 f-stops away from hitting the 
 	// target, we switch to exponential.
