Invalidate path tracer accumulation when a material instance parameter changes

This is similar to the previous implementation, but loops over the active scenes from the game thread to avoid thread safety issues with the AllocatedScenes member of the render module.

diff --git a/Engine/Source/Runtime/Engine/Private/Materials/MaterialInstance.cpp b/Engine/Source/Runtime/Engine/Private/Materials/MaterialInstance.cpp
--- a/Engine/Source/Runtime/Engine/Private/Materials/MaterialInstance.cpp
+++ b/Engine/Source/Runtime/Engine/Private/Materials/MaterialInstance.cpp
@@ -339,6 +339,7 @@ void GameThread_UpdateMIParameter(const UMaterialInstance* Instance, const Param
 {
 	if (FApp::CanEverRender())
 	{
+		GetRendererModule().InvalidatePathTracedOutput();
 		FMaterialInstanceResource* Resource = Instance->Resource;
 		const FMaterialParameterInfo& ParameterInfo = Parameter.ParameterInfo;
 		typename ParameterType::ValueType Value = ParameterType::GetValue(Parameter);
diff --git a/Engine/Source/Runtime/Engine/Public/SceneInterface.h b/Engine/Source/Runtime/Engine/Public/SceneInterface.h
--- a/Engine/Source/Runtime/Engine/Public/SceneInterface.h
+++ b/Engine/Source/Runtime/Engine/Public/SceneInterface.h
@@ -202,6 +202,9 @@ public:
 	/** Invalidates pages in a runtime virtual texture object. */
 	virtual void InvalidateRuntimeVirtualTexture(class URuntimeVirtualTextureComponent* Component, FBoxSphereBounds const& WorldBounds) {}
 
+	/** Mark scene as needing to restart path tracer accumulation. */
+	virtual void InvalidatePathTracedOutput() {}
+
 	/** 
 	 * Retrieves primitive uniform shader parameters that are internal to the renderer.
 	 */
diff --git a/Engine/Source/Runtime/RenderCore/Public/RendererInterface.h b/Engine/Source/Runtime/RenderCore/Public/RendererInterface.h
--- a/Engine/Source/Runtime/RenderCore/Public/RendererInterface.h
+++ b/Engine/Source/Runtime/RenderCore/Public/RendererInterface.h
@@ -876,5 +876,8 @@ public:
 	 * The intended use is for stand-alone rendering that involves Scene proxies (that then may need the machinery to render GPU-Scene aware primitives.
 	 */
 	virtual IScenePrimitiveRenderingContext* BeginScenePrimitiveRendering(FRDGBuilder& GraphBuilder, FSceneViewFamily* ViewFamily) = 0;
+
+	/** Mark all the current scenes as needing to restart path tracer accumulation */
+	virtual void InvalidatePathTracedOutput() = 0;
 };
 
diff --git a/Engine/Source/Runtime/Renderer/Private/RendererModule.h b/Engine/Source/Runtime/Renderer/Private/RendererModule.h
--- a/Engine/Source/Runtime/Renderer/Private/RendererModule.h
+++ b/Engine/Source/Runtime/Renderer/Private/RendererModule.h
@@ -120,6 +120,8 @@ public:
 
 	virtual IScenePrimitiveRenderingContext* BeginScenePrimitiveRendering(FRDGBuilder& GraphBuilder, FSceneViewFamily* ViewFamily) override;
 
+	virtual void InvalidatePathTracedOutput() override;
+
 private:
 	TSet<FSceneInterface*> AllocatedScenes;
 	FOnPostOpaqueRender PostOpaqueRenderDelegate;
diff --git a/Engine/Source/Runtime/Renderer/Private/RendererScene.cpp b/Engine/Source/Runtime/Renderer/Private/RendererScene.cpp
--- a/Engine/Source/Runtime/Renderer/Private/RendererScene.cpp
+++ b/Engine/Source/Runtime/Renderer/Private/RendererScene.cpp
@@ -2671,6 +2671,12 @@ void FScene::InvalidateRuntimeVirtualTexture(class URuntimeVirtualTextureCompone
 	}
 }
 
+void FScene::InvalidatePathTracedOutput()
+{
+	// NOTE: this is an atomic, so this function is ok to call from any thread
+	bPathTracingNeedsInvalidation = true;
+}
+
 void FScene::FlushDirtyRuntimeVirtualTextures()
 {
 	checkSlow(IsInRenderingThread());
@@ -4804,9 +4810,10 @@ void FRendererModule::RemoveScene(FSceneInterface* Scene)
 void FRendererModule::UpdateStaticDrawLists()
 {
 	// Update all static meshes in order to recache cached mesh draw commands.
-	for (TSet<FSceneInterface*>::TConstIterator SceneIt(AllocatedScenes); SceneIt; ++SceneIt)
+	check(IsInGameThread()); // AllocatedScenes is managed by the game thread
+	for (FSceneInterface* Scene : AllocatedScenes)
 	{
-		(*SceneIt)->UpdateStaticDrawLists();
+		Scene->UpdateStaticDrawLists();
 	}
 }
 
@@ -4888,6 +4895,15 @@ FSceneViewStateInterface* FRendererModule::AllocateViewState()
 	return new FSceneViewState();
 }
 
+void FRendererModule::InvalidatePathTracedOutput()
+{
+	check(IsInGameThread()); // AllocatedScenes is managed by the game thread
+	for (FSceneInterface* Scene : AllocatedScenes)
+	{
+		Scene->InvalidatePathTracedOutput();
+	}
+}
+
 //////////////////////////////////////////////////////////////////////////
 
 FLatentGPUTimer::FLatentGPUTimer(FRenderQueryPoolRHIRef InTimerQueryPool, int32 InAvgSamples)
diff --git a/Engine/Source/Runtime/Renderer/Private/ScenePrivate.h b/Engine/Source/Runtime/Renderer/Private/ScenePrivate.h
--- a/Engine/Source/Runtime/Renderer/Private/ScenePrivate.h
+++ b/Engine/Source/Runtime/Renderer/Private/ScenePrivate.h
@@ -2707,7 +2707,7 @@ public:
 	bool bScenesPrimitivesNeedStaticMeshElementUpdate;
 
 	/** True if a change to the scene that requires to invalidate the path tracer buffers has happened. */
-	bool bPathTracingNeedsInvalidation;
+	TAtomic<bool> bPathTracingNeedsInvalidation;
 
 #if RHI_RAYTRACING
 	/** What mode where the cached RT commands prepared for last? */
@@ -2988,6 +2988,7 @@ public:
 	virtual void RemoveRuntimeVirtualTexture(class URuntimeVirtualTextureComponent* Component) override;
 	virtual void GetRuntimeVirtualTextureHidePrimitiveMask(uint8& bHideMaskEditor, uint8& bHideMaskGame) const override;
 	virtual void InvalidateRuntimeVirtualTexture(class URuntimeVirtualTextureComponent* Component, FBoxSphereBounds const& WorldBounds) override;
+	virtual void InvalidatePathTracedOutput() override;
 	virtual void GetPrimitiveUniformShaderParameters_RenderThread(const FPrimitiveSceneInfo* PrimitiveSceneInfo, bool& bHasPrecomputedVolumetricLightmap, FMatrix& PreviousLocalToWorld, int32& SingleCaptureIndex, bool& bOutputVelocity) const override;
 	virtual void UpdateLightTransform(ULightComponent* Light) override;
 	virtual void UpdateLightColorAndBrightness(ULightComponent* Light) override;
diff --git a/Engine/Source/Runtime/Renderer/Private/SceneVisibility.cpp b/Engine/Source/Runtime/Renderer/Private/SceneVisibility.cpp
--- a/Engine/Source/Runtime/Renderer/Private/SceneVisibility.cpp
+++ b/Engine/Source/Runtime/Renderer/Private/SceneVisibility.cpp
@@ -3515,30 +3515,35 @@ void FSceneRenderer::PreVisibilityFrameSetup(FRDGBuilder& GraphBuilder, const FS
 				ViewState->PrevFrameViewInfo.ViewMatrices.GetViewOrigin(),
 				0.18f /*degree*/, 0.1f /*cm*/);
 			const bool bIsProjMatrixDifferent = View.ViewMatrices.GetProjectionNoAAMatrix() != View.ViewState->PrevFrameViewInfo.ViewMatrices.GetProjectionNoAAMatrix();
-			bool bInvalidatePathTracer = false;
 			
 			if (View.bIsOfflineRender)
 			{
 				// In the offline context, we want precise control over when to restart the path tracer's accumulation to allow for motion blur
 				// So we use the camera cut signal only. In particular - we should not use bForceCameraVisibilityReset since this has
 				// interactions with the motion blur post process effect in tiled rendering (see comment below).
-				bInvalidatePathTracer = View.bCameraCut || View.bForcePathTracerReset;
+				if (View.bCameraCut || View.bForcePathTracerReset)
+				{
+					ViewState->PathTracingInvalidate();
+				}
 			}
 			else
 			{
 				// for interactive usage - any movement or scene change should restart the path tracer
-				bInvalidatePathTracer = bResetCamera ||
-					Scene->bPathTracingNeedsInvalidation ||
+
+				// We use an atomic to signal path tracer invalidation easily from other threads (like the game thread). This means
+				// we need to be very careful to read and reset the atomic flag in a single step to avoid missing anything. This is
+				// also why we check this field first before looking at the other conditions.
+				bool bNeedsInvalidation = true;
+				if (Scene->bPathTracingNeedsInvalidation.CompareExchange(bNeedsInvalidation, false) ||
+					bResetCamera ||
 					bIsProjMatrixDifferent ||
 					bIsCameraMove ||
-					View.bForcePathTracerReset;
+					View.bForcePathTracerReset)
+				{
+					ViewState->PathTracingInvalidate();
+				}
 			}
 
-			if (bInvalidatePathTracer)
-			{
-				ViewState->PathTracingInvalidate();
-				Scene->bPathTracingNeedsInvalidation = false;
-			}
 #endif // RHI_RAYTRACING
 
 			if (bResetCamera)
