Use the input viewport to decide when we need to manually clamp UVs in the gaussian blur shader, rather than the output viewport.

#rnx

diff --git a/Engine/Source/Runtime/Renderer/Private/PostProcess/PostProcessWeightedSampleSum.cpp b/Engine/Source/Runtime/Renderer/Private/PostProcess/PostProcessWeightedSampleSum.cpp
--- a/Engine/Source/Runtime/Renderer/Private/PostProcess/PostProcessWeightedSampleSum.cpp
+++ b/Engine/Source/Runtime/Renderer/Private/PostProcess/PostProcessWeightedSampleSum.cpp
@@ -355,7 +355,7 @@ FScreenPassTexture AddGaussianBlurPass(
 
 	const bool bCombineAdditive = Additive.IsValid();
 
-	const bool bManualUVBorder = !OutputViewport.IsFullscreen();
+	const bool bManualUVBorder = Filter.ViewRect.Min != FIntPoint::ZeroValue || Filter.ViewRect.Max != Filter.Texture->Desc.Extent;
 
 	const uint32 SampleCount = SampleOffsets.Num();
 
