Allow material updates to happen from the parallel game thread as well as the game thread

diff --git a/Engine/Source/Runtime/Renderer/Private/RendererScene.cpp b/Engine/Source/Runtime/Renderer/Private/RendererScene.cpp
--- a/Engine/Source/Runtime/Renderer/Private/RendererScene.cpp
+++ b/Engine/Source/Runtime/Renderer/Private/RendererScene.cpp
@@ -4897,7 +4897,14 @@ FSceneViewStateInterface* FRendererModule::AllocateViewState()
 
 void FRendererModule::InvalidatePathTracedOutput()
 {
-	check(IsInGameThread()); // AllocatedScenes is managed by the game thread
+	// AllocatedScenes is managed by the game thread
+
+	// #jira UE-130700:
+	// Because material updates call this function and could happen in parallel, we also allow the parallel game thread here.
+	// We assume that no changes will be made to AllocatedScene during this time, otherwise locking would need to
+	// be introduced (which could have performance implications). 
+
+	check(IsInGameThread() || IsInParallelGameThread());
 	for (FSceneInterface* Scene : AllocatedScenes)
 	{
 		Scene->InvalidatePathTracedOutput();
