Allow vertex fogging on mobile by default

diff --git a/Engine/Source/Runtime/Renderer/Private/MobileBasePassRendering.cpp b/Engine/Source/Runtime/Renderer/Private/MobileBasePassRendering.cpp
--- a/Engine/Source/Runtime/Renderer/Private/MobileBasePassRendering.cpp
+++ b/Engine/Source/Runtime/Renderer/Private/MobileBasePassRendering.cpp
@@ -22,7 +22,7 @@
 // Changing this causes a full shader recompile
 static TAutoConsoleVariable<int32> CVarMobileDisableVertexFog(
 	TEXT("r.Mobile.DisableVertexFog"),
-	1,
+	0,
 	TEXT("Set to 1 to disable vertex fogging in all mobile shaders."),
 	ECVF_ReadOnly | ECVF_RenderThreadSafe);
 
