Added experimental local light support to volumetric cloud.
This is slow and only aims at being used for cinematics.
It is only available for the compute path to reduce permutation count.

diff --git a/Engine/Shaders/Private/VolumetricCloud.usf b/Engine/Shaders/Private/VolumetricCloud.usf
--- a/Engine/Shaders/Private/VolumetricCloud.usf
+++ b/Engine/Shaders/Private/VolumetricCloud.usf
@@ -28,6 +28,17 @@
 #define VIRTUAL_SHADOW_MAP 0
 #endif
 
+#ifndef CLOUD_SAMPLE_LOCAL_LIGHTS
+#define CLOUD_SAMPLE_LOCAL_LIGHTS 0
+#endif
+
+#if CLOUD_SAMPLE_LOCAL_LIGHTS
+#define ForwardLightData RenderVolumetricCloudParameters.ForwardLightData
+#include "LightGridCommon.ush"
+#define SUPPORT_CONTACT_SHADOWS 0
+#include "DeferredLightingCommon.ush"
+#endif
+
 #if defined(SHADER_MAINVS) || defined(SHADER_SHADOW_PS) || SHADER_RENDERVIEW_PS || SHADER_RENDERVIEW_CS
 
 #include "/Engine/Generated/Material.ush"
@@ -334,7 +345,7 @@ ParticipatingMediaContext SetupParticipatingMediaContext(float3 BaseAlbedo, floa
 	PMC.TransmittanceToLight1[0] = InitialTransmittanceToLight1;
 #endif
 
-	[unroll]
+	UNROLL
 	for (int ms = 1; ms < MSCOUNT; ++ms)
 	{
 		PMC.ScatteringCoefficients[ms] = PMC.ScatteringCoefficients[ms - 1] * MsSFactor;
@@ -368,7 +379,7 @@ ParticipatingMediaPhaseContext SetupParticipatingMediaPhaseContext(float BasePha
 	PMPC.Phase1[0] = BasePhase1;
 #endif
 
-	[unroll]
+	UNROLL
 	for (int ms = 1; ms < MSCOUNT; ++ms)
 	{
 		PMPC.Phase0[ms] = lerp(IsotropicPhase(), PMPC.Phase0[0], MsPhaseFactor);
@@ -801,7 +812,7 @@ void MainCommon(in FMaterialPixelParameters MaterialParameters, in float4 SvPosi
 				}
 #endif // VIRTUALSHADOW_MAP
 
-				[unroll]
+				UNROLL
 				for (ms = 0; ms < MSCOUNT; ++ms)
 				{
 					PMC.TransmittanceToLight0[ms] *= OpaqueShadow;
@@ -821,7 +832,7 @@ void MainCommon(in FMaterialPixelParameters MaterialParameters, in float4 SvPosi
 				float OutOpticalDepth = 0.0f;
 				float CloudShadow = GetCloudVolumetricShadow(SampleWorldPosition, RenderVolumetricCloudParameters.CloudShadowmapWorldToLightClipMatrix[0], RenderVolumetricCloudParameters.CloudShadowmapFarDepthKm[0].x,
 					RenderVolumetricCloudParameters.CloudShadowTexture0, RenderVolumetricCloudParameters.CloudBilinearTextureSampler, OutOpticalDepth);
-				[unroll]
+				UNROLL
 				for (ms = 0; ms < MSCOUNT; ++ms)
 				{
 					PMC.TransmittanceToLight0[ms] *= exp(-OutOpticalDepth * (ms == 0 ? 1.0f : pow(MsExtinFactor, ms)));
@@ -829,7 +840,7 @@ void MainCommon(in FMaterialPixelParameters MaterialParameters, in float4 SvPosi
 #else 
 	// Use raymarched shadows
 				DebugShadowSampleCount++;
-				[unroll]
+				UNROLL
 				for (ms = 0; ms < MSCOUNT; ++ms)
 				{
 					ExtinctionAcc[ms] = 0.0f;
@@ -876,14 +887,14 @@ void MainCommon(in FMaterialPixelParameters MaterialParameters, in float4 SvPosi
 
 					ParticipatingMediaContext ShadowPMC = SetupParticipatingMediaContext(0.0f, ShadowExtinctionCoefficients, MsScattFactor, MsExtinFactor, 0.0f, 0.0f);
 
-					[unroll]
+					UNROLL
 					for (ms = 0; ms < MSCOUNT; ++ms)
 					{
 						ExtinctionAcc[ms] += ShadowPMC.ExtinctionCoefficients[ms] * ExtinctionFactor;
 					}
 				}
 
-				[unroll]
+				UNROLL
 				for (ms = 0; ms < MSCOUNT; ++ms)
 				{
 					PMC.TransmittanceToLight0[ms] *= exp(-ExtinctionAcc[ms] * ShadowDtMeter);
@@ -899,14 +910,14 @@ void MainCommon(in FMaterialPixelParameters MaterialParameters, in float4 SvPosi
 					float OutOpticalDepth = 0.0f;
 					float CloudShadow = GetCloudVolumetricShadow(SampleWorldPosition, RenderVolumetricCloudParameters.CloudShadowmapWorldToLightClipMatrix[1], RenderVolumetricCloudParameters.CloudShadowmapFarDepthKm[1].x,
 						RenderVolumetricCloudParameters.CloudShadowTexture1, RenderVolumetricCloudParameters.CloudBilinearTextureSampler, OutOpticalDepth);
-					[unroll]
+					UNROLL
 					for (ms = 0; ms < MSCOUNT; ++ms)
 					{
 						PMC.TransmittanceToLight1[ms] *= exp(-OutOpticalDepth * (ms == 0 ? 1.0f : pow(MsExtinFactor, ms)));
 					}
 #else // Use cloud shadowmap
 					DebugShadowSampleCount++;
-					[unroll]
+					UNROLL
 					for (ms = 0; ms < MSCOUNT; ++ms)
 					{
 						ExtinctionAcc[ms] = 0.0f;
@@ -942,14 +953,14 @@ void MainCommon(in FMaterialPixelParameters MaterialParameters, in float4 SvPosi
 
 						ParticipatingMediaContext ShadowPMC = SetupParticipatingMediaContext(0.0f, ShadowExtinctionCoefficients, MsScattFactor, MsExtinFactor, 0.0f, 0.0f);
 
-						[unroll]
+						UNROLL
 						for (ms = 0; ms < MSCOUNT; ++ms)
 						{
 							ExtinctionAcc[ms] += ShadowPMC.ExtinctionCoefficients[ms] * ExtinctionFactor;
 						}
 					}
 
-					[unroll]
+					UNROLL
 					for (ms = 0; ms < MSCOUNT; ++ms)
 					{
 						PMC.TransmittanceToLight1[ms] *= exp(-ExtinctionAcc[ms] * ShadowDtMeter);
@@ -974,10 +985,142 @@ void MainCommon(in FMaterialPixelParameters MaterialParameters, in float4 SvPosi
 			}
 
 
+			//////////////////////////////
+			// Evaluate local light in a slow way just so that it can be used for cinematics for now.
+			//////////////////////////////
+		#if CLOUD_SAMPLE_LOCAL_LIGHTS
+			float3 LocalLightScatteredLuminance[MSCOUNT];
+			UNROLL
+			for (ms = 0; ms < MSCOUNT; ++ms)
+			{
+				LocalLightScatteredLuminance[ms] = 0.0f;
+			}
+			if (ForwardLightData.NumLocalLights > 0)
+			{ 
+				const float2 PixelPos = SvPosition.xy * RenderVolumetricCloudParameters.TracingCoordToFullResPixelCoordScaleBias.xy - ResolvedView.ViewRectMin.xy;// *RenderVolumetricCloudParameters.TracingCoordToFullResPixelCoordScaleBias.xy; // We do not take into account the bias so this owuld not work for split screens / VR.
+
+				float4 ClipPos = mul(float4(SampleWorldPosition, 1.0f), ResolvedView.WorldToClip);
+				ClipPos /= ClipPos.w;
+				const float SceneDepth = ConvertFromDeviceZ(ClipPos.z);
+				 
+				uint GridIndex = ComputeLightGridCellIndex(PixelPos, SceneDepth, 0);
+				const FCulledLightsGridData CulledLightsGrid = GetCulledLightsGrid(GridIndex, 0);
+				FRectTexture RectTexture = InitRectTexture(ForwardLightData.DummyRectLightSourceTexture);
+
+				// This loop will process all unshadowed lights only.
+				LOOP
+				for (uint LocalLightListIndex = 0; LocalLightListIndex < CulledLightsGrid.NumLocalLights; LocalLightListIndex++)
+				{
+					const FLocalLightData LocalLight = GetLocalLightData(CulledLightsGrid.DataStartIndex + LocalLightListIndex, 0);
+			
+					float VolumetricScatteringIntensity = f16tof32(asuint(LocalLight.SpotAnglesAndSourceRadiusPacked.w) >> 16);
+
+					if (VolumetricScatteringIntensity > 0)
+					{
+						FDeferredLightData LightData = (FDeferredLightData)0;
+						LightData.Position = LocalLight.LightPositionAndInvRadius.xyz;
+						LightData.InvRadius = LocalLight.LightPositionAndInvRadius.w;
+						LightData.Color = LocalLight.LightColorAndFalloffExponent.xyz;
+						LightData.FalloffExponent = LocalLight.LightColorAndFalloffExponent.w;
+						LightData.Direction = LocalLight.LightDirectionAndShadowMask.xyz;
+						LightData.SpotAngles = LocalLight.SpotAnglesAndSourceRadiusPacked.xy;
+						LightData.SourceRadius = LocalLight.SpotAnglesAndSourceRadiusPacked.z;
+						LightData.SourceLength = f16tof32(asuint(LocalLight.SpotAnglesAndSourceRadiusPacked.w));
+						LightData.Tangent = LocalLight.LightTangentAndSoftSourceRadius.xyz;
+						LightData.SoftSourceRadius = LocalLight.LightTangentAndSoftSourceRadius.w;
+						LightData.bInverseSquared = LightData.FalloffExponent == 0;
+						LightData.RectLightBarnCosAngle = LocalLight.RectBarnDoor.x;
+						LightData.RectLightBarnLength = LocalLight.RectBarnDoor.y;
+						LightData.bRadialLight = true;
+						LightData.bSpotLight = LightData.SpotAngles.x > -2.0f;
+						LightData.bRectLight = LightData.RectLightBarnLength > -2.0f;
+
+						float3 L = 0;
+						float3 ToLight = 0;
+						float LightMask = GetLocalLightAttenuation(SampleWorldPosition, LightData, ToLight, L);
+
+						float Lighting;
+						if( LightData.bRectLight )
+						{
+							FRect Rect = GetRect( ToLight, LightData );
+							Lighting = IntegrateLight(Rect, RectTexture);
+						}
+						else
+						{
+							FCapsuleLight Capsule = GetCapsule(ToLight, LightData);
+							Capsule.DistBiasSqr = 1.0f;
+							Lighting = IntegrateLight(Capsule, LightData.bInverseSquared);
+						}
+
+						float3 CombinedAttenuation = Lighting * LightMask;
+						if (CombinedAttenuation.r <= 0.0f)
+						{
+							continue;	// Skip shadow tracing if we know light is already not visible here
+						}
+						CombinedAttenuation *= LightData.Color;
+
+
+						MATVEC ExtinctionAcc[MSCOUNT];
+						const float SampleToLightLen = length(ToLight);
+						const float3 SampleToLightNorm = ToLight / max(1e-5, SampleToLightLen);
+						const float ShadowLengthTest = SampleToLightLen;
+						const float ShadowStepCount = RenderVolumetricCloudParameters.LocalLightsShadowSampleCount;
+						const float InvShadowStepCount = 1.0f / ShadowStepCount;
+						const float ShadowJitterNorm = 0.5f;
+						const float ShadowDtMeter = ShadowLengthTest * CENTIMETER_TO_METER / ShadowStepCount;
+
+						UNROLL
+						for (ms = 0; ms < MSCOUNT; ++ms)
+						{
+							ExtinctionAcc[ms] = 0.0f;
+						}
+						if (ShadowStepCount > 0.0f)
+						{
+							// Non-linear shadow sample distribution, hardcoded to x^2
+							for (float ShadowT = ShadowJitterNorm; ShadowT < ShadowStepCount; ShadowT += 1.0f)
+							{
+								UpdateMaterialCloudParam(MaterialParameters, SampleWorldPosition + SampleToLightNorm * ShadowLengthTest * (ShadowT * InvShadowStepCount), 
+									ResolvedView, CloudLayerParams, TRACING_SHADOW_DISTANCE_OFF);
+								const float ExtinctionFactor = 1.0f;
+
+								CalcPixelMaterialInputs(MaterialParameters, PixelMaterialInputs);
+								MATVEC ShadowExtinctionCoefficients = SampleExtinctionCoefficients(PixelMaterialInputs);
+
+								if (MaterialParameters.CloudSampleNormAltitudeInLayer < 0.0f || MaterialParameters.CloudSampleNormAltitudeInLayer > 1.0f)
+								{
+									ShadowExtinctionCoefficients = 0.0f; // Ignore remaining samples since we have just traveld out of the cloud layer.
+								}
+
+								ParticipatingMediaContext ShadowPMC = SetupParticipatingMediaContext(0.0f, ShadowExtinctionCoefficients, MsScattFactor, MsExtinFactor, 0.0f, 0.0f);
+
+								UNROLL
+								for (ms = 0; ms < MSCOUNT; ++ms)
+								{
+									ExtinctionAcc[ms] += ShadowPMC.ExtinctionCoefficients[ms] * ExtinctionFactor;
+								}
+							}
+						}
+
+						UNROLL
+						for (ms = 0; ms < MSCOUNT; ++ms)
+						{
+							MATVEC LocalLightTransmittance = exp(-ExtinctionAcc[ms] * ShadowDtMeter);
+
+							LocalLightScatteredLuminance[ms] += 
+								VolumetricScatteringIntensity * CombinedAttenuation
+								* IsotropicPhase()				// TODO add support for the cloud phase function
+								* LocalLightTransmittance;		// Transmittance accounting for multiple scattering approximation
+						}
+					}
+				}
+			}
+		#endif // CLOUD_SAMPLE_LOCAL_LIGHTS
+
+
 			//////////////////////////////
 			// Evaluate scattered luminance towards camera as well as view transmittance.
 			//////////////////////////////
-			[unroll]
+			UNROLL
 			for (ms = MSCOUNT - 1; ms >= 0; --ms) // Must terminate with 0 because this is where TransmittanceToView is updated.
 			{
 				const MATVEC ScatteringCoefficients = PMC.ScatteringCoefficients[ms];
@@ -996,9 +1139,14 @@ void MainCommon(in FMaterialPixelParameters MaterialParameters, in float4 SvPosi
 				// This could be removed when occlusion is better handled or approximated with a simple/expenssive optional  tracing.
 				SunSkyLuminance += (ms == 0 ? DistantLightLuminance : float3(0.0f, 0.0f, 0.0f));
 
+			#if CLOUD_SAMPLE_LOCAL_LIGHTS
+				SunSkyLuminance += LocalLightScatteredLuminance[ms];
+			#endif
+
 				// *** EmissiveLuminance: it should be (EmissiveLuminance * AbsorptionCoefficients) but not intuitive for artist ==> can be added later as an option for consistency during path tracing. 
 				// See "Production Volume Rendering", 2017, Section 2.2. So right now EmissiveLuminance is in fact LuminancePerMeter.
-				const float3 ScatteredLuminance = SunSkyLuminance * ScatteringCoefficients + EmissiveLuminance;// No dt, because it is part of the analytical integral below.
+				// dt is not part of ScatteredLuminance because it is part of the analytical integral below.
+				const float3 ScatteredLuminance = SunSkyLuminance * ScatteringCoefficients + EmissiveLuminance;
 
 #if 0
 	// Default iterative integration
diff --git a/Engine/Source/Runtime/Renderer/Private/DeferredShadingRenderer.cpp b/Engine/Source/Runtime/Renderer/Private/DeferredShadingRenderer.cpp
--- a/Engine/Source/Runtime/Renderer/Private/DeferredShadingRenderer.cpp
+++ b/Engine/Source/Runtime/Renderer/Private/DeferredShadingRenderer.cpp
@@ -2057,6 +2057,7 @@ void FDeferredShadingSceneRenderer::Render(FRDGBuilder& GraphBuilder)
 
 		bComputeLightGrid |= (
 			ShouldRenderVolumetricFog() ||
+			VolumetricCloudWantsToSampleLocalLights(Scene, ViewFamily.EngineShowFlags) ||
 			ViewFamily.ViewMode != VMI_Lit ||
 			bAnyLumenEnabled ||
 			VirtualShadowMapArray.IsEnabled());
diff --git a/Engine/Source/Runtime/Renderer/Private/LightGridInjection.cpp b/Engine/Source/Runtime/Renderer/Private/LightGridInjection.cpp
--- a/Engine/Source/Runtime/Renderer/Private/LightGridInjection.cpp
+++ b/Engine/Source/Runtime/Renderer/Private/LightGridInjection.cpp
@@ -25,6 +25,7 @@ LightGridInjection.cpp
 #include "ScenePrivate.h"
 #include "ClearQuad.h"
 #include "VolumetricFog.h"
+#include "VolumetricCloudRendering.h"
 #include "Components/LightComponent.h"
 #include "Engine/MapBuildDataRegistry.h"
 
@@ -783,7 +784,7 @@ void FDeferredShadingSceneRenderer::GatherLightsAndComputeLightGrid(FRDGBuilder&
 	for (int32 ViewIndex = 0; ViewIndex < Views.Num(); ViewIndex++)
 	{
 		const FViewInfo& View = Views[ViewIndex];
-		bAnyViewUsesForwardLighting |= View.bTranslucentSurfaceLighting || ShouldRenderVolumetricFog() || View.bHasSingleLayerWaterMaterial;
+		bAnyViewUsesForwardLighting |= View.bTranslucentSurfaceLighting || ShouldRenderVolumetricFog() || View.bHasSingleLayerWaterMaterial || VolumetricCloudWantsToSampleLocalLights(Scene, ViewFamily.EngineShowFlags);
 		bAnyViewUsesLumen |= GetViewPipelineState(View).DiffuseIndirectMethod == EDiffuseIndirectMethod::Lumen || GetViewPipelineState(View).ReflectionsMethod == EReflectionsMethod::Lumen;
 	}
 	
diff --git a/Engine/Source/Runtime/Renderer/Private/VolumetricCloudRendering.cpp b/Engine/Source/Runtime/Renderer/Private/VolumetricCloudRendering.cpp
--- a/Engine/Source/Runtime/Renderer/Private/VolumetricCloudRendering.cpp
+++ b/Engine/Source/Runtime/Renderer/Private/VolumetricCloudRendering.cpp
@@ -23,6 +23,7 @@
 #include "BlueNoise.h"
 #include "FogRendering.h"
 #include "SkyAtmosphereRendering.h"
+#include "BasePassRendering.h"
 
 
 //PRAGMA_DISABLE_OPTIMIZATION
@@ -186,7 +187,17 @@ static TAutoConsoleVariable<int32> CVarVolumetricCloudEnableDistantSkyLightSampl
 
 static TAutoConsoleVariable<int32> CVarVolumetricCloudEnableAtmosphericLightsSampling(
 	TEXT("r.VolumetricCloud.EnableAtmosphericLightsSampling"), 1,
-	TEXT("Enable/disable the atmospheric lights contribution on clouds."),
+	TEXT("Enable/disable atmospheric lights contribution on clouds."),
+	ECVF_RenderThreadSafe);
+
+static TAutoConsoleVariable<int32> CVarVolumetricCloudEnableLocalLightsSampling(
+	TEXT("r.VolumetricCloud.EnableLocalLightsSampling"), 0,
+	TEXT("[EXPERIMENTAL] Enable/disable local lights contribution on clouds. Expenssive! Use for cinematics if needed."),
+	ECVF_RenderThreadSafe);
+
+static TAutoConsoleVariable<float> CVarVolumetricCloudLocalLightsShadowSampleCount(
+	TEXT("r.VolumetricCloud.LocalLights.ShadowSampleCount"), 12,
+	TEXT("[EXPERIMENTAL] Set the volumetric shadow sample count when evaluating local lights. Expenssive! Use for cinematics if needed."),
 	ECVF_RenderThreadSafe);
 
 ////////////////////////////////////////////////////////////////////////// 
@@ -194,6 +205,7 @@ static TAutoConsoleVariable<int32> CVarVolumetricCloudEnableAtmosphericLightsSam
 static TAutoConsoleVariable<int32> CVarVolumetricCloudDebugSampleCountMode(
 	TEXT("r.VolumetricCloud.Debug.SampleCountMode"), 0,
 	TEXT("Only for developers. [0] Disabled [1] Primary material sample count [2] Advanced:raymarched shadow sample count [3] Shadow material sample count [4] Advanced:ground shadow sample count [5] Advanced:ground shadow material sample count"));
+
 ////////////////////////////////////////////////////////////////////////// 
 
 
@@ -239,6 +251,11 @@ bool ShouldViewVisualizeVolumetricCloudConservativeDensity(const FViewInfo& View
 		&& GIsEditor;
 }
 
+bool VolumetricCloudWantsToSampleLocalLights(const FScene* Scene, const FEngineShowFlags& EngineShowFlags)
+{
+	return ShouldRenderVolumetricCloud(Scene, EngineShowFlags) && CVarVolumetricCloudEnableLocalLightsSampling.GetValueOnRenderThread();
+}
+
 static bool ShouldRenderCloudShadowmap(const FLightSceneProxy* AtmosphericLight)
 {
 	return CVarVolumetricCloudShadowMap.GetValueOnRenderThread() > 0 && AtmosphericLight && AtmosphericLight->GetCastCloudShadows();
@@ -496,6 +513,7 @@ BEGIN_GLOBAL_SHADER_PARAMETER_STRUCT(FRenderVolumetricCloudGlobalParameters, )
 	SHADER_PARAMETER(int32, VirtualShadowMapId0)
 //	SHADER_PARAMETER_STRUCT(FBlueNoise, BlueNoise)
 	SHADER_PARAMETER(FUintVector4, TracingCoordToZbufferCoordScaleBias)
+	SHADER_PARAMETER(FUintVector4, TracingCoordToFullResPixelCoordScaleBias)
 	SHADER_PARAMETER(FUintVector4, SceneDepthTextureMinMaxCoord)
 	SHADER_PARAMETER(uint32, NoiseFrameIndexModPattern)
 	SHADER_PARAMETER(int32, OpaqueIntersectionMode)
@@ -505,6 +523,7 @@ BEGIN_GLOBAL_SHADER_PARAMETER_STRUCT(FRenderVolumetricCloudGlobalParameters, )
 	SHADER_PARAMETER(uint32, HasValidHZB)
 	SHADER_PARAMETER(uint32, ClampRayTToDepthBufferPostHZB)
 	SHADER_PARAMETER(uint32, TraceShadowmap)
+	SHADER_PARAMETER(float, LocalLightsShadowSampleCount)
 	SHADER_PARAMETER(FVector3f, HZBUvFactor)
 	SHADER_PARAMETER(FVector4, HZBSize)
 	SHADER_PARAMETER_RDG_TEXTURE(Texture2D<float>, HZBTexture)
@@ -516,6 +535,7 @@ BEGIN_GLOBAL_SHADER_PARAMETER_STRUCT(FRenderVolumetricCloudGlobalParameters, )
 	SHADER_PARAMETER(int32, EnableAtmosphericLightsSampling)
 	SHADER_PARAMETER(int32, EnableHeightFog)
 	SHADER_PARAMETER_STRUCT_INCLUDE(FFogUniformParameters, FogStruct)
+	SHADER_PARAMETER_STRUCT(FForwardLightData, ForwardLightData)
 END_GLOBAL_SHADER_PARAMETER_STRUCT()
 
 IMPLEMENT_STATIC_UNIFORM_BUFFER_STRUCT(FRenderVolumetricCloudGlobalParameters, "RenderVolumetricCloudParameters", SceneTextures);
@@ -551,6 +571,7 @@ void SetupDefaultRenderVolumetricCloudGlobalParameters(FRDGBuilder& GraphBuilder
 	VolumetricCloudParams.BlueNoise.Texture = GEngine->HighFrequencyNoiseTexture->Resource->TextureRHI;
 #endif*/
 	VolumetricCloudParams.TracingCoordToZbufferCoordScaleBias = FUintVector4(1, 1, 0, 0);
+	VolumetricCloudParams.TracingCoordToFullResPixelCoordScaleBias = FUintVector4(1, 1, 0, 0);
 	VolumetricCloudParams.NoiseFrameIndexModPattern = 0;
 	VolumetricCloudParams.VolumetricRenderTargetMode = ViewInfo.ViewState ? ViewInfo.ViewState->VolumetricCloudRenderTarget.GetMode() : 0;
 	VolumetricCloudParams.SampleCountDebugMode = GetVolumetricCloudDebugSampleCountMode(ViewInfo.Family->EngineShowFlags);
@@ -565,6 +586,9 @@ void SetupDefaultRenderVolumetricCloudGlobalParameters(FRDGBuilder& GraphBuilder
 	VolumetricCloudParams.EnableHeightFog = ViewInfo.Family->Scene->HasAnyExponentialHeightFog() && ShouldRenderFog(*ViewInfo.Family);
 	SetupFogUniformParameters(GraphBuilder, ViewInfo, VolumetricCloudParams.FogStruct);
 
+	VolumetricCloudParams.ForwardLightData = ViewInfo.ForwardLightingResources->ForwardLightData;
+	VolumetricCloudParams.LocalLightsShadowSampleCount = FMath::Clamp(CVarVolumetricCloudLocalLightsShadowSampleCount.GetValueOnRenderThread(), 0.0f, 128.0f);
+
 	ESceneTextureSetupMode SceneTextureSetupMode = ESceneTextureSetupMode::All;
 	EnumRemoveFlags(SceneTextureSetupMode, ESceneTextureSetupMode::SceneColor);
 	EnumRemoveFlags(SceneTextureSetupMode, ESceneTextureSetupMode::SceneVelocity);
@@ -740,12 +764,14 @@ class FRenderVolumetricCloudRenderViewCS : public FMeshMaterialShader
 	class FCloudPerSampleAtmosphereTransmittance : SHADER_PERMUTATION_BOOL("CLOUD_PER_SAMPLE_ATMOSPHERE_TRANSMITTANCE");
 	class FCloudSampleAtmosphericLightShadowmap : SHADER_PERMUTATION_BOOL("CLOUD_SAMPLE_ATMOSPHERIC_LIGHT_SHADOWMAP");
 	class FCloudSampleSecondLight : SHADER_PERMUTATION_BOOL("CLOUD_SAMPLE_SECOND_LIGHT");
+	class FCloudSampleLocalLights : SHADER_PERMUTATION_BOOL("CLOUD_SAMPLE_LOCAL_LIGHTS");
 	class FCloudSampleCountDebugMode : SHADER_PERMUTATION_BOOL("CLOUD_SAMPLE_COUNT_DEBUG_MODE");
 
 	using FPermutationDomain = TShaderPermutationDomain<
 		FCloudPerSampleAtmosphereTransmittance,
 		FCloudSampleAtmosphericLightShadowmap,
 		FCloudSampleSecondLight,
+		FCloudSampleLocalLights,
 		FCloudSampleCountDebugMode>;
 
 	BEGIN_SHADER_PARAMETER_STRUCT(FParameters, )
@@ -775,7 +801,8 @@ class FRenderVolumetricCloudRenderViewCS : public FMeshMaterialShader
 				&& EnumHasAllFlags(Parameters.Flags, EShaderPermutationFlags::HasEditorOnlyData)
 				&& !PermutationVector.Get<FCloudPerSampleAtmosphereTransmittance>()
 				&& !PermutationVector.Get<FCloudSampleAtmosphericLightShadowmap>()
-				&& !PermutationVector.Get<FCloudSampleSecondLight>();
+				&& !PermutationVector.Get<FCloudSampleSecondLight>()
+				&& !PermutationVector.Get<FCloudSampleLocalLights>();
 		}
 
 		return bIsCompatible;
@@ -804,6 +831,8 @@ class FRenderVolumetricCloudRenderViewCS : public FMeshMaterialShader
 
 		// This shader must support typed UAV load and we are testing if it is supported at runtime using RHIIsTypedUAVLoadSupported
 		OutEnvironment.CompilerFlags.Add(CFLAG_AllowTypedUAVLoads);
+
+		FForwardLightingParameters::ModifyCompilationEnvironment(Parameters.Platform, OutEnvironment);
 	}
 };
 
@@ -1843,6 +1872,7 @@ void FSceneRenderer::InitVolumetricCloudsForViews(FRDGBuilder& GraphBuilder, boo
 FCloudRenderContext::FCloudRenderContext()
 {
 	TracingCoordToZbufferCoordScaleBias = FUintVector4(1, 1, 0, 0);
+	TracingCoordToFullResPixelCoordScaleBias = FUintVector4(1, 1, 0, 0);
 	NoiseFrameIndexModPattern = 0;
 
 	bIsReflectionRendering = false;
@@ -1890,6 +1920,7 @@ static TRDGUniformBufferRef<FRenderVolumetricCloudGlobalParameters> CreateCloudP
 	VolumetricCloudParams.CloudShadowTexture0 = CloudRC.VolumetricCloudShadowTexture[0];
 	VolumetricCloudParams.CloudShadowTexture1 = CloudRC.VolumetricCloudShadowTexture[1];
 	VolumetricCloudParams.TracingCoordToZbufferCoordScaleBias = CloudRC.TracingCoordToZbufferCoordScaleBias;
+	VolumetricCloudParams.TracingCoordToFullResPixelCoordScaleBias = CloudRC.TracingCoordToFullResPixelCoordScaleBias;
 	VolumetricCloudParams.NoiseFrameIndexModPattern = CloudRC.NoiseFrameIndexModPattern;
 	VolumetricCloudParams.IsReflectionRendering = CloudRC.bIsReflectionRendering ? 1 : 0;
 
@@ -1982,6 +2013,7 @@ void FSceneRenderer::RenderVolumetricCloudsInternal(FRDGBuilder& GraphBuilder, F
 	TUniformBufferRef<FViewUniformShaderParameters> ViewUniformBuffer = CloudRC.ViewUniformBuffer;
 	TRDGUniformBufferRef<FRenderVolumetricCloudGlobalParameters> CloudPassUniformBuffer = CreateCloudPassUniformBuffer(GraphBuilder, CloudRC);
 
+	const bool bCloudEnableLocalLightSampling = CVarVolumetricCloudEnableLocalLightsSampling.GetValueOnRenderThread() > 0;
 	if (ShouldUseComputeForCloudTracing())
 	{
 		FRDGTextureRef CloudColorCubeTexture;
@@ -2007,6 +2039,7 @@ void FSceneRenderer::RenderVolumetricCloudsInternal(FRDGBuilder& GraphBuilder, F
 		PermutationVector.Set<typename FRenderVolumetricCloudRenderViewCS::FCloudPerSampleAtmosphereTransmittance>(!bVisualizeConservativeDensityOrDebugSampleCount && ShouldUsePerSampleAtmosphereTransmittance(Scene, &MainView));
 		PermutationVector.Set<typename FRenderVolumetricCloudRenderViewCS::FCloudSampleAtmosphericLightShadowmap>(!bVisualizeConservativeDensityOrDebugSampleCount && !bSkipAtmosphericLightShadowmap && CVarVolumetricCloudShadowSampleAtmosphericLightShadowmap.GetValueOnRenderThread() > 0);
 		PermutationVector.Set<typename FRenderVolumetricCloudRenderViewCS::FCloudSampleSecondLight>(!bVisualizeConservativeDensityOrDebugSampleCount && bSecondAtmosphereLightEnabled);
+		PermutationVector.Set<typename FRenderVolumetricCloudRenderViewCS::FCloudSampleLocalLights>(!bVisualizeConservativeDensityOrDebugSampleCount && bCloudEnableLocalLightSampling && !CloudRC.bIsSkyRealTimeReflectionRendering);
 		PermutationVector.Set<typename FRenderVolumetricCloudRenderViewCS::FCloudSampleCountDebugMode>(bVisualizeConservativeDensityOrDebugSampleCount);
 
 		TShaderRef<FRenderVolumetricCloudRenderViewCS> ComputeShader = MaterialResource->GetShader<FRenderVolumetricCloudRenderViewCS>(&FLocalVertexFactory::StaticType, PermutationVector, false);
@@ -2058,6 +2091,8 @@ void FSceneRenderer::RenderVolumetricCloudsInternal(FRDGBuilder& GraphBuilder, F
 	}
 	else
 	{
+		ensureMsgf(!bCloudEnableLocalLightSampling, TEXT("We do not support lighting clouds with local light through the pixel shader in order to reduce permutations count."));
+
 		FRenderVolumetricCloudRenderViewParametersPS* RenderViewPassParameters = GraphBuilder.AllocParameters<FRenderVolumetricCloudRenderViewParametersPS>();
 		RenderViewPassParameters->VolumetricCloudRenderViewParamsUB = CloudPassUniformBuffer;
 		RenderViewPassParameters->CloudShadowTexture0 = CloudRC.VolumetricCloudShadowTexture[0];
@@ -2204,6 +2239,7 @@ bool FSceneRenderer::RenderVolumetricCloud(
 				FRDGTextureRef DestinationRT = nullptr;
 				FRDGTextureRef DestinationRTDepth = nullptr;
 				CloudRC.TracingCoordToZbufferCoordScaleBias = FUintVector4(1, 1, 0, 0);
+				CloudRC.TracingCoordToFullResPixelCoordScaleBias = FUintVector4(1, 1, 0, 0);
 				CloudRC.NoiseFrameIndexModPattern = ViewInfo.CachedViewUniformShaderParameters->StateFrameIndexMod8;
 
 				if (bShouldViewRenderVolumetricCloudRenderTarget)
@@ -2228,6 +2264,8 @@ bool FSceneRenderer::RenderVolumetricCloud(
 					// Also take into account the view rect min to be able to read correct depth
 					CloudRC.TracingCoordToZbufferCoordScaleBias.Z += ViewInfo.CachedViewUniformShaderParameters->ViewRectMin.X;
 					CloudRC.TracingCoordToZbufferCoordScaleBias.W += ViewInfo.CachedViewUniformShaderParameters->ViewRectMin.Y / ((VRT.GetMode() == 0 || VRT.GetMode() == 3) ? 2 : 1);
+					CloudRC.TracingCoordToFullResPixelCoordScaleBias = VRT.GetTracingCoordToFullResPixelCoordScaleBias();
+
 					CloudRC.NoiseFrameIndexModPattern = VRT.GetNoiseFrameIndexModPattern();
 
 					check(VRT.GetMode() != 0 || CloudRC.bVisualizeConservativeDensityOrDebugSampleCount || HalfResolutionDepthCheckerboardMinMaxTexture);
diff --git a/Engine/Source/Runtime/Renderer/Private/VolumetricCloudRendering.h b/Engine/Source/Runtime/Renderer/Private/VolumetricCloudRendering.h
--- a/Engine/Source/Runtime/Renderer/Private/VolumetricCloudRendering.h
+++ b/Engine/Source/Runtime/Renderer/Private/VolumetricCloudRendering.h
@@ -102,6 +102,7 @@ private:
 
 bool ShouldRenderVolumetricCloud(const FScene* Scene, const FEngineShowFlags& EngineShowFlags);
 bool ShouldViewVisualizeVolumetricCloudConservativeDensity(const FViewInfo& ViewInfo, const FEngineShowFlags& EngineShowFlags);
+bool VolumetricCloudWantsToSampleLocalLights(const FScene* Scene, const FEngineShowFlags& EngineShowFlags);
 uint32 GetVolumetricCloudDebugSampleCountMode(const FEngineShowFlags& ShowFlags);
 
 
@@ -134,6 +135,7 @@ struct FCloudRenderContext
 	bool bVisualizeConservativeDensityOrDebugSampleCount;
 
 	FUintVector4 TracingCoordToZbufferCoordScaleBias;
+	FUintVector4 TracingCoordToFullResPixelCoordScaleBias;
 	uint32 NoiseFrameIndexModPattern;
 
 	FVolumeShadowingShaderParametersGlobal0 LightShadowShaderParams0;
diff --git a/Engine/Source/Runtime/Renderer/Private/VolumetricRenderTarget.cpp b/Engine/Source/Runtime/Renderer/Private/VolumetricRenderTarget.cpp
--- a/Engine/Source/Runtime/Renderer/Private/VolumetricRenderTarget.cpp
+++ b/Engine/Source/Runtime/Renderer/Private/VolumetricRenderTarget.cpp
@@ -357,6 +357,14 @@ FUintVector4 FVolumetricRenderTargetViewStateData::GetTracingCoordToZbufferCoord
 		CurrentPixelOffset.Y * VolumetricReconstructRTDownsampleFactor / VolumetricReconstructRTDownsampleFactor);
 }
 
+FUintVector4 FVolumetricRenderTargetViewStateData::GetTracingCoordToFullResPixelCoordScaleBias() const
+{
+	// In this case, the source depth buffer full resolution depth buffer is the full resolution scene one
+	const uint32 CombinedDownsampleFactor = VolumetricReconstructRTDownsampleFactor * VolumetricTracingRTDownsampleFactor;
+	return FUintVector4(CombinedDownsampleFactor, CombinedDownsampleFactor,																// Scale is the combined downsample factor
+		CurrentPixelOffset.X * VolumetricReconstructRTDownsampleFactor, CurrentPixelOffset.Y * VolumetricReconstructRTDownsampleFactor);// Each sample will then sample from full res according to reconstructed RT offset times its downsample factor
+}
+
 
 /*=============================================================================
 	FSceneRenderer implementation.
diff --git a/Engine/Source/Runtime/Renderer/Private/VolumetricRenderTargetViewStateData.h b/Engine/Source/Runtime/Renderer/Private/VolumetricRenderTargetViewStateData.h
--- a/Engine/Source/Runtime/Renderer/Private/VolumetricRenderTargetViewStateData.h
+++ b/Engine/Source/Runtime/Renderer/Private/VolumetricRenderTargetViewStateData.h
@@ -49,6 +49,7 @@ public:
 	const uint32 GetVolumetricTracingRTDownsampleFactor() const { return VolumetricTracingRTDownsampleFactor; }
 
 	FUintVector4 GetTracingCoordToZbufferCoordScaleBias() const;
+	FUintVector4 GetTracingCoordToFullResPixelCoordScaleBias() const;
 
 	float GetUvNoiseScale()		const { return UvNoiseScale; }
 	int32 GetMode()				const { return Mode; }
