Added a cvar to override the lightgrid far distance.
This is usefull to have distant light within volumetric clouds when the "local light in volumetric cloud" feature is used.

diff --git a/Engine/Source/Runtime/Renderer/Private/LightGridInjection.cpp b/Engine/Source/Runtime/Renderer/Private/LightGridInjection.cpp
--- a/Engine/Source/Runtime/Renderer/Private/LightGridInjection.cpp
+++ b/Engine/Source/Runtime/Renderer/Private/LightGridInjection.cpp
@@ -74,6 +74,16 @@ FAutoConsoleVariableRef CVarLightCullingQuality(
 	ECVF_RenderThreadSafe
 );
 
+float GLightCullingMaxDistanceOverrideKilometers = -1.0f;
+FAutoConsoleVariableRef CVarLightCullingMaxDistanceOverride(
+	TEXT("r.LightCulling.MaxDistanceOverrideKilometers"),
+	GLightCullingMaxDistanceOverrideKilometers,
+	TEXT("Used to override the maximum far distance at which we can store data in the light grid.\n If this is increase, you might want to update r.Forward.LightGridSizeZ to a reasonable value according to your use case light count and distribution.")
+	TEXT(" <=0: off \n")
+	TEXT(" >0: the far distance in kilometers.\n"),
+	ECVF_RenderThreadSafe
+);
+
 extern TAutoConsoleVariable<int32> CVarVirtualShadowOnePassProjection;
 
 /** A minimal forwarding lighting setup. */
@@ -569,7 +579,9 @@ void FSceneRenderer::ComputeLightGrid(FRDGBuilder& GraphBuilder, bool bCullLight
 		ForwardLightData.DirectLightingShowFlag = ViewFamily.EngineShowFlags.DirectLighting ? 1 : 0;
 
 		// Clamp far plane to something reasonable
-		float FarPlane = FMath::Min(FMath::Max(FurthestLight, View.FurthestReflectionCaptureDistance), (float)HALF_WORLD_MAX / 5.0f);
+		const float KilometersToCentimeters = 100000.0f;
+		const float LightCullingMaxDistance = GLightCullingMaxDistanceOverrideKilometers <= 0.0f ? (float)HALF_WORLD_MAX / 5.0f : GLightCullingMaxDistanceOverrideKilometers * KilometersToCentimeters;
+		float FarPlane = FMath::Min(FMath::Max(FurthestLight, View.FurthestReflectionCaptureDistance), LightCullingMaxDistance);
 		FVector ZParams = GetLightGridZParams(View.NearClippingDistance, FarPlane + 10.f);
 		ForwardLightData.LightGridZParams = ZParams;
 
