Move shadow quality check to before the scene lights loop to save a bit of overhead when shadows are disabled.
Also moved the light stat scope to only count if a light is visible. Reduces stat scope spam.

diff --git a/Engine/Source/Runtime/Renderer/Private/ShadowSetup.cpp b/Engine/Source/Runtime/Renderer/Private/ShadowSetup.cpp
--- a/Engine/Source/Runtime/Renderer/Private/ShadowSetup.cpp
+++ b/Engine/Source/Runtime/Renderer/Private/ShadowSetup.cpp
@@ -5367,136 +5367,139 @@ void FSceneRenderer::InitDynamicShadows(FRHICommandListImmediate& RHICmdList, FG
 		SCOPE_CYCLE_COUNTER(STAT_InitDynamicShadowsTime);
 		CSV_SCOPED_TIMING_STAT_EXCLUSIVE(ShadowInitDynamic);
 
-		for (TSparseArray<FLightSceneInfoCompact>::TConstIterator LightIt(Scene->Lights); LightIt; ++LightIt)
+		if (GetShadowQuality() > 0)
 		{
-			const FLightSceneInfoCompact& LightSceneInfoCompact = *LightIt;
-			FLightSceneInfo* LightSceneInfo = LightSceneInfoCompact.LightSceneInfo;
-
-			FScopeCycleCounter Context(LightSceneInfo->Proxy->GetStatId());
-
-			FVisibleLightInfo& VisibleLightInfo = VisibleLightInfos[LightSceneInfo->Id];
+			for (TSparseArray<FLightSceneInfoCompact>::TConstIterator LightIt(Scene->Lights); LightIt; ++LightIt)
+			{
+				const FLightSceneInfoCompact& LightSceneInfoCompact = *LightIt;
+				FLightSceneInfo* LightSceneInfo = LightSceneInfoCompact.LightSceneInfo;
 
-			const FLightOcclusionType OcclusionType = GetLightOcclusionType(LightSceneInfoCompact);
-			if (OcclusionType != FLightOcclusionType::Shadowmap)
-				continue;
+				FVisibleLightInfo& VisibleLightInfo = VisibleLightInfos[LightSceneInfo->Id];
 
-			// Only consider lights that may have shadows.
-			if ((LightSceneInfoCompact.bCastStaticShadow || LightSceneInfoCompact.bCastDynamicShadow) && GetShadowQuality() > 0)
-			{
-				// see if the light is visible in any view
-				bool bIsVisibleInAnyView = false;
+				const FLightOcclusionType OcclusionType = GetLightOcclusionType(LightSceneInfoCompact);
+				if (OcclusionType != FLightOcclusionType::Shadowmap)
+					continue;
 
-				for (int32 ViewIndex = 0; ViewIndex < Views.Num(); ViewIndex++)
+				// Only consider lights that may have shadows.
+				if (LightSceneInfoCompact.bCastStaticShadow || LightSceneInfoCompact.bCastDynamicShadow)
 				{
-					// View frustums are only checked when lights have visible primitives or have modulated shadows,
-					// so we don't need to check for that again here
-					bIsVisibleInAnyView = LightSceneInfo->ShouldRenderLight(Views[ViewIndex]);
+					// see if the light is visible in any view
+					bool bIsVisibleInAnyView = false;
 
-					if (bIsVisibleInAnyView) 
+					for (int32 ViewIndex = 0; ViewIndex < Views.Num(); ViewIndex++)
 					{
-						break;
-					}
-				}
+						// View frustums are only checked when lights have visible primitives or have modulated shadows,
+						// so we don't need to check for that again here
+						bIsVisibleInAnyView = LightSceneInfo->ShouldRenderLight(Views[ViewIndex]);
 
-				if (bIsVisibleInAnyView)
-				{
-					static const auto AllowStaticLightingVar = IConsoleManager::Get().FindTConsoleVariableDataInt(TEXT("r.AllowStaticLighting"));
-					const bool bAllowStaticLighting = (!AllowStaticLightingVar || AllowStaticLightingVar->GetValueOnRenderThread() != 0);
-					const bool bPointLightShadow = LightSceneInfoCompact.LightType == LightType_Point || LightSceneInfoCompact.LightType == LightType_Rect;
-					const bool bDirectionalLightShadow = LightSceneInfoCompact.LightType == LightType_Directional;
-					const bool bSpotLightShadow = LightSceneInfoCompact.LightType == LightType_Spot;
-
-					// Only create whole scene shadows for lights that don't precompute shadowing (movable lights)
-					const bool bShouldCreateShadowForMovableLight = 
-						LightSceneInfoCompact.bCastDynamicShadow
-						&& (!LightSceneInfo->Proxy->HasStaticShadowing() || !bAllowStaticLighting);
-
-					const bool bCreateShadowForMovableLight = 
-						bShouldCreateShadowForMovableLight
-						&& (!bPointLightShadow || bProjectEnablePointLightShadows)
-						&& (!bDirectionalLightShadow || bProjectEnableMovableDirectionLightShadows)
-						&& (!bSpotLightShadow || bProjectEnableMovableSpotLightShadows);
-
-					// Also create a whole scene shadow for lights with precomputed shadows that are unbuilt
-					const bool bShouldCreateShadowToPreviewStaticLight =
-						LightSceneInfo->Proxy->HasStaticShadowing()
-						&& LightSceneInfoCompact.bCastStaticShadow
-						&& !LightSceneInfo->IsPrecomputedLightingValid();						
-
-					const bool bCreateShadowToPreviewStaticLight = 
-						bShouldCreateShadowToPreviewStaticLight						
-						&& (!bPointLightShadow || bProjectEnablePointLightShadows)
-						// Stationary point light and spot light shadow are unsupported on mobile
-						&& (!bMobile || bDirectionalLightShadow);
-
-					// Create a whole scene shadow for lights that want static shadowing but didn't get assigned to a valid shadowmap channel due to overlap
-					const bool bShouldCreateShadowForOverflowStaticShadowing =
-						LightSceneInfo->Proxy->HasStaticShadowing()
-						&& !LightSceneInfo->Proxy->HasStaticLighting()
-						&& LightSceneInfoCompact.bCastStaticShadow
-						&& LightSceneInfo->IsPrecomputedLightingValid()
-						&& LightSceneInfo->Proxy->GetShadowMapChannel() == INDEX_NONE;
-
-					const bool bCreateShadowForOverflowStaticShadowing =
-						bShouldCreateShadowForOverflowStaticShadowing
-						&& (!bPointLightShadow || bProjectEnablePointLightShadows)
-						// Stationary point light and spot light shadow are unsupported on mobile
-						&& (!bMobile || bDirectionalLightShadow);
-
-					const bool bPointLightWholeSceneShadow = (bShouldCreateShadowForMovableLight || bShouldCreateShadowForOverflowStaticShadowing || bShouldCreateShadowToPreviewStaticLight) && bPointLightShadow;
-					if (bPointLightWholeSceneShadow)
-					{						
-						UsedWholeScenePointLightNames.Add(LightSceneInfoCompact.LightSceneInfo->Proxy->GetComponentName());
+						if (bIsVisibleInAnyView)
+						{
+							break;
+						}
 					}
 
-					if (bCreateShadowForMovableLight || bCreateShadowToPreviewStaticLight || bCreateShadowForOverflowStaticShadowing)
+					if (bIsVisibleInAnyView)
 					{
-						// Try to create a whole scene projected shadow.
-						CreateWholeSceneProjectedShadow(LightSceneInfo, NumPointShadowCachesUpdatedThisFrame, NumSpotShadowCachesUpdatedThisFrame);
-					}
+						FScopeCycleCounter Context(LightSceneInfo->Proxy->GetStatId());
+
+						static const auto AllowStaticLightingVar = IConsoleManager::Get().FindTConsoleVariableDataInt(TEXT("r.AllowStaticLighting"));
+						const bool bAllowStaticLighting = (!AllowStaticLightingVar || AllowStaticLightingVar->GetValueOnRenderThread() != 0);
+						const bool bPointLightShadow = LightSceneInfoCompact.LightType == LightType_Point || LightSceneInfoCompact.LightType == LightType_Rect;
+						const bool bDirectionalLightShadow = LightSceneInfoCompact.LightType == LightType_Directional;
+						const bool bSpotLightShadow = LightSceneInfoCompact.LightType == LightType_Spot;
+
+						// Only create whole scene shadows for lights that don't precompute shadowing (movable lights)
+						const bool bShouldCreateShadowForMovableLight =
+							LightSceneInfoCompact.bCastDynamicShadow
+							&& (!LightSceneInfo->Proxy->HasStaticShadowing() || !bAllowStaticLighting);
+
+						const bool bCreateShadowForMovableLight =
+							bShouldCreateShadowForMovableLight
+							&& (!bPointLightShadow || bProjectEnablePointLightShadows)
+							&& (!bDirectionalLightShadow || bProjectEnableMovableDirectionLightShadows)
+							&& (!bSpotLightShadow || bProjectEnableMovableSpotLightShadows);
+
+						// Also create a whole scene shadow for lights with precomputed shadows that are unbuilt
+						const bool bShouldCreateShadowToPreviewStaticLight =
+							LightSceneInfo->Proxy->HasStaticShadowing()
+							&& LightSceneInfoCompact.bCastStaticShadow
+							&& !LightSceneInfo->IsPrecomputedLightingValid();
+
+						const bool bCreateShadowToPreviewStaticLight =
+							bShouldCreateShadowToPreviewStaticLight
+							&& (!bPointLightShadow || bProjectEnablePointLightShadows)
+							// Stationary point light and spot light shadow are unsupported on mobile
+							&& (!bMobile || bDirectionalLightShadow);
+
+						// Create a whole scene shadow for lights that want static shadowing but didn't get assigned to a valid shadowmap channel due to overlap
+						const bool bShouldCreateShadowForOverflowStaticShadowing =
+							LightSceneInfo->Proxy->HasStaticShadowing()
+							&& !LightSceneInfo->Proxy->HasStaticLighting()
+							&& LightSceneInfoCompact.bCastStaticShadow
+							&& LightSceneInfo->IsPrecomputedLightingValid()
+							&& LightSceneInfo->Proxy->GetShadowMapChannel() == INDEX_NONE;
+
+						const bool bCreateShadowForOverflowStaticShadowing =
+							bShouldCreateShadowForOverflowStaticShadowing
+							&& (!bPointLightShadow || bProjectEnablePointLightShadows)
+							// Stationary point light and spot light shadow are unsupported on mobile
+							&& (!bMobile || bDirectionalLightShadow);
+
+						const bool bPointLightWholeSceneShadow = (bShouldCreateShadowForMovableLight || bShouldCreateShadowForOverflowStaticShadowing || bShouldCreateShadowToPreviewStaticLight) && bPointLightShadow;
+						if (bPointLightWholeSceneShadow)
+						{
+							UsedWholeScenePointLightNames.Add(LightSceneInfoCompact.LightSceneInfo->Proxy->GetComponentName());
+						}
 
-					// Allow movable and stationary lights to create CSM, or static lights that are unbuilt
-					if ((!LightSceneInfo->Proxy->HasStaticLighting() && LightSceneInfoCompact.bCastDynamicShadow) || bCreateShadowToPreviewStaticLight)
-					{
-						static_assert(UE_ARRAY_COUNT(Scene->MobileDirectionalLights) == 3, "All array entries for MobileDirectionalLights must be checked");
-						if( !bMobile ||
-							((LightSceneInfo->Proxy->UseCSMForDynamicObjects() || LightSceneInfo->Proxy->IsMovable()) 
-								// Mobile uses the scene's MobileDirectionalLights only for whole scene shadows.
-								&& (LightSceneInfo == Scene->MobileDirectionalLights[0] || LightSceneInfo == Scene->MobileDirectionalLights[1] || LightSceneInfo == Scene->MobileDirectionalLights[2])))
+						if (bCreateShadowForMovableLight || bCreateShadowToPreviewStaticLight || bCreateShadowForOverflowStaticShadowing)
 						{
-							AddViewDependentWholeSceneShadowsForView(ViewDependentWholeSceneShadows, ViewDependentWholeSceneShadowsThatNeedCulling, VisibleLightInfo, *LightSceneInfo, NumCSMCachesUpdatedThisFrame);
+							// Try to create a whole scene projected shadow.
+							CreateWholeSceneProjectedShadow(LightSceneInfo, NumPointShadowCachesUpdatedThisFrame, NumSpotShadowCachesUpdatedThisFrame);
 						}
 
+						// Allow movable and stationary lights to create CSM, or static lights that are unbuilt
+						if ((!LightSceneInfo->Proxy->HasStaticLighting() && LightSceneInfoCompact.bCastDynamicShadow) || bCreateShadowToPreviewStaticLight)
+						{
+							static_assert(UE_ARRAY_COUNT(Scene->MobileDirectionalLights) == 3, "All array entries for MobileDirectionalLights must be checked");
+							if (!bMobile ||
+								((LightSceneInfo->Proxy->UseCSMForDynamicObjects() || LightSceneInfo->Proxy->IsMovable())
+									// Mobile uses the scene's MobileDirectionalLights only for whole scene shadows.
+									&& (LightSceneInfo == Scene->MobileDirectionalLights[0] || LightSceneInfo == Scene->MobileDirectionalLights[1] || LightSceneInfo == Scene->MobileDirectionalLights[2])))
+							{
+							AddViewDependentWholeSceneShadowsForView(ViewDependentWholeSceneShadows, ViewDependentWholeSceneShadowsThatNeedCulling, VisibleLightInfo, *LightSceneInfo, NumCSMCachesUpdatedThisFrame);
+							}
+
 						// Disable per-object shadows when non-nanite VSMs are enabled
 						if( !bProjectEnableNonNaniteVirtualShadowMaps && (!bMobile || (LightSceneInfo->Proxy->CastsModulatedShadows() && !LightSceneInfo->Proxy->UseCSMForDynamicObjects())))
-						{
+							{
 							const TArray<FLightPrimitiveInteraction*>* InteractionShadowPrimitives = LightSceneInfo->GetInteractionShadowPrimitives();
 
-							if (InteractionShadowPrimitives)
-							{
-								const int32 NumPrims = InteractionShadowPrimitives->Num();
-								for (int32 Idx = 0; Idx < NumPrims; ++Idx)
+								if (InteractionShadowPrimitives)
 								{
-									SetupInteractionShadows(RHICmdList, (*InteractionShadowPrimitives)[Idx], VisibleLightInfo, bStaticSceneOnly, ViewDependentWholeSceneShadows, PreShadows);
+									const int32 NumPrims = InteractionShadowPrimitives->Num();
+									for (int32 Idx = 0; Idx < NumPrims; ++Idx)
+									{
+										SetupInteractionShadows(RHICmdList, (*InteractionShadowPrimitives)[Idx], VisibleLightInfo, bStaticSceneOnly, ViewDependentWholeSceneShadows, PreShadows);
+									}
 								}
-							}
-							else
-							{
-								// Look for individual primitives with a dynamic shadow.
-								for (FLightPrimitiveInteraction* Interaction = LightSceneInfo->GetDynamicInteractionOftenMovingPrimitiveList();
-									Interaction;
-									Interaction = Interaction->GetNextPrimitive()
-									)
+								else
 								{
-									SetupInteractionShadows(RHICmdList, Interaction, VisibleLightInfo, bStaticSceneOnly, ViewDependentWholeSceneShadows, PreShadows);
-								}
+									// Look for individual primitives with a dynamic shadow.
+								for (FLightPrimitiveInteraction* Interaction = LightSceneInfo->GetDynamicInteractionOftenMovingPrimitiveList();
+										Interaction;
+										Interaction = Interaction->GetNextPrimitive()
+										)
+									{
+										SetupInteractionShadows(RHICmdList, Interaction, VisibleLightInfo, bStaticSceneOnly, ViewDependentWholeSceneShadows, PreShadows);
+									}
 
 								for (FLightPrimitiveInteraction* Interaction = LightSceneInfo->GetDynamicInteractionStaticPrimitiveList();
-									Interaction;
-									Interaction = Interaction->GetNextPrimitive()
-									)
-								{
-									SetupInteractionShadows(RHICmdList, Interaction, VisibleLightInfo, bStaticSceneOnly, ViewDependentWholeSceneShadows, PreShadows);
+										Interaction;
+										Interaction = Interaction->GetNextPrimitive()
+										)
+									{
+										SetupInteractionShadows(RHICmdList, Interaction, VisibleLightInfo, bStaticSceneOnly, ViewDependentWholeSceneShadows, PreShadows);
+									}
 								}
 							}
 						}
@@ -5593,4 +5596,4 @@ void FSceneRenderer::AllocateMobileCSMAndSpotLightShadowDepthTargets(FRHICommand
 			}
 		}
 	}
-}
+}
\ No newline at end of file
