Changed SetupMobileDirectionalLightUniformParameters VisibleLightInfos parameter to pass by ref.

diff --git a/Engine/Source/Runtime/Renderer/Private/MobileBasePassRendering.cpp b/Engine/Source/Runtime/Renderer/Private/MobileBasePassRendering.cpp
--- a/Engine/Source/Runtime/Renderer/Private/MobileBasePassRendering.cpp
+++ b/Engine/Source/Runtime/Renderer/Private/MobileBasePassRendering.cpp
@@ -224,7 +224,7 @@ TRDGUniformBufferRef<FMobileBasePassUniformParameters> CreateMobileBasePassUnifo
 void SetupMobileDirectionalLightUniformParameters(
 	const FScene& Scene,
 	const FViewInfo& SceneView,
-	const TArray<FVisibleLightInfo,SceneRenderingAllocator> VisibleLightInfos,
+	const TArray<FVisibleLightInfo,SceneRenderingAllocator>& VisibleLightInfos,
 	int32 ChannelIdx,
 	bool bDynamicShadows,
 	FMobileDirectionalLightShaderParameters& Params)
diff --git a/Engine/Source/Runtime/Renderer/Private/MobileBasePassRendering.h b/Engine/Source/Runtime/Renderer/Private/MobileBasePassRendering.h
--- a/Engine/Source/Runtime/Renderer/Private/MobileBasePassRendering.h
+++ b/Engine/Source/Runtime/Renderer/Private/MobileBasePassRendering.h
@@ -60,7 +60,7 @@ extern TRDGUniformBufferRef<FMobileBasePassUniformParameters> CreateMobileBasePa
 extern void SetupMobileDirectionalLightUniformParameters(
 	const FScene& Scene,
 	const FViewInfo& View,
-	const TArray<FVisibleLightInfo,SceneRenderingAllocator> VisibleLightInfos,
+	const TArray<FVisibleLightInfo,SceneRenderingAllocator>& VisibleLightInfos,
 	int32 ChannelIdx,
 	bool bDynamicShadows,
 	FMobileDirectionalLightShaderParameters& Parameters);
