Fix missing RT load when drawing shader print output.

diff --git a/Engine/Source/Runtime/Renderer/Private/ShaderPrint.cpp b/Engine/Source/Runtime/Renderer/Private/ShaderPrint.cpp
--- a/Engine/Source/Runtime/Renderer/Private/ShaderPrint.cpp
+++ b/Engine/Source/Runtime/Renderer/Private/ShaderPrint.cpp
@@ -383,7 +383,7 @@ namespace ShaderPrint
 			TShaderMapRef< FShaderDrawSymbolsPS > PixelShader(GlobalShaderMap);
 
 			SHADER::FParameters* PassParameters = GraphBuilder.AllocParameters<SHADER::FParameters>();
-			PassParameters->RenderTargets[0] = FRenderTargetBinding(OutputTexture, ERenderTargetLoadAction::ENoAction);
+			PassParameters->RenderTargets[0] = FRenderTargetBinding(OutputTexture, ERenderTargetLoadAction::ELoad);
 			PassParameters->UniformBufferParameters = UniformBuffer;
 			PassParameters->MiniFontTexture = FontTexture;
 			PassParameters->SymbolsBuffer = GraphBuilder.CreateSRV(SymbolBuffer);
