Fix missing applicatino of pre-shadow onto hair strands sub-pixel mask.

This CL insure that the pre-shadow are applied onto hair strands object.

diff --git a/Engine/Plugins/Runtime/HairStrands/Source/HairStrandsCore/Private/GroomComponent.cpp b/Engine/Plugins/Runtime/HairStrands/Source/HairStrandsCore/Private/GroomComponent.cpp
--- a/Engine/Plugins/Runtime/HairStrands/Source/HairStrandsCore/Private/GroomComponent.cpp
+++ b/Engine/Plugins/Runtime/HairStrands/Source/HairStrandsCore/Private/GroomComponent.cpp
@@ -297,6 +297,7 @@ public:
 	{
 		// Forcing primitive uniform as we don't support robustly GPU scene data
 		bVFRequiresPrimitiveUniformBuffer = true;
+		bCastDeepShadow = true;
 
 		HairGroupInstances = Component->HairGroupInstances;
 		check(Component);
diff --git a/Engine/Source/Runtime/Engine/Private/PrimitiveSceneProxy.cpp b/Engine/Source/Runtime/Engine/Private/PrimitiveSceneProxy.cpp
--- a/Engine/Source/Runtime/Engine/Private/PrimitiveSceneProxy.cpp
+++ b/Engine/Source/Runtime/Engine/Private/PrimitiveSceneProxy.cpp
@@ -145,6 +145,7 @@ FPrimitiveSceneProxy::FPrimitiveSceneProxy(const UPrimitiveComponent* InComponen
 ,	bCastStaticShadow(InComponent->CastShadow && InComponent->bCastStaticShadow)
 ,	bCastVolumetricTranslucentShadow(InComponent->bCastDynamicShadow && InComponent->CastShadow && InComponent->bCastVolumetricTranslucentShadow)
 ,	bCastContactShadow(InComponent->CastShadow && InComponent->bCastContactShadow)
+,	bCastDeepShadow(false)
 ,	bCastCapsuleDirectShadow(false)
 ,	bCastsDynamicIndirectShadow(false)
 ,	bCastHiddenShadow(InComponent->CastShadow && InComponent->bCastHiddenShadow)
diff --git a/Engine/Source/Runtime/Engine/Public/PrimitiveSceneProxy.h b/Engine/Source/Runtime/Engine/Public/PrimitiveSceneProxy.h
--- a/Engine/Source/Runtime/Engine/Public/PrimitiveSceneProxy.h
+++ b/Engine/Source/Runtime/Engine/Public/PrimitiveSceneProxy.h
@@ -581,6 +581,7 @@ public:
 	inline EIndirectLightingCacheQuality GetIndirectLightingCacheQuality() const { return IndirectLightingCacheQuality; }
 	inline bool CastsVolumetricTranslucentShadow() const { return bCastVolumetricTranslucentShadow; }
 	inline bool CastsContactShadow() const { return bCastContactShadow; }
+	inline bool CastsDeepShadow() const { return bCastDeepShadow; }
 	inline bool CastsCapsuleDirectShadow() const { return bCastCapsuleDirectShadow; }
 	inline bool CastsDynamicIndirectShadow() const { return bCastsDynamicIndirectShadow; }
 	inline float GetDynamicIndirectShadowMinVisibility() const { return DynamicIndirectShadowMinVisibility; }
@@ -963,6 +964,9 @@ protected:
 	/** Whether the object should cast a contact shadow */
 	uint8 bCastContactShadow : 1;
 
+	/** Whether the object should cast a deep shadow */
+	uint8 bCastDeepShadow : 1;
+
 	/** Whether the primitive should use capsules for direct shadowing, if present.  Forces inset shadows. */
 	uint8 bCastCapsuleDirectShadow : 1;
 
diff --git a/Engine/Source/Runtime/Renderer/Private/ShadowRendering.cpp b/Engine/Source/Runtime/Renderer/Private/ShadowRendering.cpp
--- a/Engine/Source/Runtime/Renderer/Private/ShadowRendering.cpp
+++ b/Engine/Source/Runtime/Renderer/Private/ShadowRendering.cpp
@@ -1016,10 +1016,15 @@ void FProjectedShadowInfo::RenderProjection(
 
 	if (bSubPixelShadow)
 	{
-		// Do not apply pre-shadow on hair, as this is intended only for targed opaque geometry
+		// Do not apply pre-shadow on opaque geometry during sub-pixel pass as we only care about opaque geometry 'casting' shadow (not receiving shadow)
+		// However, applied pre-shadow onto hair primitive (which are the only one able to cast deep shadow)
 		if (bPreShadow)
 		{
-			return;
+			const bool bIsValid = ReceiverPrimitives.Num() > 0 && ReceiverPrimitives[0]->Proxy->CastsDeepShadow();
+			if (!bIsValid)
+			{
+				return;
+			}
 		}
 
 		const bool bValidPlanes = FrustumVertices.Num() > 0;
@@ -1163,7 +1168,19 @@ void FProjectedShadowInfo::RenderProjectionInternal(
 	}
 	else if (bStencilTestEnabled)
 	{
-		if (GStencilOptimization)
+		// By pass depth/stencil test for rendering pre-shadow during sub-pixel shadow, as the hair geometry is not re-rendered
+		const bool bBypass = bSubPixelSupport && bPreShadow;
+		if (bBypass)
+		{
+			GraphicsPSOInit.DepthStencilState =
+				TStaticDepthStencilState<
+				false, CF_Always,
+				false, CF_Always, SO_Zero, SO_Zero, SO_Zero,
+				false, CF_Always, SO_Zero, SO_Zero, SO_Zero,
+				0xff, 0xff
+				>::GetRHI();
+		}
+		else if (GStencilOptimization)
 		{
 			// No depth test or writes, zero the stencil
 			// Note: this will disable hi-stencil on many GPUs, but still seems 
@@ -1405,10 +1422,15 @@ void FProjectedShadowInfo::RenderOnePassPointLightProjection(
 
 		if (bSubPixelShadow)
 		{
-			// Do not apply pre-shadow on hair, as this is intended only for targed opaque geometry
+			// Do not apply pre-shadow on opaque geometry during sub-pixel pass as we only care about opaque geometry 'casting' shadow (not receiving shadow)
+			// However, applied pre-shadow onto hair primitive (which are the only one able to cast deep shadow)
 			if (bPreShadow)
 			{
-				return;
+				const bool bIsValid = ReceiverPrimitives.Num() > 0 && ReceiverPrimitives[0]->Proxy->CastsDeepShadow();
+				if (!bIsValid)
+				{
+					return;
+				}
 			}
 
 			// Skip volume which does not intersect hair clusters
