PR #7872: Fixs bug of FPrimitiveSceneInfo::UnlinkAttachmentGroup  (Contributed by [REDACTED_PERSON])

diff --git a/Engine/Source/Runtime/Renderer/Private/PrimitiveSceneInfo.cpp b/Engine/Source/Runtime/Renderer/Private/PrimitiveSceneInfo.cpp
--- a/Engine/Source/Runtime/Renderer/Private/PrimitiveSceneInfo.cpp
+++ b/Engine/Source/Runtime/Renderer/Private/PrimitiveSceneInfo.cpp
@@ -1526,7 +1526,7 @@ void FPrimitiveSceneInfo::UnlinkAttachmentGroup()
 			if (AttachmentGroup->Primitives.Num() == 0)
 			{
 				// If this was the owner and the group is empty, remove it (otherwise the above will remove when the last attached goes).
-				Scene->AttachmentGroups.Remove(LightingAttachmentRoot);
+				Scene->AttachmentGroups.Remove(PrimitiveComponentId);
 			}
 		}
 	}
