Fix a bug that movable directional light can't cast csm on movable promitives.

diff --git a/Engine/Source/Runtime/Renderer/Private/MobileBasePass.cpp b/Engine/Source/Runtime/Renderer/Private/MobileBasePass.cpp
--- a/Engine/Source/Runtime/Renderer/Private/MobileBasePass.cpp
+++ b/Engine/Source/Runtime/Renderer/Private/MobileBasePass.cpp
@@ -338,7 +338,7 @@ ELightMapPolicyType MobileBasePass::SelectMeshLightmapPolicy(
 			}
 			else if ((bHasValidVLM || bHasValidILC) && bPrimitiveUsesILC)
 			{
-				if (bUseStaticAndCSM && !bUsesDeferredShading)
+				if ((bUseStaticAndCSM || bMovableWithCSM) && !bUsesDeferredShading)
 				{
 					SelectedLightmapPolicy = LMP_MOBILE_DIRECTIONAL_LIGHT_CSM_AND_SH_INDIRECT;
 				}
