Added an onscreen warning on mobile for misused SkyAtmoshpere

diff --git a/Engine/Source/Runtime/Renderer/Private/SceneRendering.cpp b/Engine/Source/Runtime/Renderer/Private/SceneRendering.cpp
--- a/Engine/Source/Runtime/Renderer/Private/SceneRendering.cpp
+++ b/Engine/Source/Runtime/Renderer/Private/SceneRendering.cpp
@@ -3052,14 +3052,20 @@ void FSceneRenderer::RenderFinish(FRDGBuilder& GraphBuilder, FRDGTextureRef View
 		const bool bShowMobileDynamicCSMWarning = bMobile && Scene->NumMobileStaticAndCSMLights_RenderThread > 0 && !(ReadOnlyCVARCache.bMobileEnableStaticAndCSMShadowReceivers && ReadOnlyCVARCache.bMobileAllowDistanceFieldShadows);
 		const bool bShowMobileMovableDirectionalLightWarning = bMobile && Scene->NumMobileMovableDirectionalLights_RenderThread > 0 && !ReadOnlyCVARCache.bMobileAllowMovableDirectionalLights;
 		
+		bool bMobileMissingSkyMaterial = false;
 		bool bMobileShowVertexFogWarning = false;
-		if (bMobile && Scene->ExponentialFogs.Num() > 0)
+		if (bMobile && (Scene->ExponentialFogs.Num() > 0 || Scene->HasSkyAtmosphere()))
 		{
 			static const auto* CVarDisableVertexFog = IConsoleManager::Get().FindTConsoleVariableDataInt(TEXT("r.Mobile.DisableVertexFog"));
 			if (CVarDisableVertexFog && CVarDisableVertexFog->GetValueOnRenderThread() != 0)
 			{
 				bMobileShowVertexFogWarning = true;
 			}
+			
+			if (Scene->HasSkyAtmosphere() && (Views.Num() > 0 && !Views[0].bSceneHasSkyMaterial))
+			{
+				bMobileMissingSkyMaterial = true;
+			}
 		}
 
 		const bool bSingleLayerWaterWarning = ShouldRenderSingleLayerWaterSkippedRenderEditorNotification(Views);
@@ -3069,7 +3075,7 @@ void FSceneRenderer::RenderFinish(FRDGBuilder& GraphBuilder, FRDGTextureRef View
 
 		const bool bAnyWarning = bShowPrecomputedVisibilityWarning || bShowDemotedLocalMemoryWarning || bShowGlobalClipPlaneWarning || bShowAtmosphericFogWarning || bShowSkylightWarning || bShowPointLightWarning
 			|| bShowDFAODisabledWarning || bShowShadowedLightOverflowWarning || bShowMobileDynamicCSMWarning || bShowMobileLowQualityLightmapWarning || bShowMobileMovableDirectionalLightWarning
-			|| bMobileShowVertexFogWarning || bShowSkinCacheOOM || bSingleLayerWaterWarning || bShowDFDisabledWarning || bShowNoSkyAtmosphereComponentWarning || bFxDebugDraw 
+			|| bMobileShowVertexFogWarning || bMobileMissingSkyMaterial || bShowSkinCacheOOM || bSingleLayerWaterWarning || bShowDFDisabledWarning || bShowNoSkyAtmosphereComponentWarning || bFxDebugDraw 
 			|| bShowSkyAtmosphereFogComponentsConflicts || bLumenEnabledButNoSoftwareTracing || bNaniteEnabledButNoAtomics || bRealTimeSkyCaptureButNothingToCapture;
 
 		for(int32 ViewIndex = 0;ViewIndex < Views.Num();ViewIndex++)
@@ -3091,7 +3097,7 @@ void FSceneRenderer::RenderFinish(FRDGBuilder& GraphBuilder, FRDGTextureRef View
 						[this, &ReadOnlyCVARCache, ViewState, GPUSkinCacheExtraRequiredMemory,
 						bLocked, bShowPrecomputedVisibilityWarning, bShowDemotedLocalMemoryWarning, bShowGlobalClipPlaneWarning, bShowDFAODisabledWarning, bShowDFDisabledWarning,
 						bShowAtmosphericFogWarning, bViewParentOrFrozen, bShowSkylightWarning, bShowPointLightWarning, bShowShadowedLightOverflowWarning,
-						bShowMobileLowQualityLightmapWarning, bShowMobileMovableDirectionalLightWarning, bShowMobileDynamicCSMWarning, bMobileShowVertexFogWarning,
+						bShowMobileLowQualityLightmapWarning, bShowMobileMovableDirectionalLightWarning, bShowMobileDynamicCSMWarning, bMobileShowVertexFogWarning, bMobileMissingSkyMaterial, 
 						bShowSkinCacheOOM, bSingleLayerWaterWarning, bShowNoSkyAtmosphereComponentWarning, bFxDebugDraw, FXInterface, bShowSkyAtmosphereFogComponentsConflicts, 
 						bLumenEnabledButNoSoftwareTracing, bNaniteEnabledButNoAtomics, bNaniteRequireAtomics, bRealTimeSkyCaptureButNothingToCapture]
 						(FCanvas& Canvas)
@@ -3215,6 +3221,13 @@ void FSceneRenderer::RenderFinish(FRDGBuilder& GraphBuilder, FRDGTextureRef View
 							Canvas.DrawShadowedText(10, Y, Message, GetStatsFont(), FLinearColor(1.0, 0.05, 0.05, 1.0));
 							Y += 14;
 						}
+						
+						if (bMobileMissingSkyMaterial)
+						{
+							static const FText Message = NSLOCTEXT("Renderer", "MobileMissingSkyMaterial", "On mobile the SkyAtmosphere component needs a mesh with a material tagged as IsSky and using the SkyAtmosphere nodes to visualize the Atmosphere.");
+							Canvas.DrawShadowedText(10, Y, Message, GetStatsFont(), FLinearColor(1.0, 0.05, 0.05, 1.0));
+							Y += 14;
+						}
 
 						if (bShowSkinCacheOOM)
 						{
