Support Panini projection on mobile.

diff --git a/Engine/Source/Runtime/Renderer/Private/PostProcess/PostProcessing.cpp b/Engine/Source/Runtime/Renderer/Private/PostProcess/PostProcessing.cpp
--- a/Engine/Source/Runtime/Renderer/Private/PostProcess/PostProcessing.cpp
+++ b/Engine/Source/Runtime/Renderer/Private/PostProcess/PostProcessing.cpp
@@ -1469,6 +1469,8 @@ void AddMobilePostProcessingPasses(FRDGBuilder& GraphBuilder, FScene* Scene, con
 	const FEyeAdaptationParameters EyeAdaptationParameters = GetEyeAdaptationParameters(View, ERHIFeatureLevel::ES3_1);
 	FRDGBufferRef LastEyeAdaptationBuffer = GetEyeAdaptationBuffer(GraphBuilder, View);
 
+	const FPaniniProjectionConfig PaniniConfig(View);
+
 	enum class EPass : uint32
 	{
 		Distortion,
@@ -1569,7 +1571,7 @@ void AddMobilePostProcessingPasses(FRDGBuilder& GraphBuilder, FScene* Scene, con
 	PassSequence.SetEnabled(EPass::SelectionOutline, false);
 	PassSequence.SetEnabled(EPass::EditorPrimitive, false);
 #endif
-	PassSequence.SetEnabled(EPass::PrimaryUpscale, bShouldPrimaryUpscale && bDisableUpscaleInTonemapper);
+	PassSequence.SetEnabled(EPass::PrimaryUpscale, PaniniConfig.IsEnabled() || (bShouldPrimaryUpscale && bDisableUpscaleInTonemapper));
 
 	PassSequence.SetEnabled(EPass::Visualize, View.Family->EngineShowFlags.ShaderComplexity);
 
@@ -2145,6 +2147,7 @@ void AddMobilePostProcessingPasses(FRDGBuilder& GraphBuilder, FScene* Scene, con
 		PassInputs.Method = EUpscaleMethod::Bilinear;
 		PassInputs.Stage = EUpscaleStage::PrimaryToOutput;
 		PassInputs.SceneColor = SceneColor;
+		PassInputs.PaniniConfig = PaniniConfig;
 		PassInputs.OverrideOutput.LoadAction = View.IsFirstInFamily() ? ERenderTargetLoadAction::EClear : ERenderTargetLoadAction::ELoad;
 
 		SceneColor = AddUpscalePass(GraphBuilder, View, PassInputs);
