Fix modulated shadows
Grey out Modulated Shadow option if the Light is not Stationary

diff --git a/Engine/Source/Runtime/Engine/Classes/Components/DirectionalLightComponent.h b/Engine/Source/Runtime/Engine/Classes/Components/DirectionalLightComponent.h
--- a/Engine/Source/Runtime/Engine/Classes/Components/DirectionalLightComponent.h
+++ b/Engine/Source/Runtime/Engine/Classes/Components/DirectionalLightComponent.h
@@ -239,7 +239,7 @@ class ENGINE_API UDirectionalLightComponent : public ULightComponent
 	/**
 	* Whether the light should cast modulated shadows from dynamic objects (mobile only).  Also requires Cast Shadows to be set to True.
 	**/
-	UPROPERTY(EditAnywhere, BlueprintReadOnly, Category = Light, AdvancedDisplay)
+	UPROPERTY(EditAnywhere, BlueprintReadOnly, Category = Light, AdvancedDisplay, meta = (EditCondition = "Mobility == EComponentMobility::Stationary"))
 	uint32 bCastModulatedShadows : 1;
 
 	/**
diff --git a/Engine/Source/Runtime/Renderer/Private/ShadowSetup.cpp b/Engine/Source/Runtime/Renderer/Private/ShadowSetup.cpp
--- a/Engine/Source/Runtime/Renderer/Private/ShadowSetup.cpp
+++ b/Engine/Source/Runtime/Renderer/Private/ShadowSetup.cpp
@@ -4821,7 +4821,9 @@ void FSceneRenderer::AllocateShadowDepthTargets(FRHICommandListImmediate& RHICmd
 					INC_DWORD_STAT(STAT_PerObjectShadows);
 				}
 
-				bool bNeedsProjection = ProjectedShadowInfo->CacheMode != SDCM_StaticPrimitivesOnly;
+				bool bNeedsProjection = ProjectedShadowInfo->CacheMode != SDCM_StaticPrimitivesOnly
+					// Filter out everything but PerObjectOpaqueShadows & Distance field shadows for ES31
+					&& (!bMobile || ProjectedShadowInfo->bPerObjectOpaqueShadow || ProjectedShadowInfo->bRayTracedDistanceField);
 
 				if (bNeedsProjection)
 				{if (ProjectedShadowInfo->bCapsuleShadow)
