Fixed sky and cloud not working in split screen (again, I guess something got lost during a merge from UE4Main).

diff --git a/Engine/Shaders/Private/SkyAtmosphere.usf b/Engine/Shaders/Private/SkyAtmosphere.usf
--- a/Engine/Shaders/Private/SkyAtmosphere.usf
+++ b/Engine/Shaders/Private/SkyAtmosphere.usf
@@ -1395,7 +1395,7 @@ void RenderCameraAerialPerspectiveVolumeCS(uint3 ThreadId : SV_DispatchThreadID)
 	float2 PixPos = float2(ThreadId.xy) + 0.5f;
 	float2 UV = PixPos  * SkyAtmosphere.CameraAerialPerspectiveVolumeSizeAndInvSize.zw;
 
-	float4 SVPos = float4(UV*View.ViewSizeAndInvSize.xy, 0.0f, 1.0f);// SV_POS as if resolution was the one from the scene view.
+	float4 SVPos = float4(View.ViewRectMin.xy + UV * View.ViewSizeAndInvSize.xy, 0.0f, 1.0f);// SV_POS as if resolution was the one from the scene view.
 	float3 WorldDir = GetScreenWorldDir(SVPos);
 	float3 CamPos = GetCameraPlanetPos();
 
diff --git a/Engine/Source/Runtime/Renderer/Private/VolumetricCloudRendering.cpp b/Engine/Source/Runtime/Renderer/Private/VolumetricCloudRendering.cpp
--- a/Engine/Source/Runtime/Renderer/Private/VolumetricCloudRendering.cpp
+++ b/Engine/Source/Runtime/Renderer/Private/VolumetricCloudRendering.cpp
@@ -2280,7 +2280,7 @@ bool FSceneRenderer::RenderVolumetricCloud(
 					CloudRC.RenderTargets[1] = FRenderTargetBinding(DestinationRTDepth, ERenderTargetLoadAction::ENoAction);
 					CloudRC.TracingCoordToZbufferCoordScaleBias = VRT.GetTracingCoordToZbufferCoordScaleBias();
 					// Also take into account the view rect min to be able to read correct depth
-					CloudRC.TracingCoordToZbufferCoordScaleBias.Z += ViewInfo.CachedViewUniformShaderParameters->ViewRectMin.X;
+					CloudRC.TracingCoordToZbufferCoordScaleBias.Z += ViewInfo.CachedViewUniformShaderParameters->ViewRectMin.X / ((VRT.GetMode() == 0 || VRT.GetMode() == 3) ? 2 : 1);
 					CloudRC.TracingCoordToZbufferCoordScaleBias.W += ViewInfo.CachedViewUniformShaderParameters->ViewRectMin.Y / ((VRT.GetMode() == 0 || VRT.GetMode() == 3) ? 2 : 1);
 					CloudRC.TracingCoordToFullResPixelCoordScaleBias = VRT.GetTracingCoordToFullResPixelCoordScaleBias();
 
