Minor SkyAtmosphereRendering code clean up.

diff --git a/Engine/Source/Runtime/Renderer/Private/SkyAtmosphereRendering.cpp b/Engine/Source/Runtime/Renderer/Private/SkyAtmosphereRendering.cpp
--- a/Engine/Source/Runtime/Renderer/Private/SkyAtmosphereRendering.cpp
+++ b/Engine/Source/Runtime/Renderer/Private/SkyAtmosphereRendering.cpp
@@ -1273,8 +1273,6 @@ void FSceneRenderer::RenderSkyAtmosphereLookUpTables(FRDGBuilder& GraphBuilder)
 	const FSkyAtmosphereSceneProxy& SkyAtmosphereSceneProxy = SkyInfo.GetSkyAtmosphereSceneProxy();
 
 	const bool bHighQualityMultiScattering = CVarSkyAtmosphereMultiScatteringLUTHighQuality.GetValueOnRenderThread() > 0;
-	const bool bFastSky = CVarSkyAtmosphereFastSkyLUT.GetValueOnRenderThread() > 0;
-	const bool bFastAerialPespective = CVarSkyAtmosphereAerialPerspectiveApplyOnOpaque.GetValueOnRenderThread() > 0;
 	const bool bSecondAtmosphereLightEnabled = Scene->IsSecondAtmosphereLightEnabled();
 
 	FRHISamplerState* SamplerLinearClamp = TStaticSamplerState<SF_Trilinear>::GetRHI();
@@ -1999,9 +1997,6 @@ void FDeferredShadingSceneRenderer::RenderDebugSkyAtmosphere(FRDGBuilder& GraphB
 
 	if (bSkyAtmosphereVisualizeShowFlag)
 	{
-		const bool bFastSky = CVarSkyAtmosphereFastSkyLUT.GetValueOnRenderThread() > 0;
-		const bool bFastAerialPespective = CVarSkyAtmosphereAerialPerspectiveApplyOnOpaque.GetValueOnRenderThread() > 0;
-
 		FRDGTextureRef TransmittanceLut = GraphBuilder.RegisterExternalTexture(SkyInfo.GetTransmittanceLutTexture());
 		FRDGTextureRef MultiScatteredLuminanceLut = GraphBuilder.RegisterExternalTexture(SkyInfo.GetMultiScatteredLuminanceLutTexture());
 
