Remove scene textures dependency onto hair sample compaction.

diff --git a/Engine/Shaders/Private/HairStrands/HairStrandsVisibilityCompaction.usf b/Engine/Shaders/Private/HairStrands/HairStrandsVisibilityCompaction.usf
--- a/Engine/Shaders/Private/HairStrands/HairStrandsVisibilityCompaction.usf
+++ b/Engine/Shaders/Private/HairStrands/HairStrandsVisibilityCompaction.usf
@@ -284,6 +284,8 @@ void MainCS(uint GroupIndex : SV_GroupIndex, uint3 DispatchThreadId : SV_Dispatc
 
 #define MERGE_SAMPLE 0
 
+Texture2D<float> SceneDepthTexture;
+
 #if PERMUTATION_MSAACOUNT > 1
 void ComputeUniqueSamplesWithoutDepthTest(
 	inout uint4 OutSamples[PERMUTATION_MSAACOUNT],
@@ -459,7 +461,7 @@ void MainCS(
 
 	FSampleSetDesc SampleDesc;
 
-	const float OpaqueDepth = SceneTexturesStruct.SceneDepthTexture.Load(uint3(PixelCoord, 0)).r;
+	const float OpaqueDepth = SceneDepthTexture.Load(uint3(PixelCoord, 0)).r;
 	uint4 Samples[PERMUTATION_MSAACOUNT];		// x:PrimitiveId|MaterialId, y:Weight, z:SampleIt, w:Depth (as uint)
 	ComputeUniqueSamplesWithoutDepthTest(
 		Samples,
diff --git a/Engine/Source/Runtime/Renderer/Private/HairStrands/HairStrandsVisibility.cpp b/Engine/Source/Runtime/Renderer/Private/HairStrands/HairStrandsVisibility.cpp
--- a/Engine/Source/Runtime/Renderer/Private/HairStrands/HairStrandsVisibility.cpp
+++ b/Engine/Source/Runtime/Renderer/Private/HairStrands/HairStrandsVisibility.cpp
@@ -1625,6 +1625,8 @@ class FHairVisibilityPrimitiveIdCompactionCS : public FGlobalShader
 
 		SHADER_PARAMETER_RDG_TEXTURE(Texture2D, ViewTransmittanceTexture)
 
+		SHADER_PARAMETER_RDG_TEXTURE(Texture2D, SceneDepthTexture)
+
 		SHADER_PARAMETER_RDG_TEXTURE_UAV(RWTexture2D, OutCompactNodeCounter)
 		SHADER_PARAMETER_RDG_TEXTURE_UAV(RWTexture2D, OutCompactNodeIndex)
 		SHADER_PARAMETER_RDG_TEXTURE_UAV(RWTexture2D, OutCoverageTexture)
@@ -1637,7 +1639,6 @@ class FHairVisibilityPrimitiveIdCompactionCS : public FGlobalShader
 		SHADER_PARAMETER_RDG_BUFFER_SRV(Buffer<uint>, TileCountBuffer)						// Tile total count (actual number of tiles)
 
 		SHADER_PARAMETER_STRUCT_REF(FViewUniformShaderParameters, ViewUniformBuffer)
-		SHADER_PARAMETER_RDG_UNIFORM_BUFFER(FSceneTextureUniformParameters, SceneTexturesStruct)
 		RDG_BUFFER_ACCESS(IndirectBufferArgs, ERHIAccess::IndirectArgs)
 	END_SHADER_PARAMETER_STRUCT()
 
@@ -1669,6 +1670,7 @@ static void AddHairVisibilityPrimitiveIdCompactionPass(
 	const bool bUsePPLL,
 	FRDGBuilder& GraphBuilder,
 	const FViewInfo& View,
+	const FRDGTextureRef& SceneDepthTexture,
 	const FHairStrandsMacroGroupDatas& MacroGroupDatas,
 	const uint32 NodeGroupSize,
 	const FHairStrandsTiles& TileData,
@@ -1761,7 +1763,7 @@ static void AddHairVisibilityPrimitiveIdCompactionPass(
 	PassParameters->CoverageThreshold = GetHairStrandsFullCoverageThreshold();
 	PassParameters->DepthTheshold = FMath::Clamp(GHairStrandsMaterialCompactionDepthThreshold, 0.f, 100.f);
 	PassParameters->CosTangentThreshold = FMath::Cos(FMath::DegreesToRadians(FMath::Clamp(GHairStrandsMaterialCompactionTangentThreshold, 0.f, 90.f)));
-	PassParameters->SceneTexturesStruct = CreateSceneTextureUniformBuffer(GraphBuilder, View.FeatureLevel);
+	PassParameters->SceneDepthTexture = SceneDepthTexture;
 	PassParameters->ViewUniformBuffer = View.ViewUniformBuffer;
 	PassParameters->OutCompactNodeCounter = GraphBuilder.CreateUAV(OutCompactCounter);
 	PassParameters->OutCompactNodeIndex = GraphBuilder.CreateUAV(OutCompactNodeIndex);
@@ -3405,6 +3407,7 @@ void RenderHairStrandsVisibilityBuffer(
 						false, // bUsePPLL
 						GraphBuilder,
 						View,
+						SceneDepthTexture,
 						MacroGroupDatas,
 						VisibilityData.NodeGroupSize,
 						VisibilityData.TileData,
@@ -3462,7 +3465,7 @@ void RenderHairStrandsVisibilityBuffer(
 
 						CompactNodeData = PassOutput.NodeData;
 
-						VisibilityData.SampleLightingViewportResolution = PassOutput .SampleLightingTexture->Desc.Extent;
+						VisibilityData.SampleLightingViewportResolution = PassOutput.SampleLightingTexture->Desc.Extent;
 						VisibilityData.SampleLightingTexture			= PassOutput.SampleLightingTexture;
 					}
 
@@ -3559,6 +3562,7 @@ void RenderHairStrandsVisibilityBuffer(
 						true, // bUsePPLL
 						GraphBuilder,
 						View,
+						SceneDepthTexture,
 						MacroGroupDatas,
 						VisibilityData.NodeGroupSize,
 						VisibilityData.TileData,
