Fix invalid check when running groom without tile data.

diff --git a/Engine/Source/Runtime/Renderer/Private/HairStrands/HairStrandsTile.cpp b/Engine/Source/Runtime/Renderer/Private/HairStrands/HairStrandsTile.cpp
--- a/Engine/Source/Runtime/Renderer/Private/HairStrands/HairStrandsTile.cpp
+++ b/Engine/Source/Runtime/Renderer/Private/HairStrands/HairStrandsTile.cpp
@@ -40,7 +40,7 @@ FHairStrandsTilePassVS::FParameters GetHairStrandsTileParameters(const FViewInfo
 	Out.bRectPrimitive			= InTile.bRectPrimitive ? 1 : 0;
 	Out.ViewMin					= InView.ViewRect.Min;
 	Out.ViewInvSize				= FVector2D(1.f / InView.ViewRect.Width(), 1.f / InView.ViewRect.Height());
-	Out.TileDataBuffer			= InTile.GetTileBufferSRV(TileType);
+	Out.TileDataBuffer			= InTile.IsValid() ? InTile.GetTileBufferSRV(TileType) : nullptr;
 	Out.TileIndirectBuffer		= InTile.TileIndirectDrawBuffer;
 	return Out;
 }
