Fix CIS issue.

diff --git a/Engine/Source/Runtime/Renderer/Private/HairStrands/HairStrandsData.h b/Engine/Source/Runtime/Renderer/Private/HairStrands/HairStrandsData.h
--- a/Engine/Source/Runtime/Renderer/Private/HairStrands/HairStrandsData.h
+++ b/Engine/Source/Runtime/Renderer/Private/HairStrands/HairStrandsData.h
@@ -13,6 +13,8 @@
 
 class FPrimitiveSceneProxy;
 class FViewInfo;
+class FScene;
+class FInstanceCullingManager;
 struct FMeshBatch;
 
 ////////////////////////////////////////////////////////////////////////////////////
