Add hair strands invalidation.

diff --git a/Engine/Plugins/Runtime/HairStrands/Source/HairStrandsCore/Private/GroomManager.cpp b/Engine/Plugins/Runtime/HairStrands/Source/HairStrandsCore/Private/GroomManager.cpp
--- a/Engine/Plugins/Runtime/HairStrands/Source/HairStrandsCore/Private/GroomManager.cpp
+++ b/Engine/Plugins/Runtime/HairStrands/Source/HairStrandsCore/Private/GroomManager.cpp
@@ -753,14 +753,16 @@ static void RunHairLODSelection(
 		const bool bIsLODDataReady = !!(ResourceStatus & EHairResourceStatus::Loading);
 		if (bIsLODDataReady)
 		{
+			EHairBindingType BindingType = Instance->HairGroupPublicData->GetBindingType(IntLODIndex);
 			Instance->HairGroupPublicData->SetLODVisibility(bIsVisible);
 			Instance->HairGroupPublicData->SetLODIndex(LODIndex);
 			Instance->HairGroupPublicData->SetLODBias(0);
 			Instance->HairGroupPublicData->SetMeshLODIndex(MeshLODIndex);
 			Instance->HairGroupPublicData->VFInput.GeometryType = GeometryType;
+			Instance->HairGroupPublicData->VFInput.BindingType = BindingType;
 			Instance->HairGroupPublicData->VFInput.bHasLODSwitch = (FMath::FloorToInt(PrevLODIndex) != FMath::FloorToInt(LODIndex));
 			Instance->GeometryType = GeometryType;
-			Instance->BindingType = Instance->HairGroupPublicData->GetBindingType(IntLODIndex);
+			Instance->BindingType = BindingType;
 			Instance->Guides.bIsSimulationEnable = Instance->HairGroupPublicData->IsSimulationEnable(IntLODIndex);
 			Instance->Guides.bHasGlobalInterpolation = Instance->HairGroupPublicData->IsGlobalInterpolationEnable(IntLODIndex);
 			Instance->Strands.bIsCullingEnabled = bCullingEnable;
diff --git a/Engine/Plugins/Runtime/HairStrands/Source/HairStrandsCore/Private/HairStrandsRendering.cpp b/Engine/Plugins/Runtime/HairStrands/Source/HairStrandsCore/Private/HairStrandsRendering.cpp
--- a/Engine/Plugins/Runtime/HairStrands/Source/HairStrandsCore/Private/HairStrandsRendering.cpp
+++ b/Engine/Plugins/Runtime/HairStrands/Source/HairStrandsCore/Private/HairStrandsRendering.cpp
@@ -1246,6 +1246,11 @@ class FHairRaytracingGeometryCS : public FGlobalShader
 
 public:
 	static bool ShouldCompilePermutation(const FGlobalShaderPermutationParameters& Parameters) { return IsHairStrandsSupported(EHairStrandsShaderType::Strands, Parameters.Platform); }
+	static void ModifyCompilationEnvironment(const FGlobalShaderPermutationParameters& Parameters, FShaderCompilerEnvironment& OutEnvironment)
+	{
+		FGlobalShader::ModifyCompilationEnvironment(Parameters, OutEnvironment);
+		OutEnvironment.SetDefine(TEXT("SHADER_RT_GEOMETRY"), 1);
+	}
 };
 
 IMPLEMENT_GLOBAL_SHADER(FHairRaytracingGeometryCS, "/Engine/Private/HairStrands/HairStrandsRaytracingGeometry.usf", "MainCS", SF_Compute);
diff --git a/Engine/Shaders/Private/HairStrands/HairStrandsRaytracingGeometry.usf b/Engine/Shaders/Private/HairStrands/HairStrandsRaytracingGeometry.usf
--- a/Engine/Shaders/Private/HairStrands/HairStrandsRaytracingGeometry.usf
+++ b/Engine/Shaders/Private/HairStrands/HairStrandsRaytracingGeometry.usf
@@ -4,6 +4,10 @@
 #include "HairStrandsClusterCommon.ush"
 #include "HairStrandsVertexFactoryCommon.ush"
 
+///////////////////////////////////////////////////////////////////////////////////////////////////
+
+#if SHADER_RT_GEOMETRY
+
 uint VertexCount;
 uint DispatchCountX;
 
@@ -116,46 +120,115 @@ void MainCS(uint2 DispatchThreadId : SV_DispatchThreadID)
 		SetInvalidSegment(VertexIndex0);
 		return;
 	}
+}
+
+#endif // SHADER_RT_GEOMETRY
+
+///////////////////////////////////////////////////////////////////////////////////////////////////
+
+#if SHADER_POSITION_CHANGED
+
+#include "../ShaderPrintCommon.ush"
+#include "../ShaderDrawDebug.ush"
+#include "HairStrandsAABBCommon.ush"
+
+uint VertexCount;
+uint bDrawInvalidElement;
+uint DispatchCountX;
+float PositionThreshold2;
+
+uint HairStrandsVF_bIsCullingEnable;
+Buffer<uint> 	 HairStrandsVF_CullingIndirectBuffer;
+Buffer<uint> 	 HairStrandsVF_CullingIndexBuffer;
+Buffer<uint4> 	 CurrPositionBuffer;
+Buffer<uint4> 	 PrevPositionBuffer;
+Buffer<int> 	 GroupAABBBuffer;
+RWBuffer<uint> 	 InvalidationBuffer;
+RWBuffer<uint> 	 InvalidationPrintCounter;
+
+#define GROUP_SIZE 64
+groupshared uint s_InvalidMask[GROUP_SIZE];
+
+[numthreads(GROUP_SIZE, 1, 1)]
+void MainCS(uint2 DispatchThreadId : SV_DispatchThreadID, uint LinearIndex : SV_GroupIndex)
+{
+	// Early out: avoid full comparison if previous groom have already invalidated the frame
+	const bool bIsAlreadyInvalid = InvalidationBuffer[0] > 0;
+	s_InvalidMask[LinearIndex] = 0;
+
+	if (!bIsAlreadyInvalid)
+	{
+		// Fetch vertex index if culling is enabled
+		uint VertexIndex0 = GetHairStrandsVertexFetchIndex(DispatchThreadId, GROUP_SIZE, DispatchCountX);
+		if (HairStrandsVF_bIsCullingEnable)
+		{
+			const uint VertexCountAfterCulling = HairStrandsVF_CullingIndirectBuffer[3];
+			uint FetchIndex0 = VertexIndex0;
+			if (FetchIndex0 < VertexCountAfterCulling)
+			{
+				VertexIndex0 = HairStrandsVF_CullingIndexBuffer[FetchIndex0];
+			}
+		}
 
-	FHairControlPoint P1 = GetControlPoint(VertexIndex1);
+		// Position comparison
+		if (VertexIndex0 < VertexCount)
+		{
+			const FHairControlPoint Curr = UnpackHairControlPoint(CurrPositionBuffer[VertexIndex0], float3(0, 0, 0), 1, 1, 1);
+			const FHairControlPoint Prev = UnpackHairControlPoint(PrevPositionBuffer[VertexIndex0], float3(0, 0, 0), 1, 1, 1);
+			const float3 Diff = Curr.Position - Prev.Position;
+			const bool bIsInvalid = dot(Diff, Diff) > PositionThreshold2;
+
+			// Draw red bounding box around the groom having some changed position 
+			if (bIsInvalid && bDrawInvalidElement)
+			{
+				uint Offset = 0;
+				InterlockedAdd(InvalidationPrintCounter[0], 1, Offset);
+				if (Offset < 2)
+				{
+					const FHairAABB Bound = ReadHairAABB(0, GroupAABBBuffer);
+					AddAABB(Bound.Min, Bound.Max, ColorRed);
+				}
+			}
+			s_InvalidMask[LinearIndex] = bIsInvalid ? 1u : 0u;
+		}
+	}
 
-	const float3 Tangent = normalize(P1.Position - P0.Position);
-	float3 B = float3(0,0,-1);
-	if (abs(dot(Tangent,B)) > 0.95f)
+	// Reduction
+	GroupMemoryBarrierWithGroupSync();
+	if (LinearIndex < 32)
+	{
+		s_InvalidMask[LinearIndex] = s_InvalidMask[LinearIndex] | s_InvalidMask[LinearIndex + 32];
+	}
+	GroupMemoryBarrierWithGroupSync();
+	if (LinearIndex < 16)
 	{
-		B = float3(0,1,0);
-	}	
-	float3 N = normalize(cross(Tangent, B));
-	B = normalize(cross(N, Tangent));
+		s_InvalidMask[LinearIndex] = s_InvalidMask[LinearIndex] | s_InvalidMask[LinearIndex + 16];
+	}
+	GroupMemoryBarrierWithGroupSync();
 
-	const uint BaseOutIndex = VertexIndex0 *12;
+	if (LinearIndex < 8)
 	{
-		const float3 P00 = P0.Position - N * P0.WorldRadius * LodRadiusScale0;
-		const float3 P01 = P0.Position + N * P0.WorldRadius * LodRadiusScale0;
-		const float3 P10 = P1.Position - N * P1.WorldRadius * LodRadiusScale1;
-		const float3 P11 = P1.Position + N * P1.WorldRadius * LodRadiusScale1;
-		
-		OutputPositionBuffer[BaseOutIndex  ] = float4(P00, 1);
-		OutputPositionBuffer[BaseOutIndex+1] = float4(P01, 1);
-		OutputPositionBuffer[BaseOutIndex+2] = float4(P11, 1);
-
-		OutputPositionBuffer[BaseOutIndex+3] = float4(P00, 1);
-		OutputPositionBuffer[BaseOutIndex+4] = float4(P11, 1);
-		OutputPositionBuffer[BaseOutIndex+5] = float4(P10, 1);
+		s_InvalidMask[LinearIndex] = s_InvalidMask[LinearIndex] | s_InvalidMask[LinearIndex + 8];
 	}
-	
+	GroupMemoryBarrierWithGroupSync();
+	if (LinearIndex < 4)
 	{
-		const float3 P00 = P0.Position - B * P0.WorldRadius * LodRadiusScale0;
-		const float3 P01 = P0.Position + B * P0.WorldRadius * LodRadiusScale0;
-		const float3 P10 = P1.Position - B * P1.WorldRadius * LodRadiusScale1;
-		const float3 P11 = P1.Position + B * P1.WorldRadius * LodRadiusScale1;
-
-		OutputPositionBuffer[BaseOutIndex+6] = float4(P00, 1);
-		OutputPositionBuffer[BaseOutIndex+7] = float4(P01, 1);
-		OutputPositionBuffer[BaseOutIndex+8] = float4(P11, 1);
-
-		OutputPositionBuffer[BaseOutIndex+9]  = float4(P00, 1);
-		OutputPositionBuffer[BaseOutIndex+10] = float4(P11, 1);
-		OutputPositionBuffer[BaseOutIndex+11] = float4(P10, 1);
+		s_InvalidMask[LinearIndex] = s_InvalidMask[LinearIndex] | s_InvalidMask[LinearIndex + 4];
 	}
-}
\ No newline at end of file
+	GroupMemoryBarrierWithGroupSync();
+	if (LinearIndex < 2)
+	{
+		s_InvalidMask[LinearIndex] = s_InvalidMask[LinearIndex] | s_InvalidMask[LinearIndex + 2];
+	}
+	GroupMemoryBarrierWithGroupSync();
+	if (LinearIndex < 1)
+	{
+		const uint InvalidMask = s_InvalidMask[LinearIndex] | s_InvalidMask[LinearIndex + 1];
+		if (InvalidMask > 0)
+		{
+			InterlockedOr(InvalidationBuffer[0], 1u);
+		}
+	}
+}
+
+#endif // SHADER_POSITION_CHANGED
diff --git a/Engine/Source/Runtime/Renderer/Private/HairStrands/HairStrandsData.h b/Engine/Source/Runtime/Renderer/Private/HairStrands/HairStrandsData.h
--- a/Engine/Source/Runtime/Renderer/Private/HairStrands/HairStrandsData.h
+++ b/Engine/Source/Runtime/Renderer/Private/HairStrands/HairStrandsData.h
@@ -395,14 +395,25 @@ struct FHairStrandsViewStateData
 	FRHIGPUBufferReadback* GetBuffer() const { return VoxelPageAllocationCountReadback; }
 	bool IsReady() const { return VoxelPageAllocationCountReadback->IsReady(); }
 	bool IsInit() const { return VoxelPageAllocationCountReadback != nullptr; }
+
 	void Init();
 	void Release();
 
 	float VoxelWorldSize = 0; // Voxel size used during the last frame allocation
-	uint32 AllocatedPageCount = 0; // Number of voxels allocated last frame
+	uint32 VoxelAllocatedPageCount = 0; // Number of voxels allocated last frame
 
 	// Buffer used for reading back the number of voxels allocated on the GPU
 	FRHIGPUBufferReadback* VoxelPageAllocationCountReadback = nullptr;
+
+	// Buffer used for reading back hair strands position changed on the GPU
+	struct FPositionChangedData
+	{
+		FRHIGPUBufferReadback* ReadbackBuffer = nullptr;
+		bool bHasPendingReadback = false;
+	};
+	void EnqueuePositionsChanged(FRDGBuilder& GraphBuilder, FRDGBufferRef InBuffer);
+	bool ReadPositionsChanged();
+	TArray<FPositionChangedData> PositionsChangedDatas;
 };
 
 namespace HairStrands
@@ -416,6 +427,7 @@ namespace HairStrands
 	bool HasViewHairStrandsData(const TArrayView<FViewInfo>& Views);
 	bool HasViewHairStrandsVoxelData(const FViewInfo& View);
 
+	bool HasPositionsChanged(FRDGBuilder& GraphBuilder, const FViewInfo& View);
 	void DrawHitProxies(FRDGBuilder& GraphBuilder, const FScene& Scene, const FViewInfo& View, FInstanceCullingManager& InstanceCullingManager, FRDGTextureRef HitProxyTexture, FRDGTextureRef HitProxyDepthTexture);
 	void DrawEditorSelection(FRDGBuilder& GraphBuilder, const FViewInfo& View, FRDGTextureRef SelectionDepthTexture);
 }
\ No newline at end of file
diff --git a/Engine/Source/Runtime/Renderer/Private/HairStrands/HairStrandsRendering.cpp b/Engine/Source/Runtime/Renderer/Private/HairStrands/HairStrandsRendering.cpp
--- a/Engine/Source/Runtime/Renderer/Private/HairStrands/HairStrandsRendering.cpp
+++ b/Engine/Source/Runtime/Renderer/Private/HairStrands/HairStrandsRendering.cpp
@@ -4,6 +4,7 @@
 #include "HairStrandsData.h"
 #include "SceneRendering.h"
 #include "ScenePrivate.h"
+#include "RenderGraphUtils.h"
 
 static TRDGUniformBufferRef<FHairStrandsViewUniformParameters> InternalCreateHairStrandsViewUniformBuffer(
 	FRDGBuilder& GraphBuilder, 
@@ -177,23 +178,70 @@ void RenderHairBasePass(
 
 void FHairStrandsViewStateData::Init()
 {
+	// Voxel adaptive sizing
 	VoxelWorldSize = 0;
-	AllocatedPageCount = 0;
+	VoxelAllocatedPageCount = 0;
 	if (VoxelPageAllocationCountReadback == nullptr)
 	{
-		VoxelPageAllocationCountReadback = new FRHIGPUBufferReadback(TEXT("Voxel page allocation readback"));
+		VoxelPageAllocationCountReadback = new FRHIGPUBufferReadback(TEXT("Hair.VoxelPageAllocationReadback"));
+	}
+	
+	// Track if hair strands positions has changed
+	PositionsChangedDatas.SetNum(4);
+	for (FPositionChangedData& Data : PositionsChangedDatas)
+	{
+		Data.ReadbackBuffer = new FRHIGPUBufferReadback(TEXT("Hair.PositionsChangedReadback"));
+		Data.bHasPendingReadback = false;
 	}
 }
 
 void FHairStrandsViewStateData::Release()
 {
+	// Voxel adaptive sizing
 	VoxelWorldSize = 0;
-	AllocatedPageCount = 0;
+	VoxelAllocatedPageCount = 0;
 	if (VoxelPageAllocationCountReadback)
 	{
 		delete VoxelPageAllocationCountReadback;
 		VoxelPageAllocationCountReadback = nullptr;
 	}
+
+	// Track if hair strands positions has changed
+	for (FPositionChangedData& Data : PositionsChangedDatas)
+	{
+		delete Data.ReadbackBuffer;
+		Data.ReadbackBuffer = nullptr;
+		Data.bHasPendingReadback = false;
+	}
+	PositionsChangedDatas.Empty();
+}
+
+void FHairStrandsViewStateData::EnqueuePositionsChanged(FRDGBuilder& GraphBuilder, FRDGBufferRef InBuffer)
+{
+	for (FPositionChangedData& Data : PositionsChangedDatas)
+	{
+		if (!Data.bHasPendingReadback)
+		{
+			AddEnqueueCopyPass(GraphBuilder, Data.ReadbackBuffer, InBuffer, 4u);
+			Data.bHasPendingReadback = true;
+			break;
+		}
+	}
+}
+
+bool FHairStrandsViewStateData::ReadPositionsChanged()
+{
+	for (FPositionChangedData& Data : PositionsChangedDatas)
+	{
+		if (Data.bHasPendingReadback && Data.ReadbackBuffer->IsReady())
+		{
+			const uint32 ReadData = *(uint32*)(Data.ReadbackBuffer->Lock(sizeof(uint32)));
+			Data.ReadbackBuffer->Unlock();
+			Data.bHasPendingReadback = false;
+			return ReadData > 0;
+		}
+	}
+	return false;
 }
 
 namespace HairStrands
diff --git a/Engine/Source/Runtime/Renderer/Private/HairStrands/HairStrandsVisibility.cpp b/Engine/Source/Runtime/Renderer/Private/HairStrands/HairStrandsVisibility.cpp
--- a/Engine/Source/Runtime/Renderer/Private/HairStrands/HairStrandsVisibility.cpp
+++ b/Engine/Source/Runtime/Renderer/Private/HairStrands/HairStrandsVisibility.cpp
@@ -84,6 +84,12 @@ static FAutoConsoleVariableRef CVarHairStrandsTile(TEXT("r.HairStrands.Tile"), G
 static int32 GHairStrandsLightSampleFormat = 1;
 static FAutoConsoleVariableRef CVarHairStrandsLightSampleFormat(TEXT("r.HairStrands.LightSampleFormat"), GHairStrandsLightSampleFormat, TEXT("Define the format used for storing the lighting of hair samples (0: RGBA-16bits, 1: RGB-11.11.10bits)"));
 
+static float GHairStrands_InvalidationPosition_Threshold = 0.05f;
+static FAutoConsoleVariableRef CVarHairStrands_InvalidationPosition_Threshold(TEXT("r.HairStrands.PathTracing.InvalidationThreshold"), GHairStrands_InvalidationPosition_Threshold, TEXT("Define the minimal distance to invalidate path tracer output when groom changes (in cm, default: 0.5mm)"));
+
+static int32 GHairStrands_InvalidationPosition_Debug = 0;
+static FAutoConsoleVariableRef CVarHairStrands_InvalidationPosition_Debug(TEXT("r.HairStrands.PathTracing.InvalidationDebug"), GHairStrands_InvalidationPosition_Debug, TEXT("Enable bounding box drawing for groom element causing path tracer invalidation"));
+
 /////////////////////////////////////////////////////////////////////////////////////////
 
 namespace HairStrandsVisibilityInternal
@@ -3140,10 +3146,10 @@ static FRasterComputeOutput AddVisibilityComputeRasterPass(
 ///////////////////////////////////////////////////////////////////////////////////////////////////
 // Selection outline
 
-class FHairStrandsEmitSelectionPS : public FNaniteShader
+class FHairStrandsEmitSelectionPS : public FGlobalShader
 {
 	DECLARE_GLOBAL_SHADER(FHairStrandsEmitSelectionPS);
-	SHADER_USE_PARAMETER_STRUCT(FHairStrandsEmitSelectionPS, FNaniteShader);
+	SHADER_USE_PARAMETER_STRUCT(FHairStrandsEmitSelectionPS, FGlobalShader);
 
 	BEGIN_SHADER_PARAMETER_STRUCT(FParameters, )
 		SHADER_PARAMETER(uint32, MaxMaterialCount)
@@ -3212,10 +3218,10 @@ void AddHairStrandsSelectionOutlinePass(
 ///////////////////////////////////////////////////////////////////////////////////////////////////
 // HitProxyId
 
-class FHairStrandsEmitHitProxyIdPS : public FNaniteShader
+class FHairStrandsEmitHitProxyIdPS : public FGlobalShader
 {
 	DECLARE_GLOBAL_SHADER(FHairStrandsEmitHitProxyIdPS);
-	SHADER_USE_PARAMETER_STRUCT(FHairStrandsEmitHitProxyIdPS, FNaniteShader);
+	SHADER_USE_PARAMETER_STRUCT(FHairStrandsEmitHitProxyIdPS, FGlobalShader);
 
 	BEGIN_SHADER_PARAMETER_STRUCT(FParameters, )
 		SHADER_PARAMETER(uint32, MaxMaterialCount)
@@ -3283,9 +3289,88 @@ void AddHairStrandsHitProxyIdPass(
 #endif
 }
 
+///////////////////////////////////////////////////////////////////////////////////////////////////
+// Groom comparison
+class FHairStrandsPositionChangedCS : public FGlobalShader
+{
+	DECLARE_GLOBAL_SHADER(FHairStrandsPositionChangedCS);
+	SHADER_USE_PARAMETER_STRUCT(FHairStrandsPositionChangedCS, FGlobalShader);
+
+	BEGIN_SHADER_PARAMETER_STRUCT(FParameters, )
+		SHADER_PARAMETER(uint32, VertexCount)
+		SHADER_PARAMETER(uint32, DispatchCountX)
+		SHADER_PARAMETER(float, PositionThreshold2)
+		SHADER_PARAMETER(uint32, HairStrandsVF_bIsCullingEnable)
+		SHADER_PARAMETER(uint32, bDrawInvalidElement)
+		SHADER_PARAMETER_RDG_BUFFER_SRV(Buffer<uint>, HairStrandsVF_CullingIndexBuffer)
+		SHADER_PARAMETER_RDG_BUFFER_SRV(Buffer<uint>, HairStrandsVF_CullingIndirectBuffer)
+		SHADER_PARAMETER_RDG_BUFFER_SRV(Buffer<uint4>, CurrPositionBuffer)
+		SHADER_PARAMETER_RDG_BUFFER_SRV(Buffer<uint4>, PrevPositionBuffer)
+		SHADER_PARAMETER_RDG_BUFFER_SRV(Buffer<int>, GroupAABBBuffer)
+		SHADER_PARAMETER_RDG_BUFFER_UAV(RWBuffer<uint>, InvalidationBuffer)
+		SHADER_PARAMETER_RDG_BUFFER_UAV(RWBuffer<uint>, InvalidationPrintCounter)
+		SHADER_PARAMETER_STRUCT_INCLUDE(ShaderDrawDebug::FShaderParameters, ShaderDrawParameters)
+	END_SHADER_PARAMETER_STRUCT()
+public:
+	static bool ShouldCompilePermutation(const FGlobalShaderPermutationParameters& Parameters) { return IsHairStrandsSupported(EHairStrandsShaderType::Strands, Parameters.Platform); }
+	static void ModifyCompilationEnvironment(const FGlobalShaderPermutationParameters& Parameters, FShaderCompilerEnvironment& OutEnvironment)
+	{
+		OutEnvironment.SetDefine(TEXT("SHADER_POSITION_CHANGED"), TEXT("1"));
+	}
+};
+
+IMPLEMENT_GLOBAL_SHADER(FHairStrandsPositionChangedCS, "/Engine/Private/HairStrands/HairStrandsRaytracingGeometry.usf", "MainCS", SF_Compute);
+
+static void AddHairStrandsHasPositionChangedPass(
+	FRDGBuilder& GraphBuilder,
+	const FViewInfo* View,
+	const FHairGroupPublicData* HairGroupPublicData,
+	FRDGBufferUAVRef InvalidationBuffer)
+{
+	const uint32 VertexCount = HairGroupPublicData->VertexCount;
+	const uint32 GroupSize = 64;
+	const FIntVector DispatchCount = ComputeDispatchCount(VertexCount, GroupSize);
+
+	FRDGBufferRef InvalidationPrintCounter = GraphBuilder.CreateBuffer(FRDGBufferDesc::CreateBufferDesc(sizeof(uint32), 1), TEXT("Hair.InvalidationPrintCounter"));
+	FRDGBufferUAVRef InvalidationPrintCounterUAV = GraphBuilder.CreateUAV(InvalidationPrintCounter, PF_R32_UINT);
+	AddClearUAVPass(GraphBuilder, InvalidationPrintCounterUAV, 0u);
+
+	FHairStrandsPositionChangedCS::FParameters* Parameters = GraphBuilder.AllocParameters<FHairStrandsPositionChangedCS::FParameters>();
+	Parameters->VertexCount = VertexCount;
+	Parameters->DispatchCountX = DispatchCount.X;
+	Parameters->PositionThreshold2 = FMath::Square(FMath::Clamp(GHairStrands_InvalidationPosition_Threshold, 0, 1.0f));
+	Parameters->bDrawInvalidElement = GHairStrands_InvalidationPosition_Debug > 0 ? 1u : 0u;
+	Parameters->HairStrandsVF_bIsCullingEnable = 0u;
+	Parameters->HairStrandsVF_CullingIndexBuffer = GraphBuilder.CreateSRV(GSystemTextures.GetDefaultBuffer(GraphBuilder, 4, 0u), PF_R32_UINT);
+	Parameters->HairStrandsVF_CullingIndirectBuffer = Parameters->HairStrandsVF_CullingIndexBuffer;
+	Parameters->CurrPositionBuffer = HairGroupPublicData->VFInput.Strands.PositionBuffer.SRV;
+	Parameters->PrevPositionBuffer = HairGroupPublicData->VFInput.Strands.PrevPositionBuffer.SRV;
+	Parameters->GroupAABBBuffer = Register(GraphBuilder, HairGroupPublicData->GetGroupAABBBuffer(), ERDGImportedBufferFlags::CreateSRV).SRV;
+	Parameters->InvalidationBuffer = InvalidationBuffer;
+	Parameters->InvalidationPrintCounter = InvalidationPrintCounterUAV;
+	ShaderDrawDebug::SetParameters(GraphBuilder, View->ShaderDrawData, Parameters->ShaderDrawParameters);
+	if (HairGroupPublicData->GetCullingResultAvailable())
+	{
+		Parameters->HairStrandsVF_CullingIndexBuffer = Register(GraphBuilder, HairGroupPublicData->GetCulledVertexIdBuffer(), ERDGImportedBufferFlags::CreateSRV).SRV;
+		Parameters->HairStrandsVF_CullingIndirectBuffer = Register(GraphBuilder, HairGroupPublicData->GetDrawIndirectRasterComputeBuffer(), ERDGImportedBufferFlags::CreateSRV).SRV;
+		Parameters->HairStrandsVF_bIsCullingEnable = 1;
+	}
+
+	TShaderMapRef<FHairStrandsPositionChangedCS> ComputeShader(View->ShaderMap);
+	FComputeShaderUtils::AddPass(
+		GraphBuilder,
+		RDG_EVENT_NAME("HairStrands::HasPositionChanged"),
+		ComputeShader,
+		Parameters,
+		DispatchCount);
+}
+
+///////////////////////////////////////////////////////////////////////////////////////////////////
+
 namespace HairStrands
 {
 
+// Draw hair strands depth value for outline selection
 void DrawEditorSelection(FRDGBuilder& GraphBuilder, const FViewInfo& View, FRDGTextureRef SelectionDepthTexture)
 {
 	AddHairStrandsSelectionOutlinePass(
@@ -3296,6 +3381,7 @@ void DrawEditorSelection(FRDGBuilder& GraphBuilder, const FViewInfo& View, FRDGT
 		SelectionDepthTexture);
 }
 
+// Draw hair strands hit proxy values
 void DrawHitProxies(
 	FRDGBuilder& GraphBuilder,
 	const FScene& Scene,
@@ -3378,6 +3464,52 @@ void DrawHitProxies(
 	AddHairStrandsHitProxyIdPass(GraphBuilder, Scene, View, VisNodeIndex, VisNodeData, HitProxyTexture, HitProxyDepthTexture);
 }
 
+// Check if any simulated/skinned-bound groom has its positions updated (e.g. for invalidating the path-tracer accumulation)
+bool HasPositionsChanged(FRDGBuilder& GraphBuilder, const FViewInfo& View)
+{
+	FHairStrandsViewStateData* HairStrandsViewStateData = const_cast<FHairStrandsViewStateData*>(&View.ViewState->HairStrandsViewStateData);
+	if (!HairStrandsViewStateData->IsInit())
+	{
+		HairStrandsViewStateData->Init();
+	}
+
+	TArray<const FHairGroupPublicData*> GroupDatas;
+	for (const FMeshBatchAndRelevance& Batch : View.HairStrandsMeshElements)
+	{
+		FHairGroupPublicData* HairGroupPublicData = reinterpret_cast<FHairGroupPublicData*>(Batch.Mesh->Elements[0].VertexFactoryUserData);
+		check(HairGroupPublicData);
+		const int32 LODIndex = FMath::FloorToInt(HairGroupPublicData->LODIndex);
+		const bool bHasSimulationOrSkinning = 
+			HairGroupPublicData->GetGeometryType(LODIndex) == EHairGeometryType::Strands &&
+			(HairGroupPublicData->IsSimulationEnable(LODIndex) || HairGroupPublicData->GetBindingType(LODIndex) == EHairBindingType::Skinning);
+		if (bHasSimulationOrSkinning)
+		{
+			GroupDatas.Add(HairGroupPublicData);
+		}
+	}
+
+	FRDGBufferDesc Desc = FRDGBufferDesc::CreateBufferDesc(sizeof(uint32), 1);
+	Desc.Usage |= BUF_SourceCopy;
+	FRDGBufferRef InvalidationBuffer = GraphBuilder.CreateBuffer(Desc, TEXT("Hair.HasSimulationRunningBuffer"));
+	FRDGBufferUAVRef InvalidationUAV = GraphBuilder.CreateUAV(InvalidationBuffer, PF_R32_UINT);
+	AddClearUAVPass(GraphBuilder, InvalidationUAV, 0u);
+
+	// Compare current/previous and enqueue aggregated comparison
+	for (const FHairGroupPublicData* GroupData : GroupDatas)
+	{
+		AddHairStrandsHasPositionChangedPass(GraphBuilder, &View, GroupData, InvalidationUAV);
+	}
+
+	// Pull a 'ready' previous frame value
+	bool bHasPositionChanged = HairStrandsViewStateData->ReadPositionsChanged();
+
+	// Enqueue new readback request
+	HairStrandsViewStateData->EnqueuePositionsChanged(GraphBuilder, InvalidationBuffer);
+
+
+	return bHasPositionChanged;
+}
+
 }
 
 ///////////////////////////////////////////////////////////////////////////////////////////////////
diff --git a/Engine/Source/Runtime/Renderer/Private/HairStrands/HairStrandsVoxelization.cpp b/Engine/Source/Runtime/Renderer/Private/HairStrands/HairStrandsVoxelization.cpp
--- a/Engine/Source/Runtime/Renderer/Private/HairStrands/HairStrandsVoxelization.cpp
+++ b/Engine/Source/Runtime/Renderer/Private/HairStrands/HairStrandsVoxelization.cpp
@@ -997,7 +997,7 @@ static FHairStrandsVoxelResources AllocateVirtualVoxelResources(
 
 		// Update state data
 		OutViewStateData->VoxelWorldSize = VoxelWorldSize;
-		OutViewStateData->AllocatedPageCount = AllocatedPageCount;
+		OutViewStateData->VoxelAllocatedPageCount = AllocatedPageCount;
 	}
 
 	FHairStrandsVoxelResources Out;
diff --git a/Engine/Source/Runtime/Renderer/Private/PathTracing.cpp b/Engine/Source/Runtime/Renderer/Private/PathTracing.cpp
--- a/Engine/Source/Runtime/Renderer/Private/PathTracing.cpp
+++ b/Engine/Source/Runtime/Renderer/Private/PathTracing.cpp
@@ -24,6 +24,7 @@ TAutoConsoleVariable<int32> CVarPathTracing(
 #include "PathTracingDefinitions.h"
 #include "RayTracing/RayTracingMaterialHitShaders.h"
 #include "RenderCore/Public/GenerateMips.h"
+#include "HairStrands/HairStrandsData.h"
 #include <limits>
 
 TAutoConsoleVariable<int32> CVarPathTracingCompaction(
@@ -1438,7 +1439,7 @@ void FDeferredShadingSceneRenderer::RenderPathTracing(
 
 	// If the scene has changed in some way (camera move, object movement, etc ...)
 	// we must invalidate the ViewState to start over from scratch
-	if (FirstTime || Config.IsDifferent(PathTracingState->LastConfig))
+	if (FirstTime || Config.IsDifferent(PathTracingState->LastConfig) || HairStrands::HasPositionsChanged(GraphBuilder, View))
 	{
 		// remember the options we used for next time
 		PathTracingState->LastConfig = Config;
diff --git a/Engine/Source/Runtime/Renderer/Public/HairStrandsInterface.h b/Engine/Source/Runtime/Renderer/Public/HairStrandsInterface.h
--- a/Engine/Source/Runtime/Renderer/Public/HairStrandsInterface.h
+++ b/Engine/Source/Runtime/Renderer/Public/HairStrandsInterface.h
@@ -171,6 +171,7 @@ public:
 	FRDGExternalBuffer& GetDrawIndirectBuffer() { return DrawIndirectBuffer; }
 	FRDGExternalBuffer& GetClusterAABBBuffer() { return ClusterAABBBuffer; }
 	FRDGExternalBuffer& GetGroupAABBBuffer() { return GroupAABBBuffer; }
+	const FRDGExternalBuffer& GetGroupAABBBuffer() const { return GroupAABBBuffer; }
 
 	const FRDGExternalBuffer& GetCulledVertexIdBuffer() const { return CulledVertexIdBuffer; }
 	const FRDGExternalBuffer& GetCulledVertexRadiusScaleBuffer() const { return CulledVertexRadiusScaleBuffer; }
