Move FillInstanceUploadBuffer/BuildRayTracingInstanceBuffer from RenderCore to Renderer
- Necessary to access GPUScene

diff --git a/Engine/Plugins/Experimental/GPULightmass/Source/GPULightmass/Private/LightmapRenderer.cpp b/Engine/Plugins/Experimental/GPULightmass/Source/GPULightmass/Private/LightmapRenderer.cpp
--- a/Engine/Plugins/Experimental/GPULightmass/Source/GPULightmass/Private/LightmapRenderer.cpp
+++ b/Engine/Plugins/Experimental/GPULightmass/Source/GPULightmass/Private/LightmapRenderer.cpp
@@ -23,7 +23,7 @@
 #include "EngineModule.h"
 #include "PostProcess/PostProcessing.h"
 #include "RayTracingGeometryManager.h"
-#include "RayTracingInstanceCopyShader.h"
+#include "RayTracingInstanceBufferUtil.h"
 #include "ScreenPass.h"
 #include "RayTracingDynamicGeometryCollection.h"
 #include "Experimental/Containers/SherwoodHashTable.h"
diff --git a/Engine/Shaders/Private/RayTracing/RayTracingInstanceBufferUtil.usf b/Engine/Shaders/Private/RayTracing/RayTracingInstanceBufferUtil.usf
new file mode 100644
--- /dev/null
+++ b/Engine/Shaders/Private/RayTracing/RayTracingInstanceBufferUtil.usf
@@ -0,0 +1,76 @@
+// Copyright Epic Games, Inc. All Rights Reserved.
+
+/*=============================================================================================
+	RayTracingInstanceCopy.usf: Build ray tracing instances in the GPU 
+===============================================================================================*/
+
+#include "RayTracingCommon.ush"
+//#include "../SceneData.ush"
+#include "../Common.ush"
+
+#ifndef THREADGROUP_SIZE
+#define THREADGROUP_SIZE 1
+#endif
+
+uint NumInstances;
+uint DescBufferOffset;
+StructuredBuffer<float3x4> InstancesTransforms;
+RWStructuredBuffer<FRayTracingInstanceDescriptor> InstancesDescriptors;
+
+
+[numthreads(THREADGROUP_SIZE, 1, 1)]
+void RayTracingInstanceCopyShaderCS(uint3 DispatchThreadId : SV_DispatchThreadID)
+{
+	uint Index = DispatchThreadId.x;	
+
+	float3x4 InstanceXForm = InstancesTransforms[Index];
+	uint DescIndex = DescBufferOffset + Index;
+
+	if (Index < NumInstances)
+	{
+		SetRayTracingInstanceTransform(InstancesDescriptors[DescIndex], InstanceXForm);
+	}
+}
+
+struct FRayTracingInstanceDescriptorInput
+{
+	//uint GPUSceneInstanceIndex;
+	float3x4 LocalToWorld;
+	uint AccelerationStructureIndex;
+	uint InstanceId;
+	uint InstanceMaskAndFlags;
+	uint InstanceContributionToHitGroupIndex;
+};
+
+StructuredBuffer<FRayTracingInstanceDescriptorInput> InputInstanceDescriptors;
+ByteAddressBuffer AccelerationStructureAddresses;
+
+[numthreads(THREADGROUP_SIZE, 1, 1)]
+void RayTracingBuildInstanceBufferCS(uint3 DispatchThreadId : SV_DispatchThreadID)
+{
+	uint Index = DispatchThreadId.x;
+
+	if (Index < NumInstances)
+	{
+		FRayTracingInstanceDescriptorInput InputDesc = InputInstanceDescriptors[Index];
+
+		//FInstanceSceneData InstanceSceneData = GetInstanceSceneData(InputDesc.GPUSceneInstanceIndex, View.InstanceSceneDataSOAStride);
+		//float4x4 LocalToWorld = InstanceSceneData.LocalToWorld;
+		float3x4 LocalToWorld = InputDesc.LocalToWorld;
+
+		uint2 BlasAddress = 0;
+
+		if (InputDesc.AccelerationStructureIndex != 0xFFFFFFFF)
+		{
+			BlasAddress = AccelerationStructureAddresses.Load2(InputDesc.AccelerationStructureIndex * 8);
+		}
+
+		InstancesDescriptors[Index] = BuildPlatformRayTracingInstanceDesc(
+			InputDesc.InstanceMaskAndFlags & 0xFF,
+			InputDesc.InstanceId,
+			TranslateRayTracingInstanceFlags((InputDesc.InstanceMaskAndFlags >> 8) & 0xFF),
+			InputDesc.InstanceContributionToHitGroupIndex,
+			LocalToWorld,
+			BlasAddress);
+	}
+}
diff --git a/Engine/Shaders/Private/RayTracing/RayTracingInstanceCopy.usf b/Engine/Shaders/Private/RayTracing/RayTracingInstanceCopy.usf
deleted file mode 100644
diff --git a/Engine/Source/Runtime/RenderCore/Private/RayTracingInstanceCopyShader.cpp b/Engine/Source/Runtime/RenderCore/Private/RayTracingInstanceCopyShader.cpp
deleted file mode 100644
diff --git a/Engine/Source/Runtime/RenderCore/Public/RayTracingInstanceCopyShader.h b/Engine/Source/Runtime/RenderCore/Public/RayTracingInstanceCopyShader.h
deleted file mode 100644
diff --git a/Engine/Source/Runtime/Renderer/Private/RayTracing/RayTracingInstanceBufferUtil.cpp b/Engine/Source/Runtime/Renderer/Private/RayTracing/RayTracingInstanceBufferUtil.cpp
new file mode 100644
--- /dev/null
+++ b/Engine/Source/Runtime/Renderer/Private/RayTracing/RayTracingInstanceBufferUtil.cpp
@@ -0,0 +1,166 @@
+// Copyright Epic Games, Inc. All Rights Reserved.
+
+#include "RayTracingInstanceBufferUtil.h"
+
+#include "RayTracingDefinitions.h"
+
+#include "RenderGraphBuilder.h"
+#include "ShaderParameterUtils.h"
+#include "RendererInterface.h"
+#include "ShaderParameterStruct.h"
+#include "GlobalShader.h"
+
+#include "Async/ParallelFor.h"
+
+#if RHI_RAYTRACING
+
+void FillInstanceUploadBuffer(
+	TConstArrayView<FRayTracingGeometryInstance> Instances,
+	TConstArrayView<uint32> InstancesGeometryIndex,
+	FRayTracingSceneRHIRef RayTracingSceneRHI,
+	TArrayView<FRayTracingInstanceDescriptorInput> OutInstanceUploadData)
+{
+	const FRayTracingSceneInitializer2& SceneInitializer = RayTracingSceneRHI->GetInitializer();
+
+	int32 NumInactiveNativeInstances = 0;
+
+	const int32 NumSceneInstances = Instances.Num();
+	ParallelFor(NumSceneInstances, [OutInstanceUploadData, Instances, InstancesGeometryIndex, &SceneInitializer, &NumInactiveNativeInstances](int32 SceneInstanceIndex)
+		{
+			const FRayTracingGeometryInstance& SceneInstance = Instances[SceneInstanceIndex];
+
+			const uint32 NumTransforms = SceneInstance.NumTransforms;
+
+			checkf(SceneInstance.UserData.Num() == 0 || SceneInstance.UserData.Num() >= int32(NumTransforms),
+				TEXT("User data array must be either be empty (Instance.DefaultUserData is used), or contain one entry per entry in Transforms array."));
+
+			check(SceneInstance.ActivationMask.IsEmpty() || SceneInstance.ActivationMask.Num() * 32 >= int32(NumTransforms));
+
+			const bool bUseUniqueUserData = SceneInstance.UserData.Num() != 0;
+
+			const bool bGpuInstance = SceneInstance.GPUTransformsSRV != nullptr;
+			const bool bCpuInstance = !bGpuInstance;
+
+			const uint32 AccelerationStructureIndex = InstancesGeometryIndex[SceneInstanceIndex];
+
+			uint32 BaseInstanceIndex = SceneInitializer.BaseInstancePrefixSum[SceneInstanceIndex];
+
+			int32 NumInactiveNativeInstancesThisSceneInstance = 0;
+			for (uint32 TransformIndex = 0; TransformIndex < NumTransforms; ++TransformIndex)
+			{
+				FRayTracingInstanceDescriptorInput& InstanceDesc = OutInstanceUploadData[BaseInstanceIndex + TransformIndex];
+				InstanceDesc.InstanceMaskAndFlags = SceneInstance.Mask | ((uint32)SceneInstance.Flags << 8);
+				InstanceDesc.InstanceContributionToHitGroupIndex = SceneInitializer.SegmentPrefixSum[SceneInstanceIndex] * RAY_TRACING_NUM_SHADER_SLOTS;
+				InstanceDesc.InstanceId = bUseUniqueUserData ? SceneInstance.UserData[TransformIndex] : SceneInstance.DefaultUserData;
+				//InstanceDesc.GPUSceneInstanceIndex = 0; // TODO
+
+				InstanceDesc.AccelerationStructureIndex = AccelerationStructureIndex;
+
+				if (!SceneInstance.ActivationMask.IsEmpty() && (SceneInstance.ActivationMask[TransformIndex / 32] & (1 << (TransformIndex % 32))) == 0)
+				{
+					// Set flag for deactivated instances
+					InstanceDesc.AccelerationStructureIndex = 0xFFFFFFFF;
+					NumInactiveNativeInstancesThisSceneInstance++;
+					continue;
+				}
+
+				if (LIKELY(bCpuInstance))
+				{
+					const FMatrix44f LocalToWorld = SceneInstance.Transforms[TransformIndex].GetTransposed();
+					InstanceDesc.LocalToWorld[0] = *(FVector4f*)&LocalToWorld.M[0];
+					InstanceDesc.LocalToWorld[1] = *(FVector4f*)&LocalToWorld.M[1];
+					InstanceDesc.LocalToWorld[2] = *(FVector4f*)&LocalToWorld.M[2];
+				}
+				else
+				{
+					// #todo: GPU-based instances transform should be copied from GPUTransformsSRV
+					FMemory::Memset(&InstanceDesc.LocalToWorld, 0, sizeof(InstanceDesc.LocalToWorld));
+				}
+			}
+
+#if STATS
+			FPlatformAtomics::InterlockedAdd(&NumInactiveNativeInstances, NumInactiveNativeInstancesThisSceneInstance);
+#endif
+		});
+
+	SET_DWORD_STAT(STAT_RayTracingInstances, SceneInitializer.NumNativeInstances - NumInactiveNativeInstances);
+}
+
+struct FRayTracingInstanceCopyCS : public FGlobalShader
+{
+	DECLARE_GLOBAL_SHADER(FRayTracingInstanceCopyCS);
+	SHADER_USE_PARAMETER_STRUCT(FRayTracingInstanceCopyCS, FGlobalShader);
+
+	BEGIN_SHADER_PARAMETER_STRUCT(FParameters, )
+		SHADER_PARAMETER_RDG_BUFFER_UAV(RWStructuredBuffer, InstancesDescriptors)
+		SHADER_PARAMETER_RDG_BUFFER_SRV(StructuredBuffer, InstancesTransforms)
+		SHADER_PARAMETER(uint32, NumInstances)
+		SHADER_PARAMETER(uint32, DescBufferOffset)
+	END_SHADER_PARAMETER_STRUCT()
+
+	static constexpr uint32 ThreadGroupSize = 64;
+	static inline void ModifyCompilationEnvironment(const FGlobalShaderPermutationParameters& Parameters, FShaderCompilerEnvironment& OutEnvironment)
+	{
+		FGlobalShader::ModifyCompilationEnvironment(Parameters, OutEnvironment);
+
+		OutEnvironment.SetDefine(TEXT("THREADGROUP_SIZE"), ThreadGroupSize);
+	}
+
+	static bool ShouldCompilePermutation(const FGlobalShaderPermutationParameters& Parameters)
+	{
+		return (ShouldCompileRayTracingShadersForProject(Parameters.Platform) && RHISupportsComputeShaders(Parameters.Platform) && !(Parameters.Platform == EShaderPlatform::SP_METAL || Parameters.Platform == EShaderPlatform::SP_METAL_TVOS || IsMobilePlatform(Parameters.Platform)));
+	}
+};
+
+IMPLEMENT_GLOBAL_SHADER(FRayTracingInstanceCopyCS, "/Engine/Private/Raytracing/RayTracingInstanceBufferUtil.usf", "RayTracingInstanceCopyShaderCS", SF_Compute);
+
+struct FRayTracingInstanceBufferCS : public FGlobalShader
+{
+	DECLARE_GLOBAL_SHADER(FRayTracingInstanceBufferCS);
+	SHADER_USE_PARAMETER_STRUCT(FRayTracingInstanceBufferCS, FGlobalShader);
+
+	BEGIN_SHADER_PARAMETER_STRUCT(FParameters, )
+		SHADER_PARAMETER_UAV(RWStructuredBuffer, InstancesDescriptors)
+		SHADER_PARAMETER_SRV(StructuredBuffer<FRayTracingInstanceDescriptorInput>, InputInstanceDescriptors)
+		SHADER_PARAMETER_SRV(ByteAddressBuffer, AccelerationStructureAddresses)
+		SHADER_PARAMETER(uint32, NumInstances)
+	END_SHADER_PARAMETER_STRUCT()
+		
+	static constexpr uint32 ThreadGroupSize = 64;
+
+	static inline void ModifyCompilationEnvironment(const FGlobalShaderPermutationParameters& Parameters, FShaderCompilerEnvironment& OutEnvironment)
+	{
+		FGlobalShader::ModifyCompilationEnvironment(Parameters, OutEnvironment);
+
+		OutEnvironment.SetDefine(TEXT("THREADGROUP_SIZE"), ThreadGroupSize);
+	}
+
+	static bool ShouldCompilePermutation(const FGlobalShaderPermutationParameters& Parameters)
+	{
+		return (ShouldCompileRayTracingShadersForProject(Parameters.Platform) && RHISupportsComputeShaders(Parameters.Platform) && !(Parameters.Platform == EShaderPlatform::SP_METAL || Parameters.Platform == EShaderPlatform::SP_METAL_TVOS || IsMobilePlatform(Parameters.Platform)));
+	}
+};
+
+IMPLEMENT_GLOBAL_SHADER(FRayTracingInstanceBufferCS, "/Engine/Private/Raytracing/RayTracingInstanceBufferUtil.usf", "RayTracingBuildInstanceBufferCS", SF_Compute);
+
+void BuildRayTracingInstanceBuffer(FRHICommandList& RHICmdList, uint32 NumInstances, FUnorderedAccessViewRHIRef InstancesUAV, FShaderResourceViewRHIRef InstanceUploadSRV, FShaderResourceViewRHIRef AccelerationStructureAddressesSRV)
+{
+	FRayTracingInstanceBufferCS::FParameters PassParams;
+	PassParams.InstancesDescriptors = InstancesUAV;
+	PassParams.InputInstanceDescriptors = InstanceUploadSRV;
+	PassParams.AccelerationStructureAddresses = AccelerationStructureAddressesSRV;
+	PassParams.NumInstances = NumInstances;
+
+	auto ComputeShader = GetGlobalShaderMap(GMaxRHIFeatureLevel)->GetShader<FRayTracingInstanceBufferCS>();
+	const int32 GroupSize = FMath::DivideAndRoundUp(PassParams.NumInstances, FRayTracingInstanceBufferCS::ThreadGroupSize);
+
+	RHICmdList.SetComputeShader(ComputeShader.GetComputeShader());
+
+	SetShaderParameters(RHICmdList, ComputeShader, ComputeShader.GetComputeShader(), PassParams);
+
+	DispatchComputeShader(RHICmdList, ComputeShader.GetShader(), GroupSize, 1, 1);
+
+	UnsetShaderUAVs(RHICmdList, ComputeShader, ComputeShader.GetComputeShader());
+}
+
+#endif //RHI_RAYTRACING
diff --git a/Engine/Source/Runtime/Renderer/Private/RayTracing/RayTracingScene.cpp b/Engine/Source/Runtime/Renderer/Private/RayTracing/RayTracingScene.cpp
--- a/Engine/Source/Runtime/Renderer/Private/RayTracing/RayTracingScene.cpp
+++ b/Engine/Source/Runtime/Renderer/Private/RayTracing/RayTracingScene.cpp
@@ -4,7 +4,7 @@
 
 #if RHI_RAYTRACING
 
-#include "RayTracingInstanceCopyShader.h"
+#include "RayTracingInstanceBufferUtil.h"
 #include "RenderCore.h"
 #include "RayTracingDefinitions.h"
 #include "RenderGraphBuilder.h"
diff --git a/Engine/Source/Runtime/Renderer/Private/Tests/RayTracingTestbed.cpp b/Engine/Source/Runtime/Renderer/Private/Tests/RayTracingTestbed.cpp
--- a/Engine/Source/Runtime/Renderer/Private/Tests/RayTracingTestbed.cpp
+++ b/Engine/Source/Runtime/Renderer/Private/Tests/RayTracingTestbed.cpp
@@ -10,7 +10,7 @@ IMPLEMENT_SIMPLE_AUTOMATION_TEST(FRayTracingTestbed, "System.Renderer.RayTracing
 
 #if RHI_RAYTRACING
 
-#include "RayTracingInstanceCopyShader.h"
+#include "RayTracingInstanceBufferUtil.h"
 #include "RayTracingDefinitions.h"
 #include "GlobalShader.h"
 #include "Containers/DynamicRHIResourceArray.h"
diff --git a/Engine/Source/Runtime/Renderer/Public/RayTracingInstanceBufferUtil.h b/Engine/Source/Runtime/Renderer/Public/RayTracingInstanceBufferUtil.h
new file mode 100644
--- /dev/null
+++ b/Engine/Source/Runtime/Renderer/Public/RayTracingInstanceBufferUtil.h
@@ -0,0 +1,33 @@
+// Copyright Epic Games, Inc. All Rights Reserved.
+
+#pragma once
+
+#include "RenderGraphDefinitions.h"
+#include "CoreMinimal.h"
+
+#if RHI_RAYTRACING
+
+struct FRayTracingInstanceDescriptorInput
+{
+	//uint32 GPUSceneInstanceIndex;
+	FVector4f LocalToWorld[3];
+	uint32 AccelerationStructureIndex;
+	uint32 InstanceId;
+	uint32 InstanceMaskAndFlags;
+	uint32 InstanceContributionToHitGroupIndex;
+};
+
+RENDERER_API void FillInstanceUploadBuffer(
+	TConstArrayView<FRayTracingGeometryInstance> Instances,
+	TConstArrayView<uint32> InstancesGeometryIndex,
+	FRayTracingSceneRHIRef RayTracingSceneRHI,
+	TArrayView<FRayTracingInstanceDescriptorInput> OutInstanceUploadData);
+
+RENDERER_API void BuildRayTracingInstanceBuffer(
+	FRHICommandList& RHICmdList, 
+	uint32 NumInstances,
+	FUnorderedAccessViewRHIRef InstancesUAV,
+	FShaderResourceViewRHIRef InstanceUploadSRV,
+	FShaderResourceViewRHIRef AccelerationStructureAddressesSRV);
+
+#endif // RHI_RAYTRACING
\ No newline at end of file
