SceneDepthAux removed from certain platforms

Resolved Mobile Texture Lookup Scene Depth outputting DeviceZ instead of WorldZ
#[robomerge] main

diff --git a/Engine/Source/Runtime/Renderer/Private/MobileShadingRenderer.cpp b/Engine/Source/Runtime/Renderer/Private/MobileShadingRenderer.cpp
--- a/Engine/Source/Runtime/Renderer/Private/MobileShadingRenderer.cpp
+++ b/Engine/Source/Runtime/Renderer/Private/MobileShadingRenderer.cpp
@@ -456,7 +456,7 @@ void FMobileSceneRenderer::InitViews(FRDGBuilder& GraphBuilder, FSceneTexturesCo
 		bRequireSeparateViewPass ||
 		bIsFullDepthPrepassEnabled ||
         bIsSimulatedLDR;
-	// never keep MSAA depth
+	// never keep MSAA depth if SceneDepthAux is enabled
 	bKeepDepthContent = ((NumMSAASamples > 1) && MobileRequiresSceneDepthAux(ShaderPlatform)) ? false : bKeepDepthContent;
 
 	// Update the bKeepDepthContent based on the mobile renderer status.
