Improve cinematic human skin rendering performance when filling most of the screens, and fix half resolutioin boundary flickering.

diff --git a/Engine/Shaders/Private/PostProcessSubsurface.usf b/Engine/Shaders/Private/PostProcessSubsurface.usf
--- a/Engine/Shaders/Private/PostProcessSubsurface.usf
+++ b/Engine/Shaders/Private/PostProcessSubsurface.usf
@@ -773,7 +773,9 @@ void SubsurfaceRecombinePS(noperspective float4 UVAndScreenPos : TEXCOORD0, floa
 #endif // SUBSURFACE_HALF_RES
 
 	{
-		float4 SSSColorWithAlpha = Texture2DSample(SubsurfaceInput1_Texture, SharedSubsurfaceSampler1, BufferUV);
+		// subsurface scattering buffer has garbage data outside of viewport bilinear min/max. 
+		float2 ClampedBufferUV = clamp(BufferUV, SubsurfaceInput1_UVViewportBilinearMin, SubsurfaceInput1_UVViewportBilinearMax);
+		float4 SSSColorWithAlpha = Texture2DSampleLevel(SubsurfaceInput1_Texture, SharedSubsurfaceSampler1, ClampedBufferUV, 0);
 
 		// renormalize to dilate RGB to fix half res upsampling artifacts
 		SSSColor = SSSColorWithAlpha.rgb / max(SSSColorWithAlpha.a, 0.00001f);
diff --git a/Engine/Source/Runtime/Renderer/Private/PostProcess/PostProcessSubsurface.cpp b/Engine/Source/Runtime/Renderer/Private/PostProcess/PostProcessSubsurface.cpp
--- a/Engine/Source/Runtime/Renderer/Private/PostProcess/PostProcessSubsurface.cpp
+++ b/Engine/Source/Runtime/Renderer/Private/PostProcess/PostProcessSubsurface.cpp
@@ -110,7 +110,7 @@ namespace
 
 	TAutoConsoleVariable<int32> CVarSSSBurleyEnableProfileIdCache(
 		TEXT("r.SSS.Burley.EnableProfileIdCache"),
-		0,
+		1,
 		TEXT("0: Disable profile id cache using in the sampling pass.\n")
 		TEXT("1: Consumes 1 byte per pixel more memory to make Burley pass much faster. (default)\n"),
 		ECVF_RenderThreadSafe);
@@ -279,8 +279,8 @@ bool IsSubsurfaceRequiredForView(const FViewInfo& View)
 
 bool IsProfileIdCacheEnabled()
 {
-	// Had to disable this at the last minute, because it uses an R8 UAV which isn't supported on all platforms. Will enable it in a later revision.
-	return 0 && CVarSSSBurleyEnableProfileIdCache.GetValueOnRenderThread() != 0;
+	const bool bIsR8UintSupported = GDynamicRHI && GDynamicRHI->RHIIsTypedUAVLoadSupported(PF_R8_UINT);
+	return bIsR8UintSupported && CVarSSSBurleyEnableProfileIdCache.GetValueOnRenderThread() != 0;
 }
 
 uint32 GetSubsurfaceRequiredViewMask(TArrayView<const FViewInfo> Views)
