Fix GPU out of bound memory write.

diff --git a/Engine/Source/Runtime/Renderer/Private/HairStrands/HairStrandsClusters.cpp b/Engine/Source/Runtime/Renderer/Private/HairStrands/HairStrandsClusters.cpp
--- a/Engine/Source/Runtime/Renderer/Private/HairStrands/HairStrandsClusters.cpp
+++ b/Engine/Source/Runtime/Renderer/Private/HairStrands/HairStrandsClusters.cpp
@@ -260,9 +260,11 @@ static void AddClusterCullingPass(
 	FRDGBufferRef GlobalIndexCountBuffer = GraphBuilder.CreateBuffer(FRDGBufferDesc::CreateBufferDesc(sizeof(uint32), ClusterData.ClusterCount), TEXT("Hair.GlobalIndexCountBuffer"));
 	FRDGBufferRef GlobalRadiusScaleBuffer = GraphBuilder.CreateBuffer(FRDGBufferDesc::CreateBufferDesc(sizeof(float), ClusterData.ClusterCount), TEXT("Hair.GlobalRadiusScaleBuffer"));
 
-	FRDGBufferRef PerBlocklTotalIndexCountBuffer = GraphBuilder.CreateBuffer(FRDGBufferDesc::CreateBufferDesc(sizeof(uint32), ClusterData.ClusterCount), TEXT("Hair.PerBlocklTotalIndexCountBuffer"));
-	FRDGBufferRef PerBlocklTotalIndexCountPreFixSumBuffer = GraphBuilder.CreateBuffer(FRDGBufferDesc::CreateBufferDesc(sizeof(uint32) * 2, ClusterData.ClusterCount), TEXT("Hair.PerBlocklTotalIndexCountPreFixSumBuffer"));
-	FRDGBufferRef PerBlocklIndexCountPreFixSumBuffer = GraphBuilder.CreateBuffer(FRDGBufferDesc::CreateBufferDesc(sizeof(uint32) * 2, ClusterData.ClusterCount), TEXT("Hair.PerBlocklIndexCountPreFixSumBuffer"));
+	const uint32 PrefixGroupSize = 512;
+	const uint32 ClusterCountRoundUpToGroupSize = FMath::DivideAndRoundUp(ClusterData.ClusterCount, PrefixGroupSize) * PrefixGroupSize;
+	FRDGBufferRef PerBlocklTotalIndexCountBuffer = GraphBuilder.CreateBuffer(FRDGBufferDesc::CreateBufferDesc(sizeof(uint32), ClusterCountRoundUpToGroupSize), TEXT("Hair.PerBlocklTotalIndexCountBuffer"));
+	FRDGBufferRef PerBlocklTotalIndexCountPreFixSumBuffer = GraphBuilder.CreateBuffer(FRDGBufferDesc::CreateBufferDesc(sizeof(uint32) * 2, ClusterCountRoundUpToGroupSize), TEXT("Hair.PerBlocklTotalIndexCountPreFixSumBuffer"));
+	FRDGBufferRef PerBlocklIndexCountPreFixSumBuffer = GraphBuilder.CreateBuffer(FRDGBufferDesc::CreateBufferDesc(sizeof(uint32) * 2, ClusterCountRoundUpToGroupSize), TEXT("Hair.PerBlocklIndexCountPreFixSumBuffer"));
 
 	FRDGImportedBuffer DrawIndirectParametersBuffer = Register(GraphBuilder, ClusterData.HairGroupPublicPtr->GetDrawIndirectBuffer(), ERDGImportedBufferFlags::CreateViews);
 	FRDGImportedBuffer DrawIndirectParametersRasterComputeBuffer = Register(GraphBuilder, ClusterData.HairGroupPublicPtr->GetDrawIndirectRasterComputeBuffer(), ERDGImportedBufferFlags::CreateViews);
