[Major] Reduce SSS performance overhead when SSS only occupies a small region (e.g., from 0.7ms to 0.2ms).
            1. Reduce unnecessary UAV Clear with property velocity projection. Only subpass one UAV is cleared conditionally when separable tiles are detected.
            2. Remove subsurface view copy and use tile based recombine.
            3. Mipmaps are generated only when there are many Burley tiles ( default r.SSS.Burley.MinGenerateMipsTileCount to 4000)
[Minor] Rename pass and texture names, and use proper descriptors.

diff --git a/Engine/Shaders/Private/ComputeGenerateMips.usf b/Engine/Shaders/Private/ComputeGenerateMips.usf
--- a/Engine/Shaders/Private/ComputeGenerateMips.usf
+++ b/Engine/Shaders/Private/ComputeGenerateMips.usf
@@ -39,6 +39,44 @@ void MainCS(uint3 DT_ID : SV_DispatchThreadID)
 #endif
 }
 
+uint2 TextureSize;
+uint Offset;
+uint NumMips;
+Buffer<uint> ConditionBuffer;
+RWBuffer<uint> RWIndirectDispatchArgsBuffer;
+
+[numthreads(8,1,1)]
+void BuildIndirectDispatchArgsCS(uint DT_ID : SV_DispatchThreadID)
+{
+	if (DT_ID < NumMips)
+	{
+		uint Condition = ConditionBuffer[Offset];
+		
+		// Generate mipmaps only when Condition > 0
+		const bool bShouldDispatchGenerateMips = (Condition > 0);
+
+		if (bShouldDispatchGenerateMips)
+		{
+
+			uint2 DestTextureSize = uint2(
+				max(TextureSize.x >> (DT_ID + 1), 1),
+				max(TextureSize.y >> (DT_ID + 1), 1));
+
+			DestTextureSize = (DestTextureSize + 8 - 1) / 8;
+
+			RWIndirectDispatchArgsBuffer[DT_ID * 3 + 0] = DestTextureSize.x;
+			RWIndirectDispatchArgsBuffer[DT_ID * 3 + 1] = DestTextureSize.y;
+			RWIndirectDispatchArgsBuffer[DT_ID * 3 + 2] = 1;
+		}
+		else
+		{
+			RWIndirectDispatchArgsBuffer[DT_ID * 3 + 0] = 0;
+			RWIndirectDispatchArgsBuffer[DT_ID * 3 + 1] = 0;
+			RWIndirectDispatchArgsBuffer[DT_ID * 3 + 2] = 0;
+		}
+	}
+}
+
 #else // !GENMIPS_COMPUTE
 
 float2 HalfTexelSize;
diff --git a/Engine/Shaders/Private/PostProcessSubsurface.usf b/Engine/Shaders/Private/PostProcessSubsurface.usf
--- a/Engine/Shaders/Private/PostProcessSubsurface.usf
+++ b/Engine/Shaders/Private/PostProcessSubsurface.usf
@@ -23,7 +23,7 @@
 
 #include "ScreenPass.ush"
 
-// x:Radius*DistanceToProjectionWindow/KernelSize*0.5, y:DistanceToProjectionWindow, z:OverrideNumSamples, w: unused
+// x:Radius*DistanceToProjectionWindow/KernelSize*0.5, y:DistanceToProjectionWindow, z:OverrideNumSamples, w: MinGenerateMipsTileCount
 float4 SubsurfaceParams;
 
 Texture2D SubsurfaceInput0_Texture;
@@ -55,6 +55,12 @@ SCREEN_PASS_TEXTURE_VIEWPORT(Output)
 
 RWTexture2D<uint> ProfileIdTexture;
 
+// the same to FSubsurfaceTiles::ETileType
+#define SUBSURFACE_TILE_TYPE_AFIS		0
+#define SUBSURFACE_TILE_TYPE_SEPARABLE	1
+#define SUBSURFACE_TILE_TYPE_ALL		2
+#define SUBSURFACE_TILE_TYPE_COUNT		3
+
 #define SUBSURFACE_PASS_ONE 0
 #define SUBSURFACE_PASS_TWO 1
 
@@ -418,6 +424,21 @@ float4 SetupSubsurfaceForOnePixel(float2 UVSceneColor, float2 Extent)
 	return Ret;
 }
 
+Buffer<uint> TileTypeCountBuffer;
+
+bool ShouldGenerateMipmaps(uint TileType)
+{
+	bool bShouldGenerateMipmaps = false;
+
+	if (SUBSURFACE_TILE_TYPE_AFIS == TileType)
+	{
+		uint TileCount = TileTypeCountBuffer[TileType];
+		bShouldGenerateMipmaps = (TileCount >= SubsurfaceParams.w);
+	}
+
+	return bShouldGenerateMipmaps;
+}
+
 // Ref: https://graphics.pixar.com/library/ApproxBSSRDF/paper.pdf
 #include "SubsurfaceBurleyNormalized.ush"
 
@@ -428,26 +449,52 @@ float4 SetupSubsurfaceForOnePixel(float2 UVSceneColor, float2 Extent)
 // Indirect dispatch parameters
 RWBuffer<uint> RWSeparableGroupBuffer;
 RWBuffer<uint> RWBurleyGroupBuffer;
+RWBuffer<uint> RWAllGroupBuffer;
 RWBuffer<uint> RWIndirectDispatchArgsBuffer;
+RWBuffer<uint> RWIndirectDrawArgsBuffer;
+RWBuffer<uint> RWTileTypeCountBuffer;
 Buffer<uint> GroupBuffer;
 
-
 [numthreads(1, 1, 1)]
 void InitValueBufferCS()
 {
-	// Clear the first buffer element which is used as the buffer counter.
-	RWSeparableGroupBuffer[0] = 0;
-	RWBurleyGroupBuffer[0] = 0;
+	// Clear the buffer counter for all types of tile.
+	UNROLL
+	for (uint i = 0; i < SUBSURFACE_TILE_TYPE_COUNT; ++i)
+	{
+		RWTileTypeCountBuffer[i] = 0;
+	}
 }
 
-// fixed to 64 by now
+uint MinGenerateMipsTileCount;
+uint Offset;
+RWBuffer<uint> RWMipsConditionBuffer;
+
 [numthreads(1, 1, 1)]
-void BuildIndirectDispatchArgsCS()
+void FillMipsConditionBufferCS()
 {
-	uint ValueCount = min(SubsurfaceUniformParameters.MaxGroupCount, GroupBuffer[0]);
-	RWIndirectDispatchArgsBuffer[0] = ValueCount;// (ValueCount + (THREAD_SIZE_1D*THREAD_SIZE_1D) - 1) / (THREAD_SIZE_1D*THREAD_SIZE_1D);
-	RWIndirectDispatchArgsBuffer[1] = 1;
-	RWIndirectDispatchArgsBuffer[2] = 1;
+	uint TileCount = TileTypeCountBuffer[Offset];
+	RWMipsConditionBuffer[Offset] = (TileCount >= MinGenerateMipsTileCount) ? 1 : 0;
+}
+
+[numthreads(8, 1, 1)]
+void BuildIndirectDispatchArgsCS(uint2 DispatchThreadId : SV_DispatchThreadID)
+{
+	if (DispatchThreadId.y == 0 && DispatchThreadId.x < SUBSURFACE_TILE_TYPE_COUNT)
+	{
+		uint TileType = DispatchThreadId.x;
+		const uint TileCount = TileTypeCountBuffer[TileType];
+		
+		// Indirect draw
+		RWIndirectDrawArgsBuffer[TileType * 4 + 0] = SubsurfaceUniformParameters.bRectPrimitive > 0 ? 4 : 6;
+		RWIndirectDrawArgsBuffer[TileType * 4 + 1] = TileCount;
+		RWIndirectDrawArgsBuffer[TileType * 4 + 2] = 1;
+
+		// Indirect dispatch
+		RWIndirectDispatchArgsBuffer[TileType * 3 + 0] = TileCount;// (TileCount + (THREAD_SIZE_1D*THREAD_SIZE_1D) - 1) / (THREAD_SIZE_1D*THREAD_SIZE_1D);
+		RWIndirectDispatchArgsBuffer[TileType * 3 + 1] = 1;
+		RWIndirectDispatchArgsBuffer[TileType * 3 + 2] = 1;
+	}
 }
 
 #define TYPE_SEPARABLE 0x1
@@ -469,15 +516,15 @@ void AddSeparableGroup(uint2 GroupXY)
 	//buffer counter is stored in the .Count in the first element.
 	uint IndexToStore = 0;
 
-	InterlockedAdd(RWSeparableGroupBuffer[0], 1, IndexToStore);
+	InterlockedAdd(RWTileTypeCountBuffer[SUBSURFACE_TILE_TYPE_SEPARABLE], 1, IndexToStore);
 
 	if (IndexToStore >= SubsurfaceUniformParameters.MaxGroupCount)
 	{
 		return;
 	}
 
-	RWSeparableGroupBuffer[2 *(IndexToStore + 1) + 0] = GroupXY.x;
-	RWSeparableGroupBuffer[2 *(IndexToStore + 1) + 1] = GroupXY.y;
+	RWSeparableGroupBuffer[2 * IndexToStore + 0] = GroupXY.x;
+	RWSeparableGroupBuffer[2 * IndexToStore + 1] = GroupXY.y;
 }
 
 void AddBurleyGroup(uint2 GroupXY)
@@ -485,15 +532,31 @@ void AddBurleyGroup(uint2 GroupXY)
 	//buffer counter is stored in the .Count in the first element.
 	uint IndexToStore = 0;
 
-	InterlockedAdd(RWBurleyGroupBuffer[0], 1, IndexToStore);
+	InterlockedAdd(RWTileTypeCountBuffer[SUBSURFACE_TILE_TYPE_AFIS], 1, IndexToStore);
 
 	if (IndexToStore >= SubsurfaceUniformParameters.MaxGroupCount)
 	{
 		return;
 	}
 
-	RWBurleyGroupBuffer[2*(IndexToStore + 1 ) + 0] = GroupXY.x;
-	RWBurleyGroupBuffer[2*(IndexToStore + 1 ) + 1] = GroupXY.y;
+	RWBurleyGroupBuffer[2* IndexToStore + 0] = GroupXY.x;
+	RWBurleyGroupBuffer[2* IndexToStore + 1] = GroupXY.y;
+}
+
+void AddAllGroup(uint2 GroupXY)
+{
+	//buffer counter is stored in the .Count in the first element.
+	uint IndexToStore = 0;
+
+	InterlockedAdd(RWTileTypeCountBuffer[SUBSURFACE_TILE_TYPE_ALL], 1, IndexToStore);
+
+	if (IndexToStore >= SubsurfaceUniformParameters.MaxGroupCount)
+	{
+		return;
+	}
+
+	RWAllGroupBuffer[2 * IndexToStore + 0] = GroupXY.x;
+	RWAllGroupBuffer[2 * IndexToStore + 1] = GroupXY.y;
 }
 
 void AddSubsurfaceComputeGroup(uint GI, uint Type, uint2 GroupXY)
@@ -511,6 +574,11 @@ void AddSubsurfaceComputeGroup(uint GI, uint Type, uint2 GroupXY)
 	{
 		AddBurleyGroup(GroupXY);
 	}
+
+	if (GI == 0 && SubsurfaceTypeFlag > 0)
+	{
+		AddAllGroup(GroupXY);
+	}
 }
 
 RWTexture2D<float4> SetupTexture;
@@ -651,7 +719,7 @@ void MainIndirectDispatchCS(
 )
 {
 	// The first element is used for counting
-	uint2 G_ID = uint2(GroupBuffer[2*(IndirectG_ID.x + 1) + 0], GroupBuffer[2*(IndirectG_ID.x + 1)+1]);
+	uint2 G_ID = uint2(GroupBuffer[2 * IndirectG_ID.x + 0], GroupBuffer[2 * IndirectG_ID.x + 1]);
 	uint2 DT_ID = G_ID * float2(THREAD_SIZE_X, THREAD_SIZE_Y) + float2(GI % THREAD_SIZE_X, GI / THREAD_SIZE_X);
 
 
@@ -711,9 +779,9 @@ void MainIndirectDispatchCS(
  
 // Recombines the half res Subsurface filtered lighting contribution (upsampled and renormalized with the alpha)
 // with the SceneColor.
-void SubsurfaceRecombinePS(noperspective float4 UVAndScreenPos : TEXCOORD0, float4 SvPosition : SV_POSITION, out float4 OutColor : SV_Target0)
+void SubsurfaceRecombinePS(noperspective float2 UV : TEXCOORD0, float4 SvPosition : SV_POSITION, out float4 OutColor : SV_Target0)
 {
-	const float2 BufferUV  = UVAndScreenPos.xy;
+	const float2 BufferUV  = UV.xy;
 	const uint2  BufferPos = SvPosition.xy;
 
 #if STRATA_ENABLED
@@ -814,7 +882,30 @@ void SubsurfaceRecombinePS(noperspective float4 UVAndScreenPos : TEXCOORD0, floa
 	// combine potentially half res with full res
 	float3 SubsurfaceLighting = lerp(DiffuseAndSpecular.Diffuse, SSSColor, FadedSubsurfaceColor);
 
-	OutColor = float4(SubsurfaceLighting * StoredBaseColor + ExtractedNonSubsurface, 0);
+	OutColor = float4(SubsurfaceLighting * StoredBaseColor + ExtractedNonSubsurface, 1.0f);
+}
+
+void SubsurfaceRecombineCopyPS(noperspective float2 UV : TEXCOORD0, float4 SvPosition : SV_POSITION, out float4 OutColor : SV_Target0)
+{
+	const float2 BufferUV = UV.xy;
+	const uint2  BufferPos = SvPosition.xy;
+
+#if STRATA_ENABLED
+	const FStrataSubsurfaceData SSSData = StrataUnpackSSSData(Strata.SSSTexture.Load(uint3(BufferPos, 0)));
+	const bool bUseSSSProfile = SSSData.bIsValid;
+#else
+	FScreenSpaceData ScreenSpaceData = GetScreenSpaceData(BufferUV);
+	const bool bUseSSSProfile = UseSubsurfaceProfile(ScreenSpaceData.GBuffer.ShadingModelID);
+#endif
+
+	if (bUseSSSProfile)
+	{
+		OutColor = float4(Texture2DSampleLevel(SubsurfaceInput0_Texture, SharedSubsurfaceSampler0, BufferUV,0).rgb,1.0f);
+	}
+	else
+	{
+		OutColor = float4(0.0f, 0.0f, 0.0f, 0.0f);
+	}
 }
 
 void SubsurfaceViewportCopyPS(noperspective float4 InUV : TEXCOORD0, out float4 OutColor : SV_Target0)
diff --git a/Engine/Shaders/Private/PostProcessSubsurfaceTile.usf b/Engine/Shaders/Private/PostProcessSubsurfaceTile.usf
new file mode 100644
--- /dev/null
+++ b/Engine/Shaders/Private/PostProcessSubsurfaceTile.usf
@@ -0,0 +1,100 @@
+// Copyright Epic Games, Inc. All Rights Reserved.
+
+/*=============================================================================
+	PostProcessSubsurfaceTile.usf: Screenspace subsurface scattering tile shaders.
+=============================================================================*/
+
+#pragma once
+
+#include "Common.ush"
+
+////////////////////////////////////////////////////////////////////////////////////////////////////////////
+#ifndef SHADER_TILE_VS
+#define SHADER_TILE_VS 0
+#endif
+
+#ifndef SUBSURFACE_COMPUTE
+#define SUBSURFACE_COMPUTE 0
+#endif
+
+#define LINE_TEXEL_SIZE 1
+
+#if SUBSURFACE_COMPUTE
+
+int2 ViewportSize;
+uint Offset;
+Buffer<uint> ConditionBuffer;
+RWBuffer<uint> RWIndirectDispatchArgsBuffer;
+
+[numthreads(1, 1, 1)]
+void BuildIndirectDispatchArgsCS(uint DT_ID : SV_DispatchThreadID)
+{
+	uint condition = ConditionBuffer[Offset];
+	const bool bShouldDispatch = (condition > 0);
+
+	if (bShouldDispatch)
+	{
+		int2 DestTextureSize = (ViewportSize + SUBSURFACE_TILE_SIZE * LINE_TEXEL_SIZE - 1)
+								/ (SUBSURFACE_TILE_SIZE * LINE_TEXEL_SIZE);
+
+		RWIndirectDispatchArgsBuffer[DT_ID * 3 + 0] = DestTextureSize.x;
+		RWIndirectDispatchArgsBuffer[DT_ID * 3 + 1] = DestTextureSize.y;
+		RWIndirectDispatchArgsBuffer[DT_ID * 3 + 2] = 1;
+	}
+	else
+	{
+		RWIndirectDispatchArgsBuffer[DT_ID * 3 + 0] = 0;
+		RWIndirectDispatchArgsBuffer[DT_ID * 3 + 1] = 0;
+		RWIndirectDispatchArgsBuffer[DT_ID * 3 + 2] = 0; 
+	}
+}
+
+uint2 TextureExtent;
+uint2 ViewportMin;
+RWTexture2D<float4> TextureOutput;
+
+[numthreads(SUBSURFACE_TILE_SIZE, SUBSURFACE_TILE_SIZE, 1)]
+void ClearUAV(uint2 DT_ID : SV_DispatchThreadID)
+{
+	uint2 BufferPos = ViewportMin + DT_ID * LINE_TEXEL_SIZE;
+
+	if (all(BufferPos < TextureExtent))
+	{
+		UNROLL
+		for (uint i = 0; i < LINE_TEXEL_SIZE; ++i)
+		{
+			UNROLL
+			for (uint j = 0; j < LINE_TEXEL_SIZE; ++j)
+			{
+				uint2 CurrentBufferPos = min(BufferPos + uint2(i, j), TextureExtent - uint2(1,1));
+				TextureOutput[CurrentBufferPos] = float4(0.0f, 0.0f, 0.0f, 0.0f);
+			}
+		}
+	}
+}
+
+#endif
+
+#if SHADER_TILE_VS
+
+int2 ViewMin;
+float2 ExtentInverse;
+uint TileType;
+Buffer<uint> TileDataBuffer;
+
+void MainVS(
+	in uint InVertexId : SV_VertexID,
+	in uint InInstanceId : SV_InstanceID,
+	out FScreenVertexOutput Out)
+{
+	Out = (FScreenVertexOutput)0;
+	const uint2 TileCoord = uint2(TileDataBuffer[InInstanceId * 2 + 0], TileDataBuffer[InInstanceId * 2 + 1]);
+
+	uint2 TileVertex = TileCoord * SUBSURFACE_TILE_SIZE;
+	TileVertex.x += InVertexId == 1 || InVertexId == 2 || InVertexId == 4 ? SUBSURFACE_TILE_SIZE : 0;
+	TileVertex.y += InVertexId == 2 || InVertexId == 4 || InVertexId == 5 ? SUBSURFACE_TILE_SIZE : 0;
+	Out.UV = (float2(TileVertex + ViewMin) + 0.5f) * ExtentInverse;
+	Out.Position = float4(Out.UV * float2(2.0f, -2.0f) + float2(-1.0, 1.0f), 0.0f, 1.0f);
+}
+
+#endif //SHADER_TILE_VS
\ No newline at end of file
diff --git a/Engine/Shaders/Private/SubsurfaceBurleyNormalized.ush b/Engine/Shaders/Private/SubsurfaceBurleyNormalized.ush
--- a/Engine/Shaders/Private/SubsurfaceBurleyNormalized.ush
+++ b/Engine/Shaders/Private/SubsurfaceBurleyNormalized.ush
@@ -432,10 +432,87 @@
 	return GetPdf(Radius *0.01f, L*0.01f, S);
 }
 
+// Get the history states
+// SubsurfaceInput2_Texture stores the encoded velocity
+// SubsurfaceInput1_Texture store the history state
+struct FHistoryState
+{
+	float4 History;
+	bool OffScreen;
+};
+
+FHistoryState GetHistoryState(float2 BufferUV)
+{
+	FHistoryState HistoryState = (FHistoryState)0;
+	float2 VelocityOffset = float2(0.0, 0.0);
+	float2 NearestBufferUV = BufferUV;
+	float2 ViewportUV = (BufferUV - SubsurfaceInput0_UVViewportMin.xy) * SubsurfaceInput0_UVViewportSizeInverse.xy;
+	bool OffScreen = false;
+
+	//	Get the history with reprojection
+	//  Implement a simplified version of temporal AA
+#if REPROJECTION
+
+	// Code adapted from temporal AA
+	float3 PosN;
+	PosN.xy = ViewportUVToScreenPos(ViewportUV);
+	PosN.z = SampleDepthTexture(BufferUV);// Direct sample without converting to world space;
+	float4 ThisClip = float4(PosN.xy, PosN.z, 1.0f);
+	float4 PrevClip = mul(ThisClip, View.ClipToPrevClip);
+	float2 PrevScreen = PrevClip.xy / PrevClip.w;
+	float2 BackN = PosN.xy - PrevScreen;
+
+	// Sample the velocity texture
+	float Velocity = 0;
+
+	float2 BackTemp = BackN * SubsurfaceInput1_ViewportSize.xy;
+
+#if (ENABLE_VELOCITY)
+	{
+		float4 EncodedVelocity = Texture2DSampleLevel(SubsurfaceInput2_Texture, SubsurfaceSampler2, NearestBufferUV + VelocityOffset, 0);
+
+		if (EncodedVelocity.x > 0.0)
+		{
+			BackN = DecodeVelocityFromTexture(EncodedVelocity).xy;
+		}
+
+		BackTemp = BackN * SubsurfaceInput1_ViewportSize.xy;
+	}
+#endif
+
+	// Update velocity
+	Velocity = sqrt(dot(BackTemp, BackTemp));
+
+	float2 HistoryScreenPosition = (PosN.xy - BackN);
+
+	// Detect if HistoryBufferUV would be outside of the viewport.
+	OffScreen = max(abs(HistoryScreenPosition.x), abs(HistoryScreenPosition.y)) >= 1.0;
+
+	float4 History = 0;
+
+	BRANCH if (!OffScreen)
+	{
+		// ScreenPos to bufferUV
+		float2 HistoryUV = ScreenPosToViewportUV(HistoryScreenPosition);
+		// Convert history uv in viewport to buffer uv
+		float2 HistoryBufferUV = HistoryUV * SubsurfaceInput0_UVViewportSize.xy + SubsurfaceInput0_UVViewportMin.xy;
+		History = Texture2DSample(SubsurfaceInput1_Texture, SubsurfaceSampler1, HistoryBufferUV);
+	}
+
+#else
+	float4 History = Texture2DSample(SubsurfaceInput1_Texture, SubsurfaceSampler1, BufferUV);
+#endif
+
+	HistoryState.History = History;
+	HistoryState.OffScreen = OffScreen;
+
+	return HistoryState;
+}
+
 int GetNumOfSamplesBasedOnQuality(float2 UV)
 {
-
-	float4 QualityMatrix = Texture2DSample(SubsurfaceInput1_Texture, SharedSubsurfaceSampler1, UV);
+	FHistoryState HistoryState = GetHistoryState(UV);
+	float4 QualityMatrix = HistoryState.History;
 	float VariableVar = QualityMatrix.b;
 
 	float NumOfSamples = 0;
@@ -704,6 +781,12 @@
 {
 	float2 ClampedUV = clamp(CenterUV+ UVOffset, SubsurfaceInput0_UVViewportBilinearMin, SubsurfaceInput0_UVViewportBilinearMax);
 	
+	// set mip level to 0 if the mipmap is not generated
+	if (!ShouldGenerateMipmaps(SUBSURFACE_TILE_TYPE_AFIS))
+	{
+		MipLevel = 0.0f;
+	}
+
 #if !(ENABLE_PROFILE_ID_CACHE && (SUBSURFACE_FORCE_SEPARABLE !=1))
 	#if STRATA_ENABLED
 	const uint2 CenterPos = ClampedUV * SubsurfaceInput0_Extent;
@@ -1093,72 +1176,16 @@
 
 float4 UpdateQualityVariance(float4 SubsurfaceColor, float2 BufferUV)
 {
-	float Weight = EXPONENTIAL_WEIGHT;
-	float WeightFinal = Weight;
-	float4 NewHistory = 0;
-	float2 VelocityOffset = float2(0.0, 0.0);
-	float2 NearestBufferUV = BufferUV;
-	float2 ViewportUV = (BufferUV - SubsurfaceInput0_UVViewportMin.xy)*SubsurfaceInput0_UVViewportSizeInverse.xy;
+	float WeightFinal = EXPONENTIAL_WEIGHT;
+	FHistoryState HistoryState = GetHistoryState(BufferUV);
 
-	//	Get the history with reprojection
-	//  Implement a simplified version of temporal AA
-#if REPROJECTION
-
-	// Code adapted from temporal AA
-	float3 PosN;
-	PosN.xy = ViewportUVToScreenPos(ViewportUV);
-	PosN.z = SampleDepthTexture(BufferUV);// Direct sample without converting to world space;
-	float4 ThisClip = float4(PosN.xy, PosN.z, 1.0f);
-	float4 PrevClip = mul(ThisClip, View.ClipToPrevClip);
-	float2 PrevScreen = PrevClip.xy / PrevClip.w;
-	float2 BackN = PosN.xy - PrevScreen;
-
-	// Sample the velocity texture
-	float Velocity = 0;
-
-	float2 BackTemp = BackN * SubsurfaceInput1_ViewportSize.xy;
-
-#if (ENABLE_VELOCITY)
+	float4 History = HistoryState.History;
+	if (HistoryState.OffScreen)
 	{
-		float4 EncodedVelocity = Texture2DSampleLevel(SubsurfaceInput2_Texture, SubsurfaceSampler2, NearestBufferUV + VelocityOffset, 0);
-		
-		if (EncodedVelocity.x > 0.0)
-		{
-			BackN = DecodeVelocityFromTexture(EncodedVelocity).xy;
-		}
-
-		BackTemp = BackN * SubsurfaceInput1_ViewportSize.xy;
-	}
-#endif
-
-	// Update velocity
-	Velocity = sqrt(dot(BackTemp, BackTemp));
-
-	float2 HistoryScreenPosition = (PosN.xy - BackN);
-
-	// Detect if HistoryBufferUV would be outside of the viewport.
-	bool OffScreen = max(abs(HistoryScreenPosition.x), abs(HistoryScreenPosition.y)) >= 1.0;
-
-	float4 History = 0;
-
-	WeightFinal = lerp(Weight, 0.2, saturate(Velocity / 40));
-
-	BRANCH if (!OffScreen)
-	{
-		// ScreenPos to bufferUV
-		float2 HistoryUV = ScreenPosToViewportUV(HistoryScreenPosition);
-		// Convert history uv in viewport to buffer uv
-		float2 HistoryBufferUV = HistoryUV * SubsurfaceInput0_UVViewportSize.xy + SubsurfaceInput0_UVViewportMin.xy;
-		History = Texture2DSample(SubsurfaceInput1_Texture, SubsurfaceSampler1, HistoryBufferUV);
-	}
-	else
-	{
-		WeightFinal = 1;
+		WeightFinal = 1.0f;
 	}
 
-#else
-	float4 History = Texture2DSample(SubsurfaceInput1_Texture, SubsurfaceSampler1, BufferUV);
-#endif
+	float4 NewHistory = (float4)0;
 
 	// Use the bilinear jitter-corrected lighting as the control variable (Constant CV coefficient = 1.0) to remove the variance over-estimation
 	// of spatial features, like peach fuzz, inside subsurface scattering region. It will not affect the output surface color.
diff --git a/Engine/Source/Runtime/RenderCore/Private/GenerateMips.cpp b/Engine/Source/Runtime/RenderCore/Private/GenerateMips.cpp
--- a/Engine/Source/Runtime/RenderCore/Private/GenerateMips.cpp
+++ b/Engine/Source/Runtime/RenderCore/Private/GenerateMips.cpp
@@ -85,6 +85,65 @@ public:
 
 IMPLEMENT_GLOBAL_SHADER(FGenerateMipsPS, "/Engine/Private/ComputeGenerateMips.usf", "MainPS", SF_Pixel);
 
+// Determine the indirect dispatch based on conditions
+class FBuildIndirectDispatchArgsBufferCS : public FGlobalShader
+{
+public:
+	DECLARE_GLOBAL_SHADER(FBuildIndirectDispatchArgsBufferCS)
+	SHADER_USE_PARAMETER_STRUCT(FBuildIndirectDispatchArgsBufferCS, FGlobalShader)
+
+		BEGIN_SHADER_PARAMETER_STRUCT(FParameters, )
+		SHADER_PARAMETER(FIntPoint, TextureSize)
+		SHADER_PARAMETER(uint32, Offset)
+		SHADER_PARAMETER(uint32, NumMips)
+		SHADER_PARAMETER_RDG_BUFFER_SRV(Buffer<uint32>, ConditionBuffer)
+		SHADER_PARAMETER_RDG_BUFFER_UAV(RWBuffer<uint32>, RWIndirectDispatchArgsBuffer)
+		END_SHADER_PARAMETER_STRUCT()
+
+		static bool ShouldCompilePermutation(const FGlobalShaderPermutationParameters& Parameters)
+	{
+		return RHISupportsComputeShaders(Parameters.Platform);
+	}
+
+	static void ModifyCompilationEnvironment(const FShaderPermutationParameters&, FShaderCompilerEnvironment& OutEnvironment)
+	{
+		OutEnvironment.SetDefine(TEXT("GENMIPS_COMPUTE"), 1);
+	}
+};
+
+IMPLEMENT_GLOBAL_SHADER(FBuildIndirectDispatchArgsBufferCS, "/Engine/Private/ComputeGenerateMips.usf", "BuildIndirectDispatchArgsCS", SF_Compute);
+
+class FGenerateMipsIndirectCS : public FGlobalShader
+{
+public:
+	DECLARE_GLOBAL_SHADER(FGenerateMipsIndirectCS)
+	SHADER_USE_PARAMETER_STRUCT(FGenerateMipsIndirectCS, FGlobalShader)
+
+		class FGenMipsSRGB : SHADER_PERMUTATION_BOOL("GENMIPS_SRGB");
+	class FGenMipsSwizzle : SHADER_PERMUTATION_BOOL("GENMIPS_SWIZZLE");
+	using FPermutationDomain = TShaderPermutationDomain<FGenMipsSRGB, FGenMipsSwizzle>;
+
+	BEGIN_SHADER_PARAMETER_STRUCT(FParameters, )
+		SHADER_PARAMETER(FVector2D, TexelSize)
+		SHADER_PARAMETER_RDG_TEXTURE_SRV(Texture2D, MipInSRV)
+		SHADER_PARAMETER_RDG_TEXTURE_UAV(RWTexture2D, MipOutUAV)
+		SHADER_PARAMETER_SAMPLER(SamplerState, MipSampler)
+		RDG_BUFFER_ACCESS(IndirectDispatchArgsBuffer, ERHIAccess::IndirectArgs)
+		END_SHADER_PARAMETER_STRUCT()
+
+		static bool ShouldCompilePermutation(const FGlobalShaderPermutationParameters& Parameters)
+	{
+		return RHISupportsComputeShaders(Parameters.Platform);
+	}
+
+	static void ModifyCompilationEnvironment(const FShaderPermutationParameters&, FShaderCompilerEnvironment& OutEnvironment)
+	{
+		OutEnvironment.SetDefine(TEXT("GENMIPS_COMPUTE"), 1);
+	}
+};
+
+IMPLEMENT_GLOBAL_SHADER(FGenerateMipsIndirectCS, "/Engine/Private/ComputeGenerateMips.usf", "MainCS", SF_Compute);
+
 void FGenerateMips::ExecuteRaster(FRDGBuilder& GraphBuilder, FRDGTextureRef Texture, FRHISamplerState* Sampler)
 {
 	check(Texture);
@@ -181,6 +240,77 @@ void FGenerateMips::ExecuteCompute(FRDGBuilder& GraphBuilder, FRDGTextureRef Tex
 	}
 }
 
+void FGenerateMips::ExecuteCompute(FRDGBuilder& GraphBuilder, FRDGTextureRef Texture, FRHISamplerState* Sampler,
+	FRDGBufferRef ConditionBuffer, uint32 Offset)
+{
+	check(Texture);
+	check(Sampler);
+
+	const FRDGTextureDesc& TextureDesc = Texture->Desc;
+
+	FRDGBufferRef IndirectDispatchArgsBuffer = GraphBuilder.CreateBuffer(
+		FRDGBufferDesc::CreateIndirectDesc<FRHIDispatchIndirectParameters>(FMath::Max(TextureDesc.NumMips - 1,1)),
+		TEXT("IndirectDispatchArgsBuffer"));
+	{
+		// build the indirect dispatch arguments buffer ( compute the group count on GPU conditionally)
+		FBuildIndirectDispatchArgsBufferCS::FParameters* PassParameters = 
+			GraphBuilder.AllocParameters<FBuildIndirectDispatchArgsBufferCS::FParameters>();
+		PassParameters->TextureSize = TextureDesc.Extent;
+		PassParameters->Offset = Offset;
+		PassParameters->NumMips = TextureDesc.NumMips;
+		PassParameters->ConditionBuffer = GraphBuilder.CreateSRV(ConditionBuffer, EPixelFormat::PF_R32_UINT);
+		PassParameters->RWIndirectDispatchArgsBuffer = GraphBuilder.CreateUAV(IndirectDispatchArgsBuffer, EPixelFormat::PF_R32_UINT);
+
+		TShaderMapRef<FBuildIndirectDispatchArgsBufferCS> ComputeShader(GetGlobalShaderMap(GMaxRHIFeatureLevel));
+
+		FComputeShaderUtils::AddPass(
+			GraphBuilder,
+			RDG_EVENT_NAME("GenerateMips BuildIndirectArgs(Mips=%d)", TextureDesc.NumMips),
+			ComputeShader,
+			PassParameters,
+			FIntVector(FMath::DivideAndRoundUp(TextureDesc.NumMips - 1,FComputeShaderUtils::kGolden2DGroupSize), 1, 1));
+	}
+
+	// Select compute shader variant (normal vs. sRGB etc.)
+	bool bMipsSRGB = EnumHasAnyFlags(TextureDesc.Flags, TexCreate_SRGB);
+#if PLATFORM_ANDROID
+	if (IsVulkanPlatform(GMaxRHIShaderPlatform))
+	{
+		// Vulkan Android seems to skip sRGB->Lin conversion when sampling texture in compute
+		bMipsSRGB = false;
+	}
+#endif
+	const bool bMipsSwizzle = false;
+
+	FGenerateMipsIndirectCS::FPermutationDomain PermutationVector;
+	PermutationVector.Set<FGenerateMipsIndirectCS::FGenMipsSRGB>(bMipsSRGB);
+	PermutationVector.Set<FGenerateMipsIndirectCS::FGenMipsSwizzle>(bMipsSwizzle);
+	TShaderMapRef<FGenerateMipsIndirectCS> ComputeShader(GetGlobalShaderMap(GMaxRHIFeatureLevel), PermutationVector);
+
+	// Loop through each level of the mips that require creation and add a dispatch pass per level.
+	for (uint32 MipLevel = 1, MipCount = TextureDesc.NumMips; MipLevel < MipCount; ++MipLevel)
+	{
+		const FIntPoint DestTextureSize(
+			FMath::Max(TextureDesc.Extent.X >> MipLevel, 1),
+			FMath::Max(TextureDesc.Extent.Y >> MipLevel, 1));
+
+		FGenerateMipsIndirectCS::FParameters* PassParameters = GraphBuilder.AllocParameters<FGenerateMipsIndirectCS::FParameters>();
+		PassParameters->TexelSize = FVector2D(1.0f / DestTextureSize.X, 1.0f / DestTextureSize.Y);
+		PassParameters->MipInSRV = GraphBuilder.CreateSRV(FRDGTextureSRVDesc::CreateForMipLevel(Texture, MipLevel - 1));
+		PassParameters->MipOutUAV = GraphBuilder.CreateUAV(FRDGTextureUAVDesc(Texture, MipLevel));
+		PassParameters->MipSampler = Sampler;
+		PassParameters->IndirectDispatchArgsBuffer = IndirectDispatchArgsBuffer;
+
+		FComputeShaderUtils::AddPass(
+			GraphBuilder,
+			RDG_EVENT_NAME("GenerateMips DestMipLevel=%d", MipLevel),
+			ComputeShader,
+			PassParameters,
+			IndirectDispatchArgsBuffer, sizeof(FRHIDispatchIndirectParameters) * (MipLevel - 1));
+	}
+	
+}
+
 BEGIN_SHADER_PARAMETER_STRUCT(FCopyDestParameters, )
 	RDG_TEXTURE_ACCESS(Texture, ERHIAccess::CopyDest)
 END_SHADER_PARAMETER_STRUCT()
diff --git a/Engine/Source/Runtime/RenderCore/Public/GenerateMips.h b/Engine/Source/Runtime/RenderCore/Public/GenerateMips.h
--- a/Engine/Source/Runtime/RenderCore/Public/GenerateMips.h
+++ b/Engine/Source/Runtime/RenderCore/Public/GenerateMips.h
@@ -48,6 +48,14 @@ public:
 	}
 
 	static void ExecuteCompute(FRDGBuilder& GraphBuilder, FRDGTextureRef Texture, FRHISamplerState* Sampler);
+	
+	/** (SM5+) Generate mips for the requested RDG texture using the compute pass conditionally.
+		if( uint(ConditionBuffer[Offset]) > 0)
+			Execute(...)
+	*/
+	static void ExecuteCompute(FRDGBuilder& GraphBuilder, FRDGTextureRef Texture, FRHISamplerState* Sampler,
+								 FRDGBufferRef ConditionBuffer, uint32 Offset = 0);
+
 	static void ExecuteRaster(FRDGBuilder& GraphBuilder, FRDGTextureRef Texture, FRHISamplerState* Sampler);
 
 	//////////////////////////////////////////////////////////////////////////
diff --git a/Engine/Source/Runtime/Renderer/Private/PostProcess/PostProcessSubsurface.cpp b/Engine/Source/Runtime/Renderer/Private/PostProcess/PostProcessSubsurface.cpp
--- a/Engine/Source/Runtime/Renderer/Private/PostProcess/PostProcessSubsurface.cpp
+++ b/Engine/Source/Runtime/Renderer/Private/PostProcess/PostProcessSubsurface.cpp
@@ -20,6 +20,7 @@
 #include "ClearQuad.h"
 #include "Strata/Strata.h"
 #include "TemporalAA.h"
+#include "SubsurfaceTiles.h"
 
 namespace
 {
@@ -121,6 +122,13 @@ namespace
 		TEXT("0: Depth Only. It is more performant (x2 faster for close view).")
 		TEXT("1: Depth and normal. It leads to better quality in regions like eyelids. (default)"),
 		ECVF_RenderThreadSafe);
+
+	TAutoConsoleVariable<int32> CVarSSSMipmapsMinTileCount(
+		TEXT("r.SSS.Burley.MinGenerateMipsTileCount"),
+		4000,
+		TEXT("4000. (default) The minimal number of tiles to trigger subsurface radiance mip generation. Set to zero to always generate mips (Experimental value)"),
+		ECVF_RenderThreadSafe);
+
 	
 	DECLARE_GPU_STAT(SubsurfaceScattering)
 }
@@ -128,9 +136,6 @@ namespace
 // Define to use a custom ps to clear UAV.
 #define USE_CUSTOM_CLEAR_UAV
 
-// Define the size of subsurface group. @TODO: Set to 16 to use LDS.
-const uint32 kSubsurfaceGroupSize = 8;
-
 enum class ESubsurfaceMode : uint32
 {
 	// Performs a full resolution scattering filter.
@@ -237,9 +242,10 @@ FSubsurfaceParameters GetSubsurfaceCommonParameters(FRDGBuilder& GraphBuilder, c
 	const float SSSScaleX = SSSScaleZ / SUBSURFACE_KERNEL_SIZE * 0.5f;
 
 	const float SSSOverrideNumSamples = float(CVarSSSBurleyNumSamplesOverride.GetValueOnRenderThread());
+	const float MinGenerateMipsTileCount = FMath::Max(0.0f, CVarSSSMipmapsMinTileCount.GetValueOnRenderThread());
 
 	FSubsurfaceParameters Parameters;
-	Parameters.SubsurfaceParams = FVector4f(SSSScaleX, SSSScaleZ, SSSOverrideNumSamples, 0);
+	Parameters.SubsurfaceParams = FVector4f(SSSScaleX, SSSScaleZ, SSSOverrideNumSamples, MinGenerateMipsTileCount);
 	Parameters.ViewUniformBuffer = View.ViewUniformBuffer;
 	Parameters.SceneTextures = SceneTextures;
 	Parameters.SSProfilesTexture = GetSubsurfaceProfileTextureWithFallback();
@@ -399,6 +405,7 @@ IMPLEMENT_GLOBAL_SHADER(FSubsurfaceViewportCopyPS, "/Engine/Private/PostProcessS
 // Subsurface uniform buffer layout
 BEGIN_GLOBAL_SHADER_PARAMETER_STRUCT(FSubsurfaceUniformParameters, )
 	SHADER_PARAMETER(uint32, MaxGroupCount)
+	SHADER_PARAMETER(uint32, bRectPrimitive)
 END_GLOBAL_SHADER_PARAMETER_STRUCT()
 
 IMPLEMENT_GLOBAL_SHADER_PARAMETER_STRUCT(FSubsurfaceUniformParameters, "SubsurfaceUniformParameters");
@@ -409,6 +416,7 @@ FSubsurfaceUniformRef CreateUniformBuffer(FViewInfo const& View, int32 MaxGroupC
 {
 	FSubsurfaceUniformParameters Parameters;
 	Parameters.MaxGroupCount = MaxGroupCount;
+	Parameters.bRectPrimitive = GRHISupportsRectTopology;
 	return FSubsurfaceUniformRef::CreateUniformBufferImmediate(Parameters, UniformBuffer_SingleFrame);
 }
 
@@ -426,8 +434,7 @@ public:
 	}
 
 	BEGIN_SHADER_PARAMETER_STRUCT(FParameters, )
-		SHADER_PARAMETER_RDG_BUFFER_UAV(RWBuffer<uint>, RWSeparableGroupBuffer)
-		SHADER_PARAMETER_RDG_BUFFER_UAV(RWBuffer<uint>, RWBurleyGroupBuffer)
+		SHADER_PARAMETER_RDG_BUFFER_UAV(RWBuffer<uint>, RWTileTypeCountBuffer)
 	END_SHADER_PARAMETER_STRUCT()
 
 	static bool ShouldCompilePermutation(const FGlobalShaderPermutationParameters& Parameters)
@@ -438,6 +445,34 @@ public:
 
 IMPLEMENT_GLOBAL_SHADER(FSubsurfaceInitValueBufferCS, "/Engine/Private/PostProcessSubsurface.usf", "InitValueBufferCS", SF_Compute);
 
+class FFillMipsConditionBufferCS : public FSubsurfaceShader
+{
+public:
+	DECLARE_GLOBAL_SHADER(FFillMipsConditionBufferCS);
+	SHADER_USE_PARAMETER_STRUCT(FFillMipsConditionBufferCS, FSubsurfaceShader);
+
+	static void ModifyCompilationEnvironment(const FGlobalShaderPermutationParameters& Parameters, FShaderCompilerEnvironment& OutEnvironment)
+	{
+		FSubsurfaceShader::ModifyCompilationEnvironment(Parameters, OutEnvironment);
+		OutEnvironment.SetDefine(TEXT("SUBSURFACE_BURLEY_COMPUTE"), 1);
+		OutEnvironment.SetDefine(TEXT("STRATA_ENABLED"), Strata::IsStrataEnabled() ? 1 : 0);
+	}
+
+	BEGIN_SHADER_PARAMETER_STRUCT(FParameters, )
+		SHADER_PARAMETER(uint32, MinGenerateMipsTileCount)
+		SHADER_PARAMETER(uint32, Offset)
+		SHADER_PARAMETER_RDG_BUFFER_SRV(Buffer<uint>, TileTypeCountBuffer)
+		SHADER_PARAMETER_RDG_BUFFER_UAV(RWBuffer<uint>, RWMipsConditionBuffer)
+		END_SHADER_PARAMETER_STRUCT()
+
+		static bool ShouldCompilePermutation(const FGlobalShaderPermutationParameters& Parameters)
+	{
+		return IsFeatureLevelSupported(Parameters.Platform, ERHIFeatureLevel::SM5);
+	}
+};
+
+IMPLEMENT_GLOBAL_SHADER(FFillMipsConditionBufferCS, "/Engine/Private/PostProcessSubsurface.usf", "FillMipsConditionBufferCS", SF_Compute);
+
 class FSubsurfaceBuildIndirectDispatchArgsCS : public FSubsurfaceShader
 {
 public:
@@ -454,7 +489,8 @@ public:
 	BEGIN_SHADER_PARAMETER_STRUCT(FParameters, )
 		SHADER_PARAMETER_STRUCT_REF(FSubsurfaceUniformParameters, SubsurfaceUniformParameters)
 		SHADER_PARAMETER_RDG_BUFFER_UAV(RWBuffer<uint>, RWIndirectDispatchArgsBuffer)
-		SHADER_PARAMETER_RDG_BUFFER_SRV(Buffer<uint>, GroupBuffer)
+		SHADER_PARAMETER_RDG_BUFFER_UAV(RWBuffer<uint>, RWIndirectDrawArgsBuffer)
+		SHADER_PARAMETER_RDG_BUFFER_SRV(Buffer<uint>, TileTypeCountBuffer)
 	END_SHADER_PARAMETER_STRUCT()
 
 	static bool ShouldCompilePermutation(const FGlobalShaderPermutationParameters& Parameters)
@@ -487,6 +523,8 @@ public:
 		SHADER_PARAMETER_RDG_TEXTURE_UAV(RWTexture2D, SetupTexture)
 		SHADER_PARAMETER_RDG_BUFFER_UAV(RWBuffer<uint>, RWSeparableGroupBuffer)
 		SHADER_PARAMETER_RDG_BUFFER_UAV(RWBuffer<uint>, RWBurleyGroupBuffer)
+		SHADER_PARAMETER_RDG_BUFFER_UAV(RWBuffer<uint>, RWAllGroupBuffer)
+		SHADER_PARAMETER_RDG_BUFFER_UAV(RWBuffer<uint>, RWTileTypeCountBuffer)
 		SHADER_PARAMETER_RDG_TEXTURE_UAV(RWTexture2D, ProfileIdTexture)
 		SHADER_PARAMETER_STRUCT_REF(FSubsurfaceUniformParameters, SubsurfaceUniformParameters)
 	END_SHADER_PARAMETER_STRUCT()
@@ -511,7 +549,7 @@ public:
 		FSubsurfaceShader::ModifyCompilationEnvironment(Parameters, OutEnvironment);
 		OutEnvironment.SetDefine(TEXT("SUBSURFACE_BURLEY_COMPUTE"), 1);
 		OutEnvironment.SetDefine(TEXT("ENABLE_VELOCITY"), 1);
-		OutEnvironment.SetDefine(TEXT("SUBSURFACE_GROUP_SIZE"), kSubsurfaceGroupSize);
+		OutEnvironment.SetDefine(TEXT("SUBSURFACE_GROUP_SIZE"), FSubsurfaceTiles::TileSize);
 		OutEnvironment.SetDefine(TEXT("STRATA_ENABLED"), Strata::IsStrataEnabled() ? 1 : 0);
 	}
 
@@ -522,6 +560,7 @@ public:
 		SHADER_PARAMETER_RDG_TEXTURE_UAV(RWTexture2D, SSSColorUAV)
 		SHADER_PARAMETER_RDG_TEXTURE_UAV(RWTexture2D, HistoryUAV)
 		SHADER_PARAMETER_RDG_BUFFER_SRV(Buffer<uint>, GroupBuffer)
+		SHADER_PARAMETER_RDG_BUFFER_SRV(Buffer<uint>, TileTypeCountBuffer)
 		RDG_BUFFER_ACCESS(IndirectDispatchArgsBuffer, ERHIAccess::IndirectArgs)
 		SHADER_PARAMETER_STRUCT(FSubsurfaceInput, SubsurfaceInput0)
 		SHADER_PARAMETER_SAMPLER(SamplerState, SubsurfaceSampler0)
@@ -661,6 +700,19 @@ public:
 				static_cast<int32>(FSubsurfaceIndirectDispatchCS::EBilateralFilterKernelFunctionType::Depth),
 				static_cast<int32>(FSubsurfaceIndirectDispatchCS::EBilateralFilterKernelFunctionType::DepthAndNormal)));
 	}
+
+	static FSubsurfaceTiles::ETileType ToSubsurfaceTileType(ESubsurfaceType Type)
+	{
+		FSubsurfaceTiles::ETileType TileType = FSubsurfaceTiles::ETileType::AFIS;
+
+		switch (Type)
+		{
+			case ESubsurfaceType::BURLEY: TileType = FSubsurfaceTiles::ETileType::AFIS; break;
+			case ESubsurfaceType::SEPARABLE: TileType = FSubsurfaceTiles::ETileType::SEPARABLE; break;
+		}
+
+		return TileType;
+	}
 };
 
 IMPLEMENT_GLOBAL_SHADER(FSubsurfaceIndirectDispatchCS, "/Engine/Private/PostProcessSubsurface.usf", "MainIndirectDispatchCS", SF_Compute);
@@ -674,6 +726,7 @@ class FSubsurfaceSRVResolvePS : public FSubsurfaceShader
 
 	BEGIN_SHADER_PARAMETER_STRUCT(FParameters, )
 		SHADER_PARAMETER_RDG_TEXTURE(Texture2D, SubsurfaceInput0_Texture)
+		RDG_BUFFER_ACCESS(IndirectDispatchArgsBuffer, ERHIAccess::IndirectArgs)
 		SHADER_PARAMETER_SAMPLER(SamplerState, SubsurfaceSampler0)
 		RENDER_TARGET_BINDING_SLOTS()
 	END_SHADER_PARAMETER_STRUCT();
@@ -711,6 +764,24 @@ FRDGTextureRef CreateBlackUAVTexture(FRDGBuilder& GraphBuilder, FRDGTextureDesc
 	return SRVTextureOutput;
 }
 
+// create a black UAV only when CondtionBuffer[Offset] > 0.
+FRDGTextureRef CreateBlackUAVTexture(FRDGBuilder& GraphBuilder, FRDGTextureDesc SRVDesc, const TCHAR* TextureName, FRDGEventName&& PassName,
+	const FScreenPassTextureViewport& SceneViewport, FRDGBufferRef ConditionBuffer, uint32 Offset)
+{
+	SRVDesc.Flags |= TexCreate_ShaderResource | TexCreate_UAV;
+	FRDGTextureRef TextureOutput = GraphBuilder.CreateTexture(SRVDesc, TextureName);
+	
+	AddConditionalClearBlackUAVPass(
+		GraphBuilder,
+		Forward<FRDGEventName>(PassName),
+		TextureOutput,
+		SceneViewport,
+		ConditionBuffer, 
+		Offset);
+
+	return TextureOutput;
+}
+
 // Helper function to use external textures for the current GraphBuilder.
 // When the texture is null, we use BlackDummy.
 FRDGTextureRef RegisterExternalRenderTarget(FRDGBuilder& GraphBuilder, TRefCountPtr<IPooledRenderTarget>* PtrExternalTexture, FIntPoint CurentViewExtent, const TCHAR* Name)
@@ -745,6 +816,7 @@ class FSubsurfaceRecombinePS : public FSubsurfaceShader
 	BEGIN_SHADER_PARAMETER_STRUCT(FParameters, )
 		SHADER_PARAMETER_RDG_UNIFORM_BUFFER(FStrataGlobalUniformParameters, Strata)
 		SHADER_PARAMETER_STRUCT_INCLUDE(FSubsurfaceParameters, Subsurface)
+		SHADER_PARAMETER_STRUCT_INCLUDE(FSubsurfaceTilePassVS::FParameters, TileParameters)
 		SHADER_PARAMETER_STRUCT(FSubsurfaceInput, SubsurfaceInput0)
 		SHADER_PARAMETER_STRUCT(FSubsurfaceInput, SubsurfaceInput1)
 		SHADER_PARAMETER_SAMPLER(SamplerState, SubsurfaceSampler0)
@@ -777,6 +849,11 @@ class FSubsurfaceRecombinePS : public FSubsurfaceShader
 	class FRunningInSeparable : SHADER_PERMUTATION_BOOL("SUBSURFACE_FORCE_SEPARABLE");
 	using FPermutationDomain = TShaderPermutationDomain<FDimensionMode, FDimensionQuality, FDimensionCheckerboard, FDimensionHalfRes, FRunningInSeparable>;
 
+	static bool ShouldCompilePermutation(const FGlobalShaderPermutationParameters& Parameters)
+	{
+		return GetMaxSupportedFeatureLevel(Parameters.Platform) >= ERHIFeatureLevel::SM5;
+	}
+
 	static void ModifyCompilationEnvironment(const FGlobalShaderPermutationParameters& Parameters, FShaderCompilerEnvironment& OutEnvironment)
 	{
 		FSubsurfaceShader::ModifyCompilationEnvironment(Parameters, OutEnvironment);
@@ -825,6 +902,30 @@ class FSubsurfaceRecombinePS : public FSubsurfaceShader
 
 IMPLEMENT_GLOBAL_SHADER(FSubsurfaceRecombinePS, "/Engine/Private/PostProcessSubsurface.usf", "SubsurfaceRecombinePS", SF_Pixel);
 
+// copy write back additively to the scene color texture
+class FSubsurfaceRecombineCopyPS : public FSubsurfaceShader
+{
+	DECLARE_GLOBAL_SHADER(FSubsurfaceRecombineCopyPS);
+	SHADER_USE_PARAMETER_STRUCT(FSubsurfaceRecombineCopyPS, FSubsurfaceShader);
+
+	BEGIN_SHADER_PARAMETER_STRUCT(FParameters, )
+		SHADER_PARAMETER_RDG_UNIFORM_BUFFER(FStrataGlobalUniformParameters, Strata)
+		SHADER_PARAMETER_STRUCT_INCLUDE(FSubsurfaceParameters, Subsurface)
+		SHADER_PARAMETER_STRUCT_INCLUDE(FSubsurfaceTilePassVS::FParameters, TileParameters)
+		SHADER_PARAMETER_STRUCT(FSubsurfaceInput, SubsurfaceInput0)
+		SHADER_PARAMETER_SAMPLER(SamplerState, SubsurfaceSampler0)
+		RENDER_TARGET_BINDING_SLOTS()
+	END_SHADER_PARAMETER_STRUCT();
+
+	static bool ShouldCompilePermutation(const FGlobalShaderPermutationParameters& Parameters)
+	{
+		return IsFeatureLevelSupported(Parameters.Platform, ERHIFeatureLevel::SM5);
+	}
+};
+
+IMPLEMENT_GLOBAL_SHADER(FSubsurfaceRecombineCopyPS, "/Engine/Private/PostProcessSubsurface.usf", "SubsurfaceRecombineCopyPS", SF_Pixel);
+
+
 void AddSubsurfaceViewPass(
 	FRDGBuilder& GraphBuilder,
 	const FViewInfo& View,
@@ -854,15 +955,14 @@ void AddSubsurfaceViewPass(
 
 	const bool bUseProfileIdCache = !bForceRunningInSeparable && IsProfileIdCacheEnabled();
 
+	const int32 MinGenerateMipsTileCount = FMath::Max(0, CVarSSSMipmapsMinTileCount.GetValueOnRenderThread());
+
 	/**
 	 * All subsurface passes within the screen-space subsurface effect can operate at half or full resolution,
 	 * depending on the subsurface mode. The values are precomputed and shared among all Subsurface textures.
 	 */
 	const FScreenPassTextureViewport SubsurfaceViewport = GetDownscaledViewport(SceneViewport, ScaleFactor);
 
-	const FIntPoint TileDimension = FIntPoint::DivideAndRoundUp(SubsurfaceViewport.Extent, kSubsurfaceGroupSize);
-	const int32 MaxGroupCount = TileDimension.X*TileDimension.Y;
-
 	const FRDGTextureDesc SceneColorTextureDescriptor = FRDGTextureDesc::Create2D(
 		SceneViewport.Extent,
 		PF_FloatRGBA,
@@ -908,10 +1008,13 @@ void AddSubsurfaceViewPass(
 	TRefCountPtr<IPooledRenderTarget>* QualityHistoryState = ViewState ? &ViewState->SubsurfaceScatteringQualityHistoryRT : NULL;
 
 	//allocate/reallocate the quality history texture. 
-	FRDGTextureRef QualityHistoryTexture = RegisterExternalRenderTarget(GraphBuilder, QualityHistoryState, SceneColorTextureDescriptor.Extent, TEXT("QualityHistoryTexture"));
+	FRDGTextureRef QualityHistoryTexture = RegisterExternalRenderTarget(GraphBuilder, QualityHistoryState, SceneColorTextureDescriptor.Extent, TEXT("Subsurface.QualityHistoryTexture"));
 	FRDGTextureRef NewQualityHistoryTexture = nullptr;
 
 	RDG_GPU_STAT_SCOPE(GraphBuilder, SubsurfaceScattering);
+
+	FSubsurfaceTiles Tiles;
+
 	/**
 	 * When in bypass mode, the setup and convolution passes are skipped, but lighting
 	 * reconstruction is still performed in the recombine pass.
@@ -919,13 +1022,13 @@ void AddSubsurfaceViewPass(
 	if (SubsurfaceMode != ESubsurfaceMode::Bypass)
 	{
 		// Support mipmaps in full resolution only.
-		SetupTexture = GraphBuilder.CreateTexture(bForceRunningInSeparable ? SubsurfaceTextureDescriptor:SubsurfaceTextureWith6MipsDescriptor, TEXT("SubsurfaceSetupTexture"));
+		SetupTexture = GraphBuilder.CreateTexture(bForceRunningInSeparable ? SubsurfaceTextureDescriptor:SubsurfaceTextureWith6MipsDescriptor, TEXT("Subsurface.SetupTexture"));
 
 		// profile cache to accelerate sampling
 		if (bUseProfileIdCache)
 		{
-			// This path was designed to get used when r.SSS.Burley.EnableProfileIdCache is true, but we had to disable this path because R8 UAVs are not supported on all platforms. 
-			ProfileIdTexture = GraphBuilder.CreateTexture(ProfileIdTextureDescriptor, TEXT("ProfileIdTexture"));
+			// This path was designed to get used when r.SSS.Burley.EnableProfileIdCache is true. 
+			ProfileIdTexture = GraphBuilder.CreateTexture(ProfileIdTextureDescriptor, TEXT("Subsurface.ProfileIdTexture"));
 		}
 		else
 		{
@@ -933,39 +1036,49 @@ void AddSubsurfaceViewPass(
 		}
 
 		FRDGTextureRef VelocityTexture = GetIfProduced(SceneTextures.Velocity, SystemTextures.Black);
-		FSubsurfaceUniformRef UniformBuffer = CreateUniformBuffer(View, MaxGroupCount);
 		
-		// Pre-allocate black UAV together.
+		// Initialize tile buffers.
 		{
-			SubsurfaceSubpassOneTex = CreateBlackUAVTexture(GraphBuilder, SubsurfaceTextureWith6MipsDescriptor, TEXT("SubsurfaceSubpassOneTex"),
-				View, SubsurfaceViewport);
-			SubsurfaceSubpassTwoTex = CreateBlackUAVTexture(GraphBuilder, SubsurfaceTextureWith6MipsDescriptor, TEXT("SubsurfaceSubpassTwoTex"),
-				View, SubsurfaceViewport);
-			// Only clear when we are in full resolution.
-			if (!bForceRunningInSeparable)
+			check(FSubsurfaceTiles::TilePerThread_GroupSize == 64); // If this value change, we need to update the shaders using 
+			check(FSubsurfaceTiles::TileSize == 8); // only size supported for now
+
+			Tiles.TileDimension = FIntPoint::DivideAndRoundUp(SubsurfaceViewport.Extent, FSubsurfaceTiles::TileSize);
+			Tiles.TileCount = Tiles.TileDimension.X * Tiles.TileDimension.Y;
+			Tiles.bRectPrimitive = GRHISupportsRectTopology;
+			Tiles.TileTypeCountBuffer = GraphBuilder.CreateBuffer(FRDGBufferDesc::CreateBufferDesc(sizeof(uint32), FSubsurfaceTiles::TileTypeCount), TEXT("Subsurface.TileCountBuffer"));
+			Tiles.TileIndirectDispatchBuffer = 
+				GraphBuilder.CreateBuffer(FRDGBufferDesc::CreateIndirectDesc<FRHIDispatchIndirectParameters>(FSubsurfaceTiles::TileTypeCount), TEXT("Subsurface.TileIndirectDispatchBuffer"));
+			Tiles.TileIndirectDrawBuffer =
+				GraphBuilder.CreateBuffer(FRDGBufferDesc::CreateIndirectDesc<FRHIDrawIndirectParameters>(FSubsurfaceTiles::TileTypeCount), TEXT("Subsurface.TileIndirectDrawBuffer"));
+
+			Tiles.TileDataBuffer[ToIndex(FSubsurfaceTiles::ETileType::All)] = 
+				GraphBuilder.CreateBuffer(FRDGBufferDesc::CreateBufferDesc(sizeof(uint32), 2 * Tiles.TileCount), TEXT("Subsurface.TileDataBuffer(All)"));
+			Tiles.TileDataBuffer[ToIndex(FSubsurfaceTiles::ETileType::AFIS)] = 
+				GraphBuilder.CreateBuffer(FRDGBufferDesc::CreateBufferDesc(sizeof(uint32), 2 * Tiles.TileCount), TEXT("Subsurface.TileDataBuffer(AFIS)"));
+			Tiles.TileDataBuffer[ToIndex(FSubsurfaceTiles::ETileType::SEPARABLE)] = 
+				GraphBuilder.CreateBuffer(FRDGBufferDesc::CreateBufferDesc(sizeof(uint32), 2 * Tiles.TileCount), TEXT("Subsurface.TileDataBuffer(SEPARABLE)"));
+
+			for (uint32 BufferIt = 0; BufferIt < FSubsurfaceTiles::TileTypeCount; ++BufferIt)
 			{
-				NewQualityHistoryTexture = CreateBlackUAVTexture(GraphBuilder, SubsurfaceTextureDescriptor, TEXT("SubsurfaceQualityHistoryState"),
-					View, SubsurfaceViewport);
+				if (Tiles.TileDataBuffer[BufferIt])
+				{
+					Tiles.TileDataSRV[BufferIt] = GraphBuilder.CreateSRV(Tiles.TileDataBuffer[BufferIt], PF_R32_UINT);
+				}
 			}
 		}
 
-		// Initialize the group buffer
-		FRDGBufferRef SeparableGroupBuffer = GraphBuilder.CreateBuffer(FRDGBufferDesc::CreateBufferDesc(sizeof(uint32), 2*(MaxGroupCount + 1)), TEXT("SeparableGroupBuffer"));;
-		FRDGBufferRef BurleyGroupBuffer = GraphBuilder.CreateBuffer(FRDGBufferDesc::CreateBufferDesc(sizeof(uint32), 2*(MaxGroupCount + 1)), TEXT("BurleyGroupBuffer"));;
-		FRDGBufferRef SeparableIndirectDispatchArgsBuffer = GraphBuilder.CreateBuffer(FRDGBufferDesc::CreateIndirectDesc(4), TEXT("SeparableIndirectDispatchArgs"));
-		FRDGBufferRef BurleyIndirectDispatchArgsBuffer = GraphBuilder.CreateBuffer(FRDGBufferDesc::CreateIndirectDesc(4), TEXT("BurleyIndirectDispatchArgs"));
-
 		// Initialize the group counters
 		{
 			typedef FSubsurfaceInitValueBufferCS SHADER;
 			TShaderMapRef<SHADER> ComputeShader(View.ShaderMap);
 			SHADER::FParameters* PassParameters = GraphBuilder.AllocParameters<SHADER::FParameters>();
-			PassParameters->RWBurleyGroupBuffer = GraphBuilder.CreateUAV(BurleyGroupBuffer, EPixelFormat::PF_R32_UINT);
-			PassParameters->RWSeparableGroupBuffer = GraphBuilder.CreateUAV(SeparableGroupBuffer, EPixelFormat::PF_R32_UINT);
+			PassParameters->RWTileTypeCountBuffer = GraphBuilder.CreateUAV(Tiles.TileTypeCountBuffer, EPixelFormat::PF_R32_UINT);
 			FComputeShaderUtils::AddPass(GraphBuilder, RDG_EVENT_NAME("InitGroupCounter"), ComputeShader, PassParameters, FIntVector(1, 1, 1));
 		}
 
-		// Call the indirect setup
+		FSubsurfaceUniformRef UniformBuffer = CreateUniformBuffer(View, Tiles.TileCount);
+
+		// Call the indirect setup (with tile classification)
 		{
 			FRDGTextureSRVDesc SceneColorTextureSRVDesc = FRDGTextureSRVDesc::Create(SceneColorTexture);
 			FRDGTextureUAVDesc SetupTextureOutDesc(SetupTexture, 0);
@@ -981,8 +1094,11 @@ void AddSubsurfaceViewPass(
 			{
 				PassParameters->ProfileIdTexture = GraphBuilder.CreateUAV(FRDGTextureUAVDesc(ProfileIdTexture));
 			}
-			PassParameters->RWBurleyGroupBuffer = GraphBuilder.CreateUAV(BurleyGroupBuffer, EPixelFormat::PF_R32_UINT);
-			PassParameters->RWSeparableGroupBuffer = GraphBuilder.CreateUAV(SeparableGroupBuffer, EPixelFormat::PF_R32_UINT);
+			PassParameters->RWBurleyGroupBuffer = GraphBuilder.CreateUAV(Tiles.TileDataBuffer[ToIndex(FSubsurfaceTiles::ETileType::AFIS)], EPixelFormat::PF_R32_UINT);
+			PassParameters->RWSeparableGroupBuffer = GraphBuilder.CreateUAV(Tiles.TileDataBuffer[ToIndex(FSubsurfaceTiles::ETileType::SEPARABLE)], EPixelFormat::PF_R32_UINT);
+			PassParameters->RWAllGroupBuffer = GraphBuilder.CreateUAV(Tiles.TileDataBuffer[ToIndex(FSubsurfaceTiles::ETileType::All)], EPixelFormat::PF_R32_UINT);
+			PassParameters->RWTileTypeCountBuffer = GraphBuilder.CreateUAV(Tiles.TileTypeCountBuffer, EPixelFormat::PF_R32_UINT);
+
 			PassParameters->SubsurfaceUniformParameters = UniformBuffer;
 			PassParameters->Strata = Strata::BindStrataGlobalUniformParameters(View.StrataSceneData);
 
@@ -1003,45 +1119,70 @@ void AddSubsurfaceViewPass(
 					SubsurfaceViewport.Extent.Y),
 				ComputeShader,
 				PassParameters,
-				FComputeShaderUtils::GetGroupCount(SubsurfaceViewport.Extent, kSubsurfaceGroupSize));
+				FComputeShaderUtils::GetGroupCount(SubsurfaceViewport.Extent, FSubsurfaceTiles::TileSize));
 		}
 
-		// In half resolution, only Separable is used. We do not need this mipmap.
-		if(!bForceRunningInSeparable)
+		// Setup the indirect dispatch & draw arguments.
 		{
-			// Generate mipmap for the diffuse scene color and depth, use bilinear filter
-			FGenerateMips::Execute(GraphBuilder, SetupTexture, BilinearBorderSampler);
+				typedef FSubsurfaceBuildIndirectDispatchArgsCS ARGSETUPSHADER;
+				ARGSETUPSHADER::FParameters* PassParameters = GraphBuilder.AllocParameters<ARGSETUPSHADER::FParameters>();
+				PassParameters->SubsurfaceUniformParameters = UniformBuffer;
+				PassParameters->RWIndirectDispatchArgsBuffer = GraphBuilder.CreateUAV(Tiles.TileIndirectDispatchBuffer, EPixelFormat::PF_R32_UINT);
+				PassParameters->RWIndirectDrawArgsBuffer = GraphBuilder.CreateUAV(Tiles.TileIndirectDrawBuffer, EPixelFormat::PF_R32_UINT);
+				PassParameters->TileTypeCountBuffer = GraphBuilder.CreateSRV(Tiles.TileTypeCountBuffer, EPixelFormat::PF_R32_UINT);
+				TShaderMapRef<ARGSETUPSHADER> ComputeShader(View.ShaderMap);
+				FComputeShaderUtils::AddPass(GraphBuilder, FRDGEventName(TEXT("BuildIndirectArgs(Dispatch & Draw)")), ComputeShader, PassParameters, FIntVector(1, 1, 1));
 		}
 
-		typedef FSubsurfaceIndirectDispatchCS SHADER;
+		// In half resolution, only Separable is used. We do not need this mipmap.
+		if (!bForceRunningInSeparable)
+		{
+			FRDGBufferRef MipsConditionBuffer = GraphBuilder.CreateBuffer(
+				FRDGBufferDesc::CreateBufferDesc(sizeof(uint32), FSubsurfaceTiles::TileTypeCount), TEXT("Subsurface.MipsConditionBuffer"));
+			const int32 Offset = ToIndex(FSubsurfaceTiles::ETileType::AFIS);
+
+			{
+				typedef FFillMipsConditionBufferCS SHADER;
+				TShaderMapRef<SHADER> ComputeShader(View.ShaderMap);
+				SHADER::FParameters* PassParameters = GraphBuilder.AllocParameters<SHADER::FParameters>();
+				PassParameters->MinGenerateMipsTileCount = MinGenerateMipsTileCount;
+				PassParameters->Offset = Offset;
+				PassParameters->TileTypeCountBuffer = GraphBuilder.CreateSRV(Tiles.TileTypeCountBuffer, EPixelFormat::PF_R32_UINT);
+				PassParameters->RWMipsConditionBuffer = GraphBuilder.CreateUAV(MipsConditionBuffer, EPixelFormat::PF_R32_UINT);
 
-		FRHISamplerState* SubsurfaceSamplerState = SHADER::GetSamplerState(bHalfRes);
-		const SHADER::EQuality SubsurfaceQuality = SHADER::GetQuality();
+				FComputeShaderUtils::AddPass(GraphBuilder, RDG_EVENT_NAME("FillMipsConditionBufferCS"), ComputeShader, PassParameters, FIntVector(1, 1, 1));
+			}
 
-		// Store the buffer
-		const FRDGBufferRef SubsurfaceBufferUsage[] = { BurleyGroupBuffer,                      SeparableGroupBuffer };
-		const FRDGBufferRef  SubsurfaceBufferArgs[] = { BurleyIndirectDispatchArgsBuffer,       SeparableIndirectDispatchArgsBuffer };
-		const TCHAR*		  SubsurfacePhaseName[] = { TEXT("BuildBurleyIndirectDispatchArgs"),TEXT("BuildSeparableIndirectDispatchArgs") };
+			// Generate mipmap for the diffuse scene color and depth, use bilinear filter conditionally
+			FGenerateMips::ExecuteCompute(GraphBuilder, SetupTexture, BilinearBorderSampler, MipsConditionBuffer, Offset);
+		}
 
-		// Setup the indirect arguments.
+		// Allocate auxiliary buffers
 		{
-			const int NumOfSubsurfaceType = 2;
-
-			for (int SubsurfaceTypeIndex = 0; SubsurfaceTypeIndex < NumOfSubsurfaceType; ++SubsurfaceTypeIndex)
+			// Pre-allocate black UAV if we have separable tiles.
+			SubsurfaceSubpassOneTex = CreateBlackUAVTexture(
+				GraphBuilder, 
+				SubsurfaceTextureDescriptor,
+				TEXT("Subsurface.SubpassOneTex"),
+				RDG_EVENT_NAME("%s (%s, %s.TileCount > 0)",TEXT("ClearUAV"), TEXT("Subsurface.SubpassOneTex"), ToString(FSubsurfaceTiles::ETileType::SEPARABLE)),
+				SubsurfaceViewport,
+				Tiles.TileTypeCountBuffer,
+				ToIndex(FSubsurfaceTiles::ETileType::SEPARABLE));
+
+			SubsurfaceSubpassTwoTex = GraphBuilder.CreateTexture(SubsurfaceTextureDescriptor, TEXT("Subsurface.SubpassTwoTex"));
+			
+			if (!bForceRunningInSeparable)
 			{
-				typedef FSubsurfaceBuildIndirectDispatchArgsCS ARGSETUPSHADER;
-				ARGSETUPSHADER::FParameters* PassParameters = GraphBuilder.AllocParameters<ARGSETUPSHADER::FParameters>();
-				PassParameters->SubsurfaceUniformParameters = UniformBuffer;
-				PassParameters->RWIndirectDispatchArgsBuffer = GraphBuilder.CreateUAV(SubsurfaceBufferArgs[SubsurfaceTypeIndex], EPixelFormat::PF_R32_UINT);
-				PassParameters->GroupBuffer = GraphBuilder.CreateSRV(SubsurfaceBufferUsage[SubsurfaceTypeIndex], EPixelFormat::PF_R32_UINT);
-
-				TShaderMapRef<ARGSETUPSHADER> ComputeShader(View.ShaderMap);
-				FComputeShaderUtils::AddPass(GraphBuilder, FRDGEventName(SubsurfacePhaseName[SubsurfaceTypeIndex]), ComputeShader, PassParameters, FIntVector(1, 1, 1));
+				NewQualityHistoryTexture = GraphBuilder.CreateTexture(SubsurfaceTextureDescriptor, TEXT("Subsurface.QualityHistoryState"));
 			}
 		}
 
 		// Major pass to combine Burley and Separable
 		{
+			typedef FSubsurfaceIndirectDispatchCS SHADER;
+
+			FRHISamplerState* SubsurfaceSamplerState = SHADER::GetSamplerState(bHalfRes);
+
 			struct FSubsurfacePassInfo
 			{
 				FSubsurfacePassInfo(const TCHAR* InName, FRDGTextureRef InInput, FRDGTextureRef InOutput,
@@ -1066,6 +1207,8 @@ void AddSubsurfaceViewPass(
 				{	 TEXT("SubsurfacePassFour_BVar"),    SubsurfaceSubpassOneTex, SubsurfaceSubpassTwoTex, SHADER::ESubsurfaceType::BURLEY	 , SHADER::ESubsurfacePass::PassTwo}  //Burley Variance
 			};
 
+			const FRDGBufferSRVRef SubsurfaceBufferUsage[] = { Tiles.GetTileBufferSRV(FSubsurfaceTiles::ETileType::AFIS), Tiles.GetTileBufferSRV(FSubsurfaceTiles::ETileType::SEPARABLE) };
+
 			//Dispatch the two phase for both SSS
 			for (int PassIndex = 0; PassIndex < NumOfSubsurfacePass; ++PassIndex)
 			{
@@ -1086,14 +1229,17 @@ void AddSubsurfaceViewPass(
 				PassParameters->SSSColorUAV = GraphBuilder.CreateUAV(SSSColorUAVDesc);
 				PassParameters->SubsurfaceInput0 = GetSubsurfaceInput(TextureInput, SubsurfaceViewportParameters);
 				PassParameters->SubsurfaceSampler0 = SubsurfaceSamplerState;
-				PassParameters->GroupBuffer = GraphBuilder.CreateSRV(SubsurfaceBufferUsage[SubsurfaceTypeIndex], EPixelFormat::PF_R32_UINT);
-				PassParameters->IndirectDispatchArgsBuffer = SubsurfaceBufferArgs[SubsurfaceTypeIndex];
+				PassParameters->GroupBuffer = SubsurfaceBufferUsage[SubsurfaceTypeIndex];
+				PassParameters->TileTypeCountBuffer = GraphBuilder.CreateSRV(Tiles.TileTypeCountBuffer,EPixelFormat::PF_R32_UINT);
+				PassParameters->IndirectDispatchArgsBuffer = Tiles.TileIndirectDispatchBuffer;
 				PassParameters->Strata = Strata::BindStrataGlobalUniformParameters(View.StrataSceneData);
 
 				if (SubsurfacePassFunction == SHADER::ESubsurfacePass::PassOne && SubsurfaceType == SHADER::ESubsurfaceType::BURLEY)
 				{
 					PassParameters->SubsurfaceInput1 = GetSubsurfaceInput(QualityHistoryTexture, SubsurfaceViewportParameters);
 					PassParameters->SubsurfaceSampler1 = PointClampSampler;
+					PassParameters->SubsurfaceInput2 = GetSubsurfaceInput(VelocityTexture, SubsurfaceViewportParameters);
+					PassParameters->SubsurfaceSampler2 = PointClampSampler;
 				}
 
 				if (SubsurfacePassFunction == SHADER::ESubsurfacePass::PassTwo && SubsurfaceType == SHADER::ESubsurfaceType::BURLEY)
@@ -1138,58 +1284,95 @@ void AddSubsurfaceViewPass(
 						ComputeShaderPermutationVector.Get<SHADER::FRunningInSeparable>() ? TEXT(" RunningInSeparable") : TEXT("")),
 					ComputeShader,
 					PassParameters,
-					SubsurfaceBufferArgs[SubsurfaceTypeIndex], 0);
+					PassParameters->IndirectDispatchArgsBuffer, Tiles.GetIndirectDispatchArgOffset(SHADER::ToSubsurfaceTileType(SubsurfaceType)));
 			}
 		}
 	}
 
-	// Recombines scattering result with scene color.
+	// Recombines scattering result with scene color, and updates only SSS region using specialized vertex shader
+	if (SubsurfaceMode != ESubsurfaceMode::Bypass)
 	{
-		FSubsurfaceRecombinePS::FParameters* PassParameters = GraphBuilder.AllocParameters<FSubsurfaceRecombinePS::FParameters>();
-		PassParameters->Subsurface = SubsurfaceCommonParameters;
-		PassParameters->RenderTargets[0] = FRenderTargetBinding(SceneColorTextureOutput, SceneColorTextureLoadAction);
-		PassParameters->SubsurfaceInput0 = GetSubsurfaceInput(SceneColorTexture, SceneViewportParameters);
-		PassParameters->SubsurfaceSampler0 = BilinearBorderSampler;
-		PassParameters->Strata = Strata::BindStrataGlobalUniformParameters(View.StrataSceneData);
-
-		// Scattering output target is only used when scattering is enabled.
-		if (SubsurfaceMode != ESubsurfaceMode::Bypass)
+		FRDGTextureRef SubsurfaceIntermediateTexture = GraphBuilder.CreateTexture(SceneColorTextureDesc, TEXT("Subsurface.Recombines"));
+
+		// recombine into intermediate textures before copying back.
 		{
+			FSubsurfaceTiles::ETileType TileType = FSubsurfaceTiles::ETileType::All;
+
+			FSubsurfaceRecombinePS::FParameters* PassParameters = GraphBuilder.AllocParameters<FSubsurfaceRecombinePS::FParameters>();
+			PassParameters->Subsurface = SubsurfaceCommonParameters;
+			PassParameters->TileParameters = GetSubsurfaceTileParameters(View, Tiles, TileType, SubsurfaceViewport.Extent);
+			PassParameters->RenderTargets[0] = FRenderTargetBinding(SubsurfaceIntermediateTexture, SceneColorTextureLoadAction);
+			PassParameters->SubsurfaceInput0 = GetSubsurfaceInput(SceneColorTexture, SceneViewportParameters);
+			PassParameters->SubsurfaceSampler0 = BilinearBorderSampler;
+			PassParameters->Strata = Strata::BindStrataGlobalUniformParameters(View.StrataSceneData);
+
+			// Scattering output target is only used when scattering is enabled.
 			PassParameters->SubsurfaceInput1 = GetSubsurfaceInput(SubsurfaceSubpassTwoTex, SubsurfaceViewportParameters);
 			PassParameters->SubsurfaceSampler1 = BilinearBorderSampler;
+
+
+			const FSubsurfaceRecombinePS::EQuality RecombineQuality = FSubsurfaceRecombinePS::GetQuality(View);
+
+			FSubsurfaceRecombinePS::FPermutationDomain PixelShaderPermutationVector;
+			PixelShaderPermutationVector.Set<FSubsurfaceRecombinePS::FDimensionMode>(SubsurfaceMode);
+			PixelShaderPermutationVector.Set<FSubsurfaceRecombinePS::FDimensionQuality>(RecombineQuality);
+			PixelShaderPermutationVector.Set<FSubsurfaceRecombinePS::FDimensionCheckerboard>(bCheckerboard);
+			PixelShaderPermutationVector.Set<FSubsurfaceRecombinePS::FDimensionHalfRes>(bHalfRes);
+			PixelShaderPermutationVector.Set<FSubsurfaceRecombinePS::FRunningInSeparable>(bForceRunningInSeparable);
+
+			TShaderMapRef<FSubsurfaceRecombinePS> PixelShader(View.ShaderMap, PixelShaderPermutationVector);
+			TShaderMapRef<FSubsurfaceTilePassVS> VertexShader(View.ShaderMap);
+			/**
+			 * See the related comment above in the prepare pass. The scene viewport is used as both the target and
+			 * texture viewport in order to ensure that the correct pixel is sampled for checkerboard rendering.
+			 */
+			AddSubsurfaceTiledScreenPass(
+				GraphBuilder,
+				RDG_EVENT_NAME("SubsurfaceRecombine(%s %s%s%s%s) %dx%d",
+					GetEventName(PixelShaderPermutationVector.Get<FSubsurfaceRecombinePS::FDimensionMode>()),
+					FSubsurfaceRecombinePS::GetEventName(PixelShaderPermutationVector.Get<FSubsurfaceRecombinePS::FDimensionQuality>()),
+					PixelShaderPermutationVector.Get<FSubsurfaceRecombinePS::FDimensionCheckerboard>() ? TEXT(" Checkerboard") : TEXT(""),
+					PixelShaderPermutationVector.Get<FSubsurfaceRecombinePS::FDimensionHalfRes>() ? TEXT(" HalfRes") : TEXT(""),
+					PixelShaderPermutationVector.Get<FSubsurfaceRecombinePS::FRunningInSeparable>() ? TEXT(" RunningInSeparable") : TEXT(""),
+					View.ViewRect.Width(),
+					View.ViewRect.Height()),
+				View,
+				PassParameters,
+				VertexShader,
+				PixelShader,
+				TStaticBlendState<CW_RGBA, BO_Add, BF_One, BF_Zero, BO_Add, BF_One, BF_Zero>::GetRHI(),
+				SceneViewport,
+				TileType);
 		}
 
-		const FSubsurfaceRecombinePS::EQuality RecombineQuality = FSubsurfaceRecombinePS::GetQuality(View);
-
-		FSubsurfaceRecombinePS::FPermutationDomain PixelShaderPermutationVector;
-		PixelShaderPermutationVector.Set<FSubsurfaceRecombinePS::FDimensionMode>(SubsurfaceMode);
-		PixelShaderPermutationVector.Set<FSubsurfaceRecombinePS::FDimensionQuality>(RecombineQuality);
-		PixelShaderPermutationVector.Set<FSubsurfaceRecombinePS::FDimensionCheckerboard>(bCheckerboard);
-		PixelShaderPermutationVector.Set<FSubsurfaceRecombinePS::FDimensionHalfRes>(bHalfRes);
-		PixelShaderPermutationVector.Set<FSubsurfaceRecombinePS::FRunningInSeparable>(bForceRunningInSeparable);
-
-		TShaderMapRef<FSubsurfaceRecombinePS> PixelShader(View.ShaderMap, PixelShaderPermutationVector);
-
-		/**
-		 * See the related comment above in the prepare pass. The scene viewport is used as both the target and
-		 * texture viewport in order to ensure that the correct pixel is sampled for checkerboard rendering.
-		 */
-		AddDrawScreenPass(
-			GraphBuilder,
-			RDG_EVENT_NAME("SubsurfaceRecombine(%s %s%s%s%s) %dx%d",
-				GetEventName(PixelShaderPermutationVector.Get<FSubsurfaceRecombinePS::FDimensionMode>()),
-				FSubsurfaceRecombinePS::GetEventName(PixelShaderPermutationVector.Get<FSubsurfaceRecombinePS::FDimensionQuality>()),
-				PixelShaderPermutationVector.Get<FSubsurfaceRecombinePS::FDimensionCheckerboard>() ? TEXT(" Checkerboard") : TEXT(""),
-				PixelShaderPermutationVector.Get<FSubsurfaceRecombinePS::FDimensionHalfRes>() ? TEXT(" HalfRes") : TEXT(""),
-				PixelShaderPermutationVector.Get<FSubsurfaceRecombinePS::FRunningInSeparable>() ? TEXT(" RunningInSeparable") : TEXT(""),
-				View.ViewRect.Width(),
-				View.ViewRect.Height()),
-			View,
-			SceneViewport,
-			SceneViewport,
-			PixelShader,
-			PassParameters,
-			EScreenPassDrawFlags::AllowHMDHiddenAreaMask);
+		//Write back to the SceneColor texture 
+		{
+			FSubsurfaceTiles::ETileType TileType = FSubsurfaceTiles::ETileType::All;
+
+			FSubsurfaceRecombineCopyPS::FParameters* PassParameters = GraphBuilder.AllocParameters<FSubsurfaceRecombineCopyPS::FParameters>();
+			PassParameters->Subsurface = SubsurfaceCommonParameters;
+			PassParameters->TileParameters = GetSubsurfaceTileParameters(View, Tiles, TileType, SubsurfaceViewport.Extent);
+			PassParameters->RenderTargets[0] = FRenderTargetBinding(SceneColorTexture, SceneColorTextureLoadAction);
+			PassParameters->SubsurfaceInput0 = GetSubsurfaceInput(SubsurfaceIntermediateTexture, SceneViewportParameters);
+			PassParameters->SubsurfaceSampler0 = PointClampSampler;
+			PassParameters->Strata = Strata::BindStrataGlobalUniformParameters(View.StrataSceneData);
+
+			TShaderMapRef<FSubsurfaceRecombineCopyPS> PixelShader(View.ShaderMap);
+			TShaderMapRef<FSubsurfaceTilePassVS> VertexShader(View.ShaderMap);
+			
+			AddSubsurfaceTiledScreenPass(
+				GraphBuilder,
+				RDG_EVENT_NAME("SubsurfaceCopyToSceneColor %dx%d",
+					View.ViewRect.Width(),
+					View.ViewRect.Height()),
+				View, 
+				PassParameters, 
+				VertexShader, 
+				PixelShader,
+				TStaticBlendState<CW_RGBA, BO_Add, BF_SourceAlpha, BF_InverseSourceAlpha, BO_Add, BF_Zero, BF_Zero>::GetRHI(),
+				SceneViewport, 
+				TileType);
+		}
 	}
 
 	if (SubsurfaceMode != ESubsurfaceMode::Bypass && QualityHistoryState && !bForceRunningInSeparable)
@@ -1210,42 +1393,7 @@ FRDGTextureRef AddSubsurfacePass(
 	check(ViewMask);
 
 	FRDGTextureRef SceneColorTexture = SceneTextures.Color.Target;
-	ERenderTargetLoadAction SceneColorTextureLoadAction = ERenderTargetLoadAction::ENoAction;
-
-	const bool bHasNonSubsurfaceView = ViewMask != ViewMaskAll;
-
-	/**
-	 * Since we are outputting to a new texture and certain views may not utilize subsurface scattering,
-	 * we need to copy all non-subsurface views onto the destination texture.
-	 */
-	if (bHasNonSubsurfaceView)
-	{
-		RDG_EVENT_SCOPE(GraphBuilder, "SubsurfaceViewportCopy");
-
-		TShaderMapRef<FSubsurfaceViewportCopyPS> PixelShader(Views[0].ShaderMap);
-
-		for (uint32 ViewIndex = 0; ViewIndex < ViewCount; ++ViewIndex)
-		{
-			const uint32 ViewBit = 1 << ViewIndex;
-
-			const bool bIsNonSubsurfaceView = (ViewMask & ViewBit) == 0;
-
-			if (bIsNonSubsurfaceView)
-			{
-				auto* PassParameters = GraphBuilder.AllocParameters<FSubsurfaceViewportCopyPS::FParameters>();
-				PassParameters->RenderTargets[0] = FRenderTargetBinding(SceneColorTextureOutput, SceneColorTextureLoadAction);
-				PassParameters->SubsurfaceInput0_Texture = SceneColorTexture;
-				PassParameters->SubsurfaceSampler0 = TStaticSamplerState<>::GetRHI();
-
-				const FViewInfo& View = Views[ViewIndex];
-				const FScreenPassTextureViewport TextureViewport(SceneColorTexture, View.ViewRect);
-				AddDrawScreenPass(GraphBuilder, RDG_EVENT_NAME("View%d", ViewIndex), View, TextureViewport, TextureViewport, PixelShader, PassParameters);
-				SceneColorTextureLoadAction = ERenderTargetLoadAction::ELoad;
-			}
-		}
-
-		SceneColorTextureLoadAction = ERenderTargetLoadAction::ELoad;
-	}
+	ERenderTargetLoadAction SceneColorTextureLoadAction = ERenderTargetLoadAction::ELoad;
 
 	for (uint32 ViewIndex = 0; ViewIndex < ViewCount; ++ViewIndex)
 	{
@@ -1331,7 +1479,6 @@ void AddSubsurfacePass(
 
 	checkf(!SceneTextures.Color.IsSeparate(), TEXT("Subsurface rendering requires the deferred renderer."));
 
-	FRDGTextureRef SceneColorOutputTexture = GraphBuilder.CreateTexture(SceneTextures.Color.Target->Desc, TEXT("SceneColorSubsurface"));
-	AddSubsurfacePass(GraphBuilder, Views, ViewMask, SceneTextures, SceneColorOutputTexture);
-	SceneTextures.Color = SceneColorOutputTexture;
+	AddSubsurfacePass(GraphBuilder, Views, ViewMask, SceneTextures, SceneTextures.Color.Target);
+	SceneTextures.Color.Resolve = SceneTextures.Color.Target;
 }
\ No newline at end of file
diff --git a/Engine/Source/Runtime/Renderer/Private/SubsurfaceTiles.cpp b/Engine/Source/Runtime/Renderer/Private/SubsurfaceTiles.cpp
new file mode 100644
--- /dev/null
+++ b/Engine/Source/Runtime/Renderer/Private/SubsurfaceTiles.cpp
@@ -0,0 +1,146 @@
+// Copyright Epic Games, Inc. All Rights Reserved.
+
+#include "SubsurfaceTiles.h"
+#include "SceneRendering.h"
+
+
+
+bool FSubsurfaceTilePassVS::ShouldCompilePermutation(const FGlobalShaderPermutationParameters& Parameters)
+{
+	return IsFeatureLevelSupported(Parameters.Platform, ERHIFeatureLevel::SM5);
+}
+
+void FSubsurfaceTilePassVS::ModifyCompilationEnvironment(const FGlobalShaderPermutationParameters& Parameters, FShaderCompilerEnvironment& OutEnvironment)
+{
+	FGlobalShader::ModifyCompilationEnvironment(Parameters, OutEnvironment);
+	OutEnvironment.SetDefine(TEXT("SHADER_TILE_VS"), 1);
+	OutEnvironment.SetDefine(TEXT("SUBSURFACE_TILE_SIZE"), FSubsurfaceTiles::TileSize);
+}
+
+FSubsurfaceTilePassVS::FParameters GetSubsurfaceTileParameters(const FViewInfo& InView,const FSubsurfaceTiles& InTile, FSubsurfaceTiles::ETileType TileType, FIntPoint BufferExtent)
+{
+	FSubsurfaceTilePassVS::FParameters Out;
+	Out.TileType = uint32(TileType);
+	Out.bRectPrimitive = InTile.bRectPrimitive ? 1 : 0;
+	Out.ViewMin = InView.ViewRect.Min;
+	Out.ExtentInverse = FVector2D(1.f / BufferExtent.X, 1.f / BufferExtent.Y);
+	Out.TileDataBuffer = InTile.GetTileBufferSRV(TileType);
+	Out.TileIndirectBuffer = InTile.TileIndirectDrawBuffer;
+	return Out;
+}
+
+
+class FClearUAVBuildIndirectDispatchBufferCS : public FGlobalShader
+{
+public:
+	DECLARE_GLOBAL_SHADER(FClearUAVBuildIndirectDispatchBufferCS)
+	SHADER_USE_PARAMETER_STRUCT(FClearUAVBuildIndirectDispatchBufferCS, FGlobalShader)
+
+		BEGIN_SHADER_PARAMETER_STRUCT(FParameters, )
+		SHADER_PARAMETER(FIntPoint, ViewportSize)
+		SHADER_PARAMETER(uint32, Offset)
+		SHADER_PARAMETER_RDG_BUFFER_SRV(Buffer<uint32>, ConditionBuffer)
+		SHADER_PARAMETER_RDG_BUFFER_UAV(RWBuffer<uint32>, RWIndirectDispatchArgsBuffer)
+		END_SHADER_PARAMETER_STRUCT()
+
+		static bool ShouldCompilePermutation(const FGlobalShaderPermutationParameters& Parameters)
+	{
+		return RHISupportsComputeShaders(Parameters.Platform);
+	}
+
+	static void ModifyCompilationEnvironment(const FShaderPermutationParameters&, FShaderCompilerEnvironment& OutEnvironment)
+	{
+		OutEnvironment.SetDefine(TEXT("SUBSURFACE_COMPUTE"), 1);
+		OutEnvironment.SetDefine(TEXT("SUBSURFACE_TILE_SIZE"), FSubsurfaceTiles::TileSize);
+	}
+};
+
+IMPLEMENT_GLOBAL_SHADER(FClearUAVBuildIndirectDispatchBufferCS, "/Engine/Private/PostprocessSubsurfaceTile.usf", "BuildIndirectDispatchArgsCS", SF_Compute);
+
+class FClearUAVCS : public FGlobalShader
+{
+public:
+	DECLARE_GLOBAL_SHADER(FClearUAVCS)
+	SHADER_USE_PARAMETER_STRUCT(FClearUAVCS, FGlobalShader)
+
+		BEGIN_SHADER_PARAMETER_STRUCT(FParameters, )
+		SHADER_PARAMETER(FIntPoint, TextureExtent)
+		SHADER_PARAMETER(FIntPoint, ViewportMin)
+		RDG_BUFFER_ACCESS(IndirectDispatchArgsBuffer, ERHIAccess::IndirectArgs)
+		SHADER_PARAMETER_RDG_TEXTURE_UAV(RWTexture2D<float4>,TextureOutput)
+		END_SHADER_PARAMETER_STRUCT()
+
+		static bool ShouldCompilePermutation(const FGlobalShaderPermutationParameters& Parameters)
+	{
+		return RHISupportsComputeShaders(Parameters.Platform);
+	}
+
+	static void ModifyCompilationEnvironment(const FShaderPermutationParameters&, FShaderCompilerEnvironment& OutEnvironment)
+	{
+		OutEnvironment.SetDefine(TEXT("SUBSURFACE_COMPUTE"), 1);
+		OutEnvironment.SetDefine(TEXT("SUBSURFACE_TILE_SIZE"), FSubsurfaceTiles::TileSize);
+	}
+};
+
+IMPLEMENT_GLOBAL_SHADER(FClearUAVCS, "/Engine/Private/PostprocessSubsurfaceTile.usf", "ClearUAV", SF_Compute);
+
+
+void AddConditionalClearBlackUAVPass(FRDGBuilder& GraphBuilder, FRDGEventName&& Name, 
+	FRDGTextureRef Texture, const FScreenPassTextureViewport& ScreenPassViewport, FRDGBufferRef ConditionBuffer, uint32 Offset)
+{
+
+	FRDGBufferRef IndirectDispatchArgsBuffer = GraphBuilder.CreateBuffer(
+		FRDGBufferDesc::CreateIndirectDesc<FRHIDispatchIndirectParameters>(1),
+		TEXT("IndirectDispatchArgsBuffer"));
+
+	{
+		// build the indirect dispatch arguments buffer ( compute the group count on GPU conditionally)
+		FClearUAVBuildIndirectDispatchBufferCS::FParameters* PassParameters =
+			GraphBuilder.AllocParameters<FClearUAVBuildIndirectDispatchBufferCS::FParameters>();
+		PassParameters->ViewportSize = FIntPoint(ScreenPassViewport.Rect.Max.X - ScreenPassViewport.Rect.Min.X + 1,
+												ScreenPassViewport.Rect.Max.Y - ScreenPassViewport.Rect.Min.Y + 1);
+		PassParameters->Offset = Offset;
+		PassParameters->ConditionBuffer = 
+			GraphBuilder.CreateSRV(ConditionBuffer, EPixelFormat::PF_R32_UINT);
+		PassParameters->RWIndirectDispatchArgsBuffer = 
+			GraphBuilder.CreateUAV(IndirectDispatchArgsBuffer, EPixelFormat::PF_R32_UINT);
+
+		TShaderMapRef<FClearUAVBuildIndirectDispatchBufferCS> ComputeShader(GetGlobalShaderMap(GMaxRHIFeatureLevel));
+
+		FComputeShaderUtils::AddPass(
+			GraphBuilder,
+			RDG_EVENT_NAME("ClearUAV (BuildIndirectDispatchBuffer)"),
+			ComputeShader,
+			PassParameters,
+			FIntVector(1, 1, 1));
+	}
+
+	FClearUAVCS::FParameters* PassParameters = GraphBuilder.AllocParameters<FClearUAVCS::FParameters>();
+	PassParameters->TextureExtent = Texture->Desc.Extent;
+	PassParameters->ViewportMin = ScreenPassViewport.Rect.Min;
+	PassParameters->TextureOutput = GraphBuilder.CreateUAV(Texture);
+	PassParameters->IndirectDispatchArgsBuffer = IndirectDispatchArgsBuffer;
+
+	TShaderMapRef<FClearUAVCS> ComputeShader(GetGlobalShaderMap(GMaxRHIFeatureLevel));
+
+	FComputeShaderUtils::AddPass(
+		GraphBuilder,
+		Forward<FRDGEventName>(Name),
+		ComputeShader,
+		PassParameters,
+		IndirectDispatchArgsBuffer, 0);
+}
+
+IMPLEMENT_GLOBAL_SHADER(FSubsurfaceTilePassVS, "/Engine/Private/PostProcessSubsurfaceTile.usf", "MainVS", SF_Vertex);
+
+const TCHAR* ToString(FSubsurfaceTiles::ETileType Type)
+{
+	switch (Type)
+	{
+	case FSubsurfaceTiles::ETileType::All:			return TEXT("SSS(All)");
+	case FSubsurfaceTiles::ETileType::AFIS:		return TEXT("SSS(AFIS)");
+	case FSubsurfaceTiles::ETileType::SEPARABLE:		return TEXT("SSS(Separable)");
+	default:											return TEXT("Unknown");
+	}
+	return TEXT("Unknown");
+}
\ No newline at end of file
diff --git a/Engine/Source/Runtime/Renderer/Private/SubsurfaceTiles.h b/Engine/Source/Runtime/Renderer/Private/SubsurfaceTiles.h
new file mode 100644
--- /dev/null
+++ b/Engine/Source/Runtime/Renderer/Private/SubsurfaceTiles.h
@@ -0,0 +1,123 @@
+// Copyright Epic Games, Inc. All Rights Reserved.
+
+/*=============================================================================
+	SubsurfaceTiles.h: Screenspace subsurface scattering tile buffer implementation.
+=============================================================================*/
+
+#pragma once
+
+#include "CoreMinimal.h"
+#include "RendererInterface.h"
+#include "RenderGraphResources.h"
+#include "GlobalShader.h"
+#include "ShaderParameterStruct.h"
+#include "ScreenPass.h"
+
+// Adapted from FHairStrandsTile
+struct FSubsurfaceTiles
+{
+	// Support two types of acceleration. It is independent of SSS profile
+	// AFIS: adaptive filtered importance sampling
+	// SEPARABLE: two pass separable filter
+	enum class ETileType : uint8 {AFIS, SEPARABLE, All, Count};
+	static const uint32 TileTypeCount = uint32(ETileType::Count);
+
+	FIntPoint			BufferResolution = FIntPoint(0, 0);
+	static const uint32 GroupSize = 64;
+
+	// Define the size of subsurface tile. @TODO: Set to 16 to use LDS.
+	static const uint32	TileSize = 8;
+
+	static const uint32	TilePerThread_GroupSize = 64;
+	uint32				TileCount = 0;
+	FIntPoint			TileDimension = FIntPoint(0, 0);
+	bool				bRectPrimitive = false;
+
+	// Buffers per tile types
+	FRDGBufferSRVRef	TileDataSRV[TileTypeCount] = { nullptr, nullptr };
+	FRDGBufferRef		TileDataBuffer[TileTypeCount] = { nullptr, nullptr };
+
+	FRDGBufferSRVRef	TileTypeCountSRV = nullptr;
+	FRDGBufferRef		TileTypeCountBuffer = nullptr;
+
+	FRDGBufferRef		TileIndirectDispatchBuffer = nullptr;
+	FRDGBufferRef		TileIndirectDrawBuffer = nullptr;
+	FRDGBufferRef		TileIndirectRayDispatchBuffer = nullptr;
+
+	FRDGBufferRef		TilePerThreadIndirectDispatchBuffer = nullptr;
+
+	static FORCEINLINE uint32 GetIndirectDrawArgOffset(ETileType Type)		{ return uint32(Type) * sizeof(FRHIDrawIndirectParameters); }
+	static FORCEINLINE uint32 GetIndirectDispatchArgOffset(ETileType Type)	{ return uint32(Type) * sizeof(FRHIDispatchIndirectParameters); }
+
+	bool IsValid() const										{ return TileCount > 0 && TileDataBuffer[uint32(ETileType::All)];}
+	FRDGBufferRef GetTileBuffer(ETileType Type) const			{ const uint32 Index = uint32(Type); check(TileDataBuffer[Index] != nullptr); return TileDataBuffer[Index]; }
+	FRDGBufferSRVRef GetTileBufferSRV(ETileType Type) const		{ const uint32 Index = uint32(Type); check(TileDataSRV[Index] != nullptr); return TileDataSRV[Index]; }
+};
+
+FORCEINLINE uint32 ToIndex(FSubsurfaceTiles::ETileType Type) { return uint32(Type); }
+const TCHAR* ToString(FSubsurfaceTiles::ETileType Type);
+
+class FSubsurfaceTilePassVS : public FGlobalShader
+{
+	DECLARE_GLOBAL_SHADER(FSubsurfaceTilePassVS);
+	SHADER_USE_PARAMETER_STRUCT(FSubsurfaceTilePassVS, FGlobalShader);
+
+	BEGIN_SHADER_PARAMETER_STRUCT(FParameters, )
+		SHADER_PARAMETER(FIntPoint, ViewMin)
+		SHADER_PARAMETER(int32, bRectPrimitive)
+		SHADER_PARAMETER(uint32, TileType)
+		SHADER_PARAMETER(FVector2D, ExtentInverse)
+		SHADER_PARAMETER(uint32, OutputFlag)
+		SHADER_PARAMETER_RDG_BUFFER_SRV(Buffer, TileDataBuffer)
+		RDG_BUFFER_ACCESS(TileIndirectBuffer, ERHIAccess::IndirectArgs)
+		END_SHADER_PARAMETER_STRUCT()
+
+	static bool ShouldCompilePermutation(const FGlobalShaderPermutationParameters& Parameters);
+	static void ModifyCompilationEnvironment(const FGlobalShaderPermutationParameters& Parameters, FShaderCompilerEnvironment& OutEnvironment);
+};
+
+FSubsurfaceTilePassVS::FParameters GetSubsurfaceTileParameters(const FViewInfo& InView, const FSubsurfaceTiles& InTile, FSubsurfaceTiles::ETileType TileType, FIntPoint BufferExtent);
+
+template <typename FSubsurfacePassPS, typename FSubsurfacePassVS>
+void AddSubsurfaceTiledScreenPass(
+	FRDGBuilder& GraphBuilder,
+	FRDGEventName&& Name,
+	const FViewInfo& View,
+	typename FSubsurfacePassPS::FParameters* PassParameters,
+	TShaderMapRef<FSubsurfacePassVS>& VertexShader,
+	TShaderMapRef<FSubsurfacePassPS>& PixelShader,
+	FRHIBlendState* BlendState,
+	const FScreenPassTextureViewport SceneViewport,
+	FSubsurfaceTiles::ETileType TileType)
+{
+	GraphBuilder.AddPass(
+		Forward<FRDGEventName>(Name),
+		PassParameters,
+		ERDGPassFlags::Raster,
+		[PassParameters, VertexShader, PixelShader, BlendState, Viewport = SceneViewport, TileType](FRHICommandList& RHICmdList)
+		{
+			typename FSubsurfacePassVS::FParameters ParametersVS = PassParameters->TileParameters;
+
+			FGraphicsPipelineStateInitializer GraphicsPSOInit;
+			RHICmdList.ApplyCachedRenderTargets(GraphicsPSOInit);
+			GraphicsPSOInit.BlendState = BlendState;
+			GraphicsPSOInit.RasterizerState = TStaticRasterizerState<>::GetRHI();
+			GraphicsPSOInit.DepthStencilState = TStaticDepthStencilState<false, CF_Never>::GetRHI();
+
+			GraphicsPSOInit.BoundShaderState.VertexDeclarationRHI = GFilterVertexDeclaration.VertexDeclarationRHI;
+			GraphicsPSOInit.BoundShaderState.VertexShaderRHI = VertexShader.GetVertexShader();
+			GraphicsPSOInit.BoundShaderState.PixelShaderRHI = PixelShader.GetPixelShader();
+			GraphicsPSOInit.PrimitiveType = PassParameters->TileParameters.bRectPrimitive > 0 ? PT_RectList : PT_TriangleList;
+			SetGraphicsPipelineState(RHICmdList, GraphicsPSOInit, 0);
+			SetShaderParameters(RHICmdList, VertexShader, VertexShader.GetVertexShader(), ParametersVS);
+			SetShaderParameters(RHICmdList, PixelShader, PixelShader.GetPixelShader(), *PassParameters);
+
+			RHICmdList.SetViewport(0.0f, 0.0f, 0.0f, Viewport.Extent.X, Viewport.Extent.Y, 1.0f);
+			RHICmdList.SetStreamSource(0, nullptr, 0);
+			RHICmdList.DrawPrimitiveIndirect(PassParameters->TileParameters.TileIndirectBuffer->GetRHI(), FSubsurfaceTiles::GetIndirectDrawArgOffset(TileType));
+		});
+}
+
+/** Clear the UAV texture to black only when ConditionBuffer[Offset] > 0 */
+void AddConditionalClearBlackUAVPass(FRDGBuilder& GraphBuilder, FRDGEventName&& PassName, 
+	FRDGTextureRef Texture, const FScreenPassTextureViewport& ScreenPassViewport, FRDGBufferRef ConditionBuffer, uint32 Offset);
