Add custom intersection support for vertex factory

Originally done by Chris for path tracing, ported for raytracing

diff --git a/Engine/Config/Windows/DataDrivenPlatformInfo.ini b/Engine/Config/Windows/DataDrivenPlatformInfo.ini
--- a/Engine/Config/Windows/DataDrivenPlatformInfo.ini
+++ b/Engine/Config/Windows/DataDrivenPlatformInfo.ini
@@ -59,6 +59,7 @@ Language=D3D
 MaxFeatureLevel=SM5
 bIsPC=true
 bSupportsRayTracing=true
+bSupportsRayTracingProceduralPrimitive=true
 bSupportsRayTracingIndirectInstanceData=true
 bSupportsPathTracing=true
 bSupportsHighEndRayTracingReflections=true
@@ -88,6 +89,7 @@ Language=D3D
 MaxFeatureLevel=SM6
 bIsPC=true
 bSupportsRayTracing=true
+bSupportsRayTracingProceduralPrimitive=true
 bSupportsRayTracingIndirectInstanceData=true
 bSupportsPathTracing=true
 bSupportsHighEndRayTracingReflections=true
diff --git a/Engine/Shaders/Private/RayTracing/RayTracingCommon.ush b/Engine/Shaders/Private/RayTracing/RayTracingCommon.ush
--- a/Engine/Shaders/Private/RayTracing/RayTracingCommon.ush
+++ b/Engine/Shaders/Private/RayTracing/RayTracingCommon.ush
@@ -72,6 +72,11 @@
 [shader("raygeneration")] void name()
 #endif // RAY_TRACING_ENTRY_RAYGEN
 
+#ifndef RAY_TRACING_ENTRY_INTERSECTION
+#define RAY_TRACING_ENTRY_INTERSECTION(name)\
+[shader("intersection")] void name()
+#endif //RAY_TRACING_ENTRY_INTERSECTION
+
 #ifndef RAY_TRACING_ENTRY_CLOSEST_HIT
 #define RAY_TRACING_ENTRY_CLOSEST_HIT(name, payload_type, payload_name, attributes_type, attributes_name)\
 [shader("closesthit")] void name(inout payload_type payload_name, in attributes_type attributes_name)
diff --git a/Engine/Shaders/Private/RayTracing/RayTracingMaterialHitShaders.usf b/Engine/Shaders/Private/RayTracing/RayTracingMaterialHitShaders.usf
--- a/Engine/Shaders/Private/RayTracing/RayTracingMaterialHitShaders.usf
+++ b/Engine/Shaders/Private/RayTracing/RayTracingMaterialHitShaders.usf
@@ -272,7 +272,12 @@ RAY_TRACING_ENTRY_CLOSEST_HIT(MaterialCHS,
 	FRayCone PropagatedCone = PropagateRayCone(PackedPayload.GetRayCone(), 0 /* surface curvature */, RayTCurrent());
 	CalcInterpolants(PropagatedCone, Attributes, Interpolants);
 
+#if VF_SUPPORTS_RAYTRACING_PREPARE_MATERIAL_PIXEL_PARAMETERS
+	float3 GeoNormal = 0;
+	FMaterialPixelParameters MaterialParameters = GetMaterialPixelParameters(WorldRayOrigin(), WorldRayDirection(), RayTCurrent(), PrimitiveIndex(), Attributes, HitKind(), GeoNormal);
+#else
 	FMaterialPixelParameters MaterialParameters = GetMaterialPixelParameters(Interpolants, SvPosition);
+#endif
 	FPixelMaterialInputs PixelMaterialInputs;
 
 	{
@@ -561,7 +566,12 @@ RAY_TRACING_ENTRY_ANY_HIT(MaterialAHS,
 
 	CalcInterpolants(PackedPayload.GetRayCone(), Attributes, Interpolants);
 
+#if VF_SUPPORTS_RAYTRACING_PREPARE_MATERIAL_PIXEL_PARAMETERS
+	float3 GeoNormal = 0;
+	FMaterialPixelParameters MaterialParameters = GetMaterialPixelParameters(WorldRayOrigin(), WorldRayDirection(), RayTCurrent(), PrimitiveIndex(), Attributes, HitKind(), GeoNormal);
+#else
 	FMaterialPixelParameters MaterialParameters = GetMaterialPixelParameters(Interpolants, SvPosition);
+#endif
 	FPixelMaterialInputs PixelMaterialInputs;
 
 	{
diff --git a/Engine/Source/Runtime/RHI/Private/RHI.cpp b/Engine/Source/Runtime/RHI/Private/RHI.cpp
--- a/Engine/Source/Runtime/RHI/Private/RHI.cpp
+++ b/Engine/Source/Runtime/RHI/Private/RHI.cpp
@@ -2232,6 +2232,7 @@ void FGenericDataDrivenShaderPlatformInfo::ParseDataDrivenShaderInfo(const FConf
 	GET_SECTION_BOOL_HELPER(bSupports4ComponentUAVReadWrite);
 	GET_SECTION_BOOL_HELPER(bSupportsRenderTargetWriteMask);
 	GET_SECTION_BOOL_HELPER(bSupportsRayTracing);
+	GET_SECTION_BOOL_HELPER(bSupportsRayTracingProceduralPrimitive);
 	GET_SECTION_BOOL_HELPER(bSupportsRayTracingIndirectInstanceData);
 	GET_SECTION_BOOL_HELPER(bSupportsPathTracing);
 	GET_SECTION_BOOL_HELPER(bSupportsHighEndRayTracingReflections);
diff --git a/Engine/Source/Runtime/RHI/Public/RHIDefinitions.h b/Engine/Source/Runtime/RHI/Public/RHIDefinitions.h
--- a/Engine/Source/Runtime/RHI/Public/RHIDefinitions.h
+++ b/Engine/Source/Runtime/RHI/Public/RHIDefinitions.h
@@ -326,6 +326,7 @@ class RHI_API FGenericDataDrivenShaderPlatformInfo
 	uint32 bSupports4ComponentUAVReadWrite: 1;
 	uint32 bSupportsRenderTargetWriteMask: 1;
 	uint32 bSupportsRayTracing: 1;
+	uint32 bSupportsRayTracingProceduralPrimitive : 1;
 	uint32 bSupportsRayTracingIndirectInstanceData : 1; // Whether instance transforms can be copied from the GPU to the TLAS instances buffer
 	uint32 bSupportsHighEndRayTracingReflections : 1; // Whether fully-featured RT reflections can be used on the platform (with multi-bounce, translucency, etc.)
 	uint32 bSupportsPathTracing : 1; // Whether real-time path tracer is supported on this platform (avoids compiling unnecessary shaders)
@@ -528,6 +529,11 @@ public:
 		return Infos[Platform].bSupportsRayTracing;
 	}
 
+	static FORCEINLINE_DEBUGGABLE const bool GetSupportsRayTracingProceduralPrimitive(const FStaticShaderPlatform Platform)
+	{
+		return Infos[Platform].bSupportsRayTracingProceduralPrimitive;
+	}
+
 	static FORCEINLINE_DEBUGGABLE const bool GetSupportsRayTracingIndirectInstanceData(const FStaticShaderPlatform Platform)
 	{
 		return Infos[Platform].bSupportsRayTracingIndirectInstanceData;
diff --git a/Engine/Source/Runtime/RenderCore/Public/VertexFactory.h b/Engine/Source/Runtime/RenderCore/Public/VertexFactory.h
--- a/Engine/Source/Runtime/RenderCore/Public/VertexFactory.h
+++ b/Engine/Source/Runtime/RenderCore/Public/VertexFactory.h
@@ -91,17 +91,18 @@ enum class EVertexInputStreamType : uint8
 
 enum class EVertexFactoryFlags : uint32
 {
-	None                              = 0u,
-	UsedWithMaterials                 = 1u << 1,
-	SupportsStaticLighting            = 1u << 2,
-	SupportsDynamicLighting           = 1u << 3,
-	SupportsPrecisePrevWorldPos       = 1u << 4,
-	SupportsPositionOnly              = 1u << 5,
-	SupportsCachingMeshDrawCommands   = 1u << 6,
-	SupportsPrimitiveIdStream         = 1u << 7,
-	SupportsNaniteRendering           = 1u << 8,
-	SupportsRayTracing                = 1u << 9,
-	SupportsRayTracingDynamicGeometry = 1u << 10,
+	None                                  = 0u,
+	UsedWithMaterials                     = 1u << 1,
+	SupportsStaticLighting                = 1u << 2,
+	SupportsDynamicLighting               = 1u << 3,
+	SupportsPrecisePrevWorldPos           = 1u << 4,
+	SupportsPositionOnly                  = 1u << 5,
+	SupportsCachingMeshDrawCommands       = 1u << 6,
+	SupportsPrimitiveIdStream             = 1u << 7,
+	SupportsNaniteRendering               = 1u << 8,
+	SupportsRayTracing                    = 1u << 9,
+	SupportsRayTracingDynamicGeometry     = 1u << 10,
+	SupportsRayTracingProceduralPrimitive = 1u << 11,
 };
 ENUM_CLASS_FLAGS(EVertexFactoryFlags);
 
@@ -393,16 +394,17 @@ public:
 		return EnumHasAllFlags(Flags, FlagsToCheck);
 	}
 
-	bool IsUsedWithMaterials() const                { return HasFlags(EVertexFactoryFlags::UsedWithMaterials); }
-	bool SupportsStaticLighting() const             { return HasFlags(EVertexFactoryFlags::SupportsStaticLighting); }
-	bool SupportsDynamicLighting() const            { return HasFlags(EVertexFactoryFlags::SupportsDynamicLighting); }
-	bool SupportsPrecisePrevWorldPos() const        { return HasFlags(EVertexFactoryFlags::SupportsPrecisePrevWorldPos); }
-	bool SupportsPositionOnly() const               { return HasFlags(EVertexFactoryFlags::SupportsPositionOnly); }
-	bool SupportsCachingMeshDrawCommands() const    { return HasFlags(EVertexFactoryFlags::SupportsCachingMeshDrawCommands); }
-	bool SupportsPrimitiveIdStream() const          { return HasFlags(EVertexFactoryFlags::SupportsPrimitiveIdStream); }
-	bool SupportsNaniteRendering() const            { return HasFlags(EVertexFactoryFlags::SupportsNaniteRendering); }
-	bool SupportsRayTracing() const                 { return HasFlags(EVertexFactoryFlags::SupportsRayTracing); }
-	bool SupportsRayTracingDynamicGeometry() const  { return HasFlags(EVertexFactoryFlags::SupportsRayTracingDynamicGeometry); }
+	bool IsUsedWithMaterials() const                   { return HasFlags(EVertexFactoryFlags::UsedWithMaterials); }
+	bool SupportsStaticLighting() const                { return HasFlags(EVertexFactoryFlags::SupportsStaticLighting); }
+	bool SupportsDynamicLighting() const               { return HasFlags(EVertexFactoryFlags::SupportsDynamicLighting); }
+	bool SupportsPrecisePrevWorldPos() const           { return HasFlags(EVertexFactoryFlags::SupportsPrecisePrevWorldPos); }
+	bool SupportsPositionOnly() const                  { return HasFlags(EVertexFactoryFlags::SupportsPositionOnly); }
+	bool SupportsCachingMeshDrawCommands() const       { return HasFlags(EVertexFactoryFlags::SupportsCachingMeshDrawCommands); }
+	bool SupportsPrimitiveIdStream() const             { return HasFlags(EVertexFactoryFlags::SupportsPrimitiveIdStream); }
+	bool SupportsNaniteRendering() const               { return HasFlags(EVertexFactoryFlags::SupportsNaniteRendering); }
+	bool SupportsRayTracing() const                    { return HasFlags(EVertexFactoryFlags::SupportsRayTracing); }
+	bool SupportsRayTracingDynamicGeometry() const     { return HasFlags(EVertexFactoryFlags::SupportsRayTracingDynamicGeometry); }
+	bool SupportsRayTracingProceduralPrimitive() const { return HasFlags(EVertexFactoryFlags::SupportsRayTracingProceduralPrimitive); }
 
 	// Hash function.
 	friend uint32 GetTypeHash(const FVertexFactoryType* Type)
diff --git a/Engine/Source/Runtime/Renderer/Private/PathTracing.cpp b/Engine/Source/Runtime/Renderer/Private/PathTracing.cpp
--- a/Engine/Source/Runtime/Renderer/Private/PathTracing.cpp
+++ b/Engine/Source/Runtime/Renderer/Private/PathTracing.cpp
@@ -603,7 +603,7 @@ class FPathTracingIESAtlasCS : public FGlobalShader
 };
 IMPLEMENT_SHADER_TYPE(, FPathTracingIESAtlasCS, TEXT("/Engine/Private/PathTracing/PathTracingIESAtlas.usf"), TEXT("PathTracingIESAtlasCS"), SF_Compute);
 
-template<bool UseAnyHitShader>
+template<bool UseAnyHitShader, bool UseIntersectionShader>
 class TPathTracingMaterial : public FMeshMaterialShader
 {
 	DECLARE_SHADER_TYPE(TPathTracingMaterial, MeshMaterial);
@@ -616,7 +616,9 @@ public:
 
 	static bool ShouldCompilePermutation(const FMeshMaterialShaderPermutationParameters& Parameters)
 	{
+		const bool bUseProceduralPrimitive = Parameters.VertexFactoryType->SupportsRayTracingProceduralPrimitive() && FDataDrivenShaderPlatformInfo::GetSupportsRayTracingProceduralPrimitive(Parameters.Platform);
 		return Parameters.VertexFactoryType->SupportsRayTracing()
+			&& (UseIntersectionShader == bUseProceduralPrimitive)
 			&& ((Parameters.MaterialParameters.bIsMasked || Parameters.MaterialParameters.BlendMode != BLEND_Opaque) == UseAnyHitShader)
 			&& ShouldCompilePathTracingShadersForProject(Parameters.Platform);
 	}
@@ -625,6 +627,7 @@ public:
 	{
 		OutEnvironment.SetDefine(TEXT("USE_MATERIAL_CLOSEST_HIT_SHADER"), 1);
 		OutEnvironment.SetDefine(TEXT("USE_MATERIAL_ANY_HIT_SHADER"), UseAnyHitShader ? 1 : 0);
+		OutEnvironment.SetDefine(TEXT("USE_MATERIAL_INTERSECTION_SHADER"), UseIntersectionShader ? 1 : 0);
 		OutEnvironment.SetDefine(TEXT("USE_RAYTRACED_TEXTURE_RAYCONE_LOD"), 0);
 		OutEnvironment.SetDefine(TEXT("SCENE_TEXTURES_DISABLED"), 1);
 		OutEnvironment.SetDefine(TEXT("SIMPLIFIED_MATERIAL_SHADER"), 0); // TODO: expose a permutation so we can unify with GPULightmass?
@@ -654,12 +657,17 @@ public:
 	}
 };
 
-using FPathTracingMaterialCHS     = TPathTracingMaterial<false>;
-using FPathTracingMaterialCHS_AHS = TPathTracingMaterial<true>;
+using FPathTracingMaterialCHS        = TPathTracingMaterial<false, false>;
+using FPathTracingMaterialCHS_AHS    = TPathTracingMaterial<true , false>;
+using FPathTracingMaterialCHS_IS     = TPathTracingMaterial<false, true >;
+using FPathTracingMaterialCHS_AHS_IS = TPathTracingMaterial<true , true >;
 
 
-IMPLEMENT_MATERIAL_SHADER_TYPE(template <>, FPathTracingMaterialCHS    , TEXT("/Engine/Private/PathTracing/PathTracingMaterialHitShader.usf"), TEXT("closesthit=PathTracingMaterialCHS"), SF_RayHitGroup);
-IMPLEMENT_MATERIAL_SHADER_TYPE(template <>, FPathTracingMaterialCHS_AHS, TEXT("/Engine/Private/PathTracing/PathTracingMaterialHitShader.usf"), TEXT("closesthit=PathTracingMaterialCHS anyhit=PathTracingMaterialAHS"), SF_RayHitGroup);
+IMPLEMENT_MATERIAL_SHADER_TYPE(template <>, FPathTracingMaterialCHS       , TEXT("/Engine/Private/PathTracing/PathTracingMaterialHitShader.usf"), TEXT("closesthit=PathTracingMaterialCHS"), SF_RayHitGroup);
+IMPLEMENT_MATERIAL_SHADER_TYPE(template <>, FPathTracingMaterialCHS_AHS   , TEXT("/Engine/Private/PathTracing/PathTracingMaterialHitShader.usf"), TEXT("closesthit=PathTracingMaterialCHS anyhit=PathTracingMaterialAHS"), SF_RayHitGroup);
+IMPLEMENT_MATERIAL_SHADER_TYPE(template <>, FPathTracingMaterialCHS_IS    , TEXT("/Engine/Private/PathTracing/PathTracingMaterialHitShader.usf"), TEXT("closesthit=PathTracingMaterialCHS intersection=MaterialIS"), SF_RayHitGroup);
+IMPLEMENT_MATERIAL_SHADER_TYPE(template <>, FPathTracingMaterialCHS_AHS_IS, TEXT("/Engine/Private/PathTracing/PathTracingMaterialHitShader.usf"), TEXT("closesthit=PathTracingMaterialCHS anyhit=PathTracingMaterialAHS intersection=MaterialIS"), SF_RayHitGroup);
+
 
 
 bool FRayTracingMeshProcessor::ProcessPathTracing(
@@ -680,13 +688,20 @@ bool FRayTracingMeshProcessor::ProcessPathTracing(
 
 	FMaterialShaderTypes ShaderTypes;
 
+	const bool bUseProceduralPrimitive = VertexFactory->GetType()->SupportsRayTracingProceduralPrimitive() && FDataDrivenShaderPlatformInfo::GetSupportsRayTracingProceduralPrimitive(GMaxRHIShaderPlatform);
 	if (MaterialResource.IsMasked() || MaterialResource.GetBlendMode() != BLEND_Opaque)
 	{
-		ShaderTypes.AddShaderType<FPathTracingMaterialCHS_AHS>();
+		if (bUseProceduralPrimitive)
+			ShaderTypes.AddShaderType<FPathTracingMaterialCHS_AHS_IS>();
+		else
+			ShaderTypes.AddShaderType<FPathTracingMaterialCHS_AHS>();
 	}
 	else
 	{
-		ShaderTypes.AddShaderType<FPathTracingMaterialCHS>();
+		if (bUseProceduralPrimitive)
+			ShaderTypes.AddShaderType<FPathTracingMaterialCHS_IS>();
+		else
+			ShaderTypes.AddShaderType<FPathTracingMaterialCHS>();
 	}
 
 	FMaterialShaders Shaders;
diff --git a/Engine/Source/Runtime/Renderer/Private/RayTracing/RayTracingMaterialHitShaders.cpp b/Engine/Source/Runtime/Renderer/Private/RayTracing/RayTracingMaterialHitShaders.cpp
--- a/Engine/Source/Runtime/Renderer/Private/RayTracing/RayTracingMaterialHitShaders.cpp
+++ b/Engine/Source/Runtime/Renderer/Private/RayTracing/RayTracingMaterialHitShaders.cpp
@@ -105,7 +105,7 @@ public:
 	}
 };
 
-template<typename LightMapPolicyType, bool UseAnyHitShader, bool UseRayConeTextureLod>
+template<typename LightMapPolicyType, bool UseAnyHitShader, bool UseIntersectionShader, bool UseRayConeTextureLod>
 class TMaterialCHS : public FMaterialCHS
 {
 	DECLARE_SHADER_TYPE(TMaterialCHS, MeshMaterial);
@@ -125,12 +125,14 @@ public:
 		}
 
 		const bool bWantAnyHitShader = (GCompileRayTracingMaterialAHS && (Parameters.MaterialParameters.bIsMasked || Parameters.MaterialParameters.BlendMode == BLEND_Translucent));
+		const bool bSupportProceduralPrimitive = Parameters.VertexFactoryType->SupportsRayTracingProceduralPrimitive() && FDataDrivenShaderPlatformInfo::GetSupportsRayTracingProceduralPrimitive(Parameters.Platform);
 
 		return IsSupportedVertexFactoryType(Parameters.VertexFactoryType)
 			&& (bWantAnyHitShader == UseAnyHitShader)
 			&& LightMapPolicyType::ShouldCompilePermutation(Parameters)
 			&& ShouldCompileRayTracingShadersForProject(Parameters.Platform)
-			&& (bool)GRayTracingUseTextureLod == UseRayConeTextureLod;
+			&& (bool)GRayTracingUseTextureLod == UseRayConeTextureLod
+			&& (UseIntersectionShader == bSupportProceduralPrimitive);
 	}
 
 	static void ModifyCompilationEnvironment(const FMaterialShaderPermutationParameters& Parameters, FShaderCompilerEnvironment& OutEnvironment)
@@ -139,6 +141,7 @@ public:
 
 		OutEnvironment.SetDefine(TEXT("USE_MATERIAL_CLOSEST_HIT_SHADER"), GCompileRayTracingMaterialCHS ? 1 : 0);
 		OutEnvironment.SetDefine(TEXT("USE_MATERIAL_ANY_HIT_SHADER"), GCompileRayTracingMaterialAHS ? 1 : 0);
+		OutEnvironment.SetDefine(TEXT("USE_MATERIAL_INTERSECTION_SHADER"), UseIntersectionShader ? 1 : 0);
 		OutEnvironment.SetDefine(TEXT("USE_RAYTRACED_TEXTURE_RAYCONE_LOD"), UseRayConeTextureLod ? 1 : 0);
 		OutEnvironment.SetDefine(TEXT("SCENE_TEXTURES_DISABLED"), 1);
 		LightMapPolicyType::ModifyCompilationEnvironment(Parameters, OutEnvironment);
@@ -201,14 +204,22 @@ public:
 IMPLEMENT_MATERIAL_SHADER_TYPE(, FTrivialMaterialCHS, TEXT("/Engine/Private/RayTracing/RayTracingMaterialDefaultHitShaders.usf"), TEXT("closesthit=OpaqueShadowCHS"), SF_RayHitGroup);
 
 #define IMPLEMENT_MATERIALCHS_TYPE(LightMapPolicyType, LightMapPolicyName, AnyHitShaderName) \
-	typedef TMaterialCHS<LightMapPolicyType, false, false> TMaterialCHS##LightMapPolicyName; \
+	typedef TMaterialCHS<LightMapPolicyType, false, false, false> TMaterialCHS##LightMapPolicyName; \
 	IMPLEMENT_MATERIAL_SHADER_TYPE(template<>, TMaterialCHS##LightMapPolicyName, TEXT("/Engine/Private/RayTracing/RayTracingMaterialHitShaders.usf"), TEXT("closesthit=MaterialCHS"), SF_RayHitGroup); \
-	typedef TMaterialCHS<LightMapPolicyType, true, false> TMaterialCHS##LightMapPolicyName##AnyHitShaderName; \
+	typedef TMaterialCHS<LightMapPolicyType, true, false, false> TMaterialCHS##LightMapPolicyName##AnyHitShaderName; \
 	IMPLEMENT_MATERIAL_SHADER_TYPE(template<>, TMaterialCHS##LightMapPolicyName##AnyHitShaderName, TEXT("/Engine/Private/RayTracing/RayTracingMaterialHitShaders.usf"), TEXT("closesthit=MaterialCHS anyhit=MaterialAHS"), SF_RayHitGroup) \
-	typedef TMaterialCHS<LightMapPolicyType, false, true> TMaterialCHSLod##LightMapPolicyName; \
+	typedef TMaterialCHS<LightMapPolicyType, false, false, true> TMaterialCHSLod##LightMapPolicyName; \
 	IMPLEMENT_MATERIAL_SHADER_TYPE(template<>, TMaterialCHSLod##LightMapPolicyName, TEXT("/Engine/Private/RayTracing/RayTracingMaterialHitShaders.usf"), TEXT("closesthit=MaterialCHS"), SF_RayHitGroup); \
-	typedef TMaterialCHS<LightMapPolicyType, true, true> TMaterialCHSLod##LightMapPolicyName##AnyHitShaderName; \
-	IMPLEMENT_MATERIAL_SHADER_TYPE(template<>, TMaterialCHSLod##LightMapPolicyName##AnyHitShaderName, TEXT("/Engine/Private/RayTracing/RayTracingMaterialHitShaders.usf"), TEXT("closesthit=MaterialCHS anyhit=MaterialAHS"), SF_RayHitGroup);
+	typedef TMaterialCHS<LightMapPolicyType, true, false, true> TMaterialCHSLod##LightMapPolicyName##AnyHitShaderName; \
+	IMPLEMENT_MATERIAL_SHADER_TYPE(template<>, TMaterialCHSLod##LightMapPolicyName##AnyHitShaderName, TEXT("/Engine/Private/RayTracing/RayTracingMaterialHitShaders.usf"), TEXT("closesthit=MaterialCHS anyhit=MaterialAHS"), SF_RayHitGroup); \
+	typedef TMaterialCHS<LightMapPolicyType, false, true, false> TMaterialCHS_IS_##LightMapPolicyName; \
+	IMPLEMENT_MATERIAL_SHADER_TYPE(template<>, TMaterialCHS_IS_##LightMapPolicyName, TEXT("/Engine/Private/RayTracing/RayTracingMaterialHitShaders.usf"), TEXT("closesthit=MaterialCHS intersection=MaterialIS"), SF_RayHitGroup); \
+	typedef TMaterialCHS<LightMapPolicyType, true, true, false> TMaterialCHS_IS_##LightMapPolicyName##AnyHitShaderName; \
+	IMPLEMENT_MATERIAL_SHADER_TYPE(template<>, TMaterialCHS_IS_##LightMapPolicyName##AnyHitShaderName, TEXT("/Engine/Private/RayTracing/RayTracingMaterialHitShaders.usf"), TEXT("closesthit=MaterialCHS anyhit=MaterialAHS intersection=MaterialIS"), SF_RayHitGroup) \
+	typedef TMaterialCHS<LightMapPolicyType, false, true, true> TMaterialCHS_IS_Lod##LightMapPolicyName; \
+	IMPLEMENT_MATERIAL_SHADER_TYPE(template<>, TMaterialCHS_IS_Lod##LightMapPolicyName, TEXT("/Engine/Private/RayTracing/RayTracingMaterialHitShaders.usf"), TEXT("closesthit=MaterialCHS intersection=MaterialIS"), SF_RayHitGroup); \
+	typedef TMaterialCHS<LightMapPolicyType, true, true, true> TMaterialCHS_IS_Lod##LightMapPolicyName##AnyHitShaderName; \
+	IMPLEMENT_MATERIAL_SHADER_TYPE(template<>, TMaterialCHS_IS_Lod##LightMapPolicyName##AnyHitShaderName, TEXT("/Engine/Private/RayTracing/RayTracingMaterialHitShaders.usf"), TEXT("closesthit=MaterialCHS anyhit=MaterialAHS intersection=MaterialIS"), SF_RayHitGroup);
 
 IMPLEMENT_MATERIALCHS_TYPE(TUniformLightMapPolicy<LMP_NO_LIGHTMAP>, FNoLightMapPolicy, FAnyHitShader);
 IMPLEMENT_MATERIALCHS_TYPE(TUniformLightMapPolicy<LMP_PRECOMPUTED_IRRADIANCE_VOLUME_INDIRECT_LIGHTING>, FPrecomputedVolumetricLightmapLightingPolicy, FAnyHitShader);
@@ -219,35 +230,60 @@ IMPLEMENT_MATERIALCHS_TYPE(TUniformLightMapPolicy<LMP_DISTANCE_FIELD_SHADOWS_AND
 IMPLEMENT_GLOBAL_SHADER(FHiddenMaterialHitGroup, "/Engine/Private/RayTracing/RayTracingMaterialDefaultHitShaders.usf", "closesthit=HiddenMaterialCHS anyhit=HiddenMaterialAHS", SF_RayHitGroup);
 IMPLEMENT_GLOBAL_SHADER(FOpaqueShadowHitGroup, "/Engine/Private/RayTracing/RayTracingMaterialDefaultHitShaders.usf", "closesthit=OpaqueShadowCHS", SF_RayHitGroup);
 
-template<typename LightMapPolicyType>
-static bool GetMaterialHitShader(const FMaterial& RESTRICT MaterialResource, const FVertexFactory* VertexFactory, bool UseTextureLod, TShaderRef<FMaterialCHS>& OutShader)
+// Select TextureLOD
+template<typename LightMapPolicyType, bool bUseAnyHitShader, bool bUseIntersectionShader>
+inline void GetMaterialHitShader_TextureLOD(FMaterialShaderTypes& ShaderTypes, bool bUseTextureLod)
 {
-	const bool bMaterialsCompiled = GCompileRayTracingMaterialAHS || GCompileRayTracingMaterialCHS;
-	checkf(bMaterialsCompiled, TEXT(""));
+	if (bUseTextureLod)
+	{
+		ShaderTypes.AddShaderType<TMaterialCHS<LightMapPolicyType, bUseAnyHitShader, bUseIntersectionShader, true>>();
+	}
+	else
+	{
+		ShaderTypes.AddShaderType<TMaterialCHS<LightMapPolicyType, bUseAnyHitShader, bUseIntersectionShader, false>>();
+	}
+}
 
-	FMaterialShaderTypes ShaderTypes;
-	if ((MaterialResource.IsMasked() || MaterialResource.GetBlendMode() == BLEND_Translucent) && GCompileRayTracingMaterialAHS)
+// Select Intersection shader
+template<typename LightMapPolicyType, bool bUseAnyHitShader>
+inline void GetMaterialHitShader_Intersection_TextureLOD(FMaterialShaderTypes& ShaderTypes, bool bUseIntersectionShader, bool bUseTextureLod)
+{
+	if (bUseIntersectionShader)
 	{
-		if(UseTextureLod)
-		{ 
-			ShaderTypes.AddShaderType<TMaterialCHS<LightMapPolicyType, true, true>>();
-		}
-		else
-		{
-			ShaderTypes.AddShaderType<TMaterialCHS<LightMapPolicyType, true, false>>();
-		}
+		GetMaterialHitShader_TextureLOD<LightMapPolicyType, bUseAnyHitShader, true>(ShaderTypes, bUseTextureLod);
 	}
 	else
 	{
-		if (UseTextureLod)
-		{
-			ShaderTypes.AddShaderType<TMaterialCHS<LightMapPolicyType, false, true>>();
-		}
-		else
-		{
-			ShaderTypes.AddShaderType<TMaterialCHS<LightMapPolicyType, false, false>>();
-		}
+		GetMaterialHitShader_TextureLOD<LightMapPolicyType, bUseAnyHitShader, false>(ShaderTypes, bUseTextureLod);
 	}
+}
+
+// Select AnyHit shader
+template<typename LightMapPolicyType>
+inline void GetMaterialHitShader_AnyHit_Intersection_TextureLOD(FMaterialShaderTypes& ShaderTypes, bool bUseAnyHitShader, bool bUseIntersectionShader, bool bUseTextureLod)
+{
+	if (bUseAnyHitShader)
+	{
+		GetMaterialHitShader_Intersection_TextureLOD<LightMapPolicyType, true>(ShaderTypes, bUseIntersectionShader, bUseTextureLod);
+	}
+	else
+	{
+		GetMaterialHitShader_Intersection_TextureLOD<LightMapPolicyType, false>(ShaderTypes, bUseIntersectionShader, bUseTextureLod);
+	}
+}
+
+template<typename LightMapPolicyType>
+static bool GetMaterialHitShader(const FMaterial& RESTRICT MaterialResource, const FVertexFactory* VertexFactory, bool UseTextureLod, TShaderRef<FMaterialCHS>& OutShader)
+{
+	const bool bMaterialsCompiled = GCompileRayTracingMaterialAHS || GCompileRayTracingMaterialCHS;
+	checkf(bMaterialsCompiled, TEXT(""));
+
+	FMaterialShaderTypes ShaderTypes;
+	const FVertexFactoryType* VFType = VertexFactory->GetType();
+	const bool bUseIntersectionShader = VFType->HasFlags(EVertexFactoryFlags::SupportsRayTracingProceduralPrimitive) && FDataDrivenShaderPlatformInfo::GetSupportsRayTracingProceduralPrimitive(GMaxRHIShaderPlatform);
+	const bool UseAnyHitShader = (MaterialResource.IsMasked() || MaterialResource.GetBlendMode() == BLEND_Translucent) && GCompileRayTracingMaterialAHS;
+
+	GetMaterialHitShader_AnyHit_Intersection_TextureLOD<LightMapPolicyType>(ShaderTypes, UseAnyHitShader, bUseIntersectionShader, UseTextureLod);
 
 	FMaterialShaders Shaders;
 	if (!MaterialResource.TryGetShaders(ShaderTypes, VertexFactory->GetType(), Shaders))
diff --git a/Engine/Source/Runtime/Renderer/Private/RayTracing/RayTracingMaterialHitShaders.h b/Engine/Source/Runtime/Renderer/Private/RayTracing/RayTracingMaterialHitShaders.h
--- a/Engine/Source/Runtime/Renderer/Private/RayTracing/RayTracingMaterialHitShaders.h
+++ b/Engine/Source/Runtime/Renderer/Private/RayTracing/RayTracingMaterialHitShaders.h
@@ -67,7 +67,7 @@ protected:
 		SharedCommand.SetShaders(PassShaders.GetUntypedShaders());
 		SharedCommand.InstanceMask = ComputeBlendModeMask(MaterialResource.GetBlendMode());
 		SharedCommand.bCastRayTracedShadows = MeshBatch.CastRayTracedShadow && MaterialResource.CastsRayTracedShadows();
-		SharedCommand.bOpaque = MaterialResource.GetBlendMode() == EBlendMode::BLEND_Opaque;
+		SharedCommand.bOpaque = MaterialResource.GetBlendMode() == EBlendMode::BLEND_Opaque && !(VertexFactory->GetType()->SupportsRayTracingProceduralPrimitive() && FDataDrivenShaderPlatformInfo::GetSupportsRayTracingProceduralPrimitive(GMaxRHIShaderPlatform));
 		SharedCommand.bDecal = MaterialResource.GetMaterialDomain() == EMaterialDomain::MD_DeferredDecal;
 		SharedCommand.bIsSky = MaterialResource.IsSky();
 		SharedCommand.bTwoSided = MaterialResource.IsTwoSided();
