Converting FMeshPassProcessors to use ODSC.

- FCustomDepthPassMeshProcessor

diff --git a/Engine/Source/Runtime/Renderer/Private/CustomDepthRendering.cpp b/Engine/Source/Runtime/Renderer/Private/CustomDepthRendering.cpp
--- a/Engine/Source/Runtime/Renderer/Private/CustomDepthRendering.cpp
+++ b/Engine/Source/Runtime/Renderer/Private/CustomDepthRendering.cpp
@@ -251,8 +251,16 @@ public:
 	virtual void AddMeshBatch(const FMeshBatch& RESTRICT MeshBatch, uint64 BatchElementMask, const FPrimitiveSceneProxy* RESTRICT PrimitiveSceneProxy, int32 StaticMeshId = -1) override final;
 
 private:
+	bool TryAddMeshBatch(
+		const FMeshBatch& RESTRICT MeshBatch,
+		uint64 BatchElementMask,
+		const FPrimitiveSceneProxy* RESTRICT PrimitiveSceneProxy,
+		int32 StaticMeshId,
+		const FMaterialRenderProxy& MaterialRenderProxy,
+		const FMaterial& Material);
+
 	template<bool bPositionOnly, bool bUsesMobileColorValue>
-	void Process(
+	bool Process(
 		const FMeshBatch& MeshBatch,
 		uint64 BatchElementMask,
 		int32 StaticMeshId,
@@ -277,93 +285,115 @@ void FCustomDepthPassMeshProcessor::AddMeshBatch(const FMeshBatch& RESTRICT Mesh
 {
 	if (PrimitiveSceneProxy->ShouldRenderCustomDepth())
 	{
-		// Determine the mesh's material and blend mode.
-		const FMaterialRenderProxy* FallbackMaterialRenderProxyPtr = nullptr;
-		const FMaterial& Material = MeshBatch.MaterialRenderProxy->GetMaterialWithFallback(FeatureLevel, FallbackMaterialRenderProxyPtr);
-		const FMaterialRenderProxy& MaterialRenderProxy = FallbackMaterialRenderProxyPtr ? *FallbackMaterialRenderProxyPtr : *MeshBatch.MaterialRenderProxy;
-
-		const EBlendMode BlendMode = Material.GetBlendMode();
-		const FMeshDrawingPolicyOverrideSettings OverrideSettings = ComputeMeshOverrideSettings(MeshBatch);
-		const ERasterizerFillMode MeshFillMode = ComputeMeshFillMode(MeshBatch, Material, OverrideSettings);
-		const ERasterizerCullMode MeshCullMode = ComputeMeshCullMode(MeshBatch, Material, OverrideSettings);
-		const bool bIsTranslucent = IsTranslucentBlendMode(BlendMode);
-		const bool bWriteCustomStencilValues = IsCustomDepthPassWritingStencil(FeatureLevel);
-		float MobileColorValue = 0.0f;
-
-		if (bWriteCustomStencilValues)
+		const FMaterialRenderProxy* MaterialRenderProxy = MeshBatch.MaterialRenderProxy;
+		while (MaterialRenderProxy)
 		{
-			const uint32 CustomDepthStencilValue = PrimitiveSceneProxy->GetCustomDepthStencilValue();
-			static FRHIDepthStencilState* StencilStates[EStencilMask::SM_Count] =
-			{
-				TStaticDepthStencilState<true, CF_DepthNearOrEqual, true, CF_Always, SO_Keep, SO_Keep, SO_Replace, false, CF_Always, SO_Keep, SO_Keep, SO_Keep, 255, 255>::GetRHI(),
-				TStaticDepthStencilState<true, CF_DepthNearOrEqual, true, CF_Always, SO_Keep, SO_Replace, SO_Replace, false, CF_Always, SO_Keep, SO_Keep, SO_Keep, 255, 255>::GetRHI(),
-				TStaticDepthStencilState<true, CF_DepthNearOrEqual, true, CF_Always, SO_Keep, SO_Replace, SO_Replace, false, CF_Always, SO_Keep, SO_Keep, SO_Keep, 255, 1>::GetRHI(),
-				TStaticDepthStencilState<true, CF_DepthNearOrEqual, true, CF_Always, SO_Keep, SO_Replace, SO_Replace, false, CF_Always, SO_Keep, SO_Keep, SO_Keep, 255, 2>::GetRHI(),
-				TStaticDepthStencilState<true, CF_DepthNearOrEqual, true, CF_Always, SO_Keep, SO_Replace, SO_Replace, false, CF_Always, SO_Keep, SO_Keep, SO_Keep, 255, 4>::GetRHI(),
-				TStaticDepthStencilState<true, CF_DepthNearOrEqual, true, CF_Always, SO_Keep, SO_Replace, SO_Replace, false, CF_Always, SO_Keep, SO_Keep, SO_Keep, 255, 8>::GetRHI(),
-				TStaticDepthStencilState<true, CF_DepthNearOrEqual, true, CF_Always, SO_Keep, SO_Replace, SO_Replace, false, CF_Always, SO_Keep, SO_Keep, SO_Keep, 255, 16>::GetRHI(),
-				TStaticDepthStencilState<true, CF_DepthNearOrEqual, true, CF_Always, SO_Keep, SO_Replace, SO_Replace, false, CF_Always, SO_Keep, SO_Keep, SO_Keep, 255, 32>::GetRHI(),
-				TStaticDepthStencilState<true, CF_DepthNearOrEqual, true, CF_Always, SO_Keep, SO_Replace, SO_Replace, false, CF_Always, SO_Keep, SO_Keep, SO_Keep, 255, 64>::GetRHI(),
-				TStaticDepthStencilState<true, CF_DepthNearOrEqual, true, CF_Always, SO_Keep, SO_Replace, SO_Replace, false, CF_Always, SO_Keep, SO_Keep, SO_Keep, 255, 128>::GetRHI()
-			};
-			checkSlow(EStencilMask::SM_Count == UE_ARRAY_COUNT(StencilStates));
-
-			PassDrawRenderState.SetDepthStencilState(StencilStates[(int32)PrimitiveSceneProxy->GetStencilWriteMask()]);
-			PassDrawRenderState.SetStencilRef(CustomDepthStencilValue);
-
-			if (FeatureLevel <= ERHIFeatureLevel::ES3_1)
+			const FMaterial* Material = MaterialRenderProxy->GetMaterialNoFallback(FeatureLevel);
+			if (Material)
 			{
-				// On mobile platforms write custom stencil value to color target
-				MobileColorValue = CustomDepthStencilValue / 255.0f;
+				if (TryAddMeshBatch(MeshBatch, BatchElementMask, PrimitiveSceneProxy, StaticMeshId, *MaterialRenderProxy, *Material))
+				{
+					break;
+				}
 			}
+
+			MaterialRenderProxy = MaterialRenderProxy->GetFallback(FeatureLevel);
 		}
-		else
+	}
+}
+
+bool FCustomDepthPassMeshProcessor::TryAddMeshBatch(
+	const FMeshBatch& RESTRICT MeshBatch,
+	uint64 BatchElementMask,
+	const FPrimitiveSceneProxy* RESTRICT PrimitiveSceneProxy,
+	int32 StaticMeshId,
+	const FMaterialRenderProxy& MaterialRenderProxy,
+	const FMaterial& Material)
+{
+	// Determine the mesh's material and blend mode.
+	const EBlendMode BlendMode = Material.GetBlendMode();
+	const FMeshDrawingPolicyOverrideSettings OverrideSettings = ComputeMeshOverrideSettings(MeshBatch);
+	const ERasterizerFillMode MeshFillMode = ComputeMeshFillMode(MeshBatch, Material, OverrideSettings);
+	const ERasterizerCullMode MeshCullMode = ComputeMeshCullMode(MeshBatch, Material, OverrideSettings);
+	const bool bIsTranslucent = IsTranslucentBlendMode(BlendMode);
+	const bool bWriteCustomStencilValues = IsCustomDepthPassWritingStencil(FeatureLevel);
+	float MobileColorValue = 0.0f;
+
+	if (bWriteCustomStencilValues)
+	{
+		const uint32 CustomDepthStencilValue = PrimitiveSceneProxy->GetCustomDepthStencilValue();
+		static FRHIDepthStencilState* StencilStates[EStencilMask::SM_Count] =
+		{
+			TStaticDepthStencilState<true, CF_DepthNearOrEqual, true, CF_Always, SO_Keep, SO_Keep, SO_Replace, false, CF_Always, SO_Keep, SO_Keep, SO_Keep, 255, 255>::GetRHI(),
+			TStaticDepthStencilState<true, CF_DepthNearOrEqual, true, CF_Always, SO_Keep, SO_Replace, SO_Replace, false, CF_Always, SO_Keep, SO_Keep, SO_Keep, 255, 255>::GetRHI(),
+			TStaticDepthStencilState<true, CF_DepthNearOrEqual, true, CF_Always, SO_Keep, SO_Replace, SO_Replace, false, CF_Always, SO_Keep, SO_Keep, SO_Keep, 255, 1>::GetRHI(),
+			TStaticDepthStencilState<true, CF_DepthNearOrEqual, true, CF_Always, SO_Keep, SO_Replace, SO_Replace, false, CF_Always, SO_Keep, SO_Keep, SO_Keep, 255, 2>::GetRHI(),
+			TStaticDepthStencilState<true, CF_DepthNearOrEqual, true, CF_Always, SO_Keep, SO_Replace, SO_Replace, false, CF_Always, SO_Keep, SO_Keep, SO_Keep, 255, 4>::GetRHI(),
+			TStaticDepthStencilState<true, CF_DepthNearOrEqual, true, CF_Always, SO_Keep, SO_Replace, SO_Replace, false, CF_Always, SO_Keep, SO_Keep, SO_Keep, 255, 8>::GetRHI(),
+			TStaticDepthStencilState<true, CF_DepthNearOrEqual, true, CF_Always, SO_Keep, SO_Replace, SO_Replace, false, CF_Always, SO_Keep, SO_Keep, SO_Keep, 255, 16>::GetRHI(),
+			TStaticDepthStencilState<true, CF_DepthNearOrEqual, true, CF_Always, SO_Keep, SO_Replace, SO_Replace, false, CF_Always, SO_Keep, SO_Keep, SO_Keep, 255, 32>::GetRHI(),
+			TStaticDepthStencilState<true, CF_DepthNearOrEqual, true, CF_Always, SO_Keep, SO_Replace, SO_Replace, false, CF_Always, SO_Keep, SO_Keep, SO_Keep, 255, 64>::GetRHI(),
+			TStaticDepthStencilState<true, CF_DepthNearOrEqual, true, CF_Always, SO_Keep, SO_Replace, SO_Replace, false, CF_Always, SO_Keep, SO_Keep, SO_Keep, 255, 128>::GetRHI()
+		};
+		checkSlow(EStencilMask::SM_Count == UE_ARRAY_COUNT(StencilStates));
+
+		PassDrawRenderState.SetDepthStencilState(StencilStates[(int32)PrimitiveSceneProxy->GetStencilWriteMask()]);
+		PassDrawRenderState.SetStencilRef(CustomDepthStencilValue);
+
+		if (FeatureLevel <= ERHIFeatureLevel::ES3_1)
 		{
-			PassDrawRenderState.SetDepthStencilState(TStaticDepthStencilState<true, CF_DepthNearOrEqual>::GetRHI());
+			// On mobile platforms write custom stencil value to color target
+			MobileColorValue = CustomDepthStencilValue / 255.0f;
 		}
+	}
+	else
+	{
+		PassDrawRenderState.SetDepthStencilState(TStaticDepthStencilState<true, CF_DepthNearOrEqual>::GetRHI());
+	}
+
+	const bool bUsesMobileColorValue = FeatureLevel <= ERHIFeatureLevel::ES3_1;
+
+	bool bResult = true;
+	if (BlendMode == BLEND_Opaque
+		&& MeshBatch.VertexFactory->SupportsPositionOnlyStream()
+		&& !Material.MaterialModifiesMeshPosition_RenderThread()
+		&& Material.WritesEveryPixel()
+		&& !bUsesMobileColorValue)
+	{
+		const FMaterialRenderProxy& DefaultProxy = *UMaterial::GetDefaultMaterial(MD_Surface)->GetRenderProxy();
+		const FMaterial& DefaultMaterial = *DefaultProxy.GetMaterialNoFallback(FeatureLevel);
+		bResult = Process<true, false>(MeshBatch, BatchElementMask, StaticMeshId, PrimitiveSceneProxy, DefaultProxy, DefaultMaterial, MeshFillMode, MeshCullMode, MobileColorValue);
+	}
+	else if (!IsTranslucentBlendMode(BlendMode) || Material.IsTranslucencyWritingCustomDepth())
+	{
+		const bool bMaterialMasked = !Material.WritesEveryPixel() || Material.IsTranslucencyWritingCustomDepth();
 
-		const bool bUsesMobileColorValue = FeatureLevel <= ERHIFeatureLevel::ES3_1;
+		const FMaterialRenderProxy* EffectiveMaterialRenderProxy = &MaterialRenderProxy;
+		const FMaterial* EffectiveMaterial = &Material;
 
+		if (!bMaterialMasked && !Material.MaterialModifiesMeshPosition_RenderThread())
+		{
+			// Override with the default material for opaque materials that are not two sided
+			EffectiveMaterialRenderProxy = UMaterial::GetDefaultMaterial(MD_Surface)->GetRenderProxy();
+			EffectiveMaterial = EffectiveMaterialRenderProxy->GetMaterialNoFallback(FeatureLevel);
+			check(EffectiveMaterial);
+		}
 
-		if (BlendMode == BLEND_Opaque
-			&& MeshBatch.VertexFactory->SupportsPositionOnlyStream()
-			&& !Material.MaterialModifiesMeshPosition_RenderThread()
-			&& Material.WritesEveryPixel()
-			&& !bUsesMobileColorValue)
+		if (bUsesMobileColorValue)
 		{
-			const FMaterialRenderProxy& DefaultProxy = *UMaterial::GetDefaultMaterial(MD_Surface)->GetRenderProxy();
-			const FMaterial& DefaultMaterial = *DefaultProxy.GetMaterialNoFallback(FeatureLevel);
-			Process<true, false>(MeshBatch, BatchElementMask, StaticMeshId, PrimitiveSceneProxy, DefaultProxy, DefaultMaterial, MeshFillMode, MeshCullMode, MobileColorValue);
+			bResult = Process<false, true>(MeshBatch, BatchElementMask, StaticMeshId, PrimitiveSceneProxy, *EffectiveMaterialRenderProxy, *EffectiveMaterial, MeshFillMode, MeshCullMode, MobileColorValue);
 		}
-		else if (!IsTranslucentBlendMode(BlendMode) || Material.IsTranslucencyWritingCustomDepth())
+		else
 		{
-			const bool bMaterialMasked = !Material.WritesEveryPixel() || Material.IsTranslucencyWritingCustomDepth();
-
-			const FMaterialRenderProxy* EffectiveMaterialRenderProxy = &MaterialRenderProxy;
-			const FMaterial* EffectiveMaterial = &Material;
-
-			if (!bMaterialMasked && !Material.MaterialModifiesMeshPosition_RenderThread())
-			{
-				// Override with the default material for opaque materials that are not two sided
-				EffectiveMaterialRenderProxy = UMaterial::GetDefaultMaterial(MD_Surface)->GetRenderProxy();
-				EffectiveMaterial = EffectiveMaterialRenderProxy->GetMaterialNoFallback(FeatureLevel);
-				check(EffectiveMaterial);
-			}
-
-			if (bUsesMobileColorValue)
-			{
-				Process<false, true>(MeshBatch, BatchElementMask, StaticMeshId, PrimitiveSceneProxy, *EffectiveMaterialRenderProxy, *EffectiveMaterial, MeshFillMode, MeshCullMode, MobileColorValue);
-			}
-			else
-			{
-				Process<false, false>(MeshBatch, BatchElementMask, StaticMeshId, PrimitiveSceneProxy, *EffectiveMaterialRenderProxy, *EffectiveMaterial, MeshFillMode, MeshCullMode, MobileColorValue);
-			}
+			bResult = Process<false, false>(MeshBatch, BatchElementMask, StaticMeshId, PrimitiveSceneProxy, *EffectiveMaterialRenderProxy, *EffectiveMaterial, MeshFillMode, MeshCullMode, MobileColorValue);
 		}
 	}
+
+	return bResult;
 }
 
 template<bool bPositionOnly, bool bUsesMobileColorValue>
-void FCustomDepthPassMeshProcessor::Process(
+bool FCustomDepthPassMeshProcessor::Process(
 	const FMeshBatch& RESTRICT MeshBatch,
 	uint64 BatchElementMask,
 	int32 StaticMeshId,
@@ -381,14 +411,17 @@ void FCustomDepthPassMeshProcessor::Process(
 		FDepthOnlyPS<bUsesMobileColorValue>> DepthPassShaders;
 
 	FShaderPipelineRef ShaderPipeline;
-	GetDepthPassShaders<bPositionOnly, bUsesMobileColorValue>(
+	if (!GetDepthPassShaders<bPositionOnly, bUsesMobileColorValue>(
 		MaterialResource,
 		VertexFactory->GetType(),
 		FeatureLevel,
 		DepthPassShaders.VertexShader,
 		DepthPassShaders.PixelShader,
 		ShaderPipeline
-		);
+		))
+	{
+		return false;
+	}
 
 	FDepthOnlyShaderElementData ShaderElementData(MobileColorValue);
 	ShaderElementData.InitializeMeshMaterialData(ViewIfDynamicMeshCommand, PrimitiveSceneProxy, MeshBatch, StaticMeshId, false);
@@ -408,6 +441,8 @@ void FCustomDepthPassMeshProcessor::Process(
 		SortKey,
 		bPositionOnly ? EMeshPassFeatures::PositionOnly : EMeshPassFeatures::Default,
 		ShaderElementData);
+
+	return true;
 }
 
 FMeshPassProcessor* CreateCustomDepthPassProcessor(const FScene* Scene, const FSceneView* InViewIfDynamicMeshCommand, FMeshPassDrawListContext* InDrawListContext)
