Support Niagara meshes in TLAS instance buffer
- Modified RayTracingBuildInstanceBufferCS to get transforms from a separate buffer.
- All CPU transforms are stored in a separate upload buffer.
- Niagara meshes particles are stored in GPUTransformsSRV.
- 1 dispatch for all CPU instances + 1 per Niagara system.
- Removed unused RayTracingInstanceCopyShaderCS.

diff --git a/Engine/Plugins/Experimental/GPULightmass/Source/GPULightmass/Private/LightmapRenderer.cpp b/Engine/Plugins/Experimental/GPULightmass/Source/GPULightmass/Private/LightmapRenderer.cpp
--- a/Engine/Plugins/Experimental/GPULightmass/Source/GPULightmass/Private/LightmapRenderer.cpp
+++ b/Engine/Plugins/Experimental/GPULightmass/Source/GPULightmass/Private/LightmapRenderer.cpp
@@ -1004,8 +1004,12 @@ void FSceneRenderState::SetupRayTracingScene(int32 LODIndex)
 		{
 			SCOPED_GPU_MASK(RHICmdList, FRHIGPUMask::All());
 
-			TArray<uint32> GeometryIndices;
-			RayTracingScene = CreateRayTracingSceneWithGeometryInstances(RayTracingGeometryInstances, RAY_TRACING_NUM_SHADER_SLOTS, 1, GeometryIndices);
+			FRayTracingSceneWithGeometryInstances SceneWithGeometryInstances = CreateRayTracingSceneWithGeometryInstances(
+				RayTracingGeometryInstances,
+				RAY_TRACING_NUM_SHADER_SLOTS,
+				1);
+
+			RayTracingScene = SceneWithGeometryInstances.Scene;
 
 			const FRayTracingSceneInitializer2& SceneInitializer = RayTracingScene->GetInitializer();
 
@@ -1030,14 +1034,32 @@ void FSceneRenderState::SetupRayTracingScene(int32 LODIndex)
 			FBufferRHIRef InstanceUploadBuffer;
 			FShaderResourceViewRHIRef InstanceUploadSRV;
 			{
-				FRHIResourceCreateInfo CreateInfo(TEXT("RayTracingTestBedInstanceUploadBuffer"));
+				FRHIResourceCreateInfo CreateInfo(TEXT("LightmassRayTracingInstanceUploadBuffer"));
 				InstanceUploadBuffer = RHICreateStructuredBuffer(sizeof(FRayTracingInstanceDescriptorInput), InstanceUploadBufferSize, BUF_ShaderResource | BUF_Volatile, CreateInfo);
 				InstanceUploadSRV = RHICreateShaderResourceView(InstanceUploadBuffer);
 			}
 
+			const uint32 TransformUploadBufferSize = SceneWithGeometryInstances.NumNativeCPUInstances * 3 * sizeof(FVector4f);
+			FBufferRHIRef TransformUploadBuffer;
+			FShaderResourceViewRHIRef TransformUploadSRV;
+			{
+				FRHIResourceCreateInfo CreateInfo(TEXT("LightmassRayTracingTransformUploadBuffer"));
+				TransformUploadBuffer = RHICreateStructuredBuffer(sizeof(FVector4f), TransformUploadBufferSize, BUF_ShaderResource | BUF_Volatile, CreateInfo);
+				TransformUploadSRV = RHICreateShaderResourceView(TransformUploadBuffer);
+			}
+
 			{
 				FRayTracingInstanceDescriptorInput* InstanceUploadData = (FRayTracingInstanceDescriptorInput*)RHICmdList.LockBuffer(InstanceUploadBuffer, 0, InstanceUploadBufferSize, RLM_WriteOnly);
-				FillRayTracingInstanceUploadBuffer(RayTracingGeometryInstances, GeometryIndices, RayTracingScene, MakeArrayView(InstanceUploadData, SceneInitializer.NumNativeInstances));
+				FVector4f* TransformUploadData = (FVector4f*)RHICmdList.LockBuffer(TransformUploadBuffer, 0, TransformUploadBufferSize, RLM_WriteOnly);
+				FillRayTracingInstanceUploadBuffer(
+					RayTracingScene,
+					RayTracingGeometryInstances,
+					SceneWithGeometryInstances.InstanceGeometryIndices,
+					SceneWithGeometryInstances.BaseUploadBufferOffsets,
+					SceneWithGeometryInstances.NumNativeCPUInstances,
+					MakeArrayView(InstanceUploadData, SceneInitializer.NumNativeInstances),
+					MakeArrayView(TransformUploadData, SceneWithGeometryInstances.NumNativeCPUInstances * 3));
+				RHICmdList.UnlockBuffer(TransformUploadBuffer);
 				RHICmdList.UnlockBuffer(InstanceUploadBuffer);
 			}
 
@@ -1059,10 +1081,12 @@ void FSceneRenderState::SetupRayTracingScene(int32 LODIndex)
 
 			BuildRayTracingInstanceBuffer(
 				RHICmdList,
-				SceneInitializer.NumNativeInstances,
 				InstanceBuffer.UAV,
 				InstanceUploadSRV,
-				AccelerationStructureAddressesBuffer.SRV);
+				AccelerationStructureAddressesBuffer.SRV,
+				TransformUploadSRV,
+				SceneWithGeometryInstances.NumNativeCPUInstances,
+				{});
 
 			RHICmdList.BindAccelerationStructureMemory(RayTracingScene, RayTracingSceneBuffer, 0);
 
diff --git a/Engine/Plugins/FX/Niagara/Shaders/Private/NiagaraGPURayTracingTransforms.usf b/Engine/Plugins/FX/Niagara/Shaders/Private/NiagaraGPURayTracingTransforms.usf
--- a/Engine/Plugins/FX/Niagara/Shaders/Private/NiagaraGPURayTracingTransforms.usf
+++ b/Engine/Plugins/FX/Niagara/Shaders/Private/NiagaraGPURayTracingTransforms.usf
@@ -35,10 +35,12 @@ float4x4	LocalTransform;
 
 Buffer<float>	ParticleDataFloatBuffer;
 //Buffer<half>	ParticleDataHalfBuffer;
-Buffer<int>	ParticleDataIntBuffer;
+Buffer<int>		ParticleDataIntBuffer;
 Buffer<uint>	GPUInstanceCountBuffer;
 
-RWStructuredBuffer<float3x4> TLASTransforms;
+// Transforms are float3x4 row-major matrices stored as 3 float4
+// because of FXC matrix packing issues when using StructuredBuffer<float3x4>
+RWStructuredBuffer<float4> TLASTransforms;
 
 //////////////////////////////////////////////////////////////////////////////////////////////////////////////
 // Helper function for buffers
@@ -111,7 +113,7 @@ void NiagaraGPURayTracingTransformsCS(uint3 DispatchThreadId : SV_DispatchThread
 			InstanceTransformTransposed = transpose(InstanceTransform);
 		}
 	}
-	TLASTransforms[InstanceIndex][0] = InstanceTransformTransposed[0];
-	TLASTransforms[InstanceIndex][1] = InstanceTransformTransposed[1];
-	TLASTransforms[InstanceIndex][2] = InstanceTransformTransposed[2];
+	TLASTransforms[InstanceIndex * 3 + 0] = InstanceTransformTransposed[0];
+	TLASTransforms[InstanceIndex * 3 + 1] = InstanceTransformTransposed[1];
+	TLASTransforms[InstanceIndex * 3 + 2] = InstanceTransformTransposed[2];
 }
\ No newline at end of file
diff --git a/Engine/Plugins/FX/Niagara/Source/Niagara/Private/NiagaraRendererMeshes.cpp b/Engine/Plugins/FX/Niagara/Source/Niagara/Private/NiagaraRendererMeshes.cpp
--- a/Engine/Plugins/FX/Niagara/Source/Niagara/Private/NiagaraRendererMeshes.cpp
+++ b/Engine/Plugins/FX/Niagara/Source/Niagara/Private/NiagaraRendererMeshes.cpp
@@ -1373,8 +1373,7 @@ void FNiagaraRendererMeshes::GetDynamicRayTracingInstances(FRayTracingMaterialGa
 				RayTracingInstance.NumTransforms = NumInstances;
 
 				FRWBufferStructured InstanceGPUTransformsBuffer;
-				//InstanceGPUTransformsBuffer.Initialize(sizeof(FMatrix), NumInstances, BUF_Static);
-				InstanceGPUTransformsBuffer.Initialize(TEXT("InstanceGPUTransformsBuffer"), 3 * 4 * sizeof(float), NumInstances, BUF_Static);
+				InstanceGPUTransformsBuffer.Initialize(TEXT("InstanceGPUTransformsBuffer"), 4 * sizeof(float), 3 * NumInstances, BUF_Static);
 				RayTracingInstance.InstanceGPUTransformsSRV = InstanceGPUTransformsBuffer.SRV;
 
 				FNiagaraGPURayTracingTransformsCS::FParameters PassParameters;
diff --git a/Engine/Shaders/Private/RayTracing/RayTracingInstanceBufferUtil.usf b/Engine/Shaders/Private/RayTracing/RayTracingInstanceBufferUtil.usf
--- a/Engine/Shaders/Private/RayTracing/RayTracingInstanceBufferUtil.usf
+++ b/Engine/Shaders/Private/RayTracing/RayTracingInstanceBufferUtil.usf
@@ -12,38 +12,26 @@
 #define THREADGROUP_SIZE 1
 #endif
 
-uint NumInstances;
-uint DescBufferOffset;
-StructuredBuffer<float3x4> InstancesTransforms;
-RWStructuredBuffer<FRayTracingInstanceDescriptor> InstancesDescriptors;
-
-
-[numthreads(THREADGROUP_SIZE, 1, 1)]
-void RayTracingInstanceCopyShaderCS(uint3 DispatchThreadId : SV_DispatchThreadID)
-{
-	uint Index = DispatchThreadId.x;	
-
-	float3x4 InstanceXForm = InstancesTransforms[Index];
-	uint DescIndex = DescBufferOffset + Index;
-
-	if (Index < NumInstances)
-	{
-		SetRayTracingInstanceTransform(InstancesDescriptors[DescIndex], InstanceXForm);
-	}
-}
-
 struct FRayTracingInstanceDescriptorInput
 {
-	//uint GPUSceneInstanceIndex;
-	float3x4 LocalToWorld;
+	uint TransformIndex;
+	uint OutputDescriptorIndex;
 	uint AccelerationStructureIndex;
 	uint InstanceId;
 	uint InstanceMaskAndFlags;
 	uint InstanceContributionToHitGroupIndex;
 };
 
+uint NumInstances;
+uint InputDescOffset;
+
+RWStructuredBuffer<FRayTracingInstanceDescriptor> InstancesDescriptors;
+
 StructuredBuffer<FRayTracingInstanceDescriptorInput> InputInstanceDescriptors;
 ByteAddressBuffer AccelerationStructureAddresses;
+// Transforms are float3x4 row-major matrices stored as 3 float4
+// because of FXC matrix packing issues when using StructuredBuffer<float3x4>
+StructuredBuffer<float4> InstanceTransforms;
 
 [numthreads(THREADGROUP_SIZE, 1, 1)]
 void RayTracingBuildInstanceBufferCS(uint3 DispatchThreadId : SV_DispatchThreadID)
@@ -52,11 +40,15 @@ void RayTracingBuildInstanceBufferCS(uint3 DispatchThreadId : SV_DispatchThreadI
 
 	if (Index < NumInstances)
 	{
-		FRayTracingInstanceDescriptorInput InputDesc = InputInstanceDescriptors[Index];
+		FRayTracingInstanceDescriptorInput InputDesc = InputInstanceDescriptors[InputDescOffset + Index];
 
 		//FInstanceSceneData InstanceSceneData = GetInstanceSceneData(InputDesc.GPUSceneInstanceIndex, View.InstanceSceneDataSOAStride);
 		//float4x4 LocalToWorld = InstanceSceneData.LocalToWorld;
-		float3x4 LocalToWorld = InputDesc.LocalToWorld;
+
+		float3x4 LocalToWorld;
+		LocalToWorld[0] = InstanceTransforms[InputDesc.TransformIndex * 3 + 0];
+		LocalToWorld[1] = InstanceTransforms[InputDesc.TransformIndex * 3 + 1];
+		LocalToWorld[2] = InstanceTransforms[InputDesc.TransformIndex * 3 + 2];
 
 		uint2 BlasAddress = 0;
 
@@ -65,7 +57,7 @@ void RayTracingBuildInstanceBufferCS(uint3 DispatchThreadId : SV_DispatchThreadI
 			BlasAddress = AccelerationStructureAddresses.Load2(InputDesc.AccelerationStructureIndex * 8);
 		}
 
-		InstancesDescriptors[Index] = BuildPlatformRayTracingInstanceDesc(
+		InstancesDescriptors[InputDesc.OutputDescriptorIndex] = BuildPlatformRayTracingInstanceDesc(
 			InputDesc.InstanceMaskAndFlags & 0xFF,
 			InputDesc.InstanceId,
 			TranslateRayTracingInstanceFlags((InputDesc.InstanceMaskAndFlags >> 8) & 0xFF),
diff --git a/Engine/Source/Runtime/Renderer/Private/RayTracing/RayTracingInstanceBufferUtil.cpp b/Engine/Source/Runtime/Renderer/Private/RayTracing/RayTracingInstanceBufferUtil.cpp
--- a/Engine/Source/Runtime/Renderer/Private/RayTracing/RayTracingInstanceBufferUtil.cpp
+++ b/Engine/Source/Runtime/Renderer/Private/RayTracing/RayTracingInstanceBufferUtil.cpp
@@ -16,15 +16,18 @@
 
 #if RHI_RAYTRACING
 
-FRayTracingSceneRHIRef CreateRayTracingSceneWithGeometryInstances(
+FRayTracingSceneWithGeometryInstances CreateRayTracingSceneWithGeometryInstances(
 	TArrayView<FRayTracingGeometryInstance> Instances,
 	uint32 NumShaderSlotsPerGeometrySegment,
-	uint32 NumMissShaderSlots,
-	TArray<uint32>& OutGeometryIndices)
+	uint32 NumMissShaderSlots)
 {
 	const uint32 NumSceneInstances = Instances.Num();
 
-	OutGeometryIndices.SetNumUninitialized(NumSceneInstances);
+	FRayTracingSceneWithGeometryInstances Output;
+	Output.NumNativeCPUInstances = 0;
+	Output.NumNativeGPUInstances = 0;
+	Output.InstanceGeometryIndices.SetNumUninitialized(NumSceneInstances);
+	Output.BaseUploadBufferOffsets.SetNumUninitialized(NumSceneInstances);
 
 	FRayTracingSceneInitializer2 Initializer;
 	Initializer.DebugName = FName(TEXT("FRayTracingScene"));
@@ -62,31 +65,68 @@ FRayTracingSceneRHIRef CreateRayTracingSceneWithGeometryInstances(
 		Initializer.NumTotalSegments += InstanceDesc.GeometryRHI->GetNumSegments();
 
 		uint32 GeometryIndex = UniqueGeometries.FindOrAdd(InstanceDesc.GeometryRHI, Initializer.ReferencedGeometries.Num());
-		OutGeometryIndices[InstanceIndex] = GeometryIndex;
+		Output.InstanceGeometryIndices[InstanceIndex] = GeometryIndex;
 		if (GeometryIndex == Initializer.ReferencedGeometries.Num())
 		{
 			Initializer.ReferencedGeometries.Add(InstanceDesc.GeometryRHI);
 		}
 
+		if (bCpuInstance)
+		{
+			Output.BaseUploadBufferOffsets[InstanceIndex] = Output.NumNativeCPUInstances;
+
+			Output.NumNativeCPUInstances += InstanceDesc.NumTransforms;
+		}
+		else
+		{
+			Output.BaseUploadBufferOffsets[InstanceIndex] = Output.NumNativeGPUInstances;
+
+			if (InstanceDesc.NumTransforms)
+			{
+				FRayTracingGPUInstance GPUInstance;
+				GPUInstance.TransformSRV = InstanceDesc.GPUTransformsSRV;
+				GPUInstance.DescBufferOffset = Output.NumNativeGPUInstances;
+				GPUInstance.NumInstances = InstanceDesc.NumTransforms;
+				Output.GPUInstances.Add(GPUInstance);
+			}
+
+			Output.NumNativeGPUInstances += InstanceDesc.NumTransforms;
+		}
+
 		Initializer.BaseInstancePrefixSum[InstanceIndex] = Initializer.NumNativeInstances;
 		Initializer.NumNativeInstances += InstanceDesc.NumTransforms;
 	}
 
-	return RHICreateRayTracingScene(MoveTemp(Initializer));
+	Output.Scene = RHICreateRayTracingScene(MoveTemp(Initializer));
+
+	return MoveTemp(Output);
 }
 
 void FillRayTracingInstanceUploadBuffer(
-	TConstArrayView<FRayTracingGeometryInstance> Instances,
-	TConstArrayView<uint32> InstancesGeometryIndex,
 	FRayTracingSceneRHIRef RayTracingSceneRHI,
-	TArrayView<FRayTracingInstanceDescriptorInput> OutInstanceUploadData)
+	TConstArrayView<FRayTracingGeometryInstance> Instances,
+	TConstArrayView<uint32> InstanceGeometryIndices,
+	TConstArrayView<uint32> BaseUploadBufferOffsets,
+	uint32 NumNativeCPUInstances,
+	TArrayView<FRayTracingInstanceDescriptorInput> OutInstanceUploadData,
+	TArrayView<FVector4f> OutTransformData)
 {
 	const FRayTracingSceneInitializer2& SceneInitializer = RayTracingSceneRHI->GetInitializer();
 
 	int32 NumInactiveNativeInstances = 0;
 
 	const int32 NumSceneInstances = Instances.Num();
-	ParallelFor(NumSceneInstances, [OutInstanceUploadData, Instances, InstancesGeometryIndex, &SceneInitializer, &NumInactiveNativeInstances](int32 SceneInstanceIndex)
+	ParallelFor(NumSceneInstances, 
+		[
+			OutInstanceUploadData,
+			OutTransformData,
+			NumNativeCPUInstances,
+			Instances,
+			InstanceGeometryIndices,
+			BaseUploadBufferOffsets,
+			&SceneInitializer,
+			&NumInactiveNativeInstances
+		](int32 SceneInstanceIndex)
 		{
 			const FRayTracingGeometryInstance& SceneInstance = Instances[SceneInstanceIndex];
 
@@ -102,18 +142,23 @@ void FillRayTracingInstanceUploadBuffer(
 			const bool bGpuInstance = SceneInstance.GPUTransformsSRV != nullptr;
 			const bool bCpuInstance = !bGpuInstance;
 
-			const uint32 AccelerationStructureIndex = InstancesGeometryIndex[SceneInstanceIndex];
-
-			uint32 BaseInstanceIndex = SceneInitializer.BaseInstancePrefixSum[SceneInstanceIndex];
+			const uint32 AccelerationStructureIndex = InstanceGeometryIndices[SceneInstanceIndex];
+			const uint32 BaseInstanceIndex = SceneInitializer.BaseInstancePrefixSum[SceneInstanceIndex];
+			// GPU instance descriptors are stored after CPU instances
+			const uint32 BaseDescriptorIndex = BaseUploadBufferOffsets[SceneInstanceIndex] + (bGpuInstance ? NumNativeCPUInstances : 0);
+			const uint32 BaseTransformIndex = bCpuInstance ? BaseDescriptorIndex : 0;
 
 			int32 NumInactiveNativeInstancesThisSceneInstance = 0;
 			for (uint32 TransformIndex = 0; TransformIndex < NumTransforms; ++TransformIndex)
 			{
-				FRayTracingInstanceDescriptorInput& InstanceDesc = OutInstanceUploadData[BaseInstanceIndex + TransformIndex];
+				FRayTracingInstanceDescriptorInput& InstanceDesc = OutInstanceUploadData[BaseDescriptorIndex + TransformIndex];
+
 				InstanceDesc.InstanceMaskAndFlags = SceneInstance.Mask | ((uint32)SceneInstance.Flags << 8);
 				InstanceDesc.InstanceContributionToHitGroupIndex = SceneInitializer.SegmentPrefixSum[SceneInstanceIndex] * RAY_TRACING_NUM_SHADER_SLOTS;
 				InstanceDesc.InstanceId = bUseUniqueUserData ? SceneInstance.UserData[TransformIndex] : SceneInstance.DefaultUserData;
 				//InstanceDesc.GPUSceneInstanceIndex = 0; // TODO
+				InstanceDesc.TransformIndex = BaseTransformIndex + TransformIndex;
+				InstanceDesc.OutputDescriptorIndex = BaseInstanceIndex + TransformIndex;
 
 				InstanceDesc.AccelerationStructureIndex = AccelerationStructureIndex;
 
@@ -128,14 +173,9 @@ void FillRayTracingInstanceUploadBuffer(
 				if (LIKELY(bCpuInstance))
 				{
 					const FMatrix44f LocalToWorld = SceneInstance.Transforms[TransformIndex].GetTransposed();
-					InstanceDesc.LocalToWorld[0] = *(FVector4f*)&LocalToWorld.M[0];
-					InstanceDesc.LocalToWorld[1] = *(FVector4f*)&LocalToWorld.M[1];
-					InstanceDesc.LocalToWorld[2] = *(FVector4f*)&LocalToWorld.M[2];
-				}
-				else
-				{
-					// #todo: GPU-based instances transform should be copied from GPUTransformsSRV
-					FMemory::Memset(&InstanceDesc.LocalToWorld, 0, sizeof(InstanceDesc.LocalToWorld));
+					OutTransformData[InstanceDesc.TransformIndex * 3 + 0] = *(FVector4f*)&LocalToWorld.M[0];
+					OutTransformData[InstanceDesc.TransformIndex * 3 + 1] = *(FVector4f*)&LocalToWorld.M[1];
+					OutTransformData[InstanceDesc.TransformIndex * 3 + 2] = *(FVector4f*)&LocalToWorld.M[2];
 				}
 			}
 
@@ -147,44 +187,18 @@ void FillRayTracingInstanceUploadBuffer(
 	SET_DWORD_STAT(STAT_RayTracingInstances, SceneInitializer.NumNativeInstances - NumInactiveNativeInstances);
 }
 
-struct FRayTracingInstanceCopyCS : public FGlobalShader
-{
-	DECLARE_GLOBAL_SHADER(FRayTracingInstanceCopyCS);
-	SHADER_USE_PARAMETER_STRUCT(FRayTracingInstanceCopyCS, FGlobalShader);
-
-	BEGIN_SHADER_PARAMETER_STRUCT(FParameters, )
-		SHADER_PARAMETER_RDG_BUFFER_UAV(RWStructuredBuffer, InstancesDescriptors)
-		SHADER_PARAMETER_RDG_BUFFER_SRV(StructuredBuffer, InstancesTransforms)
-		SHADER_PARAMETER(uint32, NumInstances)
-		SHADER_PARAMETER(uint32, DescBufferOffset)
-	END_SHADER_PARAMETER_STRUCT()
-
-	static constexpr uint32 ThreadGroupSize = 64;
-	static inline void ModifyCompilationEnvironment(const FGlobalShaderPermutationParameters& Parameters, FShaderCompilerEnvironment& OutEnvironment)
-	{
-		FGlobalShader::ModifyCompilationEnvironment(Parameters, OutEnvironment);
-
-		OutEnvironment.SetDefine(TEXT("THREADGROUP_SIZE"), ThreadGroupSize);
-	}
-
-	static bool ShouldCompilePermutation(const FGlobalShaderPermutationParameters& Parameters)
-	{
-		return (ShouldCompileRayTracingShadersForProject(Parameters.Platform) && RHISupportsComputeShaders(Parameters.Platform) && !(Parameters.Platform == EShaderPlatform::SP_METAL || Parameters.Platform == EShaderPlatform::SP_METAL_TVOS || IsMobilePlatform(Parameters.Platform)));
-	}
-};
-
-IMPLEMENT_GLOBAL_SHADER(FRayTracingInstanceCopyCS, "/Engine/Private/Raytracing/RayTracingInstanceBufferUtil.usf", "RayTracingInstanceCopyShaderCS", SF_Compute);
-
-struct FRayTracingInstanceBufferCS : public FGlobalShader
+struct FRayTracingBuildInstanceBufferCS : public FGlobalShader
 {
-	DECLARE_GLOBAL_SHADER(FRayTracingInstanceBufferCS);
-	SHADER_USE_PARAMETER_STRUCT(FRayTracingInstanceBufferCS, FGlobalShader);
+	DECLARE_GLOBAL_SHADER(FRayTracingBuildInstanceBufferCS);
+	SHADER_USE_PARAMETER_STRUCT(FRayTracingBuildInstanceBufferCS, FGlobalShader);
 
 	BEGIN_SHADER_PARAMETER_STRUCT(FParameters, )
 		SHADER_PARAMETER_UAV(RWStructuredBuffer, InstancesDescriptors)
 		SHADER_PARAMETER_SRV(StructuredBuffer<FRayTracingInstanceDescriptorInput>, InputInstanceDescriptors)
 		SHADER_PARAMETER_SRV(ByteAddressBuffer, AccelerationStructureAddresses)
+		SHADER_PARAMETER_SRV(StructuredBuffer, InstanceTransforms)
 		SHADER_PARAMETER(uint32, NumInstances)
+		SHADER_PARAMETER(uint32, InputDescOffset)
 	END_SHADER_PARAMETER_STRUCT()
 		
 	static constexpr uint32 ThreadGroupSize = 64;
@@ -202,18 +216,27 @@ struct FRayTracingInstanceBufferCS : public FGlobalShader
 	}
 };
 
-IMPLEMENT_GLOBAL_SHADER(FRayTracingInstanceBufferCS, "/Engine/Private/Raytracing/RayTracingInstanceBufferUtil.usf", "RayTracingBuildInstanceBufferCS", SF_Compute);
+IMPLEMENT_GLOBAL_SHADER(FRayTracingBuildInstanceBufferCS, "/Engine/Private/Raytracing/RayTracingInstanceBufferUtil.usf", "RayTracingBuildInstanceBufferCS", SF_Compute);
 
-void BuildRayTracingInstanceBuffer(FRHICommandList& RHICmdList, uint32 NumInstances, FUnorderedAccessViewRHIRef InstancesUAV, FShaderResourceViewRHIRef InstanceUploadSRV, FShaderResourceViewRHIRef AccelerationStructureAddressesSRV)
+void BuildRayTracingInstanceBuffer(
+	FRHICommandList& RHICmdList,
+	uint32 NumInstances,
+	uint32 InputDescOffset,
+	FUnorderedAccessViewRHIRef InstancesUAV,
+	FShaderResourceViewRHIRef InstanceUploadSRV,
+	FShaderResourceViewRHIRef AccelerationStructureAddressesSRV,
+	FShaderResourceViewRHIRef InstanceTransformSRV)
 {
-	FRayTracingInstanceBufferCS::FParameters PassParams;
+	FRayTracingBuildInstanceBufferCS::FParameters PassParams;
 	PassParams.InstancesDescriptors = InstancesUAV;
 	PassParams.InputInstanceDescriptors = InstanceUploadSRV;
 	PassParams.AccelerationStructureAddresses = AccelerationStructureAddressesSRV;
+	PassParams.InstanceTransforms = InstanceTransformSRV;
 	PassParams.NumInstances = NumInstances;
+	PassParams.InputDescOffset = InputDescOffset;
 
-	auto ComputeShader = GetGlobalShaderMap(GMaxRHIFeatureLevel)->GetShader<FRayTracingInstanceBufferCS>();
-	const int32 GroupSize = FMath::DivideAndRoundUp(PassParams.NumInstances, FRayTracingInstanceBufferCS::ThreadGroupSize);
+	auto ComputeShader = GetGlobalShaderMap(GMaxRHIFeatureLevel)->GetShader<FRayTracingBuildInstanceBufferCS>();
+	const int32 GroupSize = FMath::DivideAndRoundUp(PassParams.NumInstances, FRayTracingBuildInstanceBufferCS::ThreadGroupSize);
 
 	RHICmdList.SetComputeShader(ComputeShader.GetComputeShader());
 
@@ -224,4 +247,42 @@ void BuildRayTracingInstanceBuffer(FRHICommandList& RHICmdList, uint32 NumInstan
 	UnsetShaderUAVs(RHICmdList, ComputeShader, ComputeShader.GetComputeShader());
 }
 
+void BuildRayTracingInstanceBuffer(
+	FRHICommandList& RHICmdList,
+	FUnorderedAccessViewRHIRef InstancesUAV,
+	FShaderResourceViewRHIRef InstanceUploadSRV,
+	FShaderResourceViewRHIRef AccelerationStructureAddressesSRV,
+	FShaderResourceViewRHIRef CPUInstanceTransformSRV,
+	uint32 NumNativeCPUInstances,
+	TConstArrayView<FRayTracingGPUInstance> GPUInstances)
+{
+	RHICmdList.BeginUAVOverlap(InstancesUAV);
+
+	BuildRayTracingInstanceBuffer(
+		RHICmdList,
+		NumNativeCPUInstances,
+		0,
+		InstancesUAV,
+		InstanceUploadSRV,
+		AccelerationStructureAddressesSRV,
+		CPUInstanceTransformSRV);
+
+	for (const auto& GPUInstance : GPUInstances)
+	{
+		// GPU instance input descriptors are stored after CPU instances
+		uint32 InputDescOffset = NumNativeCPUInstances + GPUInstance.DescBufferOffset;
+
+		BuildRayTracingInstanceBuffer(
+			RHICmdList,
+			GPUInstance.NumInstances,
+			InputDescOffset,
+			InstancesUAV,
+			InstanceUploadSRV,
+			AccelerationStructureAddressesSRV,
+			GPUInstance.TransformSRV);
+	}
+
+	RHICmdList.EndUAVOverlap(InstancesUAV);
+}
+
 #endif //RHI_RAYTRACING
diff --git a/Engine/Source/Runtime/Renderer/Private/RayTracing/RayTracingScene.cpp b/Engine/Source/Runtime/Renderer/Private/RayTracing/RayTracingScene.cpp
--- a/Engine/Source/Runtime/Renderer/Private/RayTracing/RayTracingScene.cpp
+++ b/Engine/Source/Runtime/Renderer/Private/RayTracing/RayTracingScene.cpp
@@ -30,14 +30,19 @@ void FRayTracingScene::Create(FRDGBuilder& GraphBuilder)
 
 	WaitForTasks();
 
-	TArray<uint32> GeometryIndices;
-	RayTracingSceneRHI = CreateRayTracingSceneWithGeometryInstances(Instances, RAY_TRACING_NUM_SHADER_SLOTS, RAY_TRACING_NUM_MISS_SHADER_SLOTS, GeometryIndices);
+	FRayTracingSceneWithGeometryInstances SceneWithGeometryInstances = CreateRayTracingSceneWithGeometryInstances(
+		Instances,
+		RAY_TRACING_NUM_SHADER_SLOTS,
+		RAY_TRACING_NUM_MISS_SHADER_SLOTS);
+
+	RayTracingSceneRHI = SceneWithGeometryInstances.Scene;
 
 	const FRayTracingSceneInitializer2& SceneInitializer = RayTracingSceneRHI->GetInitializer();
 
 	// Round up number of instances to some multiple to avoid pathological growth reallocations.
 	static constexpr uint32 AllocationGranularity = 8 * 1024;
-	uint32 NumNativeInstancesAligned = FMath::DivideAndRoundUp(FMath::Max(SceneInitializer.NumNativeInstances, 1U), AllocationGranularity) * AllocationGranularity;
+	const uint32 NumNativeInstancesAligned = FMath::DivideAndRoundUp(FMath::Max(SceneInitializer.NumNativeInstances, 1U), AllocationGranularity) * AllocationGranularity;
+	const uint32 NumTransformsAligned = FMath::DivideAndRoundUp(FMath::Max(SceneWithGeometryInstances.NumNativeCPUInstances, 1U), AllocationGranularity) * AllocationGranularity;
 
 	SizeInfo = RHICalcRayTracingSceneSize(SceneInitializer.NumNativeInstances, ERayTracingAccelerationStructureFlags::FastTrace);
 	FRayTracingAccelerationStructureSize SizeInfoAligned = RHICalcRayTracingSceneSize(NumNativeInstancesAligned, ERayTracingAccelerationStructureFlags::FastTrace);
@@ -92,7 +97,7 @@ void FRayTracingScene::Create(FRDGBuilder& GraphBuilder)
 	}
 
 	{
-		// Lock instance upload buffer (and resize if necessary)
+		// Create/resize instance upload buffer (if necessary)
 		const uint32 UploadBufferSize = NumNativeInstancesAligned * sizeof(FRayTracingInstanceDescriptorInput);
 
 		if (!InstanceUploadBuffer.IsValid()
@@ -105,23 +110,47 @@ void FRayTracingScene::Create(FRDGBuilder& GraphBuilder)
 		}
 	}
 
+	{
+		// Create/resize transform upload buffer (if necessary)
+		const uint32 UploadBufferSize = NumTransformsAligned * sizeof(FVector4f) * 3;
+
+		if (!TransformUploadBuffer.IsValid()
+			|| UploadBufferSize > TransformUploadBuffer->GetSize()
+			|| UploadBufferSize < TransformUploadBuffer->GetSize() / 2)
+		{
+			FRHIResourceCreateInfo CreateInfo(TEXT("RayTracingSceneTransformUploadBuffer"));
+			TransformUploadBuffer = RHICreateStructuredBuffer(sizeof(FVector4f), UploadBufferSize, BUF_ShaderResource | BUF_Volatile, CreateInfo);
+			TransformUploadSRV = RHICreateShaderResourceView(TransformUploadBuffer);
+		}
+	}
+
 	if (SceneInitializer.NumNativeInstances > 0)
 	{
-		const uint32 UploadBytes = SceneInitializer.NumNativeInstances * sizeof(FRayTracingInstanceDescriptorInput);
+		const uint32 InstanceUploadBytes = SceneInitializer.NumNativeInstances * sizeof(FRayTracingInstanceDescriptorInput);
+		const uint32 TransformUploadBytes = SceneWithGeometryInstances.NumNativeCPUInstances * 3 * sizeof(FVector4f);
 
-		FRayTracingInstanceDescriptorInput* InstanceUploadData = (FRayTracingInstanceDescriptorInput*)RHILockBuffer(InstanceUploadBuffer, 0, UploadBytes, RLM_WriteOnly);
+		FRayTracingInstanceDescriptorInput* InstanceUploadData = (FRayTracingInstanceDescriptorInput*)RHILockBuffer(InstanceUploadBuffer, 0, InstanceUploadBytes, RLM_WriteOnly);
+		FVector4f* TransformUploadData = (FVector4f*)RHILockBuffer(TransformUploadBuffer, 0, TransformUploadBytes, RLM_WriteOnly);
 
 		// Fill instance upload buffer on separate thread since results are only needed in RHI thread
 		FillInstanceUploadBufferTask = FFunctionGraphTask::CreateAndDispatchWhenReady(
-			[InstanceUploadData,
-			NumNativeInstances = SceneInitializer.NumNativeInstances,
+			[InstanceUploadData = MakeArrayView(InstanceUploadData, SceneInitializer.NumNativeInstances),
+			TransformUploadData = MakeArrayView(TransformUploadData, SceneWithGeometryInstances.NumNativeCPUInstances * 3),
+			NumNativeCPUInstances = SceneWithGeometryInstances.NumNativeCPUInstances,
 			Instances = MakeArrayView(Instances),
-			GeometryIndices = MoveTemp(GeometryIndices),
+			InstanceGeometryIndices = MoveTemp(SceneWithGeometryInstances.InstanceGeometryIndices),
+			BaseUploadBufferOffsets = MoveTemp(SceneWithGeometryInstances.BaseUploadBufferOffsets),
 			RayTracingSceneRHI = RayTracingSceneRHI]()
 		{
 			FTaskTagScope TaskTagScope(ETaskTag::EParallelRenderingThread);
-
-			FillRayTracingInstanceUploadBuffer(Instances, GeometryIndices, RayTracingSceneRHI, MakeArrayView(InstanceUploadData, NumNativeInstances));
+			FillRayTracingInstanceUploadBuffer(
+				RayTracingSceneRHI,
+				Instances,
+				InstanceGeometryIndices,
+				BaseUploadBufferOffsets,
+				NumNativeCPUInstances,
+				InstanceUploadData,
+				TransformUploadData);
 		}, TStatId(), nullptr, ENamedThreads::AnyThread);
 
 		FBuildInstanceBufferPassParams* PassParams = GraphBuilder.AllocParameters<FBuildInstanceBufferPassParams>();
@@ -131,10 +160,16 @@ void FRayTracingScene::Create(FRDGBuilder& GraphBuilder)
 			RDG_EVENT_NAME("BuildTLASInstanceBuffer"),
 			PassParams,
 			ERDGPassFlags::Compute,
-			[PassParams, this, &SceneInitializer](FRHICommandListImmediate& RHICmdList)
+			[PassParams,
+			this,
+			&SceneInitializer,
+			NumNativeCPUInstances = SceneWithGeometryInstances.NumNativeCPUInstances,
+			GPUInstances = MoveTemp(SceneWithGeometryInstances.GPUInstances)
+			](FRHICommandListImmediate& RHICmdList)
 			{
 				WaitForTasks();
 				RHICmdList.UnlockBuffer(InstanceUploadBuffer);
+				RHICmdList.UnlockBuffer(TransformUploadBuffer);
 
 				RHICmdList.EnqueueLambda([this, &SceneInitializer](FRHICommandListImmediate& RHICmdList)
 					{
@@ -154,12 +189,14 @@ void FRayTracingScene::Create(FRDGBuilder& GraphBuilder)
 						RHICmdList.UnlockBuffer(AccelerationStructureAddressesBuffer.Buffer);
 					});
 
-				::BuildRayTracingInstanceBuffer(
+				BuildRayTracingInstanceBuffer(
 					RHICmdList,
-					SceneInitializer.NumNativeInstances,
 					PassParams->InstanceBuffer->GetRHI(),
 					InstanceUploadSRV,
-					AccelerationStructureAddressesBuffer.SRV);
+					AccelerationStructureAddressesBuffer.SRV,
+					TransformUploadSRV,
+					NumNativeCPUInstances,
+					GPUInstances);
 			});
 	}
 }
diff --git a/Engine/Source/Runtime/Renderer/Private/RayTracing/RayTracingScene.h b/Engine/Source/Runtime/Renderer/Private/RayTracing/RayTracingScene.h
--- a/Engine/Source/Runtime/Renderer/Private/RayTracing/RayTracingScene.h
+++ b/Engine/Source/Runtime/Renderer/Private/RayTracing/RayTracingScene.h
@@ -102,6 +102,9 @@ private:
 	FBufferRHIRef InstanceUploadBuffer;
 	FShaderResourceViewRHIRef InstanceUploadSRV;
 
+	FBufferRHIRef TransformUploadBuffer;
+	FShaderResourceViewRHIRef TransformUploadSRV;
+
 	FByteAddressBuffer AccelerationStructureAddressesBuffer;
 
 	mutable FGraphEventRef FillInstanceUploadBufferTask;
diff --git a/Engine/Source/Runtime/Renderer/Private/Tests/RayTracingTestbed.cpp b/Engine/Source/Runtime/Renderer/Private/Tests/RayTracingTestbed.cpp
--- a/Engine/Source/Runtime/Renderer/Private/Tests/RayTracingTestbed.cpp
+++ b/Engine/Source/Runtime/Renderer/Private/Tests/RayTracingTestbed.cpp
@@ -115,10 +115,12 @@ bool RunRayTracingTestbed_RenderThread(const FString& Parameters)
 	Instances[0].NumTransforms = NumTransforms;
 	Instances[0].Transforms = MakeArrayView(&FMatrix::Identity, 1);
 
-	TArray<uint32> GeometryIndices;
-	FRayTracingSceneRHIRef Scene = CreateRayTracingSceneWithGeometryInstances(Instances, RAY_TRACING_NUM_SHADER_SLOTS, 1, GeometryIndices);
+	FRayTracingSceneWithGeometryInstances RayTracingScene = CreateRayTracingSceneWithGeometryInstances(
+		Instances,
+		RAY_TRACING_NUM_SHADER_SLOTS,
+		1);
 
-	const FRayTracingSceneInitializer2& SceneInitializer = Scene->GetInitializer();
+	const FRayTracingSceneInitializer2& SceneInitializer = RayTracingScene.Scene->GetInitializer();
 
 	ERayTracingAccelerationStructureFlags SceneBuildFlags = ERayTracingAccelerationStructureFlags::FastTrace;
 	FRayTracingAccelerationStructureSize SceneSizeInfo = RHICalcRayTracingSceneSize(1, SceneBuildFlags);
@@ -144,17 +146,34 @@ bool RunRayTracingTestbed_RenderThread(const FString& Parameters)
 	FBufferRHIRef InstanceUploadBuffer;
 	FShaderResourceViewRHIRef InstanceUploadSRV;
 	{
-
 		FRHIResourceCreateInfo CreateInfo(TEXT("RayTracingTestBedInstanceUploadBuffer"));
 		InstanceUploadBuffer = RHICreateStructuredBuffer(sizeof(FRayTracingInstanceDescriptorInput), InstanceUploadBufferSize, BUF_ShaderResource | BUF_Volatile, CreateInfo);
 		InstanceUploadSRV = RHICreateShaderResourceView(InstanceUploadBuffer);
 	}
 
+	const uint32 TransformUploadBufferSize = RayTracingScene.NumNativeCPUInstances * 3 * sizeof(FVector4f);
+	FBufferRHIRef TransformUploadBuffer;
+	FShaderResourceViewRHIRef TransformUploadSRV;
+	{
+		FRHIResourceCreateInfo CreateInfo(TEXT("RayTracingTestBedInstanceUploadBuffer"));
+		TransformUploadBuffer = RHICreateStructuredBuffer(sizeof(FVector4f), TransformUploadBufferSize, BUF_ShaderResource | BUF_Volatile, CreateInfo);
+		TransformUploadSRV = RHICreateShaderResourceView(TransformUploadBuffer);
+	}
+
 	FRHICommandListImmediate& RHICmdList = FRHICommandListExecutor::GetImmediateCommandList();
 
 	{
 		FRayTracingInstanceDescriptorInput* InstanceUploadData = (FRayTracingInstanceDescriptorInput*)RHICmdList.LockBuffer(InstanceUploadBuffer, 0, InstanceUploadBufferSize, RLM_WriteOnly);
-		FillRayTracingInstanceUploadBuffer(Instances, GeometryIndices, Scene, MakeArrayView(InstanceUploadData, SceneInitializer.NumNativeInstances));
+		FVector4f* TransformUploadData = (FVector4f*)RHICmdList.LockBuffer(TransformUploadBuffer, 0, TransformUploadBufferSize, RLM_WriteOnly);
+		FillRayTracingInstanceUploadBuffer(
+			RayTracingScene.Scene,
+			Instances,
+			RayTracingScene.InstanceGeometryIndices,
+			RayTracingScene.BaseUploadBufferOffsets,
+			RayTracingScene.NumNativeCPUInstances,
+			MakeArrayView(InstanceUploadData, SceneInitializer.NumNativeInstances),
+			MakeArrayView(TransformUploadData, RayTracingScene.NumNativeCPUInstances * 3));
+		RHICmdList.UnlockBuffer(TransformUploadBuffer);
 		RHICmdList.UnlockBuffer(InstanceUploadBuffer);
 	}
 
@@ -176,12 +195,14 @@ bool RunRayTracingTestbed_RenderThread(const FString& Parameters)
 
 	BuildRayTracingInstanceBuffer(
 		RHICmdList,
-		SceneInitializer.NumNativeInstances,
 		InstanceBuffer.UAV,
 		InstanceUploadSRV,
-		AccelerationStructureAddressesBuffer.SRV);
+		AccelerationStructureAddressesBuffer.SRV,
+		TransformUploadSRV,
+		RayTracingScene.NumNativeCPUInstances,
+		{});
 
-	RHICmdList.BindAccelerationStructureMemory(Scene, SceneBuffer, 0);
+	RHICmdList.BindAccelerationStructureMemory(RayTracingScene.Scene, SceneBuffer, 0);
 
 	RHICmdList.BuildAccelerationStructure(Geometry);
 
@@ -189,7 +210,7 @@ bool RunRayTracingTestbed_RenderThread(const FString& Parameters)
 	// RHICmdList.Transition(FRHITransitionInfo(Geometry.GetReference(), ERHIAccess::BVHWrite, ERHIAccess::BVHRead));
 
 	FRayTracingSceneBuildParams Params;
-	Params.Scene = Scene;
+	Params.Scene = RayTracingScene.Scene;
 	Params.ScratchBuffer = ScratchBuffer;
 	Params.ScratchBufferOffset = 0;
 	Params.InstanceBuffer = InstanceBuffer.Buffer;
@@ -197,10 +218,10 @@ bool RunRayTracingTestbed_RenderThread(const FString& Parameters)
 
 	RHICmdList.BuildAccelerationStructure(Params);
 
-	RHICmdList.Transition(FRHITransitionInfo(Scene.GetReference(), ERHIAccess::BVHWrite, ERHIAccess::BVHRead));
+	RHICmdList.Transition(FRHITransitionInfo(RayTracingScene.Scene.GetReference(), ERHIAccess::BVHWrite, ERHIAccess::BVHRead));
 
-	RHICmdList.RayTraceOcclusion(Scene, RayBufferView, OcclusionResultBufferView, NumRays);
-	RHICmdList.RayTraceIntersection(Scene, RayBufferView, IntersectionResultBufferView, NumRays);
+	RHICmdList.RayTraceOcclusion(RayTracingScene.Scene, RayBufferView, OcclusionResultBufferView, NumRays);
+	RHICmdList.RayTraceIntersection(RayTracingScene.Scene, RayBufferView, IntersectionResultBufferView, NumRays);
 
 	const bool bValidateResults = true;
 	bool bOcclusionTestOK = false;
diff --git a/Engine/Source/Runtime/Renderer/Public/RayTracingInstanceBufferUtil.h b/Engine/Source/Runtime/Renderer/Public/RayTracingInstanceBufferUtil.h
--- a/Engine/Source/Runtime/Renderer/Public/RayTracingInstanceBufferUtil.h
+++ b/Engine/Source/Runtime/Renderer/Public/RayTracingInstanceBufferUtil.h
@@ -7,35 +7,73 @@
 
 #if RHI_RAYTRACING
 
+/*
+* 
+* Each FRayTracingGeometryInstance can translate to multiple native TLAS instances (see FRayTracingGeometryInstance::NumTransforms).
+* 
+* The FRayTracingGeometryInstance array (ie: FRayTracingScene::Instances) used to create FRayTracingSceneRHI
+* can have mix of instances using CPU transforms or GPUTransformSRV.
+* In order to reduce the number of dispatches to build the native RayTracing Instance Buffer,
+* the upload buffer containing FRayTracingInstanceDescriptorInput is split in 2 sections, CPU instances / GPU instances.
+* This way all native CPU instance descriptors can be built in a single dispatch.
+* Followed by one dispatch per GPU instance (since GPU transforms of each GPU instance are stored in separate buffers).
+* 
+*/
+
 struct FRayTracingInstanceDescriptorInput
 {
-	//uint32 GPUSceneInstanceIndex;
-	FVector4f LocalToWorld[3];
+	uint32 TransformIndex;
+	uint32 OutputDescriptorIndex;
 	uint32 AccelerationStructureIndex;
 	uint32 InstanceId;
 	uint32 InstanceMaskAndFlags;
 	uint32 InstanceContributionToHitGroupIndex;
 };
 
+struct FRayTracingGPUInstance
+{
+	FShaderResourceViewRHIRef TransformSRV;
+	uint32 NumInstances;
+	uint32 DescBufferOffset;
+};
+
+struct FRayTracingSceneWithGeometryInstances
+{
+	FRayTracingSceneRHIRef Scene;
+	uint32 NumNativeCPUInstances;
+	uint32 NumNativeGPUInstances;
+	// index of each instance geometry in FRayTracingSceneRHIRef ReferencedGeometries
+	TArray<uint32> InstanceGeometryIndices;
+	// base offset of each instance entries in the instance upload buffer
+	TArray<uint32> BaseUploadBufferOffsets;
+	TArray<FRayTracingGPUInstance> GPUInstances;
+};
+
 // Helper function to create FRayTracingSceneRHI using array of high level instances
 // Also outputs data required to build the instance buffer
-RENDERER_API FRayTracingSceneRHIRef CreateRayTracingSceneWithGeometryInstances(
+RENDERER_API FRayTracingSceneWithGeometryInstances CreateRayTracingSceneWithGeometryInstances(
 	TArrayView<FRayTracingGeometryInstance> Instances,
 	uint32 NumShaderSlotsPerGeometrySegment,
-	uint32 NumMissShaderSlots,
-	TArray<uint32>& OutGeometryIndices);
+	uint32 NumMissShaderSlots);
 
+// Helper function to fill upload buffers required by BuildRayTracingInstanceBuffer with instance descriptors
+// Transforms of CPU instances are copied to OutTransformData
 RENDERER_API void FillRayTracingInstanceUploadBuffer(
-	TConstArrayView<FRayTracingGeometryInstance> Instances,
-	TConstArrayView<uint32> InstancesGeometryIndex,
 	FRayTracingSceneRHIRef RayTracingSceneRHI,
-	TArrayView<FRayTracingInstanceDescriptorInput> OutInstanceUploadData);
+	TConstArrayView<FRayTracingGeometryInstance> Instances,
+	TConstArrayView<uint32> InstanceGeometryIndices,
+	TConstArrayView<uint32> BaseUploadBufferOffsets,
+	uint32 NumNativeCPUInstances,
+	TArrayView<FRayTracingInstanceDescriptorInput> OutInstanceUploadData,
+	TArrayView<FVector4f> OutTransformData);
 
 RENDERER_API void BuildRayTracingInstanceBuffer(
-	FRHICommandList& RHICmdList, 
-	uint32 NumInstances,
+	FRHICommandList& RHICmdList,
 	FUnorderedAccessViewRHIRef InstancesUAV,
 	FShaderResourceViewRHIRef InstanceUploadSRV,
-	FShaderResourceViewRHIRef AccelerationStructureAddressesSRV);
+	FShaderResourceViewRHIRef AccelerationStructureAddressesSRV,
+	FShaderResourceViewRHIRef CPUInstanceTransformSRV,
+	uint32 NumNativeCPUInstances,
+	TConstArrayView<FRayTracingGPUInstance> GPUInstances);
 
 #endif // RHI_RAYTRACING
\ No newline at end of file
