Allow distance field lighting to operate with ray-traced lighting.

This is necessary to enable dynamic sky light with ray-traced global illumination.

diff --git a/Engine/Source/Runtime/Renderer/Private/DeferredShadingRenderer.h b/Engine/Source/Runtime/Renderer/Private/DeferredShadingRenderer.h
--- a/Engine/Source/Runtime/Renderer/Private/DeferredShadingRenderer.h
+++ b/Engine/Source/Runtime/Renderer/Private/DeferredShadingRenderer.h
@@ -335,19 +335,19 @@ private:
 		return (GetViewPipelineState(View).DiffuseIndirectMethod == EDiffuseIndirectMethod::Lumen || GetViewPipelineState(View).ReflectionsMethod == EReflectionsMethod::Lumen);
 	}
 
-	virtual bool AnyViewHasSupportingGIMethod() const override
+	virtual bool AnyViewHasGIMethodSupportingDFAO() const override
 	{
-		bool bAnyViewHasSupportingGIMethod = false;
+		bool bAnyViewHasGIMethodSupportingDFAO = false;
 
 		for (int32 ViewIndex = 0; ViewIndex < Views.Num(); ViewIndex++)
 		{
-			if (GetViewPipelineState(Views[ViewIndex]).DiffuseIndirectMethod == EDiffuseIndirectMethod::Disabled || GetViewPipelineState(Views[ViewIndex]).DiffuseIndirectMethod == EDiffuseIndirectMethod::SSGI)
+			if (GetViewPipelineState(Views[ViewIndex]).DiffuseIndirectMethod != EDiffuseIndirectMethod::Lumen)
 			{
-				bAnyViewHasSupportingGIMethod = true;
+				bAnyViewHasGIMethodSupportingDFAO = true;
 			}
 		}
 
-		return bAnyViewHasSupportingGIMethod;
+		return bAnyViewHasGIMethodSupportingDFAO;
 	}
 
 	static FGraphEventRef TranslucencyTimestampQuerySubmittedFence[FOcclusionQueryHelpers::MaxBufferedOcclusionFrames + 1];
diff --git a/Engine/Source/Runtime/Renderer/Private/DistanceFieldAmbientOcclusion.cpp b/Engine/Source/Runtime/Renderer/Private/DistanceFieldAmbientOcclusion.cpp
--- a/Engine/Source/Runtime/Renderer/Private/DistanceFieldAmbientOcclusion.cpp
+++ b/Engine/Source/Runtime/Renderer/Private/DistanceFieldAmbientOcclusion.cpp
@@ -697,14 +697,14 @@ bool ShouldRenderDeferredDynamicSkyLight(const FScene* Scene, const FSceneViewFa
 
 bool FSceneRenderer::ShouldPrepareForDistanceFieldAO() const
 {
-	bool bAnyViewHasSupportingGIMethod = AnyViewHasSupportingGIMethod();
+	bool bAnyViewHasGIMethodSupportingDFAO = AnyViewHasGIMethodSupportingDFAO();
 
 	return SupportsDistanceFieldAO(Scene->GetFeatureLevel(), Scene->GetShaderPlatform())
-		&& ((ShouldRenderDeferredDynamicSkyLight(Scene, ViewFamily) && bAnyViewHasSupportingGIMethod && Scene->SkyLight->bCastShadows && ViewFamily.EngineShowFlags.DistanceFieldAO)
+		&& ((ShouldRenderDeferredDynamicSkyLight(Scene, ViewFamily) && bAnyViewHasGIMethodSupportingDFAO && Scene->SkyLight->bCastShadows && ViewFamily.EngineShowFlags.DistanceFieldAO)
 			|| ViewFamily.EngineShowFlags.VisualizeMeshDistanceFields
 			|| ViewFamily.EngineShowFlags.VisualizeGlobalDistanceField
 			|| ViewFamily.EngineShowFlags.VisualizeDistanceFieldAO
-			|| (GDistanceFieldAOApplyToStaticIndirect && bAnyViewHasSupportingGIMethod && ViewFamily.EngineShowFlags.DistanceFieldAO));
+			|| (GDistanceFieldAOApplyToStaticIndirect && bAnyViewHasGIMethodSupportingDFAO && ViewFamily.EngineShowFlags.DistanceFieldAO));
 }
 
 bool FSceneRenderer::ShouldPrepareDistanceFieldScene() const
diff --git a/Engine/Source/Runtime/Renderer/Private/IndirectLightRendering.cpp b/Engine/Source/Runtime/Renderer/Private/IndirectLightRendering.cpp
--- a/Engine/Source/Runtime/Renderer/Private/IndirectLightRendering.cpp
+++ b/Engine/Source/Runtime/Renderer/Private/IndirectLightRendering.cpp
@@ -1369,15 +1369,10 @@ void FDeferredShadingSceneRenderer::RenderDeferredReflectionsAndSkyLighting(
 	// If we're currently capturing a reflection capture, output SpecularColor * IndirectIrradiance for metals so they are not black in reflections,
 	// Since we don't have multiple bounce specular reflections
 	bool bReflectionCapture = false;
-	bool bGIMethodSupportsDFAO = false;
 	for (int32 ViewIndex = 0, Num = Views.Num(); ViewIndex < Num; ViewIndex++)
 	{
 		const FViewInfo& View = Views[ViewIndex];
 		bReflectionCapture = bReflectionCapture || View.bIsReflectionCapture;
-
-		bGIMethodSupportsDFAO = bGIMethodSupportsDFAO 
-			|| GetViewPipelineState(Views[ViewIndex]).DiffuseIndirectMethod == EDiffuseIndirectMethod::Disabled 
-			|| GetViewPipelineState(Views[ViewIndex]).DiffuseIndirectMethod == EDiffuseIndirectMethod::SSGI;
 	}
 
 	if (bReflectionCapture)
@@ -1391,7 +1386,7 @@ void FDeferredShadingSceneRenderer::RenderDeferredReflectionsAndSkyLighting(
 		&& (Scene->SkyLight->ProcessedTexture || Scene->SkyLight->bRealTimeCaptureEnabled)
 		&& !Scene->SkyLight->bHasStaticLighting;
 
-	bool bDynamicSkyLight = ShouldRenderDeferredDynamicSkyLight(Scene, ViewFamily) && bGIMethodSupportsDFAO;
+	bool bDynamicSkyLight = ShouldRenderDeferredDynamicSkyLight(Scene, ViewFamily) && AnyViewHasGIMethodSupportingDFAO();
 	bool bApplySkyShadowing = false;
 	if (bDynamicSkyLight)
 	{
diff --git a/Engine/Source/Runtime/Renderer/Private/SceneRendering.h b/Engine/Source/Runtime/Renderer/Private/SceneRendering.h
--- a/Engine/Source/Runtime/Renderer/Private/SceneRendering.h
+++ b/Engine/Source/Runtime/Renderer/Private/SceneRendering.h
@@ -2007,7 +2007,7 @@ public:
 	void PrepareDistanceFieldScene(FRDGBuilder& GraphBuilder, bool bSplitDispatch);
 
 	virtual bool IsLumenEnabled(const FViewInfo& View) const { return false; }
-	virtual bool AnyViewHasSupportingGIMethod() const { return false; }
+	virtual bool AnyViewHasGIMethodSupportingDFAO() const { return false; }
 
 	/** Gets a readable light name for use with a draw event. */
 	static void GetLightNameForDrawEvent(const FLightSceneProxy* LightProxy, FString& LightNameWithLevel);
