Make rough diffuse toggleable at runtime when using strata.

diff --git a/Engine/Shaders/Private/Strata/StrataEvaluation.ush b/Engine/Shaders/Private/Strata/StrataEvaluation.ush
--- a/Engine/Shaders/Private/Strata/StrataEvaluation.ush
+++ b/Engine/Shaders/Private/Strata/StrataEvaluation.ush
@@ -285,11 +285,15 @@ FStrataEvaluateResult StrataEvaluateBSDF(FStrataBSDFContext BSDFContext, FShadow
 			const float SafeRoughness = MakeRoughnessSafe(SLAB_ROUGHNESSX(BSDFContext.BSDF));
 			float a2 = Pow4(SafeRoughness);
 
-		#if STRATA_DIFFUSE_CHAN
-			Sample.DiffusePathValue = Diffuse_Chan(DiffuseColor, a2, BSDFContext.SatNoV, BSDFContext.SatNoL, BSDFContext.Context.VoH, BSDFContext.Context.NoH);
-		#else
-			Sample.DiffusePathValue = Diffuse_Lambert(DiffuseColor);
-		#endif
+			if (STRATA_DIFFUSE_CHAN && Strata.bRoughDiffuse)
+			{
+				Sample.DiffusePathValue = Diffuse_Chan(DiffuseColor, a2, BSDFContext.SatNoV, BSDFContext.SatNoL, BSDFContext.Context.VoH, BSDFContext.Context.NoH);
+			}
+			else
+			{
+				Sample.DiffusePathValue = Diffuse_Lambert(DiffuseColor);
+			}
+		
 			Sample.DiffusePDF		= BSDFContext.SatNoL / PI;
 			Sample.bSubsurface		= bHasSubsurface;
 
@@ -385,7 +389,10 @@ FStrataEvaluateResult StrataEvaluateBSDF(FStrataBSDFContext BSDFContext, FShadow
 				// To make this a safe operation:
 				//	- We clamp SpecularTChan to a threshold to avoid div by 0 when the material is a metal. 
 				//	- When the material tends to be metal, DirectionalAlbedo_SpecularTransmission will naturally tend to be 0
-				Sample.DiffusePathValue *= DirectionalAlbedo_SpecularTransmission / max(1e-8f, StrataGGXLUTResult.DirectionalAlbedo.SpecularTChan);
+				if (Strata.bRoughDiffuse)
+				{
+					Sample.DiffusePathValue *= DirectionalAlbedo_SpecularTransmission / max(1e-8f, StrataGGXLUTResult.DirectionalAlbedo.SpecularTChan);
+				}
 				#endif
 			}
 
diff --git a/Engine/Source/Runtime/Renderer/Private/Strata/Strata.cpp b/Engine/Source/Runtime/Renderer/Private/Strata/Strata.cpp
--- a/Engine/Source/Runtime/Renderer/Private/Strata/Strata.cpp
+++ b/Engine/Source/Runtime/Renderer/Private/Strata/Strata.cpp
@@ -69,6 +69,12 @@ static TAutoConsoleVariable<int32> CVarStrataLUTContinousUpdate(
 	TEXT("Update Strata energy LUT every frame (for debug purpose)."),
 	ECVF_RenderThreadSafe);
 
+static TAutoConsoleVariable<int32> CVarStrataRoughDiffuse(
+	TEXT("r.Strata.RoughDiffuse"),
+	1,
+	TEXT("Enable Strata rough diffuse model (works only if r.Material.RoughDiffuse is enabled in the project settings). Togglable at runtime"),
+	ECVF_RenderThreadSafe);
+
 static TAutoConsoleVariable<int32> CVarStrataFurnaceTest(
 	TEXT("r.Strata.FurnaceTest"),
 	0,
@@ -244,6 +250,9 @@ void InitialiseStrataFrameSceneData(FSceneRenderer& SceneRenderer, FRDGBuilder&
 	StrataSceneData.MaterialLobesBufferSRV = GraphBuilder.CreateSRV(StrataSceneData.MaterialLobesBuffer);
 	StrataSceneData.MaterialLobesBufferUAV = GraphBuilder.CreateUAV(StrataSceneData.MaterialLobesBuffer);
 
+	// Rough diffuse model
+	StrataSceneData.bRoughDiffuse = CVarStrataRoughDiffuse.GetValueOnRenderThread() > 0 ? 1u : 0u;
+
 	// Set reference to the Strata data from each view
 	for (int32 ViewIndex = 0; ViewIndex < SceneRenderer.Views.Num(); ViewIndex++)
 	{
@@ -264,6 +273,7 @@ void InitialiseStrataFrameSceneData(FSceneRenderer& SceneRenderer, FRDGBuilder&
 	if (IsStrataEnabled())
 	{
 		FStrataGlobalUniformParameters* StrataUniformParameters = GraphBuilder.AllocParameters<FStrataGlobalUniformParameters>();
+		StrataUniformParameters->bRoughDiffuse = StrataSceneData.bRoughDiffuse ? 1u : 0u;
 		StrataUniformParameters->MaxBytesPerPixel = StrataSceneData.MaxBytesPerPixel;
 		StrataUniformParameters->MaterialLobesBuffer = StrataSceneData.MaterialLobesBufferSRV;
 		StrataUniformParameters->ClassificationTexture = StrataSceneData.ClassificationTexture;
@@ -284,6 +294,7 @@ void BindStrataBasePassUniformParameters(FRDGBuilder& GraphBuilder, FStrataScene
 
 	if (IsStrataEnabled() && StrataSceneData)
 	{
+		OutStrataUniformParameters.bRoughDiffuse = StrataSceneData->bRoughDiffuse ? 1u : 0u;
 		OutStrataUniformParameters.MaxBytesPerPixel = StrataSceneData->MaxBytesPerPixel;
 		OutStrataUniformParameters.MaterialLobesBufferUAV = StrataSceneData->MaterialLobesBufferUAV;
 		OutStrataUniformParameters.GGXEnergyLUT3DTexture = StrataSceneData->GGXEnergyLUT3DTexture->GetRenderTargetItem().ShaderResourceTexture;
@@ -291,6 +302,7 @@ void BindStrataBasePassUniformParameters(FRDGBuilder& GraphBuilder, FStrataScene
 	}
 	else
 	{
+		OutStrataUniformParameters.bRoughDiffuse = 0u;
 		OutStrataUniformParameters.MaxBytesPerPixel = 0;
 		OutStrataUniformParameters.MaterialLobesBufferUAV = GraphBuilder.CreateUAV(GraphBuilder.RegisterExternalBuffer(GWhiteVertexBufferWithRDG->Buffer), PF_R32_UINT);
 		OutStrataUniformParameters.GGXEnergyLUT3DTexture = GSystemTextures.VolumetricBlackDummy->GetRenderTargetItem().ShaderResourceTexture;
diff --git a/Engine/Source/Runtime/Renderer/Private/Strata/Strata.h b/Engine/Source/Runtime/Renderer/Private/Strata/Strata.h
--- a/Engine/Source/Runtime/Renderer/Private/Strata/Strata.h
+++ b/Engine/Source/Runtime/Renderer/Private/Strata/Strata.h
@@ -17,6 +17,7 @@ struct FScreenPassTexture;
 
 BEGIN_SHADER_PARAMETER_STRUCT(FStrataBasePassUniformParameters, )
 	SHADER_PARAMETER(uint32, MaxBytesPerPixel)
+	SHADER_PARAMETER(uint32, bRoughDiffuse)
 	SHADER_PARAMETER(FVector2D, GGXEnergyLUTScaleBias)
 	SHADER_PARAMETER_TEXTURE(Texture3D<float3>, GGXEnergyLUT3DTexture)
 	SHADER_PARAMETER_TEXTURE(Texture2D<float4>, GGXEnergyLUT2DTexture)
@@ -26,6 +27,7 @@ END_SHADER_PARAMETER_STRUCT()
 
 BEGIN_GLOBAL_SHADER_PARAMETER_STRUCT(FStrataGlobalUniformParameters, )
 	SHADER_PARAMETER(uint32, MaxBytesPerPixel)
+	SHADER_PARAMETER(uint32, bRoughDiffuse)
 	SHADER_PARAMETER(FVector2D, GGXEnergyLUTScaleBias)
 	SHADER_PARAMETER_RDG_BUFFER_SRV(ByteAddressBuffer, MaterialLobesBuffer)
 	SHADER_PARAMETER_RDG_TEXTURE(Texture2D<uint>, ClassificationTexture)
@@ -46,6 +48,7 @@ enum EStrataTileMaterialType : uint32
 struct FStrataSceneData
 {
 	uint32 MaxBytesPerPixel;
+	bool bRoughDiffuse;
 
 	// Resources allocated and updated each frame
 
