Fix skylight intensity not updating through sequencer

The skylight intensity and color is baked as part of the skydome importance tables for the path tracer, RT skylight and RT reflections. If this change is done without creating a new skylight proxy (for example when keyframed via the sequencer) it is possible the change to go un-noticed. For a similar reason, the path tracer would not restart sample accumulation in this case.

The fix is to just store the color used the last time the importance tables were baked and detect when this has changed. An alternative fix would be to leave the multiplication by the skylight color outside of the importance tables, but it seems better to fold this in as skylight intensity changes are probably rare.

diff --git a/Engine/Source/Runtime/Renderer/Private/PathTracing.cpp b/Engine/Source/Runtime/Renderer/Private/PathTracing.cpp
--- a/Engine/Source/Runtime/Renderer/Private/PathTracing.cpp
+++ b/Engine/Source/Runtime/Renderer/Private/PathTracing.cpp
@@ -874,10 +874,12 @@ bool PrepareSkyTexture(FRDGBuilder& GraphBuilder, FScene* Scene, const FViewInfo
 
 	// should we remember the skylight prep for the next frame?
 	const bool IsSkylightCachingEnabled = CVarPathTracingSkylightCaching.GetValueOnAnyThread() != 0;
-
-	if (!IsSkylightCachingEnabled)
+	FLinearColor SkyColor = Scene->SkyLight->GetEffectiveLightColor();
+	const bool bSkylightColorChanged = SkyColor != Scene->PathTracingSkylightColor;
+	if (!IsSkylightCachingEnabled || bSkylightColorChanged)
 	{
-		// we don't want any caching - release what we might have been holding onto
+		// we don't want any caching (or the light color changed)
+		// release what we might have been holding onto so we get the right texture for this frame
 		Scene->PathTracingSkylightTexture.SafeRelease();
 		Scene->PathTracingSkylightPdf.SafeRelease();
 	}
@@ -894,8 +896,7 @@ bool PrepareSkyTexture(FRDGBuilder& GraphBuilder, FScene* Scene, const FViewInfo
 		return true;
 	}
 	RDG_EVENT_SCOPE(GraphBuilder, "Path Tracing SkylightPrepare");
-
-	FLinearColor SkyColor = Scene->SkyLight->GetEffectiveLightColor();
+	Scene->PathTracingSkylightColor = SkyColor;
 	// since we are resampled into an octahedral layout, we multiply the cubemap resolution by 2 to get roughly the same number of texels
 	uint32 Size = FMath::RoundUpToPowerOfTwo(2 * Scene->SkyLight->CaptureCubeMapResolution);
 
@@ -1452,6 +1453,16 @@ void FDeferredShadingSceneRenderer::RenderPathTracing(
 		Scene->PathTracingSkylightPdf.SafeRelease();
 	}
 
+	// if the skylight has changed colors, reset both the path tracer and the importance tables
+	if (Scene->SkyLight && Scene->SkyLight->GetEffectiveLightColor() != Scene->PathTracingSkylightColor)
+	{
+		Scene->PathTracingSkylightTexture.SafeRelease();
+		Scene->PathTracingSkylightPdf.SafeRelease();
+		// reset last color here as well in case we don't reach PrepareSkyLightTexture
+		Scene->PathTracingSkylightColor = Scene->SkyLight->GetEffectiveLightColor();
+		View.ViewState->PathTracingInvalidate();
+	}
+
 	// If the scene has changed in some way (camera move, object movement, etc ...)
 	// we must invalidate the ViewState to start over from scratch
 	if (FirstTime || Config.IsDifferent(PathTracingState->LastConfig) || HairStrands::HasPositionsChanged(GraphBuilder, View))
diff --git a/Engine/Source/Runtime/Renderer/Private/RendererScene.cpp b/Engine/Source/Runtime/Renderer/Private/RendererScene.cpp
--- a/Engine/Source/Runtime/Renderer/Private/RendererScene.cpp
+++ b/Engine/Source/Runtime/Renderer/Private/RendererScene.cpp
@@ -1012,6 +1012,7 @@ FScene::FScene(UWorld* InWorld, bool bInRequiresHitProxies, bool bInIsEditorScen
 ,	RealTimeSlicedReflectionCaptureFirstFrameState(ERealTimeSlicedReflectionCaptureFirstFrameState::INIT)
 ,	RealTimeSlicedReflectionCaptureState(-1)
 ,	RealTimeSlicedReflectionCaptureFrameNumber(0)
+,   PathTracingSkylightColor(0, 0, 0, 0)
 ,	SimpleDirectionalLight(NULL)
 ,	ReflectionSceneData(InFeatureLevel)
 ,	IndirectLightingCache(InFeatureLevel)
diff --git a/Engine/Source/Runtime/Renderer/Private/ScenePrivate.h b/Engine/Source/Runtime/Renderer/Private/ScenePrivate.h
--- a/Engine/Source/Runtime/Renderer/Private/ScenePrivate.h
+++ b/Engine/Source/Runtime/Renderer/Private/ScenePrivate.h
@@ -2789,9 +2789,12 @@ public:
 	 * The path tracer uses its own representation of the skylight. These textures
 	 * are updated lazily by the path tracer when missing. Any code that modifies
 	 * the skylight appearance should simply reset these pointers.
+	 * 
+	 * We also remember the last used color so we can detect changes that would require rebuilding the tables
 	 */
 	TRefCountPtr<IPooledRenderTarget> PathTracingSkylightTexture;
 	TRefCountPtr<IPooledRenderTarget> PathTracingSkylightPdf;
+	FLinearColor PathTracingSkylightColor;
 
 	/** Used to track the order that skylights were enabled in. */
 	TArray<FSkyLightSceneProxy*> SkyLightStack;
