Do not extract memoryless textures from RDG

diff --git a/Engine/Source/Runtime/Renderer/Private/MobileShadingRenderer.cpp b/Engine/Source/Runtime/Renderer/Private/MobileShadingRenderer.cpp
--- a/Engine/Source/Runtime/Renderer/Private/MobileShadingRenderer.cpp
+++ b/Engine/Source/Runtime/Renderer/Private/MobileShadingRenderer.cpp
@@ -925,6 +925,8 @@ void FMobileSceneRenderer::Render(FRDGBuilder& GraphBuilder)
 	RenderFinish(GraphBuilder, ViewFamilyTexture);
 
 	PollOcclusionQueries(GraphBuilder);
+
+	QueueSceneTextureExtractions(GraphBuilder, SceneTextures);
 }
 
 void FMobileSceneRenderer::BuildInstanceCullingDrawParams(FRDGBuilder& GraphBuilder, FViewInfo& View, FMobileRenderPassParameters* PassParameters)
@@ -1054,8 +1056,6 @@ void FMobileSceneRenderer::RenderForward(FRDGBuilder& GraphBuilder, FRDGTextureR
 			RenderForwardSinglePass(GraphBuilder, PassParameters, ViewIndex, View, SceneTextures);
 		}
 	}
-
-	QueueSceneTextureExtractions(GraphBuilder, SceneTextures);
 }
 
 void FMobileSceneRenderer::RenderForwardSinglePass(FRDGBuilder& GraphBuilder, FMobileRenderPassParameters* PassParameters, int32 ViewIndex, FViewInfo& View, FSceneTextures& SceneTextures)
@@ -1337,8 +1337,6 @@ void FMobileSceneRenderer::RenderDeferred(FRDGBuilder& GraphBuilder, const FSort
 			RenderDeferredSinglePass(GraphBuilder, PassParameters, ViewIndex, NumViews, View, SceneTextures, SortedLightSet, bUsingPixelLocalStorage);
 		}
 	}
-
-	QueueSceneTextureExtractions(GraphBuilder, SceneTextures);
 }
 
 void FMobileSceneRenderer::RenderDeferredSinglePass(FRDGBuilder& GraphBuilder, class FMobileRenderPassParameters* PassParameters, int32 ViewIndex, int32 NumViews, FViewInfo& View, FSceneTextures& SceneTextures, const FSortedLightSetSceneInfo& SortedLightSet, bool bUsingPixelLocalStorage)
diff --git a/Engine/Source/Runtime/Renderer/Private/SceneTextures.cpp b/Engine/Source/Runtime/Renderer/Private/SceneTextures.cpp
--- a/Engine/Source/Runtime/Renderer/Private/SceneTextures.cpp
+++ b/Engine/Source/Runtime/Renderer/Private/SceneTextures.cpp
@@ -923,7 +923,7 @@ void FSceneTextureExtracts::QueueExtractions(FRDGBuilder& GraphBuilder, const FS
 
 	const auto ExtractIfProduced = [&](FRDGTextureRef Texture, TRefCountPtr<IPooledRenderTarget>& OutTarget)
 	{
-		if (HasBeenProduced(Texture))
+		if (HasBeenProduced(Texture) && !EnumHasAnyFlags(Texture->Desc.Flags, TexCreate_Memoryless))
 		{
 			GraphBuilder.QueueTextureExtraction(Texture, &OutTarget);
 		}
