Fix one pass projection
- VSM is not yet allocated at this point in the frame; IsEnabled is more appropriate

diff --git a/Engine/Source/Runtime/Renderer/Private/LightGridInjection.cpp b/Engine/Source/Runtime/Renderer/Private/LightGridInjection.cpp
--- a/Engine/Source/Runtime/Renderer/Private/LightGridInjection.cpp
+++ b/Engine/Source/Runtime/Renderer/Private/LightGridInjection.cpp
@@ -790,7 +790,7 @@ void FDeferredShadingSceneRenderer::GatherLightsAndComputeLightGrid(FRDGBuilder&
 {
 	bool bShadowedLightsInClustered = ShouldUseClusteredDeferredShading()
 		&& CVarVirtualShadowOnePassProjection.GetValueOnRenderThread()
-		&& VirtualShadowMapArray.IsAllocated();
+		&& VirtualShadowMapArray.IsEnabled();
 
 	GatherAndSortLights(SortedLightSet, bShadowedLightsInClustered);
 	
