Enable Diaphgram DOF on SM6 via the DataDrivenPlatform info, and simplify related special cases for other platforms

diff --git a/Engine/Config/Mac/DataDrivenPlatformInfo.ini b/Engine/Config/Mac/DataDrivenPlatformInfo.ini
--- a/Engine/Config/Mac/DataDrivenPlatformInfo.ini
+++ b/Engine/Config/Mac/DataDrivenPlatformInfo.ini
@@ -33,4 +33,6 @@ bSupportsSSDIndirect=true
 bSupportsWaveOperations=false
 bSupportsGPUScene=true
 bSupportsFFTBloom=true
+bSupportsDiaphragmDOF=true
 bSupportsDxc=true
+
diff --git a/Engine/Config/VulkanPC/DataDrivenPlatformInfo.ini b/Engine/Config/VulkanPC/DataDrivenPlatformInfo.ini
--- a/Engine/Config/VulkanPC/DataDrivenPlatformInfo.ini
+++ b/Engine/Config/VulkanPC/DataDrivenPlatformInfo.ini
@@ -114,3 +114,4 @@ NumberOfComputeThreads = 64
 
 bSupportsHairStrandGeometry=true
 bSupportsFFTBloom=true
+bSupportsDiaphragmDOF=true
diff --git a/Engine/Config/Windows/DataDrivenPlatformInfo.ini b/Engine/Config/Windows/DataDrivenPlatformInfo.ini
--- a/Engine/Config/Windows/DataDrivenPlatformInfo.ini
+++ b/Engine/Config/Windows/DataDrivenPlatformInfo.ini
@@ -77,6 +77,7 @@ bSupportsVariableRateShading=true
 bSupportsGen5TemporalAA=true
 bSupportsHairStrandGeometry=true
 bSupportsFFTBloom=true
+bSupportsDiaphragmDOF=true
 
 [PreviewPlatform PCD3D_SM5]
 PlatformName=PCD3D_SM5
@@ -102,6 +103,7 @@ bSupportsDistanceFields=true
 bSupportsSSDIndirect=true
 bSupportsWaveOperations=true
 bSupportsDxc=true
+bSupportsDiaphragmDOF=true
 ; TODO: SM6.6
 bRequiresVendorExtensionsForAtomics=true
 bSupportsAnisotropicMaterials=true
diff --git a/Engine/Source/Runtime/Renderer/Private/PostProcess/DiaphragmDOF.h b/Engine/Source/Runtime/Renderer/Private/PostProcess/DiaphragmDOF.h
--- a/Engine/Source/Runtime/Renderer/Private/PostProcess/DiaphragmDOF.h
+++ b/Engine/Source/Runtime/Renderer/Private/PostProcess/DiaphragmDOF.h
@@ -119,11 +119,7 @@ struct FBokehModel
 inline bool IsSupported(const FStaticShaderPlatform ShaderPlatform)
 {
 	// Only compile diaphragm DOF on platform it has been tested to ensure this is not blocking anyone else.
-	return 
-		ShaderPlatform == SP_PCD3D_SM5 ||
-		IsVulkanSM5Platform(ShaderPlatform) ||
-		IsMetalSM5Platform(ShaderPlatform) ||
-		FDataDrivenShaderPlatformInfo::GetSupportsDiaphragmDOF(ShaderPlatform);
+	return FDataDrivenShaderPlatformInfo::GetSupportsDiaphragmDOF(ShaderPlatform);
 }
 
 
