Enable Path Tracing by default

This cvar mainly controls the creation of path tracer specific material permutations. Projects which know they do not want to make use the path tracer should turn this off. Note that the DataDrivenPlatformInfo is also checked to make sure the path tracer is supported on a given platform.

#p4v-preflight-copy 18094497

diff --git a/Engine/Source/Runtime/Renderer/Private/PathTracing.cpp b/Engine/Source/Runtime/Renderer/Private/PathTracing.cpp
--- a/Engine/Source/Runtime/Renderer/Private/PathTracing.cpp
+++ b/Engine/Source/Runtime/Renderer/Private/PathTracing.cpp
@@ -8,7 +8,7 @@ PathTracingDenoiserFunction* GPathTracingDenoiserFunc = nullptr;
 
 TAutoConsoleVariable<int32> CVarPathTracing(
 	TEXT("r.PathTracing"),
-	0,
+	1,
 	TEXT("Enables the path tracing renderer (to guard the compilation of path tracer specific material permutations)"),
 	ECVF_RenderThreadSafe | ECVF_ReadOnly
 );
