Fix a typo in VirtualTextureMaterialDraw_PS

diff --git a/Engine/Source/Runtime/Renderer/Private/VT/RuntimeVirtualTextureRender.cpp b/Engine/Source/Runtime/Renderer/Private/VT/RuntimeVirtualTextureRender.cpp
--- a/Engine/Source/Runtime/Renderer/Private/VT/RuntimeVirtualTextureRender.cpp
+++ b/Engine/Source/Runtime/Renderer/Private/VT/RuntimeVirtualTextureRender.cpp
@@ -372,7 +372,7 @@ namespace RuntimeVirtualTexture
 	class FShader_VirtualTextureMaterialDraw_PS : public FShader_VirtualTextureMaterialDraw
 	{
 	public:
-		DECLARE_SHADER_TYPE(FShader_VirtualTextureMaterialDraw_VS< MaterialPolicy >, MeshMaterial);
+		DECLARE_SHADER_TYPE(FShader_VirtualTextureMaterialDraw_PS< MaterialPolicy >, MeshMaterial);
 
 		FShader_VirtualTextureMaterialDraw_PS()
 		{}
