Fix Distortion Pass when SceneDepth Aux is not required.

diff --git a/Engine/Source/Runtime/Renderer/Private/MobileDistortionPass.cpp b/Engine/Source/Runtime/Renderer/Private/MobileDistortionPass.cpp
--- a/Engine/Source/Runtime/Renderer/Private/MobileDistortionPass.cpp
+++ b/Engine/Source/Runtime/Renderer/Private/MobileDistortionPass.cpp
@@ -48,13 +48,13 @@ TRDGUniformBufferRef<FMobileDistortionPassUniformParameters> CreateMobileDistort
 		SetupMode |= EMobileSceneTextureSetupMode::CustomDepth;
 	}
 
-	if (IsMobileDeferredShadingEnabled(View.GetShaderPlatform()))
+	if (MobileRequiresSceneDepthAux(View.GetShaderPlatform()))
 	{
-		SetupMode |= EMobileSceneTextureSetupMode::SceneDepth;
+		SetupMode |= EMobileSceneTextureSetupMode::SceneDepthAux;
 	}
 	else
 	{
-		SetupMode |= EMobileSceneTextureSetupMode::SceneDepthAux;
+		SetupMode |= EMobileSceneTextureSetupMode::SceneDepth;
 	}
 
 	SetupMobileSceneTextureUniformParameters(GraphBuilder, SetupMode, Parameters->SceneTextures);
