Avoid more lighting work when the path tracer or ray traced debug is active

diff --git a/Engine/Source/Runtime/Renderer/Private/DeferredShadingRenderer.cpp b/Engine/Source/Runtime/Renderer/Private/DeferredShadingRenderer.cpp
--- a/Engine/Source/Runtime/Renderer/Private/DeferredShadingRenderer.cpp
+++ b/Engine/Source/Runtime/Renderer/Private/DeferredShadingRenderer.cpp
@@ -1966,7 +1966,8 @@ void FDeferredShadingSceneRenderer::Render(FRDGBuilder& GraphBuilder)
 
 	const FRDGSystemTextures& SystemTextures = FRDGSystemTextures::Create(GraphBuilder);
 
-	const bool bAllowStaticLighting = IsStaticLightingAllowed();
+	const bool bHasRayTracedOverlay = HasRayTracedOverlay(ViewFamily);
+	const bool bAllowStaticLighting = !bHasRayTracedOverlay && IsStaticLightingAllowed();
 
 	const bool bUseVirtualTexturing = UseVirtualTexturing(FeatureLevel);
 	if (bUseVirtualTexturing)
@@ -2087,7 +2088,6 @@ void FDeferredShadingSceneRenderer::Render(FRDGBuilder& GraphBuilder)
 	}
 
 	const bool bUseGBuffer = IsUsingGBuffers(ShaderPlatform);
-	const bool bHasRayTracedOverlay = HasRayTracedOverlay(ViewFamily);
 	
 	const bool bRenderDeferredLighting = ViewFamily.EngineShowFlags.Lighting
 		&& FeatureLevel >= ERHIFeatureLevel::SM5
@@ -3073,9 +3073,12 @@ void FDeferredShadingSceneRenderer::Render(FRDGBuilder& GraphBuilder)
 		RenderMeshDistanceFieldVisualization(GraphBuilder, SceneTextures, FDistanceFieldAOParameters(Scene->DefaultMaxDistanceFieldOcclusionDistance));
 	}
 
-	RenderLumenSceneVisualization(GraphBuilder, SceneTextures);
-	FinishGatheringLumenSurfaceCacheFeedback(GraphBuilder);
-	RenderDiffuseIndirectAndAmbientOcclusion(GraphBuilder, SceneTextures, LightingChannelsTexture, /* bIsVisualizePass = */ true);
+	if (bRenderDeferredLighting)
+	{
+		RenderLumenSceneVisualization(GraphBuilder, SceneTextures);
+		FinishGatheringLumenSurfaceCacheFeedback(GraphBuilder);
+		RenderDiffuseIndirectAndAmbientOcclusion(GraphBuilder, SceneTextures, LightingChannelsTexture, /* bIsVisualizePass = */ true);
+	}
 
 	if (ViewFamily.EngineShowFlags.StationaryLightOverlap)
 	{
