Resolve Depth issue with Editor Primitives in Mobile Preview

diff --git a/Engine/Source/Runtime/Renderer/Private/MobileShadingRenderer.cpp b/Engine/Source/Runtime/Renderer/Private/MobileShadingRenderer.cpp
--- a/Engine/Source/Runtime/Renderer/Private/MobileShadingRenderer.cpp
+++ b/Engine/Source/Runtime/Renderer/Private/MobileShadingRenderer.cpp
@@ -397,8 +397,9 @@ void FMobileSceneRenderer::InitViews(FRDGBuilder& GraphBuilder, FSceneTexturesCo
 	const bool bRequiresUpscale = ((int32)RenderTargetSize.X > FamilySize.X || (int32)RenderTargetSize.Y > FamilySize.Y) || IsMobilePropagateAlphaEnabled(ViewFamily.GetShaderPlatform());
 	// ES requires that the back buffer and depth match dimensions.
 	// For the most part this is not the case when using scene captures. Thus scene captures always render to scene color target.
+	const bool bShouldCompositeEditorPrimitives = FSceneRenderer::ShouldCompositeEditorPrimitives(Views[0]);
 	const bool bStereoRenderingAndHMD = ViewFamily.EngineShowFlags.StereoRendering && ViewFamily.EngineShowFlags.HMDDistortion;
-	bRenderToSceneColor = !bGammaSpace || bStereoRenderingAndHMD || bRequiresUpscale || FSceneRenderer::ShouldCompositeEditorPrimitives(Views[0]) || Views[0].bIsSceneCapture || Views[0].bIsReflectionCapture;
+	bRenderToSceneColor = !bGammaSpace || bStereoRenderingAndHMD || bRequiresUpscale || bShouldCompositeEditorPrimitives || Views[0].bIsSceneCapture || Views[0].bIsReflectionCapture;
 	const FPlanarReflectionSceneProxy* PlanarReflectionSceneProxy = Scene ? Scene->GetForwardPassGlobalPlanarReflection() : nullptr;
 
 	bRequiresPixelProjectedPlanarRelfectionPass = IsUsingMobilePixelProjectedReflection(ShaderPlatform)
@@ -459,6 +460,12 @@ void FMobileSceneRenderer::InitViews(FRDGBuilder& GraphBuilder, FSceneTexturesCo
 	// never keep MSAA depth if SceneDepthAux is enabled
 	bKeepDepthContent = ((NumMSAASamples > 1) && MobileRequiresSceneDepthAux(ShaderPlatform)) ? false : bKeepDepthContent;
 
+	// Depth is needed for Editor Primitives
+	if (bShouldCompositeEditorPrimitives)
+	{
+		bKeepDepthContent = true;
+	}
+
 	// Update the bKeepDepthContent based on the mobile renderer status.
 	SceneTexturesConfig.bKeepDepthContent = bKeepDepthContent;
 	
