Subsurface Scattering Engine flag changes the albedo

diff --git a/Engine/Shaders/Private/PostProcessSubsurface.usf b/Engine/Shaders/Private/PostProcessSubsurface.usf
--- a/Engine/Shaders/Private/PostProcessSubsurface.usf
+++ b/Engine/Shaders/Private/PostProcessSubsurface.usf
@@ -794,9 +794,10 @@ void SubsurfaceRecombinePS(noperspective float2 UV : TEXCOORD0, float4 SvPositio
 
 	if (!bUseSSSProfile)
 	{
-		OutColor = Texture2DSample(SubsurfaceInput0_Texture, SharedSubsurfaceSampler0, BufferUV);
+		// Ingore the copy of non-scattering pixels with zero alpha.  
+		OutColor = float4(0.0f, 0.0f, 0.0f, 0.0f);
 		return;
-	} 
+	}
 
 #if STRATA_ENABLED
 	const uint ProfileId = SSSData.ProfileId;
diff --git a/Engine/Source/Runtime/Renderer/Private/PostProcess/PostProcessSubsurface.cpp b/Engine/Source/Runtime/Renderer/Private/PostProcess/PostProcessSubsurface.cpp
--- a/Engine/Source/Runtime/Renderer/Private/PostProcess/PostProcessSubsurface.cpp
+++ b/Engine/Source/Runtime/Renderer/Private/PostProcess/PostProcessSubsurface.cpp
@@ -1289,27 +1289,36 @@ void AddSubsurfaceViewPass(
 		}
 	}
 
-	// Recombines scattering result with scene color, and updates only SSS region using specialized vertex shader
-	if (SubsurfaceMode != ESubsurfaceMode::Bypass)
+	// Recombines scattering result with scene color, and updates only SSS region using specialized vertex shader.
 	{
 		FRDGTextureRef SubsurfaceIntermediateTexture = GraphBuilder.CreateTexture(SceneColorTextureDesc, TEXT("Subsurface.Recombines"));
 
-		// recombine into intermediate textures before copying back.
+		// When bypassing subsurface scattering, we use the directly full screen pass instead of building the tile
+		// and apply the tile based copy
+		const bool bShouldFallbackToFullScreenPass = SubsurfaceMode == ESubsurfaceMode::Bypass;
+		FRHIDepthStencilState* DepthStencilState = TStaticDepthStencilState<false, CF_Always>::GetRHI();
+
+		// Recombine into intermediate textures before copying back.
 		{
 			FSubsurfaceTiles::ETileType TileType = FSubsurfaceTiles::ETileType::All;
 
 			FSubsurfaceRecombinePS::FParameters* PassParameters = GraphBuilder.AllocParameters<FSubsurfaceRecombinePS::FParameters>();
 			PassParameters->Subsurface = SubsurfaceCommonParameters;
-			PassParameters->TileParameters = GetSubsurfaceTileParameters(View, Tiles, TileType, SubsurfaceViewport.Extent);
+			if (SubsurfaceMode != ESubsurfaceMode::Bypass)
+			{
+				PassParameters->TileParameters = GetSubsurfaceTileParameters(View, Tiles, TileType, SubsurfaceViewport.Extent);
+			}
 			PassParameters->RenderTargets[0] = FRenderTargetBinding(SubsurfaceIntermediateTexture, SceneColorTextureLoadAction);
 			PassParameters->SubsurfaceInput0 = GetSubsurfaceInput(SceneColorTexture, SceneViewportParameters);
 			PassParameters->SubsurfaceSampler0 = BilinearBorderSampler;
 			PassParameters->Strata = Strata::BindStrataGlobalUniformParameters(View.StrataSceneData);
 
 			// Scattering output target is only used when scattering is enabled.
-			PassParameters->SubsurfaceInput1 = GetSubsurfaceInput(SubsurfaceSubpassTwoTex, SubsurfaceViewportParameters);
-			PassParameters->SubsurfaceSampler1 = BilinearBorderSampler;
-
+			if (SubsurfaceMode != ESubsurfaceMode::Bypass)
+			{
+				PassParameters->SubsurfaceInput1 = GetSubsurfaceInput(SubsurfaceSubpassTwoTex, SubsurfaceViewportParameters);
+				PassParameters->SubsurfaceSampler1 = BilinearBorderSampler;
+			}
 
 			const FSubsurfaceRecombinePS::EQuality RecombineQuality = FSubsurfaceRecombinePS::GetQuality(View);
 
@@ -1322,27 +1331,32 @@ void AddSubsurfaceViewPass(
 
 			TShaderMapRef<FSubsurfaceRecombinePS> PixelShader(View.ShaderMap, PixelShaderPermutationVector);
 			TShaderMapRef<FSubsurfaceTilePassVS> VertexShader(View.ShaderMap);
+			
+			FRHIBlendState* BlendState = TStaticBlendState<CW_RGBA, BO_Add, BF_One, BF_Zero, BO_Add, BF_One, BF_Zero>::GetRHI();
 			/**
-			 * See the related comment above in the prepare pass. The scene viewport is used as both the target and
-			 * texture viewport in order to ensure that the correct pixel is sampled for checkerboard rendering.
-			 */
+				* See the related comment above in the prepare pass. The scene viewport is used as both the target and
+				* texture viewport in order to ensure that the correct pixel is sampled for checkerboard rendering.
+				*/
 			AddSubsurfaceTiledScreenPass(
 				GraphBuilder,
-				RDG_EVENT_NAME("SubsurfaceRecombine(%s %s%s%s%s) %dx%d",
+				RDG_EVENT_NAME("SubsurfaceRecombine(%s %s%s%s%s%s) %dx%d",
 					GetEventName(PixelShaderPermutationVector.Get<FSubsurfaceRecombinePS::FDimensionMode>()),
 					FSubsurfaceRecombinePS::GetEventName(PixelShaderPermutationVector.Get<FSubsurfaceRecombinePS::FDimensionQuality>()),
 					PixelShaderPermutationVector.Get<FSubsurfaceRecombinePS::FDimensionCheckerboard>() ? TEXT(" Checkerboard") : TEXT(""),
 					PixelShaderPermutationVector.Get<FSubsurfaceRecombinePS::FDimensionHalfRes>() ? TEXT(" HalfRes") : TEXT(""),
 					PixelShaderPermutationVector.Get<FSubsurfaceRecombinePS::FRunningInSeparable>() ? TEXT(" RunningInSeparable") : TEXT(""),
+					!bShouldFallbackToFullScreenPass ? TEXT(" Tiled") : TEXT(""),
 					View.ViewRect.Width(),
 					View.ViewRect.Height()),
 				View,
 				PassParameters,
 				VertexShader,
 				PixelShader,
-				TStaticBlendState<CW_RGBA, BO_Add, BF_One, BF_Zero, BO_Add, BF_One, BF_Zero>::GetRHI(),
+				BlendState,
+				DepthStencilState,
 				SceneViewport,
-				TileType);
+				TileType,
+				bShouldFallbackToFullScreenPass);
 		}
 
 		//Write back to the SceneColor texture 
@@ -1351,27 +1365,35 @@ void AddSubsurfaceViewPass(
 
 			FSubsurfaceRecombineCopyPS::FParameters* PassParameters = GraphBuilder.AllocParameters<FSubsurfaceRecombineCopyPS::FParameters>();
 			PassParameters->Subsurface = SubsurfaceCommonParameters;
-			PassParameters->TileParameters = GetSubsurfaceTileParameters(View, Tiles, TileType, SubsurfaceViewport.Extent);
+			if (SubsurfaceMode != ESubsurfaceMode::Bypass)
+			{
+				PassParameters->TileParameters = GetSubsurfaceTileParameters(View, Tiles, TileType, SubsurfaceViewport.Extent);
+			}
 			PassParameters->RenderTargets[0] = FRenderTargetBinding(SceneColorTexture, SceneColorTextureLoadAction);
 			PassParameters->SubsurfaceInput0 = GetSubsurfaceInput(SubsurfaceIntermediateTexture, SceneViewportParameters);
 			PassParameters->SubsurfaceSampler0 = PointClampSampler;
 			PassParameters->Strata = Strata::BindStrataGlobalUniformParameters(View.StrataSceneData);
 
-			TShaderMapRef<FSubsurfaceRecombineCopyPS> PixelShader(View.ShaderMap);
+			TShaderMapRef<FSubsurfaceRecombineCopyPS> PixelShader(View.ShaderMap);			
 			TShaderMapRef<FSubsurfaceTilePassVS> VertexShader(View.ShaderMap);
 			
+			FRHIBlendState* BlendState = TStaticBlendState<CW_RGBA, BO_Add, BF_SourceAlpha, BF_InverseSourceAlpha, BO_Add, BF_Zero, BF_Zero>::GetRHI();
+			
 			AddSubsurfaceTiledScreenPass(
 				GraphBuilder,
-				RDG_EVENT_NAME("SubsurfaceCopyToSceneColor %dx%d",
+				RDG_EVENT_NAME("SubsurfaceCopyToSceneColor%s %dx%d",
+					!bShouldFallbackToFullScreenPass ? TEXT("( Tiled)") : TEXT(""),
 					View.ViewRect.Width(),
 					View.ViewRect.Height()),
-				View, 
-				PassParameters, 
-				VertexShader, 
+				View,
+				PassParameters,
+				VertexShader,
 				PixelShader,
-				TStaticBlendState<CW_RGBA, BO_Add, BF_SourceAlpha, BF_InverseSourceAlpha, BO_Add, BF_Zero, BF_Zero>::GetRHI(),
-				SceneViewport, 
-				TileType);
+				BlendState,
+				DepthStencilState,
+				SceneViewport,
+				TileType,
+				bShouldFallbackToFullScreenPass);
 		}
 	}
 
diff --git a/Engine/Source/Runtime/Renderer/Private/SubsurfaceTiles.h b/Engine/Source/Runtime/Renderer/Private/SubsurfaceTiles.h
--- a/Engine/Source/Runtime/Renderer/Private/SubsurfaceTiles.h
+++ b/Engine/Source/Runtime/Renderer/Private/SubsurfaceTiles.h
@@ -87,35 +87,54 @@ void AddSubsurfaceTiledScreenPass(
 	TShaderMapRef<FSubsurfacePassVS>& VertexShader,
 	TShaderMapRef<FSubsurfacePassPS>& PixelShader,
 	FRHIBlendState* BlendState,
+	FRHIDepthStencilState* DepthStencilState,
 	const FScreenPassTextureViewport SceneViewport,
-	FSubsurfaceTiles::ETileType TileType)
+	FSubsurfaceTiles::ETileType TileType,
+	const bool bShouldFallbackToFullScreenPass = false)
 {
-	GraphBuilder.AddPass(
-		Forward<FRDGEventName>(Name),
-		PassParameters,
-		ERDGPassFlags::Raster,
-		[PassParameters, VertexShader, PixelShader, BlendState, Viewport = SceneViewport, TileType](FRHICommandList& RHICmdList)
-		{
-			typename FSubsurfacePassVS::FParameters ParametersVS = PassParameters->TileParameters;
-
-			FGraphicsPipelineStateInitializer GraphicsPSOInit;
-			RHICmdList.ApplyCachedRenderTargets(GraphicsPSOInit);
-			GraphicsPSOInit.BlendState = BlendState;
-			GraphicsPSOInit.RasterizerState = TStaticRasterizerState<>::GetRHI();
-			GraphicsPSOInit.DepthStencilState = TStaticDepthStencilState<false, CF_Never>::GetRHI();
-
-			GraphicsPSOInit.BoundShaderState.VertexDeclarationRHI = GFilterVertexDeclaration.VertexDeclarationRHI;
-			GraphicsPSOInit.BoundShaderState.VertexShaderRHI = VertexShader.GetVertexShader();
-			GraphicsPSOInit.BoundShaderState.PixelShaderRHI = PixelShader.GetPixelShader();
-			GraphicsPSOInit.PrimitiveType = PassParameters->TileParameters.bRectPrimitive > 0 ? PT_RectList : PT_TriangleList;
-			SetGraphicsPipelineState(RHICmdList, GraphicsPSOInit, 0);
-			SetShaderParameters(RHICmdList, VertexShader, VertexShader.GetVertexShader(), ParametersVS);
-			SetShaderParameters(RHICmdList, PixelShader, PixelShader.GetPixelShader(), *PassParameters);
-
-			RHICmdList.SetViewport(0.0f, 0.0f, 0.0f, Viewport.Extent.X, Viewport.Extent.Y, 1.0f);
-			RHICmdList.SetStreamSource(0, nullptr, 0);
-			RHICmdList.DrawPrimitiveIndirect(PassParameters->TileParameters.TileIndirectBuffer->GetRHI(), FSubsurfaceTiles::GetIndirectDrawArgOffset(TileType));
-		});
+	if (bShouldFallbackToFullScreenPass)
+	{
+		TShaderMapRef<FScreenPassVS> FallBackVertexShader(View.ShaderMap);
+		AddDrawScreenPass(GraphBuilder,
+			Forward<FRDGEventName>(Name),
+			View,
+			SceneViewport,
+			SceneViewport,
+			FallBackVertexShader,
+			PixelShader,
+			BlendState,
+			DepthStencilState,
+			PassParameters);
+	}
+	else
+	{
+		GraphBuilder.AddPass(
+			Forward<FRDGEventName>(Name),
+			PassParameters,
+			ERDGPassFlags::Raster,
+			[PassParameters, VertexShader, PixelShader, DepthStencilState, BlendState, Viewport = SceneViewport, TileType](FRHICommandList& RHICmdList)
+			{
+				typename FSubsurfacePassVS::FParameters ParametersVS = PassParameters->TileParameters;
+
+				FGraphicsPipelineStateInitializer GraphicsPSOInit;
+				RHICmdList.ApplyCachedRenderTargets(GraphicsPSOInit);
+				GraphicsPSOInit.BlendState = BlendState;
+				GraphicsPSOInit.RasterizerState = TStaticRasterizerState<>::GetRHI();
+				GraphicsPSOInit.DepthStencilState = DepthStencilState;
+
+				GraphicsPSOInit.BoundShaderState.VertexDeclarationRHI = GFilterVertexDeclaration.VertexDeclarationRHI;
+				GraphicsPSOInit.BoundShaderState.VertexShaderRHI = VertexShader.GetVertexShader();
+				GraphicsPSOInit.BoundShaderState.PixelShaderRHI = PixelShader.GetPixelShader();
+				GraphicsPSOInit.PrimitiveType = PassParameters->TileParameters.bRectPrimitive > 0 ? PT_RectList : PT_TriangleList;
+				SetGraphicsPipelineState(RHICmdList, GraphicsPSOInit, 0);
+				SetShaderParameters(RHICmdList, VertexShader, VertexShader.GetVertexShader(), ParametersVS);
+				SetShaderParameters(RHICmdList, PixelShader, PixelShader.GetPixelShader(), *PassParameters);
+
+				RHICmdList.SetViewport(0.0f, 0.0f, 0.0f, Viewport.Extent.X, Viewport.Extent.Y, 1.0f);
+				RHICmdList.SetStreamSource(0, nullptr, 0);
+				RHICmdList.DrawPrimitiveIndirect(PassParameters->TileParameters.TileIndirectBuffer->GetRHI(), FSubsurfaceTiles::GetIndirectDrawArgOffset(TileType));
+			});
+	}
 }
 
 /** Clear the UAV texture to black only when ConditionBuffer[Offset] > 0 */
