Skip Hair HasPositionsChanged task when there are no hair strands in the scene, or they are not being simulated or skinned

diff --git a/Engine/Source/Runtime/Renderer/Private/HairStrands/HairStrandsVisibility.cpp b/Engine/Source/Runtime/Renderer/Private/HairStrands/HairStrandsVisibility.cpp
--- a/Engine/Source/Runtime/Renderer/Private/HairStrands/HairStrandsVisibility.cpp
+++ b/Engine/Source/Runtime/Renderer/Private/HairStrands/HairStrandsVisibility.cpp
@@ -3467,6 +3467,12 @@ void DrawHitProxies(
 // Check if any simulated/skinned-bound groom has its positions updated (e.g. for invalidating the path-tracer accumulation)
 bool HasPositionsChanged(FRDGBuilder& GraphBuilder, const FViewInfo& View)
 {
+	if (View.HairStrandsMeshElements.IsEmpty())
+	{
+		// there are no hair strands in the scene
+		return false;
+	}
+
 	FHairStrandsViewStateData* HairStrandsViewStateData = const_cast<FHairStrandsViewStateData*>(&View.ViewState->HairStrandsViewStateData);
 	if (!HairStrandsViewStateData->IsInit())
 	{
@@ -3487,6 +3493,11 @@ bool HasPositionsChanged(FRDGBuilder& GraphBuilder, const FViewInfo& View)
 			GroupDatas.Add(HairGroupPublicData);
 		}
 	}
+	if (GroupDatas.IsEmpty())
+	{
+		// there are no strands currently being simulated or skinned
+		return false;
+	}
 
 	FRDGBufferDesc Desc = FRDGBufferDesc::CreateBufferDesc(sizeof(uint32), 1);
 	Desc.Usage |= BUF_SourceCopy;
