Fix a debug crash when GPUScene and modulated shadow are enabled.

#p4v-preflight-copy 18206900

diff --git a/Engine/Source/Runtime/Renderer/Private/SimpleMeshDrawCommandPass.cpp b/Engine/Source/Runtime/Renderer/Private/SimpleMeshDrawCommandPass.cpp
--- a/Engine/Source/Runtime/Renderer/Private/SimpleMeshDrawCommandPass.cpp
+++ b/Engine/Source/Runtime/Renderer/Private/SimpleMeshDrawCommandPass.cpp
@@ -24,7 +24,6 @@ FSimpleMeshDrawCommandPass::FSimpleMeshDrawCommandPass(const FSceneView& View, F
 
 	ERHIFeatureLevel::Type FeatureLevel = ViewInfo->GetFeatureLevel();
 	InstanceCullingContext = FInstanceCullingContext(FeatureLevel, InstanceCullingManager, ViewIds, nullptr, bUsingStereo ? EInstanceCullingMode::Stereo : EInstanceCullingMode::Normal);
-	bDynamicInstancing = IsDynamicInstancingEnabled(FeatureLevel);
 
 	InstanceFactor = static_cast<uint32>(ViewIds.Num());
 }
@@ -80,7 +79,7 @@ void FSimpleMeshDrawCommandPass::SubmitDraw(FRHICommandList& RHICmdList, const F
 		else
 		{
 			const uint32 PrimitiveIdBufferStride = FInstanceCullingContext::GetInstanceIdBufferStride(InstanceCullingContext.FeatureLevel);
-			SubmitMeshDrawCommandsRange(VisibleMeshDrawCommands, GraphicsMinimalPipelineStateSet, PrimitiveIdVertexBuffer, PrimitiveIdBufferStride, 0, bDynamicInstancing, 0, VisibleMeshDrawCommands.Num(), InstanceFactor, RHICmdList);
+			SubmitMeshDrawCommandsRange(VisibleMeshDrawCommands, GraphicsMinimalPipelineStateSet, PrimitiveIdVertexBuffer, PrimitiveIdBufferStride, 0, false, 0, VisibleMeshDrawCommands.Num(), InstanceFactor, RHICmdList);
 		}
 	}
 }
diff --git a/Engine/Source/Runtime/Renderer/Public/SimpleMeshDrawCommandPass.h b/Engine/Source/Runtime/Renderer/Public/SimpleMeshDrawCommandPass.h
--- a/Engine/Source/Runtime/Renderer/Public/SimpleMeshDrawCommandPass.h
+++ b/Engine/Source/Runtime/Renderer/Public/SimpleMeshDrawCommandPass.h
@@ -52,7 +52,6 @@ private:
 	bool bSupportsScenePrimitives = false;
 
 	bool bUsingStereo = false;
-	bool bDynamicInstancing = false;
 	uint32 InstanceFactor = 1;
 
 	// GPUCULL_TODO: Only for legacy path
