Fix invalid/too small hair sample allocation when the camera is within the hair bounding box.

This ensures that the projected bounding box get properly projected even when the camera is inside the bounding box, making the sample buffer allocation conservative.

diff --git a/Engine/Source/Runtime/Renderer/Private/HairStrands/HairStrandsUtils.cpp b/Engine/Source/Runtime/Renderer/Private/HairStrands/HairStrandsUtils.cpp
--- a/Engine/Source/Runtime/Renderer/Private/HairStrands/HairStrandsUtils.cpp
+++ b/Engine/Source/Runtime/Renderer/Private/HairStrands/HairStrandsUtils.cpp
@@ -220,11 +220,21 @@ FIntRect ComputeProjectedScreenRect(const FBox& B, const FViewInfo& View)
 	};
 
 	// Compute the MinP/MaxP in pixel coord, relative to View.ViewRect.Min
-	const FMatrix& ViewProj = View.ViewMatrices.GetViewProjectionMatrix();
+	const FMatrix& WorldToView	= View.ViewMatrices.GetViewMatrix();
+	const FMatrix& ViewToProj	= View.ViewMatrices.GetProjectionMatrix();
+	const float NearClippingDistance = View.NearClippingDistance + SMALL_NUMBER;
 	for (uint32 i = 0; i < 8; ++i)
 	{
+		// Clamp position on the near plane to get valid rect even if bounds' points are behind the camera
+		FPlane P_View = WorldToView.TransformFVector4(FVector4(Vertices[i], 1.f));
+		if (P_View.Z <= NearClippingDistance)
+		{
+			P_View.Z = NearClippingDistance;
+		}
+
+		// Project from view to projective space
 		FVector2D P;
-		if (FSceneView::ProjectWorldToScreen(Vertices[i], View.ViewRect, ViewProj, P))
+		if (FSceneView::ProjectWorldToScreen(P_View, View.ViewRect, ViewToProj, P))
 		{
 			MinP.X = FMath::Min(MinP.X, P.X);
 			MinP.Y = FMath::Min(MinP.Y, P.Y);
