Fix validation error when rendering hair strand with deep shadow and raytracing enabled.

diff --git a/Engine/Source/Runtime/Renderer/Private/HairStrands/HairStrandsTransmittance.cpp b/Engine/Source/Runtime/Renderer/Private/HairStrands/HairStrandsTransmittance.cpp
--- a/Engine/Source/Runtime/Renderer/Private/HairStrands/HairStrandsTransmittance.cpp
+++ b/Engine/Source/Runtime/Renderer/Private/HairStrands/HairStrandsTransmittance.cpp
@@ -138,7 +138,7 @@ static void AddHairStrandsClearTransmittanceMaskPass(
 	TShaderMapRef<FHairStrandsClearTransmittanceMaskCS> ComputeShader(ShaderMap, PermutationVector);
 	FComputeShaderUtils::AddPass(
 		GraphBuilder,
-		RDG_EVENT_NAME("HairStrandsClearTransmittanceMask"),
+		RDG_EVENT_NAME("HairStrands::ClearTransmittanceMask"),
 		ComputeShader,
 		Parameters,
 		FComputeShaderUtils::GetGroupCount(Parameters->ElementCount, 64));
@@ -271,7 +271,7 @@ static FRDGBufferRef AddHairStrandsVoxelTransmittanceMaskPass(
 	TShaderMapRef<FHairStrandsVoxelTransmittanceMaskCS> ComputeShader(View.ShaderMap, PermutationVector);
 	FComputeShaderUtils::AddPass(
 		GraphBuilder,
-		RDG_EVENT_NAME("HairStrandsTransmittanceMask(Voxel,%s)", PassType == EHairTransmittancePassType::OnePass ? TEXT("OnePass") : TEXT("PerLight")),
+		RDG_EVENT_NAME("HairStrands::TransmittanceMask(Voxel,%s)", PassType == EHairTransmittancePassType::OnePass ? TEXT("OnePass") : TEXT("PerLight")),
 		ComputeShader,
 		Parameters,
 		IndirectArgsBuffer,
@@ -402,7 +402,7 @@ static FRDGBufferRef AddHairStrandsDeepShadowTransmittanceMaskPass(
 	TShaderMapRef<FHairStrandsDeepShadowTransmittanceMaskCS> ComputeShader(View.ShaderMap, PermutationVector);
 	FComputeShaderUtils::AddPass(
 		GraphBuilder,
-		RDG_EVENT_NAME("HairStrandsTransmittanceMask(DeepShadow)"),
+		RDG_EVENT_NAME("HairStrands::TransmittanceMask(DeepShadow)"),
 		ComputeShader,
 		Parameters,
 		IndirectArgsBuffer,
@@ -513,7 +513,7 @@ static void AddHairStrandsVoxelShadowMaskPass(
 	FIntPoint Resolution = OutShadowMask->Desc.Extent;
 
 	GraphBuilder.AddPass(
-		RDG_EVENT_NAME("HairStrandsShadowMask(Voxel)"),
+		RDG_EVENT_NAME("HairStrands::ShadowMask(Voxel)"),
 		Parameters,
 		ERDGPassFlags::Raster,
 		[Parameters, VertexShader, PixelShader, Viewport, Resolution](FRHICommandList& RHICmdList)
@@ -671,7 +671,7 @@ static void AddHairStrandsDeepShadowMaskPass(
 	FIntPoint Resolution = OutShadowMask->Desc.Extent;
 
 	GraphBuilder.AddPass(
-		RDG_EVENT_NAME("HairStrandsShadowMask"),
+		RDG_EVENT_NAME("HairStrands::ShadowMask"),
 		Parameters,
 		ERDGPassFlags::Raster,
 		[Parameters, VertexShader, PixelShader, Viewport, Resolution](FRHICommandList& RHICmdList)
@@ -965,6 +965,41 @@ void RenderHairStrandsShadowMask(
 		return;
 	}
 
+	for (const FViewInfo& View : Views)
+	{
+		if (HairStrands::HasViewHairStrandsData(View))
+		{
+			check(View.HairStrandsViewData.VisibilityData.CoverageTexture);
+			InternalRenderHairStrandsShadowMask(
+				GraphBuilder,
+				View,
+				LightSceneInfo,
+				View.HairStrandsViewData.VisibilityData,
+				View.HairStrandsViewData.MacroGroupDatas,
+				View.HairStrandsViewData.DeepShadowResources,
+				View.HairStrandsViewData.VirtualVoxelResources,
+				OutShadowMask);
+		}
+	}
+}
+
+void RenderHairStrandsDeepShadowMask(
+	FRDGBuilder& GraphBuilder,
+	const TArray<FViewInfo>& Views,
+	const FLightSceneInfo* LightSceneInfo,
+	FRDGTextureRef OutShadowMask)
+{
+	if (Views.Num() == 0 || OutShadowMask == nullptr)
+	{
+		return;
+	}
+
+	// Render only light with deep shadow
+	if (!LightSceneInfo || !LightSceneInfo->Proxy->CastsHairStrandsDeepShadow())
+	{
+		return;
+	}
+
 	for (const FViewInfo& View : Views)
 	{
 		if (HairStrands::HasViewHairStrandsData(View))
@@ -981,4 +1016,4 @@ void RenderHairStrandsShadowMask(
 				OutShadowMask);
 		}
 	}
-}
\ No newline at end of file
+}
diff --git a/Engine/Source/Runtime/Renderer/Private/HairStrands/HairStrandsTransmittance.h b/Engine/Source/Runtime/Renderer/Private/HairStrands/HairStrandsTransmittance.h
--- a/Engine/Source/Runtime/Renderer/Private/HairStrands/HairStrandsTransmittance.h
+++ b/Engine/Source/Runtime/Renderer/Private/HairStrands/HairStrandsTransmittance.h
@@ -26,6 +26,13 @@ void RenderHairStrandsShadowMask(
 	const FLightSceneInfo* LightSceneInfo,
 	FRDGTextureRef ScreenShadowMaskTexture); 
 
+/// Write opaque hair shadow onto screen shadow mask to have fine hair details cast onto opaque geometries (deep shadow caster only)
+void RenderHairStrandsDeepShadowMask(
+	FRDGBuilder& GraphBuilder,
+	const TArray<FViewInfo>& Views,
+	const FLightSceneInfo* LightSceneInfo,
+	FRDGTextureRef OutShadowMask);
+
 /// Output hair transmittance per hair sample for a given light
 FHairStrandsTransmittanceMaskData RenderHairStrandsTransmittanceMask(
 	FRDGBuilder& GraphBuilder,
diff --git a/Engine/Source/Runtime/Renderer/Private/LightRendering.cpp b/Engine/Source/Runtime/Renderer/Private/LightRendering.cpp
--- a/Engine/Source/Runtime/Renderer/Private/LightRendering.cpp
+++ b/Engine/Source/Runtime/Renderer/Private/LightRendering.cpp
@@ -1677,7 +1677,7 @@ void FDeferredShadingSceneRenderer::RenderLights(
 						// Inject deep shadow mask if the light supports it
 						if (bUseHairDeepShadow)
 						{
-							RenderHairStrandsShadowMask(GraphBuilder, Views, &LightSceneInfo, ScreenShadowMaskTexture);
+							RenderHairStrandsDeepShadowMask(GraphBuilder, Views, &LightSceneInfo, ScreenShadowMaskTexture);
 						}
 					}
 					else if (OcclusionType == FLightOcclusionType::Raytraced)
@@ -1720,7 +1720,6 @@ void FDeferredShadingSceneRenderer::RenderLights(
 							SubPixelRayTracingShadowMaskUAV = GraphBuilder.CreateUAV(FRDGTextureUAVDesc(SubPixelRayTracingShadowMaskTexture));
 						}
 
-
 						FRDGTextureRef RayTracingShadowMaskTileTexture;
 						{
 							FRDGTextureDesc Desc = FRDGTextureDesc::Create2D(
diff --git a/Engine/Source/Runtime/Renderer/Private/RayTracing/RayTracingShadows.cpp b/Engine/Source/Runtime/Renderer/Private/RayTracing/RayTracingShadows.cpp
--- a/Engine/Source/Runtime/Renderer/Private/RayTracing/RayTracingShadows.cpp
+++ b/Engine/Source/Runtime/Renderer/Private/RayTracing/RayTracingShadows.cpp
@@ -262,7 +262,8 @@ void FDeferredShadingSceneRenderer::RenderRayTracingShadows(
 
 	// Ray generation pass for shadow occlusion.
 	{
-		const bool bUseHairLighting = HairStrands::HasViewHairStrandsData(View) && HairStrands::HasViewHairStrandsVoxelData(View) && !LightSceneProxy->CastsHairStrandsDeepShadow();
+		const bool bUseHairLighting   = HairStrands::HasViewHairStrandsData(View) && HairStrands::HasViewHairStrandsVoxelData(View);
+		const bool bUseHairDeepShadow = HairStrands::HasViewHairStrandsData(View) && LightSceneProxy->CastsHairStrandsDeepShadow();
 
 		FOcclusionRGS::FParameters* PassParameters = GraphBuilder.AllocParameters<FOcclusionRGS::FParameters>();
 		PassParameters->RWOcclusionMaskUAV = OutShadowMaskUAV;
@@ -292,7 +293,7 @@ void FDeferredShadingSceneRenderer::RenderRayTracingShadows(
 		if (bUseHairLighting)
 		{
 			const bool bUseHairVoxel = CVarRayTracingShadowsEnableHairVoxel.GetValueOnRenderThread() > 0;
-			PassParameters->bUseHairVoxel = bUseHairVoxel ? 1 : 0;
+			PassParameters->bUseHairVoxel = !bUseHairDeepShadow && bUseHairVoxel ? 1 : 0;
 			PassParameters->HairLightChannelMaskTexture = View.HairStrandsViewData.VisibilityData.LightChannelMaskTexture;
 			PassParameters->HairStrands = HairStrands::BindHairStrandsViewUniformParameters(View);
 			PassParameters->VirtualVoxel = HairStrands::BindHairStrandsVoxelUniformParameters(View);
