* Change classification texture to be encoded on 16bits rather 32bits
* Remove unused Strata VGPRs array

diff --git a/Engine/Shaders/Private/Strata/Strata.ush b/Engine/Shaders/Private/Strata/Strata.ush
--- a/Engine/Shaders/Private/Strata/Strata.ush
+++ b/Engine/Shaders/Private/Strata/Strata.ush
@@ -87,7 +87,7 @@ struct FStrataBSDF
 	float4 InlineVGPRs[3];
 #endif
 
-	float4 VGPRs[6];
+	float4 VGPRs[5];
 };
 
 
@@ -341,24 +341,23 @@ uint StrataPackClassificationData(FStrataClassification In)
 {
 	const uint TopLayerRoughness8Bits = uint(In.TopLayerRoughness * 255.0f);
 	return 
-		(In.ShadingModels & 0xFF) | 
-		((TopLayerRoughness8Bits & 0xFF)	 << 16 ) |
-		((In.bIsStrataMaterial		? 1 : 0) << 24 ) |
-		((In.bHasSubsurface			? 1 : 0) << 25 ) | 
-		((In.bHasSubsurfaceProfile	? 1 : 0) << 26 ) |
-		((In.bIsSimple				? 1 : 0) << 27 );
+		(In.ShadingModels & 0x1F) | 
+		((TopLayerRoughness8Bits & 0xFF)	 << 8 ) |
+		((In.bHasSubsurface			? 1 : 0) << 5 ) | 
+		((In.bHasSubsurfaceProfile	? 1 : 0) << 6 ) |
+		((In.bIsSimple				? 1 : 0) << 7 );
 }
 
 FStrataClassification StrataUnpackClassificationData(uint In)
 {
 	uint TopLayerRoughness8Bits = 0;
 	FStrataClassification Out;
-	Out.ShadingModels			= (In >> 0 ) & 0xFF;
-	TopLayerRoughness8Bits		= (In >> 16) & 0xFF;
-	Out.bIsStrataMaterial		= (In >> 24) & 0x1;
-	Out.bHasSubsurface			= (In >> 25) & 0x1;
-	Out.bHasSubsurfaceProfile	= (In >> 26) & 0x1;
-	Out.bIsSimple				= (In >> 27) & 0x1;
+	Out.ShadingModels			= (In >> 0 ) & 0x1F;
+	TopLayerRoughness8Bits		= (In >> 8) & 0xFF;
+	Out.bIsStrataMaterial		= Out.ShadingModels > 0;
+	Out.bHasSubsurface			= (In >> 5) & 0x1;
+	Out.bHasSubsurfaceProfile	= (In >> 6) & 0x1;
+	Out.bIsSimple				= (In >> 7) & 0x1;
 
 	Out.TopLayerRoughness = float(TopLayerRoughness8Bits) * (1.0f / 255.0f);
 	return Out;
diff --git a/Engine/Shaders/Shared/StrataDefinitions.h b/Engine/Shaders/Shared/StrataDefinitions.h
--- a/Engine/Shaders/Shared/StrataDefinitions.h
+++ b/Engine/Shaders/Shared/StrataDefinitions.h
@@ -27,7 +27,7 @@
 #define STRATA_BSDF_TYPE_UNLIT				2
 #define STRATA_BSDF_TYPE_HAIR				3
 #define STRATA_BSDF_TYPE_SINGLELAYERWATER	4
-// When more than 8 BSDF must exists, please update STATE_BITS_BSDF and FStrataClassification.ShadingModels packing in Strata.ush
+// When more than 5 BSDF must exists, please update STATE_BIT_COUNT_BSDF and FStrataClassification.ShadingModels packing in Strata.ush
 
 #define STRATA_DATA_TILE_SIZE				8
 #define STRATA_DATA_TILE_SIZE_DIV_AS_SHIFT	3
diff --git a/Engine/Source/Runtime/Renderer/Private/Strata/Strata.cpp b/Engine/Source/Runtime/Renderer/Private/Strata/Strata.cpp
--- a/Engine/Source/Runtime/Renderer/Private/Strata/Strata.cpp
+++ b/Engine/Source/Runtime/Renderer/Private/Strata/Strata.cpp
@@ -160,7 +160,7 @@ void InitialiseStrataFrameSceneData(FSceneRenderer& SceneRenderer, FRDGBuilder&
 
 		// Classification texture
 		{
-			StrataSceneData.ClassificationTexture = GraphBuilder.CreateTexture(FRDGTextureDesc::Create2D(SceneTextureExtent, PF_R32_UINT, FClearValueBinding::Black, TexCreate_ShaderResource | TexCreate_RenderTargetable), TEXT("StrataClassificationTexture"));
+			StrataSceneData.ClassificationTexture = GraphBuilder.CreateTexture(FRDGTextureDesc::Create2D(SceneTextureExtent, PF_R16_UINT, FClearValueBinding::Black, TexCreate_ShaderResource | TexCreate_RenderTargetable), TEXT("Strata.ClassificationTexture"));
 		}
 
 		// Tile classification buffers
@@ -172,11 +172,11 @@ void InitialiseStrataFrameSceneData(FSceneRenderer& SceneRenderer, FRDGBuilder&
 			// This is fine for two types, later we might want to have a single list and indirect buffer with offsets.
 			for (uint32 i = 0; i < EStrataTileMaterialType::ECount; ++i)
 			{
-				StrataSceneData.ClassificationTileListBuffer[i] = GraphBuilder.CreateBuffer(FRDGBufferDesc::CreateBufferDesc(sizeof(uint32), TileResolution.X * TileResolution.Y), i==EStrataTileMaterialType::ESimple ? TEXT("SimpleStrataTileListBuffer") : TEXT("ComplexStrataTileListBuffer"));
+				StrataSceneData.ClassificationTileListBuffer[i] = GraphBuilder.CreateBuffer(FRDGBufferDesc::CreateBufferDesc(sizeof(uint32), TileResolution.X * TileResolution.Y), i==EStrataTileMaterialType::ESimple ? TEXT("Strata.SimpleStrataTileListBuffer") : TEXT("Strata.ComplexStrataTileListBuffer"));
 				StrataSceneData.ClassificationTileListBufferSRV[i] = GraphBuilder.CreateSRV(StrataSceneData.ClassificationTileListBuffer[i], PF_R32_UINT);
 				StrataSceneData.ClassificationTileListBufferUAV[i] = GraphBuilder.CreateUAV(StrataSceneData.ClassificationTileListBuffer[i], PF_R32_UINT);
 
-				StrataSceneData.ClassificationTileIndirectBuffer[i] = GraphBuilder.CreateBuffer(FRDGBufferDesc::CreateIndirectDesc<FRHIDrawIndirectParameters>(), i == EStrataTileMaterialType::ESimple ? TEXT("SimpleStrataTileIndirectBuffer") : TEXT("ComplexStrataTileIndirectBuffer"));
+				StrataSceneData.ClassificationTileIndirectBuffer[i] = GraphBuilder.CreateBuffer(FRDGBufferDesc::CreateIndirectDesc<FRHIDrawIndirectParameters>(), i == EStrataTileMaterialType::ESimple ? TEXT("Strata.SimpleStrataTileIndirectBuffer") : TEXT("Strata.ComplexStrataTileIndirectBuffer"));
 				StrataSceneData.ClassificationTileIndirectBufferSRV[i] = GraphBuilder.CreateSRV(StrataSceneData.ClassificationTileIndirectBuffer[i], PF_R32_UINT);
 				StrataSceneData.ClassificationTileIndirectBufferUAV[i] = GraphBuilder.CreateUAV(StrataSceneData.ClassificationTileIndirectBuffer[i], PF_R32_UINT);
 
@@ -186,12 +186,12 @@ void InitialiseStrataFrameSceneData(FSceneRenderer& SceneRenderer, FRDGBuilder&
 
 		// Top layer texture
 		{
-			StrataSceneData.TopLayerNormalTexture = GraphBuilder.CreateTexture(FRDGTextureDesc::Create2D(SceneTextureExtent, PF_R32_UINT, FClearValueBinding::Black, TexCreate_ShaderResource | TexCreate_RenderTargetable), TEXT("StrataTopLayerNormalTexture"));
+			StrataSceneData.TopLayerNormalTexture = GraphBuilder.CreateTexture(FRDGTextureDesc::Create2D(SceneTextureExtent, PF_R32_UINT, FClearValueBinding::Black, TexCreate_ShaderResource | TexCreate_RenderTargetable), TEXT("Strata.TopLayerNormalTexture"));
 		}
 
 		// SSS texture
 		{
-			StrataSceneData.SSSTexture = GraphBuilder.CreateTexture(FRDGTextureDesc::Create2D(SceneTextureExtent, PF_R32G32_UINT, FClearValueBinding::Black, TexCreate_ShaderResource | TexCreate_RenderTargetable), TEXT("StrataSSSTexture"));
+			StrataSceneData.SSSTexture = GraphBuilder.CreateTexture(FRDGTextureDesc::Create2D(SceneTextureExtent, PF_R32G32_UINT, FClearValueBinding::Black, TexCreate_ShaderResource | TexCreate_RenderTargetable), TEXT("Strata.SSSTexture"));
 		}
 	}
 	else
@@ -201,7 +201,7 @@ void InitialiseStrataFrameSceneData(FSceneRenderer& SceneRenderer, FRDGBuilder&
 
 	// create the material lob buffer for all views
 	const uint32 MaterialLobesBufferByteSize = FMath::Max(4u, MaterialBufferSizeXY.X * MaterialBufferSizeXY.Y * StrataSceneData.MaxBytesPerPixel);
-	StrataSceneData.MaterialLobesBuffer = GraphBuilder.CreateBuffer(FRDGBufferDesc::CreateByteAddressDesc(MaterialLobesBufferByteSize), TEXT("StrataMaterialBuffer"));
+	StrataSceneData.MaterialLobesBuffer = GraphBuilder.CreateBuffer(FRDGBufferDesc::CreateByteAddressDesc(MaterialLobesBufferByteSize), TEXT("Strata.MaterialBuffer"));
 	StrataSceneData.MaterialLobesBufferSRV = GraphBuilder.CreateSRV(StrataSceneData.MaterialLobesBuffer);
 	StrataSceneData.MaterialLobesBufferUAV = GraphBuilder.CreateUAV(StrataSceneData.MaterialLobesBuffer);
 
