Made SkyAtmosphere default cvar to be more targeted a game rather than cinematic.
Applied from BaseScalability.ini anyway.

diff --git a/Engine/Source/Runtime/Renderer/Private/SkyAtmosphereRendering.cpp b/Engine/Source/Runtime/Renderer/Private/SkyAtmosphereRendering.cpp
--- a/Engine/Source/Runtime/Renderer/Private/SkyAtmosphereRendering.cpp
+++ b/Engine/Source/Runtime/Renderer/Private/SkyAtmosphereRendering.cpp
@@ -71,7 +71,7 @@ static TAutoConsoleVariable<int32> CVarSkyAtmosphereSampleLightShadowmap(
 ////////////////////////////////////////////////////////////////////////// Fast sky
 
 static TAutoConsoleVariable<int32> CVarSkyAtmosphereFastSkyLUT(
-	TEXT("r.SkyAtmosphere.FastSkyLUT"), 0,
+	TEXT("r.SkyAtmosphere.FastSkyLUT"), 1,
 	TEXT("When enabled, a look up texture is used to render the sky.\n")
 	TEXT("It is faster but can result in visual artefacts if there are some high frequency details\n")
 	TEXT("in the sky such as earth shadow or scattering lob."),
@@ -136,7 +136,7 @@ static TAutoConsoleVariable<float> CVarSkyAtmosphereAerialPerspectiveLUTWidth(
 	ECVF_RenderThreadSafe | ECVF_Scalability);
 
 static TAutoConsoleVariable<int32> CVarSkyAtmosphereAerialPerspectiveApplyOnOpaque(
-	TEXT("r.SkyAtmosphere.AerialPerspectiveLUT.FastApplyOnOpaque"), 0,
+	TEXT("r.SkyAtmosphere.AerialPerspectiveLUT.FastApplyOnOpaque"), 1,
 	TEXT("When enabled, the low resolution camera frustum/froxel volume containing atmospheric fog\n")
 	TEXT(", usually used for fog on translucent surface, is used to render fog on opaque.\n")
 	TEXT("It is faster but can result in visual artefacts if there are some high frequency details\n")
