Don't include objects with "Cast Shadow" disabled in Distance Field Shadows
- Added bCastShadow flag to FDFObjectBounds
- Filter objects in CullObjectsForShadowCS

diff --git a/Engine/Shaders/Private/DistanceFieldLightingShared.ush b/Engine/Shaders/Private/DistanceFieldLightingShared.ush
--- a/Engine/Shaders/Private/DistanceFieldLightingShared.ush
+++ b/Engine/Shaders/Private/DistanceFieldLightingShared.ush
@@ -154,6 +154,7 @@ struct FDFObjectBounds
 	float SphereRadius;	// World space bounding sphere extent
 	float3 BoxExtent;	// World space AABB extent
 	uint OftenMoving;
+	uint bCastShadow;
 };
 
 uint NumSceneObjects;
@@ -170,7 +171,10 @@ FDFObjectBounds LoadDFObjectBounds(uint ObjectIndex)
 
 	float4 Vector1 = SceneObjectBounds[ObjectIndex * OBJECT_BOUNDS_STRIDE + 1];
 	Bounds.BoxExtent = Vector1.xyz;
-	Bounds.OftenMoving = Vector1.w;
+
+	uint Flags = asuint(Vector1.w);
+	Bounds.OftenMoving = Flags & 1;
+	Bounds.bCastShadow = Flags & 2;
 
 	return Bounds;
 }
diff --git a/Engine/Shaders/Private/DistanceFieldShadowing.usf b/Engine/Shaders/Private/DistanceFieldShadowing.usf
--- a/Engine/Shaders/Private/DistanceFieldShadowing.usf
+++ b/Engine/Shaders/Private/DistanceFieldShadowing.usf
@@ -148,8 +148,10 @@ void CullObjectsForShadowCS(
 #if (DISTANCEFIELD_PRIMITIVE_TYPE == DFPT_SignedDistanceField)
 		FDFObjectBounds DFObjectBounds = LoadDFObjectBounds(ObjectIndex);
 
-        if ((ShadowBoundingSphere.w == 0 && ShadowConvexHullIntersectSphere(DFObjectBounds.Center, DFObjectBounds.SphereRadius))
-			|| (ShadowBoundingSphere.w > 0 && dot(ShadowBoundingSphere.xyz - DFObjectBounds.Center, ShadowBoundingSphere.xyz - DFObjectBounds.Center) < Square(ShadowBoundingSphere.w + DFObjectBounds.SphereRadius)))
+		const bool bConvexHullIntersect = ShadowBoundingSphere.w == 0 && ShadowConvexHullIntersectSphere(DFObjectBounds.Center, DFObjectBounds.SphereRadius);
+		const bool bBoundingSphereIntersect = ShadowBoundingSphere.w > 0 && dot(ShadowBoundingSphere.xyz - DFObjectBounds.Center, ShadowBoundingSphere.xyz - DFObjectBounds.Center) < Square(ShadowBoundingSphere.w + DFObjectBounds.SphereRadius);
+
+        if (DFObjectBounds.bCastShadow && (bConvexHullIntersect || bBoundingSphereIntersect))
 		{
 			FDFObjectData DFObjectData = LoadDFObjectData(ObjectIndex);
 
diff --git a/Engine/Source/Runtime/Renderer/Private/DistanceFieldObjectManagement.cpp b/Engine/Source/Runtime/Renderer/Private/DistanceFieldObjectManagement.cpp
--- a/Engine/Source/Runtime/Renderer/Private/DistanceFieldObjectManagement.cpp
+++ b/Engine/Source/Runtime/Renderer/Private/DistanceFieldObjectManagement.cpp
@@ -708,10 +708,14 @@ void FDistanceFieldSceneData::UpdateDistanceFieldObjectBuffers(
 									UploadObjectBounds[0] = ObjectBoundingSphere;
 
 									const FGlobalDFCacheType CacheType = PrimitiveSceneProxy->IsOftenMoving() ? GDF_Full : GDF_MostlyStatic;
-									const float OftenMovingValue = CacheType == GDF_Full ? 1.0f : 0.0f;
+									const uint32 bOftenMoving = CacheType == GDF_Full;
+									const uint32 bCastShadow = PrimitiveSceneProxy->CastsDynamicShadow();
 
-									const FVector4f ObjectWorldExtent(WorldSpaceMeshBounds.GetExtent(), OftenMovingValue);
-									UploadObjectBounds[1] = ObjectWorldExtent;
+									const uint32 Flags = bOftenMoving | (bCastShadow << 1);
+
+									FVector4f ObjectWorldExtentAndFlags(WorldSpaceMeshBounds.GetExtent(), 0.0f);
+									ObjectWorldExtentAndFlags.W = *(const float*)&Flags;
+									UploadObjectBounds[1] = ObjectWorldExtentAndFlags;
 
 									// Uniformly scale our Volume space to lie within [-1, 1] at the max extent
 									// This is mirrored in the SDF encoding
@@ -1140,4 +1144,4 @@ void DistanceField::BuildUpdateTrackingBounds(const TArray<FViewInfo>& Views, Di
 			}
 		}
 	}
-}
\ No newline at end of file
+}
