Properly filter shadow casting draws into VSM & non-VSM shadow maps (avoids caching empty SMs if all geometry supports VSM)
- fixes both naniite and non-nanite cases.

diff --git a/Engine/Source/Runtime/RenderCore/Public/VertexFactory.h b/Engine/Source/Runtime/RenderCore/Public/VertexFactory.h
--- a/Engine/Source/Runtime/RenderCore/Public/VertexFactory.h
+++ b/Engine/Source/Runtime/RenderCore/Public/VertexFactory.h
@@ -649,6 +649,13 @@ public:
 		return PrimitiveIdStreamIndex[TranslatePrimitiveIdStreamIndex(InFeatureLevel, InputStreamType)];
 	}
 
+	/** 
+	 */
+	inline bool SupportsGPUScene(const FStaticFeatureLevel InFeatureLevel) const
+	{
+		return GetType()->SupportsPrimitiveIdStream() &&  GetPrimitiveIdStreamIndex(InFeatureLevel, EVertexInputStreamType::Default) != INDEX_NONE;
+	}
+
 protected:
 	void SetPrimitiveIdStreamIndex(const FStaticFeatureLevel InFeatureLevel, EVertexInputStreamType InputStreamType, int32 StreamIndex)
 	{
diff --git a/Engine/Source/Runtime/Renderer/Private/PrimitiveSceneInfo.cpp b/Engine/Source/Runtime/Renderer/Private/PrimitiveSceneInfo.cpp
--- a/Engine/Source/Runtime/Renderer/Private/PrimitiveSceneInfo.cpp
+++ b/Engine/Source/Runtime/Renderer/Private/PrimitiveSceneInfo.cpp
@@ -90,6 +90,7 @@ public:
 			bool bUseSingleLayerWaterMaterial = Material.GetShadingModels().HasShadingModel(MSM_SingleLayerWater);
 			bool bUseAnisotropy = Material.GetShadingModels().HasAnyShadingModel({MSM_DefaultLit, MSM_ClearCoat}) && Material.MaterialUsesAnisotropy_RenderThread();
 			bool bSupportsNaniteRendering = SupportsNaniteRendering(StaticMesh->VertexFactory, PrimitiveSceneProxy, Mesh.MaterialRenderProxy, FeatureLevel);
+			bool bSupportsGPUScene = StaticMesh->VertexFactory->SupportsGPUScene(FeatureLevel);
 
 			FStaticMeshBatchRelevance* StaticMeshRelevance = new(PrimitiveSceneInfo->StaticMeshRelevances) FStaticMeshBatchRelevance(
 				*StaticMesh, 
@@ -99,6 +100,7 @@ public:
 				bUseSingleLayerWaterMaterial,
 				bUseAnisotropy,
 				bSupportsNaniteRendering,
+				bSupportsGPUScene,
 				FeatureLevel
 				);
 		}
diff --git a/Engine/Source/Runtime/Renderer/Private/SceneCore.h b/Engine/Source/Runtime/Renderer/Private/SceneCore.h
--- a/Engine/Source/Runtime/Renderer/Private/SceneCore.h
+++ b/Engine/Source/Runtime/Renderer/Private/SceneCore.h
@@ -163,7 +163,7 @@ private:
 class FStaticMeshBatchRelevance
 {
 public:
-	FStaticMeshBatchRelevance(const FStaticMeshBatch& StaticMesh, float InScreenSize, bool InbSupportsCachingMeshDrawCommands, bool InbUseSkyMaterial, bool bInUseSingleLayerWaterMaterial, bool bInUseAnisotropy, bool bInSupportsNaniteRendering, ERHIFeatureLevel::Type FeatureLevel)
+	FStaticMeshBatchRelevance(const FStaticMeshBatch& StaticMesh, float InScreenSize, bool InbSupportsCachingMeshDrawCommands, bool InbUseSkyMaterial, bool bInUseSingleLayerWaterMaterial, bool bInUseAnisotropy, bool bInSupportsNaniteRendering, bool bInSupportsGPUScene, ERHIFeatureLevel::Type FeatureLevel)
 		: Id(StaticMesh.Id)
 		, ScreenSize(InScreenSize)
 		, NumElements(StaticMesh.Elements.Num())
@@ -183,6 +183,7 @@ public:
 		, RuntimeVirtualTextureMaterialType(StaticMesh.RuntimeVirtualTextureMaterialType)
 		, bSupportsCachingMeshDrawCommands(InbSupportsCachingMeshDrawCommands)
 		, bSupportsNaniteRendering(bInSupportsNaniteRendering)
+		, bSupportsGPUScene(bInSupportsGPUScene)
 	{
 	}
 
@@ -227,9 +228,12 @@ public:
 	/** Cached from vertex factory to avoid dereferencing VF in InitViews. */
 	uint8 bSupportsCachingMeshDrawCommands : 1;
 
-	/** Cached from vertex factor to avoid dereferencing VF in InitViews. */
+	/** Cached from vertex factory to avoid dereferencing VF in InitViews. */
 	uint8 bSupportsNaniteRendering : 1;
 
+	/** Cached from vertex factory to avoid dereferencing VF in shadow depth rendering. */
+	uint8 bSupportsGPUScene : 1;
+
 	/** Computes index of cached mesh draw command in FPrimitiveSceneInfo::CachedMeshDrawCommandInfos, for a given mesh pass. */
 	int32 GetStaticMeshCommandInfoIndex(EMeshPass::Type MeshPass) const;
 };
diff --git a/Engine/Source/Runtime/Renderer/Private/ShadowDepthRendering.cpp b/Engine/Source/Runtime/Renderer/Private/ShadowDepthRendering.cpp
--- a/Engine/Source/Runtime/Renderer/Private/ShadowDepthRendering.cpp
+++ b/Engine/Source/Runtime/Renderer/Private/ShadowDepthRendering.cpp
@@ -2347,39 +2347,15 @@ bool FShadowDepthPassMeshProcessor::TryAddMeshBatch(const FMeshBatch& RESTRICT M
 	bool bResult = true;
 	if (bShouldCastShadow
 		&& ShouldIncludeDomainInMeshPass(Material.GetMaterialDomain())
-		&& ShouldIncludeMaterialInDefaultOpaquePass(Material))
+		&& ShouldIncludeMaterialInDefaultOpaquePass(Material)
+		&& EnumHasAnyFlags(MeshSelectionMask, MeshBatch.VertexFactory->SupportsGPUScene(FeatureLevel) ? EShadowMeshSelection::VSM : EShadowMeshSelection::SM))
 	{
 		const FMaterialRenderProxy* EffectiveMaterialRenderProxy = &MaterialRenderProxy;
 		const FMaterial* EffectiveMaterial = &Material;
 
 		OverrideWithDefaultMaterialForShadowDepth(EffectiveMaterialRenderProxy, EffectiveMaterial, FeatureLevel);
 
-		bool bDraw = true;
-		if (UseNonNaniteVirtualShadowMaps(GShaderPlatformForFeatureLevel[FeatureLevel], FeatureLevel))
-		{
-			// TODO: This uses a lot of indirections and complex logic, optimize by precomputing as far as possible
-			//       E.g., maybe store  bSupportsGpuSceneInstancing as a flag in the MeshBatch?
-			const FVertexFactory* VertexFactory = MeshBatch.VertexFactory;
-			const bool bUsePositionOnlyVS =
-				VertexFactory->SupportsPositionAndNormalOnlyStream()
-				&& EffectiveMaterial->WritesEveryPixel(true)
-				&& !EffectiveMaterial->MaterialModifiesMeshPosition_RenderThread();
-
-			// TODO: Store in MeshBatch?
-			const bool bSupportsGpuSceneInstancing = UseGPUScene(GShaderPlatformForFeatureLevel[FeatureLevel], FeatureLevel)
-				&& VertexFactory->GetPrimitiveIdStreamIndex(FeatureLevel, bUsePositionOnlyVS ? EVertexInputStreamType::PositionAndNormalOnly : EVertexInputStreamType::Default) != INDEX_NONE;
-
-			// EMeshPass::CSMShadowDepth: If no VSM: include everything, else only !bSupportsGpuSceneInstancing
-			// EMeshPass::VSMShadowDepth: If VSM: only bSupportsGpuSceneInstancing, else nothing needs to go in.
-			// TODO: I'm sure I can reduce this logic
-			bDraw = MeshPassTargetType == EMeshPass::CSMShadowDepth ?
-				(!UseVirtualShadowMaps(GShaderPlatformForFeatureLevel[FeatureLevel], FeatureLevel) || !bSupportsGpuSceneInstancing) :
-				(UseVirtualShadowMaps(GShaderPlatformForFeatureLevel[FeatureLevel], FeatureLevel) && bSupportsGpuSceneInstancing);
-		}
-		if (bDraw)
-		{
-			bResult = Process(MeshBatch, BatchElementMask, StaticMeshId, PrimitiveSceneProxy, *EffectiveMaterialRenderProxy, *EffectiveMaterial, MeshFillMode, FinalCullMode);
-		}
+		bResult = Process(MeshBatch, BatchElementMask, StaticMeshId, PrimitiveSceneProxy, *EffectiveMaterialRenderProxy, *EffectiveMaterial, MeshFillMode, FinalCullMode);
 	}
 
 	return bResult;
@@ -2415,6 +2391,16 @@ FShadowDepthPassMeshProcessor::FShadowDepthPassMeshProcessor(
 	, ShadowDepthType(InShadowDepthType)
 	, MeshPassTargetType(InMeshPassTargetType)
 {
+	if (UseNonNaniteVirtualShadowMaps(Scene->GetShaderPlatform(), Scene->GetFeatureLevel()))
+	{
+		// set up mesh filtering.
+		MeshSelectionMask = MeshPassTargetType == EMeshPass::VSMShadowDepth ? EShadowMeshSelection::VSM : EShadowMeshSelection::SM;
+	}
+	else
+	{
+		// If VSMs are disabled, pipe all kinds of draws into the regular SMs
+		MeshSelectionMask = EShadowMeshSelection::All;
+	}
 	SetStateForShadowDepth(ShadowDepthType.bOnePassPointLightShadow, ShadowDepthType.bDirectionalLight, PassDrawRenderState, MeshPassTargetType);
 }
 
diff --git a/Engine/Source/Runtime/Renderer/Private/ShadowRendering.h b/Engine/Source/Runtime/Renderer/Private/ShadowRendering.h
--- a/Engine/Source/Runtime/Renderer/Private/ShadowRendering.h
+++ b/Engine/Source/Runtime/Renderer/Private/ShadowRendering.h
@@ -97,6 +97,23 @@ public:
 
 extern FShadowDepthType CSMShadowDepthType;
 
+
+
+/**
+ * Used to select what meshes to draw into what shadow infos. Each mesh selects the support it has, and
+ * the shadow info stores a mask of what types it collects. This makes it easy to allow regular SMs to
+ * collect both types if VSMs are disabled.
+ */
+enum class EShadowMeshSelection : uint8
+{
+	SM = 1U << 0U,
+	VSM = 1U << 1U, // declared if GPU-Scene instace culling is supported.
+	All = SM | VSM,
+};
+
+ENUM_CLASS_FLAGS(EShadowMeshSelection)
+
+
 class FShadowDepthPassMeshProcessor : public FMeshPassProcessor
 {
 public:
@@ -132,6 +149,7 @@ private:
 
 	FShadowDepthType ShadowDepthType;
 	EMeshPass::Type MeshPassTargetType = EMeshPass::CSMShadowDepth;
+	EShadowMeshSelection MeshSelectionMask = EShadowMeshSelection::All;
 };
 
 enum EShadowDepthCacheMode
@@ -368,6 +386,7 @@ public:
 	/** Used to fetch the correct cached static mesh draw commands */
 	EMeshPass::Type MeshPassTargetType = EMeshPass::CSMShadowDepth;
 
+	EShadowMeshSelection MeshSelectionMask = EShadowMeshSelection::All;
 	/** */
 	TArray< FVirtualShadowMap*, TInlineAllocator<6> > VirtualShadowMaps;
 
@@ -805,6 +824,12 @@ private:
 	 */
 	FORCEINLINE_DEBUGGABLE FLODMask CalcAndUpdateLODToRender(FViewInfo& CurrentView, const FBoxSphereBounds& Bounds, const FPrimitiveSceneInfo* PrimitiveSceneInfo, int32 ForcedLOD) const;
 
+	/**
+	 * Check relevance flags and shadow mesh pass filter, returns true if the mesh is to be added.
+	 * @param bOutDrawingStaticMeshes - set to true if the relevancy checks pass, even if the mesh was discarded using the pass filter check. Otherwise the mesh is treated as dynamic.
+	 */
+	FORCEINLINE bool ShouldDrawStaticMesh(const FStaticMeshBatchRelevance& StaticMeshRelevance, const FLODMask &ShadowLODToRender, bool& bOutDrawingStaticMeshes) const;
+
 	/** Will return if we should draw the static mesh for the shadow, and will perform lazy init of primitive if it wasn't visible */
 	bool ShouldDrawStaticMeshes(FViewInfo& InCurrentView, FPrimitiveSceneInfo* InPrimitiveSceneInfo);
 
diff --git a/Engine/Source/Runtime/Renderer/Private/ShadowSetup.cpp b/Engine/Source/Runtime/Renderer/Private/ShadowSetup.cpp
--- a/Engine/Source/Runtime/Renderer/Private/ShadowSetup.cpp
+++ b/Engine/Source/Runtime/Renderer/Private/ShadowSetup.cpp
@@ -1058,6 +1058,7 @@ void FProjectedShadowInfo::SetupClipmapProjection(FLightSceneInfo* InLightSceneI
 	bIncludeInScreenSpaceShadowMask = false;
 	MaxNonFarCascadeDistance = InMaxNonFarCascadeDistance;
 	MeshPassTargetType = EMeshPass::VSMShadowDepth;
+	MeshSelectionMask = EShadowMeshSelection::VSM;
 
 
 	FVirtualShadowMapProjectionShaderData Data = VirtualShadowMapClipmap->GetProjectionShaderData(VirtualShadowMapClipmap->GetLevelCount() - 1);
@@ -1430,6 +1431,21 @@ FLODMask FProjectedShadowInfo::CalcAndUpdateLODToRender(FViewInfo& CurrentView,
 	return ShadowLODToRender;
 }
 
+FORCEINLINE bool FProjectedShadowInfo::ShouldDrawStaticMesh(const FStaticMeshBatchRelevance& StaticMeshRelevance, const FLODMask& ShadowLODToRender, bool& bOutDrawingStaticMeshes) const
+{
+	if ((StaticMeshRelevance.CastShadow || (bSelfShadowOnly && StaticMeshRelevance.bUseForDepthPass)) && ShadowLODToRender.ContainsLOD(StaticMeshRelevance.LODIndex))
+	{
+		bOutDrawingStaticMeshes = true;
+
+		if (EnumHasAnyFlags(MeshSelectionMask, StaticMeshRelevance.bSupportsGPUScene ? EShadowMeshSelection::VSM : EShadowMeshSelection::SM))
+		{
+			return true;
+		}
+	}
+	return false;
+}
+
+
 bool FProjectedShadowInfo::ShouldDrawStaticMeshes(FViewInfo& InCurrentView, FPrimitiveSceneInfo* InPrimitiveSceneInfo)
 {
 	bool WholeSceneDirectionalShadow = IsWholeSceneDirectionalShadow();
@@ -1451,7 +1467,7 @@ bool FProjectedShadowInfo::ShouldDrawStaticMeshes(FViewInfo& InCurrentView, FPri
 				const FStaticMeshBatchRelevance& StaticMeshRelevance = InPrimitiveSceneInfo->StaticMeshRelevances[MeshIndex];
 				const FStaticMeshBatch& StaticMesh = InPrimitiveSceneInfo->StaticMeshes[MeshIndex];
 
-				if ((StaticMeshRelevance.CastShadow || (bSelfShadowOnly && StaticMeshRelevance.bUseForDepthPass)) && ShadowLODToRender.ContainsLOD(StaticMeshRelevance.LODIndex))
+				if (ShouldDrawStaticMesh(StaticMeshRelevance, ShadowLODToRender, bDrawingStaticMeshes))
 				{
 					if (GetShadowDepthType() == CSMShadowDepthType && bCanCache)
 					{
@@ -1471,8 +1487,6 @@ bool FProjectedShadowInfo::ShouldDrawStaticMeshes(FViewInfo& InCurrentView, FPri
 						NumSubjectMeshCommandBuildRequestElements += StaticMeshRelevance.NumElements;
 						SubjectMeshCommandBuildRequests.Add(&StaticMesh);
 					}
-
-					bDrawingStaticMeshes = true;
 				}
 			}
 		}
@@ -1483,12 +1497,10 @@ bool FProjectedShadowInfo::ShouldDrawStaticMeshes(FViewInfo& InCurrentView, FPri
 				const FStaticMeshBatchRelevance& StaticMeshRelevance = InPrimitiveSceneInfo->StaticMeshRelevances[MeshIndex];
 				const FStaticMeshBatch& StaticMesh = InPrimitiveSceneInfo->StaticMeshes[MeshIndex];
 
-				if ((StaticMeshRelevance.CastShadow || (bSelfShadowOnly && StaticMeshRelevance.bUseForDepthPass)) && ShadowLODToRender.ContainsLOD(StaticMeshRelevance.LODIndex))
+				if (ShouldDrawStaticMesh(StaticMeshRelevance, ShadowLODToRender, bDrawingStaticMeshes))
 				{
 					NumSubjectMeshCommandBuildRequestElements += StaticMeshRelevance.NumElements;
 					SubjectMeshCommandBuildRequests.Add(&StaticMesh);
-
-					bDrawingStaticMeshes = true;
 				}
 			}
 		}
@@ -1528,7 +1540,7 @@ bool FProjectedShadowInfo::ShouldDrawStaticMeshes_AnyThread(
 				const FStaticMeshBatchRelevance& StaticMeshRelevance = PrimitiveSceneInfo->StaticMeshRelevances[MeshIndex];
 				const FStaticMeshBatch& StaticMesh = PrimitiveSceneInfo->StaticMeshes[MeshIndex];
 
-				if ((StaticMeshRelevance.CastShadow || (bSelfShadowOnly && StaticMeshRelevance.bUseForDepthPass)) && ShadowLODToRender.ContainsLOD(StaticMeshRelevance.LODIndex))
+				if (ShouldDrawStaticMesh(StaticMeshRelevance, ShadowLODToRender, bDrawingStaticMeshes))
 				{
 					if (bCanCache && GetShadowDepthType() == CSMShadowDepthType)
 					{
@@ -1539,8 +1551,6 @@ bool FProjectedShadowInfo::ShouldDrawStaticMeshes_AnyThread(
 						++OutStats.NumMDCBuildRequests;
 						OutResult.AcceptMesh(NumAcceptedStaticMeshes++, MeshIndex, OverflowBuffer);
 					}
-
-					bDrawingStaticMeshes = true;
 				}
 			}
 		}
@@ -1553,13 +1563,11 @@ bool FProjectedShadowInfo::ShouldDrawStaticMeshes_AnyThread(
 				const FStaticMeshBatchRelevance& StaticMeshRelevance = PrimitiveSceneInfo->StaticMeshRelevances[MeshIndex];
 				const FStaticMeshBatch& StaticMesh = PrimitiveSceneInfo->StaticMeshes[MeshIndex];
 
-				if ((StaticMeshRelevance.CastShadow || (bSelfShadowOnly && StaticMeshRelevance.bUseForDepthPass)) && ShadowLODToRender.ContainsLOD(StaticMeshRelevance.LODIndex))
+				if (ShouldDrawStaticMesh(StaticMeshRelevance, ShadowLODToRender, bDrawingStaticMeshes))
 				{
 					check(MeshIndex < MAX_uint16);
 					++OutStats.NumMDCBuildRequests;
 					OutResult.AcceptMesh(NumAcceptedStaticMeshes++, MeshIndex, OverflowBuffer);
-
-					bDrawingStaticMeshes = true;
 				}
 			}
 		}
@@ -3707,6 +3715,7 @@ void FSceneRenderer::CreateWholeSceneProjectedShadow(
 					}
 
 					ProjectedShadowInfo->MeshPassTargetType = EMeshPass::VSMShadowDepth;
+					ProjectedShadowInfo->MeshSelectionMask = EShadowMeshSelection::VSM;
 
 					bool bContainsNaniteSubjects = false;
 					AddInteractingPrimitives(LightSceneInfo->GetDynamicInteractionOftenMovingPrimitiveList(), ProjectedShadowInfo, bContainsNaniteSubjects);
@@ -3741,9 +3750,12 @@ void FSceneRenderer::CreateWholeSceneProjectedShadow(
 						MemStackProjectedShadows.Add(ProjectedShadowInfo);
 
 						// If we have a virtual shadow map, disable nanite rendering into the regular shadow map or else we'd get double-shadowing
+						// and also filter out any VSM-supporting meshes
 						if (bNeedsVirtualShadowMap)
 						{
 							ProjectedShadowInfo->bNaniteGeometry = false;
+							// Note: the default is all types
+							ProjectedShadowInfo->MeshSelectionMask = EShadowMeshSelection::SM;
 						}
 
 						bool bContainsNaniteSubjects = false;
@@ -3769,7 +3781,7 @@ void FSceneRenderer::CreateWholeSceneProjectedShadow(
 					
 						if (CacheMode[CacheModeIndex] == SDCM_StaticPrimitivesOnly)
 						{
-							const bool bHasStaticPrimitives = ProjectedShadowInfo->HasSubjectPrims() || bContainsNaniteSubjects;
+							const bool bHasStaticPrimitives = ProjectedShadowInfo->HasSubjectPrims() || (bContainsNaniteSubjects && !bNeedsVirtualShadowMap);
 							bRenderShadow = bHasStaticPrimitives;
 							FCachedShadowMapData& CachedShadowMapData = Scene->GetCachedShadowMapDataRef(ProjectedShadowInfo->GetLightSceneInfo().Id);
 							CachedShadowMapData.bCachedShadowMapHasPrimitives = bHasStaticPrimitives;
@@ -5042,8 +5054,16 @@ void FSceneRenderer::AllocateShadowDepthTargets(FRHICommandListImmediate& RHICmd
 		{
 			AllocateCachedShadowDepthTargets(RHICmdList, CachedWholeSceneDirectionalShadows);
 
-			AllocateCSMDepthTargets(RHICmdList, WholeSceneDirectionalShadows, SortedShadowsForShadowDepthPass.ShadowMapAtlases);
-
+			// Only allocate CSM targets if at least one of the SMs in the CSM has any primitives, but then always allocate all of them as some algorithms depend on this (notably the forward shading structures).
+			bool bAnyHasSubjectPrims = false;
+			for (const FProjectedShadowInfo* ProjectedShadowInfo : WholeSceneDirectionalShadows)
+			{
+				bAnyHasSubjectPrims = bAnyHasSubjectPrims || ProjectedShadowInfo->HasSubjectPrims();
+			}
+			if (bAnyHasSubjectPrims)
+			{
+				AllocateCSMDepthTargets(RHICmdList, WholeSceneDirectionalShadows, SortedShadowsForShadowDepthPass.ShadowMapAtlases);
+			}
 			CachedWholeSceneDirectionalShadows.Empty();
 		}
 		else
