Move strata classification in base pass in order to reduce (read) memory bandwidth.

diff --git a/Engine/Shaders/Private/BasePassPixelShader.usf b/Engine/Shaders/Private/BasePassPixelShader.usf
--- a/Engine/Shaders/Private/BasePassPixelShader.usf
+++ b/Engine/Shaders/Private/BasePassPixelShader.usf
@@ -1593,11 +1593,23 @@ void FPixelShaderInOut_MainPS(
 
 	uint2 PixelPos = uint2(In.SvPosition.xy);
 
+	// classificaiton & pre-preprocessed data
+	FStrataClassification ClassificationData = (FStrataClassification)0;
+	FStrataSubsurfaceData SSSData = (FStrataSubsurfaceData)0;
+	FStrataTopLayerData TopLayerData = (FStrataTopLayerData)0;
+
 	float3 EmissiveLuminance = 0.0f;
 	PackStrataOut(
-		OpaqueBasePass.Strata.MaterialLobesBufferUAV, Dither,
+		OpaqueBasePass.Strata.MaterialLobesBufferUAV,
+		Dither,
 		GetStrataPixelDataByteOffset(PixelPos, uint2(ResolvedView.BufferSizeAndInvSize.xy), OpaqueBasePass.Strata.MaxBytesPerPixel),
-		StrataPixelHeader, StrataData, MaterialParameters.CameraVector, bStrataSubsurfaceEnable, EmissiveLuminance);
+		StrataPixelHeader, StrataData, MaterialParameters.CameraVector, bStrataSubsurfaceEnable, EmissiveLuminance,
+		ClassificationData, SSSData, TopLayerData);
+
+	// Write out classificaiton & pre-preprocessed data
+	OpaqueBasePass.Strata.ClassificationTextureUAV[PixelPos] = StrataPackClassificationData(ClassificationData);
+	OpaqueBasePass.Strata.TopLayerNormalTextureUAV[PixelPos] = StrataPackTopLayerData(TopLayerData);
+	OpaqueBasePass.Strata.SSSTextureUAV[PixelPos] = StrataPackSSSData(SSSData);
 
 	Color += EmissiveLuminance;
 #endif // STRATA_OPAQUE_DEFERRED
diff --git a/Engine/Shaders/Private/DebugProbes.usf b/Engine/Shaders/Private/DebugProbes.usf
--- a/Engine/Shaders/Private/DebugProbes.usf
+++ b/Engine/Shaders/Private/DebugProbes.usf
@@ -273,13 +273,18 @@ void MainPS(
 		StrataPixelHeader.MaterialAO = 1.0f - IndirectIrradianceAndAO;
 		StrataPixelHeader.BSDFCount = 1;
 
+		FStrataClassification ClassificationData = (FStrataClassification)0;
+		FStrataSubsurfaceData SSSData = (FStrataSubsurfaceData)0;
+		FStrataTopLayerData TopLayerData = (FStrataTopLayerData)0;
+
 		float3 EmissiveLuminance = 0.0f;
 		bool bStrataSubsurfaceEnable = false;
 		const float Dither = 0.0f;
 		PackStrataOut(
 			MaterialLobesBufferUAV, Dither,
 			GetStrataPixelDataByteOffset(PixelCoord, uint2(ResolvedView.BufferSizeAndInvSize.xy), MaxBytesPerPixel),
-			StrataPixelHeader, StrataData, -RayDir, bStrataSubsurfaceEnable, EmissiveLuminance);
+			StrataPixelHeader, StrataData, -RayDir, bStrataSubsurfaceEnable, EmissiveLuminance,
+			ClassificationData, SSSData, TopLayerData);
 	}
 	
 #else // PROJECT_STRATA
diff --git a/Engine/Shaders/Private/ShadingCommon.ush b/Engine/Shaders/Private/ShadingCommon.ush
--- a/Engine/Shaders/Private/ShadingCommon.ush
+++ b/Engine/Shaders/Private/ShadingCommon.ush
@@ -90,6 +90,11 @@ bool GetShadingModelRequiresBackfaceLighting(uint ShadingModelID)
 }
 
 
+float F0ToDielectricSpecular(float F0)
+{
+	return saturate(F0 / 0.08f);
+}
+
 float DielectricSpecularToF0(float Specular)
 {
 	return 0.08f * Specular;
diff --git a/Engine/Shaders/Private/Strata/Strata.ush b/Engine/Shaders/Private/Strata/Strata.ush
--- a/Engine/Shaders/Private/Strata/Strata.ush
+++ b/Engine/Shaders/Private/Strata/Strata.ush
@@ -307,6 +307,7 @@ uint   PackColorLinearToGamma2AlphaLinear(float4 rgba);
 float4 UnpackColorGamma2ToLinearAlphaLinear(uint rgba);
 uint   StrataPackNormal(in float3 Normal);
 float3 StrataUnpackNormal(uint PackedNormal);
+bool IsStrataSlabFastPathCompatible(in FStrataBSDF BSDF);
 
 ///////////////////////////////////////////////////////////////////////////////
 // Classification 
@@ -1728,7 +1729,10 @@ void PackStrataOut(
 	FStrataData Strata,
 	float3 V,
 	inout bool bStrataSubsurfaceEnable,
-	inout float3 EmissiveColor)
+	inout float3 EmissiveColor,
+	inout FStrataClassification ClassificationData,
+	inout FStrataSubsurfaceData SSSData,
+	inout FStrataTopLayerData TopLayerData)
 {
 	bStrataSubsurfaceEnable = false;
 	EmissiveColor = 0.0f;
@@ -1737,6 +1741,14 @@ void PackStrataOut(
 	const float FullThroughput = 1.0f;
 	const float FullyRough = 1.0f;
 
+	// While packing Strata layer data, Classification/SSS/TopLayer data are extracted & stored for dedicated passes
+	// This avoid to run a post-basepass which would re-read all the material data
+	ClassificationData = (FStrataClassification)0;
+	SSSData = (FStrataSubsurfaceData)0;
+	TopLayerData = (FStrataTopLayerData)0;
+
+	SetIsSimpleMaterial(ClassificationData, true); // non strata material or sky pixels are considered as simple (so that (strata simple pixels) + (non-strata pixels) can build a simple tile).
+
 	uint BSDFCount = 0;
 	uint BSDFVisibleCount = 0;
 	float3 Throughput = 1.0f;
@@ -1767,11 +1779,11 @@ void PackStrataOut(
 
 			FStrataBSDF BSDF = Strata.Layers[l].BSDFs[i];
 
-#if STRATA_INLINE_SHADING
+		#if STRATA_INLINE_SHADING
 			const float BSDFCoverage = saturate(BSDF.Coverage);
-#else
+		#else
 			const float BSDFCoverage = 1.0f;
-#endif
+		#endif
 			BSDFCount++;
 
 			// STRATA_TODO apply the same normal transform as for the material
@@ -1783,6 +1795,12 @@ void PackStrataOut(
 			const float3 FakeL = float3(0, 0, 1);
 			Init(ContextNoLight, N, V, FakeL);
 
+			if (bTopLayer)
+			{
+				// together with a threshold to avoid 0-length normal after normalization in case influence is fadeout using the Weight node.
+				TopLayerData.WorldNormal += N * max(1e-5, BSDFCoverage);
+			}
+
 			switch (BSDF_GETTYPE(BSDF))
 			{
 			case STRATA_BSDF_TYPE_SLAB:
@@ -1847,6 +1865,11 @@ void PackStrataOut(
 					BSDF_SETISTHIN(Strata.Layers[l].BSDFs[i], 0);
 				}
 
+				if (bTopLayer)
+				{
+					ClassificationData.TopLayerRoughness += BSDFCoverage * SLAB_ROUGHNESS(BSDF);
+				}
+
 				// STRATA_TODO	Should we have the Fuzz influence the next layter visibility accoding to fuzz amount? Should it also occlude diffuse?
 				//				If this is done, we should also update ComputePostProcessBSDFData.
 
@@ -1857,6 +1880,11 @@ void PackStrataOut(
 
 			case STRATA_BSDF_TYPE_HAIR:
 			{
+				if (bTopLayer)
+				{
+					ClassificationData.TopLayerRoughness += BSDFCoverage * FullyRough;
+				}
+
 				// Hair are considered opaque
 				LayerCoverage	+= BSDFCoverage;
 				LayerThroughput += BSDFCoverage * OpaqueBSDFThroughput;
@@ -1867,6 +1895,11 @@ void PackStrataOut(
 			{
 				EmissiveColor += BSDF_GETEMISSIVE(BSDF) * Throughput;
 
+				if (bTopLayer)
+				{
+					ClassificationData.TopLayerRoughness += BSDFCoverage * SLW_ROUGHNESS(BSDF);
+				}
+
 				// Water is considered opaque, it has a custom process to simulate translucency
 				LayerCoverage += BSDFCoverage;
 				LayerThroughput += BSDFCoverage * OpaqueBSDFThroughput;
@@ -1876,12 +1909,12 @@ void PackStrataOut(
 			} // switch
 
 
-#if STRATA_INLINE_SHADING
+		#if STRATA_INLINE_SHADING
 			// Update the BSDF final weight with throughput towards the view from layers above and coverage if it has not been baked into the parameters.
 			Strata.Layers[l].BSDFs[i].FinalWeight = Throughput * (Strata.Layers[l].BSDFs[i].bPremultipliedCoverage ? 1.0 : Strata.Layers[l].BSDFs[i].Coverage);
-#else
+		#else
 			Strata.Layers[l].BSDFs[i].FinalWeight = Throughput;
-#endif
+		#endif
 			
 			// Cull non visible BSDF
 			if (any(Strata.Layers[l].BSDFs[i].FinalWeight > 0))
@@ -1893,6 +1926,16 @@ void PackStrataOut(
 			BSDF_SETISTOPLAYER(Strata.Layers[l].BSDFs[i], bTopLayer ? 1 : 0);
 		}
 
+		// Average the TopLayer data once we have read all the data
+		if (bTopLayer)
+		{
+			// Top layer roughness & normal have been weighted with BSDF Coverage only (since the accumulated throughput is 1.f for the top-layer)
+			const float BSDFWeightAvg = dot(LayerCoverage, 1.0f / 3.0f);
+
+			TopLayerData.WorldNormal = normalize(TopLayerData.WorldNormal);
+			ClassificationData.TopLayerRoughness = BSDFWeightAvg > 0.0f ? ClassificationData.TopLayerRoughness / BSDFWeightAvg : 0.0f;
+		}
+
 		// Take into account the non covered part for this layer and update throughput accordingly
 		LayerThroughput += (1.0 - saturate(LayerCoverage)) * FullThroughput;
 
@@ -1903,6 +1946,7 @@ void PackStrataOut(
 	// Now write out Strata data
 
 	// 1. the header
+	AndIsSimpleMaterial(ClassificationData, BSDFVisibleCount == 1);
 	STRATA_STORE_UINT1(PackStrataHeaderIn(BSDFVisibleCount, StrataPixelHeader));
 #if STRATA_INLINE_SHADING
 	UNROLL
@@ -1950,10 +1994,14 @@ void PackStrataOut(
 				}
 
 				const uint BSDFType = BSDF_GETTYPE(BSDF);
+				AddShadingModels(ClassificationData, StrataShadingModelBit(BSDFType));
+
 				switch (BSDF_GETTYPE(BSDF))
 				{
 				case STRATA_BSDF_TYPE_SLAB:
 				{
+					AndIsSimpleMaterial(ClassificationData, IsStrataSlabFastPathCompatible(BSDF));
+
 					// Now we pack diffuse and F0 in a special way: both encoded as R7G7B6 with gamma 2.0. The low bit cound will be hidden by dithering and TAA.
 					// We pack 32 bits in the first uint and the remaining bits are pack is the lowest significant bit of the second uint.
 					const uint PackedDiffuse20Bits	= PackR7G7B6Gamma2(SLAB_DIFFUSEALBEDO(BSDF));
@@ -1980,6 +2028,25 @@ void PackStrataOut(
 					const bool bIsThin = BSDF_GETISTHIN(BSDF);
 					if (BSDF_GETHASSSS(BSDF) || bIsSimpleVolume || bIsThin)
 					{
+						SetHasSubsurface(ClassificationData, !bIsThin && BSDF_GETHASSSS(BSDF));
+						SetHasSubsurfaceProfile(ClassificationData, !bIsThin && BSDF_GETHASSSSPROFILE(BSDF));
+
+						if (!bIsThin && BSDF_GETHASSSS(BSDF))
+						{
+							SSSData.BaseColor = SLAB_DIFFUSEALBEDO(BSDF);
+							SSSData.Specular  = F0ToDielectricSpecular(SLAB_F0(BSDF).x);
+
+							// Profile ID: we only keep the latest ID, and profile will overtake any non-profile based SSS
+							SSSData.bIsValid 			= true;
+							SSSData.bIsProfile 			= SSSData.bIsProfile || BSDF_GETHASSSSPROFILE(BSDF);
+							SSSData.ProfileId 			= StrataSubsurfaceProfileIdTo8bits(SLAB_SSSPROFILEID(BSDF));
+							SSSData.ProfileRadiusScale 	= SLAB_SSSPROFILERADIUSSCALE(BSDF);
+
+							// STRATA_TODO This should be accumulated according to the BSDF weight when we run this in the base pass.
+							// We cannot do that now because BSDFWeight contains weight*coverage*trhougput all mixed together.
+							SSSData.DMFP = SLAB_SSSDMFP(BSDF);
+						}
+
 						if (BSDF_GETHASSSSPROFILE(BSDF) && !bIsSimpleVolume)
 						{
 							// Simple volume overrides SSS profile because it takes over for when not at the bottom of the BSDF layer, or during forward rendering.
@@ -2002,6 +2069,7 @@ void PackStrataOut(
 				break;
 				case STRATA_BSDF_TYPE_HAIR:
 				{
+					AndIsSimpleMaterial(ClassificationData, false);
 					STRATA_STORE_UINT1(PackColorLinearToGamma2AlphaLinear(float4(HAIR_BASECOLOR(BSDF), HAIR_ROUGHNESS(BSDF))));
 					STRATA_STORE_UINT1(PackRGBA8(float4(HAIR_SCATTER(BSDF), HAIR_SPECULAR(BSDF), HAIR_BACKLIT(BSDF), HAIR_COMPLEXTRANSMITTANCE(BSDF))));
 					// 8 bytes
@@ -2009,6 +2077,7 @@ void PackStrataOut(
 				break;
 				case STRATA_BSDF_TYPE_SINGLELAYERWATER:
 				{
+					AndIsSimpleMaterial(ClassificationData, false);
 					STRATA_STORE_UINT1(PackColorLinearToGamma2AlphaLinear(float4(SLW_BASECOLOR(BSDF), 0.0f)));
 					STRATA_STORE_UINT1(PackRGBA8(float4(SLW_METALLIC(BSDF), SLW_SPECULAR(BSDF), SLW_ROUGHNESS(BSDF), SLW_TOPMATERIALOPACITY(BSDF))));
 				}
diff --git a/Engine/Source/Runtime/Renderer/Private/Strata/Strata.cpp b/Engine/Source/Runtime/Renderer/Private/Strata/Strata.cpp
--- a/Engine/Source/Runtime/Renderer/Private/Strata/Strata.cpp
+++ b/Engine/Source/Runtime/Renderer/Private/Strata/Strata.cpp
@@ -77,6 +77,10 @@ void FStrataSceneData::Reset()
 	TopLayerNormalTexture = nullptr;
 	SSSTexture = nullptr;
 
+	ClassificationTextureUAV = nullptr;
+	TopLayerNormalTextureUAV = nullptr;
+	SSSTextureUAV = nullptr;
+
 	MaterialLobesBuffer = nullptr;
 	MaterialLobesBufferUAV = nullptr;
 	MaterialLobesBufferSRV = nullptr;
@@ -160,7 +164,8 @@ void InitialiseStrataFrameSceneData(FSceneRenderer& SceneRenderer, FRDGBuilder&
 
 		// Classification texture
 		{
-			StrataSceneData.ClassificationTexture = GraphBuilder.CreateTexture(FRDGTextureDesc::Create2D(SceneTextureExtent, PF_R16_UINT, FClearValueBinding::Black, TexCreate_ShaderResource | TexCreate_RenderTargetable), TEXT("Strata.ClassificationTexture"));
+			StrataSceneData.ClassificationTexture = GraphBuilder.CreateTexture(FRDGTextureDesc::Create2D(SceneTextureExtent, PF_R16_UINT, FClearValueBinding::Black, TexCreate_ShaderResource | TexCreate_UAV), TEXT("Strata.ClassificationTexture"));
+			StrataSceneData.ClassificationTextureUAV = GraphBuilder.CreateUAV(StrataSceneData.ClassificationTexture);
 		}
 
 		// Tile classification buffers
@@ -186,12 +191,14 @@ void InitialiseStrataFrameSceneData(FSceneRenderer& SceneRenderer, FRDGBuilder&
 
 		// Top layer texture
 		{
-			StrataSceneData.TopLayerNormalTexture = GraphBuilder.CreateTexture(FRDGTextureDesc::Create2D(SceneTextureExtent, PF_R32_UINT, FClearValueBinding::Black, TexCreate_ShaderResource | TexCreate_RenderTargetable), TEXT("Strata.TopLayerNormalTexture"));
+			StrataSceneData.TopLayerNormalTexture = GraphBuilder.CreateTexture(FRDGTextureDesc::Create2D(SceneTextureExtent, PF_R32_UINT, FClearValueBinding::Black, TexCreate_ShaderResource | TexCreate_UAV), TEXT("Strata.TopLayerNormalTexture"));
+			StrataSceneData.TopLayerNormalTextureUAV = GraphBuilder.CreateUAV(StrataSceneData.TopLayerNormalTexture);
 		}
 
 		// SSS texture
 		{
-			StrataSceneData.SSSTexture = GraphBuilder.CreateTexture(FRDGTextureDesc::Create2D(SceneTextureExtent, PF_R32G32_UINT, FClearValueBinding::Black, TexCreate_ShaderResource | TexCreate_RenderTargetable), TEXT("Strata.SSSTexture"));
+			StrataSceneData.SSSTexture = GraphBuilder.CreateTexture(FRDGTextureDesc::Create2D(SceneTextureExtent, PF_R32G32_UINT, FClearValueBinding::Black, TexCreate_ShaderResource | TexCreate_UAV), TEXT("Strata.SSSTexture"));
+			StrataSceneData.SSSTextureUAV = GraphBuilder.CreateUAV(StrataSceneData.SSSTexture);
 		}
 	}
 	else
@@ -236,12 +243,20 @@ void BindStrataBasePassUniformParameters(FRDGBuilder& GraphBuilder, FStrataScene
 		OutStrataUniformParameters.bRoughDiffuse = StrataSceneData->bRoughDiffuse ? 1u : 0u;
 		OutStrataUniformParameters.MaxBytesPerPixel = StrataSceneData->MaxBytesPerPixel;
 		OutStrataUniformParameters.MaterialLobesBufferUAV = StrataSceneData->MaterialLobesBufferUAV;
+
+		OutStrataUniformParameters.ClassificationTextureUAV = StrataSceneData->ClassificationTextureUAV;
+		OutStrataUniformParameters.TopLayerNormalTextureUAV = StrataSceneData->TopLayerNormalTextureUAV;
+		OutStrataUniformParameters.SSSTextureUAV = StrataSceneData->SSSTextureUAV;
 	}
 	else
 	{
+		const FRDGSystemTextures& SystemTextures = FRDGSystemTextures::Get(GraphBuilder);
 		OutStrataUniformParameters.bRoughDiffuse = 0u;
 		OutStrataUniformParameters.MaxBytesPerPixel = 0;
 		OutStrataUniformParameters.MaterialLobesBufferUAV = GraphBuilder.CreateUAV(GraphBuilder.RegisterExternalBuffer(GWhiteVertexBufferWithRDG->Buffer), PF_R32_UINT);
+		OutStrataUniformParameters.ClassificationTextureUAV = GraphBuilder.CreateUAV(SystemTextures.Black);
+		OutStrataUniformParameters.TopLayerNormalTextureUAV = GraphBuilder.CreateUAV(SystemTextures.Black);
+		OutStrataUniformParameters.SSSTextureUAV = GraphBuilder.CreateUAV(SystemTextures.Black);
 	}
 }
 
@@ -353,51 +368,6 @@ static void AddVisualizeMaterialPasses(FRDGBuilder& GraphBuilder, const FViewInf
 	}
 }
 
-//////////////////////////////////////////////////////////////////////////
-// Material classification pass
-// * Classification texture (shading models, BSDF bits, ...)
-// For future:
-// * SSS: hasSSS, Normal, ProfilID, BaseColor, Opacity,  MFPAlbedo/MFPRadius, Shadingmodel | 64bit?
-// * SSR: depth, roughness, normal, (clear coat amount/roughness), tangent, aniso
-
-// SSS/SSR/Auxilary data (AO/ShadowMask/...)
-class FStrataMaterialClassificationPassPS : public FGlobalShader
-{
-	DECLARE_GLOBAL_SHADER(FStrataMaterialClassificationPassPS);
-	SHADER_USE_PARAMETER_STRUCT(FStrataMaterialClassificationPassPS, FGlobalShader);
-
-	using FPermutationDomain = TShaderPermutationDomain<>;
-
-	BEGIN_SHADER_PARAMETER_STRUCT(FParameters, )
-		SHADER_PARAMETER_STRUCT_REF(FViewUniformShaderParameters, ViewUniformBuffer)
-		SHADER_PARAMETER(uint32, MaxBytesPerPixel)
-		SHADER_PARAMETER_RDG_BUFFER_SRV(ByteAddressBuffer, MaterialLobesBuffer)
-		SHADER_PARAMETER_STRUCT_INCLUDE(FSceneTextureParameters, SceneTextures)
-		SHADER_PARAMETER_STRUCT_INCLUDE(ShaderDrawDebug::FShaderParameters, ShaderDrawParameters)
-		RENDER_TARGET_BINDING_SLOTS()
-	END_SHADER_PARAMETER_STRUCT()
-
-	static FPermutationDomain RemapPermutation(FPermutationDomain PermutationVector)
-	{
-		return PermutationVector;
-	}
-
-	static bool ShouldCompilePermutation(const FGlobalShaderPermutationParameters& Parameters)
-	{
-		return GetMaxSupportedFeatureLevel(Parameters.Platform) >= ERHIFeatureLevel::SM5 && Strata::IsStrataEnabled();
-	}
-
-	static void ModifyCompilationEnvironment(const FGlobalShaderPermutationParameters& Parameters, FShaderCompilerEnvironment& OutEnvironment)
-	{
-		FGlobalShader::ModifyCompilationEnvironment(Parameters, OutEnvironment);
-		OutEnvironment.SetDefine(TEXT("SHADER_CATEGORIZATION"), 1);
-		OutEnvironment.SetRenderTargetOutputFormat(0, PF_R32_UINT);
-		OutEnvironment.SetRenderTargetOutputFormat(1, PF_R32_UINT);
-		OutEnvironment.SetRenderTargetOutputFormat(2, PF_R32G32_UINT);
-	}
-};
-IMPLEMENT_GLOBAL_SHADER(FStrataMaterialClassificationPassPS, "/Engine/Private/Strata/StrataMaterialClassification.usf", "MainPS", SF_Pixel);
-
 ////////////////////////////////////////////////////////////////////////////////////////////////////////////////////////////////////////////////
 
 class FStrataClearMaterialBufferCS : public FGlobalShader
@@ -650,34 +620,7 @@ void AddStrataMaterialClassificationPass(FRDGBuilder& GraphBuilder, const FMinim
 	{
 		const FViewInfo& View = Views[i];
 		
-		// Classification
-		{
-			FStrataMaterialClassificationPassPS::FPermutationDomain PermutationVector;
-			TShaderMapRef<FStrataMaterialClassificationPassPS> PixelShader(View.ShaderMap, PermutationVector);
-			FStrataMaterialClassificationPassPS::FParameters* PassParameters = GraphBuilder.AllocParameters<FStrataMaterialClassificationPassPS::FParameters>();
-			PassParameters->ViewUniformBuffer = View.ViewUniformBuffer;
-			PassParameters->MaxBytesPerPixel = View.StrataSceneData->MaxBytesPerPixel;
-			PassParameters->MaterialLobesBuffer = View.StrataSceneData->MaterialLobesBufferSRV;
-			PassParameters->SceneTextures = GetSceneTextureParameters(GraphBuilder);		
-			PassParameters->RenderTargets[0] = FRenderTargetBinding(View.StrataSceneData->ClassificationTexture, ERenderTargetLoadAction::EClear);
-			PassParameters->RenderTargets[1] = FRenderTargetBinding(View.StrataSceneData->TopLayerNormalTexture, ERenderTargetLoadAction::EClear);
-			PassParameters->RenderTargets[2] = FRenderTargetBinding(View.StrataSceneData->SSSTexture, ERenderTargetLoadAction::EClear);
-
-			if (ShaderDrawDebug::IsEnabled())
-			{
-				ShaderDrawDebug::SetParameters(GraphBuilder, View.ShaderDrawData, PassParameters->ShaderDrawParameters);
-			}
-
-			FPixelShaderUtils::AddFullscreenPass<FStrataMaterialClassificationPassPS>(
-				GraphBuilder, 
-				View.ShaderMap, 
-				RDG_EVENT_NAME("Strata::MaterialClassification"),
-				PixelShader, 
-				PassParameters, 
-				View.ViewRect);
-		}
-
-		// Downsampling
+		// Tile reduction
 		if (IsClassificationEnabled())
 		{
 			TShaderMapRef<FStrataMaterialTileClassificationPassCS> ComputeShader(View.ShaderMap);
diff --git a/Engine/Source/Runtime/Renderer/Private/Strata/Strata.h b/Engine/Source/Runtime/Renderer/Private/Strata/Strata.h
--- a/Engine/Source/Runtime/Renderer/Private/Strata/Strata.h
+++ b/Engine/Source/Runtime/Renderer/Private/Strata/Strata.h
@@ -19,6 +19,9 @@ BEGIN_SHADER_PARAMETER_STRUCT(FStrataBasePassUniformParameters, )
 	SHADER_PARAMETER(uint32, MaxBytesPerPixel)
 	SHADER_PARAMETER(uint32, bRoughDiffuse)
 	SHADER_PARAMETER_RDG_BUFFER_UAV(RWByteAddressBuffer, MaterialLobesBufferUAV)
+	SHADER_PARAMETER_RDG_TEXTURE_UAV(RWTexture2D<uint>, ClassificationTextureUAV)
+	SHADER_PARAMETER_RDG_TEXTURE_UAV(RWTexture2D<uint>, TopLayerNormalTextureUAV)
+	SHADER_PARAMETER_RDG_TEXTURE_UAV(RWTexture2D<uint2>, SSSTextureUAV)
 END_SHADER_PARAMETER_STRUCT()
 
 BEGIN_GLOBAL_SHADER_PARAMETER_STRUCT(FStrataGlobalUniformParameters, )
@@ -60,6 +63,10 @@ struct FStrataSceneData
 	FRDGTextureRef TopLayerNormalTexture;
 	FRDGTextureRef SSSTexture;
 
+	FRDGTextureUAVRef ClassificationTextureUAV;
+	FRDGTextureUAVRef TopLayerNormalTextureUAV;
+	FRDGTextureUAVRef SSSTextureUAV;
+
 	TRDGUniformBufferRef<FStrataGlobalUniformParameters> StrataGlobalUniformParameters{};
 
 	FStrataSceneData()
