Fix invalid dummy UAV creation.

diff --git a/Engine/Source/Runtime/Renderer/Private/Strata/Strata.cpp b/Engine/Source/Runtime/Renderer/Private/Strata/Strata.cpp
--- a/Engine/Source/Runtime/Renderer/Private/Strata/Strata.cpp
+++ b/Engine/Source/Runtime/Renderer/Private/Strata/Strata.cpp
@@ -250,13 +250,15 @@ void BindStrataBasePassUniformParameters(FRDGBuilder& GraphBuilder, FStrataScene
 	}
 	else
 	{
+		FRDGTextureRef DummyWritableTexture = GraphBuilder.CreateTexture(FRDGTextureDesc::Create2D(FIntPoint(1,1), PF_R32_UINT, FClearValueBinding::None, TexCreate_ShaderResource | TexCreate_UAV), TEXT("Strata.DummyWritableTexture"));
+		FRDGTextureUAVRef DummyWritableUAV = GraphBuilder.CreateUAV(DummyWritableTexture);
 		const FRDGSystemTextures& SystemTextures = FRDGSystemTextures::Get(GraphBuilder);
 		OutStrataUniformParameters.bRoughDiffuse = 0u;
 		OutStrataUniformParameters.MaxBytesPerPixel = 0;
 		OutStrataUniformParameters.MaterialLobesBufferUAV = GraphBuilder.CreateUAV(GraphBuilder.RegisterExternalBuffer(GWhiteVertexBufferWithRDG->Buffer), PF_R32_UINT);
-		OutStrataUniformParameters.ClassificationTextureUAV = GraphBuilder.CreateUAV(SystemTextures.Black);
-		OutStrataUniformParameters.TopLayerNormalTextureUAV = GraphBuilder.CreateUAV(SystemTextures.Black);
-		OutStrataUniformParameters.SSSTextureUAV = GraphBuilder.CreateUAV(SystemTextures.Black);
+		OutStrataUniformParameters.ClassificationTextureUAV = DummyWritableUAV;
+		OutStrataUniformParameters.TopLayerNormalTextureUAV = DummyWritableUAV;
+		OutStrataUniformParameters.SSSTextureUAV = DummyWritableUAV;
 	}
 }
 
