Avoid temporary FString allocs for breadcrumb events
  - Reduce alloc count from 15.4M to 13.3M for 24s of gameplay ~85K allocs/s

diff --git a/Engine/Source/Runtime/Renderer/Private/MeshPassProcessor.cpp b/Engine/Source/Runtime/Renderer/Private/MeshPassProcessor.cpp
--- a/Engine/Source/Runtime/Renderer/Private/MeshPassProcessor.cpp
+++ b/Engine/Source/Runtime/Renderer/Private/MeshPassProcessor.cpp
@@ -1209,10 +1209,17 @@ void FMeshDrawCommand::SubmitDraw(
 	FRHIBuffer* IndirectArgsOverrideBuffer,
 	uint32 IndirectArgsOverrideByteOffset)
 {
-#if MESH_DRAW_COMMAND_DEBUG_DATA
-	BREADCRUMB_EVENTF(RHICmdList, MeshDrawCommand, TEXT("%s %s"),
-		*MeshDrawCommand.DebugData.MaterialName,
-		MeshDrawCommand.DebugData.ResourceName.IsValid() ? *MeshDrawCommand.DebugData.ResourceName.ToString() : TEXT(""));
+#if MESH_DRAW_COMMAND_DEBUG_DATA && RHI_WANT_BREADCRUMB_EVENTS
+	if (MeshDrawCommand.DebugData.ResourceName.IsValid())
+	{
+		TCHAR NameBuffer[FName::StringBufferSize];
+		const uint32 NameLen = MeshDrawCommand.DebugData.ResourceName.ToString(NameBuffer);
+		BREADCRUMB_EVENTF(RHICmdList, MeshDrawCommand, TEXT("%s %.*s"), *MeshDrawCommand.DebugData.MaterialName, NameLen, NameBuffer);
+	}
+	else
+	{
+		BREADCRUMB_EVENTF(RHICmdList, MeshDrawCommand, TEXT("%s"), *MeshDrawCommand.DebugData.MaterialName);
+	}
 #endif
 #if WANTS_DRAW_MESH_EVENTS
 	FMeshDrawEvent MeshEvent(MeshDrawCommand, InstanceFactor, RHICmdList);
