Add waveops permutation for strata material classification.

Reduce classifcation ~15/20%

diff --git a/Engine/Shaders/Private/Strata/StrataMaterialClassification.usf b/Engine/Shaders/Private/Strata/StrataMaterialClassification.usf
--- a/Engine/Shaders/Private/Strata/StrataMaterialClassification.usf
+++ b/Engine/Shaders/Private/Strata/StrataMaterialClassification.usf
@@ -185,7 +185,9 @@ RWBuffer<uint> SimpleTileListDataBuffer;
 RWBuffer<uint> ComplexTileIndirectDataBuffer;
 RWBuffer<uint> ComplexTileListDataBuffer;
 
+#if !PERMUTATION_WAVE_OPS
 groupshared uint s_IsComplexIsMaterialTile[GROUP_THREAD_COUNT];
+#endif
 
 [numthreads(TILE_SIZE, TILE_SIZE, 1)]
 void TileMainCS(uint2 DispatchThreadId : SV_DispatchThreadID, uint LinearIndex : SV_GroupIndex, uint3 GroupId : SV_GroupID)
@@ -209,6 +211,30 @@ void TileMainCS(uint2 DispatchThreadId : SV_DispatchThreadID, uint LinearIndex :
 		bContainsComplexMaterial = !IsSimpleMaterial(Data);
 		bContainsStrataMaterial = IsStrataMaterial(Data);
 	}
+
+#if PERMUTATION_WAVE_OPS
+
+	const bool bHasAnyStrata  = WaveActiveAnyTrue(bContainsStrataMaterial);
+	const bool bHasAnyComplex = WaveActiveAnyTrue(bContainsComplexMaterial);
+
+	if (LinearIndex < 1 && bHasAnyStrata)
+	{
+		if (bHasAnyComplex)
+		{
+			uint WriteToIndex;
+			InterlockedAdd(ComplexTileIndirectDataBuffer[1], 1, WriteToIndex);
+			ComplexTileListDataBuffer[WriteToIndex] = EncodeTile(GroupId.xy);
+		}
+		else
+		{
+			uint WriteToIndex;
+			InterlockedAdd(SimpleTileIndirectDataBuffer[1], 1, WriteToIndex);
+			SimpleTileListDataBuffer[WriteToIndex] = EncodeTile(GroupId.xy);
+		}
+	}
+
+#else // PERMUTATION_WAVE_OPS
+
 	s_IsComplexIsMaterialTile[LinearIndex] = (bContainsComplexMaterial ? 1 : 0) | (bContainsStrataMaterial ? 2 : 0);
 
 	GroupMemoryBarrierWithGroupSync();
@@ -258,6 +284,7 @@ void TileMainCS(uint2 DispatchThreadId : SV_DispatchThreadID, uint LinearIndex :
 			SimpleTileListDataBuffer[WriteToIndex] = EncodeTile(GroupId.xy);
 		}
 	}
+#endif // PERMUTATION_WAVE_OPS
 }
 #endif // SHADER_TILE_CATEGORIZATION
 
diff --git a/Engine/Source/Runtime/Renderer/Private/Strata/Strata.cpp b/Engine/Source/Runtime/Renderer/Private/Strata/Strata.cpp
--- a/Engine/Source/Runtime/Renderer/Private/Strata/Strata.cpp
+++ b/Engine/Source/Runtime/Renderer/Private/Strata/Strata.cpp
@@ -405,7 +405,8 @@ class FStrataMaterialTileClassificationPassCS : public FGlobalShader
 	DECLARE_GLOBAL_SHADER(FStrataMaterialTileClassificationPassCS);
 	SHADER_USE_PARAMETER_STRUCT(FStrataMaterialTileClassificationPassCS, FGlobalShader);
 
-	using FPermutationDomain = TShaderPermutationDomain<>;
+	class FWaveOps : SHADER_PERMUTATION_BOOL("PERMUTATION_WAVE_OPS");
+	using FPermutationDomain = TShaderPermutationDomain<FWaveOps>;
 
 	BEGIN_SHADER_PARAMETER_STRUCT(FParameters, )
 		SHADER_PARAMETER_STRUCT_REF(FViewUniformShaderParameters, ViewUniformBuffer)
@@ -419,13 +420,14 @@ class FStrataMaterialTileClassificationPassCS : public FGlobalShader
 		SHADER_PARAMETER_RDG_BUFFER_UAV(RWBuffer, ComplexTileListDataBuffer)
 	END_SHADER_PARAMETER_STRUCT()
 
-	static FPermutationDomain RemapPermutation(FPermutationDomain PermutationVector)
-	{
-		return PermutationVector;
-	}
-
 	static bool ShouldCompilePermutation(const FGlobalShaderPermutationParameters& Parameters)
 	{
+		const bool bUseWaveIntrinsics = FDataDrivenShaderPlatformInfo::GetSupportsWaveOperations(Parameters.Platform);
+		FPermutationDomain PermutationVector(Parameters.PermutationId);
+		if (PermutationVector.Get<FWaveOps>() && !bUseWaveIntrinsics)
+		{
+			return false;
+		}
 		return GetMaxSupportedFeatureLevel(Parameters.Platform) >= ERHIFeatureLevel::SM5 && Strata::IsStrataEnabled();
 	}
 
@@ -433,6 +435,12 @@ class FStrataMaterialTileClassificationPassCS : public FGlobalShader
 	{
 		FGlobalShader::ModifyCompilationEnvironment(Parameters, OutEnvironment);
 		OutEnvironment.SetDefine(TEXT("SHADER_TILE_CATEGORIZATION"), 1);
+
+		FPermutationDomain PermutationVector(Parameters.PermutationId);
+		if (PermutationVector.Get<FWaveOps>())
+		{
+			OutEnvironment.CompilerFlags.Add(CFLAG_WaveOperations);
+		}
 	}
 };
 IMPLEMENT_GLOBAL_SHADER(FStrataMaterialTileClassificationPassCS, "/Engine/Private/Strata/StrataMaterialClassification.usf", "TileMainCS", SF_Compute);
@@ -625,7 +633,15 @@ void AddStrataMaterialClassificationPass(FRDGBuilder& GraphBuilder, const FMinim
 		// Tile reduction
 		if (IsClassificationEnabled())
 		{
-			TShaderMapRef<FStrataMaterialTileClassificationPassCS> ComputeShader(View.ShaderMap);
+			bool bWaveOps = FDataDrivenShaderPlatformInfo::GetSupportsWaveOperations(View.GetShaderPlatform());
+		#if PLATFORM_WINDOWS
+			// Tile reduction requires 64-wide wave
+			bWaveOps = bWaveOps && !IsRHIDeviceNVIDIA();
+		#endif
+
+			FStrataMaterialTileClassificationPassCS::FPermutationDomain PermutationVector;
+			PermutationVector.Set< FStrataMaterialTileClassificationPassCS::FWaveOps >(bWaveOps);
+			TShaderMapRef<FStrataMaterialTileClassificationPassCS> ComputeShader(View.ShaderMap, PermutationVector);
 			FStrataMaterialTileClassificationPassCS::FParameters* PassParameters = GraphBuilder.AllocParameters<FStrataMaterialTileClassificationPassCS::FParameters>();
 			PassParameters->TileSize = GetStrataBufferTileSize();	// STRATA_TODO not sure we want to tie the buffer tile optimisation for cache and the Categotisation tile size?
 			PassParameters->bRectPrimitive = GRHISupportsRectTopology ? 1 : 0;
@@ -636,11 +652,10 @@ void AddStrataMaterialClassificationPass(FRDGBuilder& GraphBuilder, const FMinim
 			PassParameters->ComplexTileListDataBuffer = View.StrataSceneData->ClassificationTileListBufferUAV[EStrataTileMaterialType::EComplex];
 			PassParameters->ComplexTileIndirectDataBuffer = View.StrataSceneData->ClassificationTileIndirectBufferUAV[EStrataTileMaterialType::EComplex];
 
-			// Add 64 threads permutation
 			const uint32 GroupSize = 8;
 			FComputeShaderUtils::AddPass(
 				GraphBuilder,
-				RDG_EVENT_NAME("Strata::MaterialTileClassification"),
+				RDG_EVENT_NAME("Strata::MaterialTileClassification(%s)", bWaveOps ? TEXT("Wave") : TEXT("SharedMemory")),
 				ComputeShader,
 				PassParameters,
 				FComputeShaderUtils::GetGroupCount(View.StrataSceneData->ClassificationTexture->Desc.Extent, GroupSize));
