Ensure Strata tile reduction pass uses wave-ops permuation only if the platform supports it (in addition to be supported by the shader platforms).

diff --git a/Engine/Source/Runtime/Renderer/Private/Strata/Strata.cpp b/Engine/Source/Runtime/Renderer/Private/Strata/Strata.cpp
--- a/Engine/Source/Runtime/Renderer/Private/Strata/Strata.cpp
+++ b/Engine/Source/Runtime/Renderer/Private/Strata/Strata.cpp
@@ -633,7 +633,7 @@ void AddStrataMaterialClassificationPass(FRDGBuilder& GraphBuilder, const FMinim
 		// Tile reduction
 		if (IsClassificationEnabled())
 		{
-			bool bWaveOps = FDataDrivenShaderPlatformInfo::GetSupportsWaveOperations(View.GetShaderPlatform());
+			bool bWaveOps = GRHISupportsWaveOperations && FDataDrivenShaderPlatformInfo::GetSupportsWaveOperations(View.GetShaderPlatform());
 		#if PLATFORM_WINDOWS
 			// Tile reduction requires 64-wide wave
 			bWaveOps = bWaveOps && !IsRHIDeviceNVIDIA();
