Fix a dark screen issue in splitscreen mode. It happens when view 0 is under water and view 1 is above water. The cause is that the render target isn't cleared at the beginning of the  above water translucency pass because NumProcessedViews has been incremented for view 0 which is an under water view.

#localization none
#tests ran a replay in both normal and splitscreen modes on a platform

diff --git a/Engine/Source/Runtime/Renderer/Private/TranslucentRendering.cpp b/Engine/Source/Runtime/Renderer/Private/TranslucentRendering.cpp
--- a/Engine/Source/Runtime/Renderer/Private/TranslucentRendering.cpp
+++ b/Engine/Source/Runtime/Renderer/Private/TranslucentRendering.cpp
@@ -1189,7 +1189,7 @@ void FDeferredShadingSceneRenderer::RenderTranslucencyInner(
 		// Create resources shared by each view (each view data is tiled into each of the render target resources)
 		FSeparateTranslucencyTextures LocalSeparateTranslucencyTextures(SeparateTranslucencyDimensions);
 
-		for (int32 ViewIndex = 0, NumProcessedViews = 0; ViewIndex < Views.Num(); ++ViewIndex, ++NumProcessedViews)
+		for (int32 ViewIndex = 0, NumProcessedViews = 0; ViewIndex < Views.Num(); ++ViewIndex)
 		{
 			FViewInfo& View = Views[ViewIndex];
 			const ETranslucencyView TranslucencyView = GetTranslucencyView(View);
@@ -1291,6 +1291,7 @@ void FDeferredShadingSceneRenderer::RenderTranslucencyInner(
 			}
 
 			AddEndSeparateTranslucencyTimerPass(GraphBuilder, View, TranslucencyPass);
+			++NumProcessedViews;
 		}
 	}
 	else
