Disable CSM caching if virtual shadowmap is enabled.

-fyi Ola.Olsson
#p4v-preflight-copy 18257889

diff --git a/Engine/Source/Runtime/Renderer/Private/ShadowSetup.cpp b/Engine/Source/Runtime/Renderer/Private/ShadowSetup.cpp
--- a/Engine/Source/Runtime/Renderer/Private/ShadowSetup.cpp
+++ b/Engine/Source/Runtime/Renderer/Private/ShadowSetup.cpp
@@ -3224,7 +3224,7 @@ void ComputeViewDependentWholeSceneShadowCacheModes(
 	int32& OutNumShadowMaps,
 	EShadowDepthCacheMode* OutCacheModes)
 {
-	if (CVarCSMCaching.GetValueOnAnyThread() == 1)
+	if (CVarCSMCaching.GetValueOnAnyThread() == 1 && !UseVirtualShadowMaps(Scene->GetShaderPlatform(), Scene->GetFeatureLevel()))
 	{
 		TArray<FCachedShadowMapData>* CachedShadowMapDatas = Scene->GetCachedShadowMapDatas(LightSceneInfo->Id);
 
