Strata - Also clear the classification texture to not have complex material at edge with sky/empty pixel or material tile leaking all over the screen when rotating the camera.

diff --git a/Engine/Shaders/Private/Strata/StrataMaterialClassification.usf b/Engine/Shaders/Private/Strata/StrataMaterialClassification.usf
--- a/Engine/Shaders/Private/Strata/StrataMaterialClassification.usf
+++ b/Engine/Shaders/Private/Strata/StrataMaterialClassification.usf
@@ -347,6 +347,7 @@ void StencilMainPS(
 #if SHADER_CLEAR_MATERIAL_BUFFER
 
 RWByteAddressBuffer MaterialLobesBufferUAV;
+RWTexture2D<uint> ClassificationTextureUAV;
 uint MaxBytesPerPixel;
 int2 TiledViewBufferResolution;
 
@@ -367,6 +368,10 @@ void ClearMaterialBufferMainCS(uint2 DispatchThreadId : SV_DispatchThreadID)
 
 	FStrataAddressing StrataAddressing = GetStrataPixelDataByteOffset(DispatchThreadId.xy, TiledViewBufferResolution, MaxBytesPerPixel);
 	MaterialLobesBufferUAV.Store(StrataAddressing.CurrentAddr, PackedStrataHeader);
+
+	FStrataClassification ClassificationData = (FStrataClassification)0;
+	SetIsSimpleMaterial(ClassificationData, true);
+	ClassificationTextureUAV[DispatchThreadId.xy] = StrataPackClassificationData(ClassificationData);
 }
 
 #endif // SHADER_CLEAR_MATERIAL_BUFFER
diff --git a/Engine/Source/Runtime/Renderer/Private/Strata/Strata.cpp b/Engine/Source/Runtime/Renderer/Private/Strata/Strata.cpp
--- a/Engine/Source/Runtime/Renderer/Private/Strata/Strata.cpp
+++ b/Engine/Source/Runtime/Renderer/Private/Strata/Strata.cpp
@@ -114,7 +114,7 @@ namespace Strata
 {
 
 // Forward declaration
-static void AddStrataClearMaterialBufferPass(FRDGBuilder& GraphBuilder, FRDGBufferUAVRef MaterialLobesBufferUAV, uint32 MaxBytesPerPixel, FIntPoint TiledViewBufferResolution);
+static void AddStrataClearMaterialBufferPass(FRDGBuilder& GraphBuilder, FRDGBufferUAVRef MaterialLobesBufferUAV, FRDGTextureUAVRef ClassificationTextureUAV, uint32 MaxBytesPerPixel, FIntPoint TiledViewBufferResolution);
 
 bool IsStrataEnabled()
 {
@@ -224,7 +224,7 @@ void InitialiseStrataFrameSceneData(FSceneRenderer& SceneRenderer, FRDGBuilder&
 
 	if (IsStrataEnabled())
 	{
-		AddStrataClearMaterialBufferPass(GraphBuilder, StrataSceneData.MaterialLobesBufferUAV, StrataSceneData.MaxBytesPerPixel, MaterialBufferSizeXY);
+		AddStrataClearMaterialBufferPass(GraphBuilder, StrataSceneData.MaterialLobesBufferUAV, StrataSceneData.ClassificationTextureUAV, StrataSceneData.MaxBytesPerPixel, MaterialBufferSizeXY);
 	}
 
 	// Create the readable uniform buffers for each views once for all (it is view independent and all the views should be tiled into the render target textures & material buffer)
@@ -381,6 +381,7 @@ class FStrataClearMaterialBufferCS : public FGlobalShader
 
 	BEGIN_SHADER_PARAMETER_STRUCT(FParameters, )
 		SHADER_PARAMETER_RDG_BUFFER_UAV(RWByteAddressBuffer, MaterialLobesBufferUAV)
+		SHADER_PARAMETER_RDG_TEXTURE_UAV(RWTexture2D<uint>, ClassificationTextureUAV)
 		SHADER_PARAMETER(uint32, MaxBytesPerPixel)
 		SHADER_PARAMETER(FIntPoint, TiledViewBufferResolution)
 	END_SHADER_PARAMETER_STRUCT()
@@ -663,11 +664,12 @@ void AddStrataMaterialClassificationPass(FRDGBuilder& GraphBuilder, const FMinim
 	}
 }
 
-static void AddStrataClearMaterialBufferPass(FRDGBuilder& GraphBuilder, FRDGBufferUAVRef MaterialLobesBufferUAV, uint32 MaxBytesPerPixel, FIntPoint TiledViewBufferResolution)
+static void AddStrataClearMaterialBufferPass(FRDGBuilder& GraphBuilder, FRDGBufferUAVRef MaterialLobesBufferUAV, FRDGTextureUAVRef ClassificationTextureUAV, uint32 MaxBytesPerPixel, FIntPoint TiledViewBufferResolution)
 {
 	TShaderMapRef<FStrataClearMaterialBufferCS> ComputeShader(GetGlobalShaderMap(GMaxRHIFeatureLevel));
 	FStrataClearMaterialBufferCS::FParameters* PassParameters = GraphBuilder.AllocParameters<FStrataClearMaterialBufferCS::FParameters>();
 	PassParameters->MaterialLobesBufferUAV = MaterialLobesBufferUAV;
+	PassParameters->ClassificationTextureUAV = ClassificationTextureUAV;
 	PassParameters->MaxBytesPerPixel = MaxBytesPerPixel;
 	PassParameters->TiledViewBufferResolution = TiledViewBufferResolution;
 
