* Add new version of barycentrics ray tracing debug view that uses InlineRayTracing (enabled by default).
* Add r.Raytracing.AllowInline/AllowPipeline to control which effects are used.

diff --git a/Engine/Shaders/Private/RayTracing/RayTracingBarycentrics.usf b/Engine/Shaders/Private/RayTracing/RayTracingBarycentrics.usf
--- a/Engine/Shaders/Private/RayTracing/RayTracingBarycentrics.usf
+++ b/Engine/Shaders/Private/RayTracing/RayTracingBarycentrics.usf
@@ -3,22 +3,37 @@
 #include "../Common.ush"
 #include "RayTracingCommon.ush"
 
+#if COMPUTESHADER
+#include "RayTracingDeferredMaterials.ush"
+#include "TraceRayInline.ush"
+#endif // COMPUTESHADER
+
 RWTexture2D<float4> Output;
 RaytracingAccelerationStructure TLAS;
 
-RAY_TRACING_ENTRY_RAYGEN(RayTracingBarycentricsMainRGS)
+void RayTracingBarycentricsCommon(uint2 DispatchThreadIndex)
 {
-	uint2 PixelCoord = DispatchRaysIndex().xy + View.ViewRectMin;
+	uint2 PixelCoord = DispatchThreadIndex.xy + View.ViewRectMin;
 
 	float2 RenderTargetUV = (float2(PixelCoord) + .5f) * View.BufferSizeAndInvSize.zw;
 
 	RayDesc Ray = CreatePrimaryRay(RenderTargetUV);
 	uint RayFlags = 0
-	              | RAY_FLAG_CULL_BACK_FACING_TRIANGLES      // use back face culling
-	              | RAY_FLAG_FORCE_OPAQUE;                   // don't run anyhit shader
+		| RAY_FLAG_CULL_BACK_FACING_TRIANGLES      // use back face culling
+		| RAY_FLAG_FORCE_OPAQUE;                   // don't run anyhit shader
 
 	FDefaultPayload Payload = (FDefaultPayload)0;
+#if COMPUTESHADER
+	FTraceRayInlineResult TraceRayResult = TraceRayInline(
+		TLAS,   // AccelerationStructure
+		RayFlags,
+		RAY_TRACING_MASK_ALL,             // InstanceInclusionMask
+		Ray    // RayDesc
+	);
 
+	Payload.HitT = TraceRayResult.HitT;
+	Payload.Barycentrics = TraceRayResult.Barycentrics;
+#else
 	TraceRay(
 		TLAS,   // AccelerationStructure
 		RayFlags,
@@ -29,7 +44,7 @@ RAY_TRACING_ENTRY_RAYGEN(RayTracingBarycentricsMainRGS)
 		Ray,    // RayDesc
 		Payload // Payload
 	);
-
+#endif // COMPUTESHADER
 	float4 Result = float4(0, 0, 0, 1);
 	if (Payload.IsHit())
 	{
@@ -38,6 +53,21 @@ RAY_TRACING_ENTRY_RAYGEN(RayTracingBarycentricsMainRGS)
 	Output[PixelCoord] = Result;
 }
 
+#if COMPUTESHADER
+
+[numthreads(INLINE_RAY_TRACING_THREAD_GROUP_SIZE_X, INLINE_RAY_TRACING_THREAD_GROUP_SIZE_Y, 1)]
+void RayTracingBarycentricsMainCS(uint2 DispatchThreadIndex : SV_DispatchThreadID)
+{
+	RayTracingBarycentricsCommon(DispatchThreadIndex);
+}
+
+#else
+
+RAY_TRACING_ENTRY_RAYGEN(RayTracingBarycentricsMainRGS)
+{
+	RayTracingBarycentricsCommon(DispatchRaysIndex().xy);
+}
+
 RAY_TRACING_ENTRY_CLOSEST_HIT(RayTracingBarycentricsMainCHS,
 	FDefaultPayload, Payload,
 	FDefaultAttributes, Attributes)
@@ -45,3 +75,5 @@ RAY_TRACING_ENTRY_CLOSEST_HIT(RayTracingBarycentricsMainCHS,
 	Payload.Barycentrics = Attributes.Barycentrics;
 	Payload.HitT = RayTCurrent();
 }
+
+#endif
\ No newline at end of file
diff --git a/Engine/Shaders/Private/RayTracing/RayTracingCommon.ush b/Engine/Shaders/Private/RayTracing/RayTracingCommon.ush
--- a/Engine/Shaders/Private/RayTracing/RayTracingCommon.ush
+++ b/Engine/Shaders/Private/RayTracing/RayTracingCommon.ush
@@ -360,7 +360,7 @@ void ApplyCameraRelativeDepthBias(inout float3 RayOrigin, float3 RayDirection, u
 	ApplyPositionBias(RayOrigin, RayDirection, WorldNormal, RayBias);
 }
 
-#if !COMPUTESHADER
+#if !COMPUTESHADER || PLATFORM_SUPPORTS_INLINE_RAY_TRACING
 
 RayDesc CreatePrimaryRay(float2 UV)
 {
@@ -381,6 +381,10 @@ RayDesc CreatePrimaryRay(float2 UV)
 	return Ray;
 }
 
+#endif // !COMPUTESHADER || PLATFORM_SUPPORTS_INLINE_RAY_TRACING
+
+#if !COMPUTESHADER
+
 struct FPackedMaterialClosestHitPayload : FMinimalPayload
 {
 	// float FMinimalPayload::HitT                       // 4  bytes
diff --git a/Engine/Source/Runtime/Renderer/Private/DeferredShadingRenderer.cpp b/Engine/Source/Runtime/Renderer/Private/DeferredShadingRenderer.cpp
--- a/Engine/Source/Runtime/Renderer/Private/DeferredShadingRenderer.cpp
+++ b/Engine/Source/Runtime/Renderer/Private/DeferredShadingRenderer.cpp
@@ -124,6 +124,20 @@ static TAutoConsoleVariable<int32> CVarForceAllRayTracingEffects(
 	TEXT(" 1: All ray tracing effects enabled"),
 	ECVF_RenderThreadSafe);
 
+static int32 GRayTracingAllowInline = 1;
+static TAutoConsoleVariable<int32> CVarRayTracingAllowInline(
+	TEXT("r.RayTracing.AllowInline"),
+	GRayTracingAllowInline,
+	TEXT("Allow use of Inline Ray Tracing if supported (default=1)."),	
+	ECVF_RenderThreadSafe);
+
+static int32 GRayTracingAllowPipeline = 1;
+static TAutoConsoleVariable<int32> CVarRayTracingAllowPipeline(
+	TEXT("r.RayTracing.AllowPipeline"),
+	GRayTracingAllowPipeline,
+	TEXT("Allow use of Ray Tracing pipelines if supported (default=1)."),
+	ECVF_RenderThreadSafe);
+
 static int32 GRayTracingSceneCaptures = -1;
 static FAutoConsoleVariableRef CVarRayTracingSceneCaptures(
 	TEXT("r.RayTracing.SceneCaptures"),
@@ -3234,8 +3248,13 @@ bool ShouldRenderRayTracingEffect(bool bEffectEnabled, ERayTracingPipelineCompat
 		return false;
 	}
 
-	const bool bAllowPipeline = GRHISupportsRayTracingShaders && EnumHasAnyFlags(CompatibilityFlags, ERayTracingPipelineCompatibilityFlags::FullPipeline);
-	const bool bAllowInline = GRHISupportsInlineRayTracing && EnumHasAnyFlags(CompatibilityFlags, ERayTracingPipelineCompatibilityFlags::Inline);
+	const bool bAllowPipeline = GRHISupportsRayTracingShaders && 
+								CVarRayTracingAllowPipeline.GetValueOnRenderThread() &&
+								EnumHasAnyFlags(CompatibilityFlags, ERayTracingPipelineCompatibilityFlags::FullPipeline);
+
+	const bool bAllowInline = GRHISupportsInlineRayTracing && 
+							  CVarRayTracingAllowInline.GetValueOnRenderThread() &&
+							  EnumHasAnyFlags(CompatibilityFlags, ERayTracingPipelineCompatibilityFlags::Inline);
 
 	// Disable the effect if current machine does not support the full ray tracing pipeline and the effect can't fall back to inline mode or vice versa.
 	if (!bAllowPipeline && !bAllowInline)
diff --git a/Engine/Source/Runtime/Renderer/Private/RayTracing/RayTracingBarycentrics.cpp b/Engine/Source/Runtime/Renderer/Private/RayTracing/RayTracingBarycentrics.cpp
--- a/Engine/Source/Runtime/Renderer/Private/RayTracing/RayTracingBarycentrics.cpp
+++ b/Engine/Source/Runtime/Renderer/Private/RayTracing/RayTracingBarycentrics.cpp
@@ -9,6 +9,7 @@
 #include "SceneRenderTargets.h"
 #include "RenderGraphBuilder.h"
 #include "RHI/Public/PipelineStateCache.h"
+#include "RayTracing/RaytracingOptions.h"
 
 class FRayTracingBarycentricsRGS : public FGlobalShader
 {
@@ -52,10 +53,69 @@ public:
 };
 IMPLEMENT_SHADER_TYPE(, FRayTracingBarycentricsCHS, TEXT("/Engine/Private/RayTracing/RayTracingBarycentrics.usf"), TEXT("RayTracingBarycentricsMainCHS"), SF_RayHitGroup);
 
-void FDeferredShadingSceneRenderer::RenderRayTracingBarycentrics(FRDGBuilder& GraphBuilder, const FViewInfo& View, FRDGTextureRef SceneColor)
+class FRayTracingBarycentricsCS : public FGlobalShader
 {
-	FGlobalShaderMap* ShaderMap = GetGlobalShaderMap(FeatureLevel);
+	DECLARE_GLOBAL_SHADER(FRayTracingBarycentricsCS)
+	SHADER_USE_PARAMETER_STRUCT(FRayTracingBarycentricsCS, FGlobalShader)
+
+	BEGIN_SHADER_PARAMETER_STRUCT(FParameters, )
+		SHADER_PARAMETER_SRV(RaytracingAccelerationStructure, TLAS)
+		SHADER_PARAMETER_RDG_TEXTURE_UAV(RWTexture2D<float4>, Output)
+		SHADER_PARAMETER_STRUCT_REF(FViewUniformShaderParameters, ViewUniformBuffer)
+	END_SHADER_PARAMETER_STRUCT()
+
+	static void ModifyCompilationEnvironment(const FGlobalShaderPermutationParameters& Parameters, FShaderCompilerEnvironment& OutEnvironment)
+	{
+		OutEnvironment.CompilerFlags.Add(CFLAG_Wave32);
+		OutEnvironment.CompilerFlags.Add(CFLAG_InlineRayTracing);
+
+		OutEnvironment.SetDefine(TEXT("INLINE_RAY_TRACING_THREAD_GROUP_SIZE_X"), ThreadGroupSizeX);
+		OutEnvironment.SetDefine(TEXT("INLINE_RAY_TRACING_THREAD_GROUP_SIZE_Y"), ThreadGroupSizeY);
+	}
+
+	static bool ShouldCompilePermutation(const FGlobalShaderPermutationParameters& Parameters)
+	{
+		return RHISupportsRayTracing(Parameters.Platform) && RHISupportsInlineRayTracing(Parameters.Platform);
+	}
+
+	// Current inline ray tracing implementation requires 1:1 mapping between thread groups and waves and only supports wave32 mode.
+	static constexpr uint32 ThreadGroupSizeX = 8;
+	static constexpr uint32 ThreadGroupSizeY = 4;
+};
+IMPLEMENT_GLOBAL_SHADER(FRayTracingBarycentricsCS, "/Engine/Private/RayTracing/RayTracingBarycentrics.usf", "RayTracingBarycentricsMainCS", SF_Compute);
+
+void RenderRayTracingBarycentricsCS(FRDGBuilder& GraphBuilder, const FViewInfo& View, FRDGTextureRef SceneColor, FGlobalShaderMap* ShaderMap)
+{
+	FRayTracingBarycentricsCS::FParameters* PassParameters = GraphBuilder.AllocParameters<FRayTracingBarycentricsCS::FParameters>();
+
+	PassParameters->TLAS = View.GetRayTracingSceneViewChecked();
+	PassParameters->ViewUniformBuffer = View.ViewUniformBuffer;
+	PassParameters->Output = GraphBuilder.CreateUAV(SceneColor);
+
+	FIntRect ViewRect = View.ViewRect;	
+
+	auto ComputeShader = ShaderMap->GetShader<FRayTracingBarycentricsCS>();
+	ClearUnusedGraphResources(ComputeShader, PassParameters);
+	GraphBuilder.AddPass(
+		RDG_EVENT_NAME("Barycentrics"),
+		PassParameters,
+		ERDGPassFlags::Compute,
+		[PassParameters, ComputeShader, &View, ViewRect](FRHIRayTracingCommandList& RHICmdList)
+		{
+			FRHIComputeShader* ShaderRHI = ComputeShader.GetComputeShader();
+			RHICmdList.SetComputeShader(ShaderRHI);
+			SetShaderParameters(RHICmdList, ComputeShader, ShaderRHI, *PassParameters);
+
+			const FIntPoint GroupSize(FRayTracingBarycentricsCS::ThreadGroupSizeX, FRayTracingBarycentricsCS::ThreadGroupSizeY);
+			const FIntVector GroupCount = FComputeShaderUtils::GetGroupCount(ViewRect.Size(), GroupSize);
+			DispatchComputeShader(RHICmdList, ComputeShader.GetShader(), GroupCount.X, GroupCount.Y, 1);
+
+			UnsetShaderUAVs(RHICmdList, ComputeShader, ShaderRHI);
+		});
+}
 
+void RenderRayTracingBarycentricsRGS(FRDGBuilder& GraphBuilder, const FViewInfo& View, FRDGTextureRef SceneColor, FGlobalShaderMap* ShaderMap)
+{
 	auto RayGenShader = ShaderMap->GetShader<FRayTracingBarycentricsRGS>();
 	auto ClosestHitShader = ShaderMap->GetShader<FRayTracingBarycentricsCHS>();
 
@@ -82,7 +142,7 @@ void FDeferredShadingSceneRenderer::RenderRayTracingBarycentrics(FRDGBuilder& Gr
 		RDG_EVENT_NAME("Barycentrics"),
 		RayGenParameters,
 		ERDGPassFlags::Compute,
-		[this, RayGenParameters, RayGenShader, &View, Pipeline, ViewRect](FRHIRayTracingCommandList& RHICmdList)
+		[RayGenParameters, RayGenShader, &View, Pipeline, ViewRect](FRHIRayTracingCommandList& RHICmdList)
 	{
 		FRayTracingShaderBindingsWriter GlobalResources;
 		SetShaderParameters(GlobalResources, RayGenShader, *RayGenParameters);
@@ -91,4 +151,21 @@ void FDeferredShadingSceneRenderer::RenderRayTracingBarycentrics(FRDGBuilder& Gr
 		RHICmdList.RayTraceDispatch(Pipeline, RayGenShader.GetRayTracingShader(), View.GetRayTracingSceneChecked(), GlobalResources, ViewRect.Size().X, ViewRect.Size().Y);
 	});
 }
+
+void FDeferredShadingSceneRenderer::RenderRayTracingBarycentrics(FRDGBuilder& GraphBuilder, const FViewInfo& View, FRDGTextureRef SceneColor)
+{
+	FGlobalShaderMap* ShaderMap = GetGlobalShaderMap(FeatureLevel);
+
+	const bool bRayTracingInline = ShouldRenderRayTracingEffect(ERayTracingPipelineCompatibilityFlags::Inline);
+	const bool bRayTracingPipeline = ShouldRenderRayTracingEffect(ERayTracingPipelineCompatibilityFlags::FullPipeline);
+
+	if (bRayTracingInline)
+	{
+		RenderRayTracingBarycentricsCS(GraphBuilder, View, SceneColor, ShaderMap);
+	}
+	else if (bRayTracingPipeline)
+	{
+		RenderRayTracingBarycentricsRGS(GraphBuilder, View, SceneColor, ShaderMap);
+	}
+}
 #endif
\ No newline at end of file
diff --git a/Engine/Source/Runtime/Renderer/Private/RayTracing/RayTracingDebug.cpp b/Engine/Source/Runtime/Renderer/Private/RayTracing/RayTracingDebug.cpp
--- a/Engine/Source/Runtime/Renderer/Private/RayTracing/RayTracingDebug.cpp
+++ b/Engine/Source/Runtime/Renderer/Private/RayTracing/RayTracingDebug.cpp
@@ -91,7 +91,8 @@ static bool RequiresRayTracingDebugCHS(uint32 DebugVisualizationMode)
 void FDeferredShadingSceneRenderer::PrepareRayTracingDebug(const FSceneViewFamily& ViewFamily, TArray<FRHIRayTracingShader*>& OutRayGenShaders)
 {
 	// Declare all RayGen shaders that require material closest hit shaders to be bound
-	if (ViewFamily.EngineShowFlags.RayTracingDebug)
+	bool bEnabled = ViewFamily.EngineShowFlags.RayTracingDebug && ShouldRenderRayTracingEffect(ERayTracingPipelineCompatibilityFlags::FullPipeline);
+	if (bEnabled)
 	{
 		auto RayGenShader = GetGlobalShaderMap(ViewFamily.GetShaderPlatform())->GetShader<FRayTracingDebugRGS>();
 		OutRayGenShaders.Add(RayGenShader.GetRayTracingShader());
@@ -148,12 +149,18 @@ void FDeferredShadingSceneRenderer::RenderRayTracingDebug(FRDGBuilder& GraphBuil
 		DebugVisualizationMode = RAY_TRACING_DEBUG_VIZ_BASE_COLOR;
 	}
 
-
 	if (DebugVisualizationMode == RAY_TRACING_DEBUG_VIZ_BARYCENTRICS)
 	{
 		return RenderRayTracingBarycentrics(GraphBuilder, View, SceneColorTexture);
 	}
 
+	// Debug modes other than barycentrics require pipeline support.
+	const bool bRayTracingPipeline = ShouldRenderRayTracingEffect(ERayTracingPipelineCompatibilityFlags::FullPipeline);
+	if (!bRayTracingPipeline)
+	{
+		return;
+	}
+
 	if (DebugVisualizationMode == RAY_TRACING_DEBUG_VIZ_PRIMARY_RAYS) 
 	{
 		FRDGTextureRef OutputColor = nullptr;
diff --git a/Engine/Source/Runtime/Renderer/Private/RayTracing/RayTracingShadows.cpp b/Engine/Source/Runtime/Renderer/Private/RayTracing/RayTracingShadows.cpp
--- a/Engine/Source/Runtime/Renderer/Private/RayTracing/RayTracingShadows.cpp
+++ b/Engine/Source/Runtime/Renderer/Private/RayTracing/RayTracingShadows.cpp
@@ -180,7 +180,7 @@ void FDeferredShadingSceneRenderer::PrepareRayTracingShadows(const FViewInfo& Vi
 	// Declare all RayGen shaders that require material closest hit shaders to be bound
 	// Ray tracing shadows shaders should be properly configured even if r.RayTracing.Shadows is 0 because lights can have raytracing shadows enabled independently of that CVar
 
-	const bool bRayTracingShadows = ShouldRenderRayTracingEffect(true, ERayTracingPipelineCompatibilityFlags::FullPipeline);
+	const bool bRayTracingShadows = ShouldRenderRayTracingEffect(ERayTracingPipelineCompatibilityFlags::FullPipeline);
 
 	if (!bRayTracingShadows)
 	{
diff --git a/Engine/Source/Runtime/Renderer/Private/RayTracing/RaytracingOptions.h b/Engine/Source/Runtime/Renderer/Private/RayTracing/RaytracingOptions.h
--- a/Engine/Source/Runtime/Renderer/Private/RayTracing/RaytracingOptions.h
+++ b/Engine/Source/Runtime/Renderer/Private/RayTracing/RaytracingOptions.h
@@ -149,3 +149,8 @@ FORCEINLINE bool CanUseRayTracingAMDHitToken()
 }
 
 #endif // RHI_RAYTRACING
+
+FORCEINLINE bool ShouldRenderRayTracingEffect(ERayTracingPipelineCompatibilityFlags CompatibilityFlags)
+{
+	return ShouldRenderRayTracingEffect(true, CompatibilityFlags);
+}
\ No newline at end of file
diff --git a/Engine/Source/Runtime/Renderer/Private/RayTracing/RaytracingSkylight.cpp b/Engine/Source/Runtime/Renderer/Private/RayTracing/RaytracingSkylight.cpp
--- a/Engine/Source/Runtime/Renderer/Private/RayTracing/RaytracingSkylight.cpp
+++ b/Engine/Source/Runtime/Renderer/Private/RayTracing/RaytracingSkylight.cpp
@@ -133,7 +133,7 @@ bool ShouldRenderRayTracingSkyLight(const FSkyLightSceneProxy* SkyLightSceneProx
 	bool bRayTracingSkyEnabled = (GRayTracingSkyLight  > 0 && SkyLightSceneProxy->CastRayTracedShadow == ECastRayTracedShadow::UseProjectSetting)
 								||  SkyLightSceneProxy->CastRayTracedShadow == ECastRayTracedShadow::Enabled;
 
-	return bRayTracingSkyEnabled && ShouldRenderRayTracingEffect(true, ERayTracingPipelineCompatibilityFlags::FullPipeline) && (GetSkyLightSamplesPerPixel(SkyLightSceneProxy) > 0);
+	return bRayTracingSkyEnabled && ShouldRenderRayTracingEffect(ERayTracingPipelineCompatibilityFlags::FullPipeline) && (GetSkyLightSamplesPerPixel(SkyLightSceneProxy) > 0);
 }
 
 IMPLEMENT_GLOBAL_SHADER_PARAMETER_STRUCT(FSkyLightData, "SkyLight");
