Fix various issues with ray tracing scene construction when running without ray tracing shader support (inline-only mode).

* Add IsRayTracingEnabledForProject() to check if ray tracing is enabled independently from RT shader support
* Skip creating the RayTracingMaterialPipeline when RT shaders are not supported
* Update FRayTracingMeshProcessor::Process() to work without full RT sahders, as it's still needed to compute the RT instance masks

diff --git a/Engine/Source/Runtime/RenderCore/Public/RenderResource.h b/Engine/Source/Runtime/RenderCore/Public/RenderResource.h
--- a/Engine/Source/Runtime/RenderCore/Public/RenderResource.h
+++ b/Engine/Source/Runtime/RenderCore/Public/RenderResource.h
@@ -627,13 +627,24 @@ public:
 	virtual FString GetFriendlyName() const override { return TEXT("FIndexBuffer"); }
 };
 
+FORCEINLINE bool IsRayTracingEnabledForProject(EShaderPlatform ShaderPlatform)
+{
+	if (RHISupportsRayTracing(ShaderPlatform))
+	{
+		extern RENDERCORE_API uint64 GRayTracingPlaformMask;
+		return !!(GRayTracingPlaformMask & (1ull << ShaderPlatform));
+	}
+	else
+	{
+		return false;
+	}
+}
 
 FORCEINLINE bool ShouldCompileRayTracingShadersForProject(EShaderPlatform ShaderPlatform)
 {
 	if (RHISupportsRayTracingShaders(ShaderPlatform))
 	{
-		extern RENDERCORE_API uint64 GRayTracingPlaformMask;
-		return !!(GRayTracingPlaformMask & (1ull << ShaderPlatform));
+		return IsRayTracingEnabledForProject(ShaderPlatform);
 	}
 	else
 	{
diff --git a/Engine/Source/Runtime/Renderer/Private/DeferredShadingRenderer.cpp b/Engine/Source/Runtime/Renderer/Private/DeferredShadingRenderer.cpp
--- a/Engine/Source/Runtime/Renderer/Private/DeferredShadingRenderer.cpp
+++ b/Engine/Source/Runtime/Renderer/Private/DeferredShadingRenderer.cpp
@@ -1366,50 +1366,60 @@ bool FDeferredShadingSceneRenderer::SetupRayTracingPipelineStates(FRHICommandLis
 		ReferenceView.AddRayTracingMeshBatchTaskList.Empty();
 	}
 
-	// #dxr_todo: UE-72565: refactor ray tracing effects to not be member functions of DeferredShadingRenderer. register each effect at startup and just loop over them automatically to gather all required shaders
-	TArray<FRHIRayTracingShader*> RayGenShaders; // TODO: inline allocator here?
 	const bool bIsPathTracing = ViewFamily.EngineShowFlags.PathTracing;
-	if (bIsPathTracing)
-	{
-		// This view only needs the path tracing raygen shaders as all other
-		// passes should be disabled.
-		PreparePathTracing(ViewFamily, RayGenShaders);
-	}
-	else
+
+	if (GRHISupportsRayTracingShaders)
 	{
-		// Path tracing is disabled, get all other possible raygen shaders
-		PrepareRayTracingDebug(ViewFamily, RayGenShaders);
+		// #dxr_todo: UE-72565: refactor ray tracing effects to not be member functions of DeferredShadingRenderer. 
+		// Should register each effect at startup and just loop over them automatically to gather all required shaders.
 
-		// These other cases do potentially depend on the camera position since they are
-		// driven by FinalPostProcessSettings, which is why we need to merge them across views
-		if (!IsForwardShadingEnabled(ShaderPlatform))
+		TArray<FRHIRayTracingShader*> RayGenShaders;
+
+		// We typically see ~120 raygen shaders, but allow some headroom to avoid reallocation if our estimate is wrong.
+		RayGenShaders.Reserve(256);
+
+		if (bIsPathTracing)
 		{
-			for (const FViewInfo& View : Views)
-			{
-				PrepareRayTracingReflections(View, *Scene, RayGenShaders);
-				PrepareSingleLayerWaterRayTracingReflections(View, *Scene, RayGenShaders);
-				PrepareRayTracingShadows(View, RayGenShaders);
-				PrepareRayTracingAmbientOcclusion(View, RayGenShaders);
-				PrepareRayTracingSkyLight(View, *Scene, RayGenShaders);
-				PrepareRayTracingGlobalIllumination(View, RayGenShaders);
-				PrepareRayTracingTranslucency(View, RayGenShaders);
+			// This view only needs the path tracing raygen shaders as all other
+			// passes should be disabled.
+			PreparePathTracing(ViewFamily, RayGenShaders);
+		}
+		else
+		{
+			// Path tracing is disabled, get all other possible raygen shaders
+			PrepareRayTracingDebug(ViewFamily, RayGenShaders);
 
-				if (DoesPlatformSupportLumenGI(ShaderPlatform))
+			// These other cases do potentially depend on the camera position since they are
+			// driven by FinalPostProcessSettings, which is why we need to merge them across views
+			if (!IsForwardShadingEnabled(ShaderPlatform))
+			{
+				for (const FViewInfo& View : Views)
 				{
-					PrepareLumenHardwareRayTracingScreenProbeGather(View, RayGenShaders);
-					PrepareLumenHardwareRayTracingRadianceCache(View, RayGenShaders);
-					PrepareLumenHardwareRayTracingTranslucencyVolume(View, RayGenShaders);
-					PrepareLumenHardwareRayTracingReflections(View, RayGenShaders);
-					PrepareLumenHardwareRayTracingVisualize(View, RayGenShaders);
+					PrepareRayTracingReflections(View, *Scene, RayGenShaders);
+					PrepareSingleLayerWaterRayTracingReflections(View, *Scene, RayGenShaders);
+					PrepareRayTracingShadows(View, RayGenShaders);
+					PrepareRayTracingAmbientOcclusion(View, RayGenShaders);
+					PrepareRayTracingSkyLight(View, *Scene, RayGenShaders);
+					PrepareRayTracingGlobalIllumination(View, RayGenShaders);
+					PrepareRayTracingTranslucency(View, RayGenShaders);
+
+					if (DoesPlatformSupportLumenGI(ShaderPlatform) && Lumen::UseHardwareRayTracing())
+					{
+						PrepareLumenHardwareRayTracingScreenProbeGather(View, RayGenShaders);
+						PrepareLumenHardwareRayTracingRadianceCache(View, RayGenShaders);
+						PrepareLumenHardwareRayTracingTranslucencyVolume(View, RayGenShaders);
+						PrepareLumenHardwareRayTracingReflections(View, RayGenShaders);
+						PrepareLumenHardwareRayTracingVisualize(View, RayGenShaders);
+					}
 				}
 			}
+			DeduplicateRayGenerationShaders(RayGenShaders);
 		}
-		DeduplicateRayGenerationShaders(RayGenShaders);
-	}
 
-	if (RayGenShaders.Num())
-	{
-		ReferenceView.RayTracingMaterialPipeline = BindRayTracingMaterialPipeline(RHICmdList, ReferenceView, RayGenShaders);
+		if (RayGenShaders.Num())
+		{
+			ReferenceView.RayTracingMaterialPipeline = BindRayTracingMaterialPipeline(RHICmdList, ReferenceView, RayGenShaders);
+		}
 	}
 
 	// Initialize common resources used for lighting in ray tracing effects
@@ -3047,7 +3057,12 @@ void FDeferredShadingSceneRenderer::Render(FRDGBuilder& GraphBuilder)
 #if RHI_RAYTRACING
 	if (IsRayTracingEnabled())
 	{
-		if (ViewFamily.EngineShowFlags.PathTracing && FDataDrivenShaderPlatformInfo::GetSupportsPathTracing(Scene->GetShaderPlatform()))
+		// Path tracer requires the full ray tracing pipeline support, as well as specialized extra shaders.
+		// Most of the ray tracing debug visualizations also require the full pipeline, but some support inline mode.
+		
+		if (ViewFamily.EngineShowFlags.PathTracing 
+			&& FDataDrivenShaderPlatformInfo::GetSupportsPathTracing(Scene->GetShaderPlatform())
+			&& ShouldRenderRayTracingEffect(ERayTracingPipelineCompatibilityFlags::FullPipeline))
 		{
 			for (const FViewInfo& View : Views)
 			{
diff --git a/Engine/Source/Runtime/Renderer/Private/RayTracing/RayTracingBarycentrics.cpp b/Engine/Source/Runtime/Renderer/Private/RayTracing/RayTracingBarycentrics.cpp
--- a/Engine/Source/Runtime/Renderer/Private/RayTracing/RayTracingBarycentrics.cpp
+++ b/Engine/Source/Runtime/Renderer/Private/RayTracing/RayTracingBarycentrics.cpp
@@ -75,7 +75,7 @@ class FRayTracingBarycentricsCS : public FGlobalShader
 
 	static bool ShouldCompilePermutation(const FGlobalShaderPermutationParameters& Parameters)
 	{
-		return RHISupportsRayTracing(Parameters.Platform) && RHISupportsInlineRayTracing(Parameters.Platform);
+		return IsRayTracingEnabledForProject(Parameters.Platform) && RHISupportsRayTracing(Parameters.Platform) && RHISupportsInlineRayTracing(Parameters.Platform);
 	}
 
 	// Current inline ray tracing implementation requires 1:1 mapping between thread groups and waves and only supports wave32 mode.
diff --git a/Engine/Source/Runtime/Renderer/Private/RayTracing/RayTracingInstanceBufferUtil.cpp b/Engine/Source/Runtime/Renderer/Private/RayTracing/RayTracingInstanceBufferUtil.cpp
--- a/Engine/Source/Runtime/Renderer/Private/RayTracing/RayTracingInstanceBufferUtil.cpp
+++ b/Engine/Source/Runtime/Renderer/Private/RayTracing/RayTracingInstanceBufferUtil.cpp
@@ -258,9 +258,7 @@ struct FRayTracingBuildInstanceBufferCS : public FGlobalShader
 
 	static bool ShouldCompilePermutation(const FGlobalShaderPermutationParameters& Parameters)
 	{
-		//return (ShouldCompileRayTracingShadersForProject(Parameters.Platform) && RHISupportsComputeShaders(Parameters.Platform) && !(Parameters.Platform == EShaderPlatform::SP_METAL || Parameters.Platform == EShaderPlatform::SP_METAL_TVOS || IsMobilePlatform(Parameters.Platform)));
-		// temp: don't check ShouldCompileRayTracingShadersForProject since can't enable it for Vulkan yet.
-		return RHISupportsComputeShaders(Parameters.Platform) && !IsMobilePlatform(Parameters.Platform);
+		return IsRayTracingEnabledForProject(Parameters.Platform) && RHISupportsComputeShaders(Parameters.Platform);
 	}
 };
 
diff --git a/Engine/Source/Runtime/Renderer/Private/RayTracing/RayTracingLighting.cpp b/Engine/Source/Runtime/Renderer/Private/RayTracing/RayTracingLighting.cpp
--- a/Engine/Source/Runtime/Renderer/Private/RayTracing/RayTracingLighting.cpp
+++ b/Engine/Source/Runtime/Renderer/Private/RayTracing/RayTracingLighting.cpp
@@ -31,14 +31,15 @@ class FSetupRayTracingLightCullData : public FGlobalShader
 	DECLARE_GLOBAL_SHADER(FSetupRayTracingLightCullData);
 	SHADER_USE_PARAMETER_STRUCT(FSetupRayTracingLightCullData, FGlobalShader)
 
-		static int32 GetGroupSize()
+	static int32 GetGroupSize()
 	{
 		return 32;
 	}
 
 	static bool ShouldCompilePermutation(const FGlobalShaderPermutationParameters& Parameters)
 	{
-		return ShouldCompileRayTracingShadersForProject(Parameters.Platform);
+		// Allow this shader to be compiled if either inline or full pipeline ray tracing mode is supported by the platform
+		return IsRayTracingEnabledForProject(Parameters.Platform);
 	}
 
 	static void ModifyCompilationEnvironment(const FGlobalShaderPermutationParameters& Parameters, FShaderCompilerEnvironment& OutEnvironment)
diff --git a/Engine/Source/Runtime/Renderer/Private/RayTracing/RayTracingMaterialHitShaders.cpp b/Engine/Source/Runtime/Renderer/Private/RayTracing/RayTracingMaterialHitShaders.cpp
--- a/Engine/Source/Runtime/Renderer/Private/RayTracing/RayTracingMaterialHitShaders.cpp
+++ b/Engine/Source/Runtime/Renderer/Private/RayTracing/RayTracingMaterialHitShaders.cpp
@@ -295,25 +295,16 @@ static bool GetMaterialHitShader(const FMaterial& RESTRICT MaterialResource, con
 	return true;
 }
 
-bool FRayTracingMeshProcessor::Process(
-	const FMeshBatch& RESTRICT MeshBatch,
-	uint64 BatchElementMask,
-	const FPrimitiveSceneProxy* RESTRICT PrimitiveSceneProxy,
-	const FMaterialRenderProxy& RESTRICT MaterialRenderProxy,
+static bool GetRayTracingMeshProcessorShaders(
+	const FUniformLightMapPolicy& RESTRICT LightMapPolicy,
+	const FVertexFactory* VertexFactory,
 	const FMaterial& RESTRICT MaterialResource,
-	const FUniformLightMapPolicy& RESTRICT LightMapPolicy)
+	TShaderRef<FMaterialCHS>& OutRayHitGroupShader)
 {
+	check(GRHISupportsRayTracingShaders);
 
 	const bool bMaterialsCompiled = GCompileRayTracingMaterialAHS || GCompileRayTracingMaterialCHS;
 
-	const FVertexFactory* VertexFactory = MeshBatch.VertexFactory;
-
-	TMeshProcessorShaders<
-		FMeshMaterialShader,
-		FMeshMaterialShader,
-		FMeshMaterialShader,
-		FMaterialCHS> RayTracingShaders;
-
 	if (bMaterialsCompiled)
 	{
 		const bool bUseTextureLOD = bool(GRayTracingUseTextureLod);
@@ -321,31 +312,31 @@ bool FRayTracingMeshProcessor::Process(
 		switch (LightMapPolicy.GetIndirectPolicy())
 		{
 		case LMP_PRECOMPUTED_IRRADIANCE_VOLUME_INDIRECT_LIGHTING:
-			if (!GetMaterialHitShader<TUniformLightMapPolicy<LMP_PRECOMPUTED_IRRADIANCE_VOLUME_INDIRECT_LIGHTING>>(MaterialResource, VertexFactory, bUseTextureLOD, RayTracingShaders.RayHitGroupShader))
+			if (!GetMaterialHitShader<TUniformLightMapPolicy<LMP_PRECOMPUTED_IRRADIANCE_VOLUME_INDIRECT_LIGHTING>>(MaterialResource, VertexFactory, bUseTextureLOD, OutRayHitGroupShader))
 			{
 				return false;
 			}
 			break;
 		case LMP_LQ_LIGHTMAP:
-			if (!GetMaterialHitShader<TUniformLightMapPolicy<LMP_LQ_LIGHTMAP>>(MaterialResource, VertexFactory, bUseTextureLOD, RayTracingShaders.RayHitGroupShader))
+			if (!GetMaterialHitShader<TUniformLightMapPolicy<LMP_LQ_LIGHTMAP>>(MaterialResource, VertexFactory, bUseTextureLOD, OutRayHitGroupShader))
 			{
 				return false;
 			}
 			break;
 		case LMP_HQ_LIGHTMAP:
-			if (!GetMaterialHitShader<TUniformLightMapPolicy<LMP_HQ_LIGHTMAP>>(MaterialResource, VertexFactory, bUseTextureLOD, RayTracingShaders.RayHitGroupShader))
+			if (!GetMaterialHitShader<TUniformLightMapPolicy<LMP_HQ_LIGHTMAP>>(MaterialResource, VertexFactory, bUseTextureLOD, OutRayHitGroupShader))
 			{
 				return false;
 			}
 			break;
 		case LMP_DISTANCE_FIELD_SHADOWS_AND_HQ_LIGHTMAP:
-			if (!GetMaterialHitShader<TUniformLightMapPolicy<LMP_DISTANCE_FIELD_SHADOWS_AND_HQ_LIGHTMAP>>(MaterialResource, VertexFactory, bUseTextureLOD, RayTracingShaders.RayHitGroupShader))
+			if (!GetMaterialHitShader<TUniformLightMapPolicy<LMP_DISTANCE_FIELD_SHADOWS_AND_HQ_LIGHTMAP>>(MaterialResource, VertexFactory, bUseTextureLOD, OutRayHitGroupShader))
 			{
 				return false;
 			}
 			break;
 		case LMP_NO_LIGHTMAP:
-			if (!GetMaterialHitShader<TUniformLightMapPolicy<LMP_NO_LIGHTMAP>>(MaterialResource, VertexFactory, bUseTextureLOD, RayTracingShaders.RayHitGroupShader))
+			if (!GetMaterialHitShader<TUniformLightMapPolicy<LMP_NO_LIGHTMAP>>(MaterialResource, VertexFactory, bUseTextureLOD, OutRayHitGroupShader))
 			{
 				return false;
 			}
@@ -365,7 +356,32 @@ bool FRayTracingMeshProcessor::Process(
 			return false;
 		}
 
-		Shaders.TryGetShader(SF_RayHitGroup, RayTracingShaders.RayHitGroupShader);
+		Shaders.TryGetShader(SF_RayHitGroup, OutRayHitGroupShader);
+	}
+
+	return true;
+}
+
+bool FRayTracingMeshProcessor::Process(
+	const FMeshBatch& RESTRICT MeshBatch,
+	uint64 BatchElementMask,
+	const FPrimitiveSceneProxy* RESTRICT PrimitiveSceneProxy,
+	const FMaterialRenderProxy& RESTRICT MaterialRenderProxy,
+	const FMaterial& RESTRICT MaterialResource,
+	const FUniformLightMapPolicy& RESTRICT LightMapPolicy)
+{
+	TMeshProcessorShaders<
+		FMeshMaterialShader,
+		FMeshMaterialShader,
+		FMeshMaterialShader,
+		FMaterialCHS> RayTracingShaders;
+
+	if (GRHISupportsRayTracingShaders)
+	{
+		if (!GetRayTracingMeshProcessorShaders(LightMapPolicy, MeshBatch.VertexFactory, MaterialResource, RayTracingShaders.RayHitGroupShader))
+		{
+			return false;
+		}
 	}
 
 	TBasePassShaderElementData<FUniformLightMapPolicy> ShaderElementData(MeshBatch.LCI);
diff --git a/Engine/Source/Runtime/Renderer/Private/RayTracing/RayTracingMaterialHitShaders.h b/Engine/Source/Runtime/Renderer/Private/RayTracing/RayTracingMaterialHitShaders.h
--- a/Engine/Source/Runtime/Renderer/Private/RayTracing/RayTracingMaterialHitShaders.h
+++ b/Engine/Source/Runtime/Renderer/Private/RayTracing/RayTracingMaterialHitShaders.h
@@ -64,7 +64,11 @@ protected:
 
 		FRayTracingMeshCommand SharedCommand;
 
-		SharedCommand.SetShaders(PassShaders.GetUntypedShaders());
+		if (GRHISupportsRayTracingShaders)
+		{
+			SharedCommand.SetShaders(PassShaders.GetUntypedShaders());
+		}
+
 		SharedCommand.InstanceMask = ComputeBlendModeMask(MaterialResource.GetBlendMode());
 		SharedCommand.bCastRayTracedShadows = MeshBatch.CastRayTracedShadow && MaterialResource.CastsRayTracedShadows();
 		SharedCommand.bOpaque = MaterialResource.GetBlendMode() == EBlendMode::BLEND_Opaque && !(VertexFactory->GetType()->SupportsRayTracingProceduralPrimitive() && FDataDrivenShaderPlatformInfo::GetSupportsRayTracingProceduralPrimitive(GMaxRHIShaderPlatform));
diff --git a/Engine/Source/Runtime/Renderer/Private/RayTracing/RayTracingPrimaryRays.cpp b/Engine/Source/Runtime/Renderer/Private/RayTracing/RayTracingPrimaryRays.cpp
--- a/Engine/Source/Runtime/Renderer/Private/RayTracing/RayTracingPrimaryRays.cpp
+++ b/Engine/Source/Runtime/Renderer/Private/RayTracing/RayTracingPrimaryRays.cpp
@@ -70,6 +70,12 @@ IMPLEMENT_GLOBAL_SHADER(FRayTracingPrimaryRaysRGS, "/Engine/Private/RayTracing/R
 
 void FDeferredShadingSceneRenderer::PrepareRayTracingTranslucency(const FViewInfo& View, TArray<FRHIRayTracingShader*>& OutRayGenShaders)
 {
+	// Translucency and primary ray tracing requires the full ray tracing pipeline with material bindings.
+	if (!ShouldRenderRayTracingEffect(ERayTracingPipelineCompatibilityFlags::FullPipeline))
+	{
+		return;
+	}
+
 	// Declare all RayGen shaders that require material closest hit shaders to be bound.
 	// NOTE: Translucency shader may be used for primary ray debug view mode.
 	if (GetRayTracingTranslucencyOptions(View).bEnabled || View.Family->EngineShowFlags.RayTracingDebug)
