Strata - no longer write SSSdata during baspass.
Added  corresponding clear.

diff --git a/Engine/Shaders/Private/BasePassPixelShader.usf b/Engine/Shaders/Private/BasePassPixelShader.usf
--- a/Engine/Shaders/Private/BasePassPixelShader.usf
+++ b/Engine/Shaders/Private/BasePassPixelShader.usf
@@ -1632,7 +1632,13 @@ void FPixelShaderInOut_MainPS(
 	// Write out classificaiton & pre-preprocessed data
 	OpaqueBasePass.Strata.ClassificationTextureUAV[PixelPos] = StrataPackClassificationData(ClassificationData);
 	OpaqueBasePass.Strata.TopLayerNormalTextureUAV[PixelPos] = StrataPackTopLayerData(TopLayerData);
-	OpaqueBasePass.Strata.SSSTextureUAV[PixelPos] = StrataPackSSSData(SSSData);
+
+	// Only write SSS datta if needed
+	BRANCH
+	if(SSSData.bIsValid)
+	{
+		OpaqueBasePass.Strata.SSSTextureUAV[PixelPos] = StrataPackSSSData(SSSData);
+	}
 
 	Color += EmissiveLuminance;
 #endif // STRATA_OPAQUE_DEFERRED
diff --git a/Engine/Shaders/Private/Strata/StrataMaterialClassification.usf b/Engine/Shaders/Private/Strata/StrataMaterialClassification.usf
--- a/Engine/Shaders/Private/Strata/StrataMaterialClassification.usf
+++ b/Engine/Shaders/Private/Strata/StrataMaterialClassification.usf
@@ -348,6 +348,7 @@ void StencilMainPS(
 
 RWByteAddressBuffer MaterialLobesBufferUAV;
 RWTexture2D<uint> ClassificationTextureUAV;
+RWTexture2D<uint2> SSSTextureUAV;
 uint MaxBytesPerPixel;
 int2 TiledViewBufferResolution;
 
@@ -372,6 +373,9 @@ void ClearMaterialBufferMainCS(uint2 DispatchThreadId : SV_DispatchThreadID)
 	FStrataClassification ClassificationData = (FStrataClassification)0;
 	SetIsSimpleMaterial(ClassificationData, true);
 	ClassificationTextureUAV[DispatchThreadId.xy] = StrataPackClassificationData(ClassificationData);
+
+	FStrataSubsurfaceData StrataSubsurfaceData = (FStrataSubsurfaceData)0;
+	SSSTextureUAV[DispatchThreadId.xy] = StrataPackSSSData(StrataSubsurfaceData);
 }
 
 #endif // SHADER_CLEAR_MATERIAL_BUFFER
diff --git a/Engine/Source/Runtime/Renderer/Private/Strata/Strata.cpp b/Engine/Source/Runtime/Renderer/Private/Strata/Strata.cpp
--- a/Engine/Source/Runtime/Renderer/Private/Strata/Strata.cpp
+++ b/Engine/Source/Runtime/Renderer/Private/Strata/Strata.cpp
@@ -114,7 +114,13 @@ namespace Strata
 {
 
 // Forward declaration
-static void AddStrataClearMaterialBufferPass(FRDGBuilder& GraphBuilder, FRDGBufferUAVRef MaterialLobesBufferUAV, FRDGTextureUAVRef ClassificationTextureUAV, uint32 MaxBytesPerPixel, FIntPoint TiledViewBufferResolution);
+static void AddStrataClearMaterialBufferPass(
+	FRDGBuilder& GraphBuilder, 
+	FRDGBufferUAVRef MaterialLobesBufferUAV,
+	FRDGTextureUAVRef ClassificationTextureUAV,
+	FRDGTextureUAVRef SSSTextureUAV,
+	uint32 MaxBytesPerPixel, 
+	FIntPoint TiledViewBufferResolution);
 
 bool IsStrataEnabled()
 {
@@ -224,7 +230,13 @@ void InitialiseStrataFrameSceneData(FSceneRenderer& SceneRenderer, FRDGBuilder&
 
 	if (IsStrataEnabled())
 	{
-		AddStrataClearMaterialBufferPass(GraphBuilder, StrataSceneData.MaterialLobesBufferUAV, StrataSceneData.ClassificationTextureUAV, StrataSceneData.MaxBytesPerPixel, MaterialBufferSizeXY);
+		AddStrataClearMaterialBufferPass(
+			GraphBuilder, 
+			StrataSceneData.MaterialLobesBufferUAV,
+			StrataSceneData.ClassificationTextureUAV,
+			StrataSceneData.SSSTextureUAV,
+			StrataSceneData.MaxBytesPerPixel, 
+			MaterialBufferSizeXY);
 	}
 
 	// Create the readable uniform buffers for each views once for all (it is view independent and all the views should be tiled into the render target textures & material buffer)
@@ -382,6 +394,7 @@ class FStrataClearMaterialBufferCS : public FGlobalShader
 	BEGIN_SHADER_PARAMETER_STRUCT(FParameters, )
 		SHADER_PARAMETER_RDG_BUFFER_UAV(RWByteAddressBuffer, MaterialLobesBufferUAV)
 		SHADER_PARAMETER_RDG_TEXTURE_UAV(RWTexture2D<uint>, ClassificationTextureUAV)
+		SHADER_PARAMETER_RDG_TEXTURE_UAV(RWTexture2D<uint2>, SSSTextureUAV)
 		SHADER_PARAMETER(uint32, MaxBytesPerPixel)
 		SHADER_PARAMETER(FIntPoint, TiledViewBufferResolution)
 	END_SHADER_PARAMETER_STRUCT()
@@ -664,12 +677,19 @@ void AddStrataMaterialClassificationPass(FRDGBuilder& GraphBuilder, const FMinim
 	}
 }
 
-static void AddStrataClearMaterialBufferPass(FRDGBuilder& GraphBuilder, FRDGBufferUAVRef MaterialLobesBufferUAV, FRDGTextureUAVRef ClassificationTextureUAV, uint32 MaxBytesPerPixel, FIntPoint TiledViewBufferResolution)
+static void AddStrataClearMaterialBufferPass(
+	FRDGBuilder& GraphBuilder,
+	FRDGBufferUAVRef MaterialLobesBufferUAV,
+	FRDGTextureUAVRef ClassificationTextureUAV,
+	FRDGTextureUAVRef SSSTextureUAV,
+	uint32 MaxBytesPerPixel,
+	FIntPoint TiledViewBufferResolution)
 {
 	TShaderMapRef<FStrataClearMaterialBufferCS> ComputeShader(GetGlobalShaderMap(GMaxRHIFeatureLevel));
 	FStrataClearMaterialBufferCS::FParameters* PassParameters = GraphBuilder.AllocParameters<FStrataClearMaterialBufferCS::FParameters>();
 	PassParameters->MaterialLobesBufferUAV = MaterialLobesBufferUAV;
 	PassParameters->ClassificationTextureUAV = ClassificationTextureUAV;
+	PassParameters->SSSTextureUAV = SSSTextureUAV;
 	PassParameters->MaxBytesPerPixel = MaxBytesPerPixel;
 	PassParameters->TiledViewBufferResolution = TiledViewBufferResolution;
 
