Enabled Legend drawing on the Light Complexity debug visualization modes

#rnx

diff --git a/Engine/Source/Runtime/Renderer/Private/PostProcess/PostProcessing.cpp b/Engine/Source/Runtime/Renderer/Private/PostProcess/PostProcessing.cpp
--- a/Engine/Source/Runtime/Renderer/Private/PostProcess/PostProcessing.cpp
+++ b/Engine/Source/Runtime/Renderer/Private/PostProcess/PostProcessing.cpp
@@ -1068,6 +1068,7 @@ void AddPostProcessingPasses(FRDGBuilder& GraphBuilder, const FViewInfo& View, c
 		PassInputs.SceneColor = OriginalSceneColor;
 		PassInputs.Colors = GEngine->StationaryLightOverlapColors;
 		PassInputs.ColorSamplingMethod = FVisualizeComplexityInputs::EColorSamplingMethod::Ramp;
+		PassInputs.bDrawLegend = true;
 
 		SceneColor = AddVisualizeComplexityPass(GraphBuilder, View, PassInputs);
 	}
@@ -1085,6 +1086,7 @@ void AddPostProcessingPasses(FRDGBuilder& GraphBuilder, const FViewInfo& View, c
 		PassInputs.Colors = GEngine->LightComplexityColors;
 		PassInputs.ColorSamplingMethod = FVisualizeComplexityInputs::EColorSamplingMethod::Linear;
 		PassInputs.ComplexityScale = ComplexityScale;
+		PassInputs.bDrawLegend = true;
 
 		SceneColor = AddVisualizeComplexityPass(GraphBuilder, View, PassInputs);
 	}
