Strata - remove unecessary TexCreate_TargetArraySlicesIndependently flag

diff --git a/Engine/Source/Runtime/Renderer/Private/Strata/Strata.cpp b/Engine/Source/Runtime/Renderer/Private/Strata/Strata.cpp
--- a/Engine/Source/Runtime/Renderer/Private/Strata/Strata.cpp
+++ b/Engine/Source/Runtime/Renderer/Private/Strata/Strata.cpp
@@ -206,9 +206,11 @@ void InitialiseStrataFrameSceneData(FSceneRenderer& SceneRenderer, FRDGBuilder&
 
 	// Create the material data container
 	FIntPoint SceneTextureExtent = IsStrataEnabled() ? GetSceneTextureExtent() : FIntPoint(2, 2);
-	const FRDGTextureDesc MaterialTextureDesc = FRHITextureCreateInfo::Create2DArray(SceneTextureExtent, PF_R32_UINT, FClearValueBinding::Transparent, TexCreate_TargetArraySlicesIndependently | TexCreate_RenderTargetable | TexCreate_ShaderResource | TexCreate_UAV, StrataSceneData.MaxBytesPerPixel / 4, 1, 1);
+
 	const FRDGTextureDesc MaterialTextureDescMRTs = FRHITextureCreateInfo::Create2DArray(SceneTextureExtent, PF_R32_UINT, FClearValueBinding::Transparent, TexCreate_TargetArraySlicesIndependently | TexCreate_RenderTargetable | TexCreate_ShaderResource | TexCreate_UAV, 2, 1, 1);
 	StrataSceneData.MaterialTextureArrayMRTs = GraphBuilder.CreateTexture(MaterialTextureDescMRTs, TEXT("Strata.MaterialRT"));
+
+	const FRDGTextureDesc MaterialTextureDesc = FRHITextureCreateInfo::Create2DArray(SceneTextureExtent, PF_R32_UINT, FClearValueBinding::Transparent, TexCreate_RenderTargetable | TexCreate_ShaderResource | TexCreate_UAV, StrataSceneData.MaxBytesPerPixel / 4, 1, 1);
 	StrataSceneData.MaterialTextureArray = GraphBuilder.CreateTexture(MaterialTextureDesc, TEXT("Strata.Material"));
 	StrataSceneData.MaterialTextureArraySRV = GraphBuilder.CreateSRV(FRDGTextureSRVDesc::Create(StrataSceneData.MaterialTextureArray));
 	StrataSceneData.MaterialTextureArrayUAV = GraphBuilder.CreateUAV(FRDGTextureUAVDesc(StrataSceneData.MaterialTextureArray, 0));
