Fix missing output render target format for Strata MRTs

diff --git a/Engine/Source/Runtime/Renderer/Private/BasePassRendering.h b/Engine/Source/Runtime/Renderer/Private/BasePassRendering.h
--- a/Engine/Source/Runtime/Renderer/Private/BasePassRendering.h
+++ b/Engine/Source/Runtime/Renderer/Private/BasePassRendering.h
@@ -320,6 +320,8 @@ public:
 			OutEnvironment.SetRenderTargetOutputFormat(VelocityIndex, PF_G16R16);
 		}
 
+		Strata::SetBasePassRenderTargetOutputFormat(Parameters.Platform, OutEnvironment);
+
 		FForwardLightingParameters::ModifyCompilationEnvironment(Parameters.Platform, OutEnvironment);
 	}
 
diff --git a/Engine/Source/Runtime/Renderer/Private/Strata/Strata.cpp b/Engine/Source/Runtime/Renderer/Private/Strata/Strata.cpp
--- a/Engine/Source/Runtime/Renderer/Private/Strata/Strata.cpp
+++ b/Engine/Source/Runtime/Renderer/Private/Strata/Strata.cpp
@@ -11,6 +11,7 @@
 #include "SceneTextureParameters.h"
 #include "StrataDefinitions.h"
 #include "ScreenPass.h"
+#include "ShaderCompiler.h"
 
 //PRAGMA_DISABLE_OPTIMIZATION
 
@@ -632,6 +633,8 @@ void ApprendStrataMRTs(FSceneRenderer& SceneRenderer, uint32& RenderTargetCount,
 {
 	if (Strata::IsStrataEnabled() && SceneRenderer.Scene)
 	{
+		// If this function changes, update Strata::SetBasePassRenderTargetOutputFormat()
+		 
 		// Add 2 uint for Strata fast path
 		auto AddStrataOutputTarget = [&](int16 StrataMaterialArraySlice)
 		{
@@ -649,6 +652,22 @@ void ApprendStrataMRTs(FSceneRenderer& SceneRenderer, uint32& RenderTargetCount,
 	}
 }
 
+void SetBasePassRenderTargetOutputFormat(const EShaderPlatform Platform, FShaderCompilerEnvironment& OutEnvironment)
+{
+	if (Strata::IsStrataEnabled())
+	{
+		const FGBufferParams GBufferParams = FShaderCompileUtilities::FetchGBufferParamsRuntime(Platform);
+		const FGBufferInfo BufferInfo = FetchFullGBufferInfo(GBufferParams);
+
+		// Add 2 uint for Strata fast path
+		OutEnvironment.SetRenderTargetOutputFormat(BufferInfo.NumTargets + 0, PF_R32_UINT);
+		OutEnvironment.SetRenderTargetOutputFormat(BufferInfo.NumTargets + 1, PF_R32_UINT);
+
+		// Add another MRT for Strata top layer information
+		OutEnvironment.SetRenderTargetOutputFormat(BufferInfo.NumTargets + 2, PF_R32_UINT);
+	}
+}
+
 ////////////////////////////////////////////////////////////////////////////////////////////////////////////////////////////////////////////////
 
 class FStrataPostBasePassCS : public FGlobalShader
diff --git a/Engine/Source/Runtime/Renderer/Private/Strata/Strata.h b/Engine/Source/Runtime/Renderer/Private/Strata/Strata.h
--- a/Engine/Source/Runtime/Renderer/Private/Strata/Strata.h
+++ b/Engine/Source/Runtime/Renderer/Private/Strata/Strata.h
@@ -91,6 +91,7 @@ void BindStrataBasePassUniformParameters(FRDGBuilder& GraphBuilder, FStrataScene
 void BindStrataGlobalUniformParameters(FRDGBuilder& GraphBuilder, FStrataSceneData* StrataSceneData, FStrataGlobalUniformParameters& OutStrataUniformParameters);
 
 void ApprendStrataMRTs(FSceneRenderer& SceneRenderer, uint32& BasePassTextureCount, TStaticArray<FTextureRenderTargetBinding, MaxSimultaneousRenderTargets>& BasePassTextures);
+void SetBasePassRenderTargetOutputFormat(const EShaderPlatform Platform, FShaderCompilerEnvironment& OutEnvironment);
 void PostBasePass(FRDGBuilder& GraphBuilder, const TArray<FViewInfo>& Views);
 
 TRDGUniformBufferRef<FStrataGlobalUniformParameters> BindStrataGlobalUniformParameters(FStrataSceneData* StrataSceneData);
