Move SS profiles and SS pre-intregrated profiles textures onto the view in order to make them accessible to all passes without extra manual binding.

diff --git a/Engine/Shaders/Private/BasePassCommon.ush b/Engine/Shaders/Private/BasePassCommon.ush
--- a/Engine/Shaders/Private/BasePassCommon.ush
+++ b/Engine/Shaders/Private/BasePassCommon.ush
@@ -4,19 +4,16 @@
 	BasePassCommon.ush: Base pass definitions used by both vertex and pixel shader
 =============================================================================*/
 
-// Defining SSProfilesTexture here overrides the default definition in SubsurfaceProfileCommon.ush
 #if MATERIALBLENDING_ANY_TRANSLUCENT
 	#define ForwardLightData TranslucentBasePass.Shared.Forward
 	#define ReflectionStruct TranslucentBasePass.Shared.Reflection
 	#define PlanarReflectionStruct TranslucentBasePass.Shared.PlanarReflection
 	#define FogStruct TranslucentBasePass.Shared.Fog
-	#define ActualSSProfilesTexture TranslucentBasePass.Shared.SSProfilesTexture
 #else
 	#define ForwardLightData OpaqueBasePass.Shared.Forward
 	#define ReflectionStruct OpaqueBasePass.Shared.Reflection
 	#define PlanarReflectionStruct OpaqueBasePass.Shared.PlanarReflection
 	#define FogStruct OpaqueBasePass.Shared.Fog
-	#define ActualSSProfilesTexture OpaqueBasePass.Shared.SSProfilesTexture
 #endif 
 
 #undef NEEDS_LIGHTMAP_COORDINATE
diff --git a/Engine/Shaders/Private/BurleyNormalizedSSSCommon.ush b/Engine/Shaders/Private/BurleyNormalizedSSSCommon.ush
--- a/Engine/Shaders/Private/BurleyNormalizedSSSCommon.ush
+++ b/Engine/Shaders/Private/BurleyNormalizedSSSCommon.ush
@@ -43,7 +43,7 @@ float4 DecodeDiffuseMeanFreePath(float4 EncodedDiffuseMeanFreePath)
 bool GetSubsurfaceTransmittanceProfileUseBurley(uint SubsurfaceProfileInt)
 {
 	// 0..255, which SubSurface profile to pick
-	float Type = ActualSSProfilesTexture.Load(int3(SSSS_BOUNDARY_COLOR_BLEED_OFFSET, SubsurfaceProfileInt, 0)).a;
+	float Type = View.SSProfilesTexture.Load(int3(SSSS_BOUNDARY_COLOR_BLEED_OFFSET, SubsurfaceProfileInt, 0)).a;
 
 	return abs(Type - SSS_TYPE_BURLEY) < 0.01f;
 }
@@ -51,7 +51,7 @@ bool GetSubsurfaceTransmittanceProfileUseBurley(uint SubsurfaceProfileInt)
 bool GetSubsurfaceProfileUseBurley(uint SubsurfaceProfileInt)
 {
 	// 0..255, which SubSurface profile to pick
-	float Type = ActualSSProfilesTexture.Load(int3(SSSS_BOUNDARY_COLOR_BLEED_OFFSET, SubsurfaceProfileInt, 0)).a;
+	float Type = View.SSProfilesTexture.Load(int3(SSSS_BOUNDARY_COLOR_BLEED_OFFSET, SubsurfaceProfileInt, 0)).a;
 
 	return abs(Type - SSS_TYPE_BURLEY) < 0.01f;
 }
@@ -59,7 +59,7 @@ bool GetSubsurfaceProfileUseBurley(uint SubsurfaceProfileInt)
 bool GetSubsurfaceProfileUseSeparable(uint SubsurfaceProfileInt)
 {
 	// 0..255, which SubSurface profile to pick
-	float Type = ActualSSProfilesTexture.Load(int3(SSSS_BOUNDARY_COLOR_BLEED_OFFSET, SubsurfaceProfileInt, 0)).a;
+	float Type = View.SSProfilesTexture.Load(int3(SSSS_BOUNDARY_COLOR_BLEED_OFFSET, SubsurfaceProfileInt, 0)).a;
 
 	return abs(Type - SSS_TYPE_SSSS) < 0.01f;
 }
@@ -185,9 +185,9 @@ float4 GetBurleyTransmissionProfile(FGBufferData GBufferData, float Thickness)
 	// 0..255, which SubSurface profile to pick
 
 	uint SubsurfaceProfileInt = ExtractSubsurfaceProfileInt(GBufferData);
-	float WorldUnitScale = DecodeWorldUnitScale(ActualSSProfilesTexture.Load(int3(SSSS_SUBSURFACE_COLOR_OFFSET, SubsurfaceProfileInt,0)).a)*10.0f; //in cm. we do not multiply by 100.0f as we did in the subsurface shader.
-	float4 MeanFreePath = DecodeDiffuseMeanFreePath(ActualSSProfilesTexture.Load(int3(BSSS_DMFP_OFFSET, SubsurfaceProfileInt, 0)));
-	float3 SurfaceAlbedo = ActualSSProfilesTexture.Load(int3(BSSS_SURFACEALBEDO_OFFSET, SubsurfaceProfileInt, 0)).rgb;
+	float WorldUnitScale = DecodeWorldUnitScale(View.SSProfilesTexture.Load(int3(SSSS_SUBSURFACE_COLOR_OFFSET, SubsurfaceProfileInt,0)).a)*10.0f; //in cm. we do not multiply by 100.0f as we did in the subsurface shader.
+	float4 MeanFreePath = DecodeDiffuseMeanFreePath(View.SSProfilesTexture.Load(int3(BSSS_DMFP_OFFSET, SubsurfaceProfileInt, 0)));
+	float3 SurfaceAlbedo = View.SSProfilesTexture.Load(int3(BSSS_SURFACEALBEDO_OFFSET, SubsurfaceProfileInt, 0)).rgb;
 	float3 ScalingFactor = GetSearchLightDiffuseScalingFactor3D(SurfaceAlbedo);//assuming that the volume albedo is the same to the surface albedo for transmission;
 	float4 Output=float4(0,0,0,1);
 	
diff --git a/Engine/Shaders/Private/MobileBasePassPixelShader.usf b/Engine/Shaders/Private/MobileBasePassPixelShader.usf
--- a/Engine/Shaders/Private/MobileBasePassPixelShader.usf
+++ b/Engine/Shaders/Private/MobileBasePassPixelShader.usf
@@ -10,13 +10,6 @@
 #define MobileSceneTextures MobileBasePass.SceneTextures
 #define EyeAdaptationStruct MobileBasePass
 
-#define ActualSSProfilesTexture MobileBasePass.SSProfilesTexture
-#define SSProfilesSampler MobileBasePass.SSProfilesSampler
-
-#define SSProfilesPreIntegratedTexture MobileBasePass.SSProfilesPreIntegratedTexture
-#define SSProfilesPreIntegratedSampler MobileBasePass.SSProfilesPreIntegratedSampler
-#define SSProfilesTextureSizeAndInvSize MobileBasePass.SSProfilesTextureSizeAndInvSize
-
 #ifndef MOBILE_QL_FORCE_DISABLE_PREINTEGRATEDGF
 #define MOBILE_QL_FORCE_DISABLE_PREINTEGRATEDGF 0
 #endif
diff --git a/Engine/Shaders/Private/MobileShadingModels.ush b/Engine/Shaders/Private/MobileShadingModels.ush
--- a/Engine/Shaders/Private/MobileShadingModels.ush
+++ b/Engine/Shaders/Private/MobileShadingModels.ush
@@ -140,7 +140,7 @@ half4 GetSSProfilePreIntegratedValue(uint SubsurfaceProfileInt, half NoL, half C
 {
 	float3 UV = float3((NoL * .5 + .5), Curvature, SubsurfaceProfileInt);
 
-	return Texture2DArraySampleLevel(SSProfilesPreIntegratedTexture, SSProfilesPreIntegratedSampler, UV, 0);
+	return Texture2DArraySampleLevel(View.SSProfilesPreIntegratedTexture, View.SSProfilesPreIntegratedSampler, UV, 0);
 }
 #endif
 
diff --git a/Engine/Shaders/Private/PathTracing/Material/PathTracingMaterialSampling.ush b/Engine/Shaders/Private/PathTracing/Material/PathTracingMaterialSampling.ush
--- a/Engine/Shaders/Private/PathTracing/Material/PathTracingMaterialSampling.ush
+++ b/Engine/Shaders/Private/PathTracing/Material/PathTracingMaterialSampling.ush
@@ -380,9 +380,9 @@ FPathTracingPayload UnpackPathTracingPayloadMaterial(FPackedPathTracingPayload P
 				// NOTE: would be better to just precompute all this on the CPU and just have a single value to lookup
 
 				// Burley parameterization
-				float3 SurfaceAlbedo = ActualSSProfilesTexture.Load(int3(BSSS_SURFACEALBEDO_OFFSET, ProfileId, 0)).xyz;
-				float3 DiffuseMeanFreePath = DecodeDiffuseMeanFreePath(ActualSSProfilesTexture.Load(int3(BSSS_DMFP_OFFSET, ProfileId, 0))).xyz;
-				float WorldUnitScale = DecodeWorldUnitScale(ActualSSProfilesTexture.Load(int3(SSSS_SUBSURFACE_COLOR_OFFSET, ProfileId, 0)).a) * BURLEY_MM_2_CM;
+				float3 SurfaceAlbedo = View.SSProfilesTexture.Load(int3(BSSS_SURFACEALBEDO_OFFSET, ProfileId, 0)).xyz;
+				float3 DiffuseMeanFreePath = DecodeDiffuseMeanFreePath(View.SSProfilesTexture.Load(int3(BSSS_DMFP_OFFSET, ProfileId, 0))).xyz;
+				float WorldUnitScale = DecodeWorldUnitScale(View.SSProfilesTexture.Load(int3(SSSS_SUBSURFACE_COLOR_OFFSET, ProfileId, 0)).a) * BURLEY_MM_2_CM;
 
 				// Opacity acts as a per-pixel radius multiplier
 				// NOTE: this seems backwards? Opacity=0 acts like default-lit while Opacity=1 acts like SSS?
@@ -405,7 +405,7 @@ FPathTracingPayload UnpackPathTracingPayloadMaterial(FPackedPathTracingPayload P
 			// Decode dual roughness info
 
 			// See logic in ShadingModels.ush
-			float4 Data = ActualSSProfilesTexture.Load(int3(SSSS_DUAL_SPECULAR_OFFSET, ProfileId, 0));
+			float4 Data = View.SSProfilesTexture.Load(int3(SSSS_DUAL_SPECULAR_OFFSET, ProfileId, 0));
 			float AverageToRoughness0 = Data.x * SSSS_MAX_DUAL_SPECULAR_ROUGHNESS;
 			float AverageToRoughness1 = Data.y * SSSS_MAX_DUAL_SPECULAR_ROUGHNESS;
 			float LobeMix = Data.z;
diff --git a/Engine/Shaders/Private/PostProcessSubsurface.usf b/Engine/Shaders/Private/PostProcessSubsurface.usf
--- a/Engine/Shaders/Private/PostProcessSubsurface.usf
+++ b/Engine/Shaders/Private/PostProcessSubsurface.usf
@@ -148,7 +148,7 @@ float4 GetSubsurfaceProfileKernel(uint SampleIndex, uint SubsurfaceProfileInt)
 #endif
 	const float4 TableMax = float4(1, 1, 1, SUBSURFACE_KERNEL_SIZE);
 
-	return ActualSSProfilesTexture.Load(int3(SampleIndex, SubsurfaceProfileInt, 0)) * TableMax;
+	return View.SSProfilesTexture.Load(int3(SampleIndex, SubsurfaceProfileInt, 0)) * TableMax;
 }
 
 float3 GetSubsurfaceProfileNormal(float2 BufferUV)
diff --git a/Engine/Shaders/Private/RayTracing/RayTracingLightingMS.usf b/Engine/Shaders/Private/RayTracing/RayTracingLightingMS.usf
--- a/Engine/Shaders/Private/RayTracing/RayTracingLightingMS.usf
+++ b/Engine/Shaders/Private/RayTracing/RayTracingLightingMS.usf
@@ -12,7 +12,6 @@
 #define LTCAmpSampler				RaytracingLightsDataPacked.LTCAmpSampler
 #define PreIntegratedGF				ReflectionStruct.PreIntegratedGF			
 #define PreIntegratedGFSampler		ReflectionStruct.PreIntegratedGFSampler	
-#define SSProfilesTexture			RaytracingLightsDataPacked.SSProfilesTexture
 
 #include "../SceneTextureParameters.ush"
 #include "RayTracingCommon.ush"
diff --git a/Engine/Shaders/Private/SSProfilePreIntegratedMobile.usf b/Engine/Shaders/Private/SSProfilePreIntegratedMobile.usf
--- a/Engine/Shaders/Private/SSProfilePreIntegratedMobile.usf
+++ b/Engine/Shaders/Private/SSProfilePreIntegratedMobile.usf
@@ -7,24 +7,21 @@
 #include "Common.ush"
 #include "BurleyNormalizedSSSCommon.ush"
 
-float4 SSProfilesPreIntegratedTextureSizeAndInvSize;
-uint SubsurfaceProfileInt;
+uint		 SourceSubsurfaceProfileInt;
+float4		 TargetSSProfilesPreIntegratedTextureSizeAndInvSize;
+float4		 SourceSSProfilesTextureSizeAndInvSize;
+Texture2D	 SourceSSProfilesTexture;
+SamplerState SourceSSProfilesSampler;
 
-Texture2DArray InTextureArray;
-SamplerState InTextureSampler;
-
-void MainPSCopy(
-	FScreenVertexOutput Input,
-	out float4 OutColor : SV_Target0
-)
+float4 GetSourceSubsurfaceProfileTexture(uint InSampleIndex, uint InSubsurfaceProfileInt)
 {
-	OutColor = Texture2DArraySample(InTextureArray, InTextureSampler, float3(Input.UV, SubsurfaceProfileInt));
+	return GetSubsurfaceProfileTexture(SourceSSProfilesTexture, SourceSSProfilesSampler, SourceSSProfilesTextureSizeAndInvSize, InSampleIndex, InSubsurfaceProfileInt);
 }
 
 bool IsSubsurfaceProfileUseBurley()
 {
 	// 0..255, which SubSurface profile to pick
-	float Type = GetSubsurfaceProfileTexture(SSSS_BOUNDARY_COLOR_BLEED_OFFSET, SubsurfaceProfileInt).a;
+	float Type = GetSourceSubsurfaceProfileTexture(SSSS_BOUNDARY_COLOR_BLEED_OFFSET, SourceSubsurfaceProfileInt).a;
 
 	return abs(Type - SSS_TYPE_BURLEY) < 0.01f;
 }
@@ -121,7 +118,7 @@ void SSProfilePreIntegratedPS(
 {
 	OutColor = HALF4_TYPE(1.0f, 1.0f, 1.0f, 0.0f);
 
-	float2 UVPos = SvPosition.xy * SSProfilesPreIntegratedTextureSizeAndInvSize.zw;
+	float2 UVPos = SvPosition.xy * TargetSSProfilesPreIntegratedTextureSizeAndInvSize.zw;
 
 	float UnClampedNoL = UVPos.x * 2.0f - 1.0f;
 	float Curvature = UVPos.y;
@@ -137,11 +134,11 @@ void SSProfilePreIntegratedPS(
 		DiffuseMeanFreePathSampleIndex = SSSS_APPROXIMATE_DMFP_OFFSET;
 	}
 
-	float4 SurfaceAlbedo = GetSubsurfaceProfileTexture(SurfaceAlbedoSampleIndex, SubsurfaceProfileInt);
+	float4 SurfaceAlbedo = GetSourceSubsurfaceProfileTexture(SurfaceAlbedoSampleIndex, SourceSubsurfaceProfileInt);
 
-	float WorldUnitScale = DecodeWorldUnitScale(GetSubsurfaceProfileTexture(SSSS_SUBSURFACE_COLOR_OFFSET, SubsurfaceProfileInt).a) * BURLEY_CM_2_MM;
+	float WorldUnitScale = DecodeWorldUnitScale(GetSourceSubsurfaceProfileTexture(SSSS_SUBSURFACE_COLOR_OFFSET, SourceSubsurfaceProfileInt).a) * BURLEY_CM_2_MM;
 
-	float4 DiffuseMeanFreePath = max(DecodeDiffuseMeanFreePath(GetSubsurfaceProfileTexture(DiffuseMeanFreePathSampleIndex, SubsurfaceProfileInt)) * WorldUnitScale, 0.01f);;
+	float4 DiffuseMeanFreePath = max(DecodeDiffuseMeanFreePath(GetSourceSubsurfaceProfileTexture(DiffuseMeanFreePathSampleIndex, SourceSubsurfaceProfileInt)) * WorldUnitScale, 0.01f);;
 
 	float CosTheta = UnClampedNoL;
 	float Theta = acos(CosTheta);
diff --git a/Engine/Shaders/Private/ShadingModels.ush b/Engine/Shaders/Private/ShadingModels.ush
--- a/Engine/Shaders/Private/ShadingModels.ush
+++ b/Engine/Shaders/Private/ShadingModels.ush
@@ -563,7 +563,7 @@ void GetProfileDualSpecular(FGBufferData GBuffer, out float AverageToRoughness0,
 	// 0..255, which SubSurface profile to pick
 	uint SubsurfaceProfileInt = ExtractSubsurfaceProfileInt(GBuffer);
 
-	float4 Data = ActualSSProfilesTexture.Load(int3(SSSS_DUAL_SPECULAR_OFFSET, SubsurfaceProfileInt, 0));
+	float4 Data = View.SSProfilesTexture.Load(int3(SSSS_DUAL_SPECULAR_OFFSET, SubsurfaceProfileInt, 0));
 	AverageToRoughness0 = Data.x * SSSS_MAX_DUAL_SPECULAR_ROUGHNESS;
 	AverageToRoughness1 = Data.y * SSSS_MAX_DUAL_SPECULAR_ROUGHNESS;
 	LobeMix = Data.z;
diff --git a/Engine/Shaders/Private/SubsurfaceBurleyNormalized.ush b/Engine/Shaders/Private/SubsurfaceBurleyNormalized.ush
--- a/Engine/Shaders/Private/SubsurfaceBurleyNormalized.ush
+++ b/Engine/Shaders/Private/SubsurfaceBurleyNormalized.ush
@@ -152,17 +152,17 @@ struct FBurleyParameter
 //Surface albedo and mean free path length
 float4 GetSubsurfaceProfileSurfaceAlbedo(uint SubsurfaceProfileInt)
 {
-	return ActualSSProfilesTexture.Load(int3(BSSS_SURFACEALBEDO_OFFSET, SubsurfaceProfileInt, 0));
+	return View.SSProfilesTexture.Load(int3(BSSS_SURFACEALBEDO_OFFSET, SubsurfaceProfileInt, 0));
 }
 
 float4 GetSubsurfaceProfileDiffuseMeanFreePath(uint SubsurfaceProfileInt)
 {
-	return ActualSSProfilesTexture.Load(int3(BSSS_DMFP_OFFSET, SubsurfaceProfileInt, 0));
+	return View.SSProfilesTexture.Load(int3(BSSS_DMFP_OFFSET, SubsurfaceProfileInt, 0));
 }
 
 float GetSubsurfaceProfileWorldUnitScale(uint SubsurfaceProfileInt)
 {
-	return ActualSSProfilesTexture.Load(int3(SSSS_SUBSURFACE_COLOR_OFFSET, SubsurfaceProfileInt, 0)).a;
+	return View.SSProfilesTexture.Load(int3(SSSS_SUBSURFACE_COLOR_OFFSET, SubsurfaceProfileInt, 0)).a;
 }
 
 float SampleDepthTexture(float2 ScreenUV)
diff --git a/Engine/Shaders/Private/SubsurfaceProfileCommon.ush b/Engine/Shaders/Private/SubsurfaceProfileCommon.ush
--- a/Engine/Shaders/Private/SubsurfaceProfileCommon.ush
+++ b/Engine/Shaders/Private/SubsurfaceProfileCommon.ush
@@ -33,30 +33,22 @@
 // Threshold value at which model switches from SSS to default lit
 #define SSSS_OPACITY_THRESHOLD_EPS 0.10
 
-// One profile per line, encoded using constants above. See FSubsurfaceProfileTexture::CreateTexture() in SubsurfaceProfile.cpp.
-// By default, SSProfilesTexture is "Texture2D SSProfilesTexture", but it can be remapped as in BasePassCommon.usf.
-#ifndef SSProfilesTexture
-Texture2D SSProfilesTexture;
-#endif
-// Rename to allow the base pass to intercept with a value from its pass uniform buffer
-#ifndef ActualSSProfilesTexture
-#define ActualSSProfilesTexture SSProfilesTexture
-
-	Texture2DArray SSProfilesPreIntegratedTexture;
-	SamplerState SSProfilesPreIntegratedSampler;
-	#if COMPILER_GLSL_ES3_1
-	SamplerState SSProfilesSampler;
-	float4 SSProfilesTextureSizeAndInvSize;
-	#endif
-
-#endif
-
 half4 GetSubsurfaceProfileTexture(uint SampleIndex, uint SubsurfaceProfileInt)
 {
+	// One profile per line, encoded using constants above. See FSubsurfaceProfileTexture::CreateTexture() in SubsurfaceProfile.cpp.
 #if COMPILER_GLSL_ES3_1 // Force to use a point sampler for Texture2D.Load on OpenGLES platform
-	return ActualSSProfilesTexture.SampleLevel(SSProfilesSampler, (uint2(SampleIndex, SubsurfaceProfileInt) + float2(0.5f, 0.5f)) * SSProfilesTextureSizeAndInvSize.zw, 0);
+	return View.SSProfilesTexture.SampleLevel(View.SSProfilesSampler, (uint2(SampleIndex, SubsurfaceProfileInt) + float2(0.5f, 0.5f)) * View.SSProfilesTextureSizeAndInvSize.zw, 0);
 #else
-	return ActualSSProfilesTexture.Load(int3(SampleIndex, SubsurfaceProfileInt, 0));
+	return View.SSProfilesTexture.Load(int3(SampleIndex, SubsurfaceProfileInt, 0));
 #endif 
 }
 
+half4 GetSubsurfaceProfileTexture(Texture2D InSSProfilesTexture, SamplerState InSSProfilesSampler, float4 InSSProfilesTextureSizeAndInvSize, uint SampleIndex, uint SubsurfaceProfileInt)
+{
+	// One profile per line, encoded using constants above. See FSubsurfaceProfileTexture::CreateTexture() in SubsurfaceProfile.cpp.
+#if COMPILER_GLSL_ES3_1 // Force to use a point sampler for Texture2D.Load on OpenGLES platform
+	return InSSProfilesTexture.SampleLevel(InSSProfilesSampler, (uint2(SampleIndex, SubsurfaceProfileInt) + float2(0.5f, 0.5f)) * InSSProfilesTextureSizeAndInvSize.zw, 0);
+#else
+	return InSSProfilesTexture.Load(int3(SampleIndex, SubsurfaceProfileInt, 0));
+#endif 
+}
diff --git a/Engine/Shaders/Private/TransmissionCommon.ush b/Engine/Shaders/Private/TransmissionCommon.ush
--- a/Engine/Shaders/Private/TransmissionCommon.ush
+++ b/Engine/Shaders/Private/TransmissionCommon.ush
@@ -9,8 +9,6 @@
 #include "SubsurfaceProfileCommon.ush"
 #include "BurleyNormalizedSSSCommon.ush"
 
-SamplerState TransmissionProfilesLinearSampler;
-
 struct FTransmissionProfileParams
 {
 	float ExtinctionScale;
@@ -34,7 +32,7 @@ FTransmissionProfileParams GetTransmissionProfileParams(uint SubsurfaceProfileIn
 	FTransmissionProfileParams Result;
 
 	//X:ExtinctionScale, Y:Normal Scale, Z:ScatteringDistribution, W:OneOverIOR
-	float4 Data = ActualSSProfilesTexture.Load(int3(SSSS_TRANSMISSION_OFFSET, SubsurfaceProfileInt, 0));
+	float4 Data = View.SSProfilesTexture.Load(int3(SSSS_TRANSMISSION_OFFSET, SubsurfaceProfileInt, 0));
 
 	Result.ExtinctionScale = Data.x;
 	Result.NormalScale = Data.y;
@@ -54,12 +52,11 @@ FTransmissionProfileParams GetTransmissionProfileParams(FGBufferData GBufferData
 float4 GetTransmissionProfile(uint SubsurfaceProfileInt, float Thickness)
 {
 	// 0..255, which SubSurface profile to pick
-	float3 TextureSize;
-	ActualSSProfilesTexture.GetDimensions(0, TextureSize.x, TextureSize.y,TextureSize.z); // TODO: it's slightly faster to pass texture dimensions through a constant buffer
+	const float2 TextureSize = View.SSProfilesTextureSizeAndInvSize.xy;
 	bool bUseBurley = GetSubsurfaceTransmittanceProfileUseBurley(SubsurfaceProfileInt);
 	float ProfileOffset = bUseBurley ? BSSS_TRANSMISSION_PROFILE_OFFSET : SSSS_TRANSMISSION_PROFILE_OFFSET;
 	float2 uv = float2((ProfileOffset + Thickness / SSSS_MAX_TRANSMISSION_PROFILE_DISTANCE * (SSSS_TRANSMISSION_PROFILE_SIZE - 1) + 0.5f), SubsurfaceProfileInt + 0.5f) / (TextureSize.xy);
-	return Texture2DSampleLevel(ActualSSProfilesTexture, TransmissionProfilesLinearSampler, uv, 0);
+	return Texture2DSampleLevel(View.SSProfilesTexture, View.SSProfilesTransmissionSampler, uv, 0);
 }
 
 float4 GetTransmissionProfile(FGBufferData GBufferData, float Thickness)
diff --git a/Engine/Source/Runtime/Engine/Classes/Engine/SubsurfaceProfile.h b/Engine/Source/Runtime/Engine/Classes/Engine/SubsurfaceProfile.h
--- a/Engine/Source/Runtime/Engine/Classes/Engine/SubsurfaceProfile.h
+++ b/Engine/Source/Runtime/Engine/Classes/Engine/SubsurfaceProfile.h
@@ -256,4 +256,10 @@ ENGINE_API FRHITexture* GetSubsurfaceProfileTexture();
 ENGINE_API FRHITexture* GetSubsurfaceProfileTextureWithFallback();
 
 // Returns the Preintegrated texture if it exists, or black.
-ENGINE_API FRHITexture* GetSSProfilesPreIntegratedTextureWithFallback();
\ No newline at end of file
+ENGINE_API FRHITexture* GetSSProfilesPreIntegratedTextureWithFallback();
+
+// Returns the subsurface profile ID shader parameter name
+ENGINE_API FName GetSubsurfaceProfileParameterName();
+
+// Returns the subsurface profile ID for a given Sub-surface Profile object
+ENGINE_API float GetSubsurfaceProfileId(const USubsurfaceProfile* In);
\ No newline at end of file
diff --git a/Engine/Source/Runtime/Engine/Private/Materials/Material.cpp b/Engine/Source/Runtime/Engine/Private/Materials/Material.cpp
--- a/Engine/Source/Runtime/Engine/Private/Materials/Material.cpp
+++ b/Engine/Source/Runtime/Engine/Private/Materials/Material.cpp
@@ -416,24 +416,9 @@ public:
 		const FMaterialResource* MaterialResource = Material->GetMaterialResource(Context.Material.GetFeatureLevel());
 		if (MaterialResource && MaterialResource->GetRenderingThreadShaderMap())
 		{
-			static FName NameSubsurfaceProfile(TEXT("__SubsurfaceProfile"));
-			if (Type == EMaterialParameterType::Scalar && ParameterInfo.Name == NameSubsurfaceProfile)
+			if (Type == EMaterialParameterType::Scalar && ParameterInfo.Name == GetSubsurfaceProfileParameterName())
 			{
-				const USubsurfaceProfile* MySubsurfaceProfileRT = GetSubsurfaceProfileRT();
-
-				int32 AllocationId = 0;
-				if (MySubsurfaceProfileRT)
-				{
-					// can be optimized (cached)
-					AllocationId = GSubsurfaceProfileTextureObject.FindAllocationId(MySubsurfaceProfileRT);
-				}
-				else
-				{
-					// no profile specified means we use the default one stored at [0] which is human skin
-					AllocationId = 0;
-				}
-
-				OutValue = AllocationId / 255.0f;
+				OutValue = GetSubsurfaceProfileId(GetSubsurfaceProfileRT());
 				return true;
 			}
 
diff --git a/Engine/Source/Runtime/Engine/Private/Materials/MaterialInstance.cpp b/Engine/Source/Runtime/Engine/Private/Materials/MaterialInstance.cpp
--- a/Engine/Source/Runtime/Engine/Private/Materials/MaterialInstance.cpp
+++ b/Engine/Source/Runtime/Engine/Private/Materials/MaterialInstance.cpp
@@ -228,27 +228,14 @@ bool FMaterialInstanceResource::GetParameterValue(EMaterialParameterType Type, c
 {
 	checkSlow(IsInParallelRenderingThread());
 
-	static FName NameSubsurfaceProfile(TEXT("__SubsurfaceProfile"));
 	bool bResult = false;
 
 	// Check for hard-coded parameters
-	if (Type == EMaterialParameterType::Scalar && ParameterInfo.Name == NameSubsurfaceProfile)
+	if (Type == EMaterialParameterType::Scalar && ParameterInfo.Name == GetSubsurfaceProfileParameterName())
 	{
 		check(ParameterInfo.Association == EMaterialParameterAssociation::GlobalParameter);
 		const USubsurfaceProfile* MySubsurfaceProfileRT = GetSubsurfaceProfileRT();
-
-		int32 AllocationId = 0;
-		if (MySubsurfaceProfileRT)
-		{
-			// can be optimized (cached)
-			AllocationId = GSubsurfaceProfileTextureObject.FindAllocationId(MySubsurfaceProfileRT);
-		}
-		else
-		{
-			// no profile specified means we use the default one stored at [0] which is human skin
-			AllocationId = 0;
-		}
-		OutValue = AllocationId / 255.0f;
+		OutValue = GetSubsurfaceProfileId(MySubsurfaceProfileRT);
 		bResult = true;
 	}
 
diff --git a/Engine/Source/Runtime/Engine/Private/Materials/MaterialInterface.cpp b/Engine/Source/Runtime/Engine/Private/Materials/MaterialInterface.cpp
--- a/Engine/Source/Runtime/Engine/Private/Materials/MaterialInterface.cpp
+++ b/Engine/Source/Runtime/Engine/Private/Materials/MaterialInterface.cpp
@@ -988,10 +988,9 @@ void UMaterialInterface::UpdateMaterialRenderProxy(FMaterialRenderProxy& Proxy)
 	// for better performance we only update SubsurfaceProfileRT if the feature is used
 	if (UseSubsurfaceProfile(MaterialShadingModels))
 	{
-		FSubsurfaceProfileStruct Settings;
-
 		USubsurfaceProfile* LocalSubsurfaceProfile = GetSubsurfaceProfile_Internal();
 		
+		FSubsurfaceProfileStruct Settings;
 		if (LocalSubsurfaceProfile)
 		{
 			Settings = LocalSubsurfaceProfile->Settings;
@@ -1001,12 +1000,9 @@ void UMaterialInterface::UpdateMaterialRenderProxy(FMaterialRenderProxy& Proxy)
 		ENQUEUE_RENDER_COMMAND(UpdateMaterialRenderProxySubsurface)(
 			[Settings, LocalSubsurfaceProfile, InProxy](FRHICommandListImmediate& RHICmdList)
 			{
-				uint32 AllocationId = 0;
-
 				if (LocalSubsurfaceProfile)
 				{
-					AllocationId = GSubsurfaceProfileTextureObject.AddOrUpdateProfile(Settings, LocalSubsurfaceProfile);
-
+					const uint32 AllocationId = GSubsurfaceProfileTextureObject.AddOrUpdateProfile(Settings, LocalSubsurfaceProfile);
 					check(AllocationId >= 0 && AllocationId <= 255);
 				}
 				InProxy->SetSubsurfaceProfileRT(LocalSubsurfaceProfile);
diff --git a/Engine/Source/Runtime/Engine/Private/Rendering/SubsurfaceProfile.cpp b/Engine/Source/Runtime/Engine/Private/Rendering/SubsurfaceProfile.cpp
--- a/Engine/Source/Runtime/Engine/Private/Rendering/SubsurfaceProfile.cpp
+++ b/Engine/Source/Runtime/Engine/Private/Rendering/SubsurfaceProfile.cpp
@@ -141,11 +141,11 @@ public:
 	SHADER_USE_PARAMETER_STRUCT(FSSProfilePreIntegratedPS, FGlobalShader);
 
 	BEGIN_SHADER_PARAMETER_STRUCT(FParameters, )
-		SHADER_PARAMETER_TEXTURE(Texture2D, SSProfilesTexture)
-		SHADER_PARAMETER_SAMPLER(SamplerState, SSProfilesSampler)
-		SHADER_PARAMETER(FVector4f, SSProfilesTextureSizeAndInvSize)
-		SHADER_PARAMETER(FVector4f, SSProfilesPreIntegratedTextureSizeAndInvSize)
-		SHADER_PARAMETER(int32, SubsurfaceProfileInt)
+		SHADER_PARAMETER_TEXTURE(Texture2D, SourceSSProfilesTexture)
+		SHADER_PARAMETER_SAMPLER(SamplerState, SourceSSProfilesSampler)
+		SHADER_PARAMETER(FVector4f, SourceSSProfilesTextureSizeAndInvSize)
+		SHADER_PARAMETER(FVector4f, TargetSSProfilesPreIntegratedTextureSizeAndInvSize)
+		SHADER_PARAMETER(int32, SourceSubsurfaceProfileInt)
 		RENDER_TARGET_BINDING_SLOTS()
 	END_SHADER_PARAMETER_STRUCT()
 
@@ -193,23 +193,23 @@ IPooledRenderTarget* FSubsurfaceProfileTexture::GetSSProfilesPreIntegratedTextur
 			FSSProfilePreIntegratedPS::FParameters* PassParameters = GraphBuilder.AllocParameters<FSSProfilePreIntegratedPS::FParameters>();
 			PassParameters->RenderTargets[0] = FRenderTargetBinding(GraphBuilder.RegisterExternalTexture(GSSProfilesPreIntegratedTexture), ERenderTargetLoadAction::EClear, 0, i);
 
-			PassParameters->SSProfilesTexture = GetSubsurfaceProfileTextureWithFallback();
-			PassParameters->SSProfilesSampler = TStaticSamplerState<SF_Point, AM_Clamp, AM_Clamp, AM_Clamp>::GetRHI();
+			PassParameters->SourceSSProfilesTexture = GetSubsurfaceProfileTextureWithFallback();
+			PassParameters->SourceSSProfilesSampler = TStaticSamplerState<SF_Point, AM_Clamp, AM_Clamp, AM_Clamp>::GetRHI();
 
-			FIntVector SSProfilesTextureSize = PassParameters->SSProfilesTexture->GetSizeXYZ();
+			FIntVector SourceSSProfilesTextureSize = PassParameters->SourceSSProfilesTexture->GetSizeXYZ();
 
-			PassParameters->SSProfilesTextureSizeAndInvSize = FVector4(SSProfilesTextureSize.X, SSProfilesTextureSize.Y, 1.0f / SSProfilesTextureSize.X, 1.0f / SSProfilesTextureSize.Y);
+			PassParameters->SourceSSProfilesTextureSizeAndInvSize = FVector4(SourceSSProfilesTextureSize.X, SourceSSProfilesTextureSize.Y, 1.0f / SourceSSProfilesTextureSize.X, 1.0f / SourceSSProfilesTextureSize.Y);
 
-			PassParameters->SSProfilesPreIntegratedTextureSizeAndInvSize = FVector4(SSProfilesPreIntegratedTextureResolution, SSProfilesPreIntegratedTextureResolution, 1.0f / SSProfilesPreIntegratedTextureResolution, 1.0f / SSProfilesPreIntegratedTextureResolution);
+			PassParameters->TargetSSProfilesPreIntegratedTextureSizeAndInvSize = FVector4(SSProfilesPreIntegratedTextureResolution, SSProfilesPreIntegratedTextureResolution, 1.0f / SSProfilesPreIntegratedTextureResolution, 1.0f / SSProfilesPreIntegratedTextureResolution);
 
-			PassParameters->SubsurfaceProfileInt = i;
+			PassParameters->SourceSubsurfaceProfileInt = i;
 
 			FIntRect ViewRect = FIntRect(FIntPoint(0, 0), SSProfilesPreIntegratedTextureResolution);
 
 			const auto GlobalShaderMap = GetGlobalShaderMap(ERHIFeatureLevel::ES3_1);
 			TShaderMapRef<FSSProfilePreIntegratedPS> PixelShader(GlobalShaderMap);
 
-			FPixelShaderUtils::AddFullscreenPass(GraphBuilder, GlobalShaderMap, RDG_EVENT_NAME("SSProfilePreIntegrated"), PixelShader, PassParameters, ViewRect);
+			FPixelShaderUtils::AddFullscreenPass(GraphBuilder, GlobalShaderMap, RDG_EVENT_NAME("SSS::SSProfilePreIntegrated"), PixelShader, PassParameters, ViewRect);
 		}
 	}
 
@@ -644,6 +644,28 @@ void FSubsurfaceProfileTexture::Dump()
 #endif
 }
 
+FName GetSubsurfaceProfileParameterName()
+{
+	static FName NameSubsurfaceProfile(TEXT("__SubsurfaceProfile"));
+	return NameSubsurfaceProfile;
+}
+
+float GetSubsurfaceProfileId(const USubsurfaceProfile* In)
+{
+	int32 AllocationId = 0;
+	if (In)
+	{
+		// can be optimized (cached)
+		AllocationId = GSubsurfaceProfileTextureObject.FindAllocationId(In);
+	}
+	else
+	{
+		// no profile specified means we use the default one stored at [0] which is human skin
+		AllocationId = 0;
+	}
+	return AllocationId / 255.0f;
+}
+
 FRHITexture* GetSubsurfaceProfileTexture()
 {
 	return GSSProfiles ? GSSProfiles->GetShaderResourceRHI() : nullptr;
diff --git a/Engine/Source/Runtime/Engine/Private/SceneManagement.cpp b/Engine/Source/Runtime/Engine/Private/SceneManagement.cpp
--- a/Engine/Source/Runtime/Engine/Private/SceneManagement.cpp
+++ b/Engine/Source/Runtime/Engine/Private/SceneManagement.cpp
@@ -798,6 +798,16 @@ FViewUniformShaderParameters::FViewUniformShaderParameters()
 	ShadingEnergyClothSpecTexture = GBlackTextureWithSRV->TextureRHI;
 	ShadingEnergyDiffuseTexture = GBlackTextureWithSRV->TextureRHI;
 
+	// Subsurface profiles/pre-intregrated
+	SSProfilesTextureSizeAndInvSize = FVector4f(1.f,1.f,1.f,1.f);
+	SSProfilesTexture = GBlackTextureWithSRV->TextureRHI;
+	SSProfilesSampler = TStaticSamplerState<SF_Point, AM_Clamp, AM_Clamp, AM_Clamp>::GetRHI();;
+	SSProfilesTransmissionSampler = TStaticSamplerState<SF_Bilinear, AM_Clamp, AM_Clamp, AM_Clamp>::GetRHI();
+
+	SSProfilesPreIntegratedTextureSizeAndInvSize = FVector4f(1.f,1.f,1.f,1.f);
+	SSProfilesPreIntegratedTexture = GBlackTextureWithSRV->TextureRHI;
+	SSProfilesPreIntegratedSampler = TStaticSamplerState<SF_Bilinear, AM_Clamp, AM_Clamp, AM_Clamp>::GetRHI();
+
 	//this can be deleted once sm4 support is removed.
 	if (!PrimitiveSceneData)
 	{
diff --git a/Engine/Source/Runtime/Engine/Public/SceneView.h b/Engine/Source/Runtime/Engine/Public/SceneView.h
--- a/Engine/Source/Runtime/Engine/Public/SceneView.h
+++ b/Engine/Source/Runtime/Engine/Public/SceneView.h
@@ -674,7 +674,6 @@ enum ETranslucencyVolumeCascade
 	VIEW_UNIFORM_BUFFER_MEMBER(float, DepthOfFieldNearTransitionRegion) \
 	VIEW_UNIFORM_BUFFER_MEMBER(float, DepthOfFieldFarTransitionRegion) \
 	VIEW_UNIFORM_BUFFER_MEMBER(float, MotionBlurNormalizedToPixel) \
-	VIEW_UNIFORM_BUFFER_MEMBER(float, bSubsurfacePostprocessEnabled) \
 	VIEW_UNIFORM_BUFFER_MEMBER(float, GeneralPurposeTweak) \
 	VIEW_UNIFORM_BUFFER_MEMBER(float, GeneralPurposeTweak2) \
 	VIEW_UNIFORM_BUFFER_MEMBER_EX(float, DemosaicVposOffset, EShaderPrecisionModifier::Half) \
@@ -764,6 +763,9 @@ enum ETranslucencyVolumeCascade
 	VIEW_UNIFORM_BUFFER_MEMBER(uint32, EnableSkyLight) \
 	VIEW_UNIFORM_BUFFER_MEMBER(uint32, HairRenderInfoBits) \
 	VIEW_UNIFORM_BUFFER_MEMBER(uint32, HairComponents) \
+	VIEW_UNIFORM_BUFFER_MEMBER(float, bSubsurfacePostprocessEnabled) \
+	VIEW_UNIFORM_BUFFER_MEMBER(FVector4f, SSProfilesTextureSizeAndInvSize) \
+	VIEW_UNIFORM_BUFFER_MEMBER(FVector4f, SSProfilesPreIntegratedTextureSizeAndInvSize) \
 	VIEW_UNIFORM_BUFFER_MEMBER(FVector3f, PhysicsFieldClipmapCenter) \
 	VIEW_UNIFORM_BUFFER_MEMBER(float, PhysicsFieldClipmapDistance) \
 	VIEW_UNIFORM_BUFFER_MEMBER(int, PhysicsFieldClipmapResolution) \
@@ -842,7 +844,7 @@ BEGIN_GLOBAL_SHADER_PARAMETER_STRUCT_WITH_CONSTRUCTOR(FViewUniformShaderParamete
 	SHADER_PARAMETER_SRV(StructuredBuffer<float4>, InstancePayloadData)
 	SHADER_PARAMETER_SRV(StructuredBuffer<float4>, LightmapSceneData)
 	SHADER_PARAMETER_SRV(StructuredBuffer<float4>, SkyIrradianceEnvironmentMap)
-
+	// Atmosphere
 	SHADER_PARAMETER_TEXTURE(Texture2D, TransmittanceLutTexture)
 	SHADER_PARAMETER_SAMPLER(SamplerState, TransmittanceLutTextureSampler)
 	SHADER_PARAMETER_TEXTURE(Texture2D, SkyViewLutTexture)
@@ -851,9 +853,10 @@ BEGIN_GLOBAL_SHADER_PARAMETER_STRUCT_WITH_CONSTRUCTOR(FViewUniformShaderParamete
 	SHADER_PARAMETER_SAMPLER(SamplerState, DistantSkyLightLutTextureSampler)
 	SHADER_PARAMETER_TEXTURE(Texture3D, CameraAerialPerspectiveVolume)
 	SHADER_PARAMETER_SAMPLER(SamplerState, CameraAerialPerspectiveVolumeSampler)
+	// Hair
 	SHADER_PARAMETER_TEXTURE(Texture3D, HairScatteringLUTTexture)
 	SHADER_PARAMETER_SAMPLER(SamplerState, HairScatteringLUTSampler)
-
+	// Energy conservation
 	SHADER_PARAMETER(uint32, bShadingEnergyConservation)
 	SHADER_PARAMETER(uint32, bShadingEnergyPreservation)
 	SHADER_PARAMETER_TEXTURE(Texture2D<float2>, ShadingEnergyGGXSpecTexture)
@@ -861,7 +864,13 @@ BEGIN_GLOBAL_SHADER_PARAMETER_STRUCT_WITH_CONSTRUCTOR(FViewUniformShaderParamete
 	SHADER_PARAMETER_TEXTURE(Texture2D<float2>, ShadingEnergyClothSpecTexture)
 	SHADER_PARAMETER_TEXTURE(Texture2D<float>, ShadingEnergyDiffuseTexture)
 	SHADER_PARAMETER_SAMPLER(SamplerState, ShadingEnergySampler)
-
+	// SSS
+	SHADER_PARAMETER_TEXTURE(Texture2D, SSProfilesTexture)
+	SHADER_PARAMETER_SAMPLER(SamplerState, SSProfilesSampler)
+	SHADER_PARAMETER_SAMPLER(SamplerState, SSProfilesTransmissionSampler)
+	SHADER_PARAMETER_TEXTURE(Texture2DArray, SSProfilesPreIntegratedTexture)
+	SHADER_PARAMETER_SAMPLER(SamplerState, SSProfilesPreIntegratedSampler)
+	// Water
 	SHADER_PARAMETER_SRV(Buffer<float4>, WaterIndirection)
 	SHADER_PARAMETER_SRV(Buffer<float4>, WaterData)
 
diff --git a/Engine/Source/Runtime/Renderer/Private/BasePassRendering.cpp b/Engine/Source/Runtime/Renderer/Private/BasePassRendering.cpp
--- a/Engine/Source/Runtime/Renderer/Private/BasePassRendering.cpp
+++ b/Engine/Source/Runtime/Renderer/Private/BasePassRendering.cpp
@@ -526,8 +526,6 @@ void SetupSharedBasePassParameters(
 	SetupReflectionUniformParameters(View, SharedParameters.Reflection);
 	SetupPlanarReflectionUniformParameters(View, ReflectionSceneProxy, SharedParameters.PlanarReflection);
 
-	SharedParameters.SSProfilesTexture = GetSubsurfaceProfileTextureWithFallback();
-
 	// Skip base pass skylight if Lumen GI is enabled, as Lumen handles the skylight.
 	// Ideally we would choose a different shader permutation to skip skylight, but Lumen GI is only known per-view
 	SharedParameters.UseBasePassSkylight = View.FinalPostProcessSettings.DynamicGlobalIlluminationMethod == EDynamicGlobalIlluminationMethod::Lumen ? 0 : 1;
diff --git a/Engine/Source/Runtime/Renderer/Private/BasePassRendering.h b/Engine/Source/Runtime/Renderer/Private/BasePassRendering.h
--- a/Engine/Source/Runtime/Renderer/Private/BasePassRendering.h
+++ b/Engine/Source/Runtime/Renderer/Private/BasePassRendering.h
@@ -68,7 +68,6 @@ BEGIN_GLOBAL_SHADER_PARAMETER_STRUCT(FSharedBasePassUniformParameters,)
 	SHADER_PARAMETER_STRUCT(FPlanarReflectionUniformParameters, PlanarReflection) // Single global planar reflection for the forward pass.
 	SHADER_PARAMETER_STRUCT(FFogUniformParameters, Fog)
 	SHADER_PARAMETER_STRUCT(FFogUniformParameters, FogISR)
-	SHADER_PARAMETER_TEXTURE(Texture2D, SSProfilesTexture)
 	SHADER_PARAMETER(uint32, UseBasePassSkylight)
 END_GLOBAL_SHADER_PARAMETER_STRUCT()
 
diff --git a/Engine/Source/Runtime/Renderer/Private/ClusteredDeferredShadingPass.cpp b/Engine/Source/Runtime/Renderer/Private/ClusteredDeferredShadingPass.cpp
--- a/Engine/Source/Runtime/Renderer/Private/ClusteredDeferredShadingPass.cpp
+++ b/Engine/Source/Runtime/Renderer/Private/ClusteredDeferredShadingPass.cpp
@@ -113,9 +113,6 @@ class FClusteredShadingPS : public FGlobalShader
 		SHADER_PARAMETER_TEXTURE(Texture2D, LTCAmpTexture)
 		SHADER_PARAMETER_SAMPLER(SamplerState, LTCAmpSampler)
 
-		SHADER_PARAMETER_TEXTURE(Texture2D, SSProfilesTexture)
-		SHADER_PARAMETER_SAMPLER(SamplerState, TransmissionProfilesLinearSampler)
-
 		RENDER_TARGET_BINDING_SLOTS()
 	END_SHADER_PARAMETER_STRUCT()
 
@@ -168,8 +165,6 @@ static void InternalAddClusteredDeferredShadingPass(
 	PassParameters->LTCMatSampler = TStaticSamplerState<SF_Bilinear, AM_Clamp, AM_Clamp, AM_Clamp>::GetRHI();
 	PassParameters->LTCAmpTexture = GSystemTextures.LTCAmp->GetShaderResourceRHI();
 	PassParameters->LTCAmpSampler = TStaticSamplerState<SF_Bilinear, AM_Clamp, AM_Clamp, AM_Clamp>::GetRHI();
-	PassParameters->SSProfilesTexture = GetSubsurfaceProfileTexture();
-	PassParameters->TransmissionProfilesLinearSampler = TStaticSamplerState<SF_Bilinear, AM_Clamp, AM_Clamp, AM_Clamp>::GetRHI();
 
 	ShaderDrawDebug::SetParameters(GraphBuilder, View.ShaderDrawData, PassParameters->ShaderDrawParameters);
 	ShaderPrint::SetParameters(GraphBuilder, View, PassParameters->ShaderPrintUniformBuffer);
diff --git a/Engine/Source/Runtime/Renderer/Private/LightRendering.cpp b/Engine/Source/Runtime/Renderer/Private/LightRendering.cpp
--- a/Engine/Source/Runtime/Renderer/Private/LightRendering.cpp
+++ b/Engine/Source/Runtime/Renderer/Private/LightRendering.cpp
@@ -537,8 +537,6 @@ class FDeferredLightPS : public FGlobalShader
 		IESTextureSampler.Bind(Initializer.ParameterMap, TEXT("IESTextureSampler"));
 		LightingChannelsTexture.Bind(Initializer.ParameterMap, TEXT("LightingChannelsTexture"));
 		LightingChannelsSampler.Bind(Initializer.ParameterMap, TEXT("LightingChannelsSampler"));
-		TransmissionProfilesTexture.Bind(Initializer.ParameterMap, TEXT("SSProfilesTexture"));
-		TransmissionProfilesLinearSampler.Bind(Initializer.ParameterMap, TEXT("TransmissionProfilesLinearSampler"));
 
 		HairTransmittanceBuffer.Bind(Initializer.ParameterMap, TEXT("HairTransmittanceBuffer"));
 		HairTransmittanceBufferMaxCount.Bind(Initializer.ParameterMap, TEXT("HairTransmittanceBufferMaxCount"));
@@ -661,18 +659,6 @@ private:
 				);
 		}
 
-		if( TransmissionProfilesTexture.IsBound() )
-		{
-			FRHITexture* SubsurfaceTextureRHI = GetSubsurfaceProfileTextureWithFallback();
-
-			SetTextureParameter(RHICmdList,
-				ShaderRHI,
-				TransmissionProfilesTexture,
-				TransmissionProfilesLinearSampler,
-				TStaticSamplerState<SF_Bilinear, AM_Clamp, AM_Clamp, AM_Clamp>::GetRHI(),
-				SubsurfaceTextureRHI);
-		}
-
 		if (HairTransmittanceBuffer.IsBound())
 		{
 			const uint32 TransmittanceBufferMaxCount = RenderLightParams ? RenderLightParams->DeepShadow_TransmittanceMaskBufferMaxCount : 0;
@@ -763,8 +749,6 @@ private:
 	LAYOUT_FIELD(FShaderResourceParameter, IESTextureSampler);
 	LAYOUT_FIELD(FShaderResourceParameter, LightingChannelsTexture);
 	LAYOUT_FIELD(FShaderResourceParameter, LightingChannelsSampler);
-	LAYOUT_FIELD(FShaderResourceParameter, TransmissionProfilesTexture);
-	LAYOUT_FIELD(FShaderResourceParameter, TransmissionProfilesLinearSampler);
 
 	LAYOUT_FIELD(FShaderParameter, HairTransmittanceBufferMaxCount);
 	LAYOUT_FIELD(FShaderResourceParameter, HairTransmittanceBuffer);
diff --git a/Engine/Source/Runtime/Renderer/Private/MobileBasePassRendering.cpp b/Engine/Source/Runtime/Renderer/Private/MobileBasePassRendering.cpp
--- a/Engine/Source/Runtime/Renderer/Private/MobileBasePassRendering.cpp
+++ b/Engine/Source/Runtime/Renderer/Private/MobileBasePassRendering.cpp
@@ -221,16 +221,6 @@ void SetupMobileBasePassUniformParameters(
 		BasePassParameters.ScreenSpaceShadowMaskTexture = GSystemTextures.GetWhiteDummy(GraphBuilder);
 		BasePassParameters.ScreenSpaceShadowMaskSampler = TStaticSamplerState<SF_Point, AM_Clamp, AM_Clamp, AM_Clamp>::GetRHI();
 	}
-
-	BasePassParameters.SSProfilesTexture = GetSubsurfaceProfileTextureWithFallback();
-	BasePassParameters.SSProfilesSampler = TStaticSamplerState<SF_Point, AM_Clamp, AM_Clamp, AM_Clamp>::GetRHI();
-
-	FIntVector SSProfilesTextureSize = BasePassParameters.SSProfilesTexture->GetSizeXYZ();
-
-	BasePassParameters.SSProfilesTextureSizeAndInvSize = FVector4f(SSProfilesTextureSize.X, SSProfilesTextureSize.Y, 1.0f / SSProfilesTextureSize.X, 1.0f / SSProfilesTextureSize.Y);
-
-	BasePassParameters.SSProfilesPreIntegratedTexture = GetSSProfilesPreIntegratedTextureWithFallback();
-	BasePassParameters.SSProfilesPreIntegratedSampler = TStaticSamplerState<SF_Bilinear, AM_Clamp, AM_Clamp, AM_Clamp>::GetRHI();
 }
 
 TRDGUniformBufferRef<FMobileBasePassUniformParameters> CreateMobileBasePassUniformBuffer(
diff --git a/Engine/Source/Runtime/Renderer/Private/MobileBasePassRendering.h b/Engine/Source/Runtime/Renderer/Private/MobileBasePassRendering.h
--- a/Engine/Source/Runtime/Renderer/Private/MobileBasePassRendering.h
+++ b/Engine/Source/Runtime/Renderer/Private/MobileBasePassRendering.h
@@ -33,7 +33,6 @@ struct FMobileBasePassTextures
 
 BEGIN_GLOBAL_SHADER_PARAMETER_STRUCT(FMobileBasePassUniformParameters, )
 	SHADER_PARAMETER(float, AmbientOcclusionStaticFraction)
-	SHADER_PARAMETER(FVector4f, SSProfilesTextureSizeAndInvSize)
 	SHADER_PARAMETER_STRUCT(FFogUniformParameters, Fog)
 	SHADER_PARAMETER_STRUCT(FPlanarReflectionUniformParameters, PlanarReflection) // Single global planar reflection for the forward pass.
 	SHADER_PARAMETER_STRUCT(FMobileSceneTextureUniformParameters, SceneTextures)
@@ -45,10 +44,6 @@ BEGIN_GLOBAL_SHADER_PARAMETER_STRUCT(FMobileBasePassUniformParameters, )
 	SHADER_PARAMETER_SAMPLER(SamplerState, AmbientOcclusionSampler)
 	SHADER_PARAMETER_RDG_TEXTURE(Texture2D, ScreenSpaceShadowMaskTexture)
 	SHADER_PARAMETER_SAMPLER(SamplerState, ScreenSpaceShadowMaskSampler)
-	SHADER_PARAMETER_TEXTURE(Texture2D, SSProfilesTexture)
-	SHADER_PARAMETER_SAMPLER(SamplerState, SSProfilesSampler)
-	SHADER_PARAMETER_TEXTURE(Texture2DArray, SSProfilesPreIntegratedTexture)
-	SHADER_PARAMETER_SAMPLER(SamplerState, SSProfilesPreIntegratedSampler)
 END_GLOBAL_SHADER_PARAMETER_STRUCT()
 
 enum class EMobileBasePass
diff --git a/Engine/Source/Runtime/Renderer/Private/PathTracing.cpp b/Engine/Source/Runtime/Renderer/Private/PathTracing.cpp
--- a/Engine/Source/Runtime/Renderer/Private/PathTracing.cpp
+++ b/Engine/Source/Runtime/Renderer/Private/PathTracing.cpp
@@ -528,8 +528,6 @@ class FPathTracingRG : public FGlobalShader
 		// Rect lights
 		SHADER_PARAMETER_TEXTURE_ARRAY(Texture2D, RectLightTexture, [PATHTRACER_MAX_RECT_TEXTURES])
 		SHADER_PARAMETER_SAMPLER(SamplerState, RectLightSampler) // Shared sampler for all rectlights
-		// Subsurface data
-		SHADER_PARAMETER_TEXTURE(Texture2D, SSProfilesTexture)
 		// Used by multi-GPU rendering
 		SHADER_PARAMETER(FIntVector, TileOffset)
 
@@ -1539,8 +1537,6 @@ void FDeferredShadingSceneRenderer::RenderPathTracing(
 			PassParameters->AlbedoTexture   = GraphBuilder.CreateUAV(AlbedoTexture);
 			PassParameters->NormalTexture   = GraphBuilder.CreateUAV(NormalTexture);
 
-			PassParameters->SSProfilesTexture = GetSubsurfaceProfileTexture();
-
 			// TODO: in multi-gpu case, split image into tiles
 			PassParameters->TileOffset.X = 0;
 			PassParameters->TileOffset.Y = 0;
diff --git a/Engine/Source/Runtime/Renderer/Private/PostProcess/PostProcessSubsurface.cpp b/Engine/Source/Runtime/Renderer/Private/PostProcess/PostProcessSubsurface.cpp
--- a/Engine/Source/Runtime/Renderer/Private/PostProcess/PostProcessSubsurface.cpp
+++ b/Engine/Source/Runtime/Renderer/Private/PostProcess/PostProcessSubsurface.cpp
@@ -232,7 +232,6 @@ BEGIN_SHADER_PARAMETER_STRUCT(FSubsurfaceParameters, )
 	SHADER_PARAMETER(FVector4f, SubsurfaceParams)
 	SHADER_PARAMETER_RDG_UNIFORM_BUFFER(FSceneTextureUniformParameters, SceneTextures)
 	SHADER_PARAMETER_STRUCT_REF(FViewUniformShaderParameters, ViewUniformBuffer)
-	SHADER_PARAMETER_TEXTURE(Texture2D, SSProfilesTexture)
 END_SHADER_PARAMETER_STRUCT()
 
 FSubsurfaceParameters GetSubsurfaceCommonParameters(FRDGBuilder& GraphBuilder, const FViewInfo& View, TRDGUniformBufferRef<FSceneTextureUniformParameters> SceneTextures)
@@ -248,7 +247,6 @@ FSubsurfaceParameters GetSubsurfaceCommonParameters(FRDGBuilder& GraphBuilder, c
 	Parameters.SubsurfaceParams = FVector4f(SSSScaleX, SSSScaleZ, SSSOverrideNumSamples, MinGenerateMipsTileCount);
 	Parameters.ViewUniformBuffer = View.ViewUniformBuffer;
 	Parameters.SceneTextures = SceneTextures;
-	Parameters.SSProfilesTexture = GetSubsurfaceProfileTextureWithFallback();
 	return Parameters;
 }
 
diff --git a/Engine/Source/Runtime/Renderer/Private/RayTracing/RayTracingGlobalIllumination.cpp b/Engine/Source/Runtime/Renderer/Private/RayTracing/RayTracingGlobalIllumination.cpp
--- a/Engine/Source/Runtime/Renderer/Private/RayTracing/RayTracingGlobalIllumination.cpp
+++ b/Engine/Source/Runtime/Renderer/Private/RayTracing/RayTracingGlobalIllumination.cpp
@@ -455,8 +455,6 @@ class FGlobalIlluminationRGS : public FGlobalShader
 		SHADER_PARAMETER(uint32, SceneLightCount)
 		SHADER_PARAMETER_STRUCT_INCLUDE(FPathTracingSkylight, SkylightParameters)
 		SHADER_PARAMETER_STRUCT_INCLUDE(FSceneTextureParameters, SceneTextures)
-		SHADER_PARAMETER_TEXTURE(Texture2D, SSProfilesTexture)
-		SHADER_PARAMETER_SAMPLER(SamplerState, TransmissionProfilesLinearSampler)
 	END_SHADER_PARAMETER_STRUCT()
 };
 
@@ -523,8 +521,6 @@ class FRayTracingGlobalIlluminationCreateGatherPointsRGS : public FGlobalShader
 
 		// Shading data
 		SHADER_PARAMETER_STRUCT_INCLUDE(FSceneTextureParameters, SceneTextures)
-		SHADER_PARAMETER_TEXTURE(Texture2D, SSProfilesTexture)
-		SHADER_PARAMETER_SAMPLER(SamplerState, TransmissionProfilesLinearSampler)
 
 		SHADER_PARAMETER(FIntPoint, GatherPointsResolution)
 		SHADER_PARAMETER(FIntPoint, TileAlignedResolution)
@@ -590,8 +586,6 @@ class FRayTracingGlobalIlluminationCreateGatherPointsTraceRGS : public FGlobalSh
 
 		// Shading data
 		SHADER_PARAMETER_STRUCT_INCLUDE(FSceneTextureParameters, SceneTextures)
-		SHADER_PARAMETER_TEXTURE(Texture2D, SSProfilesTexture)
-		SHADER_PARAMETER_SAMPLER(SamplerState, TransmissionProfilesLinearSampler)
 
 		SHADER_PARAMETER(FIntPoint, GatherPointsResolution)
 		SHADER_PARAMETER(FIntPoint, TileAlignedResolution)
@@ -645,8 +639,6 @@ class FRayTracingGlobalIlluminationFinalGatherRGS : public FGlobalShader
 
 		// Shading data
 		SHADER_PARAMETER_STRUCT_INCLUDE(FSceneTextureParameters, SceneTextures)
-		SHADER_PARAMETER_TEXTURE(Texture2D, SSProfilesTexture)
-		SHADER_PARAMETER_SAMPLER(SamplerState, TransmissionProfilesLinearSampler)
 
 		// Gather points
 		SHADER_PARAMETER_RDG_BUFFER_SRV(StructuredBuffer<GatherPoints>, GatherPointsBuffer)
@@ -845,8 +837,6 @@ void CopyGatherPassParameters(
 	NewParameters->SkylightParameters = PassParameters.SkylightParameters;
 
 	NewParameters->SceneTextures = PassParameters.SceneTextures;
-	NewParameters->SSProfilesTexture = PassParameters.SSProfilesTexture;
-	NewParameters->TransmissionProfilesLinearSampler = PassParameters.TransmissionProfilesLinearSampler;
 
 	NewParameters->GatherPointsResolution = PassParameters.GatherPointsResolution;
 	NewParameters->TileAlignedResolution = PassParameters.TileAlignedResolution;
@@ -886,8 +876,6 @@ void CopyGatherPassParameters(
 	NewParameters->SkylightParameters = PassParameters.SkylightParameters;
 
 	NewParameters->SceneTextures = PassParameters.SceneTextures;
-	NewParameters->SSProfilesTexture = PassParameters.SSProfilesTexture;
-	NewParameters->TransmissionProfilesLinearSampler = PassParameters.TransmissionProfilesLinearSampler;
 
 	NewParameters->GatherPointsResolution = PassParameters.GatherPointsResolution;
 	NewParameters->TileAlignedResolution = PassParameters.TileAlignedResolution;
@@ -957,10 +945,6 @@ void FDeferredShadingSceneRenderer::RayTracingGlobalIlluminationCreateGatherPoin
 	SetupLightParameters(Scene, View, GraphBuilder, &PassParameters->SceneLights, &PassParameters->SceneLightCount, &PassParameters->SkylightParameters);
 	PassParameters->SceneTextures = SceneTextures;
 
-	// Shading data
-	PassParameters->SSProfilesTexture = GetSubsurfaceProfileTextureWithFallback();
-	PassParameters->TransmissionProfilesLinearSampler = TStaticSamplerState<SF_Bilinear, AM_Clamp, AM_Clamp, AM_Clamp>::GetRHI();
-
 	// Output
 	FIntPoint DispatchResolution = FIntPoint::DivideAndRoundUp(View.ViewRect.Size(), UpscaleFactor);
 	FIntVector LocalGatherPointsResolution(DispatchResolution.X, DispatchResolution.Y, GatherSamples);
@@ -1177,8 +1161,6 @@ void FDeferredShadingSceneRenderer::RenderRayTracingGlobalIlluminationFinalGathe
 
 	// Shading data
 	PassParameters->SceneTextures = SceneTextures;
-	PassParameters->SSProfilesTexture = GetSubsurfaceProfileTextureWithFallback();
-	PassParameters->TransmissionProfilesLinearSampler = TStaticSamplerState<SF_Bilinear, AM_Clamp, AM_Clamp, AM_Clamp>::GetRHI();
 
 	// Gather points
 	PassParameters->GatherPointsResolution = FIntPoint(SceneViewState->GatherPointsResolution.X, SceneViewState->GatherPointsResolution.Y);
@@ -1265,8 +1247,6 @@ void FDeferredShadingSceneRenderer::RenderRayTracingGlobalIlluminationBruteForce
 	PassParameters->ViewUniformBuffer = View.ViewUniformBuffer;
 	SetupLightParameters(Scene, View, GraphBuilder, &PassParameters->SceneLights, &PassParameters->SceneLightCount, &PassParameters->SkylightParameters);
 	PassParameters->SceneTextures = SceneTextures;
-	PassParameters->SSProfilesTexture = GetSubsurfaceProfileTextureWithFallback();
-	PassParameters->TransmissionProfilesLinearSampler = TStaticSamplerState<SF_Bilinear, AM_Clamp, AM_Clamp, AM_Clamp>::GetRHI();
 	PassParameters->RWGlobalIlluminationUAV = GraphBuilder.CreateUAV(OutDenoiserInputs->Color);
 	PassParameters->RWGlobalIlluminationRayDistanceUAV = GraphBuilder.CreateUAV(OutDenoiserInputs->RayHitDistance);
 	PassParameters->RenderTileOffsetX = 0;
diff --git a/Engine/Source/Runtime/Renderer/Private/RayTracing/RaytracingSkylight.cpp b/Engine/Source/Runtime/Renderer/Private/RayTracing/RaytracingSkylight.cpp
--- a/Engine/Source/Runtime/Renderer/Private/RayTracing/RaytracingSkylight.cpp
+++ b/Engine/Source/Runtime/Renderer/Private/RayTracing/RaytracingSkylight.cpp
@@ -246,8 +246,6 @@ class FRayTracingSkyLightRGS : public FGlobalShader
 		
 		SHADER_PARAMETER_STRUCT_INCLUDE(FSkyLightVisibilityRaysData, SkyLightVisibilityRaysData)
 		SHADER_PARAMETER_STRUCT_INCLUDE(FSceneTextureParameters, SceneTextures)
-		SHADER_PARAMETER_TEXTURE(Texture2D, SSProfilesTexture)
-		SHADER_PARAMETER_SAMPLER(SamplerState, TransmissionProfilesLinearSampler)
 	END_SHADER_PARAMETER_STRUCT()
 };
 
@@ -421,8 +419,6 @@ void FDeferredShadingSceneRenderer::RenderRayTracingSkyLight(
 	// Fill Scene Texture parameters
 	FSceneTextureParameters SceneTextures = GetSceneTextureParameters(GraphBuilder);
 
-	FRHITexture* SubsurfaceProfileTexture = GetSubsurfaceProfileTextureWithFallback();
-
 	int32 ViewIndex = 0;
 	for (FViewInfo& View : Views)
 	{
@@ -440,8 +436,6 @@ void FDeferredShadingSceneRenderer::RenderRayTracingSkyLight(
 		{
 			PassParameters->SkyLightVisibilityRaysData.SkyLightVisibilityRays = GraphBuilder.CreateSRV(SkyLightVisibilityRaysBuffer, EPixelFormat::PF_R32_UINT);
 		}
-		PassParameters->SSProfilesTexture = SubsurfaceProfileTexture;
-		PassParameters->TransmissionProfilesLinearSampler = TStaticSamplerState<SF_Bilinear, AM_Clamp, AM_Clamp, AM_Clamp>::GetRHI();
 		PassParameters->SceneTextures = SceneTextures;
 		PassParameters->UpscaleFactor = UpscaleFactor;
 
diff --git a/Engine/Source/Runtime/Renderer/Private/SceneRendering.cpp b/Engine/Source/Runtime/Renderer/Private/SceneRendering.cpp
--- a/Engine/Source/Runtime/Renderer/Private/SceneRendering.cpp
+++ b/Engine/Source/Runtime/Renderer/Private/SceneRendering.cpp
@@ -1629,7 +1629,29 @@ void FViewInfo::SetupUniformBufferParameters(
 	ViewUniformShaderParameters.DepthOfFieldScale = FinalPostProcessSettings.DepthOfFieldScale;
 	ViewUniformShaderParameters.DepthOfFieldFocalLength = 50.0f;
 
-	ViewUniformShaderParameters.bSubsurfacePostprocessEnabled = IsSubsurfaceEnabled() ? 1.0f : 0.0f;
+	// Subsurface
+	{
+		ViewUniformShaderParameters.bSubsurfacePostprocessEnabled = IsSubsurfaceEnabled() ? 1.0f : 0.0f;
+
+		// Profiles
+		{
+			FRHITexture* Texture = GetSubsurfaceProfileTextureWithFallback();
+			FIntVector TextureSize = Texture->GetSizeXYZ();
+			ViewUniformShaderParameters.SSProfilesTextureSizeAndInvSize = FVector4f(TextureSize.X, TextureSize.Y, 1.0f / TextureSize.X, 1.0f / TextureSize.Y);
+			ViewUniformShaderParameters.SSProfilesTexture = Texture;
+			ViewUniformShaderParameters.SSProfilesSampler = TStaticSamplerState<SF_Point, AM_Clamp, AM_Clamp, AM_Clamp>::GetRHI();
+			ViewUniformShaderParameters.SSProfilesTransmissionSampler = TStaticSamplerState<SF_Bilinear, AM_Clamp, AM_Clamp, AM_Clamp>::GetRHI();
+		}
+
+		// Pre-integrated profiles
+		{
+			FRHITexture* Texture = GetSSProfilesPreIntegratedTextureWithFallback();
+			FIntVector TextureSize = Texture->GetSizeXYZ();
+			ViewUniformShaderParameters.SSProfilesPreIntegratedTextureSizeAndInvSize = FVector4f(TextureSize.X, TextureSize.Y, 1.0f / TextureSize.X, 1.0f / TextureSize.Y);
+			ViewUniformShaderParameters.SSProfilesPreIntegratedTexture = Texture;
+			ViewUniformShaderParameters.SSProfilesPreIntegratedSampler = TStaticSamplerState<SF_Bilinear, AM_Clamp, AM_Clamp, AM_Clamp>::GetRHI();
+		}
+	}
 
 	{
 		// This is the CVar default
diff --git a/Engine/Source/Runtime/Renderer/Private/ShadingEnergyConservation.cpp b/Engine/Source/Runtime/Renderer/Private/ShadingEnergyConservation.cpp
--- a/Engine/Source/Runtime/Renderer/Private/ShadingEnergyConservation.cpp
+++ b/Engine/Source/Runtime/Renderer/Private/ShadingEnergyConservation.cpp
@@ -68,8 +68,6 @@ class FShadingFurnaceTestPassPS : public FGlobalShader
 		SHADER_PARAMETER_STRUCT_REF(FViewUniformShaderParameters, ViewUniformBuffer)
 		SHADER_PARAMETER_RDG_UNIFORM_BUFFER(FStrataGlobalUniformParameters, Strata)
 		SHADER_PARAMETER_RDG_UNIFORM_BUFFER(FSceneTextureUniformParameters, SceneTexturesStruct)
-		SHADER_PARAMETER_TEXTURE(Texture2D, SSProfilesTexture)
-		SHADER_PARAMETER_SAMPLER(SamplerState, TransmissionProfilesLinearSampler)
 		SHADER_PARAMETER(uint32, NumSamplesPerSet)
 		SHADER_PARAMETER_STRUCT_INCLUDE(ShaderPrint::FShaderParameters, ShaderPrintUniformBuffer)
 		SHADER_PARAMETER_STRUCT_INCLUDE(ShaderDrawDebug::FShaderParameters, ShaderDrawUniformBuffer)
@@ -108,8 +106,6 @@ static void AddShadingFurnacePass(
 	Parameters->SceneTexturesStruct				= SceneTexturesUniformBuffer;
 	Parameters->NumSamplesPerSet				= FMath::Clamp(CVarShadingFurnaceTest_SampleCount.GetValueOnAnyThread(), 16, 2048);
 	Parameters->RenderTargets[0]				= FRenderTargetBinding(OutTexture, ERenderTargetLoadAction::ELoad);
-	Parameters->SSProfilesTexture				= GetSubsurfaceProfileTextureWithFallback();
-	Parameters->TransmissionProfilesLinearSampler = TStaticSamplerState<SF_Bilinear, AM_Clamp, AM_Clamp, AM_Clamp>::GetRHI();
 	if (Strata::IsStrataEnabled())
 	{
 		Parameters->Strata = Strata::BindStrataGlobalUniformParameters(View.StrataSceneData);
diff --git a/Engine/Source/Runtime/Renderer/Private/ShadowRendering.h b/Engine/Source/Runtime/Renderer/Private/ShadowRendering.h
--- a/Engine/Source/Runtime/Renderer/Private/ShadowRendering.h
+++ b/Engine/Source/Runtime/Renderer/Private/ShadowRendering.h
@@ -1126,7 +1126,6 @@ public:
 		ProjectionParameters.Bind(Initializer);
 		ShadowFadeFraction.Bind(Initializer.ParameterMap,TEXT("ShadowFadeFraction"));
 		ShadowSharpen.Bind(Initializer.ParameterMap,TEXT("ShadowSharpen"));
-		TransmissionProfilesTexture.Bind(Initializer.ParameterMap, TEXT("SSProfilesTexture"));
 		LightPosition.Bind(Initializer.ParameterMap, TEXT("LightPositionAndInvRadius"));
 
 	}
@@ -1188,8 +1187,6 @@ public:
 		{
 			Scene = View.Family->Scene->GetRenderScene();
 		}
-
-		SetTextureParameter(RHICmdList, ShaderRHI, TransmissionProfilesTexture, GetSubsurfaceProfileTextureWithFallback());
 	}
 
 protected:
@@ -1197,7 +1194,6 @@ protected:
 	LAYOUT_FIELD(FShaderParameter, ShadowFadeFraction);
 	LAYOUT_FIELD(FShaderParameter, ShadowSharpen);
 	LAYOUT_FIELD(FShaderParameter, LightPosition);
-	LAYOUT_FIELD(FShaderResourceParameter, TransmissionProfilesTexture);
 };
 
 /** Pixel shader to project modulated shadows onto the scene. */
@@ -1465,7 +1461,6 @@ public:
 		ShadowSharpen.Bind(Initializer.ParameterMap,TEXT("ShadowSharpen"));
 		PointLightDepthBias.Bind(Initializer.ParameterMap,TEXT("PointLightDepthBias"));
 		PointLightProjParameters.Bind(Initializer.ParameterMap, TEXT("PointLightProjParameters"));
-		TransmissionProfilesTexture.Bind(Initializer.ParameterMap, TEXT("SSProfilesTexture"));
 	}
 
 	static void ModifyCompilationEnvironment(const FGlobalShaderPermutationParameters& Parameters, FShaderCompilerEnvironment& OutEnvironment)
@@ -1519,8 +1514,6 @@ public:
 			SetUniformBufferParameter(RHICmdList, ShaderRHI, StrataGlobalParameters, StrataUniformBuffer->GetRHIRef());
 		}
 
-		SetTextureParameter(RHICmdList, ShaderRHI, TransmissionProfilesTexture, GetSubsurfaceProfileTextureWithFallback());
-
 		FScene* Scene = nullptr;
 
 		if (View.Family->Scene)
@@ -1546,7 +1539,6 @@ private:
 	LAYOUT_FIELD(FShaderParameter, ShadowSharpen);
 	LAYOUT_FIELD(FShaderParameter, PointLightDepthBias);
 	LAYOUT_FIELD(FShaderParameter, PointLightProjParameters);
-	LAYOUT_FIELD(FShaderResourceParameter, TransmissionProfilesTexture);
 	LAYOUT_FIELD(FShaderUniformBufferParameter, HairStrandsParameters);
 	LAYOUT_FIELD(FShaderUniformBufferParameter, StrataGlobalParameters);
 };
diff --git a/Engine/Source/Runtime/Renderer/Private/TiledDeferredLightRendering.cpp b/Engine/Source/Runtime/Renderer/Private/TiledDeferredLightRendering.cpp
--- a/Engine/Source/Runtime/Renderer/Private/TiledDeferredLightRendering.cpp
+++ b/Engine/Source/Runtime/Renderer/Private/TiledDeferredLightRendering.cpp
@@ -109,8 +109,6 @@ public:
 		LTCMatSampler.Bind(Initializer.ParameterMap, TEXT("LTCMatSampler"));
 		LTCAmpTexture.Bind(Initializer.ParameterMap, TEXT("LTCAmpTexture"));
 		LTCAmpSampler.Bind(Initializer.ParameterMap, TEXT("LTCAmpSampler"));
-		TransmissionProfilesTexture.Bind(Initializer.ParameterMap, TEXT("SSProfilesTexture"));
-		TransmissionProfilesLinearSampler.Bind(Initializer.ParameterMap, TEXT("TransmissionProfilesLinearSampler"));
 	}
 
 	FTiledDeferredLightingCS()
@@ -160,15 +158,6 @@ public:
 
 		const int32 NumLightsToRenderInSortedLights = TiledDeferredLightsEnd - TiledDeferredLightsStart;
 
-		if (TransmissionProfilesTexture.IsBound())
-		{
-			SetTextureParameter(RHICmdList,
-				ShaderRHI,
-				TransmissionProfilesTexture,
-				TransmissionProfilesLinearSampler,
-				TStaticSamplerState<SF_Bilinear, AM_Clamp, AM_Clamp, AM_Clamp>::GetRHI(),
-				GetSubsurfaceProfileTextureWithFallback());
-		}
 
 		static const auto AllowStaticLightingVar = IConsoleManager::Get().FindTConsoleVariableDataInt(TEXT("r.AllowStaticLighting"));
 		const bool bAllowStaticLighting = (!AllowStaticLightingVar || AllowStaticLightingVar->GetValueOnRenderThread() != 0);
@@ -268,8 +257,6 @@ private:
 	LAYOUT_FIELD(FShaderResourceParameter, LTCMatSampler);
 	LAYOUT_FIELD(FShaderResourceParameter, LTCAmpTexture);
 	LAYOUT_FIELD(FShaderResourceParameter, LTCAmpSampler);
-	LAYOUT_FIELD(FShaderResourceParameter, TransmissionProfilesTexture);
-	LAYOUT_FIELD(FShaderResourceParameter, TransmissionProfilesLinearSampler);
 };
 
 // #define avoids a lot of code duplication
