Enable velocity rendering when visualize motion blur flag is on

diff --git a/Engine/Source/Runtime/Renderer/Private/VelocityRendering.cpp b/Engine/Source/Runtime/Renderer/Private/VelocityRendering.cpp
--- a/Engine/Source/Runtime/Renderer/Private/VelocityRendering.cpp
+++ b/Engine/Source/Runtime/Renderer/Private/VelocityRendering.cpp
@@ -170,6 +170,7 @@ bool FDeferredShadingSceneRenderer::ShouldRenderVelocities() const
 
 		bool bTemporalAA = IsTemporalAccumulationBasedMethod(View.AntiAliasingMethod) && !View.bCameraCut;
 		bool bMotionBlur = IsMotionBlurEnabled(View);
+		bool bVisualizeMotionblur = View.Family->EngineShowFlags.VisualizeMotionBlur;
 		bool bDistanceFieldAO = ShouldPrepareForDistanceFieldAO();
 
 		bool bSSRTemporal = ScreenSpaceRayTracing::ShouldRenderScreenSpaceReflections(View) && ScreenSpaceRayTracing::IsSSRTemporalPassRequired(View);
@@ -182,7 +183,7 @@ bool FDeferredShadingSceneRenderer::ShouldRenderVelocities() const
 		bool bSSGI = ViewPipelineState.DiffuseIndirectMethod == EDiffuseIndirectMethod::SSGI;
 		bool bLumen = ViewPipelineState.DiffuseIndirectMethod == EDiffuseIndirectMethod::Lumen || ViewPipelineState.ReflectionsMethod == EReflectionsMethod::Lumen;
 		
-		bNeedsVelocity |= bMotionBlur || bTemporalAA || bDistanceFieldAO || bSSRTemporal || bDenoise || bSSGI || bLumen;
+		bNeedsVelocity |= bVisualizeMotionblur || bMotionBlur || bTemporalAA || bDistanceFieldAO || bSSRTemporal || bDenoise || bSSGI || bLumen;
 	}
 
 	return bNeedsVelocity;
