Fixed for cloud to support being rendered in volumetric render target + reconstruction when using scene capture component. This is important for performance.
https://horde.devtools.epicgames.com/job/61a757bb1d2787d8395c7323

diff --git a/Engine/Source/Runtime/Renderer/Private/VolumetricRenderTarget.cpp b/Engine/Source/Runtime/Renderer/Private/VolumetricRenderTarget.cpp
--- a/Engine/Source/Runtime/Renderer/Private/VolumetricRenderTarget.cpp
+++ b/Engine/Source/Runtime/Renderer/Private/VolumetricRenderTarget.cpp
@@ -65,7 +65,7 @@ static bool ShouldPipelineCompileVolumetricRenderTargetShaders(EShaderPlatform S
 bool ShouldViewRenderVolumetricCloudRenderTarget(const FViewInfo& ViewInfo)
 {
 	return CVarVolumetricRenderTarget.GetValueOnRenderThread() && ShouldPipelineCompileVolumetricRenderTargetShaders(ViewInfo.GetShaderPlatform())
-		&& (ViewInfo.ViewState != nullptr) && !(ViewInfo.bIsReflectionCapture || ViewInfo.bIsSceneCapture);
+		&& (ViewInfo.ViewState != nullptr) && !ViewInfo.bIsReflectionCapture;
 }
 
 bool IsVolumetricRenderTargetEnabled()
