Fixed skylight real time capture when there is a SceneCapture component in the scene.
https://horde.devtools.epicgames.com/job/61a7585a6c7d8a7295fcd5cc

diff --git a/Engine/Source/Runtime/Renderer/Private/ReflectionEnvironmentRealTimeCapture.cpp b/Engine/Source/Runtime/Renderer/Private/ReflectionEnvironmentRealTimeCapture.cpp
--- a/Engine/Source/Runtime/Renderer/Private/ReflectionEnvironmentRealTimeCapture.cpp
+++ b/Engine/Source/Runtime/Renderer/Private/ReflectionEnvironmentRealTimeCapture.cpp
@@ -288,8 +288,9 @@ void FScene::AllocateAndCaptureFrameSkyEnvMap(
 {
 	check(SkyLight && SkyLight->bRealTimeCaptureEnabled && !SkyLight->bHasStaticLighting);
 
-	// Ignore viewfamilies without the Atmosphere showflag enabled as the sky capture may fail otherwise.
-	if (!MainView.Family->EngineShowFlags.Atmosphere)
+	// Ignore viewfamilies without the Atmosphere showflag enabled as the sky capture may fail otherwise 
+	// as well as all views being "scene captures" which cannot be used to update the sky light data.
+	if (MainView.bIsSceneCapture || !MainView.Family->EngineShowFlags.Atmosphere)
 	{
 		return;
 	}
