Fix on screen warning delegate registration race condition crash.

diff --git a/Engine/Source/Runtime/Renderer/Private/Renderer.cpp b/Engine/Source/Runtime/Renderer/Private/Renderer.cpp
--- a/Engine/Source/Runtime/Renderer/Private/Renderer.cpp
+++ b/Engine/Source/Runtime/Renderer/Private/Renderer.cpp
@@ -35,6 +35,7 @@
 #include "VT/VirtualTextureSystem.h"
 #include "PostProcess/TemporalAA.h"
 #include "CanvasRender.h"
+#include "RendererOnScreenNotification.h"
 DEFINE_LOG_CATEGORY(LogRenderer);
 
 IMPLEMENT_MODULE(FRendererModule, Renderer);
@@ -56,6 +57,7 @@ void FRendererModule::StartupModule()
 	GScreenSpaceDenoiser = IScreenSpaceDenoiser::GetDefaultDenoiser();
 	GTemporalUpscaler = ITemporalUpscaler::GetDefaultTemporalUpscaler();
 
+	FRendererOnScreenNotification::Get();
 	FVirtualTextureSystem::Initialize();
 
 	StopRenderingThreadDelegate = RegisterStopRenderingThreadDelegate(FStopRenderingThreadDelegate::CreateLambda([this]
@@ -73,6 +75,7 @@ void FRendererModule::ShutdownModule()
 	UnregisterStopRenderingThreadDelegate(StopRenderingThreadDelegate);
 
 	FVirtualTextureSystem::Shutdown();
+	FRendererOnScreenNotification::TearDown();
 
 	// Free up the memory of the default denoiser. Responsibility of the plugin to free up theirs.
 	delete IScreenSpaceDenoiser::GetDefaultDenoiser();
diff --git a/Engine/Source/Runtime/Renderer/Private/RendererOnScreenNotification.h b/Engine/Source/Runtime/Renderer/Private/RendererOnScreenNotification.h
new file mode 100644
--- /dev/null
+++ b/Engine/Source/Runtime/Renderer/Private/RendererOnScreenNotification.h
@@ -0,0 +1,78 @@
+// Copyright Epic Games, Inc. All Rights Reserved.
+
+#pragma once
+
+#include "Misc/CoreDelegates.h"
+#include "Misc/LazySingleton.h"
+
+/** 
+ * Class to wrap FCoreDelegates::FGetOnScreenMessagesDelegate for access from any thread. 
+ * This avoids race conditions in registration/unregistration that would happen if using FCoreDelegates::FGetOnScreenMessagesDelegate directly from render thread.
+ * Note that the Broadcast() still happens on game thread, so care needs to be taken with how data is accessed there.
+ * If that becomes an issue we could change so that the proxy delegate Broadcasts on render thread and buffers to game thread.
+ */
+class FRendererOnScreenNotification
+{
+public:
+	/** 
+	 * Create or get singleton instance. 
+	 * First call should be on game thread. After that any thread will do.
+	 */
+	static FRendererOnScreenNotification& Get()
+	{
+		return TLazySingleton<FRendererOnScreenNotification>::Get();
+	}
+
+	/** 
+	 * Tear down singleton instance. 
+	 * Must be called on game thread. 
+	 */
+	static void TearDown()
+	{
+		TLazySingleton<FRendererOnScreenNotification>::TearDown();
+	}
+
+	/** 
+	 * Relay to AddLambda() of underlying delegate. 
+	 * Thia takes a lock so that it can be called from any thread.
+	 * The lambda will be called from the game thread!
+	 */
+	template<typename FunctorType, typename... VarTypes>
+	FDelegateHandle AddLambda(FunctorType&& InFunctor, VarTypes... Vars)
+	{
+		FScopeLock Lock(&DelgateCS);
+		return ProxyDelegate.AddLambda(MoveTemp(InFunctor), Vars...);
+	}
+
+	/** 
+	 * Relay to Remove() of underlying delegate. 
+	 * Thia takes a lock so that it can be called from any thread.
+	 */
+	bool Remove(FDelegateHandle InHandle)
+	{
+		FScopeLock Lock(&DelgateCS);
+		return ProxyDelegate.Remove(InHandle);
+	}
+
+private:
+	friend FLazySingleton;
+
+	FRendererOnScreenNotification()
+	{
+		BaseDelegateHandle = FCoreDelegates::OnGetOnScreenMessages.AddLambda([this](FCoreDelegates::FSeverityMessageMap& OutMessages)
+		{
+			FScopeLock Lock(&DelgateCS);
+			ProxyDelegate.Broadcast(OutMessages);
+		});
+	}
+
+	~FRendererOnScreenNotification()
+	{
+		FCoreDelegates::OnGetOnScreenMessages.Remove(BaseDelegateHandle);
+	}
+
+private:
+	FCriticalSection DelgateCS;
+	FCoreDelegates::FGetOnScreenMessagesDelegate ProxyDelegate;
+	FDelegateHandle BaseDelegateHandle;
+};
diff --git a/Engine/Source/Runtime/Renderer/Private/VT/RuntimeVirtualTextureSceneProxy.cpp b/Engine/Source/Runtime/Renderer/Private/VT/RuntimeVirtualTextureSceneProxy.cpp
--- a/Engine/Source/Runtime/Renderer/Private/VT/RuntimeVirtualTextureSceneProxy.cpp
+++ b/Engine/Source/Runtime/Renderer/Private/VT/RuntimeVirtualTextureSceneProxy.cpp
@@ -3,7 +3,7 @@
 #include "VT/RuntimeVirtualTextureSceneProxy.h"
 
 #include "Components/RuntimeVirtualTextureComponent.h"
-#include "Misc/CoreDelegates.h"
+#include "RendererOnScreenNotification.h"
 #include "VirtualTextureSystem.h"
 #include "VT/RuntimeVirtualTexture.h"
 #include "VT/RuntimeVirtualTextureProducer.h"
@@ -61,7 +61,7 @@ FRuntimeVirtualTextureSceneProxy::FRuntimeVirtualTextureSceneProxy(URuntimeVirtu
 #if !UE_BUILD_SHIPPING
 				// Notify that streaming texture is invalid since this can cause performance regression.
 				const FString Name = InComponent->GetPathName();
-				OnScreenWarningDelegateHandle = FCoreDelegates::OnGetOnScreenMessages.AddLambda([Name](FCoreDelegates::FSeverityMessageMap& OutMessages)
+				OnScreenWarningDelegateHandle = FRendererOnScreenNotification::Get().AddLambda([Name](FCoreDelegates::FSeverityMessageMap& OutMessages)
 				{
 					OutMessages.Add(
 						FCoreDelegates::EOnScreenMessageSeverity::Warning,
@@ -96,7 +96,7 @@ FRuntimeVirtualTextureSceneProxy::~FRuntimeVirtualTextureSceneProxy()
 	checkSlow(IsInRenderingThread());
 
 #if !UE_BUILD_SHIPPING
-	FCoreDelegates::OnGetOnScreenMessages.Remove(OnScreenWarningDelegateHandle);
+	FRendererOnScreenNotification::Get().Remove(OnScreenWarningDelegateHandle);
 #endif
 }
 
