Fix buffer offset during buffer backed texture binding.

diff --git a/Engine/Source/Runtime/Apple/AGXRHI/Private/AGXStateCache.cpp b/Engine/Source/Runtime/Apple/AGXRHI/Private/AGXStateCache.cpp
--- a/Engine/Source/Runtime/Apple/AGXRHI/Private/AGXStateCache.cpp
+++ b/Engine/Source/Runtime/Apple/AGXRHI/Private/AGXStateCache.cpp
@@ -1569,7 +1569,9 @@ void FAGXStateCache::SetShaderUnorderedAccessView(EAGXShaderStages ShaderStage,
 					uint32 ElementsPerRow = BytesPerRow / GPixelFormats[(EPixelFormat)Texture->GetFormat()].BlockBytes;
 					
 					FAGXBuffer Buffer(Surface->Texture.GetBuffer(), false);
-					SetShaderBuffer(ShaderStage, Buffer, nullptr, 0, Surface->Texture.GetBuffer().GetLength(), BindIndex, mtlpp::ResourceUsage(mtlpp::ResourceUsage::Read | mtlpp::ResourceUsage::Write), PF_MAX, ElementsPerRow);
+					const uint32 BufferOffset = Surface->Texture.GetBufferOffset();
+					const uint32 BufferSize = Surface->Texture.GetBuffer().GetLength();
+					SetShaderBuffer(ShaderStage, Buffer, nullptr, BufferOffset, BufferSize, BindIndex, mtlpp::ResourceUsage(mtlpp::ResourceUsage::Read | mtlpp::ResourceUsage::Write), PF_MAX, ElementsPerRow);
 				}
 			}
 			else
diff --git a/Engine/Source/Runtime/Apple/MetalRHI/Private/MetalStateCache.cpp b/Engine/Source/Runtime/Apple/MetalRHI/Private/MetalStateCache.cpp
--- a/Engine/Source/Runtime/Apple/MetalRHI/Private/MetalStateCache.cpp
+++ b/Engine/Source/Runtime/Apple/MetalRHI/Private/MetalStateCache.cpp
@@ -1585,7 +1585,9 @@ void FMetalStateCache::SetShaderUnorderedAccessView(EMetalShaderStages ShaderSta
 					uint32 ElementsPerRow = BytesPerRow / GPixelFormats[(EPixelFormat)Texture->GetFormat()].BlockBytes;
 				
 					FMetalBuffer Buffer(Surface->Texture.GetBuffer(), false);
-					SetShaderBuffer(ShaderStage, Buffer, nullptr, 0, Surface->Texture.GetBuffer().GetLength(), BindIndex, mtlpp::ResourceUsage(mtlpp::ResourceUsage::Read | mtlpp::ResourceUsage::Write), PF_MAX, ElementsPerRow);
+					const uint32 BufferOffset = Surface->Texture.GetBufferOffset();
+					const uint32 BufferSize = Surface->Texture.GetBuffer().GetLength();
+					SetShaderBuffer(ShaderStage, Buffer, nullptr, BufferOffset, BufferSize, BindIndex, mtlpp::ResourceUsage(mtlpp::ResourceUsage::Read | mtlpp::ResourceUsage::Write), PF_MAX, ElementsPerRow);
 				}
 			}
 			else
diff --git a/Engine/Source/Runtime/Renderer/Private/PostProcess/PostProcessHistogram.cpp b/Engine/Source/Runtime/Renderer/Private/PostProcess/PostProcessHistogram.cpp
--- a/Engine/Source/Runtime/Renderer/Private/PostProcess/PostProcessHistogram.cpp
+++ b/Engine/Source/Runtime/Renderer/Private/PostProcess/PostProcessHistogram.cpp
@@ -300,12 +300,7 @@ static FRDGTextureRef AddHistogramAtomicPass(
 		// 		FIntPoint(FHistogramAtomicCS::HistogramTexelCount, 1),
 		// 		PF_R32G32_UINT,
 		// 		FClearValueBinding::None,
-		// 		TexCreate_UAV);
-		// 
-		// 	if (IsMetalPlatform(View.GetShaderPlatform()))
-		// 	{
-		// 		Desc.Flags |= TexCreate_NoTiling;
-		// 	}
+		// 		TexCreate_UAV | TexCreate_AtomicCompatible);
 		// 
 		// 	HistogramScatter64Texture = GraphBuilder.CreateTexture(Desc, TEXT("Histogram.Scatter64"));
 		// }
@@ -315,12 +310,7 @@ static FRDGTextureRef AddHistogramAtomicPass(
 				FIntPoint(FHistogramAtomicCS::HistogramTexelCount * 2, 1),
 				PF_R32_UINT,
 				FClearValueBinding::None,
-				TexCreate_ShaderResource | TexCreate_UAV);
-
-			if (IsMetalPlatform(View.GetShaderPlatform()))
-			{
-				Desc.Flags |= TexCreate_NoTiling;
-			}
+				TexCreate_ShaderResource | TexCreate_UAV | TexCreate_AtomicCompatible);
 
 			HistogramScatter32Texture = GraphBuilder.CreateTexture(Desc, TEXT("Histogram.Scatter32"));
 		}
@@ -460,4 +450,4 @@ FRDGTextureRef AddLocalExposurePass(
 	}
 
 	return LocalExposureTexture;
-}
\ No newline at end of file
+}
