- Refactored how bRequireFullVisibilityToRender is handled : When level using bRequireFullVisibilityToRender is added to the world, primitives are not waiting level to become visible before being added to the scene. Instead they are added but hidden until the level becomes visible.
- Reduces spikes in SendAllEndOfFrame, UpdateAllPrimitiveSceneInfos, AddPrimitiveSceneInfos.
- Optimized HLOD visibility state changes.
- Optimized FPrimitiveSceneInfo::RemoveFromScene when associated proxy needs addto/removefrom world notification.
- AddPrimitive is now part of AddToWorld budget when Level bRequireFullVisibilityToRender flag is set.
- Temporary implementation for Nanite proxies to support force hidden flag.

diff --git a/Engine/Source/Runtime/Engine/Classes/Components/ActorComponent.h b/Engine/Source/Runtime/Engine/Classes/Components/ActorComponent.h
--- a/Engine/Source/Runtime/Engine/Classes/Components/ActorComponent.h
+++ b/Engine/Source/Runtime/Engine/Classes/Components/ActorComponent.h
@@ -42,6 +42,7 @@ public:
 		AddPrimitiveBatches.Add(PrimitiveComponent);
 	}
 
+	int32 Count() const { return AddPrimitiveBatches.Num(); }
 	void Process();
 
 private:
diff --git a/Engine/Source/Runtime/Engine/Classes/Engine/CoreSettings.h b/Engine/Source/Runtime/Engine/Classes/Engine/CoreSettings.h
--- a/Engine/Source/Runtime/Engine/Classes/Engine/CoreSettings.h
+++ b/Engine/Source/Runtime/Engine/Classes/Engine/CoreSettings.h
@@ -87,6 +87,12 @@ protected:
 		ToolTip = "Batching granularity used to register actor components during level streaming."))
 	int32 LevelStreamingComponentsRegistrationGranularity;
 
+	/** Batching granularity used to add primitives to scene in parallel when registering actor components during level streaming */
+	UPROPERTY(EditAnywhere, config, Category = LevelStreaming, AdvancedDisplay, meta = (
+		ConsoleVariable = "s.LevelStreamingAddPrimitiveGranularity", DisplayName = "Added Primitive Granularity",
+		ToolTip = "Batching granularity used to add primitives to scene in parallel when registering actor components during level streaming."))
+	int32 LevelStreamingAddPrimitiveGranularity;
+
 	/** Maximum allowed time to spend while unregistering components during level streaming (ms per frame) */
 	UPROPERTY(EditAnywhere, config, Category = LevelStreaming, AdvancedDisplay, meta = (
 		ConsoleVariable = "s.UnregisterComponentsTimeLimit", DisplayName = "Component Unregister Update Time Limit",
@@ -132,6 +138,8 @@ extern ENGINE_API float GLevelStreamingActorsUpdateTimeLimit;
 extern ENGINE_API float GPriorityLevelStreamingActorsUpdateExtraTime;
 /** Batching granularity used to register actor components during level streaming. */
 extern ENGINE_API int32 GLevelStreamingComponentsRegistrationGranularity;
+/** Batching granularity used to add primitives to scene in parallel when registering actor components during level streaming. */
+extern ENGINE_API int32 GLevelStreamingAddPrimitiveGranularity;
 /** Batching granularity used to unregister actor components during level streaming.  */
 extern ENGINE_API int32 GLevelStreamingComponentsUnregistrationGranularity;
 /** Maximum allowed time to spend for actor unregistration steps during level streaming (ms per frame). If this is 0.0 then we don't timeslice.*/
diff --git a/Engine/Source/Runtime/Engine/Private/Components/SceneComponent.cpp b/Engine/Source/Runtime/Engine/Private/Components/SceneComponent.cpp
--- a/Engine/Source/Runtime/Engine/Private/Components/SceneComponent.cpp
+++ b/Engine/Source/Runtime/Engine/Private/Components/SceneComponent.cpp
@@ -3020,13 +3020,6 @@ bool USceneComponent::IsVisible() const
 		return false;
 	}
 
-	// If the level has the bRequireFullVisibilityToRender flag and it ain't visible yet, this component should not be visible.
-	const ULevel* ComponentLevel = GetComponentLevel();
-	if (ComponentLevel && ComponentLevel->bRequireFullVisibilityToRender && !ComponentLevel->bIsVisible)
-	{
-		return false;
-	}
-	
 	return (GetVisibleFlag() && (!CachedLevelCollection || CachedLevelCollection->IsVisible())); 
 }
 
diff --git a/Engine/Source/Runtime/Engine/Private/CoreSettings.cpp b/Engine/Source/Runtime/Engine/Private/CoreSettings.cpp
--- a/Engine/Source/Runtime/Engine/Private/CoreSettings.cpp
+++ b/Engine/Source/Runtime/Engine/Private/CoreSettings.cpp
@@ -18,6 +18,7 @@ float GLevelStreamingActorsUpdateTimeLimit = 5.0f;
 float GPriorityLevelStreamingActorsUpdateExtraTime = 5.0f;
 float GLevelStreamingUnregisterComponentsTimeLimit = 1.0f;
 int32 GLevelStreamingComponentsRegistrationGranularity = 10;
+int32 GLevelStreamingAddPrimitiveGranularity = 120;
 int32 GLevelStreamingComponentsUnregistrationGranularity = 5;
 int32 GLevelStreamingForceGCAfterLevelStreamedOut = 1;
 int32 GLevelStreamingContinuouslyIncrementalGCWhileLevelsPendingPurge = 1;
@@ -82,6 +83,13 @@ static FAutoConsoleVariableRef CVarLevelStreamingComponentsRegistrationGranulari
 	ECVF_Default
 	);
 
+static FAutoConsoleVariableRef CVarLevelStreamingAddPrimitiveGranularity(
+	TEXT("s.LevelStreamingAddPrimitiveGranularity"),
+	GLevelStreamingAddPrimitiveGranularity,
+	TEXT("Batching granularity used to add primitives to scene in parallel when registering actor components during level streaming."),
+	ECVF_Default
+);
+
 static FAutoConsoleVariableRef CVarLevelStreamingComponentsUnregistrationGranularity(
 	TEXT("s.LevelStreamingComponentsUnregistrationGranularity"),
 	GLevelStreamingComponentsUnregistrationGranularity,
@@ -148,6 +156,7 @@ UStreamingSettings::UStreamingSettings()
 	LevelStreamingActorsUpdateTimeLimit = 5.0f;
 	PriorityLevelStreamingActorsUpdateExtraTime = 5.0f;
 	LevelStreamingComponentsRegistrationGranularity = 10;
+	LevelStreamingAddPrimitiveGranularity = 120;
 	LevelStreamingUnregisterComponentsTimeLimit = 1.0f;
 	LevelStreamingComponentsUnregistrationGranularity = 5;
 	EventDrivenLoaderEnabled = false;
diff --git a/Engine/Source/Runtime/Engine/Private/LevelStreaming.cpp b/Engine/Source/Runtime/Engine/Private/LevelStreaming.cpp
--- a/Engine/Source/Runtime/Engine/Private/LevelStreaming.cpp
+++ b/Engine/Source/Runtime/Engine/Private/LevelStreaming.cpp
@@ -594,12 +594,6 @@ void ULevelStreaming::UpdateStreamingState(bool& bOutUpdateAgain, bool& bOutRede
 
 				if (LoadedLevel->bIsVisible)
 				{
-					if (LoadedLevel->bRequireFullVisibilityToRender)
-					{
-						// Make sure components from this level are made visible
-						LoadedLevel->MarkLevelComponentsRenderStateDirty();
-					}
-
 					// immediately discard previous level
 					DiscardPendingUnloadLevel(World);
 
@@ -638,8 +632,6 @@ void ULevelStreaming::UpdateStreamingState(bool& bOutUpdateAgain, bool& bOutRede
 			const bool LevelBecameInvisible = bWasVisible && !LoadedLevel->bIsVisible;
 			if (LoadedLevel->bRequireFullVisibilityToRender && LevelBecameInvisible)
 			{
-				// Make sure components from this level are hidden
-				LoadedLevel->MarkLevelComponentsRenderStateDirty();
 				RemoveLevelFromScene();
 			}
 
diff --git a/Engine/Source/Runtime/Engine/Private/PrimitiveSceneProxy.cpp b/Engine/Source/Runtime/Engine/Private/PrimitiveSceneProxy.cpp
--- a/Engine/Source/Runtime/Engine/Private/PrimitiveSceneProxy.cpp
+++ b/Engine/Source/Runtime/Engine/Private/PrimitiveSceneProxy.cpp
@@ -220,6 +220,16 @@ FPrimitiveSceneProxy::FPrimitiveSceneProxy(const UPrimitiveComponent* InComponen
 {
 	check(Scene);
 
+	// Initialize ForceHidden flag based on Level's visibility (only if Level bRequireFullVisibilityToRender is set)
+	if (ULevel* Level = InComponent->GetComponentLevel())
+	{
+		bShouldNotifyOnWorldAddRemove = Level->bRequireFullVisibilityToRender;
+		if (bShouldNotifyOnWorldAddRemove)
+		{
+			SetForceHidden(!Level->bIsVisible);
+		}
+	}
+
 #if STATS
 	{
 		UObject const* StatObject = InComponent->AdditionalStatObject(); // prefer the additional object, this is usually the thing related to the component
@@ -333,6 +343,19 @@ FPrimitiveSceneProxy::FPrimitiveSceneProxy(const UPrimitiveComponent* InComponen
 	}
 }
 
+bool FPrimitiveSceneProxy::OnLevelAddedToWorld_RenderThread()
+{
+	check(bShouldNotifyOnWorldAddRemove);
+	SetForceHidden(false);
+	return false;
+}
+
+void FPrimitiveSceneProxy::OnLevelRemovedFromWorld_RenderThread()
+{
+	check(bShouldNotifyOnWorldAddRemove);
+	SetForceHidden(true);
+}
+
 #if WITH_EDITOR
 void FPrimitiveSceneProxy::SetUsedMaterialForVerification(const TArray<UMaterialInterface*>& InUsedMaterialsForVerification)
 {
@@ -821,6 +844,12 @@ void FPrimitiveSceneProxy::SetCollisionEnabled_RenderThread(const bool bNewEnabl
 /** @return True if the primitive is visible in the given View. */
 bool FPrimitiveSceneProxy::IsShown(const FSceneView* View) const
 {
+	// If primitive is forcibly hidden
+	if (IsForceHidden())
+	{
+		return false;
+	}
+
 #if WITH_EDITOR
 	// Don't draw editor specific actors during game mode
 	if (View->Family->EngineShowFlags.Game)
@@ -836,12 +865,6 @@ bool FPrimitiveSceneProxy::IsShown(const FSceneView* View) const
 		return false;
 	}
 
-	// If primitive is forcibly hidden
-	if (bForceHidden)
-	{
-		return false;
-	}
-
 	// After checking for VR/Desktop Edit mode specific actors, check for Editor vs. Game
 	if(View->Family->EngineShowFlags.Editor)
 	{
@@ -958,7 +981,7 @@ bool FPrimitiveSceneProxy::IsShadowCast(const FSceneView* View) const
 #endif	//#if WITH_EDITOR
 
 		// If primitive is forcibly hidden
-		if (bForceHidden)
+		if (IsForceHidden())
 		{
 			return false;
 		}
diff --git a/Engine/Source/Runtime/Engine/Private/World.cpp b/Engine/Source/Runtime/Engine/Private/World.cpp
--- a/Engine/Source/Runtime/Engine/Private/World.cpp
+++ b/Engine/Source/Runtime/Engine/Private/World.cpp
@@ -2773,14 +2773,26 @@ void UWorld::AddToWorld( ULevel* Level, const FTransform& LevelTransform, bool b
 			// from disk rather than duplicated
 			bRerunConstructionScript = !(IsGameWorld() && (Level->bHasRerunConstructionScripts || !bRerunConstructionDuringEditorStreaming));
 		}
-		
+
+		// Prepare context used to store batch/parallelize AddPrimitive calls
+		FRegisterComponentContext Context(this);
+		FRegisterComponentContext* ContextPtr = GLevelStreamingAddPrimitiveGranularity == 0 ? nullptr : &Context;
+
 		// Incrementally update components.
 		int32 NumComponentsToUpdate = (!bConsiderTimeLimit || !IsGameWorld() || IsRunningCommandlet() ? 0 : GLevelStreamingComponentsRegistrationGranularity);
 		do
 		{
-			Level->IncrementalUpdateComponents( NumComponentsToUpdate, bRerunConstructionScript );
+			Level->IncrementalUpdateComponents( NumComponentsToUpdate, bRerunConstructionScript, ContextPtr);
+			// Process AddPrimitives if threshold is reached
+			if (Context.Count() > GLevelStreamingAddPrimitiveGranularity)
+			{
+				Context.Process();
+			}
 		}
 		while (!Level->bAreComponentsCurrentlyRegistered && !IsTimeLimitExceeded(TEXT("updating components"), StartTime, Level, TimeLimit));
+		
+		// Process remaining AddPrimitives
+		Context.Process();
 
 		// We are done once all components are attached.
 		Level->bAlreadyUpdatedComponents	= Level->bAreComponentsCurrentlyRegistered;
diff --git a/Engine/Source/Runtime/Engine/Private/WorldPartition/HLOD/HLODActor.cpp b/Engine/Source/Runtime/Engine/Private/WorldPartition/HLOD/HLODActor.cpp
--- a/Engine/Source/Runtime/Engine/Private/WorldPartition/HLOD/HLODActor.cpp
+++ b/Engine/Source/Runtime/Engine/Private/WorldPartition/HLOD/HLODActor.cpp
@@ -38,9 +38,12 @@ UPrimitiveComponent* AWorldPartitionHLOD::GetHLODComponent()
 
 void AWorldPartitionHLOD::SetVisibility(bool bInVisible)
 {
-	if (GetRootComponent())
+	// When propagating visibility state to children, SetVisibility dirties all attached components.
+	// Because we know that the visibility flag of all components of an HLOD actor are always in sync, 
+	// we test on RootComponent to check if call is required. This way we avoid dirtying all primitive proxies render state.
+	if (RootComponent && (RootComponent->GetVisibleFlag() != bInVisible))
 	{
-		GetRootComponent()->SetVisibility(bInVisible, /*bPropagateToChildren*/ true);
+		RootComponent->SetVisibility(bInVisible, /*bPropagateToChildren*/ true);
 	}
 }
 
diff --git a/Engine/Source/Runtime/Engine/Public/PrimitiveSceneProxy.h b/Engine/Source/Runtime/Engine/Public/PrimitiveSceneProxy.h
--- a/Engine/Source/Runtime/Engine/Public/PrimitiveSceneProxy.h
+++ b/Engine/Source/Runtime/Engine/Public/PrimitiveSceneProxy.h
@@ -402,14 +402,14 @@ public:
 	 * 
 	 * @return return true if the primitive should be added to the scene
 	 */
-	virtual bool OnLevelAddedToWorld_RenderThread() {return false;}
+	ENGINE_API virtual bool OnLevelAddedToWorld_RenderThread();
 
 	/**
 	 * Called to notify the proxy that the level has been fully removed from
 	 * the world and the primitive will not be rendered.
 	 * Only called if bShouldNotifyOnWorldAddRemove is set to true.
 	 */
-	virtual void OnLevelRemovedFromWorld_RenderThread() {}
+	ENGINE_API virtual void OnLevelRemovedFromWorld_RenderThread();
 
 	/**
 	* @return true if the proxy can be culled when occluded by other primitives
@@ -995,9 +995,9 @@ private:
 	/** If true this primitive Renders in the mainPass */
 	uint8 bRenderInMainPass : 1;
 
-	/** If true this primitive will always return that it is hidden if queried*/
+	/** If true this primitive is hidden (used when level is not yet visible) */
 	uint8 bForceHidden : 1;
-	
+
 	/** Whether this component has any collision enabled */
 	uint8 bCollisionEnabled : 1;
 
diff --git a/Engine/Source/Runtime/Renderer/Private/GPUScene.cpp b/Engine/Source/Runtime/Renderer/Private/GPUScene.cpp
--- a/Engine/Source/Runtime/Renderer/Private/GPUScene.cpp
+++ b/Engine/Source/Runtime/Renderer/Private/GPUScene.cpp
@@ -217,6 +217,13 @@ struct FLightMapUploadInfo
 	int32 LightmapDataOffset = 0;
 };
 
+// TODO: Temporary hack : For FPrimitiveSceneProxy::IsForceHidden() to work with Nanite proxies, return an invalid primitive ID if IsForceHidden() returns true.
+static FORCEINLINE int32 GetPrimitiveID(const FScene& InScene, const int32 InPrimitiveID)
+{
+	const FPrimitiveSceneProxy* PrimitiveSceneProxy = InScene.PrimitiveSceneProxies[InPrimitiveID];
+	return (PrimitiveSceneProxy->IsNaniteMesh() && PrimitiveSceneProxy->IsForceHidden()) ? INVALID_PRIMITIVE_ID : InPrimitiveID;
+}
+
 /**
  * Implements a thin data abstraction such that the UploadGeneral function can upload primitive data from
  * both scene primitives and dynamic primitives (which are not stored in the same way). 
@@ -292,7 +299,7 @@ struct FUploadDataSourceAdapterScenePrimitives
 
 			InstanceUploadInfo.InstanceSceneDataOffset = PrimitiveSceneInfo->GetInstanceSceneDataOffset();
 			InstanceUploadInfo.LastUpdateSceneFrameNumber = SceneFrameNumber;
-			InstanceUploadInfo.PrimitiveID = PrimitiveID;
+			InstanceUploadInfo.PrimitiveID = GetPrimitiveID(Scene, PrimitiveID);
 			InstanceUploadInfo.PrimitiveToWorld = AbsoluteOrigin.MakeToRelativeWorldMatrix(LocalToWorld);
 
 			{
@@ -452,7 +459,7 @@ void FGPUScene::UpdateInternal(FRDGBuilder& GraphBuilder, FScene& Scene)
 				if (PrimitiveSceneInfo->GetInstanceSceneDataOffset() >= 0)
 				{
 					InstanceSceneDataToClear.SetRange(PrimitiveSceneInfo->GetInstanceSceneDataOffset(), PrimitiveSceneInfo->GetNumInstanceSceneDataEntries(), false);
-					IdOnlyUpdateData.Add(PrimitiveSceneInfo->GetInstanceSceneDataOffset(), PrimitiveSceneInfo->GetNumInstanceSceneDataEntries(), PrimitiveId);
+					IdOnlyUpdateData.Add(PrimitiveSceneInfo->GetInstanceSceneDataOffset(), PrimitiveSceneInfo->GetNumInstanceSceneDataEntries(), GetPrimitiveID(Scene, PrimitiveId));
 				}
 			}
 		}
diff --git a/Engine/Source/Runtime/Renderer/Private/PrimitiveSceneInfo.cpp b/Engine/Source/Runtime/Renderer/Private/PrimitiveSceneInfo.cpp
--- a/Engine/Source/Runtime/Renderer/Private/PrimitiveSceneInfo.cpp
+++ b/Engine/Source/Runtime/Renderer/Private/PrimitiveSceneInfo.cpp
@@ -164,6 +164,7 @@ FPrimitiveSceneInfo::FPrimitiveSceneInfo(UPrimitiveComponent* InComponent,FScene
 	bIndirectLightingCacheBufferDirty(false),
 	bRegisteredVirtualTextureProducerCallback(false),
 	bRegisteredWithVelocityData(false),
+	LevelUpdateNotificationIndex(INDEX_NONE),
 	InstanceSceneDataOffset(INDEX_NONE),
 	NumInstanceSceneDataEntries(0),
 	InstancePayloadDataOffset(INDEX_NONE),
@@ -1339,15 +1340,18 @@ void FPrimitiveSceneInfo::AddToScene(FRHICommandListImmediate& RHICmdList, FScen
 		INC_MEMORY_STAT_BY(STAT_PrimitiveInfoMemory, sizeof(*SceneInfo) + SceneInfo->StaticMeshes.GetAllocatedSize() + SceneInfo->StaticMeshRelevances.GetAllocatedSize() + Proxy->GetMemoryFootprint());
 	}
 
-	// Some primitive types cannot add their meshes until the level is added to the world.
-	for (FPrimitiveSceneInfo* SceneInfo : SceneInfos)
 	{
-		if (SceneInfo->Proxy->ShouldNotifyOnWorldAddRemove())
+		SCOPED_NAMED_EVENT(FPrimitiveSceneInfo_AddToScene_LevelNotifyPrimitives, FColor::Blue);
+		for (FPrimitiveSceneInfo* SceneInfo : SceneInfos)
 		{
-			TArray<FPrimitiveSceneInfo*>& LevelNotifyPrimitives = Scene->PrimitivesNeedingLevelUpdateNotification.FindOrAdd(SceneInfo->Proxy->GetLevelName());
-			LevelNotifyPrimitives.Add(SceneInfo);
+			if (SceneInfo->Proxy->ShouldNotifyOnWorldAddRemove())
+			{
+				TArray<FPrimitiveSceneInfo*>& LevelNotifyPrimitives = Scene->PrimitivesNeedingLevelUpdateNotification.FindOrAdd(SceneInfo->Proxy->GetLevelName());
+				SceneInfo->LevelUpdateNotificationIndex = LevelNotifyPrimitives.Num();
+				LevelNotifyPrimitives.Add(SceneInfo);
+			}
 		}
-	}	
+	}
 }
 
 void FPrimitiveSceneInfo::RemoveStaticMeshes()
@@ -1442,11 +1446,17 @@ void FPrimitiveSceneInfo::RemoveFromScene(bool bUpdateStaticDrawLists)
 		TArray<FPrimitiveSceneInfo*>* LevelNotifyPrimitives = Scene->PrimitivesNeedingLevelUpdateNotification.Find(Proxy->GetLevelName());
 		if (LevelNotifyPrimitives != nullptr)
 		{
-			LevelNotifyPrimitives->Remove(this);
+			checkSlow(LevelUpdateNotificationIndex != INDEX_NONE);
+			LevelNotifyPrimitives->RemoveAtSwap(LevelUpdateNotificationIndex, 1, false);
 			if (LevelNotifyPrimitives->Num() == 0)
 			{
 				Scene->PrimitivesNeedingLevelUpdateNotification.Remove(Proxy->GetLevelName());
 			}
+			else if (LevelUpdateNotificationIndex < LevelNotifyPrimitives->Num())
+			{
+				// Update swapped element's LevelUpdateNotificationIndex
+				((*LevelNotifyPrimitives)[LevelUpdateNotificationIndex])->LevelUpdateNotificationIndex = LevelUpdateNotificationIndex;
+			}
 		}
 	}
 }
diff --git a/Engine/Source/Runtime/Renderer/Private/RendererScene.cpp b/Engine/Source/Runtime/Renderer/Private/RendererScene.cpp
--- a/Engine/Source/Runtime/Renderer/Private/RendererScene.cpp
+++ b/Engine/Source/Runtime/Renderer/Private/RendererScene.cpp
@@ -3880,6 +3880,12 @@ void FScene::OnLevelAddedToWorld_RenderThread(FName InLevelName)
 				Primitive->RemoveStaticMeshes();
 				PrimitivesToAdd.Add(Primitive);
 			}
+			// Invalidate primitive proxy entry in GPU Scene. 
+			// This is necessary for Nanite::FSceneProxy to be uploaded to GPU scene (see GetPrimitiveID in GPUScene.cpp)
+			if (Primitive->Proxy->IsNaniteMesh())
+			{
+				Primitive->RequestGPUSceneUpdate();
+			}
 		}
 	}
 
@@ -3912,6 +3918,12 @@ void FScene::OnLevelRemovedFromWorld_RenderThread(FName InLevelName)
 		for (FPrimitiveSceneInfo* Primitive : *LevelPrimitives)
 		{
 			Primitive->Proxy->OnLevelRemovedFromWorld_RenderThread();
+			// Invalidate primitive proxy entry in GPU Scene.
+			// This is necessary for Nanite::FSceneProxy to be uploaded to GPU scene (see GetPrimitiveID in GPUScene.cpp)
+			if (Primitive->Proxy->IsNaniteMesh())
+			{
+				Primitive->RequestGPUSceneUpdate();
+			}
 		}
 	}
 }
diff --git a/Engine/Source/Runtime/Renderer/Public/PrimitiveSceneInfo.h b/Engine/Source/Runtime/Renderer/Public/PrimitiveSceneInfo.h
--- a/Engine/Source/Runtime/Renderer/Public/PrimitiveSceneInfo.h
+++ b/Engine/Source/Runtime/Renderer/Public/PrimitiveSceneInfo.h
@@ -614,6 +614,9 @@ private:
 	/** True if the primitive registered with velocity data and needs to remove itself when being removed from the scene. */
 	bool bRegisteredWithVelocityData : 1;
 
+	/** Index into the scene's PrimitivesNeedingLevelUpdateNotification array for this primitive scene info level. */
+	int32 LevelUpdateNotificationIndex;
+
 	/** Offset into the scene's instance scene data buffer, when GPUScene is enabled. */
 	int32 InstanceSceneDataOffset;
 
