Ray Traced skylight gets disabled if Lumen GI is enabled, or the GlobalIllumination show flag is disabled

diff --git a/Engine/Source/Runtime/Renderer/Private/DeferredShadingRenderer.cpp b/Engine/Source/Runtime/Renderer/Private/DeferredShadingRenderer.cpp
--- a/Engine/Source/Runtime/Renderer/Private/DeferredShadingRenderer.cpp
+++ b/Engine/Source/Runtime/Renderer/Private/DeferredShadingRenderer.cpp
@@ -2847,7 +2847,10 @@ void FDeferredShadingSceneRenderer::Render(FRDGBuilder& GraphBuilder)
 		}
 
 	#if RHI_RAYTRACING
-		if (ShouldRenderRayTracingSkyLight(Scene->SkyLight))
+		if (ShouldRenderRayTracingSkyLight(Scene->SkyLight) 
+			//@todo - integrate RenderRayTracingSkyLight into RenderDiffuseIndirectAndAmbientOcclusion
+			&& GetViewPipelineState(Views[0]).DiffuseIndirectMethod != EDiffuseIndirectMethod::Lumen
+			&& ViewFamily.EngineShowFlags.GlobalIllumination)
 		{
 			FRDGTextureRef SkyLightTexture = nullptr;
 			FRDGTextureRef SkyLightHitDistanceTexture = nullptr;
