Skip storing PrevFrameViewInfo.DepthBuffer for Lumen, which has its own previous depth buffer

diff --git a/Engine/Source/Runtime/Renderer/Private/DeferredShadingRenderer.cpp b/Engine/Source/Runtime/Renderer/Private/DeferredShadingRenderer.cpp
--- a/Engine/Source/Runtime/Renderer/Private/DeferredShadingRenderer.cpp
+++ b/Engine/Source/Runtime/Renderer/Private/DeferredShadingRenderer.cpp
@@ -3141,15 +3141,24 @@ void FDeferredShadingSceneRenderer::Render(FRDGBuilder& GraphBuilder)
 	{
 		const FViewInfo& View = Views[ViewIndex];
 
-		if (((View.FinalPostProcessSettings.DynamicGlobalIlluminationMethod == EDynamicGlobalIlluminationMethod::ScreenSpace && ScreenSpaceRayTracing::ShouldKeepBleedFreeSceneColor(View))
-			|| GetViewPipelineState(View).DiffuseIndirectMethod == EDiffuseIndirectMethod::Lumen
-			|| GetViewPipelineState(View).ReflectionsMethod == EReflectionsMethod::Lumen)
-			&& !View.bStatePrevViewInfoIsReadOnly)
+		if (!View.bStatePrevViewInfoIsReadOnly)
 		{
 			// Keep scene color and depth for next frame screen space ray tracing.
 			FSceneViewState* ViewState = View.ViewState;
-			GraphBuilder.QueueTextureExtraction(SceneTextures.Depth.Resolve, &ViewState->PrevFrameViewInfo.DepthBuffer);
-			GraphBuilder.QueueTextureExtraction(SceneTextures.Color.Resolve, &ViewState->PrevFrameViewInfo.ScreenSpaceRayTracingInput);
+
+			const bool bSSGI = View.FinalPostProcessSettings.DynamicGlobalIlluminationMethod == EDynamicGlobalIlluminationMethod::ScreenSpace && ScreenSpaceRayTracing::ShouldKeepBleedFreeSceneColor(View);
+			
+			if (bSSGI)
+			{
+				GraphBuilder.QueueTextureExtraction(SceneTextures.Depth.Resolve, &ViewState->PrevFrameViewInfo.DepthBuffer);
+			}
+			
+			if (bSSGI
+				|| GetViewPipelineState(View).DiffuseIndirectMethod == EDiffuseIndirectMethod::Lumen
+				|| GetViewPipelineState(View).ReflectionsMethod == EReflectionsMethod::Lumen)
+			{
+				GraphBuilder.QueueTextureExtraction(SceneTextures.Color.Resolve, &ViewState->PrevFrameViewInfo.ScreenSpaceRayTracingInput);
+			}
 		}
 	}
 
