Occlusion cull instances drawing into non-nanite VSM (controlled by r.Shadow.Virtual.NonNanite.UseHZB, default mode 2)
- Added stats for non-nanite VSM instance culling (moved VSM stats functionality into own file).
- r.Shadow.Virtual.NonNanite.UseHZB == 2 (default) uses the current-frame Nanite VSM HZB as this enables correct culling for camera cuts & light movement and contains most of the occluding geometry.

diff --git a/Engine/Shaders/Private/VirtualShadowMaps/BuildPerPageDrawCommands.usf b/Engine/Shaders/Private/VirtualShadowMaps/BuildPerPageDrawCommands.usf
--- a/Engine/Shaders/Private/VirtualShadowMaps/BuildPerPageDrawCommands.usf
+++ b/Engine/Shaders/Private/VirtualShadowMaps/BuildPerPageDrawCommands.usf
@@ -11,6 +11,7 @@
 #include "../InstanceCulling/InstanceCullingCommon.ush"
 #include "../InstanceCulling/InstanceCullingLoadBalancer.ush"
 #include "../WaveOpUtil.ush"
+#include "Stats.ush"
 
 
 #ifndef NEAR_CLIP
@@ -54,6 +55,13 @@ RWBuffer<uint> DrawIndirectArgsBufferOut;
 RWStructuredBuffer<uint> OutInvalidatingInstances;
 uint NumInvalidatingInstanceSlots;
 
+
+#if USE_HZB_OCCLUSION
+// Mode 1 == previous frame HZB, 
+// Mode 2 == this frame HZB (Nanite only).
+uint HZBMode;
+#endif
+
 void MarkInstanceForInvalidation(uint InstanceId)
 {
 	// Note: InstanceId can be higher than NumInvalidatingInstanceSlots because of Dynamic primitives and their instances as they are added later in the frame
@@ -136,6 +144,12 @@ void CullPerPageDrawCommandsCs(uint3 GroupId : SV_GroupID, int GroupThreadIndex
 
 		FFrustumCullData Cull = BoxCullFrustum(InstanceData.LocalBoundsCenter, InstanceData.LocalBoundsExtent, LocalToClip, bNearClip, false);
 
+#if USE_HZB_OCCLUSION
+		const bool bViewHZB = (NaniteView.Flags & VIEW_FLAG_HZBTEST) != 0 || HZBMode == 2;
+#endif
+
+		StatsBufferInterlockedAdd(VSM_STAT_NON_NANITE_INSTANCES_TOTAL, NaniteView.TargetNumMipLevels);
+
 		if (Cull.bIsVisible)
 		{
 			// Loop over mip levels and count number of output visible instances, also retain a bit for each mip level (used to skip empty levels in output loop)
@@ -160,12 +174,46 @@ void CullPerPageDrawCommandsCs(uint3 GroupId : SV_GroupID, int GroupThreadIndex
 					RectPages.zw = min(RectPages.zw, AllocatedBounds.zw);
 					if (all(RectPages.zw >= RectPages.xy))
 					{
+#if USE_HZB_OCCLUSION
+						if (bViewHZB)
+						{
+							// Assume:
+							//  1. Virtual address space has not changed (if so it'd have been invalidated)
+							//  2. Physical pages that are chached have not changed (guaranteed by persistent caching)
+							//  OR (mode 2): Using this frame's HZB from Nanite shadow passes.
+							// Then we can:
+							//  1. Iterate the clipped page rect
+							//  2. Test against only valid pages which enables
+							//  3. rejecting any that don't pass HZB at least once (or touches an uncached page)
+
+							// In Mode 2 we are using the current-frame HZB and can use any page whether it is cached or not.
+							uint VisiblePageMask = HZBMode == 1 ? VSM_UNCACHED_FLAG : 0U;
+							if (!IsVisibleValidMaskedHZB(HZBMode == 1 ? MipView.TargetPrevLayerIndex : VirtualShadowMapId, VirtualShadowMapId, MipLevel, RectPages, Rect, VisiblePageMask))
+							{
+								StatsBufferInterlockedInc(VSM_STAT_NON_NANITE_INSTANCES_HZB_CULLED);
+								continue;
+							}
+						}
+#endif // USE_HZB_OCCLUSION
+
 						++ThreadTotalForAllViews;
 						WriteCmd(RectPages.xy, MipViewId, InstanceId, IndirectArgIndex);
 					}
+					else
+					{
+						StatsBufferInterlockedInc(VSM_STAT_NON_NANITE_INSTANCES_EMPTY_RECT_CULLED);
+					}
+				}
+				else
+				{
+					StatsBufferInterlockedInc(VSM_STAT_NON_NANITE_INSTANCES_PAGE_MASK_CULLED);
 				}
 			}
 		}
+		else
+		{
+			StatsBufferInterlockedAdd(VSM_STAT_NON_NANITE_INSTANCES_FRUSTUM_CULLED, NaniteView.TargetNumMipLevels);
+		}
 	}
 
 	// If the instance has WPO or PDO and rendered to any mip, queue for invalidation. This will be processed the before next frame (after a new HZB is built)
@@ -175,6 +223,8 @@ void CullPerPageDrawCommandsCs(uint3 GroupId : SV_GroupID, int GroupThreadIndex
 		MarkInstanceForInvalidation(InstanceId);
 	}
 
+	StatsBufferInterlockedAdd(VSM_STAT_NON_NANITE_INSTANCES_DRAWN, ThreadTotalForAllViews);
+
 	// Accumulate total number of instances for each indirect argument, is also used to allocate space and output compact range of instances later
 	InterlockedAdd(DrawIndirectArgsBufferOut[IndirectArgIndex * INDIRECT_ARGS_NUM_WORDS + 1], ThreadTotalForAllViews);
 }
diff --git a/Engine/Shaders/Private/VirtualShadowMaps/PageManagement.usf b/Engine/Shaders/Private/VirtualShadowMaps/PageManagement.usf
--- a/Engine/Shaders/Private/VirtualShadowMaps/PageManagement.usf
+++ b/Engine/Shaders/Private/VirtualShadowMaps/PageManagement.usf
@@ -20,27 +20,7 @@
 #include "ProjectionCommon.ush"
 #include "ProjectionSpot.ush"
 #include "ProjectionDirectional.ush"
-
-#ifndef VSM_GENERATE_STATS
-#define VSM_GENERATE_STATS 0
-#endif // VSM_GENERATE_STATS
-
-#if VSM_GENERATE_STATS
-RWStructuredBuffer<uint> OutStatsBuffer;
-#endif // VSM_GENERATE_STATS
-
-void StatsBufferWrite(uint Index, uint Value)
-{
-#if VSM_GENERATE_STATS
-	OutStatsBuffer[Index] = Value;
-#endif
-}
-void StatsBufferInterlockedAdd(uint Index, uint Value)
-{
-#if VSM_GENERATE_STATS
-	InterlockedAdd(OutStatsBuffer[Index], Value);
-#endif
-}
+#include "Stats.ush"
 
 // Type of input data consume by the page allocation (i.e., data read from the source buffer: Gbuffer, HairStrands data, ...)
 #define INPUT_TYPE_GBUFFER 0
@@ -509,7 +489,7 @@ void CreateCachedPageMappings(uint2 ThreadId : SV_DispatchThreadID)
 			// stats for those pages who die from moving stuff
 			if (bPrevPageAllocated && bPrevDynamicPageFlag)
 			{
-				StatsBufferInterlockedAdd(2, 1);
+				StatsBufferInterlockedAdd(VSM_STAT_INVALIDATED_BY_DYNAMIC_PAGES, 1);
 			}
 
 			// If it was either rendered or cached previous frame, and not drawn to by a moving thing we may re-use
@@ -519,10 +499,10 @@ void CreateCachedPageMappings(uint2 ThreadId : SV_DispatchThreadID)
 				uint PhysicalPageIndex = VSMPhysicalPageAddressToIndex(PhysicalAddress);
 
 				FPhysicalPageMetaData PrevMetaData = PrevPhysicalPageMetaData[PhysicalPageIndex];
-				if (PrevMetaData.Flags & VSM_PHYSICAL_PAGE_FLAG_ALLOCATED != 0)
+				if ((PrevMetaData.Flags & VSM_PHYSICAL_PAGE_FLAG_ALLOCATED) != 0)
 				{
-					StatsBufferInterlockedAdd(0, 1);	// Total pages
-					StatsBufferInterlockedAdd(1, 1);	// Cached pages
+					StatsBufferInterlockedAdd(VSM_STAT_ALLOCATED_PAGES, 1);	// Total pages
+					StatsBufferInterlockedAdd(VSM_STAT_CACHED_PAGES, 1);	// Cached pages
 
 					// Keep the same cached physical page
 					// Note that we explicitly do not retain the propogated mip data as that could change this
@@ -586,7 +566,7 @@ void AllocateNewPageMappings(uint2 ThreadId : SV_DispatchThreadID, uint2 GroupId
 	uint RequestFlags = PageRequestFlags[GlobalPageOffset];
 	if (RequestFlags != 0 && OutPageFlags[GlobalPageOffset] == 0)
 	{
-		StatsBufferInterlockedAdd(0, 1);	// Total pages
+		StatsBufferInterlockedAdd(VSM_STAT_ALLOCATED_PAGES, 1);	// Total pages
 
 		// Grab a free page from the list
 		int FreeListOffset = 0;
diff --git a/Engine/Shaders/Private/VirtualShadowMaps/PageOverlap.ush b/Engine/Shaders/Private/VirtualShadowMaps/PageOverlap.ush
--- a/Engine/Shaders/Private/VirtualShadowMaps/PageOverlap.ush
+++ b/Engine/Shaders/Private/VirtualShadowMaps/PageOverlap.ush
@@ -178,4 +178,41 @@ bool IsPageVisibleHZB(uint2 vPage, uint PageFlagOffset, FPageTestScreenRect Page
 	}
 
 	return true;
+}
+
+/**
+ * Perform HZB-Test for a rectangle of pages, returning true if the Rect is visible in at least one page. 
+ * Skips unallocated pages, and returns 'visible' (true) for any page that contains any bit in the VisiblePageMask. This can be used
+ * to default un-cached pages to visible. If the PrevShadowMapID is INDEX_NONE, the function returns visible.
+ */
+bool IsVisibleValidMaskedHZB(uint PrevShadowMapID, uint VirtualShadowMapId, uint MipLevel, uint4 RectPages, FScreenRect Rect, uint VisiblePageMask)
+{
+	// Don't have an HZB to test.
+	if (PrevShadowMapID == ~0u)
+	{
+		return true;
+	}
+
+	FPageTestScreenRect HZBTestRect = SetupPageHZBRect(Rect);
+	uint PageTableLevelOffset = CalcPageTableLevelOffset(VirtualShadowMapId, MipLevel);
+	for (uint y = RectPages.y; y <= RectPages.w; y++)
+	{
+		for (uint x = RectPages.x; x <= RectPages.z; x++)
+		{
+			uint PageFlagOffset = PageTableLevelOffset + CalcPageOffsetInLevel(MipLevel, uint2(x, y));
+			uint PageFlag = PageFlags[PageFlagOffset];
+
+			// Skip un-allocated pages as we are not going to draw into them anyway
+			if ((PageFlag & VSM_ALLOCATED_FLAG) != 0)
+			{
+				// Treat pages with the VisiblePageMask as visible - can be used to select only cached pages
+				if ((PageFlag & VisiblePageMask) != 0 || IsPageVisibleHZB(uint2(x, y), PageFlagOffset, HZBTestRect))
+				{
+					return true;
+				}
+			}
+		}
+	}
+
+	return false;
 }
\ No newline at end of file
diff --git a/Engine/Shaders/Private/VirtualShadowMaps/PrintStats.usf b/Engine/Shaders/Private/VirtualShadowMaps/PrintStats.usf
--- a/Engine/Shaders/Private/VirtualShadowMaps/PrintStats.usf
+++ b/Engine/Shaders/Private/VirtualShadowMaps/PrintStats.usf
@@ -3,6 +3,7 @@
 #include "../Common.ush"
 #include "../ShaderPrintCommon.ush"
 #include "PageOverlap.ush"
+#include "Stats.ush"
 
 StructuredBuffer<uint>	InStatsBuffer;
 StructuredBuffer<uint4> AllocatedPageRectBounds;
@@ -38,78 +39,76 @@ void PrintStats()
 	float2 ScreenPos = float2(HeadlineX, TopMargin);
 
 	ScreenPos.x = ItemX;
-	ScreenPos = ShaderPrintSymbol(ScreenPos, _M_);
-	ScreenPos = ShaderPrintSymbol(ScreenPos, _A_);
-	ScreenPos = ShaderPrintSymbol(ScreenPos, _X_);
-	ScreenPos = ShaderPrintSymbol(ScreenPos, _SPC_);
-	ScreenPos = ShaderPrintSymbol(ScreenPos, _P_);
-	ScreenPos = ShaderPrintSymbol(ScreenPos, _A_);
-	ScreenPos = ShaderPrintSymbol(ScreenPos, _G_);
-	ScreenPos = ShaderPrintSymbol(ScreenPos, _E_);
-	ScreenPos = ShaderPrintSymbol(ScreenPos, _S_);
+	ScreenPos = ShaderPrintText(ScreenPos, TEXT("MAX PAGES"), FontWhite);
 	ScreenPos.x = ValueX;
 	ScreenPos = ShaderPrintLeftAlign(ScreenPos, VirtualShadowMap.MaxPhysicalPages);
 	ScreenPos = ShaderPrintNewline(ScreenPos);
 
 	ScreenPos.x = ItemX;
-	ScreenPos = ShaderPrintSymbol(ScreenPos, _P_);
-	ScreenPos = ShaderPrintSymbol(ScreenPos, _A_);
-	ScreenPos = ShaderPrintSymbol(ScreenPos, _G_);
-	ScreenPos = ShaderPrintSymbol(ScreenPos, _E_);
-	ScreenPos = ShaderPrintSymbol(ScreenPos, _S_);
+	ScreenPos = ShaderPrintText(ScreenPos, TEXT("PAGES"), FontWhite);
 	ScreenPos.x = ValueX;
-	ScreenPos = ShaderPrintLeftAlign(ScreenPos, InStatsBuffer[0]);
+	ScreenPos = ShaderPrintLeftAlign(ScreenPos, InStatsBuffer[VSM_STAT_ALLOCATED_PAGES]);
 	ScreenPos = ShaderPrintNewline(ScreenPos);
 
 	ScreenPos.x = ItemX;
-	ScreenPos = ShaderPrintSymbol(ScreenPos, _C_);
-	ScreenPos = ShaderPrintSymbol(ScreenPos, _A_);
-	ScreenPos = ShaderPrintSymbol(ScreenPos, _C_);
-	ScreenPos = ShaderPrintSymbol(ScreenPos, _H_);
-	ScreenPos = ShaderPrintSymbol(ScreenPos, _E_);
-	ScreenPos = ShaderPrintSymbol(ScreenPos, _D_);
-	ScreenPos = ShaderPrintSymbol(ScreenPos, _SPC_);
-	ScreenPos = ShaderPrintSymbol(ScreenPos, _P_);
-	ScreenPos = ShaderPrintSymbol(ScreenPos, _A_);
-	ScreenPos = ShaderPrintSymbol(ScreenPos, _G_);
-	ScreenPos = ShaderPrintSymbol(ScreenPos, _E_);
-	ScreenPos = ShaderPrintSymbol(ScreenPos, _S_);
+	ScreenPos = ShaderPrintText(ScreenPos, TEXT("CACHED PAGES"), FontWhite);
 	ScreenPos.x = ValueX;
-	ScreenPos = ShaderPrintLeftAlign(ScreenPos, InStatsBuffer[1]);
+	ScreenPos = ShaderPrintLeftAlign(ScreenPos, InStatsBuffer[VSM_STAT_CACHED_PAGES]);
 	ScreenPos = ShaderPrintNewline(ScreenPos);
 
 	ScreenPos.x = ItemX;
-	ScreenPos = ShaderPrintSymbol(ScreenPos, _C_);
-	ScreenPos = ShaderPrintSymbol(ScreenPos, _A_);
-	ScreenPos = ShaderPrintSymbol(ScreenPos, _C_);
-	ScreenPos = ShaderPrintSymbol(ScreenPos, _H_);
-	ScreenPos = ShaderPrintSymbol(ScreenPos, _E_);
-	ScreenPos = ShaderPrintSymbol(ScreenPos, _D_);
-	ScreenPos = ShaderPrintSymbol(ScreenPos, _SPC_);
-	ScreenPos = ShaderPrintSymbol(ScreenPos, _P_);
-	ScreenPos = ShaderPrintSymbol(ScreenPos, _E_);
-	ScreenPos = ShaderPrintSymbol(ScreenPos, _R_);
-	ScreenPos = ShaderPrintSymbol(ScreenPos, _C_);
-	ScreenPos = ShaderPrintSymbol(ScreenPos, _T_);
+	ScreenPos = ShaderPrintText(ScreenPos, TEXT("CACHED PERCT"), FontWhite);
 	ScreenPos.x = ValueX;
-	if (InStatsBuffer[0])
+	if (InStatsBuffer[VSM_STAT_ALLOCATED_PAGES])
 	{
-		ScreenPos = ShaderPrintLeftAlign(ScreenPos, (InStatsBuffer[1] * 100) / InStatsBuffer[0]);
+		ScreenPos = ShaderPrintLeftAlign(ScreenPos, (InStatsBuffer[VSM_STAT_CACHED_PAGES] * 100) / InStatsBuffer[VSM_STAT_ALLOCATED_PAGES]);
 	}
 	ScreenPos = ShaderPrintNewline(ScreenPos);
 
 	ScreenPos.x = ItemX;
-	ScreenPos = ShaderPrintSymbol(ScreenPos, _D_);
-	ScreenPos = ShaderPrintSymbol(ScreenPos, _Y_);
-	ScreenPos = ShaderPrintSymbol(ScreenPos, _N_);
-	ScreenPos = ShaderPrintSymbol(ScreenPos, _SPC_);
-	ScreenPos = ShaderPrintSymbol(ScreenPos, _I_);
-	ScreenPos = ShaderPrintSymbol(ScreenPos, _N_);
-	ScreenPos = ShaderPrintSymbol(ScreenPos, _V_);
-	ScreenPos = ShaderPrintSymbol(ScreenPos, _A_);
-	ScreenPos = ShaderPrintSymbol(ScreenPos, _L_);
+	ScreenPos = ShaderPrintText(ScreenPos, TEXT("DYN INVAL"), FontWhite);
+	ScreenPos.x = ValueX;
+	ScreenPos = ShaderPrintLeftAlign(ScreenPos, InStatsBuffer[VSM_STAT_INVALIDATED_BY_DYNAMIC_PAGES]);
+	ScreenPos = ShaderPrintNewline(ScreenPos);
+
+	ScreenPos.x = HeadlineX;
+	ScreenPos = ShaderPrintText(ScreenPos, TEXT("Non-Nanite Instances"), FontCyan);
+	ScreenPos = ShaderPrintNewline(ScreenPos);
+
+	ScreenPos.x = ItemX;
+	ScreenPos = ShaderPrintText(ScreenPos, TEXT("Total"), FontCyan);
+	ScreenPos.x = ValueX;
+	ScreenPos = ShaderPrintLeftAlign(ScreenPos, InStatsBuffer[VSM_STAT_NON_NANITE_INSTANCES_TOTAL]);
+	ScreenPos = ShaderPrintNewline(ScreenPos);
+
+	ScreenPos.x = ItemX;
+	ScreenPos = ShaderPrintText(ScreenPos, TEXT("Drawn"), FontCyan);
+	ScreenPos.x = ValueX;
+	ScreenPos = ShaderPrintLeftAlign(ScreenPos, InStatsBuffer[VSM_STAT_NON_NANITE_INSTANCES_DRAWN]);
+	ScreenPos = ShaderPrintNewline(ScreenPos);
+
+	ScreenPos.x = ItemX;
+	ScreenPos = ShaderPrintText(ScreenPos, TEXT("HZB Culled"), FontCyan);
+	ScreenPos.x = ValueX;
+	ScreenPos = ShaderPrintLeftAlign(ScreenPos, InStatsBuffer[VSM_STAT_NON_NANITE_INSTANCES_HZB_CULLED]);
+	ScreenPos = ShaderPrintNewline(ScreenPos);
+
+	ScreenPos.x = ItemX;
+	ScreenPos = ShaderPrintText(ScreenPos, TEXT("Page Mask Culled"), FontCyan);
+	ScreenPos.x = ValueX;
+	ScreenPos = ShaderPrintLeftAlign(ScreenPos, InStatsBuffer[VSM_STAT_NON_NANITE_INSTANCES_PAGE_MASK_CULLED]);
+	ScreenPos = ShaderPrintNewline(ScreenPos);
+
+	ScreenPos.x = ItemX;
+	ScreenPos = ShaderPrintText(ScreenPos, TEXT("Empty Rect Culled"), FontCyan);
+	ScreenPos.x = ValueX;
+	ScreenPos = ShaderPrintLeftAlign(ScreenPos, InStatsBuffer[VSM_STAT_NON_NANITE_INSTANCES_EMPTY_RECT_CULLED]);
+	ScreenPos = ShaderPrintNewline(ScreenPos);
+
+	ScreenPos.x = ItemX;
+	ScreenPos = ShaderPrintText(ScreenPos, TEXT("Frustum Culled"), FontCyan);
 	ScreenPos.x = ValueX;
-	ScreenPos = ShaderPrintLeftAlign(ScreenPos, InStatsBuffer[2]);
+	ScreenPos = ShaderPrintLeftAlign(ScreenPos, InStatsBuffer[VSM_STAT_NON_NANITE_INSTANCES_FRUSTUM_CULLED]);
 	ScreenPos = ShaderPrintNewline(ScreenPos);
 }
 
diff --git a/Engine/Shaders/Private/VirtualShadowMaps/Stats.ush b/Engine/Shaders/Private/VirtualShadowMaps/Stats.ush
new file mode 100644
--- /dev/null
+++ b/Engine/Shaders/Private/VirtualShadowMaps/Stats.ush
@@ -0,0 +1,49 @@
+// Copyright Epic Games, Inc. All Rights Reserved.
+
+/*=============================================================================
+	Stats.ush: 
+=============================================================================*/
+
+#include "../Common.ush"
+#include "../WaveOpUtil.ush"
+
+
+#define VSM_STAT_ALLOCATED_PAGES            0
+#define VSM_STAT_CACHED_PAGES                    1
+#define VSM_STAT_INVALIDATED_BY_DYNAMIC_PAGES    2
+#define VSM_STAT_NUM_SMS                         3
+#define VSM_STAT_NON_NANITE_INSTANCES_TOTAL      4
+#define VSM_STAT_NON_NANITE_INSTANCES_DRAWN      5
+#define VSM_STAT_NON_NANITE_INSTANCES_HZB_CULLED 6
+#define VSM_STAT_NON_NANITE_INSTANCES_PAGE_MASK_CULLED 7
+#define VSM_STAT_NON_NANITE_INSTANCES_EMPTY_RECT_CULLED 8
+#define VSM_STAT_NON_NANITE_INSTANCES_FRUSTUM_CULLED 9
+
+#ifndef VSM_GENERATE_STATS
+#define VSM_GENERATE_STATS 0
+#endif // VSM_GENERATE_STATS
+
+#if VSM_GENERATE_STATS
+RWStructuredBuffer<uint> OutStatsBuffer;
+#endif // VSM_GENERATE_STATS
+
+void StatsBufferWrite(uint Index, uint Value)
+{
+#if VSM_GENERATE_STATS
+	OutStatsBuffer[Index] = Value;
+#endif
+}
+
+void StatsBufferInterlockedAdd(uint Index, uint Value)
+{
+#if VSM_GENERATE_STATS
+	InterlockedAdd(OutStatsBuffer[Index], Value);
+#endif
+}
+
+void StatsBufferInterlockedInc(uint Index)
+{
+#if VSM_GENERATE_STATS
+	WaveInterlockedAddScalar(OutStatsBuffer[Index], 1U);
+#endif
+}
diff --git a/Engine/Source/Runtime/Renderer/Private/ShadowDepthRendering.cpp b/Engine/Source/Runtime/Renderer/Private/ShadowDepthRendering.cpp
--- a/Engine/Source/Runtime/Renderer/Private/ShadowDepthRendering.cpp
+++ b/Engine/Source/Runtime/Renderer/Private/ShadowDepthRendering.cpp
@@ -1840,20 +1840,11 @@ void FSceneRenderer::RenderShadowDepthMaps(FRDGBuilder& GraphBuilder, FInstanceC
 							Params.Flags = 0;
 
 							// TODO: Clean this up - could be stored in a single structure for the whole clipmap
-							int32 AbsoluteClipmapLevel = Clipmap->GetClipmapLevel(ClipmapLevelIndex);		// NOTE: Can be negative!
-							int32 ClipmapLevelKey = AbsoluteClipmapLevel + 128;
-							check(ClipmapLevelKey > 0 && ClipmapLevelKey < 256);
-							int32 HZBKey = Clipmap->GetLightSceneInfo().Id + (ClipmapLevelKey << 24);
+							int32 HZBKey = Clipmap->GetHZBKey(ClipmapLevelIndex);
 
 							if (PrevHZBPhysical)
 							{
-								FVirtualShadowMapHZBMetadata* PrevHZBMeta = CacheManager->PrevBuffers.HZBMetadata.Find(HZBKey);
-								if (PrevHZBMeta)
-							{
-									Params.PrevTargetLayerIndex = PrevHZBMeta->TargetLayerIndex;
-									Params.PrevViewMatrices = PrevHZBMeta->ViewMatrices;
-								Params.Flags = VIEW_FLAG_HZBTEST;
-							}
+								CacheManager->SetHZBViewParams(HZBKey, Params);
 							}
 
 							// If we're going to generate a new HZB this frame, save the associated metadata
@@ -1905,13 +1896,7 @@ void FSceneRenderer::RenderShadowDepthMaps(FRDGBuilder& GraphBuilder, FInstanceC
 							int32 HZBKey = ProjectedShadowInfo->GetLightSceneInfo().Id + (i << 24);
 							if (PrevHZBPhysical)
 							{
-								FVirtualShadowMapHZBMetadata* PrevHZBMeta = CacheManager->PrevBuffers.HZBMetadata.Find(HZBKey);
-								if (PrevHZBMeta)
-							{
-									Params.PrevTargetLayerIndex = PrevHZBMeta->TargetLayerIndex;
-									Params.PrevViewMatrices = PrevHZBMeta->ViewMatrices;
-								Params.Flags = VIEW_FLAG_HZBTEST;
-							}
+								CacheManager->SetHZBViewParams(HZBKey, Params);
 							}
 
 							// If we're going to generate a new HZB this frame, save the associated metadata
diff --git a/Engine/Source/Runtime/Renderer/Private/VirtualShadowMaps/VirtualShadowMapArray.cpp b/Engine/Source/Runtime/Renderer/Private/VirtualShadowMaps/VirtualShadowMapArray.cpp
--- a/Engine/Source/Runtime/Renderer/Private/VirtualShadowMaps/VirtualShadowMapArray.cpp
+++ b/Engine/Source/Runtime/Renderer/Private/VirtualShadowMaps/VirtualShadowMapArray.cpp
@@ -134,6 +134,12 @@ FAutoConsoleVariableRef CVarEnableNonNaniteVSM(
 	ECVF_RenderThreadSafe | ECVF_ReadOnly
 );
 
+static TAutoConsoleVariable<int32> CVarNonNaniteVsmUseHzb(
+	TEXT("r.Shadow.Virtual.NonNanite.UseHZB"),
+	2,
+	TEXT("Cull Non-Nanite instances using HZB. If set to 2, attempt to use Nanite-HZB from the current frame."),
+	ECVF_RenderThreadSafe);
+
 #if !UE_BUILD_SHIPPING
 bool GDumpVSMLightNames = false;
 void DumpVSMLightNames()
@@ -1211,6 +1217,16 @@ class FVirtualSmPrintStatsCS : public FVirtualPageManagementShader
 		SHADER_PARAMETER_RDG_BUFFER_SRV(StructuredBuffer< uint >, InStatsBuffer)
 		SHADER_PARAMETER_RDG_BUFFER_SRV(StructuredBuffer<FIntVector4>, AllocatedPageRectBounds)
 	END_SHADER_PARAMETER_STRUCT()
+
+
+	static void ModifyCompilationEnvironment(const FGlobalShaderPermutationParameters& Parameters, FShaderCompilerEnvironment& OutEnvironment)
+	{
+		FVirtualPageManagementShader::ModifyCompilationEnvironment(Parameters, OutEnvironment);
+		// Disable optimizations as shader print causes long compile times
+		OutEnvironment.CompilerFlags.Add(CFLAG_SkipOptimizations);
+	}
+
+		
 };
 IMPLEMENT_GLOBAL_SHADER(FVirtualSmPrintStatsCS, "/Engine/Private/VirtualShadowMaps/PrintStats.usf", "PrintStats", SF_Compute);
 
@@ -1369,7 +1385,9 @@ class FCullPerPageDrawCommandsCs : public FGlobalShader
 
 	class FNearClipDim		: SHADER_PERMUTATION_BOOL( "NEAR_CLIP" );
 	class FLoopOverViewsDim	: SHADER_PERMUTATION_BOOL( "LOOP_OVER_VIEWS" );
-	using FPermutationDomain = TShaderPermutationDomain< FNearClipDim, FLoopOverViewsDim >;
+	class FUseHzbDim : SHADER_PERMUTATION_BOOL("USE_HZB_OCCLUSION");
+	class FGenerateStatsDim : SHADER_PERMUTATION_BOOL("VSM_GENERATE_STATS");
+	using FPermutationDomain = TShaderPermutationDomain< FUseHzbDim, FNearClipDim, FLoopOverViewsDim, FGenerateStatsDim >;
 
 public:
 	static bool ShouldCompilePermutation(const FGlobalShaderPermutationParameters& Parameters)
@@ -1413,8 +1431,16 @@ public:
 		SHADER_PARAMETER_RDG_BUFFER_UAV(RWBuffer<uint>, DrawIndirectArgsBufferOut)
 		SHADER_PARAMETER_RDG_BUFFER_UAV(RWStructuredBuffer<uint>, VisibleInstanceCountBufferOut)
 
+		SHADER_PARAMETER_RDG_BUFFER_SRV(StructuredBuffer< uint >, ShadowHZBPageTable)
+		SHADER_PARAMETER_RDG_TEXTURE(Texture2D, HZBTexture)
+		SHADER_PARAMETER_SAMPLER(SamplerState, HZBSampler)
+		SHADER_PARAMETER(FVector2D, HZBSize)
+		SHADER_PARAMETER(uint32, HZBMode)
+
 		SHADER_PARAMETER_RDG_BUFFER_UAV(RWStructuredBuffer<uint>, OutInvalidatingInstances)
 		SHADER_PARAMETER(uint32, NumInvalidatingInstanceSlots)
+
+		SHADER_PARAMETER_RDG_BUFFER_UAV(RWStructuredBuffer< uint >, OutStatsBuffer)
 	END_SHADER_PARAMETER_STRUCT()
 };
 IMPLEMENT_GLOBAL_SHADER(FCullPerPageDrawCommandsCs, "/Engine/Private/VirtualShadowMaps/BuildPerPageDrawCommands.usf", "CullPerPageDrawCommandsCs", SF_Compute);
@@ -1507,6 +1533,26 @@ void FVirtualShadowMapArray::RenderVirtualShadowMapsHw(FRDGBuilder& GraphBuilder
 	RDG_EVENT_SCOPE(GraphBuilder, "RenderVirtualShadowMaps(Non-Nanite)");
 
 	FGPUScene& GPUScene = Scene.GPUScene;
+
+	FRDGBufferSRVRef PrevPageTableRDGSRV = CacheManager->IsValid() ? GraphBuilder.CreateSRV(GraphBuilder.RegisterExternalBuffer(CacheManager->PrevBuffers.PageTable, TEXT("Shadow.Virtual.PrevPageTable"))) : nullptr;
+
+	int32 HZBMode = CVarNonNaniteVsmUseHzb.GetValueOnRenderThread();
+
+	auto InitHZB = [&]()->FRDGTextureRef
+	{
+		if (HZBMode == 1 && CacheManager->IsValid())
+		{
+			return GraphBuilder.RegisterExternalTexture(CacheManager->PrevBuffers.HZBPhysical);
+		}
+
+		if (HZBMode == 2 && HZBPhysical != nullptr)
+		{
+			return HZBPhysical;
+		}
+		return nullptr;
+	};
+	const FRDGTextureRef HZBTexture = InitHZB();
+
 	// TODO: create a VirtualClipMapMeshRenderInfo - or something - that contains the required info:
 	//       The View used to set up (to get all the usual crap), The MeshCommandPass (to perform setup & culling), Culling volume?, Ref to Clipmap itself.
 
@@ -1553,6 +1599,11 @@ void FVirtualShadowMapArray::RenderVirtualShadowMapsHw(FRDGBuilder& GraphBuilder
 				Params.PrevTargetLayerIndex = INDEX_NONE;
 				Params.PrevViewMatrices = Params.ViewMatrices;
 
+				if (HZBTexture)
+				{
+					CacheManager->SetHZBViewParams(Clipmap->GetHZBKey(ClipmapLevelIndex), Params);
+				}
+
 				VirtualShadowViews.Add(Nanite::CreatePackedView(Params));
 
 				// This is used to mark all the referenced physical pages as being updated. 
@@ -1583,6 +1634,11 @@ void FVirtualShadowMapArray::RenderVirtualShadowMapsHw(FRDGBuilder& GraphBuilder
 				Params.TargetLayerIndex = VirtualShadowMap->ID;
 				Params.ViewMatrices = ProjectedShadowInfo->GetShadowDepthRenderingViewMatrices(i, true);
 
+				if (HZBTexture)
+				{
+					CacheManager->SetHZBViewParams(ProjectedShadowInfo->GetLightSceneInfo().Id + (i << 24), Params);
+				}
+
 				VirtualShadowViews.Add(Nanite::CreatePackedView(Params));
 
 				if (VirtualShadowMap->VirtualShadowMapCacheEntry)
@@ -1672,9 +1728,27 @@ void FVirtualShadowMapArray::RenderVirtualShadowMapsHw(FRDGBuilder& GraphBuilder
 				PassParameters->OutInvalidatingInstances = GraphBuilder.CreateUAV(InvalidatingInstancesRDG);
 				PassParameters->NumInvalidatingInstanceSlots = NumInvalidatingInstanceSlots;
 
+				if (HZBTexture)
+				{
+					// Mode 2 uses the current frame HZB  & page table.
+					PassParameters->ShadowHZBPageTable = HZBMode == 2 ? GraphBuilder.CreateSRV(PageTableRDG) : PrevPageTableRDGSRV;
+					PassParameters->HZBTexture = HZBTexture;
+					PassParameters->HZBSize = HZBTexture->Desc.Extent;
+					PassParameters->HZBSampler = TStaticSamplerState< SF_Point, AM_Clamp, AM_Clamp, AM_Clamp >::GetRHI();
+					PassParameters->HZBMode = HZBMode;
+				}
+
+				bool bGenerateStats = StatsBufferRDG != nullptr;
+				if (bGenerateStats)
+				{
+					PassParameters->OutStatsBuffer = GraphBuilder.CreateUAV(StatsBufferRDG);
+				}
+
 				FCullPerPageDrawCommandsCs::FPermutationDomain PermutationVector;
 				PermutationVector.Set< FCullPerPageDrawCommandsCs::FNearClipDim >( !Clipmap.IsValid() );
 				PermutationVector.Set< FCullPerPageDrawCommandsCs::FLoopOverViewsDim >( Clipmap.IsValid() );
+				PermutationVector.Set< FCullPerPageDrawCommandsCs::FUseHzbDim >(HZBTexture != nullptr);
+				PermutationVector.Set< FCullPerPageDrawCommandsCs::FGenerateStatsDim >(bGenerateStats);
 
 				auto ComputeShader = ShaderMap->GetShader<FCullPerPageDrawCommandsCs>( PermutationVector );
 
diff --git a/Engine/Source/Runtime/Renderer/Private/VirtualShadowMaps/VirtualShadowMapArray.h b/Engine/Source/Runtime/Renderer/Private/VirtualShadowMaps/VirtualShadowMapArray.h
--- a/Engine/Source/Runtime/Renderer/Private/VirtualShadowMaps/VirtualShadowMapArray.h
+++ b/Engine/Source/Runtime/Renderer/Private/VirtualShadowMaps/VirtualShadowMapArray.h
@@ -263,12 +263,9 @@ public:
 	FRDGTextureRef HZBPhysical = nullptr;
 	TMap<int32, FVirtualShadowMapHZBMetadata> HZBMetadata;
 
-	static constexpr uint32 NumStats = 5;
-	// 0 - allocated pages
-	// 1 - re-usable pages
-	// 2 - Touched by dynamic
-	// 3 - NumSms
-	// 4 - RandRobin invalidated
+	// See Engine\Shaders\Private\VirtualShadowMaps\Stats.ush for definitions of the different stat indexes
+	static constexpr uint32 NumStats = 16;
+
 	FRDGBufferRef StatsBufferRDG = nullptr;
 
 	TRefCountPtr<IPooledRenderTarget> DebugVisualizationOutput;
diff --git a/Engine/Source/Runtime/Renderer/Private/VirtualShadowMaps/VirtualShadowMapCacheManager.cpp b/Engine/Source/Runtime/Renderer/Private/VirtualShadowMaps/VirtualShadowMapCacheManager.cpp
--- a/Engine/Source/Runtime/Renderer/Private/VirtualShadowMaps/VirtualShadowMapCacheManager.cpp
+++ b/Engine/Source/Runtime/Renderer/Private/VirtualShadowMaps/VirtualShadowMapCacheManager.cpp
@@ -354,26 +354,36 @@ void FVirtualShadowMapArrayCacheManager::ExtractStats(FRDGBuilder& GraphBuilder,
 		ensure(FileToLogTo);
 		if (FileToLogTo)
 		{
-			static const FString StatNames[FVirtualShadowMapArray::NumStats] =
+			static const FString StatNames[] =
 			{
 				TEXT("Allocated"),
 				TEXT("Cached"),
 				TEXT("Dynamic"),
 				TEXT("NumSms"),
-				TEXT("RandRobin"),
+				TEXT("NonNaniteInstances"),
+				TEXT("NonNaniteInstancesDrawn"),
+				TEXT("NonNaniteInstancesHZBCulled"),
+				TEXT("NonNaniteInstancesPageMaskCulled"),
+				TEXT("NonNaniteInstancesEmptyRectCulled"),
+				TEXT("NonNaniteInstancesFrustumCulled"),
 			};
 
-
 			// Print header
 			FString StringToPrint;
-			for (const FString &StatName : StatNames)
+			for (int32 Index = 0; Index < FVirtualShadowMapArray::NumStats; ++Index)
 			{
 				if (!StringToPrint.IsEmpty())
 				{
 					StringToPrint += TEXT(",");
 				}
-
-				StringToPrint += StatName;
+				if (Index < int32(UE_ARRAY_COUNT(StatNames)))
+				{
+					StringToPrint.Append(StatNames[Index]);
+				}
+				else
+				{
+					StringToPrint.Appendf(TEXT("Stat_%d"), Index);
+				}
 			}
 
 			StringToPrint += TEXT("\n");
@@ -581,6 +591,17 @@ TRDGUniformBufferRef<FVirtualShadowMapUniformParameters> FVirtualShadowMapArrayC
 	return GraphBuilder.CreateUniformBuffer(VersionedParameters);
 }
 
+void FVirtualShadowMapArrayCacheManager::SetHZBViewParams(int32 HZBKey, Nanite::FPackedViewParams& OutParams)
+{
+	FVirtualShadowMapHZBMetadata* PrevHZBMeta = PrevBuffers.HZBMetadata.Find(HZBKey);
+	if (PrevHZBMeta)
+	{
+		OutParams.PrevTargetLayerIndex = PrevHZBMeta->TargetLayerIndex;
+		OutParams.PrevViewMatrices = PrevHZBMeta->ViewMatrices;
+		OutParams.Flags = VIEW_FLAG_HZBTEST;
+	}
+}
+
 static void SetupCommonParameters(FRDGBuilder& GraphBuilder, FVirtualShadowMapArrayCacheManager* CacheManager, int32 TotalInstanceCount, const FGPUScene& GPUScene, 
 	FVirtualSmInvalidateInstancePagesCS::FParameters& OutPassParameters,
 	FVirtualSmInvalidateInstancePagesCS::FPermutationDomain &OutPermutationVector)
diff --git a/Engine/Source/Runtime/Renderer/Private/VirtualShadowMaps/VirtualShadowMapCacheManager.h b/Engine/Source/Runtime/Renderer/Private/VirtualShadowMaps/VirtualShadowMapCacheManager.h
--- a/Engine/Source/Runtime/Renderer/Private/VirtualShadowMaps/VirtualShadowMapCacheManager.h
+++ b/Engine/Source/Runtime/Renderer/Private/VirtualShadowMaps/VirtualShadowMapCacheManager.h
@@ -126,6 +126,8 @@ public:
 	FVirtualShadowMapArrayFrameData PrevBuffers;
 	FVirtualShadowMapUniformParameters PrevUniformParameters;
 		
+	void SetHZBViewParams(int32 HZBKey, Nanite::FPackedViewParams& OutParams);
+
 private:
 	// Must match shader...
 	struct FInstanceSceneDataRange
diff --git a/Engine/Source/Runtime/Renderer/Private/VirtualShadowMaps/VirtualShadowMapClipmap.h b/Engine/Source/Runtime/Renderer/Private/VirtualShadowMaps/VirtualShadowMapClipmap.h
--- a/Engine/Source/Runtime/Renderer/Private/VirtualShadowMaps/VirtualShadowMapClipmap.h
+++ b/Engine/Source/Runtime/Renderer/Private/VirtualShadowMaps/VirtualShadowMapClipmap.h
@@ -49,6 +49,14 @@ public:
 		return LightSceneInfo;
 	}
 
+	int32 GetHZBKey(int32 ClipmapLevelIndex) const
+	{
+		int32 AbsoluteClipmapLevel = GetClipmapLevel(ClipmapLevelIndex);		// NOTE: Can be negative!
+		int32 ClipmapLevelKey = AbsoluteClipmapLevel + 128;
+		check(ClipmapLevelKey > 0 && ClipmapLevelKey < 256);
+		return GetLightSceneInfo().Id + (ClipmapLevelKey << 24);
+	}
+
 	FVirtualShadowMapProjectionShaderData GetProjectionShaderData(int32 ClipmapIndex) const;
 
 	FVector GetWorldOrigin() const
