Temporary fix for various issues with ray tracing culling.

diff --git a/Engine/Source/Runtime/Renderer/Private/SceneVisibility.cpp b/Engine/Source/Runtime/Renderer/Private/SceneVisibility.cpp
--- a/Engine/Source/Runtime/Renderer/Private/SceneVisibility.cpp
+++ b/Engine/Source/Runtime/Renderer/Private/SceneVisibility.cpp
@@ -391,8 +391,17 @@ bool FViewInfo::IsDistanceCulled_AnyThread(float DistanceSquared, float MinDrawD
 	const float MaxDrawDistanceScale = GetCachedScalabilityCVars().ViewDistanceScale;
 	const float FadeRadius = GDisableLODFade ? 0.0f : GDistanceFadeMaxTravel;
 	const float MaxDrawDistance = InMaxDrawDistance * MaxDrawDistanceScale;
+
+	bool bHasMaxDrawDistance = InMaxDrawDistance != FLT_MAX;
+	bool bHasMinDrawDistance = InMaxDrawDistance > 0;
 	bOutMayBeFading = false;
 
+
+	if (!bHasMaxDrawDistance && !bHasMinDrawDistance)
+	{
+		return false;
+	}
+
 	// If cull distance is disabled, always show (except foliage)
 	if (Family->EngineShowFlags.DistanceCulledPrimitives && !PrimitiveSceneInfo->Proxy->IsDetailMesh())
 	{
@@ -400,13 +409,13 @@ bool FViewInfo::IsDistanceCulled_AnyThread(float DistanceSquared, float MinDrawD
 	}
 
 	// The primitive is always culled if it exceeds the max fade distance.
-	if (DistanceSquared > FMath::Square(MaxDrawDistance + FadeRadius) || DistanceSquared < FMath::Square(MinDrawDistance))
+	if ((bHasMaxDrawDistance && DistanceSquared > FMath::Square(MaxDrawDistance + FadeRadius)) || (bHasMaxDrawDistance && DistanceSquared < FMath::Square(MinDrawDistance)))
 	{
 		return true;
 	}
 
-	const bool bDistanceCulled = (DistanceSquared > FMath::Square(MaxDrawDistance));
-	const bool bMayBeFading = (DistanceSquared > FMath::Square(MaxDrawDistance - FadeRadius));
+	const bool bDistanceCulled = bHasMaxDrawDistance && (DistanceSquared > FMath::Square(MaxDrawDistance));
+	const bool bMayBeFading = bHasMaxDrawDistance && (DistanceSquared > FMath::Square(MaxDrawDistance - FadeRadius));
 	
 	if (!GDisableLODFade && bMayBeFading && State != NULL && !bDisableDistanceBasedFadeTransitions && PrimitiveSceneInfo->Proxy->IsUsingDistanceCullFade())
 	{
