Copy paste logic error in previous change 17471856.

diff --git a/Engine/Source/Runtime/Renderer/Private/SceneVisibility.cpp b/Engine/Source/Runtime/Renderer/Private/SceneVisibility.cpp
--- a/Engine/Source/Runtime/Renderer/Private/SceneVisibility.cpp
+++ b/Engine/Source/Runtime/Renderer/Private/SceneVisibility.cpp
@@ -409,7 +409,7 @@ bool FViewInfo::IsDistanceCulled_AnyThread(float DistanceSquared, float MinDrawD
 	}
 
 	// The primitive is always culled if it exceeds the max fade distance.
-	if ((bHasMaxDrawDistance && DistanceSquared > FMath::Square(MaxDrawDistance + FadeRadius)) || (bHasMaxDrawDistance && DistanceSquared < FMath::Square(MinDrawDistance)))
+	if ((bHasMaxDrawDistance && DistanceSquared > FMath::Square(MaxDrawDistance + FadeRadius)) || (bHasMinDrawDistance && DistanceSquared < FMath::Square(MinDrawDistance)))
 	{
 		return true;
 	}
