Temporary fix for incorrect ray tracing culling which is fall out from 17459502.
Submitting for brandon.dawson

diff --git a/Engine/Source/Runtime/Renderer/Private/DeferredShadingRenderer.cpp b/Engine/Source/Runtime/Renderer/Private/DeferredShadingRenderer.cpp
--- a/Engine/Source/Runtime/Renderer/Private/DeferredShadingRenderer.cpp
+++ b/Engine/Source/Runtime/Renderer/Private/DeferredShadingRenderer.cpp
@@ -655,12 +655,15 @@ bool FDeferredShadingSceneRenderer::GatherRayTracingWorldInstancesForView(FRDGBu
 			{
 				continue;
 			}
-				
+			
+			// Disabling for now, this is not correct for ray tracing visibility
+			/*
 			// Check if the primitive has been culled already during frustum culling
 			if (!View.PrimitiveVisibilityMap[PrimitiveIndex])
 			{
 				continue;
 			}
+			*/
 
 			// Marked visible and used after point, check if streaming then mark as used in the TLAS (so it can be streamed in)
 			if (EnumHasAnyFlags(Scene->PrimitiveRayTracingFlags[PrimitiveIndex], ERayTracingPrimitiveFlags::Streaming))
