This CL is a merge of:
* Change hair cards/meshes vertex factory to use GPU scene data.
* Revert hair cards factory GPU support for now, as it produces incorrect motion vectors.

So it only checked in the non-reverted changes.

diff --git a/Engine/Plugins/Runtime/HairStrands/Source/HairStrandsCore/Private/GroomComponent.cpp b/Engine/Plugins/Runtime/HairStrands/Source/HairStrandsCore/Private/GroomComponent.cpp
--- a/Engine/Plugins/Runtime/HairStrands/Source/HairStrandsCore/Private/GroomComponent.cpp
+++ b/Engine/Plugins/Runtime/HairStrands/Source/HairStrandsCore/Private/GroomComponent.cpp
@@ -859,7 +859,7 @@ public:
 
 		BatchElement.FirstIndex = 0;
 		BatchElement.NumInstances = 1;
-		BatchElement.PrimitiveIdMode = PrimID_ForceZero;
+		BatchElement.PrimitiveIdMode = GeometryType == EHairGeometryType::Strands ? PrimID_ForceZero : PrimID_DynamicPrimitiveShaderData;
 		if (bUseCulling)
 		{
 			BatchElement.NumPrimitives = 0;
diff --git a/Engine/Source/Runtime/Renderer/Private/SceneVisibility.cpp b/Engine/Source/Runtime/Renderer/Private/SceneVisibility.cpp
--- a/Engine/Source/Runtime/Renderer/Private/SceneVisibility.cpp
+++ b/Engine/Source/Runtime/Renderer/Private/SceneVisibility.cpp
@@ -2287,11 +2287,6 @@ struct FRelevancePacket
 			const bool bTranslucentRelevance = ViewRelevance.HasTranslucency();
 
 			const bool bHairStrandsEnabled = ViewRelevance.bHairStrands && IsHairStrandsEnabled(EHairStrandsShaderType::All, Scene->GetShaderPlatform());
-			if (!bEditorRelevance && bHairStrandsEnabled)
-			{
-				++NumVisibleDynamicPrimitives;
-				OutHasDynamicMeshElementsMasks[BitIndex] |= ViewBit;
-			}
 
 			if (View.bIsReflectionCapture && !PrimitiveSceneInfo->Proxy->IsVisibleInReflectionCaptures())
 			{
@@ -2342,6 +2337,12 @@ struct FRelevancePacket
 					VisibleDynamicPrimitivesWithSimpleLights.AddPrim(PrimitiveSceneInfo);
 				}
 			}
+			else if (bHairStrandsEnabled)
+			{
+				// Strands MeshElement
+				++NumVisibleDynamicPrimitives;
+				OutHasDynamicMeshElementsMasks[BitIndex] |= ViewBit;
+			}
 
 			if (bTranslucentRelevance && !bEditorRelevance && ViewRelevance.bRenderInMainPass)
 			{
