[Renderer] - HACK HACK HACK - hacks to avoid blocking sequencer issues - these should be backed out asap and replaced with real fixes

diff --git a/Engine/Source/Runtime/Renderer/Private/DeferredShadingRenderer.cpp b/Engine/Source/Runtime/Renderer/Private/DeferredShadingRenderer.cpp
--- a/Engine/Source/Runtime/Renderer/Private/DeferredShadingRenderer.cpp
+++ b/Engine/Source/Runtime/Renderer/Private/DeferredShadingRenderer.cpp
@@ -1168,13 +1168,16 @@ bool FDeferredShadingSceneRenderer::GatherRayTracingWorldInstancesForView(FRDGBu
 
 				// At the moment we only support SM & ISMs on this path
 				check(EnumHasAnyFlags(Flags, ERayTracingPrimitiveFlags::CacheMeshCommands));
-				for (int32 CommandIndex : SceneInfo->CachedRayTracingMeshCommandIndicesPerLOD[0])
+				if (SceneInfo->CachedRayTracingMeshCommandIndicesPerLOD.Num() > 0 && SceneInfo->CachedRayTracingMeshCommandIndicesPerLOD[0].Num() > 0)
 				{
-					FVisibleRayTracingMeshCommand NewVisibleMeshCommand;
+					for (int32 CommandIndex : SceneInfo->CachedRayTracingMeshCommandIndicesPerLOD[0])
+					{
+						FVisibleRayTracingMeshCommand NewVisibleMeshCommand;
 
-					NewVisibleMeshCommand.RayTracingMeshCommand = &Scene->CachedRayTracingMeshCommands[CommandIndex];
-					NewVisibleMeshCommand.InstanceIndex = NewInstanceIndex;
-					View.VisibleRayTracingMeshCommands.Add(NewVisibleMeshCommand);
+						NewVisibleMeshCommand.RayTracingMeshCommand = &Scene->CachedRayTracingMeshCommands[CommandIndex];
+						NewVisibleMeshCommand.InstanceIndex = NewInstanceIndex;
+						View.VisibleRayTracingMeshCommands.Add(NewVisibleMeshCommand);
+					}
 				}
 
 				checkf(SceneInfo->CachedRayTracingInstance.GeometryRHI, TEXT("Ray tracing instance must have a valid geometry."));
diff --git a/Engine/Source/Runtime/Renderer/Private/GPUScene.cpp b/Engine/Source/Runtime/Renderer/Private/GPUScene.cpp
--- a/Engine/Source/Runtime/Renderer/Private/GPUScene.cpp
+++ b/Engine/Source/Runtime/Renderer/Private/GPUScene.cpp
@@ -936,53 +936,55 @@ void FGPUScene::UploadGeneral(FRHICommandListImmediate& RHICmdList, FScene *Scen
 
 							const TArray<uint32>& PassMaterialSlots = PrimitiveSceneInfo->NaniteMaterialSlots[NaniteMeshPass];
 							const TArray<Nanite::FSceneProxyBase::FMaterialSection>& PassMaterials = NaniteSceneProxy->GetMaterialSections();
-							check(PassMaterials.Num() == PassMaterialSlots.Num());
-
-							const uint32 TableEntryCount = uint32(NaniteSceneProxy->GetMaterialMaxIndex() + 1);
-							check(TableEntryCount >= uint32(PassMaterials.Num()));
-						#if WITH_EDITOR
-							const uint32 HitProxyEntryCount = (NaniteMeshPass == ENaniteMeshPass::BasePass) ? TableEntryCount : NANITE_MAX_MATERIALS;
-						#endif
-
-							void* MaterialSlotRange = nullptr;
-						#if WITH_EDITOR
-							void* HitProxyTable = nullptr;
-						#endif
+							if (PassMaterials.Num() == PassMaterialSlots.Num())
 							{
-								FNaniteMaterialCommandsLock NaniteLock(NaniteMaterials, SLT_Write);
-								MaterialSlotRange = NaniteMaterials.GetMaterialSlotPtr(UploadInfo.PrimitiveID, TableEntryCount);
-							#if WITH_EDITOR
-								HitProxyTable = NaniteMaterials.GetHitProxyTablePtr(UploadInfo.PrimitiveID, HitProxyEntryCount);
-							#endif
-							}
 
-							uint32* MaterialSlots = static_cast<uint32*>(MaterialSlotRange);
-							for (int32 Entry = 0; Entry < PassMaterialSlots.Num(); ++Entry)
-							{
-								MaterialSlots[PassMaterials[Entry].MaterialIndex] = PassMaterialSlots[Entry];
-							}
+								const uint32 TableEntryCount = uint32(NaniteSceneProxy->GetMaterialMaxIndex() + 1);
+								check(TableEntryCount >= uint32(PassMaterials.Num()));
+#if WITH_EDITOR
+								const uint32 HitProxyEntryCount = (NaniteMeshPass == ENaniteMeshPass::BasePass) ? TableEntryCount : NANITE_MAX_MATERIALS;
+#endif
 
-						#if WITH_EDITOR
-							if (NaniteMeshPass == ENaniteMeshPass::BasePass)
-							{
-								const TArray<uint32>& PassHitProxyIds = PrimitiveSceneInfo->NaniteHitProxyIds;
+								void* MaterialSlotRange = nullptr;
+#if WITH_EDITOR
+								void* HitProxyTable = nullptr;
+#endif
+								{
+									FNaniteMaterialCommandsLock NaniteLock(NaniteMaterials, SLT_Write);
+									MaterialSlotRange = NaniteMaterials.GetMaterialSlotPtr(UploadInfo.PrimitiveID, TableEntryCount);
+#if WITH_EDITOR
+									HitProxyTable = NaniteMaterials.GetHitProxyTablePtr(UploadInfo.PrimitiveID, HitProxyEntryCount);
+#endif
+								}
 
-								uint32* HitProxyEntry = static_cast<uint32*>(HitProxyTable);
-								for (int32 Entry = 0; Entry < PassHitProxyIds.Num(); ++Entry)
+								uint32* MaterialSlots = static_cast<uint32*>(MaterialSlotRange);
+								for (int32 Entry = 0; Entry < PassMaterialSlots.Num(); ++Entry)
 								{
-									HitProxyEntry[PassMaterials[Entry].MaterialIndex] = PassHitProxyIds[Entry];
+									MaterialSlots[PassMaterials[Entry].MaterialIndex] = PassMaterialSlots[Entry];
 								}
-							}
-							else
-							{
-								// Other passes don't use hit proxies. TODO: Shouldn't even need to do this.
-								uint64* DualHitProxyEntry = static_cast<uint64*>(HitProxyTable);
-								for (uint32 DualEntry = 0; DualEntry < NANITE_MAX_MATERIALS >> 1; ++DualEntry)
+
+#if WITH_EDITOR
+								if (NaniteMeshPass == ENaniteMeshPass::BasePass)
 								{
-									DualHitProxyEntry[DualEntry] = 0;
+									const TArray<uint32>& PassHitProxyIds = PrimitiveSceneInfo->NaniteHitProxyIds;
+
+									uint32* HitProxyEntry = static_cast<uint32*>(HitProxyTable);
+									for (int32 Entry = 0; Entry < PassHitProxyIds.Num(); ++Entry)
+									{
+										HitProxyEntry[PassMaterials[Entry].MaterialIndex] = PassHitProxyIds[Entry];
+									}
 								}
+								else
+								{
+									// Other passes don't use hit proxies. TODO: Shouldn't even need to do this.
+									uint64* DualHitProxyEntry = static_cast<uint64*>(HitProxyTable);
+									for (uint32 DualEntry = 0; DualEntry < NANITE_MAX_MATERIALS >> 1; ++DualEntry)
+									{
+										DualHitProxyEntry[DualEntry] = 0;
+									}
+								}
+#endif
 							}
-						#endif
 						}
 					}
 				}
diff --git a/Engine/Source/Runtime/Renderer/Private/ScreenSpaceRayTracing.cpp b/Engine/Source/Runtime/Renderer/Private/ScreenSpaceRayTracing.cpp
--- a/Engine/Source/Runtime/Renderer/Private/ScreenSpaceRayTracing.cpp
+++ b/Engine/Source/Runtime/Renderer/Private/ScreenSpaceRayTracing.cpp
@@ -782,7 +782,7 @@ FPrevSceneColorMip ReducePrevSceneColorMip(
 	// Number of mip in the mip chain
 	const int32 kNumMips = 5;
 
-	bool bUseLeakFree = View.PrevViewInfo.ScreenSpaceRayTracingInput != nullptr;
+	bool bUseLeakFree = View.PrevViewInfo.ScreenSpaceRayTracingInput != nullptr && View.PrevViewInfo.DepthBuffer != nullptr;
 	check(bUseLeakFree == true);
 
 	// Allocate FPrevSceneColorMip.
