[TEST] - testing lockdown preflight

diff --git a/Engine/Source/Runtime/Renderer/Private/DeferredShadingRenderer.cpp b/Engine/Source/Runtime/Renderer/Private/DeferredShadingRenderer.cpp
--- a/Engine/Source/Runtime/Renderer/Private/DeferredShadingRenderer.cpp
+++ b/Engine/Source/Runtime/Renderer/Private/DeferredShadingRenderer.cpp
@@ -1162,7 +1162,7 @@ bool FDeferredShadingSceneRenderer::GatherRayTracingWorldInstancesForView(FRDGBu
 			if (EnumHasAnyFlags(Flags, ERayTracingPrimitiveFlags::CacheInstances))
 			{
 				// TODO: support GRayTracingExcludeDecals, but not in the form of RayTracingMeshCommand.bDecal as that requires looping over all cached MDCs
-				// Instead, either make r.RayTracing.ExcludeDecals read only or request a recache of all ray tracing commands during which decals are excluded
+				// Instead, either make r.RayTracing.ExcludeDecals read. only or request a recache of all ray tracing commands during which decals are excluded
 
 				const int32 NewInstanceIndex = RayTracingScene.Instances.Num();
 
