Undo changelist 17467278 to unblock crashes rendering cinematics with SSGI

diff --git a/Engine/Source/Runtime/Renderer/Private/DeferredShadingRenderer.cpp b/Engine/Source/Runtime/Renderer/Private/DeferredShadingRenderer.cpp
--- a/Engine/Source/Runtime/Renderer/Private/DeferredShadingRenderer.cpp
+++ b/Engine/Source/Runtime/Renderer/Private/DeferredShadingRenderer.cpp
@@ -3147,24 +3147,15 @@ void FDeferredShadingSceneRenderer::Render(FRDGBuilder& GraphBuilder)
 	{
 		const FViewInfo& View = Views[ViewIndex];
 
-		if (!View.bStatePrevViewInfoIsReadOnly)
+		if (((View.FinalPostProcessSettings.DynamicGlobalIlluminationMethod == EDynamicGlobalIlluminationMethod::ScreenSpace && ScreenSpaceRayTracing::ShouldKeepBleedFreeSceneColor(View))
+			|| GetViewPipelineState(View).DiffuseIndirectMethod == EDiffuseIndirectMethod::Lumen
+			|| GetViewPipelineState(View).ReflectionsMethod == EReflectionsMethod::Lumen)
+			&& !View.bStatePrevViewInfoIsReadOnly)
 		{
 			// Keep scene color and depth for next frame screen space ray tracing.
 			FSceneViewState* ViewState = View.ViewState;
-
-			const bool bSSGI = View.FinalPostProcessSettings.DynamicGlobalIlluminationMethod == EDynamicGlobalIlluminationMethod::ScreenSpace && ScreenSpaceRayTracing::ShouldKeepBleedFreeSceneColor(View);
-			
-			if (bSSGI)
-			{
-				GraphBuilder.QueueTextureExtraction(SceneTextures.Depth.Resolve, &ViewState->PrevFrameViewInfo.DepthBuffer);
-			}
-			
-			if (bSSGI
-				|| GetViewPipelineState(View).DiffuseIndirectMethod == EDiffuseIndirectMethod::Lumen
-				|| GetViewPipelineState(View).ReflectionsMethod == EReflectionsMethod::Lumen)
-			{
-				GraphBuilder.QueueTextureExtraction(SceneTextures.Color.Resolve, &ViewState->PrevFrameViewInfo.ScreenSpaceRayTracingInput);
-			}
+			GraphBuilder.QueueTextureExtraction(SceneTextures.Depth.Resolve, &ViewState->PrevFrameViewInfo.DepthBuffer);
+			GraphBuilder.QueueTextureExtraction(SceneTextures.Color.Resolve, &ViewState->PrevFrameViewInfo.ScreenSpaceRayTracingInput);
 		}
 	}
 
