Make VSM physical pageinitialization and merge use indirect dispatch to only operate on needed pages.

diff --git a/Engine/Shaders/Private/VirtualShadowMaps/PageManagement.usf b/Engine/Shaders/Private/VirtualShadowMaps/PageManagement.usf
--- a/Engine/Shaders/Private/VirtualShadowMaps/PageManagement.usf
+++ b/Engine/Shaders/Private/VirtualShadowMaps/PageManagement.usf
@@ -749,7 +749,7 @@ void InitializePhysicalPages(uint2 PixelCoord : SV_DispatchThreadID)
 	if ((MetaData.Flags & VSM_ALLOCATED_FLAG) == 0)
 	{
 		// Page not used, we're done
-#if 1 // DEBUG clear unallocated pages
+#if 0 // DEBUG clear unallocated pages
 		OutPhysicalPagePool[PixelCoord] = 0;
 		OutPhysicalPagePool[PixelCoord + GetVirtualShadowMapStaticPixelOffset()] = 0;
 #endif
@@ -779,6 +779,15 @@ void InitializePhysicalPages(uint2 PixelCoord : SV_DispatchThreadID)
 	}
 }
 
+// Helper function to merge static and dynamic depth.
+void MergePhysicalPixel(uint2 PixelCoord)
+{
+	// 1:1 pixels so this is safe RMW
+	OutPhysicalPagePool[PixelCoord] = max(
+		OutPhysicalPagePool[PixelCoord],
+		OutPhysicalPagePool[PixelCoord + GetVirtualShadowMapStaticPixelOffset()]);
+}
+
 // Same as above - optimization to come
 [numthreads(16, 16, 1)]
 void MergeStaticPhysicalPages(uint2 PixelCoord : SV_DispatchThreadID)
@@ -798,8 +807,167 @@ void MergeStaticPhysicalPages(uint2 PixelCoord : SV_DispatchThreadID)
 	else if (!bFullyCached)
 	{
 		// 1:1 pixels so this is safe RMW
-		OutPhysicalPagePool[PixelCoord] = max(
-			OutPhysicalPagePool[PixelCoord],
-			OutPhysicalPagePool[PixelCoord + GetVirtualShadowMapStaticPixelOffset()]);
+		MergePhysicalPixel(PixelCoord);
+	}
+}
+
+RWBuffer<uint> OutIndirectArgsBuffer;
+
+// Set dispatch Args to 0,1,1 (3D grid dim) ready for atomic adding on GPU
+[numthreads(1, 1, 1)]
+void ClearIndirectDispatchArgs1DCS()
+{
+	OutIndirectArgsBuffer[0] = 0;
+	OutIndirectArgsBuffer[1] = 1;
+	OutIndirectArgsBuffer[2] = 1;
+}
+
+// Log2 2D dimension of thread group size, 2^4 == 16, 
+#define LOG2_TILE_THREAD_GROUP_SIZE_XY 4U
+#define TILE_THREAD_GROUP_SIZE_XY (1U << LOG2_TILE_THREAD_GROUP_SIZE_XY)
+
+// Each thread takes 2x2 samples to work with, so tile size is 2x thread group size
+#define LOG2_TILE_SIZE_XY (LOG2_TILE_THREAD_GROUP_SIZE_XY + 1U)	
+
+#if VSM_LOG2_PAGE_SIZE < LOG2_TILE_SIZE
+#error "VSM_LOG2_PAGE_SIZE must be larger than LOG2_TILE_SIZE, either increase one or reduce the other"
+#endif 
+
+// Number of tiles (thread groups) in each dimension to cover the page
+#define LOG2_TILES_PER_PAGE_XY ( VSM_LOG2_PAGE_SIZE - LOG2_TILE_SIZE_XY )
+// Log2 1D tile count to cover the page  LOG2_TILES_PER_PAGE_XY * LOG2_TILES_PER_PAGE_XY
+#define LOG2_TILES_PER_PAGE_1D ( 2U * LOG2_TILES_PER_PAGE_XY )
+// 1D tile count to cover the page  
+#define TILES_PER_PAGE_1D ( 1U << LOG2_TILES_PER_PAGE_1D )
+
+#define TILES_PER_PAGE_XY_MASK ( ( 1U << LOG2_TILES_PER_PAGE_XY ) - 1U )
+#define TILES_PER_PAGE_1D_MASK ( ( 1U << LOG2_TILES_PER_PAGE_1D ) - 1U )
+
+RWBuffer<uint> OutInitializePagesIndirectArgsBuffer;
+RWStructuredBuffer<uint> OutPhysicalPagesToInitialize;
+
+void EmitPageToProcess(RWBuffer<uint> OutIndirectArgsBuffer, RWStructuredBuffer<uint> OutSelectedPhysicalIndexBuffer, uint PhysicalPageIndex)
+{
+	int GroupCount = 0;
+	// Each page needs TILES_PER_PAGE_1D groups launched
+	WaveInterlockedAddScalar_(OutIndirectArgsBuffer[0], TILES_PER_PAGE_1D, GroupCount);
+	OutSelectedPhysicalIndexBuffer[GroupCount >> LOG2_TILES_PER_PAGE_1D] = PhysicalPageIndex;
+
+}
+
+[numthreads(VSM_DEFAULT_CS_GROUP_X, 1, 1)]
+void SelectPagesToInitializeCS(uint PhysicalPageIndex : SV_DispatchThreadID)
+{
+	if (PhysicalPageIndex >= VirtualShadowMap.MaxPhysicalPages)
+	{
+		return;
+	}
+
+	FPhysicalPageMetaData MetaData = PhysicalPageMetaData[PhysicalPageIndex];
+
+	bool bFullyCached = (MetaData.Flags & (VSM_DYNAMIC_UNCACHED_FLAG | VSM_STATIC_UNCACHED_FLAG)) == 0;
+	bool bStaticUncached = (MetaData.Flags & VSM_STATIC_UNCACHED_FLAG) != 0;
+
+	if ((MetaData.Flags & VSM_ALLOCATED_FLAG) == 0)
+	{
+		// Page not used, we're done
+	}
+	else if (bFullyCached)
+	{
+		// Page fully cached. Leave the data alone.
+	}
+	else
+	{
+		// TODO: In the relatively common case of static cached/dynamic uncached we could
+		// copy the static data to initialize the dynamic data and avoid the merge later.
+		// Before doing this we need to verify it works properly with geometry getting
+		// "added on top" of cached pages though, which would break this optimization.
+
+		// At least one of the pages is uncached
+		// NOTE: Dynamic cached/static uncached is currently an invalid state
+		// Since we merge the static stuff over the dynamic stuff after rendering we can't
+		// actually maintain separate dynamic cached pages when "only" the (theoretically)
+		// static moved. Thus if not fully cached, we always regenerate the dynamic page.
+		EmitPageToProcess(OutInitializePagesIndirectArgsBuffer, OutPhysicalPagesToInitialize, PhysicalPageIndex);
+		StatsBufferInterlockedInc(VSM_STAT_NUM_PAGES_TO_CLEAR);
+
+		if (bStaticUncached && VirtualShadowMapShouldCacheStaticSeparately())
+		{
+			EmitPageToProcess(OutInitializePagesIndirectArgsBuffer, OutPhysicalPagesToInitialize, PhysicalPageIndex + VirtualShadowMap.StaticPageIndexOffset);
+			StatsBufferInterlockedInc(VSM_STAT_NUM_PAGES_TO_CLEAR);
+		}
+	}
+}
+
+uint2 GetTileBasePos(uint2 TileThreadID, uint GroupIndex, StructuredBuffer<uint> PageIndexBuffer)
+{
+	const uint PageInputIndex = GroupIndex >> LOG2_TILES_PER_PAGE_1D;
+	const uint PageIndex = PageIndexBuffer[PageInputIndex];
+	// Each page has 1 << LOG2_TILES_PER_PAGE_XY groups (aka tiles) assigned to work on it.
+	const uint LocalTileIndex = GroupIndex & TILES_PER_PAGE_1D_MASK;
+	// wrap to 2D tile coord
+	const uint2 LocalTile = uint2(LocalTileIndex & TILES_PER_PAGE_XY_MASK, LocalTileIndex >> LOG2_TILES_PER_PAGE_XY);
+
+	uint2 PhysPageAddress = VSMPhysicalIndexToPageAddress(PageIndex);
+	// Pixel address of tile region for this thread group.
+	const uint2 TileOffset = (PhysPageAddress << uint2(VSM_LOG2_PAGE_SIZE, VSM_LOG2_PAGE_SIZE)) + (LocalTile << uint2(LOG2_TILE_SIZE_XY, LOG2_TILE_SIZE_XY));
+	// Pixel address of 2x2 region to sample for this thread.
+	const uint2 BasePos = TileOffset + (TileThreadID.xy << 1u);
+
+	return BasePos;
+}
+
+StructuredBuffer<uint> PhysicalPagesToInitialize;
+
+[numthreads(TILE_THREAD_GROUP_SIZE_XY, TILE_THREAD_GROUP_SIZE_XY, 1)]
+void InitializePhysicalPagesIndirectCS(uint2 TileThreadID : SV_GroupThreadID, uint GroupIndex : SV_GroupID)
+{
+	uint2 BasePos = GetTileBasePos(TileThreadID, GroupIndex, PhysicalPagesToInitialize);
+
+	OutPhysicalPagePool[BasePos + uint2(0U, 0U)] = 0;
+	OutPhysicalPagePool[BasePos + uint2(1U, 0U)] = 0;
+	OutPhysicalPagePool[BasePos + uint2(0U, 1U)] = 0;
+	OutPhysicalPagePool[BasePos + uint2(1U, 1U)] = 0;
+}
+
+RWBuffer<uint> OutMergePagesIndirectArgsBuffer;
+RWStructuredBuffer<uint> OutPhysicalPagesToMerge;
+
+
+[numthreads(VSM_DEFAULT_CS_GROUP_X, 1, 1)]
+void SelectPagesToMergeCS(uint PhysicalPageIndex : SV_DispatchThreadID)
+{
+	if (PhysicalPageIndex >= VirtualShadowMap.MaxPhysicalPages)
+	{
+		return;
 	}
+
+	FPhysicalPageMetaData MetaData = PhysicalPageMetaData[PhysicalPageIndex];
+
+	bool bFullyCached = (MetaData.Flags & (VSM_DYNAMIC_UNCACHED_FLAG | VSM_STATIC_UNCACHED_FLAG)) == 0;
+
+	// Merge any page that was not fully cached
+	// NOTE: This shader does not run if VirtualShadowMapShouldCacheStaticSeparately() is false
+	if ((MetaData.Flags & VSM_ALLOCATED_FLAG) == 0)
+	{
+	}
+	else if (!bFullyCached)
+	{
+		StatsBufferInterlockedInc(VSM_STAT_NUM_PAGES_TO_MERGE);
+		EmitPageToProcess(OutMergePagesIndirectArgsBuffer, OutPhysicalPagesToMerge, PhysicalPageIndex);
+	}
+}
+
+StructuredBuffer<uint> PhysicalPagesToMerge;
+
+[numthreads(TILE_THREAD_GROUP_SIZE_XY, TILE_THREAD_GROUP_SIZE_XY, 1)]
+void MergeStaticPhysicalPagesIndirectCS(uint2 TileThreadID : SV_GroupThreadID, uint GroupIndex : SV_GroupID)
+{
+	uint2 BasePos = GetTileBasePos(TileThreadID, GroupIndex, PhysicalPagesToMerge);
+
+	// 1:1 pixels so this is safe RMW
+	MergePhysicalPixel(BasePos + uint2(0U, 0U));
+	MergePhysicalPixel(BasePos + uint2(1U, 0U));
+	MergePhysicalPixel(BasePos + uint2(0U, 1U));
+	MergePhysicalPixel(BasePos + uint2(1U, 1U));
 }
diff --git a/Engine/Shaders/Private/VirtualShadowMaps/PrintStats.usf b/Engine/Shaders/Private/VirtualShadowMaps/PrintStats.usf
--- a/Engine/Shaders/Private/VirtualShadowMaps/PrintStats.usf
+++ b/Engine/Shaders/Private/VirtualShadowMaps/PrintStats.usf
@@ -89,6 +89,12 @@ void PrintStats()
 			ScreenPos.x = ValueX;
 			ScreenPos = ShaderPrintLeftAlign(ScreenPos, InStatsBuffer[VSM_STAT_DYNAMIC_INVALIDATED_BY_DYNAMIC_PAGES]);
 			ScreenPos = ShaderPrintNewline(ScreenPos);
+
+			ScreenPos.x = ItemX;
+			ScreenPos = ShaderPrintText(ScreenPos, TEXT("Merged"), FontWhite);
+			ScreenPos.x = ValueX;
+			ScreenPos = ShaderPrintLeftAlign(ScreenPos, InStatsBuffer[VSM_STAT_NUM_PAGES_TO_MERGE]);
+			ScreenPos = ShaderPrintNewline(ScreenPos);
 		}
 		else
 		{
@@ -106,6 +112,13 @@ void PrintStats()
 		}
 	}
 
+	ScreenPos.x = ItemX;
+	ScreenPos = ShaderPrintText(ScreenPos, TEXT("Cleared"), FontWhite);
+	ScreenPos.x = ValueX;
+	ScreenPos = ShaderPrintLeftAlign(ScreenPos, InStatsBuffer[VSM_STAT_NUM_PAGES_TO_CLEAR]);
+	ScreenPos = ShaderPrintNewline(ScreenPos);
+
+
 	if (bShowNonNaniteInstances)
 	{
 		ScreenPos.x = HeadlineX;
diff --git a/Engine/Shaders/Private/VirtualShadowMaps/Stats.ush b/Engine/Shaders/Private/VirtualShadowMaps/Stats.ush
--- a/Engine/Shaders/Private/VirtualShadowMaps/Stats.ush
+++ b/Engine/Shaders/Private/VirtualShadowMaps/Stats.ush
@@ -20,6 +20,8 @@
 #define VSM_STAT_NON_NANITE_INSTANCES_PAGE_MASK_CULLED	 9
 #define VSM_STAT_NON_NANITE_INSTANCES_EMPTY_RECT_CULLED	10
 #define VSM_STAT_NON_NANITE_INSTANCES_FRUSTUM_CULLED	11
+#define VSM_STAT_NUM_PAGES_TO_MERGE	12
+#define VSM_STAT_NUM_PAGES_TO_CLEAR	13
 
 #ifndef VSM_GENERATE_STATS
 #define VSM_GENERATE_STATS 0
diff --git a/Engine/Source/Runtime/Renderer/Private/VirtualShadowMaps/VirtualShadowMapArray.cpp b/Engine/Source/Runtime/Renderer/Private/VirtualShadowMaps/VirtualShadowMapArray.cpp
--- a/Engine/Source/Runtime/Renderer/Private/VirtualShadowMaps/VirtualShadowMapArray.cpp
+++ b/Engine/Source/Runtime/Renderer/Private/VirtualShadowMaps/VirtualShadowMapArray.cpp
@@ -192,6 +192,21 @@ static TAutoConsoleVariable<int32> CVarVirtualShadowMapDebugProjection(
 );
 #endif // !UE_BUILD_SHIPPING
 
+TAutoConsoleVariable<int32> CVarInitializePhysicalUsingIndirect(
+	TEXT("r.Shadow.Virtual.InitPhysicalUsingIndirect"),
+	0,
+	TEXT("."),
+	ECVF_RenderThreadSafe
+);
+
+TAutoConsoleVariable<int32> CVarMergePhysicalUsingIndirect(
+	TEXT("r.Shadow.Virtual.MergePhysicalUsingIndirect"),
+	0,
+	TEXT("."),
+	ECVF_RenderThreadSafe
+);
+
+
 FMatrix CalcTranslatedWorldToShadowUVMatrix(
 	const FMatrix& TranslatedWorldToShadowView,
 	const FMatrix& ViewToClip)
@@ -298,11 +313,14 @@ void FVirtualShadowMapArray::Initialize(FRDGBuilder& GraphBuilder, FVirtualShado
 	{
 		// Store the static pages below the dynamic/merged pages
 		UniformParameters.StaticCachedPixelOffsetY = PhysicalPagesY * FVirtualShadowMap::PageSize;
+		// Offset to static pages in linear page index
+		UniformParameters.StaticPageIndexOffset = PhysicalPagesY * PhysicalPagesX;
 		PhysicalPagesY *= 2;
 	}
 	else
 	{
 		UniformParameters.StaticCachedPixelOffsetY = 0;
+		UniformParameters.StaticPageIndexOffset = 0;
 	}
 
 	uint32 PhysicalX = PhysicalPagesX * FVirtualShadowMap::PageSize;
@@ -590,6 +608,65 @@ class FInitializePhysicalPagesCS : public FVirtualPageManagementShader
 };
 IMPLEMENT_GLOBAL_SHADER(FInitializePhysicalPagesCS, "/Engine/Private/VirtualShadowMaps/PageManagement.usf", "InitializePhysicalPages", SF_Compute);
 
+class FSelectPagesToInitializeCS : public FVirtualPageManagementShader
+{
+	DECLARE_GLOBAL_SHADER(FSelectPagesToInitializeCS);
+	SHADER_USE_PARAMETER_STRUCT(FSelectPagesToInitializeCS, FVirtualPageManagementShader)
+
+	class FGenerateStatsDim : SHADER_PERMUTATION_BOOL("VSM_GENERATE_STATS");
+	using FPermutationDomain = TShaderPermutationDomain<FGenerateStatsDim>;
+
+	BEGIN_SHADER_PARAMETER_STRUCT(FParameters, )
+		SHADER_PARAMETER_RDG_UNIFORM_BUFFER(FVirtualShadowMapUniformParameters, VirtualShadowMap)
+		SHADER_PARAMETER_RDG_BUFFER_SRV(StructuredBuffer<FPhysicalPageMetaData>, PhysicalPageMetaData)
+		SHADER_PARAMETER_RDG_BUFFER_UAV(RWBuffer<uint>, OutInitializePagesIndirectArgsBuffer)
+		SHADER_PARAMETER_RDG_BUFFER_UAV(RWStructuredBuffer< uint >, OutPhysicalPagesToInitialize)
+		SHADER_PARAMETER_RDG_BUFFER_UAV(RWStructuredBuffer< uint >, OutStatsBuffer)
+	END_SHADER_PARAMETER_STRUCT()
+};
+IMPLEMENT_GLOBAL_SHADER(FSelectPagesToInitializeCS, "/Engine/Private/VirtualShadowMaps/PageManagement.usf", "SelectPagesToInitializeCS", SF_Compute);
+
+class FInitializePhysicalPagesIndirectCS : public FVirtualPageManagementShader
+{
+	DECLARE_GLOBAL_SHADER(FInitializePhysicalPagesIndirectCS);
+	SHADER_USE_PARAMETER_STRUCT(FInitializePhysicalPagesIndirectCS, FVirtualPageManagementShader)
+
+	BEGIN_SHADER_PARAMETER_STRUCT(FParameters, )
+		SHADER_PARAMETER_RDG_UNIFORM_BUFFER(FVirtualShadowMapUniformParameters, VirtualShadowMap)
+		SHADER_PARAMETER_RDG_BUFFER_SRV(StructuredBuffer< uint >, PhysicalPagesToInitialize)
+		SHADER_PARAMETER_RDG_TEXTURE_UAV(RWTexture2D<uint>, OutPhysicalPagePool)
+		RDG_BUFFER_ACCESS(IndirectArgs, ERHIAccess::IndirectArgs)
+	END_SHADER_PARAMETER_STRUCT()
+};
+IMPLEMENT_GLOBAL_SHADER(FInitializePhysicalPagesIndirectCS, "/Engine/Private/VirtualShadowMaps/PageManagement.usf", "InitializePhysicalPagesIndirectCS", SF_Compute);
+
+class FClearIndirectDispatchArgs1DCS : public FVirtualPageManagementShader
+{
+	DECLARE_GLOBAL_SHADER(FClearIndirectDispatchArgs1DCS);
+	SHADER_USE_PARAMETER_STRUCT(FClearIndirectDispatchArgs1DCS, FVirtualPageManagementShader)
+
+	BEGIN_SHADER_PARAMETER_STRUCT(FParameters, )
+		SHADER_PARAMETER_RDG_BUFFER_UAV(RWBuffer<uint>, OutIndirectArgsBuffer)
+	END_SHADER_PARAMETER_STRUCT()
+};
+IMPLEMENT_GLOBAL_SHADER(FClearIndirectDispatchArgs1DCS, "/Engine/Private/VirtualShadowMaps/PageManagement.usf", "ClearIndirectDispatchArgs1DCS", SF_Compute);
+
+static void AddClearIndirectDispatchArgs1DPass(FRDGBuilder& GraphBuilder, FRDGBufferRef IndirectArgsRDG)
+{
+	FClearIndirectDispatchArgs1DCS::FParameters* PassParameters = GraphBuilder.AllocParameters<FClearIndirectDispatchArgs1DCS::FParameters>();
+	PassParameters->OutIndirectArgsBuffer = GraphBuilder.CreateUAV(IndirectArgsRDG);
+
+	auto ComputeShader = GetGlobalShaderMap(GMaxRHIFeatureLevel)->GetShader<FClearIndirectDispatchArgs1DCS>();
+
+	FComputeShaderUtils::AddPass(
+		GraphBuilder,
+		RDG_EVENT_NAME("ClearIndirectDispatchArgs"),
+		ComputeShader,
+		PassParameters,
+		FIntVector(1, 1, 1)
+	);
+}
+
 class FMergeStaticPhysicalPagesCS : public FVirtualPageManagementShader
 {
 	DECLARE_GLOBAL_SHADER(FMergeStaticPhysicalPagesCS);
@@ -603,6 +680,40 @@ class FMergeStaticPhysicalPagesCS : public FVirtualPageManagementShader
 };
 IMPLEMENT_GLOBAL_SHADER(FMergeStaticPhysicalPagesCS, "/Engine/Private/VirtualShadowMaps/PageManagement.usf", "MergeStaticPhysicalPages", SF_Compute);
 
+class FSelectPagesToMergeCS : public FVirtualPageManagementShader
+{
+	DECLARE_GLOBAL_SHADER(FSelectPagesToMergeCS);
+	SHADER_USE_PARAMETER_STRUCT(FSelectPagesToMergeCS, FVirtualPageManagementShader)
+
+	class FGenerateStatsDim : SHADER_PERMUTATION_BOOL("VSM_GENERATE_STATS");
+	using FPermutationDomain = TShaderPermutationDomain<FGenerateStatsDim>;
+
+	BEGIN_SHADER_PARAMETER_STRUCT(FParameters, )
+		SHADER_PARAMETER_RDG_UNIFORM_BUFFER(FVirtualShadowMapUniformParameters, VirtualShadowMap)
+		SHADER_PARAMETER_RDG_BUFFER_SRV(StructuredBuffer<FPhysicalPageMetaData>, PhysicalPageMetaData)
+		SHADER_PARAMETER_RDG_BUFFER_UAV(RWBuffer<uint>, OutMergePagesIndirectArgsBuffer)
+		SHADER_PARAMETER_RDG_BUFFER_UAV(RWStructuredBuffer< uint >, OutPhysicalPagesToMerge)
+		SHADER_PARAMETER_RDG_BUFFER_UAV(RWStructuredBuffer< uint >, OutStatsBuffer)
+	END_SHADER_PARAMETER_STRUCT()
+};
+IMPLEMENT_GLOBAL_SHADER(FSelectPagesToMergeCS, "/Engine/Private/VirtualShadowMaps/PageManagement.usf", "SelectPagesToMergeCS", SF_Compute);
+
+class FMergeStaticPhysicalPagesIndirectCS : public FVirtualPageManagementShader
+{
+	DECLARE_GLOBAL_SHADER(FMergeStaticPhysicalPagesIndirectCS);
+	SHADER_USE_PARAMETER_STRUCT(FMergeStaticPhysicalPagesIndirectCS, FVirtualPageManagementShader)
+
+	BEGIN_SHADER_PARAMETER_STRUCT(FParameters, )
+		SHADER_PARAMETER_RDG_UNIFORM_BUFFER(FVirtualShadowMapUniformParameters, VirtualShadowMap)
+		SHADER_PARAMETER_RDG_BUFFER_SRV(StructuredBuffer< uint >, PhysicalPagesToMerge)
+		SHADER_PARAMETER_RDG_TEXTURE_UAV(RWTexture2D<uint>, OutPhysicalPagePool)
+		RDG_BUFFER_ACCESS(IndirectArgs, ERHIAccess::IndirectArgs)
+	END_SHADER_PARAMETER_STRUCT()
+};
+IMPLEMENT_GLOBAL_SHADER(FMergeStaticPhysicalPagesIndirectCS, "/Engine/Private/VirtualShadowMaps/PageManagement.usf", "MergeStaticPhysicalPagesIndirectCS", SF_Compute);
+
+
+
 void FVirtualShadowMapArray::MergeStaticPhysicalPages(FRDGBuilder& GraphBuilder)
 {
 	check(IsEnabled());
@@ -612,6 +723,60 @@ void FVirtualShadowMapArray::MergeStaticPhysicalPages(FRDGBuilder& GraphBuilder)
 	}
 
 	RDG_EVENT_SCOPE(GraphBuilder, "FVirtualShadowMapArray::MergeStaticPhysicalPages");
+	if (CVarMergePhysicalUsingIndirect.GetValueOnRenderThread() != 0)
+	{
+		FRDGBufferRef MergePagesIndirectArgsRDG = GraphBuilder.CreateBuffer(FRDGBufferDesc::CreateIndirectDesc(3), TEXT("Shadow.Virtual.MergePagesIndirectArgs"));
+		// Note: creating buffer of 2x MaxPhysicalPages as MaxPhysicalPages represents the number before adding the static page overload. The selection shader emits both types of pages.
+		FRDGBufferRef PhysicalPagesToMergeRDG = GraphBuilder.CreateBuffer(FRDGBufferDesc::CreateStructuredDesc(sizeof(int32), 2U * UniformParameters.MaxPhysicalPages + 1), TEXT("Shadow.Virtual.PhysicalPagesToMerge"));
+
+		// 1. Initialize the indirect args buffer
+		AddClearIndirectDispatchArgs1DPass(GraphBuilder, MergePagesIndirectArgsRDG);
+		// 2. Filter the relevant physical pages and set up the indirect args
+		{
+			FSelectPagesToMergeCS::FParameters* PassParameters = GraphBuilder.AllocParameters<FSelectPagesToMergeCS::FParameters>();
+			PassParameters->VirtualShadowMap = GetUniformBuffer(GraphBuilder);
+			PassParameters->PhysicalPageMetaData = GraphBuilder.CreateSRV(PhysicalPageMetaDataRDG);
+			PassParameters->OutMergePagesIndirectArgsBuffer = GraphBuilder.CreateUAV(MergePagesIndirectArgsRDG);
+			PassParameters->OutPhysicalPagesToMerge = GraphBuilder.CreateUAV(PhysicalPagesToMergeRDG);
+			bool bGenerateStats = StatsBufferRDG != nullptr;
+			if (bGenerateStats)
+			{
+				PassParameters->OutStatsBuffer = GraphBuilder.CreateUAV(StatsBufferRDG);
+			}
+			FSelectPagesToMergeCS::FPermutationDomain PermutationVector;
+			PermutationVector.Set<FSelectPagesToMergeCS::FGenerateStatsDim>(bGenerateStats);
+
+			auto ComputeShader = GetGlobalShaderMap(GMaxRHIFeatureLevel)->GetShader<FSelectPagesToMergeCS>(PermutationVector);
+
+			FComputeShaderUtils::AddPass(
+				GraphBuilder,
+				RDG_EVENT_NAME("SelectPagesToMerge"),
+				ComputeShader,
+				PassParameters,
+				FIntVector(FMath::DivideAndRoundUp(UniformParameters.MaxPhysicalPages, FSelectPagesToMergeCS::DefaultCSGroupX), 1, 1)
+			);
+
+		}
+		// 3. Indirect dispatch to clear the selected pages
+		{
+			FMergeStaticPhysicalPagesIndirectCS::FParameters* PassParameters = GraphBuilder.AllocParameters<FMergeStaticPhysicalPagesIndirectCS::FParameters>();
+			PassParameters->VirtualShadowMap = GetUniformBuffer(GraphBuilder);
+			PassParameters->OutPhysicalPagePool = GraphBuilder.CreateUAV(PhysicalPagePoolRDG);
+			PassParameters->IndirectArgs = MergePagesIndirectArgsRDG;
+			PassParameters->PhysicalPagesToMerge = GraphBuilder.CreateSRV(PhysicalPagesToMergeRDG);
+			auto ComputeShader = GetGlobalShaderMap(GMaxRHIFeatureLevel)->GetShader<FMergeStaticPhysicalPagesIndirectCS>();
+
+			FComputeShaderUtils::AddPass(
+				GraphBuilder,
+				RDG_EVENT_NAME("MergeStaticPhysicalPagesIndirect"),
+				ComputeShader,
+				PassParameters,
+				PassParameters->IndirectArgs,
+				0
+			);
+		}
+	}
+	else
 	{
 		FMergeStaticPhysicalPagesCS::FParameters* PassParameters = GraphBuilder.AllocParameters<FMergeStaticPhysicalPagesCS::FParameters>();
 		PassParameters->VirtualShadowMap = GetUniformBuffer(GraphBuilder);
@@ -1152,7 +1317,61 @@ void FVirtualShadowMapArray::BuildPageAllocations(
 		// Initialize the physical page pool
 		check(PhysicalPagePoolRDG != nullptr);
 		{
-			RDG_EVENT_SCOPE( GraphBuilder, "FVirtualShadowMapArray::InitializePhysicalPages" );
+			RDG_EVENT_SCOPE( GraphBuilder, "InitializePhysicalPages" );
+			if (CVarInitializePhysicalUsingIndirect.GetValueOnRenderThread() != 0)
+			{
+				FRDGBufferRef InitializePagesIndirectArgsRDG = GraphBuilder.CreateBuffer(FRDGBufferDesc::CreateIndirectDesc(3), TEXT("Shadow.Virtual.InitializePagesIndirectArgs"));
+				// Note: creating buffer of 2x MaxPhysicalPages as MaxPhysicalPages represents the number before adding the static page overload. The selection shader emits both types of pages.
+				FRDGBufferRef PhysicalPagesToInitializeRDG = GraphBuilder.CreateBuffer(FRDGBufferDesc::CreateStructuredDesc(sizeof(int32), 2U * UniformParameters.MaxPhysicalPages + 1), TEXT("Shadow.Virtual.PhysicalPagesToInitialize"));
+
+				// 1. Initialize the indirect args buffer
+				AddClearIndirectDispatchArgs1DPass(GraphBuilder, InitializePagesIndirectArgsRDG);
+				// 2. Filter the relevant physical pages and set up the indirect args
+				{
+					FSelectPagesToInitializeCS::FParameters* PassParameters = GraphBuilder.AllocParameters<FSelectPagesToInitializeCS::FParameters>();
+					PassParameters->VirtualShadowMap = GetUniformBuffer(GraphBuilder);
+					PassParameters->PhysicalPageMetaData = GraphBuilder.CreateSRV(PhysicalPageMetaDataRDG);
+					PassParameters->OutInitializePagesIndirectArgsBuffer = GraphBuilder.CreateUAV(InitializePagesIndirectArgsRDG);
+					PassParameters->OutPhysicalPagesToInitialize = GraphBuilder.CreateUAV(PhysicalPagesToInitializeRDG);
+					bool bGenerateStats = StatsBufferRDG != nullptr;
+					if (bGenerateStats)
+					{
+						PassParameters->OutStatsBuffer = GraphBuilder.CreateUAV(StatsBufferRDG);
+					}
+					FSelectPagesToInitializeCS::FPermutationDomain PermutationVector;
+					PermutationVector.Set<FSelectPagesToInitializeCS::FGenerateStatsDim>(bGenerateStats);
+
+					auto ComputeShader = GetGlobalShaderMap(GMaxRHIFeatureLevel)->GetShader<FSelectPagesToInitializeCS>(PermutationVector);
+
+					FComputeShaderUtils::AddPass(
+						GraphBuilder,
+						RDG_EVENT_NAME("SelectPagesToInitialize"),
+						ComputeShader,
+						PassParameters,
+						FIntVector(FMath::DivideAndRoundUp(UniformParameters.MaxPhysicalPages, FSelectPagesToInitializeCS::DefaultCSGroupX), 1, 1)
+					);
+
+				}
+				// 3. Indirect dispatch to clear the selected pages
+				{
+					FInitializePhysicalPagesIndirectCS::FParameters* PassParameters = GraphBuilder.AllocParameters<FInitializePhysicalPagesIndirectCS::FParameters>();
+					PassParameters->VirtualShadowMap = GetUniformBuffer(GraphBuilder);
+					PassParameters->OutPhysicalPagePool = GraphBuilder.CreateUAV(PhysicalPagePoolRDG);
+					PassParameters->IndirectArgs = InitializePagesIndirectArgsRDG;
+					PassParameters->PhysicalPagesToInitialize = GraphBuilder.CreateSRV(PhysicalPagesToInitializeRDG);
+					auto ComputeShader = GetGlobalShaderMap(GMaxRHIFeatureLevel)->GetShader<FInitializePhysicalPagesIndirectCS>();
+
+					FComputeShaderUtils::AddPass(
+						GraphBuilder,
+						RDG_EVENT_NAME("InitializePhysicalMemoryIndirect"),
+						ComputeShader,
+						PassParameters,
+						PassParameters->IndirectArgs,
+						0
+					);
+				}
+			}
+			else
 			{
 				FInitializePhysicalPagesCS::FParameters* PassParameters = GraphBuilder.AllocParameters<FInitializePhysicalPagesCS::FParameters>();
 				PassParameters->VirtualShadowMap = GetUniformBuffer(GraphBuilder);
diff --git a/Engine/Source/Runtime/Renderer/Private/VirtualShadowMaps/VirtualShadowMapArray.h b/Engine/Source/Runtime/Renderer/Private/VirtualShadowMaps/VirtualShadowMapArray.h
--- a/Engine/Source/Runtime/Renderer/Private/VirtualShadowMaps/VirtualShadowMapArray.h
+++ b/Engine/Source/Runtime/Renderer/Private/VirtualShadowMaps/VirtualShadowMapArray.h
@@ -119,6 +119,7 @@ BEGIN_GLOBAL_SHADER_PARAMETER_STRUCT(FVirtualShadowMapUniformParameters, )
 	SHADER_PARAMETER(uint32, MaxPhysicalPages)
 	// Set to 0 if separate static caching is disabled
 	SHADER_PARAMETER(uint32, StaticCachedPixelOffsetY)
+	SHADER_PARAMETER(uint32, StaticPageIndexOffset)
 	// use to map linear index to x,y page coord
 	SHADER_PARAMETER(uint32, PhysicalPageRowMask)
 	SHADER_PARAMETER(uint32, PhysicalPageRowShift)
