Set default MotionBlurPerObjectSize to 0 so that by default we don't skip velocity for small objects.

diff --git a/Engine/Source/Runtime/Engine/Private/Scene.cpp b/Engine/Source/Runtime/Engine/Private/Scene.cpp
--- a/Engine/Source/Runtime/Engine/Private/Scene.cpp
+++ b/Engine/Source/Runtime/Engine/Private/Scene.cpp
@@ -559,7 +559,7 @@ FPostProcessSettings::FPostProcessSettings()
 	MotionBlurAmount = 0.5f;
 	MotionBlurMax = 5.0f;
 	MotionBlurTargetFPS = 30;
-	MotionBlurPerObjectSize = 0.5f;
+	MotionBlurPerObjectSize = 0.f;
 	ScreenPercentage_DEPRECATED = 100.0f;
 	ReflectionsType_DEPRECATED = EReflectionsType::RayTracing;
 	ReflectionMethod = EReflectionMethod::Lumen;
diff --git a/Engine/Source/Runtime/Renderer/Private/VelocityRendering.cpp b/Engine/Source/Runtime/Renderer/Private/VelocityRendering.cpp
--- a/Engine/Source/Runtime/Renderer/Private/VelocityRendering.cpp
+++ b/Engine/Source/Runtime/Renderer/Private/VelocityRendering.cpp
@@ -400,7 +400,7 @@ bool FVelocityMeshProcessor::PrimitiveHasVelocityForView(const FViewInfo& View,
 	const FBoxSphereBounds& PrimitiveBounds = PrimitiveSceneProxy->GetBounds();
 	const float LODFactorDistanceSquared = (PrimitiveBounds.Origin - View.ViewMatrices.GetViewOrigin()).SizeSquared() * FMath::Square(View.LODDistanceFactor);
 
-	// The minimum projected screen radius for a primitive to be drawn in the velocity pass, as a fraction of half the horizontal screen width (likely to be 0.08f)
+	// The minimum projected screen radius for a primitive to be drawn in the velocity pass, as a fraction of half the horizontal screen width.
 	float MinScreenRadiusForVelocityPass = View.FinalPostProcessSettings.MotionBlurPerObjectSize * (2.0f / 100.0f);
 	float MinScreenRadiusForVelocityPassSquared = FMath::Square(MinScreenRadiusForVelocityPass);
 
