Merge ray tracing primitive culling into standard primitive culling.  This allows the culling to use the same memory access patterns as the normal primitive culling, and ensures that the distance culled primitives and the fading primitives, propegate into the ray tracing primitives.
Bonus: Saves about 1.1 MS in the current branch.

diff --git a/Engine/Source/Runtime/RenderCore/Private/RenderCore.cpp b/Engine/Source/Runtime/RenderCore/Private/RenderCore.cpp
--- a/Engine/Source/Runtime/RenderCore/Private/RenderCore.cpp
+++ b/Engine/Source/Runtime/RenderCore/Private/RenderCore.cpp
@@ -112,13 +112,14 @@ DEFINE_STAT(STAT_ViewRelevance);
 DEFINE_STAT(STAT_ComputeViewRelevance);
 DEFINE_STAT(STAT_OcclusionCull);
 DEFINE_STAT(STAT_UpdatePrimitiveFading);
-DEFINE_STAT(STAT_FrustumCull);
+DEFINE_STAT(STAT_PrimitiveCull);
 DEFINE_STAT(STAT_DecompressPrecomputedOcclusion);
 DEFINE_STAT(STAT_ViewVisibilityTime);
 
 DEFINE_STAT(STAT_RayTracingInstances);
 DEFINE_STAT(STAT_ProcessedPrimitives);
 DEFINE_STAT(STAT_CulledPrimitives);
+DEFINE_STAT(STAT_VisibleRayTracingPrimitives);
 DEFINE_STAT(STAT_StaticallyOccludedPrimitives);
 DEFINE_STAT(STAT_OccludedPrimitives);
 DEFINE_STAT(STAT_OcclusionQueries);
diff --git a/Engine/Source/Runtime/RenderCore/Public/RenderCore.h b/Engine/Source/Runtime/RenderCore/Public/RenderCore.h
--- a/Engine/Source/Runtime/RenderCore/Public/RenderCore.h
+++ b/Engine/Source/Runtime/RenderCore/Public/RenderCore.h
@@ -48,7 +48,7 @@ DECLARE_DWORD_COUNTER_STAT_EXTERN(TEXT("Ray tracing instances"), STAT_RayTracing
 
 DECLARE_CYCLE_STAT_EXTERN(TEXT("View Visibility"),STAT_ViewVisibilityTime,STATGROUP_InitViews, RENDERCORE_API);
 DECLARE_CYCLE_STAT_EXTERN(TEXT("Decompress Occlusion"),STAT_DecompressPrecomputedOcclusion,STATGROUP_InitViews, RENDERCORE_API);
-DECLARE_CYCLE_STAT_EXTERN(TEXT("Frustum Cull"),STAT_FrustumCull,STATGROUP_InitViews, RENDERCORE_API);
+DECLARE_CYCLE_STAT_EXTERN(TEXT("Primitive Cull"), STAT_PrimitiveCull,STATGROUP_InitViews, RENDERCORE_API);
 DECLARE_CYCLE_STAT_EXTERN(TEXT("Update Fading"),STAT_UpdatePrimitiveFading,STATGROUP_InitViews, RENDERCORE_API);
 DECLARE_CYCLE_STAT_EXTERN(TEXT("Occlusion Cull"),STAT_OcclusionCull,STATGROUP_InitViews, RENDERCORE_API);
 DECLARE_CYCLE_STAT_EXTERN(TEXT("View Relevance"),STAT_ViewRelevance,STATGROUP_InitViews, RENDERCORE_API);
@@ -74,6 +74,7 @@ DECLARE_CYCLE_STAT_EXTERN(TEXT("GatherShadowPrimitives"),STAT_GatherShadowPrimit
 DECLARE_CYCLE_STAT_EXTERN(TEXT("BuildCSMVisibilityState"), STAT_BuildCSMVisibilityState, STATGROUP_InitViews, RENDERCORE_API);
 DECLARE_DWORD_COUNTER_STAT_EXTERN(TEXT("Processed primitives"),STAT_ProcessedPrimitives,STATGROUP_InitViews, RENDERCORE_API);
 DECLARE_DWORD_COUNTER_STAT_EXTERN(TEXT("Frustum Culled primitives"),STAT_CulledPrimitives,STATGROUP_InitViews, RENDERCORE_API);
+DECLARE_DWORD_COUNTER_STAT_EXTERN(TEXT("Visible Raytracing Primitives"), STAT_VisibleRayTracingPrimitives, STATGROUP_InitViews, RENDERCORE_API);
 DECLARE_DWORD_COUNTER_STAT_EXTERN(TEXT("Statically occluded primitives"),STAT_StaticallyOccludedPrimitives,STATGROUP_InitViews, RENDERCORE_API);
 DECLARE_DWORD_COUNTER_STAT_EXTERN(TEXT("Occluded primitives"),STAT_OccludedPrimitives,STATGROUP_InitViews, RENDERCORE_API);
 DECLARE_DWORD_COUNTER_STAT_EXTERN(TEXT("Occlusion queries"),STAT_OcclusionQueries,STATGROUP_InitViews, RENDERCORE_API);
diff --git a/Engine/Source/Runtime/Renderer/Private/DeferredShadingRenderer.cpp b/Engine/Source/Runtime/Renderer/Private/DeferredShadingRenderer.cpp
--- a/Engine/Source/Runtime/Renderer/Private/DeferredShadingRenderer.cpp
+++ b/Engine/Source/Runtime/Renderer/Private/DeferredShadingRenderer.cpp
@@ -588,6 +588,9 @@ bool FDeferredShadingSceneRenderer::GatherRayTracingWorldInstancesForView(FRDGBu
 	TArray<FPrimitiveSceneInfo*> DirtyCachedRayTracingPrimitives;
 	DirtyCachedRayTracingPrimitives.Reserve(Scene->PrimitiveSceneProxies.Num());
 
+	int32 VisiblePrimitives = 0;
+	bool bPerformRayTracing = View.State != nullptr && !View.bIsReflectionCapture;
+	if (bPerformRayTracing)
 	{
 		TRACE_CPUPROFILER_EVENT_SCOPE(GatherRayTracingWorldInstances_RelevantPrimitives);
 
@@ -608,45 +611,14 @@ bool FDeferredShadingSceneRenderer::GatherRayTracingWorldInstancesForView(FRDGBu
 				continue;
 			}
 
-			if (EnumHasAnyFlags(Scene->PrimitiveRayTracingFlags[PrimitiveIndex], ERayTracingPrimitiveFlags::Excluded))
+			// Get primitive visibility state from culling
+			if (!View.PrimitiveRayTracingVisibilityMap[PrimitiveIndex])
 			{
 				continue;
 			}
 
 			const FPrimitiveSceneInfo* SceneInfo = Scene->Primitives[PrimitiveIndex];
 
-			// Skip far field if not enabled
-			const bool bIsFarFieldPrimitive = EnumHasAnyFlags(Scene->PrimitiveRayTracingFlags[PrimitiveIndex], ERayTracingPrimitiveFlags::FarField);
-			if (!View.RayTracingCullingParameters.bIsRayTracingFarField && bIsFarFieldPrimitive)
-			{
-				continue;
-			}
-
-			if (RayTracing::ShouldCullBounds(View.RayTracingCullingParameters, Scene->PrimitiveBounds[PrimitiveIndex].BoxSphereBounds, bIsFarFieldPrimitive))
-			{
-				continue;
-			}
-
-			if (!View.State)
-			{
-				continue;
-			}
-
-			if (View.bIsReflectionCapture)
-			{
-				continue;
-			}
-
-			if (View.HiddenPrimitives.Contains(Scene->PrimitiveComponentIds[PrimitiveIndex]))
-			{
-				continue;
-			}
-
-			if (View.ShowOnlyPrimitives.IsSet() && !View.ShowOnlyPrimitives->Contains(Scene->PrimitiveComponentIds[PrimitiveIndex]))
-			{
-				continue;
-			}
-
 			// #dxr_todo: ray tracing in scene captures should re-use the persistent RT scene. (UE-112448)
 			bool bShouldRayTraceSceneCapture = GRayTracingSceneCaptures > 0
 				|| (GRayTracingSceneCaptures == -1 && View.bSceneCaptureUsesRayTracing);
@@ -655,15 +627,6 @@ bool FDeferredShadingSceneRenderer::GatherRayTracingWorldInstancesForView(FRDGBu
 			{
 				continue;
 			}
-			
-			// Disabling for now, this is not correct for ray tracing visibility
-			/*
-			// Check if the primitive has been culled already during frustum culling
-			if (!View.PrimitiveVisibilityMap[PrimitiveIndex])
-			{
-				continue;
-			}
-			*/
 
 			// Marked visible and used after point, check if streaming then mark as used in the TLAS (so it can be streamed in)
 			if (EnumHasAnyFlags(Scene->PrimitiveRayTracingFlags[PrimitiveIndex], ERayTracingPrimitiveFlags::Streaming))
@@ -678,6 +641,8 @@ bool FDeferredShadingSceneRenderer::GatherRayTracingWorldInstancesForView(FRDGBu
 				RayTracingScene.UsedCoarseMeshStreamingHandles.Add(SceneInfo->CoarseMeshStreamingHandle);
 			}
 
+			VisiblePrimitives++;
+
 			//#dxr_todo UE-68621  The Raytracing code path does not support ShowFlags since data moved to the SceneInfo. 
 			//Touching the SceneProxy to determine this would simply cost too much
 			static const auto RayTracingStaticMeshesCVar = IConsoleManager::Get().FindTConsoleVariableDataInt(TEXT("r.RayTracing.Geometry.StaticMeshes"));
@@ -686,7 +651,7 @@ bool FDeferredShadingSceneRenderer::GatherRayTracingWorldInstancesForView(FRDGBu
 			Item.PrimitiveIndex = PrimitiveIndex;
 
 			if (EnumHasAnyFlags(Scene->PrimitiveRayTracingFlags[PrimitiveIndex], ERayTracingPrimitiveFlags::CacheMeshCommands)
-				&& View.Family->EngineShowFlags.StaticMeshes 
+				&& View.Family->EngineShowFlags.StaticMeshes
 				&& RayTracingStaticMeshesCVar && RayTracingStaticMeshesCVar->GetValueOnRenderThread() > 0)
 			{
 				Item.bStatic = true;
@@ -700,6 +665,8 @@ bool FDeferredShadingSceneRenderer::GatherRayTracingWorldInstancesForView(FRDGBu
 		}
 	}
 
+	INC_DWORD_STAT_BY(STAT_VisibleRayTracingPrimitives, VisiblePrimitives);
+
 	FPrimitiveSceneInfo::UpdateCachedRaytracingData(Scene, DirtyCachedRayTracingPrimitives);
 
 	FGraphEventArray LODTaskList;
diff --git a/Engine/Source/Runtime/Renderer/Private/DistanceFieldObjectManagement.cpp b/Engine/Source/Runtime/Renderer/Private/DistanceFieldObjectManagement.cpp
--- a/Engine/Source/Runtime/Renderer/Private/DistanceFieldObjectManagement.cpp
+++ b/Engine/Source/Runtime/Renderer/Private/DistanceFieldObjectManagement.cpp
@@ -747,7 +747,7 @@ void FDistanceFieldSceneData::UpdateDistanceFieldObjectBuffers(
 
 									const int32 PrimIdx = PrimAndInst.Primitive->GetIndex();
 									const FPrimitiveBounds& PrimBounds = PrimitiveBounds[PrimIdx];
-									float MinDrawDist2 = PrimBounds.MinDrawDistanceSq;
+									float MinDrawDist2 = FMath::Square(PrimBounds.MinDrawDistance);
 									// For IEEE compatible machines, float operations goes to inf if overflow
 									// In this case, it will effectively disable max draw distance culling
 									float MaxDrawDist = FMath::Max(PrimBounds.MaxCullDistance, 0.f) * GetCachedScalabilityCVars().ViewDistanceScale;
diff --git a/Engine/Source/Runtime/Renderer/Private/PrimitiveSceneInfo.cpp b/Engine/Source/Runtime/Renderer/Private/PrimitiveSceneInfo.cpp
--- a/Engine/Source/Runtime/Renderer/Private/PrimitiveSceneInfo.cpp
+++ b/Engine/Source/Runtime/Renderer/Private/PrimitiveSceneInfo.cpp
@@ -1256,7 +1256,7 @@ void FPrimitiveSceneInfo::AddToScene(FRHICommandListImmediate& RHICmdList, FScen
 			FPrimitiveBounds& PrimitiveBounds = Scene->PrimitiveBounds[PackedIndex];
 			FBoxSphereBounds BoxSphereBounds = Proxy->GetBounds();
 			PrimitiveBounds.BoxSphereBounds = BoxSphereBounds;
-			PrimitiveBounds.MinDrawDistanceSq = FMath::Square(Proxy->GetMinDrawDistance());
+			PrimitiveBounds.MinDrawDistance = Proxy->GetMinDrawDistance();
 			PrimitiveBounds.MaxDrawDistance = Proxy->GetMaxDrawDistance();
 			PrimitiveBounds.MaxCullDistance = PrimitiveBounds.MaxDrawDistance;
 
diff --git a/Engine/Source/Runtime/Renderer/Private/RayTracing/RayTracingInstanceCulling.cpp b/Engine/Source/Runtime/Renderer/Private/RayTracing/RayTracingInstanceCulling.cpp
--- a/Engine/Source/Runtime/Renderer/Private/RayTracing/RayTracingInstanceCulling.cpp
+++ b/Engine/Source/Runtime/Renderer/Private/RayTracing/RayTracingInstanceCulling.cpp
@@ -55,6 +55,7 @@ void FRayTracingCullingParameters::Init(FViewInfo& View)
 	FarFieldCullingRadius = Lumen::GetFarFieldMaxTraceDistance();
 	CullAngleThreshold = CVarRayTracingCullingAngle.GetValueOnRenderThread();
 	AngleThresholdRatio = FMath::Tan(FMath::Min(89.99f, CullAngleThreshold) * PI / 180.0f);
+	AngleThresholdRatioSq = FMath::Square(AngleThresholdRatio);
 	ViewOrigin = View.ViewMatrices.GetViewOrigin();
 	ViewDirection = View.GetViewDirection();
 	bCullAllObjects = CullInRayTracing == 2 || CullInRayTracing == 3;
@@ -64,43 +65,84 @@ void FRayTracingCullingParameters::Init(FViewInfo& View)
 
 namespace RayTracing
 {
-
-bool ShouldCullBounds(const FRayTracingCullingParameters& CullingParameters, FBoxSphereBounds ObjectBounds, bool bIsFarFieldPrimitive)
+bool CullPrimitiveByFlags(const FRayTracingCullingParameters& CullingParameters, const FScene* RESTRICT Scene, int32 PrimitiveIndex)
 {
-	if (CullingParameters.CullInRayTracing > 0)
+	if (EnumHasAnyFlags(Scene->PrimitiveRayTracingFlags[PrimitiveIndex], ERayTracingPrimitiveFlags::UnsupportedProxyType))
 	{
-		const float ObjectRadius = ObjectBounds.SphereRadius;
-		const FVector ObjectCenter = ObjectBounds.Origin + 0.5 * ObjectBounds.BoxExtent;
-		const FVector CameraToObjectCenter = FVector(ObjectCenter - CullingParameters.ViewOrigin);
+		return true;
+	}
 
-		const bool bConsiderCulling = CullingParameters.bCullAllObjects || FVector::DotProduct(CullingParameters.ViewDirection, CameraToObjectCenter) < -ObjectRadius;
+	if (EnumHasAnyFlags(Scene->PrimitiveRayTracingFlags[PrimitiveIndex], ERayTracingPrimitiveFlags::Excluded))
+	{
+		return true;
+	}
 
-		if (bConsiderCulling)
+	// Skip far field if not enabled
+	const bool bIsFarFieldPrimitive = EnumHasAnyFlags(Scene->PrimitiveRayTracingFlags[PrimitiveIndex], ERayTracingPrimitiveFlags::FarField);
+	if (!CullingParameters.bIsRayTracingFarField && bIsFarFieldPrimitive)
+	{
+		return true;
+	}
+
+	return false;
+}
+
+// Tests if the given primitive should be culled, given the pre-calculated inputs for the primitive, returns true if the primitive SHOULD be culled
+template<bool bCullByRadiusOrDistance>
+bool CullBounds(const FRayTracingCullingParameters& CullingParameters, const FBoxSphereBounds& RESTRICT ObjectBounds, bool bIsFarFieldPrimitive)
+{
+	const float ObjectRadius = ObjectBounds.SphereRadius;
+	const FVector ObjectCenter = ObjectBounds.Origin + 0.5 * ObjectBounds.BoxExtent;
+	float CameraToObjectCenterLengthSq = FVector::DistSquared(ObjectCenter, CullingParameters.ViewOrigin);
+
+	if (bIsFarFieldPrimitive)
+	{
+		if (CameraToObjectCenterLengthSq > FMath::Square(CullingParameters.FarFieldCullingRadius + ObjectRadius))
 		{
-			const float CameraToObjectCenterLength = CameraToObjectCenter.Size();
+			return true;
+		}
+	}
+	else
+	{
+		const bool bIsFarEnoughToCull = CameraToObjectCenterLengthSq > FMath::Square(CullingParameters.CullingRadius + ObjectRadius);
 
-			if (bIsFarFieldPrimitive)
+		if (bCullByRadiusOrDistance && bIsFarEnoughToCull)
+		{
+			return true;
+		}
+		else
+		{
+			// Cull by solid angle: check the radius of bounding sphere against angle threshold
+			const bool bAngleIsSmallEnoughToCull = FMath::Square(ObjectRadius) / CameraToObjectCenterLengthSq < CullingParameters.AngleThresholdRatioSq;
+			if (bCullByRadiusOrDistance && bAngleIsSmallEnoughToCull)
 			{
-				if (CameraToObjectCenterLength > (CullingParameters.FarFieldCullingRadius + ObjectRadius))
-				{
-					return true;
-				}
+				return true;
 			}
-			else
+			else if (bIsFarEnoughToCull && bAngleIsSmallEnoughToCull)
 			{
-				const bool bIsFarEnoughToCull = CameraToObjectCenterLength > (CullingParameters.CullingRadius + ObjectRadius);
+				return true;
+			}
+		}
+	}
 
-				// Cull by solid angle: check the radius of bounding sphere against angle threshold
-				const bool bAngleIsSmallEnoughToCull = ObjectRadius / CameraToObjectCenterLength < CullingParameters.AngleThresholdRatio;
+	return false;
+}
 
-				if (CullingParameters.bCullByRadiusOrDistance && (bIsFarEnoughToCull || bAngleIsSmallEnoughToCull))
-				{
-					return true;
-				}
-				else if (bIsFarEnoughToCull && bAngleIsSmallEnoughToCull)
-				{
-					return true;
-				}
+bool ShouldCullBounds(const FRayTracingCullingParameters& CullingParameters, const FBoxSphereBounds& RESTRICT ObjectBounds, bool bIsFarFieldPrimitive)
+{
+	if (CullingParameters.CullInRayTracing > 0)
+	{
+		const bool bConsiderCulling = CullingParameters.bCullAllObjects || ShouldConsiderCulling(CullingParameters, ObjectBounds);
+
+		if (bConsiderCulling)
+		{
+			if (CullingParameters.bCullByRadiusOrDistance)
+			{
+				return CullBounds<true>(CullingParameters, ObjectBounds, bIsFarFieldPrimitive);
+			}
+			else
+			{
+				return CullBounds<false>(CullingParameters, ObjectBounds, bIsFarFieldPrimitive);
 			}
 		}
 	}
@@ -113,7 +155,7 @@ bool ShouldSkipPerInstanceCullingForPrimitive(const FRayTracingCullingParameters
 	bool bSkipCulling = false;
 
 	const float ObjectRadius = ObjectBounds.SphereRadius;
-	const FVector ObjectCenter = ObjectBounds.Origin + 0.5 * ObjectBounds.BoxExtent;
+	const FVector ObjectCenter = ObjectBounds.Origin;
 	const FVector CameraToObjectCenter = FVector(ObjectCenter - CullingParameters.ViewOrigin);
 
 	const FVector CameraToFurthestInstanceCenter = CameraToObjectCenter * (CameraToObjectCenter.Size() + ObjectRadius + SmallestInstanceBounds.SphereRadius) / CameraToObjectCenter.Size();
diff --git a/Engine/Source/Runtime/Renderer/Private/RayTracing/RayTracingInstanceCulling.h b/Engine/Source/Runtime/Renderer/Private/RayTracing/RayTracingInstanceCulling.h
--- a/Engine/Source/Runtime/Renderer/Private/RayTracing/RayTracingInstanceCulling.h
+++ b/Engine/Source/Runtime/Renderer/Private/RayTracing/RayTracingInstanceCulling.h
@@ -14,9 +14,20 @@ int32 GetRayTracingCullingPerInstance();
 namespace RayTracing
 {
 
-bool ShouldCullBounds(const FRayTracingCullingParameters& CullingParameters, FBoxSphereBounds ObjectBounds, bool bIsFarFieldPrimitive);
-bool ShouldSkipPerInstanceCullingForPrimitive(const FRayTracingCullingParameters& CullingParameters, FBoxSphereBounds ObjectBounds, FBoxSphereBounds SmallestInstanceBounds, bool bIsFarFieldPrimitive);
+// Tests if a primitive with the given inputs should be considered for culling.  Does NOT test configuration values.  The assumption is that the config values have already been tested.
+FORCEINLINE bool ShouldConsiderCulling(const FRayTracingCullingParameters& CullingParameters, const FBoxSphereBounds& ObjectBounds)
+{
+	FVector CameraToObjectCenter = ObjectBounds.Origin - CullingParameters.ViewOrigin;
+	return FVector::DotProduct(CullingParameters.ViewDirection, CameraToObjectCenter) < -ObjectBounds.SphereRadius;
+}
 
+// Returns true if the primitive should be culled out due to it's ray tracing flags
+bool CullPrimitiveByFlags(const FRayTracingCullingParameters& CullingParameters, const FScene* RESTRICT Scene, int32 PrimitiveIndex);
+
+// Completely test if the bounds should be culled for ray tracing. This includes all configuration values.
+bool ShouldCullBounds(const FRayTracingCullingParameters& CullingParameters, const FBoxSphereBounds& ObjectBounds, bool bIsFarFieldPrimitive);
+
+bool ShouldSkipPerInstanceCullingForPrimitive(const FRayTracingCullingParameters& CullingParameters, FBoxSphereBounds ObjectBounds, FBoxSphereBounds SmallestInstanceBounds, bool bIsFarFieldPrimitive);
 }
 
 struct FRayTracingCullPrimitiveInstancesClosure
diff --git a/Engine/Source/Runtime/Renderer/Private/ScenePrivate.h b/Engine/Source/Runtime/Renderer/Private/ScenePrivate.h
--- a/Engine/Source/Runtime/Renderer/Private/ScenePrivate.h
+++ b/Engine/Source/Runtime/Renderer/Private/ScenePrivate.h
@@ -2368,7 +2368,7 @@ struct FPrimitiveBounds
 {
 	FBoxSphereBounds BoxSphereBounds;
 	/** Square of the minimum draw distance for the primitive. */
-	float MinDrawDistanceSq;
+	float MinDrawDistance;
 	/** Maximum draw distance for the primitive. */
 	float MaxDrawDistance;
 	/** Maximum cull distance for the primitive. This is only different from the MaxDrawDistance for HLOD.*/
diff --git a/Engine/Source/Runtime/Renderer/Private/SceneRendering.h b/Engine/Source/Runtime/Renderer/Private/SceneRendering.h
--- a/Engine/Source/Runtime/Renderer/Private/SceneRendering.h
+++ b/Engine/Source/Runtime/Renderer/Private/SceneRendering.h
@@ -1139,6 +1139,7 @@ struct FRayTracingCullingParameters
 	float FarFieldCullingRadius;
 	float CullAngleThreshold;
 	float AngleThresholdRatio;
+	float AngleThresholdRatioSq;
 	FVector ViewOrigin;
 	FVector ViewDirection;
 	bool bCullAllObjects;
@@ -1170,15 +1171,15 @@ public:
 	/** A map from primitive ID to a boolean visibility value. */
 	FSceneBitArray PrimitiveVisibilityMap;
 
-	/** Bit set when a primitive is known to be unoccluded. */
+	/** A map from primitive ID to a boolean ray tracing visibility value. */
+	FSceneBitArray PrimitiveRayTracingVisibilityMap;
+
+	/** Bit set when a primitive is known to be un-occluded. */
 	FSceneBitArray PrimitiveDefinitelyUnoccludedMap;
 
 	/** A map from primitive ID to a boolean is fading value. */
 	FSceneBitArray PotentiallyFadingPrimitiveMap;
 
-	/** A map from primitive ID to whether the primitve was culled for ray tracing */
-	FSceneBitArray RayTracingCullingPrimitiveMap;
-
 	/** Primitive fade uniform buffers, indexed by packed primitive index. */
 	TArray<FRHIUniformBuffer*,SceneRenderingAllocator> PrimitiveFadeUniformBuffers;
 
diff --git a/Engine/Source/Runtime/Renderer/Private/SceneVisibility.cpp b/Engine/Source/Runtime/Renderer/Private/SceneVisibility.cpp
--- a/Engine/Source/Runtime/Renderer/Private/SceneVisibility.cpp
+++ b/Engine/Source/Runtime/Renderer/Private/SceneVisibility.cpp
@@ -43,6 +43,7 @@
 #include "Algo/Unique.h"
 #include "InstanceCulling/InstanceCullingManager.h"
 #include "TemporalAA.h"
+#include "RayTracing/RayTracingInstanceCulling.h"
 
 /*------------------------------------------------------------------------------
 	Globals
@@ -239,6 +240,12 @@ static FAutoConsoleVariableRef CVarLightMaxDrawDistanceScale(
 	ECVF_Scalability | ECVF_RenderThreadSafe
 );
 
+static TAutoConsoleVariable<bool> CVarEnableFrustumCull(
+	TEXT("r.EnableFrustumCull"),
+	true,
+	TEXT("Enables or disables frustum culling.  Useful for comparing results to ensure culling is functioning properly."),
+	ECVF_RenderThreadSafe);
+
 static bool GOcclusionSingleRHIThreadStall = false;
 static FAutoConsoleVariableRef CVarOcclusionSingleRHIThreadStall(
 	TEXT("r.Occlusion.SingleRHIThreadStall"),
@@ -517,9 +524,6 @@ static FAutoConsoleVariableRef CVarFrustumCullNumWordsPerTask(
 	ECVF_Default
 	);
 
-
-#define NANITE_MESH_ALWAYS_VISIBLE 0
-
 static TAutoConsoleVariable CVarNaniteMeshsAlwaysVisible(
 	TEXT("r.Nanite.PrimitivesAlwaysVisible"),
 	0,
@@ -527,21 +531,10 @@ static TAutoConsoleVariable CVarNaniteMeshsAlwaysVisible(
 	ECVF_Default
 );
 
-// Templated version for performance
-template<bool bNaniteAlwaysVisible>
-FORCEINLINE bool IsAlwaysVisible(const FScene* RESTRICT Scene, int32 Index);
-
-// Use template specialization to ensure proper optimizations are applied.  Runtime branching will still occur otherwise.
-template<>
-FORCEINLINE bool IsAlwaysVisible<true>(const FScene* RESTRICT Scene, int32 Index)
-{
-	return Scene->PrimitiveFlagsCompact[Index].bIsNaniteMesh;
-}
-
-template<>
-FORCEINLINE bool IsAlwaysVisible<false>(const FScene* RESTRICT Scene, int32 Index)
+// Access when not on the render thread
+FORCEINLINE bool IsAlwaysVisible(const FScene* RESTRICT Scene, int32 Index, bool bNaniteAlwaysVisible)
 {
-	return false;
+	return bNaniteAlwaysVisible ? Scene->PrimitiveFlagsCompact[Index].bIsNaniteMesh : false;
 }
 
 // Non template version
@@ -555,25 +548,40 @@ FORCEINLINE bool IsAlwaysVisible(const FScene* RESTRICT Scene, int32 Index)
 	return false;
 }
 
+struct FPrimitiveCullingFlags
+{
+	bool bShouldVisibilityCull;
+	bool bUseCustomCulling;
+	bool bAlsoUseSphereTest;
+	bool bUseFastIntersect;
+	bool bUseVisibilityOctree;
+	bool bNaniteAlwaysVisible;
+	bool bHasHiddenPrimitives;
+	bool bHasShowOnlyPrimitives;
+
+#if RHI_RAYTRACING
+	bool bCullInRayTracing;
+#endif
+};
+
 // Returns true if the frustum and bounds intersect
-template<bool UseCustomCulling, bool bAlsoUseSphereTest, bool bUseFastIntersect>
-FORCEINLINE static bool IsPrimitiveVisible(FViewInfo& View, const FPlane* PermutedPlanePtr, const FPrimitiveBounds& RESTRICT Bounds, int32 VisibilityId)
+FORCEINLINE bool IsPrimitiveVisible(FViewInfo& View, const FPlane* PermutedPlanePtr, const FPrimitiveBounds& RESTRICT Bounds, int32 VisibilityId, const FPrimitiveCullingFlags& Flags)
 {
 	// The custom culling and sphere culling are additional tests, meaning that if they pass, the
 	// remaining culling tests will still be performed.  If any of the tests fail, then the primitive
 	// is culled, and the remaining tests do not need be performed
 
-	if (UseCustomCulling && !View.CustomVisibilityQuery->IsVisible(VisibilityId, FBoxSphereBounds(Bounds.BoxSphereBounds.Origin, Bounds.BoxSphereBounds.BoxExtent, Bounds.BoxSphereBounds.SphereRadius)))
+	if (Flags.bUseCustomCulling && !View.CustomVisibilityQuery->IsVisible(VisibilityId, FBoxSphereBounds(Bounds.BoxSphereBounds.Origin, Bounds.BoxSphereBounds.BoxExtent, Bounds.BoxSphereBounds.SphereRadius)))
 	{
 		return false;
 	}
 
-	if (bAlsoUseSphereTest && !View.ViewFrustum.IntersectSphere(Bounds.BoxSphereBounds.Origin, Bounds.BoxSphereBounds.SphereRadius))
+	if (Flags.bAlsoUseSphereTest && !View.ViewFrustum.IntersectSphere(Bounds.BoxSphereBounds.Origin, Bounds.BoxSphereBounds.SphereRadius))
 	{
 		return false;
 	}
 
-	if (bUseFastIntersect)
+	if (Flags.bUseFastIntersect)
 	{
 		return IntersectBox8Plane(Bounds.BoxSphereBounds.Origin, Bounds.BoxSphereBounds.BoxExtent, PermutedPlanePtr);
 	}
@@ -583,13 +591,89 @@ FORCEINLINE static bool IsPrimitiveVisible(FViewInfo& View, const FPlane* Permut
 	}
 }
 
+FORCEINLINE bool IsPrimitiveHidden(const FScene* RESTRICT Scene, FViewInfo& View, int PrimitiveIndex, const FPrimitiveCullingFlags& Flags)
+{
+	// If any primitives are explicitly hidden, remove them now.
+	if (Flags.bHasHiddenPrimitives && View.HiddenPrimitives.Contains(Scene->PrimitiveComponentIds[PrimitiveIndex]))
+	{
+		return true;
+	}
+
+	// If the view has any show only primitives, hide everything else
+	if (Flags.bHasShowOnlyPrimitives && !View.ShowOnlyPrimitives->Contains(Scene->PrimitiveComponentIds[PrimitiveIndex]))
+	{
+		return true;
+	}
+
+	return false;
+}
+
+#if RHI_RAYTRACING
+
+FORCEINLINE bool ShouldCullForRayTracing(const FScene* RESTRICT Scene, FViewInfo& View, int32 PrimitiveIndex)
+{
+	const FRayTracingCullingParameters& RayTracingCullingParameters = View.RayTracingCullingParameters;
+
+	if (RayTracing::CullPrimitiveByFlags(RayTracingCullingParameters, Scene, PrimitiveIndex))
+	{
+		return true;
+	}
+
+	const bool bIsFarFieldPrimitive = EnumHasAnyFlags(Scene->PrimitiveRayTracingFlags[PrimitiveIndex], ERayTracingPrimitiveFlags::FarField);
+	const FPrimitiveBounds& RESTRICT Bounds = Scene->PrimitiveBounds[PrimitiveIndex];
+	return RayTracing::ShouldCullBounds(RayTracingCullingParameters, Bounds.BoxSphereBounds, bIsFarFieldPrimitive);
+};
+#endif //RHI_RAYTRACING
+
+static FORCEINLINE void CullOctree(const FScene* RESTRICT Scene, FViewInfo& View, const FPrimitiveCullingFlags& Flags, FSceneBitArray& OutVisibleNodes)
+{
+	TRACE_CPUPROFILER_EVENT_SCOPE(SceneVisibility_CullOctree);
+
+	// Two bits per octree node, 1st bit is Inisde Frustum, 2nd bit is Outside Frustum
+	OutVisibleNodes.Init(false, Scene->PrimitiveOctree.GetNumNodes() * 2);
+
+	Scene->PrimitiveOctree.FindNodesWithPredicate(
+		[&View, &OutVisibleNodes, &Flags](FScenePrimitiveOctree::FNodeIndex ParentNodeIndex, FScenePrimitiveOctree::FNodeIndex NodeIndex, const FBoxCenterAndExtent& NodeBounds)
+		{
+			// If the parent node is completely contained there is no need to test containment
+			if (ParentNodeIndex != INDEX_NONE && !OutVisibleNodes[(ParentNodeIndex * 2) + 1])
+			{
+				OutVisibleNodes[NodeIndex * 2] = true;
+				OutVisibleNodes[NodeIndex * 2 + 1] = false;
+				return true;
+			}
+
+			const FPlane* PermutedPlanePtr = View.ViewFrustum.PermutedPlanes.GetData();
+			bool bIntersects = false;
+
+			if (Flags.bUseFastIntersect)
+			{
+				bIntersects = IntersectBox8Plane(NodeBounds.Center, NodeBounds.Extent, PermutedPlanePtr);
+			}
+			else
+			{
+				bIntersects = View.ViewFrustum.IntersectBox(NodeBounds.Center, NodeBounds.Extent);
+			}
+
+			if (bIntersects)
+			{
+				OutVisibleNodes[NodeIndex * 2] = true;
+				OutVisibleNodes[NodeIndex * 2 + 1] = View.ViewFrustum.GetBoxIntersectionOutcode(NodeBounds.Center, NodeBounds.Extent).GetOutside();
+			}
+
+			return bIntersects;
+		},
+		[](FScenePrimitiveOctree::FNodeIndex /*ParentNodeIndex*/, FScenePrimitiveOctree::FNodeIndex /*NodeIndex*/, const FBoxCenterAndExtent& /*NodeBounds*/)
+		{
+
+		});
+}
 
-template<bool UseCustomCulling, bool bAlsoUseSphereTest, bool bUseFastIntersect, bool bUseVisibilityOctree, bool bNaniteAlwaysVisible>
-FORCENOINLINE void FrustumCullTask(FThreadSafeCounter& NumCulledPrimitives, const FScene* RESTRICT Scene, FViewInfo& View, float MaxDrawDistanceScale, const FHLODVisibilityState* const HLODState, FSceneBitArray& VisibleNodes, int32 TaskIndex)
+static void PrimitiveCullTask(FThreadSafeCounter& NumCulledPrimitives, const FScene* RESTRICT Scene, FViewInfo& View, FPrimitiveCullingFlags Flags, float MaxDrawDistanceScale, const FHLODVisibilityState* const HLODState, const FSceneBitArray& VisibleNodes, int32 TaskIndex)
 {
-	TRACE_CPUPROFILER_EVENT_SCOPE(SceneVisibility_FrustumCull);
-	SCOPED_NAMED_EVENT(SceneVisibility_FrustumCull, FColor::Red);
-	QUICK_SCOPE_CYCLE_COUNTER(STAT_FrustumCull_Loop);
+	TRACE_CPUPROFILER_EVENT_SCOPE(SceneVisibility_PrimitiveCull);
+	SCOPED_NAMED_EVENT(SceneVisibility_PrimitiveCull, FColor::Red);
+	QUICK_SCOPE_CYCLE_COUNTER(STAT_PrimitiveCull_Loop);
 
 	FTaskTagScope TaskTagScope(ETaskTag::EParallelRenderingThread);
 
@@ -605,106 +689,146 @@ FORCENOINLINE void FrustumCullTask(FThreadSafeCounter& NumCulledPrimitives, cons
 	// Primitives may be explicitly removed from stereo views when using mono
 	const int32 TaskWordOffset = TaskIndex * FrustumCullNumWordsPerTask;
 
+	FVector ViewOrigin = View.ViewMatrices.GetViewOrigin();
+
 	for (int32 WordIndex = TaskWordOffset; WordIndex < TaskWordOffset + FrustumCullNumWordsPerTask && WordIndex * NumBitsPerDWORD < BitArrayNumInner; WordIndex++)
 	{
 		uint32 Mask = 0x1;
 		uint32 VisBits = 0;
 		uint32 FadingBits = 0;
+
+		// If visibility culling is disabled, make sure to use the existing visibility state
+		if (!Flags.bShouldVisibilityCull)
+		{
+			VisBits = View.PrimitiveVisibilityMap.GetData()[WordIndex];
+		}
+		
+
+#if RHI_RAYTRACING
+		uint32 RayTracingBits = 0;
+#endif //RHI_RAYTRACING
 		for (int32 BitSubIndex = 0; BitSubIndex < NumBitsPerDWORD && WordIndex * NumBitsPerDWORD + BitSubIndex < BitArrayNumInner; BitSubIndex++, Mask <<= 1)
 		{
 			int32 Index = WordIndex * NumBitsPerDWORD + BitSubIndex;
-			bool bIsVisible = true;
+			bool bPrimitiveIsHidden = IsPrimitiveHidden(Scene, View, Index, Flags);
+			bool bIsVisible = Flags.bShouldVisibilityCull ? true : (VisBits & Mask) == Mask;
+			
+			bIsVisible = bIsVisible && !bPrimitiveIsHidden;
+
+#if RHI_RAYTRACING
+			bool bIsVisibleInRayTracing = true;
+
+			if (bPrimitiveIsHidden || (Flags.bCullInRayTracing && ShouldCullForRayTracing(Scene, View, Index)))
+			{
+				bIsVisibleInRayTracing = false;
+			}
+#endif //RHI_RAYTRACING
+
+			const FPrimitiveBounds& RESTRICT Bounds = Scene->PrimitiveBounds[Index];
 
 			// Handle primitives that are not always visible.
-			if (!IsAlwaysVisible<bNaniteAlwaysVisible>(Scene, Index))
+			if (Flags.bShouldVisibilityCull && bIsVisible && !IsAlwaysVisible(Scene, Index, Flags.bNaniteAlwaysVisible))
 			{
+				bool bShouldDistanceCull = true;
 				bool bPartiallyOutside = true;
-				if (bUseVisibilityOctree)
-				{
-					// If the parent octree node was completely contained by the frustum, there is no need do an additional frustum test on the primitive bounds
-					// If the parent octree node is partially in the frustum, perform an additional test on the primitive bounds
-					uint32 OctreeNodeIndex = Scene->PrimitiveOctreeIndex[Index];
+				bool bShouldFrustumCull = true;
 
-					bIsVisible = VisibleNodes[OctreeNodeIndex * 2];
-					bPartiallyOutside = VisibleNodes[OctreeNodeIndex * 2 + 1];
+				// Fading HLODs and their children must be visible, objects hidden by HLODs can be culled
+				if (HLODState)
+				{
+					if (HLODState->IsNodeForcedVisible(Index))
+					{
+						bShouldFrustumCull = false;
+						bShouldDistanceCull = false;
+					}
+					else if (HLODState->IsNodeForcedHidden(Index))
+					{
+						bShouldDistanceCull = false;
+						bShouldFrustumCull = false;
+						bIsVisible = false;
+					}
 				}
 
 				// Frustum first
-				if (bIsVisible)
+				if (bShouldFrustumCull)
 				{
-					const FPrimitiveBounds& RESTRICT Bounds = Scene->PrimitiveBounds[Index];
-
-					int32 VisibilityId = INDEX_NONE;
-
-					if (UseCustomCulling &&
-						((Scene->PrimitiveOcclusionFlags[Index] & CustomVisibilityFlags) == CustomVisibilityFlags))
+					if (Flags.bUseVisibilityOctree)
 					{
-						VisibilityId = Scene->PrimitiveVisibilityIds[Index].ByteIndex;
-					}
+						// If the parent octree node was completely contained by the frustum, there is no need do an additional frustum test on the primitive bounds
+						// If the parent octree node is partially in the frustum, perform an additional test on the primitive bounds
+						uint32 OctreeNodeIndex = Scene->PrimitiveOctreeIndex[Index];
 
-					bIsVisible = !bPartiallyOutside || IsPrimitiveVisible<UseCustomCulling, bAlsoUseSphereTest, bUseFastIntersect>(View, PermutedPlanePtr, Bounds, VisibilityId);
+						bIsVisible = VisibleNodes[OctreeNodeIndex * 2];
+						bPartiallyOutside = VisibleNodes[OctreeNodeIndex * 2 + 1];
+					}
 
-					// Distance cull if frustum cull passed
 					if (bIsVisible)
 					{
-						bool bShouldDistanceCull = true;
+						int32 VisibilityId = INDEX_NONE;
 
-						// If cull distance is disabled, always show the primitive (except foliage)
-						if (View.Family->EngineShowFlags.DistanceCulledPrimitives
-							&& !Scene->Primitives[Index]->Proxy->IsDetailMesh()) // Proxy call is intentionally behind the DistancedCulledPrimitives check to prevent an expensive memory read
+						if (Flags.bUseCustomCulling &&
+							((Scene->PrimitiveOcclusionFlags[Index] & CustomVisibilityFlags) == CustomVisibilityFlags))
 						{
-							bShouldDistanceCull = false;
+							VisibilityId = Scene->PrimitiveVisibilityIds[Index].ByteIndex;
 						}
 
-						// Fading HLODs and their children must be visible, objects hidden by HLODs can be culled
-						if (HLODState)
+						bIsVisible = !bPartiallyOutside || IsPrimitiveVisible(View, PermutedPlanePtr, Bounds, VisibilityId, Flags);
+					}
+				}
+
+				bShouldDistanceCull = bShouldDistanceCull && bIsVisible;
+
+				// Distance cull if frustum cull passed
+				if (bShouldDistanceCull)
+				{
+					// If cull distance is disabled, always show the primitive (except foliage)
+					if (View.Family->EngineShowFlags.DistanceCulledPrimitives
+						&& !Scene->Primitives[Index]->Proxy->IsDetailMesh()) // Proxy call is intentionally behind the DistancedCulledPrimitives check to prevent an expensive memory read
+					{
+						bShouldDistanceCull = false;
+					}
+				}
+
+				if (bShouldDistanceCull)
+				{
+					// Preserve infinite draw distance
+					bool bHasMaxDrawDistance = Bounds.MaxCullDistance < FLT_MAX;
+					bool bHasMinDrawDistance = Bounds.MinDrawDistance > 0;
+
+					if (bHasMaxDrawDistance || bHasMinDrawDistance)
+					{
+						float MaxDrawDistance = Bounds.MaxCullDistance * MaxDrawDistanceScale;
+						float MinDrawDistanceSq = FMath::Square(Bounds.MinDrawDistance);
+						float DistanceSquared = FVector::DistSquared(Bounds.BoxSphereBounds.Origin, ViewOriginForDistanceCulling);
+
+						// Always test the fade in distance.  If a primitive was set to always draw, it may need to be faded in.
+						if (bHasMaxDrawDistance)
 						{
-							if (HLODState->IsNodeForcedVisible(Index))
+							float MaxFadeDistanceSquared = FMath::Square(MaxDrawDistance + FadeRadius);
+							float MinFadeDistanceSquared = FMath::Square(MaxDrawDistance - FadeRadius);
+							if ((DistanceSquared < MaxFadeDistanceSquared && DistanceSquared > MinFadeDistanceSquared)
+								&& Scene->Primitives[Index]->Proxy->IsUsingDistanceCullFade())  // Proxy call is intentionally behind the fade check to prevent an expensive memory read
 							{
-								bShouldDistanceCull = false;
-							}
-							else if (HLODState->IsNodeForcedHidden(Index))
-							{
-								bShouldDistanceCull = false;
-								bIsVisible = false;
+								FadingBits |= Mask;
 							}
 						}
 
-						if (bShouldDistanceCull)
-						{
-							// Preserve infinite draw distance
-							bool bHasMaxDrawDistance = Bounds.MaxCullDistance < FLT_MAX;
-							bool bHasMinDrawDistance = Bounds.MinDrawDistanceSq > 0;
+						// Check for distance culling first
+						const bool bFarDistanceCulled = bHasMaxDrawDistance && (DistanceSquared > FMath::Square(MaxDrawDistance));
+						const bool bNearDistanceCulled = bHasMinDrawDistance && (DistanceSquared < MinDrawDistanceSq);
+						bool bIsDistanceCulled = bNearDistanceCulled || bFarDistanceCulled;
 
-							if (bHasMaxDrawDistance || bHasMinDrawDistance)
-							{
-								float MaxDrawDistance = Bounds.MaxCullDistance * MaxDrawDistanceScale;
-								float MinDrawDistanceSq = Bounds.MinDrawDistanceSq;
-								float DistanceSquared = FVector::DistSquared(Bounds.BoxSphereBounds.Origin, ViewOriginForDistanceCulling);
-
-								// Always test the fade in distance.  If a primitive was set to always draw, it may need to be faded in.
-								if (bHasMaxDrawDistance)
-								{
-									float MaxFadeDistanceSquared = FMath::Square(MaxDrawDistance + FadeRadius);
-									float MinFadeDistanceSquared = FMath::Square(MaxDrawDistance - FadeRadius);
-									if ((DistanceSquared < MaxFadeDistanceSquared && DistanceSquared > MinFadeDistanceSquared)
-										&& Scene->Primitives[Index]->Proxy->IsUsingDistanceCullFade())  // Proxy call is intentionally behind the fade check to prevent an expensive memory read
-									{
-										FadingBits |= Mask;
-									}
-								}
-
-								// Check for distance culling first
-								const bool bFarDistanceCulled = bHasMaxDrawDistance && (DistanceSquared > FMath::Square(MaxDrawDistance));
-								const bool bNearDistanceCulled = bHasMinDrawDistance && (DistanceSquared < MinDrawDistanceSq);
-								bool bIsDistanceCulled = bNearDistanceCulled || bFarDistanceCulled;
+						if (bIsDistanceCulled)
+						{
+							bIsVisible = false;
+						}
 
-								if (bIsDistanceCulled)
-								{
-									bIsVisible = false;
-								}
-							}
+#if RHI_RAYTRACING
+						if (bFarDistanceCulled)
+						{
+							bIsVisibleInRayTracing = false;
 						}
+#endif //RHI_RAYTRACING
 					}
 				}
 			}
@@ -718,81 +842,65 @@ FORCENOINLINE void FrustumCullTask(FThreadSafeCounter& NumCulledPrimitives, cons
 			{
 				STAT(++NumPrimitivesCulledForTask);
 			}
+
+#if RHI_RAYTRACING
+			if (bIsVisibleInRayTracing)
+			{
+				RayTracingBits |= Mask;
+			}
+#endif //RHI_RAYTRACING
 		}
-		if (FadingBits)
+
+		if (Flags.bShouldVisibilityCull && FadingBits)
 		{
 			checkSlow(!View.PotentiallyFadingPrimitiveMap.GetData()[WordIndex]); // this should start at zero
 			View.PotentiallyFadingPrimitiveMap.GetData()[WordIndex] = FadingBits;
 		}
-		if (VisBits)
+
+		if (Flags.bShouldVisibilityCull && VisBits)
 		{
 			checkSlow(!View.PrimitiveVisibilityMap.GetData()[WordIndex]); // this should start at zero
 			View.PrimitiveVisibilityMap.GetData()[WordIndex] = VisBits;
 		}
+
+#if RHI_RAYTRACING
+		if (RayTracingBits)
+		{
+			checkSlow(!View.PrimitiveVisibilityMap.GetData()[WordIndex]); // this should start at zero
+			View.PrimitiveRayTracingVisibilityMap.GetData()[WordIndex] = RayTracingBits;
+		}
+#endif
 	}
 
 	STAT(NumCulledPrimitives.Add(NumPrimitivesCulledForTask));
 }
 
-template<bool UseCustomCulling, bool bAlsoUseSphereTest, bool bUseFastIntersect, bool bUseVisibilityOctree, bool bNaniteAlwaysVisible>
-static int32 FrustumCullTemplated(const FScene* RESTRICT Scene, FViewInfo& View)
+static int32 PrimitiveCull(const FScene* RESTRICT Scene, FViewInfo& View, bool bShouldVisibilityCull)
 {
-	SCOPE_CYCLE_COUNTER(STAT_FrustumCull);
+	FPrimitiveCullingFlags Flags;
+	Flags.bShouldVisibilityCull = bShouldVisibilityCull;
+	Flags.bUseCustomCulling = View.CustomVisibilityQuery && View.CustomVisibilityQuery->Prepare();
+	Flags.bAlsoUseSphereTest = CVarAlsoUseSphereForFrustumCull.GetValueOnRenderThread() > 0;
+	Flags.bUseFastIntersect = (View.ViewFrustum.PermutedPlanes.Num() == 8) && CVarUseFastIntersect.GetValueOnRenderThread();
+	Flags.bUseVisibilityOctree = CVarUseVisibilityOctree.GetValueOnRenderThread() > 0;
+	Flags.bNaniteAlwaysVisible = CVarNaniteMeshsAlwaysVisible.GetValueOnRenderThread() > 0;
+	Flags.bHasHiddenPrimitives = View.HiddenPrimitives.Num() > 0;
+	Flags.bHasShowOnlyPrimitives = View.ShowOnlyPrimitives.IsSet();
 
-	FThreadSafeCounter NumCulledPrimitives;
-	float MaxDrawDistanceScale = GetCachedScalabilityCVars().ViewDistanceScale;
-	MaxDrawDistanceScale *= GetCachedScalabilityCVars().CalculateFieldOfViewDistanceScale(View.DesiredFOV);
+#if RHI_RAYTRACING
+	View.RayTracingCullingParameters.Init(View);
+	Flags.bCullInRayTracing = View.RayTracingCullingParameters.CullInRayTracing > 0;
+#endif
 
-	FSceneViewState* ViewState = (FSceneViewState*)View.State;
-	const bool bHLODActive = Scene->SceneLODHierarchy.IsActive();
-	const FHLODVisibilityState* const HLODState = bHLODActive && ViewState ? &ViewState->HLODVisibilityState : nullptr;
+	SCOPE_CYCLE_COUNTER(STAT_PrimitiveCull);
 
 	FMemMark MemStackMark(FMemStack::Get());
 
 	FSceneBitArray VisibleNodes;
 
-	if (bUseVisibilityOctree)
+	if (bShouldVisibilityCull && Flags.bUseVisibilityOctree)
 	{
-		TRACE_CPUPROFILER_EVENT_SCOPE(SceneVisibility_CullOctree);
-
-		// Two bits per octree node, 1st bit is Inisde Frustum, 2nd bit is Outside Frustum
-		VisibleNodes.Init(false, Scene->PrimitiveOctree.GetNumNodes() * 2);
-
-		Scene->PrimitiveOctree.FindNodesWithPredicate(
-			[&View, &VisibleNodes](FScenePrimitiveOctree::FNodeIndex ParentNodeIndex, FScenePrimitiveOctree::FNodeIndex NodeIndex, const FBoxCenterAndExtent& NodeBounds)
-			{
-				// If the parent node is completely contained there is no need to test containment
-				if (ParentNodeIndex != INDEX_NONE && !VisibleNodes[(ParentNodeIndex * 2) + 1])
-				{
-					VisibleNodes[NodeIndex * 2] = true;
-					VisibleNodes[NodeIndex * 2 + 1] = false;
-					return true;
-				}
-
-				const FPlane* PermutedPlanePtr = View.ViewFrustum.PermutedPlanes.GetData();
-				bool bIntersects = false;
-
-				if (bUseFastIntersect)
-				{
-					bIntersects = IntersectBox8Plane(NodeBounds.Center, NodeBounds.Extent, PermutedPlanePtr);
-				}
-				else
-				{
-					bIntersects = View.ViewFrustum.IntersectBox(NodeBounds.Center, NodeBounds.Extent);
-				}
-
-				if (bIntersects)
-				{
-					VisibleNodes[NodeIndex * 2] = true;
-					VisibleNodes[NodeIndex * 2 + 1] = View.ViewFrustum.GetBoxIntersectionOutcode(NodeBounds.Center, NodeBounds.Extent).GetOutside();
-				}
-
-				return bIntersects;
-			},
-			[&View, &VisibleNodes](FScenePrimitiveOctree::FNodeIndex /*ParentNodeIndex*/, FScenePrimitiveOctree::FNodeIndex /*NodeIndex*/, const FBoxCenterAndExtent& /*NodeBounds*/)
-			{
-				
-			});
+		CullOctree(Scene, View, Flags, VisibleNodes);
 	}
 
 	//Primitives per ParallelFor task
@@ -800,93 +908,26 @@ static int32 FrustumCullTemplated(const FScene* RESTRICT Scene, FViewInfo& View)
 	//Performance varies on total primitive count and tasks scheduled. Check the mentioned link above for some measurements.
 	//There have been some changes as compared to the code measured in the link
 
+	FThreadSafeCounter NumCulledPrimitives;
+	FSceneViewState* ViewState = (FSceneViewState*)View.State;
+	const bool bHLODActive = Scene->SceneLODHierarchy.IsActive();
+	const FHLODVisibilityState* const HLODState = bHLODActive && ViewState ? &ViewState->HLODVisibilityState : nullptr;
+	float MaxDrawDistanceScale = GetCachedScalabilityCVars().ViewDistanceScale * GetCachedScalabilityCVars().CalculateFieldOfViewDistanceScale(View.DesiredFOV);
+
 	const int32 BitArrayNum = View.PrimitiveVisibilityMap.Num();
 	const int32 BitArrayWords = FMath::DivideAndRoundUp(BitArrayNum, (int32)NumBitsPerDWORD);
 	const int32 NumTasks = FMath::DivideAndRoundUp(BitArrayWords, FrustumCullNumWordsPerTask);
 
 	ParallelFor(NumTasks,
-		[&NumCulledPrimitives, Scene, &View, MaxDrawDistanceScale, HLODState, &VisibleNodes](int32 TaskIndex)
+		[&NumCulledPrimitives, Scene, &View, MaxDrawDistanceScale, HLODState, &VisibleNodes, &Flags](int32 TaskIndex)
 		{
-			FrustumCullTask<UseCustomCulling, bAlsoUseSphereTest, bUseFastIntersect, bUseVisibilityOctree, bNaniteAlwaysVisible>(NumCulledPrimitives, Scene, View, MaxDrawDistanceScale, HLODState, VisibleNodes, TaskIndex);
+			PrimitiveCullTask(NumCulledPrimitives, Scene, View, Flags, MaxDrawDistanceScale, HLODState, VisibleNodes, TaskIndex);
 		},
-		!FApp::ShouldUseThreadingForPerformance() || (UseCustomCulling && !View.CustomVisibilityQuery->IsThreadsafe()) || CVarParallelInitViews.GetValueOnRenderThread() == 0 || !IsInActualRenderingThread()
+		!FApp::ShouldUseThreadingForPerformance() || (Flags.bUseCustomCulling && !View.CustomVisibilityQuery->IsThreadsafe()) || CVarParallelInitViews.GetValueOnRenderThread() == 0 || !IsInActualRenderingThread()
 		);
-	
-
-	return NumCulledPrimitives.GetValue();
-}
-
-// Macros to help fill the templated frustum cull table without errors
-#define FRUSTUM_CULL_TABLE1(...) \
-	FrustumCullTemplated<false, __VA_ARGS__>, \
-	FrustumCullTemplated<true, __VA_ARGS__>, 
-
-#define FRUSTUM_CULL_TABLE2(...) \
-	FRUSTUM_CULL_TABLE1(false, __VA_ARGS__) \
-	FRUSTUM_CULL_TABLE1(true, __VA_ARGS__) 
-
-#define FRUSTUM_CULL_TABLE3(...) \
-	FRUSTUM_CULL_TABLE2(false, __VA_ARGS__) \
-	FRUSTUM_CULL_TABLE2(true, __VA_ARGS__) 
-
-#define FRUSTUM_CULL_TABLE4(...) \
-	FRUSTUM_CULL_TABLE3(false, __VA_ARGS__) \
-	FRUSTUM_CULL_TABLE3(true, __VA_ARGS__)
 
-#define FRUSTUM_CULL_TABLE5() \
-	FRUSTUM_CULL_TABLE4(false) \
-	FRUSTUM_CULL_TABLE4(true)
 
-// Matrix of template inputs to template functions.  See EFrustumCullFlags below
-typedef int32(*GFunctionCullMethod)(const FScene* RESTRICT Scene, FViewInfo& View);
-static const GFunctionCullMethod GFrustumCullTable[] =
-{
-	FRUSTUM_CULL_TABLE5()
-};
-
-static int32 FrustumCull(const FScene* RESTRICT Scene, FViewInfo& View)
-{
-	enum EFrustumCullFlags
-	{
-		UseCustomCulling = 1 << 0,
-		AlsoUseSphereTest = 1 << 1,
-		UseFastInterscect = 1 << 2,
-		UseVisibilityOctree = 1 << 3,
-		NaniteAlwaysVisible = 1 << 4,
-		EFrustumCullFlags_End
-	};
-
-	// Make sure the the table matches the enumeration size
-	static_assert(((EFrustumCullFlags_End -1) << 1)  == (sizeof(GFrustumCullTable) / sizeof(GFunctionCullMethod)), "Invalid number of frustum cull methods found in GFrustumCullTable");
-
-	uint32 Flags = 0;
-	
-	if (View.CustomVisibilityQuery && View.CustomVisibilityQuery->Prepare())
-	{
-		Flags |= UseCustomCulling;
-	}
-
-	if (CVarAlsoUseSphereForFrustumCull.GetValueOnRenderThread())
-	{
-		Flags |= AlsoUseSphereTest;
-	}
-
-	if ((View.ViewFrustum.PermutedPlanes.Num() == 8) && CVarUseFastIntersect.GetValueOnRenderThread())
-	{
-		Flags |= UseFastInterscect;
-	}
-
-	if (CVarUseVisibilityOctree.GetValueOnRenderThread())
-	{
-		Flags |= UseVisibilityOctree;
-	}
-
-	if (CVarNaniteMeshsAlwaysVisible.GetValueOnRenderThread())
-	{
-		Flags |= NaniteAlwaysVisible;
-	}
-
-	return GFrustumCullTable[Flags](Scene, View);
+	return NumCulledPrimitives.GetValue();
 }
 
 /**
@@ -930,6 +971,17 @@ static void UpdatePrimitiveFading(const FScene* Scene, FViewInfo& View)
 				{
 					// If the primitive is fading out make sure it remains visible.
 					View.PrimitiveVisibilityMap.AccessCorrespondingBit(BitIt) = true;
+
+#if RHI_RAYTRACING
+					// Cannot just assume the ray tracing visibility will be true, so a complete recalculation for its culling needs to happen
+					// This should be a very rare occurrence, so the hit is not worrisome.
+					// TODO:  Could this be moved into the actual culling phase?
+
+					if (View.RayTracingCullingParameters.CullInRayTracing == 0 || !ShouldCullForRayTracing(Scene, View, BitIt.GetIndex()))
+					{
+						View.PrimitiveRayTracingVisibilityMap.AccessCorrespondingBit(BitIt) = true;
+					}
+#endif //RHI_RAYTRACING
 				}
 				View.PrimitiveFadeUniformBuffers[BitIt.GetIndex()] = UniformBuffer;
 				View.PrimitiveFadeUniformBufferMap[BitIt.GetIndex()] = UniformBuffer != nullptr;
@@ -4283,6 +4335,7 @@ void FSceneRenderer::ComputeViewVisibility(
 
 			// Allocate the view's visibility maps.
 			View.PrimitiveVisibilityMap.Init(false, Scene->Primitives.Num());
+			View.PrimitiveRayTracingVisibilityMap.Init(false, Scene->Primitives.Num());
 
 			// These are not initialized here, as we overwrite the whole array in GatherDynamicMeshElements().
 			View.DynamicMeshEndIndices.SetNumUninitialized(Scene->Primitives.Num());
@@ -4343,7 +4396,7 @@ void FSceneRenderer::ComputeViewVisibility(
 				bUsedPrecomputedVisibility = true;
 			}
 
-			bool bNeedsFrustumCulling = true;
+			bool bNeedsFrustumCulling = CVarEnableFrustumCull.GetValueOnRenderThread();
 
 			// Development builds sometimes override frustum culling, e.g. dependent views in the editor.
 #if !(UE_BUILD_SHIPPING || UE_BUILD_TEST)
@@ -4382,56 +4435,35 @@ void FSceneRenderer::ComputeViewVisibility(
 #endif
 
 			// Most views use standard frustum culling.
-			if (bNeedsFrustumCulling)
+			if(bNeedsFrustumCulling)
 			{
+				// Update HLOD transition/visibility states to allow use during distance culling
+				FLODSceneTree& HLODTree = Scene->SceneLODHierarchy;
+				if (HLODTree.IsActive())
 				{
-					TRACE_CPUPROFILER_EVENT_SCOPE(FSceneRenderer_Cull);
-					// Update HLOD transition/visibility states to allow use during distance culling
-					FLODSceneTree& HLODTree = Scene->SceneLODHierarchy;
-					if (HLODTree.IsActive())
-					{
-						QUICK_SCOPE_CYCLE_COUNTER(STAT_ViewVisibilityTime_HLODUpdate);
-						HLODTree.UpdateVisibilityStates(View);
-					}
-					else
-					{
-						HLODTree.ClearVisibilityState(View);
-					}
-
-					int32 NumCulledPrimitivesForView = FrustumCull(Scene, View);
-					STAT(NumCulledPrimitives += NumCulledPrimitivesForView);
+					QUICK_SCOPE_CYCLE_COUNTER(STAT_ViewVisibilityTime_HLODUpdate);
+					HLODTree.UpdateVisibilityStates(View);
 				}
+				else
 				{
-					TRACE_CPUPROFILER_EVENT_SCOPE(FSceneRenderer_UpdatePrimitiveFading);
-					UpdatePrimitiveFading(Scene, View);
+					HLODTree.ClearVisibilityState(View);
 				}
 			}
 
-			// If any primitives are explicitly hidden, remove them now.
-			if (View.HiddenPrimitives.Num())
 			{
-				for (FSceneSetBitIterator BitIt(View.PrimitiveVisibilityMap); BitIt; ++BitIt)
-				{
-					if (View.HiddenPrimitives.Contains(Scene->PrimitiveComponentIds[BitIt.GetIndex()]))
-					{
-						// Should respect this over the always visible bit set, since we want the editor
-						// to support hiding layers of primitives for various reasons.
-						View.PrimitiveVisibilityMap.AccessCorrespondingBit(BitIt) = false;
-					}
-				}
+				TRACE_CPUPROFILER_EVENT_SCOPE(FSceneRenderer_Cull);
+				int32 NumCulledPrimitivesForView = PrimitiveCull(Scene, View, bNeedsFrustumCulling);
+				STAT(NumCulledPrimitives += NumCulledPrimitivesForView);
+			}
+
+			{
+				TRACE_CPUPROFILER_EVENT_SCOPE(FSceneRenderer_UpdatePrimitiveFading);
+				UpdatePrimitiveFading(Scene, View);
 			}
 
-			// If the view has any show only primitives, hide everything else
 			if (View.ShowOnlyPrimitives.IsSet())
 			{
 				View.bHasNoVisiblePrimitive = View.ShowOnlyPrimitives->Num() == 0;
-				for (FSceneSetBitIterator BitIt(View.PrimitiveVisibilityMap); BitIt; ++BitIt)
-				{
-					if (!View.ShowOnlyPrimitives->Contains(Scene->PrimitiveComponentIds[BitIt.GetIndex()]))
-					{
-						View.PrimitiveVisibilityMap.AccessCorrespondingBit(BitIt) = false;
-					}
-				}
 			}
 
 			if (View.bStaticSceneOnly)
@@ -5192,11 +5224,11 @@ void FLODSceneTree::UpdateVisibilityStates(FViewInfo& View)
 			}
 
 			FPrimitiveBounds& Bounds = Scene->PrimitiveBounds[NodeIndex];
-			const bool bForcedIntoView = FMath::IsNearlyZero(Bounds.MinDrawDistanceSq);
+			const bool bForcedIntoView = FMath::IsNearlyZero(Bounds.MinDrawDistance);
 
 			// Update visibility states of this node and owned children
 			const float DistanceSquared = Bounds.BoxSphereBounds.ComputeSquaredDistanceFromBoxToPoint(View.ViewMatrices.GetViewOrigin());
-			const bool bNearCulled = DistanceSquared < Bounds.MinDrawDistanceSq * HLODState.FOVDistanceScaleSq;
+			const bool bNearCulled = DistanceSquared < FMath::Square(Bounds.MinDrawDistance) * HLODState.FOVDistanceScaleSq;
 			const bool bFarCulled = DistanceSquared > Bounds.MaxDrawDistance * Bounds.MaxDrawDistance * HLODState.FOVDistanceScaleSq;
 			const bool bIsInDrawRange = !bNearCulled && !bFarCulled;
 
