Enable indirect clear/merge by default

diff --git a/Engine/Source/Runtime/Renderer/Private/VirtualShadowMaps/VirtualShadowMapArray.cpp b/Engine/Source/Runtime/Renderer/Private/VirtualShadowMaps/VirtualShadowMapArray.cpp
--- a/Engine/Source/Runtime/Renderer/Private/VirtualShadowMaps/VirtualShadowMapArray.cpp
+++ b/Engine/Source/Runtime/Renderer/Private/VirtualShadowMaps/VirtualShadowMapArray.cpp
@@ -194,14 +194,14 @@ static TAutoConsoleVariable<int32> CVarVirtualShadowMapDebugProjection(
 
 TAutoConsoleVariable<int32> CVarInitializePhysicalUsingIndirect(
 	TEXT("r.Shadow.Virtual.InitPhysicalUsingIndirect"),
-	0,
+	1,
 	TEXT("."),
 	ECVF_RenderThreadSafe
 );
 
 TAutoConsoleVariable<int32> CVarMergePhysicalUsingIndirect(
 	TEXT("r.Shadow.Virtual.MergePhysicalUsingIndirect"),
-	0,
+	1,
 	TEXT("."),
 	ECVF_RenderThreadSafe
 );
