Removed non-required memory allocations from FMeshDrawCommand::GetPipelineStateSortingKey when GRHISupportsPipelineStateSortKey is false.

diff --git a/Engine/Source/Runtime/Renderer/Private/MeshPassProcessor.cpp b/Engine/Source/Runtime/Renderer/Private/MeshPassProcessor.cpp
--- a/Engine/Source/Runtime/Renderer/Private/MeshPassProcessor.cpp
+++ b/Engine/Source/Runtime/Renderer/Private/MeshPassProcessor.cpp
@@ -1230,17 +1230,14 @@ void FMeshDrawCommand::SubmitDraw(
 
 uint64 FMeshDrawCommand::GetPipelineStateSortingKey(FRHICommandList& RHICmdList) const
 {
-	FGraphicsMinimalPipelineStateSet GraphicsMinimalPipelineStateSet;
-	const FGraphicsMinimalPipelineStateInitializer& MeshPipelineState = CachedPipelineId.GetPipelineState(GraphicsMinimalPipelineStateSet);
-
 	// Default fallback sort key
 	uint64 SortKey = CachedPipelineId.GetId();
-	//MeshPipelineState.BoundShaderState.PixelShaderIndex;
-	//SortKey = SortKey << 32llu;
-	//SortKey |= MeshPipelineState.BoundShaderState.VertexShaderIndex;
 
 	if (GRHISupportsPipelineStateSortKey)
 	{
+		FGraphicsMinimalPipelineStateSet GraphicsMinimalPipelineStateSet;
+		const FGraphicsMinimalPipelineStateInitializer& MeshPipelineState = CachedPipelineId.GetPipelineState(GraphicsMinimalPipelineStateSet);
+
 		FGraphicsPipelineStateInitializer GraphicsPSOInit = MeshPipelineState.AsGraphicsPipelineStateInitializer();
 		RHICmdList.ApplyCachedRenderTargets(GraphicsPSOInit);
 		FGraphicsPipelineState* PipelineState = PipelineStateCache::GetAndOrCreateGraphicsPipelineState(RHICmdList, GraphicsPSOInit, EApplyRendertargetOption::CheckApply);
