Fix typo in previous commit

diff --git a/Engine/Source/Runtime/Renderer/Private/VirtualShadowMaps/VirtualShadowMapArray.cpp b/Engine/Source/Runtime/Renderer/Private/VirtualShadowMaps/VirtualShadowMapArray.cpp
--- a/Engine/Source/Runtime/Renderer/Private/VirtualShadowMaps/VirtualShadowMapArray.cpp
+++ b/Engine/Source/Runtime/Renderer/Private/VirtualShadowMaps/VirtualShadowMapArray.cpp
@@ -1216,9 +1216,9 @@ void FVirtualShadowMapArray::BuildPageAllocations(
 			PassParameters->OutPageTable			 = GraphBuilder.CreateUAV(PageTableRDG);
 			PassParameters->OutPhysicalPageMetaData  = GraphBuilder.CreateUAV(PhysicalPageMetaDataRDG);
 			PassParameters->OutPageFlags			 = GraphBuilder.CreateUAV(PageFlagsRDG);
-			PassParameters->bDynamicPageInvalidation = 0;
+			PassParameters->bDynamicPageInvalidation = 1;
 #if !UE_BUILD_SHIPPING
-			PassParameters->bDynamicPageInvalidation = CVarDebugSkipDynamicPageInvalidation.GetValueOnRenderThread();
+			PassParameters->bDynamicPageInvalidation = CVarDebugSkipDynamicPageInvalidation.GetValueOnRenderThread() == 0 ? 1 : 0;
 #endif
 
 			bool bCacheEnabled = CacheManager->IsValid();
