Remove TLAS auto scratch buffer code path.
- Added checks to verify scratch buffer alignment.

diff --git a/Engine/Plugins/Experimental/GPULightmass/Source/GPULightmass/Private/Scene/Scene.h b/Engine/Plugins/Experimental/GPULightmass/Source/GPULightmass/Private/Scene/Scene.h
--- a/Engine/Plugins/Experimental/GPULightmass/Source/GPULightmass/Private/Scene/Scene.h
+++ b/Engine/Plugins/Experimental/GPULightmass/Source/GPULightmass/Private/Scene/Scene.h
@@ -99,6 +99,7 @@ public:
 	FRayTracingSceneRHIRef RayTracingScene;
 	FShaderResourceViewRHIRef RayTracingSceneSRV;
 	FBufferRHIRef RayTracingSceneBuffer;
+	FBufferRHIRef RayTracingScratchBuffer;
 
 	FRayTracingPipelineState* RayTracingPipelineState;
 	TSharedPtr<FViewInfo> ReferenceView;
diff --git a/Engine/Source/Runtime/D3D12RHI/Private/D3D12RayTracing.cpp b/Engine/Source/Runtime/D3D12RHI/Private/D3D12RayTracing.cpp
--- a/Engine/Source/Runtime/D3D12RHI/Private/D3D12RayTracing.cpp
+++ b/Engine/Source/Runtime/D3D12RHI/Private/D3D12RayTracing.cpp
@@ -3906,22 +3906,18 @@ void FD3D12RayTracingScene::BuildAccelerationStructure(FD3D12CommandContext& Com
 		TEXT("TLAS build scratch buffer now requires %lld bytes, but only %lld was calculated in the constructor."),
 			PrebuildInfo.ScratchDataSizeInBytes, SizeInfo.BuildScratchSize);
 
-	TRefCountPtr<FD3D12Buffer> AutoScratchBuffer;
-	if (ScratchBuffer == nullptr)
-	{
-		static const FName ScratchBufferName("AutoBuildScratchTLAS");
-		AutoScratchBuffer = CreateRayTracingBuffer(Adapter, GPUIndex, PrebuildInfo.ScratchDataSizeInBytes, ERayTracingBufferType::Scratch, ScratchBufferName);
-		ScratchBuffer = AutoScratchBuffer.GetReference();
-		ScratchBufferOffset = 0;
-	}
-
 	checkf(ScratchBuffer, TEXT("TLAS build requires scratch buffer of at least %lld bytes."), PrebuildInfo.ScratchDataSizeInBytes);
 
 	const D3D12_GPU_VIRTUAL_ADDRESS ScratchAddress = ScratchBuffer->ResourceLocation.GetGPUVirtualAddress() + ScratchBufferOffset;
 
-	checkf(ScratchAddress % GRHIRayTracingScratchBufferAlignment == 0,
+	checkf(ScratchAddress % GRHIRayTracingAccelerationStructureAlignment == 0,
 		TEXT("TLAS scratch buffer (plus offset) must be aligned to %lld bytes."),
-		GRHIRayTracingScratchBufferAlignment);
+		GRHIRayTracingAccelerationStructureAlignment);
+		
+	checkf(ScratchBuffer, TEXT("TLAS build requires scratch buffer of at least %lld bytes."), PrebuildInfo.ScratchDataSizeInBytes);
+
+	const D3D12_GPU_VIRTUAL_ADDRESS ScratchAddress = ScratchBuffer->ResourceLocation.GetGPUVirtualAddress() + ScratchBufferOffset;
+
 
 	checkf(PrebuildInfo.ScratchDataSizeInBytes + ScratchBufferOffset <= ScratchBuffer->GetSize(),
 		TEXT("TLAS scratch buffer size is %lld bytes with offset %lld (%lld bytes available), but the build requires %lld bytes. ")
diff --git a/Engine/Source/Runtime/RHI/Public/RHICommandList.h b/Engine/Source/Runtime/RHI/Public/RHICommandList.h
--- a/Engine/Source/Runtime/RHI/Public/RHICommandList.h
+++ b/Engine/Source/Runtime/RHI/Public/RHICommandList.h
@@ -2977,14 +2977,6 @@ public:
 		}
 	}
 
-	UE_DEPRECATED(5.0, "Use FRayTracingSceneBuildParams instead.")
-	FORCEINLINE_DEBUGGABLE void BuildAccelerationStructure(FRHIRayTracingScene* Scene)
-	{
-		FRayTracingSceneBuildParams Params;
-		Params.Scene = Scene;
-		BuildAccelerationStructure(Params);
-	}
-
 	FORCEINLINE_DEBUGGABLE void BindAccelerationStructureMemory(FRHIRayTracingScene* Scene, FRHIBuffer* Buffer, uint32 BufferOffset)
 	{
 		if (Bypass())
diff --git a/Engine/Source/Runtime/RHI/Public/RHIContext.h b/Engine/Source/Runtime/RHI/Public/RHIContext.h
--- a/Engine/Source/Runtime/RHI/Public/RHIContext.h
+++ b/Engine/Source/Runtime/RHI/Public/RHIContext.h
@@ -739,13 +739,6 @@ public:
 		checkNoEntry();
 	}
 
-	void RHIBuildAccelerationStructure(FRHIRayTracingScene* Scene)
-	{
-		FRayTracingSceneBuildParams Params;
-		Params.Scene = Scene;
-		RHIBuildAccelerationStructure(Params);
-	}
-
 	virtual void RHIRayTraceOcclusion(FRHIRayTracingScene* Scene,
 		FRHIShaderResourceView* Rays,
 		FRHIUnorderedAccessView* Output,
diff --git a/Engine/Source/Runtime/Renderer/Private/DeferredShadingRenderer.cpp b/Engine/Source/Runtime/Renderer/Private/DeferredShadingRenderer.cpp
--- a/Engine/Source/Runtime/Renderer/Private/DeferredShadingRenderer.cpp
+++ b/Engine/Source/Runtime/Renderer/Private/DeferredShadingRenderer.cpp
@@ -1486,11 +1486,16 @@ bool FDeferredShadingSceneRenderer::DispatchRayTracingWorldUpdates(FRDGBuilder&
 		PassParams->RayTracingSceneScratchBuffer = Scene->RayTracingScene.BuildScratchBuffer;
 		PassParams->RayTracingSceneInstanceBuffer = Scene->RayTracingScene.InstanceBuffer;
 
-		GraphBuilder.AddPass(RDG_EVENT_NAME("RayTracingScene"), PassParams, ERDGPassFlags::Compute | ERDGPassFlags::NeverCull,
+		GraphBuilder.AddPass(RDG_EVENT_NAME("BuildAccelerationStructure"), PassParams, ERDGPassFlags::Compute | ERDGPassFlags::NeverCull,
 			[this, PassParams](FRHICommandListImmediate& RHICmdList)
 			{
+				check(RayTracingDynamicGeometryUpdateEndTransition == nullptr);
+				const FRHITransition* RayTracingDynamicGeometryUpdateBeginTransition = RHICreateTransition(FRHITransitionCreateInfo(ERHIPipeline::Graphics, ERHIPipeline::AsyncCompute));
+				RayTracingDynamicGeometryUpdateEndTransition = RHICreateTransition(FRHITransitionCreateInfo(ERHIPipeline::AsyncCompute, ERHIPipeline::Graphics));
+
 				FRHIAsyncComputeCommandListImmediate& RHIAsyncCmdList = FRHICommandListExecutor::GetImmediateAsyncComputeCommandList();
 
+
 				const FRHITransition* RayTracingDynamicGeometryUpdateBeginTransition = RHICreateTransition(FRHITransitionCreateInfo(ERHIPipeline::Graphics, ERHIPipeline::AsyncCompute));
 				RHICmdList.BeginTransition(RayTracingDynamicGeometryUpdateBeginTransition);
 				RHIAsyncCmdList.EndTransition(RayTracingDynamicGeometryUpdateBeginTransition);
@@ -1520,7 +1525,6 @@ bool FDeferredShadingSceneRenderer::DispatchRayTracingWorldUpdates(FRDGBuilder&
 				check(RayTracingDynamicGeometryUpdateEndTransition == nullptr)
 				RayTracingDynamicGeometryUpdateEndTransition = RHICreateTransition(FRHITransitionCreateInfo(ERHIPipeline::AsyncCompute, ERHIPipeline::Graphics));
 				RHIAsyncCmdList.BeginTransition(RayTracingDynamicGeometryUpdateEndTransition);
-
 				FRHIAsyncComputeCommandListImmediate::ImmediateDispatch(RHIAsyncCmdList);
 			});
 	}
