Fixes motion blur when using sequencer driven time dilation

This is the meat of the change, whereas https://p4-swarm.epicgames.net/reviews/17582565 contains all the engine refactor to be submitted in the final same CL.

diff --git a/Engine/Source/Runtime/Engine/Classes/Engine/World.h b/Engine/Source/Runtime/Engine/Classes/Engine/World.h
--- a/Engine/Source/Runtime/Engine/Classes/Engine/World.h
+++ b/Engine/Source/Runtime/Engine/Classes/Engine/World.h
@@ -139,6 +139,87 @@ private:
 	friend UWorld;
 };
 
+/** Contains all the timings of a gaming frame, to handle pause and time dilation (for instance bullet time) of the world. */
+struct ENGINE_API FGameTime
+{
+	using FTimeType = float;
+
+	FORCEINLINE_DEBUGGABLE FGameTime()
+		: RealTimeSeconds(FTimeType(0))
+		, WorldTimeSeconds(FTimeType(0))
+		, DeltaRealTimeSeconds(FTimeType(0))
+		, DeltaWorldTimeSeconds(FTimeType(0))
+	{ }
+
+	FGameTime(const FGameTime&) = default;
+	FGameTime& operator = (const FGameTime&) = default;
+
+	// Returns the game time since GStartTime.
+	static FGameTime GetTimeSinceAppStart();
+
+	static FORCEINLINE_DEBUGGABLE FGameTime CreateUndilated(FTimeType InRealTimeSeconds, FTimeType InDeltaRealTimeSeconds)
+	{
+		return FGameTime::CreateDilated(InRealTimeSeconds, InDeltaRealTimeSeconds, InRealTimeSeconds, InDeltaRealTimeSeconds);
+	}
+
+	static FORCEINLINE_DEBUGGABLE FGameTime CreateDilated(FTimeType InRealTimeSeconds, FTimeType InDeltaRealTimeSeconds, FTimeType InWorldTimeSeconds, FTimeType InDeltaWorldTimeSeconds)
+	{
+		return FGameTime(InRealTimeSeconds, InDeltaRealTimeSeconds, InWorldTimeSeconds, InDeltaWorldTimeSeconds);
+	}
+
+	/** Returns time in seconds since level began play, but IS NOT paused when the game is paused, and IS NOT dilated/clamped. */
+	FORCEINLINE_DEBUGGABLE FTimeType GetRealTimeSeconds() const
+	{
+		return RealTimeSeconds;
+	}
+
+	/** Returns frame delta time in seconds with no adjustment for time dilation and pause. */
+	FORCEINLINE_DEBUGGABLE FTimeType GetDeltaRealTimeSeconds() const
+	{
+		return DeltaRealTimeSeconds;
+	}
+
+	/** Returns time in seconds since level began play, but IS paused when the game is paused, and IS dilated/clamped. */
+	FORCEINLINE_DEBUGGABLE FTimeType GetWorldTimeSeconds() const
+	{
+		return WorldTimeSeconds;
+	}
+
+	/** Returns frame delta time in seconds adjusted by e.g. time dilation. */
+	FORCEINLINE_DEBUGGABLE FTimeType GetDeltaWorldTimeSeconds() const
+	{
+		return DeltaWorldTimeSeconds;
+	}
+
+	/** Returns how much world time is slowed compared to real time. */
+	FORCEINLINE_DEBUGGABLE float GetTimeDilation() const
+	{
+		ensure(DeltaRealTimeSeconds > FTimeType(0));
+		return float(DeltaWorldTimeSeconds / DeltaRealTimeSeconds);
+	}
+
+	/** Returns whether the world time is paused. */
+	FORCEINLINE_DEBUGGABLE bool IsPaused() const
+	{
+		return DeltaWorldTimeSeconds == FTimeType(0);
+	}
+
+private:
+	FTimeType RealTimeSeconds;
+	FTimeType WorldTimeSeconds;
+
+	FTimeType DeltaRealTimeSeconds;
+	FTimeType DeltaWorldTimeSeconds;
+
+	FORCEINLINE_DEBUGGABLE FGameTime(FTimeType InRealTimeSeconds, FTimeType InDeltaRealTimeSeconds, FTimeType InWorldTimeSeconds, FTimeType InDeltaWorldTimeSeconds)
+		: RealTimeSeconds(InRealTimeSeconds)
+		, WorldTimeSeconds(InWorldTimeSeconds)
+		, DeltaRealTimeSeconds(InDeltaRealTimeSeconds)
+		, DeltaWorldTimeSeconds(InDeltaWorldTimeSeconds)
+	{ }
+
+};
+
 
 ENGINE_API DECLARE_LOG_CATEGORY_EXTERN(LogSpawn, Warning, All);
 
@@ -1638,6 +1719,9 @@ public:
 	/** Time in seconds since level began play, but IS paused when the game is paused, and IS NOT dilated/clamped. */
 	float AudioTimeSeconds;
 
+	/** Frame delta time in seconds with no adjustment for time dilation. */
+	float DeltaRealTimeSeconds;
+
 	/** Frame delta time in seconds adjusted by e.g. time dilation. */
 	float DeltaTimeSeconds;
 
@@ -2474,6 +2558,13 @@ public:
 	 */
 	float GetDeltaSeconds() const;
 	
+	/**
+	 * Returns the dilatable time
+	 *
+	 * @return Returns the dilatable time
+	 */
+	FGameTime GetTime() const;
+
 	/** Helper for getting the time since a certain time. */
 	float TimeSince( float Time ) const;
 
@@ -4095,6 +4186,13 @@ FORCEINLINE_DEBUGGABLE float UWorld::GetDeltaSeconds() const
 	return DeltaTimeSeconds;
 }
 
+FORCEINLINE_DEBUGGABLE FGameTime UWorld::GetTime() const
+{
+	return FGameTime::CreateDilated(
+		RealTimeSeconds, DeltaRealTimeSeconds,
+		TimeSeconds, DeltaTimeSeconds);
+}
+
 FORCEINLINE_DEBUGGABLE float UWorld::TimeSince(float Time) const
 {
 	return GetTimeSeconds() - Time;
diff --git a/Engine/Source/Runtime/Engine/Private/LevelTick.cpp b/Engine/Source/Runtime/Engine/Private/LevelTick.cpp
--- a/Engine/Source/Runtime/Engine/Private/LevelTick.cpp
+++ b/Engine/Source/Runtime/Engine/Private/LevelTick.cpp
@@ -1422,6 +1422,7 @@ void UWorld::Tick( ELevelTick TickType, float DeltaSeconds )
 
 	DeltaSeconds = GameDeltaSeconds;
 	DeltaTimeSeconds = DeltaSeconds;
+	DeltaRealTimeSeconds = RealDeltaSeconds;
 
 	UnpausedTimeSeconds += DeltaSeconds;
 
diff --git a/Engine/Source/Runtime/Engine/Private/SceneView.cpp b/Engine/Source/Runtime/Engine/Private/SceneView.cpp
--- a/Engine/Source/Runtime/Engine/Private/SceneView.cpp
+++ b/Engine/Source/Runtime/Engine/Private/SceneView.cpp
@@ -2511,11 +2511,11 @@ void FSceneView::SetupCommonViewUniformBufferParameters(
 	else
 #endif
 	{
-		ViewUniformShaderParameters.PrevFrameGameTime = Family->CurrentWorldTime - Family->DeltaWorldTime;
-		ViewUniformShaderParameters.PrevFrameRealTime = Family->CurrentRealTime - Family->DeltaWorldTime;
-		ViewUniformShaderParameters.GameTime = Family->CurrentWorldTime;
-		ViewUniformShaderParameters.RealTime = Family->CurrentRealTime;
-		ViewUniformShaderParameters.DeltaTime = Family->DeltaWorldTime;
+		ViewUniformShaderParameters.PrevFrameGameTime = Family->Time.GetWorldTimeSeconds() - Family->Time.GetDeltaWorldTimeSeconds();
+		ViewUniformShaderParameters.PrevFrameRealTime = Family->Time.GetRealTimeSeconds() - Family->Time.GetDeltaRealTimeSeconds();
+		ViewUniformShaderParameters.GameTime = Family->Time.GetWorldTimeSeconds();
+		ViewUniformShaderParameters.RealTime = Family->Time.GetRealTimeSeconds();
+		ViewUniformShaderParameters.DeltaTime = Family->Time.GetDeltaWorldTimeSeconds();
 	}
 
 #if !(UE_BUILD_SHIPPING || UE_BUILD_TEST)
@@ -2582,9 +2582,10 @@ FSceneViewFamily::FSceneViewFamily(const ConstructionValues& CVS)
 	RenderTarget(CVS.RenderTarget),
 	Scene(CVS.Scene),
 	EngineShowFlags(CVS.EngineShowFlags),
-	CurrentWorldTime(CVS.CurrentWorldTime),
-	DeltaWorldTime(CVS.DeltaWorldTime),
-	CurrentRealTime(CVS.CurrentRealTime),
+	Time(CVS.Time),
+	CurrentWorldTime(CVS.Time.GetWorldTimeSeconds()),
+	DeltaWorldTime(CVS.Time.GetDeltaWorldTimeSeconds()),
+	CurrentRealTime(CVS.Time.GetRealTimeSeconds()),
 	FrameNumber(UINT_MAX),
 	bAdditionalViewFamily(CVS.bAdditionalViewFamily),
 	bRealtimeUpdate(CVS.bRealtimeUpdate),
@@ -2605,15 +2606,14 @@ FSceneViewFamily::FSceneViewFamily(const ConstructionValues& CVS)
 	PrimarySpatialUpscalerInterface(nullptr),
 	SecondarySpatialUpscalerInterface(nullptr)
 {
-	// If we do not pass a valid scene pointer then SetWorldTimes must be called to initialized with valid times.
+	// If we do not pass a valid scene pointer then SetTime() must be called to initialized with valid times.
 	ensure(CVS.bTimesSet);
 
 #if !(UE_BUILD_SHIPPING || UE_BUILD_TEST)
 	int32 Value = CVarRenderTimeFrozen.GetValueOnAnyThread();
 	if(Value)
 	{
-		CurrentWorldTime = 0;
-		CurrentRealTime = 0;
+		Time = FGameTime::CreateDilated(0.0f, Time.GetDeltaRealTimeSeconds(), 0.0f, Time.GetDeltaWorldTimeSeconds());
 	}
 
 	DebugViewShaderMode = ChooseDebugViewShaderMode();
diff --git a/Engine/Source/Runtime/Engine/Private/UserInterface/Canvas.cpp b/Engine/Source/Runtime/Engine/Private/UserInterface/Canvas.cpp
--- a/Engine/Source/Runtime/Engine/Private/UserInterface/Canvas.cpp
+++ b/Engine/Source/Runtime/Engine/Private/UserInterface/Canvas.cpp
@@ -261,11 +261,11 @@ int32 FCanvasWordWrapper::FindEndOfLastWholeGraphemeCluster(const int32 InStartI
 	return BreakIndex;
 }
 
-FCanvas* FCanvas::Create(FRDGBuilder& GraphBuilder, FRDGTextureRef InRenderTarget, FHitProxyConsumer* InHitProxyConsumer, float InRealTime, float InWorldTime, float InWorldDeltaTime, ERHIFeatureLevel::Type InFeatureLevel, float InDPIScale)
+FCanvas* FCanvas::Create(FRDGBuilder& GraphBuilder, FRDGTextureRef InRenderTarget, FHitProxyConsumer* InHitProxyConsumer, const FGameTime& Time, ERHIFeatureLevel::Type InFeatureLevel, float InDPIScale)
 {
 	return GraphBuilder.AllocObject<FCanvas>(
 		GraphBuilder.AllocObject<FRDGCanvasRenderTarget>(InRenderTarget),
-		InHitProxyConsumer, InRealTime, InWorldTime, InWorldDeltaTime, InFeatureLevel, InDPIScale);
+		InHitProxyConsumer, Time, InFeatureLevel, InDPIScale);
 }
 
 FCanvas::FCanvas(FRenderTarget* InRenderTarget, FHitProxyConsumer* InHitProxyConsumer, UWorld* InWorld, ERHIFeatureLevel::Type InFeatureLevel, ECanvasDrawMode InDrawMode, float InDPIScale)
@@ -276,9 +276,6 @@ FCanvas::FCanvas(FRenderTarget* InRenderTarget, FHitProxyConsumer* InHitProxyCon
 ,	Scene(InWorld ? InWorld->Scene : NULL)
 ,	AllowedModes(0xFFFFFFFF)
 ,	bRenderTargetDirty(false)
-,	CurrentRealTime(0)
-,	CurrentWorldTime(0)
-,	CurrentDeltaWorldTime(0)
 ,	FeatureLevel(InFeatureLevel)
 ,	bUseInternalTexture(false)
 ,	DrawMode(InDrawMode)
@@ -288,13 +285,11 @@ FCanvas::FCanvas(FRenderTarget* InRenderTarget, FHitProxyConsumer* InHitProxyCon
 
 	if (InWorld)
 	{
-		CurrentRealTime = InWorld->GetRealTimeSeconds();
-		CurrentWorldTime = InWorld->GetTimeSeconds();
-		CurrentDeltaWorldTime = InWorld->GetDeltaSeconds();
+		Time = InWorld->GetTime();
 	}
 }
 
-FCanvas::FCanvas(FRenderTarget* InRenderTarget,FHitProxyConsumer* InHitProxyConsumer, float InRealTime, float InWorldTime, float InWorldDeltaTime, ERHIFeatureLevel::Type InFeatureLevel, float InDPIScale)
+FCanvas::FCanvas(FRenderTarget* InRenderTarget,FHitProxyConsumer* InHitProxyConsumer, const FGameTime& InTime, ERHIFeatureLevel::Type InFeatureLevel, float InDPIScale)
 :	ViewRect(0,0,0,0)
 ,	ScissorRect(0,0,0,0)
 ,	RenderTarget(InRenderTarget)
@@ -302,9 +297,7 @@ FCanvas::FCanvas(FRenderTarget* InRenderTarget,FHitProxyConsumer* InHitProxyCons
 ,	Scene(NULL)
 ,	AllowedModes(0xFFFFFFFF)
 ,	bRenderTargetDirty(false)
-,	CurrentRealTime(InRealTime)
-,	CurrentWorldTime(InWorldTime)
-,	CurrentDeltaWorldTime(InWorldDeltaTime)
+,	Time(InTime)
 ,	bAllowsToSwitchVerticalAxis(false)
 ,	FeatureLevel(InFeatureLevel)
 ,	bUseInternalTexture(false)
diff --git a/Engine/Source/Runtime/Engine/Private/World.cpp b/Engine/Source/Runtime/Engine/Private/World.cpp
--- a/Engine/Source/Runtime/Engine/Private/World.cpp
+++ b/Engine/Source/Runtime/Engine/Private/World.cpp
@@ -3121,6 +3121,13 @@ void UWorld::RemoveFromWorld( ULevel* Level, bool bAllowIncrementalRemoval )
 	}
 }
 
+
+// static
+FGameTime FGameTime::GetTimeSinceAppStart()
+{
+	return FGameTime::CreateUndilated(FApp::GetCurrentTime() - GStartTime, FApp::GetDeltaTime());
+}
+
 /************************************************************************/
 /* FLevelStreamingGCHelper implementation                               */
 /************************************************************************/
diff --git a/Engine/Source/Runtime/Engine/Public/CanvasTypes.h b/Engine/Source/Runtime/Engine/Public/CanvasTypes.h
--- a/Engine/Source/Runtime/Engine/Public/CanvasTypes.h
+++ b/Engine/Source/Runtime/Engine/Public/CanvasTypes.h
@@ -228,7 +228,13 @@ public:
 		CDM_ImmediateDrawing
 	};
 
-	ENGINE_API static FCanvas* Create(FRDGBuilder& GraphBuilder, FRDGTextureRef InRenderTarget, FHitProxyConsumer* InHitProxyConsumer, float InRealTime, float InWorldTime, float InWorldDeltaTime, ERHIFeatureLevel::Type InFeatureLevel, float InDPIScale = 1.0f);
+	ENGINE_API static FCanvas* Create(FRDGBuilder& GraphBuilder, FRDGTextureRef InRenderTarget, FHitProxyConsumer* InHitProxyConsumer, const FGameTime& Time, ERHIFeatureLevel::Type InFeatureLevel, float InDPIScale = 1.0f);
+
+	UE_DEPRECATED(5.0, "Pass down a FGameTime instead.")
+	FORCEINLINE_DEBUGGABLE static FCanvas* Create(FRDGBuilder& GraphBuilder, FRDGTextureRef InRenderTarget, FHitProxyConsumer* InHitProxyConsumer, float InRealTime, float InWorldTime, float InWorldDeltaTime, ERHIFeatureLevel::Type InFeatureLevel, float InDPIScale = 1.0f)
+	{
+		return Create(GraphBuilder, InRenderTarget, InHitProxyConsumer, FGameTime::CreateDilated(InRealTime, InWorldDeltaTime, InWorldTime, InWorldDeltaTime), InFeatureLevel, InDPIScale);
+	}
 
 	/**
 	* Constructor.
@@ -238,7 +244,12 @@ public:
 	/**
 	* Constructor. For situations where a world is not available, but time information is
 	*/
-	ENGINE_API FCanvas(FRenderTarget* InRenderTarget, FHitProxyConsumer* InHitProxyConsumer, float InRealTime, float InWorldTime, float InWorldDeltaTime, ERHIFeatureLevel::Type InFeatureLevel, float InDPIScale = 1.0f);
+	ENGINE_API FCanvas(FRenderTarget* InRenderTarget, FHitProxyConsumer* InHitProxyConsumer, const FGameTime& Time, ERHIFeatureLevel::Type InFeatureLevel, float InDPIScale = 1.0f);
+	
+	UE_DEPRECATED(5.0, "Pass down a FGameTime instead.")
+	FORCEINLINE_DEBUGGABLE FCanvas(FRenderTarget* InRenderTarget, FHitProxyConsumer* InHitProxyConsumer, float InRealTime, float InWorldTime, float InWorldDeltaTime, ERHIFeatureLevel::Type InFeatureLevel, float InDPIScale = 1.0f)
+		: FCanvas(InRenderTarget, InHitProxyConsumer, FGameTime::CreateDilated(InRealTime, InWorldDeltaTime, InWorldTime, InWorldDeltaTime), InFeatureLevel, InDPIScale)
+	{ }
 
 	/**
 	* Destructor.
@@ -617,12 +628,8 @@ private:
 	uint32 AllowedModes;
 	/** true if the render target has been rendered to since last calling SetRenderTarget() */
 	bool bRenderTargetDirty;	
-	/** Current real time in seconds */
-	float CurrentRealTime;
-	/** Current world time in seconds */
-	float CurrentWorldTime;
-	/** Current world time in seconds */
-	float CurrentDeltaWorldTime;
+	/** Current gameplay time */
+	FGameTime Time;
 	/** true, if Canvas should be scaled to whole render target */
 	bool bScaledToRenderTarget;
 	// True if canvas allows switching vertical axis; false will ignore any flip
@@ -650,20 +657,22 @@ private:
 public:	
 
 	/**
-	 * Access current real time 
+	 * Access gameplay time
 	 */
-	float GetCurrentRealTime() const { return CurrentRealTime; }
+	const FGameTime& GetTime() const
+	{
+		return Time;
+	}
 
-	/**
-	 * Access current world time 
-	 */
-	float GetCurrentWorldTime() const { return CurrentWorldTime; }
+	UE_DEPRECATED(5.0, "Use FCanvas::GetTime()")
+	float GetCurrentRealTime() const { return GetTime().GetRealTimeSeconds(); }
 
-	/**
-	 * Access current delta time 
-	 */
-	float GetCurrentDeltaWorldTime() const { return CurrentDeltaWorldTime; }
+	UE_DEPRECATED(5.0, "Use FCanvas::GetTime()")
+	float GetCurrentWorldTime() const { return GetTime().GetWorldTimeSeconds(); }
 
+	UE_DEPRECATED(5.0, "Use FCanvas::GetTime()")
+	float GetCurrentDeltaWorldTime() const { return GetTime().GetDeltaWorldTimeSeconds(); }
+	
 	/** 
 	 * Draw a CanvasItem
 	 *
diff --git a/Engine/Source/Runtime/Engine/Public/SceneView.h b/Engine/Source/Runtime/Engine/Public/SceneView.h
--- a/Engine/Source/Runtime/Engine/Public/SceneView.h
+++ b/Engine/Source/Runtime/Engine/Public/SceneView.h
@@ -1542,9 +1542,6 @@ public:
 		,	Scene(InScene)
 		,	EngineShowFlags(InEngineShowFlags)
 		,	ViewModeParam(-1)
-		,	CurrentWorldTime(0.0f)
-		,	DeltaWorldTime(0.0f)
-		,	CurrentRealTime(0.0f)
 		,	GammaCorrection(1.0f)
 		,	bAdditionalViewFamily(false)
 		,	bRealtimeUpdate(false)
@@ -1557,11 +1554,8 @@ public:
 				UWorld* World = InScene->GetWorld();
 				// Ensure the world is valid and that we are being called from a game thread (GetRealTimeSeconds requires this)
 				if( World && IsInGameThread() )
-				{					
-					CurrentWorldTime = World->GetTimeSeconds();
-					DeltaWorldTime = World->GetDeltaSeconds();
-					CurrentRealTime = World->GetRealTimeSeconds();
-					bTimesSet = true;
+				{	
+					SetTime(World->GetTime());
 				}
 			}
 		}
@@ -1579,14 +1573,8 @@ public:
 		/** An name bound to the current viewmode param. (example : texture name) */
 		FName ViewModeParamName;
 
-		/** The current world time. */
-		float CurrentWorldTime;
-
-		/** The difference between the last world time and CurrentWorldTime. */
-		float DeltaWorldTime;
-		
-		/** The current real time. */
-		float CurrentRealTime;
+		/** The current time. */
+		FGameTime Time;
 
 		/** Gamma correction used when rendering this family. Default is 1.0 */
 		float GammaCorrection;
@@ -1606,8 +1594,19 @@ public:
 		/** Safety check to ensure valid times are set either from a valid world/scene pointer or via the SetWorldTimes function */
 		uint32 bTimesSet:1;
 
-		/** Set the world time ,difference between the last world time and CurrentWorldTime and current real time. */
-		ConstructionValues& SetWorldTimes(const float InCurrentWorldTime,const float InDeltaWorldTime,const float InCurrentRealTime) { CurrentWorldTime = InCurrentWorldTime; DeltaWorldTime = InDeltaWorldTime; CurrentRealTime = InCurrentRealTime;bTimesSet = true;return *this; }
+		/** Set the world time and real time independently to handle time dilation. */
+		ConstructionValues& SetTime(const FGameTime& InTime)
+		{
+			Time = InTime;
+			bTimesSet = true;
+			return *this;
+		}
+
+		UE_DEPRECATED(5.0, "Use FSceneViewFamily::ConstructionValues::SetTime()")
+		ConstructionValues& SetWorldTimes(float InCurrentWorldTime, float InDeltaWorldTime, float InCurrentRealTime)
+		{
+			return SetTime(FGameTime::CreateDilated(InCurrentRealTime, /* InCurrentRealTime = */ InDeltaWorldTime, InCurrentWorldTime, InDeltaWorldTime));
+		}
 
 		/** Set  whether the view family is additional. */
 		ConstructionValues& SetAdditionalViewFamily(const bool Value) { bAdditionalViewFamily = Value; return *this; }
@@ -1643,13 +1642,16 @@ public:
 	/** The new show flags for the views (meant to replace the old system). */
 	FEngineShowFlags EngineShowFlags;
 
-	/** The current world time. */
+	/** The current time. */
+	FGameTime Time;
+
+	UE_DEPRECATED(5.0, "Use FSceneViewFamily::Time")
 	float CurrentWorldTime;
 
-	/** The difference between the last world time and CurrentWorldTime. */
+	UE_DEPRECATED(5.0, "Use FSceneViewFamily::Time")
 	float DeltaWorldTime;
 
-	/** The current real time. */
+	UE_DEPRECATED(5.0, "Use FSceneViewFamily::Time")
 	float CurrentRealTime;
 
 	/** Copy from main thread GFrameNumber to be accessible on render thread side. UINT_MAX before CreateSceneRenderer() or BeginRenderingViewFamily() was called */
@@ -1868,11 +1870,16 @@ private:
 	/** whether the translucency are allowed to render after DOF, if not they will be rendered in standard translucency. */
 	bool bAllowTranslucencyAfterDOF;
 
+	PRAGMA_DISABLE_DEPRECATION_WARNINGS
 	// Only FSceneRenderer can copy a view family.
 	FSceneViewFamily(const FSceneViewFamily&) = default;
+	PRAGMA_ENABLE_DEPRECATION_WARNINGS
 
 	friend class FSceneRenderer;
+
+PRAGMA_DISABLE_DEPRECATION_WARNINGS
 };
+PRAGMA_ENABLE_DEPRECATION_WARNINGS
 
 /**
  * A view family which deletes its views when it goes out of scope.
diff --git a/Engine/Source/Runtime/Renderer/Private/ScenePrivate.h b/Engine/Source/Runtime/Renderer/Private/ScenePrivate.h
--- a/Engine/Source/Runtime/Renderer/Private/ScenePrivate.h
+++ b/Engine/Source/Runtime/Renderer/Private/ScenePrivate.h
@@ -1185,10 +1185,10 @@ public:
 	void UpdateLastRenderTime(const FSceneViewFamily& Family)
 	{
 		// The editor can trigger multiple update calls within a frame
-		if(Family.CurrentRealTime != LastRenderTime)
+		if(Family.Time.GetRealTimeSeconds() != LastRenderTime)
 		{
-			LastRenderTimeDelta = Family.CurrentRealTime - LastRenderTime;
-			LastRenderTime = Family.CurrentRealTime;
+			LastRenderTimeDelta = Family.Time.GetRealTimeSeconds() - LastRenderTime;
+			LastRenderTime = Family.Time.GetRealTimeSeconds();
 		}
 	}
 
diff --git a/Engine/Source/Runtime/Renderer/Private/SceneVisibility.cpp b/Engine/Source/Runtime/Renderer/Private/SceneVisibility.cpp
--- a/Engine/Source/Runtime/Renderer/Private/SceneVisibility.cpp
+++ b/Engine/Source/Runtime/Renderer/Private/SceneVisibility.cpp
@@ -44,6 +44,7 @@
 #include "InstanceCulling/InstanceCullingManager.h"
 #include "TemporalAA.h"
 #include "RayTracing/RayTracingInstanceCulling.h"
+#include "RendererModule.h"
 
 /*------------------------------------------------------------------------------
 	Globals
@@ -320,7 +321,7 @@ static void UpdatePrimitiveFadingStateHelper(FPrimitiveFadingState& FadingState,
 	{
 		if (FadingState.bIsVisible != bVisible)
 		{
-			float CurrentRealTime = View.Family->CurrentRealTime;
+			float CurrentRealTime = View.Family->Time.GetRealTimeSeconds();
 
 			// Need to kick off a fade, so make sure that we have fading state for that
 			if( !IsValidRef(FadingState.UniformBuffer) )
@@ -942,7 +943,7 @@ static void UpdatePrimitiveFading(const FScene* Scene, FViewInfo& View)
 	if (ViewState)
 	{
 		uint32 PrevFrameNumber = ViewState->PrevFrameNumber;
-		float CurrentRealTime = View.Family->CurrentRealTime;
+		float CurrentRealTime = View.Family->Time.GetRealTimeSeconds();
 
 		// First clear any stale fading states.
 		for (FPrimitiveFadingStateMap::TIterator It(ViewState->PrimitiveFadingStates); It; ++It)
@@ -1138,7 +1139,7 @@ static void FetchVisibilityForPrimitives_Range(FVisForPrimParams& Params, FGloba
 	FSceneViewState* ViewState = (FSceneViewState*)View.State;
 	const int32 NumBufferedFrames = FOcclusionQueryHelpers::GetNumBufferedFrames(Scene->GetFeatureLevel());
 	bool bClearQueries = !View.Family->EngineShowFlags.HitProxies;
-	const float CurrentRealTime = View.Family->CurrentRealTime;
+	const float CurrentRealTime = View.Family->Time.GetRealTimeSeconds();
 	uint32 OcclusionFrameCounter = ViewState->OcclusionFrameCounter;
 	FHZBOcclusionTester& HZBOcclusionTests = ViewState->HZBOcclusionTests;
 
@@ -1995,7 +1996,7 @@ static int32 OcclusionCull(FRHICommandListImmediate& RHICmdList, const FScene* S
 		}
 	}
 
-	float CurrentRealTime = View.Family->CurrentRealTime;
+	float CurrentRealTime = View.Family->Time.GetRealTimeSeconds();
 	if (ViewState)
 	{
 		bool bSubmitQueries = !View.bDisableQuerySubmissions;
@@ -2274,8 +2275,8 @@ struct FRelevancePacket
 		FPrimitiveViewMasks& InOutHasDynamicEditorMeshElementsMasks,
 		uint8* InMarkMasks)
 
-		: CurrentWorldTime(InView.Family->CurrentWorldTime)
-		, DeltaWorldTime(InView.Family->DeltaWorldTime)
+		: CurrentWorldTime(InView.Family->Time.GetWorldTimeSeconds())
+		, DeltaWorldTime(InView.Family->Time.GetDeltaWorldTimeSeconds())
 		, RHICmdList(InRHICmdList)
 		, Scene(InScene)
 		, View(InView)
@@ -3783,7 +3784,7 @@ void FSceneRenderer::PreVisibilityFrameSetup(FRDGBuilder& GraphBuilder, const FS
 			}
 
 			// determine if we are initializing or we should reset the persistent state
-			const float DeltaTime = View.Family->CurrentRealTime - ViewState->LastRenderTime;
+			const float DeltaTime = View.Family->Time.GetRealTimeSeconds() - ViewState->LastRenderTime;
 			const bool bFirstFrameOrTimeWasReset = DeltaTime < -0.0001f || ViewState->LastRenderTime < 0.0001f;
 			const bool bIsLargeCameraMovement = IsLargeCameraMovement(
 				View,
@@ -3862,7 +3863,7 @@ void FSceneRenderer::PreVisibilityFrameSetup(FRDGBuilder& GraphBuilder, const FS
 
 			// detect conditions where we should reset occlusion queries
 			if (bFirstFrameOrTimeWasReset || 
-				ViewState->LastRenderTime + GEngine->PrimitiveProbablyVisibleTime < View.Family->CurrentRealTime ||
+				GEngine->PrimitiveProbablyVisibleTime < View.Family->Time.GetDeltaRealTimeSeconds() ||
 				View.bCameraCut ||
 				View.bForceCameraVisibilityReset ||
 				IsLargeCameraMovement(
@@ -3944,7 +3945,7 @@ void FSceneViewState::UpdateMotionBlurTimeScale(const FViewInfo& View)
 	const int32 MotionBlurTargetFPS = View.FinalPostProcessSettings.MotionBlurTargetFPS;
 
 	// Ensure we can divide by the Delta Time later without a divide by zero.
-	float DeltaWorldTime = FMath::Max(View.Family->DeltaWorldTime, SMALL_NUMBER);
+	float DeltaRealTime = FMath::Max(View.Family->Time.GetDeltaRealTimeSeconds(), SMALL_NUMBER);
 
 	// Track the current FPS by using an exponential moving average of the current delta time.
 	if (MotionBlurTargetFPS <= 0)
@@ -3953,12 +3954,12 @@ void FSceneViewState::UpdateMotionBlurTimeScale(const FViewInfo& View)
 		if (GetSequencerState() == ESS_Paused)
 		{
 			// Reset the moving average to the current delta time.
-			MotionBlurTargetDeltaTime = DeltaWorldTime;
+			MotionBlurTargetDeltaTime = DeltaRealTime;
 		}
 		else
 		{
 			// Smooth the target delta time using a moving average.
-			MotionBlurTargetDeltaTime = FMath::Lerp(MotionBlurTargetDeltaTime, DeltaWorldTime, 0.1f);
+			MotionBlurTargetDeltaTime = FMath::Lerp(MotionBlurTargetDeltaTime, DeltaRealTime, 0.1f);
 		}
 	}
 	else // Track a fixed target FPS.
@@ -3967,13 +3968,14 @@ void FSceneViewState::UpdateMotionBlurTimeScale(const FViewInfo& View)
 		// Tuned for content compatibility with existing content when target is the default 30 FPS.
 		if (GetSequencerState() == ESS_Paused)
 		{
-			DeltaWorldTime = 1.0f / 60.0f;
+			DeltaRealTime = 1.0f / 60.0f;
 		}
 
+
 		MotionBlurTargetDeltaTime = 1.0f / static_cast<float>(MotionBlurTargetFPS);
 	}
 
-	MotionBlurTimeScale = MotionBlurTargetDeltaTime / DeltaWorldTime;
+	MotionBlurTimeScale = MotionBlurTargetDeltaTime / DeltaRealTime;
 }
 
 void UpdateReflectionSceneData(FScene* Scene)
@@ -4284,7 +4286,7 @@ void FSceneRenderer::ComputeViewVisibility(
 	}
 
 	int32 NumPrimitives = Scene->Primitives.Num();
-	float CurrentRealTime = ViewFamily.CurrentRealTime;
+	float CurrentRealTime = ViewFamily.Time.GetRealTimeSeconds();
 
 	FPrimitiveViewMasks HasDynamicMeshElementsMasks;
 	HasDynamicMeshElementsMasks.AddZeroed(NumPrimitives);
