Compile fix

diff --git a/Engine/Source/Runtime/Renderer/Private/VirtualShadowMaps/VirtualShadowMapArray.cpp b/Engine/Source/Runtime/Renderer/Private/VirtualShadowMaps/VirtualShadowMapArray.cpp
--- a/Engine/Source/Runtime/Renderer/Private/VirtualShadowMaps/VirtualShadowMapArray.cpp
+++ b/Engine/Source/Runtime/Renderer/Private/VirtualShadowMaps/VirtualShadowMapArray.cpp
@@ -1752,7 +1752,7 @@ public:
 		SHADER_PARAMETER_RDG_BUFFER_SRV(StructuredBuffer< uint >, ShadowHZBPageTable)
 		SHADER_PARAMETER_RDG_TEXTURE(Texture2D, HZBTexture)
 		SHADER_PARAMETER_SAMPLER(SamplerState, HZBSampler)
-		SHADER_PARAMETER(FVector2D, HZBSize)
+		SHADER_PARAMETER(FVector2f, HZBSize)
 		SHADER_PARAMETER(uint32, HZBMode)
 
 		SHADER_PARAMETER_RDG_BUFFER_UAV(RWStructuredBuffer<uint>, OutInvalidatingInstances)
