Add stub functions for virtual texture page id recording and playback.
Will flesh these out very shortly, but this allows us to submit things that will depend on the API.

diff --git a/Engine/Source/Runtime/RenderCore/Public/RendererInterface.h b/Engine/Source/Runtime/RenderCore/Public/RendererInterface.h
--- a/Engine/Source/Runtime/RenderCore/Public/RendererInterface.h
+++ b/Engine/Source/Runtime/RenderCore/Public/RendererInterface.h
@@ -860,12 +860,16 @@ public:
 	virtual void ReleaseVirtualTexturePendingResources() = 0;
 
 	/**	Provided a list of packed virtual texture tile ids, let the VT system request them. Note this should be called as long as the tiles are needed.*/
+	virtual void RequestVirtualTextureTiles(TArrayView<uint64> InPageRequests) = 0;
 	virtual void RequestVirtualTextureTiles(const FVector2D& InScreenSpaceSize, int32 InMipLevel) = 0;
 	virtual void RequestVirtualTextureTilesForRegion(IAllocatedVirtualTexture* AllocatedVT, const FVector2D& InScreenSpaceSize, const FVector2D& InViewportPosition, const FVector2D& InViewportSize, const FVector2D& InUV0, const FVector2D& InUV1, int32 InMipLevel) = 0;
 
 	/** Ensure that any tiles requested by 'RequestVirtualTextureTilesForRegion' are loaded, must be called from render thread */
 	virtual void LoadPendingVirtualTextureTiles(FRHICommandListImmediate& RHICmdList, ERHIFeatureLevel::Type FeatureLevel) = 0;
 
+	virtual void SetVirtualTextureRequestRecordBuffer(uint64 Handle) = 0;
+	virtual uint64 GetVirtualTextureRequestRecordBuffer(TSet<uint64>& OutPageRequests) = 0;
+
 	/** Evict all data from virtual texture caches*/
 	virtual void FlushVirtualTextureCache() = 0;
 
diff --git a/Engine/Source/Runtime/Renderer/Private/RendererModule.h b/Engine/Source/Runtime/Renderer/Private/RendererModule.h
--- a/Engine/Source/Runtime/Renderer/Private/RendererModule.h
+++ b/Engine/Source/Runtime/Renderer/Private/RendererModule.h
@@ -99,9 +99,12 @@ public:
 	virtual void AddVirtualTextureProducerDestroyedCallback(const FVirtualTextureProducerHandle& Handle, FVTProducerDestroyedFunction* Function, void* Baton) override;
 	virtual uint32 RemoveAllVirtualTextureProducerDestroyedCallbacks(const void* Baton) override;
 	virtual void ReleaseVirtualTexturePendingResources() override;
+	virtual void RequestVirtualTextureTiles(TArrayView<uint64> InPageRequests) override;
 	virtual void RequestVirtualTextureTiles(const FVector2D& InScreenSpaceSize, int32 InMipLevel) override;
 	virtual void RequestVirtualTextureTilesForRegion(IAllocatedVirtualTexture* AllocatedVT, const FVector2D& InScreenSpaceSize, const FVector2D& InViewportPosition, const FVector2D& InViewportSize, const FVector2D& InUV0, const FVector2D& InUV1, int32 InMipLevel) override;
 	virtual void LoadPendingVirtualTextureTiles(FRHICommandListImmediate& RHICmdList, ERHIFeatureLevel::Type FeatureLevel) override;
+	virtual void SetVirtualTextureRequestRecordBuffer(uint64 Handle) override;
+	virtual uint64 GetVirtualTextureRequestRecordBuffer(TSet<uint64>& OutPageRequests) override;
 	virtual void FlushVirtualTextureCache() override;
 	
 	virtual void RegisterPersistentViewUniformBufferExtension(IPersistentViewUniformBufferExtension* Extension) override;
diff --git a/Engine/Source/Runtime/Renderer/Private/SceneRendering.cpp b/Engine/Source/Runtime/Renderer/Private/SceneRendering.cpp
--- a/Engine/Source/Runtime/Renderer/Private/SceneRendering.cpp
+++ b/Engine/Source/Runtime/Renderer/Private/SceneRendering.cpp
@@ -4576,6 +4576,10 @@ uint32 FRendererModule::RemoveAllVirtualTextureProducerDestroyedCallbacks(const
 	return FVirtualTextureSystem::Get().RemoveAllProducerDestroyedCallbacks(Baton);
 }
 
+void FRendererModule::RequestVirtualTextureTiles(TArrayView<uint64> InPageRequests)
+{
+}
+
 void FRendererModule::RequestVirtualTextureTiles(const FVector2D& InScreenSpaceSize, int32 InMipLevel)
 {
 	FVirtualTextureSystem::Get().RequestTiles(InScreenSpaceSize, InMipLevel);
@@ -4594,6 +4598,15 @@ void FRendererModule::LoadPendingVirtualTextureTiles(FRHICommandListImmediate& R
 	GraphBuilder.Execute();
 }
 
+void FRendererModule::SetVirtualTextureRequestRecordBuffer(uint64 Handle)
+{
+}
+
+uint64 FRendererModule::GetVirtualTextureRequestRecordBuffer(TSet<uint64>& OutPageRequests)
+{
+	return (uint64)-1;
+}
+
 void FRendererModule::FlushVirtualTextureCache()
 {
 	FVirtualTextureSystem::Get().FlushCache();
