Fix incorrect checkSlow in ray tracing scene-visibility introduced in CL 17552671

diff --git a/Engine/Source/Runtime/Renderer/Private/SceneVisibility.cpp b/Engine/Source/Runtime/Renderer/Private/SceneVisibility.cpp
--- a/Engine/Source/Runtime/Renderer/Private/SceneVisibility.cpp
+++ b/Engine/Source/Runtime/Renderer/Private/SceneVisibility.cpp
@@ -867,7 +867,7 @@ static void PrimitiveCullTask(FThreadSafeCounter& NumCulledPrimitives, const FSc
 #if RHI_RAYTRACING
 		if (RayTracingBits)
 		{
-			checkSlow(!View.PrimitiveVisibilityMap.GetData()[WordIndex]); // this should start at zero
+			checkSlow(!View.PrimitiveRayTracingVisibilityMap.GetData()[WordIndex]); // this should start at zero
 			View.PrimitiveRayTracingVisibilityMap.GetData()[WordIndex] = RayTracingBits;
 		}
 #endif
