Ensure only visible groom elements run through view dependent pass (cluster culling, visibility, voxelization).

This fixes rendering/voxelization overhead when grooms are not visible (!ShouldRenderInMainPass or culled with hair clip).

diff --git a/Engine/Plugins/Runtime/HairStrands/Source/HairStrandsCore/Private/GroomManager.cpp b/Engine/Plugins/Runtime/HairStrands/Source/HairStrandsCore/Private/GroomManager.cpp
--- a/Engine/Plugins/Runtime/HairStrands/Source/HairStrandsCore/Private/GroomManager.cpp
+++ b/Engine/Plugins/Runtime/HairStrands/Source/HairStrandsCore/Private/GroomManager.cpp
@@ -861,6 +861,10 @@ static void RunHairLODSelection(
 			if (Instance->Strands.RenRaytracingResource)	{ Instance->Strands.RenRaytracingResource->Allocate(GraphBuilder, LoadingType, ResourceStatus); }
 			#endif
 			Instance->Strands.VertexFactory->InitResources();
+
+			// Early initialization, so that when filtering MeshBatch block in SceneVisiblity, we can use this value to know 
+			// if the hair instance is visible or not (i.e., HairLengthScale > 0)
+			Instance->HairGroupPublicData->VFInput.Strands.HairLengthScale = Instance->Strands.Modifier.HairLengthScale;
 		}
 
 		// Only switch LOD if the data are ready to be used
diff --git a/Engine/Source/Runtime/Renderer/Private/HairStrands/HairStrandsInterface.cpp b/Engine/Source/Runtime/Renderer/Private/HairStrands/HairStrandsInterface.cpp
--- a/Engine/Source/Runtime/Renderer/Private/HairStrands/HairStrandsInterface.cpp
+++ b/Engine/Source/Runtime/Renderer/Private/HairStrands/HairStrandsInterface.cpp
@@ -425,6 +425,7 @@ FHairStrandsBookmarkParameters CreateHairStrandsBookmarkParameters(FScene* Scene
 	Out.VisibleInstances.Reserve(View.HairStrandsMeshElements.Num());
 	for (const FMeshBatchAndRelevance& MeshBatch : View.HairStrandsMeshElements)
 	{
+		check(MeshBatch.PrimitiveSceneProxy && MeshBatch.PrimitiveSceneProxy->ShouldRenderInMainPass());
 		if (MeshBatch.Mesh && MeshBatch.Mesh->Elements.Num() > 0)
 		{
 			FHairGroupPublicData* HairGroupPublicData = reinterpret_cast<FHairGroupPublicData*>(MeshBatch.Mesh->Elements[0].VertexFactoryUserData);
@@ -482,3 +483,17 @@ bool IsHairStrandsCompatible(const FMeshBatch* Mesh)
 	return false;
 }
 
+bool IsHairStrandsVisible(const FMeshBatchAndRelevance& MeshBatch)
+{
+	if (MeshBatch.Mesh && MeshBatch.PrimitiveSceneProxy && MeshBatch.PrimitiveSceneProxy->ShouldRenderInMainPass())
+	{
+		const FHairGroupPublicData* Data = reinterpret_cast<const FHairGroupPublicData*>(MeshBatch.Mesh->Elements[0].VertexFactoryUserData);
+		switch (Data->VFInput.GeometryType)
+		{
+			case EHairGeometryType::Strands: return Data->VFInput.Strands.HairLengthScale > 0;
+			case EHairGeometryType::Cards  : return true;
+			case EHairGeometryType::Meshes : return true;
+		}
+	}
+	return false;
+}
diff --git a/Engine/Source/Runtime/Renderer/Private/SceneVisibility.cpp b/Engine/Source/Runtime/Renderer/Private/SceneVisibility.cpp
--- a/Engine/Source/Runtime/Renderer/Private/SceneVisibility.cpp
+++ b/Engine/Source/Runtime/Renderer/Private/SceneVisibility.cpp
@@ -3096,6 +3096,7 @@ static void ComputeAndMarkRelevanceForViewParallel(
 	}
 }
 
+bool IsHairStrandsVisible(const FMeshBatchAndRelevance& MeshBatch);
 void ComputeDynamicMeshRelevance(EShadingPath ShadingPath, bool bAddLightmapDensityCommands, const FPrimitiveViewRelevance& ViewRelevance, const FMeshBatchAndRelevance& MeshBatch, FViewInfo& View, FMeshPassMask& PassMask, FPrimitiveSceneInfo* PrimitiveSceneInfo, const FPrimitiveBounds& Bounds)
 {
 	const int32 NumElements = MeshBatch.Mesh->Elements.Num();
@@ -3279,7 +3280,7 @@ void ComputeDynamicMeshRelevance(EShadingPath ShadingPath, bool bAddLightmapDens
 	}
 	
 	const bool bIsHairStrandsCompatible = ViewRelevance.bHairStrands && IsHairStrandsEnabled(EHairStrandsShaderType::All, View.GetShaderPlatform()) && IsHairStrandsCompatible(MeshBatch.Mesh);
-	if (bIsHairStrandsCompatible)
+	if (bIsHairStrandsCompatible && IsHairStrandsVisible(MeshBatch))
 	{
 		View.HairStrandsMeshElements.AddUninitialized(1);
 		FMeshBatchAndRelevance& BatchAndProxy = View.HairStrandsMeshElements.Last();
