Reduce CachedRayTracingInstance memory usage.
- Use InstanceSceneData instead of caching local transform (r.CachedRayTracingInstances.CacheLocalTransform).
- Lazily cache world trasform (r.CachedRayTracingInstances.LazyUpdate).

diff --git a/Engine/Source/Runtime/RenderCore/Private/RenderCore.cpp b/Engine/Source/Runtime/RenderCore/Private/RenderCore.cpp
--- a/Engine/Source/Runtime/RenderCore/Private/RenderCore.cpp
+++ b/Engine/Source/Runtime/RenderCore/Private/RenderCore.cpp
@@ -87,6 +87,10 @@ DEFINE_STAT(STAT_ViewStateMemory);
 DEFINE_STAT(STAT_RenderingMemStackMemory);
 DEFINE_STAT(STAT_LightInteractionMemory);
 
+#if RHI_RAYTRACING
+DEFINE_STAT(STAT_CachedRayTracingInstancesMemory);
+#endif
+
 // The InitViews stats group contains information on how long visibility culling took and how effective it was
 
 DEFINE_STAT(STAT_GatherShadowPrimitivesTime);
diff --git a/Engine/Source/Runtime/RenderCore/Public/RenderCore.h b/Engine/Source/Runtime/RenderCore/Public/RenderCore.h
--- a/Engine/Source/Runtime/RenderCore/Public/RenderCore.h
+++ b/Engine/Source/Runtime/RenderCore/Public/RenderCore.h
@@ -148,6 +148,10 @@ DECLARE_MEMORY_STAT_EXTERN(TEXT("ViewState memory"),STAT_ViewStateMemory,STATGRO
 DECLARE_MEMORY_STAT_EXTERN(TEXT("Rendering mem stack memory"),STAT_RenderingMemStackMemory,STATGROUP_SceneMemory, RENDERCORE_API);
 DECLARE_MEMORY_STAT_EXTERN(TEXT("Light interaction memory"),STAT_LightInteractionMemory,STATGROUP_SceneMemory, RENDERCORE_API);
 
+#if RHI_RAYTRACING
+DECLARE_MEMORY_STAT_EXTERN(TEXT("Cached Ray Tracing Instances"), STAT_CachedRayTracingInstancesMemory, STATGROUP_SceneMemory, RENDERCORE_API);
+#endif
+
 DECLARE_DWORD_COUNTER_STAT_EXTERN(TEXT("Num reflective shadow map lights"),STAT_NumReflectiveShadowMapLights,STATGROUP_LightRendering, RENDERCORE_API);
 
 DECLARE_MEMORY_STAT_EXTERN(TEXT("Pool Size"), STAT_RenderTargetPoolSize, STATGROUP_RenderTargetPool, RENDERCORE_API);
diff --git a/Engine/Source/Runtime/Renderer/Private/DeferredShadingRenderer.cpp b/Engine/Source/Runtime/Renderer/Private/DeferredShadingRenderer.cpp
--- a/Engine/Source/Runtime/Renderer/Private/DeferredShadingRenderer.cpp
+++ b/Engine/Source/Runtime/Renderer/Private/DeferredShadingRenderer.cpp
@@ -1123,7 +1123,7 @@ bool FDeferredShadingSceneRenderer::GatherRayTracingWorldInstancesForView(FRDGBu
 		for (const FRelevantPrimitive& RelevantPrimitive : RelevantPrimitives)
 		{
 			const int32 PrimitiveIndex = RelevantPrimitive.PrimitiveIndex;
-			const FPrimitiveSceneInfo* SceneInfo = Scene->Primitives[PrimitiveIndex];
+			FPrimitiveSceneInfo* SceneInfo = Scene->Primitives[PrimitiveIndex];
 			ERayTracingPrimitiveFlags Flags = Scene->PrimitiveRayTracingFlags[PrimitiveIndex];
 
 			if (EnumHasAnyFlags(Flags, ERayTracingPrimitiveFlags::CacheInstances))
@@ -1148,6 +1148,7 @@ bool FDeferredShadingSceneRenderer::GatherRayTracingWorldInstancesForView(FRDGBu
 				}
 
 				checkf(SceneInfo->CachedRayTracingInstance.GeometryRHI, TEXT("Ray tracing instance must have a valid geometry."));
+				SceneInfo->UpdateCachedRayTracingInstanceWorldTransforms();
 				RayTracingScene.Instances.Add(SceneInfo->CachedRayTracingInstance);
 
 				if (View.RayTracingCullingParameters.CullInRayTracing > 0 && GetRayTracingCullingPerInstance() && SceneInfo->CachedRayTracingInstance.NumTransforms > 1)
diff --git a/Engine/Source/Runtime/Renderer/Private/PrimitiveSceneInfo.cpp b/Engine/Source/Runtime/Renderer/Private/PrimitiveSceneInfo.cpp
--- a/Engine/Source/Runtime/Renderer/Private/PrimitiveSceneInfo.cpp
+++ b/Engine/Source/Runtime/Renderer/Private/PrimitiveSceneInfo.cpp
@@ -28,6 +28,23 @@
 extern int32 GGPUSceneInstanceClearList;
 extern int32 GGPUSceneInstanceBVH;
 
+#if RHI_RAYTRACING
+static int32 GCachedRayTracingInstancesCacheLocalTransform = 0;
+static FAutoConsoleVariableRef CVarCachedRayTracingInstancesUseInstanceData(
+	TEXT("r.CachedRayTracingInstances.CacheLocalTransform"),
+	GCachedRayTracingInstancesCacheLocalTransform,
+	TEXT("Cache Local Transform instead of using InstanceData (increases memory usage)."),
+	ECVF_ReadOnly);
+
+static int32 GCachedRayTracingInstancesLazyUpdate = 1;
+static FAutoConsoleVariableRef CVarCachedRayTracingInstancesLazyUpdate(
+	TEXT("r.CachedRayTracingInstances.LazyUpdate"),
+	GCachedRayTracingInstancesLazyUpdate,
+	TEXT("Lazy update cached ray tracing instances world transforms. \n")
+	TEXT("Reduces memory usage by only caching world transforms of primitives when necessary."),
+	ECVF_ReadOnly);
+#endif
+
 /** An implementation of FStaticPrimitiveDrawInterface that stores the drawn elements for the rendering thread to use. */
 class FBatchingSPDI : public FStaticPrimitiveDrawInterface
 {
@@ -215,6 +232,11 @@ FPrimitiveSceneInfo::~FPrimitiveSceneInfo()
 	{
 		check(StaticMeshCommandInfos.Num() == 0);
 	}
+
+#if RHI_RAYTRACING
+	DEC_MEMORY_STAT_BY(STAT_CachedRayTracingInstancesMemory, CachedRayTracingInstanceLocalTransforms.Num() * sizeof(FMatrix));
+	DEC_MEMORY_STAT_BY(STAT_CachedRayTracingInstancesMemory, CachedRayTracingInstanceWorldTransforms.Num() * sizeof(FMatrix));
+#endif
 }
 
 #if RHI_RAYTRACING
@@ -804,17 +826,61 @@ void FPrimitiveSceneInfo::CacheRayTracingPrimitives(FScene* Scene, const TArrayV
 	}
 }
 
+void FPrimitiveSceneInfo::UpdateCachedRayTracingInstanceWorldTransforms()
+{
+	if (bUpdateCachedRayTracingInstanceWorldTransforms && GCachedRayTracingInstancesLazyUpdate)
+	{
+		QUICK_SCOPE_CYCLE_COUNTER(STAT_UpdateCachedRayTracingInstanceWorldTransforms);
+
+		bUpdateCachedRayTracingInstanceWorldTransforms = false;
+
+		DEC_MEMORY_STAT_BY(STAT_CachedRayTracingInstancesMemory, CachedRayTracingInstanceWorldTransforms.Num() * sizeof(FMatrix));
+
+		CachedRayTracingInstanceWorldTransforms.SetNumUninitialized(CachedRayTracingInstance.NumTransforms);
+
+		INC_MEMORY_STAT_BY(STAT_CachedRayTracingInstancesMemory, CachedRayTracingInstanceWorldTransforms.Num() * sizeof(FMatrix));
+
+		// Apply local offset to far-field object
+		FMatrix LocalToWorld = Proxy->GetLocalToWorld();
+		if (Proxy->IsRayTracingFarField())
+		{
+			LocalToWorld = LocalToWorld.ConcatTranslation(Lumen::GetFarFieldReferencePos());
+		}
+
+		TConstArrayView<FPrimitiveInstance> InstanceSceneData = Proxy->GetInstanceSceneData();
+
+		for (int32 Index = 0; Index < CachedRayTracingInstanceWorldTransforms.Num(); Index++)
+		{
+			FMatrix LocalTransform = GCachedRayTracingInstancesCacheLocalTransform ? CachedRayTracingInstanceLocalTransforms[Index] : InstanceSceneData[Index].LocalToPrimitive.ToMatrix();
+
+			CachedRayTracingInstanceWorldTransforms[Index] = LocalTransform * LocalToWorld;
+		}
+
+		CachedRayTracingInstance.Transforms = MakeArrayView(CachedRayTracingInstanceWorldTransforms);
+	}
+}
+
 void FPrimitiveSceneInfo::UpdateCachedRayTracingInstance(FPrimitiveSceneInfo* SceneInfo, FRayTracingInstance& CachedRayTracingInstance, ERayTracingPrimitiveFlags& Flags)
 {
 	if (EnumHasAnyFlags(Flags, ERayTracingPrimitiveFlags::CacheInstances))
 	{
-		// Cache a copy of local transforms so that they can be updated in the future
-		// TODO: this is actually not needed for static meshes with non-movable mobility (except in editor)
-		SceneInfo->CachedRayTracingInstanceLocalTransforms = CachedRayTracingInstance.InstanceTransforms;
+		if (GCachedRayTracingInstancesCacheLocalTransform)
+		{
+			// Cache a copy of local transforms so that they can be updated in the future
+			// TODO: this is actually not needed for static meshes with non-movable mobility (except in editor)
+			DEC_MEMORY_STAT_BY(STAT_CachedRayTracingInstancesMemory, SceneInfo->CachedRayTracingInstanceLocalTransforms.Num() * sizeof(FMatrix));
+			SceneInfo->CachedRayTracingInstanceLocalTransforms = CachedRayTracingInstance.InstanceTransforms;
+			INC_MEMORY_STAT_BY(STAT_CachedRayTracingInstancesMemory, SceneInfo->CachedRayTracingInstanceLocalTransforms.Num() * sizeof(FMatrix));
+		}
 		// TODO: allocate from FRayTracingScene & do better low-level caching
 		SceneInfo->CachedRayTracingInstance.NumTransforms = CachedRayTracingInstance.NumTransforms;
-		SceneInfo->CachedRayTracingInstanceWorldTransforms.Empty();
-		SceneInfo->CachedRayTracingInstanceWorldTransforms.AddUninitialized(CachedRayTracingInstance.NumTransforms);
+		if (!GCachedRayTracingInstancesLazyUpdate)
+		{
+			DEC_MEMORY_STAT_BY(STAT_CachedRayTracingInstancesMemory, SceneInfo->CachedRayTracingInstanceWorldTransforms.Num() * sizeof(FMatrix));
+			SceneInfo->CachedRayTracingInstanceWorldTransforms.Empty();
+			SceneInfo->CachedRayTracingInstanceWorldTransforms.AddUninitialized(CachedRayTracingInstance.NumTransforms);
+			INC_MEMORY_STAT_BY(STAT_CachedRayTracingInstancesMemory, SceneInfo->CachedRayTracingInstanceWorldTransforms.Num() * sizeof(FMatrix));
+		}
 
 		// Apply local offset to far-field object
 		FMatrix LocalToWorld = SceneInfo->Proxy->GetLocalToWorld();
@@ -827,7 +893,11 @@ void FPrimitiveSceneInfo::UpdateCachedRayTracingInstance(FPrimitiveSceneInfo* Sc
 		SceneInfo->CachedRayTracingInstanceWorldBounds.AddUninitialized(CachedRayTracingInstance.NumTransforms);
 
 		SceneInfo->UpdateCachedRayTracingInstanceTransforms(LocalToWorld);
-		SceneInfo->CachedRayTracingInstance.Transforms = MakeArrayView(SceneInfo->CachedRayTracingInstanceWorldTransforms);
+
+		if (!GCachedRayTracingInstancesLazyUpdate)
+		{
+			SceneInfo->CachedRayTracingInstance.Transforms = MakeArrayView(SceneInfo->CachedRayTracingInstanceWorldTransforms);
+		}
 
 		check(SceneInfo->CachedRayTracingInstance.NumTransforms >= uint32(SceneInfo->CachedRayTracingInstance.Transforms.Num()));
 
@@ -879,17 +949,28 @@ void FPrimitiveSceneInfo::RemoveCachedRayTracingPrimitives()
 
 void FPrimitiveSceneInfo::UpdateCachedRayTracingInstanceTransforms(const FMatrix& NewPrimitiveLocalToWorld)
 {
+	QUICK_SCOPE_CYCLE_COUNTER(STAT_UpdateCachedRayTracingInstanceTransforms);
 	TRACE_CPUPROFILER_EVENT_SCOPE(UpdateCachedRayTracingInstanceTransforms);
+
+	TConstArrayView<FPrimitiveInstance> InstanceSceneData = Proxy->GetInstanceSceneData();
 	
 	SmallestRayTracingInstanceWorldBoundsIndex = 0;
 
-	for (int32 Index = 0; Index < CachedRayTracingInstanceLocalTransforms.Num(); Index++)
+	for (int32 Index = 0; Index < CachedRayTracingInstanceWorldBounds.Num(); Index++)
 	{
 		const FRenderBounds& LocalBoundingBox = Proxy->GetInstanceLocalBounds(Index);
-		CachedRayTracingInstanceWorldTransforms[Index] = CachedRayTracingInstanceLocalTransforms[Index] * NewPrimitiveLocalToWorld;
-		CachedRayTracingInstanceWorldBounds[Index] = LocalBoundingBox.TransformBy(CachedRayTracingInstanceLocalTransforms[Index] * NewPrimitiveLocalToWorld).ToBoxSphereBounds();
+
+		FMatrix LocalTransform = GCachedRayTracingInstancesCacheLocalTransform ? CachedRayTracingInstanceLocalTransforms[Index] : InstanceSceneData[Index].LocalToPrimitive.ToMatrix();
+
+		CachedRayTracingInstanceWorldBounds[Index] = LocalBoundingBox.TransformBy(LocalTransform * NewPrimitiveLocalToWorld).ToBoxSphereBounds();
+		if(!GCachedRayTracingInstancesLazyUpdate)
+		{
+			CachedRayTracingInstanceWorldTransforms[Index] = LocalTransform * NewPrimitiveLocalToWorld;
+		}
 		SmallestRayTracingInstanceWorldBoundsIndex = CachedRayTracingInstanceWorldBounds[Index].SphereRadius < CachedRayTracingInstanceWorldBounds[SmallestRayTracingInstanceWorldBoundsIndex].SphereRadius ? Index : SmallestRayTracingInstanceWorldBoundsIndex;
 	}
+
+	bUpdateCachedRayTracingInstanceWorldTransforms = true;
 }
 #endif
 
diff --git a/Engine/Source/Runtime/Renderer/Public/PrimitiveSceneInfo.h b/Engine/Source/Runtime/Renderer/Public/PrimitiveSceneInfo.h
--- a/Engine/Source/Runtime/Renderer/Public/PrimitiveSceneInfo.h
+++ b/Engine/Source/Runtime/Renderer/Public/PrimitiveSceneInfo.h
@@ -403,6 +403,7 @@ public:
 	TArray<FMatrix> CachedRayTracingInstanceWorldTransforms;
 	TArray<FBoxSphereBounds> CachedRayTracingInstanceWorldBounds;
 	int32 SmallestRayTracingInstanceWorldBoundsIndex;
+	bool bUpdateCachedRayTracingInstanceWorldTransforms;
 #endif
 
 	/** Initialization constructor. */
@@ -577,6 +578,8 @@ public:
 #if RHI_RAYTRACING
 	static void UpdateCachedRaytracingData(FScene* Scene, const TArrayView<FPrimitiveSceneInfo*>& SceneInfos);
 	RENDERER_API FRHIRayTracingGeometry* GetStaticRayTracingGeometryInstance(int LodLevel) const;
+
+	void UpdateCachedRayTracingInstanceWorldTransforms();
 #endif
 
 	/** Return primitive fullname (for debugging only). */
