Add hash for allocated virtual textures that can persist across runs.
Use this as id for recording and playback of virtual texture feedback.
Flesh out API for recording and playback of collected ids.

diff --git a/Engine/Source/Runtime/Engine/Private/Texture2D.cpp b/Engine/Source/Runtime/Engine/Private/Texture2D.cpp
--- a/Engine/Source/Runtime/Engine/Private/Texture2D.cpp
+++ b/Engine/Source/Runtime/Engine/Private/Texture2D.cpp
@@ -1297,6 +1297,7 @@ void FVirtualTexture2DResource::InitRHI()
 
 	FVTProducerDescription ProducerDesc;
 	ProducerDesc.Name = TextureOwner->GetFName();
+	ProducerDesc.FullNameHash = GetTypeHash(TextureOwner->GetFullName());
 	ProducerDesc.bContinuousUpdate = bContinuousUpdate;
 	ProducerDesc.Dimensions = 2;
 	ProducerDesc.TileSize = VTData->TileSize;
diff --git a/Engine/Source/Runtime/Engine/Private/VT/RuntimeVirtualTexture.cpp b/Engine/Source/Runtime/Engine/Private/VT/RuntimeVirtualTexture.cpp
--- a/Engine/Source/Runtime/Engine/Private/VT/RuntimeVirtualTexture.cpp
+++ b/Engine/Source/Runtime/Engine/Private/VT/RuntimeVirtualTexture.cpp
@@ -289,6 +289,7 @@ int32 URuntimeVirtualTexture::GetPageTableSize() const
 void URuntimeVirtualTexture::GetProducerDescription(FVTProducerDescription& OutDesc, FInitSettings const& InitSettings, FTransform const& VolumeToWorld) const
 {
 	OutDesc.Name = GetFName();
+	OutDesc.FullNameHash = GetTypeHash(GetFullName());
 	OutDesc.Dimensions = 2;
 	OutDesc.DepthInTiles = 1;
 	OutDesc.WidthInBlocks = 1;
diff --git a/Engine/Source/Runtime/RenderCore/Public/RendererInterface.h b/Engine/Source/Runtime/RenderCore/Public/RendererInterface.h
--- a/Engine/Source/Runtime/RenderCore/Public/RendererInterface.h
+++ b/Engine/Source/Runtime/RenderCore/Public/RendererInterface.h
@@ -859,16 +859,18 @@ public:
 	virtual uint32 RemoveAllVirtualTextureProducerDestroyedCallbacks(const void* Baton) = 0;
 	virtual void ReleaseVirtualTexturePendingResources() = 0;
 
-	/**	Provided a list of packed virtual texture tile ids, let the VT system request them. Note this should be called as long as the tiles are needed.*/
-	virtual void RequestVirtualTextureTiles(TArrayView<uint64> InPageRequests) = 0;
 	virtual void RequestVirtualTextureTiles(const FVector2D& InScreenSpaceSize, int32 InMipLevel) = 0;
 	virtual void RequestVirtualTextureTilesForRegion(IAllocatedVirtualTexture* AllocatedVT, const FVector2D& InScreenSpaceSize, const FVector2D& InViewportPosition, const FVector2D& InViewportSize, const FVector2D& InUV0, const FVector2D& InUV1, int32 InMipLevel) = 0;
 
 	/** Ensure that any tiles requested by 'RequestVirtualTextureTilesForRegion' are loaded, must be called from render thread */
 	virtual void LoadPendingVirtualTextureTiles(FRHICommandListImmediate& RHICmdList, ERHIFeatureLevel::Type FeatureLevel) = 0;
 
+	/** Allocate a buffer and record all virtual texture page requests until the next call to either SetVirtualTextureRequestRecordBuffer or GetVirtualTextureRequestRecordBuffer. */
 	virtual void SetVirtualTextureRequestRecordBuffer(uint64 Handle) = 0;
+	/** Fetch the virtual texture page requests recorded since the last call to SetVirtualTextureRequestRecordBuffer. Returns the handle that was passed in. */
 	virtual uint64 GetVirtualTextureRequestRecordBuffer(TSet<uint64>& OutPageRequests) = 0;
+	/**	Request an array of virtual texture page requests that was captured with SetVirtualTextureRequestRecordBuffer. Note that the array will be moved and ownership is taken. */
+	virtual void RequestVirtualTextureTiles(TArray<uint64>&& InPageRequests) = 0;
 
 	/** Evict all data from virtual texture caches*/
 	virtual void FlushVirtualTextureCache() = 0;
diff --git a/Engine/Source/Runtime/RenderCore/Public/VirtualTexturing.h b/Engine/Source/Runtime/RenderCore/Public/VirtualTexturing.h
--- a/Engine/Source/Runtime/RenderCore/Public/VirtualTexturing.h
+++ b/Engine/Source/Runtime/RenderCore/Public/VirtualTexturing.h
@@ -129,6 +129,7 @@ inline uint32 GetTypeHash(const FAllocatedVTDescription& Description)
 struct FVTProducerDescription
 {
 	FName Name; /** Will be name of UTexture for streaming VTs, mostly here for debugging */
+	uint32 FullNameHash;
 	
 	bool bPersistentHighestMip = true;
 	bool bContinuousUpdate = false;
@@ -343,6 +344,7 @@ public:
 		, VirtualPageY(~0u)
 	{}
 
+	virtual uint32 GetPersistentHash() const = 0;
 	virtual uint32 GetNumPageTableTextures() const = 0;
 	virtual FRHITexture* GetPageTableTexture(uint32 InPageTableIndex) const = 0;
 	virtual FRHITexture* GetPageTableIndirectionTexture() const = 0;
diff --git a/Engine/Source/Runtime/Renderer/Private/RendererModule.h b/Engine/Source/Runtime/Renderer/Private/RendererModule.h
--- a/Engine/Source/Runtime/Renderer/Private/RendererModule.h
+++ b/Engine/Source/Runtime/Renderer/Private/RendererModule.h
@@ -99,7 +99,7 @@ public:
 	virtual void AddVirtualTextureProducerDestroyedCallback(const FVirtualTextureProducerHandle& Handle, FVTProducerDestroyedFunction* Function, void* Baton) override;
 	virtual uint32 RemoveAllVirtualTextureProducerDestroyedCallbacks(const void* Baton) override;
 	virtual void ReleaseVirtualTexturePendingResources() override;
-	virtual void RequestVirtualTextureTiles(TArrayView<uint64> InPageRequests) override;
+	virtual void RequestVirtualTextureTiles(TArray<uint64>&& InPageRequests) override;
 	virtual void RequestVirtualTextureTiles(const FVector2D& InScreenSpaceSize, int32 InMipLevel) override;
 	virtual void RequestVirtualTextureTilesForRegion(IAllocatedVirtualTexture* AllocatedVT, const FVector2D& InScreenSpaceSize, const FVector2D& InViewportPosition, const FVector2D& InViewportSize, const FVector2D& InUV0, const FVector2D& InUV1, int32 InMipLevel) override;
 	virtual void LoadPendingVirtualTextureTiles(FRHICommandListImmediate& RHICmdList, ERHIFeatureLevel::Type FeatureLevel) override;
diff --git a/Engine/Source/Runtime/Renderer/Private/SceneRendering.cpp b/Engine/Source/Runtime/Renderer/Private/SceneRendering.cpp
--- a/Engine/Source/Runtime/Renderer/Private/SceneRendering.cpp
+++ b/Engine/Source/Runtime/Renderer/Private/SceneRendering.cpp
@@ -4576,10 +4576,6 @@ uint32 FRendererModule::RemoveAllVirtualTextureProducerDestroyedCallbacks(const
 	return FVirtualTextureSystem::Get().RemoveAllProducerDestroyedCallbacks(Baton);
 }
 
-void FRendererModule::RequestVirtualTextureTiles(TArrayView<uint64> InPageRequests)
-{
-}
-
 void FRendererModule::RequestVirtualTextureTiles(const FVector2D& InScreenSpaceSize, int32 InMipLevel)
 {
 	FVirtualTextureSystem::Get().RequestTiles(InScreenSpaceSize, InMipLevel);
@@ -4600,11 +4596,23 @@ void FRendererModule::LoadPendingVirtualTextureTiles(FRHICommandListImmediate& R
 
 void FRendererModule::SetVirtualTextureRequestRecordBuffer(uint64 Handle)
 {
+#if WITH_EDITOR
+	FVirtualTextureSystem::Get().SetVirtualTextureRequestRecordBuffer(Handle);
+#endif
 }
 
 uint64 FRendererModule::GetVirtualTextureRequestRecordBuffer(TSet<uint64>& OutPageRequests)
 {
+#if WITH_EDITOR
+	return FVirtualTextureSystem::Get().GetVirtualTextureRequestRecordBuffer(OutPageRequests);
+#else
 	return (uint64)-1;
+#endif
+}
+
+void FRendererModule::RequestVirtualTextureTiles(TArray<uint64>&& InPageRequests)
+{
+	FVirtualTextureSystem::Get().RequestRecordedTiles(MoveTemp(InPageRequests));
 }
 
 void FRendererModule::FlushVirtualTextureCache()
diff --git a/Engine/Source/Runtime/Renderer/Private/VT/AllocatedVirtualTexture.cpp b/Engine/Source/Runtime/Renderer/Private/VT/AllocatedVirtualTexture.cpp
--- a/Engine/Source/Runtime/Renderer/Private/VT/AllocatedVirtualTexture.cpp
+++ b/Engine/Source/Runtime/Renderer/Private/VT/AllocatedVirtualTexture.cpp
@@ -60,6 +60,9 @@ FAllocatedVirtualTexture::FAllocatedVirtualTexture(FVirtualTextureSystem* InSyst
 
 	MaxLevel = FMath::Min(MaxLevel, VIRTUALTEXTURE_LOG2_MAX_PAGETABLE_SIZE - 1u);
 
+	// Get the persistent hash that we will use for identification of AllocatedVT objects across runs.
+	PersistentHash = CalculatePersistentHash(InDesc, InProducers);
+
 	// Lock tiles
 	LockOrUnlockTiles(InSystem, true);
 
@@ -281,6 +284,29 @@ bool FAllocatedVirtualTexture::TryMapLockedTiles(FVirtualTextureSystem* InSystem
 	return !bHasMissingTiles;
 }
 
+uint32 FAllocatedVirtualTexture::CalculatePersistentHash(FAllocatedVTDescription const& InDesc, FVirtualTextureProducer* const* InProducers) const
+{
+	uint32 Hash = 0;
+
+	for (uint32 LayerIndex = 0u; LayerIndex < InDesc.NumTextureLayers; ++LayerIndex)
+	{
+		FVirtualTextureProducer const* Producer = InProducers[LayerIndex];
+		if (Producer != nullptr)
+		{
+			Hash = HashCombine(Hash, LayerIndex);
+			Hash = HashCombine(Hash, Producer->GetDescription().FullNameHash);
+		}
+	}
+
+	Hash = HashCombine(Hash, GetTypeHash(InDesc.TileSize));
+	Hash = HashCombine(Hash, GetTypeHash(InDesc.TileBorderSize));
+	Hash = HashCombine(Hash, GetTypeHash(InDesc.Dimensions));
+	Hash = HashCombine(Hash, GetTypeHash(InDesc.NumTextureLayers));
+	Hash = HashCombine(Hash, GetTypeHash(InDesc.PackedFlags));
+
+	return Hash;
+}
+
 uint32 FAllocatedVirtualTexture::AddUniqueProducer(FVirtualTextureProducerHandle const& InHandle, const FVirtualTextureProducer* InProducer)
 {
 	for (int32 Index = 0u; Index < UniqueProducers.Num(); ++Index)
diff --git a/Engine/Source/Runtime/Renderer/Private/VT/AllocatedVirtualTexture.h b/Engine/Source/Runtime/Renderer/Private/VT/AllocatedVirtualTexture.h
--- a/Engine/Source/Runtime/Renderer/Private/VT/AllocatedVirtualTexture.h
+++ b/Engine/Source/Runtime/Renderer/Private/VT/AllocatedVirtualTexture.h
@@ -17,7 +17,8 @@ class FVirtualTextureSystem;
 class FAllocatedVirtualTexture final : public IAllocatedVirtualTexture
 {
 public:
-	FAllocatedVirtualTexture(FVirtualTextureSystem* InSystem,
+	FAllocatedVirtualTexture(
+		FVirtualTextureSystem* InSystem,
 		uint32 InFrame,
 		const FAllocatedVTDescription& InDesc,
 		FVirtualTextureProducer* const* InProducers,
@@ -36,6 +37,7 @@ public:
 	void LockOrUnlockTiles(FVirtualTextureSystem* InSystem, bool bLock) const;
 
 	// begin IAllocatedVirtualTexture
+	virtual uint32 GetPersistentHash() const override { return PersistentHash; }
 	virtual uint32 GetNumPageTableTextures() const override;
 	virtual FRHITexture* GetPageTableTexture(uint32 InPageTableIndex) const override;
 	virtual FRHITexture* GetPageTableIndirectionTexture() const override;
@@ -61,10 +63,12 @@ public:
 	inline uint32 GetProducerPhysicalGroupIndexForPageTableLayer(uint32 InLayerIndex) const { return UniquePageTableLayers[InLayerIndex].ProducerPhysicalGroupIndex; }
 
 private:
+	uint32 CalculatePersistentHash(FAllocatedVTDescription const& InDesc, FVirtualTextureProducer* const* InProducers) const;
 	uint32 AddUniqueProducer(FVirtualTextureProducerHandle const& InHandle, const FVirtualTextureProducer* InProducer);
 	uint32 AddUniquePhysicalSpace(FVirtualTexturePhysicalSpace* InPhysicalSpace, uint32 InUniqueProducerIndex, uint32 InProducerPhysicalSpaceIndex);
 
 	uint32 FrameAllocated;
+	uint32 PersistentHash;
 
 	FVirtualTextureSpace* Space;
 
diff --git a/Engine/Source/Runtime/Renderer/Private/VT/UniquePageList.h b/Engine/Source/Runtime/Renderer/Private/VT/UniquePageList.h
--- a/Engine/Source/Runtime/Renderer/Private/VT/UniquePageList.h
+++ b/Engine/Source/Runtime/Renderer/Private/VT/UniquePageList.h
@@ -24,8 +24,8 @@ public:
 private:
 	enum
 	{
-		HashSize		= 8*1024,
-		MaxUniquePages	= 4*1024,
+		HashSize		= 16*1024,
+		MaxUniquePages	= 8*1024,
 	};
 
 	bool bInitialized;
diff --git a/Engine/Source/Runtime/Renderer/Private/VT/VirtualTextureSystem.cpp b/Engine/Source/Runtime/Renderer/Private/VT/VirtualTextureSystem.cpp
--- a/Engine/Source/Runtime/Renderer/Private/VT/VirtualTextureSystem.cpp
+++ b/Engine/Source/Runtime/Renderer/Private/VT/VirtualTextureSystem.cpp
@@ -160,18 +160,6 @@ struct FFeedbackAnalysisParameters
 	uint32 FeedbackSize = 0u;
 };
 
-struct FGatherRequestsParameters
-{
-	FVirtualTextureSystem* System = nullptr;
-	const FUniquePageList* UniquePageList = nullptr;
-	FPageUpdateBuffer* PageUpdateBuffers = nullptr;
-	FUniqueRequestList* RequestList = nullptr;
-	uint32 PageUpdateFlushCount = 0u;
-	uint32 PageStartIndex = 0u;
-	uint32 NumPages = 0u;
-	uint32 FrameRequested;
-};
-
 class FFeedbackAnalysisTask
 {
 public:
@@ -196,6 +184,50 @@ public:
 	FORCEINLINE TStatId GetStatId() const { return TStatId(); }
 };
 
+struct FAddRequestedTilesParameters
+{
+	FVirtualTextureSystem* System = nullptr;
+	const uint64* RequestBuffer = nullptr;
+	FUniquePageList* UniquePageList = nullptr;
+	uint32 NumRequests = 0u;
+};
+
+class FAddRequestedTilesTask
+{
+public:
+	explicit FAddRequestedTilesTask(const FAddRequestedTilesParameters& InParams) : Parameters(InParams) {}
+
+	FAddRequestedTilesParameters Parameters;
+
+	static void DoTask(FAddRequestedTilesParameters& InParams)
+	{
+		InParams.UniquePageList->Initialize();
+		InParams.System->AddRequestedTilesTask(InParams);
+	}
+
+	void DoTask(ENamedThreads::Type CurrentThread, const FGraphEventRef& MyCompletionGraphEvent)
+	{
+		FTaskTagScope TaskTagScope(ETaskTag::EParallelRenderingThread);
+		DoTask(Parameters);
+	}
+
+	static ESubsequentsMode::Type GetSubsequentsMode() { return ESubsequentsMode::TrackSubsequents; }
+	ENamedThreads::Type GetDesiredThread() { return ENamedThreads::AnyNormalThreadNormalTask; }
+	FORCEINLINE TStatId GetStatId() const { return TStatId(); }
+};
+
+struct FGatherRequestsParameters
+{
+	FVirtualTextureSystem* System = nullptr;
+	const FUniquePageList* UniquePageList = nullptr;
+	FPageUpdateBuffer* PageUpdateBuffers = nullptr;
+	FUniqueRequestList* RequestList = nullptr;
+	uint32 PageUpdateFlushCount = 0u;
+	uint32 PageStartIndex = 0u;
+	uint32 NumPages = 0u;
+	uint32 FrameRequested;
+};
+
 class FGatherRequestsTask
 {
 public:
@@ -254,6 +286,7 @@ FVirtualTextureSystem::FVirtualTextureSystem()
 #if WITH_EDITOR
 	, SaveAllocatorImages(TEXT("r.VT.SaveAllocatorImages"), TEXT("Save images showing allocator usage."),
 		FConsoleCommandDelegate::CreateRaw(this, &FVirtualTextureSystem::SaveAllocatorImagesFromConsole))
+	, PageRequestRecordHandle(~0ull)
 #endif
 {
 #if !UE_BUILD_SHIPPING
@@ -272,6 +305,7 @@ FVirtualTextureSystem::~FVirtualTextureSystem()
 	DestroyPendingVirtualTextures(true);
 
 	check(AllocatedVTs.Num() == 0);
+	check(PersistentVTMap.Num() == 0);
 
 	for (uint32 SpaceID = 0u; SpaceID < MaxSpaces; ++SpaceID)
 	{
@@ -576,6 +610,12 @@ IAllocatedVirtualTexture* FVirtualTextureSystem::AllocateVirtualTexture(const FA
 	{
 		AllocatedVTsToMap.Add(AllocatedVT);
 	}
+
+	// Add to deterministic map that should apply across runs.
+	// Note that this may overwrite a duplicate old mapping whenever a new AllocatedVT is recreated after its producers have been recreated.
+	// In contrast AllocatedVTs maintains the duplicate entries temporarily until the older entry is deleted.
+	PersistentVTMap.Add(TPair<uint32, IAllocatedVirtualTexture*>(AllocatedVT->GetPersistentHash(), AllocatedVT));
+	
 	return AllocatedVT;
 }
 
@@ -633,6 +673,14 @@ void FVirtualTextureSystem::DestroyPendingVirtualTextures(bool bForceDestroyAll)
 		// shouldn't be more than 1 instance of this in the list
 		verify(AllocatedVTsToMap.Remove(AllocatedVT) <= 1);
 		verify(AllocatedVTs.Remove(AllocatedVT->GetDescription()) == 1);
+		
+		// persistent entry might have already been reallocated
+		IAllocatedVirtualTexture** Found = PersistentVTMap.Find(AllocatedVT->GetPersistentHash());
+		if (Found != nullptr && *Found == AllocatedVT)
+		{
+			PersistentVTMap.Remove(AllocatedVT->GetPersistentHash());
+		}
+		
 		AllocatedVT->Destroy(this);
 		delete AllocatedVT;
 	}
@@ -1344,6 +1392,32 @@ void FVirtualTextureSystem::Update(FRDGBuilder& GraphBuilder, ERHIFeatureLevel::
 		GVirtualTextureFeedback.Unmap(GraphBuilder.RHICmdList, FeedbackResult.MapHandle);
 	}
 
+#if WITH_EDITOR
+	// If we're are recording page requests, then copy off pages to the recording buffer.
+	if (PageRequestRecordHandle != ~0ull)
+	{
+		RecordPageRequests(MergedUniquePageList, PageRequestRecordBuffer);
+	}
+#endif
+
+	// Add any page requests from recording playback.
+	// todo: We can split this into concurrent tasks. 
+	if (PageRequestPlaybackBuffer.Num() > 0)
+	{
+		FMemMark FeedbackMark(MemStack);
+
+		FAddRequestedTilesParameters Parameters;
+		Parameters.RequestBuffer = PageRequestPlaybackBuffer.GetData();
+		Parameters.NumRequests = PageRequestPlaybackBuffer.Num();
+		Parameters.System = this;
+		Parameters.UniquePageList = new(MemStack) FUniquePageList;
+
+		FAddRequestedTilesTask::DoTask(Parameters);
+		MergedUniquePageList->MergePages(Parameters.UniquePageList);
+
+		PageRequestPlaybackBuffer.Reset(0);
+	}
+
 	FUniqueRequestList* MergedRequestList = new(MemStack) FUniqueRequestList(MemStack);
 	MergedRequestList->Initialize();
 
@@ -2486,6 +2560,10 @@ void FVirtualTextureSystem::UpdateResidencyNotifications()
 			FString const& FormatString = PhysicalSpace->GetFormatString();
 			const float MipBias = PhysicalSpace->GetResidencyMipMapBias();
 
+#if CSV_PROFILER
+			CSV_CUSTOM_STAT_GLOBAL(VirtualTextureMipBias, MipBias, ECsvCustomStatOp::Set);
+#endif
+
 			OnScreenMessages.Add(
 				FCoreDelegates::EOnScreenMessageSeverity::Warning,
 				FText::Format(LOCTEXT("VTOversubscribed", "VT Pool [{0}] is oversubscribed. Setting MipBias {1}"), FText::FromString(FormatString), FText::AsNumber(MipBias)));
@@ -2553,4 +2631,124 @@ void FVirtualTextureSystem::DrawResidencyHud(UCanvas* InCanvas, APlayerControlle
 
 #endif // !UE_BUILD_SHIPPING
 
+#if WITH_EDITOR
+
+void FVirtualTextureSystem::SetVirtualTextureRequestRecordBuffer(uint64 Handle)
+{
+	check(IsInRenderingThread());
+	check(PageRequestRecordHandle == ~0ull && PageRequestRecordBuffer.Num() == 0);
+	
+	PageRequestRecordHandle = Handle;
+	PageRequestRecordBuffer.Reset();
+}
+
+void FVirtualTextureSystem::RecordPageRequests(FUniquePageList const* UniquePageList, TSet<uint64>& OutPages)
+{
+	const uint32 PageCount = UniquePageList->GetNum();
+	OutPages.Reserve(PageCount);
+
+	for (uint32 PageIndex = 0; PageIndex < PageCount; ++PageIndex)
+	{
+		const uint32 PageId = UniquePageList->GetPage(PageIndex);
+
+		const uint32 SpaceId = (PageId >> 28);
+		const FVirtualTextureSpace* RESTRICT Space = GetSpace(SpaceId);
+		if (Space == nullptr)
+		{
+			continue;
+		}
+
+		const uint32 vLevelPlusOne = ((PageId >> 24) & 0x0f);
+		const uint32 vLevel = FMath::Max(vLevelPlusOne, 1u) - 1;
+		const uint32 vPageX = (PageId & 0xfff) << vLevel;
+		const uint32 vPageY = ((PageId >> 12) & 0xfff) << vLevel;
+
+		const uint32 vAddress = FMath::MortonCode2(vPageX) | (FMath::MortonCode2(vPageY) << 1);
+		const FAllocatedVirtualTexture* RESTRICT AllocatedVT = Space->GetAllocator().Find(vAddress);
+		if (AllocatedVT == nullptr)
+		{
+			continue;
+		}
+
+		const uint32 LocalAddress = vAddress - AllocatedVT->GetVirtualAddress();
+		const uint32 PersistentHash = AllocatedVT->GetPersistentHash();
+		const uint64 ExportPageId = ((uint64)PersistentHash << 32)| ((uint64)vLevelPlusOne << 28) | (uint64)LocalAddress;
+		
+		OutPages.Add(ExportPageId);
+	}
+}
+
+uint64 FVirtualTextureSystem::GetVirtualTextureRequestRecordBuffer(TSet<uint64>& OutPageRequests)
+{
+	check(IsInRenderingThread());
+
+	if (PageRequestRecordHandle == ~0ull)
+	{
+		return ~0ull;
+	}
+
+	uint64 Ret = PageRequestRecordHandle;
+	OutPageRequests = MoveTemp(PageRequestRecordBuffer);
+	PageRequestRecordBuffer.Reset();
+	PageRequestRecordHandle = ~0ull;
+	return Ret;
+}
+
+#endif // WITH_EDITOR
+
+void FVirtualTextureSystem::RequestRecordedTiles(TArray<uint64>&& InPageRequests)
+{
+	check(IsInRenderingThread());
+
+	if (PageRequestPlaybackBuffer.Num() == 0)
+	{
+		PageRequestPlaybackBuffer = MoveTemp(InPageRequests);
+	}
+	else
+	{
+		PageRequestPlaybackBuffer.Append(InPageRequests);
+	}
+}
+
+void FVirtualTextureSystem::AddRequestedTilesTask(const FAddRequestedTilesParameters& Parameters)
+{
+	FUniquePageList* RESTRICT RequestedPageList = Parameters.UniquePageList;
+	const uint64* RESTRICT Buffer = Parameters.RequestBuffer;
+	const uint32 BufferSize = Parameters.NumRequests;
+
+	uint32 CurrentPersistentHash = ~0u;
+	IAllocatedVirtualTexture* AllocatedVT = nullptr;
+
+	for (uint32 Index = 0; Index < BufferSize; Index++)
+	{
+		uint64 PageRequest = Buffer[Index];
+
+		// We expect requests to be at least partially sorted. 
+		// So a small optimization is to only resolve the PersistentHash to AllocatedVT on change.
+		const uint32 PersistentHash = (uint32)(PageRequest >> 32);
+		if (CurrentPersistentHash != PersistentHash)
+		{
+			IAllocatedVirtualTexture** AllocatedVTPtr = PersistentVTMap.Find(PersistentHash);
+			AllocatedVT = AllocatedVTPtr != nullptr ? *AllocatedVTPtr : nullptr;
+			CurrentPersistentHash = PersistentHash;
+		}
+
+		if (AllocatedVT != nullptr)
+		{
+			const uint32 LevelPlusOne = (uint32)(PageRequest >> 28) & 0xf;
+			const uint32 LocalAddress = (uint32)(PageRequest & 0xffffff);
+
+			const uint32 BaseAddress = AllocatedVT->GetVirtualAddress();
+			const uint32 Address = BaseAddress + LocalAddress;
+			const uint32 Level = FMath::Max(LevelPlusOne, 1u) - 1;
+			const uint32 PageX = FMath::ReverseMortonCode2(Address) >> Level;
+			const uint32 PageY = FMath::ReverseMortonCode2(Address >> 1) >> Level;
+			const uint32 SpaceId = AllocatedVT->GetSpaceID();
+			const uint32 PageId = PageX | (PageY << 12) | (LevelPlusOne << 24) | (SpaceId << 28);
+
+			RequestedPageList->Add(PageId, 0xffff);
+		}
+	}
+}
+
 #undef LOCTEXT_NAMESPACE
diff --git a/Engine/Source/Runtime/Renderer/Private/VT/VirtualTextureSystem.h b/Engine/Source/Runtime/Renderer/Private/VT/VirtualTextureSystem.h
--- a/Engine/Source/Runtime/Renderer/Private/VT/VirtualTextureSystem.h
+++ b/Engine/Source/Runtime/Renderer/Private/VT/VirtualTextureSystem.h
@@ -25,6 +25,7 @@ struct FVTSpaceDescription;
 struct FVTPhysicalSpaceDescription;
 union FPhysicalSpaceIDAndAddress;
 struct FFeedbackAnalysisParameters;
+struct FAddRequestedTilesParameters;
 struct FGatherRequestsParameters;
 struct FPageUpdateBuffer;
 
@@ -72,6 +73,12 @@ public:
 	void RequestTilesForRegion(IAllocatedVirtualTexture* AllocatedVT, const FVector2D& InScreenSpaceSize, const FVector2D& InViewportPosition, const FVector2D& InViewportSize, const FVector2D& InUV0, const FVector2D& InUV1, int32 InMipLevel = -1);
 	void LoadPendingTiles(FRDGBuilder& GraphBuilder, ERHIFeatureLevel::Type FeatureLevel);
 	
+#if WITH_EDITOR
+	void SetVirtualTextureRequestRecordBuffer(uint64 Handle);
+	uint64 GetVirtualTextureRequestRecordBuffer(TSet<uint64>& OutPageRequests);
+#endif
+	void RequestRecordedTiles(TArray<uint64>&& InPageRequests);
+
 	void FlushCache();
 	void FlushCache(FVirtualTextureProducerHandle const& ProducerHandle, FIntRect const& TextureRegion, uint32 MaxLevel);
 
@@ -79,6 +86,7 @@ public:
 
 private:
 	friend class FFeedbackAnalysisTask;
+	friend class FAddRequestedTilesTask;
 	friend class FGatherRequestsTask;
 
 	FVirtualTextureSystem();
@@ -100,13 +108,18 @@ private:
 
 	void AddPageUpdate(FPageUpdateBuffer* Buffers, uint32 FlushCount, uint32 PhysicalSpaceID, uint16 pAddress);
 
-	void GatherRequestsTask(const FGatherRequestsParameters& Parameters);
 	void FeedbackAnalysisTask(const FFeedbackAnalysisParameters& Parameters);
+	void AddRequestedTilesTask(const FAddRequestedTilesParameters& Parameters);
+	void GatherRequestsTask(const FGatherRequestsParameters& Parameters);
 
 	void GetContinuousUpdatesToProduce(FUniqueRequestList const* RequestList, int32 MaxTilesToProduce);
 
 	void UpdateResidencyTracking() const;
 
+#if WITH_EDITOR
+	void RecordPageRequests(FUniquePageList const* UniquePageList, TSet<uint64>& OutPages);
+#endif
+
 	uint32	Frame;
 
 	static const uint32 MaxNumTasks = 16;
@@ -120,6 +133,7 @@ private:
 
 	TMap<FAllocatedVTDescription, FAllocatedVirtualTexture*> AllocatedVTs;
 	TArray<IAllocatedVirtualTexture*> AllocatedVTsToMap;
+	TMap<uint32, IAllocatedVirtualTexture*> PersistentVTMap;
 
 	FAdaptiveVirtualTexture* AdaptiveVTs[MaxSpaces] = { nullptr };
 
@@ -152,6 +166,12 @@ private:
 	TArray<FVirtualTextureLocalTile> TransientCollectedPages;
 	TArray<IVirtualTextureFinalizer*> Finalizers;
 
+#if WITH_EDITOR
+	uint64 PageRequestRecordHandle;
+	TSet<uint64> PageRequestRecordBuffer;
+#endif
+	TArray<uint64> PageRequestPlaybackBuffer;
+
 #if !UE_BUILD_SHIPPING
 	void GetOnScreenMessages(FCoreDelegates::FSeverityMessageMap& OutMessages);
 	FCoreDelegates::FSeverityMessageMap OnScreenMessages;
