Avoid allocating stats buffer if stats accumulation is not enabled.

diff --git a/Engine/Source/Runtime/Renderer/Private/VirtualShadowMaps/VirtualShadowMapCacheManager.cpp b/Engine/Source/Runtime/Renderer/Private/VirtualShadowMaps/VirtualShadowMapCacheManager.cpp
--- a/Engine/Source/Runtime/Renderer/Private/VirtualShadowMaps/VirtualShadowMapCacheManager.cpp
+++ b/Engine/Source/Runtime/Renderer/Private/VirtualShadowMaps/VirtualShadowMapCacheManager.cpp
@@ -288,22 +288,23 @@ void FVirtualShadowMapArrayCacheManager::ExtractStats(FRDGBuilder& GraphBuilder,
 	FRDGBufferRef AccumulatedStatsBufferRDG = nullptr;
 
 	// Note: stats accumulation thing is here because it needs to persist over frames.
-	if (!AccumulatedStatsBuffer.IsValid())
-	{
-		FRDGBufferDesc Desc = FRDGBufferDesc::CreateBufferDesc(4, 1 + FVirtualShadowMapArray::NumStats * MaxStatFrames);
-		Desc.Usage = EBufferUsageFlags(Desc.Usage | BUF_SourceCopy);
-
-		AccumulatedStatsBufferRDG = GraphBuilder.CreateBuffer(Desc, TEXT("Shadow.Virtual.AccumulatedStatsBuffer"));	// TODO: Can't be a structured buffer as EnqueueCopy is only defined for vertex buffers
-		AddClearUAVPass(GraphBuilder, GraphBuilder.CreateUAV(AccumulatedStatsBufferRDG, PF_R32_UINT), 0);
-		AccumulatedStatsBuffer = GraphBuilder.ConvertToExternalBuffer(AccumulatedStatsBufferRDG);
-	}
-	else
+	if (AccumulatedStatsBuffer.IsValid())
 	{
 		AccumulatedStatsBufferRDG = GraphBuilder.RegisterExternalBuffer(AccumulatedStatsBuffer, TEXT("Shadow.Virtual.AccumulatedStatsBuffer"));
 	}
 
 	if (IsAccumulatingStats())
 	{
+		if (!AccumulatedStatsBuffer.IsValid())
+		{
+			FRDGBufferDesc Desc = FRDGBufferDesc::CreateBufferDesc(4, 1 + FVirtualShadowMapArray::NumStats * MaxStatFrames);
+			Desc.Usage = EBufferUsageFlags(Desc.Usage | BUF_SourceCopy);
+
+			AccumulatedStatsBufferRDG = GraphBuilder.CreateBuffer(Desc, TEXT("Shadow.Virtual.AccumulatedStatsBuffer"));	// TODO: Can't be a structured buffer as EnqueueCopy is only defined for vertex buffers
+			AddClearUAVPass(GraphBuilder, GraphBuilder.CreateUAV(AccumulatedStatsBufferRDG, PF_R32_UINT), 0);
+			AccumulatedStatsBuffer = GraphBuilder.ConvertToExternalBuffer(AccumulatedStatsBufferRDG);
+		}
+
 		// Initialize/clear
 		if (!bAccumulatingStats)
 		{
@@ -331,9 +332,13 @@ void FVirtualShadowMapArrayCacheManager::ExtractStats(FRDGBuilder& GraphBuilder,
 	{
 		bAccumulatingStats = false;
 
-		GPUBufferReadback = new FRHIGPUBufferReadback(TEXT("Shadow.Virtual.AccumulatedStatsBuffer"));
+		GPUBufferReadback = new FRHIGPUBufferReadback(TEXT("Shadow.Virtual.AccumulatedStatsBufferReadback"));
 		AddEnqueueCopyPass(GraphBuilder, GPUBufferReadback, AccumulatedStatsBufferRDG, 0u);
 	}
+	else if (AccumulatedStatsBuffer.IsValid())
+	{
+		AccumulatedStatsBuffer.SafeRelease();
+	}
 
 	if (GPUBufferReadback && GPUBufferReadback->IsReady())
 	{
