Adds a GPU stat for motion blur

diff --git a/Engine/Source/Runtime/Renderer/Private/PostProcess/PostProcessMotionBlur.cpp b/Engine/Source/Runtime/Renderer/Private/PostProcess/PostProcessMotionBlur.cpp
--- a/Engine/Source/Runtime/Renderer/Private/PostProcess/PostProcessMotionBlur.cpp
+++ b/Engine/Source/Runtime/Renderer/Private/PostProcess/PostProcessMotionBlur.cpp
@@ -74,6 +74,9 @@ namespace
 
 		return MotionBlurQuality;
 	}
+
+DECLARE_GPU_STAT(MotionBlur)
+
 }
 
 const int32 kMotionBlurTileSize = 16;
@@ -909,6 +912,7 @@ FScreenPassTexture AddMotionBlurPass(FRDGBuilder& GraphBuilder, const FViewInfo&
 	const FMotionBlurViewports Viewports(FScreenPassTextureViewport(Inputs.SceneColor), FScreenPassTextureViewport(Inputs.SceneDepth));
 
 	RDG_EVENT_SCOPE(GraphBuilder, "MotionBlur");
+	RDG_GPU_STAT_SCOPE(GraphBuilder, MotionBlur);
 
 	FRDGTextureRef VelocityFlatTexture = nullptr;
 	FVecocityTileTextures VelocityTileTextures;
