Add building VSM HZB for only pages that are allocated, using indirect-dispatch, saves around 0.4ms in typical views.

diff --git a/Engine/Shaders/Private/VirtualShadowMaps/PageManagement.usf b/Engine/Shaders/Private/VirtualShadowMaps/PageManagement.usf
--- a/Engine/Shaders/Private/VirtualShadowMaps/PageManagement.usf
+++ b/Engine/Shaders/Private/VirtualShadowMaps/PageManagement.usf
@@ -14,6 +14,7 @@
 #include "../Nanite/NaniteDataDecode.ush"
 #include "../HairStrands/HairStrandsVisibilityCommon.ush"
 #include "../HairStrands/HairStrandsTileCommon.ush"
+#include "../ReductionCommon.ush"
 #include "PageAccessCommon.ush"
 #include "PageCacheCommon.ush"
 #include "PageOverlap.ush"
@@ -899,7 +900,8 @@ void SelectPagesToInitializeCS(uint PhysicalPageIndex : SV_DispatchThreadID)
 	}
 }
 
-uint2 GetTileBasePos(uint2 TileThreadID, uint GroupIndex, StructuredBuffer<uint> PageIndexBuffer)
+
+uint2 GetTileOffset(uint GroupIndex, StructuredBuffer<uint> PageIndexBuffer)
 {
 	const uint PageInputIndex = GroupIndex >> LOG2_TILES_PER_PAGE_1D;
 	const uint PageIndex = PageIndexBuffer[PageInputIndex];
@@ -911,6 +913,14 @@ uint2 GetTileBasePos(uint2 TileThreadID, uint GroupIndex, StructuredBuffer<uint>
 	uint2 PhysPageAddress = VSMPhysicalIndexToPageAddress(PageIndex);
 	// Pixel address of tile region for this thread group.
 	const uint2 TileOffset = (PhysPageAddress << uint2(VSM_LOG2_PAGE_SIZE, VSM_LOG2_PAGE_SIZE)) + (LocalTile << uint2(LOG2_TILE_SIZE_XY, LOG2_TILE_SIZE_XY));
+
+	return TileOffset;
+}
+
+uint2 GetTileBasePos(uint2 TileThreadID, uint GroupIndex, StructuredBuffer<uint> PageIndexBuffer)
+{
+	// Pixel address of tile region for this thread group.
+	const uint2 TileOffset = GetTileOffset(GroupIndex, PageIndexBuffer);
 	// Pixel address of 2x2 region to sample for this thread.
 	const uint2 BasePos = TileOffset + (TileThreadID.xy << 1u);
 
@@ -971,3 +981,227 @@ void MergeStaticPhysicalPagesIndirectCS(uint2 TileThreadID : SV_GroupThreadID, u
 	MergePhysicalPixel(BasePos + uint2(0U, 1U));
 	MergePhysicalPixel(BasePos + uint2(1U, 1U));
 }
+
+// Indirect HZB building:
+
+RWBuffer<uint> OutPagesForHZBIndirectArgsBuffer;
+RWStructuredBuffer<uint> OutPhysicalPagesForHZB;
+
+[numthreads(VSM_DEFAULT_CS_GROUP_X, 1, 1)]
+void SelectPagesForHZBCS(uint PhysicalPageIndex : SV_DispatchThreadID)
+{
+	if (PhysicalPageIndex >= VirtualShadowMap.MaxPhysicalPages)
+	{
+		return;
+	}
+
+	FPhysicalPageMetaData MetaData = PhysicalPageMetaData[PhysicalPageIndex];
+
+	if ((MetaData.Flags & VSM_ALLOCATED_FLAG) != 0)
+	{
+		int GroupCount = 0;
+		// Each page needs TILES_PER_PAGE_1D groups launched
+		WaveInterlockedAddScalar_(OutPagesForHZBIndirectArgsBuffer[0], TILES_PER_PAGE_1D, GroupCount);
+		OutPhysicalPagesForHZB[GroupCount >> LOG2_TILES_PER_PAGE_1D] = PhysicalPageIndex;
+
+		// Each top-reduction needs only one group launched
+		WaveInterlockedAddScalar_(OutPagesForHZBIndirectArgsBuffer[0 + 4], 1U, GroupCount);
+	}
+}
+
+
+SamplerState PhysicalPagePoolSampler;
+Texture2D<uint> PhysicalPagePool;
+
+float4 Gather4VisZ(uint2 PixelCoord)
+{
+#if COMPILER_SUPPORTS_GATHER_UINT
+	// Offset to 2x2 footprint center and scale to UV space
+	float2 UV = float2(PixelCoord + uint2(1U, 1U)) * VirtualShadowMap.RecPhysicalPoolSize;
+	return asfloat(PhysicalPagePool.Gather(PhysicalPagePoolSampler, UV, 0));
+#else
+	uint4 PixelRect = uint4(PixelCoord.xy, PixelCoord.xy + uint2(1U, 1U));
+	uint4 UintDepths = uint4(
+		PhysicalPagePool[PixelRect.xw].r,	// (-, +)
+		PhysicalPagePool[PixelRect.zw].r,	// (+, +)
+		PhysicalPagePool[PixelRect.zy].r,	// (+, -)
+		PhysicalPagePool[PixelRect.xy].r	// (-, -)
+	);
+	return asfloat(UintDepths);
+#endif
+}
+
+StructuredBuffer<uint> PhysicalPagesForHzb;
+//                                         out                input               output   
+RWTexture2D<float> FurthestHZBOutput_0; // 64 // 1     Group: 32 (16 threads x2)  16
+RWTexture2D<float> FurthestHZBOutput_1; // 32 // 1            16                  8
+RWTexture2D<float> FurthestHZBOutput_2; // 16                 8                   4
+RWTexture2D<float> FurthestHZBOutput_3; // 8                  4                   2
+RWTexture2D<float> FurthestHZBOutput_4; // 4                  2                   1
+
+
+groupshared float SharedMinDeviceZ[TILE_THREAD_GROUP_SIZE_XY * TILE_THREAD_GROUP_SIZE_XY];
+groupshared float SharedMaxDeviceZ[TILE_THREAD_GROUP_SIZE_XY * TILE_THREAD_GROUP_SIZE_XY];
+
+#define DIM_FURTHEST 1
+#define DIM_CLOSEST 0
+
+void OutputMipLevel(uint MipLevel, uint2 OutputPixelPos, float FurthestDeviceZ, float ClosestDeviceZ)
+{
+#if DIM_FURTHEST
+#define COND_OUTPUT_LEVEL(_level_) \
+	if (MipLevel == _level_) \
+	{ \
+		FurthestHZBOutput_##_level_[OutputPixelPos] = FurthestDeviceZ; \
+		return; \
+	}
+#endif
+#if DIM_CLOSEST
+		ClosestHZBOutput_1[OutputPixelPos] = ClosestDeviceZ;
+#endif
+
+COND_OUTPUT_LEVEL(1)
+COND_OUTPUT_LEVEL(2)
+COND_OUTPUT_LEVEL(3)
+COND_OUTPUT_LEVEL(4)
+
+#undef COND_OUTPUT_LEVEL
+}
+
+[numthreads(TILE_THREAD_GROUP_SIZE_XY, TILE_THREAD_GROUP_SIZE_XY, 1)]
+void BuildHZBPerPageCS(uint GroupThreadIndex : SV_GroupIndex, uint GroupIndex : SV_GroupID)
+{
+	uint2 SrcTileOffset = GetTileOffset(GroupIndex, PhysicalPagesForHzb);
+
+	uint2 RemappedGroupThreadIndex = InitialTilePixelPositionForReduction2x2(LOG2_TILE_SIZE_XY - 1U, GroupThreadIndex);
+
+	uint2 SrcPos = SrcTileOffset + (RemappedGroupThreadIndex << uint2(1U, 1U));
+	// Sample 2x2 footprint - thread group covers 32x32 area
+	float4 DeviceZ = Gather4VisZ(SrcPos);
+	float MinDeviceZ = min(min3(DeviceZ.x, DeviceZ.y, DeviceZ.z), DeviceZ.w);
+
+	float MaxDeviceZ = 0.0f;//max(max3(DeviceZ.x, DeviceZ.y, DeviceZ.z), DeviceZ.w);
+	//uint LinearGroupThreadID = RemappedGroupThreadIndex.y << LOG2_TILE_THREAD_GROUP_SIZE_XY + RemappedGroupThreadIndex.x;
+	
+	// Broadcast to all threads (16x16).
+	SharedMinDeviceZ[GroupThreadIndex] = MinDeviceZ;
+	// Write base HZB level (half physical page size, e.g., 64x64) 
+	uint2 OutPixelPos = SrcPos >> 1U;
+	FurthestHZBOutput_0[OutPixelPos] = MinDeviceZ;
+
+	// Build next 4 levels: 32, 16, 8, 4
+	UNROLL
+	for (uint MipLevel = 1U; MipLevel < LOG2_TILE_SIZE_XY; ++MipLevel)
+	{
+		// 8x8, 4x4, 2x2, 1x1
+		const uint OutTileDim = uint(TILE_THREAD_GROUP_SIZE_XY) >> MipLevel;
+		const uint ReduceBankSize = OutTileDim * OutTileDim;
+			
+		// LDS has been written before.
+		GroupMemoryBarrierWithGroupSync();
+
+		BRANCH
+		if (GroupThreadIndex < ReduceBankSize)
+		{
+			float4 ParentMinDeviceZ;
+			//float4 ParentMaxDeviceZ;
+			ParentMinDeviceZ[0] = MinDeviceZ;
+			//ParentMaxDeviceZ[0] = MaxDeviceZ;
+
+			UNROLL
+			for (uint i = 1; i < 4; i++)
+			{
+				uint LDSIndex = GroupThreadIndex + i * ReduceBankSize;
+				ParentMinDeviceZ[i] = SharedMinDeviceZ[LDSIndex];
+				//ParentMaxDeviceZ[i] = SharedMaxDeviceZ[LDSIndex];
+			}
+				
+			MinDeviceZ = min(min3(ParentMinDeviceZ.x, ParentMinDeviceZ.y, ParentMinDeviceZ.z), ParentMinDeviceZ.w);
+			//MaxDeviceZ = max(max3(ParentMaxDeviceZ.x, ParentMaxDeviceZ.y, ParentMaxDeviceZ.z), ParentMaxDeviceZ.w);
+	
+			OutPixelPos = OutPixelPos >> 1;
+			OutputMipLevel(MipLevel, OutPixelPos, MinDeviceZ, MaxDeviceZ);
+				
+			SharedMinDeviceZ[GroupThreadIndex] = MinDeviceZ;
+			//SharedMaxDeviceZ[GroupThreadIndex] = MaxDeviceZ;
+		}
+	}
+}
+
+float4 Gather4(Texture2D Texture, SamplerState TextureSampler, uint2 SrcPos, float2 InvSize)
+{
+	float2 SrcUV = float2(SrcPos) * InvSize;
+#if COMPILER_GLSL || FEATURE_LEVEL < FEATURE_LEVEL_SM5
+	float2 HalfTexelOffset = float2(0.5f, 0.5f) * InvSize;
+
+	float4 Out;
+	Out.x = Texture.SampleLevel(TextureSampler, SrcUV + float2(-HalfTexelOffset.x, -HalfTexelOffset.y), 0 ).r;
+	Out.y = Texture.SampleLevel(TextureSampler, SrcUV + float2( HalfTexelOffset.x, -HalfTexelOffset.y), 0 ).r;
+	Out.z = Texture.SampleLevel(TextureSampler, SrcUV + float2(-HalfTexelOffset.x,  HalfTexelOffset.y), 0 ).r;
+	Out.w = Texture.SampleLevel(TextureSampler, SrcUV + float2( HalfTexelOffset.x,  HalfTexelOffset.y), 0 ).r;
+
+	return Out;
+#else
+	return Texture.GatherRed(TextureSampler, SrcUV, 0);
+#endif
+}
+
+
+Texture2D ParentTextureMip;
+SamplerState ParentTextureMipSampler;
+
+float2 InvHzbInputSize;
+
+#define TOP_MIP_TILE_SIZE_XY 4
+// Each fetches 2x2 using gather
+#define TOP_MIP_TILE_THREAD_GROUP_SIZE_XY (TOP_MIP_TILE_SIZE_XY/2)
+
+[numthreads(TOP_MIP_TILE_THREAD_GROUP_SIZE_XY, TOP_MIP_TILE_THREAD_GROUP_SIZE_XY, 1)]
+void BuildHZBPerPageTopCS(uint2 GroupThreadId : SV_GroupThreadID, uint PageInputIndex : SV_GroupID)
+{
+	const uint PageIndex = PhysicalPagesForHzb[PageInputIndex];
+	uint2 PhysPageAddress = VSMPhysicalIndexToPageAddress(PageIndex);
+	
+	// Pixel address of tile region for this thread group.
+	const uint2 SrcTileOffset = PhysPageAddress * uint2(TOP_MIP_TILE_SIZE_XY, TOP_MIP_TILE_SIZE_XY);
+
+	uint2 SrcPos = SrcTileOffset +  (GroupThreadId << uint2(1U, 1U));
+
+	// Sample 2x2 footprint - thread group covers 32x32 area
+	float4 DeviceZ = Gather4(ParentTextureMip, ParentTextureMipSampler, SrcPos + uint2(1U, 1U), InvHzbInputSize);
+	float MinDeviceZ = min(min3(DeviceZ.x, DeviceZ.y, DeviceZ.z), DeviceZ.w);
+
+	float MaxDeviceZ = 0.0f;//max(max3(DeviceZ.x, DeviceZ.y, DeviceZ.z), DeviceZ.w);
+	//uint LinearGroupThreadID = RemappedGroupThreadIndex.y << LOG2_TILE_THREAD_GROUP_SIZE_XY + RemappedGroupThreadIndex.x;
+	
+	// Broadcast to all threads.
+	SharedMinDeviceZ[GroupThreadId.y * TOP_MIP_TILE_THREAD_GROUP_SIZE_XY + GroupThreadId.x] = MinDeviceZ;
+	// Write first HZB output level (half size) 
+	uint2 OutPixelPos = SrcPos >> 1U;
+	FurthestHZBOutput_0[OutPixelPos] = MinDeviceZ;
+
+	// Build last level
+	GroupMemoryBarrierWithGroupSync();
+
+	BRANCH
+	if (all(GroupThreadId.xy == uint2(0U, 0U)))
+	{
+		float4 ParentMinDeviceZ;
+		//float4 ParentMaxDeviceZ;
+		ParentMinDeviceZ[0] = MinDeviceZ;
+		//ParentMaxDeviceZ[0] = MaxDeviceZ;
+
+		UNROLL
+		for (uint Index = 1; Index < 4; ++Index)
+		{
+			ParentMinDeviceZ[Index] = SharedMinDeviceZ[Index];
+			//ParentMaxDeviceZ[i] = SharedMaxDeviceZ[LDSIndex];
+		}
+				
+		MinDeviceZ = min(min3(ParentMinDeviceZ.x, ParentMinDeviceZ.y, ParentMinDeviceZ.z), ParentMinDeviceZ.w);
+		//MaxDeviceZ = max(max3(ParentMaxDeviceZ.x, ParentMaxDeviceZ.y, ParentMaxDeviceZ.z), ParentMaxDeviceZ.w);
+	
+		OutPixelPos = OutPixelPos >> 1;
+		FurthestHZBOutput_1[OutPixelPos] = MinDeviceZ;
+	}
+}
diff --git a/Engine/Source/Runtime/Renderer/Private/ShadowDepthRendering.cpp b/Engine/Source/Runtime/Renderer/Private/ShadowDepthRendering.cpp
--- a/Engine/Source/Runtime/Renderer/Private/ShadowDepthRendering.cpp
+++ b/Engine/Source/Runtime/Renderer/Private/ShadowDepthRendering.cpp
@@ -1918,20 +1918,7 @@ void FSceneRenderer::RenderVirtualShadowMaps(FRDGBuilder& GraphBuilder, bool bNa
 
 			if (bVSMUseHZB)
 			{
-				FRDGTextureRef SceneDepth = GraphBuilder.RegisterExternalTexture( GSystemTextures.BlackDummy );
-
-				// NOTE: 32-bit HZB is important to not lose precision (and thus culling efficiency) with
-				// some of the shadow depth functions.
-				BuildHZBFurthest(
-					GraphBuilder,
-					SceneDepth,
-					VirtualShadowMapArray.PhysicalPagePoolRDG,
-					VirtualShadowViewRect,
-					FeatureLevel,
-					ShaderPlatform,
-					TEXT("Shadow.Virtual.HZB"),
-					/* OutFurthestHZBTexture = */ &VirtualShadowMapArray.HZBPhysical,
-					PF_R32_FLOAT);
+				VirtualShadowMapArray.HZBPhysical = VirtualShadowMapArray.BuildHZBFurthest(GraphBuilder);
 			}
 		}
 	}
diff --git a/Engine/Source/Runtime/Renderer/Private/VirtualShadowMaps/VirtualShadowMapArray.cpp b/Engine/Source/Runtime/Renderer/Private/VirtualShadowMaps/VirtualShadowMapArray.cpp
--- a/Engine/Source/Runtime/Renderer/Private/VirtualShadowMaps/VirtualShadowMapArray.cpp
+++ b/Engine/Source/Runtime/Renderer/Private/VirtualShadowMaps/VirtualShadowMapArray.cpp
@@ -13,6 +13,7 @@
 #include "VirtualShadowMapClipmap.h"
 #include "ComponentRecreateRenderStateContext.h"
 #include "HairStrands/HairStrandsData.h"
+#include "SceneTextureReductions.h"
 
 IMPLEMENT_STATIC_UNIFORM_BUFFER_SLOT(VirtualShadowMapUbSlot);
 
@@ -210,6 +211,12 @@ TAutoConsoleVariable<int32> CVarDebugSkipDynamicPageInvalidation(
 );
 #endif // !UE_BUILD_SHIPPING
 
+TAutoConsoleVariable<int32> CVarVSMBuildHZBPerPage(
+	TEXT("r.Shadow.Virtual.BuildHZBPerPage"),
+	1,
+	TEXT("."),
+	ECVF_RenderThreadSafe
+);
 
 FMatrix CalcTranslatedWorldToShadowUVMatrix(
 	const FMatrix& TranslatedWorldToShadowView,
@@ -650,14 +657,18 @@ class FClearIndirectDispatchArgs1DCS : public FVirtualPageManagementShader
 	SHADER_USE_PARAMETER_STRUCT(FClearIndirectDispatchArgs1DCS, FVirtualPageManagementShader)
 
 	BEGIN_SHADER_PARAMETER_STRUCT(FParameters, )
+		SHADER_PARAMETER(uint32, NumIndirectArgs)
+		SHADER_PARAMETER(uint32, IndirectArgStride)
 		SHADER_PARAMETER_RDG_BUFFER_UAV(RWBuffer<uint>, OutIndirectArgsBuffer)
 	END_SHADER_PARAMETER_STRUCT()
 };
 IMPLEMENT_GLOBAL_SHADER(FClearIndirectDispatchArgs1DCS, "/Engine/Private/VirtualShadowMaps/PageManagement.usf", "ClearIndirectDispatchArgs1DCS", SF_Compute);
 
-static void AddClearIndirectDispatchArgs1DPass(FRDGBuilder& GraphBuilder, FRDGBufferRef IndirectArgsRDG)
+static void AddClearIndirectDispatchArgs1DPass(FRDGBuilder& GraphBuilder, FRDGBufferRef IndirectArgsRDG, uint32 NumIndirectArgs = 1U, uint32 IndirectArgStride = 4U)
 {
 	FClearIndirectDispatchArgs1DCS::FParameters* PassParameters = GraphBuilder.AllocParameters<FClearIndirectDispatchArgs1DCS::FParameters>();
+	PassParameters->NumIndirectArgs = NumIndirectArgs;
+	PassParameters->IndirectArgStride = IndirectArgStride;
 	PassParameters->OutIndirectArgsBuffer = GraphBuilder.CreateUAV(IndirectArgsRDG);
 
 	auto ComputeShader = GetGlobalShaderMap(GMaxRHIFeatureLevel)->GetShader<FClearIndirectDispatchArgs1DCS>();
@@ -667,7 +678,7 @@ static void AddClearIndirectDispatchArgs1DPass(FRDGBuilder& GraphBuilder, FRDGBu
 		RDG_EVENT_NAME("ClearIndirectDispatchArgs"),
 		ComputeShader,
 		PassParameters,
-		FIntVector(1, 1, 1)
+		FComputeShaderUtils::GetGroupCount(NumIndirectArgs, 64)
 	);
 }
 
@@ -2216,3 +2227,178 @@ void FVirtualShadowMapArray::RenderVirtualShadowMapsHw(FRDGBuilder& GraphBuilder
 		}
 	}
 }
+
+class FSelectPagesForHZBCS : public FVirtualPageManagementShader
+{
+	DECLARE_GLOBAL_SHADER(FSelectPagesForHZBCS);
+	SHADER_USE_PARAMETER_STRUCT(FSelectPagesForHZBCS, FVirtualPageManagementShader)
+
+	BEGIN_SHADER_PARAMETER_STRUCT(FParameters, )
+		SHADER_PARAMETER_RDG_UNIFORM_BUFFER(FVirtualShadowMapUniformParameters, VirtualShadowMap)
+		SHADER_PARAMETER_RDG_BUFFER_SRV(StructuredBuffer<FPhysicalPageMetaData>, PhysicalPageMetaData)
+		SHADER_PARAMETER_RDG_BUFFER_UAV(RWBuffer<uint>, OutPagesForHZBIndirectArgsBuffer)
+		SHADER_PARAMETER_RDG_BUFFER_UAV(RWStructuredBuffer< uint >, OutPhysicalPagesForHZB)
+	END_SHADER_PARAMETER_STRUCT()
+};
+IMPLEMENT_GLOBAL_SHADER(FSelectPagesForHZBCS, "/Engine/Private/VirtualShadowMaps/PageManagement.usf", "SelectPagesForHZBCS", SF_Compute);
+
+
+class FVirtualSmBuildHZBPerPageCS : public FVirtualPageManagementShader
+{
+	DECLARE_GLOBAL_SHADER(FVirtualSmBuildHZBPerPageCS);
+	SHADER_USE_PARAMETER_STRUCT(FVirtualSmBuildHZBPerPageCS, FVirtualPageManagementShader)
+
+	static constexpr uint32 TotalHZBLevels = FVirtualShadowMap::Log2PageSize;
+	static constexpr uint32 HZBLevelsBase = TotalHZBLevels - 2U;
+
+	static_assert(HZBLevelsBase == 5U, "The shader is expecting 5 levels, if the page size is changed, this needs to be massaged");
+
+	BEGIN_SHADER_PARAMETER_STRUCT(FParameters, )
+		RDG_BUFFER_ACCESS(IndirectArgs, ERHIAccess::IndirectArgs)
+		SHADER_PARAMETER_RDG_UNIFORM_BUFFER(FVirtualShadowMapUniformParameters, VirtualShadowMap)
+		SHADER_PARAMETER_RDG_BUFFER_SRV(StructuredBuffer< uint >, PhysicalPagesForHZB)
+		SHADER_PARAMETER_SAMPLER(SamplerState, PhysicalPagePoolSampler)
+		SHADER_PARAMETER_RDG_TEXTURE(Texture2D<uint>, PhysicalPagePool)
+
+		SHADER_PARAMETER_RDG_TEXTURE_UAV_ARRAY(RWTexture2D<float>, FurthestHZBOutput, [HZBLevelsBase])
+	END_SHADER_PARAMETER_STRUCT()
+};
+IMPLEMENT_GLOBAL_SHADER(FVirtualSmBuildHZBPerPageCS, "/Engine/Private/VirtualShadowMaps/PageManagement.usf", "BuildHZBPerPageCS", SF_Compute);
+
+
+class FVirtualSmBBuildHZBPerPageTopCS : public FVirtualPageManagementShader
+{
+	DECLARE_GLOBAL_SHADER(FVirtualSmBBuildHZBPerPageTopCS);
+	SHADER_USE_PARAMETER_STRUCT(FVirtualSmBBuildHZBPerPageTopCS, FVirtualPageManagementShader)
+
+	// We need one level less as HZB starts at half-size (not really sure if we really need 1x1 and 2x2 sized levels).
+	static constexpr uint32 HZBLevelsTop = 2;
+
+	BEGIN_SHADER_PARAMETER_STRUCT(FParameters, )
+		RDG_BUFFER_ACCESS(IndirectArgs, ERHIAccess::IndirectArgs)
+		SHADER_PARAMETER_RDG_UNIFORM_BUFFER(FVirtualShadowMapUniformParameters, VirtualShadowMap)
+		SHADER_PARAMETER_RDG_BUFFER_SRV(StructuredBuffer< uint >, PhysicalPagesForHZB)
+		SHADER_PARAMETER_SAMPLER(SamplerState, ParentTextureMipSampler)
+		SHADER_PARAMETER_RDG_TEXTURE_SRV(Texture2D, ParentTextureMip)
+		SHADER_PARAMETER(FVector2D, InvHzbInputSize)
+
+		SHADER_PARAMETER_RDG_TEXTURE_UAV_ARRAY(RWTexture2D<float>, FurthestHZBOutput, [HZBLevelsTop])
+	END_SHADER_PARAMETER_STRUCT()
+};
+IMPLEMENT_GLOBAL_SHADER(FVirtualSmBBuildHZBPerPageTopCS, "/Engine/Private/VirtualShadowMaps/PageManagement.usf", "BuildHZBPerPageTopCS", SF_Compute);
+
+FRDGTextureRef FVirtualShadowMapArray::BuildHZBFurthest(FRDGBuilder& GraphBuilder)
+{
+	const FIntRect ViewRect(0, 0, GetPhysicalPoolSize().X, GetPhysicalPoolSize().Y);
+	const EPixelFormat Format = PF_R32_FLOAT;
+
+	FRDGTextureRef OutFurthestHZBTexture = nullptr;
+	if (CVarVSMBuildHZBPerPage.GetValueOnRenderThread())
+	{
+		// 1. Gather up all physical pages that are allocated
+		FRDGBufferRef PagesForHZBIndirectArgsRDG = GraphBuilder.CreateBuffer(FRDGBufferDesc::CreateIndirectDesc(2U * 4U), TEXT("Shadow.Virtual.PagesForHZBIndirectArgs"));
+		FRDGBufferRef PhysicalPagesForHZBRDG = GraphBuilder.CreateBuffer(FRDGBufferDesc::CreateStructuredDesc(sizeof(int32), UniformParameters.MaxPhysicalPages + 1), TEXT("Shadow.Virtual.PhysicalPagesForHZB"));
+
+		// 1. Clear the indirect args buffer (note 2x args)
+		AddClearIndirectDispatchArgs1DPass(GraphBuilder, PagesForHZBIndirectArgsRDG, 2U);
+
+		// 2. Filter the relevant physical pages and set up the indirect args
+		{
+			FSelectPagesForHZBCS::FParameters* PassParameters = GraphBuilder.AllocParameters<FSelectPagesForHZBCS::FParameters>();
+			PassParameters->VirtualShadowMap = GetUniformBuffer(GraphBuilder);
+			PassParameters->PhysicalPageMetaData = GraphBuilder.CreateSRV(PhysicalPageMetaDataRDG);
+			PassParameters->OutPagesForHZBIndirectArgsBuffer = GraphBuilder.CreateUAV(PagesForHZBIndirectArgsRDG);
+			PassParameters->OutPhysicalPagesForHZB = GraphBuilder.CreateUAV(PhysicalPagesForHZBRDG);
+
+			auto ComputeShader = GetGlobalShaderMap(GMaxRHIFeatureLevel)->GetShader<FSelectPagesForHZBCS>();
+
+			FComputeShaderUtils::AddPass(
+				GraphBuilder,
+				RDG_EVENT_NAME("SelectPagesForHZB"),
+				ComputeShader,
+				PassParameters,
+				FIntVector(FMath::DivideAndRoundUp(UniformParameters.MaxPhysicalPages, FSelectPagesForHZBCS::DefaultCSGroupX), 1, 1)
+			);
+
+		}
+
+
+		FIntPoint HZBSize(GetPhysicalPoolSize().X / 2, GetPhysicalPoolSize().Y / 2);
+
+		FRDGTextureDesc HZBDesc = FRDGTextureDesc::Create2D(HZBSize, Format, FClearValueBinding::None, TexCreate_ShaderResource | TexCreate_UAV, FVirtualSmBuildHZBPerPageCS::TotalHZBLevels);
+		HZBDesc.Flags |= GFastVRamConfig.HZB;
+
+		/** Closest and furthest HZB are intentionally in separate render target, because majority of the case you only one or the other.
+		 * Keeping them separate avoid doubling the size in cache for this cases, to avoid performance regression.
+		 */
+		OutFurthestHZBTexture = GraphBuilder.CreateTexture(HZBDesc, TEXT("Shadow.Virtual.PreviousOccluderHZB"));
+		
+		{
+			FVirtualSmBuildHZBPerPageCS::FParameters* PassParameters = GraphBuilder.AllocParameters<FVirtualSmBuildHZBPerPageCS::FParameters>();
+
+			PassParameters->VirtualShadowMap = GetUniformBuffer(GraphBuilder);
+			for (int32 DestMip = 0; DestMip < FVirtualSmBuildHZBPerPageCS::HZBLevelsBase; DestMip++)
+			{
+				PassParameters->FurthestHZBOutput[DestMip] = GraphBuilder.CreateUAV(FRDGTextureUAVDesc(OutFurthestHZBTexture, DestMip));
+			}
+			PassParameters->PhysicalPagePool = PhysicalPagePoolRDG;
+			PassParameters->PhysicalPagePoolSampler = TStaticSamplerState<SF_Point>::GetRHI();
+
+			PassParameters->IndirectArgs = PagesForHZBIndirectArgsRDG;
+			PassParameters->PhysicalPagesForHZB = GraphBuilder.CreateSRV(PhysicalPagesForHZBRDG);
+			auto ComputeShader = GetGlobalShaderMap(GMaxRHIFeatureLevel)->GetShader<FVirtualSmBuildHZBPerPageCS>();
+
+			FComputeShaderUtils::AddPass(
+				GraphBuilder,
+				RDG_EVENT_NAME("BuildHZBPerPage"),
+				ComputeShader,
+				PassParameters,
+				PassParameters->IndirectArgs,
+				0
+			);
+		}
+		{
+			FVirtualSmBBuildHZBPerPageTopCS::FParameters* PassParameters = GraphBuilder.AllocParameters<FVirtualSmBBuildHZBPerPageTopCS::FParameters>();
+
+			PassParameters->VirtualShadowMap = GetUniformBuffer(GraphBuilder);
+
+			uint32 StartDestMip = FVirtualSmBuildHZBPerPageCS::HZBLevelsBase;
+			for (int32 DestMip = 0; DestMip < FVirtualSmBBuildHZBPerPageTopCS::HZBLevelsTop; DestMip++)
+			{
+				PassParameters->FurthestHZBOutput[DestMip] = GraphBuilder.CreateUAV(FRDGTextureUAVDesc(OutFurthestHZBTexture, StartDestMip + DestMip));
+			}
+			FIntPoint SrcSize = FIntPoint::DivideAndRoundUp(HZBSize, 1 << int32(StartDestMip - 1));
+			PassParameters->InvHzbInputSize = FVector2D(1.0f / SrcSize.X, 1.0f / SrcSize.Y);;
+			PassParameters->ParentTextureMip = GraphBuilder.CreateSRV(FRDGTextureSRVDesc::CreateForMipLevel(OutFurthestHZBTexture, StartDestMip - 1));
+			PassParameters->ParentTextureMipSampler = TStaticSamplerState<SF_Point>::GetRHI();
+
+			PassParameters->IndirectArgs = PagesForHZBIndirectArgsRDG;
+			PassParameters->PhysicalPagesForHZB = GraphBuilder.CreateSRV(PhysicalPagesForHZBRDG);
+			auto ComputeShader = GetGlobalShaderMap(GMaxRHIFeatureLevel)->GetShader<FVirtualSmBBuildHZBPerPageTopCS>();
+
+			FComputeShaderUtils::AddPass(
+				GraphBuilder,
+				RDG_EVENT_NAME("BuildHZBPerPageTop"),
+				ComputeShader,
+				PassParameters,
+				PassParameters->IndirectArgs,
+				// NOTE: offset 4 to get second set of args in the buffer.
+				4U
+			);
+		}
+	}
+	else
+	{
+		::BuildHZBFurthest(
+			GraphBuilder,
+			FRDGSystemTextures::Get(GraphBuilder).Black,
+			PhysicalPagePoolRDG,
+			ViewRect,
+			GMaxRHIFeatureLevel,
+			GMaxRHIShaderPlatform,
+			TEXT("Shadow.Virtual.HZB"),
+			/* OutFurthestHZBTexture = */ &OutFurthestHZBTexture,
+			Format);
+	}
+	return OutFurthestHZBTexture;
+}
diff --git a/Engine/Source/Runtime/Renderer/Private/VirtualShadowMaps/VirtualShadowMapArray.h b/Engine/Source/Runtime/Renderer/Private/VirtualShadowMaps/VirtualShadowMapArray.h
--- a/Engine/Source/Runtime/Renderer/Private/VirtualShadowMaps/VirtualShadowMapArray.h
+++ b/Engine/Source/Runtime/Renderer/Private/VirtualShadowMaps/VirtualShadowMapArray.h
@@ -228,6 +228,8 @@ public:
 
 	bool ShouldCullBackfacingPixels() const { return bCullBackfacingPixels; }
 
+	FRDGTextureRef BuildHZBFurthest(FRDGBuilder& GraphBuilder);
+
 	// We keep a reference to the cache manager that was used to initialize this frame as it owns some of the buffers
 	FVirtualShadowMapArrayCacheManager* CacheManager = nullptr;
 
