Added a CVar to disable VT page playback.

diff --git a/Engine/Source/Runtime/Renderer/Private/VT/VirtualTextureSystem.cpp b/Engine/Source/Runtime/Renderer/Private/VT/VirtualTextureSystem.cpp
--- a/Engine/Source/Runtime/Renderer/Private/VT/VirtualTextureSystem.cpp
+++ b/Engine/Source/Runtime/Renderer/Private/VT/VirtualTextureSystem.cpp
@@ -74,13 +74,12 @@ static TAutoConsoleVariable<int32> CVarVTVerbose(
 	ECVF_RenderThreadSafe
 );
 
-static TAutoConsoleVariable<int32> CVarVTEnableFeedBack(
-	TEXT("r.VT.EnableFeedBack"),
+static TAutoConsoleVariable<int32> CVarVTEnableFeedback(
+	TEXT("r.VT.EnableFeedback"),
 	1,
-	TEXT("process readback buffer? dev option."),
+	TEXT("Enable processing of the GPU generated feedback buffer."),
 	ECVF_RenderThreadSafe
 );
-
 static TAutoConsoleVariable<int32> CVarVTParallelFeedbackTasks(
 	TEXT("r.VT.ParallelFeedbackTasks"),
 	0,
@@ -93,6 +92,12 @@ static TAutoConsoleVariable<int32> CVarVTNumFeedbackTasks(
 	TEXT("Number of tasks to create to read virtual texture feedback."),
 	ECVF_RenderThreadSafe
 );
+static TAutoConsoleVariable<int32> CVarVTEnablePlayback(
+	TEXT("r.VT.EnablePlayback"),
+	1,
+	TEXT("Enable playback of recorded feedback requests."),
+	ECVF_RenderThreadSafe
+);
 static TAutoConsoleVariable<int32> CVarVTNumGatherTasks(
 	TEXT("r.VT.NumGatherTasks"),
 	1,
@@ -1307,7 +1312,7 @@ void FVirtualTextureSystem::Update(FRDGBuilder& GraphBuilder, ERHIFeatureLevel::
 	FUniquePageList* MergedUniquePageList = new(MemStack) FUniquePageList;
 	MergedUniquePageList->Initialize();
 	
-	if (CVarVTEnableFeedBack.GetValueOnRenderThread())
+	if (CVarVTEnableFeedback.GetValueOnRenderThread())
 	{
 		FMemMark FeedbackMark(MemStack);
 
@@ -1401,19 +1406,22 @@ void FVirtualTextureSystem::Update(FRDGBuilder& GraphBuilder, ERHIFeatureLevel::
 #endif
 
 	// Add any page requests from recording playback.
-	// todo: We can split this into concurrent tasks. 
 	if (PageRequestPlaybackBuffer.Num() > 0)
 	{
-		FMemMark FeedbackMark(MemStack);
+		if (CVarVTEnablePlayback.GetValueOnRenderThread())
+		{
+			FMemMark FeedbackMark(MemStack);
 
-		FAddRequestedTilesParameters Parameters;
-		Parameters.RequestBuffer = PageRequestPlaybackBuffer.GetData();
-		Parameters.NumRequests = PageRequestPlaybackBuffer.Num();
-		Parameters.System = this;
-		Parameters.UniquePageList = new(MemStack) FUniquePageList;
+			// todo: We can split this into concurrent tasks. 
+			FAddRequestedTilesParameters Parameters;
+			Parameters.RequestBuffer = PageRequestPlaybackBuffer.GetData();
+			Parameters.NumRequests = PageRequestPlaybackBuffer.Num();
+			Parameters.System = this;
+			Parameters.UniquePageList = new(MemStack) FUniquePageList;
 
-		FAddRequestedTilesTask::DoTask(Parameters);
-		MergedUniquePageList->MergePages(Parameters.UniquePageList);
+			FAddRequestedTilesTask::DoTask(Parameters);
+			MergedUniquePageList->MergePages(Parameters.UniquePageList);
+		}
 
 		PageRequestPlaybackBuffer.Reset(0);
 	}
