Fix indirect HZB generation for static pages
Minor cleanup to logic to try and avoid similar issues

diff --git a/Engine/Shaders/Private/VirtualShadowMaps/PageManagement.usf b/Engine/Shaders/Private/VirtualShadowMaps/PageManagement.usf
--- a/Engine/Shaders/Private/VirtualShadowMaps/PageManagement.usf
+++ b/Engine/Shaders/Private/VirtualShadowMaps/PageManagement.usf
@@ -1006,6 +1006,15 @@ void SelectPagesForHZBCS(uint PhysicalPageIndex : SV_DispatchThreadID)
 
 		// Each top-reduction needs only one group launched
 		WaveInterlockedAddScalar_(OutPagesForHZBIndirectArgsBuffer[0 + 4], 1U, GroupCount);
+
+		if (VirtualShadowMapShouldCacheStaticSeparately())
+		{
+			// Also queue build for the static page
+			int GroupCountStatic = 0;
+			WaveInterlockedAddScalar_(OutPagesForHZBIndirectArgsBuffer[0], TILES_PER_PAGE_1D, GroupCountStatic);
+			OutPhysicalPagesForHZB[GroupCountStatic >> LOG2_TILES_PER_PAGE_1D] = PhysicalPageIndex + VirtualShadowMap.StaticPageIndexOffset;
+			WaveInterlockedAddScalar_(OutPagesForHZBIndirectArgsBuffer[0 + 4], 1U, GroupCountStatic);
+		}
 	}
 }
 
diff --git a/Engine/Source/Runtime/Renderer/Private/VirtualShadowMaps/VirtualShadowMapArray.cpp b/Engine/Source/Runtime/Renderer/Private/VirtualShadowMaps/VirtualShadowMapArray.cpp
--- a/Engine/Source/Runtime/Renderer/Private/VirtualShadowMaps/VirtualShadowMapArray.cpp
+++ b/Engine/Source/Runtime/Renderer/Private/VirtualShadowMaps/VirtualShadowMapArray.cpp
@@ -318,9 +318,7 @@ void FVirtualShadowMapArray::Initialize(FRDGBuilder& GraphBuilder, FVirtualShado
 
 	UniformParameters.MaxPhysicalPages = PhysicalPagesX * PhysicalPagesY;
 
-	// TODO: Only enable separate caching when caching is enabled,
-	// but useful to be able to do it regardless for testing purposes.
-	if (/*CacheManager->IsValid() && */CVarCacheStaticSeparately.GetValueOnRenderThread() != 0)
+	if (CacheManager->IsValid() && CVarCacheStaticSeparately.GetValueOnRenderThread() != 0)
 	{
 		// Store the static pages below the dynamic/merged pages
 		UniformParameters.StaticCachedPixelOffsetY = PhysicalPagesY * FVirtualShadowMap::PageSize;
@@ -385,6 +383,13 @@ FIntPoint FVirtualShadowMapArray::GetDynamicPhysicalPoolSize() const
 		ShouldCacheStaticSeparately() ? UniformParameters.StaticCachedPixelOffsetY : UniformParameters.PhysicalPoolSize.Y);
 }
 
+uint32 FVirtualShadowMapArray::GetTotalAllocatedPhysicalPages() const
+{
+	check(bInitialized);
+	return ShouldCacheStaticSeparately() ? UniformParameters.MaxPhysicalPages : (2U * UniformParameters.MaxPhysicalPages);
+}
+
+
 TRDGUniformBufferRef<FVirtualShadowMapUniformParameters> FVirtualShadowMapArray::GetUniformBuffer(FRDGBuilder& GraphBuilder) const
 {
 	// NOTE: Need to allocate new parameter space since the UB changes over the frame as dummy references are replaced
@@ -748,8 +753,8 @@ void FVirtualShadowMapArray::MergeStaticPhysicalPages(FRDGBuilder& GraphBuilder)
 	if (CVarMergePhysicalUsingIndirect.GetValueOnRenderThread() != 0)
 	{
 		FRDGBufferRef MergePagesIndirectArgsRDG = GraphBuilder.CreateBuffer(FRDGBufferDesc::CreateIndirectDesc(3), TEXT("Shadow.Virtual.MergePagesIndirectArgs"));
-		// Note: creating buffer of 2x MaxPhysicalPages as MaxPhysicalPages represents the number before adding the static page overload. The selection shader emits both types of pages.
-		FRDGBufferRef PhysicalPagesToMergeRDG = GraphBuilder.CreateBuffer(FRDGBufferDesc::CreateStructuredDesc(sizeof(int32), 2U * UniformParameters.MaxPhysicalPages + 1), TEXT("Shadow.Virtual.PhysicalPagesToMerge"));
+		// Note: We use GetTotalAllocatedPhysicalPages() to size the buffer as the selection shader emits both static/dynamic pages separately when enabled.
+		FRDGBufferRef PhysicalPagesToMergeRDG = GraphBuilder.CreateBuffer(FRDGBufferDesc::CreateStructuredDesc(sizeof(int32), GetTotalAllocatedPhysicalPages() + 1), TEXT("Shadow.Virtual.PhysicalPagesToMerge"));
 
 		// 1. Initialize the indirect args buffer
 		AddClearIndirectDispatchArgs1DPass(GraphBuilder, MergePagesIndirectArgsRDG);
@@ -775,7 +780,7 @@ void FVirtualShadowMapArray::MergeStaticPhysicalPages(FRDGBuilder& GraphBuilder)
 				RDG_EVENT_NAME("SelectPagesToMerge"),
 				ComputeShader,
 				PassParameters,
-				FIntVector(FMath::DivideAndRoundUp(UniformParameters.MaxPhysicalPages, FSelectPagesToMergeCS::DefaultCSGroupX), 1, 1)
+				FIntVector(FMath::DivideAndRoundUp(GetMaxPhysicalPages(), FSelectPagesToMergeCS::DefaultCSGroupX), 1, 1)
 			);
 
 		}
@@ -1195,10 +1200,10 @@ void FVirtualShadowMapArray::BuildPageAllocations(
 		AddClearUAVPass(GraphBuilder, GraphBuilder.CreateUAV(HInvalidPageFlagsRDG), 0);
 
 		// One additional element as the last element is used as an atomic counter
-		FRDGBufferRef FreePhysicalPagesRDG = GraphBuilder.CreateBuffer(FRDGBufferDesc::CreateStructuredDesc(sizeof(int32), UniformParameters.MaxPhysicalPages + 1), TEXT("Shadow.Virtual.FreePhysicalPages"));
+		FRDGBufferRef FreePhysicalPagesRDG = GraphBuilder.CreateBuffer(FRDGBufferDesc::CreateStructuredDesc(sizeof(int32), GetMaxPhysicalPages() + 1), TEXT("Shadow.Virtual.FreePhysicalPages"));
 		
 		// Enough space for all physical pages that might be allocated
-		PhysicalPageMetaDataRDG = GraphBuilder.CreateBuffer(FRDGBufferDesc::CreateStructuredDesc(sizeof(FPhysicalPageMetaData), UniformParameters.MaxPhysicalPages), TEXT("Shadow.Virtual.PhysicalPageMetaData"));
+		PhysicalPageMetaDataRDG = GraphBuilder.CreateBuffer(FRDGBufferDesc::CreateStructuredDesc(sizeof(FPhysicalPageMetaData), GetMaxPhysicalPages()), TEXT("Shadow.Virtual.PhysicalPageMetaData"));
 
 		AllocatedPageRectBoundsRDG = GraphBuilder.CreateBuffer(FRDGBufferDesc::CreateStructuredDesc(sizeof(FIntVector4), NumPageRects), TEXT("Shadow.Virtual.AllocatedPageRectBounds"));
 
@@ -1214,7 +1219,7 @@ void FVirtualShadowMapArray::BuildPageAllocations(
 				RDG_EVENT_NAME("InitPhysicalPageMetaData"),
 				ComputeShader,
 				PassParameters,
-				FIntVector(FMath::DivideAndRoundUp(UniformParameters.MaxPhysicalPages, FInitPhysicalPageMetaData::DefaultCSGroupX), 1, 1)
+				FIntVector(FMath::DivideAndRoundUp(GetMaxPhysicalPages(), FInitPhysicalPageMetaData::DefaultCSGroupX), 1, 1)
 			);
 		}
 		
@@ -1271,7 +1276,7 @@ void FVirtualShadowMapArray::BuildPageAllocations(
 				RDG_EVENT_NAME("PackFreePages"),
 				ComputeShader,
 				PassParameters,
-				FIntVector(FMath::DivideAndRoundUp(UniformParameters.MaxPhysicalPages, FPackFreePagesCS::DefaultCSGroupX), 1, 1)
+				FIntVector(FMath::DivideAndRoundUp(GetMaxPhysicalPages(), FPackFreePagesCS::DefaultCSGroupX), 1, 1)
 			);
 		}
 
@@ -1346,8 +1351,8 @@ void FVirtualShadowMapArray::BuildPageAllocations(
 			if (CVarInitializePhysicalUsingIndirect.GetValueOnRenderThread() != 0)
 			{
 				FRDGBufferRef InitializePagesIndirectArgsRDG = GraphBuilder.CreateBuffer(FRDGBufferDesc::CreateIndirectDesc(3), TEXT("Shadow.Virtual.InitializePagesIndirectArgs"));
-				// Note: creating buffer of 2x MaxPhysicalPages as MaxPhysicalPages represents the number before adding the static page overload. The selection shader emits both types of pages.
-				FRDGBufferRef PhysicalPagesToInitializeRDG = GraphBuilder.CreateBuffer(FRDGBufferDesc::CreateStructuredDesc(sizeof(int32), 2U * UniformParameters.MaxPhysicalPages + 1), TEXT("Shadow.Virtual.PhysicalPagesToInitialize"));
+				// Note: We use GetTotalAllocatedPhysicalPages() to size the buffer as the selection shader emits both static/dynamic pages separately when enabled.
+				FRDGBufferRef PhysicalPagesToInitializeRDG = GraphBuilder.CreateBuffer(FRDGBufferDesc::CreateStructuredDesc(sizeof(int32), GetTotalAllocatedPhysicalPages() + 1), TEXT("Shadow.Virtual.PhysicalPagesToInitialize"));
 
 				// 1. Initialize the indirect args buffer
 				AddClearIndirectDispatchArgs1DPass(GraphBuilder, InitializePagesIndirectArgsRDG);
@@ -1373,7 +1378,7 @@ void FVirtualShadowMapArray::BuildPageAllocations(
 						RDG_EVENT_NAME("SelectPagesToInitialize"),
 						ComputeShader,
 						PassParameters,
-						FIntVector(FMath::DivideAndRoundUp(UniformParameters.MaxPhysicalPages, FSelectPagesToInitializeCS::DefaultCSGroupX), 1, 1)
+						FIntVector(FMath::DivideAndRoundUp(GetMaxPhysicalPages(), FSelectPagesToInitializeCS::DefaultCSGroupX), 1, 1)
 					);
 
 				}
@@ -2297,7 +2302,8 @@ FRDGTextureRef FVirtualShadowMapArray::BuildHZBFurthest(FRDGBuilder& GraphBuilde
 	{
 		// 1. Gather up all physical pages that are allocated
 		FRDGBufferRef PagesForHZBIndirectArgsRDG = GraphBuilder.CreateBuffer(FRDGBufferDesc::CreateIndirectDesc(2U * 4U), TEXT("Shadow.Virtual.PagesForHZBIndirectArgs"));
-		FRDGBufferRef PhysicalPagesForHZBRDG = GraphBuilder.CreateBuffer(FRDGBufferDesc::CreateStructuredDesc(sizeof(int32), UniformParameters.MaxPhysicalPages + 1), TEXT("Shadow.Virtual.PhysicalPagesForHZB"));
+		// NOTE: Total allocated pages since the shader outputs separate entries for static/dynamic pages
+		FRDGBufferRef PhysicalPagesForHZBRDG = GraphBuilder.CreateBuffer(FRDGBufferDesc::CreateStructuredDesc(sizeof(int32), GetTotalAllocatedPhysicalPages() + 1), TEXT("Shadow.Virtual.PhysicalPagesForHZB"));
 
 		// 1. Clear the indirect args buffer (note 2x args)
 		AddClearIndirectDispatchArgs1DPass(GraphBuilder, PagesForHZBIndirectArgsRDG, 2U);
diff --git a/Engine/Source/Runtime/Renderer/Private/VirtualShadowMaps/VirtualShadowMapArray.h b/Engine/Source/Runtime/Renderer/Private/VirtualShadowMaps/VirtualShadowMapArray.h
--- a/Engine/Source/Runtime/Renderer/Private/VirtualShadowMaps/VirtualShadowMapArray.h
+++ b/Engine/Source/Runtime/Renderer/Private/VirtualShadowMaps/VirtualShadowMapArray.h
@@ -179,6 +179,12 @@ public:
 	// Size of the physical pool for only the dynamic pages (if static are cached separately)
 	FIntPoint GetDynamicPhysicalPoolSize() const;
 
+	// Maximum number of physical pages to allocate. This value is NOT doubled when static caching is
+	// enabled as we always allocate both as pairs (offset in the page pool).
+	uint32 GetMaxPhysicalPages() const { return UniformParameters.MaxPhysicalPages; }
+	// Total physical page count that includes separate static pages
+	uint32 GetTotalAllocatedPhysicalPages() const;
+
 	static void SetShaderDefines(FShaderCompilerEnvironment& OutEnvironment);
 
 	void MergeStaticPhysicalPages(FRDGBuilder& GraphBuilder);
