Fixes warning in motion blur at cook time

diff --git a/Engine/Source/Runtime/Renderer/Private/PostProcess/PostProcessMotionBlur.cpp b/Engine/Source/Runtime/Renderer/Private/PostProcess/PostProcessMotionBlur.cpp
--- a/Engine/Source/Runtime/Renderer/Private/PostProcess/PostProcessMotionBlur.cpp
+++ b/Engine/Source/Runtime/Renderer/Private/PostProcess/PostProcessMotionBlur.cpp
@@ -217,7 +217,7 @@ BEGIN_SHADER_PARAMETER_STRUCT(FVecocityTileTextures, )
 END_SHADER_PARAMETER_STRUCT()
 
 BEGIN_SHADER_PARAMETER_STRUCT(FVecocityTileUAVs, )
-	SHADER_PARAMETER_RDG_TEXTURE_UAV_ARRAY(Texture2D, Textures, [2])
+	SHADER_PARAMETER_RDG_TEXTURE_UAV_ARRAY(RWTexture2D, Textures, [2])
 END_SHADER_PARAMETER_STRUCT()
 
 FVecocityTileTextures CreateVecocityTileTextures(FRDGBuilder& GraphBuilder, FIntPoint VelocityTileCount, const TCHAR* DebugName, bool ScatterDilatation = false)
