Destroys FViewInfo::PrevViewInfo before FRDGBuilder::Execute() to make sure all reference to prior histories are only within the graph to maximise memory resuse

diff --git a/Engine/Source/Runtime/Renderer/Private/DeferredShadingRenderer.cpp b/Engine/Source/Runtime/Renderer/Private/DeferredShadingRenderer.cpp
--- a/Engine/Source/Runtime/Renderer/Private/DeferredShadingRenderer.cpp
+++ b/Engine/Source/Runtime/Renderer/Private/DeferredShadingRenderer.cpp
@@ -3263,6 +3263,12 @@ void FDeferredShadingSceneRenderer::Render(FRDGBuilder& GraphBuilder)
 	}
 
 	QueueSceneTextureExtractions(GraphBuilder, SceneTextures);
+
+	// Release the view's previous frame histories so that their memory can be reused at the graph's execution.
+	for (int32 ViewIndex = 0; ViewIndex < Views.Num(); ViewIndex++)
+	{
+		Views[ViewIndex].PrevViewInfo = FPreviousViewInfo();
+	}
 }
 
 #if RHI_RAYTRACING
