Add r.VT.ProduceLockedTilesOnFlush to not sync when producing locked tiles.

diff --git a/Engine/Source/Runtime/Renderer/Private/VT/VirtualTextureSystem.cpp b/Engine/Source/Runtime/Renderer/Private/VT/VirtualTextureSystem.cpp
--- a/Engine/Source/Runtime/Renderer/Private/VT/VirtualTextureSystem.cpp
+++ b/Engine/Source/Runtime/Renderer/Private/VT/VirtualTextureSystem.cpp
@@ -116,6 +116,12 @@ static TAutoConsoleVariable<int32> CVarVTForceContinuousUpdate(
 	TEXT("Force continuous update on all virtual textures."),
 	ECVF_RenderThreadSafe
 );
+static TAutoConsoleVariable<int32> CVarVTSyncProduceLockedTiles(
+	TEXT("r.VT.SyncProduceLockedTiles"),
+	1,
+	TEXT("Should we sync loading when producing locked tiles"),
+	ECVF_RenderThreadSafe
+);
 static TAutoConsoleVariable<int32> CVarVTProduceLockedTilesOnFlush(
 	TEXT("r.VT.ProduceLockedTilesOnFlush"),
 	1,
@@ -2210,6 +2216,8 @@ void FVirtualTextureSystem::SubmitRequests(FRDGBuilder& GraphBuilder, ERHIFeatur
 		TArray<FProducePageDataPrepareTask> PrepareTasks;
 		PrepareTasks.Reserve(RequestList->GetNumLoadRequests());
 
+		const bool bSyncProduceLockedTiles = CVarVTSyncProduceLockedTiles.GetValueOnRenderThread() != 0;
+
 		const uint32 MaxPagesProduced = VirtualTextureScalability::GetMaxPagesProducedPerFrame();
 		const uint32 PageFreeThreshold = VirtualTextureScalability::GetPageFreeThreshold();
 		uint32 NumStacksProduced = 0u;
@@ -2219,7 +2227,7 @@ void FVirtualTextureSystem::SubmitRequests(FRDGBuilder& GraphBuilder, ERHIFeatur
 		for (uint32 RequestIndex = 0u; RequestIndex < RequestList->GetNumLoadRequests(); ++RequestIndex)
 		{
 			const bool bLockTile = RequestList->IsLocked(RequestIndex);
-			const bool bForceProduceTile = (bLockTile || !bAsync);
+			const bool bForceProduceTile = !bAsync || (bLockTile && bSyncProduceLockedTiles);
 			const FVirtualTextureLocalTile TileToLoad = RequestList->GetLoadRequest(RequestIndex);
 			const FVirtualTextureProducerHandle ProducerHandle = TileToLoad.GetProducerHandle();
 			const FVirtualTextureProducer& Producer = Producers.GetProducer(ProducerHandle);
@@ -2393,6 +2401,7 @@ void FVirtualTextureSystem::SubmitRequests(FRDGBuilder& GraphBuilder, ERHIFeatur
 			static bool bWaitForTasks = true;
 			if (bWaitForTasks)
 			{
+				TRACE_CPUPROFILER_EVENT_SCOPE(FVirtualTextureSystem::ProcessRequests_Wait);
 				QUICK_SCOPE_CYCLE_COUNTER(ProcessRequests_Wait);
 				FTaskGraphInterface::Get().WaitUntilTasksComplete(ProducePageTasks, ENamedThreads::GetRenderThread_Local());
 			}
