Need to remove/add DistanceFieldData when a primitive's instance count changes.

- In the future this should do a partial update based on the primitive updates in the command buffer.

diff --git a/Engine/Source/Runtime/Renderer/Private/RendererScene.cpp b/Engine/Source/Runtime/Renderer/Private/RendererScene.cpp
--- a/Engine/Source/Runtime/Renderer/Private/RendererScene.cpp
+++ b/Engine/Source/Runtime/Renderer/Private/RendererScene.cpp
@@ -4758,7 +4758,8 @@ void FScene::UpdateAllPrimitiveSceneInfos(FRDGBuilder& GraphBuilder, bool bAsync
 			PrimitiveSceneInfo->FreeGPUSceneInstances();
 			FPrimitiveSceneInfo::AllocateGPUSceneInstances(this, MakeArrayView(&PrimitiveSceneInfo, 1));
 
-			DistanceFieldSceneData.UpdatePrimitive(PrimitiveSceneInfo);
+			DistanceFieldSceneData.RemovePrimitive(PrimitiveSceneInfo);
+			DistanceFieldSceneData.AddPrimitive(PrimitiveSceneInfo);
 
 			LumenSceneData->RemovePrimitive(PrimitiveSceneInfo, PrimitiveSceneInfo->PackedIndex);
 			LumenSceneData->AddPrimitive(PrimitiveSceneInfo);
