Add support for mip bias of prestream VT page requests.
This includes a fixed per platform component and a dynamic component that takes into account the current pool residency mip bias.

diff --git a/Engine/Source/Runtime/Renderer/Private/VT/VirtualTextureSystem.cpp b/Engine/Source/Runtime/Renderer/Private/VT/VirtualTextureSystem.cpp
--- a/Engine/Source/Runtime/Renderer/Private/VT/VirtualTextureSystem.cpp
+++ b/Engine/Source/Runtime/Renderer/Private/VT/VirtualTextureSystem.cpp
@@ -98,6 +98,12 @@ static TAutoConsoleVariable<int32> CVarVTEnablePlayback(
 	TEXT("Enable playback of recorded feedback requests."),
 	ECVF_RenderThreadSafe
 );
+static TAutoConsoleVariable<float> CVarVTPlaybackMipBias(
+	TEXT("r.VT.PlaybackMipBias"),
+	0,
+	TEXT("Mip bias to apply during playback of recorded feedback requests."),
+	ECVF_RenderThreadSafe
+);
 static TAutoConsoleVariable<int32> CVarVTNumGatherTasks(
 	TEXT("r.VT.NumGatherTasks"),
 	1,
@@ -198,9 +204,10 @@ public:
 struct FAddRequestedTilesParameters
 {
 	FVirtualTextureSystem* System = nullptr;
+	uint32 LevelBias = 0u;
 	const uint64* RequestBuffer = nullptr;
-	FUniquePageList* UniquePageList = nullptr;
 	uint32 NumRequests = 0u;
+	FUniquePageList* UniquePageList = nullptr;
 };
 
 class FAddRequestedTilesTask
@@ -1420,9 +1427,10 @@ void FVirtualTextureSystem::Update(FRDGBuilder& GraphBuilder, ERHIFeatureLevel::
 
 			// todo: We can split this into concurrent tasks. 
 			FAddRequestedTilesParameters Parameters;
+			Parameters.System = this;
+			Parameters.LevelBias = FMath::FloorToInt(CVarVTPlaybackMipBias.GetValueOnRenderThread() + GetGlobalMipBias() + 0.5f);
 			Parameters.RequestBuffer = PageRequestPlaybackBuffer.GetData();
 			Parameters.NumRequests = PageRequestPlaybackBuffer.Num();
-			Parameters.System = this;
 			Parameters.UniquePageList = new(MemStack) FUniquePageList;
 
 			FAddRequestedTilesTask::DoTask(Parameters);
@@ -2732,6 +2740,7 @@ void FVirtualTextureSystem::AddRequestedTilesTask(const FAddRequestedTilesParame
 	FUniquePageList* RESTRICT RequestedPageList = Parameters.UniquePageList;
 	const uint64* RESTRICT Buffer = Parameters.RequestBuffer;
 	const uint32 BufferSize = Parameters.NumRequests;
+	const uint32 LevelBias = Parameters.LevelBias;
 
 	uint32 CurrentPersistentHash = ~0u;
 	IAllocatedVirtualTexture* AllocatedVT = nullptr;
@@ -2757,7 +2766,9 @@ void FVirtualTextureSystem::AddRequestedTilesTask(const FAddRequestedTilesParame
 
 			const uint32 BaseAddress = AllocatedVT->GetVirtualAddress();
 			const uint32 Address = BaseAddress + LocalAddress;
-			const uint32 Level = FMath::Max(LevelPlusOne, 1u) - 1;
+			const uint32 BaseLevel = FMath::Max(LevelPlusOne, 1u) - 1;
+			const uint32 MaxLevel = AllocatedVT->GetMaxLevel();
+			const uint32 Level = FMath::Min(BaseLevel + LevelBias, MaxLevel);
 			const uint32 PageX = FMath::ReverseMortonCode2(Address) >> Level;
 			const uint32 PageY = FMath::ReverseMortonCode2(Address >> 1) >> Level;
 			const uint32 SpaceId = AllocatedVT->GetSpaceID();
