Round HZB up to pow2 size in indirect VSM path as some of the sampling code depends on it.
Add cvar for 32 vs 16-bit VSM. Default 32 as currently but will revisit.
Minor cvar rename for consistency.
Enforce SM5 for VSM shaders (already implied by Nanite support, but make it explicit.)

diff --git a/Engine/Source/Runtime/Renderer/Private/VirtualShadowMaps/VirtualShadowMapArray.cpp b/Engine/Source/Runtime/Renderer/Private/VirtualShadowMaps/VirtualShadowMapArray.cpp
--- a/Engine/Source/Runtime/Renderer/Private/VirtualShadowMaps/VirtualShadowMapArray.cpp
+++ b/Engine/Source/Runtime/Renderer/Private/VirtualShadowMaps/VirtualShadowMapArray.cpp
@@ -51,8 +51,8 @@ TAutoConsoleVariable<int32> CVarMaxPhysicalPages(
 	ECVF_RenderThreadSafe
 );
 
-TAutoConsoleVariable<int32> CVarCacheStaticSeparately(
-	TEXT("r.Shadow.Virtual.Cache.StaticSeparately"),
+TAutoConsoleVariable<int32> CVarCacheStaticSeparate(
+	TEXT("r.Shadow.Virtual.Cache.StaticSeparate"),
 	0,
 	TEXT("When enabled, caches static objects in separate pages from dynamic objects.\n")
 	TEXT("This can improve performance in largely static scenes, but doubles the memory cost of the physical page pool."),
@@ -218,6 +218,13 @@ TAutoConsoleVariable<int32> CVarVSMBuildHZBPerPage(
 	ECVF_RenderThreadSafe
 );
 
+TAutoConsoleVariable<int32> CVarVSMHZB32Bit(
+	TEXT("r.Shadow.Virtual.HZB32Bit"),
+	1,
+	TEXT("Use a full 32-bit HZB buffer. This uses more memory but can offer more precise culling in cases with lots of overlapping detailed geometry."),
+	ECVF_RenderThreadSafe
+);
+
 FMatrix CalcTranslatedWorldToShadowUVMatrix(
 	const FMatrix& TranslatedWorldToShadowView,
 	const FMatrix& ViewToClip)
@@ -318,7 +325,7 @@ void FVirtualShadowMapArray::Initialize(FRDGBuilder& GraphBuilder, FVirtualShado
 
 	UniformParameters.MaxPhysicalPages = PhysicalPagesX * PhysicalPagesY;
 
-	if (CacheManager->IsValid() && CVarCacheStaticSeparately.GetValueOnRenderThread() != 0)
+	if (CacheManager->IsValid() && CVarCacheStaticSeparate.GetValueOnRenderThread() != 0)
 	{
 		// Store the static pages below the dynamic/merged pages
 		UniformParameters.StaticCachedPixelOffsetY = PhysicalPagesY * FVirtualShadowMap::PageSize;
@@ -449,7 +456,8 @@ public:
 
 	static bool ShouldCompilePermutation(const FGlobalShaderPermutationParameters& Parameters)
 	{
-		return DoesPlatformSupportNanite(Parameters.Platform);
+		return IsFeatureLevelSupported(Parameters.Platform, ERHIFeatureLevel::SM5) &&
+			DoesPlatformSupportNanite(Parameters.Platform);
 	}
 
 	/**
@@ -2295,7 +2303,7 @@ IMPLEMENT_GLOBAL_SHADER(FVirtualSmBBuildHZBPerPageTopCS, "/Engine/Private/Virtua
 FRDGTextureRef FVirtualShadowMapArray::BuildHZBFurthest(FRDGBuilder& GraphBuilder)
 {
 	const FIntRect ViewRect(0, 0, GetPhysicalPoolSize().X, GetPhysicalPoolSize().Y);
-	const EPixelFormat Format = PF_R32_FLOAT;
+	const EPixelFormat Format = CVarVSMHZB32Bit.GetValueOnRenderThread() ? PF_R32_FLOAT : PF_R16F;
 
 	FRDGTextureRef OutFurthestHZBTexture = nullptr;
 	if (CVarVSMBuildHZBPerPage.GetValueOnRenderThread())
@@ -2329,7 +2337,14 @@ FRDGTextureRef FVirtualShadowMapArray::BuildHZBFurthest(FRDGBuilder& GraphBuilde
 		}
 
 
-		FIntPoint HZBSize(GetPhysicalPoolSize().X / 2, GetPhysicalPoolSize().Y / 2);
+		//FIntPoint HZBSize(GetPhysicalPoolSize().X / 2, GetPhysicalPoolSize().Y / 2);
+		FIntPoint PhysicalPoolSize = GetPhysicalPoolSize();
+
+		FIntPoint HZBSize;
+		//HZBSize = FIntPoint(PhysicalPoolSize.X / 2, PhysicalPoolSize.Y / 2);
+		// NOTE: IsVisibleHZB currently assumes HZB is pow2 size, so force that. See HZBCull.ush.
+		HZBSize.X = FMath::Max( FPlatformMath::RoundUpToPowerOfTwo( PhysicalPoolSize.X ) >> 1, 1u );
+		HZBSize.Y = FMath::Max( FPlatformMath::RoundUpToPowerOfTwo( PhysicalPoolSize.Y ) >> 1, 1u );
 
 		FRDGTextureDesc HZBDesc = FRDGTextureDesc::Create2D(HZBSize, Format, FClearValueBinding::None, TexCreate_ShaderResource | TexCreate_UAV, FVirtualSmBuildHZBPerPageCS::TotalHZBLevels);
 		HZBDesc.Flags |= GFastVRamConfig.HZB;
@@ -2402,7 +2417,7 @@ FRDGTextureRef FVirtualShadowMapArray::BuildHZBFurthest(FRDGBuilder& GraphBuilde
 			ViewRect,
 			GMaxRHIFeatureLevel,
 			GMaxRHIShaderPlatform,
-			TEXT("Shadow.Virtual.HZB"),
+			TEXT("Shadow.Virtual.PreviousOccluderHZB"),
 			/* OutFurthestHZBTexture = */ &OutFurthestHZBTexture,
 			Format);
 	}
