Don't assert on case where locked tile is invalid to produce.
This must be happening after the change to defer producing locked tiles.
Needs some further investigation to understand why it is happening, but this quick fix should unblock.

diff --git a/Engine/Source/Runtime/Renderer/Private/VT/VirtualTextureSystem.cpp b/Engine/Source/Runtime/Renderer/Private/VT/VirtualTextureSystem.cpp
--- a/Engine/Source/Runtime/Renderer/Private/VT/VirtualTextureSystem.cpp
+++ b/Engine/Source/Runtime/Renderer/Private/VT/VirtualTextureSystem.cpp
@@ -2252,7 +2252,7 @@ void FVirtualTextureSystem::SubmitRequests(FRDGBuilder& GraphBuilder, ERHIFeatur
 
 			const EVTRequestPagePriority Priority = bLockTile ? EVTRequestPagePriority::High : EVTRequestPagePriority::Normal;
 			FVTRequestPageResult RequestPageResult = Producer.GetVirtualTexture()->RequestPageData(ProducerHandle, ProducerTextureLayerMask, TileToLoad.Local_vLevel, TileToLoad.Local_vAddress, Priority);
-			checkf(!bLockTile || RequestPageResult.Status != EVTRequestPageStatus::Invalid, TEXT("Tried to lock an invalid VT tile"));
+			//checkf(!bLockTile || RequestPageResult.Status != EVTRequestPageStatus::Invalid, TEXT("Tried to lock an invalid VT tile"));
 			if (RequestPageResult.Status == EVTRequestPageStatus::Pending && bForceProduceTile)
 			{
 				// If we're trying to lock this tile, we're OK producing data now (and possibly waiting) as long as data is pending
@@ -2267,8 +2267,10 @@ void FVirtualTextureSystem::SubmitRequests(FRDGBuilder& GraphBuilder, ERHIFeatur
 			}
 
 			bool bTileLoaded = false;
+			bool bTileInvalid = false;
 			if (RequestPageResult.Status == EVTRequestPageStatus::Invalid)
 			{
+				bTileInvalid = true;
 				if (CVarVTVerbose.GetValueOnRenderThread())
 				{
 					UE_LOG(LogConsoleResponse, Display, TEXT("vAddr %i@%i is not a valid request for AllocatedVT but is still requested."), TileToLoad.Local_vAddress, TileToLoad.Local_vLevel);
@@ -2389,7 +2391,7 @@ void FVirtualTextureSystem::SubmitRequests(FRDGBuilder& GraphBuilder, ERHIFeatur
 				}
 			}
 
-			if (bLockTile && !bTileLoaded)
+			if (bLockTile && !bTileLoaded && !bTileInvalid)
 			{
 				// Want to lock this tile, but didn't manage to load it this frame, add it back to the list to try the lock again next frame
 				TilesToLock.Add(TileToLoad);
