Allocate transient scratch buffer using RDG for dynamic geometry BLAS updates.
Expose FRayTracingAccelerationStructureSize in FRHIRayTracingGeometry.

diff --git a/Engine/Plugins/Experimental/GPULightmass/Source/GPULightmass/Private/LightmapRenderer.cpp b/Engine/Plugins/Experimental/GPULightmass/Source/GPULightmass/Private/LightmapRenderer.cpp
--- a/Engine/Plugins/Experimental/GPULightmass/Source/GPULightmass/Private/LightmapRenderer.cpp
+++ b/Engine/Plugins/Experimental/GPULightmass/Source/GPULightmass/Private/LightmapRenderer.cpp
@@ -920,7 +920,7 @@ void FSceneRenderState::SetupRayTracingScene(int32 LODIndex)
 								StaticMeshInstanceRenderStates.Elements.Num() + InstanceGroupRenderStates.Elements.Num() + LandscapeIndex
 							);
 
-							DynamicGeometryCollection.DispatchUpdates(RHICmdList);
+							DynamicGeometryCollection.DispatchUpdates(RHICmdList, nullptr);
 
 							// Landscape VF doesn't really use the vertex buffer in HitGroupSystemParameters
 							// We can release after all related RHI cmds get dispatched onto the cmd list
diff --git a/Engine/Source/Runtime/D3D12RHI/Private/D3D12RayTracing.h b/Engine/Source/Runtime/D3D12RHI/Private/D3D12RayTracing.h
--- a/Engine/Source/Runtime/D3D12RHI/Private/D3D12RayTracing.h
+++ b/Engine/Source/Runtime/D3D12RHI/Private/D3D12RayTracing.h
@@ -32,6 +32,7 @@ public:
 		return AccelerationStructureBuffers[GPUIndex]->ResourceLocation.GetGPUVirtualAddress();
 	}
 	virtual uint32 GetNumSegments() const final override { return Segments.Num(); }
+	virtual FRayTracingAccelerationStructureSize GetSizeInfo() const final override { return SizeInfo; }
 
 	void SetupHitGroupSystemParameters(uint32 InGPUIndex);
 	void TransitionBuffers(FD3D12CommandContext& CommandContext);
diff --git a/Engine/Source/Runtime/RHI/Public/DynamicRHI.h b/Engine/Source/Runtime/RHI/Public/DynamicRHI.h
--- a/Engine/Source/Runtime/RHI/Public/DynamicRHI.h
+++ b/Engine/Source/Runtime/RHI/Public/DynamicRHI.h
@@ -298,13 +298,6 @@ struct FRayTracingSceneInitializer2
 	FName DebugName;
 };
 
-struct FRayTracingAccelerationStructureSize
-{
-	uint64 ResultSize = 0;
-	uint64 BuildScratchSize = 0;
-	uint64 UpdateScratchSize = 0;
-};
-
 struct FShaderResourceViewInitializer
 {
 	struct FBufferShaderResourceViewInitializer
diff --git a/Engine/Source/Runtime/RHI/Public/RHIResources.h b/Engine/Source/Runtime/RHI/Public/RHIResources.h
--- a/Engine/Source/Runtime/RHI/Public/RHIResources.h
+++ b/Engine/Source/Runtime/RHI/Public/RHIResources.h
@@ -1912,6 +1912,13 @@ typedef TRefCountPtr<FRHIRayTracingPipelineState> FRayTracingPipelineStateRHIRef
 // Ray tracing resources
 //
 
+struct FRayTracingAccelerationStructureSize
+{
+	uint64 ResultSize = 0;
+	uint64 BuildScratchSize = 0;
+	uint64 UpdateScratchSize = 0;
+};
+
 class FRHIRayTracingAccelerationStructure
 	: public FRHIResource
 #if ENABLE_RHI_VALIDATION
@@ -1929,6 +1936,7 @@ class FRHIRayTracingGeometry : public FRHIRayTracingAccelerationStructure
 public:
 	virtual FRayTracingAccelerationStructureAddress GetAccelerationStructureAddress(uint64 GPUIndex) const = 0;
 	virtual uint32 GetNumSegments() const = 0;
+	virtual FRayTracingAccelerationStructureSize GetSizeInfo() const = 0;
 };
 
 typedef TRefCountPtr<FRHIRayTracingGeometry>     FRayTracingGeometryRHIRef;
diff --git a/Engine/Source/Runtime/Renderer/Private/DeferredShadingRenderer.cpp b/Engine/Source/Runtime/Renderer/Private/DeferredShadingRenderer.cpp
--- a/Engine/Source/Runtime/Renderer/Private/DeferredShadingRenderer.cpp
+++ b/Engine/Source/Runtime/Renderer/Private/DeferredShadingRenderer.cpp
@@ -1301,6 +1301,7 @@ static void DeduplicateRayGenerationShaders(TArray< FRHIRayTracingShader*>& RayG
 
 BEGIN_SHADER_PARAMETER_STRUCT(FBuildAccelerationStructurePassParams, )
 	RDG_BUFFER_ACCESS(RayTracingSceneScratchBuffer, ERHIAccess::UAVCompute)
+	RDG_BUFFER_ACCESS(DynamicGeometryScratchBuffer, ERHIAccess::UAVCompute)
 	RDG_BUFFER_ACCESS(RayTracingSceneInstanceBuffer, ERHIAccess::SRVCompute)
 END_SHADER_PARAMETER_STRUCT()
 
@@ -1435,7 +1436,7 @@ BEGIN_SHADER_PARAMETER_STRUCT(FGPUSkinCacheBLASUpdateParams, )
 	RDG_BUFFER_ACCESS(SharedScratchBuffer, ERHIAccess::UAVCompute)
 END_SHADER_PARAMETER_STRUCT()
 
-bool FDeferredShadingSceneRenderer::DispatchRayTracingWorldUpdates(FRDGBuilder& GraphBuilder)
+bool FDeferredShadingSceneRenderer::DispatchRayTracingWorldUpdates(FRDGBuilder& GraphBuilder, FRDGBufferRef& OutDynamicGeometryScratchBuffer)
 {
 	bool bAnyRayTracingPassEnabled = false;
 	for (int32 ViewIndex = 0; ViewIndex < Views.Num(); ViewIndex++)
@@ -1503,12 +1504,28 @@ bool FDeferredShadingSceneRenderer::DispatchRayTracingWorldUpdates(FRDGBuilder&
 
 	RayTracingScene.Create(GraphBuilder, Scene->GPUScene);
 
+	OutDynamicGeometryScratchBuffer = nullptr;
+
+	const uint32 BLASScratchSize = Scene->GetRayTracingDynamicGeometryCollection()->ComputeScratchBufferSize();
+	if (BLASScratchSize > 0)
+	{
+		const uint32 ScratchAlignment = GRHIRayTracingAccelerationStructureAlignment;
+		FRDGBufferDesc ScratchBufferDesc;
+		ScratchBufferDesc.UnderlyingType = FRDGBufferDesc::EUnderlyingType::StructuredBuffer;
+		ScratchBufferDesc.Usage = BUF_UnorderedAccess;
+		ScratchBufferDesc.BytesPerElement = ScratchAlignment;
+		ScratchBufferDesc.NumElements = FMath::DivideAndRoundUp(BLASScratchSize, ScratchAlignment);
+
+		OutDynamicGeometryScratchBuffer = GraphBuilder.CreateBuffer(ScratchBufferDesc, TEXT("DynamicGeometry.BLASSharedScratchBuffer"));
+	}
+
 	const bool bRayTracingAsyncBuild = CVarRayTracingAsyncBuild.GetValueOnRenderThread() != 0;
 
 	if (bRayTracingAsyncBuild && GRHISupportsRayTracingAsyncBuildAccelerationStructure)
 	{
 		FBuildAccelerationStructurePassParams* PassParams = GraphBuilder.AllocParameters<FBuildAccelerationStructurePassParams>();
 		PassParams->RayTracingSceneScratchBuffer = Scene->RayTracingScene.BuildScratchBuffer;
+		PassParams->DynamicGeometryScratchBuffer = OutDynamicGeometryScratchBuffer;
 		PassParams->RayTracingSceneInstanceBuffer = Scene->RayTracingScene.InstanceBuffer;
 
 		GraphBuilder.AddPass(RDG_EVENT_NAME("BuildAccelerationStructure"), PassParams, ERDGPassFlags::Compute | ERDGPassFlags::NeverCull,
@@ -1525,7 +1542,9 @@ bool FDeferredShadingSceneRenderer::DispatchRayTracingWorldUpdates(FRDGBuilder&
 				RHICmdList.BeginTransition(RayTracingDynamicGeometryUpdateBeginTransition);
 				RHIAsyncCmdList.EndTransition(RayTracingDynamicGeometryUpdateBeginTransition);
 
-				Scene->GetRayTracingDynamicGeometryCollection()->DispatchUpdates(RHIAsyncCmdList);
+				FRHIBuffer* DynamicGeometryScratchBuffer = PassParams->DynamicGeometryScratchBuffer ? PassParams->DynamicGeometryScratchBuffer->GetRHI() : nullptr;
+
+				Scene->GetRayTracingDynamicGeometryCollection()->DispatchUpdates(RHIAsyncCmdList, DynamicGeometryScratchBuffer);
 
 				FRHIRayTracingScene* RayTracingSceneRHI = Scene->RayTracingScene.GetRHIRayTracingSceneChecked();
 				FRHIBuffer* AccelerationStructureBuffer = Scene->RayTracingScene.GetBufferChecked();
@@ -1557,10 +1576,17 @@ bool FDeferredShadingSceneRenderer::DispatchRayTracingWorldUpdates(FRDGBuilder&
 	{
 		{
 			RDG_GPU_STAT_SCOPE(GraphBuilder, RayTracingGeometry);
-			AddPass(GraphBuilder, RDG_EVENT_NAME("RayTracingGeometry"), [this](FRHICommandListImmediate& RHICmdList)
-			{
-				Scene->GetRayTracingDynamicGeometryCollection()->DispatchUpdates(RHICmdList);
-			});
+
+			FBuildAccelerationStructurePassParams* PassParams = GraphBuilder.AllocParameters<FBuildAccelerationStructurePassParams>();
+			PassParams->RayTracingSceneScratchBuffer = nullptr;
+			PassParams->DynamicGeometryScratchBuffer = OutDynamicGeometryScratchBuffer;
+
+			GraphBuilder.AddPass(RDG_EVENT_NAME("RayTracingGeometry"), PassParams, ERDGPassFlags::Compute | ERDGPassFlags::NeverCull,
+				[this, PassParams](FRHICommandListImmediate& RHICmdList)
+				{
+					FRHIBuffer* DynamicGeometryScratchBuffer = PassParams->DynamicGeometryScratchBuffer ? PassParams->DynamicGeometryScratchBuffer->GetRHI() : nullptr;
+					Scene->GetRayTracingDynamicGeometryCollection()->DispatchUpdates(RHICmdList, DynamicGeometryScratchBuffer);
+				});
 		}
 
 		{
@@ -1568,7 +1594,6 @@ bool FDeferredShadingSceneRenderer::DispatchRayTracingWorldUpdates(FRDGBuilder&
 
 			FBuildAccelerationStructurePassParams* PassParams = GraphBuilder.AllocParameters<FBuildAccelerationStructurePassParams>();
 			PassParams->RayTracingSceneScratchBuffer = Scene->RayTracingScene.BuildScratchBuffer;
-			PassParams->RayTracingSceneInstanceBuffer = Scene->RayTracingScene.InstanceBuffer;
 
 			GraphBuilder.AddPass(RDG_EVENT_NAME("RayTracingScene"), PassParams, ERDGPassFlags::Compute | ERDGPassFlags::NeverCull,
 				[this, PassParams](FRHICommandListImmediate& RHICmdList)
@@ -1650,7 +1675,7 @@ static void ReleaseRaytracingResources(FRDGBuilder& GraphBuilder, TArrayView<FVi
 	});
 }
 
-void FDeferredShadingSceneRenderer::WaitForRayTracingScene(FRDGBuilder& GraphBuilder)
+void FDeferredShadingSceneRenderer::WaitForRayTracingScene(FRDGBuilder& GraphBuilder, FRDGBufferRef DynamicGeometryScratchBuffer)
 {
 	bool bAnyRayTracingPassEnabled = false;
 	for (int32 ViewIndex = 0; ViewIndex < Views.Num(); ViewIndex++)
@@ -1686,6 +1711,7 @@ void FDeferredShadingSceneRenderer::WaitForRayTracingScene(FRDGBuilder& GraphBui
 	// Scratch buffer must be referenced in this pass, as it must live until the BVH build is complete.
 	FBuildAccelerationStructurePassParams* PassParams = GraphBuilder.AllocParameters<FBuildAccelerationStructurePassParams>();
 	PassParams->RayTracingSceneScratchBuffer = Scene->RayTracingScene.BuildScratchBuffer;
+	PassParams->DynamicGeometryScratchBuffer = DynamicGeometryScratchBuffer;
 
 	GraphBuilder.AddPass(RDG_EVENT_NAME("WaitForRayTracingScene"), PassParams, ERDGPassFlags::Compute | ERDGPassFlags::NeverCull,
 		[this, PassParams](FRHICommandListImmediate& RHICmdList)
@@ -2612,7 +2638,8 @@ void FDeferredShadingSceneRenderer::Render(FRDGBuilder& GraphBuilder)
 
 #if RHI_RAYTRACING
 	// Async AS builds can potentially overlap with BasePass
-	DispatchRayTracingWorldUpdates(GraphBuilder);
+	FRDGBufferRef DynamicGeometryScratchBuffer;
+	DispatchRayTracingWorldUpdates(GraphBuilder, DynamicGeometryScratchBuffer);
 #endif
 
 	{
@@ -2706,7 +2733,7 @@ void FDeferredShadingSceneRenderer::Render(FRDGBuilder& GraphBuilder)
 		// Lumen scene lighting requires ray tracing scene to be ready if HWRT shadows are desired
 		if (Lumen::UseHardwareRayTracedSceneLighting())
 		{
-			WaitForRayTracingScene(GraphBuilder);
+			WaitForRayTracingScene(GraphBuilder, DynamicGeometryScratchBuffer);
 			bRayTracingSceneReady = true;
 		}
 #endif // RHI_RAYTRACING
@@ -2815,7 +2842,7 @@ void FDeferredShadingSceneRenderer::Render(FRDGBuilder& GraphBuilder)
 	// If Lumen did not force an earlier ray tracing scene sync, we must wait for it here.
 	if (!bRayTracingSceneReady)
 	{
-		WaitForRayTracingScene(GraphBuilder);
+		WaitForRayTracingScene(GraphBuilder, DynamicGeometryScratchBuffer);
 		bRayTracingSceneReady = true;
 	}
 #endif // RHI_RAYTRACING
diff --git a/Engine/Source/Runtime/Renderer/Private/DeferredShadingRenderer.h b/Engine/Source/Runtime/Renderer/Private/DeferredShadingRenderer.h
--- a/Engine/Source/Runtime/Renderer/Private/DeferredShadingRenderer.h
+++ b/Engine/Source/Runtime/Renderer/Private/DeferredShadingRenderer.h
@@ -929,7 +929,7 @@ private:
 	void ComputePathCompaction(FRHICommandListImmediate& RHICmdList, const FViewInfo& View, FRHITexture* RadianceTexture, FRHITexture* SampleCountTexture, FRHITexture* PixelPositionTexture,
 		FRHIUnorderedAccessView* RadianceSortedRedUAV, FRHIUnorderedAccessView* RadianceSortedGreenUAV, FRHIUnorderedAccessView* RadianceSortedBlueUAV, FRHIUnorderedAccessView* RadianceSortedAlphaUAV, FRHIUnorderedAccessView* SampleCountSortedUAV);
 
-	void WaitForRayTracingScene(FRDGBuilder& GraphBuilder);
+	void WaitForRayTracingScene(FRDGBuilder& GraphBuilder, FRDGBufferRef DynamicGeometryScratchBuffer);
 
 	/** Debug ray tracing functions. */
 	void RenderRayTracingDebug(FRDGBuilder& GraphBuilder, const FViewInfo& View, FRDGTextureRef SceneColorOutputTexture);
@@ -939,7 +939,7 @@ private:
 	bool GatherRayTracingWorldInstancesForView(FRDGBuilder& GraphBuilder, FViewInfo& View, FRayTracingScene& RayTracingScene);
 
 	bool SetupRayTracingPipelineStates(FRHICommandListImmediate& RHICmdList);
-	bool DispatchRayTracingWorldUpdates(FRDGBuilder& GraphBuilder);
+	bool DispatchRayTracingWorldUpdates(FRDGBuilder& GraphBuilder, FRDGBufferRef& OutDynamicGeometryScratchBuffer);
 	FRayTracingPipelineState* BindRayTracingMaterialPipeline(FRHICommandList& RHICmdList, FViewInfo& View, const TArrayView<FRHIRayTracingShader*>& RayGenShaderTable);
 	FRayTracingPipelineState* BindRayTracingDeferredMaterialGatherPipeline(FRHICommandList& RHICmdList, const FViewInfo& View, const TArrayView<FRHIRayTracingShader*>& RayGenShaderTable);
 	FRayTracingPipelineState* BindLumenHardwareRayTracingMaterialPipeline(FRHICommandList& RHICmdList, const FViewInfo& View, const TArrayView<FRHIRayTracingShader*>& RayGenShaderTable, FRHIBuffer* OutHitGroupDataBuffer);
diff --git a/Engine/Source/Runtime/Renderer/Private/RayTracing/RayTracingDynamicGeometry.cpp b/Engine/Source/Runtime/Renderer/Private/RayTracing/RayTracingDynamicGeometry.cpp
--- a/Engine/Source/Runtime/Renderer/Private/RayTracing/RayTracingDynamicGeometry.cpp
+++ b/Engine/Source/Runtime/Renderer/Private/RayTracing/RayTracingDynamicGeometry.cpp
@@ -330,7 +330,7 @@ void FRayTracingDynamicGeometryCollection::AddDynamicMeshBatchForGeometryUpdate(
 	}
 }
 
-void FRayTracingDynamicGeometryCollection::DispatchUpdates(FRHIComputeCommandList& ParentCmdList)
+void FRayTracingDynamicGeometryCollection::DispatchUpdates(FRHIComputeCommandList& ParentCmdList, FRHIBuffer* ScratchBuffer)
 {
 #if WANTS_DRAW_MESH_EVENTS
 #define SCOPED_DRAW_OR_COMPUTE_EVENT(ParentCmdList, Name) FDrawEvent PREPROCESSOR_JOIN(Event_##Name,__LINE__); if(GetEmitDrawEvents()) PREPROCESSOR_JOIN(Event_##Name,__LINE__).Start(&ParentCmdList, FColor(0), TEXT(#Name));
@@ -361,6 +361,9 @@ void FRayTracingDynamicGeometryCollection::DispatchUpdates(FRHIComputeCommandLis
 			{
 				TRACE_CPUPROFILER_EVENT_SCOPE(SetupSegmentData);
 
+				const uint64 ScratchAlignment = GRHIRayTracingAccelerationStructureAlignment;
+				uint32 ScratchBLASOffset = 0;
+
 				// Setup the array views on final allocated segments array
 				FRayTracingGeometrySegment* SegmentData = Segments.GetData();
 				for (FRayTracingGeometryBuildParams& Param : BuildParams)
@@ -371,6 +374,15 @@ void FRayTracingDynamicGeometryCollection::DispatchUpdates(FRHIComputeCommandLis
 						Param.Segments = MakeArrayView(SegmentData, SegmentCount);
 						SegmentData += SegmentCount;
 					}
+
+					Param.ScratchBuffer = ScratchBuffer;
+					Param.ScratchBufferOffset = ScratchBLASOffset;
+
+					// Update the offset
+
+					const FRayTracingAccelerationStructureSize BLASSizeInfo = Param.Geometry->GetSizeInfo();
+					const uint64 ScratchSize = Param.BuildMode == EAccelerationStructureBuildMode::Build ? BLASSizeInfo.BuildScratchSize : BLASSizeInfo.UpdateScratchSize;
+					ScratchBLASOffset = Align(ScratchBLASOffset + ScratchSize, ScratchAlignment);
 				}
 			}
 
@@ -515,6 +527,24 @@ void FRayTracingDynamicGeometryCollection::EndUpdate(FRHICommandListImmediate& R
 	RHICmdList.EnqueueLambda([ArrayOwnedByRHIThread = MoveTemp(Segments)](FRHICommandListImmediate&){});
 }
 
+uint32 FRayTracingDynamicGeometryCollection::ComputeScratchBufferSize()
+{
+	TRACE_CPUPROFILER_EVENT_SCOPE(FRayTracingDynamicGeometryCollection::ComputeScratchBufferSize);
+
+	const uint64 ScratchAlignment = GRHIRayTracingAccelerationStructureAlignment;
+
+	uint32 BLASScratchSize = 0;
+
+	for (FRayTracingGeometryBuildParams& Params : BuildParams)
+	{
+		const FRayTracingAccelerationStructureSize BLASSizeInfo = Params.Geometry->GetSizeInfo();
+		const uint64 ScratchSize = Params.BuildMode == EAccelerationStructureBuildMode::Build ? BLASSizeInfo.BuildScratchSize : BLASSizeInfo.UpdateScratchSize;
+		BLASScratchSize = Align(BLASScratchSize + ScratchSize, ScratchAlignment);
+	}
+
+	return BLASScratchSize;
+}
+
 #undef USE_RAY_TRACING_DYNAMIC_GEOMETRY_PARALLEL_COMMAND_LISTS
 
 #endif // RHI_RAYTRACING
diff --git a/Engine/Source/Runtime/Renderer/Public/RayTracingDynamicGeometryCollection.h b/Engine/Source/Runtime/Renderer/Public/RayTracingDynamicGeometryCollection.h
--- a/Engine/Source/Runtime/Renderer/Public/RayTracingDynamicGeometryCollection.h
+++ b/Engine/Source/Runtime/Renderer/Public/RayTracingDynamicGeometryCollection.h
@@ -20,10 +20,13 @@ public:
 
 	// Starts an update batch and returns the current shared buffer generation ID which is used for validation.
 	int64 BeginUpdate();
-	void DispatchUpdates(FRHIComputeCommandList& ParentCmdList);
+	void DispatchUpdates(FRHIComputeCommandList& ParentCmdList, FRHIBuffer* ScratchBuffer);
 	void EndUpdate(FRHICommandListImmediate& RHICmdList);
 
+	uint32 ComputeScratchBufferSize();
+
 private:
+
 	TArray<struct FMeshComputeDispatchCommand> DispatchCommands;
 	TArray<FRayTracingGeometryBuildParams> BuildParams;
 	TArray<FRayTracingGeometrySegment> Segments;
diff --git a/Engine/Source/Runtime/VulkanRHI/Private/VulkanRayTracing.cpp b/Engine/Source/Runtime/VulkanRHI/Private/VulkanRayTracing.cpp
--- a/Engine/Source/Runtime/VulkanRHI/Private/VulkanRayTracing.cpp
+++ b/Engine/Source/Runtime/VulkanRHI/Private/VulkanRayTracing.cpp
@@ -325,6 +325,10 @@ FVulkanRayTracingGeometry::FVulkanRayTracingGeometry(FRayTracingGeometryInitiali
 	ZeroVulkanStruct(DeviceAddressInfo, VK_STRUCTURE_TYPE_ACCELERATION_STRUCTURE_DEVICE_ADDRESS_INFO_KHR);
 	DeviceAddressInfo.accelerationStructure = Handle;
 	Address = vkGetAccelerationStructureDeviceAddressKHR(NativeDevice, &DeviceAddressInfo);
+	
+	SizeInfo.ResultSize = BuildData.SizesInfo.accelerationStructureSize;
+	SizeInfo.BuildScratchSize = BuildData.SizesInfo.buildScratchSize;
+	SizeInfo.UpdateScratchSize = BuildData.SizesInfo.updateScratchSize;
 }
 
 FVulkanRayTracingGeometry::~FVulkanRayTracingGeometry()
diff --git a/Engine/Source/Runtime/VulkanRHI/Private/VulkanRayTracing.h b/Engine/Source/Runtime/VulkanRHI/Private/VulkanRayTracing.h
--- a/Engine/Source/Runtime/VulkanRHI/Private/VulkanRayTracing.h
+++ b/Engine/Source/Runtime/VulkanRHI/Private/VulkanRayTracing.h
@@ -85,6 +85,7 @@ public:
 
 	virtual FRayTracingAccelerationStructureAddress GetAccelerationStructureAddress(uint64 GPUIndex) const final override { return Address; }
 	virtual uint32 GetNumSegments() const final override { return Initializer.Segments.Num(); }
+	virtual FRayTracingAccelerationStructureSize GetSizeInfo() const final override { return SizeInfo; }
 
 	void BuildAccelerationStructure(FVulkanCommandListContext& CommandContext, EAccelerationStructureBuildMode BuildMode);
 
@@ -97,6 +98,8 @@ private:
 	VkDeviceAddress Address = 0;
 	TRefCountPtr<FVulkanResourceMultiBuffer> AccelerationStructureBuffer;
 	TRefCountPtr<FVulkanResourceMultiBuffer> ScratchBuffer;
+	
+	FRayTracingAccelerationStructureSize SizeInfo = {};
 };
 
 class FVulkanRayTracingScene : public FRHIRayTracingScene
