Don't start waiting on produce tasks if we know that there is nothing to wait for.

diff --git a/Engine/Source/Runtime/Renderer/Private/VT/VirtualTextureSystem.cpp b/Engine/Source/Runtime/Renderer/Private/VT/VirtualTextureSystem.cpp
--- a/Engine/Source/Runtime/Renderer/Private/VT/VirtualTextureSystem.cpp
+++ b/Engine/Source/Runtime/Renderer/Private/VT/VirtualTextureSystem.cpp
@@ -2224,7 +2224,9 @@ void FVirtualTextureSystem::SubmitRequests(FRDGBuilder& GraphBuilder, ERHIFeatur
 		TArray<FProducePageDataPrepareTask> PrepareTasks;
 		PrepareTasks.Reserve(RequestList->GetNumLoadRequests());
 
+		// If we render a frame without all locked tiles loaded, may render garbage VT data, as there won't be low mip fallback for unloaded tiles
 		const bool bSyncProduceLockedTiles = CVarVTSyncProduceLockedTiles.GetValueOnRenderThread() != 0;
+		bool bWaitForProducers = false;
 
 		const uint32 MaxPagesProduced = VirtualTextureScalability::GetMaxPagesProducedPerFrame();
 		const uint32 PageFreeThreshold = VirtualTextureScalability::GetPageFreeThreshold();
@@ -2252,12 +2254,11 @@ void FVirtualTextureSystem::SubmitRequests(FRDGBuilder& GraphBuilder, ERHIFeatur
 
 			const EVTRequestPagePriority Priority = bLockTile ? EVTRequestPagePriority::High : EVTRequestPagePriority::Normal;
 			FVTRequestPageResult RequestPageResult = Producer.GetVirtualTexture()->RequestPageData(ProducerHandle, ProducerTextureLayerMask, TileToLoad.Local_vLevel, TileToLoad.Local_vAddress, Priority);
-			//checkf(!bLockTile || RequestPageResult.Status != EVTRequestPageStatus::Invalid, TEXT("Tried to lock an invalid VT tile"));
 			if (RequestPageResult.Status == EVTRequestPageStatus::Pending && bForceProduceTile)
 			{
-				// If we're trying to lock this tile, we're OK producing data now (and possibly waiting) as long as data is pending
-				// If we render a frame without all locked tiles loaded, may render garbage VT data, as there won't be low mip fallback for unloaded tiles
+				// If we're forcing production of this tile, we're OK producing data now (and possibly waiting) as long as data is pending
 				RequestPageResult.Status = EVTRequestPageStatus::Available;
+				bWaitForProducers = true;
 			}
 
 			if (RequestPageResult.Status == EVTRequestPageStatus::Available && !bForceProduceTile && NumPagesProduced >= MaxPagesProduced)
@@ -2270,6 +2271,8 @@ void FVirtualTextureSystem::SubmitRequests(FRDGBuilder& GraphBuilder, ERHIFeatur
 			bool bTileInvalid = false;
 			if (RequestPageResult.Status == EVTRequestPageStatus::Invalid)
 			{
+				//checkf(!bLockTile, TEXT("Tried to lock an invalid VT tile"));
+
 				bTileInvalid = true;
 				if (CVarVTVerbose.GetValueOnRenderThread())
 				{
@@ -2400,20 +2403,23 @@ void FVirtualTextureSystem::SubmitRequests(FRDGBuilder& GraphBuilder, ERHIFeatur
 
 		if (PrepareTasks.Num())
 		{
-			FGraphEventArray ProducePageTasks;
-			ProducePageTasks.Reserve(PrepareTasks.Num());
-
-			for (FProducePageDataPrepareTask& Task : PrepareTasks)
-			{
-				Task.VirtualTexture->GatherProducePageDataTasks(Task.RequestHandle, ProducePageTasks);
-			}
-
 			static bool bWaitForTasks = true;
-			if (bWaitForTasks)
+			if (bWaitForProducers && bWaitForTasks)
 			{
-				TRACE_CPUPROFILER_EVENT_SCOPE(FVirtualTextureSystem::ProcessRequests_Wait);
-				QUICK_SCOPE_CYCLE_COUNTER(ProcessRequests_Wait);
-				FTaskGraphInterface::Get().WaitUntilTasksComplete(ProducePageTasks, ENamedThreads::GetRenderThread_Local());
+				// Wait for all producers here instead of inside each individual call to ProducePageData()
+				FGraphEventArray ProducePageTasks;
+				ProducePageTasks.Reserve(PrepareTasks.Num());
+
+				for (FProducePageDataPrepareTask& Task : PrepareTasks)
+				{
+					Task.VirtualTexture->GatherProducePageDataTasks(Task.RequestHandle, ProducePageTasks);
+				}
+
+				{
+					TRACE_CPUPROFILER_EVENT_SCOPE(FVirtualTextureSystem::ProcessRequests_Wait);
+					QUICK_SCOPE_CYCLE_COUNTER(ProcessRequests_Wait);
+					FTaskGraphInterface::Get().WaitUntilTasksComplete(ProducePageTasks, ENamedThreads::GetRenderThread_Local());
+				}
 			}
 
 			for (FProducePageDataPrepareTask& Task : PrepareTasks)
