Adding CSV stat for capturing RemovePrimitiveSceneInfos to match the existing AddPrimitiveSceneInfos, UpdatePrimitiveSceneInfos, and UpdatePrimitiveInstances csv markers.

diff --git a/Engine/Source/Runtime/Renderer/Private/RendererScene.cpp b/Engine/Source/Runtime/Renderer/Private/RendererScene.cpp
--- a/Engine/Source/Runtime/Renderer/Private/RendererScene.cpp
+++ b/Engine/Source/Runtime/Renderer/Private/RendererScene.cpp
@@ -4112,6 +4112,7 @@ void FScene::UpdateAllPrimitiveSceneInfos(FRDGBuilder& GraphBuilder, bool bAsync
 
 	GPUScene.ResizeDirtyState(Primitives.Num());
 	{
+		CSV_SCOPED_TIMING_STAT_EXCLUSIVE(RemovePrimitiveSceneInfos);
 		SCOPED_NAMED_EVENT(FScene_RemovePrimitiveSceneInfos, FColor::Red);
 		SCOPE_CYCLE_COUNTER(STAT_RemoveScenePrimitiveTime);
 		for (FPrimitiveSceneInfo* PrimitiveSceneInfo : RemovedLocalPrimitiveSceneInfos)
