Always sync locked VT tiles in Editor.
This waits correctly for DDC.
Fixes nightly render issues.

diff --git a/Engine/Source/Runtime/Renderer/Private/VT/VirtualTextureSystem.cpp b/Engine/Source/Runtime/Renderer/Private/VT/VirtualTextureSystem.cpp
--- a/Engine/Source/Runtime/Renderer/Private/VT/VirtualTextureSystem.cpp
+++ b/Engine/Source/Runtime/Renderer/Private/VT/VirtualTextureSystem.cpp
@@ -2225,7 +2225,12 @@ void FVirtualTextureSystem::SubmitRequests(FRDGBuilder& GraphBuilder, ERHIFeatur
 		PrepareTasks.Reserve(RequestList->GetNumLoadRequests());
 
 		// If we render a frame without all locked tiles loaded, may render garbage VT data, as there won't be low mip fallback for unloaded tiles
+#if WITH_EDITOR
+		const bool bSyncProduceLockedTiles = true;
+#else
 		const bool bSyncProduceLockedTiles = CVarVTSyncProduceLockedTiles.GetValueOnRenderThread() != 0;
+#endif
+
 		bool bWaitForProducers = false;
 
 		const uint32 MaxPagesProduced = VirtualTextureScalability::GetMaxPagesProducedPerFrame();
