Removes PRAGMA_DISABLE_OPTIMIZATION in TSR

diff --git a/Engine/Source/Runtime/Renderer/Private/PostProcess/TemporalSuperResolution.cpp b/Engine/Source/Runtime/Renderer/Private/PostProcess/TemporalSuperResolution.cpp
--- a/Engine/Source/Runtime/Renderer/Private/PostProcess/TemporalSuperResolution.cpp
+++ b/Engine/Source/Runtime/Renderer/Private/PostProcess/TemporalSuperResolution.cpp
@@ -11,8 +11,6 @@
 #include "ScenePrivate.h"
 #include "RendererModule.h"
 
-PRAGMA_DISABLE_OPTIMIZATION
-
 #define COMPILE_TSR_DEBUG_PASSES (!UE_BUILD_SHIPPING)
 
 namespace
