Fix inverted condition in VSM causing artifacts with dynamic/static separate pages

diff --git a/Engine/Source/Runtime/Renderer/Private/VirtualShadowMaps/VirtualShadowMapArray.cpp b/Engine/Source/Runtime/Renderer/Private/VirtualShadowMaps/VirtualShadowMapArray.cpp
--- a/Engine/Source/Runtime/Renderer/Private/VirtualShadowMaps/VirtualShadowMapArray.cpp
+++ b/Engine/Source/Runtime/Renderer/Private/VirtualShadowMaps/VirtualShadowMapArray.cpp
@@ -393,7 +393,7 @@ FIntPoint FVirtualShadowMapArray::GetDynamicPhysicalPoolSize() const
 uint32 FVirtualShadowMapArray::GetTotalAllocatedPhysicalPages() const
 {
 	check(bInitialized);
-	return ShouldCacheStaticSeparately() ? UniformParameters.MaxPhysicalPages : (2U * UniformParameters.MaxPhysicalPages);
+	return ShouldCacheStaticSeparately() ? (2U * UniformParameters.MaxPhysicalPages) : UniformParameters.MaxPhysicalPages;
 }
 
 
