Improve load-balancing of the VSM invalidations, using the existing load balancer for instances.
 - saves ~0.5ms GPU time in many views (sandbox).

diff --git a/Engine/Shaders/Private/VirtualShadowMaps/CacheManagement.usf b/Engine/Shaders/Private/VirtualShadowMaps/CacheManagement.usf
--- a/Engine/Shaders/Private/VirtualShadowMaps/CacheManagement.usf
+++ b/Engine/Shaders/Private/VirtualShadowMaps/CacheManagement.usf
@@ -15,13 +15,9 @@
 #include "../ColorMap.ush"
 #endif
 
-struct FInstanceSceneDataRange
-{
-	int InstanceSceneDataOffset;
-	int NumInstanceSceneDataEntries;
-};
-StructuredBuffer<FInstanceSceneDataRange> InstanceSceneRanges;
-uint NumRemovedItems;
+#if INPUT_KIND == INPUT_KIND_LOAD_BALANCER
+#include "../InstanceCulling/InstanceCullingLoadBalancer.ush"
+#endif //
 
 RWStructuredBuffer<uint> OutDynamicCasterPageFlags;
 uint InstanceSceneDataSOAStride;
@@ -42,20 +38,18 @@ uint NumInvalidatingInstanceSlots;
 [numthreads(CS_1D_GROUP_SIZE_X, 1, 1)]
 void VirtualSmInvalidateInstancePagesCS(
 	uint DispatchIndex : SV_DispatchThreadID,
-	uint GroupId : SV_GroupID,
+	uint3 GroupId : SV_GroupID,
 	uint GroupThreadIndex : SV_GroupIndex)
 {
-#if INPUT_KIND == INPUT_KIND_LARGE_RANGES
-	// Large instance count ranges use one group per range instead, and use a group-stride loop
-	FInstanceSceneDataRange InstanceSceneDataRange = InstanceSceneRanges[GroupId];
-	for (int InstanceId = InstanceSceneDataRange.InstanceSceneDataOffset + GroupThreadIndex; InstanceId < InstanceSceneDataRange.NumInstanceSceneDataEntries + InstanceSceneDataRange.InstanceSceneDataOffset; InstanceId += CS_1D_GROUP_SIZE_X)
-#elif INPUT_KIND == INPUT_KIND_SMALL_RANGES
-	if (DispatchIndex >= NumRemovedItems)
+#if INPUT_KIND == INPUT_KIND_LOAD_BALANCER
+	FInstanceWorkSetup WorkSetup = InstanceCullingLoadBalancer_Setup(GroupId, GroupThreadIndex, 0U);
+
+	if (!WorkSetup.bValid)
 	{
 		return;
 	}
-	FInstanceSceneDataRange InstanceSceneDataRange = InstanceSceneRanges[DispatchIndex];
-	for (int InstanceId = InstanceSceneDataRange.InstanceSceneDataOffset; InstanceId < InstanceSceneDataRange.NumInstanceSceneDataEntries + InstanceSceneDataRange.InstanceSceneDataOffset; ++InstanceId)
+
+	uint InstanceId = WorkSetup.Item.InstanceDataOffset + uint(WorkSetup.LocalItemIndex);
 #else // INPUT_KIND == INPUT_KIND_GPU_INSTANCES
 	// The 0th index stores the total number appended
 	if (DispatchIndex >= InvalidatingInstances[0])
diff --git a/Engine/Source/Runtime/Renderer/Private/VirtualShadowMaps/VirtualShadowMapCacheManager.cpp b/Engine/Source/Runtime/Renderer/Private/VirtualShadowMaps/VirtualShadowMapCacheManager.cpp
--- a/Engine/Source/Runtime/Renderer/Private/VirtualShadowMaps/VirtualShadowMapCacheManager.cpp
+++ b/Engine/Source/Runtime/Renderer/Private/VirtualShadowMaps/VirtualShadowMapCacheManager.cpp
@@ -443,8 +443,7 @@ void FVirtualShadowMapArrayCacheManager::ProcessRemovedPrimives(FRDGBuilder& Gra
 	{
 		RDG_EVENT_SCOPE(GraphBuilder, "Shadow.Virtual.ProcessRemovedPrimives [%d]", RemovedPrimitiveSceneInfos.Num());
 
-		TArray<FInstanceSceneDataRange, SceneRenderingAllocator> InstanceRangesLarge;
-		TArray<FInstanceSceneDataRange, SceneRenderingAllocator> InstanceRangesSmall;
+		FInstanceGPULoadBalancer LoadBalancer;
 
 		int32 TotalInstanceCount = 0;
 		for (const FPrimitiveSceneInfo* PrimitiveSceneInfo : RemovedPrimitiveSceneInfos)
@@ -452,18 +451,11 @@ void FVirtualShadowMapArrayCacheManager::ProcessRemovedPrimives(FRDGBuilder& Gra
 			if (PrimitiveSceneInfo->GetInstanceSceneDataOffset() != INDEX_NONE)
 			{
 				const int32 NumInstanceSceneDataEntries = PrimitiveSceneInfo->GetNumInstanceSceneDataEntries();
+				LoadBalancer.Add(PrimitiveSceneInfo->GetInstanceSceneDataOffset(), NumInstanceSceneDataEntries, 0U);
 				TotalInstanceCount += NumInstanceSceneDataEntries;
-				if (NumInstanceSceneDataEntries >= 8)
-				{
-					InstanceRangesLarge.Add(FInstanceSceneDataRange{ PrimitiveSceneInfo->GetInstanceSceneDataOffset(), NumInstanceSceneDataEntries });
-				}
-				else
-				{
-					InstanceRangesSmall.Add(FInstanceSceneDataRange{ PrimitiveSceneInfo->GetInstanceSceneDataOffset(), NumInstanceSceneDataEntries });
-				}
 			}
 		}
-		ProcessInstanceRangeInvalidation(GraphBuilder, InstanceRangesLarge, InstanceRangesSmall, TotalInstanceCount, GPUScene);
+		ProcessInvalidations(GraphBuilder, LoadBalancer, TotalInstanceCount, GPUScene);
 	}
 }
 
@@ -479,9 +471,7 @@ void FVirtualShadowMapArrayCacheManager::ProcessPrimitivesToUpdate(FRDGBuilder&
 
 		if (GPUScene.PrimitivesToUpdate.Num() > 0)
 		{
-			// TODO: As a slight CPU optimization just pass primitive ID list and use instance ranges stored in GPU scene
-			TArray<FInstanceSceneDataRange, SceneRenderingAllocator> InstanceRangesLarge;
-			TArray<FInstanceSceneDataRange, SceneRenderingAllocator> InstanceRangesSmall;
+			FInstanceGPULoadBalancer LoadBalancer;
 
 			int32 TotalInstanceCount = 0;
 			for (const int32 PrimitiveId : GPUScene.PrimitivesToUpdate)
@@ -499,21 +489,14 @@ void FVirtualShadowMapArrayCacheManager::ProcessPrimitivesToUpdate(FRDGBuilder&
 						const FPrimitiveSceneInfo* PrimitiveSceneInfo = Scene.Primitives[PrimitiveId];
 						if (PrimitiveSceneInfo->GetInstanceSceneDataOffset() != INDEX_NONE)
 						{
-							int32 NumInstanceSceneDataEntries = PrimitiveSceneInfo->GetNumInstanceSceneDataEntries();
+							const int32 NumInstanceSceneDataEntries = PrimitiveSceneInfo->GetNumInstanceSceneDataEntries();
+							LoadBalancer.Add(PrimitiveSceneInfo->GetInstanceSceneDataOffset(), NumInstanceSceneDataEntries, PrimitiveId);
 							TotalInstanceCount += NumInstanceSceneDataEntries;
-							if (NumInstanceSceneDataEntries >= 8)
-							{
-								InstanceRangesLarge.Add(FInstanceSceneDataRange{ PrimitiveSceneInfo->GetInstanceSceneDataOffset(), NumInstanceSceneDataEntries });
-							}
-							else
-							{
-								InstanceRangesSmall.Add(FInstanceSceneDataRange{ PrimitiveSceneInfo->GetInstanceSceneDataOffset(), NumInstanceSceneDataEntries });
-							}
 						}
 					}
 				}
 			}
-			ProcessInstanceRangeInvalidation(GraphBuilder, InstanceRangesLarge, InstanceRangesSmall, TotalInstanceCount, GPUScene);
+			ProcessInvalidations(GraphBuilder, LoadBalancer, TotalInstanceCount, GPUScene);
 		}
 	}
 }
@@ -528,9 +511,8 @@ class FVirtualSmInvalidateInstancePagesCS : public FGlobalShader
 
 	enum EInputDataKind
 	{
-		EInputKind_SmallRanges,
-		EInputKind_LargeRanges,
 		EInputKind_GPUInstances,
+		EInputKind_LoadBalancer,
 		EInputKind_Num
 	};
 
@@ -545,8 +527,6 @@ public:
 		SHADER_PARAMETER_STRUCT_INCLUDE(ShaderDrawDebug::FShaderParameters, ShaderDrawUniformBuffer)
 		SHADER_PARAMETER(uint32, bDrawBounds)
 
-		SHADER_PARAMETER_RDG_BUFFER_SRV(StructuredBuffer< FInstanceSceneDataRange >, InstanceSceneRanges)
-		SHADER_PARAMETER(uint32, NumRemovedItems)
 		SHADER_PARAMETER_RDG_BUFFER_SRV(StructuredBuffer< uint >, PageFlags)
 		SHADER_PARAMETER_RDG_BUFFER_SRV(StructuredBuffer< uint >, HPageFlags)
 		SHADER_PARAMETER_RDG_BUFFER_SRV(StructuredBuffer< uint4 >, PageRectBounds)
@@ -566,9 +546,11 @@ public:
 		SHADER_PARAMETER_RDG_BUFFER_SRV(StructuredBuffer<uint>, InvalidatingInstances)
 		SHADER_PARAMETER(uint32, NumInvalidatingInstanceSlots)
 		RDG_BUFFER_ACCESS(IndirectArgs, ERHIAccess::IndirectArgs)
+
+		SHADER_PARAMETER_STRUCT_INCLUDE(FGPUScene::FInstanceGPULoadBalancer::FShaderParameters, LoadBalancerParameters)
 	END_SHADER_PARAMETER_STRUCT()
 
-	static constexpr int Cs1dGroupSizeX = 64;
+	static constexpr int Cs1dGroupSizeX = FVirtualShadowMapArrayCacheManager::FInstanceGPULoadBalancer::ThreadGroupSize;
 
 	static bool ShouldCompilePermutation(const FGlobalShaderPermutationParameters& Parameters)
 	{
@@ -583,9 +565,10 @@ public:
 		OutEnvironment.SetDefine(TEXT("USE_GLOBAL_GPU_SCENE_DATA"), 1);
 		OutEnvironment.SetDefine(TEXT("VF_SUPPORTS_PRIMITIVE_SCENE_DATA"), 1);
 		
-		OutEnvironment.SetDefine(TEXT("INPUT_KIND_SMALL_RANGES"), EInputKind_SmallRanges);
-		OutEnvironment.SetDefine(TEXT("INPUT_KIND_LARGE_RANGES"), EInputKind_LargeRanges);
 		OutEnvironment.SetDefine(TEXT("INPUT_KIND_GPU_INSTANCES"), EInputKind_GPUInstances);
+		OutEnvironment.SetDefine(TEXT("INPUT_KIND_LOAD_BALANCER"), EInputKind_LoadBalancer);
+
+		FGPUScene::FInstanceGPULoadBalancer::SetShaderDefines(OutEnvironment);
 	}
 };
 IMPLEMENT_GLOBAL_SHADER(FVirtualSmInvalidateInstancePagesCS, "/Engine/Private/VirtualShadowMaps/CacheManagement.usf", "VirtualSmInvalidateInstancePagesCS", SF_Compute);
@@ -675,63 +658,36 @@ static void SetupCommonParameters(FRDGBuilder& GraphBuilder, FVirtualShadowMapAr
 	OutPermutationVector.Set<FVirtualSmInvalidateInstancePagesCS::FUseHzbDim>(PrevHZBPhysical != nullptr);
 }
 
-void FVirtualShadowMapArrayCacheManager::ProcessInstanceRangeInvalidation(FRDGBuilder& GraphBuilder, const TArray<FInstanceSceneDataRange, SceneRenderingAllocator>& InstanceRangesLarge, const TArray<FInstanceSceneDataRange, SceneRenderingAllocator>& InstanceRangesSmall, int32 TotalInstanceCount, const FGPUScene& GPUScene)
+
+void FVirtualShadowMapArrayCacheManager::ProcessInvalidations(FRDGBuilder& GraphBuilder, FInstanceGPULoadBalancer& Instances, int32 TotalInstanceCount, const FGPUScene& GPUScene)
 {
-	if (InstanceRangesSmall.IsEmpty() && InstanceRangesLarge.IsEmpty())
+	if (Instances.IsEmpty())
 	{
 		return;
 	}
 
-	FRDGBufferRef InstanceRangesSmallRDG = !InstanceRangesSmall.IsEmpty() ? CreateStructuredBuffer(GraphBuilder, TEXT("Shadow.Virtual.InstanceRangesSmall"), InstanceRangesSmall) : nullptr;
-	FRDGBufferRef InstanceRangesLargeRDG = !InstanceRangesLarge.IsEmpty() ? CreateStructuredBuffer(GraphBuilder, TEXT("Shadow.Virtual.InstanceRangesSmall"), InstanceRangesLarge) : nullptr;
+	Instances.FinalizeBatches();
 
+	RDG_EVENT_SCOPE(GraphBuilder, "ProcessInvalidations [%d batches]", Instances.GetBatches().Num());
 
-	FVirtualSmInvalidateInstancePagesCS::FParameters PassParametersTmp;
-	FVirtualSmInvalidateInstancePagesCS::FPermutationDomain PermutationVector;
-	SetupCommonParameters(GraphBuilder, this, TotalInstanceCount, GPUScene, PassParametersTmp, PermutationVector);
+	FVirtualSmInvalidateInstancePagesCS::FParameters* PassParameters = GraphBuilder.AllocParameters<FVirtualSmInvalidateInstancePagesCS::FParameters>();
 
-	if (InstanceRangesSmall.Num())
-	{
-		RDG_EVENT_SCOPE(GraphBuilder, "ProcessInstanceRangeInvalidation [%d small-ranges]", InstanceRangesSmall.Num());
-
-		FVirtualSmInvalidateInstancePagesCS::FParameters* PassParameters = GraphBuilder.AllocParameters<FVirtualSmInvalidateInstancePagesCS::FParameters>();
-		*PassParameters = PassParametersTmp;
-		PassParameters->InstanceSceneRanges = GraphBuilder.CreateSRV(InstanceRangesSmallRDG);
-		PassParameters->NumRemovedItems = InstanceRangesSmall.Num();
-
-		PermutationVector.Set<FVirtualSmInvalidateInstancePagesCS::FInputKindDim>(FVirtualSmInvalidateInstancePagesCS::EInputKind_SmallRanges);
-
-		auto ComputeShader = GetGlobalShaderMap(GMaxRHIFeatureLevel)->GetShader<FVirtualSmInvalidateInstancePagesCS>(PermutationVector);
-
-		FComputeShaderUtils::AddPass(
-			GraphBuilder,
-			RDG_EVENT_NAME("VirtualSmInvalidateInstancePagesCS"),
-			ComputeShader,
-			PassParameters,
-			FIntVector(FMath::DivideAndRoundUp(InstanceRangesSmall.Num(), FVirtualSmInvalidateInstancePagesCS::Cs1dGroupSizeX), 1, 1)
-		);
-	}
-	if (InstanceRangesLarge.Num())
-	{
-		RDG_EVENT_SCOPE(GraphBuilder, "ProcessInstanceRangeInvalidation [%d large-ranges]", InstanceRangesLarge.Num());
+	FVirtualSmInvalidateInstancePagesCS::FPermutationDomain PermutationVector;
+	SetupCommonParameters(GraphBuilder, this, TotalInstanceCount, GPUScene, *PassParameters, PermutationVector);
+	Instances.Upload(GraphBuilder).GetShaderParameters(GraphBuilder, PassParameters->LoadBalancerParameters);
 
-		FVirtualSmInvalidateInstancePagesCS::FParameters* PassParameters = GraphBuilder.AllocParameters<FVirtualSmInvalidateInstancePagesCS::FParameters>();
-		*PassParameters = PassParametersTmp;
-		PassParameters->InstanceSceneRanges = GraphBuilder.CreateSRV(InstanceRangesLargeRDG);
-		PassParameters->NumRemovedItems = InstanceRangesLarge.Num();
+	PermutationVector.Set<FVirtualSmInvalidateInstancePagesCS::FInputKindDim>(FVirtualSmInvalidateInstancePagesCS::EInputKind_LoadBalancer);
 
-		PermutationVector.Set<FVirtualSmInvalidateInstancePagesCS::FInputKindDim>(FVirtualSmInvalidateInstancePagesCS::EInputKind_LargeRanges);
+	auto ComputeShader = GetGlobalShaderMap(GMaxRHIFeatureLevel)->GetShader<FVirtualSmInvalidateInstancePagesCS>(PermutationVector);
 
-		auto ComputeShader = GetGlobalShaderMap(GMaxRHIFeatureLevel)->GetShader<FVirtualSmInvalidateInstancePagesCS>(PermutationVector);
+	FComputeShaderUtils::AddPass(
+		GraphBuilder,
+		RDG_EVENT_NAME("VirtualSmInvalidateInstancePagesCS"),
+		ComputeShader,
+		PassParameters,
+		Instances.GetWrappedCsGroupCount()
+	);
 
-		FComputeShaderUtils::AddPass(
-			GraphBuilder,
-			RDG_EVENT_NAME("VirtualSmInvalidateInstancePagesCS"),
-			ComputeShader,
-			PassParameters,
-			FIntVector(InstanceRangesLarge.Num(), 1, 1)
-		);
-	}
 }
 
 
@@ -739,7 +695,7 @@ void FVirtualShadowMapArrayCacheManager::ProcessGPUInstanceInvalidations(FRDGBui
 {
 	// Dispatch CS indirectly to process instances that are marked to update from the GPU side.
 	{
-		RDG_EVENT_SCOPE(GraphBuilder, "ProcessInstanceRangeInvalidation [GPU-Instances]");
+		RDG_EVENT_SCOPE(GraphBuilder, "ProcessGPUInstanceInvalidations [GPU-Instances]");
 
 		FRDGBufferRef InvalidatingInstancesBufferRDG = GraphBuilder.RegisterExternalBuffer(PrevBuffers.InvalidatingInstancesBuffer, TEXT("Shadow.Virtual.PrevInvalidatingInstancesBuffer"));
 		FRDGBufferRef IndirectArgs = FComputeShaderUtils::AddIndirectArgsSetupCsPass1D(GraphBuilder, InvalidatingInstancesBufferRDG, TEXT("Shadow.Virtual.ProcessGPUInstanceInvalidationsIndirectArgs"), FVirtualSmInvalidateInstancePagesCS::Cs1dGroupSizeX);
diff --git a/Engine/Source/Runtime/Renderer/Private/VirtualShadowMaps/VirtualShadowMapCacheManager.h b/Engine/Source/Runtime/Renderer/Private/VirtualShadowMaps/VirtualShadowMapCacheManager.h
--- a/Engine/Source/Runtime/Renderer/Private/VirtualShadowMaps/VirtualShadowMapCacheManager.h
+++ b/Engine/Source/Runtime/Renderer/Private/VirtualShadowMaps/VirtualShadowMapCacheManager.h
@@ -8,6 +8,7 @@
 #include "CoreMinimal.h"
 #include "VirtualShadowMapArray.h"
 #include "SceneManagement.h"
+#include "InstanceCulling/InstanceCullingLoadBalancer.h"
 
 class FRHIGPUBufferReadback;
 class FGPUScene;
@@ -128,20 +129,10 @@ public:
 		
 	void SetHZBViewParams(int32 HZBKey, Nanite::FPackedViewParams& OutParams);
 
-private:
-	// Must match shader...
-	struct FInstanceSceneDataRange
-	{
-		int32 InstanceSceneDataOffset;
-		int32 NumInstanceSceneDataEntries;
-	};
+	using FInstanceGPULoadBalancer = TInstanceCullingLoadBalancer<SceneRenderingAllocator>;
 
-	void ProcessInstanceRangeInvalidation(
-		FRDGBuilder& GraphBuilder,
-		const TArray<FInstanceSceneDataRange, SceneRenderingAllocator>& InstanceRangesLarge,
-		const TArray<FInstanceSceneDataRange, SceneRenderingAllocator>& InstanceRangesSmall,
-		int32 TotalInstanceCount,
-		const FGPUScene& GPUScene);
+private:
+	void ProcessInvalidations(FRDGBuilder& GraphBuilder, FInstanceGPULoadBalancer& Instances, int32 TotalInstanceCount, const FGPUScene& GPUScene);
 
 	void ProcessGPUInstanceInvalidations(FRDGBuilder& GraphBuilder, const FGPUScene& GPUScene);
 
