Optimize GPU-Scene to not track invalid instances on an individual basis and offload clearing to GPU

diff --git a/Engine/Source/Runtime/Renderer/Private/GPUScene.cpp b/Engine/Source/Runtime/Renderer/Private/GPUScene.cpp
--- a/Engine/Source/Runtime/Renderer/Private/GPUScene.cpp
+++ b/Engine/Source/Runtime/Renderer/Private/GPUScene.cpp
@@ -490,14 +490,9 @@ void FGPUScene::UpdateInternal(FRDGBuilder& GraphBuilder, FScene& Scene)
 			PrimitivesToUpdate.Add(Index);
 		}
 
-		// Clear the full instance data range, except primitives that use a slot (they will unset the bits).
-		InstanceSceneDataToClear.Init(true, InstanceSceneDataToClear.Num());
-
-		// Set entire instance range for possible clearing.
-		for (int32 Index = 0; Index < InstanceSceneDataToClear.Num(); ++Index)
-		{
-			InstanceSceneDataClearList.Add(Index);
-		}
+		// Clear the full instance data range
+		InstanceRangesToClear.Empty();
+		InstanceRangesToClear.Add(FInstanceRange{ 0U, uint32(GetNumInstances()) });
 
 		bUpdateAllPrimitives = false;
 	}
@@ -508,6 +503,9 @@ void FGPUScene::UpdateInternal(FRDGBuilder& GraphBuilder, FScene& Scene)
 	FUploadDataSourceAdapterScenePrimitives Adapter(Scene, SceneFrameNumber, MoveTemp(PrimitivesToUpdate), MoveTemp(PrimitiveDirtyState));
 	FGPUSceneBufferState BufferState = UpdateBufferState(GraphBuilder, &Scene, Adapter);
 
+	// Run a pass that clears (Sets ID to invalid) any instances that need it
+	AddClearInstancesPass(GraphBuilder);
+
 	// Pull out instances needing only primitive ID update, they still have to go to the general update such that the primitive gets updated (as it moved)
 	{
 		FInstanceGPULoadBalancer IdOnlyUpdateData;
@@ -520,7 +518,6 @@ void FGPUScene::UpdateInternal(FRDGBuilder& GraphBuilder, FScene& Scene)
 				const FPrimitiveSceneInfo* PrimitiveSceneInfo = Scene.Primitives[PrimitiveId];
 				if (PrimitiveSceneInfo->GetInstanceSceneDataOffset() >= 0)
 				{
-					InstanceSceneDataToClear.SetRange(PrimitiveSceneInfo->GetInstanceSceneDataOffset(), PrimitiveSceneInfo->GetNumInstanceSceneDataEntries(), false);
 					IdOnlyUpdateData.Add(PrimitiveSceneInfo->GetInstanceSceneDataOffset(), PrimitiveSceneInfo->GetNumInstanceSceneDataEntries(), GetPrimitiveID(Scene, PrimitiveId));
 				}
 			}
@@ -923,9 +920,6 @@ void FGPUScene::UploadGeneral(FRHICommandListImmediate& RHICmdList, FScene *Scen
 
 		const uint32 LightMapDataBufferSize = BufferState.LightMapDataBufferSize;
 
-		// These should always be in sync with each other.
-		check(InstanceSceneDataToClear.Num() == InstanceSceneDataAllocator.GetMaxSize());
-
 		const int32 NumPrimitiveDataUploads = UploadDataSourceAdapter.NumPrimitivesToUpload();
 
 		if (Scene != nullptr)
@@ -960,13 +954,6 @@ void FGPUScene::UploadGeneral(FRHICommandListImmediate& RHICmdList, FScene *Scen
 							PrimitiveUploadBufferCS.Lock();
 						}
 
-						// Cancel out any pending clear bits for these instances.
-						if (UploadInfo.InstanceSceneDataUploads > 0)
-						{
-							check(UploadInfo.InstanceSceneDataOffset != INDEX_NONE);
-							InstanceSceneDataToClear.SetRange(UploadInfo.InstanceSceneDataOffset, UploadInfo.InstanceSceneDataUploads, false);
-						}
-
 						NumLightmapDataUploads += UploadInfo.LightmapUploadCount; // Not thread safe
 						NumInstanceSceneDataUploads += UploadInfo.InstanceSceneDataUploads; // Not thread safe
 						NumInstancePayloadDataUploads += UploadInfo.InstancePayloadDataUploads; // Not thread safe
@@ -1122,63 +1109,6 @@ void FGPUScene::UploadGeneral(FRHICommandListImmediate& RHICmdList, FScene *Scen
 				}
 			}
 		}
-		// Make sure instance buffer always has valid or properly reset entries.
-		TArray<uint32, TInlineAllocator<64, SceneRenderingAllocator>> InstancesToClear;
-
-	#if !UE_BUILD_SHIPPING
-		static const bool bVerifyClearList = false;
-		if (bVerifyClearList)
-		{
-			// We need to make sure that every set clear bit in the total list is
-			// represented by an entry in the clear list. We can safely ignore unset
-			// bits - such as clear list contains an entry from an earlier removal,
-			// but clear bit was unset prior to clearing because a slot was reused
-			// by an instance added.
-			for (int32 InstanceIndex = 0; InstanceIndex < InstanceSceneDataAllocator.GetMaxSize(); ++InstanceIndex)
-			{
-				if (InstanceSceneDataToClear[InstanceIndex])
-				{
-					check(InstanceSceneDataClearList.Contains(InstanceIndex));
-				}
-			}
-		}
-	#endif
-
-		TSet<uint32> InstanceSceneDataClearOverflow;
-		for (uint32 InstanceIndex : InstanceSceneDataClearList)
-		{
-			// Any clear bits set after enumerating the primitives being updated are
-			// stale sections of the instance data buffer, so they should be reset to
-			// an invalid state and skipped on the GPU.
-			if (InstanceSceneDataToClear[InstanceIndex])
-			{
-				if (InstanceIndex < BufferState.InstanceSceneDataSOAStride)
-				{
-					InstanceSceneDataToClear[InstanceIndex] = false;
-					InstancesToClear.Add(InstanceIndex);
-				}
-				else
-				{
-					InstanceSceneDataClearOverflow.Add(InstanceIndex);
-				}
-			}
-		}
-
-		InstanceSceneDataClearList.Reset();
-
-		Swap(InstanceSceneDataClearOverflow, InstanceSceneDataClearList);
-
-#if DO_GUARD_SLOW
-		// Verify that if InstanceClearList is empty all the upload bits are cleared as well.
-		if (InstanceSceneDataClearList.Num() == 0)
-		{
-			ensure(InstanceSceneDataToClear.Find(true) == INDEX_NONE);
-		}
-#endif
-
-		// Clears count toward the total instance scene data uploads - batched together for efficiency.
-		NumInstanceSceneDataUploads += InstancesToClear.Num();
-
 		{
 			if (NumInstancePayloadDataUploads > 0)
 			{
@@ -1190,51 +1120,11 @@ void FGPUScene::UploadGeneral(FRHICommandListImmediate& RHICmdList, FScene *Scen
 			{
 				InstanceSceneUploadBuffer.Init(NumInstanceSceneDataUploads * InstanceSceneDataNumArrays, sizeof(FVector4f), true, TEXT("InstanceSceneUploadBuffer"));
 
-				int32 RangeCount = PartitionUpdateRanges(ParallelRanges, InstancesToClear.Num(), bExecuteInParallel);
-
-				FInstanceSceneShaderData ClearedShaderData;
-				ClearedShaderData.Build(INVALID_PRIMITIVE_ID, 0, 0, INVALID_LAST_UPDATE_FRAME, 0, 0.0f);
-
-				// Reset any instance slots marked for clearing.
-				ParallelFor(RangeCount,
-					[this, &InstancesToClear, &ParallelRanges, RangeCount, InstanceSceneDataNumArrays, &ClearedShaderData , &BufferState](int32 RangeIndex)
-					{
-						for (int32 ItemIndex = ParallelRanges.Range[RangeIndex].ItemStart; ItemIndex < ParallelRanges.Range[RangeIndex].ItemStart + ParallelRanges.Range[RangeIndex].ItemCount; ++ItemIndex)
-						{
-							const int32 Index = InstancesToClear[ItemIndex];
-
-							void* DstRefs[FInstanceSceneShaderData::DataStrideInFloat4s];
-							if (RangeCount > 1)
-							{
-								PrimitiveUploadBufferCS.Lock();
-							}
-
-							for (uint32 RefIndex = 0; RefIndex < InstanceSceneDataNumArrays; ++RefIndex)
-							{
-								DstRefs[RefIndex] = InstanceSceneUploadBuffer.Add_GetRef(RefIndex * BufferState.InstanceSceneDataSOAStride + Index);
-							}
-
-							if (RangeCount > 1)
-							{
-								PrimitiveUploadBufferCS.Unlock();
-							}
-
-							// TODO: This is silly, use a custom shader to splat the identity shader data over multiple output locations - way more efficient bandwidth and memory usage.
-							for (uint32 RefIndex = 0; RefIndex < InstanceSceneDataNumArrays; ++RefIndex) //TODO: make a SOA version of InstanceUploadBuffer.Add
-							{
-								FVector4f* DstVector = static_cast<FVector4f*>(DstRefs[RefIndex]);
-								*DstVector = ClearedShaderData.Data[RefIndex];
-							}
-						}
-					},
-					RangeCount == 1
-				);
-
 				if (NumPrimitiveDataUploads > 0)
 				{
 					// Note: we iterate over the primitives, whether they have instances or not (which is a bit wasteful) but this is the way we currently get to the instance data.
 					// GPUCULL_TODO: move instance data ownership to GPU-scene such that it can be put in a compact list or something, and be tracked independent of primitives?
-					RangeCount = PartitionUpdateRanges(ParallelRanges, NumPrimitiveDataUploads, bExecuteInParallel);
+					int32 RangeCount = PartitionUpdateRanges(ParallelRanges, NumPrimitiveDataUploads, bExecuteInParallel);
 
 					// Upload any out of date instance slots.
 					ParallelFor(RangeCount,
@@ -1597,6 +1487,9 @@ void FGPUScene::UploadDynamicPrimitiveShaderDataForViewInternal(FRDGBuilder& Gra
 		FUploadDataSourceAdapterDynamicPrimitives UploadAdapter(Collector.UploadData->PrimitiveShaderData, UploadIdStart, Collector.UploadData->InstanceSceneDataOffset, SceneFrameNumber);
 		FGPUSceneBufferState BufferState = UpdateBufferState(GraphBuilder, Scene, UploadAdapter);
 
+		// Run a pass that clears (Sets ID to invalid) any instances that need it.
+		AddClearInstancesPass(GraphBuilder);
+
 		AddPass(GraphBuilder, RDG_EVENT_NAME("GPUScene::UploadDynamicPrimitiveShaderDataForView"),
 			[this, Scene, UploadAdapter, BufferState = MoveTemp(BufferState)](FRHICommandListImmediate& RHICmdList)
 		{
@@ -1677,27 +1570,7 @@ int32 FGPUScene::AllocateInstanceSceneDataSlots(int32 NumInstanceSceneDataEntrie
 		if (NumInstanceSceneDataEntries > 0)
 		{
 			int32 InstanceSceneDataOffset = InstanceSceneDataAllocator.Allocate(NumInstanceSceneDataEntries);
-
-			// Allocate enough storage space, if needed.
-			const int32 NewSize = InstanceSceneDataOffset + NumInstanceSceneDataEntries;
-			if (NewSize >= InstanceSceneDataToClear.Num())
-			{
-				InstanceSceneDataToClear.Add(false, NewSize - InstanceSceneDataToClear.Num());
-			}
-
-			// Set all bits associated with newly allocated instance data, otherwise deferred uploads will result in uninitialized instances (primarily for dynamic primitives).
-			for (int32 AddIndex = 0; AddIndex < NumInstanceSceneDataEntries; ++AddIndex)
-			{
-				int32 InstanceIndex = InstanceSceneDataOffset + AddIndex;
-				if (!InstanceSceneDataToClear[InstanceIndex])
-				{
-					// Note: we could keep ranges in InstanceSceneDataClearList to make it far more compact when many instances / primitive are used
-					InstanceSceneDataClearList.Add(InstanceIndex);
-				}
-			}
-			InstanceSceneDataToClear.SetRange(InstanceSceneDataOffset, NumInstanceSceneDataEntries, true);
-
-			check(InstanceSceneDataToClear.Num() == InstanceSceneDataAllocator.GetMaxSize());
+			InstanceRangesToClear.Add(FInstanceRange{ uint32(InstanceSceneDataOffset), uint32(NumInstanceSceneDataEntries) });
 
 			return InstanceSceneDataOffset;
 		}
@@ -1713,26 +1586,7 @@ void FGPUScene::FreeInstanceSceneDataSlots(int32 InstanceSceneDataOffset, int32
 	if (bIsEnabled)
 	{
 		InstanceSceneDataAllocator.Free(InstanceSceneDataOffset, NumInstanceSceneDataEntries);
-		for (int32 AddIndex = 0; AddIndex < NumInstanceSceneDataEntries; ++AddIndex)
-		{
-			int32 InstanceIndex = InstanceSceneDataOffset + AddIndex;
-			if (!InstanceSceneDataToClear[InstanceIndex])
-			{
-				// Note: we could keep ranges in InstanceSceneDataClearList to make it far more compact when many instances / primitive are used
-				InstanceSceneDataClearList.Add(InstanceIndex);
-			}
-		}
-		InstanceSceneDataToClear.SetRange(InstanceSceneDataOffset, NumInstanceSceneDataEntries, true);
-
-		// Resize bit arrays to match new high watermark
-		if (InstanceSceneDataToClear.Num() > InstanceSceneDataAllocator.GetMaxSize())
-		{
-			const int32 OldBitCount = InstanceSceneDataToClear.Num();
-			const int32 NewBitCount = InstanceSceneDataAllocator.GetMaxSize();
-			const int32 RemBitCount = OldBitCount - NewBitCount;
-			InstanceSceneDataToClear.RemoveAt(NewBitCount, RemBitCount);
-			check(InstanceSceneDataToClear.Num() == InstanceSceneDataAllocator.GetMaxSize());
-		}
+		InstanceRangesToClear.Add(FInstanceRange{ uint32(InstanceSceneDataOffset), uint32(NumInstanceSceneDataEntries) });
 	}
 }
 
@@ -2134,4 +1988,15 @@ void FGPUSceneCompactInstanceData::Init(const FScene* Scene, int32 PrimitiveId)
 	InstanceTransform3		= LocalToRelativeWorld.GetScaledAxis(EAxis::Z);
 	InstanceOriginAndId.W	= *(float*)&PrimitiveId;
 	InstanceAuxData			= FVector4f(0);
-}
\ No newline at end of file
+}
+
+void FGPUScene::AddClearInstancesPass(FRDGBuilder& GraphBuilder)
+{
+	FInstanceGPULoadBalancer ClearIdData;
+	for (FInstanceRange Range : InstanceRangesToClear)
+	{
+		ClearIdData.Add(Range.InstanceSceneDataOffset, Range.NumInstanceSceneDataEntries, INVALID_PRIMITIVE_ID);
+	}
+	AddUpdatePrimitiveIdsPass(GraphBuilder, ClearIdData);
+	InstanceRangesToClear.Empty();
+}
diff --git a/Engine/Source/Runtime/Renderer/Private/GPUScene.h b/Engine/Source/Runtime/Renderer/Private/GPUScene.h
--- a/Engine/Source/Runtime/Renderer/Private/GPUScene.h
+++ b/Engine/Source/Runtime/Renderer/Private/GPUScene.h
@@ -334,8 +334,13 @@ public:
 private:
 	TArray<EPrimitiveDirtyState> PrimitiveDirtyState;
 
-	TBitArray<>				InstanceSceneDataToClear;
-	TSet<uint32>			InstanceSceneDataClearList;
+	struct FInstanceRange
+	{
+		uint32 InstanceSceneDataOffset;
+		uint32 NumInstanceSceneDataEntries;
+	};
+
+	TArray<FInstanceRange> InstanceRangesToClear;
 
 	TRange<int32> CommitPrimitiveCollector(FGPUScenePrimitiveCollector& PrimitiveCollector);
 
@@ -372,6 +377,8 @@ private:
 	void UpdateInternal(FRDGBuilder& GraphBuilder, FScene& Scene);
 
 	void AddUpdatePrimitiveIdsPass(FRDGBuilder& GraphBuilder, FInstanceGPULoadBalancer& IdOnlyUpdateItems);
+
+	void AddClearInstancesPass(FRDGBuilder& GraphBuilder);
 };
 
 class FGPUSceneScopeBeginEndHelper
