Refactor GPU-Scene upload to be parallel, results in about 3x speed improvement (e.g., ~15ms -> ~3ms for bad cases).

diff --git a/Engine/Source/Runtime/RenderCore/Private/UnifiedBuffer.cpp b/Engine/Source/Runtime/RenderCore/Private/UnifiedBuffer.cpp
--- a/Engine/Source/Runtime/RenderCore/Private/UnifiedBuffer.cpp
+++ b/Engine/Source/Runtime/RenderCore/Private/UnifiedBuffer.cpp
@@ -673,6 +673,20 @@ void FScatterUploadBuffer::Init( uint32 NumElements, uint32 InNumBytesPerElement
 	}
 }
 
+void FScatterUploadBuffer::Init(TArrayView<const uint32> ElementScatterOffsets, uint32 InNumBytesPerElement, bool bInFloat4Buffer, const TCHAR* DebugName)
+{
+	Init(ElementScatterOffsets.Num(), InNumBytesPerElement, bInFloat4Buffer, DebugName);
+	FMemory::ParallelMemcpy(ScatterData, ElementScatterOffsets.GetData(), ElementScatterOffsets.Num() * ElementScatterOffsets.GetTypeSize(), EMemcpyCachePolicy::StoreUncached);
+	NumScatters = ElementScatterOffsets.Num();
+	ScatterData += ElementScatterOffsets.Num();
+}
+
+void FScatterUploadBuffer::InitPreSized(uint32 NumElements, uint32 InNumBytesPerElement, bool bInFloat4Buffer, const TCHAR* DebugName)
+{
+	Init(NumElements, InNumBytesPerElement, bInFloat4Buffer, DebugName);
+	NumScatters = NumElements;
+}
+
 // Helper type used to initialize the buffer data on creation
 struct FScatterUploadBufferResourceArray : public FResourceArrayInterface
 {
diff --git a/Engine/Source/Runtime/RenderCore/Public/UnifiedBuffer.h b/Engine/Source/Runtime/RenderCore/Public/UnifiedBuffer.h
--- a/Engine/Source/Runtime/RenderCore/Public/UnifiedBuffer.h
+++ b/Engine/Source/Runtime/RenderCore/Public/UnifiedBuffer.h
@@ -108,6 +108,18 @@ public:
 		return Result;
 	}
 
+	void* Set_GetRef(uint32 ElementIndex, uint32 ElementScatterOffset, uint32 Num = 1)
+	{
+		checkSlow(ElementIndex + Num <= MaxScatters );
+		checkSlow( ScatterData != nullptr );
+		checkSlow( UploadData != nullptr );
+		for (uint32 i = 0; i < Num; i++)
+		{
+			ScatterData[ElementIndex + i] = ElementScatterOffset + i;
+		}
+		return UploadData + ElementIndex * NumBytesPerElement;
+	}
+
 	void Release()
 	{
 		ScatterBuffer.Release();
@@ -145,4 +157,24 @@ public:
 			bUploadViaCreate = bInUploadViaCreate;
 		}
 	}
+	
+	/**
+	 * Init with presized num scatters, expecting each to be set at a later point. Requires the user to keep track of the offsets to use.
+	 */
+	RENDERCORE_API void InitPreSized(uint32 NumElements, uint32 InNumBytesPerElement, bool bInFloat4Buffer, const TCHAR* DebugName);
+
+	/**
+	 * Init with pre-existing destination index data, performs a bulk-copy.
+	 */
+	RENDERCORE_API void Init(TArrayView<const uint32> ElementScatterOffsets, uint32 InNumBytesPerElement, bool bInFloat4Buffer, const TCHAR* DebugName);
+	/**
+	 * Get pointer to an element data area, given the index of the element (not the destination scatter offset).
+	 */
+	FORCEINLINE void* GetRef(uint32 ElementIndex)
+	{
+		checkSlow(ScatterData != nullptr);
+		checkSlow(UploadData != nullptr);
+
+		return UploadData + ElementIndex * NumBytesPerElement;
+	}
 };
\ No newline at end of file
diff --git a/Engine/Source/Runtime/Renderer/Private/GPUScene.cpp b/Engine/Source/Runtime/Renderer/Private/GPUScene.cpp
--- a/Engine/Source/Runtime/Renderer/Private/GPUScene.cpp
+++ b/Engine/Source/Runtime/Renderer/Private/GPUScene.cpp
@@ -97,6 +97,8 @@ static TAutoConsoleVariable<float> CVarGPUSceneDebugDrawRange(
 	ECVF_RenderThreadSafe
 );
 
+constexpr uint32 InstanceSceneDataNumArrays = FInstanceSceneShaderData::DataStrideInFloat4s;
+
 int32 GGPUSceneInstanceUploadViaCreate = 1;
 FAutoConsoleVariableRef CVarGPUSceneInstanceUploadViaCreate(
 	TEXT("r.GPUScene.InstanceUploadViaCreate"),
@@ -174,26 +176,28 @@ struct FBVHNode
 };
 
 /**
- * Info needed by the uploader to update a primitive.
+ * Info needed by the uploader to prepare to upload a primitive.
  */
-struct FPrimitiveUploadInfo
+struct FPrimitiveUploadInfoHeader
 {
-	/** Required */
-	FPrimitiveSceneShaderData PrimitiveSceneData;
 	int32 PrimitiveID = INDEX_NONE;
 
 	/** Optional */
-	int32 InstanceSceneDataOffset = INDEX_NONE;
-	int32 InstanceSceneDataUploads = 0;
-	//int32 InstancePayloadDataOffset = INDEX_NONE;
-	int32 InstancePayloadDataUploads = 0;
+	int32 NumInstanceUploads = 0;
+	int32 NumInstancePayloadDataUploads = 0;
 	int32 LightmapUploadCount = 0;
 
 	/** NaniteSceneProxy must be set if the proxy is a Nanite proxy */
 	const Nanite::FSceneProxyBase* NaniteSceneProxy = nullptr;
 	const FPrimitiveSceneInfo* PrimitiveSceneInfo = nullptr;
 };
-
+/**
+ * Info needed by the uploader to update a primitive.
+ */
+struct FPrimitiveUploadInfo : public FPrimitiveUploadInfoHeader
+{
+	FPrimitiveSceneShaderData PrimitiveSceneData;
+};
 /**
  * Info required by the uploader to update the instances that belong to a primitive.
  */
@@ -291,124 +295,135 @@ struct FUploadDataSourceAdapterScenePrimitives
 		return PrimitivesToUpdate.Num(); 
 	}
 
+	FORCEINLINE TArrayView<const uint32> GetItemPrimitiveIds() const
+	{
+		return TArrayView<const uint32>(reinterpret_cast<const uint32 *>(PrimitivesToUpdate.GetData()), PrimitivesToUpdate.Num());
+	}
+
 	/**
 	 * Populate the primitive info for a given item index.
 	 * 
 	 */
-	FORCEINLINE bool GetPrimitiveInfo(int32 ItemIndex, FPrimitiveUploadInfo& PrimitiveUploadInfo) const
+	FORCEINLINE void GetPrimitiveInfoHeader(int32 ItemIndex, FPrimitiveUploadInfoHeader& PrimitiveUploadInfo) const
 	{
 		int32 PrimitiveID = PrimitivesToUpdate[ItemIndex];
-		if (PrimitiveID < Scene.PrimitiveSceneProxies.Num())
-		{
-			const FPrimitiveSceneProxy* RESTRICT PrimitiveSceneProxy = Scene.PrimitiveSceneProxies[PrimitiveID];
-			const FPrimitiveSceneInfo* PrimitiveSceneInfo = PrimitiveSceneProxy->GetPrimitiveSceneInfo();
+		check(PrimitiveID < Scene.PrimitiveSceneProxies.Num());
 
-			PrimitiveUploadInfo.PrimitiveID = PrimitiveID;;
-			PrimitiveUploadInfo.LightmapUploadCount = PrimitiveSceneInfo->GetNumLightmapDataEntries();
-			PrimitiveUploadInfo.NaniteSceneProxy = PrimitiveSceneProxy->IsNaniteMesh() ? static_cast<const Nanite::FSceneProxyBase*>(PrimitiveSceneProxy) : nullptr;
-			PrimitiveUploadInfo.PrimitiveSceneInfo = PrimitiveSceneInfo;
+		const FPrimitiveSceneProxy* RESTRICT PrimitiveSceneProxy = Scene.PrimitiveSceneProxies[PrimitiveID];
+		const FPrimitiveSceneInfo* RESTRICT PrimitiveSceneInfo = PrimitiveSceneProxy->GetPrimitiveSceneInfo();
 
-			// Count all primitive instances represented in the instance data buffer.
-			if (PrimitiveSceneProxy->SupportsInstanceDataBuffer())
-			{
-				const TConstArrayView<FPrimitiveInstance> InstanceSceneData = PrimitiveSceneProxy->GetInstanceSceneData();
-				PrimitiveUploadInfo.InstanceSceneDataOffset = PrimitiveSceneProxy->GetPrimitiveSceneInfo()->GetInstanceSceneDataOffset();
-				PrimitiveUploadInfo.InstanceSceneDataUploads = InstanceSceneData.Num();
-				PrimitiveUploadInfo.InstancePayloadDataUploads = PrimitiveSceneProxy->GetPrimitiveSceneInfo()->GetInstancePayloadDataStride() * InstanceSceneData.Num();
-			}
-			else
-			{
-				PrimitiveUploadInfo.InstanceSceneDataOffset = PrimitiveSceneProxy->GetPrimitiveSceneInfo()->GetInstanceSceneDataOffset();
-				PrimitiveUploadInfo.InstanceSceneDataUploads = 1;
-				PrimitiveUploadInfo.InstancePayloadDataUploads = 0;
+		PrimitiveUploadInfo.PrimitiveID = PrimitiveID;;
+		PrimitiveUploadInfo.LightmapUploadCount = PrimitiveSceneInfo->GetNumLightmapDataEntries();
+		PrimitiveUploadInfo.NaniteSceneProxy = PrimitiveSceneProxy->IsNaniteMesh() ? static_cast<const Nanite::FSceneProxyBase*>(PrimitiveSceneProxy) : nullptr;
+		PrimitiveUploadInfo.PrimitiveSceneInfo = PrimitiveSceneInfo;
 
-				checkSlow(!PrimitiveSceneProxy->HasAnyPerInstancePayloadData());
-			}
-
-			PrimitiveUploadInfo.PrimitiveSceneData = FPrimitiveSceneShaderData(PrimitiveSceneProxy);
-			return true;
+		// Prevent these from allocating instance update work
+		if (PrimitiveDirtyState[PrimitiveID] == EPrimitiveDirtyState::ChangedId)
+		{
+			PrimitiveUploadInfo.NumInstanceUploads = 0;
+			PrimitiveUploadInfo.NumInstancePayloadDataUploads = 0;
+		}
+		else 
+		{
+			PrimitiveUploadInfo.NumInstanceUploads = PrimitiveSceneInfo->GetNumInstanceSceneDataEntries();
+			PrimitiveUploadInfo.NumInstancePayloadDataUploads = PrimitiveSceneInfo->GetInstancePayloadDataStride() * PrimitiveUploadInfo.NumInstanceUploads;
 		}
+	}
 
-		return false;
+	/**
+	 * Populate the primitive info for a given item index.
+	 * 
+	 */
+	FORCEINLINE void GetPrimitiveInfo(int32 ItemIndex, FPrimitiveUploadInfo& PrimitiveUploadInfo) const
+	{
+		int32 PrimitiveID = PrimitivesToUpdate[ItemIndex];
+		check(PrimitiveID < Scene.PrimitiveSceneProxies.Num());
+
+		GetPrimitiveInfoHeader(ItemIndex, PrimitiveUploadInfo);
+
+		const FPrimitiveSceneProxy* RESTRICT PrimitiveSceneProxy = Scene.PrimitiveSceneProxies[PrimitiveID];
+		const FPrimitiveSceneInfo* PrimitiveSceneInfo = PrimitiveSceneProxy->GetPrimitiveSceneInfo();
+
+		PrimitiveUploadInfo.PrimitiveSceneData = FPrimitiveSceneShaderData(PrimitiveSceneProxy);
 	}
 
-	FORCEINLINE bool GetInstanceInfo(int32 ItemIndex, FInstanceUploadInfo& InstanceUploadInfo) const
+	FORCEINLINE void GetInstanceInfo(int32 ItemIndex, FInstanceUploadInfo& InstanceUploadInfo) const
 	{
 		const int32 PrimitiveID = PrimitivesToUpdate[ItemIndex];
-		if (PrimitiveID < Scene.PrimitiveSceneProxies.Num() && PrimitiveDirtyState[PrimitiveID] != EPrimitiveDirtyState::ChangedId)
-		{
-			FPrimitiveSceneProxy* PrimitiveSceneProxy = Scene.PrimitiveSceneProxies[PrimitiveID];
-			const FPrimitiveSceneInfo* PrimitiveSceneInfo = PrimitiveSceneProxy->GetPrimitiveSceneInfo();
-			
-			const FMatrix LocalToWorld = PrimitiveSceneProxy->GetLocalToWorld();
-			const FLargeWorldRenderPosition AbsoluteOrigin(LocalToWorld.GetOrigin());
 
-			InstanceUploadInfo.InstanceSceneDataOffset = PrimitiveSceneInfo->GetInstanceSceneDataOffset();
-			InstanceUploadInfo.InstancePayloadDataOffset = PrimitiveSceneInfo->GetInstancePayloadDataOffset();
-			InstanceUploadInfo.InstancePayloadDataStride = PrimitiveSceneInfo->GetInstancePayloadDataStride();
+		check(PrimitiveID < Scene.PrimitiveSceneProxies.Num());
+		check(PrimitiveDirtyState[PrimitiveID] != EPrimitiveDirtyState::ChangedId);
 
-			InstanceUploadInfo.LastUpdateSceneFrameNumber = SceneFrameNumber;
-			InstanceUploadInfo.PrimitiveID = GetPrimitiveID(Scene, PrimitiveID);
-			InstanceUploadInfo.PrimitiveToWorld = AbsoluteOrigin.MakeToRelativeWorldMatrix(LocalToWorld);
+		FPrimitiveSceneProxy* PrimitiveSceneProxy = Scene.PrimitiveSceneProxies[PrimitiveID];
+		const FPrimitiveSceneInfo* PrimitiveSceneInfo = PrimitiveSceneProxy->GetPrimitiveSceneInfo();
+		
+		const FMatrix LocalToWorld = PrimitiveSceneProxy->GetLocalToWorld();
+		const FLargeWorldRenderPosition AbsoluteOrigin(LocalToWorld.GetOrigin());
 
-			{
-				bool bHasPrecomputedVolumetricLightmap{};
-				bool bOutputVelocity{};
-				int32 SingleCaptureIndex{};
+		InstanceUploadInfo.InstanceSceneDataOffset = PrimitiveSceneInfo->GetInstanceSceneDataOffset();
+		check(InstanceUploadInfo.InstanceSceneDataOffset >= 0);
+		InstanceUploadInfo.InstancePayloadDataOffset = PrimitiveSceneInfo->GetInstancePayloadDataOffset();
+		InstanceUploadInfo.InstancePayloadDataStride = PrimitiveSceneInfo->GetInstancePayloadDataStride();
 
-				FMatrix PreviousLocalToWorld;
-				Scene.GetPrimitiveUniformShaderParameters_RenderThread(PrimitiveSceneInfo, bHasPrecomputedVolumetricLightmap, PreviousLocalToWorld, SingleCaptureIndex, bOutputVelocity);
-				InstanceUploadInfo.PrevPrimitiveToWorld = AbsoluteOrigin.MakeClampedToRelativeWorldMatrix(PreviousLocalToWorld);
-			}
+		InstanceUploadInfo.LastUpdateSceneFrameNumber = SceneFrameNumber;
+		InstanceUploadInfo.PrimitiveID = GetPrimitiveID(Scene, PrimitiveID);
+		InstanceUploadInfo.PrimitiveToWorld = AbsoluteOrigin.MakeToRelativeWorldMatrix(LocalToWorld);
 
-			InstanceUploadInfo.PayloadDataFlags = PrimitiveSceneProxy->GetPayloadDataFlags();
-			InstanceUploadInfo.InstanceLocalBounds = PrimitiveSceneProxy->GetInstanceLocalBounds();
-			if (InstanceUploadInfo.InstanceLocalBounds.Num() == 0)
-			{
-				InstanceUploadInfo.DummyLocalBounds = PrimitiveSceneProxy->GetLocalBounds();
-				InstanceUploadInfo.InstanceLocalBounds = TConstArrayView<FRenderBounds>(&InstanceUploadInfo.DummyLocalBounds, 1);
-			}
+		{
+			bool bHasPrecomputedVolumetricLightmap{};
+			bool bOutputVelocity{};
+			int32 SingleCaptureIndex{};
 
-			if (PrimitiveSceneProxy->SupportsInstanceDataBuffer())
-			{
-				InstanceUploadInfo.PrimitiveInstances = PrimitiveSceneProxy->GetInstanceSceneData();
-				InstanceUploadInfo.InstanceDynamicData = PrimitiveSceneProxy->GetInstanceDynamicData();
-				InstanceUploadInfo.InstanceLightShadowUVBias = PrimitiveSceneProxy->GetInstanceLightShadowUVBias();
-				InstanceUploadInfo.InstanceCustomData = PrimitiveSceneProxy->GetInstanceCustomData();
-				InstanceUploadInfo.InstanceRandomID = PrimitiveSceneProxy->GetInstanceRandomID();
-				InstanceUploadInfo.InstanceHierarchyOffset = PrimitiveSceneProxy->GetInstanceHierarchyOffset();
-			}
-			else
-			{
-				check(InstanceUploadInfo.PayloadDataFlags == 0x0);
-				check(InstanceUploadInfo.InstancePayloadDataOffset == INDEX_NONE && InstanceUploadInfo.InstancePayloadDataStride == 0);
-
-				// We always create an instance to ensure that we can always use the same code paths in the shader
-				// In the future we should remove redundant data from the primitive, and then the instances should be
-				// provided by the proxy. However, this is a lot of work before we can just enable it in the base proxy class.
-				InstanceUploadInfo.DummyInstance.LocalToPrimitive.SetIdentity();
-
-				InstanceUploadInfo.PrimitiveInstances = TConstArrayView<FPrimitiveInstance>(&InstanceUploadInfo.DummyInstance, 1);
-				InstanceUploadInfo.InstanceDynamicData = TConstArrayView<FPrimitiveInstanceDynamicData>();
-				InstanceUploadInfo.InstanceLightShadowUVBias = TConstArrayView<FVector4f>();
-				InstanceUploadInfo.InstanceCustomData = TConstArrayView<float>();
-				InstanceUploadInfo.InstanceRandomID = TConstArrayView<float>();
-				InstanceUploadInfo.InstanceHierarchyOffset = TConstArrayView<uint32>();
-			}
+			FMatrix PreviousLocalToWorld;
+			Scene.GetPrimitiveUniformShaderParameters_RenderThread(PrimitiveSceneInfo, bHasPrecomputedVolumetricLightmap, PreviousLocalToWorld, SingleCaptureIndex, bOutputVelocity);
+			InstanceUploadInfo.PrevPrimitiveToWorld = AbsoluteOrigin.MakeClampedToRelativeWorldMatrix(PreviousLocalToWorld);;
+		}
 
-			InstanceUploadInfo.InstanceCustomDataCount = 0;
-			if (InstanceUploadInfo.InstanceCustomData.Num() > 0)
-			{
-				InstanceUploadInfo.InstanceCustomDataCount = InstanceUploadInfo.InstanceCustomData.Num() / InstanceUploadInfo.PrimitiveInstances.Num();
-			}
+		InstanceUploadInfo.PayloadDataFlags = PrimitiveSceneProxy->GetPayloadDataFlags();
+		InstanceUploadInfo.InstanceLocalBounds = PrimitiveSceneProxy->GetInstanceLocalBounds();
+		if (InstanceUploadInfo.InstanceLocalBounds.Num() == 0)
+		{
+			InstanceUploadInfo.DummyLocalBounds = PrimitiveSceneProxy->GetLocalBounds();
+			InstanceUploadInfo.InstanceLocalBounds = TConstArrayView<FRenderBounds>(&InstanceUploadInfo.DummyLocalBounds, 1);
+		}
+
+		if (PrimitiveSceneProxy->SupportsInstanceDataBuffer())
+		{
+			InstanceUploadInfo.PrimitiveInstances = PrimitiveSceneProxy->GetInstanceSceneData();
+			InstanceUploadInfo.InstanceDynamicData = PrimitiveSceneProxy->GetInstanceDynamicData();
+			InstanceUploadInfo.InstanceLightShadowUVBias = PrimitiveSceneProxy->GetInstanceLightShadowUVBias();
+			InstanceUploadInfo.InstanceCustomData = PrimitiveSceneProxy->GetInstanceCustomData();
+			InstanceUploadInfo.InstanceRandomID = PrimitiveSceneProxy->GetInstanceRandomID();
+			InstanceUploadInfo.InstanceHierarchyOffset = PrimitiveSceneProxy->GetInstanceHierarchyOffset();
+		}
+		else
+		{
+			check(InstanceUploadInfo.PayloadDataFlags == 0x0);
+			check(InstanceUploadInfo.InstancePayloadDataOffset == INDEX_NONE && InstanceUploadInfo.InstancePayloadDataStride == 0);
 
-			ValidateInstanceUploadInfo(InstanceUploadInfo);
+			// We always create an instance to ensure that we can always use the same code paths in the shader
+			// In the future we should remove redundant data from the primitive, and then the instances should be
+			// provided by the proxy. However, this is a lot of work before we can just enable it in the base proxy class.
+			InstanceUploadInfo.DummyInstance.LocalToPrimitive.SetIdentity();
 
-			// Only trigger upload if this primitive has instances
-			return InstanceUploadInfo.PrimitiveInstances.Num() > 0;
+			InstanceUploadInfo.PrimitiveInstances = TConstArrayView<FPrimitiveInstance>(&InstanceUploadInfo.DummyInstance, 1);
+			InstanceUploadInfo.InstanceDynamicData = TConstArrayView<FPrimitiveInstanceDynamicData>();
+			InstanceUploadInfo.InstanceLightShadowUVBias = TConstArrayView<FVector4f>();
+			InstanceUploadInfo.InstanceCustomData = TConstArrayView<float>();
+			InstanceUploadInfo.InstanceRandomID = TConstArrayView<float>();
+			InstanceUploadInfo.InstanceHierarchyOffset = TConstArrayView<uint32>();
 		}
 
-		return false;
+		InstanceUploadInfo.InstanceCustomDataCount = 0;
+		if (InstanceUploadInfo.InstanceCustomData.Num() > 0)
+		{
+			InstanceUploadInfo.InstanceCustomDataCount = InstanceUploadInfo.InstanceCustomData.Num() / InstanceUploadInfo.PrimitiveInstances.Num();
+		}
+
+		ValidateInstanceUploadInfo(InstanceUploadInfo);
+
+		// Only trigger upload if this primitive has instances
+		check(InstanceUploadInfo.PrimitiveInstances.Num() > 0);
 	}
 
 	FORCEINLINE bool GetLightMapInfo(int32 ItemIndex, FLightMapUploadInfo &UploadInfo) const
@@ -429,7 +444,7 @@ struct FUploadDataSourceAdapterScenePrimitives
 
 	FScene& Scene;
 	const uint32 SceneFrameNumber;
-	TArray<int> PrimitivesToUpdate;
+	TArray<int32> PrimitivesToUpdate;
 	TArray<EPrimitiveDirtyState> PrimitiveDirtyState;
 };
 
@@ -500,6 +515,19 @@ void FGPUScene::UpdateInternal(FRDGBuilder& GraphBuilder, FScene& Scene)
 	// Store in GPU-scene to enable validation that update has been carried out.
 	SceneFrameNumber = Scene.GetFrameNumber();
 
+	// Strip all out-of-range ID's (left over because of deletes) so we don't need to check later
+	for (int32 Index = 0; Index < PrimitivesToUpdate.Num();)
+	{
+		if (PrimitivesToUpdate[Index] >= Scene.PrimitiveSceneProxies.Num())
+		{
+			PrimitivesToUpdate.RemoveAtSwap(Index, 1, false);
+		}
+		else
+		{
+			++Index;
+		}
+	}
+
 	FUploadDataSourceAdapterScenePrimitives Adapter(Scene, SceneFrameNumber, MoveTemp(PrimitivesToUpdate), MoveTemp(PrimitiveDirtyState));
 	FGPUSceneBufferState BufferState = UpdateBufferState(GraphBuilder, &Scene, Adapter);
 
@@ -513,13 +541,12 @@ void FGPUScene::UpdateInternal(FRDGBuilder& GraphBuilder, FScene& Scene)
 		{
 			int32 PrimitiveId = Adapter.PrimitivesToUpdate[Index];
 
-			if (PrimitiveId < Scene.PrimitiveSceneProxies.Num() && Adapter.PrimitiveDirtyState[PrimitiveId] == EPrimitiveDirtyState::ChangedId)
+			check(PrimitiveId < Scene.PrimitiveSceneProxies.Num());
+			if (Adapter.PrimitiveDirtyState[PrimitiveId] == EPrimitiveDirtyState::ChangedId)
 			{
 				const FPrimitiveSceneInfo* PrimitiveSceneInfo = Scene.Primitives[PrimitiveId];
-				if (PrimitiveSceneInfo->GetInstanceSceneDataOffset() >= 0)
-				{
-					IdOnlyUpdateData.Add(PrimitiveSceneInfo->GetInstanceSceneDataOffset(), PrimitiveSceneInfo->GetNumInstanceSceneDataEntries(), GetPrimitiveID(Scene, PrimitiveId));
-				}
+				check(PrimitiveSceneInfo->GetInstanceSceneDataOffset() >= 0 || PrimitiveSceneInfo->GetNumInstanceSceneDataEntries() == 0);
+				IdOnlyUpdateData.Add(PrimitiveSceneInfo->GetInstanceSceneDataOffset(), PrimitiveSceneInfo->GetNumInstanceSceneDataEntries(), GetPrimitiveID(Scene, PrimitiveId));
 			}
 		}
 		AddUpdatePrimitiveIdsPass(GraphBuilder, IdOnlyUpdateData);
@@ -897,6 +924,81 @@ FGPUSceneBufferState FGPUScene::UpdateBufferState(FRDGBuilder& GraphBuilder, FSc
 	return BufferState;
 }
 
+/**
+ * Used to queue up load-balanced chunks of instance upload work such that it can be spread over a large number of cores.
+ */
+struct FInstanceUploadBatch
+{
+	static constexpr int32 MaxItems = 64;
+	static constexpr int32 MaxCost = MaxItems * 2; // Selected to allow filling the array when 1:1 primitive / instances
+
+	struct FItem
+	{
+		int32 ItemIndex;
+		int32 FirstInstance;
+		int32 NumInstances;
+	};
+
+	int32 FirstItem = 0;
+	int32 NumItems = 0;
+};
+
+struct FInstanceBatcher
+{
+	struct FPrimitiveItemInfo
+	{
+		int32 InstanceSceneDataUploadOffset;
+		int32 InstancePayloadDataUploadOffset;
+	};
+
+	FInstanceUploadBatch* CurrentBatch = nullptr;
+	TArray<FInstanceUploadBatch, TInlineAllocator<1, SceneRenderingAllocator> > UpdateBatches;
+	TArray<FInstanceUploadBatch::FItem, TInlineAllocator<256, SceneRenderingAllocator> > UpdateBatchItems;
+	TArray<FPrimitiveItemInfo, SceneRenderingAllocator> PerPrimitiveItemInfo;
+
+	int32 CurrentBatchCost = 0;
+
+	int32 InstanceSceneDataUploadOffset = 0;
+	int32 InstancePayloadDataUploadOffset = 0; // Count of float4s
+
+	FInstanceBatcher()
+	{
+		CurrentBatch = &UpdateBatches.AddDefaulted_GetRef();
+	}
+
+	void QueueInstances(const FPrimitiveUploadInfoHeader &UploadInfo, int32 ItemIndex, const FPrimitiveItemInfo & PrimitiveItemInfo)
+	{
+		PerPrimitiveItemInfo[ItemIndex] = PrimitiveItemInfo;
+		const int32 NumInstances = UploadInfo.NumInstanceUploads;
+		int32 InstancesAdded = 0;
+		while (InstancesAdded < NumInstances)
+		{
+			// Populate the last batch until full. Max items/batch = 64, for balance use cost estimate of 1:1 for primitive:instance
+			// Fill to minimum cost (1
+
+			// Can add one less to account for primitive cost
+			int32 MaxInstancesThisBatch = FInstanceUploadBatch::MaxCost - CurrentBatchCost - 1;
+
+			if (MaxInstancesThisBatch > 0)
+			{
+				const int NumInstancesThisItem = FMath::Min(MaxInstancesThisBatch, NumInstances - InstancesAdded);
+				UpdateBatchItems.Add(FInstanceUploadBatch::FItem{ ItemIndex, InstancesAdded, NumInstancesThisItem });
+				CurrentBatch->NumItems += 1;
+				InstancesAdded += NumInstancesThisItem;
+				CurrentBatchCost += NumInstancesThisItem + 1;
+			}
+
+			// Flush batch if it is not possible to add any more items (for one of the reasons)
+			if (MaxInstancesThisBatch <= 0 || CurrentBatchCost > FInstanceUploadBatch::MaxCost - 1 || CurrentBatch->NumItems >= FInstanceUploadBatch::MaxItems)
+			{
+				CurrentBatchCost = 0;
+				CurrentBatch = &UpdateBatches.AddDefaulted_GetRef();
+				CurrentBatch->FirstItem = UpdateBatchItems.Num();
+			}
+		}
+	}
+};
+
 template<typename FUploadDataSourceAdapter>
 void FGPUScene::UploadGeneral(FRHICommandListImmediate& RHICmdList, FScene *Scene, const FUploadDataSourceAdapter& UploadDataSourceAdapter, const FGPUSceneBufferState& BufferState)
 {
@@ -915,7 +1017,6 @@ void FGPUScene::UploadGeneral(FRHICommandListImmediate& RHICmdList, FScene *Scen
 		const bool bExecuteInParallel = GGPUSceneParallelUpdate != 0 && FApp::ShouldUseThreadingForPerformance();
 		const bool bNaniteEnabled = DoesPlatformSupportNanite(GMaxRHIShaderPlatform);
 
-		const uint32 InstanceSceneDataNumArrays = FInstanceSceneShaderData::DataStrideInFloat4s;
 		FUAVTransitionStateScopeHelper InstanceSceneDataTransitionHelper(RHICmdList, BufferState.InstanceSceneDataBuffer.UAV, ERHIAccess::Unknown, ERHIAccess::SRVMask);
 
 		const uint32 LightMapDataBufferSize = BufferState.LightMapDataBufferSize;
@@ -937,166 +1038,130 @@ void FGPUScene::UploadGeneral(FRHICommandListImmediate& RHICmdList, FScene *Scen
 		int32 NumInstanceSceneDataUploads = 0;
 		int32 NumInstancePayloadDataUploads = 0; // Count of float4s
 
-		static FCriticalSection PrimitiveUploadBufferCS;
-
-		FParallelUpdateRanges ParallelRanges;
+		FInstanceBatcher InstanceUpdates;
 
 		if (NumPrimitiveDataUploads > 0)
 		{
-			auto ProcessPrimitiveFn = [this, &UploadDataSourceAdapter, Scene, &NumLightmapDataUploads, &NumInstanceSceneDataUploads, &NumInstancePayloadDataUploads, bNaniteEnabled](int32 ItemIndex, bool bThreaded)
-			{
-				FPrimitiveUploadInfo UploadInfo;
-				if (UploadDataSourceAdapter.GetPrimitiveInfo(ItemIndex, UploadInfo))
-				{
-					{
-						if (bThreaded)
-						{
-							PrimitiveUploadBufferCS.Lock();
-						}
+			SCOPED_NAMED_EVENT(STAT_UpdateGPUScenePrimitives, FColor::Green);
+			QUICK_SCOPE_CYCLE_COUNTER(STAT_UpdateGPUScenePrimitives);
 
-						NumLightmapDataUploads += UploadInfo.LightmapUploadCount; // Not thread safe
-						NumInstanceSceneDataUploads += UploadInfo.InstanceSceneDataUploads; // Not thread safe
-						NumInstancePayloadDataUploads += UploadInfo.InstancePayloadDataUploads; // Not thread safe
+			FUAVTransitionStateScopeHelper PrimitiveDataTransitionHelper(RHICmdList, BufferState.PrimitiveBuffer.UAV, ERHIAccess::Unknown, ERHIAccess::SRVMask);
 
-						void* UploadDst = PrimitiveUploadBuffer.Add_GetRef(UploadInfo.PrimitiveID); // Not thread safe
+			{
+				SCOPED_DRAW_EVENTF(RHICmdList, UpdateGPUScene, TEXT("UpdateGPUScene NumPrimitiveDataUploads %u"), NumPrimitiveDataUploads);
 
-						// TODO: Reserve entire payload data range here for all instances on this primitive, then per instance just write without add ref
+				{
+					SCOPED_NAMED_EVENT(STAT_UpdateGPUScenePrimitives_Seq, FColor::Green);
+					QUICK_SCOPE_CYCLE_COUNTER(STAT_UpdateGPUScenePrimitives_Seq);
 
-						if (bThreaded)
-						{
-							PrimitiveUploadBufferCS.Unlock();
-						}
+					InstanceUpdates.PerPrimitiveItemInfo.SetNumUninitialized(NumPrimitiveDataUploads);
+					PrimitiveUploadBuffer.Init(UploadDataSourceAdapter.GetItemPrimitiveIds(), sizeof(FPrimitiveSceneShaderData::Data), true, TEXT("PrimitiveUploadBuffer"));
+					// 1. do sequential work first
+					for (int32 ItemIndex = 0; ItemIndex < NumPrimitiveDataUploads; ++ItemIndex)
+					{
+						FPrimitiveUploadInfoHeader UploadInfo;
+						UploadDataSourceAdapter.GetPrimitiveInfoHeader(ItemIndex, UploadInfo);
 
-						FVector4f* DstData = static_cast<FVector4f*>(UploadDst);
-						for (uint32 VectorIndex = 0; VectorIndex < FPrimitiveSceneShaderData::DataStrideInFloat4s; ++VectorIndex)
-						{
-							DstData[VectorIndex] = UploadInfo.PrimitiveSceneData.Data[VectorIndex];
-						}
-					}
+						FInstanceBatcher::FPrimitiveItemInfo PrimitiveItemInfo;
+						PrimitiveItemInfo.InstanceSceneDataUploadOffset = NumInstanceSceneDataUploads;
+						PrimitiveItemInfo.InstancePayloadDataUploadOffset = NumInstancePayloadDataUploads;
+						InstanceUpdates.QueueInstances(UploadInfo, ItemIndex, PrimitiveItemInfo);
 
-					// Update Nanite material tables associated with this primitive index.
-					// GPUCULL_TODO: Abstract this also if we ever need Nanite in dynamic primitives
-					if (Scene != nullptr && bNaniteEnabled && UploadInfo.NaniteSceneProxy != nullptr)
-					{
-						check(UploadDataSourceAdapter.bUpdateNaniteMaterialTables);
-						check(UploadInfo.PrimitiveSceneInfo != nullptr);
-						const FPrimitiveSceneInfo* PrimitiveSceneInfo = UploadInfo.PrimitiveSceneInfo;
-						const Nanite::FSceneProxyBase* NaniteSceneProxy = UploadInfo.NaniteSceneProxy;
+						NumLightmapDataUploads += UploadInfo.LightmapUploadCount; // Not thread safe
+						NumInstanceSceneDataUploads += UploadInfo.NumInstanceUploads; // Not thread safe
+						NumInstancePayloadDataUploads += UploadInfo.NumInstancePayloadDataUploads; // Not thread safe
 
-						// Update material depth and hit proxy ID remapping tables.
-						for (int32 NaniteMeshPass = 0; NaniteMeshPass < ENaniteMeshPass::Num; ++NaniteMeshPass)
+						if (Scene != nullptr && bNaniteEnabled && UploadInfo.NaniteSceneProxy != nullptr)
 						{
-							FNaniteMaterialCommands& NaniteMaterials = Scene->NaniteMaterials[NaniteMeshPass];
+							check(UploadDataSourceAdapter.bUpdateNaniteMaterialTables);
+							check(UploadInfo.PrimitiveSceneInfo != nullptr);
+							const FPrimitiveSceneInfo* PrimitiveSceneInfo = UploadInfo.PrimitiveSceneInfo;
+							const Nanite::FSceneProxyBase* NaniteSceneProxy = UploadInfo.NaniteSceneProxy;
 
-							const TArray<uint32>& PassMaterialSlots = PrimitiveSceneInfo->NaniteMaterialSlots[NaniteMeshPass];
-							const TArray<Nanite::FSceneProxyBase::FMaterialSection>& PassMaterials = NaniteSceneProxy->GetMaterialSections();
-							if (PassMaterials.Num() == PassMaterialSlots.Num())
+							// Update material depth and hit proxy ID remapping tables.
+							for (int32 NaniteMeshPass = 0; NaniteMeshPass < ENaniteMeshPass::Num; ++NaniteMeshPass)
 							{
+								FNaniteMaterialCommands& NaniteMaterials = Scene->NaniteMaterials[NaniteMeshPass];
+
+								const TArray<uint32>& PassMaterialSlots = PrimitiveSceneInfo->NaniteMaterialSlots[NaniteMeshPass];
+								const TArray<Nanite::FSceneProxyBase::FMaterialSection>& PassMaterials = NaniteSceneProxy->GetMaterialSections();
+								if (PassMaterials.Num() == PassMaterialSlots.Num())
+								{
 
-								const uint32 TableEntryCount = uint32(NaniteSceneProxy->GetMaterialMaxIndex() + 1);
-								check(TableEntryCount >= uint32(PassMaterials.Num()));
+									const uint32 TableEntryCount = uint32(NaniteSceneProxy->GetMaterialMaxIndex() + 1);
+									check(TableEntryCount >= uint32(PassMaterials.Num()));
 #if WITH_EDITOR
-								const uint32 HitProxyEntryCount = (NaniteMeshPass == ENaniteMeshPass::BasePass) ? TableEntryCount : NANITE_MAX_MATERIALS;
+									const uint32 HitProxyEntryCount = (NaniteMeshPass == ENaniteMeshPass::BasePass) ? TableEntryCount : NANITE_MAX_MATERIALS;
 #endif
 
-								void* MaterialSlotRange = nullptr;
+									void* MaterialSlotRange = nullptr;
 #if WITH_EDITOR
-								void* HitProxyTable = nullptr;
+									void* HitProxyTable = nullptr;
 #endif
-								{
-									FNaniteMaterialCommandsLock NaniteLock(NaniteMaterials, SLT_Write);
-									MaterialSlotRange = NaniteMaterials.GetMaterialSlotPtr(UploadInfo.PrimitiveID, TableEntryCount);
+									{
+										FNaniteMaterialCommandsLock NaniteLock(NaniteMaterials, SLT_Write);
+										MaterialSlotRange = NaniteMaterials.GetMaterialSlotPtr(UploadInfo.PrimitiveID, TableEntryCount);
 #if WITH_EDITOR
-									HitProxyTable = NaniteMaterials.GetHitProxyTablePtr(UploadInfo.PrimitiveID, HitProxyEntryCount);
+										HitProxyTable = NaniteMaterials.GetHitProxyTablePtr(UploadInfo.PrimitiveID, HitProxyEntryCount);
 #endif
-								}
+									}
 
-								uint32* MaterialSlots = static_cast<uint32*>(MaterialSlotRange);
-								for (int32 Entry = 0; Entry < PassMaterialSlots.Num(); ++Entry)
-								{
-									MaterialSlots[PassMaterials[Entry].MaterialIndex] = PassMaterialSlots[Entry];
-								}
+									uint32* MaterialSlots = static_cast<uint32*>(MaterialSlotRange);
+									for (int32 Entry = 0; Entry < PassMaterialSlots.Num(); ++Entry)
+									{
+										MaterialSlots[PassMaterials[Entry].MaterialIndex] = PassMaterialSlots[Entry];
+									}
 
 #if WITH_EDITOR
-								if (NaniteMeshPass == ENaniteMeshPass::BasePass)
-								{
-									const TArray<uint32>& PassHitProxyIds = PrimitiveSceneInfo->NaniteHitProxyIds;
-
-									uint32* HitProxyEntry = static_cast<uint32*>(HitProxyTable);
-									for (int32 Entry = 0; Entry < PassHitProxyIds.Num(); ++Entry)
+									if (NaniteMeshPass == ENaniteMeshPass::BasePass)
 									{
-										HitProxyEntry[PassMaterials[Entry].MaterialIndex] = PassHitProxyIds[Entry];
+										const TArray<uint32>& PassHitProxyIds = PrimitiveSceneInfo->NaniteHitProxyIds;
+
+										uint32* HitProxyEntry = static_cast<uint32*>(HitProxyTable);
+										for (int32 Entry = 0; Entry < PassHitProxyIds.Num(); ++Entry)
+										{
+											HitProxyEntry[PassMaterials[Entry].MaterialIndex] = PassHitProxyIds[Entry];
+										}
 									}
-								}
-								else
-								{
-									// Other passes don't use hit proxies. TODO: Shouldn't even need to do this.
-									uint64* DualHitProxyEntry = static_cast<uint64*>(HitProxyTable);
-									for (uint32 DualEntry = 0; DualEntry < NANITE_MAX_MATERIALS >> 1; ++DualEntry)
+									else
 									{
-										DualHitProxyEntry[DualEntry] = 0;
+										// Other passes don't use hit proxies. TODO: Shouldn't even need to do this.
+										uint64* DualHitProxyEntry = static_cast<uint64*>(HitProxyTable);
+										for (uint32 DualEntry = 0; DualEntry < NANITE_MAX_MATERIALS >> 1; ++DualEntry)
+										{
+											DualHitProxyEntry[DualEntry] = 0;
+										}
 									}
-								}
 #endif
+								}
 							}
 						}
 					}
 				}
-			};
-
-			const bool bShouldUploadViaCreate = GRHISupportsEfficientUploadOnResourceCreation && GGPUSceneInstanceUploadViaCreate != 0;
-			PrimitiveUploadBuffer.SetUploadViaCreate(bShouldUploadViaCreate);
-			InstanceSceneUploadBuffer.SetUploadViaCreate(bShouldUploadViaCreate);
-
-			FUAVTransitionStateScopeHelper PrimitiveDataTransitionHelper(RHICmdList, BufferState.PrimitiveBuffer.UAV, ERHIAccess::Unknown, ERHIAccess::SRVMask);
-
-			const int32 MaxPrimitivesUploads = GetMaxPrimitivesUpdate(NumPrimitiveDataUploads, FPrimitiveSceneShaderData::DataStrideInFloat4s);
-			if (MaxPrimitivesUploads == NumPrimitiveDataUploads)
-			{
-				// One large batch
-				SCOPED_DRAW_EVENTF(RHICmdList, UpdateGPUScene, TEXT("UpdateGPUScene NumPrimitiveDataUploads %u"), NumPrimitiveDataUploads);
-
-				PrimitiveUploadBuffer.Init(NumPrimitiveDataUploads, sizeof(FPrimitiveSceneShaderData::Data), true, TEXT("PrimitiveUploadBuffer"));
-
-				int32 RangeCount = PartitionUpdateRanges(ParallelRanges, NumPrimitiveDataUploads, bExecuteInParallel);
-
-				ParallelFor(RangeCount,
-					[&ProcessPrimitiveFn, &ParallelRanges, RangeCount](int32 RangeIndex)
+				const bool bShouldUploadViaCreate = GRHISupportsEfficientUploadOnResourceCreation && GGPUSceneInstanceUploadViaCreate != 0;
+				PrimitiveUploadBuffer.SetUploadViaCreate(bShouldUploadViaCreate);
+				InstanceSceneUploadBuffer.SetUploadViaCreate(bShouldUploadViaCreate);
+				
+				{
+					SCOPED_NAMED_EVENT(STAT_UpdateGPUScenePrimitives_Par, FColor::Green);
+					QUICK_SCOPE_CYCLE_COUNTER(STAT_UpdateGPUScenePrimitives_Par);
+					// 1. do (potentially) parallel work next
+					ParallelFor(NumPrimitiveDataUploads, [&UploadDataSourceAdapter, this](int32 ItemIndex)
 					{
-						for (int32 ItemIndex = ParallelRanges.Range[RangeIndex].ItemStart; ItemIndex < ParallelRanges.Range[RangeIndex].ItemStart + ParallelRanges.Range[RangeIndex].ItemCount; ++ItemIndex)
+						FPrimitiveUploadInfo UploadInfo;
+						UploadDataSourceAdapter.GetPrimitiveInfo(ItemIndex, UploadInfo);
+
+						FVector4f* DstData = static_cast<FVector4f*>(PrimitiveUploadBuffer.GetRef(ItemIndex));
+						for (uint32 VectorIndex = 0; VectorIndex < FPrimitiveSceneShaderData::DataStrideInFloat4s; ++VectorIndex)
 						{
-							ProcessPrimitiveFn(ItemIndex, RangeCount > 1);
+							DstData[VectorIndex] = UploadInfo.PrimitiveSceneData.Data[VectorIndex];
 						}
-					},
-					RangeCount == 1
-				);
+					}, !bExecuteInParallel);
+				}
 
 				PrimitiveDataTransitionHelper.TransitionTo(ERHIAccess::UAVCompute);
 				PrimitiveUploadBuffer.ResourceUploadTo(RHICmdList, BufferState.PrimitiveBuffer, true);
 			}
-			else
-			{
-				// Break into multiple batches
-				for (int32 PrimitiveOffset = 0; PrimitiveOffset < NumPrimitiveDataUploads; PrimitiveOffset += MaxPrimitivesUploads)
-				{
-					SCOPED_DRAW_EVENTF(RHICmdList, UpdateGPUScene, TEXT("UpdateGPUScene NumPrimitiveDataUploads and Offset = %u %u"), NumPrimitiveDataUploads, PrimitiveOffset);
-
-					PrimitiveUploadBuffer.Init(MaxPrimitivesUploads, sizeof(FPrimitiveSceneShaderData::Data), true, TEXT("PrimitiveUploadBuffer"));
-
-					for (int32 IndexUpdate = 0; (IndexUpdate < MaxPrimitivesUploads) && ((IndexUpdate + PrimitiveOffset) < NumPrimitiveDataUploads); ++IndexUpdate)
-					{
-						const int32 ItemIndex = IndexUpdate + PrimitiveOffset;
-						ProcessPrimitiveFn(ItemIndex, false /* threaded */);
-					}
-
-					PrimitiveDataTransitionHelper.TransitionTo(ERHIAccess::UAVCompute);
-
-					{
-						QUICK_SCOPE_CYCLE_COUNTER(UploadTo);
-						PrimitiveUploadBuffer.ResourceUploadTo(RHICmdList, BufferState.PrimitiveBuffer, true);
-					}
-				}
-			}
 		}
 
 		if (Scene != nullptr)
@@ -1112,160 +1177,142 @@ void FGPUScene::UploadGeneral(FRHICommandListImmediate& RHICmdList, FScene *Scen
 		{
 			if (NumInstancePayloadDataUploads > 0)
 			{
-				InstancePayloadUploadBuffer.Init(NumInstancePayloadDataUploads, sizeof(FVector4), true, TEXT("InstancePayloadUploadBuffer"));
+				InstancePayloadUploadBuffer.InitPreSized(NumInstancePayloadDataUploads, sizeof(FVector4f), true, TEXT("InstancePayloadUploadBuffer"));
 			}
 
 			// Upload instancing data for the scene.
 			if (NumInstanceSceneDataUploads > 0)
 			{
-				InstanceSceneUploadBuffer.Init(NumInstanceSceneDataUploads * InstanceSceneDataNumArrays, sizeof(FVector4f), true, TEXT("InstanceSceneUploadBuffer"));
+				SCOPED_NAMED_EVENT(STAT_UpdateGPUSceneInstances, FColor::Green);
+				QUICK_SCOPE_CYCLE_COUNTER(STAT_UpdateGPUSceneInstances);
 
-				if (NumPrimitiveDataUploads > 0)
-				{
-					// Note: we iterate over the primitives, whether they have instances or not (which is a bit wasteful) but this is the way we currently get to the instance data.
-					// GPUCULL_TODO: move instance data ownership to GPU-scene such that it can be put in a compact list or something, and be tracked independent of primitives?
-					int32 RangeCount = PartitionUpdateRanges(ParallelRanges, NumPrimitiveDataUploads, bExecuteInParallel);
+				InstanceSceneUploadBuffer.InitPreSized(NumInstanceSceneDataUploads * InstanceSceneDataNumArrays, sizeof(FVector4f), true, TEXT("InstanceSceneUploadBuffer"));
 
-					// Upload any out of date instance slots.
-					ParallelFor(RangeCount,
-						[&UploadDataSourceAdapter, this, &ParallelRanges, RangeCount, InstanceSceneDataNumArrays, NumInstancePayloadDataUploads, &BufferState](int32 RangeIndex)
+				if (!InstanceUpdates.UpdateBatches.IsEmpty())
+				{
+					ParallelFor(InstanceUpdates.UpdateBatches.Num(), [this, &InstanceUpdates, &Scene, &UploadDataSourceAdapter, NumInstancePayloadDataUploads, &BufferState](int32 BatchIndex)
+					{
+						const FInstanceUploadBatch Batch = InstanceUpdates.UpdateBatches[BatchIndex];
+						for (int32 BatchItemIndex = 0; BatchItemIndex < Batch.NumItems; ++BatchItemIndex)
 						{
-							for (int32 ItemIndex = ParallelRanges.Range[RangeIndex].ItemStart; ItemIndex < ParallelRanges.Range[RangeIndex].ItemStart + ParallelRanges.Range[RangeIndex].ItemCount; ++ItemIndex)
+							const FInstanceUploadBatch::FItem Item = InstanceUpdates.UpdateBatchItems[Batch.FirstItem + BatchItemIndex];
+
+							const int32 ItemIndex = Item.ItemIndex;
+							FInstanceUploadInfo UploadInfo;
+							UploadDataSourceAdapter.GetInstanceInfo(ItemIndex, UploadInfo);
+							FInstanceBatcher::FPrimitiveItemInfo PrimitiveItemInfo = InstanceUpdates.PerPrimitiveItemInfo[ItemIndex];
+
+							check(NumInstancePayloadDataUploads > 0 || UploadInfo.InstancePayloadDataStride == 0); // Sanity check
+
+							for (int32 BatchInstanceIndex = 0; BatchInstanceIndex < Item.NumInstances; ++BatchInstanceIndex)
 							{
-								FInstanceUploadInfo UploadInfo;
-								if (UploadDataSourceAdapter.GetInstanceInfo(ItemIndex, UploadInfo))
+								int32 InstanceIndex = Item.FirstInstance + BatchInstanceIndex;
+								const FPrimitiveInstance& SceneData = UploadInfo.PrimitiveInstances[InstanceIndex];
+
+								// Directly embedded in instance scene data
+								const float RandomID = (UploadInfo.PayloadDataFlags & INSTANCE_SCENE_DATA_FLAG_HAS_RANDOM) ? UploadInfo.InstanceRandomID[InstanceIndex] : 0.0f;
+
+								FInstanceSceneShaderData InstanceSceneData;
+								InstanceSceneData.Build(
+									UploadInfo.PrimitiveID,
+									InstanceIndex,
+									UploadInfo.PayloadDataFlags,
+									UploadInfo.LastUpdateSceneFrameNumber,
+									UploadInfo.InstanceCustomDataCount,
+									RandomID,
+									SceneData.LocalToPrimitive,
+									UploadInfo.PrimitiveToWorld,
+									UploadInfo.PrevPrimitiveToWorld
+								);
+
+								// RefIndex* BufferState.InstanceSceneDataSOAStride + UploadInfo.InstanceSceneDataOffset + InstanceIndex
+								const uint32 UploadInstanceItemOffset = (PrimitiveItemInfo.InstanceSceneDataUploadOffset + InstanceIndex) * InstanceSceneDataNumArrays;
+
+								for (uint32 RefIndex = 0; RefIndex < InstanceSceneDataNumArrays; ++RefIndex)
 								{
-									check(NumInstancePayloadDataUploads > 0 || UploadInfo.InstancePayloadDataStride == 0); // Sanity check
+									FVector4f* DstVector = static_cast<FVector4f*>(InstanceSceneUploadBuffer.Set_GetRef(UploadInstanceItemOffset + RefIndex, RefIndex * BufferState.InstanceSceneDataSOAStride + UploadInfo.InstanceSceneDataOffset + InstanceIndex));
+									*DstVector = InstanceSceneData.Data[RefIndex];
+								}
 
-									// Update each primitive instance with current data.
-									for (int32 InstanceIndex = 0; InstanceIndex < UploadInfo.PrimitiveInstances.Num(); ++InstanceIndex)
-									{
-										const FPrimitiveInstance& SceneData = UploadInfo.PrimitiveInstances[InstanceIndex];
-
-										// Directly embedded in instance scene data
-										const float RandomID = UploadInfo.PayloadDataFlags & INSTANCE_SCENE_DATA_FLAG_HAS_RANDOM ? UploadInfo.InstanceRandomID[InstanceIndex] : 0.0f;
-
-										FInstanceSceneShaderData InstanceSceneData;
-										InstanceSceneData.Build(
-											UploadInfo.PrimitiveID,
-											InstanceIndex,
-											UploadInfo.PayloadDataFlags,
-											UploadInfo.LastUpdateSceneFrameNumber,
-											UploadInfo.InstanceCustomDataCount,
-											RandomID,
-											SceneData.LocalToPrimitive,
-											UploadInfo.PrimitiveToWorld,
-											UploadInfo.PrevPrimitiveToWorld
-										);
-
-										int32 PayloadDataStart = UploadInfo.InstancePayloadDataOffset + (InstanceIndex * UploadInfo.InstancePayloadDataStride); 
-										
-										void* DstRefs[FInstanceSceneShaderData::DataStrideInFloat4s];
-										FVector4* DstPayloadData = nullptr;
-
-										if (RangeCount > 1)
-										{
-											PrimitiveUploadBufferCS.Lock();
-										}
+								// BEGIN PAYLOAD
+								if (UploadInfo.InstancePayloadDataStride > 0)
+								{
+									const uint32 UploadPayloadItemOffset = PrimitiveItemInfo.InstancePayloadDataUploadOffset + InstanceIndex * UploadInfo.InstancePayloadDataStride;
 
-										for (uint32 RefIndex = 0; RefIndex < InstanceSceneDataNumArrays; ++RefIndex)
-										{
-											DstRefs[RefIndex] = InstanceSceneUploadBuffer.Add_GetRef(RefIndex * BufferState.InstanceSceneDataSOAStride + UploadInfo.InstanceSceneDataOffset + InstanceIndex);
-										}
+									int32 PayloadDataStart = UploadInfo.InstancePayloadDataOffset + (InstanceIndex * UploadInfo.InstancePayloadDataStride);
+									FVector4f* DstPayloadData = static_cast<FVector4f*>(InstancePayloadUploadBuffer.Set_GetRef(UploadPayloadItemOffset, PayloadDataStart, UploadInfo.InstancePayloadDataStride));
+									TArrayView<FVector4f> InstancePayloadData = TArrayView<FVector4f>(DstPayloadData, UploadInfo.InstancePayloadDataStride);
 
-										if (UploadInfo.InstancePayloadDataStride > 0)
-										{
-											DstPayloadData = static_cast<FVector4*>(InstancePayloadUploadBuffer.Add_GetRef(PayloadDataStart, UploadInfo.InstancePayloadDataStride)); // Not thread safe
-										}
+									int32 PayloadPosition = 0;
 
-										if (RangeCount > 1)
-										{
-											PrimitiveUploadBufferCS.Unlock();
-										}
+									if (UploadInfo.PayloadDataFlags & (INSTANCE_SCENE_DATA_FLAG_HAS_HIERARCHY_OFFSET | INSTANCE_SCENE_DATA_FLAG_HAS_LOCAL_BOUNDS))
+									{
+										const uint32 InstanceHierarchyOffset = (UploadInfo.PayloadDataFlags & INSTANCE_SCENE_DATA_FLAG_HAS_HIERARCHY_OFFSET) ? UploadInfo.InstanceHierarchyOffset[InstanceIndex] : 0;
+										InstancePayloadData[PayloadPosition].X = *(const float*)&InstanceHierarchyOffset;
 
-										for (uint32 RefIndex = 0; RefIndex < InstanceSceneDataNumArrays; ++RefIndex)
+										if (UploadInfo.PayloadDataFlags & INSTANCE_SCENE_DATA_FLAG_HAS_LOCAL_BOUNDS)
 										{
-											FVector4f* DstVector = static_cast<FVector4f*>(DstRefs[RefIndex]);
-											*DstVector = InstanceSceneData.Data[RefIndex];
+											check(UploadInfo.InstanceLocalBounds.Num() == UploadInfo.PrimitiveInstances.Num());
+											const FRenderBounds& InstanceLocalBounds = UploadInfo.InstanceLocalBounds[InstanceIndex];
+											const FVector3f BoundsOrigin = InstanceLocalBounds.GetCenter();
+											const FVector3f BoundsExtent = InstanceLocalBounds.GetExtent();
+
+											InstancePayloadData[PayloadPosition + 0].Y = *(const float*)&BoundsOrigin.X;
+											InstancePayloadData[PayloadPosition + 0].Z = *(const float*)&BoundsOrigin.Y;
+											InstancePayloadData[PayloadPosition + 0].W = *(const float*)&BoundsOrigin.Z;
+
+											InstancePayloadData[PayloadPosition + 1].X = *(const float*)&BoundsExtent.X;
+											InstancePayloadData[PayloadPosition + 1].Y = *(const float*)&BoundsExtent.Y;
+											InstancePayloadData[PayloadPosition + 1].Z = *(const float*)&BoundsExtent.Z;
+											// .W = Unused
 										}
 
-										// BEGIN PAYLOAD
+										PayloadPosition += (UploadInfo.PayloadDataFlags & INSTANCE_SCENE_DATA_FLAG_HAS_LOCAL_BOUNDS) ? 2 : 1;
+									}
+
+									if (UploadInfo.PayloadDataFlags & INSTANCE_SCENE_DATA_FLAG_HAS_DYNAMIC_DATA)
+									{
+										check(UploadInfo.InstanceDynamicData.Num() == UploadInfo.PrimitiveInstances.Num());
+										const FRenderTransform PrevLocalToWorld = UploadInfo.InstanceDynamicData[InstanceIndex].ComputePrevLocalToWorld(UploadInfo.PrevPrimitiveToWorld);
+									#if INSTANCE_COMPRESSED_TRANSFORMS
+										check(PayloadPosition + 1 < InstancePayloadData.Num()); // Sanity check
+										FCompressedTransform CompressedPrevLocalToWorld(PrevLocalToWorld);
+										InstancePayloadData[PayloadPosition + 0] = *(const FVector4*)&CompressedPrevLocalToWorld.Rotation[0];
+										InstancePayloadData[PayloadPosition + 1] = *(const FVector3f*)&CompressedPrevLocalToWorld.Translation;
+										PayloadPosition += 2;
+									#else
+										// Note: writes 3x float4s
+										check(PayloadPosition + 2 < InstancePayloadData.Num()); // Sanity check
+										PrevLocalToWorld.To3x4MatrixTranspose((float*)&InstancePayloadData[PayloadPosition]);
+										PayloadPosition += 3;
+									#endif
+									}
+
+									if (UploadInfo.PayloadDataFlags & INSTANCE_SCENE_DATA_FLAG_HAS_LIGHTSHADOW_UV_BIAS)
+									{
+										check(UploadInfo.InstanceLightShadowUVBias.Num() == UploadInfo.PrimitiveInstances.Num());
+										InstancePayloadData[PayloadPosition] = UploadInfo.InstanceLightShadowUVBias[InstanceIndex];
+										PayloadPosition += 1;
+									}
 
-										if (UploadInfo.InstancePayloadDataStride > 0)
+									if (UploadInfo.PayloadDataFlags & INSTANCE_SCENE_DATA_FLAG_HAS_CUSTOM_DATA)
+									{
+										check(PayloadPosition + (UploadInfo.InstanceCustomDataCount >> 2u) <= InstancePayloadData.Num());
+										const float* CustomDataGetPtr = &UploadInfo.InstanceCustomData[(InstanceIndex * UploadInfo.InstanceCustomDataCount)];
+										float* CustomDataPutPtr = (float*)&InstancePayloadData[PayloadPosition];
+										for (uint32 FloatIndex = 0; FloatIndex < uint32(UploadInfo.InstanceCustomDataCount); ++FloatIndex)
 										{
-											TArrayView<FVector4> InstancePayloadData = TArrayView<FVector4>(DstPayloadData, UploadInfo.InstancePayloadDataStride);
-
-											int32 PayloadPosition = 0;
-
-											if (UploadInfo.PayloadDataFlags & (INSTANCE_SCENE_DATA_FLAG_HAS_HIERARCHY_OFFSET | INSTANCE_SCENE_DATA_FLAG_HAS_LOCAL_BOUNDS))
-											{
-												const uint32 InstanceHierarchyOffset = (UploadInfo.PayloadDataFlags & INSTANCE_SCENE_DATA_FLAG_HAS_HIERARCHY_OFFSET) ? UploadInfo.InstanceHierarchyOffset[InstanceIndex] : 0;
-												InstancePayloadData[PayloadPosition].X = *(const float*)&InstanceHierarchyOffset;
-
-												if (UploadInfo.PayloadDataFlags & INSTANCE_SCENE_DATA_FLAG_HAS_LOCAL_BOUNDS)
-												{
-													check(UploadInfo.InstanceLocalBounds.Num() == UploadInfo.PrimitiveInstances.Num());
-													const FRenderBounds& InstanceLocalBounds = UploadInfo.InstanceLocalBounds[InstanceIndex];
-													const FVector3f BoundsOrigin = InstanceLocalBounds.GetCenter();
-													const FVector3f BoundsExtent = InstanceLocalBounds.GetExtent();
-
-													InstancePayloadData[PayloadPosition + 0].Y = *(const float*)&BoundsOrigin.X;
-													InstancePayloadData[PayloadPosition + 0].Z = *(const float*)&BoundsOrigin.Y;
-													InstancePayloadData[PayloadPosition + 0].W = *(const float*)&BoundsOrigin.Z;
-
-													InstancePayloadData[PayloadPosition + 1].X = *(const float*)&BoundsExtent.X;
-													InstancePayloadData[PayloadPosition + 1].Y = *(const float*)&BoundsExtent.Y;
-													InstancePayloadData[PayloadPosition + 1].Z = *(const float*)&BoundsExtent.Z;
-													// .W = Unused
-												}
-
-												PayloadPosition += (UploadInfo.PayloadDataFlags & INSTANCE_SCENE_DATA_FLAG_HAS_LOCAL_BOUNDS) ? 2 : 1;
-											}
-
-											if (UploadInfo.PayloadDataFlags & INSTANCE_SCENE_DATA_FLAG_HAS_DYNAMIC_DATA)
-											{
-												check(UploadInfo.InstanceDynamicData.Num() == UploadInfo.PrimitiveInstances.Num());
-												const FRenderTransform PrevLocalToWorld = UploadInfo.InstanceDynamicData[InstanceIndex].ComputePrevLocalToWorld(UploadInfo.PrevPrimitiveToWorld);
-											#if INSTANCE_COMPRESSED_TRANSFORMS
-												check(PayloadPosition + 1 < InstancePayloadData.Num()); // Sanity check
-												FCompressedTransform CompressedPrevLocalToWorld(PrevLocalToWorld);
-												InstancePayloadData[PayloadPosition + 0] = *(const FVector4*)&CompressedPrevLocalToWorld.Rotation[0];
-												InstancePayloadData[PayloadPosition + 1] = *(const FVector3f*)&CompressedPrevLocalToWorld.Translation;
-												PayloadPosition += 2;
-											#else
-												// Note: writes 3x float4s
-												check(PayloadPosition + 2 < InstancePayloadData.Num()); // Sanity check
-												PrevLocalToWorld.To3x4MatrixTranspose((float*)&InstancePayloadData[PayloadPosition]);
-												PayloadPosition += 3;
-											#endif
-											}
-
-											if (UploadInfo.PayloadDataFlags & INSTANCE_SCENE_DATA_FLAG_HAS_LIGHTSHADOW_UV_BIAS)
-											{
-												check(UploadInfo.InstanceLightShadowUVBias.Num() == UploadInfo.PrimitiveInstances.Num());
-												InstancePayloadData[PayloadPosition] = UploadInfo.InstanceLightShadowUVBias[InstanceIndex];
-												PayloadPosition += 1;
-											}
-
-											if (UploadInfo.PayloadDataFlags & INSTANCE_SCENE_DATA_FLAG_HAS_CUSTOM_DATA)
-											{
-												check(PayloadPosition + (UploadInfo.InstanceCustomDataCount >> 2u) <= InstancePayloadData.Num());
-												const float* CustomDataGetPtr = &UploadInfo.InstanceCustomData[(InstanceIndex * UploadInfo.InstanceCustomDataCount)];
-												float* CustomDataPutPtr = (float*)&InstancePayloadData[PayloadPosition];
-												for (uint32 FloatIndex = 0; FloatIndex < uint32(UploadInfo.InstanceCustomDataCount); ++FloatIndex)
-												{
-													CustomDataPutPtr[FloatIndex] = CustomDataGetPtr[FloatIndex];
-												}
-											}
+											CustomDataPutPtr[FloatIndex] = CustomDataGetPtr[FloatIndex];
 										}
-
-										// END PAYLOAD
 									}
 								}
+
+								// END PAYLOAD
+
+
 							}
-						},
-						RangeCount == 1
-					);
+						}
+					}, !bExecuteInParallel);
 				}
 
 				if (NumInstancePayloadDataUploads > 0)
@@ -1310,6 +1357,9 @@ void FGPUScene::UploadGeneral(FRHICommandListImmediate& RHICmdList, FScene *Scen
 
 			if (NumLightmapDataUploads > 0)
 			{
+				SCOPED_NAMED_EVENT(STAT_UpdateGPUSceneLightMaps, FColor::Green);
+				QUICK_SCOPE_CYCLE_COUNTER(STAT_UpdateGPUSceneLightMaps);
+
 				FUAVTransitionStateScopeHelper LightMapTransitionHelper(RHICmdList, BufferState.LightmapDataBuffer.UAV, ERHIAccess::Unknown, ERHIAccess::SRVMask);
 
 				// GPUCULL_TODO: This code is wrong: the intention is to break it up into batches such that the uploaded data fits in the max buffer size
@@ -1376,36 +1426,53 @@ struct FUploadDataSourceAdapterDynamicPrimitives
 		, PrimitiveIDStartOffset(InPrimitiveIDStartOffset)
 		, InstanceIDStartOffset(InInstanceIDStartOffset)
 		, SceneFrameNumber(InSceneFrameNumber)
-	{}
+	{
+		// Need to create this explicit for optimizing the common path
+		PrimitivesIds.SetNumUninitialized(PrimitiveShaderData.Num());
+		for (int32 Index = 0; Index < PrimitivesIds.Num(); ++Index)
+		{
+			PrimitivesIds[Index] = uint32(PrimitiveIDStartOffset + Index);
+		}
+	}
 
 	FORCEINLINE int32 NumPrimitivesToUpload() const
 	{ 
 		return PrimitiveShaderData.Num();
 	}
 
-	FORCEINLINE bool GetPrimitiveInfo(int32 ItemIndex, FPrimitiveUploadInfo& PrimitiveUploadInfo) const
+
+	FORCEINLINE TArrayView<const uint32> GetItemPrimitiveIds() const
+	{
+		return TArrayView<const uint32>(PrimitivesIds.GetData(), PrimitivesIds.Num());
+	}
+
+
+	FORCEINLINE void GetPrimitiveInfoHeader(int32 ItemIndex, FPrimitiveUploadInfoHeader& PrimitiveUploadInfo) const
 	{
 		PrimitiveUploadInfo.LightmapUploadCount = 0;
 		PrimitiveUploadInfo.NaniteSceneProxy = nullptr;
 		PrimitiveUploadInfo.PrimitiveSceneInfo = nullptr;
 
-		if (ItemIndex < PrimitiveShaderData.Num())
-		{
-			// Needed to ensure the link back to instance list is up to date
-			FPrimitiveUniformShaderParameters Tmp = PrimitiveShaderData[ItemIndex];
-			Tmp.InstanceSceneDataOffset = InstanceIDStartOffset + ItemIndex;
-			Tmp.NumInstanceSceneDataEntries = 1;
+		check(ItemIndex < PrimitiveShaderData.Num());
 
-			PrimitiveUploadInfo.InstanceSceneDataOffset = InstanceIDStartOffset + ItemIndex;
-			PrimitiveUploadInfo.InstanceSceneDataUploads = 1;
-			PrimitiveUploadInfo.InstancePayloadDataUploads = 0;
-			PrimitiveUploadInfo.PrimitiveID = PrimitiveIDStartOffset + ItemIndex;
-			PrimitiveUploadInfo.PrimitiveSceneData = FPrimitiveSceneShaderData(Tmp);
+		PrimitiveUploadInfo.NumInstanceUploads = 1;
+		PrimitiveUploadInfo.NumInstancePayloadDataUploads = 0;
+		PrimitiveUploadInfo.PrimitiveID = PrimitiveIDStartOffset + ItemIndex;
+	}
 
-			return true;
-		}
+	FORCEINLINE void GetPrimitiveInfo(int32 ItemIndex, FPrimitiveUploadInfo& PrimitiveUploadInfo) const
+	{
+		GetPrimitiveInfoHeader(ItemIndex, PrimitiveUploadInfo);
 
-		return false;
+		// Needed to ensure the link back to instance list is up to date
+		FPrimitiveUniformShaderParameters Tmp = PrimitiveShaderData[ItemIndex];
+		Tmp.InstanceSceneDataOffset = InstanceIDStartOffset + ItemIndex;
+		Tmp.NumInstanceSceneDataEntries = 1;
+
+		PrimitiveUploadInfo.NumInstanceUploads = 1;
+		PrimitiveUploadInfo.NumInstancePayloadDataUploads = 0;
+		PrimitiveUploadInfo.PrimitiveID = PrimitiveIDStartOffset + ItemIndex;
+		PrimitiveUploadInfo.PrimitiveSceneData = FPrimitiveSceneShaderData(Tmp);
 	}
 
 	FORCEINLINE bool GetInstanceInfo(int32 ItemIndex, FInstanceUploadInfo& InstanceUploadInfo) const
@@ -1448,6 +1515,7 @@ struct FUploadDataSourceAdapterDynamicPrimitives
 	const int32 PrimitiveIDStartOffset;
 	const int32 InstanceIDStartOffset;
 	const uint32 SceneFrameNumber;
+	TArray<uint32, SceneRenderingAllocator> PrimitivesIds;
 };
 
 void FGPUScene::UploadDynamicPrimitiveShaderDataForViewInternal(FRDGBuilder& GraphBuilder, FScene *Scene, FViewInfo& View)
