UE5:  Optimization to virtual texture allocator.

Partially free blocks are organized into multiple lists by the alignment of sub-allocation that could fit in a block, allowing the allocation loop to completely avoid scanning irrelevant blocks.  The sub-allocation alignments that could fit are determined efficiently using an 8x8 bitmap packed into a single 64-bit word, maintained as nodes are allocated and freed.  Improves worst case performance by a factor of 20.

#rnx

diff --git a/Engine/Source/Runtime/Renderer/Private/VT/VirtualTextureAllocator.cpp b/Engine/Source/Runtime/Renderer/Private/VT/VirtualTextureAllocator.cpp
--- a/Engine/Source/Runtime/Renderer/Private/VT/VirtualTextureAllocator.cpp
+++ b/Engine/Source/Runtime/Renderer/Private/VT/VirtualTextureAllocator.cpp
@@ -51,10 +51,10 @@ void FVirtualTextureAllocator::Initialize(uint32 MaxSize)
 	for (uint8 i = 0; i < vLogSize; i++)
 	{
 		FreeList[i] = 0xffff;
-		PartiallyFreeList[i] = 0xffff;
+		FMemory::Memset(&PartiallyFreeList[i], 0xff, sizeof(FPartiallyFreeMip));
 	}
 	FreeList[vLogSize] = 0;
-	PartiallyFreeList[vLogSize] = 0xffff;
+	FMemory::Memset(&PartiallyFreeList[vLogSize], 0xff, sizeof(FPartiallyFreeMip));
 
 	// Init global free list
 	GlobalFreeList = 0xffff;
@@ -193,7 +193,8 @@ void FVirtualTextureAllocator::SubdivideBlock(uint32 ParentIndex)
 	// Only free blocks can be subdivided, move to the partially free list
 	check(AddressBlocks[ParentIndex].FirstChild == 0xffff);
 	UnlinkFreeList(FreeList[vParentLogSize], EBlockState::FreeList, ParentIndex);
-	LinkFreeList(PartiallyFreeList[vParentLogSize], EBlockState::PartiallyFreeList, ParentIndex);
+	AddressBlocks[ParentIndex].FreeMip = 0;
+	LinkFreeList(PartiallyFreeList[vParentLogSize].Mips[0], EBlockState::PartiallyFreeList, ParentIndex);
 
 	const uint32 vAddress = AddressBlocks[ParentIndex].vAddress;
 	const int32 SortedIndex = FindAddressBlock(vAddress);
@@ -239,6 +240,30 @@ void FVirtualTextureAllocator::SubdivideBlock(uint32 ParentIndex)
 	}
 }
 
+void FVirtualTextureAllocator::FreeMipUpdateParents(uint16 ParentIndex)
+{
+	for (uint32 ParentDepth = 0; ParentDepth < PartiallyFreeMipDepth; ParentDepth++)
+	{
+		if (ParentIndex == 0xffff)
+		{
+			break;
+		}
+
+		FAddressBlock& ParentBlock = AddressBlocks[ParentIndex];
+		uint8 OldFreeMip = ParentBlock.FreeMip;
+		uint8 NewFreeMip = ComputeFreeMip(ParentIndex);
+
+		if (NewFreeMip != OldFreeMip)
+		{
+			UnlinkFreeList(PartiallyFreeList[ParentBlock.vLogSize].Mips[OldFreeMip], EBlockState::PartiallyFreeList, ParentIndex);
+			ParentBlock.FreeMip = NewFreeMip;
+			LinkFreeList(PartiallyFreeList[ParentBlock.vLogSize].Mips[NewFreeMip], EBlockState::PartiallyFreeList, ParentIndex);
+		}
+
+		ParentIndex = ParentBlock.Parent;
+	}
+}
+
 void FVirtualTextureAllocator::MarkBlockAllocated(uint32 Index, uint32 vAllocatedTileX0, uint32 vAllocatedTileY0, FAllocatedVirtualTexture* VT)
 {
 	FAddressBlock* AllocBlock = &AddressBlocks[Index];
@@ -283,6 +308,8 @@ void FVirtualTextureAllocator::MarkBlockAllocated(uint32 Index, uint32 vAllocate
 
 			AllocBlock->VT = VT;
 			AllocBlock->State = EBlockState::AllocatedTexture;
+
+			FreeMipUpdateParents(AllocBlock->Parent);
 		}
 		else
 		{
@@ -407,14 +434,22 @@ uint32 FVirtualTextureAllocator::Alloc(FAllocatedVirtualTexture* VT)
 		}
 	}
 
-	// Look for a partially allocated block that has room for this allocation
-	// Only need to check partially allocated blocks of the correct size
-	// Larger partially allocated blocks will contain a child block that's completely free, that will already be discovered by the above search
+	// Look for a partially allocated block that has room for this allocation.  Only need to check partially allocated blocks
+	// of the correct size, which could allocate a sub-block of MaxLevel alignment.  Larger partially allocated blocks will
+	// contain a child block that's completely free, that will already be discovered by the above search.
+	check((uint32)vLogMaxSize >= MaxLevel);
+	uint32 FreeMipCount = FMath::Min(3u, vLogMaxSize - MaxLevel);
+
+	for (uint32 FreeMipIndex = 0; FreeMipIndex <= FreeMipCount; FreeMipIndex++)
 	{
-		uint16 FreeIndex = PartiallyFreeList[vLogMaxSize];
+		uint16 FreeIndex = PartiallyFreeList[vLogMaxSize].Mips[FreeMipIndex];
 		while (FreeIndex != 0xffff)
 		{
 			FAddressBlock& AllocBlock = AddressBlocks[FreeIndex];
+
+			// TODO:  consider making "ComputeFreeMip" test a checkSlow once system is proven reliable
+			check((AllocBlock.FreeMip == FreeMipIndex) && (AllocBlock.FreeMip == ComputeFreeMip(FreeIndex)));
+
 			if (AllocBlock.vAddress < vAddress)
 			{
 				check(AllocBlock.State == EBlockState::PartiallyFreeList);
@@ -497,6 +532,11 @@ void FVirtualTextureAllocator::Free(FAllocatedVirtualTexture* VT)
 			AddressBlock.State = EBlockState::None;
 			AddressBlock.VT = nullptr;
 
+			// TODO (perf):  This is slightly wasteful for textures composed of multiple blocks,
+			// as it will update the same parents more than once.  Many textures only have one
+			// block, so perhaps this isn't a large overhead.
+			FreeMipUpdateParents(AddressBlock.Parent);
+
 			check(NumAllocations > 0u);
 			--NumAllocations;
 
@@ -526,7 +566,7 @@ void FVirtualTextureAllocator::FreeAddressBlock(uint32 Index, bool bTopLevelBloc
 	else
 	{
 		// Block was freed by consolidating children
-		UnlinkFreeList(PartiallyFreeList[AddressBlock.vLogSize], EBlockState::PartiallyFreeList, Index);
+		UnlinkFreeList(PartiallyFreeList[AddressBlock.vLogSize].Mips[AddressBlock.FreeMip], EBlockState::PartiallyFreeList, Index);
 	}
 
 	check(AddressBlock.VT == nullptr);
@@ -598,6 +638,136 @@ void FVirtualTextureAllocator::FreeAddressBlock(uint32 Index, bool bTopLevelBloc
 	}
 }
 
+void FVirtualTextureAllocator::RecurseComputeFreeMip(uint16 BlockIndex, uint32 Depth, uint64& IoBlockMap) const
+{
+	const FAddressBlock& Block = AddressBlocks[BlockIndex];
+
+	// Add children first...
+	if (Depth < 3)
+	{
+		uint32 ChildIndex = Block.FirstChild;
+		while (ChildIndex != 0xffff)
+		{
+			RecurseComputeFreeMip(ChildIndex, Depth + 1, IoBlockMap);
+			ChildIndex = AddressBlocks[ChildIndex].NextSibling;
+		}
+	}
+
+	if (Block.State == EBlockState::AllocatedTexture)
+	{
+		// Bit mask of pixels covered by a block of the given log2 size.  Think of each byte
+		// as a row of 8 single bit pixels, moving left in bits representing increasing X, and bytes
+		// increasing Y, similar to how a linear 2D array of texels is usually arranged.  Here's a
+		// diagram of pixels that are covered by increasing powers of 2, and the corresponding bytes,
+		// to show where the entries in the table come from:
+		//
+		//								1x1     2x2     4x4     8x8
+		//        1 2 4 4 8 8 8 8		0x01    0x03    0x0f    0xff
+		//        2 2 4 4 8 8 8 8		0x00    0x03    0x0f    0xff
+		//        4 4 4 4 8 8 8 8		0x00    0x00    0x0f    0xff
+		//        4 4 4 4 8 8 8 8		0x00    0x00    0x0f    0xff
+		//        8 8 8 8 8 8 8 8		0x00    0x00    0x00    0xff
+		//        8 8 8 8 8 8 8 8		0x00    0x00    0x00    0xff
+		//        8 8 8 8 8 8 8 8		0x00    0x00    0x00    0xff
+		//        8 8 8 8 8 8 8 8		0x00    0x00    0x00    0xff
+		//
+		static const uint64 BlockMaskByDepth[4] =
+		{
+			0xffffffffffffffffull,		// 8x8 block
+			0x000000000f0f0f0full,		// 4x4 block
+			0x0000000000000303ull,		// 2x2 block
+			0x0000000000000001ull,		// 1x1 block
+		};
+
+		// Absolute address
+		uint32 X = FMath::ReverseMortonCode2(Block.vAddress);
+		uint32 Y = FMath::ReverseMortonCode2(Block.vAddress >> 1);
+
+		// Parent block relative address.  Let's say the original parent block was 32 by 32 -- these
+		// bit mask operations would mask out the bottom 5 bits of the address for any given block,
+		// producing offsets in the range [0..32).  For the original parent, Depth will be zero, and
+		// Block.vLogSize would be 5, so the mask generated would be (1 << (5 + 0)) - 1 == 0x1f.  For
+		// the next child, Block.vLogSize would be 4, and depth 1, generating the same 0x1f mask, and
+		// so on.
+		uint32 RelativeX = X & ((1 << (Block.vLogSize + Depth)) - 1);
+		uint32 RelativeY = Y & ((1 << (Block.vLogSize + Depth)) - 1);
+
+		// Mapping address (8x8).  For our 32x32 case, we want to map our [0..32) range to [0..8), which
+		// requires us to divide by 32 and multiply by 8.  In bit shifts, that means a right shift by
+		// Block.vLogSize + Depth == 5, and a left shift by 3.  If the original tile is smaller than 8x8,
+		// the shift could go negative, so we need to handle that case.
+		int32 MapShift = (Block.vLogSize + Depth - 3);
+		uint32 MapX;
+		uint32 MapY;
+		if (MapShift >= 0)
+		{
+			MapX = RelativeX >> MapShift;
+			MapY = RelativeY >> MapShift;
+		}
+		else
+		{
+			MapX = RelativeX << (-MapShift);
+			MapY = RelativeY << (-MapShift);
+		}
+
+		// Set bits in our bitmap
+		uint32 BitIndex = MapX + MapY * 8;
+		IoBlockMap |= BlockMaskByDepth[Depth] << BitIndex;
+	}
+}
+
+uint32 FVirtualTextureAllocator::ComputeFreeMip(uint16 BlockIndex) const
+{
+	// First we need to generate a map of the block, where data is allocated.  This is an
+	// 8x8 pixel map in bits in a single 64-bit word.
+	// 
+	// TODO (perf):  It might be interesting to actually store this map in the block itself,
+	// since we've spent the time computing it, as you could use it to early out in
+	// "TestAllocation".
+	uint64 BlockMap = 0;
+	RecurseComputeFreeMip(BlockIndex, 0, BlockMap);
+
+	// Mapping that specifies pixels that need to be covered by child blocks to block any allocation at the
+	// given alignment.  If the corner pixel at a given resolution alone is covered, we can't allocate,
+	// because aligned block addresses always start at pixel corners.  But if not, there is potential to
+	// squeeze an allocation in there (at least a 1x1 if nothing else).  For finer granularity alignments,
+	// we need to check multiple pixel corners to cover all allocation offsets.  Allocation offsets
+	// correspond to steps by the "vAddressAlignment" variable in FVirtualTextureAllocator::Alloc.  Here's
+	// a diagram of the pixels we are looking at, and the corresponding bytes.  X increases as you go left
+	// in bits, and Y increases as you go left in bytes:
+	//
+	//								1x1     2x2     4x4     8x8
+	//        8 1 2 1 4 1 2 1		0xff    0x33    0x11    0x01
+	//        1 1 1 1 1 1 1 1		0xff    0x00    0x00    0x00
+	//        2 1 2 1 2 1 2 1		0xff    0x33    0x00    0x00
+	//        1 1 1 1 1 1 1 1		0xff    0x00    0x00    0x00
+	//        4 1 2 1 4 1 2 1		0xff    0x33    0x11    0x00
+	//        1 1 1 1 1 1 1 1		0xff    0x00    0x00    0x00
+	//        2 1 2 1 2 1 2 1		0xff    0x33    0x00    0x00
+	//        1 1 1 1 1 1 1 1		0xff    0x00    0x00    0x00
+	//
+	static const uint64 BlockOverlapByDepth[4] =
+	{
+		0x0000000000000001ull,		// 8x8 block
+		0x0000001100000011ull,		// 4x4 block
+		0x0033003300330033ull,		// 2x2 block
+		0xffffffffffffffffull,		// 1x1 block (unused by code, just here for reference)
+	};
+
+	uint32 FreeMip;
+	for (FreeMip = 0; FreeMip < 3; FreeMip++)
+	{
+		// Are all the corner pixels at this alignment resolution covered?  If not,
+		// break out of the loop.
+		if ((BlockOverlapByDepth[FreeMip] & BlockMap) != BlockOverlapByDepth[FreeMip])
+		{
+			break;
+		}
+	}
+
+	return FreeMip;
+}
+
 void FVirtualTextureAllocator::DumpToConsole(bool verbose)
 {
 	for (int32 BlockID = SortedIndices.Num() - 1; BlockID >= 0; BlockID--)
diff --git a/Engine/Source/Runtime/Renderer/Private/VT/VirtualTextureAllocator.h b/Engine/Source/Runtime/Renderer/Private/VT/VirtualTextureAllocator.h
--- a/Engine/Source/Runtime/Renderer/Private/VT/VirtualTextureAllocator.h
+++ b/Engine/Source/Runtime/Renderer/Private/VT/VirtualTextureAllocator.h
@@ -92,6 +92,11 @@ private:
 
 	bool TestAllocation(uint32 Index, uint32 vTileX0, uint32 vTileY0, uint32 vTileX1, uint32 vTileY1) const;
 
+	void RecurseComputeFreeMip(uint16 BlockIndex, uint32 Depth, uint64& IoBlockMap) const;
+	uint32 ComputeFreeMip(uint16 BlockIndex) const;
+
+	void FreeMipUpdateParents(uint16 ParentIndex);
+
 #if WITH_EDITOR
 	void FillDebugImage(uint32 Index, uint32* ImageData, TMap<FAllocatedVirtualTexture*, uint32>& ColorMap) const;
 #endif
@@ -109,6 +114,7 @@ private:
 		uint16						NextFree;
 		uint16						PrevFree;
 		EBlockState                 State;
+		uint8						FreeMip;
 
 		FAddressBlock()
 		{}
@@ -125,6 +131,7 @@ private:
 			, NextFree(0xffff)
 			, PrevFree(0xffff)
 			, State(EBlockState::None)
+			, FreeMip(0)
 		{}
 
 		FAddressBlock(const FAddressBlock& Block, uint32 Offset, uint32 Dimensions)
@@ -139,22 +146,35 @@ private:
 			, NextFree(0xffff)
 			, PrevFree(0xffff)
 			, State(EBlockState::None)
+			, FreeMip(0)
 		{}
 	};
 
-	const uint32			vDimensions;
-	uint32					AllocatedWidth;
-	uint32					AllocatedHeight;
-
-	TArray< FAddressBlock >	AddressBlocks;
-	TArray< uint16 >		FreeList;
-	TArray< uint16 >		PartiallyFreeList;
-	uint16					GlobalFreeList;
-	TArray< uint32 >		SortedAddresses;
-	TArray< uint16 >		SortedIndices;
-	FHashTable				HashTable;
-	uint16                  RootIndex;
-
-	uint32					NumAllocations;
-	uint32					NumAllocatedPages;
+	// We separate items in the partially free list by the alignment of sub-allocation that can
+	// potentially fit, based on what's already allocated.  A level of zero means a block could
+	// be allocated at 1x1 tile alignment, one means 2x2, etc.  It caps at level 3 (8x8) to
+	// keep the cost of updating this information manageable.  Capping it means we don't need
+	// to recurse too far into children or update too many parents when a block's state changes.
+	static const uint32 PartiallyFreeMipDepth = 4;
+
+	struct FPartiallyFreeMip
+	{
+		uint16	Mips[PartiallyFreeMipDepth];
+	};
+
+	const uint32				vDimensions;
+	uint32						AllocatedWidth;
+	uint32						AllocatedHeight;
+
+	TArray< FAddressBlock >		AddressBlocks;
+	TArray< uint16 >			FreeList;
+	TArray< FPartiallyFreeMip >	PartiallyFreeList;
+	uint16						GlobalFreeList;
+	TArray< uint32 >			SortedAddresses;
+	TArray< uint16 >			SortedIndices;
+	FHashTable					HashTable;
+	uint16						RootIndex;
+
+	uint32						NumAllocations;
+	uint32						NumAllocatedPages;
 };
