Optimization to remove stalling in FDeferredShadingSceneRenderer::GatherRayTracingWorldInstancesForView()

diff --git a/Engine/Source/Runtime/Renderer/Private/DeferredShadingRenderer.cpp b/Engine/Source/Runtime/Renderer/Private/DeferredShadingRenderer.cpp
--- a/Engine/Source/Runtime/Renderer/Private/DeferredShadingRenderer.cpp
+++ b/Engine/Source/Runtime/Renderer/Private/DeferredShadingRenderer.cpp
@@ -215,6 +215,13 @@ static FAutoConsoleVariableRef CVarRayTracingDebugForceOpaque(
 	TEXT("Forces all ray tracing geometry instances to be opaque, effectively disabling any-hit shaders. This is useful for debugging and profiling. (default = 0)")
 );
 
+static int32 GNumLODTasksToInline = 10;
+FAutoConsoleVariableRef CVarNumLODTasksToInline(
+	TEXT("r.RayTracing.GatherWorldInstancingInlineThreshold"),
+	GNumLODTasksToInline,
+	TEXT(""),
+	ECVF_Scalability | ECVF_RenderThreadSafe);
+
 
 #if !UE_BUILD_SHIPPING
 static TAutoConsoleVariable<int32> CVarForceBlackVelocityBuffer(
@@ -683,20 +690,12 @@ bool FDeferredShadingSceneRenderer::GatherRayTracingWorldInstancesForView(FRDGBu
 		const uint32 NumTasks = FMath::Max(1u, FMath::DivideAndRoundUp(NumTotalItems, TargetItemsPerTask));
 		const uint32 ItemsPerTask = FMath::DivideAndRoundUp(NumTotalItems, NumTasks); // Evenly divide commands between tasks (avoiding potential short last task)
 
-		LODTaskList.Reserve(NumTasks);
-
-		for (uint32 TaskIndex = 0; TaskIndex < NumTasks; ++TaskIndex)
-		{
-			const uint32 FirstTaskItemIndex = TaskIndex * ItemsPerTask;
-
-			LODTaskList.Add(FFunctionGraphTask::CreateAndDispatchWhenReady(
-			[	Items = RelevantPrimitives.GetData() + FirstTaskItemIndex,
-				NumItems = FMath::Min(ItemsPerTask, NumTotalItems - FirstTaskItemIndex),
-				&View,
+		auto ComputeLOD =
+			[	&View,
 				Scene = this->Scene,
 				LODScaleCVarValue,
 				ForcedLODLevel
-			]()
+			](FRelevantPrimitive* Items, uint32 NumItems)
 			{
 				TRACE_CPUPROFILER_EVENT_SCOPE(GatherRayTracingWorldInstances_ComputeLOD_Task);
 
@@ -773,8 +772,35 @@ bool FDeferredShadingSceneRenderer::GatherRayTracingWorldInstancesForView(FRDGBu
 						}
 					}
 				}
-			},
-			TStatId(), nullptr, ENamedThreads::AnyThread));
+			};
+
+		if (NumTasks > (uint32)GNumLODTasksToInline)
+		{
+			SCOPED_NAMED_EVENT(DispatchParallelComputeLOD, FColor::Red);
+			
+			LODTaskList.Reserve(NumTasks);
+
+			for (uint32 TaskIndex = 0; TaskIndex < NumTasks; ++TaskIndex)
+			{
+				const uint32 FirstTaskItemIndex = TaskIndex * ItemsPerTask;
+
+				LODTaskList.Add(FFunctionGraphTask::CreateAndDispatchWhenReady(
+					[	FirstTaskItemIndex,
+						Items = RelevantPrimitives.GetData() + FirstTaskItemIndex,
+						NumItems = FMath::Min(ItemsPerTask, NumTotalItems - FirstTaskItemIndex),
+						ComputeLOD
+					]()
+					{
+						ComputeLOD(Items, NumItems);
+					},
+					TStatId(), nullptr, ENamedThreads::AnyNormalThreadHiPriTask));
+			}
+		}
+		else
+		{
+			SCOPED_NAMED_EVENT(ComputeLOD, FColor::Magenta);
+			
+			ComputeLOD(RelevantPrimitives.GetData(), NumTotalItems);
 		}
 	}
 
@@ -1053,7 +1079,9 @@ bool FDeferredShadingSceneRenderer::GatherRayTracingWorldInstancesForView(FRDGBu
 
 		const bool bAutoInstance = CVarRayTracingAutoInstance.GetValueOnRenderThread() != 0;
 
+		if (LODTaskList.Num() > 0)
 		{
+			SCOPED_NAMED_EVENT(WaitForParallelComputeLOD, FColor::Red);
 			TRACE_CPUPROFILER_EVENT_SCOPE(WaitForLODTasks);
 			FTaskGraphInterface::Get().WaitUntilTasksComplete(LODTaskList, ENamedThreads::GetRenderThread_Local());
 		}
