Fixes TSR outputing black on D3D11

D3D11 doesn't allow to bind a UAV more than once per shader

diff --git a/Engine/Source/Runtime/Renderer/Private/PostProcess/TemporalSuperResolution.cpp b/Engine/Source/Runtime/Renderer/Private/PostProcess/TemporalSuperResolution.cpp
--- a/Engine/Source/Runtime/Renderer/Private/PostProcess/TemporalSuperResolution.cpp
+++ b/Engine/Source/Runtime/Renderer/Private/PostProcess/TemporalSuperResolution.cpp
@@ -1484,7 +1484,15 @@ void AddTemporalSuperResolutionPasses(
 		}
 		else
 		{
-			PassParameters->SceneColorOutputMip1 = PassParameters->SceneColorOutputMip0;
+			FRDGTextureDesc Desc = FRDGTextureDesc::Create2D(
+				FIntPoint(1, 1),
+				PF_FloatR11G11B10,
+				FClearValueBinding::None,
+				/* InFlags = */ TexCreate_ShaderResource | TexCreate_UAV);
+
+			FRDGTextureRef DummyTexture = GraphBuilder.CreateTexture(Desc, TEXT("TSR.DummyOutput"));
+			PassParameters->SceneColorOutputMip1 = GraphBuilder.CreateUAV(DummyTexture);
+			GraphBuilder.RemoveUnusedTextureWarning(DummyTexture);
 		}
 		PassParameters->DebugOutput = CreateDebugUAV(HistoryExtent, TEXT("Debug.TSR.UpdateHistory"));
 
