- Add Cvar to set the raytracing dynamic geometry shared vertex buffer size and changed default from 16MB to 4MB
- Add garbage collection on unused shared vertex buffers for 30 updates (cvar tweakable)

diff --git a/Engine/Source/Runtime/Renderer/Private/RayTracing/RayTracingDynamicGeometry.cpp b/Engine/Source/Runtime/Renderer/Private/RayTracing/RayTracingDynamicGeometry.cpp
--- a/Engine/Source/Runtime/Renderer/Private/RayTracing/RayTracingDynamicGeometry.cpp
+++ b/Engine/Source/Runtime/Renderer/Private/RayTracing/RayTracingDynamicGeometry.cpp
@@ -7,6 +7,22 @@
 
 #if RHI_RAYTRACING
 
+static int32 GRTDynGeomSharedVertexBufferSizeInMB = 4;
+static FAutoConsoleVariableRef CVarRTDynGeomSharedVertexBufferSizeInMB(
+	TEXT("r.RayTracing.DynamicGeometry.SharedVertexBufferSizeInMB"),
+	GRTDynGeomSharedVertexBufferSizeInMB,
+	TEXT("Size of the a single shared vertex buffer used during the BLAS update of dynamic geometries (default 4MB)"),
+	ECVF_RenderThreadSafe
+);
+
+static int32 GRTDynGeomSharedVertexBufferGarbageCollectLatency = 30;
+static FAutoConsoleVariableRef CVarRTDynGeomSharedVertexBufferGarbageCollectLatency(
+	TEXT("r.RayTracing.DynamicGeometry.SharedVertexBufferSizeInMB"),
+	GRTDynGeomSharedVertexBufferGarbageCollectLatency,
+	TEXT("Amount of update cycles before a heap is deleted when not used (default 30)."),
+	ECVF_RenderThreadSafe
+);
+
 DECLARE_CYCLE_STAT(TEXT("RTDynGeomDispatch"), STAT_CLM_RTDynGeomDispatch, STATGROUP_ParallelCommandListMarkers);
 DECLARE_CYCLE_STAT(TEXT("RTDynGeomBuild"), STAT_CLM_RTDynGeomBuild, STATGROUP_ParallelCommandListMarkers);
 
@@ -109,9 +125,18 @@ int64 FRayTracingDynamicGeometryCollection::BeginUpdate()
 	Segments.Empty(Segments.Max());
 
 	// Vertex buffer data can be immediatly reused the next frame, because it's already 'consumed' for building the AccelerationStructure data
-	for (FVertexPositionBuffer* Buffer : VertexPositionBuffers)
+	// Garbage collect unused buffers for n generations
+	for (int32 BufferIndex = 0; BufferIndex < VertexPositionBuffers.Num(); ++BufferIndex)
 	{
+		FVertexPositionBuffer* Buffer = VertexPositionBuffers[BufferIndex];
 		Buffer->UsedSize = 0;
+
+		if (Buffer->LastUsedGenerationID + GRTDynGeomSharedVertexBufferGarbageCollectLatency <= SharedBufferGenerationID)
+		{
+			VertexPositionBuffers.RemoveAtSwap(BufferIndex);
+			delete Buffer;
+			BufferIndex--;
+		}
 	}
 
 	// Increment generation ID used for validation
@@ -161,13 +186,16 @@ void FRayTracingDynamicGeometryCollection::AddDynamicMeshBatchForGeometryUpdate(
 			VertexPositionBuffer = new FVertexPositionBuffer;
 			VertexPositionBuffers.Add(VertexPositionBuffer);
 
-			static const uint32 VertexBufferCacheSize = 16 * 1024 * 1024;
+			static const uint32 VertexBufferCacheSize = GRTDynGeomSharedVertexBufferSizeInMB * 1024 * 1024;
 			uint32 AllocationSize = FMath::Max(VertexBufferCacheSize, UpdateParams.VertexBufferSize);
 
 			VertexPositionBuffer->RWBuffer.Initialize(TEXT("FRayTracingDynamicGeometryCollection::RayTracingDynamicVertexBuffer"), sizeof(float), AllocationSize / sizeof(float), PF_R32_FLOAT, BUF_UnorderedAccess | BUF_ShaderResource);
 			VertexPositionBuffer->UsedSize = 0;
 		}
 
+		// Update the last used generation ID
+		VertexPositionBuffer->LastUsedGenerationID = SharedBufferGenerationID;
+
 		// Get the offset and update used size
 		VertexBufferOffset = VertexPositionBuffer->UsedSize;
 		VertexPositionBuffer->UsedSize += UpdateParams.VertexBufferSize;
diff --git a/Engine/Source/Runtime/Renderer/Public/RayTracingDynamicGeometryCollection.h b/Engine/Source/Runtime/Renderer/Public/RayTracingDynamicGeometryCollection.h
--- a/Engine/Source/Runtime/Renderer/Public/RayTracingDynamicGeometryCollection.h
+++ b/Engine/Source/Runtime/Renderer/Public/RayTracingDynamicGeometryCollection.h
@@ -35,6 +35,7 @@ private:
 	{
 		FRWBuffer RWBuffer;
 		uint32 UsedSize = 0;
+		uint32 LastUsedGenerationID = 0;
 	};
 	TArray<FVertexPositionBuffer*> VertexPositionBuffers;
 
