Saves 0.5ms in post processing in sandbox

1) Allows motion blur to output half and/or quarter res based on what the following passes needs. This avoids adding downsampling passes after
2) Runs FFT bloom on async compute to overlap with all the histogram exposure and local pre exposure stuf
3) Fixes a bug where r.Bloom.HalfResolutionFFT resolution was also dependent r.PostProcessing.QuarterResolutionDownsample

Frosty configurations:
1) Configures to use simple 2x2 dowsampling in post processing to be able to use the fast motion blur code path
2) Configures histogram exposure and local exposure to quarter res.

diff --git a/Engine/Shaders/Private/MotionBlur/MotionBlurApply.usf b/Engine/Shaders/Private/MotionBlur/MotionBlurApply.usf
--- a/Engine/Shaders/Private/MotionBlur/MotionBlurApply.usf
+++ b/Engine/Shaders/Private/MotionBlur/MotionBlurApply.usf
@@ -5,6 +5,7 @@
 
 #include "MotionBlurCommon.ush"
 #include "../Random.ush"
+#include "/Engine/Public/WaveBroadcastIntrinsics.ush"
 
 
 //------------------------------------------------------- DEBUG
@@ -81,6 +82,8 @@ SCREEN_PASS_TEXTURE_VIEWPORT(Velocity)
 SCREEN_PASS_TEXTURE_VIEWPORT(VelocityTile)
 FScreenTransform ColorToVelocity;
 uint MaxSampleCount;
+uint OutputMip1;
+uint OutputMip2;
 
 uint TileListOffset;
 StructuredBuffer<uint> TileListsBuffer;
@@ -117,12 +120,24 @@ float2 TranslucencyExtentInverse;
 	#define SharedDepthSampler DepthSampler
 #endif
 
-RWTexture2D<mb_half4> OutColorTexture;
+RWTexture2D<mb_half4> SceneColorOutputMip0;
+RWTexture2D<mb_half4> SceneColorOutputMip1;
+RWTexture2D<mb_half4> SceneColorOutputMip2;
 RWTexture2D<float4> DebugOutput;
 
 
 //------------------------------------------------------- LDS
 
+#if !PLATFORM_SUPPORTS_WAVE_BROADCAST
+
+groupshared mb_half4 SharedArray0[TILE_SIZE * TILE_SIZE];
+
+#if POST_PROCESS_ALPHA
+groupshared mb_half  SharedArray1[TILE_SIZE * TILE_SIZE];
+#endif
+
+#endif // !PLATFORM_SUPPORTS_WAVE_BROADCAST
+
 #if !COMPILER_SUPPORTS_WAVE_VOTE
 
 groupshared uint SharedSampleCount;
@@ -313,7 +328,7 @@ void MainCS(
 		
 		bFastPath = length2(MinVelocityPixels) > 0.4 * (MaxPixelLength0 * MaxPixelLength0);
 
-		#if COMPILER_SUPPORTS_WAVE_MINMAX
+		#if COMPILER_SUPPORTS_WAVE_MINMAX && COMPILER_SUPPORTS_WAVE_VOTE
 		{
 			SampleCount = ToScalarMemory(WaveActiveMax(RecommendedSampleCount));
 			bFastPath = WaveActiveAllTrue(bFastPath);
@@ -700,31 +715,16 @@ void MainCS(
 	BRANCH
 	if (bIsHalfRes)
 	{
-		#if CONFIG_SHUFFLE_HALF_RES
-			mb_short2 OutputPixelCoord = mb_short2(ColorUV * Color_Extent);
-			mb_short2 OutputPixelCoord0 = (OutputPixelCoord & mb_short(~0x3)) | ((OutputPixelCoord >> mb_short(1)) & mb_short(0x1));
-			mb_short2 OutputPixelCoord1 = OutputPixelCoord0 + mb_short2(2, 0);
-			mb_short2 OutputPixelCoord2 = OutputPixelCoord0 + mb_short2(0, 2);
-			mb_short2 OutputPixelCoord3 = OutputPixelCoord0 + mb_short2(2, 2);
-		#else
-			mb_short2 OutputPixelCoord0 = mb_short2(uint2(ColorUV * Color_Extent)) & mb_short(~0x1);
-			mb_short2 OutputPixelCoord1 = OutputPixelCoord0 + mb_short2(1, 0);
-			mb_short2 OutputPixelCoord2 = OutputPixelCoord0 + mb_short2(0, 1);
-			mb_short2 OutputPixelCoord3 = OutputPixelCoord0 + mb_short2(1, 1);
-		#endif
+		mb_short2 OutputPixelCoord0 = mb_short2(uint2(ColorUV * Color_Extent)) & mb_short(~0x1);
+		mb_short2 OutputPixelCoord1 = OutputPixelCoord0 + mb_short2(1, 0);
+		mb_short2 OutputPixelCoord2 = OutputPixelCoord0 + mb_short2(0, 1);
+		mb_short2 OutputPixelCoord3 = OutputPixelCoord0 + mb_short2(1, 1);
 
 		#if 1
 			float2 PostMotionBlurTranslucencyUV0 = ApplyScreenTransform((float2(uint2(OutputPixelCoord0)) + 0.5) * Color_ExtentInverse, ColorToTranslucency);
 			float2 PostMotionBlurTranslucencyUV1 = ApplyScreenTransform((float2(uint2(OutputPixelCoord1)) + 0.5) * Color_ExtentInverse, ColorToTranslucency);
 			float2 PostMotionBlurTranslucencyUV2 = ApplyScreenTransform((float2(uint2(OutputPixelCoord2)) + 0.5) * Color_ExtentInverse, ColorToTranslucency);
 			float2 PostMotionBlurTranslucencyUV3 = ApplyScreenTransform((float2(uint2(OutputPixelCoord3)) + 0.5) * Color_ExtentInverse, ColorToTranslucency);
-		#elif CONFIG_SHUFFLE_HALF_RES
-			float2 ColorUV0 = (float2(uint2(OutputPixelCoord0)) + 0.5) * Color_ExtentInverse;
-
-			float2 PostMotionBlurTranslucencyUV0 = ApplyScreenTransform(ColorUV0, ColorToTranslucency);
-			float2 PostMotionBlurTranslucencyUV1 = PostMotionBlurTranslucencyUV0 + float2(2.0, 0.0) * TranslucencyExtentInverse;
-			float2 PostMotionBlurTranslucencyUV2 = PostMotionBlurTranslucencyUV0 + float2(0.0, 2.0) * TranslucencyExtentInverse;
-			float2 PostMotionBlurTranslucencyUV3 = PostMotionBlurTranslucencyUV0 + float2(2.0, 2.0) * TranslucencyExtentInverse;
 		#else
 			float2 PostMotionBlurTranslucencyUV = ApplyScreenTransform(ColorUV, ColorToTranslucency);
 			float2 PostMotionBlurTranslucencyUV0 = PostMotionBlurTranslucencyUV + float2(-0.5, -0.5) * TranslucencyExtentInverse;
@@ -784,11 +784,84 @@ void MainCS(
 		}
 		#endif
 
+		// Swizzle the half res 2x2 motion blur color across the full res 4x4 quad
+		//  0  0  1  1      0  1  0  1
+		//  0  0  1  1  ->  2  3  2  3
+		//  2  2  3  3      0  1  0  1
+		//  2  2  3  3      2  3  2  3
+		mb_half4 MotionBlurColor0;
+		mb_half4 MotionBlurColor1;
+		mb_half4 MotionBlurColor2;
+		mb_half4 MotionBlurColor3;
+		mb_half FullResBlend0;
+		mb_half FullResBlend1;
+		mb_half FullResBlend2;
+		mb_half FullResBlend3;
+		#if CONFIG_SHUFFLE_HALF_RES && PLATFORM_SUPPORTS_WAVE_BROADCAST
+		{
+			const uint LaneGroupSize = 4;
+			const uint InnerLaneGroupSize = 1;
+
+			const FWaveBroadcastSettings Broadcast0 = InitWaveBroadcastLaneGroup(LaneGroupSize, InnerLaneGroupSize, /* InnerLaneGroupId = */ 0);
+			const FWaveBroadcastSettings Broadcast1 = InitWaveBroadcastLaneGroup(LaneGroupSize, InnerLaneGroupSize, /* InnerLaneGroupId = */ 1);
+			const FWaveBroadcastSettings Broadcast2 = InitWaveBroadcastLaneGroup(LaneGroupSize, InnerLaneGroupSize, /* InnerLaneGroupId = */ 2);
+			const FWaveBroadcastSettings Broadcast3 = InitWaveBroadcastLaneGroup(LaneGroupSize, InnerLaneGroupSize, /* InnerLaneGroupId = */ 3);
+			
+			MotionBlurColor0 = WaveBroadcast(Broadcast0, MotionBlurColor);
+			MotionBlurColor1 = WaveBroadcast(Broadcast1, MotionBlurColor);
+			MotionBlurColor2 = WaveBroadcast(Broadcast2, MotionBlurColor);
+			MotionBlurColor3 = WaveBroadcast(Broadcast3, MotionBlurColor);
+		
+			FullResBlend0 = WaveBroadcast(Broadcast0, FullResBlend);
+			FullResBlend1 = WaveBroadcast(Broadcast1, FullResBlend);
+			FullResBlend2 = WaveBroadcast(Broadcast2, FullResBlend);
+			FullResBlend3 = WaveBroadcast(Broadcast3, FullResBlend);
+		}
+		#elif CONFIG_SHUFFLE_HALF_RES && !PLATFORM_SUPPORTS_WAVE_BROADCAST
+		{
+			#if POST_PROCESS_ALPHA
+				SharedArray0[GroupThreadIndex] = MotionBlurColor;
+				SharedArray1[GroupThreadIndex] = FullResBlend;
+			#else
+				SharedArray0[GroupThreadIndex] = float4(MotionBlurColor.rgb, FullResBlend);
+			#endif
+
+			MotionBlurColor0 = SharedArray0[(GroupThreadIndex & (~0x3)) | 0x0];
+			MotionBlurColor1 = SharedArray0[(GroupThreadIndex & (~0x3)) | 0x1];
+			MotionBlurColor2 = SharedArray0[(GroupThreadIndex & (~0x3)) | 0x2];
+			MotionBlurColor3 = SharedArray0[(GroupThreadIndex & (~0x3)) | 0x3];
+		
+			#if POST_PROCESS_ALPHA
+				FullResBlend0 = SharedArray1[(GroupThreadIndex & (~0x3)) | 0x0];
+				FullResBlend1 = SharedArray1[(GroupThreadIndex & (~0x3)) | 0x1];
+				FullResBlend2 = SharedArray1[(GroupThreadIndex & (~0x3)) | 0x2];
+				FullResBlend3 = SharedArray1[(GroupThreadIndex & (~0x3)) | 0x3];
+			#else
+				FullResBlend0 = MotionBlurColor0.a;
+				FullResBlend1 = MotionBlurColor1.a;
+				FullResBlend2 = MotionBlurColor2.a;
+				FullResBlend3 = MotionBlurColor3.a;
+			#endif
+		}
+		#else // !CONFIG_SHUFFLE_HALF_RES
+		{
+			MotionBlurColor0 = MotionBlurColor;
+			MotionBlurColor1 = MotionBlurColor;
+			MotionBlurColor2 = MotionBlurColor;
+			MotionBlurColor3 = MotionBlurColor;
+		
+			FullResBlend0 = FullResBlend;
+			FullResBlend1 = FullResBlend;
+			FullResBlend2 = FullResBlend;
+			FullResBlend3 = FullResBlend;
+		}
+		#endif
+
 		// Blend full res and motion blur
-		mb_half4 OutputColor0 = CenterColor0 * FullResBlend + MotionBlurColor;
-		mb_half4 OutputColor1 = CenterColor1 * FullResBlend + MotionBlurColor;
-		mb_half4 OutputColor2 = CenterColor2 * FullResBlend + MotionBlurColor;
-		mb_half4 OutputColor3 = CenterColor3 * FullResBlend + MotionBlurColor;
+		mb_half4 OutputColor0 = CenterColor0 * FullResBlend0 + MotionBlurColor0;
+		mb_half4 OutputColor1 = CenterColor1 * FullResBlend1 + MotionBlurColor1;
+		mb_half4 OutputColor2 = CenterColor2 * FullResBlend2 + MotionBlurColor2;
+		mb_half4 OutputColor3 = CenterColor3 * FullResBlend3 + MotionBlurColor3;
 		
 		// Adds debug overlay
 		#if DEBUG_OVERLAY_TILE_CLASSIFICATION || DEBUG_OVERLAY_SKIP_CENTER || DEBUG_OVERLAY_SAMPLES
@@ -819,20 +892,60 @@ void MainCS(
 		}
 		#endif
 	
-		if (all(OutputPixelCoord0 < mb_short2(Color_ViewportMax)))
+		// Compute the half res.
+		mb_half4 HalfResOutput = mb_half(0.25) * (OutputColor0 + OutputColor1 + OutputColor2 + OutputColor3);
+		
+		// Compute the quarter res
+		mb_half4 QuarterResOutput;
+		BRANCH
+		if (OutputMip2)
+		#if PLATFORM_SUPPORTS_WAVE_BROADCAST
 		{
-			OutColorTexture[OutputPixelCoord0] = OutputColor0;
-			OutColorTexture[OutputPixelCoord1] = OutputColor1;
-			OutColorTexture[OutputPixelCoord2] = OutputColor2;
-			OutColorTexture[OutputPixelCoord3] = OutputColor3;
+			FWaveBroadcastSettings Horizontal = InitWaveXorButterfly(/* XorButterFly = */ 0x1);
+			FWaveBroadcastSettings Vertical = InitWaveXorButterfly(/* XorButterFly = */ 0x2);
 			
-			#if DEBUG_OUTPUT
-				DebugOutput[OutputPixelCoord0] = Debug;
-				DebugOutput[OutputPixelCoord1] = Debug;
-				DebugOutput[OutputPixelCoord2] = Debug;
-				DebugOutput[OutputPixelCoord3] = Debug;
-			#endif
+			QuarterResOutput = HalfResOutput * mb_half(0.25);
+			QuarterResOutput += WaveBroadcast(Horizontal, QuarterResOutput);
+			QuarterResOutput += WaveBroadcast(Vertical, QuarterResOutput);
+		}
+		#else
+		{
+			QuarterResOutput = HalfResOutput * mb_half(0.25);
+
+			SharedArray0[GroupThreadIndex] = QuarterResOutput;
+			QuarterResOutput += SharedArray0[GroupThreadIndex ^ 0x1];
+			
+			SharedArray0[GroupThreadIndex] = QuarterResOutput;
+			QuarterResOutput += SharedArray0[GroupThreadIndex ^ 0x2];
+		}
+		#endif
+		else
+		{
+			QuarterResOutput = 0.0;
 		}
+		
+		bool bIsValid = all(OutputPixelCoord0 < mb_short2(Color_ViewportMax));
+		mb_short2 OutputPixelCoord0Mip0 = (bIsValid) ? OutputPixelCoord0 : mb_short(~0).xx;
+		mb_short2 OutputPixelCoord1Mip0 = (bIsValid) ? OutputPixelCoord1 : mb_short(~0).xx;
+		mb_short2 OutputPixelCoord2Mip0 = (bIsValid) ? OutputPixelCoord2 : mb_short(~0).xx;
+		mb_short2 OutputPixelCoord3Mip0 = (bIsValid) ? OutputPixelCoord3 : mb_short(~0).xx;
+		mb_short2 OutputPixelCoordMip1 = (bIsValid && (OutputMip1 != 0)) ? (OutputPixelCoord0 >> mb_short(1)) : mb_short(~0).xx;
+		mb_short2 OutputPixelCoordMip2 = (bIsValid && (OutputMip2 != 0) && (OutputPixelCoord0 & 0x3) == 0) ? (OutputPixelCoord0 >> mb_short(2)) : mb_short(~0).xx;
+
+		SceneColorOutputMip0[OutputPixelCoord0Mip0] = OutputColor0;
+		SceneColorOutputMip0[OutputPixelCoord1Mip0] = OutputColor1;
+		SceneColorOutputMip0[OutputPixelCoord2Mip0] = OutputColor2;
+		SceneColorOutputMip0[OutputPixelCoord3Mip0] = OutputColor3;
+			
+		SceneColorOutputMip1[OutputPixelCoordMip1] = HalfResOutput;
+		SceneColorOutputMip2[OutputPixelCoordMip2] = QuarterResOutput;
+		
+		#if DEBUG_OUTPUT
+			DebugOutput[OutputPixelCoord0Mip0] = Debug;
+			DebugOutput[OutputPixelCoord1Mip0] = Debug;
+			DebugOutput[OutputPixelCoord2Mip0] = Debug;
+			DebugOutput[OutputPixelCoord3Mip0] = Debug;
+		#endif
 	}
 	else
 	{
@@ -886,14 +999,73 @@ void MainCS(
 		}
 		#endif
 	
-		if (all(OutputPixelCoord < mb_short2(Color_ViewportMax)))
+		mb_half4 HalfResOutput;
+		BRANCH
+		if (OutputMip1 || OutputMip2)
+		#if PLATFORM_SUPPORTS_WAVE_BROADCAST
+		{
+			FWaveBroadcastSettings Horizontal = InitWaveXorButterfly(/* XorButterFly = */ 0x1);
+			FWaveBroadcastSettings Vertical = InitWaveXorButterfly(/* XorButterFly = */ 0x2);
+			
+			HalfResOutput = OutputColor * mb_half(0.25);
+			HalfResOutput += WaveBroadcast(Horizontal, HalfResOutput);
+			HalfResOutput += WaveBroadcast(Vertical, HalfResOutput);
+		}
+		#else
 		{
-			OutColorTexture[OutputPixelCoord] = OutputColor;
+			HalfResOutput = OutputColor * mb_half(0.25);
 
-			#if DEBUG_OUTPUT
-				DebugOutput[OutputPixelCoord] = Debug;
-			#endif
+			SharedArray0[GroupThreadIndex] = HalfResOutput;
+			HalfResOutput += SharedArray0[GroupThreadIndex ^ 0x1];
+			
+			SharedArray0[GroupThreadIndex] = HalfResOutput;
+			HalfResOutput += SharedArray0[GroupThreadIndex ^ 0x2];
 		}
+		#endif
+		else
+		{
+			HalfResOutput = 0.0;
+		}
+		
+		mb_half4 QuarterResOutput;
+		BRANCH
+		if (OutputMip2)
+		#if PLATFORM_SUPPORTS_WAVE_BROADCAST
+		{
+			FWaveBroadcastSettings Horizontal = InitWaveXorButterfly(/* XorButterFly = */ 0x4);
+			FWaveBroadcastSettings Vertical = InitWaveXorButterfly(/* XorButterFly = */ 0x8);
+			
+			QuarterResOutput = HalfResOutput * mb_half(0.25);
+			QuarterResOutput += WaveBroadcast(Horizontal, QuarterResOutput);
+			QuarterResOutput += WaveBroadcast(Vertical, QuarterResOutput);
+		}
+		#else
+		{
+			QuarterResOutput = HalfResOutput * mb_half(0.25);
+
+			SharedArray0[GroupThreadIndex] = QuarterResOutput;
+			QuarterResOutput += SharedArray0[GroupThreadIndex ^ 0x4];
+			
+			SharedArray0[GroupThreadIndex] = QuarterResOutput;
+			QuarterResOutput += SharedArray0[GroupThreadIndex ^ 0x8];
+		}
+		#endif
+		else
+		{
+			QuarterResOutput = 0.0;
+		}
+
+		bool bIsValid = all(OutputPixelCoord < mb_short2(Color_ViewportMax));
+		mb_short2 OutputPixelCoordMip0 = (bIsValid) ? OutputPixelCoord : mb_short(~0).xx;
+		mb_short2 OutputPixelCoordMip1 = (bIsValid && (OutputMip1 != 0) && (OutputPixelCoordMip0 & 0x1) == 0) ? (OutputPixelCoord >> mb_short(1)) : mb_short(~0).xx;
+		mb_short2 OutputPixelCoordMip2 = (bIsValid && (OutputMip2 != 0) && (OutputPixelCoordMip0 & 0x3) == 0) ? (OutputPixelCoord >> mb_short(2)) : mb_short(~0).xx;
+
+		SceneColorOutputMip0[OutputPixelCoordMip0] = OutputColor;
+		SceneColorOutputMip1[OutputPixelCoordMip1] = HalfResOutput;
+		SceneColorOutputMip2[OutputPixelCoordMip2] = QuarterResOutput;
+
+		#if DEBUG_OUTPUT
+			DebugOutput[OutputPixelCoordMip0] = Debug;
+		#endif
 	}
 }
-
diff --git a/Engine/Shaders/Private/ScreenSpaceDenoise/SSDSpatialKernel.ush b/Engine/Shaders/Private/ScreenSpaceDenoise/SSDSpatialKernel.ush
--- a/Engine/Shaders/Private/ScreenSpaceDenoise/SSDSpatialKernel.ush
+++ b/Engine/Shaders/Private/ScreenSpaceDenoise/SSDSpatialKernel.ush
@@ -1461,7 +1461,7 @@ void AccumulateSquare3x3Kernel(
 			/* bForceSample = */ false);
 	}
 
-	static const float2 SampleOffsetArray[5] = {
+	static const float2 SampleOffsetArray[4] = {
 		float2(-1.0,  0.0),
 		float2( 0.0, -1.0),
 		float2(-1.0, 1.0),
diff --git a/Engine/Source/Runtime/Renderer/Private/GPUFastFourierTransform.cpp b/Engine/Source/Runtime/Renderer/Private/GPUFastFourierTransform.cpp
--- a/Engine/Source/Runtime/Renderer/Private/GPUFastFourierTransform.cpp
+++ b/Engine/Source/Runtime/Renderer/Private/GPUFastFourierTransform.cpp
@@ -460,6 +460,7 @@ void DispatchReorderFFTPassCS(
 	FComputeShaderUtils::AddPass(
 		GraphBuilder,
 		RDG_EVENT_NAME("FFT Multipass: Complex %s Reorder pass of size %d", FFTDesc.FFT_TypeName(), TransformLength),
+		FFTDesc.ComputePassFlags,
 		ComputeShader,
 		PassParameters,
 		FIntVector(1, 1, TransformLength / Radix));
@@ -515,6 +516,7 @@ void DispatchGSSubComplexFFTPassCS(
 		GraphBuilder,
 		RDG_EVENT_NAME("FFT Multipass: %d GS Subpasses Complex %s of size %d",
 			NumSubRegions, FFTDesc.FFT_TypeName(), FGroupShardSubFFTPassCS::SubPassLength()),
+		FFTDesc.ComputePassFlags,
 		ComputeShader,
 		PassParameters,
 		FIntVector(1, 1, NumSignals));
@@ -567,6 +569,7 @@ void DispatchComplexFFTPassCS(
 	FComputeShaderUtils::AddPass(
 		GraphBuilder,
 		RDG_EVENT_NAME("FFT Multipass: Pass %d of Complex %s of size %d", PassLength, FFTDesc.FFT_TypeName(), TransformLength),
+		FFTDesc.ComputePassFlags,
 		ComputeShader,
 		PassParameters,
 		FIntVector(1, 1, TransformLength / Radix));
@@ -623,6 +626,7 @@ void DispatchPackTwoForOneFFTPassCS(
 	FComputeShaderUtils::AddPass(
 		GraphBuilder,
 		RDG_EVENT_NAME("FFT Multipass: TwoForOne Combine/split result of %s of size %d", FFTDesc.FFT_TypeName(), TransformLength),
+		FFTDesc.ComputePassFlags,
 		ComputeShader,
 		PassParameters,
 		FIntVector(1, 1, RealTransformLength));
@@ -645,6 +649,7 @@ void DispatchPackTwoForOneFFTPassCS(
  */
 void DispatchCopyWindowCS(
 	FRDGBuilder& GraphBuilder,
+	ERDGPassFlags ComputePassFlags,
 	const FGlobalShaderMap* ShaderMap,
 	FRDGTextureRef SrcTexture, const FIntRect& SrcWindow,
 	FRDGTextureRef DstTexture, const FIntRect& DstWindow,
@@ -662,6 +667,7 @@ void DispatchCopyWindowCS(
 	FComputeShaderUtils::AddPass(
 		GraphBuilder,
 		RDG_EVENT_NAME("FFT Multipass: Copy Subwindow"),
+		ComputePassFlags,
 		ComputeShader,
 		PassParameters,
 		FComputeShaderUtils::GetGroupCount(DstWindow.Size(), FCopyWindowCS::ThreadCount()));
@@ -697,6 +703,7 @@ void DispatchCopyWindowCS(
  */
 void DispatchComplexMultiplyImagesCS(
 	FRDGBuilder& GraphBuilder,
+	ERDGPassFlags ComputePassFlags,
 	const FGlobalShaderMap* ShaderMap,
 	const bool bHorizontalScanlines,
 	FRDGTextureRef SrcTexture, const FIntRect& SrcWindow,
@@ -722,6 +729,7 @@ void DispatchComplexMultiplyImagesCS(
 	FComputeShaderUtils::AddPass(
 		GraphBuilder,
 		RDG_EVENT_NAME("FFT Multipass: Convolution in freq-space"),
+		ComputePassFlags,
 		ComputeShader,
 		PassParameters,
 		FIntVector(1, 1, NumScanLines));
@@ -780,6 +788,7 @@ void DispatchGSComplexFFTCS(
 	FComputeShaderUtils::AddPass(
 		GraphBuilder,
 		RDG_EVENT_NAME("FFT: Complex %s of size %d", FFTDesc.FFT_TypeName(), TransformLength),
+		FFTDesc.ComputePassFlags,
 		ComputeShader,
 		PassParameters,
 		FIntVector(1, 1, NumSignals));
@@ -842,6 +851,7 @@ void DispatchGSTwoForOneFFTCS(
 	FComputeShaderUtils::AddPass(
 		GraphBuilder,
 		RDG_EVENT_NAME("FFT: Two-For-One %s of size %d of buffer %d x %d", FFTDesc.FFT_TypeName(), TransformLength, SrcRect.Size().X, SrcRect.Size().Y),
+		FFTDesc.ComputePassFlags,
 		ComputeShader,
 		PassParameters,
 		FIntVector(1, 1, NumScanLines));
@@ -901,6 +911,7 @@ void DispatchGSConvolutionWithTextureCS(
 	FComputeShaderUtils::AddPass(
 		GraphBuilder,
 		RDG_EVENT_NAME("FFT: Apply %s Transform, Multiply Texture, and InverseTransform size %d of buffer %d x %d", FFTDesc.FFT_TypeName(), SignalLength, SrcRectSize.X, SrcRectSize.Y),
+		FFTDesc.ComputePassFlags,
 		ComputeShader,
 		PassParameters,
 		FIntVector(1, 1, NumSignals));
@@ -923,12 +934,13 @@ bool FitsInGroupSharedMemory(const GPUFFT::FFTDescription& FFTDesc)
 
 void GPUFFT::CopyImage2D(
 	FRDGBuilder& GraphBuilder,
+	ERDGPassFlags ComputePassFlags,
 	const FGlobalShaderMap* ShaderMap,
 	FRDGTextureRef SrcTexture, const FIntRect& SrcWindow,
 	FRDGTextureRef DstTexture, const FIntRect& DstWindow,
 	const FPreFilter& PreFilter)
 {
-	DispatchCopyWindowCS(GraphBuilder, ShaderMap, SrcTexture, SrcWindow, DstTexture, DstWindow, PreFilter);
+	DispatchCopyWindowCS(GraphBuilder, ComputePassFlags, ShaderMap, SrcTexture, SrcWindow, DstTexture, DstWindow, PreFilter);
 }
 
 void GPUFFT::ComplexFFTImage1D::MultiPass(
@@ -1071,7 +1083,7 @@ void GPUFFT::TwoForOneRealFFTImage1D::MultiPass(
 			FRDGTextureRef PreFilteredTexture = GraphBuilder.CreateTexture(Desc, TEXT("FFT.PreFilter"));
 
 			CopyImage2D(
-				GraphBuilder, ShaderMap,
+				GraphBuilder, FFTDesc.ComputePassFlags, ShaderMap,
 				/* SrcTexture = */ SrcTexture, SrcWindow,
 				/* DstTexture = */ PreFilteredTexture, SrcWindow,
 				PreFilter);
@@ -1189,7 +1201,7 @@ void GPUFFT::ConvolutionWithTextureImage1D::MultiPass(
 		// Apply convolution kernel.
 		FRDGTextureRef SpectralTexture1 = GraphBuilder.CreateTexture(SrcTexture->Desc, TEXT("FFT.Spectral"));
 		DispatchComplexMultiplyImagesCS(
-			GraphBuilder, ShaderMap, FFTDesc.IsHorizontal(),
+			GraphBuilder, FFTDesc.ComputePassFlags, ShaderMap, FFTDesc.IsHorizontal(),
 			/* SrcTexture = */ SpectralTexture0, TargetRect,
 			TransformedKernel, 
 			/* DstTexture = */ SpectralTexture1);
@@ -1253,6 +1265,7 @@ FRDGTextureDesc CreateFrequencyDesc(const FIntPoint& FrequencySize, const bool b
 
 void GPUFFT::ConvolutionWithTextureImage2D(
 	FRDGBuilder& GraphBuilder,
+	ERDGPassFlags ComputePassFlags,
 	const FGlobalShaderMap* ShaderMap,
 	const FIntPoint& FrequencySize, bool bHorizontalFirst,
 	FRDGTextureRef TransformedKernel,
@@ -1271,6 +1284,7 @@ void GPUFFT::ConvolutionWithTextureImage2D(
 	TwoForOneFFTDesc.XFormType = bHorizontalFirst ? FFT_XFORM_TYPE::FORWARD_HORIZONTAL : FFT_XFORM_TYPE::FORWARD_VERTICAL;
 	TwoForOneFFTDesc.SignalLength = bHorizontalFirst ? FrequencySize.X : FrequencySize.Y;
 	TwoForOneFFTDesc.NumScanLines = (bHorizontalFirst) ? ROISize.Y : ROISize.X;
+	TwoForOneFFTDesc.ComputePassFlags = ComputePassFlags;
 
 	// The output has two more elements
 	FIntRect SpectralOutputRect;
@@ -1284,9 +1298,11 @@ void GPUFFT::ConvolutionWithTextureImage2D(
 	ConvolutionFFTDesc.XFormType = bHorizontalFirst ? FFT_XFORM_TYPE::FORWARD_VERTICAL : FFT_XFORM_TYPE::FORWARD_HORIZONTAL;
 	ConvolutionFFTDesc.SignalLength = bHorizontalFirst ? FrequencySize.Y : FrequencySize.X;
 	ConvolutionFFTDesc.NumScanLines = TwoForOneFFTDesc.SignalLength + 2; // two-for-one transform generated two additional elements
+	ConvolutionFFTDesc.ComputePassFlags = ComputePassFlags;
 
 	FFTDescription TwoForOneIvnFFTDesc = TwoForOneFFTDesc;
 	TwoForOneIvnFFTDesc.XFormType = GetInverseOfXForm(TwoForOneFFTDesc.XFormType);
+	TwoForOneIvnFFTDesc.ComputePassFlags = ComputePassFlags;
 	
 	
 
diff --git a/Engine/Source/Runtime/Renderer/Private/GPUFastFourierTransform.h b/Engine/Source/Runtime/Renderer/Private/GPUFastFourierTransform.h
--- a/Engine/Source/Runtime/Renderer/Private/GPUFastFourierTransform.h
+++ b/Engine/Source/Runtime/Renderer/Private/GPUFastFourierTransform.h
@@ -67,6 +67,7 @@ namespace GPUFFT
 		FFT_XFORM_TYPE  XFormType = FFT_XFORM_TYPE::FORWARD_HORIZONTAL;
 		uint32 SignalLength = 0;
 		uint32 NumScanLines = 0;
+		ERDGPassFlags ComputePassFlags = ERDGPassFlags::Compute;
 	};
 
 	/**
@@ -280,6 +281,7 @@ namespace GPUFFT
 	*/
 	void ConvolutionWithTextureImage2D(
 		FRDGBuilder& GraphBuilder,
+		ERDGPassFlags ComputePassFlags,
 		const FGlobalShaderMap* ShaderMap,
 		const FIntPoint& FrequencySize,
 		bool bHorizontalFirst,
@@ -308,6 +310,7 @@ namespace GPUFFT
 	*/
 	void CopyImage2D(
 		FRDGBuilder& GraphBuilder,
+		ERDGPassFlags ComputePassFlags,
 		const FGlobalShaderMap* ShaderMap,
 		FRDGTextureRef SrcTexture, const FIntRect& SrcWindow,
 		FRDGTextureRef DstTexture, const FIntRect& DstWindow,
diff --git a/Engine/Source/Runtime/Renderer/Private/PostProcess/PostProcessFFTBloom.cpp b/Engine/Source/Runtime/Renderer/Private/PostProcess/PostProcessFFTBloom.cpp
--- a/Engine/Source/Runtime/Renderer/Private/PostProcess/PostProcessFFTBloom.cpp
+++ b/Engine/Source/Runtime/Renderer/Private/PostProcess/PostProcessFFTBloom.cpp
@@ -21,6 +21,11 @@ TAutoConsoleVariable<int32> CVarHalfResFFTBloom(
 	TEXT(" 2: Quarter-resolution convolution.\n"),
 	ECVF_Scalability | ECVF_RenderThreadSafe);
 
+TAutoConsoleVariable<int32> CVarAsynComputeFFTBloom(
+	TEXT("r.Bloom.AsyncCompute"), 0,
+	TEXT("Whether to run FFT bloom on async compute.\n"),
+	ECVF_RenderThreadSafe);
+
 static bool DoesPlatformSupportFFTBloom(EShaderPlatform Platform)
 {
 	return FDataDrivenShaderPlatformInfo::GetSupportsFFTBloom(Platform);
@@ -191,9 +196,9 @@ IMPLEMENT_GLOBAL_SHADER(FBloomFinalizeApplyConstantsCS,     "/Engine/Private/Blo
 
 } //! namespace
 
-bool IsFFTBloomHalfResolutionEnabled()
+bool IsFFTBloomFullResolutionEnabled()
 {
-	return CVarHalfResFFTBloom.GetValueOnRenderThread() != 0;
+	return CVarHalfResFFTBloom.GetValueOnRenderThread() == 0;
 }
 
 bool IsFFTBloomQuarterResolutionEnabled()
@@ -283,6 +288,8 @@ struct FFFTBloomIntermediates
 	// The order of the two-dimensional transform.  This implicitly defines
 	// the data layout in transform space for both the kernel and image transform
 	bool bDoHorizontalFirst = false;
+
+	ERDGPassFlags ComputePassFlags = ERDGPassFlags::Compute;
 };
 
 FFFTBloomIntermediates GetFFTBloomIntermediates(
@@ -319,7 +326,7 @@ FFFTBloomIntermediates GetFFTBloomIntermediates(
 		return FilterRadius;
 	};
 
-	const bool bHalfResolutionFFT = IsFFTBloomHalfResolutionEnabled();
+	const bool bHalfResolutionFFT = !IsFFTBloomFullResolutionEnabled();
 
 	FFFTBloomIntermediates Intermediates;
 	if (bHalfResolutionFFT)
@@ -368,6 +375,8 @@ FFFTBloomIntermediates GetFFTBloomIntermediates(
 
 	Intermediates.bDoHorizontalFirst = ((Intermediates.FrequencySize.Y * PaddedImageSize.X) > (Intermediates.FrequencySize.X * PaddedImageSize.Y));
 
+	Intermediates.ComputePassFlags = (GSupportsEfficientAsyncCompute && CVarAsynComputeFFTBloom.GetValueOnRenderThread() != 0) ? ERDGPassFlags::AsyncCompute : ERDGPassFlags::Compute;
+
 	return Intermediates;
 }
 
@@ -452,6 +461,7 @@ void InitDomainAndGetKernel(
 		FComputeShaderUtils::AddPass(
 			GraphBuilder,
 			RDG_EVENT_NAME("FFTBloom FindKernelCenter %dx%d", SpatialKernelTexture->Desc.Extent.X, SpatialKernelTexture->Desc.Extent.Y),
+			Intermediates.ComputePassFlags,
 			ComputeShader,
 			PassParameters, 
 			FComputeShaderUtils::GetGroupCount(SpatialKernelTexture->Desc.Extent, 8));
@@ -500,6 +510,7 @@ void InitDomainAndGetKernel(
 			FComputeShaderUtils::AddPass(
 				GraphBuilder,
 				RDG_EVENT_NAME("FFTBloom ReduceKernelSurvey(Op=%d) %d", Op, SurveyGroupCount),
+				Intermediates.ComputePassFlags,
 				ComputeShader,
 				PassParameters,
 				FComputeShaderUtils::GetGroupCount(SurveyGroupCount, 64));
@@ -524,6 +535,7 @@ void InitDomainAndGetKernel(
 			FComputeShaderUtils::AddPass(
 				GraphBuilder,
 				RDG_EVENT_NAME("FFTBloom SurveyMaxScatterDispersion"),
+				Intermediates.ComputePassFlags,
 				ComputeShader,
 				PassParameters,
 				FIntVector(SurveyGroupGridSize, SurveyGroupGridSize, 1));
@@ -551,6 +563,7 @@ void InitDomainAndGetKernel(
 			FComputeShaderUtils::AddPass(
 				GraphBuilder,
 				RDG_EVENT_NAME("FFTBloom SurveyKernelCenterEnergy"),
+				Intermediates.ComputePassFlags,
 				ComputeShader,
 				PassParameters,
 				FIntVector(SurveyGroupGridSize, SurveyGroupGridSize, 1));
@@ -594,6 +607,7 @@ void InitDomainAndGetKernel(
 				RDG_EVENT_NAME("FFTBloom SumScatterDispersionEnergy %dx%d -> %dx%d",
 					ScatterDispersionTexture->Desc.Extent.X, ScatterDispersionTexture->Desc.Extent.Y,
 					NewScatterDispersionTexture->Desc.Extent.X, NewScatterDispersionTexture->Desc.Extent.Y),
+				Intermediates.ComputePassFlags,
 				ComputeShader,
 				PassParameters,
 				FIntVector(NewScatterDispersionTexture->Desc.Extent.X, NewScatterDispersionTexture->Desc.Extent.Y, 1));
@@ -620,6 +634,7 @@ void InitDomainAndGetKernel(
 		FComputeShaderUtils::AddPass(
 			GraphBuilder,
 			RDG_EVENT_NAME("FFTBloom PackKernelConstants"),
+			Intermediates.ComputePassFlags,
 			ComputeShader,
 			PassParameters,
 			FIntVector(1, 1, 1));
@@ -654,6 +669,7 @@ void InitDomainAndGetKernel(
 				RDG_EVENT_NAME("FFTBloom ClampedKernel %dx%d",
 					SpatialKernelTexture->Desc.Extent.X,
 					SpatialKernelTexture->Desc.Extent.Y),
+				Intermediates.ComputePassFlags,
 				ComputeShader,
 				PassParameters,
 				FComputeShaderUtils::GetGroupCount(SpatialKernelTexture->Desc.Extent, 8));
@@ -681,6 +697,7 @@ void InitDomainAndGetKernel(
 			FComputeShaderUtils::AddPass(
 				GraphBuilder,
 				RDG_EVENT_NAME("FFTBloom PreProcessKernel %dx%d", Intermediates.FrequencySize.X, Intermediates.FrequencySize.Y),
+				Intermediates.ComputePassFlags,
 				ComputeShader,
 				PassParameters,
 				FComputeShaderUtils::GetGroupCount(PaddedFrequencySize, 8));
@@ -723,8 +740,6 @@ FBloomOutputs AddFFTBloomPass(FRDGBuilder& GraphBuilder, const FViewInfo& View,
 	check(Inputs.HalfResolutionTexture);
 	check(!Inputs.HalfResolutionViewRect.IsEmpty());
 
-	const bool bHalfResolutionFFT = IsFFTBloomHalfResolutionEnabled();
-
 	FFFTBloomIntermediates Intermediates = GetFFTBloomIntermediates(View, Inputs);
 
 	RDG_EVENT_SCOPE(GraphBuilder, "FFTBloom %dx%d", Intermediates.ImageRect.Width(), Intermediates.ImageRect.Height());
@@ -756,6 +771,7 @@ FBloomOutputs AddFFTBloomPass(FRDGBuilder& GraphBuilder, const FViewInfo& View,
 		FComputeShaderUtils::AddPass(
 			GraphBuilder,
 			RDG_EVENT_NAME("FFTBloom FinalizeApplyConstants(ScatterDispersion=%f)", PassParameters->ScatterDispersionIntensity),
+			Intermediates.ComputePassFlags,
 			ComputeShader,
 			PassParameters,
 			FIntVector(1, 1, 1));
@@ -772,6 +788,7 @@ FBloomOutputs AddFFTBloomPass(FRDGBuilder& GraphBuilder, const FViewInfo& View,
 
 	GPUFFT::ConvolutionWithTextureImage2D(
 		GraphBuilder,
+		Intermediates.ComputePassFlags,
 		View.ShaderMap,
 		Intermediates.FrequencySize,
 		Intermediates.bDoHorizontalFirst,
diff --git a/Engine/Source/Runtime/Renderer/Private/PostProcess/PostProcessFFTBloom.h b/Engine/Source/Runtime/Renderer/Private/PostProcess/PostProcessFFTBloom.h
--- a/Engine/Source/Runtime/Renderer/Private/PostProcess/PostProcessFFTBloom.h
+++ b/Engine/Source/Runtime/Renderer/Private/PostProcess/PostProcessFFTBloom.h
@@ -19,6 +19,7 @@ struct FFFTBloomInputs
 	FIntRect HalfResolutionViewRect;
 };
 
+bool IsFFTBloomFullResolutionEnabled();
 bool IsFFTBloomQuarterResolutionEnabled();
 
 FBloomOutputs AddFFTBloomPass(FRDGBuilder& GraphBuilder, const FViewInfo& View, const FFFTBloomInputs& Inputs);
diff --git a/Engine/Source/Runtime/Renderer/Private/PostProcess/PostProcessLensFlares.cpp b/Engine/Source/Runtime/Renderer/Private/PostProcess/PostProcessLensFlares.cpp
--- a/Engine/Source/Runtime/Renderer/Private/PostProcess/PostProcessLensFlares.cpp
+++ b/Engine/Source/Runtime/Renderer/Private/PostProcess/PostProcessLensFlares.cpp
@@ -341,68 +341,73 @@ FScreenPassTexture AddLensFlaresPass(
 	return FScreenPassTexture(LensFlareTexture, OutputViewRect);
 }
 
+bool IsLensFlaresEnabled(const FViewInfo& View)
+{
+	const ELensFlareQuality LensFlareQuality = GetLensFlareQuality();
+
+	const FPostProcessSettings& Settings = View.FinalPostProcessSettings;
+
+	return (LensFlareQuality != ELensFlareQuality::Disabled &&
+		!Settings.LensFlareTint.IsAlmostBlack() &&
+		Settings.LensFlareBokehSize > SMALL_NUMBER &&
+		Settings.LensFlareIntensity > SMALL_NUMBER);
+}
+
 FScreenPassTexture AddLensFlaresPass(
 	FRDGBuilder& GraphBuilder,
 	const FViewInfo& View,
 	FScreenPassTexture Bloom,
 	const FSceneDownsampleChain& SceneDownsampleChain)
 {
+	ensure(IsLensFlaresEnabled(View));
 	const ELensFlareQuality LensFlareQuality = GetLensFlareQuality();
 
 	const FPostProcessSettings& Settings = View.FinalPostProcessSettings;
 
-	if (LensFlareQuality != ELensFlareQuality::Disabled &&
-		!Settings.LensFlareTint.IsAlmostBlack() &&
-		Settings.LensFlareBokehSize > SMALL_NUMBER &&
-		Settings.LensFlareIntensity > SMALL_NUMBER)
-	{
-		RDG_GPU_STAT_SCOPE(GraphBuilder, LensFlare);
-
-		FRHITexture* BokehTextureRHI = GWhiteTexture->TextureRHI;
+	RDG_GPU_STAT_SCOPE(GraphBuilder, LensFlare);
 
-		if (GEngine->DefaultBokehTexture)
-		{
-			FTextureResource* BokehTextureResource = GEngine->DefaultBokehTexture->GetResource();
+	FRHITexture* BokehTextureRHI = GWhiteTexture->TextureRHI;
 
-			if (BokehTextureResource && BokehTextureResource->TextureRHI)
-			{
-				BokehTextureRHI = BokehTextureResource->TextureRHI;
-			}
-		}
+	if (GEngine->DefaultBokehTexture)
+	{
+		FTextureResource* BokehTextureResource = GEngine->DefaultBokehTexture->GetResource();
 
-		if (Settings.LensFlareBokehShape)
+		if (BokehTextureResource && BokehTextureResource->TextureRHI)
 		{
-			FTextureResource* BokehTextureResource = Settings.LensFlareBokehShape->GetResource();
-
-			if (BokehTextureResource && BokehTextureResource->TextureRHI)
-			{
-				BokehTextureRHI = BokehTextureResource->TextureRHI;
-			}
+			BokehTextureRHI = BokehTextureResource->TextureRHI;
 		}
+	}
+
+	if (Settings.LensFlareBokehShape)
+	{
+		FTextureResource* BokehTextureResource = Settings.LensFlareBokehShape->GetResource();
 
-		// The quality level controls which downsample stage we use as the flare input texture.
-		const uint32 LensFlareDownsampleStageIndex = static_cast<uint32>(ELensFlareQuality::MAX) - static_cast<uint32>(LensFlareQuality) - 1;
-
-		FLensFlareInputs LensFlareInputs;
-		LensFlareInputs.Bloom = Bloom;
-		LensFlareInputs.Flare = SceneDownsampleChain.GetTexture(LensFlareDownsampleStageIndex);
-		LensFlareInputs.BokehShapeTexture = BokehTextureRHI;
-		LensFlareInputs.TintColorsPerFlare = Settings.LensFlareTints;
-		LensFlareInputs.TintColor = Settings.LensFlareTint;
-		LensFlareInputs.BokehSizePercent = Settings.LensFlareBokehSize;
-		LensFlareInputs.Intensity = Settings.LensFlareIntensity;
-		LensFlareInputs.Threshold = Settings.LensFlareThreshold;
-
-		// If a bloom output texture isn't available, substitute the half resolution scene color instead, but disable bloom
-		// composition. The pass needs a primary input in order to access the image descriptor and viewport for output.
-		if (!Bloom.IsValid())
+		if (BokehTextureResource && BokehTextureResource->TextureRHI)
 		{
-			LensFlareInputs.Bloom = SceneDownsampleChain.GetFirstTexture();
-			LensFlareInputs.bCompositeWithBloom = false;
+			BokehTextureRHI = BokehTextureResource->TextureRHI;
 		}
+	}
 
-		return AddLensFlaresPass(GraphBuilder, View, LensFlareInputs);
+	// The quality level controls which downsample stage we use as the flare input texture.
+	const uint32 LensFlareDownsampleStageIndex = static_cast<uint32>(ELensFlareQuality::MAX) - static_cast<uint32>(LensFlareQuality) - 1;
+
+	FLensFlareInputs LensFlareInputs;
+	LensFlareInputs.Bloom = Bloom;
+	LensFlareInputs.Flare = SceneDownsampleChain.GetTexture(LensFlareDownsampleStageIndex);
+	LensFlareInputs.BokehShapeTexture = BokehTextureRHI;
+	LensFlareInputs.TintColorsPerFlare = Settings.LensFlareTints;
+	LensFlareInputs.TintColor = Settings.LensFlareTint;
+	LensFlareInputs.BokehSizePercent = Settings.LensFlareBokehSize;
+	LensFlareInputs.Intensity = Settings.LensFlareIntensity;
+	LensFlareInputs.Threshold = Settings.LensFlareThreshold;
+
+	// If a bloom output texture isn't available, substitute the half resolution scene color instead, but disable bloom
+	// composition. The pass needs a primary input in order to access the image descriptor and viewport for output.
+	if (!Bloom.IsValid())
+	{
+		LensFlareInputs.Bloom = SceneDownsampleChain.GetFirstTexture();
+		LensFlareInputs.bCompositeWithBloom = false;
 	}
 
-	return FScreenPassTexture();
+	return AddLensFlaresPass(GraphBuilder, View, LensFlareInputs);
 }
\ No newline at end of file
diff --git a/Engine/Source/Runtime/Renderer/Private/PostProcess/PostProcessLensFlares.h b/Engine/Source/Runtime/Renderer/Private/PostProcess/PostProcessLensFlares.h
--- a/Engine/Source/Runtime/Renderer/Private/PostProcess/PostProcessLensFlares.h
+++ b/Engine/Source/Runtime/Renderer/Private/PostProcess/PostProcessLensFlares.h
@@ -57,7 +57,8 @@ struct FLensFlareInputs
 
 using FLensFlareOutputs = FScreenPassTexture;
 
-FLensFlareOutputs AddLensFlaresPass(FRDGBuilder& GraphBuilder, const FViewInfo& View, const FLensFlareInputs& Inputs);
+
+bool IsLensFlaresEnabled(const FViewInfo& View);
 
 // Helper function which pulls inputs from the post process settings of the view.
 FScreenPassTexture AddLensFlaresPass(
diff --git a/Engine/Source/Runtime/Renderer/Private/PostProcess/PostProcessMotionBlur.cpp b/Engine/Source/Runtime/Renderer/Private/PostProcess/PostProcessMotionBlur.cpp
--- a/Engine/Source/Runtime/Renderer/Private/PostProcess/PostProcessMotionBlur.cpp
+++ b/Engine/Source/Runtime/Renderer/Private/PostProcess/PostProcessMotionBlur.cpp
@@ -428,6 +428,8 @@ class FMotionBlurFilterCS : public FMotionBlurShader
 
 		SHADER_PARAMETER(FScreenTransform, ColorToVelocity)
 		SHADER_PARAMETER(int32, MaxSampleCount)
+		SHADER_PARAMETER(int32, OutputMip1)
+		SHADER_PARAMETER(int32, OutputMip2)
 
 		SHADER_PARAMETER_RDG_TEXTURE(Texture2D, ColorTexture)
 		SHADER_PARAMETER_RDG_TEXTURE(Texture2D, VelocityFlatTexture)
@@ -450,7 +452,9 @@ class FMotionBlurFilterCS : public FMotionBlurShader
 		SHADER_PARAMETER_RDG_BUFFER_SRV(RWStructuredBuffer<uint>, TileListsSizeBuffer)
 		RDG_BUFFER_ACCESS(DispatchParameters, ERHIAccess::IndirectArgs)
 
-		SHADER_PARAMETER_RDG_TEXTURE_UAV(RWTexture2D<float4>, OutColorTexture)
+		SHADER_PARAMETER_RDG_TEXTURE_UAV(RWTexture2D<float4>, SceneColorOutputMip0)
+		SHADER_PARAMETER_RDG_TEXTURE_UAV(RWTexture2D<float4>, SceneColorOutputMip1)
+		SHADER_PARAMETER_RDG_TEXTURE_UAV(RWTexture2D<float4>, SceneColorOutputMip2)
 		SHADER_PARAMETER_RDG_TEXTURE_UAV(RWTexture2D<float4>, DebugOutput)
 	END_SHADER_PARAMETER_STRUCT()
 
@@ -781,9 +785,10 @@ void AddMotionBlurVelocityPass(
 	*VelocityTileTexturesOutput = VelocityTileTextures;
 }
 
-FRDGTextureRef AddMotionBlurFilterPass(
+FMotionBlurOutputs AddMotionBlurFilterPass(
 	FRDGBuilder& GraphBuilder,
 	const FViewInfo& View,
+	const FMotionBlurInputs& Inputs,
 	const FMotionBlurViewports& Viewports,
 	FRDGTextureRef ColorTexture,
 	FRDGTextureRef VelocityFlatTexture,
@@ -901,8 +906,9 @@ FRDGTextureRef AddMotionBlurFilterPass(
 			FIntVector(1, 1, 1));
 	}
 
+
 	// Apply the filter.
-	FRDGTextureRef ColorTextureOutput;
+	FMotionBlurOutputs Output;
 	{
 		static const TCHAR* kTileClassificationNames[] = {
 			TEXT("GatherHalfRes"),
@@ -919,6 +925,8 @@ FRDGTextureRef AddMotionBlurFilterPass(
 		OriginalPassParameters.VelocityTile = Viewports.VelocityTileParameters;
 		OriginalPassParameters.ColorToVelocity = Viewports.ColorToVelocityTransform;
 		OriginalPassParameters.MaxSampleCount = MaxSampleCount;
+		OriginalPassParameters.OutputMip1 = Inputs.bOutputHalfRes ? 1 : 0;
+		OriginalPassParameters.OutputMip2 = Inputs.bOutputQuarterRes ? 1 : 0;
 
 		OriginalPassParameters.ColorTexture = ColorTexture;
 		OriginalPassParameters.VelocityFlatTexture = VelocityFlatTexture;
@@ -961,21 +969,76 @@ FRDGTextureRef AddMotionBlurFilterPass(
 		}
 
 		{
-			FRDGTextureDesc OutColorDesc = FRDGTextureDesc::Create2D(
+			FRDGTextureDesc Desc = FRDGTextureDesc::Create2D(
 				ColorTexture->Desc.Extent,
 				IsPostProcessingWithAlphaChannelSupported() ? PF_FloatRGBA : PF_FloatRGB,
 				FClearValueBinding::None,
 				TexCreate_UAV | TexCreate_ShaderResource | GFastVRamConfig.MotionBlur);
 
-			ColorTextureOutput = GraphBuilder.CreateTexture(OutColorDesc, TEXT("MotionBlur.SceneColor"));
+			Output.FullRes.Texture = GraphBuilder.CreateTexture(Desc, TEXT("MotionBlur.SceneColor"));
+			Output.FullRes.ViewRect = Viewports.Color.Rect;
+			OriginalPassParameters.SceneColorOutputMip0 = GraphBuilder.CreateUAV(Output.FullRes.Texture, ERDGUnorderedAccessViewFlags::SkipBarrier);
+		}
+
+		if (OriginalPassParameters.OutputMip1)
+		{
+			FRDGTextureDesc Desc = FRDGTextureDesc::Create2D(
+				Output.FullRes.Texture->Desc.Extent / 2,
+				Output.FullRes.Texture->Desc.Format,
+				FClearValueBinding::None,
+				TexCreate_UAV | TexCreate_ShaderResource);
+
+			Output.HalfRes.Texture = GraphBuilder.CreateTexture(Desc, TEXT("MotionBlur.SceneColor.HalfRes"));
+			Output.HalfRes.ViewRect.Min = Viewports.Color.Rect.Min / 2;
+			Output.HalfRes.ViewRect.Max = Output.HalfRes.ViewRect.Min + FIntPoint::DivideAndRoundUp(Viewports.Color.Rect.Size(), 2);
+			OriginalPassParameters.SceneColorOutputMip1 = GraphBuilder.CreateUAV(Output.HalfRes.Texture, ERDGUnorderedAccessViewFlags::SkipBarrier);
+		}
+		else
+		{
+			FRDGTextureDesc Desc = FRDGTextureDesc::Create2D(
+				FIntPoint(1, 1),
+				PF_FloatR11G11B10,
+				FClearValueBinding::None,
+				/* InFlags = */ TexCreate_ShaderResource | TexCreate_UAV);
+
+			FRDGTextureRef DummyTexture = GraphBuilder.CreateTexture(Desc, TEXT("MotionBlur.DummyOutput"));
+			OriginalPassParameters.SceneColorOutputMip1 = GraphBuilder.CreateUAV(DummyTexture, ERDGUnorderedAccessViewFlags::SkipBarrier);
+			GraphBuilder.RemoveUnusedTextureWarning(DummyTexture);
+		}
+
+		if (OriginalPassParameters.OutputMip2)
+		{
+			FRDGTextureDesc Desc = FRDGTextureDesc::Create2D(
+				Output.FullRes.Texture->Desc.Extent / 4,
+				Output.FullRes.Texture->Desc.Format,
+				FClearValueBinding::None,
+				TexCreate_UAV | TexCreate_ShaderResource);
+
+			Output.QuarterRes.Texture = GraphBuilder.CreateTexture(Desc, TEXT("MotionBlur.SceneColor.QuarterRes"));
+			Output.QuarterRes.ViewRect.Min = Viewports.Color.Rect.Min / 4;
+			Output.QuarterRes.ViewRect.Max = Output.QuarterRes.ViewRect.Min + FIntPoint::DivideAndRoundUp(Viewports.Color.Rect.Size(), 4);
+			OriginalPassParameters.SceneColorOutputMip2 = GraphBuilder.CreateUAV(Output.QuarterRes.Texture, ERDGUnorderedAccessViewFlags::SkipBarrier);
+		}
+		else
+		{
+			FRDGTextureDesc Desc = FRDGTextureDesc::Create2D(
+				FIntPoint(1, 1),
+				PF_FloatR11G11B10,
+				FClearValueBinding::None,
+				/* InFlags = */ TexCreate_ShaderResource | TexCreate_UAV);
+
+			FRDGTextureRef DummyTexture = GraphBuilder.CreateTexture(Desc, TEXT("MotionBlur.DummyOutput"));
+			OriginalPassParameters.SceneColorOutputMip2 = GraphBuilder.CreateUAV(DummyTexture, ERDGUnorderedAccessViewFlags::SkipBarrier);
+			GraphBuilder.RemoveUnusedTextureWarning(DummyTexture);
 		}
 
-		OriginalPassParameters.OutColorTexture = GraphBuilder.CreateUAV(ColorTextureOutput, ERDGUnorderedAccessViewFlags::SkipBarrier);
-		OriginalPassParameters.DebugOutput = CreateDebugUAV(GraphBuilder, ColorTextureOutput->Desc.Extent, TEXT("Debug.MotionBlur.Filter"));
+		OriginalPassParameters.DebugOutput = CreateDebugUAV(GraphBuilder, Output.FullRes.Texture->Desc.Extent, TEXT("Debug.MotionBlur.Filter"));
 
-		RDG_EVENT_SCOPE(GraphBuilder, "MotionBlur FullResFilter(BlurDirections=%d MaxSamples=%d%s) %dx%d",
+		RDG_EVENT_SCOPE(GraphBuilder, "MotionBlur FullResFilter(BlurDirections=%d MaxSamples=%d%s%s%s) %dx%d",
 			BlurDirections, OriginalPassParameters.MaxSampleCount,
 			PostMotionBlurTranslucency ? TEXT(" ComposeTranslucency") : TEXT(""),
+			OriginalPassParameters.OutputMip1 ? TEXT(" OutputMip1") : TEXT(""),
+			OriginalPassParameters.OutputMip2 ? TEXT(" OutputMip2") : TEXT(""),
 			Viewports.Color.Rect.Width(), Viewports.Color.Rect.Height());
 
 		for (int32 TileClassifcation = 0; TileClassifcation < int32(FMotionBlurFilterCS::ETileClassification::MAX); TileClassifcation++)
@@ -1005,7 +1068,7 @@ FRDGTextureRef AddMotionBlurFilterPass(
 		}
 	}
 
-	return ColorTextureOutput;
+	return Output;
 }
 
 FScreenPassTexture AddVisualizeMotionBlurPass(FRDGBuilder& GraphBuilder, const FViewInfo& View, const FMotionBlurInputs& Inputs)
@@ -1086,7 +1149,7 @@ FScreenPassTexture AddVisualizeMotionBlurPass(FRDGBuilder& GraphBuilder, const F
 	return MoveTemp(Output);
 }
 
-FScreenPassTexture AddMotionBlurPass(FRDGBuilder& GraphBuilder, const FViewInfo& View, const FMotionBlurInputs& Inputs)
+FMotionBlurOutputs AddMotionBlurPass(FRDGBuilder& GraphBuilder, const FViewInfo& View, const FMotionBlurInputs& Inputs)
 {
 	check(Inputs.SceneColor.IsValid());
 	check(Inputs.SceneDepth.IsValid());
@@ -1112,14 +1175,13 @@ FScreenPassTexture AddMotionBlurPass(FRDGBuilder& GraphBuilder, const FViewInfo&
 		&VelocityFlatTexture,
 		&VelocityTileTextures);
 
-	FScreenPassTexture Output;
-	Output.ViewRect = Viewports.Color.Rect;
-
+	FMotionBlurOutputs Output;
 	if (Inputs.Filter == EMotionBlurFilter::Separable)
 	{
 		FRDGTextureRef MotionBlurFilterTexture = AddMotionBlurFilterPass(
 			GraphBuilder,
 			View,
+			Inputs,
 			Viewports,
 			Inputs.SceneColor.Texture,
 			VelocityFlatTexture,
@@ -1127,11 +1189,12 @@ FScreenPassTexture AddMotionBlurPass(FRDGBuilder& GraphBuilder, const FViewInfo&
 			nullptr,
 			FIntPoint(0, 0),
 			EMotionBlurFilterPass::Separable0,
-			Inputs.Quality);
+			Inputs.Quality).FullRes.Texture;
 
-		Output.Texture = AddMotionBlurFilterPass(
+		Output = AddMotionBlurFilterPass(
 			GraphBuilder,
 			View,
+			Inputs,
 			Viewports,
 			MotionBlurFilterTexture,
 			VelocityFlatTexture,
@@ -1143,9 +1206,10 @@ FScreenPassTexture AddMotionBlurPass(FRDGBuilder& GraphBuilder, const FViewInfo&
 	}
 	else
 	{
-		Output.Texture = AddMotionBlurFilterPass(
+		Output = AddMotionBlurFilterPass(
 			GraphBuilder,
 			View,
+			Inputs,
 			Viewports,
 			Inputs.SceneColor.Texture,
 			VelocityFlatTexture,
diff --git a/Engine/Source/Runtime/Renderer/Private/PostProcess/PostProcessMotionBlur.h b/Engine/Source/Runtime/Renderer/Private/PostProcess/PostProcessMotionBlur.h
--- a/Engine/Source/Runtime/Renderer/Private/PostProcess/PostProcessMotionBlur.h
+++ b/Engine/Source/Runtime/Renderer/Private/PostProcess/PostProcessMotionBlur.h
@@ -41,6 +41,9 @@ EMotionBlurFilter GetMotionBlurFilter();
 
 struct FMotionBlurInputs
 {
+	bool bOutputHalfRes = false;
+	bool bOutputQuarterRes = false;
+
 	// [Optional] Render to the specified output. If invalid, a new texture is created and returned.
 	FScreenPassRenderTarget OverrideOutput;
 
@@ -63,5 +66,12 @@ struct FMotionBlurInputs
 	EMotionBlurFilter Filter = EMotionBlurFilter::Separable;
 };
 
-FScreenPassTexture AddMotionBlurPass(FRDGBuilder& GraphBuilder, const FViewInfo& View, const FMotionBlurInputs& Inputs);
+struct FMotionBlurOutputs
+{
+	FScreenPassTexture FullRes;
+	FScreenPassTexture HalfRes;
+	FScreenPassTexture QuarterRes;
+};
+
+FMotionBlurOutputs AddMotionBlurPass(FRDGBuilder& GraphBuilder, const FViewInfo& View, const FMotionBlurInputs& Inputs);
 FScreenPassTexture AddVisualizeMotionBlurPass(FRDGBuilder& GraphBuilder, const FViewInfo& View, const FMotionBlurInputs& Inputs);
\ No newline at end of file
diff --git a/Engine/Source/Runtime/Renderer/Private/PostProcess/PostProcessing.cpp b/Engine/Source/Runtime/Renderer/Private/PostProcess/PostProcessing.cpp
--- a/Engine/Source/Runtime/Renderer/Private/PostProcess/PostProcessing.cpp
+++ b/Engine/Source/Runtime/Renderer/Private/PostProcess/PostProcessing.cpp
@@ -598,7 +598,26 @@ void AddPostProcessingPasses(FRDGBuilder& GraphBuilder, const FViewInfo& View, c
 
 		PassSequence.Finalize();
 
-		bool bMotionBlurNeedsHalfResInput = PassSequence.IsEnabled(EPass::MotionBlur) && DoesMotionBlurNeedsHalfResInput();
+		const bool bLensFlareEnabled = bBloomEnabled && IsLensFlaresEnabled(View);
+		const bool bFFTBloomEnabled = bBloomEnabled && IsFFTBloomEnabled(View);
+
+		const bool bBasicEyeAdaptationEnabled = bEyeAdaptationEnabled && (AutoExposureMethod == EAutoExposureMethod::AEM_Basic);
+		const bool bLocalExposureBlurredLum = bLocalExposureEnabled && View.FinalPostProcessSettings.LocalExposureBlurredLuminanceBlend > 0.0f;
+
+		const bool bProcessQuarterResolution = IsPostProcessingQuarterResolutionDownsampleEnabled();
+		const bool bMotionBlurNeedsHalfResInput = PassSequence.IsEnabled(EPass::MotionBlur) && DoesMotionBlurNeedsHalfResInput();
+
+		const bool bIsFullResFFT = bFFTBloomEnabled && IsFFTBloomFullResolutionEnabled();
+		const bool bIsHalfResFFT = bFFTBloomEnabled && !IsFFTBloomFullResolutionEnabled() && !IsFFTBloomQuarterResolutionEnabled();
+		const bool bIsQuarterResFFT = bFFTBloomEnabled && IsFFTBloomQuarterResolutionEnabled();
+
+		const bool bProduceSceneColorChain = (
+			bBasicEyeAdaptationEnabled ||
+			(bBloomEnabled && !bFFTBloomEnabled) ||
+			(bLensFlareEnabled && bFFTBloomEnabled) ||
+			bLocalExposureBlurredLum);
+		const bool bNeedPostMotionBlurHalfRes = !bProcessQuarterResolution || bIsHalfResFFT;
+		const bool bNeedPostMotionBlurQuarterRes = bProcessQuarterResolution || bIsQuarterResFFT;
 
 		// Post Process Material Chain - Before Translucency
 		{
@@ -656,11 +675,12 @@ void AddPostProcessingPasses(FRDGBuilder& GraphBuilder, const FViewInfo& View, c
 			}
 		}
 
-		FScreenPassTexture DownsampledSceneColor;
-
 		// Scene color view rectangle after temporal AA upscale to secondary screen percentage.
 		FIntRect SecondaryViewRect = PrimaryViewRect;
 
+		FScreenPassTexture HalfResSceneColor;
+		FScreenPassTexture QuarterResSceneColor;
+
 		if (TAAConfig != EMainTAAPassConfig::Disabled)
 		{
 			// Whether we allow the temporal AA pass to downsample scene color. It may choose not to based on internal context,
@@ -668,9 +688,10 @@ void AddPostProcessingPasses(FRDGBuilder& GraphBuilder, const FViewInfo& View, c
 			const bool bAllowSceneDownsample =
 				IsTemporalAASceneDownsampleAllowed(View) &&
 				// We can only merge if the normal downsample pass would happen immediately after.
+				bNeedPostMotionBlurHalfRes &&
 				!bMotionBlurEnabled && !bVisualizeMotionBlur &&
 				// TemporalAA is only able to match the low quality mode (box filter).
-				GetDownsampleQuality() == EDownsampleQuality::Low;
+				DownsampleQuality == EDownsampleQuality::Low;
 
 			const ITemporalUpscaler* UpscalerToUse = (TAAConfig == EMainTAAPassConfig::ThirdParty) ? View.Family->GetTemporalUpscalerInterface() : ITemporalUpscaler::GetDefaultTemporalUpscaler();
 			check(UpscalerToUse);
@@ -702,8 +723,8 @@ void AddPostProcessingPasses(FRDGBuilder& GraphBuilder, const FViewInfo& View, c
 				UpscalerPassInputs,
 				&SceneColor.Texture,
 				&SecondaryViewRect,
-				&DownsampledSceneColor.Texture,
-				&DownsampledSceneColor.ViewRect);
+				&HalfResSceneColor.Texture,
+				&HalfResSceneColor.ViewRect);
 		}
 		else if (ScreenSpaceRayTracing::ShouldRenderScreenSpaceReflections(View))
 		{
@@ -747,6 +768,8 @@ void AddPostProcessingPasses(FRDGBuilder& GraphBuilder, const FViewInfo& View, c
 		{
 			FMotionBlurInputs PassInputs;
 			PassSequence.AcceptOverrideIfLastPass(EPass::MotionBlur, PassInputs.OverrideOutput);
+			PassInputs.bOutputHalfRes = bNeedPostMotionBlurHalfRes && DownsampleQuality == EDownsampleQuality::Low;
+			PassInputs.bOutputQuarterRes = bNeedPostMotionBlurQuarterRes && DownsampleQuality == EDownsampleQuality::Low;
 			PassInputs.SceneColor = SceneColor;
 			PassInputs.SceneDepth = SceneDepth;
 			PassInputs.SceneVelocity = Velocity;
@@ -761,7 +784,10 @@ void AddPostProcessingPasses(FRDGBuilder& GraphBuilder, const FViewInfo& View, c
 			}
 			else
 			{
-				SceneColor = AddMotionBlurPass(GraphBuilder, View, PassInputs);
+				FMotionBlurOutputs PassOutputs = AddMotionBlurPass(GraphBuilder, View, PassInputs);
+				SceneColor = PassOutputs.FullRes;
+				HalfResSceneColor = PassOutputs.HalfRes;
+				QuarterResSceneColor = PassOutputs.QuarterRes;
 			}
 		}
 		else
@@ -778,122 +804,127 @@ void AddPostProcessingPasses(FRDGBuilder& GraphBuilder, const FViewInfo& View, c
 				/* bApplyModulateOnly = */ false);
 		}
 
-		SceneColor = AddAfterPass(EPass::MotionBlur, SceneColor);
-
-		// If TAA didn't do it, downsample the scene color texture by half.
-		if (!DownsampledSceneColor.IsValid())
 		{
-			FDownsamplePassInputs PassInputs;
-			PassInputs.Name = TEXT("HalfResolutionSceneColor");
-			PassInputs.SceneColor = SceneColor;
-			PassInputs.Quality = DownsampleQuality;
-			PassInputs.FormatOverride = DownsampleOverrideFormat;
-			PassInputs.UserSuppliedOutput = View.PrevViewInfo.HalfResTemporalAAHistory;
+			FScreenPassTexture NewSceneColor = AddAfterPass(EPass::MotionBlur, SceneColor);
 
-			DownsampledSceneColor = AddDownsamplePass(GraphBuilder, View, PassInputs);
+			// Invalidate half and quarter res.
+			if (NewSceneColor != SceneColor)
+			{
+				HalfResSceneColor = FScreenPassTexture();
+				QuarterResSceneColor = FScreenPassTexture();
+			}
+
+			SceneColor = NewSceneColor;
 		}
 
-		if (IsPostProcessingQuarterResolutionDownsampleEnabled())
+		// Generate post motion blur lower res scene color if they have not been generated.
 		{
-			FDownsamplePassInputs PassInputs;
-			PassInputs.Name = TEXT("QuarterResolutionSceneColor");
-			PassInputs.SceneColor = DownsampledSceneColor;
-			PassInputs.Quality = DownsampleQuality;
+			if ((bNeedPostMotionBlurHalfRes ||
+				(bNeedPostMotionBlurQuarterRes && !QuarterResSceneColor.IsValid())) && !HalfResSceneColor.IsValid())
+			{
+				FDownsamplePassInputs PassInputs;
+				PassInputs.Name = TEXT("PostProcessing.SceneColor.HalfRes");
+				PassInputs.SceneColor = SceneColor;
+				PassInputs.Quality = DownsampleQuality;
+				PassInputs.FormatOverride = DownsampleOverrideFormat;
+				PassInputs.UserSuppliedOutput = View.PrevViewInfo.HalfResTemporalAAHistory;
 
-			DownsampledSceneColor = AddDownsamplePass(GraphBuilder, View, PassInputs);
+				HalfResSceneColor = AddDownsamplePass(GraphBuilder, View, PassInputs);
+			}
+
+			if (bNeedPostMotionBlurQuarterRes && !QuarterResSceneColor.IsValid())
+			{
+				FDownsamplePassInputs PassInputs;
+				PassInputs.Name = TEXT("PostProcessing.SceneColor.QuarterRes");
+				PassInputs.SceneColor = HalfResSceneColor;
+				PassInputs.Quality = DownsampleQuality;
+
+				QuarterResSceneColor = AddDownsamplePass(GraphBuilder, View, PassInputs);
+			}
 		}
 
+
+
 		// Store half res scene color in the history
 		extern int32 GSSRHalfResSceneColor;
 		if (ScreenSpaceRayTracing::ShouldRenderScreenSpaceReflections(View) && !View.bStatePrevViewInfoIsReadOnly && GSSRHalfResSceneColor)
 		{
 			check(View.ViewState);
-			GraphBuilder.QueueTextureExtraction(DownsampledSceneColor.Texture, &View.ViewState->PrevFrameViewInfo.HalfResTemporalAAHistory);
+			GraphBuilder.QueueTextureExtraction(HalfResSceneColor.Texture, &View.ViewState->PrevFrameViewInfo.HalfResTemporalAAHistory);
 		}
 
 		if (bLocalExposureEnabled)
 		{
-			LocalExposureTexture = AddLocalExposurePass(GraphBuilder, View, EyeAdaptationParameters, DownsampledSceneColor);
+			LocalExposureTexture = AddLocalExposurePass(
+				GraphBuilder, View,
+				EyeAdaptationParameters,
+				bProcessQuarterResolution ? QuarterResSceneColor : HalfResSceneColor);
 		}
 
 		if (bHistogramEnabled)
 		{
-			HistogramTexture = AddHistogramPass(GraphBuilder, View, EyeAdaptationParameters, DownsampledSceneColor, LastEyeAdaptationTexture);
+			HistogramTexture = AddHistogramPass(
+				GraphBuilder, View,
+				EyeAdaptationParameters,
+				bProcessQuarterResolution ? QuarterResSceneColor : HalfResSceneColor,
+				LastEyeAdaptationTexture);
 		}
 
-		const bool bFFTBloomEnabled = IsFFTBloomEnabled(View);
-		const bool bBasicEyeAdaptationEnabled = bEyeAdaptationEnabled && (AutoExposureMethod == EAutoExposureMethod::AEM_Basic);
-		const bool bLocalExposureBlurredLum = bLocalExposureEnabled && View.FinalPostProcessSettings.LocalExposureBlurredLuminanceBlend > 0.0f;
-
 		FSceneDownsampleChain SceneDownsampleChain;
-		if (bBasicEyeAdaptationEnabled || (bBloomEnabled && bFFTBloomEnabled) || bLocalExposureBlurredLum)
+		if (bProduceSceneColorChain)
 		{
 			const bool bLogLumaInAlpha = true;
-			SceneDownsampleChain.Init(GraphBuilder, View, EyeAdaptationParameters, DownsampledSceneColor, DownsampleQuality, bLogLumaInAlpha);
+			SceneDownsampleChain.Init(
+				GraphBuilder, View,
+				EyeAdaptationParameters,
+				bProcessQuarterResolution ? QuarterResSceneColor : HalfResSceneColor,
+				DownsampleQuality,
+				bLogLumaInAlpha);
 		}
 
 		if (bLocalExposureBlurredLum)
 		{
-			LocalExposureBlurredLogLumTexture = AddLocalExposureBlurredLogLuminancePass(GraphBuilder, View, EyeAdaptationParameters, SceneDownsampleChain.GetTexture(4));
+			const uint32 BlurredLumMip = bProcessQuarterResolution ? 3 : 4;
+			LocalExposureBlurredLogLumTexture = AddLocalExposureBlurredLogLuminancePass(
+				GraphBuilder, View,
+				EyeAdaptationParameters, SceneDownsampleChain.GetTexture(BlurredLumMip));
 		}
 
 		if (bBasicEyeAdaptationEnabled)
 		{
 			// Use the alpha channel in the last downsample (smallest) to compute eye adaptations values.
-			EyeAdaptationTexture = AddBasicEyeAdaptationPass(GraphBuilder, View, EyeAdaptationParameters, SceneDownsampleChain.GetLastTexture(), LastEyeAdaptationTexture);
+			EyeAdaptationTexture = AddBasicEyeAdaptationPass(
+				GraphBuilder, View,
+				EyeAdaptationParameters,
+				SceneDownsampleChain.GetLastTexture(),
+				LastEyeAdaptationTexture);
 		}
 		// Add histogram eye adaptation pass even if no histogram exists to support the manual clamping mode.
 		else if (bEyeAdaptationEnabled)
 		{
-			EyeAdaptationTexture = AddHistogramEyeAdaptationPass(GraphBuilder, View, EyeAdaptationParameters, HistogramTexture);
+			EyeAdaptationTexture = AddHistogramEyeAdaptationPass(
+				GraphBuilder, View,
+				EyeAdaptationParameters,
+				HistogramTexture);
 		}
 
 		FBloomOutputs Bloom;
-
 		if (bBloomEnabled)
 		{
-			FSceneDownsampleChain BloomDownsampleChain;
-
-			const bool bBloomThresholdEnabled = View.FinalPostProcessSettings.BloomThreshold > -1.0f;
-
-			// Reuse the main scene downsample chain if a threshold isn't required for gaussian bloom.
-			const FSceneDownsampleChain* GaussianBloomSceneDownsampleChain;
-			if (SceneDownsampleChain.IsInitialized() && !bBloomThresholdEnabled)
-			{
-				GaussianBloomSceneDownsampleChain = &SceneDownsampleChain;
-			}
-			else
-			{
-				FScreenPassTexture DownsampleInput = DownsampledSceneColor;
-
-				if (bBloomThresholdEnabled)
-				{
-					const float BloomThreshold = View.FinalPostProcessSettings.BloomThreshold;
-
-					FBloomSetupInputs SetupPassInputs;
-					SetupPassInputs.SceneColor = DownsampleInput;
-					SetupPassInputs.EyeAdaptationTexture = EyeAdaptationTexture;
-					SetupPassInputs.Threshold = BloomThreshold;
-
-					DownsampleInput = AddBloomSetupPass(GraphBuilder, View, SetupPassInputs);
-				}
-
-				const bool bLogLumaInAlpha = false;
-				BloomDownsampleChain.Init(GraphBuilder, View, EyeAdaptationParameters, DownsampleInput, DownsampleQuality, bLogLumaInAlpha);
-
-				GaussianBloomSceneDownsampleChain = &BloomDownsampleChain;
-			}
+			const FSceneDownsampleChain* LensFlareSceneDownsampleChain;
 
 			if (bFFTBloomEnabled)
 			{
+				LensFlareSceneDownsampleChain = &SceneDownsampleChain;
+
 				FScreenPassTexture HalfResolution;
 				if (IsFFTBloomQuarterResolutionEnabled())
 				{
-					HalfResolution = SceneDownsampleChain.GetTexture(1);
+					HalfResolution = QuarterResSceneColor;
 				}
 				else
 				{
-					HalfResolution = SceneDownsampleChain.GetFirstTexture();
+					HalfResolution = HalfResSceneColor;
 				}
 
 				FFFTBloomInputs PassInputs;
@@ -906,15 +937,43 @@ void AddPostProcessingPasses(FRDGBuilder& GraphBuilder, const FViewInfo& View, c
 			}
 			else
 			{
-				Bloom = AddGaussianBloomPasses(GraphBuilder, View, GaussianBloomSceneDownsampleChain);
-			}
+				FSceneDownsampleChain BloomDownsampleChain;
+
+				const bool bBloomThresholdEnabled = View.FinalPostProcessSettings.BloomThreshold > -1.0f;
 
-			FScreenPassTexture LensFlares = AddLensFlaresPass(GraphBuilder, View, Bloom.Bloom, *GaussianBloomSceneDownsampleChain);
+				// Reuse the main scene downsample chain if a threshold isn't required for gaussian bloom.
+				if (SceneDownsampleChain.IsInitialized() && !bBloomThresholdEnabled)
+				{
+					LensFlareSceneDownsampleChain = &SceneDownsampleChain;
+				}
+				else
+				{
+					FScreenPassTexture DownsampleInput = bProcessQuarterResolution ? QuarterResSceneColor : HalfResSceneColor;
+
+					if (bBloomThresholdEnabled)
+					{
+						const float BloomThreshold = View.FinalPostProcessSettings.BloomThreshold;
+
+						FBloomSetupInputs SetupPassInputs;
+						SetupPassInputs.SceneColor = DownsampleInput;
+						SetupPassInputs.EyeAdaptationTexture = EyeAdaptationTexture;
+						SetupPassInputs.Threshold = BloomThreshold;
+
+						DownsampleInput = AddBloomSetupPass(GraphBuilder, View, SetupPassInputs);
+					}
+
+					const bool bLogLumaInAlpha = false;
+					BloomDownsampleChain.Init(GraphBuilder, View, EyeAdaptationParameters, DownsampleInput, DownsampleQuality, bLogLumaInAlpha);
+
+					LensFlareSceneDownsampleChain = &BloomDownsampleChain;
+				}
+
+				Bloom = AddGaussianBloomPasses(GraphBuilder, View, LensFlareSceneDownsampleChain);
+			}
 
-			if (LensFlares.IsValid())
+			if (bLensFlareEnabled)
 			{
-				// Lens flares are composited with bloom.
-				Bloom.Bloom = LensFlares;
+				Bloom.Bloom = AddLensFlaresPass(GraphBuilder, View, Bloom.Bloom, *LensFlareSceneDownsampleChain);
 			}
 		}
 
