Fixes some warning when cooking motion blur shaders

diff --git a/Engine/Source/Runtime/Renderer/Private/PostProcess/PostProcessMotionBlur.cpp b/Engine/Source/Runtime/Renderer/Private/PostProcess/PostProcessMotionBlur.cpp
--- a/Engine/Source/Runtime/Renderer/Private/PostProcess/PostProcessMotionBlur.cpp
+++ b/Engine/Source/Runtime/Renderer/Private/PostProcess/PostProcessMotionBlur.cpp
@@ -386,7 +386,7 @@ class FSetupMotionBlurFilterCS : public FMotionBlurShader
 	SHADER_USE_PARAMETER_STRUCT(FSetupMotionBlurFilterCS, FMotionBlurShader);
 
 	BEGIN_SHADER_PARAMETER_STRUCT(FParameters, )
-		SHADER_PARAMETER_RDG_BUFFER_SRV(RWStructuredBuffer<uint>, TileListsSizeBuffer)
+		SHADER_PARAMETER_RDG_BUFFER_SRV(StructuredBuffer<uint>, TileListsSizeBuffer)
 		SHADER_PARAMETER_RDG_BUFFER_UAV(RWStructuredBuffer<uint>, DispatchParametersOutput)
 	END_SHADER_PARAMETER_STRUCT()
 
@@ -448,8 +448,8 @@ class FMotionBlurFilterCS : public FMotionBlurShader
 		SHADER_PARAMETER(FVector2f, TranslucencyExtentInverse)
 
 		SHADER_PARAMETER(int32, TileListOffset)
-		SHADER_PARAMETER_RDG_BUFFER_SRV(RWStructuredBuffer<uint>, TileListsBuffer)
-		SHADER_PARAMETER_RDG_BUFFER_SRV(RWStructuredBuffer<uint>, TileListsSizeBuffer)
+		SHADER_PARAMETER_RDG_BUFFER_SRV(StructuredBuffer<uint>, TileListsBuffer)
+		SHADER_PARAMETER_RDG_BUFFER_SRV(StructuredBuffer<uint>, TileListsSizeBuffer)
 		RDG_BUFFER_ACCESS(DispatchParameters, ERHIAccess::IndirectArgs)
 
 		SHADER_PARAMETER_RDG_TEXTURE_UAV(RWTexture2D<float4>, SceneColorOutputMip0)
