Fix cook warning with double cvar name

diff --git a/Engine/Source/Runtime/Renderer/Private/RayTracing/RayTracingDynamicGeometry.cpp b/Engine/Source/Runtime/Renderer/Private/RayTracing/RayTracingDynamicGeometry.cpp
--- a/Engine/Source/Runtime/Renderer/Private/RayTracing/RayTracingDynamicGeometry.cpp
+++ b/Engine/Source/Runtime/Renderer/Private/RayTracing/RayTracingDynamicGeometry.cpp
@@ -17,7 +17,7 @@ static FAutoConsoleVariableRef CVarRTDynGeomSharedVertexBufferSizeInMB(
 
 static int32 GRTDynGeomSharedVertexBufferGarbageCollectLatency = 30;
 static FAutoConsoleVariableRef CVarRTDynGeomSharedVertexBufferGarbageCollectLatency(
-	TEXT("r.RayTracing.DynamicGeometry.SharedVertexBufferSizeInMB"),
+	TEXT("r.RayTracing.DynamicGeometry.SharedVertexBufferGarbageCollectLatency"),
 	GRTDynGeomSharedVertexBufferGarbageCollectLatency,
 	TEXT("Amount of update cycles before a heap is deleted when not used (default 30)."),
 	ECVF_RenderThreadSafe
