Allows TSR to produce MotionBlur's velocity flatten textures.

Adds 20us to TSR's DilateVelocity, but eliminate MotionBlur's velocity flatten that is 120us on console.

diff --git a/Engine/Shaders/Private/MotionBlur/MotionBlurApply.usf b/Engine/Shaders/Private/MotionBlur/MotionBlurApply.usf
--- a/Engine/Shaders/Private/MotionBlur/MotionBlurApply.usf
+++ b/Engine/Shaders/Private/MotionBlur/MotionBlurApply.usf
@@ -158,7 +158,7 @@ float4 NormalizeAccumulator(float4 Accumulator, float AccumulatorWeight)
 	return -min(-Accumulator * rcp(AccumulatorWeight), float(0.0));
 }
 
-#if CONFIG_COMPILE_FP16
+#if CONFIG_MOTION_BLUR_COMPILE_FP16
 
 mb_half NormalizeAccumulator(mb_half Accumulator, mb_half AccumulatorWeight)
 {
@@ -170,7 +170,7 @@ mb_half4 NormalizeAccumulator(mb_half4 Accumulator, mb_half AccumulatorWeight)
 	return -min(-Accumulator * rcp(AccumulatorWeight), mb_half(0.0));
 }
 
-#endif // CONFIG_COMPILE_FP16
+#endif // CONFIG_MOTION_BLUR_COMPILE_FP16
 
 void NormalizeAccumulatorWithHoleFill(
 	mb_half4 Color,
@@ -942,7 +942,7 @@ void MainCS(
 		SceneColorOutputMip1[OutputPixelCoordMip1] = HalfResOutput;
 		SceneColorOutputMip2[OutputPixelCoordMip2] = QuarterResOutput;
 		
-		#if DEBUG_OUTPUT
+		#if DEBUG_MOTION_BLUR_OUTPUT
 			DebugOutput[OutputPixelCoord0Mip0] = Debug;
 			DebugOutput[OutputPixelCoord1Mip0] = Debug;
 			DebugOutput[OutputPixelCoord2Mip0] = Debug;
@@ -1066,7 +1066,7 @@ void MainCS(
 		SceneColorOutputMip1[OutputPixelCoordMip1] = HalfResOutput;
 		SceneColorOutputMip2[OutputPixelCoordMip2] = QuarterResOutput;
 
-		#if DEBUG_OUTPUT
+		#if DEBUG_MOTION_BLUR_OUTPUT
 			DebugOutput[OutputPixelCoordMip0] = Debug;
 		#endif
 	}
diff --git a/Engine/Shaders/Private/MotionBlur/MotionBlurCommon.ush b/Engine/Shaders/Private/MotionBlur/MotionBlurCommon.ush
--- a/Engine/Shaders/Private/MotionBlur/MotionBlurCommon.ush
+++ b/Engine/Shaders/Private/MotionBlur/MotionBlurCommon.ush
@@ -13,7 +13,7 @@
 
 //------------------------------------------------------- CONFIG
 
-#define DEBUG_OUTPUT 0
+#define DEBUG_MOTION_BLUR_OUTPUT 0
 
 #ifndef CONFIG_MAX_RANGE_SIZE
 	#error need to set CONFIG_MAX_RANGE_SIZE
@@ -26,9 +26,9 @@
 #define VELOCITY_FILTER_TILE_SIZE 16
 
 #if PLATFORM_SUPPORTS_REAL_TYPES && 0
-	#define CONFIG_COMPILE_FP16 1
+	#define CONFIG_MOTION_BLUR_COMPILE_FP16 1
 #else
-	#define CONFIG_COMPILE_FP16 0
+	#define CONFIG_MOTION_BLUR_COMPILE_FP16 0
 #endif
 
 #define CONFIG_MINIMAL_PIXEL_VELOCITY 0.5
@@ -36,7 +36,7 @@
 
 //------------------------------------------------------- CONSTANTS
 
-#if CONFIG_COMPILE_FP16
+#if CONFIG_MOTION_BLUR_COMPILE_FP16
 	#define mb_half  half
 	#define mb_half2 half2
 	#define mb_half3 half3
diff --git a/Engine/Shaders/Private/MotionBlur/MotionBlurFilterTileClassify.usf b/Engine/Shaders/Private/MotionBlur/MotionBlurFilterTileClassify.usf
--- a/Engine/Shaders/Private/MotionBlur/MotionBlurFilterTileClassify.usf
+++ b/Engine/Shaders/Private/MotionBlur/MotionBlurFilterTileClassify.usf
@@ -178,7 +178,7 @@ void MainCS(
 	{
 		TileListsOutput[GlobalAllocation + TileListMaxSize * TileClassify] = uint(FilterTileCoord.x) | (uint(FilterTileCoord.y) << 16);
 
-		#if DEBUG_OUTPUT
+		#if DEBUG_MOTION_BLUR_OUTPUT
 			DebugOutput[FilterTileCoord] = Debug;
 		#endif
 	}
diff --git a/Engine/Shaders/Private/MotionBlur/MotionBlurVelocityFlatten.usf b/Engine/Shaders/Private/MotionBlur/MotionBlurVelocityFlatten.usf
--- a/Engine/Shaders/Private/MotionBlur/MotionBlurVelocityFlatten.usf
+++ b/Engine/Shaders/Private/MotionBlur/MotionBlurVelocityFlatten.usf
@@ -19,80 +19,23 @@
 //------------------------------------------------------- PARAMETERS
 
 SCREEN_PASS_TEXTURE_VIEWPORT(Velocity)
-SCREEN_PASS_TEXTURE_VIEWPORT(VelocityTile)
 
-#if CAMERA_MOTION_BLUR_MODE == 2
+#if CAMERA_MOTION_BLUR_MODE == 2 // TODO: move into MotionBlurVelocityFlatten.ush
 float4x4 ClipToPrevClipOverride;
 #endif
 
-float2 VelocityScale;
-float VelocityMax;
-
 Texture2D VelocityTexture;
-Texture2D VelocityTileTexture;
 Texture2D DepthTexture;
 
-RWTexture2D<float4>	OutVelocityFlatTexture;
+RWTexture2D<float3>	OutVelocityFlatTexture;
 RWTexture2D<float4>	OutVelocityTile_Textures_0;
 RWTexture2D<float4>	OutVelocityTile_Textures_1;
 
 RWTexture2D<float4>	DebugOutput;
 
 
-//------------------------------------------------------- LDS
-
-#if DEBUG_OUTPUT
-groupshared float4 Debug[THREADGROUP_TOTALSIZE];
-#endif
-
-groupshared float4 SharedArray0[THREADGROUP_TOTALSIZE];
-groupshared float2 SharedArray1[THREADGROUP_TOTALSIZE];
-
-
 //------------------------------------------------------- FUNCTIONS
 
-void WritePolarVelocityRangeToLDS(uint LDSSlotId, FVelocityRange V)
-{
-	SharedArray0[LDSSlotId] = float4(V.Max[0], V.Min);
-	
-	#if CONFIG_MAX_RANGE_SIZE > 1
-		SharedArray1[LDSSlotId] = V.Max[1];
-	#endif
-}
-
-FVelocityRange ReadPolarVelocityRangeFromLDS(uint LDSSlotId)
-{
-	FVelocityRange V;
-	V.Max[0] = SharedArray0[LDSSlotId].xy;
-	#if CONFIG_MAX_RANGE_SIZE > 1
-		V.Max[1] = SharedArray1[LDSSlotId].xy;
-	#endif
-	V.Min = SharedArray0[LDSSlotId].zw;
-	return V;
-}
-
-void VelovityFlattenStep(uint GroupIndex, const uint GroupSize, inout FVelocityRange V)
-{
-	if (GroupIndex < GroupSize)
-	{
-		FVelocityRange V1 = ReadPolarVelocityRangeFromLDS(GroupIndex + GroupSize);
-		
-		//if (GroupSize == uint(View.GeneralPurposeTweak))
-		//{
-		//	Debug[GroupIndex].r = V.Max[1].g;
-		//}
-
-		V = ReducePolarVelocityRange(V, V1);
-		
-		//if (GroupSize == uint(View.GeneralPurposeTweak))
-		//{
-		//	Debug[GroupIndex].g = V.Max[1].g;
-		//}
-
-		WritePolarVelocityRangeToLDS(GroupIndex, V);
-	}
-}
-
 float2 GetCameraMotionVelocity(uint2 DispatchThreadId, float Depth, float4x4 ClipToPrevClip)
 {
 	// Compute velocity due to camera motion.
@@ -107,16 +50,14 @@ float2 GetCameraMotionVelocity(uint2 DispatchThreadId, float Depth, float4x4 Cli
 
 //------------------------------------------------------- ENTRY POINT
 
-[numthreads(TILE_SIZE, TILE_SIZE, 1)]
+[numthreads(VELOCITY_FLATTEN_TILE_SIZE, VELOCITY_FLATTEN_TILE_SIZE, 1)]
 void MainCS(
 	uint2 GroupId : SV_GroupID,
 	uint2 DispatchThreadId : SV_DispatchThreadID,
 	uint2 GroupThreadId : SV_GroupThreadID,
 	uint GroupIndex : SV_GroupIndex) 
 { 
-#if DEBUG_OUTPUT
-	Debug[GroupIndex] = 0.0;
-#endif
+	float4 Debug = 0;
 
 	int2 PixelPos = min(DispatchThreadId.xy + Velocity_ViewportMin, Velocity_ViewportMax - 1);
 	int2 DepthOffset = 0;
@@ -154,7 +95,7 @@ void MainCS(
 	}
 
 	float4 EncodedVelocity = VelocityTexture[ PixelPos + DepthOffset ];
-	float Depth = DepthTexture[ PixelPos + DepthOffset ].x;
+	float DeviceZ = DepthTexture[ PixelPos + DepthOffset ].x;
 
 	float2 Velocity = 0.0f;
 	if (EncodedVelocity.x > 0.0)
@@ -162,11 +103,11 @@ void MainCS(
 		Velocity = DecodeVelocityFromTexture(EncodedVelocity).xy;
 #if CAMERA_MOTION_BLUR_MODE != 1
 		// Cancel motion blur caused by the current camera movements
-		Velocity -= GetCameraMotionVelocity(DispatchThreadId, Depth, View.ClipToPrevClip);
+		Velocity -= GetCameraMotionVelocity(DispatchThreadId, DeviceZ, View.ClipToPrevClip);
 #endif
 #if CAMERA_MOTION_BLUR_MODE == 2
 		// Apply motion blur caused by the custom camera movements
-		Velocity += GetCameraMotionVelocity(DispatchThreadId, Depth, ClipToPrevClipOverride);
+		Velocity += GetCameraMotionVelocity(DispatchThreadId, DeviceZ, ClipToPrevClipOverride);
 #endif
 	}
 #if CAMERA_MOTION_BLUR_MODE > 0
@@ -174,85 +115,31 @@ void MainCS(
 	{
 #if CAMERA_MOTION_BLUR_MODE == 2
 		// Apply motion blur caused by the custom camera movements
-		Velocity = GetCameraMotionVelocity(DispatchThreadId, Depth, ClipToPrevClipOverride);
+		Velocity = GetCameraMotionVelocity(DispatchThreadId, DeviceZ, ClipToPrevClipOverride);
 #else
 		// Apply motion blur caused by the current camera movements
-		Velocity = GetCameraMotionVelocity(DispatchThreadId, Depth, View.ClipToPrevClip);
+		Velocity = GetCameraMotionVelocity(DispatchThreadId, DeviceZ, View.ClipToPrevClip);
 #endif
 	}
 #endif
 
-	Velocity *= VelocityScale;
-
-	float2 VelocityPolar = CartesianToPolar(Velocity);
+	FVelocityRange VelocityPolarRange;
+	float2 VelocityPolar;
+	ReduceVelocityFlattenTile(GroupIndex, Velocity, /* out */ VelocityPolarRange, /* out */ VelocityPolar);
 
-	// If the velocity vector was zero length, VelocityPolar will contain NaNs.
-	if (any(isnan(VelocityPolar)))
+	// Output velcity flatten texture.
 	{
-		VelocityPolar = float2(0.0f, 0.0f);
+		uint2 OutputPos = PixelPos.xy < Velocity_ViewportMax ? PixelPos : uint(~0).xx;
+		OutVelocityFlatTexture[OutputPos] = EncodeVelocityFlatten(VelocityPolar, DeviceZ);
 	}
 
-	bool bInsideViewport = all(PixelPos.xy < Velocity_ViewportMax);
-
-	// 11:11:10  (VelocityLength, VelocityAngle, Depth)
-	{
-		float2 EncodedPolarVelocity;
-		EncodedPolarVelocity.x = VelocityPolar.x;
-		EncodedPolarVelocity.y = VelocityPolar.y * (0.5 / PI) + 0.5;
-
-		BRANCH
-		if (bInsideViewport)
-		{
-			OutVelocityFlatTexture[PixelPos] = float3(EncodedPolarVelocity, ConvertFromDeviceZ(Depth)).xyzz;
-		}
-	}
-
-	// Limit velocity
-	VelocityPolar.x = min(VelocityPolar.x, VelocityMax);
-
-	FVelocityRange VelocityPolarRange = SetupPolarVelocityRange(VelocityPolar);
-
-	// TODO
-	//float4 VelocityMinMax = VelocityPolar.xyxy;
-	//VelocityMinMax.x = bInsideViewport ? VelocityMinMax.x : 2;
-	//VelocityMinMax.z = bInsideViewport ? VelocityMinMax.z : 0;
-
-	WritePolarVelocityRangeToLDS(GroupIndex, VelocityPolarRange);
-
-	GroupMemoryBarrierWithGroupSync();
-
-#if THREADGROUP_TOTALSIZE > 512
-	VelovityFlattenStep(GroupIndex, /* GroupSize = */ 512, /* inout */ VelocityPolarRange);
-	GroupMemoryBarrierWithGroupSync();
-#endif
-#if THREADGROUP_TOTALSIZE > 256
-	VelovityFlattenStep(GroupIndex, /* GroupSize = */ 256, /* inout */ VelocityPolarRange);
-	GroupMemoryBarrierWithGroupSync();
-#endif
-#if THREADGROUP_TOTALSIZE > 128
-	VelovityFlattenStep(GroupIndex, /* GroupSize = */ 128, /* inout */ VelocityPolarRange);
-	GroupMemoryBarrierWithGroupSync();
-#endif
-#if THREADGROUP_TOTALSIZE > 64
-	VelovityFlattenStep(GroupIndex, /* GroupSize = */ 64, /* inout */ VelocityPolarRange);
-	GroupMemoryBarrierWithGroupSync();
-#endif
-
-	// Safe for vector sizes 32 or larger, AMD and NV
-	// TODO Intel variable size vector
-	VelovityFlattenStep(GroupIndex, /* GroupSize = */ 32, /* inout */ VelocityPolarRange);
-	VelovityFlattenStep(GroupIndex, /* GroupSize = */ 16, /* inout */ VelocityPolarRange);
-	VelovityFlattenStep(GroupIndex, /* GroupSize = */  8, /* inout */ VelocityPolarRange);
-	VelovityFlattenStep(GroupIndex, /* GroupSize = */  4, /* inout */ VelocityPolarRange);
-	VelovityFlattenStep(GroupIndex, /* GroupSize = */  2, /* inout */ VelocityPolarRange);
-	VelovityFlattenStep(GroupIndex, /* GroupSize = */  1, /* inout */ VelocityPolarRange);
-
-	if (GroupIndex == 0)
+	// Output tile
 	{
-		StoreVelocityRange(OutVelocityTile_Textures_0, OutVelocityTile_Textures_1, GroupId, PolarToCartesian(VelocityPolarRange));
+		uint2 OutputPos = GroupIndex == 0 ? GroupId : uint(~0).xx;
+		StoreVelocityRange(OutVelocityTile_Textures_0, OutVelocityTile_Textures_1, OutputPos, PolarToCartesian(VelocityPolarRange));
 	}
 	
-#if DEBUG_OUTPUT
-	DebugOutput[PixelPos] = Debug[GroupIndex];
+#if DEBUG_MOTION_BLUR_OUTPUT
+	DebugOutput[PixelPos] = Debug;
 #endif
 }
diff --git a/Engine/Shaders/Private/MotionBlur/MotionBlurVelocityFlatten.ush b/Engine/Shaders/Private/MotionBlur/MotionBlurVelocityFlatten.ush
--- a/Engine/Shaders/Private/MotionBlur/MotionBlurVelocityFlatten.ush
+++ b/Engine/Shaders/Private/MotionBlur/MotionBlurVelocityFlatten.ush
@@ -3,4 +3,123 @@
 #include "MotionBlurCommon.ush"
 
 
-// TODO: merge velocity flatten pass within TSR's dilate velocity pass.
+//------------------------------------------------------- CONSTANTS
+
+#define VELOCITY_FLATTEN_THREAD_COUNT (VELOCITY_FLATTEN_TILE_SIZE * VELOCITY_FLATTEN_TILE_SIZE)
+
+
+//------------------------------------------------------- PARAMETERS
+
+float2 VelocityScale;
+float VelocityMax;
+
+
+//------------------------------------------------------- LDS
+
+groupshared float4 SharedArray0[VELOCITY_FLATTEN_THREAD_COUNT];
+groupshared float2 SharedArray1[VELOCITY_FLATTEN_THREAD_COUNT];
+
+
+//------------------------------------------------------- FUNCTIONS
+
+void WritePolarVelocityRangeToLDS(uint LDSSlotId, FVelocityRange V)
+{
+	SharedArray0[LDSSlotId] = float4(V.Max[0], V.Min);
+	
+	#if CONFIG_MAX_RANGE_SIZE > 1
+		SharedArray1[LDSSlotId] = V.Max[1];
+	#endif
+}
+
+FVelocityRange ReadPolarVelocityRangeFromLDS(uint LDSSlotId)
+{
+	FVelocityRange V;
+	V.Max[0] = SharedArray0[LDSSlotId].xy;
+	#if CONFIG_MAX_RANGE_SIZE > 1
+		V.Max[1] = SharedArray1[LDSSlotId].xy;
+	#endif
+	V.Min = SharedArray0[LDSSlotId].zw;
+	return V;
+}
+
+void VelocityFlattenStep(uint GroupIndex, const uint GroupSize, inout FVelocityRange V)
+{
+	if (GroupIndex < GroupSize)
+	{
+		FVelocityRange V1 = ReadPolarVelocityRangeFromLDS(GroupIndex + GroupSize);
+		
+		//if (GroupSize == uint(View.GeneralPurposeTweak))
+		//{
+		//	Debug[GroupIndex].r = V.Max[1].g;
+		//}
+
+		V = ReducePolarVelocityRange(V, V1);
+		
+		//if (GroupSize == uint(View.GeneralPurposeTweak))
+		//{
+		//	Debug[GroupIndex].g = V.Max[1].g;
+		//}
+
+		WritePolarVelocityRangeToLDS(GroupIndex, V);
+	}
+}
+
+void ReduceVelocityFlattenTile(uint GroupIndex, float2 Velocity, out FVelocityRange OutVelocityPolarRange, out float2 VelocityPolar)
+{
+	Velocity *= VelocityScale;
+
+	VelocityPolar = CartesianToPolar(Velocity);
+
+	// If the velocity vector was zero length, VelocityPolar will contain NaNs.
+	if (any(isnan(VelocityPolar)))
+	{
+		VelocityPolar = float2(0.0f, 0.0f);
+	}
+
+	// Limit velocity
+	VelocityPolar.x = min(VelocityPolar.x, VelocityMax);
+
+	FVelocityRange VelocityPolarRange = SetupPolarVelocityRange(VelocityPolar);
+
+	WritePolarVelocityRangeToLDS(GroupIndex, VelocityPolarRange);
+
+	GroupMemoryBarrierWithGroupSync();
+
+#if VELOCITY_FLATTEN_THREAD_COUNT > 512
+	VelocityFlattenStep(GroupIndex, /* GroupSize = */ 512, /* inout */ VelocityPolarRange);
+	GroupMemoryBarrierWithGroupSync();
+#endif
+#if VELOCITY_FLATTEN_THREAD_COUNT > 256
+	VelocityFlattenStep(GroupIndex, /* GroupSize = */ 256, /* inout */ VelocityPolarRange);
+	GroupMemoryBarrierWithGroupSync();
+#endif
+#if VELOCITY_FLATTEN_THREAD_COUNT > 128
+	VelocityFlattenStep(GroupIndex, /* GroupSize = */ 128, /* inout */ VelocityPolarRange);
+	GroupMemoryBarrierWithGroupSync();
+#endif
+#if VELOCITY_FLATTEN_THREAD_COUNT > 64
+	VelocityFlattenStep(GroupIndex, /* GroupSize = */ 64, /* inout */ VelocityPolarRange);
+	GroupMemoryBarrierWithGroupSync();
+#endif
+
+	// Safe for vector sizes 32 or larger, AMD and NV
+	// TODO Intel variable size vector
+	VelocityFlattenStep(GroupIndex, /* GroupSize = */ 32, /* inout */ VelocityPolarRange);
+	VelocityFlattenStep(GroupIndex, /* GroupSize = */ 16, /* inout */ VelocityPolarRange);
+	VelocityFlattenStep(GroupIndex, /* GroupSize = */  8, /* inout */ VelocityPolarRange);
+	VelocityFlattenStep(GroupIndex, /* GroupSize = */  4, /* inout */ VelocityPolarRange);
+	VelocityFlattenStep(GroupIndex, /* GroupSize = */  2, /* inout */ VelocityPolarRange);
+	VelocityFlattenStep(GroupIndex, /* GroupSize = */  1, /* inout */ VelocityPolarRange);
+
+	OutVelocityPolarRange = VelocityPolarRange;
+}
+
+float3 EncodeVelocityFlatten(float2 VelocityPolar, float DeviceZ)
+{
+	// 11:11:10  (VelocityLength, VelocityAngle, Depth)
+	float2 EncodedPolarVelocity;
+	EncodedPolarVelocity.x = VelocityPolar.x;
+	EncodedPolarVelocity.y = VelocityPolar.y * (0.5 / PI) + 0.5;
+
+	return float3(EncodedPolarVelocity, ConvertFromDeviceZ(DeviceZ)).xyz;
+}
diff --git a/Engine/Shaders/Private/TemporalSuperResolution/TSRDilateVelocity.usf b/Engine/Shaders/Private/TemporalSuperResolution/TSRDilateVelocity.usf
--- a/Engine/Shaders/Private/TemporalSuperResolution/TSRDilateVelocity.usf
+++ b/Engine/Shaders/Private/TemporalSuperResolution/TSRDilateVelocity.usf
@@ -2,10 +2,19 @@
 
 #include "TSRCommon.ush"
 
+#if DIM_MOTION_BLUR_DIRECTIONS > 0
+	#define CONFIG_MAX_RANGE_SIZE (DIM_MOTION_BLUR_DIRECTIONS)
+	#include "../MotionBlur/MotionBlurVelocityFlatten.ush"
+#endif
+
 
 //------------------------------------------------------- CONFIG
 
-#define TILE_SIZE 8
+#if DIM_MOTION_BLUR_DIRECTIONS > 0
+	#define TILE_SIZE (VELOCITY_FLATTEN_TILE_SIZE)
+#else
+	#define TILE_SIZE 8
+#endif
 
 #define AA_CROSS 1
 
@@ -27,6 +36,10 @@ RWTexture2D<uint> PrevUseCountOutput;
 RWTexture2D<uint> PrevClosestDepthOutput;
 RWTexture2D<uint> ParallaxFactorOutput;
 
+RWTexture2D<float3>	VelocityFlattenOutput;
+RWTexture2D<float4>	VelocityTileOutput0;
+RWTexture2D<float4>	VelocityTileOutput1;
+
 
 //------------------------------------------------------- ENTRY POINT
 
@@ -152,23 +165,30 @@ void MainCS(
 	}
 	#endif
 
-	// Scatter parallax rejection
-	float2 PrevScreenPos = ScreenPos - FinalScreenVelocity;
-	float2 PrevInputBufferUV = (InputInfo_ScreenPosToViewportScale * PrevScreenPos + InputInfo_ScreenPosToViewportBias) * InputInfo_ExtentInverse;
-
+	// Do motion blur velocity flatten.
+	#if DIM_MOTION_BLUR_DIRECTIONS > 0
+		FVelocityRange VelocityPolarRange;
+		float2 VelocityPolar;
+		ReduceVelocityFlattenTile(GroupThreadIndex, FinalScreenVelocity, /* out */ VelocityPolarRange, /* out */ VelocityPolar);
+	#endif
 
 	Debug.rg = VelocityPixelOffset * 0.5 + 0.5;
 
-	if (all(InputPixelPos < InputInfo_ViewportMax))
 	{
+		bool bValidOutputPixel = all(InputPixelPos < InputInfo_ViewportMax);
+		
+		// Scatter parallax rejection
 		#if 1
 		{
+			float2 PrevScreenPos = ScreenPos - FinalScreenVelocity;
+			float2 PrevInputBufferUV = (InputInfo_ScreenPosToViewportScale * PrevScreenPos + InputInfo_ScreenPosToViewportBias) * InputInfo_ExtentInverse;
+
 			float2 PixelVelocity = FinalScreenVelocity * 0.5 * InputInfo_ViewportSize;
 			uint EncodedHoleFillingVelocity = EncodeHoleFillingVelocity(PixelVelocity);
 
 			uint PrevClosestDepth = (f32tof16(ClosestDeviceZ) << VELOCITY_HOLE_FILLING_BITS) | EncodedHoleFillingVelocity;
 			
-			bool bValidHistoryCoord = all(PrevInputBufferUV > PrevOutputBufferUVMin && PrevInputBufferUV < PrevOutputBufferUVMax);
+			bool bValidHistoryCoord = bValidOutputPixel && all(PrevInputBufferUV > PrevOutputBufferUVMin && PrevInputBufferUV < PrevOutputBufferUVMax);
 
 			FBilinearSampleInfos BilinearInter = GetBilinearSampleLevelInfos(PrevInputBufferUV, InputInfo_Extent, InputInfo_ExtentInverse);
 
@@ -200,14 +220,26 @@ void MainCS(
 		
 		float4 EncodedVelocity = SceneVelocityTexture[uint2(InputBufferUV * InputInfo_Extent)];
 		
-		//DilatedVelocityOutput[InputPixelPos] = EncodedVelocity.xy; //EncodeVelocityToTexture(float3(FinalScreenVelocity, 0.0)).xy;
-		DilatedVelocityOutput[InputPixelPos] = EncodeVelocityToTexture(float3(FinalScreenVelocity, 0.0)).xy;
-		ClosestDepthOutput[InputPixelPos] = ClosestDeviceZ;
-		ParallaxFactorOutput[InputPixelPos] = CompressParallaxFactor(ParallaxFactorForHistory) | (CompressParallaxFactor(ParallaxFactorForHistoryRejection) << 4);
+		uint2 OutputPixelPos = bValidOutputPixel ? InputPixelPos : uint(~0).xx;
+
+		//DilatedVelocityOutput[OutputPixelPos] = EncodedVelocity.xy; //EncodeVelocityToTexture(float3(FinalScreenVelocity, 0.0)).xy;
+		DilatedVelocityOutput[OutputPixelPos] = EncodeVelocityToTexture(float3(FinalScreenVelocity, 0.0)).xy;
+		ClosestDepthOutput[OutputPixelPos] = ClosestDeviceZ;
+		ParallaxFactorOutput[OutputPixelPos] = CompressParallaxFactor(ParallaxFactorForHistory) | (CompressParallaxFactor(ParallaxFactorForHistoryRejection) << 4);
+
+		// Output motion blur velocity flatten
+		#if DIM_MOTION_BLUR_DIRECTIONS > 0
+		{
+			VelocityFlattenOutput[OutputPixelPos] = EncodeVelocityFlatten(VelocityPolar, ClosestDeviceZ);
+			
+			uint2 OutputTilePos = GroupThreadIndex == 0 ? GroupId : uint(~0).xx;
+			StoreVelocityRange(VelocityTileOutput0, VelocityTileOutput1, OutputTilePos, PolarToCartesian(VelocityPolarRange));
+		}
+		#endif
 
 		#if DEBUG_OUTPUT
 		{
-			DebugOutput[InputPixelPos] = Debug;
+			DebugOutput[OutputPixelPos] = Debug;
 		}
 		#endif
 	}
diff --git a/Engine/Source/Runtime/Renderer/Private/PostProcess/PostProcessMotionBlur.cpp b/Engine/Source/Runtime/Renderer/Private/PostProcess/PostProcessMotionBlur.cpp
--- a/Engine/Source/Runtime/Renderer/Private/PostProcess/PostProcessMotionBlur.cpp
+++ b/Engine/Source/Runtime/Renderer/Private/PostProcess/PostProcessMotionBlur.cpp
@@ -53,6 +53,11 @@ namespace
 		TEXT("Whether to do motion blur filter dynamically at half res under heavy motion."),
 		ECVF_RenderThreadSafe);
 
+	TAutoConsoleVariable<int32> CVarMotionBlurAllowExternalVelocityFlatten(
+		TEXT("r.MotionBlur.AllowExternalVelocityFlatten"), 1,
+		TEXT("Whether to allow motion blur's velocity flatten into other pass."),
+		ECVF_RenderThreadSafe);
+
 	TAutoConsoleVariable<int32> CVarAllowMotionBlurInVR(
 		TEXT("vr.AllowMotionBlurInVR"),
 		0,
@@ -88,7 +93,7 @@ DECLARE_GPU_STAT(MotionBlur)
 
 }
 
-const int32 kMotionBlurFlattenTileSize = 16;
+const int32 kMotionBlurFlattenTileSize = FVelocityFlattenTextures::kTileSize;
 const int32 kMotionBlurFilterTileSize = 16;
 const int32 kMotionBlurComputeTileSizeX = 8;
 const int32 kMotionBlurComputeTileSizeY = 8;
@@ -175,6 +180,12 @@ int32 GetMotionBlurDirections()
 	return FMath::Clamp(CVarMotionBlurDirections.GetValueOnRenderThread(), 1, 2);
 }
 
+// static
+bool FVelocityFlattenTextures::AllowExternal()
+{
+	return CVarMotionBlurAllowExternalVelocityFlatten.GetValueOnRenderThread() != 0;
+}
+
 FRHISamplerState* GetMotionBlurColorSampler()
 {
 	bool bFiltered = false;
@@ -228,11 +239,11 @@ FRDGTextureUAVRef CreateDebugUAV(FRDGBuilder& GraphBuilder, const FIntPoint& Ext
 }
 
 BEGIN_SHADER_PARAMETER_STRUCT(FVecocityTileTextures, )
-	SHADER_PARAMETER_RDG_TEXTURE_ARRAY(Texture2D, Textures, [2])
+	SHADER_PARAMETER_RDG_TEXTURE_ARRAY(Texture2D, Textures, [FVelocityFlattenTextures::kMaxVelocityTileTextureCount])
 END_SHADER_PARAMETER_STRUCT()
 
 BEGIN_SHADER_PARAMETER_STRUCT(FVecocityTileUAVs, )
-	SHADER_PARAMETER_RDG_TEXTURE_UAV_ARRAY(RWTexture2D, Textures, [2])
+	SHADER_PARAMETER_RDG_TEXTURE_UAV_ARRAY(RWTexture2D, Textures, [FVelocityFlattenTextures::kMaxVelocityTileTextureCount])
 END_SHADER_PARAMETER_STRUCT()
 
 FVecocityTileTextures CreateVecocityTileTextures(FRDGBuilder& GraphBuilder, FIntPoint VelocityTileCount, const TCHAR* DebugName, bool ScatterDilatation = false)
@@ -295,6 +306,7 @@ public:
 	BEGIN_SHADER_PARAMETER_STRUCT(FParameters, )
 		SHADER_PARAMETER_STRUCT_REF(FViewUniformShaderParameters, View)
 		SHADER_PARAMETER_STRUCT(FScreenPassTextureViewportParameters, Velocity)
+		SHADER_PARAMETER_STRUCT_INCLUDE(FVelocityFlattenParameters, VelocityFlattenParameters)
 		SHADER_PARAMETER(FMatrix44f, ClipToPrevClipOverride)
 		SHADER_PARAMETER(FVector2f, VelocityScale)
 		SHADER_PARAMETER(float, VelocityMax)
@@ -557,19 +569,35 @@ struct FMotionBlurViewports
 	FScreenTransform ColorToVelocityTransform;
 };
 
+FVelocityFlattenParameters GetVelocityFlattenParameters(const FViewInfo& View)
+{
+	const FSceneViewState* ViewState = View.ViewState;
+
+	const float SceneViewportSizeX = View.GetSecondaryViewRectSize().X;
+	const float SceneViewportSizeY = View.GetSecondaryViewRectSize().Y;
+	const float MotionBlurTimeScale = ViewState ? ViewState->MotionBlurTimeScale : 1.0f;
+
+	// Scale by 0.5 due to blur samples going both ways.
+	const float VelocityScale = MotionBlurTimeScale * View.FinalPostProcessSettings.MotionBlurAmount * 0.5f;
+
+	// 0:no 1:full screen width, percent conversion
+	const float UVVelocityMax = View.FinalPostProcessSettings.MotionBlurMax / 100.0f;
+
+	FVelocityFlattenParameters VelocityFlattenParameters;
+	VelocityFlattenParameters.VelocityScale.X = SceneViewportSizeX * 0.5f * VelocityScale;
+	VelocityFlattenParameters.VelocityScale.Y = -SceneViewportSizeY * 0.5f * VelocityScale;
+	VelocityFlattenParameters.VelocityMax = SceneViewportSizeX * 0.5f * UVVelocityMax;
+	return VelocityFlattenParameters;
+}
+
 void AddMotionBlurVelocityPass(
 	FRDGBuilder& GraphBuilder,
 	const FViewInfo& View,
 	const FMotionBlurViewports& Viewports,
-	FRDGTextureRef ColorTexture,
-	FRDGTextureRef DepthTexture,
-	FRDGTextureRef VelocityTexture,
+	const FMotionBlurInputs& Inputs,
 	FRDGTextureRef* VelocityFlatTextureOutput,
 	FVecocityTileTextures* VelocityTileTexturesOutput)
 {
-	check(ColorTexture);
-	check(DepthTexture);
-	check(VelocityTexture);
 	check(VelocityFlatTextureOutput);
 	check(VelocityTileTexturesOutput);
 
@@ -577,54 +605,50 @@ void AddMotionBlurVelocityPass(
 
 	const FIntPoint VelocityTileCount = Viewports.VelocityTile.Extent;
 
-	// NOTE: Use scene depth's dimensions because velocity can actually be a 1x1 black texture when there are no moving objects in sight.
-	FRDGTextureRef VelocityFlatTexture =
-		GraphBuilder.CreateTexture(
-			FRDGTextureDesc::Create2D(
-				DepthTexture->Desc.Extent,
-				PF_FloatR11G11B10,
-				FClearValueBinding::None,
-				GFastVRamConfig.VelocityFlat | TexCreate_ShaderResource | TexCreate_UAV),
-			TEXT("MotionBlur.VelocityFlat"));
-
-	FVecocityTileTextures VelocityTileTexturesSetup = CreateVecocityTileTextures(
-		GraphBuilder, VelocityTileCount, TEXT("MotionBlur.VelocityTile"));
-
 	// Velocity flatten pass: combines depth / velocity into a single target for sampling efficiency.
+	FRDGTextureRef VelocityFlatTexture = nullptr;
+	FVecocityTileTextures VelocityTileTexturesSetup;
+	if (Inputs.VelocityFlattenTextures.IsValid())
 	{
-		const bool bEnableCameraMotionBlur = View.bCameraMotionBlur.Get(true);
-		const bool bOverrideCameraMotionBlur = View.ClipToPrevClipOverride.IsSet();
-
-		FMotionBlurVelocityFlattenCS::FParameters* PassParameters = GraphBuilder.AllocParameters<FMotionBlurVelocityFlattenCS::FParameters>();
+		ensure(Inputs.VelocityFlattenTextures.VelocityFlatten.ViewRect == View.ViewRect);
+		ensure(Inputs.VelocityFlattenTextures.VelocityTile[0].ViewRect == FIntRect(FIntPoint::ZeroValue, VelocityTileCount));
 
+		VelocityFlatTexture = Inputs.VelocityFlattenTextures.VelocityFlatten.Texture;
+		for (int32 i = 0; i < FVelocityFlattenTextures::kMaxVelocityTileTextureCount; i++)
 		{
-			const FSceneViewState* ViewState = View.ViewState;
-
-			const float TileSize = kMotionBlurFlattenTileSize;
-			const float SceneViewportSizeX = Viewports.Color.Rect.Width();
-			const float SceneViewportSizeY = Viewports.Color.Rect.Height();
-			const float MotionBlurTimeScale = ViewState ? ViewState->MotionBlurTimeScale : 1.0f;
+			VelocityTileTexturesSetup.Textures[i] = Inputs.VelocityFlattenTextures.VelocityTile[i].Texture;
+		}
+	}
+	else
+	{
+		{
+			// NOTE: Use scene depth's dimensions because velocity can actually be a 1x1 black texture when there are no moving objects in sight.
+			FRDGTextureDesc Desc = FRDGTextureDesc::Create2D(
+				Inputs.SceneDepth.Texture->Desc.Extent,
+				PF_FloatR11G11B10,
+				FClearValueBinding::None,
+				GFastVRamConfig.VelocityFlat | TexCreate_ShaderResource | TexCreate_UAV);
 
-			// Scale by 0.5 due to blur samples going both ways.
-			const float VelocityScale = MotionBlurTimeScale * View.FinalPostProcessSettings.MotionBlurAmount * 0.5f;
+			VelocityFlatTexture = GraphBuilder.CreateTexture(Desc, TEXT("MotionBlur.VelocityTile"));
+		}
 
-			// 0:no 1:full screen width, percent conversion
-			const float UVVelocityMax = View.FinalPostProcessSettings.MotionBlurMax / 100.0f;
+		VelocityTileTexturesSetup = CreateVecocityTileTextures(
+			GraphBuilder, VelocityTileCount, TEXT("MotionBlur.VelocityTile"));
 
-			PassParameters->VelocityScale.X = SceneViewportSizeX * 0.5f * VelocityScale;
-			PassParameters->VelocityScale.Y = - SceneViewportSizeY * 0.5f * VelocityScale;
-			PassParameters->VelocityMax = SceneViewportSizeX * 0.5f * UVVelocityMax;
-		}
+		const bool bEnableCameraMotionBlur = View.bCameraMotionBlur.Get(true);
+		const bool bOverrideCameraMotionBlur = View.ClipToPrevClipOverride.IsSet();
 
+		FMotionBlurVelocityFlattenCS::FParameters* PassParameters = GraphBuilder.AllocParameters<FMotionBlurVelocityFlattenCS::FParameters>();
 		PassParameters->View = View.ViewUniformBuffer;
+		PassParameters->VelocityFlattenParameters = GetVelocityFlattenParameters(View);
 		if (bOverrideCameraMotionBlur)
 		{
 			PassParameters->ClipToPrevClipOverride = View.ClipToPrevClipOverride.GetValue();
 		}
 
 		PassParameters->Velocity = Viewports.VelocityParameters;
-		PassParameters->DepthTexture = DepthTexture;
-		PassParameters->VelocityTexture = VelocityTexture;
+		PassParameters->DepthTexture = Inputs.SceneDepth.Texture;
+		PassParameters->VelocityTexture = Inputs.SceneVelocity.Texture;
 
 		PassParameters->OutVelocityFlatTexture = GraphBuilder.CreateUAV(VelocityFlatTexture);
 		PassParameters->OutVelocityTile = CreateUAVs(GraphBuilder, VelocityTileTexturesSetup);
@@ -1169,9 +1193,7 @@ FMotionBlurOutputs AddMotionBlurPass(FRDGBuilder& GraphBuilder, const FViewInfo&
 		GraphBuilder,
 		View,
 		Viewports,
-		Inputs.SceneColor.Texture,
-		Inputs.SceneDepth.Texture,
-		Inputs.SceneVelocity.Texture,
+		Inputs,
 		&VelocityFlatTexture,
 		&VelocityTileTextures);
 
diff --git a/Engine/Source/Runtime/Renderer/Private/PostProcess/PostProcessMotionBlur.h b/Engine/Source/Runtime/Renderer/Private/PostProcess/PostProcessMotionBlur.h
--- a/Engine/Source/Runtime/Renderer/Private/PostProcess/PostProcessMotionBlur.h
+++ b/Engine/Source/Runtime/Renderer/Private/PostProcess/PostProcessMotionBlur.h
@@ -39,6 +39,35 @@ EMotionBlurQuality GetMotionBlurQuality();
 // Returns the global setting for motion blur filter.
 EMotionBlurFilter GetMotionBlurFilter();
 
+// Number of simultaneous direction motion blur can blur per tiles.
+int32 GetMotionBlurDirections();
+
+// Shader parameters for MotionBlurVelocityFlatten.ush when another pass can also generate the FVelocityFlattenTextures
+BEGIN_SHADER_PARAMETER_STRUCT(FVelocityFlattenParameters, )
+	SHADER_PARAMETER(FVector2D, VelocityScale)
+	SHADER_PARAMETER(float, VelocityMax)
+END_SHADER_PARAMETER_STRUCT()
+
+FVelocityFlattenParameters GetVelocityFlattenParameters(const FViewInfo& View);
+
+// Textures generated by the velocity flattening pass.
+struct FVelocityFlattenTextures
+{
+	static constexpr int32 kMaxVelocityTileTextureCount = 2;
+	static constexpr int32 kTileSize = 16;
+
+	FScreenPassTexture VelocityFlatten;
+	FScreenPassTexture VelocityTile[kMaxVelocityTileTextureCount];
+
+	bool IsValid() const
+	{
+		return VelocityFlatten.IsValid() && VelocityTile[0].IsValid();
+	}
+
+	// returns whether FVelocityFlattenTextures can be generated in external system.
+	static bool AllowExternal();
+};
+
 struct FMotionBlurInputs
 {
 	bool bOutputHalfRes = false;
@@ -64,6 +93,9 @@ struct FMotionBlurInputs
 
 	// [Required] Filter to use when processing motion blur.
 	EMotionBlurFilter Filter = EMotionBlurFilter::Separable;
+
+	// [Optional] the velocity flatten is can be generated by a earlier pass.
+	FVelocityFlattenTextures VelocityFlattenTextures;
 };
 
 struct FMotionBlurOutputs
diff --git a/Engine/Source/Runtime/Renderer/Private/PostProcess/PostProcessing.cpp b/Engine/Source/Runtime/Renderer/Private/PostProcess/PostProcessing.cpp
--- a/Engine/Source/Runtime/Renderer/Private/PostProcess/PostProcessing.cpp
+++ b/Engine/Source/Runtime/Renderer/Private/PostProcess/PostProcessing.cpp
@@ -605,7 +605,7 @@ void AddPostProcessingPasses(FRDGBuilder& GraphBuilder, const FViewInfo& View, c
 		const bool bLocalExposureBlurredLum = bLocalExposureEnabled && View.FinalPostProcessSettings.LocalExposureBlurredLuminanceBlend > 0.0f;
 
 		const bool bProcessQuarterResolution = IsPostProcessingQuarterResolutionDownsampleEnabled();
-		const bool bMotionBlurNeedsHalfResInput = PassSequence.IsEnabled(EPass::MotionBlur) && DoesMotionBlurNeedsHalfResInput();
+		const bool bMotionBlurNeedsHalfResInput = PassSequence.IsEnabled(EPass::MotionBlur) && DoesMotionBlurNeedsHalfResInput() && !bVisualizeMotionBlur;
 
 		const bool bIsFullResFFT = bFFTBloomEnabled && IsFFTBloomFullResolutionEnabled();
 		const bool bIsHalfResFFT = bFFTBloomEnabled && !IsFFTBloomFullResolutionEnabled() && !IsFFTBloomQuarterResolutionEnabled();
@@ -675,12 +675,9 @@ void AddPostProcessingPasses(FRDGBuilder& GraphBuilder, const FViewInfo& View, c
 			}
 		}
 
-		// Scene color view rectangle after temporal AA upscale to secondary screen percentage.
-		FIntRect SecondaryViewRect = PrimaryViewRect;
-
 		FScreenPassTexture HalfResSceneColor;
 		FScreenPassTexture QuarterResSceneColor;
-
+		FVelocityFlattenTextures VelocityFlattenTextures;
 		if (TAAConfig != EMainTAAPassConfig::Disabled)
 		{
 			// Whether we allow the temporal AA pass to downsample scene color. It may choose not to based on internal context,
@@ -711,20 +708,21 @@ void AddPostProcessingPasses(FRDGBuilder& GraphBuilder, const FViewInfo& View, c
 
 			UpscalerPassInputs.bAllowDownsampleSceneColor = bAllowSceneDownsample;
 			UpscalerPassInputs.bGenerateOutputMip1 = bMotionBlurNeedsHalfResInput;
+			UpscalerPassInputs.bGenerateVelocityFlattenTextures = PassSequence.IsEnabled(EPass::MotionBlur) && FVelocityFlattenTextures::AllowExternal() && !bVisualizeMotionBlur;
 			UpscalerPassInputs.DownsampleOverrideFormat = DownsampleOverrideFormat;
 			UpscalerPassInputs.SceneColorTexture = SceneColor.Texture;
 			UpscalerPassInputs.SceneDepthTexture = SceneDepth.Texture;
 			UpscalerPassInputs.SceneVelocityTexture = Velocity.Texture;
 			UpscalerPassInputs.SeparateTranslucencyTextures = Inputs.SeparateTranslucencyTextures;
 
-			UpscalerToUse->AddPasses(
+			ITemporalUpscaler::FOutputs Outputs = UpscalerToUse->AddPasses(
 				GraphBuilder,
 				View,
-				UpscalerPassInputs,
-				&SceneColor.Texture,
-				&SecondaryViewRect,
-				&HalfResSceneColor.Texture,
-				&HalfResSceneColor.ViewRect);
+				UpscalerPassInputs);
+
+			SceneColor = Outputs.FullRes;
+			HalfResSceneColor = Outputs.HalfRes;
+			VelocityFlattenTextures = Outputs.VelocityFlattenTextures;
 		}
 		else if (ScreenSpaceRayTracing::ShouldRenderScreenSpaceReflections(View))
 		{
@@ -745,9 +743,6 @@ void AddPostProcessingPasses(FRDGBuilder& GraphBuilder, const FViewInfo& View, c
 			}
 		}
 
-		//! SceneColorTexture is now upsampled to the SecondaryViewRect. Use SecondaryViewRect for input / output.
-		SceneColor.ViewRect = SecondaryViewRect;
-
 		// Post Process Material Chain - SSR Input
 		if (View.ViewState && !View.bStatePrevViewInfoIsReadOnly)
 		{
@@ -776,6 +771,7 @@ void AddPostProcessingPasses(FRDGBuilder& GraphBuilder, const FViewInfo& View, c
 			PassInputs.PostMotionBlurTranslucency = PostMotionBlurTranslucency;
 			PassInputs.Quality = GetMotionBlurQuality();
 			PassInputs.Filter = GetMotionBlurFilter();
+			PassInputs.VelocityFlattenTextures = VelocityFlattenTextures;
 
 			// Motion blur visualization replaces motion blur when enabled.
 			if (bVisualizeMotionBlur)
@@ -2265,26 +2261,16 @@ void AddMobilePostProcessingPasses(FRDGBuilder& GraphBuilder, FScene* Scene, con
 				View.GetSecondaryViewRectSize().X, View.GetSecondaryViewRectSize().Y);
 
 			ITemporalUpscaler::FPassInputs UpscalerPassInputs;
-
-			UpscalerPassInputs.bAllowDownsampleSceneColor = false;
 			UpscalerPassInputs.SceneColorTexture = SceneColor.Texture;
 			UpscalerPassInputs.SceneDepthTexture = SceneDepth.Texture;
 			UpscalerPassInputs.SceneVelocityTexture = Velocity.Texture;
 
-			FIntRect SecondaryViewRect;
-			FScreenPassTexture HalfResolutionSceneColor;
-
-			UpscalerToUse->AddPasses(
+			ITemporalUpscaler::FOutputs Outputs = UpscalerToUse->AddPasses(
 				GraphBuilder,
 				View,
-				UpscalerPassInputs,
-				&SceneColor.Texture,
-				&SecondaryViewRect,
-				&HalfResolutionSceneColor.Texture,
-				&HalfResolutionSceneColor.ViewRect);
-
-			//! SceneColorTexture is now upsampled to the SecondaryViewRect. Use SecondaryViewRect for input / output.
-			SceneColor.ViewRect = SecondaryViewRect;
+				UpscalerPassInputs);
+
+			SceneColor = Outputs.FullRes;
 		}
 	}
 	else
diff --git a/Engine/Source/Runtime/Renderer/Private/PostProcess/TemporalAA.cpp b/Engine/Source/Runtime/Renderer/Private/PostProcess/TemporalAA.cpp
--- a/Engine/Source/Runtime/Renderer/Private/PostProcess/TemporalAA.cpp
+++ b/Engine/Source/Runtime/Renderer/Private/PostProcess/TemporalAA.cpp
@@ -883,12 +883,10 @@ static void AddGen4MainTemporalAAPasses(
 
 bool DoesPlatformSupportTSR(EShaderPlatform Platform);
 
-void AddTemporalSuperResolutionPasses(
+ITemporalUpscaler::FOutputs AddTemporalSuperResolutionPasses(
 	FRDGBuilder& GraphBuilder,
 	const FViewInfo& View,
-	const ITemporalUpscaler::FPassInputs& PassInputs,
-	FRDGTextureRef* OutSceneColorTexture,
-	FIntRect* OutSceneColorViewRect);
+	const ITemporalUpscaler::FPassInputs& PassInputs);
 
 const ITemporalUpscaler* GTemporalUpscaler = nullptr;
 
@@ -901,37 +899,31 @@ public:
 		return TEXT("FDefaultTemporalUpscaler");
 	}
 
-	virtual void AddPasses(
+	virtual FOutputs AddPasses(
 		FRDGBuilder& GraphBuilder,
 		const FViewInfo& View,
-		const FPassInputs& PassInputs,
-		FRDGTextureRef* OutSceneColorTexture,
-		FIntRect* OutSceneColorViewRect,
-		FRDGTextureRef* OutSceneColorHalfResTexture,
-		FIntRect* OutSceneColorHalfResViewRect) const final
+		const FPassInputs& PassInputs) const final
 	{
 		if (ITemporalUpscaler::GetMainTAAPassConfig(View) == EMainTAAPassConfig::TSR)
 		{
-			*OutSceneColorHalfResTexture = nullptr;
-			//*OutSceneColorHalfResViewRect; // TODO.
-
 			return AddTemporalSuperResolutionPasses(
 				GraphBuilder,
 				View,
-				PassInputs,
-				OutSceneColorTexture,
-				OutSceneColorViewRect);
+				PassInputs);
 		}
 		else
 		{
-			return AddGen4MainTemporalAAPasses(
+			FOutputs Outputs;
+			AddGen4MainTemporalAAPasses(
 				GraphBuilder,
 				View,
 				PassInputs,
-				OutSceneColorTexture,
-				OutSceneColorViewRect,
-				OutSceneColorHalfResTexture,
-				OutSceneColorHalfResViewRect);
+				&Outputs.FullRes.Texture,
+				&Outputs.FullRes.ViewRect,
+				&Outputs.HalfRes.Texture,
+				&Outputs.HalfRes.ViewRect);
+
+			return Outputs;
 		}
 	}
 
diff --git a/Engine/Source/Runtime/Renderer/Private/PostProcess/TemporalAA.h b/Engine/Source/Runtime/Renderer/Private/PostProcess/TemporalAA.h
--- a/Engine/Source/Runtime/Renderer/Private/PostProcess/TemporalAA.h
+++ b/Engine/Source/Runtime/Renderer/Private/PostProcess/TemporalAA.h
@@ -3,6 +3,7 @@
 #pragma once
 
 #include "ScreenPass.h"
+#include "PostProcessMotionBlur.h"
 
 struct FTemporalAAHistory;
 class FSeparateTranslucencyTextures;
@@ -179,6 +180,7 @@ public:
 	{
 		bool bAllowDownsampleSceneColor = false;
 		bool bGenerateOutputMip1 = false;
+		bool bGenerateVelocityFlattenTextures = false;
 		EPixelFormat DownsampleOverrideFormat;
 		FRDGTextureRef SceneColorTexture = nullptr;
 		FRDGTextureRef SceneDepthTexture = nullptr;
@@ -186,6 +188,13 @@ public:
 		const FSeparateTranslucencyTextures* SeparateTranslucencyTextures = nullptr;
 	};
 
+	struct FOutputs
+	{
+		FScreenPassTexture FullRes;
+		FScreenPassTexture HalfRes;
+		FVelocityFlattenTextures VelocityFlattenTextures;
+	};
+
 	virtual ~ITemporalUpscaler() {};
 
 	virtual const TCHAR* GetDebugName() const = 0;
@@ -195,14 +204,10 @@ public:
 	// *  scene color (or null if it was not performed), and the secondary view rect.
 	// */
 
-	virtual void AddPasses(
+	virtual FOutputs AddPasses(
 		FRDGBuilder& GraphBuilder,
 		const FViewInfo& View,
-		const FPassInputs& PassInputs,
-		FRDGTextureRef* OutSceneColorTexture,
-		FIntRect* OutSceneColorViewRect,
-		FRDGTextureRef* OutSceneColorHalfResTexture,
-		FIntRect* OutSceneColorHalfResViewRect) const = 0;
+		const FPassInputs& PassInputs) const = 0;
 
 
 	virtual float GetMinUpsampleResolutionFraction() const = 0;
diff --git a/Engine/Source/Runtime/Renderer/Private/PostProcess/TemporalSuperResolution.cpp b/Engine/Source/Runtime/Renderer/Private/PostProcess/TemporalSuperResolution.cpp
--- a/Engine/Source/Runtime/Renderer/Private/PostProcess/TemporalSuperResolution.cpp
+++ b/Engine/Source/Runtime/Renderer/Private/PostProcess/TemporalSuperResolution.cpp
@@ -220,8 +220,12 @@ class FTSRDilateVelocityCS : public FTSRShader
 	DECLARE_GLOBAL_SHADER(FTSRDilateVelocityCS);
 	SHADER_USE_PARAMETER_STRUCT(FTSRDilateVelocityCS, FTSRShader);
 
+	class FMotionBlurDirectionsDim : SHADER_PERMUTATION_INT("DIM_MOTION_BLUR_DIRECTIONS", 3);
+	using FPermutationDomain = TShaderPermutationDomain<FMotionBlurDirectionsDim>;
+
 	BEGIN_SHADER_PARAMETER_STRUCT(FParameters, )
 		SHADER_PARAMETER_STRUCT_INCLUDE(FTSRCommonParameters, CommonParameters)
+		SHADER_PARAMETER_STRUCT_INCLUDE(FVelocityFlattenParameters, VelocityFlattenParameters)
 
 		SHADER_PARAMETER(FVector2f, PrevOutputBufferUVMin)
 		SHADER_PARAMETER(FVector2f, PrevOutputBufferUVMax)
@@ -235,6 +239,11 @@ class FTSRDilateVelocityCS : public FTSRShader
 		SHADER_PARAMETER_RDG_TEXTURE_UAV(RWTexture2D, PrevUseCountOutput)
 		SHADER_PARAMETER_RDG_TEXTURE_UAV(RWTexture2D, PrevClosestDepthOutput)
 		SHADER_PARAMETER_RDG_TEXTURE_UAV(RWTexture2D, ParallaxFactorOutput)
+
+		SHADER_PARAMETER_RDG_TEXTURE_UAV(RWTexture2D, VelocityFlattenOutput)
+		SHADER_PARAMETER_RDG_TEXTURE_UAV(RWTexture2D, VelocityTileOutput0)
+		SHADER_PARAMETER_RDG_TEXTURE_UAV(RWTexture2D, VelocityTileOutput1)
+
 		SHADER_PARAMETER_RDG_TEXTURE_UAV(RWTexture2D, DebugOutput)
 	END_SHADER_PARAMETER_STRUCT()
 }; // class FTSRDilateVelocityCS
@@ -586,12 +595,10 @@ bool ComposeSeparateTranslucencyInTSR(const FViewInfo& View)
 	return CVarTSRTranslucencySeparateTemporalAccumulation.GetValueOnRenderThread() != 0;
 }
 
-void AddTemporalSuperResolutionPasses(
+ITemporalUpscaler::FOutputs AddTemporalSuperResolutionPasses(
 	FRDGBuilder& GraphBuilder,
 	const FViewInfo& View,
-	const ITemporalUpscaler::FPassInputs& PassInputs,
-	FRDGTextureRef* OutSceneColorTexture,
-	FIntRect* OutSceneColorViewRect)
+	const ITemporalUpscaler::FPassInputs& PassInputs)
 {
 	const FTSRHistory& InputHistory = View.PrevViewInfo.TSRHistory;
 
@@ -762,6 +769,7 @@ void AddTemporalSuperResolutionPasses(
 	FRDGTextureRef DilatedVelocityTexture;
 	FRDGTextureRef ClosestDepthTexture;
 	FRDGTextureRef ParallaxFactorTexture;
+	FVelocityFlattenTextures VelocityFlattenTextures;
 	{
 		{
 			FRDGTextureDesc Desc = FRDGTextureDesc::Create2D(
@@ -779,6 +787,8 @@ void AddTemporalSuperResolutionPasses(
 			ParallaxFactorTexture = GraphBuilder.CreateTexture(Desc, TEXT("TSR.ParallaxFactor"));
 		}
 
+		int32 TileSize = 8;
+		FTSRDilateVelocityCS::FPermutationDomain PermutationVector;
 		FTSRDilateVelocityCS::FParameters* PassParameters = GraphBuilder.AllocParameters<FTSRDilateVelocityCS::FParameters>();
 		PassParameters->CommonParameters = CommonParameters;
 		PassParameters->PrevOutputBufferUVMin = CommonParameters.InputInfo.UVViewportBilinearMin - CommonParameters.InputInfo.ExtentInverse;
@@ -793,20 +803,63 @@ void AddTemporalSuperResolutionPasses(
 		}
 		PassParameters->SceneDepthTexture = PassInputs.SceneDepthTexture;
 		PassParameters->SceneVelocityTexture = PassInputs.SceneVelocityTexture;
+
 		PassParameters->DilatedVelocityOutput = GraphBuilder.CreateUAV(DilatedVelocityTexture);
 		PassParameters->ClosestDepthOutput = GraphBuilder.CreateUAV(ClosestDepthTexture);
 		PassParameters->PrevUseCountOutput = GraphBuilder.CreateUAV(PrevUseCountTexture);
 		PassParameters->PrevClosestDepthOutput = GraphBuilder.CreateUAV(PrevClosestDepthTexture);
 		PassParameters->ParallaxFactorOutput = GraphBuilder.CreateUAV(ParallaxFactorTexture);
+
+		// Setup up the motion blur's velocity flatten pass.
+		if (PassInputs.bGenerateVelocityFlattenTextures)
+		{
+			const int32 MotionBlurDirections = GetMotionBlurDirections();
+			PermutationVector.Set<FTSRDilateVelocityCS::FMotionBlurDirectionsDim>(MotionBlurDirections);
+			TileSize = FVelocityFlattenTextures::kTileSize;
+
+			{
+				FRDGTextureDesc Desc = FRDGTextureDesc::Create2D(
+					InputExtent,
+					PF_FloatR11G11B10,
+					FClearValueBinding::None,
+					GFastVRamConfig.VelocityFlat | TexCreate_ShaderResource | TexCreate_UAV);
+
+				VelocityFlattenTextures.VelocityFlatten.Texture = GraphBuilder.CreateTexture(Desc, TEXT("MotionBlur.VelocityTile"));
+				VelocityFlattenTextures.VelocityFlatten.ViewRect = InputRect;
+			}
+
+			{
+				FRDGTextureDesc Desc = FRDGTextureDesc::Create2D(
+					FIntPoint::DivideAndRoundUp(InputRect.Size(), FVelocityFlattenTextures::kTileSize),
+					PF_FloatRGBA,
+					FClearValueBinding::None,
+					GFastVRamConfig.MotionBlur | TexCreate_ShaderResource | TexCreate_UAV);
+
+				VelocityFlattenTextures.VelocityTile[0].Texture = GraphBuilder.CreateTexture(Desc, TEXT("MotionBlur.VelocityTile"));
+				VelocityFlattenTextures.VelocityTile[0].ViewRect = FIntRect(FIntPoint::ZeroValue, Desc.Extent);
+
+				Desc.Format = PF_G16R16F;
+				VelocityFlattenTextures.VelocityTile[1].Texture = GraphBuilder.CreateTexture(Desc, TEXT("MotionBlur.VelocityTile"));
+				VelocityFlattenTextures.VelocityTile[1].ViewRect = FIntRect(FIntPoint::ZeroValue, Desc.Extent);
+			}
+
+			PassParameters->VelocityFlattenParameters = GetVelocityFlattenParameters(View);
+			PassParameters->VelocityFlattenOutput = GraphBuilder.CreateUAV(VelocityFlattenTextures.VelocityFlatten.Texture);
+			PassParameters->VelocityTileOutput0 = GraphBuilder.CreateUAV(VelocityFlattenTextures.VelocityTile[0].Texture);
+			PassParameters->VelocityTileOutput1 = GraphBuilder.CreateUAV(VelocityFlattenTextures.VelocityTile[1].Texture);
+		}
+
 		PassParameters->DebugOutput = CreateDebugUAV(InputExtent, TEXT("Debug.TSR.DilateVelocity"));
 
-		TShaderMapRef<FTSRDilateVelocityCS> ComputeShader(View.ShaderMap);
+		TShaderMapRef<FTSRDilateVelocityCS> ComputeShader(View.ShaderMap, PermutationVector);
 		FComputeShaderUtils::AddPass(
 			GraphBuilder,
-			RDG_EVENT_NAME("TSR DilateVelocity %dx%d", InputRect.Width(), InputRect.Height()),
+			RDG_EVENT_NAME("TSR DilateVelocity(MotionBlurDirections=%d) %dx%d",
+				PermutationVector.Get<FTSRDilateVelocityCS::FMotionBlurDirectionsDim>(),
+				InputRect.Width(), InputRect.Height()),
 			ComputeShader,
 			PassParameters,
-			FComputeShaderUtils::GetGroupCount(InputRect.Size(), 8));
+			FComputeShaderUtils::GetGroupCount(InputRect.Size(), TileSize));
 	}
 
 	// Setup a dummy history
@@ -1590,6 +1643,10 @@ void AddTemporalSuperResolutionPasses(
 		}
 	}
 
-	*OutSceneColorTexture = SceneColorOutputTexture;
-	*OutSceneColorViewRect = OutputRect;
+
+	ITemporalUpscaler::FOutputs Outputs;
+	Outputs.FullRes.Texture = SceneColorOutputTexture;
+	Outputs.FullRes.ViewRect = OutputRect;
+	Outputs.VelocityFlattenTextures = VelocityFlattenTextures;
+	return Outputs;
 } // AddTemporalSuperResolutionPasses()
