Move some log spam to verbose

diff --git a/Engine/Source/Runtime/Renderer/Private/VT/AllocatedVirtualTexture.cpp b/Engine/Source/Runtime/Renderer/Private/VT/AllocatedVirtualTexture.cpp
--- a/Engine/Source/Runtime/Renderer/Private/VT/AllocatedVirtualTexture.cpp
+++ b/Engine/Source/Runtime/Renderer/Private/VT/AllocatedVirtualTexture.cpp
@@ -274,11 +274,15 @@ bool FAllocatedVirtualTexture::TryMapLockedTiles(FVirtualTextureSystem* InSystem
 		}
 	}
 
-	// Display a warning message (once) if we've failed to map pages for this after a set number of frames
+	// Display a warning message if we've failed to map pages for this after a set number of frames
 	// Generally there should be no delay, but if the system is saturated, it's possible that locked pages may not be loaded immediately
-	if (bHasMissingTiles && InSystem->GetFrame() == FrameAllocated + 30u)
+	if (bHasMissingTiles && InSystem->GetFrame() > FrameAllocated + 30u)
 	{
-		UE_LOG(LogVirtualTexturing, Warning, TEXT("Failed to map lowest resolution mip for AllocatedVT %s"), *Description.Name.ToString());
+		static const auto CVar = IConsoleManager::Get().FindTConsoleVariableDataInt(TEXT("r.VT.Verbose"));
+		if (CVar->GetValueOnRenderThread())
+		{ 
+			UE_LOG(LogVirtualTexturing, Warning, TEXT("Failed to map lowest resolution mip for AllocatedVT %s (%d frames)"), *Description.Name.ToString(), InSystem->GetFrame() - FrameAllocated);
+		}
 	}
 
 	return !bHasMissingTiles;
