Fixes compilation failure in TSR

diff --git a/Engine/Source/Runtime/Renderer/Private/PostProcess/TemporalSuperResolution.cpp b/Engine/Source/Runtime/Renderer/Private/PostProcess/TemporalSuperResolution.cpp
--- a/Engine/Source/Runtime/Renderer/Private/PostProcess/TemporalSuperResolution.cpp
+++ b/Engine/Source/Runtime/Renderer/Private/PostProcess/TemporalSuperResolution.cpp
@@ -114,8 +114,8 @@ BEGIN_SHADER_PARAMETER_STRUCT(FTSRCommonParameters, )
 	SHADER_PARAMETER(FVector2f, InputJitter)
 	SHADER_PARAMETER(int32, bCameraCut)
 	SHADER_PARAMETER(float, MoireInvExposure)
-	SHADER_PARAMETER(FVector2D, ScreenVelocityToInputPixelVelocity)
-	SHADER_PARAMETER(FVector2D, InputPixelVelocityToScreenVelocity)
+	SHADER_PARAMETER(FVector2f, ScreenVelocityToInputPixelVelocity)
+	SHADER_PARAMETER(FVector2f, InputPixelVelocityToScreenVelocity)
 
 	SHADER_PARAMETER_STRUCT_REF(FViewUniformShaderParameters, ViewUniformBuffer)
 END_SHADER_PARAMETER_STRUCT()
