Add Virtual Texture CSV stats.
Will help determine future VT pool size changes.

diff --git a/Engine/Source/Runtime/Renderer/Private/VT/VirtualTexturePhysicalSpace.cpp b/Engine/Source/Runtime/Renderer/Private/VT/VirtualTexturePhysicalSpace.cpp
--- a/Engine/Source/Runtime/Renderer/Private/VT/VirtualTexturePhysicalSpace.cpp
+++ b/Engine/Source/Runtime/Renderer/Private/VT/VirtualTexturePhysicalSpace.cpp
@@ -9,6 +9,8 @@
 #include "VT/VirtualTextureScalability.h"
 #include "VT/VirtualTextureSystem.h"
 
+CSV_DECLARE_CATEGORY_EXTERN(VirtualTexturing);
+
 static TAutoConsoleVariable<float> CVarVTResidencyMaxMipMapBias(
 	TEXT("r.VT.Residency.MaxMipMapBias"),
 	4,
@@ -302,3 +304,19 @@ void FVirtualTexturePhysicalSpace::DrawResidencyGraph(FCanvas* Canvas, FBox2D Ca
 	}
 #endif // UE_BUILD_SHIPPING
 }
+
+void FVirtualTexturePhysicalSpace::UpdateCsvStats() const
+{
+#if CSV_PROFILER
+	FCsvProfiler* Profiler = FCsvProfiler::Get();
+	if (Profiler->IsCapturing_Renderthread())
+	{
+		const FString LockedTitle = FString::Printf(TEXT("%s_%d/LockedPages"), *FormatString, GetID());
+		Profiler->RecordCustomStat(*LockedTitle, CSV_CATEGORY_INDEX(VirtualTexturing), LockedHistory[HistoryIndex], ECsvCustomStatOp::Set);
+		const FString VisbileTitle = FString::Printf(TEXT("%s_%d/VisiblePages"), *FormatString, GetID());
+		Profiler->RecordCustomStat(*VisbileTitle, CSV_CATEGORY_INDEX(VirtualTexturing), VisibleHistory[HistoryIndex], ECsvCustomStatOp::Set);
+		const FString MipBiasTitle = FString::Printf(TEXT("%s_%d/MipBias"), *FormatString, GetID());
+		Profiler->RecordCustomStat(*MipBiasTitle, CSV_CATEGORY_INDEX(VirtualTexturing), MipMapBiasHistory[HistoryIndex], ECsvCustomStatOp::Set);
+	}
+#endif
+}
diff --git a/Engine/Source/Runtime/Renderer/Private/VT/VirtualTexturePhysicalSpace.h b/Engine/Source/Runtime/Renderer/Private/VT/VirtualTexturePhysicalSpace.h
--- a/Engine/Source/Runtime/Renderer/Private/VT/VirtualTexturePhysicalSpace.h
+++ b/Engine/Source/Runtime/Renderer/Private/VT/VirtualTexturePhysicalSpace.h
@@ -116,6 +116,8 @@ public:
 	inline uint32 GetLastFrameOversubscribed() const { return LastFrameOversubscribed; }
 	/** Draw residency graph on screen. */
 	void DrawResidencyGraph(class FCanvas* Canvas, FBox2D CanvasPosition, bool bDrawKey);
+	/** Write residency stats to CSV profiler */
+	void UpdateCsvStats() const;
 
 private:
 	FVTPhysicalSpaceDescription Description;
diff --git a/Engine/Source/Runtime/Renderer/Private/VT/VirtualTextureSystem.cpp b/Engine/Source/Runtime/Renderer/Private/VT/VirtualTextureSystem.cpp
--- a/Engine/Source/Runtime/Renderer/Private/VT/VirtualTextureSystem.cpp
+++ b/Engine/Source/Runtime/Renderer/Private/VT/VirtualTextureSystem.cpp
@@ -26,6 +26,8 @@
 
 #define LOCTEXT_NAMESPACE "VirtualTexture"
 
+CSV_DEFINE_CATEGORY(VirtualTexturing, true);
+
 DECLARE_CYCLE_STAT(TEXT("VirtualTextureSystem Update"), STAT_VirtualTextureSystem_Update, STATGROUP_VirtualTexturing);
 
 DECLARE_CYCLE_STAT(TEXT("Gather Requests"), STAT_ProcessRequests_Gather, STATGROUP_VirtualTexturing);
@@ -147,6 +149,14 @@ static TAutoConsoleVariable<int32> CVarVTResidencyNotify(
 	TEXT("Show on screen notifications for virtual texture physical pool residency"),
 	ECVF_Default
 );
+static TAutoConsoleVariable<int32> CVarVTCsvStats(
+	TEXT("r.VT.CsvStats"),
+	1,
+	TEXT("Send virtual texturing stats to CSV profiler\n")
+	TEXT("0=off, 1=on, 2=verbose"),
+	ECVF_Default
+);
+
 
 static FORCEINLINE uint32 EncodePage(uint32 ID, uint32 vLevel, uint32 vTileX, uint32 vTileY)
 {
@@ -1574,6 +1584,7 @@ void FVirtualTextureSystem::Update(FRDGBuilder& GraphBuilder, ERHIFeatureLevel::
 
 #if !UE_BUILD_SHIPPING
 	UpdateNotifications();
+	UpdateCsvStats();
 #endif
 
 	ReleasePendingSpaces();
@@ -2616,10 +2627,6 @@ void FVirtualTextureSystem::UpdateResidencyNotifications()
 			FString const& FormatString = PhysicalSpace->GetFormatString();
 			const float MipBias = PhysicalSpace->GetResidencyMipMapBias();
 
-#if CSV_PROFILER
-			CSV_CUSTOM_STAT_GLOBAL(VirtualTextureMipBias, MipBias, ECsvCustomStatOp::Set);
-#endif
-
 			OnScreenMessages.Add(
 				FCoreDelegates::EOnScreenMessageSeverity::Warning,
 				FText::Format(LOCTEXT("VTOversubscribed", "VT Pool [{0}] is oversubscribed. Setting MipBias {1}"), FText::FromString(FormatString), FText::AsNumber(MipBias)));
@@ -2627,6 +2634,34 @@ void FVirtualTextureSystem::UpdateResidencyNotifications()
 	}
 }
 
+void FVirtualTextureSystem::UpdateCsvStats()
+{
+#if CSV_PROFILER
+	if (CVarVTCsvStats.GetValueOnRenderThread() == 0)
+	{
+		return;
+	}
+
+	float MaxResidencyMipMapBias = 0.f;
+	for (int32 SpaceIndex = 0u; SpaceIndex < PhysicalSpaces.Num(); ++SpaceIndex)
+	{
+		const FVirtualTexturePhysicalSpace* PhysicalSpace = PhysicalSpaces[SpaceIndex];
+		if (PhysicalSpace != nullptr)
+		{
+			const float ResidencyMipMapBias = PhysicalSpace->GetResidencyMipMapBias();
+			MaxResidencyMipMapBias = FMath::Max(MaxResidencyMipMapBias, ResidencyMipMapBias);
+
+			if (CVarVTCsvStats.GetValueOnRenderThread() == 2)
+			{
+				PhysicalSpace->UpdateCsvStats();
+			}
+		}
+	}
+
+	CSV_CUSTOM_STAT(VirtualTexturing, ResidencyMipBias, MaxResidencyMipMapBias, ECsvCustomStatOp::Set);
+#endif
+}
+
 void FVirtualTextureSystem::DrawResidencyHud(UCanvas* InCanvas, APlayerController* InController)
 {
 	if (CVarVTResidencyShow.GetValueOnGameThread() == 0)
diff --git a/Engine/Source/Runtime/Renderer/Private/VT/VirtualTextureSystem.h b/Engine/Source/Runtime/Renderer/Private/VT/VirtualTextureSystem.h
--- a/Engine/Source/Runtime/Renderer/Private/VT/VirtualTextureSystem.h
+++ b/Engine/Source/Runtime/Renderer/Private/VT/VirtualTextureSystem.h
@@ -181,10 +181,9 @@ private:
 
 	void UpdateNotifications();
 	void UpdateResidencyNotifications();
+	void UpdateCsvStats();
 
 	void DrawResidencyHud(class UCanvas*, class APlayerController*);
 	FDelegateHandle	DrawResidencyHudDelegateHandle;
 #endif
 };
-
-
