Fix compile issues introduced in CL-17955874 and CL-17945620
FPrimitiveUniformShaderParameters::LocalToWorld / PreviousLocalToWorld renamed to LocalToRelativeWorld / PreviousLocalToRelativeWorld.

diff --git a/Engine/Source/Runtime/Renderer/Private/GPUScene.cpp b/Engine/Source/Runtime/Renderer/Private/GPUScene.cpp
--- a/Engine/Source/Runtime/Renderer/Private/GPUScene.cpp
+++ b/Engine/Source/Runtime/Renderer/Private/GPUScene.cpp
@@ -1473,7 +1473,7 @@ struct FUploadDataSourceAdapterDynamicPrimitives
 
 			InstanceUploadInfo.PrimitiveID					= PrimitiveIDStartOffset + ItemIndex;
 			InstanceUploadInfo.PrimitiveToWorld				= PrimitiveData.LocalToRelativeWorld;
-			InstanceUploadInfo.PrevPrimitiveToWorld			= PrimitiveData.PreviousLocalToWorld;
+			InstanceUploadInfo.PrevPrimitiveToWorld			= PrimitiveData.PreviousLocalToRelativeWorld;
 			InstanceUploadInfo.PrimitiveInstances			= TConstArrayView<FPrimitiveInstance>(&InstanceUploadInfo.DummyInstance, 1);
 			InstanceUploadInfo.InstanceDynamicData			= TConstArrayView<FPrimitiveInstanceDynamicData>();
 			InstanceUploadInfo.InstanceLightShadowUVBias	= TConstArrayView<FVector4f>();
diff --git a/Engine/Source/Runtime/Renderer/Private/Renderer.cpp b/Engine/Source/Runtime/Renderer/Private/Renderer.cpp
--- a/Engine/Source/Runtime/Renderer/Private/Renderer.cpp
+++ b/Engine/Source/Runtime/Renderer/Private/Renderer.cpp
@@ -180,8 +180,8 @@ void FRendererModule::DrawTileMesh(FCanvasRenderContext& RenderContext, FMeshPas
 					INVALID_LAST_UPDATE_FRAME,
 					0 /* Custom Data Count */,
 					0.0f /* Random ID */,
-					PrimitiveParams.LocalToWorld,
-					PrimitiveParams.PreviousLocalToWorld
+					PrimitiveParams.LocalToRelativeWorld,
+					PrimitiveParams.PreviousLocalToRelativeWorld
 				);
 
 				// TODO: Payload dummy?
