Fixes compilation failure in VirtualShadowMapArray.cpp

diff --git a/Engine/Source/Runtime/Renderer/Private/VirtualShadowMaps/VirtualShadowMapArray.cpp b/Engine/Source/Runtime/Renderer/Private/VirtualShadowMaps/VirtualShadowMapArray.cpp
--- a/Engine/Source/Runtime/Renderer/Private/VirtualShadowMaps/VirtualShadowMapArray.cpp
+++ b/Engine/Source/Runtime/Renderer/Private/VirtualShadowMaps/VirtualShadowMapArray.cpp
@@ -2293,7 +2293,7 @@ class FVirtualSmBBuildHZBPerPageTopCS : public FVirtualPageManagementShader
 		SHADER_PARAMETER_RDG_BUFFER_SRV(StructuredBuffer< uint >, PhysicalPagesForHZB)
 		SHADER_PARAMETER_SAMPLER(SamplerState, ParentTextureMipSampler)
 		SHADER_PARAMETER_RDG_TEXTURE_SRV(Texture2D, ParentTextureMip)
-		SHADER_PARAMETER(FVector2D, InvHzbInputSize)
+		SHADER_PARAMETER(FVector2f, InvHzbInputSize)
 
 		SHADER_PARAMETER_RDG_TEXTURE_UAV_ARRAY(RWTexture2D<float>, FurthestHZBOutput, [HZBLevelsTop])
 	END_SHADER_PARAMETER_STRUCT()
@@ -2389,7 +2389,7 @@ FRDGTextureRef FVirtualShadowMapArray::BuildHZBFurthest(FRDGBuilder& GraphBuilde
 				PassParameters->FurthestHZBOutput[DestMip] = GraphBuilder.CreateUAV(FRDGTextureUAVDesc(OutFurthestHZBTexture, StartDestMip + DestMip));
 			}
 			FIntPoint SrcSize = FIntPoint::DivideAndRoundUp(HZBSize, 1 << int32(StartDestMip - 1));
-			PassParameters->InvHzbInputSize = FVector2D(1.0f / SrcSize.X, 1.0f / SrcSize.Y);;
+			PassParameters->InvHzbInputSize = FVector2f(1.0f / SrcSize.X, 1.0f / SrcSize.Y);;
 			PassParameters->ParentTextureMip = GraphBuilder.CreateSRV(FRDGTextureSRVDesc::CreateForMipLevel(OutFurthestHZBTexture, StartDestMip - 1));
 			PassParameters->ParentTextureMipSampler = TStaticSamplerState<SF_Point>::GetRHI();
 
