Fix compile issue introduced in CL-18004883
RHISupportsMultithreadedShaderCreation changed to CVar from function, need to remove parameter.

diff --git a/Engine/Source/Runtime/Renderer/Private/PrimitiveSceneInfo.cpp b/Engine/Source/Runtime/Renderer/Private/PrimitiveSceneInfo.cpp
--- a/Engine/Source/Runtime/Renderer/Private/PrimitiveSceneInfo.cpp
+++ b/Engine/Source/Runtime/Renderer/Private/PrimitiveSceneInfo.cpp
@@ -857,7 +857,7 @@ void FPrimitiveSceneInfo::CacheRayTracingPrimitives(FScene* Scene, const TArrayV
 		CSV_SCOPED_TIMING_STAT_EXCLUSIVE(FPrimitiveSceneInfo_CacheRayTracingPrimitives)
 		SCOPED_NAMED_EVENT(FPrimitiveSceneInfo_CacheRayTracingPrimitives, FColor::Emerald);
 
-		checkf(GRHISupportsMultithreadedShaderCreation(GMaxRHIShaderPlatform), TEXT("Raytracing code needs the ability to create shaders from task threads."));
+		checkf(GRHISupportsMultithreadedShaderCreation, TEXT("Raytracing code needs the ability to create shaders from task threads."));
 		
 		FCachedRayTracingMeshCommandStorage& CachedRayTracingMeshCommands = Scene->CachedRayTracingMeshCommands;
 
