Eliminate usage of RWLocks for cached mesh draws, opting instead to defer adding meshcommands to the scene's lists after they are created in parallel.
This is generally cheaper or the same under low-contention scenarios, but saves a large amount of time in high-contention scenarios.

diff --git a/Engine/Source/Runtime/Renderer/Private/MeshPassProcessor.cpp b/Engine/Source/Runtime/Renderer/Private/MeshPassProcessor.cpp
--- a/Engine/Source/Runtime/Renderer/Private/MeshPassProcessor.cpp
+++ b/Engine/Source/Runtime/Renderer/Private/MeshPassProcessor.cpp
@@ -1493,17 +1493,17 @@ FMeshDrawCommandPrimitiveIdInfo FMeshPassProcessor::GetDrawCommandPrimitiveId(
 	return PrimitiveIdInfo;
 }
 
-FCachedPassMeshDrawListContext::FCachedPassMeshDrawListContext(FCachedMeshDrawCommandInfo& InCommandInfo, FRWLock& InCachedMeshDrawCommandLock, FCachedPassMeshDrawList& InCachedDrawLists, FStateBucketMap& InCachedMeshDrawCommandStateBuckets, const FScene& InScene) :
-	CommandInfo(InCommandInfo),
-	CachedMeshDrawCommandLock(InCachedMeshDrawCommandLock),
-	CachedDrawLists(InCachedDrawLists),
-	CachedMeshDrawCommandStateBuckets(InCachedMeshDrawCommandStateBuckets),
-	Scene(InScene)
+FCachedPassMeshDrawListContext::FCachedPassMeshDrawListContext(FScene& InScene)
+	: Scene(InScene)
+	, bUseGPUScene(UseGPUScene(GMaxRHIShaderPlatform, GMaxRHIFeatureLevel))
 {
 }
 
 FMeshDrawCommand& FCachedPassMeshDrawListContext::AddCommand(FMeshDrawCommand& Initializer, uint32 NumElements)
 {
+	checkf(CurrMeshPass < EMeshPass::Num, TEXT("BeginMeshPass() must be called before adding commands to this context"));
+	ensureMsgf(CommandInfo.CommandIndex == -1 && CommandInfo.StateBucketId == -1, TEXT("GetCommandInfoAndReset() wasn't called since the last command was added"));
+
 	if (NumElements == 1)
 	{
 		return Initializer;
@@ -1515,10 +1515,30 @@ FMeshDrawCommand& FCachedPassMeshDrawListContext::AddCommand(FMeshDrawCommand& I
 	}
 }
 
-void FCachedPassMeshDrawListContext::FinalizeCommand(
+void FCachedPassMeshDrawListContext::BeginMeshPass(EMeshPass::Type MeshPass)
+{
+	checkf(CurrMeshPass == EMeshPass::Num, TEXT("BeginMeshPass() was called without a matching EndMeshPass()"));
+	check(MeshPass < EMeshPass::Num);
+	CurrMeshPass = MeshPass;
+}
+
+void FCachedPassMeshDrawListContext::EndMeshPass()
+{
+	checkf(CurrMeshPass < EMeshPass::Num, TEXT("EndMeshPass() was called without matching BeginMeshPass()"));
+	CurrMeshPass = EMeshPass::Num;
+}
+
+FCachedMeshDrawCommandInfo FCachedPassMeshDrawListContext::GetCommandInfoAndReset()
+{
+	FCachedMeshDrawCommandInfo Ret = CommandInfo;
+	CommandInfo.CommandIndex = -1;
+	CommandInfo.StateBucketId = -1;
+	return Ret;
+}
+
+void FCachedPassMeshDrawListContext::FinalizeCommandCommon(
 	const FMeshBatch& MeshBatch, 
 	int32 BatchElementIndex,
-	const FMeshDrawCommandPrimitiveIdInfo& IdInfo,
 	ERasterizerFillMode MeshFillMode,
 	ERasterizerCullMode MeshCullMode,
 	FMeshDrawCommandSortKey SortKey,
@@ -1527,61 +1547,191 @@ void FCachedPassMeshDrawListContext::FinalizeCommand(
 	const FMeshProcessorShaders* ShadersForDebugging,
 	FMeshDrawCommand& MeshDrawCommand)
 {
-	// disabling this by default as it incurs a high cost in perf captures due to sheer volume.  Recommendation is to re-enable locally if you need to profile this particular code.
-	// QUICK_SCOPE_CYCLE_COUNTER(STAT_FinalizeCachedMeshDrawCommand);
-
 	FGraphicsMinimalPipelineStateId PipelineId = FGraphicsMinimalPipelineStateId::GetPersistentId(PipelineState);
 
 	MeshDrawCommand.SetDrawParametersAndFinalize(MeshBatch, BatchElementIndex, PipelineId, ShadersForDebugging);
 
-	if (UseGPUScene(GMaxRHIShaderPlatform, GMaxRHIFeatureLevel))
-		{
-		Experimental::FHashElementId SetId;
-		auto hash = CachedMeshDrawCommandStateBuckets.ComputeHash(MeshDrawCommand);
-		{
-			FRWScopeLock Lock(CachedMeshDrawCommandLock, SLT_ReadOnly);
+	CommandInfo = FCachedMeshDrawCommandInfo(CurrMeshPass);
+	CommandInfo.SortKey = SortKey;
+	CommandInfo.MeshFillMode = MeshFillMode;
+	CommandInfo.MeshCullMode = MeshCullMode;
+	CommandInfo.Flags = Flags;
 
-#if UE_BUILD_DEBUG
-			FMeshDrawCommand MeshDrawCommandDebug = FMeshDrawCommand(MeshDrawCommand);
-			check(MeshDrawCommandDebug.ShaderBindings.GetDynamicInstancingHash() == MeshDrawCommand.ShaderBindings.GetDynamicInstancingHash());
-			check(MeshDrawCommandDebug.GetDynamicInstancingHash() == MeshDrawCommand.GetDynamicInstancingHash());
+#if MESH_DRAW_COMMAND_DEBUG_DATA
+	if (bUseGPUScene)
+	{
+		MeshDrawCommand.ClearDebugPrimitiveSceneProxy(); //When using State Buckets multiple PrimitiveSceneProxies use the same MeshDrawCommand, so The PrimitiveSceneProxy pointer can't be stored.
+	}
 #endif
-			SetId = CachedMeshDrawCommandStateBuckets.FindIdByHash(hash, MeshDrawCommand);
-
-			if (!SetId.IsValid())
-			{
-				Lock.ReleaseReadOnlyLockAndAcquireWriteLock_USE_WITH_CAUTION();
-				SetId = CachedMeshDrawCommandStateBuckets.FindOrAddIdByHash(hash, MeshDrawCommand, FMeshDrawCommandCount());
-			}
+#if DO_GUARD_SLOW
+	if (bUseGPUScene)
+	{
+		FMeshDrawCommand MeshDrawCommandDebug = FMeshDrawCommand(MeshDrawCommand);
+		check(MeshDrawCommandDebug.ShaderBindings.GetDynamicInstancingHash() == MeshDrawCommand.ShaderBindings.GetDynamicInstancingHash());
+		check(MeshDrawCommandDebug.GetDynamicInstancingHash() == MeshDrawCommand.GetDynamicInstancingHash());
+	}
+	if (Scene.GetShadingPath() == EShadingPath::Deferred)
+	{
+		ensureMsgf(MeshDrawCommand.VertexStreams.GetAllocatedSize() == 0, TEXT("Cached Mesh Draw command overflows VertexStreams. VertexStream inline size should be tweaked."));
 
-			FMeshDrawCommandCount& DrawCount = CachedMeshDrawCommandStateBuckets.GetByElementId(SetId).Value;
-			DrawCount.Num++;
+		if (CurrMeshPass == EMeshPass::BasePass || CurrMeshPass == EMeshPass::DepthPass || CurrMeshPass == EMeshPass::CSMShadowDepth || CurrMeshPass == EMeshPass::VSMShadowDepth)
+		{
+			TArray<EShaderFrequency, TInlineAllocator<SF_NumFrequencies>> ShaderFrequencies;
+			MeshDrawCommand.ShaderBindings.GetShaderFrequencies(ShaderFrequencies);
 
-#if MESH_DRAW_COMMAND_DEBUG_DATA
-			if (DrawCount.Num == 1)
+			int32 DataOffset = 0;
+			for (int32 i = 0; i < ShaderFrequencies.Num(); i++)
 			{
-				MeshDrawCommand.ClearDebugPrimitiveSceneProxy(); //When using State Buckets multiple PrimitiveSceneProxies use the same MeshDrawCommand, so The PrimitiveSceneProxy pointer can't be stored.
+				FMeshDrawSingleShaderBindings SingleShaderBindings = MeshDrawCommand.ShaderBindings.GetSingleShaderBindings(ShaderFrequencies[i], DataOffset);
+				if (SingleShaderBindings.GetParameterMapInfo().LooseParameterBuffers.Num() != 0)
+				{
+					bAnyLooseParameterBuffers = true;
+				}
+				ensureMsgf(SingleShaderBindings.GetParameterMapInfo().SRVs.Num() == 0, TEXT("Cached Mesh Draw command uses individual SRVs.  This will break dynamic instancing in performance critical pass.  Use Uniform Buffers instead."));
+				ensureMsgf(SingleShaderBindings.GetParameterMapInfo().TextureSamplers.Num() == 0, TEXT("Cached Mesh Draw command uses individual Texture Samplers.  This will break dynamic instancing in performance critical pass.  Use Uniform Buffers instead."));
 			}
+		}
+	}
 #endif
+}
+
+void FCachedPassMeshDrawListContextImmediate::FinalizeCommand(
+		const FMeshBatch& MeshBatch, 
+		int32 BatchElementIndex,
+		const FMeshDrawCommandPrimitiveIdInfo& IdInfo,
+		ERasterizerFillMode MeshFillMode,
+		ERasterizerCullMode MeshCullMode,
+		FMeshDrawCommandSortKey SortKey,
+		EFVisibleMeshDrawCommandFlags Flags,
+		const FGraphicsMinimalPipelineStateInitializer& PipelineState,
+		const FMeshProcessorShaders* ShadersForDebugging,
+		FMeshDrawCommand& MeshDrawCommand)
+{
+	// disabling this by default as it incurs a high cost in perf captures due to sheer volume.  Recommendation is to re-enable locally if you need to profile this particular code.
+	// QUICK_SCOPE_CYCLE_COUNTER(STAT_FinalizeCachedMeshDrawCommand);
+
+	FinalizeCommandCommon(
+		MeshBatch, 
+		BatchElementIndex,
+		MeshFillMode,
+		MeshCullMode,
+		SortKey,
+		Flags,
+		PipelineState,
+		ShadersForDebugging,
+		MeshDrawCommand
+	);
+
+	if (bUseGPUScene)
+	{
+		Experimental::FHashElementId SetId;
+		FStateBucketMap& BucketMap = Scene.CachedMeshDrawCommandStateBuckets[CurrMeshPass];
+		auto hash = BucketMap.ComputeHash(MeshDrawCommand);
+		{
+			SetId = BucketMap.FindOrAddIdByHash(hash, MeshDrawCommand, FMeshDrawCommandCount());
+			FMeshDrawCommandCount& DrawCount = BucketMap.GetByElementId(SetId).Value;
+			DrawCount.Num++;
 		}
 
-		check(CommandInfo.StateBucketId == -1);
 		CommandInfo.StateBucketId = SetId.GetIndex();
-		check(CommandInfo.CommandIndex == -1);
 	}
 	else
 	{
-		check(CommandInfo.CommandIndex == -1);
-		FRWScopeLock Lock(CachedMeshDrawCommandLock, SLT_Write);
 		// Only one FMeshDrawCommand supported per FStaticMesh in a pass
 		// Allocate at lowest free index so that 'r.DoLazyStaticMeshUpdate' can shrink the TSparseArray more effectively
+		FCachedPassMeshDrawList& CachedDrawLists = Scene.CachedDrawLists[CurrMeshPass];
 		CommandInfo.CommandIndex = CachedDrawLists.MeshDrawCommands.EmplaceAtLowestFreeIndex(CachedDrawLists.LowestFreeIndexSearchStart, MeshDrawCommand);
+	}	
+}
+
+void FCachedPassMeshDrawListContextDeferred::FinalizeCommand(
+		const FMeshBatch& MeshBatch, 
+		int32 BatchElementIndex,
+		const FMeshDrawCommandPrimitiveIdInfo& IdInfo,
+		ERasterizerFillMode MeshFillMode,
+		ERasterizerCullMode MeshCullMode,
+		FMeshDrawCommandSortKey SortKey,
+		EFVisibleMeshDrawCommandFlags Flags,
+		const FGraphicsMinimalPipelineStateInitializer& PipelineState,
+		const FMeshProcessorShaders* ShadersForDebugging,
+		FMeshDrawCommand& MeshDrawCommand)
+{
+	// disabling this by default as it incurs a high cost in perf captures due to sheer volume.  Recommendation is to re-enable locally if you need to profile this particular code.
+	// QUICK_SCOPE_CYCLE_COUNTER(STAT_FinalizeCachedMeshDrawCommand);
+
+	FinalizeCommandCommon(
+		MeshBatch, 
+		BatchElementIndex,
+		MeshFillMode,
+		MeshCullMode,
+		SortKey,
+		Flags,
+		PipelineState,
+		ShadersForDebugging,
+		MeshDrawCommand
+	);
+
+	const int32 Index = DeferredCommands.Add(MeshDrawCommand);
+
+	if (bUseGPUScene)
+	{
+		// Cache the hash here to make the deferred finalize less expensive
+		DeferredCommandHashes.Add(FStateBucketMap::ComputeHash(MeshDrawCommand));
+
+		CommandInfo.StateBucketId = Index;
 	}
+	else
+	{
+		CommandInfo.CommandIndex = Index;
+	}
+}
 
-	CommandInfo.SortKey = SortKey;
-	CommandInfo.MeshFillMode = MeshFillMode;
-	CommandInfo.MeshCullMode = MeshCullMode;
-	CommandInfo.Flags = Flags;
+void FCachedPassMeshDrawListContextDeferred::DeferredFinalizeMeshDrawCommands(const TArrayView<FPrimitiveSceneInfo*>& SceneInfos, int32 Start, int32 End)
+{
+	if (bUseGPUScene)
+	{
+		for (int32 SceneInfoIndex = Start; SceneInfoIndex < End; ++SceneInfoIndex)
+		{
+			FPrimitiveSceneInfo* SceneInfo = SceneInfos[SceneInfoIndex];
+			for (auto& CmdInfo : SceneInfo->StaticMeshCommandInfos)
+			{				
+				check(CmdInfo.MeshPass < EMeshPass::Num);
+				FStateBucketMap& BucketMap = Scene.CachedMeshDrawCommandStateBuckets[CmdInfo.MeshPass];
+
+				check(CmdInfo.StateBucketId >= 0 && CmdInfo.StateBucketId < DeferredCommands.Num());
+				check(CmdInfo.CommandIndex == -1);
+				FMeshDrawCommand& Command = DeferredCommands[CmdInfo.StateBucketId];
+				const Experimental::FHashType CommandHash = DeferredCommandHashes[CmdInfo.StateBucketId];
+
+				Experimental::FHashElementId SetId = BucketMap.FindOrAddIdByHash(CommandHash, MoveTemp(Command), FMeshDrawCommandCount());
+				FMeshDrawCommandCount& DrawCount = BucketMap.GetByElementId(SetId).Value;
+				DrawCount.Num++;
+
+				CmdInfo.StateBucketId = SetId.GetIndex();
+			}
+		}
+	}
+	else
+	{
+		for (int32 SceneInfoIndex = Start; SceneInfoIndex < End; ++SceneInfoIndex)
+		{
+			FPrimitiveSceneInfo* SceneInfo = SceneInfos[SceneInfoIndex];
+			for (auto& CmdInfo : SceneInfo->StaticMeshCommandInfos)
+			{				
+				check(CmdInfo.MeshPass < EMeshPass::Num);
+				FCachedPassMeshDrawList& CachedDrawLists = Scene.CachedDrawLists[CmdInfo.MeshPass];
+
+				check(CmdInfo.CommandIndex >= 0 && CmdInfo.CommandIndex < DeferredCommands.Num());
+				check(CmdInfo.StateBucketId == -1);
+				FMeshDrawCommand& Command = DeferredCommands[CmdInfo.CommandIndex];
+				
+				CmdInfo.CommandIndex = CachedDrawLists.MeshDrawCommands.EmplaceAtLowestFreeIndex(CachedDrawLists.LowestFreeIndexSearchStart, MoveTemp(Command));
+			}
+		}
+	}
+
+	DeferredCommands.Reset();
+	DeferredCommandHashes.Reset();
 }
 
 PassProcessorCreateFunction FPassProcessorManager::JumpTable[(int32)EShadingPath::Num][EMeshPass::Num] = {};
diff --git a/Engine/Source/Runtime/Renderer/Private/PrimitiveSceneInfo.cpp b/Engine/Source/Runtime/Renderer/Private/PrimitiveSceneInfo.cpp
--- a/Engine/Source/Runtime/Renderer/Private/PrimitiveSceneInfo.cpp
+++ b/Engine/Source/Runtime/Renderer/Private/PrimitiveSceneInfo.cpp
@@ -45,16 +45,23 @@ static FAutoConsoleVariableRef CVarCachedRayTracingInstancesLazyUpdate(
 	ECVF_ReadOnly);
 #endif
 
+static int32 GMeshDrawCommandsCacheMultithreaded = 1;
+static FAutoConsoleVariableRef CVarDrawCommandsCacheMultithreaded(
+	TEXT("r.MeshDrawCommands.CacheMultithreaded"),
+	GMeshDrawCommandsCacheMultithreaded,
+	TEXT("Enable multithreading of draw command caching for static meshes. 0=disabled, 1=enabled (default)"),
+	ECVF_RenderThreadSafe);
+
 static int32 GNaniteDrawCommandCacheMultithreaded = 1;
 static FAutoConsoleVariableRef CVarNaniteDrawCommandCacheMultithreaded(
-	TEXT("r.Nanite.DrawCommands.CacheMultithreaded"),
+	TEXT("r.Nanite.MeshDrawCommands.CacheMultithreaded"),
 	GNaniteDrawCommandCacheMultithreaded,
 	TEXT("Enable multithreading of draw command caching for Nanite materials. 0=disabled, 1=enabled (default)"),
-	ECVF_ReadOnly);
+	ECVF_RenderThreadSafe);
 
 static int32 GRayTracingPrimitiveCacheMultithreaded = 1;
 static FAutoConsoleVariableRef CVarRayTracingPrimitiveCacheMultithreaded(
-	TEXT("r.RayTracing.MeshCommands.CacheMultithreaded"),
+	TEXT("r.RayTracing.MeshDrawCommands.CacheMultithreaded"),
 	GRayTracingPrimitiveCacheMultithreaded,
 	TEXT("Enable multithreading of raytracing primitive mesh command caching. 0=disabled, 1=enabled (default)"),
 	ECVF_RenderThreadSafe);
@@ -283,6 +290,7 @@ void FPrimitiveSceneInfo::CacheMeshDrawCommands(FRHICommandListImmediate& RHICmd
 	FMaterialRenderProxy::UpdateDeferredCachedUniformExpressions();
 
 	SCOPED_NAMED_EVENT(FPrimitiveSceneInfo_CacheMeshDrawCommands, FColor::Emerald);
+	CSV_SCOPED_TIMING_STAT_EXCLUSIVE(FPrimitiveSceneInfo_CacheMeshDrawCommands);
 
 	QUICK_SCOPE_CYCLE_COUNTER(STAT_CacheMeshDrawCommands);
 	FMemMark Mark(FMemStack::Get());
@@ -290,7 +298,7 @@ void FPrimitiveSceneInfo::CacheMeshDrawCommands(FRHICommandListImmediate& RHICmd
 	static constexpr int BATCH_SIZE = 64;
 	const int NumBatches = (SceneInfos.Num() + BATCH_SIZE - 1) / BATCH_SIZE;
 
-	auto DoWorkLambda = [Scene, SceneInfos](int32 Index)
+	auto DoWorkLambda = [Scene, SceneInfos](FCachedPassMeshDrawListContext& DrawListContext, int32 Index)
 	{
 		SCOPED_NAMED_EVENT(FPrimitiveSceneInfo_CacheMeshDrawCommand, FColor::Green);
 
@@ -334,16 +342,10 @@ void FPrimitiveSceneInfo::CacheMeshDrawCommands(FRHICommandListImmediate& RHICmd
 
 			if ((FPassProcessorManager::GetPassFlags(ShadingPath, PassType) & EMeshPassFlags::CachedMeshCommands) != EMeshPassFlags::None)
 			{
-				FOptionalTaskTagScope Scope(ETaskTag::EParallelRenderingThread);
-				FCachedMeshDrawCommandInfo CommandInfo(PassType);
-
-				FRWLock& CachedMeshDrawCommandLock = Scene->CachedMeshDrawCommandLock[PassType];
-				FCachedPassMeshDrawList& SceneDrawList = Scene->CachedDrawLists[PassType];
-				FStateBucketMap& CachedMeshDrawCommandStateBuckets = Scene->CachedMeshDrawCommandStateBuckets[PassType];
-				FCachedPassMeshDrawListContext CachedPassMeshDrawListContext(CommandInfo, CachedMeshDrawCommandLock, SceneDrawList, CachedMeshDrawCommandStateBuckets, *Scene);
+				FCachedPassMeshDrawListContext::FMeshPassScope MeshPassScope(DrawListContext, PassType);
 
 				PassProcessorCreateFunction CreateFunction = FPassProcessorManager::GetCreateFunction(ShadingPath, PassType);
-				FMeshPassProcessor* PassMeshProcessor = CreateFunction(Scene, nullptr, &CachedPassMeshDrawListContext);
+				FMeshPassProcessor* PassMeshProcessor = CreateFunction(Scene, nullptr, &DrawListContext);
 
 				if (PassMeshProcessor != nullptr)
 				{
@@ -352,15 +354,15 @@ void FPrimitiveSceneInfo::CacheMeshDrawCommands(FRHICommandListImmediate& RHICmd
 						FPrimitiveSceneInfo* SceneInfo = SceneInfos[MeshAndInfo.InfoIndex];
 						FStaticMeshBatch& Mesh = SceneInfo->StaticMeshes[MeshAndInfo.MeshIndex];
 
-						CommandInfo = FCachedMeshDrawCommandInfo(PassType);
 						FStaticMeshBatchRelevance& MeshRelevance = SceneInfo->StaticMeshRelevances[MeshAndInfo.MeshIndex];
 
 						check(!MeshRelevance.CommandInfosMask.Get(PassType));
 
 						uint64 BatchElementMask = ~0ull;
-						// NOTE: Modifies CommandInfo (through a reference), AddMeshBatch calls FCachedPassMeshDrawListContext::FinalizeCommand
+						// NOTE: AddMeshBatch calls FCachedPassMeshDrawListContext::FinalizeCommand
 						PassMeshProcessor->AddMeshBatch(Mesh, BatchElementMask, SceneInfo->Proxy);
 
+						FCachedMeshDrawCommandInfo CommandInfo = DrawListContext.GetCommandInfoAndReset();
 						if (CommandInfo.CommandIndex != -1 || CommandInfo.StateBucketId != -1)
 						{
 							static_assert(sizeof(MeshRelevance.CommandInfosMask) * 8 >= EMeshPass::Num, "CommandInfosMask is too small to contain all mesh passes.");
@@ -370,36 +372,6 @@ void FPrimitiveSceneInfo::CacheMeshDrawCommands(FRHICommandListImmediate& RHICmd
 							int CommandInfoIndex = MeshAndInfo.MeshIndex * EMeshPass::Num + PassType;
 							check(SceneInfo->StaticMeshCommandInfos[CommandInfoIndex].MeshPass == EMeshPass::Num);
 							SceneInfo->StaticMeshCommandInfos[CommandInfoIndex] = CommandInfo;
-#if DO_GUARD_SLOW
-							if (ShadingPath == EShadingPath::Deferred)
-							{
-								FRWScopeLock Lock(CachedMeshDrawCommandLock, SLT_ReadOnly);
-								const FMeshDrawCommand* MeshDrawCommand = CommandInfo.StateBucketId >= 0
-									? &Scene->CachedMeshDrawCommandStateBuckets[PassType].GetByElementId(CommandInfo.StateBucketId).Key
-									: &SceneDrawList.MeshDrawCommands[CommandInfo.CommandIndex];
-
-								ensureMsgf(MeshDrawCommand->VertexStreams.GetAllocatedSize() == 0, TEXT("Cached Mesh Draw command overflows VertexStreams.  VertexStream inline size should be tweaked."));
-
-								if (PassType == EMeshPass::BasePass || PassType == EMeshPass::DepthPass || PassType == EMeshPass::CSMShadowDepth || PassType == EMeshPass::VSMShadowDepth)
-								{
-									TArray<EShaderFrequency, TInlineAllocator<SF_NumFrequencies>> ShaderFrequencies;
-									MeshDrawCommand->ShaderBindings.GetShaderFrequencies(ShaderFrequencies);
-
-									int32 DataOffset = 0;
-									for (int32 i = 0; i < ShaderFrequencies.Num(); i++)
-									{
-										FMeshDrawSingleShaderBindings SingleShaderBindings = const_cast<FMeshDrawCommand*>(MeshDrawCommand)->ShaderBindings.GetSingleShaderBindings(ShaderFrequencies[i], DataOffset);
-										static int32 LogCount = 0;
-										if (SingleShaderBindings.GetParameterMapInfo().LooseParameterBuffers.Num() != 0 && ((LogCount++ % 1000) == 0))
-										{
-											UE_LOG(LogRenderer, Warning, TEXT("Cached Mesh Draw command uses loose parameters.  This causes overhead and will break dynamic instancing, potentially reducing performance further.  Use Uniform Buffers instead."));
-										}
-										ensureMsgf(SingleShaderBindings.GetParameterMapInfo().SRVs.Num() == 0, TEXT("Cached Mesh Draw command uses individual SRVs.  This will break dynamic instancing in performance critical pass.  Use Uniform Buffers instead."));
-										ensureMsgf(SingleShaderBindings.GetParameterMapInfo().TextureSamplers.Num() == 0, TEXT("Cached Mesh Draw command uses individual Texture Samplers.  This will break dynamic instancing in performance critical pass.  Use Uniform Buffers instead."));
-									}
-								}
-							}
-#endif
 						}
 					}
 
@@ -453,17 +425,59 @@ void FPrimitiveSceneInfo::CacheMeshDrawCommands(FRHICommandListImmediate& RHICmd
 		}
 	};
 
-	if (FApp::ShouldUseThreadingForPerformance())
+	bool bAnyLooseParameterBuffers = false;
+	if (GMeshDrawCommandsCacheMultithreaded && FApp::ShouldUseThreadingForPerformance())
 	{
-		ParallelForTemplate(NumBatches, DoWorkLambda, EParallelForFlags::PumpRenderingThread | EParallelForFlags::Unbalanced);
+		TArray<FCachedPassMeshDrawListContextDeferred, TMemStackAllocator<>> DrawListContexts;
+		DrawListContexts.Reserve(NumBatches);
+		for(int32 ContextIndex = 0; ContextIndex < NumBatches; ++ContextIndex)
+		{
+			DrawListContexts.Emplace(*Scene);
+		}
+
+		ParallelForTemplate(
+			NumBatches, 
+			[&DrawListContexts, &DoWorkLambda](int32 Index)
+			{
+				FTaskTagScope Scope(ETaskTag::EParallelRenderingThread);
+				DoWorkLambda(DrawListContexts[Index], Index);
+			},
+			EParallelForFlags::PumpRenderingThread | EParallelForFlags::Unbalanced
+		);
+
+		if (NumBatches > 0)
+		{
+			SCOPED_NAMED_EVENT(DeferredFinalizeMeshDrawCommands, FColor::Emerald);
+
+			for (int32 Index = 0; Index < NumBatches; ++Index)
+			{
+				FCachedPassMeshDrawListContextDeferred& DrawListContext = DrawListContexts[Index];
+				const int32 Start = Index * BATCH_SIZE;
+				const int32 End = FMath::Min((Index * BATCH_SIZE) + BATCH_SIZE, SceneInfos.Num());
+				DrawListContext.DeferredFinalizeMeshDrawCommands(SceneInfos, Start, End);
+				bAnyLooseParameterBuffers |= DrawListContext.HasAnyLooseParameterBuffers();
+			}
+		}
 	}
 	else
 	{
+		FCachedPassMeshDrawListContextImmediate DrawListContext(*Scene);
 		for (int Idx = 0; Idx < NumBatches; Idx++)
 		{
-			DoWorkLambda(Idx);
+			DoWorkLambda(DrawListContext, Idx);
+		}
+		bAnyLooseParameterBuffers = DrawListContext.HasAnyLooseParameterBuffers();
+	}
+
+#if DO_GUARD_SLOW
+	{
+		static int32 LogCount = 0;
+		if (bAnyLooseParameterBuffers && (LogCount++ % 1000) == 0)
+		{
+			UE_LOG(LogRenderer, Warning, TEXT("One or more Cached Mesh Draw commands use loose parameters. This causes overhead and will break dynamic instancing, potentially reducing performance further. Use Uniform Buffers instead."));
 		}
 	}
+#endif
 
 	if (!FParallelMeshDrawCommandPass::IsOnDemandShaderCreationEnabled())
 	{
@@ -485,8 +499,6 @@ void FPrimitiveSceneInfo::RemoveCachedMeshDrawCommands()
 			FGraphicsMinimalPipelineStateId CachedPipelineId;
 
 			{
-				FRWScopeLock Lock(Scene->CachedMeshDrawCommandLock[PassIndex], SLT_ReadOnly);
-
 				auto& ElementKVP = Scene->CachedMeshDrawCommandStateBuckets[PassIndex].GetByElementId(CachedCommand.StateBucketId);
 				CachedPipelineId = ElementKVP.Key.CachedPipelineId;
 
@@ -495,12 +507,7 @@ void FPrimitiveSceneInfo::RemoveCachedMeshDrawCommands()
 				StateBucketCount.Num--;
 				if (StateBucketCount.Num == 0)
 				{
-					Lock.ReleaseReadOnlyLockAndAcquireWriteLock_USE_WITH_CAUTION();
-
-					if (StateBucketCount.Num == 0)
-					{
-						Scene->CachedMeshDrawCommandStateBuckets[PassIndex].RemoveByElementId(CachedCommand.StateBucketId);
-					}
+					Scene->CachedMeshDrawCommandStateBuckets[PassIndex].RemoveByElementId(CachedCommand.StateBucketId);
 				}
 			}
 
@@ -535,6 +542,7 @@ static void BuildNaniteDrawCommands(FRHICommandListImmediate& RHICmdList, FScene
 void FPrimitiveSceneInfo::CacheNaniteDrawCommands(FRHICommandListImmediate& RHICmdList, FScene* Scene, const TArrayView<FPrimitiveSceneInfo*>& SceneInfos)
 {
 	SCOPED_NAMED_EVENT(FPrimitiveSceneInfo_CacheNaniteDrawCommands, FColor::Emerald);
+	CSV_SCOPED_TIMING_STAT_EXCLUSIVE(FPrimitiveSceneInfo_CacheNaniteDrawCommands);
 	QUICK_SCOPE_CYCLE_COUNTER(STAT_CacheNaniteDrawCommands);
 
 	FMemMark Mark(FMemStack::Get());
diff --git a/Engine/Source/Runtime/Renderer/Private/ScenePrivate.h b/Engine/Source/Runtime/Renderer/Private/ScenePrivate.h
--- a/Engine/Source/Runtime/Renderer/Private/ScenePrivate.h
+++ b/Engine/Source/Runtime/Renderer/Private/ScenePrivate.h
@@ -2689,7 +2689,6 @@ public:
 	FPersistentUniformBuffers UniformBuffers;
 
 	/** Instancing state buckets.  These are stored on the scene as they are precomputed at FPrimitiveSceneInfo::AddToScene time. */
-	FRWLock CachedMeshDrawCommandLock[EMeshPass::Num];
 	FStateBucketMap CachedMeshDrawCommandStateBuckets[EMeshPass::Num];
 	FCachedPassMeshDrawList CachedDrawLists[EMeshPass::Num];
 
diff --git a/Engine/Source/Runtime/Renderer/Private/SceneVisibility.cpp b/Engine/Source/Runtime/Renderer/Private/SceneVisibility.cpp
--- a/Engine/Source/Runtime/Renderer/Private/SceneVisibility.cpp
+++ b/Engine/Source/Runtime/Renderer/Private/SceneVisibility.cpp
@@ -4168,8 +4168,6 @@ void FSceneRenderer::DumpPrimitives(const FViewCommands& ViewCommands)
 	Primitives.Reserve(ViewCommands.MeshCommands[EMeshPass::BasePass].Num() + ViewCommands.DynamicMeshCommandBuildRequests[EMeshPass::BasePass].Num());
 
 	{
-		FRWScopeLock Lock(Scene->CachedMeshDrawCommandLock[EMeshPass::BasePass], SLT_ReadOnly);
-
 		for (const FVisibleMeshDrawCommand& Mesh : ViewCommands.MeshCommands[EMeshPass::BasePass])
 		{
 			int32 PrimitiveId = Mesh.PrimitiveIdInfo.ScenePrimitiveId;
diff --git a/Engine/Source/Runtime/Renderer/Public/MeshPassProcessor.h b/Engine/Source/Runtime/Renderer/Public/MeshPassProcessor.h
--- a/Engine/Source/Runtime/Renderer/Public/MeshPassProcessor.h
+++ b/Engine/Source/Runtime/Renderer/Public/MeshPassProcessor.h
@@ -1592,10 +1592,64 @@ using FStateBucketMap = Experimental::TRobinHoodHashMap<FMeshDrawCommand, FMeshD
 class FCachedPassMeshDrawListContext : public FMeshPassDrawListContext
 {
 public:
-	FCachedPassMeshDrawListContext(FCachedMeshDrawCommandInfo& InCommandInfo, FRWLock& InCachedMeshDrawCommandLock, FCachedPassMeshDrawList& InCachedDrawLists, FStateBucketMap& InCachedMeshDrawCommandStateBuckets, const FScene& InScene);
+	struct FMeshPassScope
+	{
+		FMeshPassScope(const FMeshPassScope&) = delete;
+		FMeshPassScope& operator=(const FMeshPassScope&) = delete;
+
+		inline FMeshPassScope(FCachedPassMeshDrawListContext& InContext, EMeshPass::Type MeshPass)
+			: Context(InContext)
+		{
+			Context.BeginMeshPass(MeshPass);
+		}
+
+		inline ~FMeshPassScope()
+		{
+			Context.EndMeshPass();
+		}
+		
+	private:
+		FCachedPassMeshDrawListContext& Context;
+	};
+
+	FCachedPassMeshDrawListContext(FScene& InScene);
 
 	virtual FMeshDrawCommand& AddCommand(FMeshDrawCommand& Initializer, uint32 NumElements) override final;
 
+	void BeginMeshPass(EMeshPass::Type MeshPass);
+	void EndMeshPass();
+
+	void BeginMesh(int32 SceneInfoIndex, int32 MeshIndex);
+	void EndMesh();
+
+	FCachedMeshDrawCommandInfo GetCommandInfoAndReset();
+	bool HasAnyLooseParameterBuffers() const { return bAnyLooseParameterBuffers; }	
+
+protected:
+	void FinalizeCommandCommon(
+		const FMeshBatch& MeshBatch, 
+		int32 BatchElementIndex,
+		ERasterizerFillMode MeshFillMode,
+		ERasterizerCullMode MeshCullMode,
+		FMeshDrawCommandSortKey SortKey,
+		EFVisibleMeshDrawCommandFlags Flags,
+		const FGraphicsMinimalPipelineStateInitializer& PipelineState,
+		const FMeshProcessorShaders* ShadersForDebugging,
+		FMeshDrawCommand& MeshDrawCommand);
+
+	FScene& Scene;
+	FMeshDrawCommand MeshDrawCommandForStateBucketing;
+	FCachedMeshDrawCommandInfo CommandInfo;
+	EMeshPass::Type CurrMeshPass = EMeshPass::Num;
+	bool bUseGPUScene = false;
+	bool bAnyLooseParameterBuffers = false;
+};
+
+class FCachedPassMeshDrawListContextImmediate : public FCachedPassMeshDrawListContext
+{
+public:
+	FCachedPassMeshDrawListContextImmediate(FScene& InScene) : FCachedPassMeshDrawListContext(InScene) {}
+
 	virtual void FinalizeCommand(
 		const FMeshBatch& MeshBatch, 
 		int32 BatchElementIndex,
@@ -1607,14 +1661,30 @@ public:
 		const FGraphicsMinimalPipelineStateInitializer& PipelineState,
 		const FMeshProcessorShaders* ShadersForDebugging,
 		FMeshDrawCommand& MeshDrawCommand) override final;
+};
+
+class FCachedPassMeshDrawListContextDeferred : public FCachedPassMeshDrawListContext
+{
+public:
+	FCachedPassMeshDrawListContextDeferred(FScene& InScene) : FCachedPassMeshDrawListContext(InScene) {}
+
+	virtual void FinalizeCommand(
+		const FMeshBatch& MeshBatch, 
+		int32 BatchElementIndex,
+		const FMeshDrawCommandPrimitiveIdInfo& IdInfo,
+		ERasterizerFillMode MeshFillMode,
+		ERasterizerCullMode MeshCullMode,
+		FMeshDrawCommandSortKey SortKey,
+		EFVisibleMeshDrawCommandFlags Flags,
+		const FGraphicsMinimalPipelineStateInitializer& PipelineState,
+		const FMeshProcessorShaders* ShadersForDebugging,
+		FMeshDrawCommand& MeshDrawCommand) override final;
+
+	void DeferredFinalizeMeshDrawCommands(const TArrayView<FPrimitiveSceneInfo*>& SceneInfos, int32 Start, int32 End);
 
 private:
-	FMeshDrawCommand MeshDrawCommandForStateBucketing;
-	FCachedMeshDrawCommandInfo& CommandInfo;
-	FRWLock& CachedMeshDrawCommandLock;
-	FCachedPassMeshDrawList& CachedDrawLists;
-	FStateBucketMap& CachedMeshDrawCommandStateBuckets;
-	const FScene& Scene;
+	TArray<FMeshDrawCommand> DeferredCommands;
+	TArray<Experimental::FHashType> DeferredCommandHashes;
 };
 
 template<typename VertexType, typename PixelType, typename GeometryType = FMeshMaterialShader, typename RayHitGroupType = FMeshMaterialShader, typename ComputeType = FMeshMaterialShader>
