When sampling VT use a fallback color held in a uniform if the root page isn't mapped.

diff --git a/Engine/Shaders/Private/VirtualTextureCommon.ush b/Engine/Shaders/Private/VirtualTextureCommon.ush
--- a/Engine/Shaders/Private/VirtualTextureCommon.ush
+++ b/Engine/Shaders/Private/VirtualTextureCommon.ush
@@ -632,21 +632,32 @@ struct VTUniform
 	float pPageSize;
 	float vPageSize;
 	float vPageBorderSize;
-	uint PageTableFormat;
+	bool bPageTableExtraBits;
+	float4 FallbackValue;
 };
 
 /** Unpack the physical sample uniform. */
 VTUniform VTUniform_Unpack(uint4 PackedUniform)
 {
 	VTUniform result;
-	result.pPageSize = asfloat(PackedUniform.w);
+	result.pPageSize = abs(asfloat(PackedUniform.w));
 	result.vPageSize = asfloat(PackedUniform.y);
 	result.vPageBorderSize = asfloat(PackedUniform.z);
-	result.PageTableFormat = PackedUniform.x;
-	//result.PageCoordinateBitCount = PackedUniform.x == 0 ? 8 : 6;
+	result.bPageTableExtraBits = asfloat(PackedUniform.w) > 0;
+	result.FallbackValue.b = float((PackedUniform.x >> 0) & 0xFF) * (1.0f / 255.0f);
+	result.FallbackValue.g = float((PackedUniform.x >> 8) & 0xFF) * (1.0f / 255.0f);
+	result.FallbackValue.r = float((PackedUniform.x >> 16) & 0xFF) * (1.0f / 255.0f);
+	result.FallbackValue.a = float((PackedUniform.x >> 24) & 0xFF) * (1.0f / 255.0f);
 	return result;
 }
 
+/** We use 0 to mark an unmapped page table entry. */
+bool IsValid(VTPageTableResult PageTableResult, uint LayerIndex)
+{
+	const uint PackedPageTableValue = PageTableResult.PageTableValue[LayerIndex / 4u][LayerIndex & 3u];
+	return (PackedPageTableValue >> 4) != 0;
+}
+
 /** Applies proper scaling to dUVdx/dUVdy in PageTableResult. */
 float2 VTComputePhysicalUVs(in out VTPageTableResult PageTableResult, uint LayerIndex, VTUniform Uniform)
 {
@@ -657,8 +668,8 @@ float2 VTComputePhysicalUVs(in out VTPageTableResult PageTableResult, uint Layer
 	const float UVScale = float(4096u >> vLevel) * (1.0f / 4096.0f);
 
 	// This will compile to runtime branch, but should in theory be conditional moves selecting 1 of 2 different bitfield extracting instructions
-	const uint pPageX = (Uniform.PageTableFormat == 0u) ? (PackedPageTableValue >> 4) & 0xff : (PackedPageTableValue >> 4) & 0x3f;
-	const uint pPageY = (Uniform.PageTableFormat == 0u) ? (PackedPageTableValue >> 12) & 0xff : (PackedPageTableValue >> 10) & 0x3f;
+	const uint pPageX = Uniform.bPageTableExtraBits ? (PackedPageTableValue >> 4) & 0xff : (PackedPageTableValue >> 4) & 0x3f;
+	const uint pPageY = Uniform.bPageTableExtraBits ? (PackedPageTableValue >> 12) & 0xff : (PackedPageTableValue >> 10) & 0x3f;
 
 	const float2 vPageFrac = frac(PageTableResult.UV * UVScale);
 	const float2 pUV = float2(pPageX, pPageY) * Uniform.pPageSize + (vPageFrac * Uniform.vPageSize + Uniform.vPageBorderSize);
@@ -675,13 +686,18 @@ MaterialFloat4 TextureVirtualSample(
 	VTPageTableResult PageTableResult, uint LayerIndex,
 	VTUniform Uniform)
 {
-	const float2 pUV = VTComputePhysicalUVs(PageTableResult, LayerIndex, Uniform);
+	if (IsValid(PageTableResult, LayerIndex))
+	{
+		const float2 pUV = VTComputePhysicalUVs(PageTableResult, LayerIndex, Uniform);
 #if VIRTUAL_TEXTURE_ANISOTROPIC_FILTERING
-	return Physical.SampleGrad(PhysicalSampler, pUV, PageTableResult.dUVdx, PageTableResult.dUVdy);
+		return Physical.SampleGrad(PhysicalSampler, pUV, PageTableResult.dUVdx, PageTableResult.dUVdy);
 #else
-	// no need for dUVdx/dUVdy unless we have anistropic filtering enabled
-	return Physical.SampleLevel(PhysicalSampler, pUV, 0.0f);
+		// no need for dUVdx/dUVdy unless we have anistropic filtering enabled
+		return Physical.SampleLevel(PhysicalSampler, pUV, 0.0f);
 #endif
+	}
+
+	return Uniform.FallbackValue;
 }
 
 // SamplePhysicalTexture: Aniso Off
@@ -690,9 +706,14 @@ MaterialFloat4 TextureVirtualSampleLevel(
 	VTPageTableResult PageTableResult, uint LayerIndex,
 	VTUniform Uniform)
 {
-	const float2 pUV = VTComputePhysicalUVs(PageTableResult, LayerIndex, Uniform);
-	// No need to apply dUVdx/dUVdy, don't support anisotropic when sampling a specific level
-	return Physical.SampleLevel(PhysicalSampler, pUV, 0.0f);
+	if (IsValid(PageTableResult, LayerIndex))
+	{
+		const float2 pUV = VTComputePhysicalUVs(PageTableResult, LayerIndex, Uniform);
+		// No need to apply dUVdx/dUVdy, don't support anisotropic when sampling a specific level
+		return Physical.SampleLevel(PhysicalSampler, pUV, 0.0f);
+	}
+
+	return Uniform.FallbackValue;
 }
 
 
diff --git a/Engine/Source/Runtime/Engine/Private/Texture2D.cpp b/Engine/Source/Runtime/Engine/Private/Texture2D.cpp
--- a/Engine/Source/Runtime/Engine/Private/Texture2D.cpp
+++ b/Engine/Source/Runtime/Engine/Private/Texture2D.cpp
@@ -1313,6 +1313,7 @@ void FVirtualTexture2DResource::InitRHI()
 	for (uint32 LayerIndex = 0u; LayerIndex < VTData->GetNumLayers(); ++LayerIndex)
 	{
 		ProducerDesc.LayerFormat[LayerIndex] = VTData->LayerTypes[LayerIndex];
+		ProducerDesc.LayerFallbackColor[LayerIndex] = VTData->LayerFallbackColors[LayerIndex];
 		ProducerDesc.PhysicalGroupIndex[LayerIndex] = bSinglePhysicalSpace ? 0 : LayerIndex;
 	}
 
diff --git a/Engine/Source/Runtime/Engine/Private/TextureDerivedData.cpp b/Engine/Source/Runtime/Engine/Private/TextureDerivedData.cpp
--- a/Engine/Source/Runtime/Engine/Private/TextureDerivedData.cpp
+++ b/Engine/Source/Runtime/Engine/Private/TextureDerivedData.cpp
@@ -65,7 +65,7 @@
 
 // This GUID is mixed into DDC version for virtual textures only, this allows updating DDC version for VT without invalidating DDC for all textures
 // This is useful during development, but once large numbers of VT are present in shipped content, it will have the same problem as TEXTURE_DERIVEDDATA_VER
-#define TEXTURE_VT_DERIVEDDATA_VER	TEXT("A5849A269D64445585DDD909F2D49882")
+#define TEXTURE_VT_DERIVEDDATA_VER	TEXT("4FA538D2450243E380565F65CE8E44B7")
 
 #if ENABLE_COOK_STATS
 namespace TextureCookStats
diff --git a/Engine/Source/Runtime/Engine/Private/VT/VirtualTextureBuiltData.cpp b/Engine/Source/Runtime/Engine/Private/VT/VirtualTextureBuiltData.cpp
--- a/Engine/Source/Runtime/Engine/Private/VT/VirtualTextureBuiltData.cpp
+++ b/Engine/Source/Runtime/Engine/Private/VT/VirtualTextureBuiltData.cpp
@@ -344,6 +344,11 @@ void FVirtualTextureBuiltData::Serialize(FArchive& Ar, UObject* Owner, int32 Fir
 		}
 	}
 	
+ 	for (uint32 Layer = 0; Layer < NumLayers; Layer++)
+ 	{
+ 		Ar << LayerFallbackColors[Layer];
+ 	}
+
 	// Serialize the chunks
 	int32 NumChunksToSerialize = Chunks.Num() - NumChunksToStrip;
 	Ar << NumChunksToSerialize;
diff --git a/Engine/Source/Runtime/Engine/Private/VT/VirtualTextureBuiltData.h b/Engine/Source/Runtime/Engine/Private/VT/VirtualTextureBuiltData.h
--- a/Engine/Source/Runtime/Engine/Private/VT/VirtualTextureBuiltData.h
+++ b/Engine/Source/Runtime/Engine/Private/VT/VirtualTextureBuiltData.h
@@ -136,11 +136,17 @@ struct FVirtualTextureBuiltData
 	uint32 TileBorderSize; // A BorderSize pixel border will be added around all tiles
 
 	/**
-	 * The pixel format output of the data on the i'th layer. The actual data
+	 * The pixel format output of the data for each layer. The actual data
 	 * may still be compressed but will decompress to this pixel format (e.g. zipped DXT5 data).
 	 */
 	TEnumAsByte<EPixelFormat> LayerTypes[VIRTUALTEXTURE_DATA_MAXLAYERS];
 
+	/**
+	 * The fallback color to use for each layer. This color is used whenever we need to sample
+	 * the VT but have not yet streamed the root pages.
+	 */
+	FLinearColor LayerFallbackColors[VIRTUALTEXTURE_DATA_MAXLAYERS];
+
 	/**
 	 * Tile data is packed into separate chunks, typically there is 1 mip level in each chunk for high resolution mips.
 	 * After a certain threshold, all remaining low resolution mips will be packed into one final chunk.
@@ -189,6 +195,7 @@ struct FVirtualTextureBuiltData
 		, TileBorderSize(0u)
 	{
 		FMemory::Memzero(LayerTypes);
+		FMemory::Memzero(LayerFallbackColors);
 	}
 
 	inline bool IsInitialized() const { return TileSize != 0u; }
diff --git a/Engine/Source/Runtime/Engine/Private/VT/VirtualTextureDataBuilder.cpp b/Engine/Source/Runtime/Engine/Private/VT/VirtualTextureDataBuilder.cpp
--- a/Engine/Source/Runtime/Engine/Private/VT/VirtualTextureDataBuilder.cpp
+++ b/Engine/Source/Runtime/Engine/Private/VT/VirtualTextureDataBuilder.cpp
@@ -1235,6 +1235,14 @@ void FVirtualTextureDataBuilder::BuildSourcePixels(const FTextureSourceData& Sou
 		}
 	}
 
+	// Extract fallback color from last mip.
+	for (int32 LayerIndex = 0; LayerIndex < NumLayers; ++LayerIndex)
+	{
+		FImage OnePixelImage(1, 1, 1, ERawImageFormat::RGBA32F);
+		SourceBlocks.Last().MipsPerLayer[LayerIndex].Last().ResizeTo(OnePixelImage, 1, 1, ERawImageFormat::RGBA32F, EGammaSpace::Linear);
+		OutData.LayerFallbackColors[LayerIndex] = OnePixelImage.AsRGBA32F()[0];
+	}
+
 	for (int32 LayerIndex = 0; LayerIndex < NumLayers; ++LayerIndex)
 	{
 		const FTextureBuildSettings& BuildSettingsForLayer = SettingsPerLayer[LayerIndex];
diff --git a/Engine/Source/Runtime/RenderCore/Public/VirtualTexturing.h b/Engine/Source/Runtime/RenderCore/Public/VirtualTexturing.h
--- a/Engine/Source/Runtime/RenderCore/Public/VirtualTexturing.h
+++ b/Engine/Source/Runtime/RenderCore/Public/VirtualTexturing.h
@@ -159,6 +159,7 @@ struct FVTProducerDescription
 	 */
 	uint8 NumTextureLayers = 0u;
 	TEnumAsByte<EPixelFormat> LayerFormat[VIRTUALTEXTURE_SPACE_MAXLAYERS] = { PF_Unknown };
+	FLinearColor LayerFallbackColor[VIRTUALTEXTURE_SPACE_MAXLAYERS] = { FLinearColor::Black };
 	
 	uint8 NumPhysicalGroups = 0u;
 	uint8 PhysicalGroupIndex[VIRTUALTEXTURE_SPACE_MAXLAYERS] = { 0 };
diff --git a/Engine/Source/Runtime/Renderer/Private/VT/AllocatedVirtualTexture.cpp b/Engine/Source/Runtime/Renderer/Private/VT/AllocatedVirtualTexture.cpp
--- a/Engine/Source/Runtime/Renderer/Private/VT/AllocatedVirtualTexture.cpp
+++ b/Engine/Source/Runtime/Renderer/Private/VT/AllocatedVirtualTexture.cpp
@@ -23,6 +23,7 @@ FAllocatedVirtualTexture::FAllocatedVirtualTexture(FVirtualTextureSystem* InSyst
 {
 	check(IsInRenderingThread());
 	FMemory::Memzero(TextureLayers);
+	FMemory::Memzero(FallbackColorPerTextureLayer);
 
 	for (uint32 LayerIndex = 0u; LayerIndex < Description.NumTextureLayers; ++LayerIndex)
 	{
@@ -45,6 +46,11 @@ FAllocatedVirtualTexture::FAllocatedVirtualTexture(FVirtualTextureSystem* InSyst
 			
 			TextureLayers[LayerIndex].UniquePageTableLayerIndex = UniquePhysicalSpaceIndex;
 			TextureLayers[LayerIndex].PhysicalTextureIndex = PageTableLayerLocalIndex;
+
+			if (Producer)
+			{
+				FallbackColorPerTextureLayer[LayerIndex] = Producer->GetDescription().LayerFallbackColor[ProducerLayerIndex].ToFColor(false).DWColor();
+			}
 		}
 	}
 
@@ -268,6 +274,8 @@ bool FAllocatedVirtualTexture::TryMapLockedTiles(FVirtualTextureSystem* InSystem
 					else
 					{
 						bHasMissingTiles = true;
+						// Mark page table entry as invalid if we can't map.
+						PageMap.InvalidateUnmappedRootPage(Space, PageTableLayer.PhysicalSpace, PageTableLayerIndex, MaxLevel, MaxLevel, vAddress, MaxLevel + LocalMipBias);
 					}
 				}
 			}
@@ -481,10 +489,13 @@ void FAllocatedVirtualTexture::GetPackedUniform(FUintVector4* OutUniform, uint32
 		const float RcpPhysicalTextureSize = 1.0f / float(PhysicalTextureSize);
 		const uint32 pPageSize = vPageSize + PageBorderSize * 2u;
 
-		OutUniform->X = GetPageTableFormat() == EVTPageTableFormat::UInt16 ? 1 : 0;
+		OutUniform->X = FallbackColorPerTextureLayer[LayerIndex];
 		OutUniform->Y = BitcastFloatToUInt32((float)vPageSize * RcpPhysicalTextureSize);
 		OutUniform->Z = BitcastFloatToUInt32((float)PageBorderSize * RcpPhysicalTextureSize);
-		OutUniform->W = BitcastFloatToUInt32((float)pPageSize * RcpPhysicalTextureSize);
+		// Pack page table format bool as sign bit on page size.
+		const bool bPageTableExtraBits = GetPageTableFormat() == EVTPageTableFormat::UInt32;
+		const float PackedSignBit = bPageTableExtraBits ? 1.f : -1.f;
+		OutUniform->W = BitcastFloatToUInt32((float)pPageSize * RcpPhysicalTextureSize * PackedSignBit);
 	}
 	else
 	{
diff --git a/Engine/Source/Runtime/Renderer/Private/VT/AllocatedVirtualTexture.h b/Engine/Source/Runtime/Renderer/Private/VT/AllocatedVirtualTexture.h
--- a/Engine/Source/Runtime/Renderer/Private/VT/AllocatedVirtualTexture.h
+++ b/Engine/Source/Runtime/Renderer/Private/VT/AllocatedVirtualTexture.h
@@ -101,4 +101,9 @@ private:
 		uint8 PhysicalTextureIndex;
 	};
 	FTextureLayerDesc TextureLayers[VIRTUALTEXTURE_SPACE_MAXLAYERS];
+
+	// The fallback color to use on each texture layer when unmapped.
+	// This is stored as linear UNorm instead of full FLinearColor since that is what the shader uses.
+	// That means that HDR formats will be clamped.
+	uint32 FallbackColorPerTextureLayer[VIRTUALTEXTURE_SPACE_MAXLAYERS];
 };
diff --git a/Engine/Source/Runtime/Renderer/Private/VT/TexturePageMap.cpp b/Engine/Source/Runtime/Renderer/Private/VT/TexturePageMap.cpp
--- a/Engine/Source/Runtime/Renderer/Private/VT/TexturePageMap.cpp
+++ b/Engine/Source/Runtime/Renderer/Private/VT/TexturePageMap.cpp
@@ -225,6 +225,14 @@ void FTexturePageMap::MapPage(FVirtualTextureSpace* Space, FVirtualTexturePhysic
 	SortedKeysDirty = true;
 }
 
+void FTexturePageMap::InvalidateUnmappedRootPage(FVirtualTextureSpace* Space, FVirtualTexturePhysicalSpace* PhysicalSpace, uint32 PackedProducerHandle, uint8 MaxLevel, uint8 vLogSize, uint32 vAddress, uint8 Local_vLevel)
+{
+	// Page should not be mapped.
+	check(FindPageIndex(vLogSize, vAddress) == ~0u);
+	// Queue clear page table entry.
+	Space->QueueUpdate(LayerIndex, vLogSize, vAddress, Local_vLevel, FIntVector::ZeroValue);
+}
+
 void FTexturePageMap::GetMappedPagesInRange(uint32 vAddress, uint32 Width, uint32 Height, TArray<FMappedTexturePage>& OutMappedPages) const
 {
 	const uint32 vTileX0 = FMath::ReverseMortonCode2(vAddress);
diff --git a/Engine/Source/Runtime/Renderer/Private/VT/TexturePageMap.h b/Engine/Source/Runtime/Renderer/Private/VT/TexturePageMap.h
--- a/Engine/Source/Runtime/Renderer/Private/VT/TexturePageMap.h
+++ b/Engine/Source/Runtime/Renderer/Private/VT/TexturePageMap.h
@@ -106,6 +106,8 @@ public:
 	void		ExpandPageTableUpdatePainters(FVirtualTextureSystem* System, FPageTableUpdate Update, TArray< FPageTableUpdate >* Output);
 	void		ExpandPageTableUpdateMasked(FVirtualTextureSystem* System, FPageTableUpdate Update, TArray< FPageTableUpdate >* Output);
 
+	void		InvalidateUnmappedRootPage(FVirtualTextureSpace* Space, FVirtualTexturePhysicalSpace* PhysicalSpace, uint32 PackedProducerHandle, uint8 MaxLevel, uint8 vLogSize, uint32 vAddress, uint8 Local_vLevel);
+
 private:
 	void		BuildSortedKeys();
 	void        ReleaseUnmappedPages();
diff --git a/Engine/Source/Runtime/Renderer/Private/VT/TexturePagePool.cpp b/Engine/Source/Runtime/Renderer/Private/VT/TexturePagePool.cpp
--- a/Engine/Source/Runtime/Renderer/Private/VT/TexturePagePool.cpp
+++ b/Engine/Source/Runtime/Renderer/Private/VT/TexturePagePool.cpp
@@ -5,11 +5,14 @@
 #include "VirtualTextureSpace.h"
 #include "VirtualTextureSystem.h"
 
+// Reserve pAddress=0 to indicate unmapped page table entry.
+const uint32 FTexturePagePool::NumReservedPages = 1u;
+
 FTexturePagePool::FTexturePagePool()
 	: PageHash(16u * 1024)
 	, NumPages(0u)
 	, NumPagesMapped(0u)
-	, NumPagesAllocated(0u)
+	, NumPagesAllocated(NumReservedPages)
 {
 }
 
@@ -25,7 +28,7 @@ void FTexturePagePool::Initialize(uint32 InNumPages)
 
 	FreeHeap.Resize(InNumPages, InNumPages);
 
-	for (uint32 i = 0; i < InNumPages; i++)
+	for (uint32 i = NumReservedPages; i < InNumPages; i++)
 	{
 		FreeHeap.Add(0, i);
 	}
@@ -101,7 +104,7 @@ void FTexturePagePool::UnmapAllPagesForSpace(FVirtualTextureSystem* System, uint
 
 void FTexturePagePool::EvictPages(FVirtualTextureSystem* System, const FVirtualTextureProducerHandle& ProducerHandle)
 {
-	for (uint32 pAddress = 0u; pAddress < NumPages; ++pAddress)
+	for (uint32 pAddress = NumReservedPages; pAddress < NumPages; ++pAddress)
 	{
 		const FPageEntry& PageEntry = Pages[pAddress];
 		if (PageEntry.PackedProducerHandle == ProducerHandle.PackedValue)
@@ -114,7 +117,7 @@ void FTexturePagePool::EvictPages(FVirtualTextureSystem* System, const FVirtualT
 
 void FTexturePagePool::EvictPages(FVirtualTextureSystem* System, FVirtualTextureProducerHandle const& ProducerHandle, FVTProducerDescription const& Desc, FIntRect const& TextureRegion, uint32 MaxLevel, TArray<union FVirtualTextureLocalTile>& OutLocked)
 {
-	for (uint32 i = 0; i < NumPages; ++i)
+	for (uint32 i = NumReservedPages; i < NumPages; ++i)
 	{
 		if (Pages[i].PackedProducerHandle == ProducerHandle.PackedValue)
 		{
@@ -150,7 +153,7 @@ void FTexturePagePool::GetAllLockedPages(FVirtualTextureSystem* System, TSet<FVi
 {
 	OutPages.Reserve(OutPages.Num() + GetNumLockedPages());
 
-	for (uint32 i = 0; i < NumPages; ++i)
+	for (uint32 i = NumReservedPages; i < NumPages; ++i)
 	{
 		if (!FreeHeap.IsPresent(i))
 		{
@@ -311,7 +314,7 @@ void FTexturePagePool::UpdateUsage(uint32 Frame, uint16 pAddress)
 uint32 FTexturePagePool::GetNumVisiblePages(uint32 Frame) const
 {
 	uint32 Count = 0;
-	for (uint32 i = 0; i < NumPages; ++i)
+	for (uint32 i = NumReservedPages; i < NumPages; ++i)
 	{
 		if (FreeHeap.IsPresent(i))
 		{
@@ -333,7 +336,7 @@ uint32 FTexturePagePool::GetNumVisiblePages(uint32 Frame) const
 
 void FTexturePagePool::CollectProducerCounts(TMap<uint32, uint32>& OutProducerCountMap) const
 {
-	for (uint32 i = 0; i < NumPages; ++i)
+	for (uint32 i = NumReservedPages; i < NumPages; ++i)
 	{
 		const uint32 PackedProducerHandle = Pages[i].PackedProducerHandle;
 		if (PackedProducerHandle != 0u)
@@ -346,6 +349,7 @@ void FTexturePagePool::CollectProducerCounts(TMap<uint32, uint32>& OutProducerCo
 
 void FTexturePagePool::MapPage(FVirtualTextureSpace* Space, FVirtualTexturePhysicalSpace* PhysicalSpace, uint8 PageTableLayerIndex, uint8 MaxLevel, uint8 vLogSize, uint32 vAddress, uint8 Local_vLevel, uint16 pAddress)
 {
+	check(pAddress >= NumReservedPages);
 	check(pAddress < NumPages);
 	const FPageEntry& PageEntry = Pages[pAddress];
 	checkf(PageEntry.PackedProducerHandle != 0u, TEXT("Trying to map pAddress %04x that hasn't been allocated"), pAddress);
@@ -390,7 +394,7 @@ void FTexturePagePool::UnmapAllPages(FVirtualTextureSystem* System, uint16 pAddr
 	FPageEntry& PageEntry = Pages[pAddress];
 	if (PageEntry.PackedProducerHandle != 0u)
 	{
-		check(NumPagesAllocated > 0u);
+		check(NumPagesAllocated > NumReservedPages);
 		--NumPagesAllocated;
 		PageHash.Remove(GetPageHash(PageEntry), pAddress);
 		PageEntry.PackedValue = 0u;
@@ -417,7 +421,7 @@ void FTexturePagePool::RemapPages(FVirtualTextureSystem* System, uint8 SpaceID,
 	const uint32 NewBaseX = FMath::ReverseMortonCode2(NewVirtualAddress);
 	const uint32 NewBaseY = FMath::ReverseMortonCode2(NewVirtualAddress >> 1);
 
-	for (uint32 pAddress = 0u; pAddress < NumPages; ++pAddress)
+	for (uint32 pAddress = NumReservedPages; pAddress < NumPages; ++pAddress)
 	{
 		FPageEntry& PageEntry = Pages[pAddress];
 		if (PageEntry.PackedProducerHandle == OldProducerHandle.PackedValue)
diff --git a/Engine/Source/Runtime/Renderer/Private/VT/TexturePagePool.h b/Engine/Source/Runtime/Renderer/Private/VT/TexturePagePool.h
--- a/Engine/Source/Runtime/Renderer/Private/VT/TexturePagePool.h
+++ b/Engine/Source/Runtime/Renderer/Private/VT/TexturePagePool.h
@@ -241,4 +241,6 @@ private:
 	uint32 NumPages;
 	uint32 NumPagesMapped;
 	uint32 NumPagesAllocated;
+
+	static const uint32 NumReservedPages;
 };
diff --git a/Engine/Source/Runtime/Renderer/Private/VT/VirtualTextureSystem.cpp b/Engine/Source/Runtime/Renderer/Private/VT/VirtualTextureSystem.cpp
--- a/Engine/Source/Runtime/Renderer/Private/VT/VirtualTextureSystem.cpp
+++ b/Engine/Source/Runtime/Renderer/Private/VT/VirtualTextureSystem.cpp
@@ -1102,7 +1102,10 @@ void FVirtualTextureSystem::RequestTiles(const FMaterialRenderProxy* InMaterialR
 
 	for (IAllocatedVirtualTexture* AllocatedVT : InMaterialRenderProxy->UniformExpressionCache[InFeatureLevel].AllocatedVTs)
 	{
-		RequestTilesInternal(AllocatedVT, InScreenSpaceSize, INDEX_NONE);
+		if (AllocatedVT != nullptr)
+		{
+			RequestTilesInternal(AllocatedVT, InScreenSpaceSize, INDEX_NONE);
+		}
 	}
 }
 
@@ -1975,7 +1978,7 @@ void FVirtualTextureSystem::GatherRequestsTask(const FGatherRequestsParameters&
 
 							// either it wasn't mapped, or it's mapped to the current physical address...
 							// otherwise that means that the same local tile is mapped to two separate physical addresses, which is an error
-							ensure(PrevPhysicalSpaceIDAndAddress.Packed == ~0u || PrevPhysicalSpaceIDAndAddress == PhysicalSpaceIDAndAddress);
+							ensure(PrevPhysicalSpaceIDAndAddress.Packed == ~0u || PrevPhysicalSpaceIDAndAddress.Packed == 0u || PrevPhysicalSpaceIDAndAddress == PhysicalSpaceIDAndAddress);
 
 							if (PrevPhysicalSpaceIDAndAddress.Packed == ~0u)
 							{
@@ -2253,10 +2256,14 @@ void FVirtualTextureSystem::SubmitRequests(FRDGBuilder& GraphBuilder, ERHIFeatur
 		TArray<FProducePageDataPrepareTask> PrepareTasks;
 		PrepareTasks.Reserve(RequestList->GetNumLoadRequests());
 
-		// If we render a frame without all locked tiles loaded, may render garbage VT data, as there won't be low mip fallback for unloaded tiles
 #if WITH_EDITOR
+		// Always force loading of locked root pages in editor to be sure that any build steps give best quality.
 		const bool bSyncProduceLockedTiles = true;
 #else
+		// if bSyncProduceLockedTiles is false then we may render with a VT before the root pages are mapped.
+		// When that happens the shader switches to using the single color fallback value instead of sampling the VT.
+		// When the root page is finally mapped we will return to normal high quality VT sampling.
+		// todo[VT]: Make root pages always resident so that we never need to load sync the root pages.
 		const bool bSyncProduceLockedTiles = CVarVTSyncProduceLockedTiles.GetValueOnRenderThread() != 0;
 #endif
 
