For the Scene Infos that we _don't_ update static draw lists for we still need to updated cached ray tracing instances.  This work was previously done when we re-added the prim to the scene and re-created the static draw lists which we are no longer doing for each primitive.

diff --git a/Engine/Source/Runtime/Renderer/Private/RendererScene.cpp b/Engine/Source/Runtime/Renderer/Private/RendererScene.cpp
--- a/Engine/Source/Runtime/Renderer/Private/RendererScene.cpp
+++ b/Engine/Source/Runtime/Renderer/Private/RendererScene.cpp
@@ -4736,6 +4736,11 @@ void FScene::UpdateAllPrimitiveSceneInfos(FRDGBuilder& GraphBuilder, bool bAsync
 		TArray<FPrimitiveSceneInfo*> UpdatedSceneInfosWithoutStaticDrawListUpdate;
 		UpdatedSceneInfosWithoutStaticDrawListUpdate.Reserve(UpdatedInstances.Num());
 
+#if RHI_RAYTRACING
+		TArray<FPrimitiveSceneInfo*> RayTracingPrimitivesToUpdate;
+		RayTracingPrimitivesToUpdate.Reserve(UpdatedInstances.Num());
+#endif
+
 		for (const auto& UpdateInstance : UpdatedInstances)
 		{
 			FPrimitiveSceneProxy* PrimitiveSceneProxy = UpdateInstance.Key;
@@ -4765,6 +4770,9 @@ void FScene::UpdateAllPrimitiveSceneInfos(FRDGBuilder& GraphBuilder, bool bAsync
 			}
 			else
 			{
+#if RHI_RAYTRACING
+				RayTracingPrimitivesToUpdate.Add(PrimitiveSceneInfo);
+#endif
 				UpdatedSceneInfosWithoutStaticDrawListUpdate.Push(PrimitiveSceneInfo);
 				PrimitiveSceneInfo->RemoveFromScene(false);
 			}
@@ -4815,6 +4823,9 @@ void FScene::UpdateAllPrimitiveSceneInfos(FRDGBuilder& GraphBuilder, bool bAsync
 		if (UpdatedSceneInfosWithoutStaticDrawListUpdate.Num() > 0)
 		{
 			FPrimitiveSceneInfo::AddToScene(GraphBuilder.RHICmdList, this, UpdatedSceneInfosWithoutStaticDrawListUpdate, false, true, bAsyncCreateLPIs);
+#if RHI_RAYTRACING
+			FPrimitiveSceneInfo::UpdateCachedRayTracingInstances(this, RayTracingPrimitivesToUpdate);
+#endif
 		}
 	}
 
