Fix broken motion blur on camera cuts when motion vector simulation is disabled. This restores the intention from Jan 2017 (CL 3278757) where velocities were disabled on camera cut, but only if motion sim didn't provide a previous transform (likely broken during MDC or RDG refactors over ttime). When sequencer reuses an object by warping it between shots (vs. respawning new objects every time), the previous transform is nonsensical, and can be 1km away from the new location (causing numerous motion blur problems).

diff --git a/Engine/Source/Runtime/Renderer/Private/PostProcess/PostProcessing.cpp b/Engine/Source/Runtime/Renderer/Private/PostProcess/PostProcessing.cpp
--- a/Engine/Source/Runtime/Renderer/Private/PostProcess/PostProcessing.cpp
+++ b/Engine/Source/Runtime/Renderer/Private/PostProcess/PostProcessing.cpp
@@ -39,6 +39,7 @@
 #include "PostProcess/PostProcessFFTBloom.h"
 #include "PostProcess/PostProcessStreamingAccuracyLegend.h"
 #include "PostProcess/PostProcessSubsurface.h"
+#include "Rendering/MotionVectorSimulation.h"
 #include "ShaderPrint.h"
 #include "ShaderDebug.h"
 #include "HighResScreenshot.h"
@@ -537,8 +538,11 @@ void AddPostProcessingPasses(FRDGBuilder& GraphBuilder, const FViewInfo& View, c
 		const EDownsampleQuality DownsampleQuality = GetDownsampleQuality();
 		const EPixelFormat DownsampleOverrideFormat = PF_FloatRGB;
 
+		// Previous transforms are nonsensical on camera cuts, unless motion vector simulation is enabled (providing FrameN+1 transforms to FrameN+0)
+		const bool bMotionBlurValid = FMotionVectorSimulation::IsEnabled() || (!View.bCameraCut && !View.bPrevTransformsReset);
+
 		// Motion blur gets replaced by the visualization pass.
-		const bool bMotionBlurEnabled = !bVisualizeMotionBlur && IsMotionBlurEnabled(View);
+		const bool bMotionBlurEnabled = !bVisualizeMotionBlur && IsMotionBlurEnabled(View) && bMotionBlurValid;
 
 		// Skip tonemapping for visualizers which overwrite the HDR scene color.
 		const bool bTonemapEnabled = !bVisualizeMotionBlur;
