Optimization for FVirtualTextureProducer::Release
Accelerate lookup of tiles per producer in ForceUnlockAllTiles

diff --git a/Engine/Source/Runtime/Renderer/Private/VT/TexturePageLocks.cpp b/Engine/Source/Runtime/Renderer/Private/VT/TexturePageLocks.cpp
--- a/Engine/Source/Runtime/Renderer/Private/VT/TexturePageLocks.cpp
+++ b/Engine/Source/Runtime/Renderer/Private/VT/TexturePageLocks.cpp
@@ -37,6 +37,7 @@ bool FTexturePageLocks::Lock(const FVirtualTextureLocalTile& Tile)
 	LockedTiles[Index] = Tile;
 	LockCounts[Index] = 1u;
 	TileHash.Add(Hash, Index);
+	ProducerToTileIndex.Add(MurmurFinalize32(Tile.PackedProducerHandle), Index);
 	return true;
 }
 
@@ -54,6 +55,7 @@ bool FTexturePageLocks::Unlock(const FVirtualTextureLocalTile& Tile)
 				// no longer locked
 				FreeIndices.Add(Index);
 				TileHash.Remove(Hash, Index);
+				ProducerToTileIndex.Remove(MurmurFinalize32(Tile.PackedProducerHandle), Index);
 				LockCounts[Index] = 0u;
 				LockedTiles[Index].PackedValue = 0u;
 				return true;
@@ -71,21 +73,28 @@ bool FTexturePageLocks::Unlock(const FVirtualTextureLocalTile& Tile)
 
 void FTexturePageLocks::ForceUnlockAll(const FVirtualTextureProducerHandle& ProducerHandle, TArray<FVirtualTextureLocalTile>& OutUnlockedTiles)
 {
-	for (int32 Index = 0u; Index < LockedTiles.Num(); ++Index)
+	TArray<int32, TInlineAllocator<32>> TileIndicesToUnlock;
+	
+	const uint32 Hash = MurmurFinalize32(ProducerHandle.PackedValue);
+	for (uint32 Index = ProducerToTileIndex.First(Hash); ProducerToTileIndex.IsValid(Index); Index = ProducerToTileIndex.Next(Index))
 	{
-		if (LockCounts[Index] > 0u)
+		if (LockedTiles[Index].PackedProducerHandle == ProducerHandle.PackedValue && LockCounts[Index] > 0u)
 		{
-			const FVirtualTextureLocalTile& Tile = LockedTiles[Index];
-			if (Tile.GetProducerHandle() == ProducerHandle)
-			{
-				OutUnlockedTiles.Add(Tile);
-				FreeIndices.Add(Index);
-				TileHash.Remove(MurmurFinalize64(Tile.PackedValue), Index);
-				LockCounts[Index] = 0u;
-				LockedTiles[Index].PackedValue = 0u;
-			}
+			TileIndicesToUnlock.Add(Index);
 		}
 	}
+
+	for (int32 Index : TileIndicesToUnlock)
+	{
+		const FVirtualTextureLocalTile& Tile = LockedTiles[Index];
+		check(Tile.GetProducerHandle() == ProducerHandle);
+		OutUnlockedTiles.Add(Tile);
+		FreeIndices.Add(Index);
+		TileHash.Remove(MurmurFinalize64(Tile.PackedValue), Index);
+		ProducerToTileIndex.Remove(MurmurFinalize32(Tile.PackedProducerHandle), Index);
+		LockCounts[Index] = 0u;
+		LockedTiles[Index].PackedValue = 0u;
+	}
 }
 
 bool FTexturePageLocks::IsLocked(const FVirtualTextureLocalTile& Tile) const
diff --git a/Engine/Source/Runtime/Renderer/Private/VT/TexturePageLocks.h b/Engine/Source/Runtime/Renderer/Private/VT/TexturePageLocks.h
--- a/Engine/Source/Runtime/Renderer/Private/VT/TexturePageLocks.h
+++ b/Engine/Source/Runtime/Renderer/Private/VT/TexturePageLocks.h
@@ -29,6 +29,7 @@ private:
 	uint32 AcquireIndex();
 
 	FHashTable TileHash;
+	FHashTable ProducerToTileIndex;
 	TArray<FVirtualTextureLocalTile> LockedTiles;
 	TArray<uint16> LockCounts;
 	TArray<uint32> FreeIndices;
diff --git a/Engine/Source/Runtime/Renderer/Private/VT/VirtualTextureProducer.cpp b/Engine/Source/Runtime/Renderer/Private/VT/VirtualTextureProducer.cpp
--- a/Engine/Source/Runtime/Renderer/Private/VT/VirtualTextureProducer.cpp
+++ b/Engine/Source/Runtime/Renderer/Private/VT/VirtualTextureProducer.cpp
@@ -10,6 +10,8 @@ FVirtualTextureProducer::~FVirtualTextureProducer()
 
 void FVirtualTextureProducer::Release(FVirtualTextureSystem* System, const FVirtualTextureProducerHandle& HandleToSelf)
 {
+	TRACE_CPUPROFILER_EVENT_SCOPE(FVirtualTextureProducer::Release);
+
 	if (Description.bPersistentHighestMip)
 	{
 		System->ForceUnlockAllTiles(HandleToSelf, this);
