Optimization for FVirtualTextureProducer::Release
Accelerate lookup of pages per producer in EvictAll()

diff --git a/Engine/Source/Runtime/Renderer/Private/VT/TexturePagePool.cpp b/Engine/Source/Runtime/Renderer/Private/VT/TexturePagePool.cpp
--- a/Engine/Source/Runtime/Renderer/Private/VT/TexturePagePool.cpp
+++ b/Engine/Source/Runtime/Renderer/Private/VT/TexturePagePool.cpp
@@ -104,15 +104,22 @@ void FTexturePagePool::UnmapAllPagesForSpace(FVirtualTextureSystem* System, uint
 
 void FTexturePagePool::EvictPages(FVirtualTextureSystem* System, const FVirtualTextureProducerHandle& ProducerHandle)
 {
-	for (uint32 pAddress = NumReservedPages; pAddress < NumPages; ++pAddress)
+	TArray<uint32, TInlineAllocator<256>> ToEvict;
+	const uint32 Hash = MurmurFinalize32(ProducerHandle.PackedValue);
+	for (uint32 pAddress = ProducerToPageIndex.First(Hash); ProducerToPageIndex.IsValid(pAddress); pAddress = ProducerToPageIndex.Next(pAddress))
 	{
 		const FPageEntry& PageEntry = Pages[pAddress];
 		if (PageEntry.PackedProducerHandle == ProducerHandle.PackedValue)
 		{
-			UnmapAllPages(System, pAddress, false);
-			FreeHeap.Update(0, pAddress);
+			ToEvict.Add(pAddress);
 		}
 	}
+
+	for (uint32 pAddress : ToEvict)
+	{
+		UnmapAllPages(System, pAddress, false);
+		FreeHeap.Update(0, pAddress);
+	}
 }
 
 void FTexturePagePool::EvictPages(FVirtualTextureSystem* System, FVirtualTextureProducerHandle const& ProducerHandle, FVTProducerDescription const& Desc, FIntRect const& TextureRegion, uint32 MaxLevel, TArray<union FVirtualTextureLocalTile>& OutLocked)
@@ -260,6 +267,7 @@ uint32 FTexturePagePool::Alloc(FVirtualTextureSystem* System, uint32 Frame, cons
 	PageEntry.Local_vLevel = Local_vLevel;
 	PageEntry.GroupIndex = GroupIndex;
 	PageHash.Add(GetPageHash(PageEntry), pAddress);
+	ProducerToPageIndex.Add(MurmurFinalize32(ProducerHandle.PackedValue), pAddress);
 
 	if (bLock)
 	{
@@ -397,6 +405,7 @@ void FTexturePagePool::UnmapAllPages(FVirtualTextureSystem* System, uint16 pAddr
 		check(NumPagesAllocated > NumReservedPages);
 		--NumPagesAllocated;
 		PageHash.Remove(GetPageHash(PageEntry), pAddress);
+		ProducerToPageIndex.Remove(MurmurFinalize32(PageEntry.PackedProducerHandle), pAddress);
 		PageEntry.PackedValue = 0u;
 	}
 
@@ -416,6 +425,9 @@ void FTexturePagePool::UnmapAllPages(FVirtualTextureSystem* System, uint16 pAddr
 
 void FTexturePagePool::RemapPages(FVirtualTextureSystem* System, uint8 SpaceID, FVirtualTexturePhysicalSpace* PhysicalSpace, FVirtualTextureProducerHandle const& OldProducerHandle, uint32 OldVirtualAddress, FVirtualTextureProducerHandle const& NewProducerHandle, uint32 NewVirtualAddress, int32 vLevelBias, uint32 Frame)
 {
+	const uint32 OldProducerHash = MurmurFinalize32(OldProducerHandle.PackedValue);
+	const uint32 NewProducerHash = MurmurFinalize32(NewProducerHandle.PackedValue);
+
 	const uint32 OldBaseX = FMath::ReverseMortonCode2(OldVirtualAddress);
 	const uint32 OldBaseY = FMath::ReverseMortonCode2(OldVirtualAddress >> 1);
 	const uint32 NewBaseX = FMath::ReverseMortonCode2(NewVirtualAddress);
@@ -441,6 +453,9 @@ void FTexturePagePool::RemapPages(FVirtualTextureSystem* System, uint8 SpaceID,
 				PageEntry.Local_vLevel += vLevelBias;
 				PageHash.Add(GetPageHash(PageEntry), pAddress);
 
+				ProducerToPageIndex.Remove(OldProducerHash, pAddress);
+				ProducerToPageIndex.Add(NewProducerHash, pAddress);
+
 				if (FreeHeap.IsPresent(pAddress))
 				{
 					FreeHeap.Update((Frame << 4) + PageEntry.Local_vLevel, pAddress);
diff --git a/Engine/Source/Runtime/Renderer/Private/VT/TexturePagePool.h b/Engine/Source/Runtime/Renderer/Private/VT/TexturePagePool.h
--- a/Engine/Source/Runtime/Renderer/Private/VT/TexturePagePool.h
+++ b/Engine/Source/Runtime/Renderer/Private/VT/TexturePagePool.h
@@ -230,6 +230,7 @@ private:
 	FBinaryHeap<uint32, uint16> FreeHeap;
 
 	FHashTable PageHash;
+	FHashTable ProducerToPageIndex;
 	TArray<FPageEntry> Pages;
 
 	// Holds linked lists of mappings for each  physical page in the pool
