Throttle release of allocated virtual textures.

diff --git a/Engine/Source/Runtime/Engine/Private/VT/VirtualTextureScalability.cpp b/Engine/Source/Runtime/Engine/Private/VT/VirtualTextureScalability.cpp
--- a/Engine/Source/Runtime/Engine/Private/VT/VirtualTextureScalability.cpp
+++ b/Engine/Source/Runtime/Engine/Private/VT/VirtualTextureScalability.cpp
@@ -36,6 +36,13 @@ namespace VirtualTextureScalability
 		ECVF_RenderThreadSafe | ECVF_Scalability
 	);
 
+	static TAutoConsoleVariable<int32> CVarVTMaxReleasedPerFrame(
+		TEXT("r.VT.MaxReleasedPerFrame"),
+		0,
+		TEXT("Max number of allocated virtual textures to release per frame"),
+		ECVF_RenderThreadSafe | ECVF_Scalability
+	);
+
 #if WITH_EDITOR
 	static TAutoConsoleVariable<int32> CVarVTMaxContinuousUpdatesPerFrameInEditor(
 		TEXT("r.VT.MaxContinuousUpdatesPerFrameInEditor"),
@@ -231,7 +238,12 @@ namespace VirtualTextureScalability
 	{
 		return CVarMaxPagesProducedPerFrame.GetValueOnAnyThread();
 	}
-	
+
+	int32 GetMaxAllocatedVTReleasedPerFrame()
+	{
+		return CVarVTMaxReleasedPerFrame.GetValueOnAnyThread();
+	}
+
 	int32 GetMaxContinuousUpdatesPerFrame()
 	{
 #if WITH_EDITOR
diff --git a/Engine/Source/Runtime/Engine/Public/VT/VirtualTextureScalability.h b/Engine/Source/Runtime/Engine/Public/VT/VirtualTextureScalability.h
--- a/Engine/Source/Runtime/Engine/Public/VT/VirtualTextureScalability.h
+++ b/Engine/Source/Runtime/Engine/Public/VT/VirtualTextureScalability.h
@@ -12,6 +12,8 @@ namespace VirtualTextureScalability
 	ENGINE_API int32 GetMaxPagesProducedPerFrame();
 	/** Get max update rate of already mapped virtual texture pages. */
 	ENGINE_API int32 GetMaxContinuousUpdatesPerFrame();
+	/** Get max allocated virtual textures to release per frame. */
+	ENGINE_API int32 GetMaxAllocatedVTReleasedPerFrame();
 	/** Get scale factor for virtual texture physical pool sizes. */
 	ENGINE_API float GetPoolSizeScale();
 	/** Get resolution bias for runtime virtual textures. */
diff --git a/Engine/Source/Runtime/Renderer/Private/VT/VirtualTextureSystem.cpp b/Engine/Source/Runtime/Renderer/Private/VT/VirtualTextureSystem.cpp
--- a/Engine/Source/Runtime/Renderer/Private/VT/VirtualTextureSystem.cpp
+++ b/Engine/Source/Runtime/Renderer/Private/VT/VirtualTextureSystem.cpp
@@ -570,9 +570,6 @@ IAllocatedVirtualTexture* FVirtualTextureSystem::AllocateVirtualTexture(const FA
 		return AllocatedVT;
 	}
 
-	// Clear out any AllocatedVTs that are ready to be deleted, before trying a new allocation
-	DestroyPendingVirtualTextures(false);
-
 	uint32 BlockWidthInTiles = 0u;
 	uint32 BlockHeightInTiles = 0u;
 	uint32 WidthInBlocks = 0u;
@@ -663,7 +660,9 @@ void FVirtualTextureSystem::DestroyVirtualTexture(IAllocatedVirtualTexture* Allo
 void FVirtualTextureSystem::DestroyPendingVirtualTextures(bool bForceDestroyAll)
 {
 	check(IsInRenderingThread());
-	
+
+	TRACE_CPUPROFILER_EVENT_SCOPE(FVirtualTextureSystem::DestroyPendingVirtualTextures);
+
 	TArray<IAllocatedVirtualTexture*> AllocatedVTsToDelete;
 	{
 		FScopeLock Lock(&AllocatedVTLock);
@@ -674,6 +673,7 @@ void FVirtualTextureSystem::DestroyPendingVirtualTextures(bool bForceDestroyAll)
 		}
 		else
 		{
+			const int32 MaxDeleteBudget = VirtualTextureScalability::GetMaxAllocatedVTReleasedPerFrame();
 			const uint32 CurrentFrame = Frame;
 			int32 Index = 0;
 			while (Index < PendingDeleteAllocatedVTs.Num())
@@ -682,9 +682,13 @@ void FVirtualTextureSystem::DestroyPendingVirtualTextures(bool bForceDestroyAll)
 				const FAllocatedVTDescription& Desc = AllocatedVT->GetDescription();
 				check(AllocatedVT->NumRefs == 0);
 
-				// If the AllocatedVT is using a private space release it immediately, we don't want to hold references to these private spaces any longer then needed
-				// Otherwise keep deleted VTs around for a few frames, in case they are reused
-				if (Desc.bPrivateSpace || CurrentFrame >= AllocatedVT->FrameDeleted + 60u)
+				// If the AllocatedVT is using a private space release it immediately, we don't want to hold references to these private spaces any longer then needed.
+				const bool bForceDelete = Desc.bPrivateSpace;
+				// Keep deleted VTs around for a few frames, in case they are reused.
+				const bool bCanDeleteForAge = CurrentFrame >= AllocatedVT->FrameDeleted + 60u;
+				// Time slice deletion unless we can make it cheaper.
+				const bool bCanDeleteForBudget = MaxDeleteBudget <= 0 || AllocatedVTsToDelete.Num() < MaxDeleteBudget;
+				if (bForceDelete || (bCanDeleteForAge && bCanDeleteForBudget))
 				{
 					AllocatedVTsToDelete.Add(AllocatedVT);
 					PendingDeleteAllocatedVTs.RemoveAtSwap(Index, 1, false);
