Merge //UE5/Release-5.0 into //UE5/Private-Frosty-To-5.0-Staging

diff --git a/Engine/Plugins/Runtime/HairStrands/Source/HairStrandsCore/Private/GroomManager.cpp b/Engine/Plugins/Runtime/HairStrands/Source/HairStrandsCore/Private/GroomManager.cpp
--- a/Engine/Plugins/Runtime/HairStrands/Source/HairStrandsCore/Private/GroomManager.cpp
+++ b/Engine/Plugins/Runtime/HairStrands/Source/HairStrandsCore/Private/GroomManager.cpp
@@ -575,6 +575,22 @@ static void AddCopyHairStrandsPositionPass(
 #endif
 }
 
+#if RHI_RAYTRACING
+static void AllocateRaytracingResources(FHairGroupInstance* Instance)
+{
+	if (IsHairRayTracingEnabled() && !Instance->Strands.RenRaytracingResource)
+	{
+		check(Instance->Strands.Data);
+
+		// Allocate dynamic raytracing resources (owned by the groom component/instance)
+		FHairResourceName ResourceName(FName(Instance->Debug.GroomAssetName), Instance->Debug.GroupIndex);
+		Instance->Strands.RenRaytracingResource		 = new FHairStrandsRaytracingResource(*Instance->Strands.Data, ResourceName);
+		Instance->Strands.RenRaytracingResourceOwned = true;
+	}
+	Instance->Strands.ViewRayTracingMask |= EHairViewRayTracingMask::PathTracing;
+}
+#endif
+
 void AddHairStreamingRequest(FHairGroupInstance* Instance, int32 InLODIndex)
 {
 	if (Instance && InLODIndex >= 0)
@@ -737,7 +753,6 @@ static void RunHairLODSelection(
 		const float MinLOD = FMath::Max(0, GHairStrandsMinLOD);
 		float LODIndex = Instance->Debug.LODForcedIndex >= 0 ? FMath::Max(Instance->Debug.LODForcedIndex, MinLOD) : -1.0f;
 		float LODViewIndex = -1;
-		bool bCameraCut = false;
 		{
 			const FSphere SphereBound = Instance->ProxyBounds ? Instance->ProxyBounds->GetSphere() : FSphere(0);
 			for (const FSceneView* View : Views)
@@ -748,8 +763,6 @@ static void RunHairLODSelection(
 
 				// Select highest LOD accross all views
 				LODViewIndex = LODViewIndex < 0 ? CurrLODViewIndex : FMath::Min(LODViewIndex, CurrLODViewIndex);
-
-				bCameraCut = bCameraCut || View->bCameraCut;
 			}
 
 			if (LODIndex < 0)
@@ -764,8 +777,11 @@ static void RunHairLODSelection(
 		}
 
 		// Function for selecting, loading, & initializing LOD resources
-		auto SelectValidLOD = [Instance, ShaderPlatform, ShaderMap, PrevLODIndex, LODCount, MeshLODIndex, PrevMeshLODIndex, bHasPathTracingView] (
-			FRDGBuilder& GraphBuilder, float LODIndex) -> bool
+		auto SelectValidLOD = [Instance, ShaderPlatform, ShaderMap, PrevLODIndex, LODCount, MeshLODIndex, PrevMeshLODIndex
+		#if RHI_RAYTRACING
+		, bHasPathTracingView
+		#endif // RHI_RAYTRACING
+		] (FRDGBuilder& GraphBuilder, float LODIndex) -> bool
 		{
 			const int32 IntLODIndex = FMath::Clamp(FMath::FloorToInt(LODIndex), 0, LODCount - 1);
 			const TArray<EHairGeometryType>& LODGeometryTypes = Instance->HairGroupPublicData->GetLODGeometryTypes();
diff --git a/Engine/Source/Runtime/D3D12RHI/Private/D3D12RayTracing.cpp b/Engine/Source/Runtime/D3D12RHI/Private/D3D12RayTracing.cpp
--- a/Engine/Source/Runtime/D3D12RHI/Private/D3D12RayTracing.cpp
+++ b/Engine/Source/Runtime/D3D12RHI/Private/D3D12RayTracing.cpp
@@ -3910,6 +3910,15 @@ void FD3D12RayTracingScene::BuildAccelerationStructure(FD3D12CommandContext& Com
 		TEXT("TLAS build scratch buffer now requires %lld bytes, but only %lld was calculated in the constructor."),
 			PrebuildInfo.ScratchDataSizeInBytes, SizeInfo.BuildScratchSize);
 
+	TRefCountPtr<FD3D12Buffer> AutoScratchBuffer;
+	if (ScratchBuffer == nullptr)
+	{
+		static const FName ScratchBufferName("AutoBuildScratchTLAS");
+		AutoScratchBuffer = CreateRayTracingBuffer(Adapter, GPUIndex, PrebuildInfo.ScratchDataSizeInBytes, ERayTracingBufferType::Scratch, ScratchBufferName);
+		ScratchBuffer = AutoScratchBuffer.GetReference();
+		ScratchBufferOffset = 0;
+	}
+
 	checkf(ScratchBuffer, TEXT("TLAS build requires scratch buffer of at least %lld bytes."), PrebuildInfo.ScratchDataSizeInBytes);
 
 	const D3D12_GPU_VIRTUAL_ADDRESS ScratchAddress = ScratchBuffer->ResourceLocation.GetGPUVirtualAddress() + ScratchBufferOffset;
diff --git a/Engine/Source/Runtime/Engine/Private/WorldPartition/DataLayer/WorldDataLayers.cpp b/Engine/Source/Runtime/Engine/Private/WorldPartition/DataLayer/WorldDataLayers.cpp
--- a/Engine/Source/Runtime/Engine/Private/WorldPartition/DataLayer/WorldDataLayers.cpp
+++ b/Engine/Source/Runtime/Engine/Private/WorldPartition/DataLayer/WorldDataLayers.cpp
@@ -153,7 +153,10 @@ PRAGMA_DISABLE_DEPRECATION_WARNINGS
 				{
 					if (!DataLayersFilterDelegate.Execute(DataLayerLabelName, CurrentState, InState))
 					{
-						UE_LOG(LogWorldPartition, Log, TEXT("Data Layer '%s' was filtered out: %s -> %s"),  *DataLayerLabelName.ToString(), GetDataLayerStateName(CurrentState), GetDataLayerStateName(InState));
+						UE_LOG(LogWorldPartition, Log, TEXT("Data Layer '%s' was filtered out: %s -> %s"),
+							*DataLayerLabelName.ToString(),
+							*StaticEnum<EDataLayerRuntimeState>()->GetDisplayNameTextByValue((int64)CurrentState).ToString(),
+							*StaticEnum<EDataLayerRuntimeState>()->GetDisplayNameTextByValue((int64)InState).ToString());
 						return;
 					}
 				}
@@ -171,7 +174,7 @@ PRAGMA_ENABLE_DEPRECATION_WARNINGS
 			{
 				ActiveDataLayerNames.Add(InDataLayer.Name);
 			}
-			else if (InState == EDataLayerState::Unloaded)
+			else if (InState == EDataLayerRuntimeState::Unloaded)
 			{
 				GLevelStreamingContinuouslyIncrementalGCWhileLevelsPendingPurgeOverride = 1;
 			}
@@ -184,8 +187,12 @@ PRAGMA_ENABLE_DEPRECATION_WARNINGS
 
 #if !NO_LOGGING || CSV_PROFILER
 			const FString DataLayerLabel = DataLayer->GetDataLayerLabel().ToString();
-			UE_LOG(LogWorldPartition, Log, TEXT("Data Layer '%s' state changed: %s -> %s"),  *DataLayerLabel,  GetDataLayerStateName(CurrentState), GetDataLayerStateName(InState));
-			CSV_EVENT_GLOBAL(TEXT("DataLayer-%s-%s"), *DataLayerLabel, GetDataLayerStateName(InState));
+			UE_LOG(LogWorldPartition, Log, TEXT("Data Layer '%s' state changed: %s -> %s"),
+				*DataLayerLabel, 
+				*StaticEnum<EDataLayerRuntimeState>()->GetDisplayNameTextByValue((int64)CurrentState).ToString(),
+				*StaticEnum<EDataLayerRuntimeState>()->GetDisplayNameTextByValue((int64)InState).ToString());
+
+			CSV_EVENT_GLOBAL(TEXT("DataLayer-%s-%s"), *DataLayerLabel, *StaticEnum<EDataLayerRuntimeState>()->GetDisplayNameTextByValue((int64)InState).ToString());
 #endif
 
 			ResolveEffectiveRuntimeState(DataLayer);
diff --git a/Engine/Source/Runtime/Engine/Private/WorldPartition/WorldPartitionRuntimeSpatialHash.cpp b/Engine/Source/Runtime/Engine/Private/WorldPartition/WorldPartitionRuntimeSpatialHash.cpp
--- a/Engine/Source/Runtime/Engine/Private/WorldPartition/WorldPartitionRuntimeSpatialHash.cpp
+++ b/Engine/Source/Runtime/Engine/Private/WorldPartition/WorldPartitionRuntimeSpatialHash.cpp
@@ -211,7 +211,7 @@ void FSpatialHashStreamingGrid::GetCells(const TArray<FWorldPartitionStreamingSo
 						const FSpatialHashStreamingGridLayerCell& LayerCell = GridLevels[Coords.Z].LayerCells[*LayerCellIndexPtr];
 						for (const UWorldPartitionRuntimeCell* Cell : LayerCell.GridCells)
 						{
-						if (!Cell->HasDataLayers() || (DataLayerSubsystem && DataLayerSubsystem->IsAnyDataLayerInEffectiveRuntimeState(Cell->GetDataLayers(), EDataLayerRuntimeState::Activated)))
+							if (!Cell->HasDataLayers() || (DataLayerSubsystem && DataLayerSubsystem->IsAnyDataLayerInEffectiveRuntimeState(Cell->GetDataLayers(), EDataLayerRuntimeState::Activated)))
 							{
 								if (Source.TargetState == EStreamingSourceTargetState::Loaded)
 								{
@@ -223,7 +223,7 @@ void FSpatialHashStreamingGrid::GetCells(const TArray<FWorldPartitionStreamingSo
 									OutActivateCells.AddCell(Cell, Info);
 								}
 							}
-						else if (DataLayerSubsystem && DataLayerSubsystem->IsAnyDataLayerInEffectiveRuntimeState(Cell->GetDataLayers(), EDataLayerRuntimeState::Loaded))
+							else if (DataLayerSubsystem && DataLayerSubsystem->IsAnyDataLayerInEffectiveRuntimeState(Cell->GetDataLayers(), EDataLayerRuntimeState::Loaded))
 							{
 								OutLoadCells.AddCell(Cell, Info);
 							}
@@ -667,7 +667,7 @@ void UWorldPartitionRuntimeSpatialHash::SetDefaultValues()
 	MainGrid.DebugColor = FLinearColor::Gray;
 }
 
-bool UWorldPartitionRuntimeSpatialHash::GenerateStreaming(UWorldPartitionStreamingPolicy* StreamingPolicy, const FActorClusterContext& ActorClusterContext, TArray<FString>* OutPackagesToGenerate)
+bool UWorldPartitionRuntimeSpatialHash::GenerateStreaming(EWorldPartitionStreamingMode Mode, UWorldPartitionStreamingPolicy* StreamingPolicy, TArray<FString>* OutPackagesToGenerate)
 {
 	TRACE_CPUPROFILER_EVENT_SCOPE(UWorldPartitionRuntimeSpatialHash::GenerateStreaming);
 	UWorldPartition* WorldPartition = GetOuterUWorldPartition();
@@ -848,6 +848,14 @@ bool UWorldPartitionRuntimeSpatialHash::CreateStreamingGrid(const FSpatialHashRu
 				{
 					for (const FActorInstance& ActorInstance : GridCellDataChunk.GetActors())
 					{
+						if (ActorInstance.ShouldStripFromStreaming())
+						{
+							const FWorldPartitionActorDescView& ActorDescView = ActorInstance.GetActorDescView();
+							UE_LOG(LogWorldPartition, Verbose, TEXT("Stripping Actor %s (%s) from streaming grid (Container %08x)"),
+								*(ActorDescView.GetActorPath().ToString()), *ActorInstance.Actor.ToString(EGuidFormats::UniqueObjectGuid), ActorInstance.ContainerInstance->ID);
+							continue;
+						}
+
 						if (bIsMainWorldPartition && !IsRunningCookCommandlet())
 						{
 							const FWorldPartitionActorDescView& ActorDescView = ActorInstance.GetActorDescView();
diff --git a/Engine/Source/Runtime/Engine/Private/WorldPartition/WorldPartitionSubsystem.cpp b/Engine/Source/Runtime/Engine/Private/WorldPartition/WorldPartitionSubsystem.cpp
--- a/Engine/Source/Runtime/Engine/Private/WorldPartition/WorldPartitionSubsystem.cpp
+++ b/Engine/Source/Runtime/Engine/Private/WorldPartition/WorldPartitionSubsystem.cpp
@@ -173,6 +173,7 @@ void UWorldPartitionSubsystem::Deinitialize()
 	UWorldPartition* WorldPartition = GetWorldPartition();
 	
 	check(WorldPartition);
+	check(WorldPartition->IsInitialized());
 
 	if (GetWorld()->IsGameWorld())
 	{
diff --git a/Engine/Source/Runtime/Engine/Public/WorldPartition/DataLayer/WorldDataLayers.h b/Engine/Source/Runtime/Engine/Public/WorldPartition/DataLayer/WorldDataLayers.h
--- a/Engine/Source/Runtime/Engine/Public/WorldPartition/DataLayer/WorldDataLayers.h
+++ b/Engine/Source/Runtime/Engine/Public/WorldPartition/DataLayer/WorldDataLayers.h
@@ -144,7 +144,7 @@ private:
 	static int32 DataLayersStateEpoch;
 
 public:
-	DECLARE_DELEGATE_RetVal_ThreeParams(bool, FDataLayersFilterDelegate, FName /*DataLayerName*/, EDataLayerState /*CurrentState*/, EDataLayerState /*TargetState*/);
+	DECLARE_DELEGATE_RetVal_ThreeParams(bool, FDataLayersFilterDelegate, FName /*DataLayerName*/, EDataLayerRuntimeState /*CurrentState*/, EDataLayerRuntimeState /*TargetState*/);
 
 	UE_DEPRECATED(5.00, "do not use, will be replaced by another mechanism for initial release.")
 	FDataLayersFilterDelegate DataLayersFilterDelegate;
diff --git a/Engine/Source/Runtime/MovieSceneTracks/Private/Systems/MovieSceneDataLayerSystem.cpp b/Engine/Source/Runtime/MovieSceneTracks/Private/Systems/MovieSceneDataLayerSystem.cpp
--- a/Engine/Source/Runtime/MovieSceneTracks/Private/Systems/MovieSceneDataLayerSystem.cpp
+++ b/Engine/Source/Runtime/MovieSceneTracks/Private/Systems/MovieSceneDataLayerSystem.cpp
@@ -331,14 +331,24 @@ EDataLayerUpdateFlags FDesiredLayerStates::Apply(FPreAnimatedDataLayerStorage* P
 					if (DesiredStateValue == EDataLayerRuntimeState::Activated && !IsDataLayerReady(DataLayer, EDataLayerRuntimeState::Loaded, false))
 					{
 						Flags |= EDataLayerUpdateFlags::FlushStreamingFull;
-						UE_LOG(LogMovieScene, Warning, TEXT("[UMovieSceneDataLayerSystem] Data layer with name '%s' is causing a full streaming flush (%s)"), *DataLayer->GetDataLayerLabel().ToString(), GetDataLayerStateName(DesiredStateValue));
-						CSV_EVENT_GLOBAL(TEXT("SeqDataLayerFlushFull-%s-%s"), *DataLayer->GetDataLayerLabel().ToString(), GetDataLayerStateName(DesiredStateValue));
+						UE_LOG(LogMovieScene, Warning, TEXT("[UMovieSceneDataLayerSystem] Data layer with name '%s' is causing a full streaming flush (%s)"),
+							*DataLayer->GetDataLayerLabel().ToString(),
+							*StaticEnum<EDataLayerRuntimeState>()->GetDisplayNameTextByValue((int64)DesiredStateValue).ToString());
+							
+						CSV_EVENT_GLOBAL(TEXT("SeqDataLayerFlushFull-%s-%s"),
+							*DataLayer->GetDataLayerLabel().ToString(),
+							*StaticEnum<EDataLayerRuntimeState>()->GetDisplayNameTextByValue((int64)CurrentState).ToString());
 					}
 					else
 					{
 						Flags |= EDataLayerUpdateFlags::FlushStreamingVisibility;
-						UE_LOG(LogMovieScene, Log, TEXT("[UMovieSceneDataLayerSystem] Data layer with name '%s' is causing a visibility streaming flush (%s)"), *DataLayer->GetDataLayerLabel().ToString(), GetDataLayerStateName(DesiredStateValue));
-						CSV_EVENT_GLOBAL(TEXT("SeqDataLayerFlushVis-%s-%s"), *DataLayer->GetDataLayerLabel().ToString(), GetDataLayerStateName(DesiredStateValue));
+						UE_LOG(LogMovieScene, Log, TEXT("[UMovieSceneDataLayerSystem] Data layer with name '%s' is causing a visibility streaming flush (%s)"),
+							*DataLayer->GetDataLayerLabel().ToString(),
+							*StaticEnum<EDataLayerRuntimeState>()->GetDisplayNameTextByValue((int64)DesiredStateValue).ToString());
+							
+						CSV_EVENT_GLOBAL(TEXT("SeqDataLayerFlushVis-%s-%s"),
+							*DataLayer->GetDataLayerLabel().ToString(),
+							*StaticEnum<EDataLayerRuntimeState>()->GetDisplayNameTextByValue((int64)DesiredStateValue).ToString());
 					}
 				}
 			}
@@ -588,7 +598,7 @@ void UMovieSceneDataLayerSystem::UpdateDesiredStates()
 			}
 
 			EDataLayerRuntimeState DesiredState = bPreroll ? Section->GetPrerollState() : Section->GetDesiredState();
-			const bool bRequiresStreamingFlush = (DesiredState == EDataLayerState::Unloaded) ? Section->GetFlushOnUnload() : !bPreroll;
+			const bool bRequiresStreamingFlush = (DesiredState == EDataLayerRuntimeState::Unloaded) ? Section->GetFlushOnUnload() : !bPreroll;
 
 			for (const FActorDataLayer& ActorDataLayer : Section->GetDataLayers())
 			{
diff --git a/Engine/Source/Runtime/Renderer/Private/DeferredShadingRenderer.cpp b/Engine/Source/Runtime/Renderer/Private/DeferredShadingRenderer.cpp
--- a/Engine/Source/Runtime/Renderer/Private/DeferredShadingRenderer.cpp
+++ b/Engine/Source/Runtime/Renderer/Private/DeferredShadingRenderer.cpp
@@ -1553,16 +1553,12 @@ bool FDeferredShadingSceneRenderer::DispatchRayTracingWorldUpdates(FRDGBuilder&
 	{
 		FBuildAccelerationStructurePassParams* PassParams = GraphBuilder.AllocParameters<FBuildAccelerationStructurePassParams>();
 		PassParams->RayTracingSceneScratchBuffer = Scene->RayTracingScene.BuildScratchBuffer;
-		PassParams->DynamicGeometryScratchBuffer = OutDynamicGeometryScratchBuffer;
 		PassParams->RayTracingSceneInstanceBuffer = Scene->RayTracingScene.InstanceBuffer;
+		PassParams->DynamicGeometryScratchBuffer = OutDynamicGeometryScratchBuffer;
 
-		GraphBuilder.AddPass(RDG_EVENT_NAME("BuildAccelerationStructure"), PassParams, ERDGPassFlags::Compute | ERDGPassFlags::NeverCull,
+		GraphBuilder.AddPass(RDG_EVENT_NAME("RayTracingScene"), PassParams, ERDGPassFlags::Compute | ERDGPassFlags::NeverCull,
 			[this, PassParams](FRHICommandListImmediate& RHICmdList)
 			{
-				check(RayTracingDynamicGeometryUpdateEndTransition == nullptr);
-				const FRHITransition* RayTracingDynamicGeometryUpdateBeginTransition = RHICreateTransition(FRHITransitionCreateInfo(ERHIPipeline::Graphics, ERHIPipeline::AsyncCompute));
-				RayTracingDynamicGeometryUpdateEndTransition = RHICreateTransition(FRHITransitionCreateInfo(ERHIPipeline::AsyncCompute, ERHIPipeline::Graphics));
-
 				FRHIAsyncComputeCommandListImmediate& RHIAsyncCmdList = FRHICommandListExecutor::GetImmediateAsyncComputeCommandList();
 
 
@@ -1571,7 +1567,6 @@ bool FDeferredShadingSceneRenderer::DispatchRayTracingWorldUpdates(FRDGBuilder&
 				RHIAsyncCmdList.EndTransition(RayTracingDynamicGeometryUpdateBeginTransition);
 
 				FRHIBuffer* DynamicGeometryScratchBuffer = PassParams->DynamicGeometryScratchBuffer ? PassParams->DynamicGeometryScratchBuffer->GetRHI() : nullptr;
-
 				Scene->GetRayTracingDynamicGeometryCollection()->DispatchUpdates(RHIAsyncCmdList, DynamicGeometryScratchBuffer);
 
 				FRHIRayTracingScene* RayTracingSceneRHI = Scene->RayTracingScene.GetRHIRayTracingSceneChecked();
@@ -1597,6 +1592,7 @@ bool FDeferredShadingSceneRenderer::DispatchRayTracingWorldUpdates(FRDGBuilder&
 				check(RayTracingDynamicGeometryUpdateEndTransition == nullptr)
 				RayTracingDynamicGeometryUpdateEndTransition = RHICreateTransition(FRHITransitionCreateInfo(ERHIPipeline::AsyncCompute, ERHIPipeline::Graphics));
 				RHIAsyncCmdList.BeginTransition(RayTracingDynamicGeometryUpdateEndTransition);
+
 				FRHIAsyncComputeCommandListImmediate::ImmediateDispatch(RHIAsyncCmdList);
 			});
 	}
@@ -1604,17 +1600,16 @@ bool FDeferredShadingSceneRenderer::DispatchRayTracingWorldUpdates(FRDGBuilder&
 	{
 		{
 			RDG_GPU_STAT_SCOPE(GraphBuilder, RayTracingGeometry);
-
+			
 			FBuildAccelerationStructurePassParams* PassParams = GraphBuilder.AllocParameters<FBuildAccelerationStructurePassParams>();
 			PassParams->RayTracingSceneScratchBuffer = nullptr;
 			PassParams->DynamicGeometryScratchBuffer = OutDynamicGeometryScratchBuffer;
-
 			GraphBuilder.AddPass(RDG_EVENT_NAME("RayTracingGeometry"), PassParams, ERDGPassFlags::Compute | ERDGPassFlags::NeverCull,
 				[this, PassParams](FRHICommandListImmediate& RHICmdList)
-				{
-					FRHIBuffer* DynamicGeometryScratchBuffer = PassParams->DynamicGeometryScratchBuffer ? PassParams->DynamicGeometryScratchBuffer->GetRHI() : nullptr;
-					Scene->GetRayTracingDynamicGeometryCollection()->DispatchUpdates(RHICmdList, DynamicGeometryScratchBuffer);
-				});
+			{
+				FRHIBuffer* DynamicGeometryScratchBuffer = PassParams->DynamicGeometryScratchBuffer ? PassParams->DynamicGeometryScratchBuffer->GetRHI() : nullptr;
+				Scene->GetRayTracingDynamicGeometryCollection()->DispatchUpdates(RHICmdList, DynamicGeometryScratchBuffer);
+			});
 		}
 
 		{
@@ -1622,6 +1617,8 @@ bool FDeferredShadingSceneRenderer::DispatchRayTracingWorldUpdates(FRDGBuilder&
 
 			FBuildAccelerationStructurePassParams* PassParams = GraphBuilder.AllocParameters<FBuildAccelerationStructurePassParams>();
 			PassParams->RayTracingSceneScratchBuffer = Scene->RayTracingScene.BuildScratchBuffer;
+			PassParams->RayTracingSceneInstanceBuffer = Scene->RayTracingScene.InstanceBuffer;
+			PassParams->DynamicGeometryScratchBuffer = nullptr;
 
 			GraphBuilder.AddPass(RDG_EVENT_NAME("RayTracingScene"), PassParams, ERDGPassFlags::Compute | ERDGPassFlags::NeverCull,
 				[this, PassParams](FRHICommandListImmediate& RHICmdList)
@@ -3131,6 +3128,11 @@ void FDeferredShadingSceneRenderer::Render(FRDGBuilder& GraphBuilder)
 		ShadingEnergyConservation::Debug(GraphBuilder, View, SceneTextures);
 	}
 
+	for (FViewInfo& View : Views)
+	{
+		ShadingEnergyConservation::Debug(GraphBuilder, View, SceneTextures);
+	}
+
 	if (!bHasRayTracedOverlay && ViewFamily.EngineShowFlags.LightShafts)
 	{
 		SCOPE_CYCLE_COUNTER(STAT_FDeferredShadingSceneRenderer_RenderLightShaftBloom);
diff --git a/Engine/Source/Runtime/Renderer/Private/ScenePrivate.h b/Engine/Source/Runtime/Renderer/Private/ScenePrivate.h
--- a/Engine/Source/Runtime/Renderer/Private/ScenePrivate.h
+++ b/Engine/Source/Runtime/Renderer/Private/ScenePrivate.h
@@ -3038,6 +3038,7 @@ public:
 	virtual void ReleasePrimitive(UPrimitiveComponent* Primitive) override;
 	virtual void UpdateAllPrimitiveSceneInfos(FRDGBuilder& GraphBuilder, bool bAsyncCreateLPIs = false) override;
 	virtual void UpdatePrimitiveTransform(UPrimitiveComponent* Primitive) override;
+	virtual void UpdatePrimitiveOcclusionBoundsSlack(UPrimitiveComponent* Primitive, float NewSlack) override;
 	virtual void UpdatePrimitiveAttachment(UPrimitiveComponent* Primitive) override;
 	virtual void UpdateCustomPrimitiveData(UPrimitiveComponent* Primitive) override;
 	virtual void UpdatePrimitiveDistanceFieldSceneData_GameThread(UPrimitiveComponent* Primitive) override;
