Fix mismatched primitive scene data strides & another Float4 -> Float4f in GPUScene

diff --git a/Engine/Shaders/Private/SceneData.ush b/Engine/Shaders/Private/SceneData.ush
--- a/Engine/Shaders/Private/SceneData.ush
+++ b/Engine/Shaders/Private/SceneData.ush
@@ -164,7 +164,7 @@ RWStructuredBuffer<float4> GPUScenePrimitiveSceneDataRW;
 #endif
 
 // Stride of a single primitive's data in float4's, must match C++
-#define PRIMITIVE_SCENE_DATA_STRIDE 38
+#define PRIMITIVE_SCENE_DATA_STRIDE 39
 
 float4 LoadPrimitivePrimitiveSceneDataElement(uint PrimitiveIndex, uint ItemIndex)
 {
diff --git a/Engine/Source/Runtime/Engine/Public/PrimitiveUniformShaderParameters.h b/Engine/Source/Runtime/Engine/Public/PrimitiveUniformShaderParameters.h
--- a/Engine/Source/Runtime/Engine/Public/PrimitiveUniformShaderParameters.h
+++ b/Engine/Source/Runtime/Engine/Public/PrimitiveUniformShaderParameters.h
@@ -401,7 +401,7 @@ extern ENGINE_API TGlobalResource<FIdentityPrimitiveUniformBuffer> GIdentityPrim
 
 struct FPrimitiveSceneShaderData
 {
-	// Must match usf
+	// Must match PRIMITIVE_SCENE_DATA_STRIDE in SceneData.ush
 	enum { DataStrideInFloat4s = 39 };
 
 	TStaticArray<FVector4f, DataStrideInFloat4s> Data;
diff --git a/Engine/Source/Runtime/Renderer/Private/GPUScene.cpp b/Engine/Source/Runtime/Renderer/Private/GPUScene.cpp
--- a/Engine/Source/Runtime/Renderer/Private/GPUScene.cpp
+++ b/Engine/Source/Runtime/Renderer/Private/GPUScene.cpp
@@ -615,7 +615,7 @@ void FGPUScene::UpdateInternal(FRDGBuilder& GraphBuilder, FScene& Scene)
 						const FPrimitiveSceneShaderData& Item = PrimitiveBufferPtr[Index + IndexOffset];
 						for (int32 DataIndex = 0; DataIndex < FPrimitiveSceneShaderData::DataStrideInFloat4s; ++DataIndex)
 						{
-							check(FMemory::Memcmp(&PrimitiveSceneData.Data[DataIndex], &Item.Data[DataIndex], sizeof(FVector4)) == 0);
+							check(FMemory::Memcmp(&PrimitiveSceneData.Data[DataIndex], &Item.Data[DataIndex], sizeof(FVector4f)) == 0);
 						}
 					}
 				}
@@ -1294,7 +1294,7 @@ void FGPUScene::UploadGeneral(FRHICommandListImmediate& RHICmdList, FScene *Scen
 									#if INSTANCE_COMPRESSED_TRANSFORMS
 										check(PayloadPosition + 1 < InstancePayloadData.Num()); // Sanity check
 										FCompressedTransform CompressedPrevLocalToWorld(PrevLocalToWorld);
-										InstancePayloadData[PayloadPosition + 0] = *(const FVector4*)&CompressedPrevLocalToWorld.Rotation[0];
+										InstancePayloadData[PayloadPosition + 0] = *(const FVector4f*)&CompressedPrevLocalToWorld.Rotation[0];
 										InstancePayloadData[PayloadPosition + 1] = *(const FVector3f*)&CompressedPrevLocalToWorld.Translation;
 										PayloadPosition += 2;
 									#else
