Change FTaskTagScope to FOptionalTaskTagScope to prevent ensures.
In non-shipping / test builds, IsInRenderingThread caches its state from call to call to do validation checking.
This was causing an ensure in some analytics code.

diff --git a/Engine/Source/Runtime/Renderer/Private/PrimitiveSceneInfo.cpp b/Engine/Source/Runtime/Renderer/Private/PrimitiveSceneInfo.cpp
--- a/Engine/Source/Runtime/Renderer/Private/PrimitiveSceneInfo.cpp
+++ b/Engine/Source/Runtime/Renderer/Private/PrimitiveSceneInfo.cpp
@@ -879,7 +879,7 @@ void FPrimitiveSceneInfo::CacheRayTracingPrimitives(FScene* Scene, const TArrayV
 				[Scene, &SceneInfos](FCacheRayTracingPrimitivesContext<FTempRayTracingMeshCommandStorage>& Context, int32 Index)
 				{
 					FMemMark Mark(FMemStack::Get());
-					FTaskTagScope Scope(ETaskTag::EParallelRenderingThread);
+					FOptionalTaskTagScope Scope(ETaskTag::EParallelRenderingThread);
 
 					FPrimitiveSceneInfo* SceneInfo = SceneInfos[Index];
 					FRayTracingInstance CachedInstance;
