Fix VSM invalidation handling such that uninitialized instances are not processed when multiple UpdateAllPrimitiveSceneInfos are called back to back (without an UpdateGPUScene in between).
- Fix debug visualization that was inverted (red should be for invalidated)

diff --git a/Engine/Shaders/Private/VirtualShadowMaps/CacheManagement.usf b/Engine/Shaders/Private/VirtualShadowMaps/CacheManagement.usf
--- a/Engine/Shaders/Private/VirtualShadowMaps/CacheManagement.usf
+++ b/Engine/Shaders/Private/VirtualShadowMaps/CacheManagement.usf
@@ -21,6 +21,8 @@
 
 RWStructuredBuffer<uint> OutDynamicCasterPageFlags;
 uint InstanceSceneDataSOAStride;
+uint GPUSceneNumAllocatedInstances;
+uint GPUSceneNumAllocatedPrimitives;
 
 #if ENABLE_DEBUG_MODE
 uint bDrawBounds;
@@ -59,7 +61,16 @@ void VirtualSmInvalidateInstancePagesCS(
 	int InstanceId = InvalidatingInstances[1 + DispatchIndex];
 #endif
 	{
+		checkSlow(InstanceId >= 0 && InstanceId < InstanceSceneDataSOAStride);
+		checkSlow(InstanceId >= 0 && InstanceId < GPUSceneNumAllocatedInstances);
+
 		FInstanceSceneData InstanceSceneData = GetInstanceSceneData(InstanceId, InstanceSceneDataSOAStride);
+
+		if (!InstanceSceneData.ValidInstance)
+		{
+			return;
+		}
+
 		uint InvalidationFlags = ShouldCacheInstanceAsStatic(InstanceSceneData) ? VSM_STATIC_UNCACHED_FLAG : VSM_DYNAMIC_UNCACHED_FLAG;
 
 		// TODO: Clean up hardcoded flag field test.
diff --git a/Engine/Shaders/Private/VirtualShadowMaps/VirtualShadowMapProjection.usf b/Engine/Shaders/Private/VirtualShadowMaps/VirtualShadowMapProjection.usf
--- a/Engine/Shaders/Private/VirtualShadowMaps/VirtualShadowMapProjection.usf
+++ b/Engine/Shaders/Private/VirtualShadowMaps/VirtualShadowMapProjection.usf
@@ -397,8 +397,8 @@ void OutputDebug( uint2 PixelPos, FVirtualShadowMapSampleResult Result )
 		}
 		else if ( DebugOutputType == 7 )	// Cached or not?
 		{
-			bool bCached = (PhysicalPageMetaData[VSMPhysicalPageAddressToIndex(pPage)].Flags & VSM_DYNAMIC_UNCACHED_FLAG) != 0;
-			float3 Color = GreenToRedTurbo(1.0f - float(bCached));
+			bool bUncached = (PhysicalPageMetaData[VSMPhysicalPageAddressToIndex(pPage)].Flags & VSM_DYNAMIC_UNCACHED_FLAG) != 0;
+			float3 Color = GreenToRedTurbo(float(bUncached));
 			float3 Color2 = IntToColor(vPage.x + (vPage.y << 10U));
 
 			Output = lerp(Color, Color2, 0.35f);
diff --git a/Engine/Source/Runtime/Renderer/Private/GPUScene.cpp b/Engine/Source/Runtime/Renderer/Private/GPUScene.cpp
--- a/Engine/Source/Runtime/Renderer/Private/GPUScene.cpp
+++ b/Engine/Source/Runtime/Renderer/Private/GPUScene.cpp
@@ -23,6 +23,8 @@
 #include "ShaderDebug.h"
 #include "ShaderPrint.h"
 
+#define LOG_INSTANCE_ALLOCATIONS 0
+
 // TODO: Global setting/define
 #ifndef INSTANCE_COMPRESSED_TRANSFORMS
 #define INSTANCE_COMPRESSED_TRANSFORMS 1
@@ -242,14 +244,15 @@ struct FInstanceUploadInfo
 	uint32 LastUpdateSceneFrameNumber = ~uint32(0);
 };
 
-void ValidateInstanceUploadInfo(const FInstanceUploadInfo& UploadInfo)
+void ValidateInstanceUploadInfo(const FInstanceUploadInfo& UploadInfo, const FGPUSceneBufferState& BufferState)
 {
 #if DO_CHECK
 	const bool bHasRandomID			= (UploadInfo.PayloadDataFlags & INSTANCE_SCENE_DATA_FLAG_HAS_RANDOM) != 0u;
 	const bool bHasCustomData		= (UploadInfo.PayloadDataFlags & INSTANCE_SCENE_DATA_FLAG_HAS_CUSTOM_DATA) != 0u;
 	const bool bHasDynamicData		= (UploadInfo.PayloadDataFlags & INSTANCE_SCENE_DATA_FLAG_HAS_DYNAMIC_DATA) != 0u;
-	const bool bHasLMSMScaleBias	= (UploadInfo.PayloadDataFlags & INSTANCE_SCENE_DATA_FLAG_HAS_LIGHTSHADOW_UV_BIAS) != 0u;
+	const bool bHasLightShadowUVBias = (UploadInfo.PayloadDataFlags & INSTANCE_SCENE_DATA_FLAG_HAS_LIGHTSHADOW_UV_BIAS) != 0u;
 	const bool bHasHierarchyOffset	= (UploadInfo.PayloadDataFlags & INSTANCE_SCENE_DATA_FLAG_HAS_HIERARCHY_OFFSET) != 0u;
+	const bool bHasLocalBounds		= (UploadInfo.PayloadDataFlags & INSTANCE_SCENE_DATA_FLAG_HAS_LOCAL_BOUNDS) != 0u;
 #if WITH_EDITOR
 	const bool bHasEditorData		= (UploadInfo.PayloadDataFlags & INSTANCE_SCENE_DATA_FLAG_HAS_EDITOR_DATA) != 0u;
 #endif
@@ -257,7 +260,7 @@ void ValidateInstanceUploadInfo(const FInstanceUploadInfo& UploadInfo)
 	const int32 InstanceCount = UploadInfo.PrimitiveInstances.Num();
 	check(UploadInfo.InstanceRandomID.Num()				== (bHasRandomID		? InstanceCount : 0));
 	check(UploadInfo.InstanceDynamicData.Num()			== (bHasDynamicData		? InstanceCount : 0));
-	check(UploadInfo.InstanceLightShadowUVBias.Num()	== (bHasLMSMScaleBias	? InstanceCount : 0));
+	check(UploadInfo.InstanceLightShadowUVBias.Num()	== (bHasLightShadowUVBias ? InstanceCount : 0));
 	check(UploadInfo.InstanceHierarchyOffset.Num()		== (bHasHierarchyOffset	? InstanceCount : 0));
 #if WITH_EDITOR
 	check(UploadInfo.InstanceEditorData.Num() == (bHasEditorData ? InstanceCount : 0));
@@ -272,6 +275,17 @@ void ValidateInstanceUploadInfo(const FInstanceUploadInfo& UploadInfo)
 	{
 		check(UploadInfo.InstanceCustomData.Num() == 0 && UploadInfo.InstanceCustomDataCount == 0);
 	}
+
+	// RandomID is not stored in the payload but in the instance scene data.
+	const bool bHasAnyPayloadData = bHasHierarchyOffset || bHasLocalBounds || bHasDynamicData || bHasLightShadowUVBias || bHasCustomData/*|| bHasRandomID*/;
+
+	if (bHasAnyPayloadData)
+	{
+		check(UploadInfo.InstancePayloadDataOffset != INDEX_NONE);
+
+		const int32 PayloadBufferSize = BufferState.InstancePayloadDataBuffer.Buffer->GetSize() / BufferState.InstancePayloadDataBuffer.Buffer->GetStride();
+		check(UploadInfo.InstancePayloadDataOffset < PayloadBufferSize);
+	}
 #endif
 }
 
@@ -447,8 +461,6 @@ struct FUploadDataSourceAdapterScenePrimitives
 			InstanceUploadInfo.InstanceCustomDataCount = InstanceUploadInfo.InstanceCustomData.Num() / InstanceUploadInfo.PrimitiveInstances.Num();
 		}
 
-		ValidateInstanceUploadInfo(InstanceUploadInfo);
-
 		// Only trigger upload if this primitive has instances
 		check(InstanceUploadInfo.PrimitiveInstances.Num() > 0);
 	}
@@ -1214,6 +1226,7 @@ void FGPUScene::UploadGeneral(FRHICommandListImmediate& RHICmdList, FScene *Scen
 							const int32 ItemIndex = Item.ItemIndex;
 							FInstanceUploadInfo UploadInfo;
 							UploadDataSourceAdapter.GetInstanceInfo(ItemIndex, UploadInfo);
+							ValidateInstanceUploadInfo(UploadInfo, BufferState);
 							FInstanceBatcher::FPrimitiveItemInfo PrimitiveItemInfo = InstanceUpdates.PerPrimitiveItemInfo[ItemIndex];
 
 							check(NumInstancePayloadDataUploads > 0 || UploadInfo.InstancePayloadDataStride == 0); // Sanity check
@@ -1485,7 +1498,8 @@ struct FUploadDataSourceAdapterDynamicPrimitives
 		FPrimitiveUniformShaderParameters Tmp = PrimitiveShaderData[ItemIndex];
 		Tmp.InstanceSceneDataOffset = InstanceIDStartOffset + ItemIndex;
 		Tmp.NumInstanceSceneDataEntries = 1;
-
+		Tmp.InstancePayloadDataOffset = INDEX_NONE;
+		Tmp.InstancePayloadDataStride = 0;
 		PrimitiveUploadInfo.NumInstanceUploads = 1;
 		PrimitiveUploadInfo.NumInstancePayloadDataUploads = 0;
 		PrimitiveUploadInfo.PrimitiveID = PrimitiveIDStartOffset + ItemIndex;
@@ -1630,12 +1644,6 @@ void FGPUScene::Update(FRDGBuilder& GraphBuilder, FScene& Scene)
 	if (bIsEnabled)
 	{
 		ensure(bInBeginEndBlock);
-		// Invoke the cache manager to invalidate the previous location of all instances that are to be updated, 
-		// must be done prior to update of GPU-side data to use the previous transforms.
-		if (Scene.VirtualShadowMapArrayCacheManager)
-		{
-			Scene.VirtualShadowMapArrayCacheManager->ProcessPrimitivesToUpdate(GraphBuilder, Scene);
-		}
 		
 		UpdateInternal(GraphBuilder, Scene);
 	}
@@ -1659,6 +1667,9 @@ int32 FGPUScene::AllocateInstanceSceneDataSlots(int32 NumInstanceSceneDataEntrie
 		{
 			int32 InstanceSceneDataOffset = InstanceSceneDataAllocator.Allocate(NumInstanceSceneDataEntries);
 			InstanceRangesToClear.Add(FInstanceRange{ uint32(InstanceSceneDataOffset), uint32(NumInstanceSceneDataEntries) });
+#if LOG_INSTANCE_ALLOCATIONS
+			UE_LOG(LogTemp, Warning, TEXT("AllocateInstanceSceneDataSlots: [%6d,%6d)"), InstanceSceneDataOffset, InstanceSceneDataOffset + NumInstanceSceneDataEntries);
+#endif
 
 			return InstanceSceneDataOffset;
 		}
@@ -1675,6 +1686,9 @@ void FGPUScene::FreeInstanceSceneDataSlots(int32 InstanceSceneDataOffset, int32
 	{
 		InstanceSceneDataAllocator.Free(InstanceSceneDataOffset, NumInstanceSceneDataEntries);
 		InstanceRangesToClear.Add(FInstanceRange{ uint32(InstanceSceneDataOffset), uint32(NumInstanceSceneDataEntries) });
+#if LOG_INSTANCE_ALLOCATIONS
+		UE_LOG(LogTemp, Warning, TEXT("FreeInstanceSceneDataSlots: [%6d,%6d)"), InstanceSceneDataOffset, InstanceSceneDataOffset + NumInstanceSceneDataEntries);
+#endif
 	}
 }
 
@@ -2101,10 +2115,19 @@ void FGPUSceneCompactInstanceData::Init(const FScene* Scene, int32 PrimitiveId)
 void FGPUScene::AddClearInstancesPass(FRDGBuilder& GraphBuilder)
 {
 	FInstanceGPULoadBalancer ClearIdData;
+#if LOG_INSTANCE_ALLOCATIONS
+	FString RangesStr;
+#endif
 	for (FInstanceRange Range : InstanceRangesToClear)
 	{
 		ClearIdData.Add(Range.InstanceSceneDataOffset, Range.NumInstanceSceneDataEntries, INVALID_PRIMITIVE_ID);
+#if LOG_INSTANCE_ALLOCATIONS
+		RangesStr.Appendf(TEXT("[%6d, %6d), "), Range.InstanceSceneDataOffset, Range.InstanceSceneDataOffset + Range.NumInstanceSceneDataEntries);
+#endif
 	}
+#if LOG_INSTANCE_ALLOCATIONS
+	UE_LOG(LogTemp, Warning, TEXT("AddClearInstancesPass: \n%s"), *RangesStr);
+#endif
 	AddUpdatePrimitiveIdsPass(GraphBuilder, ClearIdData);
 	InstanceRangesToClear.Empty();
 }
diff --git a/Engine/Source/Runtime/Renderer/Private/GPUScene.h b/Engine/Source/Runtime/Renderer/Private/GPUScene.h
--- a/Engine/Source/Runtime/Renderer/Private/GPUScene.h
+++ b/Engine/Source/Runtime/Renderer/Private/GPUScene.h
@@ -265,7 +265,14 @@ public:
 		}
 	}
 
-	FORCEINLINE EPrimitiveDirtyState GetPrimitiveDirtyState(int32 PrimitiveId) const { return PrimitiveDirtyState[PrimitiveId]; }
+	FORCEINLINE EPrimitiveDirtyState GetPrimitiveDirtyState(int32 PrimitiveId) const 
+	{ 
+		if (PrimitiveId >= PrimitiveDirtyState.Num())
+		{
+			return EPrimitiveDirtyState::None;
+		}
+		return PrimitiveDirtyState[PrimitiveId]; 
+	}
 
 	FORCEINLINE void ResizeDirtyState(int32 NewSizeIn)
 	{
diff --git a/Engine/Source/Runtime/Renderer/Private/RendererScene.cpp b/Engine/Source/Runtime/Renderer/Private/RendererScene.cpp
--- a/Engine/Source/Runtime/Renderer/Private/RendererScene.cpp
+++ b/Engine/Source/Runtime/Renderer/Private/RendererScene.cpp
@@ -4089,20 +4089,27 @@ void FScene::UpdateAllPrimitiveSceneInfos(FRDGBuilder& GraphBuilder, bool bAsync
 
 	RemovedLocalPrimitiveSceneInfos.Sort(FPrimitiveArraySortKey());
 
-	if (VirtualShadowMapArrayCacheManager)
+	if (VirtualShadowMapArrayCacheManager && VirtualShadowMapArrayCacheManager->IsValid())
 	{
-		VirtualShadowMapArrayCacheManager->ProcessRemovedPrimives(GraphBuilder, GPUScene, RemovedLocalPrimitiveSceneInfos);
+		FVirtualShadowMapArrayCacheManager::FInvalidatingPrimitiveCollector InvalidatingPrimitiveCollector(Primitives.Num(), GPUScene);
 
-		if (UpdatedInstances.Num())
+		// All removed primitives must invalidate their footprints in the VSM before leaving
+		for (const FPrimitiveSceneInfo* PrimitiveSceneInfo : RemovedLocalPrimitiveSceneInfos)
 		{
-			TArray<FPrimitiveSceneInfo*> RemovedInstances;
-			RemovedInstances.Reserve(UpdatedInstances.Num());
-			for (auto& Instance : UpdatedInstances)
-			{
-				RemovedInstances.Add(Instance.Key->GetPrimitiveSceneInfo());
-			}
-			VirtualShadowMapArrayCacheManager->ProcessRemovedPrimives(GraphBuilder, GPUScene, RemovedInstances);
+			InvalidatingPrimitiveCollector.Add(PrimitiveSceneInfo);
 		}
+		// All updated instances must also before moving or re-allocating (TODO: filter out only those actually updated)
+		for (const auto& Instance : UpdatedInstances)
+		{
+			InvalidatingPrimitiveCollector.Add(Instance.Key->GetPrimitiveSceneInfo());
+		}
+		// As must all primitive updates, 
+		for (const auto& Transform : UpdatedTransforms)
+		{
+			InvalidatingPrimitiveCollector.Add(Transform.Key->GetPrimitiveSceneInfo());
+		}
+
+		VirtualShadowMapArrayCacheManager->ProcessRemovedOrUpdatedPrimitives(GraphBuilder, GPUScene, InvalidatingPrimitiveCollector);
 	}
 	TArray<FPrimitiveSceneInfo*> AddedLocalPrimitiveSceneInfos(AddedPrimitiveSceneInfos.Array());
 	AddedLocalPrimitiveSceneInfos.Sort(FPrimitiveArraySortKey());
@@ -4915,7 +4922,7 @@ void FScene::UpdateAllPrimitiveSceneInfos(FRDGBuilder& GraphBuilder, bool bAsync
 			bPathTracingNeedsInvalidation = true;
 		}
 	}
-
+	
 	UpdatedAttachmentRoots.Reset();
 	UpdatedTransforms.Reset();
 	UpdatedInstances.Reset();
diff --git a/Engine/Source/Runtime/Renderer/Private/VirtualShadowMaps/VirtualShadowMapCacheManager.cpp b/Engine/Source/Runtime/Renderer/Private/VirtualShadowMaps/VirtualShadowMapCacheManager.cpp
--- a/Engine/Source/Runtime/Renderer/Private/VirtualShadowMaps/VirtualShadowMapCacheManager.cpp
+++ b/Engine/Source/Runtime/Renderer/Private/VirtualShadowMaps/VirtualShadowMapCacheManager.cpp
@@ -437,70 +437,19 @@ bool FVirtualShadowMapArrayCacheManager::IsAccumulatingStats()
 	return CVarAccumulateStats.GetValueOnRenderThread() != 0;
 }
 
-void FVirtualShadowMapArrayCacheManager::ProcessRemovedPrimives(FRDGBuilder& GraphBuilder, const FGPUScene& GPUScene, const TArray<FPrimitiveSceneInfo*>& RemovedPrimitiveSceneInfos)
+void FVirtualShadowMapArrayCacheManager::ProcessRemovedOrUpdatedPrimitives(FRDGBuilder& GraphBuilder, const FGPUScene& GPUScene, FInvalidatingPrimitiveCollector& InvalidatingPrimitiveCollector)
 {
-	if (CVarCacheVirtualSMs.GetValueOnRenderThread() != 0 && RemovedPrimitiveSceneInfos.Num() > 0 && PrevBuffers.DynamicCasterPageFlags.IsValid())
+	if (CVarCacheVirtualSMs.GetValueOnRenderThread() != 0 && !InvalidatingPrimitiveCollector.IsEmpty() && PrevBuffers.DynamicCasterPageFlags.IsValid())
 	{
-		RDG_EVENT_SCOPE(GraphBuilder, "Shadow.Virtual.ProcessRemovedPrimives [%d]", RemovedPrimitiveSceneInfos.Num());
-
-		FInstanceGPULoadBalancer LoadBalancer;
-
-		int32 TotalInstanceCount = 0;
-		for (const FPrimitiveSceneInfo* PrimitiveSceneInfo : RemovedPrimitiveSceneInfos)
-		{
-			if (PrimitiveSceneInfo->GetInstanceSceneDataOffset() != INDEX_NONE)
-			{
-				const int32 NumInstanceSceneDataEntries = PrimitiveSceneInfo->GetNumInstanceSceneDataEntries();
-				LoadBalancer.Add(PrimitiveSceneInfo->GetInstanceSceneDataOffset(), NumInstanceSceneDataEntries, 0U);
-				TotalInstanceCount += NumInstanceSceneDataEntries;
-			}
-		}
-		ProcessInvalidations(GraphBuilder, LoadBalancer, TotalInstanceCount, GPUScene);
+		RDG_EVENT_SCOPE(GraphBuilder, "Shadow.Virtual.ProcessRemovedOrUpdatedPrimitives");
+#if VSM_LOG_INVALIDATIONS
+		UE_LOG(LogTemp, Warning, TEXT("ProcessRemovedOrUpdatedPrimitives: \n%s"), *InvalidatingPrimitiveCollector.RangesStr);
+#endif
+		ProcessInvalidations(GraphBuilder, InvalidatingPrimitiveCollector.LoadBalancer, InvalidatingPrimitiveCollector.TotalInstanceCount, GPUScene);
 	}
 }
 
 
-void FVirtualShadowMapArrayCacheManager::ProcessPrimitivesToUpdate(FRDGBuilder& GraphBuilder, const FScene& Scene)
-{
-	const FGPUScene& GPUScene = Scene.GPUScene;
-	if (IsValid())
-	{
-		RDG_EVENT_SCOPE(GraphBuilder, "Shadow.Virtual.ProcessPrimitivesToUpdate [%d]", GPUScene.PrimitivesToUpdate.Num());
-
-		ProcessGPUInstanceInvalidations(GraphBuilder, GPUScene);
-
-		if (GPUScene.PrimitivesToUpdate.Num() > 0)
-		{
-			FInstanceGPULoadBalancer LoadBalancer;
-
-			int32 TotalInstanceCount = 0;
-			for (const int32 PrimitiveId : GPUScene.PrimitivesToUpdate)
-			{
-				// There may possibly be IDs that are out of range if they were marked for update and then removed.
-				if (PrimitiveId < Scene.Primitives.Num())
-				{
-					EPrimitiveDirtyState PrimitiveDirtyState = GPUScene.GetPrimitiveDirtyState(PrimitiveId);
-
-					// SKIP if marked for Add, because this means it has no previous location to invalidate.
-					// SKIP if transform has not changed, as this means no invalidation needs to take place.
-					const bool bDoInvalidation = !EnumHasAnyFlags(PrimitiveDirtyState, EPrimitiveDirtyState::Added) && EnumHasAnyFlags(PrimitiveDirtyState, EPrimitiveDirtyState::ChangedTransform);
-					if (bDoInvalidation)
-					{
-						const FPrimitiveSceneInfo* PrimitiveSceneInfo = Scene.Primitives[PrimitiveId];
-						if (PrimitiveSceneInfo->GetInstanceSceneDataOffset() != INDEX_NONE)
-						{
-							const int32 NumInstanceSceneDataEntries = PrimitiveSceneInfo->GetNumInstanceSceneDataEntries();
-							LoadBalancer.Add(PrimitiveSceneInfo->GetInstanceSceneDataOffset(), NumInstanceSceneDataEntries, PrimitiveId);
-							TotalInstanceCount += NumInstanceSceneDataEntries;
-						}
-					}
-				}
-			}
-			ProcessInvalidations(GraphBuilder, LoadBalancer, TotalInstanceCount, GPUScene);
-		}
-	}
-}
-
 /**
  * Compute shader to project and invalidate the rectangles of given instances.
  */
@@ -537,6 +486,8 @@ public:
 		SHADER_PARAMETER_SRV(StructuredBuffer<float4>, GPUScenePrimitiveSceneData)
 		SHADER_PARAMETER(uint32, GPUSceneFrameNumber)
 		SHADER_PARAMETER(uint32, InstanceSceneDataSOAStride)
+		SHADER_PARAMETER(uint32, GPUSceneNumAllocatedInstances)
+		SHADER_PARAMETER(uint32, GPUSceneNumAllocatedPrimitives)
 
 		SHADER_PARAMETER_RDG_BUFFER_SRV(StructuredBuffer< uint >, ShadowHZBPageTable)
 		SHADER_PARAMETER_RDG_TEXTURE(Texture2D, HZBTexture)
@@ -635,6 +586,9 @@ static void SetupCommonParameters(FRDGBuilder& GraphBuilder, FVirtualShadowMapAr
 	OutPassParameters.GPUScenePrimitiveSceneData = GPUScene.PrimitiveBuffer.SRV;
 	OutPassParameters.GPUSceneInstancePayloadData = GPUScene.InstancePayloadDataBuffer.SRV;
 	OutPassParameters.GPUSceneFrameNumber = GPUScene.GetSceneFrameNumber();
+	OutPassParameters.GPUSceneNumAllocatedInstances = GPUScene.GetNumInstances();
+	OutPassParameters.GPUSceneNumAllocatedPrimitives = GPUScene.GetNumPrimitives();
+
 	OutPassParameters.InstanceSceneDataSOAStride = GPUScene.InstanceSceneDataSOAStride;
 	OutPassParameters.bDrawBounds = bDrawBounds;
 
diff --git a/Engine/Source/Runtime/Renderer/Private/VirtualShadowMaps/VirtualShadowMapCacheManager.h b/Engine/Source/Runtime/Renderer/Private/VirtualShadowMaps/VirtualShadowMapCacheManager.h
--- a/Engine/Source/Runtime/Renderer/Private/VirtualShadowMaps/VirtualShadowMapCacheManager.h
+++ b/Engine/Source/Runtime/Renderer/Private/VirtualShadowMaps/VirtualShadowMapCacheManager.h
@@ -9,10 +9,13 @@
 #include "VirtualShadowMapArray.h"
 #include "SceneManagement.h"
 #include "InstanceCulling/InstanceCullingLoadBalancer.h"
+#include "GPUScene.h"
 
 class FRHIGPUBufferReadback;
 class FGPUScene;
 
+#define VSM_LOG_INVALIDATIONS 0
+
 class FVirtualShadowMapCacheEntry
 {
 public:
@@ -110,17 +113,60 @@ public:
 
 	bool IsAccumulatingStats();
 
+	using FInstanceGPULoadBalancer = TInstanceCullingLoadBalancer<SceneRenderingAllocator>;
+
+	/**
+	 * Helper to collect primitives that need invalidation, filters out redundant adds and also those that are not yet known to the GPU
+	 */
+	class FInvalidatingPrimitiveCollector
+	{
+	public:
+		FInvalidatingPrimitiveCollector(int32 MaxPrimitiveID, const FGPUScene& InGPUScene)
+			: AlreadyAddedPrimitives(false, MaxPrimitiveID)
+			, GPUScene(InGPUScene)
+		{
+		}
+
+		/**
+		 * Add a primitive to invalidate the instances for, the function filters redundant primitive adds, and thus expects valid IDs (so can't be called for primitives that have not yet been added)
+		 * and unchanging IDs (so can't be used over a span that include any scene mutation).
+		 */
+		void Add(const FPrimitiveSceneInfo* PrimitiveSceneInfo)
+		{
+			int32 PrimitiveID = PrimitiveSceneInfo->GetIndex();
+			if (PrimitiveID >= 0
+				&& !AlreadyAddedPrimitives[PrimitiveID]
+				&& PrimitiveSceneInfo->GetInstanceSceneDataOffset() != INDEX_NONE
+				// Don't process primitives that are still in the 'added' state because this means that they
+				// have not been uploaded to the GPU yet and may be pending from a previous call to update primitive scene infos.
+				&& !EnumHasAnyFlags(GPUScene.GetPrimitiveDirtyState(PrimitiveID), EPrimitiveDirtyState::Added))
+			{
+				AlreadyAddedPrimitives[PrimitiveID] = true;
+				const int32 NumInstanceSceneDataEntries = PrimitiveSceneInfo->GetNumInstanceSceneDataEntries();
+				LoadBalancer.Add(PrimitiveSceneInfo->GetInstanceSceneDataOffset(), NumInstanceSceneDataEntries, 0U);
+#if VSM_LOG_INVALIDATIONS
+				RangesStr.Appendf(TEXT("[%6d, %6d), "), PrimitiveSceneInfo->GetInstanceSceneDataOffset(), PrimitiveSceneInfo->GetInstanceSceneDataOffset() + NumInstanceSceneDataEntries);
+#endif
+				TotalInstanceCount += NumInstanceSceneDataEntries;
+			}
+		}
+
+		bool IsEmpty() const { return LoadBalancer.IsEmpty(); }
+
+		TBitArray<SceneRenderingAllocator> AlreadyAddedPrimitives;
+		FInstanceGPULoadBalancer LoadBalancer;
+		int32 TotalInstanceCount = 0;
+#if VSM_LOG_INVALIDATIONS
+		FString RangesStr;
+#endif
+		const FGPUScene& GPUScene;
+	};
 	/**
 	 * This must to be executed before the instances are actually removed / updated, otherwise the wrong position will be used. 
 	 * In particular, it must be processed before the Scene primitive IDs are updated/compacted as part of the removal.
 	 * Invalidate pages that are touched by (the instances of) the removed primitives. 
 	 */
-	void ProcessRemovedPrimives(FRDGBuilder& GraphBuilder, const FGPUScene& GPUScene, const TArray<FPrimitiveSceneInfo*>& RemovedPrimitiveSceneInfos);
-
-	/**
-	 * Invalidate pages that are touched by (the instances of) all primitive about to be removed. Must be called before the GPU-Scene is updated (but after all upates are registered).
-	 */
-	void ProcessPrimitivesToUpdate(FRDGBuilder& GraphBuilder, const FScene& Scene);
+	void ProcessRemovedOrUpdatedPrimitives(FRDGBuilder& GraphBuilder, const FGPUScene& GPUScene, FInvalidatingPrimitiveCollector& InvalidatingPrimitiveCollector);
 
 	TRDGUniformBufferRef<FVirtualShadowMapUniformParameters> GetPreviousUniformBuffer(FRDGBuilder& GraphBuilder) const;
 
@@ -129,7 +175,6 @@ public:
 		
 	void SetHZBViewParams(int32 HZBKey, Nanite::FPackedViewParams& OutParams);
 
-	using FInstanceGPULoadBalancer = TInstanceCullingLoadBalancer<SceneRenderingAllocator>;
 
 private:
 	void ProcessInvalidations(FRDGBuilder& GraphBuilder, FInstanceGPULoadBalancer& Instances, int32 TotalInstanceCount, const FGPUScene& GPUScene);
