Add message to read back page pool status to host & display error message
- also make GPUMessage thread safe

diff --git a/Engine/Shaders/Private/VirtualShadowMaps/PageManagement.usf b/Engine/Shaders/Private/VirtualShadowMaps/PageManagement.usf
--- a/Engine/Shaders/Private/VirtualShadowMaps/PageManagement.usf
+++ b/Engine/Shaders/Private/VirtualShadowMaps/PageManagement.usf
@@ -22,6 +22,7 @@
 #include "ProjectionSpot.ush"
 #include "ProjectionDirectional.ush"
 #include "Stats.ush"
+#include "../GPUMessaging.ush"
 
 #ifndef HAS_CACHE_DATA
 #define HAS_CACHE_DATA 0
@@ -1214,3 +1215,19 @@ void BuildHZBPerPageTopCS(uint2 GroupThreadId : SV_GroupThreadID, uint PageInput
 		FurthestHZBOutput_1[OutPixelPos] = MinDeviceZ;
 	}
 }
+
+uint StatusMessageId;
+StructuredBuffer<int> FreePhysicalPages;
+
+
+[numthreads(1, 1, 1)]
+void FeedbackStatusCS()
+{
+	{
+		FGPUMessageWriter Mw = GPUMessageBegin(StatusMessageId, 2U);
+		// 1. Write out frame number
+		GPUMessageWriteItem(Mw, 0U);
+		// 2. Write out the free list status
+		GPUMessageWriteItem(Mw, FreePhysicalPages[VirtualShadowMap.MaxPhysicalPages]);
+	}
+}
diff --git a/Engine/Source/Runtime/RenderCore/Private/GPUMessaging.cpp b/Engine/Source/Runtime/RenderCore/Private/GPUMessaging.cpp
--- a/Engine/Source/Runtime/RenderCore/Private/GPUMessaging.cpp
+++ b/Engine/Source/Runtime/RenderCore/Private/GPUMessaging.cpp
@@ -17,7 +17,7 @@ public:
 	SHADER_USE_PARAMETER_STRUCT(FClearMessageBufferCS, FGlobalShader)
 
 	BEGIN_SHADER_PARAMETER_STRUCT(FParameters, )
-		SHADER_PARAMETER_RDG_BUFFER_UAV(StructuredBuffer<uint>, GPUMessageDataBuffer)
+		SHADER_PARAMETER_RDG_BUFFER_UAV(RWStructuredBuffer<uint>, GPUMessageDataBuffer)
 	END_SHADER_PARAMETER_STRUCT()
 };
 
@@ -48,10 +48,13 @@ static FAutoConsoleVariableRef CVarMaxBufferSize(
 class FSystem : public FRenderResource
 {
 public:
+	// Need to guard access because it may happen during FScene initialization (Game Thread) while another scene is being rendered.
+	FCriticalSection AccessCS;
+
+
 	FSocket RegisterHandler(const TSharedPtr<FHandler>& Handler)
 	{
-		checkf(!MessageBuffer, TEXT("Handlers must be registered outside of an active GPU message scope."));
-
+		FScopeLock Lock(&AccessCS);
 		const FMessageId MessageId(NextMessageId++);
 		check(!MessageHandlers.Contains(MessageId));
 		MessageHandlers.Add(MessageId, Handler);
@@ -60,8 +63,7 @@ public:
 
 	void RemoveHandler(FMessageId MessageId)
 	{
-		checkf(!MessageBuffer, TEXT("Handlers must be unregistered outside of an active GPU message scope."));
-
+		FScopeLock Lock(&AccessCS);
 		check(MessageId.IsValid());
 		check(MessageHandlers.Contains(MessageId));
 		MessageHandlers.Remove(MessageId);
@@ -70,6 +72,8 @@ public:
 	void BeginMessageScope(FRDGBuilder& GraphBuilder)
 	{
 		check(MessageBuffer == nullptr);
+		FScopeLock Lock(&AccessCS);
+
 		PollMessages();
 
 		if (MessageHandlers.IsEmpty())
@@ -97,6 +101,8 @@ public:
 
 	void EndMessageScope(FRDGBuilder& GraphBuilder)
 	{
+		FScopeLock Lock(&AccessCS);
+
 		if (MessageHandlers.IsEmpty())
 		{
 			check(MessageBuffer == nullptr);
@@ -213,6 +219,8 @@ private:
 
 	void ReleaseDynamicRHI() override
 	{
+		FScopeLock Lock(&AccessCS);
+
 		check(MessageBuffer == nullptr);
 
 		for (FRHIGPUBufferReadback* ReadbackBuffer : MessageReadbackBuffersInUse)
diff --git a/Engine/Source/Runtime/RenderCore/Public/GPUMessaging.h b/Engine/Source/Runtime/RenderCore/Public/GPUMessaging.h
--- a/Engine/Source/Runtime/RenderCore/Public/GPUMessaging.h
+++ b/Engine/Source/Runtime/RenderCore/Public/GPUMessaging.h
@@ -81,6 +81,8 @@ public:
 		return DefaultValue;
 	}
 
+	uint32 GetPayloadSize() const { return PayloadSize; }
+
 private:
 	FReader(FMessageId InMessageId, uint32 InPayloadSize, const uint32* InPayloadData)
 		: MessageId(InMessageId)
@@ -149,7 +151,7 @@ FSocket RegisterHandler(const TCHAR* MessageDebugId, LambdaType&& Lambda)
 
 /** Shader parameters used to write to write a GPU message. */
 BEGIN_SHADER_PARAMETER_STRUCT(FParameters, RENDERCORE_API)
-	SHADER_PARAMETER_RDG_BUFFER_UAV(StructuredBuffer<uint>, GPUMessageDataBuffer)
+	SHADER_PARAMETER_RDG_BUFFER_UAV(RWStructuredBuffer<uint>, GPUMessageDataBuffer)
 	SHADER_PARAMETER(uint32, GPUMessageDataBufferSize)
 END_SHADER_PARAMETER_STRUCT()
 
diff --git a/Engine/Source/Runtime/Renderer/Private/RendererScene.cpp b/Engine/Source/Runtime/Renderer/Private/RendererScene.cpp
--- a/Engine/Source/Runtime/Renderer/Private/RendererScene.cpp
+++ b/Engine/Source/Runtime/Renderer/Private/RendererScene.cpp
@@ -1126,7 +1126,7 @@ FScene::FScene(UWorld* InWorld, bool bInRequiresHitProxies, bool bInIsEditorScen
 
 	LumenSceneData = new FLumenSceneData(GShaderPlatformForFeatureLevel[InFeatureLevel], InWorld->WorldType);
 
-	VirtualShadowMapArrayCacheManager = new FVirtualShadowMapArrayCacheManager;
+	VirtualShadowMapArrayCacheManager = new FVirtualShadowMapArrayCacheManager(this);
 }
 
 FScene::~FScene()
diff --git a/Engine/Source/Runtime/Renderer/Private/VirtualShadowMaps/VirtualShadowMapArray.cpp b/Engine/Source/Runtime/Renderer/Private/VirtualShadowMaps/VirtualShadowMapArray.cpp
--- a/Engine/Source/Runtime/Renderer/Private/VirtualShadowMaps/VirtualShadowMapArray.cpp
+++ b/Engine/Source/Runtime/Renderer/Private/VirtualShadowMaps/VirtualShadowMapArray.cpp
@@ -14,6 +14,8 @@
 #include "ComponentRecreateRenderStateContext.h"
 #include "HairStrands/HairStrandsData.h"
 #include "SceneTextureReductions.h"
+#include "ShaderDebug.h"
+#include "GPUMessaging.h"
 
 IMPLEMENT_STATIC_UNIFORM_BUFFER_SLOT(VirtualShadowMapUbSlot);
 
@@ -419,6 +421,7 @@ void FVirtualShadowMapArray::SetShaderDefines(FShaderCompilerEnvironment& OutEnv
 	OutEnvironment.SetDefine(TEXT("VSM_VIRTUAL_MAX_RESOLUTION_XY"), FVirtualShadowMap::VirtualMaxResolutionXY);
 	OutEnvironment.SetDefine(TEXT("VSM_RASTER_WINDOW_PAGES"), FVirtualShadowMap::RasterWindowPages);
 	OutEnvironment.SetDefine(TEXT("VSM_PAGE_TABLE_SIZE"), FVirtualShadowMap::PageTableSize);
+	OutEnvironment.SetDefine(TEXT("VSM_NUM_STATS"), NumStats);
 	OutEnvironment.SetDefine(TEXT("INDEX_NONE"), INDEX_NONE);
 }
 
@@ -867,6 +870,29 @@ class FInitPageRectBoundsCS : public FVirtualPageManagementShader
 IMPLEMENT_GLOBAL_SHADER(FInitPageRectBoundsCS, "/Engine/Private/VirtualShadowMaps/PageManagement.usf", "InitPageRectBounds", SF_Compute);
 
 
+
+class FVirtualSmFeedbackStatusCS : public FVirtualPageManagementShader
+{
+	DECLARE_GLOBAL_SHADER(FVirtualSmFeedbackStatusCS);
+	SHADER_USE_PARAMETER_STRUCT(FVirtualSmFeedbackStatusCS, FVirtualPageManagementShader)
+
+	BEGIN_SHADER_PARAMETER_STRUCT(FParameters, )
+		SHADER_PARAMETER_RDG_UNIFORM_BUFFER(FVirtualShadowMapUniformParameters, VirtualShadowMap)
+		SHADER_PARAMETER_RDG_BUFFER_SRV(StructuredBuffer< uint >, FreePhysicalPages)
+		SHADER_PARAMETER_STRUCT_INCLUDE(GPUMessage::FParameters, GPUMessageParams)
+		SHADER_PARAMETER(uint32, StatusMessageId)
+	END_SHADER_PARAMETER_STRUCT()
+
+
+	static void ModifyCompilationEnvironment(const FGlobalShaderPermutationParameters& Parameters, FShaderCompilerEnvironment& OutEnvironment)
+	{
+		FVirtualPageManagementShader::ModifyCompilationEnvironment(Parameters, OutEnvironment);
+	}
+};
+IMPLEMENT_GLOBAL_SHADER(FVirtualSmFeedbackStatusCS, "/Engine/Private/VirtualShadowMaps/PageManagement.usf", "FeedbackStatusCS", SF_Compute);
+
+
+
 static FString GetLightNameForDebug(const FLightSceneProxy* Proxy)
 {
 	// TODO: ActorLabel is editor-only... need a better solution that goes right into Proxy->GetComponentName()
@@ -1436,6 +1462,26 @@ void FVirtualShadowMapArray::BuildPageAllocations(
 		}
 
 		UniformParameters.PageTable = GraphBuilder.CreateSRV(PageTableRDG);
+
+		// Add pass to pipe back important stats
+		{
+
+			FVirtualSmFeedbackStatusCS::FParameters* PassParameters = GraphBuilder.AllocParameters<FVirtualSmFeedbackStatusCS::FParameters>();
+			PassParameters->FreePhysicalPages = GraphBuilder.CreateSRV(FreePhysicalPagesRDG);
+			PassParameters->GPUMessageParams = GPUMessage::GetShaderParameters(GraphBuilder);
+			PassParameters->StatusMessageId = CacheManager->StatusFeedbackSocket.GetMessageId().GetIndex();
+			PassParameters->VirtualShadowMap = GetUniformBuffer(GraphBuilder);
+
+			auto ComputeShader = GetGlobalShaderMap(GMaxRHIFeatureLevel)->GetShader<FVirtualSmFeedbackStatusCS>();
+
+			FComputeShaderUtils::AddPass(
+				GraphBuilder,
+				RDG_EVENT_NAME("Feedback Status"),
+				ComputeShader,
+				PassParameters,
+				FIntVector(1, 1, 1)
+			);
+		}
 	}
 
 #if !UE_BUILD_SHIPPING
diff --git a/Engine/Source/Runtime/Renderer/Private/VirtualShadowMaps/VirtualShadowMapCacheManager.cpp b/Engine/Source/Runtime/Renderer/Private/VirtualShadowMaps/VirtualShadowMapCacheManager.cpp
--- a/Engine/Source/Runtime/Renderer/Private/VirtualShadowMaps/VirtualShadowMapCacheManager.cpp
+++ b/Engine/Source/Runtime/Renderer/Private/VirtualShadowMaps/VirtualShadowMapCacheManager.cpp
@@ -30,7 +30,10 @@ static TAutoConsoleVariable<int32> CVarCacheVirtualSMs(
 static TAutoConsoleVariable<int32> CVarDrawInvalidatingBounds(
 	TEXT("r.Shadow.Virtual.DrawInvalidatingBounds"),
 	0,
-	TEXT("Turn on debug render cache invalidating instance bounds."),
+	TEXT("Turn on debug render cache invalidating instance bounds, heat mapped by number of pages invalidated.\n")
+	TEXT("   1  = Draw all bounds.\n")
+	TEXT("   2  = Draw those invalidating static cached pages only\n")
+	TEXT("   3  = Draw those invalidating dynamic cached pages only"),
 	ECVF_RenderThreadSafe
 );
 
@@ -122,6 +125,70 @@ void FVirtualShadowMapCacheEntry::UpdateLocal(int32 VirtualShadowMapId, const FW
 	bCurrentRendered = false;
 }
 
+FVirtualShadowMapArrayCacheManager::FVirtualShadowMapArrayCacheManager(FScene* InScene) 
+	: Scene(InScene)
+{
+	// Handle message with status sent back from GPU
+	StatusFeedbackSocket = GPUMessage::RegisterHandler(TEXT("Shadow.Virtual.StatusFeedback"), [this](GPUMessage::FReader Message)
+	{
+		// Get the frame that the message was sent.
+		uint32 FrameNumber = Message.Read<uint32>(0);
+		// Goes negative on underflow
+		int32 NumPagesFree = Message.Read<int32>(0);
+
+		if (NumPagesFree < 0)
+		{
+			static const auto* CVarResolutionLodBiasLocalPtr = IConsoleManager::Get().FindTConsoleVariableDataFloat(TEXT("r.Shadow.Virtual.ResolutionLodBiasLocal"));
+			const float LodBiasLocal = CVarResolutionLodBiasLocalPtr->GetValueOnRenderThread();
+
+			static const auto* CVarResolutionLodBiasDirectionalPtr = IConsoleManager::Get().FindTConsoleVariableDataFloat(TEXT("r.Shadow.Virtual.ResolutionLodBiasDirectional"));
+			const float LodBiasDirectional = CVarResolutionLodBiasDirectionalPtr->GetValueOnRenderThread();
+
+			static const auto* CVarMaxPhysicalPagesPtr = IConsoleManager::Get().FindTConsoleVariableDataInt(TEXT("r.Shadow.Virtual.MaxPhysicalPages"));
+			const int32 MaxPhysicalPages = CVarMaxPhysicalPagesPtr->GetValueOnRenderThread();
+
+#if !UE_BUILD_SHIPPING
+			if (!bLoggedPageOverflow)
+			{
+				UE_LOG(LogRenderer, Warning, TEXT("Virtual Shadow Map Page Pool overflow (%d page allocations were not served), this will produce visual artifacts (missing shadow), increase the page pool limit or reduce resolution bias to avoid.\n")
+					TEXT(" See r.Shadow.Virtual.MaxPhysicalPages (%d), r.Shadow.Virtual.ResolutionLodBiasLocal (%.2f), and r.Shadow.Virtual.ResolutionLodBiasDirectional (%.2f)"),
+					-NumPagesFree,
+					MaxPhysicalPages,
+					LodBiasLocal,
+					LodBiasDirectional);
+				bLoggedPageOverflow = true;
+			}
+			LastOverflowFrame = Scene->GetFrameNumber();
+#endif
+		}
+#if !UE_BUILD_SHIPPING
+		else
+		{
+			bLoggedPageOverflow = false;
+		}
+#endif
+	});
+
+#if !UE_BUILD_SHIPPING
+	ScreenMessageDelegate = FCoreDelegates::OnGetOnScreenMessages.AddLambda([this](TMultiMap<FCoreDelegates::EOnScreenMessageSeverity, FText >& OutMessages)
+	{
+		// Show for ~5s after last overflow
+		int32 CurrentFrameNumber = Scene->GetFrameNumber();
+		if (LastOverflowFrame >= 0 && CurrentFrameNumber - LastOverflowFrame < 30 * 5)
+		{
+			OutMessages.Add(FCoreDelegates::EOnScreenMessageSeverity::Warning, FText::FromString(FString::Printf(TEXT("Virtual Shadow Map Page Pool overflow detected (%d frames ago)"), CurrentFrameNumber - LastOverflowFrame)));
+		}
+	});
+#endif
+}
+
+FVirtualShadowMapArrayCacheManager::~FVirtualShadowMapArrayCacheManager()
+{
+#if !UE_BUILD_SHIPPING
+	FCoreDelegates::OnGetOnScreenMessages.Remove(ScreenMessageDelegate);
+#endif
+}
+
 
 TRefCountPtr<IPooledRenderTarget> FVirtualShadowMapArrayCacheManager::SetPhysicalPoolSize(FRDGBuilder& GraphBuilder, FIntPoint RequestedSize)
 {
diff --git a/Engine/Source/Runtime/Renderer/Private/VirtualShadowMaps/VirtualShadowMapCacheManager.h b/Engine/Source/Runtime/Renderer/Private/VirtualShadowMaps/VirtualShadowMapCacheManager.h
--- a/Engine/Source/Runtime/Renderer/Private/VirtualShadowMaps/VirtualShadowMapCacheManager.h
+++ b/Engine/Source/Runtime/Renderer/Private/VirtualShadowMaps/VirtualShadowMapCacheManager.h
@@ -10,6 +10,7 @@
 #include "SceneManagement.h"
 #include "InstanceCulling/InstanceCullingLoadBalancer.h"
 #include "GPUScene.h"
+#include "GPUMessaging.h"
 
 class FRHIGPUBufferReadback;
 class FGPUScene;
@@ -81,6 +82,9 @@ struct FVirtualShadowMapArrayFrameData
 class FVirtualShadowMapArrayCacheManager
 {
 public:
+	FVirtualShadowMapArrayCacheManager(FScene *InScene);
+	~FVirtualShadowMapArrayCacheManager();
+
 	// Enough for er lots...
 	static constexpr uint32 MaxStatFrames = 512*1024U;
 
@@ -176,6 +180,8 @@ public:
 	void SetHZBViewParams(int32 HZBKey, Nanite::FPackedViewParams& OutParams);
 
 
+	GPUMessage::FSocket StatusFeedbackSocket;
+
 private:
 	void ProcessInvalidations(FRDGBuilder& GraphBuilder, FInstanceGPULoadBalancer& Instances, int32 TotalInstanceCount, const FGPUScene& GPUScene);
 
@@ -196,4 +202,10 @@ private:
 	TRefCountPtr<FRDGPooledBuffer> AccumulatedStatsBuffer;
 	bool bAccumulatingStats = false;
 	FRHIGPUBufferReadback* GPUBufferReadback = nullptr;
+#if !UE_BUILD_SHIPPING
+	FDelegateHandle ScreenMessageDelegate;
+	int32 LastOverflowFrame = -1;
+	bool bLoggedPageOverflow = false;
+#endif
+	FScene* Scene;
 };
