Flush virtual textures pending delete before allocating new virtual textures if we are about to exceed virtual texture space limit.
Avoids hitting space limit in some cases where pending queue is otherwise slow to flush (such as during cook).

diff --git a/Engine/Source/Runtime/Renderer/Private/VT/VirtualTextureSystem.cpp b/Engine/Source/Runtime/Renderer/Private/VT/VirtualTextureSystem.cpp
--- a/Engine/Source/Runtime/Renderer/Private/VT/VirtualTextureSystem.cpp
+++ b/Engine/Source/Runtime/Renderer/Private/VT/VirtualTextureSystem.cpp
@@ -630,6 +630,14 @@ IAllocatedVirtualTexture* FVirtualTextureSystem::AllocateVirtualTexture(const FA
 		NumPhysicalGroups = Desc.NumTextureLayers;
 	}
 
+	if (NumAllocatedSpaces == MaxSpaces)
+	{
+		// We have already allocated the maximum number of spaces. If this allocation needs a new space then it will fail.
+		// To give us the best chance of success, flush all pending destroys now. This may free up a space.
+		// This is potentially an expensive operation, which is why we only do it if necessary.
+		DestroyPendingVirtualTextures(true);
+	}
+
 	AllocatedVT = new FAllocatedVirtualTexture(this, Frame, Desc, ProducerForLayer, BlockWidthInTiles, BlockHeightInTiles, WidthInBlocks, HeightInBlocks, DepthInTiles);
 	AllocatedVT->NumRefs = 1;
 	if (bAnyLayerProducerWantsPersistentHighestMip)
@@ -808,6 +816,7 @@ FVirtualTextureSpace* FVirtualTextureSystem::AcquireSpace(const FVTSpaceDescript
 				const uint32 InitialPageTableSize = InDesc.bPrivateSpace ? InDesc.MaxSpaceSize : FMath::Max(AllocatedVT->GetWidthInTiles(), AllocatedVT->GetHeightInTiles());
 				FVirtualTextureSpace* Space = new FVirtualTextureSpace(this, SpaceIndex, InDesc, InitialPageTableSize);
 				Spaces[SpaceIndex].Reset(Space);
+				NumAllocatedSpaces++;
 				INC_MEMORY_STAT_BY(STAT_TotalPagetableMemory, Space->GetSizeInBytes());
 				BeginInitResource(Space);
 
@@ -837,6 +846,7 @@ void FVirtualTextureSystem::ReleaseSpace(FVirtualTextureSpace* Space)
 		DEC_MEMORY_STAT_BY(STAT_TotalPagetableMemory, Space->GetSizeInBytes());
 		Space->ReleaseResource();
 		Spaces[Space->GetID()].Reset();
+		NumAllocatedSpaces--;
 	}
 }
 
diff --git a/Engine/Source/Runtime/Renderer/Private/VT/VirtualTextureSystem.h b/Engine/Source/Runtime/Renderer/Private/VT/VirtualTextureSystem.h
--- a/Engine/Source/Runtime/Renderer/Private/VT/VirtualTextureSystem.h
+++ b/Engine/Source/Runtime/Renderer/Private/VT/VirtualTextureSystem.h
@@ -126,6 +126,7 @@ private:
 
 	static const uint32 MaxNumTasks = 16;
 	static const uint32 MaxSpaces = 16;
+	uint32 NumAllocatedSpaces = 0;
 	TUniquePtr<FVirtualTextureSpace> Spaces[MaxSpaces];
 	TArray<FVirtualTexturePhysicalSpace*> PhysicalSpaces;
 	FVirtualTextureProducerCollection Producers;
