Working around D3D11 hitches caused by unpaged memory returned from Map. The workaround is to have UnifiedBuffer create brand new resources instead of updating existing ones.

diff --git a/Engine/Source/Runtime/RenderCore/Private/UnifiedBuffer.cpp b/Engine/Source/Runtime/RenderCore/Private/UnifiedBuffer.cpp
--- a/Engine/Source/Runtime/RenderCore/Private/UnifiedBuffer.cpp
+++ b/Engine/Source/Runtime/RenderCore/Private/UnifiedBuffer.cpp
@@ -678,7 +678,6 @@ void FScatterUploadBuffer::Init(TArrayView<const uint32> ElementScatterOffsets,
 	Init(ElementScatterOffsets.Num(), InNumBytesPerElement, bInFloat4Buffer, DebugName);
 	FMemory::ParallelMemcpy(ScatterData, ElementScatterOffsets.GetData(), ElementScatterOffsets.Num() * ElementScatterOffsets.GetTypeSize(), EMemcpyCachePolicy::StoreUncached);
 	NumScatters = ElementScatterOffsets.Num();
-	ScatterData += ElementScatterOffsets.Num();
 }
 
 void FScatterUploadBuffer::InitPreSized(uint32 NumElements, uint32 InNumBytesPerElement, bool bInFloat4Buffer, const TCHAR* DebugName)
diff --git a/Engine/Source/Runtime/RenderCore/Public/UnifiedBuffer.h b/Engine/Source/Runtime/RenderCore/Public/UnifiedBuffer.h
--- a/Engine/Source/Runtime/RenderCore/Public/UnifiedBuffer.h
+++ b/Engine/Source/Runtime/RenderCore/Public/UnifiedBuffer.h
@@ -177,4 +177,15 @@ public:
 
 		return UploadData + ElementIndex * NumBytesPerElement;
 	}
+
+	void SetUploadViaCreate(bool bInUploadViaCreate)
+	{
+		if (bInUploadViaCreate != bUploadViaCreate)
+		{
+			// When switching the upload path, just free everything.
+			Release();
+
+			bUploadViaCreate = bInUploadViaCreate;
+		}
+	}
 };
\ No newline at end of file
diff --git a/Engine/Source/Runtime/Renderer/Private/GPUScene.cpp b/Engine/Source/Runtime/Renderer/Private/GPUScene.cpp
--- a/Engine/Source/Runtime/Renderer/Private/GPUScene.cpp
+++ b/Engine/Source/Runtime/Renderer/Private/GPUScene.cpp
@@ -99,7 +99,6 @@ static TAutoConsoleVariable<float> CVarGPUSceneDebugDrawRange(
 	ECVF_RenderThreadSafe
 );
 
-
 static int32 GGPUSceneAllowDeferredAllocatorMerges = 1;
 FAutoConsoleVariableRef CVarGPUSceneAllowDeferredAllocatorMerges(
 	TEXT("r.GPUScene.AllowDeferredAllocatorMerges"),
@@ -108,6 +107,14 @@ FAutoConsoleVariableRef CVarGPUSceneAllowDeferredAllocatorMerges(
 	ECVF_RenderThreadSafe
 );
 
+int32 GGPUSceneInstanceUploadViaCreate = 1;
+FAutoConsoleVariableRef CVarGPUSceneInstanceUploadViaCreate(
+	TEXT("r.GPUScene.InstanceUploadViaCreate"),
+	GGPUSceneInstanceUploadViaCreate,
+	TEXT("When uploading GPUScene InstanceData, upload via resource creation when the RHI supports it efficiently."),
+	ECVF_RenderThreadSafe
+);
+
 
 constexpr uint32 InstanceSceneDataNumArrays = FInstanceSceneShaderData::DataStrideInFloat4s;
 
@@ -1084,6 +1091,10 @@ void FGPUScene::UploadGeneral(FRHICommandListImmediate& RHICmdList, FScene *Scen
 			SCOPED_NAMED_EVENT(STAT_UpdateGPUScenePrimitives, FColor::Green);
 			QUICK_SCOPE_CYCLE_COUNTER(STAT_UpdateGPUScenePrimitives);
 
+			const bool bShouldUploadViaCreate = GRHISupportsEfficientUploadOnResourceCreation && GGPUSceneInstanceUploadViaCreate != 0;
+			PrimitiveUploadBuffer.SetUploadViaCreate(bShouldUploadViaCreate);
+			InstanceSceneUploadBuffer.SetUploadViaCreate(bShouldUploadViaCreate);
+
 			FUAVTransitionStateScopeHelper PrimitiveDataTransitionHelper(RHICmdList, BufferState.PrimitiveBuffer.UAV, ERHIAccess::Unknown, ERHIAccess::SRVMask);
 
 			{
