- Generalize one pass projection shadow mask to support up to 32 lights/pixel with 4bpp quantization
- Add some dither noise to both the SMRT result and the shadow mask lookup to minimize banding
- Fall back to a single sample VSM lookup (with a generous static bias) when overflowing the number of lights in one pass projection path
- Fix clamping issue with page dilation that was setting extraneous pages with point lights
- Fix SMRT issue with local lights jammed right next to geometry viewed at a distance
- Separate settings for page dilation for local and directional lights
- Add simple debug output for # lights in one pass projection
- Remove some dead code/parameters

diff --git a/Engine/Shaders/Private/ClusteredDeferredShadingPixelShader.usf b/Engine/Shaders/Private/ClusteredDeferredShadingPixelShader.usf
--- a/Engine/Shaders/Private/ClusteredDeferredShadingPixelShader.usf
+++ b/Engine/Shaders/Private/ClusteredDeferredShadingPixelShader.usf
@@ -38,7 +38,7 @@
 // Extracts out the simple lights into their own loop that is executed first, independently of shading model.
 #define ENABLE_SEPARATE_SIMPLE_LIGHT_LOOP 0
 
-Texture2D< uint > ShadowMaskBits;
+Texture2D< uint4 > ShadowMaskBits;
 
 #if USE_HAIR_LIGHTING
 uint				HairTransmittanceBufferMaxCount;
@@ -61,7 +61,7 @@ float4 GetLightGridLocalLighting(const FCulledLightsGridData InLightGridData, fl
 	const uint NumLocalLights = min(InLightGridData.NumLocalLights, GetNumLocalLights(EyeIndex));
 
 	int2 PixelPos = int2( ScreenUV.xy * View.BufferSizeAndInvSize.xy );
-	uint ShadowMask = ~ShadowMaskBits[ PixelPos ];
+	uint4 ShadowMask = ~ShadowMaskBits[ PixelPos ];
 	half SubsurfaceOpacity = GetSubsurfaceOpacityFromGbuffer( PixelPos );
 
 	LOOP
@@ -126,11 +126,22 @@ float4 GetLightGridLocalLighting(const FCulledLightsGridData InLightGridData, fl
 		float4 LightAttenuation = float4(1, 1, 1, 1);
 
 		LightData.ShadowedBits = 1;
-		// Be careful with the shift here as the shift amount is mod 32, so >>32 is a noop rather than zero
-		// We treat lights past the 10th as unshadowed
-		if (LocalLightListIndex < GetShadowMaskMaxLightCount())
+		if (LocalLightListIndex < GetPackedShadowMaskMaxLightCount())
 		{
-			LightAttenuation = UnpackShadowMask(ShadowMask, LocalLightListIndex).xxxx;
+			LightAttenuation = UnpackShadowMask(ShadowMask, LocalLightListIndex, Dither).xxxx;
+		}
+		else if (LocalLight.VirtualShadowMapId != INDEX_NONE)
+		{
+			// Single sample VSM lookup
+			// TODO: Normal bias? Screen space ray? Etc? Do we need a path that does more of that stuff
+			// unified but without SMRT? For now we'll just keep this as a simple lookup to avoid potential performance tradeoffs
+			// as this is meant to be the fallback path on overflowing tiles.
+			FVirtualShadowMapSampleResult VSMResult = SampleVirtualShadowMap(
+				LocalLight.VirtualShadowMapId,
+				WorldPosition,
+				30.0f // RayStartDistance TODO... maybe at least make this a proportional screen offset like the contact shadow trace
+			);
+			LightAttenuation = VSMResult.ShadowFactor.xxxx;
 		}
 
 		if ( SubsurfaceOpacity < 1.0f )
@@ -156,10 +167,10 @@ float4 GetLightGridLocalLighting(const FCulledLightsGridData InLightGridData, fl
 
 			// Fetch precomputed transmittance
 			FHairTransmittanceMask TransmittanceMask = InitHairTransmittanceMask();
-			if (LocalLightListIndex < GetShadowMaskMaxLightCount())
+			if (LocalLightListIndex < GetPackedShadowMaskMaxLightCount())
 			{
 				const uint TransmittanceSampleIndex = SampleIndex + LocalLightListIndex * TotalSampleCount;
-				const uint TransmittanceSampleMaxCount = TotalSampleCount * GetShadowMaskMaxLightCount();
+				const uint TransmittanceSampleMaxCount = TotalSampleCount * GetPackedShadowMaskMaxLightCount();
 				if (TransmittanceSampleIndex < TransmittanceSampleMaxCount)
 				{
 					TransmittanceMask = UnpackTransmittanceMask(HairTransmittanceBuffer[TransmittanceSampleIndex]);		
diff --git a/Engine/Shaders/Private/HairStrands/HairStrandsDeepTransmittanceMask.usf b/Engine/Shaders/Private/HairStrands/HairStrandsDeepTransmittanceMask.usf
--- a/Engine/Shaders/Private/HairStrands/HairStrandsDeepTransmittanceMask.usf
+++ b/Engine/Shaders/Private/HairStrands/HairStrandsDeepTransmittanceMask.usf
@@ -444,7 +444,7 @@ void MainCS(uint2 DispatchThreadId : SV_DispatchThreadID)
 	const float3 WorldPosition	= ToWorldPosition(UV, SceneDepth);
 	const float PixelRadius		= SceneDepth * VirtualVoxel.HairCoveragePixelRadiusAtDepth1;
 
-	const uint PackedShadowMask = ~ShadowMaskBitsTexture[PixelCoord];
+	const uint4 PackedShadowMask = ~ShadowMaskBitsTexture[PixelCoord];
 	const float3 Random = GetHairVoxelJitter(PixelCoord.xy, View.StateFrameIndexMod8, VirtualVoxel.JitterMode);
 
 	const uint2 LocalPixelPos = PixelCoord;
@@ -453,7 +453,7 @@ void MainCS(uint2 DispatchThreadId : SV_DispatchThreadID)
 	const uint GridLinearIndex = ComputeLightGridCellIndex(LocalPixelPos, SceneDepth, EyeIndex);
 	const FCulledLightsGridData CulledLightGridData = GetCulledLightsGrid(GridLinearIndex, EyeIndex);
 
-	const uint LightCount = min(GetShadowMaskMaxLightCount(), CulledLightGridData.NumLocalLights);
+	const uint LightCount = min(GetPackedShadowMaskMaxLightCount(), CulledLightGridData.NumLocalLights);
 
 	LOOP
 	for (uint LightIndex = 0; LightIndex < LightCount; ++LightIndex)
diff --git a/Engine/Shaders/Private/LightGridCommon.ush b/Engine/Shaders/Private/LightGridCommon.ush
--- a/Engine/Shaders/Private/LightGridCommon.ush
+++ b/Engine/Shaders/Private/LightGridCommon.ush
@@ -172,7 +172,7 @@ struct FLocalLightData
 	float4 SpotAnglesAndSourceRadiusPacked;
     float4 LightDirectionAndShadowMask;
     float4 LightTangentAndSoftSourceRadius;
-	float4 RectBarnDoor;
+	float4 RectBarnDoorAndVirtualShadowMapId;
 	/** Flag is true if clustered deferred is supported for this light. They are always first / together.*/
 	bool bClusteredDeferredSupported; 
 	/** Flag is true if it is a simple light. They are always first / together.*/
@@ -200,8 +200,8 @@ FLocalLightData GetLocalLightData(uint GridIndex, uint EyeIndex)
 		Result.LightDirectionAndShadowMask = ForwardLightData.ForwardLocalLightBuffer[LocalLightBaseIndex + 2];
 		Result.SpotAnglesAndSourceRadiusPacked = ForwardLightData.ForwardLocalLightBuffer[LocalLightBaseIndex + 3];
 		Result.LightTangentAndSoftSourceRadius = ForwardLightData.ForwardLocalLightBuffer[LocalLightBaseIndex + 4];
-		Result.RectBarnDoor = ForwardLightData.ForwardLocalLightBuffer[LocalLightBaseIndex + 5];
-		Result.VirtualShadowMapId = asint(ForwardLightData.ForwardLocalLightBuffer[LocalLightBaseIndex + 5].z);
+		Result.RectBarnDoorAndVirtualShadowMapId = ForwardLightData.ForwardLocalLightBuffer[LocalLightBaseIndex + 5];
+		Result.VirtualShadowMapId = int(Result.RectBarnDoorAndVirtualShadowMapId.z);
 #if INSTANCED_STEREO
 	}
 	else
@@ -215,8 +215,8 @@ FLocalLightData GetLocalLightData(uint GridIndex, uint EyeIndex)
 		Result.LightDirectionAndShadowMask = ForwardLightDataISR.ForwardLocalLightBuffer[LocalLightBaseIndex + 2];
 		Result.SpotAnglesAndSourceRadiusPacked = ForwardLightDataISR.ForwardLocalLightBuffer[LocalLightBaseIndex + 3];
 		Result.LightTangentAndSoftSourceRadius = ForwardLightDataISR.ForwardLocalLightBuffer[LocalLightBaseIndex + 4];
-		Result.RectBarnDoor = ForwardLightDataISR.ForwardLocalLightBuffer[LocalLightBaseIndex + 5];
-		Result.VirtualShadowMapId = asint(ForwardLightDataISR.ForwardLocalLightBuffer[LocalLightBaseIndex + 5].z);
+		Result.RectBarnDoorAndVirtualShadowMapId = ForwardLightDataISR.ForwardLocalLightBuffer[LocalLightBaseIndex + 5];
+		Result.VirtualShadowMapId = int(Result.RectBarnDoorAndVirtualShadowMapId.z);
 	}
 #endif
 
@@ -232,20 +232,46 @@ float4 UnpackShadowMapChannelMask(uint ShadowMapChannelMaskPacked)
 }
 
 /**
- * Helper to pack/unpack the shadow mask for cluster shading and virtual shadow map
- * This packing encode up to 10 lights into a 32bits uint
+ * Helpers to pack/unpack the shadow mask for cluster shading and virtual shadow map
+ * Currently hard-coded for 4 bits per light, up to 32 lights per pixel in a uint4
  */
-uint GetShadowMaskMaxLightCount()
+uint4 InitializePackedShadowMask()
 {
-	return 10;
+	return uint(0xffffffff).xxxx;
 }
 
-void PackShadowMask(inout uint OutShadowMask, float InShadowFactor, uint InLightIndex)
+uint GetPackedShadowMaskMaxLightCount()
 {
-	OutShadowMask ^= (~uint(round(InShadowFactor * 7.0)) & 7u) << (InLightIndex * 3u);
+	return VirtualShadowMap.PackedShadowMaskMaxLightCount;
 }
 
-float UnpackShadowMask(uint InShadowMask, uint InLightIndex)
+void PackShadowMask(inout uint4 InOutShadowMask, float InShadowFactor, uint InLightIndex)
 {
-	return ((InShadowMask >> (InLightIndex * 3u)) & 7u) / 7.0;
+	uint Value = ~uint(round(InShadowFactor * 15.0f)) & 15U;
+	uint Dword = InLightIndex / 8;
+
+	InOutShadowMask.x ^= (Dword == 0U) ? (Value << (InLightIndex      ) * 4U) : 0U;
+	InOutShadowMask.y ^= (Dword == 1U) ? (Value << (InLightIndex -  8U) * 4U) : 0U;
+	InOutShadowMask.z ^= (Dword == 2U) ? (Value << (InLightIndex - 16U) * 4U) : 0U;
+	InOutShadowMask.w ^= (Dword == 3U) ? (Value << (InLightIndex - 24U) * 4U) : 0U;
+
+	//OutShadowMask ^= (~uint(round(InShadowFactor * 7.0)) & 7u) << (InLightIndex * 3u);
+}
+
+float UnpackShadowMask(uint4 InShadowMask, uint InLightIndex)
+{	
+	uint Dword = InLightIndex / 8;
+	return ((InShadowMask[Dword] >> (InLightIndex - Dword*8U) * 4U) & 15U) / 15.0f;	
+	//return ((InShadowMask >> (InLightIndex * 3u)) & 7u) / 7.0;
+}
+
+// Unpack with dither to hide some of the quantization
+float UnpackShadowMask(uint4 InShadowMask, uint InLightIndex, float Dither)
+{	
+	float ShadowFactor = UnpackShadowMask(InShadowMask, InLightIndex);
+	if (ShadowFactor > 0.0f && ShadowFactor < 1.0f)
+	{
+		ShadowFactor = saturate(ShadowFactor + (Dither - 0.5f) * (1.0f / 16.0f));
+	}
+	return ShadowFactor;
 }
diff --git a/Engine/Shaders/Private/VirtualShadowMaps/PageManagement.usf b/Engine/Shaders/Private/VirtualShadowMaps/PageManagement.usf
--- a/Engine/Shaders/Private/VirtualShadowMaps/PageManagement.usf
+++ b/Engine/Shaders/Private/VirtualShadowMaps/PageManagement.usf
@@ -38,7 +38,8 @@ RWStructuredBuffer<uint> OutPageRequestFlags;
 StructuredBuffer<int> DirectionalLightIds;
 
 float ResolutionLodBiasLocal;
-float PageDilationBorderSize;
+float PageDilationBorderSizeDirectional;
+float PageDilationBorderSizeLocal;
 uint InputType;
 uint bCullBackfacingPixels;
 
@@ -108,7 +109,7 @@ uint GetMipLevelLocal(int VirtualShadowMapId, float3 WorldPosition, float SceneD
 	return CalcMipLevelLocal(Footprint);
 }
 
-void MarkPage(uint VirtualShadowMapId, uint MipLevel, float3 WorldPosition, float2 PageDilationOffset)
+void MarkPage(uint VirtualShadowMapId, uint MipLevel, float3 WorldPosition, bool bUsePageDilation, float2 PageDilationOffset)
 {
 	FVirtualShadowMapProjectionShaderData ProjectionData = GetVirtualShadowMapProjectionData(VirtualShadowMapId);
 	float4 ShadowUVz = mul(float4(WorldPosition + ProjectionData.ShadowPreViewTranslation, 1.0f), ProjectionData.TranslatedWorldToShadowUVMatrix);
@@ -126,21 +127,22 @@ void MarkPage(uint VirtualShadowMapId, uint MipLevel, float3 WorldPosition, floa
 
 	uint Flags = VSM_ALLOCATED_FLAG | VSM_NON_NANITE_FLAG;
 
+	uint MaxPageAddress = CalcLevelDimsPages(MipLevel) - 1U;
 	float2 PageAddressFloat = ShadowUVz.xy * CalcLevelDimsPages(MipLevel);
-	uint2 PageAddress = uint2(PageAddressFloat);
+	uint2 PageAddress = clamp(uint2(PageAddressFloat), 0U, MaxPageAddress);
 	uint PageOffset = CalcPageOffset(VirtualShadowMapId, MipLevel, PageAddress);
 	OutPageRequestFlags[PageOffset] = Flags;
 
 	// PageDilationBorderSize == 0 implies PageDilationOffset.xy == 0
-	if (PageDilationBorderSize > 0.0f)
+	if (bUsePageDilation)
 	{
-		uint2 PageAddress2 = uint2(PageAddressFloat + PageDilationOffset);
+		uint2 PageAddress2 = clamp(uint2(PageAddressFloat + PageDilationOffset), 0, MaxPageAddress);
 		uint PageOffset2 = CalcPageOffset(VirtualShadowMapId, MipLevel, PageAddress2);
 		if (PageOffset2 != PageOffset)
 		{
 			OutPageRequestFlags[PageOffset2] = Flags;
 		}
-		uint2 PageAddress3 = uint2(max(float2(0, 0), PageAddressFloat - PageDilationOffset));
+		uint2 PageAddress3 = clamp(uint2(PageAddressFloat - PageDilationOffset), 0, MaxPageAddress);
 		uint PageOffset3 = CalcPageOffset(VirtualShadowMapId, MipLevel, PageAddress3);
 		if (PageOffset3 != PageOffset)
 		{
@@ -215,38 +217,46 @@ void GeneratePageFlagsFromPixels(
 	// We don't need to to check all 8 adjacent pages; as long as there's at least a single pixel near the edge
 	// the adjacent one will get mapped. In practice only checking one diagonal seems to work fine and have minimal
 	// overhead.
-	float2 PageDilationOffset = PageDilationBorderSize * 
-		float2(GroupIndex & 1 ? 1.0f : -1.0f,
-			   GroupIndex & 2 ? 1.0f : -1.0f);
-	
+	const float2 PageDilationDither = float2(
+		(GroupIndex & 1) ? 1.0f : -1.0f,
+		(GroupIndex & 2) ? 1.0f : -1.0f);
+		
 	// Directional lights
-	for (uint Index = 0; Index < VirtualShadowMap.NumDirectionalLights; ++Index)
 	{
-		int ClipmapStartId = DirectionalLightIds[Index];
-		FVirtualShadowMapProjectionShaderData BaseProjectionData = GetVirtualShadowMapProjectionData(ClipmapStartId);
+		const bool bUsePageDilation = PageDilationBorderSizeDirectional > 0.0f;
+		const float2 PageDilationOffset = PageDilationBorderSizeDirectional * PageDilationDither;
 
-#if PERMUTATION_INPUT_TYPE != INPUT_TYPE_HAIRSTRANDS
-		// Backface test if requested
-		float3 LightDirection = -BaseProjectionData.TranslatedWorldToShadowViewMatrix._13_23_33;
-		if (bBackfaceCull && IsBackfaceToDirectionalLight(normalize(WorldNormal), normalize(LightDirection), BaseProjectionData.LightSourceRadius))
+		for (uint Index = 0; Index < VirtualShadowMap.NumDirectionalLights; ++Index)
 		{
-			continue;
-		}
+			int ClipmapStartId = DirectionalLightIds[Index];
+			FVirtualShadowMapProjectionShaderData BaseProjectionData = GetVirtualShadowMapProjectionData(ClipmapStartId);
+
+#if PERMUTATION_INPUT_TYPE != INPUT_TYPE_HAIRSTRANDS
+			// Backface test if requested
+			float3 LightDirection = -BaseProjectionData.TranslatedWorldToShadowViewMatrix._13_23_33;
+			if (bBackfaceCull && IsBackfaceToDirectionalLight(normalize(WorldNormal), normalize(LightDirection), BaseProjectionData.LightSourceRadius))
+			{
+				continue;
+			}
 #endif
 
-		// Possibly use a single shadow projection data for the whole clipmap with just scales/biases for
-		// each level (to handle snapping). For now we maintain one per level.
-		//int ClipmapLevel = CalcClipmapLevel(length(WorldPosition - ProjectionData.ClipmapWorldOrigin), ClipmapResolutionLodBias);
-		const int ClipmapLevel = CalcClipmapLevel(BaseProjectionData, WorldPosition);
-		int ClipmapIndex = max(0, ClipmapLevel - BaseProjectionData.ClipmapLevel);
-		if (ClipmapIndex < BaseProjectionData.ClipmapLevelCount)
-		{
-			MarkPage(ClipmapStartId + ClipmapIndex, 0, WorldPosition, PageDilationOffset);
+			// Possibly use a single shadow projection data for the whole clipmap with just scales/biases for
+			// each level (to handle snapping). For now we maintain one per level.
+			//int ClipmapLevel = CalcClipmapLevel(length(WorldPosition - ProjectionData.ClipmapWorldOrigin), ClipmapResolutionLodBias);
+			const int ClipmapLevel = CalcClipmapLevel(BaseProjectionData, WorldPosition);
+			int ClipmapIndex = max(0, ClipmapLevel - BaseProjectionData.ClipmapLevel);
+			if (ClipmapIndex < BaseProjectionData.ClipmapLevelCount)
+			{
+				MarkPage(ClipmapStartId + ClipmapIndex, 0, WorldPosition, bUsePageDilation, PageDilationOffset);
+			}
 		}
 	}
 
 	// Local lights
 	{
+		const bool bUsePageDilation = PageDilationBorderSizeLocal > 0.0f;
+		const float2 PageDilationOffset = PageDilationBorderSizeLocal * PageDilationDither;
+
 		uint EyeIndex = 0; // ??
 		uint3 GridCoordinate = ComputeLightGridCellCoordinate(PixelPos, SceneDepth, EyeIndex);
 		uint GridLinearIndex = ComputeLightGridCellIndex(GridCoordinate, EyeIndex);
@@ -281,11 +291,11 @@ void GeneratePageFlagsFromPixels(
 				bool bSpotLight = LightData.SpotAnglesAndSourceRadiusPacked.x > -2.0f;
 				if( !bSpotLight )
 				{
-					VirtualShadowMapId += VirtualShadowMapGetCubeFace( -ToLight );
+					VirtualShadowMapId += VirtualShadowMapGetCubeFace(-ToLight);
 				}
 
 				uint MipLevel = GetMipLevelLocal(VirtualShadowMapId, WorldPosition, SceneDepth);
-				MarkPage(VirtualShadowMapId, MipLevel, WorldPosition, PageDilationOffset);
+				MarkPage(VirtualShadowMapId, MipLevel, WorldPosition, bUsePageDilation, PageDilationOffset);
 			}
 		}
 	}
diff --git a/Engine/Shaders/Private/VirtualShadowMaps/ProjectionDirectional.ush b/Engine/Shaders/Private/VirtualShadowMaps/ProjectionDirectional.ush
--- a/Engine/Shaders/Private/VirtualShadowMaps/ProjectionDirectional.ush
+++ b/Engine/Shaders/Private/VirtualShadowMaps/ProjectionDirectional.ush
@@ -210,7 +210,7 @@ FVirtualShadowMapSampleResult TraceDirectional(
 	const half Noise,
 	bool bCullBackface,
 	half3 WorldNormal,
-	float TexelDitherScale)
+	half TexelDitherScale)
 {
 	float3 WorldPosition = TranslatedWorldPosition - LWCHackToFloat(PrimaryView.PreViewTranslation);
 	float3 ViewPosition = mul(float4(TranslatedWorldPosition, 1.0f), View.TranslatedWorldToView).xyz;
@@ -284,7 +284,7 @@ FVirtualShadowMapSampleResult TraceDirectional(
 		}
 		uint RayCount = min(i + 1, MaxRayCount);		// break vs regular for loop exit
 
-		Result.ShadowFactor = float(RayMissCount) / float(RayCount);
+		Result.ShadowFactor = SMRTComputeRayMissFactor(RayMissCount, RayCount, MaxRayCount, Noise);;
 		Result.OccluderDistance = -1.0f;	// TODO?
 		Result.RayCount = RayCount;
 	}
diff --git a/Engine/Shaders/Private/VirtualShadowMaps/ProjectionSpot.ush b/Engine/Shaders/Private/VirtualShadowMaps/ProjectionSpot.ush
--- a/Engine/Shaders/Private/VirtualShadowMaps/ProjectionSpot.ush
+++ b/Engine/Shaders/Private/VirtualShadowMaps/ProjectionSpot.ush
@@ -269,7 +269,6 @@ bool IsBackfaceToLocalLight(float3 Position, half3 Normal, float3 LightPosition,
 FVirtualShadowMapSampleResult TraceLocalLight(
 	int VirtualShadowMapId,
 	FLightShaderParameters Light,
-	uint LightType,
 	uint2 PixelPos,
 	const float SceneDepth,
 	float3 TranslatedWorldPosition,
@@ -283,6 +282,8 @@ FVirtualShadowMapSampleResult TraceLocalLight(
 {
 	const float3 WorldPosition = TranslatedWorldPosition - LWCHackToFloat(PrimaryView.PreViewTranslation);
 
+	uint LightType = Light.SpotAngles.x > -2.0f ? LIGHT_TYPE_SPOT : LIGHT_TYPE_POINT;
+
 	half3 ConeAxis;
 	half ConeSin;
 	half DistToLight;	
@@ -292,7 +293,6 @@ FVirtualShadowMapSampleResult TraceLocalLight(
 	Result.bValid = true;
 	Result.ShadowFactor = 0.0f;
 
-	half RayMissFactor = 0;
 	if (!bBackface || !bCullBackface)
 	{
 		float StepOffset = Noise;
@@ -300,6 +300,7 @@ FVirtualShadowMapSampleResult TraceLocalLight(
 		uint TileID = MortonEncode( PixelPos );
 
 		uint i = 0;
+		uint RayMissCount = 0;
 		for( ; i < MaxRayCount; i++ )
 		{
 			RandomSequence RandSequence;
@@ -315,9 +316,7 @@ FVirtualShadowMapSampleResult TraceLocalLight(
 			float4 RandSample = SobolSampler( RandSequence.SampleIndex, RandSequence.SampleSeed, 8 );
 #endif
 
-			float2 E = RandSample.xy;
-
-			half2 LightUV = UniformSampleDiskConcentricApprox( E ).xy;
+			half2 LightUV = UniformSampleDiskConcentricApprox( RandSample.xy ).xy;
 			LightUV *= ConeSin;
 
 			half SinTheta2 = dot( LightUV, LightUV );
@@ -331,10 +330,11 @@ FVirtualShadowMapSampleResult TraceLocalLight(
 			// where the projection becomes extreme.
 			
 			// Clip ray to max angle to limit penumbra sizes
-			float ClampedLength = DistToLight * 0.75 * saturate( 1.5 / ( CosTheta + CotMaxRayAngleFromLight * SinTheta ) );
+			float ClampedEnd = DistToLight * 0.75 * saturate( 1.5 / ( CosTheta + CotMaxRayAngleFromLight * SinTheta ) );
+			float Start = min(RayStartOffset, ClampedEnd - 1e-6f);
 
-			float3 RayStart = WorldPosition + Dir * RayStartOffset;
-			float3 RayEnd = WorldPosition + Dir * ClampedLength;
+			float3 RayStart = WorldPosition + Dir * Start;
+			float3 RayEnd = WorldPosition + Dir * ClampedEnd;
 
 			FSMRTResult SMRTResult;
 			FSMRTSingleRayState RayState;	// For debug output
@@ -359,7 +359,7 @@ FVirtualShadowMapSampleResult TraceLocalLight(
 					RayState);
 			}
 
-			RayMissFactor += SMRTResult.bValidHit ? 0 : 1;
+			RayMissCount += SMRTResult.bValidHit ? 0 : 1;
 			
 			// Debug output (DCE'd if not used)
 			Result.ClipmapIndexOrMipLevel = RayState.MipLevel;
@@ -371,18 +371,20 @@ FVirtualShadowMapSampleResult TraceLocalLight(
 			if( i == 0 )
 			{
 				bool bHit = SMRTResult.bValidHit;
-				if( bHit )
-				{
-				#if 0
-					//float WorldBluringRadius = ComputeLightSampleWorldBluringRadius( TranslatedWorldPosition, LIGHT_TYPE_SPOT, Light, SMRTResult.OccluderDistance );
-					//float PixelBluringRadius = ProjectWorldDistanceToPixelDistance( SceneDepth, WorldBluringRadius );
-					float WorldBluringRadius = Light.SourceRadius * SMRTResult.OccluderDistance / max( DistToLight - SMRTResult.OccluderDistance, 0.001 );
-					float PixelBluringRadius = WorldBluringRadius * View.ViewSizeAndInvSize.x * 0.5 / ( SceneDepth * View.ClipToView[0][0] );
-					bHit = PixelBluringRadius > 2;
-				#else
-					bHit = SMRTResult.OccluderDistance * Light.SourceRadius * ( 0.25 * View.ViewSizeAndInvSize.y / ( SceneDepth * View.ClipToView[1][1] ) ) > max( DistToLight - Result.OccluderDistance, 0.001 );
+				#if 0			// Disabled heuristic penumbra size path for now as it creates some speckle artifacts near the edges of sharp shadows
+					if( bHit )
+					{
+					#if 0
+						//float WorldBluringRadius = ComputeLightSampleWorldBluringRadius( TranslatedWorldPosition, LIGHT_TYPE_SPOT, Light, SMRTResult.OccluderDistance );
+						//float PixelBluringRadius = ProjectWorldDistanceToPixelDistance( SceneDepth, WorldBluringRadius );
+						float WorldBluringRadius = Light.SourceRadius * SMRTResult.OccluderDistance / max( DistToLight - SMRTResult.OccluderDistance, 0.001 );
+						float PixelBluringRadius = WorldBluringRadius * View.ViewSizeAndInvSize.x * 0.5 / ( SceneDepth * View.ClipToView[0][0] );
+						bHit = PixelBluringRadius > 2;
+					#else
+						bHit = SMRTResult.OccluderDistance * Light.SourceRadius * ( 0.25 * View.ViewSizeAndInvSize.y / ( SceneDepth * View.ClipToView[1][1] ) ) > max( DistToLight - Result.OccluderDistance, 0.001 );
+					#endif
+					}
 				#endif
-				}
 
 				// All lanes miss or hit something close enough to not be noisy.
 				bool bAllLanesMiss = WaveActiveAllTrue( !bHit );
@@ -394,7 +396,8 @@ FVirtualShadowMapSampleResult TraceLocalLight(
 			else if( i == 1 )
 			{
 				// After 2 iterations and all have hit, must be in umbra.
-				bool bAllLanesHit = WaveActiveAllTrue( RayMissFactor == 0 );
+				// TODO: Modify this heuristic based on ray counts to avoid blocky artifacts in near-umbra
+				bool bAllLanesHit = WaveActiveAllTrue( RayMissCount == 0 );
 				if( bAllLanesHit )
 				{
 					break;
@@ -405,9 +408,8 @@ FVirtualShadowMapSampleResult TraceLocalLight(
 			StepOffset = ( EvolveSobolSeed( Seed ) >> 8 ) * 5.96046447754e-08; // * 2^-24;
 		}
 		uint RayCount = min(i + 1, MaxRayCount);		// break vs regular for loop exit
-		RayMissFactor /= float(RayCount);
 
-		Result.ShadowFactor = RayMissFactor;
+		Result.ShadowFactor = SMRTComputeRayMissFactor(RayMissCount, RayCount, MaxRayCount, Noise);
 		Result.RayCount = RayCount;
 	}
 
diff --git a/Engine/Shaders/Private/VirtualShadowMaps/SMRTCommon.ush b/Engine/Shaders/Private/VirtualShadowMaps/SMRTCommon.ush
--- a/Engine/Shaders/Private/VirtualShadowMaps/SMRTCommon.ush
+++ b/Engine/Shaders/Private/VirtualShadowMaps/SMRTCommon.ush
@@ -1,7 +1,7 @@
 // Copyright Epic Games, Inc. All Rights Reserved.
 
 /*=============================================================================
-SMRTTemplateCommon.ush: 
+SMRTCommon.ush: 
 =============================================================================*/
 
 #pragma once
@@ -20,3 +20,16 @@ struct FSMRTResult
 	float OccluderDistance;
 	int SampleMissCount;
 };
+
+float SMRTComputeRayMissFactor(uint RayMissCount, uint RayCount, uint MaxRayCount, half Noise)
+{
+	float RayMissFactor = float(RayMissCount) / float(RayCount);
+
+	// Add some simple dither to break up banding
+	if (RayMissFactor > 0.0f && RayMissFactor < 1.0f)
+	{
+		RayMissFactor += (Noise - 0.5f) * float(1.0f / MaxRayCount);
+	}
+
+	return RayMissFactor;
+}
diff --git a/Engine/Shaders/Private/VirtualShadowMaps/VirtualShadowMapProjection.usf b/Engine/Shaders/Private/VirtualShadowMaps/VirtualShadowMapProjection.usf
--- a/Engine/Shaders/Private/VirtualShadowMaps/VirtualShadowMapProjection.usf
+++ b/Engine/Shaders/Private/VirtualShadowMaps/VirtualShadowMapProjection.usf
@@ -140,7 +140,7 @@ float SMRTCotMaxRayAngleFromLight;	// Spot/point lights
 float SMRTTexelDitherScale;			// Currently only directional lights
 
 // One pass projection
-RWTexture2D< uint > RWShadowMaskBits;
+RWTexture2D< uint4 > RWShadowMaskBits;
 
 // Single light per pass
 // Light parameters loaded via GetRootLightShaderParameters();
@@ -167,20 +167,20 @@ FVirtualShadowMapSampleResult ProjectLight(
 	float ContactShadowLengthWorld,
 	float3 TranslatedWorldPosition,
 	half3 GBufferNormal,
-	half SubsurfaceOpacity,
 	const half Noise)
 {
+	const uint ShadingModelID = GetGBufferDataUint(PixelPos).ShadingModelID;
+	const bool bIsHairInput = InputType == INPUT_TYPE_HAIRSTRANDS;
+
 	const float3 BaseWorldPosition = TranslatedWorldPosition - LWCHackToFloat(PrimaryView.PreViewTranslation);
 	// NOTE: If PrimaryView.PreViewTranslation is guaranteed to be the camera position we can simplify
 	const float DistanceToCamera = length(TranslatedWorldPosition - View.TranslatedWorldCameraOrigin);
 	const float NormalBiasLength = max(0.01f, NormalBias * DistanceToCamera / View.ViewToClip._11);
 
 #if DIRECTIONAL_LIGHT
-	uint LightType = LIGHT_TYPE_DIRECTIONAL;
 	float3 L = Light.Direction;
 	bool bInLightRegion = true;
-#else
-	uint LightType = Light.SpotAngles.x > -2.0f ? LIGHT_TYPE_SPOT : LIGHT_TYPE_POINT;
+#else	
 	float3 ToLight = Light.Position - BaseWorldPosition;
 	float d2 = dot( ToLight, ToLight );
 	float InvDist = rsqrt( d2 );
@@ -193,16 +193,14 @@ FVirtualShadowMapSampleResult ProjectLight(
 
 	FVirtualShadowMapSampleResult Result = InitVirtualShadowMapSampleResult();
 
-	const uint ShadingModelID = GetGBufferDataUint(PixelPos).ShadingModelID;
-
-	const bool bIsHair = InputType == INPUT_TYPE_HAIRSTRANDS;
-	const bool bBackfaceCull = (bCullBackfacingPixels != 0) && !bIsHair && !IsSubsurfaceModel(ShadingModelID);
-	const bool bValidPixel = bIsHair || ShadingModelID != SHADINGMODELID_UNLIT;
+	const bool bValidPixel = bIsHairInput || ShadingModelID != SHADINGMODELID_UNLIT;
 
 	BRANCH
 	if (bInLightRegion && bValidPixel)
 	{
-		const half3 WorldNormal = (bIsHair || ShadingModelID == SHADINGMODELID_HAIR) ? L : GBufferNormal; 
+		const bool bBackfaceCull = (bCullBackfacingPixels != 0) && !bIsHairInput && !IsSubsurfaceModel(ShadingModelID);
+
+		const half3 WorldNormal = (bIsHairInput || ShadingModelID == SHADINGMODELID_HAIR) ? L : GBufferNormal;
 		TranslatedWorldPosition += WorldNormal * NormalBiasLength;
 
 		if (SMRTRayCount > 0)
@@ -226,7 +224,6 @@ FVirtualShadowMapSampleResult ProjectLight(
 			Result = TraceLocalLight(
 				VirtualShadowMapId,
 				Light,
-				LightType,
 				PixelPos,
 				SceneDepth,
 				TranslatedWorldPosition,
@@ -250,7 +247,7 @@ FVirtualShadowMapSampleResult ProjectLight(
 
 		// Do not run contact shadow when computing the hair shadow mask (i.e. shadow mask applied on hair, has the scene 
 		// depth buffer contains fully opaque depth, which create false positive intersection resulting in wrong self shadowing)
-		if (!bIsHair)
+		if (!bIsHairInput)
 		{
 			if (ContactShadowLengthWorld > 0.0f)
 			{
@@ -265,7 +262,7 @@ FVirtualShadowMapSampleResult ProjectLight(
 
 		// Hair strands voxel traversal to apply hair shadow on top of opaque geometry
 		#if HAS_HAIR_STRANDS
-		if (!bIsHair)
+		if (!bIsHairInput)
 		{
 			float3 RayStart = 0;
 			float3 RayEnd = 0;
@@ -356,12 +353,15 @@ FVirtualShadowMapSampleResult ProjectLight(
 	return Result;
 }
 
-void OutputDebug( uint2 PixelPos, FVirtualShadowMapSampleResult Result )
+void OutputDebug( uint2 PixelPos, FVirtualShadowMapSampleResult Result, uint LightCount )
 {
 	float3 Output = float3( 1, 0, 1 );
-		
+	bool bWriteOutput = false;
+
 	if ( Result.bValid )
 	{
+		bWriteOutput = true;
+
 		uint2 vPage = Result.VirtualTexelAddress.xy >> VSM_LOG2_PAGE_SIZE;
 		uint2 pPage = Result.PhysicalTexelAddress.xy >> VSM_LOG2_PAGE_SIZE;
 		bool bValidPageAddress = all( pPage < uint2(VirtualShadowMap.PhysicalPoolSizePages) );
@@ -403,9 +403,24 @@ void OutputDebug( uint2 PixelPos, FVirtualShadowMapSampleResult Result )
 
 			Output = lerp(Color, Color2, 0.35f);
 		}
+		else
+		{
+			bWriteOutput = false;
+		}
+	}
+
+#if ONE_PASS_PROJECTION
+	if ( DebugOutputType == 8 )	// Light count (only useful for one pass projection)
+	{
+		Output = GreenToRedTurbo( float( LightCount ) / float( 50 ) );
+		bWriteOutput = true;
 	}
+#endif
 
-	RWDebug[ PixelPos ] = float4( Output, 1.0f );
+	if (bWriteOutput)
+	{
+		RWDebug[PixelPos] = float4(Output, 1.0f);
+	}
 }
 
 
@@ -425,11 +440,11 @@ void VirtualShadowMapProjection(
 	}
 	
 	float DeviceZ = SceneTexturesStruct.SceneDepthTexture.Load( int3( PixelPos, 0 ) ).r;
-	const bool bIsHair = InputType == INPUT_TYPE_HAIRSTRANDS;
+	const bool bIsHairInput = InputType == INPUT_TYPE_HAIRSTRANDS;
 	#if HAS_HAIR_STRANDS
-	if (bIsHair)
+	if (bIsHairInput)
 	{
-		DeviceZ = HairStrands.HairOnlyDepthTexture.Load(int3(PixelPos, 0));
+		DeviceZ = HairStrands.HairOnlyDepthTexture.Load(int3(PixelPos, 0)).x;
 		if (DeviceZ == 0)
 		{
 			return;
@@ -445,10 +460,11 @@ void VirtualShadowMapProjection(
 	const half Noise = InterleavedGradientNoise( SvPosition.xy, View.StateFrameIndexMod8 );
 
 	const half3 GBufferNormal = GetGBufferDataUint( PixelPos, true ).WorldNormal;
-	const half SubsurfaceOpacity = bIsHair ? 1.0f : GetSubsurfaceOpacityFromGbuffer( PixelPos );
+	const half SubsurfaceOpacity = bIsHairInput ? 1.0f : GetSubsurfaceOpacityFromGbuffer( PixelPos );
 
 	FVirtualShadowMapSampleResult DebugResult = InitVirtualShadowMapSampleResult();
-	bool bWriteDebug = false;
+
+	uint LightCountDebug = 1;
 
 #if ONE_PASS_PROJECTION
 	uint EyeIndex = 0; // TODO: Instanced stereo
@@ -456,9 +472,10 @@ void VirtualShadowMapProjection(
 	const FCulledLightsGridData CulledLightGridData = GetCulledLightsGrid( GridLinearIndex, EyeIndex );
 
 	// We can only handle 10 lights in our output encoding right now, so no purpose in computing more
-	const uint LightCount = min(GetShadowMaskMaxLightCount(), CulledLightGridData.NumLocalLights );
+	LightCountDebug = CulledLightGridData.NumLocalLights;
+	uint LightCount = min(GetPackedShadowMaskMaxLightCount(), CulledLightGridData.NumLocalLights);
 
-	uint ShadowMask = 0xffffffff;
+	uint4 ShadowMask = InitializePackedShadowMask();
 
 	LOOP
 	for (uint Index = 0; Index < LightCount; ++Index)
@@ -478,15 +495,15 @@ void VirtualShadowMapProjection(
 				ContactShadowLengthWorld,
 				TranslatedWorldPosition,
 				GBufferNormal,
-				SubsurfaceOpacity,
 				Noise);
 
+			// NOTE: Subsurface opacity is handled in the clustered shading pass in the one pass projection path
+
 			PackShadowMask(ShadowMask, Result.ShadowFactor, Index);
 
 			if (DebugVirtualShadowMapId == VirtualShadowMapId)
 			{
 				DebugResult = Result;
-				bWriteDebug = true;
 			}
 		}
 	}
@@ -506,7 +523,6 @@ void VirtualShadowMapProjection(
 			ContactShadowLengthWorld,
 			TranslatedWorldPosition,
 			GBufferNormal,
-			SubsurfaceOpacity,
 			Noise);
 		
 		half SSSTransmission = Result.ShadowFactor;
@@ -523,15 +539,14 @@ void VirtualShadowMapProjection(
 		if (DebugVirtualShadowMapId == VirtualShadowMapId)
 		{
 			DebugResult = Result;
-			bWriteDebug = true;
 		}
 	}
 #endif // ONE_PASS_PROJECTION
 
 #if DEBUG_OUTPUT
-	if (bWriteDebug && InputType == INPUT_TYPE_GBUFFER)
+	if (InputType == INPUT_TYPE_GBUFFER)
 	{
-		OutputDebug( PixelPos, DebugResult );
+		OutputDebug( PixelPos, DebugResult, LightCountDebug );
 	}
 #endif // DEBUG_OUTPUT
 }
diff --git a/Engine/Shaders/Private/VolumetricCloud.usf b/Engine/Shaders/Private/VolumetricCloud.usf
--- a/Engine/Shaders/Private/VolumetricCloud.usf
+++ b/Engine/Shaders/Private/VolumetricCloud.usf
@@ -1029,8 +1029,8 @@ void MainCommon(in FMaterialPixelParameters MaterialParameters, in float4 SvPosi
 						LightData.Tangent = LocalLight.LightTangentAndSoftSourceRadius.xyz;
 						LightData.SoftSourceRadius = LocalLight.LightTangentAndSoftSourceRadius.w;
 						LightData.bInverseSquared = LightData.FalloffExponent == 0;
-						LightData.RectLightBarnCosAngle = LocalLight.RectBarnDoor.x;
-						LightData.RectLightBarnLength = LocalLight.RectBarnDoor.y;
+						LightData.RectLightBarnCosAngle = LocalLight.RectBarnDoorAndVirtualShadowMapId.x;
+						LightData.RectLightBarnLength = LocalLight.RectBarnDoorAndVirtualShadowMapId.y;
 						LightData.bRadialLight = true;
 						LightData.bSpotLight = LightData.SpotAngles.x > -2.0f;
 						LightData.bRectLight = LightData.RectLightBarnLength > -2.0f;
diff --git a/Engine/Shaders/Private/VolumetricFog.usf b/Engine/Shaders/Private/VolumetricFog.usf
--- a/Engine/Shaders/Private/VolumetricFog.usf
+++ b/Engine/Shaders/Private/VolumetricFog.usf
@@ -711,8 +711,8 @@ void LightScatteringCS(
 				LightData.Tangent = LocalLight.LightTangentAndSoftSourceRadius.xyz;
 				LightData.SoftSourceRadius = LocalLight.LightTangentAndSoftSourceRadius.w;
 				LightData.bInverseSquared = LightData.FalloffExponent == 0;
-				LightData.RectLightBarnCosAngle = LocalLight.RectBarnDoor.x;
-				LightData.RectLightBarnLength = LocalLight.RectBarnDoor.y;
+				LightData.RectLightBarnCosAngle = LocalLight.RectBarnDoorAndVirtualShadowMapId.x;
+				LightData.RectLightBarnLength = LocalLight.RectBarnDoorAndVirtualShadowMapId.y;
 				LightData.bRadialLight = true;
 				LightData.bSpotLight = LightData.SpotAngles.x > -2.0f;
 				LightData.bRectLight = LightData.RectLightBarnLength > -2.0f;
diff --git a/Engine/Source/Runtime/RenderCore/Private/RenderCore.cpp b/Engine/Source/Runtime/RenderCore/Private/RenderCore.cpp
--- a/Engine/Source/Runtime/RenderCore/Private/RenderCore.cpp
+++ b/Engine/Source/Runtime/RenderCore/Private/RenderCore.cpp
@@ -158,7 +158,7 @@ DEFINE_STAT(STAT_LightRendering);
 
 DEFINE_STAT(STAT_NumShadowedLights);
 DEFINE_STAT(STAT_NumLightFunctionOnlyLights);
-DEFINE_STAT(STAT_NumUnshadowedLights);
+DEFINE_STAT(STAT_NumBatchedLights);
 DEFINE_STAT(STAT_NumLightsInjectedIntoTranslucency);
 DEFINE_STAT(STAT_NumLightsUsingTiledDeferred);
 DEFINE_STAT(STAT_NumLightsUsingSimpleTiledDeferred);
diff --git a/Engine/Source/Runtime/RenderCore/Public/RenderCore.h b/Engine/Source/Runtime/RenderCore/Public/RenderCore.h
--- a/Engine/Source/Runtime/RenderCore/Public/RenderCore.h
+++ b/Engine/Source/Runtime/RenderCore/Public/RenderCore.h
@@ -107,7 +107,7 @@ DECLARE_DWORD_COUNTER_STAT_EXTERN(TEXT("Num using standard deferred"),STAT_NumLi
 DECLARE_DWORD_COUNTER_STAT_EXTERN(TEXT("Num using tiled deferred"),STAT_NumLightsUsingTiledDeferred,STATGROUP_LightRendering, RENDERCORE_API);
 DECLARE_DWORD_COUNTER_STAT_EXTERN(TEXT("Num using simple lighting model"),STAT_NumLightsUsingSimpleTiledDeferred,STATGROUP_LightRendering, RENDERCORE_API);
 DECLARE_DWORD_COUNTER_STAT_EXTERN(TEXT("Num injected into translucency"),STAT_NumLightsInjectedIntoTranslucency,STATGROUP_LightRendering, RENDERCORE_API);
-DECLARE_DWORD_COUNTER_STAT_EXTERN(TEXT("Num unshadowed"),STAT_NumUnshadowedLights,STATGROUP_LightRendering, RENDERCORE_API);
+DECLARE_DWORD_COUNTER_STAT_EXTERN(TEXT("Num batched"),STAT_NumBatchedLights,STATGROUP_LightRendering, RENDERCORE_API);
 DECLARE_DWORD_COUNTER_STAT_EXTERN(TEXT("Num light function only"),STAT_NumLightFunctionOnlyLights,STATGROUP_LightRendering, RENDERCORE_API);
 DECLARE_DWORD_COUNTER_STAT_EXTERN(TEXT("Num shadowed"),STAT_NumShadowedLights,STATGROUP_LightRendering, RENDERCORE_API);
 
diff --git a/Engine/Source/Runtime/Renderer/Private/BasePassRendering.h b/Engine/Source/Runtime/Renderer/Private/BasePassRendering.h
--- a/Engine/Source/Runtime/Renderer/Private/BasePassRendering.h
+++ b/Engine/Source/Runtime/Renderer/Private/BasePassRendering.h
@@ -51,7 +51,7 @@ public:
 	FVector4f LightDirectionAndShadowMapChannelMask;
 	FVector4f SpotAnglesAndSourceRadiusPacked;
 	FVector4f LightTangentAndSoftSourceRadius;
-	FVector4f RectBarnDoor;
+	FVector4f RectBarnDoorAndVirtualShadowMapId;
 };
 
 struct FForwardBasePassTextures
diff --git a/Engine/Source/Runtime/Renderer/Private/ClusteredDeferredShadingPass.cpp b/Engine/Source/Runtime/Renderer/Private/ClusteredDeferredShadingPass.cpp
--- a/Engine/Source/Runtime/Renderer/Private/ClusteredDeferredShadingPass.cpp
+++ b/Engine/Source/Runtime/Renderer/Private/ClusteredDeferredShadingPass.cpp
@@ -252,7 +252,7 @@ void FDeferredShadingSceneRenderer::AddClusteredDeferredShadingPass(
 
 			if (HairStrands::HasViewHairStrandsData(View))
 			{
-				FHairStrandsTransmittanceMaskData TransmittanceMask = RenderHairStrandsOnePassTransmittanceMask(GraphBuilder, View, HairStrandsShadowMaskBits);
+				FHairStrandsTransmittanceMaskData TransmittanceMask = RenderHairStrandsOnePassTransmittanceMask(GraphBuilder, View, HairStrandsShadowMaskBits, VirtualShadowMapArray);
 				InternalAddClusteredDeferredShadingPass(
 					GraphBuilder,
 					View,
diff --git a/Engine/Source/Runtime/Renderer/Private/DeferredShadingRenderer.cpp b/Engine/Source/Runtime/Renderer/Private/DeferredShadingRenderer.cpp
--- a/Engine/Source/Runtime/Renderer/Private/DeferredShadingRenderer.cpp
+++ b/Engine/Source/Runtime/Renderer/Private/DeferredShadingRenderer.cpp
@@ -2539,8 +2539,8 @@ void FDeferredShadingSceneRenderer::Render(FRDGBuilder& GraphBuilder)
 	}
 
 	CSV_CUSTOM_STAT(LightCount, All,  float(SortedLightSet.SortedLights.Num()), ECsvCustomStatOp::Set);
-	CSV_CUSTOM_STAT(LightCount, ShadowOff, float(SortedLightSet.AttenuationLightStart), ECsvCustomStatOp::Set);
-	CSV_CUSTOM_STAT(LightCount, ShadowOn, float(SortedLightSet.SortedLights.Num()) - float(SortedLightSet.AttenuationLightStart), ECsvCustomStatOp::Set);
+	CSV_CUSTOM_STAT(LightCount, Batched, float(SortedLightSet.UnbatchedLightStart), ECsvCustomStatOp::Set);
+	CSV_CUSTOM_STAT(LightCount, Unbatched, float(SortedLightSet.SortedLights.Num()) - float(SortedLightSet.UnbatchedLightStart), ECsvCustomStatOp::Set);
 
 	FCompositionLighting CompositionLighting(Views, SceneTextures, [this] (int32 ViewIndex)
 	{
diff --git a/Engine/Source/Runtime/Renderer/Private/HairStrands/HairStrandsTransmittance.cpp b/Engine/Source/Runtime/Renderer/Private/HairStrands/HairStrandsTransmittance.cpp
--- a/Engine/Source/Runtime/Renderer/Private/HairStrands/HairStrandsTransmittance.cpp
+++ b/Engine/Source/Runtime/Renderer/Private/HairStrands/HairStrandsTransmittance.cpp
@@ -183,6 +183,7 @@ class FHairStrandsVoxelTransmittanceMaskCS : public FGlobalShader
 		SHADER_PARAMETER_STRUCT_INCLUDE(ShaderPrint::FShaderParameters, ShaderPrintUniformBuffer)
 		SHADER_PARAMETER_RDG_UNIFORM_BUFFER(FHairStrandsViewUniformParameters, HairStrands)
 		SHADER_PARAMETER_RDG_UNIFORM_BUFFER(FVirtualVoxelParameters, VirtualVoxel)
+		SHADER_PARAMETER_STRUCT_INCLUDE(FVirtualShadowMapSamplingParameters, VirtualShadowMap)
 
 		SHADER_PARAMETER(float, LightRadius)
 		SHADER_PARAMETER(FVector3f, LightDirection)
@@ -224,7 +225,8 @@ static FRDGBufferRef AddHairStrandsVoxelTransmittanceMaskPass(
 	const FHairStrandsTransmittanceLightParams& Params,
 	const uint32 NodeGroupSize,
 	FRDGBufferRef IndirectArgsBuffer,
-	FRDGTextureRef ShadowMaskTexture)
+	FRDGTextureRef ShadowMaskTexture,
+	FVirtualShadowMapArray* VirtualShadowMapArray = nullptr)
 {
 	check(HairStrands::HasViewHairStrandsVoxelData(View));
 	check(NodeGroupSize == 64 || NodeGroupSize == 32);
@@ -239,6 +241,8 @@ static FRDGBufferRef AddHairStrandsVoxelTransmittanceMaskPass(
 	Parameters->OutTransmittanceMask = GraphBuilder.CreateUAV(OutBuffer, FHairStrandsTransmittanceMaskData::Format);
 	if (PassType == EHairTransmittancePassType::OnePass)
 	{
+		check(VirtualShadowMapArray != nullptr);
+		Parameters->VirtualShadowMap = VirtualShadowMapArray->GetSamplingParameters(GraphBuilder);
 		Parameters->ForwardLightData = View.ForwardLightingResources->ForwardLightDataUniformBuffer;
 		Parameters->RayMarchMaskTexture = nullptr;
 		Parameters->ShadowMaskBitsTexture = ShadowMaskTexture;
@@ -811,7 +815,8 @@ static FHairStrandsTransmittanceMaskData InternalRenderHairStrandsTransmittanceM
 FHairStrandsTransmittanceMaskData RenderHairStrandsOnePassTransmittanceMask(
 	FRDGBuilder& GraphBuilder,
 	const FViewInfo& View, 
-	FRDGTextureRef ShadowMaskBits)
+	FRDGTextureRef ShadowMaskBits,
+	FVirtualShadowMapArray& VirtualShadowMapArray)
 {
 	FHairStrandsTransmittanceMaskData Out;
 	if (HairStrands::HasViewHairStrandsData(View) && View.HairStrandsViewData.MacroGroupDatas.Num() > 0)
@@ -834,7 +839,8 @@ FHairStrandsTransmittanceMaskData RenderHairStrandsOnePassTransmittanceMask(
 				DummyParams,
 				View.HairStrandsViewData.VisibilityData.NodeGroupSize,
 				View.HairStrandsViewData.VisibilityData.NodeIndirectArg,
-				ShadowMaskBits);
+				ShadowMaskBits,
+				&VirtualShadowMapArray);
 		}
 	}
 	return Out;
diff --git a/Engine/Source/Runtime/Renderer/Private/HairStrands/HairStrandsTransmittance.h b/Engine/Source/Runtime/Renderer/Private/HairStrands/HairStrandsTransmittance.h
--- a/Engine/Source/Runtime/Renderer/Private/HairStrands/HairStrandsTransmittance.h
+++ b/Engine/Source/Runtime/Renderer/Private/HairStrands/HairStrandsTransmittance.h
@@ -12,6 +12,7 @@
 
 class FLightSceneInfo;
 class FViewInfo;
+class FVirtualShadowMapArray;
 
 struct FHairStrandsTransmittanceMaskData
 {
@@ -44,4 +45,5 @@ FHairStrandsTransmittanceMaskData RenderHairStrandsTransmittanceMask(
 FHairStrandsTransmittanceMaskData RenderHairStrandsOnePassTransmittanceMask(
 	FRDGBuilder& GraphBuilder,
 	const FViewInfo& View,
-	FRDGTextureRef ShadowMaskBits);
\ No newline at end of file
+	FRDGTextureRef ShadowMaskBits,
+	FVirtualShadowMapArray& VirtualShadowMapArray);
\ No newline at end of file
diff --git a/Engine/Source/Runtime/Renderer/Private/LightGridInjection.cpp b/Engine/Source/Runtime/Renderer/Private/LightGridInjection.cpp
--- a/Engine/Source/Runtime/Renderer/Private/LightGridInjection.cpp
+++ b/Engine/Source/Runtime/Renderer/Private/LightGridInjection.cpp
@@ -372,7 +372,7 @@ void FSceneRenderer::ComputeLightGrid(FRDGBuilder& GraphBuilder, bool bCullLight
 
 					LightData.SpotAnglesAndSourceRadiusPacked = FVector4f(-2, 1, 0, *(float*)&PackedWInt);
 					LightData.LightTangentAndSoftSourceRadius = FVector4f(1.0f, 0.0f, 0.0f, 0.0f);
-					LightData.RectBarnDoor = FVector4f(0, -2, 0, 0);
+					LightData.RectBarnDoorAndVirtualShadowMapId = FVector4f(0, -2, 0, 0);
 
 #if ENABLE_LIGHT_CULLING_VIEW_SPACE_BUILD_DATA
 					FVector4f ViewSpacePosAndRadius(FVector(View.ViewMatrices.GetViewMatrix().TransformPosition(SimpleLightPerViewData.Position)), SimpleLight.Radius);
@@ -458,7 +458,10 @@ void FSceneRenderer::ComputeLightGrid(FRDGBuilder& GraphBuilder, bool bCullLight
 
 						LightData.LightTangentAndSoftSourceRadius = FVector4f(LightParameters.Tangent, LightParameters.SoftSourceRadius);
 
-						LightData.RectBarnDoor = FVector4f(LightParameters.RectLightBarnCosAngle, LightParameters.RectLightBarnLength, *((float*)&VirtualShadowMapId), 0);
+						// NOTE: This cast of VirtualShadowMapId to float is not ideal, but bitcast has issues here with INDEX_NONE -> NaN
+						// and 32-bit floats have enough mantissa to cover all reasonable numbers here for now.
+						LightData.RectBarnDoorAndVirtualShadowMapId = FVector4f(LightParameters.RectLightBarnCosAngle, LightParameters.RectLightBarnLength, float(VirtualShadowMapId), 0);
+						checkSlow(int(LightData.RectBarnDoorAndVirtualShadowMapId.Z) == VirtualShadowMapId);
 
 						float VolumetricScatteringIntensity = LightProxy->GetVolumetricScatteringIntensity();
 
diff --git a/Engine/Source/Runtime/Renderer/Private/LightRendering.cpp b/Engine/Source/Runtime/Renderer/Private/LightRendering.cpp
--- a/Engine/Source/Runtime/Renderer/Private/LightRendering.cpp
+++ b/Engine/Source/Runtime/Renderer/Private/LightRendering.cpp
@@ -936,7 +936,6 @@ void FSceneRenderer::GatherAndSortLights(FSortedLightSetSceneInfo& OutSortedLigh
 					// These are not simple lights.
 					SortedLightInfo->SortKey.Fields.bIsNotSimpleLight = 1;
 
-
 					// tiled and clustered deferred lighting only supported for certain lights that don't use any additional features
 					// And also that are not directional (mostly because it does'nt make so much sense to insert them into every grid cell in the universe)
 					// In the forward case one directional light gets put into its own variables, and in the deferred case it gets a full-screen pass.
@@ -950,6 +949,11 @@ void FSceneRenderer::GatherAndSortLights(FSortedLightSetSceneInfo& OutSortedLigh
 						&& LightSceneInfoCompact.LightType != LightType_Directional
 						&& LightSceneInfoCompact.LightType != LightType_Rect;
 
+					// One pass projection is supported for lights with only virtual shadow maps
+					// TODO: Exclude lights that also have non-virtual shadow maps
+					bool bHasVirtualShadowMap = VisibleLightInfos[LightSceneInfo->Id].GetVirtualShadowMapId(&Views[ViewIndex]) != INDEX_NONE;
+					SortedLightInfo->SortKey.Fields.bDoesNotWriteIntoPackedShadowMask = !bTiledOrClusteredDeferredSupported || !bHasVirtualShadowMap;
+
 					SortedLightInfo->SortKey.Fields.bTiledDeferredNotSupported = !(bTiledOrClusteredDeferredSupported && LightSceneInfo->Proxy->IsTiledDeferredLightingSupported());
 
 					SortedLightInfo->SortKey.Fields.bClusteredDeferredNotSupported = !bTiledOrClusteredDeferredSupported;
@@ -1001,7 +1005,7 @@ void FSceneRenderer::GatherAndSortLights(FSortedLightSetSceneInfo& OutSortedLigh
 	OutSortedLights.SimpleLightsEnd = SortedLights.Num();
 	OutSortedLights.TiledSupportedEnd = SortedLights.Num();
 	OutSortedLights.ClusteredSupportedEnd = SortedLights.Num();
-	OutSortedLights.AttenuationLightStart = SortedLights.Num();
+	OutSortedLights.UnbatchedLightStart = SortedLights.Num();
 
 	// Iterate over all lights to be rendered and build ranges for tiled deferred and unshadowed lights
 	for (int32 LightIndex = 0; LightIndex < SortedLights.Num(); LightIndex++)
@@ -1032,9 +1036,9 @@ void FSceneRenderer::GatherAndSortLights(FSortedLightSetSceneInfo& OutSortedLigh
 
 		if( (bDrawShadows || bDrawLightFunction || bLightingChannels) && SortedLightInfo.SortKey.Fields.bClusteredDeferredNotSupported )
 		{
-			// Once we find a shadowed light, we can exit the loop, these lights should never support tiled deferred rendering either
+			// Once we find an unbatched shadowed light, we can exit the loop
 			check(SortedLightInfo.SortKey.Fields.bTiledDeferredNotSupported);
-			OutSortedLights.AttenuationLightStart = LightIndex;
+			OutSortedLights.UnbatchedLightStart = LightIndex;
 			break;
 		}
 	}
@@ -1042,7 +1046,7 @@ void FSceneRenderer::GatherAndSortLights(FSortedLightSetSceneInfo& OutSortedLigh
 	// Make sure no obvious things went wrong!
 	check(OutSortedLights.TiledSupportedEnd >= OutSortedLights.SimpleLightsEnd);
 	check(OutSortedLights.ClusteredSupportedEnd >= OutSortedLights.TiledSupportedEnd);
-	check(OutSortedLights.AttenuationLightStart >= OutSortedLights.ClusteredSupportedEnd);
+	check(OutSortedLights.UnbatchedLightStart >= OutSortedLights.ClusteredSupportedEnd);
 }
 
 /** Shader parameters to use when creating a RenderLight(...) pass. */
@@ -1126,7 +1130,7 @@ void FDeferredShadingSceneRenderer::RenderLights(
 
 	const FSimpleLightArray &SimpleLights = SortedLightSet.SimpleLights;
 	const TArray<FSortedLightSceneInfo, SceneRenderingAllocator> &SortedLights = SortedLightSet.SortedLights;
-	const int32 AttenuationLightStart = SortedLightSet.AttenuationLightStart;
+	const int32 UnbatchedLightStart = SortedLightSet.UnbatchedLightStart;
 	const int32 SimpleLightsEnd = SortedLightSet.SimpleLightsEnd;
 
 	FHairStrandsTransmittanceMaskData DummyTransmittanceMaskData;
@@ -1148,8 +1152,8 @@ void FDeferredShadingSceneRenderer::RenderLights(
 
 		if(ViewFamily.EngineShowFlags.DirectLighting)
 		{
-			RDG_EVENT_SCOPE(GraphBuilder, "NonShadowedLights");
-			INC_DWORD_STAT_BY(STAT_NumUnshadowedLights, AttenuationLightStart);
+			RDG_EVENT_SCOPE(GraphBuilder, "BatchedLights");
+			INC_DWORD_STAT_BY(STAT_NumBatchedLights, UnbatchedLightStart);
 
 			// Currently they have a special path anyway in case of standard deferred so always skip the simple lights
 			int32 StandardDeferredStart = SortedLightSet.SimpleLightsEnd;
@@ -1239,10 +1243,10 @@ void FDeferredShadingSceneRenderer::RenderLights(
 						RDG_EVENT_NAME("StandardDeferredLighting"),
 						PassParameters,
 						ERDGPassFlags::Raster,
-						[this, &View, &SortedLights, LightingChannelsTexture, StandardDeferredStart, AttenuationLightStart, PassParameters](FRHICommandList& RHICmdList)
+						[this, &View, &SortedLights, LightingChannelsTexture, StandardDeferredStart, UnbatchedLightStart, PassParameters](FRHICommandList& RHICmdList)
 					{
 						// Draw non-shadowed non-light function lights without changing render targets between them
-						for (int32 LightIndex = StandardDeferredStart; LightIndex < AttenuationLightStart; LightIndex++)
+						for (int32 LightIndex = StandardDeferredStart; LightIndex < UnbatchedLightStart; LightIndex++)
 						{
 							const FSortedLightSceneInfo& SortedLightInfo = SortedLights[LightIndex];
 							const FLightSceneInfo* const LightSceneInfo = SortedLightInfo.LightSceneInfo;
@@ -1266,7 +1270,7 @@ void FDeferredShadingSceneRenderer::RenderLights(
 					if (HairStrands::HasViewHairStrandsData(View))
 					{
 						// Draw non-shadowed non-light function lights without changing render targets between them
-						for (int32 LightIndex = StandardDeferredStart; LightIndex < AttenuationLightStart; LightIndex++)
+						for (int32 LightIndex = StandardDeferredStart; LightIndex < UnbatchedLightStart; LightIndex++)
 						{
 							const FSortedLightSceneInfo& SortedLightInfo = SortedLights[LightIndex];
 							const FLightSceneInfo* const LightSceneInfo = SortedLightInfo.LightSceneInfo;
@@ -1278,10 +1282,10 @@ void FDeferredShadingSceneRenderer::RenderLights(
 
 			if (GUseTranslucentLightingVolumes && GSupportsVolumeTextureRendering)
 			{
-				if (AttenuationLightStart)
+				if (UnbatchedLightStart)
 				{
 					// Inject non-shadowed, non-simple, non-light function lights in to the volume.
-					InjectTranslucencyLightingVolumeArray(GraphBuilder, Views, Scene, *this, TranslucencyLightingVolumeTextures, VisibleLightInfos, SortedLights, TInterval<int32>(SimpleLightsEnd, AttenuationLightStart));
+					InjectTranslucencyLightingVolumeArray(GraphBuilder, Views, Scene, *this, TranslucencyLightingVolumeTextures, VisibleLightInfos, SortedLights, TInterval<int32>(SimpleLightsEnd, UnbatchedLightStart));
 				}
 
 				if (SimpleLights.InstanceData.Num() > 0)
@@ -1308,7 +1312,7 @@ void FDeferredShadingSceneRenderer::RenderLights(
 		}
 
 		{
-			RDG_EVENT_SCOPE(GraphBuilder, "ShadowedLights");
+			RDG_EVENT_SCOPE(GraphBuilder, "UnbatchedLights");
 
 			const int32 DenoiserMode = CVarShadowUseDenoiser.GetValueOnRenderThread();
 
@@ -1358,7 +1362,7 @@ void FDeferredShadingSceneRenderer::RenderLights(
 			FRDGTextureRef SharedScreenShadowMaskSubPixelTexture = nullptr;
 
 			// Draw shadowed and light function lights
-			for (int32 LightIndex = AttenuationLightStart; LightIndex < SortedLights.Num(); LightIndex++)
+			for (int32 LightIndex = UnbatchedLightStart; LightIndex < SortedLights.Num(); LightIndex++)
 			{
 				const FSortedLightSceneInfo& SortedLightInfo = SortedLights[LightIndex];
 				const FLightSceneInfo& LightSceneInfo = *SortedLightInfo.LightSceneInfo;
@@ -1425,7 +1429,7 @@ void FDeferredShadingSceneRenderer::RenderLights(
 						if (
 							RHI_RAYTRACING &&
 							bWantsBatchedShadow &&
-							(PreprocessedShadowMaskTextures.Num() == 0 || !PreprocessedShadowMaskTextures[LightIndex - AttenuationLightStart]))
+							(PreprocessedShadowMaskTextures.Num() == 0 || !PreprocessedShadowMaskTextures[LightIndex - UnbatchedLightStart]))
 						{
 							RDG_EVENT_SCOPE(GraphBuilder, "ShadowBatch");
 							TStaticArray<IScreenSpaceDenoiser::FShadowVisibilityParameters, IScreenSpaceDenoiser::kMaxBatchSize> DenoisingQueue;
@@ -1467,7 +1471,7 @@ void FDeferredShadingSceneRenderer::RenderLights(
 									const FLightSceneInfo* LocalLightSceneInfo = DenoisingQueue[i].LightSceneInfo;
 
 									int32 LocalLightIndex = LightIndices[i];
-									FRDGTextureRef& RefDestination = PreprocessedShadowMaskTextures[LocalLightIndex - AttenuationLightStart];
+									FRDGTextureRef& RefDestination = PreprocessedShadowMaskTextures[LocalLightIndex - UnbatchedLightStart];
 									check(RefDestination == nullptr);
 									RefDestination = Outputs[i].Mask;
 									DenoisingQueue[i].LightSceneInfo = nullptr;
@@ -1565,7 +1569,7 @@ void FDeferredShadingSceneRenderer::RenderLights(
 									
 									if (HairStrands::HasViewHairStrandsData(View))
 									{
-										FRDGTextureRef& RefDestination = PreprocessedShadowMaskSubPixelTextures[LightBatchIndex - AttenuationLightStart];
+										FRDGTextureRef& RefDestination = PreprocessedShadowMaskSubPixelTextures[LightBatchIndex - UnbatchedLightStart];
 										check(RefDestination == nullptr);
 										RefDestination = SubPixelRayTracingShadowMaskTexture;
 									}
@@ -1602,7 +1606,7 @@ void FDeferredShadingSceneRenderer::RenderLights(
 								}
 								else
 								{
-									PreprocessedShadowMaskTextures[LightBatchIndex - AttenuationLightStart] = RayTracingShadowMaskTexture;
+									PreprocessedShadowMaskTextures[LightBatchIndex - UnbatchedLightStart] = RayTracingShadowMaskTexture;
 								}
 
 								// terminate batch if we filled a denoiser batch or hit our max light batch
@@ -1621,9 +1625,9 @@ void FDeferredShadingSceneRenderer::RenderLights(
 						}
 					} // end inline batched raytraced shadow
 
-					if (RHI_RAYTRACING && PreprocessedShadowMaskTextures.Num() > 0 && PreprocessedShadowMaskTextures[LightIndex - AttenuationLightStart])
+					if (RHI_RAYTRACING && PreprocessedShadowMaskTextures.Num() > 0 && PreprocessedShadowMaskTextures[LightIndex - UnbatchedLightStart])
 					{
-						const uint32 ShadowMaskIndex = LightIndex - AttenuationLightStart;
+						const uint32 ShadowMaskIndex = LightIndex - UnbatchedLightStart;
 						ScreenShadowMaskTexture = PreprocessedShadowMaskTextures[ShadowMaskIndex];
 						PreprocessedShadowMaskTextures[ShadowMaskIndex] = nullptr;
 
@@ -2583,7 +2587,7 @@ void FDeferredShadingSceneRenderer::RenderLightsForHair(
 {
 	const FSimpleLightArray &SimpleLights = SortedLightSet.SimpleLights;
 	const TArray<FSortedLightSceneInfo, SceneRenderingAllocator> &SortedLights = SortedLightSet.SortedLights;
-	const int32 AttenuationLightStart = SortedLightSet.AttenuationLightStart;
+	const int32 UnbatchedLightStart = SortedLightSet.UnbatchedLightStart;
 	const int32 SimpleLightsEnd = SortedLightSet.SimpleLightsEnd;
 
 	if (ViewFamily.EngineShowFlags.DirectLighting)
@@ -2598,7 +2602,7 @@ void FDeferredShadingSceneRenderer::RenderLightsForHair(
 			}
 
 			FHairStrandsTransmittanceMaskData DummyTransmittanceMaskData = CreateDummyHairStrandsTransmittanceMaskData(GraphBuilder, View.ShaderMap);
-			for (int32 LightIndex = AttenuationLightStart; LightIndex < SortedLights.Num(); LightIndex++)
+			for (int32 LightIndex = UnbatchedLightStart; LightIndex < SortedLights.Num(); LightIndex++)
 			{
 				const FSortedLightSceneInfo& SortedLightInfo = SortedLights[LightIndex];
 				const FLightSceneInfo& LightSceneInfo = *SortedLightInfo.LightSceneInfo;
diff --git a/Engine/Source/Runtime/Renderer/Private/LightSceneInfo.h b/Engine/Source/Runtime/Renderer/Private/LightSceneInfo.h
--- a/Engine/Source/Runtime/Renderer/Private/LightSceneInfo.h
+++ b/Engine/Source/Runtime/Renderer/Private/LightSceneInfo.h
@@ -72,21 +72,23 @@ struct FSortedLightSceneInfo
 		struct
 		{
 			// Note: the order of these members controls the light sort order!
-			// Currently bTiledDeferredNotSupported is the MSB and LightType is LSB
+			// Currently bClusteredDeferredNotSupported is the MSB and LightType is LSB
 			/** The type of light. */
 			uint32 LightType : LightType_NumBits;
 			/** Whether the light has a texture profile. */
 			uint32 bTextureProfile : 1;
 			/** Whether the light uses a light function. */
 			uint32 bLightFunction : 1;
-			/** Whether the light casts shadows. */
-			uint32 bShadowed : 1;
 			/** Whether the light uses lighting channels. */
 			uint32 bUsesLightingChannels : 1;
+			/** Whether the light casts shadows. */
+			uint32 bShadowed : 1;
 			/** Whether the light is NOT a simple light - they always support tiled/clustered but may want to be selected separately. */
 			uint32 bIsNotSimpleLight : 1;
 			/** True if the light doesn't support tiled deferred, logic is inverted so that lights that DO support tiled deferred will sort first in list */
 			uint32 bTiledDeferredNotSupported : 1;
+			/* We want to sort the lights that write into the packed shadow mask (when enabled) to the front of the list so we don't waste slots in the packed shadow mask. */
+			uint32 bDoesNotWriteIntoPackedShadowMask : 1;
 			/** 
 			 * True if the light doesn't support clustered deferred, logic is inverted so that lights that DO support clustered deferred will sort first in list 
 			 * Super-set of lights supporting tiled, so the tiled lights will end up in the first part of this range.
@@ -121,7 +123,8 @@ struct FSortedLightSceneInfo
 /** 
  * Stores info about sorted lights and ranges. 
  * The sort-key in FSortedLightSceneInfo gives rise to the following order:
- *  [SimpleLights,Tiled/Clustered,LightFunction/Shadow/LightChannels/TextureProfile]
+ *  [SimpleLights,Tiled/Clustered,UnbatchedLights]
+ * Note that some shadowed lights can be included in the clustered pass when virtual shadow maps and one pass projection are used.
  */
 struct FSortedLightSetSceneInfo
 {
@@ -130,7 +133,7 @@ struct FSortedLightSetSceneInfo
 	int ClusteredSupportedEnd;
 
 	/** First light with shadow map or */
-	int AttenuationLightStart;
+	int UnbatchedLightStart;
 
 	FSimpleLightArray SimpleLights;
 	TArray<FSortedLightSceneInfo, SceneRenderingAllocator> SortedLights;
diff --git a/Engine/Source/Runtime/Renderer/Private/SceneVisibility.cpp b/Engine/Source/Runtime/Renderer/Private/SceneVisibility.cpp
--- a/Engine/Source/Runtime/Renderer/Private/SceneVisibility.cpp
+++ b/Engine/Source/Runtime/Renderer/Private/SceneVisibility.cpp
@@ -4284,7 +4284,7 @@ void FSceneRenderer::ComputeViewVisibility(
 	// Allocate the visible light info.
 	if (Scene->Lights.GetMaxIndex() > 0)
 	{
-		VisibleLightInfos.AddZeroed(Scene->Lights.GetMaxIndex());
+		VisibleLightInfos.AddDefaulted(Scene->Lights.GetMaxIndex());
 	}
 
 	int32 NumPrimitives = Scene->Primitives.Num();
diff --git a/Engine/Source/Runtime/Renderer/Private/VirtualShadowMaps/VirtualShadowMapArray.cpp b/Engine/Source/Runtime/Renderer/Private/VirtualShadowMaps/VirtualShadowMapArray.cpp
--- a/Engine/Source/Runtime/Renderer/Private/VirtualShadowMaps/VirtualShadowMapArray.cpp
+++ b/Engine/Source/Runtime/Renderer/Private/VirtualShadowMaps/VirtualShadowMapArray.cpp
@@ -83,10 +83,18 @@ static TAutoConsoleVariable<float> CVarResolutionLodBiasLocal(
 	ECVF_RenderThreadSafe
 );
 
-static TAutoConsoleVariable<float> CVarPageDilationBorderSize(
-	TEXT("r.Shadow.Virtual.PageDilationBorderSize"),
+static TAutoConsoleVariable<float> CVarPageDilationBorderSizeDirectional(
+	TEXT("r.Shadow.Virtual.PageDilationBorderSizeDirectional"),
 	0.05f,
-	TEXT("If a screen pixel falls within this fraction of a page border, the adacent page will also be mapped.")
+	TEXT("If a screen pixel falls within this fraction of a page border for directional lights, the adacent page will also be mapped.")
+	TEXT("Higher values can reduce page misses at screen edges or disocclusions, but increase total page counts."),
+	ECVF_RenderThreadSafe
+);
+
+static TAutoConsoleVariable<float> CVarPageDilationBorderSizeLocal(
+	TEXT("r.Shadow.Virtual.PageDilationBorderSizeLocal"),
+	0.05f,
+	TEXT("If a screen pixel falls within this fraction of a page border for local lights, the adacent page will also be mapped.")
 	TEXT("Higher values can reduce page misses at screen edges or disocclusions, but increase total page counts."),
 	ECVF_RenderThreadSafe
 );
@@ -165,6 +173,14 @@ TAutoConsoleVariable<int32> CVarMergePhysicalUsingIndirect(
 	ECVF_RenderThreadSafe
 );
 
+TAutoConsoleVariable<int32> CVarVirtualShadowOnePassProjectionMaxLights(
+	TEXT("r.Shadow.Virtual.OnePassProjection.MaxLightsPerPixel"),
+	16,
+	TEXT("Maximum lights per pixel that get full filtering when using one pass projection and clustered shading.")
+	TEXT("Generally set to 8 (32bpp), 16 (64bpp) or 32 (128bpp). Lower values require less transient VRAM during the lighting pass."),
+	ECVF_RenderThreadSafe
+);
+
 #if !UE_BUILD_SHIPPING
 bool GDumpVSMLightNames = false;
 void DumpVSMLightNames()
@@ -350,6 +366,10 @@ void FVirtualShadowMapArray::Initialize(FRDGBuilder& GraphBuilder, FVirtualShado
 	UniformParameters.PhysicalPoolSize = FIntPoint( PhysicalX, PhysicalY );
 	UniformParameters.PhysicalPoolSizePages = FIntPoint( PhysicalPagesX, PhysicalPagesY );
 
+	// TODO: Parameterize this in a useful way; potentially modify it automatically
+	// when there are fewer lights in the scene and/or clustered shading settings differ.
+	UniformParameters.PackedShadowMaskMaxLightCount = FMath::Min(CVarVirtualShadowOnePassProjectionMaxLights.GetValueOnRenderThread(), 32);
+
 	// Reference dummy data in the UB initially
 	const uint32 DummyPageElement = 0xFFFFFFFF;
 	UniformParameters.PageTable = GraphBuilder.CreateSRV(GSystemTextures.GetDefaultStructuredBuffer(GraphBuilder, sizeof(DummyPageElement), DummyPageElement));
@@ -378,6 +398,26 @@ FVirtualShadowMapArray::~FVirtualShadowMapArray()
 	}
 }
 
+EPixelFormat FVirtualShadowMapArray::GetPackedShadowMaskFormat() const
+{
+	// TODO: Check if we're after any point that determines the format later too (light setup)
+	check(bInitialized);
+	// NOTE: Currently 4bpp/light
+	if (UniformParameters.PackedShadowMaskMaxLightCount <= 8)
+	{
+		return PF_R32_UINT;
+	}
+	else if (UniformParameters.PackedShadowMaskMaxLightCount <= 16)
+	{
+		return PF_R32G32_UINT;
+	}
+	else
+	{
+		check(UniformParameters.PackedShadowMaskMaxLightCount <= 32);
+		return PF_R32G32B32A32_UINT;
+	}
+}
+
 FIntPoint FVirtualShadowMapArray::GetPhysicalPoolSize() const
 {
 	check(bInitialized);
@@ -505,7 +545,8 @@ class FGeneratePageFlagsFromPixelsCS : public FVirtualPageManagementShader
 		SHADER_PARAMETER(uint32, NumDirectionalLightSmInds)
 		SHADER_PARAMETER(uint32, bPostBasePass)
 		SHADER_PARAMETER(float, ResolutionLodBiasLocal)
-		SHADER_PARAMETER(float, PageDilationBorderSize)
+		SHADER_PARAMETER(float, PageDilationBorderSizeDirectional)
+		SHADER_PARAMETER(float, PageDilationBorderSizeLocal)
 		SHADER_PARAMETER(uint32, bCullBackfacingPixels)
 	END_SHADER_PARAMETER_STRUCT()
 };
@@ -972,9 +1013,6 @@ void FVirtualShadowMapArray::BuildPageAllocations(
 	check(IsEnabled());
 	RDG_EVENT_SCOPE(GraphBuilder, "FVirtualShadowMapArray::BuildPageAllocation");
 
-	const float ResolutionLodBiasLocal = CVarResolutionLodBiasLocal.GetValueOnRenderThread();
-	const float PageDilationBorderSize = CVarPageDilationBorderSize.GetValueOnRenderThread();
-
 #if !UE_BUILD_SHIPPING
 	bool bDebugOutputEnabled = false;
 	if (GDumpVSMLightNames)
@@ -1164,8 +1202,9 @@ void FVirtualShadowMapArray::BuildPageAllocations(
 						PassParameters->OutPageRequestFlags = PageRequestFlagsUAV;
 						PassParameters->ForwardLightData = View.ForwardLightingResources->ForwardLightDataUniformBuffer;
 						PassParameters->DirectionalLightIds = GraphBuilder.CreateSRV(DirectionalLightIdsRDG);
-						PassParameters->ResolutionLodBiasLocal = ResolutionLodBiasLocal;
-						PassParameters->PageDilationBorderSize = PageDilationBorderSize;
+						PassParameters->ResolutionLodBiasLocal = CVarResolutionLodBiasLocal.GetValueOnRenderThread();
+						PassParameters->PageDilationBorderSizeLocal = CVarPageDilationBorderSizeLocal.GetValueOnRenderThread();
+						PassParameters->PageDilationBorderSizeDirectional = CVarPageDilationBorderSizeDirectional.GetValueOnRenderThread();
 						PassParameters->bCullBackfacingPixels = ShouldCullBackfacingPixels() ? 1 : 0;
 
 						auto ComputeShader = View.ShaderMap->GetShader<FGeneratePageFlagsFromPixelsCS>(PermutationVector);
diff --git a/Engine/Source/Runtime/Renderer/Private/VirtualShadowMaps/VirtualShadowMapArray.h b/Engine/Source/Runtime/Renderer/Private/VirtualShadowMaps/VirtualShadowMapArray.h
--- a/Engine/Source/Runtime/Renderer/Private/VirtualShadowMaps/VirtualShadowMapArray.h
+++ b/Engine/Source/Runtime/Renderer/Private/VirtualShadowMaps/VirtualShadowMapArray.h
@@ -126,6 +126,7 @@ BEGIN_GLOBAL_SHADER_PARAMETER_STRUCT(FVirtualShadowMapUniformParameters, )
 	SHADER_PARAMETER(FVector4f, RecPhysicalPoolSize)
 	SHADER_PARAMETER(FIntPoint, PhysicalPoolSize)
 	SHADER_PARAMETER(FIntPoint, PhysicalPoolSizePages)
+	SHADER_PARAMETER(uint32, PackedShadowMaskMaxLightCount)
 
 	SHADER_PARAMETER_RDG_BUFFER_SRV(ByteAddressBuffer, ProjectionData)
 	SHADER_PARAMETER_RDG_BUFFER_SRV(StructuredBuffer<uint>, PageTable)
@@ -185,6 +186,8 @@ public:
 	// Total physical page count that includes separate static pages
 	uint32 GetTotalAllocatedPhysicalPages() const;
 
+	EPixelFormat GetPackedShadowMaskFormat() const;
+
 	static void SetShaderDefines(FShaderCompilerEnvironment& OutEnvironment);
 
 	void MergeStaticPhysicalPages(FRDGBuilder& GraphBuilder);
diff --git a/Engine/Source/Runtime/Renderer/Private/VirtualShadowMaps/VirtualShadowMapProjection.cpp b/Engine/Source/Runtime/Renderer/Private/VirtualShadowMaps/VirtualShadowMapProjection.cpp
--- a/Engine/Source/Runtime/Renderer/Private/VirtualShadowMaps/VirtualShadowMapProjection.cpp
+++ b/Engine/Source/Runtime/Renderer/Private/VirtualShadowMaps/VirtualShadowMapProjection.cpp
@@ -205,7 +205,7 @@ class FVirtualShadowMapProjectionCS : public FGlobalShader
 		SHADER_PARAMETER(uint32, bCullBackfacingPixels)
 		// One pass projection parameters
 		SHADER_PARAMETER_STRUCT_REF(FForwardLightData, ForwardLightData)
-		SHADER_PARAMETER_RDG_TEXTURE_UAV(RWTexture2D< uint >, RWShadowMaskBits)
+		SHADER_PARAMETER_RDG_TEXTURE_UAV(RWTexture2D, RWShadowMaskBits)
 		// Pass per light parameters
 		SHADER_PARAMETER_STRUCT(FLightShaderParameters, Light)
 		SHADER_PARAMETER(int32, LightUniformVirtualShadowMapId)
@@ -240,7 +240,7 @@ class FVirtualShadowMapProjectionCS : public FGlobalShader
 		FPermutationDomain PermutationVector(Parameters.PermutationId);
 
 		// Directional lights are always in separate passes as forward light data structure currently
-		// only contains a single single directional light.
+		// only contains a single directional light.
 		if( PermutationVector.Get< FDirectionalLightDim >() && PermutationVector.Get< FOnePassProjectionDim >() )
 		{
 			return false;
@@ -368,11 +368,11 @@ FRDGTextureRef RenderVirtualShadowMapProjectionOnePass(
 
 	const FRDGTextureDesc ShadowMaskDesc = FRDGTextureDesc::Create2D(
 		SceneTextures.Config.Extent,
-		PF_R32_UINT,
+		VirtualShadowMapArray.GetPackedShadowMaskFormat(),
 		FClearValueBinding::None,
-		TexCreate_ShaderResource | TexCreate_UAV );
+		TexCreate_ShaderResource | TexCreate_UAV);
 
-	FRDGTextureRef ShadowMaskBits = GraphBuilder.CreateTexture( ShadowMaskDesc, InputType == EVirtualShadowMapProjectionInputType::HairStrands ? TEXT("ShadowMaskBits(HairStrands)") : TEXT("ShadowMaskBits(Gbuffer)"));
+	FRDGTextureRef ShadowMaskBits = GraphBuilder.CreateTexture(ShadowMaskDesc, InputType == EVirtualShadowMapProjectionInputType::HairStrands ? TEXT("ShadowMaskBits(HairStrands)") : TEXT("ShadowMaskBits(Gbuffer)"));
 	
 	RenderVirtualShadowMapProjectionCommon(
 		GraphBuilder,
