Implement batching of instance culling for non-nanite VSM, all local lights are processed as one dispatch, also:
- Factor out the batching setup & merging from the deferred instance culling
- Refactor the instance culling load balancer to be able to share types between different allocators
- Make batching logic use scene rendering allocator
- Rename RenderVirtualShadowMapsHw to RenderVirtualShadowMapsNonNanite (and similar names)
- Refactor view setup logic for Nanite & Non into common functions

diff --git a/Engine/Shaders/Private/InstanceCulling/BuildInstanceDrawCommands.usf b/Engine/Shaders/Private/InstanceCulling/BuildInstanceDrawCommands.usf
--- a/Engine/Shaders/Private/InstanceCulling/BuildInstanceDrawCommands.usf
+++ b/Engine/Shaders/Private/InstanceCulling/BuildInstanceDrawCommands.usf
@@ -26,10 +26,6 @@ uint GPUSceneNumInstances;
 uint GPUSceneNumPrimitives;
 uint GPUSceneNumLightmapDataItems;
 
-#ifndef BATCH_PROCESSING_MODE_NUM
-	#define BATCH_PROCESSING_MODE_NUM 1
-#endif // BATCH_PROCESSING_MODE_NUM
-
 StructuredBuffer<FDrawCommandDesc> DrawCommandDescs;
 Buffer<uint> InstanceIdOffsetBuffer;
 StructuredBuffer<uint> ViewIds;
@@ -40,33 +36,6 @@ uint DynamicInstanceIdMax;
 uint NumCullingViews;
 uint CurrentBatchProcessingMode;
 
-struct FContextBatchInfo
-{
-	uint IndirectArgOffset;
-	uint InstanceDataWriteOffset;
-	uint ViewIdsOffset;
-	uint NumViewIds;
-	uint DynamicInstanceIdOffset;
-	uint DynamicInstanceIdMax;
-	// Must be equal to ItemDataOffset[BATCH_PROCESSING_MODE_NUM]
-	// Not using array here to avoid GL shader compiler bug on Adreno devices
-	uint ItemDataOffset0;
-	uint ItemDataOffset1;
-};
-
-uint GetItemDataOffset(FContextBatchInfo BatchInfo, uint BatchProcessingMode)
-{
-#if BATCH_PROCESSING_MODE_NUM > 2
-	#error Make sure FContextBatchInfo::ItemDataOffset is up to date with BATCH_PROCESSING_MODE_NUM
-#else
-	return (BatchProcessingMode == 0) ? BatchInfo.ItemDataOffset0 : BatchInfo.ItemDataOffset1;
-#endif
-}
-
-#if ENABLE_BATCH_MODE
-StructuredBuffer<FContextBatchInfo> BatchInfos;
-StructuredBuffer<uint> BatchInds;
-#endif // ENABLE_BATCH_MODE
 
 #if DEBUG_MODE
 int bDrawOnlyVSMInvalidatingGeometry;
diff --git a/Engine/Shaders/Private/InstanceCulling/InstanceCullingCommon.ush b/Engine/Shaders/Private/InstanceCulling/InstanceCullingCommon.ush
--- a/Engine/Shaders/Private/InstanceCulling/InstanceCullingCommon.ush
+++ b/Engine/Shaders/Private/InstanceCulling/InstanceCullingCommon.ush
@@ -4,7 +4,40 @@
 
 #include "../Common.ush"
 
+#ifndef BATCH_PROCESSING_MODE_NUM
+#define BATCH_PROCESSING_MODE_NUM 1
+#endif // BATCH_PROCESSING_MODE_NUM
+
+
 struct FDrawCommandDesc
 {
 	uint bMaterialMayModifyPosition;
 };
+
+struct FContextBatchInfo
+{
+	uint IndirectArgOffset;
+	uint InstanceDataWriteOffset;
+	uint ViewIdsOffset;
+	uint NumViewIds;
+	uint DynamicInstanceIdOffset;
+	uint DynamicInstanceIdMax;
+	// Must be equal to ItemDataOffset[BATCH_PROCESSING_MODE_NUM]
+	// Not using array here to avoid GL shader compiler bug on Adreno devices
+	uint ItemDataOffset0;
+	uint ItemDataOffset1;
+};
+
+uint GetItemDataOffset(FContextBatchInfo BatchInfo, uint BatchProcessingMode)
+{
+#if BATCH_PROCESSING_MODE_NUM > 2
+#error Make sure FContextBatchInfo::ItemDataOffset is up to date with BATCH_PROCESSING_MODE_NUM
+#else
+	return (BatchProcessingMode == 0U) ? BatchInfo.ItemDataOffset0 : BatchInfo.ItemDataOffset1;
+#endif
+}
+
+#if ENABLE_BATCH_MODE
+StructuredBuffer<FContextBatchInfo> BatchInfos;
+StructuredBuffer<uint> BatchInds;
+#endif // ENABLE_BATCH_MODE
\ No newline at end of file
diff --git a/Engine/Shaders/Private/VirtualShadowMaps/BuildPerPageDrawCommands.usf b/Engine/Shaders/Private/VirtualShadowMaps/BuildPerPageDrawCommands.usf
--- a/Engine/Shaders/Private/VirtualShadowMaps/BuildPerPageDrawCommands.usf
+++ b/Engine/Shaders/Private/VirtualShadowMaps/BuildPerPageDrawCommands.usf
@@ -28,13 +28,25 @@ struct FVisibleInstanceCmd
 };
 
 RWStructuredBuffer<FVisibleInstanceCmd> VisibleInstancesOut;
-uint FirstPrimaryView;
-uint NumPrimaryViews;
 RWStructuredBuffer<uint> VisibleInstanceCountBufferOut;
 uint InstanceSceneDataSOAStride;
+uint TotalPrimaryViews;
+uint VisibleInstancesBufferNum;
+
+struct FVSMCullingBatchInfo
+{
+	uint FirstPrimaryView;
+	uint NumPrimaryViews;
+};
+
+#if ENABLE_BATCH_MODE
+StructuredBuffer<FVSMCullingBatchInfo> VSMCullingBatchInfos;
+#else // !ENABLE_BATCH_MODE
+uint FirstPrimaryView;
+uint NumPrimaryViews;
 uint DynamicInstanceIdOffset;
 uint DynamicInstanceIdMax;
-
+#endif // ENABLE_BATCH_MODE
 
 StructuredBuffer<FDrawCommandDesc> DrawCommandDescs;
 
@@ -84,16 +96,39 @@ void WriteCmd(uint MipViewId, uint InstanceId, uint IndirectArgIndex, bool bStat
 
 	uint VisibleInstanceOutputOffset = 0U;
 	WaveInterlockedAddScalar_(VisibleInstanceCountBufferOut[0], 1U, VisibleInstanceOutputOffset);
-	VisibleInstancesOut[VisibleInstanceOutputOffset] = VisibleInstanceCmd;
+	if (VisibleInstanceOutputOffset < VisibleInstancesBufferNum)
+	{
+		// TODO: Flag & feedback on overflow.
+		VisibleInstancesOut[VisibleInstanceOutputOffset] = VisibleInstanceCmd;
+	}
 }
 
 [numthreads(NUM_THREADS_PER_GROUP, 1, 1)]
 void CullPerPageDrawCommandsCs(uint3 GroupId : SV_GroupID, int GroupThreadIndex : SV_GroupIndex)
 {
-	// TODO: Support batching? Currently VSM are culled before each draw. Batching could work explicitly or using the same deferred mechanism as other InstanceCullingContexts.
-	//       However, these would need to be kept separate, and also means all instance data must be pre-allocated for all the lights (and cannot be re-used between lights).
-	//       For supporting many small lights, we will certainly need batching, but perhaps they will use a different path anyway.
-	FInstanceWorkSetup WorkSetup = InstanceCullingLoadBalancer_Setup(GroupId, GroupThreadIndex, 0);
+	uint DispatchGroupId = GetUnWrappedDispatchGroupId(GroupId);
+
+	if (DispatchGroupId >= InstanceCullingLoadBalancer_GetNumBatches())
+	{
+		return;
+	}
+#if ENABLE_BATCH_MODE
+	// Load Instance culling context batch info, indirection per group
+	FContextBatchInfo BatchInfo = BatchInfos[BatchInds[DispatchGroupId]];
+	FVSMCullingBatchInfo VSMCullingBatchInfo = VSMCullingBatchInfos[BatchInds[DispatchGroupId]];
+#else // !ENABLE_BATCH_MODE
+	// Single Instance culling context batch in the call, set up batch from the kernel parameters
+	FContextBatchInfo BatchInfo = (FContextBatchInfo)0;
+	BatchInfo.DynamicInstanceIdOffset = DynamicInstanceIdOffset;
+	BatchInfo.DynamicInstanceIdMax = DynamicInstanceIdMax;
+
+	FVSMCullingBatchInfo VSMCullingBatchInfo;
+	VSMCullingBatchInfo.FirstPrimaryView = FirstPrimaryView;
+	VSMCullingBatchInfo.NumPrimaryViews = NumPrimaryViews;
+#endif // ENABLE_BATCH_MODE
+
+	uint CurrentBatchProcessingMode = 0U;
+	FInstanceWorkSetup WorkSetup = InstanceCullingLoadBalancer_Setup(GroupId, GroupThreadIndex, GetItemDataOffset(BatchInfo, CurrentBatchProcessingMode));
 	if (!WorkSetup.bValid)
 	{
 		return;
@@ -104,13 +139,13 @@ void CullPerPageDrawCommandsCs(uint3 GroupId : SV_GroupID, int GroupThreadIndex
 
 	if (bDynamicInstanceDataOffset)
 	{
-		InstanceDataOffset += DynamicInstanceIdOffset;
-		checkSlow(InstanceDataOffset + uint(WorkSetup.LocalItemIndex) < DynamicInstanceIdMax);
+		InstanceDataOffset += BatchInfo.DynamicInstanceIdOffset;
+		checkSlow(InstanceDataOffset + uint(WorkSetup.LocalItemIndex) < BatchInfo.DynamicInstanceIdMax);
 	}
 
-
 	uint InstanceId = InstanceDataOffset + uint(WorkSetup.LocalItemIndex);
-	uint IndirectArgIndex = WorkSetup.Item.Payload >> 1U;
+	// Work item payload stores the IndirectArgsOffset, which must be offset for the batch.
+	uint IndirectArgIndex = BatchInfo.IndirectArgOffset + (WorkSetup.Item.Payload >> 1U);
 	FDrawCommandDesc DrawCommandDesc = DrawCommandDescs[IndirectArgIndex];
 
 	const bool bNearClip = (NEAR_CLIP != 0);
@@ -121,7 +156,7 @@ void CullPerPageDrawCommandsCs(uint3 GroupId : SV_GroupID, int GroupThreadIndex
 
 	uint ThreadTotalForAllViews = 0;
 	// Loop over views and output visible instance (i.e., those that overlap a valid page)
-	for (uint PrimaryViewId = FirstPrimaryView; PrimaryViewId < FirstPrimaryView + NumPrimaryViews; ++PrimaryViewId)
+	for (uint PrimaryViewId = VSMCullingBatchInfo.FirstPrimaryView; PrimaryViewId < VSMCullingBatchInfo.FirstPrimaryView + VSMCullingBatchInfo.NumPrimaryViews; ++PrimaryViewId)
 	{
 		FNaniteView NaniteView = GetNaniteView(PrimaryViewId);
 		const uint2 TargetViewSize = uint2(NaniteView.ViewSizeAndInvSize.xy);
@@ -140,9 +175,9 @@ void CullPerPageDrawCommandsCs(uint3 GroupId : SV_GroupID, int GroupThreadIndex
 		if (Cull.bIsVisible)
 		{
 			// Loop over mip levels and count number of output visible instances, also retain a bit for each mip level (used to skip empty levels in output loop)
-			for (uint MipLevel = 0; MipLevel < uint(NaniteView.TargetNumMipLevels); ++MipLevel)
+			for (uint MipLevel = 0U; MipLevel < uint(NaniteView.TargetNumMipLevels); ++MipLevel)
 			{
-				uint MipViewId = MipLevel * NumPrimaryViews + PrimaryViewId;
+				uint MipViewId = MipLevel * TotalPrimaryViews + PrimaryViewId;
 				FNaniteView MipView = GetNaniteView(MipViewId);
 				uint VirtualShadowMapId = uint(MipView.TargetLayerIndex);
 
diff --git a/Engine/Source/Runtime/RenderCore/Public/RenderGraphUtils.h b/Engine/Source/Runtime/RenderCore/Public/RenderGraphUtils.h
--- a/Engine/Source/Runtime/RenderCore/Public/RenderGraphUtils.h
+++ b/Engine/Source/Runtime/RenderCore/Public/RenderGraphUtils.h
@@ -798,6 +798,24 @@ FORCEINLINE FRDGBufferRef CreateStructuredBuffer(
 	return CreateStructuredBuffer(GraphBuilder, Name, InitialData.GetTypeSize(), InitialData.Num(), InitialData.GetData(), InitialData.Num() * InitialData.GetTypeSize(), InitialDataFlags);
 }
 
+/**
+ * Helper to create a structured buffer with initial data from a TConstArrayView.
+ */
+template <typename ElementType>
+FORCEINLINE FRDGBufferRef CreateStructuredBuffer(
+	FRDGBuilder& GraphBuilder,
+	const TCHAR* Name,
+	const TConstArrayView<ElementType>& InitialData,
+	ERDGInitialDataFlags InitialDataFlags = ERDGInitialDataFlags::None)
+{
+	static const ElementType DummyElement = ElementType();
+	if (InitialData.Num() == 0)
+	{
+		return CreateStructuredBuffer(GraphBuilder, Name, InitialData.GetTypeSize(), 1, &DummyElement, InitialData.GetTypeSize(), ERDGInitialDataFlags::NoCopy);
+	}
+	return CreateStructuredBuffer(GraphBuilder, Name, InitialData.GetTypeSize(), InitialData.Num(), InitialData.GetData(), InitialData.Num() * InitialData.GetTypeSize(), InitialDataFlags);
+}
+
 template <typename ElementType>
 FORCEINLINE FRDGBufferRef CreateStructuredBuffer(
 	FRDGBuilder& GraphBuilder,
diff --git a/Engine/Source/Runtime/Renderer/Private/InstanceCulling/InstanceCullingContext.cpp b/Engine/Source/Runtime/Renderer/Private/InstanceCulling/InstanceCullingContext.cpp
--- a/Engine/Source/Runtime/Renderer/Private/InstanceCulling/InstanceCullingContext.cpp
+++ b/Engine/Source/Runtime/Renderer/Private/InstanceCulling/InstanceCullingContext.cpp
@@ -12,6 +12,7 @@
 #include "ScenePrivate.h"
 #include "InstanceCulling/InstanceCullingManager.h"
 #include "InstanceCullingLoadBalancer.h"
+#include "InstanceCullingMergedContext.h"
 
 static TAutoConsoleVariable<int32> CVarCullInstances(
 	TEXT("r.CullInstances"),
@@ -50,7 +51,7 @@ FMeshDrawCommandOverrideArgs GetMeshDrawCommandOverrideArgs(const FInstanceCulli
 	return Result;
 }
 
-static uint32 StepInstanceDataOffset(ERHIFeatureLevel::Type FeatureLevel, uint32 NumStepInstances, uint32 NumStepDraws)
+uint32 FInstanceCullingContext::StepInstanceDataOffset(ERHIFeatureLevel::Type FeatureLevel, uint32 NumStepInstances, uint32 NumStepDraws)
 {
 	// mobile uses one instance step rate, on desktop step is once per draw
 	if (FeatureLevel == ERHIFeatureLevel::ES3_1)
@@ -271,12 +272,12 @@ const TRDGUniformBufferRef<FInstanceCullingGlobalUniforms> FInstanceCullingConte
 }
 
 
-class FInstanceCullingDeferredContext
+class FInstanceCullingDeferredContext : public FInstanceCullingMergedContext
 {
 public:
 	FInstanceCullingDeferredContext(ERHIFeatureLevel::Type InFeatureLevel, FInstanceCullingManager* InInstanceCullingManager = nullptr) 
-		: InstanceCullingManager(InInstanceCullingManager)
-		, FeatureLevel(InFeatureLevel)
+		: FInstanceCullingMergedContext(InFeatureLevel)
+		, InstanceCullingManager(InInstanceCullingManager)
 	{}
 
 	FInstanceCullingManager* InstanceCullingManager;
@@ -285,46 +286,9 @@ public:
 	FRDGBufferRef InstanceDataBuffer = nullptr;
 	TRDGUniformBufferRef<FInstanceCullingGlobalUniforms> UniformBuffer = nullptr;
 
-	/** Batches of GPU instance culling input data. */
-	TArray<FInstanceCullingContext::FBatchItem, SceneRenderingAllocator> Batches;
-
-	TArray<int32, SceneRenderingAllocator> ViewIds;
-	//TArray<FMeshDrawCommandInfo, SceneRenderingAllocator> MeshDrawCommandInfos;
-	TArray<FRHIDrawIndexedIndirectParameters, SceneRenderingAllocator> IndirectArgs;
-	TArray<FInstanceCullingContext::FDrawCommandDesc, SceneRenderingAllocator> DrawCommandDescs;
-	TArray<uint32, SceneRenderingAllocator> InstanceIdOffsets;
-
-	FInstanceCullingContext::LoadBalancerArray LoadBalancers = FInstanceCullingContext::LoadBalancerArray(InPlace, nullptr);
-	TStaticArray<TArray<uint32, SceneRenderingAllocator>, static_cast<uint32>(EBatchProcessingMode::Num)> BatchInds;
-	TArray<FInstanceCullingContext::FContextBatchInfo, SceneRenderingAllocator> BatchInfos;
-
-	ERHIFeatureLevel::Type FeatureLevel = ERHIFeatureLevel::Num;
-	// Counters to sum up all sizes to facilitate pre-sizing
-	uint32 InstanceIdBufferSize = 0U;
-	TStaticArray<int32, uint32(EBatchProcessingMode::Num)> TotalBatches = TStaticArray<int32, uint32(EBatchProcessingMode::Num)>(InPlace, 0);
-	TStaticArray<int32, uint32(EBatchProcessingMode::Num)> TotalItems = TStaticArray<int32, uint32(EBatchProcessingMode::Num)>(InPlace, 0);
-	int32 TotalIndirectArgs = 0;
-	int32 TotalViewIds = 0;
-
-	// Single Previous frame HZB which is shared among all batched contexts, thus only one is allowed (but the same can be used in multiple passes). (Needs atlas or bindless to expand(.
-	FRDGTextureRef PrevHZB = nullptr;
-	int32 NumCullingViews = 0;
-
 	bool bProcessed = false;
 
 	void ProcessBatched(TStaticArray<FBuildInstanceIdBufferAndCommandsFromPrimitiveIdsCs::FParameters*, static_cast<uint32>(EBatchProcessingMode::Num)> PassParameters);
-
-
-	~FInstanceCullingDeferredContext()
-	{
-		for (auto& LoadBalancer : LoadBalancers)
-		{
-			if (LoadBalancer != nullptr)
-			{
-				delete LoadBalancer;
-			}
-		}
-	}
 };
 
 static FRDGBufferDesc CreateInstanceIdsBufferDesc(ERHIFeatureLevel::Type FeatureLevel, uint32 BufferSize)
@@ -370,37 +334,7 @@ void FInstanceCullingContext::BuildRenderingCommands(
 			Results.DrawIndirectArgsBuffer = DeferredContext->DrawIndirectArgsBuffer;
 			Results.InstanceDataBuffer = DeferredContext->InstanceDataBuffer;
 			Results.UniformBuffer = DeferredContext->UniformBuffer;
-			DeferredContext->Batches.Add(FBatchItem{ this, InstanceCullingDrawParams, DynamicInstanceIdOffset, DynamicInstanceIdNum });
-
-			// Set HZB texture for the merged batches
-			if (bOcclusionCullInstances)
-			{
-				// Verify that each batch contains the same HZB if not null as we only support one
-				check(DeferredContext->PrevHZB == nullptr || DeferredContext->PrevHZB == GraphBuilder.RegisterExternalTexture(PrevHZB));
-
-				if (DeferredContext->PrevHZB == nullptr)
-				{
-					DeferredContext->PrevHZB = GraphBuilder.RegisterExternalTexture(PrevHZB);
-				}
-			}
-
-			// Accumulate the totals so the deferred processing can pre-size the arrays
-			for (uint32 Mode = 0U; Mode < uint32(EBatchProcessingMode::Num); ++Mode)
-			{
-				LoadBalancers[Mode]->FinalizeBatches();
-				DeferredContext->TotalBatches[Mode] += LoadBalancers[Mode]->GetBatches().Max();
-				DeferredContext->TotalItems[Mode] += LoadBalancers[Mode]->GetItems().Max();
-			}
-#if DO_CHECK
-			for (int32 ViewId : ViewIds)
-			{
-				checkf(ViewId < DeferredContext->NumCullingViews, TEXT("Attempting to defer a culling context that references a view that has not been uploaded yet."));
-			}
-#endif 
-
-			DeferredContext->TotalIndirectArgs += IndirectArgs.Num();
-			DeferredContext->TotalViewIds += ViewIds.Num();
-			DeferredContext->InstanceIdBufferSize += InstanceIdBufferSize;
+			DeferredContext->AddBatch(GraphBuilder, this, DynamicInstanceIdOffset, DynamicInstanceIdNum, InstanceCullingDrawParams);
 		}
 		return;
 	}
@@ -535,8 +469,6 @@ void FInstanceCullingContext::BuildRenderingCommands(
 	}
 }
 
-
-
 void FInstanceCullingDeferredContext::ProcessBatched(TStaticArray<FBuildInstanceIdBufferAndCommandsFromPrimitiveIdsCs::FParameters*, static_cast<uint32>(EBatchProcessingMode::Num)> PassParameters)
 {
 	if (bProcessed)
@@ -544,93 +476,16 @@ void FInstanceCullingDeferredContext::ProcessBatched(TStaticArray<FBuildInstance
 		return;
 	}
 
+	MergeBatches();
 	bProcessed = true;
-	for (uint32 Mode = 0U; Mode < uint32(EBatchProcessingMode::Num); ++Mode)
-	{
-		LoadBalancers[Mode]= new FInstanceProcessingGPULoadBalancer;
-		LoadBalancers[Mode]->ReserveStorage(TotalBatches[Mode], TotalItems[Mode]);
-	}
-	// Pre-size all arrays
-	IndirectArgs.Empty(TotalIndirectArgs);
-	DrawCommandDescs.Empty(TotalIndirectArgs);
-	InstanceIdOffsets.Empty(TotalIndirectArgs);
-	ViewIds.Empty(TotalViewIds);
-
-	BatchInfos.AddDefaulted(Batches.Num());
-	uint32 InstanceIdBufferOffset = 0U;
-	uint32 InstanceDataByteOffset = 0U;
-	const uint32 InstanceIdBufferStride = FInstanceCullingContext::GetInstanceIdBufferStride(FeatureLevel);
-
-	// Index that maps from each command to the corresponding batch - maybe not the utmost efficiency
-	for (int32 BatchIndex = 0; BatchIndex < Batches.Num(); ++BatchIndex)
-	{
-		const FInstanceCullingContext::FBatchItem& BatchItem = Batches[BatchIndex];
-		const FInstanceCullingContext& InstanceCullingContext = *BatchItem.Context;
-
-		FInstanceCullingContext::FContextBatchInfo& BatchInfo = BatchInfos[BatchIndex];
-
-		BatchInfo.IndirectArgsOffset = IndirectArgs.Num();
-		//BatchInfo.NumIndirectArgs = InstanceCullingContext.IndirectArgs.Num();
-		IndirectArgs.Append(InstanceCullingContext.IndirectArgs);
-
-		check(InstanceCullingContext.DrawCommandDescs.Num() == InstanceCullingContext.IndirectArgs.Num());
-		DrawCommandDescs.Append(InstanceCullingContext.DrawCommandDescs);
-
-		check(InstanceCullingContext.InstanceIdOffsets.Num() == InstanceCullingContext.IndirectArgs.Num());
-		InstanceIdOffsets.AddDefaulted(InstanceCullingContext.InstanceIdOffsets.Num());
-		// TODO: perform offset on GPU
-		// InstanceIdOffsets.Append(InstanceCullingContext.InstanceIdOffsets);
-		for (int32 Index = 0; Index < InstanceCullingContext.InstanceIdOffsets.Num(); ++Index)
-		{
-			InstanceIdOffsets[BatchInfo.IndirectArgsOffset + Index] = InstanceCullingContext.InstanceIdOffsets[Index] + InstanceIdBufferOffset;
-		}
-
-		BatchInfo.ViewIdsOffset = ViewIds.Num();
-		BatchInfo.NumViewIds = InstanceCullingContext.ViewIds.Num();
-		ViewIds.Append(InstanceCullingContext.ViewIds);
-
-		BatchInfo.DynamicInstanceIdOffset = BatchItem.DynamicInstanceIdOffset;
-		BatchInfo.DynamicInstanceIdMax = BatchItem.DynamicInstanceIdOffset + BatchItem.DynamicInstanceIdNum;
-
-		for (uint32 Mode = 0U; Mode < uint32(EBatchProcessingMode::Num); ++Mode)
-		{
-			int32 StartIndex = BatchInds[Mode].Num();
-			FInstanceProcessingGPULoadBalancer* MergedLoadBalancer = LoadBalancers[Mode];
-				
-			BatchInfo.ItemDataOffset[Mode] = MergedLoadBalancer->GetItems().Num();
-			FInstanceProcessingGPULoadBalancer* LoadBalancer = InstanceCullingContext.LoadBalancers[Mode];
-			LoadBalancer->FinalizeBatches();
-
-			// UnCulled bucket is used for a single instance mode
-			check(EBatchProcessingMode(Mode) != EBatchProcessingMode::UnCulled || LoadBalancer->HasSingleInstanceItemsOnly());
-
-			BatchInds[Mode].AddDefaulted(LoadBalancer->GetBatches().Num());
-
-			MergedLoadBalancer->AppendData(*LoadBalancer);
-			for (int32 Index = StartIndex; Index < BatchInds[Mode].Num(); ++Index)
-			{
-				BatchInds[Mode][Index] = BatchIndex;
-			}
-		}
-		const uint32 BatchTotalInstances = InstanceCullingContext.TotalInstances * InstanceCullingContext.ViewIds.Num();
-		const uint32 BatchTotalDraws = InstanceCullingContext.InstanceIdOffsets.Num();
 
-		FInstanceCullingDrawParams& Result = *BatchItem.Result;
-		Result.InstanceDataByteOffset = InstanceDataByteOffset;
-		Result.IndirectArgsByteOffset = BatchInfo.IndirectArgsOffset * FInstanceCullingContext::IndirectArgsNumWords * sizeof(uint32);
-
-		BatchInfo.InstanceDataWriteOffset = InstanceIdBufferOffset;
-		InstanceIdBufferOffset += BatchTotalInstances;
-		// Advance offset into per-instance buffer
-		InstanceDataByteOffset += StepInstanceDataOffset(FeatureLevel, BatchTotalInstances, BatchTotalDraws) * InstanceIdBufferStride;
-	}
 
 	// Finalize culling pass parameters
 	for (uint32 Mode = 0U; Mode < uint32(EBatchProcessingMode::Num); ++Mode)
 	{
 		PassParameters[Mode]->NumViewIds = ViewIds.Num();
-		PassParameters[Mode]->LoadBalancerParameters.NumBatches = LoadBalancers[Mode]->GetBatches().Num();
-		PassParameters[Mode]->LoadBalancerParameters.NumItems = LoadBalancers[Mode]->GetItems().Num();
+		PassParameters[Mode]->LoadBalancerParameters.NumBatches = LoadBalancers[Mode].GetBatches().Num();
+		PassParameters[Mode]->LoadBalancerParameters.NumItems = LoadBalancers[Mode].GetItems().Num();
 
 		const bool bOcclusionCullInstances = PrevHZB != nullptr && FInstanceCullingContext::IsOcclusionCullingEnabled();
 		if (bOcclusionCullInstances)
@@ -638,9 +493,9 @@ void FInstanceCullingDeferredContext::ProcessBatched(TStaticArray<FBuildInstance
 			PassParameters[Mode]->HZBTexture = PrevHZB;
 			PassParameters[Mode]->HZBSize = PrevHZB->Desc.Extent;
 			PassParameters[Mode]->HZBSampler = TStaticSamplerState< SF_Point, AM_Clamp, AM_Clamp, AM_Clamp >::GetRHI();
+		}
 	}
 }
-}
 
 template <typename DataType>
 FORCEINLINE int32 GetArrayDataSize(const TArrayView<const DataType>& Array)
@@ -775,17 +630,17 @@ FInstanceCullingDeferredContext *FInstanceCullingContext::CreateDeferredContext(
 			GraphBuilder,
 			TEXT("InstanceCullingLoadBalancer.Batches"),
 			sizeof(FInstanceProcessingGPULoadBalancer::FPackedBatch),
-			INST_CULL_CALLBACK_MODE(DeferredContext->LoadBalancers[Mode]->GetBatches().Num()),
-			INST_CULL_CALLBACK_MODE(DeferredContext->LoadBalancers[Mode]->GetBatches().GetData()),
-			INST_CULL_CALLBACK_MODE(GetArrayDataSize(DeferredContext->LoadBalancers[Mode]->GetBatches())));
+			INST_CULL_CALLBACK_MODE(DeferredContext->LoadBalancers[Mode].GetBatches().Num()),
+			INST_CULL_CALLBACK_MODE(DeferredContext->LoadBalancers[Mode].GetBatches().GetData()),
+			INST_CULL_CALLBACK_MODE(GetArrayDataSize(DeferredContext->LoadBalancers[Mode].GetBatches())));
 
 		FRDGBufferRef ItemBuffer = CreateStructuredBuffer(
 			GraphBuilder,
 			TEXT("InstanceCullingLoadBalancer.Items"),
 			sizeof(FInstanceProcessingGPULoadBalancer::FPackedItem),
-			INST_CULL_CALLBACK_MODE(DeferredContext->LoadBalancers[Mode]->GetItems().Num()),
-			INST_CULL_CALLBACK_MODE(DeferredContext->LoadBalancers[Mode]->GetItems().GetData()),
-			INST_CULL_CALLBACK_MODE(GetArrayDataSize(DeferredContext->LoadBalancers[Mode]->GetItems())));
+			INST_CULL_CALLBACK_MODE(DeferredContext->LoadBalancers[Mode].GetItems().Num()),
+			INST_CULL_CALLBACK_MODE(DeferredContext->LoadBalancers[Mode].GetItems().GetData()),
+			INST_CULL_CALLBACK_MODE(GetArrayDataSize(DeferredContext->LoadBalancers[Mode].GetItems())));
 
 		PassParameters[Mode]->LoadBalancerParameters.BatchBuffer = GraphBuilder.CreateSRV(BatchBuffer);
 		PassParameters[Mode]->LoadBalancerParameters.ItemBuffer = GraphBuilder.CreateSRV(ItemBuffer);
@@ -813,7 +668,7 @@ FInstanceCullingDeferredContext *FInstanceCullingContext::CreateDeferredContext(
 			RDG_EVENT_NAME("CullInstances(%s)", BatchProcessingModeStr[Mode]),
 			ComputeShader,
 			PassParameters[Mode],
-			INST_CULL_CALLBACK_MODE(DeferredContext->LoadBalancers[Mode]->GetWrappedCsGroupCount()));
+			INST_CULL_CALLBACK_MODE(DeferredContext->LoadBalancers[Mode].GetWrappedCsGroupCount()));
 	}
 	
 	if (FeatureLevel > ERHIFeatureLevel::ES3_1)
diff --git a/Engine/Source/Runtime/Renderer/Private/InstanceCulling/InstanceCullingLoadBalancer.h b/Engine/Source/Runtime/Renderer/Private/InstanceCulling/InstanceCullingLoadBalancer.h
--- a/Engine/Source/Runtime/Renderer/Private/InstanceCulling/InstanceCullingLoadBalancer.h
+++ b/Engine/Source/Runtime/Renderer/Private/InstanceCulling/InstanceCullingLoadBalancer.h
@@ -28,13 +28,9 @@
 /*
  * Helper to build the needed data to run per-instance operation on the GPU in a balanced way
  */
-template <typename InAllocatorType = FDefaultAllocator>
-class TInstanceCullingLoadBalancer
+class FInstanceCullingLoadBalancerBase
 {
 public:
-	using AllocatorType = InAllocatorType;
-	static constexpr ERDGInitialDataFlags DefaultRDGInitialDataFlags = InstanceCullingImplementationDetails::AllocatorTypeRDGInitialDataFlags<AllocatorType>::Flags;
-
 	static constexpr uint32 ThreadGroupSize = 64U;
 
 	// Number of bits needed for prefix sum storage
@@ -45,6 +41,76 @@
 	static constexpr uint32 NumInstancesItemBits = PrefixBits + 1U;
 	static constexpr uint32 NumInstancesItemMask = (1U << NumInstancesItemBits) - 1U;
 
+
+	struct FPackedBatch
+	{
+		uint32 FirstItem_NumItems;
+	};
+
+	FPackedBatch PackBatch(uint32 FirstItem, uint32 NumItems)
+	{
+		checkSlow(NumItems < (1U << NumInstancesItemBits));
+		checkSlow(FirstItem < (1U << (32U - NumInstancesItemBits)));
+
+		return FPackedBatch{ (FirstItem << NumInstancesItemBits) | (NumItems & NumInstancesItemMask) };
+	}
+
+	struct FPackedItem
+	{
+		// packed 32-NumInstancesItemBits:NumInstancesItemBits - need one more bit for the case where one item has ThreadGroupSize work to do
+		uint32 InstanceDataOffset_NumInstances;
+		// packed 32-PrefixBits:PrefixBits
+		uint32 Payload_BatchPrefixOffset;
+	};
+	FPackedItem PackItem(uint32 InstanceDataOffset, uint32 NumInstances, uint32 Payload, uint32 BatchPrefixSum)
+	{
+		checkSlow(NumInstances < (1U << NumInstancesItemBits));
+		checkSlow(InstanceDataOffset < (1U << (32U - NumInstancesItemBits)));
+		checkSlow(BatchPrefixSum < (1U << PrefixBits));
+		checkSlow(Payload < (1U << (32U - PrefixBits)));
+
+		return FPackedItem
+		{
+			(InstanceDataOffset << NumInstancesItemBits) | (NumInstances & NumInstancesItemMask),
+			(Payload << PrefixBits) | (BatchPrefixSum & PrefixBitMask)
+		};
+	}
+
+	BEGIN_SHADER_PARAMETER_STRUCT(FShaderParameters, )
+		SHADER_PARAMETER_RDG_BUFFER_SRV(StructuredBuffer< FPackedBatch >, BatchBuffer)
+		SHADER_PARAMETER_RDG_BUFFER_SRV(StructuredBuffer< FPackedItem >, ItemBuffer)
+		SHADER_PARAMETER(uint32, NumBatches)
+		SHADER_PARAMETER(uint32, NumItems)
+	END_SHADER_PARAMETER_STRUCT()
+
+	struct FGPUData
+	{
+		int32 NumBatches = 0;
+		int32 NumItems = 0;
+		FRDGBufferRef BatchBuffer = nullptr;
+		FRDGBufferRef ItemBuffer = nullptr;
+
+		void GetShaderParameters(FRDGBuilder& GraphBuilder, FShaderParameters& ShaderParameters)
+		{
+			ShaderParameters.BatchBuffer = GraphBuilder.CreateSRV(BatchBuffer);
+			ShaderParameters.ItemBuffer = GraphBuilder.CreateSRV(ItemBuffer);
+			ShaderParameters.NumBatches = NumBatches;
+			ShaderParameters.NumItems = NumItems;
+		}
+	};
+};
+
+/*
+ * Helper to build the needed data to run per-instance operation on the GPU in a balanced way
+ */
+template <typename InAllocatorType = FDefaultAllocator>
+class TInstanceCullingLoadBalancer : public FInstanceCullingLoadBalancerBase
+{
+public:
+	using AllocatorType = InAllocatorType;
+	static constexpr ERDGInitialDataFlags DefaultRDGInitialDataFlags = InstanceCullingImplementationDetails::AllocatorTypeRDGInitialDataFlags<AllocatorType>::Flags;
+
+
 	/*
 	 * Publish constants to a shader implementing a kernel using the load balancer.
 	 * Call from ModifyCompilationEnvironment
@@ -107,63 +173,6 @@
 		return Items.IsEmpty();
 	}
 
-	struct FPackedBatch
-	{
-		uint32 FirstItem_NumItems;
-	};
-
-	FPackedBatch PackBatch(uint32 FirstItem, uint32 NumItems)
-	{
-		checkSlow(NumItems < (1U << NumInstancesItemBits));
-		checkSlow(FirstItem < (1U << (32U - NumInstancesItemBits)));
-
-		return FPackedBatch{ (FirstItem << NumInstancesItemBits) | (NumItems & NumInstancesItemMask) };
-	}
-
-	struct FPackedItem
-	{
-		// packed 32-NumInstancesItemBits:NumInstancesItemBits - need one more bit for the case where one item has ThreadGroupSize work to do
-		uint32 InstanceDataOffset_NumInstances;
-		// packed 32-PrefixBits:PrefixBits
-		uint32 Payload_BatchPrefixOffset;
-	};
-	FPackedItem PackItem(uint32 InstanceDataOffset, uint32 NumInstances, uint32 Payload, uint32 BatchPrefixSum)
-	{
-		checkSlow(NumInstances < (1U << NumInstancesItemBits));
-		checkSlow(InstanceDataOffset < (1U << (32U - NumInstancesItemBits)));
-		checkSlow(BatchPrefixSum < (1U << PrefixBits));
-		checkSlow(Payload < (1U << (32U - PrefixBits)));
-
-		return FPackedItem
-		{
-			(InstanceDataOffset << NumInstancesItemBits) | (NumInstances & NumInstancesItemMask),
-			(Payload << PrefixBits) | (BatchPrefixSum & PrefixBitMask)
-		};
-	}
-
-	BEGIN_SHADER_PARAMETER_STRUCT(FShaderParameters, )
-		SHADER_PARAMETER_RDG_BUFFER_SRV(StructuredBuffer< FPackedBatch >, BatchBuffer)
-		SHADER_PARAMETER_RDG_BUFFER_SRV(StructuredBuffer< FPackedItem >, ItemBuffer)
-		SHADER_PARAMETER(uint32, NumBatches)
-		SHADER_PARAMETER(uint32, NumItems)
-	END_SHADER_PARAMETER_STRUCT()
-
-	struct FGPUData
-	{
-		int32 NumBatches = 0;
-		int32 NumItems = 0;
-		FRDGBufferRef BatchBuffer = nullptr;
-		FRDGBufferRef ItemBuffer = nullptr;
-
-		void GetShaderParameters(FRDGBuilder& GraphBuilder, FShaderParameters& ShaderParameters)
-		{
-			ShaderParameters.BatchBuffer = GraphBuilder.CreateSRV(BatchBuffer);
-			ShaderParameters.ItemBuffer = GraphBuilder.CreateSRV(ItemBuffer);
-			ShaderParameters.NumBatches = NumBatches;
-			ShaderParameters.NumItems = NumItems;
-		}
-	};
-
 	FGPUData Upload(FRDGBuilder& GraphBuilder, ERDGInitialDataFlags RDGInitialDataFlags = DefaultRDGInitialDataFlags)
 	{
 		FinalizeBatches();
@@ -210,14 +219,14 @@
 		return Items;
 	}
 
+	uint32 GetTotalNumInstances() const { return TotalInstances; }
+
 	template <typename AllocatorType>
 	void AppendData(const TInstanceCullingLoadBalancer<AllocatorType> &Other)
 	{
-		check(Other.CurrentBatchNumItems == 0);
-
-		Batches.Append(Other.Batches);
-		Items.Append(Other.Items);
-		TotalInstances += Other.TotalInstances;
+		Batches.Append(Other.GetBatches().GetData(), Other.GetBatches().Num());
+		Items.Append(Other.GetItems().GetData(), Other.GetItems().Num());
+		TotalInstances += Other.GetTotalNumInstances();
 	}
 
 	bool HasSingleInstanceItemsOnly() const
diff --git a/Engine/Source/Runtime/Renderer/Private/InstanceCulling/InstanceCullingMergedContext.cpp b/Engine/Source/Runtime/Renderer/Private/InstanceCulling/InstanceCullingMergedContext.cpp
new file mode 100644
--- /dev/null
+++ b/Engine/Source/Runtime/Renderer/Private/InstanceCulling/InstanceCullingMergedContext.cpp
@@ -0,0 +1,125 @@
+// Copyright Epic Games, Inc. All Rights Reserved.
+
+#include "InstanceCulling/InstanceCullingMergedContext.h"
+#include "InstanceCulling/InstanceCullingManager.h"
+
+void FInstanceCullingMergedContext::MergeBatches()
+{
+	for (uint32 Mode = 0U; Mode < uint32(EBatchProcessingMode::Num); ++Mode)
+	{
+		LoadBalancers[Mode].ReserveStorage(TotalBatches[Mode], TotalItems[Mode]);
+	}
+	// Pre-size all arrays
+	IndirectArgs.Empty(TotalIndirectArgs);
+	DrawCommandDescs.Empty(TotalIndirectArgs);
+	InstanceIdOffsets.Empty(TotalIndirectArgs);
+	ViewIds.Empty(TotalViewIds);
+
+	BatchInfos.AddDefaulted(Batches.Num());
+	uint32 InstanceIdBufferOffset = 0U;
+	uint32 InstanceDataByteOffset = 0U;
+	const uint32 InstanceIdBufferStride = FInstanceCullingContext::GetInstanceIdBufferStride(FeatureLevel);
+
+	// Index that maps from each command to the corresponding batch - maybe not the utmost efficiency
+	for (int32 BatchIndex = 0; BatchIndex < Batches.Num(); ++BatchIndex)
+	{
+		const FBatchItem& BatchItem = Batches[BatchIndex];
+		const FInstanceCullingContext& InstanceCullingContext = *BatchItem.Context;
+
+		FContextBatchInfo& BatchInfo = BatchInfos[BatchIndex];
+
+		BatchInfo.IndirectArgsOffset = IndirectArgs.Num();
+		//BatchInfo.NumIndirectArgs = InstanceCullingContext.IndirectArgs.Num();
+		IndirectArgs.Append(InstanceCullingContext.IndirectArgs);
+
+		check(InstanceCullingContext.DrawCommandDescs.Num() == InstanceCullingContext.IndirectArgs.Num());
+		DrawCommandDescs.Append(InstanceCullingContext.DrawCommandDescs);
+
+		check(InstanceCullingContext.InstanceIdOffsets.Num() == InstanceCullingContext.IndirectArgs.Num());
+		InstanceIdOffsets.AddDefaulted(InstanceCullingContext.InstanceIdOffsets.Num());
+		// TODO: perform offset on GPU
+		// InstanceIdOffsets.Append(InstanceCullingContext.InstanceIdOffsets);
+		for (int32 Index = 0; Index < InstanceCullingContext.InstanceIdOffsets.Num(); ++Index)
+		{
+			InstanceIdOffsets[BatchInfo.IndirectArgsOffset + Index] = InstanceCullingContext.InstanceIdOffsets[Index] + InstanceIdBufferOffset;
+		}
+
+		BatchInfo.ViewIdsOffset = ViewIds.Num();
+		BatchInfo.NumViewIds = InstanceCullingContext.ViewIds.Num();
+		ViewIds.Append(InstanceCullingContext.ViewIds);
+
+		BatchInfo.DynamicInstanceIdOffset = BatchItem.DynamicInstanceIdOffset;
+		BatchInfo.DynamicInstanceIdMax = BatchItem.DynamicInstanceIdOffset + BatchItem.DynamicInstanceIdNum;
+
+		for (uint32 Mode = 0U; Mode < uint32(EBatchProcessingMode::Num); ++Mode)
+		{
+			int32 StartIndex = BatchInds[Mode].Num();
+			TInstanceCullingLoadBalancer<SceneRenderingAllocator>* MergedLoadBalancer = &LoadBalancers[Mode];
+
+			BatchInfo.ItemDataOffset[Mode] = MergedLoadBalancer->GetItems().Num();
+			FInstanceProcessingGPULoadBalancer* LoadBalancer = InstanceCullingContext.LoadBalancers[Mode];
+			LoadBalancer->FinalizeBatches();
+
+			// UnCulled bucket is used for a single instance mode
+			check(EBatchProcessingMode(Mode) != EBatchProcessingMode::UnCulled || LoadBalancer->HasSingleInstanceItemsOnly());
+
+			BatchInds[Mode].AddDefaulted(LoadBalancer->GetBatches().Num());
+
+			MergedLoadBalancer->AppendData(*LoadBalancer);
+			for (int32 Index = StartIndex; Index < BatchInds[Mode].Num(); ++Index)
+			{
+				BatchInds[Mode][Index] = BatchIndex;
+			}
+		}
+		const uint32 BatchTotalInstances = InstanceCullingContext.TotalInstances * InstanceCullingContext.ViewIds.Num();
+		const uint32 BatchTotalDraws = InstanceCullingContext.InstanceIdOffsets.Num();
+
+		FInstanceCullingDrawParams& Result = *BatchItem.Result;
+		Result.InstanceDataByteOffset = InstanceDataByteOffset;
+		Result.IndirectArgsByteOffset = BatchInfo.IndirectArgsOffset * FInstanceCullingContext::IndirectArgsNumWords * sizeof(uint32);
+
+		BatchInfo.InstanceDataWriteOffset = InstanceIdBufferOffset;
+		InstanceIdBufferOffset += BatchTotalInstances;
+		// Advance offset into per-instance buffer
+		InstanceDataByteOffset += FInstanceCullingContext::StepInstanceDataOffset(FeatureLevel, BatchTotalInstances, BatchTotalDraws) * InstanceIdBufferStride;
+	}
+}
+
+void FInstanceCullingMergedContext::AddBatch(FRDGBuilder& GraphBuilder, const FInstanceCullingContext* Context, int32 DynamicInstanceIdOffset,	int32 DynamicInstanceIdNum, FInstanceCullingDrawParams* InstanceCullingDrawParams)
+{
+	Batches.Add(FBatchItem{ Context, InstanceCullingDrawParams, DynamicInstanceIdOffset, DynamicInstanceIdNum });
+
+	const bool bOcclusionCullInstances = Context->PrevHZB.IsValid() && FInstanceCullingContext::IsOcclusionCullingEnabled();
+
+	// Set HZB texture for the merged batches
+	if (bOcclusionCullInstances)
+	{
+		// Verify that each batch contains the same HZB if not null as we only support one
+		check(PrevHZB == nullptr || PrevHZB == GraphBuilder.RegisterExternalTexture(Context->PrevHZB));
+
+		if (PrevHZB == nullptr)
+		{
+			// Note: performing the registration here because the final merge of contexts may happen during RDG execute (in case of the deferred culling) which seems ill-defined
+			PrevHZB = GraphBuilder.RegisterExternalTexture(Context->PrevHZB);
+		}
+	}
+
+	// Accumulate the totals so the deferred processing can pre-size the arrays
+	for (uint32 Mode = 0U; Mode < uint32(EBatchProcessingMode::Num); ++Mode)
+	{
+		Context->LoadBalancers[Mode]->FinalizeBatches();
+		TotalBatches[Mode] += Context->LoadBalancers[Mode]->GetBatches().Max();
+		TotalItems[Mode] += Context->LoadBalancers[Mode]->GetItems().Max();
+	}
+#if DO_CHECK
+	for (int32 ViewId : Context->ViewIds)
+	{
+		checkf(NumCullingViews < 0 || ViewId < NumCullingViews, TEXT("Attempting to defer a culling context that references a view that has not been uploaded yet."));
+	}
+#endif 
+
+	TotalIndirectArgs += Context->IndirectArgs.Num();
+	TotalViewIds += Context->ViewIds.Num();
+	InstanceIdBufferSize += Context->TotalInstances * Context->ViewIds.Num();
+	TotalInstances += Context->TotalInstances;
+}
\ No newline at end of file
diff --git a/Engine/Source/Runtime/Renderer/Private/InstanceCulling/InstanceCullingMergedContext.h b/Engine/Source/Runtime/Renderer/Private/InstanceCulling/InstanceCullingMergedContext.h
new file mode 100644
--- /dev/null
+++ b/Engine/Source/Runtime/Renderer/Private/InstanceCulling/InstanceCullingMergedContext.h
@@ -0,0 +1,72 @@
+// Copyright Epic Games, Inc. All Rights Reserved.
+
+#pragma once
+
+#include "CoreMinimal.h"
+#include "InstanceCulling/InstanceCullingContext.h"
+#include "InstanceCulling/InstanceCullingLoadBalancer.h"
+
+class FInstanceCullingMergedContext
+{
+public:
+	FInstanceCullingMergedContext(ERHIFeatureLevel::Type InFeatureLevel)
+		: FeatureLevel(InFeatureLevel)
+	{}
+
+	struct FBatchItem
+	{
+		const FInstanceCullingContext* Context = nullptr;
+		FInstanceCullingDrawParams* Result = nullptr;
+		int32 DynamicInstanceIdOffset = 0;
+		int32 DynamicInstanceIdNum = 0;
+	};
+
+	// Info about a batch of culling work produced by a context, when part of a batched job
+	// Store once per context, provides start offsets to commands/etc for the context.
+	struct FContextBatchInfo
+	{
+		uint32 IndirectArgsOffset;
+		uint32 InstanceDataWriteOffset;
+		uint32 ViewIdsOffset;
+		uint32 NumViewIds;
+		uint32 DynamicInstanceIdOffset;
+		uint32 DynamicInstanceIdMax;
+		uint32 ItemDataOffset[uint32(EBatchProcessingMode::Num)];
+	};
+
+	/** Batches of GPU instance culling input data. */
+	TArray<FBatchItem, SceneRenderingAllocator> Batches;
+
+
+	/** 
+	 * Merged data, derived in MergeBatches(), follows.
+	 */
+	TArray<int32, SceneRenderingAllocator> ViewIds;
+	//TArray<FMeshDrawCommandInfo, SceneRenderingAllocator> MeshDrawCommandInfos;
+	TArray<FRHIDrawIndexedIndirectParameters, SceneRenderingAllocator> IndirectArgs;
+	TArray<FInstanceCullingContext::FDrawCommandDesc, SceneRenderingAllocator> DrawCommandDescs;
+	TArray<uint32, SceneRenderingAllocator> InstanceIdOffsets;
+
+	TStaticArray<TInstanceCullingLoadBalancer<SceneRenderingAllocator>, static_cast<uint32>(EBatchProcessingMode::Num)> LoadBalancers;
+	TStaticArray<TArray<uint32, SceneRenderingAllocator>, static_cast<uint32>(EBatchProcessingMode::Num)> BatchInds;
+	TArray<FContextBatchInfo, SceneRenderingAllocator> BatchInfos;
+
+	ERHIFeatureLevel::Type FeatureLevel = ERHIFeatureLevel::Num;
+	// Counters to sum up all sizes to facilitate pre-sizing
+	uint32 InstanceIdBufferSize = 0U;
+	TStaticArray<int32, uint32(EBatchProcessingMode::Num)> TotalBatches = TStaticArray<int32, uint32(EBatchProcessingMode::Num)>(InPlace, 0);
+	TStaticArray<int32, uint32(EBatchProcessingMode::Num)> TotalItems = TStaticArray<int32, uint32(EBatchProcessingMode::Num)>(InPlace, 0);
+	int32 TotalIndirectArgs = 0;
+	int32 TotalViewIds = 0;
+	int32 TotalInstances = 0;
+
+	// Single Previous frame HZB which is shared among all batched contexts, thus only one is allowed (but the same can be used in multiple passes). (Needs atlas or bindless to expand).
+	FRDGTextureRef PrevHZB = nullptr;
+	int32 NumCullingViews = 0;
+
+	// Merge the queued batches and populate the derived data.
+	void MergeBatches();
+
+
+	void AddBatch(FRDGBuilder& GraphBuilder, const FInstanceCullingContext* Context, int32 DynamicInstanceIdOffset, int32 DynamicInstanceIdNum, FInstanceCullingDrawParams* InstanceCullingDrawParams = nullptr);
+};
\ No newline at end of file
diff --git a/Engine/Source/Runtime/Renderer/Private/ShadowDepthRendering.cpp b/Engine/Source/Runtime/Renderer/Private/ShadowDepthRendering.cpp
--- a/Engine/Source/Runtime/Renderer/Private/ShadowDepthRendering.cpp
+++ b/Engine/Source/Runtime/Renderer/Private/ShadowDepthRendering.cpp
@@ -1770,57 +1770,12 @@ void FSceneRenderer::RenderVirtualShadowMaps(FRDGBuilder& GraphBuilder, bool bNa
 				{
 					for (const TSharedPtr<FVirtualShadowMapClipmap>& Clipmap : SortedShadowsForShadowDepthPass.VirtualShadowMapClipmaps)
 					{
-						// TODO: Decide if this sort of logic belongs here or in Nanite (as with the mip level view expansion logic)
-						// We're eventually going to want to snap/quantize these rectangles/positions somewhat so probably don't want it
-						// entirely within Nanite, but likely makes sense to have some sort of "multi-viewport" notion in Nanite that can
-						// handle both this and mips.
-						// NOTE: There's still the additional VSM view logic that runs on top of this in Nanite too (see CullRasterize variant)
-						Nanite::FPackedViewParams BaseParams;
-						BaseParams.ViewRect = FIntRect(0, 0, FVirtualShadowMap::VirtualMaxResolutionXY, FVirtualShadowMap::VirtualMaxResolutionXY);
-						BaseParams.HZBTestViewRect = BaseParams.ViewRect;
-						BaseParams.RasterContextSize = VirtualShadowMapArray.GetPhysicalPoolSize();
-						BaseParams.LODScaleFactor = ComputeNaniteShadowsLODScaleFactor();
-						BaseParams.PrevTargetLayerIndex = INDEX_NONE;
-						BaseParams.TargetMipLevel = 0;
-						BaseParams.TargetMipCount = 1;	// No mips for clipmaps
-
-						for (int32 ClipmapLevelIndex = 0; ClipmapLevelIndex < Clipmap->GetLevelCount(); ++ClipmapLevelIndex)
-						{
-							FVirtualShadowMap* VirtualShadowMap = Clipmap->GetVirtualShadowMap(ClipmapLevelIndex);
-
-							Nanite::FPackedViewParams Params = BaseParams;
-							Params.TargetLayerIndex = VirtualShadowMap->ID;
-							Params.ViewMatrices = Clipmap->GetViewMatrices(ClipmapLevelIndex);
-							Params.PrevTargetLayerIndex = INDEX_NONE;
-							Params.PrevViewMatrices = Params.ViewMatrices;
-							Params.Flags = 0;
-
-							// TODO: Clean this up - could be stored in a single structure for the whole clipmap
-							int32 HZBKey = Clipmap->GetHZBKey(ClipmapLevelIndex);
-
-							if (PrevHZBPhysical)
-							{
-								CacheManager->SetHZBViewParams(HZBKey, Params);
-							}
-
-							// If we're going to generate a new HZB this frame, save the associated metadata
-							if (bVSMUseHZB)
-							{
-								FVirtualShadowMapHZBMetadata& HZBMeta = VirtualShadowMapArray.HZBMetadata.FindOrAdd(HZBKey);
-								HZBMeta.TargetLayerIndex = Params.TargetLayerIndex;
-								HZBMeta.ViewMatrices = Params.ViewMatrices;
-								HZBMeta.ViewRect = Params.ViewRect;
-							}
-
-							Nanite::FPackedView View = Nanite::CreatePackedView(Params);
-							VirtualShadowViews.Add(View);
-
-							// Mark that we rendered to this VSM for caching purposes
-							if (VirtualShadowMap->VirtualShadowMapCacheEntry)
-							{
-								VirtualShadowMap->VirtualShadowMapCacheEntry->MarkRendered();
-							}		
-						}
+						VirtualShadowMapArray.AddRenderViews(
+							Clipmap,
+							ComputeNaniteShadowsLODScaleFactor(),
+							PrevHZBPhysical.IsValid(),
+							bVSMUseHZB,
+							VirtualShadowViews);
 					}
 				}
 
@@ -1828,51 +1783,12 @@ void FSceneRenderer::RenderVirtualShadowMaps(FRDGBuilder& GraphBuilder, bool bNa
 				{
 					if (ProjectedShadowInfo->ShouldClampToNearPlane() == bShouldClampToNearPlane && ProjectedShadowInfo->HasVirtualShadowMap())
 					{
-						Nanite::FPackedViewParams BaseParams;
-						BaseParams.ViewRect = ProjectedShadowInfo->GetOuterViewRect();
-						BaseParams.HZBTestViewRect = BaseParams.ViewRect;
-						BaseParams.RasterContextSize = VirtualShadowMapArray.GetPhysicalPoolSize();
-						BaseParams.LODScaleFactor = ComputeNaniteShadowsLODScaleFactor();
-						BaseParams.PrevTargetLayerIndex = INDEX_NONE;
-						BaseParams.TargetMipLevel = 0;
-						BaseParams.TargetMipCount = FVirtualShadowMap::MaxMipLevels;
-
-						int32 NumMaps = ProjectedShadowInfo->bOnePassPointLightShadow ? 6 : 1;
-						for( int32 i = 0; i < NumMaps; i++ )
-						{
-							FVirtualShadowMap* VirtualShadowMap = ProjectedShadowInfo->VirtualShadowMaps[i];
-
-							Nanite::FPackedViewParams Params = BaseParams;
-							Params.TargetLayerIndex = VirtualShadowMap->ID;
-							Params.ViewMatrices = ProjectedShadowInfo->GetShadowDepthRenderingViewMatrices(i, true);
-							Params.PrevTargetLayerIndex = INDEX_NONE;
-							Params.PrevViewMatrices = Params.ViewMatrices;
-							Params.Flags = 0;
-
-							int32 HZBKey = ProjectedShadowInfo->GetLightSceneInfo().Id + (i << 24);
-							if (PrevHZBPhysical)
-							{
-								CacheManager->SetHZBViewParams(HZBKey, Params);
-							}
-
-							// If we're going to generate a new HZB this frame, save the associated metadata
-							if (bVSMUseHZB)
-							{
-								FVirtualShadowMapHZBMetadata& HZBMeta = VirtualShadowMapArray.HZBMetadata.FindOrAdd(HZBKey);
-								HZBMeta.TargetLayerIndex = Params.TargetLayerIndex;
-								HZBMeta.ViewMatrices = Params.ViewMatrices;
-								HZBMeta.ViewRect = Params.ViewRect;
-							}
-
-							Nanite::FPackedView View = Nanite::CreatePackedView(Params);
-							VirtualShadowViews.Add(View);
-
-							// Mark that we rendered to this VSM for caching purposes
-							if (VirtualShadowMap->VirtualShadowMapCacheEntry)
-							{
-								VirtualShadowMap->VirtualShadowMapCacheEntry->MarkRendered();
-							}							
-						}
+						VirtualShadowMapArray.AddRenderViews(
+							ProjectedShadowInfo, 
+							ComputeNaniteShadowsLODScaleFactor(),
+							PrevHZBPhysical.IsValid(),
+							bVSMUseHZB,
+							VirtualShadowViews);
 					}
 				}
 
@@ -1941,7 +1857,7 @@ void FSceneRenderer::RenderVirtualShadowMaps(FRDGBuilder& GraphBuilder, bool bNa
 
 	if (UseNonNaniteVirtualShadowMaps(ShaderPlatform, FeatureLevel))
 	{
-		VirtualShadowMapArray.RenderVirtualShadowMapsHw(GraphBuilder, SortedShadowsForShadowDepthPass.VirtualShadowMapShadows, *Scene);
+		VirtualShadowMapArray.RenderVirtualShadowMapsNonNanite(GraphBuilder, SortedShadowsForShadowDepthPass.VirtualShadowMapShadows, *Scene);
 	}
 
 	// If separate static/dynamic caching is enabled, we may need to merge some pages after rendering
diff --git a/Engine/Source/Runtime/Renderer/Private/VirtualShadowMaps/VirtualShadowMapArray.cpp b/Engine/Source/Runtime/Renderer/Private/VirtualShadowMaps/VirtualShadowMapArray.cpp
--- a/Engine/Source/Runtime/Renderer/Private/VirtualShadowMaps/VirtualShadowMapArray.cpp
+++ b/Engine/Source/Runtime/Renderer/Private/VirtualShadowMaps/VirtualShadowMapArray.cpp
@@ -16,6 +16,7 @@
 #include "SceneTextureReductions.h"
 #include "ShaderDebug.h"
 #include "GPUMessaging.h"
+#include "InstanceCulling/InstanceCullingMergedContext.h"
 
 IMPLEMENT_STATIC_UNIFORM_BUFFER_SLOT(VirtualShadowMapUbSlot);
 
@@ -181,6 +182,12 @@ TAutoConsoleVariable<int32> CVarVirtualShadowOnePassProjectionMaxLights(
 	ECVF_RenderThreadSafe
 );
 
+TAutoConsoleVariable<int32> CVarDoNonNaniteBatching(
+	TEXT("r.Shadow.Virtual.NonNanite.Batch"),
+	1,
+	TEXT("."),
+	ECVF_RenderThreadSafe
+);
 #if !UE_BUILD_SHIPPING
 bool GDumpVSMLightNames = false;
 void DumpVSMLightNames()
@@ -1755,7 +1762,7 @@ void FVirtualShadowMapArray::CreateMipViews( TArray<Nanite::FPackedView, SceneRe
 	}
 
 	// Remove unused mip views
-	check(MaxMips > 0);
+	check(Views.IsEmpty() || MaxMips > 0);
 	Views.SetNum(MaxMips * NumPrimaryViews, false);
 }
 
@@ -1798,6 +1805,13 @@ BEGIN_SHADER_PARAMETER_STRUCT(FVirtualShadowDepthPassParameters,)
 END_SHADER_PARAMETER_STRUCT()
 
 
+struct FVSMCullingBatchInfo
+{
+	uint32 FirstPrimaryView;
+	uint32 NumPrimaryViews;
+};
+
+
 struct FVisibleInstanceCmd
 {
 	uint32 PackedPageInfo;
@@ -1811,10 +1825,10 @@ class FCullPerPageDrawCommandsCs : public FGlobalShader
 	SHADER_USE_PARAMETER_STRUCT(FCullPerPageDrawCommandsCs, FGlobalShader)
 
 	class FNearClipDim		: SHADER_PERMUTATION_BOOL( "NEAR_CLIP" );
-	class FLoopOverViewsDim	: SHADER_PERMUTATION_BOOL( "LOOP_OVER_VIEWS" );
 	class FUseHzbDim : SHADER_PERMUTATION_BOOL("USE_HZB_OCCLUSION");
 	class FGenerateStatsDim : SHADER_PERMUTATION_BOOL("VSM_GENERATE_STATS");
-	using FPermutationDomain = TShaderPermutationDomain< FUseHzbDim, FNearClipDim, FLoopOverViewsDim, FGenerateStatsDim >;
+	class FBatchedDim : SHADER_PERMUTATION_BOOL("ENABLE_BATCH_MODE");
+	using FPermutationDomain = TShaderPermutationDomain< FUseHzbDim, FNearClipDim, FBatchedDim, FGenerateStatsDim >;
 
 public:
 	static bool ShouldCompilePermutation(const FGlobalShaderPermutationParameters& Parameters)
@@ -1836,6 +1850,17 @@ public:
 		OutEnvironment.SetDefine(TEXT("USE_GLOBAL_GPU_SCENE_DATA"), 1);
 	}
 
+
+	BEGIN_SHADER_PARAMETER_STRUCT(FHZBShaderParameters, )
+		SHADER_PARAMETER_RDG_BUFFER_SRV(StructuredBuffer< uint >, ShadowHZBPageTable)
+		SHADER_PARAMETER_RDG_TEXTURE(Texture2D, HZBTexture)
+		SHADER_PARAMETER_SAMPLER(SamplerState, HZBSampler)
+		SHADER_PARAMETER(FVector2f, HZBSize)
+		SHADER_PARAMETER(uint32, HZBMode)
+	END_SHADER_PARAMETER_STRUCT()
+
+
+
 	BEGIN_SHADER_PARAMETER_STRUCT(FParameters, )
 		SHADER_PARAMETER_STRUCT_INCLUDE(FVirtualShadowMapPageTableParameters, PageTableParams)
 
@@ -1849,20 +1874,22 @@ public:
 
 		SHADER_PARAMETER(int32, FirstPrimaryView)
 		SHADER_PARAMETER(int32, NumPrimaryViews)
+		SHADER_PARAMETER(uint32, TotalPrimaryViews)
+		SHADER_PARAMETER(uint32, VisibleInstancesBufferNum)
 		SHADER_PARAMETER(int32, DynamicInstanceIdOffset)
 		SHADER_PARAMETER(int32, DynamicInstanceIdMax)
 		SHADER_PARAMETER_RDG_BUFFER_SRV(StructuredBuffer< FPackedView >, InViews)
 		SHADER_PARAMETER_RDG_BUFFER_SRV(StructuredBuffer< FInstanceCullingContext::FDrawCommandDesc >, DrawCommandDescs)
 
+		SHADER_PARAMETER_RDG_BUFFER_SRV(StructuredBuffer< FContextBatchInfo >, BatchInfos)
+		SHADER_PARAMETER_RDG_BUFFER_SRV(StructuredBuffer< FVSMCullingBatchInfo >, VSMCullingBatchInfos)
+		SHADER_PARAMETER_RDG_BUFFER_SRV(StructuredBuffer< uint32 >, BatchInds)
+
 		SHADER_PARAMETER_RDG_BUFFER_UAV(RWStructuredBuffer<FVisibleInstanceCmd>, VisibleInstancesOut)
 		SHADER_PARAMETER_RDG_BUFFER_UAV(RWBuffer<uint>, DrawIndirectArgsBufferOut)
 		SHADER_PARAMETER_RDG_BUFFER_UAV(RWStructuredBuffer<uint>, VisibleInstanceCountBufferOut)
 
-		SHADER_PARAMETER_RDG_BUFFER_SRV(StructuredBuffer< uint >, ShadowHZBPageTable)
-		SHADER_PARAMETER_RDG_TEXTURE(Texture2D, HZBTexture)
-		SHADER_PARAMETER_SAMPLER(SamplerState, HZBSampler)
-		SHADER_PARAMETER(FVector2f, HZBSize)
-		SHADER_PARAMETER(uint32, HZBMode)
+		SHADER_PARAMETER_STRUCT_INCLUDE(FHZBShaderParameters, HZBShaderParameters)
 
 		SHADER_PARAMETER_RDG_BUFFER_UAV(RWStructuredBuffer<uint>, OutInvalidatingInstances)
 		SHADER_PARAMETER(uint32, NumInvalidatingInstanceSlots)
@@ -1950,359 +1977,474 @@ public:
 };
 IMPLEMENT_GLOBAL_SHADER(FOutputCommandInstanceListsCs, "/Engine/Private/VirtualShadowMaps/BuildPerPageDrawCommands.usf", "OutputCommandInstanceListsCs", SF_Compute);
 
-void FVirtualShadowMapArray::RenderVirtualShadowMapsHw(FRDGBuilder& GraphBuilder, const TArray<FProjectedShadowInfo *, SceneRenderingAllocator>& VirtualSmMeshCommandPasses, FScene& Scene)
+
+struct FCullingResult
 {
-	if (VirtualSmMeshCommandPasses.Num() == 0)
-	{
-		return;
-	}
+	FRDGBufferRef DrawIndirectArgsRDG;
+	FRDGBufferRef InstanceIdOffsetBufferRDG;
+	FRDGBufferRef InstanceIdsBuffer;
+	FRDGBufferRef PageInfoBuffer;
+	uint32 MaxNumInstancesPerPass;
+};
 
-	RDG_EVENT_SCOPE(GraphBuilder, "RenderVirtualShadowMaps(Non-Nanite)");
+template <typename InstanceCullingLoadBalancerType>
+static FCullingResult AddCullingPasses(FRDGBuilder& GraphBuilder,
+	const TConstArrayView<FRHIDrawIndexedIndirectParameters> &IndirectArgs,
+	const TConstArrayView<FInstanceCullingContext::FDrawCommandDesc>& DrawCommandDescs,
+	const TConstArrayView<uint32>& InstanceIdOffsets,
+	InstanceCullingLoadBalancerType *LoadBalancer,
+	const TConstArrayView<FInstanceCullingMergedContext::FContextBatchInfo> BatchInfos,
+	const TConstArrayView<FVSMCullingBatchInfo> VSMCullingBatchInfos,
+	const TConstArrayView<uint32> BatchInds,
+	bool bUseNearClip,
+	uint32 TotalInstances,
+	uint32 TotalPrimaryViews,
+	FRDGBufferRef VirtualShadowViewsRDG,
+	const FCullPerPageDrawCommandsCs::FHZBShaderParameters &HZBShaderParameters,
+	FVirtualShadowMapArray *VirtualShadowMapArray,
+	FGPUScene& GPUScene)
+{
+	int32 NumIndirectArgs = IndirectArgs.Num();
 
-	FGPUScene& GPUScene = Scene.GPUScene;
+	FRDGBufferRef TmpInstanceIdOffsetBufferRDG = CreateStructuredBuffer(GraphBuilder, TEXT("Shadow.Virtual.TmpInstanceIdOffsetBuffer"), sizeof(uint32), NumIndirectArgs, nullptr, 0);
 
-	FRDGBufferSRVRef PrevPageTableRDGSRV = CacheManager->IsValid() ? GraphBuilder.CreateSRV(GraphBuilder.RegisterExternalBuffer(CacheManager->PrevBuffers.PageTable, TEXT("Shadow.Virtual.PrevPageTable"))) : nullptr;
+	// TODO: This is both not right, and also over conservative when running with the atomic path
+	FCullingResult CullingResult;
+	CullingResult.MaxNumInstancesPerPass = TotalInstances * 64u;
+	FRDGBufferRef VisibleInstancesRdg = CreateStructuredBuffer(GraphBuilder, TEXT("Shadow.Virtual.VisibleInstances"), sizeof(FVisibleInstanceCmd), CullingResult.MaxNumInstancesPerPass, nullptr, 0);
 
-	int32 HZBMode = CVarNonNaniteVsmUseHzb.GetValueOnRenderThread();
+	FRDGBufferRef VisibleInstanceWriteOffsetRDG = CreateStructuredBuffer(GraphBuilder, TEXT("Shadow.Virtual.VisibleInstanceWriteOffset"), sizeof(uint32), 1, nullptr, 0);
+	FRDGBufferRef OutputOffsetBufferRDG = CreateStructuredBuffer(GraphBuilder, TEXT("Shadow.Virtual.OutputOffsetBuffer"), sizeof(uint32), 1, nullptr, 0);
 
-	auto InitHZB = [&]()->FRDGTextureRef
-	{
-		if (HZBMode == 1 && CacheManager->IsValid())
-		{
-			return GraphBuilder.RegisterExternalTexture(CacheManager->PrevBuffers.HZBPhysical);
-		}
+	AddClearUAVPass(GraphBuilder, GraphBuilder.CreateUAV(VisibleInstanceWriteOffsetRDG), 0);
+	AddClearUAVPass(GraphBuilder, GraphBuilder.CreateUAV(OutputOffsetBufferRDG), 0);
 
-		if (HZBMode == 2 && HZBPhysical != nullptr)
-		{
-			return HZBPhysical;
-		}
-		return nullptr;
-	};
-	const FRDGTextureRef HZBTexture = InitHZB();
+	// Create buffer for indirect args and upload draw arg data, also clears the instance to zero
+	CullingResult.DrawIndirectArgsRDG = GraphBuilder.CreateBuffer(FRDGBufferDesc::CreateIndirectDesc(FInstanceCullingContext::IndirectArgsNumWords * IndirectArgs.Num()), TEXT("Shadow.Virtual.DrawIndirectArgsBuffer"));
+	GraphBuilder.QueueBufferUpload(CullingResult.DrawIndirectArgsRDG, IndirectArgs.GetData(), IndirectArgs.GetTypeSize() * IndirectArgs.Num());
+
+	FGlobalShaderMap* ShaderMap = GetGlobalShaderMap(GMaxRHIFeatureLevel);
 
-	// TODO: create a VirtualClipMapMeshRenderInfo - or something - that contains the required info:
-	//       The View used to set up (to get all the usual crap), The MeshCommandPass (to perform setup & culling), Culling volume?, Ref to Clipmap itself.
+	// Note: we redundantly clear the instance counts here as there is some issue with replays on certain consoles.
+	FInstanceCullingContext::AddClearIndirectArgInstanceCountPass(GraphBuilder, ShaderMap, CullingResult.DrawIndirectArgsRDG);
 
-	// TODO: A lot of this loop is similar to the setup in ShadowDepthRendering... we could probably combine something pretty easily
-	// that returns a packed view list.
+	// not using structured buffer as we have to get at it as a vertex buffer 
+	CullingResult.InstanceIdOffsetBufferRDG = GraphBuilder.CreateBuffer(FRDGBufferDesc::CreateBufferDesc(sizeof(uint32), InstanceIdOffsets.Num()), TEXT("Shadow.Virtual.InstanceIdOffsetBuffer"));
 
-	// Loop over the mesh command passes needed (one for each SM, we need to do this on a per-SM basis mainly because different views may have different view infos)
-	// In practice there is no fundamental reason why we could not merge all the draws into one giant submission - should probably consider this at some point. 
-	// One question is whether the shaders refer to the view, and actually can get sensible data out (I'm not sure this is the case at present anyway, as it contains a bastardization of main view and shadow view info)
-	// The other one to beware of is the fact that we have these dynamic primitives, which until this is fixed will cause the creation of an entire copy of the GPU scene data for each shadow view that contains one...
-	//   Starting to become urgent!
-	for (int Index = 0; Index < VirtualSmMeshCommandPasses.Num(); ++Index)
 	{
-		FProjectedShadowInfo* ProjectedShadowInfo = VirtualSmMeshCommandPasses[Index];
+		FCullPerPageDrawCommandsCs::FParameters* PassParameters = GraphBuilder.AllocParameters<FCullPerPageDrawCommandsCs::FParameters>();
 
-		FParallelMeshDrawCommandPass& MeshCommandPass = ProjectedShadowInfo->GetShadowDepthPass();
-		// TODO: Right now, this is for clipmaps only! Get hold of them somehow!
-		const TSharedPtr<FVirtualShadowMapClipmap> Clipmap = ProjectedShadowInfo->VirtualShadowMapClipmap;
-		// TODO: Also get the associated view such that all the parameters can be set.
-		const FViewInfo* ViewUsedToCreateShadow = ProjectedShadowInfo->DependentView;
-		FViewInfo* ShadowDepthView = ProjectedShadowInfo->ShadowDepthView;
+		VirtualShadowMapArray->GetPageTableParameters(GraphBuilder, PassParameters->PageTableParams);
 
-		const FViewInfo& View = *ViewUsedToCreateShadow;
-		TArray<Nanite::FPackedView, SceneRenderingAllocator> VirtualShadowViews;
+		PassParameters->GPUSceneInstanceSceneData = GPUScene.InstanceSceneDataBuffer.SRV;
+		PassParameters->GPUSceneInstancePayloadData = GPUScene.InstancePayloadDataBuffer.SRV;
+		PassParameters->GPUScenePrimitiveSceneData = GPUScene.PrimitiveBuffer.SRV;
+		PassParameters->GPUSceneFrameNumber = GPUScene.GetSceneFrameNumber();
+		PassParameters->InstanceSceneDataSOAStride = GPUScene.InstanceSceneDataSOAStride;
 
-		ProjectedShadowInfo->BeginRenderView(GraphBuilder, &Scene);
+		PassParameters->DynamicInstanceIdOffset = BatchInfos[0].DynamicInstanceIdOffset;
+		PassParameters->DynamicInstanceIdMax = BatchInfos[0].DynamicInstanceIdMax;
 
-		if( Clipmap )
-		{
-			Nanite::FPackedViewParams BaseParams;
-			BaseParams.ViewRect = FIntRect(0, 0, FVirtualShadowMap::VirtualMaxResolutionXY, FVirtualShadowMap::VirtualMaxResolutionXY);
-			BaseParams.RasterContextSize = GetPhysicalPoolSize();
-			BaseParams.PrevTargetLayerIndex = INDEX_NONE;
-			BaseParams.TargetMipLevel = 0;
-			BaseParams.TargetMipCount = 1;	// No mips for clipmaps
-
-			for (int32 ClipmapLevelIndex = 0; ClipmapLevelIndex < Clipmap->GetLevelCount(); ++ClipmapLevelIndex)
-			{
-				FVirtualShadowMap* VirtualShadowMap = Clipmap->GetVirtualShadowMap(ClipmapLevelIndex);
+		auto GPUData = LoadBalancer->Upload(GraphBuilder);
+		GPUData.GetShaderParameters(GraphBuilder, PassParameters->LoadBalancerParameters);
 
-				Nanite::FPackedViewParams Params = BaseParams;
-				Params.TargetLayerIndex = VirtualShadowMap->ID;
-				Params.ViewMatrices = Clipmap->GetViewMatrices(ClipmapLevelIndex);
-				Params.PrevTargetLayerIndex = INDEX_NONE;
-				Params.PrevViewMatrices = Params.ViewMatrices;
+		PassParameters->FirstPrimaryView = VSMCullingBatchInfos[0].FirstPrimaryView;
+		PassParameters->NumPrimaryViews = VSMCullingBatchInfos[0].NumPrimaryViews;
+		PassParameters->TotalPrimaryViews = TotalPrimaryViews;
+		PassParameters->VisibleInstancesBufferNum = CullingResult.MaxNumInstancesPerPass;
+		PassParameters->InViews = GraphBuilder.CreateSRV(VirtualShadowViewsRDG);
+		PassParameters->DrawCommandDescs = GraphBuilder.CreateSRV(CreateStructuredBuffer(GraphBuilder, TEXT("Shadow.Virtual.DrawCommandDescs"), DrawCommandDescs));
 
-				if (HZBTexture)
-				{
-					CacheManager->SetHZBViewParams(Clipmap->GetHZBKey(ClipmapLevelIndex), Params);
-				}
+		const bool bUseBatchMode = !BatchInds.IsEmpty();
+		if (bUseBatchMode)
+		{
+			PassParameters->BatchInfos = GraphBuilder.CreateSRV(CreateStructuredBuffer(GraphBuilder, TEXT("Shadow.Virtual.BatchInfos"), BatchInfos));
+			PassParameters->VSMCullingBatchInfos = GraphBuilder.CreateSRV(CreateStructuredBuffer(GraphBuilder, TEXT("Shadow.Virtual.VSMCullingBatchInfos"), VSMCullingBatchInfos));
+			PassParameters->BatchInds = GraphBuilder.CreateSRV(CreateStructuredBuffer(GraphBuilder, TEXT("Shadow.Virtual.BatchInds"), BatchInds));
+		}
 
-				VirtualShadowViews.Add(Nanite::CreatePackedView(Params));
+		PassParameters->DrawIndirectArgsBufferOut = GraphBuilder.CreateUAV(CullingResult.DrawIndirectArgsRDG, PF_R32_UINT);
 
-				// This is used to mark all the referenced physical pages as being updated. 
-				// If rendering was disabled or some such we can't assume the content of a mapped page is useful.
-				if (VirtualShadowMap->VirtualShadowMapCacheEntry)
-				{
-					VirtualShadowMap->VirtualShadowMapCacheEntry->MarkRendered();
-				}
-			}
-		}
-		else if (ProjectedShadowInfo->HasVirtualShadowMap())
+		PassParameters->VisibleInstancesOut = GraphBuilder.CreateUAV(VisibleInstancesRdg);
+		PassParameters->VisibleInstanceCountBufferOut = GraphBuilder.CreateUAV(VisibleInstanceWriteOffsetRDG);
+		PassParameters->OutInvalidatingInstances = GraphBuilder.CreateUAV(VirtualShadowMapArray->InvalidatingInstancesRDG);
+		PassParameters->NumInvalidatingInstanceSlots = VirtualShadowMapArray->NumInvalidatingInstanceSlots;
+		PassParameters->HZBShaderParameters = HZBShaderParameters;
+
+		bool bGenerateStats = VirtualShadowMapArray->StatsBufferRDG != nullptr;
+		if (bGenerateStats)
 		{
-			Nanite::FPackedViewParams BaseParams;
-			BaseParams.ViewRect = ProjectedShadowInfo->GetOuterViewRect();
-			BaseParams.HZBTestViewRect = BaseParams.ViewRect;
-			BaseParams.RasterContextSize = GetPhysicalPoolSize();
-			//BaseParams.LODScaleFactor = ComputeNaniteShadowsLODScaleFactor();
-			BaseParams.PrevTargetLayerIndex = INDEX_NONE;
-			BaseParams.TargetMipLevel = 0;
-			BaseParams.TargetMipCount = FVirtualShadowMap::MaxMipLevels;
-
-			int32 NumMaps = ProjectedShadowInfo->bOnePassPointLightShadow ? 6 : 1;
-			for( int32 i = 0; i < NumMaps; i++ )
-			{
-				FVirtualShadowMap* VirtualShadowMap = ProjectedShadowInfo->VirtualShadowMaps[i];
+			PassParameters->OutStatsBuffer = GraphBuilder.CreateUAV(VirtualShadowMapArray->StatsBufferRDG);
+		}
 
-				Nanite::FPackedViewParams Params = BaseParams;
-				Params.TargetLayerIndex = VirtualShadowMap->ID;
-				Params.ViewMatrices = ProjectedShadowInfo->GetShadowDepthRenderingViewMatrices(i, true);
+		FCullPerPageDrawCommandsCs::FPermutationDomain PermutationVector;
+		PermutationVector.Set< FCullPerPageDrawCommandsCs::FNearClipDim >(bUseNearClip);
+		PermutationVector.Set< FCullPerPageDrawCommandsCs::FBatchedDim >(bUseBatchMode);
+		PermutationVector.Set< FCullPerPageDrawCommandsCs::FUseHzbDim >(HZBShaderParameters.HZBTexture != nullptr);
+		PermutationVector.Set< FCullPerPageDrawCommandsCs::FGenerateStatsDim >(bGenerateStats);
 
-				if (HZBTexture)
-				{
-					CacheManager->SetHZBViewParams(ProjectedShadowInfo->GetLightSceneInfo().Id + (i << 24), Params);
-				}
+		auto ComputeShader = ShaderMap->GetShader<FCullPerPageDrawCommandsCs>(PermutationVector);
 
-				VirtualShadowViews.Add(Nanite::CreatePackedView(Params));
+		FComputeShaderUtils::AddPass(
+			GraphBuilder,
+			RDG_EVENT_NAME("CullPerPageDrawCommands"),
+			ComputeShader,
+			PassParameters,
+			LoadBalancer->GetWrappedCsGroupCount()
+		);
+	}
+	// 2.2.Allocate space for the final instance ID output and so on.
+	if (true)
+	{
+		FAllocateCommandInstanceOutputSpaceCs::FParameters* PassParameters = GraphBuilder.AllocParameters<FAllocateCommandInstanceOutputSpaceCs::FParameters>();
 
-				if (VirtualShadowMap->VirtualShadowMapCacheEntry)
-				{
-					VirtualShadowMap->VirtualShadowMapCacheEntry->MarkRendered();
-				}
-			}
-		}
+		FRDGBufferRef InstanceIdOutOffsetBufferRDG = CreateStructuredBuffer(GraphBuilder, TEXT("InstanceCulling.OutputOffsetBufferOut"), sizeof(uint32), 1, nullptr, 0);
+		AddClearUAVPass(GraphBuilder, GraphBuilder.CreateUAV(InstanceIdOutOffsetBufferRDG), 0);
 
-		int32 NumPrimaryViews = VirtualShadowViews.Num();
-		CreateMipViews( VirtualShadowViews );
+		PassParameters->NumIndirectArgs = NumIndirectArgs;
+		PassParameters->InstanceIdOffsetBufferOut = GraphBuilder.CreateUAV(CullingResult.InstanceIdOffsetBufferRDG, PF_R32_UINT);;
+		PassParameters->OutputOffsetBufferOut = GraphBuilder.CreateUAV(InstanceIdOutOffsetBufferRDG);
+		PassParameters->TmpInstanceIdOffsetBufferOut = GraphBuilder.CreateUAV(TmpInstanceIdOffsetBufferRDG);
+		PassParameters->DrawIndirectArgsBuffer = GraphBuilder.CreateSRV(CullingResult.DrawIndirectArgsRDG, PF_R32_UINT);
 
-		MeshCommandPass.WaitForSetupTask();
+		auto ComputeShader = ShaderMap->GetShader<FAllocateCommandInstanceOutputSpaceCs>();
 
-		FInstanceCullingContext* InstanceCullingContext = MeshCommandPass.GetInstanceCullingContext();
+		FComputeShaderUtils::AddPass(
+			GraphBuilder,
+			RDG_EVENT_NAME("AllocateCommandInstanceOutputSpaceCs"),
+			ComputeShader,
+			PassParameters,
+			FComputeShaderUtils::GetGroupCount(NumIndirectArgs, FAllocateCommandInstanceOutputSpaceCs::NumThreadsPerGroup)
+		);
+	}
+	// 2.3. Perform final pass to re-shuffle the instance ID's to their final resting places
+	CullingResult.InstanceIdsBuffer = GraphBuilder.CreateBuffer(FRDGBufferDesc::CreateStructuredDesc(sizeof(uint32), CullingResult.MaxNumInstancesPerPass), TEXT("Shadow.Virtual.InstanceIdsBuffer"));
+	CullingResult.PageInfoBuffer = GraphBuilder.CreateBuffer(FRDGBufferDesc::CreateStructuredDesc(sizeof(uint32), CullingResult.MaxNumInstancesPerPass), TEXT("Shadow.Virtual.PageInfoBuffer"));
 
-		if (InstanceCullingContext->HasCullingCommands())
-		{
-			FString LightNameWithLevel;
-			FSceneRenderer::GetLightNameForDrawEvent(ProjectedShadowInfo->GetLightSceneInfo().Proxy, LightNameWithLevel);
-			RDG_EVENT_SCOPE(GraphBuilder, "%s", *LightNameWithLevel);
+	FRDGBufferRef OutputPassIndirectArgs = FComputeShaderUtils::AddIndirectArgsSetupCsPass1D(GraphBuilder, VisibleInstanceWriteOffsetRDG, TEXT("Shadow.Virtual.IndirectArgs"), FOutputCommandInstanceListsCs::NumThreadsPerGroup);
+	if (true)
+	{
 
+		FOutputCommandInstanceListsCs::FParameters* PassParameters = GraphBuilder.AllocParameters<FOutputCommandInstanceListsCs::FParameters>();
 
-			uint32 DynamicInstanceIdOffset = ShadowDepthView->DynamicPrimitiveCollector.GetInstanceSceneDataOffset();
-			uint32 DynamicInstanceIdMax = DynamicInstanceIdOffset + ShadowDepthView->DynamicPrimitiveCollector.NumInstances();
+		PassParameters->VisibleInstances = GraphBuilder.CreateSRV(VisibleInstancesRdg);
+		PassParameters->PageInfoBufferOut = GraphBuilder.CreateUAV(CullingResult.PageInfoBuffer);
+		PassParameters->InstanceIdsBufferOut = GraphBuilder.CreateUAV(CullingResult.InstanceIdsBuffer);
+		PassParameters->TmpInstanceIdOffsetBufferOut = GraphBuilder.CreateUAV(TmpInstanceIdOffsetBufferRDG);
+		PassParameters->VisibleInstanceCountBuffer = GraphBuilder.CreateSRV(VisibleInstanceWriteOffsetRDG);
+		PassParameters->IndirectArgs = OutputPassIndirectArgs;
 
-			int32 NumIndirectArgs = InstanceCullingContext->IndirectArgs.Num();
+		auto ComputeShader = ShaderMap->GetShader<FOutputCommandInstanceListsCs>();
 
-			FRDGBufferRef TmpInstanceIdOffsetBufferRDG = CreateStructuredBuffer(GraphBuilder, TEXT("Shadow.Virtual.TmpInstanceIdOffsetBuffer"), sizeof(uint32), NumIndirectArgs, nullptr, 0);
+		FComputeShaderUtils::AddPass(
+			GraphBuilder,
+			RDG_EVENT_NAME("OutputCommandInstanceListsCs"),
+			ComputeShader,
+			PassParameters,
+			OutputPassIndirectArgs,
+			0
+		);
+	}
 
-			// TODO: This is both not right, and also over conservative when running with the atomic path
-			const uint32 MaxNumInstancesPerPass = InstanceCullingContext->TotalInstances * 64u;
-			FRDGBufferRef VisibleInstancesRdg = CreateStructuredBuffer(GraphBuilder, TEXT("Shadow.Virtual.VisibleInstances"), sizeof(FVisibleInstanceCmd), MaxNumInstancesPerPass, nullptr, 0);
+	return CullingResult;
+}
 
-			FRDGBufferRef VisibleInstanceWriteOffsetRDG = CreateStructuredBuffer(GraphBuilder, TEXT("Shadow.Virtual.VisibleInstanceWriteOffset"), sizeof(uint32),  1, nullptr, 0);
-			FRDGBufferRef OutputOffsetBufferRDG = CreateStructuredBuffer(GraphBuilder, TEXT("Shadow.Virtual.OutputOffsetBuffer"), sizeof(uint32), 1, nullptr, 0);
 
-			AddClearUAVPass(GraphBuilder, GraphBuilder.CreateUAV(VisibleInstanceWriteOffsetRDG), 0);
-			AddClearUAVPass(GraphBuilder, GraphBuilder.CreateUAV(OutputOffsetBufferRDG), 0);
+static void AddRasterPass(
+	FRDGBuilder& GraphBuilder, 
+	FRDGEventName&& PassName,
+	const FViewInfo * ShadowDepthView, 
+	const TRDGUniformBufferRef<FShadowDepthPassUniformParameters> &ShadowDepthPassUniformBuffer,
+	FVirtualShadowMapArray* VirtualShadowMapArray,
+	FRDGBufferRef VirtualShadowViewsRDG,
+	const FCullingResult &CullingResult, 
+	FParallelMeshDrawCommandPass& MeshCommandPass,
+	FVirtualShadowDepthPassParameters* PassParameters,
+	TRDGUniformBufferRef<FInstanceCullingGlobalUniforms> InstanceCullingUniformBuffer)
+{
+	PassParameters->View = ShadowDepthView->ViewUniformBuffer;
+	PassParameters->ShadowDepthPass = ShadowDepthPassUniformBuffer;
 
-			FRDGBufferRef VirtualShadowViewsRDG = CreateStructuredBuffer(GraphBuilder, TEXT("Shadow.Virtual.VirtualShadowViews"), VirtualShadowViews);
+	PassParameters->VirtualShadowMap = VirtualShadowMapArray->GetUniformBuffer(GraphBuilder);
+	PassParameters->InViews = GraphBuilder.CreateSRV(VirtualShadowViewsRDG);
+	PassParameters->InstanceCullingDrawParams.DrawIndirectArgsBuffer = CullingResult.DrawIndirectArgsRDG;
+	PassParameters->InstanceCullingDrawParams.InstanceIdOffsetBuffer = CullingResult.InstanceIdOffsetBufferRDG;
+	PassParameters->InstanceCullingDrawParams.InstanceCulling = InstanceCullingUniformBuffer;
 
-			const auto& IndirectArgs = InstanceCullingContext->IndirectArgs;
-			const auto& DrawCommandDescs = InstanceCullingContext->DrawCommandDescs;
-			const TArray<uint32> &InstanceIdOffsets = InstanceCullingContext->InstanceIdOffsets;
+	FIntRect ViewRect;
+	ViewRect.Max = FVirtualShadowMap::VirtualMaxResolutionXY;
 
-			// Create buffer for indirect args and upload draw arg data, also clears the instance to zero
-			FRDGBufferRef DrawIndirectArgsRDG = GraphBuilder.CreateBuffer(FRDGBufferDesc::CreateIndirectDesc(FInstanceCullingContext::IndirectArgsNumWords * IndirectArgs.Num()), TEXT("Shadow.Virtual.DrawIndirectArgsBuffer"));
-			GraphBuilder.QueueBufferUpload(DrawIndirectArgsRDG, IndirectArgs.GetData(), IndirectArgs.GetTypeSize() * IndirectArgs.Num());
+	GraphBuilder.AddPass(
+		MoveTemp(PassName),
+		PassParameters,
+		ERDGPassFlags::Raster | ERDGPassFlags::SkipRenderPass,
+		[&MeshCommandPass, PassParameters, ViewRect](FRHICommandList& RHICmdList)
+		{
+			FRHIRenderPassInfo RPInfo;
+			RPInfo.ResolveParameters.DestRect.X1 = ViewRect.Min.X;
+			RPInfo.ResolveParameters.DestRect.Y1 = ViewRect.Min.Y;
+			RPInfo.ResolveParameters.DestRect.X2 = ViewRect.Max.X;
+			RPInfo.ResolveParameters.DestRect.Y2 = ViewRect.Max.Y;
+			RHICmdList.BeginRenderPass(RPInfo, TEXT("RasterizeVirtualShadowMaps(Non-Nanite)"));
 
-			FGlobalShaderMap* ShaderMap = GetGlobalShaderMap(GMaxRHIFeatureLevel);
+			RHICmdList.SetViewport(ViewRect.Min.X, ViewRect.Min.Y, 0.0f, FMath::Min(ViewRect.Max.X, 32767), FMath::Min(ViewRect.Max.Y, 32767), 1.0f);
 
-			// Note: we redundantly clear the instance counts here as there is some issue with replays on certain consoles.
-			FInstanceCullingContext::AddClearIndirectArgInstanceCountPass(GraphBuilder, ShaderMap, DrawIndirectArgsRDG);
+			MeshCommandPass.DispatchDraw(nullptr, RHICmdList, &PassParameters->InstanceCullingDrawParams);
+			RHICmdList.EndRenderPass();
+		});
+}
 
-			// not using structured buffer as we have to get at it as a vertex buffer 
-			FRDGBufferRef InstanceIdOffsetBufferRDG = GraphBuilder.CreateBuffer(FRDGBufferDesc::CreateBufferDesc(sizeof(uint32), InstanceIdOffsets.Num()), TEXT("Shadow.Virtual.InstanceIdOffsetBuffer"));
+void FVirtualShadowMapArray::RenderVirtualShadowMapsNonNanite(FRDGBuilder& GraphBuilder, const TArray<FProjectedShadowInfo *, SceneRenderingAllocator>& VirtualSmMeshCommandPasses, FScene& Scene)
+{
+	if (VirtualSmMeshCommandPasses.Num() == 0)
+	{
+		return;
+	}
 
-			{
-				FCullPerPageDrawCommandsCs::FParameters* PassParameters = GraphBuilder.AllocParameters<FCullPerPageDrawCommandsCs::FParameters>();
+	RDG_EVENT_SCOPE(GraphBuilder, "RenderVirtualShadowMaps(Non-Nanite)");
 
-				GetPageTableParameters(GraphBuilder, PassParameters->PageTableParams);
+	FGPUScene& GPUScene = Scene.GPUScene;
 
-				auto LoadBalancer = InstanceCullingContext->LoadBalancers[uint32(EBatchProcessingMode::Generic)];
+	FRDGBufferSRVRef PrevPageTableRDGSRV = CacheManager->IsValid() ? GraphBuilder.CreateSRV(GraphBuilder.RegisterExternalBuffer(CacheManager->PrevBuffers.PageTable, TEXT("Shadow.Virtual.PrevPageTable"))) : nullptr;
 
-				PassParameters->GPUSceneInstanceSceneData = GPUScene.InstanceSceneDataBuffer.SRV;
-				PassParameters->GPUSceneInstancePayloadData = GPUScene.InstancePayloadDataBuffer.SRV;
-				PassParameters->GPUScenePrimitiveSceneData = GPUScene.PrimitiveBuffer.SRV;
-				PassParameters->GPUSceneFrameNumber = GPUScene.GetSceneFrameNumber();
-				PassParameters->InstanceSceneDataSOAStride = GPUScene.InstanceSceneDataSOAStride;
+	int32 HZBMode = CVarNonNaniteVsmUseHzb.GetValueOnRenderThread();
 
-				PassParameters->DynamicInstanceIdOffset = DynamicInstanceIdOffset;
-				PassParameters->DynamicInstanceIdMax = DynamicInstanceIdMax;
+	auto InitHZB = [&]()->FRDGTextureRef
+	{
+		if (HZBMode == 1 && CacheManager->IsValid())
+		{
+			return GraphBuilder.RegisterExternalTexture(CacheManager->PrevBuffers.HZBPhysical);
+		}
 
+		if (HZBMode == 2 && HZBPhysical != nullptr)
+		{
+			return HZBPhysical;
+		}
+		return nullptr;
+	};
+	const FRDGTextureRef HZBTexture = InitHZB();
 
-				auto GPUData = LoadBalancer->Upload(GraphBuilder);
-				GPUData.GetShaderParameters(GraphBuilder, PassParameters->LoadBalancerParameters);
+	TArray<FVSMCullingBatchInfo, SceneRenderingAllocator> UnBatchedVSMCullingBatchInfo;
+	TArray<FProjectedShadowInfo*, SceneRenderingAllocator> BatchedVirtualSmMeshCommandPasses;
+	TArray<FProjectedShadowInfo*, SceneRenderingAllocator> UnBatchedVirtualSmMeshCommandPasses;
+	UnBatchedVSMCullingBatchInfo.Reserve(VirtualSmMeshCommandPasses.Num());
+	BatchedVirtualSmMeshCommandPasses.Reserve(VirtualSmMeshCommandPasses.Num());
+	UnBatchedVirtualSmMeshCommandPasses.Reserve(VirtualSmMeshCommandPasses.Num());
+	TArray<Nanite::FPackedView, SceneRenderingAllocator> VirtualShadowViews;
 
-				PassParameters->FirstPrimaryView = 0;
-				PassParameters->NumPrimaryViews = NumPrimaryViews;
-				PassParameters->InViews = GraphBuilder.CreateSRV(VirtualShadowViewsRDG);
-				PassParameters->DrawCommandDescs = GraphBuilder.CreateSRV(CreateStructuredBuffer(GraphBuilder, TEXT("Shadow.Virtual.DrawCommandDescs"), DrawCommandDescs));
-				PassParameters->DrawIndirectArgsBufferOut = GraphBuilder.CreateUAV(DrawIndirectArgsRDG, PF_R32_UINT);
+	TArray<FVSMCullingBatchInfo, SceneRenderingAllocator> VSMCullingBatchInfos;
+	VSMCullingBatchInfos.Reserve(VirtualSmMeshCommandPasses.Num());
 
-				PassParameters->VisibleInstancesOut = GraphBuilder.CreateUAV(VisibleInstancesRdg);
-				PassParameters->VisibleInstanceCountBufferOut = GraphBuilder.CreateUAV(VisibleInstanceWriteOffsetRDG);
-				PassParameters->OutInvalidatingInstances = GraphBuilder.CreateUAV(InvalidatingInstancesRDG);
-				PassParameters->NumInvalidatingInstanceSlots = NumInvalidatingInstanceSlots;
+	TArray<FVirtualShadowDepthPassParameters*, SceneRenderingAllocator> BatchedPassParameters;
+	BatchedPassParameters.Reserve(VirtualSmMeshCommandPasses.Num());
 
-				if (HZBTexture)
-				{
-					// Mode 2 uses the current frame HZB  & page table.
-					PassParameters->ShadowHZBPageTable = HZBMode == 2 ? GraphBuilder.CreateSRV(PageTableRDG) : PrevPageTableRDGSRV;
-					PassParameters->HZBTexture = HZBTexture;
-					PassParameters->HZBSize = HZBTexture->Desc.Extent;
-					PassParameters->HZBSampler = TStaticSamplerState< SF_Point, AM_Clamp, AM_Clamp, AM_Clamp >::GetRHI();
-					PassParameters->HZBMode = HZBMode;
-				}
+	FInstanceCullingMergedContext InstanceCullingMergedContext(GMaxRHIFeatureLevel);
+	// We don't use the registered culling views (this redundancy should probably be addressed at some point), set the number to disable index range checking
+	InstanceCullingMergedContext.NumCullingViews = -1;
+	for (int32 Index = 0; Index < VirtualSmMeshCommandPasses.Num(); ++Index)
+	{
+		FProjectedShadowInfo* ProjectedShadowInfo = VirtualSmMeshCommandPasses[Index];
+		ProjectedShadowInfo->BeginRenderView(GraphBuilder, &Scene);
 
-				bool bGenerateStats = StatsBufferRDG != nullptr;
-				if (bGenerateStats)
-				{
-					PassParameters->OutStatsBuffer = GraphBuilder.CreateUAV(StatsBufferRDG);
-				}
+		FVSMCullingBatchInfo VSMCullingBatchInfo;
+		VSMCullingBatchInfo.FirstPrimaryView = uint32(VirtualShadowViews.Num());
+		VSMCullingBatchInfo.NumPrimaryViews = 0U;
 
-				FCullPerPageDrawCommandsCs::FPermutationDomain PermutationVector;
-				PermutationVector.Set< FCullPerPageDrawCommandsCs::FNearClipDim >( !Clipmap.IsValid() );
-				PermutationVector.Set< FCullPerPageDrawCommandsCs::FLoopOverViewsDim >( Clipmap.IsValid() );
-				PermutationVector.Set< FCullPerPageDrawCommandsCs::FUseHzbDim >(HZBTexture != nullptr);
-				PermutationVector.Set< FCullPerPageDrawCommandsCs::FGenerateStatsDim >(bGenerateStats);
+		const TSharedPtr<FVirtualShadowMapClipmap> Clipmap = ProjectedShadowInfo->VirtualShadowMapClipmap;
+		if (Clipmap)
+		{
+			VSMCullingBatchInfo.NumPrimaryViews = AddRenderViews(Clipmap, 1.0f, HZBTexture != nullptr, false, VirtualShadowViews);
+			UnBatchedVSMCullingBatchInfo.Add(VSMCullingBatchInfo);
+			UnBatchedVirtualSmMeshCommandPasses.Add(ProjectedShadowInfo);
+		}
+		else if (ProjectedShadowInfo->HasVirtualShadowMap())
+		{
+			FParallelMeshDrawCommandPass& MeshCommandPass = ProjectedShadowInfo->GetShadowDepthPass();
+			MeshCommandPass.WaitForSetupTask();
 
-				auto ComputeShader = ShaderMap->GetShader<FCullPerPageDrawCommandsCs>( PermutationVector );
+			FInstanceCullingContext* InstanceCullingContext = MeshCommandPass.GetInstanceCullingContext();
 
-				FComputeShaderUtils::AddPass(
-					GraphBuilder,
-					RDG_EVENT_NAME("CullPerPageDrawCommands"),
-					ComputeShader,
-					PassParameters,
-					LoadBalancer->GetWrappedCsGroupCount()
-				);
-			}
-			// 2.2.Allocate space for the final instance ID output and so on.
+			if (InstanceCullingContext->HasCullingCommands())
 			{
-				FAllocateCommandInstanceOutputSpaceCs::FParameters* PassParameters = GraphBuilder.AllocParameters<FAllocateCommandInstanceOutputSpaceCs::FParameters>();
+				VSMCullingBatchInfo.NumPrimaryViews = AddRenderViews(ProjectedShadowInfo, 1.0f, HZBTexture != nullptr, false, VirtualShadowViews);
 
-				FRDGBufferRef InstanceIdOutOffsetBufferRDG = CreateStructuredBuffer(GraphBuilder, TEXT("InstanceCulling.OutputOffsetBufferOut"), sizeof(uint32), 1, nullptr, 0);
-				AddClearUAVPass(GraphBuilder, GraphBuilder.CreateUAV(InstanceIdOutOffsetBufferRDG), 0);
+				if (CVarDoNonNaniteBatching.GetValueOnRenderThread())
+				{
+					FViewInfo* ShadowDepthView = ProjectedShadowInfo->ShadowDepthView;
+					uint32 DynamicInstanceIdOffset = ShadowDepthView->DynamicPrimitiveCollector.GetInstanceSceneDataOffset();
+					uint32 DynamicInstanceIdMax = DynamicInstanceIdOffset + ShadowDepthView->DynamicPrimitiveCollector.NumInstances();
+
+					VSMCullingBatchInfos.Add(VSMCullingBatchInfo);
+
+					// Note: we have to allocate these up front as the context merging machinery writes the offsets directly to the &PassParameters->InstanceCullingDrawParams, 
+					// this is a side-effect from sharing the code with the deferred culling. Should probably be refactored.
+					FVirtualShadowDepthPassParameters* PassParameters = GraphBuilder.AllocParameters<FVirtualShadowDepthPassParameters>();
+					InstanceCullingMergedContext.AddBatch(GraphBuilder, InstanceCullingContext, DynamicInstanceIdOffset, ShadowDepthView->DynamicPrimitiveCollector.NumInstances(), &PassParameters->InstanceCullingDrawParams);
+					BatchedVirtualSmMeshCommandPasses.Add(ProjectedShadowInfo);
+					BatchedPassParameters.Add(PassParameters);
+				}
+				else
+				{
+					UnBatchedVSMCullingBatchInfo.Add(VSMCullingBatchInfo);
+					UnBatchedVirtualSmMeshCommandPasses.Add(ProjectedShadowInfo);
+				}
+			}
+		}
+	}
+	uint32 TotalPrimaryViews = uint32(VirtualShadowViews.Num());
+	CreateMipViews(VirtualShadowViews);
+	FRDGBufferRef VirtualShadowViewsRDG = CreateStructuredBuffer(GraphBuilder, TEXT("Shadow.Virtual.VirtualShadowViews"), VirtualShadowViews);
 
-				PassParameters->NumIndirectArgs = NumIndirectArgs;
-				PassParameters->InstanceIdOffsetBufferOut = GraphBuilder.CreateUAV(InstanceIdOffsetBufferRDG, PF_R32_UINT);;
-				PassParameters->OutputOffsetBufferOut = GraphBuilder.CreateUAV(InstanceIdOutOffsetBufferRDG);
-				PassParameters->TmpInstanceIdOffsetBufferOut = GraphBuilder.CreateUAV(TmpInstanceIdOffsetBufferRDG);
-				PassParameters->DrawIndirectArgsBuffer = GraphBuilder.CreateSRV(DrawIndirectArgsRDG, PF_R32_UINT);
+	// Helper function to create raster pass UB - only really need two of these ever
+	auto CreateShadowDepthPassUniformBuffer = [this, &VirtualShadowViewsRDG, &GraphBuilder](bool bClampToNearPlane)
+	{
+		FShadowDepthPassUniformParameters* ShadowDepthPassParameters = GraphBuilder.AllocParameters<FShadowDepthPassUniformParameters>();
+		check(PhysicalPagePoolRDG != nullptr);
+		// TODO: These are not used for this case anyway
+		ShadowDepthPassParameters->ProjectionMatrix = FMatrix::Identity;
+		ShadowDepthPassParameters->ViewMatrix = FMatrix::Identity;
+		ShadowDepthPassParameters->ShadowParams = FVector4f(0.0f, 0.0f, 0.0f, 1.0f);
+		ShadowDepthPassParameters->bRenderToVirtualShadowMap = true;
+
+		ShadowDepthPassParameters->VirtualSmPageTable = GraphBuilder.CreateSRV(PageTableRDG);
+		ShadowDepthPassParameters->PackedNaniteViews = GraphBuilder.CreateSRV(VirtualShadowViewsRDG);
+		ShadowDepthPassParameters->PageRectBounds = GraphBuilder.CreateSRV(PageRectBoundsRDG);
+		ShadowDepthPassParameters->OutDepthBuffer = GraphBuilder.CreateUAV(PhysicalPagePoolRDG, ERDGUnorderedAccessViewFlags::SkipBarrier);
+		SetupSceneTextureUniformParameters(GraphBuilder, GMaxRHIFeatureLevel, ESceneTextureSetupMode::None, ShadowDepthPassParameters->SceneTextures);
+		ShadowDepthPassParameters->bClampToNearPlane = bClampToNearPlane;
+
+		return GraphBuilder.CreateUniformBuffer(ShadowDepthPassParameters);
+	};
 
-				auto ComputeShader = ShaderMap->GetShader<FAllocateCommandInstanceOutputSpaceCs>();
+	FCullPerPageDrawCommandsCs::FHZBShaderParameters HZBShaderParameters;
+	if (HZBTexture)
+	{
+		// Mode 2 uses the current frame HZB  & page table.
+		HZBShaderParameters.ShadowHZBPageTable = HZBMode == 2 ? GraphBuilder.CreateSRV(PageTableRDG) : PrevPageTableRDGSRV;
+		HZBShaderParameters.HZBTexture = HZBTexture;
+		HZBShaderParameters.HZBSize = HZBTexture->Desc.Extent;
+		HZBShaderParameters.HZBSampler = TStaticSamplerState< SF_Point, AM_Clamp, AM_Clamp, AM_Clamp >::GetRHI();
+		HZBShaderParameters.HZBMode = HZBMode;
+	}
 
-				FComputeShaderUtils::AddPass(
-					GraphBuilder,
-					RDG_EVENT_NAME("AllocateCommandInstanceOutputSpaceCs"),
-					ComputeShader,
-					PassParameters,
-					FComputeShaderUtils::GetGroupCount(NumIndirectArgs, FAllocateCommandInstanceOutputSpaceCs::NumThreadsPerGroup)
-				);
-			}
-			// 2.3. Perform final pass to re-shuffle the instance ID's to their final resting places
-			FRDGBufferRef InstanceIdsBuffer = GraphBuilder.CreateBuffer(FRDGBufferDesc::CreateStructuredDesc(sizeof(uint32), MaxNumInstancesPerPass), TEXT("Shadow.Virtual.InstanceIdsBuffer"));
-			FRDGBufferRef PageInfoBuffer = GraphBuilder.CreateBuffer(FRDGBufferDesc::CreateStructuredDesc(sizeof(uint32), MaxNumInstancesPerPass), TEXT("Shadow.Virtual.PageInfoBuffer"));
+	// Process batched passes
+	if (!InstanceCullingMergedContext.Batches.IsEmpty())
+	{
+		RDG_EVENT_SCOPE(GraphBuilder, "Batched");
 
-			{
-				FRDGBufferRef OutputPassIndirectArgs = FComputeShaderUtils::AddIndirectArgsSetupCsPass1D(GraphBuilder, VisibleInstanceWriteOffsetRDG, TEXT("Shadow.Virtual.IndirectArgs"), FOutputCommandInstanceListsCs::NumThreadsPerGroup);
+		InstanceCullingMergedContext.MergeBatches();
 
-				FOutputCommandInstanceListsCs::FParameters* PassParameters = GraphBuilder.AllocParameters<FOutputCommandInstanceListsCs::FParameters>();
+		GraphBuilder.BeginEventScope(RDG_EVENT_NAME("CullingPasses"));
+		FCullingResult CullingResult = AddCullingPasses(
+			GraphBuilder,
+			InstanceCullingMergedContext.IndirectArgs,
+			InstanceCullingMergedContext.DrawCommandDescs,
+			InstanceCullingMergedContext.InstanceIdOffsets,
+			&InstanceCullingMergedContext.LoadBalancers[uint32(EBatchProcessingMode::Generic)],
+			InstanceCullingMergedContext.BatchInfos,
+			VSMCullingBatchInfos,
+			InstanceCullingMergedContext.BatchInds[uint32(EBatchProcessingMode::Generic)],
+			true,
+			InstanceCullingMergedContext.TotalInstances,
+			TotalPrimaryViews,
+			VirtualShadowViewsRDG,
+			HZBShaderParameters,
+			this,
+			GPUScene
+		);
+		GraphBuilder.EndEventScope();
 
-				PassParameters->VisibleInstances = GraphBuilder.CreateSRV(VisibleInstancesRdg);
-				PassParameters->PageInfoBufferOut = GraphBuilder.CreateUAV(PageInfoBuffer);
-				PassParameters->InstanceIdsBufferOut = GraphBuilder.CreateUAV(InstanceIdsBuffer);
-				PassParameters->TmpInstanceIdOffsetBufferOut = GraphBuilder.CreateUAV(TmpInstanceIdOffsetBufferRDG);
-				PassParameters->VisibleInstanceCountBuffer = GraphBuilder.CreateSRV(VisibleInstanceWriteOffsetRDG);
-				PassParameters->IndirectArgs = OutputPassIndirectArgs;
+		TRDGUniformBufferRef<FShadowDepthPassUniformParameters> ShadowDepthPassUniformBuffer = CreateShadowDepthPassUniformBuffer(false);
 
-				auto ComputeShader = ShaderMap->GetShader<FOutputCommandInstanceListsCs>();
+		FInstanceCullingGlobalUniforms* InstanceCullingGlobalUniforms = GraphBuilder.AllocParameters<FInstanceCullingGlobalUniforms>();
+		InstanceCullingGlobalUniforms->InstanceIdsBuffer = GraphBuilder.CreateSRV(CullingResult.InstanceIdsBuffer);
+		InstanceCullingGlobalUniforms->PageInfoBuffer = GraphBuilder.CreateSRV(CullingResult.PageInfoBuffer);
+		InstanceCullingGlobalUniforms->BufferCapacity = CullingResult.MaxNumInstancesPerPass;
+		TRDGUniformBufferRef<FInstanceCullingGlobalUniforms> InstanceCullingUniformBuffer = GraphBuilder.CreateUniformBuffer(InstanceCullingGlobalUniforms);
 
-				FComputeShaderUtils::AddPass(
-					GraphBuilder,
-					RDG_EVENT_NAME("OutputCommandInstanceListsCs"),
-					ComputeShader,
-					PassParameters,
-					OutputPassIndirectArgs,
-					0
-				);
-			}
+		{
+			RDG_EVENT_SCOPE(GraphBuilder, "RasterPasses");
 
-			check(PhysicalPagePoolRDG != nullptr);
+			for (int Index = 0; Index < BatchedVirtualSmMeshCommandPasses.Num(); ++Index)
+			{
+				FProjectedShadowInfo* ProjectedShadowInfo = BatchedVirtualSmMeshCommandPasses[Index];
+				FParallelMeshDrawCommandPass& MeshCommandPass = ProjectedShadowInfo->GetShadowDepthPass();
+				FViewInfo* ShadowDepthView = ProjectedShadowInfo->ShadowDepthView;
 
-			FVirtualShadowDepthPassParameters* PassParameters = GraphBuilder.AllocParameters<FVirtualShadowDepthPassParameters>();
-			PassParameters->View = ShadowDepthView->ViewUniformBuffer;
+				// Local lights are assumed to not use the clamp to near-plane (this is used for some per-object SMs but these should never be used fotr VSM).
+				check(!ProjectedShadowInfo->ShouldClampToNearPlane());
+			
+				FString LightNameWithLevel;
+				FSceneRenderer::GetLightNameForDrawEvent(ProjectedShadowInfo->GetLightSceneInfo().Proxy, LightNameWithLevel);
+				AddRasterPass(GraphBuilder, RDG_EVENT_NAME("Rasterize[%s]", *LightNameWithLevel), ShadowDepthView, ShadowDepthPassUniformBuffer, this, VirtualShadowViewsRDG, CullingResult, MeshCommandPass, BatchedPassParameters[Index], InstanceCullingUniformBuffer);
+			}
+		}
+	}
 
-			FShadowDepthPassUniformParameters* ShadowDepthPassParameters = GraphBuilder.AllocParameters<FShadowDepthPassUniformParameters>();
+	// Loop over the un batched mesh command passes needed, these are all the clipmaps (but we may change the criteria)
+	for (int Index = 0; Index < UnBatchedVirtualSmMeshCommandPasses.Num(); ++Index)
+	{
+		const auto VSMCullingBatchInfo = UnBatchedVSMCullingBatchInfo[Index];
+		FProjectedShadowInfo* ProjectedShadowInfo = UnBatchedVirtualSmMeshCommandPasses[Index];
+		FInstanceCullingMergedContext::FContextBatchInfo CullingBatchInfo = FInstanceCullingMergedContext::FContextBatchInfo{ 0 };
 
-			SetupSceneTextureUniformParameters(GraphBuilder, GMaxRHIFeatureLevel, ESceneTextureSetupMode::None, ShadowDepthPassParameters->SceneTextures);
+		FParallelMeshDrawCommandPass& MeshCommandPass = ProjectedShadowInfo->GetShadowDepthPass();
+		const TSharedPtr<FVirtualShadowMapClipmap> Clipmap = ProjectedShadowInfo->VirtualShadowMapClipmap;
+		FViewInfo* ShadowDepthView = ProjectedShadowInfo->ShadowDepthView;
 
-			ShadowDepthPassParameters->bClampToNearPlane = ProjectedShadowInfo->ShouldClampToNearPlane();
+		MeshCommandPass.WaitForSetupTask();
 
-			// TODO: These are not used for this case anyway
-			ShadowDepthPassParameters->ProjectionMatrix = FMatrix::Identity;
-			ShadowDepthPassParameters->ViewMatrix = FMatrix::Identity;
-			ShadowDepthPassParameters->ShadowParams = FVector4f(0.0f, 0.0f, 0.0f, 1.0f);
-			ShadowDepthPassParameters->bRenderToVirtualShadowMap = true;
-			
-			ShadowDepthPassParameters->VirtualSmPageTable	= GraphBuilder.CreateSRV( PageTableRDG );
-			ShadowDepthPassParameters->PackedNaniteViews	= GraphBuilder.CreateSRV( VirtualShadowViewsRDG );
-			ShadowDepthPassParameters->PageRectBounds		= GraphBuilder.CreateSRV( PageRectBoundsRDG );
-			ShadowDepthPassParameters->OutDepthBuffer		= GraphBuilder.CreateUAV( PhysicalPagePoolRDG, ERDGUnorderedAccessViewFlags::SkipBarrier );
-			
-			PassParameters->ShadowDepthPass = GraphBuilder.CreateUniformBuffer(ShadowDepthPassParameters);
-			PassParameters->VirtualShadowMap = GetUniformBuffer(GraphBuilder);
-			PassParameters->InViews = GraphBuilder.CreateSRV(VirtualShadowViewsRDG);
+		FInstanceCullingContext* InstanceCullingContext = MeshCommandPass.GetInstanceCullingContext();
 
-			FInstanceCullingResult InstanceCullingResult;
-			InstanceCullingResult.DrawIndirectArgsBuffer = DrawIndirectArgsRDG;
-			InstanceCullingResult.InstanceDataBuffer = InstanceIdOffsetBufferRDG;
+		if (InstanceCullingContext->HasCullingCommands())
+		{
+			FString LightNameWithLevel;
+			FSceneRenderer::GetLightNameForDrawEvent(ProjectedShadowInfo->GetLightSceneInfo().Proxy, LightNameWithLevel);
+			RDG_EVENT_SCOPE(GraphBuilder, "%s", *LightNameWithLevel);
 
-			FInstanceCullingGlobalUniforms* InstanceCullingGlobalUniforms = GraphBuilder.AllocParameters<FInstanceCullingGlobalUniforms>();
-			InstanceCullingGlobalUniforms->InstanceIdsBuffer = GraphBuilder.CreateSRV(InstanceIdsBuffer);
-			InstanceCullingGlobalUniforms->PageInfoBuffer = GraphBuilder.CreateSRV(PageInfoBuffer);
-			InstanceCullingGlobalUniforms->BufferCapacity = MaxNumInstancesPerPass;
-			InstanceCullingResult.UniformBuffer = GraphBuilder.CreateUniformBuffer(InstanceCullingGlobalUniforms);
+			CullingBatchInfo.DynamicInstanceIdOffset = ShadowDepthView->DynamicPrimitiveCollector.GetInstanceSceneDataOffset();
+			CullingBatchInfo.DynamicInstanceIdMax = CullingBatchInfo.DynamicInstanceIdOffset + ShadowDepthView->DynamicPrimitiveCollector.NumInstances();
 
-			InstanceCullingResult.GetDrawParameters(PassParameters->InstanceCullingDrawParams);
+			FCullingResult CullingResult = AddCullingPasses(
+				GraphBuilder,
+				InstanceCullingContext->IndirectArgs, 
+				InstanceCullingContext->DrawCommandDescs,
+				InstanceCullingContext->InstanceIdOffsets,
+				InstanceCullingContext->LoadBalancers[uint32(EBatchProcessingMode::Generic)],
+				MakeArrayView(&CullingBatchInfo, 1),
+				MakeArrayView(&VSMCullingBatchInfo, 1),
+				MakeArrayView<const uint32>(nullptr, 0),
+				!Clipmap.IsValid(),
+				InstanceCullingContext->TotalInstances,
+				TotalPrimaryViews,
+				VirtualShadowViewsRDG,
+				HZBShaderParameters,
+				this,
+				GPUScene
+			);
 
-			FIntRect ViewRect;
-			ViewRect.Max = FVirtualShadowMap::VirtualMaxResolutionXY;
+			TRDGUniformBufferRef<FShadowDepthPassUniformParameters> ShadowDepthPassUniformBuffer = CreateShadowDepthPassUniformBuffer(ProjectedShadowInfo->ShouldClampToNearPlane());
 
-			GraphBuilder.AddPass(
-				RDG_EVENT_NAME("RenderVirtualShadowMapsHw"),
-				PassParameters,
-				ERDGPassFlags::Raster | ERDGPassFlags::SkipRenderPass,
-				[this, ProjectedShadowInfo , &MeshCommandPass, PassParameters, ShadowDepthPassParameters, ViewRect](FRHICommandList& RHICmdList)
-				{
-					FRHIRenderPassInfo RPInfo;
-					RPInfo.ResolveParameters.DestRect.X1 = ViewRect.Min.X;
-					RPInfo.ResolveParameters.DestRect.Y1 = ViewRect.Min.Y;
-					RPInfo.ResolveParameters.DestRect.X2 = ViewRect.Max.X;
-					RPInfo.ResolveParameters.DestRect.Y2 = ViewRect.Max.Y;
-					RHICmdList.BeginRenderPass(RPInfo, TEXT("RenderVirtualShadowMapsHw"));
-
-					RHICmdList.SetViewport( ViewRect.Min.X, ViewRect.Min.Y, 0.0f, FMath::Min( ViewRect.Max.X, 32767 ), FMath::Min( ViewRect.Max.Y, 32767 ), 1.0f );
-
-					MeshCommandPass.DispatchDraw(nullptr, RHICmdList, &PassParameters->InstanceCullingDrawParams);
-					RHICmdList.EndRenderPass();
-				});
+			FInstanceCullingGlobalUniforms* InstanceCullingGlobalUniforms = GraphBuilder.AllocParameters<FInstanceCullingGlobalUniforms>();
+			InstanceCullingGlobalUniforms->InstanceIdsBuffer = GraphBuilder.CreateSRV(CullingResult.InstanceIdsBuffer);
+			InstanceCullingGlobalUniforms->PageInfoBuffer = GraphBuilder.CreateSRV(CullingResult.PageInfoBuffer);
+			InstanceCullingGlobalUniforms->BufferCapacity = CullingResult.MaxNumInstancesPerPass;
+			TRDGUniformBufferRef<FInstanceCullingGlobalUniforms> InstanceCullingUniformBuffer = GraphBuilder.CreateUniformBuffer(InstanceCullingGlobalUniforms);
+
+			FVirtualShadowDepthPassParameters* DepthPassParams = GraphBuilder.AllocParameters<FVirtualShadowDepthPassParameters>();
+			DepthPassParams->InstanceCullingDrawParams.IndirectArgsByteOffset = 0;
+			DepthPassParams->InstanceCullingDrawParams.InstanceDataByteOffset = 0;
+			AddRasterPass(GraphBuilder, RDG_EVENT_NAME("Rasterize"), ShadowDepthView, ShadowDepthPassUniformBuffer, this, VirtualShadowViewsRDG, CullingResult, MeshCommandPass, DepthPassParams, InstanceCullingUniformBuffer);
 		}
 
 
 		//
 		if (Index == CVarShowClipmapStats.GetValueOnRenderThread())
 		{
+			// The 'main' view the shadow was created with respect to
+			const FViewInfo* ViewUsedToCreateShadow = ProjectedShadowInfo->DependentView;
+			const FViewInfo& View = *ViewUsedToCreateShadow;
+
 			FVirtualSmPrintClipmapStatsCS::FParameters* PassParameters = GraphBuilder.AllocParameters<FVirtualSmPrintClipmapStatsCS::FParameters>();
 
 			ShaderPrint::SetParameters(GraphBuilder, View, PassParameters->ShaderPrintStruct);
@@ -2508,3 +2650,108 @@ FRDGTextureRef FVirtualShadowMapArray::BuildHZBFurthest(FRDGBuilder& GraphBuilde
 	}
 	return OutFurthestHZBTexture;
 }
+
+
+uint32 FVirtualShadowMapArray::AddRenderViews(const TSharedPtr<FVirtualShadowMapClipmap>& Clipmap, float LODScaleFactor, bool bSetHzbParams, bool bUpdateHZBMetaData, TArray<Nanite::FPackedView, SceneRenderingAllocator> &OutVirtualShadowViews)
+{
+	// TODO: Decide if this sort of logic belongs here or in Nanite (as with the mip level view expansion logic)
+	// We're eventually going to want to snap/quantize these rectangles/positions somewhat so probably don't want it
+	// entirely within Nanite, but likely makes sense to have some sort of "multi-viewport" notion in Nanite that can
+	// handle both this and mips.
+	// NOTE: There's still the additional VSM view logic that runs on top of this in Nanite too (see CullRasterize variant)
+	Nanite::FPackedViewParams BaseParams;
+	BaseParams.ViewRect = FIntRect(0, 0, FVirtualShadowMap::VirtualMaxResolutionXY, FVirtualShadowMap::VirtualMaxResolutionXY);
+	BaseParams.HZBTestViewRect = BaseParams.ViewRect;
+	BaseParams.RasterContextSize = GetPhysicalPoolSize();
+	BaseParams.LODScaleFactor = LODScaleFactor;
+	BaseParams.PrevTargetLayerIndex = INDEX_NONE;
+	BaseParams.TargetMipLevel = 0;
+	BaseParams.TargetMipCount = 1;	// No mips for clipmaps
+
+	for (int32 ClipmapLevelIndex = 0; ClipmapLevelIndex < Clipmap->GetLevelCount(); ++ClipmapLevelIndex)
+	{
+		FVirtualShadowMap* VirtualShadowMap = Clipmap->GetVirtualShadowMap(ClipmapLevelIndex);
+
+		Nanite::FPackedViewParams Params = BaseParams;
+		Params.TargetLayerIndex = VirtualShadowMap->ID;
+		Params.ViewMatrices = Clipmap->GetViewMatrices(ClipmapLevelIndex);
+		Params.PrevTargetLayerIndex = INDEX_NONE;
+		Params.PrevViewMatrices = Params.ViewMatrices;
+		Params.Flags = 0;
+
+		// TODO: Clean this up - could be stored in a single structure for the whole clipmap
+		int32 HZBKey = Clipmap->GetHZBKey(ClipmapLevelIndex);
+
+		if (bSetHzbParams)
+		{
+			CacheManager->SetHZBViewParams(HZBKey, Params);
+		}
+
+		// If we're going to generate a new HZB this frame, save the associated metadata
+		if (bUpdateHZBMetaData)
+		{
+			FVirtualShadowMapHZBMetadata& HZBMeta = HZBMetadata.FindOrAdd(HZBKey);
+			HZBMeta.TargetLayerIndex = Params.TargetLayerIndex;
+			HZBMeta.ViewMatrices = Params.ViewMatrices;
+			HZBMeta.ViewRect = Params.ViewRect;
+		}
+
+		Nanite::FPackedView View = Nanite::CreatePackedView(Params);
+		OutVirtualShadowViews.Add(View);
+
+		// Mark that we rendered to this VSM for caching purposes
+		if (VirtualShadowMap->VirtualShadowMapCacheEntry)
+		{
+			VirtualShadowMap->VirtualShadowMapCacheEntry->MarkRendered();
+		}
+	}
+
+	return uint32(Clipmap->GetLevelCount());
+}
+
+uint32 FVirtualShadowMapArray::AddRenderViews(const FProjectedShadowInfo* ProjectedShadowInfo, float LODScaleFactor, bool bSetHzbParams, bool bUpdateHZBMetaData, TArray<Nanite::FPackedView, SceneRenderingAllocator>& OutVirtualShadowViews)
+{
+	Nanite::FPackedViewParams BaseParams;
+	BaseParams.ViewRect = ProjectedShadowInfo->GetOuterViewRect();
+	BaseParams.HZBTestViewRect = BaseParams.ViewRect;
+	BaseParams.RasterContextSize = GetPhysicalPoolSize();
+	BaseParams.LODScaleFactor = LODScaleFactor;
+	BaseParams.PrevTargetLayerIndex = INDEX_NONE;
+	BaseParams.TargetMipLevel = 0;
+	BaseParams.TargetMipCount = FVirtualShadowMap::MaxMipLevels;
+
+	int32 NumMaps = ProjectedShadowInfo->bOnePassPointLightShadow ? 6 : 1;
+	for (int32 Index = 0; Index < NumMaps; ++Index)
+	{
+		FVirtualShadowMap* VirtualShadowMap = ProjectedShadowInfo->VirtualShadowMaps[Index];
+
+		Nanite::FPackedViewParams Params = BaseParams;
+		Params.TargetLayerIndex = VirtualShadowMap->ID;
+		Params.ViewMatrices = ProjectedShadowInfo->GetShadowDepthRenderingViewMatrices(Index, true);
+
+		int32 HZBKey = ProjectedShadowInfo->GetLightSceneInfo().Id + (Index << 24);
+
+		if (bSetHzbParams)
+		{
+			CacheManager->SetHZBViewParams(HZBKey, Params);
+		}
+
+		// If we're going to generate a new HZB this frame, save the associated metadata
+		if (bUpdateHZBMetaData)
+		{
+			FVirtualShadowMapHZBMetadata& HZBMeta = HZBMetadata.FindOrAdd(HZBKey);
+			HZBMeta.TargetLayerIndex = Params.TargetLayerIndex;
+			HZBMeta.ViewMatrices = Params.ViewMatrices;
+			HZBMeta.ViewRect = Params.ViewRect;
+		}
+
+		OutVirtualShadowViews.Add(Nanite::CreatePackedView(Params));
+
+		if (VirtualShadowMap->VirtualShadowMapCacheEntry)
+		{
+			VirtualShadowMap->VirtualShadowMapCacheEntry->MarkRendered();
+		}
+	}
+
+	return uint32(NumMaps);
+}
diff --git a/Engine/Source/Runtime/Renderer/Private/VirtualShadowMaps/VirtualShadowMapArray.h b/Engine/Source/Runtime/Renderer/Private/VirtualShadowMaps/VirtualShadowMapArray.h
--- a/Engine/Source/Runtime/Renderer/Private/VirtualShadowMaps/VirtualShadowMapArray.h
+++ b/Engine/Source/Runtime/Renderer/Private/VirtualShadowMaps/VirtualShadowMapArray.h
@@ -16,6 +16,7 @@ class FVisibleLightInfo;
 class FVirtualShadowMapCacheEntry;
 class FVirtualShadowMapArrayCacheManager;
 struct FSortedLightSetSceneInfo;
+class FVirtualShadowMapClipmap;
 
 // TODO: does this exist?
 constexpr uint32 ILog2Const(uint32 n)
@@ -215,9 +216,9 @@ public:
 	void CreateMipViews( TArray<Nanite::FPackedView, SceneRenderingAllocator>& Views ) const;
 
 	/**
-	 * Draw old-school hardware based shadow map tiles into virtual SM.
+	 * Draw Non-Nanite geometry into the VSMs.
 	 */
-	void RenderVirtualShadowMapsHw(FRDGBuilder& GraphBuilder, const TArray<FProjectedShadowInfo*, SceneRenderingAllocator>& VirtualSmMeshCommandPasses, FScene& Scene);
+	void RenderVirtualShadowMapsNonNanite(FRDGBuilder& GraphBuilder, const TArray<FProjectedShadowInfo*, SceneRenderingAllocator>& VirtualSmMeshCommandPasses, FScene& Scene);
 
 	// Draw debug info into render target 'VSMDebug' of screen-size, the mode is controlled by 'r.Shadow.Virtual.DebugVisualize'.
 	void RenderDebugInfo(FRDGBuilder& GraphBuilder);
@@ -239,6 +240,10 @@ public:
 
 	FRDGTextureRef BuildHZBFurthest(FRDGBuilder& GraphBuilder);
 
+	// Add render views, and mark shadow maps as rendered for a given clipmap or set of VSMs, returns the number of primary views added.
+	uint32 AddRenderViews(const TSharedPtr<FVirtualShadowMapClipmap>& Clipmap, float LODScaleFactor, bool bSetHzbParams, bool bUpdateHZBMetaData, TArray<Nanite::FPackedView, SceneRenderingAllocator>& OutVirtualShadowViews);
+	uint32 AddRenderViews(const FProjectedShadowInfo* ProjectedShadowInfo, float LODScaleFactor, bool bSetHzbParams, bool bUpdateHZBMetaData, TArray<Nanite::FPackedView, SceneRenderingAllocator>& OutVirtualShadowViews);
+
 	// We keep a reference to the cache manager that was used to initialize this frame as it owns some of the buffers
 	FVirtualShadowMapArrayCacheManager* CacheManager = nullptr;
 
diff --git a/Engine/Source/Runtime/Renderer/Public/InstanceCulling/InstanceCullingContext.h b/Engine/Source/Runtime/Renderer/Public/InstanceCulling/InstanceCullingContext.h
--- a/Engine/Source/Runtime/Renderer/Public/InstanceCulling/InstanceCullingContext.h
+++ b/Engine/Source/Runtime/Renderer/Public/InstanceCulling/InstanceCullingContext.h
@@ -60,6 +60,7 @@ class FInstanceCullingContext
 public:
 	static constexpr uint32 IndirectArgsNumWords = 5;
 	RENDERER_API static uint32 GetInstanceIdBufferStride(ERHIFeatureLevel::Type FeatureLevel);
+	RENDERER_API static uint32 StepInstanceDataOffset(ERHIFeatureLevel::Type FeatureLevel, uint32 NumStepInstances, uint32 NumStepDraws);
 
 	FInstanceCullingContext() {}
 
@@ -83,14 +84,6 @@ public:
 
 	static bool IsOcclusionCullingEnabled();
 
-	struct FBatchItem
-	{
-		const FInstanceCullingContext* Context = nullptr;
-		FInstanceCullingDrawParams* Result = nullptr;
-		int32 DynamicInstanceIdOffset = 0;
-		int32 DynamicInstanceIdNum = 0;
-	};
-
 	/**
 	 * Call to empty out the culling commands & other culling data.
 	 */
@@ -192,19 +185,6 @@ public:
 		uint32 bMaterialMayModifyPosition;
 	};
 
-	// Info about a batch of culling work produced by a context, when part of a batched job
-	// Store once per context, provides start offsets to commands/etc for the context.
-	struct FContextBatchInfo
-	{
-		uint32 IndirectArgsOffset;
-		uint32 InstanceDataWriteOffset;
-		uint32 ViewIdsOffset;
-		uint32 NumViewIds;
-		uint32 DynamicInstanceIdOffset;
-		uint32 DynamicInstanceIdMax;
-		uint32 ItemDataOffset[uint32(EBatchProcessingMode::Num)];
-	};
-	
 	TArray<FMeshDrawCommandInfo> MeshDrawCommandInfos;
 	TArray<FRHIDrawIndexedIndirectParameters> IndirectArgs;
 	TArray<FDrawCommandDesc> DrawCommandDescs;
