Reduced TVolumetricFogLightScatteringCS::FSuperSampleCount choices to achieve a more reasonable number of permutations: 1024 -> 128

diff --git a/Engine/Source/Runtime/Renderer/Private/VolumetricFog.cpp b/Engine/Source/Runtime/Renderer/Private/VolumetricFog.cpp
--- a/Engine/Source/Runtime/Renderer/Private/VolumetricFog.cpp
+++ b/Engine/Source/Runtime/Renderer/Private/VolumetricFog.cpp
@@ -797,7 +797,7 @@ class TVolumetricFogLightScatteringCS : public FGlobalShader
 
 	class FTemporalReprojection			: SHADER_PERMUTATION_BOOL("USE_TEMPORAL_REPROJECTION");
 	class FDistanceFieldSkyOcclusion	: SHADER_PERMUTATION_BOOL("DISTANCE_FIELD_SKY_OCCLUSION");
-	class FSuperSampleCount				: SHADER_PERMUTATION_RANGE_INT("HISTORY_MISS_SUPER_SAMPLE_COUNT", 1, 16);
+	class FSuperSampleCount				: SHADER_PERMUTATION_SPARSE_INT("HISTORY_MISS_SUPER_SAMPLE_COUNT", 1, 4, 8, 16);
 	class FLumenGI						: SHADER_PERMUTATION_BOOL("LUMEN_GI");
 	class FVirtualShadowMap				: SHADER_PERMUTATION_BOOL("VIRTUAL_SHADOW_MAP");
 	class FCloudTransmittance			: SHADER_PERMUTATION_BOOL("USE_CLOUD_TRANSMITTANCE");
@@ -840,6 +840,24 @@ public:
 		return FIntVector(4, 4, 4);
 	}
 
+	static int32 GetSuperSampleCount(int32 InSampleCount)
+	{
+		if (InSampleCount <= 1)
+		{
+			return 1;
+		}
+		else if (InSampleCount <= 4)
+		{
+			return 4;
+		}
+		else if (InSampleCount <= 8)
+		{
+			return 8;
+		}
+		
+		return 16;
+	}
+
 	static void ModifyCompilationEnvironment(const FGlobalShaderPermutationParameters& Parameters, FShaderCompilerEnvironment& OutEnvironment)
 	{
 		FGlobalShader::ModifyCompilationEnvironment(Parameters, OutEnvironment);
@@ -1494,10 +1512,12 @@ void FDeferredShadingSceneRenderer::ComputeVolumetricFog(FRDGBuilder& GraphBuild
 				&& Views.Num() == 1
 				&& View.IsPerspectiveProjection();
 
+			const int32 SuperSampleCount = TVolumetricFogLightScatteringCS::GetSuperSampleCount(GVolumetricFogHistoryMissSupersampleCount);
+
 			TVolumetricFogLightScatteringCS::FPermutationDomain PermutationVector;
 			PermutationVector.Set< TVolumetricFogLightScatteringCS::FTemporalReprojection >(bUseTemporalReprojection);
 			PermutationVector.Set< TVolumetricFogLightScatteringCS::FDistanceFieldSkyOcclusion >(bUseDistanceFieldSkyOcclusion);
-			PermutationVector.Set< TVolumetricFogLightScatteringCS::FSuperSampleCount >(GVolumetricFogHistoryMissSupersampleCount);
+			PermutationVector.Set< TVolumetricFogLightScatteringCS::FSuperSampleCount >(SuperSampleCount);
 			PermutationVector.Set< TVolumetricFogLightScatteringCS::FLumenGI >(bUseLumenGI);
 			PermutationVector.Set< TVolumetricFogLightScatteringCS::FVirtualShadowMap >( VirtualShadowMapArray.IsAllocated() );
 			PermutationVector.Set< TVolumetricFogLightScatteringCS::FCloudTransmittance >(AtmosphericDirectionalLightIndex >= 0);
@@ -1518,7 +1538,7 @@ void FDeferredShadingSceneRenderer::ComputeVolumetricFog(FRDGBuilder& GraphBuild
 					VolumetricFogGridSize.X,
 					VolumetricFogGridSize.Y,
 					VolumetricFogGridSize.Z,
-					GVolumetricFogHistoryMissSupersampleCount,
+					SuperSampleCount,
 					bUseDistanceFieldSkyOcclusion ? TEXT("DFAO") : TEXT(""),
 					PassParameters->LightFunctionTexture ? TEXT("LF") : TEXT(""),
 					bUseLumenGI ? TEXT("Lumen") : TEXT("")),
