Moves the bloom's intensity to be applied before writing to the bloom rdg texture

This allows lens flares to work with FFT bloom without introducing additional memory bandwidth costs in the tonemapper.

diff --git a/Engine/Shaders/Private/Bloom/BloomCommon.ush b/Engine/Shaders/Private/Bloom/BloomCommon.ush
--- a/Engine/Shaders/Private/Bloom/BloomCommon.ush
+++ b/Engine/Shaders/Private/Bloom/BloomCommon.ush
@@ -28,9 +28,3 @@ struct FBloomKernelInfo
 	// Some padding to make vulkan happy
 	uint2 _Padding;
 };
-
-struct FBloomApplyInfo
-{
-	float4 SceneColorMultiply;
-	float4 BloomMultiply;
-};
diff --git a/Engine/Shaders/Private/Bloom/BloomFinalizeApplyConstants.usf b/Engine/Shaders/Private/Bloom/BloomFinalizeApplyConstants.usf
--- a/Engine/Shaders/Private/Bloom/BloomFinalizeApplyConstants.usf
+++ b/Engine/Shaders/Private/Bloom/BloomFinalizeApplyConstants.usf
@@ -8,7 +8,8 @@
 float ScatterDispersionIntensity;
 
 StructuredBuffer<FBloomKernelInfo> KernelConstantsBuffer;
-RWStructuredBuffer<FBloomApplyInfo> BloomApplyConstantsOutput;
+RWStructuredBuffer<float4> SceneColorApplyOutput;
+RWStructuredBuffer<float4> FFTMulitplyOutput;
 
 
 //------------------------------------------------------- ENTRY POINT
@@ -36,6 +37,6 @@ void MainCS(
 	
 	float4 FinalCenterEnergy = OriginalTotalEnergy - FinalScatterDispersionEnergy;
 
-	BloomApplyConstantsOutput[0].SceneColorMultiply = Tint * saturate(FinalCenterEnergy / OriginalTotalEnergy);
-	BloomApplyConstantsOutput[0].BloomMultiply = Tint * saturate(FinalScatterDispersionEnergy / OriginalTotalEnergy);
+	SceneColorApplyOutput[0] = Tint * saturate(FinalCenterEnergy / OriginalTotalEnergy);
+	FFTMulitplyOutput[0] = Tint * saturate(FinalScatterDispersionEnergy / OriginalTotalEnergy);
 }
diff --git a/Engine/Shaders/Private/GPUFastFourierTransform.usf b/Engine/Shaders/Private/GPUFastFourierTransform.usf
--- a/Engine/Shaders/Private/GPUFastFourierTransform.usf
+++ b/Engine/Shaders/Private/GPUFastFourierTransform.usf
@@ -322,6 +322,9 @@ uint TransformType;
 uint2 SrcRectMin;
 uint2 SrcRectMax;
 uint4 DstRect;
+
+StructuredBuffer<float4> DstPostFilterParameters;
+
 [numthreads(NUMTHREADSX, 1, 1)]
 void GroupSharedTwoForOneFFTCS(uint3 GroupID : SV_GroupID, uint3 GroupThreadID : SV_GroupThreadID )
 {
@@ -402,6 +405,15 @@ void GroupSharedTwoForOneFFTCS(uint3 GroupID : SV_GroupID, uint3 GroupThreadID :
 		// This is specialized for images, where floating point errors may have
 		// resulted in (very very small) negative color values.
 		ScrubNANs(LocalBuffer);
+		
+		float4 Scale = DstPostFilterParameters[0];
+
+		UNROLL
+		for (uint r = 0; r < RADIX; r++)
+		{
+			LocalBuffer[0][r] *= Scale.xy;
+			LocalBuffer[1][r] *= Scale.xw;
+		}
 
 		CopyDataLocalToDstWindow(LocalBuffer, bIsHorizontal, ScanIdx, Head, Stride, DstRect);
 	}
@@ -999,6 +1011,7 @@ uint TransformType;
 uint BitCount;  // log(2, TransformLength) + 1;
 uint PowTwoLength;  // Power of two.  2^p where p is the pass number
 
+StructuredBuffer<float4> DstPostFilterParameters;
 
 [numthreads(NUMTHREADS_PER_COL, 1, 1)]
 void ComplexFFTPassCS(uint3 GroupID : SV_GroupID, uint3 GroupThreadID : SV_GroupThreadID )
@@ -1068,7 +1081,9 @@ void ComplexFFTPassCS(uint3 GroupID : SV_GroupID, uint3 GroupThreadID : SV_Group
 	float4 SrcValueB;
 
 	const uint TransformLength = 1 << (BitCount-1);
-	const float Scale = (Ns == 1 && !bIsForward) ? 1.f / float(TransformLength) : 1.f;
+	
+	float4 Scale = DstPostFilterParameters[0] * ((Ns == 1 && !bIsForward) ? 1.f / float(TransformLength) : 1.f);
+
 	if (bIsHorizontal)
 	{
 		uint CurRowIdx = RowIdx;
diff --git a/Engine/Shaders/Private/PostProcessMobile.usf b/Engine/Shaders/Private/PostProcessMobile.usf
--- a/Engine/Shaders/Private/PostProcessMobile.usf
+++ b/Engine/Shaders/Private/PostProcessMobile.usf
@@ -1056,6 +1056,7 @@ void SunMergeVS_Mobile(
 
 float4 SunColorVignetteIntensity;
 float3 BloomColor;
+float BloomIntensity;
 
 Texture2D BloomDirtMaskTexture;
 SamplerState BloomDirtMaskSampler;
@@ -1113,6 +1114,8 @@ void SunMergePS_Mobile(
 		OutColor.rgb += SunColorVignetteIntensity.rgb * Sun;
 	#endif
 
+	OutColor.rgb *= BloomIntensity;
+
 	OutColor.a = 1.0f;
 }
 
diff --git a/Engine/Shaders/Private/PostProcessTonemap.usf b/Engine/Shaders/Private/PostProcessTonemap.usf
--- a/Engine/Shaders/Private/PostProcessTonemap.usf
+++ b/Engine/Shaders/Private/PostProcessTonemap.usf
@@ -6,8 +6,8 @@
 
 #define EYE_ADAPTATION_LOOSE_PARAMETERS 1
 
-#ifndef SUPPORTS_BLOOM_APPLY_PARAMETERS
-	#define SUPPORTS_BLOOM_APPLY_PARAMETERS 0
+#ifndef SUPPORTS_SCENE_COLOR_APPLY_PARAMETERS
+	#define Undefined SUPPORTS_SCENE_COLOR_APPLY_PARAMETERS 0
 #endif
 
 
@@ -18,9 +18,7 @@
 #include "PostProcessHistogramCommon.ush"
 #include "PixelQuadMessagePassing.ush"
 #include "ScreenPass.ush"
-#if SUPPORTS_BLOOM_APPLY_PARAMETERS
-	#include "Bloom/BloomCommon.ush"
-#endif
+
 
 /** Values of DIM_OUTPUT_DEVICE (matching C++'s ETonemapperOutputDevice) */
 #define TONEMAPPER_OUTPUT_sRGB 0
@@ -55,12 +53,11 @@ SamplerState ColorSampler;
 FScreenTransform ColorToBloom;
 float2 BloomUVViewportBilinearMin;
 float2 BloomUVViewportBilinearMax;
-float4 BloomMultiply;
 Texture2D BloomTexture;
 SamplerState BloomSampler;
 
-#if SUPPORTS_BLOOM_APPLY_PARAMETERS
-	StructuredBuffer<FBloomApplyInfo> BloomApplyParameters;
+#if SUPPORTS_SCENE_COLOR_APPLY_PARAMETERS
+	StructuredBuffer<float4> SceneColorApplyParamaters;
 #endif
 
 Texture3D LumBilateralGrid;
@@ -583,6 +580,13 @@ float4 TonemapCommonPS(
 	
 		half3 LinearColor = SceneColor.rgb * ColorScale0.rgb;
 
+		#if SUPPORTS_SCENE_COLOR_APPLY_PARAMETERS
+		{
+			float4 SceneColorMultiply = SceneColorApplyParamaters[0];
+			LinearColor *= SceneColorMultiply;
+		}
+		#endif
+
 		#if USE_BLOOM
 		{
 			float2 BloomUV;
@@ -628,30 +632,15 @@ float4 TonemapCommonPS(
 
 				float3 BloomDirtMaskColor = Texture2DSample(BloomDirtMaskTexture, BloomDirtMaskSampler, DirtLensUV * .5f + .5f).rgb * BloomDirtMaskTint.rgb;
 				
-				#if SUPPORTS_BLOOM_APPLY_PARAMETERS
-				{
-					FBloomApplyInfo BloomApplyInfo = BloomApplyParameters[0];
-
-					LinearColor = (
-						LinearColor.rgb * BloomApplyInfo.SceneColorMultiply.rgb +
-						CombinedBloom.rgb * (BloomApplyInfo.BloomMultiply.rgb + BloomDirtMaskColor));
-					
-					SceneColor.a = (
-						SceneColor.a * BloomApplyInfo.SceneColorMultiply.a +
-						CombinedBloom.a * BloomApplyInfo.BloomMultiply.a);
-				}
-				#else
-				{
-					LinearColor += CombinedBloom.rgb * (BloomMultiply.rgb + BloomDirtMaskColor);
-				}
-				#endif
+				LinearColor += CombinedBloom.rgb * (1.0 + BloomDirtMaskColor);
 			}
 			#endif
 		}
 		#endif
 
+		// TODO: doesn't need a shader permutation for a scalar.
 		#if NO_EYEADAPTATION_EXPOSURE_FIX == 1
-			ExposureScale = BloomDirtMaskTint.w;
+			ExposureScale = 1.0;
 		#endif
 		
 		LinearColor *= ExposureScale;
diff --git a/Engine/Source/Runtime/Renderer/Private/GPUFastFourierTransform.cpp b/Engine/Source/Runtime/Renderer/Private/GPUFastFourierTransform.cpp
--- a/Engine/Source/Runtime/Renderer/Private/GPUFastFourierTransform.cpp
+++ b/Engine/Source/Runtime/Renderer/Private/GPUFastFourierTransform.cpp
@@ -5,6 +5,7 @@
 #include "GlobalShader.h"
 #include "RenderGraph.h" 
 #include "ShaderCompilerCore.h"
+#include "SystemTextures.h"
 
 bool ShouldCompileGPUFFT(EShaderPlatform Platform)
 {
@@ -223,6 +224,7 @@ class FComplexFFTPassCS : public FFFTShader
 		SHADER_PARAMETER(int32, BitCount)
 
 		SHADER_PARAMETER_RDG_TEXTURE(Texture2D, SrcSRV)
+		SHADER_PARAMETER_RDG_BUFFER_SRV(StructuredBuffer<float4>, DstPostFilterParameters)
 		SHADER_PARAMETER_RDG_TEXTURE_UAV(RWTexture2D, DstUAV)
 	END_SHADER_PARAMETER_STRUCT()
 
@@ -349,6 +351,7 @@ class FGSTwoForOneTransformCS : public FFFTShader
 		SHADER_PARAMETER(int32, TransformType)
 
 		SHADER_PARAMETER_RDG_TEXTURE(Texture2D, SrcTexture)
+		SHADER_PARAMETER_RDG_BUFFER_SRV(StructuredBuffer<float4>, DstPostFilterParameters)
 		SHADER_PARAMETER_RDG_TEXTURE_UAV(RWTexture2D, DstTexture)
 	END_SHADER_PARAMETER_STRUCT()
 
@@ -540,8 +543,8 @@ void DispatchComplexFFTPassCS(
 	const GPUFFT::FFTDescription& FFTDesc,
 	const uint32 PassLength,
 	FRDGTextureRef SrcTexture, const FIntRect& SrcWindow,
-	FRDGTextureRef DstTexture,
-	const bool bScrubNaNs = false)
+	FRDGBufferRef DstPostFilterParameters,
+	FRDGTextureRef DstTexture)
 {
 	// Using multiple radix two passes.
 	const uint32 Radix = 2;
@@ -550,9 +553,9 @@ void DispatchComplexFFTPassCS(
 
 	// Translate the transform type. 
 	uint32 TransformTypeValue = GPUFFT::BitEncode(FFTDesc.XFormType);
-	if (bScrubNaNs)
+	if (!DstPostFilterParameters)
 	{
-		TransformTypeValue |= 4;
+		DstPostFilterParameters = GSystemTextures.GetDefaultStructuredBuffer(GraphBuilder, sizeof(FVector4f), FVector4f(1.0f, 1.0f, 1.0f, 1.0f));
 	}
 
 	FComplexFFTPassCS::FParameters* PassParameters = GraphBuilder.AllocParameters<FComplexFFTPassCS::FParameters>();
@@ -563,6 +566,7 @@ void DispatchComplexFFTPassCS(
 	PassParameters->BitCount = GPUFFT::BitSize(TransformLength);
 
 	PassParameters->SrcSRV = SrcTexture;
+	PassParameters->DstPostFilterParameters = GraphBuilder.CreateSRV(DstPostFilterParameters);
 	PassParameters->DstUAV = GraphBuilder.CreateUAV(DstTexture);
 
 	TShaderMapRef<FComplexFFTPassCS> ComputeShader(ShaderMap);
@@ -816,13 +820,19 @@ void DispatchGSTwoForOneFFTCS(
 	const GPUFFT::FFTDescription& FFTDesc,
 	FRDGTextureRef SrcTexture, const FIntRect& SrcRect,
 	FRDGTextureRef DstTexture, const FIntRect& DstRect,
-	const GPUFFT::FPreFilter& PreFilter)
+	const GPUFFT::FPreFilter& PreFilter,
+	FRDGBufferRef DstPostFilterParameters)
 {
 	const uint32 TransformLength = FFTDesc.SignalLength;
 	
 	// The number of signals to transform simultaneously (i.e. number of scan lines)
 	const uint32 NumScanLines = (FFTDesc.IsHorizontal()) ? SrcRect.Size().Y : SrcRect.Size().X;
 
+	if (!DstPostFilterParameters)
+	{
+		DstPostFilterParameters = GSystemTextures.GetDefaultStructuredBuffer(GraphBuilder, sizeof(FVector4f), FVector4f(1.0f, 1.0f, 1.0f, 1.0f));
+	}
+
 	FGSTwoForOneTransformCS::FParameters* PassParameters = GraphBuilder.AllocParameters<FGSTwoForOneTransformCS::FParameters>();
 	PassParameters->SrcRectMin = SrcRect.Min;
 	PassParameters->SrcRectMax = SrcRect.Max;
@@ -842,6 +852,7 @@ void DispatchGSTwoForOneFFTCS(
 	}
 
 	PassParameters->SrcTexture = SrcTexture;
+	PassParameters->DstPostFilterParameters = GraphBuilder.CreateSRV(DstPostFilterParameters);
 	PassParameters->DstTexture = GraphBuilder.CreateUAV(DstTexture);
 
 	FGSTwoForOneTransformCS::FPermutationDomain PermutationVector;
@@ -949,6 +960,7 @@ void GPUFFT::ComplexFFTImage1D::MultiPass(
 	const FFTDescription& FFTDesc,
 	FRDGTextureRef SrcTexture, const FIntRect& SrcWindow,
 	FRDGTextureRef DstTexture,
+	FRDGBufferRef PostFilterParameters,
 	const bool bScrubNaNs)
 {
 	if (FitsInGroupSharedMemory(FFTDesc))
@@ -1030,6 +1042,7 @@ void GPUFFT::ComplexFFTImage1D::MultiPass(
 		for (uint32 Ns = FGroupShardSubFFTPassCS::SubPassLength(); Ns < TransformLength; Ns *= 2)
 		{
 			FRDGTextureRef OpDstTexture;
+			FRDGBufferRef DstPostFilterParameters = nullptr;
 			if (Ns * 2 < TransformLength)
 			{
 				FIntPoint SrcBufferSize = PrevSpectralTexture->Desc.Extent;
@@ -1040,12 +1053,14 @@ void GPUFFT::ComplexFFTImage1D::MultiPass(
 			else
 			{
 				OpDstTexture = DstTexture;
+				DstPostFilterParameters = PostFilterParameters;
 			}
 
 			DispatchComplexFFTPassCS(
 				GraphBuilder, ShaderMap, FFTDesc,
 				Ns,
 				/* SrcTexture = */ PrevSpectralTexture, XFormWindow,
+				/* DstPostFilterParameters = */ DstPostFilterParameters,
 				/* DstTexture = */ OpDstTexture);
 
 			PrevSpectralTexture = OpDstTexture;
@@ -1061,11 +1076,17 @@ void GPUFFT::TwoForOneRealFFTImage1D::MultiPass(
 	const FFTDescription& FFTDesc,
 	FRDGTextureRef SrcTexture, const FIntRect& SrcWindow,
 	FRDGTextureRef DstTexture, const FIntRect& DstWindow,
-	const FPreFilter& PreFilter)
+	const FPreFilter& PreFilter,
+	FRDGBufferRef PostFilterParameters)
 {
 	if (FitsInGroupSharedMemory(FFTDesc))
 	{
-		DispatchGSTwoForOneFFTCS(GraphBuilder, ShaderMap, FFTDesc, SrcTexture, SrcWindow, DstTexture, DstWindow, PreFilter);
+		DispatchGSTwoForOneFFTCS(
+			GraphBuilder, ShaderMap, FFTDesc,
+			SrcTexture, SrcWindow,
+			DstTexture, DstWindow,
+			PreFilter,
+			PostFilterParameters);
 	}
 	else if (FFTDesc.IsForward())
 	{
@@ -1096,6 +1117,7 @@ void GPUFFT::TwoForOneRealFFTImage1D::MultiPass(
 			GraphBuilder, ShaderMap, FFTDesc,
 			/* SrcTexture = */ SrcTexture, SrcWindow,
 			/* DstTexture = */ SpectralTexture,
+			/* PostFilterParameters = */ nullptr,
 			/* bScrubNaNs = */ true);
 		
 		// Unpack the complex transform into transform of real data
@@ -1122,7 +1144,9 @@ void GPUFFT::TwoForOneRealFFTImage1D::MultiPass(
 		ComplexFFTImage1D::MultiPass(
 			GraphBuilder, ShaderMap, FFTDesc,
 			/* SrcTexture = */ SpectralTexture, SrcWindow,
-			/* DstTexture = */ DstTexture);
+			/* DstTexture = */ DstTexture,
+			/* PostFilterParameters = */ PostFilterParameters,
+			/* bScrubNaNs = */ false);
 	}
 }
 
@@ -1271,7 +1295,8 @@ void GPUFFT::ConvolutionWithTextureImage2D(
 	FRDGTextureRef TransformedKernel,
 	FRDGTextureRef SrcTexture, const FIntRect& ROIRect,
 	FRDGTextureRef DstTexture, const FIntRect& DstRect,
-	const FPreFilter& PreFilter)
+	const FPreFilter& PreFilter,
+	FRDGBufferRef PostFilterParameters)
 {
 	FRDGTextureDesc SpectralDesc = CreateFrequencyDesc(FrequencySize, bHorizontalFirst, ROIRect.Size());
 
@@ -1313,7 +1338,8 @@ void GPUFFT::ConvolutionWithTextureImage2D(
 		GraphBuilder, ShaderMap, TwoForOneFFTDesc,
 		SrcTexture, ROIRect,
 		TwoForOneFFTOutput, SpectralOutputRect,
-		PreFilter);
+		PreFilter,
+		/* PostFilterParameters = */ nullptr);
 
 	// ---- 1 D Convolution --- 
 
@@ -1333,5 +1359,6 @@ void GPUFFT::ConvolutionWithTextureImage2D(
 		TwoForOneIvnFFTDesc,
 		ConvolutionOutput, SpectralOutputRect,
 		DstTexture, DstRect,
-		PreFilter);
+		PreFilter,
+		PostFilterParameters);
 }
diff --git a/Engine/Source/Runtime/Renderer/Private/GPUFastFourierTransform.h b/Engine/Source/Runtime/Renderer/Private/GPUFastFourierTransform.h
--- a/Engine/Source/Runtime/Renderer/Private/GPUFastFourierTransform.h
+++ b/Engine/Source/Runtime/Renderer/Private/GPUFastFourierTransform.h
@@ -131,6 +131,7 @@ namespace GPUFFT
 			const FFTDescription& FFTDesc,
 			FRDGTextureRef SrcTexture, const FIntRect& SrcWindow,
 			FRDGTextureRef DstTexture,
+			FRDGBufferRef PostFilterParameters = nullptr,
 			const bool bScrubNaNs = false);
 
 	};
@@ -180,7 +181,8 @@ namespace GPUFFT
 			const FFTDescription& FFTDesc,
 			FRDGTextureRef SrcTexture, const FIntRect& SrcWindow,
 			FRDGTextureRef DstTexture, const FIntRect& DstWindow,
-			const FPreFilter& PreFilter = FPreFilter(TNumericLimits<float>::Max(), TNumericLimits<float>::Lowest(), 0.f));
+			const FPreFilter& PreFilter = FPreFilter(TNumericLimits<float>::Max(), TNumericLimits<float>::Lowest(), 0.f),
+			FRDGBufferRef PostFilterParameters = nullptr);
 
 	};
 
@@ -288,7 +290,8 @@ namespace GPUFFT
 		FRDGTextureRef TransformedKernel,
 		FRDGTextureRef SrcTexture, const FIntRect& ROIRect,
 		FRDGTextureRef DstTexture, const FIntRect& DstRect,
-		const FPreFilter& PreFilter);
+		const FPreFilter& PreFilter,
+		FRDGBufferRef PostFilterParameters);
 
 	FIntPoint Convolution2DBufferSize(const FIntPoint& FrequencySize, const bool bHorizontalFirst, const FIntPoint& SrcExtent);
 
diff --git a/Engine/Source/Runtime/Renderer/Private/PostProcess/PostProcessBloomSetup.cpp b/Engine/Source/Runtime/Renderer/Private/PostProcess/PostProcessBloomSetup.cpp
--- a/Engine/Source/Runtime/Renderer/Private/PostProcess/PostProcessBloomSetup.cpp
+++ b/Engine/Source/Runtime/Renderer/Private/PostProcess/PostProcessBloomSetup.cpp
@@ -109,11 +109,6 @@ public:
 IMPLEMENT_GLOBAL_SHADER(FBloomSetupCS, "/Engine/Private/PostProcessBloom.usf", "BloomSetupCS", SF_Compute);
 } //!namespace
 
-bool FBloomOutputs::SupportsApplyParametersBuffer(EShaderPlatform Platform)
-{
-	return FDataDrivenShaderPlatformInfo::GetSupportsFFTBloom(Platform);
-}
-
 FScreenPassTexture AddBloomSetupPass(FRDGBuilder& GraphBuilder, const FViewInfo& View, const FBloomSetupInputs& Inputs)
 {
 	check(Inputs.SceneColor.IsValid());
@@ -185,7 +180,7 @@ static_assert(
 	static_cast<uint32>(EBloomQuality::MAX) == FSceneDownsampleChain::StageCount,
 	"The total number of stages in the scene downsample chain and the number of bloom quality levels must match.");
 
-FBloomOutputs AddGaussianBloomPasses(FRDGBuilder& GraphBuilder, const FViewInfo& View, const FSceneDownsampleChain* SceneDownsampleChain)
+FScreenPassTexture AddGaussianBloomPasses(FRDGBuilder& GraphBuilder, const FViewInfo& View, const FSceneDownsampleChain* SceneDownsampleChain)
 {
 	check(SceneDownsampleChain);
 	check(!IsFFTBloomEnabled(View));
@@ -194,7 +189,7 @@ FBloomOutputs AddGaussianBloomPasses(FRDGBuilder& GraphBuilder, const FViewInfo&
 
 	const EBloomQuality BloomQuality = GetBloomQuality();
 
-	FBloomOutputs PassOutputs;
+	FScreenPassTexture PassOutputs;
 	if (BloomQuality != EBloomQuality::Disabled)
 	{
 		RDG_EVENT_SCOPE(GraphBuilder, "Bloom");
@@ -239,7 +234,7 @@ FBloomOutputs AddGaussianBloomPasses(FRDGBuilder& GraphBuilder, const FViewInfo&
 		// Use bloom quality to select the number of downsample stages to use for bloom.
 		const uint32 BloomStageCount = BloomQualityToSceneDownsampleStage[BloomQualityIndex];
 
-		const float TintScale = 1.0f / BloomQualityCountMax;
+		const float TintScale = (1.0f / BloomQualityCountMax) * Settings.BloomIntensity;
 
 		for (uint32 StageIndex = 0, SourceIndex = BloomQualityCountMax - 1; StageIndex < BloomStageCount; ++StageIndex, --SourceIndex)
 		{
@@ -251,30 +246,14 @@ FBloomOutputs AddGaussianBloomPasses(FRDGBuilder& GraphBuilder, const FViewInfo&
 				PassInputs.NameX = TEXT("BloomX");
 				PassInputs.NameY = TEXT("BloomY");
 				PassInputs.Filter = SceneDownsampleChain->GetTexture(SourceIndex);
-				PassInputs.Additive = PassOutputs.Bloom;
+				PassInputs.Additive = PassOutputs;
 				PassInputs.CrossCenterWeight = CrossCenterWeight;
 				PassInputs.KernelSizePercent = BloomStage.Size * Settings.BloomSizeScale;
 				PassInputs.TintColor = BloomStage.Tint * TintScale;
 
-				PassOutputs.Bloom = AddGaussianBlurPass(GraphBuilder, View, PassInputs);
+				PassOutputs = AddGaussianBlurPass(GraphBuilder, View, PassInputs);
 			}
 		}
-
-		if (FBloomOutputs::SupportsApplyParametersBuffer(View.GetShaderPlatform()))
-		{
-			FBloomOutputs::FApplyInfo ApplyParameters;
-			ApplyParameters.SceneColorMultiply = FLinearColor::White;
-			ApplyParameters.BloomMultiply = FLinearColor::White * View.FinalPostProcessSettings.BloomIntensity;
-			ApplyParameters.BloomMultiply.A = 0.0f;
-
-			TArray<FBloomOutputs::FApplyInfo, TInlineAllocator<1> > ApplyParametersArray;
-			ApplyParametersArray.Add(ApplyParameters);
-
-			PassOutputs.ApplyParameters = CreateStructuredBuffer(
-				GraphBuilder,
-				TEXT("Bloom.ApplyParameters"),
-				ApplyParametersArray);
-		}
 	}
 
 	return PassOutputs;
diff --git a/Engine/Source/Runtime/Renderer/Private/PostProcess/PostProcessFFTBloom.h b/Engine/Source/Runtime/Renderer/Private/PostProcess/PostProcessFFTBloom.h
--- a/Engine/Source/Runtime/Renderer/Private/PostProcess/PostProcessFFTBloom.h
+++ b/Engine/Source/Runtime/Renderer/Private/PostProcess/PostProcessFFTBloom.h
@@ -4,6 +4,7 @@
 
 #include "RenderGraph.h"
 #include "ScenePrivate.h"
+#include "ScreenPass.h"
 
 struct FBloomOutputs;
 
@@ -19,7 +20,13 @@ struct FFFTBloomInputs
 	FIntRect HalfResolutionViewRect;
 };
 
+struct FFFTBloomOutput
+{
+	FScreenPassTexture BloomTexture;
+	FRDGBufferRef SceneColorApplyParameters = nullptr;
+};
+
 bool IsFFTBloomFullResolutionEnabled();
 bool IsFFTBloomQuarterResolutionEnabled();
 
-FBloomOutputs AddFFTBloomPass(FRDGBuilder& GraphBuilder, const FViewInfo& View, const FFFTBloomInputs& Inputs);
+FFFTBloomOutput AddFFTBloomPass(FRDGBuilder& GraphBuilder, const FViewInfo& View, const FFFTBloomInputs& Inputs);
diff --git a/Engine/Source/Runtime/Renderer/Private/PostProcess/PostProcessLensFlares.cpp b/Engine/Source/Runtime/Renderer/Private/PostProcess/PostProcessLensFlares.cpp
--- a/Engine/Source/Runtime/Renderer/Private/PostProcess/PostProcessLensFlares.cpp
+++ b/Engine/Source/Runtime/Renderer/Private/PostProcess/PostProcessLensFlares.cpp
@@ -398,7 +398,7 @@ FScreenPassTexture AddLensFlaresPass(
 	LensFlareInputs.TintColorsPerFlare = Settings.LensFlareTints;
 	LensFlareInputs.TintColor = Settings.LensFlareTint;
 	LensFlareInputs.BokehSizePercent = Settings.LensFlareBokehSize;
-	LensFlareInputs.Intensity = Settings.LensFlareIntensity;
+	LensFlareInputs.Intensity = Settings.LensFlareIntensity * Settings.BloomIntensity;
 	LensFlareInputs.Threshold = Settings.LensFlareThreshold;
 
 	// If a bloom output texture isn't available, substitute the half resolution scene color instead, but disable bloom
diff --git a/Engine/Source/Runtime/Renderer/Private/PostProcess/PostProcessMobile.cpp b/Engine/Source/Runtime/Renderer/Private/PostProcess/PostProcessMobile.cpp
--- a/Engine/Source/Runtime/Renderer/Private/PostProcess/PostProcessMobile.cpp
+++ b/Engine/Source/Runtime/Renderer/Private/PostProcess/PostProcessMobile.cpp
@@ -1027,11 +1027,12 @@ public:
 	using FPermutationDomain = TShaderPermutationDomain<
 		FUseBloomDim,
 		FUseSunDim>;
-
+	
 	BEGIN_SHADER_PARAMETER_STRUCT(FParameters, )
 		SHADER_PARAMETER(FVector4f, BloomDirtMaskTint)
 		SHADER_PARAMETER(FVector4f, SunColorVignetteIntensity)
 		SHADER_PARAMETER(FVector3f, BloomColor)
+		SHADER_PARAMETER(float, BloomIntensity)
 		SHADER_PARAMETER_STRUCT_REF(FViewUniformShaderParameters, View)
 		SHADER_PARAMETER_RDG_TEXTURE(Texture2D, SunBlurTexture)
 		SHADER_PARAMETER_SAMPLER(SamplerState, SunBlurSampler)
@@ -1139,6 +1140,7 @@ FScreenPassTexture AddMobileSunMergePass(FRDGBuilder& GraphBuilder, const FViewI
 	PSShaderParameters->BloomDirtMaskTint = Settings.BloomDirtMaskTint * Settings.BloomDirtMaskIntensity;
 	PSShaderParameters->SunColorVignetteIntensity = SunColorVignetteIntensityParam;
 	PSShaderParameters->BloomColor = FVector(BloomColor.R, BloomColor.G, BloomColor.B);
+	PSShaderParameters->BloomIntensity = Settings.BloomIntensity;
 	PSShaderParameters->SunBlurTexture = Inputs.SunBlur.Texture;
 	PSShaderParameters->SunBlurSampler = TStaticSamplerState<SF_Bilinear, AM_Clamp, AM_Clamp, AM_Clamp>::GetRHI();
 	PSShaderParameters->BloomSetup_BloomTexture = Inputs.BloomSetup_Bloom.Texture;
diff --git a/Engine/Source/Runtime/Renderer/Private/PostProcess/PostProcessTonemap.cpp b/Engine/Source/Runtime/Renderer/Private/PostProcess/PostProcessTonemap.cpp
--- a/Engine/Source/Runtime/Renderer/Private/PostProcess/PostProcessTonemap.cpp
+++ b/Engine/Source/Runtime/Renderer/Private/PostProcess/PostProcessTonemap.cpp
@@ -17,6 +17,7 @@
 #include "PipelineStateCache.h"
 #include "Rendering/Texture2DResource.h"
 #include "Math/Halton.h"
+#include "SystemTextures.h"
 
 bool SupportsFilmGrain(EShaderPlatform Platform)
 {
@@ -341,14 +342,15 @@ BEGIN_SHADER_PARAMETER_STRUCT(FTonemapParameters, )
 	SHADER_PARAMETER_STRUCT(FEyeAdaptationParameters, EyeAdaptation)
 	SHADER_PARAMETER_RDG_TEXTURE(Texture2D, ColorTexture)
 
+	// Parameters to apply to the scene color.
+	SHADER_PARAMETER_RDG_BUFFER_SRV(StructuredBuffer<float4>, SceneColorApplyParamaters)
+
 	// Bloom texture
 	SHADER_PARAMETER(FScreenTransform, ColorToBloom)
 	SHADER_PARAMETER(FVector2f, BloomUVViewportBilinearMin)
 	SHADER_PARAMETER(FVector2f, BloomUVViewportBilinearMax)
-	SHADER_PARAMETER(FVector4f, BloomMultiply)
 	SHADER_PARAMETER_RDG_TEXTURE(Texture2D, BloomTexture)
 	SHADER_PARAMETER_SAMPLER(SamplerState, BloomSampler)
-	SHADER_PARAMETER_RDG_BUFFER_SRV(StructuredBuffer<float4>, BloomApplyParameters)
 
 	SHADER_PARAMETER_RDG_TEXTURE(Texture3D, LumBilateralGrid)
 	SHADER_PARAMETER_RDG_TEXTURE(Texture2D, BlurredLogLum)
@@ -462,7 +464,7 @@ public:
 		const int UseVolumeLut = PipelineVolumeTextureLUTSupportGuaranteedAtRuntime(Parameters.Platform) ? 1 : 0;
 		OutEnvironment.SetDefine(TEXT("USE_VOLUME_LUT"), UseVolumeLut);
 
-		OutEnvironment.SetDefine(TEXT("SUPPORTS_BLOOM_APPLY_PARAMETERS"), FBloomOutputs::SupportsApplyParametersBuffer(Parameters.Platform) ? 1 : 0);
+		OutEnvironment.SetDefine(TEXT("SUPPORTS_SCENE_COLOR_APPLY_PARAMETERS"), FTonemapInputs::SupportsSceneColorApplyParametersBuffer(Parameters.Platform) ? 1 : 0);
 		OutEnvironment.SetDefine(TEXT("SUPPORTS_FILM_GRAIN"), SupportsFilmGrain(Parameters.Platform) ? 1 : 0);
 	}
 };
@@ -502,7 +504,7 @@ public:
 		const int UseVolumeLut = PipelineVolumeTextureLUTSupportGuaranteedAtRuntime(Parameters.Platform) ? 1 : 0;
 		OutEnvironment.SetDefine(TEXT("USE_VOLUME_LUT"), UseVolumeLut);
 
-		OutEnvironment.SetDefine(TEXT("SUPPORTS_BLOOM_APPLY_PARAMETERS"), FBloomOutputs::SupportsApplyParametersBuffer(Parameters.Platform) ? 1 : 0);
+		OutEnvironment.SetDefine(TEXT("SUPPORTS_SCENE_COLOR_APPLY_PARAMETERS"), FTonemapInputs::SupportsSceneColorApplyParametersBuffer(Parameters.Platform) ? 1 : 0);
 		OutEnvironment.SetDefine(TEXT("SUPPORTS_FILM_GRAIN"), SupportsFilmGrain(Parameters.Platform) ? 1 : 0);
 	}
 };
@@ -513,6 +515,11 @@ IMPLEMENT_GLOBAL_SHADER(FTonemapVS, "/Engine/Private/PostProcessTonemap.usf", "M
 IMPLEMENT_GLOBAL_SHADER(FTonemapPS, "/Engine/Private/PostProcessTonemap.usf", "MainPS", SF_Pixel);
 IMPLEMENT_GLOBAL_SHADER(FTonemapCS, "/Engine/Private/PostProcessTonemap.usf", "MainCS", SF_Compute);
 
+bool FTonemapInputs::SupportsSceneColorApplyParametersBuffer(EShaderPlatform Platform)
+{
+	return FDataDrivenShaderPlatformInfo::GetSupportsFFTBloom(Platform);
+}
+
 static
 FRDGBufferRef BuildFilmGrainConstants(FRDGBuilder& GraphBuilder, const FViewInfo& View, FRDGTextureRef FilmGrainTexture)
 {
@@ -622,13 +629,6 @@ FScreenPassTexture AddTonemapPass(FRDGBuilder& GraphBuilder, const FViewInfo& Vi
 
 	const FScreenPassTextureViewport OutputViewport(Output);
 
-	FRHITexture* BloomDirtMaskTexture = GBlackTexture->TextureRHI;
-
-	if (PostProcessSettings.BloomDirtMask && PostProcessSettings.BloomDirtMask->GetResource())
-	{
-		BloomDirtMaskTexture = PostProcessSettings.BloomDirtMask->GetResource()->TextureRHI;
-	}
-
 	FRHISamplerState* BilinearClampSampler = TStaticSamplerState<SF_Bilinear, AM_Clamp, AM_Clamp, AM_Clamp>::GetRHI();
 	FRHISamplerState* PointClampSampler = TStaticSamplerState<SF_Point, AM_Clamp, AM_Clamp, AM_Clamp>::GetRHI();
 
@@ -780,12 +780,9 @@ FScreenPassTexture AddTonemapPass(FRDGBuilder& GraphBuilder, const FViewInfo& Vi
 	CommonParameters.EyeAdaptationTexture = Inputs.EyeAdaptationTexture;
 	CommonParameters.EyeAdaptation = *Inputs.EyeAdaptationParameters;
 	CommonParameters.ColorGradingLUT = Inputs.ColorGradingTexture;
-	CommonParameters.BloomDirtMaskTexture = BloomDirtMaskTexture;
 	CommonParameters.ColorSampler = BilinearClampSampler;
 	CommonParameters.ColorGradingLUTSampler = BilinearClampSampler;
-	CommonParameters.BloomDirtMaskSampler = BilinearClampSampler;
 	CommonParameters.ColorScale0 = PostProcessSettings.SceneColorTint;
-	CommonParameters.BloomDirtMaskTint = PostProcessSettings.BloomDirtMaskTint * PostProcessSettings.BloomDirtMaskIntensity;
 	CommonParameters.ChromaticAberrationParams = ChromaticAberrationParams;
 	CommonParameters.TonemapperParams = FVector4f(PostProcessSettings.VignetteIntensity, SharpenDiv6, 0.0f, 0.0f);
 	CommonParameters.SwitchVerticalAxis = Inputs.bFlipYAxis;
@@ -794,16 +791,43 @@ FScreenPassTexture AddTonemapPass(FRDGBuilder& GraphBuilder, const FViewInfo& Vi
 	CommonParameters.bOutputInHDR = ViewFamily.bIsHDR;
 	CommonParameters.LensPrincipalPointOffsetScale = View.LensPrincipalPointOffsetScale;
 
+	// TODO: PostProcessSettings.BloomDirtMask->GetResource() is not thread safe
+	if (PostProcessSettings.BloomDirtMask && PostProcessSettings.BloomDirtMask->GetResource() && Inputs.bAllowDirtMask && PostProcessSettings.BloomDirtMaskIntensity > 0 && PostProcessSettings.BloomIntensity > 0.0)
+	{
+		CommonParameters.BloomDirtMaskTint = PostProcessSettings.BloomDirtMaskTint * PostProcessSettings.BloomDirtMaskIntensity / PostProcessSettings.BloomIntensity;
+		CommonParameters.BloomDirtMaskTexture = PostProcessSettings.BloomDirtMask->GetResource()->TextureRHI;
+		CommonParameters.BloomDirtMaskSampler = BilinearClampSampler;
+	}
+	else
+	{
+		CommonParameters.BloomDirtMaskTint = FLinearColor::Black;
+		CommonParameters.BloomDirtMaskTexture = GBlackTexture->TextureRHI;
+		CommonParameters.BloomDirtMaskSampler = BilinearClampSampler;
+	}
+
+	if (!FTonemapInputs::SupportsSceneColorApplyParametersBuffer(View.GetShaderPlatform()))
+	{
+		check(Inputs.SceneColorApplyParamaters == nullptr);
+	}
+	else if (Inputs.SceneColorApplyParamaters)
+	{
+		CommonParameters.SceneColorApplyParamaters = GraphBuilder.CreateSRV(Inputs.SceneColorApplyParamaters);
+	}
+	else
+	{
+		FRDGBufferRef ApplyParametersBuffer = GSystemTextures.GetDefaultStructuredBuffer(GraphBuilder, sizeof(FVector4f), FVector4f(1.0f, 1.0f, 1.0f, 1.0f));
+		CommonParameters.SceneColorApplyParamaters = GraphBuilder.CreateSRV(ApplyParametersBuffer);
+	}
+
 	// Bloom parameters
 	{
-		if (Inputs.Bloom.Bloom.Texture)
+		if (Inputs.Bloom.Texture)
 		{
-			const FScreenPassTextureViewport BloomViewport(Inputs.Bloom.Bloom);
+			const FScreenPassTextureViewport BloomViewport(Inputs.Bloom);
 			CommonParameters.ColorToBloom = FScreenTransform::ChangeTextureUVCoordinateFromTo(SceneColorViewport, BloomViewport);
 			CommonParameters.BloomUVViewportBilinearMin = GetScreenPassTextureViewportParameters(BloomViewport).UVViewportBilinearMin;
 			CommonParameters.BloomUVViewportBilinearMax = GetScreenPassTextureViewportParameters(BloomViewport).UVViewportBilinearMax;
-			CommonParameters.BloomMultiply = FLinearColor::White * PostProcessSettings.BloomIntensity;
-			CommonParameters.BloomTexture = Inputs.Bloom.Bloom.Texture;
+			CommonParameters.BloomTexture = Inputs.Bloom.Texture;
 			CommonParameters.BloomSampler = BilinearClampSampler;
 		}
 		else
@@ -811,35 +835,9 @@ FScreenPassTexture AddTonemapPass(FRDGBuilder& GraphBuilder, const FViewInfo& Vi
 			CommonParameters.ColorToBloom = FScreenTransform::Identity;
 			CommonParameters.BloomUVViewportBilinearMin = FVector2D::ZeroVector;
 			CommonParameters.BloomUVViewportBilinearMax = FVector2D::UnitVector;
-			CommonParameters.BloomMultiply = FLinearColor::Transparent;
 			CommonParameters.BloomTexture = GSystemTextures.GetBlackDummy(GraphBuilder);
 			CommonParameters.BloomSampler = BilinearClampSampler;
 		}
-
-		if (!FBloomOutputs::SupportsApplyParametersBuffer(View.GetShaderPlatform()))
-		{
-			check(Inputs.Bloom.ApplyParameters == nullptr);
-		}
-		else if (Inputs.Bloom.ApplyParameters)
-		{
-			CommonParameters.BloomApplyParameters = GraphBuilder.CreateSRV(Inputs.Bloom.ApplyParameters);
-		}
-		else
-		{
-			FBloomOutputs::FApplyInfo ApplyParameters;
-			ApplyParameters.SceneColorMultiply = FLinearColor::White;
-			ApplyParameters.BloomMultiply = FLinearColor(CommonParameters.BloomMultiply);
-
-			TArray<FBloomOutputs::FApplyInfo, TInlineAllocator<1> > ApplyParametersArray;
-			ApplyParametersArray.Add(ApplyParameters);
-
-			FRDGBufferRef ApplyParametersBuffer = CreateStructuredBuffer(
-				GraphBuilder,
-				TEXT("Bloom.ApplyParameters"),
-				ApplyParametersArray);
-
-			CommonParameters.BloomApplyParameters = GraphBuilder.CreateSRV(ApplyParametersBuffer);
-		}
 	}
 
 	// forward transformation from shader:
diff --git a/Engine/Source/Runtime/Renderer/Private/PostProcess/PostProcessTonemap.h b/Engine/Source/Runtime/Renderer/Private/PostProcess/PostProcessTonemap.h
--- a/Engine/Source/Runtime/Renderer/Private/PostProcess/PostProcessTonemap.h
+++ b/Engine/Source/Runtime/Renderer/Private/PostProcess/PostProcessTonemap.h
@@ -51,7 +51,10 @@ struct FTonemapInputs
 	FScreenPassTexture SceneColor;
 
 	// [Required] Filtered bloom texture to composite with tonemapped scene color. This should be transparent black for no bloom.
-	FBloomOutputs Bloom;
+	FScreenPassTexture Bloom;
+
+	// [Optional] structured buffer of multiply parameters to apply to the scene color.
+	FRDGBufferRef SceneColorApplyParamaters;
 
 	// [Optional] Luminance bilateral grid. If this is null, local exposure is disabled.
 	FRDGTextureRef LocalExposureTexture = nullptr;
@@ -71,6 +74,9 @@ struct FTonemapInputs
 	// [Optional, ES31] Eye adaptation buffer used to compute exposure. 
 	FRDGBufferRef EyeAdaptationBuffer = nullptr;
 
+	// Weather should allow dirk mask.
+	bool bAllowDirtMask = true;
+
 	// [Raster Only, Mobile] Flips the image vertically on output.
 	bool bFlipYAxis = false;
 
@@ -84,6 +90,9 @@ struct FTonemapInputs
 	bool bOutputInHDR = false;
 
 	bool bMetalMSAAHDRDecode = false;
+
+	// Returns whether ApplyParameters is supported by the tonemapper.
+	static bool SupportsSceneColorApplyParametersBuffer(EShaderPlatform Platform);
 };
 
 FScreenPassTexture AddTonemapPass(FRDGBuilder& GraphBuilder, const FViewInfo& View, const FTonemapInputs& Inputs);
diff --git a/Engine/Source/Runtime/Renderer/Private/PostProcess/PostProcessing.cpp b/Engine/Source/Runtime/Renderer/Private/PostProcess/PostProcessing.cpp
--- a/Engine/Source/Runtime/Renderer/Private/PostProcess/PostProcessing.cpp
+++ b/Engine/Source/Runtime/Renderer/Private/PostProcess/PostProcessing.cpp
@@ -908,7 +908,8 @@ void AddPostProcessingPasses(FRDGBuilder& GraphBuilder, const FViewInfo& View, c
 				HistogramTexture);
 		}
 
-		FBloomOutputs Bloom;
+		FScreenPassTexture Bloom;
+		FRDGBufferRef SceneColorApplyParameters = nullptr;
 		if (bBloomEnabled)
 		{
 			const FSceneDownsampleChain* LensFlareSceneDownsampleChain;
@@ -933,7 +934,9 @@ void AddPostProcessingPasses(FRDGBuilder& GraphBuilder, const FViewInfo& View, c
 				PassInputs.HalfResolutionTexture = HalfResolution.Texture;
 				PassInputs.HalfResolutionViewRect = HalfResolution.ViewRect;
 
-				Bloom = AddFFTBloomPass(GraphBuilder, View, PassInputs);
+				FFFTBloomOutput Outputs = AddFFTBloomPass(GraphBuilder, View, PassInputs);
+				Bloom = Outputs.BloomTexture;
+				SceneColorApplyParameters = Outputs.SceneColorApplyParameters;
 			}
 			else
 			{
@@ -973,7 +976,7 @@ void AddPostProcessingPasses(FRDGBuilder& GraphBuilder, const FViewInfo& View, c
 
 			if (bLensFlareEnabled)
 			{
-				Bloom.Bloom = AddLensFlaresPass(GraphBuilder, View, Bloom.Bloom, *LensFlareSceneDownsampleChain);
+				Bloom = AddLensFlaresPass(GraphBuilder, View, Bloom, *LensFlareSceneDownsampleChain);
 			}
 		}
 
@@ -991,7 +994,7 @@ void AddPostProcessingPasses(FRDGBuilder& GraphBuilder, const FViewInfo& View, c
 				PassSequence.AcceptOverrideIfLastPass(EPass::Tonemap, PassInputs.OverrideOutput);
 				PassInputs.SetInput(EPostProcessMaterialInput::SceneColor, SceneColor);
 				PassInputs.SetInput(EPostProcessMaterialInput::SeparateTranslucency, SeparateTranslucency);
-				PassInputs.SetInput(EPostProcessMaterialInput::CombinedBloom, Bloom.Bloom);
+				PassInputs.SetInput(EPostProcessMaterialInput::CombinedBloom, Bloom);
 				PassInputs.SceneTextures = GetSceneTextureShaderParameters(Inputs.SceneTextures);
 				PassInputs.CustomDepthTexture = CustomDepth.Texture;
 
@@ -1020,11 +1023,13 @@ void AddPostProcessingPasses(FRDGBuilder& GraphBuilder, const FViewInfo& View, c
 				PassSequence.AcceptOverrideIfLastPass(EPass::Tonemap, PassInputs.OverrideOutput);
 				PassInputs.SceneColor = SceneColor;
 				PassInputs.Bloom = Bloom;
+				PassInputs.SceneColorApplyParamaters = SceneColorApplyParameters;
 				PassInputs.LocalExposureTexture = LocalExposureTexture;
 				PassInputs.BlurredLogLuminanceTexture = LocalExposureBlurredLogLumTexture;
 				PassInputs.EyeAdaptationParameters = &EyeAdaptationParameters;
 				PassInputs.EyeAdaptationTexture = EyeAdaptationTexture;
 				PassInputs.ColorGradingTexture = ColorGradingTexture;
+				PassInputs.bAllowDirtMask = !bFFTBloomEnabled;
 				PassInputs.bWriteAlphaChannel = AntiAliasingMethod == AAM_FXAA || IsPostProcessingWithAlphaChannelSupported();
 				PassInputs.bOutputInHDR = bTonemapOutputInHDR;
 
@@ -2343,7 +2348,7 @@ void AddMobilePostProcessingPasses(FRDGBuilder& GraphBuilder, FScene* Scene, con
 		}
 			
 		TonemapperInputs.SceneColor = SceneColor;
-		TonemapperInputs.Bloom.Bloom = BloomOutput;
+		TonemapperInputs.Bloom = BloomOutput;
 		TonemapperInputs.EyeAdaptationParameters = &EyeAdaptationParameters;
 		TonemapperInputs.EyeAdaptationTexture = nullptr;
 		TonemapperInputs.ColorGradingTexture = ColorGradingTexture;
