Another macro fix for RenderDebugViewMode

diff --git a/Engine/Source/Runtime/Renderer/Private/BasePassRendering.cpp b/Engine/Source/Runtime/Renderer/Private/BasePassRendering.cpp
--- a/Engine/Source/Runtime/Renderer/Private/BasePassRendering.cpp
+++ b/Engine/Source/Runtime/Renderer/Private/BasePassRendering.cpp
@@ -1110,6 +1110,7 @@ void FDeferredShadingSceneRenderer::RenderBasePassInternal(
 				RDG_GPU_MASK_SCOPE(GraphBuilder, View.GPUMask);
 				RDG_EVENT_SCOPE_CONDITIONAL(GraphBuilder, Views.Num() > 1, "View%d", ViewIndex);
 
+				#if WITH_DEBUG_VIEW_MODES
 				Nanite::FRasterResults& RasterResults = NaniteRasterResults[ViewIndex];
 				Nanite::RenderDebugViewMode(
 					GraphBuilder,
@@ -1124,6 +1125,7 @@ void FDeferredShadingSceneRenderer::RenderBasePassInternal(
 					SceneTextures.QuadOverdraw,
 					BasePassRenderTargets
 				);
+				#endif
 			}
 		}
 	}
