Fixes crash when using a view mode

diff --git a/Engine/Source/Runtime/Renderer/Private/PostProcess/PostProcessTonemap.h b/Engine/Source/Runtime/Renderer/Private/PostProcess/PostProcessTonemap.h
--- a/Engine/Source/Runtime/Renderer/Private/PostProcess/PostProcessTonemap.h
+++ b/Engine/Source/Runtime/Renderer/Private/PostProcess/PostProcessTonemap.h
@@ -54,7 +54,7 @@ struct FTonemapInputs
 	FScreenPassTexture Bloom;
 
 	// [Optional] structured buffer of multiply parameters to apply to the scene color.
-	FRDGBufferRef SceneColorApplyParamaters;
+	FRDGBufferRef SceneColorApplyParamaters = nullptr;
 
 	// [Optional] Luminance bilateral grid. If this is null, local exposure is disabled.
 	FRDGTextureRef LocalExposureTexture = nullptr;
