Fix crash due to delegate being de-registered on the render thread.

diff --git a/Engine/Source/Runtime/Renderer/Private/VirtualShadowMaps/VirtualShadowMapCacheManager.cpp b/Engine/Source/Runtime/Renderer/Private/VirtualShadowMaps/VirtualShadowMapCacheManager.cpp
--- a/Engine/Source/Runtime/Renderer/Private/VirtualShadowMaps/VirtualShadowMapCacheManager.cpp
+++ b/Engine/Source/Runtime/Renderer/Private/VirtualShadowMaps/VirtualShadowMapCacheManager.cpp
@@ -12,6 +12,7 @@
 
 #include "PrimitiveSceneInfo.h"
 #include "ShaderDebug.h"
+#include "RendererOnScreenNotification.h"
 
 static TAutoConsoleVariable<int32> CVarAccumulateStats(
 	TEXT("r.Shadow.Virtual.AccumulateStats"),
@@ -170,7 +171,7 @@ FVirtualShadowMapArrayCacheManager::FVirtualShadowMapArrayCacheManager(FScene* I
 	});
 
 #if !UE_BUILD_SHIPPING
-	ScreenMessageDelegate = FCoreDelegates::OnGetOnScreenMessages.AddLambda([this](TMultiMap<FCoreDelegates::EOnScreenMessageSeverity, FText >& OutMessages)
+	ScreenMessageDelegate = FRendererOnScreenNotification::Get().AddLambda([this](TMultiMap<FCoreDelegates::EOnScreenMessageSeverity, FText >& OutMessages)
 	{
 		// Show for ~5s after last overflow
 		int32 CurrentFrameNumber = Scene->GetFrameNumber();
@@ -185,7 +186,7 @@ FVirtualShadowMapArrayCacheManager::FVirtualShadowMapArrayCacheManager(FScene* I
 FVirtualShadowMapArrayCacheManager::~FVirtualShadowMapArrayCacheManager()
 {
 #if !UE_BUILD_SHIPPING
-	FCoreDelegates::OnGetOnScreenMessages.Remove(ScreenMessageDelegate);
+	FRendererOnScreenNotification::Get().Remove(ScreenMessageDelegate);
 #endif
 }
 
