* Updates to GPU Scene and the GPUSceneCollector to support handling the upload of multiple instances and instance payload data for dynamic primitives
* Updates to the VirtualShadowMapArrayCacheManager to support the invalidation of cache pages upon removal of dynamic primitives
* Various updates and fixes to support writing dynamic primitives' instance scene data and instance payload data from the GPU
* All these changes are needed by a few existing systems that will be using them soon to render via GPU Scene

diff --git a/Engine/Shaders/Private/GPUScene/GPUSceneWriter.ush b/Engine/Shaders/Private/GPUScene/GPUSceneWriter.ush
--- a/Engine/Shaders/Private/GPUScene/GPUSceneWriter.ush
+++ b/Engine/Shaders/Private/GPUScene/GPUSceneWriter.ush
@@ -21,19 +21,164 @@ uint GPUSceneNumAllocatedPrimitives;
 
 void StoreInstanceSceneDataElement(uint InstanceId, uint ElementIndex, float4 DataElement)
 {
+	checkSlow(InstanceId < GPUSceneNumAllocatedInstances);
 	GPUSceneInstanceSceneDataRW[ElementIndex * GPUSceneInstanceSceneDataSOAStride + InstanceId] = DataElement;
 }
 
+void StoreInstanceSceneDataElementX(uint InstanceId, uint ElementIndex, float Data)
+{
+	// NOTE: Structured buffers allow for DWORD granularity storage, so you don't have to write the whole float4
+	checkSlow(InstanceId < GPUSceneNumAllocatedInstances);
+	GPUSceneInstanceSceneDataRW[ElementIndex * GPUSceneInstanceSceneDataSOAStride + InstanceId].x = Data;
+}
+
+void StoreInstancePayloadDataElement(uint PayloadDataOffset, uint ElementIndex, float4 DataElement)
+{
+	GPUSceneInstancePayloadDataRW[PayloadDataOffset + ElementIndex] = DataElement;
+}
+
+/**
+ * Helpers to load and store the primitive ID and flags for the specified instance, which are packed together
+ */
+void LoadInstancePrimitiveIdAndFlags(uint InstanceId, inout uint OutPrimitiveId, inout uint OutInstanceFlags)
+{
+	LoadInstancePrimitiveIdAndFlags(InstanceId, GPUSceneInstanceSceneDataSOAStride, OutPrimitiveId, OutInstanceFlags);
+}
+
+void WriteInstancePrimitiveIdAndFlags(uint InstanceId, uint PrimitiveId, uint InstanceFlags)
+{
+	const uint Packed0 = PackPrimitiveIdAndInstanceFlags(PrimitiveId, InstanceFlags);
+	StoreInstanceSceneDataElementX(InstanceId, 0, asfloat(Packed0));
+}
+
 /**
  * Store the contents of InstanceData.PrimitiveId to the GPUScene buffers. 
- * NOTE: As the granularity is float4 this also reads/stores other members that are co-located.
+ * NOTE: As the primitive ID is packed with the instance flags, the flags are loaded and re-stored
  */
 void WriteInstancePrimitiveId(uint InstanceId, uint NewPrimitiveId)
 {
-	float4 DataElement = LoadInstanceSceneDataElement(0 * GPUSceneInstanceSceneDataSOAStride + InstanceId);
-	const uint InstanceFlags = BitFieldExtractU32(asuint(DataElement.x),  12, 20); // Max of 12 instance flags
-	const uint Packed0 = (InstanceFlags << 20u) | NewPrimitiveId;
-	DataElement.x = asfloat(Packed0);
-	checkSlow(InstanceId < GPUSceneNumAllocatedInstances);
-	StoreInstanceSceneDataElement(InstanceId, 0, DataElement);
+	uint OldPrimitiveId, InstanceFlags;
+	LoadInstancePrimitiveIdAndFlags(InstanceId, OldPrimitiveId, InstanceFlags);
+	WriteInstancePrimitiveIdAndFlags(InstanceId, NewPrimitiveId, InstanceFlags);
+}
+
+/**
+ * Store the contents of InstanceData.LocalToWorld to the GPUScene buffers. Optionally, also calculates the determinant
+ * of the transform's rotation to set or unset the determinant sign flag on the instance flags
+ **/
+void WriteInstanceLocalToWorld(uint InstanceId, float4x4 LocalToWorld, bool StoreDeterminantSign)
+{
+#if INSTANCE_SCENE_DATA_COMPRESSED_TRANSFORMS
+	uint4 RotationScale;
+	float3 Translation;
+	EncodeTransform(LocalToWorld, RotationScale, Translation);
+	StoreInstanceSceneDataElement(InstanceId, 1, asfloat(RotationScale));
+	StoreInstanceSceneDataElement(InstanceId, 2, float4(Translation, 0.0f));
+#else
+	LocalToWorld = transpose(LocalToWorld);
+	StoreInstanceSceneDataElement(InstanceId, 1, LocalToWorld[0]);
+	StoreInstanceSceneDataElement(InstanceId, 2, LocalToWorld[1]);
+	StoreInstanceSceneDataElement(InstanceId, 3, LocalToWorld[2]);
+#endif
+
+	// check to update the determinant sign flag on the instance flags as well
+	if (StoreDeterminantSign)
+	{
+		// Check to update the determinant sign flag on the instance flags
+		uint PrimitiveId, InstanceFlags;
+		LoadInstancePrimitiveIdAndFlags(InstanceId, PrimitiveId, InstanceFlags);
+		SetInstanceDeterminantSignFlag(determinant((float3x3)LocalToWorld), InstanceFlags);
+		WriteInstancePrimitiveIdAndFlags(InstanceId, PrimitiveId, InstanceFlags);
+	}
+}
+
+/**
+ * Initializes the instance scene data of the specified instance data slot.
+ * NOTE: This function must match FInstanceSceneShaderData::BuildInternal
+ **/
+ void InitializeInstanceSceneData(
+ 	uint InstanceId,
+ 	uint PrimitiveId,
+ 	uint InstanceRelativeId,
+ 	uint PayloadDataFlags,
+ 	uint CustomDataCount,
+ 	float RandomId,
+ 	float4x4 LocalToWorld)
+ {
+ 	uint InstanceFlags = PayloadDataFlags;
+ 	SetInstanceDeterminantSignFlag(determinant((float3x3)LocalToWorld), InstanceFlags);
+ 	
+ 	float4 Element0 =
+ 	{
+ 		asfloat(PackPrimitiveIdAndInstanceFlags(PrimitiveId, InstanceFlags)),
+ 		asfloat(PackInstanceRelativeIdAndCustomDataCount(InstanceRelativeId, CustomDataCount)),
+ 		asfloat(GetGPUSceneFrameNumber()),
+ 		RandomId
+	};
+ 	StoreInstanceSceneDataElement(InstanceId, 0, Element0);
+ 	WriteInstanceLocalToWorld(InstanceId, LocalToWorld, false);
+ }
+
+/**
+ * Helper for loading instance payload offsets for GPU writers (assumes global parameter for SOA stride)
+ **/
+ FInstancePayloadDataOffsets GetInstancePayloadDataOffsets(uint InstanceId)
+{
+	uint PrimitiveId, InstanceFlags, InstanceRelativeId, CustomDataCount;
+	LoadInstancePrimitiveIdAndFlags(InstanceId, GPUSceneInstanceSceneDataSOAStride, PrimitiveId, InstanceFlags);
+	LoadInstanceRelativeIdAndCustomDataCount(InstanceId, GPUSceneInstanceSceneDataSOAStride, InstanceRelativeId, CustomDataCount);
+	return GetInstancePayloadDataOffsets(PrimitiveId, InstanceFlags, InstanceRelativeId);
+}
+
+/**
+ * Store the data for dynamic instances to the GPUScene sidecar payload buffer.
+ **/
+void WriteInstanceDynamicData(FInstancePayloadDataOffsets PayloadOffsets, float4x4 PrevLocalToWorld)
+{	
+	BRANCH
+	if (PayloadOffsets.DynamicData == INVALID_INSTANCE_PAYLOAD_OFFSET)
+	{
+		// Instance doesn't have dynamic data
+		return;
+	}
+
+#if INSTANCE_SCENE_DATA_COMPRESSED_TRANSFORMS
+	uint4 RotationScale;
+	float3 Translation;
+	EncodeTransform(PrevLocalToWorld, RotationScale, Translation);
+	StoreInstancePayloadDataElement(PayloadOffsets.DynamicData, 0, asfloat(RotationScale));
+	StoreInstancePayloadDataElement(PayloadOffsets.DynamicData, 1, float4(Translation, 0.0f));
+#else
+	PrevLocalToWorld = transpose(PrevLocalToWorld);
+	StoreInstancePayloadDataElement(PayloadOffsets.DynamicData, 0, PrevLocalToWorld[0]);
+	StoreInstancePayloadDataElement(PayloadOffsets.DynamicData, 1, PrevLocalToWorld[1]);
+	StoreInstancePayloadDataElement(PayloadOffsets.DynamicData, 2, PrevLocalToWorld[2]);
+#endif
+}
+
+void WriteInstanceDynamicData(uint InstanceId, float4x4 PrevLocalToWorld)
+{
+	FInstancePayloadDataOffsets PayloadOffsets = GetInstancePayloadDataOffsets(InstanceId);
+	WriteInstanceDynamicData(PayloadOffsets, PrevLocalToWorld);
 }
+
+/**
+ * Stores a float4 element of custom data for the given instance
+ */
+ void WriteInstanceCustomData(FInstancePayloadDataOffsets PayloadOffsets, uint ElementIndex, float4 ElementData)
+ {
+ 	BRANCH
+ 	if (PayloadOffsets.CustomData == INVALID_INSTANCE_PAYLOAD_OFFSET)
+ 	{
+		// Instance has no custom data
+ 		return;
+ 	}
+
+ 	StoreInstancePayloadDataElement(PayloadOffsets.CustomData, ElementIndex, ElementData);
+ }
+
+ void WriteInstanceCustomData(uint InstanceId, uint ElementIndex, float4 ElementData)
+ {
+	FInstancePayloadDataOffsets PayloadOffsets = GetInstancePayloadDataOffsets(InstanceId);
+ 	WriteInstanceCustomData(PayloadOffsets, ElementIndex, ElementData);
+ }
\ No newline at end of file
diff --git a/Engine/Shaders/Private/SceneData.ush b/Engine/Shaders/Private/SceneData.ush
--- a/Engine/Shaders/Private/SceneData.ush
+++ b/Engine/Shaders/Private/SceneData.ush
@@ -25,8 +25,15 @@
 	#define NEEDS_LIGHTMAP_COORDINATE 0
 #endif
 
+#ifndef VF_REQUIRES_PER_INSTANCE_CUSTOM_DATA
+	#define VF_REQUIRES_PER_INSTANCE_CUSTOM_DATA 0
+#endif
+
+#define ENABLE_PER_INSTANCE_CUSTOM_DATA (USES_PER_INSTANCE_CUSTOM_DATA || VF_REQUIRES_PER_INSTANCE_CUSTOM_DATA)
+
+// When transforms come from the InstanceSceneData buffer, indicates whether or not the transforms are compressed (ignored on mobile)
 // TODO: Global setting/define
-#define INSTANCE_COMPRESSED_TRANSFORMS 1
+#define INSTANCE_SCENE_DATA_COMPRESSED_TRANSFORMS 1
 
 // Whether to fetch primitive values (eg LocalToWorld) by dynamically indexing a scene-wide buffer, or to reference a single Primitive uniform buffer
 #if VF_SUPPORTS_PRIMITIVE_SCENE_DATA
@@ -60,6 +67,17 @@
 // GPUCULL_TODO: Eventually we need to remove this workaround
 #define VF_TREAT_INSTANCE_ID_OFFSET_AS_PRIMITIVE_ID_FLAG (1U << 31U)
 
+#define PRIMITIVE_ID_NUM_BITS				(20u) // Max of 1,048,576 primitives
+#define INSTANCE_SCENE_DATA_FLAGS_NUM_BITS	(12u) // Max of 12 flags
+#define INSTANCE_RELATIVE_ID_NUM_BITS		(24u) // Max of 16,777,216 instances per primitive
+#define INSTANCE_CUSTOM_DATA_COUNT_NUM_BITS	 (8u) // Max of 255 custom data floats per instance
+
+#define PRIMITIVE_ID_MASK			((1u << PRIMITIVE_ID_NUM_BITS) - 1u)
+#define INSTANCE_RELATIVE_ID_MASK	((1u << INSTANCE_RELATIVE_ID_NUM_BITS) - 1u)
+
+#define INVALID_PRIMITIVE_ID PRIMITIVE_ID_MASK
+#define INVALID_INSTANCE_PAYLOAD_OFFSET (0xFFFFFFFFu)
+
 // Must match InstanceUniformShaderParameters.h
 #define INSTANCE_SCENE_DATA_FLAG_DETERMINANT_SIGN			0x1
 #define INSTANCE_SCENE_DATA_FLAG_HAS_RANDOM					0x2
@@ -70,9 +88,11 @@
 #define INSTANCE_SCENE_DATA_FLAG_HAS_LOCAL_BOUNDS			0x40
 #define INSTANCE_SCENE_DATA_FLAG_HAS_EDITOR_DATA			0x80
 
-#define INVALID_PRIMITIVE_ID 0x000FFFFFu
-
-#define INSTANCE_COMPRESSED_TRANSFORM	0
+#if INSTANCE_SCENE_DATA_COMPRESSED_TRANSFORMS
+	static const uint InstanceTransformSizeFloat4Count = 2u; // compressed transform
+#else
+	static const uint InstanceTransformSizeFloat4Count = 3u; // encoded scale/rotation (uint4) and translation (float3)
+#endif
 
 #define NUM_CUSTOM_PRIMITIVE_DATA 9 // Num float4s used for custom data. Must match FCustomPrimitiveData::NumCustomPrimitiveDataFloat4s in SceneTypes.h
 
@@ -351,6 +371,16 @@ float3 UnpackHitProxyId(uint HitProxyPacked)
 	) * (1.0f / 255.0f);
 }
 
+struct FInstancePayloadDataOffsets
+{
+	uint HierarchyOffset;
+	uint EditorData;
+	uint LocalBounds;
+	uint DynamicData;
+	uint LightShadowUVBias;
+	uint CustomData;
+};
+
 struct FInstanceSceneData
 {
 	FLWCMatrix LocalToWorld;
@@ -370,7 +400,7 @@ struct FInstanceSceneData
 #if 1//USES_PER_INSTANCE_RANDOM
 	float    RandomID;
 #endif
-#if USES_PER_INSTANCE_CUSTOM_DATA
+#if ENABLE_PER_INSTANCE_CUSTOM_DATA
 	uint     CustomDataOffset;
 	uint     CustomDataCount;
 #endif
@@ -426,6 +456,31 @@ float4 LoadInstancePayloadDataElement(uint Index)
 #endif
 }
 
+float4 LoadInstanceCustomDataElement(FInstanceSceneData SceneData, uint Float4Index)
+{
+#if ENABLE_PER_INSTANCE_CUSTOM_DATA
+	const uint NumCustomFloat4s = (SceneData.CustomDataCount + 3u) >> 2u;
+	if (SceneData.CustomDataOffset != 0xFFFFFFFFu && Float4Index < NumCustomFloat4s)
+	{
+		return LoadInstancePayloadDataElement(SceneData.CustomDataOffset + Float4Index);
+	}
+#endif
+
+	return (float4)0.0f;
+}
+
+float LoadInstanceCustomDataFloat(FInstanceSceneData SceneData, uint FloatIndex)
+{
+#if ENABLE_PER_INSTANCE_CUSTOM_DATA
+	const uint Float4Index = FloatIndex >> 2u;
+	const uint ComponentIndex = FloatIndex % 4u;
+	const float4 Element = LoadInstanceCustomDataElement(SceneData, Float4Index);
+	return Element[ComponentIndex];
+#else
+	return 0.0f;
+#endif
+}
+
 // [Frisvad 2012, "Building an Orthonormal Basis from a 3D Unit Vector Without Normalization"]
 void GetHemiOrthoBasis( inout float3 BasisX, inout float3 BasisY, float3 BasisZ )
 {
@@ -435,6 +490,107 @@ void GetHemiOrthoBasis( inout float3 BasisX, inout float3 BasisY, float3 BasisZ
 	BasisY = float3( B, 1.0f - BasisZ.y * BasisZ.y * A, -BasisZ.y );
 }
 
+uint4 EncodeScaleAndRotation(float3 Scale, float3x3 Axis)
+{
+	const uint ExpBits			= 8;
+	const uint ExpBias			= ( 1u << (ExpBits - 1) ) - 1;
+	const uint SignMantissaBits	= 16;
+	const uint SignMantissaMask	= (1u << SignMantissaBits) - 1;
+	const uint MantissaBits		= SignMantissaBits - 1;
+	const float Sqrt2 			= 1.41421356f;
+
+	uint4 Output;
+
+	// Rotation
+	{
+		if( Axis[2].z < 0.0f )
+		{
+			Axis[2] *= -1.0f;
+			Scale.z *= -1.0f;
+		}
+
+		float2 OctZ = UnitVectorToHemiOctahedron( Axis[2] );
+
+		float3 BasisX, BasisY;
+		GetHemiOrthoBasis( BasisX, BasisY, Axis[2] );
+
+		float X = dot(Axis[0], BasisX);
+		float Y = dot(Axis[0], BasisY);
+
+		float aX = abs( X );
+		float aY = abs( Y );
+
+		bool bSpinIsX = aX < aY;
+		float Spin0 = bSpinIsX ? X : Y;
+		float Spin1 = bSpinIsX ? Y : X;
+		float Sign1 = Spin1 < 0.0f ? -1.0f : 1.0f;
+	
+		//Axis[0]	*= Sign1;
+		Scale.x *= Sign1;
+		Spin0	*= Sign1;
+
+		float3 GeneratedY = cross(Axis[2], Axis[0]);
+		Scale.y *= dot( Axis[1], GeneratedY ) < 0.0f ? -Sign1 : Sign1;
+
+		// Avoid sign extension in shader by biasing
+		Output.x  = (((int)round( OctZ.x * 32767.0f ) + 32768) & 0xFFFF) <<  0;
+		Output.x |= (((int)round( OctZ.y * 32767.0f ) + 32768) & 0xFFFF) << 16;
+
+		// NOTE: Masking the bits with `& 0x7FFF` below causes the whole int to be optimized to 0 on some shader platforms.
+		// This is okay, as long as Spin0 is in [0, 1], which it should be.
+		Output.y  = ((int)round( Spin0 * 16383.0f * Sqrt2 ) + 16384); // & 0x7FFF;
+		Output.y |=	bSpinIsX ? (1u << 15) : 0;
+	}
+
+	// Scale
+	{
+		float MaxComponent = max3(abs(Scale.x), abs(Scale.y), abs(Scale.z));
+		uint MaxComponentExponent = (asuint(MaxComponent) & 0x7f800000u) >> 23;
+
+		// Need +1 because of losing the implicit leading bit of mantissa
+		// TODO assumes ExpBits == 8
+		// TODO clamp to expressable range
+		uint SharedExp = MaxComponentExponent + 1;
+
+		float ExpScale = asfloat(((127 + ExpBias + MantissaBits - SharedExp) & 0xFFu) << 23);
+
+		if( (uint)round( MaxComponent * ExpScale ) == (1u << MantissaBits) )
+		{
+			// Mantissa rounded up
+			SharedExp++;
+			ExpScale *= 0.5f;
+		}
+
+		Output.z  = (((int)round( Scale.x * ExpScale ) + (1u << MantissaBits)) & 0xFFFFu) <<  0;
+		Output.z |= (((int)round( Scale.y * ExpScale ) + (1u << MantissaBits)) & 0xFFFFu) << 16;
+		Output.w  = (((int)round( Scale.z * ExpScale ) + (1u << MantissaBits)) & 0xFFFFu) <<  0;
+		Output.w |= SharedExp << 16;		
+	}
+
+	return Output;
+}
+
+uint4 EncodeScaleAndRotation( float3x3 InTransform )
+{
+	float3 Scale = {
+		length(InTransform[0]),
+		length(InTransform[1]),
+		length(InTransform[2])
+	};
+	float3x3 Axis = {
+		InTransform[0] / Scale.x,
+		InTransform[1] / Scale.y,
+		InTransform[2] / Scale.z
+	};
+	return EncodeScaleAndRotation(Scale, Axis);
+}
+
+void EncodeTransform( float4x4 InTransform, inout uint4 OutRotationScale, inout float3 OutTranslation )
+{
+	OutRotationScale = EncodeScaleAndRotation((float3x3)InTransform);
+	OutTranslation = InTransform[3].xyz;
+}
+
 float4x4 DecodeTransform( uint4 RotationScale, float3 Translation, inout float3 Scale )
 {
 	float4x4 M = 0.0;
@@ -510,24 +666,145 @@ float4x4 DecodeTransform( uint4 RotationScale, float3 Translation, inout float3
 	return M;
 }
 
+// Helpers to pack/unpack the primitive ID and flags for the specified instance, which are packed together in a uint
+void UnpackPrimitiveIdAndInstanceFlags(uint PackedPrimitiveIdAndFlags, inout uint OutPrimitiveId, inout uint OutInstanceFlags)
+{
+	OutPrimitiveId		= BitFieldExtractU32(PackedPrimitiveIdAndFlags, PRIMITIVE_ID_NUM_BITS,  0);
+	OutInstanceFlags	= BitFieldExtractU32(PackedPrimitiveIdAndFlags, INSTANCE_SCENE_DATA_FLAGS_NUM_BITS, PRIMITIVE_ID_NUM_BITS);	
+}
+
+uint PackPrimitiveIdAndInstanceFlags(uint PrimitiveId, uint InstanceFlags)
+{
+	return (PrimitiveId & PRIMITIVE_ID_MASK) | (InstanceFlags << PRIMITIVE_ID_NUM_BITS);
+}
+
+void LoadInstancePrimitiveIdAndFlags(uint InstanceId, uint SOAStride, inout uint OutPrimitiveId, inout uint OutInstanceFlags)
+{
+	const uint PackedPrimitiveIdAndFlags = asuint(LoadInstanceSceneDataElement(0 * SOAStride + InstanceId).x);
+	UnpackPrimitiveIdAndInstanceFlags(PackedPrimitiveIdAndFlags, OutPrimitiveId, OutInstanceFlags);
+}
+
+// Helpers to pack/unpack the instance relative ID and custom data count for the specified instance, which are packed together in a uint
+void UnpackInstanceRelativeIdAndCustomDataCount(uint PackedRelativeIdAndCustomDataCount, inout uint OutRelativeId, inout uint OutCustomDataCount)
+{
+	OutRelativeId		= BitFieldExtractU32(PackedRelativeIdAndCustomDataCount, INSTANCE_RELATIVE_ID_NUM_BITS,  0);
+	OutCustomDataCount	= BitFieldExtractU32(PackedRelativeIdAndCustomDataCount, INSTANCE_CUSTOM_DATA_COUNT_NUM_BITS, INSTANCE_RELATIVE_ID_NUM_BITS);
+}
+
+uint PackInstanceRelativeIdAndCustomDataCount(uint RelativeId, uint CustomDataCount)
+{
+	return (RelativeId & INSTANCE_RELATIVE_ID_MASK) | (CustomDataCount << INSTANCE_RELATIVE_ID_NUM_BITS);
+}
+
+void LoadInstanceRelativeIdAndCustomDataCount(uint InstanceId, uint SOAStride, inout uint OutPrimitiveId, inout uint OutInstanceFlags)
+{
+	const uint PackedRelativeIdAndCustomDataCount = asuint(LoadInstanceSceneDataElement(0 * SOAStride + InstanceId).y);
+	UnpackInstanceRelativeIdAndCustomDataCount(PackedRelativeIdAndCustomDataCount, OutPrimitiveId, OutInstanceFlags);
+}
+
+// Helpers for getting/setting the instance determinant sign from instance data flags
+float GetInstanceDeterminantSignFromFlags(uint Flags)
+{
+	// Scale.x * Scale.y * Scale.z < 0.0 ? -1.0 : 1.0;
+	return CondMask(Flags & INSTANCE_SCENE_DATA_FLAG_DETERMINANT_SIGN, -1.0f, 1.0f);
+}
+
+void SetInstanceDeterminantSignFlag(float Determinant, inout uint Flags)
+{
+	if (Determinant < 0.0f)
+	{
+		Flags |= INSTANCE_SCENE_DATA_FLAG_DETERMINANT_SIGN;
+	}
+	else
+	{
+		Flags &= ~INSTANCE_SCENE_DATA_FLAG_DETERMINANT_SIGN;
+	}
+}
+
+// Determine the offsets into the payload data buffer for the given instance
+FInstancePayloadDataOffsets GetInstancePayloadDataOffsets(uint PrimitiveId, uint Flags, uint InstanceRelativeId)
+{
+	const uint PayloadDataRelativeOffset	= InstanceRelativeId * GetPrimitiveData(PrimitiveId).InstancePayloadDataStride;
+	const uint PayloadDataGlobalOffset		= PayloadDataRelativeOffset + GetPrimitiveData(PrimitiveId).InstancePayloadDataOffset;
+
+	const bool bHasHierarchyOffset			= (Flags & INSTANCE_SCENE_DATA_FLAG_HAS_HIERARCHY_OFFSET) != 0u;
+	const bool bHasLocalBounds				= (Flags & INSTANCE_SCENE_DATA_FLAG_HAS_LOCAL_BOUNDS) != 0u;
+	const bool bHasDynamicData				= (Flags & INSTANCE_SCENE_DATA_FLAG_HAS_DYNAMIC_DATA) != 0u;
+	const bool bHasLightShadowUVBias		= (Flags & INSTANCE_SCENE_DATA_FLAG_HAS_LIGHTSHADOW_UV_BIAS) != 0u;
+	const bool bHasCustomData				= (Flags & INSTANCE_SCENE_DATA_FLAG_HAS_CUSTOM_DATA) != 0u;
+#if USE_EDITOR_SHADERS
+	const bool bHasEditorData				= (Flags & INSTANCE_SCENE_DATA_FLAG_HAS_EDITOR_DATA) != 0u;
+#else
+	const bool bHasEditorData				= false;
+#endif
+
+	uint CurOffset = PayloadDataGlobalOffset;
+
+	// Offsets are in float4s
+	FInstancePayloadDataOffsets Offsets;
+	Offsets.HierarchyOffset 	= INVALID_INSTANCE_PAYLOAD_OFFSET;
+	Offsets.EditorData 			= INVALID_INSTANCE_PAYLOAD_OFFSET;
+	Offsets.LocalBounds 		= INVALID_INSTANCE_PAYLOAD_OFFSET;
+	Offsets.DynamicData 		= INVALID_INSTANCE_PAYLOAD_OFFSET;
+	Offsets.LightShadowUVBias 	= INVALID_INSTANCE_PAYLOAD_OFFSET;
+	Offsets.CustomData 			= INVALID_INSTANCE_PAYLOAD_OFFSET;
+	
+	// Hierarchy Offset -> float0.x
+	if (bHasHierarchyOffset)
+	{
+		Offsets.HierarchyOffset = CurOffset;
+	}
+	// EditorData -> float0.y
+	if (bHasEditorData)
+	{
+		Offsets.EditorData = CurOffset;
+	}
+	// LocalBounds -> float0.zw & float1.xyzw
+	if (bHasLocalBounds)
+	{
+		Offsets.LocalBounds = CurOffset;
+	}
+	CurOffset += CondMask(bHasLocalBounds, 2u, CondMask(bHasHierarchyOffset || bHasEditorData, 1u, 0u));
+
+	if (bHasDynamicData)
+	{
+		Offsets.DynamicData = CurOffset;
+		CurOffset += InstanceTransformSizeFloat4Count;
+	}
+
+	if (bHasLightShadowUVBias)
+	{
+		Offsets.LightShadowUVBias = CurOffset;
+		++CurOffset;
+	}
+
+	if (bHasCustomData)
+	{
+		Offsets.CustomData = CurOffset;
+	}
+
+	return Offsets;
+}
+
 void ComputeInstanceDerivedData(inout FInstanceSceneData InstanceData, float3 TilePosition, float4x4 LocalToRelativeWorld)
 {
 	//
 	// Do not put any load operations here!
 	//
 
-#if !INSTANCE_COMPRESSED_TRANSFORM
+#if (VF_USE_PRIMITIVE_SCENE_DATA == 2) || !INSTANCE_SCENE_DATA_COMPRESSED_TRANSFORMS
+	// Non-uniform scale must be computed from the transform because it was not already computed when decoding it (see below in GetInstanceSceneData)
 	float3 Scale2;
 	Scale2.x = length2(LocalToRelativeWorld[0].xyz);
 	Scale2.y = length2(LocalToRelativeWorld[1].xyz);
 	Scale2.z = length2(LocalToRelativeWorld[2].xyz);
 
 	InstanceData.InvNonUniformScale							= rsqrt(Scale2);
-	InstanceData.DeterminantSign							= CondMask(InstanceData.Flags & INSTANCE_SCENE_DATA_FLAG_DETERMINANT_SIGN, -1.0f, 1.0f);
 	InstanceData.NonUniformScale.xyz						= Scale2 * InstanceData.InvNonUniformScale;
 #endif
 
 	InstanceData.NonUniformScale.w							= max3( InstanceData.NonUniformScale.x, InstanceData.NonUniformScale.y, InstanceData.NonUniformScale.z );
+	InstanceData.DeterminantSign							= GetInstanceDeterminantSignFromFlags(InstanceData.Flags);
 
 	float4x4 RelativeWorldToLocal							= LocalToRelativeWorld;
 	RelativeWorldToLocal[0].xyz							   *= Pow2(InstanceData.InvNonUniformScale.x);
@@ -537,14 +814,6 @@ void ComputeInstanceDerivedData(inout FInstanceSceneData InstanceData, float3 Ti
 	RelativeWorldToLocal								    = transpose(RelativeWorldToLocal);
 	RelativeWorldToLocal[3].xyz						        = mul(float4(-LocalToRelativeWorld[3].xyz, 0.0f), RelativeWorldToLocal).xyz;
 	InstanceData.WorldToLocal								= MakeLWCInverseMatrix(TilePosition, RelativeWorldToLocal);
-
-// See defines in FInstanceSceneData
-#if 0
-	InstanceData.HasRandomID								= (InstanceData.Flags & INSTANCE_SCENE_DATA_FLAG_HAS_RANDOM);
-	InstanceData.HasCustomData								= (InstanceData.Flags & INSTANCE_SCENE_DATA_FLAG_HAS_CUSTOM_DATA);
-	InstanceData.HasDynamicData								= (InstanceData.Flags & INSTANCE_SCENE_DATA_FLAG_HAS_DYNAMIC_DATA);
-	InstanceData.HasLightShadowUVBias						= (InstanceData.Flags & INSTANCE_SCENE_DATA_FLAG_HAS_LIGHTSHADOW_UV_BIAS);
-#endif
 }
 
 // Fetch from scene primitive buffer
@@ -558,11 +827,8 @@ FInstanceSceneData GetInstanceSceneData(uint InstanceId, uint SOAStride, bool bC
 	//
 
 	// Only process valid instances
-	// Note: keep primitive ID packing in sync with GPUSceneWriter.ush
-	const uint FlagsAndPrimitiveId	= asuint(LoadInstanceSceneDataElement(0 * SOAStride + InstanceId).x);
-	InstanceData.PrimitiveId		= BitFieldExtractU32(FlagsAndPrimitiveId,  20, 0);  // Max of 1,048,576 primitives
-	InstanceData.Flags				= BitFieldExtractU32(FlagsAndPrimitiveId,  12, 20); // Max of 12 instance flags
-	InstanceData.ValidInstance		= InstanceData.PrimitiveId != INVALID_PRIMITIVE_ID;
+	LoadInstancePrimitiveIdAndFlags(InstanceId, SOAStride, InstanceData.PrimitiveId, InstanceData.Flags);
+	InstanceData.ValidInstance = InstanceData.PrimitiveId != INVALID_PRIMITIVE_ID;
 
 	// Payload Data Layout
 	// NOTE: Per-instance local bounds and hierarchy offset are always mutually inclusive, so pack together.
@@ -572,7 +838,7 @@ FInstanceSceneData GetInstanceSceneData(uint InstanceId, uint SOAStride, bool bC
 	// LocalBounds Center		float0.yzw
 	// LocalBounds Extent		float1.xyz
 	// __UNUSED					float1.w
-#if INSTANCE_COMPRESSED_TRANSFORMS
+#if INSTANCE_SCENE_DATA_COMPRESSED_TRANSFORMS
 	// Previous Transform[0]	float2.xyzw
 	// Previous Transform[1]	float3.xyzw
 	// LM/SM Scale Bias			float4.xyzw
@@ -585,28 +851,26 @@ FInstanceSceneData GetInstanceSceneData(uint InstanceId, uint SOAStride, bool bC
 	// Custom Data Float4s		float6.xyzw ... floatN.xyzw
 #endif
 
-#if INSTANCE_COMPRESSED_TRANSFORMS
-	const uint TransformSizeFloat4Count = 2u;
-#else
-	const uint TransformSizeFloat4Count = 3u;
-#endif
-
 	BRANCH
 	if (!bCheckValid || InstanceData.ValidInstance)
 	{
-		const uint RelativeIdAndCustomDataCount					= asuint(LoadInstanceSceneDataElement(0 * SOAStride + InstanceId).y);
-		InstanceData.RelativeId									= BitFieldExtractU32(RelativeIdAndCustomDataCount, 24, 0); // Max of 16,777,216 instances per primitive
-	#if USES_PER_INSTANCE_CUSTOM_DATA
-		InstanceData.CustomDataCount							= BitFieldExtractU32(RelativeIdAndCustomDataCount, 8, 24); // Max of 255 custom data floats per instance
+		uint CustomDataCount;
+		LoadInstanceRelativeIdAndCustomDataCount(InstanceId, SOAStride, InstanceData.RelativeId, CustomDataCount);
+
+		FInstancePayloadDataOffsets Offsets 	= GetInstancePayloadDataOffsets(InstanceData.PrimitiveId, InstanceData.Flags, InstanceData.RelativeId);
+	
+	#if ENABLE_PER_INSTANCE_CUSTOM_DATA
+		InstanceData.CustomDataCount			= CustomDataCount;
+		InstanceData.CustomDataOffset			= Offsets.CustomData;
 	#endif
-		InstanceData.LastUpdateSceneFrameNumber					= asuint(LoadInstanceSceneDataElement(0 * SOAStride + InstanceId).z);
+		InstanceData.LastUpdateSceneFrameNumber	= asuint(LoadInstanceSceneDataElement(0 * SOAStride + InstanceId).z);
 	#if 1//USES_PER_INSTANCE_RANDOM
-		InstanceData.RandomID									= LoadInstanceSceneDataElement(0 * SOAStride + InstanceId).w;
+		InstanceData.RandomID					= LoadInstanceSceneDataElement(0 * SOAStride + InstanceId).w;
 	#endif
 
 		float3 TilePosition 									= GetPrimitiveData(InstanceData.PrimitiveId).TilePosition;
 
-#if INSTANCE_COMPRESSED_TRANSFORMS
+#if INSTANCE_SCENE_DATA_COMPRESSED_TRANSFORMS
 		uint4  RotationScale									= asuint(LoadInstanceSceneDataElement(1 * SOAStride + InstanceId));
 		float3 Translation										= LoadInstanceSceneDataElement(2 * SOAStride + InstanceId).xyz;
 		float3 Scale = 0;
@@ -637,60 +901,35 @@ FInstanceSceneData GetInstanceSceneData(uint InstanceId, uint SOAStride, bool bC
 		InstanceData.PrevLocalToWorld	= MakeLWCMatrix(TilePosition, PrevLocalToRelativeWorld);
 		ComputeInstanceDerivedData(InstanceData, TilePosition, LocalToRelativeWorld);
 
-		const uint PayloadDataRelativeOffset	= InstanceData.RelativeId   * GetPrimitiveData(InstanceData.PrimitiveId).InstancePayloadDataStride;
-		const uint PayloadDataGlobalOffset		= PayloadDataRelativeOffset + GetPrimitiveData(InstanceData.PrimitiveId).InstancePayloadDataOffset;
-
-		const bool bHasHierarchyOffset			= (InstanceData.Flags & INSTANCE_SCENE_DATA_FLAG_HAS_HIERARCHY_OFFSET) != 0u;
-		const bool bHasLocalBounds				= (InstanceData.Flags & INSTANCE_SCENE_DATA_FLAG_HAS_LOCAL_BOUNDS) != 0u;
-		const bool bHasDynamicData				= (InstanceData.Flags & INSTANCE_SCENE_DATA_FLAG_HAS_DYNAMIC_DATA) != 0u;
-		const bool bHasLightShadowUVBias		= (InstanceData.Flags & INSTANCE_SCENE_DATA_FLAG_HAS_LIGHTSHADOW_UV_BIAS) != 0u;
-		const bool bHasCustomData				= (InstanceData.Flags & INSTANCE_SCENE_DATA_FLAG_HAS_CUSTOM_DATA) != 0u;
-		const bool bHasRandomID					= (InstanceData.Flags & INSTANCE_SCENE_DATA_FLAG_HAS_RANDOM) != 0u;
-#if USE_EDITOR_SHADERS
-		const bool bHasEditorData				= (InstanceData.Flags & INSTANCE_SCENE_DATA_FLAG_HAS_EDITOR_DATA) != 0u;
-#else //
-		const bool bHasEditorData				= false;
-#endif
-
-		// Offsets are in float4s
-		const uint Offset_HierarchyOffset		= PayloadDataGlobalOffset + 0u; // float0.x
-		const uint Offset_EditorData			= PayloadDataGlobalOffset + 0u; // float0.y
-		const uint Offset_LocalBounds			= Offset_HierarchyOffset; // float0.zw + float1.xyzw
-		const uint Offset_DynamicData			= CondMask(bHasLocalBounds, 2u, CondMask(bHasHierarchyOffset || bHasEditorData, 1u, 0u)) + PayloadDataGlobalOffset;
-		const uint Offset_LightShadowUVBias		= CondMask(bHasDynamicData, TransformSizeFloat4Count, 0u) + Offset_DynamicData;
-	#if USES_PER_INSTANCE_CUSTOM_DATA
-		InstanceData.CustomDataOffset			= CondMask(bHasCustomData, CondMask(bHasLightShadowUVBias, 1u, 0u) + Offset_LightShadowUVBias, 0xFFFFFFFFu);
-	#endif
-
 		InstanceData.NaniteRuntimeResourceID	= GetPrimitiveData(InstanceData.PrimitiveId).NaniteResourceID;
 		InstanceData.NaniteHierarchyOffset		= GetPrimitiveData(InstanceData.PrimitiveId).NaniteHierarchyOffset;
 
 		BRANCH
-		if (bHasHierarchyOffset)
+		if (Offsets.HierarchyOffset != INVALID_INSTANCE_PAYLOAD_OFFSET)
 		{
-			const uint HierarchyRootOffset = asuint(LoadInstancePayloadDataElement(Offset_HierarchyOffset)).x;
+			const uint HierarchyRootOffset = asuint(LoadInstancePayloadDataElement(Offsets.HierarchyOffset)).x;
 
 			// Combine this instance's hierarchy offset with the primitive's root hierarchy offset
 			InstanceData.NaniteHierarchyOffset += HierarchyRootOffset;
 		}
 
-#if USE_EDITOR_SHADERS
+	#if USE_EDITOR_SHADERS
 		BRANCH
-		if (bHasEditorData)
+		if (Offsets.EditorData != INVALID_INSTANCE_PAYLOAD_OFFSET)
 		{
-			const uint PackedEditorData = asuint(LoadInstancePayloadDataElement(Offset_EditorData)).y;
+			const uint PackedEditorData = asuint(LoadInstancePayloadDataElement(Offsets.EditorData)).y;
 
 			InstanceData.EditorData.bIsSelected    = (PackedEditorData >> 24u) != 0;
 			InstanceData.EditorData.HitProxyPacked =  PackedEditorData & 0x00FFFFFFu;
 			InstanceData.EditorData.HitProxyId = UnpackHitProxyId(InstanceData.EditorData.HitProxyPacked);
 		}
-#endif
+	#endif
 
 		BRANCH
-		if (bHasLocalBounds)
+		if (Offsets.LocalBounds != INVALID_INSTANCE_PAYLOAD_OFFSET)
 		{
-			InstanceData.LocalBoundsCenter = float3(LoadInstancePayloadDataElement(Offset_LocalBounds + 0).zw, LoadInstancePayloadDataElement(Offset_LocalBounds + 1).x);
-			InstanceData.LocalBoundsExtent = LoadInstancePayloadDataElement(Offset_LocalBounds + 1).yzw;
+			InstanceData.LocalBoundsCenter = float3(LoadInstancePayloadDataElement(Offsets.LocalBounds + 0).zw, LoadInstancePayloadDataElement(Offsets.LocalBounds + 1).x);
+			InstanceData.LocalBoundsExtent = LoadInstancePayloadDataElement(Offsets.LocalBounds + 1).yzw;
 		}
 		else
 		{
@@ -699,17 +938,17 @@ FInstanceSceneData GetInstanceSceneData(uint InstanceId, uint SOAStride, bool bC
 		}
 
 		BRANCH
-		if (bHasDynamicData)
+		if (Offsets.DynamicData != INVALID_INSTANCE_PAYLOAD_OFFSET)
 		{
-		#if INSTANCE_COMPRESSED_TRANSFORMS
-			uint4  PrevRotationScale			= asuint(LoadInstancePayloadDataElement(Offset_DynamicData + 0));
-			float3 PrevTranslation				= LoadInstancePayloadDataElement(Offset_DynamicData + 1).xyz;
+		#if INSTANCE_SCENE_DATA_COMPRESSED_TRANSFORMS
+			uint4  PrevRotationScale			= asuint(LoadInstancePayloadDataElement(Offsets.DynamicData + 0));
+			float3 PrevTranslation				= LoadInstancePayloadDataElement(Offsets.DynamicData + 1).xyz;
 			float3 PrevScale					= 0;
 			float4x4 PrevLocalToRelativeWorld	= DecodeTransform(PrevRotationScale, PrevTranslation, PrevScale);
 		#else
-			float4x4 PrevLocalToRelativeWorld	= transpose(float4x4(LoadInstancePayloadDataElement(Offset_DynamicData + 0),
-																 LoadInstancePayloadDataElement(Offset_DynamicData + 1),
-																 LoadInstancePayloadDataElement(Offset_DynamicData + 2),
+			float4x4 PrevLocalToRelativeWorld	= transpose(float4x4(LoadInstancePayloadDataElement(Offsets.DynamicData + 0),
+																 LoadInstancePayloadDataElement(Offsets.DynamicData + 1),
+																 LoadInstancePayloadDataElement(Offsets.DynamicData + 2),
 																 float4(0.0f, 0.0f, 0.0f, 1.0f)));
 		#endif
 
@@ -718,7 +957,7 @@ FInstanceSceneData GetInstanceSceneData(uint InstanceId, uint SOAStride, bool bC
 		}
 		else
 		{
-		#if INSTANCE_COMPRESSED_TRANSFORMS
+		#if INSTANCE_SCENE_DATA_COMPRESSED_TRANSFORMS
 			// TODO: Temporary PrevVelocityHack
 			uint4  PrevRotationScale			= asuint(LoadInstanceSceneDataElement(3 * SOAStride + InstanceId));
 			float3 PrevTranslation				= LoadInstanceSceneDataElement(4 * SOAStride + InstanceId).xyz;
@@ -734,9 +973,9 @@ FInstanceSceneData GetInstanceSceneData(uint InstanceId, uint SOAStride, bool bC
 
 	#if 1 //NEEDS_LIGHTMAP_COORDINATE
 		BRANCH
-		if (bHasLightShadowUVBias)
+		if (Offsets.LightShadowUVBias != INVALID_INSTANCE_PAYLOAD_OFFSET)
 		{
-			InstanceData.LightMapAndShadowMapUVBias = LoadInstancePayloadDataElement(Offset_LightShadowUVBias);
+			InstanceData.LightMapAndShadowMapUVBias = LoadInstancePayloadDataElement(Offsets.LightShadowUVBias);
 		}
 	#endif
 	}
@@ -852,10 +1091,6 @@ FSceneDataIntermediates GetSceneDataIntermediates(uint DrawInstanceId, float4 In
 	InstanceData.RandomID = InstanceTransform3.w;
 #endif
 
-#if INSTANCE_COMPRESSED_TRANSFORM
-	#error Compressed instance transform is not supported on mobile path
-#endif
-
 	ComputeInstanceDerivedData(InstanceData, TilePosition, LocalToRelativeWorld);
 
 	//
diff --git a/Engine/Source/Runtime/Engine/Private/InstanceUniformShaderParameters.cpp b/Engine/Source/Runtime/Engine/Private/InstanceUniformShaderParameters.cpp
--- a/Engine/Source/Runtime/Engine/Private/InstanceUniformShaderParameters.cpp
+++ b/Engine/Source/Runtime/Engine/Private/InstanceUniformShaderParameters.cpp
@@ -4,13 +4,6 @@
 #include "PrimitiveSceneProxy.h"
 #include "PrimitiveSceneInfo.h"
 
-// TODO: Global setting/define
-#ifndef INSTANCE_COMPRESSED_TRANSFORMS
-#define INSTANCE_COMPRESSED_TRANSFORMS 1
-#else
-#error "Should not be defined yet"
-#endif
-
 void FInstanceSceneShaderData::Build
 (
 	uint32 PrimitiveId,
@@ -59,7 +52,7 @@ void FInstanceSceneShaderData::BuildInternal
 	const FRenderTransform& PrevLocalToWorld // Assumes shear has been removed already // TODO: Temporary PrevVelocityHack
 )
 {
-	// Note: layout must match GetInstanceData in SceneData.ush
+	// Note: layout must match GetInstanceData in SceneData.ush and InitializeInstanceSceneData in GPUSceneWriter.ush
 
 	uint32 InstanceFlags = PayloadDataFlags;
 	if (LocalToWorld.RotDeterminant() < 0.0f)
@@ -85,7 +78,7 @@ void FInstanceSceneShaderData::BuildInternal
 	Data[0].W  = *(const float*)&RandomID;
 
 	// TODO: Temporary PrevVelocityHack
-#if INSTANCE_COMPRESSED_TRANSFORMS
+#if INSTANCE_SCENE_DATA_COMPRESSED_TRANSFORMS
 	FCompressedTransform CompressedLocalToWorld(LocalToWorld);
 	Data[1] = *(const FVector4f*)&CompressedLocalToWorld.Rotation[0];
 	Data[2] = *(const FVector3f*)&CompressedLocalToWorld.Translation;
diff --git a/Engine/Source/Runtime/Engine/Private/PrimitiveSceneProxy.cpp b/Engine/Source/Runtime/Engine/Private/PrimitiveSceneProxy.cpp
--- a/Engine/Source/Runtime/Engine/Private/PrimitiveSceneProxy.cpp
+++ b/Engine/Source/Runtime/Engine/Private/PrimitiveSceneProxy.cpp
@@ -15,11 +15,13 @@
 #include "Materials/Material.h"
 #include "SceneManagement.h"
 #include "VT/RuntimeVirtualTexture.h"
+#include "PrimitiveInstanceUpdateCommand.h"
+#include "InstanceUniformShaderParameters.h"
+
 #if WITH_EDITOR
 #include "FoliageHelper.h"
 #include "ObjectCacheEventSink.h"
 #endif
-#include "PrimitiveInstanceUpdateCommand.h"
 
 static TAutoConsoleVariable<int32> CVarForceSingleSampleShadowingFromStationary(
 	TEXT("r.Shadow.ForceSingleSampleShadowingFromStationary"),
@@ -472,11 +474,10 @@ uint32 FPrimitiveSceneProxy::GetPayloadDataStride() const
 	uint32 PayloadDataCount = 0;
 
 	// Random ID is packed into scene data currently
-	// TODO: Global setting/define for INSTANCE_COMPRESSED_TRANSFORMS
-#if 0
-	PayloadDataCount += HasPerInstanceDynamicData() ? 3 : 0;	// FRenderTransform
-#else
+#if INSTANCE_SCENE_DATA_COMPRESSED_TRANSFORMS
 	PayloadDataCount += HasPerInstanceDynamicData() ? 2 : 0;	// Compressed transform
+#else
+	PayloadDataCount += HasPerInstanceDynamicData() ? 3 : 0;	// FRenderTransform
 #endif
 		
 	// Hierarchy is packed in with local bounds if they are both present (almost always the case)
diff --git a/Engine/Source/Runtime/Engine/Private/SceneManagement.cpp b/Engine/Source/Runtime/Engine/Private/SceneManagement.cpp
--- a/Engine/Source/Runtime/Engine/Private/SceneManagement.cpp
+++ b/Engine/Source/Runtime/Engine/Private/SceneManagement.cpp
@@ -291,11 +291,17 @@ void FMeshElementCollector::AddMesh(int32 ViewIndex, FMeshBatch& MeshBatch)
 		for (int32 Index = 0; Index < MeshBatch.Elements.Num(); ++Index)
 		{
 			const TUniformBuffer<FPrimitiveUniformShaderParameters>* PrimitiveUniformBufferResource = MeshBatch.Elements[Index].PrimitiveUniformBufferResource;
-
 			if (PrimitiveUniformBufferResource)
 			{
-				MeshBatch.Elements[Index].PrimitiveIdMode = PrimID_DynamicPrimitiveShaderData;
-				MeshBatch.Elements[Index].DynamicPrimitiveShaderDataIndex = DynamicPrimitiveCollectorPerView[ViewIndex]->Add(reinterpret_cast<const FPrimitiveUniformShaderParameters*>(PrimitiveUniformBufferResource->GetContents()), 1);;
+				FMeshBatchElement& Element = MeshBatch.Elements[Index];
+
+				Element.PrimitiveIdMode = PrimID_DynamicPrimitiveShaderData;
+				DynamicPrimitiveCollectorPerView[ViewIndex]->Add(
+					Element.DynamicPrimitiveData,
+					*reinterpret_cast<const FPrimitiveUniformShaderParameters*>(PrimitiveUniformBufferResource->GetContents()),
+					Element.NumInstances,
+					Element.DynamicPrimitiveIndex,
+					Element.DynamicPrimitiveInstanceSceneDataOffset);
 			}
 		}
 	}
diff --git a/Engine/Source/Runtime/Engine/Public/GPUSceneWriter.h b/Engine/Source/Runtime/Engine/Public/GPUSceneWriter.h
new file mode 100644
--- /dev/null
+++ b/Engine/Source/Runtime/Engine/Public/GPUSceneWriter.h
@@ -0,0 +1,50 @@
+// Copyright Epic Games, Inc. All Rights Reserved.
+
+#pragma once
+
+#include "CoreMinimal.h"
+#include "ShaderParameters.h"
+
+class FSceneView;
+class FRDGBuilder;
+
+/** Include in shader parameters struct for shaders that modify GPU-Scene instances */
+BEGIN_SHADER_PARAMETER_STRUCT(FGPUSceneWriterParameters, ENGINE_API)
+	SHADER_PARAMETER_UAV(RWStructuredBuffer<float4>, GPUSceneInstanceSceneDataRW)
+	SHADER_PARAMETER_UAV(RWStructuredBuffer<float4>, GPUSceneInstancePayloadDataRW)
+	SHADER_PARAMETER_UAV(RWStructuredBuffer<float4>, GPUScenePrimitiveSceneDataRW)
+	SHADER_PARAMETER(uint32, GPUSceneInstanceSceneDataSOAStride)
+	SHADER_PARAMETER(uint32, GPUSceneFrameNumber)
+	SHADER_PARAMETER(uint32, GPUSceneNumAllocatedInstances)
+	SHADER_PARAMETER(uint32, GPUSceneNumAllocatedPrimitives)
+END_SHADER_PARAMETER_STRUCT()
+
+/** Specifies the point during scene rendering that GPU Scene can be written to */
+enum class EGPUSceneGPUWritePass : int8
+{
+	/** Invalid GPU write pass. Also used to signify writes that occur during upload. */
+	None = -1,
+
+	/** Writes to the GPU Scene will happen after rendering opaques (as well has after FXSystem post-opaque updates). */
+	PostOpaqueRendering,
+
+	Num
+};
+
+/** The parameters passed to the GPUScene writer delegate */
+struct FGPUSceneWriteDelegateParams
+{
+	/** The ID of the primitive that writes must be limited to */
+	uint32 PrimitiveId = INDEX_NONE;
+	/** The ID of the first instance scene data of the primitive */
+	uint32 InstanceSceneDataOffset = INDEX_NONE;
+	/** The GPU Scene write pass that is currently executing. (NOTE: A value of None specifies that it is occurring on upload) */
+	EGPUSceneGPUWritePass GPUWritePass = EGPUSceneGPUWritePass::None;
+	/** The view for which this primitive belongs (for dynamic primitives) */
+	FSceneView* View = nullptr;
+	/** The shader parameters the delegate can use to perform writes on GPU Scene data */
+	FGPUSceneWriterParameters GPUWriteParams;
+};
+
+/** Delegate used by GPUScene to stage writing to the GPUScene primitive and instance data buffers via the GPU */
+DECLARE_DELEGATE_TwoParams(FGPUSceneWriteDelegate, FRDGBuilder&, const FGPUSceneWriteDelegateParams&);
\ No newline at end of file
diff --git a/Engine/Source/Runtime/Engine/Public/InstanceUniformShaderParameters.h b/Engine/Source/Runtime/Engine/Public/InstanceUniformShaderParameters.h
--- a/Engine/Source/Runtime/Engine/Public/InstanceUniformShaderParameters.h
+++ b/Engine/Source/Runtime/Engine/Public/InstanceUniformShaderParameters.h
@@ -25,7 +25,8 @@
 
 #define INVALID_LAST_UPDATE_FRAME 0xFFFFFFFFu
 
-#define INSTANCE_COMPRESSED_TRANSFORM	0
+#define INSTANCE_SCENE_DATA_COMPRESSED_TRANSFORMS	1
+
 // TODO: Rename to FInstanceSceneData
 struct FPrimitiveInstance
 {
@@ -83,12 +84,11 @@ struct FPrimitiveInstanceDynamicData
 struct FInstanceSceneShaderData
 {
 	// Must match GetInstanceSceneData() in SceneData.ush
-	// TODO: Global setting/define for INSTANCE_COMPRESSED_TRANSFORMS
-#if INSTANCE_COMPRESSED_TRANSFORM
-	enum { DataStrideInFloat4s = 4 };
-#else
+#if INSTANCE_SCENE_DATA_COMPRESSED_TRANSFORMS
 	// Compressed transform
 	enum { DataStrideInFloat4s = 5 }; // TODO: Temporary PrevVelocityHack
+#else
+	enum { DataStrideInFloat4s = 4 };
 #endif
 
 	TStaticArray<FVector4f, DataStrideInFloat4s> Data;
diff --git a/Engine/Source/Runtime/Engine/Public/MeshBatch.h b/Engine/Source/Runtime/Engine/Public/MeshBatch.h
--- a/Engine/Source/Runtime/Engine/Public/MeshBatch.h
+++ b/Engine/Source/Runtime/Engine/Public/MeshBatch.h
@@ -8,8 +8,10 @@
 #include "MaterialShared.h"
 #include "Engine/Scene.h"
 #include "PrimitiveUniformShaderParameters.h"
+#include "InstanceUniformShaderParameters.h"
 #include "VT/RuntimeVirtualTextureEnum.h"
 #include "EngineDefines.h"
+#include "GPUSceneWriter.h"
 
 #define USE_MESH_BATCH_VALIDATION !UE_BUILD_SHIPPING
 
@@ -24,8 +26,8 @@ enum EPrimitiveIdMode
 	PrimID_FromPrimitiveSceneInfo		= 0,
 
 	/** 
-     * The renderer will upload Primitive data from the FMeshBatchElement's PrimitiveUniformBufferResource to the end of the GPUScene PrimitiveBuffer, and assign the offset to DynamicPrimitiveShaderDataIndex.
-	 * PrimitiveId for drawing will be computed as Scene->NumPrimitives + FMeshBatchElement's DynamicPrimitiveShaderDataIndex. 
+     * The renderer will upload Primitive data from the FMeshBatchElement's PrimitiveUniformBufferResource to the end of the GPUScene PrimitiveBuffer, and assign the offset to DynamicPrimitiveIndex.
+	 * PrimitiveId for drawing will be computed as Scene->NumPrimitives + FMeshBatchElement's DynamicPrimitiveIndex. 
 	 */
 	PrimID_DynamicPrimitiveShaderData	= 1,
 
@@ -58,6 +60,87 @@ struct FMeshBatchElementDynamicIndexBuffer
 	}
 };
 
+/**
+ * Dynamic primitive/instance data for a mesh batch element.
+ * 
+ * NOTES:
+ * - When applied to a FMeshBatchElement, data provided to the TConstArrayView members are expected to live until the end of the frame on the render thread
+ * - If `DataWriterGPU` is bound and the TConstArrayView members are left empty, the delegate is expected to write any missing data, as it will not be uploaded
+ */
+struct FMeshBatchDynamicPrimitiveData
+{
+	TConstArrayView<FPrimitiveInstance> InstanceSceneData;
+	TConstArrayView<FPrimitiveInstanceDynamicData> InstanceDynamicData;
+	TConstArrayView<float> InstanceCustomData;
+	FGPUSceneWriteDelegate DataWriterGPU;		
+	EGPUSceneGPUWritePass DataWriterGPUPass = EGPUSceneGPUWritePass::None;
+	uint16 PayloadDataFlags = 0;
+	uint32 NumInstanceCustomDataFloats = 0;
+
+	FORCEINLINE void SetPayloadDataFlags(uint16 Flags, bool bValue)
+	{
+		if (bValue)
+		{
+			PayloadDataFlags |= Flags;
+		}
+		else
+		{
+			PayloadDataFlags &= ~Flags;
+		}
+	}
+	
+	FORCEINLINE void EnableInstanceDynamicData(bool bEnable) { SetPayloadDataFlags(INSTANCE_SCENE_DATA_FLAG_HAS_DYNAMIC_DATA, bEnable); }
+	FORCEINLINE void SetNumInstanceCustomDataFloats(uint32 NumFloats)
+	{
+		SetPayloadDataFlags(INSTANCE_SCENE_DATA_FLAG_HAS_CUSTOM_DATA, NumFloats > 0);
+		NumInstanceCustomDataFloats = NumFloats;
+	}
+
+	FORCEINLINE uint32 GetPayloadFloat4Stride() const
+	{
+		uint32 Total = 0;
+		if (PayloadDataFlags & INSTANCE_SCENE_DATA_FLAG_HAS_DYNAMIC_DATA)
+		{
+#if INSTANCE_SCENE_DATA_COMPRESSED_TRANSFORMS
+			Total += 2;
+#else
+			Total += 3;
+#endif
+		}
+		if (PayloadDataFlags & INSTANCE_SCENE_DATA_FLAG_HAS_CUSTOM_DATA)
+		{
+			Total += FMath::DivideAndRoundUp(NumInstanceCustomDataFloats, 4u);
+		}
+
+		return Total;
+	}
+
+	FORCEINLINE void Validate(uint32 NumInstances) const
+	{
+#if DO_CHECK
+		// Ensure array views are sized exactly for all instances, or are empty and there is a GPU writer
+		const bool bGPUWrite = DataWriterGPU.IsBound();
+		checkf(uint32(InstanceSceneData.Num()) == NumInstances || (bGPUWrite && InstanceSceneData.Num() == 0),
+			TEXT("DynamicPrimitiveData provided should have %u instances in InstanceSceneData. Found %d"),
+			NumInstances, InstanceSceneData.Num());
+		if (PayloadDataFlags & INSTANCE_SCENE_DATA_FLAG_HAS_DYNAMIC_DATA)
+		{
+			checkf(uint32(InstanceDynamicData.Num()) == NumInstances || (bGPUWrite && InstanceDynamicData.Num() == 0),
+				TEXT("DynamicPrimitiveData provided should have %u elements in InstanceDynamicData. Found %d"),
+				NumInstances, InstanceDynamicData.Num());
+		}
+		if (PayloadDataFlags & INSTANCE_SCENE_DATA_FLAG_HAS_CUSTOM_DATA)
+		{
+			checkf(NumInstanceCustomDataFloats > 0,
+				TEXT("DynamicPrimitiveData provided has the custom data flag set, but NumInstanceCustomDataFloats == 0"));
+			checkf(uint32(InstanceCustomData.Num()) == NumInstances * NumInstanceCustomDataFloats || (bGPUWrite && InstanceCustomData.Num() == 0),
+				TEXT("DynamicPrimitiveData provided should have %u elements in InstanceCustomData. Found %d"),
+				NumInstances * NumInstanceCustomDataFloats, InstanceCustomData.Num());
+		}
+#endif
+	}
+};
+
 /**
  * A batch mesh element definition.
  */
@@ -101,9 +184,6 @@ struct FMeshBatchElement
 	/** Assigned by renderer */
 	EPrimitiveIdMode PrimitiveIdMode : PrimID_NumBits + 1;
 
-	/** Assigned by renderer */
-	uint32 DynamicPrimitiveShaderDataIndex : 24;
-
 	uint32 FirstIndex;
 	/** When 0, IndirectArgsBuffer will be used. */
 	uint32 NumPrimitives;
@@ -128,6 +208,14 @@ struct FMeshBatchElement
 	int32 VisualizeElementIndex : 8;
 #endif
 
+	/**
+	 * Source instance scene data and payload data for dynamic primitives. Must be provided for dynamic primitives that have more than a single instance.
+	 * NOTE: The lifetime of the object pointed to is expected to match or exceed that of the mesh batch itself.
+	 **/
+	const FMeshBatchDynamicPrimitiveData* DynamicPrimitiveData;
+	uint32 DynamicPrimitiveIndex;
+	uint32 DynamicPrimitiveInstanceSceneDataOffset;
+
 	FORCEINLINE int32 GetNumPrimitives() const
 	{
 		if (bIsInstanceRuns && InstanceRuns)
@@ -155,7 +243,6 @@ struct FMeshBatchElement
 	,	IndirectArgsBuffer(nullptr)
 	,	IndirectArgsOffset(0)
 	,	PrimitiveIdMode(PrimID_FromPrimitiveSceneInfo)
-	,	DynamicPrimitiveShaderDataIndex(0)
 	,	NumInstances(1)
 	,	BaseVertexIndex(0)
 	,	UserIndex(-1)
@@ -166,9 +253,12 @@ struct FMeshBatchElement
 	,	bUserDataIsColorVertexBuffer(false)
 	,	bIsSplineProxy(false)
 	,	bIsInstanceRuns(false)
-#if !(UE_BUILD_SHIPPING || UE_BUILD_TEST)
+#if UE_ENABLE_DEBUG_DRAWING
 	,	VisualizeElementIndex(INDEX_NONE)
 #endif
+	,	DynamicPrimitiveData(nullptr)
+	,	DynamicPrimitiveIndex(INDEX_NONE)
+	,	DynamicPrimitiveInstanceSceneDataOffset(INDEX_NONE)
 	{
 	}
 };
diff --git a/Engine/Source/Runtime/Renderer/Private/DeferredShadingRenderer.cpp b/Engine/Source/Runtime/Renderer/Private/DeferredShadingRenderer.cpp
--- a/Engine/Source/Runtime/Renderer/Private/DeferredShadingRenderer.cpp
+++ b/Engine/Source/Runtime/Renderer/Private/DeferredShadingRenderer.cpp
@@ -3032,6 +3032,11 @@ void FDeferredShadingSceneRenderer::Render(FRDGBuilder& GraphBuilder)
 
 	RenderOpaqueFX(GraphBuilder, Views, FXSystem, SceneTextures.UniformBuffer);
 
+	if (Scene->GPUScene.ExecuteDeferredGPUWritePass(GraphBuilder, Views, EGPUSceneGPUWritePass::PostOpaqueRendering))
+	{
+		InstanceCullingManager.BeginDeferredCulling(GraphBuilder, Scene->GPUScene);
+	}
+
 	if (GetHairStrandsComposition() == EHairStrandsCompositionType::BeforeTranslucent)
 	{
 		RDG_GPU_STAT_SCOPE(GraphBuilder, HairRendering);
diff --git a/Engine/Source/Runtime/Renderer/Private/DynamicPrimitiveDrawing.inl b/Engine/Source/Runtime/Renderer/Private/DynamicPrimitiveDrawing.inl
--- a/Engine/Source/Runtime/Renderer/Private/DynamicPrimitiveDrawing.inl
+++ b/Engine/Source/Runtime/Renderer/Private/DynamicPrimitiveDrawing.inl
@@ -244,13 +244,18 @@ inline int32 FViewElementPDI::DrawMesh(const FMeshBatch& Mesh)
 					{
 						for (int32 Index = 0; Index < NewMesh->Elements.Num(); ++Index)
 						{
-							const TUniformBuffer<FPrimitiveUniformShaderParameters>* PrimitiveUniformBufferResource = NewMesh->Elements[Index].PrimitiveUniformBufferResource;
+							FMeshBatchElement& Element = NewMesh->Elements[Index];
+							const TUniformBuffer<FPrimitiveUniformShaderParameters>* PrimitiveUniformBufferResource = Element.PrimitiveUniformBufferResource;
 
 							if (PrimitiveUniformBufferResource)
 							{
-								const int32 DataIndex = DynamicPrimitiveCollectorForRT->Add(reinterpret_cast<const FPrimitiveUniformShaderParameters*>(PrimitiveUniformBufferResource->GetContents()), 1);
-								NewMesh->Elements[Index].PrimitiveIdMode = PrimID_DynamicPrimitiveShaderData;
-								NewMesh->Elements[Index].DynamicPrimitiveShaderDataIndex = DataIndex;
+								Element.PrimitiveIdMode = PrimID_DynamicPrimitiveShaderData;
+								DynamicPrimitiveCollectorForRT->Add(
+									Element.DynamicPrimitiveData,
+									*reinterpret_cast<const FPrimitiveUniformShaderParameters*>(PrimitiveUniformBufferResource->GetContents()),
+									Element.NumInstances,
+									Element.DynamicPrimitiveIndex,
+									Element.DynamicPrimitiveInstanceSceneDataOffset);
 							}
 						}
 					}
diff --git a/Engine/Source/Runtime/Renderer/Private/GPUScene.cpp b/Engine/Source/Runtime/Renderer/Private/GPUScene.cpp
--- a/Engine/Source/Runtime/Renderer/Private/GPUScene.cpp
+++ b/Engine/Source/Runtime/Renderer/Private/GPUScene.cpp
@@ -21,17 +21,11 @@
 #include "HAL/LowLevelMemTracker.h"
 #include "HAL/LowLevelMemStats.h"
 #include "ShaderDebug.h"
+#include "InstanceUniformShaderParameters.h"
 #include "ShaderPrint.h"
 
 #define LOG_INSTANCE_ALLOCATIONS 0
 
-// TODO: Global setting/define
-#ifndef INSTANCE_COMPRESSED_TRANSFORMS
-#define INSTANCE_COMPRESSED_TRANSFORMS 1
-#else
-#error "This should not be defined yet"
-#endif
-
 int32 GGPUSceneUploadEveryFrame = 0;
 FAutoConsoleVariableRef CVarGPUSceneUploadEveryFrame(
 	TEXT("r.GPUScene.UploadEveryFrame"),
@@ -172,6 +166,84 @@ static int32 PartitionUpdateRanges(FParallelUpdateRanges& Ranges, int32 ItemCoun
 	return Ranges.Range[3].ItemCount > 0 ? 4 : 3;
 }
 
+void FGPUScenePrimitiveCollector::Add(
+	const FMeshBatchDynamicPrimitiveData* MeshBatchData,
+	const FPrimitiveUniformShaderParameters& PrimitiveShaderParams,
+	uint32 NumInstances,
+	uint32& OutPrimitiveIndex,
+	uint32& OutInstanceSceneDataOffset)
+{
+	check(GPUSceneDynamicContext != nullptr);
+	check(!bCommitted);
+	
+
+	// Lazy allocation of the upload data to not waste space and processing if none was needed.
+	if (UploadData == nullptr)
+	{
+		UploadData = AllocateUploadData();
+	}
+
+	const int32 PrimitiveIndex = UploadData->PrimitiveData.Num();
+	FPrimitiveData& PrimitiveData = UploadData->PrimitiveData.AddDefaulted_GetRef();		
+
+	if (MeshBatchData != nullptr)
+	{
+		// make sure the source data is appropriately structured
+		MeshBatchData->Validate(NumInstances);
+		PrimitiveData.SourceData = *MeshBatchData;
+	}
+	
+	const int32 PayloadFloat4Stride = PrimitiveData.SourceData.GetPayloadFloat4Stride();
+	
+	PrimitiveData.ShaderParams = &PrimitiveShaderParams;
+	PrimitiveData.NumInstances = NumInstances;
+	PrimitiveData.LocalInstanceSceneDataOffset = UploadData->TotalInstanceCount;
+	PrimitiveData.LocalPayloadDataOffset = PayloadFloat4Stride > 0 ? UploadData->InstancePayloadDataFloat4Count : INDEX_NONE;
+
+	UploadData->TotalInstanceCount += NumInstances;	
+	UploadData->InstancePayloadDataFloat4Count += PayloadFloat4Stride * NumInstances;
+
+	if (PrimitiveData.SourceData.DataWriterGPU.IsBound())
+	{
+		// Enqueue this primitive data to be executed (either upon upload or deferred to a later GPU write pass)
+		UploadData->GPUWritePrimitives.Add(PrimitiveIndex);
+	}
+
+	// Set the output data offsets
+	OutPrimitiveIndex = PrimitiveIndex;
+	OutInstanceSceneDataOffset = PrimitiveData.LocalInstanceSceneDataOffset;
+}
+
+#if DO_CHECK
+
+bool FGPUScenePrimitiveCollector::IsPrimitiveProcessed(uint32 PrimitiveIndex, const FGPUScene& GPUScene) const
+{
+	if (UploadData == nullptr || !bCommitted)
+	{
+		// The collector hasn't collected anything or hasn't been uploaded
+		return false;
+	}
+
+	if (PrimitiveIndex >= uint32(UploadData->PrimitiveData.Num()))
+	{
+		// The specified index is out of range
+		return false;
+	}
+
+	const FMeshBatchDynamicPrimitiveData& SourceData = UploadData->PrimitiveData[PrimitiveIndex].SourceData;
+	if (!SourceData.DataWriterGPU.IsBound() || SourceData.DataWriterGPUPass == EGPUSceneGPUWritePass::None)
+	{
+		// The primitive doesn't have a pending GPU write and has been uploaded or written to by the GPU already
+		return true;
+	}
+
+	// If the GPU scene still has a pending deferred write for the primitive, then it has not been fully processed yet
+	const uint32 PrimitiveId = GetPrimitiveIdRange().GetLowerBoundValue() + PrimitiveIndex;
+	return !GPUScene.HasPendingGPUWrite(PrimitiveId);
+}
+
+#endif // DO_CHECK
+
 void FGPUScenePrimitiveCollector::Commit()
 {
 	ensure(!bCommitted);
@@ -541,6 +613,8 @@ void FGPUScene::UpdateInternal(FRDGBuilder& GraphBuilder, FScene& Scene)
 	ensure(NumScenePrimitives == Scene.Primitives.Num());
 	ensure(DynamicPrimitivesOffset >= Scene.Primitives.Num());
 
+	LastDeferredGPUWritePass = EGPUSceneGPUWritePass::None;
+
 	if (GGPUSceneUploadEveryFrame || bUpdateAllPrimitives)
 	{
 		PrimitivesToUpdate.Reset();
@@ -1303,7 +1377,7 @@ void FGPUScene::UploadGeneral(FRHICommandListImmediate& RHICmdList, FScene *Scen
 									{
 										check(UploadInfo.InstanceDynamicData.Num() == UploadInfo.PrimitiveInstances.Num());
 										const FRenderTransform PrevLocalToWorld = UploadInfo.InstanceDynamicData[InstanceIndex].ComputePrevLocalToWorld(UploadInfo.PrevPrimitiveToWorld);
-									#if INSTANCE_COMPRESSED_TRANSFORMS
+									#if INSTANCE_SCENE_DATA_COMPRESSED_TRANSFORMS
 										check(PayloadPosition + 1 < InstancePayloadData.Num()); // Sanity check
 										FCompressedTransform CompressedPrevLocalToWorld(PrevLocalToWorld);
 										InstancePayloadData[PayloadPosition + 0] = *(const FVector4f*)&CompressedPrevLocalToWorld.Rotation[0];
@@ -1447,17 +1521,19 @@ struct FUploadDataSourceAdapterDynamicPrimitives
 	static constexpr bool bUpdateNaniteMaterialTables = false;
 
 	FUploadDataSourceAdapterDynamicPrimitives(
-		const TArray<FPrimitiveUniformShaderParameters, TInlineAllocator<8>>& InPrimitiveShaderData,
+		const TArray<FGPUScenePrimitiveCollector::FPrimitiveData, TInlineAllocator<8>>& InPrimitiveData,
 		int32 InPrimitiveIDStartOffset,
 		int32 InInstanceIDStartOffset,
+		int32 InPayloadStartOffset,
 		uint32 InSceneFrameNumber)
-		: PrimitiveShaderData(InPrimitiveShaderData)
+		: PrimitiveData(InPrimitiveData)
 		, PrimitiveIDStartOffset(InPrimitiveIDStartOffset)
 		, InstanceIDStartOffset(InInstanceIDStartOffset)
+		, PayloadStartOffset(InPayloadStartOffset)
 		, SceneFrameNumber(InSceneFrameNumber)
 	{
 		// Need to create this explicit for optimizing the common path
-		PrimitivesIds.SetNumUninitialized(PrimitiveShaderData.Num());
+		PrimitivesIds.SetNumUninitialized(PrimitiveData.Num());
 		for (int32 Index = 0; Index < PrimitivesIds.Num(); ++Index)
 		{
 			PrimitivesIds[Index] = uint32(PrimitiveIDStartOffset + Index);
@@ -1466,7 +1542,7 @@ struct FUploadDataSourceAdapterDynamicPrimitives
 
 	FORCEINLINE int32 NumPrimitivesToUpload() const
 	{ 
-		return PrimitiveShaderData.Num();
+		return PrimitiveData.Num();
 	}
 
 
@@ -1482,11 +1558,23 @@ struct FUploadDataSourceAdapterDynamicPrimitives
 		PrimitiveUploadInfo.NaniteSceneProxy = nullptr;
 		PrimitiveUploadInfo.PrimitiveSceneInfo = nullptr;
 
-		check(ItemIndex < PrimitiveShaderData.Num());
+		check(ItemIndex < PrimitiveData.Num());
 
-		PrimitiveUploadInfo.NumInstanceUploads = 1;
-		PrimitiveUploadInfo.NumInstancePayloadDataUploads = 0;
 		PrimitiveUploadInfo.PrimitiveID = PrimitiveIDStartOffset + ItemIndex;
+
+		const FGPUScenePrimitiveCollector::FPrimitiveData& PrimData = PrimitiveData[ItemIndex];
+		PrimitiveUploadInfo.NumInstanceUploads = PrimData.NumInstances;
+		PrimitiveUploadInfo.NumInstancePayloadDataUploads = PrimData.SourceData.GetPayloadFloat4Stride() * PrimData.NumInstances;
+		
+		if (PrimData.SourceData.DataWriterGPU.IsBound())
+		{
+			// Only upload if we have data, otherwise expect the delegate to handle missing data
+			PrimitiveUploadInfo.NumInstanceUploads = PrimData.SourceData.InstanceSceneData.Num();
+			if (PrimData.SourceData.InstanceCustomData.Num() == 0)
+			{
+				PrimitiveUploadInfo.NumInstancePayloadDataUploads = 0;
+			}
+		}
 	}
 
 	FORCEINLINE void GetPrimitiveInfo(int32 ItemIndex, FPrimitiveUploadInfo& PrimitiveUploadInfo) const
@@ -1494,44 +1582,55 @@ struct FUploadDataSourceAdapterDynamicPrimitives
 		GetPrimitiveInfoHeader(ItemIndex, PrimitiveUploadInfo);
 
 		// Needed to ensure the link back to instance list is up to date
-		FPrimitiveUniformShaderParameters Tmp = PrimitiveShaderData[ItemIndex];
-		Tmp.InstanceSceneDataOffset = InstanceIDStartOffset + ItemIndex;
-		Tmp.NumInstanceSceneDataEntries = 1;
-		Tmp.InstancePayloadDataOffset = INDEX_NONE;
-		Tmp.InstancePayloadDataStride = 0;
-		PrimitiveUploadInfo.NumInstanceUploads = 1;
-		PrimitiveUploadInfo.NumInstancePayloadDataUploads = 0;
-		PrimitiveUploadInfo.PrimitiveID = PrimitiveIDStartOffset + ItemIndex;
+		const FGPUScenePrimitiveCollector::FPrimitiveData& PrimData = PrimitiveData[ItemIndex];		
+		FPrimitiveUniformShaderParameters Tmp = *PrimData.ShaderParams;
+		Tmp.InstanceSceneDataOffset = InstanceIDStartOffset + PrimData.LocalInstanceSceneDataOffset;
+		Tmp.NumInstanceSceneDataEntries = PrimData.NumInstances;			
+		if (PrimitiveData[ItemIndex].LocalPayloadDataOffset != INDEX_NONE)
+		{
+			Tmp.InstancePayloadDataOffset = PayloadStartOffset + PrimData.LocalPayloadDataOffset;
+			Tmp.InstancePayloadDataStride = PrimData.SourceData.GetPayloadFloat4Stride();
+		}
+		else
+		{
+			Tmp.InstancePayloadDataOffset = INDEX_NONE;
+			Tmp.InstancePayloadDataStride = 0;
+		}
+
 		PrimitiveUploadInfo.PrimitiveSceneData = FPrimitiveSceneShaderData(Tmp);
 	}
 
 	FORCEINLINE bool GetInstanceInfo(int32 ItemIndex, FInstanceUploadInfo& InstanceUploadInfo) const
 	{
-		if (ItemIndex < PrimitiveShaderData.Num())
+		if (ItemIndex < PrimitiveData.Num())
 		{
-			const FPrimitiveUniformShaderParameters& PrimitiveData = PrimitiveShaderData[ItemIndex];
-
-			InstanceUploadInfo.DummyInstance.LocalToPrimitive.SetIdentity();
-			InstanceUploadInfo.DummyLocalBounds = FRenderBounds(PrimitiveData.LocalObjectBoundsMin, PrimitiveData.LocalObjectBoundsMax);
+			const FGPUScenePrimitiveCollector::FPrimitiveData& PrimData = PrimitiveData[ItemIndex];
+			const FPrimitiveUniformShaderParameters& ShaderParams = *PrimData.ShaderParams;
 
 			InstanceUploadInfo.PrimitiveID					= PrimitiveIDStartOffset + ItemIndex;
-			InstanceUploadInfo.PrimitiveToWorld				= PrimitiveData.LocalToRelativeWorld;
-			InstanceUploadInfo.PrevPrimitiveToWorld			= PrimitiveData.PreviousLocalToRelativeWorld;
-			InstanceUploadInfo.PrimitiveInstances			= TConstArrayView<FPrimitiveInstance>(&InstanceUploadInfo.DummyInstance, 1);
-			InstanceUploadInfo.InstanceDynamicData			= TConstArrayView<FPrimitiveInstanceDynamicData>();
-			InstanceUploadInfo.InstanceLightShadowUVBias	= TConstArrayView<FVector4f>();
-			InstanceUploadInfo.InstanceCustomData			= TConstArrayView<float>();
+			InstanceUploadInfo.PrimitiveToWorld				= ShaderParams.LocalToRelativeWorld;
+			InstanceUploadInfo.PrevPrimitiveToWorld			= ShaderParams.PreviousLocalToRelativeWorld;
+			InstanceUploadInfo.InstanceSceneDataOffset		= InstanceIDStartOffset + PrimData.LocalInstanceSceneDataOffset;
+			InstanceUploadInfo.InstancePayloadDataOffset	= PrimData.LocalPayloadDataOffset == INDEX_NONE ? INDEX_NONE : PayloadStartOffset + PrimData.LocalPayloadDataOffset;
+			InstanceUploadInfo.InstancePayloadDataStride	= PrimData.SourceData.GetPayloadFloat4Stride();
+			InstanceUploadInfo.InstanceCustomDataCount		= PrimData.SourceData.NumInstanceCustomDataFloats;
+			InstanceUploadInfo.PayloadDataFlags				= PrimData.SourceData.PayloadDataFlags;
+			InstanceUploadInfo.PrimitiveInstances 			= PrimData.SourceData.InstanceSceneData;
+			InstanceUploadInfo.InstanceDynamicData 			= PrimData.SourceData.InstanceDynamicData;
+			InstanceUploadInfo.InstanceCustomData 			= PrimData.SourceData.InstanceCustomData;
 			InstanceUploadInfo.InstanceRandomID				= TConstArrayView<float>();
-			InstanceUploadInfo.InstanceLocalBounds			= TConstArrayView<FRenderBounds>(&InstanceUploadInfo.DummyLocalBounds, 1);
 			InstanceUploadInfo.InstanceHierarchyOffset		= TConstArrayView<uint32>();
+			InstanceUploadInfo.InstanceLightShadowUVBias	= TConstArrayView<FVector4f>();
 #if WITH_EDITOR
 			InstanceUploadInfo.InstanceEditorData			= TConstArrayView<uint32>();
 #endif
-			InstanceUploadInfo.InstanceSceneDataOffset		= InstanceIDStartOffset + ItemIndex; 
-			InstanceUploadInfo.PayloadDataFlags				= 0;
-			InstanceUploadInfo.InstancePayloadDataOffset	= INDEX_NONE;
-			InstanceUploadInfo.InstancePayloadDataStride	= 0;
-			InstanceUploadInfo.InstanceCustomDataCount		= 0;
+
+			// upload dummies where applicable
+			if (InstanceUploadInfo.PrimitiveInstances.Num() == 0)
+			{
+				InstanceUploadInfo.DummyInstance.LocalToPrimitive.SetIdentity();
+				InstanceUploadInfo.PrimitiveInstances = TConstArrayView<FPrimitiveInstance>(&InstanceUploadInfo.DummyInstance, 1);
+			}
 
 			return true;
 		}
@@ -1544,9 +1643,10 @@ struct FUploadDataSourceAdapterDynamicPrimitives
 		return false;
 	}
 
-	const TArray<FPrimitiveUniformShaderParameters, TInlineAllocator<8>> &PrimitiveShaderData;
+	const TArray<FGPUScenePrimitiveCollector::FPrimitiveData, TInlineAllocator<8>> &PrimitiveData;
 	const int32 PrimitiveIDStartOffset;
 	const int32 InstanceIDStartOffset;
+	const int32 PayloadStartOffset;
 	const uint32 SceneFrameNumber;
 	TArray<uint32, SceneRenderingAllocator> PrimitivesIds;
 };
@@ -1555,6 +1655,8 @@ void FGPUScene::UploadDynamicPrimitiveShaderDataForViewInternal(FRDGBuilder& Gra
 {
 	LLM_SCOPE_BYTAG(GPUScene);
 
+	RDG_EVENT_SCOPE(GraphBuilder, "GPUScene.UploadDynamicPrimitiveShaderDataForView");
+
 	ensure(bInBeginEndBlock);
 	ensure(Scene == nullptr || DynamicPrimitivesOffset >= Scene->Primitives.Num());
 
@@ -1574,24 +1676,48 @@ void FGPUScene::UploadDynamicPrimitiveShaderDataForViewInternal(FRDGBuilder& Gra
 	ensure(Collector.UploadData == nullptr || CurrentDynamicContext->DymamicPrimitiveUploadData.Find(Collector.UploadData) != INDEX_NONE);
 
 	// Skip uploading empty & already uploaded data
-	if (Collector.UploadData != nullptr && NumPrimitiveDataUploads > 0 && !Collector.UploadData->bIsUploaded)
+	const bool bNeedsUpload = Collector.UploadData != nullptr && NumPrimitiveDataUploads > 0 && !Collector.UploadData->bIsUploaded;
+	if (bNeedsUpload)
 	{
 		CSV_SCOPED_TIMING_STAT_EXCLUSIVE(UploadDynamicPrimitiveShaderData);
 		QUICK_SCOPE_CYCLE_COUNTER(STAT_UploadDynamicPrimitiveShaderData);
 
 		Collector.UploadData->bIsUploaded = true;
 
-		int32 UploadIdStart = Collector.GetPrimitiveIdRange().GetLowerBoundValue();
+		const int32 UploadIdStart = Collector.GetPrimitiveIdRange().GetLowerBoundValue();
+		const int32 InstanceIdStart = Collector.UploadData->InstanceSceneDataOffset;
 		ensure(UploadIdStart < DynamicPrimitivesOffset);
-		ensure(Collector.UploadData->InstanceSceneDataOffset != INDEX_NONE);
+		ensure(InstanceIdStart != INDEX_NONE);
 
-		FUploadDataSourceAdapterDynamicPrimitives UploadAdapter(Collector.UploadData->PrimitiveShaderData, UploadIdStart, Collector.UploadData->InstanceSceneDataOffset, SceneFrameNumber);
+		if (Scene != nullptr && Scene->VirtualShadowMapArrayCacheManager != nullptr)
+		{
+			// Enqueue cache invalidations for all dynamic primitives' instances, as they will be removed this frame and are not associated
+			// with any particular FPrimitiveSceneInfo. Will occur on the next call to UpdateAllPrimitiveSceneInfos
+			for (const FGPUScenePrimitiveCollector::FPrimitiveData& PrimitiveData : Collector.UploadData->PrimitiveData)
+			{
+				ensure(PrimitiveData.LocalInstanceSceneDataOffset != INDEX_NONE);				
+				DynamicPrimitiveInstancesToInvalidate.Add(
+					FInstanceRange
+					{
+						PrimitiveData.LocalInstanceSceneDataOffset + InstanceIdStart,
+						PrimitiveData.NumInstances
+					}
+				);
+			}
+		}
+
+		FUploadDataSourceAdapterDynamicPrimitives UploadAdapter(
+			Collector.UploadData->PrimitiveData,
+			UploadIdStart,
+			InstanceIdStart,
+			Collector.UploadData->InstancePayloadDataOffset,
+			SceneFrameNumber);
 		FGPUSceneBufferState BufferState = UpdateBufferState(GraphBuilder, Scene, UploadAdapter);
 
 		// Run a pass that clears (Sets ID to invalid) any instances that need it.
 		AddClearInstancesPass(GraphBuilder);
 
-		AddPass(GraphBuilder, RDG_EVENT_NAME("GPUScene::UploadDynamicPrimitiveShaderDataForView"),
+		AddPass(GraphBuilder, RDG_EVENT_NAME("Uploads"),
 			[this, Scene, UploadAdapter, BufferState = MoveTemp(BufferState)](FRHICommandListImmediate& RHICmdList)
 		{
 			UploadGeneral<FUploadDataSourceAdapterDynamicPrimitives>(RHICmdList, Scene, UploadAdapter, BufferState);
@@ -1612,6 +1738,63 @@ void FGPUScene::UploadDynamicPrimitiveShaderDataForViewInternal(FRDGBuilder& Gra
 	View.CachedViewUniformShaderParameters->InstanceSceneDataSOAStride = InstanceSceneDataSOAStride;
 
 	View.ViewUniformBuffer.UpdateUniformBufferImmediate(*View.CachedViewUniformShaderParameters);
+
+	// Execute any instance data GPU writer callbacks. (Note: Done after the UB update, in case the user requires it)
+	if (bNeedsUpload) 
+	{
+		const uint32 PrimitiveIdStart = Collector.GetPrimitiveIdRange().GetLowerBoundValue();
+		const uint32 InstanceIdStart = Collector.UploadData->InstanceSceneDataOffset;
+
+		// Determine if we have any GPU data writers this frame and simultaneously defer any writes that must happen later in the frame
+		TArray<uint32, TMemStackAllocator<>> ImmediateWrites;
+		ImmediateWrites.Reserve(Collector.UploadData->GPUWritePrimitives.Num());
+		for (uint32 PrimitiveIndex : Collector.UploadData->GPUWritePrimitives)
+		{
+			const FGPUScenePrimitiveCollector::FPrimitiveData& PrimData = Collector.UploadData->PrimitiveData[PrimitiveIndex];
+			const EGPUSceneGPUWritePass GPUWritePass = PrimData.SourceData.DataWriterGPUPass;
+			
+			// We're going to immediately execute any GPU writers whose write pass is immediate or has already happened this frame
+			if (GPUWritePass == EGPUSceneGPUWritePass::None || GPUWritePass <= LastDeferredGPUWritePass)
+			{
+				ImmediateWrites.Add(PrimitiveIndex);
+			}
+			else
+			{
+				// Defer this write to a later GPU write pass
+				FDeferredGPUWrite DeferredWrite;
+				DeferredWrite.DataWriterGPU = PrimData.SourceData.DataWriterGPU;
+				DeferredWrite.ViewId = View.GPUSceneViewId;
+				DeferredWrite.PrimitiveId = PrimitiveIdStart + PrimitiveIndex;
+				DeferredWrite.InstanceSceneDataOffset = InstanceIdStart + PrimData.LocalInstanceSceneDataOffset;
+
+				uint32 PassIndex = uint32(PrimData.SourceData.DataWriterGPUPass);
+				DeferredGPUWritePassDelegates[PassIndex].Add(DeferredWrite);
+			}
+		}		
+
+		if (ImmediateWrites.Num() > 0)
+		{
+			// Execute writes that should execute immediately
+			RDG_EVENT_SCOPE(GraphBuilder, "GPU Writer Delegates");
+			BeginReadWriteAccess(GraphBuilder, true);			
+
+			FGPUSceneWriteDelegateParams Params;
+			Params.View = &View;
+			Params.GPUWritePass = EGPUSceneGPUWritePass::None;
+			GetWriteParameters(Params.GPUWriteParams);
+
+			for (uint32 PrimitiveIndex : ImmediateWrites)
+			{
+				const FGPUScenePrimitiveCollector::FPrimitiveData& PrimData = Collector.UploadData->PrimitiveData[PrimitiveIndex];
+				Params.PrimitiveId = PrimitiveIdStart + PrimitiveIndex;
+				Params.InstanceSceneDataOffset = InstanceIdStart + PrimData.LocalInstanceSceneDataOffset;
+
+				PrimData.SourceData.DataWriterGPU.Execute(GraphBuilder, Params);
+			}
+
+			EndReadWriteAccess(GraphBuilder, ERHIAccess::SRVMask);
+		}
+	}
 }
 
 void AddPrimitiveToUpdateGPU(FScene& Scene, int32 PrimitiveId)
@@ -1921,13 +2104,79 @@ TRange<int32> FGPUScene::CommitPrimitiveCollector(FGPUScenePrimitiveCollector& P
 	ensure(CurrentDynamicContext == nullptr || CurrentDynamicContext->DymamicPrimitiveUploadData.Find(PrimitiveCollector.UploadData) != INDEX_NONE);
 
 	int32 StartOffset = DynamicPrimitivesOffset;
-	DynamicPrimitivesOffset += PrimitiveCollector.UploadData->PrimitiveShaderData.Num();
+	DynamicPrimitivesOffset += PrimitiveCollector.Num();
 
-	PrimitiveCollector.UploadData->InstanceSceneDataOffset = AllocateInstanceSceneDataSlots(PrimitiveCollector.UploadData->PrimitiveShaderData.Num());
+	PrimitiveCollector.UploadData->InstanceSceneDataOffset = AllocateInstanceSceneDataSlots(PrimitiveCollector.NumInstances());
+	PrimitiveCollector.UploadData->InstancePayloadDataOffset = AllocateInstancePayloadDataSlots(PrimitiveCollector.NumPayloadDataSlots());
 
 	return TRange<int32>(StartOffset, DynamicPrimitivesOffset);
 }
 
+
+bool FGPUScene::ExecuteDeferredGPUWritePass(FRDGBuilder& GraphBuilder, const TArrayView<FViewInfo>& Views, EGPUSceneGPUWritePass GPUWritePass)
+{
+	check(GPUWritePass != EGPUSceneGPUWritePass::None && GPUWritePass < EGPUSceneGPUWritePass::Num);
+	check(LastDeferredGPUWritePass < GPUWritePass);
+
+	if (!bIsEnabled)
+	{
+		return false;
+	}
+
+	// Mark this pass as having executed for the frame
+	LastDeferredGPUWritePass = GPUWritePass;
+
+	const uint32 PassIndex = uint32(GPUWritePass);
+	if (DeferredGPUWritePassDelegates[PassIndex].Num() == 0)
+	{
+		// No deferred writes to make for this pass this frame
+		return false;
+	}
+
+	RDG_EVENT_SCOPE(GraphBuilder, "GPUScene.DeferredGPUWrites - Pass %u", uint32(GPUWritePass));
+	BeginReadWriteAccess(GraphBuilder, true);
+
+	FGPUSceneWriteDelegateParams Params;
+	Params.GPUWritePass = GPUWritePass;
+	GetWriteParameters(Params.GPUWriteParams);
+
+	for (const FDeferredGPUWrite& DeferredWrite : DeferredGPUWritePassDelegates[PassIndex])
+	{
+		FViewInfo* View = Views.FindByPredicate([&DeferredWrite](const FViewInfo& V){ return V.GPUSceneViewId == DeferredWrite.ViewId; });
+		checkf(View != nullptr, TEXT("Deferred GPU Write found with no matching view in the view family"));
+		
+		Params.View = View;
+		Params.PrimitiveId = DeferredWrite.PrimitiveId;
+		Params.InstanceSceneDataOffset = DeferredWrite.InstanceSceneDataOffset;
+
+		DeferredWrite.DataWriterGPU.Execute(GraphBuilder, Params);
+	}
+
+	EndReadWriteAccess(GraphBuilder, ERHIAccess::SRVMask);
+
+	DeferredGPUWritePassDelegates[PassIndex].Reset();
+	return true;
+}
+
+
+bool FGPUScene::HasPendingGPUWrite(uint32 PrimitiveId) const
+{
+	for (uint32 PassIndex = uint32(LastDeferredGPUWritePass) + 1; PassIndex < uint32(EGPUSceneGPUWritePass::Num); ++PassIndex)
+	{
+		if (DeferredGPUWritePassDelegates[PassIndex].FindByPredicate(
+			[PrimitiveId](const FDeferredGPUWrite& Write)
+			{
+				return Write.PrimitiveId == PrimitiveId;
+			}))
+		{
+			return true;
+		}
+	}
+
+	return false;
+}
+
+
 FGPUSceneDynamicContext::~FGPUSceneDynamicContext()
 {
 	Release();
@@ -1938,7 +2187,11 @@ void FGPUSceneDynamicContext::Release()
 	for (auto UploadData : DymamicPrimitiveUploadData)
 	{
 		check(UploadData->InstanceSceneDataOffset != INDEX_NONE);
-		GPUScene.FreeInstanceSceneDataSlots(UploadData->InstanceSceneDataOffset, UploadData->PrimitiveShaderData.Num());
+		GPUScene.FreeInstanceSceneDataSlots(UploadData->InstanceSceneDataOffset, UploadData->TotalInstanceCount);
+		if (UploadData->InstancePayloadDataOffset != INDEX_NONE)
+		{
+			GPUScene.FreeInstancePayloadDataSlots(UploadData->InstancePayloadDataOffset, UploadData->InstancePayloadDataFloat4Count);
+		}
 		delete UploadData;
 	}
 	DymamicPrimitiveUploadData.Empty();
@@ -1956,12 +2209,17 @@ FGPUScenePrimitiveCollector::FUploadData* FGPUSceneDynamicContext::AllocateDynam
 /**
  * Call before accessing the GPU scene in a read/write pass.
  */
-bool FGPUScene::BeginReadWriteAccess(FRDGBuilder& GraphBuilder)
+bool FGPUScene::BeginReadWriteAccess(FRDGBuilder& GraphBuilder, bool bAllowUAVOverlap)
 {
 	if (IsEnabled())
 	{
+		checkf(!bReadWriteAccess, TEXT("GPUScene's buffers already have r/w access"));
+		bReadWriteAccess = true;
+		bReadWriteUAVOverlap = bAllowUAVOverlap;
+
 		// TODO: Remove this when everything is properly RDG'd
-		AddPass(GraphBuilder, RDG_EVENT_NAME("GPUScene::TransitionInstanceSceneDataBuffer"), [this](FRHICommandList& RHICmdList)
+		AddPass(GraphBuilder, RDG_EVENT_NAME("GPUScene::TransitionInstanceSceneDataBuffer"),
+			[this, bAllowUAVOverlap](FRHICommandList& RHICmdList)
 		{
 			FRHITransitionInfo 	Transitions[] =
 			{
@@ -1971,6 +2229,11 @@ bool FGPUScene::BeginReadWriteAccess(FRDGBuilder& GraphBuilder)
 			};
 
 			RHICmdList.Transition(Transitions);
+
+			if (bAllowUAVOverlap) // NOTE: using a capture because the member is set on a different timeline
+			{
+				RHICmdList.BeginUAVOverlap({ InstanceSceneDataBuffer.UAV, InstancePayloadDataBuffer.UAV, PrimitiveBuffer.UAV });
+			}
 		});
 		return true;
 	}
@@ -1999,9 +2262,17 @@ void FGPUScene::EndReadWriteAccess(FRDGBuilder& GraphBuilder, ERHIAccess FinalAc
 {
 	if (IsEnabled())
 	{
+		checkf(bReadWriteAccess, TEXT("GPUScene's buffers do not currently have r/w access"));
+
 		// TODO: Remove this when everything is properly RDG'd
-		AddPass(GraphBuilder, RDG_EVENT_NAME("GPUScene::TransitionInstanceSceneDataBuffer"), [this, FinalAccessState](FRHICommandList& RHICmdList)
+		AddPass(GraphBuilder, RDG_EVENT_NAME("GPUScene::TransitionInstanceSceneDataBuffer"),
+			[this, FinalAccessState, bEndUAVOverlap=bReadWriteUAVOverlap](FRHICommandList& RHICmdList)
 		{
+			if (bEndUAVOverlap) // NOTE: using a capture because the member is set on a different timeline
+			{
+				RHICmdList.EndUAVOverlap({ InstanceSceneDataBuffer.UAV, InstancePayloadDataBuffer.UAV, PrimitiveBuffer.UAV });
+			}
+
 			FRHITransitionInfo 	Transitions[] =
 			{
 				FRHITransitionInfo(InstanceSceneDataBuffer.UAV, ERHIAccess::UAVCompute, FinalAccessState),
@@ -2011,6 +2282,9 @@ void FGPUScene::EndReadWriteAccess(FRDGBuilder& GraphBuilder, ERHIAccess FinalAc
 
 			RHICmdList.Transition(Transitions);
 		});
+
+		bReadWriteAccess = false;
+		bReadWriteUAVOverlap = false;
 	}
 }
 
@@ -2026,7 +2300,7 @@ public:
 	BEGIN_SHADER_PARAMETER_STRUCT(FParameters, )
 		SHADER_PARAMETER_STRUCT_INCLUDE(FGPUScene::FInstanceGPULoadBalancer::FShaderParameters, BatcherParameters)
 		SHADER_PARAMETER_STRUCT_INCLUDE(FGPUSceneWriterParameters, GPUSceneWriterParameters)
-		END_SHADER_PARAMETER_STRUCT()
+	END_SHADER_PARAMETER_STRUCT()
 
 	static constexpr int32 NumThreadsPerGroup = FGPUScene::FInstanceGPULoadBalancer::ThreadGroupSize;
 
@@ -2082,7 +2356,7 @@ void FGPUSceneCompactInstanceData::Init(const FGPUScenePrimitiveCollector* Primi
 		DynamicPrimitiveId = PrimitiveCollector->GetPrimitiveIdRange().GetLowerBoundValue() + PrimitiveId;
 		if (PrimitiveCollector->GetPrimitiveIdRange().Contains(DynamicPrimitiveId))
 		{
-			const FPrimitiveUniformShaderParameters& PrimitiveData = PrimitiveCollector->UploadData->PrimitiveShaderData[PrimitiveId];
+			const FPrimitiveUniformShaderParameters& PrimitiveData = *PrimitiveCollector->UploadData->PrimitiveData[PrimitiveId].ShaderParams;
 			LocalToRelativeWorld = PrimitiveData.LocalToRelativeWorld;
 		}
 	}
diff --git a/Engine/Source/Runtime/Renderer/Private/GPUScene.h b/Engine/Source/Runtime/Renderer/Private/GPUScene.h
--- a/Engine/Source/Runtime/Renderer/Private/GPUScene.h
+++ b/Engine/Source/Runtime/Renderer/Private/GPUScene.h
@@ -8,6 +8,7 @@
 #include "PrimitiveSceneInfo.h"
 #include "GrowOnlySpanAllocator.h"
 #include "InstanceCulling/InstanceCullingLoadBalancer.h"
+#include "MeshBatch.h"
 
 class FRHICommandList;
 class FScene;
@@ -33,38 +34,16 @@ public:
 	{};
 
 	/**
-	 * Add data for a number of primitives, 
+	 * Add data for a primitive with a number of instances.
 	 * May be called outside (before) a FGPUScene::Begin/EndRender block.
-	 * @return The local index of the first primitive added - this needs to be offset after commiting to the GPU-Scene.
+	 * Note: needs to be virtual to prevent a linker error
 	 */
-	FORCEINLINE int32 Add(const FPrimitiveUniformShaderParameters* Data, int32 Num)
-	{
-		check(GPUSceneDynamicContext != nullptr);
-		check(!bCommitted);
-
-		// Lazy allocation of the upload data to not waste space and processing if none was needed.
-		if (UploadData == nullptr)
-		{
-			UploadData = AllocateUploadData();
-		}
-
-		const int32 LocalIndex = UploadData->PrimitiveShaderData.Num();
-		UploadData->PrimitiveShaderData.Append(Data, Num);
-		return LocalIndex;
-	}
-
-	/**
-	 * Translates the LocalIndex returned when adding a dynamic primitive into the final GPU-Scene index,
-	 * Only valid to call after Commit has been called.
-	 */
-	FORCEINLINE int32 GetFinalId(int32 LocalIndex) const
-	{
-		check(bCommitted);
-
-		check(PrimitiveIdRange.GetLowerBoundValue() >= 0);
-		check(PrimitiveIdRange.Size<int32>() >= LocalIndex);
-		return PrimitiveIdRange.GetLowerBoundValue() + LocalIndex;
-	}
+	virtual void Add(
+		const FMeshBatchDynamicPrimitiveData* MeshBatchData,
+		const FPrimitiveUniformShaderParameters& PrimitiveShaderParams,
+		uint32 NumInstances,
+		uint32& OutPrimitiveIndex,
+		uint32& OutInstanceSceneDataOffset);
 
 	/**
 	 * Allocates the range in GPUScene and queues the data for upload. 
@@ -82,41 +61,60 @@ public:
 		return PrimitiveIdRange;
 	}
 
-	FORCEINLINE int32 GetInstanceSceneDataOffset(int32 PrimitiveId) const
+	FORCEINLINE int32 GetInstanceSceneDataOffset() const
 	{
-		check(bCommitted);
-		check(PrimitiveIdRange.Contains(PrimitiveId));
-		check(UploadData->bIsUploaded);
+		check(bCommitted || UploadData == nullptr);
 
-		// Assume a 1:1 mapping between primitive ID and instance ID
-		return UploadData->InstanceSceneDataOffset + (PrimitiveId - PrimitiveIdRange.GetLowerBoundValue());
+		return UploadData != nullptr ? UploadData->InstanceSceneDataOffset : 0;
 	}
 
-	FORCEINLINE int32 GetInstanceSceneDataOffset() const
+	FORCEINLINE int32 GetInstancePayloadDataOffset() const
 	{
 		check(bCommitted || UploadData == nullptr);
 
-		// Assume a 1:1 mapping between primitive ID and instance ID
-		return UploadData != nullptr ? UploadData->InstanceSceneDataOffset : 0;
+		return UploadData != nullptr ? UploadData->InstancePayloadDataOffset : 0;
 	}
 
-	int32 Num() const {	return UploadData != nullptr ? UploadData->PrimitiveShaderData.Num() : 0; }
-	int32 NumInstances() const { return Num(); }
+	int32 Num() const {	return UploadData != nullptr ? UploadData->PrimitiveData.Num() : 0; }
+	int32 NumInstances() const { return UploadData != nullptr ? UploadData->TotalInstanceCount : 0; }
+	int32 NumPayloadDataSlots() const { return UploadData != nullptr ? UploadData->InstancePayloadDataFloat4Count : 0; }
+
+#if DO_CHECK
+	/**
+	 * Determines if the specified primitive has been sufficiently processed and its data can be read
+	 */
+	bool IsPrimitiveProcessed(uint32 PrimitiveIndex, const FGPUScene& GPUScene) const;
+#endif // DO_CHECK
+
 private:
 
 	friend class FGPUScene;
 	friend class FGPUSceneDynamicContext;
 	friend struct FGPUSceneCompactInstanceData;
+	friend struct FUploadDataSourceAdapterDynamicPrimitives;
+
+	struct FPrimitiveData
+	{
+		FMeshBatchDynamicPrimitiveData SourceData;
+		const FPrimitiveUniformShaderParameters* ShaderParams = nullptr;
+		uint32 NumInstances = 0;
+		uint32 LocalInstanceSceneDataOffset = INDEX_NONE;
+		uint32 LocalPayloadDataOffset = INDEX_NONE;
+	};
 
 	struct FUploadData
 	{
-		TArray<FPrimitiveUniformShaderParameters, TInlineAllocator<8>> PrimitiveShaderData;
-		int32 InstanceSceneDataOffset = INDEX_NONE;
+		TArray<FPrimitiveData, TInlineAllocator<8>> PrimitiveData;
+		TArray<uint32> GPUWritePrimitives;
+
+		uint32 InstanceSceneDataOffset = INDEX_NONE;
+		uint32 TotalInstanceCount = 0;
+		uint32 InstancePayloadDataOffset = INDEX_NONE;
+		uint32 InstancePayloadDataFloat4Count = 0;
 		bool bIsUploaded = false;
 	};
 
-	// Note: needs to be virtual to prevent a liker error
-	virtual FUploadData* AllocateUploadData();
+	FUploadData* AllocateUploadData();
 
 	/**
 	 * Range in GPU scene allocated to the dynamic primitives.
@@ -171,18 +169,6 @@ struct FGPUSceneBufferState
 	bool bResizedLightmapData = false;
 };
 
-
-// Include in shader that modifies GPU-Scene instances
-BEGIN_SHADER_PARAMETER_STRUCT(FGPUSceneWriterParameters, )
-	SHADER_PARAMETER_UAV(RWStructuredBuffer<float4>, GPUSceneInstanceSceneDataRW)
-	SHADER_PARAMETER_UAV(RWStructuredBuffer<float4>, GPUSceneInstancePayloadDataRW)
-	SHADER_PARAMETER_UAV(RWStructuredBuffer<float4>, GPUScenePrimitiveSceneDataRW)
-	SHADER_PARAMETER(uint32, GPUSceneInstanceSceneDataSOAStride)
-	SHADER_PARAMETER(uint32, GPUSceneFrameNumber)
-	SHADER_PARAMETER(uint32, GPUSceneNumAllocatedInstances)
-	SHADER_PARAMETER(uint32, GPUSceneNumAllocatedPrimitives)
-END_SHADER_PARAMETER_STRUCT()
-
 class FGPUScene
 {
 public:
@@ -287,7 +273,7 @@ public:
 	/**
 	 * Call before accessing the GPU scene in a read/write pass, returns false if read/write acces is not supported on the platform or the GPU scene is not enabled.
 	 */
-	bool BeginReadWriteAccess(FRDGBuilder& GraphBuilder);
+	bool BeginReadWriteAccess(FRDGBuilder& GraphBuilder, bool bAllowUAVOverlap = false);
 
 	/**
 	 * Fills in the FGPUSceneWriterParameters to use for read/write access to the GPU Scene.
@@ -318,6 +304,14 @@ public:
 	void BeginDeferAllocatorMerges();
 	void EndDeferAllocatorMerges();
 
+	/**
+	 * Executes GPUScene writes that were deferred until a later point in scene rendering
+	 **/
+	bool ExecuteDeferredGPUWritePass(FRDGBuilder& GraphBuilder, const TArrayView<FViewInfo>& Views, EGPUSceneGPUWritePass Pass);
+
+	/** Returns whether or not a GPU Write is pending for the specified primitive */
+	bool HasPendingGPUWrite(uint32 PrimitiveId) const;
+
 	bool bUpdateAllPrimitives;
 
 	/** Indices of primitives that need to be updated in GPU Scene */
@@ -344,18 +338,32 @@ public:
 	FRWBufferStructured		LightmapDataBuffer;
 	FScatterUploadBuffer	LightmapUploadBuffer;
 
-	using FInstanceGPULoadBalancer = TInstanceCullingLoadBalancer<SceneRenderingAllocator>;
-private:
-	TArray<EPrimitiveDirtyState> PrimitiveDirtyState;
-
 	struct FInstanceRange
 	{
 		uint32 InstanceSceneDataOffset;
 		uint32 NumInstanceSceneDataEntries;
 	};
 
+	TArray<FInstanceRange> DynamicPrimitiveInstancesToInvalidate;
+
+	using FInstanceGPULoadBalancer = TInstanceCullingLoadBalancer<SceneRenderingAllocator>;
+private:
+	TArray<EPrimitiveDirtyState> PrimitiveDirtyState;
+
 	TArray<FInstanceRange> InstanceRangesToClear;
 
+	struct FDeferredGPUWrite
+	{
+		FGPUSceneWriteDelegate DataWriterGPU;
+		int32 ViewId = INDEX_NONE;
+		uint32 PrimitiveId = INDEX_NONE;
+		uint32 InstanceSceneDataOffset = INDEX_NONE;
+	};
+
+	static constexpr uint32 NumDeferredGPUWritePasses = uint32(EGPUSceneGPUWritePass::Num);
+	TArray<FDeferredGPUWrite> DeferredGPUWritePassDelegates[NumDeferredGPUWritePasses];
+	EGPUSceneGPUWritePass LastDeferredGPUWritePass = EGPUSceneGPUWritePass::None;
+
 	TRange<int32> CommitPrimitiveCollector(FGPUScenePrimitiveCollector& PrimitiveCollector);
 
 
@@ -371,6 +379,8 @@ private:
 
 	bool bIsEnabled = false;
 	bool bInBeginEndBlock = false;
+	bool bReadWriteAccess = false;
+	bool bReadWriteUAVOverlap = false;
 	FGPUSceneDynamicContext* CurrentDynamicContext = nullptr;
 	int32 NumScenePrimitives = 0;
 
diff --git a/Engine/Source/Runtime/Renderer/Private/MeshDrawCommands.cpp b/Engine/Source/Runtime/Renderer/Private/MeshDrawCommands.cpp
--- a/Engine/Source/Runtime/Renderer/Private/MeshDrawCommands.cpp
+++ b/Engine/Source/Runtime/Renderer/Private/MeshDrawCommands.cpp
@@ -1337,6 +1337,20 @@ void FParallelMeshDrawCommandPass::BuildRenderingCommands(
 	{
 		check(!bHasInstanceCullingDrawParameters);
 		WaitForMeshPassSetupTask();
+
+#if DO_CHECK
+		for (const FVisibleMeshDrawCommand& VisibleMeshDrawCommand : TaskContext.MeshDrawCommands)
+		{
+			if (VisibleMeshDrawCommand.PrimitiveIdInfo.bIsDynamicPrimitive)
+			{
+				uint32 PrimitiveIndex = VisibleMeshDrawCommand.PrimitiveIdInfo.DrawPrimitiveId & ~GPrimIDDynamicFlag;
+				checkf(TaskContext.View->DynamicPrimitiveCollector.IsPrimitiveProcessed(PrimitiveIndex, GPUScene),
+					TEXT("Dynamic Primitive index %u has not been fully processed. It may be an invalid index for the collector or it has a pending GPU write."),
+					PrimitiveIndex);
+			}
+		}
+#endif
+
 		// 2. Run or queue finalize culling commands pass
 		TaskContext.InstanceCullingContext.BuildRenderingCommands(GraphBuilder, GPUScene, TaskContext.View->DynamicPrimitiveCollector.GetInstanceSceneDataOffset(), TaskContext.View->DynamicPrimitiveCollector.NumInstances(), TaskContext.InstanceCullingResult, &OutInstanceCullingDrawParams);
 		TaskContext.InstanceCullingResult.GetDrawParameters(OutInstanceCullingDrawParams);
diff --git a/Engine/Source/Runtime/Renderer/Private/MeshPassProcessor.cpp b/Engine/Source/Runtime/Renderer/Private/MeshPassProcessor.cpp
--- a/Engine/Source/Runtime/Renderer/Private/MeshPassProcessor.cpp
+++ b/Engine/Source/Runtime/Renderer/Private/MeshPassProcessor.cpp
@@ -1478,8 +1478,8 @@ FMeshDrawCommandPrimitiveIdInfo FMeshPassProcessor::GetDrawCommandPrimitiveId(
 		else if (BatchElement.PrimitiveIdMode == PrimID_DynamicPrimitiveShaderData && ViewIfDynamicMeshCommand != nullptr)
 		{
 			// Mark using GPrimIDDynamicFlag (top bit) as we defer this to later.
-			PrimitiveIdInfo.DrawPrimitiveId = BatchElement.DynamicPrimitiveShaderDataIndex | GPrimIDDynamicFlag;
-			PrimitiveIdInfo.InstanceSceneDataOffset = BatchElement.DynamicPrimitiveShaderDataIndex;
+			PrimitiveIdInfo.DrawPrimitiveId = BatchElement.DynamicPrimitiveIndex | GPrimIDDynamicFlag;
+			PrimitiveIdInfo.InstanceSceneDataOffset = BatchElement.DynamicPrimitiveInstanceSceneDataOffset;
 			PrimitiveIdInfo.bIsDynamicPrimitive = 1U;
 		}
 		else
diff --git a/Engine/Source/Runtime/Renderer/Private/VirtualShadowMaps/VirtualShadowMapCacheManager.h b/Engine/Source/Runtime/Renderer/Private/VirtualShadowMaps/VirtualShadowMapCacheManager.h
--- a/Engine/Source/Runtime/Renderer/Private/VirtualShadowMaps/VirtualShadowMapCacheManager.h
+++ b/Engine/Source/Runtime/Renderer/Private/VirtualShadowMaps/VirtualShadowMapCacheManager.h
@@ -125,10 +125,21 @@ public:
 	class FInvalidatingPrimitiveCollector
 	{
 	public:
-		FInvalidatingPrimitiveCollector(int32 MaxPrimitiveID, const FGPUScene& InGPUScene)
+		FInvalidatingPrimitiveCollector(int32 MaxPrimitiveID, FGPUScene& InGPUScene)
 			: AlreadyAddedPrimitives(false, MaxPrimitiveID)
 			, GPUScene(InGPUScene)
 		{
+			// Add and clear pending invalidations enqueued on the GPU Scene from dynamic primitives added since last invalidation
+			for (const FGPUScene::FInstanceRange& Range : GPUScene.DynamicPrimitiveInstancesToInvalidate)
+			{
+				LoadBalancer.Add(Range.InstanceSceneDataOffset, Range.NumInstanceSceneDataEntries, 0U);
+#if VSM_LOG_INVALIDATIONS
+				RangesStr.Appendf(TEXT("[%6d, %6d), "), Range.InstanceSceneDataOffset, Range.InstanceSceneDataOffset + Range.NumInstanceSceneDataEntries);
+#endif
+				TotalInstanceCount += Range.NumInstanceSceneDataEntries;
+			}
+
+			GPUScene.DynamicPrimitiveInstancesToInvalidate.Reset();
 		}
 
 		/**
@@ -163,8 +174,9 @@ public:
 #if VSM_LOG_INVALIDATIONS
 		FString RangesStr;
 #endif
-		const FGPUScene& GPUScene;
+		FGPUScene& GPUScene;
 	};
+
 	/**
 	 * This must to be executed before the instances are actually removed / updated, otherwise the wrong position will be used. 
 	 * In particular, it must be processed before the Scene primitive IDs are updated/compacted as part of the removal.
