Removing duplicate definition due to bad automerge
#rnx

diff --git a/Engine/Source/Runtime/Renderer/Private/GPUScene.cpp b/Engine/Source/Runtime/Renderer/Private/GPUScene.cpp
--- a/Engine/Source/Runtime/Renderer/Private/GPUScene.cpp
+++ b/Engine/Source/Runtime/Renderer/Private/GPUScene.cpp
@@ -118,14 +118,6 @@ FAutoConsoleVariableRef CVarGPUSceneInstanceUploadViaCreate(
 
 constexpr uint32 InstanceSceneDataNumArrays = FInstanceSceneShaderData::DataStrideInFloat4s;
 
-int32 GGPUSceneInstanceUploadViaCreate = 1;
-FAutoConsoleVariableRef CVarGPUSceneInstanceUploadViaCreate(
-	TEXT("r.GPUScene.InstanceUploadViaCreate"),
-	GGPUSceneInstanceUploadViaCreate,
-	TEXT("When uploading GPUScene InstanceData, upload via resource creation when the RHI supports it efficiently."),
-	ECVF_RenderThreadSafe
-);
-
 LLM_DECLARE_TAG_API(GPUScene, RENDERER_API);
 DECLARE_LLM_MEMORY_STAT(TEXT("GPUScene"), STAT_GPUSceneLLM, STATGROUP_LLMFULL);
 DECLARE_LLM_MEMORY_STAT(TEXT("GPUScene"), STAT_GPUSceneSummaryLLM, STATGROUP_LLM);
