Removing duplicate definition due to bad automerge
#rnx

diff --git a/Engine/Source/Runtime/Renderer/Private/SceneVisibility.cpp b/Engine/Source/Runtime/Renderer/Private/SceneVisibility.cpp
--- a/Engine/Source/Runtime/Renderer/Private/SceneVisibility.cpp
+++ b/Engine/Source/Runtime/Renderer/Private/SceneVisibility.cpp
@@ -247,14 +247,6 @@ static TAutoConsoleVariable<bool> CVarEnableFrustumCull(
 	TEXT("Enables or disables frustum culling.  Useful for comparing results to ensure culling is functioning properly."),
 	ECVF_RenderThreadSafe);
 
-static bool GOcclusionSingleRHIThreadStall = false;
-static FAutoConsoleVariableRef CVarOcclusionSingleRHIThreadStall(
-	TEXT("r.Occlusion.SingleRHIThreadStall"),
-	GOcclusionSingleRHIThreadStall,
-	TEXT("Enable a single RHI thread stall before polling occlusion queries. This will only happen if the RHI's occlusion queries would normally stall the RHI thread themselves."),
-	ECVF_RenderThreadSafe
-);
-
 static TAutoConsoleVariable<int32> CVarAlsoUseSphereForFrustumCull(
 	TEXT("r.AlsoUseSphereForFrustumCull"),
 	0,
