Removing duplicated code due to bad merge
#rnx

diff --git a/Engine/Source/Runtime/Renderer/Private/GPUScene.cpp b/Engine/Source/Runtime/Renderer/Private/GPUScene.cpp
--- a/Engine/Source/Runtime/Renderer/Private/GPUScene.cpp
+++ b/Engine/Source/Runtime/Renderer/Private/GPUScene.cpp
@@ -1155,9 +1155,6 @@ void FGPUScene::UploadGeneral(FRHICommandListImmediate& RHICmdList, FScene *Scen
 						}
 					}
 				}
-				const bool bShouldUploadViaCreate = GRHISupportsEfficientUploadOnResourceCreation && GGPUSceneInstanceUploadViaCreate != 0;
-				PrimitiveUploadBuffer.SetUploadViaCreate(bShouldUploadViaCreate);
-				InstanceSceneUploadBuffer.SetUploadViaCreate(bShouldUploadViaCreate);
 				
 				{
 					SCOPED_NAMED_EVENT(STAT_UpdateGPUScenePrimitives_Par, FColor::Green);
