Fix motion blur shader parameters bad merge.

diff --git a/Engine/Source/Runtime/Renderer/Private/PostProcess/PostProcessMotionBlur.cpp b/Engine/Source/Runtime/Renderer/Private/PostProcess/PostProcessMotionBlur.cpp
--- a/Engine/Source/Runtime/Renderer/Private/PostProcess/PostProcessMotionBlur.cpp
+++ b/Engine/Source/Runtime/Renderer/Private/PostProcess/PostProcessMotionBlur.cpp
@@ -308,8 +308,6 @@ public:
 		SHADER_PARAMETER_STRUCT(FScreenPassTextureViewportParameters, Velocity)
 		SHADER_PARAMETER_STRUCT_INCLUDE(FVelocityFlattenParameters, VelocityFlattenParameters)
 		SHADER_PARAMETER(FMatrix44f, ClipToPrevClipOverride)
-		SHADER_PARAMETER(FVector2f, VelocityScale)
-		SHADER_PARAMETER(float, VelocityMax)
 
 		SHADER_PARAMETER_RDG_TEXTURE(Texture2D, VelocityTexture)
 		SHADER_PARAMETER_RDG_TEXTURE(Texture2D, DepthTexture)
