Fix for shipping server
#rnx

diff --git a/Engine/Source/Runtime/Renderer/Private/VT/VirtualTexturePhysicalSpace.cpp b/Engine/Source/Runtime/Renderer/Private/VT/VirtualTexturePhysicalSpace.cpp
--- a/Engine/Source/Runtime/Renderer/Private/VT/VirtualTexturePhysicalSpace.cpp
+++ b/Engine/Source/Runtime/Renderer/Private/VT/VirtualTexturePhysicalSpace.cpp
@@ -307,7 +307,7 @@ void FVirtualTexturePhysicalSpace::DrawResidencyGraph(FCanvas* Canvas, FBox2D Ca
 
 void FVirtualTexturePhysicalSpace::UpdateCsvStats() const
 {
-#if CSV_PROFILER
+#if CSV_PROFILER && !UE_BUILD_SHIPPING
 	FCsvProfiler* Profiler = FCsvProfiler::Get();
 	if (Profiler->IsCapturing_Renderthread())
 	{
