* Make function for checking hair strands/cards VF common
* Clean up hair public data access

diff --git a/Engine/Plugins/Runtime/HairStrands/Source/HairStrandsCore/Public/GroomInstance.h b/Engine/Plugins/Runtime/HairStrands/Source/HairStrandsCore/Public/GroomInstance.h
--- a/Engine/Plugins/Runtime/HairStrands/Source/HairStrandsCore/Public/GroomInstance.h
+++ b/Engine/Plugins/Runtime/HairStrands/Source/HairStrandsCore/Public/GroomInstance.h
@@ -283,6 +283,9 @@ struct HAIRSTRANDSCORE_API FHairGroupInstance : public FHairStrandsInstance
 		return Meshes.IsValid() || Cards.IsValid() || Strands.IsValid();
 	}
 
+	virtual const FBoxSphereBounds* GetBounds() override { return ProxyBounds;  }
+	virtual const FHairGroupPublicData* GetHairData() { return HairGroupPublicData; }
+
 	/** Get the current local to world transform according to the internal binding type */
 	FORCEINLINE const FTransform& GetCurrentLocalToWorld() const
 	{
diff --git a/Engine/Source/Runtime/Renderer/Private/HairStrands/HairStrandsCluster.cpp b/Engine/Source/Runtime/Renderer/Private/HairStrands/HairStrandsCluster.cpp
--- a/Engine/Source/Runtime/Renderer/Private/HairStrands/HairStrandsCluster.cpp
+++ b/Engine/Source/Runtime/Renderer/Private/HairStrands/HairStrandsCluster.cpp
@@ -3,6 +3,7 @@
 
 #include "HairStrandsCluster.h"
 #include "HairStrandsUtils.h"
+#include "HairStrandsData.h"
 #include "SceneRendering.h"
 #include "SceneManagement.h"
 #include "RendererInterface.h"
@@ -10,6 +11,7 @@
 #include "GlobalShader.h"
 #include "ShaderParameters.h"
 #include "ShaderParameterStruct.h"
+#include "ScenePrivate.h"
 
 class FHairMacroGroupAABBCS : public FGlobalShader
 {
@@ -126,36 +128,26 @@ static bool DoesGroupExists(uint32 ResourceId, uint32 GroupIndex, const FHairStr
 	return false;
 }
 
-inline const FHairGroupPublicData* GetHairStrandsPublicData(const FMeshBatch* InMeshBatch)
+static void InternalUpdateMacroGroup(FHairStrandsMacroGroupData& MacroGroup, int32& MaterialId, FHairGroupPublicData* HairData, const FMeshBatch* Mesh, const FPrimitiveSceneProxy* Proxy)
 {
-	 return reinterpret_cast<const FHairGroupPublicData*>(InMeshBatch->Elements[0].VertexFactoryUserData);
-}
+	check(HairData);
 
-static void InternalUpdateMacroGroup(FHairStrandsMacroGroupData& MacroGroup, int32& MaterialId, const FMeshBatchAndRelevance* MeshBatchAndRelevance)
-{
-	if (MeshBatchAndRelevance)
+	if (HairData->VFInput.Strands.bScatterSceneLighting)
 	{
-		check(MeshBatchAndRelevance->Mesh);
-		check(MeshBatchAndRelevance->Mesh->Elements.Num() == 1);
-
-		FHairGroupPublicData* HairGroupPublicData = reinterpret_cast<FHairGroupPublicData*>(MeshBatchAndRelevance->Mesh->Elements[0].VertexFactoryUserData);
-		if (HairGroupPublicData->VFInput.Strands.bScatterSceneLighting)
-		{
-			MacroGroup.bNeedScatterSceneLighting = true;
-		}
+		MacroGroup.bNeedScatterSceneLighting = true;
+	}
 
-		FHairStrandsMacroGroupData::PrimitiveInfo& PrimitiveInfo = MacroGroup.PrimitivesInfos.AddZeroed_GetRef();
-		PrimitiveInfo.Mesh = MeshBatchAndRelevance->Mesh;
-		PrimitiveInfo.PrimitiveSceneProxy = MeshBatchAndRelevance->PrimitiveSceneProxy;
-		PrimitiveInfo.MaterialId = MaterialId++;
-		PrimitiveInfo.ResourceId = reinterpret_cast<uint64>(MeshBatchAndRelevance->Mesh->Elements[0].UserData);
-		PrimitiveInfo.GroupIndex = HairGroupPublicData->GetGroupIndex();
-		PrimitiveInfo.PublicDataPtr = HairGroupPublicData;
+	FHairStrandsMacroGroupData::PrimitiveInfo& PrimitiveInfo = MacroGroup.PrimitivesInfos.AddZeroed_GetRef();
+	PrimitiveInfo.Mesh = Mesh;
+	PrimitiveInfo.PrimitiveSceneProxy = Proxy;
+	PrimitiveInfo.MaterialId = MaterialId++;
+	PrimitiveInfo.ResourceId = Mesh ? reinterpret_cast<uint64>(Mesh->Elements[0].UserData) : ~0u;
+	PrimitiveInfo.GroupIndex = HairData->GetGroupIndex();
+	PrimitiveInfo.PublicDataPtr = HairData;
 
-		if (HairGroupPublicData->DoesSupportVoxelization())
-		{
-			MacroGroup.bSupportVoxelization = true;
-		}
+	if (HairData->DoesSupportVoxelization())
+	{
+		MacroGroup.bSupportVoxelization = true;
 	}
 }
 
@@ -165,32 +157,28 @@ void CreateHairStrandsMacroGroups(
 	const FViewInfo& View, 
 	FHairStrandsViewData& OutHairStrandsViewData)
 {
-	if (!View.Family || View.HairStrandsMeshElements.Num() == 0 || View.bIsPlanarReflection || View.bIsReflectionCapture)
+	const bool bHasHairStrandsElements = View.HairStrandsMeshElements.Num() != 0 || Scene->HairStrandsSceneData.RegisteredProxies.Num() != 0;
+	if (!View.Family || !bHasHairStrandsElements || View.bIsPlanarReflection || View.bIsReflectionCapture)
 	{
 		return;
 	}
 
-	static const FVertexFactoryType* CompatibleVF = FVertexFactoryType::GetVFByName(TEXT("FHairStrandsVertexFactory"));
-
 	TArray<FHairStrandsMacroGroupData, SceneRenderingAllocator>& MacroGroups = OutHairStrandsViewData.MacroGroupDatas;
 
 	int32 MaterialId = 0;
 
 	// Aggregate all hair primitives within the same area into macro groups, for allocating/rendering DOM/voxel
 	uint32 MacroGroupId = 0;
-	auto UpdateMacroGroup = [&MacroGroups, &MacroGroupId, &MaterialId](const FMeshBatchAndRelevance* MeshBatchAndRelevance, const FPrimitiveSceneProxy* Proxy)
+	auto UpdateMacroGroup = [&MacroGroups, &MacroGroupId, &MaterialId](FHairGroupPublicData* HairData, const FMeshBatch* Mesh,  const FPrimitiveSceneProxy* Proxy, const FBoxSphereBounds* Bounds)
 	{
-		const bool bIsHairStrandsFactory = MeshBatchAndRelevance->Mesh->VertexFactory->GetType()->GetHashedName() == CompatibleVF->GetHashedName();
-		if (!bIsHairStrandsFactory)
-			return;
+		check(HairData);
 
 		// Ensure that the element has been initialized
-		const FHairGroupPublicData* HairGroupPublicData = reinterpret_cast<const FHairGroupPublicData*>(MeshBatchAndRelevance->Mesh->Elements[0].VertexFactoryUserData);
-		const bool bIsValid = HairGroupPublicData->VFInput.Strands.PositionBufferRHISRV != nullptr;
+		const bool bIsValid = HairData->VFInput.Strands.PositionBufferRHISRV != nullptr;
 		if (!bIsValid)
 			return;
 
-		const FBoxSphereBounds& PrimitiveBounds = Proxy->GetBounds();
+		const FBoxSphereBounds& PrimitiveBounds = Proxy ? Proxy->GetBounds() : *Bounds;
 
 		bool bFound = false;
 		float MinDistance = FLT_MAX;
@@ -207,7 +195,8 @@ void CreateHairStrandsMacroGroups(
 			if (bIntersect)
 			{
 				MacroGroup.Bounds = Union(MacroGroup.Bounds, PrimitiveBounds);
-				InternalUpdateMacroGroup(MacroGroup, MaterialId, MeshBatchAndRelevance);
+
+				InternalUpdateMacroGroup(MacroGroup, MaterialId, HairData, Mesh, Proxy);
 				bFound = true;
 				break;
 			}
@@ -229,13 +218,13 @@ void CreateHairStrandsMacroGroups(
 				FHairStrandsMacroGroupData& MacroGroup = MacroGroups[ClosestMacroGroupId];
 				check(MacroGroup.MacroGroupId == ClosestMacroGroupId);
 				MacroGroup.Bounds = Union(MacroGroup.Bounds, PrimitiveBounds);
-				InternalUpdateMacroGroup(MacroGroup, MaterialId, MeshBatchAndRelevance);
+				InternalUpdateMacroGroup(MacroGroup, MaterialId, HairData, Mesh, Proxy);
 			}
 			else
 			{
 				FHairStrandsMacroGroupData MacroGroup;
 				MacroGroup.MacroGroupId = MacroGroupId++;
-				InternalUpdateMacroGroup(MacroGroup, MaterialId, MeshBatchAndRelevance);
+				InternalUpdateMacroGroup(MacroGroup, MaterialId, HairData, Mesh, Proxy);
 				MacroGroup.Bounds = PrimitiveBounds;
 				MacroGroups.Add(MacroGroup);
 			}
@@ -244,7 +233,13 @@ void CreateHairStrandsMacroGroups(
 
 	for (const FMeshBatchAndRelevance& MeshBatchAndRelevance : View.HairStrandsMeshElements)
 	{
-		UpdateMacroGroup(&MeshBatchAndRelevance, MeshBatchAndRelevance.PrimitiveSceneProxy);
+		if (HairStrands::IsHairStrandsVF(MeshBatchAndRelevance.Mesh))
+		{
+			if (FHairGroupPublicData* HairData = HairStrands::GetHairData(MeshBatchAndRelevance.Mesh))
+			{
+				UpdateMacroGroup(HairData, MeshBatchAndRelevance.Mesh, MeshBatchAndRelevance.PrimitiveSceneProxy, nullptr);
+			}
+		}
 	}
 
 	for (FHairStrandsMacroGroupData& MacroGroup : MacroGroups)
@@ -275,6 +270,6 @@ void CreateHairStrandsMacroGroups(
 
 bool FHairStrandsMacroGroupData::PrimitiveInfo::IsCullingEnable() const
 {
-	const FHairGroupPublicData* HairGroupPublicData = GetHairStrandsPublicData(Mesh);
-	return HairGroupPublicData->GetCullingResultAvailable();
-}
\ No newline at end of file
+	const FHairGroupPublicData* HairData = HairStrands::GetHairData(Mesh);
+	return HairData->GetCullingResultAvailable();
+}
diff --git a/Engine/Source/Runtime/Renderer/Private/HairStrands/HairStrandsData.h b/Engine/Source/Runtime/Renderer/Private/HairStrands/HairStrandsData.h
--- a/Engine/Source/Runtime/Renderer/Private/HairStrands/HairStrandsData.h
+++ b/Engine/Source/Runtime/Renderer/Private/HairStrands/HairStrandsData.h
@@ -430,4 +430,10 @@ namespace HairStrands
 	bool HasPositionsChanged(FRDGBuilder& GraphBuilder, const FViewInfo& View);
 	void DrawHitProxies(FRDGBuilder& GraphBuilder, const FScene& Scene, const FViewInfo& View, FInstanceCullingManager& InstanceCullingManager, FRDGTextureRef HitProxyTexture, FRDGTextureRef HitProxyDepthTexture);
 	void DrawEditorSelection(FRDGBuilder& GraphBuilder, const FViewInfo& View, FRDGTextureRef SelectionDepthTexture);
+
+	// Mesh batch helpers
+	FHairGroupPublicData* GetHairData(const FMeshBatch* In);
+	bool IsHairStrandsCompatible(const FMeshBatch* Mesh);
+	bool IsHairStrandsVF(const FMeshBatch* Mesh);
+	bool IsHairCardsVF(const FMeshBatch* Mesh);
 }
\ No newline at end of file
diff --git a/Engine/Source/Runtime/Renderer/Private/HairStrands/HairStrandsDebug.cpp b/Engine/Source/Runtime/Renderer/Private/HairStrands/HairStrandsDebug.cpp
--- a/Engine/Source/Runtime/Renderer/Private/HairStrands/HairStrandsDebug.cpp
+++ b/Engine/Source/Runtime/Renderer/Private/HairStrands/HairStrandsDebug.cpp
@@ -11,6 +11,7 @@
 #include "HairStrandsInterface.h"
 #include "HairStrandsMeshProjection.h"
 #include "HairStrandsTile.h"
+#include "HairStrandsData.h"
 
 #include "Shader.h"
 #include "GlobalShader.h"
@@ -1029,14 +1030,13 @@ static void AddDrawDebugClusterPass(
 				for (const FHairStrandsMacroGroupData::PrimitiveInfo& PrimitiveInfo : MacroGroupData.PrimitivesInfos)
 				{
 					check(PrimitiveInfo.Mesh && PrimitiveInfo.Mesh->Elements.Num() > 0);
-					const FHairGroupPublicData* HairGroupPublicData = reinterpret_cast<const FHairGroupPublicData*>(PrimitiveInfo.Mesh->Elements[0].VertexFactoryUserData);
 
 					for (int DataIndex = 0; DataIndex < HairClusterData.HairGroups.Num(); ++DataIndex)
 					{
 						const FHairStrandClusterData::FHairGroup& HairGroupClusters = HairClusterData.HairGroups[DataIndex];
 
 						// Find a better/less hacky way
-						if (HairGroupPublicData != HairGroupClusters.HairGroupPublicPtr)
+						if (PrimitiveInfo.PublicDataPtr != HairGroupClusters.HairGroupPublicPtr)
 							continue;
 
 						if (ShaderDrawDebug::IsEnabled(View) && HairGroupClusters.CulledClusterCountBuffer)
@@ -1119,7 +1119,7 @@ static void InternalRenderHairStrandsDebugInfo(
 		{
 			for (const FMeshBatchAndRelevance& Mesh : View.HairStrandsMeshElements)
 			{
-				const FHairGroupPublicData* GroupData = reinterpret_cast<const FHairGroupPublicData*>(Mesh.Mesh->Elements[0].VertexFactoryUserData);
+				const FHairGroupPublicData* GroupData = HairStrands::GetHairData(Mesh.Mesh);
 				if (GroupData->DebugMode == EHairStrandsDebugMode::RenderHairTangent)
 				{
 					bTangentEnabled = true;
diff --git a/Engine/Source/Runtime/Renderer/Private/HairStrands/HairStrandsInterface.cpp b/Engine/Source/Runtime/Renderer/Private/HairStrands/HairStrandsInterface.cpp
--- a/Engine/Source/Runtime/Renderer/Private/HairStrands/HairStrandsInterface.cpp
+++ b/Engine/Source/Runtime/Renderer/Private/HairStrands/HairStrandsInterface.cpp
@@ -7,6 +7,7 @@
 #include "HairStrandsInterface.h"
 #include "HairStrandsRendering.h"
 #include "HairStrandsMeshProjection.h"
+#include "HairStrandsData.h"
 
 #include "GPUSkinCache.h"
 #include "Rendering/SkeletalMeshRenderData.h"
@@ -430,7 +431,7 @@ FHairStrandsBookmarkParameters CreateHairStrandsBookmarkParameters(FScene* Scene
 		check(MeshBatch.PrimitiveSceneProxy && MeshBatch.PrimitiveSceneProxy->ShouldRenderInMainPass());
 		if (MeshBatch.Mesh && MeshBatch.Mesh->Elements.Num() > 0)
 		{
-			FHairGroupPublicData* HairGroupPublicData = reinterpret_cast<FHairGroupPublicData*>(MeshBatch.Mesh->Elements[0].VertexFactoryUserData);
+			FHairGroupPublicData* HairGroupPublicData = HairStrands::GetHairData(MeshBatch.Mesh);
 			if (HairGroupPublicData && HairGroupPublicData->Instance)
 			{
 				Out.VisibleInstances.Add(HairGroupPublicData->Instance);
@@ -473,23 +474,11 @@ FHairStrandsBookmarkParameters CreateHairStrandsBookmarkParameters(FScene* Scene
 	return Out;
 }
 
-bool IsHairStrandsCompatible(const FMeshBatch* Mesh)
-{
-	if (Mesh)
-	{
-		static const FHashedName& VFType0 = FVertexFactoryType::GetVFByName(TEXT("FHairStrandsVertexFactory"))->GetHashedName();
-		static const FHashedName& VFType1 = FVertexFactoryType::GetVFByName(TEXT("FHairCardsVertexFactory"))->GetHashedName();
-		const FHashedName& VFType = Mesh->VertexFactory->GetType()->GetHashedName();
-		return VFType == VFType0 || VFType == VFType1;
-	}
-	return false;
-}
-
 bool IsHairStrandsVisible(const FMeshBatchAndRelevance& MeshBatch)
 {
 	if (MeshBatch.Mesh && MeshBatch.PrimitiveSceneProxy && MeshBatch.PrimitiveSceneProxy->ShouldRenderInMainPass())
 	{
-		const FHairGroupPublicData* Data = reinterpret_cast<const FHairGroupPublicData*>(MeshBatch.Mesh->Elements[0].VertexFactoryUserData);
+		const FHairGroupPublicData* Data = HairStrands::GetHairData(MeshBatch.Mesh);
 		switch (Data->VFInput.GeometryType)
 		{
 			case EHairGeometryType::Strands: return Data->VFInput.Strands.HairLengthScale > 0;
@@ -499,3 +488,40 @@ bool IsHairStrandsVisible(const FMeshBatchAndRelevance& MeshBatch)
 	}
 	return false;
 }
+
+namespace HairStrands
+{
+
+bool IsHairStrandsVF(const FMeshBatch* Mesh)
+{
+	if (Mesh)
+	{
+		static const FHashedName& VFTypeRef = FVertexFactoryType::GetVFByName(TEXT("FHairStrandsVertexFactory"))->GetHashedName();
+		const FHashedName& VFType = Mesh->VertexFactory->GetType()->GetHashedName();
+		return VFType == VFTypeRef;
+	}
+	return false;
+}
+
+bool IsHairCardsVF(const FMeshBatch* Mesh)
+{
+	if (Mesh)
+	{
+		static const FHashedName& VFTypeRef = FVertexFactoryType::GetVFByName(TEXT("FHairCardsVertexFactory"))->GetHashedName();
+		const FHashedName& VFType = Mesh->VertexFactory->GetType()->GetHashedName();
+		return VFType == VFTypeRef;
+	}
+	return false;
+}
+
+bool IsHairStrandsCompatible(const FMeshBatch* Mesh)
+{
+	return IsHairStrandsVF(Mesh) || IsHairCardsVF(Mesh);
+}
+
+FHairGroupPublicData* GetHairData(const FMeshBatch* Mesh)
+{
+	return reinterpret_cast<FHairGroupPublicData*>(Mesh->Elements[0].VertexFactoryUserData);
+}
+
+} // namespace HairStrands
diff --git a/Engine/Source/Runtime/Renderer/Private/HairStrands/HairStrandsVisibility.cpp b/Engine/Source/Runtime/Renderer/Private/HairStrands/HairStrandsVisibility.cpp
--- a/Engine/Source/Runtime/Renderer/Private/HairStrands/HairStrandsVisibility.cpp
+++ b/Engine/Source/Runtime/Renderer/Private/HairStrands/HairStrandsVisibility.cpp
@@ -2857,7 +2857,7 @@ void AddMeshDrawTransitionPass(
 	{
 		for (const FHairStrandsMacroGroupData::PrimitiveInfo& PrimitiveInfo : MacroGroup.PrimitivesInfos)
 		{
-			FHairGroupPublicData* HairGroupPublicData = reinterpret_cast<FHairGroupPublicData*>(PrimitiveInfo.Mesh->Elements[0].VertexFactoryUserData);
+			FHairGroupPublicData* HairGroupPublicData = PrimitiveInfo.PublicDataPtr;
 			check(HairGroupPublicData);
 
 			FRDGResourceAccessFinalizer ResourceAccessFinalizer;
@@ -3099,9 +3099,8 @@ static FRasterComputeOutput AddVisibilityComputeRasterPass(
 			PassParameters->OutVisibilityTexture2 = VisibilityTexture2UAV;
 			PassParameters->OutVisibilityTexture3 = VisibilityTexture3UAV;
 			
-			check(PrimitiveInfo.Mesh && PrimitiveInfo.Mesh->Elements.Num() > 0);
-			const FHairGroupPublicData* HairGroupPublicData = reinterpret_cast<const FHairGroupPublicData*>(PrimitiveInfo.Mesh->Elements[0].VertexFactoryUserData);
-			check(HairGroupPublicData);
+			check(PrimitiveInfo.PublicDataPtr);
+			const FHairGroupPublicData* HairGroupPublicData = PrimitiveInfo.PublicDataPtr;
 
 			const FHairGroupPublicData::FVertexFactoryInput& VFInput = HairGroupPublicData->VFInput;
 			PassParameters->HairStrandsVF_PositionBuffer	= VFInput.Strands.PositionBuffer.SRV;
@@ -3176,6 +3175,11 @@ void AddHairStrandsSelectionOutlinePass(
 	FRDGBufferRef VisNodeData,
 	FRDGTextureRef SelectionDepthTexture)
 {
+	if (View.HairStrandsMeshElements.Num() == 0)
+	{
+		return;
+	}
+
 #if WITH_EDITOR
 	// Create mapping table between PrimitiveId and BatchId
 	TArray<uint32> SelectionMaterialId;
@@ -3251,6 +3255,11 @@ void AddHairStrandsHitProxyIdPass(
 	FRDGTextureRef HitProxyDepthTexture)
 {
 #if WITH_EDITOR
+	if (View.HairStrandsMeshElements.Num() == 0)
+	{
+		return;
+	}
+
 	// Create mapping table between PrimitiveId and BatchId
 	TArray<uint32> MaterialIdToHitProxyId;
 	MaterialIdToHitProxyId.Reserve(View.HairStrandsMeshElements.Num());
@@ -3478,7 +3487,7 @@ bool HasPositionsChanged(FRDGBuilder& GraphBuilder, const FViewInfo& View)
 	TArray<const FHairGroupPublicData*> GroupDatas;
 	for (const FMeshBatchAndRelevance& Batch : View.HairStrandsMeshElements)
 	{
-		FHairGroupPublicData* HairGroupPublicData = reinterpret_cast<FHairGroupPublicData*>(Batch.Mesh->Elements[0].VertexFactoryUserData);
+		FHairGroupPublicData* HairGroupPublicData = GetHairData(Batch.Mesh);
 		check(HairGroupPublicData);
 		const int32 LODIndex = FMath::FloorToInt(HairGroupPublicData->LODIndex);
 		const bool bHasSimulationOrSkinning = 
diff --git a/Engine/Source/Runtime/Renderer/Private/HairStrands/HairStrandsVoxelization.cpp b/Engine/Source/Runtime/Renderer/Private/HairStrands/HairStrandsVoxelization.cpp
--- a/Engine/Source/Runtime/Renderer/Private/HairStrands/HairStrandsVoxelization.cpp
+++ b/Engine/Source/Runtime/Renderer/Private/HairStrands/HairStrandsVoxelization.cpp
@@ -1291,8 +1291,8 @@ static void AddVirtualVoxelizationComputeRasterPass(
 		TShaderMapRef<FVoxelRasterComputeCS> ComputeShader_CullingOn(ViewInfo->ShaderMap, PermutationVector_On);
 		for (const FHairStrandsMacroGroupData::PrimitiveInfo& PrimitiveInfo : PrimitiveSceneInfos)
 		{
-			check(PrimitiveInfo.Mesh && PrimitiveInfo.Mesh->Elements.Num() > 0);
-			const FHairGroupPublicData* HairGroupPublicData = reinterpret_cast<const FHairGroupPublicData*>(PrimitiveInfo.Mesh->Elements[0].VertexFactoryUserData);
+			check(PrimitiveInfo.PublicDataPtr);
+			const FHairGroupPublicData* HairGroupPublicData = PrimitiveInfo.PublicDataPtr;
 
 			if (HairGroupPublicData->DoesSupportVoxelization())
 			{
diff --git a/Engine/Source/Runtime/Renderer/Private/SceneVisibility.cpp b/Engine/Source/Runtime/Renderer/Private/SceneVisibility.cpp
--- a/Engine/Source/Runtime/Renderer/Private/SceneVisibility.cpp
+++ b/Engine/Source/Runtime/Renderer/Private/SceneVisibility.cpp
@@ -37,6 +37,7 @@
 #include "TranslucentRendering.h"
 #include "Async/ParallelFor.h"
 #include "HairStrands/HairStrandsRendering.h"
+#include "HairStrands/HairStrandsData.h"
 #include "RectLightSceneProxy.h"
 #include "Math/Halton.h"
 #include "ProfilingDebugging/DiagnosticTable.h"
@@ -3290,7 +3291,7 @@ void ComputeDynamicMeshRelevance(EShadingPath ShadingPath, bool bAddLightmapDens
 		BatchAndProxy.SortKey = MeshBatch.PrimitiveSceneProxy->GetTranslucencySortPriority();
 	}
 	
-	const bool bIsHairStrandsCompatible = ViewRelevance.bHairStrands && IsHairStrandsEnabled(EHairStrandsShaderType::All, View.GetShaderPlatform()) && IsHairStrandsCompatible(MeshBatch.Mesh);
+	const bool bIsHairStrandsCompatible = ViewRelevance.bHairStrands && IsHairStrandsEnabled(EHairStrandsShaderType::All, View.GetShaderPlatform()) && HairStrands::IsHairStrandsCompatible(MeshBatch.Mesh);
 	if (bIsHairStrandsCompatible && IsHairStrandsVisible(MeshBatch))
 	{
 		View.HairStrandsMeshElements.AddUninitialized(1);
diff --git a/Engine/Source/Runtime/Renderer/Public/HairStrandsInterface.h b/Engine/Source/Runtime/Renderer/Public/HairStrandsInterface.h
--- a/Engine/Source/Runtime/Renderer/Public/HairStrandsInterface.h
+++ b/Engine/Source/Runtime/Renderer/Public/HairStrandsInterface.h
@@ -144,6 +144,8 @@ struct RENDERER_API FHairStrandsInstance
 	uint32 AddRef() const;
 	uint32 Release() const;
 	int32 RegisteredIndex = -1;
+	virtual const FBoxSphereBounds* GetBounds() { return nullptr; }
+	virtual const FHairGroupPublicData* GetHairData() { return nullptr; }
 protected:
 	mutable uint32 RefCount = 0;
 };
@@ -395,8 +397,6 @@ RENDERER_API bool IsHairStrandsSupported(EHairStrandsShaderType Type, EShaderPla
 RENDERER_API bool IsHairStrandsEnabled(EHairStrandsShaderType Type, EShaderPlatform Platform = EShaderPlatform::SP_NumPlatforms);
 RENDERER_API void SetHairStrandsEnabled(bool In);
 
-RENDERER_API bool IsHairStrandsCompatible(const FMeshBatch* Mesh);
-
 RENDERER_API bool IsHairRayTracingEnabled();
 
 // Return true if hair simulation is enabled.
