Split hair strands and hair cards mesh batches filtering.

diff --git a/Engine/Source/Runtime/Renderer/Private/HairStrands/HairStrandsData.h b/Engine/Source/Runtime/Renderer/Private/HairStrands/HairStrandsData.h
--- a/Engine/Source/Runtime/Renderer/Private/HairStrands/HairStrandsData.h
+++ b/Engine/Source/Runtime/Renderer/Private/HairStrands/HairStrandsData.h
@@ -16,6 +16,7 @@ class FViewInfo;
 class FScene;
 class FInstanceCullingManager;
 struct FMeshBatch;
+struct FMeshBatchAndRelevance;
 
 ////////////////////////////////////////////////////////////////////////////////////
 // HairStrands uniform buffer
@@ -433,7 +434,13 @@ namespace HairStrands
 
 	// Mesh batch helpers
 	FHairGroupPublicData* GetHairData(const FMeshBatch* In);
-	bool IsHairStrandsCompatible(const FMeshBatch* Mesh);
+	bool IsHairCompatible(const FMeshBatch* Mesh);
 	bool IsHairStrandsVF(const FMeshBatch* Mesh);
 	bool IsHairCardsVF(const FMeshBatch* Mesh);
+	bool IsHairVisible(const FMeshBatchAndRelevance& MeshBatch);
+
+	// Hair helpers
+	bool HasHairInstanceInScene(const FScene& Scene);
+	bool HasHairCardsVisible(const TArrayView<FViewInfo>& Views);
+	bool HasHairStrandsVisible(const TArrayView<FViewInfo>& Views);
 }
\ No newline at end of file
diff --git a/Engine/Source/Runtime/Renderer/Private/HairStrands/HairStrandsDebug.cpp b/Engine/Source/Runtime/Renderer/Private/HairStrands/HairStrandsDebug.cpp
--- a/Engine/Source/Runtime/Renderer/Private/HairStrands/HairStrandsDebug.cpp
+++ b/Engine/Source/Runtime/Renderer/Private/HairStrands/HairStrandsDebug.cpp
@@ -1111,7 +1111,6 @@ static void InternalRenderHairStrandsDebugInfo(
 	}
 
 	// Display tangent vector for strands/cards/meshes
-	if (GHairTangentDebug > 0 || View.HairStrandsMeshElements.Num() > 0)
 	{
 		// Check among the hair instances, if hair tangent debug mode is requested
 		bool bTangentEnabled = GHairTangentDebug > 0;
@@ -1127,6 +1126,18 @@ static void InternalRenderHairStrandsDebugInfo(
 				}
 			}
 		}
+		if (!bTangentEnabled)
+		{
+			for (const FMeshBatchAndRelevance& Mesh : View.HairCardsMeshElements)
+			{
+				const FHairGroupPublicData* GroupData = HairStrands::GetHairData(Mesh.Mesh);
+				if (GroupData->DebugMode == EHairStrandsDebugMode::RenderHairTangent)
+				{
+					bTangentEnabled = true;
+					break;
+				}
+			}
+		}
 		if (bTangentEnabled)
 		{
 			AddDebugHairTangentPass(GraphBuilder, View, SceneTextures, SceneColorTexture);
diff --git a/Engine/Source/Runtime/Renderer/Private/HairStrands/HairStrandsInterface.cpp b/Engine/Source/Runtime/Renderer/Private/HairStrands/HairStrandsInterface.cpp
--- a/Engine/Source/Runtime/Renderer/Private/HairStrands/HairStrandsInterface.cpp
+++ b/Engine/Source/Runtime/Renderer/Private/HairStrands/HairStrandsInterface.cpp
@@ -437,7 +437,21 @@ FHairStrandsBookmarkParameters CreateHairStrandsBookmarkParameters(FScene* Scene
 				Out.VisibleInstances.Add(HairGroupPublicData->Instance);
 			}
 		}
-	}	
+	}
+
+	for (const FMeshBatchAndRelevance& MeshBatch : View.HairCardsMeshElements)
+	{
+		check(MeshBatch.PrimitiveSceneProxy && MeshBatch.PrimitiveSceneProxy->ShouldRenderInMainPass());
+		if (MeshBatch.Mesh && MeshBatch.Mesh->Elements.Num() > 0)
+		{
+			FHairGroupPublicData* HairGroupPublicData = HairStrands::GetHairData(MeshBatch.Mesh);
+			if (HairGroupPublicData && HairGroupPublicData->Instance)
+			{
+				Out.VisibleInstances.Add(HairGroupPublicData->Instance);
+			}
+		}
+	}
+
 	Out.ShaderDebugData			= ShaderDrawDebug::IsEnabled(View) ? &View.ShaderDrawData : nullptr;
 	Out.ShaderPrintData			= ShaderPrint::IsEnabled(View) ? &View.ShaderPrintData : nullptr;
 	Out.SkinCache				= View.Family->Scene->GetGPUSkinCache();
@@ -474,21 +488,6 @@ FHairStrandsBookmarkParameters CreateHairStrandsBookmarkParameters(FScene* Scene
 	return Out;
 }
 
-bool IsHairStrandsVisible(const FMeshBatchAndRelevance& MeshBatch)
-{
-	if (MeshBatch.Mesh && MeshBatch.PrimitiveSceneProxy && MeshBatch.PrimitiveSceneProxy->ShouldRenderInMainPass())
-	{
-		const FHairGroupPublicData* Data = HairStrands::GetHairData(MeshBatch.Mesh);
-		switch (Data->VFInput.GeometryType)
-		{
-			case EHairGeometryType::Strands: return Data->VFInput.Strands.HairLengthScale > 0;
-			case EHairGeometryType::Cards  : return true;
-			case EHairGeometryType::Meshes : return true;
-		}
-	}
-	return false;
-}
-
 namespace HairStrands
 {
 
@@ -514,11 +513,26 @@ bool IsHairCardsVF(const FMeshBatch* Mesh)
 	return false;
 }
 
-bool IsHairStrandsCompatible(const FMeshBatch* Mesh)
+bool IsHairCompatible(const FMeshBatch* Mesh)
 {
 	return IsHairStrandsVF(Mesh) || IsHairCardsVF(Mesh);
 }
 
+bool IsHairVisible(const FMeshBatchAndRelevance& MeshBatch)
+{
+	if (MeshBatch.Mesh && MeshBatch.PrimitiveSceneProxy && MeshBatch.PrimitiveSceneProxy->ShouldRenderInMainPass())
+	{
+		const FHairGroupPublicData* Data = HairStrands::GetHairData(MeshBatch.Mesh);
+		switch (Data->VFInput.GeometryType)
+		{
+		case EHairGeometryType::Strands: return Data->VFInput.Strands.HairLengthScale > 0;
+		case EHairGeometryType::Cards: return true;
+		case EHairGeometryType::Meshes: return true;
+		}
+	}
+	return false;
+}
+
 FHairGroupPublicData* GetHairData(const FMeshBatch* Mesh)
 {
 	return reinterpret_cast<FHairGroupPublicData*>(Mesh->Elements[0].VertexFactoryUserData);
diff --git a/Engine/Source/Runtime/Renderer/Private/HairStrands/HairStrandsRendering.cpp b/Engine/Source/Runtime/Renderer/Private/HairStrands/HairStrandsRendering.cpp
--- a/Engine/Source/Runtime/Renderer/Private/HairStrands/HairStrandsRendering.cpp
+++ b/Engine/Source/Runtime/Renderer/Private/HairStrands/HairStrandsRendering.cpp
@@ -306,4 +306,33 @@ bool HasViewHairStrandsData(const TArrayView<FViewInfo>& Views)
 	return false;
 }
 
+bool HasHairStrandsVisible(const TArrayView<FViewInfo>& Views)
+{
+	for (const FViewInfo& View : Views)
+	{
+		if (View.HairStrandsMeshElements.Num() > 0)
+		{
+			return true;
+		}
+	}
+	return false;
 }
+
+bool HasHairCardsVisible(const TArrayView<FViewInfo>& Views)
+{
+	for (const FViewInfo& View : Views)
+	{
+		if (View.HairCardsMeshElements.Num() > 0)
+		{
+			return true;
+		}
+	}
+	return false;
+}
+
+bool HasHairInstanceInScene(const FScene& Scene)
+{	
+	return Scene.HairStrandsSceneData.RegisteredProxies.Num() > 0;
+}
+
+} // HairStrands
diff --git a/Engine/Source/Runtime/Renderer/Private/HairStrands/HairStrandsVisibility.cpp b/Engine/Source/Runtime/Renderer/Private/HairStrands/HairStrandsVisibility.cpp
--- a/Engine/Source/Runtime/Renderer/Private/HairStrands/HairStrandsVisibility.cpp
+++ b/Engine/Source/Runtime/Renderer/Private/HairStrands/HairStrandsVisibility.cpp
@@ -816,11 +816,13 @@ static FMaterialPassOutput AddHairMaterialPass(
 	{
 		for (const FHairStrandsMacroGroupData::PrimitiveInfo& PrimitiveInfo : MacroGroupData.PrimitivesInfos)
 		{
-			const FMeshBatch& MeshBatch = *PrimitiveInfo.Mesh;
-			if (MeshBatch.MaterialRenderProxy->GetIncompleteMaterialWithFallback(FeatureLevel).HasEmissiveColorConnected())
+			if (const FMeshBatch* MeshBatch = PrimitiveInfo.Mesh)
 			{
-				bHasEmissiveMaterial = true;
-				break;
+				if (MeshBatch->MaterialRenderProxy->GetIncompleteMaterialWithFallback(FeatureLevel).HasEmissiveColorConnected())
+				{
+					bHasEmissiveMaterial = true;
+					break;
+				}
 			}
 		}
 	}
@@ -886,18 +888,20 @@ static FMaterialPassOutput AddHairMaterialPass(
 			{
 				for (const FHairStrandsMacroGroupData::PrimitiveInfo& PrimitiveInfo : MacroGroupData.PrimitivesInfos)
 				{
-					const FMeshBatch& MeshBatch = *PrimitiveInfo.Mesh;
-					const uint64 BatchElementMask = ~0ull;
-					bool bIsCompatible = true;
-					if (Filter != EHairMaterialPassFilter::All)
+					if (const FMeshBatch* MeshBatch = PrimitiveInfo.Mesh)
 					{
-						const bool bHasEmissive = MeshBatch.MaterialRenderProxy->GetIncompleteMaterialWithFallback(FeatureLevel).HasEmissiveColorConnected();
-						bIsCompatible = (bHasEmissive && Filter == EHairMaterialPassFilter::EmissiveOnly) || (!bHasEmissive && Filter == EHairMaterialPassFilter::NonEmissiveOnly);
-					}
+						const uint64 BatchElementMask = ~0ull;
+						bool bIsCompatible = true;
+						if (Filter != EHairMaterialPassFilter::All)
+						{
+							const bool bHasEmissive = MeshBatch->MaterialRenderProxy->GetIncompleteMaterialWithFallback(FeatureLevel).HasEmissiveColorConnected();
+							bIsCompatible = (bHasEmissive && Filter == EHairMaterialPassFilter::EmissiveOnly) || (!bHasEmissive && Filter == EHairMaterialPassFilter::NonEmissiveOnly);
+						}
 
-					if (bIsCompatible)
-					{
-						MeshProcessor.AddMeshBatch(MeshBatch, BatchElementMask, PrimitiveInfo.PrimitiveSceneProxy, -1, MacroGroupData.MacroGroupId, PrimitiveInfo.MaterialId);
+						if (bIsCompatible)
+						{
+							MeshProcessor.AddMeshBatch(*MeshBatch, BatchElementMask, PrimitiveInfo.PrimitiveSceneProxy, -1, MacroGroupData.MacroGroupId, PrimitiveInfo.MaterialId);
+						}
 					}
 				}
 			}
@@ -2234,9 +2238,11 @@ static void AddHairVisibilityCommonPass(
 			{
 				for (const FHairStrandsMacroGroupData::PrimitiveInfo& PrimitiveInfo : MacroGroupData.PrimitivesInfos)
 				{
-					const FMeshBatch& MeshBatch = *PrimitiveInfo.Mesh;
-					const uint64 BatchElementMask = ~0ull;
-					MeshProcessor.AddMeshBatch(MeshBatch, BatchElementMask, PrimitiveInfo.PrimitiveSceneProxy, -1, MacroGroupData.MacroGroupId, PrimitiveInfo.MaterialId, PrimitiveInfo.IsCullingEnable());
+					if (const FMeshBatch* MeshBatch = PrimitiveInfo.Mesh)
+					{
+						const uint64 BatchElementMask = ~0ull;
+						MeshProcessor.AddMeshBatch(*MeshBatch, BatchElementMask, PrimitiveInfo.PrimitiveSceneProxy, -1, MacroGroupData.MacroGroupId, PrimitiveInfo.MaterialId, PrimitiveInfo.IsCullingEnable());
+					}
 				}
 			}
 		}
diff --git a/Engine/Source/Runtime/Renderer/Private/SceneRendering.h b/Engine/Source/Runtime/Renderer/Private/SceneRendering.h
--- a/Engine/Source/Runtime/Renderer/Private/SceneRendering.h
+++ b/Engine/Source/Runtime/Renderer/Private/SceneRendering.h
@@ -1271,8 +1271,9 @@ public:
 	/* [PrimitiveIndex] = end index index in DynamicMeshElements[], to support GetDynamicMeshElementRange(). Contains valid values only for visible primitives with bDynamicRelevance. */
 	TArray<uint32, SceneRenderingAllocator> DynamicMeshEndIndices;
 
-	/** Hair strands dynamic mesh element. */
+	/** Hair strands & cards dynamic mesh element. */
 	TArray<FMeshBatchAndRelevance, SceneRenderingAllocator> HairStrandsMeshElements;
+	TArray<FMeshBatchAndRelevance, SceneRenderingAllocator> HairCardsMeshElements;
 
 	/* Mesh pass relevance for gathered dynamic mesh elements. */
 	TArray<FMeshPassMask, SceneRenderingAllocator> DynamicMeshElementsPassRelevance;
diff --git a/Engine/Source/Runtime/Renderer/Private/SceneVisibility.cpp b/Engine/Source/Runtime/Renderer/Private/SceneVisibility.cpp
--- a/Engine/Source/Runtime/Renderer/Private/SceneVisibility.cpp
+++ b/Engine/Source/Runtime/Renderer/Private/SceneVisibility.cpp
@@ -3108,7 +3108,6 @@ static void ComputeAndMarkRelevanceForViewParallel(
 	}
 }
 
-bool IsHairStrandsVisible(const FMeshBatchAndRelevance& MeshBatch);
 void ComputeDynamicMeshRelevance(EShadingPath ShadingPath, bool bAddLightmapDensityCommands, const FPrimitiveViewRelevance& ViewRelevance, const FMeshBatchAndRelevance& MeshBatch, FViewInfo& View, FMeshPassMask& PassMask, FPrimitiveSceneInfo* PrimitiveSceneInfo, const FPrimitiveBounds& Bounds)
 {
 	const int32 NumElements = MeshBatch.Mesh->Elements.Num();
@@ -3291,12 +3290,22 @@ void ComputeDynamicMeshRelevance(EShadingPath ShadingPath, bool bAddLightmapDens
 		BatchAndProxy.SortKey = MeshBatch.PrimitiveSceneProxy->GetTranslucencySortPriority();
 	}
 	
-	const bool bIsHairStrandsCompatible = ViewRelevance.bHairStrands && IsHairStrandsEnabled(EHairStrandsShaderType::All, View.GetShaderPlatform()) && HairStrands::IsHairStrandsCompatible(MeshBatch.Mesh);
-	if (bIsHairStrandsCompatible && IsHairStrandsVisible(MeshBatch))
+	const bool bIsHairStrandsCompatible = ViewRelevance.bHairStrands && IsHairStrandsEnabled(EHairStrandsShaderType::All, View.GetShaderPlatform());
+	if (bIsHairStrandsCompatible && HairStrands::IsHairVisible(MeshBatch))
 	{
-		View.HairStrandsMeshElements.AddUninitialized(1);
-		FMeshBatchAndRelevance& BatchAndProxy = View.HairStrandsMeshElements.Last();
-		BatchAndProxy = MeshBatch;
+		if (HairStrands::IsHairStrandsVF(MeshBatch.Mesh))
+		{
+			View.HairStrandsMeshElements.AddUninitialized(1);
+			FMeshBatchAndRelevance& BatchAndProxy = View.HairStrandsMeshElements.Last();
+			BatchAndProxy = MeshBatch;
+		}
+
+		if (HairStrands::IsHairCardsVF(MeshBatch.Mesh))
+		{
+			View.HairCardsMeshElements.AddUninitialized(1);
+			FMeshBatchAndRelevance& BatchAndProxy = View.HairCardsMeshElements.Last();
+			BatchAndProxy = MeshBatch;
+		}
 	}
 }
 
