Clean how active hair instances are queried.

diff --git a/Engine/Plugins/Runtime/HairStrands/Source/HairStrandsCore/Private/GroomManager.cpp b/Engine/Plugins/Runtime/HairStrands/Source/HairStrandsCore/Private/GroomManager.cpp
--- a/Engine/Plugins/Runtime/HairStrands/Source/HairStrandsCore/Private/GroomManager.cpp
+++ b/Engine/Plugins/Runtime/HairStrands/Source/HairStrandsCore/Private/GroomManager.cpp
@@ -1112,9 +1112,4 @@ void ProcessHairStrandsBookmark(
 			Parameters.View->UnscaledViewRect,
 			Parameters.View->ViewUniformBuffer);
 	}
-}
-
-void ProcessHairStrandsParameters(FHairStrandsBookmarkParameters& Parameters)
-{
-	Parameters.bHasElements = Parameters.Instances && Parameters.Instances->Num() > 0;
-}
+}
\ No newline at end of file
diff --git a/Engine/Plugins/Runtime/HairStrands/Source/HairStrandsCore/Private/HairStrandsCore.cpp b/Engine/Plugins/Runtime/HairStrands/Source/HairStrandsCore/Private/HairStrandsCore.cpp
--- a/Engine/Plugins/Runtime/HairStrands/Source/HairStrandsCore/Private/HairStrandsCore.cpp
+++ b/Engine/Plugins/Runtime/HairStrands/Source/HairStrandsCore/Private/HairStrandsCore.cpp
@@ -17,13 +17,11 @@ void ProcessHairStrandsBookmark(
 	EHairStrandsBookmark Bookmark,
 	FHairStrandsBookmarkParameters& Parameters);
 
-void ProcessHairStrandsParameters(FHairStrandsBookmarkParameters& Parameters);
-
 FHairAssetHelper HairStrandsCore_AssetHelper;
 
 void FHairStrandsCore::StartupModule()
 {
-	RegisterBookmarkFunction(ProcessHairStrandsBookmark, ProcessHairStrandsParameters);
+	RegisterBookmarkFunction(ProcessHairStrandsBookmark);
 
 	// Maps virtual shader source directory /Plugin/FX/Niagara to the plugin's actual Shaders directory.
 	FString PluginShaderDir = FPaths::Combine(IPluginManager::Get().FindPlugin(TEXT("HairStrands"))->GetBaseDir(), TEXT("Shaders"));
diff --git a/Engine/Source/Runtime/Renderer/Private/DeferredShadingRenderer.cpp b/Engine/Source/Runtime/Renderer/Private/DeferredShadingRenderer.cpp
--- a/Engine/Source/Runtime/Renderer/Private/DeferredShadingRenderer.cpp
+++ b/Engine/Source/Runtime/Renderer/Private/DeferredShadingRenderer.cpp
@@ -2282,10 +2282,10 @@ void FDeferredShadingSceneRenderer::Render(FRDGBuilder& GraphBuilder)
 
 		// Interpolation needs to happen after the skin cache run as there is a dependency 
 		// on the skin cache output.
-		const bool bRunHairStrands = HairStrandsBookmarkParameters.bHasElements && (Views.Num() > 0);
+		const bool bRunHairStrands = HairStrandsBookmarkParameters.HasInstances() && (Views.Num() > 0);
 		if (bRunHairStrands)
 		{
-			if (HairStrandsBookmarkParameters.bStrandsGeometryEnabled)
+			if (IsHairStrandsEnabled(EHairStrandsShaderType::Strands, Scene->GetShaderPlatform()))
 			{
 				RunHairStrandsBookmark(GraphBuilder, EHairStrandsBookmark::ProcessGatherCluster, HairStrandsBookmarkParameters);
 
@@ -2316,7 +2316,7 @@ void FDeferredShadingSceneRenderer::Render(FRDGBuilder& GraphBuilder)
 	const bool bShouldRenderVelocities = ShouldRenderVelocities();
 	const bool bBasePassCanOutputVelocity = FVelocityRendering::BasePassCanOutputVelocity(FeatureLevel);
 	const bool bUseSelectiveBasePassOutputs = IsUsingSelectiveBasePassOutputs(ShaderPlatform);
-	const bool bHairEnable = HairStrandsBookmarkParameters.bHasElements && Views.Num() > 0 && IsHairStrandsEnabled(EHairStrandsShaderType::Strands, Views[0].GetShaderPlatform());
+	const bool bHairStrandsEnable = HairStrandsBookmarkParameters.HasInstances() && Views.Num() > 0 && IsHairStrandsEnabled(EHairStrandsShaderType::Strands, Views[0].GetShaderPlatform());
 
 	{
 		GraphBuilder.SetCommandListStat(GET_STATID(STAT_CLM_PrePass));
@@ -2347,7 +2347,7 @@ void FDeferredShadingSceneRenderer::Render(FRDGBuilder& GraphBuilder)
 		{
 			// Render the velocities of movable objects
 			GraphBuilder.SetCommandListStat(GET_STATID(STAT_CLM_Velocity));
-			RenderVelocities(GraphBuilder, SceneTextures, EVelocityPass::Opaque, bHairEnable);
+			RenderVelocities(GraphBuilder, SceneTextures, EVelocityPass::Opaque, bHairStrandsEnable);
 			GraphBuilder.SetCommandListStat(GET_STATID(STAT_CLM_AfterVelocity));
 		}
 	}
@@ -2638,7 +2638,7 @@ void FDeferredShadingSceneRenderer::Render(FRDGBuilder& GraphBuilder)
 		ensureMsgf(!VirtualShadowMapArray.IsEnabled(), TEXT("Virtual shadow maps are not supported in the forward shading path"));
 		RenderShadowDepthMaps(GraphBuilder, InstanceCullingManager);
 
-		if (bHairEnable && !bHasRayTracedOverlay)
+		if (bHairStrandsEnable && !bHasRayTracedOverlay)
 		{
 			RenderHairPrePass(GraphBuilder, Scene, Views, InstanceCullingManager);
 			RenderHairBasePass(GraphBuilder, Scene, SceneTextures, Views, InstanceCullingManager);
@@ -2736,7 +2736,7 @@ void FDeferredShadingSceneRenderer::Render(FRDGBuilder& GraphBuilder)
 	}
 
 	// Render hair
-	if (bHairEnable && !IsForwardShadingEnabled(ShaderPlatform) && !bHasRayTracedOverlay)
+	if (bHairStrandsEnable && !IsForwardShadingEnabled(ShaderPlatform) && !bHasRayTracedOverlay)
 	{
 		RenderHairPrePass(GraphBuilder, Scene, Views, InstanceCullingManager);
 		RenderHairBasePass(GraphBuilder, Scene, SceneTextures, Views, InstanceCullingManager);
@@ -2830,7 +2830,7 @@ void FDeferredShadingSceneRenderer::Render(FRDGBuilder& GraphBuilder)
 	{
 		// Render the velocities of movable objects
 		GraphBuilder.SetCommandListStat(GET_STATID(STAT_CLM_Velocity));
-		RenderVelocities(GraphBuilder, SceneTextures, EVelocityPass::Opaque, bHairEnable);
+		RenderVelocities(GraphBuilder, SceneTextures, EVelocityPass::Opaque, bHairStrandsEnable);
 		GraphBuilder.SetCommandListStat(GET_STATID(STAT_CLM_AfterVelocity));
 
 		// TODO: Populate velocity buffer from Nanite visibility buffer.
@@ -3117,7 +3117,7 @@ void FDeferredShadingSceneRenderer::Render(FRDGBuilder& GraphBuilder)
 #endif
 
 	{
-		if (HairStrandsBookmarkParameters.bHasElements)
+		if (HairStrandsBookmarkParameters.HasInstances())
 		{
 			RenderHairStrandsDebugInfo(GraphBuilder, Scene, Views, HairStrandsBookmarkParameters.HairClusterData, SceneTextures.Color.Target, SceneTextures.Depth.Target);
 		}
diff --git a/Engine/Source/Runtime/Renderer/Private/HairStrands/HairStrandsDebug.cpp b/Engine/Source/Runtime/Renderer/Private/HairStrands/HairStrandsDebug.cpp
--- a/Engine/Source/Runtime/Renderer/Private/HairStrands/HairStrandsDebug.cpp
+++ b/Engine/Source/Runtime/Renderer/Private/HairStrands/HairStrandsDebug.cpp
@@ -1092,7 +1092,7 @@ static void InternalRenderHairStrandsDebugInfo(
 	FHairStrandsBookmarkParameters Params = CreateHairStrandsBookmarkParameters(Scene, View);
 	Params.SceneColorTexture = SceneColorTexture;
 	Params.SceneDepthTexture = SceneDepthTexture;
-	if (!Params.bHasElements)
+	if (!Params.HasInstances())
 	{
 		return;
 	}
diff --git a/Engine/Source/Runtime/Renderer/Private/HairStrands/HairStrandsInterface.cpp b/Engine/Source/Runtime/Renderer/Private/HairStrands/HairStrandsInterface.cpp
--- a/Engine/Source/Runtime/Renderer/Private/HairStrands/HairStrandsInterface.cpp
+++ b/Engine/Source/Runtime/Renderer/Private/HairStrands/HairStrandsInterface.cpp
@@ -392,18 +392,12 @@ void TransitBufferToReadable(FRDGBuilder& GraphBuilder, FBufferTransitionQueue&
 ///////////////////////////////////////////////////////////////////////////////////////////////////
 // Bookmark API
 THairStrandsBookmarkFunction  GHairStrandsBookmarkFunction = nullptr;
-THairStrandsParameterFunction GHairStrandsParameterFunction = nullptr;
-void RegisterBookmarkFunction(THairStrandsBookmarkFunction Bookmark, THairStrandsParameterFunction Parameter)
+void RegisterBookmarkFunction(THairStrandsBookmarkFunction Bookmark)
 {
 	if (Bookmark)
 	{
 		GHairStrandsBookmarkFunction = Bookmark;
 	}
-
-	if (Parameter)
-	{
-		GHairStrandsParameterFunction = Parameter;
-	}
 }
 
 void RunHairStrandsBookmark(FRDGBuilder& GraphBuilder, EHairStrandsBookmark Bookmark, FHairStrandsBookmarkParameters& Parameters)
@@ -461,17 +455,11 @@ FHairStrandsBookmarkParameters CreateHairStrandsBookmarkParameters(FScene* Scene
 	Out.ViewRect				= View.ViewRect;
 	Out.ViewUniqueID			= View.ViewState ? View.ViewState->UniqueID : ~0;
 	Out.SceneColorTexture		= nullptr;
-	Out.bStrandsGeometryEnabled = IsHairStrandsEnabled(EHairStrandsShaderType::Strands, View.GetShaderPlatform());
+	Out.bHzbRequest				= false; // Out.HasInstances() && IsHairStrandsEnabled(EHairStrandsShaderType::Strands, View.GetShaderPlatform());
 
 	// Sanity check
 	check(Out.Instances->Num() >= Out.VisibleInstances.Num());
 
-	if (GHairStrandsParameterFunction)
-	{
-		GHairStrandsParameterFunction(Out);
-	}
-	Out.bHzbRequest = false; // Out.bHasElements&& Out.bStrandsGeometryEnabled;
-
 	return Out;
 }
 
diff --git a/Engine/Source/Runtime/Renderer/Private/SceneRendering.cpp b/Engine/Source/Runtime/Renderer/Private/SceneRendering.cpp
--- a/Engine/Source/Runtime/Renderer/Private/SceneRendering.cpp
+++ b/Engine/Source/Runtime/Renderer/Private/SceneRendering.cpp
@@ -4090,7 +4090,7 @@ static void RenderViewFamily_RenderThread(FRHICommandListImmediate& RHICmdList,
 			if (IsHairStrandsEnabled(EHairStrandsShaderType::All, SceneRenderer->Scene->GetShaderPlatform()) && SceneRenderer->Views.Num() > 0 && !ViewFamily.EngineShowFlags.HitProxies)
 			{
 				FHairStrandsBookmarkParameters Parameters = CreateHairStrandsBookmarkParameters(SceneRenderer->Scene, SceneRenderer->Views);
-				if (Parameters.bHasElements)
+				if (Parameters.HasInstances())
 				{
 					RunHairStrandsBookmark(EHairStrandsBookmark::ProcessEndOfFrame, Parameters);
 				}
diff --git a/Engine/Source/Runtime/Renderer/Private/SceneVisibility.cpp b/Engine/Source/Runtime/Renderer/Private/SceneVisibility.cpp
--- a/Engine/Source/Runtime/Renderer/Private/SceneVisibility.cpp
+++ b/Engine/Source/Runtime/Renderer/Private/SceneVisibility.cpp
@@ -3488,7 +3488,7 @@ void FSceneRenderer::PreVisibilityFrameSetup(FRDGBuilder& GraphBuilder, const FS
 	if (Views.Num() > 0 && !ViewFamily.EngineShowFlags.HitProxies)
 	{
 		FHairStrandsBookmarkParameters Parameters = CreateHairStrandsBookmarkParameters(Scene, Views);
-		if (Parameters.bHasElements)
+		if (Parameters.HasInstances())
 		{
 			RunHairStrandsBookmark(GraphBuilder, EHairStrandsBookmark::ProcessLODSelection, Parameters);
 		}
diff --git a/Engine/Source/Runtime/Renderer/Public/HairStrandsInterface.h b/Engine/Source/Runtime/Renderer/Public/HairStrandsInterface.h
--- a/Engine/Source/Runtime/Renderer/Public/HairStrandsInterface.h
+++ b/Engine/Source/Runtime/Renderer/Public/HairStrandsInterface.h
@@ -450,14 +450,13 @@ struct FHairStrandsBookmarkParameters
 	FRDGTextureRef SceneDepthTexture = nullptr; 
 
 	bool bHzbRequest = false;
-	bool bHasElements = false;
-	bool bStrandsGeometryEnabled = false;
 	uint32 FrameIndex = ~0;
 
 	// Temporary
 	FHairStrandClusterData HairClusterData;
+
+	inline bool HasInstances() const { return Instances != nullptr && Instances->Num() > 0; }
 };
 
-typedef void (*THairStrandsParameterFunction)(FHairStrandsBookmarkParameters& Parameters);
 typedef void (*THairStrandsBookmarkFunction)(FRDGBuilder* GraphBuilder, EHairStrandsBookmark Bookmark, FHairStrandsBookmarkParameters& Parameters);
-RENDERER_API void RegisterBookmarkFunction(THairStrandsBookmarkFunction Bookmark, THairStrandsParameterFunction Parameters);
+RENDERER_API void RegisterBookmarkFunction(THairStrandsBookmarkFunction Bookmark);
