Add option to fix invalid shadow culling of hair strands when non-visible in the primary view.

Disable by default for now as this can  add non-negligeable voxelization cost due to inaccurate shadow view culling.

diff --git a/Engine/Plugins/Runtime/HairStrands/Source/HairStrandsCore/Public/GroomInstance.h b/Engine/Plugins/Runtime/HairStrands/Source/HairStrandsCore/Public/GroomInstance.h
--- a/Engine/Plugins/Runtime/HairStrands/Source/HairStrandsCore/Public/GroomInstance.h
+++ b/Engine/Plugins/Runtime/HairStrands/Source/HairStrandsCore/Public/GroomInstance.h
@@ -283,8 +283,8 @@ struct HAIRSTRANDSCORE_API FHairGroupInstance : public FHairStrandsInstance
 		return Meshes.IsValid() || Cards.IsValid() || Strands.IsValid();
 	}
 
-	virtual const FBoxSphereBounds* GetBounds() override { return ProxyBounds;  }
-	virtual const FHairGroupPublicData* GetHairData() { return HairGroupPublicData; }
+	virtual const FBoxSphereBounds* GetBounds() const override { return ProxyBounds;  }
+	virtual const FHairGroupPublicData* GetHairData() const { return HairGroupPublicData; }
 
 	/** Get the current local to world transform according to the internal binding type */
 	FORCEINLINE const FTransform& GetCurrentLocalToWorld() const
diff --git a/Engine/Source/Runtime/Renderer/Private/HairStrands/HairStrandsCluster.cpp b/Engine/Source/Runtime/Renderer/Private/HairStrands/HairStrandsCluster.cpp
--- a/Engine/Source/Runtime/Renderer/Private/HairStrands/HairStrandsCluster.cpp
+++ b/Engine/Source/Runtime/Renderer/Private/HairStrands/HairStrandsCluster.cpp
@@ -128,6 +128,9 @@ static bool DoesGroupExists(uint32 ResourceId, uint32 GroupIndex, const FHairStr
 	return false;
 }
 
+bool IsHairStrandsNonVisibleShadowCastingEnable();
+bool IsHairStrandsVisibleInShadows(const FViewInfo& View, const FHairStrandsInstance& Instance);
+
 static void InternalUpdateMacroGroup(FHairStrandsMacroGroupData& MacroGroup, int32& MaterialId, FHairGroupPublicData* HairData, const FMeshBatch* Mesh, const FPrimitiveSceneProxy* Proxy)
 {
 	check(HairData);
@@ -231,6 +234,9 @@ void CreateHairStrandsMacroGroups(
 		}
 	};
 
+	// 1. Add all visible hair-strands instances
+	const int32 ActiveInstanceCount = Scene->HairStrandsSceneData.RegisteredProxies.Num();
+	TBitArray InstancesVisibility(false, ActiveInstanceCount);
 	for (const FMeshBatchAndRelevance& MeshBatchAndRelevance : View.HairStrandsMeshElements)
 	{
 		if (HairStrands::IsHairStrandsVF(MeshBatchAndRelevance.Mesh))
@@ -238,10 +244,28 @@ void CreateHairStrandsMacroGroups(
 			if (FHairGroupPublicData* HairData = HairStrands::GetHairData(MeshBatchAndRelevance.Mesh))
 			{
 				UpdateMacroGroup(HairData, MeshBatchAndRelevance.Mesh, MeshBatchAndRelevance.PrimitiveSceneProxy, nullptr);
+				InstancesVisibility[HairData->Instance->RegisteredIndex] = true;
+			}
+		}
+	}
+
+	// 2. Add all hair-strands instances which are non-visible in primary view(s) but visible in shadow view(s)
+	// Slow Linear search
+	if (IsHairStrandsNonVisibleShadowCastingEnable())
+	{
+		for (FHairStrandsInstance* Instance : Scene->HairStrandsSceneData.RegisteredProxies)
+		{
+			if (Instance->RegisteredIndex >= 0 && Instance->RegisteredIndex < ActiveInstanceCount && !InstancesVisibility[Instance->RegisteredIndex])
+			{
+				if (IsHairStrandsVisibleInShadows(View, *Instance))
+				{
+					UpdateMacroGroup(const_cast<FHairGroupPublicData*>(Instance->GetHairData()), nullptr, nullptr, Instance->GetBounds());
+				}
 			}
 		}
 	}
 
+	// Compute the screen size of macro group projection, for allocation purpose
 	for (FHairStrandsMacroGroupData& MacroGroup : MacroGroups)
 	{
 		MacroGroup.ScreenRect = ComputeProjectedScreenRect(MacroGroup.Bounds.GetBox(), View);
diff --git a/Engine/Source/Runtime/Renderer/Private/HairStrands/HairStrandsData.h b/Engine/Source/Runtime/Renderer/Private/HairStrands/HairStrandsData.h
--- a/Engine/Source/Runtime/Renderer/Private/HairStrands/HairStrandsData.h
+++ b/Engine/Source/Runtime/Renderer/Private/HairStrands/HairStrandsData.h
@@ -388,6 +388,10 @@ struct FHairStrandsViewData
 	FHairStrandsVoxelResources VirtualVoxelResources;
 	FHairStrandsMacroGroupResources MacroGroupResources;
 	FHairStrandsDebugData DebugData;
+
+	// Transient: store all light visible in primary view(s)
+	TArray<const FLightSceneInfo*> VisibleShadowCastingLights;
+	TArray<FSphere> VisibleShadowCastingBounds;
 };
 
 // View State data (i.e., persistent accross fram)
@@ -443,4 +447,7 @@ namespace HairStrands
 	bool HasHairInstanceInScene(const FScene& Scene);
 	bool HasHairCardsVisible(const TArrayView<FViewInfo>& Views);
 	bool HasHairStrandsVisible(const TArrayView<FViewInfo>& Views);
+
+	void AddVisibleShadowCastingLight(const FScene& Scene, TArray<FViewInfo>& Views, const FLightSceneInfo* LightSceneInfo);
+	void AddVisibleShadowCastingLight(const FScene& Scene, TArray<FViewInfo>& Views, const FSphere& Bounds);
 }
\ No newline at end of file
diff --git a/Engine/Source/Runtime/Renderer/Private/HairStrands/HairStrandsInterface.cpp b/Engine/Source/Runtime/Renderer/Private/HairStrands/HairStrandsInterface.cpp
--- a/Engine/Source/Runtime/Renderer/Private/HairStrands/HairStrandsInterface.cpp
+++ b/Engine/Source/Runtime/Renderer/Private/HairStrands/HairStrandsInterface.cpp
@@ -59,6 +59,11 @@ static TAutoConsoleVariable<int32> CVarHairStrandsSimulation(
 	TEXT("Enable/disable hair simulation"),
 	ECVF_RenderThreadSafe | ECVF_Scalability);
 
+static TAutoConsoleVariable<int32> CVarHairStrandsNonVisibleShadowCasting(
+	TEXT("r.HairStrands.Shadow.CastShadowWhenNonVisible"), 0,
+	TEXT("Enable shadow casting for hair strands even when culled out from the primary view"),
+	ECVF_RenderThreadSafe);
+
 ////////////////////////////////////////////////////////////////////////////////////////////////////////////////////////
 // Hair strands instance ref. counting for debug purpose only
 uint32 FHairStrandsInstance::GetRefCount() const
@@ -389,6 +394,54 @@ void TransitBufferToReadable(FRDGBuilder& GraphBuilder, FBufferTransitionQueue&
 	}
 }
 
+///////////////////////////////////////////////////////////////////////////////////////////////////
+
+bool IsHairStrandsNonVisibleShadowCastingEnable()
+{
+	return CVarHairStrandsNonVisibleShadowCasting.GetValueOnAnyThread() > 0;
+}
+
+bool IsHairStrandsVisibleInShadows(const FViewInfo& View, const FHairStrandsInstance& Instance)
+{
+	bool bIsVisibleInShadow = false;
+	if (const FHairGroupPublicData* HairData = Instance.GetHairData())
+	{
+		const int32 LODIndex = FMath::CeilToInt(HairData->LODIndex);
+		const bool bIsStrands = LODIndex >= 0 && HairData->IsVisible(LODIndex) && HairData->GetGeometryType(LODIndex) == EHairGeometryType::Strands;
+		if (bIsStrands)
+		{
+			if (const FBoxSphereBounds* Bounds = Instance.GetBounds())
+			{
+				{
+					for (const FLightSceneInfo* LightInfo : View.HairStrandsViewData.VisibleShadowCastingLights)
+					{
+						// Influence radius check
+						if (LightInfo->Proxy->AffectsBounds(*Bounds))
+						{
+							bIsVisibleInShadow = true;
+							break;
+						}
+					}
+				}
+
+				if (!bIsVisibleInShadow)
+				{
+					for (const FSphere& LightBound : View.HairStrandsViewData.VisibleShadowCastingBounds)
+					{
+						// Influence radius check
+						if (Bounds->GetSphere().Intersects(LightBound))
+						{
+							bIsVisibleInShadow = true;
+							break;
+						}
+					}
+				}
+			}
+		}
+	}
+	return bIsVisibleInShadow;
+}
+
 ///////////////////////////////////////////////////////////////////////////////////////////////////
 // Bookmark API
 THairStrandsBookmarkFunction  GHairStrandsBookmarkFunction = nullptr;
@@ -418,30 +471,53 @@ void RunHairStrandsBookmark(EHairStrandsBookmark Bookmark, FHairStrandsBookmarkP
 
 FHairStrandsBookmarkParameters CreateHairStrandsBookmarkParameters(FScene* Scene, FViewInfo& View)
 {
+	const int32 ActiveInstanceCount = Scene->HairStrandsSceneData.RegisteredProxies.Num();
+	TBitArray InstancesVisibility(false, ActiveInstanceCount);
+
 	FHairStrandsBookmarkParameters Out;
 	Out.VisibleInstances.Reserve(View.HairStrandsMeshElements.Num());
+
+	// 1. Add all visible strands instances
 	for (const FMeshBatchAndRelevance& MeshBatch : View.HairStrandsMeshElements)
 	{
 		check(MeshBatch.PrimitiveSceneProxy && MeshBatch.PrimitiveSceneProxy->ShouldRenderInMainPass());
 		if (MeshBatch.Mesh && MeshBatch.Mesh->Elements.Num() > 0)
 		{
-			FHairGroupPublicData* HairGroupPublicData = HairStrands::GetHairData(MeshBatch.Mesh);
-			if (HairGroupPublicData && HairGroupPublicData->Instance)
+			FHairGroupPublicData* HairData = HairStrands::GetHairData(MeshBatch.Mesh);
+			if (HairData && HairData->Instance)
 			{
-				Out.VisibleInstances.Add(HairGroupPublicData->Instance);
+				Out.VisibleInstances.Add(HairData->Instance);
+				InstancesVisibility[HairData->Instance->RegisteredIndex] = true;
 			}
 		}
 	}
 
+	// 2. Add all visible cards instances
 	for (const FMeshBatchAndRelevance& MeshBatch : View.HairCardsMeshElements)
 	{
 		check(MeshBatch.PrimitiveSceneProxy && MeshBatch.PrimitiveSceneProxy->ShouldRenderInMainPass());
 		if (MeshBatch.Mesh && MeshBatch.Mesh->Elements.Num() > 0)
 		{
-			FHairGroupPublicData* HairGroupPublicData = HairStrands::GetHairData(MeshBatch.Mesh);
-			if (HairGroupPublicData && HairGroupPublicData->Instance)
+			FHairGroupPublicData* HairData = HairStrands::GetHairData(MeshBatch.Mesh);
+			if (HairData && HairData->Instance)
+			{
+				Out.VisibleInstances.Add(HairData->Instance);
+				InstancesVisibility[HairData->Instance->RegisteredIndex] = true;
+			}
+		}
+	}
+
+	// 3. Add all instances non-visible primary view(s) but visible in shadow view(s)
+	if (IsHairStrandsNonVisibleShadowCastingEnable())
+	{
+		for (FHairStrandsInstance* Instance : Scene->HairStrandsSceneData.RegisteredProxies)
+		{
+			if (Instance->RegisteredIndex >= 0 && Instance->RegisteredIndex < ActiveInstanceCount && !InstancesVisibility[Instance->RegisteredIndex])
 			{
-				Out.VisibleInstances.Add(HairGroupPublicData->Instance);
+				if (IsHairStrandsVisibleInShadows(View, *Instance))
+				{
+					Out.VisibleInstances.Add(Instance);
+				}
 			}
 		}
 	}
@@ -526,4 +602,30 @@ FHairGroupPublicData* GetHairData(const FMeshBatch* Mesh)
 	return reinterpret_cast<FHairGroupPublicData*>(Mesh->Elements[0].VertexFactoryUserData);
 }
 
+void AddVisibleShadowCastingLight(const FScene& Scene, TArray<FViewInfo>& Views, const FLightSceneInfo* LightSceneInfo)
+{
+	for (FViewInfo& View : Views)
+	{
+		// If any hair data are registered, track which lights are visible so that hair strands can cast shadow even if not visibible in primary view
+		if (Scene.HairStrandsSceneData.RegisteredProxies.Num() > 0)
+		{
+			View.HairStrandsViewData.VisibleShadowCastingLights.Add(LightSceneInfo);
+			break;
+		}
+	}
+}
+
+void AddVisibleShadowCastingLight(const FScene& Scene, TArray<FViewInfo>& Views, const FSphere& Bounds)
+{
+	for (FViewInfo& View : Views)
+	{
+		// If any hair data are registered, track which lights are visible so that hair strands can cast shadow even if not visibible in primary view
+		if (Scene.HairStrandsSceneData.RegisteredProxies.Num() > 0)
+		{
+			View.HairStrandsViewData.VisibleShadowCastingBounds.Add(Bounds);
+			break;
+		}
+	}
+}
+
 } // namespace HairStrands
diff --git a/Engine/Source/Runtime/Renderer/Private/ShadowSetup.cpp b/Engine/Source/Runtime/Renderer/Private/ShadowSetup.cpp
--- a/Engine/Source/Runtime/Renderer/Private/ShadowSetup.cpp
+++ b/Engine/Source/Runtime/Renderer/Private/ShadowSetup.cpp
@@ -3804,6 +3804,8 @@ void FSceneRenderer::InitProjectedShadowVisibility()
 	SCOPE_CYCLE_COUNTER(STAT_InitProjectedShadowVisibility);
 	int32 NumBufferedFrames = FOcclusionQueryHelpers::GetNumBufferedFrames(FeatureLevel);
 
+	const bool bHairStrands = HairStrands::HasHairInstanceInScene(*Scene);
+
 	// Initialize the views' ProjectedShadowVisibilityMaps and remove shadows without subjects.
 	for(auto LightIt = Scene->Lights.CreateConstIterator();LightIt;++LightIt)
 	{
@@ -3934,6 +3936,12 @@ void FSceneRenderer::InitProjectedShadowVisibility()
 					}
 				}
 			}
+
+			// Register visible lights for allowing hair strands to cast shadow (directional light)
+			if (bHairStrands && LightIt->LightType == ELightComponentType::LightType_Directional)
+			{
+				HairStrands::AddVisibleShadowCastingLight(*Scene, Views, ProjectedShadowInfo.ShadowBounds);
+			}
 		}
 	}
 
@@ -5518,6 +5526,7 @@ FDynamicShadowsTaskData* FSceneRenderer::BeginInitDynamicShadows(bool bRunningEa
 	SCOPED_NAMED_EVENT_TEXT("FSceneRenderer::BeginInitDynamicShadows", FColor::Magenta);
 
 	const bool bMobile = FeatureLevel < ERHIFeatureLevel::SM5;
+	const bool bHairStrands = HairStrands::HasHairInstanceInScene(*Scene);
 
 	const bool bProjectEnablePointLightShadows = Scene->ReadOnlyCVARCache.bEnablePointLightShadows && !bMobile; // Point light shadow is unsupported on mobile for now.
 	const bool bProjectEnableMovableDirectionLightShadows = !bMobile || Scene->ReadOnlyCVARCache.bMobileAllowMovableDirectionalLights;
@@ -5628,6 +5637,12 @@ FDynamicShadowsTaskData* FSceneRenderer::BeginInitDynamicShadows(bool bRunningEa
 							CreateWholeSceneProjectedShadow(LightSceneInfo, NumPointShadowCachesUpdatedThisFrame, NumSpotShadowCachesUpdatedThisFrame);
 						}
 
+						// Register visible lights for allowing hair strands to cast shadow (non-directional light)
+						if (bHairStrands && LightSceneInfo->Proxy->GetLightType() != LightType_Directional)
+						{
+							HairStrands::AddVisibleShadowCastingLight(*Scene, Views, LightSceneInfo);
+						}
+
 						// Allow movable and stationary lights to create CSM, or static lights that are unbuilt
 						if ((!LightSceneInfo->Proxy->HasStaticLighting() && LightSceneInfoCompact.bCastDynamicShadow) || bCreateShadowToPreviewStaticLight)
 						{
diff --git a/Engine/Source/Runtime/Renderer/Public/HairStrandsInterface.h b/Engine/Source/Runtime/Renderer/Public/HairStrandsInterface.h
--- a/Engine/Source/Runtime/Renderer/Public/HairStrandsInterface.h
+++ b/Engine/Source/Runtime/Renderer/Public/HairStrandsInterface.h
@@ -144,8 +144,8 @@ struct RENDERER_API FHairStrandsInstance
 	uint32 AddRef() const;
 	uint32 Release() const;
 	int32 RegisteredIndex = -1;
-	virtual const FBoxSphereBounds* GetBounds() { return nullptr; }
-	virtual const FHairGroupPublicData* GetHairData() { return nullptr; }
+	virtual const FBoxSphereBounds* GetBounds() const { return nullptr; }
+	virtual const FHairGroupPublicData* GetHairData() const { return nullptr; }
 protected:
 	mutable uint32 RefCount = 0;
 };
@@ -195,6 +195,12 @@ public:
 	void SetLODVisibilities(const TArray<bool>& InLODVisibility) { LODVisibilities = InLODVisibility; }
 	const TArray<bool>& GetLODVisibilities() const { return LODVisibilities; }
 
+	bool IsVisible(int32 InLODIndex) const
+	{
+		if (InLODIndex < 0 && InLODIndex >= LODVisibilities.Num()) return false;
+		return LODVisibilities[InLODIndex];
+	}
+
 	EHairGeometryType GetGeometryType(int32 InLODIndex) const
 	{
 		if (InLODIndex < 0 && InLODIndex >= LODGeometryTypes.Num()) return EHairGeometryType::NoneGeometry;
