Fixing no unity error
#rnx

diff --git a/Engine/Source/Runtime/Renderer/Private/HairStrands/HairStrandsData.h b/Engine/Source/Runtime/Renderer/Private/HairStrands/HairStrandsData.h
--- a/Engine/Source/Runtime/Renderer/Private/HairStrands/HairStrandsData.h
+++ b/Engine/Source/Runtime/Renderer/Private/HairStrands/HairStrandsData.h
@@ -11,6 +11,7 @@
 #include "RHIGPUReadback.h"
 #include "Shader.h"
 
+class FLightSceneInfo;
 class FPrimitiveSceneProxy;
 class FViewInfo;
 class FScene;
