Strata
- Strata render targets follow base pass load action.
- There is now an option to never clear the non first uint target
- On console, fast clear is leveraged
- SSData is cleared during the categorisation pass if the pixel is non SSS only
- This shaved 1.3ms on one platofrms and 0.3ms on another

diff --git a/Engine/Shaders/Private/Strata/StrataMaterialClassification.usf b/Engine/Shaders/Private/Strata/StrataMaterialClassification.usf
--- a/Engine/Shaders/Private/Strata/StrataMaterialClassification.usf
+++ b/Engine/Shaders/Private/Strata/StrataMaterialClassification.usf
@@ -37,6 +37,10 @@ RWBuffer<uint> ComplexTileListDataBuffer;
 groupshared uint s_IsComplexIsMaterialTile[GROUP_THREAD_COUNT];
 #endif
 
+#if PERMUTATION_STRATA_CLEAR_DURING_CATEGORIZATION
+RWTexture2D<uint2> SSSTextureUAV;
+#endif
+
 [numthreads(TILE_SIZE, TILE_SIZE, 1)]
 void TileMainCS(uint2 DispatchThreadId : SV_DispatchThreadID, uint LinearIndex : SV_GroupIndex, uint3 GroupId : SV_GroupID)
 {
@@ -52,6 +56,7 @@ void TileMainCS(uint2 DispatchThreadId : SV_DispatchThreadID, uint LinearIndex :
 	// Pixels outside of the view area are considered simple to enable screen borders to receive the simple permutation when not aligned to shader group size.
 	bool bContainsComplexMaterial = false;
 	bool bContainsStrataMaterial = false;
+	bool bHasScreenSpaceSubsurfaceScattering = false;
 	if (all(PixelCoord < uint2(ViewResolution)))
 	{
 		FStrataAddressing StrataAddressing = GetStrataPixelDataByteOffset(PixelCoord, uint2(View.BufferSizeAndInvSize.xy), MaxBytesPerPixel);
@@ -60,8 +65,19 @@ void TileMainCS(uint2 DispatchThreadId : SV_DispatchThreadID, uint LinearIndex :
 		const bool bIsSimple = IsSimpleMaterial(StrataPixelHeader) || StrataPixelHeader.BSDFCount == 0; // BSDFCount == 0 ensures that non-strata pixel, like sky pixels, won't make a simple tile flagged as complex
 		bContainsStrataMaterial  = StrataPixelHeader.BSDFCount > 0;
 		bContainsComplexMaterial = !bIsSimple;
+		bHasScreenSpaceSubsurfaceScattering = HasSubsurface(StrataPixelHeader);
 	}
 
+#if PERMUTATION_STRATA_CLEAR_DURING_CATEGORIZATION
+	BRANCH
+	if (!bHasScreenSpaceSubsurfaceScattering)
+	{
+		// We must fill all the pixel which doe not hav subsurface scattering with default so that the SSS code is not executed where it should not.
+		FStrataSubsurfaceData StrataSubsurfaceData = (FStrataSubsurfaceData)0;
+		SSSTextureUAV[DispatchThreadId.xy] = StrataPackSSSData(StrataSubsurfaceData);
+	}
+#endif
+
 #if PERMUTATION_WAVE_OPS
 
 	const bool bHasAnyStrata  = WaveActiveAnyTrue(bContainsStrataMaterial);
diff --git a/Engine/Source/Runtime/RenderCore/Public/RenderGraphUtils.h b/Engine/Source/Runtime/RenderCore/Public/RenderGraphUtils.h
--- a/Engine/Source/Runtime/RenderCore/Public/RenderGraphUtils.h
+++ b/Engine/Source/Runtime/RenderCore/Public/RenderGraphUtils.h
@@ -73,24 +73,24 @@ struct FTextureRenderTargetBinding
 {
 	FRDGTextureRef Texture;
 	int16 ArraySlice;
-	ERenderTargetLoadAction OverrideLoadAction;
+	bool bNeverClear;
 
 	FTextureRenderTargetBinding()
 		: Texture(nullptr)
 		, ArraySlice(-1)
-		, OverrideLoadAction(ERenderTargetLoadAction::Num)
+		, bNeverClear(false)
 	{}
 
-	FTextureRenderTargetBinding(FRDGTextureRef InTexture, ERenderTargetLoadAction InOverrideLoadAction)
+	FTextureRenderTargetBinding(FRDGTextureRef InTexture, bool bInNeverClear)
 		: Texture(InTexture)
 		, ArraySlice(-1)
-		, OverrideLoadAction(InOverrideLoadAction)
+		, bNeverClear(bInNeverClear)
 	{}
 
-	FTextureRenderTargetBinding(FRDGTextureRef InTexture, int16 InArraySlice = -1, ERenderTargetLoadAction InOverrideLoadAction = ERenderTargetLoadAction::Num)
+	FTextureRenderTargetBinding(FRDGTextureRef InTexture, int16 InArraySlice = -1, bool bInNeverClear = false)
 		: Texture(InTexture)
 		, ArraySlice(InArraySlice)
-		, OverrideLoadAction(InOverrideLoadAction)
+		, bNeverClear(bInNeverClear)
 	{}
 };
 inline FRenderTargetBindingSlots GetRenderTargetBindings(ERenderTargetLoadAction ColorLoadAction, TArrayView<FTextureRenderTargetBinding> ColorTextures)
@@ -102,9 +102,9 @@ inline FRenderTargetBindingSlots GetRenderTargetBindings(ERenderTargetLoadAction
 	{
 		check(ColorTextures[Index].Texture);
 		BindingSlots[Index] = FRenderTargetBinding(ColorTextures[Index].Texture, ColorLoadAction, 0, ColorTextures[Index].ArraySlice);
-		if (ColorTextures[Index].OverrideLoadAction != ERenderTargetLoadAction::Num)
+		if (ColorLoadAction == ERenderTargetLoadAction::EClear && ColorTextures[Index].bNeverClear)
 		{
-			BindingSlots[Index].SetLoadAction(ColorTextures[Index].OverrideLoadAction);
+			BindingSlots[Index].SetLoadAction(ERenderTargetLoadAction::ELoad);
 		}
 	}
 	return BindingSlots;
diff --git a/Engine/Source/Runtime/Renderer/Private/Strata/Strata.cpp b/Engine/Source/Runtime/Renderer/Private/Strata/Strata.cpp
--- a/Engine/Source/Runtime/Renderer/Private/Strata/Strata.cpp
+++ b/Engine/Source/Runtime/Renderer/Private/Strata/Strata.cpp
@@ -66,6 +66,12 @@ static TAutoConsoleVariable<int32> CVarMaterialRoughDiffuse(
 	TEXT("Enable rough diffuse material."),
 	ECVF_ReadOnly | ECVF_RenderThreadSafe);
 
+// STRATA_TODO we keep this for now and can remove it once battletested.
+static TAutoConsoleVariable<int32> CVarClearDuringCategorization(
+	TEXT("r.strata.ClearDuringCategorization"),
+	1,
+	TEXT("TEST."),
+	ECVF_RenderThreadSafe);
 
 
 IMPLEMENT_GLOBAL_SHADER_PARAMETER_STRUCT(FStrataGlobalUniformParameters, "Strata");
@@ -196,7 +202,7 @@ void InitialiseStrataFrameSceneData(FSceneRenderer& SceneRenderer, FRDGBuilder&
 
 		// SSS texture
 		{
-			StrataSceneData.SSSTexture = GraphBuilder.CreateTexture(FRDGTextureDesc::Create2D(SceneTextureExtent, PF_R32G32_UINT, FClearValueBinding::Black, TexCreate_ShaderResource | TexCreate_UAV), TEXT("Strata.SSSTexture"));
+			StrataSceneData.SSSTexture = GraphBuilder.CreateTexture(FRDGTextureDesc::Create2D(SceneTextureExtent, PF_R32G32_UINT, FClearValueBinding::Black, TexCreate_DisableDCC | TexCreate_NoFastClear | TexCreate_ShaderResource | TexCreate_UAV), TEXT("Strata.SSSTexture"));
 			StrataSceneData.SSSTextureUAV = GraphBuilder.CreateUAV(StrataSceneData.SSSTexture);
 		}
 	}
@@ -427,8 +433,9 @@ class FStrataMaterialTileClassificationPassCS : public FGlobalShader
 	DECLARE_GLOBAL_SHADER(FStrataMaterialTileClassificationPassCS);
 	SHADER_USE_PARAMETER_STRUCT(FStrataMaterialTileClassificationPassCS, FGlobalShader);
 
+	class FStrataClearDuringCategorization : SHADER_PERMUTATION_BOOL("PERMUTATION_STRATA_CLEAR_DURING_CATEGORIZATION");
 	class FWaveOps : SHADER_PERMUTATION_BOOL("PERMUTATION_WAVE_OPS");
-	using FPermutationDomain = TShaderPermutationDomain<FWaveOps>;
+	using FPermutationDomain = TShaderPermutationDomain<FWaveOps, FStrataClearDuringCategorization>;
 
 	BEGIN_SHADER_PARAMETER_STRUCT(FParameters, )
 		SHADER_PARAMETER_STRUCT_REF(FViewUniformShaderParameters, ViewUniformBuffer)
@@ -442,6 +449,7 @@ class FStrataMaterialTileClassificationPassCS : public FGlobalShader
 		SHADER_PARAMETER_RDG_BUFFER_UAV(RWBuffer, SimpleTileListDataBuffer)
 		SHADER_PARAMETER_RDG_BUFFER_UAV(RWBuffer, ComplexTileIndirectDataBuffer)
 		SHADER_PARAMETER_RDG_BUFFER_UAV(RWBuffer, ComplexTileListDataBuffer)
+		SHADER_PARAMETER_RDG_TEXTURE_UAV(RWTexture2D<uint2>, SSSTextureUAV)
 	END_SHADER_PARAMETER_STRUCT()
 
 	static bool ShouldCompilePermutation(const FGlobalShaderPermutationParameters& Parameters)
@@ -648,18 +656,20 @@ void ApprendStrataMRTs(FSceneRenderer& SceneRenderer, uint32& RenderTargetCount,
 	{
 		// If this function changes, update Strata::SetBasePassRenderTargetOutputFormat()
 		 
-		// Add 2 uint for Strata fast path
-		auto AddStrataOutputTarget = [&](int16 StrataMaterialArraySlice)
+		// Add 2 uint for Strata fast path. 
+		// - We must clear the first uint to 0 to identify pixels that have not been written to.
+		// - We must never clear the second uint, it will only be written/read if needed.
+		auto AddStrataOutputTarget = [&](int16 StrataMaterialArraySlice, bool bNeverClear = false)
 		{
-			RenderTargets[RenderTargetCount] = FTextureRenderTargetBinding(SceneRenderer.Scene->StrataSceneData.MaterialTextureArrayMRTs, StrataMaterialArraySlice, ERenderTargetLoadAction::ELoad);
+			RenderTargets[RenderTargetCount] = FTextureRenderTargetBinding(SceneRenderer.Scene->StrataSceneData.MaterialTextureArrayMRTs, StrataMaterialArraySlice, bNeverClear);
 			RenderTargetCount++;
 		};
 		AddStrataOutputTarget(0);
-		AddStrataOutputTarget(1);
+		AddStrataOutputTarget(1, true); // Never clear that target
 
-		// Add another MRT for Strata top layer information
+		// Add another MRT for Strata top layer information. We want to follow the usual clear process which can leverage fast clear.
 		{
-			RenderTargets[RenderTargetCount] = FTextureRenderTargetBinding(SceneRenderer.Scene->StrataSceneData.TopLayerTexture, ERenderTargetLoadAction::ELoad);
+			RenderTargets[RenderTargetCount] = FTextureRenderTargetBinding(SceneRenderer.Scene->StrataSceneData.TopLayerTexture);
 			RenderTargetCount++;
 		};
 	}
@@ -759,6 +769,7 @@ void AddStrataMaterialClassificationPass(FRDGBuilder& GraphBuilder, const FMinim
 		#endif
 
 			FStrataMaterialTileClassificationPassCS::FPermutationDomain PermutationVector;
+			PermutationVector.Set< FStrataMaterialTileClassificationPassCS::FStrataClearDuringCategorization >(CVarClearDuringCategorization.GetValueOnRenderThread() > 0);
 			PermutationVector.Set< FStrataMaterialTileClassificationPassCS::FWaveOps >(bWaveOps);
 			TShaderMapRef<FStrataMaterialTileClassificationPassCS> ComputeShader(View.ShaderMap, PermutationVector);
 			FStrataMaterialTileClassificationPassCS::FParameters* PassParameters = GraphBuilder.AllocParameters<FStrataMaterialTileClassificationPassCS::FParameters>();
@@ -773,6 +784,7 @@ void AddStrataMaterialClassificationPass(FRDGBuilder& GraphBuilder, const FMinim
 			PassParameters->SimpleTileIndirectDataBuffer = View.StrataSceneData->ClassificationTileIndirectBufferUAV[EStrataTileMaterialType::ESimple];
 			PassParameters->ComplexTileListDataBuffer = View.StrataSceneData->ClassificationTileListBufferUAV[EStrataTileMaterialType::EComplex];
 			PassParameters->ComplexTileIndirectDataBuffer = View.StrataSceneData->ClassificationTileIndirectBufferUAV[EStrataTileMaterialType::EComplex];
+			PassParameters->SSSTextureUAV = View.StrataSceneData->SSSTextureUAV;
 
 			const uint32 GroupSize = 8;
 			FComputeShaderUtils::AddPass(
@@ -793,6 +805,11 @@ static void AddStrataClearMaterialBufferPass(
 	uint32 MaxBytesPerPixel,
 	FIntPoint TiledViewBufferResolution)
 {
+	if (CVarClearDuringCategorization.GetValueOnRenderThread() > 0)
+	{
+		return;
+	}
+
 	TShaderMapRef<FStrataClearMaterialBufferCS> ComputeShader(GetGlobalShaderMap(GMaxRHIFeatureLevel));
 	FStrataClearMaterialBufferCS::FParameters* PassParameters = GraphBuilder.AllocParameters<FStrataClearMaterialBufferCS::FParameters>();
 	PassParameters->MaterialTextureArrayUAV = MaterialTextureArrayUAV;
