Change RendererOnScreenNotification to broadcast on render thread and buffer results to game thread.
This should make the use a bit more intuitive.

diff --git a/Engine/Source/Runtime/Renderer/Private/RendererOnScreenNotification.h b/Engine/Source/Runtime/Renderer/Private/RendererOnScreenNotification.h
--- a/Engine/Source/Runtime/Renderer/Private/RendererOnScreenNotification.h
+++ b/Engine/Source/Runtime/Renderer/Private/RendererOnScreenNotification.h
@@ -6,10 +6,9 @@
 #include "Misc/LazySingleton.h"
 
 /** 
- * Class to wrap FCoreDelegates::FGetOnScreenMessagesDelegate for access from any thread. 
- * This avoids race conditions in registration/unregistration that would happen if using FCoreDelegates::FGetOnScreenMessagesDelegate directly from render thread.
- * Note that the Broadcast() still happens on game thread, so care needs to be taken with how data is accessed there.
- * If that becomes an issue we could change so that the proxy delegate Broadcasts on render thread and buffers to game thread.
+ * Class to wrap FCoreDelegates::FGetOnScreenMessagesDelegate for use on the render thread. 
+ * This avoids race conditions in registration/unregistration that would happen if using FCoreDelegates::FGetOnScreenMessagesDelegate not on the game thread.
+ * Registration/Unregistration can happen on any thread and a proxy delegate calls Broadcast() on the render thread.
  */
 class FRendererOnScreenNotification
 {
@@ -33,9 +32,8 @@ public:
 	}
 
 	/** 
-	 * Relay to AddLambda() of underlying delegate. 
+	 * Relay to AddLambda() of underlying proxy delegate. 
 	 * Thia takes a lock so that it can be called from any thread.
-	 * The lambda will be called from the game thread!
 	 */
 	template<typename FunctorType, typename... VarTypes>
 	FDelegateHandle AddLambda(FunctorType&& InFunctor, VarTypes... Vars)
@@ -45,7 +43,7 @@ public:
 	}
 
 	/** 
-	 * Relay to Remove() of underlying delegate. 
+	 * Relay to Remove() of underlying proxy delegate. 
 	 * Thia takes a lock so that it can be called from any thread.
 	 */
 	bool Remove(FDelegateHandle InHandle)
@@ -54,6 +52,25 @@ public:
 		return ProxyDelegate.Remove(InHandle);
 	}
 
+	/** 
+	 * Calls Broadcast on the proxy delegate.
+	 * The results are buffered for later collection by the base game thread delegate.
+	 */
+	void Broadcast()
+	{
+		// Broadcast using temporary container to reduce any lock contention.
+		MessagesTmp.Empty(Messages.Num());
+		
+		{
+			FScopeLock Lock(&DelgateCS);
+			ProxyDelegate.Broadcast(MessagesTmp);
+		}
+		{
+			FScopeLock Lock(&MessageCS);
+			Messages = MoveTemp(MessagesTmp);
+		}
+	}
+
 private:
 	friend FLazySingleton;
 
@@ -61,8 +78,8 @@ private:
 	{
 		BaseDelegateHandle = FCoreDelegates::OnGetOnScreenMessages.AddLambda([this](FCoreDelegates::FSeverityMessageMap& OutMessages)
 		{
-			FScopeLock Lock(&DelgateCS);
-			ProxyDelegate.Broadcast(OutMessages);
+			FScopeLock Lock(&MessageCS);
+			OutMessages.Append(Messages);
 		});
 	}
 
@@ -73,6 +90,9 @@ private:
 
 private:
 	FCriticalSection DelgateCS;
+	FCriticalSection MessageCS;
 	FCoreDelegates::FGetOnScreenMessagesDelegate ProxyDelegate;
 	FDelegateHandle BaseDelegateHandle;
+	FCoreDelegates::FSeverityMessageMap Messages;
+	FCoreDelegates::FSeverityMessageMap MessagesTmp;
 };
diff --git a/Engine/Source/Runtime/Renderer/Private/SceneRendering.cpp b/Engine/Source/Runtime/Renderer/Private/SceneRendering.cpp
--- a/Engine/Source/Runtime/Renderer/Private/SceneRendering.cpp
+++ b/Engine/Source/Runtime/Renderer/Private/SceneRendering.cpp
@@ -73,6 +73,7 @@
 #include "Lumen/Lumen.h"
 #include "Nanite/Nanite.h"
 #include "DistanceFieldLightingShared.h"
+#include "RendererOnScreenNotification.h"
 #include "Rendering/NaniteCoarseMeshStreamingManager.h"
 #include "Rendering/NaniteStreamingManager.h"
 
@@ -4137,6 +4138,11 @@ static void RenderViewFamily_RenderThread(FRHICommandListImmediate& RHICmdList,
 		}
 #endif
 
+#if !UE_BUILD_SHIPPING
+		// Update on screen notifications.
+		FRendererOnScreenNotification::Get().Broadcast();
+#endif
+
 		GRenderTargetPool.SetEventRecordingActive(false);
 	}
 
