Use RendererOnScreenNotification for VirtualTextureSystem on screen messages.

diff --git a/Engine/Source/Runtime/Renderer/Private/VT/VirtualTextureSystem.cpp b/Engine/Source/Runtime/Renderer/Private/VT/VirtualTextureSystem.cpp
--- a/Engine/Source/Runtime/Renderer/Private/VT/VirtualTextureSystem.cpp
+++ b/Engine/Source/Runtime/Renderer/Private/VT/VirtualTextureSystem.cpp
@@ -10,6 +10,7 @@
 #include "PostProcess/SceneRenderTargets.h"
 #include "ProfilingDebugging/CsvProfiler.h"
 #include "ProfilingDebugging/CpuProfilerTrace.h"
+#include "RendererOnScreenNotification.h"
 #include "ScenePrivate.h"
 #include "SceneUtils.h"
 #include "Stats/Stats.h"
@@ -319,7 +320,7 @@ FVirtualTextureSystem::FVirtualTextureSystem()
 #endif
 {
 #if !UE_BUILD_SHIPPING
-	FCoreDelegates::OnGetOnScreenMessages.AddRaw(this, &FVirtualTextureSystem::GetOnScreenMessages);
+	OnScreenMessageDelegateHandle = FRendererOnScreenNotification::Get().AddLambda([this](FCoreDelegates::FSeverityMessageMap& OutMessages) { GetOnScreenMessages(OutMessages); });
 	DrawResidencyHudDelegateHandle = UDebugDrawService::Register(TEXT("VirtualTextureResidency"), FDebugDrawDelegate::CreateRaw(this, &FVirtualTextureSystem::DrawResidencyHud));
 #endif
 }
@@ -327,7 +328,7 @@ FVirtualTextureSystem::FVirtualTextureSystem()
 FVirtualTextureSystem::~FVirtualTextureSystem()
 {
 #if !UE_BUILD_SHIPPING
-	FCoreDelegates::OnGetOnScreenMessages.RemoveAll(this);
+	FRendererOnScreenNotification::Get().Remove(OnScreenMessageDelegateHandle);
 	UDebugDrawService::Unregister(DrawResidencyHudDelegateHandle);
 #endif
 
@@ -1600,7 +1601,6 @@ void FVirtualTextureSystem::Update(FRDGBuilder& GraphBuilder, ERHIFeatureLevel::
 	UpdateResidencyTracking();
 
 #if !UE_BUILD_SHIPPING
-	UpdateNotifications();
 	UpdateCsvStats();
 #endif
 
@@ -2621,19 +2621,6 @@ float FVirtualTextureSystem::GetGlobalMipBias() const
 #if !UE_BUILD_SHIPPING
 
 void FVirtualTextureSystem::GetOnScreenMessages(FCoreDelegates::FSeverityMessageMap& OutMessages)
-{
-	FScopeLock Locker(&OnScreenMessageLock);
-	OutMessages.Append(OnScreenMessages);
-}
-
-void FVirtualTextureSystem::UpdateNotifications()
-{
-	FScopeLock Locker(&OnScreenMessageLock);
-	OnScreenMessages.Reset();
-	UpdateResidencyNotifications();
-}
-
-void FVirtualTextureSystem::UpdateResidencyNotifications()
 {
 	if (CVarVTResidencyNotify.GetValueOnRenderThread() == 0)
 	{
@@ -2648,7 +2635,7 @@ void FVirtualTextureSystem::UpdateResidencyNotifications()
 			FString const& FormatString = PhysicalSpace->GetFormatString();
 			const float MipBias = PhysicalSpace->GetResidencyMipMapBias();
 
-			OnScreenMessages.Add(
+			OutMessages.Add(
 				FCoreDelegates::EOnScreenMessageSeverity::Warning,
 				FText::Format(LOCTEXT("VTOversubscribed", "VT Pool [{0}] is oversubscribed. Setting MipBias {1}"), FText::FromString(FormatString), FText::AsNumber(MipBias)));
 		}
diff --git a/Engine/Source/Runtime/Renderer/Private/VT/VirtualTextureSystem.h b/Engine/Source/Runtime/Renderer/Private/VT/VirtualTextureSystem.h
--- a/Engine/Source/Runtime/Renderer/Private/VT/VirtualTextureSystem.h
+++ b/Engine/Source/Runtime/Renderer/Private/VT/VirtualTextureSystem.h
@@ -177,14 +177,9 @@ private:
 
 #if !UE_BUILD_SHIPPING
 	void GetOnScreenMessages(FCoreDelegates::FSeverityMessageMap& OutMessages);
-	FCoreDelegates::FSeverityMessageMap OnScreenMessages;
-	FCriticalSection OnScreenMessageLock;
-
-	void UpdateNotifications();
-	void UpdateResidencyNotifications();
-	void UpdateCsvStats();
-
+	FDelegateHandle OnScreenMessageDelegateHandle;
 	void DrawResidencyHud(class UCanvas*, class APlayerController*);
 	FDelegateHandle	DrawResidencyHudDelegateHandle;
+	void UpdateCsvStats();
 #endif
 };
