Removes GTemporalUpscaler since the interface get actually set by custom screen percentage or view family extension

diff --git a/Engine/Source/Runtime/Renderer/Private/PostProcess/TemporalAA.cpp b/Engine/Source/Runtime/Renderer/Private/PostProcess/TemporalAA.cpp
--- a/Engine/Source/Runtime/Renderer/Private/PostProcess/TemporalAA.cpp
+++ b/Engine/Source/Runtime/Renderer/Private/PostProcess/TemporalAA.cpp
@@ -888,8 +888,6 @@ ITemporalUpscaler::FOutputs AddTemporalSuperResolutionPasses(
 	const FViewInfo& View,
 	const ITemporalUpscaler::FPassInputs& PassInputs);
 
-const ITemporalUpscaler* GTemporalUpscaler = nullptr;
-
 class FDefaultTemporalUpscaler : public ITemporalUpscaler
 {
 public:
diff --git a/Engine/Source/Runtime/Renderer/Private/PostProcess/TemporalAA.h b/Engine/Source/Runtime/Renderer/Private/PostProcess/TemporalAA.h
--- a/Engine/Source/Runtime/Renderer/Private/PostProcess/TemporalAA.h
+++ b/Engine/Source/Runtime/Renderer/Private/PostProcess/TemporalAA.h
@@ -217,6 +217,3 @@ public:
 
 	static EMainTAAPassConfig GetMainTAAPassConfig(const FViewInfo& View);
 }; 
-
-extern RENDERER_API const ITemporalUpscaler* GTemporalUpscaler;
-
diff --git a/Engine/Source/Runtime/Renderer/Private/Renderer.cpp b/Engine/Source/Runtime/Renderer/Private/Renderer.cpp
--- a/Engine/Source/Runtime/Renderer/Private/Renderer.cpp
+++ b/Engine/Source/Runtime/Renderer/Private/Renderer.cpp
@@ -55,7 +55,6 @@ FAutoConsoleVariableRef CVarFlushRenderTargetsOnWorldCleanup(TEXT("r.bFlushRende
 void FRendererModule::StartupModule()
 {
 	GScreenSpaceDenoiser = IScreenSpaceDenoiser::GetDefaultDenoiser();
-	GTemporalUpscaler = ITemporalUpscaler::GetDefaultTemporalUpscaler();
 
 	FRendererOnScreenNotification::Get();
 	FVirtualTextureSystem::Initialize();
