Fix crash with hair strands in forward rendering.

diff --git a/Engine/Source/Runtime/Renderer/Private/HairStrands/HairStrandsVisibility.cpp b/Engine/Source/Runtime/Renderer/Private/HairStrands/HairStrandsVisibility.cpp
--- a/Engine/Source/Runtime/Renderer/Private/HairStrands/HairStrandsVisibility.cpp
+++ b/Engine/Source/Runtime/Renderer/Private/HairStrands/HairStrandsVisibility.cpp
@@ -957,7 +957,6 @@ class FHairVelocityCS: public FGlobalShader
 		SHADER_PARAMETER_RDG_TEXTURE_UAV(RWTexture2D, OutVelocityTexture)
 		SHADER_PARAMETER_RDG_TEXTURE_UAV(RWTexture2D, OutResolveMaskTexture)
 		SHADER_PARAMETER_STRUCT_REF(FViewUniformShaderParameters, ViewUniformBuffer)
-		SHADER_PARAMETER_RDG_UNIFORM_BUFFER(FSceneTextureUniformParameters, SceneTexturesStruct)
 
 		SHADER_PARAMETER(FIntPoint, TileCountXY)
 		SHADER_PARAMETER(uint32, TileSize)
@@ -1024,7 +1023,6 @@ static void AddHairVelocityPass(
 
 	FHairVelocityCS::FParameters* PassParameters = GraphBuilder.AllocParameters<FHairVelocityCS::FParameters>();
 	PassParameters->bNeedClear = bNeedClear ? 1u : 0u;
-	PassParameters->SceneTexturesStruct = CreateSceneTextureUniformBuffer(GraphBuilder, View.FeatureLevel);
 	PassParameters->ViewUniformBuffer = View.ViewUniformBuffer;
 	PassParameters->VelocityThreshold = GetHairFastResolveVelocityThreshold(Resolution);
 	PassParameters->CoverageThreshold = GetHairWriteVelocityCoverageThreshold();
@@ -2658,7 +2656,8 @@ static FRDGTextureRef AddHairLightChannelMaskPass(
 	const FHairStrandsTiles& TileData,
 	const FRDGTextureRef& CoverageTexture,
 	const FRDGTextureRef& HairSampleOffset,
-	const FRDGBufferRef& HairSampleData)
+	const FRDGBufferRef& HairSampleData,
+	const FRDGTextureRef& SceneDepthTexture)
 {
 	check(IsRayTracingEnabled());
 	FRDGTextureRef OutLightChannelMask = CreateLigthtChannelMaskTexture(GraphBuilder, View.ViewRect.Size());
@@ -2674,7 +2673,7 @@ static FRDGTextureRef AddHairLightChannelMaskPass(
 		nullptr,
 		nullptr,
 		nullptr,
-		nullptr,
+		SceneDepthTexture,
 		OutLightChannelMask);
 	return OutLightChannelMask;
 }
@@ -4030,7 +4029,8 @@ void RenderHairStrandsVisibilityBuffer(
 					VisibilityData.TileData,
 					CoverageTexture,
 					CompactNodeIndex,
-					CompactNodeData);
+					CompactNodeData,
+					SceneDepthTexture);
 			}
 		#endif
 		}
