Fix calculation for FTexturePagePool::GetNumLockedPages().
Needs to take into account the reserved page in the pool.
Fixes check() firing on VT physical pool destruction.

diff --git a/Engine/Source/Runtime/Renderer/Private/VT/TexturePagePool.h b/Engine/Source/Runtime/Renderer/Private/VT/TexturePagePool.h
--- a/Engine/Source/Runtime/Renderer/Private/VT/TexturePagePool.h
+++ b/Engine/Source/Runtime/Renderer/Private/VT/TexturePagePool.h
@@ -31,7 +31,7 @@ public:
 	FCriticalSection& GetLock() { return CriticalSection; }
 
 	uint32 GetNumPages() const { return NumPages; }
-	uint32 GetNumLockedPages() const { return GetNumPages() - FreeHeap.Num(); }
+	uint32 GetNumLockedPages() const { return GetNumPages() - FreeHeap.Num() - NumReservedPages; }
 	uint32 GetNumMappedPages() const { return NumPagesMapped; }
 	uint32 GetNumAllocatedPages() const { return NumPagesAllocated; }
 
