Replaces r.Bloom.HalfResConvolution with r.Bloom.ScreenPercentage

diff --git a/Engine/Config/BaseScalability.ini b/Engine/Config/BaseScalability.ini
--- a/Engine/Config/BaseScalability.ini
+++ b/Engine/Config/BaseScalability.ini
@@ -258,7 +258,7 @@ r.LensFlareQuality=0
 r.SceneColorFringeQuality=0
 r.EyeAdaptationQuality=2
 r.BloomQuality=4
-r.Bloom.HalfResolutionFFT=2
+r.Bloom.ScreenPercentage=25.000
 r.FastBlurThreshold=0
 r.Upscale.Quality=1
 r.Tonemapper.GrainQuantization=0
@@ -278,7 +278,7 @@ r.LensFlareQuality=0
 r.SceneColorFringeQuality=0
 r.EyeAdaptationQuality=2
 r.BloomQuality=4
-r.Bloom.HalfResolutionFFT=2
+r.Bloom.ScreenPercentage=35.355
 r.FastBlurThreshold=2
 r.Upscale.Quality=2
 r.Tonemapper.GrainQuantization=0
@@ -310,7 +310,7 @@ r.LensFlareQuality=2
 r.SceneColorFringeQuality=1
 r.EyeAdaptationQuality=2
 r.BloomQuality=5
-r.Bloom.HalfResolutionFFT=1
+r.Bloom.ScreenPercentage=50.000
 r.FastBlurThreshold=3
 r.Upscale.Quality=2
 r.Tonemapper.GrainQuantization=1
@@ -344,7 +344,7 @@ r.LensFlareQuality=2
 r.SceneColorFringeQuality=1
 r.EyeAdaptationQuality=2
 r.BloomQuality=5
-r.Bloom.HalfResolutionFFT=1
+r.Bloom.ScreenPercentage=70.711
 r.FastBlurThreshold=100
 r.Upscale.Quality=3
 r.Tonemapper.GrainQuantization=1
@@ -381,7 +381,7 @@ r.LensFlareQuality=3
 r.SceneColorFringeQuality=1
 r.EyeAdaptationQuality=2
 r.BloomQuality=5
-r.Bloom.HalfResolutionFFT=0
+r.Bloom.ScreenPercentage=100
 r.FastBlurThreshold=100
 r.Upscale.Quality=3
 r.Tonemapper.GrainQuantization=1
diff --git a/Engine/Shaders/Private/PostProcessDownsample.usf b/Engine/Shaders/Private/PostProcessDownsample.usf
--- a/Engine/Shaders/Private/PostProcessDownsample.usf
+++ b/Engine/Shaders/Private/PostProcessDownsample.usf
@@ -68,6 +68,8 @@ float4 DownsampleCommon(float2 UV)
 }
 
 // pixel shader entry point
+#if PIXELSHADER
+
 void MainPS(float4 SvPosition : SV_POSITION, out float4 OutColor : SV_Target0)
 {
 	const float2 UV = SvPosition.xy * Output_ExtentInverse;
@@ -75,20 +77,22 @@ void MainPS(float4 SvPosition : SV_POSITION, out float4 OutColor : SV_Target0)
 	OutColor = DownsampleCommon(UV);
 }
 
-#if COMPUTESHADER
+#elif COMPUTESHADER
+
+FScreenTransform DispatchThreadIdToInputUV;
 
 RWTexture2D<float4> OutComputeTexture;
 
 [numthreads(THREADGROUP_SIZEX, THREADGROUP_SIZEY, 1)]
 void MainCS(uint2 DispatchThreadId : SV_DispatchThreadID)
 {
-	uint2 iPixelPos = DispatchThreadId + Output_ViewportMin;
-	float2 PixelPos = float2(iPixelPos);
-	float2 UV = (PixelPos + 0.5) * Output_ExtentInverse;
+	float2 InputUV = ApplyScreenTransform(float2(DispatchThreadId), DispatchThreadIdToInputUV);
+	uint2 OutputPixelPos = DispatchThreadId + Output_ViewportMin;
 
-	if (!IsComputeUVOutOfBounds(UV))
+	if (all(OutputPixelPos < Output_ViewportMax))
 	{
-		OutComputeTexture[PixelPos] = DownsampleCommon(UV);
+		OutComputeTexture[OutputPixelPos] = DownsampleCommon(InputUV);
 	}
 }
+
 #endif 
diff --git a/Engine/Source/Runtime/Renderer/Private/PostProcess/PostProcessBloomSetup.cpp b/Engine/Source/Runtime/Renderer/Private/PostProcess/PostProcessBloomSetup.cpp
--- a/Engine/Source/Runtime/Renderer/Private/PostProcess/PostProcessBloomSetup.cpp
+++ b/Engine/Source/Runtime/Renderer/Private/PostProcess/PostProcessBloomSetup.cpp
@@ -11,7 +11,7 @@ const int32 GBloomSetupComputeTileSizeX = 8;
 const int32 GBloomSetupComputeTileSizeY = 8;
 
 TAutoConsoleVariable<float> CVarBloomCross(
-	TEXT("r.Bloom.Cross"),
+	TEXT("r.GaussianBloom.Cross"),
 	0.0f,
 	TEXT("Experimental feature to give bloom kernel a more bright center sample (values between 1 and 3 work without causing aliasing)\n")
 	TEXT("Existing bloom get lowered to match the same brightness\n")
diff --git a/Engine/Source/Runtime/Renderer/Private/PostProcess/PostProcessDownsample.cpp b/Engine/Source/Runtime/Renderer/Private/PostProcess/PostProcessDownsample.cpp
--- a/Engine/Source/Runtime/Renderer/Private/PostProcess/PostProcessDownsample.cpp
+++ b/Engine/Source/Runtime/Renderer/Private/PostProcess/PostProcessDownsample.cpp
@@ -9,14 +9,6 @@ namespace
 const int32 GDownsampleTileSizeX = 8;
 const int32 GDownsampleTileSizeY = 8;
 
-TAutoConsoleVariable<int32> CVarDownsampleQuality(
-	TEXT("r.Downsample.Quality"),
-	1,
-	TEXT("Defines the quality in which the Downsample passes. we might add more quality levels later.\n")
-	TEXT(" 0: low quality\n")
-	TEXT(">0: high quality (default: 1)\n"),
-	ECVF_Scalability | ECVF_RenderThreadSafe);
-
 BEGIN_SHADER_PARAMETER_STRUCT(FDownsampleParameters, )
 	SHADER_PARAMETER_STRUCT_REF(FViewUniformShaderParameters, ViewUniformBuffer)
 	SHADER_PARAMETER_STRUCT(FScreenPassTextureViewportParameters, Input)
@@ -76,6 +68,7 @@ public:
 
 	BEGIN_SHADER_PARAMETER_STRUCT(FParameters, )
 		SHADER_PARAMETER_STRUCT_INCLUDE(FDownsampleParameters, Common)
+		SHADER_PARAMETER(FScreenTransform, DispatchThreadIdToInputUV)
 		SHADER_PARAMETER_RDG_TEXTURE_UAV(RWTexture2D<float4>, OutComputeTexture)
 	END_SHADER_PARAMETER_STRUCT()
 
@@ -95,13 +88,6 @@ public:
 IMPLEMENT_GLOBAL_SHADER(FDownsampleCS, "/Engine/Private/PostProcessDownsample.usf", "MainCS", SF_Compute);
 } //! namespace
 
-EDownsampleQuality GetDownsampleQuality()
-{
-	const int32 DownsampleQuality = FMath::Clamp(CVarDownsampleQuality.GetValueOnRenderThread(), 0, 1);
-
-	return static_cast<EDownsampleQuality>(DownsampleQuality);
-}
-
 FScreenPassTexture AddDownsamplePass(
 	FRDGBuilder& GraphBuilder,
 	const FViewInfo& View,
@@ -147,47 +133,62 @@ FScreenPassTexture AddDownsamplePass(
 		Output.LoadAction = ERenderTargetLoadAction::ENoAction;
 	}
 
-	FDownsamplePermutationDomain PermutationVector;
-	PermutationVector.Set<FDownsampleQualityDimension>(Inputs.Quality);
-
-	const FScreenPassTextureViewport SceneColorViewport(Inputs.SceneColor);
-	const FScreenPassTextureViewport OutputViewport(Output);
-
 	if (bIsComputePass)
 	{
-		FDownsampleCS::FParameters* PassParameters = GraphBuilder.AllocParameters<FDownsampleCS::FParameters>();
-		PassParameters->Common = GetDownsampleParameters(View, Output, Inputs.SceneColor, Inputs.Quality);
-		PassParameters->OutComputeTexture = GraphBuilder.CreateUAV(Output.Texture);
-
-		TShaderMapRef<FDownsampleCS> ComputeShader(View.ShaderMap, PermutationVector);
-
-		FComputeShaderUtils::AddPass(
-			GraphBuilder,
-			RDG_EVENT_NAME("Downsample.%s %dx%d (CS)", Inputs.Name, Inputs.SceneColor.ViewRect.Width(), Inputs.SceneColor.ViewRect.Height()),
-			ComputeShader,
-			PassParameters,
-			FComputeShaderUtils::GetGroupCount(OutputViewport.Rect.Size(), FIntPoint(GDownsampleTileSizeX, GDownsampleTileSizeY)));
+		AddDownsampleComputePass(GraphBuilder, View, Inputs.SceneColor, Output, Inputs.Quality, bIsComputePass ? ERDGPassFlags::Compute : ERDGPassFlags::Raster);
 	}
 	else
 	{
+		FDownsamplePermutationDomain PermutationVector;
+		PermutationVector.Set<FDownsampleQualityDimension>(Inputs.Quality);
+
 		FDownsamplePS::FParameters* PassParameters = GraphBuilder.AllocParameters<FDownsamplePS::FParameters>();
 		PassParameters->Common = GetDownsampleParameters(View, Output, Inputs.SceneColor, Inputs.Quality);
 		PassParameters->RenderTargets[0] = Output.GetRenderTargetBinding();
 
 		TShaderMapRef<FDownsamplePS> PixelShader(View.ShaderMap, PermutationVector);
-
 		FPixelShaderUtils::AddFullscreenPass(
 			GraphBuilder,
 			View.ShaderMap,
-			RDG_EVENT_NAME("Downsample.%s %dx%d (PS)", Inputs.Name, Inputs.SceneColor.ViewRect.Width(), Inputs.SceneColor.ViewRect.Height()),
+			RDG_EVENT_NAME("Downsample(%s Quality=%s PS) %dx%d -> %dx%d",
+				Output.Texture->Name,
+				Inputs.Quality == EDownsampleQuality::High ? TEXT("High") : TEXT("Bilinear"),
+				Inputs.SceneColor.ViewRect.Width(), Inputs.SceneColor.ViewRect.Height(),
+				Output.ViewRect.Width(), Output.ViewRect.Height()),
 			PixelShader,
 			PassParameters,
-			OutputViewport.Rect);
+			Output.ViewRect);
 	}
 
 	return MoveTemp(Output);
 }
 
+void AddDownsampleComputePass(FRDGBuilder& GraphBuilder, const FViewInfo& View, FScreenPassTexture Input, FScreenPassTexture Output, EDownsampleQuality Quality, ERDGPassFlags PassFlags)
+{
+	check(PassFlags == ERDGPassFlags::Compute || PassFlags == ERDGPassFlags::AsyncCompute);
+
+	FDownsampleCS::FParameters* PassParameters = GraphBuilder.AllocParameters<FDownsampleCS::FParameters>();
+	PassParameters->Common = GetDownsampleParameters(View, Output, Input, Quality);
+	PassParameters->DispatchThreadIdToInputUV = ((FScreenTransform::Identity + 0.5f) / Output.ViewRect.Size()) * FScreenTransform::ChangeTextureBasisFromTo(FScreenPassTextureViewport(Input), FScreenTransform::ETextureBasis::ViewportUV, FScreenTransform::ETextureBasis::TextureUV);
+	PassParameters->OutComputeTexture = GraphBuilder.CreateUAV(Output.Texture);
+
+	FDownsamplePermutationDomain PermutationVector;
+	PermutationVector.Set<FDownsampleQualityDimension>(Quality);
+
+	TShaderMapRef<FDownsampleCS> ComputeShader(View.ShaderMap, PermutationVector);
+	FComputeShaderUtils::AddPass(
+		GraphBuilder,
+		RDG_EVENT_NAME("Downsample(%s Quality=%s CS) %dx%d -> %dx%d",
+			Output.Texture->Name,
+			Quality == EDownsampleQuality::High ? TEXT("High") : TEXT("Bilinear"),
+			Input.ViewRect.Width(), Input.ViewRect.Height(),
+			Output.ViewRect.Width(), Output.ViewRect.Height()),
+		PassFlags,
+		ComputeShader,
+		PassParameters,
+		FComputeShaderUtils::GetGroupCount(Output.ViewRect.Size(), FIntPoint(GDownsampleTileSizeX, GDownsampleTileSizeY)));
+}
+
 void FSceneDownsampleChain::Init(
 	FRDGBuilder& GraphBuilder,
 	const FViewInfo& View,
diff --git a/Engine/Source/Runtime/Renderer/Private/PostProcess/PostProcessDownsample.h b/Engine/Source/Runtime/Renderer/Private/PostProcess/PostProcessDownsample.h
--- a/Engine/Source/Runtime/Renderer/Private/PostProcess/PostProcessDownsample.h
+++ b/Engine/Source/Runtime/Renderer/Private/PostProcess/PostProcessDownsample.h
@@ -26,9 +26,6 @@ enum class EDownsampleQuality : uint8
 	MAX
 };
 
-// Returns the global downsample quality specified by the r.Downsample.Quality CVar.
-EDownsampleQuality GetDownsampleQuality();
-
 // The set of inputs needed to add a downsample pass to RDG.
 struct FDownsamplePassInputs
 {
@@ -55,6 +52,8 @@ struct FDownsamplePassInputs
 
 FScreenPassTexture AddDownsamplePass(FRDGBuilder& GraphBuilder, const FViewInfo& View, const FDownsamplePassInputs& Inputs);
 
+void AddDownsampleComputePass(FRDGBuilder& GraphBuilder, const FViewInfo& View, FScreenPassTexture Input, FScreenPassTexture Output, EDownsampleQuality Quality, ERDGPassFlags PassFlags);
+
 class FSceneDownsampleChain
 {
 public:
diff --git a/Engine/Source/Runtime/Renderer/Private/PostProcess/PostProcessFFTBloom.cpp b/Engine/Source/Runtime/Renderer/Private/PostProcess/PostProcessFFTBloom.cpp
--- a/Engine/Source/Runtime/Renderer/Private/PostProcess/PostProcessFFTBloom.cpp
+++ b/Engine/Source/Runtime/Renderer/Private/PostProcess/PostProcessFFTBloom.cpp
@@ -3,25 +3,24 @@
 #include "PostProcess/PostProcessFFTBloom.h"
 #include "PostProcess/PostProcessBloomSetup.h"
 #include "PostProcess/PostProcessTonemap.h"
+#include "PostProcess/PostProcessDownsample.h"
 #include "GPUFastFourierTransform.h"
 #include "RendererModule.h"
 #include "Rendering/Texture2DResource.h"
 
 namespace
 {
+
+TAutoConsoleVariable<float> CVarBloomScreenPercentage(
+	TEXT("r.Bloom.ScreenPercentage"), 100.0f,
+	TEXT("Controles the axis resolution of the FFT convolution for bloom.\n"),
+	ECVF_Scalability | ECVF_RenderThreadSafe);
+
 TAutoConsoleVariable<int32> CVarBloomCacheKernel(
 	TEXT("r.Bloom.CacheKernel"), 1,
 	TEXT("Whether to cache the kernel in spectral domain."),
 	ECVF_RenderThreadSafe);
 
-TAutoConsoleVariable<int32> CVarHalfResFFTBloom(
-	TEXT("r.Bloom.HalfResolutionFFT"), 0,
-	TEXT("Experimental half-resolution FFT Bloom convolution. \n")
-	TEXT(" 0: Standard full resolution convolution bloom;")
-	TEXT(" 1: Half-resolution convolution;\n")
-	TEXT(" 2: Quarter-resolution convolution.\n"),
-	ECVF_Scalability | ECVF_RenderThreadSafe);
-
 TAutoConsoleVariable<int32> CVarAsynComputeFFTBloom(
 	TEXT("r.Bloom.AsyncCompute"), 0,
 	TEXT("Whether to run FFT bloom on async compute.\n"),
@@ -198,14 +197,9 @@ IMPLEMENT_GLOBAL_SHADER(FBloomFinalizeApplyConstantsCS,     "/Engine/Private/Blo
 
 } //! namespace
 
-bool IsFFTBloomFullResolutionEnabled()
+float GetFFTBloomResolutionFraction()
 {
-	return CVarHalfResFFTBloom.GetValueOnRenderThread() == 0;
-}
-
-bool IsFFTBloomQuarterResolutionEnabled()
-{
-	return CVarHalfResFFTBloom.GetValueOnRenderThread() == 2;
+	return FMath::Clamp(CVarBloomScreenPercentage.GetValueOnRenderThread() / 100.0f, 0.1f, 1.0f);
 }
 
 bool IsFFTBloomEnabled(const FViewInfo& View)
@@ -269,14 +263,9 @@ FRDGTextureRef TransformKernelFFT(
 
 struct FFFTBloomIntermediates
 {
-	FRDGTextureRef InputTexture;
-
-	// The size of the input buffer.
-	FIntPoint InputBufferSize;
-
 	// The sub-domain of the Input/Output buffers 
 	// where the image lives, i.e. the region of interest
-	FIntRect  ImageRect;
+	FIntPoint ImageSize;
 
 	// Image space, padded by black for kernel and  rounded up to powers of two
 	// this defines the size of the FFT in each direction.
@@ -294,94 +283,6 @@ struct FFFTBloomIntermediates
 	ERDGPassFlags ComputePassFlags = ERDGPassFlags::Compute;
 };
 
-FFFTBloomIntermediates GetFFTBloomIntermediates(
-	const FViewInfo& View,
-	const FFFTBloomInputs& Inputs)
-{
-	FSceneViewState* ViewState = (FSceneViewState*)View.State;
-	check(ViewState);
-
-	const auto& PPSettings = View.FinalPostProcessSettings;
-
-	// The kernel parameters on the FinalPostProcess.
-
-	const float BloomConvolutionSize = PPSettings.BloomConvolutionSize;
-	const float KernelSupportScaleClamp = FMath::Clamp(PPSettings.BloomConvolutionBufferScale, 0.f, 1.f);
-
-	// Clip the Kernel support (i.e. bloom size) to 100% the screen width 
-	const float MaxBloomSize = 1.f;
-	const float KernelSupportScale = FMath::Clamp(BloomConvolutionSize, 0.f, MaxBloomSize);
-
-	// We padd by 1/2 the number of pixels the kernel needs in the x-direction
-	// so if the kernel is being applied on the edge of the image it will see padding and not periodicity
-	// NB:  If the kernel padding would force a transform buffer that is too big for group shared memory (> 4096)
-	//      we clamp it.  This could result in a wrap-around in the bloom (from one side of the screen to the other),
-	//      but since the amplitude of the bloom kernel tails is usually very small, this shouldn't be too bad.
-	auto KernelRadiusSupportFunctor = [KernelSupportScale, KernelSupportScaleClamp](const FIntPoint& Size) ->int32
-	{
-		float ClampedKernelSupportScale = (KernelSupportScaleClamp > 0) ? FMath::Min(KernelSupportScale, KernelSupportScaleClamp) : KernelSupportScale;
-		int32 FilterRadius = FMath::CeilToInt(0.5 * ClampedKernelSupportScale * Size.X);
-		const int32 MaxFFTSize = GPUFFT::MaxScanLineLength();
-		int32 MaxDim = FMath::Max(Size.X, Size.Y);
-		if (MaxDim + FilterRadius > MaxFFTSize && MaxDim < MaxFFTSize) FilterRadius = MaxFFTSize - MaxDim;
-
-		return FilterRadius;
-	};
-
-	const bool bHalfResolutionFFT = !IsFFTBloomFullResolutionEnabled();
-
-	FFFTBloomIntermediates Intermediates;
-	if (bHalfResolutionFFT)
-	{
-		Intermediates.InputTexture = Inputs.HalfResolutionTexture;
-		Intermediates.InputBufferSize = Inputs.HalfResolutionViewRect.Size();
-		Intermediates.ImageRect = Inputs.HalfResolutionViewRect;
-	}
-	else
-	{
-		Intermediates.InputTexture = Inputs.FullResolutionTexture;
-		Intermediates.InputBufferSize = Inputs.FullResolutionViewRect.Size();
-		Intermediates.ImageRect = Inputs.FullResolutionViewRect;
-	}
-
-	Intermediates.KernelSupportScale = KernelSupportScale;
-	Intermediates.KernelSupportScaleClamp = KernelSupportScaleClamp;
-
-	// The pre-filter boost parameters for bright pixels. Because the Convolution PP work in pre-exposure space, the min and max needs adjustment.
-	Intermediates.PreFilter = FVector(PPSettings.BloomConvolutionPreFilterMin, PPSettings.BloomConvolutionPreFilterMax, PPSettings.BloomConvolutionPreFilterMult);
-
-	// Capture the region of interest
-	const FIntPoint ImageSize = Intermediates.ImageRect.Size();
-
-	// The length of the a side of the square kernel image in pixels
-
-	const int32 KernelSize = FMath::CeilToInt(KernelSupportScale * FMath::Max(ImageSize.X, ImageSize.Y));
-
-	const int32 SpectralPadding = KernelRadiusSupportFunctor(ImageSize);
-
-	// The following are mathematically equivalent
-	// 1) Horizontal FFT / Vertical FFT / Filter / Vertical InvFFT / Horizontal InvFFT
-	// 2) Vertical FFT / Horizontal FFT / Filter / Horizontal InvFFT / Vertical InvFFT
-	// but we choose the one with the smallest intermediate buffer size
-
-	// The size of the input image plus padding that accounts for
-	// the width of the kernel.  The ImageRect is virtually padded
-	// with black to account for the gather action of the convolution.
-	FIntPoint PaddedImageSize = ImageSize + FIntPoint(SpectralPadding, SpectralPadding);
-	PaddedImageSize.X = FMath::Max(PaddedImageSize.X, KernelSize);
-	PaddedImageSize.Y = FMath::Max(PaddedImageSize.Y, KernelSize);
-
-	Intermediates.FrequencySize = FIntPoint(FMath::RoundUpToPowerOfTwo(PaddedImageSize.X), FMath::RoundUpToPowerOfTwo(PaddedImageSize.Y));
-
-	// Choose to do to transform in the direction that results in writing the least amount of data to main memory.
-
-	Intermediates.bDoHorizontalFirst = ((Intermediates.FrequencySize.Y * PaddedImageSize.X) > (Intermediates.FrequencySize.X * PaddedImageSize.Y));
-
-	Intermediates.ComputePassFlags = (GSupportsEfficientAsyncCompute && CVarAsynComputeFFTBloom.GetValueOnRenderThread() != 0) ? ERDGPassFlags::AsyncCompute : ERDGPassFlags::Compute;
-
-	return Intermediates;
-}
-
 void InitDomainAndGetKernel(
 	FRDGBuilder& GraphBuilder,
 	const FViewInfo& View,
@@ -427,7 +328,7 @@ void InitDomainAndGetKernel(
 			&& TransformDesc.Extent == PrevCachedSpectralKernel->Desc.Extent;
 
 		const bool bSameKernelSize = FMath::IsNearlyEqual(FFTKernel.Scale, BloomConvolutionSize, float(1.e-6) /*tol*/);
-		const bool bSameImageSize = (Intermediates.ImageRect.Size() == FFTKernel.ImageSize);
+		const bool bSameImageSize = (Intermediates.ImageSize == FFTKernel.ImageSize);
 		const bool bSameMipLevel = bSameTexture && FFTKernel.PhysicalMipLevel == BloomConvolutionTextureResource->GetCurrentMipCount();
 
 		if (bSameTexture && bSameSpectralBuffer && bSameKernelSize && bSameImageSize && bSameMipLevel)
@@ -475,7 +376,7 @@ void InitDomainAndGetKernel(
 	{
 		const int32 SurveyTileSize = 8;
 
-		float KernelSizeInDstPixels = FMath::Max(float(Intermediates.ImageRect.Width()) * Intermediates.KernelSupportScale, 1.0f);
+		float KernelSizeInDstPixels = FMath::Max(float(Intermediates.ImageSize.X) * Intermediates.KernelSupportScale, 1.0f);
 
 		// Diameter of a view texel in the kernel.
 		float ViewTexelDiameterInKernelTexels = FMath::Max(SpatialKernelTexture->Desc.Extent.X / KernelSizeInDstPixels, 1.0f);
@@ -683,7 +584,7 @@ void InitDomainAndGetKernel(
 
 			FBloomResizeKernelCS::FParameters* PassParameters = GraphBuilder.AllocParameters<FBloomResizeKernelCS::FParameters>();
 			PassParameters->DstExtent = Intermediates.FrequencySize;
-			PassParameters->ImageExtent = Intermediates.ImageRect.Size();
+			PassParameters->ImageExtent = Intermediates.ImageSize;
 			PassParameters->KernelSpatialTextureInvSize.X = 1.0f / float(SpatialKernelTexture->Desc.Extent.X);
 			PassParameters->KernelSpatialTextureInvSize.Y = 1.0f / float(SpatialKernelTexture->Desc.Extent.Y);
 			PassParameters->DstBufferExtent = Intermediates.FrequencySize;
@@ -717,7 +618,7 @@ void InitDomainAndGetKernel(
 	if (ViewState && CVarBloomCacheKernel.GetValueOnRenderThread() != 0)
 	{
 		ViewState->BloomFFTKernel.Scale = BloomConvolutionSize;
-		ViewState->BloomFFTKernel.ImageSize = Intermediates.ImageRect.Size();
+		ViewState->BloomFFTKernel.ImageSize = Intermediates.ImageSize;
 		ViewState->BloomFFTKernel.Physical = View.FinalPostProcessSettings.BloomConvolutionTexture;
 		ViewState->BloomFFTKernel.PhysicalRHI = PhysicalSpaceKernelTextureRef;
 		ViewState->BloomFFTKernel.PhysicalMipLevel = BloomConvolutionTextureResource->GetCurrentMipCount();
@@ -735,16 +636,105 @@ void InitDomainAndGetKernel(
 	*OutKernelConstantsBuffer = KernelConstantsBuffer;
 }
 
-FFFTBloomOutput AddFFTBloomPass(FRDGBuilder& GraphBuilder, const FViewInfo& View, const FFFTBloomInputs& Inputs)
+FFFTBloomOutput AddFFTBloomPass(FRDGBuilder& GraphBuilder, const FViewInfo& View, const FScreenPassTexture& InputSceneColor, float InputResolutionFraction)
 {
-	check(Inputs.FullResolutionTexture);
-	check(!Inputs.FullResolutionViewRect.IsEmpty());
-	check(Inputs.HalfResolutionTexture);
-	check(!Inputs.HalfResolutionViewRect.IsEmpty());
+	check(InputSceneColor.IsValid());
+
+	const float ResolutionFraction = GetFFTBloomResolutionFraction();
+	const float DownscaleResolutionFraction = ResolutionFraction / InputResolutionFraction;
+	check(DownscaleResolutionFraction <= 1.0f);
+
+	FFFTBloomIntermediates Intermediates;
+	{
+		FSceneViewState* ViewState = (FSceneViewState*)View.State;
+		check(ViewState);
+
+		const auto& PPSettings = View.FinalPostProcessSettings;
+
+		// The kernel parameters on the FinalPostProcess.
+
+		const float BloomConvolutionSize = PPSettings.BloomConvolutionSize;
+		const float KernelSupportScaleClamp = FMath::Clamp(PPSettings.BloomConvolutionBufferScale, 0.f, 1.f);
+
+		// Clip the Kernel support (i.e. bloom size) to 100% the screen width 
+		const float MaxBloomSize = 1.f;
+		const float KernelSupportScale = FMath::Clamp(BloomConvolutionSize, 0.f, MaxBloomSize);
+
+		// We padd by 1/2 the number of pixels the kernel needs in the x-direction
+		// so if the kernel is being applied on the edge of the image it will see padding and not periodicity
+		// NB:  If the kernel padding would force a transform buffer that is too big for group shared memory (> 4096)
+		//      we clamp it.  This could result in a wrap-around in the bloom (from one side of the screen to the other),
+		//      but since the amplitude of the bloom kernel tails is usually very small, this shouldn't be too bad.
+		auto KernelRadiusSupportFunctor = [KernelSupportScale, KernelSupportScaleClamp](const FIntPoint& Size) ->int32
+		{
+			float ClampedKernelSupportScale = (KernelSupportScaleClamp > 0) ? FMath::Min(KernelSupportScale, KernelSupportScaleClamp) : KernelSupportScale;
+			int32 FilterRadius = FMath::CeilToInt(0.5 * ClampedKernelSupportScale * Size.X);
+			const int32 MaxFFTSize = GPUFFT::MaxScanLineLength();
+			int32 MaxDim = FMath::Max(Size.X, Size.Y);
+			if (MaxDim + FilterRadius > MaxFFTSize && MaxDim < MaxFFTSize) FilterRadius = MaxFFTSize - MaxDim;
+
+			return FilterRadius;
+		};
+
+		Intermediates.ImageSize.X = FMath::CeilToInt(InputSceneColor.ViewRect.Width() * DownscaleResolutionFraction);
+		Intermediates.ImageSize.Y = FMath::CeilToInt(InputSceneColor.ViewRect.Height() * DownscaleResolutionFraction);
+
+		Intermediates.KernelSupportScale = KernelSupportScale;
+		Intermediates.KernelSupportScaleClamp = KernelSupportScaleClamp;
+
+		// The pre-filter boost parameters for bright pixels. Because the Convolution PP work in pre-exposure space, the min and max needs adjustment.
+		Intermediates.PreFilter = FVector(PPSettings.BloomConvolutionPreFilterMin, PPSettings.BloomConvolutionPreFilterMax, PPSettings.BloomConvolutionPreFilterMult);
+
+		// The length of the a side of the square kernel image in pixels
+		const int32 KernelSize = FMath::CeilToInt(KernelSupportScale * FMath::Max(Intermediates.ImageSize.X, Intermediates.ImageSize.Y));
 
-	FFFTBloomIntermediates Intermediates = GetFFTBloomIntermediates(View, Inputs);
+		const int32 SpectralPadding = KernelRadiusSupportFunctor(Intermediates.ImageSize);
 
-	RDG_EVENT_SCOPE(GraphBuilder, "FFTBloom %dx%d", Intermediates.ImageRect.Width(), Intermediates.ImageRect.Height());
+		// The following are mathematically equivalent
+		// 1) Horizontal FFT / Vertical FFT / Filter / Vertical InvFFT / Horizontal InvFFT
+		// 2) Vertical FFT / Horizontal FFT / Filter / Horizontal InvFFT / Vertical InvFFT
+		// but we choose the one with the smallest intermediate buffer size
+
+		// The size of the input image plus padding that accounts for
+		// the width of the kernel.  The ImageRect is virtually padded
+		// with black to account for the gather action of the convolution.
+		FIntPoint PaddedImageSize = Intermediates.ImageSize + FIntPoint(SpectralPadding, SpectralPadding);
+		PaddedImageSize.X = FMath::Max(PaddedImageSize.X, KernelSize);
+		PaddedImageSize.Y = FMath::Max(PaddedImageSize.Y, KernelSize);
+
+		Intermediates.FrequencySize = FIntPoint(FMath::RoundUpToPowerOfTwo(PaddedImageSize.X), FMath::RoundUpToPowerOfTwo(PaddedImageSize.Y));
+
+		// Choose to do to transform in the direction that results in writing the least amount of data to main memory.
+
+		Intermediates.bDoHorizontalFirst = ((Intermediates.FrequencySize.Y * PaddedImageSize.X) > (Intermediates.FrequencySize.X * PaddedImageSize.Y));
+
+		Intermediates.ComputePassFlags = (GSupportsEfficientAsyncCompute && CVarAsynComputeFFTBloom.GetValueOnRenderThread() != 0) ? ERDGPassFlags::AsyncCompute : ERDGPassFlags::Compute;
+	}
+
+	RDG_EVENT_SCOPE(GraphBuilder, "Bloom %dx%d", Intermediates.ImageSize.X, Intermediates.ImageSize.Y);
+
+	// Downscale the input to the final resolution fraction
+	FScreenPassTexture FFTInputSceneColor;
+	if (ResolutionFraction != InputResolutionFraction)
+	{
+		FIntPoint Extent;
+		QuantizeSceneBufferSize(Intermediates.ImageSize, Extent);
+
+		FRDGTextureDesc Desc = FRDGTextureDesc::Create2D(
+			Extent,
+			InputSceneColor.Texture->Desc.Format,
+			FClearValueBinding::None,
+			TexCreate_ShaderResource | TexCreate_UAV);
+
+		FFTInputSceneColor.Texture = GraphBuilder.CreateTexture(Desc, TEXT("Bloom.FFT.Input"));
+		FFTInputSceneColor.ViewRect = FIntRect(FIntPoint::ZeroValue, Intermediates.ImageSize);
+
+		AddDownsampleComputePass(GraphBuilder, View, InputSceneColor, FFTInputSceneColor, EDownsampleQuality::Low, Intermediates.ComputePassFlags);
+	}
+	else
+	{
+		FFTInputSceneColor = InputSceneColor;
+	}
 
 	// Init the domain data update the cached kernel if needed.
 	FRDGTextureRef SpectralKernelTexture;
@@ -785,13 +775,13 @@ FFFTBloomOutput AddFFTBloomPass(FRDGBuilder& GraphBuilder, const FViewInfo& View
 	}
 
 	FRDGTextureDesc OutputSceneColorDesc = FRDGTextureDesc::Create2D(
-		Intermediates.InputTexture->Desc.Extent,
-		Intermediates.InputTexture->Desc.Format,
+		FFTInputSceneColor.Texture->Desc.Extent,
+		FFTInputSceneColor.Texture->Desc.Format,
 		FClearValueBinding::None,
 		TexCreate_ShaderResource | TexCreate_UAV);
 		
 	BloomOutput.BloomTexture.Texture = GraphBuilder.CreateTexture(OutputSceneColorDesc, TEXT("Bloom.FFT.ScatterDispersion"));
-	BloomOutput.BloomTexture.ViewRect = Intermediates.ImageRect;
+	BloomOutput.BloomTexture.ViewRect = FFTInputSceneColor.ViewRect;
 
 	GPUFFT::ConvolutionWithTextureImage2D(
 		GraphBuilder,
@@ -800,7 +790,7 @@ FFFTBloomOutput AddFFTBloomPass(FRDGBuilder& GraphBuilder, const FViewInfo& View
 		Intermediates.FrequencySize,
 		Intermediates.bDoHorizontalFirst,
 		SpectralKernelTexture,
-		Intermediates.InputTexture, Intermediates.ImageRect,
+		FFTInputSceneColor.Texture, FFTInputSceneColor.ViewRect,
 		BloomOutput.BloomTexture.Texture, BloomOutput.BloomTexture.ViewRect,
 		Intermediates.PreFilter,
 		FFTMulitplyParameters);
diff --git a/Engine/Source/Runtime/Renderer/Private/PostProcess/PostProcessFFTBloom.h b/Engine/Source/Runtime/Renderer/Private/PostProcess/PostProcessFFTBloom.h
--- a/Engine/Source/Runtime/Renderer/Private/PostProcess/PostProcessFFTBloom.h
+++ b/Engine/Source/Runtime/Renderer/Private/PostProcess/PostProcessFFTBloom.h
@@ -11,14 +11,7 @@ struct FBloomOutputs;
 // Returns whether FFT bloom is enabled for the view.
 bool IsFFTBloomEnabled(const FViewInfo& View);
 
-struct FFFTBloomInputs
-{
-	FRDGTextureRef FullResolutionTexture;
-	FIntRect FullResolutionViewRect;
-
-	FRDGTextureRef HalfResolutionTexture;
-	FIntRect HalfResolutionViewRect;
-};
+float GetFFTBloomResolutionFraction();
 
 struct FFFTBloomOutput
 {
@@ -26,7 +19,6 @@ struct FFFTBloomOutput
 	FRDGBufferRef SceneColorApplyParameters = nullptr;
 };
 
-bool IsFFTBloomFullResolutionEnabled();
-bool IsFFTBloomQuarterResolutionEnabled();
+float GetFFTBloomResolutionFraction();
 
-FFFTBloomOutput AddFFTBloomPass(FRDGBuilder& GraphBuilder, const FViewInfo& View, const FFFTBloomInputs& Inputs);
+FFFTBloomOutput AddFFTBloomPass(FRDGBuilder& GraphBuilder, const FViewInfo& View, const FScreenPassTexture& InputSceneColor, float InputResolutionFraction);
diff --git a/Engine/Source/Runtime/Renderer/Private/PostProcess/PostProcessing.cpp b/Engine/Source/Runtime/Renderer/Private/PostProcess/PostProcessing.cpp
--- a/Engine/Source/Runtime/Renderer/Private/PostProcess/PostProcessing.cpp
+++ b/Engine/Source/Runtime/Renderer/Private/PostProcess/PostProcessing.cpp
@@ -96,6 +96,20 @@ TAutoConsoleVariable<int32> CVarPostProcessingQuarterResolutionDownsample(
 	TEXT("Uses quarter resolution downsample instead of half resolution to feed into exposure / bloom."),
 	ECVF_RenderThreadSafe);
 
+TAutoConsoleVariable<int32> CVarDownsampleQuality(
+	TEXT("r.PostProcessing.DownsampleQuality"), 0,
+	TEXT("Defines the quality used for downsampling to half or quarter res the scene color in post processing chain.\n")
+	TEXT(" 0: low quality (default)\n")
+	TEXT(" 1: high quality\n"),
+	ECVF_Scalability | ECVF_RenderThreadSafe);
+
+TAutoConsoleVariable<int32> CVarDownsampleChainQuality(
+	TEXT("r.PostProcessing.DownsampleChainQuality"), 1,
+	TEXT("Defines the quality used for downsampling to the scene color in scene color chains.\n")
+	TEXT(" 0: low quality\n")
+	TEXT(" 1: high quality (default)\n"),
+	ECVF_Scalability | ECVF_RenderThreadSafe);
+
 #if !(UE_BUILD_SHIPPING)
 TAutoConsoleVariable<int32> CVarPostProcessingForceAsyncDispatch(
 	TEXT("r.PostProcessing.ForceAsyncDispatch"),
@@ -107,6 +121,12 @@ TAutoConsoleVariable<int32> CVarPostProcessingForceAsyncDispatch(
 
 } //! namespace
 
+EDownsampleQuality GetDownsampleQuality(const TAutoConsoleVariable<int32>& CVar)
+{
+	const int32 DownsampleQuality = FMath::Clamp(CVar.GetValueOnRenderThread(), 0, 1);
+	return static_cast<EDownsampleQuality>(DownsampleQuality);
+}
+
 bool IsPostProcessingWithComputeEnabled(ERHIFeatureLevel::Type FeatureLevel)
 {
 	// Any thread is used due to FViewInfo initialization.
@@ -535,7 +555,8 @@ void AddPostProcessingPasses(FRDGBuilder& GraphBuilder, const FViewInfo& View, c
 
 		const EAutoExposureMethod AutoExposureMethod = GetAutoExposureMethod(View);
 		const EAntiAliasingMethod AntiAliasingMethod = !bVisualizeDepthOfField ? View.AntiAliasingMethod : AAM_None;
-		const EDownsampleQuality DownsampleQuality = GetDownsampleQuality();
+		const EDownsampleQuality DownsampleQuality = GetDownsampleQuality(CVarDownsampleQuality);
+		const EDownsampleQuality DownsampleChainQuality = GetDownsampleQuality(CVarDownsampleChainQuality);
 		const EPixelFormat DownsampleOverrideFormat = PF_FloatRGB;
 
 		// Previous transforms are nonsensical on camera cuts, unless motion vector simulation is enabled (providing FrameN+1 transforms to FrameN+0)
@@ -611,17 +632,15 @@ void AddPostProcessingPasses(FRDGBuilder& GraphBuilder, const FViewInfo& View, c
 		const bool bProcessQuarterResolution = IsPostProcessingQuarterResolutionDownsampleEnabled();
 		const bool bMotionBlurNeedsHalfResInput = PassSequence.IsEnabled(EPass::MotionBlur) && DoesMotionBlurNeedsHalfResInput() && !bVisualizeMotionBlur;
 
-		const bool bIsFullResFFT = bFFTBloomEnabled && IsFFTBloomFullResolutionEnabled();
-		const bool bIsHalfResFFT = bFFTBloomEnabled && !IsFFTBloomFullResolutionEnabled() && !IsFFTBloomQuarterResolutionEnabled();
-		const bool bIsQuarterResFFT = bFFTBloomEnabled && IsFFTBloomQuarterResolutionEnabled();
+		const float FFTBloomResolutionFraction = GetFFTBloomResolutionFraction();
 
 		const bool bProduceSceneColorChain = (
 			bBasicEyeAdaptationEnabled ||
 			(bBloomEnabled && !bFFTBloomEnabled) ||
 			(bLensFlareEnabled && bFFTBloomEnabled) ||
 			bLocalExposureBlurredLum);
-		const bool bNeedPostMotionBlurHalfRes = !bProcessQuarterResolution || bIsHalfResFFT;
-		const bool bNeedPostMotionBlurQuarterRes = bProcessQuarterResolution || bIsQuarterResFFT;
+		const bool bNeedPostMotionBlurHalfRes = !bProcessQuarterResolution || (bFFTBloomEnabled && FFTBloomResolutionFraction > 0.25f && FFTBloomResolutionFraction <= 0.5f);
+		const bool bNeedPostMotionBlurQuarterRes = bProcessQuarterResolution || (bFFTBloomEnabled && FFTBloomResolutionFraction <= 0.25f);
 
 		// Post Process Material Chain - Before Translucency
 		{
@@ -819,8 +838,8 @@ void AddPostProcessingPasses(FRDGBuilder& GraphBuilder, const FViewInfo& View, c
 
 		// Generate post motion blur lower res scene color if they have not been generated.
 		{
-			if ((bNeedPostMotionBlurHalfRes ||
-				(bNeedPostMotionBlurQuarterRes && !QuarterResSceneColor.IsValid())) && !HalfResSceneColor.IsValid())
+			if ((bNeedPostMotionBlurHalfRes && !HalfResSceneColor.IsValid()) ||
+				(bNeedPostMotionBlurQuarterRes && !QuarterResSceneColor.IsValid() && !HalfResSceneColor.IsValid()))
 			{
 				FDownsamplePassInputs PassInputs;
 				PassInputs.Name = TEXT("PostProcessing.SceneColor.HalfRes");
@@ -878,7 +897,7 @@ void AddPostProcessingPasses(FRDGBuilder& GraphBuilder, const FViewInfo& View, c
 				GraphBuilder, View,
 				EyeAdaptationParameters,
 				bProcessQuarterResolution ? QuarterResSceneColor : HalfResSceneColor,
-				DownsampleQuality,
+				DownsampleChainQuality,
 				bLogLumaInAlpha);
 		}
 
@@ -918,23 +937,26 @@ void AddPostProcessingPasses(FRDGBuilder& GraphBuilder, const FViewInfo& View, c
 			{
 				LensFlareSceneDownsampleChain = &SceneDownsampleChain;
 
-				FScreenPassTexture HalfResolution;
-				if (IsFFTBloomQuarterResolutionEnabled())
+				float InputResolutionFraction;
+				FScreenPassTexture InputSceneColor;
+
+				if (FFTBloomResolutionFraction <= 0.25f)
 				{
-					HalfResolution = QuarterResSceneColor;
+					InputSceneColor = QuarterResSceneColor;
+					InputResolutionFraction = 0.25f;
+				}
+				else if (FFTBloomResolutionFraction <= 0.5f)
+				{
+					InputSceneColor = HalfResSceneColor;
+					InputResolutionFraction = 0.5f;
 				}
 				else
 				{
-					HalfResolution = HalfResSceneColor;
+					InputSceneColor = SceneColor;
+					InputResolutionFraction = 1.0f;
 				}
 
-				FFFTBloomInputs PassInputs;
-				PassInputs.FullResolutionTexture = SceneColor.Texture;
-				PassInputs.FullResolutionViewRect = SceneColor.ViewRect;
-				PassInputs.HalfResolutionTexture = HalfResolution.Texture;
-				PassInputs.HalfResolutionViewRect = HalfResolution.ViewRect;
-
-				FFFTBloomOutput Outputs = AddFFTBloomPass(GraphBuilder, View, PassInputs);
+				FFFTBloomOutput Outputs = AddFFTBloomPass(GraphBuilder, View, InputSceneColor, InputResolutionFraction);
 				Bloom = Outputs.BloomTexture;
 				SceneColorApplyParameters = Outputs.SceneColorApplyParameters;
 			}
@@ -966,7 +988,7 @@ void AddPostProcessingPasses(FRDGBuilder& GraphBuilder, const FViewInfo& View, c
 					}
 
 					const bool bLogLumaInAlpha = false;
-					BloomDownsampleChain.Init(GraphBuilder, View, EyeAdaptationParameters, DownsampleInput, DownsampleQuality, bLogLumaInAlpha);
+					BloomDownsampleChain.Init(GraphBuilder, View, EyeAdaptationParameters, DownsampleInput, DownsampleChainQuality, bLogLumaInAlpha);
 
 					LensFlareSceneDownsampleChain = &BloomDownsampleChain;
 				}
