Fix hair strands shadowing in forward rendering.

Forward rendering now supports up to four shadow casting light, which are stored into the various channels of the shadow masks textures. This CL ensures that the hair mask is written in the proper channels. It also fixes potential light leaks due to invalid transmittance initialization when shadow mask is fully opaque.

diff --git a/Engine/Shaders/Private/DeferredLightPixelShaders.usf b/Engine/Shaders/Private/DeferredLightPixelShaders.usf
--- a/Engine/Shaders/Private/DeferredLightPixelShaders.usf
+++ b/Engine/Shaders/Private/DeferredLightPixelShaders.usf
@@ -321,6 +321,7 @@ Texture2D<float4> ScreenShadowMaskSubPixelTexture;
 uint HairTransmittanceBufferMaxCount;
 Buffer<uint> HairTransmittanceBuffer;
 uint HairShadowMaskValid;
+float4 ShadowChannelMask;
 
 void DeferredLightPixelMain(
 	float4 SVPos : SV_POSITION,
@@ -400,6 +401,10 @@ void DeferredLightPixelMain(
 
 	const float Dither = InterleavedGradientNoise(PixelCoord, View.StateFrameIndexMod8);
 	float4 LightAttenuation = HairShadowMaskValid ? ScreenShadowMaskSubPixelTexture.Load(uint3(PixelCoord,0)) : 1;
+	if (any(ShadowChannelMask < 1.0f))
+	{
+		LightAttenuation = dot(LightAttenuation, ShadowChannelMask).xxxx;
+	}
 	LightAttenuation = min(LightAttenuation, LightData.HairTransmittance.OpaqueVisibility.xxxx);
 
 	FRectTexture RectTexture = InitRectTexture(DeferredLightUniforms.SourceTexture);
diff --git a/Engine/Shaders/Private/HairStrands/HairStrandsDeepTransmittanceMask.usf b/Engine/Shaders/Private/HairStrands/HairStrandsDeepTransmittanceMask.usf
--- a/Engine/Shaders/Private/HairStrands/HairStrandsDeepTransmittanceMask.usf
+++ b/Engine/Shaders/Private/HairStrands/HairStrandsDeepTransmittanceMask.usf
@@ -52,6 +52,7 @@ uint2				DeepShadow_Resolution;
 float4x4			DeepShadow_CPUWorldToLightTransforms[MAX_HAIR_MACROGROUP_COUNT];
 float3				LightDirection;
 uint				LightChannelMask;
+float4				ShadowChannelMask;
 float4				LightPosition;
 float				DeepShadow_DepthBiasScale;
 float				DeepShadow_DensityScale;
@@ -339,6 +340,7 @@ float				LightRadius;
 Texture2D<uint> 	ShadowMaskBitsTexture;
 #else
 Texture2D<float4> 	RayMarchMaskTexture;
+float4				ShadowChannelMask;
 #endif
 Buffer<uint4>		IndirectArgsBuffer;
 RWBuffer<uint>		OutTransmittanceMask;
@@ -518,11 +520,10 @@ void MainCS(uint2 DispatchThreadId : SV_DispatchThreadID)
 	}
 
 	const uint2 PixelCoord = HairStrands.HairSampleCoords[SampleIndex];
-	if (RayMarchMaskTexture.Load(uint3(PixelCoord, 0)).r == 0.0f)
+	if (dot(RayMarchMaskTexture.Load(uint3(PixelCoord, 0)), ShadowChannelMask) == 0)
 	{
 		// This pixel is already fully shadowed opaque. (RayMarchMaskTexture is computed using the closest sample as of today)
-		const FHairTransmittanceMask DefaultMask = InitHairTransmittanceMask();
-		OutTransmittanceMask[SampleIndex] = PackTransmittanceMask(DefaultMask);
+		OutTransmittanceMask[SampleIndex] = InitNullPackedHairTransmittanceMask();
 		return;
 	}
 
diff --git a/Engine/Source/Runtime/Renderer/Private/DeferredShadingRenderer.h b/Engine/Source/Runtime/Renderer/Private/DeferredShadingRenderer.h
--- a/Engine/Source/Runtime/Renderer/Private/DeferredShadingRenderer.h
+++ b/Engine/Source/Runtime/Renderer/Private/DeferredShadingRenderer.h
@@ -734,7 +734,8 @@ private:
 		const FLightSceneInfo* LightSceneInfo,
 		FRDGTextureRef ScreenShadowMaskSubPixelTexture,
 		FRDGTextureRef LightingChannelsTexture,
-		const FHairStrandsTransmittanceMaskData& InTransmittanceMaskData);
+		const FHairStrandsTransmittanceMaskData& InTransmittanceMaskData,
+		const bool bForwardRendering);
 
 	/** Renders an array of simple lights using standard deferred shading. */
 	void RenderSimpleLightsStandardDeferred(
diff --git a/Engine/Source/Runtime/Renderer/Private/HairStrands/HairStrandsTransmittance.cpp b/Engine/Source/Runtime/Renderer/Private/HairStrands/HairStrandsTransmittance.cpp
--- a/Engine/Source/Runtime/Renderer/Private/HairStrands/HairStrandsTransmittance.cpp
+++ b/Engine/Source/Runtime/Renderer/Private/HairStrands/HairStrandsTransmittance.cpp
@@ -94,6 +94,7 @@ struct FHairStrandsTransmittanceLightParams
 	FVector3f LightDirection = FVector3f::ZeroVector;
 	FVector4f LightPosition = FVector4f(0, 0, 0, 0);
 	uint32 LightChannelMask = 0;
+	uint32 ShadowChannelMask = 0;
 	float LightRadius = 0;
 };
 
@@ -189,6 +190,7 @@ class FHairStrandsVoxelTransmittanceMaskCS : public FGlobalShader
 		SHADER_PARAMETER(FVector3f, LightDirection)
 		SHADER_PARAMETER(FVector4f, LightPosition)
 		SHADER_PARAMETER(uint32, LightChannelMask)
+		SHADER_PARAMETER(FVector4f, ShadowChannelMask)
 
 		SHADER_PARAMETER_RDG_TEXTURE(Texture2D<uint>, ShadowMaskBitsTexture)
 		SHADER_PARAMETER_RDG_TEXTURE(Texture2D, RayMarchMaskTexture)
@@ -256,6 +258,8 @@ static FRDGBufferRef AddHairStrandsVoxelTransmittanceMaskPass(
 		Parameters->ShadowMaskBitsTexture = nullptr;
 	}
 
+	Parameters->ShadowChannelMask = FVector4f(0, 0, 0, 0);
+	Parameters->ShadowChannelMask[FMath::Clamp(Params.ShadowChannelMask, 0, 3)] = 1.0f;
 	Parameters->LightChannelMask = Params.LightChannelMask;
 	Parameters->IndirectArgsBuffer = IndirectArgsBuffer;
 	Parameters->HairStrands = HairStrands::BindHairStrandsViewUniformParameters(View);
@@ -307,6 +311,7 @@ class FHairStrandsDeepShadowTransmittanceMaskCS : public FGlobalShader
 		SHADER_PARAMETER(FVector3f, LightDirection)
 		SHADER_PARAMETER(FVector4f, LightPosition)
 		SHADER_PARAMETER(uint32, LightChannelMask)
+		SHADER_PARAMETER(FVector4f, ShadowChannelMask)
 		SHADER_PARAMETER(FVector4f, DeepShadow_LayerDepths)
 		SHADER_PARAMETER(float, DeepShadow_DepthBiasScale)
 		SHADER_PARAMETER(float, DeepShadow_DensityScale)
@@ -352,7 +357,9 @@ struct FHairStrandsDeepShadowTransmittanceLightParams : FHairStrandsTransmittanc
 	float DeepShadow_DepthBiasScale = 0;
 	float DeepShadow_DensityScale = 0;
 	FMatrix DeepShadow_ShadowToWorld = FMatrix::Identity;
-	
+	uint32 OutputChannel = ~0;
+	uint32 ShadowChannelMask = 0;
+
 	FRDGTextureRef DeepShadow_FrontDepthTexture = nullptr;
 	FRDGTextureRef DeepShadow_DomTexture = nullptr;
 };
@@ -392,6 +399,8 @@ static FRDGBufferRef AddHairStrandsDeepShadowTransmittanceMaskPass(
 	Parameters->DeepShadow_bIsGPUDriven = Params.DeepShadow_bIsGPUDriven ? 1 : 0;;
 	Parameters->DeepShadow_WorldToLightTransformBuffer = Params.DeepShadow_WorldToLightTransformBuffer;
 	Parameters->RayMarchMaskTexture = ScreenShadowMaskSubPixelTexture ? ScreenShadowMaskSubPixelTexture : GraphBuilder.RegisterExternalTexture(GSystemTextures.WhiteDummy);
+	Parameters->ShadowChannelMask = FVector4f(0, 0, 0, 0);
+	Parameters->ShadowChannelMask[FMath::Clamp(Params.ShadowChannelMask, 0, 3)] = 1.0f;
 
 	for (uint32 SlotIndex=0;SlotIndex< FHairStrandsDeepShadowData::MaxMacroGroupCount;++SlotIndex)
 	{
@@ -458,6 +467,7 @@ IMPLEMENT_GLOBAL_SHADER(FHairStrandsVoxelShadowMaskPS, "/Engine/Private/HairStra
 struct FHairStrandsVoxelShadowParams
 {
 	bool			bIsWholeSceneLight = false;
+	uint32			OutputChannel = ~0;
 	FVector3f		Voxel_LightDirection = FVector3f::ZeroVector;
 	FVector4f		Voxel_LightPosition = FVector4f(0, 0, 0, 0);
 	uint32			Voxel_MacroGroupId;
@@ -512,6 +522,7 @@ static void AddHairStrandsVoxelShadowMaskPass(
 	TShaderMapRef<FHairStrandsVoxelShadowMaskPS> PixelShader(View.ShaderMap, PermutationVector);
 	const FGlobalShaderMap* GlobalShaderMap = View.ShaderMap;
 	const FIntRect Viewport = View.ViewRect;
+	const uint32 OutputChannel = Params.OutputChannel;
 
 	ClearUnusedGraphResources(PixelShader, Parameters);
 	FIntPoint Resolution = OutShadowMask->Desc.Extent;
@@ -520,11 +531,19 @@ static void AddHairStrandsVoxelShadowMaskPass(
 		RDG_EVENT_NAME("HairStrands::ShadowMask(Voxel)"),
 		Parameters,
 		ERDGPassFlags::Raster,
-		[Parameters, VertexShader, PixelShader, Viewport, Resolution](FRHICommandList& RHICmdList)
+		[Parameters, VertexShader, PixelShader, Viewport, Resolution, OutputChannel](FRHICommandList& RHICmdList)
 	{
 		FGraphicsPipelineStateInitializer GraphicsPSOInit;
 		RHICmdList.ApplyCachedRenderTargets(GraphicsPSOInit);
-		GraphicsPSOInit.BlendState = TStaticBlendState<CW_RGBA, BO_Min, BF_One, BF_One, BO_Min, BF_One, BF_One>::GetRHI(); // Min Operator
+
+		switch (OutputChannel)
+		{
+		case 0:  GraphicsPSOInit.BlendState = TStaticBlendState<CW_RED,   BO_Min, BF_One, BF_One, BO_Min, BF_One, BF_One>::GetRHI(); break; // Min Operator
+		case 1:  GraphicsPSOInit.BlendState = TStaticBlendState<CW_GREEN, BO_Min, BF_One, BF_One, BO_Min, BF_One, BF_One>::GetRHI(); break; // Min Operator
+		case 2:  GraphicsPSOInit.BlendState = TStaticBlendState<CW_BLUE,  BO_Min, BF_One, BF_One, BO_Min, BF_One, BF_One>::GetRHI(); break; // Min Operator
+		case 3:  GraphicsPSOInit.BlendState = TStaticBlendState<CW_ALPHA, BO_Min, BF_One, BF_One, BO_Min, BF_One, BF_One>::GetRHI(); break; // Min Operator
+		default: GraphicsPSOInit.BlendState = TStaticBlendState<CW_RGBA,  BO_Min, BF_One, BF_One, BO_Min, BF_One, BF_One>::GetRHI(); break; // Min Operator
+		}
 		GraphicsPSOInit.RasterizerState = TStaticRasterizerState<>::GetRHI();
 		GraphicsPSOInit.DepthStencilState = TStaticDepthStencilState<false, CF_Always>::GetRHI();
 
@@ -604,6 +623,7 @@ IMPLEMENT_GLOBAL_SHADER(FHairStrandsDeepShadowMaskPS, "/Engine/Private/HairStran
 struct FHairStrandsDeepShadowParams
 {
 	bool			bIsWholeSceneLight = false;
+	uint32			OutputChannel = ~0;
 	FRDGBufferSRVRef DeepShadow_WorldToLightTransformBuffer = nullptr;
 	FMatrix			DeepShadow_CPUWorldToLightTransform;
 	FIntRect		DeepShadow_AtlasRect;
@@ -670,6 +690,7 @@ static void AddHairStrandsDeepShadowMaskPass(
 	TShaderMapRef<FHairStrandsDeepShadowMaskPS> PixelShader(View.ShaderMap, PermutationVector);
 	const FGlobalShaderMap* GlobalShaderMap = View.ShaderMap;
 	const FIntRect Viewport = View.ViewRect;
+	const uint32 OutputChannel = Params.OutputChannel;
 
 	ClearUnusedGraphResources(PixelShader, Parameters);
 	FIntPoint Resolution = OutShadowMask->Desc.Extent;
@@ -678,11 +699,18 @@ static void AddHairStrandsDeepShadowMaskPass(
 		RDG_EVENT_NAME("HairStrands::ShadowMask"),
 		Parameters,
 		ERDGPassFlags::Raster,
-		[Parameters, VertexShader, PixelShader, Viewport, Resolution](FRHICommandList& RHICmdList)
+		[Parameters, VertexShader, PixelShader, Viewport, Resolution, OutputChannel](FRHICommandList& RHICmdList)
 	{
 		FGraphicsPipelineStateInitializer GraphicsPSOInit;
 		RHICmdList.ApplyCachedRenderTargets(GraphicsPSOInit);
-		GraphicsPSOInit.BlendState = TStaticBlendState<CW_RGBA, BO_Min, BF_One, BF_One, BO_Min, BF_One, BF_One>::GetRHI(); // Min Operator
+		switch (OutputChannel)
+		{
+		case 0:  GraphicsPSOInit.BlendState = TStaticBlendState<CW_RED, BO_Min, BF_One, BF_One, BO_Min, BF_One, BF_One>::GetRHI(); break; // Min Operator
+		case 1:  GraphicsPSOInit.BlendState = TStaticBlendState<CW_GREEN, BO_Min, BF_One, BF_One, BO_Min, BF_One, BF_One>::GetRHI(); break; // Min Operator
+		case 2:  GraphicsPSOInit.BlendState = TStaticBlendState<CW_BLUE, BO_Min, BF_One, BF_One, BO_Min, BF_One, BF_One>::GetRHI(); break; // Min Operator
+		case 3:  GraphicsPSOInit.BlendState = TStaticBlendState<CW_ALPHA, BO_Min, BF_One, BF_One, BO_Min, BF_One, BF_One>::GetRHI(); break; // Min Operator
+		default: GraphicsPSOInit.BlendState = TStaticBlendState<CW_RGBA, BO_Min, BF_One, BF_One, BO_Min, BF_One, BF_One>::GetRHI(); break; // Min Operator
+		}
 		GraphicsPSOInit.RasterizerState = TStaticRasterizerState<>::GetRHI();
 		GraphicsPSOInit.DepthStencilState = TStaticDepthStencilState<false, CF_Always>::GetRHI();
 
@@ -718,6 +746,7 @@ static FHairStrandsTransmittanceMaskData InternalRenderHairStrandsTransmittanceM
 	const FHairStrandsMacroGroupDatas& MacroGroupDatas,
 	const FHairStrandsDeepShadowResources& DeepShadowResources,
 	const FHairStrandsVoxelResources& VoxelResources,
+	const bool bProjectingForForwardShading,
 	FRDGTextureRef ScreenShadowMaskSubPixelTexture)
 {
 	FHairStrandsTransmittanceMaskData Out;
@@ -763,6 +792,7 @@ static FHairStrandsTransmittanceMaskData InternalRenderHairStrandsTransmittanceM
 					Params.LightPosition = DeepShadowData.LightPosition;
 					Params.LightRadius = 0;
 					Params.LightChannelMask = LightSceneInfo->Proxy->GetLightingChannelMask();
+					Params.ShadowChannelMask = bProjectingForForwardShading ? LightSceneInfo->GetDynamicShadowMapChannel() : 0;
 					Params.DeepShadow_LayerDepths = ComputeDeepShadowLayerDepths(DeepShadowData.LayerDistribution);
 					Params.DeepShadow_AtlasSlotOffsets_AtlasSlotIndex[DeepShadowData.MacroGroupId] = FIntVector4(DeepShadowData.AtlasRect.Min.X, DeepShadowData.AtlasRect.Min.Y, DeepShadowData.AtlasSlotIndex, 0);
 					Params.DeepShadow_CPUWorldToLightTransforms[DeepShadowData.MacroGroupId] = DeepShadowData.CPU_WorldToLightTransform;
@@ -796,6 +826,7 @@ static FHairStrandsTransmittanceMaskData InternalRenderHairStrandsTransmittanceM
 		Params.LightDirection = LightSceneInfo->Proxy->GetDirection();
 		Params.LightPosition = FVector4f(FVector(LightSceneInfo->Proxy->GetPosition()), LightSceneInfo->Proxy->GetLightType() == ELightComponentType::LightType_Directional ? 0 : 1);
 		Params.LightChannelMask = LightSceneInfo->Proxy->GetLightingChannelMask();
+		Params.ShadowChannelMask = bProjectingForForwardShading ? LightSceneInfo->GetDynamicShadowMapChannel() : 0;
 		Params.LightRadius = FMath::Max(LightParameters.SourceLength, LightParameters.SourceRadius);
 
 		Out.TransmittanceMask = AddHairStrandsVoxelTransmittanceMaskPass(
@@ -850,6 +881,7 @@ FHairStrandsTransmittanceMaskData RenderHairStrandsTransmittanceMask(
 	FRDGBuilder& GraphBuilder,
 	const FViewInfo& View,
 	const FLightSceneInfo* LightSceneInfo,
+	const bool bProjectingForForwardShading,
 	FRDGTextureRef ScreenShadowMaskSubPixelTexture)
 {
 	FHairStrandsTransmittanceMaskData TransmittanceMaskData;
@@ -863,6 +895,7 @@ FHairStrandsTransmittanceMaskData RenderHairStrandsTransmittanceMask(
 			View.HairStrandsViewData.MacroGroupDatas,
 			View.HairStrandsViewData.DeepShadowResources,
 			View.HairStrandsViewData.VirtualVoxelResources,
+			bProjectingForForwardShading,
 			ScreenShadowMaskSubPixelTexture);
 	}
 	return TransmittanceMaskData;
@@ -878,6 +911,7 @@ static void InternalRenderHairStrandsShadowMask(
 	const FHairStrandsMacroGroupDatas& InMacroGroupDatas,
 	const FHairStrandsDeepShadowResources& DeepShadowResources,
 	const FHairStrandsVoxelResources& VoxelResources,
+	const bool bProjectingForForwardShading, 
 	FRDGTextureRef OutShadowMask)
 {
 	if (InMacroGroupDatas.Num() == 0)
@@ -925,6 +959,7 @@ static void InternalRenderHairStrandsShadowMask(
 				Params.DeepShadow_DepthBiasScale = GetDeepShadowDepthBiasScale();
 				Params.DeepShadow_DensityScale = GetDeepShadowDensityScale();
 				Params.DeepShadow_LayerDepths = ComputeDeepShadowLayerDepths(DomData.LayerDistribution);
+				Params.OutputChannel = bProjectingForForwardShading ? LightSceneInfo->GetDynamicShadowMapChannel() : ~0;
 				AddHairStrandsDeepShadowMaskPass(
 					GraphBuilder,
 					SceneTextures.Depth.Resolve,
@@ -949,6 +984,7 @@ static void InternalRenderHairStrandsShadowMask(
 			Params.Voxel_LightPosition = FVector4f(FVector(LightSceneInfo->Proxy->GetPosition()), LightSceneInfo->Proxy->GetLightType() == ELightComponentType::LightType_Directional ? 0 : 1);
 			Params.Voxel_MacroGroupId = MacroGroupData.MacroGroupId;
 			Params.bIsWholeSceneLight = bIsWholeSceneLight ? 1 : 0;
+			Params.OutputChannel = bProjectingForForwardShading ? LightSceneInfo->GetDynamicShadowMapChannel() : ~0;
 
 			AddHairStrandsVoxelShadowMaskPass(
 				GraphBuilder,
@@ -964,6 +1000,7 @@ void RenderHairStrandsShadowMask(
 	FRDGBuilder& GraphBuilder,
 	const TArray<FViewInfo>& Views,
 	const FLightSceneInfo* LightSceneInfo,
+	const bool bProjectingForForwardShading,
 	FRDGTextureRef OutShadowMask)
 {
 	if (Views.Num() == 0 || OutShadowMask == nullptr)
@@ -984,6 +1021,7 @@ void RenderHairStrandsShadowMask(
 				View.HairStrandsViewData.MacroGroupDatas,
 				View.HairStrandsViewData.DeepShadowResources,
 				View.HairStrandsViewData.VirtualVoxelResources,
+				bProjectingForForwardShading,
 				OutShadowMask);
 		}
 	}
@@ -1019,6 +1057,7 @@ void RenderHairStrandsDeepShadowMask(
 				View.HairStrandsViewData.MacroGroupDatas,
 				View.HairStrandsViewData.DeepShadowResources,
 				View.HairStrandsViewData.VirtualVoxelResources,
+				false,
 				OutShadowMask);
 		}
 	}
diff --git a/Engine/Source/Runtime/Renderer/Private/HairStrands/HairStrandsTransmittance.h b/Engine/Source/Runtime/Renderer/Private/HairStrands/HairStrandsTransmittance.h
--- a/Engine/Source/Runtime/Renderer/Private/HairStrands/HairStrandsTransmittance.h
+++ b/Engine/Source/Runtime/Renderer/Private/HairStrands/HairStrandsTransmittance.h
@@ -25,6 +25,7 @@ void RenderHairStrandsShadowMask(
 	FRDGBuilder& GraphBuilder,
 	const TArray<FViewInfo>& Views,
 	const FLightSceneInfo* LightSceneInfo,
+	const bool bProjectingForForwardShading,
 	FRDGTextureRef ScreenShadowMaskTexture); 
 
 /// Write opaque hair shadow onto screen shadow mask to have fine hair details cast onto opaque geometries (deep shadow caster only)
@@ -39,6 +40,7 @@ FHairStrandsTransmittanceMaskData RenderHairStrandsTransmittanceMask(
 	FRDGBuilder& GraphBuilder,
 	const FViewInfo& View,
 	const class FLightSceneInfo* LightSceneInfo,
+	const bool bProjectingForForwardShading,
 	FRDGTextureRef ScreenShadowMaskSubPixelTexture);
 
 /// Output hair transmittance per hair sample for all lights using the forward cluster lights
diff --git a/Engine/Source/Runtime/Renderer/Private/LightGridInjection.cpp b/Engine/Source/Runtime/Renderer/Private/LightGridInjection.cpp
--- a/Engine/Source/Runtime/Renderer/Private/LightGridInjection.cpp
+++ b/Engine/Source/Runtime/Renderer/Private/LightGridInjection.cpp
@@ -899,7 +899,7 @@ void FDeferredShadingSceneRenderer::RenderForwardShadowProjections(
 
 				if (bIsHairEnable)
 				{
-					RenderHairStrandsShadowMask(GraphBuilder, Views, LightSceneInfo, ForwardScreenSpaceShadowMask.Target);
+					RenderHairStrandsShadowMask(GraphBuilder, Views, LightSceneInfo, bProjectingForForwardShading, ForwardScreenSpaceShadowMask.Target);
 				}
 			}
 
diff --git a/Engine/Source/Runtime/Renderer/Private/LightRendering.cpp b/Engine/Source/Runtime/Renderer/Private/LightRendering.cpp
--- a/Engine/Source/Runtime/Renderer/Private/LightRendering.cpp
+++ b/Engine/Source/Runtime/Renderer/Private/LightRendering.cpp
@@ -360,6 +360,7 @@ struct FRenderLightParams
 	FShaderResourceViewRHIRef DeepShadow_TransmittanceMaskBuffer = nullptr;
 	uint32 DeepShadow_TransmittanceMaskBufferMaxCount = 0;
 	FRHITexture* ScreenShadowMaskSubPixelTexture = nullptr;
+	FVector4f ShadowChannelMask = FVector4f(1,1,1,1);
 
 	// Cloud shadow data
 	FMatrix Cloud_WorldToLightClipShadowMatrix;
@@ -539,6 +540,7 @@ class FDeferredLightPS : public FGlobalShader
 		ScreenShadowMaskSubPixelTexture.Bind(Initializer.ParameterMap, TEXT("ScreenShadowMaskSubPixelTexture")); // TODO hook the shader itself
 
 		HairShadowMaskValid.Bind(Initializer.ParameterMap, TEXT("HairShadowMaskValid"));
+		ShadowChannelMask.Bind(Initializer.ParameterMap, TEXT("ShadowChannelMask"));
 		HairStrandsParameters.Bind(Initializer.ParameterMap, FHairStrandsViewUniformParameters::StaticStructMetadata.GetShaderVariableName());
 
 		DummyRectLightTextureForCapsuleCompilerWarning.Bind(Initializer.ParameterMap, TEXT("DummyRectLightTextureForCapsuleCompilerWarning"));
@@ -670,6 +672,12 @@ private:
 					ShaderRHI,
 					HairShadowMaskValid,
 					InHairShadowMaskValid);
+
+				SetShaderValue(
+					RHICmdList,
+					ShaderRHI,
+					ShadowChannelMask,
+					RenderLightParams->ShadowChannelMask);
 			}
 		}
 
@@ -730,6 +738,7 @@ private:
 	LAYOUT_FIELD(FShaderUniformBufferParameter, HairStrandsParameters);
 	LAYOUT_FIELD(FShaderResourceParameter, ScreenShadowMaskSubPixelTexture);
 	LAYOUT_FIELD(FShaderParameter, HairShadowMaskValid);
+	LAYOUT_FIELD(FShaderParameter, ShadowChannelMask);
 
 	LAYOUT_FIELD(FShaderResourceParameter, DummyRectLightTextureForCapsuleCompilerWarning);
 	LAYOUT_FIELD(FShaderResourceParameter, DummyRectLightSamplerForCapsuleCompilerWarning);
@@ -1274,7 +1283,7 @@ void FDeferredShadingSceneRenderer::RenderLights(
 						{
 							const FSortedLightSceneInfo& SortedLightInfo = SortedLights[LightIndex];
 							const FLightSceneInfo* const LightSceneInfo = SortedLightInfo.LightSceneInfo;
-							RenderLightForHair(GraphBuilder, View, SceneTextures.UniformBuffer, LightSceneInfo, NullScreenShadowMaskSubPixelTexture, LightingChannelsTexture, DummyTransmittanceMaskData);
+							RenderLightForHair(GraphBuilder, View, SceneTextures.UniformBuffer, LightSceneInfo, NullScreenShadowMaskSubPixelTexture, LightingChannelsTexture, DummyTransmittanceMaskData, false /*bForwardRendering*/);
 						}
 					}
 				}
@@ -1772,7 +1781,7 @@ void FDeferredShadingSceneRenderer::RenderLights(
 						// Inject deep shadow mask if the light supports it
 						if (bUseHairDeepShadow)
 						{
-							RenderHairStrandsShadowMask(GraphBuilder, Views, &LightSceneInfo, ScreenShadowMaskTexture);
+							RenderHairStrandsShadowMask(GraphBuilder, Views, &LightSceneInfo, false /*bForwardShading*/, ScreenShadowMaskTexture);
 						}
 					}
 					else // (OcclusionType == FOcclusionType::Shadowmap)
@@ -1905,7 +1914,7 @@ void FDeferredShadingSceneRenderer::RenderLights(
 					{
 						if (bDrawHairShadow && HairStrands::HasViewHairStrandsData(View))
 						{
-							FHairStrandsTransmittanceMaskData TransmittanceMaskData = RenderHairStrandsTransmittanceMask(GraphBuilder, View, &LightSceneInfo, ScreenShadowMaskSubPixelTexture);							
+							FHairStrandsTransmittanceMaskData TransmittanceMaskData = RenderHairStrandsTransmittanceMask(GraphBuilder, View, &LightSceneInfo, false, ScreenShadowMaskSubPixelTexture);
 							if (TransmittanceMaskData.TransmittanceMask == nullptr)
 							{
 								TransmittanceMaskData = DummyTransmittanceMaskData;
@@ -1913,7 +1922,7 @@ void FDeferredShadingSceneRenderer::RenderLights(
 
 							// Note: ideally the light should still be evaluated for hair when not casting shadow, but for preserving the old behavior, and not adding 
 							// any perf. regression, we disable this light for hair rendering 
-							RenderLightForHair(GraphBuilder, View, SceneTextures.UniformBuffer, &LightSceneInfo, ScreenShadowMaskSubPixelTexture, LightingChannelsTexture, TransmittanceMaskData);
+							RenderLightForHair(GraphBuilder, View, SceneTextures.UniformBuffer, &LightSceneInfo, ScreenShadowMaskSubPixelTexture, LightingChannelsTexture, TransmittanceMaskData, false /*bForwardRendering*/);
 						}
 					}
 				}
@@ -2443,7 +2452,8 @@ void FDeferredShadingSceneRenderer::RenderLightForHair(
 	const FLightSceneInfo* LightSceneInfo,
 	FRDGTextureRef HairShadowMaskTexture,
 	FRDGTextureRef LightingChannelsTexture,
-	const FHairStrandsTransmittanceMaskData& InTransmittanceMaskData)
+	const FHairStrandsTransmittanceMaskData& InTransmittanceMaskData,
+	const bool bForwardRendering)
 {
 	const bool bHairRenderingEnabled = HairStrands::HasViewHairStrandsData(View);
 	if (!bHairRenderingEnabled)
@@ -2476,6 +2486,13 @@ void FDeferredShadingSceneRenderer::RenderLightForHair(
 			return;
 		}
 
+		FVector4f ShadowChannelMask = FVector4f(1, 1, 1, 1);
+		if (bForwardRendering)
+		{
+			ShadowChannelMask = FVector4f(0, 0, 0, 0);
+			ShadowChannelMask[FMath::Clamp(LightSceneInfo->GetDynamicShadowMapChannel(), 0, 3)] = 1.f;
+		}
+
 		FRenderLightForHairParameters* PassParameters = GraphBuilder.AllocParameters<FRenderLightForHairParameters>();
 		GetRenderLightParameters(
 			View,
@@ -2500,7 +2517,7 @@ void FDeferredShadingSceneRenderer::RenderLightForHair(
 			{},
 			PassParameters,
 			ERDGPassFlags::Raster,
-			[this, &HairVisibilityData, &View, PassParameters, LightSceneInfo, MaxTransmittanceElementCount, HairShadowMaskTexture, LightingChannelsTexture, bIsShadowMaskValid](FRHICommandList& RHICmdList)
+			[this, &HairVisibilityData, &View, PassParameters, LightSceneInfo, MaxTransmittanceElementCount, HairShadowMaskTexture, LightingChannelsTexture, bIsShadowMaskValid, ShadowChannelMask](FRHICommandList& RHICmdList)
 		{
 			RHICmdList.SetViewport(0, 0, 0.0f, HairVisibilityData.SampleLightingViewportResolution.X, HairVisibilityData.SampleLightingViewportResolution.Y, 1.0f);
 
@@ -2508,6 +2525,7 @@ void FDeferredShadingSceneRenderer::RenderLightForHair(
 			RenderLightParams.DeepShadow_TransmittanceMaskBufferMaxCount = MaxTransmittanceElementCount;
 			RenderLightParams.ScreenShadowMaskSubPixelTexture = bIsShadowMaskValid ? PassParameters->Light.ShadowMaskTexture->GetRHI() : GSystemTextures.WhiteDummy->GetShaderResourceRHI();
 			RenderLightParams.DeepShadow_TransmittanceMaskBuffer = PassParameters->HairTransmittanceMaskSRV->GetRHI();
+			RenderLightParams.ShadowChannelMask = ShadowChannelMask;
 
 			FGraphicsPipelineStateInitializer GraphicsPSOInit;
 			RHICmdList.ApplyCachedRenderTargets(GraphicsPSOInit);
@@ -2612,7 +2630,7 @@ void FDeferredShadingSceneRenderer::RenderLightsForHair(
 					FHairStrandsTransmittanceMaskData TransmittanceMaskData = DummyTransmittanceMaskData;
 					if (bDrawHairShadow)
 					{
-						TransmittanceMaskData = RenderHairStrandsTransmittanceMask(GraphBuilder, View, &LightSceneInfo, ScreenShadowMaskSubPixelTexture);
+						TransmittanceMaskData = RenderHairStrandsTransmittanceMask(GraphBuilder, View, &LightSceneInfo, true, ScreenShadowMaskSubPixelTexture);
 					}
 
 					RenderLightForHair(
@@ -2622,7 +2640,8 @@ void FDeferredShadingSceneRenderer::RenderLightsForHair(
 						&LightSceneInfo,
 						ScreenShadowMaskSubPixelTexture,
 						LightingChannelsTexture,
-						TransmittanceMaskData);
+						TransmittanceMaskData,
+						true /*bForwardRendering*/);
 				}
 			}
 		}
diff --git a/Engine/Source/Runtime/Renderer/Private/ShadowRendering.cpp b/Engine/Source/Runtime/Renderer/Private/ShadowRendering.cpp
--- a/Engine/Source/Runtime/Renderer/Private/ShadowRendering.cpp
+++ b/Engine/Source/Runtime/Renderer/Private/ShadowRendering.cpp
@@ -2329,7 +2329,7 @@ void FDeferredShadingSceneRenderer::RenderDeferredShadowProjections(
 		}
 		if (bNeedHairShadowMaskPass)
 		{
-			RenderHairStrandsShadowMask(GraphBuilder, Views, LightSceneInfo, ScreenShadowMaskTexture);
+			RenderHairStrandsShadowMask(GraphBuilder, Views, LightSceneInfo, false /*bForward*/, ScreenShadowMaskTexture);
 		}
 	}
 }
