Fix groom selection.

diff --git a/Engine/Source/Runtime/Renderer/Public/HairStrandsInterface.h b/Engine/Source/Runtime/Renderer/Public/HairStrandsInterface.h
--- a/Engine/Source/Runtime/Renderer/Public/HairStrandsInterface.h
+++ b/Engine/Source/Runtime/Renderer/Public/HairStrandsInterface.h
@@ -272,7 +272,7 @@ public:
 			float HairRootScale = 0;
 			float HairTipScale = 0;
 			float HairRaytracingRadiusScale = 0;
-			float HairLengthScale = 0;
+			float HairLengthScale = 1.f;
 			float HairLength = 0;
 			float HairDensity = 0;
 			bool bUseStableRasterization = false;
