Refactors FDeferredDecalVS

diff --git a/Engine/Source/Runtime/Renderer/Private/DecalRenderingShared.cpp b/Engine/Source/Runtime/Renderer/Private/DecalRenderingShared.cpp
--- a/Engine/Source/Runtime/Renderer/Private/DecalRenderingShared.cpp
+++ b/Engine/Source/Runtime/Renderer/Private/DecalRenderingShared.cpp
@@ -37,32 +37,16 @@ FTransientDecalRenderData::FTransientDecalRenderData(const FScene& InScene, cons
  */
 class FDeferredDecalVS : public FGlobalShader
 {
-	DECLARE_SHADER_TYPE(FDeferredDecalVS,Global);
-public:
-
-	static bool ShouldCompilePermutation(const FGlobalShaderPermutationParameters& Parameters)
-	{
-		return true;
-	}
-
-	FDeferredDecalVS( )	{ }
-	FDeferredDecalVS(const ShaderMetaType::CompiledShaderInitializerType& Initializer)
-		: FGlobalShader(Initializer)
-	{
-		FrustumComponentToClip.Bind(Initializer.ParameterMap, TEXT("FrustumComponentToClip"));
-	}
-
-	void SetParameters(FRHICommandList& RHICmdList, const FMatrix44f& InFrustumComponentToClip)
-	{
-		FRHIVertexShader* ShaderRHI = RHICmdList.GetBoundVertexShader();
-		SetShaderValue(RHICmdList, ShaderRHI, FrustumComponentToClip, InFrustumComponentToClip);
-	}
+	DECLARE_GLOBAL_SHADER(FDeferredDecalVS);
+	SHADER_USE_PARAMETER_STRUCT(FDeferredDecalVS, FGlobalShader);
 
-private:
-	LAYOUT_FIELD(FShaderParameter, FrustumComponentToClip);
+	BEGIN_SHADER_PARAMETER_STRUCT(FParameters, )
+		SHADER_PARAMETER(FMatrix44f, FrustumComponentToClip)
+		SHADER_PARAMETER_STRUCT_REF(FPrimitiveUniformShaderParameters, PrimitiveUniformBuffer)
+	END_SHADER_PARAMETER_STRUCT()
 };
 
-IMPLEMENT_SHADER_TYPE(,FDeferredDecalVS,TEXT("/Engine/Private/DeferredDecal.usf"),TEXT("MainVS"),SF_Vertex);
+IMPLEMENT_GLOBAL_SHADER(FDeferredDecalVS, "/Engine/Private/DeferredDecal.usf", "MainVS" ,SF_Vertex);
 
 /**
  * A pixel shader for projecting a deferred decal onto the scene.
@@ -393,31 +377,33 @@ namespace DecalRendering
 			GraphicsPSOInit.PrimitiveType = PT_TriangleList;
 
 			SetGraphicsPipelineState(RHICmdList, GraphicsPSOInit, StencilRef);
-			PixelShader->SetParameters(RHICmdList, View, DecalData.MaterialProxy, *DecalData.DecalProxy, DecalData.FadeAlpha);
 		}
 
-		// SetUniformBufferParameter() need to happen after the shader has been set otherwise a DebugBreak could occur.
+		// Set vertex shader parameters.
+		{
+			FDeferredDecalVS::FParameters ShaderParameters;
+			ShaderParameters.FrustumComponentToClip = FrustumComponentToClip;
+			ShaderParameters.PrimitiveUniformBuffer = GIdentityPrimitiveUniformBuffer.GetUniformBufferRef();
+
+			SetShaderParameters(RHICmdList, VertexShader, VertexShader.GetVertexShader(), ShaderParameters);
+		}
 
-		// we don't have the Primitive uniform buffer setup for decals (later we want to batch)
+		// Set pixel shader parameters.
 		{
-			auto& PrimitiveVS = VertexShader->GetUniformBufferParameter<FPrimitiveUniformShaderParameters>();
+			PixelShader->SetParameters(RHICmdList, View, DecalData.MaterialProxy, *DecalData.DecalProxy, DecalData.FadeAlpha);
+
 			auto& PrimitivePS = PixelShader->GetUniformBufferParameter<FPrimitiveUniformShaderParameters>();
 
 			// uncomment to track down usage of the Primitive uniform buffer
 			//	check(!PrimitiveVS.IsBound());
 			//	check(!PrimitivePS.IsBound());
 
-			// to prevent potential shader error (UE-18852 ElementalDemo crashes due to nil constant buffer)
-			SetUniformBufferParameter(RHICmdList, VertexShader.GetVertexShader(), PrimitiveVS, GIdentityPrimitiveUniformBuffer);
-
 			if (DebugViewMode == DVSM_None)
 			{
 				SetUniformBufferParameter(RHICmdList, PixelShader.GetPixelShader(), PrimitivePS, GIdentityPrimitiveUniformBuffer);
 			}
 		}
 
-		VertexShader->SetParameters(RHICmdList, FrustumComponentToClip);
-
 		// Set stream source after updating cached strides
 		RHICmdList.SetStreamSource(0, GetUnitCubeVertexBuffer(), 0);
 	}
@@ -431,6 +417,15 @@ namespace DecalRendering
 		GraphicsPSOInit.PrimitiveType = PT_TriangleList;
 
 		SetGraphicsPipelineState(RHICmdList, GraphicsPSOInit, 0);
-		VertexShader->SetParameters(RHICmdList, FrustumComponentToClip);
+
+		// Set vertex shader parameters.
+		{
+			FDeferredDecalVS::FParameters ShaderParameters;
+			ShaderParameters.FrustumComponentToClip = FrustumComponentToClip;
+			ShaderParameters.PrimitiveUniformBuffer = GIdentityPrimitiveUniformBuffer.GetUniformBufferRef();
+
+			SetShaderParameters(RHICmdList, VertexShader, VertexShader.GetVertexShader(), ShaderParameters);
+		}
+
 	}
 }
