Fixes static analysis warning in TSR

diff --git a/Engine/Source/Runtime/Renderer/Private/PostProcess/TemporalSuperResolution.cpp b/Engine/Source/Runtime/Renderer/Private/PostProcess/TemporalSuperResolution.cpp
--- a/Engine/Source/Runtime/Renderer/Private/PostProcess/TemporalSuperResolution.cpp
+++ b/Engine/Source/Runtime/Renderer/Private/PostProcess/TemporalSuperResolution.cpp
@@ -855,7 +855,7 @@ ITemporalUpscaler::FOutputs AddTemporalSuperResolutionPasses(
 		FComputeShaderUtils::AddPass(
 			GraphBuilder,
 			RDG_EVENT_NAME("TSR DilateVelocity(MotionBlurDirections=%d) %dx%d",
-				PermutationVector.Get<FTSRDilateVelocityCS::FMotionBlurDirectionsDim>(),
+				int32(PermutationVector.Get<FTSRDilateVelocityCS::FMotionBlurDirectionsDim>()),
 				InputRect.Width(), InputRect.Height()),
 			ComputeShader,
 			PassParameters,
