New ray tracing debug visualization modes:

Traversal Node - shows bounding box nodes intersections
Traversal Triangle - shows triangle/leaf nodes intersections
Traversal All - shows total node intersection count

diff --git a/Engine/Shaders/Private/RayTracing/RayTracingDebugTraversal.usf b/Engine/Shaders/Private/RayTracing/RayTracingDebugTraversal.usf
new file mode 100644
--- /dev/null
+++ b/Engine/Shaders/Private/RayTracing/RayTracingDebugTraversal.usf
@@ -0,0 +1,53 @@
+// Copyright Epic Games, Inc. All Rights Reserved.
+
+#include "../Common.ush"
+#include "../ShadingCommon.ush"
+#include "../ColorMap.ush"
+
+#include "RayTracingCommon.ush"
+#include "RayTracingDeferredMaterials.ush"
+
+#include "/Engine/Shared/RayTracingDebugDefinitions.h"
+
+#include "TraceRayInline.ush"
+
+RWTexture2D<float4> Output;
+RaytracingAccelerationStructure TLAS;
+
+uint VisualizationMode;
+float TraversalBoxScale;
+float TraversalTriangleScale;
+
+[numthreads(INLINE_RAY_TRACING_THREAD_GROUP_SIZE_X, INLINE_RAY_TRACING_THREAD_GROUP_SIZE_Y, 1)]
+void RayTracingDebugTraversalCS(uint3 DispatchThread : SV_DispatchThreadID)
+{
+	uint2 PixelCoord = DispatchThread.xy + View.ViewRectMin;
+
+	float2 RenderTargetUV = (float2(PixelCoord) + .5f) * View.BufferSizeAndInvSize.zw;
+
+	RayDesc Ray = CreatePrimaryRay(RenderTargetUV);
+
+	uint RayFlags = RAY_FLAG_CULL_BACK_FACING_TRIANGLES | RAY_FLAG_FORCE_OPAQUE;
+
+	FTraceRayInlineTraversalStatistics TraversalStatistics;
+	FTraceRayInlineResult TraceResult = TraceRayInline(TLAS, RayFlags, RAY_TRACING_MASK_ALL, Ray, TraversalStatistics, 0xFFFFFFFF);
+
+	float4 Result = 0;
+
+	const uint Intersections = TraversalStatistics.NodeIntersectionCount + TraversalStatistics.TriangleIntersectionCount;
+	switch (VisualizationMode)
+	{
+	default:
+	case RAY_TRACING_DEBUG_VIZ_TRAVERSAL_ALL:
+		Result.rgb = ColorMapTurbo(Intersections / (TraversalBoxScale + TraversalTriangleScale));
+		break;
+	case RAY_TRACING_DEBUG_VIZ_TRAVERSAL_NODE:
+		Result.rgb = ColorMapTurbo(TraversalStatistics.NodeIntersectionCount / TraversalBoxScale);
+		break;
+	case RAY_TRACING_DEBUG_VIZ_TRAVERSAL_TRIANGLE:
+		Result.rgb = ColorMapTurbo(TraversalStatistics.TriangleIntersectionCount / TraversalTriangleScale);
+		break;
+	}
+
+	Output[PixelCoord] = Result;
+}
diff --git a/Engine/Shaders/Private/RayTracing/TraceRayInline.ush b/Engine/Shaders/Private/RayTracing/TraceRayInline.ush
--- a/Engine/Shaders/Private/RayTracing/TraceRayInline.ush
+++ b/Engine/Shaders/Private/RayTracing/TraceRayInline.ush
@@ -25,7 +25,8 @@
 FTraceRayInlineResult TraceRayInline(RaytracingAccelerationStructure AccelerationStructure,
 	uint RayFlags,
 	uint InstanceInclusionMask,
-	RayDesc Ray)
+	RayDesc Ray,
+	uint MaxIterations = INLINE_RAY_TRACING_MAX_ITERATIONS)
 {
 	FTraceRayInlineResult TraceRayResult = (FTraceRayInlineResult)0;
 	TraceRayResult.SetMiss();
@@ -56,6 +57,19 @@ FTraceRayInlineResult TraceRayInline(RaytracingAccelerationStructure Acceleratio
 	return TraceRayResult;
 }
 
+FTraceRayInlineResult TraceRayInline(RaytracingAccelerationStructure AccelerationStructure,
+	uint RayFlags,
+	uint InstanceInclusionMask,
+	RayDesc Ray,
+	inout FTraceRayInlineTraversalStatistics Statistics,
+	uint MaxIterations = INLINE_RAY_TRACING_MAX_ITERATIONS)
+{
+	Statistics = (FTraceRayInlineTraversalStatistics)0;
+	FTraceRayInlineResult TraceRayResult = TraceRayInline(AccelerationStructure, RayFlags, InstanceInclusionMask, Ray, MaxIterations);
+
+	return TraceRayResult;
+}
+
 #endif // OVERRIDE_TRACERAYINLINE_USH
 
 #ifndef PLATFORM_SUPPORTS_INLINE_RAY_TRACING_TRIANGLE_NORMALS
diff --git a/Engine/Shaders/Shared/RaytracingDebugDefinitions.h b/Engine/Shaders/Shared/RaytracingDebugDefinitions.h
--- a/Engine/Shaders/Shared/RaytracingDebugDefinitions.h
+++ b/Engine/Shaders/Shared/RaytracingDebugDefinitions.h
@@ -33,3 +33,6 @@
 #define RAY_TRACING_DEBUG_VIZ_PERFORMANCE				23 
 #define RAY_TRACING_DEBUG_VIZ_TRIANGLES					24
 #define RAY_TRACING_DEBUG_VIZ_FAR_FIELD					25
+#define RAY_TRACING_DEBUG_VIZ_TRAVERSAL_NODE			26
+#define RAY_TRACING_DEBUG_VIZ_TRAVERSAL_TRIANGLE		27
+#define RAY_TRACING_DEBUG_VIZ_TRAVERSAL_ALL				28
diff --git a/Engine/Source/Runtime/Renderer/Private/RayTracing/RayTracingDebug.cpp b/Engine/Source/Runtime/Renderer/Private/RayTracing/RayTracingDebug.cpp
--- a/Engine/Source/Runtime/Renderer/Private/RayTracing/RayTracingDebug.cpp
+++ b/Engine/Source/Runtime/Renderer/Private/RayTracing/RayTracingDebug.cpp
@@ -38,6 +38,18 @@ static TAutoConsoleVariable<float> CVarRayTracingDebugTimingScale(
 	TEXT("Scaling factor for ray timing heat map visualization. (default = 1)\n")
 );
 
+static TAutoConsoleVariable<float> CVarRayTracingDebugTraversalBoxScale(
+	TEXT("r.RayTracing.DebugTraversalScale.Box"),
+	150.0f,
+	TEXT("Scaling factor for box traversal heat map visualization. (default = 150)\n")
+);
+
+static TAutoConsoleVariable<float> CVarRayTracingDebugTraversalTriangleScale(
+	TEXT("r.RayTracing.DebugTraversalScale.Triangle"),
+	30.0f,
+	TEXT("Scaling factor for triangle traversal heat map visualization. (default = 30)\n")
+);
+
 class FRayTracingDebugRGS : public FGlobalShader
 {
 	DECLARE_GLOBAL_SHADER(FRayTracingDebugRGS)
@@ -86,11 +98,52 @@ public:
 };
 IMPLEMENT_SHADER_TYPE(, FRayTracingDebugCHS, TEXT("/Engine/Private/RayTracing/RayTracingDebug.usf"), TEXT("RayTracingDebugMainCHS"), SF_RayHitGroup);
 
+class FRayTracingDebugTraversalCS : public FGlobalShader
+{
+	DECLARE_GLOBAL_SHADER(FRayTracingDebugTraversalCS)
+	SHADER_USE_PARAMETER_STRUCT(FRayTracingDebugTraversalCS, FGlobalShader)
+
+	BEGIN_SHADER_PARAMETER_STRUCT(FParameters, )
+		SHADER_PARAMETER(uint32, VisualizationMode)
+		SHADER_PARAMETER(float, TraversalBoxScale)
+		SHADER_PARAMETER(float, TraversalTriangleScale)
+		SHADER_PARAMETER_SRV(RaytracingAccelerationStructure, TLAS)
+		SHADER_PARAMETER_RDG_TEXTURE_UAV(RWTexture2D<float4>, Output)
+		SHADER_PARAMETER_STRUCT_REF(FViewUniformShaderParameters, ViewUniformBuffer)
+	END_SHADER_PARAMETER_STRUCT()
+
+	static void ModifyCompilationEnvironment(const FGlobalShaderPermutationParameters& Parameters, FShaderCompilerEnvironment& OutEnvironment)
+	{
+		OutEnvironment.CompilerFlags.Add(CFLAG_Wave32);
+		OutEnvironment.CompilerFlags.Add(CFLAG_InlineRayTracing);
+
+		OutEnvironment.SetDefine(TEXT("INLINE_RAY_TRACING_THREAD_GROUP_SIZE_X"), ThreadGroupSizeX);
+		OutEnvironment.SetDefine(TEXT("INLINE_RAY_TRACING_THREAD_GROUP_SIZE_Y"), ThreadGroupSizeY);
+	}
+
+	static bool ShouldCompilePermutation(const FGlobalShaderPermutationParameters& Parameters)
+	{
+		return IsRayTracingEnabledForProject(Parameters.Platform) && RHISupportsRayTracing(Parameters.Platform) && RHISupportsInlineRayTracing(Parameters.Platform);
+	}
+
+	static constexpr uint32 ThreadGroupSizeX = 8;
+	static constexpr uint32 ThreadGroupSizeY = 4;
+	static_assert(ThreadGroupSizeX*ThreadGroupSizeY == 32, "Current inline ray tracing implementation requires 1:1 mapping between thread groups and waves and only supports wave32 mode.");
+};
+IMPLEMENT_GLOBAL_SHADER(FRayTracingDebugTraversalCS, "/Engine/Private/RayTracing/RayTracingDebugTraversal.usf", "RayTracingDebugTraversalCS", SF_Compute);
+
 static bool RequiresRayTracingDebugCHS(uint32 DebugVisualizationMode)
 {
 	return DebugVisualizationMode == RAY_TRACING_DEBUG_VIZ_INSTANCES || DebugVisualizationMode == RAY_TRACING_DEBUG_VIZ_TRIANGLES;
 }
 
+static bool IsRayTracingDebugTraversalMode(uint32 DebugVisualizationMode)
+{
+	return DebugVisualizationMode == RAY_TRACING_DEBUG_VIZ_TRAVERSAL_NODE || 
+		DebugVisualizationMode == RAY_TRACING_DEBUG_VIZ_TRAVERSAL_TRIANGLE || 
+		DebugVisualizationMode == RAY_TRACING_DEBUG_VIZ_TRAVERSAL_ALL;
+}
+
 void FDeferredShadingSceneRenderer::PrepareRayTracingDebug(const FSceneViewFamily& ViewFamily, TArray<FRHIRayTracingShader*>& OutRayGenShaders)
 {
 	// Declare all RayGen shaders that require material closest hit shaders to be bound
@@ -133,6 +186,9 @@ void FDeferredShadingSceneRenderer::RenderRayTracingDebug(FRDGBuilder& GraphBuil
 		RayTracingDebugVisualizationModes.Emplace(FName(*LOCTEXT("Performance", "Performance").ToString()),										RAY_TRACING_DEBUG_VIZ_PERFORMANCE);
 		RayTracingDebugVisualizationModes.Emplace(FName(*LOCTEXT("Triangles", "Triangles").ToString()),											RAY_TRACING_DEBUG_VIZ_TRIANGLES);
 		RayTracingDebugVisualizationModes.Emplace(FName(*LOCTEXT("FarField", "FarField").ToString()),											RAY_TRACING_DEBUG_VIZ_FAR_FIELD);
+		RayTracingDebugVisualizationModes.Emplace(FName(*LOCTEXT("Traversal Node", "Traversal Node").ToString()),								RAY_TRACING_DEBUG_VIZ_TRAVERSAL_NODE);
+		RayTracingDebugVisualizationModes.Emplace(FName(*LOCTEXT("Traversal Triangle", "Traversal Triangle").ToString()),						RAY_TRACING_DEBUG_VIZ_TRAVERSAL_TRIANGLE);
+		RayTracingDebugVisualizationModes.Emplace(FName(*LOCTEXT("Traversal All", "Traversal All").ToString()),									RAY_TRACING_DEBUG_VIZ_TRAVERSAL_ALL);
 	}
 
 	uint32 DebugVisualizationMode;
@@ -158,7 +214,31 @@ void FDeferredShadingSceneRenderer::RenderRayTracingDebug(FRDGBuilder& GraphBuil
 		return RenderRayTracingBarycentrics(GraphBuilder, View, SceneColorTexture);
 	}
 
-	// Debug modes other than barycentrics require pipeline support.
+	if (IsRayTracingDebugTraversalMode(DebugVisualizationMode) && ShouldRenderRayTracingEffect(ERayTracingPipelineCompatibilityFlags::Inline))
+	{
+		FRayTracingDebugTraversalCS::FParameters* PassParameters = GraphBuilder.AllocParameters<FRayTracingDebugTraversalCS::FParameters>();
+
+		PassParameters->VisualizationMode = DebugVisualizationMode;
+		PassParameters->TraversalBoxScale = CVarRayTracingDebugTraversalBoxScale.GetValueOnAnyThread();
+		PassParameters->TraversalTriangleScale = CVarRayTracingDebugTraversalTriangleScale.GetValueOnAnyThread();
+		PassParameters->TLAS = View.GetRayTracingSceneViewChecked();
+		PassParameters->ViewUniformBuffer = View.ViewUniformBuffer;
+		PassParameters->Output = GraphBuilder.CreateUAV(SceneColorTexture);
+
+		FIntRect ViewRect = View.ViewRect;
+
+		RDG_GPU_STAT_SCOPE(GraphBuilder, RayTracingDebug);
+
+		const FIntPoint GroupSize(FRayTracingDebugTraversalCS::ThreadGroupSizeX, FRayTracingDebugTraversalCS::ThreadGroupSizeY);
+		const FIntVector GroupCount = FComputeShaderUtils::GetGroupCount(ViewRect.Size(), GroupSize);
+		
+		TShaderRef<FRayTracingDebugTraversalCS> ComputeShader = GetGlobalShaderMap(FeatureLevel)->GetShader<FRayTracingDebugTraversalCS>();
+
+		FComputeShaderUtils::AddPass(GraphBuilder, RDG_EVENT_NAME("RayTracingDebug"), ComputeShader, PassParameters, GroupCount);
+		return;
+	}
+
+	// Debug modes other than barycentrics and traversal require pipeline support.
 	const bool bRayTracingPipeline = ShouldRenderRayTracingEffect(ERayTracingPipelineCompatibilityFlags::FullPipeline);
 	if (!bRayTracingPipeline)
 	{
