Height fog conversion to RDG.
Ran enginetest and other manual test with vfog, texture inscattering.
Preflight https://horde.devtools.epicgames.com/job/61af58e6813adc3fbb2fcb61

diff --git a/Engine/Shaders/Private/HeightFogCommon.ush b/Engine/Shaders/Private/HeightFogCommon.ush
--- a/Engine/Shaders/Private/HeightFogCommon.ush
+++ b/Engine/Shaders/Private/HeightFogCommon.ush
@@ -4,16 +4,16 @@
 	HeightFogCommon.usf: 
 =============================================================================*/
 
-#ifndef SUPPORT_FOG_INSCATTERING_TEXTURE
-#define SUPPORT_FOG_INSCATTERING_TEXTURE (FEATURE_LEVEL >= FEATURE_LEVEL_SM4)
+#ifndef PERMUTATION_SUPPORT_FOG_INSCATTERING_TEXTURE
+#define PERMUTATION_SUPPORT_FOG_INSCATTERING_TEXTURE (FEATURE_LEVEL >= FEATURE_LEVEL_SM4)
 #endif
 
-#ifndef SUPPORT_FOG_DIRECTIONAL_LIGHT_INSCATTERING
-#define SUPPORT_FOG_DIRECTIONAL_LIGHT_INSCATTERING 1
+#ifndef PERMUTATION_SUPPORT_FOG_DIRECTIONAL_LIGHT_INSCATTERING
+#define PERMUTATION_SUPPORT_FOG_DIRECTIONAL_LIGHT_INSCATTERING 1
 #endif
 
-#ifndef SUPPORT_VOLUMETRIC_FOG
-#define SUPPORT_VOLUMETRIC_FOG (FEATURE_LEVEL >= FEATURE_LEVEL_SM4)
+#ifndef PERMUTATION_SUPPORT_VOLUMETRIC_FOG
+#define PERMUTATION_SUPPORT_VOLUMETRIC_FOG (FEATURE_LEVEL >= FEATURE_LEVEL_SM4)
 #endif
 
 #if INSTANCED_STEREO
@@ -37,7 +37,7 @@ float3 ComputeInscatteringColor(float3 CameraToReceiver, float CameraToReceiverL
 {
 	half3 Inscattering = FogStruct.ExponentialFogColorParameter.xyz;
 
-#if SUPPORT_FOG_INSCATTERING_TEXTURE
+#if PERMUTATION_SUPPORT_FOG_INSCATTERING_TEXTURE
 	BRANCH
 	if (FogStruct.ExponentialFogParameters3.z > 0)
 	{
@@ -136,11 +136,11 @@ half4 GetExponentialHeightFog(float3 WorldPositionRelativeToCamera, float Exclud
 	half3 InscatteringColor = ComputeInscatteringColor(CameraToReceiver, CameraToReceiverLength);
 	half3 DirectionalInscattering = 0;
 
-#if SUPPORT_FOG_DIRECTIONAL_LIGHT_INSCATTERING
+#if PERMUTATION_SUPPORT_FOG_DIRECTIONAL_LIGHT_INSCATTERING
 	// if InscatteringLightDirection.w is negative then it's disabled, otherwise it holds directional inscattering start distance
 	BRANCH
 	if (FogStruct.InscatteringLightDirection.w >= 0
-	#if SUPPORT_FOG_INSCATTERING_TEXTURE
+	#if PERMUTATION_SUPPORT_FOG_INSCATTERING_TEXTURE
 		&& FogStruct.ExponentialFogParameters3.z == 0
 	#endif
 	)
@@ -200,7 +200,7 @@ half4 CalculateHeightFog(float3 WorldPositionRelativeToCamera)
 	
 	float ExcludeDistance = 0;
 
-#if SUPPORT_VOLUMETRIC_FOG
+#if PERMUTATION_SUPPORT_VOLUMETRIC_FOG
 	float InvCosAngle = length(WorldPositionRelativeToCamera) / dot(WorldPositionRelativeToCamera, View.ViewForward);
 	// Volumetric fog covers up to MaxDistance along ViewZ, exclude analytical fog from this range
 	ExcludeDistance = max(View.VolumetricFogMaxDistance * InvCosAngle, 0);
@@ -220,7 +220,7 @@ float4 CombineVolumetricFog(float4 GlobalFog, float3 VolumeUV, uint EyeIndex)
 {
 	float4 VolumetricFogLookup = float4(0, 0, 0, 1);
 
-#if SUPPORT_VOLUMETRIC_FOG
+#if PERMUTATION_SUPPORT_VOLUMETRIC_FOG
 	if (FogStruct.ApplyVolumetricFog > 0)
 	{
 	#if INSTANCED_STEREO
diff --git a/Engine/Source/Runtime/Renderer/Private/DeferredShadingRenderer.h b/Engine/Source/Runtime/Renderer/Private/DeferredShadingRenderer.h
--- a/Engine/Source/Runtime/Renderer/Private/DeferredShadingRenderer.h
+++ b/Engine/Source/Runtime/Renderer/Private/DeferredShadingRenderer.h
@@ -41,7 +41,6 @@ namespace LumenRadianceCache
 class FRenderLightParameters;
 
 struct FSceneWithoutWaterTextures;
-struct FHeightFogRenderingParameters;
 struct FRayTracingReflectionOptions;
 struct FHairStrandsTransmittanceMaskData;
 struct FVolumetricFogLocalLightFunctionInfo;
diff --git a/Engine/Source/Runtime/Renderer/Private/FogRendering.cpp b/Engine/Source/Runtime/Renderer/Private/FogRendering.cpp
--- a/Engine/Source/Runtime/Renderer/Private/FogRendering.cpp
+++ b/Engine/Source/Runtime/Renderer/Private/FogRendering.cpp
@@ -53,16 +53,6 @@ static TAutoConsoleVariable<float> CVarUpsampleJitterMultiplier(
 
 IMPLEMENT_GLOBAL_SHADER_PARAMETER_STRUCT(FFogUniformParameters, "FogStruct");
 
-struct FHeightFogRenderingParameters
-{
-	FRHITexture* LightShaftOcclusionRHI = nullptr;
-	FRHITexture* LinearDepthTextureRHI = nullptr;
-	FIntRect ViewRect;
-	float LinearDepthReadScale = 1.0f;
-	FVector4f LinearDepthMinMaxUV;
-	FVector4f LightShaftOcclusionMinMaxUV;
-};
-
 void SetupFogUniformParameters(FRDGBuilder& GraphBuilder, const FViewInfo& View, FFogUniformParameters& OutParameters)
 {
 	// Exponential Height Fog
@@ -114,130 +104,53 @@ TRDGUniformBufferRef<FFogUniformParameters> CreateFogUniformBuffer(FRDGBuilder&
 /** A vertex shader for rendering height fog. */
 class FHeightFogVS : public FGlobalShader
 {
-	DECLARE_SHADER_TYPE(FHeightFogVS,Global);
-public:
+	DECLARE_GLOBAL_SHADER(FHeightFogVS);
+	SHADER_USE_PARAMETER_STRUCT(FHeightFogVS, FGlobalShader);
+
+	BEGIN_SHADER_PARAMETER_STRUCT(FParameters, )
+		SHADER_PARAMETER_STRUCT_REF(FViewUniformShaderParameters, ViewUniformBuffer)
+	END_SHADER_PARAMETER_STRUCT()
 
 	static bool ShouldCompilePermutation(const FGlobalShaderPermutationParameters& Parameters)
 	{
 		return IsFeatureLevelSupported(Parameters.Platform, ERHIFeatureLevel::SM5);
 	}
-
-	FHeightFogVS( )	{ }
-	FHeightFogVS(const ShaderMetaType::CompiledShaderInitializerType& Initializer):
-		FGlobalShader(Initializer)
-	{
-	}
-
-	void SetParameters(FRHICommandList& RHICmdList, const FViewInfo& View)
-	{
-		FGlobalShader::SetParameters<FViewUniformShaderParameters>(RHICmdList, RHICmdList.GetBoundVertexShader(), View.ViewUniformBuffer);
-	}
-
-private:
 };
 
-IMPLEMENT_SHADER_TYPE(,FHeightFogVS,TEXT("/Engine/Private/HeightFogVertexShader.usf"),TEXT("Main"),SF_Vertex);
-
-enum class EHeightFogFeature
-{
-	HeightFog,
-	InscatteringTexture,
-	DirectionalLightInscattering,
-	HeightFogAndVolumetricFog,
-	InscatteringTextureAndVolumetricFog,
-	DirectionalLightInscatteringAndVolumetricFog
-};
+IMPLEMENT_GLOBAL_SHADER(FHeightFogVS, "/Engine/Private/HeightFogVertexShader.usf", "Main", SF_Vertex);
 
 /** A pixel shader for rendering exponential height fog. */
-template<EHeightFogFeature HeightFogFeature>
-class TExponentialHeightFogPS : public FGlobalShader
+class FExponentialHeightFogPS : public FGlobalShader
 {
-	DECLARE_SHADER_TYPE(TExponentialHeightFogPS,Global);
-public:
+	DECLARE_GLOBAL_SHADER(FExponentialHeightFogPS);
+	SHADER_USE_PARAMETER_STRUCT(FExponentialHeightFogPS, FGlobalShader);
+
+	class FSupportFogInScatteringTexture : SHADER_PERMUTATION_BOOL("PERMUTATION_SUPPORT_FOG_INSCATTERING_TEXTURE");
+	class FSupportFogDirectionalLightInScattering : SHADER_PERMUTATION_BOOL("PERMUTATION_SUPPORT_FOG_DIRECTIONAL_LIGHT_INSCATTERING");
+	class FSupportVolumetricFog : SHADER_PERMUTATION_BOOL("PERMUTATION_SUPPORT_VOLUMETRIC_FOG");
+	using FPermutationDomain = TShaderPermutationDomain<FSupportFogInScatteringTexture, FSupportFogDirectionalLightInScattering, FSupportVolumetricFog>;
+
+	BEGIN_SHADER_PARAMETER_STRUCT(FParameters, )
+		SHADER_PARAMETER_STRUCT_REF(FViewUniformShaderParameters, ViewUniformBuffer)
+		SHADER_PARAMETER_RDG_UNIFORM_BUFFER(FFogUniformParameters, FogUniformBuffer)
+		SHADER_PARAMETER_RDG_TEXTURE(Texture2D, OcclusionTexture)
+		SHADER_PARAMETER_SAMPLER(SamplerState, OcclusionSampler)
+		SHADER_PARAMETER_RDG_TEXTURE(Texture2D, LinearDepthTexture)
+		SHADER_PARAMETER_SAMPLER(SamplerState, LinearDepthSampler)
+		SHADER_PARAMETER(float, bOnlyOnRenderedOpaque)
+		SHADER_PARAMETER(float, bUseLinearDepthTexture)
+		SHADER_PARAMETER(float, UpsampleJitterMultiplier)
+		SHADER_PARAMETER(FVector4f, LinearDepthTextureMinMaxUV)
+		SHADER_PARAMETER(FVector4f, OcclusionTextureMinMaxUV)
+	END_SHADER_PARAMETER_STRUCT()
 
 	static bool ShouldCompilePermutation(const FGlobalShaderPermutationParameters& Parameters)
 	{
 		return IsFeatureLevelSupported(Parameters.Platform, ERHIFeatureLevel::SM5);
 	}
-
-	static void ModifyCompilationEnvironment(const FGlobalShaderPermutationParameters& Parameters, FShaderCompilerEnvironment& OutEnvironment)
-	{
-		OutEnvironment.SetDefine(TEXT("SUPPORT_FOG_INSCATTERING_TEXTURE"), HeightFogFeature == EHeightFogFeature::InscatteringTexture || HeightFogFeature == EHeightFogFeature::InscatteringTextureAndVolumetricFog);
-		OutEnvironment.SetDefine(TEXT("SUPPORT_FOG_DIRECTIONAL_LIGHT_INSCATTERING"), HeightFogFeature == EHeightFogFeature::DirectionalLightInscattering || HeightFogFeature == EHeightFogFeature::DirectionalLightInscatteringAndVolumetricFog);
-		OutEnvironment.SetDefine(TEXT("SUPPORT_VOLUMETRIC_FOG"), HeightFogFeature == EHeightFogFeature::HeightFogAndVolumetricFog || HeightFogFeature == EHeightFogFeature::InscatteringTextureAndVolumetricFog || HeightFogFeature == EHeightFogFeature::DirectionalLightInscatteringAndVolumetricFog);
-	}
-
-	TExponentialHeightFogPS( )	{ }
-	TExponentialHeightFogPS(const ShaderMetaType::CompiledShaderInitializerType& Initializer):
-		FGlobalShader(Initializer)
-	{
-		OcclusionTexture.Bind(Initializer.ParameterMap, TEXT("OcclusionTexture"));
-		OcclusionSampler.Bind(Initializer.ParameterMap, TEXT("OcclusionSampler"));
-		LinearDepthTexture.Bind(Initializer.ParameterMap, TEXT("LinearDepthTexture"));
-		LinearDepthSampler.Bind(Initializer.ParameterMap, TEXT("LinearDepthSampler"));
-		bOnlyOnRenderedOpaque.Bind(Initializer.ParameterMap, TEXT("bOnlyOnRenderedOpaque"));
-		bUseLinearDepthTexture.Bind(Initializer.ParameterMap, TEXT("bUseLinearDepthTexture"));
-		UpsampleJitterMultiplier.Bind(Initializer.ParameterMap, TEXT("UpsampleJitterMultiplier"));
-		LinearDepthTextureMinMaxUV.Bind(Initializer.ParameterMap, TEXT("LinearDepthTextureMinMaxUV"));
-		LightShaftOcclusionMinMaxUV.Bind(Initializer.ParameterMap, TEXT("OcclusionTextureMinMaxUV"));
-	}
-
-	void SetParameters(FRHICommandList& RHICmdList, const FViewInfo& View, const FHeightFogRenderingParameters& Params, FRHIUniformBuffer* FogUniformBuffer)
-	{
-		FGlobalShader::SetParameters<FViewUniformShaderParameters>(RHICmdList, RHICmdList.GetBoundPixelShader(), View.ViewUniformBuffer);
-		FGlobalShader::SetParameters<FFogUniformParameters>(RHICmdList, RHICmdList.GetBoundPixelShader(), FogUniformBuffer);
-
-		FRHITexture* TextureRHI = Params.LightShaftOcclusionRHI;
-		
-		if (!TextureRHI)
-		{
-			TextureRHI = GWhiteTexture->TextureRHI;
-		}
-
-		SetTextureParameter(
-			RHICmdList, 
-			RHICmdList.GetBoundPixelShader(),
-			OcclusionTexture, OcclusionSampler,
-			TStaticSamplerState<SF_Bilinear,AM_Clamp,AM_Clamp,AM_Clamp>::GetRHI(),
-			TextureRHI
-			);
-
-		const bool bUseLinearDepthTextureEnabled = Params.LinearDepthTextureRHI != nullptr;
-		FRHITexture* LinearDepthTextureRHI = bUseLinearDepthTextureEnabled ? Params.LinearDepthTextureRHI : GSystemTextures.DepthDummy->GetShaderResourceRHI();
-		SetTextureParameter(
-			RHICmdList, 
-			RHICmdList.GetBoundPixelShader(),
-			LinearDepthTexture, LinearDepthSampler,
-			TStaticSamplerState<SF_Point,AM_Clamp,AM_Clamp,AM_Clamp>::GetRHI(),
-			LinearDepthTextureRHI
-			);
-
-		SetShaderValue(RHICmdList, RHICmdList.GetBoundPixelShader(), bOnlyOnRenderedOpaque, View.bFogOnlyOnRenderedOpaque ? 1.0f : 0.0f);
-		SetShaderValue(RHICmdList, RHICmdList.GetBoundPixelShader(), bUseLinearDepthTexture, bUseLinearDepthTextureEnabled ? Params.LinearDepthReadScale : 0.0f);
-		SetShaderValue(RHICmdList, RHICmdList.GetBoundPixelShader(), LinearDepthTextureMinMaxUV, Params.LinearDepthMinMaxUV);
-		extern int32 GVolumetricFogGridPixelSize;
-		SetShaderValue(RHICmdList, RHICmdList.GetBoundPixelShader(), UpsampleJitterMultiplier, CVarUpsampleJitterMultiplier.GetValueOnRenderThread() * GVolumetricFogGridPixelSize);
-		SetShaderValue(RHICmdList, RHICmdList.GetBoundPixelShader(), LightShaftOcclusionMinMaxUV, Params.LightShaftOcclusionMinMaxUV);
-	}
-
-private:
-	LAYOUT_FIELD(FShaderResourceParameter, OcclusionTexture)
-	LAYOUT_FIELD(FShaderResourceParameter, OcclusionSampler)
-	LAYOUT_FIELD(FShaderResourceParameter, LinearDepthTexture);
-	LAYOUT_FIELD(FShaderResourceParameter, LinearDepthSampler);
-	LAYOUT_FIELD(FShaderParameter, bOnlyOnRenderedOpaque)
-	LAYOUT_FIELD(FShaderParameter, bUseLinearDepthTexture);
-	LAYOUT_FIELD(FShaderParameter, UpsampleJitterMultiplier);
-	LAYOUT_FIELD(FShaderParameter, LinearDepthTextureMinMaxUV);
-	LAYOUT_FIELD(FShaderParameter, LightShaftOcclusionMinMaxUV);
 };
 
-IMPLEMENT_SHADER_TYPE(template<>,TExponentialHeightFogPS<EHeightFogFeature::HeightFog>,TEXT("/Engine/Private/HeightFogPixelShader.usf"), TEXT("ExponentialPixelMain"),SF_Pixel)
-IMPLEMENT_SHADER_TYPE(template<>,TExponentialHeightFogPS<EHeightFogFeature::InscatteringTexture>,TEXT("/Engine/Private/HeightFogPixelShader.usf"), TEXT("ExponentialPixelMain"),SF_Pixel)
-IMPLEMENT_SHADER_TYPE(template<>,TExponentialHeightFogPS<EHeightFogFeature::DirectionalLightInscattering>,TEXT("/Engine/Private/HeightFogPixelShader.usf"), TEXT("ExponentialPixelMain"),SF_Pixel)
-IMPLEMENT_SHADER_TYPE(template<>,TExponentialHeightFogPS<EHeightFogFeature::HeightFogAndVolumetricFog>,TEXT("/Engine/Private/HeightFogPixelShader.usf"), TEXT("ExponentialPixelMain"),SF_Pixel)
-IMPLEMENT_SHADER_TYPE(template<>,TExponentialHeightFogPS<EHeightFogFeature::InscatteringTextureAndVolumetricFog>,TEXT("/Engine/Private/HeightFogPixelShader.usf"), TEXT("ExponentialPixelMain"),SF_Pixel)
-IMPLEMENT_SHADER_TYPE(template<>,TExponentialHeightFogPS<EHeightFogFeature::DirectionalLightInscatteringAndVolumetricFog>,TEXT("/Engine/Private/HeightFogPixelShader.usf"), TEXT("ExponentialPixelMain"),SF_Pixel)
+IMPLEMENT_GLOBAL_SHADER(FExponentialHeightFogPS, "/Engine/Private/HeightFogPixelShader.usf", "ExponentialPixelMain", SF_Pixel);
 
 /** The fog vertex declaration resource type. */
 class FFogVertexDeclaration : public FRenderResource
@@ -346,140 +259,117 @@ void FSceneRenderer::InitFogConstants()
 	}
 }
 
-/** Sets the bound shader state for either the per-pixel or per-sample fog pass. */
-static void SetFogShaders(FRHICommandList& RHICmdList, FGraphicsPipelineStateInitializer& GraphicsPSOInit, const FViewInfo& View, bool bShouldRenderVolumetricFog, const FHeightFogRenderingParameters& Params, FRHIUniformBuffer* FogUniformBuffer)
-{
-	TShaderMapRef<FHeightFogVS> VertexShader(View.ShaderMap);
-	GraphicsPSOInit.BoundShaderState.VertexDeclarationRHI = GFogVertexDeclaration.VertexDeclarationRHI;
-	GraphicsPSOInit.BoundShaderState.VertexShaderRHI = VertexShader.GetVertexShader();
-
-	if (bShouldRenderVolumetricFog)
-	{
-		if (View.FogInscatteringColorCubemap)
-		{
-			TShaderMapRef<TExponentialHeightFogPS<EHeightFogFeature::InscatteringTextureAndVolumetricFog> > ExponentialHeightFogPixelShader(View.ShaderMap);
+BEGIN_SHADER_PARAMETER_STRUCT(FFogPassParameters, )
+	SHADER_PARAMETER_STRUCT_INCLUDE(FHeightFogVS::FParameters, VS)
+	SHADER_PARAMETER_STRUCT_INCLUDE(FExponentialHeightFogPS::FParameters, PS)
+	SHADER_PARAMETER_RDG_UNIFORM_BUFFER(FSceneTextureUniformParameters, SceneTextures)
+	RENDER_TARGET_BINDING_SLOTS()
+END_SHADER_PARAMETER_STRUCT()
 
-			GraphicsPSOInit.BoundShaderState.PixelShaderRHI = ExponentialHeightFogPixelShader.GetPixelShader();
-			SetGraphicsPipelineState(RHICmdList, GraphicsPSOInit, 0);
-			VertexShader->SetParameters(RHICmdList, View);
-			ExponentialHeightFogPixelShader->SetParameters(RHICmdList, View, Params, FogUniformBuffer);
-		}
-		else if (View.bUseDirectionalInscattering)
-		{
-			TShaderMapRef<TExponentialHeightFogPS<EHeightFogFeature::DirectionalLightInscatteringAndVolumetricFog> > ExponentialHeightFogPixelShader(View.ShaderMap);
+static FFogPassParameters* CreateDefaultFogPassParameters(
+	FRDGBuilder& GraphBuilder,
+	const FViewInfo& View,
+	TRDGUniformBufferRef<FSceneTextureUniformParameters> SceneTexturesUniformbuffer,
+	TRDGUniformBufferRef<FFogUniformParameters>& FogUniformBuffer,
+	FRDGTextureRef LightShaftOcclusionTexture,
+	const FScreenPassTextureViewportParameters LightShaftParameters)
+{
+	extern int32 GVolumetricFogGridPixelSize;
+
+	FFogPassParameters* PassParameters = GraphBuilder.AllocParameters<FFogPassParameters>();
+	PassParameters->SceneTextures = SceneTexturesUniformbuffer;
+	PassParameters->VS.ViewUniformBuffer = GetShaderBinding(View.ViewUniformBuffer);
+	PassParameters->PS.ViewUniformBuffer = GetShaderBinding(View.ViewUniformBuffer);
+	PassParameters->PS.FogUniformBuffer = FogUniformBuffer;
+	PassParameters->PS.OcclusionTexture = LightShaftOcclusionTexture != nullptr ? LightShaftOcclusionTexture : GSystemTextures.GetWhiteDummy(GraphBuilder);
+	PassParameters->PS.OcclusionSampler = TStaticSamplerState<SF_Bilinear, AM_Clamp, AM_Clamp, AM_Clamp>::GetRHI();
+	PassParameters->PS.LinearDepthTexture = GSystemTextures.GetDepthDummy(GraphBuilder);
+	PassParameters->PS.LinearDepthSampler = TStaticSamplerState<SF_Point, AM_Clamp, AM_Clamp, AM_Clamp>::GetRHI();
+	PassParameters->PS.OcclusionTextureMinMaxUV = FVector4f(LightShaftParameters.UVViewportBilinearMin, LightShaftParameters.UVViewportBilinearMax);
+	PassParameters->PS.LinearDepthTextureMinMaxUV = FVector4f::Zero();
+	PassParameters->PS.UpsampleJitterMultiplier = CVarUpsampleJitterMultiplier.GetValueOnRenderThread() * GVolumetricFogGridPixelSize;
+	PassParameters->PS.bOnlyOnRenderedOpaque = View.bFogOnlyOnRenderedOpaque;
+	PassParameters->PS.bUseLinearDepthTexture = 0.0f;
+	return PassParameters;
+}
 
-			GraphicsPSOInit.BoundShaderState.PixelShaderRHI = ExponentialHeightFogPixelShader.GetPixelShader();
-			SetGraphicsPipelineState(RHICmdList, GraphicsPSOInit, 0);
-			VertexShader->SetParameters(RHICmdList, View);
-			ExponentialHeightFogPixelShader->SetParameters(RHICmdList, View, Params, FogUniformBuffer);
-		}
-		else
-		{
-			TShaderMapRef<TExponentialHeightFogPS<EHeightFogFeature::HeightFogAndVolumetricFog> > ExponentialHeightFogPixelShader(View.ShaderMap);
+static void RenderViewFog(
+	FRHICommandList& RHICmdList, 
+	const FViewInfo& View, 
+	FIntRect ViewRect, 
+	FFogPassParameters* PassParameters, 
+	bool bShouldRenderVolumetricFog)
+{
+	FGraphicsPipelineStateInitializer GraphicsPSOInit;
+	RHICmdList.ApplyCachedRenderTargets(GraphicsPSOInit);
+	RHICmdList.SetViewport(ViewRect.Min.X, ViewRect.Min.Y, 0.0f, ViewRect.Max.X, ViewRect.Max.Y, 1.0f);
 
-			GraphicsPSOInit.BoundShaderState.PixelShaderRHI = ExponentialHeightFogPixelShader.GetPixelShader();
-			SetGraphicsPipelineState(RHICmdList, GraphicsPSOInit, 0);
-			VertexShader->SetParameters(RHICmdList, View);
-			ExponentialHeightFogPixelShader->SetParameters(RHICmdList, View, Params, FogUniformBuffer);
-		}
-	}
-	else
-	{
-		GraphicsPSOInit.bDepthBounds = GSupportsDepthBoundsTest && CVarFogUseDepthBounds.GetValueOnAnyThread();
+	GraphicsPSOInit.RasterizerState = TStaticRasterizerState<FM_Solid, CM_None>::GetRHI();
+	GraphicsPSOInit.DepthStencilState = TStaticDepthStencilState<false, CF_Always>::GetRHI();
+	GraphicsPSOInit.PrimitiveType = PT_TriangleList;
 
-		if (GraphicsPSOInit.bDepthBounds)
-		{
-			// The fog can be set to start at a certain euclidean distance.
-			// clamp the value to be behind the near plane z
-			float FogStartDistance = FMath::Max(30.0f, View.ExponentialFogParameters.W);
-
-			// Here we compute the nearest z value the fog can start
-			// to skip shader execution on pixels that are closer.
-			// This means with a bigger distance specified more pixels are
-			// are culled and don't need to be rendered. This is faster if
-			// there is opaque content nearer than the computed z.
-			// This optimisation is achieved using depth bound tests.
-			// Mobile platofrms typically does not support that feature 
-			// but typically renders the world using forward shading 
-			// with height fog evaluated as part of the material vertex or pixel shader.
-			FMatrix InvProjectionMatrix = View.ViewMatrices.GetInvProjectionMatrix();
-			FVector ViewSpaceCorner = InvProjectionMatrix.TransformFVector4(FVector4(1, 1, 1, 1));
-			float Ratio = ViewSpaceCorner.Z / ViewSpaceCorner.Size();
-			FVector ViewSpaceStartFogPoint(0.0f, 0.0f, FogStartDistance * Ratio);
-			FVector4f ClipSpaceMaxDistance = View.ViewMatrices.GetProjectionMatrix().TransformPosition(ViewSpaceStartFogPoint);
-			float FogClipSpaceZ = ClipSpaceMaxDistance.Z / ClipSpaceMaxDistance.W;
-
-			if (bool(ERHIZBuffer::IsInverted))
-			{
-				RHICmdList.SetDepthBounds(0.0f, FogClipSpaceZ);
-			}
-			else
-			{
-				RHICmdList.SetDepthBounds(FogClipSpaceZ, 1.0f);
-			}
-		}
+	// disable alpha writes in order to preserve scene depth values on PC
+	GraphicsPSOInit.BlendState = TStaticBlendState<CW_RGB, BO_Add, BF_One, BF_SourceAlpha>::GetRHI();
 
-		if (View.FogInscatteringColorCubemap)
-		{
-			TShaderMapRef<TExponentialHeightFogPS<EHeightFogFeature::InscatteringTexture> > ExponentialHeightFogPixelShader(View.ShaderMap);
+	TShaderMapRef<FHeightFogVS> VertexShader(View.ShaderMap);
+	GraphicsPSOInit.BoundShaderState.VertexDeclarationRHI = GFogVertexDeclaration.VertexDeclarationRHI;
+	GraphicsPSOInit.BoundShaderState.VertexShaderRHI = VertexShader.GetVertexShader();
 
-			GraphicsPSOInit.BoundShaderState.PixelShaderRHI = ExponentialHeightFogPixelShader.GetPixelShader();
-			SetGraphicsPipelineState(RHICmdList, GraphicsPSOInit, 0);
-			VertexShader->SetParameters(RHICmdList, View);
-			ExponentialHeightFogPixelShader->SetParameters(RHICmdList, View, Params, FogUniformBuffer);
-		}
-		else if (View.bUseDirectionalInscattering)
+	const bool bUseFogInscatteringColorCubemap = View.FogInscatteringColorCubemap != nullptr;
+	FExponentialHeightFogPS::FPermutationDomain PsPermutationVector;
+	PsPermutationVector.Set<FExponentialHeightFogPS::FSupportFogInScatteringTexture>(bUseFogInscatteringColorCubemap);
+	PsPermutationVector.Set<FExponentialHeightFogPS::FSupportFogDirectionalLightInScattering>(!bUseFogInscatteringColorCubemap && View.bUseDirectionalInscattering);
+	PsPermutationVector.Set<FExponentialHeightFogPS::FSupportVolumetricFog>(bShouldRenderVolumetricFog);
+	TShaderMapRef<FExponentialHeightFogPS> PixelShader(View.ShaderMap, PsPermutationVector);
+	GraphicsPSOInit.BoundShaderState.PixelShaderRHI = PixelShader.GetPixelShader();
+
+	// Setup the depth bound optimization if possible on that platform.
+	GraphicsPSOInit.bDepthBounds = GSupportsDepthBoundsTest && CVarFogUseDepthBounds.GetValueOnAnyThread();
+	if (GraphicsPSOInit.bDepthBounds)
+	{
+		// The fog can be set to start at a certain euclidean distance.
+		// clamp the value to be behind the near plane z
+		float FogStartDistance = FMath::Max(30.0f, View.ExponentialFogParameters.W);
+
+		// Here we compute the nearest z value the fog can start
+		// to skip shader execution on pixels that are closer.
+		// This means with a bigger distance specified more pixels are
+		// are culled and don't need to be rendered. This is faster if
+		// there is opaque content nearer than the computed z.
+		// This optimization is achieved using depth bound tests.
+		// Mobile platforms typically does not support that feature 
+		// but typically renders the world using forward shading 
+		// with height fog evaluated as part of the material vertex or pixel shader.
+		FMatrix InvProjectionMatrix = View.ViewMatrices.GetInvProjectionMatrix();
+		FVector ViewSpaceCorner = InvProjectionMatrix.TransformFVector4(FVector4(1, 1, 1, 1));
+		float Ratio = ViewSpaceCorner.Z / ViewSpaceCorner.Size();
+		FVector ViewSpaceStartFogPoint(0.0f, 0.0f, FogStartDistance * Ratio);
+		FVector4f ClipSpaceMaxDistance = View.ViewMatrices.GetProjectionMatrix().TransformPosition(ViewSpaceStartFogPoint);
+		float FogClipSpaceZ = ClipSpaceMaxDistance.Z / ClipSpaceMaxDistance.W;
+
+		if (bool(ERHIZBuffer::IsInverted))
 		{
-			TShaderMapRef<TExponentialHeightFogPS<EHeightFogFeature::DirectionalLightInscattering> > ExponentialHeightFogPixelShader(View.ShaderMap);
-
-			GraphicsPSOInit.BoundShaderState.PixelShaderRHI = ExponentialHeightFogPixelShader.GetPixelShader();
-			SetGraphicsPipelineState(RHICmdList, GraphicsPSOInit, 0);
-			VertexShader->SetParameters(RHICmdList, View);
-			ExponentialHeightFogPixelShader->SetParameters(RHICmdList, View, Params, FogUniformBuffer);
+			RHICmdList.SetDepthBounds(0.0f, FogClipSpaceZ);
 		}
 		else
 		{
-			TShaderMapRef<TExponentialHeightFogPS<EHeightFogFeature::HeightFog> > ExponentialHeightFogPixelShader(View.ShaderMap);
-
-			GraphicsPSOInit.BoundShaderState.PixelShaderRHI = ExponentialHeightFogPixelShader.GetPixelShader();
-			SetGraphicsPipelineState(RHICmdList, GraphicsPSOInit, 0);
-			VertexShader->SetParameters(RHICmdList, View);
-			ExponentialHeightFogPixelShader->SetParameters(RHICmdList, View, Params, FogUniformBuffer);
+			RHICmdList.SetDepthBounds(FogClipSpaceZ, 1.0f);
 		}
 	}
-}
 
-static void RenderViewFog(FRHICommandList& RHICmdList, const FViewInfo& View, bool bShouldRenderVolumetricFog, const FHeightFogRenderingParameters& Params, FRHIUniformBuffer* FogUniformBuffer)
-{
-	FGraphicsPipelineStateInitializer GraphicsPSOInit;
-	RHICmdList.ApplyCachedRenderTargets(GraphicsPSOInit);
+	SetGraphicsPipelineState(RHICmdList, GraphicsPSOInit, 0);
 
-	// Set the device viewport for the view.
-	RHICmdList.SetViewport(Params.ViewRect.Min.X, Params.ViewRect.Min.Y, 0.0f, Params.ViewRect.Max.X, Params.ViewRect.Max.Y, 1.0f);
-			
-	GraphicsPSOInit.RasterizerState = TStaticRasterizerState<FM_Solid, CM_None>::GetRHI();
-			
-	// disable alpha writes in order to preserve scene depth values on PC
-	GraphicsPSOInit.BlendState = TStaticBlendState<CW_RGB, BO_Add, BF_One, BF_SourceAlpha>::GetRHI();
+	ClearUnusedGraphResources(VertexShader, &PassParameters->VS);
+	ClearUnusedGraphResources(PixelShader, &PassParameters->PS);
 
-	GraphicsPSOInit.DepthStencilState = TStaticDepthStencilState<false, CF_Always>::GetRHI();
-	GraphicsPSOInit.PrimitiveType = PT_TriangleList;
-
-	SetFogShaders(RHICmdList, GraphicsPSOInit, View, bShouldRenderVolumetricFog, Params, FogUniformBuffer);
+	SetShaderParameters(RHICmdList, VertexShader, VertexShader.GetVertexShader(), PassParameters->VS);
+	SetShaderParameters(RHICmdList, PixelShader, PixelShader.GetPixelShader(), PassParameters->PS);
 
 	// Draw a quad covering the view.
 	RHICmdList.SetStreamSource(0, GScreenSpaceVertexBuffer.VertexBufferRHI, 0);
 	RHICmdList.DrawIndexedPrimitive(GTwoTrianglesIndexBuffer.IndexBufferRHI, 0, 0, 4, 0, 2, 1);
 }
 
-BEGIN_SHADER_PARAMETER_STRUCT(FFogPassParameters, )
-	SHADER_PARAMETER_RDG_UNIFORM_BUFFER(FSceneTextureUniformParameters, SceneTextures)
-	SHADER_PARAMETER_RDG_UNIFORM_BUFFER(FFogUniformParameters, Fog)
-	RDG_TEXTURE_ACCESS(LightShaftOcclusionTexture, ERHIAccess::SRVGraphics)
-	RDG_TEXTURE_ACCESS(LinearDepthTexture, ERHIAccess::SRVGraphics)
-	RENDER_TARGET_BINDING_SLOTS()
-END_SHADER_PARAMETER_STRUCT()
-
 void FDeferredShadingSceneRenderer::RenderFog(
 	FRDGBuilder& GraphBuilder,
 	const FMinimalSceneTextures& SceneTextures,
@@ -504,29 +394,18 @@ void FDeferredShadingSceneRenderer::RenderFog(
 
 				TRDGUniformBufferRef<FFogUniformParameters> FogUniformBuffer = CreateFogUniformBuffer(GraphBuilder, View);
 
-				FFogPassParameters* PassParameters = GraphBuilder.AllocParameters<FFogPassParameters>();
-				PassParameters->SceneTextures = SceneTextures.UniformBuffer;
-				PassParameters->Fog = FogUniformBuffer;
-				PassParameters->LightShaftOcclusionTexture = LightShaftOcclusionTexture;
+				const FScreenPassTextureViewport SceneViewport(SceneTextures.Config.Extent, View.ViewRect);
+				const FScreenPassTextureViewport OutputViewport(GetDownscaledViewport(SceneViewport, GetLightShaftDownsampleFactor()));
+				const FScreenPassTextureViewportParameters LightShaftParameters = GetScreenPassTextureViewportParameters(OutputViewport);
+
+				FFogPassParameters* PassParameters = CreateDefaultFogPassParameters(GraphBuilder, View, SceneTextures.UniformBuffer, FogUniformBuffer, LightShaftOcclusionTexture, LightShaftParameters);
 				PassParameters->RenderTargets[0] = FRenderTargetBinding(SceneTextures.Color.Target, ERenderTargetLoadAction::ELoad);
 				PassParameters->RenderTargets.DepthStencil = FDepthStencilBinding(SceneTextures.Depth.Target, ERenderTargetLoadAction::ELoad, ERenderTargetLoadAction::ELoad, FExclusiveDepthStencil::DepthRead_StencilWrite);
 
-				FIntPoint SceneTexturesExtent = SceneTextures.Config.Extent;
-
-				GraphBuilder.AddPass(RDG_EVENT_NAME("Fog"), PassParameters, ERDGPassFlags::Raster, [this, &View, SceneTexturesExtent, bShouldRenderVolumetricFog, LightShaftOcclusionTexture, FogUniformBuffer](FRHICommandList& RHICmdList)
+				GraphBuilder.AddPass(RDG_EVENT_NAME("Fog"), PassParameters, ERDGPassFlags::Raster, 
+					[this, &View, PassParameters, bShouldRenderVolumetricFog](FRHICommandList& RHICmdList)
 				{
-					FHeightFogRenderingParameters Parameters;
-					Parameters.ViewRect = View.ViewRect;
-					if (LightShaftOcclusionTexture)
-					{
-						Parameters.LightShaftOcclusionRHI = LightShaftOcclusionTexture->GetRHI();
-
-						const FScreenPassTextureViewport SceneViewport(SceneTexturesExtent, View.ViewRect);
-						const FScreenPassTextureViewport OutputViewport(GetDownscaledViewport(SceneViewport, GetLightShaftDownsampleFactor()));
-						const FScreenPassTextureViewportParameters LightShaftParameters = GetScreenPassTextureViewportParameters(OutputViewport);
-						Parameters.LightShaftOcclusionMinMaxUV = FVector4f(LightShaftParameters.UVViewportBilinearMin, LightShaftParameters.UVViewportBilinearMax);
-					}
-					RenderViewFog(RHICmdList, View, bShouldRenderVolumetricFog, Parameters, FogUniformBuffer->GetRHI());
+					RenderViewFog(RHICmdList, View, View.ViewRect, PassParameters, bShouldRenderVolumetricFog);
 				});
 			}
 		}
@@ -551,9 +430,7 @@ void FDeferredShadingSceneRenderer::RenderUnderWaterFog(
 		const bool bShouldRenderVolumetricFog = ShouldRenderVolumetricFog();
 
 		// This must match SINGLE_LAYER_WATER_DEPTH_SCALE from SingleLayerWaterCommon.ush and SingleLayerWaterComposite.usf.
-		const float kSingleLayerWaterDepthScale = 100.0f; 
-		FHeightFogRenderingParameters Parameters;
-		Parameters.LinearDepthReadScale = kSingleLayerWaterDepthScale;
+		const float kSingleLayerWaterDepthScale = 100.0f;
 
 		for (int32 ViewIndex = 0; ViewIndex < Views.Num(); ViewIndex++)
 		{
@@ -563,23 +440,23 @@ void FDeferredShadingSceneRenderer::RenderUnderWaterFog(
 				RDG_EVENT_SCOPE_CONDITIONAL(GraphBuilder, Views.Num() > 1, "View%d", ViewIndex);
 				RDG_GPU_MASK_SCOPE(GraphBuilder, View.GPUMask);
 
+				const auto& SceneWithoutWaterView = SceneWithoutWaterTextures.Views[ViewIndex];
+
 				TRDGUniformBufferRef<FFogUniformParameters> FogUniformBuffer = CreateFogUniformBuffer(GraphBuilder, View);
 
-				FFogPassParameters* PassParameters = GraphBuilder.AllocParameters<FFogPassParameters>();
-				PassParameters->SceneTextures = SceneTexturesWithDepth;
-				PassParameters->Fog = FogUniformBuffer;
-				PassParameters->LinearDepthTexture = LinearDepthTexture;
-				PassParameters->RenderTargets[0] = FRenderTargetBinding(SceneWithoutWaterTextures.ColorTexture, ERenderTargetLoadAction::ELoad);
+				// TODO add support for occlusion texture on water
+				FRDGTextureRef LightShaftOcclusionTexture = nullptr;	
+				FScreenPassTextureViewportParameters LightShaftParameters = FScreenPassTextureViewportParameters();
 
-				const auto& SceneWithoutWaterView = SceneWithoutWaterTextures.Views[ViewIndex];
-				Parameters.ViewRect = SceneWithoutWaterView.ViewRect;
-				Parameters.LinearDepthMinMaxUV = SceneWithoutWaterView.MinMaxUV;
+				FFogPassParameters* PassParameters = CreateDefaultFogPassParameters(GraphBuilder, View, SceneTexturesWithDepth, FogUniformBuffer, LightShaftOcclusionTexture, LightShaftParameters);
+				PassParameters->PS.LinearDepthTexture = LinearDepthTexture;
+				PassParameters->PS.bUseLinearDepthTexture = kSingleLayerWaterDepthScale;
+				PassParameters->PS.LinearDepthTextureMinMaxUV = SceneWithoutWaterView.MinMaxUV;
+				PassParameters->RenderTargets[0] = FRenderTargetBinding(SceneWithoutWaterTextures.ColorTexture, ERenderTargetLoadAction::ELoad);
 
-				GraphBuilder.AddPass(RDG_EVENT_NAME("Fog"), PassParameters, ERDGPassFlags::Raster, [this, &View, bShouldRenderVolumetricFog, Parameters, LinearDepthTexture, FogUniformBuffer](FRHICommandList& RHICmdList)
+				GraphBuilder.AddPass(RDG_EVENT_NAME("FogBehindWater"), PassParameters, ERDGPassFlags::Raster, [this, &View, SceneWithoutWaterView, PassParameters, bShouldRenderVolumetricFog](FRHICommandList& RHICmdList)
 				{
-					FHeightFogRenderingParameters LocalParameters = Parameters;
-					LocalParameters.LinearDepthTextureRHI = LinearDepthTexture->GetRHI();
-					RenderViewFog(RHICmdList, View, bShouldRenderVolumetricFog, LocalParameters, FogUniformBuffer->GetRHI());
+					RenderViewFog(RHICmdList, View, SceneWithoutWaterView.ViewRect, PassParameters, bShouldRenderVolumetricFog);
 				});
 			}
 		}
