Fixes a bunch of vector precisions in post processing

diff --git a/Engine/Source/Runtime/Renderer/Private/PostProcess/DiaphragmDOF.cpp b/Engine/Source/Runtime/Renderer/Private/PostProcess/DiaphragmDOF.cpp
--- a/Engine/Source/Runtime/Renderer/Private/PostProcess/DiaphragmDOF.cpp
+++ b/Engine/Source/Runtime/Renderer/Private/PostProcess/DiaphragmDOF.cpp
@@ -390,15 +390,15 @@ const TCHAR* GetEventName(EDiaphragmDOFDilateCocMode e)
 }
 
 // Returns X and Y for F(M) = saturate(M * X + Y) so that F(LowM) = 0 and F(HighM) = 1
-FVector2D GenerateSaturatedAffineTransformation(float LowM, float HighM)
+FVector2f GenerateSaturatedAffineTransformation(float LowM, float HighM)
 {
 	float X = 1.0f / (HighM - LowM);
-	return FVector2D(X, -X * LowM);
+	return FVector2f(X, -X * LowM);
 }
 
 // Affine transformtations that always return 0 or 1.
-const FVector2D kContantlyPassingAffineTransformation(0, 1);
-const FVector2D kContantlyBlockingAffineTransformation(0, 0);
+const FVector2f kContantlyPassingAffineTransformation(0, 1);
+const FVector2f kContantlyBlockingAffineTransformation(0, 0);
 
 
 /** Base shader class for diaphragm DOF. */
@@ -1570,7 +1570,7 @@ FRDGTextureRef DiaphragmDOF::AddPasses(
 			PassParameters->CommonParameters = CommonParameters;
 			SetCocModelParameters(&PassParameters->CocModel, CocModel, CocRadiusBasis);
 			PassParameters->ViewportRect = FIntRect(FIntPoint::ZeroValue, PassViewSize);
-			PassParameters->CocRadiusBasis = FVector2D(GatheringViewSize.X, PreprocessViewSize.X);
+			PassParameters->CocRadiusBasis = FVector2f(GatheringViewSize.X, PreprocessViewSize.X);
 			PassParameters->SceneColorTexture = InputSceneColor;
 			PassParameters->SceneDepthTexture = SceneTextures.SceneDepthTexture;
 		
@@ -1670,10 +1670,10 @@ FRDGTextureRef DiaphragmDOF::AddPasses(
 			FDiaphragmDOFCocFlattenCS::FParameters* PassParameters = GraphBuilder.AllocParameters<FDiaphragmDOFCocFlattenCS::FParameters>();
 			PassParameters->CommonParameters = CommonParameters;
 			PassParameters->ViewportRect = FIntRect(0, 0, GatheringViewSize.X, GatheringViewSize.Y);
-			PassParameters->ThreadIdToBufferUV = FVector2D(
+			PassParameters->ThreadIdToBufferUV = FVector2f(
 				PreprocessViewSize.X / float(GatheringViewSize.X * SrcSize.X),
 				PreprocessViewSize.Y / float(GatheringViewSize.Y * SrcSize.Y));
-			PassParameters->MaxBufferUV = FVector2D(
+			PassParameters->MaxBufferUV = FVector2f(
 				(PreprocessViewSize.X - 1.0f) / float(SrcSize.X),
 				(PreprocessViewSize.Y - 1.0f) / float(SrcSize.Y));
 			PassParameters->GatherInput = HalfResGatherInputTextures;
@@ -1898,7 +1898,7 @@ FRDGTextureRef DiaphragmDOF::AddPasses(
 			FDiaphragmDOFDownsampleCS::FParameters* PassParameters = GraphBuilder.AllocParameters<FDiaphragmDOFDownsampleCS::FParameters>();
 			PassParameters->CommonParameters = CommonParameters;
 			PassParameters->ViewportRect = FIntRect(0, 0, PassViewSize.X, PassViewSize.Y);
-			PassParameters->MaxBufferUV = FVector2D(
+			PassParameters->MaxBufferUV = FVector2f(
 				(PreprocessViewSize.X - 0.5f) / float(SrcSize.X),
 				(PreprocessViewSize.Y - 0.5f) / float(SrcSize.Y));
 			PassParameters->OutputCocRadiusMultiplier = PreProcessingToProcessingCocRadiusFactor;
@@ -1946,7 +1946,7 @@ FRDGTextureRef DiaphragmDOF::AddPasses(
 
 			FDiaphragmDOFReduceCS::FParameters* PassParameters = GraphBuilder.AllocParameters<FDiaphragmDOFReduceCS::FParameters>();
 			PassParameters->ViewportRect = FIntRect(0, 0, PassViewSize.X, PassViewSize.Y);
-			PassParameters->MaxInputBufferUV = FVector2D(
+			PassParameters->MaxInputBufferUV = FVector2f(
 				(PreprocessViewSize.X - 0.5f) / SrcSize.X,
 				(PreprocessViewSize.Y - 0.5f) / SrcSize.Y);
 			PassParameters->MaxScatteringGroupCount = MaxScatteringGroupCount;
@@ -2010,7 +2010,7 @@ FRDGTextureRef DiaphragmDOF::AddPasses(
 
 			FDiaphragmDOFReduceCS::FParameters* PassParameters = GraphBuilder.AllocParameters<FDiaphragmDOFReduceCS::FParameters>();
 			PassParameters->ViewportRect = FIntRect(0, 0, PassViewSize.X, PassViewSize.Y);
-			PassParameters->MaxInputBufferUV = FVector2D(
+			PassParameters->MaxInputBufferUV = FVector2f(
 				(PreprocessViewSize.X - 0.5f) / SrcSize.X,
 				(PreprocessViewSize.Y - 0.5f) / SrcSize.Y);
 
@@ -2229,9 +2229,9 @@ FRDGTextureRef DiaphragmDOF::AddPasses(
 			PermutationVector = FDiaphragmDOFGatherCS::RemapPermutation(PermutationVector);
 
 			// Affine transformtation to control whether a CocRadius is considered or not.
-			FVector2D ConsiderCocRadiusAffineTransformation0 = kContantlyPassingAffineTransformation;
-			FVector2D ConsiderCocRadiusAffineTransformation1 = kContantlyPassingAffineTransformation;
-			FVector2D ConsiderAbsCocRadiusAffineTransformation = kContantlyPassingAffineTransformation;
+			FVector2f ConsiderCocRadiusAffineTransformation0 = kContantlyPassingAffineTransformation;
+			FVector2f ConsiderCocRadiusAffineTransformation1 = kContantlyPassingAffineTransformation;
+			FVector2f ConsiderAbsCocRadiusAffineTransformation = kContantlyPassingAffineTransformation;
 			{
 				// Gathering scalability.
 				const float GatheringScalingDownFactor = float(PreprocessViewSize.X) / float(GatheringViewSize.X);
@@ -2275,13 +2275,13 @@ FRDGTextureRef DiaphragmDOF::AddPasses(
 			PassParameters->ViewportSize = FVector4f(GatheringViewSize.X, GatheringViewSize.Y, 1.0f / GatheringViewSize.X, 1.0f / GatheringViewSize.Y);
 			PassParameters->ViewportRect = FIntRect(0, 0, GatheringViewSize.X, GatheringViewSize.Y);
 			PassParameters->TemporalJitterPixels = View.TemporalJitterPixels;
-			PassParameters->DispatchThreadIdToInputBufferUV = FVector2D(
+			PassParameters->DispatchThreadIdToInputBufferUV = FVector2f(
 				float(PreprocessViewSize.X) / float(GatheringViewSize.X * SrcSize.X),
 				float(PreprocessViewSize.Y) / float(GatheringViewSize.Y * SrcSize.Y));;
 			PassParameters->ConsiderCocRadiusAffineTransformation0 = ConsiderCocRadiusAffineTransformation0;
 			PassParameters->ConsiderCocRadiusAffineTransformation1 = ConsiderCocRadiusAffineTransformation1;
 			PassParameters->ConsiderAbsCocRadiusAffineTransformation = ConsiderAbsCocRadiusAffineTransformation;
-			PassParameters->InputBufferUVToOutputPixel = FVector2D(
+			PassParameters->InputBufferUVToOutputPixel = FVector2f(
 				float(SrcSize.X * GatheringViewSize.X) / float(PreprocessViewSize.X),
 				float(SrcSize.Y * GatheringViewSize.Y) / float(PreprocessViewSize.Y));
 			PassParameters->MipBias = FMath::Log2(float(PreprocessViewSize.X) / float(GatheringViewSize.X));
@@ -2292,7 +2292,7 @@ FRDGTextureRef DiaphragmDOF::AddPasses(
 			PassParameters->CommonParameters = CommonParameters;
 		
 			PassParameters->GatherInputSize = FVector4f(SrcSize.X, SrcSize.Y, 1.0f / SrcSize.X, 1.0f / SrcSize.Y);
-			PassParameters->GatherInputViewportSize = FVector2D(PreprocessViewSize.X, PreprocessViewSize.Y);
+			PassParameters->GatherInputViewportSize = FVector2f(PreprocessViewSize.X, PreprocessViewSize.Y);
 			PassParameters->GatherInput = ReducedGatherInputTextures;
 		
 			PassParameters->TileClassification = TileClassificationTextures;
@@ -2365,7 +2365,7 @@ FRDGTextureRef DiaphragmDOF::AddPasses(
 
 			FDiaphragmDOFPostfilterCS::FParameters* PassParameters = GraphBuilder.AllocParameters<FDiaphragmDOFPostfilterCS::FParameters>();
 			PassParameters->ViewportRect = FIntRect(0, 0, GatheringViewSize.X, GatheringViewSize.Y);
-			PassParameters->MaxInputBufferUV = FVector2D(
+			PassParameters->MaxInputBufferUV = FVector2f(
 				(GatheringViewSize.X - 0.5f) / float(RefBufferSize.X),
 				(GatheringViewSize.Y - 0.5f) / float(RefBufferSize.Y));
 			PassParameters->TileDecisionParameters.MinGatherRadius = MaxRecombineAbsCocRadius - 1;
@@ -2631,7 +2631,7 @@ FRDGTextureRef DiaphragmDOF::AddPasses(
 		PassParameters->ViewportRect = PassViewRect;
 		PassParameters->ViewportSize = FVector4f(PassViewRect.Width(), PassViewRect.Height(), 1.0f / PassViewRect.Width(), 1.0f / PassViewRect.Height());
 		PassParameters->TemporalJitterPixels = View.TemporalJitterPixels;
-		PassParameters->DOFBufferUVMax = FVector2D(
+		PassParameters->DOFBufferUVMax = FVector2f(
 			(GatheringViewSize.X - 0.5f) / float(RefBufferSize.X),
 			(GatheringViewSize.Y - 0.5f) / float(RefBufferSize.Y));
 
@@ -2657,7 +2657,7 @@ FRDGTextureRef DiaphragmDOF::AddPasses(
 		PassParameters->FullResDepthTexture = GraphBuilder.CreateSRV(FRDGTextureSRVDesc::CreateForMetaData(SceneTextures.SceneDepthTexture, ERDGTextureMetaDataAccess::Depth));
 		PassParameters->LowResDepthTexture = GraphBuilder.CreateSRV(FRDGTextureSRVDesc::CreateForMetaData(SeparateTranslucencyDepth, ERDGTextureMetaDataAccess::Depth));
 		const FIntPoint LowResExtent = SeparateTranslucency->Desc.Extent;
-		PassParameters->SeparateTranslucencyTextureLowResExtentInverse = FVector2D(1.0f / LowResExtent.X, 1.0f / LowResExtent.Y);
+		PassParameters->SeparateTranslucencyTextureLowResExtentInverse = FVector2f(1.0f / LowResExtent.X, 1.0f / LowResExtent.Y);
 
 		if (!bGatherForeground && !bGatherBackground)
 		{
diff --git a/Engine/Source/Runtime/Renderer/Private/PostProcess/PostProcessCombineLUTs.cpp b/Engine/Source/Runtime/Renderer/Private/PostProcess/PostProcessCombineLUTs.cpp
--- a/Engine/Source/Runtime/Renderer/Private/PostProcess/PostProcessCombineLUTs.cpp
+++ b/Engine/Source/Runtime/Renderer/Private/PostProcess/PostProcessCombineLUTs.cpp
@@ -86,7 +86,7 @@ FColorRemapParameters GetColorRemapParameters()
 	float a = ColorTransform.MaxValue - ColorTransform.MinValue - b;
 
 	FColorRemapParameters Parameters;
-	Parameters.MappingPolynomial = FVector(a, b, c);
+	Parameters.MappingPolynomial = FVector3f(a, b, c);
 	return Parameters;
 }
 
@@ -435,7 +435,7 @@ FRDGTextureRef AddCombineLUTPass(FRDGBuilder& GraphBuilder, const FViewInfo& Vie
 	{
 		FLUTBlenderCS::FParameters* PassParameters = GraphBuilder.AllocParameters<FLUTBlenderCS::FParameters>();
 		GetCombineLUTParameters(PassParameters->CombineLUT, View, LocalTextures, LocalWeights, LocalCount);
-		PassParameters->OutputExtentInverse = FVector2D(1.0f, 1.0f) / FVector2D(OutputViewSize);
+		PassParameters->OutputExtentInverse = FVector2f(1.0f, 1.0f) / FVector2f(OutputViewSize);
 		PassParameters->RWOutputTexture = GraphBuilder.CreateUAV(OutputTexture);
 
 		const bool ShouldSkipTemperature = FMath::IsNearlyEqual(PassParameters->CombineLUT.WhiteTemp, DefaultTemperature) && FMath::IsNearlyEqual(PassParameters->CombineLUT.WhiteTint, DefaultTint);
diff --git a/Engine/Source/Runtime/Renderer/Private/PostProcess/PostProcessCompositeEditorPrimitives.cpp b/Engine/Source/Runtime/Renderer/Private/PostProcess/PostProcessCompositeEditorPrimitives.cpp
--- a/Engine/Source/Runtime/Renderer/Private/PostProcess/PostProcessCompositeEditorPrimitives.cpp
+++ b/Engine/Source/Runtime/Renderer/Private/PostProcess/PostProcessCompositeEditorPrimitives.cpp
@@ -337,7 +337,7 @@ FScreenPassTexture AddEditorPrimitivePass(
 		NumSamples);
 
 	FScreenPassTexture SceneDepth = Inputs.SceneDepth;
-	FVector2D SceneDepthJitter = View.TemporalJitterPixels;
+	FVector2f SceneDepthJitter = View.TemporalJitterPixels;
 
 	// The editor primitive composition pass is also used when rendering VMI_WIREFRAME in order to use MSAA.
 	// So we need to check whether the editor primitives are enabled inside this function.
@@ -427,7 +427,7 @@ FScreenPassTexture AddEditorPrimitivePass(
 				}
 
 				SceneDepth = History;
-				SceneDepthJitter = FVector2D::ZeroVector;
+				SceneDepthJitter = FVector2f::ZeroVector;
 			}
 
 			// Populate the MSAA depth buffer from depth buffer or temporally upscaled depth buffer
diff --git a/Engine/Source/Runtime/Renderer/Private/PostProcess/PostProcessFFTBloom.cpp b/Engine/Source/Runtime/Renderer/Private/PostProcess/PostProcessFFTBloom.cpp
--- a/Engine/Source/Runtime/Renderer/Private/PostProcess/PostProcessFFTBloom.cpp
+++ b/Engine/Source/Runtime/Renderer/Private/PostProcess/PostProcessFFTBloom.cpp
@@ -271,7 +271,7 @@ struct FFFTBloomIntermediates
 	// this defines the size of the FFT in each direction.
 	FIntPoint FrequencySize;
 
-	FVector PreFilter;
+	FVector3f PreFilter;
 
 	float KernelSupportScale = 0.0f;
 	float KernelSupportScaleClamp = 0.0f;
@@ -301,7 +301,7 @@ void InitDomainAndGetKernel(
 	check(BloomConvolutionTextureResource && PhysicalSpaceKernelTextureRef);
 
 	const float BloomConvolutionSize = PPSettings.BloomConvolutionSize;
-	const FVector2D CenterUV = PPSettings.BloomConvolutionCenterUV; // TODO: remove
+	const FVector2f CenterUV = PPSettings.BloomConvolutionCenterUV; // TODO: remove
 
 	// Our frequency storage layout adds two elements to the first transform direction. 
 	const FIntPoint FrequencyPadding = (Intermediates.bDoHorizontalFirst) ? FIntPoint(2, 0) : FIntPoint(0, 2);
@@ -683,7 +683,7 @@ FFFTBloomOutput AddFFTBloomPass(FRDGBuilder& GraphBuilder, const FViewInfo& View
 		Intermediates.KernelSupportScaleClamp = KernelSupportScaleClamp;
 
 		// The pre-filter boost parameters for bright pixels. Because the Convolution PP work in pre-exposure space, the min and max needs adjustment.
-		Intermediates.PreFilter = FVector(PPSettings.BloomConvolutionPreFilterMin, PPSettings.BloomConvolutionPreFilterMax, PPSettings.BloomConvolutionPreFilterMult);
+		Intermediates.PreFilter = FVector3f(PPSettings.BloomConvolutionPreFilterMin, PPSettings.BloomConvolutionPreFilterMax, PPSettings.BloomConvolutionPreFilterMult);
 
 		// The length of the a side of the square kernel image in pixels
 		const int32 KernelSize = FMath::CeilToInt(KernelSupportScale * FMath::Max(Intermediates.ImageSize.X, Intermediates.ImageSize.Y));
diff --git a/Engine/Source/Runtime/Renderer/Private/PostProcess/PostProcessLensFlares.cpp b/Engine/Source/Runtime/Renderer/Private/PostProcess/PostProcessLensFlares.cpp
--- a/Engine/Source/Runtime/Renderer/Private/PostProcess/PostProcessLensFlares.cpp
+++ b/Engine/Source/Runtime/Renderer/Private/PostProcess/PostProcessLensFlares.cpp
@@ -271,7 +271,7 @@ FScreenPassTexture AddLensFlaresPass(
 
 	const FIntRect OutputViewRect = BloomViewport.Rect;
 
-	const FVector2D OutputCenter = FVector2D(OutputViewRect.Min + OutputViewRect.Max) / 2;
+	const FVector2f OutputCenter = FVector2f(OutputViewRect.Min + OutputViewRect.Max) / 2;
 
 	// Scales normalized flare tint alpha to a viewport scale factor.
 	const float AlphaScale = static_cast<float>(Inputs.LensFlareCount - 1);
@@ -296,9 +296,9 @@ FScreenPassTexture AddLensFlaresPass(
 		// Alpha of the tint color is used to derive a scale of the flare quad.
 		const float FinalOutputScale = (LensFlareTint.A * AlphaScale + AlphaBias) * GuardBandScale;
 
-		const FVector2D QuadSize = FVector2D(OutputViewRect.Size()) * FinalOutputScale;
+		const FVector2f QuadSize = FVector2f(OutputViewRect.Size()) * FinalOutputScale;
 
-		const FVector2D QuadOffset = OutputCenter - 0.5f * QuadSize;
+		const FVector2f QuadOffset = OutputCenter - 0.5f * QuadSize;
 
 		TShaderMapRef<FScreenPassVS> VertexShader(View.ShaderMap);
 		TShaderMapRef<FLensFlareCompositePS> PixelShader(View.ShaderMap);
diff --git a/Engine/Source/Runtime/Renderer/Private/PostProcess/PostProcessLocalExposure.cpp b/Engine/Source/Runtime/Renderer/Private/PostProcess/PostProcessLocalExposure.cpp
--- a/Engine/Source/Runtime/Renderer/Private/PostProcess/PostProcessLocalExposure.cpp
+++ b/Engine/Source/Runtime/Renderer/Private/PostProcess/PostProcessLocalExposure.cpp
@@ -94,7 +94,7 @@ FRDGTextureRef AddLocalExposureBlurredLogLuminancePass(
 		GaussianBlurInputs.NameY = TEXT("LocalExposureGaussianY");
 		GaussianBlurInputs.Filter = FScreenPassTexture(GaussianLumSetupTexture, InputTexture.ViewRect);
 		GaussianBlurInputs.TintColor = FLinearColor::White;
-		GaussianBlurInputs.CrossCenterWeight = FVector2D::ZeroVector;
+		GaussianBlurInputs.CrossCenterWeight = FVector2f::ZeroVector;
 		GaussianBlurInputs.KernelSizePercent = View.FinalPostProcessSettings.LocalExposureBlurredLuminanceKernelSizePercent;
 		GaussianBlurInputs.UseMirrorAddressMode = true;
 
diff --git a/Engine/Source/Runtime/Renderer/Private/PostProcess/PostProcessTonemap.cpp b/Engine/Source/Runtime/Renderer/Private/PostProcess/PostProcessTonemap.cpp
--- a/Engine/Source/Runtime/Renderer/Private/PostProcess/PostProcessTonemap.cpp
+++ b/Engine/Source/Runtime/Renderer/Private/PostProcess/PostProcessTonemap.cpp
@@ -308,7 +308,7 @@ FTonemapperOutputDeviceParameters GetTonemapperOutputDeviceParameters(const FSce
 		OutputDeviceValue = ETonemapperOutputDevice::ExplicitGammaMapping;
 	}
 
-	FVector InvDisplayGammaValue;
+	FVector3f InvDisplayGammaValue;
 	InvDisplayGammaValue.X = 1.0f / Family.RenderTarget->GetDisplayGamma();
 	InvDisplayGammaValue.Y = 2.2f / Family.RenderTarget->GetDisplayGamma();
 	InvDisplayGammaValue.Z = 1.0f / FMath::Max(Gamma, 1.0f);
@@ -725,12 +725,12 @@ FScreenPassTexture AddTonemapPass(FRDGBuilder& GraphBuilder, const FViewInfo& Vi
 		int32 RandomSequenceLength = CVarFilmGrainSequenceLength.GetValueOnRenderThread();
 		int32 RandomSequenceIndex = (View.ViewState ? View.ViewState->FrameIndex : 0) % RandomSequenceLength;
 
-		FVector2D RandomGrainTextureUVOffset;
+		FVector2f RandomGrainTextureUVOffset;
 		RandomGrainTextureUVOffset.X = Halton(RandomSequenceIndex + 1, 2);
 		RandomGrainTextureUVOffset.Y = Halton(RandomSequenceIndex + 1, 3);
 
-		FVector2D TextureSize(View.FilmGrainTexture->GetSizeX(), View.FilmGrainTexture->GetSizeY());
-		FVector2D OutputSizeF(1920.0f, 1920.0f * float(OutputViewport.Rect.Height()) / float(OutputViewport.Rect.Width()));
+		FVector2f TextureSize(View.FilmGrainTexture->GetSizeX(), View.FilmGrainTexture->GetSizeY());
+		FVector2f OutputSizeF(1920.0f, 1920.0f * float(OutputViewport.Rect.Height()) / float(OutputViewport.Rect.Width()));
 
 		CommonParameters.FilmGrain.ScreenPosToFilmGrainTextureUV = 
 			FScreenTransform::ScreenPosToViewportUV *
@@ -833,8 +833,8 @@ FScreenPassTexture AddTonemapPass(FRDGBuilder& GraphBuilder, const FViewInfo& Vi
 		else
 		{
 			CommonParameters.ColorToBloom = FScreenTransform::Identity;
-			CommonParameters.BloomUVViewportBilinearMin = FVector2D::ZeroVector;
-			CommonParameters.BloomUVViewportBilinearMax = FVector2D::UnitVector;
+			CommonParameters.BloomUVViewportBilinearMin = FVector2f::ZeroVector;
+			CommonParameters.BloomUVViewportBilinearMax = FVector2f::UnitVector;
 			CommonParameters.BloomTexture = GSystemTextures.GetBlackDummy(GraphBuilder);
 			CommonParameters.BloomSampler = BilinearClampSampler;
 		}
diff --git a/Engine/Source/Runtime/Renderer/Private/PostProcess/PostProcessUpscale.cpp b/Engine/Source/Runtime/Renderer/Private/PostProcess/PostProcessUpscale.cpp
--- a/Engine/Source/Runtime/Renderer/Private/PostProcess/PostProcessUpscale.cpp
+++ b/Engine/Source/Runtime/Renderer/Private/PostProcess/PostProcessUpscale.cpp
@@ -51,7 +51,7 @@ TAutoConsoleVariable<int32> CVarUpscaleQuality(
 	TEXT(" 5: 36-tap Gaussian-filtered unsharp mask (very expensive, but good for extreme upsampling).\n"),
 	ECVF_Scalability | ECVF_RenderThreadSafe);
 
-FVector2D PaniniProjection(FVector2D OM, float d, float s)
+FVector2f PaniniProjection(FVector2f OM, float d, float s)
 {
 	float PaniniDirectionXZInvLength = 1.0f / FMath::Sqrt(1.0f + OM.X * OM.X);
 	float SinPhi = OM.X * PaniniDirectionXZInvLength;
@@ -59,7 +59,7 @@ FVector2D PaniniProjection(FVector2D OM, float d, float s)
 	float CosPhi = FMath::Sqrt(1.0f - SinPhi * SinPhi);
 	float S = (d + 1.0f) / (d + CosPhi);
 
-	return S * FVector2D(SinPhi, FMath::Lerp(TanTheta, TanTheta / CosPhi, s));
+	return S * FVector2f(SinPhi, FMath::Lerp(TanTheta, TanTheta / CosPhi, s));
 }
 } //! namespace
 
@@ -85,10 +85,10 @@ END_SHADER_PARAMETER_STRUCT()
 
 FPaniniProjectionParameters GetPaniniProjectionParameters(FPaniniProjectionConfig InPaniniConfig, const FViewInfo& View)
 {
-	const FVector2D FOVPerAxis = View.ViewMatrices.ComputeHalfFieldOfViewPerAxis();
-	const FVector2D ScreenPosToPaniniFactor = FVector2D(FMath::Tan(FOVPerAxis.X), FMath::Tan(FOVPerAxis.Y));
-	const FVector2D PaniniDirection = FVector2D(1.0f, 0.0f) * ScreenPosToPaniniFactor;
-	const FVector2D PaniniPosition = PaniniProjection(PaniniDirection, InPaniniConfig.D, InPaniniConfig.S);
+	const FVector2f FOVPerAxis = View.ViewMatrices.ComputeHalfFieldOfViewPerAxis();
+	const FVector2f ScreenPosToPaniniFactor = FVector2f(FMath::Tan(FOVPerAxis.X), FMath::Tan(FOVPerAxis.Y));
+	const FVector2f PaniniDirection = FVector2f(1.0f, 0.0f) * ScreenPosToPaniniFactor;
+	const FVector2f PaniniPosition = PaniniProjection(PaniniDirection, InPaniniConfig.D, InPaniniConfig.S);
 
 	const float WidthFit = ScreenPosToPaniniFactor.X / PaniniPosition.X;
 	const float ScreenPosScale = FMath::Lerp(1.0f, WidthFit, InPaniniConfig.ScreenFit);
diff --git a/Engine/Source/Runtime/Renderer/Private/PostProcess/PostProcessWeightedSampleSum.cpp b/Engine/Source/Runtime/Renderer/Private/PostProcess/PostProcessWeightedSampleSum.cpp
--- a/Engine/Source/Runtime/Renderer/Private/PostProcess/PostProcessWeightedSampleSum.cpp
+++ b/Engine/Source/Runtime/Renderer/Private/PostProcess/PostProcessWeightedSampleSum.cpp
@@ -61,7 +61,7 @@ void GetFilterParameters(
 	FFilterParameters& OutParameters,
 	const FScreenPassTextureInput& Filter,
 	const FScreenPassTextureInput& Additive,
-	TArrayView<const FVector2D> SampleOffsets,
+	TArrayView<const FVector2f> SampleOffsets,
 	TArrayView<const FLinearColor> SampleWeights)
 {
 	OutParameters.Filter = Filter;
@@ -143,7 +143,7 @@ int GetIntegerKernelRadius(uint32 SampleCountMax, float KernelRadius)
 	return FMath::Min<int32>(FMath::CeilToInt(GetClampedKernelRadius(SampleCountMax, KernelRadius)), SampleCountMax - 1);
 }
 
-uint32 Compute1DGaussianFilterKernel(FVector2D OutOffsetAndWeight[MAX_FILTER_SAMPLES], uint32 SampleCountMax,	float KernelRadius, float CrossCenterWeight)
+uint32 Compute1DGaussianFilterKernel(FVector2f OutOffsetAndWeight[MAX_FILTER_SAMPLES], uint32 SampleCountMax,	float KernelRadius, float CrossCenterWeight)
 {
 	const float FilterSizeScale = FMath::Clamp(CVarFilterSizeScale.GetValueOnRenderThread(), 0.1f, 10.0f);
 
@@ -342,7 +342,7 @@ FScreenPassTexture AddGaussianBlurPass(
 	FScreenPassTextureViewport OutputViewport,
 	FScreenPassTexture Filter,
 	FScreenPassTexture Additive,
-	TArrayView<const FVector2D> SampleOffsets,
+	TArrayView<const FVector2f> SampleOffsets,
 	TArrayView<const FLinearColor> SampleWeights,
 	bool UseMirrorAddressMode)
 {
@@ -470,7 +470,7 @@ FScreenPassTexture AddGaussianBlurPass(
 
 	const uint32 SampleCountMax = GetSampleCountMax(View.GetFeatureLevel(), View.GetShaderPlatform());
 
-	const FVector2D InverseFilterTextureExtent(
+	const FVector2f InverseFilterTextureExtent(
 		1.0f / static_cast<float>(FilterViewport.Extent.X),
 		1.0f / static_cast<float>(FilterViewport.Extent.Y));
 
@@ -485,9 +485,9 @@ FScreenPassTexture AddGaussianBlurPass(
 
 	// Horizontal Pass
 	{
-		FVector2D OffsetAndWeight[MAX_FILTER_SAMPLES];
+		FVector2f OffsetAndWeight[MAX_FILTER_SAMPLES];
 		FLinearColor SampleWeights[MAX_FILTER_SAMPLES];
-		FVector2D SampleOffsets[MAX_FILTER_SAMPLES];
+		FVector2f SampleOffsets[MAX_FILTER_SAMPLES];
 
 		const uint32 SampleCount = Compute1DGaussianFilterKernel(OffsetAndWeight, SampleCountMax, BlurRadius, Inputs.CrossCenterWeight.X);
 
@@ -503,7 +503,7 @@ FScreenPassTexture AddGaussianBlurPass(
 		{
 			const float Offset = OffsetAndWeight[i].X;
 
-			SampleOffsets[i] = FVector2D(InverseFilterTextureExtent.X * Offset, 0.0f);
+			SampleOffsets[i] = FVector2f(InverseFilterTextureExtent.X * Offset, 0.0f);
 		}
 
 		// Horizontal pass doesn't use additive combine.
@@ -516,16 +516,16 @@ FScreenPassTexture AddGaussianBlurPass(
 			HorizontalOutputViewport,
 			Inputs.Filter,
 			Additive,
-			TArrayView<const FVector2D>(SampleOffsets, SampleCount),
+			TArrayView<const FVector2f>(SampleOffsets, SampleCount),
 			TArrayView<const FLinearColor>(SampleWeights, SampleCount),
 			Inputs.UseMirrorAddressMode);
 	}
 
 	// Vertical Pass
 	{
-		FVector2D OffsetAndWeight[MAX_FILTER_SAMPLES];
+		FVector2f OffsetAndWeight[MAX_FILTER_SAMPLES];
 		FLinearColor SampleWeights[MAX_FILTER_SAMPLES];
-		FVector2D SampleOffsets[MAX_FILTER_SAMPLES];
+		FVector2f SampleOffsets[MAX_FILTER_SAMPLES];
 
 		const uint32 SampleCount = Compute1DGaussianFilterKernel(OffsetAndWeight, SampleCountMax, BlurRadius, Inputs.CrossCenterWeight.Y);
 
@@ -541,7 +541,7 @@ FScreenPassTexture AddGaussianBlurPass(
 		{
 			const float Offset = OffsetAndWeight[i].X;
 
-			SampleOffsets[i] = FVector2D(0, InverseFilterTextureExtent.Y * Offset);
+			SampleOffsets[i] = FVector2f(0, InverseFilterTextureExtent.Y * Offset);
 		}
 
 		return AddGaussianBlurPass(
@@ -551,7 +551,7 @@ FScreenPassTexture AddGaussianBlurPass(
 			FilterViewport,
 			HorizontalOutput,
 			Inputs.Additive,
-			TArrayView<const FVector2D>(SampleOffsets, SampleCount),
+			TArrayView<const FVector2f>(SampleOffsets, SampleCount),
 			TArrayView<const FLinearColor>(SampleWeights, SampleCount),
 			Inputs.UseMirrorAddressMode);
 	}
diff --git a/Engine/Source/Runtime/Renderer/Private/PostProcess/PostProcessWeightedSampleSum.h b/Engine/Source/Runtime/Renderer/Private/PostProcess/PostProcessWeightedSampleSum.h
--- a/Engine/Source/Runtime/Renderer/Private/PostProcess/PostProcessWeightedSampleSum.h
+++ b/Engine/Source/Runtime/Renderer/Private/PostProcess/PostProcessWeightedSampleSum.h
@@ -20,7 +20,7 @@ struct FGaussianBlurInputs
 	FLinearColor TintColor;
 
 	// Controls the cross shape of the blur, in both X / Y directions. See r.Bloom.Cross.
-	FVector2D CrossCenterWeight = FVector2D::ZeroVector;
+	FVector2f CrossCenterWeight = FVector2f::ZeroVector;
 
 	// The filter kernel size in percentage of the screen.
 	float KernelSizePercent = 0.0f;
diff --git a/Engine/Source/Runtime/Renderer/Private/PostProcess/PostProcessing.cpp b/Engine/Source/Runtime/Renderer/Private/PostProcess/PostProcessing.cpp
--- a/Engine/Source/Runtime/Renderer/Private/PostProcess/PostProcessing.cpp
+++ b/Engine/Source/Runtime/Renderer/Private/PostProcess/PostProcessing.cpp
@@ -271,15 +271,15 @@ FRDGTextureRef AddTranslucencyCompositionPass(
 			bPostMotionBlur ? TEXT("PostMotionBlurTranslucency.SceneColor") : TEXT("PostDOFTranslucency.SceneColor"));
 	}
 
-	const FVector2D SceneColorSize(SceneColor.ViewRect.Size());
-	const FVector2D SceneColorSizeInv = FVector2D(1.0f, 1.0f) / SceneColorSize;	
-	const FVector2D SceneColorExtent(NewSceneColor->Desc.Extent);
-	const FVector2D SceneColorExtentInv = FVector2D(1.0f, 1.0f) / SceneColorExtent;
+	const FVector2f SceneColorSize(SceneColor.ViewRect.Size());
+	const FVector2f SceneColorSizeInv = FVector2f(1.0f, 1.0f) / SceneColorSize;	
+	const FVector2f SceneColorExtent(NewSceneColor->Desc.Extent);
+	const FVector2f SceneColorExtentInv = FVector2f(1.0f, 1.0f) / SceneColorExtent;
 	
 	const FIntRect SeparateTranslucencyRect = SeparateTranslucencyTextures.GetDimensions().GetViewport(View.ViewRect).Rect;
-	const FVector2D SeparateTranslucencySize(SeparateTranslucencyRect.Size());
-	const FVector2D SeparateTranslucencyExtent((bApplyModulateOnly ? SeparateModulationTexture : SeparateTranslucencyTexture)->Desc.Extent);
-	const FVector2D SeparateTranslucencyExtentInv = FVector2D(1.0f, 1.0f) / SeparateTranslucencyExtent;
+	const FVector2f SeparateTranslucencySize(SeparateTranslucencyRect.Size());
+	const FVector2f SeparateTranslucencyExtent((bApplyModulateOnly ? SeparateModulationTexture : SeparateTranslucencyTexture)->Desc.Extent);
+	const FVector2f SeparateTranslucencyExtentInv = FVector2f(1.0f, 1.0f) / SeparateTranslucencyExtent;
 
 	const FScreenPassTextureViewport SceneColorViewport(SceneColor);
 	const FScreenPassTextureViewport TranslucencyViewport = SeparateTranslucencyTextures.GetDimensions().GetViewport(View.ViewRect);
@@ -2090,7 +2090,7 @@ void AddMobilePostProcessingPasses(FRDGBuilder& GraphBuilder, FScene* Scene, con
 							GaussianBlurInputs.NameY = BlurDebugY;
 							GaussianBlurInputs.Filter = DOFSetup;
 							GaussianBlurInputs.TintColor = FLinearColor::White;
-							GaussianBlurInputs.CrossCenterWeight = FVector2D::ZeroVector;
+							GaussianBlurInputs.CrossCenterWeight = FVector2f::ZeroVector;
 							GaussianBlurInputs.KernelSizePercent = KernelSizePercent;
 
 							return AddGaussianBlurPass(GraphBuilder, View, GaussianBlurInputs);
@@ -2159,7 +2159,7 @@ void AddMobilePostProcessingPasses(FRDGBuilder& GraphBuilder, FScene* Scene, con
 				FMobileBloomUpInputs BloomUpInputs;
 				BloomUpInputs.BloomUpSourceA = BloomUpSourceA;
 				BloomUpInputs.BloomUpSourceB = BloomUpSourceB;
-				BloomUpInputs.ScaleAB = FVector2D(BloomSourceScale, BloomSourceScale);
+				BloomUpInputs.ScaleAB = FVector2f(BloomSourceScale, BloomSourceScale);
 				BloomUpInputs.TintA = TintA;
 				BloomUpInputs.TintB = TintB;
 
diff --git a/Engine/Source/Runtime/Renderer/Private/PostProcess/TemporalAA.cpp b/Engine/Source/Runtime/Renderer/Private/PostProcess/TemporalAA.cpp
--- a/Engine/Source/Runtime/Renderer/Private/PostProcess/TemporalAA.cpp
+++ b/Engine/Source/Runtime/Renderer/Private/PostProcess/TemporalAA.cpp
@@ -259,7 +259,7 @@ float CatmullRom(float x)
 		return (1.5f * ax - 2.5f) * ax*ax + 1.0f;
 }
 
-void SetupSampleWeightParameters(FTemporalAACS::FParameters* OutTAAParameters, const FTAAPassParameters& PassParameters, FVector2D TemporalJitterPixels)
+void SetupSampleWeightParameters(FTemporalAACS::FParameters* OutTAAParameters, const FTAAPassParameters& PassParameters, FVector2f TemporalJitterPixels)
 {
 	float JitterX = TemporalJitterPixels.X;
 	float JitterY = TemporalJitterPixels.Y;
@@ -626,7 +626,7 @@ FTAAOutputs AddTemporalAAPass(
 			if (bCameraCut)
 			{
 				PassParameters->ScreenPosToHistoryBufferUV = FVector4f(1.0f, 1.0f, 1.0f, 1.0f);
-				PassParameters->ScreenPosAbsMax = FVector2D(0.0f, 0.0f);
+				PassParameters->ScreenPosAbsMax = FVector2f(0.0f, 0.0f);
 				PassParameters->HistoryBufferUVMinMax = FVector4f(0.0f, 0.0f, 0.0f, 0.0f);
 				PassParameters->HistoryBufferSize = FVector4f(1.0f, 1.0f, 1.0f, 1.0f);
 
@@ -657,7 +657,7 @@ FTAAOutputs AddTemporalAAPass(
 				FIntPoint ViewportExtent = FIntPoint::DivideAndRoundUp(ReferenceViewportExtent, Inputs.ResolutionDivisor);
 				FIntPoint BufferSize = ReferenceBufferSize / Inputs.ResolutionDivisor;
 
-				PassParameters->ScreenPosAbsMax = FVector2D(1.0f - 1.0f / float(ViewportExtent.X), 1.0f - 1.0f / float(ViewportExtent.Y));
+				PassParameters->ScreenPosAbsMax = FVector2f(1.0f - 1.0f / float(ViewportExtent.X), 1.0f - 1.0f / float(ViewportExtent.Y));
 
 				float InvBufferSizeX = 1.f / float(BufferSize.X);
 				float InvBufferSizeY = 1.f / float(BufferSize.Y);
@@ -725,7 +725,7 @@ FTAAOutputs AddTemporalAAPass(
 			PassParameters->TemporalJitterPixels = InputViewSizeScale * View.TemporalJitterPixels;
 			PassParameters->ScreenPercentage = float(InputViewRect.Width()) / float(OutputViewRect.Width());
 			PassParameters->UpscaleFactor = float(OutputViewRect.Width()) / float(InputViewRect.Width());
-			PassParameters->InputViewMin = InputViewSizeScale * FVector2D(InputViewRect.Min.X, InputViewRect.Min.Y);
+			PassParameters->InputViewMin = InputViewSizeScale * FVector2f(InputViewRect.Min.X, InputViewRect.Min.Y);
 			PassParameters->InputViewSize = FVector4f(
 				InputViewSizeScale * InputViewRect.Width(), InputViewSizeScale * InputViewRect.Height(),
 				InputViewSizeInvScale / InputViewRect.Width(), InputViewSizeInvScale / InputViewRect.Height());
