Fixes vector precision in selection outline

diff --git a/Engine/Source/Runtime/Renderer/Private/PostProcess/PostProcessSelectionOutline.cpp b/Engine/Source/Runtime/Renderer/Private/PostProcess/PostProcessSelectionOutline.cpp
--- a/Engine/Source/Runtime/Renderer/Private/PostProcess/PostProcessSelectionOutline.cpp
+++ b/Engine/Source/Runtime/Renderer/Private/PostProcess/PostProcessSelectionOutline.cpp
@@ -182,9 +182,9 @@ FScreenPassTexture AddSelectionOutlinePass(FRDGBuilder& GraphBuilder, const FVie
 		PassParameters->DepthSampler = PointClampSampler;
 		PassParameters->EditorPrimitivesDepth = DepthStencilTexture;
 		PassParameters->EditorPrimitivesStencil = GraphBuilder.CreateSRV(FRDGTextureSRVDesc::CreateWithPixelFormat(DepthStencilTexture, PF_X24_G8));
-		PassParameters->OutlineColor = FVector(View.SelectionOutlineColor);
+		PassParameters->OutlineColor = FVector3f(View.SelectionOutlineColor);
 		PassParameters->SelectionHighlightIntensity = GEngine->SelectionHighlightIntensity;
-		PassParameters->SubduedOutlineColor = FVector(View.SubduedSelectionOutlineColor);
+		PassParameters->SubduedOutlineColor = FVector3f(View.SubduedSelectionOutlineColor);
 		PassParameters->BSPSelectionIntensity = GEngine->BSPSelectionHighlightIntensity;
 
 		FSelectionOutlinePS::FPermutationDomain PermutationVector;
