Removes FFXAAVS

diff --git a/Engine/Shaders/Private/FXAAShader.usf b/Engine/Shaders/Private/FXAAShader.usf
--- a/Engine/Shaders/Private/FXAAShader.usf
+++ b/Engine/Shaders/Private/FXAAShader.usf
@@ -5,6 +5,8 @@
 
 SCREEN_PASS_TEXTURE_VIEWPORT(Input)
 
+FScreenTransform SvPositionToInputTextureUV;
+
 Texture2D Input_Texture;
 SamplerState Input_Sampler;
 
@@ -38,21 +40,6 @@ SamplerState Input_Sampler;
 	#define FXAA_360_OPT 0
 #endif
 
-/*-----------------------------------------------------------------------------
-	FXAA vertex shader
------------------------------------------------------------------------------*/
-void FxaaVS(
-	in float4 InPosition : ATTRIBUTE0,
-	in float2 UV : ATTRIBUTE1,
-	out noperspective float2 OutCenterUV : TEXCOORD0,
-	out noperspective float4 OutCornersUV : TEXCOORD1,
-	out float4 OutPosition : SV_POSITION
-	)
-{
-	DrawRectangle(InPosition, UV, OutPosition, OutCenterUV);
-	OutCornersUV = float4(OutCenterUV.xy - 0.5f * Input_ExtentInverse, OutCenterUV.xy + 0.5f * Input_ExtentInverse);
-}
-
 
 /*-----------------------------------------------------------------------------
 	FXAA pixel shader
@@ -82,12 +69,17 @@ void FxaaVS(
 	// include NVIDIA FXAA source
 	#include "Fxaa3_11.ush"
 
-	void FxaaPS(noperspective float2 TexCenter : TEXCOORD0, noperspective float4 TexCorners : TEXCOORD1, out float4 OutColor : SV_Target0)
+	void FxaaPS(
+		float4 SvPosition : SV_POSITION,
+		out float4 OutColor : SV_Target0)
 	{
 		FxaaTex TextureAndSampler;
 		TextureAndSampler.tex = Input_Texture;
 		TextureAndSampler.smpl = Input_Sampler;
 		TextureAndSampler.UVMinMax = float4(Input_UVViewportBilinearMin, Input_UVViewportBilinearMax);
+		
+		float2 TexCenter = ApplyScreenTransform(SvPosition.xy, SvPositionToInputTextureUV);
+		float4 TexCorners = float4(TexCenter - 0.5f * Input_ExtentInverse, TexCenter + 0.5f * Input_ExtentInverse);
 
 		OutColor = FxaaPixelShader(
 			TexCenter, TexCorners,
diff --git a/Engine/Source/Runtime/Renderer/Private/PostProcess/PostProcessAA.cpp b/Engine/Source/Runtime/Renderer/Private/PostProcess/PostProcessAA.cpp
--- a/Engine/Source/Runtime/Renderer/Private/PostProcess/PostProcessAA.cpp
+++ b/Engine/Source/Runtime/Renderer/Private/PostProcess/PostProcessAA.cpp
@@ -2,6 +2,7 @@
 
 #include "PostProcess/PostProcessAA.h"
 #include "PostProcess/PostProcessing.h"
+#include "PixelShaderUtils.h"
 
 TAutoConsoleVariable<int32> CVarFXAAQuality(
 	TEXT("r.FXAA.Quality"), 4,
@@ -14,36 +15,6 @@ TAutoConsoleVariable<int32> CVarFXAAQuality(
 	TEXT(" 5: PC extrem quality 12-samples"),
 	ECVF_Scalability | ECVF_RenderThreadSafe);
 
-BEGIN_SHADER_PARAMETER_STRUCT(FFXAAParameters, )
-	SHADER_PARAMETER_STRUCT(FScreenPassTextureInput, Input)
-	SHADER_PARAMETER(FVector4f, fxaaConsoleRcpFrameOpt)
-	SHADER_PARAMETER(FVector4f, fxaaConsoleRcpFrameOpt2)
-	SHADER_PARAMETER(float, fxaaQualitySubpix)
-	SHADER_PARAMETER(float, fxaaQualityEdgeThreshold)
-	SHADER_PARAMETER(float, fxaaQualityEdgeThresholdMin)
-	SHADER_PARAMETER(float, fxaaConsoleEdgeSharpness)
-	SHADER_PARAMETER(float, fxaaConsoleEdgeThreshold)
-	SHADER_PARAMETER(float, fxaaConsoleEdgeThresholdMin)
-	RENDER_TARGET_BINDING_SLOTS()
-END_SHADER_PARAMETER_STRUCT()
-
-class FFXAAVS : public FGlobalShader
-{
-public:
-	DECLARE_GLOBAL_SHADER(FFXAAVS);
-	// FDrawRectangleParameters is filled by DrawScreenPass.
-	SHADER_USE_PARAMETER_STRUCT_WITH_LEGACY_BASE(FFXAAVS, FGlobalShader);
-
-	static bool ShouldCompilePermutation(const FGlobalShaderPermutationParameters& Parameters)
-	{
-		return true;
-	}
-
-	using FParameters = FFXAAParameters;
-};
-
-IMPLEMENT_GLOBAL_SHADER(FFXAAVS, "/Engine/Private/FXAAShader.usf", "FxaaVS", SF_Vertex);
-
 class FFXAAPS : public FGlobalShader
 {
 public:
@@ -57,7 +28,20 @@ public:
 
 	class FQualityDimension : SHADER_PERMUTATION_ENUM_CLASS("FXAA_PRESET", EFXAAQuality);
 	using FPermutationDomain = TShaderPermutationDomain<FQualityDimension>;
-	using FParameters = FFXAAParameters;
+
+	BEGIN_SHADER_PARAMETER_STRUCT(FParameters, )
+		SHADER_PARAMETER_STRUCT(FScreenPassTextureInput, Input)
+		SHADER_PARAMETER(FScreenTransform, SvPositionToInputTextureUV)
+		SHADER_PARAMETER(FVector4f, fxaaConsoleRcpFrameOpt)
+		SHADER_PARAMETER(FVector4f, fxaaConsoleRcpFrameOpt2)
+		SHADER_PARAMETER(float, fxaaQualitySubpix)
+		SHADER_PARAMETER(float, fxaaQualityEdgeThreshold)
+		SHADER_PARAMETER(float, fxaaQualityEdgeThresholdMin)
+		SHADER_PARAMETER(float, fxaaConsoleEdgeSharpness)
+		SHADER_PARAMETER(float, fxaaConsoleEdgeThreshold)
+		SHADER_PARAMETER(float, fxaaConsoleEdgeThresholdMin)
+		RENDER_TARGET_BINDING_SLOTS()
+	END_SHADER_PARAMETER_STRUCT()
 };
 
 IMPLEMENT_GLOBAL_SHADER(FFXAAPS, "/Engine/Private/FXAAShader.usf", "FxaaPS", SF_Pixel);
@@ -83,9 +67,12 @@ FScreenPassTexture AddFXAAPass(FRDGBuilder& GraphBuilder, const FViewInfo& View,
 
 	FRHISamplerState* BilinearClampSampler = TStaticSamplerState<SF_Bilinear, AM_Clamp, AM_Clamp, AM_Clamp>::GetRHI();
 
-	FFXAAParameters* PassParameters = GraphBuilder.AllocParameters<FFXAAParameters>();
+	FFXAAPS::FParameters* PassParameters = GraphBuilder.AllocParameters<FFXAAPS::FParameters>();
 	PassParameters->RenderTargets[0] = Output.GetRenderTargetBinding();
 	PassParameters->Input = GetScreenPassTextureInput(Inputs.SceneColor, BilinearClampSampler);
+	PassParameters->SvPositionToInputTextureUV = (
+		FScreenTransform::ChangeTextureBasisFromTo(FScreenPassTextureViewport(Output), FScreenTransform::ETextureBasis::TexelPosition, FScreenTransform::ETextureBasis::ViewportUV) *
+		FScreenTransform::ChangeTextureBasisFromTo(FScreenPassTextureViewport(Inputs.SceneColor), FScreenTransform::ETextureBasis::ViewportUV, FScreenTransform::ETextureBasis::TextureUV));
 
 	{
 		float N = 0.5f;
@@ -109,25 +96,16 @@ FScreenPassTexture AddFXAAPass(FRDGBuilder& GraphBuilder, const FViewInfo& View,
 	FFXAAPS::FPermutationDomain PixelPermutationVector;
 	PixelPermutationVector.Set<FFXAAPS::FQualityDimension>(Inputs.Quality);
 
-	TShaderMapRef<FFXAAVS> VertexShader(View.ShaderMap);
 	TShaderMapRef<FFXAAPS> PixelShader(View.ShaderMap, PixelPermutationVector);
 
-	const FScreenPassTextureViewport OutputViewport(Output);
-
-	AddDrawScreenPass(
+	FPixelShaderUtils::AddFullscreenPass(
 		GraphBuilder,
-		RDG_EVENT_NAME("FXAA(Quality=%d) %dx%d PS", Inputs.Quality, OutputViewport.Rect.Width(), OutputViewport.Rect.Height()),
-		View,
-		OutputViewport,
-		FScreenPassTextureViewport(Inputs.SceneColor),
-		FScreenPassPipelineState(VertexShader, PixelShader),
+		View.ShaderMap,
+		RDG_EVENT_NAME("FXAA(Quality=%d) %dx%d PS",
+			Inputs.Quality, Output.ViewRect.Width(), Output.ViewRect.Height()),
+		PixelShader,
 		PassParameters,
-		EScreenPassDrawFlags::AllowHMDHiddenAreaMask,
-		[VertexShader, PixelShader, PassParameters](FRHICommandList& RHICmdList)
-	{
-		SetShaderParameters(RHICmdList, VertexShader, VertexShader.GetVertexShader(), *PassParameters);
-		SetShaderParameters(RHICmdList, PixelShader, PixelShader.GetPixelShader(), *PassParameters);
-	});
+		Output.ViewRect);
 
 	return MoveTemp(Output);
 }
