Adds RENDERER_API to AddGaussianBlurPass for pluggins

diff --git a/Engine/Source/Runtime/Renderer/Private/PostProcess/PostProcessWeightedSampleSum.h b/Engine/Source/Runtime/Renderer/Private/PostProcess/PostProcessWeightedSampleSum.h
--- a/Engine/Source/Runtime/Renderer/Private/PostProcess/PostProcessWeightedSampleSum.h
+++ b/Engine/Source/Runtime/Renderer/Private/PostProcess/PostProcessWeightedSampleSum.h
@@ -30,4 +30,4 @@ struct FGaussianBlurInputs
 
 using FGaussianBlurOutputs = FScreenPassTexture;
 
-FGaussianBlurOutputs AddGaussianBlurPass(FRDGBuilder& GraphBuilder, const FViewInfo& View, const FGaussianBlurInputs& Inputs);
+extern RENDERER_API FGaussianBlurOutputs AddGaussianBlurPass(FRDGBuilder& GraphBuilder, const FViewInfo& View, const FGaussianBlurInputs& Inputs);
