Removes FBloomSetupVS

diff --git a/Engine/Shaders/Private/PostProcessBloom.usf b/Engine/Shaders/Private/PostProcessBloom.usf
--- a/Engine/Shaders/Private/PostProcessBloom.usf
+++ b/Engine/Shaders/Private/PostProcessBloom.usf
@@ -47,17 +47,22 @@ float4 BloomSetupCommon(float2 UV, float ExposureScale)
 	return float4(BloomAmount * LinearColor, 0) * View.PreExposure;
 }
 
+#if PIXELSHADER
+
+FScreenTransform SvPositionToInputTextureUV;
+
 void BloomSetupPS(
-	noperspective float4 UVAndScreenPos : TEXCOORD0,
-	nointerpolation float InExposureScale : TEXCOORD1,
+	float4 SvPosition : SV_POSITION,
 	out float4 OutColor : SV_Target0)
 {
-	float2 UV = UVAndScreenPos.xy;
+	float2 UV = ApplyScreenTransform(SvPosition.xy, SvPositionToInputTextureUV);
+	
+	float ExposureScale = EyeAdaptationLookup();
 
-	OutColor = BloomSetupCommon(UV, InExposureScale);
+	OutColor = BloomSetupCommon(UV, ExposureScale);
 }
 
-#if COMPUTESHADER
+#elif COMPUTESHADER
 
 RWTexture2D<float4> RWOutputTexture;
 
diff --git a/Engine/Source/Runtime/Renderer/Private/PostProcess/PostProcessBloomSetup.cpp b/Engine/Source/Runtime/Renderer/Private/PostProcess/PostProcessBloomSetup.cpp
--- a/Engine/Source/Runtime/Renderer/Private/PostProcess/PostProcessBloomSetup.cpp
+++ b/Engine/Source/Runtime/Renderer/Private/PostProcess/PostProcessBloomSetup.cpp
@@ -4,6 +4,7 @@
 #include "PostProcess/PostProcessDownsample.h"
 #include "PostProcess/PostProcessFFTBloom.h"
 #include "PostProcess/PostProcessWeightedSampleSum.h"
+#include "PixelShaderUtils.h"
 
 namespace
 {
@@ -46,23 +47,6 @@ FBloomSetupParameters GetBloomSetupParameters(
 	return Parameters;
 }
 
-class FBloomSetupVS : public FGlobalShader
-{
-	DECLARE_GLOBAL_SHADER(FBloomSetupVS);
-
-	// FDrawRectangleParameters is filled by DrawScreenPass.
-	SHADER_USE_PARAMETER_STRUCT_WITH_LEGACY_BASE(FBloomSetupVS, FGlobalShader);
-
-	using FParameters = FBloomSetupParameters;
-
-	static bool ShouldCompilePermutation(const FGlobalShaderPermutationParameters& Parameters)
-	{
-		return IsFeatureLevelSupported(Parameters.Platform, ERHIFeatureLevel::SM5);
-	}
-};
-
-IMPLEMENT_GLOBAL_SHADER(FBloomSetupVS, "/Engine/Private/PostProcessBloom.usf", "BloomSetupVS", SF_Vertex);
-
 class FBloomSetupPS : public FGlobalShader
 {
 public:
@@ -71,6 +55,7 @@ public:
 
 	BEGIN_SHADER_PARAMETER_STRUCT(FParameters, )
 		SHADER_PARAMETER_STRUCT_INCLUDE(FBloomSetupParameters, BloomSetup)
+		SHADER_PARAMETER(FScreenTransform, SvPositionToInputTextureUV)
 		RENDER_TARGET_BINDING_SLOTS()
 	END_SHADER_PARAMETER_STRUCT()
 
@@ -143,24 +128,19 @@ FScreenPassTexture AddBloomSetupPass(FRDGBuilder& GraphBuilder, const FViewInfo&
 	{
 		FBloomSetupPS::FParameters* PassParameters = GraphBuilder.AllocParameters<FBloomSetupPS::FParameters>();
 		PassParameters->BloomSetup = GetBloomSetupParameters(View, Viewport, Inputs.SceneColor.Texture, Inputs.EyeAdaptationTexture, Inputs.Threshold);
+		PassParameters->SvPositionToInputTextureUV = (
+			FScreenTransform::ChangeTextureBasisFromTo(FScreenPassTextureViewport(Output), FScreenTransform::ETextureBasis::TexelPosition, FScreenTransform::ETextureBasis::ViewportUV) *
+			FScreenTransform::ChangeTextureBasisFromTo(FScreenPassTextureViewport(Inputs.SceneColor), FScreenTransform::ETextureBasis::ViewportUV, FScreenTransform::ETextureBasis::TextureUV));
 		PassParameters->RenderTargets[0] = Output.GetRenderTargetBinding();
 
-		TShaderMapRef<FBloomSetupVS> VertexShader(View.ShaderMap);
 		TShaderMapRef<FBloomSetupPS> PixelShader(View.ShaderMap);
-
-		AddDrawScreenPass(
+		FPixelShaderUtils::AddFullscreenPass(
 			GraphBuilder,
+			View.ShaderMap,
 			RDG_EVENT_NAME("BloomSetup %dx%d (PS)", Viewport.Rect.Width(), Viewport.Rect.Height()),
-			View,
-			Viewport,
-			Viewport,
-			FScreenPassPipelineState(VertexShader, PixelShader),
+			PixelShader,
 			PassParameters,
-			[VertexShader, PixelShader, PassParameters] (FRHICommandList& RHICmdList)
-		{
-			SetShaderParameters(RHICmdList, VertexShader, VertexShader.GetVertexShader(), PassParameters->BloomSetup);
-			SetShaderParameters(RHICmdList, PixelShader, PixelShader.GetPixelShader(), *PassParameters);
-		});
+			Output.ViewRect);
 	}
 
 	return FScreenPassTexture(Output);
