Removes SHADER_USE_PARAMETER_STRUCT_WITH_LEGACY_BASE from FFilterVS

diff --git a/Engine/Shaders/Private/FilterPixelShader.usf b/Engine/Shaders/Private/FilterPixelShader.usf
--- a/Engine/Shaders/Private/FilterPixelShader.usf
+++ b/Engine/Shaders/Private/FilterPixelShader.usf
@@ -29,6 +29,7 @@
 #define USE_DYNAMIC_LOOP (STATIC_SAMPLE_COUNT == 0)
 
 SCREEN_PASS_TEXTURE_VIEWPORT(Filter)
+
 SamplerState Filter_Sampler;
 Texture2D Filter_Texture; 
 
@@ -68,12 +69,18 @@ float4 SampleAdditiveTexture(float2 UV)
 
 #if USE_DYNAMIC_LOOP
 
+FScreenTransform SvPositionToTextureUV;
+
 float4 SampleOffsets[PACKED_SAMPLE_COUNT_MAX];
 float4 SampleWeights[SAMPLE_COUNT_MAX];
 int SampleCount;
 
-void MainPS(noperspective float2 InUV : TEXCOORD0, out float4 OutColor : SV_Target0)
+void MainPS(
+	float4 SvPosition : SV_POSITION,
+	out float4 OutColor : SV_Target0)
 {
+	float2 InUV = ApplyScreenTransform(SvPosition.xy, SvPositionToTextureUV);
+
 	float4 Color = 0;
 
 	int SampleIndex = 0;
diff --git a/Engine/Shaders/Private/FilterVertexShader.usf b/Engine/Shaders/Private/FilterVertexShader.usf
--- a/Engine/Shaders/Private/FilterVertexShader.usf
+++ b/Engine/Shaders/Private/FilterVertexShader.usf
@@ -5,6 +5,7 @@
 =============================================================================*/
 
 #include "Common.ush"
+#include "ScreenPass.ush"
 
 #ifndef STATIC_SAMPLE_COUNT
 	#error Static sample count is undefined!
@@ -14,21 +15,23 @@
 
 #define USE_UNPACKED_UVS (ES3_1_PROFILE)
 
+FScreenTransform ViewportUVToTextureUV;
+
 float4 SampleOffsets[PACKED_STATIC_SAMPLE_COUNT];
 
 #if USE_UNPACKED_UVS
 
 void MainVS(
-	float4 InPosition : ATTRIBUTE0,
-	float2 UV : ATTRIBUTE1,
+	float2 ViewportUV : ATTRIBUTE0,
+	float2 Unused     : ATTRIBUTE1, // TODO: kill
 	out noperspective float2 OutUV : TEXCOORD0,
 	out noperspective float2 OutOffsetUVs[STATIC_SAMPLE_COUNT] : TEXCOORD1,
 	out float4 OutPosition : SV_POSITION)
 {
-	DrawRectangle(InPosition, UV, OutPosition, UV);
-
+	OutPosition = float4(ViewportUVToScreenPos(ViewportUV), 0, 1);
+	
+	float2 UV = ApplyScreenTransform(ViewportUV, ViewportUVToTextureUV);
 	
-
 	for (int SampleIndex = 0; SampleIndex < STATIC_SAMPLE_COUNT - 1; SampleIndex += 2)
 	{
 		half4 OffsetUVUV = SampleOffsets[SampleIndex / 2];
@@ -48,15 +51,17 @@ void MainVS(
 #else
 
 void MainVS(
-	float4 InPosition : ATTRIBUTE0,
-	float2 UV : ATTRIBUTE1,
+	float2 ViewportUV : ATTRIBUTE0,
+	float2 Unused     : ATTRIBUTE1, // TODO: kill
 	out noperspective float2 OutUV : TEXCOORD0,
 	out noperspective float4 OutOffsetUVs[PACKED_STATIC_SAMPLE_COUNT] : TEXCOORD1,
 	out float4 OutPosition : SV_POSITION
 	)
 {
-	DrawRectangle(InPosition, UV, OutPosition, UV);
-
+	OutPosition = float4(ViewportUVToScreenPos(ViewportUV), 0, 1);
+	
+	float2 UV = ApplyScreenTransform(ViewportUV, ViewportUVToTextureUV);
+	
 	UNROLL for(int OffsetIndex = 0; OffsetIndex < PACKED_STATIC_SAMPLE_COUNT; ++OffsetIndex)
 	{
 		OutOffsetUVs[OffsetIndex] = UV.xyxy + SampleOffsets[OffsetIndex];
diff --git a/Engine/Shaders/Private/Tools/FullscreenVertexShader.usf b/Engine/Shaders/Private/Tools/FullscreenVertexShader.usf
--- a/Engine/Shaders/Private/Tools/FullscreenVertexShader.usf
+++ b/Engine/Shaders/Private/Tools/FullscreenVertexShader.usf
@@ -7,9 +7,9 @@
 #include "../Common.ush"
 
 void MainVS(
-	float2 InPosition : ATTRIBUTE0,
-	float2 InUV       : ATTRIBUTE1, // TODO: kill
-	out float4 Position : SV_POSITION)
+	float2 ViewportUV : ATTRIBUTE0,
+	float2 UV         : ATTRIBUTE1, // TODO: kill
+	out float4 OutPosition : SV_POSITION)
 {
-	Position = float4(InPosition.x * 2.0 - 1.0, 1.0 - 2.0 * InPosition.y, 0, 1);
+	OutPosition = float4(ViewportUVToScreenPos(ViewportUV), 0, 1);
 }
diff --git a/Engine/Source/Runtime/Renderer/Private/PostProcess/PostProcessWeightedSampleSum.cpp b/Engine/Source/Runtime/Renderer/Private/PostProcess/PostProcessWeightedSampleSum.cpp
--- a/Engine/Source/Runtime/Renderer/Private/PostProcess/PostProcessWeightedSampleSum.cpp
+++ b/Engine/Source/Runtime/Renderer/Private/PostProcess/PostProcessWeightedSampleSum.cpp
@@ -1,6 +1,7 @@
 // Copyright Epic Games, Inc. All Rights Reserved.
 
 #include "PostProcess/PostProcessWeightedSampleSum.h"
+#include "PixelShaderUtils.h"
 
 // maximum number of sample using the shader that has the dynamic loop
 #define MAX_FILTER_SAMPLES	128
@@ -252,16 +253,37 @@ public:
 	{}
 };
 
+class FFilterPS : public FFilterShader
+{
+public:
+	DECLARE_GLOBAL_SHADER(FFilterPS);
+	SHADER_USE_PARAMETER_STRUCT(FFilterPS, FFilterShader);
+
+	using FPermutationDomain = TShaderPermutationDomain<FStaticSampleCount, FCombineAdditive, FManualUVBorder>;
+
+	BEGIN_SHADER_PARAMETER_STRUCT(FParameters, )
+		SHADER_PARAMETER_STRUCT_INCLUDE(FFilterParameters, Filter)
+		SHADER_PARAMETER(FScreenTransform, SvPositionToTextureUV)
+		SHADER_PARAMETER(FScreenTransform, ViewportUVToTextureUV)
+		RENDER_TARGET_BINDING_SLOTS()
+	END_SHADER_PARAMETER_STRUCT()
+
+	static bool ShouldCompilePermutation(const FGlobalShaderPermutationParameters& Parameters)
+	{
+		const FPermutationDomain PermutationVector(Parameters.PermutationId);
+
+		return FFilterShader::ShouldCompilePermutation(Parameters.Platform, PermutationVector.Get<FStaticSampleCount>());
+	}
+};
+
 class FFilterVS : public FFilterShader
 {
 public:
 	DECLARE_GLOBAL_SHADER(FFilterVS);
-
-	// FDrawRectangleParameters is filled by DrawScreenPass.
-	SHADER_USE_PARAMETER_STRUCT_WITH_LEGACY_BASE(FFilterVS, FFilterShader);
+	SHADER_USE_PARAMETER_STRUCT(FFilterVS, FFilterShader);
 
 	using FPermutationDomain = TShaderPermutationDomain<FStaticSampleCount>;
-	using FParameters = FFilterParameters;
+	using FParameters = FFilterPS::FParameters;
 
 	static bool ShouldCompilePermutation(const FGlobalShaderPermutationParameters& Parameters)
 	{
@@ -281,28 +303,6 @@ public:
 };
 
 IMPLEMENT_GLOBAL_SHADER(FFilterVS, "/Engine/Private/FilterVertexShader.usf", "MainVS", SF_Vertex);
-
-class FFilterPS : public FFilterShader
-{
-public:
-	DECLARE_GLOBAL_SHADER(FFilterPS);
-	SHADER_USE_PARAMETER_STRUCT(FFilterPS, FFilterShader);
-
-	using FPermutationDomain = TShaderPermutationDomain<FStaticSampleCount, FCombineAdditive, FManualUVBorder>;
-
-	BEGIN_SHADER_PARAMETER_STRUCT(FParameters, )
-		SHADER_PARAMETER_STRUCT_INCLUDE(FFilterParameters, Filter)
-		RENDER_TARGET_BINDING_SLOTS()
-	END_SHADER_PARAMETER_STRUCT()
-
-	static bool ShouldCompilePermutation(const FGlobalShaderPermutationParameters& Parameters)
-	{
-		const FPermutationDomain PermutationVector(Parameters.PermutationId);
-
-		return FFilterShader::ShouldCompilePermutation(Parameters.Platform, PermutationVector.Get<FStaticSampleCount>());
-	}
-};
-
 IMPLEMENT_GLOBAL_SHADER(FFilterPS, "/Engine/Private/FilterPixelShader.usf", "MainPS", SF_Pixel);
 
 class FFilterCS : public FFilterShader
@@ -413,6 +413,11 @@ FScreenPassTexture AddGaussianBlurPass(
 	{
 		FFilterPS::FParameters* PassParameters = GraphBuilder.AllocParameters<FFilterPS::FParameters>();
 		GetFilterParameters(PassParameters->Filter, FilterInput, AdditiveInput, SampleOffsets, SampleWeights);
+		PassParameters->ViewportUVToTextureUV = (
+			FScreenTransform::ChangeTextureBasisFromTo(FScreenPassTextureViewport(Filter), FScreenTransform::ETextureBasis::ViewportUV, FScreenTransform::ETextureBasis::TextureUV));
+		PassParameters->SvPositionToTextureUV = (
+			FScreenTransform::ChangeTextureBasisFromTo(FScreenPassTextureViewport(Output), FScreenTransform::ETextureBasis::TexelPosition, FScreenTransform::ETextureBasis::ViewportUV) *
+			PassParameters->ViewportUVToTextureUV);
 		PassParameters->RenderTargets[0] = Output.GetRenderTargetBinding();
 
 		FFilterPS::FPermutationDomain PixelPermutationVector;
@@ -437,20 +442,19 @@ FScreenPassTexture AddGaussianBlurPass(
 				PassParameters,
 				[VertexShader, PixelShader, PassParameters] (FRHICommandList& RHICmdList)
 			{
-				SetShaderParameters(RHICmdList, VertexShader, VertexShader.GetVertexShader(), PassParameters->Filter);
+				SetShaderParameters(RHICmdList, VertexShader, VertexShader.GetVertexShader(), *PassParameters);
 				SetShaderParameters(RHICmdList, PixelShader, PixelShader.GetPixelShader(), *PassParameters);
 			});
 		}
 		else
 		{
-			AddDrawScreenPass(
+			FPixelShaderUtils::AddFullscreenPass(
 				GraphBuilder,
+				View.ShaderMap,
 				RDG_EVENT_NAME("GaussianBlur.%s %dx%d (PS, Dynamic)", Name, OutputViewport.Rect.Width(), OutputViewport.Rect.Height()),
-				View,
-				OutputViewport,
-				FScreenPassTextureViewport(Filter),
 				PixelShader,
-				PassParameters);
+				PassParameters,
+				Output.ViewRect);
 		}
 	}
 
