Fix strata toplayer clearing cost due to missing CMask.

The top layer texture had a unecessary UAV flag preventing to create cmask metadata.

diff --git a/Engine/Shaders/Private/Strata/StrataMaterialClassification.usf b/Engine/Shaders/Private/Strata/StrataMaterialClassification.usf
--- a/Engine/Shaders/Private/Strata/StrataMaterialClassification.usf
+++ b/Engine/Shaders/Private/Strata/StrataMaterialClassification.usf
@@ -213,7 +213,6 @@ void StencilMainPS(
 #if SHADER_CLEAR_MATERIAL_BUFFER
 
 RWTexture2DArray<uint> MaterialTextureArrayUAV;
-RWTexture2D<uint> TopLayerTextureUAV;
 RWTexture2D<uint2> SSSTextureUAV;
 uint MaxBytesPerPixel;
 int2 TiledViewBufferResolution;
@@ -236,8 +235,6 @@ void ClearMaterialBufferMainCS(uint2 DispatchThreadId : SV_DispatchThreadID)
 	const uint FirstSlice = 0;
 	MaterialTextureArrayUAV[uint3(DispatchThreadId.xy, FirstSlice)] = PackedStrataHeader;
 
-	TopLayerTextureUAV[DispatchThreadId.xy] = 0;
-
 	FStrataSubsurfaceData StrataSubsurfaceData = (FStrataSubsurfaceData)0;
 	SSSTextureUAV[DispatchThreadId.xy] = StrataPackSSSData(StrataSubsurfaceData);
 }
diff --git a/Engine/Source/Runtime/Renderer/Private/Strata/Strata.cpp b/Engine/Source/Runtime/Renderer/Private/Strata/Strata.cpp
--- a/Engine/Source/Runtime/Renderer/Private/Strata/Strata.cpp
+++ b/Engine/Source/Runtime/Renderer/Private/Strata/Strata.cpp
@@ -76,14 +76,10 @@ static TAutoConsoleVariable<int32> CVarClearDuringCategorization(
 
 IMPLEMENT_GLOBAL_SHADER_PARAMETER_STRUCT(FStrataGlobalUniformParameters, "Strata");
 
-
-
 void FStrataSceneData::Reset()
 {
 	TopLayerTexture = nullptr;
 	SSSTexture = nullptr;
-
-	TopLayerTextureUAV = nullptr;
 	SSSTextureUAV = nullptr;
 
 	MaterialTextureArray = nullptr;
@@ -113,7 +109,10 @@ const TCHAR* ToString(EStrataTileMaterialType Type)
 	return TEXT("Unknown");
 }
 
-
+FORCEINLINE bool ClearDuringCategorization()
+{
+	return CVarClearDuringCategorization.GetValueOnRenderThread() > 0;
+}
 
 namespace Strata
 {
@@ -122,7 +121,6 @@ namespace Strata
 static void AddStrataClearMaterialBufferPass(
 	FRDGBuilder& GraphBuilder, 
 	FRDGTextureUAVRef MaterialTextureArrayUAV,
-	FRDGTextureUAVRef TopLayerTextureUAV,
 	FRDGTextureUAVRef SSSTextureUAV,
 	uint32 MaxBytesPerPixel, 
 	FIntPoint TiledViewBufferResolution);
@@ -196,8 +194,7 @@ void InitialiseStrataFrameSceneData(FSceneRenderer& SceneRenderer, FRDGBuilder&
 
 		// Top layer texture
 		{
-			StrataSceneData.TopLayerTexture = GraphBuilder.CreateTexture(FRDGTextureDesc::Create2D(SceneTextureExtent, PF_R32_UINT, FClearValueBinding::Black, TexCreate_RenderTargetable | TexCreate_ShaderResource | TexCreate_UAV), TEXT("Strata.TopLayerTexture"));
-			StrataSceneData.TopLayerTextureUAV = GraphBuilder.CreateUAV(StrataSceneData.TopLayerTexture);
+			StrataSceneData.TopLayerTexture = GraphBuilder.CreateTexture(FRDGTextureDesc::Create2D(SceneTextureExtent, PF_R32_UINT, FClearValueBinding::Black, TexCreate_RenderTargetable | TexCreate_ShaderResource), TEXT("Strata.TopLayerTexture"));
 		}
 
 		// SSS texture
@@ -237,7 +234,6 @@ void InitialiseStrataFrameSceneData(FSceneRenderer& SceneRenderer, FRDGBuilder&
 		AddStrataClearMaterialBufferPass(
 			GraphBuilder, 
 			GraphBuilder.CreateUAV(FRDGTextureUAVDesc(StrataSceneData.MaterialTextureArrayMRTs, 0)),
-			StrataSceneData.TopLayerTextureUAV,
 			StrataSceneData.SSSTextureUAV,
 			StrataSceneData.MaxBytesPerPixel, 
 			MaterialBufferSizeXY);
@@ -407,7 +403,6 @@ class FStrataClearMaterialBufferCS : public FGlobalShader
 
 	BEGIN_SHADER_PARAMETER_STRUCT(FParameters, )
 		SHADER_PARAMETER_RDG_TEXTURE_UAV(RWTexture2DArray<uint>, MaterialTextureArrayUAV)
-		SHADER_PARAMETER_RDG_TEXTURE_UAV(RWTexture2D<uint>, TopLayerTextureUAV)
 		SHADER_PARAMETER_RDG_TEXTURE_UAV(RWTexture2D<uint2>, SSSTextureUAV)
 		SHADER_PARAMETER(uint32, MaxBytesPerPixel)
 		SHADER_PARAMETER(FIntPoint, TiledViewBufferResolution)
@@ -767,9 +762,9 @@ void AddStrataMaterialClassificationPass(FRDGBuilder& GraphBuilder, const FMinim
 			// Tile reduction requires 64-wide wave
 			bWaveOps = bWaveOps && !IsRHIDeviceNVIDIA();
 		#endif
-
+			const bool bClear = ClearDuringCategorization();
 			FStrataMaterialTileClassificationPassCS::FPermutationDomain PermutationVector;
-			PermutationVector.Set< FStrataMaterialTileClassificationPassCS::FStrataClearDuringCategorization >(CVarClearDuringCategorization.GetValueOnRenderThread() > 0);
+			PermutationVector.Set< FStrataMaterialTileClassificationPassCS::FStrataClearDuringCategorization >(bClear);
 			PermutationVector.Set< FStrataMaterialTileClassificationPassCS::FWaveOps >(bWaveOps);
 			TShaderMapRef<FStrataMaterialTileClassificationPassCS> ComputeShader(View.ShaderMap, PermutationVector);
 			FStrataMaterialTileClassificationPassCS::FParameters* PassParameters = GraphBuilder.AllocParameters<FStrataMaterialTileClassificationPassCS::FParameters>();
@@ -789,7 +784,7 @@ void AddStrataMaterialClassificationPass(FRDGBuilder& GraphBuilder, const FMinim
 			const uint32 GroupSize = 8;
 			FComputeShaderUtils::AddPass(
 				GraphBuilder,
-				RDG_EVENT_NAME("Strata::MaterialTileClassification(%s)", bWaveOps ? TEXT("Wave") : TEXT("SharedMemory")),
+				RDG_EVENT_NAME("Strata::MaterialTileClassification(%s%s)", bWaveOps ? TEXT("Wave") : TEXT("SharedMemory"), bClear ? TEXT(", Clear") : TEXT("")),
 				ComputeShader,
 				PassParameters,
 				FComputeShaderUtils::GetGroupCount(PassParameters->ViewResolution, GroupSize));
@@ -800,12 +795,11 @@ void AddStrataMaterialClassificationPass(FRDGBuilder& GraphBuilder, const FMinim
 static void AddStrataClearMaterialBufferPass(
 	FRDGBuilder& GraphBuilder,
 	FRDGTextureUAVRef MaterialTextureArrayUAV,
-	FRDGTextureUAVRef TopLayerTextureUAV,
 	FRDGTextureUAVRef SSSTextureUAV,
 	uint32 MaxBytesPerPixel,
 	FIntPoint TiledViewBufferResolution)
 {
-	if (CVarClearDuringCategorization.GetValueOnRenderThread() > 0)
+	if (ClearDuringCategorization())
 	{
 		return;
 	}
@@ -813,7 +807,6 @@ static void AddStrataClearMaterialBufferPass(
 	TShaderMapRef<FStrataClearMaterialBufferCS> ComputeShader(GetGlobalShaderMap(GMaxRHIFeatureLevel));
 	FStrataClearMaterialBufferCS::FParameters* PassParameters = GraphBuilder.AllocParameters<FStrataClearMaterialBufferCS::FParameters>();
 	PassParameters->MaterialTextureArrayUAV = MaterialTextureArrayUAV;
-	PassParameters->TopLayerTextureUAV = TopLayerTextureUAV;
 	PassParameters->SSSTextureUAV = SSSTextureUAV;
 	PassParameters->MaxBytesPerPixel = MaxBytesPerPixel;
 	PassParameters->TiledViewBufferResolution = TiledViewBufferResolution;
