Fixed typo in strata

diff --git a/Engine/Source/Runtime/Renderer/Private/BasePassRendering.cpp b/Engine/Source/Runtime/Renderer/Private/BasePassRendering.cpp
--- a/Engine/Source/Runtime/Renderer/Private/BasePassRendering.cpp
+++ b/Engine/Source/Runtime/Renderer/Private/BasePassRendering.cpp
@@ -750,7 +750,7 @@ void FDeferredShadingSceneRenderer::RenderBasePass(
 
 	TStaticArray<FTextureRenderTargetBinding, MaxSimultaneousRenderTargets> BasePassTextures;
 	uint32 BasePassTextureCount = SceneTextures.GetGBufferRenderTargets(BasePassTextures);
-	Strata::ApprendStrataMRTs(*this, BasePassTextureCount, BasePassTextures);
+	Strata::AppendStrataMRTs(*this, BasePassTextureCount, BasePassTextures);
 	TArrayView<FTextureRenderTargetBinding> BasePassTexturesView = MakeArrayView(BasePassTextures.GetData(), BasePassTextureCount);
 	FRDGTextureRef BasePassDepthTexture = SceneTextures.Depth.Target;
 	FLinearColor SceneColorClearValue;
diff --git a/Engine/Source/Runtime/Renderer/Private/Strata/Strata.cpp b/Engine/Source/Runtime/Renderer/Private/Strata/Strata.cpp
--- a/Engine/Source/Runtime/Renderer/Private/Strata/Strata.cpp
+++ b/Engine/Source/Runtime/Renderer/Private/Strata/Strata.cpp
@@ -645,7 +645,7 @@ void AddStrataStencilPass(
 
 ////////////////////////////////////////////////////////////////////////////////////////////////////////////////////////////////////////////////
 
-void ApprendStrataMRTs(FSceneRenderer& SceneRenderer, uint32& RenderTargetCount, TStaticArray<FTextureRenderTargetBinding, MaxSimultaneousRenderTargets>& RenderTargets)
+void AppendStrataMRTs(FSceneRenderer& SceneRenderer, uint32& RenderTargetCount, TStaticArray<FTextureRenderTargetBinding, MaxSimultaneousRenderTargets>& RenderTargets)
 {
 	if (Strata::IsStrataEnabled() && SceneRenderer.Scene)
 	{
diff --git a/Engine/Source/Runtime/Renderer/Private/Strata/Strata.h b/Engine/Source/Runtime/Renderer/Private/Strata/Strata.h
--- a/Engine/Source/Runtime/Renderer/Private/Strata/Strata.h
+++ b/Engine/Source/Runtime/Renderer/Private/Strata/Strata.h
@@ -95,7 +95,7 @@ void BindStrataBasePassUniformParameters(FRDGBuilder& GraphBuilder, FStrataScene
 void BindStrataGlobalUniformParameters(FRDGBuilder& GraphBuilder, FStrataSceneData* StrataSceneData, FStrataGlobalUniformParameters& OutStrataUniformParameters);
 void BindStrataForwardPasslUniformParameters(FRDGBuilder& GraphBuilder, FStrataSceneData* StrataSceneData, FStrataForwardPassUniformParameters& OutStrataUniformParameters);
 
-void ApprendStrataMRTs(FSceneRenderer& SceneRenderer, uint32& BasePassTextureCount, TStaticArray<FTextureRenderTargetBinding, MaxSimultaneousRenderTargets>& BasePassTextures);
+void AppendStrataMRTs(FSceneRenderer& SceneRenderer, uint32& BasePassTextureCount, TStaticArray<FTextureRenderTargetBinding, MaxSimultaneousRenderTargets>& BasePassTextures);
 void SetBasePassRenderTargetOutputFormat(const EShaderPlatform Platform, FShaderCompilerEnvironment& OutEnvironment);
 void PostBasePass(FRDGBuilder& GraphBuilder, const TArray<FViewInfo>& Views);
 
