Initial VSM scalability settings
Expand Z range of clipmap levels to fix some rare artifacts related to SMRT walking off the Z range with certain settings.

diff --git a/Engine/Config/BaseScalability.ini b/Engine/Config/BaseScalability.ini
--- a/Engine/Config/BaseScalability.ini
+++ b/Engine/Config/BaseScalability.ini
@@ -101,6 +101,13 @@ r.DistanceFieldAO=0
 r.VolumetricFog=0
 r.LightMaxDrawDistanceScale=0
 r.CapsuleShadows=0
+r.Shadow.Virtual.MaxPhysicalPages=512
+r.Shadow.Virtual.ResolutionLodBiasDirectional=0.0
+r.Shadow.Virtual.ResolutionLodBiasLocal=1.0
+r.Shadow.Virtual.SMRT.RayCountDirectional=0
+r.Shadow.Virtual.SMRT.SamplesPerRayDirectional=2
+r.Shadow.Virtual.SMRT.RayCountLocal=0
+r.Shadow.Virtual.SMRT.SamplesPerRayLocal=2
 
 [ShadowQuality@1]
 r.LightFunctionQuality=1
@@ -117,6 +124,13 @@ r.DistanceFieldAO=0
 r.VolumetricFog=0
 r.LightMaxDrawDistanceScale=.5
 r.CapsuleShadows=1
+r.Shadow.Virtual.MaxPhysicalPages=512
+r.Shadow.Virtual.ResolutionLodBiasDirectional=0.0
+r.Shadow.Virtual.ResolutionLodBiasLocal=1.0
+r.Shadow.Virtual.SMRT.RayCountDirectional=4
+r.Shadow.Virtual.SMRT.SamplesPerRayDirectional=2
+r.Shadow.Virtual.SMRT.RayCountLocal=4
+r.Shadow.Virtual.SMRT.SamplesPerRayLocal=2
 
 [ShadowQuality@2]
 r.LightFunctionQuality=1
@@ -137,6 +151,13 @@ r.VolumetricFog.GridSizeZ=64
 r.VolumetricFog.HistoryMissSupersampleCount=4
 r.LightMaxDrawDistanceScale=1
 r.CapsuleShadows=1
+r.Shadow.Virtual.MaxPhysicalPages=2048
+r.Shadow.Virtual.ResolutionLodBiasDirectional=-0.5
+r.Shadow.Virtual.ResolutionLodBiasLocal=0.0
+r.Shadow.Virtual.SMRT.RayCountDirectional=8
+r.Shadow.Virtual.SMRT.SamplesPerRayDirectional=4
+r.Shadow.Virtual.SMRT.RayCountLocal=8
+r.Shadow.Virtual.SMRT.SamplesPerRayLocal=4
 
 [ShadowQuality@3]
 r.LightFunctionQuality=1
@@ -157,6 +178,13 @@ r.VolumetricFog.GridSizeZ=128
 r.VolumetricFog.HistoryMissSupersampleCount=4
 r.LightMaxDrawDistanceScale=1
 r.CapsuleShadows=1
+r.Shadow.Virtual.MaxPhysicalPages=4096
+r.Shadow.Virtual.ResolutionLodBiasDirectional=-1.5
+r.Shadow.Virtual.ResolutionLodBiasLocal=0.0
+r.Shadow.Virtual.SMRT.RayCountDirectional=8
+r.Shadow.Virtual.SMRT.SamplesPerRayDirectional=8
+r.Shadow.Virtual.SMRT.RayCountLocal=8
+r.Shadow.Virtual.SMRT.SamplesPerRayLocal=8
 
 [ShadowQuality@Cine]
 r.LightFunctionQuality=1
@@ -177,6 +205,13 @@ r.VolumetricFog.GridSizeZ=128
 r.VolumetricFog.HistoryMissSupersampleCount=16
 r.LightMaxDrawDistanceScale=1
 r.CapsuleShadows=1
+r.Shadow.Virtual.MaxPhysicalPages=8192
+r.Shadow.Virtual.ResolutionLodBiasDirectional=-2.5
+r.Shadow.Virtual.ResolutionLodBiasLocal=-1.0
+r.Shadow.Virtual.SMRT.RayCountDirectional=16
+r.Shadow.Virtual.SMRT.SamplesPerRayDirectional=16
+r.Shadow.Virtual.SMRT.RayCountLocal=16
+r.Shadow.Virtual.SMRT.SamplesPerRayLocal=8
 
 ;-----------------------------------------------------------------------------------------------------------------
 
diff --git a/Engine/Source/Runtime/Renderer/Private/VirtualShadowMaps/VirtualShadowMapArray.cpp b/Engine/Source/Runtime/Renderer/Private/VirtualShadowMaps/VirtualShadowMapArray.cpp
--- a/Engine/Source/Runtime/Renderer/Private/VirtualShadowMaps/VirtualShadowMapArray.cpp
+++ b/Engine/Source/Runtime/Renderer/Private/VirtualShadowMaps/VirtualShadowMapArray.cpp
@@ -51,7 +51,7 @@ TAutoConsoleVariable<int32> CVarMaxPhysicalPages(
 	TEXT("r.Shadow.Virtual.MaxPhysicalPages"),
 	2048,
 	TEXT("Maximum number of physical pages in the pool."),
-	ECVF_RenderThreadSafe
+	ECVF_Scalability | ECVF_RenderThreadSafe
 );
 
 TAutoConsoleVariable<int32> CVarCacheStaticSeparate(
@@ -81,7 +81,7 @@ static TAutoConsoleVariable<float> CVarResolutionLodBiasLocal(
 	TEXT("r.Shadow.Virtual.ResolutionLodBiasLocal"),
 	0.0f,
 	TEXT("Bias applied to LOD calculations for local lights. -1.0 doubles resolution, 1.0 halves it and so on."),
-	ECVF_RenderThreadSafe
+	ECVF_Scalability | ECVF_RenderThreadSafe
 );
 
 static TAutoConsoleVariable<float> CVarPageDilationBorderSizeDirectional(
@@ -112,7 +112,7 @@ TAutoConsoleVariable<int32> CVarMarkCoarsePagesDirectional(
 	1,
 	TEXT("Marks coarse pages in directional light virtual shadow maps so that low resolution data is available everywhere.")
 	TEXT("Ability to disable is primarily for profiling and debugging."),
-	ECVF_RenderThreadSafe
+	ECVF_Scalability | ECVF_RenderThreadSafe
 );
 
 TAutoConsoleVariable<int32> CVarMarkCoarsePagesLocal(
@@ -120,7 +120,7 @@ TAutoConsoleVariable<int32> CVarMarkCoarsePagesLocal(
 	1,
 	TEXT("Marks coarse pages in local light virtual shadow maps so that low resolution data is available everywhere.")
 	TEXT("Ability to disable is primarily for profiling and debugging."),
-	ECVF_RenderThreadSafe
+	ECVF_Scalability | ECVF_RenderThreadSafe
 );
 
 TAutoConsoleVariable<int32> CVarCoarsePagesIncludeNonNanite(
@@ -128,7 +128,7 @@ TAutoConsoleVariable<int32> CVarCoarsePagesIncludeNonNanite(
 	1,
 	TEXT("Include non-nanite geometry in coarse pages.")
 	TEXT("Rendering non-nanite geometry into large coarse pages can be expensive; disabling this can be a significant performance win."),
-	ECVF_RenderThreadSafe
+	ECVF_Scalability | ECVF_RenderThreadSafe
 );
 
 static TAutoConsoleVariable<int32> CVarShowClipmapStats(
@@ -179,7 +179,7 @@ TAutoConsoleVariable<int32> CVarVirtualShadowOnePassProjectionMaxLights(
 	16,
 	TEXT("Maximum lights per pixel that get full filtering when using one pass projection and clustered shading.")
 	TEXT("Generally set to 8 (32bpp), 16 (64bpp) or 32 (128bpp). Lower values require less transient VRAM during the lighting pass."),
-	ECVF_RenderThreadSafe
+	ECVF_Scalability | ECVF_RenderThreadSafe
 );
 
 TAutoConsoleVariable<int32> CVarDoNonNaniteBatching(
diff --git a/Engine/Source/Runtime/Renderer/Private/VirtualShadowMaps/VirtualShadowMapCacheManager.cpp b/Engine/Source/Runtime/Renderer/Private/VirtualShadowMaps/VirtualShadowMapCacheManager.cpp
--- a/Engine/Source/Runtime/Renderer/Private/VirtualShadowMaps/VirtualShadowMapCacheManager.cpp
+++ b/Engine/Source/Runtime/Renderer/Private/VirtualShadowMaps/VirtualShadowMapCacheManager.cpp
@@ -83,11 +83,16 @@ void FVirtualShadowMapCacheEntry::UpdateClipmap(
 		}
 	}
 
-	if (bCacheValid)
+	bool bRadiusMatches = (ViewRadiusZ == Clipmap.ViewRadiusZ);
+	if (!bRadiusMatches)
+	{
+		// These really should be exact by construction currently
+		UE_LOG(LogRenderer, Warning, TEXT("Invalidated clipmap level (VSM %d) due to Z radius mismatch"), VirtualShadowMapId);
+	}
+
+	if (bCacheValid && bRadiusMatches)
 	{
 		PrevVirtualShadowMapId = CurrentVirtualShadowMapId;
-		// These really should be exact by construction
-		check(ViewRadiusZ == Clipmap.ViewRadiusZ);
 	}
 	else
 	{
diff --git a/Engine/Source/Runtime/Renderer/Private/VirtualShadowMaps/VirtualShadowMapClipmap.cpp b/Engine/Source/Runtime/Renderer/Private/VirtualShadowMaps/VirtualShadowMapClipmap.cpp
--- a/Engine/Source/Runtime/Renderer/Private/VirtualShadowMaps/VirtualShadowMapClipmap.cpp
+++ b/Engine/Source/Runtime/Renderer/Private/VirtualShadowMaps/VirtualShadowMapClipmap.cpp
@@ -15,33 +15,33 @@ static TAutoConsoleVariable<float> CVarVirtualShadowMapResolutionLodBiasDirectio
 	TEXT( "r.Shadow.Virtual.ResolutionLodBiasDirectional" ),
 	-0.5f,
 	TEXT( "Bias applied to LOD calculations for directional lights. -1.0 doubles resolution, 1.0 halves it and so on." ),
-	ECVF_RenderThreadSafe
+	ECVF_Scalability | ECVF_RenderThreadSafe
 );
 static TAutoConsoleVariable<int32> CVarVirtualShadowMapClipmapFirstLevel(
 	TEXT( "r.Shadow.Virtual.Clipmap.FirstLevel" ),
 	6,
 	TEXT( "First level of the virtual clipmap. Lower values allow higher resolution shadows closer to the camera." ),
-	ECVF_RenderThreadSafe
+	ECVF_Scalability | ECVF_RenderThreadSafe
 );
 static TAutoConsoleVariable<int32> CVarVirtualShadowMapClipmapLastLevel(
 	TEXT( "r.Shadow.Virtual.Clipmap.LastLevel" ),
 	22,
 	TEXT( "Last level of the virtual climap. Indirectly determines radius the clipmap can cover." ),
-	ECVF_RenderThreadSafe
+	ECVF_Scalability | ECVF_RenderThreadSafe
 );
 
 TAutoConsoleVariable<int32> CVarVirtualShadowMapClipmapFirstCoarseLevel(
 	TEXT("r.Shadow.Virtual.Clipmap.FirstCoarseLevel"),
 	15,
 	TEXT("First level of the clipmap to mark coarse pages for. Lower values allow higher resolution coarse pages near the camera but increase total page counts."),
-	ECVF_RenderThreadSafe
+	ECVF_Scalability | ECVF_RenderThreadSafe
 );
 
 TAutoConsoleVariable<int32> CVarVirtualShadowMapClipmapLastCoarseLevel(
 	TEXT("r.Shadow.Virtual.Clipmap.LastCoarseLevel"),
 	18,
 	TEXT("Last level of the clipmap to mark coarse pages for. Higher values provide dense clipmap data for a longer radius but increase total page counts."),
-	ECVF_RenderThreadSafe
+	ECVF_Scalability | ECVF_RenderThreadSafe
 );
 
 // "Virtual" clipmap level to clipmap radius
@@ -143,8 +143,10 @@ FVirtualShadowMapClipmap::FVirtualShadowMapClipmap(
 
 		// We expand the depth range of the clipmap level to allow for camera movement without having to invalidate cached shadow data
 		// (See VirtualShadowMapCacheManager::UpdateClipmap for invalidation logic.)
-		// Hard-coded constant currently; should probably always be at least 2.0f, otherwise the cache will be frequently invalidated.
-		const float ViewRadiusZMultiplier = 5.0f;
+		// This also better accomodates SMRT where we want to avoid stepping outside of the Z bounds of a given clipmap
+		// NOTE: It's tempting to use a single global Z range for the entire clipmap (which avoids some SMRT overhead too)
+		// but this can cause precision issues with cached pages very near the camera.
+		const float ViewRadiusZMultiplier = 1000.0f;
 
 		float ViewRadiusZ = RawLevelRadius * ViewRadiusZMultiplier;
 		float ViewCenterDeltaZ = 0.0f;
diff --git a/Engine/Source/Runtime/Renderer/Private/VirtualShadowMaps/VirtualShadowMapProjection.cpp b/Engine/Source/Runtime/Renderer/Private/VirtualShadowMaps/VirtualShadowMapProjection.cpp
--- a/Engine/Source/Runtime/Renderer/Private/VirtualShadowMaps/VirtualShadowMapProjection.cpp
+++ b/Engine/Source/Runtime/Renderer/Private/VirtualShadowMaps/VirtualShadowMapProjection.cpp
@@ -29,7 +29,7 @@ static TAutoConsoleVariable<float> CVarContactShadowLength(
 	TEXT( "r.Shadow.Virtual.ContactShadowLength" ),
 	0.02f,
 	TEXT( "Length of the screen space contact shadow trace (smart shadow bias) before the virtual shadow map lookup." ),
-	ECVF_RenderThreadSafe
+	ECVF_Scalability | ECVF_RenderThreadSafe
 );
 
 static TAutoConsoleVariable<float> CVarNormalBias(
@@ -37,28 +37,28 @@ static TAutoConsoleVariable<float> CVarNormalBias(
 	0.5f,
 	TEXT( "Receiver offset along surface normal for shadow lookup. Scaled by distance to camera." )
 	TEXT( "Higher values avoid artifacts on surfaces nearly parallel to the light, but also visibility offset shadows and increase the chance of hitting unmapped pages." ),
-	ECVF_RenderThreadSafe
+	ECVF_Scalability | ECVF_RenderThreadSafe
 );
 
 TAutoConsoleVariable<int32> CVarVirtualShadowOnePassProjection(
 	TEXT("r.Shadow.Virtual.OnePassProjection"),
 	0,
 	TEXT("Single pass projects all local VSMs culled with the light grid. Used in conjunction with clustered deferred shading."),
-	ECVF_RenderThreadSafe
+	ECVF_Scalability | ECVF_RenderThreadSafe
 );
 
 static TAutoConsoleVariable<int32> CVarSMRTRayCountLocal(
 	TEXT( "r.Shadow.Virtual.SMRT.RayCountLocal" ),
 	7,
 	TEXT( "Ray count for shadow map tracing of local lights. 0 = disabled." ),
-	ECVF_RenderThreadSafe
+	ECVF_Scalability | ECVF_RenderThreadSafe
 );
 
 static TAutoConsoleVariable<int32> CVarSMRTSamplesPerRayLocal(
 	TEXT( "r.Shadow.Virtual.SMRT.SamplesPerRayLocal" ),
 	8,
 	TEXT( "Shadow map samples per ray for local lights" ),
-	ECVF_RenderThreadSafe
+	ECVF_Scalability | ECVF_RenderThreadSafe
 );
 
 static TAutoConsoleVariable<float> CVarSMRTMaxRayAngleFromLight(
@@ -74,14 +74,14 @@ static TAutoConsoleVariable<int32> CVarSMRTRayCountDirectional(
 	TEXT( "r.Shadow.Virtual.SMRT.RayCountDirectional" ),
 	7,
 	TEXT( "Ray count for shadow map tracing of directional lights. 0 = disabled." ),
-	ECVF_RenderThreadSafe
+	ECVF_Scalability | ECVF_RenderThreadSafe
 );
 
 static TAutoConsoleVariable<int32> CVarSMRTSamplesPerRayDirectional(
 	TEXT( "r.Shadow.Virtual.SMRT.SamplesPerRayDirectional" ),
 	8,
 	TEXT( "Shadow map samples per ray for directional lights" ),
-	ECVF_RenderThreadSafe
+	ECVF_Scalability | ECVF_RenderThreadSafe
 );
 
 static TAutoConsoleVariable<float> CVarSMRTRayLengthScaleDirectional(
@@ -97,7 +97,7 @@ static TAutoConsoleVariable<int32> CVarSMRTAdaptiveRayCount(
 	TEXT( "r.Shadow.Virtual.SMRT.AdaptiveRayCount" ),
 	1,
 	TEXT( "Shoot fewer rays in fully shadowed and unshadowed regions. Currently only supported with OnePassProjection. " ),
-	ECVF_RenderThreadSafe
+	ECVF_Scalability | ECVF_RenderThreadSafe
 );
 
 static TAutoConsoleVariable<float> CVarSMRTTexelDitherScale(
@@ -105,7 +105,7 @@ static TAutoConsoleVariable<float> CVarSMRTTexelDitherScale(
 	2.0f,
 	TEXT( "Applies a dither to the shadow map ray casts to help hide aliasing due to insufficient shadow resolution.\n" )
 	TEXT( "This is usually desirable, but it can occasionally cause shadows from thin geometry to separate from their casters at shallow light angles." ),
-	ECVF_RenderThreadSafe
+	ECVF_Scalability | ECVF_RenderThreadSafe
 );
 
 
