Remove RenderLight from DeferredRendering, and change it to a standalone static/internal function to light rendering to ease its refactor to RDG.

diff --git a/Engine/Source/Runtime/Renderer/Private/AnisotropyRendering.cpp b/Engine/Source/Runtime/Renderer/Private/AnisotropyRendering.cpp
--- a/Engine/Source/Runtime/Renderer/Private/AnisotropyRendering.cpp
+++ b/Engine/Source/Runtime/Renderer/Private/AnisotropyRendering.cpp
@@ -238,21 +238,31 @@ bool FAnisotropyMeshProcessor::Process(
 	return true;
 }
 
-bool FDeferredShadingSceneRenderer::ShouldRenderAnisotropyPass() const
+bool ShouldRenderAnisotropyPass(const FViewInfo& View)
 {
-	if (!SupportsAnisotropicMaterials(FeatureLevel, ShaderPlatform))
+	if (!SupportsAnisotropicMaterials(View.FeatureLevel, View.GetShaderPlatform()))
 	{
 		return false;
 	}
 
-	if (IsAnyForwardShadingEnabled(GetFeatureLevelShaderPlatform(FeatureLevel)))
+	if (IsAnyForwardShadingEnabled(GetFeatureLevelShaderPlatform(View.FeatureLevel)))
 	{
 		return false;
 	}
 
-	for (auto& View : Views)
+	if (View.ShouldRenderView() && View.ParallelMeshDrawCommandPasses[EMeshPass::AnisotropyPass].HasAnyDraw())
 	{
-		if (View.ShouldRenderView() && View.ParallelMeshDrawCommandPasses[EMeshPass::AnisotropyPass].HasAnyDraw())
+		return true;
+	}
+
+	return false;
+}
+
+bool ShouldRenderAnisotropyPass(const TArray<FViewInfo>& Views)
+{
+	for (const FViewInfo& View : Views)
+	{
+		if (ShouldRenderAnisotropyPass(View))
 		{
 			return true;
 		}
diff --git a/Engine/Source/Runtime/Renderer/Private/AnisotropyRendering.h b/Engine/Source/Runtime/Renderer/Private/AnisotropyRendering.h
--- a/Engine/Source/Runtime/Renderer/Private/AnisotropyRendering.h
+++ b/Engine/Source/Runtime/Renderer/Private/AnisotropyRendering.h
@@ -49,3 +49,6 @@ protected:
 		ERasterizerCullMode MeshCullMode
 		);
 };
+
+bool ShouldRenderAnisotropyPass(const FViewInfo& Views);
+bool ShouldRenderAnisotropyPass(const TArray<FViewInfo>& Views);
\ No newline at end of file
diff --git a/Engine/Source/Runtime/Renderer/Private/BasePassRendering.cpp b/Engine/Source/Runtime/Renderer/Private/BasePassRendering.cpp
--- a/Engine/Source/Runtime/Renderer/Private/BasePassRendering.cpp
+++ b/Engine/Source/Runtime/Renderer/Private/BasePassRendering.cpp
@@ -19,6 +19,7 @@
 #include "ClearQuad.h"
 #include "ProfilingDebugging/CpuProfilerTrace.h"
 #include "DebugProbeRendering.h"
+#include "AnisotropyRendering.h"
 #include "Nanite/NaniteVisualize.h"
 
 // Changing this causes a full shader recompile
@@ -889,7 +890,7 @@ void FDeferredShadingSceneRenderer::RenderBasePass(
 		ClearGBufferAtMaxZ(GraphBuilder, Views, BasePassRenderTargets, SceneColorClearValue);
 	}
 
-	if (ShouldRenderAnisotropyPass())
+	if (ShouldRenderAnisotropyPass(Views))
 	{
 		GraphBuilder.SetCommandListStat(GET_STATID(STAT_CLM_AnisotropyPass));
 		RenderAnisotropyPass(GraphBuilder, SceneTextures, bEnableParallelBasePasses);
diff --git a/Engine/Source/Runtime/Renderer/Private/DeferredShadingRenderer.h b/Engine/Source/Runtime/Renderer/Private/DeferredShadingRenderer.h
--- a/Engine/Source/Runtime/Renderer/Private/DeferredShadingRenderer.h
+++ b/Engine/Source/Runtime/Renderer/Private/DeferredShadingRenderer.h
@@ -214,8 +214,6 @@ public:
 		bool bNaniteEnabled,
 		const TArrayView<Nanite::FRasterResults>& NaniteRasterResults);
 
-	bool ShouldRenderAnisotropyPass() const;
-
 	void RenderAnisotropyPass(
 		FRDGBuilder& GraphBuilder,
 		FSceneTextures& SceneTextures,
@@ -699,25 +697,6 @@ private:
 		bool bRenderingPreviewShadowsIndicator, 
 		bool bUseHairStrands);
 
-	/** Used by RenderLights to render a light to the scene color buffer. */
-	void RenderLight(
-		FRHICommandList& RHICmdList,
-		const FViewInfo& View,
-		const FLightSceneInfo* LightSceneInfo,
-		FRenderLightParameters* PassParameters,
-		FRHITexture* ScreenShadowMaskTexture,
-		FRHITexture* LightingChannelTexture,
-		bool bRenderOverlap,
-		bool bIssueDrawEvent);
-
-	void RenderLight(
-		FRDGBuilder& GraphBuilder,
-		const FMinimalSceneTextures& SceneTextures,
-		const FLightSceneInfo* LightSceneInfo,
-		FRDGTextureRef ScreenShadowMaskTexture,
-		FRDGTextureRef LightingChannelsTexture,
-		bool bRenderOverlap);
-
 	void RenderLightsForHair(
 		FRDGBuilder& GraphBuilder,
 		TRDGUniformBufferRef<FSceneTextureUniformParameters> SceneTexturesUniformBuffer,
diff --git a/Engine/Source/Runtime/Renderer/Private/LightRendering.cpp b/Engine/Source/Runtime/Renderer/Private/LightRendering.cpp
--- a/Engine/Source/Runtime/Renderer/Private/LightRendering.cpp
+++ b/Engine/Source/Runtime/Renderer/Private/LightRendering.cpp
@@ -23,6 +23,7 @@
 #include "Strata/Strata.h"
 #include "VirtualShadowMaps/VirtualShadowMapProjection.h"
 #include "HairStrands/HairStrandsData.h"
+#include "AnisotropyRendering.h"
 #include "Engine/SubsurfaceProfile.h"
 
 // ENABLE_DEBUG_DISCARD_PROP is used to test the lighting code by allowing to discard lights to see how performance scales
@@ -179,6 +180,26 @@ bool ShouldRenderRayTracingShadowsForLight(const FLightSceneInfoCompact& LightIn
 }
 #endif // RHI_RAYTRACING
 
+static void RenderLight(
+	FRHICommandList& RHICmdList,
+	const FScene* Scene,
+	const FViewInfo& View,
+	const FLightSceneInfo* LightSceneInfo,
+	FRenderLightParameters* PassParameters,	// If this is null, it means we cannot use Strata tiles and fallback to previous behavior.
+	FRHITexture* ScreenShadowMaskTexture,
+	FRHITexture* LightingChannelsTexture,
+	bool bRenderOverlap, bool bIssueDrawEvent);
+
+static void RenderLight(
+	FRDGBuilder& GraphBuilder,
+	const FScene* Scene,
+	const TArray<FViewInfo>& Views,
+	const FMinimalSceneTextures& SceneTextures,
+	const FLightSceneInfo* LightSceneInfo,
+	FRDGTextureRef ScreenShadowMaskTexture,
+	FRDGTextureRef LightingChannelsTexture,
+	bool bRenderOverlap);
+
 FDeferredLightUniformStruct GetDeferredLightParameters(const FSceneView& View, const FLightSceneInfo& LightSceneInfo)
 {
 	FDeferredLightUniformStruct Parameters;
@@ -1264,7 +1285,7 @@ void FDeferredShadingSceneRenderer::RenderLights(
 							SCOPED_GPU_MASK(RHICmdList, View.GPUMask);
 
 							// Render the light to the scene color buffer, using a 1x1 white texture as input
-							RenderLight(RHICmdList, View, LightSceneInfo, PassParameters, nullptr, TryGetRHI(LightingChannelsTexture), false, false);
+							RenderLight(RHICmdList, Scene, View, LightSceneInfo, PassParameters, nullptr, TryGetRHI(LightingChannelsTexture), false, false);
 						}
 					});
 				}
@@ -1905,7 +1926,7 @@ void FDeferredShadingSceneRenderer::RenderLights(
 				if (bDirectLighting)
 				{
 					const bool bRenderOverlap = false;
-					RenderLight(GraphBuilder, SceneTextures, &LightSceneInfo, ScreenShadowMaskTexture, LightingChannelsTexture, bRenderOverlap);
+					RenderLight(GraphBuilder, Scene, Views, SceneTextures, &LightSceneInfo, ScreenShadowMaskTexture, LightingChannelsTexture, bRenderOverlap);
 				}
 
 				if (bUseHairLighting)
@@ -1959,7 +1980,7 @@ void FDeferredShadingSceneRenderer::RenderLightArrayForOverlapViewmode(
 		for (const FViewInfo& View : Views)
 		{
 			SCOPED_GPU_MASK(RHICmdList, View.GPUMask);
-			RenderLight(RHICmdList, View, LightSceneInfo, nullptr, nullptr, LightingChannelsTexture, true, false);
+			RenderLight(RHICmdList, Scene, View, LightSceneInfo, nullptr, nullptr, LightingChannelsTexture, true, false);
 		}
 	}
 }
@@ -2139,8 +2160,9 @@ static void BindAtmosphereAndCloudResources(
  * @return true if anything got rendered
  */
 
-void FDeferredShadingSceneRenderer::RenderLight(
+static void RenderLight(
 	FRHICommandList& RHICmdList,
+	const FScene* Scene,
 	const FViewInfo& View,
 	const FLightSceneInfo* LightSceneInfo,
 	FRenderLightParameters* PassParameters,	// If this is null, it means we cannot use Strata tiles and fallback to previous behavior.
@@ -2239,7 +2261,7 @@ void FDeferredShadingSceneRenderer::RenderLight(
 			PermutationVector.Set< FDeferredLightPS::FInverseSquaredDim >( false );
 			PermutationVector.Set< FDeferredLightPS::FVisualizeCullingDim >( View.Family->EngineShowFlags.VisualizeLightCulling );
 			PermutationVector.Set< FDeferredLightPS::FLightingChannelsDim >( View.bUsesLightingChannels );
-			PermutationVector.Set< FDeferredLightPS::FAnistropicMaterials >(ShouldRenderAnisotropyPass());
+			PermutationVector.Set< FDeferredLightPS::FAnistropicMaterials >(ShouldRenderAnisotropyPass(View));
 			PermutationVector.Set< FDeferredLightPS::FTransmissionDim >( bTransmission );
 			PermutationVector.Set< FDeferredLightPS::FHairLighting>(0);
 			// Only directional lights are rendered in this path, so we only need to check if it is use to light the atmosphere
@@ -2320,7 +2342,7 @@ void FDeferredShadingSceneRenderer::RenderLight(
 			PermutationVector.Set< FDeferredLightPS::FInverseSquaredDim >( LightProxy->IsInverseSquared() );
 			PermutationVector.Set< FDeferredLightPS::FVisualizeCullingDim >( View.Family->EngineShowFlags.VisualizeLightCulling );
 			PermutationVector.Set< FDeferredLightPS::FLightingChannelsDim >( View.bUsesLightingChannels );
-			PermutationVector.Set< FDeferredLightPS::FAnistropicMaterials >(ShouldRenderAnisotropyPass() && !LightSceneInfo->Proxy->IsRectLight());
+			PermutationVector.Set< FDeferredLightPS::FAnistropicMaterials >(ShouldRenderAnisotropyPass(View) && !LightSceneInfo->Proxy->IsRectLight());
 			PermutationVector.Set< FDeferredLightPS::FTransmissionDim >( bTransmission );
 			PermutationVector.Set< FDeferredLightPS::FHairLighting>(0);
 			PermutationVector.Set < FDeferredLightPS::FAtmosphereTransmittance >(false);
@@ -2400,8 +2422,10 @@ void FDeferredShadingSceneRenderer::RenderLight(
 	}
 }
 
-void FDeferredShadingSceneRenderer::RenderLight(
+static void RenderLight(
 	FRDGBuilder& GraphBuilder,
+	const FScene* Scene,
+	const TArray<FViewInfo>& Views,
 	const FMinimalSceneTextures& SceneTextures,
 	const FLightSceneInfo* LightSceneInfo,
 	FRDGTextureRef ScreenShadowMaskTexture,
@@ -2424,10 +2448,11 @@ void FDeferredShadingSceneRenderer::RenderLight(
 			RDG_EVENT_NAME("StandardDeferredLighting"),
 			PassParameters,
 			PassFlags,
-			[this, &View, LightSceneInfo, ScreenShadowMaskTexture, LightingChannelsTexture, bRenderOverlap, PassParameters](FRHICommandList& RHICmdList)
+			[Scene, &View, LightSceneInfo, ScreenShadowMaskTexture, LightingChannelsTexture, bRenderOverlap, PassParameters](FRHICommandList& RHICmdList)
 		{
 			RenderLight(
 				RHICmdList,
+				Scene, 
 				View,
 				LightSceneInfo,
 				PassParameters,
