Split RenderLight function to avoid lambdas calls.

diff --git a/Engine/Source/Runtime/Renderer/Private/LightRendering.cpp b/Engine/Source/Runtime/Renderer/Private/LightRendering.cpp
--- a/Engine/Source/Runtime/Renderer/Private/LightRendering.cpp
+++ b/Engine/Source/Runtime/Renderer/Private/LightRendering.cpp
@@ -2160,7 +2160,7 @@ static void BindAtmosphereAndCloudResources(
  * @return true if anything got rendered
  */
 
-static void RenderLight(
+static void InternalRenderLight(
 	FRHICommandList& RHICmdList,
 	const FScene* Scene,
 	const FViewInfo& View,
@@ -2168,21 +2168,9 @@ static void RenderLight(
 	FRenderLightParameters* PassParameters,	// If this is null, it means we cannot use Strata tiles and fallback to previous behavior.
 	FRHITexture* ScreenShadowMaskTexture,
 	FRHITexture* LightingChannelsTexture,
-	bool bRenderOverlap, bool bIssueDrawEvent)
+	bool bRenderOverlap, bool bIssueDrawEvent,
+	bool bEnableStrataStencilTest, bool bEnableStrataTiledPass, uint32 StrataTileMaterialType)
 {
-	// Ensure the light is valid for this view
-	if (!LightSceneInfo->ShouldRenderLight(View))
-	{
-		return;
-	}
-
-	SCOPE_CYCLE_COUNTER(STAT_DirectLightRenderingTime);
-	INC_DWORD_STAT(STAT_NumLightsUsingStandardDeferred);
-	SCOPED_CONDITIONAL_DRAW_EVENT(RHICmdList, StandardDeferredLighting, bIssueDrawEvent);
-
-	auto RenderInternalLight = [&]( bool bEnableStrataStencilTest, bool bEnableStrataTiledPass, uint32 StrataTileMaterialType)
-	{
-
 	FGraphicsPipelineStateInitializer GraphicsPSOInit;
 	RHICmdList.ApplyCachedRenderTargets(GraphicsPSOInit);
 
@@ -2389,8 +2377,28 @@ static void RenderLight(
 		{
 			StencilingGeometry::DrawCone(RHICmdList);
 		}
+	}	
+}
+
+static void RenderLight(
+	FRHICommandList& RHICmdList,
+	const FScene* Scene,
+	const FViewInfo& View,
+	const FLightSceneInfo* LightSceneInfo,
+	FRenderLightParameters* PassParameters,	// If this is null, it means we cannot use Strata tiles and fallback to previous behavior.
+	FRHITexture* ScreenShadowMaskTexture,
+	FRHITexture* LightingChannelsTexture,
+	bool bRenderOverlap, bool bIssueDrawEvent)
+{
+	// Ensure the light is valid for this view
+	if (!LightSceneInfo->ShouldRenderLight(View))
+	{
+		return;
 	}
-	};
+
+	SCOPE_CYCLE_COUNTER(STAT_DirectLightRenderingTime);
+	INC_DWORD_STAT(STAT_NumLightsUsingStandardDeferred);
+	SCOPED_CONDITIONAL_DRAW_EVENT(RHICmdList, StandardDeferredLighting, bIssueDrawEvent);
 
 	const bool bStrataClassificationEnabled = Strata::IsStrataEnabled() && Strata::IsClassificationEnabled();
 	const bool bTilePassesReadingStrataEnabled = Strata::ShouldPassesReadingStrataBeTiled(Scene->GetFeatureLevel());
@@ -2402,11 +2410,11 @@ static void RenderLight(
 
 		{
 			SCOPED_DRAW_EVENT(RHICmdList, StrataSimpleMaterial);
-			RenderInternalLight(bEnableStrataStencilTest, bEnableStrataTiledPass, EStrataTileMaterialType::ESimple);
+			InternalRenderLight(RHICmdList, Scene, View, LightSceneInfo, PassParameters, ScreenShadowMaskTexture, LightingChannelsTexture, bRenderOverlap, bIssueDrawEvent, bEnableStrataStencilTest, bEnableStrataTiledPass, EStrataTileMaterialType::ESimple);
 		}
 		{
 			SCOPED_DRAW_EVENT(RHICmdList, StrataComplexMaterial);
-			RenderInternalLight(bEnableStrataStencilTest, bEnableStrataTiledPass, EStrataTileMaterialType::EComplex);
+			InternalRenderLight(RHICmdList, Scene, View, LightSceneInfo, PassParameters, ScreenShadowMaskTexture, LightingChannelsTexture, bRenderOverlap, bIssueDrawEvent, bEnableStrataStencilTest, bEnableStrataTiledPass, EStrataTileMaterialType::EComplex);
 		}
 	}
 	else
@@ -2414,10 +2422,10 @@ static void RenderLight(
 		const bool bEnableStrataTiledPass = false;
 		const bool bEnableStrataStencilTest = bStrataClassificationEnabled;
 
-		RenderInternalLight(bEnableStrataStencilTest, bEnableStrataTiledPass, EStrataTileMaterialType::EComplex);
+		InternalRenderLight(RHICmdList, Scene, View, LightSceneInfo, PassParameters, ScreenShadowMaskTexture, LightingChannelsTexture, bRenderOverlap, bIssueDrawEvent, bEnableStrataStencilTest, bEnableStrataTiledPass, EStrataTileMaterialType::EComplex);
 		if (Strata::IsStrataEnabled() && Strata::IsClassificationEnabled())
 		{
-			RenderInternalLight(bEnableStrataStencilTest, bEnableStrataTiledPass, EStrataTileMaterialType::ESimple);
+			InternalRenderLight(RHICmdList, Scene, View, LightSceneInfo, PassParameters, ScreenShadowMaskTexture, LightingChannelsTexture, bRenderOverlap, bIssueDrawEvent, bEnableStrataStencilTest, bEnableStrataTiledPass, EStrataTileMaterialType::ESimple);
 		}
 	}
 }
