Change RenderLight to use render graph version.

diff --git a/Engine/Source/Runtime/Renderer/Private/DeferredShadingRenderer.h b/Engine/Source/Runtime/Renderer/Private/DeferredShadingRenderer.h
--- a/Engine/Source/Runtime/Renderer/Private/DeferredShadingRenderer.h
+++ b/Engine/Source/Runtime/Renderer/Private/DeferredShadingRenderer.h
@@ -589,12 +589,6 @@ private:
 		FRDGTextureRef LightingChannelsTexture,
 		FSortedLightSetSceneInfo& SortedLightSet);
 
-	/** Renders an array of lights for the stationary light overlap viewmode. */
-	void RenderLightArrayForOverlapViewmode(
-		FRHICommandList& RHICmdList,
-		FRHITexture* LightingChannelsTexture,
-		const TSparseArray<FLightSceneInfoCompact, TAlignedSparseArrayAllocator<alignof(FLightSceneInfoCompact)>>& LightArray);
-
 	/** Render stationary light overlap as complexity to scene color. */
 	void RenderStationaryLightOverlap(
 		FRDGBuilder& GraphBuilder,
diff --git a/Engine/Source/Runtime/Renderer/Private/LightRendering.cpp b/Engine/Source/Runtime/Renderer/Private/LightRendering.cpp
--- a/Engine/Source/Runtime/Renderer/Private/LightRendering.cpp
+++ b/Engine/Source/Runtime/Renderer/Private/LightRendering.cpp
@@ -180,25 +180,16 @@ bool ShouldRenderRayTracingShadowsForLight(const FLightSceneInfoCompact& LightIn
 }
 #endif // RHI_RAYTRACING
 
-static void RenderLight(
-	FRHICommandList& RHICmdList,
-	const FScene* Scene,
-	const FViewInfo& View,
-	const FLightSceneInfo* LightSceneInfo,
-	FRenderLightParameters* PassParameters,	// If this is null, it means we cannot use Strata tiles and fallback to previous behavior.
-	FRHITexture* ScreenShadowMaskTexture,
-	FRHITexture* LightingChannelsTexture,
-	bool bRenderOverlap, bool bIssueDrawEvent);
-
 static void RenderLight(
 	FRDGBuilder& GraphBuilder,
 	const FScene* Scene,
-	const TArray<FViewInfo>& Views,
+	const FViewInfo& View,
 	const FMinimalSceneTextures& SceneTextures,
 	const FLightSceneInfo* LightSceneInfo,
 	FRDGTextureRef ScreenShadowMaskTexture,
 	FRDGTextureRef LightingChannelsTexture,
-	bool bRenderOverlap);
+	bool bRenderOverlap,
+	bool bCloudShadow);
 
 FDeferredLightUniformStruct GetDeferredLightParameters(const FSceneView& View, const FLightSceneInfo& LightSceneInfo)
 {
@@ -1262,32 +1253,18 @@ void FDeferredShadingSceneRenderer::RenderLights(
 				RenderSimpleLightsStandardDeferred(GraphBuilder, SceneTextures, SortedLightSet.SimpleLights);
 			}
 
+			// Draw non-shadowed non-light function lights without changing render targets between them
+			for (int32 ViewIndex = 0, ViewCount = Views.Num(); ViewIndex < ViewCount; ++ViewIndex)
 			{
-				for (int32 ViewIndex = 0, ViewCount = Views.Num(); ViewIndex < ViewCount; ++ViewIndex)
-				{
-					const FViewInfo& View = Views[ViewIndex];
-					FRenderLightParameters* PassParameters = GraphBuilder.AllocParameters<FRenderLightParameters>();
-					GetRenderLightParameters(View, SceneTextures, View.HairStrandsViewData, nullptr, LightingChannelsTexture, {}, *PassParameters);
-
-					GraphBuilder.AddPass(
-						RDG_EVENT_NAME("StandardDeferredLighting"),
-						PassParameters,
-						ERDGPassFlags::Raster,
-						[this, &View, &SortedLights, LightingChannelsTexture, StandardDeferredStart, UnbatchedLightStart, PassParameters](FRHICommandList& RHICmdList)
-					{
-						// Draw non-shadowed non-light function lights without changing render targets between them
-						for (int32 LightIndex = StandardDeferredStart; LightIndex < UnbatchedLightStart; LightIndex++)
-						{
-							const FSortedLightSceneInfo& SortedLightInfo = SortedLights[LightIndex];
-							const FLightSceneInfo* const LightSceneInfo = SortedLightInfo.LightSceneInfo;
-
-							
-							SCOPED_GPU_MASK(RHICmdList, View.GPUMask);
+				const FViewInfo& View = Views[ViewIndex];
+				RDG_EVENT_SCOPE_CONDITIONAL(GraphBuilder, ViewCount > 1, "View%d", ViewIndex);
+				SCOPED_GPU_MASK(GraphBuilder.RHICmdList, View.GPUMask);
 
-							// Render the light to the scene color buffer, using a 1x1 white texture as input
-							RenderLight(RHICmdList, Scene, View, LightSceneInfo, PassParameters, nullptr, TryGetRHI(LightingChannelsTexture), false, false);
-						}
-					});
+				for (int32 LightIndex = StandardDeferredStart; LightIndex < UnbatchedLightStart; LightIndex++)
+				{
+					// Render the light to the scene color buffer, using a 1x1 white texture as input
+					const FLightSceneInfo* LightSceneInfo = SortedLights[LightIndex].LightSceneInfo;
+					RenderLight(GraphBuilder, Scene, View, SceneTextures, LightSceneInfo, nullptr, LightingChannelsTexture, false /*bRenderOverlap*/, false /*bCloudShadow*/);
 				}
 			}
 
@@ -1302,8 +1279,7 @@ void FDeferredShadingSceneRenderer::RenderLights(
 						// Draw non-shadowed non-light function lights without changing render targets between them
 						for (int32 LightIndex = StandardDeferredStart; LightIndex < UnbatchedLightStart; LightIndex++)
 						{
-							const FSortedLightSceneInfo& SortedLightInfo = SortedLights[LightIndex];
-							const FLightSceneInfo* const LightSceneInfo = SortedLightInfo.LightSceneInfo;
+							const FLightSceneInfo* LightSceneInfo = SortedLights[LightIndex].LightSceneInfo;
 							RenderLightForHair(GraphBuilder, View, SceneTextures.UniformBuffer, LightSceneInfo, NullScreenShadowMaskSubPixelTexture, LightingChannelsTexture, DummyTransmittanceMaskData, false /*bForwardRendering*/);
 						}
 					}
@@ -1925,8 +1901,14 @@ void FDeferredShadingSceneRenderer::RenderLights(
 				// Render the light to the scene color buffer, conditionally using the attenuation buffer or a 1x1 white texture as input 
 				if (bDirectLighting)
 				{
-					const bool bRenderOverlap = false;
-					RenderLight(GraphBuilder, Scene, Views, SceneTextures, &LightSceneInfo, ScreenShadowMaskTexture, LightingChannelsTexture, bRenderOverlap);
+					for (int32 ViewIndex = 0, ViewCount = Views.Num(); ViewIndex < ViewCount; ++ViewIndex)
+					{
+						const FViewInfo& View = Views[ViewIndex];
+
+						RDG_EVENT_SCOPE_CONDITIONAL(GraphBuilder, ViewCount > 1, "View%d", ViewIndex);
+						SCOPED_GPU_MASK(GraphBuilder.RHICmdList, View.GPUMask);
+						RenderLight(GraphBuilder, Scene, View, SceneTextures, &LightSceneInfo, ScreenShadowMaskTexture, LightingChannelsTexture, false /*bRenderOverlap*/, true /*bCloudShadow*/);
+					}
 				}
 
 				if (bUseHairLighting)
@@ -1952,9 +1934,12 @@ void FDeferredShadingSceneRenderer::RenderLights(
 	}
 }
 
-void FDeferredShadingSceneRenderer::RenderLightArrayForOverlapViewmode(
-	FRHICommandList& RHICmdList,
-	FRHITexture* LightingChannelsTexture,
+static void RenderLightArrayForOverlapViewmode(
+	FRDGBuilder& GraphBuilder,
+	const FScene* Scene,
+	const TArray<FViewInfo>& Views,
+	const FMinimalSceneTextures& SceneTextures,
+	FRDGTextureRef LightingChannelsTexture,
 	const TSparseArray<FLightSceneInfoCompact, TAlignedSparseArrayAllocator<alignof(FLightSceneInfoCompact)>>& LightArray)
 {
 	for (auto LightIt = LightArray.CreateConstIterator(); LightIt; ++LightIt)
@@ -1963,7 +1948,7 @@ void FDeferredShadingSceneRenderer::RenderLightArrayForOverlapViewmode(
 		const FLightSceneInfo* LightSceneInfo = LightSceneInfoCompact.LightSceneInfo;
 
 		// Nothing to do for black lights.
-		if(LightSceneInfoCompact.Color.IsAlmostBlack())
+		if (LightSceneInfoCompact.Color.IsAlmostBlack())
 		{
 			continue;
 		}
@@ -1979,8 +1964,8 @@ void FDeferredShadingSceneRenderer::RenderLightArrayForOverlapViewmode(
 		// Check if the light is visible in any of the views.
 		for (const FViewInfo& View : Views)
 		{
-			SCOPED_GPU_MASK(RHICmdList, View.GPUMask);
-			RenderLight(RHICmdList, Scene, View, LightSceneInfo, nullptr, nullptr, LightingChannelsTexture, true, false);
+			SCOPED_GPU_MASK(GraphBuilder.RHICmdList, View.GPUMask);
+			RenderLight(GraphBuilder, Scene, View, SceneTextures, LightSceneInfo, nullptr, LightingChannelsTexture, true /*bRenderOverlap*/, false /*bCloudShadow*/);
 		}
 	}
 }
@@ -1992,26 +1977,14 @@ void FDeferredShadingSceneRenderer::RenderStationaryLightOverlap(
 {
 	if (Scene->bIsEditorScene)
 	{
-		FRenderLightParameters* PassParameters = GraphBuilder.AllocParameters<FRenderLightParameters>();
-		GetRenderLightParameters(Views[0], SceneTextures, Views[0].HairStrandsViewData, nullptr, LightingChannelsTexture, {}, *PassParameters);
-
-		GraphBuilder.AddPass(
-			RDG_EVENT_NAME("StationaryLightOverlap"),
-			PassParameters,
-			ERDGPassFlags::Raster,
-			[this, LightingChannelsTexture](FRHICommandList& RHICmdList)
-		{
-			FRHITexture* LightingChannelsTextureRHI = TryGetRHI(LightingChannelsTexture);
-
-			// Clear to discard base pass values in scene color since we didn't skip that, to have valid scene depths
-			DrawClearQuad(RHICmdList, FLinearColor::Black);
+		// Clear to discard base pass values in scene color since we didn't skip that, to have valid scene depths
+		AddClearRenderTargetPass(GraphBuilder, SceneTextures.Color.Target, FLinearColor::Black);
 
-			RenderLightArrayForOverlapViewmode(RHICmdList, LightingChannelsTextureRHI, Scene->Lights);
+		RenderLightArrayForOverlapViewmode(GraphBuilder, Scene, Views, SceneTextures, LightingChannelsTexture, Scene->Lights);
 
-			//Note: making use of FScene::InvisibleLights, which contains lights that haven't been added to the scene in the same way as visible lights
-			// So code called by RenderLightArrayForOverlapViewmode must be careful what it accesses
-			RenderLightArrayForOverlapViewmode(RHICmdList, LightingChannelsTextureRHI, Scene->InvisibleLights);
-		});
+		//Note: making use of FScene::InvisibleLights, which contains lights that haven't been added to the scene in the same way as visible lights
+		// So code called by RenderLightArrayForOverlapViewmode must be careful what it accesses
+		RenderLightArrayForOverlapViewmode(GraphBuilder, Scene, Views, SceneTextures, LightingChannelsTexture, Scene->InvisibleLights);
 	}
 }
 
@@ -2168,7 +2141,7 @@ static void InternalRenderLight(
 	FRenderLightParameters* PassParameters,	// If this is null, it means we cannot use Strata tiles and fallback to previous behavior.
 	FRHITexture* ScreenShadowMaskTexture,
 	FRHITexture* LightingChannelsTexture,
-	bool bRenderOverlap, bool bIssueDrawEvent,
+	bool bRenderOverlap,
 	bool bEnableStrataStencilTest, bool bEnableStrataTiledPass, uint32 StrataTileMaterialType)
 {
 	FGraphicsPipelineStateInitializer GraphicsPSOInit;
@@ -2380,7 +2353,7 @@ static void InternalRenderLight(
 	}	
 }
 
-static void RenderLight(
+static void TempRenderLight(
 	FRHICommandList& RHICmdList,
 	const FScene* Scene,
 	const FViewInfo& View,
@@ -2388,17 +2361,10 @@ static void RenderLight(
 	FRenderLightParameters* PassParameters,	// If this is null, it means we cannot use Strata tiles and fallback to previous behavior.
 	FRHITexture* ScreenShadowMaskTexture,
 	FRHITexture* LightingChannelsTexture,
-	bool bRenderOverlap, bool bIssueDrawEvent)
+	bool bRenderOverlap)
 {
-	// Ensure the light is valid for this view
-	if (!LightSceneInfo->ShouldRenderLight(View))
-	{
-		return;
-	}
-
 	SCOPE_CYCLE_COUNTER(STAT_DirectLightRenderingTime);
 	INC_DWORD_STAT(STAT_NumLightsUsingStandardDeferred);
-	SCOPED_CONDITIONAL_DRAW_EVENT(RHICmdList, StandardDeferredLighting, bIssueDrawEvent);
 
 	const bool bStrataClassificationEnabled = Strata::IsStrataEnabled() && Strata::IsClassificationEnabled();
 	const bool bTilePassesReadingStrataEnabled = Strata::ShouldPassesReadingStrataBeTiled(Scene->GetFeatureLevel());
@@ -2410,11 +2376,11 @@ static void RenderLight(
 
 		{
 			SCOPED_DRAW_EVENT(RHICmdList, StrataSimpleMaterial);
-			InternalRenderLight(RHICmdList, Scene, View, LightSceneInfo, PassParameters, ScreenShadowMaskTexture, LightingChannelsTexture, bRenderOverlap, bIssueDrawEvent, bEnableStrataStencilTest, bEnableStrataTiledPass, EStrataTileMaterialType::ESimple);
+			InternalRenderLight(RHICmdList, Scene, View, LightSceneInfo, PassParameters, ScreenShadowMaskTexture, LightingChannelsTexture, bRenderOverlap, bEnableStrataStencilTest, bEnableStrataTiledPass, EStrataTileMaterialType::ESimple);
 		}
 		{
 			SCOPED_DRAW_EVENT(RHICmdList, StrataComplexMaterial);
-			InternalRenderLight(RHICmdList, Scene, View, LightSceneInfo, PassParameters, ScreenShadowMaskTexture, LightingChannelsTexture, bRenderOverlap, bIssueDrawEvent, bEnableStrataStencilTest, bEnableStrataTiledPass, EStrataTileMaterialType::EComplex);
+			InternalRenderLight(RHICmdList, Scene, View, LightSceneInfo, PassParameters, ScreenShadowMaskTexture, LightingChannelsTexture, bRenderOverlap, bEnableStrataStencilTest, bEnableStrataTiledPass, EStrataTileMaterialType::EComplex);
 		}
 	}
 	else
@@ -2422,10 +2388,10 @@ static void RenderLight(
 		const bool bEnableStrataTiledPass = false;
 		const bool bEnableStrataStencilTest = bStrataClassificationEnabled;
 
-		InternalRenderLight(RHICmdList, Scene, View, LightSceneInfo, PassParameters, ScreenShadowMaskTexture, LightingChannelsTexture, bRenderOverlap, bIssueDrawEvent, bEnableStrataStencilTest, bEnableStrataTiledPass, EStrataTileMaterialType::EComplex);
+		InternalRenderLight(RHICmdList, Scene, View, LightSceneInfo, PassParameters, ScreenShadowMaskTexture, LightingChannelsTexture, bRenderOverlap, bEnableStrataStencilTest, bEnableStrataTiledPass, EStrataTileMaterialType::EComplex);
 		if (Strata::IsStrataEnabled() && Strata::IsClassificationEnabled())
 		{
-			InternalRenderLight(RHICmdList, Scene, View, LightSceneInfo, PassParameters, ScreenShadowMaskTexture, LightingChannelsTexture, bRenderOverlap, bIssueDrawEvent, bEnableStrataStencilTest, bEnableStrataTiledPass, EStrataTileMaterialType::ESimple);
+			InternalRenderLight(RHICmdList, Scene, View, LightSceneInfo, PassParameters, ScreenShadowMaskTexture, LightingChannelsTexture, bRenderOverlap, bEnableStrataStencilTest, bEnableStrataTiledPass, EStrataTileMaterialType::ESimple);
 		}
 	}
 }
@@ -2433,43 +2399,47 @@ static void RenderLight(
 static void RenderLight(
 	FRDGBuilder& GraphBuilder,
 	const FScene* Scene,
-	const TArray<FViewInfo>& Views,
+	const FViewInfo& View,
 	const FMinimalSceneTextures& SceneTextures,
 	const FLightSceneInfo* LightSceneInfo,
 	FRDGTextureRef ScreenShadowMaskTexture,
 	FRDGTextureRef LightingChannelsTexture,
-	bool bRenderOverlap)
+	bool bRenderOverlap, 
+	bool bCloudShadow)
 {
+	// Ensure the light is valid for this view
+	if (!LightSceneInfo->ShouldRenderLight(View))
+	{
+		return;
+	}
+
 	ERDGPassFlags PassFlags = ERDGPassFlags::Raster;
 
 	const FVolumetricCloudRenderSceneInfo* CloudInfo = Scene->GetVolumetricCloudSceneInfo();
 
-	for (int32 ViewIndex = 0, ViewCount = Views.Num(); ViewIndex < ViewCount; ++ViewIndex)
-	{
-		RDG_EVENT_SCOPE_CONDITIONAL(GraphBuilder, ViewCount > 1, "View%d", ViewIndex);
-		const FViewInfo& View = Views[ViewIndex];
-
-		FRenderLightParameters* PassParameters = GraphBuilder.AllocParameters<FRenderLightParameters>();
+	FRenderLightParameters* PassParameters = GraphBuilder.AllocParameters<FRenderLightParameters>();
+	if (bCloudShadow)
 		GetRenderLightParameters(View, SceneTextures, View.HairStrandsViewData, ScreenShadowMaskTexture, LightingChannelsTexture, GetCloudShadowAOParameters(GraphBuilder, View, CloudInfo), *PassParameters);
+	else
+		GetRenderLightParameters(View, SceneTextures, View.HairStrandsViewData, ScreenShadowMaskTexture, LightingChannelsTexture, {}, *PassParameters);
 
-		GraphBuilder.AddPass(
-			RDG_EVENT_NAME("StandardDeferredLighting"),
+	GraphBuilder.AddPass(
+		RDG_EVENT_NAME("StandardDeferredLighting"),
+		PassParameters,
+		PassFlags,
+		[Scene, &View, LightSceneInfo, ScreenShadowMaskTexture, LightingChannelsTexture, bRenderOverlap, PassParameters](FRHICommandList& RHICmdList)
+	{
+		TempRenderLight(
+			RHICmdList,
+			Scene, 
+			View,
+			LightSceneInfo,
 			PassParameters,
-			PassFlags,
-			[Scene, &View, LightSceneInfo, ScreenShadowMaskTexture, LightingChannelsTexture, bRenderOverlap, PassParameters](FRHICommandList& RHICmdList)
-		{
-			RenderLight(
-				RHICmdList,
-				Scene, 
-				View,
-				LightSceneInfo,
-				PassParameters,
-				TryGetRHI(ScreenShadowMaskTexture),
-				TryGetRHI(LightingChannelsTexture),
-				bRenderOverlap,
-				false);
-		});
-	}
+			TryGetRHI(ScreenShadowMaskTexture),
+			TryGetRHI(LightingChannelsTexture),
+			bRenderOverlap);
+
+	});
 }
 
 BEGIN_SHADER_PARAMETER_STRUCT(FRenderLightForHairParameters, )
