Change DeferredLightOverlapPS to use Permutation domain instead of template.

diff --git a/Engine/Source/Runtime/Renderer/Private/LightRendering.cpp b/Engine/Source/Runtime/Renderer/Private/LightRendering.cpp
--- a/Engine/Source/Runtime/Renderer/Private/LightRendering.cpp
+++ b/Engine/Source/Runtime/Renderer/Private/LightRendering.cpp
@@ -382,9 +382,9 @@ struct FRenderLightParams
 };
 
 
-class TDeferredLightHairVS : public FGlobalShader
+class FDeferredLightHairVS : public FGlobalShader
 {
-	DECLARE_SHADER_TYPE(TDeferredLightHairVS, Global);
+	DECLARE_SHADER_TYPE(FDeferredLightHairVS, Global);
 public:
 
 	static bool ShouldCompilePermutation(const FGlobalShaderPermutationParameters& Parameters)
@@ -398,8 +398,8 @@ public:
 		OutEnvironment.SetDefine(TEXT("SHADER_HAIR"), 1);
 	}
 
-	TDeferredLightHairVS() {}
-	TDeferredLightHairVS(const ShaderMetaType::CompiledShaderInitializerType& Initializer) :
+	FDeferredLightHairVS() {}
+	FDeferredLightHairVS(const ShaderMetaType::CompiledShaderInitializerType& Initializer) :
 		FGlobalShader(Initializer)
 	{
 		HairStrandsParameters.Bind(Initializer.ParameterMap, FHairStrandsViewUniformParameters::StaticStructMetadata.GetShaderVariableName());
@@ -420,7 +420,7 @@ private:
 	LAYOUT_FIELD(FShaderUniformBufferParameter, HairStrandsParameters);
 };
 
-IMPLEMENT_SHADER_TYPE(, TDeferredLightHairVS, TEXT("/Engine/Private/DeferredLightVertexShaders.usf"), TEXT("HairVertexMain"), SF_Vertex);
+IMPLEMENT_SHADER_TYPE(, FDeferredLightHairVS, TEXT("/Engine/Private/DeferredLightVertexShaders.usf"), TEXT("HairVertexMain"), SF_Vertex);
 
 
 enum class ELightSourceShape
@@ -764,12 +764,15 @@ private:
 
 IMPLEMENT_GLOBAL_SHADER(FDeferredLightPS, "/Engine/Private/DeferredLightPixelShaders.usf", "DeferredLightPixelMain", SF_Pixel);
 
-
 /** Shader used to visualize stationary light overlap. */
-template<bool bRadialAttenuation>
-class TDeferredLightOverlapPS : public FGlobalShader
+class FDeferredLightOverlapPS : public FGlobalShader
 {
-	DECLARE_SHADER_TYPE(TDeferredLightOverlapPS,Global)
+	DECLARE_SHADER_TYPE(FDeferredLightOverlapPS,Global)
+
+	class FRadialAttenuation : SHADER_PERMUTATION_BOOL("RADIAL_ATTENUATION");
+
+	using FPermutationDomain = TShaderPermutationDomain<FRadialAttenuation>;
+
 public:
 
 	static bool ShouldCompilePermutation(const FGlobalShaderPermutationParameters& Parameters)
@@ -780,16 +783,15 @@ public:
 	static void ModifyCompilationEnvironment(const FGlobalShaderPermutationParameters& Parameters, FShaderCompilerEnvironment& OutEnvironment)
 	{
 		FGlobalShader::ModifyCompilationEnvironment(Parameters, OutEnvironment);
-		OutEnvironment.SetDefine(TEXT("RADIAL_ATTENUATION"), (uint32)bRadialAttenuation);
 	}
 
-	TDeferredLightOverlapPS(const ShaderMetaType::CompiledShaderInitializerType& Initializer)
+	FDeferredLightOverlapPS(const ShaderMetaType::CompiledShaderInitializerType& Initializer)
 	:	FGlobalShader(Initializer)
 	{
 		HasValidChannel.Bind(Initializer.ParameterMap, TEXT("HasValidChannel"));
 	}
 
-	TDeferredLightOverlapPS()
+	FDeferredLightOverlapPS()
 	{
 	}
 
@@ -806,8 +808,7 @@ private:
 	LAYOUT_FIELD(FShaderParameter, HasValidChannel);
 };
 
-IMPLEMENT_SHADER_TYPE(template<>, TDeferredLightOverlapPS<true>, TEXT("/Engine/Private/StationaryLightOverlapShaders.usf"), TEXT("OverlapRadialPixelMain"), SF_Pixel);
-IMPLEMENT_SHADER_TYPE(template<>, TDeferredLightOverlapPS<false>, TEXT("/Engine/Private/StationaryLightOverlapShaders.usf"), TEXT("OverlapDirectionalPixelMain"), SF_Pixel);
+IMPLEMENT_GLOBAL_SHADER(FDeferredLightOverlapPS, "/Engine/Private/StationaryLightOverlapShaders.usf", "OverlapRadialPixelMain", SF_Pixel);
 
 static void SplitSimpleLightsByView(TArrayView<const FViewInfo> Views, const FSimpleLightArray& SimpleLights, TArrayView<FSimpleLightArray> SimpleLightsByView)
 {
@@ -2197,7 +2198,9 @@ static void InternalRenderLight(
 
 		if (bRenderOverlap)
 		{
-			TShaderMapRef<TDeferredLightOverlapPS<false> > PixelShader(View.ShaderMap);
+			FDeferredLightOverlapPS::FPermutationDomain PermutationVector;
+			PermutationVector.Set<FDeferredLightOverlapPS::FRadialAttenuation>(false);
+			TShaderMapRef<FDeferredLightOverlapPS> PixelShader(View.ShaderMap, PermutationVector);
 			GraphicsPSOInit.BoundShaderState.VertexDeclarationRHI = GFilterVertexDeclaration.VertexDeclarationRHI;
 			GraphicsPSOInit.BoundShaderState.VertexShaderRHI = VertexShader.GetVertexShader();
 			GraphicsPSOInit.BoundShaderState.PixelShaderRHI = PixelShader.GetPixelShader();
@@ -2286,7 +2289,10 @@ static void InternalRenderLight(
 
 		if (bRenderOverlap)
 		{
-			TShaderMapRef<TDeferredLightOverlapPS<true> > PixelShader(View.ShaderMap);
+			FDeferredLightOverlapPS::FPermutationDomain PermutationVector;
+			PermutationVector.Set<FDeferredLightOverlapPS::FRadialAttenuation>(true);
+
+			TShaderMapRef<FDeferredLightOverlapPS> PixelShader(View.ShaderMap, PermutationVector);
 			GraphicsPSOInit.BoundShaderState.VertexDeclarationRHI = GetVertexDeclarationFVector4();
 			GraphicsPSOInit.BoundShaderState.VertexShaderRHI = VertexShader.GetVertexShader();
 			GraphicsPSOInit.BoundShaderState.PixelShaderRHI = PixelShader.GetPixelShader();
@@ -2570,7 +2576,7 @@ void FDeferredShadingSceneRenderer::RenderLightForHair(
 			PermutationVector.Set< FDeferredLightPS::FTransmissionDim >(false);
 			PermutationVector.Set< FDeferredLightPS::FHairLighting>(1);
 
-			TShaderMapRef<TDeferredLightHairVS> VertexShader(View.ShaderMap);
+			TShaderMapRef<FDeferredLightHairVS> VertexShader(View.ShaderMap);
 			TShaderMapRef<FDeferredLightPS> PixelShader(View.ShaderMap, PermutationVector);
 
 			GraphicsPSOInit.BoundShaderState.VertexDeclarationRHI = GetVertexDeclarationFVector4();
