Fix mac crash for subsurface scattering recombine fall back pass when engine flag of SSS is disabled.

diff --git a/Engine/Shaders/Private/PostProcessSubsurface.usf b/Engine/Shaders/Private/PostProcessSubsurface.usf
--- a/Engine/Shaders/Private/PostProcessSubsurface.usf
+++ b/Engine/Shaders/Private/PostProcessSubsurface.usf
@@ -904,9 +904,9 @@ void MainIndirectDispatchCS(
  
 // Recombines the half res Subsurface filtered lighting contribution (upsampled and renormalized with the alpha)
 // with the SceneColor.
-void SubsurfaceRecombinePS(noperspective float2 UV : TEXCOORD0, float4 SvPosition : SV_POSITION, out float4 OutColor : SV_Target0)
+void SubsurfaceRecombinePS(in FScreenVertexOutput Input, out float4 OutColor : SV_Target0)
 {
-	const float2 BufferUV  = UV.xy;
+	const float2 BufferUV  = Input.UV.xy;
 
 #if STRATA_ENABLED
 	const FStrataSubsurfaceData SSSData = LoadStataSSSData(BufferUV);
@@ -1010,9 +1010,9 @@ void SubsurfaceRecombinePS(noperspective float2 UV : TEXCOORD0, float4 SvPositio
 	OutColor = float4(SubsurfaceLighting * StoredBaseColor + ExtractedNonSubsurface, 1.0f);
 }
 
-void SubsurfaceRecombineCopyPS(noperspective float2 UV : TEXCOORD0, float4 SvPosition : SV_POSITION, out float4 OutColor : SV_Target0)
+void SubsurfaceRecombineCopyPS(in FScreenVertexOutput Input, out float4 OutColor : SV_Target0)
 {
-	const float2 BufferUV = UV.xy;
+	const float2 BufferUV = Input.UV.xy;
 
 #if STRATA_ENABLED
 	const FStrataSubsurfaceData SSSData = LoadStataSSSData(BufferUV);
diff --git a/Engine/Shaders/Private/PostProcessSubsurfaceTile.usf b/Engine/Shaders/Private/PostProcessSubsurfaceTile.usf
--- a/Engine/Shaders/Private/PostProcessSubsurfaceTile.usf
+++ b/Engine/Shaders/Private/PostProcessSubsurfaceTile.usf
@@ -97,4 +97,12 @@ void MainVS(
 	Out.Position = float4(Out.UV * float2(2.0f, -2.0f) + float2(-1.0, 1.0f), 0.0f, 1.0f);
 }
 
-#endif //SHADER_TILE_VS
\ No newline at end of file
+#endif //SHADER_TILE_VS
+
+void SubsurfaceTileFallbackScreenPassVS(
+	in float4 InPosition : ATTRIBUTE0,
+	in float2 InTexCoord : ATTRIBUTE1,
+	out FScreenVertexOutput Out)
+{
+	DrawRectangle(InPosition, InTexCoord, Out.Position, Out.UV);
+}
\ No newline at end of file
diff --git a/Engine/Source/Runtime/Renderer/Private/SubsurfaceTiles.cpp b/Engine/Source/Runtime/Renderer/Private/SubsurfaceTiles.cpp
--- a/Engine/Source/Runtime/Renderer/Private/SubsurfaceTiles.cpp
+++ b/Engine/Source/Runtime/Renderer/Private/SubsurfaceTiles.cpp
@@ -17,6 +17,11 @@ void FSubsurfaceTilePassVS::ModifyCompilationEnvironment(const FGlobalShaderPerm
 	OutEnvironment.SetDefine(TEXT("SUBSURFACE_TILE_SIZE"), FSubsurfaceTiles::TileSize);
 }
 
+bool FSubsurfaceTileFallbackScreenPassVS::ShouldCompilePermutation(const FGlobalShaderPermutationParameters& Parameters)
+{
+	return IsFeatureLevelSupported(Parameters.Platform, ERHIFeatureLevel::SM5);
+}
+
 FSubsurfaceTilePassVS::FParameters GetSubsurfaceTileParameters(const FViewInfo& InView,const FSubsurfaceTiles& InTile, FSubsurfaceTiles::ETileType TileType, FIntPoint BufferExtent)
 {
 	FSubsurfaceTilePassVS::FParameters Out;
@@ -132,6 +137,7 @@ void AddConditionalClearBlackUAVPass(FRDGBuilder& GraphBuilder, FRDGEventName&&
 }
 
 IMPLEMENT_GLOBAL_SHADER(FSubsurfaceTilePassVS, "/Engine/Private/PostProcessSubsurfaceTile.usf", "MainVS", SF_Vertex);
+IMPLEMENT_GLOBAL_SHADER(FSubsurfaceTileFallbackScreenPassVS, "/Engine/Private/PostProcessSubsurfaceTile.usf", "SubsurfaceTileFallbackScreenPassVS", SF_Vertex);
 
 const TCHAR* ToString(FSubsurfaceTiles::ETileType Type)
 {
diff --git a/Engine/Source/Runtime/Renderer/Private/SubsurfaceTiles.h b/Engine/Source/Runtime/Renderer/Private/SubsurfaceTiles.h
--- a/Engine/Source/Runtime/Renderer/Private/SubsurfaceTiles.h
+++ b/Engine/Source/Runtime/Renderer/Private/SubsurfaceTiles.h
@@ -76,6 +76,20 @@ class FSubsurfaceTilePassVS : public FGlobalShader
 	static void ModifyCompilationEnvironment(const FGlobalShaderPermutationParameters& Parameters, FShaderCompilerEnvironment& OutEnvironment);
 };
 
+class FSubsurfaceTileFallbackScreenPassVS : public FGlobalShader
+{
+public:
+	DECLARE_GLOBAL_SHADER(FSubsurfaceTileFallbackScreenPassVS);
+
+	static bool ShouldCompilePermutation(const FGlobalShaderPermutationParameters& Parameters);
+
+	FSubsurfaceTileFallbackScreenPassVS() = default;
+	FSubsurfaceTileFallbackScreenPassVS(const ShaderMetaType::CompiledShaderInitializerType& Initializer)
+		: FGlobalShader(Initializer)
+	{}
+};
+
+
 FSubsurfaceTilePassVS::FParameters GetSubsurfaceTileParameters(const FViewInfo& InView, const FSubsurfaceTiles& InTile, FSubsurfaceTiles::ETileType TileType, FIntPoint BufferExtent);
 
 template <typename FSubsurfacePassPS, typename FSubsurfacePassVS>
@@ -94,7 +108,7 @@ void AddSubsurfaceTiledScreenPass(
 {
 	if (bShouldFallbackToFullScreenPass)
 	{
-		TShaderMapRef<FScreenPassVS> FallBackVertexShader(View.ShaderMap);
+		TShaderMapRef<FSubsurfaceTileFallbackScreenPassVS> FallBackVertexShader(View.ShaderMap);
 		AddDrawScreenPass(GraphBuilder,
 			Forward<FRDGEventName>(Name),
 			View,
