Change FDeferredLightVS to use Permutation vector rather than template.

diff --git a/Engine/Shaders/Private/DeferredLightVertexShaders.usf b/Engine/Shaders/Private/DeferredLightVertexShaders.usf
--- a/Engine/Shaders/Private/DeferredLightVertexShaders.usf
+++ b/Engine/Shaders/Private/DeferredLightVertexShaders.usf
@@ -16,7 +16,7 @@ float3 StencilingPreViewTranslation;
 
 #if defined(SHADER_RADIAL_LIGHT) && SHADER_RADIAL_LIGHT == 0
 /** Vertex shader for rendering a directional light using a full screen quad. */
-void DirectionalVertexMain(
+void VertexMain(
 	in float2 InPosition : ATTRIBUTE0,
 	in float2 InUV       : ATTRIBUTE1,
 	out float2 OutTexCoord : TEXCOORD0,
@@ -29,9 +29,9 @@ void DirectionalVertexMain(
 }
 #endif
 
-#if SHADER_RADIAL_LIGHT == 1
+#if defined(SHADER_RADIAL_LIGHT) && SHADER_RADIAL_LIGHT == 1
 /** Vertex shader for rendering a point or spot light using approximate bounding geometry. */
-void RadialVertexMain(
+void VertexMain(
 	in uint InVertexId : SV_VertexID,
 	in float3 InPosition : ATTRIBUTE0,
 	out float4 OutScreenPosition : TEXCOORD0,
diff --git a/Engine/Source/Runtime/Renderer/Private/LightRendering.cpp b/Engine/Source/Runtime/Renderer/Private/LightRendering.cpp
--- a/Engine/Source/Runtime/Renderer/Private/LightRendering.cpp
+++ b/Engine/Source/Runtime/Renderer/Private/LightRendering.cpp
@@ -360,11 +360,8 @@ TGlobalResource<StencilingGeometry::FStencilConeVertexBuffer> StencilingGeometry
 /** The stencil cone index buffer. */
 TGlobalResource<StencilingGeometry::FStencilConeIndexBuffer> StencilingGeometry::GStencilConeIndexBuffer;
 
-
 // Implement a version for directional lights, and a version for point / spot lights
-IMPLEMENT_SHADER_TYPE(template<>,TDeferredLightVS<false>,TEXT("/Engine/Private/DeferredLightVertexShaders.usf"),TEXT("DirectionalVertexMain"),SF_Vertex);
-IMPLEMENT_SHADER_TYPE(template<>,TDeferredLightVS<true>,TEXT("/Engine/Private/DeferredLightVertexShaders.usf"),TEXT("RadialVertexMain"),SF_Vertex);
-
+IMPLEMENT_GLOBAL_SHADER(FDeferredLightVS, "/Engine/Private/DeferredLightVertexShaders.usf", "VertexMain", SF_Vertex);
 
 struct FRenderLightParams
 {
@@ -420,8 +417,7 @@ private:
 	LAYOUT_FIELD(FShaderUniformBufferParameter, HairStrandsParameters);
 };
 
-IMPLEMENT_SHADER_TYPE(, FDeferredLightHairVS, TEXT("/Engine/Private/DeferredLightVertexShaders.usf"), TEXT("HairVertexMain"), SF_Vertex);
-
+IMPLEMENT_GLOBAL_SHADER(FDeferredLightHairVS, "/Engine/Private/DeferredLightVertexShaders.usf", "HairVertexMain", SF_Vertex);
 
 enum class ELightSourceShape
 {
@@ -2185,7 +2181,9 @@ static void InternalRenderLight(
 	{
 		// Turn DBT back off
 		GraphicsPSOInit.bDepthBounds = false;
-		TShaderMapRef<TDeferredLightVS<false> > VertexShader(View.ShaderMap);
+		FDeferredLightVS::FPermutationDomain PermutationVectorVS;
+		PermutationVectorVS.Set<FDeferredLightVS::FRadialLight>(false);
+		TShaderMapRef<FDeferredLightVS> VertexShader(View.ShaderMap, PermutationVectorVS);
 
 		Strata::FStrataTilePassVS::FParameters VSParameters;
 		Strata::FStrataTilePassVS::FPermutationDomain VSPermutationVector;
@@ -2276,7 +2274,9 @@ static void InternalRenderLight(
 		// Disable depth bound when hair rendering is enabled as this rejects partially covered pixel write (with opaque background)
 		GraphicsPSOInit.bDepthBounds = GSupportsDepthBoundsTest && GAllowDepthBoundsTest != 0;
 
-		TShaderMapRef<TDeferredLightVS<true> > VertexShader(View.ShaderMap);
+		FDeferredLightVS::FPermutationDomain PermutationVectorVS;
+		PermutationVectorVS.Set<FDeferredLightVS::FRadialLight>(true);
+		TShaderMapRef<FDeferredLightVS> VertexShader(View.ShaderMap, PermutationVectorVS);
 
 		const bool bCameraInsideLightGeometry = ((FVector)View.ViewMatrices.GetViewOrigin() - LightBounds.Center).SizeSquared() < FMath::Square(LightBounds.W * 1.05f + View.NearClippingDistance * 2.0f)
 		//const bool bCameraInsideLightGeometry = LightProxy->AffectsBounds( FSphere( View.ViewMatrices.GetViewOrigin(), View.NearClippingDistance * 2.0f ) )
@@ -2704,7 +2704,9 @@ void FDeferredShadingSceneRenderer::RenderSimpleLightsStandardDeferred(
 				// Set the device viewport for the view.
 				RHICmdList.SetViewport(View.ViewRect.Min.X, View.ViewRect.Min.Y, 0.0f, View.ViewRect.Max.X, View.ViewRect.Max.Y, 1.0f);
 
-				TShaderMapRef<TDeferredLightVS<true> > VertexShader(View.ShaderMap);
+				FDeferredLightVS::FPermutationDomain PermutationVector;
+				PermutationVector.Set<FDeferredLightVS::FRadialLight>(true);
+				TShaderMapRef<FDeferredLightVS> VertexShader(View.ShaderMap, PermutationVector);
 
 				const bool bCameraInsideLightGeometry = ((FVector)View.ViewMatrices.GetViewOrigin() - LightBounds.Center).SizeSquared() < FMath::Square(LightBounds.W * 1.05f + View.NearClippingDistance * 2.0f)
 								// Always draw backfaces in ortho
diff --git a/Engine/Source/Runtime/Renderer/Private/LightRendering.h b/Engine/Source/Runtime/Renderer/Private/LightRendering.h
--- a/Engine/Source/Runtime/Renderer/Private/LightRendering.h
+++ b/Engine/Source/Runtime/Renderer/Private/LightRendering.h
@@ -435,15 +435,18 @@ private:
 
 
 /** A vertex shader for rendering the light in a deferred pass. */
-template<bool bRadialLight>
-class TDeferredLightVS : public FGlobalShader
+class FDeferredLightVS : public FGlobalShader
 {
-	DECLARE_SHADER_TYPE(TDeferredLightVS,Global);
-public:
+	DECLARE_SHADER_TYPE(FDeferredLightVS,Global);
+
+	class FRadialLight : SHADER_PERMUTATION_BOOL("SHADER_RADIAL_LIGHT");
+	using FPermutationDomain = TShaderPermutationDomain<FRadialLight>;
 
+public:
 	static bool ShouldCompilePermutation(const FGlobalShaderPermutationParameters& Parameters)
 	{
-		if (bRadialLight)
+		FPermutationDomain PermutationVector(Parameters.PermutationId);
+		if (PermutationVector.Get<FRadialLight>())
 		{
 			return IsFeatureLevelSupported(Parameters.Platform, ERHIFeatureLevel::SM5) || IsMobileDeferredShadingEnabled(Parameters.Platform);
 		}
@@ -454,11 +457,10 @@ public:
 	static void ModifyCompilationEnvironment(const FGlobalShaderPermutationParameters& Parameters, FShaderCompilerEnvironment& OutEnvironment)
 	{
 		FGlobalShader::ModifyCompilationEnvironment(Parameters, OutEnvironment);
-		OutEnvironment.SetDefine(TEXT("SHADER_RADIAL_LIGHT"), bRadialLight ? 1 : 0);
 	}
 
-	TDeferredLightVS()	{}
-	TDeferredLightVS(const ShaderMetaType::CompiledShaderInitializerType& Initializer):
+	FDeferredLightVS()	{}
+	FDeferredLightVS(const ShaderMetaType::CompiledShaderInitializerType& Initializer):
 		FGlobalShader(Initializer)
 	{
 		StencilingGeometryParameters.Bind(Initializer.ParameterMap);
diff --git a/Engine/Source/Runtime/Renderer/Private/MobileDeferredShadingPass.cpp b/Engine/Source/Runtime/Renderer/Private/MobileDeferredShadingPass.cpp
--- a/Engine/Source/Runtime/Renderer/Private/MobileDeferredShadingPass.cpp
+++ b/Engine/Source/Runtime/Renderer/Private/MobileDeferredShadingPass.cpp
@@ -363,8 +363,10 @@ static void RenderLocalLight_StencilMask(FRHICommandListImmediate& RHICmdList, c
 		true, CF_Always, SO_Keep, SO_Replace, SO_Keep,		
 		false, CF_Always, SO_Keep, SO_Keep, SO_Keep,
 		0x00, STENCIL_SANDBOX_MASK>::GetRHI();
-	   	
-	TShaderMapRef<TDeferredLightVS<true> > VertexShader(View.ShaderMap);
+
+	FDeferredLightVS::FPermutationDomain PermutationVector;
+	PermutationVector.Set<FDeferredLightVS::FRadialLight>(true);
+	TShaderMapRef<FDeferredLightVS> VertexShader(View.ShaderMap, PermutationVector);
 	GraphicsPSOInit.BoundShaderState.VertexDeclarationRHI = GetVertexDeclarationFVector4();
 	GraphicsPSOInit.BoundShaderState.VertexShaderRHI = VertexShader.GetVertexShader();
 	GraphicsPSOInit.BoundShaderState.PixelShaderRHI = nullptr;
@@ -430,7 +432,9 @@ static void RenderLocalLight(
 		SetLocalLightRasterizerAndDepthState(GraphicsPSOInit, View, LightBounds);
 	}
 
-	TShaderMapRef<TDeferredLightVS<true>> VertexShader(View.ShaderMap);
+	FDeferredLightVS::FPermutationDomain PermutationVectorVS;
+	PermutationVectorVS.Set<FDeferredLightVS::FRadialLight>(true);
+	TShaderMapRef<FDeferredLightVS> VertexShader(View.ShaderMap, PermutationVectorVS);
 		
 	const FMaterialRenderProxy* LightFunctionMaterialProxy = nullptr;
 	if (View.Family->EngineShowFlags.LightFunctions)
@@ -483,7 +487,7 @@ static void RenderLocalLight(
 static void SetupSimpleLightPSO(
 	FRHICommandListImmediate& RHICmdList, 
 	const FViewInfo& View,
-	const TShaderMapRef<TDeferredLightVS<true>>& VertexShader,
+	const TShaderMapRef<FDeferredLightVS>& VertexShader,
 	const TShaderRef<FMobileRadialLightFunctionPS>& PixelShader, 
 	FGraphicsPipelineStateInitializer& GraphicsPSOInit)
 {
@@ -508,7 +512,9 @@ static void RenderSimpleLights(
 {
 	const FSimpleLightArray& SimpleLights = SortedLightSet.SimpleLights;
 
-	TShaderMapRef<TDeferredLightVS<true>> VertexShader(View.ShaderMap);
+	FDeferredLightVS::FPermutationDomain PermutationVectorVS;
+	PermutationVectorVS.Set<FDeferredLightVS::FRadialLight>(true);
+	TShaderMapRef<FDeferredLightVS> VertexShader(View.ShaderMap, PermutationVectorVS);
 	TShaderRef<FMobileRadialLightFunctionPS> PixelShaders[2];
 	{
 		const FMaterialShaderMap* MaterialShaderMap = DefaultMaterial.Material->GetRenderingThreadShaderMap();
diff --git a/Engine/Source/Runtime/Renderer/Private/PlanarReflectionRendering.cpp b/Engine/Source/Runtime/Renderer/Private/PlanarReflectionRendering.cpp
--- a/Engine/Source/Runtime/Renderer/Private/PlanarReflectionRendering.cpp
+++ b/Engine/Source/Runtime/Renderer/Private/PlanarReflectionRendering.cpp
@@ -205,7 +205,9 @@ void PrefilterPlanarReflection(
 		PassParameters->SceneColorInputSampler = TStaticSamplerState<SF_Bilinear, AM_Clamp, AM_Clamp, AM_Clamp>::GetRHI();
 		PassParameters->RenderTargets[0] = FRenderTargetBinding(ViewFamilyTexture, RTLoadAction);
 
-		TShaderMapRef<TDeferredLightVS<false> > VertexShader(View.ShaderMap);
+		FDeferredLightVS::FPermutationDomain PermutationVector;
+		PermutationVector.Set<FDeferredLightVS::FRadialLight>(false);
+		TShaderMapRef<FDeferredLightVS> VertexShader(View.ShaderMap, PermutationVector);
 		TShaderMapRef<FPrefilterPlanarReflectionPS> PixelShader(View.ShaderMap);
 
 		const FIntPoint SceneColorExtent = SceneColorTexture->Desc.Extent;
@@ -837,7 +839,9 @@ void FDeferredShadingSceneRenderer::RenderDeferredPlanarReflections(FRDGBuilder&
 
 			SCOPED_DRAW_EVENTF(RHICmdList, PlanarReflection, *ReflectionSceneProxy->OwnerName.ToString());
 
-			TShaderMapRef<TDeferredLightVS<false> > VertexShader(View.ShaderMap);
+			FDeferredLightVS::FPermutationDomain PermutationVector;
+			PermutationVector.Set<FDeferredLightVS::FRadialLight>(false);
+			TShaderMapRef<FDeferredLightVS> VertexShader(View.ShaderMap, PermutationVector);
 			TShaderMapRef<FPlanarReflectionPS> PixelShader(View.ShaderMap);
 
 			GraphicsPSOInit.BoundShaderState.VertexDeclarationRHI = GFilterVertexDeclaration.VertexDeclarationRHI;
