Simplify StandardDeferredSimpleLights.

diff --git a/Engine/Source/Runtime/Renderer/Private/LightRendering.cpp b/Engine/Source/Runtime/Renderer/Private/LightRendering.cpp
--- a/Engine/Source/Runtime/Renderer/Private/LightRendering.cpp
+++ b/Engine/Source/Runtime/Renderer/Private/LightRendering.cpp
@@ -2022,36 +2022,6 @@ void SetBoundingGeometryRasterizerAndDepthState(FGraphicsPipelineStateInitialize
 	}
 }
 
-template<bool bUseIESProfile, bool bRadialAttenuation, bool bInverseSquaredFalloff>
-static void SetShaderTemplLightingSimple(
-	FRHICommandList& RHICmdList,
-	FGraphicsPipelineStateInitializer& GraphicsPSOInit,
-	const FViewInfo& View,
-	const TShaderRef<FShader>& VertexShader,
-	const FSimpleLightEntry& SimpleLight,
-	const FSimpleLightPerViewEntry& SimpleLightPerViewData)
-{
-	FDeferredLightPS::FPermutationDomain PermutationVector;
-	PermutationVector.Set< FDeferredLightPS::FSourceShapeDim >( ELightSourceShape::Capsule );
-	PermutationVector.Set< FDeferredLightPS::FIESProfileDim >( bUseIESProfile );
-	PermutationVector.Set< FDeferredLightPS::FInverseSquaredDim >( bInverseSquaredFalloff );
-	PermutationVector.Set< FDeferredLightPS::FVisualizeCullingDim >( View.Family->EngineShowFlags.VisualizeLightCulling );
-	PermutationVector.Set< FDeferredLightPS::FLightingChannelsDim >( false );
-	PermutationVector.Set< FDeferredLightPS::FAnistropicMaterials >(false);
-	PermutationVector.Set< FDeferredLightPS::FTransmissionDim >( false );
-	PermutationVector.Set< FDeferredLightPS::FHairLighting>( 0 );
-	PermutationVector.Set< FDeferredLightPS::FAtmosphereTransmittance >( false );
-	PermutationVector.Set< FDeferredLightPS::FCloudTransmittance >(false);
-	PermutationVector.Set< FDeferredLightPS::FStrataFastPath>(false);
-
-	TShaderMapRef< FDeferredLightPS > PixelShader( View.ShaderMap, PermutationVector );
-	GraphicsPSOInit.BoundShaderState.VertexDeclarationRHI = GetVertexDeclarationFVector4();
-	GraphicsPSOInit.BoundShaderState.VertexShaderRHI = VertexShader.GetVertexShader();
-	GraphicsPSOInit.BoundShaderState.PixelShaderRHI = PixelShader.GetPixelShader();
-	SetGraphicsPipelineState(RHICmdList, GraphicsPSOInit, 0);
-	PixelShader->SetParametersSimpleLight(RHICmdList, View, SimpleLight, SimpleLightPerViewData);
-}
-
 // Use DBT to allow work culling on shadow lights
 void CalculateLightNearFarDepthFromBounds(const FViewInfo& View, const FSphere &LightBounds, float &NearDepth, float &FarDepth)
 {
@@ -2704,9 +2674,9 @@ void FDeferredShadingSceneRenderer::RenderSimpleLightsStandardDeferred(
 				// Set the device viewport for the view.
 				RHICmdList.SetViewport(View.ViewRect.Min.X, View.ViewRect.Min.Y, 0.0f, View.ViewRect.Max.X, View.ViewRect.Max.Y, 1.0f);
 
-				FDeferredLightVS::FPermutationDomain PermutationVector;
-				PermutationVector.Set<FDeferredLightVS::FRadialLight>(true);
-				TShaderMapRef<FDeferredLightVS> VertexShader(View.ShaderMap, PermutationVector);
+				FDeferredLightVS::FPermutationDomain PermutationVectorVS;
+				PermutationVectorVS.Set<FDeferredLightVS::FRadialLight>(true);
+				TShaderMapRef<FDeferredLightVS> VertexShader(View.ShaderMap, PermutationVectorVS);
 
 				const bool bCameraInsideLightGeometry = ((FVector)View.ViewMatrices.GetViewOrigin() - LightBounds.Center).SizeSquared() < FMath::Square(LightBounds.W * 1.05f + View.NearClippingDistance * 2.0f)
 								// Always draw backfaces in ortho
@@ -2715,16 +2685,27 @@ void FDeferredShadingSceneRenderer::RenderSimpleLightsStandardDeferred(
 
 				SetBoundingGeometryRasterizerAndDepthState(GraphicsPSOInit, View, bCameraInsideLightGeometry);
 
-				if (SimpleLight.Exponent == 0)
-				{
-					// inverse squared
-					SetShaderTemplLightingSimple<false, true, true>(RHICmdList, GraphicsPSOInit, View, VertexShader, SimpleLight, SimpleLightPerViewData);
-				}
-				else
-				{
-					// light's exponent, not inverse squared
-					SetShaderTemplLightingSimple<false, true, false>(RHICmdList, GraphicsPSOInit, View, VertexShader, SimpleLight, SimpleLightPerViewData);
-				}
+				const bool bInverseSquaredFalloff = SimpleLight.Exponent == 0;
+
+				FDeferredLightPS::FPermutationDomain PermutationVector;
+				PermutationVector.Set< FDeferredLightPS::FSourceShapeDim >(ELightSourceShape::Capsule);
+				PermutationVector.Set< FDeferredLightPS::FIESProfileDim >(false);
+				PermutationVector.Set< FDeferredLightPS::FInverseSquaredDim >(bInverseSquaredFalloff);
+				PermutationVector.Set< FDeferredLightPS::FVisualizeCullingDim >(View.Family->EngineShowFlags.VisualizeLightCulling);
+				PermutationVector.Set< FDeferredLightPS::FLightingChannelsDim >(false);
+				PermutationVector.Set< FDeferredLightPS::FAnistropicMaterials >(false);
+				PermutationVector.Set< FDeferredLightPS::FTransmissionDim >(false);
+				PermutationVector.Set< FDeferredLightPS::FHairLighting>(0);
+				PermutationVector.Set< FDeferredLightPS::FAtmosphereTransmittance >(false);
+				PermutationVector.Set< FDeferredLightPS::FCloudTransmittance >(false);
+				PermutationVector.Set< FDeferredLightPS::FStrataFastPath>(false);
+
+				TShaderMapRef< FDeferredLightPS > PixelShader(View.ShaderMap, PermutationVector);
+				GraphicsPSOInit.BoundShaderState.VertexDeclarationRHI = GetVertexDeclarationFVector4();
+				GraphicsPSOInit.BoundShaderState.VertexShaderRHI = VertexShader.GetVertexShader();
+				GraphicsPSOInit.BoundShaderState.PixelShaderRHI = PixelShader.GetPixelShader();
+				SetGraphicsPipelineState(RHICmdList, GraphicsPSOInit, 0);
+				PixelShader->SetParametersSimpleLight(RHICmdList, View, SimpleLight, SimpleLightPerViewData);
 
 				VertexShader->SetSimpleLightParameters(RHICmdList, View, LightBounds);
 
