Fixes nan that causes motion blur tile classification to fails and make motion blur a lot more expensive

diff --git a/Engine/Shaders/Private/MotionBlur/MotionBlurCommon.ush b/Engine/Shaders/Private/MotionBlur/MotionBlurCommon.ush
--- a/Engine/Shaders/Private/MotionBlur/MotionBlurCommon.ush
+++ b/Engine/Shaders/Private/MotionBlur/MotionBlurCommon.ush
@@ -125,7 +125,7 @@ float4 MinMaxLengthPolar(float4 v0, float4 v1)
 float2 CartesianToPolar(float2 Velocity)
 {
 	float Length = length(Velocity);
-	float Angle = atan2Fast(Velocity.y, Velocity.x);
+	float Angle = Length > 0.0 ? atan2Fast(Velocity.y, Velocity.x) : 0.0;
 	return float2(Length, Angle);
 }
 
diff --git a/Engine/Source/Runtime/Renderer/Private/PostProcess/PostProcessMotionBlur.cpp b/Engine/Source/Runtime/Renderer/Private/PostProcess/PostProcessMotionBlur.cpp
--- a/Engine/Source/Runtime/Renderer/Private/PostProcess/PostProcessMotionBlur.cpp
+++ b/Engine/Source/Runtime/Renderer/Private/PostProcess/PostProcessMotionBlur.cpp
@@ -627,7 +627,7 @@ void AddMotionBlurVelocityPass(
 				FClearValueBinding::None,
 				GFastVRamConfig.VelocityFlat | TexCreate_ShaderResource | TexCreate_UAV);
 
-			VelocityFlatTexture = GraphBuilder.CreateTexture(Desc, TEXT("MotionBlur.VelocityTile"));
+			VelocityFlatTexture = GraphBuilder.CreateTexture(Desc, TEXT("MotionBlur.VelocityFlatten"));
 		}
 
 		VelocityTileTexturesSetup = CreateVecocityTileTextures(
