Fix clang

diff --git a/Engine/Source/Runtime/Renderer/Private/DynamicBVH.h b/Engine/Source/Runtime/Renderer/Private/DynamicBVH.h
--- a/Engine/Source/Runtime/Renderer/Private/DynamicBVH.h
+++ b/Engine/Source/Runtime/Renderer/Private/DynamicBVH.h
@@ -414,7 +414,7 @@ uint32 FDynamicBVH< MaxChildren, FCostMetric >::FindBestInsertion_Greedy( const
 					// Induced cost for children
 					ChildCost[j] = TotalCost[j] - CostMetric( NodeBounds );
 
-					FVector Delta = ( Bounds.Min + Bounds.Max ) - ( NodeBounds.Min + NodeBounds.Max );
+					FVector3f Delta = ( Bounds.Min + Bounds.Max ) - ( NodeBounds.Min + NodeBounds.Max );
 					Delta = Delta.GetAbs();
 					Dist[j] = Delta.X + Delta.Y + Delta.Z;
 					Dist[j] += ChildCost[j] * 4.0f;
