Change FDeferredLightVS to use FParameters.

diff --git a/Engine/Shaders/Private/DeferredLightVertexShaders.usf b/Engine/Shaders/Private/DeferredLightVertexShaders.usf
--- a/Engine/Shaders/Private/DeferredLightVertexShaders.usf
+++ b/Engine/Shaders/Private/DeferredLightVertexShaders.usf
@@ -5,7 +5,7 @@
 =============================================================================*/
 
 #include "HairStrands/HairStrandsVisibilityCommonStruct.ush" 
-#include "Common.ush"
+#include "Common.ush"  
 
 #if defined(SHADER_RADIAL_LIGHT)
 float4 StencilingGeometryPosAndScale;
@@ -15,6 +15,23 @@ float3 StencilingPreViewTranslation;
 #endif
 
 #if defined(SHADER_RADIAL_LIGHT) && SHADER_RADIAL_LIGHT == 0
+
+float4 PosScaleBias;
+float4 UVScaleBias;
+float4 InvTargetSizeAndTextureSize; 
+
+void DrawFullScreenRect(
+	in float4 InPosition,
+	in float2 InTexCoord,
+	out float4 OutPosition,
+	out float2 OutTexCoord)
+{
+	OutPosition = InPosition;
+	OutPosition.xy = -1.0f + 2.0f * (PosScaleBias.zw + (InPosition.xy * PosScaleBias.xy)) * InvTargetSizeAndTextureSize.xy;
+	OutPosition.xy *= float2(1, -1);
+	OutTexCoord.xy = (UVScaleBias.zw + (InTexCoord.xy * UVScaleBias.xy)) * InvTargetSizeAndTextureSize.zw;
+}
+
 /** Vertex shader for rendering a directional light using a full screen quad. */
 void VertexMain(
 	in float2 InPosition : ATTRIBUTE0,
@@ -24,7 +41,7 @@ void VertexMain(
 	out float4 OutPosition : SV_POSITION
 	)
 {	
-	DrawRectangle(float4(InPosition.xy, 0, 1), InUV, OutPosition, OutTexCoord);
+	DrawFullScreenRect(float4(InPosition.xy, 0, 1), InUV, OutPosition, OutTexCoord);
 	OutScreenVector = mul(float4(OutPosition.xy, 1, 0), View.ScreenToTranslatedWorld).xyz;
 }
 #endif
diff --git a/Engine/Source/Runtime/Renderer/Private/LightRendering.cpp b/Engine/Source/Runtime/Renderer/Private/LightRendering.cpp
--- a/Engine/Source/Runtime/Renderer/Private/LightRendering.cpp
+++ b/Engine/Source/Runtime/Renderer/Private/LightRendering.cpp
@@ -1068,6 +1068,7 @@ void FSceneRenderer::GatherAndSortLights(FSortedLightSetSceneInfo& OutSortedLigh
 
 /** Shader parameters to use when creating a RenderLight(...) pass. */
 BEGIN_SHADER_PARAMETER_STRUCT(FRenderLightParameters, )
+	SHADER_PARAMETER_STRUCT_REF(FViewUniformShaderParameters, View)
 	SHADER_PARAMETER_RDG_UNIFORM_BUFFER(FSceneTextureUniformParameters, SceneTextures)
 	SHADER_PARAMETER_RDG_UNIFORM_BUFFER(FHairStrandsViewUniformParameters, HairStrands)
 	SHADER_PARAMETER_RDG_UNIFORM_BUFFER(FStrataGlobalUniformParameters, Strata)
@@ -1093,6 +1094,7 @@ void GetRenderLightParameters(
 	const FVolumetricCloudShadowAOParameters& CloudShadowAOParameters,
 	FRenderLightParameters& Parameters)
 {
+	Parameters.View = View.ViewUniformBuffer;
 	Parameters.SceneTextures = SceneTexturesUniformBuffer;
 	Parameters.HairStrands = HairStrandsUniformBuffer;
 	Parameters.Strata = Strata::BindStrataGlobalUniformParameters(View.StrataSceneData);
@@ -2217,7 +2219,8 @@ static void InternalRenderLight(
 
 		if (!bEnableStrataTiledPass)
 		{
-			VertexShader->SetParameters(RHICmdList, View, LightSceneInfo);
+			FDeferredLightVS::FParameters VSParameters2 = FDeferredLightVS::GetParameters(View);
+			SetShaderParameters(RHICmdList, VertexShader, VertexShader.GetVertexShader(), VSParameters2);
 
 			// Apply the directional light as a full screen quad
 			DrawRectangle(
@@ -2294,7 +2297,8 @@ static void InternalRenderLight(
 			PixelShader->SetParameters(RHICmdList, View, LightSceneInfo, ScreenShadowMaskTexture, LightingChannelsTexture, &RenderLightParams, nullptr);
 		}
 
-		VertexShader->SetParameters(RHICmdList, View, LightSceneInfo);
+		FDeferredLightVS::FParameters VSParameters2 = FDeferredLightVS::GetParameters(View, LightSceneInfo);
+		SetShaderParameters(RHICmdList, VertexShader, VertexShader.GetVertexShader(), VSParameters2);
 
 		// Use DBT to allow work culling on shadow lights
 		if (GraphicsPSOInit.bDepthBounds)
@@ -2707,7 +2711,8 @@ void FDeferredShadingSceneRenderer::RenderSimpleLightsStandardDeferred(
 				SetGraphicsPipelineState(RHICmdList, GraphicsPSOInit, 0);
 				PixelShader->SetParametersSimpleLight(RHICmdList, View, SimpleLight, SimpleLightPerViewData);
 
-				VertexShader->SetSimpleLightParameters(RHICmdList, View, LightBounds);
+				FDeferredLightVS::FParameters ParametersVS = FDeferredLightVS::GetParameters(View, LightBounds);
+				SetShaderParameters(RHICmdList, VertexShader, VertexShader.GetVertexShader(), ParametersVS);
 
 				// Apply the point or spot light with some approximately bounding geometry,
 				// So we can get speedups from depth testing and not processing pixels outside of the light's influence.
diff --git a/Engine/Source/Runtime/Renderer/Private/LightRendering.h b/Engine/Source/Runtime/Renderer/Private/LightRendering.h
--- a/Engine/Source/Runtime/Renderer/Private/LightRendering.h
+++ b/Engine/Source/Runtime/Renderer/Private/LightRendering.h
@@ -373,7 +373,16 @@ namespace StencilingGeometry
 class FStencilingGeometryShaderParameters
 {
 	DECLARE_TYPE_LAYOUT(FStencilingGeometryShaderParameters, NonVirtual);
+
 public:
+
+	BEGIN_SHADER_PARAMETER_STRUCT(FParameters, )
+		SHADER_PARAMETER(FVector4f, StencilingGeometryPosAndScale)
+		SHADER_PARAMETER(FVector4f, StencilingConeParameters)
+		SHADER_PARAMETER(FMatrix44f, StencilingConeTransform)
+		SHADER_PARAMETER(FVector3f, StencilingPreViewTranslation)
+	END_SHADER_PARAMETER_STRUCT()
+
 	void Bind(const FShaderParameterMap& ParameterMap)
 	{
 		StencilGeometryPosAndScale.Bind(ParameterMap, TEXT("StencilingGeometryPosAndScale"));
@@ -384,34 +393,68 @@ public:
 
 	void Set(FRHICommandList& RHICmdList, FShader* Shader, const FVector4f& InStencilingGeometryPosAndScale) const
 	{
-		SetShaderValue(RHICmdList, RHICmdList.GetBoundVertexShader(), StencilGeometryPosAndScale, InStencilingGeometryPosAndScale);
-		SetShaderValue(RHICmdList, RHICmdList.GetBoundVertexShader(), StencilConeParameters, FVector4f(0.0f, 0.0f, 0.0f, 0.0f));
+		const FParameters P = GetParameters(InStencilingGeometryPosAndScale);
+		SetShaderValue(RHICmdList, RHICmdList.GetBoundVertexShader(), StencilGeometryPosAndScale, P.StencilingGeometryPosAndScale);
+		SetShaderValue(RHICmdList, RHICmdList.GetBoundVertexShader(), StencilConeParameters, P.StencilingConeParameters);
 	}
 
 	void Set(FRHICommandList& RHICmdList, FShader* Shader, const FSceneView& View, const FLightSceneInfo* LightSceneInfo) const
 	{
-		FVector4 GeometryPosAndScale;
-		if( LightSceneInfo->Proxy->GetLightType() == LightType_Point ||
-			LightSceneInfo->Proxy->GetLightType() == LightType_Rect )
+		const FParameters P = GetParameters(View, LightSceneInfo);
+		if (LightSceneInfo->Proxy->GetLightType() == LightType_Point ||
+			LightSceneInfo->Proxy->GetLightType() == LightType_Rect)
 		{
-			StencilingGeometry::GStencilSphereVertexBuffer.CalcTransform(GeometryPosAndScale, LightSceneInfo->Proxy->GetBoundingSphere(), View.ViewMatrices.GetPreViewTranslation());
-			SetShaderValue(RHICmdList, RHICmdList.GetBoundVertexShader(), StencilGeometryPosAndScale, FVector4f(GeometryPosAndScale));
-			SetShaderValue(RHICmdList, RHICmdList.GetBoundVertexShader(), StencilConeParameters, FVector4f(0.0f, 0.0f, 0.0f, 0.0f));
+			SetShaderValue(RHICmdList, RHICmdList.GetBoundVertexShader(), StencilGeometryPosAndScale, P.StencilingGeometryPosAndScale);
+			SetShaderValue(RHICmdList, RHICmdList.GetBoundVertexShader(), StencilConeParameters, P.StencilingConeParameters);
 		}
-		else if(LightSceneInfo->Proxy->GetLightType() == LightType_Spot)
+		else if (LightSceneInfo->Proxy->GetLightType() == LightType_Spot)
 		{
-			SetShaderValue(RHICmdList, RHICmdList.GetBoundVertexShader(), StencilConeTransform, (FMatrix44f)LightSceneInfo->Proxy->GetLightToWorld());
+			SetShaderValue(RHICmdList, RHICmdList.GetBoundVertexShader(), StencilConeTransform, P.StencilingConeTransform);
 			SetShaderValue(
-				RHICmdList, 
+				RHICmdList,
 				RHICmdList.GetBoundVertexShader(),
 				StencilConeParameters,
+				P.StencilingConeParameters);
+			SetShaderValue(RHICmdList, RHICmdList.GetBoundVertexShader(), StencilPreViewTranslation, P.StencilingPreViewTranslation);
+		}
+	}
+
+	static FParameters GetParameters(const FVector4f& InStencilingGeometryPosAndScale)
+	{
+		FParameters Out;
+		Out.StencilingGeometryPosAndScale = InStencilingGeometryPosAndScale;
+		Out.StencilingConeParameters = FVector4f(0.0f, 0.0f, 0.0f, 0.0f);
+		Out.StencilingConeTransform = FMatrix44f::Identity;
+		Out.StencilingPreViewTranslation = FVector3f(0, 0, 0);
+		return Out;
+	}
+
+	static FParameters GetParameters(const FSceneView& View, const FLightSceneInfo* LightSceneInfo)
+	{
+		FParameters Out;
+		if (LightSceneInfo->Proxy->GetLightType() == LightType_Point ||
+			LightSceneInfo->Proxy->GetLightType() == LightType_Rect)
+		{
+			FVector4 GeometryPosAndScale;
+			StencilingGeometry::GStencilSphereVertexBuffer.CalcTransform(GeometryPosAndScale, LightSceneInfo->Proxy->GetBoundingSphere(), View.ViewMatrices.GetPreViewTranslation());
+			Out.StencilingGeometryPosAndScale = FVector4f(GeometryPosAndScale);
+			Out.StencilingConeParameters = FVector4f(0.0f, 0.0f, 0.0f, 0.0f);
+			Out.StencilingConeTransform = FMatrix44f::Identity;
+			Out.StencilingPreViewTranslation = FVector3f(0, 0, 0);
+		}
+		else if (LightSceneInfo->Proxy->GetLightType() == LightType_Spot)
+		{
+			Out.StencilingGeometryPosAndScale = FVector4f(0.0f, 0.0f, 0.0f, 0.0f);
+			Out.StencilingConeTransform = (FMatrix44f)LightSceneInfo->Proxy->GetLightToWorld();
+			Out.StencilingConeParameters =
 				FVector4f(
 					StencilingGeometry::FStencilConeIndexBuffer::NumSides,
 					StencilingGeometry::FStencilConeIndexBuffer::NumSlices,
 					LightSceneInfo->Proxy->GetOuterConeAngle(),
-					LightSceneInfo->Proxy->GetRadius()));
-			SetShaderValue(RHICmdList, RHICmdList.GetBoundVertexShader(), StencilPreViewTranslation, (FVector3f)View.ViewMatrices.GetPreViewTranslation());
+					LightSceneInfo->Proxy->GetRadius());
+			Out.StencilingPreViewTranslation = (FVector3f)View.ViewMatrices.GetPreViewTranslation();
 		}
+		return Out;
 	}
 
 	/** Serializer. */ 
@@ -426,22 +469,34 @@ public:
 
 private:
 	
-		LAYOUT_FIELD(FShaderParameter, StencilGeometryPosAndScale)
-		LAYOUT_FIELD(FShaderParameter, StencilConeParameters)
-		LAYOUT_FIELD(FShaderParameter, StencilConeTransform)
-		LAYOUT_FIELD(FShaderParameter, StencilPreViewTranslation)
-	
+	LAYOUT_FIELD(FShaderParameter, StencilGeometryPosAndScale)
+	LAYOUT_FIELD(FShaderParameter, StencilConeParameters)
+	LAYOUT_FIELD(FShaderParameter, StencilConeTransform)
+	LAYOUT_FIELD(FShaderParameter, StencilPreViewTranslation)
 };
 
 
+BEGIN_SHADER_PARAMETER_STRUCT(FDrawFullScreenRectangleParameters, )
+	SHADER_PARAMETER(FVector4f, PosScaleBias)
+	SHADER_PARAMETER(FVector4f, UVScaleBias)
+	SHADER_PARAMETER(FVector4f, InvTargetSizeAndTextureSize)
+END_SHADER_PARAMETER_STRUCT()
+
 /** A vertex shader for rendering the light in a deferred pass. */
 class FDeferredLightVS : public FGlobalShader
 {
 	DECLARE_SHADER_TYPE(FDeferredLightVS,Global);
+	SHADER_USE_PARAMETER_STRUCT(FDeferredLightVS, FGlobalShader);
 
 	class FRadialLight : SHADER_PERMUTATION_BOOL("SHADER_RADIAL_LIGHT");
 	using FPermutationDomain = TShaderPermutationDomain<FRadialLight>;
 
+	BEGIN_SHADER_PARAMETER_STRUCT(FParameters, )
+		SHADER_PARAMETER_STRUCT_REF(FViewUniformShaderParameters, View)
+		SHADER_PARAMETER_STRUCT_INCLUDE(FDrawFullScreenRectangleParameters, FullScreenRect)
+		SHADER_PARAMETER_STRUCT_INCLUDE(FStencilingGeometryShaderParameters::FParameters, Geometry)
+	END_SHADER_PARAMETER_STRUCT()
+
 public:
 	static bool ShouldCompilePermutation(const FGlobalShaderPermutationParameters& Parameters)
 	{
@@ -459,31 +514,74 @@ public:
 		FGlobalShader::ModifyCompilationEnvironment(Parameters, OutEnvironment);
 	}
 
-	FDeferredLightVS()	{}
-	FDeferredLightVS(const ShaderMetaType::CompiledShaderInitializerType& Initializer):
-		FGlobalShader(Initializer)
+	static FDrawFullScreenRectangleParameters GetFullScreenRectParameters(
+		float X, float Y,
+		float SizeX, float SizeY,
+		float U, float V,
+		float SizeU, float SizeV,
+		FIntPoint InTargetSize,
+		FIntPoint InTextureSize)
 	{
-		StencilingGeometryParameters.Bind(Initializer.ParameterMap);
+		FDrawFullScreenRectangleParameters Out;
+		Out.PosScaleBias = FVector4f(SizeX, SizeY, X, Y);
+		Out.UVScaleBias  = FVector4f(SizeU, SizeV, U, V);
+		Out.InvTargetSizeAndTextureSize = FVector4f(
+			1.0f / InTargetSize.X, 1.0f / InTargetSize.Y,
+			1.0f / InTextureSize.X, 1.0f / InTextureSize.Y);
+		return Out;
 	}
 
-	void SetParameters(FRHICommandList& RHICmdList, const FViewInfo& View, const FLightSceneInfo* LightSceneInfo)
+	static FParameters GetParameters(const FViewInfo& View, 
+		float X, float Y,
+		float SizeX, float SizeY,
+		float U, float V,
+		float SizeU, float SizeV,
+		FIntPoint TargetSize,
+		FIntPoint TextureSize)
 	{
-		FGlobalShader::SetParameters<FViewUniformShaderParameters>(RHICmdList, RHICmdList.GetBoundVertexShader(), View.ViewUniformBuffer);
-		StencilingGeometryParameters.Set(RHICmdList, this, View, LightSceneInfo);
+		FParameters Out;
+		Out.View = View.ViewUniformBuffer;
+		Out.Geometry = FStencilingGeometryShaderParameters::GetParameters(FVector4(0,0,0,0));
+		Out.FullScreenRect = GetFullScreenRectParameters(X, Y, SizeX, SizeY, U, V, SizeU, SizeV, TargetSize, TextureSize);
+		return Out;
 	}
 
-	void SetSimpleLightParameters(FRHICommandList& RHICmdList, const FViewInfo& View, const FSphere& LightBounds)
+	static FParameters GetParameters(const FViewInfo& View)
 	{
-		FGlobalShader::SetParameters<FViewUniformShaderParameters>(RHICmdList, RHICmdList.GetBoundVertexShader(), View.ViewUniformBuffer);
+		FParameters Out;
+		Out.View = View.ViewUniformBuffer;
+		Out.Geometry = FStencilingGeometryShaderParameters::GetParameters(FVector4(0, 0, 0, 0));
+		Out.FullScreenRect = GetFullScreenRectParameters(
+			0, 0,
+			View.ViewRect.Width(), View.ViewRect.Height(),
+			View.ViewRect.Min.X, View.ViewRect.Min.Y,
+			View.ViewRect.Width(), View.ViewRect.Height(),
+			View.ViewRect.Size(),
+			GetSceneTextureExtent());
+
+		return Out;
+	}
 
+	static FParameters GetParameters(const FViewInfo& View, const FSphere& LightBounds)
+	{
 		FVector4 StencilingSpherePosAndScale;
 		StencilingGeometry::GStencilSphereVertexBuffer.CalcTransform(StencilingSpherePosAndScale, LightBounds, View.ViewMatrices.GetPreViewTranslation());
-		StencilingGeometryParameters.Set(RHICmdList, this, StencilingSpherePosAndScale);
-	}
 
-private:
+		FParameters Out;
+		Out.View = View.ViewUniformBuffer;
+		Out.Geometry = FStencilingGeometryShaderParameters::GetParameters(StencilingSpherePosAndScale);
+		Out.FullScreenRect = GetFullScreenRectParameters(0, 0, 0, 0, 0, 0, 0, 0, FIntPoint(1, 1), FIntPoint(1, 1));
+		return Out;
+	}
 
-	LAYOUT_FIELD(FStencilingGeometryShaderParameters, StencilingGeometryParameters);
+	static FParameters GetParameters(const FViewInfo& View, const FLightSceneInfo* LightSceneInfo)
+	{
+		FParameters Out;
+		Out.View = View.ViewUniformBuffer;
+		Out.Geometry = FStencilingGeometryShaderParameters::GetParameters(View, LightSceneInfo);
+		Out.FullScreenRect = GetFullScreenRectParameters(0, 0, 0, 0, 0, 0, 0, 0, FIntPoint(1, 1), FIntPoint(1, 1));
+		return Out;
+	}
 };
 
 void SetBoundingGeometryRasterizerAndDepthState(FGraphicsPipelineStateInitializer& GraphicsPSOInit, const FViewInfo& View, bool bCameraInsideLightGeometry);
diff --git a/Engine/Source/Runtime/Renderer/Private/MobileDeferredShadingPass.cpp b/Engine/Source/Runtime/Renderer/Private/MobileDeferredShadingPass.cpp
--- a/Engine/Source/Runtime/Renderer/Private/MobileDeferredShadingPass.cpp
+++ b/Engine/Source/Runtime/Renderer/Private/MobileDeferredShadingPass.cpp
@@ -372,7 +372,9 @@ static void RenderLocalLight_StencilMask(FRHICommandListImmediate& RHICmdList, c
 	GraphicsPSOInit.BoundShaderState.PixelShaderRHI = nullptr;
 
 	SetGraphicsPipelineState(RHICmdList, GraphicsPSOInit, 1);
-	VertexShader->SetParameters(RHICmdList, View, &LightSceneInfo);
+
+	FDeferredLightVS::FParameters ParametersVS = FDeferredLightVS::GetParameters(View, &LightSceneInfo);
+	SetShaderParameters(RHICmdList, VertexShader, VertexShader.GetVertexShader(), ParametersVS);
 
 	if (LightType == LightType_Point)
 	{
@@ -457,7 +459,8 @@ static void RenderLocalLight(
 	uint8 StencilRef = GET_STENCIL_MOBILE_SM_MASK(MSM_DefaultLit);
 	SetGraphicsPipelineState(RHICmdList, GraphicsPSOInit, StencilRef);
 
-	VertexShader->SetParameters(RHICmdList, View, &LightSceneInfo);
+	FDeferredLightVS::FParameters ParametersVS = FDeferredLightVS::GetParameters(View, &LightSceneInfo);
+	SetShaderParameters(RHICmdList, VertexShader, VertexShader.GetVertexShader(), ParametersVS);
 
 	FMobileRadialLightFunctionPS::FParameters PassParameters;
 	PassParameters.DeferredLightUniforms = TUniformBufferRef<FDeferredLightUniformStruct>::CreateUniformBufferImmediate(GetDeferredLightParameters(View, LightSceneInfo), EUniformBufferUsage::UniformBuffer_SingleFrame);
@@ -566,7 +569,10 @@ static void RenderSimpleLights(
 
 		// Render light mask
 		SetGraphicsPipelineState(RHICmdList, GraphicsPSOLightMask, 1);
-		VertexShader->SetSimpleLightParameters(RHICmdList, View, LightBounds);
+
+		FDeferredLightVS::FParameters ParametersVS = FDeferredLightVS::GetParameters(View, LightBounds);
+		SetShaderParameters(RHICmdList, VertexShader, VertexShader.GetVertexShader(), ParametersVS);
+
 		StencilingGeometry::DrawSphere(RHICmdList);
 
 		// Render light
@@ -590,7 +596,6 @@ static void RenderSimpleLights(
 			SetGraphicsPipelineState(RHICmdList, GraphicsPSOLight[0], StencilRef);
 			FMobileRadialLightFunctionPS::SetParameters(RHICmdList, PixelShaders[0], View, DefaultMaterial.MaterialProxy, *DefaultMaterial.Material, PassParameters);
 		}
-		VertexShader->SetSimpleLightParameters(RHICmdList, View, LightBounds);
 
 		// Apply the point or spot light with some approximately bounding geometry,
 		// So we can get speedups from depth testing and not processing pixels outside of the light's influence.
diff --git a/Engine/Source/Runtime/Renderer/Private/PlanarReflectionRendering.cpp b/Engine/Source/Runtime/Renderer/Private/PlanarReflectionRendering.cpp
--- a/Engine/Source/Runtime/Renderer/Private/PlanarReflectionRendering.cpp
+++ b/Engine/Source/Runtime/Renderer/Private/PlanarReflectionRendering.cpp
@@ -233,9 +233,6 @@ void PrefilterPlanarReflection(
 
 			SetGraphicsPipelineState(RHICmdList, GraphicsPSOInit, 0);
 
-			SetShaderParameters(RHICmdList, PixelShader, PixelShader.GetPixelShader(), *PassParameters);
-			VertexShader->SetSimpleLightParameters(RHICmdList, View, FSphere(0));
-
 			FIntPoint UV = View.ViewRect.Min;
 			FIntPoint UVSize = View.ViewRect.Size();
 
@@ -245,6 +242,16 @@ void PrefilterPlanarReflection(
 				UVSize.Y = -UVSize.Y;
 			}
 
+			FDeferredLightVS::FParameters ParametersVS = FDeferredLightVS::GetParameters(View, 
+				0, 0,
+				View.ViewRect.Width(), View.ViewRect.Height(),
+				UV.X, UV.Y,
+				UVSize.X, UVSize.Y,
+				View.ViewRect.Size(),
+				SceneColorExtent);
+			SetShaderParameters(RHICmdList, PixelShader, PixelShader.GetPixelShader(), *PassParameters);
+			SetShaderParameters(RHICmdList, VertexShader, VertexShader.GetVertexShader(), ParametersVS);
+
 			DrawRectangle(
 				RHICmdList,
 				0, 0,
@@ -851,7 +858,8 @@ void FDeferredShadingSceneRenderer::RenderDeferredPlanarReflections(FRDGBuilder&
 
 			SetGraphicsPipelineState(RHICmdList, GraphicsPSOInit, 0);
 
-			VertexShader->SetSimpleLightParameters(RHICmdList, View, FSphere(0));
+			FDeferredLightVS::FParameters ParametersVS = FDeferredLightVS::GetParameters(View);
+			SetShaderParameters(RHICmdList, VertexShader, VertexShader.GetVertexShader(), ParametersVS);
 
 			{
 				FPlanarReflectionUniformParameters PlanarReflectionUniformParameters;
