Fix invalidations due to things like WPO (that are fed back on the GPU)

diff --git a/Engine/Source/Runtime/Renderer/Private/VirtualShadowMaps/VirtualShadowMapCacheManager.cpp b/Engine/Source/Runtime/Renderer/Private/VirtualShadowMaps/VirtualShadowMapCacheManager.cpp
--- a/Engine/Source/Runtime/Renderer/Private/VirtualShadowMaps/VirtualShadowMapCacheManager.cpp
+++ b/Engine/Source/Runtime/Renderer/Private/VirtualShadowMaps/VirtualShadowMapCacheManager.cpp
@@ -512,13 +512,18 @@ bool FVirtualShadowMapArrayCacheManager::IsAccumulatingStats()
 
 void FVirtualShadowMapArrayCacheManager::ProcessRemovedOrUpdatedPrimitives(FRDGBuilder& GraphBuilder, const FGPUScene& GPUScene, FInvalidatingPrimitiveCollector& InvalidatingPrimitiveCollector)
 {
-	if (CVarCacheVirtualSMs.GetValueOnRenderThread() != 0 && !InvalidatingPrimitiveCollector.IsEmpty() && PrevBuffers.DynamicCasterPageFlags.IsValid())
+	if (CVarCacheVirtualSMs.GetValueOnRenderThread() != 0 && PrevBuffers.DynamicCasterPageFlags.IsValid())
 	{
 		RDG_EVENT_SCOPE(GraphBuilder, "Shadow.Virtual.ProcessRemovedOrUpdatedPrimitives");
+		ProcessGPUInstanceInvalidations(GraphBuilder, GPUScene);
+
+		if (!InvalidatingPrimitiveCollector.IsEmpty())
+		{
 #if VSM_LOG_INVALIDATIONS
-		UE_LOG(LogTemp, Warning, TEXT("ProcessRemovedOrUpdatedPrimitives: \n%s"), *InvalidatingPrimitiveCollector.RangesStr);
+			UE_LOG(LogTemp, Warning, TEXT("ProcessRemovedOrUpdatedPrimitives: \n%s"), *InvalidatingPrimitiveCollector.RangesStr);
 #endif
-		ProcessInvalidations(GraphBuilder, InvalidatingPrimitiveCollector.LoadBalancer, InvalidatingPrimitiveCollector.TotalInstanceCount, GPUScene);
+			ProcessInvalidations(GraphBuilder, InvalidatingPrimitiveCollector.LoadBalancer, InvalidatingPrimitiveCollector.TotalInstanceCount, GPUScene);
+		}
 	}
 }
 
@@ -721,6 +726,7 @@ void FVirtualShadowMapArrayCacheManager::ProcessInvalidations(FRDGBuilder& Graph
 void FVirtualShadowMapArrayCacheManager::ProcessGPUInstanceInvalidations(FRDGBuilder& GraphBuilder, const FGPUScene& GPUScene)
 {
 	// Dispatch CS indirectly to process instances that are marked to update from the GPU side.
+	if (PrevBuffers.InvalidatingInstancesBuffer.IsValid())
 	{
 		RDG_EVENT_SCOPE(GraphBuilder, "ProcessGPUInstanceInvalidations [GPU-Instances]");
 
@@ -747,5 +753,7 @@ void FVirtualShadowMapArrayCacheManager::ProcessGPUInstanceInvalidations(FRDGBui
 			0
 		);
 
+		// Drop the InvalidatingInstancesBuffer to make sure we don't redundantly process the associated invalidations if ProcessRemovedOrUpdatedPrimitives is called multiple times.
+		PrevBuffers.InvalidatingInstancesBuffer.SafeRelease();
 	}
 }
\ No newline at end of file
