Refactor remaining DF shaders to use RDG - Part 1
- Convert FAOParameters to SHADER_PARAMETER_STRUCT.
- Added DistanceField::SetupAOShaderParameters helper function to perform the necessary logic.
- Removed all legacy shader parameters from FComputeDistanceFieldNormalPS, FComputeDistanceFieldNormalCS, FBuildTileConesCS, FObjectCullVS, FObjectCullPS
- Use FComputeShaderUtils::AddPass for the compute shaders listed above.

diff --git a/Engine/Source/Runtime/Renderer/Private/DistanceFieldAmbientOcclusion.cpp b/Engine/Source/Runtime/Renderer/Private/DistanceFieldAmbientOcclusion.cpp
--- a/Engine/Source/Runtime/Renderer/Private/DistanceFieldAmbientOcclusion.cpp
+++ b/Engine/Source/Runtime/Renderer/Private/DistanceFieldAmbientOcclusion.cpp
@@ -21,7 +21,6 @@
 #include "Lumen/Lumen.h"
 #include "Strata/Strata.h"
 
-IMPLEMENT_TYPE_LAYOUT(FAOParameters);
 IMPLEMENT_TYPE_LAYOUT(FDFAOUpsampleParameters);
 IMPLEMENT_TYPE_LAYOUT(FScreenGridParameters);
 
@@ -243,6 +242,20 @@ float GAOConeHalfAngle = FMath::Acos(1 - 1.0f / (float)UE_ARRAY_COUNT(SpacedVect
 // Must match shader code
 uint32 GAONumConeSteps = 10;
 
+FAOParameters DistanceField::SetupAOShaderParameters(const FDistanceFieldAOParameters& Parameters)
+{
+	const float AOLargestSampleOffset = Parameters.ObjectMaxOcclusionDistance / (1 + FMath::Tan(GAOConeHalfAngle));
+
+	FAOParameters ShaderParameters;
+	ShaderParameters.AOObjectMaxDistance = Parameters.ObjectMaxOcclusionDistance;
+	ShaderParameters.AOStepScale = AOLargestSampleOffset / FMath::Pow(2.0f, GAOStepExponentScale * (GAONumConeSteps - 1));
+	ShaderParameters.AOStepExponentScale = GAOStepExponentScale;
+	ShaderParameters.AOMaxViewDistance = GetMaxAOViewDistance();
+	ShaderParameters.AOGlobalMaxOcclusionDistance = Parameters.GlobalMaxOcclusionDistance;
+
+	return ShaderParameters;
+}
+
 bool bListMemoryNextFrame = false;
 
 void OnListMemory(UWorld* InWorld)
@@ -273,11 +286,13 @@ class FComputeDistanceFieldNormalPS : public FGlobalShader
 {
 public:
 	DECLARE_GLOBAL_SHADER(FComputeDistanceFieldNormalPS);
+	SHADER_USE_PARAMETER_STRUCT(FComputeDistanceFieldNormalPS, FGlobalShader);
 
 	BEGIN_SHADER_PARAMETER_STRUCT(FParameters, )
 		SHADER_PARAMETER_STRUCT_REF(FViewUniformShaderParameters, View)
 		SHADER_PARAMETER_RDG_UNIFORM_BUFFER(FSceneTextureUniformParameters, SceneTextures)
 		SHADER_PARAMETER_RDG_UNIFORM_BUFFER(FStrataGlobalUniformParameters, Strata)
+		SHADER_PARAMETER_STRUCT_INCLUDE(FAOParameters, AOParameters)
 		RENDER_TARGET_BINDING_SLOTS()
 	END_SHADER_PARAMETER_STRUCT()
 
@@ -293,37 +308,21 @@ public:
 		OutEnvironment.SetDefine(TEXT("THREADGROUP_SIZEY"), GDistanceFieldAOTileSizeY);
 		OutEnvironment.SetDefine(TEXT("STRATA_ENABLED"), Strata::IsStrataEnabled() ? 1u : 0u);
 	}
-
-	FComputeDistanceFieldNormalPS() = default;
-	FComputeDistanceFieldNormalPS(const ShaderMetaType::CompiledShaderInitializerType& Initializer)
-		: FGlobalShader(Initializer)
-	{
-		BindForLegacyShaderParameters<FParameters>(this, Initializer.PermutationId, Initializer.ParameterMap, false);
-		AOParameters.Bind(Initializer.ParameterMap);
-	}
-
-	void SetParameters(FRHICommandList& RHICmdList, const FSceneView& View, const FDistanceFieldAOParameters& Parameters)
-	{
-		FRHIPixelShader* ShaderRHI = RHICmdList.GetBoundPixelShader();
-		AOParameters.Set(RHICmdList, ShaderRHI, Parameters);
-	}
-
-private:
-	LAYOUT_FIELD(FAOParameters, AOParameters);
 };
 
-IMPLEMENT_SHADER_TYPE(,FComputeDistanceFieldNormalPS,TEXT("/Engine/Private/DistanceFieldScreenGridLighting.usf"),TEXT("ComputeDistanceFieldNormalPS"),SF_Pixel);
-
+IMPLEMENT_GLOBAL_SHADER(FComputeDistanceFieldNormalPS, "/Engine/Private/DistanceFieldScreenGridLighting.usf", "ComputeDistanceFieldNormalPS", SF_Pixel);
 
 class FComputeDistanceFieldNormalCS : public FGlobalShader
 {
 public:
 	DECLARE_GLOBAL_SHADER(FComputeDistanceFieldNormalCS);
+	SHADER_USE_PARAMETER_STRUCT(FComputeDistanceFieldNormalCS, FGlobalShader);
 
 	BEGIN_SHADER_PARAMETER_STRUCT(FParameters, )
 		SHADER_PARAMETER_STRUCT_REF(FViewUniformShaderParameters, View)
 		SHADER_PARAMETER_RDG_UNIFORM_BUFFER(FSceneTextureUniformParameters, SceneTextures)
 		SHADER_PARAMETER_RDG_UNIFORM_BUFFER(FStrataGlobalUniformParameters, Strata)
+		SHADER_PARAMETER_STRUCT_INCLUDE(FAOParameters, AOParameters)
 		SHADER_PARAMETER_RDG_TEXTURE_UAV(RWTexture2D<float4>, RWDistanceFieldNormal)
 	END_SHADER_PARAMETER_STRUCT()
 
@@ -339,26 +338,9 @@ public:
 		OutEnvironment.SetDefine(TEXT("THREADGROUP_SIZEY"), GDistanceFieldAOTileSizeY);
 		OutEnvironment.SetDefine(TEXT("STRATA_ENABLED"), Strata::IsStrataEnabled() ? 1u : 0u);
 	}
-
-	FComputeDistanceFieldNormalCS() = default;
-	FComputeDistanceFieldNormalCS(const ShaderMetaType::CompiledShaderInitializerType& Initializer)
-		: FGlobalShader(Initializer)
-	{
-		BindForLegacyShaderParameters<FParameters>(this, Initializer.PermutationId, Initializer.ParameterMap, false);
-		AOParameters.Bind(Initializer.ParameterMap);
-	}
-
-	void SetParameters(FRHICommandList& RHICmdList, const FDistanceFieldAOParameters& Parameters)
-	{
-		FRHIComputeShader* ShaderRHI = RHICmdList.GetBoundComputeShader();
-		AOParameters.Set(RHICmdList, ShaderRHI, Parameters);
-	}
-
-private:
-	LAYOUT_FIELD(FAOParameters, AOParameters);
 };
 
-IMPLEMENT_SHADER_TYPE(,FComputeDistanceFieldNormalCS,TEXT("/Engine/Private/DistanceFieldScreenGridLighting.usf"),TEXT("ComputeDistanceFieldNormalCS"),SF_Compute);
+IMPLEMENT_GLOBAL_SHADER(FComputeDistanceFieldNormalCS, "/Engine/Private/DistanceFieldScreenGridLighting.usf", "ComputeDistanceFieldNormalCS", SF_Compute);
 
 void ComputeDistanceFieldNormal(
 	FRDGBuilder& GraphBuilder,
@@ -381,23 +363,12 @@ void ComputeDistanceFieldNormal(
 			PassParameters->View = View.ViewUniformBuffer;
 			PassParameters->SceneTextures = SceneTexturesUniformBuffer;
 			PassParameters->Strata = Strata::BindStrataGlobalUniformParameters(View.StrataSceneData);
+			PassParameters->AOParameters = DistanceField::SetupAOShaderParameters(Parameters);
 			PassParameters->RWDistanceFieldNormal = GraphBuilder.CreateUAV(DistanceFieldNormal);
 
 			TShaderMapRef<FComputeDistanceFieldNormalCS> ComputeShader(View.ShaderMap);
 
-			GraphBuilder.AddPass(
-				RDG_EVENT_NAME("ComputeNormalCS"),
-				PassParameters,
-				ERDGPassFlags::Compute,
-				[ComputeShader, PassParameters, Parameters, GroupSizeX, GroupSizeY](FRHICommandList& RHICmdList)
-			{
-				FRHIComputeShader* ShaderRHI = ComputeShader.GetComputeShader();
-				RHICmdList.SetComputeShader(ComputeShader.GetComputeShader());
-				ComputeShader->SetParameters(RHICmdList, Parameters);
-				SetShaderParameters(RHICmdList, ComputeShader, ShaderRHI, *PassParameters);
-				DispatchComputeShader(RHICmdList, ComputeShader.GetShader(), GroupSizeX, GroupSizeY, 1);
-				UnsetShaderUAVs(RHICmdList, ComputeShader, ShaderRHI);
-			});
+			FComputeShaderUtils::AddPass(GraphBuilder, RDG_EVENT_NAME("ComputeNormalCS"), ComputeShader, PassParameters, FIntVector(GroupSizeX, GroupSizeY, 1));
 		}
 	}
 	else
@@ -411,13 +382,14 @@ void ComputeDistanceFieldNormal(
 			PassParameters->View = View.ViewUniformBuffer;
 			PassParameters->SceneTextures = SceneTexturesUniformBuffer;
 			PassParameters->Strata = Strata::BindStrataGlobalUniformParameters(View.StrataSceneData);
+			PassParameters->AOParameters = DistanceField::SetupAOShaderParameters(Parameters);
 			PassParameters->RenderTargets[0] = FRenderTargetBinding(DistanceFieldNormal, ViewIndex == 0 ? ERenderTargetLoadAction::EClear : ERenderTargetLoadAction::ELoad);
 
 			GraphBuilder.AddPass(
 				RDG_EVENT_NAME("ComputeNormal"),
 				PassParameters,
 				ERDGPassFlags::Raster,
-				[&View, PassParameters, Parameters](FRHICommandList& RHICmdList)
+				[&View, PassParameters](FRHICommandList& RHICmdList)
 			{
 				RHICmdList.SetViewport(0, 0, 0.0f, View.ViewRect.Width() / GAODownsampleFactor, View.ViewRect.Height() / GAODownsampleFactor, 1.0f);
 
@@ -438,7 +410,6 @@ void ComputeDistanceFieldNormal(
 				SetGraphicsPipelineState(RHICmdList, GraphicsPSOInit, 0);
 
 				SetShaderParameters(RHICmdList, PixelShader, PixelShader.GetPixelShader(), *PassParameters);
-				PixelShader->SetParameters(RHICmdList, View, Parameters);
 
 				DrawRectangle(
 					RHICmdList,
diff --git a/Engine/Source/Runtime/Renderer/Private/DistanceFieldAmbientOcclusion.h b/Engine/Source/Runtime/Renderer/Private/DistanceFieldAmbientOcclusion.h
--- a/Engine/Source/Runtime/Renderer/Private/DistanceFieldAmbientOcclusion.h
+++ b/Engine/Source/Runtime/Renderer/Private/DistanceFieldAmbientOcclusion.h
@@ -90,58 +90,17 @@ inline float GetMaxAOViewDistance()
 	return FMath::Min(GAOMaxViewDistance, 65000.0f);
 }
 
-class FAOParameters
-{
-	DECLARE_TYPE_LAYOUT(FAOParameters, NonVirtual);
-public:
-	void Bind(const FShaderParameterMap& ParameterMap)
-	{
-		AOObjectMaxDistance.Bind(ParameterMap,TEXT("AOObjectMaxDistance"));
-		AOStepScale.Bind(ParameterMap,TEXT("AOStepScale"));
-		AOStepExponentScale.Bind(ParameterMap,TEXT("AOStepExponentScale"));
-		AOMaxViewDistance.Bind(ParameterMap,TEXT("AOMaxViewDistance"));
-		AOGlobalMaxOcclusionDistance.Bind(ParameterMap,TEXT("AOGlobalMaxOcclusionDistance"));
-	}
-
-	friend FArchive& operator<<(FArchive& Ar,FAOParameters& Parameters)
-	{
-		Ar << Parameters.AOObjectMaxDistance;
-		Ar << Parameters.AOStepScale;
-		Ar << Parameters.AOStepExponentScale;
-		Ar << Parameters.AOMaxViewDistance;
-		Ar << Parameters.AOGlobalMaxOcclusionDistance;
-		return Ar;
-	}
-
-	template<typename ShaderRHIParamRef>
-	void Set(FRHICommandList& RHICmdList, const ShaderRHIParamRef ShaderRHI, const FDistanceFieldAOParameters& Parameters)
-	{
-		SetShaderValue(RHICmdList, ShaderRHI, AOObjectMaxDistance, Parameters.ObjectMaxOcclusionDistance);
-
-		extern float GAOConeHalfAngle;
-		const float AOLargestSampleOffset = Parameters.ObjectMaxOcclusionDistance / (1 + FMath::Tan(GAOConeHalfAngle));
-
-		extern float GAOStepExponentScale;
-		extern uint32 GAONumConeSteps;
-		float AOStepScaleValue = AOLargestSampleOffset / FMath::Pow(2.0f, GAOStepExponentScale * (GAONumConeSteps - 1));
-		SetShaderValue(RHICmdList, ShaderRHI, AOStepScale, AOStepScaleValue);
-
-		SetShaderValue(RHICmdList, ShaderRHI, AOStepExponentScale, GAOStepExponentScale);
-
-		SetShaderValue(RHICmdList, ShaderRHI, AOMaxViewDistance, GetMaxAOViewDistance());
-
-		const float GlobalMaxOcclusionDistance = Parameters.GlobalMaxOcclusionDistance;
-		SetShaderValue(RHICmdList, ShaderRHI, AOGlobalMaxOcclusionDistance, GlobalMaxOcclusionDistance);
-	}
+BEGIN_SHADER_PARAMETER_STRUCT(FAOParameters, )
+	SHADER_PARAMETER(float, AOObjectMaxDistance)
+	SHADER_PARAMETER(float, AOStepScale)
+	SHADER_PARAMETER(float, AOStepExponentScale)
+	SHADER_PARAMETER(float, AOMaxViewDistance)
+	SHADER_PARAMETER(float, AOGlobalMaxOcclusionDistance)
+END_SHADER_PARAMETER_STRUCT()
 
-private:
-	
-		LAYOUT_FIELD(FShaderParameter, AOObjectMaxDistance)
-		LAYOUT_FIELD(FShaderParameter, AOStepScale)
-		LAYOUT_FIELD(FShaderParameter, AOStepExponentScale)
-		LAYOUT_FIELD(FShaderParameter, AOMaxViewDistance)
-		LAYOUT_FIELD(FShaderParameter, AOGlobalMaxOcclusionDistance)
-	
+namespace DistanceField
+{
+	FAOParameters SetupAOShaderParameters(const FDistanceFieldAOParameters& AOParameters);
 };
 
 class FDFAOUpsampleParameters
diff --git a/Engine/Source/Runtime/Renderer/Private/DistanceFieldLightingPost.cpp b/Engine/Source/Runtime/Renderer/Private/DistanceFieldLightingPost.cpp
--- a/Engine/Source/Runtime/Renderer/Private/DistanceFieldLightingPost.cpp
+++ b/Engine/Source/Runtime/Renderer/Private/DistanceFieldLightingPost.cpp
@@ -173,6 +173,7 @@ public:
 	// TODO(RDG) Hook these up to the shader. They are just informing RDG of transitions ATM.
 	BEGIN_SHADER_PARAMETER_STRUCT(FParameters, )
 		SHADER_PARAMETER_RDG_UNIFORM_BUFFER(FSceneTextureUniformParameters, SceneTextures)
+		SHADER_PARAMETER_STRUCT_INCLUDE(FAOParameters, AOParameters)
 		RDG_TEXTURE_ACCESS(DistanceFieldNormal, ERHIAccess::SRVGraphics)
 		RDG_TEXTURE_ACCESS(DistanceFieldAOBentNormal, ERHIAccess::SRVGraphics)
 		RDG_TEXTURE_ACCESS(BentNormalHistoryTexture, ERHIAccess::SRVGraphics)
@@ -197,7 +198,6 @@ public:
 	FUpdateHistoryDepthRejectionPS(const ShaderMetaType::CompiledShaderInitializerType& Initializer)
 		: FGlobalShader(Initializer)
 	{
-		AOParameters.Bind(Initializer.ParameterMap);
 		GeometryAwareUpsampleParameters.Bind(Initializer.ParameterMap);
 		BentNormalHistoryTexture.Bind(Initializer.ParameterMap, TEXT("BentNormalHistoryTexture"));
 		BentNormalHistorySampler.Bind(Initializer.ParameterMap, TEXT("BentNormalHistorySampler"));
@@ -217,13 +217,11 @@ public:
 		FRHITexture* DistanceFieldNormal, 
 		FRHITexture* DistanceFieldAOBentNormal,
 		FRHITexture* BentNormalHistoryTextureValue, 
-		FRHITexture* VelocityTextureValue,
-		const FDistanceFieldAOParameters& Parameters)
+		FRHITexture* VelocityTextureValue)
 	{
 		FRHIPixelShader* ShaderRHI = RHICmdList.GetBoundPixelShader();
 
 		FGlobalShader::SetParameters<FViewUniformShaderParameters>(RHICmdList, ShaderRHI, View.ViewUniformBuffer);
-		AOParameters.Set(RHICmdList, ShaderRHI, Parameters);
 		GeometryAwareUpsampleParameters.Set(RHICmdList, ShaderRHI, View, DistanceFieldNormal, DistanceFieldAOBentNormal);
 
 		SetTextureParameter(
@@ -273,7 +271,6 @@ public:
 	}
 
 private:
-	LAYOUT_FIELD(FAOParameters, AOParameters);
 	LAYOUT_FIELD(FGeometryAwareUpsampleParameters, GeometryAwareUpsampleParameters);
 	LAYOUT_FIELD(FShaderResourceParameter, BentNormalHistoryTexture);
 	LAYOUT_FIELD(FShaderResourceParameter, BentNormalHistorySampler);
@@ -400,6 +397,7 @@ class FGeometryAwareUpsamplePS : public FGlobalShader
 public:
 
 	BEGIN_SHADER_PARAMETER_STRUCT(FParameters, )
+		SHADER_PARAMETER_STRUCT_INCLUDE(FAOParameters, AOParameters)
 		RDG_TEXTURE_ACCESS(DistanceFieldNormal, ERHIAccess::SRVGraphics)
 		RDG_TEXTURE_ACCESS(BentNormalInterpolation, ERHIAccess::SRVGraphics)
 		RENDER_TARGET_BINDING_SLOTS()
@@ -422,7 +420,6 @@ public:
 	FGeometryAwareUpsamplePS(const ShaderMetaType::CompiledShaderInitializerType& Initializer)
 		: FGlobalShader(Initializer)
 	{
-		AOParameters.Bind(Initializer.ParameterMap);
 		GeometryAwareUpsampleParameters.Bind(Initializer.ParameterMap);
 	}
 
@@ -430,19 +427,16 @@ public:
 		FRHICommandList& RHICmdList,
 		const FViewInfo& View,
 		FRHITexture* DistanceFieldNormal,
-		FRHITexture* DistanceFieldAOBentNormal,
-		const FDistanceFieldAOParameters& Parameters)
+		FRHITexture* DistanceFieldAOBentNormal)
 	{
 		FRHIPixelShader* ShaderRHI = RHICmdList.GetBoundPixelShader();
 
 		FGlobalShader::SetParameters<FViewUniformShaderParameters>(RHICmdList, ShaderRHI, View.ViewUniformBuffer);
 
-		AOParameters.Set(RHICmdList, ShaderRHI, Parameters);
 		GeometryAwareUpsampleParameters.Set(RHICmdList, ShaderRHI, View, DistanceFieldNormal, DistanceFieldAOBentNormal);
 	}
 
 private:
-	LAYOUT_FIELD(FAOParameters, AOParameters);
 	LAYOUT_FIELD(FGeometryAwareUpsampleParameters, GeometryAwareUpsampleParameters);
 };
 
@@ -462,6 +456,7 @@ void AllocateOrReuseAORenderTarget(FRDGBuilder& GraphBuilder, FRDGTextureRef& Te
 void GeometryAwareUpsample(FRDGBuilder& GraphBuilder, const FViewInfo& View, FRDGTextureRef DistanceFieldAOBentNormal, FRDGTextureRef DistanceFieldNormal, FRDGTextureRef BentNormalInterpolation, const FDistanceFieldAOParameters& Parameters)
 {
 	auto* PassParameters = GraphBuilder.AllocParameters<FGeometryAwareUpsamplePS::FParameters>();
+	PassParameters->AOParameters = DistanceField::SetupAOShaderParameters(Parameters);
 	PassParameters->DistanceFieldNormal = DistanceFieldNormal;
 	PassParameters->BentNormalInterpolation = BentNormalInterpolation;
 	PassParameters->RenderTargets[0] = FRenderTargetBinding(DistanceFieldAOBentNormal, ERenderTargetLoadAction::ELoad);
@@ -470,7 +465,7 @@ void GeometryAwareUpsample(FRDGBuilder& GraphBuilder, const FViewInfo& View, FRD
 		RDG_EVENT_NAME("GeometryAwareUpsample"),
 		PassParameters,
 		ERDGPassFlags::Raster,
-		[&View, DistanceFieldNormal, BentNormalInterpolation, Parameters](FRHICommandList& RHICmdList)
+		[&View, DistanceFieldNormal, BentNormalInterpolation](FRHICommandList& RHICmdList)
 	{
 		RHICmdList.SetViewport(0, 0, 0.0f, View.ViewRect.Width() / GAODownsampleFactor, View.ViewRect.Height() / GAODownsampleFactor, 1.0f);
 
@@ -488,7 +483,7 @@ void GeometryAwareUpsample(FRDGBuilder& GraphBuilder, const FViewInfo& View, FRD
 		TShaderMapRef<FGeometryAwareUpsamplePS> PixelShader(View.ShaderMap);
 		GraphicsPSOInit.BoundShaderState.PixelShaderRHI = PixelShader.GetPixelShader();
 		SetGraphicsPipelineState(RHICmdList, GraphicsPSOInit, 0);
-		PixelShader->SetParameters(RHICmdList, View, DistanceFieldNormal->GetRHI(), BentNormalInterpolation->GetRHI(), Parameters);
+		PixelShader->SetParameters(RHICmdList, View, DistanceFieldNormal->GetRHI(), BentNormalInterpolation->GetRHI());
 
 		DrawRectangle(
 			RHICmdList,
@@ -551,6 +546,7 @@ void UpdateHistory(
 
 				auto* PassParameters = GraphBuilder.AllocParameters<FUpdateHistoryDepthRejectionPS::FParameters>();
 				PassParameters->SceneTextures = SceneTexturesUniformBuffer;
+				PassParameters->AOParameters = DistanceField::SetupAOShaderParameters(Parameters);
 				PassParameters->DistanceFieldNormal = DistanceFieldNormal;
 				PassParameters->DistanceFieldAOBentNormal = BentNormalInterpolation;
 				PassParameters->BentNormalHistoryTexture = BentNormalHistoryTexture;
@@ -564,7 +560,7 @@ void UpdateHistory(
 					RDG_EVENT_NAME("UpdateHistory"),
 					PassParameters,
 					ERDGPassFlags::Raster,
-					[VertexShader, PixelShader, &View, DistanceFieldNormal, BentNormalInterpolation, BentNormalHistoryTexture, VelocityTexture, PrevHistoryViewRect, Parameters, SceneTextureExtent]
+					[VertexShader, PixelShader, &View, DistanceFieldNormal, BentNormalInterpolation, BentNormalHistoryTexture, VelocityTexture, PrevHistoryViewRect, SceneTextureExtent]
 					(FRHICommandList& RHICmdList)
 				{
 					RHICmdList.SetViewport(0, 0, 0.0f, View.ViewRect.Width() / GAODownsampleFactor, View.ViewRect.Height() / GAODownsampleFactor, 1.0f);
@@ -585,8 +581,7 @@ void UpdateHistory(
 						DistanceFieldNormal->GetRHI(),
 						BentNormalInterpolation->GetRHI(),
 						BentNormalHistoryTexture->GetRHI(),
-						VelocityTexture ? VelocityTexture->GetRHI() : nullptr,
-						Parameters);
+						VelocityTexture ? VelocityTexture->GetRHI() : nullptr);
 
 					DrawRectangle(
 						RHICmdList,
diff --git a/Engine/Source/Runtime/Renderer/Private/DistanceFieldObjectCulling.cpp b/Engine/Source/Runtime/Renderer/Private/DistanceFieldObjectCulling.cpp
--- a/Engine/Source/Runtime/Renderer/Private/DistanceFieldObjectCulling.cpp
+++ b/Engine/Source/Runtime/Renderer/Private/DistanceFieldObjectCulling.cpp
@@ -106,15 +106,20 @@ void CullObjectsToView(FRDGBuilder& GraphBuilder, FScene* Scene, const FViewInfo
 /**  */
 class FBuildTileConesCS : public FGlobalShader
 {
-	DECLARE_GLOBAL_SHADER(FBuildTileConesCS);
 public:
+	DECLARE_GLOBAL_SHADER(FBuildTileConesCS);
+	SHADER_USE_PARAMETER_STRUCT(FBuildTileConesCS, FGlobalShader);
 
 	BEGIN_SHADER_PARAMETER_STRUCT(FParameters, )
+		SHADER_PARAMETER_STRUCT_REF(FViewUniformShaderParameters, View)
 		SHADER_PARAMETER_RDG_BUFFER_UAV(RWStructuredBuffer<FVector4f>, RWTileConeAxisAndCos)
 		SHADER_PARAMETER_RDG_BUFFER_UAV(RWStructuredBuffer<FVector4f>, RWTileConeDepthRanges)
 		SHADER_PARAMETER_RDG_UNIFORM_BUFFER(FSceneTextureUniformParameters, SceneTextures)
 		SHADER_PARAMETER_RDG_TEXTURE(Texture2D, DistanceFieldNormalTexture)
 		SHADER_PARAMETER_SAMPLER(SamplerState, DistanceFieldNormalSampler)
+		SHADER_PARAMETER_STRUCT_INCLUDE(FAOParameters, AOParameters)
+		SHADER_PARAMETER(FUintVector4, ViewDimensions)
+		SHADER_PARAMETER(FVector2f, NumGroups)
 	END_SHADER_PARAMETER_STRUCT()
 
 	static bool ShouldCompilePermutation(const FGlobalShaderPermutationParameters& Parameters)
@@ -132,35 +137,6 @@ public:
 		// To reduce shader compile time of compute shaders with shared memory, doesn't have an impact on generated code with current compiler (June 2010 DX SDK)
 		OutEnvironment.CompilerFlags.Add(CFLAG_StandardOptimization);
 	}
-
-	FBuildTileConesCS(const ShaderMetaType::CompiledShaderInitializerType& Initializer)
-		: FGlobalShader(Initializer)
-	{
-		BindForLegacyShaderParameters<FParameters>(this, Initializer.PermutationId, Initializer.ParameterMap, false);
-		AOParameters.Bind(Initializer.ParameterMap);
-		NumGroups.Bind(Initializer.ParameterMap, TEXT("NumGroups"));
-		ViewDimensionsParameter.Bind(Initializer.ParameterMap, TEXT("ViewDimensions"));
-	}
-
-	FBuildTileConesCS()
-	{
-	}
-	void SetParameters(FRHICommandList& RHICmdList, const FViewInfo& View, FScene* Scene, FVector2D NumGroupsValue, const FDistanceFieldAOParameters& Parameters)
-	{
-		FRHIComputeShader* ShaderRHI = RHICmdList.GetBoundComputeShader();
-
-		FGlobalShader::SetParameters<FViewUniformShaderParameters>(RHICmdList, ShaderRHI, View.ViewUniformBuffer);
-		AOParameters.Set(RHICmdList, ShaderRHI, Parameters);
-
-		SetShaderValue(RHICmdList, ShaderRHI, ViewDimensionsParameter, View.ViewRect);
-
-		SetShaderValue(RHICmdList, ShaderRHI, NumGroups, FVector2f(NumGroupsValue));
-	}
-
-private:
-	LAYOUT_FIELD(FAOParameters, AOParameters);
-	LAYOUT_FIELD(FShaderParameter, ViewDimensionsParameter);
-	LAYOUT_FIELD(FShaderParameter, NumGroups);
 };
 
 IMPLEMENT_GLOBAL_SHADER(FBuildTileConesCS, "/Engine/Private/DistanceFieldObjectCulling.usf", "BuildTileConesMain", SF_Compute);
@@ -168,61 +144,39 @@ IMPLEMENT_GLOBAL_SHADER(FBuildTileConesCS, "/Engine/Private/DistanceFieldObjectC
 /**  */
 class FObjectCullVS : public FGlobalShader
 {
-	DECLARE_GLOBAL_SHADER(FObjectCullVS);
 public:
+	DECLARE_GLOBAL_SHADER(FObjectCullVS);
+	SHADER_USE_PARAMETER_STRUCT(FObjectCullVS, FGlobalShader);
 
 	BEGIN_SHADER_PARAMETER_STRUCT(FParameters, )
+		SHADER_PARAMETER_STRUCT_REF(FViewUniformShaderParameters, View)
 		SHADER_PARAMETER_STRUCT_INCLUDE(FDistanceFieldCulledObjectBufferParameters, DistanceFieldCulledObjectBuffers)
 		SHADER_PARAMETER_STRUCT_INCLUDE(FDistanceFieldAtlasParameters, DistanceFieldAtlas)
+		SHADER_PARAMETER_STRUCT_INCLUDE(FAOParameters, AOParameters)
+		SHADER_PARAMETER(float, ConservativeRadiusScale)
 	END_SHADER_PARAMETER_STRUCT()
 
 	static bool ShouldCompilePermutation(const FGlobalShaderPermutationParameters& Parameters) 
 	{ 
 		return ShouldCompileDistanceFieldShaders(Parameters.Platform);
 	}
-
-	FObjectCullVS(const ShaderMetaType::CompiledShaderInitializerType& Initializer):
-		FGlobalShader(Initializer)
-	{
-		BindForLegacyShaderParameters<FParameters>(this, Initializer.PermutationId, Initializer.ParameterMap, false);
-		AOParameters.Bind(Initializer.ParameterMap);
-		ConservativeRadiusScale.Bind(Initializer.ParameterMap, TEXT("ConservativeRadiusScale"));
-	}
-
-	FObjectCullVS() {}
-
-	void SetParameters(FRHICommandList& RHICmdList, const FSceneView& View, const FDistanceFieldAOParameters& Parameters)
-	{
-		FRHIVertexShader* ShaderRHI = RHICmdList.GetBoundVertexShader();
-		FGlobalShader::SetParameters<FViewUniformShaderParameters>(RHICmdList, ShaderRHI, View.ViewUniformBuffer);
-		
-		AOParameters.Set(RHICmdList, ShaderRHI, Parameters);
-
-		const int32 NumRings = StencilingGeometry::GLowPolyStencilSphereVertexBuffer.GetNumRings();
-		const float RadiansPerRingSegment = PI / (float)NumRings;
-
-		// Boost the effective radius so that the edges of the sphere approximation lie on the sphere, instead of the vertices
-		const float ConservativeRadiusScaleValue = 1.0f / FMath::Cos(RadiansPerRingSegment);
-
-		SetShaderValue(RHICmdList, ShaderRHI, ConservativeRadiusScale, ConservativeRadiusScaleValue);
-	}
-
-private:
-	LAYOUT_FIELD(FAOParameters, AOParameters);
-	LAYOUT_FIELD(FShaderParameter, ConservativeRadiusScale);
 };
 
 IMPLEMENT_GLOBAL_SHADER(FObjectCullVS, "/Engine/Private/DistanceFieldObjectCulling.usf", "ObjectCullVS", SF_Vertex);
 
 class FObjectCullPS : public FGlobalShader
 {
-	DECLARE_GLOBAL_SHADER(FObjectCullPS);
 public:
+	DECLARE_GLOBAL_SHADER(FObjectCullPS);
+	SHADER_USE_PARAMETER_STRUCT(FObjectCullPS, FGlobalShader);
 
 	BEGIN_SHADER_PARAMETER_STRUCT(FParameters, )
+		SHADER_PARAMETER_STRUCT_REF(FViewUniformShaderParameters, View)
 		SHADER_PARAMETER_STRUCT_INCLUDE(FTileIntersectionParameters, TileIntersectionParameters)
 		SHADER_PARAMETER_STRUCT_INCLUDE(FDistanceFieldCulledObjectBufferParameters, DistanceFieldCulledObjectBuffers)
 		SHADER_PARAMETER_STRUCT_INCLUDE(FDistanceFieldAtlasParameters, DistanceFieldAtlas)
+		SHADER_PARAMETER_STRUCT_INCLUDE(FAOParameters, AOParameters)
+		SHADER_PARAMETER(FVector2f, NumGroups)
 	END_SHADER_PARAMETER_STRUCT()
 
 	class FCountingPass : SHADER_PERMUTATION_BOOL("SCATTER_CULLING_COUNT_PASS");
@@ -240,32 +194,6 @@ public:
 
 		OutEnvironment.SetDefine(TEXT("DOWNSAMPLE_FACTOR"), GAODownsampleFactor);
 	}
-
-	/** Default constructor. */
-	FObjectCullPS() {}
-
-	/** Initialization constructor. */
-	FObjectCullPS(const ShaderMetaType::CompiledShaderInitializerType& Initializer)
-		: FGlobalShader(Initializer)
-	{
-		BindForLegacyShaderParameters<FParameters>(this, Initializer.PermutationId, Initializer.ParameterMap, false);
-		AOParameters.Bind(Initializer.ParameterMap);
-		NumGroups.Bind(Initializer.ParameterMap, TEXT("NumGroups"));
-	}
-
-	void SetParameters(FRHICommandList& RHICmdList, const FSceneView& View, FVector2D NumGroupsValue, const FDistanceFieldAOParameters& Parameters)
-	{
-		FRHIPixelShader* ShaderRHI = RHICmdList.GetBoundPixelShader();
-		FGlobalShader::SetParameters<FViewUniformShaderParameters>(RHICmdList, ShaderRHI, View.ViewUniformBuffer);
-
-		AOParameters.Set(RHICmdList, ShaderRHI, Parameters);
-
-		SetShaderValue(RHICmdList, ShaderRHI, NumGroups, FVector2f(NumGroupsValue));
-	}
-
-private:
-	LAYOUT_FIELD(FAOParameters, AOParameters);
-	LAYOUT_FIELD(FShaderParameter, NumGroups);
 };
 
 IMPLEMENT_GLOBAL_SHADER(FObjectCullPS, "/Engine/Private/DistanceFieldObjectCulling.usf", "ObjectCullPS", SF_Pixel);
@@ -331,11 +259,26 @@ void ScatterTilesToObjects(
 	auto PixelShader = View.ShaderMap->GetShader<FObjectCullPS>(PermutationVector);
 
 	auto* PassParameters = GraphBuilder.AllocParameters<FObjectCullParameters>();
+	PassParameters->VS.View = GetShaderBinding(View.ViewUniformBuffer);
 	PassParameters->VS.DistanceFieldCulledObjectBuffers = CulledObjectBufferParameters;
 	PassParameters->VS.DistanceFieldAtlas = DistanceField::SetupAtlasParameters(DistanceFieldSceneData);
+	PassParameters->VS.AOParameters = DistanceField::SetupAOShaderParameters(Parameters);
+
+	{
+		const int32 NumRings = StencilingGeometry::GLowPolyStencilSphereVertexBuffer.GetNumRings();
+		const float RadiansPerRingSegment = PI / (float)NumRings;
+
+		// Boost the effective radius so that the edges of the sphere approximation lie on the sphere, instead of the vertices
+		PassParameters->VS.ConservativeRadiusScale = 1.0f / FMath::Cos(RadiansPerRingSegment);
+	}
+
+	PassParameters->PS.View = GetShaderBinding(View.ViewUniformBuffer);
 	PassParameters->PS.TileIntersectionParameters = TileIntersectionParameters;
 	PassParameters->PS.DistanceFieldCulledObjectBuffers = CulledObjectBufferParameters;
 	PassParameters->PS.DistanceFieldAtlas = DistanceField::SetupAtlasParameters(DistanceFieldSceneData);
+	PassParameters->PS.AOParameters = DistanceField::SetupAOShaderParameters(Parameters);
+	PassParameters->PS.NumGroups = FVector2D(TileListGroupSize.X, TileListGroupSize.Y);
+
 	PassParameters->SceneTextures = SceneTexturesUniformBuffer;
 	PassParameters->ObjectIndirectArguments = ObjectIndirectArguments;
 
@@ -352,7 +295,7 @@ void ScatterTilesToObjects(
 		bCountingPass ? RDG_EVENT_NAME("CountTileObjectIntersections") : RDG_EVENT_NAME("CullTilesToObjects"),
 		PassParameters,
 		ERDGPassFlags::Raster,
-		[PassParameters, VertexShader, PixelShader, &View, TileListGroupSize, Parameters, ObjectIndirectArguments](FRHICommandListImmediate& RHICmdList)
+		[PassParameters, VertexShader, PixelShader, &View, TileListGroupSize, ObjectIndirectArguments](FRHICommandListImmediate& RHICmdList)
 		{
 			FGraphicsPipelineStateInitializer GraphicsPSOInit;
 			RHICmdList.ApplyCachedRenderTargets(GraphicsPSOInit);
@@ -371,9 +314,6 @@ void ScatterTilesToObjects(
 
 			SetGraphicsPipelineState(RHICmdList, GraphicsPSOInit, 0);
 
-			VertexShader->SetParameters(RHICmdList, View, Parameters);
-			PixelShader->SetParameters(RHICmdList, View, FVector2D(TileListGroupSize.X, TileListGroupSize.Y), Parameters);
-
 			SetShaderParameters(RHICmdList, VertexShader, VertexShader.GetVertexShader(), PassParameters->VS);
 			SetShaderParameters(RHICmdList, PixelShader, PixelShader.GetPixelShader(), PassParameters->PS);
 
@@ -418,31 +358,18 @@ void BuildTileObjectLists(
 
 		{
 			auto* PassParameters = GraphBuilder.AllocParameters<FBuildTileConesCS::FParameters>();
+			PassParameters->View = View.ViewUniformBuffer;
 			PassParameters->RWTileConeAxisAndCos = TileIntersectionParameters.RWTileConeAxisAndCos;
 			PassParameters->RWTileConeDepthRanges = TileIntersectionParameters.RWTileConeDepthRanges;
 			PassParameters->SceneTextures = SceneTexturesUniformBuffer;
 			PassParameters->DistanceFieldNormalTexture = DistanceFieldNormal;
 			PassParameters->DistanceFieldNormalSampler = TStaticSamplerState<SF_Point, AM_Clamp, AM_Clamp, AM_Clamp>::GetRHI();
+			PassParameters->AOParameters = DistanceField::SetupAOShaderParameters(Parameters);
+			PassParameters->NumGroups = FVector2f(TileListGroupSize.X, TileListGroupSize.Y);
 
 			auto ComputeShader = View.ShaderMap->GetShader<FBuildTileConesCS>();
 
-			ClearUnusedGraphResources(ComputeShader, PassParameters);
-
-			GraphBuilder.AddPass(
-				RDG_EVENT_NAME("BuildTileCones"),
-				PassParameters,
-				ERDGPassFlags::Compute,
-				[PassParameters, ComputeShader, Scene, &View, TileListGroupSize, Parameters](FRHICommandListImmediate& RHICmdList)
-				{
-					RHICmdList.SetComputeShader(ComputeShader.GetComputeShader());
-
-					ComputeShader->SetParameters(RHICmdList, View, Scene, FVector2f(TileListGroupSize.X, TileListGroupSize.Y), Parameters);
-					SetShaderParameters(RHICmdList, ComputeShader, ComputeShader.GetComputeShader(), *PassParameters);
-
-					DispatchComputeShader(RHICmdList, ComputeShader.GetShader(), TileListGroupSize.X, TileListGroupSize.Y, 1);
-
-					UnsetShaderUAVs(RHICmdList, ComputeShader, ComputeShader.GetComputeShader());
-				});
+			FComputeShaderUtils::AddPass(GraphBuilder, RDG_EVENT_NAME("BuildTileCones"), ComputeShader, PassParameters, FIntVector(TileListGroupSize.X, TileListGroupSize.Y, 1));
 		}
 
 		// Start at 0 tiles per object
diff --git a/Engine/Source/Runtime/Renderer/Private/DistanceFieldScreenGridLighting.cpp b/Engine/Source/Runtime/Renderer/Private/DistanceFieldScreenGridLighting.cpp
--- a/Engine/Source/Runtime/Renderer/Private/DistanceFieldScreenGridLighting.cpp
+++ b/Engine/Source/Runtime/Renderer/Private/DistanceFieldScreenGridLighting.cpp
@@ -82,6 +82,7 @@ public:
 		SHADER_PARAMETER_STRUCT_INCLUDE(FDistanceFieldAtlasParameters, DistanceFieldAtlas)
 		SHADER_PARAMETER_STRUCT_INCLUDE(FTileIntersectionParameters, TileIntersectionParameters)
 		SHADER_PARAMETER_STRUCT_INCLUDE(FAOScreenGridParameters, AOScreenGridParameters)
+		SHADER_PARAMETER_STRUCT_INCLUDE(FAOParameters, AOParameters)
 		RDG_BUFFER_ACCESS(ObjectTilesIndirectArguments, ERHIAccess::IndirectArgs)
 		SHADER_PARAMETER_RDG_UNIFORM_BUFFER(FSceneTextureUniformParameters, SceneTextures)
 		RDG_TEXTURE_ACCESS(DistanceFieldNormal, ERHIAccess::SRVCompute)
@@ -109,7 +110,6 @@ public:
 		: FGlobalShader(Initializer)
 	{
 		BindForLegacyShaderParameters<FParameters>(this, Initializer.PermutationId, Initializer.ParameterMap, false);
-		AOParameters.Bind(Initializer.ParameterMap);
 		ScreenGridParameters.Bind(Initializer.ParameterMap);
 		GlobalDistanceFieldParameters.Bind(Initializer.ParameterMap);
 		TanConeHalfAngle.Bind(Initializer.ParameterMap, TEXT("TanConeHalfAngle"));
@@ -123,15 +123,13 @@ public:
 	void SetParameters(
 		FRHICommandList& RHICmdList, 
 		const FViewInfo& View, 
-		FRHITexture* DistanceFieldNormal,  
-		const FDistanceFieldAOParameters& Parameters,
+		FRHITexture* DistanceFieldNormal,
 		bool bUseGlobalDistanceField,
 		const FGlobalDistanceFieldInfo& GlobalDistanceFieldInfo)
 	{
 		FRHIComputeShader* ShaderRHI = RHICmdList.GetBoundComputeShader();
 		FGlobalShader::SetParameters<FViewUniformShaderParameters>(RHICmdList, ShaderRHI, View.ViewUniformBuffer);
 
-		AOParameters.Set(RHICmdList, ShaderRHI, Parameters);
 		ScreenGridParameters.Set(RHICmdList, ShaderRHI, View, DistanceFieldNormal);
 
 		if (bUseGlobalDistanceField)
@@ -170,7 +168,6 @@ public:
 	}
 
 private:
-	LAYOUT_FIELD(FAOParameters, AOParameters);
 	LAYOUT_FIELD(FScreenGridParameters, ScreenGridParameters);
 	LAYOUT_FIELD(FGlobalDistanceFieldParameters, GlobalDistanceFieldParameters);
 	LAYOUT_FIELD(FShaderParameter, TanConeHalfAngle);
@@ -190,6 +187,7 @@ public:
 		SHADER_PARAMETER_STRUCT_INCLUDE(FDistanceFieldCulledObjectBufferParameters, DistanceFieldCulledObjectBuffers)
 		SHADER_PARAMETER_STRUCT_INCLUDE(FDistanceFieldAtlasParameters, DistanceFieldAtlas)
 		SHADER_PARAMETER_STRUCT_INCLUDE(FAOScreenGridParameters, AOScreenGridParameters)
+		SHADER_PARAMETER_STRUCT_INCLUDE(FAOParameters, AOParameters)
 		SHADER_PARAMETER_RDG_BUFFER_SRV(Buffer<FVector4f>, TileConeDepthRanges)
 		SHADER_PARAMETER_RDG_UNIFORM_BUFFER(FSceneTextureUniformParameters, SceneTextures)
 		RDG_TEXTURE_ACCESS(DistanceFieldNormal, ERHIAccess::SRVCompute)
@@ -219,7 +217,6 @@ public:
 		: FGlobalShader(Initializer)
 	{
 		BindForLegacyShaderParameters<FParameters>(this, Initializer.PermutationId, Initializer.ParameterMap, false);
-		AOParameters.Bind(Initializer.ParameterMap);
 		ScreenGridParameters.Bind(Initializer.ParameterMap);
 		GlobalDistanceFieldParameters.Bind(Initializer.ParameterMap);
 		TileListGroupSize.Bind(Initializer.ParameterMap, TEXT("TileListGroupSize"));
@@ -235,14 +232,12 @@ public:
 		FRHICommandList& RHICmdList, 
 		const FViewInfo& View,
 		FIntPoint TileListGroupSizeValue, 
-		FRHITexture* DistanceFieldNormal, 
-		const FDistanceFieldAOParameters& Parameters,
+		FRHITexture* DistanceFieldNormal,
 		const FGlobalDistanceFieldInfo& GlobalDistanceFieldInfo)
 	{
 		FRHIComputeShader* ShaderRHI = RHICmdList.GetBoundComputeShader();
 		FGlobalShader::SetParameters<FViewUniformShaderParameters>(RHICmdList, ShaderRHI, View.ViewUniformBuffer);
 
-		AOParameters.Set(RHICmdList, ShaderRHI, Parameters);
 		ScreenGridParameters.Set(RHICmdList, ShaderRHI, View, DistanceFieldNormal);
 		GlobalDistanceFieldParameters.Set(RHICmdList, ShaderRHI, GlobalDistanceFieldInfo.ParameterData);
 
@@ -276,7 +271,6 @@ public:
 	}
 
 private:
-	LAYOUT_FIELD(FAOParameters, AOParameters);
 	LAYOUT_FIELD(FScreenGridParameters, ScreenGridParameters);
 	LAYOUT_FIELD(FGlobalDistanceFieldParameters, GlobalDistanceFieldParameters);
 	LAYOUT_FIELD(FShaderParameter, TileListGroupSize);
@@ -449,6 +443,7 @@ void FDeferredShadingSceneRenderer::RenderDistanceFieldAOScreenGrid(
 		PassParameters->DistanceFieldCulledObjectBuffers = CulledObjectBufferParameters;
 		PassParameters->DistanceFieldAtlas = DistanceField::SetupAtlasParameters(DistanceFieldSceneData);
 		PassParameters->AOScreenGridParameters = AOScreenGridParameters;
+		PassParameters->AOParameters = DistanceField::SetupAOShaderParameters(Parameters);
 		PassParameters->TileConeDepthRanges = TileIntersectionParameters.TileConeDepthRanges;
 		PassParameters->SceneTextures = SceneTextures.UniformBuffer;
 		PassParameters->DistanceFieldNormal = DistanceFieldNormal;
@@ -464,14 +459,14 @@ void FDeferredShadingSceneRenderer::RenderDistanceFieldAOScreenGrid(
 			RDG_EVENT_NAME("ConeTraceGlobal"),
 			PassParameters,
 			ERDGPassFlags::Compute,
-			[PassParameters, ComputeShader, &View, Parameters, DistanceFieldNormal, TileListGroupSize, bUseObjectDistanceField](FRHICommandList& RHICmdList)
+			[PassParameters, ComputeShader, &View, DistanceFieldNormal, TileListGroupSize, bUseObjectDistanceField](FRHICommandList& RHICmdList)
 			{
 				const uint32 GroupSizeX = FMath::DivideAndRoundUp(View.ViewRect.Size().X / GAODownsampleFactor / GConeTraceDownsampleFactor, GConeTraceGlobalDFTileSize);
 				const uint32 GroupSizeY = FMath::DivideAndRoundUp(View.ViewRect.Size().Y / GAODownsampleFactor / GConeTraceDownsampleFactor, GConeTraceGlobalDFTileSize);
 
 				RHICmdList.SetComputeShader(ComputeShader.GetComputeShader());
 
-				ComputeShader->SetParameters(RHICmdList, View, TileListGroupSize, DistanceFieldNormal->GetRHI(), Parameters, View.GlobalDistanceFieldInfo);
+				ComputeShader->SetParameters(RHICmdList, View, TileListGroupSize, DistanceFieldNormal->GetRHI(), View.GlobalDistanceFieldInfo);
 				SetShaderParameters(RHICmdList, ComputeShader, ComputeShader.GetComputeShader(), *PassParameters);
 
 				DispatchComputeShader(RHICmdList, ComputeShader.GetShader(), GroupSizeX, GroupSizeY, 1);
@@ -489,6 +484,7 @@ void FDeferredShadingSceneRenderer::RenderDistanceFieldAOScreenGrid(
 		PassParameters->DistanceFieldAtlas = DistanceField::SetupAtlasParameters(DistanceFieldSceneData);
 		PassParameters->TileIntersectionParameters = TileIntersectionParameters;
 		PassParameters->AOScreenGridParameters = AOScreenGridParameters;
+		PassParameters->AOParameters = DistanceField::SetupAOShaderParameters(Parameters);
 		PassParameters->ObjectTilesIndirectArguments = ObjectTilesIndirectArguments;
 		PassParameters->SceneTextures = SceneTextures.UniformBuffer;
 		PassParameters->DistanceFieldNormal = DistanceFieldNormal;
@@ -505,11 +501,11 @@ void FDeferredShadingSceneRenderer::RenderDistanceFieldAOScreenGrid(
 			RDG_EVENT_NAME("ConeTraceObjects"),
 			PassParameters,
 			ERDGPassFlags::Compute,
-			[PassParameters, ComputeShader, &View, Parameters, DistanceFieldNormal, bUseGlobalDistanceField, ObjectTilesIndirectArguments](FRHICommandList& RHICmdList)
+			[PassParameters, ComputeShader, &View, DistanceFieldNormal, bUseGlobalDistanceField, ObjectTilesIndirectArguments](FRHICommandList& RHICmdList)
 			{
 				RHICmdList.SetComputeShader(ComputeShader.GetComputeShader());
 
-				ComputeShader->SetParameters(RHICmdList, View, DistanceFieldNormal->GetRHI(), Parameters, bUseGlobalDistanceField, View.GlobalDistanceFieldInfo);
+				ComputeShader->SetParameters(RHICmdList, View, DistanceFieldNormal->GetRHI(), bUseGlobalDistanceField, View.GlobalDistanceFieldInfo);
 				SetShaderParameters(RHICmdList, ComputeShader, ComputeShader.GetComputeShader(), *PassParameters);
 
 				DispatchIndirectComputeShader(RHICmdList, ComputeShader.GetShader(), ObjectTilesIndirectArguments->GetIndirectRHICallBuffer(), 0);
diff --git a/Engine/Source/Runtime/Renderer/Private/DistanceFieldVisualization.cpp b/Engine/Source/Runtime/Renderer/Private/DistanceFieldVisualization.cpp
--- a/Engine/Source/Runtime/Renderer/Private/DistanceFieldVisualization.cpp
+++ b/Engine/Source/Runtime/Renderer/Private/DistanceFieldVisualization.cpp
@@ -31,6 +31,7 @@ public:
 		SHADER_PARAMETER_STRUCT_REF(FReflectionUniformParameters, ReflectionStruct)
 		SHADER_PARAMETER_STRUCT_INCLUDE(FDistanceFieldCulledObjectBufferParameters, DistanceFieldCulledObjectBuffers)
 		SHADER_PARAMETER_STRUCT_INCLUDE(FDistanceFieldAtlasParameters, DistanceFieldAtlas)
+		SHADER_PARAMETER_STRUCT_INCLUDE(FAOParameters, AOParameters)
 		SHADER_PARAMETER_RDG_UNIFORM_BUFFER(FSceneTextureUniformParameters, SceneTextures)
 		SHADER_PARAMETER_RDG_TEXTURE_UAV(RWTexture2D<float4>, RWVisualizeMeshDistanceFields)
 		SHADER_PARAMETER(FVector2f, NumGroups)
@@ -53,19 +54,15 @@ public:
 		: FGlobalShader(Initializer)
 	{
 		BindForLegacyShaderParameters<FParameters>(this, Initializer.PermutationId, Initializer.ParameterMap);
-		AOParameters.Bind(Initializer.ParameterMap);
 		GlobalDistanceFieldParameters.Bind(Initializer.ParameterMap);
 	}
 
 	void SetParameters(
 		FRHICommandList& RHICmdList,
-		const FDistanceFieldAOParameters& Parameters,
 		const FGlobalDistanceFieldInfo& GlobalDistanceFieldInfo)
 	{
 		FRHIComputeShader* ShaderRHI = RHICmdList.GetBoundComputeShader();
 
-		AOParameters.Set(RHICmdList, ShaderRHI, Parameters);
-
 		if (GlobalDistanceFieldParameters.IsBound())
 		{
 			GlobalDistanceFieldParameters.Set(RHICmdList, ShaderRHI, GlobalDistanceFieldInfo.ParameterData);
@@ -73,7 +70,6 @@ public:
 	}
 
 private:
-	LAYOUT_FIELD(FAOParameters, AOParameters);
 	LAYOUT_FIELD(FGlobalDistanceFieldParameters, GlobalDistanceFieldParameters);
 };
 
@@ -169,6 +165,7 @@ void FDeferredShadingSceneRenderer::RenderMeshDistanceFieldVisualization(
 		PassParameters->ReflectionStruct = CreateReflectionUniformBuffer(View, UniformBuffer_MultiFrame);
 		PassParameters->DistanceFieldCulledObjectBuffers = CulledObjectBufferParameters;
 		PassParameters->DistanceFieldAtlas = DistanceField::SetupAtlasParameters(Scene->DistanceFieldSceneData);
+		PassParameters->AOParameters = DistanceField::SetupAOShaderParameters(Parameters);
 		PassParameters->NumGroups = FVector2f(GroupSizeX, GroupSizeY);
 		PassParameters->SceneTextures = SceneTextures.UniformBuffer;
 		PassParameters->RWVisualizeMeshDistanceFields = GraphBuilder.CreateUAV(VisualizeResultTexture);
@@ -179,12 +176,12 @@ void FDeferredShadingSceneRenderer::RenderMeshDistanceFieldVisualization(
 			RDG_EVENT_NAME("VisualizeMeshDistanceFieldCS"),
 			PassParameters,
 			ERDGPassFlags::Compute,
-			[&View, Parameters, ComputeShader, PassParameters, GroupSizeX, GroupSizeY](FRHICommandList& RHICmdList)
+			[&View, ComputeShader, PassParameters, GroupSizeX, GroupSizeY](FRHICommandList& RHICmdList)
 		{
 			FRHIComputeShader* ShaderRHI = ComputeShader.GetComputeShader();
 			RHICmdList.SetComputeShader(ComputeShader.GetComputeShader());
 			SetShaderParameters(RHICmdList, ComputeShader, ShaderRHI, *PassParameters);
-			ComputeShader->SetParameters(RHICmdList, Parameters, View.GlobalDistanceFieldInfo);
+			ComputeShader->SetParameters(RHICmdList,View.GlobalDistanceFieldInfo);
 			DispatchComputeShader(RHICmdList, ComputeShader.GetShader(), GroupSizeX, GroupSizeY, 1);
 			UnsetShaderUAVs(RHICmdList, ComputeShader, ShaderRHI);
 		});
diff --git a/Engine/Source/Runtime/Renderer/Private/HeightfieldLighting.cpp b/Engine/Source/Runtime/Renderer/Private/HeightfieldLighting.cpp
--- a/Engine/Source/Runtime/Renderer/Private/HeightfieldLighting.cpp
+++ b/Engine/Source/Runtime/Renderer/Private/HeightfieldLighting.cpp
@@ -1094,6 +1094,7 @@ public:
 
 	BEGIN_SHADER_PARAMETER_STRUCT(FParameters, )
 		SHADER_PARAMETER_STRUCT_INCLUDE(FAOScreenGridParameters, AOScreenGridParameters)
+		SHADER_PARAMETER_STRUCT_INCLUDE(FAOParameters, AOParameters)
 		SHADER_PARAMETER_RDG_UNIFORM_BUFFER(FSceneTextureUniformParameters, SceneTextures)
 		RDG_TEXTURE_ACCESS(DistanceFieldNormal, ERHIAccess::SRVCompute)
 	END_SHADER_PARAMETER_STRUCT()
@@ -1119,7 +1120,6 @@ public:
 		: FGlobalShader(Initializer)
 	{
 		BindForLegacyShaderParameters<FParameters>(this, Initializer.PermutationId, Initializer.ParameterMap, false);
-		AOParameters.Bind(Initializer.ParameterMap);
 		ScreenGridParameters.Bind(Initializer.ParameterMap);
 		HeightfieldDescriptionParameters.Bind(Initializer.ParameterMap);
 		HeightfieldTextureParameters.Bind(Initializer.ParameterMap);
@@ -1135,13 +1135,11 @@ public:
 		const FViewInfo& View,
 		UTexture2D* HeightfieldTextureValue,
 		int32 NumHeightfieldsValue,
-		FRHITexture* DistanceFieldNormal,
-		const FDistanceFieldAOParameters& Parameters)
+		FRHITexture* DistanceFieldNormal)
 	{
 		FRHIComputeShader* ShaderRHI = RHICmdList.GetBoundComputeShader();
 		FGlobalShader::SetParameters<FViewUniformShaderParameters>(RHICmdList, ShaderRHI, View.ViewUniformBuffer);
 
-		AOParameters.Set(RHICmdList, ShaderRHI, Parameters);
 		ScreenGridParameters.Set(RHICmdList, ShaderRHI, View, DistanceFieldNormal);
 		HeightfieldDescriptionParameters.Set(RHICmdList, ShaderRHI, GetHeightfieldDescriptionsSRV(), NumHeightfieldsValue);
 		HeightfieldTextureParameters.Set(RHICmdList, ShaderRHI, HeightfieldTextureValue, nullptr, nullptr);
@@ -1163,8 +1161,6 @@ public:
 	}
 
 private:
-
-	LAYOUT_FIELD(FAOParameters, AOParameters);
 	LAYOUT_FIELD(FScreenGridParameters, ScreenGridParameters);
 	LAYOUT_FIELD(FHeightfieldDescriptionParameters, HeightfieldDescriptionParameters);
 	LAYOUT_FIELD(FHeightfieldTextureParameters, HeightfieldTextureParameters);
@@ -1195,6 +1191,7 @@ void FHeightfieldLightingViewInfo::ComputeOcclusionForScreenGrid(
 			{
 				auto* PassParameters = GraphBuilder.AllocParameters<FCalculateHeightfieldOcclusionScreenGridCS::FParameters>();
 				PassParameters->AOScreenGridParameters = AOScreenGridParameters;
+				PassParameters->AOParameters = DistanceField::SetupAOShaderParameters(Parameters);
 				PassParameters->SceneTextures = SceneTextures.UniformBuffer;
 				PassParameters->DistanceFieldNormal = DistanceFieldNormal;
 
@@ -1206,7 +1203,7 @@ void FHeightfieldLightingViewInfo::ComputeOcclusionForScreenGrid(
 					RDG_EVENT_NAME("HeightfieldOcclusion"),
 					PassParameters,
 					ERDGPassFlags::Compute,
-					[PassParameters, ComputeShader, &View, Parameters, DistanceFieldNormal, &HeightfieldDescriptions, HeightfieldTexture, DownsampleFactor](FRHICommandListImmediate& RHICmdList)
+					[PassParameters, ComputeShader, &View, DistanceFieldNormal, &HeightfieldDescriptions, HeightfieldTexture, DownsampleFactor](FRHICommandListImmediate& RHICmdList)
 					{
 						UploadHeightfieldDescriptions(HeightfieldDescriptions, FVector2D(1, 1), 1.0f / DownsampleFactor);
 
@@ -1215,7 +1212,7 @@ void FHeightfieldLightingViewInfo::ComputeOcclusionForScreenGrid(
 
 						RHICmdList.SetComputeShader(ComputeShader.GetComputeShader());
 
-						ComputeShader->SetParameters(RHICmdList, View, HeightfieldTexture, HeightfieldDescriptions.Num(), DistanceFieldNormal->GetRHI(), Parameters);
+						ComputeShader->SetParameters(RHICmdList, View, HeightfieldTexture, HeightfieldDescriptions.Num(), DistanceFieldNormal->GetRHI());
 						SetShaderParameters(RHICmdList, ComputeShader, ComputeShader.GetComputeShader(), *PassParameters);
 
 						DispatchComputeShader(RHICmdList, ComputeShader.GetShader(), GroupSizeX, GroupSizeY, 1);
diff --git a/Engine/Source/Runtime/Renderer/Private/VolumetricFog.cpp b/Engine/Source/Runtime/Renderer/Private/VolumetricFog.cpp
--- a/Engine/Source/Runtime/Renderer/Private/VolumetricFog.cpp
+++ b/Engine/Source/Runtime/Renderer/Private/VolumetricFog.cpp
@@ -825,6 +825,7 @@ class TVolumetricFogLightScatteringCS : public FGlobalShader
 		SHADER_PARAMETER(uint32, UseEmissive)
 		SHADER_PARAMETER_RDG_UNIFORM_BUFFER(FLumenTranslucencyLightingUniforms, LumenGIVolumeStruct)
 		SHADER_PARAMETER_STRUCT_INCLUDE(FVirtualShadowMapSamplingParameters, VirtualShadowMapSamplingParameters)
+		SHADER_PARAMETER_STRUCT_INCLUDE(FAOParameters, AOParameters)
 		SHADER_PARAMETER_RDG_TEXTURE_UAV(RWTexture2D, RWLightScattering)
 	END_SHADER_PARAMETER_STRUCT()
 
@@ -892,7 +893,6 @@ public:
 		InverseSquaredLightDistanceBiasScale.Bind(Initializer.ParameterMap, TEXT("InverseSquaredLightDistanceBiasScale"));
 		UseHeightFogColors.Bind(Initializer.ParameterMap, TEXT("UseHeightFogColors"));
 		UseDirectionalLightShadowing.Bind(Initializer.ParameterMap, TEXT("UseDirectionalLightShadowing"));
-		AOParameters.Bind(Initializer.ParameterMap);
 		GlobalDistanceFieldParameters.Bind(Initializer.ParameterMap);
 		LightScatteringSampleJitterMultiplier.Bind(Initializer.ParameterMap, TEXT("LightScatteringSampleJitterMultiplier"));
 
@@ -942,7 +942,6 @@ public:
 		SetSamplerParameter(RHICmdList, ShaderRHI, LightFunctionSampler, TStaticSamplerState<SF_Bilinear, AM_Clamp, AM_Clamp, AM_Clamp>::GetRHI());
 
 		FScene* Scene = (FScene*)View.Family->Scene;
-		FDistanceFieldAOParameters AOParameterData(Scene->DefaultMaxDistanceFieldOcclusionDistance);
 		FSkyLightSceneProxy* SkyLight = Scene->SkyLight;
 
 		if (SkyLight
@@ -958,8 +957,6 @@ public:
 			SetShaderValue(RHICmdList, ShaderRHI, SkySH, (FVector4f&)SkyIrradiance.R.V, 0);
 			SetShaderValue(RHICmdList, ShaderRHI, SkySH, (FVector4f&)SkyIrradiance.G.V, 1);
 			SetShaderValue(RHICmdList, ShaderRHI, SkySH, (FVector4f&)SkyIrradiance.B.V, 2);
-
-			AOParameterData = FDistanceFieldAOParameters(SkyLight->OcclusionMaxDistance, SkyLight->Contrast);
 		}
 		else
 		{
@@ -987,7 +984,6 @@ public:
 			OverrideDirectionalLightInScatteringUsingHeightFog(View, FogInfo) ? 1.0f : 0.0f,
 			OverrideSkyLightInScatteringUsingHeightFog(View, FogInfo) ? 1.0f : 0.0f ));
 
-		AOParameters.Set(RHICmdList, ShaderRHI, AOParameterData);
 		GlobalDistanceFieldParameters.Set(RHICmdList, ShaderRHI, View.GlobalDistanceFieldInfo.ParameterData);
 
 		if (CloudShadowmapTexture.IsBound())
@@ -1050,7 +1046,6 @@ private:
 	LAYOUT_FIELD(FShaderParameter, InverseSquaredLightDistanceBiasScale);
 	LAYOUT_FIELD(FShaderParameter, UseHeightFogColors);
 	LAYOUT_FIELD(FShaderParameter, UseDirectionalLightShadowing);
-	LAYOUT_FIELD(FAOParameters, AOParameters);
 	LAYOUT_FIELD(FGlobalDistanceFieldParameters, GlobalDistanceFieldParameters);
 	LAYOUT_FIELD(FShaderResourceParameter, CloudShadowmapTexture);
 	LAYOUT_FIELD(FShaderResourceParameter, CloudShadowmapSampler);
@@ -1497,6 +1492,16 @@ void FDeferredShadingSceneRenderer::ComputeVolumetricFog(FRDGBuilder& GraphBuild
 			PassParameters->RWLightScattering = IntegrationData.LightScatteringUAV;
 			PassParameters->VirtualShadowMapSamplingParameters = VirtualShadowMapArray.GetSamplingParameters(GraphBuilder);
 
+			FDistanceFieldAOParameters AOParameterData(Scene->DefaultMaxDistanceFieldOcclusionDistance);
+			if (Scene->SkyLight
+				// Skylights with static lighting had their diffuse contribution baked into lightmaps
+				&& !Scene->SkyLight->bHasStaticLighting
+				&& View.Family->EngineShowFlags.SkyLighting)
+			{
+				AOParameterData = FDistanceFieldAOParameters(Scene->SkyLight->OcclusionMaxDistance, Scene->SkyLight->Contrast);
+			}
+			PassParameters->AOParameters = DistanceField::SetupAOShaderParameters(AOParameterData);
+
 			const bool bUseLumenGI = View.LumenTranslucencyGIVolume.Texture0 != nullptr;
 			const bool bUseGlobalDistanceField = UseGlobalDistanceField() && Scene->DistanceFieldSceneData.NumObjectsInBuffer > 0;
 
