Refactor remaining DF shaders to use RDG - Part 2
- Convert FDFAOUpsampleParameters to SHADER_PARAMETER_STRUCT.
    - Added DistanceField::SetupAOUpsampleParameters helper function to perform the necessary logic.
- Convert FGeometryAwareUpsampleParameters to SHADER_PARAMETER_STRUCT.
    - Added SetupGeometryAwareUpsampleParameters helper function to perform the necessary logic.
- Updated TFilterHistoryPS to use SHADER_PERMUTATION_BOOL.
- Removed all legacy shader parameters from FUpdateHistoryDepthRejectionPS, FFilterHistoryPS, FGeometryAwareUpsamplePS, FDistanceFieldAOUpsamplePS.

diff --git a/Engine/Source/Runtime/Renderer/Private/DistanceFieldAmbientOcclusion.cpp b/Engine/Source/Runtime/Renderer/Private/DistanceFieldAmbientOcclusion.cpp
--- a/Engine/Source/Runtime/Renderer/Private/DistanceFieldAmbientOcclusion.cpp
+++ b/Engine/Source/Runtime/Renderer/Private/DistanceFieldAmbientOcclusion.cpp
@@ -21,7 +21,6 @@
 #include "Lumen/Lumen.h"
 #include "Strata/Strata.h"
 
-IMPLEMENT_TYPE_LAYOUT(FDFAOUpsampleParameters);
 IMPLEMENT_TYPE_LAYOUT(FScreenGridParameters);
 
 int32 GDistanceFieldAO = 1;
@@ -242,6 +241,8 @@ float GAOConeHalfAngle = FMath::Acos(1 - 1.0f / (float)UE_ARRAY_COUNT(SpacedVect
 // Must match shader code
 uint32 GAONumConeSteps = 10;
 
+extern float GAOViewFadeDistanceScale;
+
 FAOParameters DistanceField::SetupAOShaderParameters(const FDistanceFieldAOParameters& Parameters)
 {
 	const float AOLargestSampleOffset = Parameters.ObjectMaxOcclusionDistance / (1 + FMath::Tan(GAOConeHalfAngle));
@@ -256,6 +257,25 @@ FAOParameters DistanceField::SetupAOShaderParameters(const FDistanceFieldAOParam
 	return ShaderParameters;
 }
 
+FDFAOUpsampleParameters DistanceField::SetupAOUpsampleParameters(const FViewInfo& View, FRDGTextureRef DistanceFieldAOBentNormal)
+{
+	const float DistanceFadeScaleValue = 1.0f / ((1.0f - GAOViewFadeDistanceScale) * GetMaxAOViewDistance());
+
+	const FIntPoint AOBufferSize = GetBufferSizeForAO();
+	const FVector2f UVMax(
+		(View.ViewRect.Width() / GAODownsampleFactor - 0.51f) / AOBufferSize.X, // 0.51 - so bilateral gather4 won't sample invalid texels
+		(View.ViewRect.Height() / GAODownsampleFactor - 0.51f) / AOBufferSize.Y);
+
+	FDFAOUpsampleParameters ShaderParameters;
+	ShaderParameters.BentNormalAOTexture = DistanceFieldAOBentNormal;
+	ShaderParameters.BentNormalAOSampler = TStaticSamplerState<SF_Bilinear>::GetRHI();
+	ShaderParameters.AOBufferBilinearUVMax = UVMax;
+	ShaderParameters.DistanceFadeScale = DistanceFadeScaleValue;
+	ShaderParameters.AOMaxViewDistance = GetMaxAOViewDistance();
+
+	return ShaderParameters;
+}
+
 bool bListMemoryNextFrame = false;
 
 void OnListMemory(UWorld* InWorld)
diff --git a/Engine/Source/Runtime/Renderer/Private/DistanceFieldAmbientOcclusion.h b/Engine/Source/Runtime/Renderer/Private/DistanceFieldAmbientOcclusion.h
--- a/Engine/Source/Runtime/Renderer/Private/DistanceFieldAmbientOcclusion.h
+++ b/Engine/Source/Runtime/Renderer/Private/DistanceFieldAmbientOcclusion.h
@@ -98,60 +98,18 @@ BEGIN_SHADER_PARAMETER_STRUCT(FAOParameters, )
 	SHADER_PARAMETER(float, AOGlobalMaxOcclusionDistance)
 END_SHADER_PARAMETER_STRUCT()
 
+BEGIN_SHADER_PARAMETER_STRUCT(FDFAOUpsampleParameters, )
+	SHADER_PARAMETER_RDG_TEXTURE(Texture2D, BentNormalAOTexture)
+	SHADER_PARAMETER_SAMPLER(SamplerState, BentNormalAOSampler)
+	SHADER_PARAMETER(FVector2f, AOBufferBilinearUVMax)
+	SHADER_PARAMETER(float, DistanceFadeScale)
+	SHADER_PARAMETER(float, AOMaxViewDistance)
+END_SHADER_PARAMETER_STRUCT()
+
 namespace DistanceField
 {
 	FAOParameters SetupAOShaderParameters(const FDistanceFieldAOParameters& AOParameters);
-};
-
-class FDFAOUpsampleParameters
-{
-	DECLARE_TYPE_LAYOUT(FDFAOUpsampleParameters, NonVirtual);
-public:
-	void Bind(const FShaderParameterMap& ParameterMap)
-	{
-		BentNormalAOTexture.Bind(ParameterMap, TEXT("BentNormalAOTexture"));
-		BentNormalAOSampler.Bind(ParameterMap, TEXT("BentNormalAOSampler"));
-		AOBufferBilinearUVMax.Bind(ParameterMap, TEXT("AOBufferBilinearUVMax"));
-		DistanceFadeScale.Bind(ParameterMap, TEXT("DistanceFadeScale"));
-		AOMaxViewDistance.Bind(ParameterMap, TEXT("AOMaxViewDistance"));
-	}
-
-	void Set(FRHICommandList& RHICmdList, FRHIPixelShader* ShaderRHI, const FViewInfo& View, FRHITexture* DistanceFieldAOBentNormal)
-	{
-		FRHITexture* BentNormalAO = DistanceFieldAOBentNormal ? DistanceFieldAOBentNormal : (FRHITexture*)GWhiteTexture->TextureRHI;
-		SetTextureParameter(RHICmdList, ShaderRHI, BentNormalAOTexture, BentNormalAOSampler, TStaticSamplerState<SF_Bilinear>::GetRHI(), BentNormalAO);
-
-		FIntPoint const AOBufferSize = GetBufferSizeForAO();
-		FVector2f const UVMax(
-			(View.ViewRect.Width() / GAODownsampleFactor - 0.51f) / AOBufferSize.X, // 0.51 - so bilateral gather4 won't sample invalid texels
-			(View.ViewRect.Height() / GAODownsampleFactor - 0.51f) / AOBufferSize.Y);
-		SetShaderValue(RHICmdList, ShaderRHI, AOBufferBilinearUVMax, UVMax);
-
-		SetShaderValue(RHICmdList, ShaderRHI, AOMaxViewDistance, GetMaxAOViewDistance());
-
-		extern float GAOViewFadeDistanceScale;
-		const float DistanceFadeScaleValue = 1.0f / ((1.0f - GAOViewFadeDistanceScale) * GetMaxAOViewDistance());
-		SetShaderValue(RHICmdList, ShaderRHI, DistanceFadeScale, DistanceFadeScaleValue);
-	}
-
-	/** Serializer. */
-	friend FArchive& operator<<(FArchive& Ar, FDFAOUpsampleParameters& P)
-	{
-		Ar << P.BentNormalAOTexture;
-		Ar << P.BentNormalAOSampler;
-		Ar << P.AOBufferBilinearUVMax;
-		Ar << P.DistanceFadeScale;
-		Ar << P.AOMaxViewDistance;
-
-		return Ar;
-	}
-
-private:
-	LAYOUT_FIELD(FShaderResourceParameter, BentNormalAOTexture);
-	LAYOUT_FIELD(FShaderResourceParameter, BentNormalAOSampler);
-	LAYOUT_FIELD(FShaderParameter, AOBufferBilinearUVMax);
-	LAYOUT_FIELD(FShaderParameter, DistanceFadeScale);
-	LAYOUT_FIELD(FShaderParameter, AOMaxViewDistance);
+	FDFAOUpsampleParameters SetupAOUpsampleParameters(const FViewInfo& View, FRDGTextureRef DistanceFieldAOBentNormal);
 };
 
 class FMaxSizedRWBuffers : public FRenderResource
diff --git a/Engine/Source/Runtime/Renderer/Private/DistanceFieldLightingPost.cpp b/Engine/Source/Runtime/Renderer/Private/DistanceFieldLightingPost.cpp
--- a/Engine/Source/Runtime/Renderer/Private/DistanceFieldLightingPost.cpp
+++ b/Engine/Source/Runtime/Renderer/Private/DistanceFieldLightingPost.cpp
@@ -70,114 +70,72 @@ bool UseAOHistoryStabilityPass()
 	return GAOHistoryStabilityPass && GDistanceFieldAOQuality >= 2;
 }
 
-class FGeometryAwareUpsampleParameters
+BEGIN_SHADER_PARAMETER_STRUCT(FGeometryAwareUpsampleParameters, )
+	SHADER_PARAMETER_RDG_TEXTURE(Texture2D, DistanceFieldNormalTexture)
+	SHADER_PARAMETER_SAMPLER(SamplerState, DistanceFieldNormalSampler)
+	SHADER_PARAMETER_RDG_TEXTURE(Texture2D, BentNormalAOTexture)
+	SHADER_PARAMETER_SAMPLER(SamplerState, BentNormalAOSampler)
+	SHADER_PARAMETER(FVector4f, BentNormalBufferAndTexelSize)
+	SHADER_PARAMETER(FVector2f, DistanceFieldGBufferTexelSize)
+	SHADER_PARAMETER(FVector2f, DistanceFieldGBufferJitterOffset)
+	SHADER_PARAMETER(FVector2f, JitterOffset)
+	SHADER_PARAMETER(float, MinDownsampleFactorToBaseLevel)
+	SHADER_PARAMETER(float, DistanceFadeScale)
+END_SHADER_PARAMETER_STRUCT()
+
+FGeometryAwareUpsampleParameters SetupGeometryAwareUpsampleParameters(const FViewInfo& View, FRDGTextureRef DistanceFieldNormal, FRDGTextureRef DistanceFieldAOBentNormal)
 {
-	DECLARE_INLINE_TYPE_LAYOUT(FGeometryAwareUpsampleParameters, NonVirtual);
-public:
-	void Bind(const FShaderParameterMap& ParameterMap)
-	{
-		DistanceFieldNormalTexture.Bind(ParameterMap, TEXT("DistanceFieldNormalTexture"));
-		DistanceFieldNormalSampler.Bind(ParameterMap, TEXT("DistanceFieldNormalSampler"));
-		BentNormalAOTexture.Bind(ParameterMap, TEXT("BentNormalAOTexture"));
-		BentNormalAOSampler.Bind(ParameterMap, TEXT("BentNormalAOSampler"));
-		DistanceFieldGBufferTexelSize.Bind(ParameterMap, TEXT("DistanceFieldGBufferTexelSize"));
-		DistanceFieldGBufferJitterOffset.Bind(ParameterMap, TEXT("DistanceFieldGBufferJitterOffset"));
-		BentNormalBufferAndTexelSize.Bind(ParameterMap, TEXT("BentNormalBufferAndTexelSize"));
-		MinDownsampleFactorToBaseLevel.Bind(ParameterMap, TEXT("MinDownsampleFactorToBaseLevel"));
-		DistanceFadeScale.Bind(ParameterMap, TEXT("DistanceFadeScale"));
-		JitterOffset.Bind(ParameterMap, TEXT("JitterOffset"));
-	}
-
-	void Set(FRHICommandList& RHICmdList, FRHIPixelShader* ShaderRHI, const FViewInfo& View, FRHITexture* DistanceFieldNormal, FRHITexture* DistanceFieldAOBentNormal)
-	{
-		SetTextureParameter(
-			RHICmdList,
-			ShaderRHI,
-			DistanceFieldNormalTexture,
-			DistanceFieldNormalSampler,
-			TStaticSamplerState<SF_Bilinear, AM_Clamp, AM_Clamp, AM_Clamp>::GetRHI(),
-			DistanceFieldNormal
-		);
-
-		SetTextureParameter(
-			RHICmdList,
-			ShaderRHI,
-			BentNormalAOTexture,
-			BentNormalAOSampler,
-			TStaticSamplerState<SF_Bilinear, AM_Clamp, AM_Clamp, AM_Clamp>::GetRHI(),
-			DistanceFieldAOBentNormal
-		);
-
-		extern FVector2f GetJitterOffset(int32 SampleIndex);
-		FVector2f const JitterOffsetValue = GetJitterOffset(View.GetDistanceFieldTemporalSampleIndex());
-
-		const FIntPoint DownsampledBufferSize = GetBufferSizeForAO();
-		const FVector2f BaseLevelTexelSizeValue(1.0f / DownsampledBufferSize.X, 1.0f / DownsampledBufferSize.Y);
-		SetShaderValue(RHICmdList, ShaderRHI, DistanceFieldGBufferTexelSize, BaseLevelTexelSizeValue);
-
-		SetShaderValue(RHICmdList, ShaderRHI, DistanceFieldGBufferJitterOffset, BaseLevelTexelSizeValue * JitterOffsetValue);
-
-		extern FIntPoint GetBufferSizeForConeTracing();
-		const FIntPoint ConeTracingBufferSize = GetBufferSizeForConeTracing();
-		const FVector4f BentNormalBufferAndTexelSizeValue(ConeTracingBufferSize.X, ConeTracingBufferSize.Y, 1.0f / ConeTracingBufferSize.X, 1.0f / ConeTracingBufferSize.Y);
-		SetShaderValue(RHICmdList, ShaderRHI, BentNormalBufferAndTexelSize, BentNormalBufferAndTexelSizeValue);
-
-		extern int32 GConeTraceDownsampleFactor;
-		const float MinDownsampleFactor = GConeTraceDownsampleFactor;
-		SetShaderValue(RHICmdList, ShaderRHI, MinDownsampleFactorToBaseLevel, MinDownsampleFactor);
-
-		extern float GAOViewFadeDistanceScale;
-		const float DistanceFadeScaleValue = 1.0f / ((1.0f - GAOViewFadeDistanceScale) * GetMaxAOViewDistance());
-		SetShaderValue(RHICmdList, ShaderRHI, DistanceFadeScale, DistanceFadeScaleValue);
-
-
-		SetShaderValue(RHICmdList, ShaderRHI, JitterOffset, JitterOffsetValue);
-	}
-
-	/** Serializer. */
-	friend FArchive& operator<<(FArchive& Ar, FGeometryAwareUpsampleParameters& P)
-	{
-		Ar << P.DistanceFieldNormalTexture;
-		Ar << P.DistanceFieldNormalSampler;
-		Ar << P.BentNormalAOTexture;
-		Ar << P.BentNormalAOSampler;
-		Ar << P.DistanceFieldGBufferTexelSize;
-		Ar << P.DistanceFieldGBufferJitterOffset;
-		Ar << P.BentNormalBufferAndTexelSize;
-		Ar << P.MinDownsampleFactorToBaseLevel;
-		Ar << P.DistanceFadeScale;
-		Ar << P.JitterOffset;
-		return Ar;
-	}
-
-private:
-	
-		LAYOUT_FIELD(FShaderResourceParameter, DistanceFieldNormalTexture)
-		LAYOUT_FIELD(FShaderResourceParameter, DistanceFieldNormalSampler)
-		LAYOUT_FIELD(FShaderResourceParameter, BentNormalAOTexture)
-		LAYOUT_FIELD(FShaderResourceParameter, BentNormalAOSampler)
-		LAYOUT_FIELD(FShaderParameter, DistanceFieldGBufferTexelSize)
-		LAYOUT_FIELD(FShaderParameter, DistanceFieldGBufferJitterOffset)
-		LAYOUT_FIELD(FShaderParameter, BentNormalBufferAndTexelSize)
-		LAYOUT_FIELD(FShaderParameter, MinDownsampleFactorToBaseLevel)
-		LAYOUT_FIELD(FShaderParameter, DistanceFadeScale)
-		LAYOUT_FIELD(FShaderParameter, JitterOffset)
-	
-};
+	extern FVector2f GetJitterOffset(int32 SampleIndex);
+	FVector2f const JitterOffsetValue = GetJitterOffset(View.GetDistanceFieldTemporalSampleIndex());
+
+	const FIntPoint DownsampledBufferSize = GetBufferSizeForAO();
+	const FVector2f BaseLevelTexelSizeValue(1.0f / DownsampledBufferSize.X, 1.0f / DownsampledBufferSize.Y);
+
+	extern FIntPoint GetBufferSizeForConeTracing();
+	const FIntPoint ConeTracingBufferSize = GetBufferSizeForConeTracing();
+	const FVector4f BentNormalBufferAndTexelSizeValue(ConeTracingBufferSize.X, ConeTracingBufferSize.Y, 1.0f / ConeTracingBufferSize.X, 1.0f / ConeTracingBufferSize.Y);
+
+	extern int32 GConeTraceDownsampleFactor;
+	const float MinDownsampleFactor = GConeTraceDownsampleFactor;
+
+	extern float GAOViewFadeDistanceScale;
+	const float DistanceFadeScaleValue = 1.0f / ((1.0f - GAOViewFadeDistanceScale) * GetMaxAOViewDistance());
+
+	FGeometryAwareUpsampleParameters ShaderParameters;
+	ShaderParameters.DistanceFieldNormalTexture = DistanceFieldNormal;
+	ShaderParameters.DistanceFieldNormalSampler = TStaticSamplerState<SF_Bilinear, AM_Clamp, AM_Clamp, AM_Clamp>::GetRHI();
+	ShaderParameters.BentNormalAOTexture = DistanceFieldAOBentNormal;
+	ShaderParameters.BentNormalAOSampler = TStaticSamplerState<SF_Bilinear, AM_Clamp, AM_Clamp, AM_Clamp>::GetRHI();
+	ShaderParameters.BentNormalBufferAndTexelSize = BentNormalBufferAndTexelSizeValue;
+	ShaderParameters.DistanceFieldGBufferTexelSize = BaseLevelTexelSizeValue;
+	ShaderParameters.DistanceFieldGBufferJitterOffset = BaseLevelTexelSizeValue * JitterOffsetValue;
+	ShaderParameters.JitterOffset = JitterOffsetValue;
+	ShaderParameters.MinDownsampleFactorToBaseLevel = MinDownsampleFactor;
+	ShaderParameters.DistanceFadeScale = DistanceFadeScaleValue;
+
+	return ShaderParameters;
+}
 
 class FUpdateHistoryDepthRejectionPS : public FGlobalShader
 {
-	DECLARE_SHADER_TYPE(FUpdateHistoryDepthRejectionPS, Global);
 public:
+	DECLARE_GLOBAL_SHADER(FUpdateHistoryDepthRejectionPS);
+	SHADER_USE_PARAMETER_STRUCT(FUpdateHistoryDepthRejectionPS, FGlobalShader);
 
-	// TODO(RDG) Hook these up to the shader. They are just informing RDG of transitions ATM.
 	BEGIN_SHADER_PARAMETER_STRUCT(FParameters, )
+		SHADER_PARAMETER_STRUCT_REF(FViewUniformShaderParameters, View)
 		SHADER_PARAMETER_RDG_UNIFORM_BUFFER(FSceneTextureUniformParameters, SceneTextures)
 		SHADER_PARAMETER_STRUCT_INCLUDE(FAOParameters, AOParameters)
-		RDG_TEXTURE_ACCESS(DistanceFieldNormal, ERHIAccess::SRVGraphics)
-		RDG_TEXTURE_ACCESS(DistanceFieldAOBentNormal, ERHIAccess::SRVGraphics)
-		RDG_TEXTURE_ACCESS(BentNormalHistoryTexture, ERHIAccess::SRVGraphics)
-		RDG_TEXTURE_ACCESS(VelocityTexture, ERHIAccess::SRVGraphics)
+		SHADER_PARAMETER_STRUCT_INCLUDE(FGeometryAwareUpsampleParameters, GeometryAwareUpsampleParameters)
+		SHADER_PARAMETER_RDG_TEXTURE(Texture2D, BentNormalHistoryTexture)
+		SHADER_PARAMETER_SAMPLER(SamplerState, BentNormalHistorySampler)
+		SHADER_PARAMETER_RDG_TEXTURE(Texture2D, VelocityTexture)
+		SHADER_PARAMETER_SAMPLER(SamplerState, VelocityTextureSampler)
+		SHADER_PARAMETER(FVector4f, HistoryScreenPositionScaleBias)
+		SHADER_PARAMETER(FVector4f, HistoryUVMinMax)
+		SHADER_PARAMETER(float, HistoryWeight)
+		SHADER_PARAMETER(float, HistoryDistanceThreshold)
+		SHADER_PARAMETER(float, UseHistoryFilter)
 		RENDER_TARGET_BINDING_SLOTS()
 	END_SHADER_PARAMETER_STRUCT()
 
@@ -190,116 +148,32 @@ public:
 	{
 		OutEnvironment.SetDefine(TEXT("DOWNSAMPLE_FACTOR"), GAODownsampleFactor);
 	}
-
-	/** Default constructor. */
-	FUpdateHistoryDepthRejectionPS() {}
-
-	/** Initialization constructor. */
-	FUpdateHistoryDepthRejectionPS(const ShaderMetaType::CompiledShaderInitializerType& Initializer)
-		: FGlobalShader(Initializer)
-	{
-		GeometryAwareUpsampleParameters.Bind(Initializer.ParameterMap);
-		BentNormalHistoryTexture.Bind(Initializer.ParameterMap, TEXT("BentNormalHistoryTexture"));
-		BentNormalHistorySampler.Bind(Initializer.ParameterMap, TEXT("BentNormalHistorySampler"));
-		HistoryWeight.Bind(Initializer.ParameterMap, TEXT("HistoryWeight"));
-		HistoryDistanceThreshold.Bind(Initializer.ParameterMap, TEXT("HistoryDistanceThreshold"));
-		UseHistoryFilter.Bind(Initializer.ParameterMap, TEXT("UseHistoryFilter"));
-		VelocityTexture.Bind(Initializer.ParameterMap, TEXT("VelocityTexture"));
-		VelocityTextureSampler.Bind(Initializer.ParameterMap, TEXT("VelocityTextureSampler"));
-		HistoryScreenPositionScaleBias.Bind(Initializer.ParameterMap, TEXT("HistoryScreenPositionScaleBias"));
-		HistoryUVMinMax.Bind(Initializer.ParameterMap, TEXT("HistoryUVMinMax"));
-	}
-
-	void SetParameters(
-		FRHICommandList& RHICmdList, 
-		const FViewInfo& View,
-		const FIntRect& HistoryViewRect,
-		FRHITexture* DistanceFieldNormal, 
-		FRHITexture* DistanceFieldAOBentNormal,
-		FRHITexture* BentNormalHistoryTextureValue, 
-		FRHITexture* VelocityTextureValue)
-	{
-		FRHIPixelShader* ShaderRHI = RHICmdList.GetBoundPixelShader();
-
-		FGlobalShader::SetParameters<FViewUniformShaderParameters>(RHICmdList, ShaderRHI, View.ViewUniformBuffer);
-		GeometryAwareUpsampleParameters.Set(RHICmdList, ShaderRHI, View, DistanceFieldNormal, DistanceFieldAOBentNormal);
-
-		SetTextureParameter(
-			RHICmdList,
-			ShaderRHI,
-			BentNormalHistoryTexture,
-			BentNormalHistorySampler,
-			TStaticSamplerState<SF_Bilinear,AM_Clamp,AM_Clamp,AM_Clamp>::GetRHI(),
-			BentNormalHistoryTextureValue
-			);
-
-		SetShaderValue(RHICmdList, ShaderRHI, HistoryWeight, GAOHistoryWeight);
-		SetShaderValue(RHICmdList, ShaderRHI, HistoryDistanceThreshold, GAOHistoryDistanceThreshold);
-		SetShaderValue(RHICmdList, ShaderRHI, UseHistoryFilter, UseAOHistoryStabilityPass() ? 1.0f : 0.0f);
-
-		SetTextureParameter(
-			RHICmdList,
-			ShaderRHI,
-			VelocityTexture,
-			VelocityTextureSampler,
-			TStaticSamplerState<SF_Bilinear,AM_Clamp,AM_Clamp,AM_Clamp>::GetRHI(),
-			VelocityTextureValue ? (FRHITexture2D*)VelocityTextureValue : GBlackTexture->TextureRHI
-			);
-
-		{
-			FIntPoint HistoryBufferSize = GetSceneTextureExtent() / FIntPoint(GAODownsampleFactor, GAODownsampleFactor);
-
-			const float InvBufferSizeX = 1.0f / HistoryBufferSize.X;
-			const float InvBufferSizeY = 1.0f / HistoryBufferSize.Y;
-
-			const FVector4f HistoryScreenPositionScaleBiasValue(
-				HistoryViewRect.Width() * InvBufferSizeX / +2.0f,
-				HistoryViewRect.Height() * InvBufferSizeY / (-2.0f * GProjectionSignY),
-				(HistoryViewRect.Height() / 2.0f + HistoryViewRect.Min.Y) * InvBufferSizeY,
-				(HistoryViewRect.Width() / 2.0f + HistoryViewRect.Min.X) * InvBufferSizeX);
-
-			// Pull in the max UV to exclude the region which will read outside the viewport due to bilinear filtering
-			const FVector4f HistoryUVMinMaxValue(
-				(HistoryViewRect.Min.X + 0.5f) * InvBufferSizeX,
-				(HistoryViewRect.Min.Y + 0.5f) * InvBufferSizeY,
-				(HistoryViewRect.Max.X - 0.5f) * InvBufferSizeX,
-				(HistoryViewRect.Max.Y - 0.5f) * InvBufferSizeY);
-
-			SetShaderValue(RHICmdList, ShaderRHI, HistoryScreenPositionScaleBias, HistoryScreenPositionScaleBiasValue);
-			SetShaderValue(RHICmdList, ShaderRHI, HistoryUVMinMax, HistoryUVMinMaxValue);
-		}
-	}
-
-private:
-	LAYOUT_FIELD(FGeometryAwareUpsampleParameters, GeometryAwareUpsampleParameters);
-	LAYOUT_FIELD(FShaderResourceParameter, BentNormalHistoryTexture);
-	LAYOUT_FIELD(FShaderResourceParameter, BentNormalHistorySampler);
-	LAYOUT_FIELD(FShaderParameter, HistoryWeight);
-	LAYOUT_FIELD(FShaderParameter, HistoryDistanceThreshold);
-	LAYOUT_FIELD(FShaderParameter, UseHistoryFilter);
-	LAYOUT_FIELD(FShaderResourceParameter, VelocityTexture);
-	LAYOUT_FIELD(FShaderResourceParameter, VelocityTextureSampler);
-	LAYOUT_FIELD(FShaderParameter, HistoryScreenPositionScaleBias);
-	LAYOUT_FIELD(FShaderParameter, HistoryUVMinMax);
 };
 
-IMPLEMENT_SHADER_TYPE(, FUpdateHistoryDepthRejectionPS,TEXT("/Engine/Private/DistanceFieldLightingPost.usf"),TEXT("UpdateHistoryDepthRejectionPS"),SF_Pixel);
+IMPLEMENT_GLOBAL_SHADER(FUpdateHistoryDepthRejectionPS, "/Engine/Private/DistanceFieldLightingPost.usf", "UpdateHistoryDepthRejectionPS", SF_Pixel);
 
-
-template<bool bManuallyClampUV>
-class TFilterHistoryPS : public FGlobalShader
+class FFilterHistoryPS : public FGlobalShader
 {
-	DECLARE_SHADER_TYPE(TFilterHistoryPS, Global);
 public:
-	
-	// TODO(RDG) Hook these up to the shader. They are just informing RDG of transitions ATM.
+	DECLARE_GLOBAL_SHADER(FFilterHistoryPS);
+	SHADER_USE_PARAMETER_STRUCT(FFilterHistoryPS, FGlobalShader);
+
 	BEGIN_SHADER_PARAMETER_STRUCT(FParameters, )
+		SHADER_PARAMETER_STRUCT_REF(FViewUniformShaderParameters, View)
 		SHADER_PARAMETER_RDG_UNIFORM_BUFFER(FSceneTextureUniformParameters, SceneTextures)
-		RDG_TEXTURE_ACCESS(DistanceFieldNormal, ERHIAccess::SRVGraphics)
-		RDG_TEXTURE_ACCESS(NewBentNormalHistory, ERHIAccess::SRVGraphics)
+		SHADER_PARAMETER_RDG_TEXTURE(Texture2D, BentNormalAOTexture)
+		SHADER_PARAMETER_RDG_TEXTURE(Texture2D, DistanceFieldNormalTexture)
+		SHADER_PARAMETER_SAMPLER(SamplerState, BentNormalAOSampler)
+		SHADER_PARAMETER_SAMPLER(SamplerState, DistanceFieldNormalSampler)
+		SHADER_PARAMETER(FVector2f, BentNormalAOTexelSize)
+		SHADER_PARAMETER(FVector2f, MaxSampleBufferUV)
+		SHADER_PARAMETER(float, HistoryWeight)
 		RENDER_TARGET_BINDING_SLOTS()
 	END_SHADER_PARAMETER_STRUCT()
 
+	class FManuallyClampUV : SHADER_PERMUTATION_BOOL("MANUALLY_CLAMP_UV");
+	using FPermutationDomain = TShaderPermutationDomain<FManuallyClampUV>;
+
 	static bool ShouldCompilePermutation(const FGlobalShaderPermutationParameters& Parameters)
 	{
 		return ShouldCompileDistanceFieldShaders(Parameters.Platform);
@@ -308,98 +182,21 @@ public:
 	static void ModifyCompilationEnvironment(const FGlobalShaderPermutationParameters& Parameters, FShaderCompilerEnvironment& OutEnvironment)
 	{
 		OutEnvironment.SetDefine(TEXT("DOWNSAMPLE_FACTOR"), GAODownsampleFactor);
-		OutEnvironment.SetDefine(TEXT("MANUALLY_CLAMP_UV"), bManuallyClampUV);
-	}
-
-	/** Default constructor. */
-	TFilterHistoryPS() {}
-
-	/** Initialization constructor. */
-	TFilterHistoryPS(const ShaderMetaType::CompiledShaderInitializerType& Initializer)
-		: FGlobalShader(Initializer)
-	{
-		BentNormalAOTexture.Bind(Initializer.ParameterMap, TEXT("BentNormalAOTexture"));
-		BentNormalAOSampler.Bind(Initializer.ParameterMap, TEXT("BentNormalAOSampler"));
-		HistoryWeight.Bind(Initializer.ParameterMap, TEXT("HistoryWeight"));
-		BentNormalAOTexelSize.Bind(Initializer.ParameterMap, TEXT("BentNormalAOTexelSize"));
-		MaxSampleBufferUV.Bind(Initializer.ParameterMap, TEXT("MaxSampleBufferUV"));
-		DistanceFieldNormalTexture.Bind(Initializer.ParameterMap, TEXT("DistanceFieldNormalTexture"));
-		DistanceFieldNormalSampler.Bind(Initializer.ParameterMap, TEXT("DistanceFieldNormalSampler"));
-	}
-
-	void SetParameters(
-		FRHICommandList& RHICmdList, 
-		const FViewInfo& View, 
-		FRHITexture* DistanceFieldNormal, 
-		FRHITexture* BentNormalHistoryTextureValue)
-	{
-		FRHIPixelShader* ShaderRHI = RHICmdList.GetBoundPixelShader();
-
-		FGlobalShader::SetParameters<FViewUniformShaderParameters>(RHICmdList, ShaderRHI, View.ViewUniformBuffer);
-
-		SetTextureParameter(
-			RHICmdList,
-			ShaderRHI,
-			BentNormalAOTexture,
-			BentNormalAOSampler,
-			TStaticSamplerState<SF_Bilinear,AM_Clamp,AM_Clamp,AM_Clamp>::GetRHI(),
-			BentNormalHistoryTextureValue
-			);
-
-		SetTextureParameter(
-			RHICmdList,
-			ShaderRHI,
-			DistanceFieldNormalTexture,
-			DistanceFieldNormalSampler,
-			TStaticSamplerState<SF_Bilinear,AM_Clamp,AM_Clamp,AM_Clamp>::GetRHI(),
-			DistanceFieldNormal
-			);
-
-		SetShaderValue(RHICmdList, ShaderRHI, HistoryWeight, GAOHistoryWeight);
-
-		const FIntPoint DownsampledBufferSize(GetSceneTextureExtent() / FIntPoint(GAODownsampleFactor, GAODownsampleFactor));
-		const FVector2f BaseLevelTexelSizeValue(1.0f / DownsampledBufferSize.X, 1.0f / DownsampledBufferSize.Y);
-		SetShaderValue(RHICmdList, ShaderRHI, BentNormalAOTexelSize, BaseLevelTexelSizeValue);
-
-		if (bManuallyClampUV)
-		{
-			FVector2f MaxSampleBufferUVValue(
-				(View.ViewRect.Width() / GAODownsampleFactor - 0.5f - GAODownsampleFactor) / DownsampledBufferSize.X,
-				(View.ViewRect.Height() / GAODownsampleFactor - 0.5f - GAODownsampleFactor) / DownsampledBufferSize.Y);
-			SetShaderValue(RHICmdList, ShaderRHI, MaxSampleBufferUV, MaxSampleBufferUVValue);
-		}
 	}
-
-private:
-
-	LAYOUT_FIELD(FShaderResourceParameter, BentNormalAOTexture);
-	LAYOUT_FIELD(FShaderResourceParameter, BentNormalAOSampler);
-	LAYOUT_FIELD(FShaderParameter, HistoryWeight);
-	LAYOUT_FIELD(FShaderParameter, BentNormalAOTexelSize);
-	LAYOUT_FIELD(FShaderParameter, MaxSampleBufferUV);
-	LAYOUT_FIELD(FShaderResourceParameter, DistanceFieldNormalTexture);
-	LAYOUT_FIELD(FShaderResourceParameter, DistanceFieldNormalSampler);
 };
 
-
-#define VARIATION1(A) \
-	typedef TFilterHistoryPS<A> TFilterHistoryPS##A; \
-	IMPLEMENT_SHADER_TYPE(template<>,TFilterHistoryPS##A,TEXT("/Engine/Private/DistanceFieldLightingPost.usf"),TEXT("FilterHistoryPS"),SF_Pixel);
-
-VARIATION1(false)
-VARIATION1(true)
-
-#undef VARIATION1
+IMPLEMENT_GLOBAL_SHADER(FFilterHistoryPS, "/Engine/Private/DistanceFieldLightingPost.usf", "FilterHistoryPS", SF_Pixel);
 
 class FGeometryAwareUpsamplePS : public FGlobalShader
 {
-	DECLARE_SHADER_TYPE(FGeometryAwareUpsamplePS, Global);
 public:
+	DECLARE_GLOBAL_SHADER(FGeometryAwareUpsamplePS);
+	SHADER_USE_PARAMETER_STRUCT(FGeometryAwareUpsamplePS, FGlobalShader);
 
 	BEGIN_SHADER_PARAMETER_STRUCT(FParameters, )
+		SHADER_PARAMETER_STRUCT_REF(FViewUniformShaderParameters, View)
 		SHADER_PARAMETER_STRUCT_INCLUDE(FAOParameters, AOParameters)
-		RDG_TEXTURE_ACCESS(DistanceFieldNormal, ERHIAccess::SRVGraphics)
-		RDG_TEXTURE_ACCESS(BentNormalInterpolation, ERHIAccess::SRVGraphics)
+		SHADER_PARAMETER_STRUCT_INCLUDE(FGeometryAwareUpsampleParameters, GeometryAwareUpsampleParameters)
 		RENDER_TARGET_BINDING_SLOTS()
 	END_SHADER_PARAMETER_STRUCT()
 
@@ -412,32 +209,6 @@ public:
 	{
 		OutEnvironment.SetDefine(TEXT("DOWNSAMPLE_FACTOR"), GAODownsampleFactor);
 	}
-
-	/** Default constructor. */
-	FGeometryAwareUpsamplePS() {}
-
-	/** Initialization constructor. */
-	FGeometryAwareUpsamplePS(const ShaderMetaType::CompiledShaderInitializerType& Initializer)
-		: FGlobalShader(Initializer)
-	{
-		GeometryAwareUpsampleParameters.Bind(Initializer.ParameterMap);
-	}
-
-	void SetParameters(
-		FRHICommandList& RHICmdList,
-		const FViewInfo& View,
-		FRHITexture* DistanceFieldNormal,
-		FRHITexture* DistanceFieldAOBentNormal)
-	{
-		FRHIPixelShader* ShaderRHI = RHICmdList.GetBoundPixelShader();
-
-		FGlobalShader::SetParameters<FViewUniformShaderParameters>(RHICmdList, ShaderRHI, View.ViewUniformBuffer);
-
-		GeometryAwareUpsampleParameters.Set(RHICmdList, ShaderRHI, View, DistanceFieldNormal, DistanceFieldAOBentNormal);
-	}
-
-private:
-	LAYOUT_FIELD(FGeometryAwareUpsampleParameters, GeometryAwareUpsampleParameters);
 };
 
 IMPLEMENT_SHADER_TYPE(,FGeometryAwareUpsamplePS, TEXT("/Engine/Private/DistanceFieldLightingPost.usf"), TEXT("GeometryAwareUpsamplePS"), SF_Pixel);
@@ -456,16 +227,21 @@ void AllocateOrReuseAORenderTarget(FRDGBuilder& GraphBuilder, FRDGTextureRef& Te
 void GeometryAwareUpsample(FRDGBuilder& GraphBuilder, const FViewInfo& View, FRDGTextureRef DistanceFieldAOBentNormal, FRDGTextureRef DistanceFieldNormal, FRDGTextureRef BentNormalInterpolation, const FDistanceFieldAOParameters& Parameters)
 {
 	auto* PassParameters = GraphBuilder.AllocParameters<FGeometryAwareUpsamplePS::FParameters>();
+	PassParameters->View = View.ViewUniformBuffer;
 	PassParameters->AOParameters = DistanceField::SetupAOShaderParameters(Parameters);
-	PassParameters->DistanceFieldNormal = DistanceFieldNormal;
-	PassParameters->BentNormalInterpolation = BentNormalInterpolation;
+	PassParameters->GeometryAwareUpsampleParameters = SetupGeometryAwareUpsampleParameters(View, DistanceFieldNormal, BentNormalInterpolation);
 	PassParameters->RenderTargets[0] = FRenderTargetBinding(DistanceFieldAOBentNormal, ERenderTargetLoadAction::ELoad);
 
+	auto VertexShader = View.ShaderMap->GetShader<FPostProcessVS>();
+	auto PixelShader = View.ShaderMap->GetShader<FGeometryAwareUpsamplePS>();
+
+	ClearUnusedGraphResources(PixelShader, PassParameters);
+
 	GraphBuilder.AddPass(
 		RDG_EVENT_NAME("GeometryAwareUpsample"),
 		PassParameters,
 		ERDGPassFlags::Raster,
-		[&View, DistanceFieldNormal, BentNormalInterpolation](FRHICommandList& RHICmdList)
+		[PassParameters, VertexShader, PixelShader, &View](FRHICommandList& RHICmdList)
 	{
 		RHICmdList.SetViewport(0, 0, 0.0f, View.ViewRect.Width() / GAODownsampleFactor, View.ViewRect.Height() / GAODownsampleFactor, 1.0f);
 
@@ -474,16 +250,14 @@ void GeometryAwareUpsample(FRDGBuilder& GraphBuilder, const FViewInfo& View, FRD
 		GraphicsPSOInit.RasterizerState = TStaticRasterizerState<FM_Solid, CM_None>::GetRHI();
 		GraphicsPSOInit.DepthStencilState = TStaticDepthStencilState<false, CF_Always>::GetRHI();
 		GraphicsPSOInit.BlendState = TStaticBlendState<>::GetRHI();
-
-		TShaderMapRef<FPostProcessVS> VertexShader(View.ShaderMap);
 		GraphicsPSOInit.BoundShaderState.VertexDeclarationRHI = GFilterVertexDeclaration.VertexDeclarationRHI;
 		GraphicsPSOInit.BoundShaderState.VertexShaderRHI = VertexShader.GetVertexShader();
+		GraphicsPSOInit.BoundShaderState.PixelShaderRHI = PixelShader.GetPixelShader();
 		GraphicsPSOInit.PrimitiveType = PT_TriangleList;
 
-		TShaderMapRef<FGeometryAwareUpsamplePS> PixelShader(View.ShaderMap);
-		GraphicsPSOInit.BoundShaderState.PixelShaderRHI = PixelShader.GetPixelShader();
 		SetGraphicsPipelineState(RHICmdList, GraphicsPSOInit, 0);
-		PixelShader->SetParameters(RHICmdList, View, DistanceFieldNormal->GetRHI(), BentNormalInterpolation->GetRHI());
+
+		SetShaderParameters(RHICmdList, PixelShader, PixelShader.GetPixelShader(), *PassParameters);
 
 		DrawRectangle(
 			RHICmdList,
@@ -512,6 +286,8 @@ void UpdateHistory(
 	FRDGTextureRef& BentNormalHistoryOutput,
 	const FDistanceFieldAOParameters& Parameters)
 {
+	const FRDGSystemTextures& SystemTextures = FRDGSystemTextures::Get(GraphBuilder);
+
 	const FIntPoint SceneTextureExtent = GetSceneTextureExtent();
 
 	if (BentNormalHistoryState && GAOUseHistory)
@@ -545,22 +321,54 @@ void UpdateHistory(
 				FIntRect PrevHistoryViewRect = *DistanceFieldAOHistoryViewRect;
 
 				auto* PassParameters = GraphBuilder.AllocParameters<FUpdateHistoryDepthRejectionPS::FParameters>();
+				PassParameters->View = View.ViewUniformBuffer;
 				PassParameters->SceneTextures = SceneTexturesUniformBuffer;
 				PassParameters->AOParameters = DistanceField::SetupAOShaderParameters(Parameters);
-				PassParameters->DistanceFieldNormal = DistanceFieldNormal;
-				PassParameters->DistanceFieldAOBentNormal = BentNormalInterpolation;
+				PassParameters->GeometryAwareUpsampleParameters = SetupGeometryAwareUpsampleParameters(View, DistanceFieldNormal, BentNormalInterpolation);
 				PassParameters->BentNormalHistoryTexture = BentNormalHistoryTexture;
-				PassParameters->VelocityTexture = VelocityTexture;
+				PassParameters->BentNormalHistorySampler = TStaticSamplerState<SF_Bilinear, AM_Clamp, AM_Clamp, AM_Clamp>::GetRHI();
+				PassParameters->VelocityTexture = GetIfProduced(VelocityTexture, SystemTextures.Black);
+				PassParameters->VelocityTextureSampler = TStaticSamplerState<SF_Bilinear, AM_Clamp, AM_Clamp, AM_Clamp>::GetRHI();
+
+				PassParameters->HistoryWeight = GAOHistoryWeight;
+				PassParameters->HistoryDistanceThreshold = GAOHistoryDistanceThreshold;
+				PassParameters->UseHistoryFilter = UseAOHistoryStabilityPass() ? 1.0f : 0.0f;
+
+				{
+					FIntPoint HistoryBufferSize = GetSceneTextureExtent() / FIntPoint(GAODownsampleFactor, GAODownsampleFactor);
+
+					const float InvBufferSizeX = 1.0f / HistoryBufferSize.X;
+					const float InvBufferSizeY = 1.0f / HistoryBufferSize.Y;
+
+					const FVector4f HistoryScreenPositionScaleBiasValue(
+						PrevHistoryViewRect.Width() * InvBufferSizeX / +2.0f,
+						PrevHistoryViewRect.Height() * InvBufferSizeY / (-2.0f * GProjectionSignY),
+						(PrevHistoryViewRect.Height() / 2.0f + PrevHistoryViewRect.Min.Y) * InvBufferSizeY,
+						(PrevHistoryViewRect.Width() / 2.0f + PrevHistoryViewRect.Min.X) * InvBufferSizeX);
+
+					// Pull in the max UV to exclude the region which will read outside the viewport due to bilinear filtering
+					const FVector4f HistoryUVMinMaxValue(
+						(PrevHistoryViewRect.Min.X + 0.5f) * InvBufferSizeX,
+						(PrevHistoryViewRect.Min.Y + 0.5f) * InvBufferSizeY,
+						(PrevHistoryViewRect.Max.X - 0.5f) * InvBufferSizeX,
+						(PrevHistoryViewRect.Max.Y - 0.5f) * InvBufferSizeY);
+
+					PassParameters->HistoryScreenPositionScaleBias = HistoryScreenPositionScaleBiasValue;
+					PassParameters->HistoryUVMinMax = HistoryUVMinMaxValue;
+				}
+
 				PassParameters->RenderTargets[0] = FRenderTargetBinding(NewBentNormalHistory, ERenderTargetLoadAction::ELoad);
 
 				TShaderMapRef<FPostProcessVS> VertexShader(View.ShaderMap);
 				TShaderMapRef<FUpdateHistoryDepthRejectionPS> PixelShader(View.ShaderMap);
 
+				ClearUnusedGraphResources(PixelShader, PassParameters);
+
 				GraphBuilder.AddPass(
 					RDG_EVENT_NAME("UpdateHistory"),
 					PassParameters,
 					ERDGPassFlags::Raster,
-					[VertexShader, PixelShader, &View, DistanceFieldNormal, BentNormalInterpolation, BentNormalHistoryTexture, VelocityTexture, PrevHistoryViewRect, SceneTextureExtent]
+					[PassParameters, VertexShader, PixelShader, &View, SceneTextureExtent]
 					(FRHICommandList& RHICmdList)
 				{
 					RHICmdList.SetViewport(0, 0, 0.0f, View.ViewRect.Width() / GAODownsampleFactor, View.ViewRect.Height() / GAODownsampleFactor, 1.0f);
@@ -577,11 +385,7 @@ void UpdateHistory(
 
 					SetGraphicsPipelineState(RHICmdList, GraphicsPSOInit, 0);
 
-					PixelShader->SetParameters(RHICmdList, View, PrevHistoryViewRect,
-						DistanceFieldNormal->GetRHI(),
-						BentNormalInterpolation->GetRHI(),
-						BentNormalHistoryTexture->GetRHI(),
-						VelocityTexture ? VelocityTexture->GetRHI() : nullptr);
+					SetShaderParameters(RHICmdList, PixelShader, PixelShader.GetPixelShader(), *PassParameters);
 
 					DrawRectangle(
 						RHICmdList,
@@ -608,22 +412,41 @@ void UpdateHistory(
 					AllocateOrReuseAORenderTarget(GraphBuilder, BentNormalHistoryTexture, BentNormalHistoryRTName, PF_FloatRGBA);
 				}
 
-				auto* PassParameters = GraphBuilder.AllocParameters<FUpdateHistoryDepthRejectionPS::FParameters>();
+				const bool bManuallyClampUV = View.ViewRect.Min != FIntPoint::ZeroValue || View.ViewRect.Max != SceneTextureExtent;
+
+				FFilterHistoryPS::FPermutationDomain PermutationVector;
+				PermutationVector.Set<FFilterHistoryPS::FManuallyClampUV>(bManuallyClampUV);
+
+				auto VertexShader = View.ShaderMap->GetShader<FPostProcessVS>();
+				auto PixelShader = View.ShaderMap->GetShader<FFilterHistoryPS>(PermutationVector);
+
+				const FIntPoint DownsampledBufferSize(GetSceneTextureExtent() / FIntPoint(GAODownsampleFactor, GAODownsampleFactor));
+				FVector2f MaxSampleBufferUV(
+					(View.ViewRect.Width() / GAODownsampleFactor - 0.5f - GAODownsampleFactor) / DownsampledBufferSize.X,
+					(View.ViewRect.Height() / GAODownsampleFactor - 0.5f - GAODownsampleFactor) / DownsampledBufferSize.Y);
+
+				auto* PassParameters = GraphBuilder.AllocParameters<FFilterHistoryPS::FParameters>();
+				PassParameters->View = View.ViewUniformBuffer;
 				PassParameters->SceneTextures = SceneTexturesUniformBuffer;
-				PassParameters->DistanceFieldNormal = DistanceFieldNormal;
-				PassParameters->BentNormalHistoryTexture = NewBentNormalHistory;
+				PassParameters->DistanceFieldNormalTexture = DistanceFieldNormal;
+				PassParameters->BentNormalAOTexture = NewBentNormalHistory;
+				PassParameters->DistanceFieldNormalSampler = TStaticSamplerState<SF_Bilinear, AM_Clamp, AM_Clamp, AM_Clamp>::GetRHI();
+				PassParameters->BentNormalAOSampler = TStaticSamplerState<SF_Bilinear, AM_Clamp, AM_Clamp, AM_Clamp>::GetRHI();
+				PassParameters->BentNormalAOTexelSize = FVector2f(1.0f / DownsampledBufferSize.X, 1.0f / DownsampledBufferSize.Y);
+				PassParameters->MaxSampleBufferUV = MaxSampleBufferUV;
+				PassParameters->HistoryWeight = GAOHistoryWeight;
 				PassParameters->RenderTargets[0] = FRenderTargetBinding(BentNormalHistoryTexture, ERenderTargetLoadAction::ELoad);
 
+				ClearUnusedGraphResources(PixelShader, PassParameters);
+
 				GraphBuilder.AddPass(
 					RDG_EVENT_NAME("UpdateHistoryStability"),
 					PassParameters,
 					ERDGPassFlags::Raster,
-					[&View, DistanceFieldNormal, NewBentNormalHistory, SceneTextureExtent](FRHICommandList& RHICmdList)
+					[PassParameters, VertexShader, PixelShader, &View, SceneTextureExtent](FRHICommandList& RHICmdList)
 				{
 					RHICmdList.SetViewport(0, 0, 0.0f, View.ViewRect.Width() / GAODownsampleFactor, View.ViewRect.Height() / GAODownsampleFactor, 1.0f);
 
-					TShaderMapRef<FPostProcessVS> VertexShader(View.ShaderMap);
-
 					FGraphicsPipelineStateInitializer GraphicsPSOInit;
 					RHICmdList.ApplyCachedRenderTargets(GraphicsPSOInit);
 					GraphicsPSOInit.RasterizerState = TStaticRasterizerState<FM_Solid, CM_None>::GetRHI();
@@ -631,24 +454,12 @@ void UpdateHistory(
 					GraphicsPSOInit.BlendState = TStaticBlendState<>::GetRHI();
 					GraphicsPSOInit.BoundShaderState.VertexDeclarationRHI = GFilterVertexDeclaration.VertexDeclarationRHI;
 					GraphicsPSOInit.BoundShaderState.VertexShaderRHI = VertexShader.GetVertexShader();
+					GraphicsPSOInit.BoundShaderState.PixelShaderRHI = PixelShader.GetPixelShader();
 					GraphicsPSOInit.PrimitiveType = PT_TriangleList;
 
-					if (View.ViewRect.Min == FIntPoint::ZeroValue && View.ViewRect.Max == SceneTextureExtent)
-					{
-						TShaderMapRef<TFilterHistoryPS<false> > PixelShader(View.ShaderMap);
-
-						GraphicsPSOInit.BoundShaderState.PixelShaderRHI = PixelShader.GetPixelShader();
-						SetGraphicsPipelineState(RHICmdList, GraphicsPSOInit, 0);
-						PixelShader->SetParameters(RHICmdList, View, DistanceFieldNormal->GetRHI(), NewBentNormalHistory->GetRHI());
-					}
-					else
-					{
-						TShaderMapRef<TFilterHistoryPS<true> > PixelShader(View.ShaderMap);
+					SetGraphicsPipelineState(RHICmdList, GraphicsPSOInit, 0);
 
-						GraphicsPSOInit.BoundShaderState.PixelShaderRHI = PixelShader.GetPixelShader();
-						SetGraphicsPipelineState(RHICmdList, GraphicsPSOInit, 0);
-						PixelShader->SetParameters(RHICmdList, View, DistanceFieldNormal->GetRHI(), NewBentNormalHistory->GetRHI());
-					}
+					SetShaderParameters(RHICmdList, PixelShader, PixelShader.GetPixelShader(), *PassParameters);
 
 					DrawRectangle(
 						RHICmdList,
@@ -708,10 +519,13 @@ class FDistanceFieldAOUpsamplePS : public FGlobalShader
 {
 public:
 	DECLARE_GLOBAL_SHADER(FDistanceFieldAOUpsamplePS);
+	SHADER_USE_PARAMETER_STRUCT(FDistanceFieldAOUpsamplePS, FGlobalShader);
 
 	BEGIN_SHADER_PARAMETER_STRUCT(FParameters, )
 		SHADER_PARAMETER_STRUCT_REF(FViewUniformShaderParameters, View)
 		SHADER_PARAMETER_RDG_UNIFORM_BUFFER(FSceneTextureUniformParameters, SceneTextures)
+		SHADER_PARAMETER_STRUCT_INCLUDE(FDFAOUpsampleParameters, DFAOUpsampleParameters)
+		SHADER_PARAMETER(float, MinIndirectDiffuseOcclusion)
 		RDG_TEXTURE_ACCESS(DistanceFieldAOBentNormal, ERHIAccess::SRVGraphics)
 		RENDER_TARGET_BINDING_SLOTS()
 	END_SHADER_PARAMETER_STRUCT()
@@ -723,30 +537,6 @@ public:
 	{
 		return ShouldCompileDistanceFieldShaders(Parameters.Platform);
 	}
-
-	FDistanceFieldAOUpsamplePS() = default;
-	FDistanceFieldAOUpsamplePS(const ShaderMetaType::CompiledShaderInitializerType& Initializer)
-		: FGlobalShader(Initializer)
-	{
-		BindForLegacyShaderParameters<FParameters>(this, Initializer.PermutationId, Initializer.ParameterMap, false);
-		DFAOUpsampleParameters.Bind(Initializer.ParameterMap);
-		MinIndirectDiffuseOcclusion.Bind(Initializer.ParameterMap,TEXT("MinIndirectDiffuseOcclusion"));
-	}
-
-	void SetParameters(FRHICommandList& RHICmdList, const FViewInfo& View, FRHITexture* DistanceFieldAOBentNormal)
-	{
-		FRHIPixelShader* ShaderRHI = RHICmdList.GetBoundPixelShader();
-
-		DFAOUpsampleParameters.Set(RHICmdList, ShaderRHI, View, DistanceFieldAOBentNormal);
-
-		FScene* Scene = (FScene*)View.Family->Scene;
-		const float MinOcclusion = Scene->SkyLight ? Scene->SkyLight->MinOcclusion : 0;
-		SetShaderValue(RHICmdList, ShaderRHI, MinIndirectDiffuseOcclusion, MinOcclusion);
-	}
-	
-private:
-	LAYOUT_FIELD(FDFAOUpsampleParameters, DFAOUpsampleParameters);
-	LAYOUT_FIELD(FShaderParameter, MinIndirectDiffuseOcclusion);
 };
 
 IMPLEMENT_GLOBAL_SHADER(FDistanceFieldAOUpsamplePS, "/Engine/Private/DistanceFieldLightingPost.usf", "AOUpsamplePS", SF_Pixel);
@@ -762,11 +552,15 @@ void UpsampleBentNormalAO(
 	for (int32 ViewIndex = 0; ViewIndex < Views.Num(); ViewIndex++)
 	{
 		const FViewInfo& View = Views[ViewIndex];
+		FScene* Scene = (FScene*)View.Family->Scene;
+
 		RDG_GPU_MASK_SCOPE(GraphBuilder, View.GPUMask);
 
 		auto* PassParameters = GraphBuilder.AllocParameters<FDistanceFieldAOUpsamplePS::FParameters>();
 		PassParameters->View = View.ViewUniformBuffer;
 		PassParameters->SceneTextures = SceneTexturesUniformBuffer;
+		PassParameters->DFAOUpsampleParameters = DistanceField::SetupAOUpsampleParameters(View, DistanceFieldAOBentNormal ? DistanceFieldAOBentNormal : GSystemTextures.GetWhiteDummy(GraphBuilder));
+		PassParameters->MinIndirectDiffuseOcclusion = Scene->SkyLight ? Scene->SkyLight->MinOcclusion : 0;
 		PassParameters->RenderTargets[0] = FRenderTargetBinding(SceneColorTexture, ERenderTargetLoadAction::ELoad);
 		PassParameters->DistanceFieldAOBentNormal = DistanceFieldAOBentNormal;
 
@@ -803,7 +597,6 @@ void UpsampleBentNormalAO(
 			GraphicsPSOInit.BoundShaderState.VertexShaderRHI = VertexShader.GetVertexShader();
 			GraphicsPSOInit.BoundShaderState.PixelShaderRHI = PixelShader.GetPixelShader();
 			SetGraphicsPipelineState(RHICmdList, GraphicsPSOInit, 0);
-			PixelShader->SetParameters(RHICmdList, View, DistanceFieldAOBentNormal->GetRHI());
 			SetShaderParameters(RHICmdList, PixelShader, PixelShader.GetPixelShader(), *PassParameters);
 
 			DrawRectangle(
