Remove deprecated DFGI and Heightfield GI

diff --git a/Engine/Shaders/Private/HeightfieldLighting.usf b/Engine/Shaders/Private/HeightfieldLighting.usf
deleted file mode 100644
diff --git a/Engine/Source/Runtime/Renderer/Private/DistanceFieldObjectManagement.cpp b/Engine/Source/Runtime/Renderer/Private/DistanceFieldObjectManagement.cpp
--- a/Engine/Source/Runtime/Renderer/Private/DistanceFieldObjectManagement.cpp
+++ b/Engine/Source/Runtime/Renderer/Private/DistanceFieldObjectManagement.cpp
@@ -1041,8 +1041,6 @@ void FSceneRenderer::PrepareDistanceFieldScene(FRDGBuilder& GraphBuilder, bool b
 
 			RDG_GPU_MASK_SCOPE(GraphBuilder, View.GPUMask);
 
-			View.HeightfieldLightingViewInfo.SetupVisibleHeightfields(View, GraphBuilder);
-
 			if (ShouldPrepareGlobalDistanceField())
 			{
 				float OcclusionMaxDistance = Scene->DefaultMaxDistanceFieldOcclusionDistance;
diff --git a/Engine/Source/Runtime/Renderer/Private/DistanceFieldScreenGridLighting.cpp b/Engine/Source/Runtime/Renderer/Private/DistanceFieldScreenGridLighting.cpp
--- a/Engine/Source/Runtime/Renderer/Private/DistanceFieldScreenGridLighting.cpp
+++ b/Engine/Source/Runtime/Renderer/Private/DistanceFieldScreenGridLighting.cpp
@@ -514,9 +514,6 @@ void FDeferredShadingSceneRenderer::RenderDistanceFieldAOScreenGrid(
 			});
 	}
 
-	// Compute heightfield occlusion after heightfield GI, otherwise it self-shadows incorrectly
-	View.HeightfieldLightingViewInfo.ComputeOcclusionForScreenGrid(GraphBuilder, View, SceneTextures, DistanceFieldNormal, AOScreenGridParameters, Parameters);
-
 	FRDGTextureRef DownsampledBentNormal = nullptr;
 
 	{
diff --git a/Engine/Source/Runtime/Renderer/Private/HeightfieldLighting.cpp b/Engine/Source/Runtime/Renderer/Private/HeightfieldLighting.cpp
--- a/Engine/Source/Runtime/Renderer/Private/HeightfieldLighting.cpp
+++ b/Engine/Source/Runtime/Renderer/Private/HeightfieldLighting.cpp
@@ -17,595 +17,15 @@ HeightfieldLighting.cpp
 #include "LightRendering.h"
 #include "PipelineStateCache.h"
 
-IMPLEMENT_TYPE_LAYOUT(FHeightfieldDescriptionParameters);
-IMPLEMENT_TYPE_LAYOUT(FHeightfieldTextureParameters);
-
-// Currently disabled because the bHasHeightfieldRepresentation GBuffer bit has been reallocated, and self-shadowing artifacts are too severe without that bit
-int32 GAOHeightfieldOcclusion = 0;
-FAutoConsoleVariableRef CVarAOHeightfieldOcclusion(
-	TEXT("r.AOHeightfieldOcclusion"),
-	GAOHeightfieldOcclusion,
-	TEXT("Whether to compute AO from heightfields (landscape)"),
-	ECVF_RenderThreadSafe
-);
-
-int32 GHeightfieldGlobalIllumination = 1;
-FAutoConsoleVariableRef CVarHeightfieldGlobalIllumination(
-	TEXT("r.HeightfieldGlobalIllumination"),
-	GHeightfieldGlobalIllumination,
-	TEXT(""),
-	ECVF_RenderThreadSafe
-);
-
-float GHeightfieldInnerBounceDistance = 3000;
-FAutoConsoleVariableRef CVarHeightfieldInnerBounceDistancer(
-	TEXT("r.HeightfieldInnerBounceDistance"),
-	GHeightfieldInnerBounceDistance,
-	TEXT(""),
-	ECVF_RenderThreadSafe
-);
-
-float GHeightfieldOuterBounceDistanceScale = 3;
-FAutoConsoleVariableRef CVarHeightfieldOuterBounceDistanceScale(
-	TEXT("r.HeightfieldOuterBounceDistanceScale"),
-	GHeightfieldOuterBounceDistanceScale,
-	TEXT(""),
-	ECVF_RenderThreadSafe
-);
-
-float GetGHeightfieldBounceDistance()
-{
-	return GHeightfieldInnerBounceDistance * GHeightfieldOuterBounceDistanceScale;
-}
-
-float GHeightfieldTargetUnitsPerTexel = 200;
-FAutoConsoleVariableRef CVarHeightfieldTargetUnitsPerTexel(
-	TEXT("r.HeightfieldTargetUnitsPerTexel"),
-	GHeightfieldTargetUnitsPerTexel,
-	TEXT(""),
-	ECVF_RenderThreadSafe
-);
-
-void FHeightfieldLightingAtlas::InitDynamicRHI()
-{
-	FRHICommandListImmediate& RHICmdList = FRHICommandListExecutor::GetImmediateCommandList();
-	if (AtlasSize.GetMin() > 0)
-	{
-		FPooledRenderTargetDesc Desc(FPooledRenderTargetDesc::Create2DDesc(
-			AtlasSize,
-			PF_G16,
-			FClearValueBinding::Transparent,
-			TexCreate_None,
-			TexCreate_RenderTargetable,
-			false));
-
-		GRenderTargetPool.FindFreeElement(RHICmdList, Desc, Height, TEXT("HeightAtlas"));
-
-		FPooledRenderTargetDesc Desc2(FPooledRenderTargetDesc::Create2DDesc(
-			AtlasSize,
-			PF_R8G8,
-			FClearValueBinding::Transparent,
-			TexCreate_None,
-			TexCreate_RenderTargetable,
-			false));
-
-		GRenderTargetPool.FindFreeElement(RHICmdList, Desc2, Normal, TEXT("NormalAtlas"));
-
-		FPooledRenderTargetDesc Desc3(FPooledRenderTargetDesc::Create2DDesc(
-			AtlasSize,
-			PF_R8G8B8A8,
-			FClearValueBinding::Transparent,
-			TexCreate_None,
-			TexCreate_RenderTargetable,
-			false));
-
-		GRenderTargetPool.FindFreeElement(RHICmdList, Desc3, DiffuseColor, TEXT("DiffuseColorAtlas"));
-
-		FPooledRenderTargetDesc Desc4(FPooledRenderTargetDesc::Create2DDesc(
-			AtlasSize,
-			PF_G8,
-			FClearValueBinding::White,
-			TexCreate_None,
-			TexCreate_RenderTargetable,
-			false));
-
-		GRenderTargetPool.FindFreeElement(RHICmdList, Desc4, DirectionalLightShadowing, TEXT("HeightfieldShadowingAtlas"));
-
-		FPooledRenderTargetDesc Desc5(FPooledRenderTargetDesc::Create2DDesc(
-			AtlasSize,
-			PF_FloatR11G11B10,
-			FClearValueBinding::Black,
-			TexCreate_None,
-			TexCreate_RenderTargetable,
-			false));
-		Desc5.AutoWritable = false;
-
-		GRenderTargetPool.FindFreeElement(RHICmdList, Desc5, Lighting, TEXT("HeightfieldLightingAtlas"));
-	}
-}
-
-void FHeightfieldLightingAtlas::ReleaseDynamicRHI()
-{
-	GRenderTargetPool.FreeUnusedResource(Height);
-	GRenderTargetPool.FreeUnusedResource(Normal);
-	GRenderTargetPool.FreeUnusedResource(DiffuseColor);
-	GRenderTargetPool.FreeUnusedResource(DirectionalLightShadowing);
-	GRenderTargetPool.FreeUnusedResource(Lighting);
-}
-
-void FHeightfieldLightingAtlas::InitializeForSize(FIntPoint InAtlasSize)
-{
-	if (InAtlasSize.X > AtlasSize.X || InAtlasSize.Y > AtlasSize.Y)
-	{
-		AtlasSize.X = FMath::Max(InAtlasSize.X, AtlasSize.X);
-		AtlasSize.Y = FMath::Max(InAtlasSize.Y, AtlasSize.Y);
-
-		if (IsInitialized())
-		{
-			UpdateRHI();
-		}
-		else
-		{
-			InitResource();
-		}
-	}
-}
-
-class FSubsectionHeightfieldDescriptionsResource : public FRenderResource
-{
-public:
-
-	FCPUUpdatedBuffer Data;
-
-	FSubsectionHeightfieldDescriptionsResource()
-	{
-		Data.Format = PF_A32B32G32R32F;
-		// In float4's, must match usf
-		Data.Stride = 4;
-	}
-
-	virtual void InitDynamicRHI()  override
-	{
-		Data.Initialize();
-	}
-
-	virtual void ReleaseDynamicRHI() override
-	{
-		Data.Release();
-	}
-};
-
-TGlobalResource<FSubsectionHeightfieldDescriptionsResource> GSubsectionHeightfieldDescriptions;
-
-class FSubsectionHeightfieldDescriptionParameters
-{
-	DECLARE_INLINE_TYPE_LAYOUT(FSubsectionHeightfieldDescriptionParameters, NonVirtual);
-public:
-	void Bind(const FShaderParameterMap& ParameterMap)
-	{
-		SubsectionHeightfieldDescriptions.Bind(ParameterMap, TEXT("SubsectionHeightfieldDescriptions"));
-	}
-
-	friend FArchive& operator<<(FArchive& Ar, FSubsectionHeightfieldDescriptionParameters& Parameters)
-	{
-		Ar << Parameters.SubsectionHeightfieldDescriptions;
-		return Ar;
-	}
-
-	template<typename ShaderRHIParamRef>
-	void Set(FRHICommandList& RHICmdList, const ShaderRHIParamRef ShaderRHI)
-	{
-		SetSRVParameter(RHICmdList, ShaderRHI, SubsectionHeightfieldDescriptions, GSubsectionHeightfieldDescriptions.Data.BufferSRV);
-	}
-
-private:
-	
-		LAYOUT_FIELD(FShaderResourceParameter, SubsectionHeightfieldDescriptions)
-	
-};
-
-int32 UploadSubsectionHeightfieldDescriptions(const TArray<FHeightfieldComponentDescription>& HeightfieldDescriptions, FVector2D InvLightingAtlasSize, float InvDownsampleFactor)
-{
-	TArray<FVector4f, SceneRenderingAllocator> HeightfieldDescriptionData;
-	HeightfieldDescriptionData.Empty(HeightfieldDescriptions.Num() * GSubsectionHeightfieldDescriptions.Data.Stride);
-
-	for (int32 DescriptionIndex = 0; DescriptionIndex < HeightfieldDescriptions.Num(); DescriptionIndex++)
-	{
-		const FHeightfieldComponentDescription& Description = HeightfieldDescriptions[DescriptionIndex];
-
-		for (int32 SubsectionY = 0; SubsectionY < Description.NumSubsections; SubsectionY++)
-		{
-			for (int32 SubsectionX = 0; SubsectionX < Description.NumSubsections; SubsectionX++)
-			{
-				HeightfieldDescriptionData.Add(FVector4f(SubsectionX, SubsectionY, 0, 0));
-				HeightfieldDescriptionData.Add(Description.SubsectionScaleAndBias);
-				HeightfieldDescriptionData.Add(Description.HeightfieldScaleBias);
-
-				// GlobalUVScaleBias = SubsectionSizeQuads / AtlasSize, (SubsectionBase + SubsectionId * SubsectionSizeQuads - AtlasMin) / AtlasSize
-				const FVector4f GlobalUVScaleBias(
-					Description.SubsectionScaleAndBias.X * InvLightingAtlasSize.X * InvDownsampleFactor,
-					Description.SubsectionScaleAndBias.Y * InvLightingAtlasSize.Y * InvDownsampleFactor,
-					(Description.LightingAtlasLocation.X + SubsectionX * Description.SubsectionScaleAndBias.X * InvDownsampleFactor) * InvLightingAtlasSize.X,
-					(Description.LightingAtlasLocation.Y + SubsectionY * Description.SubsectionScaleAndBias.Y * InvDownsampleFactor) * InvLightingAtlasSize.Y);
-
-				HeightfieldDescriptionData.Add(GlobalUVScaleBias);
-			}
-		}
-	}
-
-	check(HeightfieldDescriptionData.Num() % GSubsectionHeightfieldDescriptions.Data.Stride == 0);
-
-	if (HeightfieldDescriptionData.Num() > GSubsectionHeightfieldDescriptions.Data.MaxElements)
-	{
-		GSubsectionHeightfieldDescriptions.Data.MaxElements = HeightfieldDescriptionData.Num() * 5 / 4;
-		GSubsectionHeightfieldDescriptions.Data.Release();
-		GSubsectionHeightfieldDescriptions.Data.Initialize();
-	}
-
-	void* LockedBuffer = RHILockBuffer(GSubsectionHeightfieldDescriptions.Data.Buffer, 0, GSubsectionHeightfieldDescriptions.Data.Buffer->GetSize(), RLM_WriteOnly);
-	const uint32 MemcpySize = HeightfieldDescriptionData.GetTypeSize() * HeightfieldDescriptionData.Num();
-	check(GSubsectionHeightfieldDescriptions.Data.Buffer->GetSize() >= MemcpySize);
-	FPlatformMemory::Memcpy(LockedBuffer, HeightfieldDescriptionData.GetData(), MemcpySize);
-	RHIUnlockBuffer(GSubsectionHeightfieldDescriptions.Data.Buffer);
-
-	return HeightfieldDescriptionData.Num() / GSubsectionHeightfieldDescriptions.Data.Stride;
-}
-
-bool DoesPlatformSupportDistanceFieldGI(EShaderPlatform Platform)
-{
-	//@todo - remove remaining uses
-	return false;
-}
-
-class FHeightfieldSubsectionQuadVS : public FGlobalShader
-{
-	DECLARE_SHADER_TYPE(FHeightfieldSubsectionQuadVS, Global);
-public:
-
-	static bool ShouldCompilePermutation(const FGlobalShaderPermutationParameters& Parameters)
-	{
-		return IsFeatureLevelSupported(Parameters.Platform, ERHIFeatureLevel::SM5) && DoesPlatformSupportDistanceFieldGI(Parameters.Platform);
-	}
-
-	/** Default constructor. */
-	FHeightfieldSubsectionQuadVS() {}
-
-	/** Initialization constructor. */
-	FHeightfieldSubsectionQuadVS(const ShaderMetaType::CompiledShaderInitializerType& Initializer)
-		: FGlobalShader(Initializer)
-	{
-		SubsectionHeightfieldParameters.Bind(Initializer.ParameterMap);
-	}
-
-	void SetParameters(FRHICommandList& RHICmdList, const FSceneView& View)
-	{
-		FRHIVertexShader* ShaderRHI = RHICmdList.GetBoundVertexShader();
-
-		FGlobalShader::SetParameters<FViewUniformShaderParameters>(RHICmdList, ShaderRHI, View.ViewUniformBuffer);
-		SubsectionHeightfieldParameters.Set(RHICmdList, ShaderRHI);
-	}
-
-private:
-
-	LAYOUT_FIELD(FSubsectionHeightfieldDescriptionParameters, SubsectionHeightfieldParameters);
-};
-
-IMPLEMENT_SHADER_TYPE(, FHeightfieldSubsectionQuadVS, TEXT("/Engine/Private/HeightfieldLighting.usf"), TEXT("HeightfieldSubsectionQuadVS"), SF_Vertex);
-
-class FInitializeHeightfieldsPS : public FGlobalShader
-{
-	DECLARE_SHADER_TYPE(FInitializeHeightfieldsPS, Global);
-public:
-
-	static bool ShouldCompilePermutation(const FGlobalShaderPermutationParameters& Parameters)
-	{
-		return IsFeatureLevelSupported(Parameters.Platform, ERHIFeatureLevel::SM5) && DoesPlatformSupportDistanceFieldGI(Parameters.Platform);
-	}
-
-	/** Default constructor. */
-	FInitializeHeightfieldsPS() {}
-
-	/** Initialization constructor. */
-	FInitializeHeightfieldsPS(const ShaderMetaType::CompiledShaderInitializerType& Initializer)
-		: FGlobalShader(Initializer)
-	{
-		HeightfieldTextureParameters.Bind(Initializer.ParameterMap);
-	}
-
-	void SetParameters(FRHICommandList& RHICmdList, const FSceneView& View, UTexture2D* HeightfieldTextureValue, UTexture2D* DiffuseColorTextureValue)
-	{
-		FRHIPixelShader* ShaderRHI = RHICmdList.GetBoundPixelShader();
-
-		FGlobalShader::SetParameters<FViewUniformShaderParameters>(RHICmdList, ShaderRHI, View.ViewUniformBuffer);
-		HeightfieldTextureParameters.Set(RHICmdList, ShaderRHI, HeightfieldTextureValue, DiffuseColorTextureValue, nullptr);
-	}
-
-private:
-
-	LAYOUT_FIELD(FHeightfieldTextureParameters, HeightfieldTextureParameters);
-};
-
-IMPLEMENT_SHADER_TYPE(, FInitializeHeightfieldsPS, TEXT("/Engine/Private/HeightfieldLighting.usf"), TEXT("InitializeHeightfieldsPS"), SF_Pixel);
-
-class FQuadVertexBuffer : public FVertexBuffer
-{
-public:
-
-	FQuadVertexBuffer()
-	{
-	}
-
-	virtual void InitRHI() override
-	{
-		const uint32 Size = 6 * sizeof(FScreenVertex);
-		FRHIResourceCreateInfo CreateInfo(TEXT("FQuadVertexBuffer"));
-		VertexBufferRHI = RHICreateBuffer(Size, BUF_Static | BUF_VertexBuffer, 0, ERHIAccess::VertexOrIndexBuffer, CreateInfo);
-		FScreenVertex* DestVertex = (FScreenVertex*)RHILockBuffer(VertexBufferRHI, 0, Size, RLM_WriteOnly);
-
-		DestVertex[0].Position = FVector2f(1, 1);
-		DestVertex[0].UV = FVector2f(1, 1);
-
-		DestVertex[1].Position = FVector2f(0, 1);
-		DestVertex[1].UV = FVector2f(0, 1);
-
-		DestVertex[2].Position = FVector2f(1, 0);
-		DestVertex[2].UV = FVector2f(1, 0);
-
-		DestVertex[3].Position = FVector2f(1, 0);
-		DestVertex[3].UV = FVector2f(1, 0);
-
-		DestVertex[4].Position = FVector2f(0, 1);
-		DestVertex[4].UV = FVector2f(0, 1);
-
-		DestVertex[5].Position = FVector2f(0, 0);
-		DestVertex[5].UV = FVector2f(0, 0);
-
-		RHIUnlockBuffer(VertexBufferRHI);
-	}
-};
-
-TGlobalResource<FQuadVertexBuffer> GQuadVertexBuffer;
-
-bool SupportsHeightfieldLighting(ERHIFeatureLevel::Type FeatureLevel, EShaderPlatform ShaderPlatform)
-{
-	return FeatureLevel >= ERHIFeatureLevel::SM5
-		&& DoesPlatformSupportDistanceFieldGI(ShaderPlatform);
-}
-
-bool AllowHeightfieldGI(const FViewInfo& View)
-{
-	return GHeightfieldGlobalIllumination
-		&& View.State
-		//@todo - remove heightfield GI
-		&& false;
-}
-
-void FHeightfieldLightingViewInfo::SetupVisibleHeightfields(const FViewInfo& View, FRDGBuilder& GraphBuilder)
-{
-	const FScene* Scene = (const FScene*)View.Family->Scene;
-	const int32 NumPrimitives = Scene->DistanceFieldSceneData.HeightfieldPrimitives.Num();
-
-	if ((!AllowHeightfieldGI(View) && !GAOHeightfieldOcclusion)
-		|| NumPrimitives == 0
-		|| !SupportsHeightfieldLighting(View.GetFeatureLevel(), View.GetShaderPlatform()))
-	{
-		return;
-	}
-
-	AddPass(GraphBuilder, RDG_EVENT_NAME("SetupVisibleHeightfields"), [this, Scene, &View, NumPrimitives](FRHICommandListImmediate& RHICmdList)
-	{
-		const float MaxDistanceSquared = FMath::Square(GetMaxAOViewDistance() + GetGHeightfieldBounceDistance());
-		float LocalToWorldScale = 1;
-
-		for (int32 HeightfieldPrimitiveIndex = 0; HeightfieldPrimitiveIndex < NumPrimitives; HeightfieldPrimitiveIndex++)
-		{
-			const FPrimitiveSceneInfo* HeightfieldPrimitive = Scene->DistanceFieldSceneData.HeightfieldPrimitives[HeightfieldPrimitiveIndex];
-			const FBoxSphereBounds& PrimitiveBounds = HeightfieldPrimitive->Proxy->GetBounds();
-			const float DistanceToPrimitiveSq = (PrimitiveBounds.Origin - View.ViewMatrices.GetViewOrigin()).SizeSquared();
-
-			if (View.ViewFrustum.IntersectSphere(PrimitiveBounds.Origin, PrimitiveBounds.SphereRadius + GetGHeightfieldBounceDistance())
-				&& DistanceToPrimitiveSq < MaxDistanceSquared)
-			{
-				UTexture2D* HeightfieldTexture = NULL;
-				UTexture2D* DiffuseColorTexture = NULL;
-				UTexture2D* VisibilityTexture = NULL;
-
-				FHeightfieldComponentDescription NewComponentDescription(HeightfieldPrimitive->Proxy->GetLocalToWorld());
-				HeightfieldPrimitive->Proxy->GetHeightfieldRepresentation(HeightfieldTexture, DiffuseColorTexture, VisibilityTexture, NewComponentDescription);
-
-				if (HeightfieldTexture && HeightfieldTexture->GetResource()->TextureRHI)
-				{
-					const FIntPoint HeightfieldSize = NewComponentDescription.HeightfieldRect.Size();
-
-					if (Heightfield.Rect.Area() == 0)
-					{
-						Heightfield.Rect = NewComponentDescription.HeightfieldRect;
-						LocalToWorldScale = NewComponentDescription.LocalToWorld.GetScaleVector().X;
-					}
-					else
-					{
-						Heightfield.Rect.Union(NewComponentDescription.HeightfieldRect);
-					}
-
-					TArray<FHeightfieldComponentDescription>& ComponentDescriptions = Heightfield.ComponentDescriptions.FindOrAdd(FHeightfieldComponentTextures(HeightfieldTexture, DiffuseColorTexture, VisibilityTexture));
-					ComponentDescriptions.Add(NewComponentDescription);
-				}
-			}
-		}
-
-		if (AllowHeightfieldGI(View) && Heightfield.ComponentDescriptions.Num() > 0)
-		{
-			FSceneViewState* ViewState = (FSceneViewState*)View.State;
-
-			{
-				FHeightfieldLightingAtlas*& ExistingAtlas = ViewState->HeightfieldLightingAtlas;
-
-				if (!ExistingAtlas)
-				{
-					ExistingAtlas = new FHeightfieldLightingAtlas();
-				}
-
-				Heightfield.DownsampleFactor = FMath::Max(FMath::TruncToInt(GHeightfieldTargetUnitsPerTexel / LocalToWorldScale), 1);
-				Heightfield.DownsampledRect = Heightfield.Rect / Heightfield.DownsampleFactor;
-				Heightfield.Rect.Min = FIntPoint::DivideAndRoundDown(Heightfield.Rect.Min, Heightfield.DownsampleFactor) * Heightfield.DownsampleFactor;
-				ExistingAtlas->InitializeForSize(Heightfield.DownsampledRect.Size());
-
-				for (TMap<FHeightfieldComponentTextures, TArray<FHeightfieldComponentDescription>>::TIterator It(Heightfield.ComponentDescriptions); It; ++It)
-				{
-					TArray<FHeightfieldComponentDescription>& HeightfieldDescriptions = It.Value();
-
-					for (int32 ComponentIndex = 0; ComponentIndex < HeightfieldDescriptions.Num(); ComponentIndex++)
-					{
-						FIntPoint RelativeAtlasOffset = HeightfieldDescriptions[ComponentIndex].HeightfieldRect.Min - Heightfield.Rect.Min;
-						HeightfieldDescriptions[ComponentIndex].LightingAtlasLocation = FVector2f(RelativeAtlasOffset.X, RelativeAtlasOffset.Y) / Heightfield.DownsampleFactor;
-					}
-				}
-
-				{
-					SCOPED_DRAW_EVENT(RHICmdList, InitializeHeightfield);
-					const FIntPoint LightingAtlasSize = ExistingAtlas->GetAtlasSize();
-					const FVector2D InvLightingAtlasSize(1.0f / LightingAtlasSize.X, 1.0f / LightingAtlasSize.Y);
-
-					FRHITexture* RenderTargets[3] =
-					{
-						ExistingAtlas->Height->GetRenderTargetItem().TargetableTexture,
-						ExistingAtlas->Normal->GetRenderTargetItem().TargetableTexture,
-						ExistingAtlas->DiffuseColor->GetRenderTargetItem().TargetableTexture
-					};
-
-					RHICmdList.SetViewport(0, 0, 0.0f, LightingAtlasSize.X, LightingAtlasSize.Y, 1.0f);
-
-					FRHIRenderPassInfo RPInfo(UE_ARRAY_COUNT(RenderTargets), RenderTargets, ERenderTargetActions::Clear_Store);
-					RHICmdList.BeginRenderPass(RPInfo, TEXT("SetupHeightfields"));
-					{
-						FGraphicsPipelineStateInitializer GraphicsPSOInit;
-						RHICmdList.ApplyCachedRenderTargets(GraphicsPSOInit);
-
-						GraphicsPSOInit.RasterizerState = TStaticRasterizerState<FM_Solid, CM_None>::GetRHI();
-						GraphicsPSOInit.DepthStencilState = TStaticDepthStencilState<false, CF_Always>::GetRHI();
-						GraphicsPSOInit.BlendState = TStaticBlendState<>::GetRHI();
-						RHICmdList.SetStreamSource(0, GQuadVertexBuffer.VertexBufferRHI, 0);
-
-						TShaderMapRef<FHeightfieldSubsectionQuadVS> VertexShader(View.ShaderMap);
-						TShaderMapRef<FInitializeHeightfieldsPS> PixelShader(View.ShaderMap);
-
-						GraphicsPSOInit.BoundShaderState.VertexDeclarationRHI = GScreenVertexDeclaration.VertexDeclarationRHI;
-						GraphicsPSOInit.BoundShaderState.VertexShaderRHI = VertexShader.GetVertexShader();
-						GraphicsPSOInit.BoundShaderState.PixelShaderRHI = PixelShader.GetPixelShader();
-						GraphicsPSOInit.PrimitiveType = PT_TriangleList;
-
-						SetGraphicsPipelineState(RHICmdList, GraphicsPSOInit, 0);
-
-						for (TMap<FHeightfieldComponentTextures, TArray<FHeightfieldComponentDescription>>::TConstIterator It(Heightfield.ComponentDescriptions); It; ++It)
-						{
-							const TArray<FHeightfieldComponentDescription>& HeightfieldDescriptions = It.Value();
-
-							if (HeightfieldDescriptions.Num() > 0)
-							{
-								const int32 NumQuads = UploadSubsectionHeightfieldDescriptions(HeightfieldDescriptions, InvLightingAtlasSize, 1.0f / Heightfield.DownsampleFactor);
-
-								VertexShader->SetParameters(RHICmdList, View);
-								PixelShader->SetParameters(RHICmdList, View, It.Key().HeightAndNormal, It.Key().DiffuseColor);
-
-								RHICmdList.DrawPrimitive(0, 2, NumQuads);
-							}
-						}
-					}
-					RHICmdList.EndRenderPass();
-
-					TArray<FRHITransitionInfo, TInlineAllocator<3>> SRVTransitions;
-					for (FRHITexture* RHITexture : RenderTargets)
-					{
-						SRVTransitions.Add(FRHITransitionInfo(RHITexture, ERHIAccess::Unknown, ERHIAccess::SRVMask));
-					}
-					RHICmdList.Transition(MakeArrayView(SRVTransitions.GetData(), SRVTransitions.Num()));
-				}
-			}
-		}
-	});
-}
-
-void FHeightfieldLightingViewInfo::SetupHeightfieldsForScene(const FScene& Scene)
-{
-	const int32 NumPrimitives = Scene.DistanceFieldSceneData.HeightfieldPrimitives.Num();
-
-	if (NumPrimitives > 0
-		&& SupportsHeightfieldLighting(Scene.GetFeatureLevel(), Scene.GetShaderPlatform()))
-	{
-		const float MaxDistanceSquared = FMath::Square(GetMaxAOViewDistance() + GetGHeightfieldBounceDistance());
-		float LocalToWorldScale = 1;
-
-		for (int32 HeightfieldPrimitiveIndex = 0; HeightfieldPrimitiveIndex < NumPrimitives; HeightfieldPrimitiveIndex++)
-		{
-			const FPrimitiveSceneInfo* HeightfieldPrimitive = Scene.DistanceFieldSceneData.HeightfieldPrimitives[HeightfieldPrimitiveIndex];
-
-			UTexture2D* HeightfieldTexture = NULL;
-			UTexture2D* DiffuseColorTexture = NULL;
-			UTexture2D* VisibilityTexture = NULL;
-
-			FHeightfieldComponentDescription NewComponentDescription(HeightfieldPrimitive->Proxy->GetLocalToWorld());
-			HeightfieldPrimitive->Proxy->GetHeightfieldRepresentation(HeightfieldTexture, DiffuseColorTexture, VisibilityTexture, NewComponentDescription);
-
-			if (HeightfieldTexture && HeightfieldTexture->GetResource()->TextureRHI)
-			{
-				const FIntPoint HeightfieldSize = NewComponentDescription.HeightfieldRect.Size();
-
-				if (Heightfield.Rect.Area() == 0)
-				{
-					Heightfield.Rect = NewComponentDescription.HeightfieldRect;
-					LocalToWorldScale = NewComponentDescription.LocalToWorld.GetScaleVector().X;
-				}
-				else
-				{
-					Heightfield.Rect.Union(NewComponentDescription.HeightfieldRect);
-				}
-
-				TArray<FHeightfieldComponentDescription>& ComponentDescriptions = Heightfield.ComponentDescriptions.FindOrAdd(FHeightfieldComponentTextures(HeightfieldTexture, DiffuseColorTexture, VisibilityTexture));
-				ComponentDescriptions.Add(NewComponentDescription);
-			}
-		}
-	}
-}
-
-class FHeightfieldDescriptionsResource : public FRenderResource
-{
-public:
-
-	FCPUUpdatedBuffer Data;
-
-	FHeightfieldDescriptionsResource()
-	{
-		// In float4's, must match usf
-		const int32 HEIGHTFIELD_DATA_STRIDE = 12;
-
-		Data.Format = PF_A32B32G32R32F;
-		Data.Stride = HEIGHTFIELD_DATA_STRIDE;
-	}
-
-	virtual void InitDynamicRHI()  override
-	{
-		Data.Initialize();
-	}
-
-	virtual void ReleaseDynamicRHI() override
-	{
-		Data.Release();
-	}
-};
-
-TGlobalResource<FHeightfieldDescriptionsResource> GHeightfieldDescriptions;
-
-FRHIShaderResourceView* GetHeightfieldDescriptionsSRV()
-{
-	return GHeightfieldDescriptions.Data.BufferSRV;
-}
+// In float4's, must match usf
+static const int32 HEIGHTFIELD_DATA_STRIDE = 12;
 
 void FillHeightfieldDescriptionData(const TArray<FHeightfieldComponentDescription>& HeightfieldDescriptions, 
 	FVector2D InvLightingAtlasSize, 
 	float InvDownsampleFactor,
 	TArray<FVector4f, SceneRenderingAllocator>& HeightfieldDescriptionData)
 {
-	HeightfieldDescriptionData.Empty(HeightfieldDescriptions.Num() * GHeightfieldDescriptions.Data.Stride);
+	HeightfieldDescriptionData.Empty(HeightfieldDescriptions.Num() * HEIGHTFIELD_DATA_STRIDE);
 
 	for (int32 DescriptionIndex = 0; DescriptionIndex < HeightfieldDescriptions.Num(); DescriptionIndex++)
 	{
@@ -653,36 +73,11 @@ void FillHeightfieldDescriptionData(const TArray<FHeightfieldComponentDescriptio
 		HeightfieldDescriptionData.Add(ChannelMask);
 	}
 
-	check(HeightfieldDescriptionData.Num() % GHeightfieldDescriptions.Data.Stride == 0);
-}
-
-void UploadHeightfieldDescriptions(const TArray<FHeightfieldComponentDescription>& HeightfieldDescriptions, FVector2D InvLightingAtlasSize, float InvDownsampleFactor)
-{
-	TArray<FVector4f, SceneRenderingAllocator> HeightfieldDescriptionData;
-
-	FillHeightfieldDescriptionData(HeightfieldDescriptions,
-		InvLightingAtlasSize,
-		InvDownsampleFactor,
-		/*out*/ HeightfieldDescriptionData);
-
-	if (HeightfieldDescriptionData.Num() > GHeightfieldDescriptions.Data.MaxElements)
-	{
-		GHeightfieldDescriptions.Data.MaxElements = HeightfieldDescriptionData.Num() * 5 / 4;
-		GHeightfieldDescriptions.Data.Release();
-		GHeightfieldDescriptions.Data.Initialize();
-	}
-
-	void* LockedBuffer = RHILockBuffer(GHeightfieldDescriptions.Data.Buffer, 0, GHeightfieldDescriptions.Data.Buffer->GetSize(), RLM_WriteOnly);
-	const uint32 MemcpySize = HeightfieldDescriptionData.GetTypeSize() * HeightfieldDescriptionData.Num();
-	check(GHeightfieldDescriptions.Data.Buffer->GetSize() >= MemcpySize);
-	FPlatformMemory::Memcpy(LockedBuffer, HeightfieldDescriptionData.GetData(), MemcpySize);
-	RHIUnlockBuffer(GHeightfieldDescriptions.Data.Buffer);
+	check(HeightfieldDescriptionData.Num() % HEIGHTFIELD_DATA_STRIDE == 0);
 }
 
 FRDGBufferRef UploadHeightfieldDescriptions(FRDGBuilder& GraphBuilder, const TArray<FHeightfieldComponentDescription>& HeightfieldDescriptions, FVector2D InvLightingAtlasSize, float InvDownsampleFactor)
 {
-	const uint32 BufferStride = GHeightfieldDescriptions.Data.Stride;
-
 	TArray<FVector4f, SceneRenderingAllocator> HeightfieldDescriptionData;
 
 	FillHeightfieldDescriptionData(HeightfieldDescriptions,
@@ -690,536 +85,14 @@ FRDGBufferRef UploadHeightfieldDescriptions(FRDGBuilder& GraphBuilder, const TAr
 		InvDownsampleFactor,
 		/*out*/ HeightfieldDescriptionData);
 
-	FRDGBufferRef HeightfieldDescriptionsBuffer =
-		CreateUploadBuffer(GraphBuilder, TEXT("HeightfieldDescriptionsBuffer"),
-			sizeof(FVector4f), FMath::RoundUpToPowerOfTwo(FMath::Max(HeightfieldDescriptionData.Num(), 1)),
-			HeightfieldDescriptionData.GetData(), HeightfieldDescriptionData.Num() * HeightfieldDescriptionData.GetTypeSize());
+	FRDGBufferRef HeightfieldDescriptionsBuffer = CreateUploadBuffer(
+		GraphBuilder,
+		TEXT("HeightfieldDescriptionsBuffer"),
+		sizeof(FVector4f),
+		FMath::RoundUpToPowerOfTwo(FMath::Max(HeightfieldDescriptionData.Num(), 1)),
+		HeightfieldDescriptionData.GetData(),
+		HeightfieldDescriptionData.Num() * HeightfieldDescriptionData.GetTypeSize()
+	);
 
 	return HeightfieldDescriptionsBuffer;
 }
-
-class FGlobalHeightfieldParameters
-{
-	DECLARE_INLINE_TYPE_LAYOUT(FGlobalHeightfieldParameters, NonVirtual);
-public:
-	void Bind(const FShaderParameterMap& ParameterMap)
-	{
-		GlobalHeightfieldTexture.Bind(ParameterMap, TEXT("GlobalHeightfieldTexture"));
-		GlobalNormalTexture.Bind(ParameterMap, TEXT("GlobalNormalTexture"));
-		GlobalDiffuseColorTexture.Bind(ParameterMap, TEXT("GlobalDiffuseColorTexture"));
-		GlobalHeightfieldSampler.Bind(ParameterMap, TEXT("GlobalHeightfieldSampler"));
-	}
-
-	friend FArchive& operator<<(FArchive& Ar, FGlobalHeightfieldParameters& Parameters)
-	{
-		Ar << Parameters.GlobalHeightfieldTexture;
-		Ar << Parameters.GlobalNormalTexture;
-		Ar << Parameters.GlobalDiffuseColorTexture;
-		Ar << Parameters.GlobalHeightfieldSampler;
-		return Ar;
-	}
-
-	template<typename ShaderRHIParamRef>
-	void Set(FRHICommandList& RHICmdList, const ShaderRHIParamRef ShaderRHI, const FHeightfieldLightingAtlas& Atlas)
-	{
-		SetTextureParameter(RHICmdList, ShaderRHI, GlobalHeightfieldTexture, GlobalHeightfieldSampler, TStaticSamplerState<SF_Bilinear>::GetRHI(), Atlas.Height->GetRenderTargetItem().ShaderResourceTexture);
-		SetTextureParameter(RHICmdList, ShaderRHI, GlobalNormalTexture, GlobalHeightfieldSampler, TStaticSamplerState<SF_Bilinear>::GetRHI(), Atlas.Normal->GetRenderTargetItem().ShaderResourceTexture);
-		SetTextureParameter(RHICmdList, ShaderRHI, GlobalDiffuseColorTexture, GlobalHeightfieldSampler, TStaticSamplerState<SF_Bilinear>::GetRHI(), Atlas.DiffuseColor->GetRenderTargetItem().ShaderResourceTexture);
-	}
-
-private:
-	
-		LAYOUT_FIELD(FShaderResourceParameter, GlobalHeightfieldTexture)
-		LAYOUT_FIELD(FShaderResourceParameter, GlobalNormalTexture)
-		LAYOUT_FIELD(FShaderResourceParameter, GlobalDiffuseColorTexture)
-		LAYOUT_FIELD(FShaderResourceParameter, GlobalHeightfieldSampler)
-	
-};
-
-class FHeightfieldComponentQuadVS : public FGlobalShader
-{
-	DECLARE_SHADER_TYPE(FHeightfieldComponentQuadVS, Global);
-public:
-
-	static bool ShouldCompilePermutation(const FGlobalShaderPermutationParameters& Parameters)
-	{
-		return IsFeatureLevelSupported(Parameters.Platform, ERHIFeatureLevel::SM5) && DoesPlatformSupportDistanceFieldGI(Parameters.Platform);
-	}
-
-	/** Default constructor. */
-	FHeightfieldComponentQuadVS() {}
-
-	/** Initialization constructor. */
-	FHeightfieldComponentQuadVS(const ShaderMetaType::CompiledShaderInitializerType& Initializer)
-		: FGlobalShader(Initializer)
-	{
-		HeightfieldDescriptionParameters.Bind(Initializer.ParameterMap);
-	}
-
-	void SetParameters(FRHICommandList& RHICmdList, const FSceneView& View, int32 NumHeightfieldsValue)
-	{
-		FRHIVertexShader* ShaderRHI = RHICmdList.GetBoundVertexShader();
-
-		FGlobalShader::SetParameters<FViewUniformShaderParameters>(RHICmdList, ShaderRHI, View.ViewUniformBuffer);
-		HeightfieldDescriptionParameters.Set(RHICmdList, ShaderRHI, GetHeightfieldDescriptionsSRV(), NumHeightfieldsValue);
-	}
-
-private:
-	LAYOUT_FIELD(FHeightfieldDescriptionParameters, HeightfieldDescriptionParameters)
-};
-
-IMPLEMENT_SHADER_TYPE(, FHeightfieldComponentQuadVS, TEXT("/Engine/Private/HeightfieldLighting.usf"), TEXT("HeightfieldComponentQuadVS"), SF_Vertex);
-
-class FShadowHeightfieldsPS : public FGlobalShader
-{
-	DECLARE_SHADER_TYPE(FShadowHeightfieldsPS, Global);
-public:
-
-	static bool ShouldCompilePermutation(const FGlobalShaderPermutationParameters& Parameters)
-	{
-		return IsFeatureLevelSupported(Parameters.Platform, ERHIFeatureLevel::SM5) && DoesPlatformSupportDistanceFieldGI(Parameters.Platform);
-	}
-
-	/** Default constructor. */
-	FShadowHeightfieldsPS() {}
-
-	/** Initialization constructor. */
-	FShadowHeightfieldsPS(const ShaderMetaType::CompiledShaderInitializerType& Initializer)
-		: FGlobalShader(Initializer)
-	{
-		HeightfieldDescriptionParameters.Bind(Initializer.ParameterMap);
-		GlobalHeightfieldParameters.Bind(Initializer.ParameterMap);
-		WorldToShadow.Bind(Initializer.ParameterMap, TEXT("WorldToShadow"));
-		ShadowmapMinMax.Bind(Initializer.ParameterMap, TEXT("ShadowmapMinMax"));
-		ShadowDepthBias.Bind(Initializer.ParameterMap, TEXT("ShadowDepthBias"));
-		CascadeDepthMinMax.Bind(Initializer.ParameterMap, TEXT("CascadeDepthMinMax"));
-		ShadowDepthTexture.Bind(Initializer.ParameterMap, TEXT("ShadowDepthTexture"));
-		ShadowDepthTextureSampler.Bind(Initializer.ParameterMap, TEXT("ShadowDepthTextureSampler"));
-	}
-
-	void SetParameters(FRHICommandList& RHICmdList, const FSceneView& View, const FProjectedShadowInfo* ProjectedShadowInfo, int32 NumHeightfieldsValue, const FHeightfieldLightingAtlas& Atlas)
-	{
-		FRHIPixelShader* ShaderRHI = RHICmdList.GetBoundPixelShader();
-
-		FGlobalShader::SetParameters<FViewUniformShaderParameters>(RHICmdList, ShaderRHI, View.ViewUniformBuffer);
-		HeightfieldDescriptionParameters.Set(RHICmdList, ShaderRHI, GetHeightfieldDescriptionsSRV(), NumHeightfieldsValue);
-		GlobalHeightfieldParameters.Set(RHICmdList, ShaderRHI, Atlas);
-
-		FVector4f ShadowmapMinMaxValue;
-		FMatrix44f WorldToShadowMatrixValue = ProjectedShadowInfo->GetWorldToShadowMatrix(ShadowmapMinMaxValue);
-
-		SetShaderValue(RHICmdList, ShaderRHI, WorldToShadow, WorldToShadowMatrixValue);
-		SetShaderValue(RHICmdList, ShaderRHI, ShadowmapMinMax, ShadowmapMinMaxValue);
-		SetShaderValue(RHICmdList, ShaderRHI, ShadowDepthBias, ProjectedShadowInfo->GetShaderDepthBias());
-		SetShaderValue(RHICmdList, ShaderRHI, CascadeDepthMinMax, FVector2f(ProjectedShadowInfo->CascadeSettings.SplitNear, ProjectedShadowInfo->CascadeSettings.SplitFar));
-
-		FRHITexture* ShadowDepthTextureValue = ProjectedShadowInfo->RenderTargets.DepthTarget->GetRenderTargetItem().ShaderResourceTexture.GetReference();
-		FRHISamplerState* DepthSamplerState = TStaticSamplerState<SF_Point, AM_Clamp, AM_Clamp, AM_Clamp>::GetRHI();
-
-		SetTextureParameter(RHICmdList, ShaderRHI, ShadowDepthTexture, ShadowDepthTextureSampler, DepthSamplerState, ShadowDepthTextureValue);
-	}
-
-private:
-
-	LAYOUT_FIELD(FHeightfieldDescriptionParameters, HeightfieldDescriptionParameters);
-	LAYOUT_FIELD(FGlobalHeightfieldParameters, GlobalHeightfieldParameters);
-	LAYOUT_FIELD(FShaderParameter, WorldToShadow);
-	LAYOUT_FIELD(FShaderParameter, ShadowmapMinMax);
-	LAYOUT_FIELD(FShaderParameter, ShadowDepthBias);
-	LAYOUT_FIELD(FShaderParameter, CascadeDepthMinMax);
-	LAYOUT_FIELD(FShaderResourceParameter, ShadowDepthTexture);
-	LAYOUT_FIELD(FShaderResourceParameter, ShadowDepthTextureSampler);
-};
-
-IMPLEMENT_SHADER_TYPE(, FShadowHeightfieldsPS, TEXT("/Engine/Private/HeightfieldLighting.usf"), TEXT("ShadowHeightfieldsPS"), SF_Pixel);
-
-void FHeightfieldLightingViewInfo::ClearShadowing(FRDGBuilder& GraphBuilder, const FViewInfo& View, const FLightSceneInfo& LightSceneInfo) const
-{
-	if (AllowHeightfieldGI(View)
-		&& SupportsHeightfieldLighting(View.GetFeatureLevel(), View.GetShaderPlatform())
-		&& Heightfield.ComponentDescriptions.Num() > 0
-		&& LightSceneInfo.Proxy->GetLightType() == LightType_Directional
-		&& LightSceneInfo.Proxy->CastsDynamicShadow())
-	{
-		// TODO(RDG): Port this to RDG.
-		AddPass(GraphBuilder, RDG_EVENT_NAME("ClearShadowing"), [this, &View, &LightSceneInfo](FRHICommandList& RHICmdList)
-		{
-			FSceneViewState* ViewState = (FSceneViewState*)View.State;
-			const FHeightfieldLightingAtlas& Atlas = *ViewState->HeightfieldLightingAtlas;
-
-			FRHIRenderPassInfo RPInfo(Atlas.DirectionalLightShadowing->GetRenderTargetItem().TargetableTexture, ERenderTargetActions::Clear_Store);
-			RHICmdList.BeginRenderPass(RPInfo, TEXT("ClearShadowing"));
-			RHICmdList.EndRenderPass();
-		});
-	}
-}
-
-void FHeightfieldLightingViewInfo::ComputeShadowMapShadowing(FRDGBuilder& GraphBuilder, const FViewInfo& View, const FProjectedShadowInfo* ProjectedShadowInfo) const
-{
-	if (AllowHeightfieldGI(View)
-		&& SupportsHeightfieldLighting(View.GetFeatureLevel(), View.GetShaderPlatform())
-		&& Heightfield.ComponentDescriptions.Num() > 0
-		&& ProjectedShadowInfo->IsWholeSceneDirectionalShadow()
-		&& ProjectedShadowInfo->DependentView == &View
-		&& !ProjectedShadowInfo->bRayTracedDistanceField)
-	{
-		// TODO(RDG): Port this to RDG.
-		AddPass(GraphBuilder, RDG_EVENT_NAME("HeightfieldShadowMapShadowingForGI"), [this, &View, ProjectedShadowInfo](FRHICommandListImmediate& RHICmdList)
-		{
-			FSceneViewState* ViewState = (FSceneViewState*)View.State;
-			const FHeightfieldLightingAtlas& Atlas = *ViewState->HeightfieldLightingAtlas;
-
-			const FIntPoint LightingAtlasSize = Atlas.GetAtlasSize();
-			const FVector2D InvLightingAtlasSize(1.0f / LightingAtlasSize.X, 1.0f / LightingAtlasSize.Y);
-
-			FRHIRenderPassInfo RPInfo(Atlas.DirectionalLightShadowing->GetRenderTargetItem().TargetableTexture, ERenderTargetActions::Load_Store);
-			RHICmdList.BeginRenderPass(RPInfo, TEXT("ComputeShadowMapShadowing"));
-			{
-				FGraphicsPipelineStateInitializer GraphicsPSOInit;
-				RHICmdList.ApplyCachedRenderTargets(GraphicsPSOInit);
-
-				RHICmdList.SetViewport(0, 0, 0.0f, LightingAtlasSize.X, LightingAtlasSize.Y, 1.0f);
-				GraphicsPSOInit.RasterizerState = TStaticRasterizerState<FM_Solid, CM_None>::GetRHI();
-				GraphicsPSOInit.DepthStencilState = TStaticDepthStencilState<false, CF_Always>::GetRHI();
-				// Combine with other shadow types with min (ray traced)
-				GraphicsPSOInit.BlendState = TStaticBlendState<CW_RED, BO_Min, BF_One, BF_One>::GetRHI();
-				RHICmdList.SetStreamSource(0, GQuadVertexBuffer.VertexBufferRHI, 0);
-
-				TShaderMapRef<FHeightfieldComponentQuadVS> VertexShader(View.ShaderMap);
-				TShaderMapRef<FShadowHeightfieldsPS> PixelShader(View.ShaderMap);
-
-				GraphicsPSOInit.BoundShaderState.VertexDeclarationRHI = GScreenVertexDeclaration.VertexDeclarationRHI;
-				GraphicsPSOInit.BoundShaderState.VertexShaderRHI = VertexShader.GetVertexShader();
-				GraphicsPSOInit.BoundShaderState.PixelShaderRHI = PixelShader.GetPixelShader();
-				GraphicsPSOInit.PrimitiveType = PT_TriangleList;
-
-				SetGraphicsPipelineState(RHICmdList, GraphicsPSOInit, 0);
-
-				for (TMap<FHeightfieldComponentTextures, TArray<FHeightfieldComponentDescription>>::TConstIterator It(Heightfield.ComponentDescriptions); It; ++It)
-				{
-					const TArray<FHeightfieldComponentDescription>& HeightfieldDescriptions = It.Value();
-
-					if (HeightfieldDescriptions.Num() > 0)
-					{
-						//@todo - cull heightfield tiles with shadow bounds
-						UploadHeightfieldDescriptions(HeightfieldDescriptions, InvLightingAtlasSize, 1.0f / Heightfield.DownsampleFactor);
-
-						VertexShader->SetParameters(RHICmdList, View, HeightfieldDescriptions.Num());
-						PixelShader->SetParameters(RHICmdList, View, ProjectedShadowInfo, HeightfieldDescriptions.Num(), Atlas);
-
-						RHICmdList.DrawPrimitive(0, 2, HeightfieldDescriptions.Num());
-					}
-				}
-			}
-			RHICmdList.EndRenderPass();
-		});
-	}
-}
-
-class FLightHeightfieldsPS : public FMaterialShader
-{
-	DECLARE_SHADER_TYPE(FLightHeightfieldsPS, Material);
-public:
-
-	static bool ShouldCompilePermutation(const FMaterialShaderPermutationParameters& Parameters)
-	{
-		return Parameters.MaterialParameters.MaterialDomain == MD_LightFunction && IsFeatureLevelSupported(Parameters.Platform, ERHIFeatureLevel::SM5) && DoesPlatformSupportDistanceFieldGI(Parameters.Platform);
-	}
-
-	static void ModifyCompilationEnvironment(const FMaterialShaderPermutationParameters& Parameters, FShaderCompilerEnvironment& OutEnvironment)
-	{
-		OutEnvironment.SetDefine(TEXT("APPLY_LIGHT_FUNCTION"), 1);
-	}
-
-	/** Default constructor. */
-	FLightHeightfieldsPS() {}
-
-	/** Initialization constructor. */
-	FLightHeightfieldsPS(const ShaderMetaType::CompiledShaderInitializerType& Initializer)
-		: FMaterialShader(Initializer)
-	{
-		HeightfieldDescriptionParameters.Bind(Initializer.ParameterMap);
-		GlobalHeightfieldParameters.Bind(Initializer.ParameterMap);
-		LightDirection.Bind(Initializer.ParameterMap, TEXT("LightDirection"));
-		LightColor.Bind(Initializer.ParameterMap, TEXT("LightColor"));
-		SkyLightIndirectScale.Bind(Initializer.ParameterMap, TEXT("SkyLightIndirectScale"));
-		HeightfieldShadowing.Bind(Initializer.ParameterMap, TEXT("HeightfieldShadowing"));
-		HeightfieldShadowingSampler.Bind(Initializer.ParameterMap, TEXT("HeightfieldShadowingSampler"));
-		WorldToLight.Bind(Initializer.ParameterMap, TEXT("WorldToLight"));
-		LightFunctionParameters.Bind(Initializer.ParameterMap);
-	}
-
-	void SetParameters(
-		FRHICommandList& RHICmdList,
-		const FSceneView& View,
-		const FLightSceneInfo& LightSceneInfo,
-		const FMaterialRenderProxy* MaterialProxy,
-		int32 NumHeightfieldsValue,
-		const FHeightfieldLightingAtlas& Atlas,
-		float SkyLightIndirectScaleValue)
-	{
-		FRHIPixelShader* ShaderRHI = RHICmdList.GetBoundPixelShader();
-
-		const FMaterial& Material = MaterialProxy->GetMaterialWithFallback(View.GetFeatureLevel(), MaterialProxy);
-		FMaterialShader::SetViewParameters(RHICmdList, ShaderRHI, View, View.ViewUniformBuffer);
-		FMaterialShader::SetParameters(RHICmdList, ShaderRHI, MaterialProxy, Material, View);
-
-		HeightfieldDescriptionParameters.Set(RHICmdList, ShaderRHI, GetHeightfieldDescriptionsSRV(), NumHeightfieldsValue);
-		GlobalHeightfieldParameters.Set(RHICmdList, ShaderRHI, Atlas);
-
-		SetShaderValue(RHICmdList, ShaderRHI, LightDirection, FVector3f(LightSceneInfo.Proxy->GetDirection()));
-		SetShaderValue(RHICmdList, ShaderRHI, LightColor, LightSceneInfo.Proxy->GetColor() * LightSceneInfo.Proxy->GetIndirectLightingScale());
-
-		SetShaderValue(RHICmdList, ShaderRHI, SkyLightIndirectScale, SkyLightIndirectScaleValue);
-
-		const FSceneViewState* ViewState = (const FSceneViewState*)View.State;
-		SetTextureParameter(RHICmdList, ShaderRHI, HeightfieldShadowing, HeightfieldShadowingSampler, TStaticSamplerState<SF_Bilinear>::GetRHI(), Atlas.DirectionalLightShadowing->GetRenderTargetItem().ShaderResourceTexture);
-
-		const FVector Scale = LightSceneInfo.Proxy->GetLightFunctionScale();
-		// Switch x and z so that z of the user specified scale affects the distance along the light direction
-		const FVector InverseScale = FVector(1.f / Scale.Z, 1.f / Scale.Y, 1.f / Scale.X);
-		const FMatrix44f WorldToLightValue = LightSceneInfo.Proxy->GetWorldToLight() * FScaleMatrix(FVector(InverseScale));
-
-		SetShaderValue(RHICmdList, ShaderRHI, WorldToLight, WorldToLightValue);
-
-		LightFunctionParameters.Set(RHICmdList, ShaderRHI, &LightSceneInfo, 1);
-	}
-
-private:
-
-	LAYOUT_FIELD(FHeightfieldDescriptionParameters, HeightfieldDescriptionParameters);
-	LAYOUT_FIELD(FGlobalHeightfieldParameters, GlobalHeightfieldParameters);
-	LAYOUT_FIELD(FShaderParameter, LightDirection);
-	LAYOUT_FIELD(FShaderParameter, LightColor);
-	LAYOUT_FIELD(FShaderParameter, SkyLightIndirectScale);
-	LAYOUT_FIELD(FShaderResourceParameter, HeightfieldShadowing);
-	LAYOUT_FIELD(FShaderResourceParameter, HeightfieldShadowingSampler);
-	LAYOUT_FIELD(FShaderParameter, WorldToLight);
-	LAYOUT_FIELD(FLightFunctionSharedParameters, LightFunctionParameters);
-};
-
-IMPLEMENT_MATERIAL_SHADER_TYPE(, FLightHeightfieldsPS, TEXT("/Engine/Private/HeightfieldLighting.usf"), TEXT("LightHeightfieldsPS"), SF_Pixel);
-
-void FHeightfieldLightingViewInfo::ComputeLighting(FRDGBuilder& GraphBuilder, const FViewInfo& View, const FLightSceneInfo& LightSceneInfo) const
-{
-	const ERHIFeatureLevel::Type FeatureLevel = View.GetFeatureLevel();
-
-	if (!AllowHeightfieldGI(View)
-		|| !SupportsHeightfieldLighting(FeatureLevel, View.GetShaderPlatform())
-		//@todo - handle local lights
-		|| LightSceneInfo.Proxy->GetLightType() != LightType_Directional
-		|| !LightSceneInfo.Proxy->CastsDynamicShadow()
-		|| Heightfield.ComponentDescriptions.Num() == 0)
-	{
-		return;
-	}
-
-	AddPass(GraphBuilder, RDG_EVENT_NAME("HeightfieldLightingForGI"), [this, &View, &LightSceneInfo, FeatureLevel](FRHICommandListImmediate& RHICmdList)
-	{
-		FSceneViewState* ViewState = (FSceneViewState*)View.State;
-		const FHeightfieldLightingAtlas& Atlas = *ViewState->HeightfieldLightingAtlas;
-
-		const FIntPoint LightingAtlasSize = Atlas.GetAtlasSize();
-		const FVector2D InvLightingAtlasSize(1.0f / LightingAtlasSize.X, 1.0f / LightingAtlasSize.Y);
-
-		FRHIRenderPassInfo RPInfo(Atlas.Lighting->GetRenderTargetItem().TargetableTexture, ERenderTargetActions::Clear_Store);
-		RHICmdList.BeginRenderPass(RPInfo, TEXT("ComputeHeightfieldLighting"));
-		{
-			const bool bApplyLightFunction = (View.Family->EngineShowFlags.LightFunctions &&
-				LightSceneInfo.Proxy->GetLightFunctionMaterial() &&
-				LightSceneInfo.Proxy->GetLightFunctionMaterial()->GetIncompleteMaterialWithFallback(FeatureLevel).IsLightFunction());
-
-			const FMaterialRenderProxy* MaterialProxy = bApplyLightFunction ?
-				LightSceneInfo.Proxy->GetLightFunctionMaterial() :
-				UMaterial::GetDefaultMaterial(MD_LightFunction)->GetRenderProxy();
-
-			const FScene* Scene = (const FScene*)View.Family->Scene;
-
-			const float SkyLightIndirectScale = ShouldRenderDeferredDynamicSkyLight(Scene, *View.Family) ? Scene->SkyLight->IndirectLightingIntensity : 0;
-
-			// Skip rendering if the DefaultLightFunctionMaterial isn't compiled yet
-			if (MaterialProxy->GetIncompleteMaterialWithFallback(FeatureLevel).IsLightFunction())
-			{
-				RHICmdList.SetViewport(0, 0, 0.0f, LightingAtlasSize.X, LightingAtlasSize.Y, 1.0f);
-
-				FGraphicsPipelineStateInitializer GraphicsPSOInit;
-				RHICmdList.ApplyCachedRenderTargets(GraphicsPSOInit);
-
-				GraphicsPSOInit.RasterizerState = TStaticRasterizerState<FM_Solid, CM_None>::GetRHI();
-				GraphicsPSOInit.DepthStencilState = TStaticDepthStencilState<false, CF_Always>::GetRHI();
-				GraphicsPSOInit.BlendState = TStaticBlendState<>::GetRHI();
-
-				RHICmdList.SetStreamSource(0, GQuadVertexBuffer.VertexBufferRHI, 0);
-
-				TShaderMapRef<FHeightfieldComponentQuadVS> VertexShader(View.ShaderMap);
-
-				const FMaterial& Material = MaterialProxy->GetMaterialWithFallback(FeatureLevel, MaterialProxy);
-				const FMaterialShaderMap* MaterialShaderMap = Material.GetRenderingThreadShaderMap();
-				TShaderRef<FLightHeightfieldsPS> PixelShader = MaterialShaderMap->GetShader<FLightHeightfieldsPS>();
-
-				for (TMap<FHeightfieldComponentTextures, TArray<FHeightfieldComponentDescription>>::TConstIterator It(Heightfield.ComponentDescriptions); It; ++It)
-				{
-					const TArray<FHeightfieldComponentDescription>& HeightfieldDescriptions = It.Value();
-
-					if (HeightfieldDescriptions.Num() > 0)
-					{
-						UploadHeightfieldDescriptions(HeightfieldDescriptions, InvLightingAtlasSize, 1.0f / Heightfield.DownsampleFactor);
-
-						GraphicsPSOInit.BoundShaderState.VertexDeclarationRHI = GScreenVertexDeclaration.VertexDeclarationRHI;
-						GraphicsPSOInit.BoundShaderState.VertexShaderRHI = VertexShader.GetVertexShader();
-						GraphicsPSOInit.BoundShaderState.PixelShaderRHI = PixelShader.GetPixelShader();
-						GraphicsPSOInit.PrimitiveType = PT_TriangleList;
-
-						SetGraphicsPipelineState(RHICmdList, GraphicsPSOInit, 0);
-
-						VertexShader->SetParameters(RHICmdList, View, HeightfieldDescriptions.Num());
-						PixelShader->SetParameters(RHICmdList, View, LightSceneInfo, MaterialProxy, HeightfieldDescriptions.Num(), Atlas, SkyLightIndirectScale);
-
-						RHICmdList.DrawPrimitive(0, 2, HeightfieldDescriptions.Num());
-					}
-				}
-			}
-		}
-		RHICmdList.EndRenderPass();
-	});
-}
-
-const int32 GHeightfieldOcclusionDispatchSize = 8;
-
-class FCalculateHeightfieldOcclusionScreenGridCS : public FGlobalShader
-{
-	DECLARE_GLOBAL_SHADER(FCalculateHeightfieldOcclusionScreenGridCS);
-
-public:
-
-	BEGIN_SHADER_PARAMETER_STRUCT(FParameters, )
-		SHADER_PARAMETER_STRUCT_INCLUDE(FAOScreenGridParameters, AOScreenGridParameters)
-		SHADER_PARAMETER_STRUCT_INCLUDE(FAOParameters, AOParameters)
-		SHADER_PARAMETER_RDG_UNIFORM_BUFFER(FSceneTextureUniformParameters, SceneTextures)
-		RDG_TEXTURE_ACCESS(DistanceFieldNormal, ERHIAccess::SRVCompute)
-	END_SHADER_PARAMETER_STRUCT()
-
-	static bool ShouldCompilePermutation(const FGlobalShaderPermutationParameters& Parameters)
-	{
-		return IsFeatureLevelSupported(Parameters.Platform, ERHIFeatureLevel::SM5) && DoesPlatformSupportDistanceFieldGI(Parameters.Platform);
-	}
-
-	static void ModifyCompilationEnvironment(const FGlobalShaderPermutationParameters& Parameters, FShaderCompilerEnvironment& OutEnvironment)
-	{
-		FGlobalShader::ModifyCompilationEnvironment(Parameters, OutEnvironment);
-
-		OutEnvironment.SetDefine(TEXT("HEIGHTFIELD_OCCLUSION_DISPATCH_SIZEX"), GHeightfieldOcclusionDispatchSize);
-		extern int32 GConeTraceDownsampleFactor;
-		OutEnvironment.SetDefine(TEXT("TRACE_DOWNSAMPLE_FACTOR"), GConeTraceDownsampleFactor);
-
-		// To reduce shader compile time of compute shaders with shared memory, doesn't have an impact on generated code with current compiler (June 2010 DX SDK)
-		OutEnvironment.CompilerFlags.Add(CFLAG_StandardOptimization);
-	}
-
-	FCalculateHeightfieldOcclusionScreenGridCS(const ShaderMetaType::CompiledShaderInitializerType& Initializer)
-		: FGlobalShader(Initializer)
-	{
-		BindForLegacyShaderParameters<FParameters>(this, Initializer.PermutationId, Initializer.ParameterMap, false);
-		ScreenGridParameters.Bind(Initializer.ParameterMap);
-		HeightfieldDescriptionParameters.Bind(Initializer.ParameterMap);
-		HeightfieldTextureParameters.Bind(Initializer.ParameterMap);
-		TanConeHalfAngle.Bind(Initializer.ParameterMap, TEXT("TanConeHalfAngle"));
-	}
-
-	FCalculateHeightfieldOcclusionScreenGridCS()
-	{
-	}
-
-	void SetParameters(
-		FRHICommandList& RHICmdList,
-		const FViewInfo& View,
-		UTexture2D* HeightfieldTextureValue,
-		int32 NumHeightfieldsValue,
-		FRHITexture* DistanceFieldNormal)
-	{
-		FRHIComputeShader* ShaderRHI = RHICmdList.GetBoundComputeShader();
-		FGlobalShader::SetParameters<FViewUniformShaderParameters>(RHICmdList, ShaderRHI, View.ViewUniformBuffer);
-
-		ScreenGridParameters.Set(RHICmdList, ShaderRHI, View, DistanceFieldNormal);
-		HeightfieldDescriptionParameters.Set(RHICmdList, ShaderRHI, GetHeightfieldDescriptionsSRV(), NumHeightfieldsValue);
-		HeightfieldTextureParameters.Set(RHICmdList, ShaderRHI, HeightfieldTextureValue, nullptr, nullptr);
-
-		FAOSampleData2 AOSampleData;
-
-		TArray<FVector, TInlineAllocator<9> > SampleDirections;
-		GetSpacedVectors(View.Family->FrameNumber, SampleDirections);
-
-		for (int32 SampleIndex = 0; SampleIndex < NumConeSampleDirections; SampleIndex++)
-		{
-			AOSampleData.SampleDirections[SampleIndex] = FVector4f(SampleDirections[SampleIndex]);
-		}
-
-		SetUniformBufferParameterImmediate(RHICmdList, ShaderRHI, GetUniformBufferParameter<FAOSampleData2>(), AOSampleData);
-
-		extern float GAOConeHalfAngle;
-		SetShaderValue(RHICmdList, ShaderRHI, TanConeHalfAngle, FMath::Tan(GAOConeHalfAngle));
-	}
-
-private:
-	LAYOUT_FIELD(FScreenGridParameters, ScreenGridParameters);
-	LAYOUT_FIELD(FHeightfieldDescriptionParameters, HeightfieldDescriptionParameters);
-	LAYOUT_FIELD(FHeightfieldTextureParameters, HeightfieldTextureParameters);
-	LAYOUT_FIELD(FShaderParameter, TanConeHalfAngle);
-};
-
-IMPLEMENT_GLOBAL_SHADER(FCalculateHeightfieldOcclusionScreenGridCS, "/Engine/Private/HeightfieldLighting.usf", "CalculateHeightfieldOcclusionScreenGridCS", SF_Compute);
-
-
-void FHeightfieldLightingViewInfo::ComputeOcclusionForScreenGrid(
-	FRDGBuilder& GraphBuilder,
-	const FViewInfo& View,
-	const FSceneTextures& SceneTextures,
-	FRDGTextureRef DistanceFieldNormal,
-	const FAOScreenGridParameters& AOScreenGridParameters,
-	const FDistanceFieldAOParameters& Parameters) const
-{
-	int32 DownsampleFactor = Heightfield.DownsampleFactor;
-
-	if (Heightfield.ComponentDescriptions.Num() > 0 && GAOHeightfieldOcclusion)
-	{
-		for (TMap<FHeightfieldComponentTextures, TArray<FHeightfieldComponentDescription>>::TConstIterator It(Heightfield.ComponentDescriptions); It; ++It)
-		{
-			UTexture2D* HeightfieldTexture = It.Key().HeightAndNormal;
-			const TArray<FHeightfieldComponentDescription>& HeightfieldDescriptions = It.Value();
-
-			if (HeightfieldDescriptions.Num() > 0)
-			{
-				auto* PassParameters = GraphBuilder.AllocParameters<FCalculateHeightfieldOcclusionScreenGridCS::FParameters>();
-				PassParameters->AOScreenGridParameters = AOScreenGridParameters;
-				PassParameters->AOParameters = DistanceField::SetupAOShaderParameters(Parameters);
-				PassParameters->SceneTextures = SceneTextures.UniformBuffer;
-				PassParameters->DistanceFieldNormal = DistanceFieldNormal;
-
-				auto ComputeShader = View.ShaderMap->GetShader<FCalculateHeightfieldOcclusionScreenGridCS>();
-
-				ClearUnusedGraphResources(ComputeShader, PassParameters);
-
-				GraphBuilder.AddPass(
-					RDG_EVENT_NAME("HeightfieldOcclusion"),
-					PassParameters,
-					ERDGPassFlags::Compute,
-					[PassParameters, ComputeShader, &View, DistanceFieldNormal, &HeightfieldDescriptions, HeightfieldTexture, DownsampleFactor](FRHICommandListImmediate& RHICmdList)
-					{
-						UploadHeightfieldDescriptions(HeightfieldDescriptions, FVector2D(1, 1), 1.0f / DownsampleFactor);
-
-						const uint32 GroupSizeX = FMath::DivideAndRoundUp(View.ViewRect.Size().X / GAODownsampleFactor, GHeightfieldOcclusionDispatchSize);
-						const uint32 GroupSizeY = FMath::DivideAndRoundUp(View.ViewRect.Size().Y / GAODownsampleFactor, GHeightfieldOcclusionDispatchSize);
-
-						RHICmdList.SetComputeShader(ComputeShader.GetComputeShader());
-
-						ComputeShader->SetParameters(RHICmdList, View, HeightfieldTexture, HeightfieldDescriptions.Num(), DistanceFieldNormal->GetRHI());
-						SetShaderParameters(RHICmdList, ComputeShader, ComputeShader.GetComputeShader(), *PassParameters);
-
-						DispatchComputeShader(RHICmdList, ComputeShader.GetShader(), GroupSizeX, GroupSizeY, 1);
-
-						UnsetShaderUAVs(RHICmdList, ComputeShader, ComputeShader.GetComputeShader());
-					});
-			}
-		}
-	}
-}
diff --git a/Engine/Source/Runtime/Renderer/Private/HeightfieldLighting.h b/Engine/Source/Runtime/Renderer/Private/HeightfieldLighting.h
--- a/Engine/Source/Runtime/Renderer/Private/HeightfieldLighting.h
+++ b/Engine/Source/Runtime/Renderer/Private/HeightfieldLighting.h
@@ -15,41 +15,6 @@
 #include "PrimitiveSceneProxy.h"
 #include "RenderGraphResources.h"
 
-class FAOScreenGridParameters;
-class FDistanceFieldAOParameters;
-class FLightSceneInfo;
-class FProjectedShadowInfo;
-class FScene;
-class FViewInfo;
-class FGlobalDistanceFieldClipmap;
-struct Rect;
-
-class FHeightfieldLightingAtlas : public FRenderResource
-{
-public:
-
-	TRefCountPtr<IPooledRenderTarget> Height;
-	TRefCountPtr<IPooledRenderTarget> Normal;
-	TRefCountPtr<IPooledRenderTarget> DiffuseColor;
-	TRefCountPtr<IPooledRenderTarget> DirectionalLightShadowing;
-	TRefCountPtr<IPooledRenderTarget> Lighting;
-
-	FHeightfieldLightingAtlas() :
-		AtlasSize(FIntPoint(0, 0))
-	{}
-
-	virtual void InitDynamicRHI();
-	virtual void ReleaseDynamicRHI();
-
-	void InitializeForSize(FIntPoint InAtlasSize);
-
-	FIntPoint GetAtlasSize() const { return AtlasSize; }
-
-private:
-
-	FIntPoint AtlasSize;
-};
-
 class FHeightfieldComponentTextures
 {
 public:
@@ -89,128 +54,4 @@ public:
 		DownsampleFactor(1),
 		DownsampledRect(FIntRect(0, 0, 0, 0))
 	{}
-};
-
-class FHeightfieldLightingViewInfo
-{
-public:
-
-	FHeightfieldLightingViewInfo()
-	{}
-
-	void SetupVisibleHeightfields(const FViewInfo& View, FRDGBuilder& GraphBuilder);
-
-	void SetupHeightfieldsForScene(const FScene& Scene);
-
-	void ClearShadowing(FRDGBuilder& GraphBuilder, const FViewInfo& View, const FLightSceneInfo& LightSceneInfo) const;
-
-	void ComputeShadowMapShadowing(FRDGBuilder& GraphBuilder, const FViewInfo& View, const FProjectedShadowInfo* ProjectedShadowInfo) const;
-
-	void ComputeLighting(FRDGBuilder& GraphBuilder, const FViewInfo& View, const FLightSceneInfo& LightSceneInfo) const;
-
-	void ComputeOcclusionForScreenGrid(
-		FRDGBuilder& GraphBuilder,
-		const FViewInfo& View,
-		const FSceneTextures& SceneTextures,
-		FRDGTextureRef DistanceFieldNormal,
-		const FAOScreenGridParameters& AOScreenGridParameters,
-		const FDistanceFieldAOParameters& Parameters) const;
-
-private:
-
-	FHeightfieldDescription Heightfield;
-};
-
-extern FRHIShaderResourceView* GetHeightfieldDescriptionsSRV();
-
-class FHeightfieldDescriptionParameters
-{
-	DECLARE_TYPE_LAYOUT(FHeightfieldDescriptionParameters, NonVirtual);
-public:
-	void Bind(const FShaderParameterMap& ParameterMap)
-	{
-		HeightfieldDescriptions.Bind(ParameterMap, TEXT("HeightfieldDescriptions"));
-		NumHeightfields.Bind(ParameterMap, TEXT("NumHeightfields"));
-	}
-
-	friend FArchive& operator<<(FArchive& Ar, FHeightfieldDescriptionParameters& Parameters)
-	{
-		Ar << Parameters.HeightfieldDescriptions;
-		Ar << Parameters.NumHeightfields;
-		return Ar;
-	}
-
-	template<typename ShaderRHIParamRef>
-	void Set(FRHICommandList& RHICmdList, const ShaderRHIParamRef ShaderRHI, FRHIShaderResourceView* HeightfieldDescriptionsValue, int32 NumHeightfieldsValue)
-	{
-		SetSRVParameter(RHICmdList, ShaderRHI, HeightfieldDescriptions, HeightfieldDescriptionsValue);
-		SetShaderValue(RHICmdList, ShaderRHI, NumHeightfields, NumHeightfieldsValue);
-	}
-
-private:
-	
-		LAYOUT_FIELD(FShaderResourceParameter, HeightfieldDescriptions)
-		LAYOUT_FIELD(FShaderParameter, NumHeightfields)
-	
-};
-
-class FHeightfieldTextureParameters
-{
-	DECLARE_TYPE_LAYOUT(FHeightfieldTextureParameters, NonVirtual);
-public:
-	void Bind(const FShaderParameterMap& ParameterMap)
-	{
-		HeightfieldTexture.Bind(ParameterMap, TEXT("HeightfieldTexture"));
-		HeightfieldSampler.Bind(ParameterMap, TEXT("HeightfieldSampler"));
-		DiffuseColorTexture.Bind(ParameterMap, TEXT("DiffuseColorTexture"));
-		DiffuseColorSampler.Bind(ParameterMap, TEXT("DiffuseColorSampler"));
-		VisibilityTexture.Bind(ParameterMap, TEXT("VisibilityTexture"));
-		VisibilitySampler.Bind(ParameterMap, TEXT("VisibilitySampler"));
-	}
-
-	friend FArchive& operator<<(FArchive& Ar, FHeightfieldTextureParameters& Parameters)
-	{
-		Ar << Parameters.HeightfieldTexture;
-		Ar << Parameters.HeightfieldSampler;
-		Ar << Parameters.DiffuseColorTexture;
-		Ar << Parameters.DiffuseColorSampler;
-		Ar << Parameters.VisibilityTexture;
-		Ar << Parameters.VisibilitySampler;
-		return Ar;
-	}
-
-	template<typename ShaderRHIParamRef>
-	void Set(FRHICommandList& RHICmdList, const ShaderRHIParamRef ShaderRHI, UTexture2D* HeightfieldTextureValue, UTexture2D* DiffuseColorTextureValue, UTexture2D* VisibilityTextureValue)
-	{
-		//@todo - shouldn't filter the heightfield, it's packed
-		SetTextureParameter(RHICmdList, ShaderRHI, HeightfieldTexture, HeightfieldSampler, TStaticSamplerState<SF_Bilinear>::GetRHI(), HeightfieldTextureValue->GetResource()->TextureRHI);
-
-		if (DiffuseColorTextureValue)
-		{
-			SetTextureParameter(RHICmdList, ShaderRHI, DiffuseColorTexture, DiffuseColorSampler, TStaticSamplerState<SF_Bilinear>::GetRHI(), DiffuseColorTextureValue->GetResource()->TextureRHI);
-		}
-		else
-		{
-			SetTextureParameter(RHICmdList, ShaderRHI, DiffuseColorTexture, DiffuseColorSampler, TStaticSamplerState<SF_Bilinear>::GetRHI(), GBlackTexture->TextureRHI);
-		}
-
-		if (VisibilityTextureValue)
-		{
-			SetTextureParameter(RHICmdList, ShaderRHI, VisibilityTexture, VisibilitySampler, TStaticSamplerState<SF_Bilinear>::GetRHI(), VisibilityTextureValue->GetResource()->TextureRHI);
-		}
-		else
-		{
-			SetTextureParameter(RHICmdList, ShaderRHI, VisibilityTexture, VisibilitySampler, TStaticSamplerState<SF_Bilinear>::GetRHI(), GBlackTexture->TextureRHI);
-		}
-	}
-
-private:
-	
-		LAYOUT_FIELD(FShaderResourceParameter, HeightfieldTexture)
-		LAYOUT_FIELD(FShaderResourceParameter, HeightfieldSampler)
-		LAYOUT_FIELD(FShaderResourceParameter, DiffuseColorTexture)
-		LAYOUT_FIELD(FShaderResourceParameter, DiffuseColorSampler)
-		LAYOUT_FIELD(FShaderResourceParameter, VisibilityTexture)
-		LAYOUT_FIELD(FShaderResourceParameter, VisibilitySampler)
-	
-};
+};
\ No newline at end of file
diff --git a/Engine/Source/Runtime/Renderer/Private/LightRendering.cpp b/Engine/Source/Runtime/Renderer/Private/LightRendering.cpp
--- a/Engine/Source/Runtime/Renderer/Private/LightRendering.cpp
+++ b/Engine/Source/Runtime/Renderer/Private/LightRendering.cpp
@@ -1780,14 +1780,7 @@ void FDeferredShadingSceneRenderer::RenderLights(
 						}
 					}
 					else // (OcclusionType == FOcclusionType::Shadowmap)
-					{
-						for (int32 ViewIndex = 0; ViewIndex < Views.Num(); ViewIndex++)
-						{
-							const FViewInfo& View = Views[ViewIndex];
-							RDG_GPU_MASK_SCOPE(GraphBuilder, View.GPUMask);
-							View.HeightfieldLightingViewInfo.ClearShadowing(GraphBuilder, View, LightSceneInfo);
-						}
-					
+					{					
 						const auto ClearShadowMask = [&](FRDGTextureRef InScreenShadowMaskTexture)
 						{
 							// Clear light attenuation for local lights with a quad covering their extents
@@ -1845,13 +1838,6 @@ void FDeferredShadingSceneRenderer::RenderLights(
 					bUsedShadowMaskTexture = true;
 				}
 
-				for (int32 ViewIndex = 0; ViewIndex < Views.Num(); ViewIndex++)
-				{
-					const FViewInfo& View = Views[ViewIndex];
-					RDG_GPU_MASK_SCOPE(GraphBuilder, View.GPUMask);
-					View.HeightfieldLightingViewInfo.ComputeLighting(GraphBuilder, View, LightSceneInfo);
-				}
-
 				// Render light function to the attenuation buffer.
 				if (bDirectLighting)
 				{
diff --git a/Engine/Source/Runtime/Renderer/Private/RendererScene.cpp b/Engine/Source/Runtime/Renderer/Private/RendererScene.cpp
--- a/Engine/Source/Runtime/Renderer/Private/RendererScene.cpp
+++ b/Engine/Source/Runtime/Renderer/Private/RendererScene.cpp
@@ -244,8 +244,6 @@ FSceneViewState::FSceneViewState()
 
 	bRoundRobinOcclusionEnabled = false;
 
-	HeightfieldLightingAtlas = NULL;
-
 	for (int32 CascadeIndex = 0; CascadeIndex < UE_ARRAY_COUNT(TranslucencyLightingCacheAllocations); CascadeIndex++)
 	{
 		TranslucencyLightingCacheAllocations[CascadeIndex] = NULL;
@@ -304,8 +302,6 @@ FSceneViewState::~FSceneViewState()
 		delete TranslucencyLightingCacheAllocations[CascadeIndex];
 	}
 
-	DestroyRenderResource(HeightfieldLightingAtlas);
-
 	HairStrandsViewStateData.Release();
 	ShaderDrawDebugStateData.Release();
 }
diff --git a/Engine/Source/Runtime/Renderer/Private/ScenePrivate.h b/Engine/Source/Runtime/Renderer/Private/ScenePrivate.h
--- a/Engine/Source/Runtime/Renderer/Private/ScenePrivate.h
+++ b/Engine/Source/Runtime/Renderer/Private/ScenePrivate.h
@@ -964,8 +964,6 @@ public:
 	 */
 	FPreviousViewInfo PrevFrameViewInfo;
 
-	FHeightfieldLightingAtlas* HeightfieldLightingAtlas;
-
 	// Temporal AA result for light shafts of last frame
 	FTemporalAAHistory LightShaftOcclusionHistory;
 	// Temporal AA result for light shafts of last frame
diff --git a/Engine/Source/Runtime/Renderer/Private/SceneRendering.h b/Engine/Source/Runtime/Renderer/Private/SceneRendering.h
--- a/Engine/Source/Runtime/Renderer/Private/SceneRendering.h
+++ b/Engine/Source/Runtime/Renderer/Private/SceneRendering.h
@@ -1450,8 +1450,6 @@ public:
 
 	TOptional<FMobileLightShaftInfo> MobileLightShaft;
 
-	FHeightfieldLightingViewInfo HeightfieldLightingViewInfo;
-
 	FGlobalShaderMap* ShaderMap;
 
 	bool bIsSnapshot;
diff --git a/Engine/Source/Runtime/Renderer/Private/ShadowRendering.cpp b/Engine/Source/Runtime/Renderer/Private/ShadowRendering.cpp
--- a/Engine/Source/Runtime/Renderer/Private/ShadowRendering.cpp
+++ b/Engine/Source/Runtime/Renderer/Private/ShadowRendering.cpp
@@ -2285,22 +2285,6 @@ void FDeferredShadingSceneRenderer::RenderDeferredShadowProjections(
 
 	RenderCapsuleDirectShadows(GraphBuilder, SceneTextures.UniformBuffer, *LightSceneInfo, ScreenShadowMaskTexture, VisibleLightInfo.CapsuleShadowsToProject, bProjectingForForwardShading);
 
-	for (int32 ViewIndex = 0; ViewIndex < Views.Num(); ViewIndex++)
-	{
-		const FViewInfo& View = Views[ViewIndex];
-		RDG_GPU_MASK_SCOPE(GraphBuilder, View.GPUMask);
-
-		for (int32 ShadowIndex = 0; ShadowIndex < VisibleLightInfo.ShadowsToProject.Num(); ShadowIndex++)
-		{
-			const FProjectedShadowInfo* ProjectedShadowInfo = VisibleLightInfo.ShadowsToProject[ShadowIndex];
-
-			if (ProjectedShadowInfo->bAllocated && ProjectedShadowInfo->bWholeSceneShadow)
-			{
-				View.HeightfieldLightingViewInfo.ComputeShadowMapShadowing(GraphBuilder, View, ProjectedShadowInfo);
-			}
-		}
-	}
-
 	// Inject deep shadow mask for regular shadow map. When using virtual shadow map, it is directly handled in the shadow kernel.
 	if (HairStrands::HasViewHairStrandsData(Views))
 	{
