Avoid wave operations in path compaction shader which appears to give a slight speedup despite additional contention on the atomic and does not require running with SM6.

Implement tiled dispatch in the path tracer to reduce the likelyhood of GPU timeouts when rendering at high resolution. This also reduces the memory requirements for path state when running with path compaction enabled.

Change from a uint buffer to a structured buffer for storing path states which gives a small speedup.

Add indirect dispatch support to launch less work for compacted bounces (off by default as it does not seem to provide a speedup so far)

diff --git a/Engine/Shaders/Private/PathTracing/PathTracing.usf b/Engine/Shaders/Private/PathTracing/PathTracing.usf
--- a/Engine/Shaders/Private/PathTracing/PathTracing.usf
+++ b/Engine/Shaders/Private/PathTracing/PathTracing.usf
@@ -10,104 +10,63 @@ PathTracing.usf: Reference path tracing
 
 #include "PathTracingCore.ush"
 
-#if PATH_TRACER_USE_COMPACTION
+#if PATH_TRACER_USE_COMPACTION == 1
 
 int Bounce;
-RWBuffer<uint> PathStateData;
+RWStructuredBuffer<FPathTracingPackedPathState> PathStateData;
 Buffer<int> ActivePaths;
 RWBuffer<int> NextActivePaths;
-RWBuffer<int> NumPathStates;
+RWBuffer<uint> NumPathStates;
 
 
-// TODO: Experiment with AoS vs SoA vs AoSoA?
-#define PATH_STATE_STRIDE   PATHTRACER_PATHSTATE_NUM_WORDS
-
 FPathState LoadPathStateData(uint Index)
 {
+	FPathTracingPackedPathState Data = PathStateData[Index];
 	FPathState Output = (FPathState)0;
-	uint PackedPixelIndex = PathStateData[PATH_STATE_STRIDE * Index + 0];
-	Output.PixelIndex = uint2(PackedPixelIndex & 0xFFFFu, PackedPixelIndex >> 16);
-
-	Output.RandSequence.SampleIndex = PathStateData[PATH_STATE_STRIDE * Index + 1];
-	Output.RandSequence.SampleSeed  = PathStateData[PATH_STATE_STRIDE * Index + 2];
-
-	Output.Radiance.x = asfloat(PathStateData[PATH_STATE_STRIDE * Index + 3]);
-	Output.Radiance.y = asfloat(PathStateData[PATH_STATE_STRIDE * Index + 4]);
-	Output.Radiance.z = asfloat(PathStateData[PATH_STATE_STRIDE * Index + 5]);
-
-	Output.BackgroundVisibility = asfloat(PathStateData[PATH_STATE_STRIDE * Index + 6]);
-
-	uint3 PackedAlbedoNormal;
-	PackedAlbedoNormal.x = PathStateData[PATH_STATE_STRIDE * Index + 7];
-	PackedAlbedoNormal.y = PathStateData[PATH_STATE_STRIDE * Index + 8];
-	PackedAlbedoNormal.z = PathStateData[PATH_STATE_STRIDE * Index + 9];
-	Output.Albedo.x = f16tof32(PackedAlbedoNormal.x);
-	Output.Albedo.y = f16tof32(PackedAlbedoNormal.x >> 16);
-	Output.Albedo.z = f16tof32(PackedAlbedoNormal.y);
-	Output.Normal.x = f16tof32(PackedAlbedoNormal.y >> 16);
-	Output.Normal.y = f16tof32(PackedAlbedoNormal.z);
-	Output.Normal.z = f16tof32(PackedAlbedoNormal.z >> 16);
-
-	Output.Ray.Origin.x = asfloat(PathStateData[PATH_STATE_STRIDE * Index + 10]);
-	Output.Ray.Origin.y = asfloat(PathStateData[PATH_STATE_STRIDE * Index + 11]);
-	Output.Ray.Origin.z = asfloat(PathStateData[PATH_STATE_STRIDE * Index + 12]);
-	Output.Ray.Direction.x = asfloat(PathStateData[PATH_STATE_STRIDE * Index + 13]);
-	Output.Ray.Direction.y = asfloat(PathStateData[PATH_STATE_STRIDE * Index + 14]);
-	Output.Ray.Direction.z = asfloat(PathStateData[PATH_STATE_STRIDE * Index + 15]);
+	Output.PixelIndex = uint2(Data.PixelIndex & 0xFFFFu, Data.PixelIndex >> 16);
+	Output.RandSequence.SampleIndex = Data.RandSeqSampleIndex;
+	Output.RandSequence.SampleSeed  = Data.RandSeqSampleSeed;
+	Output.Radiance = Data.Radiance;
+	Output.BackgroundVisibility = Data.BackgroundVisibility;
+	Output.Albedo.x = f16tof32(Data.PackedAlbedoNormal.x);
+	Output.Albedo.y = f16tof32(Data.PackedAlbedoNormal.x >> 16);
+	Output.Albedo.z = f16tof32(Data.PackedAlbedoNormal.y);
+	Output.Normal.x = f16tof32(Data.PackedAlbedoNormal.y >> 16);
+	Output.Normal.y = f16tof32(Data.PackedAlbedoNormal.z);
+	Output.Normal.z = f16tof32(Data.PackedAlbedoNormal.z >> 16);
+	Output.Ray.Origin = Data.RayOrigin;
+	Output.Ray.Direction = Data.RayDirection;
 	Output.Ray.TMin = 0;
 	Output.Ray.TMax = POSITIVE_INFINITY;
-
-	uint PackedRayCone = PathStateData[PATH_STATE_STRIDE * Index + 16];
-	Output.RayCone.Width = f16tof32(PackedRayCone);
-	Output.RayCone.SpreadAngle = f16tof32(PackedRayCone >> 16);
-
-	Output.PathThroughput.x = asfloat(PathStateData[PATH_STATE_STRIDE * Index + 17]);
-	Output.PathThroughput.y = asfloat(PathStateData[PATH_STATE_STRIDE * Index + 18]);
-	Output.PathThroughput.z = asfloat(PathStateData[PATH_STATE_STRIDE * Index + 19]);
-
-	uint2 PackedRoughnessSigma;
-	PackedRoughnessSigma.x = PathStateData[PATH_STATE_STRIDE * Index + 20];
-	PackedRoughnessSigma.y = PathStateData[PATH_STATE_STRIDE * Index + 21];
-	Output.PathRoughness = f16tof32(PackedRoughnessSigma.x);
-	Output.SigmaT.x = f16tof32(PackedRoughnessSigma.x >> 16);
-	Output.SigmaT.y = f16tof32(PackedRoughnessSigma.y);
-	Output.SigmaT.z = f16tof32(PackedRoughnessSigma.y >> 16);
+	Output.RayCone.Width = f16tof32(Data.PackedRayCone);
+	Output.RayCone.SpreadAngle = f16tof32(Data.PackedRayCone >> 16);
+	Output.PathThroughput = Data.PathThroughput;
+	Output.PathRoughness = f16tof32(Data.PackedRoughnessSigma.x);
+	Output.SigmaT.x = f16tof32(Data.PackedRoughnessSigma.x >> 16);
+	Output.SigmaT.y = f16tof32(Data.PackedRoughnessSigma.y);
+	Output.SigmaT.z = f16tof32(Data.PackedRoughnessSigma.y >> 16);
 
 	return Output;
 }
 
 void StorePathStateData(FPathState PathState, uint Index)
 {
-	PathStateData[PATH_STATE_STRIDE * Index + 0] = PathState.PixelIndex.x + PathState.PixelIndex.y * 65536;
-
-	PathStateData[PATH_STATE_STRIDE * Index + 1] = PathState.RandSequence.SampleIndex;
-	PathStateData[PATH_STATE_STRIDE * Index + 2] = PathState.RandSequence.SampleSeed;
-
-	PathStateData[PATH_STATE_STRIDE * Index + 3] = asuint(PathState.Radiance.x);
-	PathStateData[PATH_STATE_STRIDE * Index + 4] = asuint(PathState.Radiance.y);
-	PathStateData[PATH_STATE_STRIDE * Index + 5] = asuint(PathState.Radiance.z);
-
-	PathStateData[PATH_STATE_STRIDE * Index + 6] = asuint(PathState.BackgroundVisibility);
-
-	PathStateData[PATH_STATE_STRIDE * Index + 7] = f32tof16(PathState.Albedo.x) | (f32tof16(PathState.Albedo.y) << 16);
-	PathStateData[PATH_STATE_STRIDE * Index + 8] = f32tof16(PathState.Albedo.z) | (f32tof16(PathState.Normal.x) << 16);
-	PathStateData[PATH_STATE_STRIDE * Index + 9] = f32tof16(PathState.Normal.y) | (f32tof16(PathState.Normal.z) << 16);
-
-	PathStateData[PATH_STATE_STRIDE * Index + 10] = asuint(PathState.Ray.Origin.x);
-	PathStateData[PATH_STATE_STRIDE * Index + 11] = asuint(PathState.Ray.Origin.y);
-	PathStateData[PATH_STATE_STRIDE * Index + 12] = asuint(PathState.Ray.Origin.z);
-	PathStateData[PATH_STATE_STRIDE * Index + 13] = asuint(PathState.Ray.Direction.x);
-	PathStateData[PATH_STATE_STRIDE * Index + 14] = asuint(PathState.Ray.Direction.y);
-	PathStateData[PATH_STATE_STRIDE * Index + 15] = asuint(PathState.Ray.Direction.z);
-
-	PathStateData[PATH_STATE_STRIDE * Index + 16] = f32tof16(PathState.RayCone.Width) | (f32tof16(PathState.RayCone.SpreadAngle) << 16);
-
-	PathStateData[PATH_STATE_STRIDE * Index + 17] = asuint(PathState.PathThroughput.x);
-	PathStateData[PATH_STATE_STRIDE * Index + 18] = asuint(PathState.PathThroughput.y);
-	PathStateData[PATH_STATE_STRIDE * Index + 19] = asuint(PathState.PathThroughput.z);
-
-	PathStateData[PATH_STATE_STRIDE * Index + 20] = f32tof16(PathState.PathRoughness) | (f32tof16(PathState.SigmaT.x) << 16);
-	PathStateData[PATH_STATE_STRIDE * Index + 21] = f32tof16(PathState.SigmaT.y) | (f32tof16(PathState.SigmaT.z) << 16);
+	FPathTracingPackedPathState Data;
+	Data.PixelIndex = PathState.PixelIndex.x + PathState.PixelIndex.y * 65536;;
+	Data.RandSeqSampleIndex = PathState.RandSequence.SampleIndex;
+	Data.RandSeqSampleSeed  = PathState.RandSequence.SampleSeed;
+	Data.Radiance = PathState.Radiance;
+	Data.BackgroundVisibility = PathState.BackgroundVisibility;
+	Data.PackedAlbedoNormal.x = f32tof16(PathState.Albedo.x) | (f32tof16(PathState.Albedo.y) << 16);
+	Data.PackedAlbedoNormal.y = f32tof16(PathState.Albedo.z) | (f32tof16(PathState.Normal.x) << 16);
+	Data.PackedAlbedoNormal.z = f32tof16(PathState.Normal.y) | (f32tof16(PathState.Normal.z) << 16);
+	Data.RayOrigin = PathState.Ray.Origin;
+	Data.RayDirection = PathState.Ray.Direction;
+	Data.PackedRayCone = f32tof16(PathState.RayCone.Width) | (f32tof16(PathState.RayCone.SpreadAngle) << 16);
+	Data.PathThroughput = PathState.PathThroughput;
+	Data.PackedRoughnessSigma.x = f32tof16(PathState.PathRoughness) | (f32tof16(PathState.SigmaT.x) << 16);
+	Data.PackedRoughnessSigma.y = f32tof16(PathState.SigmaT.y) | (f32tof16(PathState.SigmaT.z) << 16);
+	PathStateData[Index] = Data;
 }
 
 RAY_TRACING_ENTRY_RAYGEN(PathTracingMainRG)
@@ -116,6 +75,12 @@ RAY_TRACING_ENTRY_RAYGEN(PathTracingMainRG)
 	const uint2 DispatchDim = DispatchRaysDimensions().xy;
 
 	int LinearPixelIndex = DispatchIdx.x + DispatchDim.x * DispatchIdx.y;
+	if (LinearPixelIndex == 0)
+	{
+		// write other dimensions for indirect dispatch
+		NumPathStates[1] = 1;
+		NumPathStates[2] = 1;
+	}
 	FPathState PathState;
 	if (Bounce == 0)
 	{
@@ -132,23 +97,11 @@ RAY_TRACING_ENTRY_RAYGEN(PathTracingMainRG)
 	}
 
 	const bool KeepGoing = PathTracingKernel(PathState, Bounce);
-	// check how many lanes in our wave want to keep tracing forward
-	const uint NumActive = WaveActiveCountBits(KeepGoing);
-
-	uint GroupWriteStart;
-	if (WaveIsFirstLane() && NumActive > 0)
-	{
-		// increment the global counter to know which lane to map to
-		// NOTE: we reduce contention by only incrementing once per wave, but this is still sub-optimal
-		InterlockedAdd(NumPathStates[0], NumActive, GroupWriteStart);
-	}
-	// Broadcast the global write offset to all lanes.
-	GroupWriteStart = WaveReadLaneFirst(GroupWriteStart);
-
 	if (KeepGoing)
 	{
-		// the path should still be extended - store the path state
-		uint PathStateIndex = GroupWriteStart + WavePrefixCountBits(KeepGoing);
+		// NOTE: using wave instructions to reduce contention seems to run slightly slower? (tested on RTX-3080)
+		uint PathStateIndex;
+		InterlockedAdd(NumPathStates[0], 1, PathStateIndex);
 		NextActivePaths[PathStateIndex] = LinearPixelIndex;
 		StorePathStateData(PathState, LinearPixelIndex);
 	}
@@ -156,11 +109,11 @@ RAY_TRACING_ENTRY_RAYGEN(PathTracingMainRG)
 	{
 		// nothing left to do
 		// Accumulate radiance and update pixel variance
-		PathState.WritePixel(BlendFactor);
+		PathState.WritePixel();
 	}
 }
 
-#else
+#else // PATH_TRACER_USE_COMPACTION == 0
 
 RAY_TRACING_ENTRY_RAYGEN(PathTracingMainRG)
 {
@@ -179,7 +132,7 @@ RAY_TRACING_ENTRY_RAYGEN(PathTracingMainRG)
 	}
 
 	// Accumulate radiance and update pixel variance
-	PathState.WritePixel(BlendFactor);
+	PathState.WritePixel();
 }
 
 #endif
\ No newline at end of file
diff --git a/Engine/Shaders/Private/PathTracing/PathTracingCore.ush b/Engine/Shaders/Private/PathTracing/PathTracingCore.ush
--- a/Engine/Shaders/Private/PathTracing/PathTracingCore.ush
+++ b/Engine/Shaders/Private/PathTracing/PathTracingCore.ush
@@ -66,10 +66,9 @@ struct FPathState
 	float3 PathThroughput;          // 12 bytes
 	float PathRoughness;            //  2 bytes // half precision
 	float3 SigmaT; // Extinction    //  6 bytes // half precision 
-									// 88 bytes total
-									// 22 words - sync with PATHTRACER_PATHSTATE_NUM_WORDS
+									// 88 bytes total (see FPathTracingPackedPathState)
 
-	void WritePixel(float BlendFactor)
+	void WritePixel()
 	{
 		// Avoid reading the old pixel on the first sample on the off-chance there is a NaN/Inf pixel ...
 		float4 OldPixel  = Iteration > 0 ? RadianceTexture[PixelIndex] : 0;
diff --git a/Engine/Shaders/Shared/PathTracingDefinitions.h b/Engine/Shaders/Shared/PathTracingDefinitions.h
--- a/Engine/Shaders/Shared/PathTracingDefinitions.h
+++ b/Engine/Shaders/Shared/PathTracingDefinitions.h
@@ -7,8 +7,6 @@
 #define PATHTRACER_COHERENT_TILE_SHIFT		3
 #define PATHTRACER_COHERENT_TILE_SIZE		(1 << PATHTRACER_COHERENT_TILE_SHIFT)
 
-#define PATHTRACER_PATHSTATE_NUM_WORDS		22
-
 // Constants for 'SamplerType'
 #define PATHTRACER_SAMPLER_DEFAULT			0
 #define PATHTRACER_SAMPLER_ERROR_DIFFUSION	1
diff --git a/Engine/Shaders/Shared/RayTracingTypes.h b/Engine/Shaders/Shared/RayTracingTypes.h
--- a/Engine/Shaders/Shared/RayTracingTypes.h
+++ b/Engine/Shaders/Shared/RayTracingTypes.h
@@ -31,9 +31,27 @@ struct FPathTracingLight {
 
 static_assert(sizeof(FPathTracingLight) == 128, "Path tracing light structure should be aligned to 128 bytes for optimal access on the GPU");
 
+struct FPathTracingPackedPathState {
+	uint32    PixelIndex;
+	uint32    RandSeqSampleIndex;
+	uint32    RandSeqSampleSeed;
+	FVector3f Radiance;
+	float     BackgroundVisibility;
+	uint16    Albedo[3];
+	uint16    Normal[3];
+	FVector3f RayOrigin;
+	FVector3f RayDirection;
+	uint32    RayCone;
+	FVector3f PathThroughput;
+	uint16    PathRoughness;
+	uint16    SigmaT[3];
+};
+
+static_assert(sizeof(FPathTracingPackedPathState) == 88, "Packed Path State size should be minimized");
+
 #else
 
-// HLSL side of the struct above
+// HLSL side of the structs above
 
 struct FPathTracingLight {
 	float3  Position;
@@ -53,4 +71,18 @@ struct FPathTracingLight {
 	float Padding;
 };
 
+struct FPathTracingPackedPathState {
+	uint      PixelIndex;
+	uint      RandSeqSampleIndex;
+	uint      RandSeqSampleSeed;
+	float3    Radiance;
+	float     BackgroundVisibility;
+	uint3     PackedAlbedoNormal;
+	float3    RayOrigin;
+	float3    RayDirection;
+	uint      PackedRayCone;
+	float3    PathThroughput;
+	uint2     PackedRoughnessSigma;
+};
+
 #endif
diff --git a/Engine/Source/Runtime/Renderer/Private/PathTracing.cpp b/Engine/Source/Runtime/Renderer/Private/PathTracing.cpp
--- a/Engine/Source/Runtime/Renderer/Private/PathTracing.cpp
+++ b/Engine/Source/Runtime/Renderer/Private/PathTracing.cpp
@@ -34,6 +34,19 @@ TAutoConsoleVariable<int32> CVarPathTracingCompaction(
 	ECVF_RenderThreadSafe
 );
 
+TAutoConsoleVariable<int32> CVarPathTracingIndirectDispatch(
+	TEXT("r.PathTracing.IndirectDispatch"),
+	0,
+	TEXT("Enables indirect dispatch (if supported by the hardware) for compacted path tracing (default: 0 (disabled))"),
+	ECVF_RenderThreadSafe
+);
+
+TAutoConsoleVariable<int32> CVarPathTracingDispatchSize(
+	TEXT("r.PathTracing.DispatchSize"),
+	2048,
+	TEXT("Controls the tile size used when rendering the image. Reducing this value may prevent GPU timeouts for heavy renders. (default = 2048)"),
+	ECVF_RenderThreadSafe
+);
 
 TAutoConsoleVariable<int32> CVarPathTracingMaxBounces(
 	TEXT("r.PathTracing.MaxBounces"),
@@ -477,30 +490,19 @@ class FPathTracingRG : public FGlobalShader
 	DECLARE_GLOBAL_SHADER(FPathTracingRG)
 	SHADER_USE_ROOT_PARAMETER_STRUCT(FPathTracingRG, FGlobalShader)
 
-	class FEnableCompactionDim : SHADER_PERMUTATION_BOOL("PATH_TRACER_USE_COMPACTION");
+	class FCompactionType : SHADER_PERMUTATION_INT("PATH_TRACER_USE_COMPACTION", 2);
 
-	using FPermutationDomain = TShaderPermutationDomain<FEnableCompactionDim>;
+	using FPermutationDomain = TShaderPermutationDomain<FCompactionType>;
 
 	static bool ShouldCompilePermutation(const FGlobalShaderPermutationParameters& Parameters)
 	{
-		FPermutationDomain PermutationVector(Parameters.PermutationId);
-		const bool bUseCompaction = PermutationVector.Get<FEnableCompactionDim>();
-		if (bUseCompaction && !RHISupportsWaveOperations(Parameters.Platform))
-		{
-			// compaction requires wave operations
-			return false;
-		}
 		return ShouldCompilePathTracingShadersForProject(Parameters.Platform);
 	}
 
 	static void ModifyCompilationEnvironment(const FGlobalShaderPermutationParameters& Parameters, FShaderCompilerEnvironment& OutEnvironment)
 	{
+		FGlobalShader::ModifyCompilationEnvironment(Parameters, OutEnvironment);
 		OutEnvironment.SetDefine(TEXT("USE_RECT_LIGHT_TEXTURES"), 1);
-		FPermutationDomain PermutationVector(Parameters.PermutationId);
-		if (PermutationVector.Get<FEnableCompactionDim>())
-		{
-			OutEnvironment.CompilerFlags.Add(CFLAG_WaveOperations);
-		}
 		OutEnvironment.CompilerFlags.Add(CFLAG_WarningsAsErrors);
 	}
 
@@ -533,10 +535,12 @@ class FPathTracingRG : public FGlobalShader
 
 		// extra parameters required for path compacting kernel
 		SHADER_PARAMETER(int, Bounce)
-		SHADER_PARAMETER_RDG_BUFFER_UAV(RWBuffer<uint>, PathStateData)
+		SHADER_PARAMETER_RDG_BUFFER_UAV(RWStructuredBuffer<FPathTracingPackedPathState>, PathStateData)
 		SHADER_PARAMETER_RDG_BUFFER_SRV(Buffer<int>, ActivePaths)
 		SHADER_PARAMETER_RDG_BUFFER_UAV(RWBuffer<int>, NextActivePaths)
 		SHADER_PARAMETER_RDG_BUFFER_UAV(RWBuffer<int>, NumPathStates)
+
+		RDG_BUFFER_ACCESS(PathTracingIndirectArgs, ERHIAccess::IndirectArgs | ERHIAccess::SRVCompute)
 	END_SHADER_PARAMETER_STRUCT()
 };
 IMPLEMENT_GLOBAL_SHADER(FPathTracingRG, "/Engine/Private/PathTracing/PathTracing.usf", "PathTracingMainRG", SF_RayGen);
@@ -1301,10 +1305,10 @@ void FDeferredShadingSceneRenderer::PreparePathTracing(const FSceneViewFamily& V
 		&& ShouldCompilePathTracingShadersForProject(ViewFamily.GetShaderPlatform()))
 	{
 		// Declare all RayGen shaders that require material closest hit shaders to be bound
-		for (int EnableCompaction = 0; EnableCompaction < 2; EnableCompaction++)
+		for (int CompactionType = 0; CompactionType < FPathTracingRG::FCompactionType::PermutationCount; CompactionType++)
 		{
 			FPathTracingRG::FPermutationDomain PermutationVector;
-			PermutationVector.Set<FPathTracingRG::FEnableCompactionDim>(EnableCompaction != 0);
+			PermutationVector.Set<FPathTracingRG::FCompactionType>(CompactionType);
 			FGlobalShaderPermutationParameters Parameters(FPathTracingRG::StaticGetTypeLayout().Name, ViewFamily.GetShaderPlatform(), PermutationVector.ToDimensionValueId());
 			if (!FPathTracingRG::ShouldCompilePermutation(Parameters))
 			{
@@ -1480,98 +1484,146 @@ void FDeferredShadingSceneRenderer::RenderPathTracing(
 	{
 		// Round up to coherent path tracing tile size to simplify pixel shuffling
 		// TODO: be careful not to write extra pixels past the boundary when using multi-gpu
-		const int32 DispatchSizeX = FMath::DivideAndRoundUp(View.ViewRect.Size().X, PATHTRACER_COHERENT_TILE_SIZE) * PATHTRACER_COHERENT_TILE_SIZE;
-		const int32 DispatchSizeY = FMath::DivideAndRoundUp(View.ViewRect.Size().Y, PATHTRACER_COHERENT_TILE_SIZE) * PATHTRACER_COHERENT_TILE_SIZE;
+
+		const int32 ResX = View.ViewRect.Size().X;
+		const int32 ResY = View.ViewRect.Size().Y;
+
+		const int32 DispatchResX = FMath::DivideAndRoundUp(ResX, PATHTRACER_COHERENT_TILE_SIZE) * PATHTRACER_COHERENT_TILE_SIZE;
+		const int32 DispatchResY = FMath::DivideAndRoundUp(ResY, PATHTRACER_COHERENT_TILE_SIZE) * PATHTRACER_COHERENT_TILE_SIZE;
+
+		int32 DispatchSize = FMath::Max(CVarPathTracingDispatchSize.GetValueOnRenderThread(), 64);
+		DispatchSize = FMath::DivideAndRoundUp(DispatchSize, PATHTRACER_COHERENT_TILE_SIZE) * PATHTRACER_COHERENT_TILE_SIZE;
 
 		// should we use path compaction?
-		const bool bUseCompaction = CVarPathTracingCompaction.GetValueOnRenderThread() != 0 && GRHISupportsWaveOperations && RHISupportsWaveOperations(View.GetShaderPlatform());
+		const int CompactionType = CVarPathTracingCompaction.GetValueOnRenderThread();
+		const bool bUseIndirectDispatch = GRHISupportsRayTracingDispatchIndirect && CVarPathTracingIndirectDispatch.GetValueOnRenderThread() != 0;
 		FRDGBuffer* ActivePaths[2] = {};
-		FRDGBuffer* NumActivePaths = nullptr;
+		FRDGBuffer* NumActivePaths[2] = {};
 		FRDGBuffer* PathStateData = nullptr;
-		if (bUseCompaction)
-		{
-			ActivePaths[0] = GraphBuilder.CreateBuffer(FRDGBufferDesc::CreateBufferDesc(sizeof(int32), DispatchSizeX * DispatchSizeY), TEXT("PathTracer.ActivePaths0"));
-			ActivePaths[1] = GraphBuilder.CreateBuffer(FRDGBufferDesc::CreateBufferDesc(sizeof(int32), DispatchSizeX * DispatchSizeY), TEXT("PathTracer.ActivePaths1"));
-			NumActivePaths = GraphBuilder.CreateBuffer(FRDGBufferDesc::CreateBufferDesc(sizeof(int32), 1), TEXT("PathTracer.NumActivePaths"));
-			PathStateData = GraphBuilder.CreateBuffer(FRDGBufferDesc::CreateBufferDesc(sizeof(int32) * PATHTRACER_PATHSTATE_NUM_WORDS, DispatchSizeX * DispatchSizeY), TEXT("PathTracer.PathStateData"));
-			AddClearUAVPass(GraphBuilder, GraphBuilder.CreateUAV(ActivePaths[0], PF_R32_UINT), 0);
-		}
-		FPathTracingRG::FPermutationDomain PermutationVector;
-		PermutationVector.Set<FPathTracingRG::FEnableCompactionDim>(bUseCompaction);
-		TShaderMapRef<FPathTracingRG> RayGenShader(View.ShaderMap, PermutationVector);
-		FPathTracingRG::FParameters* PreviousPassParameters = nullptr;
-		// When using path compaction, we need to run the path tracer once per bounce
-		// otherwise, the path tracer is the one doing the bounces
-		for (int Bounce = 0, MaxBounces = bUseCompaction ? Config.PathTracingData.MaxBounces : 0; Bounce <= MaxBounces; Bounce++)
+		if (CompactionType == 1)
 		{
-			FPathTracingRG::FParameters* PassParameters = GraphBuilder.AllocParameters<FPathTracingRG::FParameters>();
-			PassParameters->TLAS = View.GetRayTracingSceneViewChecked();
-			PassParameters->ViewUniformBuffer = View.ViewUniformBuffer;
-			PassParameters->PathTracingData = Config.PathTracingData;
-			if (PreviousPassParameters == nullptr)
+			const int32 NumPaths = FMath::Min(DispatchSize * DispatchSize, DispatchResX * DispatchResY);
+			ActivePaths[0] = GraphBuilder.CreateBuffer(FRDGBufferDesc::CreateBufferDesc(sizeof(int32), NumPaths), TEXT("PathTracer.ActivePaths0"));
+			ActivePaths[1] = GraphBuilder.CreateBuffer(FRDGBufferDesc::CreateBufferDesc(sizeof(int32), NumPaths), TEXT("PathTracer.ActivePaths1"));
+			if (bUseIndirectDispatch)
 			{
-				// upload sky/lights data
-				SetLightParameters(GraphBuilder, PassParameters, Scene, View, Config.UseMISCompensation);
-				PreviousPassParameters = PassParameters;
+				NumActivePaths[0] = GraphBuilder.CreateBuffer(FRDGBufferDesc::CreateIndirectDesc<int32>(3), TEXT("PathTracer.NumActivePaths0"));
+				NumActivePaths[1] = GraphBuilder.CreateBuffer(FRDGBufferDesc::CreateIndirectDesc<int32>(3), TEXT("PathTracer.NumActivePaths1"));
 			}
 			else
 			{
-				// re-use from last iteration
-				PassParameters->IESTexture = PreviousPassParameters->IESTexture;
-				PassParameters->IESTextureSampler = PreviousPassParameters->IESTextureSampler;
-				PassParameters->LightGridParameters = PreviousPassParameters->LightGridParameters;
-				PassParameters->RectLightTexture = PreviousPassParameters->RectLightTexture;
-				PassParameters->RectLightSampler = PreviousPassParameters->RectLightSampler;
-				PassParameters->SceneLightCount = PreviousPassParameters->SceneLightCount;
-				PassParameters->SceneVisibleLightCount = PreviousPassParameters->SceneVisibleLightCount;
-				PassParameters->SceneLights = PreviousPassParameters->SceneLights;
-				PassParameters->SkylightParameters = PreviousPassParameters->SkylightParameters;
+				NumActivePaths[0] = GraphBuilder.CreateBuffer(FRDGBufferDesc::CreateBufferDesc(sizeof(int32), 3), TEXT("PathTracer.NumActivePaths"));
 			}
-			if (Config.PathTracingData.EnableDirectLighting == 0)
+			PathStateData = GraphBuilder.CreateBuffer(FRDGBufferDesc::CreateStructuredDesc(sizeof(FPathTracingPackedPathState), NumPaths), TEXT("PathTracer.PathStateData"));
+		}
+		FPathTracingRG::FPermutationDomain PermutationVector;
+		PermutationVector.Set<FPathTracingRG::FCompactionType>(CompactionType);
+		TShaderMapRef<FPathTracingRG> RayGenShader(View.ShaderMap, PermutationVector);
+		FPathTracingRG::FParameters* PreviousPassParameters = nullptr;
+		for (int32 TileY = 0; TileY < DispatchResY; TileY += DispatchSize)
+		{
+			for (int32 TileX = 0; TileX < DispatchResX; TileX += DispatchSize)
 			{
-				PassParameters->SceneVisibleLightCount = 0;
-			}
+				const int32 DispatchSizeX = FMath::Min(DispatchSize, DispatchResX - TileX);
+				const int32 DispatchSizeY = FMath::Min(DispatchSize, DispatchResY - TileY);
+				if (CompactionType == 1)
+				{
+					AddClearUAVPass(GraphBuilder, GraphBuilder.CreateUAV(ActivePaths[0], PF_R32_UINT), 0);
+				}
+				// When using path compaction, we need to run the path tracer once per bounce
+				// otherwise, the path tracer is the one doing the bounces
+				for (int Bounce = 0, MaxBounces = CompactionType == 1 ? Config.PathTracingData.MaxBounces : 0; Bounce <= MaxBounces; Bounce++)
+				{
+					FPathTracingRG::FParameters* PassParameters = GraphBuilder.AllocParameters<FPathTracingRG::FParameters>();
+					PassParameters->TLAS = View.GetRayTracingSceneViewChecked();
+					PassParameters->ViewUniformBuffer = View.ViewUniformBuffer;
+					PassParameters->PathTracingData = Config.PathTracingData;
+					if (PreviousPassParameters == nullptr)
+					{
+						// upload sky/lights data
+						SetLightParameters(GraphBuilder, PassParameters, Scene, View, Config.UseMISCompensation);
+						PreviousPassParameters = PassParameters;
+					}
+					else
+					{
+						// re-use from last iteration
+						PassParameters->IESTexture = PreviousPassParameters->IESTexture;
+						PassParameters->IESTextureSampler = PreviousPassParameters->IESTextureSampler;
+						PassParameters->LightGridParameters = PreviousPassParameters->LightGridParameters;
+						PassParameters->RectLightTexture = PreviousPassParameters->RectLightTexture;
+						PassParameters->RectLightSampler = PreviousPassParameters->RectLightSampler;
+						PassParameters->SceneLightCount = PreviousPassParameters->SceneLightCount;
+						PassParameters->SceneVisibleLightCount = PreviousPassParameters->SceneVisibleLightCount;
+						PassParameters->SceneLights = PreviousPassParameters->SceneLights;
+						PassParameters->SkylightParameters = PreviousPassParameters->SkylightParameters;
+					}
+					if (Config.PathTracingData.EnableDirectLighting == 0)
+					{
+						PassParameters->SceneVisibleLightCount = 0;
+					}
 
-			PassParameters->IESTextureSampler = TStaticSamplerState<SF_Bilinear, AM_Clamp, AM_Clamp, AM_Clamp>::GetRHI();
-			PassParameters->RadianceTexture = GraphBuilder.CreateUAV(RadianceTexture);
-			PassParameters->AlbedoTexture   = GraphBuilder.CreateUAV(AlbedoTexture);
-			PassParameters->NormalTexture   = GraphBuilder.CreateUAV(NormalTexture);
+					PassParameters->IESTextureSampler = TStaticSamplerState<SF_Bilinear, AM_Clamp, AM_Clamp, AM_Clamp>::GetRHI();
+					PassParameters->RadianceTexture = GraphBuilder.CreateUAV(RadianceTexture);
+					PassParameters->AlbedoTexture   = GraphBuilder.CreateUAV(AlbedoTexture);
+					PassParameters->NormalTexture   = GraphBuilder.CreateUAV(NormalTexture);
 
-			// TODO: in multi-gpu case, split image into tiles
-			PassParameters->TileOffset.X = 0;
-			PassParameters->TileOffset.Y = 0;
+					// TODO: in multi-gpu case, assign different tiles to different GPUs
+					PassParameters->TileOffset.X = TileX;
+					PassParameters->TileOffset.Y = TileY;
 
-			if (bUseCompaction)
-			{
-				PassParameters->Bounce = Bounce;
-				PassParameters->ActivePaths = GraphBuilder.CreateSRV(ActivePaths[Bounce & 1], PF_R32_UINT);
-				PassParameters->NextActivePaths = GraphBuilder.CreateUAV(ActivePaths[(Bounce & 1) ^ 1], PF_R32_UINT);
-				PassParameters->PathStateData = GraphBuilder.CreateUAV(PathStateData, PF_R32_UINT);
-				PassParameters->NumPathStates = GraphBuilder.CreateUAV(NumActivePaths, PF_R32_UINT);
-				AddClearUAVPass(GraphBuilder, PassParameters->NextActivePaths, ~0u); // make sure everything is initialized to -1 since paths that go inactive don't write anything
-				AddClearUAVPass(GraphBuilder, PassParameters->NumPathStates, 0); // reset number of active paths to 0
-			}
-			ClearUnusedGraphResources(RayGenShader, PassParameters);
-			GraphBuilder.AddPass(
-				bUseCompaction
-					? RDG_EVENT_NAME("Path Tracer Compute (%d x %d) Sample=%d/%d NumLights=%d Bounce=%d", View.ViewRect.Size().X, View.ViewRect.Size().Y, PathTracingState->SampleIndex, MaxSPP, PassParameters->SceneLightCount, Bounce)
-					: RDG_EVENT_NAME("Path Tracer Compute (%d x %d) Sample=%d/%d NumLights=%d", View.ViewRect.Size().X, View.ViewRect.Size().Y, PathTracingState->SampleIndex, MaxSPP, PassParameters->SceneLightCount),
-				PassParameters,
-				ERDGPassFlags::Compute,
-				[PassParameters, RayGenShader, DispatchSizeX, DispatchSizeY, &View](FRHIRayTracingCommandList& RHICmdList)
-			{
-				FRHIRayTracingScene* RayTracingSceneRHI = View.GetRayTracingSceneChecked();
+					PassParameters->Bounce = Bounce;
+					if (CompactionType == 1)
+					{
+						PassParameters->ActivePaths = GraphBuilder.CreateSRV(ActivePaths[Bounce & 1], PF_R32_SINT);
+						PassParameters->NextActivePaths = GraphBuilder.CreateUAV(ActivePaths[(Bounce & 1) ^ 1], PF_R32_SINT);
+						PassParameters->PathStateData = GraphBuilder.CreateUAV(PathStateData);
+						if (bUseIndirectDispatch)
+						{
+							PassParameters->NumPathStates = GraphBuilder.CreateUAV(NumActivePaths[Bounce & 1], PF_R32_UINT);
+							PassParameters->PathTracingIndirectArgs = NumActivePaths[(Bounce & 1) ^ 1];
+						}
+						else
+						{
+							PassParameters->NumPathStates = GraphBuilder.CreateUAV(NumActivePaths[0], PF_R32_UINT);
+							AddClearUAVPass(GraphBuilder, PassParameters->NextActivePaths, -1); // make sure everything is initialized to -1 since paths that go inactive don't write anything
+						}
+						AddClearUAVPass(GraphBuilder, PassParameters->NumPathStates, 0);
+					}
+					ClearUnusedGraphResources(RayGenShader, PassParameters);
+					GraphBuilder.AddPass(
+						CompactionType == 1
+							? RDG_EVENT_NAME("Path Tracer Compute (%d x %d) Tile=(%d,%d - %dx%d) Sample=%d/%d NumLights=%d (Bounce=%d%s)", ResX, ResY, TileX, TileY, DispatchSizeX, DispatchSizeY, PathTracingState->SampleIndex, MaxSPP, PassParameters->SceneLightCount, Bounce, bUseIndirectDispatch && Bounce > 0 ? TEXT(" indirect") : TEXT(""))
+							: RDG_EVENT_NAME("Path Tracer Compute (%d x %d) Tile=(%d,%d - %dx%d) Sample=%d/%d NumLights=%d", ResX, ResY, TileX, TileY, DispatchSizeX, DispatchSizeY, PathTracingState->SampleIndex, MaxSPP, PassParameters->SceneLightCount),
+						PassParameters,
+						ERDGPassFlags::Compute,
+						[PassParameters, RayGenShader, DispatchSizeX, DispatchSizeY, bUseIndirectDispatch, &View](FRHIRayTracingCommandList& RHICmdList)
+					{
+						FRHIRayTracingScene* RayTracingSceneRHI = View.GetRayTracingSceneChecked();
 			
-				FRayTracingShaderBindingsWriter GlobalResources;
-				SetShaderParameters(GlobalResources, RayGenShader, *PassParameters);
-
-				RHICmdList.RayTraceDispatch(
-					View.RayTracingMaterialPipeline,
-					RayGenShader.GetRayTracingShader(),
-					RayTracingSceneRHI, GlobalResources,
-					DispatchSizeX, DispatchSizeY
-				);
-			});
+						FRayTracingShaderBindingsWriter GlobalResources;
+						SetShaderParameters(GlobalResources, RayGenShader, *PassParameters);
+						if (bUseIndirectDispatch && PassParameters->Bounce > 0)
+						{
+							PassParameters->PathTracingIndirectArgs->MarkResourceAsUsed();
+							RHICmdList.RayTraceDispatchIndirect(
+								View.RayTracingMaterialPipeline,
+								RayGenShader.GetRayTracingShader(),
+								RayTracingSceneRHI, GlobalResources,
+								PassParameters->PathTracingIndirectArgs->GetIndirectRHICallBuffer(), 0
+							);
+						}
+						else
+						{
+							RHICmdList.RayTraceDispatch(
+								View.RayTracingMaterialPipeline,
+								RayGenShader.GetRayTracingShader(),
+								RayTracingSceneRHI, GlobalResources,
+								DispatchSizeX, DispatchSizeY
+							);
+						}
+					});
+				}
+			}
 		}
 		// After we are done, make sure we remember our texture for next time so that we can accumulate samples across frames
 		GraphBuilder.QueueTextureExtraction(RadianceTexture, &PathTracingState->RadianceRT);
