Convert Standard deferred lights shader to use FParameters shader parameters and native RDG resources.

diff --git a/Engine/Shaders/Private/StationaryLightOverlapShaders.usf b/Engine/Shaders/Private/StationaryLightOverlapShaders.usf
--- a/Engine/Shaders/Private/StationaryLightOverlapShaders.usf
+++ b/Engine/Shaders/Private/StationaryLightOverlapShaders.usf
@@ -6,10 +6,10 @@
 
 #include "Common.ush"
 #include "DeferredShadingCommon.ush"
-#include "DynamicLightingCommon.ush"
+#include "DynamicLightingCommon.ush" 
 
 /** 1 if the light has a valid shadowmap channel. */
-float HasValidChannel;
+float bHasValidChannel;
 
 #if RADIAL_ATTENUATION == 0
 void OverlapPixelMain(
@@ -28,7 +28,7 @@ void OverlapPixelMain(
 		// Output minor complexity to visualize overlap
 		OutColor.rgb = .04f;
 
-		if (HasValidChannel < 1)
+		if (bHasValidChannel < 1)
 		{
 			// Output a huge amount of complexity to make it clear that this light is an error
 			OutColor.rgb = .4f;
@@ -67,7 +67,7 @@ void OverlapPixelMain(
 			// Output minor complexity to visualize overlap
 			OutColor.rgb = .04f;
 
-			if (HasValidChannel < 1)
+			if (bHasValidChannel < 1)
 			{
 				// Output a huge amount of complexity to make it clear that this light is an error
 				OutColor.rgb = .4f;
diff --git a/Engine/Source/Runtime/Renderer/Private/DeferredShadingRenderer.cpp b/Engine/Source/Runtime/Renderer/Private/DeferredShadingRenderer.cpp
--- a/Engine/Source/Runtime/Renderer/Private/DeferredShadingRenderer.cpp
+++ b/Engine/Source/Runtime/Renderer/Private/DeferredShadingRenderer.cpp
@@ -2943,7 +2943,7 @@ void FDeferredShadingSceneRenderer::Render(FRDGBuilder& GraphBuilder)
 	}
 	else if (HairStrands::HasViewHairStrandsData(Views) && ViewFamily.EngineShowFlags.Lighting)
 	{
-		RenderLightsForHair(GraphBuilder, SceneTextures.UniformBuffer, SortedLightSet, ForwardScreenSpaceShadowMaskHairTexture, LightingChannelsTexture);
+		RenderLightsForHair(GraphBuilder, SceneTextures, SortedLightSet, ForwardScreenSpaceShadowMaskHairTexture, LightingChannelsTexture);
 		RenderDeferredReflectionsAndSkyLightingHair(GraphBuilder);
 	}
 
diff --git a/Engine/Source/Runtime/Renderer/Private/DeferredShadingRenderer.h b/Engine/Source/Runtime/Renderer/Private/DeferredShadingRenderer.h
--- a/Engine/Source/Runtime/Renderer/Private/DeferredShadingRenderer.h
+++ b/Engine/Source/Runtime/Renderer/Private/DeferredShadingRenderer.h
@@ -693,7 +693,7 @@ private:
 
 	void RenderLightsForHair(
 		FRDGBuilder& GraphBuilder,
-		TRDGUniformBufferRef<FSceneTextureUniformParameters> SceneTexturesUniformBuffer,
+		const FMinimalSceneTextures& SceneTextures,
 		FSortedLightSetSceneInfo& SortedLightSet,
 		FRDGTextureRef InScreenShadowMaskSubPixelTexture,
 		FRDGTextureRef LightingChannelsTexture);
@@ -702,7 +702,7 @@ private:
 	void RenderLightForHair(
 		FRDGBuilder& GraphBuilder,
 		FViewInfo& View,
-		TRDGUniformBufferRef<FSceneTextureUniformParameters> SceneTexturesUniformBuffer,
+		const FMinimalSceneTextures& SceneTextures,
 		const FLightSceneInfo* LightSceneInfo,
 		FRDGTextureRef ScreenShadowMaskSubPixelTexture,
 		FRDGTextureRef LightingChannelsTexture,
diff --git a/Engine/Source/Runtime/Renderer/Private/LightRendering.cpp b/Engine/Source/Runtime/Renderer/Private/LightRendering.cpp
--- a/Engine/Source/Runtime/Renderer/Private/LightRendering.cpp
+++ b/Engine/Source/Runtime/Renderer/Private/LightRendering.cpp
@@ -193,15 +193,15 @@ static void RenderLight(
 
 FDeferredLightUniformStruct GetDeferredLightParameters(const FSceneView& View, const FLightSceneInfo& LightSceneInfo)
 {
-	FDeferredLightUniformStruct Parameters;
-	LightSceneInfo.Proxy->GetLightShaderParameters(Parameters.LightParameters);
+	FDeferredLightUniformStruct Out;
+	LightSceneInfo.Proxy->GetLightShaderParameters(Out.LightParameters);
 
 	const bool bIsRayTracedLight = ShouldRenderRayTracingShadowsForLight(*LightSceneInfo.Proxy);
 
 	const FVector2D FadeParams = LightSceneInfo.Proxy->GetDirectionalLightDistanceFadeParameters(View.GetFeatureLevel(), !bIsRayTracedLight && LightSceneInfo.IsPrecomputedLightingValid(), View.MaxShadowCascades);
 	
 	// use MAD for efficiency in the shader
-	Parameters.DistanceFadeMAD = FVector2D(FadeParams.Y, -FadeParams.X * FadeParams.Y);
+	Out.DistanceFadeMAD = FVector2D(FadeParams.Y, -FadeParams.X * FadeParams.Y);
 	
 	int32 ShadowMapChannel = LightSceneInfo.Proxy->GetShadowMapChannel();
 
@@ -213,7 +213,7 @@ FDeferredLightUniformStruct GetDeferredLightParameters(const FSceneView& View, c
 		ShadowMapChannel = INDEX_NONE;
 	}
 
-	Parameters.ShadowMapChannelMask = FVector4f(
+	Out.ShadowMapChannelMask = FVector4f(
 		ShadowMapChannel == 0 ? 1 : 0,
 		ShadowMapChannel == 1 ? 1 : 0,
 		ShadowMapChannel == 2 ? 1 : 0,
@@ -221,66 +221,74 @@ FDeferredLightUniformStruct GetDeferredLightParameters(const FSceneView& View, c
 
 	const bool bDynamicShadows = View.Family->EngineShowFlags.DynamicShadows && GetShadowQuality() > 0;
 	const bool bHasLightFunction = LightSceneInfo.Proxy->GetLightFunctionMaterial() != NULL;
-	Parameters.ShadowedBits = LightSceneInfo.Proxy->CastsStaticShadow() || bHasLightFunction ? 1 : 0;
-	Parameters.ShadowedBits |= LightSceneInfo.Proxy->CastsDynamicShadow() && View.Family->EngineShowFlags.DynamicShadows ? 3 : 0;
+	Out.ShadowedBits = LightSceneInfo.Proxy->CastsStaticShadow() || bHasLightFunction ? 1 : 0;
+	Out.ShadowedBits |= LightSceneInfo.Proxy->CastsDynamicShadow() && View.Family->EngineShowFlags.DynamicShadows ? 3 : 0;
 
-	Parameters.VolumetricScatteringIntensity = LightSceneInfo.Proxy->GetVolumetricScatteringIntensity();
+	Out.VolumetricScatteringIntensity = LightSceneInfo.Proxy->GetVolumetricScatteringIntensity();
 
 	static auto* ContactShadowsCVar = IConsoleManager::Get().FindTConsoleVariableDataInt(TEXT("r.ContactShadows"));
 	static auto* IntensityCVar = IConsoleManager::Get().FindTConsoleVariableDataFloat(TEXT("r.ContactShadows.NonShadowCastingIntensity"));
 
-	Parameters.ContactShadowLength = 0;
-	Parameters.ContactShadowNonShadowCastingIntensity = 0.0f;
+	Out.ContactShadowLength = 0;
+	Out.ContactShadowNonShadowCastingIntensity = 0.0f;
 
 	if (ContactShadowsCVar && ContactShadowsCVar->GetValueOnRenderThread() != 0 && View.Family->EngineShowFlags.ContactShadows)
 	{
-		Parameters.ContactShadowLength = LightSceneInfo.Proxy->GetContactShadowLength();
+		Out.ContactShadowLength = LightSceneInfo.Proxy->GetContactShadowLength();
 		// Sign indicates if contact shadow length is in world space or screen space.
 		// Multiply by 2 for screen space in order to preserve old values after introducing multiply by View.ClipToView[1][1] in shader.
-		Parameters.ContactShadowLength *= LightSceneInfo.Proxy->IsContactShadowLengthInWS() ? -1.0f : 2.0f;
+		Out.ContactShadowLength *= LightSceneInfo.Proxy->IsContactShadowLengthInWS() ? -1.0f : 2.0f;
 
-		Parameters.ContactShadowNonShadowCastingIntensity = IntensityCVar ? IntensityCVar->GetValueOnRenderThread() : 0.0f;
+		Out.ContactShadowNonShadowCastingIntensity = IntensityCVar ? IntensityCVar->GetValueOnRenderThread() : 0.0f;
 	}
 
 	// When rendering reflection captures, the direct lighting of the light is actually the indirect specular from the main view
 	if (View.bIsReflectionCapture)
 	{
-		Parameters.LightParameters.Color *= LightSceneInfo.Proxy->GetIndirectLightingScale();
+		Out.LightParameters.Color *= LightSceneInfo.Proxy->GetIndirectLightingScale();
 	}
 
 	const ELightComponentType LightType = (ELightComponentType)LightSceneInfo.Proxy->GetLightType();
 	if ((LightType == LightType_Point || LightType == LightType_Spot || LightType == LightType_Rect) && View.IsPerspectiveProjection())
 	{
-		Parameters.LightParameters.Color *= GetLightFadeFactor(View, LightSceneInfo.Proxy);
+		Out.LightParameters.Color *= GetLightFadeFactor(View, LightSceneInfo.Proxy);
 	}
 
-	Parameters.LightingChannelMask = LightSceneInfo.Proxy->GetLightingChannelMask();
+	Out.LightingChannelMask = LightSceneInfo.Proxy->GetLightingChannelMask();
 
-	return Parameters;
+	return Out;
 }
 
-void SetupSimpleDeferredLightParameters(
+FDeferredLightUniformStruct GetSimpleDeferredLightParameters(
 	const FSimpleLightEntry& SimpleLight,
-	const FSimpleLightPerViewEntry &SimpleLightPerViewData,
-	FDeferredLightUniformStruct& DeferredLightUniformsValue)
+	const FVector3f& Position)
 {
-	DeferredLightUniformsValue.LightParameters.Position = SimpleLightPerViewData.Position;
-	DeferredLightUniformsValue.LightParameters.InvRadius = 1.0f / FMath::Max(SimpleLight.Radius, KINDA_SMALL_NUMBER);
-	DeferredLightUniformsValue.LightParameters.Color = SimpleLight.Color;
-	DeferredLightUniformsValue.LightParameters.FalloffExponent = SimpleLight.Exponent;
-	DeferredLightUniformsValue.LightParameters.Direction = FVector(1, 0, 0);
-	DeferredLightUniformsValue.LightParameters.Tangent = FVector(1, 0, 0);
-	DeferredLightUniformsValue.LightParameters.SpotAngles = FVector2D(-2, 1);
-	DeferredLightUniformsValue.LightParameters.SpecularScale = 1.0f;
-	DeferredLightUniformsValue.LightParameters.SourceRadius = 0.0f;
-	DeferredLightUniformsValue.LightParameters.SoftSourceRadius = 0.0f;
-	DeferredLightUniformsValue.LightParameters.SourceLength = 0.0f;
-	DeferredLightUniformsValue.LightParameters.SourceTexture = GWhiteTexture->TextureRHI;
-	DeferredLightUniformsValue.ContactShadowLength = 0.0f;
-	DeferredLightUniformsValue.DistanceFadeMAD = FVector2D(0, 0);
-	DeferredLightUniformsValue.ShadowMapChannelMask = FVector4f(0, 0, 0, 0);
-	DeferredLightUniformsValue.ShadowedBits = 0;
-	DeferredLightUniformsValue.LightingChannelMask = 0;
+	FDeferredLightUniformStruct Out;
+	Out.LightParameters.Position = Position;
+	Out.LightParameters.InvRadius = 1.0f / FMath::Max(SimpleLight.Radius, KINDA_SMALL_NUMBER);
+	Out.LightParameters.Color = SimpleLight.Color;
+	Out.LightParameters.FalloffExponent = SimpleLight.Exponent;
+	Out.LightParameters.Direction = FVector(1, 0, 0);
+	Out.LightParameters.Tangent = FVector(1, 0, 0);
+	Out.LightParameters.SpotAngles = FVector2D(-2, 1);
+	Out.LightParameters.SpecularScale = 1.0f;
+	Out.LightParameters.SourceRadius = 0.0f;
+	Out.LightParameters.SoftSourceRadius = 0.0f;
+	Out.LightParameters.SourceLength = 0.0f;
+	Out.LightParameters.SourceTexture = GWhiteTexture->TextureRHI;
+	Out.ContactShadowLength = 0.0f;
+	Out.DistanceFadeMAD = FVector2D(0, 0);
+	Out.ShadowMapChannelMask = FVector4f(0, 0, 0, 0);
+	Out.ShadowedBits = 0;
+	Out.LightingChannelMask = 0;
+
+	return Out;
+}
+FDeferredLightUniformStruct GetSimpleDeferredLightParameters(
+	const FSimpleLightEntry& SimpleLight,
+	const FSimpleLightPerViewEntry& SimpleLightPerViewData)
+{
+	return GetSimpleDeferredLightParameters(SimpleLight, SimpleLightPerViewData.Position);
 }
 
 FLightOcclusionType GetLightOcclusionType(const FLightSceneProxy& Proxy)
@@ -363,26 +371,14 @@ TGlobalResource<StencilingGeometry::FStencilConeIndexBuffer> StencilingGeometry:
 // Implement a version for directional lights, and a version for point / spot lights
 IMPLEMENT_GLOBAL_SHADER(FDeferredLightVS, "/Engine/Private/DeferredLightVertexShaders.usf", "VertexMain", SF_Vertex);
 
-struct FRenderLightParams
-{
-	// Precompute transmittance
-	FShaderResourceViewRHIRef DeepShadow_TransmittanceMaskBuffer = nullptr;
-	uint32 DeepShadow_TransmittanceMaskBufferMaxCount = 0;
-	FRHITexture* ScreenShadowMaskSubPixelTexture = nullptr;
-	FVector4f ShadowChannelMask = FVector4f(1,1,1,1);
-
-	// Cloud shadow data
-	FMatrix Cloud_WorldToLightClipShadowMatrix;
-	float Cloud_ShadowmapFarDepthKm = 0.0f;
-	FRHITexture* Cloud_ShadowmapTexture = nullptr;
-	float Cloud_ShadowmapStrength = 0.0f;
-};
-
-
 class FDeferredLightHairVS : public FGlobalShader
 {
 	DECLARE_SHADER_TYPE(FDeferredLightHairVS, Global);
-public:
+	SHADER_USE_PARAMETER_STRUCT(FDeferredLightHairVS, FGlobalShader);
+
+	BEGIN_SHADER_PARAMETER_STRUCT(FParameters, )
+		SHADER_PARAMETER_RDG_UNIFORM_BUFFER(FHairStrandsViewUniformParameters, HairStrands)
+	END_SHADER_PARAMETER_STRUCT()
 
 	static bool ShouldCompilePermutation(const FGlobalShaderPermutationParameters& Parameters)
 	{
@@ -394,27 +390,6 @@ public:
 		FGlobalShader::ModifyCompilationEnvironment(Parameters, OutEnvironment);
 		OutEnvironment.SetDefine(TEXT("SHADER_HAIR"), 1);
 	}
-
-	FDeferredLightHairVS() {}
-	FDeferredLightHairVS(const ShaderMetaType::CompiledShaderInitializerType& Initializer) :
-		FGlobalShader(Initializer)
-	{
-		HairStrandsParameters.Bind(Initializer.ParameterMap, FHairStrandsViewUniformParameters::StaticStructMetadata.GetShaderVariableName());
-	}
-
-	void SetParameters(FRHICommandList& RHICmdList, const FViewInfo& View, FRHIUniformBuffer* HairStrandsUniformBuffer)
-	{
-		FRHIVertexShader* ShaderRHI = RHICmdList.GetBoundVertexShader();
-		FGlobalShader::SetParameters<FViewUniformShaderParameters>(RHICmdList, ShaderRHI, View.ViewUniformBuffer);
-
-		if (HairStrandsUniformBuffer)
-		{
-			SetUniformBufferParameter(RHICmdList, ShaderRHI, HairStrandsParameters, HairStrandsUniformBuffer);
-		}
-	}
-
-private:
-	LAYOUT_FIELD(FShaderUniformBufferParameter, HairStrandsParameters);
 };
 
 IMPLEMENT_GLOBAL_SHADER(FDeferredLightHairVS, "/Engine/Private/DeferredLightVertexShaders.usf", "HairVertexMain", SF_Vertex);
@@ -433,6 +408,7 @@ enum class ELightSourceShape
 class FDeferredLightPS : public FGlobalShader
 {
 	DECLARE_SHADER_TYPE(FDeferredLightPS, Global)
+	SHADER_USE_PARAMETER_STRUCT(FDeferredLightPS, FGlobalShader);
 
 	class FSourceShapeDim		: SHADER_PERMUTATION_ENUM_CLASS("LIGHT_SOURCE_SHAPE", ELightSourceShape);
 	class FSourceTextureDim		: SHADER_PERMUTATION_BOOL("USE_SOURCE_TEXTURE");
@@ -461,6 +437,30 @@ class FDeferredLightPS : public FGlobalShader
 		FAnistropicMaterials,
 		FStrataFastPath>;
 
+	BEGIN_SHADER_PARAMETER_STRUCT(FParameters, )
+		SHADER_PARAMETER_RDG_UNIFORM_BUFFER(FSceneTextureUniformParameters, SceneTextures)
+		SHADER_PARAMETER_RDG_UNIFORM_BUFFER(FHairStrandsViewUniformParameters, HairStrands)
+		SHADER_PARAMETER_RDG_UNIFORM_BUFFER(FStrataGlobalUniformParameters, Strata)
+		SHADER_PARAMETER_STRUCT_INCLUDE(FVolumetricCloudShadowAOParameters, CloudShadowAO)
+		SHADER_PARAMETER_STRUCT_INCLUDE(FVolumetricCloudShadowParameters, CloudShadow)
+		SHADER_PARAMETER(uint32, HairTransmittanceBufferMaxCount)
+		SHADER_PARAMETER(uint32, HairShadowMaskValid)
+		SHADER_PARAMETER(FVector4f, ShadowChannelMask)
+		SHADER_PARAMETER_RDG_TEXTURE(Texture2D, LightAttenuationTexture)
+		SHADER_PARAMETER_SAMPLER(SamplerState, LightAttenuationTextureSampler)
+		SHADER_PARAMETER_TEXTURE(Texture2D, IESTexture)
+		SHADER_PARAMETER_SAMPLER(SamplerState, IESTextureSampler)
+		SHADER_PARAMETER_RDG_TEXTURE(Texture2D, LightingChannelsTexture)
+		SHADER_PARAMETER_SAMPLER(SamplerState, LightingChannelsSampler)
+		SHADER_PARAMETER_RDG_BUFFER_SRV(Buffer, HairTransmittanceBuffer)
+		SHADER_PARAMETER_RDG_TEXTURE(Texture2D, ScreenShadowMaskSubPixelTexture)
+		SHADER_PARAMETER_RDG_TEXTURE(Texture2D, DummyRectLightTextureForCapsuleCompilerWarning)
+		SHADER_PARAMETER_SAMPLER(SamplerState, DummyRectLightSamplerForCapsuleCompilerWarning)
+		SHADER_PARAMETER_STRUCT_REF(FViewUniformShaderParameters, View)
+		SHADER_PARAMETER_RDG_UNIFORM_BUFFER(FDeferredLightUniformStruct, DeferredLight)
+		RENDER_TARGET_BINDING_SLOTS()
+	END_SHADER_PARAMETER_STRUCT()
+
 	static bool ShouldCompilePermutation(const FGlobalShaderPermutationParameters& Parameters)
 	{
 		FPermutationDomain PermutationVector(Parameters.PermutationId);
@@ -532,230 +532,6 @@ class FDeferredLightPS : public FGlobalShader
 		OutEnvironment.SetDefine(TEXT("USE_HAIR_COMPLEX_TRANSMITTANCE"), IsHairStrandsSupported(EHairStrandsShaderType::All, Parameters.Platform) ? 1u : 0u);
 		OutEnvironment.SetDefine(TEXT("STRATA_ENABLED"), Strata::IsStrataEnabled() ? 1u : 0u);
 	}
-
-	FDeferredLightPS(const ShaderMetaType::CompiledShaderInitializerType& Initializer)
-	:	FGlobalShader(Initializer)
-	{
-		LightAttenuationTexture.Bind(Initializer.ParameterMap, TEXT("LightAttenuationTexture"));
-		LightAttenuationTextureSampler.Bind(Initializer.ParameterMap, TEXT("LightAttenuationTextureSampler"));
-		IESTexture.Bind(Initializer.ParameterMap, TEXT("IESTexture"));
-		IESTextureSampler.Bind(Initializer.ParameterMap, TEXT("IESTextureSampler"));
-		LightingChannelsTexture.Bind(Initializer.ParameterMap, TEXT("LightingChannelsTexture"));
-		LightingChannelsSampler.Bind(Initializer.ParameterMap, TEXT("LightingChannelsSampler"));
-
-		HairTransmittanceBuffer.Bind(Initializer.ParameterMap, TEXT("HairTransmittanceBuffer"));
-		HairTransmittanceBufferMaxCount.Bind(Initializer.ParameterMap, TEXT("HairTransmittanceBufferMaxCount"));
-		ScreenShadowMaskSubPixelTexture.Bind(Initializer.ParameterMap, TEXT("ScreenShadowMaskSubPixelTexture")); // TODO hook the shader itself
-
-		HairShadowMaskValid.Bind(Initializer.ParameterMap, TEXT("HairShadowMaskValid"));
-		ShadowChannelMask.Bind(Initializer.ParameterMap, TEXT("ShadowChannelMask"));
-		HairStrandsParameters.Bind(Initializer.ParameterMap, FHairStrandsViewUniformParameters::StaticStructMetadata.GetShaderVariableName());
-
-		DummyRectLightTextureForCapsuleCompilerWarning.Bind(Initializer.ParameterMap, TEXT("DummyRectLightTextureForCapsuleCompilerWarning"));
-		DummyRectLightSamplerForCapsuleCompilerWarning.Bind(Initializer.ParameterMap, TEXT("DummyRectLightSamplerForCapsuleCompilerWarning"));
-
-		CloudShadowmapTexture.Bind(Initializer.ParameterMap, TEXT("CloudShadowmapTexture"));
-		CloudShadowmapSampler.Bind(Initializer.ParameterMap, TEXT("CloudShadowmapSampler"));
-		CloudShadowmapFarDepthKm.Bind(Initializer.ParameterMap, TEXT("CloudShadowmapFarDepthKm"));
-		CloudShadowmapWorldToLightClipMatrix.Bind(Initializer.ParameterMap, TEXT("CloudShadowmapWorldToLightClipMatrix"));
-		CloudShadowmapStrength.Bind(Initializer.ParameterMap, TEXT("CloudShadowmapStrength"));
-	}
-
-	FDeferredLightPS()
-	{}
-
-public:
-	void SetParameters(
-		FRHICommandList& RHICmdList, 
-		const FViewInfo& View,
-		const FLightSceneInfo* LightSceneInfo,
-		FRHITexture* ScreenShadowMaskTexture,
-		FRHITexture* LightingChannelsTextureRHI,
-		FRenderLightParams* RenderLightParams,
-		FRHIUniformBuffer* HairStrandsUniformBuffer)
-	{
-		FRHIPixelShader* ShaderRHI = RHICmdList.GetBoundPixelShader();
-		SetParametersBase(RHICmdList, ShaderRHI, View, ScreenShadowMaskTexture, LightingChannelsTextureRHI, LightSceneInfo->Proxy->GetIESTextureResource(), RenderLightParams);
-		SetDeferredLightParameters(RHICmdList, ShaderRHI, GetUniformBufferParameter<FDeferredLightUniformStruct>(), LightSceneInfo, View);
-		if (HairStrandsUniformBuffer)
-		{
-			SetUniformBufferParameter(RHICmdList, ShaderRHI, HairStrandsParameters, HairStrandsUniformBuffer);
-		}
-	}
-
-	void SetParametersSimpleLight(FRHICommandList& RHICmdList, const FViewInfo& View, const FSimpleLightEntry& SimpleLight, const FSimpleLightPerViewEntry& SimpleLightPerViewData)
-	{
-		FRHIPixelShader* ShaderRHI = RHICmdList.GetBoundPixelShader();
-		SetParametersBase(RHICmdList, ShaderRHI, View, nullptr, nullptr, nullptr, nullptr);
-		SetSimpleDeferredLightParameters(RHICmdList, ShaderRHI, GetUniformBufferParameter<FDeferredLightUniformStruct>(), SimpleLight, SimpleLightPerViewData, View);
-	}
-
-private:
-
-	void SetParametersBase(
-		FRHICommandList& RHICmdList, 
-		FRHIPixelShader* ShaderRHI, 
-		const FViewInfo& View,
-		FRHITexture* ScreenShadowMaskTexture,
-		FRHITexture* LightingChannelsTextureRHI,
-		FTexture* IESTextureResource,
-		FRenderLightParams* RenderLightParams)
-	{
-		FGlobalShader::SetParameters<FViewUniformShaderParameters>(RHICmdList, ShaderRHI, View.ViewUniformBuffer);
-		if (TRDGUniformBufferRef<FStrataGlobalUniformParameters> StrataUniformBuffer = Strata::BindStrataGlobalUniformParameters(View.StrataSceneData))
-		{
-			FGlobalShader::SetParameters<FStrataGlobalUniformParameters>(RHICmdList, ShaderRHI, StrataUniformBuffer->GetRHIRef());
-		}
-
-		if(LightAttenuationTexture.IsBound())
-		{
-			if (!ScreenShadowMaskTexture)
-			{
-				ScreenShadowMaskTexture = GWhiteTexture->TextureRHI;
-			}
-
-			SetTextureParameter(
-				RHICmdList,
-				ShaderRHI,
-				LightAttenuationTexture,
-				LightAttenuationTextureSampler,
-				TStaticSamplerState<SF_Point,AM_Wrap,AM_Wrap,AM_Wrap>::GetRHI(),
-				ScreenShadowMaskTexture);
-		}
-
-		{
-			FRHITexture* TextureRHI = IESTextureResource ? IESTextureResource->TextureRHI : GWhiteTexture->TextureRHI;
-
-			SetTextureParameter(
-				RHICmdList,
-				ShaderRHI,
-				IESTexture,
-				IESTextureSampler,
-				TStaticSamplerState<SF_Bilinear,AM_Clamp,AM_Clamp,AM_Clamp>::GetRHI(),
-				TextureRHI
-				);
-		}
-
-		if( LightingChannelsTexture.IsBound() )
-		{
-			SetTextureParameter(
-				RHICmdList,
-				ShaderRHI,
-				LightingChannelsTexture,
-				LightingChannelsSampler,
-				TStaticSamplerState<SF_Point,AM_Clamp,AM_Clamp,AM_Clamp>::GetRHI(),
-				LightingChannelsTextureRHI ? LightingChannelsTextureRHI : GWhiteTexture->TextureRHI.GetReference()
-				);
-		}
-
-		if (HairTransmittanceBuffer.IsBound())
-		{
-			const uint32 TransmittanceBufferMaxCount = RenderLightParams ? RenderLightParams->DeepShadow_TransmittanceMaskBufferMaxCount : 0;
-			SetShaderValue(
-				RHICmdList,
-				ShaderRHI,
-				HairTransmittanceBufferMaxCount,
-				TransmittanceBufferMaxCount);
-			if (RenderLightParams && RenderLightParams->DeepShadow_TransmittanceMaskBuffer)
-			{
-				SetSRVParameter(RHICmdList, ShaderRHI, HairTransmittanceBuffer, RenderLightParams->DeepShadow_TransmittanceMaskBuffer);
-			}
-		}
-
-		if (ScreenShadowMaskSubPixelTexture.IsBound())
-		{
-			if (RenderLightParams)
-			{
-				SetTextureParameter(
-					RHICmdList,
-					ShaderRHI,
-					ScreenShadowMaskSubPixelTexture,
-					LightAttenuationTextureSampler,
-					TStaticSamplerState<SF_Point, AM_Clamp, AM_Clamp, AM_Clamp>::GetRHI(),
-					(RenderLightParams && RenderLightParams->ScreenShadowMaskSubPixelTexture) ? (FRHITexture2D*)RenderLightParams->ScreenShadowMaskSubPixelTexture : GWhiteTexture->TextureRHI);
-
-				uint32 InHairShadowMaskValid = RenderLightParams->ScreenShadowMaskSubPixelTexture ? 1 : 0;
-				SetShaderValue(
-					RHICmdList,
-					ShaderRHI,
-					HairShadowMaskValid,
-					InHairShadowMaskValid);
-
-				SetShaderValue(
-					RHICmdList,
-					ShaderRHI,
-					ShadowChannelMask,
-					RenderLightParams->ShadowChannelMask);
-			}
-		}
-
-		if (DummyRectLightTextureForCapsuleCompilerWarning.IsBound())
-		{
-			SetTextureParameter(
-				RHICmdList,
-				ShaderRHI,
-				DummyRectLightTextureForCapsuleCompilerWarning,
-				DummyRectLightSamplerForCapsuleCompilerWarning,
-				TStaticSamplerState<SF_Bilinear, AM_Clamp, AM_Clamp, AM_Clamp>::GetRHI(),
-				GSystemTextures.DepthDummy->GetShaderResourceRHI()
-			);
-		}
-
-		if (CloudShadowmapTexture.IsBound())
-		{
-			if (RenderLightParams && RenderLightParams->Cloud_ShadowmapTexture)
-			{
-				SetTextureParameter(
-					RHICmdList,
-					ShaderRHI,
-					CloudShadowmapTexture,
-					CloudShadowmapSampler,
-					TStaticSamplerState<SF_Bilinear, AM_Clamp, AM_Clamp, AM_Clamp>::GetRHI(),
-					RenderLightParams->Cloud_ShadowmapTexture ? (FRHITexture2D*)RenderLightParams->Cloud_ShadowmapTexture : GBlackVolumeTexture->TextureRHI);
-
-				SetShaderValue(
-					RHICmdList,
-					ShaderRHI,
-					CloudShadowmapFarDepthKm,
-					RenderLightParams->Cloud_ShadowmapFarDepthKm);
-
-				SetShaderValue(
-					RHICmdList,
-					ShaderRHI,
-					CloudShadowmapWorldToLightClipMatrix,
-					(FMatrix44f)RenderLightParams->Cloud_WorldToLightClipShadowMatrix);
-
-				SetShaderValue(
-					RHICmdList,
-					ShaderRHI,
-					CloudShadowmapStrength,
-					RenderLightParams->Cloud_ShadowmapStrength);
-			}
-		}
-	}
-
-	LAYOUT_FIELD(FShaderResourceParameter, LightAttenuationTexture);
-	LAYOUT_FIELD(FShaderResourceParameter, LightAttenuationTextureSampler);
-	LAYOUT_FIELD(FShaderResourceParameter, IESTexture);
-	LAYOUT_FIELD(FShaderResourceParameter, IESTextureSampler);
-	LAYOUT_FIELD(FShaderResourceParameter, LightingChannelsTexture);
-	LAYOUT_FIELD(FShaderResourceParameter, LightingChannelsSampler);
-
-	LAYOUT_FIELD(FShaderParameter, HairTransmittanceBufferMaxCount);
-	LAYOUT_FIELD(FShaderResourceParameter, HairTransmittanceBuffer);
-	LAYOUT_FIELD(FShaderUniformBufferParameter, HairStrandsParameters);
-	LAYOUT_FIELD(FShaderResourceParameter, ScreenShadowMaskSubPixelTexture);
-	LAYOUT_FIELD(FShaderParameter, HairShadowMaskValid);
-	LAYOUT_FIELD(FShaderParameter, ShadowChannelMask);
-
-	LAYOUT_FIELD(FShaderResourceParameter, DummyRectLightTextureForCapsuleCompilerWarning);
-	LAYOUT_FIELD(FShaderResourceParameter, DummyRectLightSamplerForCapsuleCompilerWarning);
-
-	LAYOUT_FIELD(FShaderResourceParameter, CloudShadowmapTexture);
-	LAYOUT_FIELD(FShaderResourceParameter, CloudShadowmapSampler);
-	LAYOUT_FIELD(FShaderParameter, CloudShadowmapFarDepthKm);
-	LAYOUT_FIELD(FShaderParameter, CloudShadowmapWorldToLightClipMatrix);
-	LAYOUT_FIELD(FShaderParameter, CloudShadowmapStrength);
 };
 
 IMPLEMENT_GLOBAL_SHADER(FDeferredLightPS, "/Engine/Private/DeferredLightPixelShaders.usf", "DeferredLightPixelMain", SF_Pixel);
@@ -763,11 +539,20 @@ IMPLEMENT_GLOBAL_SHADER(FDeferredLightPS, "/Engine/Private/DeferredLightPixelSha
 /** Shader used to visualize stationary light overlap. */
 class FDeferredLightOverlapPS : public FGlobalShader
 {
-	DECLARE_SHADER_TYPE(FDeferredLightOverlapPS,Global)
+	DECLARE_SHADER_TYPE(FDeferredLightOverlapPS, Global)
+	SHADER_USE_PARAMETER_STRUCT(FDeferredLightOverlapPS, FGlobalShader);
 
 	class FRadialAttenuation : SHADER_PERMUTATION_BOOL("RADIAL_ATTENUATION");
 	using FPermutationDomain = TShaderPermutationDomain<FRadialAttenuation>;
 
+	BEGIN_SHADER_PARAMETER_STRUCT(FParameters, )
+		SHADER_PARAMETER_STRUCT_REF(FViewUniformShaderParameters, View)
+		SHADER_PARAMETER_RDG_UNIFORM_BUFFER(FSceneTextureUniformParameters, SceneTextures)
+		SHADER_PARAMETER_RDG_UNIFORM_BUFFER(FDeferredLightUniformStruct, DeferredLight)
+		SHADER_PARAMETER(float, bHasValidChannel)
+		RENDER_TARGET_BINDING_SLOTS()
+	END_SHADER_PARAMETER_STRUCT()
+
 public:
 
 	static bool ShouldCompilePermutation(const FGlobalShaderPermutationParameters& Parameters)
@@ -779,28 +564,6 @@ public:
 	{
 		FGlobalShader::ModifyCompilationEnvironment(Parameters, OutEnvironment);
 	}
-
-	FDeferredLightOverlapPS(const ShaderMetaType::CompiledShaderInitializerType& Initializer)
-	:	FGlobalShader(Initializer)
-	{
-		HasValidChannel.Bind(Initializer.ParameterMap, TEXT("HasValidChannel"));
-	}
-
-	FDeferredLightOverlapPS()
-	{
-	}
-
-	void SetParameters(FRHICommandList& RHICmdList, const FSceneView& View, const FLightSceneInfo* LightSceneInfo)
-	{
-		FRHIPixelShader* ShaderRHI = RHICmdList.GetBoundPixelShader();
-		FGlobalShader::SetParameters<FViewUniformShaderParameters>(RHICmdList, ShaderRHI,View.ViewUniformBuffer);
-		const float HasValidChannelValue = LightSceneInfo->Proxy->GetPreviewShadowMapChannel() == INDEX_NONE ? 0.0f : 1.0f;
-		SetShaderValue(RHICmdList, ShaderRHI, HasValidChannel, HasValidChannelValue);
-		SetDeferredLightParameters(RHICmdList, ShaderRHI, GetUniformBufferParameter<FDeferredLightUniformStruct>(), LightSceneInfo, View);
-	}
-
-private:
-	LAYOUT_FIELD(FShaderParameter, HasValidChannel);
 };
 
 IMPLEMENT_GLOBAL_SHADER(FDeferredLightOverlapPS, "/Engine/Private/StationaryLightOverlapShaders.usf", "OverlapPixelMain", SF_Pixel);
@@ -1066,71 +829,8 @@ void FSceneRenderer::GatherAndSortLights(FSortedLightSetSceneInfo& OutSortedLigh
 	check(OutSortedLights.UnbatchedLightStart >= OutSortedLights.ClusteredSupportedEnd);
 }
 
-/** Shader parameters to use when creating a RenderLight(...) pass. */
-BEGIN_SHADER_PARAMETER_STRUCT(FRenderLightParameters, )
-	SHADER_PARAMETER_STRUCT_REF(FViewUniformShaderParameters, View)
-	SHADER_PARAMETER_RDG_UNIFORM_BUFFER(FSceneTextureUniformParameters, SceneTextures)
-	SHADER_PARAMETER_RDG_UNIFORM_BUFFER(FHairStrandsViewUniformParameters, HairStrands)
-	SHADER_PARAMETER_RDG_UNIFORM_BUFFER(FStrataGlobalUniformParameters, Strata)
-	SHADER_PARAMETER_STRUCT_INCLUDE(FVolumetricCloudShadowAOParameters, CloudShadowAO)
-	RDG_TEXTURE_ACCESS(ShadowMaskTexture, ERHIAccess::SRVGraphics)
-	RDG_TEXTURE_ACCESS(LightingChannelsTexture, ERHIAccess::SRVGraphics)
-	// We reference all the Strata tiled resources we might need in this pass
-	SHADER_PARAMETER_RDG_BUFFER_SRV(Buffer<uint>, TileListBufferSimple)
-	SHADER_PARAMETER_RDG_BUFFER_SRV(Buffer<uint>, TileListBufferComplex)
-	RDG_BUFFER_ACCESS(TileIndirectBufferSimple, ERHIAccess::IndirectArgs)
-	RDG_BUFFER_ACCESS(TileIndirectBufferComplex, ERHIAccess::IndirectArgs)
-	RENDER_TARGET_BINDING_SLOTS()
-END_SHADER_PARAMETER_STRUCT()
-
-void GetRenderLightParameters(
-	const FViewInfo& View,
-	FRDGTextureRef SceneColorTexture,
-	FRDGTextureRef SceneDepthTexture,
-	TRDGUniformBufferRef<FSceneTextureUniformParameters> SceneTexturesUniformBuffer,
-	TRDGUniformBufferRef<FHairStrandsViewUniformParameters> HairStrandsUniformBuffer,
-	FRDGTextureRef ShadowMaskTexture,
-	FRDGTextureRef LightingChannelsTexture,
-	const FVolumetricCloudShadowAOParameters& CloudShadowAOParameters,
-	FRenderLightParameters& Parameters)
-{
-	Parameters.View = View.ViewUniformBuffer;
-	Parameters.SceneTextures = SceneTexturesUniformBuffer;
-	Parameters.HairStrands = HairStrandsUniformBuffer;
-	Parameters.Strata = Strata::BindStrataGlobalUniformParameters(View.StrataSceneData);
-	Parameters.ShadowMaskTexture = ShadowMaskTexture;
-	Parameters.LightingChannelsTexture = LightingChannelsTexture;
-	Parameters.CloudShadowAO = CloudShadowAOParameters;
-	Parameters.RenderTargets[0] = FRenderTargetBinding(SceneColorTexture, ERenderTargetLoadAction::ELoad);
-
-	if (SceneDepthTexture)
-	{
-		Parameters.RenderTargets.DepthStencil = FDepthStencilBinding(SceneDepthTexture, ERenderTargetLoadAction::ELoad, ERenderTargetLoadAction::ELoad, FExclusiveDepthStencil::DepthRead_StencilWrite);
-	}
-
-	if (Strata::ShouldPassesReadingStrataBeTiled(View.Family->GetFeatureLevel()))
-	{
-		Parameters.TileListBufferSimple = View.StrataSceneData->ClassificationTileListBufferSRV[EStrataTileMaterialType::ESimple];
-		Parameters.TileListBufferComplex = View.StrataSceneData->ClassificationTileListBufferSRV[EStrataTileMaterialType::EComplex];
-		Parameters.TileIndirectBufferSimple = View.StrataSceneData->ClassificationTileIndirectBuffer[EStrataTileMaterialType::ESimple];
-		Parameters.TileIndirectBufferComplex = View.StrataSceneData->ClassificationTileIndirectBuffer[EStrataTileMaterialType::EComplex];
-	}
-}
-
 FHairStrandsTransmittanceMaskData CreateDummyHairStrandsTransmittanceMaskData(FRDGBuilder& GraphBuilder, FGlobalShaderMap* ShaderMap);
 
-void GetRenderLightParameters(
-	const FViewInfo& View,
-	const FMinimalSceneTextures& SceneTextures,
-	const FHairStrandsViewData& HairViewData,
-	FRDGTextureRef ShadowMaskTexture,
-	FRDGTextureRef LightingChannelsTexture,
-	const FVolumetricCloudShadowAOParameters& CloudShadowAOParameters,
-	FRenderLightParameters& Parameters)
-{
-	GetRenderLightParameters(View, SceneTextures.Color.Target, SceneTextures.Depth.Target, SceneTextures.UniformBuffer, HairViewData.UniformBuffer, ShadowMaskTexture, LightingChannelsTexture, CloudShadowAOParameters, Parameters);
-}
-
 void FDeferredShadingSceneRenderer::RenderLights(
 	FRDGBuilder& GraphBuilder,
 	FMinimalSceneTextures& SceneTextures,
@@ -1278,7 +978,7 @@ void FDeferredShadingSceneRenderer::RenderLights(
 						for (int32 LightIndex = StandardDeferredStart; LightIndex < UnbatchedLightStart; LightIndex++)
 						{
 							const FLightSceneInfo* LightSceneInfo = SortedLights[LightIndex].LightSceneInfo;
-							RenderLightForHair(GraphBuilder, View, SceneTextures.UniformBuffer, LightSceneInfo, NullScreenShadowMaskSubPixelTexture, LightingChannelsTexture, DummyTransmittanceMaskData, false /*bForwardRendering*/);
+							RenderLightForHair(GraphBuilder, View, SceneTextures, LightSceneInfo, NullScreenShadowMaskSubPixelTexture, LightingChannelsTexture, DummyTransmittanceMaskData, false /*bForwardRendering*/);
 						}
 					}
 				}
@@ -1909,7 +1609,7 @@ void FDeferredShadingSceneRenderer::RenderLights(
 
 							// Note: ideally the light should still be evaluated for hair when not casting shadow, but for preserving the old behavior, and not adding 
 							// any perf. regression, we disable this light for hair rendering 
-							RenderLightForHair(GraphBuilder, View, SceneTextures.UniformBuffer, &LightSceneInfo, ScreenShadowMaskSubPixelTexture, LightingChannelsTexture, TransmittanceMaskData, false /*bForwardRendering*/);
+							RenderLightForHair(GraphBuilder, View, SceneTextures, &LightSceneInfo, ScreenShadowMaskSubPixelTexture, LightingChannelsTexture, TransmittanceMaskData, false /*bForwardRendering*/);
 						}
 					}
 				}
@@ -2011,7 +1711,7 @@ void SetBoundingGeometryRasterizerAndDepthState(FGraphicsPipelineStateInitialize
 }
 
 // Use DBT to allow work culling on shadow lights
-void CalculateLightNearFarDepthFromBounds(const FViewInfo& View, const FSphere &LightBounds, float &NearDepth, float &FarDepth)
+static void CalculateLightNearFarDepthFromBounds(const FViewInfo& View, const FSphere &LightBounds, float &NearDepth, float &FarDepth)
 {
 	const FMatrix ViewProjection = View.ViewMatrices.GetViewProjectionMatrix();
 	const FVector ViewDirection = View.GetViewDirection();
@@ -2041,325 +1741,262 @@ void CalculateLightNearFarDepthFromBounds(const FViewInfo& View, const FSphere &
 
 }
 
-static void BindAtmosphereAndCloudResources(
+static TRDGUniformBufferRef<FDeferredLightUniformStruct> CreateDeferredLightUniformBuffer(FRDGBuilder& GraphBuilder, const FViewInfo& View, const FLightSceneInfo& LightSceneInfo)
+{
+	auto* DeferredLightStruct = GraphBuilder.AllocParameters<FDeferredLightUniformStruct>();
+	*DeferredLightStruct = GetDeferredLightParameters(View, LightSceneInfo);
+	return GraphBuilder.CreateUniformBuffer(DeferredLightStruct);
+}
+
+static TRDGUniformBufferRef<FDeferredLightUniformStruct> CreateDeferredLightUniformBuffer(FRDGBuilder& GraphBuilder, const FSimpleLightEntry& SimpleLight, const FVector3f& SimpleLightPosition)
+{
+	auto* DeferredLightStruct = GraphBuilder.AllocParameters<FDeferredLightUniformStruct>();
+	*DeferredLightStruct = GetSimpleDeferredLightParameters(SimpleLight, SimpleLightPosition);
+	return GraphBuilder.CreateUniformBuffer(DeferredLightStruct);
+}
+
+static FDeferredLightPS::FParameters GetDeferredLightPSParameters(
+	FRDGBuilder& GraphBuilder,
 	const FScene* Scene,
-	const FViewInfo* View,
-	const FLightSceneProxy* Proxy,
-	FRenderLightParams& RenderLightParams,
-	bool& bAtmospherePerPixelTransmittance,
-	bool& bCloudPerPixelTransmittance)
+	const FViewInfo& View,
+	const FLightSceneInfo* LightSceneInfo,
+	FRDGTextureRef SceneColorTexture,
+	FRDGTextureRef SceneDepthTexture,
+	TRDGUniformBufferRef<FSceneTextureUniformParameters> SceneTexturesUniformBuffer,
+	TRDGUniformBufferRef<FHairStrandsViewUniformParameters> HairStrandsUniformBuffer,
+	FRDGTextureRef ShadowMaskTexture,
+	FRDGTextureRef LightingChannelsTexture,
+	bool bCloudShadow)
 {
-	bAtmospherePerPixelTransmittance =
-		Proxy->GetLightType() == LightType_Directional &&
-		Proxy->IsUsedAsAtmosphereSunLight() &&
-		Proxy->GetUsePerPixelAtmosphereTransmittance() &&
-		ShouldRenderSkyAtmosphere(Scene, View->Family->EngineShowFlags);
-
-	const FLightSceneProxy* AtmosphereLight0Proxy = Scene->AtmosphereLights[0] ? Scene->AtmosphereLights[0]->Proxy : nullptr;
-	const FLightSceneProxy* AtmosphereLight1Proxy = Scene->AtmosphereLights[1] ? Scene->AtmosphereLights[1]->Proxy : nullptr;
-	const FVolumetricCloudRenderSceneInfo* CloudInfo = Scene->GetVolumetricCloudSceneInfo();
-	const bool VolumetricCloudShadowMap0Valid = View->VolumetricCloudShadowExtractedRenderTarget[0] != nullptr;
-	const bool VolumetricCloudShadowMap1Valid = View->VolumetricCloudShadowExtractedRenderTarget[1] != nullptr;
-	const bool bLight0CloudPerPixelTransmittance = CloudInfo && VolumetricCloudShadowMap0Valid && AtmosphereLight0Proxy == Proxy && AtmosphereLight0Proxy && AtmosphereLight0Proxy->GetCloudShadowOnSurfaceStrength() > 0.0f;
-	const bool bLight1CloudPerPixelTransmittance = CloudInfo && VolumetricCloudShadowMap1Valid && AtmosphereLight1Proxy == Proxy && AtmosphereLight1Proxy && AtmosphereLight1Proxy->GetCloudShadowOnSurfaceStrength() > 0.0f;
-	if (bLight0CloudPerPixelTransmittance)
+	FDeferredLightPS::FParameters Out;
+
+	const ELightComponentType LightType = (ELightComponentType)LightSceneInfo->Proxy->GetLightType();
+	const bool bIsDirectional = LightType == LightType_Directional;
+
+	FRDGTextureRef WhiteDummy = GSystemTextures.GetWhiteDummy(GraphBuilder);
+	FRDGTextureRef DepthDummy = GSystemTextures.GetDepthDummy(GraphBuilder);
+	FRDGBufferRef BufferDummy = GSystemTextures.GetDefaultBuffer(GraphBuilder, 4, 0u);
+	FRDGBufferSRVRef BufferDummySRV = GraphBuilder.CreateSRV(BufferDummy, PF_R32_UINT);
+
+	// PS - General parameters
+	const FVolumetricCloudRenderSceneInfo* CloudInfo = bCloudShadow ? Scene->GetVolumetricCloudSceneInfo() : nullptr;
+	Out.SceneTextures = SceneTexturesUniformBuffer;
+	Out.HairStrands = HairStrandsUniformBuffer;
+	Out.Strata = Strata::BindStrataGlobalUniformParameters(View.StrataSceneData);
+	Out.LightingChannelsTexture = LightingChannelsTexture ? LightingChannelsTexture : WhiteDummy;
+	Out.LightingChannelsSampler = TStaticSamplerState<SF_Point, AM_Clamp, AM_Clamp, AM_Clamp>::GetRHI();
+	Out.CloudShadowAO = GetCloudShadowAOParameters(GraphBuilder, View, CloudInfo);
+	Out.CloudShadow = GetCloudShadowParameters(GraphBuilder, Scene, View, LightSceneInfo->Proxy, CloudInfo);
+	Out.LightAttenuationTexture = ShadowMaskTexture ? ShadowMaskTexture : WhiteDummy;
+	Out.LightAttenuationTextureSampler = TStaticSamplerState<SF_Point, AM_Wrap, AM_Wrap, AM_Wrap>::GetRHI();
+	Out.DummyRectLightTextureForCapsuleCompilerWarning = DepthDummy;
+	Out.DummyRectLightSamplerForCapsuleCompilerWarning = TStaticSamplerState<SF_Bilinear, AM_Clamp, AM_Clamp, AM_Clamp>::GetRHI();
+	Out.IESTextureSampler = TStaticSamplerState<SF_Bilinear, AM_Clamp, AM_Clamp, AM_Clamp>::GetRHI();
+	Out.IESTexture = GSystemTextures.WhiteDummy->GetShaderResourceRHI();
+	if (LightSceneInfo->Proxy->GetIESTextureResource())
 	{
-		RenderLightParams.Cloud_ShadowmapTexture = View->VolumetricCloudShadowExtractedRenderTarget[0]->GetShaderResourceRHI();
-		RenderLightParams.Cloud_ShadowmapFarDepthKm = CloudInfo->GetVolumetricCloudCommonShaderParameters().CloudShadowmapFarDepthKm[0].X;
-		RenderLightParams.Cloud_WorldToLightClipShadowMatrix = CloudInfo->GetVolumetricCloudCommonShaderParameters().CloudShadowmapWorldToLightClipMatrix[0];
-		RenderLightParams.Cloud_ShadowmapStrength = AtmosphereLight0Proxy->GetCloudShadowOnSurfaceStrength();
+		Out.IESTexture = LightSceneInfo->Proxy->GetIESTextureResource()->TextureRHI;
 	}
-	else if (bLight1CloudPerPixelTransmittance)
+	Out.View = View.ViewUniformBuffer;
+	Out.DeferredLight = CreateDeferredLightUniformBuffer(GraphBuilder, View, *LightSceneInfo);
+	// PS - Hair (default value)
+	Out.ScreenShadowMaskSubPixelTexture = WhiteDummy;
+	Out.HairTransmittanceBuffer = BufferDummySRV;
+	Out.HairTransmittanceBufferMaxCount = 0;
+	Out.HairShadowMaskValid = false;
+	Out.ShadowChannelMask = FVector4f(1, 1, 1, 1);
+	// PS - Render Targets
+	Out.RenderTargets[0] = FRenderTargetBinding(SceneColorTexture, ERenderTargetLoadAction::ELoad);
+	if (SceneDepthTexture)
 	{
-		RenderLightParams.Cloud_ShadowmapTexture = View->VolumetricCloudShadowExtractedRenderTarget[1]->GetShaderResourceRHI();
-		RenderLightParams.Cloud_ShadowmapFarDepthKm = CloudInfo->GetVolumetricCloudCommonShaderParameters().CloudShadowmapFarDepthKm[1].X;
-		RenderLightParams.Cloud_WorldToLightClipShadowMatrix = CloudInfo->GetVolumetricCloudCommonShaderParameters().CloudShadowmapWorldToLightClipMatrix[1];
-		RenderLightParams.Cloud_ShadowmapStrength = AtmosphereLight1Proxy->GetCloudShadowOnSurfaceStrength();
+		Out.RenderTargets.DepthStencil = FDepthStencilBinding(SceneDepthTexture, ERenderTargetLoadAction::ELoad, ERenderTargetLoadAction::ELoad, FExclusiveDepthStencil::DepthRead_StencilWrite);
 	}
-	bCloudPerPixelTransmittance = bLight0CloudPerPixelTransmittance || bLight1CloudPerPixelTransmittance;
-}
 
-/**
- * Used by RenderLights to render a light to the scene color buffer.
- *
- * @param LightSceneInfo Represents the current light
- * @param LightIndex The light's index into FScene::Lights
- * @return true if anything got rendered
- */
+	return Out;
+}
 
+// Used by RenderLights to render a light to the scene color buffer.
+template<typename TShaderType, typename TParametersType>
 static void InternalRenderLight(
-	FRHICommandList& RHICmdList,
+	FRDGBuilder& GraphBuilder,
 	const FScene* Scene,
 	const FViewInfo& View,
 	const FLightSceneInfo* LightSceneInfo,
-	FRenderLightParameters* PassParameters,	// If this is null, it means we cannot use Strata tiles and fallback to previous behavior.
-	FRHITexture* ScreenShadowMaskTexture,
-	FRHITexture* LightingChannelsTexture,
-	bool bRenderOverlap,
-	bool bEnableStrataStencilTest, bool bEnableStrataTiledPass, uint32 StrataTileMaterialType)
+	TShaderType& PixelShader,
+	TParametersType* PassParameters,
+	EStrataTileMaterialType StrataTileMaterialType,
+	const TCHAR* ShaderName)
 {
-	FGraphicsPipelineStateInitializer GraphicsPSOInit;
-	RHICmdList.ApplyCachedRenderTargets(GraphicsPSOInit);
-
-	GraphicsPSOInit.BlendState = TStaticBlendState<CW_RGBA, BO_Add, BF_One, BF_One, BO_Add, BF_One, BF_One>::GetRHI();
-	GraphicsPSOInit.PrimitiveType = PT_TriangleList;
-
 	const FLightSceneProxy* RESTRICT LightProxy = LightSceneInfo->Proxy;
-
-	const FSphere LightBounds = LightProxy->GetBoundingSphere();
+	const uint32 StencilRef = StrataTileMaterialType == EStrataTileMaterialType::ESimple ? Strata::StencilBit : 0u;
 	const bool bTransmission = LightProxy->Transmission();
+	const FSphere LightBounds = LightProxy->GetBoundingSphere();
+	const ELightComponentType LightType = (ELightComponentType)LightProxy->GetLightType();
 
-	bool bUseIESTexture = false;
-
-	if(View.Family->EngineShowFlags.TexturedLightProfiles)
+	const TCHAR* PermationName = nullptr;
+	switch (StrataTileMaterialType)
 	{
-		bUseIESTexture = (LightSceneInfo->Proxy->GetIESTextureResource() != 0);
+	case EStrataTileMaterialType::ESimple : TEXT("(Strata,Simple)"); break;
+	case EStrataTileMaterialType::EComplex: TEXT("(Strata,Complex)"); break;
+	case EStrataTileMaterialType::ECount  : TEXT(""); break;
 	}
 
-	// Set the device viewport for the view.
-	RHICmdList.SetViewport(View.ViewRect.Min.X, View.ViewRect.Min.Y, 0.0f, View.ViewRect.Max.X, View.ViewRect.Max.Y, 1.0f);
-
-	FRenderLightParams RenderLightParams;
-	if (bEnableStrataStencilTest)
-	{
-		GraphicsPSOInit.DepthStencilState = TStaticDepthStencilState<
-			false, CF_Always,
-			true, CF_Equal, SO_Keep, SO_Keep, SO_Keep,
-			true, CF_Equal, SO_Keep, SO_Keep, SO_Keep,
-			Strata::StencilBit, 0x0>::GetRHI();
-	}
-	else
+	GraphBuilder.AddPass(
+		RDG_EVENT_NAME("%s%s", ShaderName, PermationName),
+		PassParameters,
+		ERDGPassFlags::Raster,
+		[Scene, &View, PixelShader, LightSceneInfo, PassParameters, StencilRef, LightBounds, LightType, StrataTileMaterialType](FRHICommandList& RHICmdList)
 	{
-		GraphicsPSOInit.DepthStencilState = TStaticDepthStencilState<false, CF_Always>::GetRHI();
-	}
-
-	const uint32 StencilRef = StrataTileMaterialType == EStrataTileMaterialType::ESimple ? Strata::StencilBit : 0u;
 
-	if (LightProxy->GetLightType() == LightType_Directional)
-	{
-		// Turn DBT back off
-		GraphicsPSOInit.bDepthBounds = false;
-		FDeferredLightVS::FPermutationDomain PermutationVectorVS;
-		PermutationVectorVS.Set<FDeferredLightVS::FRadialLight>(false);
-		TShaderMapRef<FDeferredLightVS> VertexShader(View.ShaderMap, PermutationVectorVS);
+		const bool bIsRadial = LightType != LightType_Directional;
+		const bool bEnableStrataTiledPass   = StrataTileMaterialType != EStrataTileMaterialType::ECount;
+		const bool bEnableStrataStencilTest = StrataTileMaterialType != EStrataTileMaterialType::ECount && bIsRadial;
 
-		Strata::FStrataTilePassVS::FParameters VSParameters;
-		Strata::FStrataTilePassVS::FPermutationDomain VSPermutationVector;
-		VSPermutationVector.Set< Strata::FStrataTilePassVS::FEnableDebug >(false);
-		VSPermutationVector.Set< Strata::FStrataTilePassVS::FEnableTexCoordScreenVector >(true);
-		TShaderMapRef<Strata::FStrataTilePassVS> StrataTilePassVertexShader(View.ShaderMap, VSPermutationVector);
+		FGraphicsPipelineStateInitializer GraphicsPSOInit;
+		RHICmdList.ApplyCachedRenderTargets(GraphicsPSOInit);
+		// Set the device viewport for the view.
+		RHICmdList.SetViewport(View.ViewRect.Min.X, View.ViewRect.Min.Y, 0.0f, View.ViewRect.Max.X, View.ViewRect.Max.Y, 1.0f);
 
-		GraphicsPSOInit.RasterizerState = TStaticRasterizerState<FM_Solid, CM_None>::GetRHI();
+		GraphicsPSOInit.BlendState = TStaticBlendState<CW_RGBA, BO_Add, BF_One, BF_One, BO_Add, BF_One, BF_One>::GetRHI();
+		GraphicsPSOInit.PrimitiveType = PT_TriangleList;
 
-		if (bRenderOverlap)
+		if (bEnableStrataStencilTest)
 		{
-			FDeferredLightOverlapPS::FPermutationDomain PermutationVector;
-			PermutationVector.Set<FDeferredLightOverlapPS::FRadialAttenuation>(false);
-			TShaderMapRef<FDeferredLightOverlapPS> PixelShader(View.ShaderMap, PermutationVector);
-			GraphicsPSOInit.BoundShaderState.VertexDeclarationRHI = GFilterVertexDeclaration.VertexDeclarationRHI;
-			GraphicsPSOInit.BoundShaderState.VertexShaderRHI = VertexShader.GetVertexShader();
-			GraphicsPSOInit.BoundShaderState.PixelShaderRHI = PixelShader.GetPixelShader();
-			SetGraphicsPipelineState(RHICmdList, GraphicsPSOInit, StencilRef);
-			PixelShader->SetParameters(RHICmdList, View, LightSceneInfo);
+			GraphicsPSOInit.DepthStencilState = TStaticDepthStencilState<
+				false, CF_Always,
+				true, CF_Equal, SO_Keep, SO_Keep, SO_Keep,
+				true, CF_Equal, SO_Keep, SO_Keep, SO_Keep,
+				Strata::StencilBit, 0x0>::GetRHI();
 		}
 		else
 		{
-			bool bAtmospherePerPixelTransmittance = false;
-			bool bCloudPerPixelTransmittance = false;
-			BindAtmosphereAndCloudResources(
-				Scene,
-				&View,
-				LightSceneInfo->Proxy,
-				RenderLightParams,
-				bAtmospherePerPixelTransmittance,
-				bCloudPerPixelTransmittance);
-
-			FDeferredLightPS::FPermutationDomain PermutationVector;
-			PermutationVector.Set< FDeferredLightPS::FSourceShapeDim >( ELightSourceShape::Directional );
-			PermutationVector.Set< FDeferredLightPS::FIESProfileDim >( false );
-			PermutationVector.Set< FDeferredLightPS::FInverseSquaredDim >( false );
-			PermutationVector.Set< FDeferredLightPS::FVisualizeCullingDim >( View.Family->EngineShowFlags.VisualizeLightCulling );
-			PermutationVector.Set< FDeferredLightPS::FLightingChannelsDim >( View.bUsesLightingChannels );
-			PermutationVector.Set< FDeferredLightPS::FAnistropicMaterials >(ShouldRenderAnisotropyPass(View));
-			PermutationVector.Set< FDeferredLightPS::FTransmissionDim >( bTransmission );
-			PermutationVector.Set< FDeferredLightPS::FHairLighting>(0);
-			// Only directional lights are rendered in this path, so we only need to check if it is use to light the atmosphere
-			PermutationVector.Set< FDeferredLightPS::FAtmosphereTransmittance >(bAtmospherePerPixelTransmittance);
-			PermutationVector.Set< FDeferredLightPS::FCloudTransmittance >(bCloudPerPixelTransmittance);
-			PermutationVector.Set< FDeferredLightPS::FStrataFastPath >(StrataTileMaterialType == EStrataTileMaterialType::ESimple);
-
-			TShaderMapRef< FDeferredLightPS > PixelShader( View.ShaderMap, PermutationVector );
-			GraphicsPSOInit.BoundShaderState.VertexDeclarationRHI = GFilterVertexDeclaration.VertexDeclarationRHI;
-			GraphicsPSOInit.BoundShaderState.VertexShaderRHI = VertexShader.GetVertexShader();
-			GraphicsPSOInit.BoundShaderState.PixelShaderRHI = PixelShader.GetPixelShader();
-
-			if (bEnableStrataTiledPass)
-			{
-				Strata::FillUpTiledPassData((EStrataTileMaterialType)StrataTileMaterialType, View, VSParameters, GraphicsPSOInit.PrimitiveType);
-				GraphicsPSOInit.BoundShaderState.VertexShaderRHI = StrataTilePassVertexShader.GetVertexShader();
-			}
-
-			SetGraphicsPipelineState(RHICmdList, GraphicsPSOInit, StencilRef);
-			PixelShader->SetParameters(RHICmdList, View, LightSceneInfo, ScreenShadowMaskTexture, LightingChannelsTexture, &RenderLightParams, nullptr);
+			GraphicsPSOInit.DepthStencilState = TStaticDepthStencilState<false, CF_Always>::GetRHI();
 		}
 
-		if (!bEnableStrataTiledPass)
+		if (LightType == LightType_Directional)
 		{
-			FDeferredLightVS::FParameters VSParameters2 = FDeferredLightVS::GetParameters(View);
-			SetShaderParameters(RHICmdList, VertexShader, VertexShader.GetVertexShader(), VSParameters2);
+			FDeferredLightVS::FPermutationDomain PermutationVectorVS;
+			PermutationVectorVS.Set<FDeferredLightVS::FRadialLight>(false);
+			TShaderMapRef<FDeferredLightVS> VertexShader(View.ShaderMap, PermutationVectorVS);
 
-			// Apply the directional light as a full screen quad
-			DrawRectangle(
-				RHICmdList,
-				0, 0,
-				View.ViewRect.Width(), View.ViewRect.Height(),
-				View.ViewRect.Min.X, View.ViewRect.Min.Y,
-				View.ViewRect.Width(), View.ViewRect.Height(),
-				View.ViewRect.Size(),
-				GetSceneTextureExtent(),
-				VertexShader,
-				EDRF_UseTriangleOptimization);
-		}
-		else
-		{
-			SetShaderParameters(RHICmdList, StrataTilePassVertexShader, StrataTilePassVertexShader.GetVertexShader(), VSParameters);
+			Strata::FStrataTilePassVS::FPermutationDomain VSPermutationVector;
+			VSPermutationVector.Set< Strata::FStrataTilePassVS::FEnableDebug >(false);
+			VSPermutationVector.Set< Strata::FStrataTilePassVS::FEnableTexCoordScreenVector >(true);
+			TShaderMapRef<Strata::FStrataTilePassVS> TileVertexShader(View.ShaderMap, VSPermutationVector);
 
-			RHICmdList.DrawPrimitiveIndirect(VSParameters.TileIndirectBuffer->GetIndirectRHICallBuffer(), 0);
-		}
-	}
-	else
-	{
-		// Use DBT to allow work culling on shadow lights
-		// Disable depth bound when hair rendering is enabled as this rejects partially covered pixel write (with opaque background)
-		GraphicsPSOInit.bDepthBounds = GSupportsDepthBoundsTest && GAllowDepthBoundsTest != 0;
+			// Turn DBT back off
+			GraphicsPSOInit.bDepthBounds = false;
+			GraphicsPSOInit.RasterizerState = TStaticRasterizerState<FM_Solid, CM_None>::GetRHI();
+			GraphicsPSOInit.BoundShaderState.VertexDeclarationRHI = GFilterVertexDeclaration.VertexDeclarationRHI;
+			GraphicsPSOInit.BoundShaderState.VertexShaderRHI = bEnableStrataTiledPass ? TileVertexShader.GetVertexShader() : VertexShader.GetVertexShader();
+			GraphicsPSOInit.BoundShaderState.PixelShaderRHI = PixelShader.GetPixelShader();		
+			SetGraphicsPipelineState(RHICmdList, GraphicsPSOInit, StencilRef);
 
-		FDeferredLightVS::FPermutationDomain PermutationVectorVS;
-		PermutationVectorVS.Set<FDeferredLightVS::FRadialLight>(true);
-		TShaderMapRef<FDeferredLightVS> VertexShader(View.ShaderMap, PermutationVectorVS);
+			SetShaderParameters(RHICmdList, PixelShader, PixelShader.GetPixelShader(), PassParameters->PS);
 
-		const bool bCameraInsideLightGeometry = ((FVector)View.ViewMatrices.GetViewOrigin() - LightBounds.Center).SizeSquared() < FMath::Square(LightBounds.W * 1.05f + View.NearClippingDistance * 2.0f)
-		//const bool bCameraInsideLightGeometry = LightProxy->AffectsBounds( FSphere( View.ViewMatrices.GetViewOrigin(), View.NearClippingDistance * 2.0f ) )
-			// Always draw backfaces in ortho
-			//@todo - accurate ortho camera / light intersection
-			|| !View.IsPerspectiveProjection();
+			if (bEnableStrataTiledPass)
+			{
+				Strata::FStrataTilePassVS::FParameters VSParameters;
+				Strata::FillUpTiledPassData(StrataTileMaterialType, View, VSParameters, GraphicsPSOInit.PrimitiveType);
+				SetShaderParameters(RHICmdList, TileVertexShader, TileVertexShader.GetVertexShader(), VSParameters);
 
-		SetBoundingGeometryRasterizerAndDepthState(GraphicsPSOInit, View, bCameraInsideLightGeometry);
+				RHICmdList.DrawPrimitiveIndirect(VSParameters.TileIndirectBuffer->GetIndirectRHICallBuffer(), 0);
+			}
+			else
+			{
+				FDeferredLightVS::FParameters VSParameters2 = FDeferredLightVS::GetParameters(View);
+				SetShaderParameters(RHICmdList, VertexShader, VertexShader.GetVertexShader(), VSParameters2);
 
-		if (bRenderOverlap)
+				// Apply the directional light as a full screen quad
+				DrawRectangle(
+					RHICmdList,
+					0, 0,
+					View.ViewRect.Width(), View.ViewRect.Height(),
+					View.ViewRect.Min.X, View.ViewRect.Min.Y,	
+					View.ViewRect.Width(), View.ViewRect.Height(),
+					View.ViewRect.Size(),
+					GetSceneTextureExtent(),
+					VertexShader,
+					EDRF_UseTriangleOptimization);
+			}
+		}
+		else // Radial light (LightType_Point, LightType_Spot, LightType_Rect)
 		{
-			FDeferredLightOverlapPS::FPermutationDomain PermutationVector;
-			PermutationVector.Set<FDeferredLightOverlapPS::FRadialAttenuation>(true);
+			// Use DBT to allow work culling on shadow lights
+			// Disable depth bound when hair rendering is enabled as this rejects partially covered pixel write (with opaque background)
+			GraphicsPSOInit.bDepthBounds = GSupportsDepthBoundsTest && GAllowDepthBoundsTest != 0;
 
-			TShaderMapRef<FDeferredLightOverlapPS> PixelShader(View.ShaderMap, PermutationVector);
-			GraphicsPSOInit.BoundShaderState.VertexDeclarationRHI = GetVertexDeclarationFVector4();
-			GraphicsPSOInit.BoundShaderState.VertexShaderRHI = VertexShader.GetVertexShader();
-			GraphicsPSOInit.BoundShaderState.PixelShaderRHI = PixelShader.GetPixelShader();
+			FDeferredLightVS::FPermutationDomain PermutationVectorVS;
+			PermutationVectorVS.Set<FDeferredLightVS::FRadialLight>(true);
+			TShaderMapRef<FDeferredLightVS> VertexShader(View.ShaderMap, PermutationVectorVS);
 
-			SetGraphicsPipelineState(RHICmdList, GraphicsPSOInit, StencilRef);
-			PixelShader->SetParameters(RHICmdList, View, LightSceneInfo);
-		}
-		else
-		{
-			FDeferredLightPS::FPermutationDomain PermutationVector;
-			PermutationVector.Set< FDeferredLightPS::FSourceShapeDim >( LightProxy->IsRectLight() ? ELightSourceShape::Rect : ELightSourceShape::Capsule );
-			PermutationVector.Set< FDeferredLightPS::FSourceTextureDim >( LightProxy->IsRectLight() && LightProxy->HasSourceTexture() );
-			PermutationVector.Set< FDeferredLightPS::FIESProfileDim >( bUseIESTexture );
-			PermutationVector.Set< FDeferredLightPS::FInverseSquaredDim >( LightProxy->IsInverseSquared() );
-			PermutationVector.Set< FDeferredLightPS::FVisualizeCullingDim >( View.Family->EngineShowFlags.VisualizeLightCulling );
-			PermutationVector.Set< FDeferredLightPS::FLightingChannelsDim >( View.bUsesLightingChannels );
-			PermutationVector.Set< FDeferredLightPS::FAnistropicMaterials >(ShouldRenderAnisotropyPass(View) && !LightSceneInfo->Proxy->IsRectLight());
-			PermutationVector.Set< FDeferredLightPS::FTransmissionDim >( bTransmission );
-			PermutationVector.Set< FDeferredLightPS::FHairLighting>(0);
-			PermutationVector.Set < FDeferredLightPS::FAtmosphereTransmittance >(false);
-			PermutationVector.Set< FDeferredLightPS::FCloudTransmittance >(false);
-			PermutationVector.Set< FDeferredLightPS::FStrataFastPath >(StrataTileMaterialType == EStrataTileMaterialType::ESimple);
+			const bool bCameraInsideLightGeometry = ((FVector)View.ViewMatrices.GetViewOrigin() - LightBounds.Center).SizeSquared() < FMath::Square(LightBounds.W * 1.05f + View.NearClippingDistance * 2.0f)
+			//const bool bCameraInsideLightGeometry = LightProxy->AffectsBounds( FSphere( View.ViewMatrices.GetViewOrigin(), View.NearClippingDistance * 2.0f ) )
+				// Always draw backfaces in ortho
+				//@todo - accurate ortho camera / light intersection
+				|| !View.IsPerspectiveProjection();
+
+			SetBoundingGeometryRasterizerAndDepthState(GraphicsPSOInit, View, bCameraInsideLightGeometry);
 
-			TShaderMapRef< FDeferredLightPS > PixelShader( View.ShaderMap, PermutationVector );
 			GraphicsPSOInit.BoundShaderState.VertexDeclarationRHI = GetVertexDeclarationFVector4();
 			GraphicsPSOInit.BoundShaderState.VertexShaderRHI = VertexShader.GetVertexShader();
 			GraphicsPSOInit.BoundShaderState.PixelShaderRHI = PixelShader.GetPixelShader();
+			SetGraphicsPipelineState(RHICmdList, GraphicsPSOInit, StencilRef);			
+			SetShaderParameters(RHICmdList, PixelShader, PixelShader.GetPixelShader(), PassParameters->PS);
 
-			SetGraphicsPipelineState(RHICmdList, GraphicsPSOInit, StencilRef);
-			PixelShader->SetParameters(RHICmdList, View, LightSceneInfo, ScreenShadowMaskTexture, LightingChannelsTexture, &RenderLightParams, nullptr);
-		}
+			FDeferredLightVS::FParameters VSParameters2 = FDeferredLightVS::GetParameters(View, LightSceneInfo);
+			SetShaderParameters(RHICmdList, VertexShader, VertexShader.GetVertexShader(), VSParameters2);
 
-		FDeferredLightVS::FParameters VSParameters2 = FDeferredLightVS::GetParameters(View, LightSceneInfo);
-		SetShaderParameters(RHICmdList, VertexShader, VertexShader.GetVertexShader(), VSParameters2);
+			// Use DBT to allow work culling on shadow lights
+			if (GraphicsPSOInit.bDepthBounds)
+			{
+				// Can use the depth bounds test to skip work for pixels which won't be touched by the light (i.e outside the depth range)
+				float NearDepth = 1.f;
+				float FarDepth = 0.f;
+				CalculateLightNearFarDepthFromBounds(View,LightBounds,NearDepth,FarDepth);
+				if (NearDepth <= FarDepth)
+				{
+					NearDepth = 1.0f;
+					FarDepth = 0.0f;
+				}
 
-		// Use DBT to allow work culling on shadow lights
-		if (GraphicsPSOInit.bDepthBounds)
-		{
-			// Can use the depth bounds test to skip work for pixels which won't be touched by the light (i.e outside the depth range)
-			float NearDepth = 1.f;
-			float FarDepth = 0.f;
-			CalculateLightNearFarDepthFromBounds(View,LightBounds,NearDepth,FarDepth);
+				// UE uses reversed depth, so far < near
+				RHICmdList.SetDepthBounds(FarDepth, NearDepth);
+			}
 
-			if (NearDepth <= FarDepth)
+			if( LightType == LightType_Point || LightType == LightType_Rect )
 			{
-				NearDepth = 1.0f;
-				FarDepth = 0.0f;
+				// Apply the point or spot light with some approximate bounding geometry,
+				// So we can get speedups from depth testing and not processing pixels outside of the light's influence.
+				StencilingGeometry::DrawSphere(RHICmdList);
 			}
-
-			// UE uses reversed depth, so far < near
-			RHICmdList.SetDepthBounds(FarDepth, NearDepth);
-		}
-
-		if( LightProxy->GetLightType() == LightType_Point ||
-			LightProxy->GetLightType() == LightType_Rect )
-		{
-			// Apply the point or spot light with some approximate bounding geometry,
-			// So we can get speedups from depth testing and not processing pixels outside of the light's influence.
-			StencilingGeometry::DrawSphere(RHICmdList);
-		}
-		else if (LightProxy->GetLightType() == LightType_Spot)
-		{
-			StencilingGeometry::DrawCone(RHICmdList);
-		}
-	}	
+			else if (LightType == LightType_Spot)
+			{
+				StencilingGeometry::DrawCone(RHICmdList);
+			}
+		}	
+	}); // RenderPass
 }
 
-static void TempRenderLight(
-	FRHICommandList& RHICmdList,
-	const FScene* Scene,
-	const FViewInfo& View,
-	const FLightSceneInfo* LightSceneInfo,
-	FRenderLightParameters* PassParameters,	// If this is null, it means we cannot use Strata tiles and fallback to previous behavior.
-	FRHITexture* ScreenShadowMaskTexture,
-	FRHITexture* LightingChannelsTexture,
-	bool bRenderOverlap)
-{
-	SCOPE_CYCLE_COUNTER(STAT_DirectLightRenderingTime);
-	INC_DWORD_STAT(STAT_NumLightsUsingStandardDeferred);
-
-	const bool bStrataClassificationEnabled = Strata::IsStrataEnabled() && Strata::IsClassificationEnabled();
-	const bool bTilePassesReadingStrataEnabled = Strata::ShouldPassesReadingStrataBeTiled(Scene->GetFeatureLevel());
-
-	if (bStrataClassificationEnabled && bTilePassesReadingStrataEnabled && PassParameters != nullptr)
-	{
-		const bool bEnableStrataTiledPass = true;
-		const bool bEnableStrataStencilTest = false;
 
-		{
-			SCOPED_DRAW_EVENT(RHICmdList, StrataSimpleMaterial);
-			InternalRenderLight(RHICmdList, Scene, View, LightSceneInfo, PassParameters, ScreenShadowMaskTexture, LightingChannelsTexture, bRenderOverlap, bEnableStrataStencilTest, bEnableStrataTiledPass, EStrataTileMaterialType::ESimple);
-		}
-		{
-			SCOPED_DRAW_EVENT(RHICmdList, StrataComplexMaterial);
-			InternalRenderLight(RHICmdList, Scene, View, LightSceneInfo, PassParameters, ScreenShadowMaskTexture, LightingChannelsTexture, bRenderOverlap, bEnableStrataStencilTest, bEnableStrataTiledPass, EStrataTileMaterialType::EComplex);
-		}
-	}
-	else
-	{
-		const bool bEnableStrataTiledPass = false;
-		const bool bEnableStrataStencilTest = bStrataClassificationEnabled;
+/** Shader parameters for Standard Deferred Light Overlap Debug pass. */
+BEGIN_SHADER_PARAMETER_STRUCT(FRenderLightParameters, )
+	// PS/VS parameter structs
+	SHADER_PARAMETER_STRUCT_INCLUDE(FDeferredLightPS::FParameters, PS)	
+	SHADER_PARAMETER_STRUCT_INCLUDE(FDeferredLightVS::FParameters, VS)
+	// Strata tiles
+	SHADER_PARAMETER_RDG_BUFFER_SRV(Buffer<uint>, TileListBufferSimple)
+	SHADER_PARAMETER_RDG_BUFFER_SRV(Buffer<uint>, TileListBufferComplex)
+	RDG_BUFFER_ACCESS(TileIndirectBufferSimple, ERHIAccess::IndirectArgs)
+	RDG_BUFFER_ACCESS(TileIndirectBufferComplex, ERHIAccess::IndirectArgs)
+END_SHADER_PARAMETER_STRUCT()
 
-		InternalRenderLight(RHICmdList, Scene, View, LightSceneInfo, PassParameters, ScreenShadowMaskTexture, LightingChannelsTexture, bRenderOverlap, bEnableStrataStencilTest, bEnableStrataTiledPass, EStrataTileMaterialType::EComplex);
-		if (Strata::IsStrataEnabled() && Strata::IsClassificationEnabled())
-		{
-			InternalRenderLight(RHICmdList, Scene, View, LightSceneInfo, PassParameters, ScreenShadowMaskTexture, LightingChannelsTexture, bRenderOverlap, bEnableStrataStencilTest, bEnableStrataTiledPass, EStrataTileMaterialType::ESimple);
-		}
-	}
-}
+/** Shader parameters for Standard Deferred Light pass. */
+BEGIN_SHADER_PARAMETER_STRUCT(FRenderLightOverlapParameters, )
+	// PS/VS parameter structs
+	SHADER_PARAMETER_STRUCT_INCLUDE(FDeferredLightOverlapPS::FParameters, PS)
+	SHADER_PARAMETER_STRUCT_INCLUDE(FDeferredLightVS::FParameters, VS)
+END_SHADER_PARAMETER_STRUCT()
 
 static void RenderLight(
 	FRDGBuilder& GraphBuilder,
@@ -2378,200 +2015,270 @@ static void RenderLight(
 		return;
 	}
 
-	ERDGPassFlags PassFlags = ERDGPassFlags::Raster;
+	SCOPE_CYCLE_COUNTER(STAT_DirectLightRenderingTime);
+	INC_DWORD_STAT(STAT_NumLightsUsingStandardDeferred);
 
-	const FVolumetricCloudRenderSceneInfo* CloudInfo = Scene->GetVolumetricCloudSceneInfo();
+	const FLightSceneProxy* RESTRICT LightProxy = LightSceneInfo->Proxy;
+	const bool bUseIESTexture = View.Family->EngineShowFlags.TexturedLightProfiles && (LightSceneInfo->Proxy->GetIESTextureResource() != 0);
+	const bool bTransmission = LightProxy->Transmission();
+	const FSphere LightBounds = LightProxy->GetBoundingSphere();
+	const ELightComponentType LightType = (ELightComponentType)LightProxy->GetLightType();
+	const bool bIsRadial = LightType != LightType_Directional;
 
-	FRenderLightParameters* PassParameters = GraphBuilder.AllocParameters<FRenderLightParameters>();
-	if (bCloudShadow)
-		GetRenderLightParameters(View, SceneTextures, View.HairStrandsViewData, ScreenShadowMaskTexture, LightingChannelsTexture, GetCloudShadowAOParameters(GraphBuilder, View, CloudInfo), *PassParameters);
-	else
-		GetRenderLightParameters(View, SceneTextures, View.HairStrandsViewData, ScreenShadowMaskTexture, LightingChannelsTexture, {}, *PassParameters);
+	// Debug Overlap shader
+	if (bRenderOverlap)
+	{
+		FRenderLightOverlapParameters* PassParameters = GraphBuilder.AllocParameters<FRenderLightOverlapParameters>();
+		// PS - General parameters
+		PassParameters->PS.bHasValidChannel = LightSceneInfo->Proxy->GetPreviewShadowMapChannel() == INDEX_NONE ? 0.0f : 1.0f;
+		PassParameters->PS.View = View.ViewUniformBuffer;
+		PassParameters->PS.DeferredLight = CreateDeferredLightUniformBuffer(GraphBuilder, View, *LightSceneInfo);
+		PassParameters->PS.SceneTextures = SceneTextures.UniformBuffer;
+		PassParameters->PS.RenderTargets[0] = FRenderTargetBinding(SceneTextures.Color.Target, ERenderTargetLoadAction::ELoad);
+		if (SceneTextures.Depth.Target)
+		{
+			PassParameters->PS.RenderTargets.DepthStencil = FDepthStencilBinding(SceneTextures.Depth.Target, ERenderTargetLoadAction::ELoad, ERenderTargetLoadAction::ELoad, FExclusiveDepthStencil::DepthRead_StencilWrite);
+		}
+		// VS - General parameters
+		if (bIsRadial)
+		{
+			PassParameters->VS = FDeferredLightVS::GetParameters(View, LightSceneInfo, false);
+		}
+		else
+		{
+			PassParameters->VS = FDeferredLightVS::GetParameters(View, false);
+		}
 
-	GraphBuilder.AddPass(
-		RDG_EVENT_NAME("StandardDeferredLighting"),
-		PassParameters,
-		PassFlags,
-		[Scene, &View, LightSceneInfo, ScreenShadowMaskTexture, LightingChannelsTexture, bRenderOverlap, PassParameters](FRHICommandList& RHICmdList)
+		FDeferredLightOverlapPS::FPermutationDomain PermutationVector;
+		PermutationVector.Set<FDeferredLightOverlapPS::FRadialAttenuation>(bIsRadial);
+		TShaderMapRef<FDeferredLightOverlapPS> PixelShader(View.ShaderMap, PermutationVector);
+		
+		InternalRenderLight(GraphBuilder, Scene, View, LightSceneInfo, PixelShader, PassParameters, EStrataTileMaterialType::ECount, TEXT("StandardDeferredOverlap"));
+	}
+	// Lighting shader
+	else
 	{
-		TempRenderLight(
-			RHICmdList,
-			Scene, 
-			View,
-			LightSceneInfo,
-			PassParameters,
-			TryGetRHI(ScreenShadowMaskTexture),
-			TryGetRHI(LightingChannelsTexture),
-			bRenderOverlap);
+		FRenderLightParameters* PassParameters = GraphBuilder.AllocParameters<FRenderLightParameters>();
+		// PS - Generatl parameters
+		PassParameters->PS = GetDeferredLightPSParameters(GraphBuilder, Scene, View, LightSceneInfo, SceneTextures.Color.Target, SceneTextures.Depth.Target, SceneTextures.UniformBuffer, View.HairStrandsViewData.UniformBuffer, ScreenShadowMaskTexture, LightingChannelsTexture, bCloudShadow);
+		// VS - General parameters
+		if (bIsRadial)
+		{
+			PassParameters->VS = FDeferredLightVS::GetParameters(View, LightSceneInfo, false);
+		}
+		else // Directional
+		{
+			PassParameters->VS = FDeferredLightVS::GetParameters(View, false);
+		}
+		// VS - Strata tile parameters
+		if (Strata::ShouldPassesReadingStrataBeTiled(View.Family->GetFeatureLevel()))
+		{
+			PassParameters->TileListBufferSimple = View.StrataSceneData->ClassificationTileListBufferSRV[EStrataTileMaterialType::ESimple];
+			PassParameters->TileListBufferComplex = View.StrataSceneData->ClassificationTileListBufferSRV[EStrataTileMaterialType::EComplex];
+			PassParameters->TileIndirectBufferSimple = View.StrataSceneData->ClassificationTileIndirectBuffer[EStrataTileMaterialType::ESimple];
+			PassParameters->TileIndirectBufferComplex = View.StrataSceneData->ClassificationTileIndirectBuffer[EStrataTileMaterialType::EComplex];
+		}
+		else
+		{
+			FRDGBufferRef BufferDummy = GSystemTextures.GetDefaultBuffer(GraphBuilder, 4, 0u);
+			FRDGBufferSRVRef BufferDummySRV = GraphBuilder.CreateSRV(BufferDummy, PF_R32_UINT);
+			PassParameters->TileListBufferSimple = BufferDummySRV;
+			PassParameters->TileListBufferComplex = BufferDummySRV;
+			PassParameters->TileIndirectBufferSimple = BufferDummy;
+			PassParameters->TileIndirectBufferComplex = BufferDummy;
+		}
 
-	});
+		FDeferredLightPS::FPermutationDomain PermutationVector;
+		PermutationVector.Set< FDeferredLightPS::FTransmissionDim >(bTransmission);
+		PermutationVector.Set< FDeferredLightPS::FHairLighting>(0);
+		PermutationVector.Set< FDeferredLightPS::FLightingChannelsDim >(View.bUsesLightingChannels);
+		PermutationVector.Set< FDeferredLightPS::FVisualizeCullingDim >(View.Family->EngineShowFlags.VisualizeLightCulling);
+		PermutationVector.Set< FDeferredLightPS::FStrataFastPath >(false);
+		if (bIsRadial)
+		{
+			PermutationVector.Set< FDeferredLightPS::FSourceShapeDim >(LightProxy->IsRectLight() ? ELightSourceShape::Rect : ELightSourceShape::Capsule);
+			PermutationVector.Set< FDeferredLightPS::FSourceTextureDim >(LightProxy->IsRectLight() && LightProxy->HasSourceTexture());
+			PermutationVector.Set< FDeferredLightPS::FIESProfileDim >(bUseIESTexture);
+			PermutationVector.Set< FDeferredLightPS::FInverseSquaredDim >(LightProxy->IsInverseSquared());
+			PermutationVector.Set< FDeferredLightPS::FAnistropicMaterials >(ShouldRenderAnisotropyPass(View) && !LightSceneInfo->Proxy->IsRectLight());
+			PermutationVector.Set < FDeferredLightPS::FAtmosphereTransmittance >(false);
+			PermutationVector.Set< FDeferredLightPS::FCloudTransmittance >(false);
+		}
+		else // Directional
+		{
+			PermutationVector.Set< FDeferredLightPS::FSourceShapeDim >(ELightSourceShape::Directional);
+			PermutationVector.Set< FDeferredLightPS::FSourceTextureDim >(false);
+			PermutationVector.Set< FDeferredLightPS::FIESProfileDim >(false);
+			PermutationVector.Set< FDeferredLightPS::FInverseSquaredDim >(false);
+			PermutationVector.Set< FDeferredLightPS::FAnistropicMaterials >(ShouldRenderAnisotropyPass(View));
+			// Only directional lights are rendered in this path, so we only need to check if it is use to light the atmosphere
+			PermutationVector.Set< FDeferredLightPS::FAtmosphereTransmittance >(PassParameters->PS.CloudShadow.bAtmospherePerPixelTransmittance > 0);
+			PermutationVector.Set< FDeferredLightPS::FCloudTransmittance >(PassParameters->PS.CloudShadow.bCloudPerPixelTransmittance > 0);
+		}
+
+		// Strata tile rendering: 
+		// * if the light is directional, then dispatch a set of rect tiles
+		// * if the light is radial/local, then dispatch a light geometry with stencil test. The stencil buffer has been prefilled with the tile result (simple/complex) 
+		//   so that the geometry get correctly stencil culled on complex/simple part of the screen
+		if (Strata::IsStrataEnabled() && Strata::IsClassificationEnabled())
+		{
+			// Simple tiles
+			{
+				PermutationVector.Set< FDeferredLightPS::FStrataFastPath>(true);
+				TShaderMapRef< FDeferredLightPS > PixelShader(View.ShaderMap, PermutationVector);
+				InternalRenderLight(GraphBuilder, Scene, View, LightSceneInfo, PixelShader, PassParameters, EStrataTileMaterialType::ESimple, TEXT("StandardDeferredLight"));
+			}
+			// Complex tiles
+			{
+				PermutationVector.Set< FDeferredLightPS::FStrataFastPath>(false);
+				TShaderMapRef< FDeferredLightPS > PixelShader(View.ShaderMap, PermutationVector);
+				InternalRenderLight(GraphBuilder, Scene, View, LightSceneInfo, PixelShader, PassParameters, EStrataTileMaterialType::EComplex, TEXT("StandardDeferredLight"));
+			}
+		}
+		else
+		{
+			PermutationVector.Set< FDeferredLightPS::FStrataFastPath>(false);
+			TShaderMapRef< FDeferredLightPS > PixelShader(View.ShaderMap, PermutationVector);
+			InternalRenderLight(GraphBuilder, Scene, View, LightSceneInfo, PixelShader, PassParameters, EStrataTileMaterialType::ECount, TEXT("StandardDeferredLight"));
+		}
+	}
 }
 
+/** Shader parameters for Standard Deferred Light for HairStrands pass. */
 BEGIN_SHADER_PARAMETER_STRUCT(FRenderLightForHairParameters, )
-	SHADER_PARAMETER_STRUCT_INCLUDE(FRenderLightParameters, Light)
-	SHADER_PARAMETER_RDG_UNIFORM_BUFFER(FHairStrandsViewUniformParameters, HairStrands)
-	SHADER_PARAMETER_RDG_BUFFER_SRV(Buffer, HairTransmittanceMaskSRV)
+	SHADER_PARAMETER_STRUCT_INCLUDE(FDeferredLightHairVS::FParameters, VS)
+	SHADER_PARAMETER_STRUCT_INCLUDE(FDeferredLightPS::FParameters, PS)
 END_SHADER_PARAMETER_STRUCT()
 
 void FDeferredShadingSceneRenderer::RenderLightForHair(
 	FRDGBuilder& GraphBuilder,
 	FViewInfo& View,
-	TRDGUniformBufferRef<FSceneTextureUniformParameters> SceneTexturesUniformBuffer,
+	const FMinimalSceneTextures& SceneTextures,
 	const FLightSceneInfo* LightSceneInfo,
 	FRDGTextureRef HairShadowMaskTexture,
 	FRDGTextureRef LightingChannelsTexture,
 	const FHairStrandsTransmittanceMaskData& InTransmittanceMaskData,
 	const bool bForwardRendering)
 {
+	// Ensure the light is valid for this view
 	const bool bHairRenderingEnabled = HairStrands::HasViewHairStrandsData(View);
-	if (!bHairRenderingEnabled)
+	if (!bHairRenderingEnabled || !LightSceneInfo->ShouldRenderLight(View) || View.HairStrandsViewData.VisibilityData.SampleLightingTexture == nullptr)
 	{
 		return;
 	}
 	
+	// Sanity check
+	check(InTransmittanceMaskData.TransmittanceMask);
+
 	SCOPE_CYCLE_COUNTER(STAT_DirectLightRenderingTime);
 	INC_DWORD_STAT(STAT_NumLightsUsingStandardDeferred);
 	RDG_EVENT_SCOPE(GraphBuilder, "StandardDeferredLighting_Hair");
-
-	const FSphere LightBounds = LightSceneInfo->Proxy->GetBoundingSphere();
-	const bool bTransmission = LightSceneInfo->Proxy->Transmission();
-
-	const FVolumetricCloudRenderSceneInfo* CloudInfo = Scene->GetVolumetricCloudSceneInfo();
-
-
+	RDG_GPU_MASK_SCOPE(GraphBuilder, View.GPUMask);
+
+	const bool bIsDirectional = LightSceneInfo->Proxy->GetLightType() == LightType_Directional;
+	const bool bCloudShadow   = bIsDirectional;
+
+	FRenderLightForHairParameters* PassParameters = GraphBuilder.AllocParameters<FRenderLightForHairParameters>();
+	// VS - General parameters
+	PassParameters->VS.HairStrands = HairStrands::BindHairStrandsViewUniformParameters(View);
+	// PS - General parameters
+	PassParameters->PS = GetDeferredLightPSParameters(
+		GraphBuilder,
+		Scene,
+		View,
+		LightSceneInfo,
+		SceneTextures.Color.Target,
+		SceneTextures.Depth.Target,
+		SceneTextures.UniformBuffer,
+		HairStrands::BindHairStrandsViewUniformParameters(View),
+		HairShadowMaskTexture,
+		LightingChannelsTexture,
+		bCloudShadow);
+
+	// PS - Hair parameters
+	const FIntPoint SampleLightingViewportResolution = View.HairStrandsViewData.VisibilityData.SampleLightingViewportResolution;
+	PassParameters->PS.HairTransmittanceBuffer = GraphBuilder.CreateSRV(InTransmittanceMaskData.TransmittanceMask, FHairStrandsTransmittanceMaskData::Format);
+	PassParameters->PS.HairTransmittanceBufferMaxCount = InTransmittanceMaskData.TransmittanceMask ? InTransmittanceMaskData.TransmittanceMask->Desc.NumElements : 0;
+	PassParameters->PS.ShadowChannelMask = FVector4f(1, 1, 1, 1);
+	if (HairShadowMaskTexture)
 	{
-		RDG_GPU_MASK_SCOPE(GraphBuilder, View.GPUMask);
-
-		// Ensure the light is valid for this view
-		if (!LightSceneInfo->ShouldRenderLight(View))
-		{
-			return;
-		}
-
-		const FHairStrandsVisibilityData& HairVisibilityData = View.HairStrandsViewData.VisibilityData;
-		if (!HairVisibilityData.SampleLightingTexture)
-		{
-			return;
-		}
-
-		FVector4f ShadowChannelMask = FVector4f(1, 1, 1, 1);
-		if (bForwardRendering)
-		{
-			ShadowChannelMask = FVector4f(0, 0, 0, 0);
-			ShadowChannelMask[FMath::Clamp(LightSceneInfo->GetDynamicShadowMapChannel(), 0, 3)] = 1.f;
-		}
-
-		FRenderLightForHairParameters* PassParameters = GraphBuilder.AllocParameters<FRenderLightForHairParameters>();
-		GetRenderLightParameters(
-			View,
-			HairVisibilityData.SampleLightingTexture, 
-			nullptr, 
-			SceneTexturesUniformBuffer, 
-			HairStrands::BindHairStrandsViewUniformParameters(View), 
-			HairShadowMaskTexture, 
-			LightingChannelsTexture, 
-			GetCloudShadowAOParameters(GraphBuilder, View, CloudInfo), 
-			PassParameters->Light);
-
-		// Sanity check
-		check(InTransmittanceMaskData.TransmittanceMask);
-
-		PassParameters->HairStrands = HairStrands::BindHairStrandsViewUniformParameters(View);
-		PassParameters->HairTransmittanceMaskSRV = GraphBuilder.CreateSRV(InTransmittanceMaskData.TransmittanceMask, FHairStrandsTransmittanceMaskData::Format);
-
-		const bool bIsShadowMaskValid = !!PassParameters->Light.ShadowMaskTexture;
-		const uint32 MaxTransmittanceElementCount = InTransmittanceMaskData.TransmittanceMask ? InTransmittanceMaskData.TransmittanceMask->Desc.NumElements : 0;
-		GraphBuilder.AddPass(
-			{},
-			PassParameters,
-			ERDGPassFlags::Raster,
-			[this, &HairVisibilityData, &View, PassParameters, LightSceneInfo, MaxTransmittanceElementCount, HairShadowMaskTexture, LightingChannelsTexture, bIsShadowMaskValid, ShadowChannelMask](FRHICommandList& RHICmdList)
-		{
-			RHICmdList.SetViewport(0, 0, 0.0f, HairVisibilityData.SampleLightingViewportResolution.X, HairVisibilityData.SampleLightingViewportResolution.Y, 1.0f);
+		PassParameters->PS.ScreenShadowMaskSubPixelTexture = HairShadowMaskTexture;
+		PassParameters->PS.HairShadowMaskValid = true;
+	}
+	if (bForwardRendering)
+	{
+		PassParameters->PS.ShadowChannelMask = FVector4f(0, 0, 0, 0);
+		PassParameters->PS.ShadowChannelMask[FMath::Clamp(LightSceneInfo->GetDynamicShadowMapChannel(), 0, 3)] = 1.f;
+	}
+	PassParameters->PS.RenderTargets[0] = FRenderTargetBinding(View.HairStrandsViewData.VisibilityData.SampleLightingTexture, ERenderTargetLoadAction::ELoad);
+	PassParameters->PS.RenderTargets.DepthStencil = FDepthStencilBinding(nullptr, ERenderTargetLoadAction::ENoAction, ERenderTargetLoadAction::ENoAction, FExclusiveDepthStencil::DepthNop_StencilNop);
+
+	FDeferredLightPS::FPermutationDomain PermutationVector;
+	PermutationVector.Set< FDeferredLightPS::FLightingChannelsDim >(View.bUsesLightingChannels);
+	PermutationVector.Set< FDeferredLightPS::FVisualizeCullingDim >(false);
+	PermutationVector.Set< FDeferredLightPS::FTransmissionDim >(false);
+	PermutationVector.Set< FDeferredLightPS::FHairLighting>(1);
+	if (bIsDirectional)
+	{
+		PermutationVector.Set< FDeferredLightPS::FSourceShapeDim >(ELightSourceShape::Directional);
+		PermutationVector.Set< FDeferredLightPS::FSourceTextureDim >(false);
+		PermutationVector.Set< FDeferredLightPS::FIESProfileDim >(false);
+		PermutationVector.Set< FDeferredLightPS::FInverseSquaredDim >(false);
+		PermutationVector.Set< FDeferredLightPS::FAtmosphereTransmittance >(PassParameters->PS.CloudShadow.bAtmospherePerPixelTransmittance > 0.f);
+		PermutationVector.Set< FDeferredLightPS::FCloudTransmittance >(PassParameters->PS.CloudShadow.bCloudPerPixelTransmittance > 0.f);
+	}
+	else
+	{
+		const bool bUseIESTexture = View.Family->EngineShowFlags.TexturedLightProfiles && LightSceneInfo->Proxy->GetIESTextureResource() != 0;
+		PermutationVector.Set< FDeferredLightPS::FSourceShapeDim >(LightSceneInfo->Proxy->IsRectLight() ? ELightSourceShape::Rect : ELightSourceShape::Capsule);
+		PermutationVector.Set< FDeferredLightPS::FSourceTextureDim >(LightSceneInfo->Proxy->IsRectLight() && LightSceneInfo->Proxy->HasSourceTexture());
+		PermutationVector.Set< FDeferredLightPS::FIESProfileDim >(bUseIESTexture);
+		PermutationVector.Set< FDeferredLightPS::FInverseSquaredDim >(LightSceneInfo->Proxy->IsInverseSquared());
+		PermutationVector.Set< FDeferredLightPS::FAtmosphereTransmittance >(false);
+		PermutationVector.Set< FDeferredLightPS::FCloudTransmittance >(false);
+	}
 
-			FRenderLightParams RenderLightParams;
-			RenderLightParams.DeepShadow_TransmittanceMaskBufferMaxCount = MaxTransmittanceElementCount;
-			RenderLightParams.ScreenShadowMaskSubPixelTexture = bIsShadowMaskValid ? PassParameters->Light.ShadowMaskTexture->GetRHI() : GSystemTextures.WhiteDummy->GetShaderResourceRHI();
-			RenderLightParams.DeepShadow_TransmittanceMaskBuffer = PassParameters->HairTransmittanceMaskSRV->GetRHI();
-			RenderLightParams.ShadowChannelMask = ShadowChannelMask;
+	TShaderMapRef<FDeferredLightHairVS> VertexShader(View.ShaderMap);
+	TShaderMapRef<FDeferredLightPS> PixelShader(View.ShaderMap, PermutationVector);
 
-			FGraphicsPipelineStateInitializer GraphicsPSOInit;
-			RHICmdList.ApplyCachedRenderTargets(GraphicsPSOInit);
-			GraphicsPSOInit.BlendState = TStaticBlendState<CW_RGBA, BO_Add, BF_One, BF_One, BO_Max, BF_SourceAlpha, BF_DestAlpha>::GetRHI();
-			GraphicsPSOInit.PrimitiveType = PT_TriangleList;
+	GraphBuilder.AddPass(
+		{},
+		PassParameters,
+		ERDGPassFlags::Raster,
+		[this, VertexShader, PixelShader, PassParameters, SampleLightingViewportResolution](FRHICommandList& RHICmdList)
+	{
+		RHICmdList.SetViewport(0, 0, 0.0f, SampleLightingViewportResolution.X, SampleLightingViewportResolution.Y, 1.0f);
 
-			FDeferredLightPS::FPermutationDomain PermutationVector;
-			if (LightSceneInfo->Proxy->GetLightType() == LightType_Directional)
-			{
-				bool bAtmospherePerPixelTransmittance = false;
-				bool bCloudPerPixelTransmittance = false;
-				BindAtmosphereAndCloudResources(
-					Scene,
-					&View,
-					LightSceneInfo->Proxy,
-					RenderLightParams,
-					bAtmospherePerPixelTransmittance,
-					bCloudPerPixelTransmittance);
-
-				PermutationVector.Set< FDeferredLightPS::FSourceShapeDim >(ELightSourceShape::Directional);
-				PermutationVector.Set< FDeferredLightPS::FSourceTextureDim >(false);
-				PermutationVector.Set< FDeferredLightPS::FIESProfileDim >(false);
-				PermutationVector.Set< FDeferredLightPS::FInverseSquaredDim >(false);
-				PermutationVector.Set< FDeferredLightPS::FAtmosphereTransmittance >(bAtmospherePerPixelTransmittance);
-				PermutationVector.Set< FDeferredLightPS::FCloudTransmittance >(bCloudPerPixelTransmittance);
-			}
-			else
-			{
-				const bool bUseIESTexture = View.Family->EngineShowFlags.TexturedLightProfiles && LightSceneInfo->Proxy->GetIESTextureResource() != 0;
-				PermutationVector.Set< FDeferredLightPS::FSourceShapeDim >(LightSceneInfo->Proxy->IsRectLight() ? ELightSourceShape::Rect : ELightSourceShape::Capsule);
-				PermutationVector.Set< FDeferredLightPS::FSourceTextureDim >(LightSceneInfo->Proxy->IsRectLight() && LightSceneInfo->Proxy->HasSourceTexture());
-				PermutationVector.Set< FDeferredLightPS::FIESProfileDim >(bUseIESTexture);
-				PermutationVector.Set< FDeferredLightPS::FInverseSquaredDim >(LightSceneInfo->Proxy->IsInverseSquared());
-				PermutationVector.Set< FDeferredLightPS::FAtmosphereTransmittance >(false);
-				PermutationVector.Set< FDeferredLightPS::FCloudTransmittance >(false);
-			}
-			PermutationVector.Set< FDeferredLightPS::FLightingChannelsDim >(View.bUsesLightingChannels);
-			PermutationVector.Set< FDeferredLightPS::FVisualizeCullingDim >(false);
-			PermutationVector.Set< FDeferredLightPS::FTransmissionDim >(false);
-			PermutationVector.Set< FDeferredLightPS::FHairLighting>(1);
+		FGraphicsPipelineStateInitializer GraphicsPSOInit;
+		RHICmdList.ApplyCachedRenderTargets(GraphicsPSOInit);
+		GraphicsPSOInit.BlendState = TStaticBlendState<CW_RGBA, BO_Add, BF_One, BF_One, BO_Max, BF_SourceAlpha, BF_DestAlpha>::GetRHI();
+		GraphicsPSOInit.PrimitiveType = PT_TriangleList;
 
-			TShaderMapRef<FDeferredLightHairVS> VertexShader(View.ShaderMap);
-			TShaderMapRef<FDeferredLightPS> PixelShader(View.ShaderMap, PermutationVector);
+		GraphicsPSOInit.BoundShaderState.VertexDeclarationRHI = GetVertexDeclarationFVector4();
+		GraphicsPSOInit.bDepthBounds = false;
+		GraphicsPSOInit.RasterizerState = TStaticRasterizerState<FM_Solid, CM_None>::GetRHI();
+		GraphicsPSOInit.DepthStencilState = TStaticDepthStencilState<false, CF_Always>::GetRHI();
+		GraphicsPSOInit.BoundShaderState.VertexShaderRHI = VertexShader.GetVertexShader();
+		GraphicsPSOInit.BoundShaderState.PixelShaderRHI = PixelShader.GetPixelShader();
+		GraphicsPSOInit.PrimitiveType = PT_TriangleList;
+		SetGraphicsPipelineState(RHICmdList, GraphicsPSOInit, 0);
 
-			GraphicsPSOInit.BoundShaderState.VertexDeclarationRHI = GetVertexDeclarationFVector4();
-			GraphicsPSOInit.bDepthBounds = false;
-			GraphicsPSOInit.RasterizerState = TStaticRasterizerState<FM_Solid, CM_None>::GetRHI();
-			GraphicsPSOInit.DepthStencilState = TStaticDepthStencilState<false, CF_Always>::GetRHI();
-			GraphicsPSOInit.BoundShaderState.VertexShaderRHI = VertexShader.GetVertexShader();
-			GraphicsPSOInit.BoundShaderState.PixelShaderRHI = PixelShader.GetPixelShader();
-			GraphicsPSOInit.PrimitiveType = PT_TriangleList;
-			SetGraphicsPipelineState(RHICmdList, GraphicsPSOInit, 0);
-
-			VertexShader->SetParameters(RHICmdList, View, PassParameters->HairStrands->GetRHI());
-			PixelShader->SetParameters(
-				RHICmdList,
-				View,
-				LightSceneInfo,
-				TryGetRHI(HairShadowMaskTexture),
-				TryGetRHI(LightingChannelsTexture),
-				&RenderLightParams,
-				PassParameters->HairStrands->GetRHI());
-
-			RHICmdList.SetStreamSource(0, nullptr, 0);
-			RHICmdList.DrawPrimitive(0, 1, 1);
-		});
-	}
+		SetShaderParameters(RHICmdList, VertexShader, VertexShader.GetVertexShader(), PassParameters->VS);
+		SetShaderParameters(RHICmdList, PixelShader, PixelShader.GetPixelShader(), PassParameters->PS);
+		
+		RHICmdList.SetStreamSource(0, nullptr, 0);
+		RHICmdList.DrawPrimitive(0, 1, 1);
+	});
 }
 
 // Forward lighting version for hair
 void FDeferredShadingSceneRenderer::RenderLightsForHair(
 	FRDGBuilder& GraphBuilder,
-	TRDGUniformBufferRef<FSceneTextureUniformParameters> SceneTexturesUniformBuffer,
+	const FMinimalSceneTextures& SceneTextures,
 	FSortedLightSetSceneInfo &SortedLightSet,
 	FRDGTextureRef ScreenShadowMaskSubPixelTexture,
 	FRDGTextureRef LightingChannelsTexture)
 {
-	const FSimpleLightArray &SimpleLights = SortedLightSet.SimpleLights;
 	const TArray<FSortedLightSceneInfo, SceneRenderingAllocator> &SortedLights = SortedLightSet.SortedLights;
 	const int32 UnbatchedLightStart = SortedLightSet.UnbatchedLightStart;
 	const int32 SimpleLightsEnd = SortedLightSet.SimpleLightsEnd;
@@ -2604,7 +2311,7 @@ void FDeferredShadingSceneRenderer::RenderLightsForHair(
 					RenderLightForHair(
 						GraphBuilder,
 						View,
-						SceneTexturesUniformBuffer,
+						SceneTextures,
 						&LightSceneInfo,
 						ScreenShadowMaskSubPixelTexture,
 						LightingChannelsTexture,
@@ -2617,57 +2324,136 @@ void FDeferredShadingSceneRenderer::RenderLightsForHair(
 }
 
 BEGIN_SHADER_PARAMETER_STRUCT(FSimpleLightsStandardDeferredParameters, )
-	SHADER_PARAMETER_RDG_UNIFORM_BUFFER(FSceneTextureUniformParameters, SceneTextures)
-	SHADER_PARAMETER_RDG_UNIFORM_BUFFER(FStrataGlobalUniformParameters, Strata)
-	RENDER_TARGET_BINDING_SLOTS()
+	SHADER_PARAMETER_STRUCT_INCLUDE(FDeferredLightPS::FParameters, PS)
+	SHADER_PARAMETER_STRUCT_INCLUDE(FDeferredLightVS::FParameters, VS)
 END_SHADER_PARAMETER_STRUCT()
 
+static FSimpleLightsStandardDeferredParameters GetRenderLightSimpleParameters(
+	FRDGBuilder& GraphBuilder,
+	const FScene* Scene,
+	const FViewInfo& View,
+	const FMinimalSceneTextures& SceneTextures,
+	const FSimpleLightEntry& SimpleLight,
+	const FVector3f SimpleLightPosition)
+{
+	FSimpleLightsStandardDeferredParameters Out;
+
+	FRDGTextureRef WhiteDummy = GSystemTextures.GetWhiteDummy(GraphBuilder);
+	FRDGTextureRef DepthDummy = GSystemTextures.GetDepthDummy(GraphBuilder);
+	FRDGBufferRef BufferDummy = GSystemTextures.GetDefaultBuffer(GraphBuilder, 4, 0u);
+	FRDGBufferSRVRef BufferDummySRV = GraphBuilder.CreateSRV(BufferDummy, PF_R32_UINT);
+	
+	// PS - General parmaeters
+	Out.PS.SceneTextures = SceneTextures.UniformBuffer;
+	Out.PS.HairStrands = View.HairStrandsViewData.UniformBuffer;
+	Out.PS.Strata = Strata::BindStrataGlobalUniformParameters(View.StrataSceneData);
+	Out.PS.LightingChannelsTexture = WhiteDummy;
+	Out.PS.LightingChannelsSampler = TStaticSamplerState<SF_Point, AM_Clamp, AM_Clamp, AM_Clamp>::GetRHI();
+	Out.PS.CloudShadowAO = GetCloudShadowAOParameters(GraphBuilder, View, nullptr);
+	Out.PS.CloudShadow = GetCloudShadowParameters(GraphBuilder);
+	Out.PS.LightAttenuationTexture = WhiteDummy;
+	Out.PS.LightAttenuationTextureSampler = TStaticSamplerState<SF_Point, AM_Wrap, AM_Wrap, AM_Wrap>::GetRHI();
+	Out.PS.DummyRectLightTextureForCapsuleCompilerWarning = DepthDummy;
+	Out.PS.DummyRectLightSamplerForCapsuleCompilerWarning = TStaticSamplerState<SF_Bilinear, AM_Clamp, AM_Clamp, AM_Clamp>::GetRHI();
+	Out.PS.IESTextureSampler = TStaticSamplerState<SF_Bilinear, AM_Clamp, AM_Clamp, AM_Clamp>::GetRHI();
+	Out.PS.IESTexture = GSystemTextures.WhiteDummy->GetShaderResourceRHI();
+	Out.PS.View = View.ViewUniformBuffer;
+	Out.PS.DeferredLight = CreateDeferredLightUniformBuffer(GraphBuilder, SimpleLight, SimpleLightPosition);
+	// PS - Hair (default)
+	Out.PS.ScreenShadowMaskSubPixelTexture = WhiteDummy;
+	Out.PS.HairTransmittanceBuffer = BufferDummySRV;
+	Out.PS.HairTransmittanceBufferMaxCount = 0;
+	Out.PS.HairShadowMaskValid = false;
+	Out.PS.ShadowChannelMask = FVector4f(1, 1, 1, 1);
+	// PS - RT/Depth
+	Out.PS.RenderTargets[0] = FRenderTargetBinding(SceneTextures.Color.Target, ERenderTargetLoadAction::ELoad);
+	if (SceneTextures.Depth.Target)
+	{
+		Out.PS.RenderTargets.DepthStencil = FDepthStencilBinding(SceneTextures.Depth.Target, ERenderTargetLoadAction::ELoad, ERenderTargetLoadAction::ELoad, FExclusiveDepthStencil::DepthRead_StencilWrite);
+	}
+
+	// VS - General paramters (dummy geometry, as the geometry is setup within the pass light loop)
+	FSphere DummySphere;
+	DummySphere.Center = FVector::ZeroVector;
+	DummySphere.W = 1.0f;
+	Out.VS = FDeferredLightVS::GetParameters(View, DummySphere, false);
+
+	return Out;
+}
+
 void FDeferredShadingSceneRenderer::RenderSimpleLightsStandardDeferred(
 	FRDGBuilder& GraphBuilder,
 	const FMinimalSceneTextures& SceneTextures,
 	const FSimpleLightArray& SimpleLights)
 {
+	if (SimpleLights.InstanceData.Num() == 0)
+	{
+		return;
+	}
+
 	SCOPE_CYCLE_COUNTER(STAT_DirectLightRenderingTime);
 	INC_DWORD_STAT_BY(STAT_NumLightsUsingStandardDeferred, SimpleLights.InstanceData.Num());
 
-	FSimpleLightsStandardDeferredParameters* PassParameters = GraphBuilder.AllocParameters<FSimpleLightsStandardDeferredParameters>();
-	PassParameters->SceneTextures = SceneTextures.UniformBuffer;
-	PassParameters->Strata = Strata::BindStrataGlobalUniformParameters(&Scene->StrataSceneData);
-	PassParameters->RenderTargets[0] = FRenderTargetBinding(SceneTextures.Color.Target, ERenderTargetLoadAction::ELoad);
-	PassParameters->RenderTargets.DepthStencil = FDepthStencilBinding(SceneTextures.Depth.Target, ERenderTargetLoadAction::ELoad, ERenderTargetLoadAction::ELoad, FExclusiveDepthStencil::DepthRead_StencilWrite);
-
-	GraphBuilder.AddPass(
-		RDG_EVENT_NAME("StandardDeferredSimpleLights"),
-		PassParameters,
-		ERDGPassFlags::Raster,
-		[this, &SimpleLights](FRHICommandList& RHICmdList)
+	for (int32 ViewIndex = 0, NumViews = Views.Num(); ViewIndex < NumViews; ViewIndex++)
 	{
-		FGraphicsPipelineStateInitializer GraphicsPSOInit;
-		RHICmdList.ApplyCachedRenderTargets(GraphicsPSOInit);
+		const FViewInfo& View = Views[ViewIndex];
 
-		// Use additive blending for color
-		GraphicsPSOInit.BlendState = TStaticBlendState<CW_RGBA, BO_Add, BF_One, BF_One, BO_Add, BF_One, BF_One>::GetRHI();
-		GraphicsPSOInit.PrimitiveType = PT_TriangleList;
+		FSimpleLightsStandardDeferredParameters* PassParameters = GraphBuilder.AllocParameters<FSimpleLightsStandardDeferredParameters>();
+		*PassParameters = GetRenderLightSimpleParameters(
+			GraphBuilder,
+			Scene,
+			View,
+			SceneTextures,
+			SimpleLights.InstanceData[0], // Use a dummy light to create the PassParameter buffer. The light data will be
+			FVector3f(0, 0, 0));		  // update dynamically with the pass light loop for efficiency purpose
+
+		FDeferredLightPS::FPermutationDomain PermutationVector;
+		PermutationVector.Set< FDeferredLightPS::FSourceShapeDim >(ELightSourceShape::Capsule);
+		PermutationVector.Set< FDeferredLightPS::FIESProfileDim >(false);
+		PermutationVector.Set< FDeferredLightPS::FVisualizeCullingDim >(View.Family->EngineShowFlags.VisualizeLightCulling);
+		PermutationVector.Set< FDeferredLightPS::FLightingChannelsDim >(false);
+		PermutationVector.Set< FDeferredLightPS::FAnistropicMaterials >(false);
+		PermutationVector.Set< FDeferredLightPS::FTransmissionDim >(false);
+		PermutationVector.Set< FDeferredLightPS::FHairLighting>(0);
+		PermutationVector.Set< FDeferredLightPS::FAtmosphereTransmittance >(false);
+		PermutationVector.Set< FDeferredLightPS::FCloudTransmittance >(false);
+		PermutationVector.Set< FDeferredLightPS::FStrataFastPath>(false);
+
+		FDeferredLightPS::FPermutationDomain PermutationVector_bInverseSquaredFalloff0;
+		FDeferredLightPS::FPermutationDomain PermutationVector_bInverseSquaredFalloff1;
+		PermutationVector_bInverseSquaredFalloff0.Set<FDeferredLightPS::FInverseSquaredDim>(false);
+		PermutationVector_bInverseSquaredFalloff1.Set<FDeferredLightPS::FInverseSquaredDim>(true);
+		TShaderMapRef<FDeferredLightPS> PixelShader0(View.ShaderMap, PermutationVector);
+		TShaderMapRef<FDeferredLightPS> PixelShader1(View.ShaderMap, PermutationVector);
 
-		const int32 NumViews = Views.Num();
-		for (int32 LightIndex = 0; LightIndex < SimpleLights.InstanceData.Num(); LightIndex++)
+		FDeferredLightVS::FPermutationDomain PermutationVectorVS;
+		PermutationVectorVS.Set<FDeferredLightVS::FRadialLight>(true);
+		TShaderMapRef<FDeferredLightVS> VertexShader(View.ShaderMap, PermutationVectorVS);
+
+		GraphBuilder.AddPass(
+			RDG_EVENT_NAME("StandardDeferredSimpleLights"),
+			PassParameters,
+			ERDGPassFlags::Raster,
+			[this, &View, &SimpleLights, ViewIndex, NumViews, PassParameters, PixelShader0, PixelShader1, VertexShader](FRHICommandList& RHICmdList)
 		{
-			const FSimpleLightEntry& SimpleLight = SimpleLights.InstanceData[LightIndex];
+			FGraphicsPipelineStateInitializer GraphicsPSOInit;
+			RHICmdList.ApplyCachedRenderTargets(GraphicsPSOInit);
+
+			// Use additive blending for color
+			GraphicsPSOInit.BlendState = TStaticBlendState<CW_RGBA, BO_Add, BF_One, BF_One, BO_Add, BF_One, BF_One>::GetRHI();
+			GraphicsPSOInit.PrimitiveType = PT_TriangleList;
 
-			for (int32 ViewIndex = 0; ViewIndex < NumViews; ViewIndex++)
+			for (int32 LightIndex = 0; LightIndex < SimpleLights.InstanceData.Num(); LightIndex++)
 			{
+				const FSimpleLightEntry& SimpleLight = SimpleLights.InstanceData[LightIndex];
+
 				const FSimpleLightPerViewEntry& SimpleLightPerViewData = SimpleLights.GetViewDependentData(LightIndex, ViewIndex, NumViews);
 				const FSphere LightBounds(SimpleLightPerViewData.Position, SimpleLight.Radius);
 
-				const FViewInfo& View = Views[ViewIndex];
 
 				// Set the device viewport for the view.
 				RHICmdList.SetViewport(View.ViewRect.Min.X, View.ViewRect.Min.Y, 0.0f, View.ViewRect.Max.X, View.ViewRect.Max.Y, 1.0f);
 
-				FDeferredLightVS::FPermutationDomain PermutationVectorVS;
-				PermutationVectorVS.Set<FDeferredLightVS::FRadialLight>(true);
-				TShaderMapRef<FDeferredLightVS> VertexShader(View.ShaderMap, PermutationVectorVS);
-
 				const bool bCameraInsideLightGeometry = ((FVector)View.ViewMatrices.GetViewOrigin() - LightBounds.Center).SizeSquared() < FMath::Square(LightBounds.W * 1.05f + View.NearClippingDistance * 2.0f)
 								// Always draw backfaces in ortho
 								//@todo - accurate ortho camera / light intersection
@@ -2677,35 +2463,35 @@ void FDeferredShadingSceneRenderer::RenderSimpleLightsStandardDeferred(
 
 				const bool bInverseSquaredFalloff = SimpleLight.Exponent == 0;
 
-				FDeferredLightPS::FPermutationDomain PermutationVector;
-				PermutationVector.Set< FDeferredLightPS::FSourceShapeDim >(ELightSourceShape::Capsule);
-				PermutationVector.Set< FDeferredLightPS::FIESProfileDim >(false);
-				PermutationVector.Set< FDeferredLightPS::FInverseSquaredDim >(bInverseSquaredFalloff);
-				PermutationVector.Set< FDeferredLightPS::FVisualizeCullingDim >(View.Family->EngineShowFlags.VisualizeLightCulling);
-				PermutationVector.Set< FDeferredLightPS::FLightingChannelsDim >(false);
-				PermutationVector.Set< FDeferredLightPS::FAnistropicMaterials >(false);
-				PermutationVector.Set< FDeferredLightPS::FTransmissionDim >(false);
-				PermutationVector.Set< FDeferredLightPS::FHairLighting>(0);
-				PermutationVector.Set< FDeferredLightPS::FAtmosphereTransmittance >(false);
-				PermutationVector.Set< FDeferredLightPS::FCloudTransmittance >(false);
-				PermutationVector.Set< FDeferredLightPS::FStrataFastPath>(false);
-
-				TShaderMapRef< FDeferredLightPS > PixelShader(View.ShaderMap, PermutationVector);
 				GraphicsPSOInit.BoundShaderState.VertexDeclarationRHI = GetVertexDeclarationFVector4();
 				GraphicsPSOInit.BoundShaderState.VertexShaderRHI = VertexShader.GetVertexShader();
-				GraphicsPSOInit.BoundShaderState.PixelShaderRHI = PixelShader.GetPixelShader();
+				GraphicsPSOInit.BoundShaderState.PixelShaderRHI = bInverseSquaredFalloff ? PixelShader1.GetPixelShader() : PixelShader0.GetPixelShader();
 				SetGraphicsPipelineState(RHICmdList, GraphicsPSOInit, 0);
-				PixelShader->SetParametersSimpleLight(RHICmdList, View, SimpleLight, SimpleLightPerViewData);
 
+
+				// Update the light parameters with a custom uniform buffer
+				FDeferredLightUniformStruct DeferredLightUniformsValue = GetSimpleDeferredLightParameters(SimpleLight, SimpleLightPerViewData);
+				if (bInverseSquaredFalloff)
+				{
+					SetShaderParameters(RHICmdList, PixelShader1, PixelShader1.GetPixelShader(), PassParameters->PS);
+					SetUniformBufferParameterImmediate(RHICmdList, RHICmdList.GetBoundPixelShader(), PixelShader0->GetUniformBufferParameter<FDeferredLightUniformStruct>(), DeferredLightUniformsValue);
+				}
+				else
+				{
+					SetShaderParameters(RHICmdList, PixelShader0, PixelShader0.GetPixelShader(), PassParameters->PS);
+					SetUniformBufferParameterImmediate(RHICmdList, RHICmdList.GetBoundPixelShader(), PixelShader0->GetUniformBufferParameter<FDeferredLightUniformStruct>(), DeferredLightUniformsValue);
+				}
+
+				// Update vertex shader parameters with custom parameters/uniform buffer
 				FDeferredLightVS::FParameters ParametersVS = FDeferredLightVS::GetParameters(View, LightBounds);
-				SetShaderParameters(RHICmdList, VertexShader, VertexShader.GetVertexShader(), ParametersVS);
+				SetShaderParameters(RHICmdList, VertexShader, VertexShader.GetVertexShader(), ParametersVS /*PassParameters->VS*/);
 
 				// Apply the point or spot light with some approximately bounding geometry,
 				// So we can get speedups from depth testing and not processing pixels outside of the light's influence.
-				StencilingGeometry::DrawSphere(RHICmdList);
+				StencilingGeometry::DrawSphere(RHICmdList);				
 			}
-		}
-	});
+		});
+	}
 }
 
 class FCopyStencilToLightingChannelsPS : public FGlobalShader
diff --git a/Engine/Source/Runtime/Renderer/Private/LightRendering.h b/Engine/Source/Runtime/Renderer/Private/LightRendering.h
--- a/Engine/Source/Runtime/Renderer/Private/LightRendering.h
+++ b/Engine/Source/Runtime/Renderer/Private/LightRendering.h
@@ -47,10 +47,9 @@ void SetDeferredLightParameters(
 	SetUniformBufferParameterImmediate(RHICmdList, ShaderRHI, DeferredLightUniformBufferParameter, GetDeferredLightParameters(View, *LightSceneInfo));
 }
 
-extern void SetupSimpleDeferredLightParameters(
+extern FDeferredLightUniformStruct GetSimpleDeferredLightParameters(
 	const FSimpleLightEntry& SimpleLight,
-	const FSimpleLightPerViewEntry &SimpleLightPerViewData, 
-	FDeferredLightUniformStruct& DeferredLightUniformsValue);
+	const FSimpleLightPerViewEntry &SimpleLightPerViewData);
 
 template<typename ShaderRHIParamRef>
 void SetSimpleDeferredLightParameters(
@@ -61,8 +60,7 @@ void SetSimpleDeferredLightParameters(
 	const FSimpleLightPerViewEntry &SimpleLightPerViewData,
 	const FSceneView& View)
 {
-	FDeferredLightUniformStruct DeferredLightUniformsValue;
-	SetupSimpleDeferredLightParameters(SimpleLight, SimpleLightPerViewData, DeferredLightUniformsValue);
+	FDeferredLightUniformStruct DeferredLightUniformsValue = GetSimpleDeferredLightParameters(SimpleLight, SimpleLightPerViewData);
 	SetUniformBufferParameterImmediate(RHICmdList, ShaderRHI, DeferredLightUniformBufferParameter, DeferredLightUniformsValue);
 }
 
@@ -537,19 +535,26 @@ public:
 		float U, float V,
 		float SizeU, float SizeV,
 		FIntPoint TargetSize,
-		FIntPoint TextureSize)
+		FIntPoint TextureSize,
+		bool bBindViewUniform = true)
 	{
 		FParameters Out;
-		Out.View = View.ViewUniformBuffer;
+		if (bBindViewUniform)
+		{
+			Out.View = View.ViewUniformBuffer;
+		}
 		Out.Geometry = FStencilingGeometryShaderParameters::GetParameters(FVector4(0,0,0,0));
 		Out.FullScreenRect = GetFullScreenRectParameters(X, Y, SizeX, SizeY, U, V, SizeU, SizeV, TargetSize, TextureSize);
 		return Out;
 	}
 
-	static FParameters GetParameters(const FViewInfo& View)
+	static FParameters GetParameters(const FViewInfo& View, bool bBindViewUniform = true)
 	{
 		FParameters Out;
-		Out.View = View.ViewUniformBuffer;
+		if (bBindViewUniform)
+		{
+			Out.View = View.ViewUniformBuffer;
+		}
 		Out.Geometry = FStencilingGeometryShaderParameters::GetParameters(FVector4(0, 0, 0, 0));
 		Out.FullScreenRect = GetFullScreenRectParameters(
 			0, 0,
@@ -562,22 +567,28 @@ public:
 		return Out;
 	}
 
-	static FParameters GetParameters(const FViewInfo& View, const FSphere& LightBounds)
+	static FParameters GetParameters(const FViewInfo& View, const FSphere& LightBounds, bool bBindViewUniform = true)
 	{
 		FVector4 StencilingSpherePosAndScale;
 		StencilingGeometry::GStencilSphereVertexBuffer.CalcTransform(StencilingSpherePosAndScale, LightBounds, View.ViewMatrices.GetPreViewTranslation());
 
 		FParameters Out;
-		Out.View = View.ViewUniformBuffer;
+		if (bBindViewUniform)
+		{
+			Out.View = View.ViewUniformBuffer;
+		}
 		Out.Geometry = FStencilingGeometryShaderParameters::GetParameters(StencilingSpherePosAndScale);
 		Out.FullScreenRect = GetFullScreenRectParameters(0, 0, 0, 0, 0, 0, 0, 0, FIntPoint(1, 1), FIntPoint(1, 1));
 		return Out;
 	}
 
-	static FParameters GetParameters(const FViewInfo& View, const FLightSceneInfo* LightSceneInfo)
+	static FParameters GetParameters(const FViewInfo& View, const FLightSceneInfo* LightSceneInfo, bool bBindViewUniform = true)
 	{
 		FParameters Out;
-		Out.View = View.ViewUniformBuffer;
+		if (bBindViewUniform)
+		{
+			Out.View = View.ViewUniformBuffer;
+		}
 		Out.Geometry = FStencilingGeometryShaderParameters::GetParameters(View, LightSceneInfo);
 		Out.FullScreenRect = GetFullScreenRectParameters(0, 0, 0, 0, 0, 0, 0, 0, FIntPoint(1, 1), FIntPoint(1, 1));
 		return Out;
diff --git a/Engine/Source/Runtime/Renderer/Private/MobileDeferredShadingPass.cpp b/Engine/Source/Runtime/Renderer/Private/MobileDeferredShadingPass.cpp
--- a/Engine/Source/Runtime/Renderer/Private/MobileDeferredShadingPass.cpp
+++ b/Engine/Source/Runtime/Renderer/Private/MobileDeferredShadingPass.cpp
@@ -577,8 +577,7 @@ static void RenderSimpleLights(
 
 		// Render light
 		FMobileRadialLightFunctionPS::FParameters PassParameters;
-		FDeferredLightUniformStruct DeferredLightUniformsValue;
-		SetupSimpleDeferredLightParameters(SimpleLight, SimpleLightPerViewData, DeferredLightUniformsValue);
+		FDeferredLightUniformStruct DeferredLightUniformsValue = GetSimpleDeferredLightParameters(SimpleLight, SimpleLightPerViewData);
 		PassParameters.DeferredLightUniforms = TUniformBufferRef<FDeferredLightUniformStruct>::CreateUniformBufferImmediate(DeferredLightUniformsValue, EUniformBufferUsage::UniformBuffer_SingleFrame);
 		PassParameters.IESTexture = GWhiteTexture->TextureRHI;
 		PassParameters.IESTextureSampler = GWhiteTexture->SamplerStateRHI;
diff --git a/Engine/Source/Runtime/Renderer/Private/SceneRendering.h b/Engine/Source/Runtime/Renderer/Private/SceneRendering.h
--- a/Engine/Source/Runtime/Renderer/Private/SceneRendering.h
+++ b/Engine/Source/Runtime/Renderer/Private/SceneRendering.h
@@ -1433,7 +1433,7 @@ public:
 	TUniformBufferRef<FViewUniformShaderParameters> VolumetricRenderTargetViewUniformBuffer;
 	// The effective cloud shadow target this frame independently of the fact that a view can have a state (primary view) or not (sky light reflection capture)
 	FRDGTextureRef VolumetricCloudShadowRenderTarget[NUM_ATMOSPHERE_LIGHTS] = {};
-	// We need to extract that RDG resource because the RHI must be accessed to setup FDeferredLightPS 
+	// We need to extract that RDG resource because the RHI must be accessed to setup FTranslucentLightingInjectPS & TVolumetricFogLightScatteringCS
 	TRefCountPtr<IPooledRenderTarget> VolumetricCloudShadowExtractedRenderTarget[NUM_ATMOSPHERE_LIGHTS] = {};
 
 	/** Used when there is no view state, buffers reallocate every frame. */
diff --git a/Engine/Source/Runtime/Renderer/Private/VolumetricCloudRendering.cpp b/Engine/Source/Runtime/Renderer/Private/VolumetricCloudRendering.cpp
--- a/Engine/Source/Runtime/Renderer/Private/VolumetricCloudRendering.cpp
+++ b/Engine/Source/Runtime/Renderer/Private/VolumetricCloudRendering.cpp
@@ -392,14 +392,22 @@ static bool ShouldUsePerSampleAtmosphereTransmittance(const FScene* Scene, const
 
 FVolumetricCloudShadowAOParameters GetCloudShadowAOParameters(FRDGBuilder& GraphBuilder, const FViewInfo& View, const FVolumetricCloudRenderSceneInfo* CloudInfo)
 {
-	FVolumetricCloudShadowAOParameters Parameters;
+	FVolumetricCloudShadowAOParameters Out;
 	if (CloudInfo)
 	{
-		Parameters.ShadowMap0 = View.VolumetricCloudShadowRenderTarget[0];
-		Parameters.ShadowMap1 = View.VolumetricCloudShadowRenderTarget[1];
-		Parameters.SkyAO = View.VolumetricCloudSkyAO;
+		Out.ShadowMap0 = View.VolumetricCloudShadowRenderTarget[0];
+		Out.ShadowMap1 = View.VolumetricCloudShadowRenderTarget[1];
+		Out.SkyAO = View.VolumetricCloudSkyAO;
 	}
-	return Parameters;
+	else
+	{
+		FRDGTextureRef DummyWhite = GSystemTextures.GetWhiteDummy(GraphBuilder);
+		FRDGTextureRef DummyDepth = GSystemTextures.GetDepthDummy(GraphBuilder);
+		Out.ShadowMap0 = DummyDepth;
+		Out.ShadowMap1 = DummyDepth;
+		Out.SkyAO = DummyWhite;
+	}
+	return Out;
 }
 
 void GetCloudShadowAOData(const FVolumetricCloudRenderSceneInfo* CloudInfo, const FViewInfo& View, FRDGBuilder& GraphBuilder, FCloudShadowAOData& OutData)
@@ -424,6 +432,63 @@ void GetCloudShadowAOData(const FVolumetricCloudRenderSceneInfo* CloudInfo, cons
 	OutData.bShouldSampleCloudSkyAO  = OutData.bShouldSampleCloudSkyAO  && bCloudComponentValid;
 }
 
+FVolumetricCloudShadowParameters GetCloudShadowParameters(FRDGBuilder& GraphBuilder)
+{
+	FVolumetricCloudShadowParameters Out;
+
+	Out.bAtmospherePerPixelTransmittance = 0.f;
+	Out.CloudShadowmapTexture = GSystemTextures.GetDepthDummy(GraphBuilder);
+	Out.CloudShadowmapFarDepthKm = 0.f;
+	Out.CloudShadowmapWorldToLightClipMatrix = FMatrix44f::Identity;
+	Out.CloudShadowmapStrength = 0.f;
+	Out.CloudShadowmapSampler = TStaticSamplerState<SF_Bilinear, AM_Clamp, AM_Clamp, AM_Clamp>::GetRHI();
+	Out.bCloudPerPixelTransmittance = 0.f;
+	return Out;
+}
+
+FVolumetricCloudShadowParameters GetCloudShadowParameters(FRDGBuilder& GraphBuilder, const FScene* Scene, const FViewInfo& View, const FLightSceneProxy* LightProxy, const FVolumetricCloudRenderSceneInfo* CloudInfo)
+{
+	FVolumetricCloudShadowParameters Out;
+
+	Out.bAtmospherePerPixelTransmittance =
+		LightProxy->GetLightType() == LightType_Directional &&
+		LightProxy->IsUsedAsAtmosphereSunLight() &&
+		LightProxy->GetUsePerPixelAtmosphereTransmittance() &&
+		ShouldRenderSkyAtmosphere(Scene, View.Family->EngineShowFlags) ? 1.f : 0.f;
+
+	const FLightSceneProxy* AtmosphereLight0Proxy = Scene->AtmosphereLights[0] ? Scene->AtmosphereLights[0]->Proxy : nullptr;
+	const FLightSceneProxy* AtmosphereLight1Proxy = Scene->AtmosphereLights[1] ? Scene->AtmosphereLights[1]->Proxy : nullptr;
+	const bool VolumetricCloudShadowMap0Valid = View.VolumetricCloudShadowRenderTarget[0] != nullptr;
+	const bool VolumetricCloudShadowMap1Valid = View.VolumetricCloudShadowRenderTarget[1] != nullptr;
+	const bool bLight0CloudPerPixelTransmittance = CloudInfo && VolumetricCloudShadowMap0Valid && AtmosphereLight0Proxy == LightProxy && AtmosphereLight0Proxy && AtmosphereLight0Proxy->GetCloudShadowOnSurfaceStrength() > 0.0f;
+	const bool bLight1CloudPerPixelTransmittance = CloudInfo && VolumetricCloudShadowMap1Valid && AtmosphereLight1Proxy == LightProxy && AtmosphereLight1Proxy && AtmosphereLight1Proxy->GetCloudShadowOnSurfaceStrength() > 0.0f;
+	if (bLight0CloudPerPixelTransmittance)
+	{
+		Out.CloudShadowmapTexture = View.VolumetricCloudShadowRenderTarget[0];
+		Out.CloudShadowmapFarDepthKm = CloudInfo->GetVolumetricCloudCommonShaderParameters().CloudShadowmapFarDepthKm[0].X;
+		Out.CloudShadowmapWorldToLightClipMatrix = CloudInfo->GetVolumetricCloudCommonShaderParameters().CloudShadowmapWorldToLightClipMatrix[0];
+		Out.CloudShadowmapStrength = AtmosphereLight0Proxy->GetCloudShadowOnSurfaceStrength();
+	}
+	else if (bLight1CloudPerPixelTransmittance)
+	{
+		Out.CloudShadowmapTexture = View.VolumetricCloudShadowRenderTarget[1];
+		Out.CloudShadowmapFarDepthKm = CloudInfo->GetVolumetricCloudCommonShaderParameters().CloudShadowmapFarDepthKm[1].X;
+		Out.CloudShadowmapWorldToLightClipMatrix = CloudInfo->GetVolumetricCloudCommonShaderParameters().CloudShadowmapWorldToLightClipMatrix[1];
+		Out.CloudShadowmapStrength = AtmosphereLight1Proxy->GetCloudShadowOnSurfaceStrength();
+	}
+	else
+	{
+		Out.CloudShadowmapTexture = GSystemTextures.GetDepthDummy(GraphBuilder);	
+		Out.CloudShadowmapFarDepthKm = 0.f;
+		Out.CloudShadowmapWorldToLightClipMatrix = FMatrix44f::Identity;
+		Out.CloudShadowmapStrength = 0.f;
+	}
+	Out.CloudShadowmapSampler = TStaticSamplerState<SF_Bilinear, AM_Clamp, AM_Clamp, AM_Clamp>::GetRHI();
+
+	Out.bCloudPerPixelTransmittance = bLight0CloudPerPixelTransmittance || bLight1CloudPerPixelTransmittance ? 1.f : 0.f;
+
+	return Out;
+}
 
 /*=============================================================================
 	FVolumetricCloudRenderSceneInfo implementation.
diff --git a/Engine/Source/Runtime/Renderer/Private/VolumetricCloudRendering.h b/Engine/Source/Runtime/Renderer/Private/VolumetricCloudRendering.h
--- a/Engine/Source/Runtime/Renderer/Private/VolumetricCloudRendering.h
+++ b/Engine/Source/Runtime/Renderer/Private/VolumetricCloudRendering.h
@@ -18,6 +18,7 @@ class FViewInfo;
 class FLightSceneInfo;
 class UVolumetricCloudComponent;
 class FVolumetricCloudSceneProxy;
+class FLightSceneProxy;
 
 struct FEngineShowFlags;
 
@@ -149,13 +150,26 @@ private:
 };
 
 BEGIN_SHADER_PARAMETER_STRUCT(FVolumetricCloudShadowAOParameters, )
-	RDG_TEXTURE_ACCESS(ShadowMap0, ERHIAccess::SRVGraphics)
-	RDG_TEXTURE_ACCESS(ShadowMap1, ERHIAccess::SRVGraphics)
-	RDG_TEXTURE_ACCESS(SkyAO, ERHIAccess::SRVGraphics)
+	SHADER_PARAMETER_RDG_TEXTURE(Texture2D, ShadowMap0)
+	SHADER_PARAMETER_RDG_TEXTURE(Texture2D, ShadowMap1)
+	SHADER_PARAMETER_RDG_TEXTURE(Texture2D, SkyAO)
 END_SHADER_PARAMETER_STRUCT()
 
 FVolumetricCloudShadowAOParameters GetCloudShadowAOParameters(FRDGBuilder& GraphBuilder, const FViewInfo& View, const FVolumetricCloudRenderSceneInfo* CloudInfo);
 
+BEGIN_SHADER_PARAMETER_STRUCT(FVolumetricCloudShadowParameters, )
+	SHADER_PARAMETER(float, bAtmospherePerPixelTransmittance)
+	SHADER_PARAMETER(float, bCloudPerPixelTransmittance)
+	SHADER_PARAMETER(float, CloudShadowmapFarDepthKm)
+	SHADER_PARAMETER(FMatrix44f, CloudShadowmapWorldToLightClipMatrix)
+	SHADER_PARAMETER(float, CloudShadowmapStrength)
+	SHADER_PARAMETER_RDG_TEXTURE(Texture2D, CloudShadowmapTexture)
+	SHADER_PARAMETER_SAMPLER(SamplerState, CloudShadowmapSampler)
+END_SHADER_PARAMETER_STRUCT()
+
+FVolumetricCloudShadowParameters GetCloudShadowParameters(FRDGBuilder& GraphBuilder);
+FVolumetricCloudShadowParameters GetCloudShadowParameters(FRDGBuilder& GraphBuilder, const FScene* Scene, const FViewInfo& View, const FLightSceneProxy* LightProxy, const FVolumetricCloudRenderSceneInfo* CloudInfo);
+
 struct FCloudShadowAOData
 {
 	bool bShouldSampleCloudShadow;
