Add the ability to force instance culling for single-instance GPUScene primitives to prevent the rendering of invalid instances.

diff --git a/Engine/Plugins/FX/Niagara/Source/Niagara/Private/NiagaraRendererMeshes.cpp b/Engine/Plugins/FX/Niagara/Source/Niagara/Private/NiagaraRendererMeshes.cpp
--- a/Engine/Plugins/FX/Niagara/Source/Niagara/Private/NiagaraRendererMeshes.cpp
+++ b/Engine/Plugins/FX/Niagara/Source/Niagara/Private/NiagaraRendererMeshes.cpp
@@ -978,6 +978,13 @@ void FNiagaraRendererMeshes::SetupElementForGPUScene(
 		GPUSceneRes.GPUWriteParams.DistanceCullRangeSquared = DistanceCullRange * DistanceCullRange;
 		GPUSceneRes.GPUWriteParams.bNeedsPrevTransform		= bNeedsPrevTransform ? 1 : 0;
 
+		// We need to set this flag to force the system to always cull individual instances, because we may need to discard instances that are:
+		// * Not tied to any live particles
+		// * Are distance culled
+		// * Are culled because of mismatched MeshIndex
+		// * Are culled because of mismatched VisibilityTag
+		GPUSceneRes.DynamicPrimitiveData.bForceInstanceCulling = true;
+
 		GPUSceneRes.DynamicPrimitiveData.DataWriterGPU = FGPUSceneWriteDelegate::CreateLambda(				
 			[&GPUSceneRes](FRDGBuilder& GraphBuilder, const FGPUSceneWriteDelegateParams& Params)
 			{
diff --git a/Engine/Source/Runtime/Engine/Public/MeshBatch.h b/Engine/Source/Runtime/Engine/Public/MeshBatch.h
--- a/Engine/Source/Runtime/Engine/Public/MeshBatch.h
+++ b/Engine/Source/Runtime/Engine/Public/MeshBatch.h
@@ -76,6 +76,7 @@ struct FMeshBatchDynamicPrimitiveData
 	EGPUSceneGPUWritePass DataWriterGPUPass = EGPUSceneGPUWritePass::None;
 	uint16 PayloadDataFlags = 0;
 	uint32 NumInstanceCustomDataFloats = 0;
+	bool bForceInstanceCulling = false;
 
 	FORCEINLINE void SetPayloadDataFlags(uint16 Flags, bool bValue)
 	{
diff --git a/Engine/Source/Runtime/Renderer/Private/InstanceCulling/InstanceCullingContext.cpp b/Engine/Source/Runtime/Renderer/Private/InstanceCulling/InstanceCullingContext.cpp
--- a/Engine/Source/Runtime/Renderer/Private/InstanceCulling/InstanceCullingContext.cpp
+++ b/Engine/Source/Runtime/Renderer/Private/InstanceCulling/InstanceCullingContext.cpp
@@ -150,19 +150,19 @@ uint32 FInstanceCullingContext::AllocateIndirectArgs(const FMeshDrawCommand *Mes
 // 2.1 Only allocate indirect draw cmd if needed, 
 // 3. 
 
-void FInstanceCullingContext::AddInstancesToDrawCommand(uint32 IndirectArgsOffset, int32 InstanceDataOffset, bool bDynamicInstanceDataOffset, uint32 NumInstances)
+void FInstanceCullingContext::AddInstancesToDrawCommand(uint32 IndirectArgsOffset, int32 InstanceDataOffset, bool bDynamicInstanceDataOffset, bool bForceInstanceCulling, uint32 NumInstances)
 {
 	checkSlow(InstanceDataOffset >= 0);
 
 	// We special-case the single-instance (i.e., regular primitives) as they don't need culling (again).
 	// In actual fact this is not 100% true because dynamic path primitives may not have been culled.
-	EBatchProcessingMode Mode = NumInstances == 1 ? SingleInstanceProcessingMode : EBatchProcessingMode::Generic;
+	EBatchProcessingMode Mode = (NumInstances == 1 && !bForceInstanceCulling) ? SingleInstanceProcessingMode : EBatchProcessingMode::Generic;
 	// NOTE: we pack the bDynamicInstanceDataOffset in the lowest bit because the load balancer steals the upper bits of the payload!
 	LoadBalancers[uint32(Mode)]->Add(uint32(InstanceDataOffset), NumInstances, (IndirectArgsOffset << 1U) | uint32(bDynamicInstanceDataOffset));
 	TotalInstances += NumInstances;
 }
 
-void FInstanceCullingContext::AddInstanceRunsToDrawCommand(uint32 IndirectArgsOffset, int32 InstanceDataOffset, bool bDynamicInstanceDataOffset, const uint32* Runs, uint32 NumRuns)
+void FInstanceCullingContext::AddInstanceRunsToDrawCommand(uint32 IndirectArgsOffset, int32 InstanceDataOffset, bool bDynamicInstanceDataOffset, bool bForceInstanceCulling, const uint32* Runs, uint32 NumRuns)
 {
 	// Add items to current generic batch as they are instanced for sure.
 	for (uint32 Index = 0; Index < NumRuns; ++Index)
@@ -170,7 +170,7 @@ void FInstanceCullingContext::AddInstanceRunsToDrawCommand(uint32 IndirectArgsOf
 		uint32 RunStart = Runs[Index * 2];
 		uint32 RunEndIncl = Runs[Index * 2 + 1];
 		uint32 NumInstances = (RunEndIncl + 1U) - RunStart;
-		AddInstancesToDrawCommand(IndirectArgsOffset, InstanceDataOffset + RunStart, bDynamicInstanceDataOffset, NumInstances);
+		AddInstancesToDrawCommand(IndirectArgsOffset, InstanceDataOffset + RunStart, bDynamicInstanceDataOffset, bForceInstanceCulling, NumInstances);
 	}
 }
 
@@ -832,7 +832,8 @@ void FInstanceCullingContext::SetupDrawCommands(
 
 		const bool bSupportsGPUSceneInstancing = EnumHasAnyFlags(VisibleMeshDrawCommand.Flags, EFVisibleMeshDrawCommandFlags::HasPrimitiveIdStreamIndex);
 		const bool bMaterialMayModifyPosition = EnumHasAnyFlags(VisibleMeshDrawCommand.Flags, EFVisibleMeshDrawCommandFlags::MaterialMayModifyPosition);
-		const bool bUseIndirectDraw = bAlwaysUseIndirectDraws || (VisibleMeshDrawCommand.NumRuns > 0 || MeshDrawCommand->NumInstances > 1);
+		const bool bForceInstanceCulling = EnumHasAnyFlags(VisibleMeshDrawCommand.Flags, EFVisibleMeshDrawCommandFlags::ForceInstanceCulling);
+		const bool bUseIndirectDraw = bAlwaysUseIndirectDraws || bForceInstanceCulling || (VisibleMeshDrawCommand.NumRuns > 0 || MeshDrawCommand->NumInstances > 1);
 
 		if (bCompactIdenticalCommands && CurrentStateBucketId != -1 && VisibleMeshDrawCommand.StateBucketId == CurrentStateBucketId)
 		{
@@ -846,7 +847,7 @@ void FInstanceCullingContext::SetupDrawCommands(
 			if (DrawCmd.bUseIndirect == 0)
 			{
 				DrawCmd.IndirectArgsOffsetOrNumInstances += 1;
-		}
+			}
 		}
 		else
 		{
@@ -903,11 +904,11 @@ void FInstanceCullingContext::SetupDrawCommands(
 			// This will cause all instances belonging to the Primitive to be added to the command, if they are visible etc (GPU-Scene knows all - sees all)
 			if (VisibleMeshDrawCommand.RunArray)
 			{
-				AddInstanceRunsToDrawCommand(CurrentIndirectArgsOffset, VisibleMeshDrawCommand.PrimitiveIdInfo.InstanceSceneDataOffset, VisibleMeshDrawCommand.PrimitiveIdInfo.bIsDynamicPrimitive, VisibleMeshDrawCommand.RunArray, VisibleMeshDrawCommand.NumRuns);
+				AddInstanceRunsToDrawCommand(CurrentIndirectArgsOffset, VisibleMeshDrawCommand.PrimitiveIdInfo.InstanceSceneDataOffset, VisibleMeshDrawCommand.PrimitiveIdInfo.bIsDynamicPrimitive, bForceInstanceCulling, VisibleMeshDrawCommand.RunArray, VisibleMeshDrawCommand.NumRuns);
 			}
 			else 
 			{
-				AddInstancesToDrawCommand(CurrentIndirectArgsOffset, VisibleMeshDrawCommand.PrimitiveIdInfo.InstanceSceneDataOffset, VisibleMeshDrawCommand.PrimitiveIdInfo.bIsDynamicPrimitive, VisibleMeshDrawCommand.MeshDrawCommand->NumInstances);
+				AddInstancesToDrawCommand(CurrentIndirectArgsOffset, VisibleMeshDrawCommand.PrimitiveIdInfo.InstanceSceneDataOffset, VisibleMeshDrawCommand.PrimitiveIdInfo.bIsDynamicPrimitive, bForceInstanceCulling, VisibleMeshDrawCommand.MeshDrawCommand->NumInstances);
 			}
 		}
 	}
diff --git a/Engine/Source/Runtime/Renderer/Public/InstanceCulling/InstanceCullingContext.h b/Engine/Source/Runtime/Renderer/Public/InstanceCulling/InstanceCullingContext.h
--- a/Engine/Source/Runtime/Renderer/Public/InstanceCulling/InstanceCullingContext.h
+++ b/Engine/Source/Runtime/Renderer/Public/InstanceCulling/InstanceCullingContext.h
@@ -95,13 +95,13 @@ public:
 	 * Add command to cull a range of instances for the given mesh draw command index.
 	 * Multiple commands may add to the same slot, ordering is not preserved.
 	 */
-	void AddInstancesToDrawCommand(uint32 IndirectArgsOffset, int32 InstanceDataOffset, bool bDynamicInstanceDataOffset, uint32 NumInstances);
+	void AddInstancesToDrawCommand(uint32 IndirectArgsOffset, int32 InstanceDataOffset, bool bDynamicInstanceDataOffset, bool bForceInstanceCulling, uint32 NumInstances);
 
 	/**
 	 * Command that is executed in the per-view, post-cull pass to gather up the instances belonging to this primitive.
 	 * Multiple commands may add to the same slot, ordering is not preserved.
 	 */
-	void AddInstanceRunsToDrawCommand(uint32 IndirectArgsOffset, int32 InstanceDataOffset, bool bDynamicInstanceDataOffset, const uint32* Runs, uint32 NumRuns);
+	void AddInstanceRunsToDrawCommand(uint32 IndirectArgsOffset, int32 InstanceDataOffset, bool bDynamicInstanceDataOffset, bool bForceInstanceCulling, const uint32* Runs, uint32 NumRuns);
 
 	/*
 	 * Allocate space for indirect draw call argumens for a given MeshDrawCommand and initialize with draw command data.
diff --git a/Engine/Source/Runtime/Renderer/Public/MeshPassProcessor.h b/Engine/Source/Runtime/Renderer/Public/MeshPassProcessor.h
--- a/Engine/Source/Runtime/Renderer/Public/MeshPassProcessor.h
+++ b/Engine/Source/Runtime/Renderer/Public/MeshPassProcessor.h
@@ -1261,8 +1261,11 @@ enum class EFVisibleMeshDrawCommandFlags : uint8
 	/** If set, the mesh draw command supports primitive ID steam (required for dynamic instancing and GPU-Scene instance culling). */
 	HasPrimitiveIdStreamIndex = 1U << 1U,
 
-	All = MaterialMayModifyPosition | HasPrimitiveIdStreamIndex,
-	NumBits = 2U
+	/** If set, forces individual instances to always be culled independently from the primitive */
+	ForceInstanceCulling = 1U << 2U,	
+
+	All = MaterialMayModifyPosition | HasPrimitiveIdStreamIndex | ForceInstanceCulling,
+	NumBits = 3U
 };
 ENUM_CLASS_FLAGS(EFVisibleMeshDrawCommandFlags);
 static_assert(uint32(EFVisibleMeshDrawCommandFlags::All) < (1U << uint32(EFVisibleMeshDrawCommandFlags::NumBits)), "EFVisibleMeshDrawCommandFlags::NumBits too small to represent all flags in EFVisibleMeshDrawCommandFlags.");
diff --git a/Engine/Source/Runtime/Renderer/Public/MeshPassProcessor.inl b/Engine/Source/Runtime/Renderer/Public/MeshPassProcessor.inl
--- a/Engine/Source/Runtime/Renderer/Public/MeshPassProcessor.inl
+++ b/Engine/Source/Runtime/Renderer/Public/MeshPassProcessor.inl
@@ -48,12 +48,11 @@ void FMeshPassProcessor::BuildMeshDrawCommands(
 	const FPrimitiveSceneInfo* RESTRICT PrimitiveSceneInfo = PrimitiveSceneProxy ? PrimitiveSceneProxy->GetPrimitiveSceneInfo() : nullptr;
 
 	FMeshDrawCommand SharedMeshDrawCommand;
-
-	EFVisibleMeshDrawCommandFlags Flags = EFVisibleMeshDrawCommandFlags::Default;
+	EFVisibleMeshDrawCommandFlags SharedFlags = EFVisibleMeshDrawCommandFlags::Default;
 
 	if (MaterialResource.MaterialModifiesMeshPosition_RenderThread())
 	{
-		Flags |= EFVisibleMeshDrawCommandFlags::MaterialMayModifyPosition;
+		SharedFlags |= EFVisibleMeshDrawCommandFlags::MaterialMayModifyPosition;
 	}
 
 	SharedMeshDrawCommand.SetStencilRef(DrawRenderState.GetStencilRef());
@@ -90,7 +89,7 @@ void FMeshPassProcessor::BuildMeshDrawCommands(
 
 	if (SharedMeshDrawCommand.PrimitiveIdStreamIndex != INDEX_NONE)
 	{
-		Flags |= EFVisibleMeshDrawCommandFlags::HasPrimitiveIdStreamIndex;
+		SharedFlags |= EFVisibleMeshDrawCommandFlags::HasPrimitiveIdStreamIndex;
 	}
 
 	int32 DataOffset = 0;
@@ -122,6 +121,12 @@ void FMeshPassProcessor::BuildMeshDrawCommands(
 		{
 			const FMeshBatchElement& BatchElement = MeshBatch.Elements[BatchElementIndex];
 			FMeshDrawCommand& MeshDrawCommand = DrawListContext->AddCommand(SharedMeshDrawCommand, NumElements);
+			EFVisibleMeshDrawCommandFlags Flags = SharedFlags;
+
+			if (BatchElement.DynamicPrimitiveData != nullptr && BatchElement.DynamicPrimitiveData->bForceInstanceCulling)
+			{
+				Flags |= EFVisibleMeshDrawCommandFlags::ForceInstanceCulling;
+			}
 
 			DataOffset = 0;
 			if (PassShaders.VertexShader.IsValid())
