Remove inverse square falloff permutation.

diff --git a/Engine/Shaders/Private/DeferredLightPixelShaders.usf b/Engine/Shaders/Private/DeferredLightPixelShaders.usf
--- a/Engine/Shaders/Private/DeferredLightPixelShaders.usf
+++ b/Engine/Shaders/Private/DeferredLightPixelShaders.usf
@@ -94,8 +94,7 @@ FDeferredLightData SetupLightDataForStandardDeferred()
 	LightData.DistanceFadeMAD = DeferredLightUniforms.DistanceFadeMAD;
 	LightData.ShadowMapChannelMask = DeferredLightUniforms.ShadowMapChannelMask;
 	LightData.ShadowedBits = DeferredLightUniforms.ShadowedBits;
-
-	LightData.bInverseSquared = INVERSE_SQUARED_FALLOFF;
+	LightData.bInverseSquared = DeferredLightUniforms.FalloffExponent == 0;
 	LightData.bRadialLight = LIGHT_SOURCE_SHAPE > 0;
 	//@todo - permutation opportunity
 	LightData.bSpotLight = LIGHT_SOURCE_SHAPE > 0;
diff --git a/Engine/Shaders/Private/DynamicLightingCommon.ush b/Engine/Shaders/Private/DynamicLightingCommon.ush
--- a/Engine/Shaders/Private/DynamicLightingCommon.ush
+++ b/Engine/Shaders/Private/DynamicLightingCommon.ush
@@ -86,10 +86,10 @@ float CalcLightAttenuationTWS(float3 TranslatedWorldPosition, out float3 WorldLi
 	
 	float LightRadiusMask = Square( saturate( 1 - Square( DistanceSqr * DeferredLightUniforms.InvRadius * DeferredLightUniforms.InvRadius ) ) );
 	DistanceAttenuation *= LightRadiusMask;
-	
-#if !INVERSE_SQUARED_FALLOFF
-	DistanceAttenuation = RadialAttenuation(WorldLightVector * DeferredLightUniforms.InvRadius, DeferredLightUniforms.FalloffExponent);
-#endif
+	if (DeferredLightUniforms.FalloffExponent > 0.f)
+	{
+		DistanceAttenuation = RadialAttenuation(WorldLightVector * DeferredLightUniforms.InvRadius, DeferredLightUniforms.FalloffExponent);
+	}
 #endif
 
 	float SpotFalloff = 1;
diff --git a/Engine/Shaders/Private/MobileDeferredShading.usf b/Engine/Shaders/Private/MobileDeferredShading.usf
--- a/Engine/Shaders/Private/MobileDeferredShading.usf
+++ b/Engine/Shaders/Private/MobileDeferredShading.usf
@@ -30,10 +30,6 @@
 #include "/Engine/Public/Platform/Vulkan/VulkanSubpassSupport.ush"
 #endif
 
-#ifndef INVERSE_SQUARED_FALLOFF 
-#define INVERSE_SQUARED_FALLOFF 0	
-#endif 
-
 #ifndef IS_SPOT_LIGHT 
 #define IS_SPOT_LIGHT 0 
 #endif 
@@ -394,7 +390,7 @@ void MobileRadialLightPS(
 		LightData.Direction = DeferredLightUniforms.Direction;
 		LightData.SpotAngles = DeferredLightUniforms.SpotAngles;
 		LightData.SpecularScale = 1.0;
-		LightData.bInverseSquared = INVERSE_SQUARED_FALLOFF; 
+		LightData.bInverseSquared = DeferredLightUniforms.FalloffExponent == 0;
 		LightData.bSpotLight = IS_SPOT_LIGHT; 
 	}
 
diff --git a/Engine/Shaders/Private/VolumetricFog.usf b/Engine/Shaders/Private/VolumetricFog.usf
--- a/Engine/Shaders/Private/VolumetricFog.usf
+++ b/Engine/Shaders/Private/VolumetricFog.usf
@@ -204,10 +204,6 @@ int VirtualShadowMapId;
 #define HISTORY_MISS_SUPER_SAMPLE_COUNT 1
 #endif
 
-#ifndef INVERSE_SQUARED_FALLOFF
-#define INVERSE_SQUARED_FALLOFF 1
-#endif
-
 #if	USE_LIGHT_FUNCTION
 
 /** Matrix can be either WorldToShadow or WorldToLight depending on light type. */
@@ -291,8 +287,7 @@ void InjectShadowedLocalLightPS(
 		LightData.ShadowedBits = DeferredLightUniforms.ShadowedBits;
 		LightData.RectLightBarnCosAngle = DeferredLightUniforms.RectLightBarnCosAngle;
 		LightData.RectLightBarnLength = DeferredLightUniforms.RectLightBarnLength;
-
-		LightData.bInverseSquared = INVERSE_SQUARED_FALLOFF;		
+		LightData.bInverseSquared = DeferredLightUniforms.FalloffExponent == 0;
 		LightData.bRadialLight = true;
 		LightData.bSpotLight = LightData.SpotAngles.x > -2.0f;
 		LightData.bRectLight = LightData.RectLightBarnLength > -2.0f;
diff --git a/Engine/Source/Runtime/Renderer/Private/LightRendering.cpp b/Engine/Source/Runtime/Renderer/Private/LightRendering.cpp
--- a/Engine/Source/Runtime/Renderer/Private/LightRendering.cpp
+++ b/Engine/Source/Runtime/Renderer/Private/LightRendering.cpp
@@ -256,6 +256,11 @@ FDeferredLightUniformStruct GetDeferredLightParameters(const FSceneView& View, c
 
 	Out.LightingChannelMask = LightSceneInfo.Proxy->GetLightingChannelMask();
 
+	// Ensure the light falloff exponent is set to 0 so that lighting shaders handle it as inverse-squared attenuated light
+	if (LightSceneInfo.Proxy->IsInverseSquared())
+	{
+		Out.LightParameters.FalloffExponent = 0.0f;
+	}
 	return Out;
 }
 
@@ -413,7 +418,6 @@ class FDeferredLightPS : public FGlobalShader
 	class FSourceShapeDim		: SHADER_PERMUTATION_ENUM_CLASS("LIGHT_SOURCE_SHAPE", ELightSourceShape);
 	class FSourceTextureDim		: SHADER_PERMUTATION_BOOL("USE_SOURCE_TEXTURE");
 	class FIESProfileDim		: SHADER_PERMUTATION_BOOL("USE_IES_PROFILE");
-	class FInverseSquaredDim	: SHADER_PERMUTATION_BOOL("INVERSE_SQUARED_FALLOFF");
 	class FVisualizeCullingDim	: SHADER_PERMUTATION_BOOL("VISUALIZE_LIGHT_CULLING");
 	class FLightingChannelsDim	: SHADER_PERMUTATION_BOOL("USE_LIGHTING_CHANNELS");
 	class FTransmissionDim		: SHADER_PERMUTATION_BOOL("USE_TRANSMISSION");
@@ -427,7 +431,6 @@ class FDeferredLightPS : public FGlobalShader
 		FSourceShapeDim,
 		FSourceTextureDim,
 		FIESProfileDim,
-		FInverseSquaredDim,
 		FVisualizeCullingDim,
 		FLightingChannelsDim,
 		FTransmissionDim,
@@ -465,9 +468,7 @@ class FDeferredLightPS : public FGlobalShader
 	{
 		FPermutationDomain PermutationVector(Parameters.PermutationId);
 
-		if( PermutationVector.Get< FSourceShapeDim >() == ELightSourceShape::Directional && (
-			PermutationVector.Get< FIESProfileDim >() ||
-			PermutationVector.Get< FInverseSquaredDim >() ) )
+		if (PermutationVector.Get< FSourceShapeDim >() == ELightSourceShape::Directional && PermutationVector.Get< FIESProfileDim >())
 		{
 			return false;
 		}
@@ -477,19 +478,9 @@ class FDeferredLightPS : public FGlobalShader
 			return false;
 		}
 
-		if( PermutationVector.Get< FSourceShapeDim >() == ELightSourceShape::Rect )
-		{
-			if(	!PermutationVector.Get< FInverseSquaredDim >() )
-			{
-				return false;
-			}
-		}
-		else
+		if( PermutationVector.Get< FSourceShapeDim >() != ELightSourceShape::Rect && PermutationVector.Get< FSourceTextureDim >())
 		{
-			if( PermutationVector.Get< FSourceTextureDim >() )
-			{
-				return false;
-			}
+			return false;
 		}
 
 		if (PermutationVector.Get< FHairLighting >() && (
@@ -501,6 +492,11 @@ class FDeferredLightPS : public FGlobalShader
 
 		if (PermutationVector.Get<FDeferredLightPS::FAnistropicMaterials>())
 		{
+			if (Strata::IsStrataEnabled())
+			{
+				return false;
+			}
+
 			// Anisotropic materials do not currently support rect lights
 			if (PermutationVector.Get<FSourceShapeDim>() == ELightSourceShape::Rect || PermutationVector.Get<FSourceTextureDim>())
 			{
@@ -2024,6 +2020,7 @@ static void RenderLight(
 	const FSphere LightBounds = LightProxy->GetBoundingSphere();
 	const ELightComponentType LightType = (ELightComponentType)LightProxy->GetLightType();
 	const bool bIsRadial = LightType != LightType_Directional;
+	const bool bSupportAnisotropyPermutation = ShouldRenderAnisotropyPass(View) && !Strata::IsStrataEnabled(); // Strata managed anisotropy differently than legacy path. No need for special permutation.
 
 	// Debug Overlap shader
 	if (bRenderOverlap)
@@ -2099,8 +2096,7 @@ static void RenderLight(
 			PermutationVector.Set< FDeferredLightPS::FSourceShapeDim >(LightProxy->IsRectLight() ? ELightSourceShape::Rect : ELightSourceShape::Capsule);
 			PermutationVector.Set< FDeferredLightPS::FSourceTextureDim >(LightProxy->IsRectLight() && LightProxy->HasSourceTexture());
 			PermutationVector.Set< FDeferredLightPS::FIESProfileDim >(bUseIESTexture);
-			PermutationVector.Set< FDeferredLightPS::FInverseSquaredDim >(LightProxy->IsInverseSquared());
-			PermutationVector.Set< FDeferredLightPS::FAnistropicMaterials >(ShouldRenderAnisotropyPass(View) && !LightSceneInfo->Proxy->IsRectLight());
+			PermutationVector.Set< FDeferredLightPS::FAnistropicMaterials >(bSupportAnisotropyPermutation && !LightSceneInfo->Proxy->IsRectLight());
 			PermutationVector.Set < FDeferredLightPS::FAtmosphereTransmittance >(false);
 			PermutationVector.Set< FDeferredLightPS::FCloudTransmittance >(false);
 		}
@@ -2109,8 +2105,7 @@ static void RenderLight(
 			PermutationVector.Set< FDeferredLightPS::FSourceShapeDim >(ELightSourceShape::Directional);
 			PermutationVector.Set< FDeferredLightPS::FSourceTextureDim >(false);
 			PermutationVector.Set< FDeferredLightPS::FIESProfileDim >(false);
-			PermutationVector.Set< FDeferredLightPS::FInverseSquaredDim >(false);
-			PermutationVector.Set< FDeferredLightPS::FAnistropicMaterials >(ShouldRenderAnisotropyPass(View));
+			PermutationVector.Set< FDeferredLightPS::FAnistropicMaterials >(bSupportAnisotropyPermutation);
 			// Only directional lights are rendered in this path, so we only need to check if it is use to light the atmosphere
 			PermutationVector.Set< FDeferredLightPS::FAtmosphereTransmittance >(PassParameters->PS.CloudShadow.bAtmospherePerPixelTransmittance > 0);
 			PermutationVector.Set< FDeferredLightPS::FCloudTransmittance >(PassParameters->PS.CloudShadow.bCloudPerPixelTransmittance > 0);
@@ -2223,7 +2218,6 @@ void FDeferredShadingSceneRenderer::RenderLightForHair(
 		PermutationVector.Set< FDeferredLightPS::FSourceShapeDim >(ELightSourceShape::Directional);
 		PermutationVector.Set< FDeferredLightPS::FSourceTextureDim >(false);
 		PermutationVector.Set< FDeferredLightPS::FIESProfileDim >(false);
-		PermutationVector.Set< FDeferredLightPS::FInverseSquaredDim >(false);
 		PermutationVector.Set< FDeferredLightPS::FAtmosphereTransmittance >(PassParameters->PS.CloudShadow.bAtmospherePerPixelTransmittance > 0.f);
 		PermutationVector.Set< FDeferredLightPS::FCloudTransmittance >(PassParameters->PS.CloudShadow.bCloudPerPixelTransmittance > 0.f);
 	}
@@ -2233,7 +2227,6 @@ void FDeferredShadingSceneRenderer::RenderLightForHair(
 		PermutationVector.Set< FDeferredLightPS::FSourceShapeDim >(LightSceneInfo->Proxy->IsRectLight() ? ELightSourceShape::Rect : ELightSourceShape::Capsule);
 		PermutationVector.Set< FDeferredLightPS::FSourceTextureDim >(LightSceneInfo->Proxy->IsRectLight() && LightSceneInfo->Proxy->HasSourceTexture());
 		PermutationVector.Set< FDeferredLightPS::FIESProfileDim >(bUseIESTexture);
-		PermutationVector.Set< FDeferredLightPS::FInverseSquaredDim >(LightSceneInfo->Proxy->IsInverseSquared());
 		PermutationVector.Set< FDeferredLightPS::FAtmosphereTransmittance >(false);
 		PermutationVector.Set< FDeferredLightPS::FCloudTransmittance >(false);
 	}
@@ -2418,13 +2411,7 @@ void FDeferredShadingSceneRenderer::RenderSimpleLightsStandardDeferred(
 		PermutationVector.Set< FDeferredLightPS::FAtmosphereTransmittance >(false);
 		PermutationVector.Set< FDeferredLightPS::FCloudTransmittance >(false);
 		PermutationVector.Set< FDeferredLightPS::FStrataFastPath>(false);
-
-		FDeferredLightPS::FPermutationDomain PermutationVector_bInverseSquaredFalloff0;
-		FDeferredLightPS::FPermutationDomain PermutationVector_bInverseSquaredFalloff1;
-		PermutationVector_bInverseSquaredFalloff0.Set<FDeferredLightPS::FInverseSquaredDim>(false);
-		PermutationVector_bInverseSquaredFalloff1.Set<FDeferredLightPS::FInverseSquaredDim>(true);
-		TShaderMapRef<FDeferredLightPS> PixelShader0(View.ShaderMap, PermutationVector);
-		TShaderMapRef<FDeferredLightPS> PixelShader1(View.ShaderMap, PermutationVector);
+		TShaderMapRef<FDeferredLightPS> PixelShader(View.ShaderMap, PermutationVector);
 
 		FDeferredLightVS::FPermutationDomain PermutationVectorVS;
 		PermutationVectorVS.Set<FDeferredLightVS::FRadialLight>(true);
@@ -2434,7 +2421,7 @@ void FDeferredShadingSceneRenderer::RenderSimpleLightsStandardDeferred(
 			RDG_EVENT_NAME("StandardDeferredSimpleLights"),
 			PassParameters,
 			ERDGPassFlags::Raster,
-			[this, &View, &SimpleLights, ViewIndex, NumViews, PassParameters, PixelShader0, PixelShader1, VertexShader](FRHICommandList& RHICmdList)
+			[this, &View, &SimpleLights, ViewIndex, NumViews, PassParameters, PixelShader, VertexShader](FRHICommandList& RHICmdList)
 		{
 			FGraphicsPipelineStateInitializer GraphicsPSOInit;
 			RHICmdList.ApplyCachedRenderTargets(GraphicsPSOInit);
@@ -2460,27 +2447,16 @@ void FDeferredShadingSceneRenderer::RenderSimpleLightsStandardDeferred(
 								|| !View.IsPerspectiveProjection();
 
 				SetBoundingGeometryRasterizerAndDepthState(GraphicsPSOInit, View, bCameraInsideLightGeometry);
-
-				const bool bInverseSquaredFalloff = SimpleLight.Exponent == 0;
-
 				GraphicsPSOInit.BoundShaderState.VertexDeclarationRHI = GetVertexDeclarationFVector4();
 				GraphicsPSOInit.BoundShaderState.VertexShaderRHI = VertexShader.GetVertexShader();
-				GraphicsPSOInit.BoundShaderState.PixelShaderRHI = bInverseSquaredFalloff ? PixelShader1.GetPixelShader() : PixelShader0.GetPixelShader();
+				GraphicsPSOInit.BoundShaderState.PixelShaderRHI = PixelShader.GetPixelShader();
 				SetGraphicsPipelineState(RHICmdList, GraphicsPSOInit, 0);
 
 
 				// Update the light parameters with a custom uniform buffer
 				FDeferredLightUniformStruct DeferredLightUniformsValue = GetSimpleDeferredLightParameters(SimpleLight, SimpleLightPerViewData);
-				if (bInverseSquaredFalloff)
-				{
-					SetShaderParameters(RHICmdList, PixelShader1, PixelShader1.GetPixelShader(), PassParameters->PS);
-					SetUniformBufferParameterImmediate(RHICmdList, RHICmdList.GetBoundPixelShader(), PixelShader0->GetUniformBufferParameter<FDeferredLightUniformStruct>(), DeferredLightUniformsValue);
-				}
-				else
-				{
-					SetShaderParameters(RHICmdList, PixelShader0, PixelShader0.GetPixelShader(), PassParameters->PS);
-					SetUniformBufferParameterImmediate(RHICmdList, RHICmdList.GetBoundPixelShader(), PixelShader0->GetUniformBufferParameter<FDeferredLightUniformStruct>(), DeferredLightUniformsValue);
-				}
+				SetShaderParameters(RHICmdList, PixelShader, PixelShader.GetPixelShader(), PassParameters->PS);
+				SetUniformBufferParameterImmediate(RHICmdList, RHICmdList.GetBoundPixelShader(), PixelShader->GetUniformBufferParameter<FDeferredLightUniformStruct>(), DeferredLightUniformsValue);
 
 				// Update vertex shader parameters with custom parameters/uniform buffer
 				FDeferredLightVS::FParameters ParametersVS = FDeferredLightVS::GetParameters(View, LightBounds);
diff --git a/Engine/Source/Runtime/Renderer/Private/MobileDeferredShadingPass.cpp b/Engine/Source/Runtime/Renderer/Private/MobileDeferredShadingPass.cpp
--- a/Engine/Source/Runtime/Renderer/Private/MobileDeferredShadingPass.cpp
+++ b/Engine/Source/Runtime/Renderer/Private/MobileDeferredShadingPass.cpp
@@ -131,9 +131,8 @@ public:
 	SHADER_USE_PARAMETER_STRUCT_WITH_LEGACY_BASE(FMobileRadialLightFunctionPS, FMaterialShader)
 
 	class FSpotLightDim			: SHADER_PERMUTATION_BOOL("IS_SPOT_LIGHT");
-	class FInverseSquaredDim	: SHADER_PERMUTATION_BOOL("INVERSE_SQUARED_FALLOFF");
 	class FIESProfileDim		: SHADER_PERMUTATION_BOOL("USE_IES_PROFILE");
-	using FPermutationDomain = TShaderPermutationDomain<FSpotLightDim, FInverseSquaredDim, FIESProfileDim>;
+	using FPermutationDomain = TShaderPermutationDomain<FSpotLightDim, FIESProfileDim>;
 
 	BEGIN_SHADER_PARAMETER_STRUCT(FParameters, )
 		SHADER_PARAMETER(FMatrix44f, WorldToLight)
@@ -445,7 +444,6 @@ static void RenderLocalLight(
 	}
 	FMobileRadialLightFunctionPS::FPermutationDomain PermutationVector;
 	PermutationVector.Set<FMobileRadialLightFunctionPS::FSpotLightDim>(bIsSpotLight);
-	PermutationVector.Set<FMobileRadialLightFunctionPS::FInverseSquaredDim>(LightSceneInfo.Proxy->IsInverseSquared());
 	PermutationVector.Set<FMobileRadialLightFunctionPS::FIESProfileDim>(bUseIESTexture);
 	FCachedLightMaterial LightMaterial;
 	TShaderRef<FMobileRadialLightFunctionPS> PixelShader;
@@ -518,23 +516,19 @@ static void RenderSimpleLights(
 	FDeferredLightVS::FPermutationDomain PermutationVectorVS;
 	PermutationVectorVS.Set<FDeferredLightVS::FRadialLight>(true);
 	TShaderMapRef<FDeferredLightVS> VertexShader(View.ShaderMap, PermutationVectorVS);
-	TShaderRef<FMobileRadialLightFunctionPS> PixelShaders[2];
+	TShaderRef<FMobileRadialLightFunctionPS> PixelShaders;
 	{
 		const FMaterialShaderMap* MaterialShaderMap = DefaultMaterial.Material->GetRenderingThreadShaderMap();
 		FMobileRadialLightFunctionPS::FPermutationDomain PermutationVector;
 		PermutationVector.Set<FMobileRadialLightFunctionPS::FSpotLightDim>(false);
 		PermutationVector.Set<FMobileRadialLightFunctionPS::FIESProfileDim>(false);
-		PermutationVector.Set<FMobileRadialLightFunctionPS::FInverseSquaredDim>(false);
-		PixelShaders[0] = MaterialShaderMap->GetShader<FMobileRadialLightFunctionPS>(PermutationVector);
-		PermutationVector.Set<FMobileRadialLightFunctionPS::FInverseSquaredDim>(true);
-		PixelShaders[1] = MaterialShaderMap->GetShader<FMobileRadialLightFunctionPS>(PermutationVector);
+		PixelShaders = MaterialShaderMap->GetShader<FMobileRadialLightFunctionPS>(PermutationVector);
 	}
 
 	// Setup PSOs we going to use for light rendering 
-	FGraphicsPipelineStateInitializer GraphicsPSOLight[2];
+	FGraphicsPipelineStateInitializer GraphicsPSOLight;
 	{
-		SetupSimpleLightPSO(RHICmdList, View, VertexShader, PixelShaders[0], GraphicsPSOLight[0]);
-		SetupSimpleLightPSO(RHICmdList, View, VertexShader, PixelShaders[1], GraphicsPSOLight[1]);
+		SetupSimpleLightPSO(RHICmdList, View, VertexShader, PixelShaders, GraphicsPSOLight);
 	}
 	// Setup stencil mask PSO
 	FGraphicsPipelineStateInitializer GraphicsPSOLightMask;
@@ -585,15 +579,9 @@ static void RenderSimpleLights(
 		// Shade only MSM_DefaultLit pixels
 		uint8 StencilRef = GET_STENCIL_MOBILE_SM_MASK(MSM_DefaultLit);
 
-		if (SimpleLight.Exponent == 0)
 		{
-			SetGraphicsPipelineState(RHICmdList, GraphicsPSOLight[1], StencilRef);
-			FMobileRadialLightFunctionPS::SetParameters(RHICmdList, PixelShaders[1], View, DefaultMaterial.MaterialProxy, *DefaultMaterial.Material, PassParameters);
-		}
-		else
-		{
-			SetGraphicsPipelineState(RHICmdList, GraphicsPSOLight[0], StencilRef);
-			FMobileRadialLightFunctionPS::SetParameters(RHICmdList, PixelShaders[0], View, DefaultMaterial.MaterialProxy, *DefaultMaterial.Material, PassParameters);
+			SetGraphicsPipelineState(RHICmdList, GraphicsPSOLight, StencilRef);
+			FMobileRadialLightFunctionPS::SetParameters(RHICmdList, PixelShaders, View, DefaultMaterial.MaterialProxy, *DefaultMaterial.Material, PassParameters);
 		}
 
 		// Apply the point or spot light with some approximately bounding geometry,
diff --git a/Engine/Source/Runtime/Renderer/Private/TranslucentLighting.cpp b/Engine/Source/Runtime/Renderer/Private/TranslucentLighting.cpp
--- a/Engine/Source/Runtime/Renderer/Private/TranslucentLighting.cpp
+++ b/Engine/Source/Runtime/Renderer/Private/TranslucentLighting.cpp
@@ -707,14 +707,12 @@ class FTranslucentLightingInjectPS : public FMaterialShader
 	class FRadialAttenuation	: SHADER_PERMUTATION_BOOL("RADIAL_ATTENUATION");
 	class FDynamicallyShadowed	: SHADER_PERMUTATION_BOOL("DYNAMICALLY_SHADOWED");
 	class FLightFunction		: SHADER_PERMUTATION_BOOL("APPLY_LIGHT_FUNCTION");
-	class FInverseSquared		: SHADER_PERMUTATION_BOOL("INVERSE_SQUARED_FALLOFF");
 	class FVirtualShadowMap		: SHADER_PERMUTATION_BOOL("VIRTUAL_SHADOW_MAP");
 
 	using FPermutationDomain = TShaderPermutationDomain<
 		FRadialAttenuation,
 		FDynamicallyShadowed,
 		FLightFunction,
-		FInverseSquared,
 		FVirtualShadowMap >;
 
 public:
@@ -733,13 +731,6 @@ public:
 	  */
 	static bool ShouldCompilePermutation(const FMaterialShaderPermutationParameters& Parameters)
 	{
-		FPermutationDomain PermutationVector(Parameters.PermutationId);
-
-		if( !PermutationVector.Get<FRadialAttenuation>() && !PermutationVector.Get<FInverseSquared>() )
-		{
-			return false;
-		}
-
 		return (Parameters.MaterialParameters.MaterialDomain == MD_LightFunction || Parameters.MaterialParameters.bIsSpecialEngineMaterial) &&
 			(IsFeatureLevelSupported(Parameters.Platform, ERHIFeatureLevel::SM5) &&
 			(RHISupportsGeometryShaders(Parameters.Platform) || RHISupportsVertexShaderLayer(Parameters.Platform)));
@@ -1261,7 +1252,7 @@ static void InjectTranslucentLightArray(
 					RDG_EVENT_NAME("InjectTranslucentLightArray"),
 					PassParameters,
 					ERDGPassFlags::Raster,
-					[PassParameters, VertexShader, GeometryShader, &View, &Renderer, &InjectionData, LightSceneInfo, bInverseSquared, bDirectionalLight, bUseVSM, VolumeBounds, VolumeCascadeIndex](FRHICommandList& RHICmdList)
+					[PassParameters, VertexShader, GeometryShader, &View, &Renderer, &InjectionData, LightSceneInfo, bDirectionalLight, bUseVSM, VolumeBounds, VolumeCascadeIndex](FRHICommandList& RHICmdList)
 				{
 					FGraphicsPipelineStateInitializer GraphicsPSOInit;
 					RHICmdList.ApplyCachedRenderTargets(GraphicsPSOInit);
@@ -1294,7 +1285,6 @@ static void InjectTranslucentLightArray(
 					PermutationVector.Set< FTranslucentLightingInjectPS::FRadialAttenuation >( !bDirectionalLight );
 					PermutationVector.Set< FTranslucentLightingInjectPS::FDynamicallyShadowed >( InjectionData.ProjectedShadowInfo != nullptr );
 					PermutationVector.Set< FTranslucentLightingInjectPS::FLightFunction >( InjectionData.bApplyLightFunction );
-					PermutationVector.Set< FTranslucentLightingInjectPS::FInverseSquared >( bInverseSquared );
 					PermutationVector.Set< FTranslucentLightingInjectPS::FVirtualShadowMap >( bUseVSM );
 
 					auto PixelShader = MaterialShaderMap->GetShader< FTranslucentLightingInjectPS >( PermutationVector );
diff --git a/Engine/Source/Runtime/Renderer/Private/VolumetricFog.cpp b/Engine/Source/Runtime/Renderer/Private/VolumetricFog.cpp
--- a/Engine/Source/Runtime/Renderer/Private/VolumetricFog.cpp
+++ b/Engine/Source/Runtime/Renderer/Private/VolumetricFog.cpp
@@ -293,7 +293,6 @@ class TInjectShadowedLocalLightPS : public FGlobalShader
 	END_SHADER_PARAMETER_STRUCT()
 
 	class FDynamicallyShadowed	: SHADER_PERMUTATION_BOOL("DYNAMICALLY_SHADOWED");
-	class FInverseSquared		: SHADER_PERMUTATION_BOOL("INVERSE_SQUARED_FALLOFF");
 	class FTemporalReprojection : SHADER_PERMUTATION_BOOL("USE_TEMPORAL_REPROJECTION");
 	class FLightFunction		: SHADER_PERMUTATION_BOOL("USE_LIGHT_FUNCTION");
 	class FEnableShadows		: SHADER_PERMUTATION_BOOL("ENABLE_SHADOW_COMPUTATION");
@@ -301,7 +300,6 @@ class TInjectShadowedLocalLightPS : public FGlobalShader
 
 	using FPermutationDomain = TShaderPermutationDomain<
 		FDynamicallyShadowed,
-		FInverseSquared,
 		FTemporalReprojection,
 		FLightFunction,
 		FEnableShadows,
@@ -704,7 +702,6 @@ void FDeferredShadingSceneRenderer::RenderLocalLightsForVolumetricFog(
 					const FVisibleLightInfo& VisibleLightInfo = VisibleLightInfos[LightSceneInfo->Id];
 					const FProjectedShadowInfo* ProjectedShadowInfo = GetShadowForInjectionIntoVolumetricFog(VisibleLightInfo);
 
-					const bool bInverseSquared = LightSceneInfo->Proxy->IsInverseSquared();
 					const bool bDynamicallyShadowed = ProjectedShadowInfo != NULL;
 					
 					const FSphere LightBounds = LightSceneInfo->Proxy->GetBoundingSphere();
@@ -739,7 +736,6 @@ void FDeferredShadingSceneRenderer::RenderLocalLightsForVolumetricFog(
 
 						TInjectShadowedLocalLightPS::FPermutationDomain PermutationVector;
 						PermutationVector.Set< TInjectShadowedLocalLightPS::FDynamicallyShadowed >(bDynamicallyShadowed);
-						PermutationVector.Set< TInjectShadowedLocalLightPS::FInverseSquared >(bInverseSquared);
 						PermutationVector.Set< TInjectShadowedLocalLightPS::FTemporalReprojection >(bUseTemporalReprojection);
 						PermutationVector.Set< TInjectShadowedLocalLightPS::FLightFunction >(bUsesLightFunction);
 						PermutationVector.Set< TInjectShadowedLocalLightPS::FEnableShadows >(bIsShadowed);
