Refactor remaining DF shaders to not use legacy shader parameter binding - Part 3
- Created FGlobalDistanceFieldParameters2 using SHADER_PARAMETER_STRUCT.
    - Can't remove existing struct since it's used by Niagara.
    - Added SetupGlobalDistanceFieldParameters helper function to perform the necessary logic.
- Converted FScreenGridParameters to use SHADER_PARAMETER_STRUCT.
    - Added SetupScreenGridParameters helper function to perform the necessary logic.
- Bind AO SampleDirections using SHADER_PARAMETER_STRUCT instead of GLOBAL_SHADER_PARAMETER_STRUCT + SetUniformBufferParameterImmediate.
- Removed all legacy shader parameters from FVisualizeMeshDistanceFieldCS, FConeTraceScreenGridObjectOcclusionCS, FConeTraceScreenGridGlobalOcclusionCS, FCombineConeVisibilityCS.
- Removed some unused variables

diff --git a/Engine/Shaders/Private/DistanceFieldAOShared.ush b/Engine/Shaders/Private/DistanceFieldAOShared.ush
--- a/Engine/Shaders/Private/DistanceFieldAOShared.ush
+++ b/Engine/Shaders/Private/DistanceFieldAOShared.ush
@@ -136,6 +136,7 @@ bool SphereIntersectConeWithDepthRanges(float4 SphereCenterAndRadius, float3 Con
 }
 
 Buffer<float> RecordConeVisibility;
+float4 SampleDirections[NUM_CONE_DIRECTIONS];
 float BentNormalNormalizeFactor;
 
 void FindBestAxisVectors2(float3 InZAxis, out float3 OutXAxis, out float3 OutYAxis )
@@ -155,7 +156,7 @@ float3 ComputeBentNormal(float3 RecordWorldNormal, uint RelativeRecordIndex)
 
 	for (uint ConeIndex = 0; ConeIndex < NUM_CONE_DIRECTIONS; ConeIndex++)
 	{
-		float3 ConeDirection = AOSamples2.SampleDirections[ConeIndex].xyz;
+		float3 ConeDirection = SampleDirections[ConeIndex].xyz;
 		float3 RotatedConeDirection = ConeDirection.x * TangentX + ConeDirection.y * TangentY + ConeDirection.z * RecordWorldNormal;
 
 		float ConeVisibility = RecordConeVisibility[RelativeRecordIndex * NUM_CONE_DIRECTIONS + ConeIndex];
diff --git a/Engine/Shaders/Private/DistanceFieldScreenGridLighting.usf b/Engine/Shaders/Private/DistanceFieldScreenGridLighting.usf
--- a/Engine/Shaders/Private/DistanceFieldScreenGridLighting.usf
+++ b/Engine/Shaders/Private/DistanceFieldScreenGridLighting.usf
@@ -106,7 +106,7 @@ void HemisphereConeTraceAgainstTileCulledObjects(uint ObjectIndex, uint OutputBa
 			LOOP
 			for (uint ConeIndex = 0; ConeIndex < NUM_CONE_DIRECTIONS; ConeIndex++)
 			{
-				float3 ConeDirection = AOSamples2.SampleDirections[ConeIndex].xyz;
+				float3 ConeDirection = SampleDirections[ConeIndex].xyz;
 				float3 RotatedConeDirection = ConeDirection.x * TangentX + ConeDirection.y * TangentY + ConeDirection.z * WorldNormal;
 				float3 ScaledLocalConeDirection = mul(RotatedConeDirection, (float3x3) DFObjectData.WorldToVolume).xyz;
 		
@@ -265,7 +265,7 @@ void HemisphereConeTraceAgainstGlobalDistanceFieldClipmap(
 	LOOP
 	for (uint ConeIndex = 0; ConeIndex < NUM_CONE_DIRECTIONS; ConeIndex++)
 	{
-		float3 ConeDirection = AOSamples2.SampleDirections[ConeIndex].xyz;
+		float3 ConeDirection = SampleDirections[ConeIndex].xyz;
 		float3 RotatedConeDirection = ConeDirection.x * TangentX + ConeDirection.y * TangentY + ConeDirection.z * WorldNormal;
 
 		float MinVisibility = 1;
@@ -401,7 +401,7 @@ void CombineConeVisibilityCS(
 	{
 		float ConeVisibility = asfloat(ScreenGridConeVisibility[ConeIndex * ScreenGridConeVisibilitySize.x * ScreenGridConeVisibilitySize.y + InputBaseIndex]);
 		
-		float3 ConeDirection = AOSamples2.SampleDirections[ConeIndex].xyz;
+		float3 ConeDirection = SampleDirections[ConeIndex].xyz;
 		float3 RotatedConeDirection = ConeDirection.x * TangentX + ConeDirection.y * TangentY + ConeDirection.z * WorldNormal;
 		UnoccludedDirection += ConeVisibility * RotatedConeDirection;
 	}
diff --git a/Engine/Source/Runtime/Renderer/Private/DistanceFieldAmbientOcclusion.cpp b/Engine/Source/Runtime/Renderer/Private/DistanceFieldAmbientOcclusion.cpp
--- a/Engine/Source/Runtime/Renderer/Private/DistanceFieldAmbientOcclusion.cpp
+++ b/Engine/Source/Runtime/Renderer/Private/DistanceFieldAmbientOcclusion.cpp
@@ -21,8 +21,6 @@
 #include "Lumen/Lumen.h"
 #include "Strata/Strata.h"
 
-IMPLEMENT_TYPE_LAYOUT(FScreenGridParameters);
-
 int32 GDistanceFieldAO = 1;
 FAutoConsoleVariableRef CVarDistanceFieldAO(
 	TEXT("r.DistanceFieldAO"),
@@ -130,8 +128,6 @@ int32 GDistanceFieldAOTileSizeY = 16;
 
 DEFINE_LOG_CATEGORY(LogDistanceField);
 
-IMPLEMENT_GLOBAL_SHADER_PARAMETER_STRUCT(FAOSampleData2,"AOSamples2");
-
 FDistanceFieldAOParameters::FDistanceFieldAOParameters(float InOcclusionMaxDistance, float InContrast)
 {
 	Contrast = FMath::Clamp(InContrast, .01f, 2.0f);
diff --git a/Engine/Source/Runtime/Renderer/Private/DistanceFieldAmbientOcclusion.h b/Engine/Source/Runtime/Renderer/Private/DistanceFieldAmbientOcclusion.h
--- a/Engine/Source/Runtime/Renderer/Private/DistanceFieldAmbientOcclusion.h
+++ b/Engine/Source/Runtime/Renderer/Private/DistanceFieldAmbientOcclusion.h
@@ -17,10 +17,6 @@
 #include "PostProcess/SceneRenderTargets.h"
 #include "ScenePrivate.h"
 
-const static int32 GAOMaxSupportedLevel = 6;
-/** Number of cone traced directions. */
-const int32 NumConeSampleDirections = 9;
-
 /** Base downsample factor that all distance field AO operations are done at. */
 const int32 GAODownsampleFactor = 2;
 
@@ -57,8 +53,8 @@ BEGIN_SHADER_PARAMETER_STRUCT(FTileIntersectionParameters, )
 	SHADER_PARAMETER(FIntPoint, TileListGroupSize)
 END_SHADER_PARAMETER_STRUCT()
 
-static int32 CulledTileDataStride = 2;
-static int32 ConeTraceObjectsThreadGroupSize = 64;
+static const int32 CulledTileDataStride = 2;
+static const int32 ConeTraceObjectsThreadGroupSize = 64;
 
 inline void TileIntersectionModifyCompilationEnvironment(EShaderPlatform Platform, FShaderCompilerEnvironment& OutEnvironment)
 {
@@ -78,10 +74,6 @@ END_SHADER_PARAMETER_STRUCT()
 
 extern void GetSpacedVectors(uint32 FrameNumber, TArray<FVector, TInlineAllocator<9> >& OutVectors);
 
-BEGIN_GLOBAL_SHADER_PARAMETER_STRUCT(FAOSampleData2,)
-	SHADER_PARAMETER_ARRAY(FVector4f,SampleDirections,[NumConeSampleDirections])
-END_GLOBAL_SHADER_PARAMETER_STRUCT()
-
 inline float GetMaxAOViewDistance()
 {
 	extern float GAOMaxViewDistance;
@@ -156,54 +148,6 @@ protected:
 	int32 MaxSize;
 };
 
-class FScreenGridParameters
-{
-	DECLARE_TYPE_LAYOUT(FScreenGridParameters, NonVirtual);
-public:
-	void Bind(const FShaderParameterMap& ParameterMap)
-	{
-		BaseLevelTexelSize.Bind(ParameterMap, TEXT("BaseLevelTexelSize"));
-		JitterOffset.Bind(ParameterMap, TEXT("JitterOffset"));
-		DistanceFieldNormalTexture.Bind(ParameterMap, TEXT("DistanceFieldNormalTexture"));
-		DistanceFieldNormalSampler.Bind(ParameterMap, TEXT("DistanceFieldNormalSampler"));
-	}
-
-	template<typename TParamRef>
-	void Set(FRHICommandList& RHICmdList, const TParamRef& ShaderRHI, const FViewInfo& View, FRHITexture* DistanceFieldNormal)
-	{
-		const FIntPoint DownsampledBufferSize = GetBufferSizeForAO();
-		const FVector2f BaseLevelTexelSizeValue(1.0f / DownsampledBufferSize.X, 1.0f / DownsampledBufferSize.Y);
-		SetShaderValue(RHICmdList, ShaderRHI, BaseLevelTexelSize, BaseLevelTexelSizeValue);
-
-		extern FVector2f GetJitterOffset(int32 SampleIndex);
-		SetShaderValue(RHICmdList, ShaderRHI, JitterOffset, GetJitterOffset(View.GetDistanceFieldTemporalSampleIndex()));
-
-
-		SetTextureParameter(
-			RHICmdList,
-			ShaderRHI,
-			DistanceFieldNormalTexture,
-			DistanceFieldNormalSampler,
-			TStaticSamplerState<SF_Point,AM_Wrap,AM_Wrap,AM_Wrap>::GetRHI(),
-			DistanceFieldNormal
-			);
-	}
-
-	friend FArchive& operator<<(FArchive& Ar,FScreenGridParameters& P)
-	{
-		Ar << P.BaseLevelTexelSize << P.JitterOffset << P.DistanceFieldNormalTexture << P.DistanceFieldNormalSampler;
-		return Ar;
-	}
-
-private:
-	
-		LAYOUT_FIELD(FShaderParameter, BaseLevelTexelSize)
-		LAYOUT_FIELD(FShaderParameter, JitterOffset)
-		LAYOUT_FIELD(FShaderResourceParameter, DistanceFieldNormalTexture)
-		LAYOUT_FIELD(FShaderResourceParameter, DistanceFieldNormalSampler)
-	
-};
-
 extern void TrackGPUProgress(FRHICommandListImmediate& RHICmdList, uint32 DebugId);
 
 extern bool ShouldRenderDeferredDynamicSkyLight(const FScene* Scene, const FSceneViewFamily& ViewFamily);
diff --git a/Engine/Source/Runtime/Renderer/Private/DistanceFieldScreenGridLighting.cpp b/Engine/Source/Runtime/Renderer/Private/DistanceFieldScreenGridLighting.cpp
--- a/Engine/Source/Runtime/Renderer/Private/DistanceFieldScreenGridLighting.cpp
+++ b/Engine/Source/Runtime/Renderer/Private/DistanceFieldScreenGridLighting.cpp
@@ -28,6 +28,9 @@
 #include "ClearQuad.h"
 #include "VisualizeTexture.h"
 
+/** Number of cone traced directions. */
+static const int32 NumConeSampleDirections = 9;
+
 static int32 GAOUseJitter = 1;
 static FAutoConsoleVariableRef CVarAOUseJitter(
 	TEXT("r.AOUseJitter"),
@@ -59,6 +62,7 @@ FVector2f JitterOffsets[4] =
 	FVector2f(0, .5f)
 };
 
+extern float GAOConeHalfAngle;
 extern int32 GAOUseHistory;
 
 FVector2f GetJitterOffset(int32 SampleIndex)
@@ -71,21 +75,77 @@ FVector2f GetJitterOffset(int32 SampleIndex)
 	return FVector2f(0, 0);
 }
 
-class FConeTraceScreenGridObjectOcclusionCS : public FGlobalShader
+BEGIN_SHADER_PARAMETER_STRUCT(FAOSampleParameters,)
+	SHADER_PARAMETER_ARRAY(FVector4f,SampleDirections,[NumConeSampleDirections])
+	SHADER_PARAMETER(float, BentNormalNormalizeFactor)
+END_SHADER_PARAMETER_STRUCT()
+
+static FAOSampleParameters SetupAOSampleParameters(uint32 FrameNumber)
 {
-	DECLARE_GLOBAL_SHADER(FConeTraceScreenGridObjectOcclusionCS);
+	TArray<FVector, TInlineAllocator<9> > SampleDirections;
+	GetSpacedVectors(FrameNumber, SampleDirections);
+
+	FVector UnoccludedVector(0);
+	for (int32 SampleIndex = 0; SampleIndex < NumConeSampleDirections; SampleIndex++)
+	{
+		UnoccludedVector += SampleDirections[SampleIndex];
+	}
+
+	// LWC_TODO: Was float BentNormalNormalizeFactorValue = 1.0f / (UnoccludedVector / NumConeSampleDirections).Size();
+	// This causes "error C4723: potential divide by 0" in msvc, implying the compiler has managed to evaluate (UnoccludedVector / NumConeSampleDirections).Size() as 0 at compile time.
+	// Ensuring Size() is called in a seperate line stops the warning, but this needs investigating. Clang seems happy with it. Possible compiler bug?
+	const float ConeSampleAverageSize = (UnoccludedVector / NumConeSampleDirections).Size();
 
+	FAOSampleParameters ShaderParameters;
+	ShaderParameters.BentNormalNormalizeFactor = ConeSampleAverageSize ? 1.0f / ConeSampleAverageSize : 0.f;
+
+	for (int32 SampleIndex = 0; SampleIndex < NumConeSampleDirections; SampleIndex++)
+	{
+		ShaderParameters.SampleDirections[SampleIndex] = FVector4f(SampleDirections[SampleIndex]);
+	}
+
+	return ShaderParameters;
+}
+
+BEGIN_SHADER_PARAMETER_STRUCT(FScreenGridParameters, )
+	SHADER_PARAMETER_RDG_TEXTURE(Texture2D, DistanceFieldNormalTexture)
+	SHADER_PARAMETER_SAMPLER(SamplerState, DistanceFieldNormalSampler)
+	SHADER_PARAMETER(FVector2f, BaseLevelTexelSize)
+	SHADER_PARAMETER(FVector2f, JitterOffset)
+END_SHADER_PARAMETER_STRUCT()
+
+static FScreenGridParameters SetupScreenGridParameters(const FViewInfo& View, FRDGTextureRef DistanceFieldNormal)
+{
+	const FIntPoint DownsampledBufferSize = GetBufferSizeForAO();
+
+	FScreenGridParameters ShaderParameters;
+	ShaderParameters.DistanceFieldNormalTexture = DistanceFieldNormal;
+	ShaderParameters.DistanceFieldNormalSampler = TStaticSamplerState<SF_Point, AM_Wrap, AM_Wrap, AM_Wrap>::GetRHI();
+	ShaderParameters.BaseLevelTexelSize = FVector2f(1.0f / DownsampledBufferSize.X, 1.0f / DownsampledBufferSize.Y);
+	ShaderParameters.JitterOffset = GetJitterOffset(View.GetDistanceFieldTemporalSampleIndex());
+
+	return ShaderParameters;
+}
+
+class FConeTraceScreenGridObjectOcclusionCS : public FGlobalShader
+{
 public:
+	DECLARE_GLOBAL_SHADER(FConeTraceScreenGridObjectOcclusionCS);
+	SHADER_USE_PARAMETER_STRUCT(FConeTraceScreenGridObjectOcclusionCS, FGlobalShader);
 
 	BEGIN_SHADER_PARAMETER_STRUCT(FParameters, )
+		SHADER_PARAMETER_STRUCT_REF(FViewUniformShaderParameters, View)
 		SHADER_PARAMETER_STRUCT_INCLUDE(FDistanceFieldCulledObjectBufferParameters, DistanceFieldCulledObjectBuffers)
 		SHADER_PARAMETER_STRUCT_INCLUDE(FDistanceFieldAtlasParameters, DistanceFieldAtlas)
 		SHADER_PARAMETER_STRUCT_INCLUDE(FTileIntersectionParameters, TileIntersectionParameters)
 		SHADER_PARAMETER_STRUCT_INCLUDE(FAOScreenGridParameters, AOScreenGridParameters)
 		SHADER_PARAMETER_STRUCT_INCLUDE(FAOParameters, AOParameters)
-		RDG_BUFFER_ACCESS(ObjectTilesIndirectArguments, ERHIAccess::IndirectArgs)
+		SHADER_PARAMETER_STRUCT_INCLUDE(FScreenGridParameters, ScreenGridParameters)
+		SHADER_PARAMETER_STRUCT_INCLUDE(FGlobalDistanceFieldParameters2, GlobalDistanceFieldParameters)
+		SHADER_PARAMETER_STRUCT_INCLUDE(FAOSampleParameters, AOSampleParameters)
 		SHADER_PARAMETER_RDG_UNIFORM_BUFFER(FSceneTextureUniformParameters, SceneTextures)
-		RDG_TEXTURE_ACCESS(DistanceFieldNormal, ERHIAccess::SRVCompute)
+		SHADER_PARAMETER(float, TanConeHalfAngle)
+		RDG_BUFFER_ACCESS(ObjectTilesIndirectArguments, ERHIAccess::IndirectArgs)
 	END_SHADER_PARAMETER_STRUCT()
 
 	class FUseGlobalDistanceField : SHADER_PERMUTATION_BOOL("USE_GLOBAL_DISTANCE_FIELD");
@@ -105,73 +165,6 @@ public:
 		// To reduce shader compile time of compute shaders with shared memory, doesn't have an impact on generated code with current compiler (June 2010 DX SDK)
 		OutEnvironment.CompilerFlags.Add(CFLAG_StandardOptimization);
 	}
-
-	FConeTraceScreenGridObjectOcclusionCS(const ShaderMetaType::CompiledShaderInitializerType& Initializer)
-		: FGlobalShader(Initializer)
-	{
-		BindForLegacyShaderParameters<FParameters>(this, Initializer.PermutationId, Initializer.ParameterMap, false);
-		ScreenGridParameters.Bind(Initializer.ParameterMap);
-		GlobalDistanceFieldParameters.Bind(Initializer.ParameterMap);
-		TanConeHalfAngle.Bind(Initializer.ParameterMap, TEXT("TanConeHalfAngle"));
-		BentNormalNormalizeFactor.Bind(Initializer.ParameterMap, TEXT("BentNormalNormalizeFactor"));
-	}
-
-	FConeTraceScreenGridObjectOcclusionCS()
-	{
-	}
-
-	void SetParameters(
-		FRHICommandList& RHICmdList, 
-		const FViewInfo& View, 
-		FRHITexture* DistanceFieldNormal,
-		bool bUseGlobalDistanceField,
-		const FGlobalDistanceFieldInfo& GlobalDistanceFieldInfo)
-	{
-		FRHIComputeShader* ShaderRHI = RHICmdList.GetBoundComputeShader();
-		FGlobalShader::SetParameters<FViewUniformShaderParameters>(RHICmdList, ShaderRHI, View.ViewUniformBuffer);
-
-		ScreenGridParameters.Set(RHICmdList, ShaderRHI, View, DistanceFieldNormal);
-
-		if (bUseGlobalDistanceField)
-		{
-			GlobalDistanceFieldParameters.Set(RHICmdList, ShaderRHI, GlobalDistanceFieldInfo.ParameterData);
-		}
-
-		FAOSampleData2 AOSampleData;
-
-		TArray<FVector, TInlineAllocator<9> > SampleDirections;
-		GetSpacedVectors(View.Family->FrameNumber, SampleDirections);
-
-		for (int32 SampleIndex = 0; SampleIndex < NumConeSampleDirections; SampleIndex++)
-		{
-			AOSampleData.SampleDirections[SampleIndex] = FVector4f(SampleDirections[SampleIndex]);
-		}
-
-		SetUniformBufferParameterImmediate(RHICmdList, ShaderRHI, GetUniformBufferParameter<FAOSampleData2>(), AOSampleData);
-
-		extern float GAOConeHalfAngle;
-		SetShaderValue(RHICmdList, ShaderRHI, TanConeHalfAngle, FMath::Tan(GAOConeHalfAngle));
-
-		FVector UnoccludedVector(0);
-
-		for (int32 SampleIndex = 0; SampleIndex < NumConeSampleDirections; SampleIndex++)
-		{
-			UnoccludedVector += SampleDirections[SampleIndex];
-		}
-
-		// LWC_TODO: Was float BentNormalNormalizeFactorValue = 1.0f / (UnoccludedVector / NumConeSampleDirections).Size();
-		// This causes "error C4723: potential divide by 0" in msvc, implying the compiler has managed to evaluate (UnoccludedVector / NumConeSampleDirections).Size() as 0 at compile time.
-		// Ensuring Size() is called in a seperate line stops the warning, but this needs investigating. Clang seems happy with it. Possible compiler bug?
-		float ConeSampleAverageSize = (UnoccludedVector / NumConeSampleDirections).Size();
-		float BentNormalNormalizeFactorValue = ConeSampleAverageSize ? 1.0f / ConeSampleAverageSize : 0.f;
-		SetShaderValue(RHICmdList, ShaderRHI, BentNormalNormalizeFactor, BentNormalNormalizeFactorValue);
-	}
-
-private:
-	LAYOUT_FIELD(FScreenGridParameters, ScreenGridParameters);
-	LAYOUT_FIELD(FGlobalDistanceFieldParameters, GlobalDistanceFieldParameters);
-	LAYOUT_FIELD(FShaderParameter, TanConeHalfAngle);
-	LAYOUT_FIELD(FShaderParameter, BentNormalNormalizeFactor);
 };
 
 IMPLEMENT_GLOBAL_SHADER(FConeTraceScreenGridObjectOcclusionCS, "/Engine/Private/DistanceFieldScreenGridLighting.usf", "ConeTraceObjectOcclusionCS", SF_Compute);
@@ -180,17 +173,23 @@ const int32 GConeTraceGlobalDFTileSize = 8;
 
 class FConeTraceScreenGridGlobalOcclusionCS : public FGlobalShader
 {
-	DECLARE_GLOBAL_SHADER(FConeTraceScreenGridGlobalOcclusionCS);
 public:
+	DECLARE_GLOBAL_SHADER(FConeTraceScreenGridGlobalOcclusionCS);
+	SHADER_USE_PARAMETER_STRUCT(FConeTraceScreenGridGlobalOcclusionCS, FGlobalShader);
 
 	BEGIN_SHADER_PARAMETER_STRUCT(FParameters, )
+		SHADER_PARAMETER_STRUCT_REF(FViewUniformShaderParameters, View)
 		SHADER_PARAMETER_STRUCT_INCLUDE(FDistanceFieldCulledObjectBufferParameters, DistanceFieldCulledObjectBuffers)
 		SHADER_PARAMETER_STRUCT_INCLUDE(FDistanceFieldAtlasParameters, DistanceFieldAtlas)
 		SHADER_PARAMETER_STRUCT_INCLUDE(FAOScreenGridParameters, AOScreenGridParameters)
 		SHADER_PARAMETER_STRUCT_INCLUDE(FAOParameters, AOParameters)
+		SHADER_PARAMETER_STRUCT_INCLUDE(FScreenGridParameters, ScreenGridParameters)
+		SHADER_PARAMETER_STRUCT_INCLUDE(FGlobalDistanceFieldParameters2, GlobalDistanceFieldParameters)
+		SHADER_PARAMETER_STRUCT_INCLUDE(FAOSampleParameters, AOSampleParameters)
 		SHADER_PARAMETER_RDG_BUFFER_SRV(Buffer<FVector4f>, TileConeDepthRanges)
 		SHADER_PARAMETER_RDG_UNIFORM_BUFFER(FSceneTextureUniformParameters, SceneTextures)
-		RDG_TEXTURE_ACCESS(DistanceFieldNormal, ERHIAccess::SRVCompute)
+		SHADER_PARAMETER(FIntPoint, TileListGroupSize)
+		SHADER_PARAMETER(float, TanConeHalfAngle)
 	END_SHADER_PARAMETER_STRUCT()
 
 	class FConeTraceObjects : SHADER_PERMUTATION_BOOL("CONE_TRACE_OBJECTS");
@@ -212,70 +211,6 @@ public:
 		// To reduce shader compile time of compute shaders with shared memory, doesn't have an impact on generated code with current compiler (June 2010 DX SDK)
 		OutEnvironment.CompilerFlags.Add(CFLAG_StandardOptimization);
 	}
-
-	FConeTraceScreenGridGlobalOcclusionCS(const ShaderMetaType::CompiledShaderInitializerType& Initializer)
-		: FGlobalShader(Initializer)
-	{
-		BindForLegacyShaderParameters<FParameters>(this, Initializer.PermutationId, Initializer.ParameterMap, false);
-		ScreenGridParameters.Bind(Initializer.ParameterMap);
-		GlobalDistanceFieldParameters.Bind(Initializer.ParameterMap);
-		TileListGroupSize.Bind(Initializer.ParameterMap, TEXT("TileListGroupSize"));
-		TanConeHalfAngle.Bind(Initializer.ParameterMap, TEXT("TanConeHalfAngle"));
-		BentNormalNormalizeFactor.Bind(Initializer.ParameterMap, TEXT("BentNormalNormalizeFactor"));
-	}
-
-	FConeTraceScreenGridGlobalOcclusionCS()
-	{
-	}
-
-	void SetParameters(
-		FRHICommandList& RHICmdList, 
-		const FViewInfo& View,
-		FIntPoint TileListGroupSizeValue, 
-		FRHITexture* DistanceFieldNormal,
-		const FGlobalDistanceFieldInfo& GlobalDistanceFieldInfo)
-	{
-		FRHIComputeShader* ShaderRHI = RHICmdList.GetBoundComputeShader();
-		FGlobalShader::SetParameters<FViewUniformShaderParameters>(RHICmdList, ShaderRHI, View.ViewUniformBuffer);
-
-		ScreenGridParameters.Set(RHICmdList, ShaderRHI, View, DistanceFieldNormal);
-		GlobalDistanceFieldParameters.Set(RHICmdList, ShaderRHI, GlobalDistanceFieldInfo.ParameterData);
-
-		FAOSampleData2 AOSampleData;
-
-		TArray<FVector, TInlineAllocator<9> > SampleDirections;
-		GetSpacedVectors(View.Family->FrameNumber, SampleDirections);
-
-		for (int32 SampleIndex = 0; SampleIndex < NumConeSampleDirections; SampleIndex++)
-		{
-			AOSampleData.SampleDirections[SampleIndex] = FVector4f(SampleDirections[SampleIndex]);
-		}
-
-		SetUniformBufferParameterImmediate(RHICmdList, ShaderRHI, GetUniformBufferParameter<FAOSampleData2>(), AOSampleData);
-
-		SetShaderValue(RHICmdList, ShaderRHI, TileListGroupSize, TileListGroupSizeValue);
-
-		extern float GAOConeHalfAngle;
-		SetShaderValue(RHICmdList, ShaderRHI, TanConeHalfAngle, FMath::Tan(GAOConeHalfAngle));
-
-		FVector UnoccludedVector(0);
-
-		for (int32 SampleIndex = 0; SampleIndex < NumConeSampleDirections; SampleIndex++)
-		{
-			UnoccludedVector += SampleDirections[SampleIndex];
-		}
-
-		float ConeSampleAverageSize = (UnoccludedVector / NumConeSampleDirections).Size();
-		float BentNormalNormalizeFactorValue = ConeSampleAverageSize ? 1.0f / ConeSampleAverageSize : 0.f;
-		SetShaderValue(RHICmdList, ShaderRHI, BentNormalNormalizeFactor, BentNormalNormalizeFactorValue);
-	}
-
-private:
-	LAYOUT_FIELD(FScreenGridParameters, ScreenGridParameters);
-	LAYOUT_FIELD(FGlobalDistanceFieldParameters, GlobalDistanceFieldParameters);
-	LAYOUT_FIELD(FShaderParameter, TileListGroupSize);
-	LAYOUT_FIELD(FShaderParameter, TanConeHalfAngle);
-	LAYOUT_FIELD(FShaderParameter, BentNormalNormalizeFactor);
 };
 
 IMPLEMENT_GLOBAL_SHADER(FConeTraceScreenGridGlobalOcclusionCS, "/Engine/Private/DistanceFieldScreenGridLighting.usf", "ConeTraceGlobalOcclusionCS", SF_Compute);
@@ -284,13 +219,18 @@ const int32 GCombineConesSizeX = 8;
 
 class FCombineConeVisibilityCS : public FGlobalShader
 {
-	DECLARE_GLOBAL_SHADER(FCombineConeVisibilityCS);
 public:
+	DECLARE_GLOBAL_SHADER(FCombineConeVisibilityCS);
+	SHADER_USE_PARAMETER_STRUCT(FCombineConeVisibilityCS, FGlobalShader);
 
 	BEGIN_SHADER_PARAMETER_STRUCT(FParameters, )
+		SHADER_PARAMETER_STRUCT_REF(FViewUniformShaderParameters, View)
 		SHADER_PARAMETER_STRUCT_INCLUDE(FAOScreenGridParameters, AOScreenGridParameters)
+		SHADER_PARAMETER_STRUCT_INCLUDE(FScreenGridParameters, ScreenGridParameters)
+		SHADER_PARAMETER_STRUCT_INCLUDE(FAOSampleParameters, AOSampleParameters)
 		SHADER_PARAMETER_RDG_TEXTURE_UAV(RWTexture2D, RWDistanceFieldBentNormal)
-		RDG_TEXTURE_ACCESS(DistanceFieldNormal, ERHIAccess::SRVCompute)
+		SHADER_PARAMETER(FIntPoint, ConeBufferMax)
+		SHADER_PARAMETER(FVector2f, DFNormalBufferUVMax)
 	END_SHADER_PARAMETER_STRUCT()
 
 	static bool ShouldCompilePermutation(const FGlobalShaderPermutationParameters& Parameters)
@@ -307,66 +247,6 @@ public:
 		// To reduce shader compile time of compute shaders with shared memory, doesn't have an impact on generated code with current compiler (June 2010 DX SDK)
 		OutEnvironment.CompilerFlags.Add(CFLAG_StandardOptimization);
 	}
-
-	FCombineConeVisibilityCS(const ShaderMetaType::CompiledShaderInitializerType& Initializer)
-		: FGlobalShader(Initializer)
-	{
-		BindForLegacyShaderParameters<FParameters>(this, Initializer.PermutationId, Initializer.ParameterMap, false);
-		ScreenGridParameters.Bind(Initializer.ParameterMap);
-		BentNormalNormalizeFactor.Bind(Initializer.ParameterMap, TEXT("BentNormalNormalizeFactor"));
-		ConeBufferMax.Bind(Initializer.ParameterMap, TEXT("ConeBufferMax"));
-		DFNormalBufferUVMax.Bind(Initializer.ParameterMap, TEXT("DFNormalBufferUVMax"));
-	}
-
-	FCombineConeVisibilityCS()
-	{
-	}
-
-	void SetParameters(FRHICommandList& RHICmdList, const FViewInfo& View, FRHITexture* DistanceFieldNormal)
-	{
-		FRHIComputeShader* ShaderRHI = RHICmdList.GetBoundComputeShader();
-		FGlobalShader::SetParameters<FViewUniformShaderParameters>(RHICmdList, ShaderRHI, View.ViewUniformBuffer);
-		ScreenGridParameters.Set(RHICmdList, ShaderRHI, View, DistanceFieldNormal);
-
-		FAOSampleData2 AOSampleData;
-
-		TArray<FVector, TInlineAllocator<9> > SampleDirections;
-		GetSpacedVectors(View.Family->FrameNumber, SampleDirections);
-
-		for (int32 SampleIndex = 0; SampleIndex < NumConeSampleDirections; SampleIndex++)
-		{
-			AOSampleData.SampleDirections[SampleIndex] = FVector4f(SampleDirections[SampleIndex]);
-		}
-
-		SetUniformBufferParameterImmediate(RHICmdList, ShaderRHI, GetUniformBufferParameter<FAOSampleData2>(), AOSampleData);
-
-		FVector UnoccludedVector(0);
-
-		for (int32 SampleIndex = 0; SampleIndex < NumConeSampleDirections; SampleIndex++)
-		{
-			UnoccludedVector += SampleDirections[SampleIndex];
-		}
-
-		float ConeSampleAverageSize = (UnoccludedVector / NumConeSampleDirections).Size();
-		float BentNormalNormalizeFactorValue = ConeSampleAverageSize ? 1.0f / ConeSampleAverageSize : 0.f;
-		SetShaderValue(RHICmdList, ShaderRHI, BentNormalNormalizeFactor, BentNormalNormalizeFactorValue);
-
-		FIntPoint const ConeBufferMaxValue(View.ViewRect.Width() / GAODownsampleFactor / GConeTraceDownsampleFactor - 1, View.ViewRect.Height() / GAODownsampleFactor / GConeTraceDownsampleFactor - 1);
-		SetShaderValue(RHICmdList, ShaderRHI, ConeBufferMax, ConeBufferMaxValue);
-
-		FIntPoint const DFNormalBufferSize = GetBufferSizeForAO();
-		FVector2f const DFNormalBufferUVMaxValue(
-			(View.ViewRect.Width()  / GAODownsampleFactor - 0.5f) / DFNormalBufferSize.X,
-			(View.ViewRect.Height() / GAODownsampleFactor - 0.5f) / DFNormalBufferSize.Y);
-		SetShaderValue(RHICmdList, ShaderRHI, DFNormalBufferUVMax, DFNormalBufferUVMaxValue);
-	}
-
-
-private:
-	LAYOUT_FIELD(FScreenGridParameters, ScreenGridParameters);
-	LAYOUT_FIELD(FShaderParameter, BentNormalNormalizeFactor);
-	LAYOUT_FIELD(FShaderParameter, DFNormalBufferUVMax);
-	LAYOUT_FIELD(FShaderParameter, ConeBufferMax);
 };
 
 IMPLEMENT_GLOBAL_SHADER(FCombineConeVisibilityCS, "/Engine/Private/DistanceFieldScreenGridLighting.usf", "CombineConeVisibilityCS", SF_Compute);
@@ -440,39 +320,28 @@ void FDeferredShadingSceneRenderer::RenderDistanceFieldAOScreenGrid(
 		check(View.GlobalDistanceFieldInfo.Clipmaps.Num() > 0);
 
 		auto* PassParameters = GraphBuilder.AllocParameters<FConeTraceScreenGridGlobalOcclusionCS::FParameters>();
+		PassParameters->View = View.ViewUniformBuffer;
 		PassParameters->DistanceFieldCulledObjectBuffers = CulledObjectBufferParameters;
 		PassParameters->DistanceFieldAtlas = DistanceField::SetupAtlasParameters(DistanceFieldSceneData);
 		PassParameters->AOScreenGridParameters = AOScreenGridParameters;
 		PassParameters->AOParameters = DistanceField::SetupAOShaderParameters(Parameters);
+		PassParameters->ScreenGridParameters = SetupScreenGridParameters(View, DistanceFieldNormal);
+		PassParameters->GlobalDistanceFieldParameters = SetupGlobalDistanceFieldParameters(View.GlobalDistanceFieldInfo.ParameterData);
+		PassParameters->AOSampleParameters = SetupAOSampleParameters(View.Family->FrameNumber);
 		PassParameters->TileConeDepthRanges = TileIntersectionParameters.TileConeDepthRanges;
 		PassParameters->SceneTextures = SceneTextures.UniformBuffer;
-		PassParameters->DistanceFieldNormal = DistanceFieldNormal;
+		PassParameters->TileListGroupSize = TileListGroupSize;
+		PassParameters->TanConeHalfAngle = FMath::Tan(GAOConeHalfAngle);
 
 		FConeTraceScreenGridGlobalOcclusionCS::FPermutationDomain PermutationVector;
 		PermutationVector.Set<FConeTraceScreenGridGlobalOcclusionCS::FConeTraceObjects>(bUseObjectDistanceField);
 
 		auto ComputeShader = View.ShaderMap->GetShader<FConeTraceScreenGridGlobalOcclusionCS>(PermutationVector);
 
-		ClearUnusedGraphResources(ComputeShader, PassParameters);
-
-		GraphBuilder.AddPass(
-			RDG_EVENT_NAME("ConeTraceGlobal"),
-			PassParameters,
-			ERDGPassFlags::Compute,
-			[PassParameters, ComputeShader, &View, DistanceFieldNormal, TileListGroupSize, bUseObjectDistanceField](FRHICommandList& RHICmdList)
-			{
-				const uint32 GroupSizeX = FMath::DivideAndRoundUp(View.ViewRect.Size().X / GAODownsampleFactor / GConeTraceDownsampleFactor, GConeTraceGlobalDFTileSize);
-				const uint32 GroupSizeY = FMath::DivideAndRoundUp(View.ViewRect.Size().Y / GAODownsampleFactor / GConeTraceDownsampleFactor, GConeTraceGlobalDFTileSize);
-
-				RHICmdList.SetComputeShader(ComputeShader.GetComputeShader());
-
-				ComputeShader->SetParameters(RHICmdList, View, TileListGroupSize, DistanceFieldNormal->GetRHI(), View.GlobalDistanceFieldInfo);
-				SetShaderParameters(RHICmdList, ComputeShader, ComputeShader.GetComputeShader(), *PassParameters);
-
-				DispatchComputeShader(RHICmdList, ComputeShader.GetShader(), GroupSizeX, GroupSizeY, 1);
+		const uint32 GroupSizeX = FMath::DivideAndRoundUp(View.ViewRect.Size().X / GAODownsampleFactor / GConeTraceDownsampleFactor, GConeTraceGlobalDFTileSize);
+		const uint32 GroupSizeY = FMath::DivideAndRoundUp(View.ViewRect.Size().Y / GAODownsampleFactor / GConeTraceDownsampleFactor, GConeTraceGlobalDFTileSize);
 
-				UnsetShaderUAVs(RHICmdList, ComputeShader, ComputeShader.GetComputeShader());
-			});
+		FComputeShaderUtils::AddPass(GraphBuilder, RDG_EVENT_NAME("ConeTraceGlobal"), ComputeShader, PassParameters, FIntVector(GroupSizeX, GroupSizeY, 1));
 	}
 
 	if (bUseObjectDistanceField)
@@ -480,14 +349,21 @@ void FDeferredShadingSceneRenderer::RenderDistanceFieldAOScreenGrid(
 		check(!bUseGlobalDistanceField || View.GlobalDistanceFieldInfo.Clipmaps.Num() > 0);
 
 		auto* PassParameters = GraphBuilder.AllocParameters<FConeTraceScreenGridObjectOcclusionCS::FParameters>();
+		PassParameters->View = View.ViewUniformBuffer;
 		PassParameters->DistanceFieldCulledObjectBuffers = CulledObjectBufferParameters;
 		PassParameters->DistanceFieldAtlas = DistanceField::SetupAtlasParameters(DistanceFieldSceneData);
 		PassParameters->TileIntersectionParameters = TileIntersectionParameters;
 		PassParameters->AOScreenGridParameters = AOScreenGridParameters;
 		PassParameters->AOParameters = DistanceField::SetupAOShaderParameters(Parameters);
-		PassParameters->ObjectTilesIndirectArguments = ObjectTilesIndirectArguments;
+		PassParameters->ScreenGridParameters = SetupScreenGridParameters(View, DistanceFieldNormal);
+		if (bUseGlobalDistanceField)
+		{
+			PassParameters->GlobalDistanceFieldParameters = SetupGlobalDistanceFieldParameters(View.GlobalDistanceFieldInfo.ParameterData);
+		}
+		PassParameters->AOSampleParameters = SetupAOSampleParameters(View.Family->FrameNumber);
 		PassParameters->SceneTextures = SceneTextures.UniformBuffer;
-		PassParameters->DistanceFieldNormal = DistanceFieldNormal;
+		PassParameters->TanConeHalfAngle = FMath::Tan(GAOConeHalfAngle);
+		PassParameters->ObjectTilesIndirectArguments = ObjectTilesIndirectArguments;
 
 		FConeTraceScreenGridObjectOcclusionCS::FPermutationDomain PermutationVector;
 		PermutationVector.Set<FConeTraceScreenGridObjectOcclusionCS::FUseGlobalDistanceField>(bUseGlobalDistanceField);
@@ -495,23 +371,7 @@ void FDeferredShadingSceneRenderer::RenderDistanceFieldAOScreenGrid(
 
 		auto ComputeShader = View.ShaderMap->GetShader<FConeTraceScreenGridObjectOcclusionCS>(PermutationVector);
 
-		ClearUnusedGraphResources(ComputeShader, PassParameters);
-
-		GraphBuilder.AddPass(
-			RDG_EVENT_NAME("ConeTraceObjects"),
-			PassParameters,
-			ERDGPassFlags::Compute,
-			[PassParameters, ComputeShader, &View, DistanceFieldNormal, bUseGlobalDistanceField, ObjectTilesIndirectArguments](FRHICommandList& RHICmdList)
-			{
-				RHICmdList.SetComputeShader(ComputeShader.GetComputeShader());
-
-				ComputeShader->SetParameters(RHICmdList, View, DistanceFieldNormal->GetRHI(), bUseGlobalDistanceField, View.GlobalDistanceFieldInfo);
-				SetShaderParameters(RHICmdList, ComputeShader, ComputeShader.GetComputeShader(), *PassParameters);
-
-				DispatchIndirectComputeShader(RHICmdList, ComputeShader.GetShader(), ObjectTilesIndirectArguments->GetIndirectRHICallBuffer(), 0);
-
-				UnsetShaderUAVs(RHICmdList, ComputeShader, ComputeShader.GetComputeShader());
-			});
+		FComputeShaderUtils::AddPass(GraphBuilder, RDG_EVENT_NAME("ConeTraceObjects"), ComputeShader, PassParameters, ObjectTilesIndirectArguments, 0);
 	}
 
 	FRDGTextureRef DownsampledBentNormal = nullptr;
@@ -525,30 +385,23 @@ void FDeferredShadingSceneRenderer::RenderDistanceFieldAOScreenGrid(
 		const uint32 GroupSizeX = FMath::DivideAndRoundUp(ConeTraceBufferSize.X, GCombineConesSizeX);
 		const uint32 GroupSizeY = FMath::DivideAndRoundUp(ConeTraceBufferSize.Y, GCombineConesSizeX);
 
+		const FIntPoint DFNormalBufferSize = GetBufferSizeForAO();
+		const FVector2f DFNormalBufferUVMaxValue(
+			(View.ViewRect.Width() / GAODownsampleFactor - 0.5f) / DFNormalBufferSize.X,
+			(View.ViewRect.Height() / GAODownsampleFactor - 0.5f) / DFNormalBufferSize.Y);
+
 		auto* PassParameters = GraphBuilder.AllocParameters<FCombineConeVisibilityCS::FParameters>();
+		PassParameters->View = View.ViewUniformBuffer;
 		PassParameters->AOScreenGridParameters = AOScreenGridParameters;
+		PassParameters->ScreenGridParameters = SetupScreenGridParameters(View, DistanceFieldNormal);
+		PassParameters->AOSampleParameters = SetupAOSampleParameters(View.Family->FrameNumber);
 		PassParameters->RWDistanceFieldBentNormal = GraphBuilder.CreateUAV(DownsampledBentNormal);
-		PassParameters->DistanceFieldNormal = DistanceFieldNormal;
+		PassParameters->ConeBufferMax = FIntPoint(View.ViewRect.Width() / GAODownsampleFactor / GConeTraceDownsampleFactor - 1, View.ViewRect.Height() / GAODownsampleFactor / GConeTraceDownsampleFactor - 1);
+		PassParameters->DFNormalBufferUVMax = DFNormalBufferUVMaxValue;
 
 		auto ComputeShader = View.ShaderMap->GetShader<FCombineConeVisibilityCS>();
 
-		ClearUnusedGraphResources(ComputeShader, PassParameters);
-
-		GraphBuilder.AddPass(
-			RDG_EVENT_NAME("CombineCones"),
-			PassParameters,
-			ERDGPassFlags::Compute,
-			[&View, ComputeShader, PassParameters, GroupSizeX, GroupSizeY](FRHICommandList& RHICmdList)
-			{
-				RHICmdList.SetComputeShader(ComputeShader.GetComputeShader());
-
-				ComputeShader->SetParameters(RHICmdList, View, PassParameters->DistanceFieldNormal->GetRHI());
-				SetShaderParameters(RHICmdList, ComputeShader, ComputeShader.GetComputeShader(), *PassParameters);
-
-				DispatchComputeShader(RHICmdList, ComputeShader.GetShader(), GroupSizeX, GroupSizeY, 1);
-
-				UnsetShaderUAVs(RHICmdList, ComputeShader, ComputeShader.GetComputeShader());
-			});
+		FComputeShaderUtils::AddPass(GraphBuilder, RDG_EVENT_NAME("CombineCones"), ComputeShader, PassParameters, FIntVector(GroupSizeX, GroupSizeY, 1));
 	}
 
 	PostProcessBentNormalAOScreenGrid(
diff --git a/Engine/Source/Runtime/Renderer/Private/DistanceFieldVisualization.cpp b/Engine/Source/Runtime/Renderer/Private/DistanceFieldVisualization.cpp
--- a/Engine/Source/Runtime/Renderer/Private/DistanceFieldVisualization.cpp
+++ b/Engine/Source/Runtime/Renderer/Private/DistanceFieldVisualization.cpp
@@ -21,6 +21,7 @@ class FVisualizeMeshDistanceFieldCS : public FGlobalShader
 {
 public:
 	DECLARE_GLOBAL_SHADER(FVisualizeMeshDistanceFieldCS);
+	SHADER_USE_PARAMETER_STRUCT(FVisualizeMeshDistanceFieldCS, FGlobalShader);
 
 	class FUseGlobalDistanceFieldDim : SHADER_PERMUTATION_BOOL("USE_GLOBAL_DISTANCE_FIELD");
 	using FPermutationDomain = TShaderPermutationDomain<FUseGlobalDistanceFieldDim>;
@@ -32,6 +33,7 @@ public:
 		SHADER_PARAMETER_STRUCT_INCLUDE(FDistanceFieldCulledObjectBufferParameters, DistanceFieldCulledObjectBuffers)
 		SHADER_PARAMETER_STRUCT_INCLUDE(FDistanceFieldAtlasParameters, DistanceFieldAtlas)
 		SHADER_PARAMETER_STRUCT_INCLUDE(FAOParameters, AOParameters)
+		SHADER_PARAMETER_STRUCT_INCLUDE(FGlobalDistanceFieldParameters2, GlobalDistanceFieldParameters)
 		SHADER_PARAMETER_RDG_UNIFORM_BUFFER(FSceneTextureUniformParameters, SceneTextures)
 		SHADER_PARAMETER_RDG_TEXTURE_UAV(RWTexture2D<float4>, RWVisualizeMeshDistanceFields)
 		SHADER_PARAMETER(FVector2f, NumGroups)
@@ -48,29 +50,6 @@ public:
 		OutEnvironment.SetDefine(TEXT("THREADGROUP_SIZEX"), GDistanceFieldAOTileSizeX);
 		OutEnvironment.SetDefine(TEXT("THREADGROUP_SIZEY"), GDistanceFieldAOTileSizeY);
 	}
-
-	FVisualizeMeshDistanceFieldCS() = default;
-	FVisualizeMeshDistanceFieldCS(const ShaderMetaType::CompiledShaderInitializerType& Initializer)
-		: FGlobalShader(Initializer)
-	{
-		BindForLegacyShaderParameters<FParameters>(this, Initializer.PermutationId, Initializer.ParameterMap);
-		GlobalDistanceFieldParameters.Bind(Initializer.ParameterMap);
-	}
-
-	void SetParameters(
-		FRHICommandList& RHICmdList,
-		const FGlobalDistanceFieldInfo& GlobalDistanceFieldInfo)
-	{
-		FRHIComputeShader* ShaderRHI = RHICmdList.GetBoundComputeShader();
-
-		if (GlobalDistanceFieldParameters.IsBound())
-		{
-			GlobalDistanceFieldParameters.Set(RHICmdList, ShaderRHI, GlobalDistanceFieldInfo.ParameterData);
-		}
-	}
-
-private:
-	LAYOUT_FIELD(FGlobalDistanceFieldParameters, GlobalDistanceFieldParameters);
 };
 
 IMPLEMENT_GLOBAL_SHADER(FVisualizeMeshDistanceFieldCS, "/Engine/Private/DistanceFieldVisualization.usf", "VisualizeMeshDistanceFieldCS", SF_Compute);
@@ -166,25 +145,12 @@ void FDeferredShadingSceneRenderer::RenderMeshDistanceFieldVisualization(
 		PassParameters->DistanceFieldCulledObjectBuffers = CulledObjectBufferParameters;
 		PassParameters->DistanceFieldAtlas = DistanceField::SetupAtlasParameters(Scene->DistanceFieldSceneData);
 		PassParameters->AOParameters = DistanceField::SetupAOShaderParameters(Parameters);
+		PassParameters->GlobalDistanceFieldParameters = SetupGlobalDistanceFieldParameters(View.GlobalDistanceFieldInfo.ParameterData);
 		PassParameters->NumGroups = FVector2f(GroupSizeX, GroupSizeY);
 		PassParameters->SceneTextures = SceneTextures.UniformBuffer;
 		PassParameters->RWVisualizeMeshDistanceFields = GraphBuilder.CreateUAV(VisualizeResultTexture);
 
-		ClearUnusedGraphResources(ComputeShader, PassParameters);
-
-		GraphBuilder.AddPass(
-			RDG_EVENT_NAME("VisualizeMeshDistanceFieldCS"),
-			PassParameters,
-			ERDGPassFlags::Compute,
-			[&View, ComputeShader, PassParameters, GroupSizeX, GroupSizeY](FRHICommandList& RHICmdList)
-		{
-			FRHIComputeShader* ShaderRHI = ComputeShader.GetComputeShader();
-			RHICmdList.SetComputeShader(ComputeShader.GetComputeShader());
-			SetShaderParameters(RHICmdList, ComputeShader, ShaderRHI, *PassParameters);
-			ComputeShader->SetParameters(RHICmdList,View.GlobalDistanceFieldInfo);
-			DispatchComputeShader(RHICmdList, ComputeShader.GetShader(), GroupSizeX, GroupSizeY, 1);
-			UnsetShaderUAVs(RHICmdList, ComputeShader, ShaderRHI);
-		});
+		FComputeShaderUtils::AddPass(GraphBuilder, RDG_EVENT_NAME("VisualizeMeshDistanceFieldCS"), ComputeShader, PassParameters, FIntVector(GroupSizeX, GroupSizeY, 1));
 	}
 
 	TShaderMapRef<FVisualizeDistanceFieldUpsamplePS> PixelShader(FirstView.ShaderMap);
diff --git a/Engine/Source/Runtime/Renderer/Private/GlobalDistanceField.cpp b/Engine/Source/Runtime/Renderer/Private/GlobalDistanceField.cpp
--- a/Engine/Source/Runtime/Renderer/Private/GlobalDistanceField.cpp
+++ b/Engine/Source/Runtime/Renderer/Private/GlobalDistanceField.cpp
@@ -205,6 +205,35 @@ FAutoConsoleVariableRef CVarAOGlobalDistanceFieldMipFactor(
 	ECVF_Scalability | ECVF_RenderThreadSafe
 );
 
+FGlobalDistanceFieldParameters2 SetupGlobalDistanceFieldParameters(const FGlobalDistanceFieldParameterData& ParameterData)
+{
+	FGlobalDistanceFieldParameters2 ShaderParameters;
+
+	ShaderParameters.GlobalDistanceFieldPageAtlasTexture = OrBlack3DIfNull(ParameterData.PageAtlasTexture);
+	ShaderParameters.GlobalDistanceFieldPageTableTexture = OrBlack3DIfNull(ParameterData.PageTableTexture);
+	ShaderParameters.GlobalDistanceFieldMipTexture = OrBlack3DIfNull(ParameterData.MipTexture);
+
+	for (int32 Index = 0; Index < GMaxGlobalDistanceFieldClipmaps; Index++)
+	{
+		ShaderParameters.GlobalVolumeCenterAndExtent[Index] = ParameterData.CenterAndExtent[Index];
+		ShaderParameters.GlobalVolumeWorldToUVAddAndMul[Index] = ParameterData.WorldToUVAddAndMul[Index];
+		ShaderParameters.GlobalDistanceFieldMipWorldToUVScale[Index] = ParameterData.MipWorldToUVScale[Index];
+		ShaderParameters.GlobalDistanceFieldMipWorldToUVBias[Index] = ParameterData.MipWorldToUVBias[Index];
+	}
+
+	ShaderParameters.GlobalDistanceFieldMipFactor = ParameterData.MipFactor;
+	ShaderParameters.GlobalDistanceFieldMipTransition = ParameterData.MipTransition;
+	ShaderParameters.GlobalDistanceFieldClipmapSizeInPages = ParameterData.ClipmapSizeInPages;
+	ShaderParameters.GlobalDistanceFieldInvPageAtlasSize = ParameterData.InvPageAtlasSize;
+	ShaderParameters.GlobalVolumeDimension = ParameterData.GlobalDFResolution;
+	ShaderParameters.GlobalVolumeTexelSize = 1.0f / ParameterData.GlobalDFResolution;
+	ShaderParameters.MaxGlobalDFAOConeDistance = ParameterData.MaxDFAOConeDistance;
+	ShaderParameters.NumGlobalSDFClipmaps = ParameterData.NumGlobalSDFClipmaps;
+
+	return ShaderParameters;
+}
+
+
 float GetMinMeshSDFRadius(float VoxelWorldSize)
 {
 	float MinRadius = GAOGlobalDistanceFieldMinMeshSDFRadius * (GAOGlobalDistanceFieldFastCameraMode ? 10.0f : 1.0f);
@@ -379,8 +408,6 @@ void FGlobalDistanceFieldInfo::UpdateParameterData(float MaxOcclusionDistance, b
 
 				ParameterData.MipWorldToUVScale[ClipmapIndex].Z = ParameterData.MipWorldToUVScale[ClipmapIndex].Z / Clipmaps.Num();
 				ParameterData.MipWorldToUVBias[ClipmapIndex].Z = (ParameterData.MipWorldToUVBias[ClipmapIndex].Z + ClipmapIndex) / Clipmaps.Num();
-
-				ParameterData.PageTableScrollOffset[ClipmapIndex] = Clipmap.ScrollOffset / Clipmap.Bounds.GetSize();
 			}
 			else
 			{
@@ -388,7 +415,6 @@ void FGlobalDistanceFieldInfo::UpdateParameterData(float MaxOcclusionDistance, b
 				ParameterData.WorldToUVAddAndMul[ClipmapIndex] = FVector4f(0);
 				ParameterData.MipWorldToUVScale[ClipmapIndex] = FVector(0);
 				ParameterData.MipWorldToUVBias[ClipmapIndex] = FVector(0);
-				ParameterData.PageTableScrollOffset[ClipmapIndex] = FVector(0);
 			}
 		}
 
diff --git a/Engine/Source/Runtime/Renderer/Private/VolumetricFog.cpp b/Engine/Source/Runtime/Renderer/Private/VolumetricFog.cpp
--- a/Engine/Source/Runtime/Renderer/Private/VolumetricFog.cpp
+++ b/Engine/Source/Runtime/Renderer/Private/VolumetricFog.cpp
@@ -822,6 +822,7 @@ class TVolumetricFogLightScatteringCS : public FGlobalShader
 		SHADER_PARAMETER_RDG_UNIFORM_BUFFER(FLumenTranslucencyLightingUniforms, LumenGIVolumeStruct)
 		SHADER_PARAMETER_STRUCT_INCLUDE(FVirtualShadowMapSamplingParameters, VirtualShadowMapSamplingParameters)
 		SHADER_PARAMETER_STRUCT_INCLUDE(FAOParameters, AOParameters)
+		SHADER_PARAMETER_STRUCT_INCLUDE(FGlobalDistanceFieldParameters2, GlobalDistanceFieldParameters)
 		SHADER_PARAMETER_RDG_TEXTURE_UAV(RWTexture2D, RWLightScattering)
 	END_SHADER_PARAMETER_STRUCT()
 
@@ -889,7 +890,6 @@ public:
 		InverseSquaredLightDistanceBiasScale.Bind(Initializer.ParameterMap, TEXT("InverseSquaredLightDistanceBiasScale"));
 		UseHeightFogColors.Bind(Initializer.ParameterMap, TEXT("UseHeightFogColors"));
 		UseDirectionalLightShadowing.Bind(Initializer.ParameterMap, TEXT("UseDirectionalLightShadowing"));
-		GlobalDistanceFieldParameters.Bind(Initializer.ParameterMap);
 		LightScatteringSampleJitterMultiplier.Bind(Initializer.ParameterMap, TEXT("LightScatteringSampleJitterMultiplier"));
 
 		CloudShadowmapTexture.Bind(Initializer.ParameterMap, TEXT("CloudShadowmapTexture"));
@@ -980,8 +980,6 @@ public:
 			OverrideDirectionalLightInScatteringUsingHeightFog(View, FogInfo) ? 1.0f : 0.0f,
 			OverrideSkyLightInScatteringUsingHeightFog(View, FogInfo) ? 1.0f : 0.0f ));
 
-		GlobalDistanceFieldParameters.Set(RHICmdList, ShaderRHI, View.GlobalDistanceFieldInfo.ParameterData);
-
 		if (CloudShadowmapTexture.IsBound())
 		{
 			FMatrix CloudWorldToLightClipShadowMatrix = FMatrix::Identity;
@@ -1042,7 +1040,6 @@ private:
 	LAYOUT_FIELD(FShaderParameter, InverseSquaredLightDistanceBiasScale);
 	LAYOUT_FIELD(FShaderParameter, UseHeightFogColors);
 	LAYOUT_FIELD(FShaderParameter, UseDirectionalLightShadowing);
-	LAYOUT_FIELD(FGlobalDistanceFieldParameters, GlobalDistanceFieldParameters);
 	LAYOUT_FIELD(FShaderResourceParameter, CloudShadowmapTexture);
 	LAYOUT_FIELD(FShaderResourceParameter, CloudShadowmapSampler);
 	LAYOUT_FIELD(FShaderParameter, CloudShadowmapFarDepthKm);
@@ -1497,6 +1494,7 @@ void FDeferredShadingSceneRenderer::ComputeVolumetricFog(FRDGBuilder& GraphBuild
 				AOParameterData = FDistanceFieldAOParameters(Scene->SkyLight->OcclusionMaxDistance, Scene->SkyLight->Contrast);
 			}
 			PassParameters->AOParameters = DistanceField::SetupAOShaderParameters(AOParameterData);
+			PassParameters->GlobalDistanceFieldParameters = SetupGlobalDistanceFieldParameters(View.GlobalDistanceFieldInfo.ParameterData);
 
 			const bool bUseLumenGI = View.LumenTranslucencyGIVolume.Texture0 != nullptr;
 			const bool bUseGlobalDistanceField = UseGlobalDistanceField() && Scene->DistanceFieldSceneData.NumObjectsInBuffer > 0;
diff --git a/Engine/Source/Runtime/Renderer/Public/GlobalDistanceFieldParameters.h b/Engine/Source/Runtime/Renderer/Public/GlobalDistanceFieldParameters.h
--- a/Engine/Source/Runtime/Renderer/Public/GlobalDistanceFieldParameters.h
+++ b/Engine/Source/Runtime/Renderer/Public/GlobalDistanceFieldParameters.h
@@ -31,7 +31,6 @@ public:
 	FVector4f MipWorldToUVBias[GMaxGlobalDistanceFieldClipmaps];
 	float MipFactor;
 	float MipTransition;
-	FVector PageTableScrollOffset[GMaxGlobalDistanceFieldClipmaps];
 	FRHITexture* PageAtlasTexture;
 	FRHITexture* PageTableTexture;
 	FRHITexture* MipTexture;
@@ -43,6 +42,26 @@ public:
 	int32 NumGlobalSDFClipmaps;
 };
 
+BEGIN_SHADER_PARAMETER_STRUCT(FGlobalDistanceFieldParameters2, )
+	SHADER_PARAMETER_TEXTURE(Texture3D, GlobalDistanceFieldPageAtlasTexture)
+	SHADER_PARAMETER_TEXTURE(Texture3D<uint>, GlobalDistanceFieldPageTableTexture)
+	SHADER_PARAMETER_TEXTURE(Texture3D, GlobalDistanceFieldMipTexture)
+	SHADER_PARAMETER_ARRAY(FVector4f, GlobalVolumeCenterAndExtent, [GMaxGlobalDistanceFieldClipmaps])
+	SHADER_PARAMETER_ARRAY(FVector4f, GlobalVolumeWorldToUVAddAndMul, [GMaxGlobalDistanceFieldClipmaps])
+	SHADER_PARAMETER_ARRAY(FVector4f, GlobalDistanceFieldMipWorldToUVScale, [GMaxGlobalDistanceFieldClipmaps])
+	SHADER_PARAMETER_ARRAY(FVector4f, GlobalDistanceFieldMipWorldToUVBias, [GMaxGlobalDistanceFieldClipmaps])
+	SHADER_PARAMETER(float, GlobalDistanceFieldMipFactor)
+	SHADER_PARAMETER(float, GlobalDistanceFieldMipTransition)
+	SHADER_PARAMETER(int32, GlobalDistanceFieldClipmapSizeInPages)
+	SHADER_PARAMETER(FVector3f, GlobalDistanceFieldInvPageAtlasSize)
+	SHADER_PARAMETER(float, GlobalVolumeDimension)
+	SHADER_PARAMETER(float, GlobalVolumeTexelSize)
+	SHADER_PARAMETER(float, MaxGlobalDFAOConeDistance)
+	SHADER_PARAMETER(uint32, NumGlobalSDFClipmaps)
+END_SHADER_PARAMETER_STRUCT()
+
+FGlobalDistanceFieldParameters2 SetupGlobalDistanceFieldParameters(const FGlobalDistanceFieldParameterData& ParameterData);
+
 class FGlobalDistanceFieldParameters
 {
 	DECLARE_INLINE_TYPE_LAYOUT(FGlobalDistanceFieldParameters, NonVirtual);
