Remove CVars for toggling strata classification.

Strata classification is now always enabled. This CL also simplifies/cleans up the overall setup.

diff --git a/Engine/Source/Runtime/Renderer/Private/IndirectLightRendering.cpp b/Engine/Source/Runtime/Renderer/Private/IndirectLightRendering.cpp
--- a/Engine/Source/Runtime/Renderer/Private/IndirectLightRendering.cpp
+++ b/Engine/Source/Runtime/Renderer/Private/IndirectLightRendering.cpp
@@ -230,7 +230,7 @@ class FReflectionEnvironmentSkyLightingPS : public FGlobalShader
 			PermutationVector.Set<FSkyShadowing>(false);
 		}
 
-		if (PermutationVector.Get<FStrataFastPath>() && (!Strata::IsStrataEnabled() || !Strata::IsClassificationEnabled()))
+		if (PermutationVector.Get<FStrataFastPath>() && !Strata::IsStrataEnabled())
 		{
 			PermutationVector.Set<FStrataFastPath>(false);
 		}
@@ -249,7 +249,7 @@ class FReflectionEnvironmentSkyLightingPS : public FGlobalShader
 		PermutationVector.Set<FDynamicSkyLight>(bEnableDynamicSkyLight);
 		PermutationVector.Set<FSkyShadowing>(bApplySkyShadowing);
 		PermutationVector.Set<FRayTracedReflections>(bRayTracedReflections);
-		PermutationVector.Set<FStrataFastPath>(Strata::IsStrataEnabled() && Strata::IsClassificationEnabled() && bStrataFastPath);
+		PermutationVector.Set<FStrataFastPath>(Strata::IsStrataEnabled() && bStrataFastPath);
 
 		return RemapPermutation(PermutationVector);
 	}
@@ -1162,9 +1162,7 @@ static void AddSkyReflectionPass(
 	bool bSkyLight, 
 	bool bDynamicSkyLight, 
 	bool bApplySkyShadowing,
-	bool bEnableStrataStencilTest,
-	bool bEnableStrataTiledPass,
-	uint32 StrataTileMaterialType)	// 0 simple material, 1 all/complex material. Does not matter when strata is disabled. See EStrataTileMaterialType
+	EStrataTileMaterialType StrataTileMaterialType)
 {
 	// Render the reflection environment with tiled deferred culling
 	bool bHasBoxCaptures = (View.NumBoxReflectionCaptures > 0);
@@ -1267,10 +1265,6 @@ static void AddSkyReflectionPass(
 	}
 
 	PassParameters->RenderTargets[0] = FRenderTargetBinding(SceneColorTexture.Target, ERenderTargetLoadAction::ELoad);
-	if (bEnableStrataStencilTest)
-	{		
-		PassParameters->RenderTargets.DepthStencil = FDepthStencilBinding(SceneTextureParameters.SceneDepthTexture, ERenderTargetLoadAction::ELoad, ERenderTargetLoadAction::ELoad, FExclusiveDepthStencil::DepthRead_StencilWrite);
-	}
 
 	// Bind hair data
 	const bool bCheckerboardSubsurfaceRendering = IsSubsurfaceCheckerboardFormat(SceneColorTexture.Target->Desc.Format);
@@ -1296,35 +1290,28 @@ static void AddSkyReflectionPass(
 	TShaderMapRef<Strata::FStrataTilePassVS> StrataTilePassVertexShader(View.ShaderMap, VSPermutationVector);
 	PassParameters->VS.TileIndirectBuffer = nullptr;
 	PassParameters->VS.TileListBuffer = nullptr;
-	if (bEnableStrataTiledPass)
+	const bool bStrataEnabled = Strata::IsStrataEnabled();
+	if (bStrataEnabled)
 	{
-		Strata::FillUpTiledPassData((EStrataTileMaterialType)StrataTileMaterialType, View, PassParameters->VS, StrataTilePrimitiveType);
+		check(StrataTileMaterialType != EStrataTileMaterialType::ECount);
+		Strata::FillUpTiledPassData(StrataTileMaterialType, View, PassParameters->VS, StrataTilePrimitiveType);
 		ClearUnusedGraphResources(StrataTilePassVertexShader, &PassParameters->VS);
 	}
 
 	GraphBuilder.AddPass(
-		RDG_EVENT_NAME("ReflectionEnvironmentAndSky %dx%d StrataMat=%s", 
+		RDG_EVENT_NAME("ReflectionEnvironmentAndSky(%dx%d,StrataMat=%s)", 
 			View.ViewRect.Width(), View.ViewRect.Height(), 
-			Strata::IsStrataEnabled() ? ToString(EStrataTileMaterialType(StrataTileMaterialType)) : TEXT("Off")),
+			bStrataEnabled ? ToString(StrataTileMaterialType) : TEXT("Off")),
 		PassParameters,
 		ERDGPassFlags::Raster,
 		[PassParameters, &View, PixelShader, bCheckerboardSubsurfaceRendering, 
-		bEnableStrataStencilTest, bEnableStrataTiledPass, StrataTilePassVertexShader, StrataTileMaterialType, StrataTilePrimitiveType](FRHICommandList& InRHICmdList)
+		StrataTilePassVertexShader, bStrataEnabled, StrataTilePrimitiveType](FRHICommandList& InRHICmdList)
 	{
 		InRHICmdList.SetViewport(View.ViewRect.Min.X, View.ViewRect.Min.Y, 0.0f, View.ViewRect.Max.X, View.ViewRect.Max.Y, 1.0f);
 
 		FGraphicsPipelineStateInitializer GraphicsPSOInit;
 		FPixelShaderUtils::InitFullscreenPipelineState(InRHICmdList, View.ShaderMap, PixelShader, GraphicsPSOInit);
 
-		if (bEnableStrataStencilTest)
-		{
-			GraphicsPSOInit.DepthStencilState = TStaticDepthStencilState<
-				false, CF_Always, 
-				true,  CF_Equal, SO_Keep, SO_Keep, SO_Keep, 
-				false, CF_Equal, SO_Keep, SO_Keep, SO_Keep,
-				Strata::StencilBit, Strata::StencilBit>::GetRHI();
-		}
-
 		extern int32 GAOOverwriteSceneColor;
 		if (GetReflectionEnvironmentCVar() == 2 || GAOOverwriteSceneColor)
 		{
@@ -1352,23 +1339,23 @@ static void AddSkyReflectionPass(
 			InRHICmdList.SetDepthBounds(Values.MinDepth, Values.MaxDepth);
 		}
 
-		if (bEnableStrataTiledPass)
+		if (bStrataEnabled)
 		{
 			GraphicsPSOInit.BoundShaderState.VertexShaderRHI = StrataTilePassVertexShader.GetVertexShader();
 			GraphicsPSOInit.PrimitiveType = StrataTilePrimitiveType;
 		}
 
-		SetGraphicsPipelineState(InRHICmdList, GraphicsPSOInit, bEnableStrataStencilTest && StrataTileMaterialType == EStrataTileMaterialType::ESimple ? Strata::StencilBit : 0x0);
+		SetGraphicsPipelineState(InRHICmdList, GraphicsPSOInit, 0x0);
 		SetShaderParameters(InRHICmdList, PixelShader, PixelShader.GetPixelShader(), PassParameters->PS);
 
-		if (!bEnableStrataTiledPass)
+		if (bStrataEnabled)
 		{
-			FPixelShaderUtils::DrawFullscreenTriangle(InRHICmdList);
+			SetShaderParameters(InRHICmdList, StrataTilePassVertexShader, StrataTilePassVertexShader.GetVertexShader(), PassParameters->VS);
+			InRHICmdList.DrawPrimitiveIndirect(PassParameters->VS.TileIndirectBuffer->GetIndirectRHICallBuffer(), 0);
 		}
 		else
 		{
-			SetShaderParameters(InRHICmdList, StrataTilePassVertexShader, StrataTilePassVertexShader.GetVertexShader(), PassParameters->VS);
-			InRHICmdList.DrawPrimitiveIndirect(PassParameters->VS.TileIndirectBuffer->GetIndirectRHICallBuffer(), 0);
+			FPixelShaderUtils::DrawFullscreenTriangle(InRHICmdList);
 		}
 	});
 }
@@ -1593,14 +1580,8 @@ void FDeferredShadingSceneRenderer::RenderDeferredReflectionsAndSkyLighting(
 		{
 			RDG_GPU_STAT_SCOPE(GraphBuilder, ReflectionEnvironment);
 
-			const bool bStrataClassificationEnabled = Strata::IsStrataEnabled() && Strata::IsClassificationEnabled();
-			const bool bTilePassesReadingStrataEnabled = Strata::ShouldPassesReadingStrataBeTiled(Scene->GetFeatureLevel());
-
-			if (bStrataClassificationEnabled && bTilePassesReadingStrataEnabled)
+			if (Strata::IsStrataEnabled())
 			{
-				const bool bEnableStrataTiledPass = true;
-				const bool bEnableStrataStencilTest = false;
-
 				AddSkyReflectionPass(
 					GraphBuilder,
 					View,
@@ -1613,8 +1594,6 @@ void FDeferredShadingSceneRenderer::RenderDeferredReflectionsAndSkyLighting(
 					bSkyLight,
 					bDynamicSkyLight,
 					bApplySkyShadowing,
-					bEnableStrataStencilTest,
-					bEnableStrataTiledPass,
 					EStrataTileMaterialType::ESimple);
 
 				AddSkyReflectionPass(
@@ -1629,15 +1608,10 @@ void FDeferredShadingSceneRenderer::RenderDeferredReflectionsAndSkyLighting(
 					bSkyLight,
 					bDynamicSkyLight,
 					bApplySkyShadowing,
-					bEnableStrataStencilTest,
-					bEnableStrataTiledPass,
 					EStrataTileMaterialType::EComplex);
 			}
 			else
 			{
-				const bool bEnableStrataTiledPass = false;
-				const bool bEnableStrataStencilTest = bStrataClassificationEnabled;
-
 				// Typical path uses when Strata is not enabled
 				AddSkyReflectionPass(
 					GraphBuilder,
@@ -1651,27 +1625,7 @@ void FDeferredShadingSceneRenderer::RenderDeferredReflectionsAndSkyLighting(
 					bSkyLight,
 					bDynamicSkyLight,
 					bApplySkyShadowing,
-					bEnableStrataStencilTest,
-					bEnableStrataTiledPass,
-					EStrataTileMaterialType::EComplex);
-				if (bStrataClassificationEnabled && !bTilePassesReadingStrataEnabled)
-				{
-					AddSkyReflectionPass(
-						GraphBuilder,
-						View,
-						Scene,
-						SceneTextures,
-						DynamicBentNormalAOTexture,
-						ReflectionsColor,
-						RayTracingReflectionOptions,
-						SceneTextureParameters,
-						bSkyLight,
-						bDynamicSkyLight,
-						bApplySkyShadowing,
-						bEnableStrataStencilTest,
-						bEnableStrataTiledPass,
-						EStrataTileMaterialType::ESimple);
-				}
+					EStrataTileMaterialType::ECount);
 			}
 		}
 
diff --git a/Engine/Source/Runtime/Renderer/Private/LightRendering.cpp b/Engine/Source/Runtime/Renderer/Private/LightRendering.cpp
--- a/Engine/Source/Runtime/Renderer/Private/LightRendering.cpp
+++ b/Engine/Source/Runtime/Renderer/Private/LightRendering.cpp
@@ -858,8 +858,7 @@ void FDeferredShadingSceneRenderer::RenderLights(
 		RDG_EVENT_SCOPE(GraphBuilder, "DirectLighting");
 
 		// STRATA_TODO move right after stencil clear so that it is also common with EnvLight pass
-		if (ViewFamily.EngineShowFlags.DirectLighting &&
-			Strata::IsStrataEnabled() && Strata::IsClassificationEnabled())
+		if (ViewFamily.EngineShowFlags.DirectLighting && Strata::IsStrataEnabled())
 		{
 			// Update the stencil buffer, marking simple/complex strata material only once for all the following passes.
 			Strata::AddStrataStencilPass(GraphBuilder, Views, SceneTextures);
@@ -1668,8 +1667,7 @@ void FDeferredShadingSceneRenderer::RenderStationaryLightOverlap(
 	}
 }
 
-/** Sets up rasterizer and depth state for rendering bounding geometry in a deferred pass. */
-void SetBoundingGeometryRasterizerAndDepthState(FGraphicsPipelineStateInitializer& GraphicsPSOInit, const FViewInfo& View, bool bCameraInsideLightGeometry)
+static void InternalSetBoundingGeometryRasterizerState(FGraphicsPipelineStateInitializer& GraphicsPSOInit, const FViewInfo& View, bool bCameraInsideLightGeometry)
 {
 	if (bCameraInsideLightGeometry)
 	{
@@ -1681,29 +1679,46 @@ void SetBoundingGeometryRasterizerAndDepthState(FGraphicsPipelineStateInitialize
 		// Render frontfaces with depth tests on to get the speedup from HiZ since the camera is outside the light geometry
 		GraphicsPSOInit.RasterizerState = View.bReverseCulling ? TStaticRasterizerState<FM_Solid, CM_CCW>::GetRHI() : TStaticRasterizerState<FM_Solid, CM_CW>::GetRHI();
 	}
+}
 
-	if (Strata::IsStrataEnabled() && Strata::IsClassificationEnabled())
+template<ECompareFunction CompareFunction>
+static uint32 InternalSetBoundingGeometryDepthState(FGraphicsPipelineStateInitializer& GraphicsPSOInit, EStrataTileMaterialType TileType)
+{
+	// bCameraInsideLightGeometry = true  -> CompareFunction = Always
+	// bCameraInsideLightGeometry = false -> CompareFunction = CF_DepthNearOrEqual
+	uint32 StencilRef = 0u;
+	if (Strata::IsStrataEnabled())
 	{
-		GraphicsPSOInit.DepthStencilState =
-		bCameraInsideLightGeometry
-		? TStaticDepthStencilState<
-			false, CF_Always,
-			true, CF_Equal, SO_Keep, SO_Keep, SO_Keep,
-			true, CF_Equal, SO_Keep, SO_Keep, SO_Keep,
-			Strata::StencilBit, 0x0>::GetRHI()
-		: TStaticDepthStencilState<
-			false, CF_DepthNearOrEqual,
-			true, CF_Equal, SO_Keep, SO_Keep, SO_Keep,
-			true, CF_Equal, SO_Keep, SO_Keep, SO_Keep,
-			Strata::StencilBit, 0x0>::GetRHI();
+		check(TileType != EStrataTileMaterialType::ECount);
+		switch (TileType)
+		{
+		case EStrataTileMaterialType::ESimple : StencilRef = Strata::StencilBit_Fast;    GraphicsPSOInit.DepthStencilState = TStaticDepthStencilState<false, CompareFunction, true, CF_Equal, SO_Keep, SO_Keep, SO_Keep, true, CF_Equal, SO_Keep, SO_Keep, SO_Keep, Strata::StencilBit_Fast, 0x0>::GetRHI(); break;
+		case EStrataTileMaterialType::EComplex: StencilRef = Strata::StencilBit_Complex; GraphicsPSOInit.DepthStencilState = TStaticDepthStencilState<false, CompareFunction, true, CF_Equal, SO_Keep, SO_Keep, SO_Keep, true, CF_Equal, SO_Keep, SO_Keep, SO_Keep, 0x0, 0x0>::GetRHI(); break;
+		}
 	}
 	else
 	{
-		GraphicsPSOInit.DepthStencilState =
-			bCameraInsideLightGeometry
-			? TStaticDepthStencilState<false, CF_Always>::GetRHI()
-			: TStaticDepthStencilState<false, CF_DepthNearOrEqual>::GetRHI();
+		check(TileType == EStrataTileMaterialType::ECount);
+		GraphicsPSOInit.DepthStencilState = TStaticDepthStencilState<false, CompareFunction>::GetRHI();
 	}
+	return StencilRef;
+}
+
+/** Sets up rasterizer and depth state for rendering bounding geometry in a deferred pass. */
+static uint32 SetBoundingGeometryRasterizerAndDepthState(FGraphicsPipelineStateInitializer& GraphicsPSOInit, const FViewInfo& View, bool bCameraInsideLightGeometry, EStrataTileMaterialType TileType)
+{
+	uint32 StencilRef = 0u;
+	InternalSetBoundingGeometryRasterizerState(GraphicsPSOInit, View, bCameraInsideLightGeometry);
+	if (bCameraInsideLightGeometry)
+	{
+		StencilRef = InternalSetBoundingGeometryDepthState<CF_Always>(GraphicsPSOInit, TileType);
+	}
+	else
+	{
+		StencilRef = InternalSetBoundingGeometryDepthState<CF_DepthNearOrEqual>(GraphicsPSOInit, TileType);
+	}
+
+	return StencilRef;
 }
 
 // Use DBT to allow work culling on shadow lights
@@ -1824,24 +1839,15 @@ static void InternalRenderLight(
 	const TCHAR* ShaderName)
 {
 	const FLightSceneProxy* RESTRICT LightProxy = LightSceneInfo->Proxy;
-	const uint32 StencilRef = StrataTileMaterialType == EStrataTileMaterialType::ESimple ? Strata::StencilBit : 0u;
 	const bool bTransmission = LightProxy->Transmission();
 	const FSphere LightBounds = LightProxy->GetBoundingSphere();
 	const ELightComponentType LightType = (ELightComponentType)LightProxy->GetLightType();
 
-	const TCHAR* PermationName = nullptr;
-	switch (StrataTileMaterialType)
-	{
-	case EStrataTileMaterialType::ESimple : TEXT("(Strata,Simple)"); break;
-	case EStrataTileMaterialType::EComplex: TEXT("(Strata,Complex)"); break;
-	case EStrataTileMaterialType::ECount  : TEXT(""); break;
-	}
-
 	GraphBuilder.AddPass(
-		RDG_EVENT_NAME("%s%s", ShaderName, PermationName),
+		RDG_EVENT_NAME("%s", ShaderName),
 		PassParameters,
 		ERDGPassFlags::Raster,
-		[Scene, &View, PixelShader, LightSceneInfo, PassParameters, StencilRef, LightBounds, LightType, StrataTileMaterialType](FRHICommandList& RHICmdList)
+		[Scene, &View, PixelShader, LightSceneInfo, PassParameters, LightBounds, LightType, StrataTileMaterialType](FRHICommandList& RHICmdList)
 	{
 
 		const bool bIsRadial = LightType != LightType_Directional;
@@ -1856,19 +1862,6 @@ static void InternalRenderLight(
 		GraphicsPSOInit.BlendState = TStaticBlendState<CW_RGBA, BO_Add, BF_One, BF_One, BO_Add, BF_One, BF_One>::GetRHI();
 		GraphicsPSOInit.PrimitiveType = PT_TriangleList;
 
-		if (bEnableStrataStencilTest)
-		{
-			GraphicsPSOInit.DepthStencilState = TStaticDepthStencilState<
-				false, CF_Always,
-				true, CF_Equal, SO_Keep, SO_Keep, SO_Keep,
-				true, CF_Equal, SO_Keep, SO_Keep, SO_Keep,
-				Strata::StencilBit, 0x0>::GetRHI();
-		}
-		else
-		{
-			GraphicsPSOInit.DepthStencilState = TStaticDepthStencilState<false, CF_Always>::GetRHI();
-		}
-
 		if (LightType == LightType_Directional)
 		{
 			FDeferredLightVS::FPermutationDomain PermutationVectorVS;
@@ -1886,11 +1879,12 @@ static void InternalRenderLight(
 			GraphicsPSOInit.BoundShaderState.VertexDeclarationRHI = GFilterVertexDeclaration.VertexDeclarationRHI;
 			GraphicsPSOInit.BoundShaderState.VertexShaderRHI = bEnableStrataTiledPass ? TileVertexShader.GetVertexShader() : VertexShader.GetVertexShader();
 			GraphicsPSOInit.BoundShaderState.PixelShaderRHI = PixelShader.GetPixelShader();		
-			SetGraphicsPipelineState(RHICmdList, GraphicsPSOInit, StencilRef);
+			GraphicsPSOInit.DepthStencilState = TStaticDepthStencilState<false, CF_Always>::GetRHI();
+			SetGraphicsPipelineState(RHICmdList, GraphicsPSOInit, 0x0);
 
 			SetShaderParameters(RHICmdList, PixelShader, PixelShader.GetPixelShader(), PassParameters->PS);
 
-			if (bEnableStrataTiledPass)
+			if (Strata::IsStrataEnabled())
 			{
 				Strata::FStrataTilePassVS::FParameters VSParameters;
 				Strata::FillUpTiledPassData(StrataTileMaterialType, View, VSParameters, GraphicsPSOInit.PrimitiveType);
@@ -1932,12 +1926,11 @@ static void InternalRenderLight(
 				//@todo - accurate ortho camera / light intersection
 				|| !View.IsPerspectiveProjection();
 
-			SetBoundingGeometryRasterizerAndDepthState(GraphicsPSOInit, View, bCameraInsideLightGeometry);
-
+			const uint32 StencilRef = SetBoundingGeometryRasterizerAndDepthState(GraphicsPSOInit, View, bCameraInsideLightGeometry, StrataTileMaterialType);
 			GraphicsPSOInit.BoundShaderState.VertexDeclarationRHI = GetVertexDeclarationFVector4();
 			GraphicsPSOInit.BoundShaderState.VertexShaderRHI = VertexShader.GetVertexShader();
 			GraphicsPSOInit.BoundShaderState.PixelShaderRHI = PixelShader.GetPixelShader();
-			SetGraphicsPipelineState(RHICmdList, GraphicsPSOInit, StencilRef);			
+			SetGraphicsPipelineState(RHICmdList, GraphicsPSOInit, StencilRef);
 			SetShaderParameters(RHICmdList, PixelShader, PixelShader.GetPixelShader(), PassParameters->PS);
 
 			FDeferredLightVS::FParameters VSParameters2 = FDeferredLightVS::GetParameters(View, LightSceneInfo);
@@ -2050,7 +2043,8 @@ static void RenderLight(
 		PermutationVector.Set<FDeferredLightOverlapPS::FRadialAttenuation>(bIsRadial);
 		TShaderMapRef<FDeferredLightOverlapPS> PixelShader(View.ShaderMap, PermutationVector);
 		
-		InternalRenderLight(GraphBuilder, Scene, View, LightSceneInfo, PixelShader, PassParameters, EStrataTileMaterialType::ECount, TEXT("StandardDeferredOverlap"));
+		// STRATA_TODO: fix overlap debug mode for strata (simple/complex tile rendering)
+		InternalRenderLight(GraphBuilder, Scene, View, LightSceneInfo, PixelShader, PassParameters, EStrataTileMaterialType::ECount, TEXT("Light::StandardDeferred(Overlap)"));
 	}
 	// Lighting shader
 	else
@@ -2068,7 +2062,7 @@ static void RenderLight(
 			PassParameters->VS = FDeferredLightVS::GetParameters(View, false);
 		}
 		// VS - Strata tile parameters
-		if (Strata::ShouldPassesReadingStrataBeTiled(View.Family->GetFeatureLevel()))
+		if (Strata::IsStrataEnabled())
 		{
 			PassParameters->TileListBufferSimple = View.StrataSceneData->ClassificationTileListBufferSRV[EStrataTileMaterialType::ESimple];
 			PassParameters->TileListBufferComplex = View.StrataSceneData->ClassificationTileListBufferSRV[EStrataTileMaterialType::EComplex];
@@ -2115,26 +2109,26 @@ static void RenderLight(
 		// * if the light is directional, then dispatch a set of rect tiles
 		// * if the light is radial/local, then dispatch a light geometry with stencil test. The stencil buffer has been prefilled with the tile result (simple/complex) 
 		//   so that the geometry get correctly stencil culled on complex/simple part of the screen
-		if (Strata::IsStrataEnabled() && Strata::IsClassificationEnabled())
+		if (Strata::IsStrataEnabled())
 		{
 			// Simple tiles
 			{
 				PermutationVector.Set< FDeferredLightPS::FStrataFastPath>(true);
 				TShaderMapRef< FDeferredLightPS > PixelShader(View.ShaderMap, PermutationVector);
-				InternalRenderLight(GraphBuilder, Scene, View, LightSceneInfo, PixelShader, PassParameters, EStrataTileMaterialType::ESimple, TEXT("StandardDeferredLight"));
+				InternalRenderLight(GraphBuilder, Scene, View, LightSceneInfo, PixelShader, PassParameters, EStrataTileMaterialType::ESimple, TEXT("Light::StandardDeferred(Simple)"));
 			}
 			// Complex tiles
 			{
 				PermutationVector.Set< FDeferredLightPS::FStrataFastPath>(false);
 				TShaderMapRef< FDeferredLightPS > PixelShader(View.ShaderMap, PermutationVector);
-				InternalRenderLight(GraphBuilder, Scene, View, LightSceneInfo, PixelShader, PassParameters, EStrataTileMaterialType::EComplex, TEXT("StandardDeferredLight"));
+				InternalRenderLight(GraphBuilder, Scene, View, LightSceneInfo, PixelShader, PassParameters, EStrataTileMaterialType::EComplex, TEXT("Light::StandardDeferred(Complex)"));
 			}
 		}
 		else
 		{
 			PermutationVector.Set< FDeferredLightPS::FStrataFastPath>(false);
 			TShaderMapRef< FDeferredLightPS > PixelShader(View.ShaderMap, PermutationVector);
-			InternalRenderLight(GraphBuilder, Scene, View, LightSceneInfo, PixelShader, PassParameters, EStrataTileMaterialType::ECount, TEXT("StandardDeferredLight"));
+			InternalRenderLight(GraphBuilder, Scene, View, LightSceneInfo, PixelShader, PassParameters, EStrataTileMaterialType::ECount, TEXT("Light::StandardDeferred"));
 		}
 	}
 }
@@ -2417,11 +2411,14 @@ void FDeferredShadingSceneRenderer::RenderSimpleLightsStandardDeferred(
 		PermutationVectorVS.Set<FDeferredLightVS::FRadialLight>(true);
 		TShaderMapRef<FDeferredLightVS> VertexShader(View.ShaderMap, PermutationVectorVS);
 
+		// STRATA_TODO: add simple/complex tile support for simple lights
+		const EStrataTileMaterialType TileType = Strata::IsStrataEnabled() ? EStrataTileMaterialType::EComplex : EStrataTileMaterialType::ECount;
+
 		GraphBuilder.AddPass(
 			RDG_EVENT_NAME("StandardDeferredSimpleLights"),
 			PassParameters,
 			ERDGPassFlags::Raster,
-			[this, &View, &SimpleLights, ViewIndex, NumViews, PassParameters, PixelShader, VertexShader](FRHICommandList& RHICmdList)
+			[this, &View, &SimpleLights, ViewIndex, NumViews, PassParameters, PixelShader, VertexShader, TileType](FRHICommandList& RHICmdList)
 		{
 			FGraphicsPipelineStateInitializer GraphicsPSOInit;
 			RHICmdList.ApplyCachedRenderTargets(GraphicsPSOInit);
@@ -2446,7 +2443,7 @@ void FDeferredShadingSceneRenderer::RenderSimpleLightsStandardDeferred(
 								//@todo - accurate ortho camera / light intersection
 								|| !View.IsPerspectiveProjection();
 
-				SetBoundingGeometryRasterizerAndDepthState(GraphicsPSOInit, View, bCameraInsideLightGeometry);
+				SetBoundingGeometryRasterizerAndDepthState(GraphicsPSOInit, View, bCameraInsideLightGeometry, TileType);
 				GraphicsPSOInit.BoundShaderState.VertexDeclarationRHI = GetVertexDeclarationFVector4();
 				GraphicsPSOInit.BoundShaderState.VertexShaderRHI = VertexShader.GetVertexShader();
 				GraphicsPSOInit.BoundShaderState.PixelShaderRHI = PixelShader.GetPixelShader();
@@ -2464,7 +2461,7 @@ void FDeferredShadingSceneRenderer::RenderSimpleLightsStandardDeferred(
 
 				// Apply the point or spot light with some approximately bounding geometry,
 				// So we can get speedups from depth testing and not processing pixels outside of the light's influence.
-				StencilingGeometry::DrawSphere(RHICmdList);				
+				StencilingGeometry::DrawSphere(RHICmdList);
 			}
 		});
 	}
@@ -2544,4 +2541,4 @@ FRDGTextureRef FDeferredShadingSceneRenderer::CopyStencilToLightingChannelTextur
 	}
 
 	return LightingChannelsTexture;
-}
\ No newline at end of file
+}
diff --git a/Engine/Source/Runtime/Renderer/Private/LightRendering.h b/Engine/Source/Runtime/Renderer/Private/LightRendering.h
--- a/Engine/Source/Runtime/Renderer/Private/LightRendering.h
+++ b/Engine/Source/Runtime/Renderer/Private/LightRendering.h
@@ -595,8 +595,6 @@ public:
 	}
 };
 
-void SetBoundingGeometryRasterizerAndDepthState(FGraphicsPipelineStateInitializer& GraphicsPSOInit, const FViewInfo& View, bool bCameraInsideLightGeometry);
-
 enum class FLightOcclusionType : uint8
 {
 	Shadowmap,
diff --git a/Engine/Source/Runtime/Renderer/Private/Strata/Strata.cpp b/Engine/Source/Runtime/Renderer/Private/Strata/Strata.cpp
--- a/Engine/Source/Runtime/Renderer/Private/Strata/Strata.cpp
+++ b/Engine/Source/Runtime/Renderer/Private/Strata/Strata.cpp
@@ -34,11 +34,6 @@ static TAutoConsoleVariable<int32> CVarStrataBytePerPixel(
 	TEXT("Strata allocated byte per pixel to store materials data. Higher value means more complex material can be represented."),
 	ECVF_ReadOnly | ECVF_RenderThreadSafe);
 
-static TAutoConsoleVariable<int32> CVarStrataClassification(
-	TEXT("r.Strata.Classification"),
-	1,
-	TEXT("Enable strata classification to speed up lighting pass."),
-	ECVF_RenderThreadSafe); 
 
 static TAutoConsoleVariable<int32> CVarStrataClassificationDebug(
 	TEXT("r.Strata.Classification.Debug"),
@@ -46,12 +41,6 @@ static TAutoConsoleVariable<int32> CVarStrataClassificationDebug(
 	TEXT("Enable strata classification visualization: 1 shows simple material tiles in green and complex material tiles in red."),
 	ECVF_RenderThreadSafe);
 
-static TAutoConsoleVariable<int32> CVarStrataClassificationPassesReadingStrataAreTiled(
-	TEXT("r.Strata.Classification.PassesReadingStrataAreTiled"),
-	1,
-	TEXT("Enable the tiling of passes reading strata material (when possible) instead of doing multiple full screen passes testing stencil."),
-	ECVF_RenderThreadSafe);
-
 static TAutoConsoleVariable<int32> CVarStrataRoughDiffuse(
 	TEXT("r.Strata.RoughDiffuse"),
 	1,
@@ -130,16 +119,6 @@ bool IsStrataEnabled()
 	return CVarStrata.GetValueOnAnyThread() > 0;
 }
 
-bool IsClassificationEnabled()
-{
-	return CVarStrataClassification.GetValueOnAnyThread() > 0;
-}
-
-bool ShouldPassesReadingStrataBeTiled(ERHIFeatureLevel::Type FeatureLevel)
-{
-	return IsStrataEnabled() && IsClassificationEnabled() && FeatureLevel >= ERHIFeatureLevel::SM5 && CVarStrataClassificationPassesReadingStrataAreTiled.GetValueOnAnyThread() > 0;
-}
-
 uint32 GetStrataBufferTileSize()
 {
 	return 8;
@@ -176,15 +155,21 @@ void InitialiseStrataFrameSceneData(FSceneRenderer& SceneRenderer, FRDGBuilder&
 			const int32 TileInPixel = GetStrataBufferTileSize();
 			const FIntPoint TileResolution(FMath::DivideAndRoundUp(SceneTextureExtent.X, TileInPixel), FMath::DivideAndRoundUp(SceneTextureExtent.Y, TileInPixel));
 
+			const TCHAR* StrataTileResourceNames[EStrataTileMaterialType::ECount][2] =
+			{
+				{ TEXT("Strata.StrataTileListBuffer(Simple)"),	TEXT("Strata.StrataTileIndirectBuffer(Simple)") },
+				{ TEXT("Strata.StrataTileListBuffer(Complex)"),	TEXT("Strata.StrataTileIndirectBuffer(Complex)")}
+			};
+
 			// As of today we allocate one index+indirect buffer for each EStrataTileMaterialType.
 			// This is fine for two types, later we might want to have a single list and indirect buffer with offsets.
 			for (uint32 i = 0; i < EStrataTileMaterialType::ECount; ++i)
 			{
-				StrataSceneData.ClassificationTileListBuffer[i] = GraphBuilder.CreateBuffer(FRDGBufferDesc::CreateBufferDesc(sizeof(uint32), TileResolution.X * TileResolution.Y), i==EStrataTileMaterialType::ESimple ? TEXT("Strata.SimpleStrataTileListBuffer") : TEXT("Strata.ComplexStrataTileListBuffer"));
+				StrataSceneData.ClassificationTileListBuffer[i] = GraphBuilder.CreateBuffer(FRDGBufferDesc::CreateBufferDesc(sizeof(uint32), TileResolution.X * TileResolution.Y), StrataTileResourceNames[i][0]);
 				StrataSceneData.ClassificationTileListBufferSRV[i] = GraphBuilder.CreateSRV(StrataSceneData.ClassificationTileListBuffer[i], PF_R32_UINT);
 				StrataSceneData.ClassificationTileListBufferUAV[i] = GraphBuilder.CreateUAV(StrataSceneData.ClassificationTileListBuffer[i], PF_R32_UINT);
 
-				StrataSceneData.ClassificationTileIndirectBuffer[i] = GraphBuilder.CreateBuffer(FRDGBufferDesc::CreateIndirectDesc<FRHIDrawIndirectParameters>(), i == EStrataTileMaterialType::ESimple ? TEXT("Strata.SimpleStrataTileIndirectBuffer") : TEXT("Strata.ComplexStrataTileIndirectBuffer"));
+				StrataSceneData.ClassificationTileIndirectBuffer[i] = GraphBuilder.CreateBuffer(FRDGBufferDesc::CreateIndirectDesc<FRHIDrawIndirectParameters>(), StrataTileResourceNames[i][1]);
 				StrataSceneData.ClassificationTileIndirectBufferSRV[i] = GraphBuilder.CreateSRV(StrataSceneData.ClassificationTileIndirectBuffer[i], PF_R32_UINT);
 				StrataSceneData.ClassificationTileIndirectBufferUAV[i] = GraphBuilder.CreateUAV(StrataSceneData.ClassificationTileIndirectBuffer[i], PF_R32_UINT);
 
@@ -605,12 +590,12 @@ static void AddStrataInternalClassificationTilePass(
 					false, CF_Always, 
 					true,  CF_Always, SO_Keep, SO_Keep, SO_Replace,
 					false, CF_Always, SO_Keep, SO_Keep, SO_Keep,
-					0xFF, StencilBit>::GetRHI();
+					0xFF, StencilBit_Fast>::GetRHI();
 			}
 			GraphicsPSOInit.BoundShaderState.VertexDeclarationRHI = GFilterVertexDeclaration.VertexDeclarationRHI;
 			GraphicsPSOInit.BoundShaderState.VertexShaderRHI = VertexShader.GetVertexShader();
 			GraphicsPSOInit.PrimitiveType = StrataTilePrimitiveType;
-			SetGraphicsPipelineState(RHICmdList, GraphicsPSOInit, StencilBit);
+			SetGraphicsPipelineState(RHICmdList, GraphicsPSOInit, StencilBit_Fast);
 			SetShaderParameters(RHICmdList, VertexShader, VertexShader.GetVertexShader(), ParametersPS->VS);
 			if (bDebug)
 			{
@@ -755,7 +740,6 @@ void AddStrataMaterialClassificationPass(FRDGBuilder& GraphBuilder, const FMinim
 		const FViewInfo& View = Views[i];
 		
 		// Tile reduction
-		if (IsClassificationEnabled())
 		{
 			bool bWaveOps = GRHISupportsWaveOperations && FDataDrivenShaderPlatformInfo::GetSupportsWaveOperations(View.GetShaderPlatform());
 		#if PLATFORM_WINDOWS
@@ -825,7 +809,7 @@ bool ShouldRenderStrataDebugPasses(const FViewInfo& View)
 	return IsStrataEnabled() &&
 		(
 		   (FVisualizeMaterialPS::CanRunStrataVizualizeMaterial(View.GetShaderPlatform()) && View.Family && View.Family->EngineShowFlags.VisualizeStrataMaterial)
-		|| (IsClassificationEnabled() && CVarStrataClassificationDebug.GetValueOnAnyThread() > 0)
+		|| (CVarStrataClassificationDebug.GetValueOnAnyThread() > 0)
 		);
 }
 
@@ -841,7 +825,7 @@ FScreenPassTexture AddStrataDebugPasses(FRDGBuilder& GraphBuilder, const FViewIn
 	}
 
 	const int32 StrataClassificationDebug = CVarStrataClassificationDebug.GetValueOnAnyThread();
-	if (IsClassificationEnabled() && StrataClassificationDebug > 0)
+	if (StrataClassificationDebug > 0)
 	{
 		RDG_EVENT_SCOPE(GraphBuilder, "Strata::VisualizeClassification");
 		const bool bDebugPass = true;
diff --git a/Engine/Source/Runtime/Renderer/Private/Strata/Strata.h b/Engine/Source/Runtime/Renderer/Private/Strata/Strata.h
--- a/Engine/Source/Runtime/Renderer/Private/Strata/Strata.h
+++ b/Engine/Source/Runtime/Renderer/Private/Strata/Strata.h
@@ -83,11 +83,10 @@ struct FStrataSceneData
 
 namespace Strata
 {
-constexpr uint32 StencilBit = 0x80; // In sync with SceneRenderTargets.h - GET_STENCIL_BIT_MASK(STENCIL_STRATA_FASTPATH)
+constexpr uint32 StencilBit_Complex= 0x00; // No stencil bit
+constexpr uint32 StencilBit_Fast   = 0x80; // In sync with SceneRenderTargets.h - GET_STENCIL_BIT_MASK(STENCIL_STRATA_FASTPATH)
 
 bool IsStrataEnabled();
-bool IsClassificationEnabled();
-bool ShouldPassesReadingStrataBeTiled(ERHIFeatureLevel::Type FeatureLevel);
 
 void InitialiseStrataFrameSceneData(FSceneRenderer& SceneRenderer, FRDGBuilder& GraphBuilder);
 
