Remove tiled deferred rendering path

diff --git a/Engine/Shaders/Private/ClusteredDeferredShadingPixelShader.usf b/Engine/Shaders/Private/ClusteredDeferredShadingPixelShader.usf
--- a/Engine/Shaders/Private/ClusteredDeferredShadingPixelShader.usf
+++ b/Engine/Shaders/Private/ClusteredDeferredShadingPixelShader.usf
@@ -248,9 +248,6 @@ void ClusteredShadingPixelShader(
 
 	uint PixelShadingModelID = GetScreenSpaceData(ScreenUV).GBuffer.ShadingModelID;
 
-	// TODO: This is how it was in TiledDeferred
-	//const float Dither = 0.5f;
-	// BUT: this is how it is in trad. deferred, assumed better:
 	const float Dither = InterleavedGradientNoise(SvPosition.xy, View.StateFrameIndexMod8);
 
 	// ?
@@ -412,9 +409,6 @@ void ClusteredShadingPixelShader(
 		ScreenSpaceData.GBuffer.PerObjectGBufferData = 0;
 	}
 
-	// TODO: This is how it was in TiledDeferred
-	//const float Dither = 0.5f;
-	// BUT: this is how it is in trad. deferred, assumed better:
 	const float Dither = InterleavedGradientNoise(SvPosition.xy, View.StateFrameIndexMod8);
 
 	const uint GridIndex = ComputeLightGridCellIndex(uint2(LocalPosition.x, LocalPosition.y), SceneDepth);
diff --git a/Engine/Shaders/Private/TiledDeferredLightShaders.usf b/Engine/Shaders/Private/TiledDeferredLightShaders.usf
deleted file mode 100644
diff --git a/Engine/Source/Runtime/Engine/Private/Components/LightComponent.cpp b/Engine/Source/Runtime/Engine/Private/Components/LightComponent.cpp
--- a/Engine/Source/Runtime/Engine/Private/Components/LightComponent.cpp
+++ b/Engine/Source/Runtime/Engine/Private/Components/LightComponent.cpp
@@ -346,7 +346,6 @@ FLightSceneProxy::FLightSceneProxy(const ULightComponent* InLightComponent)
 	, bUseVirtualShadowMaps(false)	// See below
 	, bCastModulatedShadows(false)
 	, bUseWholeSceneCSMForMovableObjects(false)
-	, bTiledDeferredLightingSupported(false)
 	, AtmosphereSunLightIndex(InLightComponent->GetAtmosphereSunLightIndex())
 	, AtmosphereSunDiskColorScale(InLightComponent->GetAtmosphereSunDiskColorScale())
 	, LightType(InLightComponent->GetLightType())	
diff --git a/Engine/Source/Runtime/Engine/Private/PointLightSceneProxy.h b/Engine/Source/Runtime/Engine/Private/PointLightSceneProxy.h
--- a/Engine/Source/Runtime/Engine/Private/PointLightSceneProxy.h
+++ b/Engine/Source/Runtime/Engine/Private/PointLightSceneProxy.h
@@ -39,9 +39,6 @@ public:
 	,	bInverseSquared(Component->bUseInverseSquaredFalloff)
 	{
 		UpdateRadius(Component->AttenuationRadius);
-
-		// tiled deferred only supported for point/spot lights with 0 length
-		bTiledDeferredLightingSupported = (SourceLength == 0.0f);
 	}
 
 	virtual float GetSourceRadius() const override
diff --git a/Engine/Source/Runtime/Engine/Public/SceneManagement.h b/Engine/Source/Runtime/Engine/Public/SceneManagement.h
--- a/Engine/Source/Runtime/Engine/Public/SceneManagement.h
+++ b/Engine/Source/Runtime/Engine/Public/SceneManagement.h
@@ -1525,7 +1525,6 @@ public:
 	inline bool UseRayTracedDistanceFieldShadows() const { return bUseRayTracedDistanceFieldShadows; }
 	inline bool UseVirtualShadowMaps() const { return bUseVirtualShadowMaps; }
 	inline float GetRayStartOffsetDepthScale() const { return RayStartOffsetDepthScale; }
-	inline bool IsTiledDeferredLightingSupported() const { return bTiledDeferredLightingSupported;  }
 	inline uint8 GetLightType() const { return LightType; }
 	inline uint8 GetLightingChannelMask() const { return LightingChannelMask; }
 	inline FName GetComponentFName() const { return ComponentName; }
@@ -1755,9 +1754,6 @@ protected:
 	/** Whether to render csm shadows for movable objects only (mobile). */
 	uint8 bUseWholeSceneCSMForMovableObjects : 1;
 
-	/** Whether the light supports rendering in tiled deferred pass */
-	uint8 bTiledDeferredLightingSupported : 1;
-
 	/** The index of the atmospheric light. Multiple lights can be considered when computing the sky/atmospheric scattering. */
 	const uint8 AtmosphereSunLightIndex;
 
diff --git a/Engine/Source/Runtime/RenderCore/Private/RenderCore.cpp b/Engine/Source/Runtime/RenderCore/Private/RenderCore.cpp
--- a/Engine/Source/Runtime/RenderCore/Private/RenderCore.cpp
+++ b/Engine/Source/Runtime/RenderCore/Private/RenderCore.cpp
@@ -158,8 +158,6 @@ DEFINE_STAT(STAT_NumShadowedLights);
 DEFINE_STAT(STAT_NumLightFunctionOnlyLights);
 DEFINE_STAT(STAT_NumBatchedLights);
 DEFINE_STAT(STAT_NumLightsInjectedIntoTranslucency);
-DEFINE_STAT(STAT_NumLightsUsingTiledDeferred);
-DEFINE_STAT(STAT_NumLightsUsingSimpleTiledDeferred);
 DEFINE_STAT(STAT_NumLightsUsingStandardDeferred);
 
 DEFINE_STAT(STAT_LightShaftsLights);
diff --git a/Engine/Source/Runtime/RenderCore/Public/RenderCore.h b/Engine/Source/Runtime/RenderCore/Public/RenderCore.h
--- a/Engine/Source/Runtime/RenderCore/Public/RenderCore.h
+++ b/Engine/Source/Runtime/RenderCore/Public/RenderCore.h
@@ -105,8 +105,6 @@ DECLARE_CYCLE_STAT_EXTERN(TEXT("Render Lights and Shadows"),STAT_LightRendering,
 DECLARE_CYCLE_STAT_EXTERN(TEXT("Direct lighting"),STAT_DirectLightRenderingTime,STATGROUP_LightRendering, RENDERCORE_API);
 DECLARE_CYCLE_STAT_EXTERN(TEXT("Translucent injection"),STAT_TranslucentInjectTime,STATGROUP_LightRendering, RENDERCORE_API);
 DECLARE_DWORD_COUNTER_STAT_EXTERN(TEXT("Num using standard deferred"),STAT_NumLightsUsingStandardDeferred,STATGROUP_LightRendering, RENDERCORE_API);
-DECLARE_DWORD_COUNTER_STAT_EXTERN(TEXT("Num using tiled deferred"),STAT_NumLightsUsingTiledDeferred,STATGROUP_LightRendering, RENDERCORE_API);
-DECLARE_DWORD_COUNTER_STAT_EXTERN(TEXT("Num using simple lighting model"),STAT_NumLightsUsingSimpleTiledDeferred,STATGROUP_LightRendering, RENDERCORE_API);
 DECLARE_DWORD_COUNTER_STAT_EXTERN(TEXT("Num injected into translucency"),STAT_NumLightsInjectedIntoTranslucency,STATGROUP_LightRendering, RENDERCORE_API);
 DECLARE_DWORD_COUNTER_STAT_EXTERN(TEXT("Num batched"),STAT_NumBatchedLights,STATGROUP_LightRendering, RENDERCORE_API);
 DECLARE_DWORD_COUNTER_STAT_EXTERN(TEXT("Num light function only"),STAT_NumLightFunctionOnlyLights,STATGROUP_LightRendering, RENDERCORE_API);
diff --git a/Engine/Source/Runtime/Renderer/Private/ClusteredDeferredShadingPass.cpp b/Engine/Source/Runtime/Renderer/Private/ClusteredDeferredShadingPass.cpp
--- a/Engine/Source/Runtime/Renderer/Private/ClusteredDeferredShadingPass.cpp
+++ b/Engine/Source/Runtime/Renderer/Private/ClusteredDeferredShadingPass.cpp
@@ -1,7 +1,7 @@
 // Copyright Epic Games, Inc. All Rights Reserved.
 
 /*=============================================================================
-	TiledDeferredLightRendering.cpp: Implementation of tiled deferred shading
+	ClusteredDeferredShadingPass.cpp: Implementation of tiled deferred shading
 =============================================================================*/
 
 #include "CoreMinimal.h"
diff --git a/Engine/Source/Runtime/Renderer/Private/DeferredShadingRenderer.h b/Engine/Source/Runtime/Renderer/Private/DeferredShadingRenderer.h
--- a/Engine/Source/Runtime/Renderer/Private/DeferredShadingRenderer.h
+++ b/Engine/Source/Runtime/Renderer/Private/DeferredShadingRenderer.h
@@ -545,16 +545,9 @@ private:
 	/** Mark time line for gathering Lumen virtual surface cache feedback. */
 	void BeginGatheringLumenSurfaceCacheFeedback(FRDGBuilder& GraphBuilder, const FViewInfo& View, FLumenSceneFrameTemporaries& FrameTemporaries);
 	void FinishGatheringLumenSurfaceCacheFeedback(FRDGBuilder& GraphBuilder, const FViewInfo& View, FLumenSceneFrameTemporaries& FrameTemporaries);
-	 
-	/** Whether tiled deferred is supported and can be used at all. */
-	bool CanUseTiledDeferred() const;
-
-	/** Whether to use tiled deferred shading given a number of lights that support it. */
-	bool ShouldUseTiledDeferred(int32 NumTiledDeferredLights) const;
-
+	
 	/** 
 	 * True if the 'r.UseClusteredDeferredShading' flag is 1 and sufficient feature level. 
-	 * NOTE: When true it takes precedence over the TiledDeferred path, since they handle the same lights.
 	 */
 	bool ShouldUseClusteredDeferredShading() const;
 
@@ -572,15 +565,6 @@ private:
 		FRDGTextureRef ShadowMaskBits,
 		FRDGTextureRef HairStrandsShadowMaskBits);
 
-	/** Renders the lights in SortedLights in the range [TiledDeferredLightsStart, TiledDeferredLightsEnd) using tiled deferred shading. */
-	void RenderTiledDeferredLighting(
-		FRDGBuilder& GraphBuilder,
-		FMinimalSceneTextures& SceneTextures,
-		const TArray<FSortedLightSceneInfo, SceneRenderingAllocator>& SortedLights,
-		int32 TiledDeferredLightsStart,
-		int32 TiledDeferredLightsEnd,
-		const FSimpleLightArray& SimpleLights);
-
 	/** Renders the scene's lighting. */
 	void RenderLights(
 		FRDGBuilder& GraphBuilder,
diff --git a/Engine/Source/Runtime/Renderer/Private/LightRendering.cpp b/Engine/Source/Runtime/Renderer/Private/LightRendering.cpp
--- a/Engine/Source/Runtime/Renderer/Private/LightRendering.cpp
+++ b/Engine/Source/Runtime/Renderer/Private/LightRendering.cpp
@@ -717,7 +717,7 @@ void FSceneRenderer::GatherAndSortLights(FSortedLightSetSceneInfo& OutSortedLigh
 					// In the forward case one directional light gets put into its own variables, and in the deferred case it gets a full-screen pass.
 					// Usually it'll have shadows and stuff anyway.
 					// Rect lights are not supported as the performance impact is significant even if not used, for now, left for trad. deferred.
-					const bool bTiledOrClusteredDeferredSupported =
+					const bool bClusteredDeferredSupported =
 						!SortedLightInfo->SortKey.Fields.bTextureProfile &&
 						(!SortedLightInfo->SortKey.Fields.bShadowed || bShadowedLightsInClustered) &&
 						!SortedLightInfo->SortKey.Fields.bLightFunction &&
@@ -728,11 +728,8 @@ void FSceneRenderer::GatherAndSortLights(FSortedLightSetSceneInfo& OutSortedLigh
 					// One pass projection is supported for lights with only virtual shadow maps
 					// TODO: Exclude lights that also have non-virtual shadow maps
 					bool bHasVirtualShadowMap = VisibleLightInfos[LightSceneInfo->Id].GetVirtualShadowMapId(&Views[ViewIndex]) != INDEX_NONE;
-					SortedLightInfo->SortKey.Fields.bDoesNotWriteIntoPackedShadowMask = !bTiledOrClusteredDeferredSupported || !bHasVirtualShadowMap;
-
-					SortedLightInfo->SortKey.Fields.bTiledDeferredNotSupported = !(bTiledOrClusteredDeferredSupported && LightSceneInfo->Proxy->IsTiledDeferredLightingSupported());
-
-					SortedLightInfo->SortKey.Fields.bClusteredDeferredNotSupported = !bTiledOrClusteredDeferredSupported;
+					SortedLightInfo->SortKey.Fields.bDoesNotWriteIntoPackedShadowMask = !bClusteredDeferredSupported || !bHasVirtualShadowMap;
+					SortedLightInfo->SortKey.Fields.bClusteredDeferredNotSupported = !bClusteredDeferredSupported;
 					break;
 				}
 			}
@@ -763,7 +760,6 @@ void FSceneRenderer::GatherAndSortLights(FSortedLightSetSceneInfo& OutSortedLigh
 		SortedLightInfo->SortKey.Fields.bIsNotSimpleLight = 0;
 
 		// Simple lights are ok to use with tiled and clustered deferred lighting
-		SortedLightInfo->SortKey.Fields.bTiledDeferredNotSupported = 0;
 		SortedLightInfo->SortKey.Fields.bClusteredDeferredNotSupported = 0;
 	}
 
@@ -779,7 +775,6 @@ void FSceneRenderer::GatherAndSortLights(FSortedLightSetSceneInfo& OutSortedLigh
 
 	// Scan and find ranges.
 	OutSortedLights.SimpleLightsEnd = SortedLights.Num();
-	OutSortedLights.TiledSupportedEnd = SortedLights.Num();
 	OutSortedLights.ClusteredSupportedEnd = SortedLights.Num();
 	OutSortedLights.UnbatchedLightStart = SortedLights.Num();
 
@@ -798,12 +793,6 @@ void FSceneRenderer::GatherAndSortLights(FSortedLightSetSceneInfo& OutSortedLigh
 			OutSortedLights.SimpleLightsEnd = LightIndex;
 		}
 
-		if (SortedLightInfo.SortKey.Fields.bTiledDeferredNotSupported && OutSortedLights.TiledSupportedEnd == SortedLights.Num())
-		{
-			// Mark the first index to not support tiled deferred
-			OutSortedLights.TiledSupportedEnd = LightIndex;
-		}
-
 		if (SortedLightInfo.SortKey.Fields.bClusteredDeferredNotSupported && OutSortedLights.ClusteredSupportedEnd == SortedLights.Num())
 		{
 			// Mark the first index to not support clustered deferred
@@ -813,15 +802,14 @@ void FSceneRenderer::GatherAndSortLights(FSortedLightSetSceneInfo& OutSortedLigh
 		if( (bDrawShadows || bDrawLightFunction || bLightingChannels) && SortedLightInfo.SortKey.Fields.bClusteredDeferredNotSupported )
 		{
 			// Once we find an unbatched shadowed light, we can exit the loop
-			check(SortedLightInfo.SortKey.Fields.bTiledDeferredNotSupported);
+			check(SortedLightInfo.SortKey.Fields.bClusteredDeferredNotSupported);
 			OutSortedLights.UnbatchedLightStart = LightIndex;
 			break;
 		}
 	}
 
 	// Make sure no obvious things went wrong!
-	check(OutSortedLights.TiledSupportedEnd >= OutSortedLights.SimpleLightsEnd);
-	check(OutSortedLights.ClusteredSupportedEnd >= OutSortedLights.TiledSupportedEnd);
+	check(OutSortedLights.ClusteredSupportedEnd >= OutSortedLights.SimpleLightsEnd);
 	check(OutSortedLights.UnbatchedLightStart >= OutSortedLights.ClusteredSupportedEnd);
 }
 
@@ -918,28 +906,6 @@ void FDeferredShadingSceneRenderer::RenderLights(
 
 				AddClusteredDeferredShadingPass(GraphBuilder, SceneTextures, SortedLightSet, ShadowMaskBits, HairStrandsShadowMaskBits);
 			}
-			else if (CanUseTiledDeferred())
-			{
-				bool bAnyViewIsStereo = false;
-				for (int32 ViewIndex = 0; ViewIndex < Views.Num(); ++ViewIndex)
-				{
-					if (IStereoRendering::IsStereoEyeView(Views[ViewIndex]))
-					{
-						bAnyViewIsStereo = true;
-						break;
-					}
-				}
-
-				// Use tiled deferred shading on any unshadowed lights without a texture light profile
-				if (ShouldUseTiledDeferred(SortedLightSet.TiledSupportedEnd) && !bAnyViewIsStereo)
-				{
-					// Update the range that needs to be processed by standard deferred to exclude the lights done with tiled
-					StandardDeferredStart = SortedLightSet.TiledSupportedEnd;
-					bRenderSimpleLightsStandardDeferred = false;
-
-					RenderTiledDeferredLighting(GraphBuilder, SceneTextures, SortedLights, SortedLightSet.SimpleLightsEnd, SortedLightSet.TiledSupportedEnd, SimpleLights);
-				}
-			}
 
 			if (bRenderSimpleLightsStandardDeferred)
 			{
diff --git a/Engine/Source/Runtime/Renderer/Private/LightSceneInfo.h b/Engine/Source/Runtime/Renderer/Private/LightSceneInfo.h
--- a/Engine/Source/Runtime/Renderer/Private/LightSceneInfo.h
+++ b/Engine/Source/Runtime/Renderer/Private/LightSceneInfo.h
@@ -85,8 +85,6 @@ struct FSortedLightSceneInfo
 			uint32 bShadowed : 1;
 			/** Whether the light is NOT a simple light - they always support tiled/clustered but may want to be selected separately. */
 			uint32 bIsNotSimpleLight : 1;
-			/** True if the light doesn't support tiled deferred, logic is inverted so that lights that DO support tiled deferred will sort first in list */
-			uint32 bTiledDeferredNotSupported : 1;
 			/* We want to sort the lights that write into the packed shadow mask (when enabled) to the front of the list so we don't waste slots in the packed shadow mask. */
 			uint32 bDoesNotWriteIntoPackedShadowMask : 1;
 			/** 
@@ -123,13 +121,12 @@ struct FSortedLightSceneInfo
 /** 
  * Stores info about sorted lights and ranges. 
  * The sort-key in FSortedLightSceneInfo gives rise to the following order:
- *  [SimpleLights,Tiled/Clustered,UnbatchedLights]
+ *  [SimpleLights,Clustered,UnbatchedLights]
  * Note that some shadowed lights can be included in the clustered pass when virtual shadow maps and one pass projection are used.
  */
 struct FSortedLightSetSceneInfo
 {
 	int SimpleLightsEnd;
-	int TiledSupportedEnd;
 	int ClusteredSupportedEnd;
 
 	/** First light with shadow map or */
diff --git a/Engine/Source/Runtime/Renderer/Private/TiledDeferredLightRendering.cpp b/Engine/Source/Runtime/Renderer/Private/TiledDeferredLightRendering.cpp
deleted file mode 100644
