Cleanup FHistogramCS to not use template.

diff --git a/Engine/Source/Runtime/Renderer/Private/PostProcess/PostProcessHistogram.cpp b/Engine/Source/Runtime/Renderer/Private/PostProcess/PostProcessHistogram.cpp
--- a/Engine/Source/Runtime/Renderer/Private/PostProcess/PostProcessHistogram.cpp
+++ b/Engine/Source/Runtime/Renderer/Private/PostProcess/PostProcessHistogram.cpp
@@ -15,23 +15,18 @@ TAutoConsoleVariable<int32> CVarUseAtomicHistogram(
 
 namespace
 {
-
-template <uint32 ThreadGroupSizeX, uint32 ThreadGroupSizeY, uint32 HistogramSize, bool BilateralGrid>
-class THistogramCS : public FGlobalShader
+class FHistogramCS : public FGlobalShader
 {
 public:
 	// Changing these numbers requires Histogram.usf to be recompiled.
 	static const uint32 LoopCountX = 8;
 	static const uint32 LoopCountY = 8;
 
-	// /4 as we store 4 buckets in one ARGB texel.
-	static const uint32 HistogramTexelCount = HistogramSize / 4;
+	// we store 4 buckets in one ARGB texel.
+	static const uint32 HistogramBucketsPerTexel = 4;
 
-	// The number of texels on each axis processed by a single thread group.
-	static const FIntPoint TexelsPerThreadGroup;
-
-	DECLARE_GLOBAL_SHADER(THistogramCS);
-	SHADER_USE_PARAMETER_STRUCT(THistogramCS, FGlobalShader);
+	DECLARE_GLOBAL_SHADER(FHistogramCS);
+	SHADER_USE_PARAMETER_STRUCT(FHistogramCS, FGlobalShader);
 
 	BEGIN_SHADER_PARAMETER_STRUCT(FParameters, )
 		SHADER_PARAMETER_STRUCT_REF(FViewUniformShaderParameters, View)
@@ -44,6 +39,9 @@ public:
 		SHADER_PARAMETER(FIntPoint, ThreadGroupCount)
 	END_SHADER_PARAMETER_STRUCT()
 
+	class FBilateralGrid : SHADER_PERMUTATION_BOOL("BILATERAL_GRID");
+	using FPermutationDomain = TShaderPermutationDomain<FBilateralGrid>;
+
 	static bool ShouldCompilePermutation(const FGlobalShaderPermutationParameters& Parameters)
 	{
 		return IsFeatureLevelSupported(Parameters.Platform, ERHIFeatureLevel::SM5);
@@ -52,22 +50,43 @@ public:
 	static void ModifyCompilationEnvironment(const FGlobalShaderPermutationParameters& Parameters, FShaderCompilerEnvironment& OutEnvironment)
 	{
 		FGlobalShader::ModifyCompilationEnvironment(Parameters, OutEnvironment);
-		OutEnvironment.SetDefine(TEXT("THREADGROUP_SIZEX"), ThreadGroupSizeX);
-		OutEnvironment.SetDefine(TEXT("THREADGROUP_SIZEY"), ThreadGroupSizeY);
+
+		const FPermutationDomain PermutationVector(Parameters.PermutationId);
+		const bool bBilateralGrid = PermutationVector.Get<FBilateralGrid>();
+
+		const FIntPoint ThreadGroupSize = GetThreadGroupSize(bBilateralGrid);
+		const uint32 HistogramSize = GetHistogramSize(bBilateralGrid);
+
+		OutEnvironment.SetDefine(TEXT("THREADGROUP_SIZEX"), ThreadGroupSize.X);
+		OutEnvironment.SetDefine(TEXT("THREADGROUP_SIZEY"), ThreadGroupSize.Y);
 		OutEnvironment.SetDefine(TEXT("LOOP_SIZEX"), LoopCountX);
 		OutEnvironment.SetDefine(TEXT("LOOP_SIZEY"), LoopCountY);
 		OutEnvironment.SetDefine(TEXT("HISTOGRAM_SIZE"), HistogramSize);
-		OutEnvironment.SetDefine(TEXT("BILATERAL_GRID"), BilateralGrid);
+
 		OutEnvironment.CompilerFlags.Add( CFLAG_StandardOptimization );
 	}
 
-	// One ThreadGroup processes LoopCountX*LoopCountY blocks of size ThreadGroupSizeX*ThreadGroupSizeY
-	static FIntPoint GetThreadGroupCount(FIntPoint InputExtent)
+	static uint32 GetHistogramSize(bool bBilateralGrid)
+	{
+		return bBilateralGrid ? 32 : 64;
+	}
+
+	static FIntPoint GetThreadGroupSize(bool bBilateralGrid)
 	{
+		return bBilateralGrid ? FIntPoint(8, 8) : FIntPoint(8, 4);
+	}
+
+	static FIntPoint GetThreadGroupCount(FIntPoint InputExtent, bool bBilateralGrid)
+	{
+		// One ThreadGroup ThreadGroupSizeX*ThreadGroupSizeY processes blocks of size LoopCountX*LoopCountY
+		const FIntPoint TexelsPerThreadGroup = GetThreadGroupSize(bBilateralGrid) * FIntPoint(LoopCountX, LoopCountY);
+
 		return FIntPoint::DivideAndRoundUp(InputExtent, TexelsPerThreadGroup);
 	}
 };
 
+IMPLEMENT_GLOBAL_SHADER(FHistogramCS, "/Engine/Private/PostProcessHistogram.usf", "MainCS", SF_Compute);
+
 class FHistogramReducePS : public FGlobalShader
 {
 public:
@@ -98,6 +117,8 @@ public:
 	}
 };
 
+IMPLEMENT_GLOBAL_SHADER(FHistogramReducePS, "/Engine/Private/PostProcessHistogramReduce.usf", "MainPS", SF_Pixel);
+
 class FHistogramAtomicCS : public FGlobalShader
 {
 public:
@@ -105,12 +126,6 @@ public:
 	static const uint32 ThreadGroupSizeX = 64; //tested at 32, 64, 128 all fast.  but 2 or 256 were DISASTER.
 	static const uint32 HistogramSize = 64;		//should be power of two
 
-	// /4 as we store 4 buckets in one ARGB texel.
-	static const uint32 HistogramTexelCount = HistogramSize;
-
-	// The number of texels on each axis processed by a single thread group.
-	//	static const FIntPoint TexelsPerThreadGroup;
-
 	DECLARE_GLOBAL_SHADER(FHistogramAtomicCS);
 	SHADER_USE_PARAMETER_STRUCT(FHistogramAtomicCS, FGlobalShader);
 
@@ -138,6 +153,8 @@ public:
 	}
 };
 
+IMPLEMENT_GLOBAL_SHADER(FHistogramAtomicCS, "/Engine/Private/Histogram.usf", "MainAtomicCS", SF_Compute);
+
 class FHistogramAtomicConvertCS : public FGlobalShader
 {
 public:
@@ -175,20 +192,6 @@ public:
 	}
 };
 
-typedef THistogramCS<8, 4, 64, false> FHistogramCS;
-IMPLEMENT_SHADER_TYPE(template<>, FHistogramCS, TEXT("/Engine/Private/PostProcessHistogram.usf"), TEXT("MainCS"), SF_Compute);
-template<>
-const FIntPoint FHistogramCS::TexelsPerThreadGroup(8 * LoopCountX, 4 * LoopCountY);
-
-static const uint32 BilateralGridDepth = 32;
-typedef THistogramCS<8, 8, BilateralGridDepth, true> FHistogramBilateralGridCS;
-IMPLEMENT_SHADER_TYPE(template<>, FHistogramBilateralGridCS, TEXT("/Engine/Private/PostProcessHistogram.usf"), TEXT("MainCS"), SF_Compute);
-template<>
-const FIntPoint FHistogramBilateralGridCS::TexelsPerThreadGroup(8 * LoopCountX, 8 * LoopCountY);
-
-IMPLEMENT_GLOBAL_SHADER(FHistogramReducePS, "/Engine/Private/PostProcessHistogramReduce.usf", "MainPS", SF_Pixel);
-
-IMPLEMENT_GLOBAL_SHADER(FHistogramAtomicCS,        "/Engine/Private/Histogram.usf", "MainAtomicCS", SF_Compute);
 IMPLEMENT_GLOBAL_SHADER(FHistogramAtomicConvertCS, "/Engine/Private/Histogram.usf", "HistogramConvertCS", SF_Compute);
 
 } //! namespace
@@ -203,9 +206,9 @@ static FRDGTextureRef AddHistogramLegacyPass(
 	check(SceneColor.IsValid());
 	check(EyeAdaptationTexture);
 
-	const FIntPoint HistogramThreadGroupCount = FIntPoint::DivideAndRoundUp(SceneColor.ViewRect.Size(), FHistogramCS::TexelsPerThreadGroup);
-
+	const FIntPoint HistogramThreadGroupCount = FHistogramCS::GetThreadGroupCount(SceneColor.ViewRect.Size(), false);
 	const uint32 HistogramThreadGroupCountTotal = HistogramThreadGroupCount.X * HistogramThreadGroupCount.Y;
+	const uint32 HistogramTexelCount = FHistogramCS::GetHistogramSize(false) / FHistogramCS::HistogramBucketsPerTexel;
 
 	FRDGTextureRef HistogramTexture = nullptr;
 
@@ -213,7 +216,7 @@ static FRDGTextureRef AddHistogramLegacyPass(
 
 	// First pass outputs one flattened histogram per group.
 	{
-		const FIntPoint TextureExtent = FIntPoint(FHistogramCS::HistogramTexelCount, HistogramThreadGroupCountTotal);
+		const FIntPoint TextureExtent = FIntPoint(HistogramTexelCount, HistogramThreadGroupCountTotal);
 
 		const FRDGTextureDesc TextureDesc = FRDGTextureDesc::Create2D(
 			TextureExtent,
@@ -232,7 +235,9 @@ static FRDGTextureRef AddHistogramLegacyPass(
 		PassParameters->ThreadGroupCount = HistogramThreadGroupCount;
 		PassParameters->EyeAdaptation = EyeAdaptationParameters;
 
-		TShaderMapRef<FHistogramCS> ComputeShader(View.ShaderMap);
+		FHistogramCS::FPermutationDomain PermutationVector;
+
+		auto ComputeShader = View.ShaderMap->GetShader<FHistogramCS>(PermutationVector);
 
 		FComputeShaderUtils::AddPass(
 			GraphBuilder,
@@ -246,7 +251,7 @@ static FRDGTextureRef AddHistogramLegacyPass(
 
 	// Second pass further reduces the histogram to a single line. The second line contains the eye adaptation value (two line texture).
 	{
-		const FIntPoint TextureExtent = FIntPoint(FHistogramCS::HistogramTexelCount, 2);
+		const FIntPoint TextureExtent = FIntPoint(HistogramTexelCount, 2);
 
 		const FRDGTextureDesc TextureDesc = FRDGTextureDesc::Create2D(
 			TextureExtent,
@@ -299,7 +304,7 @@ static FRDGTextureRef AddHistogramAtomicPass(
 	{
 		// {
 		// 	FRDGTextureDesc Desc = FRDGTextureDesc::Create2D(
-		// 		FIntPoint(FHistogramAtomicCS::HistogramTexelCount, 1),
+		// 		FIntPoint(FHistogramAtomicCS::HistogramSize, 1),
 		// 		PF_R32G32_UINT,
 		// 		FClearValueBinding::None,
 		// 		TexCreate_UAV | TexCreate_AtomicCompatible);
@@ -309,7 +314,7 @@ static FRDGTextureRef AddHistogramAtomicPass(
 
 		{
 			FRDGTextureDesc Desc = FRDGTextureDesc::Create2D(
-				FIntPoint(FHistogramAtomicCS::HistogramTexelCount * 2, 1),
+				FIntPoint(FHistogramAtomicCS::HistogramSize * 2, 1),
 				PF_R32_UINT,
 				FClearValueBinding::None,
 				TexCreate_ShaderResource | TexCreate_UAV | TexCreate_AtomicCompatible);
@@ -345,7 +350,7 @@ static FRDGTextureRef AddHistogramAtomicPass(
 	{
 		{
 			const FRDGTextureDesc TextureDescGather = FRDGTextureDesc::Create2D(
-				FIntPoint(FHistogramAtomicCS::HistogramTexelCount / 4, 2),
+				FIntPoint(FHistogramAtomicCS::HistogramSize / FHistogramCS::HistogramBucketsPerTexel, 2),
 				PF_A32B32G32R32F,
 				FClearValueBinding::None,
 				TexCreate_ShaderResource | TexCreate_UAV);
@@ -412,17 +417,16 @@ FRDGTextureRef AddLocalExposurePass(
 {
 	check(SceneColor.IsValid());
 
-	const FIntPoint TileSize(64, 64);
-	const FIntPoint ThreadGroupSize(8, 8);
-	const FIntPoint NumTiles = FIntPoint::DivideAndRoundUp(SceneColor.ViewRect.Size(), FHistogramBilateralGridCS::TexelsPerThreadGroup);
-	const FIntPoint ThreadGroupCount = FIntPoint::DivideAndRoundUp(SceneColor.ViewRect.Size(), FHistogramBilateralGridCS::TexelsPerThreadGroup);
+	const FIntPoint ThreadGroupSize = FHistogramCS::GetThreadGroupSize(true);
+	const FIntPoint ThreadGroupCount = FHistogramCS::GetThreadGroupCount(SceneColor.ViewRect.Size(), true);
+	const FIntPoint NumTiles = ThreadGroupCount;
 
 	FRDGTextureRef LocalExposureTexture = nullptr;
 
 	RDG_EVENT_SCOPE(GraphBuilder, "LocalExposure");
 
 	{
-		const FIntVector TextureExtent = FIntVector(NumTiles.X, NumTiles.Y, BilateralGridDepth);
+		const FIntVector TextureExtent = FIntVector(NumTiles.X, NumTiles.Y, FHistogramCS::GetHistogramSize(true));
 
 		const FRDGTextureDesc TextureDesc = FRDGTextureDesc::Create3D(
 			TextureExtent,
@@ -432,7 +436,7 @@ FRDGTextureRef AddLocalExposurePass(
 
 		LocalExposureTexture = GraphBuilder.CreateTexture(TextureDesc, TEXT("LocalExposure"));
 
-		auto* PassParameters = GraphBuilder.AllocParameters<FHistogramBilateralGridCS::FParameters>();
+		auto* PassParameters = GraphBuilder.AllocParameters<FHistogramCS::FParameters>();
 		PassParameters->View = View.ViewUniformBuffer;
 		PassParameters->EyeAdaptation = EyeAdaptationParameters;
 		PassParameters->Input = GetScreenPassTextureViewportParameters(FScreenPassTextureViewport(SceneColor));
@@ -441,12 +445,15 @@ FRDGTextureRef AddLocalExposurePass(
 		PassParameters->BilateralGridRWTexture = GraphBuilder.CreateUAV(LocalExposureTexture);
 		PassParameters->ThreadGroupCount = ThreadGroupCount;
 
-		auto ComputeShader = View.ShaderMap->GetShader<FHistogramBilateralGridCS>();
+		FHistogramCS::FPermutationDomain PermutationVector;
+		PermutationVector.Set<FHistogramCS::FBilateralGrid>(true);
+
+		auto ComputeShader = View.ShaderMap->GetShader<FHistogramCS>(PermutationVector);
 
 		FComputeShaderUtils::AddPass(
 			GraphBuilder,
 			RDG_EVENT_NAME("FLocalExposure %dx%d (CS)", SceneColor.ViewRect.Width(), SceneColor.ViewRect.Height()),
-			ERDGPassFlags::Compute | ERDGPassFlags::NeverCull,
+			ERDGPassFlags::Compute,
 			ComputeShader,
 			PassParameters,
 			FIntVector(ThreadGroupCount.X, ThreadGroupCount.Y, 1));
