Fix incorrect Strata tile topology during lighting pass when rect primitive are supported.

diff --git a/Engine/Source/Runtime/Renderer/Private/LightRendering.cpp b/Engine/Source/Runtime/Renderer/Private/LightRendering.cpp
--- a/Engine/Source/Runtime/Renderer/Private/LightRendering.cpp
+++ b/Engine/Source/Runtime/Renderer/Private/LightRendering.cpp
@@ -1839,6 +1839,12 @@ static void InternalRenderLight(
 			VSPermutationVector.Set< Strata::FStrataTilePassVS::FEnableTexCoordScreenVector >(true);
 			TShaderMapRef<Strata::FStrataTilePassVS> TileVertexShader(View.ShaderMap, VSPermutationVector);
 
+			Strata::FStrataTilePassVS::FParameters VSParameters;
+			if (Strata::IsStrataEnabled())
+			{
+				Strata::FillUpTiledPassData(StrataTileMaterialType, View, VSParameters, GraphicsPSOInit.PrimitiveType);
+			}
+
 			// Turn DBT back off
 			GraphicsPSOInit.bDepthBounds = false;
 			GraphicsPSOInit.RasterizerState = TStaticRasterizerState<FM_Solid, CM_None>::GetRHI();
@@ -1852,10 +1858,7 @@ static void InternalRenderLight(
 
 			if (Strata::IsStrataEnabled())
 			{
-				Strata::FStrataTilePassVS::FParameters VSParameters;
-				Strata::FillUpTiledPassData(StrataTileMaterialType, View, VSParameters, GraphicsPSOInit.PrimitiveType);
 				SetShaderParameters(RHICmdList, TileVertexShader, TileVertexShader.GetVertexShader(), VSParameters);
-
 				RHICmdList.DrawPrimitiveIndirect(VSParameters.TileIndirectBuffer->GetIndirectRHICallBuffer(), 0);
 			}
 			else
@@ -2070,6 +2073,7 @@ static void RenderLight(
 			PermutationVector.Set< FDeferredLightPS::FAtmosphereTransmittance >(PassParameters->PS.CloudShadow.bAtmospherePerPixelTransmittance > 0);
 			PermutationVector.Set< FDeferredLightPS::FCloudTransmittance >(PassParameters->PS.CloudShadow.bCloudPerPixelTransmittance > 0);
 		}
+		PermutationVector = FDeferredLightPS::RemapPermutation(PermutationVector);
 
 		// Strata tile rendering: 
 		// * if the light is directional, then dispatch a set of rect tiles
