Fix typo.

diff --git a/Engine/Source/Runtime/Renderer/Private/LightRendering.cpp b/Engine/Source/Runtime/Renderer/Private/LightRendering.cpp
--- a/Engine/Source/Runtime/Renderer/Private/LightRendering.cpp
+++ b/Engine/Source/Runtime/Renderer/Private/LightRendering.cpp
@@ -2073,7 +2073,6 @@ static void RenderLight(
 			PermutationVector.Set< FDeferredLightPS::FAtmosphereTransmittance >(PassParameters->PS.CloudShadow.bAtmospherePerPixelTransmittance > 0);
 			PermutationVector.Set< FDeferredLightPS::FCloudTransmittance >(PassParameters->PS.CloudShadow.bCloudPerPixelTransmittance > 0);
 		}
-		PermutationVector = FDeferredLightPS::RemapPermutation(PermutationVector);
 
 		// Strata tile rendering: 
 		// * if the light is directional, then dispatch a set of rect tiles
