Remove unecessary permutaiton from DeferredLightPS.

diff --git a/Engine/Source/Runtime/Renderer/Private/LightRendering.cpp b/Engine/Source/Runtime/Renderer/Private/LightRendering.cpp
--- a/Engine/Source/Runtime/Renderer/Private/LightRendering.cpp
+++ b/Engine/Source/Runtime/Renderer/Private/LightRendering.cpp
@@ -468,6 +468,19 @@ class FDeferredLightPS : public FGlobalShader
 	{
 		FPermutationDomain PermutationVector(Parameters.PermutationId);
 
+		// Build FVisualizeCullingDim permutation only for a restricted number of case, as they don't impact the 'estimated cost' of lighting
+		if (PermutationVector.Get< FVisualizeCullingDim >() && (
+			PermutationVector.Get< FSourceTextureDim >() ||
+			PermutationVector.Get< FIESProfileDim >() ||
+			PermutationVector.Get< FTransmissionDim >() ||
+			PermutationVector.Get< FHairLighting >() ||
+			PermutationVector.Get< FAtmosphereTransmittance >() ||
+			PermutationVector.Get< FCloudTransmittance >() ||
+			PermutationVector.Get< FAnistropicMaterials >()))
+		{
+			return false;
+		}
+
 		if (PermutationVector.Get< FSourceShapeDim >() == ELightSourceShape::Directional && PermutationVector.Get< FIESProfileDim >())
 		{
 			return false;
@@ -483,9 +496,7 @@ class FDeferredLightPS : public FGlobalShader
 			return false;
 		}
 
-		if (PermutationVector.Get< FHairLighting >() && (
-			PermutationVector.Get< FVisualizeCullingDim >() ||
-			PermutationVector.Get< FTransmissionDim >()))
+		if (PermutationVector.Get< FHairLighting >() && PermutationVector.Get< FTransmissionDim >())
 		{
 			return false;
 		}
@@ -522,6 +533,23 @@ class FDeferredLightPS : public FGlobalShader
 		return IsFeatureLevelSupported(Parameters.Platform, ERHIFeatureLevel::SM5);
 	}
 
+	static FPermutationDomain RemapPermutation(FPermutationDomain PermutationVector)
+	{
+		// Build FVisualizeCullingDim permutation only for a restricted number of case, as they don't impact the 'estimated cost' of lighting
+		if (PermutationVector.Get< FVisualizeCullingDim >())
+		{
+			PermutationVector.Set< FSourceTextureDim >(false);
+			PermutationVector.Set< FIESProfileDim >(false);
+			PermutationVector.Set< FTransmissionDim >(false);
+			PermutationVector.Set< FHairLighting >(false);
+			PermutationVector.Set< FAtmosphereTransmittance >(false);
+			PermutationVector.Set< FCloudTransmittance >(false);
+			PermutationVector.Set< FAnistropicMaterials >(false);
+		}
+
+		return PermutationVector;
+	}
+
 	static void ModifyCompilationEnvironment(const FGlobalShaderPermutationParameters& Parameters, FShaderCompilerEnvironment& OutEnvironment)
 	{
 		FGlobalShader::ModifyCompilationEnvironment(Parameters, OutEnvironment);
@@ -2073,6 +2101,7 @@ static void RenderLight(
 			PermutationVector.Set< FDeferredLightPS::FAtmosphereTransmittance >(PassParameters->PS.CloudShadow.bAtmospherePerPixelTransmittance > 0);
 			PermutationVector.Set< FDeferredLightPS::FCloudTransmittance >(PassParameters->PS.CloudShadow.bCloudPerPixelTransmittance > 0);
 		}
+		PermutationVector = FDeferredLightPS::RemapPermutation(PermutationVector);
 
 		// Strata tile rendering: 
 		// * if the light is directional, then dispatch a set of rect tiles
