Fix light overlap with Strata.

diff --git a/Engine/Source/Runtime/Renderer/Private/LightRendering.cpp b/Engine/Source/Runtime/Renderer/Private/LightRendering.cpp
--- a/Engine/Source/Runtime/Renderer/Private/LightRendering.cpp
+++ b/Engine/Source/Runtime/Renderer/Private/LightRendering.cpp
@@ -1681,9 +1681,9 @@ static uint32 InternalSetBoundingGeometryDepthState(FGraphicsPipelineStateInitia
 	// bCameraInsideLightGeometry = true  -> CompareFunction = Always
 	// bCameraInsideLightGeometry = false -> CompareFunction = CF_DepthNearOrEqual
 	uint32 StencilRef = 0u;
-	if (Strata::IsStrataEnabled())
+	if (TileType != EStrataTileMaterialType::ECount)
 	{
-		check(TileType != EStrataTileMaterialType::ECount);
+		check(Strata::IsStrataEnabled());
 		switch (TileType)
 		{
 		case EStrataTileMaterialType::ESimple : StencilRef = Strata::StencilBit_Fast;    GraphicsPSOInit.DepthStencilState = TStaticDepthStencilState<false, CompareFunction, true, CF_Equal, SO_Keep, SO_Keep, SO_Keep, true, CF_Equal, SO_Keep, SO_Keep, SO_Keep, Strata::StencilBit_Fast, 0x0>::GetRHI(); break;
@@ -1694,7 +1694,6 @@ static uint32 InternalSetBoundingGeometryDepthState(FGraphicsPipelineStateInitia
 	}
 	else
 	{
-		check(TileType == EStrataTileMaterialType::ECount);
 		GraphicsPSOInit.DepthStencilState = TStaticDepthStencilState<false, CompareFunction>::GetRHI();
 	}
 	return StencilRef;
@@ -1886,8 +1885,9 @@ static void InternalRenderLight(
 
 			SetShaderParameters(RHICmdList, PixelShader, PixelShader.GetPixelShader(), PassParameters->PS);
 
-			if (Strata::IsStrataEnabled())
+			if (StrataTileMaterialType != ECount)
 			{
+				check(Strata::IsStrataEnabled());
 				SetShaderParameters(RHICmdList, TileVertexShader, TileVertexShader.GetVertexShader(), VSParameters);
 				RHICmdList.DrawPrimitiveIndirect(VSParameters.TileIndirectBuffer->GetIndirectRHICallBuffer(), 0);
 			}
@@ -2042,9 +2042,7 @@ static void RenderLight(
 
 		FDeferredLightOverlapPS::FPermutationDomain PermutationVector;
 		PermutationVector.Set<FDeferredLightOverlapPS::FRadialAttenuation>(bIsRadial);
-		TShaderMapRef<FDeferredLightOverlapPS> PixelShader(View.ShaderMap, PermutationVector);
-		
-		// STRATA_TODO: fix overlap debug mode for strata (simple/complex tile rendering)
+		TShaderMapRef<FDeferredLightOverlapPS> PixelShader(View.ShaderMap, PermutationVector);		
 		InternalRenderLight(GraphBuilder, Scene, View, LightSceneInfo, PixelShader, PassParameters, EStrataTileMaterialType::ECount, TEXT("Light::StandardDeferred(Overlap)"));
 	}
 	// Lighting shader
