Fix missing VSM when there is only Nanite affecting a point light by tracking if Nanite geometry was rendered into a shadow map

diff --git a/Engine/Source/Runtime/Renderer/Private/ShadowRendering.h b/Engine/Source/Runtime/Renderer/Private/ShadowRendering.h
--- a/Engine/Source/Runtime/Renderer/Private/ShadowRendering.h
+++ b/Engine/Source/Runtime/Renderer/Private/ShadowRendering.h
@@ -383,6 +383,9 @@ public:
 	 * for direct lighting visibility. */
 	uint32 bIncludeInScreenSpaceShadowMask : 1;
 
+	/** Whether the the shadow overlaps any nanite primitives */
+	uint32 bContainsNaniteSubjects : 1;
+
 	/** Used to fetch the correct cached static mesh draw commands */
 	EMeshPass::Type MeshPassTargetType = EMeshPass::CSMShadowDepth;
 
diff --git a/Engine/Source/Runtime/Renderer/Private/ShadowSetup.cpp b/Engine/Source/Runtime/Renderer/Private/ShadowSetup.cpp
--- a/Engine/Source/Runtime/Renderer/Private/ShadowSetup.cpp
+++ b/Engine/Source/Runtime/Renderer/Private/ShadowSetup.cpp
@@ -706,6 +706,7 @@ FProjectedShadowInfo::FProjectedShadowInfo()
 	, bHairStrandsDeepShadow(false)
 	, bNaniteGeometry(true)
 	, bIncludeInScreenSpaceShadowMask(true)
+	, bContainsNaniteSubjects(false)
 	, PerObjectShadowFadeStart(WORLD_MAX)
 	, InvPerObjectShadowFadeLength(0.0f)
 	, OverlappedUVOnCachedShadowMap(-1.0f, -1.0f, -1.0f, -1.0f)
@@ -2088,7 +2089,8 @@ bool FProjectedShadowInfo::HasSubjectPrims() const
 	return DynamicSubjectPrimitives.Num() > 0
 		|| ShadowDepthPass.HasAnyDraw()
 		|| SubjectMeshCommandBuildRequests.Num() > 0
-		|| ShadowDepthPassVisibleCommands.Num() > 0;
+		|| ShadowDepthPassVisibleCommands.Num() > 0
+		|| (bContainsNaniteSubjects && bNaniteGeometry);
 }
 
 void FProjectedShadowInfo::AddReceiverPrimitive(FPrimitiveSceneInfo* PrimitiveSceneInfo)
@@ -3730,6 +3732,7 @@ void FSceneRenderer::CreateWholeSceneProjectedShadow(
 					bool bContainsNaniteSubjects = false;
 					AddInteractingPrimitives(LightSceneInfo->GetDynamicInteractionOftenMovingPrimitiveList(), ProjectedShadowInfo, bContainsNaniteSubjects);
 					AddInteractingPrimitives(LightSceneInfo->GetDynamicInteractionStaticPrimitiveList(), ProjectedShadowInfo, bContainsNaniteSubjects);
+					ProjectedShadowInfo->bContainsNaniteSubjects = bContainsNaniteSubjects;
 
 					VisibleLightInfo.VirtualShadowMapId = ProjectedShadowInfo->VirtualShadowMaps[0]->ID;
 					VisibleLightInfo.AllProjectedShadows.Add(ProjectedShadowInfo);
@@ -3786,12 +3789,13 @@ void FSceneRenderer::CreateWholeSceneProjectedShadow(
 								AddInteractingPrimitives(LightSceneInfo->GetDynamicInteractionStaticPrimitiveList(), ProjectedShadowInfo, bContainsNaniteSubjects);
 							}
 						}
+						ProjectedShadowInfo->bContainsNaniteSubjects = bContainsNaniteSubjects;
 
 						bool bRenderShadow = true;
 					
 						if (CacheMode[CacheModeIndex] == SDCM_StaticPrimitivesOnly)
 						{
-							const bool bHasStaticPrimitives = ProjectedShadowInfo->HasSubjectPrims() || (bContainsNaniteSubjects && !bNeedsVirtualShadowMap);
+							const bool bHasStaticPrimitives = ProjectedShadowInfo->HasSubjectPrims();
 							bRenderShadow = bHasStaticPrimitives;
 							FCachedShadowMapData& CachedShadowMapData = Scene->GetCachedShadowMapDataRef(ProjectedShadowInfo->GetLightSceneInfo().Id);
 							CachedShadowMapData.bCachedShadowMapHasPrimitives = bHasStaticPrimitives;
