FVolumetricFogLightFunctionPS now uses modern parameterisation.

diff --git a/Engine/Source/Runtime/Renderer/Private/LightRendering.h b/Engine/Source/Runtime/Renderer/Private/LightRendering.h
--- a/Engine/Source/Runtime/Renderer/Private/LightRendering.h
+++ b/Engine/Source/Runtime/Renderer/Private/LightRendering.h
@@ -75,6 +75,14 @@ public:
 		LightFunctionParameters.Bind(ParameterMap,TEXT("LightFunctionParameters"));
 	}
 
+	static FVector4f GetLightFunctionSharedParameters(const FLightSceneInfo* LightSceneInfo, float ShadowFadeFraction)
+	{
+		const bool bIsSpotLight = LightSceneInfo->Proxy->GetLightType() == LightType_Spot;
+		const bool bIsPointLight = LightSceneInfo->Proxy->GetLightType() == LightType_Point;
+		const float TanOuterAngle = bIsSpotLight ? FMath::Tan(LightSceneInfo->Proxy->GetOuterConeAngle()) : 1.0f;
+		return FVector4f(TanOuterAngle, ShadowFadeFraction, bIsSpotLight ? 1.0f : 0.0f, bIsPointLight ? 1.0f : 0.0f);
+	}
+
 	template<typename ShaderRHIParamRef>
 	void Set(FRHICommandList& RHICmdList, const ShaderRHIParamRef ShaderRHI, const FLightSceneInfo* LightSceneInfo, float ShadowFadeFraction) const
 	{
@@ -86,7 +94,7 @@ public:
 			RHICmdList, 
 			ShaderRHI, 
 			LightFunctionParameters, 
-			FVector4f(TanOuterAngle, ShadowFadeFraction, bIsSpotLight ? 1.0f : 0.0f, bIsPointLight ? 1.0f : 0.0f));
+			GetLightFunctionSharedParameters(LightSceneInfo, ShadowFadeFraction));
 	}
 
 	/** Serializer. */ 
diff --git a/Engine/Source/Runtime/Renderer/Private/VolumetricFogLightFunction.cpp b/Engine/Source/Runtime/Renderer/Private/VolumetricFogLightFunction.cpp
--- a/Engine/Source/Runtime/Renderer/Private/VolumetricFogLightFunction.cpp
+++ b/Engine/Source/Runtime/Renderer/Private/VolumetricFogLightFunction.cpp
@@ -51,31 +51,41 @@ static int32 GetVolumetricFogLightFunctionResolution()
 
 class FVolumetricFogLightFunctionPS : public FMaterialShader
 {
-	DECLARE_SHADER_TYPE(FVolumetricFogLightFunctionPS,Material);
-	
+	DECLARE_SHADER_TYPE(FVolumetricFogLightFunctionPS, Material);
+
 	class FLightType : SHADER_PERMUTATION_INT("LIGHT_TYPE", 3);
 	using FPermutationDomain = TShaderPermutationDomain< FLightType >;
-public:
 
-	static bool ShouldCompilePermutation(const FMaterialShaderPermutationParameters& Parameters)
-	{
-		return Parameters.MaterialParameters.MaterialDomain == MD_LightFunction && DoesPlatformSupportVolumetricFog(Parameters.Platform);
-	}
+	BEGIN_SHADER_PARAMETER_STRUCT(FParameters, )
+		SHADER_PARAMETER(FMatrix44f, LightFunctionWorldToLight)
+		SHADER_PARAMETER(FMatrix44f, ShadowToWorld)
+		SHADER_PARAMETER(FVector4f, LightFunctionParameters)
+		SHADER_PARAMETER(FVector3f, LightFunctionParameters2)
+		SHADER_PARAMETER(FVector3f, LightWorldPosition)
+		SHADER_PARAMETER(FVector2f, LightFunctionTexelSize)
+	END_SHADER_PARAMETER_STRUCT()
 
 	FVolumetricFogLightFunctionPS() {}
 	FVolumetricFogLightFunctionPS(const ShaderMetaType::CompiledShaderInitializerType& Initializer)
 		: FMaterialShader(Initializer)
 	{
-		LightFunctionParameters.Bind(Initializer.ParameterMap);
-		LightFunctionWorldToLight.Bind(Initializer.ParameterMap,TEXT("LightFunctionWorldToLight"));
-		LightFunctionParameters2.Bind(Initializer.ParameterMap,TEXT("LightFunctionParameters2"));
-		LightFunctionTexelSize.Bind(Initializer.ParameterMap,TEXT("LightFunctionTexelSize"));
-		ShadowToWorld.Bind(Initializer.ParameterMap,TEXT("ShadowToWorld"));
-		LightWorldPositionParam.Bind(Initializer.ParameterMap, TEXT("LightWorldPosition"));
+		Bindings.BindForLegacyShaderParameters(
+			this,
+			Initializer.PermutationId,
+			Initializer.ParameterMap,
+			*FParameters::FTypeInfo::GetStructMetadata(),
+			// Don't require full bindings, we use FMaterialShader::SetParameters
+			false);
+	}
+
+	static bool ShouldCompilePermutation(const FMaterialShaderPermutationParameters& Parameters)
+	{
+		return Parameters.MaterialParameters.MaterialDomain == MD_LightFunction && DoesPlatformSupportVolumetricFog(Parameters.Platform);
 	}
 
 	void SetParameters(
 		FRHICommandList& RHICmdList, 
+		FParameters& PS,
 		const FSceneView& View, 
 		const FLightSceneInfo* LightSceneInfo, 
 		const FMaterialRenderProxy* MaterialProxy,
@@ -88,39 +98,28 @@ public:
 		FMaterialShader::SetViewParameters(RHICmdList, ShaderRHI, View, View.ViewUniformBuffer);
 		FMaterialShader::SetParameters(RHICmdList, ShaderRHI, MaterialProxy, Material, View);
 
-		LightFunctionParameters.Set(RHICmdList, ShaderRHI, LightSceneInfo, 1.0f);
-
-		SetShaderValue(RHICmdList, ShaderRHI, LightFunctionParameters2, FVector3f(
-			LightSceneInfo->Proxy->GetLightFunctionFadeDistance(), 
+		PS.LightFunctionParameters = FLightFunctionSharedParameters::GetLightFunctionSharedParameters(LightSceneInfo, 1.0f);
+		PS.LightFunctionParameters2 = FVector3f(
+			LightSceneInfo->Proxy->GetLightFunctionFadeDistance(),
 			LightSceneInfo->Proxy->GetLightFunctionDisabledBrightness(),
-			0.0f));
+			0.0f);
 
-		if (LightFunctionWorldToLight.IsBound())
 		{
 			const FVector Scale = LightSceneInfo->Proxy->GetLightFunctionScale();
 			// Switch x and z so that z of the user specified scale affects the distance along the light direction
-			const FVector InverseScale = FVector( 1.f / Scale.Z, 1.f / Scale.Y, 1.f / Scale.X );
+			const FVector InverseScale = FVector(1.f / Scale.Z, 1.f / Scale.Y, 1.f / Scale.X);
 			const FMatrix44f WorldToLight = LightSceneInfo->Proxy->GetWorldToLight() * FScaleMatrix(FVector(InverseScale));		// LWC_TODO: Precision loss
 
-			SetShaderValue(RHICmdList, ShaderRHI, LightFunctionWorldToLight, WorldToLight);
+			PS.LightFunctionWorldToLight = WorldToLight;
 		}
 
-		SetShaderValue(RHICmdList, ShaderRHI, LightFunctionTexelSize, FVector2f(LightFunctionTexelSizeValue));
-		
-		SetShaderValue(RHICmdList, ShaderRHI, ShadowToWorld, ShadowToWorldValue);
-		SetShaderValue(RHICmdList, ShaderRHI, LightWorldPositionParam, FVector3f(LightSceneInfo->Proxy->GetPosition()));
+		PS.LightFunctionTexelSize = FVector2f(LightFunctionTexelSizeValue);
+		PS.ShadowToWorld = ShadowToWorldValue;
+		PS.LightWorldPosition = FVector3f(LightSceneInfo->Proxy->GetPosition());
 	}
-
-private:
-	LAYOUT_FIELD(FLightFunctionSharedParameters, LightFunctionParameters);
-	LAYOUT_FIELD(FShaderParameter, LightFunctionParameters2);
-	LAYOUT_FIELD(FShaderParameter, LightFunctionWorldToLight);
-	LAYOUT_FIELD(FShaderParameter, LightFunctionTexelSize);
-	LAYOUT_FIELD(FShaderParameter, ShadowToWorld);
-	LAYOUT_FIELD(FShaderParameter, LightWorldPositionParam);
 };
 
-IMPLEMENT_MATERIAL_SHADER_TYPE(,FVolumetricFogLightFunctionPS,TEXT("/Engine/Private/VolumetricFogLightFunction.usf"),TEXT("Main"),SF_Pixel);
+IMPLEMENT_MATERIAL_SHADER_TYPE(, FVolumetricFogLightFunctionPS, TEXT("/Engine/Private/VolumetricFogLightFunction.usf"), TEXT("Main"), SF_Pixel);
 
 BEGIN_SHADER_PARAMETER_STRUCT(FVolumetricFogLightFunctionParameters, )
 	SHADER_PARAMETER_RDG_UNIFORM_BUFFER(FSceneTextureUniformParameters, SceneTextures)
@@ -319,7 +318,11 @@ void FDeferredShadingSceneRenderer::RenderLightFunctionForVolumetricFog(
 
 					SetGraphicsPipelineState(RHICmdList, GraphicsPSOInit, 0);
 
-					PixelShader->SetParameters(RHICmdList, View, DirectionalLightSceneInfo, MaterialProxyForRendering, FVector2D(1.0f / LightFunctionResolution.X, 1.0f / LightFunctionResolution.Y), LightFunctionWorldToShadowMatrix.Inverse());
+					FVolumetricFogLightFunctionPS::FParameters PS;
+					PixelShader->SetParameters(RHICmdList, PS, View, DirectionalLightSceneInfo, MaterialProxyForRendering, 
+						FVector2D(1.0f / LightFunctionResolution.X, 1.0f / LightFunctionResolution.Y), LightFunctionWorldToShadowMatrix.Inverse());
+					ClearUnusedGraphResources(PixelShader, &PS);
+					SetShaderParameters(RHICmdList, PixelShader, PixelShader.GetPixelShader(), PS);
 
 					DrawRectangle(
 						RHICmdList,
@@ -436,8 +439,11 @@ void FDeferredShadingSceneRenderer::RenderLightFunctionForVolumetricFog(
 
 					SetGraphicsPipelineState(RHICmdList, GraphicsPSOInit, 0);
 
-					PixelShader->SetParameters(RHICmdList, View, LightSceneInfo, MaterialProxyForRendering, 
+					FVolumetricFogLightFunctionPS::FParameters PS;
+					PixelShader->SetParameters(RHICmdList, PS, View, LightSceneInfo, MaterialProxyForRendering,
 						FVector2D(1.0f / LightFunctionResolution, 1.0f / LightFunctionResolution), LightFunctionInfo.LightFunctionMatrix.Inverse());
+					ClearUnusedGraphResources(PixelShader, &PS);
+					SetShaderParameters(RHICmdList, PixelShader, PixelShader.GetPixelShader(), PS);
 
 					FIntPoint RectSize = LightFunctionInfo.AtlasTile.RectBound.Size();
 					DrawRectangle(
