Fix scene memory usage being wrongly attributed to LoadMapMisc

diff --git a/Engine/Source/Runtime/Renderer/Private/RendererScene.cpp b/Engine/Source/Runtime/Renderer/Private/RendererScene.cpp
--- a/Engine/Source/Runtime/Renderer/Private/RendererScene.cpp
+++ b/Engine/Source/Runtime/Renderer/Private/RendererScene.cpp
@@ -5199,6 +5199,7 @@ private:
 
 FSceneInterface* FRendererModule::AllocateScene(UWorld* World, bool bInRequiresHitProxies, bool bCreateFXSystem, ERHIFeatureLevel::Type InFeatureLevel)
 {
+	LLM_SCOPE(ELLMTag::SceneRender);
 	check(IsInGameThread());
 
 	// Create a full fledged scene if we have something to render.
