Volumetricfog refactor to RGD
- FVolumetricFogIntegrationParameters
- FVolumetricFogMaterialSetupCS
- FWriteToBoundingSphereVS
- FVolumetricFogFinalIntegrationCS

not finished:
- TInjectShadowedLocalLightPS
- FVolumetricFogLightScatteringCS

diff --git a/Engine/Shaders/Private/VolumetricFog.usf b/Engine/Shaders/Private/VolumetricFog.usf
--- a/Engine/Shaders/Private/VolumetricFog.usf
+++ b/Engine/Shaders/Private/VolumetricFog.usf
@@ -82,8 +82,8 @@ float3 ScatteringFunction()
 	//return RaleighScattering();
 }
 
-float3 GlobalAlbedo;
-float3 GlobalEmissive;
+float4 GlobalAlbedo;
+float4 GlobalEmissive;
 float GlobalExtinctionScale;
 
 #ifndef THREADGROUP_SIZE 
@@ -119,7 +119,7 @@ void MaterialSetupCS(
 	float GlobalDensitySecond = FogStruct.ExponentialFogParameters2.z * exp2(-FogStruct.ExponentialFogParameters2.y * (WorldPosition.z - FogStruct.ExponentialFogParameters2.w));
 	float GlobalDensity = GlobalDensityFirst + GlobalDensitySecond;
 
-	float3 Albedo = GlobalAlbedo;
+	float3 Albedo = GlobalAlbedo.rgb;
 
 	// Exponential height fog interprets density differently, match its behavior
 	float MatchHeightFogFactor = .5f;
@@ -132,7 +132,7 @@ void MaterialSetupCS(
 	{
 		RWVBufferA[GridCoordinate] = float4(Scattering, Absorption);
 #if USE_EMISSIVE
-		RWVBufferB[GridCoordinate] = float4(GlobalEmissive, 0); 
+		RWVBufferB[GridCoordinate] = float4(GlobalEmissive.rgb, 0); 
 #endif
 	}
 }
@@ -768,7 +768,7 @@ void LightScatteringCS(
 	
 #endif
 
-	if (all(GridCoordinate < VolumetricFog.GridSizeInt))
+	if (all(int3(GridCoordinate) < VolumetricFog.GridSizeInt))
 	{
 		ScatteringAndExtinction = MakePositiveFinite(ScatteringAndExtinction);
 		RWLightScattering[GridCoordinate] = ScatteringAndExtinction;
@@ -790,7 +790,7 @@ void FinalIntegrationCS(
 	float AccumulatedTransmittance = 1.0f;
 	float3 PreviousSliceWorldPosition = LWCHackToFloat(PrimaryView.WorldCameraOrigin);
 
-	for (uint LayerIndex = 0; LayerIndex < VolumetricFog.GridSizeInt.z; LayerIndex++)
+	for (int LayerIndex = 0; LayerIndex < VolumetricFog.GridSizeInt.z; LayerIndex++)
 	{
 		uint3 LayerCoordinate = uint3(GridCoordinate.xy, LayerIndex);
 		float4 ScatteringAndExtinction = DecodeHDR(LightScattering[LayerCoordinate]);
diff --git a/Engine/Source/Runtime/Renderer/Private/VolumetricFog.cpp b/Engine/Source/Runtime/Renderer/Private/VolumetricFog.cpp
--- a/Engine/Source/Runtime/Renderer/Private/VolumetricFog.cpp
+++ b/Engine/Source/Runtime/Renderer/Private/VolumetricFog.cpp
@@ -25,7 +25,6 @@ VolumetricFog.cpp
 #include "VirtualShadowMaps/VirtualShadowMapClipmap.h"
 #include "ProfilingDebugging/CpuProfilerTrace.h"
 
-IMPLEMENT_TYPE_LAYOUT(FVolumetricFogIntegrationParameters);
 IMPLEMENT_TYPE_LAYOUT(FVolumeShadowingParameters);
 
 int32 GVolumetricFog = 1;
@@ -161,6 +160,32 @@ FVector VolumetricFogTemporalRandom(uint32 FrameNumber)
 	return RandomOffsetValue;
 }
 
+void SetupVolumetricFogIntegrationParameters(
+	FVolumetricFogIntegrationParameters& Out,
+	FViewInfo& View,
+	const FVolumetricFogIntegrationParameterData& IntegrationData)
+{
+	Out.VolumetricFog = View.VolumetricFogResources.VolumetricFogGlobalData;
+
+	FMatrix44f UnjitteredInvTranslatedViewProjectionMatrix = View.ViewMatrices.ComputeInvProjectionNoAAMatrix() * View.ViewMatrices.GetTranslatedViewMatrix().GetTransposed();
+	Out.UnjitteredClipToTranslatedWorld = UnjitteredInvTranslatedViewProjectionMatrix;
+
+	FMatrix44f UnjitteredViewProjectionMatrix = View.PrevViewInfo.ViewMatrices.GetViewMatrix() * View.PrevViewInfo.ViewMatrices.ComputeProjectionNoAAMatrix();
+	Out.UnjitteredPrevWorldToClip = UnjitteredViewProjectionMatrix;
+
+	int32 OffsetCount = IntegrationData.FrameJitterOffsetValues.Num();
+	for (int32 i = 0; i < OffsetCount; ++i)
+	{
+		Out.FrameJitterOffsets[i] = IntegrationData.FrameJitterOffsetValues.GetData()[i];
+	}
+
+	extern float GVolumetricFogHistoryWeight;
+	Out.HistoryWeight = IntegrationData.bTemporalHistoryIsValid ? GVolumetricFogHistoryWeight : 0.0f;
+
+	extern int32 GVolumetricFogHistoryMissSupersampleCount;
+	Out.HistoryMissSuperSampleCount = FMath::Clamp(GVolumetricFogHistoryMissSupersampleCount, 1, 16);
+}
+
 static const uint32 VolumetricFogGridInjectionGroupSize = 4;
 
 namespace
@@ -170,7 +195,8 @@ class FPermutationUseEmissive : SHADER_PERMUTATION_BOOL("USE_EMISSIVE");
 
 class FVolumetricFogMaterialSetupCS : public FGlobalShader
 {
-	DECLARE_GLOBAL_SHADER(FVolumetricFogMaterialSetupCS)
+	DECLARE_GLOBAL_SHADER(FVolumetricFogMaterialSetupCS);
+	SHADER_USE_PARAMETER_STRUCT(FVolumetricFogMaterialSetupCS, FGlobalShader);
 
 	using FPermutationDomain = TShaderPermutationDomain<FPermutationUseEmissive>;
 
@@ -182,6 +208,8 @@ class FVolumetricFogMaterialSetupCS : public FGlobalShader
 		SHADER_PARAMETER_RDG_UNIFORM_BUFFER(FFogUniformParameters, Fog)
 		SHADER_PARAMETER_STRUCT_REF(FViewUniformShaderParameters, View)
 
+		SHADER_PARAMETER_STRUCT_INCLUDE(FVolumetricFogIntegrationParameters, VolumetricFogParameters)
+
 		SHADER_PARAMETER_RDG_TEXTURE_UAV(RWTexture3D<float4>, RWVBufferA)
 		SHADER_PARAMETER_RDG_TEXTURE_UAV(RWTexture3D<float4>, RWVBufferB)
 	END_SHADER_PARAMETER_STRUCT()
@@ -198,31 +226,6 @@ public:
 		FGlobalShader::ModifyCompilationEnvironment(Parameters, OutEnvironment);
 		OutEnvironment.SetDefine(TEXT("THREADGROUP_SIZE"), VolumetricFogGridInjectionGroupSize);
 	}
-
-	FVolumetricFogMaterialSetupCS(const ShaderMetaType::CompiledShaderInitializerType& Initializer)
-		: FGlobalShader(Initializer)
-	{
-		Bindings.BindForLegacyShaderParameters(this, Initializer.PermutationId, Initializer.ParameterMap, *FParameters::FTypeInfo::GetStructMetadata());
-
-		VolumetricFogParameters.Bind(Initializer.ParameterMap);
-	}
-
-	FVolumetricFogMaterialSetupCS()
-	{
-	}
-
-	void SetParameters(
-		FRHICommandList& RHICmdList,
-		const FViewInfo& View,
-		const FVolumetricFogIntegrationParameterData& IntegrationData)
-	{
-		FRHIComputeShader* ShaderRHI = RHICmdList.GetBoundComputeShader();
-		VolumetricFogParameters.Set(RHICmdList, ShaderRHI, View, IntegrationData);
-	}
-
-private:
-
-	LAYOUT_FIELD(FVolumetricFogIntegrationParameters, VolumetricFogParameters);
 };
 
 IMPLEMENT_GLOBAL_SHADER(FVolumetricFogMaterialSetupCS, "/Engine/Private/VolumetricFog.usf", "MaterialSetupCS", SF_Compute);
@@ -230,8 +233,15 @@ IMPLEMENT_GLOBAL_SHADER(FVolumetricFogMaterialSetupCS, "/Engine/Private/Volumetr
 /** Vertex shader used to write to a range of slices of a 3d volume texture. */
 class FWriteToBoundingSphereVS : public FGlobalShader
 {
-	DECLARE_SHADER_TYPE(FWriteToBoundingSphereVS, Global);
-public:
+	DECLARE_GLOBAL_SHADER(FWriteToBoundingSphereVS);
+	SHADER_USE_PARAMETER_STRUCT(FWriteToBoundingSphereVS, FGlobalShader);
+
+	BEGIN_SHADER_PARAMETER_STRUCT(FParameters, )
+		SHADER_PARAMETER_STRUCT_INCLUDE(FVolumetricFogIntegrationParameters, VolumetricFogParameters)
+		SHADER_PARAMETER(FMatrix44f, ViewToVolumeClip)
+		SHADER_PARAMETER(FVector4f, ViewSpaceBoundingSphere)
+		SHADER_PARAMETER(int32, MinZ)
+	END_SHADER_PARAMETER_STRUCT()
 
 	static bool ShouldCompilePermutation(const FGlobalShaderPermutationParameters& Parameters)
 	{
@@ -243,39 +253,9 @@ public:
 		FGlobalShader::ModifyCompilationEnvironment(Parameters, OutEnvironment);
 		OutEnvironment.CompilerFlags.Add(CFLAG_VertexToGeometryShader);
 	}
-
-	FWriteToBoundingSphereVS(const ShaderMetaType::CompiledShaderInitializerType& Initializer) :
-		FGlobalShader(Initializer)
-	{
-		MinZ.Bind(Initializer.ParameterMap, TEXT("MinZ"));
-		ViewSpaceBoundingSphere.Bind(Initializer.ParameterMap, TEXT("ViewSpaceBoundingSphere"));
-		ViewToVolumeClip.Bind(Initializer.ParameterMap, TEXT("ViewToVolumeClip"));
-		VolumetricFogParameters.Bind(Initializer.ParameterMap);
-	}
-
-	FWriteToBoundingSphereVS() {}
-
-	void SetParameters(FRHICommandList& RHICmdList, const FViewInfo& View, const FVolumetricFogIntegrationParameterData& IntegrationData, const FSphere& BoundingSphere, int32 MinZValue)
-	{
-		SetShaderValue(RHICmdList, RHICmdList.GetBoundVertexShader(), MinZ, MinZValue);
-
-		const FVector ViewSpaceBoundingSphereCenter = View.ViewMatrices.GetViewMatrix().TransformPosition(BoundingSphere.Center);
-		SetShaderValue(RHICmdList, RHICmdList.GetBoundVertexShader(), ViewSpaceBoundingSphere, FVector4f(ViewSpaceBoundingSphereCenter, BoundingSphere.W));
-
-		const FMatrix44f ProjectionMatrix = View.ViewMatrices.ComputeProjectionNoAAMatrix();
-		SetShaderValue(RHICmdList, RHICmdList.GetBoundVertexShader(), ViewToVolumeClip, ProjectionMatrix);
-
-		VolumetricFogParameters.Set(RHICmdList, RHICmdList.GetBoundVertexShader(), View, IntegrationData);
-	}
-
-private:
-	LAYOUT_FIELD(FShaderParameter, MinZ);
-	LAYOUT_FIELD(FShaderParameter, ViewSpaceBoundingSphere);
-	LAYOUT_FIELD(FShaderParameter, ViewToVolumeClip);
-	LAYOUT_FIELD(FVolumetricFogIntegrationParameters, VolumetricFogParameters);
 };
 
-IMPLEMENT_SHADER_TYPE(, FWriteToBoundingSphereVS, TEXT("/Engine/Private/VolumetricFog.usf"), TEXT("WriteToBoundingSphereVS"), SF_Vertex);
+IMPLEMENT_GLOBAL_SHADER(FWriteToBoundingSphereVS, "/Engine/Private/VolumetricFog.usf", "WriteToBoundingSphereVS", SF_Vertex);
 
 
 /** Shader that adds direct lighting contribution from the given light to the current volume lighting cascade. */
@@ -289,6 +269,7 @@ class TInjectShadowedLocalLightPS : public FGlobalShader
 		SHADER_PARAMETER_RDG_TEXTURE(Texture2D, WhiteDummyTexture)
 		SHADER_PARAMETER_RDG_TEXTURE(Texture2D, LightFunctionAtlasTexture)
 		SHADER_PARAMETER_SAMPLER(SamplerState, LightFunctionAtlasSampler)
+		SHADER_PARAMETER_STRUCT_INCLUDE(FVolumetricFogIntegrationParameters, VolumetricFogParameters)
 		RENDER_TARGET_BINDING_SLOTS()
 	END_SHADER_PARAMETER_STRUCT()
 
@@ -324,7 +305,6 @@ class TInjectShadowedLocalLightPS : public FGlobalShader
 
 		PhaseG.Bind(Initializer.ParameterMap, TEXT("PhaseG"));
 		InverseSquaredLightDistanceBiasScale.Bind(Initializer.ParameterMap, TEXT("InverseSquaredLightDistanceBiasScale"));
-		VolumetricFogParameters.Bind(Initializer.ParameterMap);
 		VolumeShadowingParameters.Bind(Initializer.ParameterMap);
 
 		LightFunctionMatrixParam.Bind(Initializer.ParameterMap, TEXT("LocalLightFunctionMatrix"));
@@ -360,8 +340,6 @@ public:
 
 		SetDeferredLightParameters(RHICmdList, ShaderRHI, GetUniformBufferParameter<FDeferredLightUniformStruct>(), LightSceneInfo, View);
 
-		VolumetricFogParameters.Set(RHICmdList, ShaderRHI, View, IntegrationData);
-
 		SetShaderValue(RHICmdList, ShaderRHI, PhaseG, FogInfo.VolumetricFogScatteringDistribution);
 		SetShaderValue(RHICmdList, ShaderRHI, InverseSquaredLightDistanceBiasScale, GInverseSquaredLightDistanceBiasScale);
 
@@ -383,7 +361,6 @@ public:
 private:
 	LAYOUT_FIELD(FShaderParameter, PhaseG);
 	LAYOUT_FIELD(FShaderParameter, InverseSquaredLightDistanceBiasScale);
-	LAYOUT_FIELD(FVolumetricFogIntegrationParameters, VolumetricFogParameters);
 	LAYOUT_FIELD(FVolumeShadowingParameters, VolumeShadowingParameters);
 	LAYOUT_FIELD(FShaderParameter, UseConservativeDepthTexture);
 	LAYOUT_FIELD(FShaderParameter, VirtualShadowMapIdParameter);
@@ -687,6 +664,7 @@ void FDeferredShadingSceneRenderer::RenderLocalLightsForVolumetricFog(
 		PassParameters->LightFunctionAtlasSampler = TStaticSamplerState<SF_Bilinear, AM_Clamp, AM_Clamp, AM_Clamp>::GetRHI();
 		// We also bind the default light function texture because when we are out of atlas tile, we fallback to use a white light function so we need the RHI to be created
 		PassParameters->WhiteDummyTexture = View.VolumetricFogResources.TransientLightFunctionTextureAtlas ? View.VolumetricFogResources.TransientLightFunctionTextureAtlas->GetDefaultLightFunctionTexture() : GSystemTextures.GetWhiteDummy(GraphBuilder);
+		SetupVolumetricFogIntegrationParameters(PassParameters->VolumetricFogParameters, View, IntegrationData);
 
 		GraphBuilder.AddPass(
 			RDG_EVENT_NAME("ShadowedLights"),
@@ -762,15 +740,21 @@ void FDeferredShadingSceneRenderer::RenderLocalLightsForVolumetricFog(
 
 						PixelShader->SetParameters(RHICmdList, View, IntegrationData, LightSceneInfo, FogInfo, ProjectedShadowInfo, bDynamicallyShadowed, VirtualShadowMapId,
 							LightFunctionMatrix, LightFunctionTexture, LightFunctionAtlasTileMinMaxUvBound);
-						VertexShader->SetParameters(RHICmdList, View, IntegrationData, LightBounds, VolumeZBounds.X);
+
+						SetShaderParameters(RHICmdList, PixelShader, PixelShader.GetPixelShader(), *PassParameters);
+
+						FWriteToBoundingSphereVS::FParameters VSPassParameters;
+						VSPassParameters.MinZ = VolumeZBounds.X;
+						VSPassParameters.ViewSpaceBoundingSphere = FVector4f(View.ViewMatrices.GetViewMatrix().TransformPosition(LightBounds.Center), LightBounds.W);
+						VSPassParameters.ViewToVolumeClip = View.ViewMatrices.ComputeProjectionNoAAMatrix();
+						VSPassParameters.VolumetricFogParameters = PassParameters->VolumetricFogParameters;
+						SetShaderParameters(RHICmdList, VertexShader, VertexShader.GetVertexShader(), VSPassParameters);
 
 						if (GeometryShader.IsValid())
 						{
 							GeometryShader->SetParameters(RHICmdList, VolumeZBounds.X);
 						}
 
-						SetShaderParameters(RHICmdList, PixelShader, PixelShader.GetPixelShader(), *PassParameters);
-
 						RHICmdList.SetStreamSource(0, GCircleRasterizeVertexBuffer.VertexBufferRHI, 0);
 						const int32 NumInstances = VolumeZBounds.Y - VolumeZBounds.X;
 						const int32 NumTriangles = FCircleRasterizeVertexBuffer::NumVertices - 2;
@@ -787,9 +771,9 @@ END_GLOBAL_SHADER_PARAMETER_STRUCT()
 
 IMPLEMENT_GLOBAL_SHADER_PARAMETER_STRUCT(FLumenTranslucencyLightingUniforms, "LumenGIVolumeStruct");
 
-class TVolumetricFogLightScatteringCS : public FGlobalShader
+class FVolumetricFogLightScatteringCS : public FGlobalShader
 {
-	DECLARE_SHADER_TYPE(TVolumetricFogLightScatteringCS, Global)
+	DECLARE_SHADER_TYPE(FVolumetricFogLightScatteringCS, Global)
 
 	class FTemporalReprojection			: SHADER_PERMUTATION_BOOL("USE_TEMPORAL_REPROJECTION");
 	class FDistanceFieldSkyOcclusion	: SHADER_PERMUTATION_BOOL("DISTANCE_FIELD_SKY_OCCLUSION");
@@ -809,6 +793,7 @@ class TVolumetricFogLightScatteringCS : public FGlobalShader
 	BEGIN_SHADER_PARAMETER_STRUCT(FParameters, )
 		SHADER_PARAMETER_STRUCT_REF(FViewUniformShaderParameters, View)
 		SHADER_PARAMETER_RDG_UNIFORM_BUFFER(FFogUniformParameters, Fog)
+		SHADER_PARAMETER_STRUCT_INCLUDE(FVolumetricFogIntegrationParameters, VolumetricFogParameters)
 
 		SHADER_PARAMETER_RDG_TEXTURE(Texture2D, VBufferA)
 		SHADER_PARAMETER_RDG_TEXTURE(Texture2D, VBufferB)
@@ -866,7 +851,7 @@ public:
 		FVirtualShadowMapArray::SetShaderDefines(OutEnvironment);
 	}
 
-	TVolumetricFogLightScatteringCS(const ShaderMetaType::CompiledShaderInitializerType& Initializer)
+	FVolumetricFogLightScatteringCS(const ShaderMetaType::CompiledShaderInitializerType& Initializer)
 		: FGlobalShader(Initializer)
 	{
 		Bindings.BindForLegacyShaderParameters(this, Initializer.PermutationId, Initializer.ParameterMap, *FParameters::FTypeInfo::GetStructMetadata());
@@ -874,7 +859,6 @@ public:
 		LocalShadowedLightScattering.Bind(Initializer.ParameterMap, TEXT("LocalShadowedLightScattering"));
 		LightScatteringHistory.Bind(Initializer.ParameterMap, TEXT("LightScatteringHistory"));
 		LightScatteringHistorySampler.Bind(Initializer.ParameterMap, TEXT("LightScatteringHistorySampler"));
-		VolumetricFogParameters.Bind(Initializer.ParameterMap);
 		DirectionalLightFunctionWorldToShadow.Bind(Initializer.ParameterMap, TEXT("DirectionalLightFunctionWorldToShadow"));
 		LightFunctionTexture.Bind(Initializer.ParameterMap, TEXT("LightFunctionTexture"));
 		LightFunctionSampler.Bind(Initializer.ParameterMap, TEXT("LightFunctionSampler"));
@@ -899,7 +883,7 @@ public:
 		CloudShadowmapStrength.Bind(Initializer.ParameterMap, TEXT("CloudShadowmapStrength"));
 	}
 
-	TVolumetricFogLightScatteringCS()
+	FVolumetricFogLightScatteringCS()
 	{
 	}
 
@@ -930,7 +914,6 @@ public:
 			TStaticSamplerState<SF_Bilinear, AM_Clamp, AM_Clamp, AM_Clamp>::GetRHI(),
 			LightScatteringHistoryTexture);
 
-		VolumetricFogParameters.Set(RHICmdList, ShaderRHI, View, IntegrationData);
 		SetUniformBufferParameter(RHICmdList, ShaderRHI, GetUniformBufferParameter<FForwardLightData>(), View.ForwardLightingResources->ForwardLightDataUniformBuffer);
 
 		SetShaderValue(RHICmdList, ShaderRHI, DirectionalLightFunctionWorldToShadow, DirectionalLightFunctionWorldToShadowValue);
@@ -1028,7 +1011,6 @@ private:
 	LAYOUT_FIELD(FShaderResourceParameter, LocalShadowedLightScattering);
 	LAYOUT_FIELD(FShaderResourceParameter, LightScatteringHistory);
 	LAYOUT_FIELD(FShaderResourceParameter, LightScatteringHistorySampler);
-	LAYOUT_FIELD(FVolumetricFogIntegrationParameters, VolumetricFogParameters);
 	LAYOUT_FIELD(FShaderParameter, DirectionalLightFunctionWorldToShadow);
 	LAYOUT_FIELD(FShaderResourceParameter, LightFunctionTexture);
 	LAYOUT_FIELD(FShaderResourceParameter, LightFunctionSampler);
@@ -1052,20 +1034,23 @@ private:
 	LAYOUT_FIELD(FShaderParameter, LightScatteringSampleJitterMultiplier)
 };
 
-IMPLEMENT_GLOBAL_SHADER(TVolumetricFogLightScatteringCS, "/Engine/Private/VolumetricFog.usf", "LightScatteringCS", SF_Compute);
+IMPLEMENT_GLOBAL_SHADER(FVolumetricFogLightScatteringCS, "/Engine/Private/VolumetricFog.usf", "LightScatteringCS", SF_Compute);
 
 uint32 VolumetricFogIntegrationGroupSize = 8;
 
 class FVolumetricFogFinalIntegrationCS : public FGlobalShader
 {
-	DECLARE_SHADER_TYPE(FVolumetricFogFinalIntegrationCS, Global)
+	DECLARE_GLOBAL_SHADER(FVolumetricFogFinalIntegrationCS);
+	SHADER_USE_PARAMETER_STRUCT(FVolumetricFogFinalIntegrationCS, FGlobalShader);
 
 	BEGIN_SHADER_PARAMETER_STRUCT(FParameters, )
+		SHADER_PARAMETER_STRUCT_REF(FViewUniformShaderParameters, ViewUniformBuffer)
 		SHADER_PARAMETER_RDG_TEXTURE(Texture3D<float4>, LightScattering)
 		SHADER_PARAMETER_RDG_TEXTURE_UAV(RWTexture3D<float4>, RWIntegratedLightScattering)
+		SHADER_PARAMETER_STRUCT_INCLUDE(FVolumetricFogIntegrationParameters, VolumetricFogParameters)
 	END_SHADER_PARAMETER_STRUCT()
 
-		static bool ShouldCompilePermutation(const FGlobalShaderPermutationParameters& Parameters)
+	static bool ShouldCompilePermutation(const FGlobalShaderPermutationParameters& Parameters)
 	{
 		return DoesPlatformSupportVolumetricFog(Parameters.Platform);
 	}
@@ -1075,33 +1060,9 @@ class FVolumetricFogFinalIntegrationCS : public FGlobalShader
 		FGlobalShader::ModifyCompilationEnvironment(Parameters, OutEnvironment);
 		OutEnvironment.SetDefine(TEXT("THREADGROUP_SIZE"), VolumetricFogIntegrationGroupSize);
 	}
-
-	FVolumetricFogFinalIntegrationCS(const ShaderMetaType::CompiledShaderInitializerType& Initializer)
-		: FGlobalShader(Initializer)
-	{
-		Bindings.BindForLegacyShaderParameters(this, Initializer.PermutationId, Initializer.ParameterMap, *FParameters::FTypeInfo::GetStructMetadata());
-
-		VolumetricFogParameters.Bind(Initializer.ParameterMap);
-	}
-
-	FVolumetricFogFinalIntegrationCS()
-	{
-	}
-
-public:
-	void SetParameters(FRHICommandList& RHICmdList, const FViewInfo& View, const FVolumetricFogIntegrationParameterData& IntegrationData)
-	{
-		FRHIComputeShader* ShaderRHI = RHICmdList.GetBoundComputeShader();
-		FGlobalShader::SetParameters<FViewUniformShaderParameters>(RHICmdList, ShaderRHI, View.ViewUniformBuffer);
-		VolumetricFogParameters.Set(RHICmdList, ShaderRHI, View, IntegrationData);
-	}
-
-private:
-
-	LAYOUT_FIELD(FVolumetricFogIntegrationParameters, VolumetricFogParameters);
 };
 
-IMPLEMENT_SHADER_TYPE(, FVolumetricFogFinalIntegrationCS, TEXT("/Engine/Private/VolumetricFog.usf"), TEXT("FinalIntegrationCS"), SF_Compute);
+IMPLEMENT_GLOBAL_SHADER(FVolumetricFogFinalIntegrationCS, "/Engine/Private/VolumetricFog.usf", "FinalIntegrationCS", SF_Compute);
 
 bool ShouldRenderVolumetricFog(const FScene* Scene, const FSceneViewFamily& ViewFamily)
 {
@@ -1419,6 +1380,7 @@ void FDeferredShadingSceneRenderer::ComputeVolumetricFog(FRDGBuilder& GraphBuild
 
 			PassParameters->Fog = FogUniformBuffer; 
 			PassParameters->View = View.ViewUniformBuffer;
+			SetupVolumetricFogIntegrationParameters(PassParameters->VolumetricFogParameters, View, IntegrationData);
 
 			FVolumetricFogMaterialSetupCS::FPermutationDomain PermutationVector;
 			PermutationVector.Set< FPermutationUseEmissive >(bUseEmissive);
@@ -1435,8 +1397,6 @@ void FDeferredShadingSceneRenderer::ComputeVolumetricFog(FRDGBuilder& GraphBuild
 
 				RHICmdList.SetComputeShader(ComputeShader.GetComputeShader());
 
-				ComputeShader->SetParameters(RHICmdList, View, IntegrationData);
-
 				SetShaderParameters(RHICmdList, ComputeShader, ComputeShader.GetComputeShader(), *PassParameters);
 				DispatchComputeShader(RHICmdList, ComputeShader.GetShader(), NumGroups.X, NumGroups.Y, NumGroups.Z);
 				UnsetShaderUAVs(RHICmdList, ComputeShader, ComputeShader.GetComputeShader());
@@ -1456,10 +1416,12 @@ void FDeferredShadingSceneRenderer::ComputeVolumetricFog(FRDGBuilder& GraphBuild
 		IntegrationData.LightScatteringUAV = GraphBuilder.CreateUAV(FRDGTextureUAVDesc(IntegrationData.LightScattering));
 
 		{
-			TVolumetricFogLightScatteringCS::FParameters* PassParameters = GraphBuilder.AllocParameters<TVolumetricFogLightScatteringCS::FParameters>();
+			FVolumetricFogLightScatteringCS::FParameters* PassParameters = GraphBuilder.AllocParameters<FVolumetricFogLightScatteringCS::FParameters>();
 
 			PassParameters->View = View.ViewUniformBuffer;
 			PassParameters->Fog = FogUniformBuffer;
+			SetupVolumetricFogIntegrationParameters(PassParameters->VolumetricFogParameters, View, IntegrationData);
+
 			PassParameters->VBufferA = IntegrationData.VBufferA;
 			PassParameters->VBufferB = IntegrationData.VBufferB ? IntegrationData.VBufferB : VolumetricBlackDummyTexture;
 			PassParameters->LocalShadowedLightScattering = LocalShadowedLightScattering;
@@ -1511,17 +1473,17 @@ void FDeferredShadingSceneRenderer::ComputeVolumetricFog(FRDGBuilder& GraphBuild
 				&& Views.Num() == 1
 				&& View.IsPerspectiveProjection();
 
-			const int32 SuperSampleCount = TVolumetricFogLightScatteringCS::GetSuperSampleCount(GVolumetricFogHistoryMissSupersampleCount);
+			const int32 SuperSampleCount = FVolumetricFogLightScatteringCS::GetSuperSampleCount(GVolumetricFogHistoryMissSupersampleCount);
 
-			TVolumetricFogLightScatteringCS::FPermutationDomain PermutationVector;
-			PermutationVector.Set< TVolumetricFogLightScatteringCS::FTemporalReprojection >(bUseTemporalReprojection);
-			PermutationVector.Set< TVolumetricFogLightScatteringCS::FDistanceFieldSkyOcclusion >(bUseDistanceFieldSkyOcclusion);
-			PermutationVector.Set< TVolumetricFogLightScatteringCS::FSuperSampleCount >(SuperSampleCount);
-			PermutationVector.Set< TVolumetricFogLightScatteringCS::FLumenGI >(bUseLumenGI);
-			PermutationVector.Set< TVolumetricFogLightScatteringCS::FVirtualShadowMap >( VirtualShadowMapArray.IsAllocated() );
-			PermutationVector.Set< TVolumetricFogLightScatteringCS::FCloudTransmittance >(AtmosphericDirectionalLightIndex >= 0);
+			FVolumetricFogLightScatteringCS::FPermutationDomain PermutationVector;
+			PermutationVector.Set< FVolumetricFogLightScatteringCS::FTemporalReprojection >(bUseTemporalReprojection);
+			PermutationVector.Set< FVolumetricFogLightScatteringCS::FDistanceFieldSkyOcclusion >(bUseDistanceFieldSkyOcclusion);
+			PermutationVector.Set< FVolumetricFogLightScatteringCS::FSuperSampleCount >(SuperSampleCount);
+			PermutationVector.Set< FVolumetricFogLightScatteringCS::FLumenGI >(bUseLumenGI);
+			PermutationVector.Set< FVolumetricFogLightScatteringCS::FVirtualShadowMap >( VirtualShadowMapArray.IsAllocated() );
+			PermutationVector.Set< FVolumetricFogLightScatteringCS::FCloudTransmittance >(AtmosphericDirectionalLightIndex >= 0);
 
-			auto ComputeShader = View.ShaderMap->GetShader< TVolumetricFogLightScatteringCS >(PermutationVector);
+			auto ComputeShader = View.ShaderMap->GetShader< FVolumetricFogLightScatteringCS >(PermutationVector);
 			ClearUnusedGraphResources(ComputeShader, PassParameters);
 
 			FRHITexture* LightScatteringHistoryTexture = GBlackVolumeTexture->TextureRHI;
@@ -1545,7 +1507,7 @@ void FDeferredShadingSceneRenderer::ComputeVolumetricFog(FRDGBuilder& GraphBuild
 				ERDGPassFlags::Compute,
 				[PassParameters, ComputeShader, &View, this, FogInfo, bUseTemporalReprojection, LightScatteringHistoryTexture, CloudInfo, VolumetricFogGridSize, IntegrationData, bUseDirectionalLightShadowing, bUseDistanceFieldSkyOcclusion, DirectionalLightFunctionWorldToShadow, AtmosphericDirectionalLightIndex, AtmosphereLightProxy](FRHICommandList& RHICmdList)
 			{
-				const FIntVector NumGroups = FComputeShaderUtils::GetGroupCount(VolumetricFogGridSize, TVolumetricFogLightScatteringCS::GetGroupSize());
+				const FIntVector NumGroups = FComputeShaderUtils::GetGroupCount(VolumetricFogGridSize, FVolumetricFogLightScatteringCS::GetGroupSize());
 
 				RHICmdList.SetComputeShader(ComputeShader.GetComputeShader());
 
@@ -1564,6 +1526,8 @@ void FDeferredShadingSceneRenderer::ComputeVolumetricFog(FRDGBuilder& GraphBuild
 			FVolumetricFogFinalIntegrationCS::FParameters* PassParameters = GraphBuilder.AllocParameters<FVolumetricFogFinalIntegrationCS::FParameters>();
 			PassParameters->LightScattering = IntegrationData.LightScattering;
 			PassParameters->RWIntegratedLightScattering = IntegratedLightScatteringUAV;
+			PassParameters->ViewUniformBuffer = View.ViewUniformBuffer;
+			SetupVolumetricFogIntegrationParameters(PassParameters->VolumetricFogParameters, View, IntegrationData);
 
 			GraphBuilder.AddPass(
 				RDG_EVENT_NAME("FinalIntegration"),
@@ -1575,7 +1539,6 @@ void FDeferredShadingSceneRenderer::ComputeVolumetricFog(FRDGBuilder& GraphBuild
 
 				auto ComputeShader = View.ShaderMap->GetShader< FVolumetricFogFinalIntegrationCS >();
 				RHICmdList.SetComputeShader(ComputeShader.GetComputeShader());
-				ComputeShader->SetParameters(RHICmdList, View, IntegrationData);
 
 				SetShaderParameters(RHICmdList, ComputeShader, ComputeShader.GetComputeShader(), *PassParameters);
 				DispatchComputeShader(RHICmdList, ComputeShader.GetShader(), NumGroups.X, NumGroups.Y, 1);
diff --git a/Engine/Source/Runtime/Renderer/Private/VolumetricFogShared.h b/Engine/Source/Runtime/Renderer/Private/VolumetricFogShared.h
--- a/Engine/Source/Runtime/Renderer/Private/VolumetricFogShared.h
+++ b/Engine/Source/Runtime/Renderer/Private/VolumetricFogShared.h
@@ -36,94 +36,19 @@ struct FVolumetricFogIntegrationParameterData
 	FRDGTextureUAV* LightScatteringUAV;
 };
 
-class FVolumetricFogIntegrationParameters
-{
-	DECLARE_TYPE_LAYOUT(FVolumetricFogIntegrationParameters, NonVirtual);
-public:
-
-	void Bind(const FShaderParameterMap& ParameterMap)
-	{
-		VBufferA.Bind(ParameterMap, TEXT("VBufferA"));
-		VBufferB.Bind(ParameterMap, TEXT("VBufferB"));
-		LightScattering.Bind(ParameterMap, TEXT("LightScattering"));
-		IntegratedLightScattering.Bind(ParameterMap, TEXT("IntegratedLightScattering"));
-		IntegratedLightScatteringSampler.Bind(ParameterMap, TEXT("IntegratedLightScatteringSampler"));
-		VolumetricFogData.Bind(ParameterMap, TEXT("VolumetricFog"));
-		UnjitteredClipToTranslatedWorld.Bind(ParameterMap, TEXT("UnjitteredClipToTranslatedWorld")); 
-		UnjitteredPrevWorldToClip.Bind(ParameterMap, TEXT("UnjitteredPrevWorldToClip")); 
-		FrameJitterOffsets.Bind(ParameterMap, TEXT("FrameJitterOffsets")); 
-		HistoryWeight.Bind(ParameterMap, TEXT("HistoryWeight"));
-		HistoryMissSuperSampleCount.Bind(ParameterMap, TEXT("HistoryMissSuperSampleCount")); 
-	}
-
-	template<typename ShaderRHIParamRef>
-	void Set(
-		FRHICommandList& RHICmdList, 
-		const ShaderRHIParamRef& ShaderRHI, 
-		const FViewInfo& View, 
-		const FVolumetricFogIntegrationParameterData& IntegrationData) const
-	{
-		SetSamplerParameter(RHICmdList, ShaderRHI, IntegratedLightScatteringSampler, TStaticSamplerState<SF_Bilinear, AM_Clamp, AM_Clamp, AM_Clamp>::GetRHI());
-		
-		if (VolumetricFogData.IsBound())
-		{
-			SetUniformBufferParameter(RHICmdList, ShaderRHI, VolumetricFogData, View.VolumetricFogResources.VolumetricFogGlobalData);
-		}
-
-		if (UnjitteredClipToTranslatedWorld.IsBound())
-		{
-			FMatrix44f UnjitteredInvTranslatedViewProjectionMatrix = View.ViewMatrices.ComputeInvProjectionNoAAMatrix() * View.ViewMatrices.GetTranslatedViewMatrix().GetTransposed();
-			SetShaderValue(RHICmdList, ShaderRHI, UnjitteredClipToTranslatedWorld, UnjitteredInvTranslatedViewProjectionMatrix);
-		}
-
-		if (UnjitteredPrevWorldToClip.IsBound())
-		{
-			FMatrix44f UnjitteredViewProjectionMatrix = View.PrevViewInfo.ViewMatrices.GetViewMatrix() * View.PrevViewInfo.ViewMatrices.ComputeProjectionNoAAMatrix();
-			SetShaderValue(RHICmdList, ShaderRHI, UnjitteredPrevWorldToClip, UnjitteredViewProjectionMatrix);
-		}
-
-		if (FrameJitterOffsets.IsBound())
-		{
-			SetShaderValueArray(RHICmdList, ShaderRHI, FrameJitterOffsets, IntegrationData.FrameJitterOffsetValues.GetData(), IntegrationData.FrameJitterOffsetValues.Num());
-		}
-
-		extern float GVolumetricFogHistoryWeight;
-		SetShaderValue(RHICmdList, ShaderRHI, HistoryWeight, IntegrationData.bTemporalHistoryIsValid ? GVolumetricFogHistoryWeight : 0.0f);
-
-		extern int32 GVolumetricFogHistoryMissSupersampleCount;
-		SetShaderValue(RHICmdList, ShaderRHI, HistoryMissSuperSampleCount, FMath::Clamp(GVolumetricFogHistoryMissSupersampleCount, 1, 16));
-	}
-
-	/** Serializer. */
-	friend FArchive& operator<<(FArchive& Ar,FVolumetricFogIntegrationParameters& P)
-	{
-		Ar << P.VBufferA;
-		Ar << P.VBufferB;
-		Ar << P.LightScattering;
-		Ar << P.IntegratedLightScattering;
-		Ar << P.IntegratedLightScatteringSampler;
-		Ar << P.VolumetricFogData;
-		Ar << P.UnjitteredClipToTranslatedWorld;
-		Ar << P.UnjitteredPrevWorldToClip;
-		Ar << P.FrameJitterOffsets;
-		Ar << P.HistoryWeight;
-		Ar << P.HistoryMissSuperSampleCount;
-		return Ar;
-	}
-
-private:
-	LAYOUT_FIELD(FRWShaderParameter, VBufferA)
-	LAYOUT_FIELD(FRWShaderParameter, VBufferB)
-	LAYOUT_FIELD(FRWShaderParameter, LightScattering)
-	LAYOUT_FIELD(FRWShaderParameter, IntegratedLightScattering)
-	LAYOUT_FIELD(FShaderResourceParameter, IntegratedLightScatteringSampler)
-	LAYOUT_FIELD(FShaderUniformBufferParameter, VolumetricFogData)
-	LAYOUT_FIELD(FShaderParameter, UnjitteredClipToTranslatedWorld)
-	LAYOUT_FIELD(FShaderParameter, UnjitteredPrevWorldToClip)
-	LAYOUT_FIELD(FShaderParameter, FrameJitterOffsets)
-	LAYOUT_FIELD(FShaderParameter, HistoryWeight)
-	LAYOUT_FIELD(FShaderParameter, HistoryMissSuperSampleCount)
-};
+BEGIN_SHADER_PARAMETER_STRUCT(FVolumetricFogIntegrationParameters, )
+	SHADER_PARAMETER_STRUCT_REF(FVolumetricFogGlobalData, VolumetricFog)
+	SHADER_PARAMETER(FMatrix44f, UnjitteredClipToTranslatedWorld)
+	SHADER_PARAMETER(FMatrix44f, UnjitteredPrevWorldToClip)
+	SHADER_PARAMETER_ARRAY(FVector4f, FrameJitterOffsets, [16])
+	SHADER_PARAMETER(float, HistoryWeight)
+	SHADER_PARAMETER(uint32, HistoryMissSuperSampleCount)
+END_SHADER_PARAMETER_STRUCT()
+
+void SetupVolumetricFogIntegrationParameters(
+	FVolumetricFogIntegrationParameters& Out,
+	FViewInfo& View,
+	const FVolumetricFogIntegrationParameterData& IntegrationData);
 
 inline int32 ComputeZSliceFromDepth(float SceneDepth, FVector GridZParams)
 {
