Fixes for HZB culling readback (part-contributed by [REDACTED_PERSON] through UDN) & update shader parameters

diff --git a/Engine/Plugins/Experimental/GPULightmass/Source/GPULightmass/Private/LightmapRenderer.cpp b/Engine/Plugins/Experimental/GPULightmass/Source/GPULightmass/Private/LightmapRenderer.cpp
--- a/Engine/Plugins/Experimental/GPULightmass/Source/GPULightmass/Private/LightmapRenderer.cpp
+++ b/Engine/Plugins/Experimental/GPULightmass/Source/GPULightmass/Private/LightmapRenderer.cpp
@@ -3101,9 +3101,9 @@ void FLightmapRenderer::BackgroundTick()
 			// FLinearColor is in RGBA while the GPU texture is in ABGR
 			// TODO: apply swizzling in the copy compute shader if this becomes a problem
 			void* LockedData[3];
-			ReadbackGroup.StagingHQLayer0Readback->LockTexture(RHICmdList, LockedData[0], ReadbackGroup.TextureData->RowPitchInPixels[0]); // This forces a GPU stall
-			ReadbackGroup.StagingHQLayer1Readback->LockTexture(RHICmdList, LockedData[1], ReadbackGroup.TextureData->RowPitchInPixels[1]); // This forces a GPU stall
-			ReadbackGroup.StagingShadowMaskReadback->LockTexture(RHICmdList, LockedData[2], ReadbackGroup.TextureData->RowPitchInPixels[2]); // This forces a GPU stall
+			LockedData[0] = ReadbackGroup.StagingHQLayer0Readback->Lock(ReadbackGroup.TextureData->RowPitchInPixels[0]); // This forces a GPU stall
+			LockedData[1] = ReadbackGroup.StagingHQLayer1Readback->Lock(ReadbackGroup.TextureData->RowPitchInPixels[1]); // This forces a GPU stall
+			LockedData[2] = ReadbackGroup.StagingShadowMaskReadback->Lock(ReadbackGroup.TextureData->RowPitchInPixels[2]); // This forces a GPU stall
 
 			ReadbackGroup.TextureData->Texture[0].AddUninitialized(ReadbackGroup.ReadbackTilePoolGPU->SizeInTiles.Y * GPreviewLightmapPhysicalTileSize * ReadbackGroup.TextureData->RowPitchInPixels[0]);
 			ReadbackGroup.TextureData->Texture[1].AddUninitialized(ReadbackGroup.ReadbackTilePoolGPU->SizeInTiles.Y * GPreviewLightmapPhysicalTileSize * ReadbackGroup.TextureData->RowPitchInPixels[1]);
diff --git a/Engine/Shaders/Private/HZBOcclusion.usf b/Engine/Shaders/Private/HZBOcclusion.usf
--- a/Engine/Shaders/Private/HZBOcclusion.usf
+++ b/Engine/Shaders/Private/HZBOcclusion.usf
@@ -8,8 +8,12 @@ SamplerState	BoundsCenterSampler;
 Texture2D		BoundsExtentTexture;
 SamplerState	BoundsExtentSampler;
 
-void HZBTestPS(float2 InUV : TEXCOORD0, out float4 OutColor : SV_Target0)
+float2 InvTestTargetSize;
+
+void HZBTestPS(in float4 SvPosition : SV_Position, out float4 OutColor : SV_Target0)
 {
+	float2 InUV = SvPosition.xy * InvTestTargetSize;
+
 	float4 BoundsCenter = BoundsCenterTexture.SampleLevel( BoundsCenterSampler, InUV, 0 );
 	float4 BoundsExtent = BoundsExtentTexture.SampleLevel( BoundsExtentSampler, InUV, 0 );
 	BoundsCenter.xyz += LWCHackToFloat(PrimaryView.PreViewTranslation);
diff --git a/Engine/Source/Runtime/Landscape/Private/LandscapePhysicalMaterial.cpp b/Engine/Source/Runtime/Landscape/Private/LandscapePhysicalMaterial.cpp
--- a/Engine/Source/Runtime/Landscape/Private/LandscapePhysicalMaterial.cpp
+++ b/Engine/Source/Runtime/Landscape/Private/LandscapePhysicalMaterial.cpp
@@ -351,9 +351,8 @@ namespace
 		FRHIGPUTextureReadback* Readback,
 		TArray<uint8>& OutPhysicalMaterialIds)
 	{
-		void* Data = nullptr;
-		int32 Pitch;
-		Readback->LockTexture(RHICmdList, Data, Pitch);
+		int32 Pitch = 0;
+		void* Data = Readback->Lock(Pitch);
 		check(Data && TargetSize.X <= Pitch);
 
 		OutPhysicalMaterialIds.Empty(TargetSize.X * TargetSize.Y);
diff --git a/Engine/Source/Runtime/RHI/Private/RHIGPUReadback.cpp b/Engine/Source/Runtime/RHI/Private/RHIGPUReadback.cpp
--- a/Engine/Source/Runtime/RHI/Private/RHIGPUReadback.cpp
+++ b/Engine/Source/Runtime/RHI/Private/RHIGPUReadback.cpp
@@ -185,7 +185,7 @@ void* FRHIGPUTextureReadback::Lock(uint32 NumBytes)
 	}
 }
 
-void FRHIGPUTextureReadback::LockTexture(FRHICommandListImmediate& RHICmdList, void*& OutBufferPtr, int32& OutRowPitchInPixels)
+void* FRHIGPUTextureReadback::Lock(int32& OutRowPitchInPixels, int32 *OutBufferHeight)
 {
 	if (DestinationStagingTexture)
 	{
@@ -193,15 +193,27 @@ void FRHIGPUTextureReadback::LockTexture(FRHICommandListImmediate& RHICmdList, v
 
 		void* ResultsBuffer = nullptr;
 		int32 BufferWidth = 0, BufferHeight = 0;
+		FRHICommandListImmediate& RHICmdList = FRHICommandListExecutor::GetImmediateCommandList();
 		RHICmdList.MapStagingSurface(DestinationStagingTexture, Fence.GetReference(), ResultsBuffer, BufferWidth, BufferHeight, LastLockGPUIndex);
-		OutBufferPtr = ResultsBuffer;
+		if (OutBufferHeight)
+		{
+			*OutBufferHeight = BufferHeight;
+		}
 		OutRowPitchInPixels = BufferWidth;
+		return ResultsBuffer;
 	}
-	else
+
+	OutRowPitchInPixels = 0;
+	if (OutBufferHeight)
 	{
-		OutBufferPtr = nullptr;
-		OutRowPitchInPixels = 0;
+		*OutBufferHeight = 0;
 	}
+	return nullptr;
+}
+
+void FRHIGPUTextureReadback::LockTexture(FRHICommandListImmediate& RHICmdList, void*& OutBufferPtr, int32& OutRowPitchInPixels)
+{
+	OutBufferPtr = Lock(OutRowPitchInPixels);
 }
 
 void FRHIGPUTextureReadback::Unlock()
diff --git a/Engine/Source/Runtime/RHI/Public/RHIGPUReadback.h b/Engine/Source/Runtime/RHI/Public/RHIGPUReadback.h
--- a/Engine/Source/Runtime/RHI/Public/RHIGPUReadback.h
+++ b/Engine/Source/Runtime/RHI/Public/RHIGPUReadback.h
@@ -113,9 +113,14 @@ public:
 	void EnqueueCopyRDG(FRHICommandList& RHICmdList, FRHITexture* SourceTexture, FResolveRect Rect = FResolveRect());
 	void EnqueueCopy(FRHICommandList& RHICmdList, FRHITexture* SourceTexture, FResolveRect Rect = FResolveRect()) override;
 
+
+	UE_DEPRECATED(5.0, "Use FRHIGPUTextureReadback::Lock( int32& OutRowPitchInPixels) instead.")
 	void* Lock(uint32 NumBytes) override;
+
+	void* Lock(int32& OutRowPitchInPixels, int32* OutBufferHeight = nullptr);
 	void Unlock() override;
 
+	UE_DEPRECATED(5.0, "Use FRHIGPUTextureReadback::Lock( int32& OutRowPitchInPixels) instead.")
 	void LockTexture(FRHICommandListImmediate& RHICmdList, void*& OutBufferPtr, int32& OutRowPitchInPixels);
 
 private:
diff --git a/Engine/Source/Runtime/Renderer/Private/PostProcess/PostProcessEyeAdaptation.cpp b/Engine/Source/Runtime/Renderer/Private/PostProcess/PostProcessEyeAdaptation.cpp
--- a/Engine/Source/Runtime/Renderer/Private/PostProcess/PostProcessEyeAdaptation.cpp
+++ b/Engine/Source/Runtime/Renderer/Private/PostProcess/PostProcessEyeAdaptation.cpp
@@ -762,7 +762,8 @@ void FSceneViewState::FEyeAdaptationManager::UpdateLastExposureFromTexture()
 	if (ReadbackTexture)
 	{
 		// Read the last request results.
-		FVector4f* ReadbackData = (FVector4f*)ReadbackTexture->Lock(sizeof(FVector4f));
+		int32 RowPitch = 0;
+		FVector4f* ReadbackData = (FVector4f*)ReadbackTexture->Lock(RowPitch);
 		if (ReadbackData)
 		{
 			LastExposure = ReadbackData->X;
diff --git a/Engine/Source/Runtime/Renderer/Private/SceneOcclusion.cpp b/Engine/Source/Runtime/Renderer/Private/SceneOcclusion.cpp
--- a/Engine/Source/Runtime/Renderer/Private/SceneOcclusion.cpp
+++ b/Engine/Source/Runtime/Renderer/Private/SceneOcclusion.cpp
@@ -18,6 +18,7 @@
 #include "ShaderParameterStruct.h"
 #include "VisualizeTexture.h"
 #include "ProfilingDebugging/CpuProfilerTrace.h"
+#include "PixelShaderUtils.h"
 
 
 /*-----------------------------------------------------------------------------
@@ -760,7 +761,8 @@ static void ExecutePlanarReflectionOcclusionQuery(FRHICommandList& RHICmdList, u
 }
 
 FHZBOcclusionTester::FHZBOcclusionTester()
-	: ResultsBuffer( NULL )
+	: ResultsBuffer(nullptr)
+	, ResultsBufferRowPitch(0)
 {
 	SetInvalidFrameNumber();
 }
@@ -796,8 +798,7 @@ void FHZBOcclusionTester::InitDynamicRHI()
 	{
 		FRHICommandListImmediate& RHICmdList = FRHICommandListExecutor::GetImmediateCommandList();
 		FPooledRenderTargetDesc Desc( FPooledRenderTargetDesc::Create2DDesc( FIntPoint( SizeX, SizeY ), PF_B8G8R8A8, FClearValueBinding::None, TexCreate_CPUReadback | TexCreate_HideInVisualizeTexture, TexCreate_None, false ) );
-		GRenderTargetPool.FindFreeElement(RHICmdList, Desc, ResultsTextureCPU, TEXT("HZBResultsCPU"), ERenderTargetTransience::NonTransient );
-		Fence = RHICreateGPUFence(TEXT("HZBGPUFence"));
+		ResultsReadback.Reset(new FRHIGPUTextureReadback(TEXT("HZBGPUReadback")));
 	}
 }
 
@@ -805,8 +806,7 @@ void FHZBOcclusionTester::ReleaseDynamicRHI()
 {
 	if (GetFeatureLevel() >= ERHIFeatureLevel::SM5)
 	{
-		GRenderTargetPool.FreeUnusedResource( ResultsTextureCPU );
-		Fence.SafeRelease();
+		ResultsReadback.Reset();
 	}
 }
 
@@ -827,10 +827,13 @@ void FHZBOcclusionTester::MapResults(FRHICommandListImmediate& RHICmdList)
 	{
 		uint32 IdleStart = FPlatformTime::Cycles();
 
-		int32 Width = 0;
-		int32 Height = 0;
+		SCOPED_GPU_MASK(RHICmdList, ResultsReadback->GetLastCopyGPUMask());
 
-		RHICmdList.MapStagingSurface(ResultsTextureCPU->GetRenderTargetItem().ShaderResourceTexture, Fence.GetReference(), *(void**)&ResultsBuffer, Width, Height);
+		int32 ResultBufferHeight = 0;
+		ResultsBufferRowPitch = 0;
+		ResultsBuffer = reinterpret_cast<const uint8*>(ResultsReadback->Lock(ResultsBufferRowPitch, &ResultBufferHeight));
+		check(ResultsBufferRowPitch >= SizeX);
+		check(ResultBufferHeight >= SizeY);
 
 		// RHIMapStagingSurface will block until the results are ready (from the previous frame) so we need to consider this RT idle time
 		GRenderThreadIdle[ERenderThreadIdleTypes::WaitingForGPUQuery] += FPlatformTime::Cycles() - IdleStart;
@@ -838,7 +841,7 @@ void FHZBOcclusionTester::MapResults(FRHICommandListImmediate& RHICmdList)
 	}
 	
 	// Can happen because of device removed, we might crash later but this occlusion culling system can behave gracefully.
-	if( ResultsBuffer == NULL )
+	if( ResultsBuffer == nullptr )
 	{
 		// First frame
 		static uint8 FirstFrameBuffer[] = { 255 };
@@ -852,9 +855,10 @@ void FHZBOcclusionTester::UnmapResults(FRHICommandListImmediate& RHICmdList)
 	check( ResultsBuffer );
 	if(!IsInvalidFrame())
 	{
-		RHICmdList.UnmapStagingSurface(ResultsTextureCPU->GetRenderTargetItem().ShaderResourceTexture);
+		SCOPED_GPU_MASK(RHICmdList, ResultsReadback->GetLastCopyGPUMask());
+		ResultsReadback->Unlock();
 	}
-	ResultsBuffer = NULL;
+	ResultsBuffer = nullptr;
 }
 
 bool FHZBOcclusionTester::IsVisible( uint32 Index ) const
@@ -886,95 +890,46 @@ bool FHZBOcclusionTester::IsVisible( uint32 Index ) const
 	const int32 x = BlockX * BlockSize + b % BlockSize;
 	const int32 y = BlockY * BlockSize + b / BlockSize;
 
-	return ResultsBuffer[ 4 * (x + y * SizeY) ] != 0;
+	return ResultsBuffer[ 4 * (x + y * ResultsBufferRowPitch) ] != 0;
 #endif
 }
 
+
 class FHZBTestPS : public FGlobalShader
 {
-	DECLARE_SHADER_TYPE(FHZBTestPS, Global);
+	DECLARE_GLOBAL_SHADER(FHZBTestPS);
+	SHADER_USE_PARAMETER_STRUCT(FHZBTestPS, FGlobalShader)
 
-	static bool ShouldCompilePermutation(const FGlobalShaderPermutationParameters& Parameters)
-	{
-		return IsFeatureLevelSupported(Parameters.Platform, ERHIFeatureLevel::SM5);
-	}
+	BEGIN_SHADER_PARAMETER_STRUCT(FParameters, )
+		SHADER_PARAMETER_STRUCT_REF(FViewUniformShaderParameters, View)
 
-	FHZBTestPS() {}
+		SHADER_PARAMETER_RDG_TEXTURE(Texture2D, HZBTexture)
+		SHADER_PARAMETER_SAMPLER(SamplerState, HZBSampler)
+		SHADER_PARAMETER_RDG_TEXTURE(Texture2D, BoundsCenterTexture)
+		SHADER_PARAMETER_SAMPLER(SamplerState, BoundsCenterSampler)
+		SHADER_PARAMETER_RDG_TEXTURE(Texture2D, BoundsExtentTexture)
+		SHADER_PARAMETER_SAMPLER(SamplerState, BoundsExtentSampler)
 
-public:
-	LAYOUT_FIELD(FShaderParameter, HZBUvFactor)
-	LAYOUT_FIELD(FShaderParameter, HZBViewSize)
-	LAYOUT_FIELD(FShaderParameter, HZBSize)
-	LAYOUT_FIELD(FShaderResourceParameter, HZBTexture)
-	LAYOUT_FIELD(FShaderResourceParameter, HZBSampler)
-	LAYOUT_FIELD(FShaderResourceParameter, BoundsCenterTexture)
-	LAYOUT_FIELD(FShaderResourceParameter, BoundsCenterSampler)
-	LAYOUT_FIELD(FShaderResourceParameter, BoundsExtentTexture)
-	LAYOUT_FIELD(FShaderResourceParameter, BoundsExtentSampler)
-
-	FHZBTestPS(const ShaderMetaType::CompiledShaderInitializerType& Initializer)
-		: FGlobalShader(Initializer)
-	{
-		HZBUvFactor.Bind( Initializer.ParameterMap, TEXT("HZBUvFactor") );
-		HZBViewSize.Bind( Initializer.ParameterMap, TEXT("HZBViewSize") );
-		HZBSize.Bind( Initializer.ParameterMap, TEXT("HZBSize") );
-		HZBTexture.Bind( Initializer.ParameterMap, TEXT("HZBTexture") );
-		HZBSampler.Bind( Initializer.ParameterMap, TEXT("HZBSampler") );
-		BoundsCenterTexture.Bind( Initializer.ParameterMap, TEXT("BoundsCenterTexture") );
-		BoundsCenterSampler.Bind( Initializer.ParameterMap, TEXT("BoundsCenterSampler") );
-		BoundsExtentTexture.Bind( Initializer.ParameterMap, TEXT("BoundsExtentTexture") );
-		BoundsExtentSampler.Bind( Initializer.ParameterMap, TEXT("BoundsExtentSampler") );
-	}
+		SHADER_PARAMETER(FVector2f, InvTestTargetSize)
+		SHADER_PARAMETER(FVector2f, HZBSize)
+		SHADER_PARAMETER(FVector2f, HZBViewSize)
 
-	void SetParameters(FRHICommandList& RHICmdList, const FViewInfo& View, FRHITexture* BoundsCenter, FRHITexture* BoundsExtent )
-	{
-		FRHIPixelShader* ShaderRHI = RHICmdList.GetBoundPixelShader();
+		RENDER_TARGET_BINDING_SLOTS()
+	END_SHADER_PARAMETER_STRUCT()
 
-		FGlobalShader::SetParameters<FViewUniformShaderParameters>(RHICmdList, ShaderRHI, View.ViewUniformBuffer);
-
-		/*
-		 * Defines the maximum number of mipmaps the HZB test is considering
-		 * to avoid memory cache trashing when rendering on high resolution.
-		 */
-		const float kHZBTestMaxMipmap = 9.0f;
-
-		const float HZBMipmapCounts = FMath::Log2(FMath::Max<float>(View.HZBMipmap0Size.X, View.HZBMipmap0Size.Y));
-		const FVector3f HZBUvFactorValue(
-			float(View.ViewRect.Width()) / float(2 * View.HZBMipmap0Size.X),
-			float(View.ViewRect.Height()) / float(2 * View.HZBMipmap0Size.Y),
-			FMath::Max(HZBMipmapCounts - kHZBTestMaxMipmap, 0.0f)
-			);
-		const FVector4f HZBSizeValue(
-			View.HZBMipmap0Size.X,
-			View.HZBMipmap0Size.Y,
-			1.0f / float(View.HZBMipmap0Size.X),
-			1.0f / float(View.HZBMipmap0Size.Y)
-			);
-		SetShaderValue(RHICmdList, ShaderRHI, HZBUvFactor, HZBUvFactorValue);
-		SetShaderValue(RHICmdList, ShaderRHI, HZBViewSize, View.ViewRect.Size());
-		SetShaderValue(RHICmdList, ShaderRHI, HZBSize, HZBSizeValue);
-
-		SetTextureParameter(RHICmdList, ShaderRHI, HZBTexture, HZBSampler, TStaticSamplerState<SF_Point,AM_Clamp,AM_Clamp,AM_Clamp>::GetRHI(), View.HZB->GetRHI() );
-
-		SetTextureParameter(RHICmdList, ShaderRHI, BoundsCenterTexture, BoundsCenterSampler, TStaticSamplerState<SF_Point,AM_Clamp,AM_Clamp,AM_Clamp>::GetRHI(), BoundsCenter );
-		SetTextureParameter(RHICmdList, ShaderRHI, BoundsExtentTexture, BoundsExtentSampler, TStaticSamplerState<SF_Point,AM_Clamp,AM_Clamp,AM_Clamp>::GetRHI(), BoundsExtent );
+	static bool ShouldCompilePermutation(const FGlobalShaderPermutationParameters& Parameters)
+	{
+		return IsFeatureLevelSupported(Parameters.Platform, ERHIFeatureLevel::SM5);
 	}
 };
 
-IMPLEMENT_SHADER_TYPE(,FHZBTestPS,TEXT("/Engine/Private/HZBOcclusion.usf"),TEXT("HZBTestPS"),SF_Pixel);
+IMPLEMENT_GLOBAL_SHADER(FHZBTestPS, "/Engine/Private/HZBOcclusion.usf", "HZBTestPS", SF_Pixel);
 
 BEGIN_SHADER_PARAMETER_STRUCT(FHZBOcclusionUpdateTexturesParameters, )
 	RDG_TEXTURE_ACCESS(BoundsCenterTexture, ERHIAccess::CopyDest)
 	RDG_TEXTURE_ACCESS(BoundsExtentTexture, ERHIAccess::CopyDest)
 END_SHADER_PARAMETER_STRUCT()
 
-BEGIN_SHADER_PARAMETER_STRUCT(FHZBOcclusionTestHZBParameters, )
-	RDG_TEXTURE_ACCESS(BoundsCenterTexture, ERHIAccess::SRVGraphics)
-	RDG_TEXTURE_ACCESS(BoundsExtentTexture, ERHIAccess::SRVGraphics)
-	RDG_TEXTURE_ACCESS(HZBTexture, ERHIAccess::SRVGraphics)
-	RENDER_TARGET_BINDING_SLOTS()
-END_SHADER_PARAMETER_STRUCT()
-
 void FHZBOcclusionTester::Submit(FRDGBuilder& GraphBuilder, const FViewInfo& View)
 {
 	RDG_EVENT_SCOPE(GraphBuilder, "SubmitHZB");
@@ -1059,59 +1014,38 @@ void FHZBOcclusionTester::Submit(FRDGBuilder& GraphBuilder, const FViewInfo& Vie
 
 	// Draw test
 	{
-		auto* PassParameters = GraphBuilder.AllocParameters<FHZBOcclusionTestHZBParameters>();
+		FHZBTestPS::FParameters* PassParameters = GraphBuilder.AllocParameters<FHZBTestPS::FParameters>();
+		PassParameters->View = View.ViewUniformBuffer;
 		PassParameters->BoundsCenterTexture = BoundsCenterTexture;
 		PassParameters->BoundsExtentTexture = BoundsExtentTexture;
 		PassParameters->HZBTexture = View.HZB;
 		PassParameters->RenderTargets[0] = FRenderTargetBinding(ResultsTextureGPU, ERenderTargetLoadAction::ENoAction);
 
-		GraphBuilder.AddPass(
-			RDG_EVENT_NAME("TestHZB"),
-			PassParameters,
-			ERDGPassFlags::Raster,
-			[this, &View, BoundsCenterTexture, BoundsExtentTexture](FRHICommandList& RHICmdList)
-		{
-			FGraphicsPipelineStateInitializer GraphicsPSOInit;
-			RHICmdList.ApplyCachedRenderTargets(GraphicsPSOInit);
-			GraphicsPSOInit.BlendState = TStaticBlendState<>::GetRHI();
-			GraphicsPSOInit.RasterizerState = TStaticRasterizerState<>::GetRHI();
-			GraphicsPSOInit.DepthStencilState = TStaticDepthStencilState<false, CF_Always>::GetRHI();
-
-			TShaderMapRef< FScreenVS >	VertexShader(View.ShaderMap);
-			TShaderMapRef< FHZBTestPS >	PixelShader(View.ShaderMap);
+		PassParameters->HZBSize = FVector2f(View.HZBMipmap0Size);
+		PassParameters->HZBViewSize = FVector2f(View.ViewRect.Size());
 
-			GraphicsPSOInit.BoundShaderState.VertexDeclarationRHI = GFilterVertexDeclaration.VertexDeclarationRHI;
-			GraphicsPSOInit.BoundShaderState.VertexShaderRHI = VertexShader.GetVertexShader();
-			GraphicsPSOInit.BoundShaderState.PixelShaderRHI = PixelShader.GetPixelShader();
-			GraphicsPSOInit.PrimitiveType = PT_TriangleList;
-
-			SetGraphicsPipelineState(RHICmdList, GraphicsPSOInit, 0);
+		PassParameters->HZBSampler = TStaticSamplerState<SF_Point, AM_Clamp, AM_Clamp, AM_Clamp>::GetRHI();
+		PassParameters->BoundsCenterSampler = TStaticSamplerState<SF_Point, AM_Clamp, AM_Clamp, AM_Clamp>::GetRHI();
+		PassParameters->BoundsExtentSampler = TStaticSamplerState<SF_Point, AM_Clamp, AM_Clamp, AM_Clamp>::GetRHI();
 
-			PixelShader->SetParameters(RHICmdList, View, BoundsCenterTexture->GetRHI(), BoundsExtentTexture->GetRHI());
+		PassParameters->InvTestTargetSize = FVector2f(1.0f / float(SizeX), 1.0f / float(SizeY));
 
-			RHICmdList.SetViewport(0, 0, 0.0f, SizeX, SizeY, 1.0f);
+		TShaderMapRef< FHZBTestPS >	PixelShader(View.ShaderMap);
 
-			// TODO draw quads covering blocks added above
-			DrawRectangle(
-				RHICmdList,
-				0, 0,
-				SizeX, SizeY,
-				0, 0,
-				SizeX, SizeY,
-				FIntPoint(SizeX, SizeY),
-				FIntPoint(SizeX, SizeY),
-				VertexShader,
-				EDRF_UseTriangleOptimization);
-		});
+		FPixelShaderUtils::AddFullscreenPass(
+			GraphBuilder,
+			View.ShaderMap,
+			RDG_EVENT_NAME("TestHZB"),
+			PixelShader,
+			PassParameters,
+			FIntRect(0, 0, SizeX, SizeY),
+			TStaticBlendState<>::GetRHI(),
+			TStaticRasterizerState<>::GetRHI(),
+			TStaticDepthStencilState<false, CF_Always>::GetRHI());
 	}
 
 	// Transfer memory GPU -> CPU
-	AddCopyToResolveTargetPass(GraphBuilder, ResultsTextureGPU, GraphBuilder.RegisterExternalTexture(ResultsTextureCPU), FResolveParams());
-
-	AddPass(GraphBuilder, RDG_EVENT_NAME("WriteGPUFence"), [this](FRHICommandList& RHICmdList)
-	{
-		RHICmdList.WriteGPUFence(Fence);
-	});
+	AddEnqueueCopyPass(GraphBuilder, ResultsReadback.Get(), ResultsTextureGPU);
 }
 
 static FViewOcclusionQueriesPerView AllocateOcclusionTests(const FScene* Scene, TArrayView<const FVisibleLightInfo> VisibleLightInfos, TArrayView<FViewInfo> Views)
diff --git a/Engine/Source/Runtime/Renderer/Private/SceneRendering.h b/Engine/Source/Runtime/Renderer/Private/SceneRendering.h
--- a/Engine/Source/Runtime/Renderer/Private/SceneRendering.h
+++ b/Engine/Source/Runtime/Renderer/Private/SceneRendering.h
@@ -481,9 +481,9 @@ private:
 
 	TArray< FOcclusionPrimitive, SceneRenderingAllocator >	Primitives;
 
-	TRefCountPtr<IPooledRenderTarget>	ResultsTextureCPU;
 	const uint8*						ResultsBuffer;
-
+	int32								ResultsBufferRowPitch;
+	TUniquePtr<FRHIGPUTextureReadback>	ResultsReadback;
 
 	bool IsInvalidFrame() const;
 
@@ -491,7 +491,6 @@ private:
 	void SetInvalidFrameNumber();
 
 	uint32 ValidFrameNumber;
-	FGPUFenceRHIRef Fence;
 };
 
 DECLARE_STATS_GROUP(TEXT("Parallel Command List Markers"), STATGROUP_ParallelCommandListMarkers, STATCAT_Advanced);
