Fix for static mesh with static lighting not being rendered in the  static skylight due to missing code setting PassMask.

diff --git a/Engine/Source/Runtime/Renderer/Private/SceneVisibility.cpp b/Engine/Source/Runtime/Renderer/Private/SceneVisibility.cpp
--- a/Engine/Source/Runtime/Renderer/Private/SceneVisibility.cpp
+++ b/Engine/Source/Runtime/Renderer/Private/SceneVisibility.cpp
@@ -3121,6 +3121,12 @@ void ComputeDynamicMeshRelevance(EShadingPath ShadingPath, bool bAddLightmapDens
 		{
 			PassMask.Set(EMeshPass::BasePass);
 			View.NumVisibleDynamicMeshElements[EMeshPass::BasePass] += NumElements;
+			
+			if (ViewRelevance.bUsesSkyMaterial)
+			{
+				PassMask.Set(EMeshPass::SkyPass);
+				View.NumVisibleDynamicMeshElements[EMeshPass::SkyPass] += NumElements;
+			}
 
 			if (ViewRelevance.bUsesAnisotropy)
 			{
