Converting  FVolumetricFogLightScatteringCS fully to RDG
Green volumetric fog screnshot tests.

diff --git a/Engine/Shaders/Private/VolumetricFog.usf b/Engine/Shaders/Private/VolumetricFog.usf
--- a/Engine/Shaders/Private/VolumetricFog.usf
+++ b/Engine/Shaders/Private/VolumetricFog.usf
@@ -265,7 +265,7 @@ void InjectShadowedLocalLightPS(
 
 	// Somehow pixels are being rasterized outside of the viewport on a 970 GTX, perhaps due to use of a geometry shader bypassing the viewport scissor.
 	// This triggers the HistoryMissSuperSampleCount path causing significant overhead for shading off-screen pixels.
-	if (all(GridCoordinate < VolumetricFog.GridSizeInt))
+	if (all(int3(GridCoordinate) < VolumetricFog.GridSizeInt))
 	{
 		FDeferredLightData LightData;
 		LightData.Position = DeferredLightUniforms.Position;
@@ -560,7 +560,7 @@ void LightScatteringCS(
 
 	// Supersample if the history was outside the camera frustum
 	// The compute shader is dispatched with extra threads, make sure those don't supersample
-	NumSuperSamples = HistoryAlpha < .001f && all(GridCoordinate < VolumetricFog.GridSizeInt) ? HISTORY_MISS_SUPER_SAMPLE_COUNT : 1;
+	NumSuperSamples = HistoryAlpha < .001f && all(int3(GridCoordinate) < VolumetricFog.GridSizeInt) ? HISTORY_MISS_SUPER_SAMPLE_COUNT : 1;
 #endif
 
 	for (uint SampleIndex = 0; SampleIndex < NumSuperSamples; SampleIndex++)
diff --git a/Engine/Source/Runtime/Renderer/Private/VolumetricFog.cpp b/Engine/Source/Runtime/Renderer/Private/VolumetricFog.cpp
--- a/Engine/Source/Runtime/Renderer/Private/VolumetricFog.cpp
+++ b/Engine/Source/Runtime/Renderer/Private/VolumetricFog.cpp
@@ -773,7 +773,8 @@ IMPLEMENT_GLOBAL_SHADER_PARAMETER_STRUCT(FLumenTranslucencyLightingUniforms, "Lu
 
 class FVolumetricFogLightScatteringCS : public FGlobalShader
 {
-	DECLARE_SHADER_TYPE(FVolumetricFogLightScatteringCS, Global)
+	DECLARE_GLOBAL_SHADER(FVolumetricFogLightScatteringCS);
+	SHADER_USE_PARAMETER_STRUCT(FVolumetricFogLightScatteringCS, FGlobalShader);
 
 	class FTemporalReprojection			: SHADER_PERMUTATION_BOOL("USE_TEMPORAL_REPROJECTION");
 	class FDistanceFieldSkyOcclusion	: SHADER_PERMUTATION_BOOL("DISTANCE_FIELD_SKY_OCCLUSION");
@@ -792,6 +793,7 @@ class FVolumetricFogLightScatteringCS : public FGlobalShader
 
 	BEGIN_SHADER_PARAMETER_STRUCT(FParameters, )
 		SHADER_PARAMETER_STRUCT_REF(FViewUniformShaderParameters, View)
+		SHADER_PARAMETER_STRUCT_REF(FForwardLightData, Forward)
 		SHADER_PARAMETER_RDG_UNIFORM_BUFFER(FFogUniformParameters, Fog)
 		SHADER_PARAMETER_STRUCT_INCLUDE(FVolumetricFogIntegrationParameters, VolumetricFogParameters)
 
@@ -799,20 +801,36 @@ class FVolumetricFogLightScatteringCS : public FGlobalShader
 		SHADER_PARAMETER_RDG_TEXTURE(Texture2D, VBufferB)
 		SHADER_PARAMETER_RDG_TEXTURE(Texture2D, LocalShadowedLightScattering)
 		SHADER_PARAMETER_RDG_TEXTURE(Texture2D, LightFunctionTexture)
+		SHADER_PARAMETER_SAMPLER(SamplerState, LightFunctionSampler)
+		SHADER_PARAMETER_RDG_TEXTURE(Texture2D, CloudShadowmapTexture)
+		SHADER_PARAMETER_SAMPLER(SamplerState, CloudShadowmapSampler)
 		SHADER_PARAMETER_RDG_TEXTURE(Texture2D, ConservativeDepthTexture)
 		SHADER_PARAMETER_RDG_TEXTURE(Texture2D, PrevConservativeDepthTexture)
-		SHADER_PARAMETER(FVector2f, PrevConservativeDepthTextureSize)
-		SHADER_PARAMETER(uint32, UseConservativeDepthTexture)
-		SHADER_PARAMETER(uint32, UseEmissive)
+		SHADER_PARAMETER_RDG_TEXTURE(Texture3D, LightScatteringHistory)
+		SHADER_PARAMETER_SAMPLER(SamplerState, LightScatteringHistorySampler)
 		SHADER_PARAMETER_RDG_UNIFORM_BUFFER(FLumenTranslucencyLightingUniforms, LumenGIVolumeStruct)
 		SHADER_PARAMETER_STRUCT_INCLUDE(FVirtualShadowMapSamplingParameters, VirtualShadowMapSamplingParameters)
 		SHADER_PARAMETER_STRUCT_INCLUDE(FAOParameters, AOParameters)
 		SHADER_PARAMETER_STRUCT_INCLUDE(FGlobalDistanceFieldParameters2, GlobalDistanceFieldParameters)
 		SHADER_PARAMETER_RDG_TEXTURE_UAV(RWTexture2D, RWLightScattering)
+		SHADER_PARAMETER_ARRAY(FVector4f, SkySH, [3])
+		SHADER_PARAMETER(FMatrix44f, DirectionalLightFunctionWorldToShadow)
+		SHADER_PARAMETER(FMatrix44f, CloudShadowmapWorldToLightClipMatrix)
+		SHADER_PARAMETER(FVector2f, PrevConservativeDepthTextureSize)
+		SHADER_PARAMETER(FVector2f, UseHeightFogColors)
+		SHADER_PARAMETER(float, StaticLightingScatteringIntensity)
+		SHADER_PARAMETER(float, SkyLightVolumetricScatteringIntensity)
+		SHADER_PARAMETER(float, SkyLightUseStaticShadowing)
+		SHADER_PARAMETER(float, PhaseG)
+		SHADER_PARAMETER(float, InverseSquaredLightDistanceBiasScale)
+		SHADER_PARAMETER(float, LightScatteringSampleJitterMultiplier)
+		SHADER_PARAMETER(float, CloudShadowmapFarDepthKm)
+		SHADER_PARAMETER(float, CloudShadowmapStrength)
+		SHADER_PARAMETER(float, UseDirectionalLightShadowing)
+		SHADER_PARAMETER(uint32, UseConservativeDepthTexture)
+		SHADER_PARAMETER(uint32, UseEmissive)
 	END_SHADER_PARAMETER_STRUCT()
 
-public:
-
 	static bool ShouldCompilePermutation(const FGlobalShaderPermutationParameters& Parameters)
 	{
 		return DoesPlatformSupportVolumetricFog(Parameters.Platform);
@@ -850,188 +868,6 @@ public:
 		FForwardLightingParameters::ModifyCompilationEnvironment(Parameters.Platform, OutEnvironment);
 		FVirtualShadowMapArray::SetShaderDefines(OutEnvironment);
 	}
-
-	FVolumetricFogLightScatteringCS(const ShaderMetaType::CompiledShaderInitializerType& Initializer)
-		: FGlobalShader(Initializer)
-	{
-		Bindings.BindForLegacyShaderParameters(this, Initializer.PermutationId, Initializer.ParameterMap, *FParameters::FTypeInfo::GetStructMetadata());
-
-		LocalShadowedLightScattering.Bind(Initializer.ParameterMap, TEXT("LocalShadowedLightScattering"));
-		LightScatteringHistory.Bind(Initializer.ParameterMap, TEXT("LightScatteringHistory"));
-		LightScatteringHistorySampler.Bind(Initializer.ParameterMap, TEXT("LightScatteringHistorySampler"));
-		DirectionalLightFunctionWorldToShadow.Bind(Initializer.ParameterMap, TEXT("DirectionalLightFunctionWorldToShadow"));
-		LightFunctionTexture.Bind(Initializer.ParameterMap, TEXT("LightFunctionTexture"));
-		LightFunctionSampler.Bind(Initializer.ParameterMap, TEXT("LightFunctionSampler"));
-		ConservativeDepthTexture.Bind(Initializer.ParameterMap, TEXT("ConservativeDepthTexture"));
-		PrevConservativeDepthTexture.Bind(Initializer.ParameterMap, TEXT("PrevConservativeDepthTexture"));
-		PrevConservativeDepthTextureSize.Bind(Initializer.ParameterMap, TEXT("PrevConservativeDepthTextureSize"));
-		UseConservativeDepthTexture.Bind(Initializer.ParameterMap, TEXT("UseConservativeDepthTexture"));
-		StaticLightingScatteringIntensity.Bind(Initializer.ParameterMap, TEXT("StaticLightingScatteringIntensity"));
-		SkyLightUseStaticShadowing.Bind(Initializer.ParameterMap, TEXT("SkyLightUseStaticShadowing"));
-		SkyLightVolumetricScatteringIntensity.Bind(Initializer.ParameterMap, TEXT("SkyLightVolumetricScatteringIntensity"));
-		SkySH.Bind(Initializer.ParameterMap, TEXT("SkySH"));
-		PhaseG.Bind(Initializer.ParameterMap, TEXT("PhaseG"));
-		InverseSquaredLightDistanceBiasScale.Bind(Initializer.ParameterMap, TEXT("InverseSquaredLightDistanceBiasScale"));
-		UseHeightFogColors.Bind(Initializer.ParameterMap, TEXT("UseHeightFogColors"));
-		UseDirectionalLightShadowing.Bind(Initializer.ParameterMap, TEXT("UseDirectionalLightShadowing"));
-		LightScatteringSampleJitterMultiplier.Bind(Initializer.ParameterMap, TEXT("LightScatteringSampleJitterMultiplier"));
-
-		CloudShadowmapTexture.Bind(Initializer.ParameterMap, TEXT("CloudShadowmapTexture"));
-		CloudShadowmapSampler.Bind(Initializer.ParameterMap, TEXT("CloudShadowmapSampler"));
-		CloudShadowmapFarDepthKm.Bind(Initializer.ParameterMap, TEXT("CloudShadowmapFarDepthKm"));
-		CloudShadowmapWorldToLightClipMatrix.Bind(Initializer.ParameterMap, TEXT("CloudShadowmapWorldToLightClipMatrix"));
-		CloudShadowmapStrength.Bind(Initializer.ParameterMap, TEXT("CloudShadowmapStrength"));
-	}
-
-	FVolumetricFogLightScatteringCS()
-	{
-	}
-
-	void SetParameters(
-		FRHICommandList& RHICmdList,
-		const FViewInfo& View,
-		const FVolumetricFogIntegrationParameterData& IntegrationData,
-		const FExponentialHeightFogSceneInfo& FogInfo,
-		FRHITexture* LightScatteringHistoryTexture,
-		bool bUseDirectionalLightShadowing,
-		const FMatrix44f& DirectionalLightFunctionWorldToShadowValue,
-		const int AtmosphericDirectionalLightIndex,
-		const FLightSceneProxy* AtmosphereLightProxy,
-		const FVolumetricCloudRenderSceneInfo* CloudInfo)
-	{
-		FRHIComputeShader* ShaderRHI = RHICmdList.GetBoundComputeShader();
-
-		if (!LightScatteringHistoryTexture)
-		{
-			LightScatteringHistoryTexture = GBlackVolumeTexture->TextureRHI;
-		}
-
-		SetTextureParameter(
-			RHICmdList,
-			ShaderRHI,
-			LightScatteringHistory,
-			LightScatteringHistorySampler,
-			TStaticSamplerState<SF_Bilinear, AM_Clamp, AM_Clamp, AM_Clamp>::GetRHI(),
-			LightScatteringHistoryTexture);
-
-		SetUniformBufferParameter(RHICmdList, ShaderRHI, GetUniformBufferParameter<FForwardLightData>(), View.ForwardLightingResources->ForwardLightDataUniformBuffer);
-
-		SetShaderValue(RHICmdList, ShaderRHI, DirectionalLightFunctionWorldToShadow, DirectionalLightFunctionWorldToShadowValue);
-
-		SetSamplerParameter(RHICmdList, ShaderRHI, LightFunctionSampler, TStaticSamplerState<SF_Bilinear, AM_Clamp, AM_Clamp, AM_Clamp>::GetRHI());
-
-		FScene* Scene = (FScene*)View.Family->Scene;
-		FSkyLightSceneProxy* SkyLight = Scene->SkyLight;
-
-		if (SkyLight
-			// Skylights with static lighting had their diffuse contribution baked into lightmaps
-			&& !SkyLight->bHasStaticLighting
-			&& View.Family->EngineShowFlags.SkyLighting)
-		{
-			const float LocalSkyLightUseStaticShadowing = SkyLight->bWantsStaticShadowing && SkyLight->bCastShadows ? 1.0f : 0.0f;
-			SetShaderValue(RHICmdList, ShaderRHI, SkyLightUseStaticShadowing, LocalSkyLightUseStaticShadowing);
-			SetShaderValue(RHICmdList, ShaderRHI, SkyLightVolumetricScatteringIntensity, SkyLight->VolumetricScatteringIntensity);
-
-			const FSHVectorRGB3& SkyIrradiance = SkyLight->IrradianceEnvironmentMap;
-			SetShaderValue(RHICmdList, ShaderRHI, SkySH, (FVector4f&)SkyIrradiance.R.V, 0);
-			SetShaderValue(RHICmdList, ShaderRHI, SkySH, (FVector4f&)SkyIrradiance.G.V, 1);
-			SetShaderValue(RHICmdList, ShaderRHI, SkySH, (FVector4f&)SkyIrradiance.B.V, 2);
-		}
-		else
-		{
-			SetShaderValue(RHICmdList, ShaderRHI, SkyLightUseStaticShadowing, 0.0f);
-			SetShaderValue(RHICmdList, ShaderRHI, SkyLightVolumetricScatteringIntensity, 0.0f);
-			SetShaderValue(RHICmdList, ShaderRHI, SkySH, FVector4f(0, 0, 0, 0), 0);
-			SetShaderValue(RHICmdList, ShaderRHI, SkySH, FVector4f(0, 0, 0, 0), 1);
-			SetShaderValue(RHICmdList, ShaderRHI, SkySH, FVector4f(0, 0, 0, 0), 2);
-		}
-
-		float StaticLightingScatteringIntensityValue = 0;
-
-		if (View.Family->EngineShowFlags.GlobalIllumination && View.Family->EngineShowFlags.VolumetricLightmap)
-		{
-			StaticLightingScatteringIntensityValue = FogInfo.VolumetricFogStaticLightingScatteringIntensity;
-		}
-
-		SetShaderValue(RHICmdList, ShaderRHI, StaticLightingScatteringIntensity, StaticLightingScatteringIntensityValue);
-
-		SetShaderValue(RHICmdList, ShaderRHI, PhaseG, FogInfo.VolumetricFogScatteringDistribution);
-		SetShaderValue(RHICmdList, ShaderRHI, InverseSquaredLightDistanceBiasScale, GInverseSquaredLightDistanceBiasScale);
-		SetShaderValue(RHICmdList, ShaderRHI, UseDirectionalLightShadowing, bUseDirectionalLightShadowing ? 1.0f : 0.0f);
-
-		SetShaderValue(RHICmdList, ShaderRHI, UseHeightFogColors, FVector2f(
-			OverrideDirectionalLightInScatteringUsingHeightFog(View, FogInfo) ? 1.0f : 0.0f,
-			OverrideSkyLightInScatteringUsingHeightFog(View, FogInfo) ? 1.0f : 0.0f ));
-
-		if (CloudShadowmapTexture.IsBound())
-		{
-			FMatrix CloudWorldToLightClipShadowMatrix = FMatrix::Identity;
-			float CloudShadowmap_FarDepthKm = 0.0f;
-			float CloudShadowmap_Strength = 0.0f;
-			IPooledRenderTarget* CloudShadowmap_Texture = nullptr;
-			if (CloudInfo && AtmosphericDirectionalLightIndex >= 0 && AtmosphereLightProxy)
-			{
-				CloudShadowmap_Texture = View.VolumetricCloudShadowExtractedRenderTarget[AtmosphericDirectionalLightIndex];
-				CloudWorldToLightClipShadowMatrix = CloudInfo->GetVolumetricCloudCommonShaderParameters().CloudShadowmapWorldToLightClipMatrix[AtmosphericDirectionalLightIndex];
-				CloudShadowmap_FarDepthKm = CloudInfo->GetVolumetricCloudCommonShaderParameters().CloudShadowmapFarDepthKm[AtmosphericDirectionalLightIndex].X;
-				CloudShadowmap_Strength = AtmosphereLightProxy->GetCloudShadowOnSurfaceStrength();
-			}
-
-			SetTextureParameter(
-				RHICmdList,
-				ShaderRHI,
-				CloudShadowmapTexture,
-				CloudShadowmapSampler,
-				TStaticSamplerState<SF_Bilinear, AM_Clamp, AM_Clamp, AM_Clamp>::GetRHI(),
-				CloudShadowmap_Texture ? CloudShadowmap_Texture->GetRenderTargetItem().ShaderResourceTexture : GBlackTexture->TextureRHI);
-
-			SetShaderValue(
-				RHICmdList,
-				ShaderRHI,
-				CloudShadowmapFarDepthKm,
-				CloudShadowmap_FarDepthKm);
-
-			SetShaderValue(
-				RHICmdList,
-				ShaderRHI,
-				CloudShadowmapWorldToLightClipMatrix,
-				(FMatrix44f)CloudWorldToLightClipShadowMatrix);
-
-			SetShaderValue(
-				RHICmdList,
-				ShaderRHI,
-				CloudShadowmapStrength,
-				CloudShadowmap_Strength);
-		}
-
-		SetShaderValue(RHICmdList, ShaderRHI, LightScatteringSampleJitterMultiplier, GVolumetricFogJitter ? GLightScatteringSampleJitterMultiplier : 0);
-	}
-
-private:
-	LAYOUT_FIELD(FShaderResourceParameter, LocalShadowedLightScattering);
-	LAYOUT_FIELD(FShaderResourceParameter, LightScatteringHistory);
-	LAYOUT_FIELD(FShaderResourceParameter, LightScatteringHistorySampler);
-	LAYOUT_FIELD(FShaderParameter, DirectionalLightFunctionWorldToShadow);
-	LAYOUT_FIELD(FShaderResourceParameter, LightFunctionTexture);
-	LAYOUT_FIELD(FShaderResourceParameter, LightFunctionSampler);
-	LAYOUT_FIELD(FShaderParameter, StaticLightingScatteringIntensity);
-	LAYOUT_FIELD(FShaderParameter, SkyLightUseStaticShadowing);
-	LAYOUT_FIELD(FShaderParameter, SkyLightVolumetricScatteringIntensity);
-	LAYOUT_FIELD(FShaderParameter, SkySH);
-	LAYOUT_FIELD(FShaderParameter, PhaseG);
-	LAYOUT_FIELD(FShaderParameter, InverseSquaredLightDistanceBiasScale);
-	LAYOUT_FIELD(FShaderParameter, UseHeightFogColors);
-	LAYOUT_FIELD(FShaderParameter, UseDirectionalLightShadowing);
-	LAYOUT_FIELD(FShaderResourceParameter, CloudShadowmapTexture);
-	LAYOUT_FIELD(FShaderResourceParameter, CloudShadowmapSampler);
-	LAYOUT_FIELD(FShaderParameter, CloudShadowmapFarDepthKm);
-	LAYOUT_FIELD(FShaderParameter, CloudShadowmapWorldToLightClipMatrix);
-	LAYOUT_FIELD(FShaderParameter, CloudShadowmapStrength);
-	LAYOUT_FIELD(FShaderResourceParameter, ConservativeDepthTexture);
-	LAYOUT_FIELD(FShaderResourceParameter, PrevConservativeDepthTexture);
-	LAYOUT_FIELD(FShaderParameter, PrevConservativeDepthTextureSize);
-	LAYOUT_FIELD(FShaderParameter, UseConservativeDepthTexture);
-	LAYOUT_FIELD(FShaderParameter, LightScatteringSampleJitterMultiplier)
 };
 
 IMPLEMENT_GLOBAL_SHADER(FVolumetricFogLightScatteringCS, "/Engine/Private/VolumetricFog.usf", "LightScatteringCS", SF_Compute);
@@ -1419,13 +1255,13 @@ void FDeferredShadingSceneRenderer::ComputeVolumetricFog(FRDGBuilder& GraphBuild
 			FVolumetricFogLightScatteringCS::FParameters* PassParameters = GraphBuilder.AllocParameters<FVolumetricFogLightScatteringCS::FParameters>();
 
 			PassParameters->View = View.ViewUniformBuffer;
+			PassParameters->Forward = View.ForwardLightingResources->ForwardLightDataUniformBuffer;
 			PassParameters->Fog = FogUniformBuffer;
 			SetupVolumetricFogIntegrationParameters(PassParameters->VolumetricFogParameters, View, IntegrationData);
 
 			PassParameters->VBufferA = IntegrationData.VBufferA;
 			PassParameters->VBufferB = IntegrationData.VBufferB ? IntegrationData.VBufferB : VolumetricBlackDummyTexture;
 			PassParameters->LocalShadowedLightScattering = LocalShadowedLightScattering;
-			PassParameters->LightFunctionTexture = DirectionalLightFunctionTexture;
 			PassParameters->ConservativeDepthTexture = ConservativeDepthTexture;
 			PassParameters->UseConservativeDepthTexture = GVolumetricFogConservativeDepth > 0 ? 1 : 0;
 			PassParameters->UseEmissive = bUseEmissive ? 1 : 0;
@@ -1441,6 +1277,10 @@ void FDeferredShadingSceneRenderer::ComputeVolumetricFog(FRDGBuilder& GraphBuild
 				PassParameters->PrevConservativeDepthTextureSize = FVector2D(1, 1);
 			}
 
+			PassParameters->DirectionalLightFunctionWorldToShadow = DirectionalLightFunctionWorldToShadow;
+			PassParameters->LightFunctionTexture = DirectionalLightFunctionTexture;
+			PassParameters->LightFunctionSampler = TStaticSamplerState<SF_Bilinear, AM_Clamp, AM_Clamp, AM_Clamp>::GetRHI();
+
 			auto* LumenUniforms = GraphBuilder.AllocParameters<FLumenTranslucencyLightingUniforms>();
 			LumenUniforms->Parameters = GetLumenTranslucencyLightingParameters(GraphBuilder, View.LumenTranslucencyGIVolume);
 			PassParameters->LumenGIVolumeStruct = GraphBuilder.CreateUniformBuffer(LumenUniforms);
@@ -1458,6 +1298,70 @@ void FDeferredShadingSceneRenderer::ComputeVolumetricFog(FRDGBuilder& GraphBuild
 			PassParameters->AOParameters = DistanceField::SetupAOShaderParameters(AOParameterData);
 			PassParameters->GlobalDistanceFieldParameters = SetupGlobalDistanceFieldParameters(View.GlobalDistanceFieldInfo.ParameterData);
 
+			FVolumetricCloudRenderSceneInfo* CloudInfo = Scene->GetVolumetricCloudSceneInfo();
+			FRDGTexture* LightScatteringHistoryRDGTexture = BlackDummyTexture;
+			if (bUseTemporalReprojection && View.ViewState->LightScatteringHistory.IsValid())
+			{
+				LightScatteringHistoryRDGTexture = GraphBuilder.RegisterExternalTexture(View.ViewState->LightScatteringHistory);
+			}
+
+			PassParameters->LightScatteringHistory = LightScatteringHistoryRDGTexture;
+			PassParameters->LightScatteringHistorySampler = TStaticSamplerState<SF_Bilinear, AM_Clamp, AM_Clamp, AM_Clamp>::GetRHI();
+
+			FSkyLightSceneProxy* SkyLight = Scene->SkyLight;
+			if (SkyLight
+				// Skylights with static lighting had their diffuse contribution baked into lightmaps
+				&& !SkyLight->bHasStaticLighting
+				&& View.Family->EngineShowFlags.SkyLighting)
+			{
+				PassParameters->SkyLightUseStaticShadowing = SkyLight->bWantsStaticShadowing && SkyLight->bCastShadows ? 1.0f : 0.0f;
+				PassParameters->SkyLightVolumetricScatteringIntensity = SkyLight->VolumetricScatteringIntensity;
+
+				const FSHVectorRGB3& SkyIrradiance = SkyLight->IrradianceEnvironmentMap;
+				PassParameters->SkySH[0] = (FVector4f&)SkyIrradiance.R.V;
+				PassParameters->SkySH[1] = (FVector4f&)SkyIrradiance.G.V;
+				PassParameters->SkySH[2] = (FVector4f&)SkyIrradiance.B.V;
+			}
+			else
+			{
+				PassParameters->SkyLightUseStaticShadowing = 0.0f;
+				PassParameters->SkyLightVolumetricScatteringIntensity = 0.0f;
+				PassParameters->SkySH[0] = FVector4f(0, 0, 0, 0);
+				PassParameters->SkySH[1] = FVector4f(0, 0, 0, 0);
+				PassParameters->SkySH[2] = FVector4f(0, 0, 0, 0);
+			}
+
+			float StaticLightingScatteringIntensityValue = 0;
+			if (View.Family->EngineShowFlags.GlobalIllumination && View.Family->EngineShowFlags.VolumetricLightmap)
+			{
+				StaticLightingScatteringIntensityValue = FogInfo.VolumetricFogStaticLightingScatteringIntensity;
+			}
+			PassParameters->StaticLightingScatteringIntensity = StaticLightingScatteringIntensityValue;
+
+			PassParameters->PhaseG = FogInfo.VolumetricFogScatteringDistribution;
+			PassParameters->InverseSquaredLightDistanceBiasScale = GInverseSquaredLightDistanceBiasScale;
+			PassParameters->UseDirectionalLightShadowing = bUseDirectionalLightShadowing ? 1.0f : 0.0f;
+			PassParameters->LightScatteringSampleJitterMultiplier = GVolumetricFogJitter ? GLightScatteringSampleJitterMultiplier : 0;
+			PassParameters->UseHeightFogColors = FVector2f(
+				OverrideDirectionalLightInScatteringUsingHeightFog(View, FogInfo) ? 1.0f : 0.0f,
+				OverrideSkyLightInScatteringUsingHeightFog(View, FogInfo) ? 1.0f : 0.0f);
+
+			FMatrix CloudWorldToLightClipShadowMatrix = FMatrix::Identity;
+			float CloudShadowmap_FarDepthKm = 0.0f;
+			float CloudShadowmap_Strength = 0.0f;
+			FRDGTexture* CloudShadowmap_RDGTexture = BlackDummyTexture;
+			if (CloudInfo && AtmosphericDirectionalLightIndex >= 0 && AtmosphereLightProxy)
+			{
+				CloudShadowmap_RDGTexture = GraphBuilder.RegisterExternalTexture(View.VolumetricCloudShadowExtractedRenderTarget[AtmosphericDirectionalLightIndex]);
+				CloudWorldToLightClipShadowMatrix = CloudInfo->GetVolumetricCloudCommonShaderParameters().CloudShadowmapWorldToLightClipMatrix[AtmosphericDirectionalLightIndex];
+				CloudShadowmap_FarDepthKm = CloudInfo->GetVolumetricCloudCommonShaderParameters().CloudShadowmapFarDepthKm[AtmosphericDirectionalLightIndex].X;
+				CloudShadowmap_Strength = AtmosphereLightProxy->GetCloudShadowOnSurfaceStrength();
+			}
+			PassParameters->CloudShadowmapTexture = CloudShadowmap_RDGTexture;
+			PassParameters->CloudShadowmapSampler = TStaticSamplerState<SF_Bilinear, AM_Clamp, AM_Clamp, AM_Clamp>::GetRHI();
+			PassParameters->CloudShadowmapFarDepthKm = CloudShadowmap_FarDepthKm;
+			PassParameters->CloudShadowmapStrength = CloudShadowmap_Strength;
+
 			const bool bUseLumenGI = View.LumenTranslucencyGIVolume.Texture0 != nullptr;
 			const bool bUseGlobalDistanceField = UseGlobalDistanceField() && Scene->DistanceFieldSceneData.NumObjectsInBuffer > 0;
 
@@ -1486,14 +1390,6 @@ void FDeferredShadingSceneRenderer::ComputeVolumetricFog(FRDGBuilder& GraphBuild
 			auto ComputeShader = View.ShaderMap->GetShader< FVolumetricFogLightScatteringCS >(PermutationVector);
 			ClearUnusedGraphResources(ComputeShader, PassParameters);
 
-			FRHITexture* LightScatteringHistoryTexture = GBlackVolumeTexture->TextureRHI;
-			if (bUseTemporalReprojection && View.ViewState->LightScatteringHistory.IsValid())
-			{
-				LightScatteringHistoryTexture = View.ViewState->LightScatteringHistory->GetRenderTargetItem().ShaderResourceTexture;
-			}
-
-			FVolumetricCloudRenderSceneInfo* CloudInfo = Scene->GetVolumetricCloudSceneInfo();
-
 			GraphBuilder.AddPass(
 				RDG_EVENT_NAME("LightScattering %dx%dx%d SS:%d %s %s %s",
 					VolumetricFogGridSize.X,
@@ -1505,14 +1401,12 @@ void FDeferredShadingSceneRenderer::ComputeVolumetricFog(FRDGBuilder& GraphBuild
 					bUseLumenGI ? TEXT("Lumen") : TEXT("")),
 				PassParameters,
 				ERDGPassFlags::Compute,
-				[PassParameters, ComputeShader, &View, this, FogInfo, bUseTemporalReprojection, LightScatteringHistoryTexture, CloudInfo, VolumetricFogGridSize, IntegrationData, bUseDirectionalLightShadowing, bUseDistanceFieldSkyOcclusion, DirectionalLightFunctionWorldToShadow, AtmosphericDirectionalLightIndex, AtmosphereLightProxy](FRHICommandList& RHICmdList)
+				[PassParameters, ComputeShader, &View, this, VolumetricFogGridSize](FRHICommandList& RHICmdList)
 			{
 				const FIntVector NumGroups = FComputeShaderUtils::GetGroupCount(VolumetricFogGridSize, FVolumetricFogLightScatteringCS::GetGroupSize());
 
 				RHICmdList.SetComputeShader(ComputeShader.GetComputeShader());
 
-				ComputeShader->SetParameters(RHICmdList, View, IntegrationData, FogInfo, LightScatteringHistoryTexture, bUseDirectionalLightShadowing, DirectionalLightFunctionWorldToShadow, AtmosphericDirectionalLightIndex, AtmosphereLightProxy, CloudInfo);
-
 				SetShaderParameters(RHICmdList, ComputeShader, ComputeShader.GetComputeShader(), *PassParameters);
 				DispatchComputeShader(RHICmdList, ComputeShader.GetShader(), NumGroups.X, NumGroups.Y, NumGroups.Z);
 				UnsetShaderUAVs(RHICmdList, ComputeShader, ComputeShader.GetComputeShader());
