Renaming TInjectShadowedLocalLightPS to FInjectShadowedLocalLightPS before some refactoring

diff --git a/Engine/Source/Runtime/Renderer/Private/VolumetricFog.cpp b/Engine/Source/Runtime/Renderer/Private/VolumetricFog.cpp
--- a/Engine/Source/Runtime/Renderer/Private/VolumetricFog.cpp
+++ b/Engine/Source/Runtime/Renderer/Private/VolumetricFog.cpp
@@ -259,9 +259,9 @@ IMPLEMENT_GLOBAL_SHADER(FWriteToBoundingSphereVS, "/Engine/Private/VolumetricFog
 
 
 /** Shader that adds direct lighting contribution from the given light to the current volume lighting cascade. */
-class TInjectShadowedLocalLightPS : public FGlobalShader
+class FInjectShadowedLocalLightPS : public FGlobalShader
 {
-	DECLARE_GLOBAL_SHADER(TInjectShadowedLocalLightPS);
+	DECLARE_GLOBAL_SHADER(FInjectShadowedLocalLightPS);
 
 	BEGIN_SHADER_PARAMETER_STRUCT(FParameters, )
 		SHADER_PARAMETER_RDG_TEXTURE(Texture2D, ConservativeDepthTexture)
@@ -297,7 +297,7 @@ class TInjectShadowedLocalLightPS : public FGlobalShader
 		return DoesPlatformSupportVolumetricFog(Parameters.Platform);
 	}
 
-	TInjectShadowedLocalLightPS(const ShaderMetaType::CompiledShaderInitializerType& Initializer) :
+	FInjectShadowedLocalLightPS(const ShaderMetaType::CompiledShaderInitializerType& Initializer) :
 		FGlobalShader(Initializer)
 	{
 		BindForLegacyShaderParameters<FParameters>(this, Initializer.PermutationId, Initializer.ParameterMap);
@@ -316,7 +316,7 @@ class TInjectShadowedLocalLightPS : public FGlobalShader
 		VirtualShadowMapIdParameter.Bind(Initializer.ParameterMap, TEXT("VirtualShadowMapId"));
 	}
 
-	TInjectShadowedLocalLightPS() {}
+	FInjectShadowedLocalLightPS() {}
 
 public:
 	// @param InnerSplitIndex which CSM shadow map level, INDEX_NONE if no directional light
@@ -370,7 +370,7 @@ private:
 	LAYOUT_FIELD(FShaderResourceParameter, LightFunctionAtlasSamplerParam);
 };
 
-IMPLEMENT_GLOBAL_SHADER(TInjectShadowedLocalLightPS, "/Engine/Private/VolumetricFog.usf", "InjectShadowedLocalLightPS", SF_Pixel);
+IMPLEMENT_GLOBAL_SHADER(FInjectShadowedLocalLightPS, "/Engine/Private/VolumetricFog.usf", "InjectShadowedLocalLightPS", SF_Pixel);
 
 void GetVolumeShadowingShaderParameters(
 	FRDGBuilder& GraphBuilder,
@@ -656,7 +656,7 @@ void FDeferredShadingSceneRenderer::RenderLocalLightsForVolumetricFog(
 	{
 		OutLocalShadowedLightScattering = GraphBuilder.CreateTexture(VolumeDesc, TEXT("VolumetricFog.LocalShadowedLightScattering"));
 
-		TInjectShadowedLocalLightPS::FParameters* PassParameters = GraphBuilder.AllocParameters<TInjectShadowedLocalLightPS::FParameters>();
+		FInjectShadowedLocalLightPS::FParameters* PassParameters = GraphBuilder.AllocParameters<FInjectShadowedLocalLightPS::FParameters>();
 		PassParameters->RenderTargets[0] = FRenderTargetBinding(OutLocalShadowedLightScattering, ERenderTargetLoadAction::EClear);
 		PassParameters->VirtualShadowMapSamplingParameters = VirtualShadowMapArray.GetSamplingParameters(GraphBuilder);
 		PassParameters->ConservativeDepthTexture = ConservativeDepthTexture;
@@ -712,16 +712,16 @@ void FDeferredShadingSceneRenderer::RenderLocalLightsForVolumetricFog(
 							LightFunctionAtlasTileMinMaxUvBound = LightFunctionData->AtlasTile.MinMaxUvBound;
 						}
 
-						TInjectShadowedLocalLightPS::FPermutationDomain PermutationVector;
-						PermutationVector.Set< TInjectShadowedLocalLightPS::FDynamicallyShadowed >(bDynamicallyShadowed);
-						PermutationVector.Set< TInjectShadowedLocalLightPS::FTemporalReprojection >(bUseTemporalReprojection);
-						PermutationVector.Set< TInjectShadowedLocalLightPS::FLightFunction >(bUsesLightFunction);
-						PermutationVector.Set< TInjectShadowedLocalLightPS::FEnableShadows >(bIsShadowed);
-						PermutationVector.Set< TInjectShadowedLocalLightPS::FVirtualShadowMap >(bUseVSM);
+						FInjectShadowedLocalLightPS::FPermutationDomain PermutationVector;
+						PermutationVector.Set< FInjectShadowedLocalLightPS::FDynamicallyShadowed >(bDynamicallyShadowed);
+						PermutationVector.Set< FInjectShadowedLocalLightPS::FTemporalReprojection >(bUseTemporalReprojection);
+						PermutationVector.Set< FInjectShadowedLocalLightPS::FLightFunction >(bUsesLightFunction);
+						PermutationVector.Set< FInjectShadowedLocalLightPS::FEnableShadows >(bIsShadowed);
+						PermutationVector.Set< FInjectShadowedLocalLightPS::FVirtualShadowMap >(bUseVSM);
 
 						auto VertexShader = View.ShaderMap->GetShader< FWriteToBoundingSphereVS >();
 						TOptionalShaderMapRef<FWriteToSliceGS> GeometryShader(View.ShaderMap);
-						auto PixelShader = View.ShaderMap->GetShader< TInjectShadowedLocalLightPS >(PermutationVector);
+						auto PixelShader = View.ShaderMap->GetShader< FInjectShadowedLocalLightPS >(PermutationVector);
 
 						FGraphicsPipelineStateInitializer GraphicsPSOInit;
 						RHICmdList.ApplyCachedRenderTargets(GraphicsPSOInit);
