Deletes r.TSR.Moire

Too unreliable

diff --git a/Engine/Shaders/Private/TemporalSuperResolution/TSRCommon.ush b/Engine/Shaders/Private/TemporalSuperResolution/TSRCommon.ush
--- a/Engine/Shaders/Private/TemporalSuperResolution/TSRCommon.ush
+++ b/Engine/Shaders/Private/TemporalSuperResolution/TSRCommon.ush
@@ -11,7 +11,7 @@
 
 
 //------------------------------------------------------- RECOMPILE HASH
-// 4E268118-A455-494C-BA81-DCFF1C9528CB
+// BE1AEDD9-8FB8-4039-8E36-8C888D492EC7
 
 
 //------------------------------------------------------- COMPILER CONFIG
@@ -76,23 +76,6 @@
 #define SUB_PIXEL_PARALLAX_FACTOR_BIT_MASK ((1 << SUB_PIXEL_PARALLAX_FACTOR_BIT_COUNT) - 1)
 
 
-// Size and payload of the super resolution history
-#define SUPERRES_RES_MULTIPLIER 2
-
-#define SUPERRES_AGE_BITS 5
-#define SUPERRES_AGE_BITMASK ((1 << SUPERRES_AGE_BITS) - 1)
-
-#define SUPERRES_INTERFERENCE_BITS 4
-#define SUPERRES_INTERFERENCE_BITMASK ((1 << SUPERRES_INTERFERENCE_BITS) - 1)
-
-#define SUPERRES_OFFSET_BITS 4
-#define SUPERRES_OFFSET_BITMASK ((1 << SUPERRES_OFFSET_BITS) - 1)
-#define SUPERRES_OFFSET_OFFSET (SUPERRES_OFFSET_BITMASK / 2)
-#define SUPERRES_OFFSET_MIN (-SUPERRES_OFFSET_OFFSET)
-#define SUPERRES_OFFSET_MAX (SUPERRES_OFFSET_BITMASK - SUPERRES_OFFSET_OFFSET)
-
-
-
 // K = Center of the nearest input pixel.
 // O = Center of the output pixel.
 //
@@ -294,7 +277,6 @@ FScreenTransform ScreenPosToPrevHistoryBufferUV;
 float HistoryPreExposureCorrection;
 
 uint bCameraCut;
-float MoireInvExposure;
 
 #if DEBUG_OUTPUT
 	RWTexture2D<float4> DebugOutput;
@@ -906,13 +888,6 @@ taa_half UncompressParallaxFactor(taa_subpixel_payload CompressedParallaxFactor)
 		taa_half(1.0);
 }
 
-taa_half LumaToInterferenceSeed(taa_half Luma)
-{
-	Luma *= taa_half(MoireInvExposure);
-
-	return Luma * HdrWeightY(Luma);
-}
-
 taa_half3x2 QuantizeForFloatRenderTarget(taa_half3x2 Color, taa_half E, const float3 QuantizationError)
 {
 	taa_half3x2 Error = dpv_mul(Color, taa_half3(QuantizationError));
@@ -942,43 +917,6 @@ taa_half3 QuantizeForFloatRenderTarget(taa_half3 Color, taa_half E, const float3
 }
 
 
-struct FSuperResHistory
-{
-	int2 TexelToSample;
-	uint Age;
-	taa_half InterferenceSeed;
-};
-
-uint2 PackSuperResHistory(FSuperResHistory SuperRes)
-{
-	uint CompressedInterferenceSeed = clamp(
-		uint(SuperRes.InterferenceSeed * taa_half(SUPERRES_INTERFERENCE_BITMASK)),
-		uint(0), uint(SUPERRES_INTERFERENCE_BITMASK));
-
-	uint2 Packed = 0;
-
-	Packed[0] |= ((SuperRes.Age) & SUPERRES_AGE_BITMASK) << 0;
-	Packed[0] |= ((CompressedInterferenceSeed) & SUPERRES_INTERFERENCE_BITMASK) << 5;
-
-	Packed[1] |= (uint(SuperRes.TexelToSample.x + SUPERRES_OFFSET_OFFSET) & SUPERRES_OFFSET_BITMASK) << 0 * SUPERRES_OFFSET_BITS;
-	Packed[1] |= (uint(SuperRes.TexelToSample.y + SUPERRES_OFFSET_OFFSET) & SUPERRES_OFFSET_BITMASK) << 1 * SUPERRES_OFFSET_BITS;
-
-	return Packed;
-}
-
-FSuperResHistory UnpackSuperResHistory(uint2 Packed)
-{
-	FSuperResHistory SuperRes;
-
-	SuperRes.Age = ((Packed[0] >> 0) & SUPERRES_AGE_BITMASK);
-	SuperRes.InterferenceSeed = taa_half((Packed[0] >> 5) & SUPERRES_INTERFERENCE_BITMASK) * rcp(taa_half(SUPERRES_INTERFERENCE_BITMASK));
-
-	SuperRes.TexelToSample.x = ((Packed[1] >> 0 * SUPERRES_OFFSET_BITS) & SUPERRES_OFFSET_BITMASK) - SUPERRES_OFFSET_OFFSET;
-	SuperRes.TexelToSample.y = ((Packed[1] >> 1 * SUPERRES_OFFSET_BITS) & SUPERRES_OFFSET_BITMASK) - SUPERRES_OFFSET_OFFSET;
-
-	return SuperRes;
-}
-
 #if COMPRESS_PREV_USE_COUNT
 
 // encode 4 use count that are 8bits shifted in the PF_R32_UINT
diff --git a/Engine/Shaders/Private/TemporalSuperResolution/TSRCompareHistory.usf b/Engine/Shaders/Private/TemporalSuperResolution/TSRCompareHistory.usf
--- a/Engine/Shaders/Private/TemporalSuperResolution/TSRCompareHistory.usf
+++ b/Engine/Shaders/Private/TemporalSuperResolution/TSRCompareHistory.usf
@@ -15,8 +15,6 @@
 Texture2D<taa_half> ParallaxRejectionMaskTexture;
 Texture2D<taa_half3> FilteredInputTexture;
 Texture2D<taa_half3> FilteredPredictionSceneColorTexture;
-Texture2D<taa_half4> FilteredInterferenceSeedTexture;
-Texture2D<taa_half> InterferenceWeightTexture;
 
 RWTexture2D<taa_half> HistoryRejectionOutput;
 
@@ -52,15 +50,6 @@ void MainCS(
 	taa_half3 InputMaxColor;
 	taa_half3 InputCenterColor;
 	taa_half ParallaxRejectionMask;
-	taa_half InterferenceWeight;
-	
-	taa_half ClosestInterferenceSeedMin;
-	taa_half ClosestInterferenceSeedMax;
-	taa_half IntegratedInterferenceSeed;
-	taa_half IntegratedInterferenceSeedMin;
-	taa_half IntegratedInterferenceSeedMax;
-	taa_half MinInterferenceSeed;
-	taa_half MaxInterferenceSeed;
 
 	{
 		InputCenterColor = TransformColorForClampingBox(FilteredInputTexture[InputPixelPos]);
@@ -68,18 +57,6 @@ void MainCS(
 
 		InputMinColor = InputCenterColor;
 		InputMaxColor = InputCenterColor;
-		
-		InterferenceWeight = InterferenceWeightTexture[InputPixelPos];
-		
-		taa_half4 InterferenceSeed = FilteredInterferenceSeedTexture[InputPixelPos];
-		
-		ClosestInterferenceSeedMin = InterferenceSeed.x;
-		ClosestInterferenceSeedMax = InterferenceSeed.x; 
-		IntegratedInterferenceSeed = InterferenceSeed.y;
-		IntegratedInterferenceSeedMin = InterferenceSeed.y;
-		IntegratedInterferenceSeedMax = InterferenceSeed.y;
-		MinInterferenceSeed = InterferenceSeed.z;
-		MaxInterferenceSeed = InterferenceSeed.w;
 	}
 	
 	taa_half NeigborMaxInterferenceSeedGradient = 0.0;
@@ -98,67 +75,11 @@ void MainCS(
 			InputPixelPosMin, InputPixelPosMax);
 
 		taa_half3 Color = TransformColorForClampingBox(FilteredInputTexture[SampleInputPixelPos]);
-		taa_half4 InterferenceSeed = FilteredInterferenceSeedTexture[SampleInputPixelPos];
 
 		InputMinColor = min(InputMinColor, Color);
 		InputMaxColor = max(InputMaxColor, Color);
-		
-		taa_half InterferenceSeedGradient = InterferenceSeed.w - InterferenceSeed.z;
-		
-		ClosestInterferenceSeedMin = min(ClosestInterferenceSeedMin, InterferenceSeed.x);
-		ClosestInterferenceSeedMax = max(ClosestInterferenceSeedMax, InterferenceSeed.x); 
-		
-		// TODO(TSR): this is InterferenceSeed.w - InterferenceSeed.z
-		IntegratedInterferenceSeedMin = min(IntegratedInterferenceSeedMin, InterferenceSeed.y);
-		IntegratedInterferenceSeedMax = max(IntegratedInterferenceSeedMax, InterferenceSeed.y);
-		
-		NeigborMaxInterferenceSeedGradient = max(NeigborMaxInterferenceSeedGradient, InterferenceSeedGradient); 
 	}
 	
-	bool bIsInterference = false;
-	#if DIM_MOIRE_HEURISTIC
-	{
-		// Error due to the low bitdepth of the inteference seeds in the history.
-		const taa_half InterferenceSeedEncodingError = rcp(taa_half(SUPERRES_INTERFERENCE_BITMASK));
-
-		taa_half InputSeedMin = LumaToInterferenceSeed(InputMinColor.x);
-		taa_half InputSeedMax = LumaToInterferenceSeed(InputMaxColor.x);
-
-		// Measure the low frequency gradien in the seed. Needs to happen both on current frame input to avoid ghosting on
-		// current frame changes with high contrast geometric edges, and previous frames to avoid ghosting on area that had 
-		// high contrast geometric edges but no longer had in current frame due to shading change.
-		taa_half InterferenceSeedGradient = max(InputSeedMax - InputSeedMin, IntegratedInterferenceSeedMax - IntegratedInterferenceSeedMin);
-		//taa_half InterferenceSeedGradient = InputSeedMax - InputSeedMin;
-
-		// Detect if there is any inteference based on any subpixel detail range from the lowest LDR to highest LDR.
-		bool bRejectInterferenceSeedWithGradient = (
-			1 && (MinInterferenceSeed < ClosestInterferenceSeedMin - InterferenceSeedEncodingError - InterferenceSeedGradient) ||
-			1 && (MaxInterferenceSeed > ClosestInterferenceSeedMax + InterferenceSeedEncodingError + InterferenceSeedGradient));
-			
-		bool bRejectInterferenceSeed = (
-			(MinInterferenceSeed < ClosestInterferenceSeedMin - InterferenceSeedEncodingError) ||
-			(MaxInterferenceSeed > ClosestInterferenceSeedMax + InterferenceSeedEncodingError));
-
-		// Detect interference solely based on the inteference seed, but need to take into account low frequency
-		// gradien that there may be in the neighborhood to avoid ghosting on high contrast geometric edges.
-		//bIsInterference = bRejectInterferenceSeed;
-		bIsInterference = bRejectInterferenceSeedWithGradient;
-
-		Debug.r = bIsInterference ? 1.0 : 0.0;
-		Debug.g = InputSeedMax - InputSeedMin;
-		Debug.b = ClosestInterferenceSeedMax;
-		Debug.a = MaxInterferenceSeed;
-		
-		Debug.b = ClosestInterferenceSeedMax;
-		Debug.a = MinInterferenceSeed;
-
-		//FLATTEN
-		//if (bOffScreen)
-		//{
-		//	bIsInterference = false;
-		//}
-	}
-	#endif
 
 	taa_half3 PrevColor = TransformColorForClampingBox(FilteredPredictionSceneColorTexture[InputPixelPos]);
 
@@ -175,19 +96,6 @@ void MainCS(
 		//Debug.rgb = abs(YCoCgToRGB(PrevColor - ClampedPrevColor)) / max3(YCoCgToRGB(InputMaxColor), YCoCgToRGB(InputCenterColor), 0.01);
 	}
 	#endif
-	
-	#if DIM_MOIRE_HEURISTIC
-	{
-		if (bIsInterference)
-		{
-			InterferenceWeight = 0.0;
-		}
-		 
-		Rejection = saturate(Rejection + 1.0 - InterferenceWeight);
-		
-		Debug.r = InterferenceWeight;
-	}
-	#endif
 
 	if (ParallaxRejectionMask < taa_half(PARALLAX_REJECTION_MASK_THRESHOLD))
 	{
diff --git a/Engine/Shaders/Private/TemporalSuperResolution/TSRDecimateHistory.usf b/Engine/Shaders/Private/TemporalSuperResolution/TSRDecimateHistory.usf
--- a/Engine/Shaders/Private/TemporalSuperResolution/TSRDecimateHistory.usf
+++ b/Engine/Shaders/Private/TemporalSuperResolution/TSRDecimateHistory.usf
@@ -30,16 +30,12 @@ Texture2D<taa_half2> PrevHistory_Metadata;
 Texture2D<taa_half4> PrevHistory_Translucency;
 Texture2D<taa_subpixel_details> PrevHistory_SubpixelDetails;
 
-Texture2D<half> PrevHistory_LowResTextures_0;
-Texture2D<uint> PrevHistory_SuperResTextures_0;
-
 RWTexture2D<taa_half3> HalfResSceneColorOutput;
 RWTexture2D<taa_half3> HalfResPredictionSceneColorOutput;
 RWTexture2D<taa_half> HalfResParallaxRejectionMaskOutput;
 
 RWTexture2D<taa_half3> PredictionSceneColorOutput;
 RWTexture2D<taa_half> ParallaxRejectionMaskOutput;
-RWTexture2D<taa_half4> InterferenceSeedOutput;
 
 
 //------------------------------------------------------- LDS
@@ -81,9 +77,6 @@ void MainCS(
 	float DeviceZ = ClosestDepthTexture[InputPixelPos];
 	taa_half ParallaxFactor = UncompressParallaxFactor((ParallaxFactorTexture[InputPixelPos] >> 4) & 0xF);
 	taa_half3 InputSceneColor = InputSceneColorTexture[InputPixelPos];
-	
-	taa_half ClosestInputLuma4 = Luma4(InputSceneColor);
-	taa_half InputInterferenceSeed = LumaToInterferenceSeed(ClosestInputLuma4);
 
 	float WorldDepth = ConvertFromDeviceZ(DeviceZ);
 	
@@ -170,99 +163,9 @@ void MainCS(
 		
 		float2 PrevInputBufferUV = (InputInfo_ScreenPosToViewportScale * PrevOutputScreenPos + InputInfo_ScreenPosToViewportBias) * InputInfo_ExtentInverse;
 
-		#if DIM_MOIRE_HEURISTIC
-		{
-			const uint KernelSize = 4;
-
-			uint SuperResPayloads[16];
-
-			uint ClosestPayload = ~0;
-			uint ClosestAge = 0;
-			taa_half ClosestDistanceToSeed = 0.2;
-
-			Debug.r = InputInterferenceSeed;
-
-			// Sample the super res history and find an aproximative age of the input pixel.
-			UNROLL
-			for (uint y = 0; y < KernelSize; y++)
-			{
-				UNROLL
-				for (uint x = 0; x < KernelSize; x++)
-				{
-					const float2 Offset = float2(x, y) - (KernelSize - 1.0) * 0.5;
-					
-					float2 SampleHistoryBufferUV = clamp(
-						(HistoryPPCk * 2.0 + Offset) * PrevHistoryInfo_ExtentInverse * rcp(float(SUPERRES_RES_MULTIPLIER)),
-						PrevHistoryInfo_UVViewportBilinearMin,
-						PrevHistoryInfo_UVViewportBilinearMax);
-
-					uint2 SuperResPixelCoord = uint2(SampleHistoryBufferUV * PrevHistoryInfo_Extent * float(SUPERRES_RES_MULTIPLIER));
-
-					uint SubpixelPayload = PrevHistory_SuperResTextures_0[SuperResPixelCoord];
-					FSuperResHistory UnpackedSubpixel = UnpackSuperResHistory(SubpixelPayload);
-
-					SuperResPayloads[y * 4 + x] = SubpixelPayload;
-
-					taa_half DistanceToSeed = abs(InputInterferenceSeed - UnpackedSubpixel.InterferenceSeed);
-
-					// UnpackedSubpixel.Age >= ClosestAge && 
-
-					FLATTEN
-					if (DistanceToSeed < ClosestDistanceToSeed)
-					{
-						ClosestPayload = SubpixelPayload;
-						ClosestAge = UnpackedSubpixel.Age;
-						ClosestDistanceToSeed = DistanceToSeed;
-					}
-				}
-			}
-			
-			uint AgeMin = UnpackSuperResHistory(ClosestPayload).Age;
-			uint AgeMax = UnpackSuperResHistory(ClosestPayload).Age + 4;
-
-			ClosestInterferenceSeed = UnpackSuperResHistory(ClosestPayload).InterferenceSeed;
-
-			MinInterferenceSeed = 1.0;
-			MaxInterferenceSeed = 0.0;
-
-			// Find out if there is high frequency details in the history with that age.
-			UNROLL
-			for (uint i = 0; i < KernelSize * KernelSize; i++)
-			{
-				FSuperResHistory UnpackedSubpixel = UnpackSuperResHistory(SuperResPayloads[i]);
-
-				bool bAgeBilateralWeight = UnpackedSubpixel.Age >= AgeMin && UnpackedSubpixel.Age <= AgeMax;
-
-				FLATTEN
-				if (bAgeBilateralWeight && ClosestPayload != ~0)
-				{
-					MinInterferenceSeed = min(MinInterferenceSeed, UnpackedSubpixel.InterferenceSeed);
-					MaxInterferenceSeed = max(MaxInterferenceSeed, UnpackedSubpixel.InterferenceSeed);
-				}
-
-			}
-			
-			// actually integrating SubpixelSeed stil suffer interference that causes instability. So instead take
-			// average of min and max that turns out interference free.
-			IntegratedInterferenceSeed = (MinInterferenceSeed + MaxInterferenceSeed) * 0.5;
-		}
-		#endif // DIM_MOIRE_HEURISTIC
-
-		taa_half PrevUninterferenceSeed = PrevHistory_LowResTextures_0.SampleLevel(GlobalBilinearClampedSampler, PrevInputBufferUV, 0);
-		taa_half PrevUninterferenceLuma = PrevUninterferenceSeed * HdrWeightInvY(PrevUninterferenceSeed);
-
-		PrevUninterferenceLuma /= MoireInvExposure;
-
-		#if !DIM_MOIRE_HEURISTIC || 2
+		#if 1
 		{
 			DecimatedColor = PrevHistory_LowFrequency.SampleLevel(GlobalBilinearClampedSampler, KernelHistoryBufferUV, 0);
-			
-			// Cancels out the flickering Luma component with the interference free luma from the inteference seed of the history.
-			#if DIM_MOIRE_HEURISTIC && 1
-			{
-				DecimatedColor *= PrevUninterferenceLuma * SafeRcp(Luma4(DecimatedColor));
-			}
-			#endif
 		}
 		#else
 		{
@@ -280,13 +183,6 @@ void MainCS(
 				
 				taa_half3 HistorySample = PrevHistory_LowFrequency.SampleLevel(GlobalBilinearClampedSampler, SampleBufferUV, 0);
 
-				// Cancels out the flickering Luma component with the interference free luma from the inteference seed of the history.
-				#if DIM_MOIRE_HEURISTIC
-				{
-					HistorySample *= PrevUninterferenceLuma * SafeRcp(Luma4(HistorySample));
-				}
-				#endif
-
 				DecimatedColor += HistorySample;
 			}
 
@@ -418,12 +314,6 @@ void MainCS(
 		
 		ParallaxRejectionMaskOutput[InputPixelPos] = ParallaxRejectionMask;
 
-		#if DIM_MOIRE_HEURISTIC
-		{
-			InterferenceSeedOutput[InputPixelPos] = taa_half4(ClosestInterferenceSeed, IntegratedInterferenceSeed, MinInterferenceSeed, MaxInterferenceSeed);
-		}
-		#endif
-
 		#if DEBUG_OUTPUT
 		{
 			DebugOutput[InputPixelPos] = Debug;
diff --git a/Engine/Shaders/Private/TemporalSuperResolution/TSRDetectInterference.usf b/Engine/Shaders/Private/TemporalSuperResolution/TSRDetectInterference.usf
deleted file mode 100644
diff --git a/Engine/Shaders/Private/TemporalSuperResolution/TSRFilterFrequencies.usf b/Engine/Shaders/Private/TemporalSuperResolution/TSRFilterFrequencies.usf
--- a/Engine/Shaders/Private/TemporalSuperResolution/TSRFilterFrequencies.usf
+++ b/Engine/Shaders/Private/TemporalSuperResolution/TSRFilterFrequencies.usf
@@ -35,12 +35,9 @@ float3 OutputQuantizationError;
 Texture2D<taa_half3> InputTexture;
 Texture2D<taa_half3> PredictionSceneColorTexture;
 Texture2D<taa_half> ParallaxRejectionMaskTexture;
-Texture2D<taa_half4> InterferenceSeedTexture;
-Texture2D<taa_half> InterferenceWeightTexture;
 
 RWTexture2D<taa_half3> FilteredInputOutput;
 RWTexture2D<taa_half3> FilteredPredictionSceneColorOutput;
-RWTexture2D<taa_half4> FilteredInterferenceSeedOutput;
 RWTexture2D<taa_half>  InputSceneColorLdrLumaOutput;
 
 
@@ -87,8 +84,6 @@ void MainCS(
 	taa_half4 FilteredInterferenceSeed;
 	taa_half FilteredInterferenceSeedWeight;
 
-	taa_half CenterSampleInterferenceWeight;
-
 	{
 		taa_half SampleSpatialWeight = 1.0;
 
@@ -96,10 +91,7 @@ void MainCS(
 		taa_half3 HistorySample = PredictionSceneColorTexture[InputPixelPos];
 		taa_half HistorySampleValidity = ParallaxRejectionMaskToHistoryValidity(
 			ParallaxRejectionMaskTexture[InputPixelPos]);
-		taa_half4 InterferenceSeedSample = InterferenceSeedTexture[InputPixelPos];
 		
-		CenterSampleInterferenceWeight = InterferenceWeightTexture[InputPixelPos];
-
 		taa_half InputWeight = ComputeColorWeight(InputSample);
 		taa_half HistoryWeight = HistorySampleValidity * ComputeColorWeight(HistorySample);
 
@@ -110,13 +102,6 @@ void MainCS(
 		
 		FilteredHistory       = (SampleSpatialWeight * HistoryWeight) * HistorySample;
 		FilteredHistoryWeight = (SampleSpatialWeight * HistoryWeight);
-		
-		#if DIM_MOIRE_HEURISTIC
-		{
-			FilteredInterferenceSeed       = (SampleSpatialWeight * HistorySampleValidity) * InterferenceSeedSample;
-			FilteredInterferenceSeedWeight = (SampleSpatialWeight * HistorySampleValidity);
-		}
-		#endif
 	}
 	
 	uint StartSampleId;
@@ -194,48 +179,24 @@ void MainCS(
 		taa_half3 HistorySample = PredictionSceneColorTexture[SampleInputPixelPos];
 		taa_half HistorySampleValidity = ParallaxRejectionMaskToHistoryValidity(
 			ParallaxRejectionMaskTexture[SampleInputPixelPos]);
-		taa_half4 InterferenceSeedSample = InterferenceSeedTexture[SampleInputPixelPos];
-
-		#if DIM_MOIRE_HEURISTIC
-			taa_half SampleInterferenceWeight = InterferenceWeightTexture[SampleInputPixelPos];
-			//SampleInterferenceWeight = 0.0; // TODO: need to dectect interference due to this pattern
-			//SampleInterferenceWeight = 1.0;
-
-			SampleInterferenceWeight = min(SampleInterferenceWeight, CenterSampleInterferenceWeight);
-		#else
-			taa_half SampleInterferenceWeight = taa_half(1.0);
-		#endif
 
 		#if CONFIG_SCENE_COLOR_OVERFLOW
 			InputSample = min(InputSample, LargestSceneColorRGB);
 		#endif
 
-		taa_half InputWeight = (HistorySampleValidity * SampleInterferenceWeight) * ComputeColorWeight(InputSample);
-		taa_half HistoryWeight = (HistorySampleValidity * SampleInterferenceWeight) * ComputeColorWeight(HistorySample);
+		taa_half InputWeight = HistorySampleValidity * ComputeColorWeight(InputSample);
+		taa_half HistoryWeight = HistorySampleValidity * ComputeColorWeight(HistorySample);
 
 		FilteredInput       += (SampleSpatialWeight * InputWeight) * InputSample;
 		FilteredInputWeight += (SampleSpatialWeight * InputWeight);
 		
 		FilteredHistory       += (SampleSpatialWeight * HistoryWeight) * HistorySample;
 		FilteredHistoryWeight += (SampleSpatialWeight * HistoryWeight);
-		
-		#if DIM_MOIRE_HEURISTIC
-		{
-			FilteredInterferenceSeed       += (SampleSpatialWeight * HistorySampleValidity * SampleInterferenceWeight) * InterferenceSeedSample;
-			FilteredInterferenceSeedWeight += (SampleSpatialWeight * HistorySampleValidity * SampleInterferenceWeight);
-		}
-		#endif
 	}
 
 	{
 		FilteredInput *= SafeRcp(FilteredInputWeight);
 		FilteredHistory *= SafeRcp(FilteredHistoryWeight);
-		
-		#if DIM_MOIRE_HEURISTIC
-		{
-			FilteredInterferenceSeed *= SafeRcp(FilteredInterferenceSeedWeight);
-		}
-		#endif
 	}
 	
 	// Compute the LDR luma used for spatial anti-aliasing's computation of the TexelOffset
@@ -284,12 +245,6 @@ void MainCS(
 
 	FilteredInputOutput[OutputPixelPos] = FilteredInput;
 	FilteredPredictionSceneColorOutput[OutputPixelPos] = FilteredHistory;
-	
-	#if DIM_MOIRE_HEURISTIC
-	{
-		FilteredInterferenceSeedOutput[OutputPixelPos] = FilteredInterferenceSeed;
-	}
-	#endif
 
 	#if DIM_OUTPUT_ANTI_ALIASING_LUMA
 	{
diff --git a/Engine/Shaders/Private/TemporalSuperResolution/TSRUpdateHistory.usf b/Engine/Shaders/Private/TemporalSuperResolution/TSRUpdateHistory.usf
--- a/Engine/Shaders/Private/TemporalSuperResolution/TSRUpdateHistory.usf
+++ b/Engine/Shaders/Private/TemporalSuperResolution/TSRUpdateHistory.usf
@@ -61,7 +61,6 @@ Texture2D<taa_half3> InputSceneColorTexture;
 Texture2D<uint2> InputSceneStencilTexture;
 Texture2D<taa_half4> InputSceneTranslucencyTexture;
 
-Texture2D<taa_half> InterferenceWeightTexture;
 Texture2D<taa_half> HistoryRejectionTexture;
 Texture2D<taa_half> TranslucencyRejectionTexture;
 
@@ -254,7 +253,6 @@ void MainCS(
 	taa_half2 LowFrequencyRejection = taa_half(1.0).xx;
 	taa_half2 TranslucencyRejection = taa_half(1.0).xx;
 	taa_half2 OutputPixelVelocity = taa_half(0.0).xx;
-	taa_half2 MoireWeight = taa_half(1.0).xx;
 	taa_half2 NoiseFiltering = taa_half(0.0).xx;
 	taa_half2x2 HoleFilledVelocityMask = dpv_interleave_mono_registers(taa_half(1.0).xx);
 	taa_subpixel_payload2 ParallaxFactorBits = taa_subpixel_payload(0x0).xx;
@@ -273,7 +271,6 @@ void MainCS(
 		taa_half RawLowFrequencyRejection[DPV_PIXEL_PER_LANE];
 		taa_half RawTranslucencyRejection[DPV_PIXEL_PER_LANE];
 		taa_half RawTranslucencyRejectionNeighborhood[CONFIG_SAMPLES - 1][DPV_PIXEL_PER_LANE];
-		taa_half RawMoireWeightNeighborhood[CONFIG_SAMPLES - 1][DPV_PIXEL_PER_LANE]; // TODO(TSR): bug
 		
 		uint RawSceneStencilRef[DPV_PIXEL_PER_LANE];
 		uint RawSceneStencilRefNeighborhood[CONFIG_SAMPLES - 1][DPV_PIXEL_PER_LANE];
@@ -331,10 +328,6 @@ void MainCS(
 				{
 					RawTranslucencyRejectionNeighborhood[NeighborId][PixelId] = TranslucencyRejectionTexture[dpv_access_pixel(SampleInputPixelPos, PixelId)];
 					RawSceneStencilRefNeighborhood[NeighborId][PixelId] = InputSceneStencilTexture.Load(int3(dpv_access_pixel(SampleInputPixelPos, PixelId), 0)) STENCIL_COMPONENT_SWIZZLE;
-			
-					#if DIM_MOIRE_HEURISTIC
-						RawMoireWeightNeighborhood[NeighborId][PixelId] = InterferenceWeightTexture[dpv_access_pixel(SampleInputPixelPos, PixelId)];
-					#endif
 				}
 			} // for (uint NeighborId = 0; NeighborId < (CONFIG_SAMPLES - 1); NeighborId++)
 			
@@ -420,10 +413,6 @@ void MainCS(
 				TranslucencyRejection = min(TranslucencyRejection, dpv_interleave_registers_array(RawTranslucencyRejectionNeighborhood[NeighborId]));
 			
 				SceneStencilRef |= dpv_interleave_registers_array(RawSceneStencilRefNeighborhood[NeighborId]);
-			
-				#if DIM_MOIRE_HEURISTIC
-					MoireWeight = min(MoireWeight, dpv_interleave_registers_array(RawMoireWeightNeighborhood[NeighborId]));
-				#endif
 			}
 		}
 
@@ -907,10 +896,6 @@ void MainCS(
 			PrevYCoCg, ClampedPrevYCoCg,
 			PrevHighFrequencyYCoCg, InputMinColor, InputMaxColor,
 			MeasureBackbufferLDRQuantizationErrorFromHDRLuma(max(PrevYCoCg[0], ClampedPrevYCoCg[0])));
-		
-		#if DIM_MOIRE_HEURISTIC
-			HighFrequencyRejection = max(HighFrequencyRejection, 1.0 - MoireWeight);
-		#endif
 
 		//PrevHighFrequencyColor = dpv_scale(PrevHighFrequencyColor, HighFrequencyRejection);
 		PrevHighFrequencyColor = dpv_scale(clamp(NormalizedPrevHighFrequencyColor, PrevHistoryMommentMin, PrevHistoryMommentMax), HighFrequencyRejection * PrevHistoryValidity);
@@ -991,11 +976,6 @@ void MainCS(
 		#if 1
 		{
 			taa_half2 MaxValidity = taa_half(1.0) - (taa_half(1.0) - PredictionOnlyValidity) * saturate(OutputPixelVelocity * taa_half(InvWeightClampingPixelSpeed));
-		
-			// Maintain full stability on interference areas.
-			#if DIM_MOIRE_HEURISTIC
-				MaxValidity = max(MaxValidity, 1.0 - MoireWeight);
-			#endif
 
 			// Clamp up the max validity to favorise stability under motion on high contrast edges.
 			#if 1
diff --git a/Engine/Shaders/Private/TemporalSuperResolution/TSRUpdateSuperResHistory.usf b/Engine/Shaders/Private/TemporalSuperResolution/TSRUpdateSuperResHistory.usf
deleted file mode 100644
diff --git a/Engine/Source/Runtime/Renderer/Private/PostProcess/TemporalSuperResolution.cpp b/Engine/Source/Runtime/Renderer/Private/PostProcess/TemporalSuperResolution.cpp
--- a/Engine/Source/Runtime/Renderer/Private/PostProcess/TemporalSuperResolution.cpp
+++ b/Engine/Source/Runtime/Renderer/Private/PostProcess/TemporalSuperResolution.cpp
@@ -45,17 +45,6 @@ TAutoConsoleVariable<int32> CVarTSRRejectionAntiAliasingQuality(
 	TEXT("Controls the quality of spatial anti-aliasing on history rejection (default=1)."),
 	ECVF_Scalability | ECVF_RenderThreadSafe);
 
-TAutoConsoleVariable<int32> CVarTSRMoireEnable(
-	TEXT("r.TSR.Moire"), 0,
-	TEXT("Enable heuristic to detect moire pattern in the input scene color caused by geometric interference between input ")
-	TEXT("pixel grid alignement and structured geometry."),
-	ECVF_RenderThreadSafe);
-
-TAutoConsoleVariable<float> CVarTSRMoireSeedExposure(
-	TEXT("r.TSR.Moire.SeedExposure"), 0.125f,
-	TEXT("Exposure of the interference seed. Lower value will allow to detect interference more towards the shadows."),
-	ECVF_RenderThreadSafe);
-
 TAutoConsoleVariable<float> CVarTSRTranslucencyHighlightLuminance(
 	TEXT("r.TSR.Translucency.HighlightLuminance"), -1.0f,
 	TEXT("Sets the liminance at which translucency is considered an highlights (default=-1.0)."),
@@ -113,7 +102,6 @@ BEGIN_SHADER_PARAMETER_STRUCT(FTSRCommonParameters, )
 
 	SHADER_PARAMETER(FVector2f, InputJitter)
 	SHADER_PARAMETER(int32, bCameraCut)
-	SHADER_PARAMETER(float, MoireInvExposure)
 	SHADER_PARAMETER(FVector2f, ScreenVelocityToInputPixelVelocity)
 	SHADER_PARAMETER(FVector2f, InputPixelVelocityToScreenVelocity)
 
@@ -127,9 +115,6 @@ BEGIN_SHADER_PARAMETER_STRUCT(FTSRHistoryTextures, )
 	SHADER_PARAMETER_RDG_TEXTURE(Texture2D, SubpixelDetails)
 	SHADER_PARAMETER_RDG_TEXTURE(Texture2D, Translucency)
 	SHADER_PARAMETER_RDG_TEXTURE(Texture2D, TranslucencyAlpha)
-	
-	SHADER_PARAMETER_RDG_TEXTURE_ARRAY(Texture2D, LowResTextures, [1])
-	SHADER_PARAMETER_RDG_TEXTURE_ARRAY(Texture2D, SuperResTextures, [2])
 END_SHADER_PARAMETER_STRUCT()
 
 BEGIN_SHADER_PARAMETER_STRUCT(FTSRHistoryUAVs, )
@@ -139,9 +124,6 @@ BEGIN_SHADER_PARAMETER_STRUCT(FTSRHistoryUAVs, )
 	SHADER_PARAMETER_RDG_TEXTURE_UAV(RWTexture2D, SubpixelDetails)
 	SHADER_PARAMETER_RDG_TEXTURE_UAV(RWTexture2D, Translucency)
 	SHADER_PARAMETER_RDG_TEXTURE_UAV(RWTexture2D, TranslucencyAlpha)
-
-	SHADER_PARAMETER_RDG_TEXTURE_UAV_ARRAY(RWTexture2D, LowResTextures, [1])
-	SHADER_PARAMETER_RDG_TEXTURE_UAV_ARRAY(RWTexture2D, SuperResTextures, [2])
 END_SHADER_PARAMETER_STRUCT()
 
 BEGIN_SHADER_PARAMETER_STRUCT(FTSRPrevHistoryParameters, )
@@ -162,25 +144,9 @@ FTSRHistoryUAVs CreateUAVs(FRDGBuilder& GraphBuilder, const FTSRHistoryTextures&
 		UAVs.Translucency = GraphBuilder.CreateUAV(Textures.Translucency);
 		UAVs.TranslucencyAlpha = GraphBuilder.CreateUAV(Textures.TranslucencyAlpha);
 	}
-
-	if (Textures.LowResTextures[0])
-	{
-		UAVs.LowResTextures[0] = GraphBuilder.CreateUAV(Textures.LowResTextures[0]);
-	}
-	if (Textures.SuperResTextures[0])
-	{
-		UAVs.SuperResTextures[0] = GraphBuilder.CreateUAV(Textures.SuperResTextures[0]);
-	}
-	if (Textures.SuperResTextures[1])
-	{
-		UAVs.SuperResTextures[1] = GraphBuilder.CreateUAV(Textures.SuperResTextures[1]);
-	}
-
 	return UAVs;
 }
 
-class FTSRMoireHeuristic : SHADER_PERMUTATION_BOOL("DIM_MOIRE_HEURISTIC");
-
 class FTSRShader : public FGlobalShader
 {
 public:
@@ -255,7 +221,7 @@ class FTSRDecimateHistoryCS : public FTSRShader
 
 	class FOutputHalfRes : SHADER_PERMUTATION_BOOL("DIM_OUTPUT_HALF_RES");
 
-	using FPermutationDomain = TShaderPermutationDomain<FOutputHalfRes, FTSRMoireHeuristic>;
+	using FPermutationDomain = TShaderPermutationDomain<FOutputHalfRes>;
 
 	BEGIN_SHADER_PARAMETER_STRUCT(FParameters, )
 		SHADER_PARAMETER_STRUCT_INCLUDE(FTSRCommonParameters, CommonParameters)
@@ -278,7 +244,6 @@ class FTSRDecimateHistoryCS : public FTSRShader
 		SHADER_PARAMETER_RDG_TEXTURE_UAV(RWTexture2D, HalfResParallaxRejectionMaskOutput)
 		SHADER_PARAMETER_RDG_TEXTURE_UAV(RWTexture2D, PredictionSceneColorOutput)
 		SHADER_PARAMETER_RDG_TEXTURE_UAV(RWTexture2D, ParallaxRejectionMaskOutput)
-		SHADER_PARAMETER_RDG_TEXTURE_UAV(RWTexture2D, InterferenceSeedOutput)
 		SHADER_PARAMETER_RDG_TEXTURE_UAV(RWTexture2D, DebugOutput)
 	END_SHADER_PARAMETER_STRUCT()
 }; // class FTSRDecimateHistoryCS
@@ -307,29 +272,6 @@ class FTSRCompareTranslucencyCS : public FTSRShader
 	END_SHADER_PARAMETER_STRUCT()
 }; // class FTSRCompareTranslucencyCS
 
-class FTSRDetectInterferenceCS : public FTSRShader
-{
-	DECLARE_GLOBAL_SHADER(FTSRDetectInterferenceCS);
-	SHADER_USE_PARAMETER_STRUCT(FTSRDetectInterferenceCS, FTSRShader);
-	
-	BEGIN_SHADER_PARAMETER_STRUCT(FParameters, )
-		SHADER_PARAMETER_STRUCT_INCLUDE(FTSRCommonParameters, CommonParameters)
-
-		SHADER_PARAMETER_RDG_TEXTURE(Texture2D, InputSceneColorTexture)
-		SHADER_PARAMETER_RDG_TEXTURE(Texture2D, DilatedVelocityTexture)
-		SHADER_PARAMETER_RDG_TEXTURE(Texture2D, PredictionSceneColorTexture)
-		SHADER_PARAMETER_RDG_TEXTURE(Texture2D, ParallaxRejectionMaskTexture)
-		SHADER_PARAMETER_RDG_TEXTURE(Texture2D, InterferenceSeedTexture)
-
-		SHADER_PARAMETER_STRUCT_INCLUDE(FTSRPrevHistoryParameters, PrevHistoryParameters)
-		SHADER_PARAMETER_STRUCT(FTSRHistoryTextures, PrevHistory)
-
-		SHADER_PARAMETER_STRUCT(FTSRHistoryUAVs, HistoryOutput)
-		SHADER_PARAMETER_RDG_TEXTURE_UAV(RWTexture2D, InterferenceWeightOutput)
-		SHADER_PARAMETER_RDG_TEXTURE_UAV(RWTexture2D, DebugOutput)
-	END_SHADER_PARAMETER_STRUCT()
-}; // class FTSRDetectInterferenceCS
-
 class FTSRHoleFillVelocityCS : public FTSRShader
 {
 	DECLARE_GLOBAL_SHADER(FTSRHoleFillVelocityCS);
@@ -356,7 +298,7 @@ class FTSRFilterFrequenciesCS : public FTSRShader
 
 	class FOutputAALumaDim : SHADER_PERMUTATION_BOOL("DIM_OUTPUT_ANTI_ALIASING_LUMA");
 
-	using FPermutationDomain = TShaderPermutationDomain<FOutputAALumaDim, FTSRMoireHeuristic>;
+	using FPermutationDomain = TShaderPermutationDomain<FOutputAALumaDim>;
 
 	BEGIN_SHADER_PARAMETER_STRUCT(FParameters, )
 		SHADER_PARAMETER_STRUCT_INCLUDE(FTSRCommonParameters, CommonParameters)
@@ -365,12 +307,9 @@ class FTSRFilterFrequenciesCS : public FTSRShader
 		SHADER_PARAMETER_RDG_TEXTURE(Texture2D, InputTexture)
 		SHADER_PARAMETER_RDG_TEXTURE(Texture2D, PredictionSceneColorTexture)
 		SHADER_PARAMETER_RDG_TEXTURE(Texture2D, ParallaxRejectionMaskTexture)
-		SHADER_PARAMETER_RDG_TEXTURE(Texture2D, InterferenceSeedTexture)
-		SHADER_PARAMETER_RDG_TEXTURE(Texture2D, InterferenceWeightTexture)
 
 		SHADER_PARAMETER_RDG_TEXTURE_UAV(RWTexture2D, FilteredInputOutput)
 		SHADER_PARAMETER_RDG_TEXTURE_UAV(RWTexture2D, FilteredPredictionSceneColorOutput)
-		SHADER_PARAMETER_RDG_TEXTURE_UAV(RWTexture2D, FilteredInterferenceSeedOutput)
 		SHADER_PARAMETER_RDG_TEXTURE_UAV(RWTexture2D, InputSceneColorLdrLumaOutput)
 		SHADER_PARAMETER_RDG_TEXTURE_UAV(RWTexture2D, DebugOutput)
 	END_SHADER_PARAMETER_STRUCT()
@@ -382,15 +321,13 @@ class FTSRCompareHistoryCS : public FTSRShader
 	SHADER_USE_PARAMETER_STRUCT(FTSRCompareHistoryCS, FTSRShader);
 
 	class FOutputHalfRes : SHADER_PERMUTATION_BOOL("DIM_OUTPUT_HALF_RES");
-	using FPermutationDomain = TShaderPermutationDomain<FOutputHalfRes, FTSRMoireHeuristic>;
+	using FPermutationDomain = TShaderPermutationDomain<FOutputHalfRes>;
 
 	BEGIN_SHADER_PARAMETER_STRUCT(FParameters, )
 		SHADER_PARAMETER_STRUCT_INCLUDE(FTSRCommonParameters, CommonParameters)
 		SHADER_PARAMETER_RDG_TEXTURE(Texture2D, ParallaxRejectionMaskTexture)
 		SHADER_PARAMETER_RDG_TEXTURE(Texture2D, FilteredInputTexture)
 		SHADER_PARAMETER_RDG_TEXTURE(Texture2D, FilteredPredictionSceneColorTexture)
-		SHADER_PARAMETER_RDG_TEXTURE(Texture2D, FilteredInterferenceSeedTexture)
-		SHADER_PARAMETER_RDG_TEXTURE(Texture2D, InterferenceWeightTexture)
 		SHADER_PARAMETER_RDG_TEXTURE_UAV(RWTexture2D, HistoryRejectionOutput)
 		SHADER_PARAMETER_RDG_TEXTURE_UAV(RWTexture2D, DebugOutput)
 	END_SHADER_PARAMETER_STRUCT()
@@ -476,26 +413,6 @@ class FTSRFilterAntiAliasingCS : public FTSRShader
 	END_SHADER_PARAMETER_STRUCT()
 }; // class FTSRFilterAntiAliasingCS
 
-class FTSRUpdateSuperResHistoryCS : public FTSRShader
-{
-	DECLARE_GLOBAL_SHADER(FTSRUpdateSuperResHistoryCS);
-	SHADER_USE_PARAMETER_STRUCT(FTSRUpdateSuperResHistoryCS, FTSRShader);
-
-	BEGIN_SHADER_PARAMETER_STRUCT(FParameters, )
-		SHADER_PARAMETER_STRUCT_INCLUDE(FTSRCommonParameters, CommonParameters)
-		SHADER_PARAMETER_RDG_TEXTURE(Texture2D, InputSceneColorTexture)
-
-		SHADER_PARAMETER_RDG_TEXTURE(Texture2D, DilatedVelocityTexture)
-		SHADER_PARAMETER_RDG_TEXTURE(Texture2D, ParallaxRejectionMaskTexture)
-
-		SHADER_PARAMETER_STRUCT_INCLUDE(FTSRPrevHistoryParameters, PrevHistoryParameters)
-		SHADER_PARAMETER_STRUCT(FTSRHistoryTextures, PrevHistory)
-
-		SHADER_PARAMETER_STRUCT(FTSRHistoryUAVs, HistoryOutput)
-		SHADER_PARAMETER_RDG_TEXTURE_UAV(RWTexture2D, DebugOutput)
-	END_SHADER_PARAMETER_STRUCT()
-}; // class FTSRUpdateSuperResHistoryCS
-
 class FTSRUpdateHistoryCS : public FTSRShader
 {
 	DECLARE_GLOBAL_SHADER(FTSRUpdateHistoryCS);
@@ -504,7 +421,7 @@ class FTSRUpdateHistoryCS : public FTSRShader
 	class FRejectionAADim : SHADER_PERMUTATION_BOOL("DIM_REJECTION_ANTI_ALIASING");
 	class FSeparateTranslucencyDim : SHADER_PERMUTATION_BOOL("DIM_SEPARATE_TRANSLUCENCY");
 
-	using FPermutationDomain = TShaderPermutationDomain<FRejectionAADim, FSeparateTranslucencyDim, FTSRMoireHeuristic>;
+	using FPermutationDomain = TShaderPermutationDomain<FRejectionAADim, FSeparateTranslucencyDim>;
 
 	BEGIN_SHADER_PARAMETER_STRUCT(FParameters, )
 		SHADER_PARAMETER_STRUCT_INCLUDE(FTSRCommonParameters, CommonParameters)
@@ -512,7 +429,6 @@ class FTSRUpdateHistoryCS : public FTSRShader
 		SHADER_PARAMETER_RDG_TEXTURE_SRV(Texture2D, InputSceneStencilTexture)
 		SHADER_PARAMETER_RDG_TEXTURE(Texture2D, InputSceneTranslucencyTexture)
 
-		SHADER_PARAMETER_RDG_TEXTURE(Texture2D, InterferenceWeightTexture)
 		SHADER_PARAMETER_RDG_TEXTURE(Texture2D, HistoryRejectionTexture)
 		SHADER_PARAMETER_RDG_TEXTURE(Texture2D, TranslucencyRejectionTexture)
 
@@ -570,14 +486,12 @@ IMPLEMENT_GLOBAL_SHADER(FTSRDilateVelocityCS,        "/Engine/Private/TemporalSu
 IMPLEMENT_GLOBAL_SHADER(FTSRDecimateHistoryCS,       "/Engine/Private/TemporalSuperResolution/TSRDecimateHistory.usf",       "MainCS", SF_Compute);
 IMPLEMENT_GLOBAL_SHADER(FTSRCompareTranslucencyCS,   "/Engine/Private/TemporalSuperResolution/TSRCompareTranslucency.usf",   "MainCS", SF_Compute);
 IMPLEMENT_GLOBAL_SHADER(FTSRHoleFillVelocityCS,      "/Engine/Private/TemporalSuperResolution/TSRHoleFillVelocity.usf",      "MainCS", SF_Compute);
-IMPLEMENT_GLOBAL_SHADER(FTSRDetectInterferenceCS,    "/Engine/Private/TemporalSuperResolution/TSRDetectInterference.usf",    "MainCS", SF_Compute);
 IMPLEMENT_GLOBAL_SHADER(FTSRFilterFrequenciesCS,     "/Engine/Private/TemporalSuperResolution/TSRFilterFrequencies.usf",     "MainCS", SF_Compute);
 IMPLEMENT_GLOBAL_SHADER(FTSRCompareHistoryCS,        "/Engine/Private/TemporalSuperResolution/TSRCompareHistory.usf",        "MainCS", SF_Compute);
 IMPLEMENT_GLOBAL_SHADER(FTSRPostfilterRejectionCS,   "/Engine/Private/TemporalSuperResolution/TSRPostfilterRejection.usf",   "MainCS", SF_Compute);
 IMPLEMENT_GLOBAL_SHADER(FTSRDilateRejectionCS,       "/Engine/Private/TemporalSuperResolution/TSRDilateRejection.usf",       "MainCS", SF_Compute);
 IMPLEMENT_GLOBAL_SHADER(FTSRSpatialAntiAliasingCS,   "/Engine/Private/TemporalSuperResolution/TSRSpatialAntiAliasing.usf",   "MainCS", SF_Compute);
 IMPLEMENT_GLOBAL_SHADER(FTSRFilterAntiAliasingCS,    "/Engine/Private/TemporalSuperResolution/TSRFilterAntiAliasing.usf",    "MainCS", SF_Compute);
-IMPLEMENT_GLOBAL_SHADER(FTSRUpdateSuperResHistoryCS, "/Engine/Private/TemporalSuperResolution/TSRUpdateSuperResHistory.usf", "MainCS", SF_Compute);
 IMPLEMENT_GLOBAL_SHADER(FTSRUpdateHistoryCS,         "/Engine/Private/TemporalSuperResolution/TSRUpdateHistory.usf",         "MainCS", SF_Compute);
 
 #if COMPILE_TSR_DEBUG_PASSES
@@ -611,8 +525,6 @@ ITemporalUpscaler::FOutputs AddTemporalSuperResolutionPasses(
 
 	bool bHalfResLowFrequency = CVarTSRHalfResShadingRejection.GetValueOnRenderThread() != 0;
 
-	bool bEnableMoireHeuristic = CVarTSRMoireEnable.GetValueOnRenderThread() != 0;
-
 	bool bIsSeperateTranslucyTexturesValid = PassInputs.SeparateTranslucencyTextures != nullptr && PassInputs.SeparateTranslucencyTextures->IsColorValid();
 
 	bool bRejectSeparateTranslucency = bIsSeperateTranslucyTexturesValid && CVarTSRTranslucencyPreviousFrameRejection.GetValueOnRenderThread() != 0;
@@ -719,7 +631,6 @@ ITemporalUpscaler::FOutputs AddTemporalSuperResolutionPasses(
 
 		CommonParameters.InputJitter = View.TemporalJitterPixels;
 		CommonParameters.bCameraCut = bCameraCut;
-		CommonParameters.MoireInvExposure = 1.0f / CVarTSRMoireSeedExposure.GetValueOnRenderThread();
 		CommonParameters.ViewUniformBuffer = View.ViewUniformBuffer;
 	}
 
@@ -872,11 +783,6 @@ ITemporalUpscaler::FOutputs AddTemporalSuperResolutionPasses(
 		DummyHistory.TranslucencyAlpha = WhiteDummy;
 
 		DummyHistory.SubpixelDetails = BlackUintDummy;
-
-		DummyHistory.LowResTextures[0] = BlackDummy;
-
-		DummyHistory.SuperResTextures[0] = BlackUintDummy;
-		DummyHistory.SuperResTextures[1] = BlackUintDummy;
 	}
 
 	// Setup the previous frame history
@@ -892,17 +798,6 @@ ITemporalUpscaler::FOutputs AddTemporalSuperResolutionPasses(
 
 		// Register non-filterable history
 		PrevHistory.SubpixelDetails = GraphBuilder.RegisterExternalTexture(InputHistory.SubpixelDetails);
-
-		PrevHistory.LowResTextures[0] = InputHistory.LowResTextures[0].IsValid()
-			? GraphBuilder.RegisterExternalTexture(InputHistory.LowResTextures[0])
-			: DummyHistory.LowResTextures[0];
-
-		PrevHistory.SuperResTextures[0] = InputHistory.SuperResTextures[0].IsValid()
-			? GraphBuilder.RegisterExternalTexture(InputHistory.SuperResTextures[0])
-			: DummyHistory.SuperResTextures[0];
-		PrevHistory.SuperResTextures[1] = InputHistory.SuperResTextures[1].IsValid()
-			? GraphBuilder.RegisterExternalTexture(InputHistory.SuperResTextures[1])
-			: DummyHistory.SuperResTextures[1];
 	}
 	else
 	{
@@ -928,17 +823,6 @@ ITemporalUpscaler::FOutputs AddTemporalSuperResolutionPasses(
 
 	// Create new history.
 	FTSRHistoryTextures History;
-	if (bEnableMoireHeuristic)
-	{
-		FRDGTextureDesc LowResDesc = FRDGTextureDesc::Create2D(
-			InputExtent,
-			PF_R8,
-			FClearValueBinding::None,
-			TexCreate_ShaderResource | TexCreate_UAV);
-
-		History.LowResTextures[0] = GraphBuilder.CreateTexture(LowResDesc, TEXT("TSR.History.LowResMetadata[0]"));
-	}
-
 	{
 		FRDGTextureDesc Desc = FRDGTextureDesc::Create2D(
 			HistoryExtent,
@@ -965,27 +849,12 @@ ITemporalUpscaler::FOutputs AddTemporalSuperResolutionPasses(
 		History.SubpixelDetails = GraphBuilder.CreateTexture(Desc, TEXT("TSR.History.SubpixelInfo"));
 	}
 
-	if (bEnableMoireHeuristic)
-	{
-		FRDGTextureDesc SuperResDesc = FRDGTextureDesc::Create2D(
-			HistoryExtent * 2,
-			PF_R16_UINT,
-			FClearValueBinding::None,
-			TexCreate_ShaderResource | TexCreate_UAV);
-
-		History.SuperResTextures[0] = GraphBuilder.CreateTexture(SuperResDesc, TEXT("TSR.History.SuperResMetadata[0]"));
-		//
-		SuperResDesc.Format = PF_R8_UINT;
-		History.SuperResTextures[1] = GraphBuilder.CreateTexture(SuperResDesc, TEXT("TSR.History.SuperResMetadata[1]"));
-	}
-
 	// Decimate input to flicker at same frequency as input.
 	FRDGTextureRef HalfResInputSceneColorTexture = nullptr;
 	FRDGTextureRef HalfResPredictionSceneColorTexture = nullptr;
 	FRDGTextureRef HalfResParallaxRejectionMaskTexture = nullptr;
 	FRDGTextureRef PredictionSceneColorTexture = nullptr;
 	FRDGTextureRef ParallaxRejectionMaskTexture = nullptr;
-	FRDGTextureRef InterferenceSeedTexture = nullptr;
 	{
 		{
 			FRDGTextureDesc Desc = FRDGTextureDesc::Create2D(
@@ -995,10 +864,6 @@ ITemporalUpscaler::FOutputs AddTemporalSuperResolutionPasses(
 				/* InFlags = */ TexCreate_ShaderResource | TexCreate_UAV);
 
 			ParallaxRejectionMaskTexture = GraphBuilder.CreateTexture(Desc, TEXT("TSR.ParallaxRejectionMask"));
-
-			// TODO(TSR): can compress to the history seed's 4bit per pixel
-			Desc.Format = PF_R8G8B8A8;
-			InterferenceSeedTexture = GraphBuilder.CreateTexture(Desc, TEXT("TSR.Moire.Seed"));
 		}
 
 		FTSRDecimateHistoryCS::FParameters* PassParameters = GraphBuilder.AllocParameters<FTSRDecimateHistoryCS::FParameters>();
@@ -1062,12 +927,10 @@ ITemporalUpscaler::FOutputs AddTemporalSuperResolutionPasses(
 		}
 
 		PassParameters->ParallaxRejectionMaskOutput = GraphBuilder.CreateUAV(ParallaxRejectionMaskTexture);
-		PassParameters->InterferenceSeedOutput = GraphBuilder.CreateUAV(InterferenceSeedTexture);
 		PassParameters->DebugOutput = CreateDebugUAV(LowFrequencyExtent, TEXT("Debug.TSR.DecimateHistory"));
 
 		FTSRDecimateHistoryCS::FPermutationDomain PermutationVector;
 		PermutationVector.Set<FTSRDecimateHistoryCS::FOutputHalfRes>(bHalfResLowFrequency);
-		PermutationVector.Set<FTSRMoireHeuristic>(bEnableMoireHeuristic);
 
 		TShaderMapRef<FTSRDecimateHistoryCS> ComputeShader(View.ShaderMap, PermutationVector);
 		FComputeShaderUtils::AddPass(
@@ -1137,54 +1000,6 @@ ITemporalUpscaler::FOutputs AddTemporalSuperResolutionPasses(
 			FComputeShaderUtils::GetGroupCount(InputRect.Size(), 8));
 	}
 
-	// Detect interference between geometry and alignement of input texel centers. It is not about awnser whether an
-	// interference has happened in the past, because interference change based on input resolution or camera position.
-	// So to remain stable on camera movement and input resolution change, it is about awnsering the question on whether
-	// an interference is possible.
-	// TODO(TSR): Could sample the interference seed in the DilateVelocity and detect interference in the decimate.
-	FRDGTextureRef InterferenceWeightTexture;
-	if (bEnableMoireHeuristic)
-	{
-		{
-			// TODO(TSR): Compress to 1bit per pixel
-			FRDGTextureDesc Desc = FRDGTextureDesc::Create2D(
-				InputExtent,
-				PF_R8,
-				FClearValueBinding::None,
-				/* InFlags = */ TexCreate_ShaderResource | TexCreate_UAV);
-
-			InterferenceWeightTexture = GraphBuilder.CreateTexture(Desc, TEXT("TSR.Moire.Weight"));
-		}
-
-		FTSRDetectInterferenceCS::FParameters* PassParameters = GraphBuilder.AllocParameters<FTSRDetectInterferenceCS::FParameters>();
-		PassParameters->CommonParameters = CommonParameters;
-
-		PassParameters->InputSceneColorTexture = PassInputs.SceneColorTexture;
-		PassParameters->DilatedVelocityTexture = DilatedVelocityTexture;
-		PassParameters->PredictionSceneColorTexture = PredictionSceneColorTexture;
-		PassParameters->ParallaxRejectionMaskTexture = ParallaxRejectionMaskTexture;
-		PassParameters->InterferenceSeedTexture = InterferenceSeedTexture;
-
-		PassParameters->PrevHistoryParameters = PrevHistoryParameters;
-		PassParameters->PrevHistory = PrevHistory;
-
-		PassParameters->HistoryOutput = CreateUAVs(GraphBuilder, History);
-		PassParameters->InterferenceWeightOutput = GraphBuilder.CreateUAV(InterferenceWeightTexture);
-		PassParameters->DebugOutput = CreateDebugUAV(LowFrequencyExtent, TEXT("Debug.TSR.DetectInterference"));
-
-		TShaderMapRef<FTSRDetectInterferenceCS> ComputeShader(View.ShaderMap);
-		FComputeShaderUtils::AddPass(
-			GraphBuilder,
-			RDG_EVENT_NAME("TSR DetectInterference %dx%d", InputRect.Width(), InputRect.Height()),
-			ComputeShader,
-			PassParameters,
-			FComputeShaderUtils::GetGroupCount(InputRect.Size(), 8));
-	}
-	else
-	{
-		InterferenceWeightTexture = WhiteDummy;
-	}
-
 	FRDGTextureRef HoleFilledVelocityMaskTexture = BlackDummy;
 	if (bHoleFillVelocity)
 	{
@@ -1232,7 +1047,6 @@ ITemporalUpscaler::FOutputs AddTemporalSuperResolutionPasses(
 		// Filter out the high frquencies
 		FRDGTextureRef FilteredInputTexture;
 		FRDGTextureRef FilteredPredictionSceneColorTexture;
-		FRDGTextureRef FilteredInterferenceSeedTexture = nullptr;
 		{
 			{
 				FRDGTextureDesc Desc = FRDGTextureDesc::Create2D(
@@ -1256,16 +1070,6 @@ ITemporalUpscaler::FOutputs AddTemporalSuperResolutionPasses(
 				InputSceneColorLdrLumaTexture = GraphBuilder.CreateTexture(Desc, TEXT("TSR.SceneColorLdrLuma"));
 			}
 
-			{
-				FRDGTextureDesc Desc = FRDGTextureDesc::Create2D(
-					LowFrequencyExtent,
-					PF_R8G8B8A8,
-					FClearValueBinding::None,
-					/* InFlags = */ TexCreate_ShaderResource | TexCreate_UAV);
-
-				FilteredInterferenceSeedTexture = GraphBuilder.CreateTexture(Desc, TEXT("TSR.Moire.FilteredSeeds"));
-			}
-
 			FTSRFilterFrequenciesCS::FParameters* PassParameters = GraphBuilder.AllocParameters<FTSRFilterFrequenciesCS::FParameters>();
 			PassParameters->CommonParameters = CommonParameters;
 			PassParameters->OutputQuantizationError = ComputePixelFormatQuantizationError(FilteredInputTexture->Desc.Format);
@@ -1281,19 +1085,15 @@ ITemporalUpscaler::FOutputs AddTemporalSuperResolutionPasses(
 				PassParameters->InputTexture = PassInputs.SceneColorTexture;
 				PassParameters->PredictionSceneColorTexture = PredictionSceneColorTexture;
 				PassParameters->ParallaxRejectionMaskTexture = ParallaxRejectionMaskTexture;
-				PassParameters->InterferenceSeedTexture = InterferenceSeedTexture;
-				PassParameters->InterferenceWeightTexture = InterferenceWeightTexture;
 			}
 
 			PassParameters->FilteredInputOutput = GraphBuilder.CreateUAV(FilteredInputTexture);
 			PassParameters->FilteredPredictionSceneColorOutput = GraphBuilder.CreateUAV(FilteredPredictionSceneColorTexture);
-			PassParameters->FilteredInterferenceSeedOutput = GraphBuilder.CreateUAV(FilteredInterferenceSeedTexture);
 			PassParameters->InputSceneColorLdrLumaOutput = InputSceneColorLdrLumaTexture ? GraphBuilder.CreateUAV(InputSceneColorLdrLumaTexture) : nullptr;
 			PassParameters->DebugOutput = CreateDebugUAV(LowFrequencyExtent, TEXT("Debug.TSR.FilterFrequencies"));
 
 			FTSRFilterFrequenciesCS::FPermutationDomain PermutationVector;
 			PermutationVector.Set<FTSRFilterFrequenciesCS::FOutputAALumaDim>(RejectionAntiAliasingQuality > 0);
-			PermutationVector.Set<FTSRMoireHeuristic>(bEnableMoireHeuristic);
 
 			TShaderMapRef<FTSRFilterFrequenciesCS> ComputeShader(View.ShaderMap, PermutationVector);
 			FComputeShaderUtils::AddPass(
@@ -1325,15 +1125,12 @@ ITemporalUpscaler::FOutputs AddTemporalSuperResolutionPasses(
 			PassParameters->ParallaxRejectionMaskTexture = ParallaxRejectionMaskTexture;
 			PassParameters->FilteredInputTexture = FilteredInputTexture;
 			PassParameters->FilteredPredictionSceneColorTexture = FilteredPredictionSceneColorTexture;
-			PassParameters->FilteredInterferenceSeedTexture = FilteredInterferenceSeedTexture;
-			PassParameters->InterferenceWeightTexture = InterferenceWeightTexture;
 
 			PassParameters->HistoryRejectionOutput = GraphBuilder.CreateUAV(HistoryRejectionTexture);
 			PassParameters->DebugOutput = CreateDebugUAV(LowFrequencyExtent, TEXT("Debug.TSR.CompareHistory"));
 
 			FTSRCompareHistoryCS::FPermutationDomain PermutationVector;
 			PermutationVector.Set<FTSRCompareHistoryCS::FOutputHalfRes>(bOutputHalfRes);
-			PermutationVector.Set<FTSRMoireHeuristic>(bEnableMoireHeuristic);
 
 			TShaderMapRef<FTSRCompareHistoryCS> ComputeShader(View.ShaderMap, PermutationVector);
 			FComputeShaderUtils::AddPass(
@@ -1462,29 +1259,6 @@ ITemporalUpscaler::FOutputs AddTemporalSuperResolutionPasses(
 			FComputeShaderUtils::GetGroupCount(RejectionRect.Size(), 8));
 	}
 
-	if (bEnableMoireHeuristic)
-	{
-		FTSRUpdateSuperResHistoryCS::FParameters* PassParameters = GraphBuilder.AllocParameters<FTSRUpdateSuperResHistoryCS::FParameters>();
-		PassParameters->CommonParameters = CommonParameters;
-		PassParameters->InputSceneColorTexture = PassInputs.SceneColorTexture;
-		PassParameters->DilatedVelocityTexture = DilatedVelocityTexture;
-		PassParameters->ParallaxRejectionMaskTexture = ParallaxRejectionMaskTexture;
-
-		PassParameters->PrevHistoryParameters = PrevHistoryParameters;
-		PassParameters->PrevHistory = PrevHistory;
-
-		PassParameters->HistoryOutput = CreateUAVs(GraphBuilder, History);
-		PassParameters->DebugOutput = CreateDebugUAV(HistoryExtent * 2, TEXT("Debug.TSR.UpdateSuperResHistory"));
-
-		TShaderMapRef<FTSRUpdateSuperResHistoryCS> ComputeShader(View.ShaderMap);
-		FComputeShaderUtils::AddPass(
-			GraphBuilder,
-			RDG_EVENT_NAME("TSR UpdateSuperResHistory %dx%d", HistorySize.X * 2, HistorySize.Y * 2),
-			ComputeShader,
-			PassParameters,
-			FComputeShaderUtils::GetGroupCount(HistorySize * 2, 8));
-	}
-
 	FRDGTextureRef SceneColorOutputTexture;
 	{
 		// Allocate output
@@ -1505,7 +1279,6 @@ ITemporalUpscaler::FOutputs AddTemporalSuperResolutionPasses(
 		PassParameters->InputSceneStencilTexture = GraphBuilder.CreateSRV(
 			FRDGTextureSRVDesc::CreateWithPixelFormat(PassInputs.SceneDepthTexture, PF_X24_G8));
 		PassParameters->InputSceneTranslucencyTexture = SeparateTranslucencyTexture;
-		PassParameters->InterferenceWeightTexture = InterferenceWeightTexture;
 		PassParameters->HistoryRejectionTexture = DilatedHistoryRejectionTexture;
 		PassParameters->TranslucencyRejectionTexture = TranslucencyRejectionTexture ? TranslucencyRejectionTexture : BlackDummy;
 
@@ -1552,7 +1325,6 @@ ITemporalUpscaler::FOutputs AddTemporalSuperResolutionPasses(
 		FTSRUpdateHistoryCS::FPermutationDomain PermutationVector;
 		PermutationVector.Set<FTSRUpdateHistoryCS::FRejectionAADim>(RejectionAntiAliasingQuality > 0);
 		PermutationVector.Set<FTSRUpdateHistoryCS::FSeparateTranslucencyDim>(bAccumulateSeparateTranslucency);
-		PermutationVector.Set<FTSRMoireHeuristic>(bEnableMoireHeuristic);
 
 		TShaderMapRef<FTSRUpdateHistoryCS> ComputeShader(View.ShaderMap, PermutationVector);
 		FComputeShaderUtils::AddPass(
@@ -1625,13 +1397,6 @@ ITemporalUpscaler::FOutputs AddTemporalSuperResolutionPasses(
 				SeparateTranslucencyTexture, &View.ViewState->PrevFrameViewInfo.SeparateTranslucency);
 		}
 
-		if (bEnableMoireHeuristic)
-		{
-			GraphBuilder.QueueTextureExtraction(History.LowResTextures[0], &OutputHistory.LowResTextures[0]);
-			GraphBuilder.QueueTextureExtraction(History.SuperResTextures[0], &OutputHistory.SuperResTextures[0]);
-			GraphBuilder.QueueTextureExtraction(History.SuperResTextures[1], &OutputHistory.SuperResTextures[1]);
-		}
-
 		// Extract the output for next frame SSR so that separate translucency shows up in SSR.
 		if (bAccumulateSeparateTranslucency || HistorySize != OutputRect.Size())
 		{
diff --git a/Engine/Source/Runtime/Renderer/Private/SceneRendering.h b/Engine/Source/Runtime/Renderer/Private/SceneRendering.h
--- a/Engine/Source/Runtime/Renderer/Private/SceneRendering.h
+++ b/Engine/Source/Runtime/Renderer/Private/SceneRendering.h
@@ -911,10 +911,6 @@ struct FTSRHistory
 	// Non-filterable output resolution
 	TRefCountPtr<IPooledRenderTarget> SubpixelDetails;
 
-	// Super resolution informaiton
-	TStaticArray<TRefCountPtr<IPooledRenderTarget>, 1> LowResTextures;
-	TStaticArray<TRefCountPtr<IPooledRenderTarget>, 2> SuperResTextures;
-
 	FIntRect OutputViewportRect;
 
 
