Fix planar reflections.

The planar reflection pass was using shader binding for the view uniform buffer, but FDeferredLightVS uses it as a global UB after CL 18416804. This was causing the VS to run without without the buffer bound, which crashed on DX12 and produced incorrect rendering on other RHIs.

#rnx

diff --git a/Engine/Source/Runtime/Renderer/Private/PlanarReflectionRendering.cpp b/Engine/Source/Runtime/Renderer/Private/PlanarReflectionRendering.cpp
--- a/Engine/Source/Runtime/Renderer/Private/PlanarReflectionRendering.cpp
+++ b/Engine/Source/Runtime/Renderer/Private/PlanarReflectionRendering.cpp
@@ -189,7 +189,7 @@ void PrefilterPlanarReflection(
 		const float FilterWidth = View.ViewRect.Width();
 
 		auto* PassParameters = GraphBuilder.AllocParameters<typename FPrefilterPlanarReflectionPS::FParameters>();
-		PassParameters->View = GetShaderBinding(View.ViewUniformBuffer);
+		PassParameters->View = View.ViewUniformBuffer;
 
 		{
 			FPlanarReflectionUniformParameters PlanarReflectionUniformParameters;
@@ -815,7 +815,7 @@ void FDeferredShadingSceneRenderer::RenderDeferredPlanarReflections(FRDGBuilder&
 	PassParameters->SceneTextures.GBufferFTexture = SceneTextures.GBufferFTexture;
 	PassParameters->SceneTextures.GBufferVelocityTexture = SceneTextures.GBufferVelocityTexture;
 
-	PassParameters->ViewUniformBuffer = GetShaderBinding(View.ViewUniformBuffer);
+	PassParameters->ViewUniformBuffer = View.ViewUniformBuffer;
 	PassParameters->RenderTargets[0] = FRenderTargetBinding(
 		ReflectionsOutputTexture, bClearReflectionsOutputTexture ? ERenderTargetLoadAction::EClear : ERenderTargetLoadAction::ELoad);
 
