UE5_RELEASE:  MGPU null reference on editor startup, in GBuffer Temporal Copy logic.  Added null check.

#rnx

diff --git a/Engine/Source/Runtime/Renderer/Private/DeferredShadingRenderer.cpp b/Engine/Source/Runtime/Renderer/Private/DeferredShadingRenderer.cpp
--- a/Engine/Source/Runtime/Renderer/Private/DeferredShadingRenderer.cpp
+++ b/Engine/Source/Runtime/Renderer/Private/DeferredShadingRenderer.cpp
@@ -2666,6 +2666,7 @@ void FDeferredShadingSceneRenderer::Render(FRDGBuilder& GraphBuilder)
 		}
 
 #if WITH_MGPU
+		if (SceneTextures.Depth.Resolve && SceneTextures.GBufferA)
 		{
 			FGBufferTemporalTextureParams* PassParameters = GraphBuilder.AllocParameters<FGBufferTemporalTextureParams>();
 
