Fixed: Crash in UE4Editor-MetalRHI!FMetalRenderPass::RestartRenderPass

diff --git a/Engine/Source/Runtime/Apple/MetalRHI/Private/MetalStateCache.cpp b/Engine/Source/Runtime/Apple/MetalRHI/Private/MetalStateCache.cpp
--- a/Engine/Source/Runtime/Apple/MetalRHI/Private/MetalStateCache.cpp
+++ b/Engine/Source/Runtime/Apple/MetalRHI/Private/MetalStateCache.cpp
@@ -1836,7 +1836,7 @@ bool FMetalStateCache::PrepareToRestart(bool const bCurrentApplied)
 				LoadAction = ERenderTargetLoadAction::ELoad;
 			}
 			RenderTargetView.Action = MakeRenderTargetActions(LoadAction, StoreAction);
-			check(RenderTargetView.RenderTarget == nil || GetStoreAction(RenderTargetView.Action) == ERenderTargetStoreAction::EStore);
+			check(RenderTargetView.RenderTarget == nil || GetStoreAction(RenderTargetView.Action) != ERenderTargetStoreAction::ENoAction);
 		}
 		
 		InvalidateRenderTargets();
diff --git a/Engine/Source/Runtime/Renderer/Private/MobileShadingRenderer.cpp b/Engine/Source/Runtime/Renderer/Private/MobileShadingRenderer.cpp
--- a/Engine/Source/Runtime/Renderer/Private/MobileShadingRenderer.cpp
+++ b/Engine/Source/Runtime/Renderer/Private/MobileShadingRenderer.cpp
@@ -443,7 +443,6 @@ void FMobileSceneRenderer::InitViews(FRDGBuilder& GraphBuilder, FSceneTexturesCo
 	const bool bSeparateTranslucencyActive = IsMobileSeparateTranslucencyActive(Views.GetData(), Views.Num()); 
 	const bool bPostProcessUsesSceneDepth = PostProcessUsesSceneDepth(Views[0]) || IsMobileDistortionActive(Views[0]);
 	const bool bRequireSeparateViewPass = Views.Num() > 1 && !Views[0].bIsMobileMultiViewEnabled;
-	const bool bIsSimulatedLDR = (!IsMobileHDR() && IsSimulatedPlatform(ShaderPlatform));
 	bRequiresMultiPass = RequiresMultiPass(RHICmdList, Views[0]);
 	bKeepDepthContent =
 		bRequiresMultiPass ||
@@ -456,8 +455,7 @@ void FMobileSceneRenderer::InitViews(FRDGBuilder& GraphBuilder, FSceneTexturesCo
 		(bDeferredShading && bPostProcessUsesSceneDepth) ||
 		bShouldRenderVelocities ||
 		bRequireSeparateViewPass ||
-		bIsFullDepthPrepassEnabled ||
-        bIsSimulatedLDR;
+		bIsFullDepthPrepassEnabled;
 	// never keep MSAA depth if SceneDepthAux is enabled
 	bKeepDepthContent = ((NumMSAASamples > 1) && MobileRequiresSceneDepthAux(ShaderPlatform)) ? false : bKeepDepthContent;
 
@@ -467,6 +465,12 @@ void FMobileSceneRenderer::InitViews(FRDGBuilder& GraphBuilder, FSceneTexturesCo
 		bKeepDepthContent = true;
 	}
 
+	// In the editor RHIs may split a render-pass into several cmd buffer submissions, so all targets need to Store
+	if (IsSimulatedPlatform(ShaderPlatform))
+	{
+		bKeepDepthContent = true;
+	}
+
 	// Update the bKeepDepthContent based on the mobile renderer status.
 	SceneTexturesConfig.bKeepDepthContent = bKeepDepthContent;
 	
@@ -1788,4 +1792,4 @@ void FMobileSceneRenderer::RenderHZB(FRDGBuilder& GraphBuilder, FRDGTextureRef S
 			View.HZB = FurthestHZBTexture;
 		}
 	}
-}
\ No newline at end of file
+}
