Fixed: Black screen when playing in editor with Android Vulkan preview with VR Template

diff --git a/Engine/Source/Runtime/Renderer/Private/MobileShadingRenderer.cpp b/Engine/Source/Runtime/Renderer/Private/MobileShadingRenderer.cpp
--- a/Engine/Source/Runtime/Renderer/Private/MobileShadingRenderer.cpp
+++ b/Engine/Source/Runtime/Renderer/Private/MobileShadingRenderer.cpp
@@ -395,7 +395,11 @@ void FMobileSceneRenderer::InitViews(FRDGBuilder& GraphBuilder, FSceneTexturesCo
 	PostVisibilityFrameSetup(ILCTaskData);
 
 	const FIntPoint RenderTargetSize = (ViewFamily.RenderTarget->GetRenderTargetTexture().IsValid()) ? ViewFamily.RenderTarget->GetRenderTargetTexture()->GetSizeXY() : ViewFamily.RenderTarget->GetSizeXY();
-	const bool bRequiresUpscale = ((int32)RenderTargetSize.X > FamilySize.X || (int32)RenderTargetSize.Y > FamilySize.Y) || IsMobilePropagateAlphaEnabled(ViewFamily.GetShaderPlatform());
+	const bool bRequiresUpscale = 
+		((int32)RenderTargetSize.X > FamilySize.X || (int32)RenderTargetSize.Y > FamilySize.Y || IsMobilePropagateAlphaEnabled(ViewFamily.GetShaderPlatform()) 
+		// in the editor color surface and backbuffer could have a different pixel formats and size, 
+		// so we always run upscale pass to blit content from scene color to backbuffer 
+		|| (GIsEditor && !IsMobileHDR() && NumMSAASamples > 1);
 	// ES requires that the back buffer and depth match dimensions.
 	// For the most part this is not the case when using scene captures. Thus scene captures always render to scene color target.
 	const bool bShouldCompositeEditorPrimitives = FSceneRenderer::ShouldCompositeEditorPrimitives(Views[0]);
diff --git a/Engine/Source/Runtime/Renderer/Private/PostProcess/PostProcessing.cpp b/Engine/Source/Runtime/Renderer/Private/PostProcess/PostProcessing.cpp
--- a/Engine/Source/Runtime/Renderer/Private/PostProcess/PostProcessing.cpp
+++ b/Engine/Source/Runtime/Renderer/Private/PostProcess/PostProcessing.cpp
@@ -1880,7 +1880,12 @@ void AddMobilePostProcessingPasses(FRDGBuilder& GraphBuilder, FScene* Scene, con
 	PassSequence.SetEnabled(EPass::SelectionOutline, false);
 	PassSequence.SetEnabled(EPass::EditorPrimitive, false);
 #endif
-	PassSequence.SetEnabled(EPass::PrimaryUpscale, PaniniConfig.IsEnabled() || (bShouldPrimaryUpscale && bDisableUpscaleInTonemapper));
+	PassSequence.SetEnabled(EPass::PrimaryUpscale, 
+		PaniniConfig.IsEnabled() || (bShouldPrimaryUpscale && bDisableUpscaleInTonemapper) 
+		// in the editor color surface and backbuffer could have a different pixel formats and size, 
+		// so we always run upscale pass to blit content from scene color to backbuffer 
+		|| (GIsEditor && !IsMobileHDR() && bMobileMSAA)
+	);
 
 	PassSequence.SetEnabled(EPass::Visualize, View.Family->EngineShowFlags.ShaderComplexity);
 
