[Backout] - CL18485836 - Change modified code that was skipped merging to UE5
Original CL Desc
-----------------------------------------------------------------
Fixed: Black screen when playing in editor with Android Vulkan preview with VR Template

diff --git a/Engine/Source/Runtime/Renderer/Private/MobileShadingRenderer.cpp b/Engine/Source/Runtime/Renderer/Private/MobileShadingRenderer.cpp
--- a/Engine/Source/Runtime/Renderer/Private/MobileShadingRenderer.cpp
+++ b/Engine/Source/Runtime/Renderer/Private/MobileShadingRenderer.cpp
@@ -395,11 +395,7 @@ void FMobileSceneRenderer::InitViews(FRDGBuilder& GraphBuilder, FSceneTexturesCo
 	PostVisibilityFrameSetup(ILCTaskData);
 
 	const FIntPoint RenderTargetSize = (ViewFamily.RenderTarget->GetRenderTargetTexture().IsValid()) ? ViewFamily.RenderTarget->GetRenderTargetTexture()->GetSizeXY() : ViewFamily.RenderTarget->GetSizeXY();
-	const bool bRequiresUpscale = 
-		((int32)RenderTargetSize.X > FamilySize.X || (int32)RenderTargetSize.Y > FamilySize.Y || IsMobilePropagateAlphaEnabled(ViewFamily.GetShaderPlatform()) 
-		// in the editor color surface and backbuffer could have a different pixel formats and size, 
-		// so we always run upscale pass to blit content from scene color to backbuffer 
-		|| (GIsEditor && !IsMobileHDR() && NumMSAASamples > 1);
+	const bool bRequiresUpscale = ((int32)RenderTargetSize.X > FamilySize.X || (int32)RenderTargetSize.Y > FamilySize.Y) || IsMobilePropagateAlphaEnabled(ViewFamily.GetShaderPlatform());
 	// ES requires that the back buffer and depth match dimensions.
 	// For the most part this is not the case when using scene captures. Thus scene captures always render to scene color target.
 	const bool bShouldCompositeEditorPrimitives = FSceneRenderer::ShouldCompositeEditorPrimitives(Views[0]);
diff --git a/Engine/Source/Runtime/Renderer/Private/PostProcess/PostProcessing.cpp b/Engine/Source/Runtime/Renderer/Private/PostProcess/PostProcessing.cpp
--- a/Engine/Source/Runtime/Renderer/Private/PostProcess/PostProcessing.cpp
+++ b/Engine/Source/Runtime/Renderer/Private/PostProcess/PostProcessing.cpp
@@ -1880,12 +1880,7 @@ void AddMobilePostProcessingPasses(FRDGBuilder& GraphBuilder, FScene* Scene, con
 	PassSequence.SetEnabled(EPass::SelectionOutline, false);
 	PassSequence.SetEnabled(EPass::EditorPrimitive, false);
 #endif
-	PassSequence.SetEnabled(EPass::PrimaryUpscale, 
-		PaniniConfig.IsEnabled() || (bShouldPrimaryUpscale && bDisableUpscaleInTonemapper) 
-		// in the editor color surface and backbuffer could have a different pixel formats and size, 
-		// so we always run upscale pass to blit content from scene color to backbuffer 
-		|| (GIsEditor && !IsMobileHDR() && bMobileMSAA)
-	);
+	PassSequence.SetEnabled(EPass::PrimaryUpscale, PaniniConfig.IsEnabled() || (bShouldPrimaryUpscale && bDisableUpscaleInTonemapper));
 
 	PassSequence.SetEnabled(EPass::Visualize, View.Family->EngineShowFlags.ShaderComplexity);
 
