Fix non-initialized simple light parameters causing them to flicker/disappear.

The simple light LWC parmaeters was incorrectly initialized, causing its tile coord to be incorrect.

diff --git a/Engine/Source/Runtime/Renderer/Private/LightRendering.cpp b/Engine/Source/Runtime/Renderer/Private/LightRendering.cpp
--- a/Engine/Source/Runtime/Renderer/Private/LightRendering.cpp
+++ b/Engine/Source/Runtime/Renderer/Private/LightRendering.cpp
@@ -266,11 +266,22 @@ FDeferredLightUniformStruct GetDeferredLightParameters(const FSceneView& View, c
 
 FDeferredLightUniformStruct GetSimpleDeferredLightParameters(
 	const FSimpleLightEntry& SimpleLight,
-	const FVector3f& Position)
+	const FVector& LWPosition)
 {
 	FDeferredLightUniformStruct Out;
-	Out.LightParameters.Position = Position;
+	Out.ShadowMapChannelMask = FVector4f(0, 0, 0, 0);
+	Out.DistanceFadeMAD = FVector2D(0, 0);
+	Out.ContactShadowLength = 0.0f;
+	Out.ContactShadowNonShadowCastingIntensity = 0.f;
+	Out.VolumetricScatteringIntensity = SimpleLight.VolumetricScatteringIntensity;
+	Out.ShadowedBits = 0;
+	Out.LightingChannelMask = 0;
+
+	const FLargeWorldRenderPosition AbsoluteWorldPosition(LWPosition);
+	Out.LightParameters.Position = AbsoluteWorldPosition.GetOffset();
+	Out.LightParameters.TilePosition = AbsoluteWorldPosition.GetTile();
 	Out.LightParameters.InvRadius = 1.0f / FMath::Max(SimpleLight.Radius, KINDA_SMALL_NUMBER);
+	Out.LightParameters.Padding = 0;
 	Out.LightParameters.Color = SimpleLight.Color;
 	Out.LightParameters.FalloffExponent = SimpleLight.Exponent;
 	Out.LightParameters.Direction = FVector(1, 0, 0);
@@ -281,12 +292,9 @@ FDeferredLightUniformStruct GetSimpleDeferredLightParameters(
 	Out.LightParameters.SoftSourceRadius = 0.0f;
 	Out.LightParameters.SourceLength = 0.0f;
 	Out.LightParameters.SourceTexture = GWhiteTexture->TextureRHI;
-	Out.ContactShadowLength = 0.0f;
-	Out.DistanceFadeMAD = FVector2D(0, 0);
-	Out.ShadowMapChannelMask = FVector4f(0, 0, 0, 0);
-	Out.ShadowedBits = 0;
-	Out.LightingChannelMask = 0;
-
+	Out.LightParameters.RectLightBarnCosAngle = 0;
+	Out.LightParameters.RectLightBarnLength = -2.0f;
+	Out.LightParameters.SourceTexture = GWhiteTexture->TextureRHI;
 	return Out;
 }
 FDeferredLightUniformStruct GetSimpleDeferredLightParameters(
@@ -1754,7 +1762,7 @@ static TRDGUniformBufferRef<FDeferredLightUniformStruct> CreateDeferredLightUnif
 	return GraphBuilder.CreateUniformBuffer(DeferredLightStruct);
 }
 
-static TRDGUniformBufferRef<FDeferredLightUniformStruct> CreateDeferredLightUniformBuffer(FRDGBuilder& GraphBuilder, const FSimpleLightEntry& SimpleLight, const FVector3f& SimpleLightPosition)
+static TRDGUniformBufferRef<FDeferredLightUniformStruct> CreateDeferredLightUniformBuffer(FRDGBuilder& GraphBuilder, const FSimpleLightEntry& SimpleLight, const FVector& SimpleLightPosition)
 {
 	auto* DeferredLightStruct = GraphBuilder.AllocParameters<FDeferredLightUniformStruct>();
 	*DeferredLightStruct = GetSimpleDeferredLightParameters(SimpleLight, SimpleLightPosition);
@@ -2332,7 +2340,7 @@ static FSimpleLightsStandardDeferredParameters GetRenderLightSimpleParameters(
 	const FViewInfo& View,
 	const FMinimalSceneTextures& SceneTextures,
 	const FSimpleLightEntry& SimpleLight,
-	const FVector3f SimpleLightPosition)
+	const FVector& SimpleLightPosition)
 {
 	FSimpleLightsStandardDeferredParameters Out;
 
@@ -2371,10 +2379,10 @@ static FSimpleLightsStandardDeferredParameters GetRenderLightSimpleParameters(
 	}
 
 	// VS - General paramters (dummy geometry, as the geometry is setup within the pass light loop)
-	FSphere DummySphere;
-	DummySphere.Center = FVector::ZeroVector;
-	DummySphere.W = 1.0f;
-	Out.VS = FDeferredLightVS::GetParameters(View, DummySphere, false);
+	FSphere SphereLight;
+	SphereLight.Center = SimpleLightPosition; // Should we account for LWC Position+Tile here?
+	SphereLight.W = SimpleLight.Radius;
+	Out.VS = FDeferredLightVS::GetParameters(View, SphereLight, false);
 
 	return Out;
 }
