Refactoring FTranslucentLightingInjectPS to RDG.

Green screenshot test: TranslucencyLightingMode

diff --git a/Engine/Source/Runtime/Renderer/Private/TranslucentLighting.cpp b/Engine/Source/Runtime/Renderer/Private/TranslucentLighting.cpp
--- a/Engine/Source/Runtime/Renderer/Private/TranslucentLighting.cpp
+++ b/Engine/Source/Runtime/Renderer/Private/TranslucentLighting.cpp
@@ -651,57 +651,28 @@ public:
 
 IMPLEMENT_GLOBAL_SHADER(FFilterTranslucentVolumePS, "/Engine/Private/TranslucentLightingShaders.usf", "FilterMainPS", SF_Pixel);
 
-/** Shader parameters needed to inject direct lighting into a volume. */
-class FTranslucentInjectParameters
-{
-	DECLARE_INLINE_TYPE_LAYOUT(FTranslucentInjectParameters, NonVirtual);
-public:
-
-	void Bind(const FShaderParameterMap& ParameterMap)
-	{
-		WorldToShadowMatrix.Bind(ParameterMap,TEXT("WorldToShadowMatrix"));
-		ShadowmapMinMax.Bind(ParameterMap,TEXT("ShadowmapMinMax"));
-		VolumeCascadeIndex.Bind(ParameterMap,TEXT("VolumeCascadeIndex"));
-	}
-
-	template<typename ShaderRHIParamRef>
-	void Set(
-		FRHICommandList& RHICmdList, 
-		const ShaderRHIParamRef ShaderRHI, 
-		FShader* Shader, 
-		const FViewInfo& View, 
-		const FLightSceneInfo* LightSceneInfo, 
-		const FProjectedShadowInfo* ShadowMap, 
-		uint32 VolumeCascadeIndexValue,
-		bool bDynamicallyShadowed) const
-	{
-		SetDeferredLightParameters(RHICmdList, ShaderRHI, Shader->GetUniformBufferParameter<FDeferredLightUniformStruct>(), LightSceneInfo, View);
-
-		if (bDynamicallyShadowed)
-		{
-			FVector4f ShadowmapMinMaxValue;
-			FMatrix44f WorldToShadowMatrixValue = ShadowMap->GetWorldToShadowMatrix(ShadowmapMinMaxValue);
-
-			SetShaderValue(RHICmdList, ShaderRHI, WorldToShadowMatrix, WorldToShadowMatrixValue);
-			SetShaderValue(RHICmdList, ShaderRHI, ShadowmapMinMax, ShadowmapMinMaxValue);
-		}
-
-		SetShaderValue(RHICmdList, ShaderRHI, VolumeCascadeIndex, VolumeCascadeIndexValue);
-	}
-
-private:
-	LAYOUT_FIELD(FShaderParameter, WorldToShadowMatrix)
-	LAYOUT_FIELD(FShaderParameter, ShadowmapMinMax)
-	LAYOUT_FIELD(FShaderParameter, VolumeCascadeIndex)
-};
-
 /** Shader that adds direct lighting contribution from the given light to the current volume lighting cascade. */
 class FTranslucentLightingInjectPS : public FMaterialShader
 {
-	DECLARE_SHADER_TYPE(FTranslucentLightingInjectPS,Material);
+	DECLARE_SHADER_TYPE(FTranslucentLightingInjectPS, Material);
 
 	BEGIN_SHADER_PARAMETER_STRUCT(FParameters, )
+		SHADER_PARAMETER_STRUCT_REF(FViewUniformShaderParameters, ViewUniformBuffer)
+		SHADER_PARAMETER_RDG_UNIFORM_BUFFER(FDeferredLightUniformStruct, DeferredLight)
+		SHADER_PARAMETER_STRUCT_INCLUDE(FVolumeShadowingShaderParameters, VolumeShadowingParameters)
 		SHADER_PARAMETER_STRUCT_INCLUDE(FVirtualShadowMapSamplingParameters, VirtualShadowMapSamplingParameters)
+		SHADER_PARAMETER_RDG_TEXTURE(Texture2D<float3>, VolumetricCloudShadowmapTexture)
+		SHADER_PARAMETER_SAMPLER(SamplerState, VolumetricCloudShadowmapTextureSampler)
+		SHADER_PARAMETER(FMatrix44f, LightFunctionWorldToLight)
+		SHADER_PARAMETER(FMatrix44f, VolumetricCloudWorldToLightClipShadowMatrix)
+		SHADER_PARAMETER(FVector4f, LightFunctionParameters)
+		SHADER_PARAMETER(float, SpotlightMask)
+		SHADER_PARAMETER(float, VolumetricCloudShadowmapFarDepthKm)
+		SHADER_PARAMETER(float, VolumetricCloudShadowmapStrength)
+		SHADER_PARAMETER(uint32, VolumetricCloudShadowEnabled)
+		SHADER_PARAMETER(uint32, AtmospherePerPixelTransmittanceEnabled)
+		SHADER_PARAMETER(uint32, VolumeCascadeIndex)
+		SHADER_PARAMETER(int32, VirtualShadowMapId)
 	END_SHADER_PARAMETER_STRUCT()
 
 	class FRadialAttenuation	: SHADER_PERMUTATION_BOOL("RADIAL_ATTENUATION");
@@ -746,135 +717,21 @@ public:
 			*FParameters::FTypeInfo::GetStructMetadata(), 
 			// Don't require full bindings, we use FMaterialShader::SetParameters
 			false); 
-
-		VolumeShadowingParameters.Bind(Initializer.ParameterMap);
-		SpotlightMask.Bind(Initializer.ParameterMap, TEXT("SpotlightMask"));
-		LightFunctionParameters.Bind(Initializer.ParameterMap);
-		TranslucentInjectParameters.Bind(Initializer.ParameterMap);
-		LightFunctionWorldToLight.Bind(Initializer.ParameterMap, TEXT("LightFunctionWorldToLight"));
-		VirtualShadowMapIdParameter.Bind(Initializer.ParameterMap, TEXT("VirtualShadowMapId"));
-
-		VolumetricCloudWorldToLightClipShadowMatrix.Bind(Initializer.ParameterMap, TEXT("VolumetricCloudWorldToLightClipShadowMatrix"));
-		VolumetricCloudShadowmapFarDepthKm.Bind(Initializer.ParameterMap, TEXT("VolumetricCloudShadowmapFarDepthKm"));
-		VolumetricCloudShadowEnabled.Bind(Initializer.ParameterMap, TEXT("VolumetricCloudShadowEnabled"));
-		VolumetricCloudShadowmapStrength.Bind(Initializer.ParameterMap, TEXT("VolumetricCloudShadowmapStrength"));
-		VolumetricCloudShadowmapTexture.Bind(Initializer.ParameterMap, TEXT("VolumetricCloudShadowmapTexture"));
-		VolumetricCloudShadowmapTextureSampler.Bind(Initializer.ParameterMap, TEXT("VolumetricCloudShadowmapTextureSampler"));
-		AtmospherePerPixelTransmittanceEnabled.Bind(Initializer.ParameterMap, TEXT("AtmospherePerPixelTransmittanceEnabled"));
 	}
+
 	FTranslucentLightingInjectPS() {}
 
-	// @param InnerSplitIndex which CSM shadow map level, INDEX_NONE if no directional light
-	// @param VolumeCascadeIndexValue which volume we render to
 	void SetParameters(
 		FRHICommandList& RHICmdList, 
 		const FViewInfo& View, 
-		const FLightSceneInfo* LightSceneInfo, 
-		const FMaterialRenderProxy* MaterialProxy, 
-		const FProjectedShadowInfo* ShadowMap, 
-		int32 InnerSplitIndex, 
-		int32 VolumeCascadeIndexValue,
-		int32 VirtualShadowMapId)
+		const FMaterialRenderProxy* MaterialProxy)
 	{
-		bool bDynamicallyShadowed = ShadowMap != nullptr;
-		
 		FRHIPixelShader* ShaderRHI = RHICmdList.GetBoundPixelShader();
-
 		const FMaterial& Material = MaterialProxy->GetMaterialWithFallback(View.GetFeatureLevel(), MaterialProxy);
-		FMaterialShader::SetViewParameters(RHICmdList, ShaderRHI, View, View.ViewUniformBuffer);
 		FMaterialShader::SetParameters(RHICmdList, ShaderRHI, MaterialProxy, Material, View);
-		
-		VolumeShadowingParameters.Set(RHICmdList, ShaderRHI, View, LightSceneInfo, ShadowMap, InnerSplitIndex, bDynamicallyShadowed);
-
-		bool bIsSpotlight = LightSceneInfo->Proxy->GetLightType() == LightType_Spot;
-		//@todo - needs to be a permutation to reduce shadow filtering work
-		SetShaderValue(RHICmdList, ShaderRHI, SpotlightMask, (bIsSpotlight ? 1.0f : 0.0f));
-
-		LightFunctionParameters.Set(RHICmdList, ShaderRHI, LightSceneInfo, 1);
-		TranslucentInjectParameters.Set(RHICmdList, ShaderRHI, this, View, LightSceneInfo, ShadowMap, VolumeCascadeIndexValue, bDynamicallyShadowed);
-
-		if (LightFunctionWorldToLight.IsBound())
-		{
-			const FVector Scale = LightSceneInfo->Proxy->GetLightFunctionScale();
-			// Switch x and z so that z of the user specified scale affects the distance along the light direction
-			const FVector InverseScale = FVector( 1.f / Scale.Z, 1.f / Scale.Y, 1.f / Scale.X );
-			const FMatrix44f WorldToLight = LightSceneInfo->Proxy->GetWorldToLight() * FScaleMatrix(FVector(InverseScale));	
-
-			SetShaderValue(RHICmdList, ShaderRHI, LightFunctionWorldToLight, WorldToLight);
-		}
-
-		SetShaderValue(RHICmdList, ShaderRHI, VirtualShadowMapIdParameter, VirtualShadowMapId);
-
-		FLightSceneProxy* AtmosphereLight0Proxy = LightSceneInfo->Scene->AtmosphereLights[0] ? LightSceneInfo->Scene->AtmosphereLights[0]->Proxy : nullptr;
-		FLightSceneProxy* AtmosphereLight1Proxy = LightSceneInfo->Scene->AtmosphereLights[1] ? LightSceneInfo->Scene->AtmosphereLights[1]->Proxy : nullptr;
-
-		if (VolumetricCloudShadowmapTexture.IsBound())
-		{
-			FVolumetricCloudRenderSceneInfo* CloudInfo = LightSceneInfo->Scene->GetVolumetricCloudSceneInfo();
-
-			const bool bLight0CloudPerPixelTransmittance = CloudInfo && View.VolumetricCloudShadowRenderTarget[0] != nullptr && AtmosphereLight0Proxy && AtmosphereLight0Proxy == LightSceneInfo->Proxy;
-			const bool bLight1CloudPerPixelTransmittance = CloudInfo && View.VolumetricCloudShadowRenderTarget[1] != nullptr && AtmosphereLight1Proxy && AtmosphereLight1Proxy == LightSceneInfo->Proxy;
-
-			if (bLight0CloudPerPixelTransmittance || bLight1CloudPerPixelTransmittance)
-			{
-				uint32 LightIndex = bLight1CloudPerPixelTransmittance ? 1 : 0;
-				SetShaderValue(RHICmdList, ShaderRHI, VolumetricCloudShadowEnabled, 1);
-				SetShaderValue(RHICmdList, ShaderRHI, VolumetricCloudWorldToLightClipShadowMatrix, CloudInfo->GetVolumetricCloudCommonShaderParameters().CloudShadowmapWorldToLightClipMatrix[LightIndex]);
-				SetShaderValue(RHICmdList, ShaderRHI, VolumetricCloudShadowmapFarDepthKm, CloudInfo->GetVolumetricCloudCommonShaderParameters().CloudShadowmapFarDepthKm[LightIndex].X);
-				SetTextureParameter(
-					RHICmdList,
-					ShaderRHI,
-					VolumetricCloudShadowmapTexture,
-					VolumetricCloudShadowmapTextureSampler,
-					TStaticSamplerState<SF_Bilinear, AM_Clamp, AM_Clamp, AM_Clamp>::GetRHI(),
-					View.VolumetricCloudShadowExtractedRenderTarget[LightIndex]->GetShaderResourceRHI());
-
-				if (bLight0CloudPerPixelTransmittance)
-				{
-					SetShaderValue(RHICmdList, ShaderRHI, VolumetricCloudShadowmapStrength, AtmosphereLight0Proxy->GetCloudShadowOnSurfaceStrength());
-				}
-				else if(bLight1CloudPerPixelTransmittance)
-				{
-					SetShaderValue(RHICmdList, ShaderRHI, VolumetricCloudShadowmapStrength, AtmosphereLight1Proxy->GetCloudShadowOnSurfaceStrength());
-				}
-			}
-			else
-			{
-				SetShaderValue(RHICmdList, ShaderRHI, VolumetricCloudShadowEnabled, 0);
-				SetShaderValue(RHICmdList, ShaderRHI, VolumetricCloudWorldToLightClipShadowMatrix, FMatrix44f::Identity);
-				SetShaderValue(RHICmdList, ShaderRHI, VolumetricCloudShadowmapFarDepthKm, 1.0f);
-				SetTextureParameter(
-					RHICmdList,
-					ShaderRHI,
-					VolumetricCloudShadowmapTexture,
-					VolumetricCloudShadowmapTextureSampler,
-					TStaticSamplerState<SF_Bilinear, AM_Clamp, AM_Clamp, AM_Clamp>::GetRHI(),
-					GBlackTexture->TextureRHI);
-				SetShaderValue(RHICmdList, ShaderRHI, VolumetricCloudShadowmapStrength, 0.0f);
-			}
-		}
-		
-		const bool bLightAtmospherePerPixelTransmittance = ShouldRenderSkyAtmosphere(LightSceneInfo->Scene, View.Family->EngineShowFlags) &&
-														(  (AtmosphereLight0Proxy == LightSceneInfo->Proxy && AtmosphereLight0Proxy && AtmosphereLight0Proxy->GetUsePerPixelAtmosphereTransmittance())
-														|| (AtmosphereLight1Proxy == LightSceneInfo->Proxy && AtmosphereLight1Proxy && AtmosphereLight1Proxy->GetUsePerPixelAtmosphereTransmittance()));
-		SetShaderValue(RHICmdList, ShaderRHI, AtmospherePerPixelTransmittanceEnabled, bLightAtmospherePerPixelTransmittance ? 1 : 0);
 	}
 
 private:
-	LAYOUT_FIELD(FVolumeShadowingParameters, VolumeShadowingParameters);
-	LAYOUT_FIELD(FShaderParameter, SpotlightMask);
-	LAYOUT_FIELD(FLightFunctionSharedParameters, LightFunctionParameters);
-	LAYOUT_FIELD(FTranslucentInjectParameters, TranslucentInjectParameters);
-	LAYOUT_FIELD(FShaderParameter, LightFunctionWorldToLight);
-	LAYOUT_FIELD(FShaderParameter, VirtualShadowMapIdParameter);
-
-	LAYOUT_FIELD(FShaderParameter, VolumetricCloudWorldToLightClipShadowMatrix);
-	LAYOUT_FIELD(FShaderParameter, VolumetricCloudShadowmapFarDepthKm);
-	LAYOUT_FIELD(FShaderParameter, VolumetricCloudShadowEnabled);
-	LAYOUT_FIELD(FShaderParameter, VolumetricCloudShadowmapStrength);
-	LAYOUT_FIELD(FShaderResourceParameter, VolumetricCloudShadowmapTexture);
-	LAYOUT_FIELD(FShaderResourceParameter, VolumetricCloudShadowmapTextureSampler);
-	LAYOUT_FIELD(FShaderParameter, AtmospherePerPixelTransmittanceEnabled);
 };
 
 IMPLEMENT_MATERIAL_SHADER_TYPE(,FTranslucentLightingInjectPS, TEXT("/Engine/Private/TranslucentLightInjectionShaders.usf"), TEXT("InjectMainPS"), SF_Pixel);
@@ -1223,6 +1080,7 @@ static void InjectTranslucentLightArray(
 		{
 			const FTranslucentLightInjectionData& InjectionData = LightInjectionData[LightIndex];
 			const FLightSceneInfo* const LightSceneInfo = InjectionData.LightSceneInfo;
+			const FVisibleLightInfo& VisibleLightInfo = Renderer.VisibleLightInfos[LightSceneInfo->Id];
 			const bool bInverseSquared = LightSceneInfo->Proxy->IsInverseSquared();
 			const bool bDirectionalLight = LightSceneInfo->Proxy->GetLightType() == LightType_Directional;
 			const bool bUseVSM = Renderer.VirtualShadowMapArray.IsAllocated();
@@ -1248,6 +1106,76 @@ static void InjectTranslucentLightArray(
 				PassParameters->RenderTargets[0] = FRenderTargetBinding(VolumeAmbientTexture, ERenderTargetLoadAction::ELoad);
 				PassParameters->RenderTargets[1] = FRenderTargetBinding(VolumeDirectionalTexture, ERenderTargetLoadAction::ELoad);
 
+				PassParameters->PS.ViewUniformBuffer = View.ViewUniformBuffer;
+
+				FDeferredLightUniformStruct* DeferredLightStruct = GraphBuilder.AllocParameters<FDeferredLightUniformStruct>();
+				*DeferredLightStruct = GetDeferredLightParameters(View, *LightSceneInfo);
+				PassParameters->PS.DeferredLight = GraphBuilder.CreateUniformBuffer(DeferredLightStruct);
+
+				// InnerSplitIndex: which CSM shadow map level, INDEX_NONE if no directional light
+				int32 InnerSplitIndex = InjectionData.ProjectedShadowInfo ? InjectionData.ProjectedShadowInfo->CascadeSettings.ShadowSplitIndex : INDEX_NONE;
+				const bool bDynamicallyShadowed = InjectionData.ProjectedShadowInfo != NULL;
+
+				GetVolumeShadowingShaderParameters(GraphBuilder, View, LightSceneInfo, InjectionData.ProjectedShadowInfo, InnerSplitIndex, bDynamicallyShadowed, PassParameters->PS.VolumeShadowingParameters);
+
+				PassParameters->PS.VirtualShadowMapId = Renderer.VisibleLightInfos[LightSceneInfo->Id].GetVirtualShadowMapId(&View);
+				PassParameters->PS.LightFunctionParameters = FLightFunctionSharedParameters::GetLightFunctionSharedParameters(LightSceneInfo, 1.0f);
+				PassParameters->PS.VolumeCascadeIndex = VolumeCascadeIndex;
+
+				bool bIsSpotlight = LightSceneInfo->Proxy->GetLightType() == LightType_Spot;
+				PassParameters->PS.SpotlightMask = bIsSpotlight ? 1.0f : 0.0f; //@todo - needs to be a permutation to reduce shadow filtering work
+
+				{
+					const FVector Scale = LightSceneInfo->Proxy->GetLightFunctionScale();
+					// Switch x and z so that z of the user specified scale affects the distance along the light direction
+					const FVector InverseScale = FVector(1.f / Scale.Z, 1.f / Scale.Y, 1.f / Scale.X);
+					const FMatrix44f WorldToLight = LightSceneInfo->Proxy->GetWorldToLight() * FScaleMatrix(FVector(InverseScale));
+					PassParameters->PS.LightFunctionWorldToLight = WorldToLight;
+				}
+
+				{
+					// TODO this should be put in a common structure / set function.
+					FLightSceneProxy* AtmosphereLight0Proxy = LightSceneInfo->Scene->AtmosphereLights[0] ? LightSceneInfo->Scene->AtmosphereLights[0]->Proxy : nullptr;
+					FLightSceneProxy* AtmosphereLight1Proxy = LightSceneInfo->Scene->AtmosphereLights[1] ? LightSceneInfo->Scene->AtmosphereLights[1]->Proxy : nullptr;
+					{
+						FVolumetricCloudRenderSceneInfo* CloudInfo = LightSceneInfo->Scene->GetVolumetricCloudSceneInfo();
+
+						const bool bLight0CloudPerPixelTransmittance = CloudInfo && View.VolumetricCloudShadowRenderTarget[0] != nullptr && AtmosphereLight0Proxy && AtmosphereLight0Proxy == LightSceneInfo->Proxy;
+						const bool bLight1CloudPerPixelTransmittance = CloudInfo && View.VolumetricCloudShadowRenderTarget[1] != nullptr && AtmosphereLight1Proxy && AtmosphereLight1Proxy == LightSceneInfo->Proxy;
+
+						if (bLight0CloudPerPixelTransmittance || bLight1CloudPerPixelTransmittance)
+						{
+							PassParameters->PS.VolumetricCloudShadowmapTexture = View.VolumetricCloudShadowRenderTarget[LightIndex];
+							PassParameters->PS.VolumetricCloudWorldToLightClipShadowMatrix = CloudInfo->GetVolumetricCloudCommonShaderParameters().CloudShadowmapWorldToLightClipMatrix[LightIndex];
+							PassParameters->PS.VolumetricCloudShadowmapFarDepthKm = CloudInfo->GetVolumetricCloudCommonShaderParameters().CloudShadowmapFarDepthKm[LightIndex].X;
+							if (bLight0CloudPerPixelTransmittance)
+							{
+								PassParameters->PS.VolumetricCloudShadowmapStrength = AtmosphereLight0Proxy->GetCloudShadowOnSurfaceStrength();
+							}
+							else if (bLight1CloudPerPixelTransmittance)
+							{
+								PassParameters->PS.VolumetricCloudShadowmapStrength = AtmosphereLight1Proxy->GetCloudShadowOnSurfaceStrength();
+							}
+							PassParameters->PS.VolumetricCloudShadowEnabled = 1;
+						}
+						else
+						{
+							PassParameters->PS.VolumetricCloudShadowmapTexture = GSystemTextures.GetBlackDummy(GraphBuilder);
+							PassParameters->PS.VolumetricCloudWorldToLightClipShadowMatrix = FMatrix44f::Identity;
+							PassParameters->PS.VolumetricCloudShadowmapFarDepthKm = 1.0f;
+							PassParameters->PS.VolumetricCloudShadowmapStrength = 0.0f;
+							PassParameters->PS.VolumetricCloudShadowEnabled = 0;
+						}
+
+						PassParameters->PS.VolumetricCloudShadowmapTextureSampler = TStaticSamplerState<SF_Bilinear, AM_Clamp, AM_Clamp, AM_Clamp>::GetRHI();
+					}
+
+					const bool bLightAtmospherePerPixelTransmittance = ShouldRenderSkyAtmosphere(LightSceneInfo->Scene, View.Family->EngineShowFlags) &&
+						((AtmosphereLight0Proxy == LightSceneInfo->Proxy && AtmosphereLight0Proxy && AtmosphereLight0Proxy->GetUsePerPixelAtmosphereTransmittance())
+							|| (AtmosphereLight1Proxy == LightSceneInfo->Proxy && AtmosphereLight1Proxy && AtmosphereLight1Proxy->GetUsePerPixelAtmosphereTransmittance()));
+					PassParameters->PS.AtmospherePerPixelTransmittanceEnabled = bLightAtmospherePerPixelTransmittance;
+				}
+
 				GraphBuilder.AddPass(
 					RDG_EVENT_NAME("InjectTranslucentLightArray"),
 					PassParameters,
@@ -1296,7 +1224,6 @@ static void InjectTranslucentLightArray(
 					SetGraphicsPipelineState(RHICmdList, GraphicsPSOInit, 0);
 
 					const int32 TranslucencyLightingVolumeDim = GetTranslucencyLightingVolumeDim();
-					const int32 VirtualShadowMapId = Renderer.VisibleLightInfos[LightSceneInfo->Id].GetVirtualShadowMapId( &View );
 
 					VertexShader->SetParameters(RHICmdList, VolumeBounds, FIntVector(TranslucencyLightingVolumeDim));
 					if (GeometryShader.IsValid())
@@ -1304,13 +1231,7 @@ static void InjectTranslucentLightArray(
 						GeometryShader->SetParameters(RHICmdList, VolumeBounds.MinZ);
 					}
 
-					PixelShader->SetParameters(RHICmdList, View, LightSceneInfo,
-						InjectionData.LightFunctionMaterialProxy,
-						InjectionData.ProjectedShadowInfo,
-						InjectionData.ProjectedShadowInfo ? InjectionData.ProjectedShadowInfo->CascadeSettings.ShadowSplitIndex : INDEX_NONE,
-						VolumeCascadeIndex,
-						VirtualShadowMapId);
-
+					PixelShader->SetParameters(RHICmdList, View, InjectionData.LightFunctionMaterialProxy);
 					SetShaderParameters(RHICmdList, PixelShader, PixelShader.GetPixelShader(), PassParameters->PS);
 
 					RasterizeToVolumeTexture(RHICmdList, VolumeBounds);
