Fixes r.TSR.HistoryScreenPercentage's history convergance rate and use r.TSR.HistoryScreenPercentage on cinematic scalability.

diff --git a/Engine/Config/BaseScalability.ini b/Engine/Config/BaseScalability.ini
--- a/Engine/Config/BaseScalability.ini
+++ b/Engine/Config/BaseScalability.ini
@@ -40,26 +40,31 @@ PerfIndexValues_ResolutionQuality="50 71 87 100 100"
 r.FXAA.Quality=0
 r.TemporalAA.Quality=0
 r.TSR.RejectionAntiAliasingQuality=0
+r.TSR.HistoryScreenPercentage=100
 
 [AntiAliasingQuality@1]
 r.FXAA.Quality=1
 r.TemporalAA.Quality=1
 r.TSR.RejectionAntiAliasingQuality=0
+r.TSR.HistoryScreenPercentage=100
 
 [AntiAliasingQuality@2]
 r.FXAA.Quality=3
 r.TemporalAA.Quality=1
 r.TSR.RejectionAntiAliasingQuality=1
+r.TSR.HistoryScreenPercentage=100
 
 [AntiAliasingQuality@3]
 r.FXAA.Quality=4
 r.TemporalAA.Quality=2
 r.TSR.RejectionAntiAliasingQuality=1
+r.TSR.HistoryScreenPercentage=100
 
 [AntiAliasingQuality@Cine]
 r.FXAA.Quality=5
 r.TemporalAA.Quality=2
 r.TSR.RejectionAntiAliasingQuality=2
+r.TSR.HistoryScreenPercentage=200
 
 
 ;-----------------------------------------------------------------------------------------------------------------
diff --git a/Engine/Shaders/Private/TemporalSuperResolution/TSRCommon.ush b/Engine/Shaders/Private/TemporalSuperResolution/TSRCommon.ush
--- a/Engine/Shaders/Private/TemporalSuperResolution/TSRCommon.ush
+++ b/Engine/Shaders/Private/TemporalSuperResolution/TSRCommon.ush
@@ -581,6 +581,16 @@ half2 SafeRcp(half2 x)
 
 #endif
 
+taa_half2 ComputeLowFrequencyContribution(taa_half2 HistoryValidity)
+{
+	return saturate((HistoryValidity - taa_half(1.0 * rcp(MAX_SAMPLE_COUNT))) * taa_half(6.0));
+}
+
+taa_half ComputeLowFrequencyContribution(taa_half HistoryValidity)
+{
+	return dpv_lo(ComputeLowFrequencyContribution(dpv_interleave_mono_registers(HistoryValidity)));
+}
+
 taa_half3x2 RGBToYCoCg(taa_half3x2 RGB)
 {
     taa_half2 R = RGB[0];
diff --git a/Engine/Shaders/Private/TemporalSuperResolution/TSRResolveHistory.usf b/Engine/Shaders/Private/TemporalSuperResolution/TSRResolveHistory.usf
new file mode 100644
--- /dev/null
+++ b/Engine/Shaders/Private/TemporalSuperResolution/TSRResolveHistory.usf
@@ -0,0 +1,237 @@
+// Copyright Epic Games, Inc. All Rights Reserved.
+
+#include "TSRCommon.ush"
+
+
+//------------------------------------------------------- CONFIG
+
+#define TILE_SIZE 8
+
+#define KERNEL_SIZE 4
+
+
+//------------------------------------------------------- CONSTANTS
+
+static const int2 kOffsetsCross3x3[4] = {
+	int2(-1, -1),
+	int2( 1, -1),
+	int2(-1,  1),
+	int2( 1,  1),
+};
+
+
+//------------------------------------------------------- PARAMETERS
+
+RWTexture2D<uint> PrevUseCountOutput;
+RWTexture2D<uint> PrevClosestDepthOutput;
+
+
+//------------------------------------------------------- LDS
+
+groupshared taa_half3 SharedArray0[TILE_SIZE * TILE_SIZE];
+
+
+//------------------------------------------------------- FUNCTIONS
+
+// Mitchell Netravali
+// B-spline:       B=1      C=0
+// Mitchell:       B=1/3    C=1/3
+// Catmull-Rom:    B=0      C=1/2
+// Robidoux:       B=0.3782 C=0.3109 (cylindrical)
+// Robidoux Sharp: B=0.2620 C=0.3690 (cylindrical)
+// Robidoux Soft:  B=0.6796 C=0.1602 (cylindrical)
+void MitchellNetravaliCoefs(const float B, const float C, out float OutQ0[4], out float OutQ1[4])
+{
+	OutQ0[0] = (6.0 - 2.0 * B) / 6.0;
+	OutQ0[1] = 0.0;
+	OutQ0[2] = (-18.0 + 12.0 * B + 6.0 * C) / 6.0;
+	OutQ0[3] = (12.0 - 9.0 * B - 6.0 * C) / 6.0;
+
+	OutQ1[0] = (8 * B + 24 * C) / 6.0;
+	OutQ1[1] = (-12 * B - 48 * C) / 6.0;
+	OutQ1[2] = (6 * B + 30 * C) / 6.0;
+	OutQ1[3] = (-B - 6 * C) / 6.0;
+}
+
+float MitchellNetravali(float d, const float B, const float C)
+{
+	// Coeficient ends up known at compile time.
+	float Q0[4];
+	float Q1[4];
+	MitchellNetravaliCoefs(B, C, Q0, Q1);
+
+	if (d < 1)
+	{
+		return Q0[0] + d * (Q0[1] + d * (Q0[2] + d * Q0[3]));
+	}
+	else if ((d >= 1) && (d < 2))
+	{
+		return Q1[0] + d * (Q1[1] + d * (Q1[2] + d * Q1[3]));
+	}
+	else
+	{
+		return 0;
+	}
+}
+
+
+//------------------------------------------------------- PARAMETERS
+
+FScreenTransform DispatchThreadToHistoryPixelPos;
+uint2 OutputViewRectMin;
+uint2 OutputViewRectMax;
+uint bGenerateOutputMip1;
+float HistoryValidityMultiply;
+
+Texture2D<taa_half3> History_LowFrequency;
+Texture2D<taa_half3> History_HighFrequency;
+Texture2D<taa_half2> History_Metadata;
+Texture2D<taa_half3> History_Translucency;
+Texture2D<taa_half>  History_TranslucencyAlpha;
+
+RWTexture2D<taa_half3> SceneColorOutputMip0;
+RWTexture2D<taa_half3> SceneColorOutputMip1;
+
+
+//------------------------------------------------------- ENTRY POINT
+
+[numthreads(TILE_SIZE, TILE_SIZE, 1)]
+void MainCS(
+	uint2 GroupId : SV_GroupID,
+	uint GroupThreadIndex : SV_GroupIndex)
+{
+	float4 Debug = 0.0;
+
+	uint2 DispatchThreadId = (
+		ZOrder2D(GroupThreadIndex, log2(TILE_SIZE)) +
+		GroupId * uint2(TILE_SIZE, TILE_SIZE));
+
+	taa_half3 LowFrequencyColor = 0.0;
+	taa_half3 HighFrequencyColor = 0.0;
+	taa_half  HistoryValidity = 0.0;
+	
+	// Downsample the history
+	{
+		const float B = rcp(3.0);
+		const float C = rcp(3.0);
+
+		// Position of the output pixel in the history.
+		float2 OutputPixelPos = ApplyScreenTransform(float2(DispatchThreadId), DispatchThreadToHistoryPixelPos);
+		float2 TopLeftKernel = floor(OutputPixelPos + (0.5 - 0.5 * KERNEL_SIZE)) + 0.5;
+		
+		taa_half LowFrequencyWeight = 0.0;
+		taa_half HighFrequencyWeight = 0.0;
+		taa_half HistoryValidityWeight = 0.0;
+
+		UNROLL_N(KERNEL_SIZE)
+		for (uint x = 0; x < KERNEL_SIZE; x++)
+		{
+			UNROLL_N(KERNEL_SIZE)
+			for (uint y = 0; y < KERNEL_SIZE; y++)
+			{
+				float2 SamplePixelPos = TopLeftKernel + float2(x, y);
+
+				float2 SampleUV = SamplePixelPos * HistoryInfo_ExtentInverse;
+
+				SampleUV = clamp(SampleUV, HistoryInfo_UVViewportBilinearMin, HistoryInfo_UVViewportBilinearMax);
+
+				taa_half3 SampleLowFrequencyColor  = History_LowFrequency.SampleLevel(GlobalBilinearClampedSampler, SampleUV, 0);
+				taa_half3 SampleHighFrequencyColor = History_HighFrequency.SampleLevel(GlobalBilinearClampedSampler, SampleUV, 0);
+				taa_half2 SampleMetadata   	       = History_Metadata.SampleLevel(GlobalBilinearClampedSampler, SampleUV, 0);
+				
+				taa_half3 SampleTranslucencyColor = History_Translucency.SampleLevel(GlobalBilinearClampedSampler, SampleUV, 0);
+				taa_half  SampleTranslucencyAlpha = History_TranslucencyAlpha.SampleLevel(GlobalBilinearClampedSampler, SampleUV, 0);
+
+				#if 1
+				{
+					SampleLowFrequencyColor  = SampleLowFrequencyColor  * SampleTranslucencyAlpha + SampleTranslucencyColor;
+					SampleHighFrequencyColor = SampleHighFrequencyColor * SampleTranslucencyAlpha + SampleTranslucencyColor;
+				}
+				#endif
+
+				taa_half SampleHistoryValidity = SampleMetadata.y;
+
+				float2 PixelOffset = abs(SamplePixelPos - OutputPixelPos);
+
+				float kernelWeight = MitchellNetravali(PixelOffset.x, B, C) * MitchellNetravali(PixelOffset.y, B, C);
+				
+				float LowFrequencyHdrWeight  = HdrWeight4(SampleLowFrequencyColor);
+				float HighFrequencyHdrWeight = HdrWeight4(SampleHighFrequencyColor * SafeRcp(SampleHistoryValidity));
+
+				LowFrequencyColor   += (kernelWeight * LowFrequencyHdrWeight)  * SampleLowFrequencyColor;
+				LowFrequencyWeight  += (kernelWeight * LowFrequencyHdrWeight);
+
+				HighFrequencyColor  += (kernelWeight * HighFrequencyHdrWeight) * SampleHighFrequencyColor;
+				HighFrequencyWeight += (kernelWeight * HighFrequencyHdrWeight) * SampleHistoryValidity;
+
+				HistoryValidity       += max(kernelWeight, 0.0) * SampleHistoryValidity;
+				HistoryValidityWeight += max(kernelWeight, 0.0);
+			}
+		}
+
+		LowFrequencyColor  *= rcp(LowFrequencyWeight);
+		HighFrequencyColor *= SafeRcp(HighFrequencyWeight);
+		HistoryValidity *= HistoryValidityMultiply / HistoryValidityWeight;
+	}
+	
+	// Compute final output
+	taa_half3 FinalOutputColor;
+	{
+		taa_half OutputBlend = ComputeLowFrequencyContribution(HistoryValidity);
+		
+		taa_half LowFrequencyWeight = HdrWeight4(LowFrequencyColor);
+		taa_half Moment1Weight = HdrWeight4(HighFrequencyColor);
+
+		taa_half2 Weights = WeightedLerpFactors(LowFrequencyWeight, Moment1Weight, OutputBlend);
+		
+		FinalOutputColor = LowFrequencyColor * Weights[0] + HighFrequencyColor * Weights[1];
+	}
+
+	taa_short2 OutputPixelPosition = InvalidateOutputPixelPos(OutputViewRectMin + DispatchThreadId, OutputViewRectMax);
+	
+	// Output final scene color Mip0
+	{
+		SceneColorOutputMip0[OutputPixelPosition] = FinalOutputColor;
+	}
+
+	// Output final scene color Mip1
+	{
+		uint LocalGroupThreadIndex = GroupThreadIndex;
+
+		taa_half3 HalfResOutput = FinalOutputColor * taa_half(0.25);
+		taa_short2 HalfResOutputPixelPos;
+		HalfResOutputPixelPos[0] = (OutputPixelPosition[0] >> taa_short(1)) | (((OutputPixelPosition[0] & taa_short(0x1)) | !bGenerateOutputMip1) * taa_short(~0));
+		HalfResOutputPixelPos[1] = (OutputPixelPosition[1] >> taa_short(1)) | (((OutputPixelPosition[1] & taa_short(0x1)) | !bGenerateOutputMip1) * taa_short(~0));
+			
+		#if PLATFORM_SUPPORTS_WAVE_BROADCAST && 0 // Support WaveBroadcast with halfs
+		BRANCH
+		if (bGenerateOutputMip1)
+		{
+			FWaveBroadcastSettings Horizontal = InitWaveXorButterfly(/* XorButterFly = */ 0x1);
+			FWaveBroadcastSettings Vertical = InitWaveXorButterfly(/* XorButterFly = */ 0x2);
+
+			HalfResOutput += WaveBroadcast(Horizontal, HalfResOutput);
+			HalfResOutput += WaveBroadcast(Vertical, HalfResOutput);
+		}
+		#else
+		BRANCH
+		if (bGenerateOutputMip1)
+		{
+			SharedArray0[LocalGroupThreadIndex] = HalfResOutput;
+			SharedArray0[LocalGroupThreadIndex] += SharedArray0[LocalGroupThreadIndex ^ 0x1];
+			SharedArray0[LocalGroupThreadIndex] += SharedArray0[LocalGroupThreadIndex ^ 0x2];
+			
+			HalfResOutput = SharedArray0[LocalGroupThreadIndex];
+		}
+		#endif
+
+		SceneColorOutputMip1[HalfResOutputPixelPos] = HalfResOutput;
+	}
+
+	#if DEBUG_OUTPUT
+	{
+		DebugOutput[OutputPixelPosition] = Debug;
+	}
+	#endif
+}
+ 
\ No newline at end of file
diff --git a/Engine/Shaders/Private/TemporalSuperResolution/TSRUpdateHistory.usf b/Engine/Shaders/Private/TemporalSuperResolution/TSRUpdateHistory.usf
--- a/Engine/Shaders/Private/TemporalSuperResolution/TSRUpdateHistory.usf
+++ b/Engine/Shaders/Private/TemporalSuperResolution/TSRUpdateHistory.usf
@@ -135,11 +135,6 @@ taa_ushort2x2 Map8x8Tile2x2LaneDPV(uint GroupThreadIndex)
 #endif
 }
 
-taa_half2 ComputeLowFrequencyContribution(taa_half2 HistoryValidity)
-{
-	return saturate((HistoryValidity - taa_half(1.0 * rcp(MAX_SAMPLE_COUNT))) * taa_half(6.0));
-}
-
 taa_half3x2 TransformColorForClampingBox(taa_half3x2 Color)
 {
 	return RGBToYCoCg(Color);
@@ -1292,7 +1287,6 @@ void MainCS(
 
 	// Compute final output
 	taa_half3x2 FinalOutputColor;
-	taa_half2 FinalOutputValidity;
 	{
 		taa_half2 OutputBlend = ComputeLowFrequencyContribution(FinalHistoryValidity);
 
@@ -1313,8 +1307,6 @@ void MainCS(
 		}
 		#endif
 
-		FinalOutputValidity = lerp(taa_half(1.0), saturate(FinalHistoryValidity), OutputBlend);
-
 		taa_half2 LowFrequencyWeight = HdrWeight4(FinalLowFrequencyColor);
 		taa_half2 Moment1Weight = HdrWeight4(FinalHighFrequencyColor);
 
diff --git a/Engine/Source/Runtime/Renderer/Private/PostProcess/TemporalSuperResolution.cpp b/Engine/Source/Runtime/Renderer/Private/PostProcess/TemporalSuperResolution.cpp
--- a/Engine/Source/Runtime/Renderer/Private/PostProcess/TemporalSuperResolution.cpp
+++ b/Engine/Source/Runtime/Renderer/Private/PostProcess/TemporalSuperResolution.cpp
@@ -2,7 +2,6 @@
 
 #include "PostProcess/TemporalAA.h"
 #include "PostProcess/PostProcessTonemap.h"
-#include "PostProcess/PostProcessMitchellNetravali.h"
 #include "PostProcess/PostProcessing.h"
 #include "ClearQuad.h"
 #include "PostProcessing.h"
@@ -20,7 +19,7 @@ TAutoConsoleVariable<float> CVarTSRHistorySP(
 	TEXT("r.TSR.HistoryScreenPercentage"),
 	100.0f,
 	TEXT("Size of TSR's history."),
-	ECVF_RenderThreadSafe);
+	ECVF_Scalability | ECVF_RenderThreadSafe);
 
 TAutoConsoleVariable<int32> CVarTSRR11G11B10History(
 	TEXT("r.TSR.R11G11B10History"), 1,
@@ -464,6 +463,27 @@ class FTSRUpdateHistoryCS : public FTSRShader
 	}
 }; // class FTSRUpdateHistoryCS
 
+class FTSRResolveHistoryCS : public FTSRShader
+{
+	DECLARE_GLOBAL_SHADER(FTSRResolveHistoryCS);
+	SHADER_USE_PARAMETER_STRUCT(FTSRResolveHistoryCS, FTSRShader);
+
+	BEGIN_SHADER_PARAMETER_STRUCT(FParameters, )
+		SHADER_PARAMETER_STRUCT_INCLUDE(FTSRCommonParameters, CommonParameters)
+		SHADER_PARAMETER(FScreenTransform, DispatchThreadToHistoryPixelPos)
+		SHADER_PARAMETER(FIntPoint, OutputViewRectMin)
+		SHADER_PARAMETER(FIntPoint, OutputViewRectMax)
+		SHADER_PARAMETER(int32, bGenerateOutputMip1)
+		SHADER_PARAMETER(float, HistoryValidityMultiply)
+
+		SHADER_PARAMETER_STRUCT(FTSRHistoryTextures, History)
+
+		SHADER_PARAMETER_RDG_TEXTURE_UAV(RWTexture2D, SceneColorOutputMip0)
+		SHADER_PARAMETER_RDG_TEXTURE_UAV(RWTexture2D, SceneColorOutputMip1)
+		SHADER_PARAMETER_RDG_TEXTURE_UAV(RWTexture2D, DebugOutput)
+	END_SHADER_PARAMETER_STRUCT()
+}; // class FTSRResolveHistoryCS
+
 #if COMPILE_TSR_DEBUG_PASSES
 
 class FTSRDebugHistoryCS : public FTSRShader
@@ -476,7 +496,7 @@ class FTSRDebugHistoryCS : public FTSRShader
 		SHADER_PARAMETER_STRUCT(FTSRHistoryTextures, History)
 		SHADER_PARAMETER_STRUCT(FTSRHistoryTextures, PrevHistory)
 		SHADER_PARAMETER_RDG_TEXTURE_UAV(RWTexture2D, DebugOutput)
-		END_SHADER_PARAMETER_STRUCT()
+	END_SHADER_PARAMETER_STRUCT()
 }; // class FTSRDebugHistoryCS
 
 #endif
@@ -493,6 +513,7 @@ IMPLEMENT_GLOBAL_SHADER(FTSRDilateRejectionCS,       "/Engine/Private/TemporalSu
 IMPLEMENT_GLOBAL_SHADER(FTSRSpatialAntiAliasingCS,   "/Engine/Private/TemporalSuperResolution/TSRSpatialAntiAliasing.usf",   "MainCS", SF_Compute);
 IMPLEMENT_GLOBAL_SHADER(FTSRFilterAntiAliasingCS,    "/Engine/Private/TemporalSuperResolution/TSRFilterAntiAliasing.usf",    "MainCS", SF_Compute);
 IMPLEMENT_GLOBAL_SHADER(FTSRUpdateHistoryCS,         "/Engine/Private/TemporalSuperResolution/TSRUpdateHistory.usf",         "MainCS", SF_Compute);
+IMPLEMENT_GLOBAL_SHADER(FTSRResolveHistoryCS,        "/Engine/Private/TemporalSuperResolution/TSRResolveHistory.usf",        "MainCS", SF_Compute);
 
 #if COMPILE_TSR_DEBUG_PASSES
 IMPLEMENT_GLOBAL_SHADER(FTSRDebugHistoryCS,          "/Engine/Private/TemporalSuperResolution/TSRDebugHistory.usf",          "MainCS", SF_Compute);
@@ -509,6 +530,20 @@ bool ComposeSeparateTranslucencyInTSR(const FViewInfo& View)
 	return CVarTSRTranslucencySeparateTemporalAccumulation.GetValueOnRenderThread() != 0;
 }
 
+static FRDGTextureUAVRef CreateDummyUAV(FRDGBuilder& GraphBuilder, EPixelFormat PixelFormat)
+{
+	FRDGTextureDesc Desc = FRDGTextureDesc::Create2D(
+		FIntPoint(1, 1),
+		PixelFormat,
+		FClearValueBinding::None,
+		/* InFlags = */ TexCreate_ShaderResource | TexCreate_UAV);
+
+	FRDGTextureRef DummyTexture = GraphBuilder.CreateTexture(Desc, TEXT("TSR.DummyOutput"));
+	GraphBuilder.RemoveUnusedTextureWarning(DummyTexture);
+
+	return GraphBuilder.CreateUAV(DummyTexture);
+};
+
 ITemporalUpscaler::FOutputs AddTemporalSuperResolutionPasses(
 	FRDGBuilder& GraphBuilder,
 	const FViewInfo& View,
@@ -1259,20 +1294,21 @@ ITemporalUpscaler::FOutputs AddTemporalSuperResolutionPasses(
 			FComputeShaderUtils::GetGroupCount(RejectionRect.Size(), 8));
 	}
 
+	// Allocate output
 	FRDGTextureRef SceneColorOutputTexture;
 	{
-		// Allocate output
-		{
-			FRDGTextureDesc Desc = FRDGTextureDesc::Create2D(
-				HistoryExtent,
-				PF_FloatR11G11B10,
-				FClearValueBinding::None,
-				/* InFlags = */ TexCreate_ShaderResource | TexCreate_UAV,
-				/* NumMips = */ PassInputs.bGenerateOutputMip1 ? 2 : 1);
+		FRDGTextureDesc Desc = FRDGTextureDesc::Create2D(
+			OutputExtent,
+			PF_FloatR11G11B10,
+			FClearValueBinding::None,
+			/* InFlags = */ TexCreate_ShaderResource | TexCreate_UAV,
+			/* NumMips = */ PassInputs.bGenerateOutputMip1 ? 2 : 1);
 
-			SceneColorOutputTexture = GraphBuilder.CreateTexture(Desc, TEXT("TSR.Output"));
-		}
+		SceneColorOutputTexture = GraphBuilder.CreateTexture(Desc, TEXT("TSR.Output"));
+	}
 
+	// Update temporal history.
+	{
 		FTSRUpdateHistoryCS::FParameters* PassParameters = GraphBuilder.AllocParameters<FTSRUpdateHistoryCS::FParameters>();
 		PassParameters->CommonParameters = CommonParameters;
 		PassParameters->InputSceneColorTexture = PassInputs.SceneColorTexture;
@@ -1297,28 +1333,28 @@ ITemporalUpscaler::FOutputs AddTemporalSuperResolutionPasses(
 		PassParameters->InvWeightClampingPixelSpeed = 1.0f / CVarTSRWeightClampingPixelSpeed.GetValueOnRenderThread();
 		PassParameters->InputToHistoryFactor = float(HistorySize.X) / float(InputRect.Width());
 		PassParameters->ResponsiveStencilMask = CVarTSREnableResponiveAA.GetValueOnRenderThread() ? (STENCIL_TEMPORAL_RESPONSIVE_AA_MASK) : 0;
-		PassParameters->bGenerateOutputMip1 = PassInputs.bGenerateOutputMip1 ? 1 : 0;
+		PassParameters->bGenerateOutputMip1 = (PassInputs.bGenerateOutputMip1 && HistorySize == OutputRect.Size()) ? 1 : 0;
 
 		PassParameters->PrevHistoryParameters = PrevHistoryParameters;
 		PassParameters->PrevHistory = PrevHistory;
 
 		PassParameters->HistoryOutput = CreateUAVs(GraphBuilder, History);
-		PassParameters->SceneColorOutputMip0 = GraphBuilder.CreateUAV(FRDGTextureUAVDesc(SceneColorOutputTexture, /* InMipLevel = */ 0));
-		if (PassInputs.bGenerateOutputMip1)
+		if (HistorySize != OutputRect.Size())
+		{
+			PassParameters->SceneColorOutputMip0 = CreateDummyUAV(GraphBuilder, PF_FloatR11G11B10);
+		}
+		else
+		{
+			PassParameters->SceneColorOutputMip0 = GraphBuilder.CreateUAV(FRDGTextureUAVDesc(SceneColorOutputTexture, /* InMipLevel = */ 0));
+		}
+		
+		if (PassParameters->bGenerateOutputMip1)
 		{
 			PassParameters->SceneColorOutputMip1 = GraphBuilder.CreateUAV(FRDGTextureUAVDesc(SceneColorOutputTexture, /* InMipLevel = */ 1));
 		}
 		else
 		{
-			FRDGTextureDesc Desc = FRDGTextureDesc::Create2D(
-				FIntPoint(1, 1),
-				PF_FloatR11G11B10,
-				FClearValueBinding::None,
-				/* InFlags = */ TexCreate_ShaderResource | TexCreate_UAV);
-
-			FRDGTextureRef DummyTexture = GraphBuilder.CreateTexture(Desc, TEXT("TSR.DummyOutput"));
-			PassParameters->SceneColorOutputMip1 = GraphBuilder.CreateUAV(DummyTexture);
-			GraphBuilder.RemoveUnusedTextureWarning(DummyTexture);
+			PassParameters->SceneColorOutputMip1 = CreateDummyUAV(GraphBuilder, PF_FloatR11G11B10);
 		}
 		PassParameters->DebugOutput = CreateDebugUAV(HistoryExtent, TEXT("Debug.TSR.UpdateHistory"));
 
@@ -1365,10 +1401,38 @@ ITemporalUpscaler::FOutputs AddTemporalSuperResolutionPasses(
 	// If we upscaled the history buffer, downsize back to the secondary screen percentage size.
 	if (HistorySize != OutputRect.Size())
 	{
-		SceneColorOutputTexture = ComputeMitchellNetravaliDownsample(
-			GraphBuilder, View,
-			/* InputViewport = */ FScreenPassTexture(SceneColorOutputTexture, FIntRect(FIntPoint(0, 0), HistorySize)),
-			/* OutputViewport = */ FScreenPassTextureViewport(OutputExtent, OutputRect));
+		FTSRResolveHistoryCS::FParameters* PassParameters = GraphBuilder.AllocParameters<FTSRResolveHistoryCS::FParameters>();
+		PassParameters->CommonParameters = CommonParameters;
+		PassParameters->DispatchThreadToHistoryPixelPos = (
+			FScreenTransform::DispatchThreadIdToViewportUV(OutputRect) *
+			FScreenTransform::ChangeTextureBasisFromTo(
+				HistoryExtent, FIntRect(FIntPoint(0, 0), HistorySize),
+				FScreenTransform::ETextureBasis::ViewportUV, FScreenTransform::ETextureBasis::TexelPosition));
+		PassParameters->OutputViewRectMin = OutputRect.Min;
+		PassParameters->OutputViewRectMax = OutputRect.Max;
+		PassParameters->bGenerateOutputMip1 = PassInputs.bGenerateOutputMip1 ? 1 : 0;
+		PassParameters->HistoryValidityMultiply = float(HistorySize.X * HistorySize.Y) / float(OutputRect.Width() * OutputRect.Height());
+
+		PassParameters->History = History;
+		
+		PassParameters->SceneColorOutputMip0 = GraphBuilder.CreateUAV(FRDGTextureUAVDesc(SceneColorOutputTexture, /* InMipLevel = */ 0));
+		if (PassParameters->bGenerateOutputMip1)
+		{
+			PassParameters->SceneColorOutputMip1 = GraphBuilder.CreateUAV(FRDGTextureUAVDesc(SceneColorOutputTexture, /* InMipLevel = */ 1));
+		}
+		else
+		{
+			PassParameters->SceneColorOutputMip1 = CreateDummyUAV(GraphBuilder, PF_FloatR11G11B10);
+		}
+		PassParameters->DebugOutput = CreateDebugUAV(HistoryExtent, TEXT("Debug.TSR.ResolveHistory"));
+
+		TShaderMapRef<FTSRResolveHistoryCS> ComputeShader(View.ShaderMap);
+		FComputeShaderUtils::AddPass(
+			GraphBuilder,
+			RDG_EVENT_NAME("TSR ResolveHistory %dx%d", OutputRect.Width(), OutputRect.Height()),
+			ComputeShader,
+			PassParameters,
+			FComputeShaderUtils::GetGroupCount(OutputRect.Size(), 8));
 	}
 
 	// Extract all resources for next frame.
diff --git a/Engine/Source/Runtime/Renderer/Private/ScreenPass.cpp b/Engine/Source/Runtime/Renderer/Private/ScreenPass.cpp
--- a/Engine/Source/Runtime/Renderer/Private/ScreenPass.cpp
+++ b/Engine/Source/Runtime/Renderer/Private/ScreenPass.cpp
@@ -83,6 +83,18 @@ FScreenTransform FScreenTransform::ChangeTextureUVCoordinateFromTo(
 		ChangeTextureBasisFromTo(DestViewport, ETextureBasis::ViewportUV, ETextureBasis::TextureUV));
 }
 
+// static
+FScreenTransform FScreenTransform::SvPositionToViewportUV(const FIntRect& SrcViewport)
+{
+	return (FScreenTransform::Identity - SrcViewport.Min) / SrcViewport.Size();
+}
+
+// static
+FScreenTransform FScreenTransform::DispatchThreadIdToViewportUV(const FIntRect& SrcViewport)
+{
+	return (FScreenTransform::Identity + 0.5f) / SrcViewport.Size();
+}
+
 void SetScreenPassPipelineState(FRHICommandList& RHICmdList, const FScreenPassPipelineState& ScreenPassDraw)
 {
 	FGraphicsPipelineStateInitializer GraphicsPSOInit;
diff --git a/Engine/Source/Runtime/Renderer/Private/ScreenPass.h b/Engine/Source/Runtime/Renderer/Private/ScreenPass.h
--- a/Engine/Source/Runtime/Renderer/Private/ScreenPass.h
+++ b/Engine/Source/Runtime/Renderer/Private/ScreenPass.h
@@ -318,6 +318,8 @@ struct FScreenTransform
 		const FScreenPassTextureViewport& SrcViewport,
 		const FScreenPassTextureViewport& DestViewport);
 
+	static FScreenTransform SvPositionToViewportUV(const FIntRect& SrcViewport);
+	static FScreenTransform DispatchThreadIdToViewportUV(const FIntRect& SrcViewport);
 }; // FScreenTransform
 
 // A utility shader parameter struct containing the viewport, texture, and sampler for a unique texture input to a shader.
diff --git a/Engine/Source/Runtime/Renderer/Private/ScreenPass.inl b/Engine/Source/Runtime/Renderer/Private/ScreenPass.inl
--- a/Engine/Source/Runtime/Renderer/Private/ScreenPass.inl
+++ b/Engine/Source/Runtime/Renderer/Private/ScreenPass.inl
@@ -227,6 +227,22 @@ inline FScreenTransform operator + (const FScreenTransform& AToB, const FIntPoin
 }
 
 
+inline FScreenTransform operator - (const FScreenTransform& AToB, const FVector2f& Bias)
+{
+	return FScreenTransform(AToB.Scale, AToB.Bias - Bias);
+}
+
+inline FScreenTransform operator - (const FScreenTransform& AToB, const float& Bias)
+{
+	return AToB - FVector2f(Bias, Bias);
+}
+
+inline FScreenTransform operator - (const FScreenTransform& AToB, const FIntPoint& Bias)
+{
+	return AToB - FVector2f(Bias.X, Bias.Y);
+}
+
+
 inline FScreenTransform operator / (const FScreenTransform& AToB, const FVector2f& InvertedScale)
 {
 	ensure(!FMath::IsNearlyZero(InvertedScale.X));
