Made more code common with default LightCloudTransmittance structure.
Code clean up.

diff --git a/Engine/Shaders/Private/TranslucentLightInjectionShaders.usf b/Engine/Shaders/Private/TranslucentLightInjectionShaders.usf
--- a/Engine/Shaders/Private/TranslucentLightInjectionShaders.usf
+++ b/Engine/Shaders/Private/TranslucentLightInjectionShaders.usf
@@ -19,13 +19,9 @@ float4 SimpleLightPositionAndRadius;
 float4 SimpleLightColorAndExponent;
 
 #if RADIAL_ATTENUATION==0
+#define USE_CLOUD_TRANSMITTANCE 1
 #include "VolumetricCloudCommon.ush"
-float4x4 VolumetricCloudWorldToLightClipShadowMatrix;
-float VolumetricCloudShadowmapFarDepthKm;
-Texture2D<float3> VolumetricCloudShadowmapTexture;
-SamplerState VolumetricCloudShadowmapTextureSampler;
 uint VolumetricCloudShadowEnabled;
-float VolumetricCloudShadowmapStrength;
 
 #include "/Engine/Private/SkyAtmosphereCommon.ush"
 uint AtmospherePerPixelTransmittanceEnabled;
@@ -210,7 +206,7 @@ void InjectMainPS(
 		if (VolumetricCloudShadowEnabled > 0)
 		{
 			float OutOpticalDepth = 0.0f;
-			ShadowFactor *= lerp(1.0, GetCloudVolumetricShadow(WorldPositionForLighting, VolumetricCloudWorldToLightClipShadowMatrix, VolumetricCloudShadowmapFarDepthKm, VolumetricCloudShadowmapTexture, VolumetricCloudShadowmapTextureSampler, OutOpticalDepth), VolumetricCloudShadowmapStrength);
+			ShadowFactor *= lerp(1.0f, GetCloudVolumetricShadow(WorldPositionForLighting, CloudShadowmapWorldToLightClipMatrix, CloudShadowmapFarDepthKm, CloudShadowmapTexture, CloudShadowmapSampler, OutOpticalDepth), CloudShadowmapStrength);
 		}
 		#endif
 		
diff --git a/Engine/Source/Runtime/Renderer/Private/LightRendering.cpp b/Engine/Source/Runtime/Renderer/Private/LightRendering.cpp
--- a/Engine/Source/Runtime/Renderer/Private/LightRendering.cpp
+++ b/Engine/Source/Runtime/Renderer/Private/LightRendering.cpp
@@ -453,7 +453,8 @@ class FDeferredLightPS : public FGlobalShader
 		SHADER_PARAMETER_RDG_UNIFORM_BUFFER(FHairStrandsViewUniformParameters, HairStrands)
 		SHADER_PARAMETER_RDG_UNIFORM_BUFFER(FStrataGlobalUniformParameters, Strata)
 		SHADER_PARAMETER_STRUCT_INCLUDE(FVolumetricCloudShadowAOParameters, CloudShadowAO)
-		SHADER_PARAMETER_STRUCT_INCLUDE(FVolumetricCloudShadowParameters, CloudShadow)
+		SHADER_PARAMETER_STRUCT_INCLUDE(FLightCloudTransmittanceParameters, CloudShadow)
+		SHADER_PARAMETER(uint32, CloudShadowEnabled)
 		SHADER_PARAMETER(uint32, HairTransmittanceBufferMaxCount)
 		SHADER_PARAMETER(uint32, HairShadowMaskValid)
 		SHADER_PARAMETER(FVector4f, ShadowChannelMask)
@@ -1800,7 +1801,7 @@ static FDeferredLightPS::FParameters GetDeferredLightPSParameters(
 	Out.LightingChannelsTexture = LightingChannelsTexture ? LightingChannelsTexture : WhiteDummy;
 	Out.LightingChannelsSampler = TStaticSamplerState<SF_Point, AM_Clamp, AM_Clamp, AM_Clamp>::GetRHI();
 	Out.CloudShadowAO = GetCloudShadowAOParameters(GraphBuilder, View, CloudInfo);
-	Out.CloudShadow = GetCloudShadowParameters(GraphBuilder, Scene, View, LightSceneInfo->Proxy, CloudInfo);
+	Out.CloudShadowEnabled = SetupLightCloudTransmittanceParameters(GraphBuilder, Scene, View, LightSceneInfo, Out.CloudShadow) ? 1 : 0;
 	Out.LightAttenuationTexture = ShadowMaskTexture ? ShadowMaskTexture : WhiteDummy;
 	Out.LightAttenuationTextureSampler = TStaticSamplerState<SF_Point, AM_Wrap, AM_Wrap, AM_Wrap>::GetRHI();
 	Out.DummyRectLightTextureForCapsuleCompilerWarning = DepthDummy;
@@ -2110,8 +2111,8 @@ static void RenderLight(
 			PermutationVector.Set< FDeferredLightPS::FIESProfileDim >(false);
 			PermutationVector.Set< FDeferredLightPS::FAnistropicMaterials >(bSupportAnisotropyPermutation);
 			// Only directional lights are rendered in this path, so we only need to check if it is use to light the atmosphere
-			PermutationVector.Set< FDeferredLightPS::FAtmosphereTransmittance >(PassParameters->PS.CloudShadow.bAtmospherePerPixelTransmittance > 0);
-			PermutationVector.Set< FDeferredLightPS::FCloudTransmittance >(PassParameters->PS.CloudShadow.bCloudPerPixelTransmittance > 0);
+			PermutationVector.Set< FDeferredLightPS::FAtmosphereTransmittance >(IsLightAtmospherePerPixelTransmittanceEnabled(Scene, View, LightSceneInfo));
+			PermutationVector.Set< FDeferredLightPS::FCloudTransmittance >(PassParameters->PS.CloudShadowEnabled > 0);
 		}
 		PermutationVector = FDeferredLightPS::RemapPermutation(PermutationVector);
 
@@ -2231,8 +2232,8 @@ void FDeferredShadingSceneRenderer::RenderLightForHair(
 		PermutationVector.Set< FDeferredLightPS::FSourceShapeDim >(ELightSourceShape::Directional);
 		PermutationVector.Set< FDeferredLightPS::FSourceTextureDim >(false);
 		PermutationVector.Set< FDeferredLightPS::FIESProfileDim >(false);
-		PermutationVector.Set< FDeferredLightPS::FAtmosphereTransmittance >(PassParameters->PS.CloudShadow.bAtmospherePerPixelTransmittance > 0.f);
-		PermutationVector.Set< FDeferredLightPS::FCloudTransmittance >(PassParameters->PS.CloudShadow.bCloudPerPixelTransmittance > 0.f);
+		PermutationVector.Set< FDeferredLightPS::FAtmosphereTransmittance >(IsLightAtmospherePerPixelTransmittanceEnabled(Scene, View, LightSceneInfo));
+		PermutationVector.Set< FDeferredLightPS::FCloudTransmittance >(PassParameters->PS.CloudShadowEnabled > 0.f);
 	}
 	else
 	{
@@ -2356,7 +2357,8 @@ static FSimpleLightsStandardDeferredParameters GetRenderLightSimpleParameters(
 	Out.PS.LightingChannelsTexture = WhiteDummy;
 	Out.PS.LightingChannelsSampler = TStaticSamplerState<SF_Point, AM_Clamp, AM_Clamp, AM_Clamp>::GetRHI();
 	Out.PS.CloudShadowAO = GetCloudShadowAOParameters(GraphBuilder, View, nullptr);
-	Out.PS.CloudShadow = GetCloudShadowParameters(GraphBuilder);
+	Out.PS.CloudShadowEnabled = 0;
+	SetupLightCloudTransmittanceParameters(GraphBuilder, nullptr, View, nullptr, Out.PS.CloudShadow);
 	Out.PS.LightAttenuationTexture = WhiteDummy;
 	Out.PS.LightAttenuationTextureSampler = TStaticSamplerState<SF_Point, AM_Wrap, AM_Wrap, AM_Wrap>::GetRHI();
 	Out.PS.DummyRectLightTextureForCapsuleCompilerWarning = DepthDummy;
diff --git a/Engine/Source/Runtime/Renderer/Private/SkyAtmosphereRendering.cpp b/Engine/Source/Runtime/Renderer/Private/SkyAtmosphereRendering.cpp
--- a/Engine/Source/Runtime/Renderer/Private/SkyAtmosphereRendering.cpp
+++ b/Engine/Source/Runtime/Renderer/Private/SkyAtmosphereRendering.cpp
@@ -434,6 +434,23 @@ void PrepareSunLightProxy(const FSkyAtmosphereRenderSceneInfo& SkyAtmosphere, ui
 	AtmosphereLight.Proxy->SetAtmosphereRelatedProperties(TransmittanceTowardSun, SunDiskOuterSpaceLuminance);
 }
 
+bool IsLightAtmospherePerPixelTransmittanceEnabled(const FScene* Scene, const FViewInfo& View, const FLightSceneInfo* const LightSceneInfo)
+{
+	if (Scene 
+		&& LightSceneInfo 
+		&& LightSceneInfo->Proxy->GetLightType() == LightType_Directional 
+		&& ShouldRenderSkyAtmosphere(LightSceneInfo->Scene, View.Family->EngineShowFlags))
+	{
+		FLightSceneProxy* AtmosphereLight0Proxy = Scene->AtmosphereLights[0] ? Scene->AtmosphereLights[0]->Proxy : nullptr;
+		FLightSceneProxy* AtmosphereLight1Proxy = Scene->AtmosphereLights[1] ? Scene->AtmosphereLights[1]->Proxy : nullptr;
+
+		// The light must be one of the atmospheric directional lights and has per pixel atmosphere transmittance enabled.
+		return (AtmosphereLight0Proxy == LightSceneInfo->Proxy && AtmosphereLight0Proxy && AtmosphereLight0Proxy->GetUsePerPixelAtmosphereTransmittance())
+			|| (AtmosphereLight1Proxy == LightSceneInfo->Proxy && AtmosphereLight1Proxy && AtmosphereLight1Proxy->GetUsePerPixelAtmosphereTransmittance());
+	}
+	return false;
+}
+
 
 
 /*=============================================================================
diff --git a/Engine/Source/Runtime/Renderer/Private/SkyAtmosphereRendering.h b/Engine/Source/Runtime/Renderer/Private/SkyAtmosphereRendering.h
--- a/Engine/Source/Runtime/Renderer/Private/SkyAtmosphereRendering.h
+++ b/Engine/Source/Runtime/Renderer/Private/SkyAtmosphereRendering.h
@@ -175,6 +175,7 @@ extern void SetupSkyAtmosphereViewSharedUniformShaderParameters(const class FVie
 // Prepare the sun light data as a function of the atmosphere state. 
 void PrepareSunLightProxy(const FSkyAtmosphereRenderSceneInfo& SkyAtmosphere, uint32 AtmosphereLightIndex, FLightSceneInfo& AtmosphereLight);
 
+bool IsLightAtmospherePerPixelTransmittanceEnabled(const FScene* Scene, const FViewInfo& View, const FLightSceneInfo* const LightSceneInfo);
 
 float GetValidAerialPerspectiveStartDepthInCm(const FViewInfo& View, const FSkyAtmosphereSceneProxy& SkyAtmosphereProxy);
 
diff --git a/Engine/Source/Runtime/Renderer/Private/TranslucentLighting.cpp b/Engine/Source/Runtime/Renderer/Private/TranslucentLighting.cpp
--- a/Engine/Source/Runtime/Renderer/Private/TranslucentLighting.cpp
+++ b/Engine/Source/Runtime/Renderer/Private/TranslucentLighting.cpp
@@ -661,18 +661,14 @@ class FTranslucentLightingInjectPS : public FMaterialShader
 		SHADER_PARAMETER_RDG_UNIFORM_BUFFER(FDeferredLightUniformStruct, DeferredLight)
 		SHADER_PARAMETER_STRUCT_INCLUDE(FVolumeShadowingShaderParameters, VolumeShadowingParameters)
 		SHADER_PARAMETER_STRUCT_INCLUDE(FVirtualShadowMapSamplingParameters, VirtualShadowMapSamplingParameters)
-		SHADER_PARAMETER_RDG_TEXTURE(Texture2D<float3>, VolumetricCloudShadowmapTexture)
-		SHADER_PARAMETER_SAMPLER(SamplerState, VolumetricCloudShadowmapTextureSampler)
+		SHADER_PARAMETER_STRUCT_INCLUDE(FLightCloudTransmittanceParameters, LightCloudTransmittanceParameters)
 		SHADER_PARAMETER(FMatrix44f, LightFunctionWorldToLight)
-		SHADER_PARAMETER(FMatrix44f, VolumetricCloudWorldToLightClipShadowMatrix)
 		SHADER_PARAMETER(FVector4f, LightFunctionParameters)
 		SHADER_PARAMETER(float, SpotlightMask)
-		SHADER_PARAMETER(float, VolumetricCloudShadowmapFarDepthKm)
-		SHADER_PARAMETER(float, VolumetricCloudShadowmapStrength)
-		SHADER_PARAMETER(uint32, VolumetricCloudShadowEnabled)
-		SHADER_PARAMETER(uint32, AtmospherePerPixelTransmittanceEnabled)
 		SHADER_PARAMETER(uint32, VolumeCascadeIndex)
 		SHADER_PARAMETER(int32, VirtualShadowMapId)
+		SHADER_PARAMETER(uint32, AtmospherePerPixelTransmittanceEnabled)
+		SHADER_PARAMETER(uint32, VolumetricCloudShadowEnabled)
 	END_SHADER_PARAMETER_STRUCT()
 
 	class FRadialAttenuation	: SHADER_PERMUTATION_BOOL("RADIAL_ATTENUATION");
@@ -730,8 +726,6 @@ public:
 		const FMaterial& Material = MaterialProxy->GetMaterialWithFallback(View.GetFeatureLevel(), MaterialProxy);
 		FMaterialShader::SetParameters(RHICmdList, ShaderRHI, MaterialProxy, Material, View);
 	}
-
-private:
 };
 
 IMPLEMENT_MATERIAL_SHADER_TYPE(,FTranslucentLightingInjectPS, TEXT("/Engine/Private/TranslucentLightInjectionShaders.usf"), TEXT("InjectMainPS"), SF_Pixel);
@@ -1133,48 +1127,10 @@ static void InjectTranslucentLightArray(
 					PassParameters->PS.LightFunctionWorldToLight = WorldToLight;
 				}
 
-				{
-					// TODO this should be put in a common structure / set function.
-					FLightSceneProxy* AtmosphereLight0Proxy = LightSceneInfo->Scene->AtmosphereLights[0] ? LightSceneInfo->Scene->AtmosphereLights[0]->Proxy : nullptr;
-					FLightSceneProxy* AtmosphereLight1Proxy = LightSceneInfo->Scene->AtmosphereLights[1] ? LightSceneInfo->Scene->AtmosphereLights[1]->Proxy : nullptr;
-					{
-						FVolumetricCloudRenderSceneInfo* CloudInfo = LightSceneInfo->Scene->GetVolumetricCloudSceneInfo();
-
-						const bool bLight0CloudPerPixelTransmittance = CloudInfo && View.VolumetricCloudShadowRenderTarget[0] != nullptr && AtmosphereLight0Proxy && AtmosphereLight0Proxy == LightSceneInfo->Proxy;
-						const bool bLight1CloudPerPixelTransmittance = CloudInfo && View.VolumetricCloudShadowRenderTarget[1] != nullptr && AtmosphereLight1Proxy && AtmosphereLight1Proxy == LightSceneInfo->Proxy;
+				const bool bCloudShadowEnabled = SetupLightCloudTransmittanceParameters(GraphBuilder, Scene, View, LightSceneInfo, PassParameters->PS.LightCloudTransmittanceParameters);
+				PassParameters->PS.VolumetricCloudShadowEnabled = bCloudShadowEnabled ? 1 : 0;
 
-						if (bLight0CloudPerPixelTransmittance || bLight1CloudPerPixelTransmittance)
-						{
-							PassParameters->PS.VolumetricCloudShadowmapTexture = View.VolumetricCloudShadowRenderTarget[LightIndex];
-							PassParameters->PS.VolumetricCloudWorldToLightClipShadowMatrix = CloudInfo->GetVolumetricCloudCommonShaderParameters().CloudShadowmapWorldToLightClipMatrix[LightIndex];
-							PassParameters->PS.VolumetricCloudShadowmapFarDepthKm = CloudInfo->GetVolumetricCloudCommonShaderParameters().CloudShadowmapFarDepthKm[LightIndex].X;
-							if (bLight0CloudPerPixelTransmittance)
-							{
-								PassParameters->PS.VolumetricCloudShadowmapStrength = AtmosphereLight0Proxy->GetCloudShadowOnSurfaceStrength();
-							}
-							else if (bLight1CloudPerPixelTransmittance)
-							{
-								PassParameters->PS.VolumetricCloudShadowmapStrength = AtmosphereLight1Proxy->GetCloudShadowOnSurfaceStrength();
-							}
-							PassParameters->PS.VolumetricCloudShadowEnabled = 1;
-						}
-						else
-						{
-							PassParameters->PS.VolumetricCloudShadowmapTexture = GSystemTextures.GetBlackDummy(GraphBuilder);
-							PassParameters->PS.VolumetricCloudWorldToLightClipShadowMatrix = FMatrix44f::Identity;
-							PassParameters->PS.VolumetricCloudShadowmapFarDepthKm = 1.0f;
-							PassParameters->PS.VolumetricCloudShadowmapStrength = 0.0f;
-							PassParameters->PS.VolumetricCloudShadowEnabled = 0;
-						}
-
-						PassParameters->PS.VolumetricCloudShadowmapTextureSampler = TStaticSamplerState<SF_Bilinear, AM_Clamp, AM_Clamp, AM_Clamp>::GetRHI();
-					}
-
-					const bool bLightAtmospherePerPixelTransmittance = ShouldRenderSkyAtmosphere(LightSceneInfo->Scene, View.Family->EngineShowFlags) &&
-						((AtmosphereLight0Proxy == LightSceneInfo->Proxy && AtmosphereLight0Proxy && AtmosphereLight0Proxy->GetUsePerPixelAtmosphereTransmittance())
-							|| (AtmosphereLight1Proxy == LightSceneInfo->Proxy && AtmosphereLight1Proxy && AtmosphereLight1Proxy->GetUsePerPixelAtmosphereTransmittance()));
-					PassParameters->PS.AtmospherePerPixelTransmittanceEnabled = bLightAtmospherePerPixelTransmittance;
-				}
+				PassParameters->PS.AtmospherePerPixelTransmittanceEnabled = IsLightAtmospherePerPixelTransmittanceEnabled(Scene, View, LightSceneInfo);
 
 				GraphBuilder.AddPass(
 					RDG_EVENT_NAME("InjectTranslucentLightArray"),
diff --git a/Engine/Source/Runtime/Renderer/Private/VolumetricCloudRendering.cpp b/Engine/Source/Runtime/Renderer/Private/VolumetricCloudRendering.cpp
--- a/Engine/Source/Runtime/Renderer/Private/VolumetricCloudRendering.cpp
+++ b/Engine/Source/Runtime/Renderer/Private/VolumetricCloudRendering.cpp
@@ -432,64 +432,6 @@ void GetCloudShadowAOData(const FVolumetricCloudRenderSceneInfo* CloudInfo, cons
 	OutData.bShouldSampleCloudSkyAO  = OutData.bShouldSampleCloudSkyAO  && bCloudComponentValid;
 }
 
-FVolumetricCloudShadowParameters GetCloudShadowParameters(FRDGBuilder& GraphBuilder)
-{
-	FVolumetricCloudShadowParameters Out;
-
-	Out.bAtmospherePerPixelTransmittance = 0.f;
-	Out.CloudShadowmapTexture = GSystemTextures.GetDepthDummy(GraphBuilder);
-	Out.CloudShadowmapFarDepthKm = 0.f;
-	Out.CloudShadowmapWorldToLightClipMatrix = FMatrix44f::Identity;
-	Out.CloudShadowmapStrength = 0.f;
-	Out.CloudShadowmapSampler = TStaticSamplerState<SF_Bilinear, AM_Clamp, AM_Clamp, AM_Clamp>::GetRHI();
-	Out.bCloudPerPixelTransmittance = 0.f;
-	return Out;
-}
-
-FVolumetricCloudShadowParameters GetCloudShadowParameters(FRDGBuilder& GraphBuilder, const FScene* Scene, const FViewInfo& View, const FLightSceneProxy* LightProxy, const FVolumetricCloudRenderSceneInfo* CloudInfo)
-{
-	FVolumetricCloudShadowParameters Out;
-
-	Out.bAtmospherePerPixelTransmittance =
-		LightProxy->GetLightType() == LightType_Directional &&
-		LightProxy->IsUsedAsAtmosphereSunLight() &&
-		LightProxy->GetUsePerPixelAtmosphereTransmittance() &&
-		ShouldRenderSkyAtmosphere(Scene, View.Family->EngineShowFlags) ? 1.f : 0.f;
-
-	const FLightSceneProxy* AtmosphereLight0Proxy = Scene->AtmosphereLights[0] ? Scene->AtmosphereLights[0]->Proxy : nullptr;
-	const FLightSceneProxy* AtmosphereLight1Proxy = Scene->AtmosphereLights[1] ? Scene->AtmosphereLights[1]->Proxy : nullptr;
-	const bool VolumetricCloudShadowMap0Valid = View.VolumetricCloudShadowRenderTarget[0] != nullptr;
-	const bool VolumetricCloudShadowMap1Valid = View.VolumetricCloudShadowRenderTarget[1] != nullptr;
-	const bool bLight0CloudPerPixelTransmittance = CloudInfo && VolumetricCloudShadowMap0Valid && AtmosphereLight0Proxy == LightProxy && AtmosphereLight0Proxy && AtmosphereLight0Proxy->GetCloudShadowOnSurfaceStrength() > 0.0f;
-	const bool bLight1CloudPerPixelTransmittance = CloudInfo && VolumetricCloudShadowMap1Valid && AtmosphereLight1Proxy == LightProxy && AtmosphereLight1Proxy && AtmosphereLight1Proxy->GetCloudShadowOnSurfaceStrength() > 0.0f;
-	if (bLight0CloudPerPixelTransmittance)
-	{
-		Out.CloudShadowmapTexture = View.VolumetricCloudShadowRenderTarget[0];
-		Out.CloudShadowmapFarDepthKm = CloudInfo->GetVolumetricCloudCommonShaderParameters().CloudShadowmapFarDepthKm[0].X;
-		Out.CloudShadowmapWorldToLightClipMatrix = CloudInfo->GetVolumetricCloudCommonShaderParameters().CloudShadowmapWorldToLightClipMatrix[0];
-		Out.CloudShadowmapStrength = AtmosphereLight0Proxy->GetCloudShadowOnSurfaceStrength();
-	}
-	else if (bLight1CloudPerPixelTransmittance)
-	{
-		Out.CloudShadowmapTexture = View.VolumetricCloudShadowRenderTarget[1];
-		Out.CloudShadowmapFarDepthKm = CloudInfo->GetVolumetricCloudCommonShaderParameters().CloudShadowmapFarDepthKm[1].X;
-		Out.CloudShadowmapWorldToLightClipMatrix = CloudInfo->GetVolumetricCloudCommonShaderParameters().CloudShadowmapWorldToLightClipMatrix[1];
-		Out.CloudShadowmapStrength = AtmosphereLight1Proxy->GetCloudShadowOnSurfaceStrength();
-	}
-	else
-	{
-		Out.CloudShadowmapTexture = GSystemTextures.GetDepthDummy(GraphBuilder);	
-		Out.CloudShadowmapFarDepthKm = 0.f;
-		Out.CloudShadowmapWorldToLightClipMatrix = FMatrix44f::Identity;
-		Out.CloudShadowmapStrength = 0.f;
-	}
-	Out.CloudShadowmapSampler = TStaticSamplerState<SF_Bilinear, AM_Clamp, AM_Clamp, AM_Clamp>::GetRHI();
-
-	Out.bCloudPerPixelTransmittance = bLight0CloudPerPixelTransmittance || bLight1CloudPerPixelTransmittance ? 1.f : 0.f;
-
-	return Out;
-}
-
 /*=============================================================================
 	FVolumetricCloudRenderSceneInfo implementation.
 =============================================================================*/
@@ -2615,9 +2557,9 @@ bool SetupLightCloudTransmittanceParameters(FRDGBuilder& GraphBuilder, const FSc
 	const FVolumetricCloudRenderSceneInfo* CloudInfo = Scene->GetVolumetricCloudSceneInfo();
 	const bool VolumetricCloudShadowMap0Valid = View.VolumetricCloudShadowRenderTarget[0] != nullptr;
 	const bool VolumetricCloudShadowMap1Valid = View.VolumetricCloudShadowRenderTarget[1] != nullptr;
-	const bool bLight0CloudPerPixelTransmittance = CloudInfo && LightSceneInfo && VolumetricCloudShadowMap0Valid && AtmosphereLight0Proxy == LightSceneInfo->Proxy && AtmosphereLight0Proxy && AtmosphereLight0Proxy->GetCloudShadowOnSurfaceStrength() > 0.0f;
-	const bool bLight1CloudPerPixelTransmittance = CloudInfo && LightSceneInfo && VolumetricCloudShadowMap1Valid && AtmosphereLight1Proxy == LightSceneInfo->Proxy && AtmosphereLight1Proxy && AtmosphereLight1Proxy->GetCloudShadowOnSurfaceStrength() > 0.0f;
-	if (bLight0CloudPerPixelTransmittance)
+	const bool bLight0CloudShadowEnabled = CloudInfo && LightSceneInfo && VolumetricCloudShadowMap0Valid && AtmosphereLight0Proxy == LightSceneInfo->Proxy && AtmosphereLight0Proxy && AtmosphereLight0Proxy->GetCloudShadowOnSurfaceStrength() > 0.0f;
+	const bool bLight1CloudShadowEnabled = CloudInfo && LightSceneInfo && VolumetricCloudShadowMap1Valid && AtmosphereLight1Proxy == LightSceneInfo->Proxy && AtmosphereLight1Proxy && AtmosphereLight1Proxy->GetCloudShadowOnSurfaceStrength() > 0.0f;
+	if (bLight0CloudShadowEnabled)
 	{
 		OutParameters.CloudShadowmapTexture = View.VolumetricCloudShadowRenderTarget[0];
 		OutParameters.CloudShadowmapSampler = TStaticSamplerState<SF_Bilinear, AM_Clamp, AM_Clamp, AM_Clamp>::GetRHI();
@@ -2625,7 +2567,7 @@ bool SetupLightCloudTransmittanceParameters(FRDGBuilder& GraphBuilder, const FSc
 		OutParameters.CloudShadowmapWorldToLightClipMatrix = CloudInfo->GetVolumetricCloudCommonShaderParameters().CloudShadowmapWorldToLightClipMatrix[0];
 		OutParameters.CloudShadowmapStrength = AtmosphereLight0Proxy->GetCloudShadowOnSurfaceStrength();
 	}
-	else if (bLight1CloudPerPixelTransmittance)
+	else if (bLight1CloudShadowEnabled)
 	{
 		OutParameters.CloudShadowmapTexture = View.VolumetricCloudShadowRenderTarget[1];
 		OutParameters.CloudShadowmapSampler = TStaticSamplerState<SF_Bilinear, AM_Clamp, AM_Clamp, AM_Clamp>::GetRHI();
@@ -2638,5 +2580,5 @@ bool SetupLightCloudTransmittanceParameters(FRDGBuilder& GraphBuilder, const FSc
 		DefaultLightCloudTransmittanceParameters(GraphBuilder, OutParameters);
 	}
 
-	return bLight0CloudPerPixelTransmittance || bLight1CloudPerPixelTransmittance;
+	return bLight0CloudShadowEnabled || bLight1CloudShadowEnabled;
 }
\ No newline at end of file
diff --git a/Engine/Source/Runtime/Renderer/Private/VolumetricCloudRendering.h b/Engine/Source/Runtime/Renderer/Private/VolumetricCloudRendering.h
--- a/Engine/Source/Runtime/Renderer/Private/VolumetricCloudRendering.h
+++ b/Engine/Source/Runtime/Renderer/Private/VolumetricCloudRendering.h
@@ -157,19 +157,6 @@ END_SHADER_PARAMETER_STRUCT()
 
 FVolumetricCloudShadowAOParameters GetCloudShadowAOParameters(FRDGBuilder& GraphBuilder, const FViewInfo& View, const FVolumetricCloudRenderSceneInfo* CloudInfo);
 
-BEGIN_SHADER_PARAMETER_STRUCT(FVolumetricCloudShadowParameters, )
-	SHADER_PARAMETER(float, bAtmospherePerPixelTransmittance)
-	SHADER_PARAMETER(float, bCloudPerPixelTransmittance)
-	SHADER_PARAMETER(float, CloudShadowmapFarDepthKm)
-	SHADER_PARAMETER(FMatrix44f, CloudShadowmapWorldToLightClipMatrix)
-	SHADER_PARAMETER(float, CloudShadowmapStrength)
-	SHADER_PARAMETER_RDG_TEXTURE(Texture2D, CloudShadowmapTexture)
-	SHADER_PARAMETER_SAMPLER(SamplerState, CloudShadowmapSampler)
-END_SHADER_PARAMETER_STRUCT()
-
-FVolumetricCloudShadowParameters GetCloudShadowParameters(FRDGBuilder& GraphBuilder);
-FVolumetricCloudShadowParameters GetCloudShadowParameters(FRDGBuilder& GraphBuilder, const FScene* Scene, const FViewInfo& View, const FLightSceneProxy* LightProxy, const FVolumetricCloudRenderSceneInfo* CloudInfo);
-
 struct FCloudShadowAOData
 {
 	bool bShouldSampleCloudShadow;
