Fixes ensures in TSR when high res screenshot or burning out at 16K

diff --git a/Engine/Source/Runtime/Renderer/Private/PostProcess/TemporalSuperResolution.cpp b/Engine/Source/Runtime/Renderer/Private/PostProcess/TemporalSuperResolution.cpp
--- a/Engine/Source/Runtime/Renderer/Private/PostProcess/TemporalSuperResolution.cpp
+++ b/Engine/Source/Runtime/Renderer/Private/PostProcess/TemporalSuperResolution.cpp
@@ -618,11 +618,17 @@ ITemporalUpscaler::FOutputs AddTemporalSuperResolutionPasses(
 	FIntPoint HistoryExtent;
 	FIntPoint HistorySize;
 	{
-		float UpscaleFactor = FMath::Clamp(CVarTSRHistorySP.GetValueOnRenderThread() / 100.0f, 1.0f, 2.0f);
+		float MaxHistoryUpscaleFactor = FMath::Max(float(GMaxTextureDimensions) / float(FMath::Max(OutputRect.Width(), OutputRect.Height())), 1.0f);
 
+		float HistoryUpscaleFactor = FMath::Clamp(CVarTSRHistorySP.GetValueOnRenderThread() / 100.0f, 1.0f, 2.0f);
+		if (HistoryUpscaleFactor > MaxHistoryUpscaleFactor)
+		{
+			HistoryUpscaleFactor = 1.0f;
+		}
+		
 		HistorySize = FIntPoint(
-			FMath::CeilToInt(OutputRect.Width() * UpscaleFactor),
-			FMath::CeilToInt(OutputRect.Height() * UpscaleFactor));
+			FMath::CeilToInt(OutputRect.Width() * HistoryUpscaleFactor),
+			FMath::CeilToInt(OutputRect.Height() * HistoryUpscaleFactor));
 
 		FIntPoint QuantizedHistoryViewSize;
 		QuantizeSceneBufferSize(HistorySize, QuantizedHistoryViewSize);
