Fixes a crash in FFT bloom in highres screenshot

diff --git a/Engine/Source/Runtime/Renderer/Private/PostProcess/PostProcessFFTBloom.cpp b/Engine/Source/Runtime/Renderer/Private/PostProcess/PostProcessFFTBloom.cpp
--- a/Engine/Source/Runtime/Renderer/Private/PostProcess/PostProcessFFTBloom.cpp
+++ b/Engine/Source/Runtime/Renderer/Private/PostProcess/PostProcessFFTBloom.cpp
@@ -197,9 +197,11 @@ IMPLEMENT_GLOBAL_SHADER(FBloomFinalizeApplyConstantsCS,     "/Engine/Private/Blo
 
 } //! namespace
 
-float GetFFTBloomResolutionFraction()
+float GetFFTBloomResolutionFraction(const FIntPoint& ViewSize)
 {
-	return FMath::Clamp(CVarBloomScreenPercentage.GetValueOnRenderThread() / 100.0f, 0.1f, 1.0f);
+	float MaxResolutionFraction = FMath::Min(float(GMaxTextureDimensions) * 0.5f / float(FMath::Max(ViewSize.X, ViewSize.Y)), 1.0f);
+
+	return FMath::Clamp(CVarBloomScreenPercentage.GetValueOnRenderThread() / 100.0f, 0.1f, MaxResolutionFraction);
 }
 
 bool IsFFTBloomEnabled(const FViewInfo& View)
@@ -640,7 +642,7 @@ FFFTBloomOutput AddFFTBloomPass(FRDGBuilder& GraphBuilder, const FViewInfo& View
 {
 	check(InputSceneColor.IsValid());
 
-	const float ResolutionFraction = GetFFTBloomResolutionFraction();
+	const float ResolutionFraction = GetFFTBloomResolutionFraction(InputSceneColor.ViewRect.Size());
 	const float DownscaleResolutionFraction = ResolutionFraction / InputResolutionFraction;
 	check(DownscaleResolutionFraction <= 1.0f);
 
diff --git a/Engine/Source/Runtime/Renderer/Private/PostProcess/PostProcessFFTBloom.h b/Engine/Source/Runtime/Renderer/Private/PostProcess/PostProcessFFTBloom.h
--- a/Engine/Source/Runtime/Renderer/Private/PostProcess/PostProcessFFTBloom.h
+++ b/Engine/Source/Runtime/Renderer/Private/PostProcess/PostProcessFFTBloom.h
@@ -11,7 +11,7 @@ struct FBloomOutputs;
 // Returns whether FFT bloom is enabled for the view.
 bool IsFFTBloomEnabled(const FViewInfo& View);
 
-float GetFFTBloomResolutionFraction();
+float GetFFTBloomResolutionFraction(const FIntPoint& ViewSize);
 
 struct FFFTBloomOutput
 {
@@ -19,6 +19,4 @@ struct FFFTBloomOutput
 	FRDGBufferRef SceneColorApplyParameters = nullptr;
 };
 
-float GetFFTBloomResolutionFraction();
-
 FFFTBloomOutput AddFFTBloomPass(FRDGBuilder& GraphBuilder, const FViewInfo& View, const FScreenPassTexture& InputSceneColor, float InputResolutionFraction);
diff --git a/Engine/Source/Runtime/Renderer/Private/PostProcess/PostProcessing.cpp b/Engine/Source/Runtime/Renderer/Private/PostProcess/PostProcessing.cpp
--- a/Engine/Source/Runtime/Renderer/Private/PostProcess/PostProcessing.cpp
+++ b/Engine/Source/Runtime/Renderer/Private/PostProcess/PostProcessing.cpp
@@ -598,6 +598,8 @@ void AddPostProcessingPasses(FRDGBuilder& GraphBuilder, const FViewInfo& View, c
 		// Whether separate translucency is composed in TSR.
 		bool bComposeSeparateTranslucencyInTSR = TAAConfig == EMainTAAPassConfig::TSR && ComposeSeparateTranslucencyInTSR(View);
 
+		const FIntPoint PostTAAViewSize = View.PrimaryScreenPercentageMethod != EPrimaryScreenPercentageMethod::TemporalUpscale ? View.GetSecondaryViewRectSize() : View.ViewRect.Size();
+
 		const FPostProcessMaterialChain PostProcessMaterialAfterTonemappingChain = GetPostProcessMaterialChain(View, BL_AfterTonemapping);
 
 		PassSequence.SetEnabled(EPass::MotionBlur, bVisualizeMotionBlur || bMotionBlurEnabled);
@@ -632,7 +634,7 @@ void AddPostProcessingPasses(FRDGBuilder& GraphBuilder, const FViewInfo& View, c
 		const bool bProcessQuarterResolution = IsPostProcessingQuarterResolutionDownsampleEnabled();
 		const bool bMotionBlurNeedsHalfResInput = PassSequence.IsEnabled(EPass::MotionBlur) && DoesMotionBlurNeedsHalfResInput() && !bVisualizeMotionBlur;
 
-		const float FFTBloomResolutionFraction = GetFFTBloomResolutionFraction();
+		const float FFTBloomResolutionFraction = GetFFTBloomResolutionFraction(PostTAAViewSize);
 
 		const bool bProduceSceneColorChain = (
 			bBasicEyeAdaptationEnabled ||
@@ -766,6 +768,8 @@ void AddPostProcessingPasses(FRDGBuilder& GraphBuilder, const FViewInfo& View, c
 			}
 		}
 
+		ensure(SceneColor.ViewRect.Size() == PostTAAViewSize);
+
 		// Post Process Material Chain - SSR Input
 		if (View.ViewState && !View.bStatePrevViewInfoIsReadOnly)
 		{
