Fix a bug of casting CSM on mobile split screen.

diff --git a/Engine/Source/Runtime/Renderer/Private/MobileBasePassRendering.cpp b/Engine/Source/Runtime/Renderer/Private/MobileBasePassRendering.cpp
--- a/Engine/Source/Runtime/Renderer/Private/MobileBasePassRendering.cpp
+++ b/Engine/Source/Runtime/Renderer/Private/MobileBasePassRendering.cpp
@@ -243,6 +243,7 @@ TRDGUniformBufferRef<FMobileBasePassUniformParameters> CreateMobileBasePassUnifo
 
 void SetupMobileDirectionalLightUniformParameters(
 	const FScene& Scene,
+	int32 ViewIndex,
 	const FViewInfo& SceneView,
 	const TArray<FVisibleLightInfo,SceneRenderingAllocator>& VisibleLightInfos,
 	int32 ChannelIdx,
@@ -266,11 +267,11 @@ void SetupMobileDirectionalLightUniformParameters(
 			const TArray<FProjectedShadowInfo*, SceneRenderingAllocator>& DirectionalLightShadowInfos = VisibleLightInfos[Light->Id].AllProjectedShadows;
 			static_assert(MAX_MOBILE_SHADOWCASCADES <= 4, "more than 4 cascades not supported by the shader and uniform buffer");
 
-			const int32 NumShadowsToCopy = FMath::Min(DirectionalLightShadowInfos.Num(), MAX_MOBILE_SHADOWCASCADES);
+			const int32 NumShadowsToCopy = FMath::Min(DirectionalLightShadowInfos.Num(), SceneView.MaxShadowCascades);
 			int32_t OutShadowIndex = 0;
 			for (int32 i = 0; i < NumShadowsToCopy; ++i)
 			{
-				const FProjectedShadowInfo* ShadowInfo = DirectionalLightShadowInfos[i];
+				const FProjectedShadowInfo* ShadowInfo = DirectionalLightShadowInfos[i + ViewIndex * NumShadowsToCopy];
 
 				if (ShadowInfo->ShadowDepthView && !ShadowInfo->bRayTracedDistanceField && ShadowInfo->CacheMode != SDCM_StaticPrimitivesOnly)
 				{
diff --git a/Engine/Source/Runtime/Renderer/Private/MobileBasePassRendering.h b/Engine/Source/Runtime/Renderer/Private/MobileBasePassRendering.h
--- a/Engine/Source/Runtime/Renderer/Private/MobileBasePassRendering.h
+++ b/Engine/Source/Runtime/Renderer/Private/MobileBasePassRendering.h
@@ -68,6 +68,7 @@ extern TRDGUniformBufferRef<FMobileBasePassUniformParameters> CreateMobileBasePa
 
 extern void SetupMobileDirectionalLightUniformParameters(
 	const FScene& Scene,
+	int32 ViewIndex,
 	const FViewInfo& View,
 	const TArray<FVisibleLightInfo,SceneRenderingAllocator>& VisibleLightInfos,
 	int32 ChannelIdx,
diff --git a/Engine/Source/Runtime/Renderer/Private/MobileShadingRenderer.cpp b/Engine/Source/Runtime/Renderer/Private/MobileShadingRenderer.cpp
--- a/Engine/Source/Runtime/Renderer/Private/MobileShadingRenderer.cpp
+++ b/Engine/Source/Runtime/Renderer/Private/MobileShadingRenderer.cpp
@@ -1049,7 +1049,7 @@ void FMobileSceneRenderer::RenderForward(FRDGBuilder& GraphBuilder, FRDGTextureR
 
 		View.BeginRenderView();
 
-		UpdateDirectionalLightUniformBuffers(GraphBuilder, View);
+		UpdateDirectionalLightUniformBuffers(GraphBuilder, ViewIndex, View);
 
 		FMobileBasePassTextures MobileBasePassTextures{};
 		MobileBasePassTextures.ScreenSpaceAO = bRequiresAmbientOcclusionPass ? SceneTextures.ScreenSpaceAO : SystemTextures.White;
@@ -1333,7 +1333,7 @@ void FMobileSceneRenderer::RenderDeferred(FRDGBuilder& GraphBuilder, const FSort
 
 		View.BeginRenderView();
 
-		UpdateDirectionalLightUniformBuffers(GraphBuilder, View);
+		UpdateDirectionalLightUniformBuffers(GraphBuilder, ViewIndex, View);
 
 		FMobileBasePassTextures MobileBasePassTextures{};
 		MobileBasePassTextures.ScreenSpaceAO = bRequiresAmbientOcclusionPass ? SceneTextures.ScreenSpaceAO : SystemTextures.White;
@@ -1569,7 +1569,7 @@ bool FMobileSceneRenderer::RequiresMultiPass(FRHICommandListImmediate& RHICmdLis
 	return true;
 }
 
-void FMobileSceneRenderer::UpdateDirectionalLightUniformBuffers(FRDGBuilder& GraphBuilder, const FViewInfo& View)
+void FMobileSceneRenderer::UpdateDirectionalLightUniformBuffers(FRDGBuilder& GraphBuilder, int32 ViewIndex, const FViewInfo& View)
 {
 	if (CachedView == &View)
 	{
@@ -1577,14 +1577,14 @@ void FMobileSceneRenderer::UpdateDirectionalLightUniformBuffers(FRDGBuilder& Gra
 	}
 	CachedView = &View;
 
-	AddPass(GraphBuilder, RDG_EVENT_NAME("UpdateDirectionalLightUniformBuffers"), [this, &View](FRHICommandListImmediate&)
+	AddPass(GraphBuilder, RDG_EVENT_NAME("UpdateDirectionalLightUniformBuffers"), [this, ViewIndex, &View](FRHICommandListImmediate&)
 	{
 		const bool bDynamicShadows = ViewFamily.EngineShowFlags.DynamicShadows;
 		// Fill in the other entries based on the lights
 		for (int32 ChannelIdx = 0; ChannelIdx < UE_ARRAY_COUNT(Scene->MobileDirectionalLights); ChannelIdx++)
 		{
 			FMobileDirectionalLightShaderParameters Params;
-			SetupMobileDirectionalLightUniformParameters(*Scene, View, VisibleLightInfos, ChannelIdx, bDynamicShadows, Params);
+			SetupMobileDirectionalLightUniformParameters(*Scene, ViewIndex, View, VisibleLightInfos, ChannelIdx, bDynamicShadows, Params);
 			Scene->UniformBuffers.MobileDirectionalLightUniformBuffers[ChannelIdx + 1].UpdateUniformBufferImmediate(Params);
 		}
 	});
diff --git a/Engine/Source/Runtime/Renderer/Private/SceneRendering.cpp b/Engine/Source/Runtime/Renderer/Private/SceneRendering.cpp
--- a/Engine/Source/Runtime/Renderer/Private/SceneRendering.cpp
+++ b/Engine/Source/Runtime/Renderer/Private/SceneRendering.cpp
@@ -922,7 +922,7 @@ void FViewInfo::Init()
 		TranslucencyLightingVolumeSize[CascadeIndex] = FVector(0);
 	}
 
-	const int32 MaxMobileShadowCascadeCount = FMath::Clamp(CVarMaxMobileShadowCascades.GetValueOnAnyThread(), 0, MAX_MOBILE_SHADOWCASCADES);
+	const int32 MaxMobileShadowCascadeCount = FMath::Clamp(CVarMaxMobileShadowCascades.GetValueOnAnyThread(), 0, MAX_MOBILE_SHADOWCASCADES / FMath::Max(Family->Views.Num(), 1));
 	const int32 MaxShadowCascadeCountUpperBound = GetFeatureLevel() >= ERHIFeatureLevel::SM5 ? 10 : MaxMobileShadowCascadeCount;
 
 	MaxShadowCascades = FMath::Clamp<int32>(CVarMaxShadowCascades.GetValueOnAnyThread(), 0, MaxShadowCascadeCountUpperBound);
diff --git a/Engine/Source/Runtime/Renderer/Private/SceneRendering.h b/Engine/Source/Runtime/Renderer/Private/SceneRendering.h
--- a/Engine/Source/Runtime/Renderer/Private/SceneRendering.h
+++ b/Engine/Source/Runtime/Renderer/Private/SceneRendering.h
@@ -2365,7 +2365,7 @@ protected:
 	void SortMobileBasePassAfterShadowInit(FExclusiveDepthStencil::Type BasePassDepthStencilAccess, FViewVisibleCommandsPerView& ViewCommandsPerView);
 	void SetupMobileBasePassAfterShadowInit(FExclusiveDepthStencil::Type BasePassDepthStencilAccess, FViewVisibleCommandsPerView& ViewCommandsPerView, FInstanceCullingManager& InstanceCullingManager);
 
-	void UpdateDirectionalLightUniformBuffers(FRDGBuilder& GraphBuilder, const FViewInfo& View);
+	void UpdateDirectionalLightUniformBuffers(FRDGBuilder& GraphBuilder, int32 ViewIndex, const FViewInfo& View);
 	void UpdateSkyReflectionUniformBuffer();
 	
 	void BuildInstanceCullingDrawParams(FRDGBuilder& GraphBuilder, FViewInfo& View, class FMobileRenderPassParameters* PassParameters);
