Fixes ensure when temporally upscaling

diff --git a/Engine/Source/Runtime/Renderer/Private/PostProcess/PostProcessing.cpp b/Engine/Source/Runtime/Renderer/Private/PostProcess/PostProcessing.cpp
--- a/Engine/Source/Runtime/Renderer/Private/PostProcess/PostProcessing.cpp
+++ b/Engine/Source/Runtime/Renderer/Private/PostProcess/PostProcessing.cpp
@@ -598,7 +598,7 @@ void AddPostProcessingPasses(FRDGBuilder& GraphBuilder, const FViewInfo& View, c
 		// Whether separate translucency is composed in TSR.
 		bool bComposeSeparateTranslucencyInTSR = TAAConfig == EMainTAAPassConfig::TSR && ComposeSeparateTranslucencyInTSR(View);
 
-		const FIntPoint PostTAAViewSize = View.PrimaryScreenPercentageMethod != EPrimaryScreenPercentageMethod::TemporalUpscale ? View.GetSecondaryViewRectSize() : View.ViewRect.Size();
+		const FIntPoint PostTAAViewSize = (View.PrimaryScreenPercentageMethod == EPrimaryScreenPercentageMethod::TemporalUpscale && TAAConfig != EMainTAAPassConfig::Disabled) ? View.GetSecondaryViewRectSize() : View.ViewRect.Size();
 
 		const FPostProcessMaterialChain PostProcessMaterialAfterTonemappingChain = GetPostProcessMaterialChain(View, BL_AfterTonemapping);
 
