Strata leverage UAV slice start/count to remove the post base pass copy step.
Ran Strata and non Strata code path.
QAGame cooked on Swi.
ShoorterGame cooked and ran on PC.

diff --git a/Engine/Shaders/Private/BasePassPixelShader.usf b/Engine/Shaders/Private/BasePassPixelShader.usf
--- a/Engine/Shaders/Private/BasePassPixelShader.usf
+++ b/Engine/Shaders/Private/BasePassPixelShader.usf
@@ -1637,10 +1637,10 @@ void FPixelShaderInOut_MainPS(
 
 	// Generate the strata material data to write out
 	FStrataAddressing StrataAddressing = GetStrataPixelDataByteOffset(PixelPos, uint2(ResolvedView.BufferSizeAndInvSize.xy), OpaqueBasePass.Strata.MaxBytesPerPixel);
-	FRWStrataMaterialContainer RWStrataMaterialContainer = InitialiseRWStrataMaterialContainer(OpaqueBasePass.Strata.MaterialTextureArrayUAV);
+	FRWStrataMaterialContainer RWStrataMaterialContainer = InitialiseRWStrataMaterialContainer(OpaqueBasePass.Strata.MaterialTextureArrayUAVWithoutRTs);
 	PackStrataOut(
 		RWStrataMaterialContainer,
-		OpaqueBasePass.Strata.MaterialTextureArrayUAV,
+		OpaqueBasePass.Strata.MaterialTextureArrayUAVWithoutRTs,
 		Dither,
 		StrataAddressing,
 		StrataPixelHeader, StrataData, MaterialParameters.CameraVector, bStrataSubsurfaceEnable, EmissiveLuminance,
diff --git a/Engine/Shaders/Private/Strata/Strata.ush b/Engine/Shaders/Private/Strata/Strata.ush
--- a/Engine/Shaders/Private/Strata/Strata.ush
+++ b/Engine/Shaders/Private/Strata/Strata.ush
@@ -35,10 +35,9 @@
 // UAV SSS data written out if needed
 // UAV if strata data overrun the MRT output size
 
-#define STARTA_MATERIAL_RENDERTARGET_COUNT 2
 struct FRWStrataMaterialContainerStruct
 {
-	uint MaterialRenderTargets[STARTA_MATERIAL_RENDERTARGET_COUNT];
+	uint MaterialRenderTargets[STRATA_BASE_PASS_MRT_OUTPUT_COUNT];
 //	RWTexture2DArray<uint> ExtraMaterialDataUAV; // If this is used, we cannot run shaders on debug...
 };
 #define FRWStrataMaterialContainer	FRWStrataMaterialContainerStruct
@@ -48,7 +47,7 @@ FRWStrataMaterialContainer InitialiseRWStrataMaterialContainer(RWTexture2DArray<
 {
 	FRWStrataMaterialContainer RWStrataMaterialContainer;
 	UNROLL
-	for (int i = 0; i < STARTA_MATERIAL_RENDERTARGET_COUNT; ++i)
+	for (int i = 0; i < STRATA_BASE_PASS_MRT_OUTPUT_COUNT; ++i)
 	{
 		RWStrataMaterialContainer.MaterialRenderTargets[i] = 0;
 	}
@@ -832,14 +831,17 @@ uint StrataLoadUint1(FStrataMaterialContainer StrataBuffer, inout FStrataAddress
 void StrataStoreUint1(inout FRWStrataMaterialContainer StrataBuffer, RWTexture2DArray<uint> ExtraMaterialDataUAV, inout FStrataAddressing StrataAddressing, uint Data)
 {
 	// We heavily rely on the compiler to optimize out the if code
-	if (StrataAddressing.CurrentIndex < STARTA_MATERIAL_RENDERTARGET_COUNT)
+	if (StrataAddressing.CurrentIndex < STRATA_BASE_PASS_MRT_OUTPUT_COUNT)
 	{
 		StrataBuffer.MaterialRenderTargets[0] = StrataAddressing.CurrentIndex == 0 ? Data : StrataBuffer.MaterialRenderTargets[0];
 		StrataBuffer.MaterialRenderTargets[1] = StrataAddressing.CurrentIndex == 1 ? Data : StrataBuffer.MaterialRenderTargets[1];
+#if STRATA_BASE_PASS_MRT_OUTPUT_COUNT != 2
+#error Strata STRATA_BASE_PASS_MRT_OUTPUT_COUNT has been update but not StrataStore function
+#endif
 	}
 	else
 	{
-		ExtraMaterialDataUAV[uint3(StrataAddressing.PixelCoords, StrataAddressing.CurrentIndex)] = Data;
+		ExtraMaterialDataUAV[uint3(StrataAddressing.PixelCoords, StrataAddressing.CurrentIndex - STRATA_BASE_PASS_MRT_OUTPUT_COUNT)] = Data;
 	}
 	StrataAddressing.CurrentIndex++;
 }
diff --git a/Engine/Shaders/Shared/StrataDefinitions.h b/Engine/Shaders/Shared/StrataDefinitions.h
--- a/Engine/Shaders/Shared/StrataDefinitions.h
+++ b/Engine/Shaders/Shared/StrataDefinitions.h
@@ -42,4 +42,4 @@
 // The thin lighting model is used when the thickness becomes strictly lower than 0.1 millimeter
 #define STRATA_LAYER_ISTHIN_THICKNESS_THRESHOLD_CM 0.01f
 
-
+#define STRATA_BASE_PASS_MRT_OUTPUT_COUNT	2
diff --git a/Engine/Source/Runtime/Renderer/Private/Strata/Strata.cpp b/Engine/Source/Runtime/Renderer/Private/Strata/Strata.cpp
--- a/Engine/Source/Runtime/Renderer/Private/Strata/Strata.cpp
+++ b/Engine/Source/Runtime/Renderer/Private/Strata/Strata.cpp
@@ -72,6 +72,7 @@ void FStrataSceneData::Reset()
 	SSSTextureUAV = nullptr;
 
 	MaterialTextureArray = nullptr;
+	MaterialTextureArrayUAVWithoutRTs = nullptr;
 	MaterialTextureArrayUAV = nullptr;
 	MaterialTextureArraySRV = nullptr;
 
@@ -192,20 +193,23 @@ void InitialiseStrataFrameSceneData(FSceneRenderer& SceneRenderer, FRDGBuilder&
 	}
 	else
 	{
-		StrataSceneData.MaxBytesPerPixel = 4u;
+		StrataSceneData.MaxBytesPerPixel = 4u * STRATA_BASE_PASS_MRT_OUTPUT_COUNT;
 	}
 
 	// Create the material data container
 	FIntPoint SceneTextureExtent = IsStrataEnabled() ? GetSceneTextureExtent() : FIntPoint(2, 2);
 
-	const FRDGTextureDesc MaterialTextureDescMRTs = FRHITextureCreateInfo::Create2DArray(SceneTextureExtent, PF_R32_UINT, FClearValueBinding::Transparent, TexCreate_TargetArraySlicesIndependently | TexCreate_RenderTargetable | TexCreate_ShaderResource | TexCreate_UAV, 2, 1, 1);
-	StrataSceneData.MaterialTextureArrayMRTs = GraphBuilder.CreateTexture(MaterialTextureDescMRTs, TEXT("Strata.MaterialRT"));
-
-	const FRDGTextureDesc MaterialTextureDesc = FRHITextureCreateInfo::Create2DArray(SceneTextureExtent, PF_R32_UINT, FClearValueBinding::Transparent, TexCreate_DisableDCC | TexCreate_NoFastClear | TexCreate_RenderTargetable | TexCreate_ShaderResource | TexCreate_UAV, StrataSceneData.MaxBytesPerPixel / 4, 1, 1);
+	const uint32 SliceCount = FMath::DivideAndRoundUp(StrataSceneData.MaxBytesPerPixel, 4u);
+	const FRDGTextureDesc MaterialTextureDesc = FRHITextureCreateInfo::Create2DArray(SceneTextureExtent, PF_R32_UINT, FClearValueBinding::Transparent, 
+		TexCreate_TargetArraySlicesIndependently | TexCreate_DisableDCC | TexCreate_NoFastClear | TexCreate_RenderTargetable | TexCreate_ShaderResource | TexCreate_UAV, SliceCount, 1, 1);
 	StrataSceneData.MaterialTextureArray = GraphBuilder.CreateTexture(MaterialTextureDesc, TEXT("Strata.Material"));
 	StrataSceneData.MaterialTextureArraySRV = GraphBuilder.CreateSRV(FRDGTextureSRVDesc::Create(StrataSceneData.MaterialTextureArray));
 	StrataSceneData.MaterialTextureArrayUAV = GraphBuilder.CreateUAV(FRDGTextureUAVDesc(StrataSceneData.MaterialTextureArray, 0));
 
+	// See AppendStrataMRTs
+	check(STRATA_BASE_PASS_MRT_OUTPUT_COUNT <= SliceCount);
+	StrataSceneData.MaterialTextureArrayUAVWithoutRTs = GraphBuilder.CreateUAV(FRDGTextureUAVDesc(StrataSceneData.MaterialTextureArray, 0, PF_Unknown, STRATA_BASE_PASS_MRT_OUTPUT_COUNT, SliceCount - STRATA_BASE_PASS_MRT_OUTPUT_COUNT));
+
 	// Rough diffuse model
 	StrataSceneData.bRoughDiffuse = CVarStrataRoughDiffuse.GetValueOnRenderThread() > 0 ? 1u : 0u;
 
@@ -220,7 +224,7 @@ void InitialiseStrataFrameSceneData(FSceneRenderer& SceneRenderer, FRDGBuilder&
 	{
 		AddStrataClearMaterialBufferPass(
 			GraphBuilder, 
-			GraphBuilder.CreateUAV(FRDGTextureUAVDesc(StrataSceneData.MaterialTextureArrayMRTs, 0)),
+			GraphBuilder.CreateUAV(FRDGTextureUAVDesc(StrataSceneData.MaterialTextureArray, 0)),
 			StrataSceneData.SSSTextureUAV,
 			StrataSceneData.MaxBytesPerPixel, 
 			MaterialBufferSizeXY);
@@ -241,7 +245,7 @@ void BindStrataBasePassUniformParameters(FRDGBuilder& GraphBuilder, FStrataScene
 	{
 		OutStrataUniformParameters.bRoughDiffuse = StrataSceneData->bRoughDiffuse ? 1u : 0u;
 		OutStrataUniformParameters.MaxBytesPerPixel = StrataSceneData->MaxBytesPerPixel;
-		OutStrataUniformParameters.MaterialTextureArrayUAV = StrataSceneData->MaterialTextureArrayUAV;
+		OutStrataUniformParameters.MaterialTextureArrayUAVWithoutRTs = StrataSceneData->MaterialTextureArrayUAVWithoutRTs;
 		OutStrataUniformParameters.SSSTextureUAV = StrataSceneData->SSSTextureUAV;
 	}
 	else
@@ -255,7 +259,7 @@ void BindStrataBasePassUniformParameters(FRDGBuilder& GraphBuilder, FStrataScene
 		const FRDGSystemTextures& SystemTextures = FRDGSystemTextures::Get(GraphBuilder);
 		OutStrataUniformParameters.bRoughDiffuse = 0u;
 		OutStrataUniformParameters.MaxBytesPerPixel = 0;
-		OutStrataUniformParameters.MaterialTextureArrayUAV = DummyWritableTextureArrayUAV;
+		OutStrataUniformParameters.MaterialTextureArrayUAVWithoutRTs = DummyWritableTextureArrayUAV;
 		OutStrataUniformParameters.SSSTextureUAV = DummyWritableTextureUAV;
 	}
 }
@@ -635,11 +639,14 @@ void AppendStrataMRTs(FSceneRenderer& SceneRenderer, uint32& RenderTargetCount,
 		// - We must never clear the second uint, it will only be written/read if needed.
 		auto AddStrataOutputTarget = [&](int16 StrataMaterialArraySlice, bool bNeverClear = false)
 		{
-			RenderTargets[RenderTargetCount] = FTextureRenderTargetBinding(SceneRenderer.Scene->StrataSceneData.MaterialTextureArrayMRTs, StrataMaterialArraySlice, bNeverClear);
+			RenderTargets[RenderTargetCount] = FTextureRenderTargetBinding(SceneRenderer.Scene->StrataSceneData.MaterialTextureArray, StrataMaterialArraySlice, bNeverClear);
 			RenderTargetCount++;
 		};
-		AddStrataOutputTarget(0);
-		AddStrataOutputTarget(1, true); // Never clear that target
+		for (int i = 0; i < STRATA_BASE_PASS_MRT_OUTPUT_COUNT; ++i)
+		{
+			const bool bNeverClear = i != 0; // Only allow clearing the first slice containing the header
+			AddStrataOutputTarget(i, bNeverClear);
+		}
 
 		// Add another MRT for Strata top layer information. We want to follow the usual clear process which can leverage fast clear.
 		{
@@ -693,7 +700,17 @@ IMPLEMENT_GLOBAL_SHADER(FStrataPostBasePassCS, "/Engine/Private/Strata/StrataMat
 
 void PostBasePass(FRDGBuilder& GraphBuilder, const TArray<FViewInfo>& Views)
 {
-	RDG_EVENT_SCOPE_CONDITIONAL(GraphBuilder, IsStrataEnabled() && Views.Num() > 0, "StrataPostBasePass");
+
+
+
+
+	////// TODO delete PostBasePas shader
+
+
+
+
+
+	/*RDG_EVENT_SCOPE_CONDITIONAL(GraphBuilder, IsStrataEnabled() && Views.Num() > 0, "StrataPostBasePass");
 	if (!IsStrataEnabled())
 	{
 		return;
@@ -716,7 +733,7 @@ void PostBasePass(FRDGBuilder& GraphBuilder, const TArray<FViewInfo>& Views)
 			ComputeShader,
 			PassParameters,
 			FComputeShaderUtils::GetGroupCount(View.ViewRect.Size(), GroupSize));
-	}
+	}*/
 }
 
 ////////////////////////////////////////////////////////////////////////////////////////////////////////////////////////////////////////////////
diff --git a/Engine/Source/Runtime/Renderer/Private/Strata/Strata.h b/Engine/Source/Runtime/Renderer/Private/Strata/Strata.h
--- a/Engine/Source/Runtime/Renderer/Private/Strata/Strata.h
+++ b/Engine/Source/Runtime/Renderer/Private/Strata/Strata.h
@@ -18,7 +18,7 @@ struct FScreenPassTexture;
 BEGIN_SHADER_PARAMETER_STRUCT(FStrataBasePassUniformParameters, )
 	SHADER_PARAMETER(uint32, MaxBytesPerPixel)
 	SHADER_PARAMETER(uint32, bRoughDiffuse)
-	SHADER_PARAMETER_RDG_TEXTURE_UAV(RWTexture2DArray<uint>, MaterialTextureArrayUAV)
+	SHADER_PARAMETER_RDG_TEXTURE_UAV(RWTexture2DArray<uint>, MaterialTextureArrayUAVWithoutRTs)
 	SHADER_PARAMETER_RDG_TEXTURE_UAV(RWTexture2D<uint2>, SSSTextureUAV)
 END_SHADER_PARAMETER_STRUCT()
 
@@ -51,10 +51,8 @@ struct FStrataSceneData
 
 	// Resources allocated and updated each frame
 
-
-	FRDGTextureRef MaterialTextureArrayMRTs;
-
 	FRDGTextureRef MaterialTextureArray;
+	FRDGTextureUAVRef MaterialTextureArrayUAVWithoutRTs;
 	FRDGTextureUAVRef MaterialTextureArrayUAV;
 	FRDGTextureSRVRef MaterialTextureArraySRV;
 
