Removed Strata post base pass no longer necessary.

diff --git a/Engine/Shaders/Private/Strata/StrataMaterialClassification.usf b/Engine/Shaders/Private/Strata/StrataMaterialClassification.usf
--- a/Engine/Shaders/Private/Strata/StrataMaterialClassification.usf
+++ b/Engine/Shaders/Private/Strata/StrataMaterialClassification.usf
@@ -193,25 +193,3 @@ void ClearMaterialBufferMainCS(uint2 DispatchThreadId : SV_DispatchThreadID)
 
 ////////////////////////////////////////////////////////////////////////////////////////////////////////////
 
-#if SHADER_POST_BASEPASS
-
-Texture2DArray<uint> MaterialTextureArrayMRTs;
-RWTexture2DArray<uint> MaterialTextureArrayUAV;
-
-[numthreads(8, 8, 1)]
-void PostBasePassCS(uint2 DispatchThreadId : SV_DispatchThreadID)
-{
-	if (any(int2(DispatchThreadId.xy) >= View.ViewSizeAndInvSize.xy))
-	{
-		return;
-	}
-
-	int3 CoordData0 = int3(View.ViewRectMin.xy + DispatchThreadId.xy, 0);
-	int3 CoordData1 = int3(View.ViewRectMin.xy + DispatchThreadId.xy, 1);
-
-	MaterialTextureArrayUAV[CoordData0] = MaterialTextureArrayMRTs[CoordData0];
-	MaterialTextureArrayUAV[CoordData1] = MaterialTextureArrayMRTs[CoordData1];
-}
-
-#endif // SHADER_POST_BASEPASS
-
diff --git a/Engine/Source/Runtime/Renderer/Private/BasePassRendering.cpp b/Engine/Source/Runtime/Renderer/Private/BasePassRendering.cpp
--- a/Engine/Source/Runtime/Renderer/Private/BasePassRendering.cpp
+++ b/Engine/Source/Runtime/Renderer/Private/BasePassRendering.cpp
@@ -903,8 +903,6 @@ void FDeferredShadingSceneRenderer::RenderBasePass(
 		StampDeferredDebugProbeMaterialPS(GraphBuilder, Views, BasePassRenderTargets, SceneTextures);
 	}
 #endif
-
-	Strata::PostBasePass(GraphBuilder, Views);
 }
 
 BEGIN_SHADER_PARAMETER_STRUCT(FOpaqueBasePassParameters, )
diff --git a/Engine/Source/Runtime/Renderer/Private/Strata/Strata.cpp b/Engine/Source/Runtime/Renderer/Private/Strata/Strata.cpp
--- a/Engine/Source/Runtime/Renderer/Private/Strata/Strata.cpp
+++ b/Engine/Source/Runtime/Renderer/Private/Strata/Strata.cpp
@@ -674,70 +674,6 @@ void SetBasePassRenderTargetOutputFormat(const EShaderPlatform Platform, FShader
 
 ////////////////////////////////////////////////////////////////////////////////////////////////////////////////////////////////////////////////
 
-class FStrataPostBasePassCS : public FGlobalShader
-{
-	DECLARE_GLOBAL_SHADER(FStrataPostBasePassCS);
-	SHADER_USE_PARAMETER_STRUCT(FStrataPostBasePassCS, FGlobalShader);
-
-	BEGIN_SHADER_PARAMETER_STRUCT(FParameters, )
-		SHADER_PARAMETER_STRUCT_REF(FViewUniformShaderParameters, ViewUniformBuffer)
-		SHADER_PARAMETER_RDG_TEXTURE_SRV(Texture2DArray<uint>, MaterialTextureArrayMRTs)
-		SHADER_PARAMETER_RDG_TEXTURE_UAV(RWTexture2DArray<uint>, MaterialTextureArrayUAV)
-	END_SHADER_PARAMETER_STRUCT()
-
-	static bool ShouldCompilePermutation(const FGlobalShaderPermutationParameters& Parameters)
-	{
-		return GetMaxSupportedFeatureLevel(Parameters.Platform) >= ERHIFeatureLevel::SM5 && Strata::IsStrataEnabled();
-	}
-
-	static void ModifyCompilationEnvironment(const FGlobalShaderPermutationParameters& Parameters, FShaderCompilerEnvironment& OutEnvironment)
-	{
-		FGlobalShader::ModifyCompilationEnvironment(Parameters, OutEnvironment);
-		OutEnvironment.SetDefine(TEXT("SHADER_POST_BASEPASS"), 1);
-	}
-};
-IMPLEMENT_GLOBAL_SHADER(FStrataPostBasePassCS, "/Engine/Private/Strata/StrataMaterialClassification.usf", "PostBasePassCS", SF_Compute);
-
-void PostBasePass(FRDGBuilder& GraphBuilder, const TArray<FViewInfo>& Views)
-{
-
-
-
-
-	////// TODO delete PostBasePas shader
-
-
-
-
-
-	/*RDG_EVENT_SCOPE_CONDITIONAL(GraphBuilder, IsStrataEnabled() && Views.Num() > 0, "StrataPostBasePass");
-	if (!IsStrataEnabled())
-	{
-		return;
-	}
-
-	for (int32 i = 0; i < Views.Num(); ++i)
-	{
-		const FViewInfo& View = Views[i];
-
-		TShaderMapRef<FStrataPostBasePassCS> ComputeShader(View.ShaderMap);
-		FStrataPostBasePassCS::FParameters* PassParameters = GraphBuilder.AllocParameters<FStrataPostBasePassCS::FParameters>();
-		PassParameters->ViewUniformBuffer = View.ViewUniformBuffer;
-		PassParameters->MaterialTextureArrayMRTs = GraphBuilder.CreateSRV(FRDGTextureSRVDesc::Create(View.StrataSceneData->MaterialTextureArrayMRTs));
-		PassParameters->MaterialTextureArrayUAV = View.StrataSceneData->MaterialTextureArrayUAV;
-
-		const uint32 GroupSize = 8;
-		FComputeShaderUtils::AddPass(
-			GraphBuilder,
-			RDG_EVENT_NAME("Strata::PostBasePassCS"),
-			ComputeShader,
-			PassParameters,
-			FComputeShaderUtils::GetGroupCount(View.ViewRect.Size(), GroupSize));
-	}*/
-}
-
-////////////////////////////////////////////////////////////////////////////////////////////////////////////////////////////////////////////////
-
 void AddStrataMaterialClassificationPass(FRDGBuilder& GraphBuilder, const FMinimalSceneTextures& SceneTextures, const TArray<FViewInfo>& Views)
 {
 	RDG_EVENT_SCOPE_CONDITIONAL(GraphBuilder, IsStrataEnabled() && Views.Num() > 0, "StrataMaterialClassification");
diff --git a/Engine/Source/Runtime/Renderer/Private/Strata/Strata.h b/Engine/Source/Runtime/Renderer/Private/Strata/Strata.h
--- a/Engine/Source/Runtime/Renderer/Private/Strata/Strata.h
+++ b/Engine/Source/Runtime/Renderer/Private/Strata/Strata.h
@@ -96,7 +96,6 @@ void BindStrataForwardPasslUniformParameters(FRDGBuilder& GraphBuilder, FStrataS
 
 void AppendStrataMRTs(FSceneRenderer& SceneRenderer, uint32& BasePassTextureCount, TStaticArray<FTextureRenderTargetBinding, MaxSimultaneousRenderTargets>& BasePassTextures);
 void SetBasePassRenderTargetOutputFormat(const EShaderPlatform Platform, FShaderCompilerEnvironment& OutEnvironment);
-void PostBasePass(FRDGBuilder& GraphBuilder, const TArray<FViewInfo>& Views);
 
 TRDGUniformBufferRef<FStrataGlobalUniformParameters> BindStrataGlobalUniformParameters(FStrataSceneData* StrataSceneData);
 
