Conversion of FInjectShadowedLocalLightPS to RDG.

Volumetricfog screenshot test tested green

diff --git a/Engine/Source/Runtime/Renderer/Private/VolumeLighting.cpp b/Engine/Source/Runtime/Renderer/Private/VolumeLighting.cpp
--- a/Engine/Source/Runtime/Renderer/Private/VolumeLighting.cpp
+++ b/Engine/Source/Runtime/Renderer/Private/VolumeLighting.cpp
@@ -155,31 +155,7 @@ static auto GetVolumeShadowingShaderParametersGlobal = [](
 	ShaderParams.WorldToStaticShadowMatrix = WorldToStaticShadow;
 	ShaderParams.StaticShadowBufferSize = StaticShadowBufferSizeValue;
 
-	//
-	// See FOnePassPointShadowProjectionShaderParameters from ShadowRendering.h
-	//
-	FRDGTexture* ShadowDepthTextureValue = ShadowInfo
-		? GraphBuilder.RegisterExternalTexture(ShadowInfo->RenderTargets.DepthTarget) 
-		: SystemTextures.BlackDepthCube;
-
-	ShaderParams.OnePassPointShadowProjection.ShadowDepthCubeTexture = ShadowDepthTextureValue;
-	ShaderParams.OnePassPointShadowProjection.ShadowDepthCubeTexture2 = ShadowDepthTextureValue;
-	ShaderParams.OnePassPointShadowProjection.ShadowDepthCubeTextureSampler = TStaticSamplerState<SF_Bilinear, AM_Clamp, AM_Clamp, AM_Clamp, 0, 0, 0, SCF_Less>::GetRHI();
-
-	if (bDynamicallyShadowed)
-	{
-		if (ShadowInfo)
-		{
-			TArray<FMatrix44f> SIOnePassShadowViewProjectionMatrices = LWC::ConvertArrayType<FMatrix44f>(ShadowInfo->OnePassShadowViewProjectionMatrices);	// LWC_TODO: Precision loss. Perf pessimization
-			memcpy(ShaderParams.OnePassPointShadowProjection.ShadowViewProjectionMatrices.GetData(), SIOnePassShadowViewProjectionMatrices.GetData(), SIOnePassShadowViewProjectionMatrices.Num() * sizeof(FMatrix44f));
-			ShaderParams.OnePassPointShadowProjection.InvShadowmapResolution = 1.0f / float(ShadowInfo->ResolutionX);
-		}
-		else
-		{
-			memset(ShaderParams.OnePassPointShadowProjection.ShadowViewProjectionMatrices.GetData(), 0, sizeof(ShaderParams.OnePassPointShadowProjection.ShadowViewProjectionMatrices));
-			ShaderParams.OnePassPointShadowProjection.InvShadowmapResolution = 0.0f;
-		}
-	}
+	GetOnePassPointShadowProjectionParameters(GraphBuilder, bDynamicallyShadowed && (LightType == LightType_Point || LightType == LightType_Rect) ? ShadowInfo : NULL, ShaderParams.OnePassPointShadowProjection);
 };
 
 
diff --git a/Engine/Source/Runtime/Renderer/Private/VolumetricFog.cpp b/Engine/Source/Runtime/Renderer/Private/VolumetricFog.cpp
--- a/Engine/Source/Runtime/Renderer/Private/VolumetricFog.cpp
+++ b/Engine/Source/Runtime/Renderer/Private/VolumetricFog.cpp
@@ -262,14 +262,24 @@ IMPLEMENT_GLOBAL_SHADER(FWriteToBoundingSphereVS, "/Engine/Private/VolumetricFog
 class FInjectShadowedLocalLightPS : public FGlobalShader
 {
 	DECLARE_GLOBAL_SHADER(FInjectShadowedLocalLightPS);
+	SHADER_USE_PARAMETER_STRUCT(FInjectShadowedLocalLightPS, FGlobalShader);
 
 	BEGIN_SHADER_PARAMETER_STRUCT(FParameters, )
-		SHADER_PARAMETER_RDG_TEXTURE(Texture2D, ConservativeDepthTexture)
+		SHADER_PARAMETER_STRUCT_REF(FViewUniformShaderParameters, ViewUniformBuffer)
+		SHADER_PARAMETER_RDG_UNIFORM_BUFFER(FDeferredLightUniformStruct, DeferredLight)
+		SHADER_PARAMETER_STRUCT_INCLUDE(FVolumeShadowingShaderParameters, VolumeShadowingShaderParameters)
 		SHADER_PARAMETER_STRUCT_INCLUDE(FVirtualShadowMapSamplingParameters, VirtualShadowMapSamplingParameters)
+		SHADER_PARAMETER_RDG_TEXTURE(Texture2D, ConservativeDepthTexture)
 		SHADER_PARAMETER_RDG_TEXTURE(Texture2D, WhiteDummyTexture)
 		SHADER_PARAMETER_RDG_TEXTURE(Texture2D, LightFunctionAtlasTexture)
 		SHADER_PARAMETER_SAMPLER(SamplerState, LightFunctionAtlasSampler)
 		SHADER_PARAMETER_STRUCT_INCLUDE(FVolumetricFogIntegrationParameters, VolumetricFogParameters)
+		SHADER_PARAMETER(FMatrix44f, LocalLightFunctionMatrix)
+		SHADER_PARAMETER(FVector4f, LightFunctionAtlasTileMinMaxUvBound)
+		SHADER_PARAMETER(float, PhaseG)
+		SHADER_PARAMETER(float, InverseSquaredLightDistanceBiasScale)
+		SHADER_PARAMETER(uint32, UseConservativeDepthTexture)
+		SHADER_PARAMETER(int32, VirtualShadowMapId)
 		RENDER_TARGET_BINDING_SLOTS()
 	END_SHADER_PARAMETER_STRUCT()
 
@@ -296,78 +306,6 @@ class FInjectShadowedLocalLightPS : public FGlobalShader
 	{
 		return DoesPlatformSupportVolumetricFog(Parameters.Platform);
 	}
-
-	FInjectShadowedLocalLightPS(const ShaderMetaType::CompiledShaderInitializerType& Initializer) :
-		FGlobalShader(Initializer)
-	{
-		BindForLegacyShaderParameters<FParameters>(this, Initializer.PermutationId, Initializer.ParameterMap);
-		Bindings.BindForLegacyShaderParameters(this, Initializer.PermutationId, Initializer.ParameterMap, *FParameters::FTypeInfo::GetStructMetadata());
-
-		PhaseG.Bind(Initializer.ParameterMap, TEXT("PhaseG"));
-		InverseSquaredLightDistanceBiasScale.Bind(Initializer.ParameterMap, TEXT("InverseSquaredLightDistanceBiasScale"));
-		VolumeShadowingParameters.Bind(Initializer.ParameterMap);
-
-		LightFunctionMatrixParam.Bind(Initializer.ParameterMap, TEXT("LocalLightFunctionMatrix"));
-		LightFunctionAtlasTileMinMaxUvBoundParam.Bind(Initializer.ParameterMap, TEXT("LightFunctionAtlasTileMinMaxUvBound"));
-		LightFunctionAtlasTextureParam.Bind(Initializer.ParameterMap, TEXT("LightFunctionAtlasTexture"));
-		LightFunctionAtlasSamplerParam.Bind(Initializer.ParameterMap, TEXT("LightFunctionAtlasSampler"));
-
-		UseConservativeDepthTexture.Bind(Initializer.ParameterMap, TEXT("UseConservativeDepthTexture"));
-		VirtualShadowMapIdParameter.Bind(Initializer.ParameterMap, TEXT("VirtualShadowMapId"));
-	}
-
-	FInjectShadowedLocalLightPS() {}
-
-public:
-	// @param InnerSplitIndex which CSM shadow map level, INDEX_NONE if no directional light
-	// @param VolumeCascadeIndexValue which volume we render to
-	void SetParameters(
-		FRHICommandList& RHICmdList,
-		const FViewInfo& View,
-		const FVolumetricFogIntegrationParameterData& IntegrationData,
-		const FLightSceneInfo* LightSceneInfo,
-		const FExponentialHeightFogSceneInfo& FogInfo,
-		const FProjectedShadowInfo* ShadowMap,
-		bool bDynamicallyShadowed,
-		int32 VirtualShadowMapId,
-		const FMatrix& LightFunctionMatrix,
-		FRDGTextureRef LightFunctionAtlasTexture,
-		FVector4f LightFunctionAtlasTileMinMaxUvBound)
-	{
-		FRHIPixelShader* ShaderRHI = RHICmdList.GetBoundPixelShader();
-
-		FGlobalShader::SetParameters<FViewUniformShaderParameters>(RHICmdList, ShaderRHI, View.ViewUniformBuffer);
-
-		SetDeferredLightParameters(RHICmdList, ShaderRHI, GetUniformBufferParameter<FDeferredLightUniformStruct>(), LightSceneInfo, View);
-
-		SetShaderValue(RHICmdList, ShaderRHI, PhaseG, FogInfo.VolumetricFogScatteringDistribution);
-		SetShaderValue(RHICmdList, ShaderRHI, InverseSquaredLightDistanceBiasScale, GInverseSquaredLightDistanceBiasScale);
-
-		SetShaderValue(RHICmdList, ShaderRHI, LightFunctionAtlasTileMinMaxUvBoundParam, LightFunctionAtlasTileMinMaxUvBound);
-		SetShaderValue(RHICmdList, ShaderRHI, LightFunctionMatrixParam, (FMatrix44f)LightFunctionMatrix);
-		if (LightFunctionAtlasTextureParam.IsBound())
-		{
-			SetTextureParameter(RHICmdList, ShaderRHI, LightFunctionAtlasTextureParam, LightFunctionAtlasSamplerParam,
-				TStaticSamplerState<SF_Bilinear, AM_Clamp, AM_Clamp, AM_Clamp>::GetRHI(), LightFunctionAtlasTexture->GetRHI());
-		}
-
-		SetShaderValue(RHICmdList, ShaderRHI, UseConservativeDepthTexture, GVolumetricFogConservativeDepth > 0 ? 1 : 0);
-
-		VolumeShadowingParameters.Set(RHICmdList, ShaderRHI, View, LightSceneInfo, ShadowMap, INDEX_NONE, bDynamicallyShadowed);
-
-		SetShaderValue(RHICmdList, ShaderRHI, VirtualShadowMapIdParameter, VirtualShadowMapId);
-	}
-
-private:
-	LAYOUT_FIELD(FShaderParameter, PhaseG);
-	LAYOUT_FIELD(FShaderParameter, InverseSquaredLightDistanceBiasScale);
-	LAYOUT_FIELD(FVolumeShadowingParameters, VolumeShadowingParameters);
-	LAYOUT_FIELD(FShaderParameter, UseConservativeDepthTexture);
-	LAYOUT_FIELD(FShaderParameter, VirtualShadowMapIdParameter);
-	LAYOUT_FIELD(FShaderParameter, LightFunctionAtlasTileMinMaxUvBoundParam);
-	LAYOUT_FIELD(FShaderParameter, LightFunctionMatrixParam);
-	LAYOUT_FIELD(FShaderResourceParameter, LightFunctionAtlasTextureParam);
-	LAYOUT_FIELD(FShaderResourceParameter, LightFunctionAtlasSamplerParam);
 };
 
 IMPLEMENT_GLOBAL_SHADER(FInjectShadowedLocalLightPS, "/Engine/Private/VolumetricFog.usf", "InjectShadowedLocalLightPS", SF_Pixel);
@@ -381,7 +319,7 @@ void GetVolumeShadowingShaderParameters(
 	bool bDynamicallyShadowed,
 	FVolumeShadowingShaderParameters& OutParameters)
 {
-	//@todo DynamicGI: remove duplication with FVolumeShadowingParameters
+	//@todo DynamicGI: remove duplication with FVolumeShadowingParameters, only FTranslucentLightingInjectPS renamins to be converted to RDG
 
 	if (bDynamicallyShadowed && ShadowMap != nullptr)
 	{
@@ -626,9 +564,8 @@ void FDeferredShadingSceneRenderer::RenderLocalLightsForVolumetricFog(
 	FRDGTexture*& OutLocalShadowedLightScattering,
 	FRDGTextureRef ConservativeDepthTexture)
 {
-	TMap<FLightSceneInfo*, FVolumetricFogLocalLightFunctionInfo>& LocalLightFunctionData = View.VolumetricFogResources.LocalLightFunctionData;
+	// Gather lights that need to be rendered with shadow from opaque or light functions.
 	TArray<const FLightSceneInfo*, SceneRenderingAllocator> LightsToInject;
-
 	for (auto LightIt = Scene->Lights.CreateConstIterator(); LightIt; ++LightIt)
 	{
 		const FLightSceneInfoCompact& LightSceneInfoCompact = *LightIt;
@@ -652,77 +589,94 @@ void FDeferredShadingSceneRenderer::RenderLocalLightsForVolumetricFog(
 		}
 	}
 
+	// Now voxelise all the light we have just gathered.
 	if (LightsToInject.Num() > 0)
 	{
 		OutLocalShadowedLightScattering = GraphBuilder.CreateTexture(VolumeDesc, TEXT("VolumetricFog.LocalShadowedLightScattering"));
+		TMap<FLightSceneInfo*, FVolumetricFogLocalLightFunctionInfo>& LocalLightFunctionData = View.VolumetricFogResources.LocalLightFunctionData;
 
-		FInjectShadowedLocalLightPS::FParameters* PassParameters = GraphBuilder.AllocParameters<FInjectShadowedLocalLightPS::FParameters>();
-		PassParameters->RenderTargets[0] = FRenderTargetBinding(OutLocalShadowedLightScattering, ERenderTargetLoadAction::EClear);
-		PassParameters->VirtualShadowMapSamplingParameters = VirtualShadowMapArray.GetSamplingParameters(GraphBuilder);
-		PassParameters->ConservativeDepthTexture = ConservativeDepthTexture;
-		PassParameters->LightFunctionAtlasTexture = View.VolumetricFogResources.TransientLightFunctionTextureAtlas ? View.VolumetricFogResources.TransientLightFunctionTextureAtlas->GetTransientLightFunctionAtlasTexture() : GSystemTextures.GetWhiteDummy(GraphBuilder);
-		PassParameters->LightFunctionAtlasSampler = TStaticSamplerState<SF_Bilinear, AM_Clamp, AM_Clamp, AM_Clamp>::GetRHI();
-		// We also bind the default light function texture because when we are out of atlas tile, we fallback to use a white light function so we need the RHI to be created
-		PassParameters->WhiteDummyTexture = View.VolumetricFogResources.TransientLightFunctionTextureAtlas ? View.VolumetricFogResources.TransientLightFunctionTextureAtlas->GetDefaultLightFunctionTexture() : GSystemTextures.GetWhiteDummy(GraphBuilder);
-		SetupVolumetricFogIntegrationParameters(PassParameters->VolumetricFogParameters, View, IntegrationData);
-
-		GraphBuilder.AddPass(
-			RDG_EVENT_NAME("ShadowedLights"),
-			PassParameters,
-			ERDGPassFlags::Raster,
-			[PassParameters, &View, this, LightsToInject, VolumetricFogGridSize, GridZParams, bUseTemporalReprojection, IntegrationData, FogInfo](FRHICommandList& RHICmdList)
-			{
-				TMap<FLightSceneInfo*, FVolumetricFogLocalLightFunctionInfo>& LocalLightFunctionData = View.VolumetricFogResources.LocalLightFunctionData;
-
-				for (int32 LightIndex = 0; LightIndex < LightsToInject.Num(); LightIndex++)
-				{
-					const FLightSceneInfo* LightSceneInfo = LightsToInject[LightIndex];
-					const FVisibleLightInfo& VisibleLightInfo = VisibleLightInfos[LightSceneInfo->Id];
-					const FProjectedShadowInfo* ProjectedShadowInfo = GetShadowForInjectionIntoVolumetricFog(VisibleLightInfo);
-
-					const bool bDynamicallyShadowed = ProjectedShadowInfo != NULL;
-					
-					const FSphere LightBounds = LightSceneInfo->Proxy->GetBoundingSphere();
-					const FIntPoint VolumeZBounds = CalculateVolumetricFogBoundsForLight(LightBounds, View, VolumetricFogGridSize, GridZParams);
-
-					if (VolumeZBounds.X < VolumeZBounds.Y)
-					{
-						bool bIsShadowed = LightNeedsSeparateInjectionIntoVolumetricFogForOpaqueShadow(View, LightSceneInfo, VisibleLightInfo);
-						bool bUsesLightFunction = ViewFamily.EngineShowFlags.LightFunctions 
-							&& CheckForLightFunction(LightSceneInfo) && LightNeedsSeparateInjectionIntoVolumetricFogForLightFunction(LightSceneInfo);
+		bool bClearExecuted = false;
+		for (int32 LightIndex = 0; LightIndex < LightsToInject.Num(); LightIndex++)
+		{
+			const FLightSceneInfo* LightSceneInfo = LightsToInject[LightIndex];
+			const FVisibleLightInfo& VisibleLightInfo = VisibleLightInfos[LightSceneInfo->Id];
 
-						int32 VirtualShadowMapId = VisibleLightInfo.GetVirtualShadowMapId( &View );
-						const bool bUseVSM = bIsShadowed && VirtualShadowMapArray.IsAllocated() && VirtualShadowMapId != INDEX_NONE;
+			const FSphere LightBounds = LightSceneInfo->Proxy->GetBoundingSphere();
+			const FIntPoint VolumeZBounds = CalculateVolumetricFogBoundsForLight(LightBounds, View, VolumetricFogGridSize, GridZParams);
+			if (VolumeZBounds.X < VolumeZBounds.Y)
+			{
+				bool bIsShadowed = LightNeedsSeparateInjectionIntoVolumetricFogForOpaqueShadow(View, LightSceneInfo, VisibleLightInfo);
+				bool bUsesLightFunction = ViewFamily.EngineShowFlags.LightFunctions
+					&& CheckForLightFunction(LightSceneInfo) && LightNeedsSeparateInjectionIntoVolumetricFogForLightFunction(LightSceneInfo);
 
-						FRDGTextureRef LightFunctionTexture = PassParameters->LightFunctionAtlasTexture;
-						FMatrix LightFunctionMatrix = FMatrix::Identity;
-						FVector4f LightFunctionAtlasTileMinMaxUvBound = FVector4f(ForceInitToZero);
-						if (bUsesLightFunction)
-						{
-							FVolumetricFogLocalLightFunctionInfo* LightFunctionData = LocalLightFunctionData.Find(LightSceneInfo);
+				int32 VirtualShadowMapId = VisibleLightInfo.GetVirtualShadowMapId(&View);
+				const bool bUseVSM = bIsShadowed && VirtualShadowMapArray.IsAllocated() && VirtualShadowMapId != INDEX_NONE;
 
-							if (!ensure(LightFunctionData != nullptr))
-							{
-								// The light function data is missing but the light requires it. Skip this light for now.
-								continue;
-							}
+				FInjectShadowedLocalLightPS::FParameters* PassParameters = GraphBuilder.AllocParameters<FInjectShadowedLocalLightPS::FParameters>();
 
-							LightFunctionMatrix = LightFunctionData->LightFunctionMatrix;
-							LightFunctionTexture = LightFunctionData->AtlasTile.Texture;
-							LightFunctionAtlasTileMinMaxUvBound = LightFunctionData->AtlasTile.MinMaxUvBound;
-						}
+				// Light function parameters
+				PassParameters->LightFunctionAtlasTexture = View.VolumetricFogResources.TransientLightFunctionTextureAtlas ? View.VolumetricFogResources.TransientLightFunctionTextureAtlas->GetTransientLightFunctionAtlasTexture() : GSystemTextures.GetWhiteDummy(GraphBuilder);;
+				PassParameters->LocalLightFunctionMatrix = FMatrix::Identity;
+				PassParameters->LightFunctionAtlasTileMinMaxUvBound = FVector4f(ForceInitToZero);
+				if (bUsesLightFunction)
+				{
+					FVolumetricFogLocalLightFunctionInfo* LightFunctionData = LocalLightFunctionData.Find(LightSceneInfo);
 
-						FInjectShadowedLocalLightPS::FPermutationDomain PermutationVector;
-						PermutationVector.Set< FInjectShadowedLocalLightPS::FDynamicallyShadowed >(bDynamicallyShadowed);
-						PermutationVector.Set< FInjectShadowedLocalLightPS::FTemporalReprojection >(bUseTemporalReprojection);
-						PermutationVector.Set< FInjectShadowedLocalLightPS::FLightFunction >(bUsesLightFunction);
-						PermutationVector.Set< FInjectShadowedLocalLightPS::FEnableShadows >(bIsShadowed);
-						PermutationVector.Set< FInjectShadowedLocalLightPS::FVirtualShadowMap >(bUseVSM);
+					if (!ensure(LightFunctionData != nullptr))
+					{
+						// The light function data is missing but the light requires it. Skip this light for now.
+						continue;
+					}
 
-						auto VertexShader = View.ShaderMap->GetShader< FWriteToBoundingSphereVS >();
-						TOptionalShaderMapRef<FWriteToSliceGS> GeometryShader(View.ShaderMap);
-						auto PixelShader = View.ShaderMap->GetShader< FInjectShadowedLocalLightPS >(PermutationVector);
+					PassParameters->LocalLightFunctionMatrix = LightFunctionData->LightFunctionMatrix;
+					PassParameters->LightFunctionAtlasTexture = LightFunctionData->AtlasTile.Texture;
+					PassParameters->LightFunctionAtlasTileMinMaxUvBound = LightFunctionData->AtlasTile.MinMaxUvBound;
+				}
 
+				PassParameters->RenderTargets[0] = FRenderTargetBinding(OutLocalShadowedLightScattering, bClearExecuted ? ERenderTargetLoadAction::ELoad : ERenderTargetLoadAction::EClear);
+				bClearExecuted = true;
+				PassParameters->VirtualShadowMapSamplingParameters = VirtualShadowMapArray.GetSamplingParameters(GraphBuilder);
+				PassParameters->ConservativeDepthTexture = ConservativeDepthTexture;
+				PassParameters->LightFunctionAtlasSampler = TStaticSamplerState<SF_Bilinear, AM_Clamp, AM_Clamp, AM_Clamp>::GetRHI();
+				// We also bind the default light function texture because when we are out of atlas tile, we fallback to use a white light function so we need the RHI to be created
+				PassParameters->WhiteDummyTexture = View.VolumetricFogResources.TransientLightFunctionTextureAtlas ? View.VolumetricFogResources.TransientLightFunctionTextureAtlas->GetDefaultLightFunctionTexture() : GSystemTextures.GetWhiteDummy(GraphBuilder);
+				SetupVolumetricFogIntegrationParameters(PassParameters->VolumetricFogParameters, View, IntegrationData);
+
+				PassParameters->ViewUniformBuffer = View.ViewUniformBuffer;
+				PassParameters->PhaseG = FogInfo.VolumetricFogScatteringDistribution;
+				PassParameters->InverseSquaredLightDistanceBiasScale = GInverseSquaredLightDistanceBiasScale;
+				PassParameters->UseConservativeDepthTexture = GVolumetricFogConservativeDepth > 0 ? 1 : 0;
+				PassParameters->VirtualShadowMapId = VirtualShadowMapId;
+
+				FDeferredLightUniformStruct* DeferredLightStruct = GraphBuilder.AllocParameters<FDeferredLightUniformStruct>();
+				*DeferredLightStruct = GetDeferredLightParameters(View, *LightSceneInfo);
+				PassParameters->DeferredLight = GraphBuilder.CreateUniformBuffer(DeferredLightStruct);
+
+				const FProjectedShadowInfo* ProjectedShadowInfo = GetShadowForInjectionIntoVolumetricFog(VisibleLightInfo);
+				const bool bDynamicallyShadowed = ProjectedShadowInfo != NULL;
+				GetVolumeShadowingShaderParameters(GraphBuilder, View, LightSceneInfo, ProjectedShadowInfo, 0, bDynamicallyShadowed, PassParameters->VolumeShadowingShaderParameters);
+
+				FInjectShadowedLocalLightPS::FPermutationDomain PermutationVector;
+				PermutationVector.Set< FInjectShadowedLocalLightPS::FDynamicallyShadowed >(bDynamicallyShadowed);
+				PermutationVector.Set< FInjectShadowedLocalLightPS::FTemporalReprojection >(bUseTemporalReprojection);
+				PermutationVector.Set< FInjectShadowedLocalLightPS::FLightFunction >(bUsesLightFunction);
+				PermutationVector.Set< FInjectShadowedLocalLightPS::FEnableShadows >(bIsShadowed);
+				PermutationVector.Set< FInjectShadowedLocalLightPS::FVirtualShadowMap >(bUseVSM);
+
+				auto VertexShader = View.ShaderMap->GetShader< FWriteToBoundingSphereVS >();
+				TOptionalShaderMapRef<FWriteToSliceGS> GeometryShader(View.ShaderMap);
+				auto PixelShader = View.ShaderMap->GetShader< FInjectShadowedLocalLightPS >(PermutationVector);
+
+				// We execute one pass per light: this is because RDG resources needs to be gathrered before and reference in the PassParameters.
+				// Not many lights cast shadow so that is acceptable (LightRendering is doing the same things).
+				// If light shadow maps woud be in a common resources (atlas, texture array, bindless) we could have a single pass for all the lights.
+				// NOTE: light functions are already in an atlas so they are not a problem.
+				GraphBuilder.AddPass(
+					RDG_EVENT_NAME("ShadowedLights"),
+					PassParameters,
+					ERDGPassFlags::Raster,
+					[PassParameters, &View, this, VertexShader, GeometryShader, PixelShader, VolumeZBounds, LightBounds](FRHICommandList& RHICmdList)
+				{
 						FGraphicsPipelineStateInitializer GraphicsPSOInit;
 						RHICmdList.ApplyCachedRenderTargets(GraphicsPSOInit);
 						GraphicsPSOInit.RasterizerState = TStaticRasterizerState<FM_Solid, CM_None>::GetRHI();
@@ -738,9 +692,6 @@ void FDeferredShadingSceneRenderer::RenderLocalLightsForVolumetricFog(
 
 						SetGraphicsPipelineState(RHICmdList, GraphicsPSOInit, 0);
 
-						PixelShader->SetParameters(RHICmdList, View, IntegrationData, LightSceneInfo, FogInfo, ProjectedShadowInfo, bDynamicallyShadowed, VirtualShadowMapId,
-							LightFunctionMatrix, LightFunctionTexture, LightFunctionAtlasTileMinMaxUvBound);
-
 						SetShaderParameters(RHICmdList, PixelShader, PixelShader.GetPixelShader(), *PassParameters);
 
 						FWriteToBoundingSphereVS::FParameters VSPassParameters;
@@ -759,9 +710,9 @@ void FDeferredShadingSceneRenderer::RenderLocalLightsForVolumetricFog(
 						const int32 NumInstances = VolumeZBounds.Y - VolumeZBounds.X;
 						const int32 NumTriangles = FCircleRasterizeVertexBuffer::NumVertices - 2;
 						RHICmdList.DrawIndexedPrimitive(GCircleRasterizeIndexBuffer.IndexBufferRHI, 0, 0, FCircleRasterizeVertexBuffer::NumVertices, 0, NumTriangles, NumInstances);
-					}
-				}
-			});
+				});
+			}
+		}
 	}
 }
 
