Fix bug caused by preshadows not being projected when cached after CL 18393307.

diff --git a/Engine/Source/Runtime/Renderer/Private/ShadowSetup.cpp b/Engine/Source/Runtime/Renderer/Private/ShadowSetup.cpp
--- a/Engine/Source/Runtime/Renderer/Private/ShadowSetup.cpp
+++ b/Engine/Source/Runtime/Renderer/Private/ShadowSetup.cpp
@@ -4974,7 +4974,8 @@ void FSceneRenderer::AllocateShadowDepthTargets(FRHICommandListImmediate& RHICmd
 				}
 			}
 			// If uncached and no primitives, skip allocations etc
-			else if (ProjectedShadowInfo->CacheMode == SDCM_Uncached && !ProjectedShadowInfo->HasSubjectPrims() && !ProjectedShadowInfo->IsWholeSceneDirectionalShadow())
+			// Note: cached PreShadows, for some reason, has CacheMode == SDCM_Uncached so need to skip them here.
+			else if ((ProjectedShadowInfo->CacheMode == SDCM_Uncached && !ProjectedShadowInfo->bPreShadow) && !ProjectedShadowInfo->HasSubjectPrims() && !ProjectedShadowInfo->IsWholeSceneDirectionalShadow())
 			{
 				bShadowIsVisible = false;
 			}
