Only output scenecolor for sunmask on IOS when MSAA is enabled.

diff --git a/Engine/Shaders/Private/PostProcessMobile.usf b/Engine/Shaders/Private/PostProcessMobile.usf
--- a/Engine/Shaders/Private/PostProcessMobile.usf
+++ b/Engine/Shaders/Private/PostProcessMobile.usf
@@ -140,7 +140,7 @@ half2 SunConstDepthMaskScaleBias()
 void SunMaskPS_Mobile(
 	float4 InUVPos : TEXCOORD0,
 	out HALF_TYPE OutSunShaftAndDof : SV_Target0
-#if !MOBILE_USEDEPTHTEXTURE
+#if MOBILE_USESUN && METAL_MSAA_HDR_DECODE
 	, out HALF4_TYPE OutColor : SV_Target1
 #endif
 	)
@@ -157,6 +157,7 @@ void SunMaskPS_Mobile(
 
 		#if METAL_MSAA_HDR_DECODE
 			SceneColor.rgb *= rcp(SceneColor.r*(-0.299) + SceneColor.g*(-0.587) + SceneColor.b*(-0.114) + 1.0);
+			OutColor = SceneColor;
 		#endif
 
 		half2 DepthMaskScaleBias = SunConstDepthMaskScaleBias();
@@ -177,10 +178,6 @@ void SunMaskPS_Mobile(
 	#if MOBILE_USEDOF
 		OutSunShaftAndDof += Coc2(InDepth);
 	#endif
-
-	#if !MOBILE_USEDEPTHTEXTURE
-		OutColor = SceneColor;
-	#endif
 }
 
 #endif
diff --git a/Engine/Source/Runtime/Renderer/Private/PostProcess/PostProcessMobile.cpp b/Engine/Source/Runtime/Renderer/Private/PostProcess/PostProcessMobile.cpp
--- a/Engine/Source/Runtime/Renderer/Private/PostProcess/PostProcessMobile.cpp
+++ b/Engine/Source/Runtime/Renderer/Private/PostProcess/PostProcessMobile.cpp
@@ -730,7 +730,7 @@ FMobileSunMaskOutputs AddMobileSunMaskPass(FRDGBuilder& GraphBuilder, const FVie
 		FRDGTextureDesc OutputDesc = Inputs.SceneColor.Texture->Desc;
 		OutputDesc.Reset();
 
-		if (!Inputs.bUseDepthTexture)
+		if (Inputs.bUseSun && Inputs.bUseMetalMSAAHDRDecode)
 		{
 			if (IsMobilePropagateAlphaEnabled(View.GetShaderPlatform()))
 			{
diff --git a/Engine/Source/Runtime/Renderer/Private/PostProcess/PostProcessing.cpp b/Engine/Source/Runtime/Renderer/Private/PostProcess/PostProcessing.cpp
--- a/Engine/Source/Runtime/Renderer/Private/PostProcess/PostProcessing.cpp
+++ b/Engine/Source/Runtime/Renderer/Private/PostProcess/PostProcessing.cpp
@@ -1997,14 +1997,13 @@ void AddMobilePostProcessingPasses(FRDGBuilder& GraphBuilder, FScene* Scene, con
 
 			PostProcessSunShaftAndDof = SunMaskOutputs.SunMask;
 
-			if (!bUseDepthTexture)
+			// The scene color will be decoded after sun mask pass and output to linear color space for following passes if sun shaft enabled
+			// set bMetalMSAAHDRDecode to false if sun shaft enabled
+			if (bMetalMSAAHDRDecode && bUseSun)
 			{
 				SceneColor = SunMaskOutputs.SceneColor;
+				bMetalMSAAHDRDecode = false;
 			}
-
-			// The scene color will be decoded after sun mask pass and output to linear color space for following passes if sun shaft enabled
-			// set bMetalMSAAHDRDecode to false if sun shaft enabled
-			bMetalMSAAHDRDecode = (bMetalMSAAHDRDecode && !bUseSun);
 			//@todo Ronin sunmask pass isnt clipping to image only.
 		}
 
