Strata - Fixed broken categorisation
- single materials were tagged as simple
- complex material were writting pixels full screen to to read and ref being 0x0.
Fixes weird colors on some materials.

diff --git a/Engine/Source/Runtime/Renderer/Private/LightRendering.cpp b/Engine/Source/Runtime/Renderer/Private/LightRendering.cpp
--- a/Engine/Source/Runtime/Renderer/Private/LightRendering.cpp
+++ b/Engine/Source/Runtime/Renderer/Private/LightRendering.cpp
@@ -1697,7 +1697,7 @@ static uint32 InternalSetBoundingGeometryDepthState(FGraphicsPipelineStateInitia
 		{
 		case EStrataTileMaterialType::ESimple : StencilRef = Strata::StencilBit_Fast;    GraphicsPSOInit.DepthStencilState = TStaticDepthStencilState<false, CompareFunction, true, CF_Equal, SO_Keep, SO_Keep, SO_Keep, true, CF_Equal, SO_Keep, SO_Keep, SO_Keep, Strata::StencilBit_Fast, 0x0>::GetRHI(); break;
 		case EStrataTileMaterialType::ESingle : StencilRef = Strata::StencilBit_Single;  GraphicsPSOInit.DepthStencilState = TStaticDepthStencilState<false, CompareFunction, true, CF_Equal, SO_Keep, SO_Keep, SO_Keep, true, CF_Equal, SO_Keep, SO_Keep, SO_Keep, Strata::StencilBit_Single, 0x0>::GetRHI(); break;
-		case EStrataTileMaterialType::EComplex: StencilRef = Strata::StencilBit_Complex; GraphicsPSOInit.DepthStencilState = TStaticDepthStencilState<false, CompareFunction, true, CF_Equal, SO_Keep, SO_Keep, SO_Keep, true, CF_Equal, SO_Keep, SO_Keep, SO_Keep, 0x0, 0x0>::GetRHI(); break;
+		case EStrataTileMaterialType::EComplex: StencilRef = Strata::StencilBit_Complex; GraphicsPSOInit.DepthStencilState = TStaticDepthStencilState<false, CompareFunction, true, CF_Equal, SO_Keep, SO_Keep, SO_Keep, true, CF_Equal, SO_Keep, SO_Keep, SO_Keep, Strata::StencilBit_Fast | Strata::StencilBit_Single, 0x0>::GetRHI(); break;
 		default: check(false);
 		}
 	}
diff --git a/Engine/Source/Runtime/Renderer/Private/Strata/Strata.cpp b/Engine/Source/Runtime/Renderer/Private/Strata/Strata.cpp
--- a/Engine/Source/Runtime/Renderer/Private/Strata/Strata.cpp
+++ b/Engine/Source/Runtime/Renderer/Private/Strata/Strata.cpp
@@ -565,11 +565,12 @@ static void AddStrataInternalClassificationTilePass(
 		RDG_EVENT_NAME("Strata::%sClassificationPass(%s)", bDebug ? TEXT("Debug") : TEXT("Stencil"), ToString(TileMaterialType)),
 		ParametersPS,
 		ERDGPassFlags::Raster,
-		[ParametersPS, VertexShader, PixelShader, OutputResolution, StrataTilePrimitiveType, bDebug](FRHICommandList& RHICmdList)
+		[ParametersPS, VertexShader, PixelShader, OutputResolution, StrataTilePrimitiveType, TileMaterialType, bDebug](FRHICommandList& RHICmdList)
 		{
 			FGraphicsPipelineStateInitializer GraphicsPSOInit;
 			RHICmdList.ApplyCachedRenderTargets(GraphicsPSOInit);
 			GraphicsPSOInit.RasterizerState = TStaticRasterizerState<>::GetRHI();
+			uint32 StencilRef = 0xFF;
 			if (bDebug)
 			{
 				// Use premultiplied alpha blending, pixel shader and depth/stencil is off
@@ -579,19 +580,34 @@ static void AddStrataInternalClassificationTilePass(
 			}
 			else
 			{
+				check(TileMaterialType == EStrataTileMaterialType::ESimple || TileMaterialType == EStrataTileMaterialType::ESingle);
+
 				// No blending and no pixel shader required. Stencil will be writen to.
 				GraphicsPSOInit.BoundShaderState.PixelShaderRHI = nullptr;
 				GraphicsPSOInit.BlendState = TStaticBlendState<>::GetRHI();
-				GraphicsPSOInit.DepthStencilState = TStaticDepthStencilState<
-					false, CF_Always, 
-					true,  CF_Always, SO_Keep, SO_Keep, SO_Replace,
-					false, CF_Always, SO_Keep, SO_Keep, SO_Keep,
-					0xFF, StencilBit_Fast>::GetRHI();
+				if (TileMaterialType == EStrataTileMaterialType::ESimple)
+				{
+					GraphicsPSOInit.DepthStencilState = TStaticDepthStencilState<
+						false, CF_Always,
+						true, CF_Always, SO_Keep, SO_Keep, SO_Replace,
+						false, CF_Always, SO_Keep, SO_Keep, SO_Keep,
+						0xFF, StencilBit_Fast>::GetRHI();
+					StencilRef = StencilBit_Fast;
+				}
+				else // EStrataTileMaterialType::ESingle
+				{
+					GraphicsPSOInit.DepthStencilState = TStaticDepthStencilState<
+						false, CF_Always,
+						true, CF_Always, SO_Keep, SO_Keep, SO_Replace,
+						false, CF_Always, SO_Keep, SO_Keep, SO_Keep,
+						0xFF, StencilBit_Single>::GetRHI();
+					StencilRef = StencilBit_Single;
+				}
 			}
 			GraphicsPSOInit.BoundShaderState.VertexDeclarationRHI = GFilterVertexDeclaration.VertexDeclarationRHI;
 			GraphicsPSOInit.BoundShaderState.VertexShaderRHI = VertexShader.GetVertexShader();
 			GraphicsPSOInit.PrimitiveType = StrataTilePrimitiveType;
-			SetGraphicsPipelineState(RHICmdList, GraphicsPSOInit, StencilBit_Fast);
+			SetGraphicsPipelineState(RHICmdList, GraphicsPSOInit, StencilRef);
 			SetShaderParameters(RHICmdList, VertexShader, VertexShader.GetVertexShader(), ParametersPS->VS);
 			if (bDebug)
 			{
