Allow custom depth and depth to be transient when extracted for Slate.

diff --git a/Engine/Source/Runtime/Renderer/Private/CustomDepthRendering.cpp b/Engine/Source/Runtime/Renderer/Private/CustomDepthRendering.cpp
--- a/Engine/Source/Runtime/Renderer/Private/CustomDepthRendering.cpp
+++ b/Engine/Source/Runtime/Renderer/Private/CustomDepthRendering.cpp
@@ -108,16 +108,12 @@ FCustomDepthTextures FCustomDepthTextures::Create(FRDGBuilder& GraphBuilder, FIn
 		const FRDGTextureDesc MobileCustomStencilDesc = FRDGTextureDesc::Create2D(CustomDepthExtent, PF_G8, FClearValueBinding::Transparent, MobileCustomStencilFlags);
 
 		CustomDepthTextures.MobileDepth = GraphBuilder.CreateTexture(MobileCustomDepthDesc, TEXT("MobileCustomDepth"));
-		CustomDepthTextures.MobileDepth->SetNonTransient();
-
 		CustomDepthTextures.MobileStencil = GraphBuilder.CreateTexture(MobileCustomStencilDesc, TEXT("MobileCustomStencil"));
-		CustomDepthTextures.MobileStencil->SetNonTransient();
 	}
 
 	const FRDGTextureDesc CustomDepthDesc = FRDGTextureDesc::Create2D(CustomDepthExtent, PF_DepthStencil, FClearValueBinding::DepthFar, GFastVRamConfig.CustomDepth | TexCreate_NoFastClear | TexCreate_DepthStencilTargetable | TexCreate_ShaderResource);
 
 	CustomDepthTextures.Depth = GraphBuilder.CreateTexture(CustomDepthDesc, TEXT("CustomDepth"));
-	CustomDepthTextures.Depth->SetNonTransient();
 
 	CustomDepthTextures.Stencil = GraphBuilder.CreateSRV(FRDGTextureSRVDesc::CreateWithPixelFormat(CustomDepthTextures.Depth, PF_X24_G8));
 	CustomDepthTextures.DepthAction = ERenderTargetLoadAction::EClear;
diff --git a/Engine/Source/Runtime/Renderer/Private/SceneTextures.cpp b/Engine/Source/Runtime/Renderer/Private/SceneTextures.cpp
--- a/Engine/Source/Runtime/Renderer/Private/SceneTextures.cpp
+++ b/Engine/Source/Runtime/Renderer/Private/SceneTextures.cpp
@@ -633,7 +633,6 @@ FSceneTextures& FMinimalSceneTextures::Create(FRDGBuilder& GraphBuilder, const F
 							 FRDGTextureDesc::Create2D(SceneTextures.Config.Extent, PF_DepthStencil, Config.DepthClearValue, Flags));
 		Desc.NumSamples = Config.NumSamples;
 		SceneTextures.Depth = GraphBuilder.CreateTexture(Desc, TEXT("SceneDepthZ"));
-		SceneTextures.Depth.Target->SetNonTransient();
 
 		if (Desc.NumSamples > 1)
 		{
@@ -646,7 +645,6 @@ FSceneTextures& FMinimalSceneTextures::Create(FRDGBuilder& GraphBuilder, const F
 			else
 			{
 				SceneTextures.Depth.Resolve = GraphBuilder.CreateTexture(Desc, TEXT("SceneDepthZ"));
-				SceneTextures.Depth.Resolve->SetNonTransient();
 			}
 		}
 
@@ -687,11 +685,6 @@ FSceneTextures& FMinimalSceneTextures::Create(FRDGBuilder& GraphBuilder, const F
 
 			SceneTextures.Color.Resolve = GraphBuilder.CreateTexture(Desc, SceneColorName);
 		}
-
-		if (FRDGBuilder::IsDrainEnabled())
-		{
-			SceneTextures.Color.Resolve->SetNonTransient();
-		}
 	}
 
 	// Custom Depth
@@ -792,14 +785,6 @@ FSceneTextures& FSceneTextures::Create(FRDGBuilder& GraphBuilder, const FSceneTe
 			const FRDGTextureDesc Desc = FRDGTextureDesc::Create2D(Config.Extent, GetGBufferFFormat(), FClearValueBinding({ 0.5f, 0.5f, 0.5f, 0.5f }), TexCreate_RenderTargetable | TexCreate_ShaderResource | FlagsToAdd | GFastVRamConfig.GBufferF);
 			SceneTextures.GBufferF = GraphBuilder.CreateTexture(Desc, TEXT("GBufferF"));
 		}
-
-		// Lighting features may require GBuffer history.
-		if (FRDGBuilder::IsDrainEnabled())
-		{
-			SceneTextures.GBufferA->SetNonTransient();
-			SceneTextures.GBufferB->SetNonTransient();
-			SceneTextures.GBufferC->SetNonTransient();
-		}
 	}
 
 
@@ -926,7 +911,7 @@ void FSceneTextureExtracts::QueueExtractions(FRDGBuilder& GraphBuilder, const FS
 	{
 		if (HasBeenProduced(Texture) && !EnumHasAnyFlags(Texture->Desc.Flags, TexCreate_Memoryless))
 		{
-			GraphBuilder.QueueTextureExtraction(Texture, &OutTarget);
+			GraphBuilder.QueueTextureExtraction(Texture, &OutTarget, ERDGResourceExtractionFlags::AllowTransient);
 		}
 	};
 
