Fix static pre-shadow applied onto hair strands geometry.

diff --git a/Engine/Source/Runtime/Renderer/Private/ShadowRendering.cpp b/Engine/Source/Runtime/Renderer/Private/ShadowRendering.cpp
--- a/Engine/Source/Runtime/Renderer/Private/ShadowRendering.cpp
+++ b/Engine/Source/Runtime/Renderer/Private/ShadowRendering.cpp
@@ -441,14 +441,14 @@ IMPLEMENT_SHADOW_PROJECTION_PIXEL_SHADER(5, true);
 
 template<typename VertexShaderType, typename PixelShaderType>
 static void BindShaderShaders(FRHICommandList& RHICmdList, FGraphicsPipelineStateInitializer& GraphicsPSOInit,
-	int32 ViewIndex, const FViewInfo& View, const FProjectedShadowInfo* ShadowInfo, FRHIUniformBuffer* HairStrandsUniformBuffer = nullptr)
+	int32 ViewIndex, const FViewInfo& View, const FProjectedShadowInfo* ShadowInfo, uint32 StencilRef, FRHIUniformBuffer* HairStrandsUniformBuffer = nullptr)
 {
 	TShaderRef<VertexShaderType> VertexShader = View.ShaderMap->GetShader<VertexShaderType>();
 	TShaderRef<PixelShaderType> PixelShader = View.ShaderMap->GetShader<PixelShaderType>();
 
 	GraphicsPSOInit.BoundShaderState.VertexShaderRHI = VertexShader.GetVertexShader();
 	GraphicsPSOInit.BoundShaderState.PixelShaderRHI = PixelShader.GetPixelShader();
-	SetGraphicsPipelineState(RHICmdList, GraphicsPSOInit, 0);
+	SetGraphicsPipelineState(RHICmdList, GraphicsPSOInit, StencilRef);
 
 	VertexShader->SetParameters(RHICmdList, View, ShadowInfo, EShadowProjectionVertexShaderFlags::DrawingFrustum);
 	PixelShader->SetParameters(RHICmdList, ViewIndex, View, ShadowInfo);
@@ -469,7 +469,7 @@ static void BindShaderShaders(FRHICommandList& RHICmdList, FGraphicsPipelineStat
 
 
 static void BindShadowProjectionShaders(int32 Quality, FRHICommandList& RHICmdList, FGraphicsPipelineStateInitializer GraphicsPSOInit, int32 ViewIndex, const FViewInfo& View,
-	const FProjectedShadowInfo* ShadowInfo, bool bMobileModulatedProjections, FRHIUniformBuffer* HairStrandsUniformBuffer)
+	const FProjectedShadowInfo* ShadowInfo, uint32 StencilRef, bool bMobileModulatedProjections, FRHIUniformBuffer* HairStrandsUniformBuffer)
 {
 	const bool bSubPixelShadow = HairStrandsUniformBuffer != nullptr;
 	if (bSubPixelShadow)
@@ -480,11 +480,11 @@ static void BindShadowProjectionShaders(int32 Quality, FRHICommandList& RHICmdLi
 		{
 			switch (Quality)
 			{
-			case 1: BindShaderShaders<FShadowProjectionNoTransformVS, TShadowProjectionPS<1, false, false, false, true> >(RHICmdList, GraphicsPSOInit, ViewIndex, View, ShadowInfo, HairStrandsUniformBuffer); break;
-			case 2: BindShaderShaders<FShadowProjectionNoTransformVS, TShadowProjectionPS<2, false, false, false, true> >(RHICmdList, GraphicsPSOInit, ViewIndex, View, ShadowInfo, HairStrandsUniformBuffer); break;
-			case 3: BindShaderShaders<FShadowProjectionNoTransformVS, TShadowProjectionPS<3, false, false, false, true> >(RHICmdList, GraphicsPSOInit, ViewIndex, View, ShadowInfo, HairStrandsUniformBuffer); break;
-			case 4: BindShaderShaders<FShadowProjectionNoTransformVS, TShadowProjectionPS<4, false, false, false, true> >(RHICmdList, GraphicsPSOInit, ViewIndex, View, ShadowInfo, HairStrandsUniformBuffer); break;
-			case 5: BindShaderShaders<FShadowProjectionNoTransformVS, TShadowProjectionPS<5, false, false, false, true> >(RHICmdList, GraphicsPSOInit, ViewIndex, View, ShadowInfo, HairStrandsUniformBuffer); break;
+			case 1: BindShaderShaders<FShadowProjectionNoTransformVS, TShadowProjectionPS<1, false, false, false, true> >(RHICmdList, GraphicsPSOInit, ViewIndex, View, ShadowInfo, StencilRef, HairStrandsUniformBuffer); break;
+			case 2: BindShaderShaders<FShadowProjectionNoTransformVS, TShadowProjectionPS<2, false, false, false, true> >(RHICmdList, GraphicsPSOInit, ViewIndex, View, ShadowInfo, StencilRef, HairStrandsUniformBuffer); break;
+			case 3: BindShaderShaders<FShadowProjectionNoTransformVS, TShadowProjectionPS<3, false, false, false, true> >(RHICmdList, GraphicsPSOInit, ViewIndex, View, ShadowInfo, StencilRef, HairStrandsUniformBuffer); break;
+			case 4: BindShaderShaders<FShadowProjectionNoTransformVS, TShadowProjectionPS<4, false, false, false, true> >(RHICmdList, GraphicsPSOInit, ViewIndex, View, ShadowInfo, StencilRef, HairStrandsUniformBuffer); break;
+			case 5: BindShaderShaders<FShadowProjectionNoTransformVS, TShadowProjectionPS<5, false, false, false, true> >(RHICmdList, GraphicsPSOInit, ViewIndex, View, ShadowInfo, StencilRef, HairStrandsUniformBuffer); break;
 			default:
 				check(0);
 			}
@@ -493,11 +493,11 @@ static void BindShadowProjectionShaders(int32 Quality, FRHICommandList& RHICmdLi
 		{
 			switch (Quality)
 			{
-			case 1: BindShaderShaders<FShadowVolumeBoundProjectionVS, TShadowProjectionPS<1, false, false, false, true> >(RHICmdList, GraphicsPSOInit, ViewIndex, View, ShadowInfo, HairStrandsUniformBuffer); break;
-			case 2: BindShaderShaders<FShadowVolumeBoundProjectionVS, TShadowProjectionPS<2, false, false, false, true> >(RHICmdList, GraphicsPSOInit, ViewIndex, View, ShadowInfo, HairStrandsUniformBuffer); break;
-			case 3: BindShaderShaders<FShadowVolumeBoundProjectionVS, TShadowProjectionPS<3, false, false, false, true> >(RHICmdList, GraphicsPSOInit, ViewIndex, View, ShadowInfo, HairStrandsUniformBuffer); break;
-			case 4: BindShaderShaders<FShadowVolumeBoundProjectionVS, TShadowProjectionPS<4, false, false, false, true> >(RHICmdList, GraphicsPSOInit, ViewIndex, View, ShadowInfo, HairStrandsUniformBuffer); break;
-			case 5: BindShaderShaders<FShadowVolumeBoundProjectionVS, TShadowProjectionPS<5, false, false, false, true> >(RHICmdList, GraphicsPSOInit, ViewIndex, View, ShadowInfo, HairStrandsUniformBuffer); break;
+			case 1: BindShaderShaders<FShadowVolumeBoundProjectionVS, TShadowProjectionPS<1, false, false, false, true> >(RHICmdList, GraphicsPSOInit, ViewIndex, View, ShadowInfo, StencilRef, HairStrandsUniformBuffer); break;
+			case 2: BindShaderShaders<FShadowVolumeBoundProjectionVS, TShadowProjectionPS<2, false, false, false, true> >(RHICmdList, GraphicsPSOInit, ViewIndex, View, ShadowInfo, StencilRef, HairStrandsUniformBuffer); break;
+			case 3: BindShaderShaders<FShadowVolumeBoundProjectionVS, TShadowProjectionPS<3, false, false, false, true> >(RHICmdList, GraphicsPSOInit, ViewIndex, View, ShadowInfo, StencilRef, HairStrandsUniformBuffer); break;
+			case 4: BindShaderShaders<FShadowVolumeBoundProjectionVS, TShadowProjectionPS<4, false, false, false, true> >(RHICmdList, GraphicsPSOInit, ViewIndex, View, ShadowInfo, StencilRef, HairStrandsUniformBuffer); break;
+			case 5: BindShaderShaders<FShadowVolumeBoundProjectionVS, TShadowProjectionPS<5, false, false, false, true> >(RHICmdList, GraphicsPSOInit, ViewIndex, View, ShadowInfo, StencilRef, HairStrandsUniformBuffer); break;
 			default:
 				check(0);
 			}
@@ -509,11 +509,11 @@ static void BindShadowProjectionShaders(int32 Quality, FRHICommandList& RHICmdLi
 	{
 		switch (Quality)
 		{
-		case 1: BindShaderShaders<FShadowVolumeBoundProjectionVS, TShadowProjectionFromTranslucencyPS<1> >(RHICmdList, GraphicsPSOInit, ViewIndex, View, ShadowInfo); break;
-		case 2: BindShaderShaders<FShadowVolumeBoundProjectionVS, TShadowProjectionFromTranslucencyPS<2> >(RHICmdList, GraphicsPSOInit, ViewIndex, View, ShadowInfo); break;
-		case 3: BindShaderShaders<FShadowVolumeBoundProjectionVS, TShadowProjectionFromTranslucencyPS<3> >(RHICmdList, GraphicsPSOInit, ViewIndex, View, ShadowInfo); break;
-		case 4: BindShaderShaders<FShadowVolumeBoundProjectionVS, TShadowProjectionFromTranslucencyPS<4> >(RHICmdList, GraphicsPSOInit, ViewIndex, View, ShadowInfo); break;
-		case 5: BindShaderShaders<FShadowVolumeBoundProjectionVS, TShadowProjectionFromTranslucencyPS<5> >(RHICmdList, GraphicsPSOInit, ViewIndex, View, ShadowInfo); break;
+		case 1: BindShaderShaders<FShadowVolumeBoundProjectionVS, TShadowProjectionFromTranslucencyPS<1> >(RHICmdList, GraphicsPSOInit, ViewIndex, View, ShadowInfo, StencilRef); break;
+		case 2: BindShaderShaders<FShadowVolumeBoundProjectionVS, TShadowProjectionFromTranslucencyPS<2> >(RHICmdList, GraphicsPSOInit, ViewIndex, View, ShadowInfo, StencilRef); break;
+		case 3: BindShaderShaders<FShadowVolumeBoundProjectionVS, TShadowProjectionFromTranslucencyPS<3> >(RHICmdList, GraphicsPSOInit, ViewIndex, View, ShadowInfo, StencilRef); break;
+		case 4: BindShaderShaders<FShadowVolumeBoundProjectionVS, TShadowProjectionFromTranslucencyPS<4> >(RHICmdList, GraphicsPSOInit, ViewIndex, View, ShadowInfo, StencilRef); break;
+		case 5: BindShaderShaders<FShadowVolumeBoundProjectionVS, TShadowProjectionFromTranslucencyPS<5> >(RHICmdList, GraphicsPSOInit, ViewIndex, View, ShadowInfo, StencilRef); break;
 		default:
 			check(0);
 		}
@@ -523,9 +523,9 @@ static void BindShadowProjectionShaders(int32 Quality, FRHICommandList& RHICmdLi
 		if (CVarFilterMethod.GetValueOnRenderThread() == 1)
 		{
 			if (ShadowInfo->CascadeSettings.FadePlaneLength > 0)
-				BindShaderShaders<FShadowProjectionNoTransformVS, TDirectionalPercentageCloserShadowProjectionPS<5, true> >(RHICmdList, GraphicsPSOInit, ViewIndex, View, ShadowInfo);
+				BindShaderShaders<FShadowProjectionNoTransformVS, TDirectionalPercentageCloserShadowProjectionPS<5, true> >(RHICmdList, GraphicsPSOInit, ViewIndex, View, ShadowInfo, StencilRef);
 			else
-				BindShaderShaders<FShadowProjectionNoTransformVS, TDirectionalPercentageCloserShadowProjectionPS<5, false> >(RHICmdList, GraphicsPSOInit, ViewIndex, View, ShadowInfo);
+				BindShaderShaders<FShadowProjectionNoTransformVS, TDirectionalPercentageCloserShadowProjectionPS<5, false> >(RHICmdList, GraphicsPSOInit, ViewIndex, View, ShadowInfo, StencilRef);
 		}
 		else if (ShadowInfo->CascadeSettings.FadePlaneLength > 0)
 		{
@@ -533,11 +533,11 @@ static void BindShadowProjectionShaders(int32 Quality, FRHICommandList& RHICmdLi
 			{
 				switch (Quality)
 				{
-				case 1: BindShaderShaders<FShadowProjectionNoTransformVS, TShadowProjectionPS<1, true, false, true> >(RHICmdList, GraphicsPSOInit, ViewIndex, View, ShadowInfo); break;
-				case 2: BindShaderShaders<FShadowProjectionNoTransformVS, TShadowProjectionPS<2, true, false, true> >(RHICmdList, GraphicsPSOInit, ViewIndex, View, ShadowInfo); break;
-				case 3: BindShaderShaders<FShadowProjectionNoTransformVS, TShadowProjectionPS<3, true, false, true> >(RHICmdList, GraphicsPSOInit, ViewIndex, View, ShadowInfo); break;
-				case 4: BindShaderShaders<FShadowProjectionNoTransformVS, TShadowProjectionPS<4, true, false, true> >(RHICmdList, GraphicsPSOInit, ViewIndex, View, ShadowInfo); break;
-				case 5: BindShaderShaders<FShadowProjectionNoTransformVS, TShadowProjectionPS<5, true, false, true> >(RHICmdList, GraphicsPSOInit, ViewIndex, View, ShadowInfo); break;
+				case 1: BindShaderShaders<FShadowProjectionNoTransformVS, TShadowProjectionPS<1, true, false, true> >(RHICmdList, GraphicsPSOInit, ViewIndex, View, ShadowInfo, StencilRef); break;
+				case 2: BindShaderShaders<FShadowProjectionNoTransformVS, TShadowProjectionPS<2, true, false, true> >(RHICmdList, GraphicsPSOInit, ViewIndex, View, ShadowInfo, StencilRef); break;
+				case 3: BindShaderShaders<FShadowProjectionNoTransformVS, TShadowProjectionPS<3, true, false, true> >(RHICmdList, GraphicsPSOInit, ViewIndex, View, ShadowInfo, StencilRef); break;
+				case 4: BindShaderShaders<FShadowProjectionNoTransformVS, TShadowProjectionPS<4, true, false, true> >(RHICmdList, GraphicsPSOInit, ViewIndex, View, ShadowInfo, StencilRef); break;
+				case 5: BindShaderShaders<FShadowProjectionNoTransformVS, TShadowProjectionPS<5, true, false, true> >(RHICmdList, GraphicsPSOInit, ViewIndex, View, ShadowInfo, StencilRef); break;
 				default:
 					check(0);
 				}
@@ -546,11 +546,11 @@ static void BindShadowProjectionShaders(int32 Quality, FRHICommandList& RHICmdLi
 			{
 				switch (Quality)
 				{
-				case 1: BindShaderShaders<FShadowProjectionNoTransformVS, TShadowProjectionPS<1, true> >(RHICmdList, GraphicsPSOInit, ViewIndex, View, ShadowInfo); break;
-				case 2: BindShaderShaders<FShadowProjectionNoTransformVS, TShadowProjectionPS<2, true> >(RHICmdList, GraphicsPSOInit, ViewIndex, View, ShadowInfo); break;
-				case 3: BindShaderShaders<FShadowProjectionNoTransformVS, TShadowProjectionPS<3, true> >(RHICmdList, GraphicsPSOInit, ViewIndex, View, ShadowInfo); break;
-				case 4: BindShaderShaders<FShadowProjectionNoTransformVS, TShadowProjectionPS<4, true> >(RHICmdList, GraphicsPSOInit, ViewIndex, View, ShadowInfo); break;
-				case 5: BindShaderShaders<FShadowProjectionNoTransformVS, TShadowProjectionPS<5, true> >(RHICmdList, GraphicsPSOInit, ViewIndex, View, ShadowInfo); break;
+				case 1: BindShaderShaders<FShadowProjectionNoTransformVS, TShadowProjectionPS<1, true> >(RHICmdList, GraphicsPSOInit, ViewIndex, View, ShadowInfo, StencilRef); break;
+				case 2: BindShaderShaders<FShadowProjectionNoTransformVS, TShadowProjectionPS<2, true> >(RHICmdList, GraphicsPSOInit, ViewIndex, View, ShadowInfo, StencilRef); break;
+				case 3: BindShaderShaders<FShadowProjectionNoTransformVS, TShadowProjectionPS<3, true> >(RHICmdList, GraphicsPSOInit, ViewIndex, View, ShadowInfo, StencilRef); break;
+				case 4: BindShaderShaders<FShadowProjectionNoTransformVS, TShadowProjectionPS<4, true> >(RHICmdList, GraphicsPSOInit, ViewIndex, View, ShadowInfo, StencilRef); break;
+				case 5: BindShaderShaders<FShadowProjectionNoTransformVS, TShadowProjectionPS<5, true> >(RHICmdList, GraphicsPSOInit, ViewIndex, View, ShadowInfo, StencilRef); break;
 				default:
 					check(0);
 				}
@@ -562,11 +562,11 @@ static void BindShadowProjectionShaders(int32 Quality, FRHICommandList& RHICmdLi
 			{
 				switch (Quality)
 				{
-				case 1: BindShaderShaders<FShadowProjectionNoTransformVS, TShadowProjectionPS<1, false, false, true> >(RHICmdList, GraphicsPSOInit, ViewIndex, View, ShadowInfo); break;
-				case 2: BindShaderShaders<FShadowProjectionNoTransformVS, TShadowProjectionPS<2, false, false, true> >(RHICmdList, GraphicsPSOInit, ViewIndex, View, ShadowInfo); break;
-				case 3: BindShaderShaders<FShadowProjectionNoTransformVS, TShadowProjectionPS<3, false, false, true> >(RHICmdList, GraphicsPSOInit, ViewIndex, View, ShadowInfo); break;
-				case 4: BindShaderShaders<FShadowProjectionNoTransformVS, TShadowProjectionPS<4, false, false, true> >(RHICmdList, GraphicsPSOInit, ViewIndex, View, ShadowInfo); break;
-				case 5: BindShaderShaders<FShadowProjectionNoTransformVS, TShadowProjectionPS<5, false, false, true> >(RHICmdList, GraphicsPSOInit, ViewIndex, View, ShadowInfo); break;
+				case 1: BindShaderShaders<FShadowProjectionNoTransformVS, TShadowProjectionPS<1, false, false, true> >(RHICmdList, GraphicsPSOInit, ViewIndex, View, ShadowInfo, StencilRef); break;
+				case 2: BindShaderShaders<FShadowProjectionNoTransformVS, TShadowProjectionPS<2, false, false, true> >(RHICmdList, GraphicsPSOInit, ViewIndex, View, ShadowInfo, StencilRef); break;
+				case 3: BindShaderShaders<FShadowProjectionNoTransformVS, TShadowProjectionPS<3, false, false, true> >(RHICmdList, GraphicsPSOInit, ViewIndex, View, ShadowInfo, StencilRef); break;
+				case 4: BindShaderShaders<FShadowProjectionNoTransformVS, TShadowProjectionPS<4, false, false, true> >(RHICmdList, GraphicsPSOInit, ViewIndex, View, ShadowInfo, StencilRef); break;
+				case 5: BindShaderShaders<FShadowProjectionNoTransformVS, TShadowProjectionPS<5, false, false, true> >(RHICmdList, GraphicsPSOInit, ViewIndex, View, ShadowInfo, StencilRef); break;
 				default:
 					check(0);
 				}
@@ -575,11 +575,11 @@ static void BindShadowProjectionShaders(int32 Quality, FRHICommandList& RHICmdLi
 			{ 
 				switch (Quality)
 				{
-				case 1: BindShaderShaders<FShadowProjectionNoTransformVS, TShadowProjectionPS<1, false> >(RHICmdList, GraphicsPSOInit, ViewIndex, View, ShadowInfo); break;
-				case 2: BindShaderShaders<FShadowProjectionNoTransformVS, TShadowProjectionPS<2, false> >(RHICmdList, GraphicsPSOInit, ViewIndex, View, ShadowInfo); break;
-				case 3: BindShaderShaders<FShadowProjectionNoTransformVS, TShadowProjectionPS<3, false> >(RHICmdList, GraphicsPSOInit, ViewIndex, View, ShadowInfo); break;
-				case 4: BindShaderShaders<FShadowProjectionNoTransformVS, TShadowProjectionPS<4, false> >(RHICmdList, GraphicsPSOInit, ViewIndex, View, ShadowInfo); break;
-				case 5: BindShaderShaders<FShadowProjectionNoTransformVS, TShadowProjectionPS<5, false> >(RHICmdList, GraphicsPSOInit, ViewIndex, View, ShadowInfo); break;
+				case 1: BindShaderShaders<FShadowProjectionNoTransformVS, TShadowProjectionPS<1, false> >(RHICmdList, GraphicsPSOInit, ViewIndex, View, ShadowInfo, StencilRef); break;
+				case 2: BindShaderShaders<FShadowProjectionNoTransformVS, TShadowProjectionPS<2, false> >(RHICmdList, GraphicsPSOInit, ViewIndex, View, ShadowInfo, StencilRef); break;
+				case 3: BindShaderShaders<FShadowProjectionNoTransformVS, TShadowProjectionPS<3, false> >(RHICmdList, GraphicsPSOInit, ViewIndex, View, ShadowInfo, StencilRef); break;
+				case 4: BindShaderShaders<FShadowProjectionNoTransformVS, TShadowProjectionPS<4, false> >(RHICmdList, GraphicsPSOInit, ViewIndex, View, ShadowInfo, StencilRef); break;
+				case 5: BindShaderShaders<FShadowProjectionNoTransformVS, TShadowProjectionPS<5, false> >(RHICmdList, GraphicsPSOInit, ViewIndex, View, ShadowInfo, StencilRef); break;
 				default:
 					check(0);
 				}
@@ -592,11 +592,11 @@ static void BindShadowProjectionShaders(int32 Quality, FRHICommandList& RHICmdLi
 		{
 			switch (Quality)
 			{
-			case 1: BindShaderShaders<FShadowVolumeBoundProjectionVS, TModulatedShadowProjection<1> >(RHICmdList, GraphicsPSOInit, ViewIndex, View, ShadowInfo); break;
-			case 2: BindShaderShaders<FShadowVolumeBoundProjectionVS, TModulatedShadowProjection<2> >(RHICmdList, GraphicsPSOInit, ViewIndex, View, ShadowInfo); break;
-			case 3: BindShaderShaders<FShadowVolumeBoundProjectionVS, TModulatedShadowProjection<3> >(RHICmdList, GraphicsPSOInit, ViewIndex, View, ShadowInfo); break;
-			case 4: BindShaderShaders<FShadowVolumeBoundProjectionVS, TModulatedShadowProjection<4> >(RHICmdList, GraphicsPSOInit, ViewIndex, View, ShadowInfo); break;
-			case 5: BindShaderShaders<FShadowVolumeBoundProjectionVS, TModulatedShadowProjection<5> >(RHICmdList, GraphicsPSOInit, ViewIndex, View, ShadowInfo); break;
+			case 1: BindShaderShaders<FShadowVolumeBoundProjectionVS, TModulatedShadowProjection<1> >(RHICmdList, GraphicsPSOInit, ViewIndex, View, ShadowInfo, StencilRef); break;
+			case 2: BindShaderShaders<FShadowVolumeBoundProjectionVS, TModulatedShadowProjection<2> >(RHICmdList, GraphicsPSOInit, ViewIndex, View, ShadowInfo, StencilRef); break;
+			case 3: BindShaderShaders<FShadowVolumeBoundProjectionVS, TModulatedShadowProjection<3> >(RHICmdList, GraphicsPSOInit, ViewIndex, View, ShadowInfo, StencilRef); break;
+			case 4: BindShaderShaders<FShadowVolumeBoundProjectionVS, TModulatedShadowProjection<4> >(RHICmdList, GraphicsPSOInit, ViewIndex, View, ShadowInfo, StencilRef); break;
+			case 5: BindShaderShaders<FShadowVolumeBoundProjectionVS, TModulatedShadowProjection<5> >(RHICmdList, GraphicsPSOInit, ViewIndex, View, ShadowInfo, StencilRef); break;
 			default:
 				check(0);
 			}
@@ -605,11 +605,11 @@ static void BindShadowProjectionShaders(int32 Quality, FRHICommandList& RHICmdLi
 		{
 			switch (Quality)
 			{
-			case 1: BindShaderShaders<FShadowVolumeBoundProjectionVS, TShadowProjectionPS<1, false, false, true> >(RHICmdList, GraphicsPSOInit, ViewIndex, View, ShadowInfo); break;
-			case 2: BindShaderShaders<FShadowVolumeBoundProjectionVS, TShadowProjectionPS<2, false, false, true> >(RHICmdList, GraphicsPSOInit, ViewIndex, View, ShadowInfo); break;
-			case 3: BindShaderShaders<FShadowVolumeBoundProjectionVS, TShadowProjectionPS<3, false, false, true> >(RHICmdList, GraphicsPSOInit, ViewIndex, View, ShadowInfo); break;
-			case 4: BindShaderShaders<FShadowVolumeBoundProjectionVS, TShadowProjectionPS<4, false, false, true> >(RHICmdList, GraphicsPSOInit, ViewIndex, View, ShadowInfo); break;
-			case 5: BindShaderShaders<FShadowVolumeBoundProjectionVS, TShadowProjectionPS<5, false, false, true> >(RHICmdList, GraphicsPSOInit, ViewIndex, View, ShadowInfo); break;
+			case 1: BindShaderShaders<FShadowVolumeBoundProjectionVS, TShadowProjectionPS<1, false, false, true> >(RHICmdList, GraphicsPSOInit, ViewIndex, View, ShadowInfo, StencilRef); break;
+			case 2: BindShaderShaders<FShadowVolumeBoundProjectionVS, TShadowProjectionPS<2, false, false, true> >(RHICmdList, GraphicsPSOInit, ViewIndex, View, ShadowInfo, StencilRef); break;
+			case 3: BindShaderShaders<FShadowVolumeBoundProjectionVS, TShadowProjectionPS<3, false, false, true> >(RHICmdList, GraphicsPSOInit, ViewIndex, View, ShadowInfo, StencilRef); break;
+			case 4: BindShaderShaders<FShadowVolumeBoundProjectionVS, TShadowProjectionPS<4, false, false, true> >(RHICmdList, GraphicsPSOInit, ViewIndex, View, ShadowInfo, StencilRef); break;
+			case 5: BindShaderShaders<FShadowVolumeBoundProjectionVS, TShadowProjectionPS<5, false, false, true> >(RHICmdList, GraphicsPSOInit, ViewIndex, View, ShadowInfo, StencilRef); break;
 			default:
 				check(0);
 			}
@@ -618,17 +618,17 @@ static void BindShadowProjectionShaders(int32 Quality, FRHICommandList& RHICmdLi
 		{
 			if (CVarFilterMethod.GetValueOnRenderThread() == 1 && ShadowInfo->GetLightSceneInfo().Proxy->GetLightType() == LightType_Spot)
 			{
-				BindShaderShaders<FShadowVolumeBoundProjectionVS, TSpotPercentageCloserShadowProjectionPS<5, false> >(RHICmdList, GraphicsPSOInit, ViewIndex, View, ShadowInfo);
+				BindShaderShaders<FShadowVolumeBoundProjectionVS, TSpotPercentageCloserShadowProjectionPS<5, false> >(RHICmdList, GraphicsPSOInit, ViewIndex, View, ShadowInfo, StencilRef);
 			}
 			else
 			{
 				switch (Quality)
 				{
-				case 1: BindShaderShaders<FShadowVolumeBoundProjectionVS, TShadowProjectionPS<1, false> >(RHICmdList, GraphicsPSOInit, ViewIndex, View, ShadowInfo); break;
-				case 2: BindShaderShaders<FShadowVolumeBoundProjectionVS, TShadowProjectionPS<2, false> >(RHICmdList, GraphicsPSOInit, ViewIndex, View, ShadowInfo); break;
-				case 3: BindShaderShaders<FShadowVolumeBoundProjectionVS, TShadowProjectionPS<3, false> >(RHICmdList, GraphicsPSOInit, ViewIndex, View, ShadowInfo); break;
-				case 4: BindShaderShaders<FShadowVolumeBoundProjectionVS, TShadowProjectionPS<4, false> >(RHICmdList, GraphicsPSOInit, ViewIndex, View, ShadowInfo); break;
-				case 5: BindShaderShaders<FShadowVolumeBoundProjectionVS, TShadowProjectionPS<5, false> >(RHICmdList, GraphicsPSOInit, ViewIndex, View, ShadowInfo); break;
+				case 1: BindShaderShaders<FShadowVolumeBoundProjectionVS, TShadowProjectionPS<1, false> >(RHICmdList, GraphicsPSOInit, ViewIndex, View, ShadowInfo, StencilRef); break;
+				case 2: BindShaderShaders<FShadowVolumeBoundProjectionVS, TShadowProjectionPS<2, false> >(RHICmdList, GraphicsPSOInit, ViewIndex, View, ShadowInfo, StencilRef); break;
+				case 3: BindShaderShaders<FShadowVolumeBoundProjectionVS, TShadowProjectionPS<3, false> >(RHICmdList, GraphicsPSOInit, ViewIndex, View, ShadowInfo, StencilRef); break;
+				case 4: BindShaderShaders<FShadowVolumeBoundProjectionVS, TShadowProjectionPS<4, false> >(RHICmdList, GraphicsPSOInit, ViewIndex, View, ShadowInfo, StencilRef); break;
+				case 5: BindShaderShaders<FShadowVolumeBoundProjectionVS, TShadowProjectionPS<5, false> >(RHICmdList, GraphicsPSOInit, ViewIndex, View, ShadowInfo, StencilRef); break;
 				default:
 					check(0);
 				}
@@ -1031,6 +1031,38 @@ void FProjectedShadowInfo::SetupProjectionStencilMask(
 	}
 }
 
+// Mark hair-strands pixels encompassed into the rasterized volumes. 
+// * Mark 1: for hair pixels encompassed into the volume
+// * Mark 2: for pixels touching the volume
+// This affects only the hair-only depth stencil texture.
+void FProjectedShadowInfo::SetupProjectionStencilMaskForHair(FRHICommandList& RHICmdList, const FViewInfo* View) const
+{
+	check(!IsWholeSceneDirectionalShadow());
+
+	FGraphicsPipelineStateInitializer GraphicsPSOInit;
+	RHICmdList.ApplyCachedRenderTargets(GraphicsPSOInit);
+
+	TShaderRef<FShadowVolumeBoundProjectionVS> VertexShader = View->ShaderMap->GetShader<FShadowVolumeBoundProjectionVS>();
+	GraphicsPSOInit.RasterizerState = TStaticRasterizerState<FM_Solid, CM_None>::GetRHI();
+	GraphicsPSOInit.bDepthBounds = false;
+	GraphicsPSOInit.DepthStencilState =
+		TStaticDepthStencilState<
+		false, CF_GreaterEqual,
+		true, CF_Always, SO_Keep, SO_Keep, SO_Increment,
+		true, CF_Always, SO_Keep, SO_Keep, SO_Increment,
+		0xff, 0xff>::GetRHI();
+	GraphicsPSOInit.BlendState = TStaticBlendState<CW_RED, BO_Add, BF_SourceAlpha, BF_InverseSourceAlpha>::GetRHI();
+	GraphicsPSOInit.PrimitiveType = PT_TriangleList;
+	GraphicsPSOInit.BoundShaderState.VertexDeclarationRHI = GetVertexDeclarationFVector4();
+	GraphicsPSOInit.BoundShaderState.VertexShaderRHI = VertexShader.GetVertexShader();
+	GraphicsPSOInit.BoundShaderState.PixelShaderRHI = nullptr;
+	SetGraphicsPipelineState(RHICmdList, GraphicsPSOInit, 0);
+	VertexShader->SetParameters(RHICmdList, *View, this, EShadowProjectionVertexShaderFlags::DrawingFrustum);
+
+	RHICmdList.SetStreamSource(0, GFrustumVertexBuffer.VertexBufferRHI, 0);
+	RHICmdList.DrawIndexedPrimitive(GCubeIndexBuffer.IndexBufferRHI, 0, 0, 8, 0, 12, 1);
+}
+
 BEGIN_SHADER_PARAMETER_STRUCT(FShadowProjectionPassParameters, )
 	SHADER_PARAMETER_STRUCT_INCLUDE(FViewShaderParameters, View)
 	SHADER_PARAMETER_RDG_UNIFORM_BUFFER(FHairStrandsViewUniformParameters, HairStrands)
@@ -1215,12 +1247,20 @@ void FProjectedShadowInfo::RenderProjectionInternal(
 	const bool bSubPixelSupport = HairStrandsUniformBuffer != nullptr;// HairStrands::HasViewHairStrandsData(*View);
 	const bool bStencilTestEnabled = !bSubPixelSupport && GShadowStencilCulling;
 	const bool bDepthBoundsTestEnabled = IsWholeSceneDirectionalShadow() && GSupportsDepthBoundsTest && CVarCSMDepthBoundsTest.GetValueOnRenderThread() != 0 && !bSubPixelSupport;
+	const uint32 StencilRef = bSubPixelSupport ? 1u : 0u;
 
 	if (!bDepthBoundsTestEnabled && bStencilTestEnabled)
 	{
 		SetupProjectionStencilMask(RHICmdList, View, ViewIndex, SceneRender, FrustumVertices, bMobileModulatedProjections, bCameraInsideShadowFrustum, InstanceCullingDrawParams);
 	}
 
+	// Mark stencil so that only hair pixel within volume bound will be affected by the pre-shadow mask
+	if (bSubPixelSupport && !IsWholeSceneDirectionalShadow())
+	{
+		DrawClearQuad(RHICmdList, false, FLinearColor::Transparent, false, 0, true, 0);
+		SetupProjectionStencilMaskForHair(RHICmdList, View);
+	}
+
 	// solid rasterization w/ back-face culling.
 	GraphicsPSOInit.RasterizerState = (View->bReverseCulling || IsWholeSceneDirectionalShadow()) ? TStaticRasterizerState<FM_Solid, CM_CCW>::GetRHI() : TStaticRasterizerState<FM_Solid, CM_CW>::GetRHI();
 
@@ -1271,7 +1311,14 @@ void FProjectedShadowInfo::RenderProjectionInternal(
 	}
 	else
 	{
-		GraphicsPSOInit.DepthStencilState = TStaticDepthStencilState<false, CF_DepthFartherOrEqual>::GetRHI();
+		if (bSubPixelSupport)
+		{
+			GraphicsPSOInit.DepthStencilState = TStaticDepthStencilState<false, CF_DepthFartherOrEqual, true, CF_Equal, SO_Keep, SO_Keep, SO_Keep, true, CF_Equal, SO_Keep, SO_Keep, SO_Keep, 0xFF, 0xFF>::GetRHI();
+		}
+		else
+		{
+			GraphicsPSOInit.DepthStencilState = TStaticDepthStencilState<false, CF_DepthFartherOrEqual>::GetRHI();
+		}
 	}
 
 	GraphicsPSOInit.BlendState = GetBlendStateForProjection(bProjectingForForwardShading, bMobileModulatedProjections);
@@ -1311,7 +1358,7 @@ void FProjectedShadowInfo::RenderProjectionInternal(
 		}
 
 		GraphicsPSOInit.BoundShaderState.VertexDeclarationRHI = GetVertexDeclarationFVector4();
-		BindShadowProjectionShaders(LocalQuality, RHICmdList, GraphicsPSOInit, ViewIndex, *View, this, bMobileModulatedProjections, HairStrandsUniformBuffer);
+		BindShadowProjectionShaders(LocalQuality, RHICmdList, GraphicsPSOInit, ViewIndex, *View, this, StencilRef, bMobileModulatedProjections, HairStrandsUniformBuffer);
 
 		if (bDepthBoundsTestEnabled)
 		{
diff --git a/Engine/Source/Runtime/Renderer/Private/ShadowRendering.h b/Engine/Source/Runtime/Renderer/Private/ShadowRendering.h
--- a/Engine/Source/Runtime/Renderer/Private/ShadowRendering.h
+++ b/Engine/Source/Runtime/Renderer/Private/ShadowRendering.h
@@ -873,6 +873,7 @@ private:
 		bool bCameraInsideShadowFrustum,
 		const FInstanceCullingDrawParams& InstanceCullingDrawParams) const;
 
+	void SetupProjectionStencilMaskForHair(FRHICommandList& RHICmdList, const FViewInfo* View) const;
 
 	FORCEINLINE bool TestPrimitiveFarCascadeConditions(bool bPrimitiveCastsFarShadow, const FBoxSphereBounds& Bounds) const;
 
