Add epsilon to fix some flickering lit pixels in shadowed regions due to numeric precision in the SMRT sample comparison.
Add a cvar for the clipmap level depth scale, primarily for debug purposes.

diff --git a/Engine/Shaders/Private/VirtualShadowMaps/SMRTTemplate.ush b/Engine/Shaders/Private/VirtualShadowMaps/SMRTTemplate.ush
--- a/Engine/Shaders/Private/VirtualShadowMaps/SMRTTemplate.ush
+++ b/Engine/Shaders/Private/VirtualShadowMaps/SMRTTemplate.ush
@@ -58,7 +58,10 @@ FSMRTResult SMRTRayCast(
 			}
 			else
 			{
-				const float CompareTolerance = abs(Sample.CompareToleranceScale * SampleCompare);
+				// Add a small relative error to the comparison to avoid missing surfaces due to numeric precision issues
+				// Without this there are occasionally flickering fireflies in fully shadowed regions with SMRT
+				const float EpsScale = 1.02f;
+				const float CompareTolerance = abs(Sample.CompareToleranceScale * SampleCompare) * EpsScale;
 
 				bool bBehind = (SampleDepth - ReferenceDepth) > (2.0f * CompareTolerance);
 				float DepthSample = bBehind ? DepthHistory : SampleDepth;
diff --git a/Engine/Source/Runtime/Renderer/Private/VirtualShadowMaps/VirtualShadowMapClipmap.cpp b/Engine/Source/Runtime/Renderer/Private/VirtualShadowMaps/VirtualShadowMapClipmap.cpp
--- a/Engine/Source/Runtime/Renderer/Private/VirtualShadowMaps/VirtualShadowMapClipmap.cpp
+++ b/Engine/Source/Runtime/Renderer/Private/VirtualShadowMaps/VirtualShadowMapClipmap.cpp
@@ -44,6 +44,13 @@ TAutoConsoleVariable<int32> CVarVirtualShadowMapClipmapLastCoarseLevel(
 	ECVF_Scalability | ECVF_RenderThreadSafe
 );
 
+TAutoConsoleVariable<float> CVarVirtualShadowMapClipmapZRangeScale(
+	TEXT("r.Shadow.Virtual.Clipmap.ZRangeScale"),
+	1000.0f,
+	TEXT("Scale of the clipmap level depth range relative to the radius. Should generally be at least 10 or it will result in excessive cache invalidations."),
+	ECVF_RenderThreadSafe
+);
+
 // "Virtual" clipmap level to clipmap radius
 // NOTE: This is the radius of around the clipmap origin that this level must cover
 // The actual clipmap dimensions will be larger due to snapping and other accomodations
@@ -146,9 +153,9 @@ FVirtualShadowMapClipmap::FVirtualShadowMapClipmap(
 		// This also better accomodates SMRT where we want to avoid stepping outside of the Z bounds of a given clipmap
 		// NOTE: It's tempting to use a single global Z range for the entire clipmap (which avoids some SMRT overhead too)
 		// but this can cause precision issues with cached pages very near the camera.
-		const float ViewRadiusZMultiplier = 1000.0f;
+		const float ViewRadiusZScale = CVarVirtualShadowMapClipmapZRangeScale.GetValueOnRenderThread();
 
-		float ViewRadiusZ = RawLevelRadius * ViewRadiusZMultiplier;
+		float ViewRadiusZ = RawLevelRadius * ViewRadiusZScale;
 		float ViewCenterDeltaZ = 0.0f;
 
 		if (CacheEntry)
