Fix DF Shadows not working for Point/Spot lights after CL 18393307.

diff --git a/Engine/Source/Runtime/Renderer/Private/ShadowSetup.cpp b/Engine/Source/Runtime/Renderer/Private/ShadowSetup.cpp
--- a/Engine/Source/Runtime/Renderer/Private/ShadowSetup.cpp
+++ b/Engine/Source/Runtime/Renderer/Private/ShadowSetup.cpp
@@ -4975,7 +4975,11 @@ void FSceneRenderer::AllocateShadowDepthTargets(FRHICommandListImmediate& RHICmd
 			}
 			// If uncached and no primitives, skip allocations etc
 			// Note: cached PreShadows, for some reason, has CacheMode == SDCM_Uncached so need to skip them here.
-			else if ((ProjectedShadowInfo->CacheMode == SDCM_Uncached && !ProjectedShadowInfo->bPreShadow) && !ProjectedShadowInfo->HasSubjectPrims() && !ProjectedShadowInfo->IsWholeSceneDirectionalShadow())
+			// Ray traced shadows use the GPU managed distance field object buffers, no CPU culling should be used
+			else if (
+				(ProjectedShadowInfo->CacheMode == SDCM_Uncached && !ProjectedShadowInfo->bPreShadow) 
+				&& (!ProjectedShadowInfo->HasSubjectPrims() && !ProjectedShadowInfo->bRayTracedDistanceField) 
+				&& !ProjectedShadowInfo->IsWholeSceneDirectionalShadow())
 			{
 				bShadowIsVisible = false;
 			}
