Set InstanceContributionToHitGroupIndex based on SceneInitializer

diff --git a/Engine/Source/Runtime/Renderer/Private/RayTracing/RayTracingInstanceBufferUtil.cpp b/Engine/Source/Runtime/Renderer/Private/RayTracing/RayTracingInstanceBufferUtil.cpp
--- a/Engine/Source/Runtime/Renderer/Private/RayTracing/RayTracingInstanceBufferUtil.cpp
+++ b/Engine/Source/Runtime/Renderer/Private/RayTracing/RayTracingInstanceBufferUtil.cpp
@@ -180,7 +180,7 @@ void FillRayTracingInstanceUploadBuffer(
 				FRayTracingInstanceDescriptorInput& InstanceDesc = OutInstanceUploadData[BaseDescriptorIndex + TransformIndex];
 
 				InstanceDesc.InstanceMaskAndFlags = SceneInstance.Mask | ((uint32)SceneInstance.Flags << 8);
-				InstanceDesc.InstanceContributionToHitGroupIndex = SceneInitializer.SegmentPrefixSum[SceneInstanceIndex] * RAY_TRACING_NUM_SHADER_SLOTS;
+				InstanceDesc.InstanceContributionToHitGroupIndex = SceneInitializer.SegmentPrefixSum[SceneInstanceIndex] * SceneInitializer.ShaderSlotsPerGeometrySegment;
 				InstanceDesc.InstanceId = bUseUniqueUserData ? SceneInstance.UserData[TransformIndex] : SceneInstance.DefaultUserData;
 				InstanceDesc.OutputDescriptorIndex = BaseInstanceIndex + TransformIndex;
 
