Implements r.TSR.History.UpdateQuality

diff --git a/Engine/Config/BaseScalability.ini b/Engine/Config/BaseScalability.ini
--- a/Engine/Config/BaseScalability.ini
+++ b/Engine/Config/BaseScalability.ini
@@ -39,32 +39,57 @@ PerfIndexValues_ResolutionQuality="50 71 87 100 100"
 [AntiAliasingQuality@0]
 r.FXAA.Quality=0
 r.TemporalAA.Quality=0
+r.TSR.History.R11G11B10=1
+r.TSR.History.ScreenPercentage=100
+r.TSR.History.UpdateQuality=0
+r.TSR.ShadingRejection.SpatialFilter=0
+r.TSR.Velocity.Extrapolation=0
+r.TSR.Velocity.HoleFill=0
 r.TSR.RejectionAntiAliasingQuality=0
-r.TSR.HistoryScreenPercentage=100
 
 [AntiAliasingQuality@1]
 r.FXAA.Quality=1
 r.TemporalAA.Quality=1
-r.TSR.RejectionAntiAliasingQuality=0
-r.TSR.HistoryScreenPercentage=100
+r.TSR.History.R11G11B10=1
+r.TSR.History.ScreenPercentage=100
+r.TSR.History.UpdateQuality=1
+r.TSR.ShadingRejection.SpatialFilter=0
+r.TSR.Velocity.Extrapolation=0
+r.TSR.Velocity.HoleFill=0
+r.TSR.RejectionAntiAliasingQuality=1
 
 [AntiAliasingQuality@2]
 r.FXAA.Quality=3
 r.TemporalAA.Quality=1
+r.TSR.History.R11G11B10=1
+r.TSR.History.ScreenPercentage=100
+r.TSR.History.UpdateQuality=2
+r.TSR.ShadingRejection.SpatialFilter=1
+r.TSR.Velocity.Extrapolation=0
+r.TSR.Velocity.HoleFill=1
 r.TSR.RejectionAntiAliasingQuality=1
-r.TSR.HistoryScreenPercentage=100
 
 [AntiAliasingQuality@3]
 r.FXAA.Quality=4
 r.TemporalAA.Quality=2
-r.TSR.RejectionAntiAliasingQuality=1
-r.TSR.HistoryScreenPercentage=100
+r.TSR.History.R11G11B10=1
+r.TSR.History.ScreenPercentage=100
+r.TSR.History.UpdateQuality=3
+r.TSR.ShadingRejection.SpatialFilter=2
+r.TSR.Velocity.Extrapolation=1
+r.TSR.Velocity.HoleFill=1
+r.TSR.RejectionAntiAliasingQuality=2
 
 [AntiAliasingQuality@Cine]
 r.FXAA.Quality=5
 r.TemporalAA.Quality=2
+r.TSR.History.R11G11B10=0
+r.TSR.History.ScreenPercentage=200
+r.TSR.History.UpdateQuality=3
+r.TSR.ShadingRejection.SpatialFilter=2
+r.TSR.Velocity.Extrapolation=1
+r.TSR.Velocity.HoleFill=1
 r.TSR.RejectionAntiAliasingQuality=2
-r.TSR.HistoryScreenPercentage=200
 
 
 ;-----------------------------------------------------------------------------------------------------------------
diff --git a/Engine/Shaders/Private/TemporalSuperResolution/TSRUpdateHistory.usf b/Engine/Shaders/Private/TemporalSuperResolution/TSRUpdateHistory.usf
--- a/Engine/Shaders/Private/TemporalSuperResolution/TSRUpdateHistory.usf
+++ b/Engine/Shaders/Private/TemporalSuperResolution/TSRUpdateHistory.usf
@@ -12,32 +12,54 @@
 
 //------------------------------------------------------- CONFIG
 
-#define CONFIG_CLAMP 1
-
-#define CONFIG_SAMPLES 5
-
-#define CONFIG_MANUAL_LDS_SPILL (CONFIG_COMPILE_FP16)
-
-#define CONFIG_LOW_FREQUENCY_DRIFT_REJECTION 1
+#if DIM_UPDATE_QUALITY == 0 // Low
+	#define CONFIG_SAMPLES 5
+	#define CONFIG_LOW_FREQUENCY_DRIFT_REJECTION 0
+	#define CONFIG_REJECTION_ANTI_ALIASING 0
+
+#elif DIM_UPDATE_QUALITY == 1 // Medium
+	#define CONFIG_SAMPLES 5
+	#define CONFIG_LOW_FREQUENCY_DRIFT_REJECTION 1
+	#define CONFIG_REJECTION_ANTI_ALIASING 1
+
+#elif DIM_UPDATE_QUALITY == 2 // High
+	#define CONFIG_SAMPLES 6
+	#define CONFIG_LOW_FREQUENCY_DRIFT_REJECTION 1
+	#define CONFIG_REJECTION_ANTI_ALIASING 1
+
+#elif DIM_UPDATE_QUALITY == 3 // Epic
+	#define CONFIG_SAMPLES 6
+	#define CONFIG_LOW_FREQUENCY_DRIFT_REJECTION 1
+	#define CONFIG_INTERPOLATE_VELOCITY 1
+	#define CONFIG_REJECTION_ANTI_ALIASING 1
 
-#define CONFIG_TRANSLUCENCY_REJECTION 1
-
-#define CONFIG_REJECTION_ANTI_ALIASING (DIM_REJECTION_ANTI_ALIASING)
-
-#define CONFIG_INTERPOLATE_VELOCITY 1
+#else
+	#error Unknown history update quality
+#endif
 
 #if DIM_SEPARATE_TRANSLUCENCY == 1
 	#define CONFIG_ACCUMULATE_TRANSLUCENCY_SEPARATLY 1
 	#define CONFIG_VFX_SURVIVAL 0
-	#define CONFIG_RESPONSIVE_STENCIL 1
 	#define CONFIG_REJECT_TRAMSLUCENCY_ON_RESPONSIVE_AA_ONLY 1
 #else
 	#define CONFIG_ACCUMULATE_TRANSLUCENCY_SEPARATLY 0
 	#define CONFIG_VFX_SURVIVAL 1
-	#define CONFIG_RESPONSIVE_STENCIL 1
 	#define CONFIG_REJECT_TRAMSLUCENCY_ON_RESPONSIVE_AA_ONLY 0
 #endif
 
+
+//------------------------------------------------------- DEFAULTS
+
+#define CONFIG_CLAMP 1
+
+#ifndef CONFIG_TRANSLUCENCY_REJECTION
+	#define CONFIG_TRANSLUCENCY_REJECTION 1
+#endif
+
+#define CONFIG_RESPONSIVE_STENCIL 1
+
+#define CONFIG_MANUAL_LDS_SPILL (CONFIG_COMPILE_FP16)
+
 #if CONFIG_COMPILE_FP16
 	// Take advantage of RDNA's v_pk_*_{uif}16 instructions
 	#define CONFIG_ENABLE_DUAL_PIXEL_VECTORIZATION 1
diff --git a/Engine/Source/Runtime/Renderer/Private/PostProcess/TemporalSuperResolution.cpp b/Engine/Source/Runtime/Renderer/Private/PostProcess/TemporalSuperResolution.cpp
--- a/Engine/Source/Runtime/Renderer/Private/PostProcess/TemporalSuperResolution.cpp
+++ b/Engine/Source/Runtime/Renderer/Private/PostProcess/TemporalSuperResolution.cpp
@@ -16,15 +16,20 @@ namespace
 {
 
 TAutoConsoleVariable<float> CVarTSRHistorySP(
-	TEXT("r.TSR.HistoryScreenPercentage"),
+	TEXT("r.TSR.History.ScreenPercentage"),
 	100.0f,
 	TEXT("Size of TSR's history."),
 	ECVF_Scalability | ECVF_RenderThreadSafe);
 
 TAutoConsoleVariable<int32> CVarTSRR11G11B10History(
-	TEXT("r.TSR.R11G11B10History"), 1,
+	TEXT("r.TSR.History.R11G11B10"), 1,
 	TEXT("Select the bitdepth of the history."),
-	ECVF_RenderThreadSafe);
+	ECVF_Scalability | ECVF_RenderThreadSafe);
+
+TAutoConsoleVariable<int32> CVarTSRHistoryUpdateQuality(
+	TEXT("r.TSR.History.UpdateQuality"), 3,
+	TEXT("Select the quality of the history update."),
+	ECVF_Scalability | ECVF_RenderThreadSafe);
 
 TAutoConsoleVariable<int32> CVarTSRHalfResShadingRejection(
 	TEXT("r.TSR.ShadingRejection.HalfRes"), 0,
@@ -37,7 +42,7 @@ TAutoConsoleVariable<int32> CVarTSRFilterShadingRejection(
 	TEXT(" 0: Disabled;\n")
 	TEXT(" 1: Spatial filter pass is run at lower resolution than CompareHistory pass (default);\n")
 	TEXT(" 2: Spatial filter pass is run CompareHistory pass resolution to improve stability."),
-	ECVF_RenderThreadSafe);
+	ECVF_Scalability | ECVF_RenderThreadSafe);
 
 TAutoConsoleVariable<int32> CVarTSRRejectionAntiAliasingQuality(
 	TEXT("r.TSR.RejectionAntiAliasingQuality"), 3,
@@ -73,7 +78,7 @@ TAutoConsoleVariable<float> CVarTSRWeightClampingPixelSpeed(
 TAutoConsoleVariable<float> CVarTSRVelocityExtrapolation(
 	TEXT("r.TSR.Velocity.Extrapolation"), 1.0f,
 	TEXT("Defines how much the velocity should be extrapolated on geometric discontinuities (Default = 1.0f)."),
-	ECVF_RenderThreadSafe);
+	ECVF_Scalability | ECVF_RenderThreadSafe);
 
 TAutoConsoleVariable<int32> CVarTSRVelocityHoleFill(
 	TEXT("r.TSR.Velocity.HoleFill"), 1,
@@ -423,10 +428,19 @@ class FTSRUpdateHistoryCS : public FTSRShader
 	DECLARE_GLOBAL_SHADER(FTSRUpdateHistoryCS);
 	SHADER_USE_PARAMETER_STRUCT(FTSRUpdateHistoryCS, FTSRShader);
 
-	class FRejectionAADim : SHADER_PERMUTATION_BOOL("DIM_REJECTION_ANTI_ALIASING");
+	enum class EQuality
+	{
+		Low,
+		Medium,
+		High,
+		Epic,
+		MAX
+	};
+
 	class FSeparateTranslucencyDim : SHADER_PERMUTATION_BOOL("DIM_SEPARATE_TRANSLUCENCY");
+	class FQualityDim : SHADER_PERMUTATION_ENUM_CLASS("DIM_UPDATE_QUALITY", EQuality);
 
-	using FPermutationDomain = TShaderPermutationDomain<FRejectionAADim, FSeparateTranslucencyDim>;
+	using FPermutationDomain = TShaderPermutationDomain<FSeparateTranslucencyDim, FQualityDim>;
 
 	BEGIN_SHADER_PARAMETER_STRUCT(FParameters, )
 		SHADER_PARAMETER_STRUCT_INCLUDE(FTSRCommonParameters, CommonParameters)
@@ -564,6 +578,8 @@ ITemporalUpscaler::FOutputs AddTemporalSuperResolutionPasses(
 	// Whether to use camera cut shader permutation or not.
 	bool bCameraCut = !InputHistory.IsValid() || View.bCameraCut;
 
+	FTSRUpdateHistoryCS::EQuality UpdateHistoryQuality = FTSRUpdateHistoryCS::EQuality(FMath::Clamp(CVarTSRHistoryUpdateQuality.GetValueOnRenderThread(), 0, int32(FTSRUpdateHistoryCS::EQuality::MAX) - 1));
+
 	bool bHalfResLowFrequency = CVarTSRHalfResShadingRejection.GetValueOnRenderThread() != 0;
 
 	bool bIsSeperateTranslucyTexturesValid = PassInputs.SeparateTranslucencyTextures != nullptr && PassInputs.SeparateTranslucencyTextures->IsColorValid();
@@ -572,7 +588,11 @@ ITemporalUpscaler::FOutputs AddTemporalSuperResolutionPasses(
 
 	bool bAccumulateSeparateTranslucency = bIsSeperateTranslucyTexturesValid && CVarTSRTranslucencySeparateTemporalAccumulation.GetValueOnRenderThread() != 0;
 
-	int32 RejectionAntiAliasingQuality = FMath::Clamp(CVarTSRRejectionAntiAliasingQuality.GetValueOnRenderThread(), 0, 2);
+	int32 RejectionAntiAliasingQuality = FMath::Clamp(CVarTSRRejectionAntiAliasingQuality.GetValueOnRenderThread(), 1, 2);
+	if (UpdateHistoryQuality == FTSRUpdateHistoryCS::EQuality::Low)
+	{
+		RejectionAntiAliasingQuality = 0; 
+	}
 
 	bool bHoleFillVelocity = CVarTSRVelocityHoleFill.GetValueOnRenderThread() != 0 && RejectionAntiAliasingQuality > 0;
 
@@ -1323,6 +1343,14 @@ ITemporalUpscaler::FOutputs AddTemporalSuperResolutionPasses(
 
 	// Update temporal history.
 	{
+		static const TCHAR* const kUpdateQualityNames[] = {
+			TEXT("Low"),
+			TEXT("Medium"),
+			TEXT("High"),
+			TEXT("Epic"),
+		};
+		static_assert(UE_ARRAY_COUNT(kUpdateQualityNames) == int32(FTSRUpdateHistoryCS::EQuality::MAX), "Fix me!");
+
 		FTSRUpdateHistoryCS::FParameters* PassParameters = GraphBuilder.AllocParameters<FTSRUpdateHistoryCS::FParameters>();
 		PassParameters->CommonParameters = CommonParameters;
 		PassParameters->InputSceneColorTexture = PassInputs.SceneColorTexture;
@@ -1373,15 +1401,15 @@ ITemporalUpscaler::FOutputs AddTemporalSuperResolutionPasses(
 		PassParameters->DebugOutput = CreateDebugUAV(HistoryExtent, TEXT("Debug.TSR.UpdateHistory"));
 
 		FTSRUpdateHistoryCS::FPermutationDomain PermutationVector;
-		PermutationVector.Set<FTSRUpdateHistoryCS::FRejectionAADim>(RejectionAntiAliasingQuality > 0);
 		PermutationVector.Set<FTSRUpdateHistoryCS::FSeparateTranslucencyDim>(bAccumulateSeparateTranslucency);
+		PermutationVector.Set<FTSRUpdateHistoryCS::FQualityDim>(UpdateHistoryQuality);
 
 		TShaderMapRef<FTSRUpdateHistoryCS> ComputeShader(View.ShaderMap, PermutationVector);
 		FComputeShaderUtils::AddPass(
 			GraphBuilder,
-			RDG_EVENT_NAME("TSR UpdateHistory(%s%s%s%s) %dx%d", 
+			RDG_EVENT_NAME("TSR UpdateHistory(Quality=%s %s%s%s) %dx%d",
+				kUpdateQualityNames[int32(PermutationVector.Get<FTSRUpdateHistoryCS::FQualityDim>())],
 				History.LowFrequency->Desc.Format == PF_FloatR11G11B10 ? TEXT("R11G11B10") : TEXT(""),
-				PermutationVector.Get<FTSRUpdateHistoryCS::FRejectionAADim>() ? TEXT(" RejectionAA") : TEXT(""),
 				PermutationVector.Get<FTSRUpdateHistoryCS::FSeparateTranslucencyDim>() ? TEXT(" SeparateTranslucency") : TEXT(""),
 				PassInputs.bGenerateOutputMip1 ? TEXT(" OutputMip1") : TEXT(""),
 				HistorySize.X, HistorySize.Y),
