Fixing broken RT light culling.

Generation of the RT culling grid depended on uploading lights' BoundingSphere value, defined natively as a VectorRegister type. Changes with LWC no longer makes this a compatible type to directly upload in a TResourceArray. Instead, the value must be type-converted to a supported format (eg: FVector4f).

diff --git a/Engine/Source/Runtime/Renderer/Private/RayTracing/RayTracingLighting.cpp b/Engine/Source/Runtime/Renderer/Private/RayTracing/RayTracingLighting.cpp
--- a/Engine/Source/Runtime/Renderer/Private/RayTracing/RayTracingLighting.cpp
+++ b/Engine/Source/Runtime/Renderer/Private/RayTracing/RayTracingLighting.cpp
@@ -108,19 +108,27 @@ static void CreateRaytracingLightCullingStructure(
 
 	const int32 CellsPerDim = GetCellsPerDim();
 
-	TResourceArray<VectorRegister> RankedLights;
+	TResourceArray<FVector4f> RankedLights;
 	RankedLights.Reserve(NumLightsToUse);
 
 	// setup light vector array sorted by rank
 	for (int32 LightIndex = 0; LightIndex < NumLightsToUse; LightIndex++)
 	{
-		RankedLights.Push(Lights[LightIndices[LightIndex]].BoundingSphereVector);
+		VectorRegister BoundingSphere = Lights[LightIndices[LightIndex]].BoundingSphereVector;
+		RankedLights.Push(
+			FVector4f(
+				VectorGetComponentImpl<0>(BoundingSphere), 
+				VectorGetComponentImpl<1>(BoundingSphere), 
+				VectorGetComponentImpl<2>(BoundingSphere), 
+				VectorGetComponentImpl<3>(BoundingSphere)
+			)
+		);
 	}
 
 	// push null vector to prevent failure in RHICreateStructuredBuffer due to requesting a zero sized allocation
 	if (RankedLights.Num() == 0)
 	{
-		RankedLights.Push(VectorRegister{});
+		RankedLights.Push(FVector4f());
 	}
 
 	FRHIResourceCreateInfo CreateInfo(TEXT("RayTracingCullLights"));
