Fix a bug of splic screen on mobile.

diff --git a/Engine/Source/Runtime/Renderer/Private/MobileShadingRenderer.cpp b/Engine/Source/Runtime/Renderer/Private/MobileShadingRenderer.cpp
--- a/Engine/Source/Runtime/Renderer/Private/MobileShadingRenderer.cpp
+++ b/Engine/Source/Runtime/Renderer/Private/MobileShadingRenderer.cpp
@@ -232,6 +232,7 @@ FMobileSceneRenderer::FMobileSceneRenderer(const FSceneViewFamily* InViewFamily,
 	bRequiresDistanceFieldShadowingPass = false;
 	bIsFullDepthPrepassEnabled = Scene->EarlyZPassMode == DDM_AllOpaque;
 	bIsMaskedOnlyDepthPrepassEnabled = Scene->EarlyZPassMode == DDM_MaskedOnly;
+	bRequiresSceneDepthAux = MobileRequiresSceneDepthAux(ShaderPlatform);
 	
 	StandardTranslucencyPass = ViewFamily.AllowTranslucencyAfterDOF() ? ETranslucencyPass::TPT_StandardTranslucency : ETranslucencyPass::TPT_AllTranslucency;
 	StandardTranslucencyMeshPass = TranslucencyPassToMeshPass(StandardTranslucencyPass);
@@ -457,7 +458,7 @@ void FMobileSceneRenderer::InitViews(FRDGBuilder& GraphBuilder, FSceneTexturesCo
 		bRequireSeparateViewPass ||
 		bIsFullDepthPrepassEnabled;
 	// never keep MSAA depth if SceneDepthAux is enabled
-	bKeepDepthContent = ((NumMSAASamples > 1) && MobileRequiresSceneDepthAux(ShaderPlatform)) ? false : bKeepDepthContent;
+	bKeepDepthContent = ((NumMSAASamples > 1) && bRequiresSceneDepthAux) ? false : bKeepDepthContent;
 
 	// Depth is needed for Editor Primitives
 	if (bShouldCompositeEditorPrimitives)
@@ -837,7 +838,7 @@ void FMobileSceneRenderer::Render(FRDGBuilder& GraphBuilder)
 	{
 		RenderFullDepthPrepass(GraphBuilder, SceneTextures);
 
-		if (!MobileRequiresSceneDepthAux(ShaderPlatform))
+		if (!bRequiresSceneDepthAux)
 		{
 			AddResolveSceneDepthPass(GraphBuilder, Views, SceneTextures.Depth);
 		}
@@ -1008,7 +1009,7 @@ void FMobileSceneRenderer::RenderForward(FRDGBuilder& GraphBuilder, FRDGTextureR
 
 	FRenderTargetBindingSlots BasePassRenderTargets;
 	BasePassRenderTargets[0] = FRenderTargetBinding(SceneColor, SceneColorResolve, ERenderTargetLoadAction::EClear);
-	if (MobileRequiresSceneDepthAux(ShaderPlatform))
+	if (bRequiresSceneDepthAux)
 	{
 		BasePassRenderTargets[1] = FRenderTargetBinding(SceneTextures.DepthAux.Target, SceneTextures.DepthAux.Resolve, ERenderTargetLoadAction::EClear);
 	}
@@ -1042,6 +1043,10 @@ void FMobileSceneRenderer::RenderForward(FRDGBuilder& GraphBuilder, FRDGTextureR
 		if (ViewIndex > 0)
 		{
 			BasePassRenderTargets[0].SetLoadAction(ERenderTargetLoadAction::ELoad);
+			if (bRequiresSceneDepthAux)
+			{
+				BasePassRenderTargets[1].SetLoadAction(ERenderTargetLoadAction::ELoad);
+			}
 			BasePassRenderTargets.DepthStencil.SetDepthLoadAction(ERenderTargetLoadAction::ELoad);
 			BasePassRenderTargets.DepthStencil.SetStencilLoadAction(ERenderTargetLoadAction::ELoad);
 			BasePassRenderTargets.DepthStencil.SetDepthStencilAccess(bIsFullDepthPrepassEnabled ? FExclusiveDepthStencil::DepthRead_StencilWrite : FExclusiveDepthStencil::DepthWrite_StencilWrite);
@@ -1161,7 +1166,7 @@ void FMobileSceneRenderer::RenderForwardMultiPass(FRDGBuilder& GraphBuilder, FMo
 	}
 
 	BasePassRenderTargets[0].SetLoadAction(ERenderTargetLoadAction::ELoad);
-	if (MobileRequiresSceneDepthAux(ShaderPlatform))
+	if (bRequiresSceneDepthAux)
 	{
 		BasePassRenderTargets[1].SetLoadAction(ERenderTargetLoadAction::ELoad);
 	}
@@ -1295,7 +1300,7 @@ void FMobileSceneRenderer::RenderDeferred(FRDGBuilder& GraphBuilder, const FSort
 		ColorTargets.Add(SceneTextures.GBufferA);
 		ColorTargets.Add(SceneTextures.GBufferB);
 		ColorTargets.Add(SceneTextures.GBufferC);
-		if(MobileRequiresSceneDepthAux(ShaderPlatform))
+		if(bRequiresSceneDepthAux)
 		{
 			ColorTargets.Add(SceneTextures.DepthAux.Target);
 		}
diff --git a/Engine/Source/Runtime/Renderer/Private/SceneRendering.h b/Engine/Source/Runtime/Renderer/Private/SceneRendering.h
--- a/Engine/Source/Runtime/Renderer/Private/SceneRendering.h
+++ b/Engine/Source/Runtime/Renderer/Private/SceneRendering.h
@@ -2404,6 +2404,7 @@ private:
 	bool bRequiresDistanceFieldShadowingPass;
 	bool bIsFullDepthPrepassEnabled;
 	bool bIsMaskedOnlyDepthPrepassEnabled;
+	bool bRequiresSceneDepthAux;
 
 	ETranslucencyPass::Type StandardTranslucencyPass;
 	EMeshPass::Type StandardTranslucencyMeshPass;
