Use RGBA16 alias for BC1/BC4 when rendering runtime virtual textures.
Removes inconsistency for some platforms.
Removes special case uint2/uint4 stuff in the compression shader.

diff --git a/Engine/Shaders/Private/BCCompressionCommon.ush b/Engine/Shaders/Private/BCCompressionCommon.ush
--- a/Engine/Shaders/Private/BCCompressionCommon.ush
+++ b/Engine/Shaders/Private/BCCompressionCommon.ush
@@ -167,7 +167,7 @@ uint2 GetPackedAlphaBlockIndices( in float Block[16], in float MinAlpha, in floa
 }
 
 // Compress a BC1 block
-uint2 CompressBC1Block( in float3 Block[16] )
+uint4 CompressBC1Block( in float3 Block[16] )
 {
 	float3 MinColor, MaxColor;
 	GetMinMax( Block, MinColor, MaxColor );
@@ -181,12 +181,12 @@ uint2 CompressBC1Block( in float3 Block[16] )
 		Indices = GetPackedColorBlockIndices( Block, MinColor, MaxColor );
 	}
 
-	return uint2((MinColor565 << 16) | MaxColor565, Indices);
+	return uint4(MaxColor565, MinColor565, Indices & 0xffff, Indices >> 16);
 }
 
 // Compress a BC1 block that will be sampled as sRGB
 // We expect linear colors as input and convert internally
-uint2 CompressBC1Block_SRGB( in float3 Block[16] )
+uint4 CompressBC1Block_SRGB( in float3 Block[16] )
 {
 #if (FEATURE_LEVEL == FEATURE_LEVEL_ES3_1)
 	// mobile uses linear color for RVT
@@ -205,7 +205,7 @@ uint2 CompressBC1Block_SRGB( in float3 Block[16] )
 		Indices = GetPackedColorBlockIndices( Block, MinColor, MaxColor );
 	}
 
-	return uint2((MinColor565 << 16) | MaxColor565, Indices);
+	return uint4(MaxColor565, MinColor565, Indices & 0xffff, Indices >> 16);
 }
 
 // Compress a BC3 block
@@ -275,7 +275,7 @@ uint4 CompressBC3Block_SRGB( in float3 BlockRGB[16], in float BlockA[16] )
 }
 
 // Compress a BC4 block
-uint2 CompressBC4Block( in float Block[16] )
+uint4 CompressBC4Block( in float Block[16] )
 {
 	float MinAlpha, MaxAlpha;
 	GetMinMax( Block, MinAlpha, MaxAlpha );
@@ -289,7 +289,7 @@ uint2 CompressBC4Block( in float Block[16] )
 		Indices = GetPackedAlphaBlockIndices( Block, MinAlpha, MaxAlpha );
 	}
 
-	return uint2((MinAlphaUint << 8) | MaxAlphaUint | Indices.x, Indices.y);
+	return uint4((MinAlphaUint << 8) | MaxAlphaUint, Indices.x >> 16, Indices.y & 0xffff, Indices.y >> 16);
 }
 
 // Compress a BC5 block
diff --git a/Engine/Shaders/Private/VirtualTextureCompress.usf b/Engine/Shaders/Private/VirtualTextureCompress.usf
--- a/Engine/Shaders/Private/VirtualTextureCompress.usf
+++ b/Engine/Shaders/Private/VirtualTextureCompress.usf
@@ -15,9 +15,6 @@
 #define ETC_PROFILE 0
 #endif
 
-// GLES does not support 2-channel images, pack RG32 into RGBA16 on all GL
-#define PACK_RG32_RGBA16 (COMPILER_GLSL || COMPILER_GLSL_ES3_1)
-
 uint4 DestRect;
 Texture2D<float4> RenderTexture0;
 SamplerState TextureSampler0;
@@ -25,19 +22,11 @@ Texture2D<float4> RenderTexture1;
 SamplerState TextureSampler1;
 Texture2D<float4> RenderTexture2;
 SamplerState TextureSampler2;
-
-RWTexture2D<uint4> OutCompressTexture0_u4;
+RWTexture2D<uint4> OutCompressTexture0;
 RWTexture2D<uint4> OutCompressTexture1;
-RWTexture2D<uint4> OutCompressTexture2_u4;
-#if PACK_RG32_RGBA16
-RWTexture2D<uint4> OutCompressTexture0_u2;
-RWTexture2D<uint4> OutCompressTexture2_u2;
-#else
-RWTexture2D<uint2> OutCompressTexture0_u2;
-RWTexture2D<uint2> OutCompressTexture2_u2;
-#endif //PACK_RG32_RGBA16
+RWTexture2D<uint4> OutCompressTexture2;
 
-#if PACK_RG32_RGBA16
+/** Pack ETC RG32 to RGBA16 for output. */
 uint4 PackRG32(in uint2 u2)
 {
 	uint x0 = (u2.x >> 16) & 0x0000ffff;
@@ -46,12 +35,6 @@ uint4 PackRG32(in uint2 u2)
 	uint y1 = (u2.y) & 0x0000ffff;
 	return uint4(x1,x0,y1,y0);
 }
-#else
-uint2 PackRG32(in uint2 u2)
-{
-	return u2;
-}
-#endif //PACK_RG32_RGBA16
 
 /** Trivial vertex shader for copy shaders. */
 void CopyVS(in uint VertexId : SV_VertexID, out float4 OutPosition : SV_POSITION, out float2 OutTexCoord : TEXCOORD0)
@@ -77,9 +60,9 @@ void CompressBaseColorCS(
 	float3 BlockBaseColor[16];
 	ReadBlockRGB(RenderTexture0, TextureSampler0, SampleUV, TexelUVSize, BlockBaseColor);
 #if ETC_PROFILE
-	OutCompressTexture0_u2[ThreadId.xy + DestRect.xy] = PackRG32(CompressBlock_ETC2_SRGB(BlockBaseColor));
+	OutCompressTexture0[ThreadId.xy + DestRect.xy] = PackRG32(CompressBlock_ETC2_SRGB(BlockBaseColor));
 #else
-	OutCompressTexture0_u2[ThreadId.xy + DestRect.xy] = PackRG32(CompressBC1Block_SRGB(BlockBaseColor));
+	OutCompressTexture0[ThreadId.xy + DestRect.xy] = CompressBC1Block_SRGB(BlockBaseColor);
 #endif
 
 }
@@ -122,10 +105,10 @@ void CompressBaseColorNormalSpecularCS(
 		BlockSpecular[i].z = round(BlockSpecular[i].z);
 	}
 #if ETC_PROFILE
-	OutCompressTexture0_u4[ThreadId.xy + DestRect.xy] = CompressBlock_ETC2_RGBA(BlockBaseColor, BlockNormalX);
+	OutCompressTexture0[ThreadId.xy + DestRect.xy] = CompressBlock_ETC2_RGBA(BlockBaseColor, BlockNormalX);
 	OutCompressTexture1[ThreadId.xy + DestRect.xy] = CompressBlock_ETC2_RGBA(BlockSpecular, BlockNormalY);
 #else
-	OutCompressTexture0_u4[ThreadId.xy + DestRect.xy] = CompressBC3Block_SRGB(BlockBaseColor, BlockNormalX);
+	OutCompressTexture0[ThreadId.xy + DestRect.xy] = CompressBC3Block_SRGB(BlockBaseColor, BlockNormalX);
 	OutCompressTexture1[ThreadId.xy + DestRect.xy] = CompressBC3Block(BlockSpecular, BlockNormalY);
 #endif
 
@@ -152,10 +135,10 @@ void CompressBaseColorNormalRoughnessCS(
 								   1.0, 1.0, 1.0, 1.0, 1.0, 1.0, 1.0, 1.0};
 	ReadBlockRGB(RenderTexture1, TextureSampler2, SampleUV, TexelUVSize, BlockNormalXRougnnessY);
 #if ETC_PROFILE
-	OutCompressTexture0_u2[ThreadId.xy + DestRect.xy] = PackRG32(CompressBlock_ETC2_RGB(BlockBaseColor));
+	OutCompressTexture0[ThreadId.xy + DestRect.xy] = PackRG32(CompressBlock_ETC2_RGB(BlockBaseColor));
 	OutCompressTexture1[ThreadId.xy + DestRect.xy] = CompressBlock_ETC2_RGBA(BlockNormalXRougnnessY, BlockNormalZ);
 #else
-	OutCompressTexture0_u2[ThreadId.xy + DestRect.xy] = PackRG32(CompressBC1Block_SRGB(BlockBaseColor));
+	OutCompressTexture0[ThreadId.xy + DestRect.xy] = CompressBC1Block_SRGB(BlockBaseColor);
 	OutCompressTexture1[ThreadId.xy + DestRect.xy] = CompressBC3Block(BlockNormalXRougnnessY, BlockNormalZ);
 #endif
 
@@ -207,13 +190,13 @@ void CompressBaseColorNormalSpecularYCoCgCS(
 	}
 
 #if ETC_PROFILE
-	OutCompressTexture0_u4[ThreadId.xy + DestRect.xy] = CompressBlock_ETC2_YCoCg(BlockBaseColor);
+	OutCompressTexture0[ThreadId.xy + DestRect.xy] = CompressBlock_ETC2_YCoCg(BlockBaseColor);
 	OutCompressTexture1[ThreadId.xy + DestRect.xy] = CompressBlock_ETC2_RG(BlockNormalX, BlockNormalY);
-	OutCompressTexture2_u2[ThreadId.xy + DestRect.xy] = PackRG32(CompressBlock_ETC2_RGB(BlockSpecular));
+	OutCompressTexture2[ThreadId.xy + DestRect.xy] = PackRG32(CompressBlock_ETC2_RGB(BlockSpecular));
 #else
-	OutCompressTexture0_u4[ThreadId.xy + DestRect.xy] = CompressBC3BlockYCoCg(BlockBaseColor);
+	OutCompressTexture0[ThreadId.xy + DestRect.xy] = CompressBC3BlockYCoCg(BlockBaseColor);
 	OutCompressTexture1[ThreadId.xy + DestRect.xy] = CompressBC5Block(BlockNormalX, BlockNormalY);
-	OutCompressTexture2_u2[ThreadId.xy + DestRect.xy] = PackRG32(CompressBC1Block(BlockSpecular));
+	OutCompressTexture2[ThreadId.xy + DestRect.xy] = CompressBC1Block(BlockSpecular);
 #endif
 }
 
@@ -266,13 +249,13 @@ void CompressBaseColorNormalSpecularMaskYCoCgCS(
 	}
 
 #if ETC_PROFILE
-	OutCompressTexture0_u4[ThreadId.xy + DestRect.xy] = CompressBlock_ETC2_YCoCg(BlockBaseColor);
+	OutCompressTexture0[ThreadId.xy + DestRect.xy] = CompressBlock_ETC2_YCoCg(BlockBaseColor);
 	OutCompressTexture1[ThreadId.xy + DestRect.xy] = CompressBlock_ETC2_RG(BlockNormalX, BlockNormalY);
-	OutCompressTexture2_u4[ThreadId.xy + DestRect.xy] = CompressBlock_ETC2_RGBA(BlockSpecular, BlockMask);
+	OutCompressTexture2[ThreadId.xy + DestRect.xy] = CompressBlock_ETC2_RGBA(BlockSpecular, BlockMask);
 #else
-	OutCompressTexture0_u4[ThreadId.xy + DestRect.xy] = CompressBC3BlockYCoCg(BlockBaseColor);
+	OutCompressTexture0[ThreadId.xy + DestRect.xy] = CompressBC3BlockYCoCg(BlockBaseColor);
 	OutCompressTexture1[ThreadId.xy + DestRect.xy] = CompressBC5Block(BlockNormalX, BlockNormalY);
-	OutCompressTexture2_u4[ThreadId.xy + DestRect.xy] = CompressBC3Block(BlockSpecular, BlockMask);
+	OutCompressTexture2[ThreadId.xy + DestRect.xy] = CompressBC3Block(BlockSpecular, BlockMask);
 #endif
 }
 
diff --git a/Engine/Source/Runtime/Renderer/Private/VT/RuntimeVirtualTextureRender.cpp b/Engine/Source/Runtime/Renderer/Private/VT/RuntimeVirtualTextureRender.cpp
--- a/Engine/Source/Runtime/Renderer/Private/VT/RuntimeVirtualTextureRender.cpp
+++ b/Engine/Source/Runtime/Renderer/Private/VT/RuntimeVirtualTextureRender.cpp
@@ -666,11 +666,9 @@ namespace RuntimeVirtualTexture
 			SHADER_PARAMETER_SAMPLER(SamplerState, TextureSampler1)
 			SHADER_PARAMETER_RDG_TEXTURE(Texture2D<float4>, RenderTexture2)
 			SHADER_PARAMETER_SAMPLER(SamplerState, TextureSampler2)
-			SHADER_PARAMETER_RDG_TEXTURE_UAV(RWTexture2D<uint2>, OutCompressTexture0_u2)
-			SHADER_PARAMETER_RDG_TEXTURE_UAV(RWTexture2D<uint4>, OutCompressTexture0_u4)
+			SHADER_PARAMETER_RDG_TEXTURE_UAV(RWTexture2D<uint4>, OutCompressTexture0)
 			SHADER_PARAMETER_RDG_TEXTURE_UAV(RWTexture2D<uint4>, OutCompressTexture1)
-			SHADER_PARAMETER_RDG_TEXTURE_UAV(RWTexture2D<uint2>, OutCompressTexture2_u2)
-			SHADER_PARAMETER_RDG_TEXTURE_UAV(RWTexture2D<uint4>, OutCompressTexture2_u4)
+			SHADER_PARAMETER_RDG_TEXTURE_UAV(RWTexture2D<uint4>, OutCompressTexture2)
 		END_SHADER_PARAMETER_STRUCT()
 
 		static bool ShouldCompilePermutation(FGlobalShaderPermutationParameters const& Parameters)
@@ -931,8 +929,7 @@ namespace RuntimeVirtualTexture
 			bCopyPass = bRenderPass && !bCopyThumbnailPass && !bCompressPass && (MaterialType == ERuntimeVirtualTextureMaterialType::BaseColor_Normal_Specular || MaterialType == ERuntimeVirtualTextureMaterialType::BaseColor_Normal_Specular_YCoCg || MaterialType == ERuntimeVirtualTextureMaterialType::BaseColor_Normal_Specular_Mask_YCoCg);
 			// Not all mobile RHIs support sRGB texture views/aliasing, use only linear targets on mobile
 			const ETextureCreateFlags VT_SRGB = FeatureLevel > ERHIFeatureLevel::ES3_1 ? TexCreate_SRGB : TexCreate_None;
-			// GLES does not support writing to 2-channel images, instead pack RG32 into RGBA16 on all GL
-			const EPixelFormat Compressed64BitFormat = IsOpenGLPlatform(GMaxRHIShaderPlatform) ? PF_R16G16B16A16_UINT : PF_R32G32_UINT;
+			const EPixelFormat Compressed64BitFormat = PF_R16G16B16A16_UINT;
 			const EPixelFormat Compressed128BitFormat = PF_R32G32B32A32_UINT;
 
 			switch (MaterialType)
@@ -944,7 +941,7 @@ namespace RuntimeVirtualTexture
 				}
 				if (bCompressPass)
 				{
-					OutputAlias0 = CompressTexture0_u2 = GraphBuilder.CreateTexture(FRDGTextureDesc::Create2D(TextureSize / 4, Compressed64BitFormat, FClearValueBinding::None, TexCreate_UAV), TEXT("CompressTexture0"));
+					OutputAlias0 = CompressTexture0 = GraphBuilder.CreateTexture(FRDGTextureDesc::Create2D(TextureSize / 4, Compressed64BitFormat, FClearValueBinding::None, TexCreate_UAV), TEXT("CompressTexture0"));
 				}
 				if (bCopyThumbnailPass)
 				{
@@ -959,7 +956,7 @@ namespace RuntimeVirtualTexture
 				}
 				if (bCompressPass)
 				{
-					OutputAlias0 = CompressTexture0_u2 = GraphBuilder.CreateTexture(FRDGTextureDesc::Create2D(TextureSize / 4, Compressed64BitFormat, FClearValueBinding::None, TexCreate_UAV), TEXT("CompressTexture0"));
+					OutputAlias0 = CompressTexture0 = GraphBuilder.CreateTexture(FRDGTextureDesc::Create2D(TextureSize / 4, Compressed64BitFormat, FClearValueBinding::None, TexCreate_UAV), TEXT("CompressTexture0"));
 					OutputAlias1 = CompressTexture1 = GraphBuilder.CreateTexture(FRDGTextureDesc::Create2D(TextureSize / 4, Compressed128BitFormat, FClearValueBinding::None, TexCreate_UAV), TEXT("CompressTexture1"));
 				}
 				if (bCopyThumbnailPass)
@@ -977,7 +974,7 @@ namespace RuntimeVirtualTexture
 				}
 				if (bCompressPass)
 				{
-					OutputAlias0 = CompressTexture0_u4 = GraphBuilder.CreateTexture(FRDGTextureDesc::Create2D(TextureSize / 4, Compressed128BitFormat, FClearValueBinding::None, TexCreate_UAV), TEXT("CompressTexture0"));
+					OutputAlias0 = CompressTexture0 = GraphBuilder.CreateTexture(FRDGTextureDesc::Create2D(TextureSize / 4, Compressed128BitFormat, FClearValueBinding::None, TexCreate_UAV), TEXT("CompressTexture0"));
 					OutputAlias1 = CompressTexture1 = GraphBuilder.CreateTexture(FRDGTextureDesc::Create2D(TextureSize / 4, Compressed128BitFormat, FClearValueBinding::None, TexCreate_UAV), TEXT("CompressTexture1"));
 				}
 				if (bCopyPass)
@@ -999,9 +996,9 @@ namespace RuntimeVirtualTexture
 				}
 				if (bCompressPass)
 				{
-					OutputAlias0 = CompressTexture0_u4 = GraphBuilder.CreateTexture(FRDGTextureDesc::Create2D(TextureSize / 4, Compressed128BitFormat, FClearValueBinding::None, TexCreate_UAV), TEXT("CompressTexture0"));
+					OutputAlias0 = CompressTexture0 = GraphBuilder.CreateTexture(FRDGTextureDesc::Create2D(TextureSize / 4, Compressed128BitFormat, FClearValueBinding::None, TexCreate_UAV), TEXT("CompressTexture0"));
 					OutputAlias1 = CompressTexture1 = GraphBuilder.CreateTexture(FRDGTextureDesc::Create2D(TextureSize / 4, Compressed128BitFormat, FClearValueBinding::None, TexCreate_UAV), TEXT("CompressTexture1"));
-					OutputAlias2 = CompressTexture2_u2 = GraphBuilder.CreateTexture(FRDGTextureDesc::Create2D(TextureSize / 4, Compressed64BitFormat, FClearValueBinding::None, TexCreate_UAV), TEXT("CompressTexture2"));
+					OutputAlias2 = CompressTexture2 = GraphBuilder.CreateTexture(FRDGTextureDesc::Create2D(TextureSize / 4, Compressed64BitFormat, FClearValueBinding::None, TexCreate_UAV), TEXT("CompressTexture2"));
 				}
 				if (bCopyPass)
 				{
@@ -1023,9 +1020,9 @@ namespace RuntimeVirtualTexture
 				}
 				if (bCompressPass)
 				{
-					OutputAlias0 = CompressTexture0_u4 = GraphBuilder.CreateTexture(FRDGTextureDesc::Create2D(TextureSize / 4, Compressed128BitFormat, FClearValueBinding::None, TexCreate_UAV), TEXT("CompressTexture0"));
+					OutputAlias0 = CompressTexture0 = GraphBuilder.CreateTexture(FRDGTextureDesc::Create2D(TextureSize / 4, Compressed128BitFormat, FClearValueBinding::None, TexCreate_UAV), TEXT("CompressTexture0"));
 					OutputAlias1 = CompressTexture1 = GraphBuilder.CreateTexture(FRDGTextureDesc::Create2D(TextureSize / 4, Compressed128BitFormat, FClearValueBinding::None, TexCreate_UAV), TEXT("CompressTexture1"));
-					OutputAlias2 = CompressTexture2_u4 = GraphBuilder.CreateTexture(FRDGTextureDesc::Create2D(TextureSize / 4, Compressed128BitFormat, FClearValueBinding::None, TexCreate_UAV), TEXT("CompressTexture2"));
+					OutputAlias2 = CompressTexture2 = GraphBuilder.CreateTexture(FRDGTextureDesc::Create2D(TextureSize / 4, Compressed128BitFormat, FClearValueBinding::None, TexCreate_UAV), TEXT("CompressTexture2"));
 				}
 				if (bCopyPass)
 				{
@@ -1061,11 +1058,9 @@ namespace RuntimeVirtualTexture
 		FRDGTextureRef RenderTexture0 = nullptr;
 		FRDGTextureRef RenderTexture1 = nullptr;
 		FRDGTextureRef RenderTexture2 = nullptr;
-		FRDGTextureRef CompressTexture0_u2 = nullptr;
-		FRDGTextureRef CompressTexture0_u4 = nullptr;
+		FRDGTextureRef CompressTexture0 = nullptr;
 		FRDGTextureRef CompressTexture1 = nullptr;
-		FRDGTextureRef CompressTexture2_u2 = nullptr;
-		FRDGTextureRef CompressTexture2_u4 = nullptr;
+		FRDGTextureRef CompressTexture2 = nullptr;
 		FRDGTextureRef CopyTexture0 = nullptr;
 		FRDGTextureRef CopyTexture1 = nullptr;
 		FRDGTextureRef CopyTexture2 = nullptr;
@@ -1200,11 +1195,9 @@ namespace RuntimeVirtualTexture
 			PassParameters->TextureSampler1 = TStaticSamplerState<SF_Point, AM_Clamp, AM_Clamp, AM_Clamp>::GetRHI();
 			PassParameters->RenderTexture2 = GraphSetup.RenderTexture2;
 			PassParameters->TextureSampler2 = TStaticSamplerState<SF_Point, AM_Clamp, AM_Clamp, AM_Clamp>::GetRHI();
-			PassParameters->OutCompressTexture0_u2 = GraphSetup.CompressTexture0_u2 ? GraphBuilder.CreateUAV(FRDGTextureUAVDesc(GraphSetup.CompressTexture0_u2)) : nullptr;
-			PassParameters->OutCompressTexture0_u4 = GraphSetup.CompressTexture0_u4 ? GraphBuilder.CreateUAV(FRDGTextureUAVDesc(GraphSetup.CompressTexture0_u4)) : nullptr;
+			PassParameters->OutCompressTexture0 = GraphSetup.CompressTexture0 ? GraphBuilder.CreateUAV(FRDGTextureUAVDesc(GraphSetup.CompressTexture0)) : nullptr;
 			PassParameters->OutCompressTexture1 = GraphSetup.CompressTexture1 ? GraphBuilder.CreateUAV(FRDGTextureUAVDesc(GraphSetup.CompressTexture1)) : nullptr;
-			PassParameters->OutCompressTexture2_u2 = GraphSetup.CompressTexture2_u2 ? GraphBuilder.CreateUAV(FRDGTextureUAVDesc(GraphSetup.CompressTexture2_u2)) : nullptr;
-			PassParameters->OutCompressTexture2_u4 = GraphSetup.CompressTexture2_u4 ? GraphBuilder.CreateUAV(FRDGTextureUAVDesc(GraphSetup.CompressTexture2_u4)) : nullptr;
+			PassParameters->OutCompressTexture2 = GraphSetup.CompressTexture2 ? GraphBuilder.CreateUAV(FRDGTextureUAVDesc(GraphSetup.CompressTexture2)) : nullptr;
 
 			AddCompressPass(GraphBuilder, Scene->GetFeatureLevel(), PassParameters, TextureSize, MaterialType);
 		}
