Remove IsMobileDistanceFieldShadowingEnabled

diff --git a/Engine/Source/Runtime/RenderCore/Private/RenderUtils.cpp b/Engine/Source/Runtime/RenderCore/Private/RenderUtils.cpp
--- a/Engine/Source/Runtime/RenderCore/Private/RenderUtils.cpp
+++ b/Engine/Source/Runtime/RenderCore/Private/RenderUtils.cpp
@@ -1185,14 +1185,6 @@ RENDERCORE_API bool IsMobileDistanceFieldEnabled(const FStaticShaderPlatform Pla
 	return IsMobilePlatform(Platform) && (FDataDrivenShaderPlatformInfo::GetSupportsMobileDistanceField(Platform)/* || IsD3DPlatform(Platform)*/) && IsUsingDistanceFields(Platform);
 }
 
-RENDERCORE_API bool IsMobileDistanceFieldShadowingEnabled(const FStaticShaderPlatform Platform)
-{
-	static auto CVar = IConsoleManager::Get().FindConsoleVariable(TEXT("r.DistanceFieldShadowing"));
-	bool bDistanceFieldShadowingEnabled = CVar && (CVar->GetInt() != 0);
-
-	return GRHISupportsPixelShaderUAVs && bDistanceFieldShadowingEnabled && IsMobileDistanceFieldEnabled(Platform);
-}
-
 RENDERCORE_API bool MobileBasePassAlwaysUsesCSM(const FStaticShaderPlatform Platform)
 {
 	static TConsoleVariableData<int32>* CVar = IConsoleManager::Get().FindTConsoleVariableDataInt(TEXT("r.Mobile.Shadow.CSMShaderCullingMethod"));
diff --git a/Engine/Source/Runtime/RenderCore/Public/RenderUtils.h b/Engine/Source/Runtime/RenderCore/Public/RenderUtils.h
--- a/Engine/Source/Runtime/RenderCore/Public/RenderUtils.h
+++ b/Engine/Source/Runtime/RenderCore/Public/RenderUtils.h
@@ -503,8 +503,6 @@ RENDERCORE_API bool UseMobileAmbientOcclusion(const FStaticShaderPlatform Platfo
 
 RENDERCORE_API bool IsMobileDistanceFieldEnabled(const FStaticShaderPlatform Platform);
 
-RENDERCORE_API bool IsMobileDistanceFieldShadowingEnabled(const FStaticShaderPlatform Platform);
-
 RENDERCORE_API bool MobileBasePassAlwaysUsesCSM(const FStaticShaderPlatform Platform);
 
 RENDERCORE_API bool SupportsGen4TAA(const FStaticShaderPlatform Platform);
diff --git a/Engine/Source/Runtime/Renderer/Private/MobileBasePassRendering.cpp b/Engine/Source/Runtime/Renderer/Private/MobileBasePassRendering.cpp
--- a/Engine/Source/Runtime/Renderer/Private/MobileBasePassRendering.cpp
+++ b/Engine/Source/Runtime/Renderer/Private/MobileBasePassRendering.cpp
@@ -207,10 +207,8 @@ void SetupMobileBasePassUniformParameters(
 	BasePassParameters.AmbientOcclusionTexture = AmbientOcclusionTexture;
 	BasePassParameters.AmbientOcclusionSampler = TStaticSamplerState<SF_Bilinear, AM_Clamp, AM_Clamp, AM_Clamp>::GetRHI();
 	BasePassParameters.AmbientOcclusionStaticFraction = FMath::Clamp(View.FinalPostProcessSettings.AmbientOcclusionStaticFraction, 0.0f, 1.0f);
-
-	bool bRequiresDistanceFieldShadowingPass = IsMobileDistanceFieldShadowingEnabled(View.GetShaderPlatform());
 	
-	if (bRequiresDistanceFieldShadowingPass && GScreenSpaceShadowMaskTextureMobileOutputs.ScreenSpaceShadowMaskTextureMobile.IsValid())
+	if (IsMobileDistanceFieldEnabled(View.GetShaderPlatform()) && GScreenSpaceShadowMaskTextureMobileOutputs.ScreenSpaceShadowMaskTextureMobile.IsValid())
 	{
 		FRDGTextureRef ScreenShadowMaskTexture = GraphBuilder.RegisterExternalTexture(GScreenSpaceShadowMaskTextureMobileOutputs.ScreenSpaceShadowMaskTextureMobile, TEXT("ScreenSpaceShadowMaskTextureMobile"));
 		BasePassParameters.ScreenSpaceShadowMaskTexture = ScreenShadowMaskTexture;
diff --git a/Engine/Source/Runtime/Renderer/Private/MobileDeferredShadingPass.cpp b/Engine/Source/Runtime/Renderer/Private/MobileDeferredShadingPass.cpp
--- a/Engine/Source/Runtime/Renderer/Private/MobileDeferredShadingPass.cpp
+++ b/Engine/Source/Runtime/Renderer/Private/MobileDeferredShadingPass.cpp
@@ -269,9 +269,7 @@ static void RenderDirectLight(FRHICommandListImmediate& RHICmdList, const FScene
 	PassParameters.PreIntegratedGF = GSystemTextures.PreintegratedGF->GetRenderTargetItem().ShaderResourceTexture;
 	PassParameters.PreIntegratedGFSampler = TStaticSamplerState<SF_Bilinear, AM_Clamp, AM_Clamp, AM_Clamp>::GetRHI();
 
-	bool bRequiresDistanceFieldShadowingPass = IsMobileDistanceFieldShadowingEnabled(View.GetShaderPlatform());
-
-	if (bRequiresDistanceFieldShadowingPass && GScreenSpaceShadowMaskTextureMobileOutputs.ScreenSpaceShadowMaskTextureMobile.IsValid())
+	if (IsMobileDistanceFieldEnabled(View.GetShaderPlatform()) && GScreenSpaceShadowMaskTextureMobileOutputs.ScreenSpaceShadowMaskTextureMobile.IsValid())
 	{
 		PassParameters.ScreenSpaceShadowMaskTexture = GScreenSpaceShadowMaskTextureMobileOutputs.ScreenSpaceShadowMaskTextureMobile->GetRenderTargetItem().ShaderResourceTexture;
 		PassParameters.ScreenSpaceShadowMaskSampler = TStaticSamplerState<SF_Point, AM_Clamp, AM_Clamp, AM_Clamp>::GetRHI();
diff --git a/Engine/Source/Runtime/Renderer/Private/MobileShadingRenderer.cpp b/Engine/Source/Runtime/Renderer/Private/MobileShadingRenderer.cpp
--- a/Engine/Source/Runtime/Renderer/Private/MobileShadingRenderer.cpp
+++ b/Engine/Source/Runtime/Renderer/Private/MobileShadingRenderer.cpp
@@ -229,7 +229,6 @@ FMobileSceneRenderer::FMobileSceneRenderer(const FSceneViewFamily* InViewFamily,
 	bShouldRenderCustomDepth = false;
 	bRequiresPixelProjectedPlanarRelfectionPass = false;
 	bRequiresAmbientOcclusionPass = false;
-	bRequiresDistanceFieldShadowingPass = false;
 	bIsFullDepthPrepassEnabled = Scene->EarlyZPassMode == DDM_AllOpaque;
 	bIsMaskedOnlyDepthPrepassEnabled = Scene->EarlyZPassMode == DDM_MaskedOnly;
 	bRequiresSceneDepthAux = MobileRequiresSceneDepthAux(ShaderPlatform);
@@ -433,8 +432,6 @@ void FMobileSceneRenderer::InitViews(FRDGBuilder& GraphBuilder, FSceneTexturesCo
 		&& !ViewFamily.EngineShowFlags.VisualizeLightCulling
 		&& !ViewFamily.UseDebugViewPS();
 
-	bRequiresDistanceFieldShadowingPass = bRequiresDistanceField && IsMobileDistanceFieldShadowingEnabled(ShaderPlatform);
-
 	bShouldRenderHZB = ShouldRenderHZB();
 		
 
@@ -449,7 +446,7 @@ void FMobileSceneRenderer::InitViews(FRDGBuilder& GraphBuilder, FSceneTexturesCo
 		bRequiresMultiPass ||
 		bForceDepthResolve ||
 		bRequiresAmbientOcclusionPass ||
-		bRequiresDistanceFieldShadowingPass ||
+		bRequiresDistanceField ||
 		bRequiresPixelProjectedPlanarRelfectionPass ||
 		bSeparateTranslucencyActive ||
 		Views[0].bIsReflectionCapture ||
@@ -500,7 +497,7 @@ void FMobileSceneRenderer::InitViews(FRDGBuilder& GraphBuilder, FSceneTexturesCo
 		InitSkyAtmosphereForViews(RHICmdList);
 	}
 
-	if (bRequiresDistanceFieldShadowingPass)
+	if (bRequiresDistanceField)
 	{
 		InitMobileSDFShadowingOutputs(RHICmdList, SceneTexturesConfig.Extent);
 	}
@@ -830,7 +827,7 @@ void FMobileSceneRenderer::Render(FRDGBuilder& GraphBuilder)
 		SceneTextures.MobileSetupMode = EMobileSceneTextureSetupMode::SceneDepth;
 		SceneTextures.MobileUniformBuffer = CreateMobileSceneTextureUniformBuffer(GraphBuilder, SceneTextures.MobileSetupMode);
 
-		if (bRequiresDistanceFieldShadowingPass)
+		if (bRequiresDistanceField)
 		{
 			CSV_SCOPED_TIMING_STAT_EXCLUSIVE(RenderMobileShadowProjections);
 			RDG_GPU_STAT_SCOPE(GraphBuilder, ShadowProjection);
diff --git a/Engine/Source/Runtime/Renderer/Private/SceneRendering.h b/Engine/Source/Runtime/Renderer/Private/SceneRendering.h
--- a/Engine/Source/Runtime/Renderer/Private/SceneRendering.h
+++ b/Engine/Source/Runtime/Renderer/Private/SceneRendering.h
@@ -2381,7 +2381,6 @@ private:
 	bool bShouldRenderVelocities;
 	bool bShouldRenderHZB;
 	bool bRequiresDistanceField;
-	bool bRequiresDistanceFieldShadowingPass;
 	bool bIsFullDepthPrepassEnabled;
 	bool bIsMaskedOnlyDepthPrepassEnabled;
 	bool bRequiresSceneDepthAux;
