Check if primary view index and family are valid when getting primary or instanced view

diff --git a/Engine/Source/Runtime/Engine/Private/SceneView.cpp b/Engine/Source/Runtime/Engine/Private/SceneView.cpp
--- a/Engine/Source/Runtime/Engine/Private/SceneView.cpp
+++ b/Engine/Source/Runtime/Engine/Private/SceneView.cpp
@@ -2641,13 +2641,17 @@ FRDGPooledBuffer* FSceneView::GetEyeAdaptationBuffer() const
 const FSceneView* FSceneView::GetPrimarySceneView() const
 {
 	// It is valid for this function to return itself if it's already the primary view.
-	return Family->Views[PrimaryViewIndex];
+	if (Family && Family->Views.IsValidIndex(PrimaryViewIndex))
+	{
+		return Family->Views[PrimaryViewIndex];
+	}
+	return this;
 }
 
 const FSceneView* FSceneView::GetInstancedSceneView() const
 {
 	// If called on the first secondary view it'll return itself.
-	if (Family->Views.IsValidIndex(PrimaryViewIndex + 1))
+	if (Family && Family->Views.IsValidIndex(PrimaryViewIndex + 1))
 	{
 		const FSceneView* SecondaryView = Family->Views[PrimaryViewIndex + 1];
 		return IStereoRendering::IsASecondaryView(*SecondaryView) ? SecondaryView : nullptr;
diff --git a/Engine/Source/Runtime/Renderer/Private/SceneRendering.cpp b/Engine/Source/Runtime/Renderer/Private/SceneRendering.cpp
--- a/Engine/Source/Runtime/Renderer/Private/SceneRendering.cpp
+++ b/Engine/Source/Runtime/Renderer/Private/SceneRendering.cpp
@@ -2025,14 +2025,18 @@ FViewShaderParameters FViewInfo::GetShaderParameters() const
 const FViewInfo* FViewInfo::GetPrimaryView() const
 {
 	// It is valid for this function to return itself if it's already the primary view.
-	const FSceneView* PrimaryView = Family->Views[PrimaryViewIndex];
-	check(PrimaryView->bIsViewInfo);
-	return static_cast<const FViewInfo*>(PrimaryView);
+	if (Family && Family->Views.IsValidIndex(PrimaryViewIndex))
+	{
+		const FSceneView* PrimaryView = Family->Views[PrimaryViewIndex];
+		check(PrimaryView->bIsViewInfo);
+		return static_cast<const FViewInfo*>(PrimaryView);
+	}
+	return this;
 }
 
 const FViewInfo* FViewInfo::GetInstancedView() const
 {
-	if ((IsInstancedStereoPass() || bIsMobileMultiViewEnabled) && Family->Views.Num() > 0)
+	if ((IsInstancedStereoPass() || bIsMobileMultiViewEnabled))
 	{
 		return static_cast<const FViewInfo*>(GetInstancedSceneView());
 	}
