Fix for dummy light injection uniform buffer emitting  a read without write dependency assert in RDG.

diff --git a/Engine/Source/Runtime/Renderer/Private/LightGridInjection.cpp b/Engine/Source/Runtime/Renderer/Private/LightGridInjection.cpp
--- a/Engine/Source/Runtime/Renderer/Private/LightGridInjection.cpp
+++ b/Engine/Source/Runtime/Renderer/Private/LightGridInjection.cpp
@@ -94,21 +94,21 @@ TRDGUniformBufferRef<FForwardLightData> CreateDummyForwardLightUniformBuffer(FRD
 	ForwardLightData->DirectionalLightShadowmapAtlas = SystemTextures.Black;
 	ForwardLightData->DirectionalLightStaticShadowmap = GBlackTexture->TextureRHI;
 
-	FRDGBufferRef ForwardLocalLightBuffer     = GraphBuilder.CreateBuffer(FRDGBufferDesc::CreateUploadDesc(sizeof(FVector4f), sizeof(FForwardLocalLightData) / sizeof(FVector4f)), TEXT("ForwardLocalLightBuffer"));
-	ForwardLightData->ForwardLocalLightBuffer = GraphBuilder.CreateSRV(FRDGBufferSRVDesc(ForwardLocalLightBuffer, PF_A32B32G32R32F));
+	FRDGBufferRef ForwardLocalLightBuffer     = GSystemTextures.GetDefaultBuffer(GraphBuilder, sizeof(FVector4f));
+	ForwardLightData->ForwardLocalLightBuffer = GraphBuilder.CreateSRV(ForwardLocalLightBuffer, PF_A32B32G32R32F);
 
-	FRDGBufferRef NumCulledLightsGrid     = GraphBuilder.CreateBuffer(FRDGBufferDesc::CreateBufferDesc(sizeof(uint32), 1), TEXT("NumCulledLightsGrid"));
-	ForwardLightData->NumCulledLightsGrid = GraphBuilder.CreateSRV(FRDGBufferSRVDesc(NumCulledLightsGrid, PF_R32_UINT));
+	FRDGBufferRef NumCulledLightsGrid     = GSystemTextures.GetDefaultBuffer(GraphBuilder, sizeof(uint32));
+	ForwardLightData->NumCulledLightsGrid = GraphBuilder.CreateSRV(NumCulledLightsGrid, PF_R32_UINT);
 
 	if (RHISupportsBufferLoadTypeConversion(GMaxRHIShaderPlatform))
 	{
-		FRDGBufferRef CulledLightDataGrid     = GraphBuilder.CreateBuffer(FRDGBufferDesc::CreateBufferDesc(sizeof(uint16), 1), TEXT("CulledLightDataGrid"));
-		ForwardLightData->CulledLightDataGrid = GraphBuilder.CreateSRV(FRDGBufferSRVDesc(CulledLightDataGrid, PF_R16_UINT));
+		FRDGBufferRef CulledLightDataGrid     = GSystemTextures.GetDefaultBuffer(GraphBuilder, sizeof(uint16));
+		ForwardLightData->CulledLightDataGrid = GraphBuilder.CreateSRV(CulledLightDataGrid, PF_R16_UINT);
 	}
 	else
 	{
-		FRDGBufferRef CulledLightDataGrid     = GraphBuilder.CreateBuffer(FRDGBufferDesc::CreateBufferDesc(sizeof(uint32), 1), TEXT("CulledLightDataGrid"));
-		ForwardLightData->CulledLightDataGrid = GraphBuilder.CreateSRV(FRDGBufferSRVDesc(CulledLightDataGrid, PF_R32_UINT));
+		FRDGBufferRef CulledLightDataGrid     = GSystemTextures.GetDefaultBuffer(GraphBuilder, sizeof(uint32));
+		ForwardLightData->CulledLightDataGrid = GraphBuilder.CreateSRV(CulledLightDataGrid, PF_R32_UINT);
 	}
 
 	return GraphBuilder.CreateUniformBuffer(ForwardLightData);
