Saves 48-bytes per UPrimitiveComponent.

- Moves InstanceUpdateCmdBuffer from UPrimitiveComponent to UInstancedStaticMeshComponent.  It is only used by the derived class so this code is safe.  This saves 48-bytes (40-bytes from the removed member and 8-bytes of padding - verified with CruncherSharp).

diff --git a/Engine/Source/Runtime/Engine/Classes/Components/InstancedStaticMeshComponent.h b/Engine/Source/Runtime/Engine/Classes/Components/InstancedStaticMeshComponent.h
--- a/Engine/Source/Runtime/Engine/Classes/Components/InstancedStaticMeshComponent.h
+++ b/Engine/Source/Runtime/Engine/Classes/Components/InstancedStaticMeshComponent.h
@@ -389,6 +389,11 @@ public:
 	*/
 	TMap<int32, int32> InstanceIdToInstanceIndexMap;
 
+	// This is here because in void FScene::UpdatePrimitiveInstance(UPrimitiveComponent* Primitive) we want access
+	// to the instances updates through the primitive component in a generic way.
+	/** Recorded modifications to per-instance data */
+	FInstanceUpdateCmdBuffer InstanceUpdateCmdBuffer;
+
 private:
 
 	/** Sets up new instance data to sensible defaults, creates physics counterparts if possible. */
diff --git a/Engine/Source/Runtime/Engine/Classes/Components/PrimitiveComponent.h b/Engine/Source/Runtime/Engine/Classes/Components/PrimitiveComponent.h
--- a/Engine/Source/Runtime/Engine/Classes/Components/PrimitiveComponent.h
+++ b/Engine/Source/Runtime/Engine/Classes/Components/PrimitiveComponent.h
@@ -23,7 +23,6 @@
 #include "HitProxies.h"
 #include "Interfaces/Interface_AsyncCompilation.h"
 #include "HLOD/HLODBatchingPolicy.h"
-#include "PrimitiveInstanceUpdateCommand.h"
 #include "PrimitiveComponent.generated.h"
 
 class AController;
@@ -800,11 +799,6 @@ public:
 	/** Last time the component was submitted for rendering (called FScene::AddPrimitive). */
 	float LastSubmitTime;
 
-	// This is here because in void FScene::UpdatePrimitiveInstance(UPrimitiveComponent* Primitive) we want access
-	// to the instances updates through the primitive component in a generic way.
-	/** Recorded modifications to per-instance data */
-	FInstanceUpdateCmdBuffer InstanceUpdateCmdBuffer;
-
 private:
 	/**
 	 * The value of WorldSettings->TimeSeconds for the frame when this component was last rendered.  This is written
diff --git a/Engine/Source/Runtime/Engine/Public/SceneInterface.h b/Engine/Source/Runtime/Engine/Public/SceneInterface.h
--- a/Engine/Source/Runtime/Engine/Public/SceneInterface.h
+++ b/Engine/Source/Runtime/Engine/Public/SceneInterface.h
@@ -26,6 +26,7 @@ class UDecalComponent;
 class ULightComponent;
 class UPlanarReflectionComponent;
 class UPrimitiveComponent;
+class UInstancedStaticMeshComponent;
 class UReflectionCaptureComponent;
 class USkyLightComponent;
 class UStaticMeshComponent;
@@ -90,7 +91,7 @@ public:
 	 *
 	 * @param Primitive - primitive component to update
 	 */
-	virtual void UpdatePrimitiveInstances(UPrimitiveComponent* Primitive) = 0;
+	virtual void UpdatePrimitiveInstances(UInstancedStaticMeshComponent* Primitive) = 0;
 
 	/** 
 	 * Updates the custom primitive data of a primitive component which has already been added to the scene. 
diff --git a/Engine/Source/Runtime/Renderer/Private/RendererScene.cpp b/Engine/Source/Runtime/Renderer/Private/RendererScene.cpp
--- a/Engine/Source/Runtime/Renderer/Private/RendererScene.cpp
+++ b/Engine/Source/Runtime/Renderer/Private/RendererScene.cpp
@@ -31,6 +31,7 @@
 #include "Components/DecalComponent.h"
 #include "Components/ReflectionCaptureComponent.h"
 #include "Components/RuntimeVirtualTextureComponent.h"
+#include "Components/InstancedStaticMeshComponent.h"
 #include "ScenePrivateBase.h"
 #include "SceneCore.h"
 #include "Rendering/MotionVectorSimulation.h"
@@ -1471,7 +1472,7 @@ void FScene::UpdatePrimitiveOcclusionBoundsSlack(UPrimitiveComponent* Primitive,
 	}
 }
 
-void FScene::UpdatePrimitiveInstances(UPrimitiveComponent* Primitive)
+void FScene::UpdatePrimitiveInstances(UInstancedStaticMeshComponent* Primitive)
 {
 	SCOPE_CYCLE_COUNTER(STAT_UpdatePrimitiveInstanceGT);
 	SCOPED_NAMED_EVENT(FScene_UpdatePrimitiveInstance, FColor::Yellow);
@@ -1481,7 +1482,7 @@ void FScene::UpdatePrimitiveInstances(UPrimitiveComponent* Primitive)
 		FUpdateInstanceCommand UpdateParams;
 		UpdateParams.PrimitiveSceneProxy = Primitive->SceneProxy;
 		UpdateParams.WorldBounds = Primitive->Bounds;
-		UpdateParams.LocalBounds = Primitive->GetLocalBounds();
+		UpdateParams.LocalBounds = Cast<UPrimitiveComponent>(Primitive)->GetLocalBounds();
 
 		// #todo (jnadro) This code should not be dependent on static mesh bounds.
 		UStaticMeshComponent* StaticMeshComponent = Cast<UStaticMeshComponent>(Primitive);
@@ -5078,7 +5079,7 @@ public:
 
 	/** Updates the transform of a primitive which has already been added to the scene. */
 	virtual void UpdatePrimitiveTransform(UPrimitiveComponent* Primitive) override {}
-	virtual void UpdatePrimitiveInstances(UPrimitiveComponent* Primitive) override {}
+	virtual void UpdatePrimitiveInstances(UInstancedStaticMeshComponent* Primitive) override {}
 	virtual void UpdatePrimitiveOcclusionBoundsSlack(UPrimitiveComponent* Primitive, float NewSlack) override {}
 	virtual void UpdatePrimitiveAttachment(UPrimitiveComponent* Primitive) override {}
 	virtual void UpdateCustomPrimitiveData(UPrimitiveComponent* Primitive) override {}
diff --git a/Engine/Source/Runtime/Renderer/Private/ScenePrivate.h b/Engine/Source/Runtime/Renderer/Private/ScenePrivate.h
--- a/Engine/Source/Runtime/Renderer/Private/ScenePrivate.h
+++ b/Engine/Source/Runtime/Renderer/Private/ScenePrivate.h
@@ -76,6 +76,7 @@ class UExponentialHeightFogComponent;
 class ULightComponent;
 class UPlanarReflectionComponent;
 class UPrimitiveComponent;
+class UInstancedStaticMeshComponent;
 class UReflectionCaptureComponent;
 class USkyLightComponent;
 class UStaticMesh;
@@ -3035,7 +3036,7 @@ public:
 	virtual void ReleasePrimitive(UPrimitiveComponent* Primitive) override;
 	virtual void UpdateAllPrimitiveSceneInfos(FRDGBuilder& GraphBuilder, bool bAsyncCreateLPIs = false) override;
 	virtual void UpdatePrimitiveTransform(UPrimitiveComponent* Primitive) override;
-	virtual void UpdatePrimitiveInstances(UPrimitiveComponent* Primitive) override;
+	virtual void UpdatePrimitiveInstances(UInstancedStaticMeshComponent* Primitive) override;
 	virtual void UpdatePrimitiveOcclusionBoundsSlack(UPrimitiveComponent* Primitive, float NewSlack) override;
 	virtual void UpdatePrimitiveAttachment(UPrimitiveComponent* Primitive) override;
 	virtual void UpdateCustomPrimitiveData(UPrimitiveComponent* Primitive) override;
