UE5_RELEASE:  Fix for use-after-free bug in SceneRenderer.  Need to pull optional profiling results destination pointer out of structure in case it gets freed after scene render.

#rnx

diff --git a/Engine/Source/Runtime/Renderer/Private/SceneRendering.cpp b/Engine/Source/Runtime/Renderer/Private/SceneRendering.cpp
--- a/Engine/Source/Runtime/Renderer/Private/SceneRendering.cpp
+++ b/Engine/Source/Runtime/Renderer/Private/SceneRendering.cpp
@@ -4337,15 +4337,18 @@ void FRendererModule::BeginRenderingViewFamily(FCanvas* Canvas, FSceneViewFamily
 		ENQUEUE_RENDER_COMMAND(FDrawSceneCommand)(
 			[SceneRenderer, DrawSceneEnqueue](FRHICommandListImmediate& RHICmdList)
 			{
+				// Cache the profiling results pointer, as SceneRenderer may be deleted after rendering finishes
+				float* ProfileSceneRenderTime = SceneRenderer->ViewFamily.ProfileSceneRenderTime;
+
 				uint64 SceneRenderStart = FPlatformTime::Cycles64();
 				const float StartDelayMillisec = FPlatformTime::ToMilliseconds64(SceneRenderStart - DrawSceneEnqueue);
 				CSV_CUSTOM_STAT_GLOBAL(DrawSceneCommand_StartDelay, StartDelayMillisec, ECsvCustomStatOp::Set);
 				RenderViewFamily_RenderThread(RHICmdList, SceneRenderer);
 				FlushPendingDeleteRHIResources_RenderThread();
 
-				if (SceneRenderer->ViewFamily.ProfileSceneRenderTime)
+				if (ProfileSceneRenderTime)
 				{
-					*SceneRenderer->ViewFamily.ProfileSceneRenderTime = (float)FPlatformTime::ToSeconds64(FPlatformTime::Cycles64() - SceneRenderStart);
+					*ProfileSceneRenderTime = (float)FPlatformTime::ToSeconds64(FPlatformTime::Cycles64() - SceneRenderStart);
 				}
 			});
 
