Remove HALF_WORLD_MAX dependency in mobile mesh sorting

diff --git a/Engine/Source/Runtime/Renderer/Private/MeshDrawCommands.cpp b/Engine/Source/Runtime/Renderer/Private/MeshDrawCommands.cpp
--- a/Engine/Source/Runtime/Renderer/Private/MeshDrawCommands.cpp
+++ b/Engine/Source/Runtime/Renderer/Private/MeshDrawCommands.cpp
@@ -217,8 +217,10 @@ static uint64 GetMobileBasePassSortKey_FrontToBack(bool bMasked, bool bBackgroun
 
 static uint64 GetMobileBasePassSortKey_ByState(bool bMasked, bool bBackground, int32 PipelineId, int32 StateBucketId, float PipelineDistance, float PrimitiveDistance)
 {
-	const float PrimitiveDepthQuantization = ((1 << 14) - 1);
-
+	// maximum primitive distance for bucketing, aprox 10.5km 
+	constexpr float PrimitiveMaxDepth = (1 << 20);
+	constexpr float PrimitiveDepthQuantization = ((1 << 14) - 1);
+	
 	union
 	{
 		uint64 PackedInt;
@@ -239,11 +241,11 @@ static uint64 GetMobileBasePassSortKey_ByState(bool bMasked, bool bBackground, i
 	Key.PackedInt = 0;
 	Key.Fields.Masked = bMasked;
 	Key.Fields.Background = bBackground;
-	F2I.F = PipelineDistance / HALF_WORLD_MAX;
+	F2I.F = PipelineDistance / PrimitiveMaxDepth;
 	Key.Fields.PipelineDepthBits = (F2I.I >> 23) & 0xff; // 8 bit exponent
 	Key.Fields.PipelineId = PipelineId;
 	Key.Fields.StateBucketId = StateBucketId;
-	Key.Fields.DepthBits = int32((FMath::Min<float>(PrimitiveDistance, HALF_WORLD_MAX) / HALF_WORLD_MAX) * PrimitiveDepthQuantization);
+	Key.Fields.DepthBits = int32((FMath::Min<float>(PrimitiveDistance, PrimitiveMaxDepth) / PrimitiveMaxDepth) * PrimitiveDepthQuantization);
 
 	return Key.PackedInt;
 }
