Fix crash on Vulkan when VRS is used (DDPI mismatch with RHI feature-level support, shader not compiled).

diff --git a/Engine/Source/Runtime/Renderer/Private/VariableRateShadingImageManager.cpp b/Engine/Source/Runtime/Renderer/Private/VariableRateShadingImageManager.cpp
--- a/Engine/Source/Runtime/Renderer/Private/VariableRateShadingImageManager.cpp
+++ b/Engine/Source/Runtime/Renderer/Private/VariableRateShadingImageManager.cpp
@@ -173,7 +173,7 @@ void FVariableRateShadingImageManager::ReleaseDynamicRHI()
 FRDGTextureRef FVariableRateShadingImageManager::GetVariableRateShadingImage(FRDGBuilder& GraphBuilder, const FSceneViewFamily& ViewFamily, const TArray<TRefCountPtr<IPooledRenderTarget>>* ExternalVRSSources, EVRSType VRSTypesToExclude)
 {
 	// If the RHI doesn't support VRS, we should always bail immediately.
-	if (!GRHISupportsAttachmentVariableRateShading || !GRHIVariableRateShadingEnabled || !GRHIAttachmentVariableRateShadingEnabled)
+	if (!GRHISupportsAttachmentVariableRateShading || !GRHIVariableRateShadingEnabled || !GRHIAttachmentVariableRateShadingEnabled || !FDataDrivenShaderPlatformInfo::GetSupportsVariableRateShading(GMaxRHIShaderPlatform))
 	{
 		return nullptr;
 	}
