Fix failing hair & fur EngineTests.

Reference stencil value was invalid for static directional light.

diff --git a/Engine/Source/Runtime/Renderer/Private/ShadowRendering.cpp b/Engine/Source/Runtime/Renderer/Private/ShadowRendering.cpp
--- a/Engine/Source/Runtime/Renderer/Private/ShadowRendering.cpp
+++ b/Engine/Source/Runtime/Renderer/Private/ShadowRendering.cpp
@@ -1240,7 +1240,7 @@ void FProjectedShadowInfo::RenderProjectionInternal(
 	const bool bSubPixelSupport = HairStrandsUniformBuffer != nullptr;// HairStrands::HasViewHairStrandsData(*View);
 	const bool bStencilTestEnabled = !bSubPixelSupport && GShadowStencilCulling;
 	const bool bDepthBoundsTestEnabled = IsWholeSceneDirectionalShadow() && GSupportsDepthBoundsTest && CVarCSMDepthBoundsTest.GetValueOnRenderThread() != 0 && !bSubPixelSupport;
-	const uint32 StencilRef = bSubPixelSupport ? 1u : 0u;
+	const uint32 StencilRef = bSubPixelSupport && !IsWholeSceneDirectionalShadow() ? 1u : 0u;
 
 	if (!bDepthBoundsTestEnabled && bStencilTestEnabled)
 	{
