Fix a crash bug of using High Resolution Screenshot with "Use custom depth as mask" enabled on mobile preview

diff --git a/Engine/Source/Runtime/Renderer/Private/PostProcess/PostProcessing.cpp b/Engine/Source/Runtime/Renderer/Private/PostProcess/PostProcessing.cpp
--- a/Engine/Source/Runtime/Renderer/Private/PostProcess/PostProcessing.cpp
+++ b/Engine/Source/Runtime/Renderer/Private/PostProcess/PostProcessing.cpp
@@ -2409,7 +2409,11 @@ void AddMobilePostProcessingPasses(FRDGBuilder& GraphBuilder, FScene* Scene, con
 	{
 		FHighResolutionScreenshotMaskInputs HighResolutionScreenshotMaskInputs;
 		HighResolutionScreenshotMaskInputs.SceneColor = SceneColor;
-		PassSequence.AcceptOverrideIfLastPass(EPass::Tonemap, HighResolutionScreenshotMaskInputs.OverrideOutput);
+		HighResolutionScreenshotMaskInputs.SceneTextures = GetSceneTextureShaderParameters(Inputs.SceneTextures);
+		HighResolutionScreenshotMaskInputs.Material = View.FinalPostProcessSettings.HighResScreenshotMaterial;
+		HighResolutionScreenshotMaskInputs.MaskMaterial = View.FinalPostProcessSettings.HighResScreenshotMaskMaterial;
+		HighResolutionScreenshotMaskInputs.CaptureRegionMaterial = View.FinalPostProcessSettings.HighResScreenshotCaptureRegionMaterial;
+		PassSequence.AcceptOverrideIfLastPass(EPass::HighResolutionScreenshotMask, HighResolutionScreenshotMaskInputs.OverrideOutput);
 		HighResolutionScreenshotMaskInputs.OverrideOutput.LoadAction = View.IsFirstInFamily() ? ERenderTargetLoadAction::EClear : ERenderTargetLoadAction::ELoad;
 
 		SceneColor = AddHighResolutionScreenshotMaskPass(GraphBuilder, View, HighResolutionScreenshotMaskInputs);
