Fix crash on mobile when Strata is disabled.

diff --git a/Engine/Source/Runtime/Renderer/Private/PostProcess/PostProcessMaterial.cpp b/Engine/Source/Runtime/Renderer/Private/PostProcess/PostProcessMaterial.cpp
--- a/Engine/Source/Runtime/Renderer/Private/PostProcess/PostProcessMaterial.cpp
+++ b/Engine/Source/Runtime/Renderer/Private/PostProcess/PostProcessMaterial.cpp
@@ -553,8 +553,10 @@ FScreenPassTexture AddPostProcessMaterialPass(
 	}
 
 	PostProcessMaterialParameters->bFlipYAxis = Inputs.bFlipYAxis && !bForceIntermediateTarget;
-	PostProcessMaterialParameters->Strata = Strata::BindStrataGlobalUniformParameters(View.StrataSceneData);
-
+	if (Strata::IsStrataEnabled())
+	{
+		PostProcessMaterialParameters->Strata = Strata::BindStrataGlobalUniformParameters(View.StrataSceneData);
+	}
 
 	ClearUnusedGraphResources(VertexShader, PixelShader, PostProcessMaterialParameters);
 
