Fixes shader print drawing before spatial upscaler

This caused inconsistency between what ShaderPrint::BeginView() believe the resolution was and what ShaderPrint::DrawView() was actually drawing to

diff --git a/Engine/Source/Editor/UnrealEd/Private/EditorViewportClient.cpp b/Engine/Source/Editor/UnrealEd/Private/EditorViewportClient.cpp
--- a/Engine/Source/Editor/UnrealEd/Private/EditorViewportClient.cpp
+++ b/Engine/Source/Editor/UnrealEd/Private/EditorViewportClient.cpp
@@ -3826,6 +3826,8 @@ void FEditorViewportClient::Draw(FViewport* InViewport, FCanvas* Canvas)
 		.SetRealtimeUpdate( IsRealtime() && FSlateThrottleManager::Get().IsAllowingExpensiveTasks() )
 		.SetViewModeParam( ViewModeParam, ViewModeParamName ) );
 
+	ViewFamily.DebugDPIScale = GetDPIScale();
+
 	ViewFamily.EngineShowFlags = UseEngineShowFlags;
 
 	ViewFamily.bIsHDR = Viewport->IsHDRViewport();
diff --git a/Engine/Source/Runtime/Engine/Private/GameViewportClient.cpp b/Engine/Source/Runtime/Engine/Private/GameViewportClient.cpp
--- a/Engine/Source/Runtime/Engine/Private/GameViewportClient.cpp
+++ b/Engine/Source/Runtime/Engine/Private/GameViewportClient.cpp
@@ -1290,6 +1290,8 @@ void UGameViewportClient::Draw(FViewport* InViewport, FCanvas* SceneCanvas)
 		EngineShowFlags)
 		.SetRealtimeUpdate(true));
 
+	ViewFamily.DebugDPIScale = GetDPIScale();
+
 #if WITH_EDITOR
 	if (GIsEditor)
 	{
diff --git a/Engine/Source/Runtime/Engine/Public/SceneView.h b/Engine/Source/Runtime/Engine/Public/SceneView.h
--- a/Engine/Source/Runtime/Engine/Public/SceneView.h
+++ b/Engine/Source/Runtime/Engine/Public/SceneView.h
@@ -1736,6 +1736,9 @@ public:
 	/** Gamma correction used when rendering this family. Default is 1.0 */
 	float GammaCorrection;
 	
+	/** DPI scale to be used for debugging font. */
+	float DebugDPIScale = 1.0f;
+
 	/** Editor setting to allow designers to override the automatic expose. 0:Automatic, following indices: -4 .. +4 */
 	FExposureSettings ExposureSettings;
 
diff --git a/Engine/Source/Runtime/Renderer/Private/PostProcess/PostProcessing.cpp b/Engine/Source/Runtime/Renderer/Private/PostProcess/PostProcessing.cpp
--- a/Engine/Source/Runtime/Renderer/Private/PostProcess/PostProcessing.cpp
+++ b/Engine/Source/Runtime/Renderer/Private/PostProcess/PostProcessing.cpp
@@ -1347,34 +1347,6 @@ void AddPostProcessingPasses(FRDGBuilder& GraphBuilder, const FViewInfo& View, c
 		SceneColor = AddHMDDistortionPass(GraphBuilder, View, PassInputs);
 	}
 
-	if (EngineShowFlags.TestImage)
-	{
-		AddTestImagePass(GraphBuilder, View, SceneColor);
-	}
-
-	if (EngineShowFlags.VisualizeNanite && NaniteRasterResults != nullptr)
-	{
-		AddVisualizeNanitePass(GraphBuilder, View, SceneColor, *NaniteRasterResults);
-	}
-
-	if (ShaderDrawDebug::IsEnabled(View))
-	{
-		ShaderDrawDebug::DrawView(GraphBuilder, View, SceneColor.Texture, SceneDepth.Texture);
-	}
-
-	if (ShaderPrint::IsEnabled(View))
-	{
-		ShaderPrint::DrawView(GraphBuilder, View, SceneColor.Texture);
-	}
-
-	if ( View.Family && View.Family->Scene )
-	{
-		if (FFXSystemInterface* FXSystem = View.Family->Scene->GetFXSystem())
-		{
-			FXSystem->DrawSceneDebug_RenderThread(GraphBuilder, View, SceneColor.Texture, SceneDepth.Texture);
-		}
-	}
-
 	if (PassSequence.IsEnabled(EPass::HighResolutionScreenshotMask))
 	{
 		FHighResolutionScreenshotMaskInputs PassInputs;
@@ -1453,6 +1425,39 @@ void AddPostProcessingPasses(FRDGBuilder& GraphBuilder, const FViewInfo& View, c
 			SceneColor = ISpatialUpscaler::AddDefaultUpscalePass(GraphBuilder, View, PassInputs, Method, FPaniniProjectionConfig());
 		}
 	}
+
+	// Draw debug stuff directly onto the back buffer
+	#if !UE_BUILD_SHIPPING
+	{
+		if (EngineShowFlags.TestImage)
+		{
+			AddTestImagePass(GraphBuilder, View, SceneColor);
+		}
+
+		if (EngineShowFlags.VisualizeNanite && NaniteRasterResults != nullptr)
+		{
+			AddVisualizeNanitePass(GraphBuilder, View, SceneColor, *NaniteRasterResults);
+		}
+
+		if (ShaderDrawDebug::IsEnabled(View))
+		{
+			ShaderDrawDebug::DrawView(GraphBuilder, View, SceneColor.Texture, SceneDepth.Texture);
+		}
+
+		if (ShaderPrint::IsEnabled(View))
+		{
+			ShaderPrint::DrawView(GraphBuilder, View, SceneColor);
+		}
+
+		if (View.Family && View.Family->Scene)
+		{
+			if (FFXSystemInterface* FXSystem = View.Family->Scene->GetFXSystem())
+			{
+				FXSystem->DrawSceneDebug_RenderThread(GraphBuilder, View, SceneColor.Texture, SceneDepth.Texture);
+			}
+		}
+	}
+	#endif
 }
 
 void AddDebugViewPostProcessingPasses(FRDGBuilder& GraphBuilder, const FViewInfo& View, const FPostProcessingInputs& Inputs, const Nanite::FRasterResults* NaniteRasterResults)
diff --git a/Engine/Source/Runtime/Renderer/Private/ShaderPrint.cpp b/Engine/Source/Runtime/Renderer/Private/ShaderPrint.cpp
--- a/Engine/Source/Runtime/Renderer/Private/ShaderPrint.cpp
+++ b/Engine/Source/Runtime/Renderer/Private/ShaderPrint.cpp
@@ -98,7 +98,7 @@ namespace ShaderPrint
 	{		
 		FUniformBufferParameters Out;
 		Out.FontSize = Data.FontSize;
-		Out.Resolution = Data.Resolution;
+		Out.Resolution = Data.OutputRect.Size();
 		Out.MaxValueCount = Data.MaxValueCount;
 		Out.MaxSymbolCount = Data.MaxSymbolCount;
 		return FUniformBufferRef::CreateUniformBufferImmediate(Out, UniformBuffer_SingleFrame);
@@ -289,14 +289,18 @@ namespace ShaderPrint
 			return;
 		}
 
-		const FIntPoint Resolution(FMath::Max(View.UnconstrainedViewRect.Size().X, 1), FMath::Max(View.UnconstrainedViewRect.Size().Y, 1));
+		const FIntPoint ViewSize(FMath::Max(View.UnconstrainedViewRect.Size().X, 1), FMath::Max(View.UnconstrainedViewRect.Size().Y, 1));
 
-		const float FontWidth = (float)FMath::Max(CVarFontSize.GetValueOnRenderThread(), 1) / (float)Resolution.X;
-		const float FontHeight = (float)FMath::Max(CVarFontSize.GetValueOnRenderThread(), 1) / (float)Resolution.Y;
-		const float SpaceWidth = (float)FMath::Max(CVarFontSpacingX.GetValueOnRenderThread(), 1) / (float)Resolution.X;
-		const float SpaceHeight = (float)FMath::Max(CVarFontSpacingY.GetValueOnRenderThread(), 1) / (float)Resolution.Y;
+		float FontSize = float(FMath::Max(CVarFontSize.GetValueOnRenderThread(), 1)) * View.Family->DebugDPIScale;
+		float FontSpacingX = float(FMath::Max(CVarFontSpacingX.GetValueOnRenderThread(), 1)) * View.Family->DebugDPIScale;
+		float FontSpacingY = float(FMath::Max(CVarFontSpacingY.GetValueOnRenderThread(), 1)) * View.Family->DebugDPIScale;
+
+		const float FontWidth = FontSize / float(ViewSize.X);
+		const float FontHeight = FontSize / float(ViewSize.Y);
+		const float SpaceWidth = FontSpacingX / float(ViewSize.X);
+		const float SpaceHeight = FontSpacingY / float(ViewSize.Y);
 		View.ShaderPrintData.FontSize = FVector4f(FontWidth, FontHeight, SpaceWidth + FontWidth, SpaceHeight + FontHeight);
-		View.ShaderPrintData.Resolution = Resolution;
+		View.ShaderPrintData.OutputRect = View.UnconstrainedViewRect;
 		View.ShaderPrintData.MaxValueCount = GetMaxValueCount();
 		View.ShaderPrintData.MaxSymbolCount = GetMaxSymbolCount();
 		GCharacterRequestCount = 0;
@@ -330,12 +334,14 @@ namespace ShaderPrint
 			FIntVector(1, 1, 1));
 	}
 
-	void DrawView(FRDGBuilder& GraphBuilder, const FViewInfo& View, FRDGTextureRef OutputTexture)
+	void DrawView(FRDGBuilder& GraphBuilder, const FViewInfo& View, FScreenPassTexture OutputTexture)
 	{
-		check(OutputTexture);
+		check(OutputTexture.IsValid());
 
 		RDG_EVENT_SCOPE(GraphBuilder, "ShaderPrintDrawView");
 
+		FIntRect OutputRect = OutputTexture.ViewRect;
+	
 		// Initialize graph managed resources
 		// Symbols buffer contains Count + 1 elements. The first element is only used as a counter.
 		FRDGBufferRef SymbolBuffer = GraphBuilder.CreateBuffer(FRDGBufferDesc::CreateStructuredDesc(sizeof(ShaderPrintItem), GetMaxSymbolCount() + 1), TEXT("ShaderPrintSymbolBuffer"));
@@ -410,7 +416,7 @@ namespace ShaderPrint
 			TShaderMapRef< FShaderDrawSymbolsPS > PixelShader(GlobalShaderMap);
 
 			SHADER::FParameters* PassParameters = GraphBuilder.AllocParameters<SHADER::FParameters>();
-			PassParameters->RenderTargets[0] = FRenderTargetBinding(OutputTexture, ERenderTargetLoadAction::ELoad);
+			PassParameters->RenderTargets[0] = FRenderTargetBinding(OutputTexture.Texture, ERenderTargetLoadAction::ELoad);
 			PassParameters->UniformBufferParameters = UniformBuffer;
 			PassParameters->MiniFontTexture = FontTexture;
 			PassParameters->SymbolsBuffer = GraphBuilder.CreateSRV(SymbolBuffer);
@@ -420,10 +426,12 @@ namespace ShaderPrint
 				RDG_EVENT_NAME("DrawSymbols"),
 				PassParameters,
 				ERDGPassFlags::Raster,
-				[VertexShader, PixelShader, PassParameters](FRHICommandList& RHICmdListImmediate)
+				[VertexShader, PixelShader, PassParameters, OutputRect](FRHICommandList& RHICmdList)
 			{
+				RHICmdList.SetViewport(OutputRect.Min.X, OutputRect.Min.Y, 0.0f, OutputRect.Min.X, OutputRect.Min.Y, 1.0f);
+				
 				FGraphicsPipelineStateInitializer GraphicsPSOInit;
-				RHICmdListImmediate.ApplyCachedRenderTargets(GraphicsPSOInit);
+				RHICmdList.ApplyCachedRenderTargets(GraphicsPSOInit);
 				GraphicsPSOInit.DepthStencilState = TStaticDepthStencilState<false, CF_Always>::GetRHI();
 				GraphicsPSOInit.BlendState = TStaticBlendState<CW_RGBA, BO_Add, BF_One, BF_InverseSourceAlpha, BO_Add, BF_Zero, BF_One>::GetRHI();
 				GraphicsPSOInit.RasterizerState = TStaticRasterizerState<>::GetRHI();
@@ -431,12 +439,12 @@ namespace ShaderPrint
 				GraphicsPSOInit.BoundShaderState.VertexDeclarationRHI = GetVertexDeclarationFVector4();
 				GraphicsPSOInit.BoundShaderState.VertexShaderRHI = VertexShader.GetVertexShader();
 				GraphicsPSOInit.BoundShaderState.PixelShaderRHI = PixelShader.GetPixelShader();
-				SetGraphicsPipelineState(RHICmdListImmediate, GraphicsPSOInit, 0);
+				SetGraphicsPipelineState(RHICmdList, GraphicsPSOInit, 0);
 
-				SetShaderParameters(RHICmdListImmediate, VertexShader, VertexShader.GetVertexShader(), *PassParameters);
-				SetShaderParameters(RHICmdListImmediate, PixelShader, PixelShader.GetPixelShader(), *PassParameters);
+				SetShaderParameters(RHICmdList, VertexShader, VertexShader.GetVertexShader(), *PassParameters);
+				SetShaderParameters(RHICmdList, PixelShader, PixelShader.GetPixelShader(), *PassParameters);
 
-				RHICmdListImmediate.DrawIndexedPrimitiveIndirect(GTwoTrianglesIndexBuffer.IndexBufferRHI, PassParameters->IndirectDrawArgsBuffer->GetIndirectRHICallBuffer(), 0);
+				RHICmdList.DrawIndexedPrimitiveIndirect(GTwoTrianglesIndexBuffer.IndexBufferRHI, PassParameters->IndirectDrawArgsBuffer->GetIndirectRHICallBuffer(), 0);
 			});
 		}
 	}
diff --git a/Engine/Source/Runtime/Renderer/Private/ShaderPrint.h b/Engine/Source/Runtime/Renderer/Private/ShaderPrint.h
--- a/Engine/Source/Runtime/Renderer/Private/ShaderPrint.h
+++ b/Engine/Source/Runtime/Renderer/Private/ShaderPrint.h
@@ -12,6 +12,7 @@
 
 #include "RenderGraph.h"
 #include "Templates/RefCounting.h"
+#include "ScreenPass.h"
 
 class FViewInfo;
 
@@ -32,7 +33,7 @@ namespace ShaderPrint
 	// Allocate the debug print buffer associated with the view
 	void BeginView(FRDGBuilder& GraphBuilder, FViewInfo& View);
 	// Draw info from the debug print buffer to the given output target
-	void DrawView(FRDGBuilder& GraphBuilder, const FViewInfo& View, FRDGTextureRef OutputTexture);
+	void DrawView(FRDGBuilder& GraphBuilder, const FViewInfo& View, FScreenPassTexture OutputTexture);
 	// Release the debug print buffer associated with the view
 	void EndView(FViewInfo& View);
 }
diff --git a/Engine/Source/Runtime/Renderer/Public/ShaderPrintParameters.h b/Engine/Source/Runtime/Renderer/Public/ShaderPrintParameters.h
--- a/Engine/Source/Runtime/Renderer/Public/ShaderPrintParameters.h
+++ b/Engine/Source/Runtime/Renderer/Public/ShaderPrintParameters.h
@@ -20,7 +20,7 @@ class FViewInfo;
 struct FShaderPrintData
 {
 	FVector4f FontSize;
-	FIntPoint Resolution;
+	FIntRect OutputRect;
 	int32 MaxValueCount = -1;
 	int32 MaxSymbolCount = -1;
 	FRDGBufferRef ShaderPrintValueBuffer = nullptr;
