Add persistent primitive index (FPrimitiveSceneInfo::PersistentIndex) which is stable for the life-time of a primitive (proxy) in FScene
- can be used to track information about a primitive over time without needing to track changes in PackedPrimitiveIndex
- will have the same high-water mark to packed index, but will persist for longer, potentially (allocator picks new from the front).
- Add PersistentPrimitiveIndex to FPrimitiveSceneData / SceneData.ush, to enable access on the GPU
- Add FSpanAllocator, an incremental improvement over FGrowOnlySpanAllocator - which shrinks as well and has efficient allocation of unit-sized allocations.

diff --git a/Engine/Shaders/Private/SceneData.ush b/Engine/Shaders/Private/SceneData.ush
--- a/Engine/Shaders/Private/SceneData.ush
+++ b/Engine/Shaders/Private/SceneData.ush
@@ -102,6 +102,7 @@ struct FPrimitiveSceneData
 	uint		Flags; // TODO: Use 16 bits?
 	int			InstanceSceneDataOffset; // Link to the range of instances that belong to this primitive
 	int			NumInstanceSceneDataEntries;
+	int			PersistentPrimitiveIndex;
 	uint		SingleCaptureIndex; // TODO: Use 16 bits? 8 bits?
 	float3		TilePosition;
 	FLWCMatrix	LocalToWorld;
@@ -170,6 +171,7 @@ FPrimitiveSceneData GetPrimitiveDataFromUniformBuffer()
 	PrimitiveData.InstancePayloadDataStride		= Primitive.InstancePayloadDataStride;
 	PrimitiveData.WireframeColor				= Primitive.WireframeColor;
 	PrimitiveData.LevelColor					= Primitive.LevelColor;
+	PrimitiveData.PersistentPrimitiveIndex		= Primitive.PersistentPrimitiveIndex;
 	UNROLL
 	for (int DataIndex = 0; DataIndex < NUM_CUSTOM_PRIMITIVE_DATA; ++DataIndex)
 	{
@@ -279,6 +281,7 @@ FPrimitiveSceneData GetPrimitiveData(uint PrimitiveId)
 	PrimitiveData.WireframeColor				= LoadPrimitivePrimitiveSceneDataElement(PrimitiveIndex, 29).xyz;
 	PrimitiveData.NaniteImposterIndex			= asuint(LoadPrimitivePrimitiveSceneDataElement(PrimitiveIndex, 29).w);
 	PrimitiveData.LevelColor					= LoadPrimitivePrimitiveSceneDataElement(PrimitiveIndex, 30).xyz;
+	PrimitiveData.PersistentPrimitiveIndex		= asint(LoadPrimitivePrimitiveSceneDataElement(PrimitiveIndex, 30).w);
 
 	// TODO: Move to associated array (and editor data) to shrink primitive data and better pack cachelines
 	UNROLL
diff --git a/Engine/Source/Runtime/Engine/Private/PrimitiveUniformShaderParameters.cpp b/Engine/Source/Runtime/Engine/Private/PrimitiveUniformShaderParameters.cpp
--- a/Engine/Source/Runtime/Engine/Private/PrimitiveUniformShaderParameters.cpp
+++ b/Engine/Source/Runtime/Engine/Private/PrimitiveUniformShaderParameters.cpp
@@ -144,6 +144,7 @@ FPrimitiveSceneShaderData::FPrimitiveSceneShaderData(const FPrimitiveSceneProxy*
 		.LightmapDataIndex(Proxy->GetPrimitiveSceneInfo()->GetLightmapDataOffset())
 		.LightmapUVIndex(Proxy->GetLightMapCoordinateIndex())
 		.SingleCaptureIndex(SingleCaptureIndex)
+		.PersistentPrimitiveIndex(Proxy->GetPrimitiveSceneInfo()->GetPersistentIndex().Index)
 		.InstanceSceneDataOffset(Proxy->GetPrimitiveSceneInfo()->GetInstanceSceneDataOffset())
 		.NumInstanceSceneDataEntries(Proxy->GetPrimitiveSceneInfo()->GetNumInstanceSceneDataEntries())
 		.InstancePayloadDataOffset(Proxy->GetPrimitiveSceneInfo()->GetInstancePayloadDataOffset())
@@ -247,7 +248,7 @@ void FPrimitiveSceneShaderData::Setup(const FPrimitiveUniformShaderParameters& P
 	Data[30].X = PrimitiveUniformShaderParameters.LevelColor.X;
 	Data[30].Y = PrimitiveUniformShaderParameters.LevelColor.Y;
 	Data[30].Z = PrimitiveUniformShaderParameters.LevelColor.Z;
-	Data[30].W = 0.0f; // Unused
+	Data[30].W = FMath::AsFloat(uint32(PrimitiveUniformShaderParameters.PersistentPrimitiveIndex));
 
 	// Set all the custom primitive data float4. This matches the loop in SceneData.ush
 	const int32 CustomPrimitiveDataStartIndex = 31;
diff --git a/Engine/Source/Runtime/Engine/Private/SpanAllocator.cpp b/Engine/Source/Runtime/Engine/Private/SpanAllocator.cpp
new file mode 100644
--- /dev/null
+++ b/Engine/Source/Runtime/Engine/Private/SpanAllocator.cpp
@@ -0,0 +1,85 @@
+// Copyright Epic Games, Inc. All Rights Reserved.
+
+#include "SpanAllocator.h"
+
+void FSpanAllocator::Consolidate()
+{
+	// Consolidation 
+
+	if (PendingFreeSpans.IsEmpty() && FirstNonEmptySpan == 0)
+	{
+		return;
+	}
+
+	// 1. Sort the Newly free list by span start, the existing free list is already sorted by construction
+	PendingFreeSpans.Sort();
+
+	// alternate free list, used during consolidation to avoid N^2 worst case, retained to avoid re-allocations.
+	TArray<FLinearAllocation, TInlineAllocator<10>> FreeSpansTmp;
+	FreeSpansTmp.Reset(FreeSpans.Num());
+
+	int32 PrevEndOffset = INDEX_NONE;
+	int32 PendingFreeIndex = 0;
+
+	// 2. Joint loop and merge over both free list and newly freed and fuse all adjacent, copy into new free list (to avoid compaction)
+	for (int32 Index = 0; Index < FreeSpans.Num() || PendingFreeIndex < PendingFreeSpans.Num(); )
+	{
+		FLinearAllocation Alloc = Index < FreeSpans.Num() ? FreeSpans[Index] : FLinearAllocation{ MAX_int32, 0 };
+		// Make sure we don't run out of both somehow...
+		check(PendingFreeIndex < PendingFreeSpans.Num() || Alloc.StartOffset < MAX_int32);
+		// Consume the next new alloc if it is before the next old one
+		if (PendingFreeIndex < PendingFreeSpans.Num() && PendingFreeSpans[PendingFreeIndex].StartOffset < Alloc.StartOffset)
+		{
+			Alloc = PendingFreeSpans[PendingFreeIndex++];
+		}
+		else
+		{
+			// Otherwise advance the old allocs
+			++Index;
+		}
+		check(Alloc.StartOffset < MAX_int32);
+
+		// Discard empty allocs
+		if (Alloc.Num > 0)
+		{
+			// Continues the previous one, fuse
+			if (PrevEndOffset == Alloc.StartOffset)
+			{
+				FreeSpansTmp.Last().Num += Alloc.Num;
+			}
+			else
+			{
+				FreeSpansTmp.Add(Alloc);
+			}
+			PrevEndOffset = FreeSpansTmp.Last().Num + FreeSpansTmp.Last().StartOffset;
+		}
+	}
+
+	// Trim last span
+	if (!FreeSpansTmp.IsEmpty() && FreeSpansTmp.Last().StartOffset + FreeSpansTmp.Last().Num == MaxSize)
+	{
+		MaxSize -= FreeSpansTmp.Last().Num;
+		FreeSpansTmp.Pop(false);
+	}
+
+	// 3. Store new free list
+	FreeSpans = MoveTemp(FreeSpansTmp);
+	PendingFreeSpans.Empty(PendingFreeSpans.Num());
+	FirstNonEmptySpan = 0;
+}
+
+int32 FSpanAllocator::SearchFreeList(int32 Num, int32 SearchStartIndex)
+{
+	// Search free list for first matching
+	for (int32 Index = SearchStartIndex; Index < FreeSpans.Num(); Index++)
+	{
+		FLinearAllocation CurrentSpan = FreeSpans[Index];
+
+		if (CurrentSpan.Num >= Num)
+		{
+			return Index;
+		}
+	}
+
+	return INDEX_NONE;
+}
\ No newline at end of file
diff --git a/Engine/Source/Runtime/Engine/Public/PrimitiveUniformShaderParameters.h b/Engine/Source/Runtime/Engine/Public/PrimitiveUniformShaderParameters.h
--- a/Engine/Source/Runtime/Engine/Public/PrimitiveUniformShaderParameters.h
+++ b/Engine/Source/Runtime/Engine/Public/PrimitiveUniformShaderParameters.h
@@ -53,7 +53,7 @@ BEGIN_GLOBAL_SHADER_PARAMETER_STRUCT(FPrimitiveUniformShaderParameters,ENGINE_AP
 	SHADER_PARAMETER(FVector3f,		WireframeColor)											// Only needed for editor/development
 	SHADER_PARAMETER(uint32,		NaniteImposterIndex)
 	SHADER_PARAMETER(FVector3f,		LevelColor)												// Only needed for editor/development
-	SHADER_PARAMETER(uint32,		Padding)
+	SHADER_PARAMETER(int32,			PersistentPrimitiveIndex)
 	SHADER_PARAMETER_ARRAY(FVector4f, CustomPrimitiveData, [FCustomPrimitiveData::NumCustomPrimitiveDataFloat4s]) // Custom data per primitive that can be accessed through material expression parameters and modified through UStaticMeshComponent
 END_GLOBAL_SHADER_PARAMETER_STRUCT()
 
@@ -106,6 +106,7 @@ public:
 		Parameters.LightmapDataIndex				= INDEX_NONE;
 		Parameters.LightmapUVIndex					= INDEX_NONE;
 		Parameters.SingleCaptureIndex				= INDEX_NONE;
+		Parameters.PersistentPrimitiveIndex			= INDEX_NONE;
 
 		// Nanite
 		Parameters.NaniteResourceID					= INDEX_NONE;
@@ -145,6 +146,7 @@ public:
 	PRIMITIVE_UNIFORM_BUILDER_METHOD(uint32,			InstancePayloadDataOffset);
 	PRIMITIVE_UNIFORM_BUILDER_METHOD(uint32,			InstancePayloadDataStride);
 	PRIMITIVE_UNIFORM_BUILDER_METHOD(int32,				SingleCaptureIndex);
+	PRIMITIVE_UNIFORM_BUILDER_METHOD(int32,				PersistentPrimitiveIndex);
 	PRIMITIVE_UNIFORM_BUILDER_METHOD(uint32,			NaniteResourceID);
 	PRIMITIVE_UNIFORM_BUILDER_METHOD(uint32,			NaniteHierarchyOffset);
 	PRIMITIVE_UNIFORM_BUILDER_METHOD(uint32,			NaniteImposterIndex);
diff --git a/Engine/Source/Runtime/Engine/Public/SpanAllocator.h b/Engine/Source/Runtime/Engine/Public/SpanAllocator.h
new file mode 100644
--- /dev/null
+++ b/Engine/Source/Runtime/Engine/Public/SpanAllocator.h
@@ -0,0 +1,146 @@
+// Copyright Epic Games, Inc. All Rights Reserved.
+
+#pragma once
+
+#include "CoreMinimal.h"
+
+
+/**
+ * Allocator for spans from some range that is allowed to grow and shrink to accomodate the allocations.
+ * Implementation is biased towards the case where allocations and frees are performed in bulk, if they are interleaved performance will degrade.
+ * Prefers allocations at the start of the range, such that the high-watermark can be reduced optimally.
+ * Primarily optimized for the use-case where span size == 1, for variable span size allocation performance is sub-optimal as it does not accelerate the search for larger allocs.
+ */
+class FSpanAllocator
+{
+public:
+	FSpanAllocator() 
+		: MaxSize(0)
+		, FirstNonEmptySpan(0)
+	{}
+
+	// Allocate a range.  Returns allocated StartOffset.
+	inline int32 Allocate(int32 Num = 1)
+	{
+		// Auto consolidate if free has been called
+		if (!PendingFreeSpans.IsEmpty())
+		{
+			Consolidate();
+		}
+
+		const int32 FoundIndex = SearchFreeList(Num, FirstNonEmptySpan);
+
+		// Use an existing free span if one is found
+		if (FoundIndex != INDEX_NONE)
+		{
+			FLinearAllocation FreeSpan = FreeSpans[FoundIndex];
+
+			// Update existing free span with remainder, 
+			// note: may become zero, not removing empty spans to avoid moving free list, will be cleaned out when consolidating.
+			FreeSpans[FoundIndex] = FLinearAllocation{ FreeSpan.StartOffset + Num, FreeSpan.Num - Num };
+
+			// If this span is now empty && the found index was the first non-empty, we update the first non-empty span index past this.
+			if (FreeSpan.Num == Num && FirstNonEmptySpan == FoundIndex)
+			{
+				FirstNonEmptySpan = FoundIndex + 1;
+			}
+
+			return FreeSpan.StartOffset;
+		}
+
+		// New allocation
+		int32 StartOffset = MaxSize;
+		MaxSize = MaxSize + Num;
+
+		return StartOffset;
+	}
+
+	// Free an already allocated range.  
+	FORCEINLINE void Free(int32 BaseOffset, int32 Num = 1)
+	{
+		checkSlow(BaseOffset + Num <= MaxSize);
+		PendingFreeSpans.Add(FLinearAllocation{ BaseOffset, Num });
+	}
+
+	inline int32 GetSparselyAllocatedSize() const
+	{
+		int32 AllocatedSize = MaxSize;
+
+		for (int32 i = 0; i < FreeSpans.Num(); i++)
+		{
+			AllocatedSize -= FreeSpans[i].Num;
+		}
+
+		return AllocatedSize;
+	}
+
+	FORCEINLINE int32 GetMaxSize() const
+	{
+		return MaxSize;
+	}
+
+	FORCEINLINE int32 GetNumFreeSpans() const
+	{
+		return FreeSpans.Num() + PendingFreeSpans.Num();
+	}
+
+	FORCEINLINE int32 GetNumPendingFreeSpans() const
+	{
+		return PendingFreeSpans.Num();
+	}
+
+#if DO_CHECK
+	/**
+	 * Loop over all free spans and check if Index is in any of them.
+	 * Note: Very costly, only intended for debugging use, and probably best if under a toggle even then.
+	 */
+	FORCEINLINE bool IsFree(int32 Index) const
+	{
+		// If outside the max size, it is considered free as the allocator can grow at will
+		if (Index >= MaxSize)
+		{
+			return true;
+		}
+
+		for (const auto& FreeSpan : FreeSpans)
+		{
+			if (Index >= FreeSpan.StartOffset && Index < FreeSpan.StartOffset + FreeSpan.Num)
+			{
+				return true;
+			}
+		}
+
+		return false;
+	}
+#endif // DO_CHECK
+
+	/**
+	 * Between these calls to Free just appends the allocation to the free list, rather than trying to merge with existing allocations.
+	 * At EndDeferMerges the free list is consolidated by sorting and merging all spans. This amortises the cost of the merge over many calls.
+	 */
+	ENGINE_API void Consolidate();
+
+private:
+	
+	struct FLinearAllocation
+	{
+		int32 StartOffset;
+		int32 Num;
+
+		inline bool operator<(const FLinearAllocation& Other) const
+		{
+			return StartOffset < Other.StartOffset;
+		}
+	};
+
+	// Size of the linear range used by the allocator
+	int32 MaxSize;
+	// First span in free list that is not empty, used to speed up search for free items
+	int32 FirstNonEmptySpan;
+	// Ordered free list from low to high
+	TArray<FLinearAllocation, TInlineAllocator<10>> FreeSpans;
+	// Unordered list of freed items since last consolidate
+	TArray<FLinearAllocation, TInlineAllocator<10>> PendingFreeSpans;
+
+	ENGINE_API int32 SearchFreeList(int32 Num, int32 SearchStartIndex = 0);
+};
\ No newline at end of file
diff --git a/Engine/Source/Runtime/Renderer/Private/GPUScene.h b/Engine/Source/Runtime/Renderer/Private/GPUScene.h
--- a/Engine/Source/Runtime/Renderer/Private/GPUScene.h
+++ b/Engine/Source/Runtime/Renderer/Private/GPUScene.h
@@ -291,6 +291,8 @@ public:
 	int32 GetNumPrimitives() const { return DynamicPrimitivesOffset; }
 	int32 GetNumLightmapDataItems() const { return LightmapDataAllocator.GetMaxSize(); }
 
+	const FGrowOnlySpanAllocator& GetInstanceSceneDataAllocator() const { return InstanceSceneDataAllocator; }
+
 	/**
 	 * Draw GPU-Scene debug info, such as bounding boxes. Call once per view at some point in the frame after GPU scene has been updated fully.
 	 * What is drawn is controlled by the CVar: r.GPUScene.DebugMode. Enabling this cvar causes ShaderDraw to be being active (if supported). 
diff --git a/Engine/Source/Runtime/Renderer/Private/PrimitiveSceneInfo.cpp b/Engine/Source/Runtime/Renderer/Private/PrimitiveSceneInfo.cpp
--- a/Engine/Source/Runtime/Renderer/Private/PrimitiveSceneInfo.cpp
+++ b/Engine/Source/Runtime/Renderer/Private/PrimitiveSceneInfo.cpp
@@ -197,6 +197,7 @@ FPrimitiveSceneInfo::FPrimitiveSceneInfo(UPrimitiveComponent* InComponent,FScene
 	CoarseMeshStreamingHandle(INDEX_NONE),
 #endif
 	PackedIndex(INDEX_NONE),
+	PersistentIndex(FPersistentPrimitiveIndex{ INDEX_NONE }),
 	ComponentForDebuggingOnly(InComponent),
 	bNeedsStaticMeshUpdateWithoutVisibilityCheck(false),
 	bNeedsUniformBufferUpdate(false),
diff --git a/Engine/Source/Runtime/Renderer/Private/RendererScene.cpp b/Engine/Source/Runtime/Renderer/Private/RendererScene.cpp
--- a/Engine/Source/Runtime/Renderer/Private/RendererScene.cpp
+++ b/Engine/Source/Runtime/Renderer/Private/RendererScene.cpp
@@ -6,6 +6,7 @@
 
 #include "CoreMinimal.h"
 #include "HAL/ThreadSafeCounter.h"
+#include "HAL/PlatformFileManager.h"
 #include "Stats/Stats.h"
 #include "HAL/IConsoleManager.h"
 #include "Misc/App.h"
@@ -146,6 +147,125 @@ static FThreadSafeCounter FSceneViewState_UniqueID;
 
 extern bool IsVelocityMergedWithDepthPass();
 
+#define ENABLE_LOG_PRIMITIVE_INSTANCE_ID_STATS_TO_CSV 0
+
+#if ENABLE_LOG_PRIMITIVE_INSTANCE_ID_STATS_TO_CSV
+
+int32 GDumpPrimitiveAllocatorStatsToCSV = 0;
+FAutoConsoleVariableRef CVarDumpPrimitiveStatsToCSV(
+	TEXT("r.DumpPrimitiveStatsToCSV"),
+	GDumpPrimitiveAllocatorStatsToCSV,
+	TEXT("Dump primitive and instance stats to CSV\n")
+	TEXT(" 0 - stop recording, dump to csv and clear array.\n")
+	TEXT(" 1 - start recording into array\n")
+	TEXT(" 2 - dump to csv and continue recording without clearing array\n"),
+	ECVF_RenderThreadSafe
+);
+
+static constexpr int32 StatStride = 8;
+TArray<int32> GPrimitiveAllocatorStats;
+
+void DumpPrimitiveAllocatorStats()
+{
+	if (!GPrimitiveAllocatorStats.IsEmpty())
+	{
+		FString FileName = FPaths::ProjectLogDir() / TEXT("PrimitiveStats-") + FDateTime::Now().ToString() + TEXT(".csv");
+
+		UE_LOG(LogRenderer, Log, TEXT("Dumping primitive allocator stats trace to: '%s'"), *FileName);
+
+		FArchive* FileToLogTo = IFileManager::Get().CreateFileWriter(*FileName, false);
+		
+		ensure(FileToLogTo);
+		if (FileToLogTo)
+		{
+			static const FString StatNames[StatStride] =
+			{
+				TEXT("NumPrimitives"),
+				TEXT("MaxPersistent"),
+				TEXT("NumPersistentAllocated"),
+				TEXT("PersistentFreeListSizeBC"),
+				TEXT("PersistentPendingListSizeBC"),
+				TEXT("PersistentFreeListSize"),
+				TEXT("MaxInstances"),
+				TEXT("NumInstancesAllocated"),
+			};
+
+			// Print header
+			FString StringToPrint;
+			for (int32 Index = 0; Index < StatStride; ++Index)
+			{
+				if (!StringToPrint.IsEmpty())
+				{
+					StringToPrint += TEXT(",");
+				}
+				if (Index < int32(UE_ARRAY_COUNT(StatNames)))
+				{
+					StringToPrint.Append(StatNames[Index]);
+				}
+				else
+				{
+					StringToPrint.Appendf(TEXT("Stat_%d"), Index);
+				}
+			}
+
+			StringToPrint += TEXT("\n");
+			FileToLogTo->Serialize(TCHAR_TO_ANSI(*StringToPrint), StringToPrint.Len());
+
+			int32 Num = GPrimitiveAllocatorStats.Num() / StatStride;
+			for (int32 Ind = 0; Ind < Num; ++Ind)
+			{
+				StringToPrint.Empty();
+
+				for (int32 StatInd = 0; StatInd < StatStride; ++StatInd)
+				{
+					if (!StringToPrint.IsEmpty())
+					{
+						StringToPrint.Append(TEXT(","));
+					}
+
+					StringToPrint.Appendf(TEXT("%d"), GPrimitiveAllocatorStats[Ind * StatStride + StatInd]);
+				}
+
+				StringToPrint += TEXT("\n");
+				FileToLogTo->Serialize(TCHAR_TO_ANSI(*StringToPrint), StringToPrint.Len());
+			}
+
+
+			FileToLogTo->Close();
+		}
+	}
+}
+
+void UpdatePrimitiveAllocatorStats(int32 NumPrimitives, int32 MaxPersistent, int32 NumPersistentAllocated, int32 PersistemFreeListSize, int32 PersistemFreeListSizeBC, int32 PersistentPendingListSizeBC, int32 MaxInstances, int32 NumInstancesAllocated)
+{
+	// Dump and clear when turning off the cvar
+	if (GDumpPrimitiveAllocatorStatsToCSV == 0 && !GPrimitiveAllocatorStats.IsEmpty())
+	{
+		DumpPrimitiveAllocatorStats();
+		GPrimitiveAllocatorStats.Empty();
+	}
+
+	// Dump snapshot (and don't clear or stop) when cvar is set to 2
+	if (GDumpPrimitiveAllocatorStatsToCSV == 2 && !GPrimitiveAllocatorStats.IsEmpty())
+	{
+		DumpPrimitiveAllocatorStats();
+		GDumpPrimitiveAllocatorStatsToCSV = 1;
+	}
+
+	if (GDumpPrimitiveAllocatorStatsToCSV != 0)
+	{
+		GPrimitiveAllocatorStats.Add(NumPrimitives);
+		GPrimitiveAllocatorStats.Add(MaxPersistent);
+		GPrimitiveAllocatorStats.Add(NumPersistentAllocated);
+		GPrimitiveAllocatorStats.Add(PersistemFreeListSizeBC);
+		GPrimitiveAllocatorStats.Add(PersistentPendingListSizeBC);
+		GPrimitiveAllocatorStats.Add(PersistemFreeListSize);
+		GPrimitiveAllocatorStats.Add(MaxInstances);
+		GPrimitiveAllocatorStats.Add(NumInstancesAllocated);
+	}
+}
+#endif // ENABLE_LOG_PRIMITIVE_INSTANCE_ID_STATS_TO_CSV
+
 /**
  * Holds the info to update SpeedTree wind per unique tree object in the scene, instead of per instance
  */
@@ -4347,6 +4467,8 @@ void FScene::UpdateAllPrimitiveSceneInfos(FRDGBuilder& GraphBuilder, bool bAsync
 						}
 					}
 				}
+
+				PersistentPrimitiveIdAllocator.Free(PrimitiveSceneInfo->PersistentIndex.Index);
 			}
 
 			for (TMap<int32, TArray<FCachedShadowMapData>>::TIterator CachedShadowMapIt(CachedShadowMaps); CachedShadowMapIt; ++CachedShadowMapIt)
@@ -4371,6 +4493,12 @@ void FScene::UpdateAllPrimitiveSceneInfos(FRDGBuilder& GraphBuilder, bool bAsync
 		CSV_SCOPED_TIMING_STAT_EXCLUSIVE(AddPrimitiveSceneInfos);
 		SCOPED_NAMED_EVENT(FScene_AddPrimitiveSceneInfos, FColor::Green);
 		SCOPE_CYCLE_COUNTER(STAT_AddScenePrimitiveRenderThreadTime);
+#if ENABLE_LOG_PRIMITIVE_INSTANCE_ID_STATS_TO_CSV
+		int32 PersistentPrimitiveFreeListSizeBeforeConsolidate = PersistentPrimitiveIdAllocator.GetNumFreeSpans();
+		int32 PersistentPrimitivePendingFreeListSizeBeforeConsolidate = PersistentPrimitiveIdAllocator.GetNumPendingFreeSpans();
+#endif
+		PersistentPrimitiveIdAllocator.Consolidate();
+
 		if (AddedLocalPrimitiveSceneInfos.Num())
 		{
 			Primitives.Reserve(Primitives.Num() + AddedLocalPrimitiveSceneInfos.Num());
@@ -4457,6 +4585,8 @@ void FScene::UpdateAllPrimitiveSceneInfos(FRDGBuilder& GraphBuilder, bool bAsync
 
 				const int32 SourceIndex = PrimitiveSceneProxies.Num() - 1;
 				PrimitiveSceneInfo->PackedIndex = SourceIndex;
+				checkSlow(PrimitiveSceneInfo->PersistentIndex.Index == INDEX_NONE);
+				PrimitiveSceneInfo->PersistentIndex = FPersistentPrimitiveIndex{ PersistentPrimitiveIdAllocator.Allocate() };
 
 				GPUScene.AddPrimitiveToUpdate(SourceIndex, EPrimitiveDirtyState::AddedMask);
 			}
@@ -4567,6 +4697,10 @@ void FScene::UpdateAllPrimitiveSceneInfos(FRDGBuilder& GraphBuilder, bool bAsync
 
 			CheckPrimitiveArrays();
 
+#if ENABLE_LOG_PRIMITIVE_INSTANCE_ID_STATS_TO_CSV
+			UpdatePrimitiveAllocatorStats(Primitives.Num(), PersistentPrimitiveIdAllocator.GetMaxSize(), PersistentPrimitiveIdAllocator.GetSparselyAllocatedSize(), PersistentPrimitiveIdAllocator.GetNumFreeSpans(), PersistentPrimitiveFreeListSizeBeforeConsolidate, PersistentPrimitivePendingFreeListSizeBeforeConsolidate, GPUScene.GetNumInstances(), GPUScene.GetInstanceSceneDataAllocator().GetSparselyAllocatedSize());
+#endif
+
 			for (int32 AddIndex = StartIndex; AddIndex < AddedLocalPrimitiveSceneInfos.Num(); AddIndex++)
 			{
 				FPrimitiveSceneInfo* PrimitiveSceneInfo = AddedLocalPrimitiveSceneInfos[AddIndex];
diff --git a/Engine/Source/Runtime/Renderer/Private/ScenePrivate.h b/Engine/Source/Runtime/Renderer/Private/ScenePrivate.h
--- a/Engine/Source/Runtime/Renderer/Private/ScenePrivate.h
+++ b/Engine/Source/Runtime/Renderer/Private/ScenePrivate.h
@@ -60,6 +60,7 @@
 #include "PrimitiveInstanceUpdateCommand.h"
 #include "OIT/OIT.h"
 #include "ShadingEnergyConservation.h"
+#include "SpanAllocator.h"
 
 /** Factor by which to grow occlusion tests **/
 #define OCCLUSION_SLOP (1.0f)
@@ -3013,6 +3014,8 @@ public:
 
 	const FReadOnlyCVARCache& ReadOnlyCVARCache;
 
+	FSpanAllocator PersistentPrimitiveIdAllocator;
+
 #if WITH_EDITOR
 	/** Editor Pixel inspector */
 	FPixelInspectorData PixelInspectorData;
@@ -3358,6 +3361,12 @@ public:
 
 	bool IsPrimitiveBeingRemoved(FPrimitiveSceneInfo* PrimitiveSceneInfo) const;
 
+	/**
+	 * Maximum used persistent Primitive Index, use to size arrays that store primitive data indexed by FPrimitiveSceneInfo::PersistentIndex.
+	 * Only changes during UpdateAllPrimitiveSceneInfos.
+	 */
+	inline int32 GetMaxPersistentPrimitiveIndex() const { return PersistentPrimitiveIdAllocator.GetMaxSize(); }
+
 protected:
 
 private:
diff --git a/Engine/Source/Runtime/Renderer/Public/PrimitiveSceneInfo.h b/Engine/Source/Runtime/Renderer/Public/PrimitiveSceneInfo.h
--- a/Engine/Source/Runtime/Renderer/Public/PrimitiveSceneInfo.h
+++ b/Engine/Source/Runtime/Renderer/Public/PrimitiveSceneInfo.h
@@ -265,6 +265,14 @@ enum class EUpdateStaticMeshFlags : uint8
 };
 ENUM_CLASS_FLAGS(EUpdateStaticMeshFlags);
 
+/**
+ * Wrapper to make it harder to confuse the packed and persistent index when used as arguments etc.
+ */
+struct FPersistentPrimitiveIndex
+{
+	int32 Index;
+};
+
 /**
  * The renderer's internal state for a single UPrimitiveComponent.  This has a one to one mapping with FPrimitiveSceneProxy, which is in the engine module.
  */
@@ -516,7 +524,19 @@ public:
 
 	/** Simple comparison against the invalid values used before/after scene add/remove. */
 	FORCEINLINE bool IsIndexValid() const { return PackedIndex != INDEX_NONE && PackedIndex != MAX_int32; }
-	
+
+	/**
+	 * Persistent index of the primitive in the range [0, FScene::GetMaxPersistentPrimitiveIndex() ). Where the max is never higher than the high watermark of the primitives in the scene.
+	 * This index remains stable for the life-time of the primitive in the scene (i.e., same as PrimitiveSceneInfo and Proxy - between FScene::AddPrimitive / FScene::RemovePrimitive). 
+	 * The intended use is to enable tracking primitive data over time in renderer systems using direct indexing, e.g., for efficiently storing a bit per primitive. 
+	 * Direct indexing also facilitates easy GPU-reflection and access of the persistent data, where it can be accessed as e.g., GetPrimitiveData(...).PersistentPrimitiveIndex.
+	 * It is allocated using FScene::PersistentPrimitiveIdAllocator from a range with holes that is kept as compact as possible up to the high-water mark of the scene primitives. 
+	 * Due to persistence this maximum can be substantially larger than the current number of Primitives at times, but never worse than the high-watermark.
+	 * In the future, the index will likely be refactored to persist for the lifetime of the component (to facilitate tracking data past proxy re-creates).
+	 * Note: It is not currently used to index any of the current FScene primitive arrays (use PackedIndex), though this is intended to change.
+	 */
+	RENDERER_API FORCEINLINE FPersistentPrimitiveIndex GetPersistentIndex() const { return PersistentIndex; }
+
 	/**
 	 * Shifts primitive position and all relevant data by an arbitrary delta.
 	 * Called on world origin changes
@@ -601,6 +621,11 @@ private:
 	 */
 	int32 PackedIndex;
 
+	/**
+	 * See GetPersistentIndex()
+	 */
+	FPersistentPrimitiveIndex PersistentIndex;
+
 	/** 
 	 * The UPrimitiveComponent this scene info is for, useful for quickly inspecting properties on the corresponding component while debugging.
 	 * This should not be dereferenced on the rendering thread.  The game thread can be modifying UObject members at any time.
