Add tracking of culled prims for cached VSM, to correctly and efficiently invalidate primitives that move (even if they are culled)
- Disable view-dependent CPU culling of primitives for local-light VSMs (to fix incorrect caching).
- Track per-light cache data for clipmaps to enable storing bit flags for rendered primitives.
- Make invalidations skip primitives that were never rendered into a clipmap, and mark invalidated primitives as not rendered.

diff --git a/Engine/Shaders/Private/VirtualShadowMaps/BuildPerPageDrawCommands.usf b/Engine/Shaders/Private/VirtualShadowMaps/BuildPerPageDrawCommands.usf
--- a/Engine/Shaders/Private/VirtualShadowMaps/BuildPerPageDrawCommands.usf
+++ b/Engine/Shaders/Private/VirtualShadowMaps/BuildPerPageDrawCommands.usf
@@ -55,6 +55,8 @@ RWBuffer<uint> DrawIndirectArgsBufferOut;
 RWStructuredBuffer<uint> OutInvalidatingInstances;
 uint NumInvalidatingInstanceSlots;
 
+StructuredBuffer<uint> PrimitiveRevealedMask;
+uint PrimitiveRevealedNum;
 
 #if USE_HZB_OCCLUSION
 // Mode 1 == previous frame HZB, 
@@ -62,6 +64,21 @@ uint NumInvalidatingInstanceSlots;
 uint HZBMode;
 #endif
 
+
+bool WasPrimitiveRevealed(uint PersistentPrimitiveId)
+{
+	// Use to filter out transient primitives as well as opting out of the test (if e.g., caching is off)
+	if (PersistentPrimitiveId >= PrimitiveRevealedNum)
+	{
+		return false;
+	}
+
+	uint PrimitiveIdWordOffset = PersistentPrimitiveId / 32U;
+	uint PrimitiveIdWordMask = 1U << (PersistentPrimitiveId % 32U);
+
+	return (PrimitiveRevealedMask[PrimitiveIdWordOffset] & PrimitiveIdWordMask) != 0U;
+}
+
 void MarkInstanceForInvalidation(uint InstanceId)
 {
 	// Note: InstanceId can be higher than NumInvalidatingInstanceSlots because of Dynamic primitives and their instances as they are added later in the frame
@@ -69,7 +86,7 @@ void MarkInstanceForInvalidation(uint InstanceId)
 	if (InstanceId < uint(NumInvalidatingInstanceSlots))
 	{
 		uint InstanceIdWordOffset = InstanceId / 32U;
-		uint InstanceIdWordMask = 1U << (InstanceId % 32);
+		uint InstanceIdWordMask = 1U << (InstanceId % 32U);
 
 		uint PreviousWordMask = 0U;
 		InterlockedOr(OutInvalidatingInstances[1 + NumInvalidatingInstanceSlots + InstanceIdWordOffset], InstanceIdWordMask, PreviousWordMask);
@@ -152,7 +169,14 @@ void CullPerPageDrawCommandsCs(uint3 GroupId : SV_GroupID, int GroupThreadIndex
 
 	// Load relevant instance data
 	FInstanceSceneData InstanceData = GetInstanceSceneData(InstanceId, InstanceSceneDataSOAStride);
-	const bool bHasMoved = GetGPUSceneFrameNumber() == InstanceData.LastUpdateSceneFrameNumber || DrawCommandDesc.bMaterialMayModifyPosition;
+
+	// A thing is considered to have moved (which causes it to render into cached pages) if:
+	// It was uploaded this frame, OR
+	const bool bHasMoved = GetGPUSceneFrameNumber() == InstanceData.LastUpdateSceneFrameNumber
+		// it has a material that uses WPO or PDO, OR
+		|| DrawCommandDesc.bMaterialMayModifyPosition
+		// the primitive was "revealed" this frame - i.e., transitioned from culled -> visible on CPU
+		|| WasPrimitiveRevealed(GetPrimitiveData(InstanceData.PrimitiveId).PersistentPrimitiveIndex);
 
 	uint ThreadTotalForAllViews = 0;
 	// Loop over views and output visible instance (i.e., those that overlap a valid page)
diff --git a/Engine/Shaders/Private/VirtualShadowMaps/CacheManagement.usf b/Engine/Shaders/Private/VirtualShadowMaps/CacheManagement.usf
--- a/Engine/Shaders/Private/VirtualShadowMaps/CacheManagement.usf
+++ b/Engine/Shaders/Private/VirtualShadowMaps/CacheManagement.usf
@@ -43,6 +43,9 @@ void VirtualSmInvalidateInstancePagesCS(
 	uint3 GroupId : SV_GroupID,
 	uint GroupThreadIndex : SV_GroupIndex)
 {
+	uint FirstVirtualShadowMapId = 0u;
+	uint NumVirtualShadowMapIds = VirtualShadowMap.NumShadowMaps;
+
 #if INPUT_KIND == INPUT_KIND_LOAD_BALANCER
 	FInstanceWorkSetup WorkSetup = InstanceCullingLoadBalancer_Setup(GroupId, GroupThreadIndex, 0U);
 
@@ -51,6 +54,14 @@ void VirtualSmInvalidateInstancePagesCS(
 		return;
 	}
 
+	bool bSkipClipmaps = (WorkSetup.Item.Payload & 1U) != 0U;
+	uint SingleShadowMapId = WorkSetup.Item.Payload >> 1U;
+	if (SingleShadowMapId != 0U)
+	{
+		FirstVirtualShadowMapId = SingleShadowMapId - 1U;
+		NumVirtualShadowMapIds = FirstVirtualShadowMapId + 1U;
+	}
+
 	uint InstanceId = WorkSetup.Item.InstanceDataOffset + uint(WorkSetup.LocalItemIndex);
 #else // INPUT_KIND == INPUT_KIND_GPU_INSTANCES
 	// The 0th index stores the total number appended
@@ -58,6 +69,8 @@ void VirtualSmInvalidateInstancePagesCS(
 	{
 		return;
 	}
+	bool bSkipClipmaps = false;
+
 	int InstanceId = InvalidatingInstances[1 + DispatchIndex];
 #endif
 	{
@@ -83,12 +96,22 @@ void VirtualSmInvalidateInstancePagesCS(
 		// TODO: test the flag on the instance instead once it is updated correctly InstanceSceneData.CastShadows
 		if (bCastShadows)
 		{
-			for (uint ShadowMapDataIndex = 0; ShadowMapDataIndex < VirtualShadowMap.NumShadowMaps; ++ShadowMapDataIndex)
+			for (uint VirtualShadowMapId = FirstVirtualShadowMapId; VirtualShadowMapId < NumVirtualShadowMapIds;)
 			{
 				// 1. Load cached projection data
-				FVirtualShadowMapProjectionShaderData ProjectionData = GetVirtualShadowMapProjectionData(ShadowMapDataIndex);
+				FVirtualShadowMapProjectionShaderData ProjectionData = GetVirtualShadowMapProjectionData(VirtualShadowMapId);
 				const bool bDirectionalLight = (ProjectionData.LightType == LIGHT_TYPE_DIRECTIONAL);
 
+				if (bSkipClipmaps && bDirectionalLight)
+				{
+					VirtualShadowMapId += uint(ProjectionData.ClipmapLevelCount);
+					continue;
+				}
+				else
+				{
+					++VirtualShadowMapId;
+				}
+
 				float4x4 LocalToTranslatedWorld = LWCHackToFloat(InstanceSceneData.LocalToWorld);
 				LocalToTranslatedWorld[3].xyz += ProjectionData.ShadowPreViewTranslation;
 
diff --git a/Engine/Source/Runtime/Renderer/Private/RendererScene.cpp b/Engine/Source/Runtime/Renderer/Private/RendererScene.cpp
--- a/Engine/Source/Runtime/Renderer/Private/RendererScene.cpp
+++ b/Engine/Source/Runtime/Renderer/Private/RendererScene.cpp
@@ -4212,7 +4212,7 @@ void FScene::UpdateAllPrimitiveSceneInfos(FRDGBuilder& GraphBuilder, bool bAsync
 
 	if (VirtualShadowMapArrayCacheManager && VirtualShadowMapArrayCacheManager->IsValid())
 	{
-		FVirtualShadowMapArrayCacheManager::FInvalidatingPrimitiveCollector InvalidatingPrimitiveCollector(Primitives.Num(), GPUScene);
+		FVirtualShadowMapArrayCacheManager::FInvalidatingPrimitiveCollector InvalidatingPrimitiveCollector(VirtualShadowMapArrayCacheManager);
 
 		// All removed primitives must invalidate their footprints in the VSM before leaving
 		for (const FPrimitiveSceneInfo* PrimitiveSceneInfo : RemovedLocalPrimitiveSceneInfos)
@@ -4870,6 +4870,9 @@ void FScene::UpdateAllPrimitiveSceneInfos(FRDGBuilder& GraphBuilder, bool bAsync
 		}
 	}
 
+	// handle scene changes
+	VirtualShadowMapArrayCacheManager->OnSceneChange();
+
 	{
 		CSV_SCOPED_TIMING_STAT_EXCLUSIVE(UpdatePrimitiveInstances);
 		SCOPED_NAMED_EVENT(FScene_UpdatePrimitiveInstances, FColor::Emerald);
diff --git a/Engine/Source/Runtime/Renderer/Private/ShadowSetup.cpp b/Engine/Source/Runtime/Renderer/Private/ShadowSetup.cpp
--- a/Engine/Source/Runtime/Renderer/Private/ShadowSetup.cpp
+++ b/Engine/Source/Runtime/Renderer/Private/ShadowSetup.cpp
@@ -2002,6 +2002,12 @@ void FProjectedShadowInfo::FinalizeAddSubjectPrimitive(
 		}
 	}
 
+	// If we have a clipmap, we need to track the state of the primitive such that we may invalidate pages when the culling changes (e.g., the primitive crosses a range boundary)
+	if (VirtualShadowMapClipmap.IsValid())
+	{
+		VirtualShadowMapClipmap->OnPrimitiveRendered(PrimitiveSceneInfo->GetPersistentIndex());
+	}
+
 	if (Result.bAddOnRenderThread)
 	{
 		AddSubjectPrimitive(PrimitiveSceneInfo, ViewArray, Result.bShouldRecordShadowSubjectsForMobile);
@@ -3660,7 +3666,7 @@ void FSceneRenderer::CreateWholeSceneProjectedShadow(
 				}
 
 				// Consolidate repeated operation into a single site.
-				auto AddInteractingPrimitives = [bStaticSceneOnly, &Views = this->Views, FeatureLevel = this->FeatureLevel, &LightViewFrustumConvexHulls](FLightPrimitiveInteraction* InteractionList, FProjectedShadowInfo* ProjectedShadowInfo, bool& bContainsNaniteSubjectsOut)
+				auto AddInteractingPrimitives = [bStaticSceneOnly, &Views = this->Views, FeatureLevel = this->FeatureLevel, &LightViewFrustumConvexHulls](FLightPrimitiveInteraction* InteractionList, FProjectedShadowInfo* ProjectedShadowInfo, bool& bContainsNaniteSubjectsOut, bool bTestConvexHull)
 				{
 					for (FLightPrimitiveInteraction* Interaction = InteractionList; Interaction; Interaction = Interaction->GetNextPrimitive())
 					{
@@ -3679,7 +3685,7 @@ void FSceneRenderer::CreateWholeSceneProjectedShadow(
 							else if (EnumHasAnyFlags(ProjectedShadowInfo->MeshSelectionMask, Interaction->ProxySupportsGPUScene() ? EShadowMeshSelection::VSM : EShadowMeshSelection::All))
 							{
 								FBoxSphereBounds const& Bounds = Interaction->GetPrimitiveSceneInfo()->Proxy->GetBounds();
-								if (IntersectsConvexHulls(LightViewFrustumConvexHulls, Bounds))
+								if (!bTestConvexHull || IntersectsConvexHulls(LightViewFrustumConvexHulls, Bounds))
 								{
 									ProjectedShadowInfo->AddSubjectPrimitive(Interaction->GetPrimitiveSceneInfo(), Views, false);
 								}
@@ -3734,8 +3740,10 @@ void FSceneRenderer::CreateWholeSceneProjectedShadow(
 					ProjectedShadowInfo->MeshSelectionMask = EShadowMeshSelection::VSM;
 
 					bool bContainsNaniteSubjects = false;
-					AddInteractingPrimitives(LightSceneInfo->GetDynamicInteractionOftenMovingPrimitiveList(), ProjectedShadowInfo, bContainsNaniteSubjects);
-					AddInteractingPrimitives(LightSceneInfo->GetDynamicInteractionStaticPrimitiveList(), ProjectedShadowInfo, bContainsNaniteSubjects);
+					// Skip convex hull tests for VSM since this causes artifacts due to caching (potentially check if caching is enabled but may lead to race condition).
+					// The interaction setup has already tested the light bounds (by calling FLightSceneProxy::AffectsBounds)
+					AddInteractingPrimitives(LightSceneInfo->GetDynamicInteractionOftenMovingPrimitiveList(), ProjectedShadowInfo, bContainsNaniteSubjects, false);
+					AddInteractingPrimitives(LightSceneInfo->GetDynamicInteractionStaticPrimitiveList(), ProjectedShadowInfo, bContainsNaniteSubjects, false);
 					ProjectedShadowInfo->bContainsNaniteSubjects = bContainsNaniteSubjects;
 
 					VisibleLightInfo.VirtualShadowMapId = ProjectedShadowInfo->VirtualShadowMaps[0]->ID;
@@ -3785,12 +3793,12 @@ void FSceneRenderer::CreateWholeSceneProjectedShadow(
 								&& (CacheMode[CacheModeIndex] != SDCM_MovablePrimitivesOnly || bCastCachedShadowFromMovablePrimitives))
 							{
 								// Add all the shadow casting primitives affected by the light to the shadow's subject primitive list.
-								AddInteractingPrimitives(LightSceneInfo->GetDynamicInteractionOftenMovingPrimitiveList(), ProjectedShadowInfo, bContainsNaniteSubjects);
+								AddInteractingPrimitives(LightSceneInfo->GetDynamicInteractionOftenMovingPrimitiveList(), ProjectedShadowInfo, bContainsNaniteSubjects, true);
 							}
 						
 							if (CacheMode[CacheModeIndex] != SDCM_MovablePrimitivesOnly)
 							{
-								AddInteractingPrimitives(LightSceneInfo->GetDynamicInteractionStaticPrimitiveList(), ProjectedShadowInfo, bContainsNaniteSubjects);
+								AddInteractingPrimitives(LightSceneInfo->GetDynamicInteractionStaticPrimitiveList(), ProjectedShadowInfo, bContainsNaniteSubjects, true);
 							}
 						}
 						ProjectedShadowInfo->bContainsNaniteSubjects = bContainsNaniteSubjects;
diff --git a/Engine/Source/Runtime/Renderer/Private/VirtualShadowMaps/VirtualShadowMapArray.cpp b/Engine/Source/Runtime/Renderer/Private/VirtualShadowMaps/VirtualShadowMapArray.cpp
--- a/Engine/Source/Runtime/Renderer/Private/VirtualShadowMaps/VirtualShadowMapArray.cpp
+++ b/Engine/Source/Runtime/Renderer/Private/VirtualShadowMaps/VirtualShadowMapArray.cpp
@@ -1862,6 +1862,9 @@ public:
 		SHADER_PARAMETER(uint32, InstanceSceneDataSOAStride)
 		SHADER_PARAMETER(uint32, GPUSceneFrameNumber)
 
+		SHADER_PARAMETER_RDG_BUFFER_SRV(StructuredBuffer<uint>, PrimitiveRevealedMask)
+		SHADER_PARAMETER(uint32, PrimitiveRevealedNum)
+
 		SHADER_PARAMETER_STRUCT_INCLUDE(FInstanceProcessingGPULoadBalancer::FShaderParameters, LoadBalancerParameters)
 
 		SHADER_PARAMETER(int32, FirstPrimaryView)
@@ -1994,7 +1997,9 @@ static FCullingResult AddCullingPasses(FRDGBuilder& GraphBuilder,
 	FRDGBufferRef VirtualShadowViewsRDG,
 	const FCullPerPageDrawCommandsCs::FHZBShaderParameters &HZBShaderParameters,
 	FVirtualShadowMapArray *VirtualShadowMapArray,
-	FGPUScene& GPUScene)
+	FGPUScene& GPUScene,
+	FRDGBufferRef PrimitiveRevealedMaskRdg,
+	int32 PrimitiveRevealedNum)
 {
 	int32 NumIndirectArgs = IndirectArgs.Num();
 
@@ -2037,6 +2042,11 @@ static FCullingResult AddCullingPasses(FRDGBuilder& GraphBuilder,
 		PassParameters->GPUSceneFrameNumber = GPUScene.GetSceneFrameNumber();
 		PassParameters->InstanceSceneDataSOAStride = GPUScene.InstanceSceneDataSOAStride;
 
+		// Make sure there is enough space in the buffer for all the primitive IDs that might be used to index.
+		check(PrimitiveRevealedMaskRdg->Desc.NumElements * 32u >= uint32(PrimitiveRevealedNum));
+		PassParameters->PrimitiveRevealedMask = GraphBuilder.CreateSRV(PrimitiveRevealedMaskRdg);
+		PassParameters->PrimitiveRevealedNum = uint32(PrimitiveRevealedNum);
+
 		PassParameters->DynamicInstanceIdOffset = BatchInfos[0].DynamicInstanceIdOffset;
 		PassParameters->DynamicInstanceIdMax = BatchInfos[0].DynamicInstanceIdMax;
 
@@ -2345,7 +2355,9 @@ void FVirtualShadowMapArray::RenderVirtualShadowMapsNonNanite(FRDGBuilder& Graph
 			VirtualShadowViewsRDG,
 			HZBShaderParameters,
 			this,
-			GPUScene
+			GPUScene,
+			GSystemTextures.GetDefaultStructuredBuffer(GraphBuilder, 4),
+			0
 		);
 		GraphBuilder.EndEventScope();
 
@@ -2400,6 +2412,17 @@ void FVirtualShadowMapArray::RenderVirtualShadowMapsNonNanite(FRDGBuilder& Graph
 			CullingBatchInfo.DynamicInstanceIdOffset = ShadowDepthView->DynamicPrimitiveCollector.GetInstanceSceneDataOffset();
 			CullingBatchInfo.DynamicInstanceIdMax = CullingBatchInfo.DynamicInstanceIdOffset + ShadowDepthView->DynamicPrimitiveCollector.NumInstances();
 
+			FRDGBufferRef PrimitiveRevealedMaskRdg = GSystemTextures.GetDefaultStructuredBuffer(GraphBuilder, 4);
+			int32 PrimitiveRevealedNum = 0;
+				
+
+			if (!Clipmap->GetRevealedPrimitivesMask().IsEmpty())
+			{
+				PrimitiveRevealedMaskRdg = CreateStructuredBuffer(GraphBuilder, TEXT("Shadow.Virtual.RevealedPrimitivesMask"), Clipmap->GetRevealedPrimitivesMask());
+				PrimitiveRevealedNum = Clipmap->GetNumRevealedPrimitives();
+
+			}
+
 			FCullingResult CullingResult = AddCullingPasses(
 				GraphBuilder,
 				InstanceCullingContext->IndirectArgs, 
@@ -2415,7 +2438,9 @@ void FVirtualShadowMapArray::RenderVirtualShadowMapsNonNanite(FRDGBuilder& Graph
 				VirtualShadowViewsRDG,
 				HZBShaderParameters,
 				this,
-				GPUScene
+				GPUScene,
+				PrimitiveRevealedMaskRdg,
+				PrimitiveRevealedNum
 			);
 
 			TRDGUniformBufferRef<FShadowDepthPassUniformParameters> ShadowDepthPassUniformBuffer = CreateShadowDepthPassUniformBuffer(ProjectedShadowInfo->ShouldClampToNearPlane());
diff --git a/Engine/Source/Runtime/Renderer/Private/VirtualShadowMaps/VirtualShadowMapCacheManager.cpp b/Engine/Source/Runtime/Renderer/Private/VirtualShadowMaps/VirtualShadowMapCacheManager.cpp
--- a/Engine/Source/Runtime/Renderer/Private/VirtualShadowMaps/VirtualShadowMapCacheManager.cpp
+++ b/Engine/Source/Runtime/Renderer/Private/VirtualShadowMaps/VirtualShadowMapCacheManager.cpp
@@ -4,6 +4,7 @@
 	VirtualShadowMap.h:
 =============================================================================*/
 #include "VirtualShadowMapCacheManager.h"
+#include "VirtualShadowMapClipmap.h"
 #include "RendererModule.h"
 #include "RenderGraphUtils.h"
 #include "RHIGPUReadback.h"
@@ -132,6 +133,76 @@ void FVirtualShadowMapCacheEntry::UpdateLocal(int32 VirtualShadowMapId, const FW
 	bCurrentRendered = false;
 }
 
+FVirtualShadowMapArrayCacheManager::FInvalidatingPrimitiveCollector::FInvalidatingPrimitiveCollector(FVirtualShadowMapArrayCacheManager* InVirtualShadowMapArrayCacheManager)
+	: AlreadyAddedPrimitives(false, InVirtualShadowMapArrayCacheManager->Scene->Primitives.Num())
+	, Scene(*InVirtualShadowMapArrayCacheManager->Scene)
+	, GPUScene(InVirtualShadowMapArrayCacheManager->Scene->GPUScene)
+	, Manager(*InVirtualShadowMapArrayCacheManager)
+{
+	// Add and clear pending invalidations enqueued on the GPU Scene from dynamic primitives added since last invalidation
+	for (const FGPUScene::FInstanceRange& Range : GPUScene.DynamicPrimitiveInstancesToInvalidate)
+	{
+		LoadBalancer.Add(Range.InstanceSceneDataOffset, Range.NumInstanceSceneDataEntries, 0U);
+#if VSM_LOG_INVALIDATIONS
+		RangesStr.Appendf(TEXT("[%6d, %6d), "), Range.InstanceSceneDataOffset, Range.InstanceSceneDataOffset + Range.NumInstanceSceneDataEntries);
+#endif
+		TotalInstanceCount += Range.NumInstanceSceneDataEntries;
+	}
+
+	GPUScene.DynamicPrimitiveInstancesToInvalidate.Reset();
+}
+void FVirtualShadowMapArrayCacheManager::FInvalidatingPrimitiveCollector::Add(const FPrimitiveSceneInfo * PrimitiveSceneInfo)
+{
+	int32 PrimitiveID = PrimitiveSceneInfo->GetIndex();
+	if (PrimitiveID >= 0
+		&& !AlreadyAddedPrimitives[PrimitiveID]
+		&& PrimitiveSceneInfo->GetInstanceSceneDataOffset() != INDEX_NONE
+		// Don't process primitives that are still in the 'added' state because this means that they
+		// have not been uploaded to the GPU yet and may be pending from a previous call to update primitive scene infos.
+		&& !EnumHasAnyFlags(GPUScene.GetPrimitiveDirtyState(PrimitiveID), EPrimitiveDirtyState::Added))
+	{
+		AlreadyAddedPrimitives[PrimitiveID] = true;
+		int32 PersistentPrimitiveIndex = PrimitiveSceneInfo->GetPersistentIndex().Index;
+
+		const int32 NumInstanceSceneDataEntries = PrimitiveSceneInfo->GetNumInstanceSceneDataEntries();
+		// Add for non-directional lights, mark for skipping clipmaps as these are handled individually below
+		LoadBalancer.Add(PrimitiveSceneInfo->GetInstanceSceneDataOffset(), NumInstanceSceneDataEntries, 1U);
+
+		// Nanite meshes need special handling because they don't get culled on CPU, thus always process invalidations for those
+		const bool bIsNaniteMesh = Scene.PrimitiveFlagsCompact[PrimitiveID].bIsNaniteMesh;
+		// Process directional lights, where we explicitly filter out primitives that were not rendered (and mark this fact)
+		for (auto& CacheEntry : Manager.PrevCacheEntries)
+		{
+			TBitArray<>& RenderedPrimitives = CacheEntry.Value->RenderedPrimitives;
+			if (bIsNaniteMesh || (PersistentPrimitiveIndex <= RenderedPrimitives.Num() && RenderedPrimitives[PersistentPrimitiveIndex]))
+			{
+				if (!bIsNaniteMesh)
+				{
+					// Clear the record as we're wiping it out.
+					RenderedPrimitives[PersistentPrimitiveIndex] = false;
+				}
+
+				// Add item for each shadow map explicitly, inflates host data but improves load balancing,
+				// TODO: maybe add permutation so we can strip the loop completely.
+				for (const auto& SmCacheEntry : CacheEntry.Value->ShadowMapEntries)
+				{
+					if (SmCacheEntry.IsValid())
+					{
+						// Lowest bit indicates whether to skip the clipmap loop, add 1 to ID so != 0 <==> single level processing
+						LoadBalancer.Add(PrimitiveSceneInfo->GetInstanceSceneDataOffset(), NumInstanceSceneDataEntries, (uint32(SmCacheEntry->CurrentVirtualShadowMapId + 1) << 1U) | 0U);
+					}
+				}
+			}
+		}
+#if VSM_LOG_INVALIDATIONS
+		RangesStr.Appendf(TEXT("[%6d, %6d), "), PrimitiveSceneInfo->GetInstanceSceneDataOffset(), PrimitiveSceneInfo->GetInstanceSceneDataOffset() + NumInstanceSceneDataEntries);
+#endif
+		TotalInstanceCount += NumInstanceSceneDataEntries;
+	}
+}
+
+
+
 FVirtualShadowMapArrayCacheManager::FVirtualShadowMapArrayCacheManager(FScene* InScene) 
 	: Scene(InScene)
 {
@@ -233,35 +304,60 @@ void FVirtualShadowMapArrayCacheManager::Invalidate()
 	CacheEntries.Reset();
 }
 
-TSharedPtr<FVirtualShadowMapCacheEntry> FVirtualShadowMapArrayCacheManager::FindCreateCacheEntry(int32 LightSceneId, int32 Index)
+TSharedPtr<FVirtualShadowMapCacheEntry> FVirtualShadowMapPerLightCacheEntry::FindCreateShadowMapEntry(int32 Index)
+{
+	check(Index >= 0);
+	ShadowMapEntries.SetNum(FMath::Max(Index + 1, ShadowMapEntries.Num()));
+
+	TSharedPtr<FVirtualShadowMapCacheEntry>& EntryRef = ShadowMapEntries[Index];
+
+	if (!EntryRef.IsValid())
+	{
+		EntryRef = MakeShared<FVirtualShadowMapCacheEntry>();
+	}
+
+	return EntryRef;
+}
+
+TSharedPtr<FVirtualShadowMapPerLightCacheEntry> FVirtualShadowMapArrayCacheManager::FindCreateLightCacheEntry(int32 LightSceneId)
 {
 	if (CVarCacheVirtualSMs.GetValueOnRenderThread() == 0)
 	{
 		return nullptr;
 	}
 
-	const FIntPoint Key(LightSceneId, Index);
-
-	if (TSharedPtr<FVirtualShadowMapCacheEntry> *VirtualShadowMapCacheEntry = CacheEntries.Find(Key))
+	if (TSharedPtr<FVirtualShadowMapPerLightCacheEntry> *LightEntry = CacheEntries.Find(LightSceneId))
 	{
-		return *VirtualShadowMapCacheEntry;
+		return *LightEntry;
 	}
 
 	// Add to current frame / active set.
-	TSharedPtr<FVirtualShadowMapCacheEntry> &NewVirtualShadowMapCacheEntry = CacheEntries.Add(Key);
+	TSharedPtr<FVirtualShadowMapPerLightCacheEntry>& NewLightEntry = CacheEntries.Add(LightSceneId);
 
 	// Copy data if available
-	if (TSharedPtr<FVirtualShadowMapCacheEntry> *PrevVirtualShadowMapCacheEntry = PrevCacheEntries.Find(Key))
+	if (TSharedPtr<FVirtualShadowMapPerLightCacheEntry>* PrevNewLightEntry = PrevCacheEntries.Find(LightSceneId))
 	{
-		NewVirtualShadowMapCacheEntry = *PrevVirtualShadowMapCacheEntry;
+		NewLightEntry = *PrevNewLightEntry;
 	}
 	else
 	{
-		NewVirtualShadowMapCacheEntry = TSharedPtr<FVirtualShadowMapCacheEntry>(new FVirtualShadowMapCacheEntry);
+		NewLightEntry = MakeShared<FVirtualShadowMapPerLightCacheEntry>(Scene->Primitives.Num());
 	}
 
 	// return entry
-	return NewVirtualShadowMapCacheEntry;
+	return NewLightEntry;
+}
+
+
+
+TSharedPtr<FVirtualShadowMapCacheEntry> FVirtualShadowMapArrayCacheManager::FindCreateCacheEntry(int32 LightSceneId, int32 Index)
+{
+	if (CVarCacheVirtualSMs.GetValueOnRenderThread() == 0)
+	{
+		return nullptr;
+	}
+
+	return FindCreateLightCacheEntry(LightSceneId)->FindCreateShadowMapEntry(Index);
 }
 
 
@@ -528,6 +624,31 @@ void FVirtualShadowMapArrayCacheManager::ProcessRemovedOrUpdatedPrimitives(FRDGB
 	}
 }
 
+static void ResizeFlagArray(TBitArray<>& BitArray, int32 NewMax)
+{
+	if (BitArray.Num() > NewMax)
+	{
+		// Trim off excess items
+		BitArray.SetNumUninitialized(NewMax);
+	}
+	else if (BitArray.Num() < NewMax)
+	{
+		// Add false
+		BitArray.Add(false, NewMax - BitArray.Num());
+	}
+}
+
+void FVirtualShadowMapArrayCacheManager::OnSceneChange()
+{
+	if (CVarCacheVirtualSMs.GetValueOnRenderThread() != 0)
+	{
+		for (auto& CacheEntry : PrevCacheEntries)
+		{
+			ResizeFlagArray(CacheEntry.Value->RenderedPrimitives, Scene->GetMaxPersistentPrimitiveIndex());
+		}
+	}
+}
+
 
 /**
  * Compute shader to project and invalidate the rectangles of given instances.
diff --git a/Engine/Source/Runtime/Renderer/Private/VirtualShadowMaps/VirtualShadowMapCacheManager.h b/Engine/Source/Runtime/Renderer/Private/VirtualShadowMaps/VirtualShadowMapCacheManager.h
--- a/Engine/Source/Runtime/Renderer/Private/VirtualShadowMaps/VirtualShadowMapCacheManager.h
+++ b/Engine/Source/Runtime/Renderer/Private/VirtualShadowMaps/VirtualShadowMapCacheManager.h
@@ -58,6 +58,24 @@ public:
 	FClipmapInfo Clipmap;
 };
 
+class FVirtualShadowMapPerLightCacheEntry
+{
+public:
+	FVirtualShadowMapPerLightCacheEntry(int32 NumScenePrimitives)
+		: RenderedPrimitives(false, NumScenePrimitives)
+	{
+	}
+
+	TSharedPtr<FVirtualShadowMapCacheEntry> FindCreateShadowMapEntry(int32 Index);
+
+	// Primitives that have been rendered (not culled) the previous frame, when a primitive transitions from being culled to not it must be rendered into the VSM
+	// Key culling reasons are small size or distance cutoff.
+	TBitArray<> RenderedPrimitives;
+
+	// One entry represents the cached state of a given shadow map in the set of either a clipmap(N), one cube map(6) or a regular VSM (1)
+	TArray< TSharedPtr<FVirtualShadowMapCacheEntry> > ShadowMapEntries;
+};
+
 // Persistent buffers that we ping pong frame by frame
 struct FVirtualShadowMapArrayFrameData
 {
@@ -105,6 +123,11 @@ public:
 		FVirtualShadowMapArray &VirtualShadowMapArray,
 		bool bEnableCaching);
 
+	/**
+	 * Finds an existing cache entry and moves to the active set or creates a fresh one.
+	 */
+	TSharedPtr<FVirtualShadowMapPerLightCacheEntry> FindCreateLightCacheEntry(int32 LightSceneId);
+
 	/**
 	 * Finds an existing cache entry and moves to the active set or creates a fresh one.
 	 */
@@ -125,46 +148,13 @@ public:
 	class FInvalidatingPrimitiveCollector
 	{
 	public:
-		FInvalidatingPrimitiveCollector(int32 MaxPrimitiveID, FGPUScene& InGPUScene)
-			: AlreadyAddedPrimitives(false, MaxPrimitiveID)
-			, GPUScene(InGPUScene)
-		{
-			// Add and clear pending invalidations enqueued on the GPU Scene from dynamic primitives added since last invalidation
-			for (const FGPUScene::FInstanceRange& Range : GPUScene.DynamicPrimitiveInstancesToInvalidate)
-			{
-				LoadBalancer.Add(Range.InstanceSceneDataOffset, Range.NumInstanceSceneDataEntries, 0U);
-#if VSM_LOG_INVALIDATIONS
-				RangesStr.Appendf(TEXT("[%6d, %6d), "), Range.InstanceSceneDataOffset, Range.InstanceSceneDataOffset + Range.NumInstanceSceneDataEntries);
-#endif
-				TotalInstanceCount += Range.NumInstanceSceneDataEntries;
-			}
-
-			GPUScene.DynamicPrimitiveInstancesToInvalidate.Reset();
-		}
+		FInvalidatingPrimitiveCollector(FVirtualShadowMapArrayCacheManager* InVirtualShadowMapArrayCacheManager);
 
 		/**
 		 * Add a primitive to invalidate the instances for, the function filters redundant primitive adds, and thus expects valid IDs (so can't be called for primitives that have not yet been added)
 		 * and unchanging IDs (so can't be used over a span that include any scene mutation).
 		 */
-		void Add(const FPrimitiveSceneInfo* PrimitiveSceneInfo)
-		{
-			int32 PrimitiveID = PrimitiveSceneInfo->GetIndex();
-			if (PrimitiveID >= 0
-				&& !AlreadyAddedPrimitives[PrimitiveID]
-				&& PrimitiveSceneInfo->GetInstanceSceneDataOffset() != INDEX_NONE
-				// Don't process primitives that are still in the 'added' state because this means that they
-				// have not been uploaded to the GPU yet and may be pending from a previous call to update primitive scene infos.
-				&& !EnumHasAnyFlags(GPUScene.GetPrimitiveDirtyState(PrimitiveID), EPrimitiveDirtyState::Added))
-			{
-				AlreadyAddedPrimitives[PrimitiveID] = true;
-				const int32 NumInstanceSceneDataEntries = PrimitiveSceneInfo->GetNumInstanceSceneDataEntries();
-				LoadBalancer.Add(PrimitiveSceneInfo->GetInstanceSceneDataOffset(), NumInstanceSceneDataEntries, 0U);
-#if VSM_LOG_INVALIDATIONS
-				RangesStr.Appendf(TEXT("[%6d, %6d), "), PrimitiveSceneInfo->GetInstanceSceneDataOffset(), PrimitiveSceneInfo->GetInstanceSceneDataOffset() + NumInstanceSceneDataEntries);
-#endif
-				TotalInstanceCount += NumInstanceSceneDataEntries;
-			}
-		}
+		void Add(const FPrimitiveSceneInfo* PrimitiveSceneInfo);
 
 		bool IsEmpty() const { return LoadBalancer.IsEmpty(); }
 
@@ -174,7 +164,9 @@ public:
 #if VSM_LOG_INVALIDATIONS
 		FString RangesStr;
 #endif
+		const FScene& Scene;
 		FGPUScene& GPUScene;
+		FVirtualShadowMapArrayCacheManager& Manager;
 	};
 
 	/**
@@ -184,6 +176,11 @@ public:
 	 */
 	void ProcessRemovedOrUpdatedPrimitives(FRDGBuilder& GraphBuilder, const FGPUScene& GPUScene, FInvalidatingPrimitiveCollector& InvalidatingPrimitiveCollector);
 
+	/**
+	 * Allow the cache manager to track scene changes, in particular track resizing of primitive tracking data.
+	 */
+	void OnSceneChange();
+
 	TRDGUniformBufferRef<FVirtualShadowMapUniformParameters> GetPreviousUniformBuffer(FRDGBuilder& GraphBuilder) const;
 
 	FVirtualShadowMapArrayFrameData PrevBuffers;
@@ -206,9 +203,9 @@ private:
 	// we store the physical pool here.
 	TRefCountPtr<IPooledRenderTarget> PhysicalPagePool;
 
-	// Index the Cache entries by the light ID and clipmap/cubemap face index (if applicable)
-	TMap< FIntPoint, TSharedPtr<FVirtualShadowMapCacheEntry> > CacheEntries;
-	TMap< FIntPoint, TSharedPtr<FVirtualShadowMapCacheEntry> > PrevCacheEntries;
+	// Index the Cache entries by the light ID
+	TMap< int32, TSharedPtr<FVirtualShadowMapPerLightCacheEntry> > CacheEntries;
+	TMap< int32, TSharedPtr<FVirtualShadowMapPerLightCacheEntry> > PrevCacheEntries;
 
 	// Stores stats over frames when activated.
 	TRefCountPtr<FRDGPooledBuffer> AccumulatedStatsBuffer;
diff --git a/Engine/Source/Runtime/Renderer/Private/VirtualShadowMaps/VirtualShadowMapClipmap.cpp b/Engine/Source/Runtime/Renderer/Private/VirtualShadowMaps/VirtualShadowMapClipmap.cpp
--- a/Engine/Source/Runtime/Renderer/Private/VirtualShadowMaps/VirtualShadowMapClipmap.cpp
+++ b/Engine/Source/Runtime/Renderer/Private/VirtualShadowMaps/VirtualShadowMapClipmap.cpp
@@ -109,6 +109,11 @@ FVirtualShadowMapClipmap::FVirtualShadowMapClipmap(
 
 	WorldOrigin = CameraViewMatrices.GetViewOrigin();
 
+	if (bCacheValid)
+	{
+		PerLightCacheEntry = VirtualShadowMapArrayCacheManager->FindCreateLightCacheEntry(LightSceneInfo.Id);
+	}
+
 	for (int32 Index = 0; Index < LevelCount; ++Index)
 	{
 		FLevelData& Level = LevelData[Index];
@@ -142,10 +147,10 @@ FVirtualShadowMapClipmap::FVirtualShadowMapClipmap(
 		// Check if we have a cache entry for this level
 		// If we do and it covers our required depth range, we can use cached pages. Otherwise we need to invalidate.
 		TSharedPtr<FVirtualShadowMapCacheEntry> CacheEntry = nullptr;
-		if (bCacheValid)
+		if (PerLightCacheEntry.IsValid())
 		{
 			// NOTE: We use the absolute (virtual) level index so that the caching is robust against changes to the chosen level range
-			CacheEntry = VirtualShadowMapArrayCacheManager->FindCreateCacheEntry(LightSceneInfo.Id, AbsoluteLevel);
+			CacheEntry = PerLightCacheEntry->FindCreateShadowMapEntry(AbsoluteLevel);
 		}
 
 		// We expand the depth range of the clipmap level to allow for camera movement without having to invalidate cached shadow data
@@ -277,3 +282,26 @@ uint32 FVirtualShadowMapClipmap::GetCoarsePageClipmapIndexMask()
 
 	return BitMask;
 }
+
+void FVirtualShadowMapClipmap::OnPrimitiveRendered(FPersistentPrimitiveIndex PersistentPrimitiveId)
+{
+	check(PersistentPrimitiveId.Index >= 0);
+
+	if (PerLightCacheEntry.IsValid())
+	{
+		check(PersistentPrimitiveId.Index < PerLightCacheEntry->RenderedPrimitives.Num());
+		// Check previous frame(s) state to detect transition from hidden->visible
+		if (!PerLightCacheEntry->RenderedPrimitives[PersistentPrimitiveId.Index])
+		{
+			if (RevealedPrimitivesMask.IsEmpty())
+			{
+				RevealedPrimitivesMask.Init(false, PerLightCacheEntry->RenderedPrimitives.Num());
+			}
+			RevealedPrimitivesMask[PersistentPrimitiveId.Index] = true;
+		}
+
+		// update rendered state (this is cleared whenever a primitive is invalidating the VSM).
+		PerLightCacheEntry->RenderedPrimitives[PersistentPrimitiveId.Index] = true;
+	}
+}
+
diff --git a/Engine/Source/Runtime/Renderer/Private/VirtualShadowMaps/VirtualShadowMapClipmap.h b/Engine/Source/Runtime/Renderer/Private/VirtualShadowMaps/VirtualShadowMapClipmap.h
--- a/Engine/Source/Runtime/Renderer/Private/VirtualShadowMaps/VirtualShadowMapClipmap.h
+++ b/Engine/Source/Runtime/Renderer/Private/VirtualShadowMaps/VirtualShadowMapClipmap.h
@@ -12,6 +12,7 @@ VirtualShadowMapClipmap.h
 struct FViewMatrices;
 class FVirtualShadowMapArray;
 class FVirtualShadowMapArrayCacheManager;
+class FVirtualShadowMapPerLightCacheEntry;
 
 class FVirtualShadowMapClipmap : FRefCountedObject
 {
@@ -76,6 +77,17 @@ public:
 	// Bits relative to FirstLevel (i.e. in terms of ClipmapIndex, not ClipmapLevel)
 	static uint32 GetCoarsePageClipmapIndexMask();
 
+	/**
+	 * Called when a primitive passes CPU-culling, note that this applies to non-nanite primitives only. Not thread safe in general.
+	 */ 
+	void OnPrimitiveRendered(FPersistentPrimitiveIndex PersistentPrimitiveId);
+
+	/**
+	 * Return array with one bit per primitive, a set bit indicates that the primitive transitioned from not rendered to rendered this frame (see OnPrimitiveRendered above).
+	 */
+	TConstArrayView<uint32> GetRevealedPrimitivesMask() const { return RevealedPrimitivesMask.IsEmpty() ? MakeArrayView<uint32>(nullptr, 0) : MakeArrayView(RevealedPrimitivesMask.GetData(), FBitSet::CalculateNumWords(RevealedPrimitivesMask.Num())); }
+	int32 GetNumRevealedPrimitives() const { return RevealedPrimitivesMask.Num(); }
+
 private:
 	void ComputeBoundingVolumes(const FViewMatrices& CameraViewMatrices);
 
@@ -111,4 +123,9 @@ private:
 
 	FSphere BoundingSphere;
 	FConvexVolume ViewFrustumBounds;
+
+	TSharedPtr<FVirtualShadowMapPerLightCacheEntry> PerLightCacheEntry;
+
+	// Set to 1 for each primitives that went from not being rendered to being rendered this frame, may be empty if no primitives were revealed - lazily initialized
+	TBitArray<> RevealedPrimitivesMask;
 };
