Fix hair outline selection when upscaler is active.

diff --git a/Engine/Shaders/Private/HairStrands/HairStrandsHitProxy.usf b/Engine/Shaders/Private/HairStrands/HairStrandsHitProxy.usf
--- a/Engine/Shaders/Private/HairStrands/HairStrandsHitProxy.usf
+++ b/Engine/Shaders/Private/HairStrands/HairStrandsHitProxy.usf
@@ -1,6 +1,6 @@
 // Copyright Epic Games, Inc. All Rights Reserved.
 
-#include "/Engine/Public/Platform.ush"
+#include "../Common.ush"
 float3 DecodeVelocityFromTexture(float4 EncodedV);
 #include "HairStrandsVisibilityCommon.ush" 
 
@@ -49,6 +49,7 @@ void EmitPS(
 
 #if SHADER_SELECTION
 
+float2							 InvViewportResolution;
 uint							 MaxMaterialCount;
 Buffer<uint>					 SelectionMaterialIdBuffer;
 Texture2D<uint>					 VisNodeIndex;
@@ -58,7 +59,8 @@ void EmitPS(
 	in float4 SvPosition : SV_Position,
 	out float OutDepth : SV_Depth)
 {
-	const uint2 PixelCoord = (uint2)SvPosition.xy;
+	const float2 UV = SvPosition * InvViewportResolution;
+	const uint2 PixelCoord = UV * View.ViewSizeAndInvSize.xy;
 
 	// Take the first sample
 	const FNodeDesc NodeDesc = DecodeNodeDesc(VisNodeIndex.Load(uint3(PixelCoord, 0)));
diff --git a/Engine/Source/Runtime/Renderer/Private/HairStrands/HairStrandsData.h b/Engine/Source/Runtime/Renderer/Private/HairStrands/HairStrandsData.h
--- a/Engine/Source/Runtime/Renderer/Private/HairStrands/HairStrandsData.h
+++ b/Engine/Source/Runtime/Renderer/Private/HairStrands/HairStrandsData.h
@@ -435,7 +435,7 @@ namespace HairStrands
 
 	bool HasPositionsChanged(FRDGBuilder& GraphBuilder, const FViewInfo& View);
 	void DrawHitProxies(FRDGBuilder& GraphBuilder, const FScene& Scene, const FViewInfo& View, FInstanceCullingManager& InstanceCullingManager, FRDGTextureRef HitProxyTexture, FRDGTextureRef HitProxyDepthTexture);
-	void DrawEditorSelection(FRDGBuilder& GraphBuilder, const FViewInfo& View, FRDGTextureRef SelectionDepthTexture);
+	void DrawEditorSelection(FRDGBuilder& GraphBuilder, const FViewInfo& View, const FIntRect& ViewportRect, FRDGTextureRef SelectionDepthTexture);
 
 	// Mesh batch helpers
 	FHairGroupPublicData* GetHairData(const FMeshBatch* In);
diff --git a/Engine/Source/Runtime/Renderer/Private/HairStrands/HairStrandsVisibility.cpp b/Engine/Source/Runtime/Renderer/Private/HairStrands/HairStrandsVisibility.cpp
--- a/Engine/Source/Runtime/Renderer/Private/HairStrands/HairStrandsVisibility.cpp
+++ b/Engine/Source/Runtime/Renderer/Private/HairStrands/HairStrandsVisibility.cpp
@@ -3152,6 +3152,7 @@ class FHairStrandsEmitSelectionPS : public FGlobalShader
 	SHADER_USE_PARAMETER_STRUCT(FHairStrandsEmitSelectionPS, FGlobalShader);
 
 	BEGIN_SHADER_PARAMETER_STRUCT(FParameters, )
+		SHADER_PARAMETER(FVector2f, InvViewportResolution)
 		SHADER_PARAMETER(uint32, MaxMaterialCount)
 		SHADER_PARAMETER_STRUCT_REF(FViewUniformShaderParameters, View)
 		SHADER_PARAMETER_RDG_TEXTURE(Texture2D, VisNodeIndex)
@@ -3176,6 +3177,7 @@ IMPLEMENT_GLOBAL_SHADER(FHairStrandsEmitSelectionPS, "/Engine/Private/HairStrand
 void AddHairStrandsSelectionOutlinePass(
 	FRDGBuilder& GraphBuilder,
 	const FViewInfo& View,
+	const FIntRect& ViewportRect,
 	FRDGTextureRef VisNodeIndex,
 	FRDGBufferRef VisNodeData,
 	FRDGTextureRef SelectionDepthTexture)
@@ -3199,6 +3201,7 @@ void AddHairStrandsSelectionOutlinePass(
 	auto* PassParameters = GraphBuilder.AllocParameters<FHairStrandsEmitSelectionPS::FParameters>();
 	PassParameters->View = View.ViewUniformBuffer;
 	PassParameters->MaxMaterialCount = SelectionMaterialId.Num();
+	PassParameters->InvViewportResolution = FVector2f(1.f/ViewportRect.Width(), 1.f/ViewportRect.Height());
 	PassParameters->VisNodeIndex = VisNodeIndex;
 	PassParameters->VisNodeData = GraphBuilder.CreateSRV(VisNodeData);
 	PassParameters->SelectionMaterialIdBuffer = GraphBuilder.CreateSRV(SelectionMaterialIdBuffer, PF_R32_UINT);
@@ -3212,7 +3215,7 @@ void AddHairStrandsSelectionOutlinePass(
 		RDG_EVENT_NAME("HairStrands::EmitSelection"),
 		PixelShader,
 		PassParameters,
-		View.ViewRect,
+		ViewportRect,
 		TStaticBlendState<>::GetRHI(),
 		TStaticRasterizerState<>::GetRHI(),
 		TStaticDepthStencilState<true, CF_DepthNearOrEqual, true, CF_Always, SO_Keep, SO_Keep, SO_Replace>::GetRHI(), 
@@ -3381,11 +3384,12 @@ namespace HairStrands
 {
 
 // Draw hair strands depth value for outline selection
-void DrawEditorSelection(FRDGBuilder& GraphBuilder, const FViewInfo& View, FRDGTextureRef SelectionDepthTexture)
+void DrawEditorSelection(FRDGBuilder& GraphBuilder, const FViewInfo& View, const FIntRect& ViewportRect, FRDGTextureRef SelectionDepthTexture)
 {
 	AddHairStrandsSelectionOutlinePass(
 		GraphBuilder,
 		View,
+		ViewportRect,
 		View.HairStrandsViewData.VisibilityData.NodeIndex,
 		View.HairStrandsViewData.VisibilityData.NodeVisData,
 		SelectionDepthTexture);
diff --git a/Engine/Source/Runtime/Renderer/Private/PostProcess/PostProcessSelectionOutline.cpp b/Engine/Source/Runtime/Renderer/Private/PostProcess/PostProcessSelectionOutline.cpp
--- a/Engine/Source/Runtime/Renderer/Private/PostProcess/PostProcessSelectionOutline.cpp
+++ b/Engine/Source/Runtime/Renderer/Private/PostProcess/PostProcessSelectionOutline.cpp
@@ -146,13 +146,12 @@ FScreenPassTexture AddSelectionOutlinePass(FRDGBuilder& GraphBuilder, const FVie
 				RHICmdList.SetScissorRect(false, 0, 0, 0, 0);
 			}
 		});
-	}
-
 
-	// Render HairStrands outlines
-	if (HairStrands::HasViewHairStrandsData(View))
-	{
-		HairStrands::DrawEditorSelection(GraphBuilder, View, DepthStencilTexture);
+		// Render HairStrands outlines
+		if (HairStrands::HasViewHairStrandsData(View))
+		{
+			HairStrands::DrawEditorSelection(GraphBuilder, View, SceneColorViewport.Rect, DepthStencilTexture);
+		}
 	}
 
 	FScreenPassRenderTarget Output = Inputs.OverrideOutput;
