Filter small hair from selection to unclutter selection outline when several component are selected.

diff --git a/Engine/Shaders/Private/HairStrands/HairStrandsHitProxy.usf b/Engine/Shaders/Private/HairStrands/HairStrandsHitProxy.usf
--- a/Engine/Shaders/Private/HairStrands/HairStrandsHitProxy.usf
+++ b/Engine/Shaders/Private/HairStrands/HairStrandsHitProxy.usf
@@ -49,11 +49,13 @@ void EmitPS(
 
 #if SHADER_SELECTION
 
+float							 CoverageThreshold;
 float2							 InvViewportResolution;
 uint							 MaxMaterialCount;
 Buffer<uint>					 SelectionMaterialIdBuffer;
 Texture2D<uint>					 VisNodeIndex;
 StructuredBuffer<FPackedHairVis> VisNodeData;
+Texture2D<float>				 CoverageTexture;
 
 void EmitPS(
 	in float4 SvPosition : SV_Position,
@@ -67,10 +69,11 @@ void EmitPS(
 	bool bIsSelected = false;
 	if (NodeDesc.Count > 0)
 	{
+		const float Coverage = CoverageTexture.Load(uint3(PixelCoord, 0));
 		const uint SampleIndex = NodeDesc.Offset + 0;
 		const FHairVis InNode = UnpackHairVis(VisNodeData[SampleIndex]);
 		const uint MaterialId = clamp(InNode.MaterialId, 0, MaxMaterialCount - 1);
-		if (SelectionMaterialIdBuffer[MaterialId] > 0)
+		if (SelectionMaterialIdBuffer[MaterialId] > 0 && Coverage > CoverageThreshold)
 		{
 			bIsSelected = true;
 			OutDepth = InNode.Depth;
diff --git a/Engine/Source/Runtime/Renderer/Private/HairStrands/HairStrandsVisibility.cpp b/Engine/Source/Runtime/Renderer/Private/HairStrands/HairStrandsVisibility.cpp
--- a/Engine/Source/Runtime/Renderer/Private/HairStrands/HairStrandsVisibility.cpp
+++ b/Engine/Source/Runtime/Renderer/Private/HairStrands/HairStrandsVisibility.cpp
@@ -90,6 +90,9 @@ static FAutoConsoleVariableRef CVarHairStrands_InvalidationPosition_Threshold(TE
 static int32 GHairStrands_InvalidationPosition_Debug = 0;
 static FAutoConsoleVariableRef CVarHairStrands_InvalidationPosition_Debug(TEXT("r.HairStrands.PathTracing.InvalidationDebug"), GHairStrands_InvalidationPosition_Debug, TEXT("Enable bounding box drawing for groom element causing path tracer invalidation"));
 
+static float GHairStrands_Selection_CoverageThreshold = 0.0f;
+static FAutoConsoleVariableRef CVarHairStrands_Selection_CoverageThreshold(TEXT("r.HairStrands.Selection.CoverageThreshold"), GHairStrands_Selection_CoverageThreshold, TEXT("Coverage threshold for making hair strands outline selection finer"));
+
 /////////////////////////////////////////////////////////////////////////////////////////
 
 namespace HairStrandsVisibilityInternal
@@ -3152,10 +3155,12 @@ class FHairStrandsEmitSelectionPS : public FGlobalShader
 	SHADER_USE_PARAMETER_STRUCT(FHairStrandsEmitSelectionPS, FGlobalShader);
 
 	BEGIN_SHADER_PARAMETER_STRUCT(FParameters, )
+		SHADER_PARAMETER(float, CoverageThreshold)
 		SHADER_PARAMETER(FVector2f, InvViewportResolution)
 		SHADER_PARAMETER(uint32, MaxMaterialCount)
 		SHADER_PARAMETER_STRUCT_REF(FViewUniformShaderParameters, View)
 		SHADER_PARAMETER_RDG_TEXTURE(Texture2D, VisNodeIndex)
+		SHADER_PARAMETER_RDG_TEXTURE(Texture2D<float>, CoverageTexture)
 		SHADER_PARAMETER_RDG_BUFFER_SRV(StructuredBuffer<FPackedHairVis>, VisNodeData)
 		SHADER_PARAMETER_RDG_BUFFER_SRV(Buffer<uint>, SelectionMaterialIdBuffer)
 		RENDER_TARGET_BINDING_SLOTS()
@@ -3180,6 +3185,7 @@ void AddHairStrandsSelectionOutlinePass(
 	const FIntRect& ViewportRect,
 	FRDGTextureRef VisNodeIndex,
 	FRDGBufferRef VisNodeData,
+	FRDGTextureRef CoverageTexture,
 	FRDGTextureRef SelectionDepthTexture)
 {
 	if (View.HairStrandsMeshElements.Num() == 0)
@@ -3200,15 +3206,18 @@ void AddHairStrandsSelectionOutlinePass(
 	FRDGBufferRef SelectionMaterialIdBuffer = CreateUploadBuffer(GraphBuilder, TEXT("Hair.MaterialIdToHitProxyIdBuffer"), sizeof(uint32), SelectionMaterialId.Num(), SelectionMaterialId.GetData(), sizeof(uint32) * SelectionMaterialId.Num());
 	auto* PassParameters = GraphBuilder.AllocParameters<FHairStrandsEmitSelectionPS::FParameters>();
 	PassParameters->View = View.ViewUniformBuffer;
+	PassParameters->CoverageThreshold = FMath::Clamp(GHairStrands_Selection_CoverageThreshold, 0.f, 1.f);
 	PassParameters->MaxMaterialCount = SelectionMaterialId.Num();
 	PassParameters->InvViewportResolution = FVector2f(1.f/ViewportRect.Width(), 1.f/ViewportRect.Height());
 	PassParameters->VisNodeIndex = VisNodeIndex;
 	PassParameters->VisNodeData = GraphBuilder.CreateSRV(VisNodeData);
+	PassParameters->CoverageTexture = CoverageTexture;
 	PassParameters->SelectionMaterialIdBuffer = GraphBuilder.CreateSRV(SelectionMaterialIdBuffer, PF_R32_UINT);
 	PassParameters->RenderTargets.DepthStencil = FDepthStencilBinding(SelectionDepthTexture, ERenderTargetLoadAction::ELoad, FExclusiveDepthStencil::DepthWrite_StencilWrite);
 
 	auto PixelShader = View.ShaderMap->GetShader<FHairStrandsEmitSelectionPS>();
 
+	const uint32 StencilRef = 3;
 	FPixelShaderUtils::AddFullscreenPass(
 		GraphBuilder,
 		View.ShaderMap,
@@ -3219,7 +3228,7 @@ void AddHairStrandsSelectionOutlinePass(
 		TStaticBlendState<>::GetRHI(),
 		TStaticRasterizerState<>::GetRHI(),
 		TStaticDepthStencilState<true, CF_DepthNearOrEqual, true, CF_Always, SO_Keep, SO_Keep, SO_Replace>::GetRHI(), 
-		3);
+		StencilRef);
 #endif
 }
 
@@ -3256,7 +3265,7 @@ IMPLEMENT_GLOBAL_SHADER(FHairStrandsEmitHitProxyIdPS, "/Engine/Private/HairStran
 void AddHairStrandsHitProxyIdPass(
 	FRDGBuilder& GraphBuilder,
 	const FScene& Scene,
-	const FViewInfo& View,	
+	const FViewInfo& View,
 	FRDGTextureRef VisNodeIndex,
 	FRDGBufferRef VisNodeData,
 	FRDGTextureRef HitProxyTexture,
@@ -3392,6 +3401,7 @@ void DrawEditorSelection(FRDGBuilder& GraphBuilder, const FViewInfo& View, const
 		ViewportRect,
 		View.HairStrandsViewData.VisibilityData.NodeIndex,
 		View.HairStrandsViewData.VisibilityData.NodeVisData,
+		View.HairStrandsViewData.VisibilityData.CoverageTexture,
 		SelectionDepthTexture);
 }
 
