Exclude editor-only shaders from games that don't support a cooked editor.

diff --git a/Engine/Source/Runtime/RenderCore/Private/Shader.cpp b/Engine/Source/Runtime/RenderCore/Private/Shader.cpp
--- a/Engine/Source/Runtime/RenderCore/Private/Shader.cpp
+++ b/Engine/Source/Runtime/RenderCore/Private/Shader.cpp
@@ -1930,3 +1930,20 @@ void ShaderMapAppendKeyString(EShaderPlatform Platform, FString& KeyString)
 		KeyString += FString::Printf(TEXT("_LWC-%d"), FMath::FloorToInt(FLargeWorldRenderScalar::GetTileSize()));
 	}
 }
+
+EShaderPermutationFlags GetShaderPermutationFlags(const FPlatformTypeLayoutParameters& LayoutParams)
+{
+	EShaderPermutationFlags Result = EShaderPermutationFlags::None;
+
+	static bool bProjectSupportsCookedEditor = []()
+	{
+		bool bSupportCookedEditorConfigValue = false;
+		return GConfig->GetBool(TEXT("CookedEditorSettings"), TEXT("bSupportCookedEditor"), bSupportCookedEditorConfigValue, GGameIni) && bSupportCookedEditorConfigValue;
+	}();
+
+	if (bProjectSupportsCookedEditor || LayoutParams.WithEditorOnly())
+	{
+		Result |= EShaderPermutationFlags::HasEditorOnlyData;
+	}
+	return Result;
+}
diff --git a/Engine/Source/Runtime/RenderCore/Public/Shader.h b/Engine/Source/Runtime/RenderCore/Public/Shader.h
--- a/Engine/Source/Runtime/RenderCore/Public/Shader.h
+++ b/Engine/Source/Runtime/RenderCore/Public/Shader.h
@@ -637,12 +637,7 @@ enum class EShaderPermutationFlags : uint32
 };
 ENUM_CLASS_FLAGS(EShaderPermutationFlags);
 
-inline EShaderPermutationFlags GetShaderPermutationFlags(const FPlatformTypeLayoutParameters& LayoutParams)
-{
-	EShaderPermutationFlags Result = EShaderPermutationFlags::None;
-	if (LayoutParams.WithEditorOnly()) Result |= EShaderPermutationFlags::HasEditorOnlyData;
-	return Result;
-}
+RENDERCORE_API EShaderPermutationFlags GetShaderPermutationFlags(const FPlatformTypeLayoutParameters& LayoutParams);
 
 struct FShaderPermutationParameters
 {
diff --git a/Engine/Source/Runtime/Renderer/Private/SceneHitProxyRendering.cpp b/Engine/Source/Runtime/Renderer/Private/SceneHitProxyRendering.cpp
--- a/Engine/Source/Runtime/Renderer/Private/SceneHitProxyRendering.cpp
+++ b/Engine/Source/Runtime/Renderer/Private/SceneHitProxyRendering.cpp
@@ -54,7 +54,7 @@ public:
 	static bool ShouldCompilePermutation(const FMeshMaterialShaderPermutationParameters& Parameters)
 	{
 		// Only compile the hit proxy vertex shader on desktop editor platforms
-		return IsPCPlatform(Parameters.Platform)// && EnumHasAllFlags(Parameters.Flags, EShaderPermutationFlags::HasEditorOnlyData)
+		return IsPCPlatform(Parameters.Platform) && EnumHasAllFlags(Parameters.Flags, EShaderPermutationFlags::HasEditorOnlyData)
 			// and only compile for the default material or materials that are masked.
 			&& (Parameters.MaterialParameters.bIsSpecialEngineMaterial ||
 				!Parameters.MaterialParameters.bWritesEveryPixel ||
@@ -111,7 +111,7 @@ public:
 	static bool ShouldCompilePermutation(const FMeshMaterialShaderPermutationParameters& Parameters)
 	{
 		// Only compile the hit proxy vertex shader on desktop editor platforms
-		return IsPCPlatform(Parameters.Platform)// && EnumHasAllFlags(Parameters.Flags, EShaderPermutationFlags::HasEditorOnlyData)
+		return IsPCPlatform(Parameters.Platform) && EnumHasAllFlags(Parameters.Flags, EShaderPermutationFlags::HasEditorOnlyData)
 			// and only compile for default materials or materials that are masked.
 			&& (Parameters.MaterialParameters.bIsSpecialEngineMaterial ||
 				!Parameters.MaterialParameters.bWritesEveryPixel ||
