Temporary work-around for assert firing on some machines

diff --git a/Engine/Source/Runtime/Renderer/Private/VirtualShadowMaps/VirtualShadowMapClipmap.cpp b/Engine/Source/Runtime/Renderer/Private/VirtualShadowMaps/VirtualShadowMapClipmap.cpp
--- a/Engine/Source/Runtime/Renderer/Private/VirtualShadowMaps/VirtualShadowMapClipmap.cpp
+++ b/Engine/Source/Runtime/Renderer/Private/VirtualShadowMaps/VirtualShadowMapClipmap.cpp
@@ -289,7 +289,13 @@ void FVirtualShadowMapClipmap::OnPrimitiveRendered(FPersistentPrimitiveIndex Per
 
 	if (PerLightCacheEntry.IsValid())
 	{
-		check(PersistentPrimitiveId.Index < PerLightCacheEntry->RenderedPrimitives.Num());
+		// FIXME: This should not happen but working around a few people hitting this until we can debug
+		//check(PersistentPrimitiveId.Index < PerLightCacheEntry->RenderedPrimitives.Num());
+		if (PersistentPrimitiveId.Index >= PerLightCacheEntry->RenderedPrimitives.Num())
+		{
+			return;
+		}		
+
 		// Check previous frame(s) state to detect transition from hidden->visible
 		if (!PerLightCacheEntry->RenderedPrimitives[PersistentPrimitiveId.Index])
 		{
