Fixing out of bounds array access with non-nanite assets.

diff --git a/Engine/Source/Runtime/Renderer/Private/VirtualShadowMaps/VirtualShadowMapCacheManager.cpp b/Engine/Source/Runtime/Renderer/Private/VirtualShadowMaps/VirtualShadowMapCacheManager.cpp
--- a/Engine/Source/Runtime/Renderer/Private/VirtualShadowMaps/VirtualShadowMapCacheManager.cpp
+++ b/Engine/Source/Runtime/Renderer/Private/VirtualShadowMaps/VirtualShadowMapCacheManager.cpp
@@ -174,7 +174,7 @@ void FVirtualShadowMapArrayCacheManager::FInvalidatingPrimitiveCollector::Add(co
 		for (auto& CacheEntry : Manager.PrevCacheEntries)
 		{
 			TBitArray<>& RenderedPrimitives = CacheEntry.Value->RenderedPrimitives;
-			if (bIsNaniteMesh || (PersistentPrimitiveIndex <= RenderedPrimitives.Num() && RenderedPrimitives[PersistentPrimitiveIndex]))
+			if (bIsNaniteMesh || (PersistentPrimitiveIndex < RenderedPrimitives.Num() && RenderedPrimitives[PersistentPrimitiveIndex]))
 			{
 				if (!bIsNaniteMesh)
 				{
