Fix bug in new rendered primitive tracking for VSM (incorrectly used upper boundfor non-persistent primitive IDs).

diff --git a/Engine/Source/Runtime/Renderer/Private/VirtualShadowMaps/VirtualShadowMapCacheManager.cpp b/Engine/Source/Runtime/Renderer/Private/VirtualShadowMaps/VirtualShadowMapCacheManager.cpp
--- a/Engine/Source/Runtime/Renderer/Private/VirtualShadowMaps/VirtualShadowMapCacheManager.cpp
+++ b/Engine/Source/Runtime/Renderer/Private/VirtualShadowMaps/VirtualShadowMapCacheManager.cpp
@@ -341,7 +341,7 @@ TSharedPtr<FVirtualShadowMapPerLightCacheEntry> FVirtualShadowMapArrayCacheManag
 	}
 	else
 	{
-		NewLightEntry = MakeShared<FVirtualShadowMapPerLightCacheEntry>(Scene->Primitives.Num());
+		NewLightEntry = MakeShared<FVirtualShadowMapPerLightCacheEntry>(Scene->GetMaxPersistentPrimitiveIndex());
 	}
 
 	// return entry
@@ -645,6 +645,10 @@ void FVirtualShadowMapArrayCacheManager::OnSceneChange()
 		{
 			ResizeFlagArray(CacheEntry.Value->RenderedPrimitives, Scene->GetMaxPersistentPrimitiveIndex());
 		}
+		for (auto& CacheEntry : CacheEntries)
+		{
+			ResizeFlagArray(CacheEntry.Value->RenderedPrimitives, Scene->GetMaxPersistentPrimitiveIndex());
+		}
 	}
 }
 
diff --git a/Engine/Source/Runtime/Renderer/Private/VirtualShadowMaps/VirtualShadowMapCacheManager.h b/Engine/Source/Runtime/Renderer/Private/VirtualShadowMaps/VirtualShadowMapCacheManager.h
--- a/Engine/Source/Runtime/Renderer/Private/VirtualShadowMaps/VirtualShadowMapCacheManager.h
+++ b/Engine/Source/Runtime/Renderer/Private/VirtualShadowMaps/VirtualShadowMapCacheManager.h
@@ -61,8 +61,8 @@ public:
 class FVirtualShadowMapPerLightCacheEntry
 {
 public:
-	FVirtualShadowMapPerLightCacheEntry(int32 NumScenePrimitives)
-		: RenderedPrimitives(false, NumScenePrimitives)
+	FVirtualShadowMapPerLightCacheEntry(int32 MaxPersistentScenePrimitiveIndex)
+		: RenderedPrimitives(false, MaxPersistentScenePrimitiveIndex)
 	{
 	}
 
diff --git a/Engine/Source/Runtime/Renderer/Private/VirtualShadowMaps/VirtualShadowMapClipmap.cpp b/Engine/Source/Runtime/Renderer/Private/VirtualShadowMaps/VirtualShadowMapClipmap.cpp
--- a/Engine/Source/Runtime/Renderer/Private/VirtualShadowMaps/VirtualShadowMapClipmap.cpp
+++ b/Engine/Source/Runtime/Renderer/Private/VirtualShadowMaps/VirtualShadowMapClipmap.cpp
@@ -289,13 +289,7 @@ void FVirtualShadowMapClipmap::OnPrimitiveRendered(FPersistentPrimitiveIndex Per
 
 	if (PerLightCacheEntry.IsValid())
 	{
-		// FIXME: This should not happen but working around a few people hitting this until we can debug
-		//check(PersistentPrimitiveId.Index < PerLightCacheEntry->RenderedPrimitives.Num());
-		if (PersistentPrimitiveId.Index >= PerLightCacheEntry->RenderedPrimitives.Num())
-		{
-			return;
-		}		
-
+		check(PersistentPrimitiveId.Index < PerLightCacheEntry->RenderedPrimitives.Num());
 		// Check previous frame(s) state to detect transition from hidden->visible
 		if (!PerLightCacheEntry->RenderedPrimitives[PersistentPrimitiveId.Index])
 		{
