Add using Nanite as a reason to force a full Depth Pass.
Fixes issues with Nanite drawing on top of masked elements when DBuffer Decals are disabled

diff --git a/Engine/Source/Runtime/Renderer/Private/DeferredShadingRenderer.cpp b/Engine/Source/Runtime/Renderer/Private/DeferredShadingRenderer.cpp
--- a/Engine/Source/Runtime/Renderer/Private/DeferredShadingRenderer.cpp
+++ b/Engine/Source/Runtime/Renderer/Private/DeferredShadingRenderer.cpp
@@ -304,6 +304,13 @@ const TCHAR* GetDepthPassReason(bool bDitheredLODTransitionsUseStencil, EShaderP
 		return TEXT("(Forced by ForwardShading)");
 	}
 
+	bool bUseNanite = UseNanite(ShaderPlatform);
+
+	if (bUseNanite)
+	{
+		return TEXT("(Forced by Nanite)");
+	}
+
 	bool bDBufferAllowed = IsUsingDBuffers(ShaderPlatform);
 
 	if (bDBufferAllowed)
diff --git a/Engine/Source/Runtime/Renderer/Private/RendererScene.cpp b/Engine/Source/Runtime/Renderer/Private/RendererScene.cpp
--- a/Engine/Source/Runtime/Renderer/Private/RendererScene.cpp
+++ b/Engine/Source/Runtime/Renderer/Private/RendererScene.cpp
@@ -3734,6 +3734,8 @@ void FScene::Release()
 
 bool ShouldForceFullDepthPass(EShaderPlatform ShaderPlatform)
 {
+	const bool bNaniteEnabled = UseNanite(ShaderPlatform);
+
 	const bool bDBufferAllowed = IsUsingDBuffers(ShaderPlatform);
 
 	static const auto StencilLODDitherCVar = IConsoleManager::Get().FindTConsoleVariableDataInt(TEXT("r.StencilForLODDither"));
@@ -3745,7 +3747,7 @@ bool ShouldForceFullDepthPass(EShaderPlatform ShaderPlatform)
 	const bool bEarlyZMaterialMasking = MaskedInEarlyPass(ShaderPlatform);
 
 	// Note: ShouldForceFullDepthPass affects which static draw lists meshes go into, so nothing it depends on can change at runtime, unless you do a FGlobalComponentRecreateRenderStateContext to propagate the cvar change
-	return bAOCompute || bDBufferAllowed || bStencilLODDither || bEarlyZMaterialMasking || IsForwardShadingEnabled(ShaderPlatform) || IsUsingSelectiveBasePassOutputs(ShaderPlatform);
+	return bNaniteEnabled || bAOCompute || bDBufferAllowed || bStencilLODDither || bEarlyZMaterialMasking || IsForwardShadingEnabled(ShaderPlatform) || IsUsingSelectiveBasePassOutputs(ShaderPlatform);
 }
 
 void FScene::UpdateEarlyZPassMode()
