Add max distance threshold cvar (r.Shadow.Virtual.MaxMaterialPositionInvalidationRange) for stopping VSM invalidations from WPO (and PDO)

diff --git a/Engine/Shaders/Private/VirtualShadowMaps/BuildPerPageDrawCommands.usf b/Engine/Shaders/Private/VirtualShadowMaps/BuildPerPageDrawCommands.usf
--- a/Engine/Shaders/Private/VirtualShadowMaps/BuildPerPageDrawCommands.usf
+++ b/Engine/Shaders/Private/VirtualShadowMaps/BuildPerPageDrawCommands.usf
@@ -35,7 +35,9 @@ uint VisibleInstancesBufferNum;
 
 struct FVSMCullingBatchInfo
 {
+	float3 CullingViewOriginOffset;
 	uint FirstPrimaryView;
+	float3 CullingViewOriginTile;
 	uint NumPrimaryViews;
 };
 
@@ -46,8 +48,15 @@ uint FirstPrimaryView;
 uint NumPrimaryViews;
 uint DynamicInstanceIdOffset;
 uint DynamicInstanceIdMax;
+
+float3 CullingViewOriginOffset;
+float3 CullingViewOriginTile;
+
 #endif // ENABLE_BATCH_MODE
 
+
+float MaxMaterialPositionInvalidationRange;
+
 StructuredBuffer<FDrawCommandDesc> DrawCommandDescs;
 
 RWBuffer<uint> DrawIndirectArgsBufferOut;
@@ -142,8 +151,13 @@ void CullPerPageDrawCommandsCs(uint3 GroupId : SV_GroupID, int GroupThreadIndex
 	FVSMCullingBatchInfo VSMCullingBatchInfo;
 	VSMCullingBatchInfo.FirstPrimaryView = FirstPrimaryView;
 	VSMCullingBatchInfo.NumPrimaryViews = NumPrimaryViews;
+	VSMCullingBatchInfo.CullingViewOriginOffset = CullingViewOriginOffset;
+	VSMCullingBatchInfo.CullingViewOriginTile = CullingViewOriginTile;
 #endif // ENABLE_BATCH_MODE
 
+	FLWCVector3 CullingViewWorldOrigin = MakeLWCVector3(VSMCullingBatchInfo.CullingViewOriginTile, VSMCullingBatchInfo.CullingViewOriginOffset);
+
+
 	uint CurrentBatchProcessingMode = 0U;
 	FInstanceWorkSetup WorkSetup = InstanceCullingLoadBalancer_Setup(GroupId, GroupThreadIndex, GetItemDataOffset(BatchInfo, CurrentBatchProcessingMode));
 	if (!WorkSetup.bValid)
@@ -170,11 +184,20 @@ void CullPerPageDrawCommandsCs(uint3 GroupId : SV_GroupID, int GroupThreadIndex
 	// Load relevant instance data
 	FInstanceSceneData InstanceData = GetInstanceSceneData(InstanceId, InstanceSceneDataSOAStride);
 
+	// Stop invalidations from material if further away than the threshold
+	bool bMaterialInvalidates = DrawCommandDesc.bMaterialMayModifyPosition;
+	BRANCH
+	if (MaxMaterialPositionInvalidationRange >= 0.0f && bMaterialInvalidates)
+	{
+		float DistSq = length2(LWCToFloat(LWCSubtract(LWCMultiply(InstanceData.LocalBoundsCenter, InstanceData.LocalToWorld), CullingViewWorldOrigin)));
+		bMaterialInvalidates = DistSq < Square(MaxMaterialPositionInvalidationRange);
+	}
+
 	// A thing is considered to have moved (which causes it to render into cached pages) if:
 	// It was uploaded this frame, OR
 	const bool bHasMoved = GetGPUSceneFrameNumber() == InstanceData.LastUpdateSceneFrameNumber
 		// it has a material that uses WPO or PDO, OR
-		|| DrawCommandDesc.bMaterialMayModifyPosition
+		|| bMaterialInvalidates
 		// the primitive was "revealed" this frame - i.e., transitioned from culled -> visible on CPU
 		|| WasPrimitiveRevealed(GetPrimitiveData(InstanceData.PrimitiveId).PersistentPrimitiveIndex);
 
@@ -264,7 +287,7 @@ void CullPerPageDrawCommandsCs(uint3 GroupId : SV_GroupID, int GroupThreadIndex
 
 	// If the instance has WPO or PDO and rendered to any mip, queue for invalidation. This will be processed the before next frame (after a new HZB is built)
 	// and invalidate any visible pages.
-	if (ThreadTotalForAllViews > 0U && DrawCommandDesc.bMaterialMayModifyPosition)
+	if (ThreadTotalForAllViews > 0U && bMaterialInvalidates)
 	{
 		MarkInstanceForInvalidation(InstanceId);
 	}
diff --git a/Engine/Source/Runtime/Renderer/Private/ShadowDepthRendering.cpp b/Engine/Source/Runtime/Renderer/Private/ShadowDepthRendering.cpp
--- a/Engine/Source/Runtime/Renderer/Private/ShadowDepthRendering.cpp
+++ b/Engine/Source/Runtime/Renderer/Private/ShadowDepthRendering.cpp
@@ -1857,7 +1857,7 @@ void FSceneRenderer::RenderVirtualShadowMaps(FRDGBuilder& GraphBuilder, bool bNa
 
 	if (UseNonNaniteVirtualShadowMaps(ShaderPlatform, FeatureLevel))
 	{
-		VirtualShadowMapArray.RenderVirtualShadowMapsNonNanite(GraphBuilder, SortedShadowsForShadowDepthPass.VirtualShadowMapShadows, *Scene);
+		VirtualShadowMapArray.RenderVirtualShadowMapsNonNanite(GraphBuilder, SortedShadowsForShadowDepthPass.VirtualShadowMapShadows, *Scene, Views);
 	}
 
 	// If separate static/dynamic caching is enabled, we may need to merge some pages after rendering
diff --git a/Engine/Source/Runtime/Renderer/Private/VirtualShadowMaps/VirtualShadowMapArray.cpp b/Engine/Source/Runtime/Renderer/Private/VirtualShadowMaps/VirtualShadowMapArray.cpp
--- a/Engine/Source/Runtime/Renderer/Private/VirtualShadowMaps/VirtualShadowMapArray.cpp
+++ b/Engine/Source/Runtime/Renderer/Private/VirtualShadowMaps/VirtualShadowMapArray.cpp
@@ -250,6 +250,16 @@ TAutoConsoleVariable<int32> CVarVSMHZB32Bit(
 	ECVF_RenderThreadSafe
 );
 
+
+static TAutoConsoleVariable<float> CVarMaxMaterialPositionInvalidationRange(
+	TEXT("r.Shadow.Virtual.MaxMaterialPositionInvalidationRange"),
+	-1.0f,
+	TEXT("Beyond this distance in world units, material position effects (e.g., WPO or PDO) cease to cause VSM invalidations.\n")
+	TEXT(" This can be used to tune performance by reducing re-draw overhead, but causes some artifacts.\n")
+	TEXT(" < 0 <=> infinite (default)"),
+	ECVF_Scalability | ECVF_RenderThreadSafe
+);
+
 FMatrix CalcTranslatedWorldToShadowUVMatrix(
 	const FMatrix& TranslatedWorldToShadowView,
 	const FMatrix& ViewToClip)
@@ -1771,7 +1781,9 @@ END_SHADER_PARAMETER_STRUCT()
 
 struct FVSMCullingBatchInfo
 {
+	FVector3f CullingViewOriginOffset;
 	uint32 FirstPrimaryView;
+	FVector3f CullingViewOriginTile;
 	uint32 NumPrimaryViews;
 };
 
@@ -1845,6 +1857,10 @@ public:
 		SHADER_PARAMETER(uint32, VisibleInstancesBufferNum)
 		SHADER_PARAMETER(int32, DynamicInstanceIdOffset)
 		SHADER_PARAMETER(int32, DynamicInstanceIdMax)
+		SHADER_PARAMETER(float, MaxMaterialPositionInvalidationRange)
+		SHADER_PARAMETER(FVector3f, CullingViewOriginOffset)
+		SHADER_PARAMETER(FVector3f, CullingViewOriginTile)
+
 		SHADER_PARAMETER_RDG_BUFFER_SRV(StructuredBuffer< FPackedView >, InViews)
 		SHADER_PARAMETER_RDG_BUFFER_SRV(StructuredBuffer< FInstanceCullingContext::FDrawCommandDesc >, DrawCommandDescs)
 
@@ -2021,12 +2037,17 @@ static FCullingResult AddCullingPasses(FRDGBuilder& GraphBuilder,
 
 		PassParameters->DynamicInstanceIdOffset = BatchInfos[0].DynamicInstanceIdOffset;
 		PassParameters->DynamicInstanceIdMax = BatchInfos[0].DynamicInstanceIdMax;
+		
+		PassParameters->MaxMaterialPositionInvalidationRange = CVarMaxMaterialPositionInvalidationRange.GetValueOnRenderThread();
 
 		auto GPUData = LoadBalancer->Upload(GraphBuilder);
 		GPUData.GetShaderParameters(GraphBuilder, PassParameters->LoadBalancerParameters);
 
 		PassParameters->FirstPrimaryView = VSMCullingBatchInfos[0].FirstPrimaryView;
 		PassParameters->NumPrimaryViews = VSMCullingBatchInfos[0].NumPrimaryViews;
+		PassParameters->CullingViewOriginOffset = VSMCullingBatchInfos[0].CullingViewOriginOffset;
+		PassParameters->CullingViewOriginTile = VSMCullingBatchInfos[0].CullingViewOriginTile;
+
 		PassParameters->TotalPrimaryViews = TotalPrimaryViews;
 		PassParameters->VisibleInstancesBufferNum = CullingResult.MaxNumInstancesPerPass;
 		PassParameters->InViews = GraphBuilder.CreateSRV(VirtualShadowViewsRDG);
@@ -2171,7 +2192,7 @@ static void AddRasterPass(
 		});
 }
 
-void FVirtualShadowMapArray::RenderVirtualShadowMapsNonNanite(FRDGBuilder& GraphBuilder, const TArray<FProjectedShadowInfo *, SceneRenderingAllocator>& VirtualSmMeshCommandPasses, FScene& Scene)
+void FVirtualShadowMapArray::RenderVirtualShadowMapsNonNanite(FRDGBuilder& GraphBuilder, const TArray<FProjectedShadowInfo *, SceneRenderingAllocator>& VirtualSmMeshCommandPasses, FScene& Scene, TArrayView<FViewInfo> Views)
 {
 	if (VirtualSmMeshCommandPasses.Num() == 0)
 	{
@@ -2215,6 +2236,37 @@ void FVirtualShadowMapArray::RenderVirtualShadowMapsNonNanite(FRDGBuilder& Graph
 	TArray<FVirtualShadowDepthPassParameters*, SceneRenderingAllocator> BatchedPassParameters;
 	BatchedPassParameters.Reserve(VirtualSmMeshCommandPasses.Num());
 
+	/**
+	 * Use the 'dependent view' i.e., the view used to set up a view dependent CSM/VSM(clipmap) OR select the view closest to the local light.
+	 * This last is important to get some kind of reasonable behaviour for split screen.
+	 */
+	auto GetCullingViewOrigin = [&Views](const FProjectedShadowInfo* ProjectedShadowInfo) -> FLargeWorldRenderPosition
+	{
+		if (ProjectedShadowInfo->DependentView != nullptr)
+		{
+			return ProjectedShadowInfo->DependentView->ShadowViewMatrices.GetViewOrigin();
+		}
+
+		// VSM supports only whole scene shadows, so those without a "DependentView" are local lights
+		// For local lights the origin is the (inverse of) pre-shadow translation. 
+		check(ProjectedShadowInfo->bWholeSceneShadow);
+
+		FVector MinOrigin = Views[0].ShadowViewMatrices.GetViewOrigin();
+		double MinDistanceSq = (MinOrigin + ProjectedShadowInfo->PreShadowTranslation).SquaredLength();
+		for (int Index = 1; Index < Views.Num(); ++Index)
+		{
+			FVector TestOrigin = Views[Index].ShadowViewMatrices.GetViewOrigin();
+			double TestDistanceSq = (TestOrigin + ProjectedShadowInfo->PreShadowTranslation).SquaredLength();
+			if (TestDistanceSq < MinDistanceSq)
+			{
+				MinOrigin = TestOrigin;
+				MinDistanceSq = TestDistanceSq;
+			}
+
+		}
+		return MinOrigin;
+	};
+
 	FInstanceCullingMergedContext InstanceCullingMergedContext(GMaxRHIFeatureLevel);
 	// We don't use the registered culling views (this redundancy should probably be addressed at some point), set the number to disable index range checking
 	InstanceCullingMergedContext.NumCullingViews = -1;
@@ -2227,6 +2279,12 @@ void FVirtualShadowMapArray::RenderVirtualShadowMapsNonNanite(FRDGBuilder& Graph
 		VSMCullingBatchInfo.FirstPrimaryView = uint32(VirtualShadowViews.Num());
 		VSMCullingBatchInfo.NumPrimaryViews = 0U;
 
+		{
+			const FLargeWorldRenderPosition CullingViewOrigin = GetCullingViewOrigin(ProjectedShadowInfo);
+			VSMCullingBatchInfo.CullingViewOriginOffset = CullingViewOrigin.GetOffset();
+			VSMCullingBatchInfo.CullingViewOriginTile = CullingViewOrigin.GetTile();
+		}
+
 		const TSharedPtr<FVirtualShadowMapClipmap> Clipmap = ProjectedShadowInfo->VirtualShadowMapClipmap;
 		if (Clipmap)
 		{
diff --git a/Engine/Source/Runtime/Renderer/Private/VirtualShadowMaps/VirtualShadowMapArray.h b/Engine/Source/Runtime/Renderer/Private/VirtualShadowMaps/VirtualShadowMapArray.h
--- a/Engine/Source/Runtime/Renderer/Private/VirtualShadowMaps/VirtualShadowMapArray.h
+++ b/Engine/Source/Runtime/Renderer/Private/VirtualShadowMaps/VirtualShadowMapArray.h
@@ -201,7 +201,7 @@ public:
 	/**
 	 * Draw Non-Nanite geometry into the VSMs.
 	 */
-	void RenderVirtualShadowMapsNonNanite(FRDGBuilder& GraphBuilder, const TArray<FProjectedShadowInfo*, SceneRenderingAllocator>& VirtualSmMeshCommandPasses, FScene& Scene);
+	void RenderVirtualShadowMapsNonNanite(FRDGBuilder& GraphBuilder, const TArray<FProjectedShadowInfo*, SceneRenderingAllocator>& VirtualSmMeshCommandPasses, FScene& Scene, TArrayView<FViewInfo> Views);
 
 	// Draw debug info into render target 'VSMDebug' of screen-size, the mode is controlled by 'r.Shadow.Virtual.DebugVisualize'.
 	void RenderDebugInfo(FRDGBuilder& GraphBuilder);
