Add opt-in ability for instanced draw calls to preserve their draw order with GPU Scene after culling.
Requires two additional compute passes post-culling to compact the draw command instances.
This feature is not currently supported on mobile platforms; more work is required to make instance compaction work with mobile vertex buffers.
Updated ISM to use order preservation for meshes that render with a translucent material

diff --git a/Engine/Plugins/FX/Niagara/Source/Niagara/Private/NiagaraRendererMeshes.cpp b/Engine/Plugins/FX/Niagara/Source/Niagara/Private/NiagaraRendererMeshes.cpp
--- a/Engine/Plugins/FX/Niagara/Source/Niagara/Private/NiagaraRendererMeshes.cpp
+++ b/Engine/Plugins/FX/Niagara/Source/Niagara/Private/NiagaraRendererMeshes.cpp
@@ -949,7 +949,7 @@ void FNiagaraRendererMeshes::SetupElementForGPUScene(
 		// * Are distance culled
 		// * Are culled because of mismatched MeshIndex
 		// * Are culled because of mismatched VisibilityTag
-		GPUSceneRes.DynamicPrimitiveData.bForceInstanceCulling = true;
+		OutMeshBatchElement.bForceInstanceCulling = true;
 
 		GPUSceneRes.DynamicPrimitiveData.DataWriterGPU = FGPUSceneWriteDelegate::CreateLambda(				
 			[&GPUSceneRes](FRDGBuilder& GraphBuilder, const FGPUSceneWriteDelegateParams& Params)
diff --git a/Engine/Shaders/Private/InstanceCulling/BuildInstanceDrawCommands.usf b/Engine/Shaders/Private/InstanceCulling/BuildInstanceDrawCommands.usf
--- a/Engine/Shaders/Private/InstanceCulling/BuildInstanceDrawCommands.usf
+++ b/Engine/Shaders/Private/InstanceCulling/BuildInstanceDrawCommands.usf
@@ -5,10 +5,16 @@
 	#define USES_PER_INSTANCE_RANDOM 1
 #endif
 
+#if ENABLE_INSTANCE_COMPACTION
+	// Currently, instance compaction is the only reason to need extended payload data, so tie it to the permutation for now
+	#define ENABLE_EXTENDED_INSTANCE_CULLING_PAYLOADS 1
+#endif
+
 #include "../Common.ush"
 #include "../SceneData.ush"
 #include "../LightmapData.ush"
 #include "InstanceCullingCommon.ush"
+#include "InstanceCompactionCommon.ush"
 #include "../Nanite/NaniteDataDecode.ush"
 #include "../Nanite/HZBCull.ush"
 #include "../WaveOpUtil.ush"
@@ -19,6 +25,7 @@
 	#define LOAD_BALANCER_SINGLE_INSTANCE_MODE 1
 #endif
 
+
 #include "InstanceCullingLoadBalancer.ush"
 
 uint InstanceSceneDataSOAStride;
@@ -46,6 +53,11 @@ RWStructuredBuffer<uint> InstanceIdsBufferOut;
 RWStructuredBuffer<float4> InstanceIdsBufferOutMobile;
 RWBuffer<uint> DrawIndirectArgsBufferOut;
 
+// Used for draw commands that need to use compaction to preserve instance order
+StructuredBuffer<FPackedDrawCommandCompactionData> DrawCommandCompactionData;
+RWStructuredBuffer<uint> CompactInstanceIdsBufferOut;
+RWStructuredBuffer<uint> CompactionBlockCounts;
+
 // this is just to avoid compiling packing code for unrelated shaders
 #define USES_PACKED_INSTANCE_DATA (FEATURE_LEVEL == FEATURE_LEVEL_ES3_1 && VF_SUPPORTS_PRIMITIVE_SCENE_DATA)
 
@@ -152,6 +164,15 @@ void WriteInstance(uint Offset, uint InstanceId, FInstanceSceneData InstanceData
 #endif
 }
 
+void WriteInstanceForCompaction(uint Offset, bool bVisible, uint InstanceId, FInstanceSceneData InstanceData, uint ViewIdIndex)
+{
+	checkSlow(InstanceId < GPUSceneNumInstances);
+
+	uint PackedId = bVisible ? InstanceId | (ViewIdIndex << 28U) : 0xFFFFFFFFU;
+	checkStructuredBufferAccessSlow(CompactInstanceIdsBufferOut, Offset);
+	CompactInstanceIdsBufferOut[Offset] = PackedId;	
+}
+
 bool ShouldRenderInstance(FInstanceSceneData InstanceData, FDrawCommandDesc Desc)
 {
 #if DEBUG_MODE
@@ -259,42 +280,65 @@ void InstanceCullBuildInstanceIdBufferCS(uint3 GroupId : SV_GroupID, int GroupTh
 		return;
 	}
 
-	// Check dynamic flag
-	const bool bDynamicInstanceDataOffset = (WorkSetup.Item.Payload & 1U) != 0U;
-	uint InstanceDataOffset = WorkSetup.Item.InstanceDataOffset;
+	// Extract the draw command payload
+	const FInstanceCullingPayload Payload = LoadInstanceCullingPayload(WorkSetup.Item.Payload, BatchInfo);
+
+	uint InstanceId = WorkSetup.Item.InstanceDataOffset + uint(WorkSetup.LocalItemIndex);
 
-	if (bDynamicInstanceDataOffset)
+	// Check dynamic flag to offset the instance ID
+	if (Payload.bDynamicInstanceDataOffset)
 	{
-		InstanceDataOffset += BatchInfo.DynamicInstanceIdOffset;
-		checkSlow(InstanceDataOffset + uint(WorkSetup.LocalItemIndex) < BatchInfo.DynamicInstanceIdMax);
+		InstanceId += BatchInfo.DynamicInstanceIdOffset;
+		checkSlow(InstanceId < BatchInfo.DynamicInstanceIdMax);
 	}
 
-	uint InstanceId = InstanceDataOffset + uint(WorkSetup.LocalItemIndex);
-
-	// Work item payload stores the IndirectArgsOffset, which must be offset for the batch.
-	uint IndirectArgIndex = BatchInfo.IndirectArgOffset + (WorkSetup.Item.Payload >> 1U);
-
-	// TODO: Add unpack and store as bit flags
 	// Load auxiliary per-instanced-draw command info
-	FDrawCommandDesc DrawCommandDesc = DrawCommandDescs[IndirectArgIndex];
+	const FDrawCommandDesc DrawCommandDesc = DrawCommandDescs[Payload.IndirectArgIndex];
+
+	// Extract compaction data (if applicable)	
+#if ENABLE_INSTANCE_COMPACTION
+	const bool bCompactInstances = Payload.CompactionDataIndex != 0xFFFFFFFFU;
+#else
+	const bool bCompactInstances = false;
+#endif
+	uint CompactOutputInstanceIndex = 0;
+	FDrawCommandCompactionData CompactionData = (FDrawCommandCompactionData)0;	
+	BRANCH
+	if (bCompactInstances)
+	{
+		CompactionData = UnpackDrawCommandCompactionData(DrawCommandCompactionData[Payload.CompactionDataIndex]);
+		
+		const uint WorkItemLocalInstanceOffset = WorkSetup.Item.InstanceDataOffset - Payload.InstanceDataOffset;
+		CompactOutputInstanceIndex = Payload.RunInstanceOffset + WorkItemLocalInstanceOffset + uint(WorkSetup.LocalItemIndex);
+	}
 
 	// TODO: This must be read-modify-written when batching such that the final offset that is fed to the VS is correct.
 	//       Then we don't need to add the batch offset (BatchInfo.InstanceDataWriteOffset)
-	uint InstanceDataOutputOffset = InstanceIdOffsetBuffer[IndirectArgIndex];
+	const uint InstanceDataOutputOffset = InstanceIdOffsetBuffer[Payload.IndirectArgIndex];
 
-	FInstanceSceneData InstanceData = GetInstanceSceneData(InstanceId, InstanceSceneDataSOAStride);
+	const FInstanceSceneData InstanceData = GetInstanceSceneData(InstanceId, InstanceSceneDataSOAStride);
 
 	if (!ShouldRenderInstance(InstanceData, DrawCommandDesc))
 	{
 		return;
 	}
 
+	uint NumVisibleInstances = 0;
 #if STEREO_CULLING_MODE
-	if (IsInstanceVisible(InstanceData, BatchInfo.ViewIdsOffset + 0U) || IsInstanceVisible(InstanceData, BatchInfo.ViewIdsOffset + 1U))
+	const bool bVisible = IsInstanceVisible(InstanceData, BatchInfo.ViewIdsOffset + 0U) || IsInstanceVisible(InstanceData, BatchInfo.ViewIdsOffset + 1U);
+	NumVisibleInstances += bVisible ? 2 : 0;
+	
+	BRANCH
+	if (bCompactInstances)
 	{
-		uint OutputOffset = 0U;
-		InterlockedAdd(DrawIndirectArgsBufferOut[IndirectArgIndex * INDIRECT_ARGS_NUM_WORDS + 1], 2U, OutputOffset);
-
+		const uint OutputOffset = CompactOutputInstanceIndex * 2U;
+		WriteInstanceForCompaction(CompactionData.SrcInstanceIdOffset + OutputOffset + 0U, bVisible, InstanceId, InstanceData, 0U);
+		WriteInstanceForCompaction(CompactionData.SrcInstanceIdOffset + OutputOffset + 1U, bVisible, InstanceId, InstanceData, 1U);
+	}
+	else if (bVisible)
+	{
+		uint OutputOffset;
+		InterlockedAdd(DrawIndirectArgsBufferOut[Payload.IndirectArgIndex * INDIRECT_ARGS_NUM_WORDS + 1], 2U, OutputOffset);				
 		WriteInstance(InstanceDataOutputOffset + OutputOffset + 0U, InstanceId, InstanceData, 0U);
 		WriteInstance(InstanceDataOutputOffset + OutputOffset + 1U, InstanceId, InstanceData, 1U);
 	}
@@ -303,18 +347,35 @@ void InstanceCullBuildInstanceIdBufferCS(uint3 GroupId : SV_GroupID, int GroupTh
 	for (uint ViewIdIndex = 0; ViewIdIndex < BatchInfo.NumViewIds; ++ViewIdIndex)
 	{
 		// Culling is disabled for multi-view
-		if (IsInstanceVisible(InstanceData, BatchInfo.ViewIdsOffset + ViewIdIndex))
+		const bool bVisible = IsInstanceVisible(InstanceData, BatchInfo.ViewIdsOffset + ViewIdIndex);
+		NumVisibleInstances += bVisible ? 1 : 0;
+		
+		BRANCH
+		if (bCompactInstances)
+		{
+			const uint OutputOffset = CompactOutputInstanceIndex * BatchInfo.NumViewIds + ViewIdIndex;
+			WriteInstanceForCompaction(CompactionData.SrcInstanceIdOffset + OutputOffset, bVisible, InstanceId, InstanceData, ViewIdIndex);
+		}
+		else if (bVisible)
 		{
-			uint OutputOffset = 0;
-			// TOOD: if all items in the group-batch target the same draw args the more efficient warp-collective functions can be used
+			// TODO: if all items in the group-batch target the same draw args the more efficient warp-collective functions can be used
 			//       detected as FInstanceBatch.NumItems == 1. Can switch dynamically or bin the items that fill a group and dispatch separately with permutation.
 			// TODO: if the arg only has a single item, and culling is not enabled, then we can skip the atomics. Again do dynamically or separate permutation.
-			InterlockedAdd(DrawIndirectArgsBufferOut[IndirectArgIndex * INDIRECT_ARGS_NUM_WORDS + 1], 1U, OutputOffset);
+			uint OutputOffset;
+			InterlockedAdd(DrawIndirectArgsBufferOut[Payload.IndirectArgIndex * INDIRECT_ARGS_NUM_WORDS + 1], 1U, OutputOffset);
 
 			WriteInstance(InstanceDataOutputOffset + OutputOffset, InstanceId, InstanceData, ViewIdIndex);
 		}
 	}
 #endif // STEREO_CULLING_MODE
+
+	BRANCH
+	if (bCompactInstances && NumVisibleInstances > 0)
+	{
+		// Determine compaction block and atomically increment its count			
+		const uint BlockIndex = GetCompactionBlockIndexFromInstanceIndex(CompactOutputInstanceIndex);
+		InterlockedAdd(CompactionBlockCounts[CompactionData.BlockOffset + BlockIndex], NumVisibleInstances);
+	}
 }
 
 
diff --git a/Engine/Shaders/Private/InstanceCulling/CompactVisibleInstances.usf b/Engine/Shaders/Private/InstanceCulling/CompactVisibleInstances.usf
new file mode 100644
--- /dev/null
+++ b/Engine/Shaders/Private/InstanceCulling/CompactVisibleInstances.usf
@@ -0,0 +1,161 @@
+// Copyright Epic Games, Inc. All Rights Reserved.
+
+#include "../Common.ush"
+#include "../ComputeShaderUtils.ush"
+#include "../ThreadGroupPrefixSum.ush"
+#include "InstanceCompactionCommon.ush"
+
+#if COMPILER_SUPPORTS_WAVE_VOTE
+	// TODO: Need to also check that wave width == thread group size (or implement workaround for other wave sizes)
+	#define COMPACTION_USES_WAVE_OPS 0
+#endif
+
+/** Data needed per draw command for compaction */
+StructuredBuffer<FPackedDrawCommandCompactionData> DrawCommandCompactionData;
+
+/**
+ * CalculateCompactBlockInstanceOffsets
+ *
+ * Compaction phase 1 - Generate output offsets of "blocks" of source instance IDs such that all valid IDs will be compacted for
+ * the draw comand. Each thread group in this dispatch aligns with a single draw command. It uses a hybrid of parallelization and
+ * serialization to calculate the compact output offsets in the instance ID buffer in a way that doesn't require multiple dispatches.
+ */
+#if CALCULATE_COMPACT_BLOCK_INSTANCE_OFFSETS
+
+/** Number of valid instances per block in a draw command (accumulated by BuildInstanceDrawCommands.usf) */
+StructuredBuffer<uint> BlockInstanceCounts;
+
+/** Destination instance offset per block of instances in a draw command. (Relative to draw command's InstanceIdOffset) */
+RWStructuredBuffer<uint> BlockDestInstanceOffsetsOut;
+/** The indirect args buffer to receive the final, compacted instance count */
+RWBuffer<uint> DrawIndirectArgsBufferOut;
+
+/** Entry point */
+[numthreads(NUM_THREADS_PER_GROUP, 1, 1)]
+void CalculateCompactBlockInstanceOffsetsCS(uint3 GroupId : SV_GroupID, uint GroupThreadIndex : SV_GroupIndex)
+{
+	const uint DrawCommandIndex = GetUnWrappedDispatchGroupId(GroupId);
+	const FDrawCommandCompactionData DrawCommandData = UnpackDrawCommandCompactionData(DrawCommandCompactionData[DrawCommandIndex]);
+	const uint BlockOffset = DrawCommandData.BlockOffset;
+	const uint NumBlocksInDrawCommand = GetCompactionBlockCount(DrawCommandData.NumInstances);
+	const uint NumWorkIterations = ((NumBlocksInDrawCommand - 1U) / NUM_THREADS_PER_GROUP) + 1U;
+
+	uint CurInstanceOffset = 0;
+	uint CurDrawCommandBlockIndex = 0;
+	for (uint IterIndex = 0; IterIndex < NumWorkIterations; ++IterIndex)
+	{
+		const uint DrawCommandBlockIndex = CurDrawCommandBlockIndex + GroupThreadIndex;
+		const uint GlobalBlockIndex = DrawCommandData.BlockOffset + DrawCommandBlockIndex;
+
+		uint BlockInstanceDestCount = 0;
+		BRANCH
+		if (DrawCommandBlockIndex < NumBlocksInDrawCommand)
+		{
+			BlockInstanceDestCount = BlockInstanceCounts[GlobalBlockIndex];
+		}
+
+		uint PrefixSum, GroupSum;
+	#if COMPACTION_USES_WAVE_OPS
+		PrefixSum = WavePrefixSum(BlockInstanceDestCount);
+		GroupSum = WaveActiveSum(BlockInstanceDestCount);
+	#else
+		PrefixSum = ThreadGroupPrefixSum(BlockInstanceDestCount, GroupThreadIndex, GroupSum);
+	#endif // COMPACTION_USES_WAVE_OPS
+		
+		// Write out the instance offsets
+		BRANCH
+		if (DrawCommandBlockIndex < NumBlocksInDrawCommand)
+		{
+			BlockDestInstanceOffsetsOut[GlobalBlockIndex] = CurInstanceOffset + PrefixSum;
+		}
+
+		CurInstanceOffset += GroupSum;
+		CurDrawCommandBlockIndex += NUM_THREADS_PER_GROUP;
+	}
+
+	if (GroupThreadIndex == 0)
+	{
+		// Set the final indirect arg instance count
+		DrawIndirectArgsBufferOut[DrawCommandData.IndirectArgIndex * INDIRECT_ARGS_NUM_WORDS + 1] = CurInstanceOffset;
+	}
+}
+
+#endif // CALCULATE_COMPACT_BLOCK_INSTANCE_OFFSETS
+
+/**
+ * CompactVisibleInstances
+ *
+ * Compaction phase 2 - Write out the visible instances in compact order
+ * Each group outputs a compact block of visible instances. The max source (uncompacted) block size is equal to
+ * COMPACTION_BLOCK_NUM_INSTANCES * NumViews. The number of threads in the group is equal to the number of instances
+ * per block.
+ */ 
+#if COMPACT_VISIBLE_INSTANCES
+
+/** Draw command index per block */
+StructuredBuffer<uint> BlockDrawCommandIndices;
+/** The raw, uncompacted instance ID buffer (provided by BuildInstanceDrawCommands.usf) */
+StructuredBuffer<uint> InstanceIdsBufferIn;
+/** Instance output offsets per block (provided by CalculateCompactBlockInstanceOffsetsCS) */
+StructuredBuffer<uint> BlockDestInstanceOffsets;
+
+/** The compacted output buffer of InstanceIds (non-mobile) */
+RWStructuredBuffer<uint> InstanceIdsBufferOut;
+/** The compacted output buffer of instance data consumed by vertex shaders on mobile */
+RWStructuredBuffer<float4> InstanceIdsBufferOutMobile;
+
+/** Entry point */
+[numthreads(NUM_THREADS_PER_GROUP, 1, 1)]
+void CompactVisibleInstances(uint3 GroupId : SV_GroupID, uint GroupThreadIndex : SV_GroupIndex)
+{
+	const uint GlobalBlockIndex = GetUnWrappedDispatchGroupId(GroupId);
+	const FDrawCommandCompactionData DrawCommandData = UnpackDrawCommandCompactionData(DrawCommandCompactionData[BlockDrawCommandIndices[GlobalBlockIndex]]);
+	const uint DrawCommandBlockIndex = GlobalBlockIndex - DrawCommandData.BlockOffset;
+	const uint MaxInstancesPerBlock = NUM_THREADS_PER_GROUP * DrawCommandData.NumViews;
+	const uint BlockSrcInstanceCount = min(DrawCommandData.NumInstances * DrawCommandData.NumViews - DrawCommandBlockIndex * MaxInstancesPerBlock, MaxInstancesPerBlock);
+	const uint BlockSrcInstanceOffset = DrawCommandData.SrcInstanceIdOffset + DrawCommandBlockIndex * MaxInstancesPerBlock;
+	const uint BlockDestInstanceOffset = DrawCommandData.DestInstanceIdOffset + BlockDestInstanceOffsets[GlobalBlockIndex];
+
+	// These are the relative offsets for the whole group from the start of the source and destination offsets for the block, respectively
+	uint CurGroupSrcOffset = 0;
+	uint CurGroupDestOffset = 0;
+
+	while (CurGroupSrcOffset < BlockSrcInstanceCount) // may loop as many times as NumViews
+	{
+		uint InstanceId = 0xFFFFFFFFU;
+		
+		const uint BlockSrcInstanceIndex = CurGroupSrcOffset + GroupThreadIndex;
+		const uint GlobalSrcInstanceIndex = BlockSrcInstanceOffset + BlockSrcInstanceIndex;
+		BRANCH
+		if (BlockSrcInstanceIndex < BlockSrcInstanceCount)
+		{
+			InstanceId = InstanceIdsBufferIn[GlobalSrcInstanceIndex];
+		}
+		
+		const bool bIsValidInstance = (InstanceId != 0xFFFFFFFFU);
+
+		// now calculate the index based on the valid index in the group
+		uint PrefixSum, GroupSum;
+	#if COMPACTION_USES_WAVE_OPS
+		PrefixSum = WavePrefixCountBits(bIsValidInstance);
+		GroupSum = WaveActiveCountBits(bIsValidInstance);
+	#else
+		// NOTE: This must be called even by threads that won't output an InstanceId
+		PrefixSum = ThreadGroupPrefixSum(bIsValidInstance ? 1U : 0U, GroupThreadIndex, GroupSum);
+	#endif // COMPACTION_USES_WAVE_OPS
+
+		if (bIsValidInstance)
+		{
+			const uint BlockDestInstanceIndex = CurGroupDestOffset + PrefixSum;
+			const uint GlobalDestInstanceIndex = BlockDestInstanceOffset + BlockDestInstanceIndex;
+
+			// TODO: This currently doesn't support mobile, which outputs transform data into a vertex buffer
+			InstanceIdsBufferOut[GlobalDestInstanceIndex] = InstanceId;
+		}
+
+		CurGroupSrcOffset += NUM_THREADS_PER_GROUP;
+		CurGroupDestOffset += GroupSum;
+	}	
+}
+
+#endif // COMPACT_VISIBLE_INSTANCES
\ No newline at end of file
diff --git a/Engine/Shaders/Private/InstanceCulling/InstanceCompactionCommon.ush b/Engine/Shaders/Private/InstanceCulling/InstanceCompactionCommon.ush
new file mode 100644
--- /dev/null
+++ b/Engine/Shaders/Private/InstanceCulling/InstanceCompactionCommon.ush
@@ -0,0 +1,57 @@
+// Copyright Epic Games, Inc. All Rights Reserved.
+
+#pragma once
+
+#ifndef COMPACTION_BLOCK_NUM_INSTANCES
+	#error COMPACTION_BLOCK_NUM_INSTANCES must be defined
+#endif
+
+/** Determines number of compaction blocks required for the specified number of instances */
+uint GetCompactionBlockCount(uint NumInstances)
+{
+	return ((NumInstances - 1U) / COMPACTION_BLOCK_NUM_INSTANCES) + 1U;
+}
+
+/** Determines number of compaction blocks required for the specified number of instances */
+uint GetCompactionBlockIndexFromInstanceIndex(uint InstanceIndex)
+{
+	return InstanceIndex / COMPACTION_BLOCK_NUM_INSTANCES;
+}
+
+struct FPackedDrawCommandCompactionData
+{
+	uint NumInstances_NumViews;
+	uint BlockOffset;
+	uint IndirectArgIndex;
+	uint SrcInstanceIdOffset;
+	uint DestInstanceIdOffset;
+};
+
+struct FDrawCommandCompactionData
+{
+	/** The number of views for which the instances were written */
+	uint NumViews;
+	/** The total number of instances in the draw command */
+	uint NumInstances;
+	/** The offset of the first compaction "block" (group of N instances * NumViewIds) of the draw command */
+	uint BlockOffset;
+	/** The indirect argument index of the draw command */
+	uint IndirectArgIndex;
+	/** The offset of the input instance ID buffer for the start of the draw command before compaction */
+	uint SrcInstanceIdOffset;
+	/** The offset of the output instance ID buffer for the start of the draw command after compaction */
+	uint DestInstanceIdOffset;
+};
+
+FDrawCommandCompactionData UnpackDrawCommandCompactionData(FPackedDrawCommandCompactionData Data)
+{
+	FDrawCommandCompactionData Output;
+	Output.NumViews = Data.NumInstances_NumViews & 0xFFU;
+	Output.NumInstances = Data.NumInstances_NumViews >> 8;
+	Output.BlockOffset = Data.BlockOffset;
+	Output.IndirectArgIndex = Data.IndirectArgIndex;
+	Output.SrcInstanceIdOffset = Data.SrcInstanceIdOffset;
+	Output.DestInstanceIdOffset = Data.DestInstanceIdOffset;
+
+	return Output;
+}
\ No newline at end of file
diff --git a/Engine/Shaders/Private/InstanceCulling/InstanceCullingCommon.ush b/Engine/Shaders/Private/InstanceCulling/InstanceCullingCommon.ush
--- a/Engine/Shaders/Private/InstanceCulling/InstanceCullingCommon.ush
+++ b/Engine/Shaders/Private/InstanceCulling/InstanceCullingCommon.ush
@@ -8,6 +8,9 @@
 #define BATCH_PROCESSING_MODE_NUM 1
 #endif // BATCH_PROCESSING_MODE_NUM
 
+#ifndef ENABLE_EXTENDED_INSTANCE_CULLING_PAYLOADS
+#define ENABLE_EXTENDED_INSTANCE_CULLING_PAYLOADS 0
+#endif // ENABLE_EXTENDED_INSTANCE_CULLING_PAYLOADS
 
 struct FDrawCommandDesc
 {
@@ -18,6 +21,8 @@ struct FContextBatchInfo
 {
 	uint IndirectArgOffset;
 	uint InstanceDataWriteOffset;
+	uint PayloadOffset;
+	uint CompactionDataOffset;
 	uint ViewIdsOffset;
 	uint NumViewIds;
 	uint DynamicInstanceIdOffset;
@@ -40,4 +45,80 @@ uint GetItemDataOffset(FContextBatchInfo BatchInfo, uint BatchProcessingMode)
 #if ENABLE_BATCH_MODE
 StructuredBuffer<FContextBatchInfo> BatchInfos;
 StructuredBuffer<uint> BatchInds;
-#endif // ENABLE_BATCH_MODE
\ No newline at end of file
+#endif // ENABLE_BATCH_MODE
+
+struct FInstanceCullingPayload
+{
+	/** Whether or not the instances are dynamic */
+	bool bDynamicInstanceDataOffset;
+	/** The IndirectArgs index for instances in the run. */
+	uint IndirectArgIndex;
+	/** The instance data offset of the start of the current instance run, or -1 if this data should not be needed for this work item */
+	uint InstanceDataOffset;
+	/** The number of instances added from other runs before this one, or -1 if this data should not be needed for this work item */
+	uint RunInstanceOffset;
+	/** When compaction is enabled, the index into the FDrawCommandCompactionData buffer (or -1 otherwise) */
+	uint CompactionDataIndex;
+};
+
+struct FPackedInstanceCullingPayload
+{
+	uint bDynamicInstanceDataOffset_IndirectArgIndex;
+	uint InstanceDataOffset;
+	uint RunInstanceOffset;
+	uint CompactionDataIndex;
+};
+
+bool IsExtendedInstanceCullingPayload(uint WorkItemPayload)
+{
+	return (WorkItemPayload & 1U) != 0U;
+}
+
+FInstanceCullingPayload UnpackInstanceCullingPayload(FPackedInstanceCullingPayload Payload)
+{
+	FInstanceCullingPayload Output;
+	Output.bDynamicInstanceDataOffset = (Payload.bDynamicInstanceDataOffset_IndirectArgIndex >> 31U) != 0U;
+	Output.IndirectArgIndex = Payload.bDynamicInstanceDataOffset_IndirectArgIndex & 0x7FFFFFFFU;
+	Output.InstanceDataOffset = Payload.InstanceDataOffset;
+	Output.RunInstanceOffset = Payload.RunInstanceOffset;
+	Output.CompactionDataIndex = Payload.CompactionDataIndex;
+
+	return Output;
+}
+
+#if ENABLE_EXTENDED_INSTANCE_CULLING_PAYLOADS
+StructuredBuffer<FPackedInstanceCullingPayload> InstanceCullingPayloads;
+#endif
+
+/**
+ * Retrieves the instance culling payload data. When instance compaction is not needed to preserve instance order, much of the
+ * instance culling payload need not take up space, so it is compacted into the 26-bit payload provided by the load balancer.
+ * Otherwise, the first bit of the load balancer's payload is set, and the remaining 25 bits are used to index into the payload
+ * structured buffer.
+ */
+FInstanceCullingPayload LoadInstanceCullingPayload(uint WorkItemPayload, FContextBatchInfo BatchInfo)
+{
+	FInstanceCullingPayload Output;
+
+#if ENABLE_EXTENDED_INSTANCE_CULLING_PAYLOADS
+	BRANCH
+	if (IsExtendedInstanceCullingPayload(WorkItemPayload))
+	{
+		Output = UnpackInstanceCullingPayload(InstanceCullingPayloads[BatchInfo.PayloadOffset + (WorkItemPayload >> 1U)]);
+		Output.IndirectArgIndex += BatchInfo.IndirectArgOffset;
+		Output.CompactionDataIndex += BatchInfo.CompactionDataOffset;
+		return Output;
+	}
+#else
+	// Ensure that the payload is not extended
+	checkSlow(!IsExtendedInstanceCullingPayload(WorkItemPayload));
+#endif
+
+	Output.bDynamicInstanceDataOffset = (WorkItemPayload & 2U) != 0U;
+	Output.IndirectArgIndex = BatchInfo.IndirectArgOffset + (WorkItemPayload >> 2U);
+	Output.InstanceDataOffset = 0xFFFFFFFFU;
+	Output.RunInstanceOffset = 0xFFFFFFFFU;
+	Output.CompactionDataIndex = 0xFFFFFFFFU;
+
+	return Output;
+}
\ No newline at end of file
diff --git a/Engine/Shaders/Private/ThreadGroupPrefixSum.ush b/Engine/Shaders/Private/ThreadGroupPrefixSum.ush
new file mode 100644
--- /dev/null
+++ b/Engine/Shaders/Private/ThreadGroupPrefixSum.ush
@@ -0,0 +1,67 @@
+// Copyright Epic Games, Inc. All Rights Reserved.
+
+#pragma once
+
+/**
+ * ThreadGroupPrefixSum(Value, GroupThreadIndex [, GroupSumOut])
+ *
+ * Calculates the prefix sum (and optionally group sum) of a given value across threads of a thread group.
+ *
+ * EXAMPLE USES:
+ *	float GroupSum;
+ *	float PrefixSum = ThreadGroupPrefixSum(ValueFloat, GroupThreadIndex, GroupSum);
+ *
+ *	int PrefixSum = ThreadGroupPrefixSum(ValueInt, GroupThreadIndex);
+ *
+ * NOTES:
+ * - PrefixSum means the sum of the value for all threads whose group thread index is LESS than the current
+ * - Only scalar types are currently supported
+ * - All threads in the group must be active when calling this method; it cannot be used in a branch (or after an early return)
+ *	 that doesn't include the entire group.
+ */
+
+#ifndef NUM_THREADS_PER_GROUP
+	#error NUM_THREADS_PER_GROUP must be defined, and must be equal to the thread group size of the caller of ThreadGroupPrefixSum
+#endif
+
+// re-use the same groupshared memory, in case the caller utilizes multiple overloads
+groupshared uint ThreadGroupPrefixSumWorkspace[2][NUM_THREADS_PER_GROUP];
+
+#define DECLARE_THREAD_GROUP_PREFIX_SUM(ValType, CastToUint, CastFromUint) \
+	ValType ThreadGroupPrefixSum(ValType Value, uint GroupThreadIndex, inout ValType GroupSum) \
+	{ \
+		uint Curr = 0; \
+		ThreadGroupPrefixSumWorkspace[Curr][GroupThreadIndex] = CastToUint(Value); \
+		GroupMemoryBarrierWithGroupSync(); \
+		for (uint i = 1U; i <= (NUM_THREADS_PER_GROUP / 2U); i *= 2U) \
+		{ \
+			const uint Next = 1U - Curr; \
+			if (GroupThreadIndex < i) \
+			{ \
+				ThreadGroupPrefixSumWorkspace[Next][GroupThreadIndex] = \
+					ThreadGroupPrefixSumWorkspace[Curr][GroupThreadIndex]; \
+			} \
+			else \
+			{ \
+				ThreadGroupPrefixSumWorkspace[Next][GroupThreadIndex] = CastToUint( \
+					CastFromUint(ThreadGroupPrefixSumWorkspace[Curr][GroupThreadIndex]) + \
+					CastFromUint(ThreadGroupPrefixSumWorkspace[Curr][GroupThreadIndex - i]) \
+				); \
+			} \
+			Curr = Next; \
+			GroupMemoryBarrierWithGroupSync(); \
+		} \
+		GroupSum = CastFromUint(ThreadGroupPrefixSumWorkspace[Curr][NUM_THREADS_PER_GROUP - 1]); \
+		return (GroupThreadIndex == 0U) ? CastFromUint(0) : CastFromUint(ThreadGroupPrefixSumWorkspace[Curr][GroupThreadIndex - 1]); \
+	} \
+	ValType ThreadGroupPrefixSum(ValType Value, uint GroupThreadIndex) \
+	{ \
+		ValType GroupSum; \
+		return ThreadGroupPrefixSum(Value, GroupThreadIndex, GroupSum); \
+	}
+
+DECLARE_THREAD_GROUP_PREFIX_SUM(float, asuint, asfloat)
+DECLARE_THREAD_GROUP_PREFIX_SUM(int, uint, int)
+DECLARE_THREAD_GROUP_PREFIX_SUM(uint, uint, uint)
+
+#undef DECLARE_THREAD_GROUP_PREFIX_SUM
\ No newline at end of file
diff --git a/Engine/Shaders/Private/VirtualShadowMaps/BuildPerPageDrawCommands.usf b/Engine/Shaders/Private/VirtualShadowMaps/BuildPerPageDrawCommands.usf
--- a/Engine/Shaders/Private/VirtualShadowMaps/BuildPerPageDrawCommands.usf
+++ b/Engine/Shaders/Private/VirtualShadowMaps/BuildPerPageDrawCommands.usf
@@ -165,19 +165,18 @@ void CullPerPageDrawCommandsCs(uint3 GroupId : SV_GroupID, int GroupThreadIndex
 		return;
 	}
 
-	const bool bDynamicInstanceDataOffset = (WorkSetup.Item.Payload & 1U) != 0U;
+	const FInstanceCullingPayload Payload = LoadInstanceCullingPayload(WorkSetup.Item.Payload, BatchInfo);
+
 	uint InstanceDataOffset = WorkSetup.Item.InstanceDataOffset;
 
-	if (bDynamicInstanceDataOffset)
+	if (Payload.bDynamicInstanceDataOffset)
 	{
 		InstanceDataOffset += BatchInfo.DynamicInstanceIdOffset;
 		checkSlow(InstanceDataOffset + uint(WorkSetup.LocalItemIndex) < BatchInfo.DynamicInstanceIdMax);
 	}
 
-	uint InstanceId = InstanceDataOffset + uint(WorkSetup.LocalItemIndex);
-	// Work item payload stores the IndirectArgsOffset, which must be offset for the batch.
-	uint IndirectArgIndex = BatchInfo.IndirectArgOffset + (WorkSetup.Item.Payload >> 1U);
-	FDrawCommandDesc DrawCommandDesc = DrawCommandDescs[IndirectArgIndex];
+	uint InstanceId = InstanceDataOffset + uint(WorkSetup.LocalItemIndex);	
+	FDrawCommandDesc DrawCommandDesc = DrawCommandDescs[Payload.IndirectArgIndex];
 
 	const bool bNearClip = (NEAR_CLIP != 0);
 
@@ -266,7 +265,7 @@ void CullPerPageDrawCommandsCs(uint3 GroupId : SV_GroupID, int GroupThreadIndex
 #endif // USE_HZB_OCCLUSION
 
 						++ThreadTotalForAllViews;
-						WriteCmd(MipViewId, InstanceId, IndirectArgIndex, ShouldCacheInstanceAsStatic(InstanceData));
+						WriteCmd(MipViewId, InstanceId, Payload.IndirectArgIndex, ShouldCacheInstanceAsStatic(InstanceData));
 					}
 					else
 					{
@@ -295,7 +294,7 @@ void CullPerPageDrawCommandsCs(uint3 GroupId : SV_GroupID, int GroupThreadIndex
 	StatsBufferInterlockedAdd(VSM_STAT_NON_NANITE_INSTANCES_DRAWN, ThreadTotalForAllViews);
 
 	// Accumulate total number of instances for each indirect argument, is also used to allocate space and output compact range of instances later
-	InterlockedAdd(DrawIndirectArgsBufferOut[IndirectArgIndex * INDIRECT_ARGS_NUM_WORDS + 1], ThreadTotalForAllViews);
+	InterlockedAdd(DrawIndirectArgsBufferOut[Payload.IndirectArgIndex * INDIRECT_ARGS_NUM_WORDS + 1], ThreadTotalForAllViews);
 }
 
 Buffer<uint> DrawIndirectArgsBuffer;
diff --git a/Engine/Source/Runtime/Engine/Private/InstancedStaticMesh.cpp b/Engine/Source/Runtime/Engine/Private/InstancedStaticMesh.cpp
--- a/Engine/Source/Runtime/Engine/Private/InstancedStaticMesh.cpp
+++ b/Engine/Source/Runtime/Engine/Private/InstancedStaticMesh.cpp
@@ -1608,6 +1608,20 @@ void FInstancedStaticMeshSceneProxy::SetupInstancedMeshBatch(int32 LODIndex, int
 	BatchElement0.InstancedLODIndex = LODIndex;
 	BatchElement0.UserIndex = 0;
 	BatchElement0.PrimitiveUniformBuffer = GetUniformBuffer();
+	
+	if (OutMeshBatch.MaterialRenderProxy)
+	{
+		// If the material on the mesh batch is translucent, then preserve the instance draw order to prevent flickering
+		// TODO: This should use depth sorting of instances instead, once that is implemented for GPU Scene instances
+		const ERHIFeatureLevel::Type FeatureLevel = GetScene().GetFeatureLevel();
+		
+		// NOTE: For now, this feature is not supported for mobile platforms
+		if (FeatureLevel > ERHIFeatureLevel::ES3_1)
+		{
+			const FMaterial& Material = OutMeshBatch.MaterialRenderProxy->GetIncompleteMaterialWithFallback(FeatureLevel);
+			BatchElement0.bPreserveInstanceOrder = IsTranslucentBlendMode(Material.GetBlendMode());
+		}
+	}
 
 	BatchElement0.NumInstances = NumInstances;
 }
diff --git a/Engine/Source/Runtime/Engine/Public/MeshBatch.h b/Engine/Source/Runtime/Engine/Public/MeshBatch.h
--- a/Engine/Source/Runtime/Engine/Public/MeshBatch.h
+++ b/Engine/Source/Runtime/Engine/Public/MeshBatch.h
@@ -76,7 +76,6 @@ struct FMeshBatchDynamicPrimitiveData
 	EGPUSceneGPUWritePass DataWriterGPUPass = EGPUSceneGPUWritePass::None;
 	uint16 PayloadDataFlags = 0;
 	uint32 NumInstanceCustomDataFloats = 0;
-	bool bForceInstanceCulling = false;
 
 	FORCEINLINE void SetPayloadDataFlags(uint16 Flags, bool bValue)
 	{
@@ -203,6 +202,8 @@ struct FMeshBatchElement
 	uint32 bUserDataIsColorVertexBuffer : 1;
 	uint32 bIsSplineProxy : 1;
 	uint32 bIsInstanceRuns : 1;
+	uint32 bForceInstanceCulling : 1;
+	uint32 bPreserveInstanceOrder : 1;
 
 #if UE_ENABLE_DEBUG_DRAWING
 	/** Conceptual element index used for debug viewmodes. */
@@ -254,6 +255,8 @@ struct FMeshBatchElement
 	,	bUserDataIsColorVertexBuffer(false)
 	,	bIsSplineProxy(false)
 	,	bIsInstanceRuns(false)
+	,	bForceInstanceCulling(false)
+	,	bPreserveInstanceOrder(false)
 #if UE_ENABLE_DEBUG_DRAWING
 	,	VisualizeElementIndex(INDEX_NONE)
 #endif
diff --git a/Engine/Source/Runtime/Renderer/Private/InstanceCulling/InstanceCullingContext.cpp b/Engine/Source/Runtime/Renderer/Private/InstanceCulling/InstanceCullingContext.cpp
--- a/Engine/Source/Runtime/Renderer/Private/InstanceCulling/InstanceCullingContext.cpp
+++ b/Engine/Source/Runtime/Renderer/Private/InstanceCulling/InstanceCullingContext.cpp
@@ -27,11 +27,16 @@ static FAutoConsoleVariableRef CVarOcclusionCullInstances(
 	TEXT("Whether to do per instance occlusion culling for GPU instance culling."),
 	ECVF_RenderThreadSafe);
 
+static int32 GInstanceCullingAllowOrderPreservation = 1;
+static FAutoConsoleVariableRef CVarInstanceCullingAllowOrderPreservation(
+	TEXT("r.InstanceCulling.AllowInstanceOrderPreservation"),
+	GInstanceCullingAllowOrderPreservation,
+	TEXT("Whether or not to allow instances to preserve instance draw order using GPU compaction."),
+	ECVF_RenderThreadSafe);
+
 IMPLEMENT_STATIC_UNIFORM_BUFFER_SLOT(InstanceCullingUbSlot);
 IMPLEMENT_STATIC_UNIFORM_BUFFER_STRUCT(FInstanceCullingGlobalUniforms, "InstanceCulling", InstanceCullingUbSlot);
 
-
-
 static const TCHAR* BatchProcessingModeStr[] =
 {
 	TEXT("Generic"),
@@ -40,6 +45,11 @@ static const TCHAR* BatchProcessingModeStr[] =
 
 static_assert(UE_ARRAY_COUNT(BatchProcessingModeStr) == uint32(EBatchProcessingMode::Num), "BatchProcessingModeStr length does not match EBatchProcessingMode::Num, these must be kept in sync.");
 
+static bool IsInstanceOrderPreservationAllowed(ERHIFeatureLevel::Type FeatureLevel)
+{
+	// Instance order preservation is currently not supported on mobile platforms
+	return GInstanceCullingAllowOrderPreservation && FeatureLevel > ERHIFeatureLevel::ES3_1;
+}
 
 FMeshDrawCommandOverrideArgs GetMeshDrawCommandOverrideArgs(const FInstanceCullingDrawParams& InstanceCullingDrawParams)
 {
@@ -76,19 +86,18 @@ uint32 FInstanceCullingContext::GetInstanceIdBufferStride(ERHIFeatureLevel::Type
 	}
 }
 
-FInstanceCullingContext::FInstanceCullingContext(ERHIFeatureLevel::Type InFeatureLevel, FInstanceCullingManager* InInstanceCullingManager, TArrayView<const int32> InViewIds, const TRefCountPtr<IPooledRenderTarget>& InPrevHZB, enum EInstanceCullingMode InInstanceCullingMode, bool bInDrawOnlyVSMInvalidatingGeometry, EBatchProcessingMode InSingleInstanceProcessingMode) :
+FInstanceCullingContext::FInstanceCullingContext(ERHIFeatureLevel::Type InFeatureLevel, FInstanceCullingManager* InInstanceCullingManager, TArrayView<const int32> InViewIds, const TRefCountPtr<IPooledRenderTarget>& InPrevHZB, EInstanceCullingMode InInstanceCullingMode, EInstanceCullingFlags InFlags, EBatchProcessingMode InSingleInstanceProcessingMode) :
 	InstanceCullingManager(InInstanceCullingManager),
 	FeatureLevel(InFeatureLevel),
 	ViewIds(InViewIds),
 	PrevHZB(InPrevHZB),
 	bIsEnabled(InInstanceCullingManager == nullptr || InInstanceCullingManager->IsEnabled()),
 	InstanceCullingMode(InInstanceCullingMode),
-	bDrawOnlyVSMInvalidatingGeometry(bInDrawOnlyVSMInvalidatingGeometry),
+	Flags(InFlags),
 	SingleInstanceProcessingMode(InSingleInstanceProcessingMode)
 {
 }
 
-
 bool FInstanceCullingContext::IsOcclusionCullingEnabled()
 {
 	return GOcclusionCullInstances != 0;
@@ -112,7 +121,18 @@ void FInstanceCullingContext::ResetCommands(int32 MaxNumCommands)
 	MeshDrawCommandInfos.Empty(MaxNumCommands);
 	DrawCommandDescs.Empty(MaxNumCommands);
 	InstanceIdOffsets.Empty(MaxNumCommands);
+	PayloadData.Empty(MaxNumCommands);
 	TotalInstances = 0U;
+
+	DrawCommandCompactionData.Empty(MaxNumCommands);
+	CompactionBlockDataIndices.Reset();
+	NumCompactionInstances = 0U;
+}
+
+bool FInstanceCullingContext::IsInstanceOrderPreservationEnabled() const
+{
+	// NOTE: Instance compaction is currently not enabled on mobile platforms
+	return IsInstanceOrderPreservationAllowed(FeatureLevel) && !EnumHasAnyFlags(Flags, EInstanceCullingFlags::NoInstanceOrderPreservation);
 }
 
 uint32 FInstanceCullingContext::AllocateIndirectArgs(const FMeshDrawCommand *MeshDrawCommand)
@@ -150,30 +170,137 @@ uint32 FInstanceCullingContext::AllocateIndirectArgs(const FMeshDrawCommand *Mes
 // 2.1 Only allocate indirect draw cmd if needed, 
 // 3. 
 
-void FInstanceCullingContext::AddInstancesToDrawCommand(uint32 IndirectArgsOffset, int32 InstanceDataOffset, bool bDynamicInstanceDataOffset, bool bForceInstanceCulling, uint32 NumInstances)
+void FInstanceCullingContext::AddInstancesToDrawCommand(uint32 IndirectArgsOffset, int32 InstanceDataOffset, uint32 RunOffset, uint32 NumInstances, EInstanceFlags InstanceFlags)
 {
 	checkSlow(InstanceDataOffset >= 0);
 
-	// We special-case the single-instance (i.e., regular primitives) as they don't need culling (again).
+	const bool bDynamicInstanceDataOffset = EnumHasAnyFlags(InstanceFlags, EInstanceFlags::DynamicInstanceDataOffset);
+	const bool bPreserveInstanceOrder = EnumHasAnyFlags(InstanceFlags, EInstanceFlags::PreserveInstanceOrder);
+	const bool bForceInstanceCulling = EnumHasAnyFlags(InstanceFlags, EInstanceFlags::ForceInstanceCulling);	
+
+	uint32 Payload;
+	if (bPreserveInstanceOrder)
+	{
+		checkSlow(!EnumHasAnyFlags(Flags, EInstanceCullingFlags::NoInstanceOrderPreservation)); // this should have already been handled
+
+		// We need to provide full payload data for these instances
+		// NOTE: The extended payload data flag is in the lowest bit instead of the highest because the payload is not a full dword		
+		Payload = 1 | (uint32(PayloadData.Num()) << 1U);
+		PayloadData.Emplace(bDynamicInstanceDataOffset, IndirectArgsOffset, InstanceDataOffset, RunOffset, DrawCommandCompactionData.Num());
+	}
+	else
+	{
+		// Conserve space by packing the relevant payload information into the dword
+		Payload = (IndirectArgsOffset << 2U) | (bDynamicInstanceDataOffset ? 2U : 0U);
+	}
+
+	// We special-case the single-instance (i.e., regular primitives) as they don't need culling (again), except where explicitly specified.
 	// In actual fact this is not 100% true because dynamic path primitives may not have been culled.
 	EBatchProcessingMode Mode = (NumInstances == 1 && !bForceInstanceCulling) ? SingleInstanceProcessingMode : EBatchProcessingMode::Generic;
-	// NOTE: we pack the bDynamicInstanceDataOffset in the lowest bit because the load balancer steals the upper bits of the payload!
-	LoadBalancers[uint32(Mode)]->Add(uint32(InstanceDataOffset), NumInstances, (IndirectArgsOffset << 1U) | uint32(bDynamicInstanceDataOffset));
+	LoadBalancers[uint32(Mode)]->Add(uint32(InstanceDataOffset), NumInstances, Payload);
 	TotalInstances += NumInstances;
 }
 
-void FInstanceCullingContext::AddInstanceRunsToDrawCommand(uint32 IndirectArgsOffset, int32 InstanceDataOffset, bool bDynamicInstanceDataOffset, bool bForceInstanceCulling, const uint32* Runs, uint32 NumRuns)
+void FInstanceCullingContext::AddInstanceRunsToDrawCommand(uint32 IndirectArgsOffset, int32 InstanceDataOffset, const uint32* Runs, uint32 NumRuns, EInstanceFlags InstanceFlags)
 {
 	// Add items to current generic batch as they are instanced for sure.
+	uint32 NumInstancesInRuns = 0;
 	for (uint32 Index = 0; Index < NumRuns; ++Index)
 	{
 		uint32 RunStart = Runs[Index * 2];
 		uint32 RunEndIncl = Runs[Index * 2 + 1];
 		uint32 NumInstances = (RunEndIncl + 1U) - RunStart;
-		AddInstancesToDrawCommand(IndirectArgsOffset, InstanceDataOffset + RunStart, bDynamicInstanceDataOffset, bForceInstanceCulling, NumInstances);
+		AddInstancesToDrawCommand(IndirectArgsOffset, InstanceDataOffset + RunStart, NumInstancesInRuns, NumInstances, InstanceFlags);
+
+		NumInstancesInRuns += NumInstances;
 	}
 }
 
+
+// Base class that provides common functionality between all compaction phases
+class FCompactVisibleInstancesBaseCs : public FGlobalShader
+{
+public:
+	/** A compaction block is a group of instance IDs sized (N * NumViews). This is N. */
+	static constexpr int32 CompactionBlockNumInstances = 64;
+
+	FCompactVisibleInstancesBaseCs() = default;
+	FCompactVisibleInstancesBaseCs(const ShaderMetaType::CompiledShaderInitializerType& Initializer)
+		: FGlobalShader(Initializer)
+	{
+	}
+
+	static bool ShouldCompilePermutation(const FGlobalShaderPermutationParameters& Parameters)
+	{
+		// Currently compaction isn't supported on mobile
+		return UseGPUScene(Parameters.Platform) && GetMaxSupportedFeatureLevel(Parameters.Platform) > ERHIFeatureLevel::ES3_1;
+	}
+
+	static void ModifyCompilationEnvironment(const FGlobalShaderPermutationParameters& Parameters, FShaderCompilerEnvironment& OutEnvironment)
+	{
+		FGlobalShader::ModifyCompilationEnvironment(Parameters, OutEnvironment);
+
+		OutEnvironment.SetDefine(TEXT("COMPACTION_BLOCK_NUM_INSTANCES"), CompactionBlockNumInstances);
+		OutEnvironment.SetDefine(TEXT("INDIRECT_ARGS_NUM_WORDS"), FInstanceCullingContext::IndirectArgsNumWords);
+	}
+};
+
+// Compaction shader for phase one - calculate instance offsets for each instance compaction "block"
+class FCalculateCompactBlockInstanceOffsetsCs final : public FCompactVisibleInstancesBaseCs
+{
+	DECLARE_GLOBAL_SHADER(FCalculateCompactBlockInstanceOffsetsCs);
+	SHADER_USE_PARAMETER_STRUCT(FCalculateCompactBlockInstanceOffsetsCs, FCompactVisibleInstancesBaseCs)
+
+public:
+	static constexpr int32 NumThreadsPerGroup = 512;
+
+	static void ModifyCompilationEnvironment(const FGlobalShaderPermutationParameters& Parameters, FShaderCompilerEnvironment& OutEnvironment)
+	{
+		FCompactVisibleInstancesBaseCs::ModifyCompilationEnvironment(Parameters, OutEnvironment);
+
+		OutEnvironment.SetDefine(TEXT("CALCULATE_COMPACT_BLOCK_INSTANCE_OFFSETS"), 1);
+		OutEnvironment.SetDefine(TEXT("NUM_THREADS_PER_GROUP"), NumThreadsPerGroup);
+	}
+
+	BEGIN_SHADER_PARAMETER_STRUCT(FParameters, )
+		SHADER_PARAMETER_RDG_BUFFER_SRV(StructuredBuffer<FInstanceCullingContext::FCompactionData>, DrawCommandCompactionData)
+		SHADER_PARAMETER_RDG_BUFFER_SRV(StructuredBuffer<uint32>, BlockInstanceCounts)
+
+		SHADER_PARAMETER_RDG_BUFFER_UAV(RWStructuredBuffer<uint32>, BlockDestInstanceOffsetsOut)
+		SHADER_PARAMETER_RDG_BUFFER_UAV(RWBuffer<uint32>, DrawIndirectArgsBufferOut)
+	END_SHADER_PARAMETER_STRUCT()
+};
+IMPLEMENT_GLOBAL_SHADER(FCalculateCompactBlockInstanceOffsetsCs, "/Engine/Private/InstanceCulling/CompactVisibleInstances.usf", "CalculateCompactBlockInstanceOffsetsCS", SF_Compute);
+
+// Compaction shader for phase two - output visible instances, compacted and in original draw order
+class FCompactVisibleInstancesCs final : public FCompactVisibleInstancesBaseCs
+{
+	DECLARE_GLOBAL_SHADER(FCompactVisibleInstancesCs);
+	SHADER_USE_PARAMETER_STRUCT(FCompactVisibleInstancesCs, FCompactVisibleInstancesBaseCs)
+
+public:
+	static constexpr int32 NumThreadsPerGroup = FCompactVisibleInstancesBaseCs::CompactionBlockNumInstances;
+
+	static void ModifyCompilationEnvironment(const FGlobalShaderPermutationParameters& Parameters, FShaderCompilerEnvironment& OutEnvironment)
+	{
+		FCompactVisibleInstancesBaseCs::ModifyCompilationEnvironment(Parameters, OutEnvironment);
+
+		OutEnvironment.SetDefine(TEXT("COMPACT_VISIBLE_INSTANCES"), 1);
+		OutEnvironment.SetDefine(TEXT("NUM_THREADS_PER_GROUP"), NumThreadsPerGroup);
+	}
+
+	BEGIN_SHADER_PARAMETER_STRUCT(FParameters, )
+		SHADER_PARAMETER_RDG_BUFFER_SRV(StructuredBuffer<FInstanceCullingContext::FCompactionData>, DrawCommandCompactionData)
+		SHADER_PARAMETER_RDG_BUFFER_SRV(StructuredBuffer<uint32>, BlockDrawCommandIndices)
+		SHADER_PARAMETER_RDG_BUFFER_SRV(StructuredBuffer<uint32>, InstanceIdsBufferIn)
+		SHADER_PARAMETER_RDG_BUFFER_SRV(StructuredBuffer<uint32>, BlockDestInstanceOffsets)
+
+		SHADER_PARAMETER_RDG_BUFFER_UAV(RWStructuredBuffer<uint32>, InstanceIdsBufferOut)
+		SHADER_PARAMETER_RDG_BUFFER_UAV(RWStructuredBuffer<FVector4f>, InstanceIdsBufferOutMobile)		
+	END_SHADER_PARAMETER_STRUCT()
+};
+IMPLEMENT_GLOBAL_SHADER(FCompactVisibleInstancesCs, "/Engine/Private/InstanceCulling/CompactVisibleInstances.usf", "CompactVisibleInstances", SF_Compute);
+
 class FBuildInstanceIdBufferAndCommandsFromPrimitiveIdsCs : public FGlobalShader
 {
 	DECLARE_GLOBAL_SHADER(FBuildInstanceIdBufferAndCommandsFromPrimitiveIdsCs);
@@ -193,12 +320,26 @@ public:
 	// TODO: maybe disable permutation in shipping builds?
 	class FDebugModeDim : SHADER_PERMUTATION_BOOL("DEBUG_MODE");
 	class FBatchedDim : SHADER_PERMUTATION_BOOL("ENABLE_BATCH_MODE");
+	class FInstanceCompactionDim : SHADER_PERMUTATION_BOOL("ENABLE_INSTANCE_COMPACTION");
 
-	using FPermutationDomain = TShaderPermutationDomain<FOutputCommandIdDim, FSingleInstanceModeDim, FCullInstancesDim, FOcclusionCullInstancesDim, FStereoModeDim, FDebugModeDim, FBatchedDim>;
+	using FPermutationDomain = TShaderPermutationDomain<FOutputCommandIdDim, FSingleInstanceModeDim, FCullInstancesDim, FOcclusionCullInstancesDim, FStereoModeDim, FDebugModeDim, FBatchedDim, FInstanceCompactionDim>;
 
 	static bool ShouldCompilePermutation(const FGlobalShaderPermutationParameters& Parameters)
 	{
-		return UseGPUScene(Parameters.Platform);
+		if (!UseGPUScene(Parameters.Platform))
+		{
+			return false;
+		}
+		
+		const FPermutationDomain PermutationVector(Parameters.PermutationId);
+		
+		// Currently, instance compaction is not supported on mobile platforms
+		if (PermutationVector.Get<FInstanceCompactionDim>() && IsMobilePlatform(Parameters.Platform))
+		{
+			return false;				
+		}
+
+		return true;		
 	}
 
 	static void ModifyCompilationEnvironment(const FGlobalShaderPermutationParameters& Parameters, FShaderCompilerEnvironment& OutEnvironment)
@@ -212,6 +353,7 @@ public:
 		OutEnvironment.SetDefine(TEXT("USE_GLOBAL_GPU_LIGHTMAP_DATA"), 1);
 		OutEnvironment.SetDefine(TEXT("NANITE_MULTI_VIEW"), 1);
 		OutEnvironment.SetDefine(TEXT("PRIM_ID_DYNAMIC_FLAG"), GPrimIDDynamicFlag);
+		OutEnvironment.SetDefine(TEXT("COMPACTION_BLOCK_NUM_INSTANCES"), FCompactVisibleInstancesBaseCs::CompactionBlockNumInstances);
 
 		OutEnvironment.SetDefine(TEXT("BATCH_PROCESSING_MODE_GENERIC"), uint32(EBatchProcessingMode::Generic));
 		OutEnvironment.SetDefine(TEXT("BATCH_PROCESSING_MODE_UNCULLED"), uint32(EBatchProcessingMode::UnCulled));
@@ -232,6 +374,7 @@ public:
 		SHADER_PARAMETER_STRUCT_INCLUDE(FInstanceProcessingGPULoadBalancer::FShaderParameters, LoadBalancerParameters)
 
 		SHADER_PARAMETER_RDG_BUFFER_SRV(StructuredBuffer< FInstanceCullingContext::FDrawCommandDesc >, DrawCommandDescs)
+		SHADER_PARAMETER_RDG_BUFFER_SRV(StructuredBuffer< FInstanceCullingContext::FPayloadData >, InstanceCullingPayloads)
 		SHADER_PARAMETER_RDG_BUFFER_SRV(StructuredBuffer< uint32 >, ViewIds)
 		SHADER_PARAMETER_RDG_BUFFER_SRV(StructuredBuffer< Nanite::FPackedView >, InViews)
 
@@ -242,17 +385,20 @@ public:
 
 		SHADER_PARAMETER_RDG_BUFFER_UAV(RWStructuredBuffer<uint>, InstanceIdsBufferOut)
 		SHADER_PARAMETER_RDG_BUFFER_UAV(RWStructuredBuffer<float4>, InstanceIdsBufferOutMobile)
-		SHADER_PARAMETER_RDG_BUFFER_UAV(RWStructuredBuffer<uint>, DrawCommandIdsBufferOut)
-
 		SHADER_PARAMETER_RDG_BUFFER_UAV(RWBuffer<uint>, DrawIndirectArgsBufferOut)
 
+		SHADER_PARAMETER_RDG_BUFFER_SRV(StructuredBuffer<FInstanceCullingContext::FCompactionData>, DrawCommandCompactionData)
+		SHADER_PARAMETER_RDG_BUFFER_UAV(RWStructuredBuffer<uint>, CompactInstanceIdsBufferOut)
+		SHADER_PARAMETER_RDG_BUFFER_UAV(RWStructuredBuffer<uint>, CompactionBlockCounts)
+
+
 		SHADER_PARAMETER_RDG_TEXTURE(Texture2D, HZBTexture)
 		SHADER_PARAMETER_SAMPLER(SamplerState, HZBSampler)
 		SHADER_PARAMETER(FVector2f, HZBSize)
 
 		SHADER_PARAMETER(uint32, NumViewIds)
 		SHADER_PARAMETER(uint32, NumCullingViews)
-		SHADER_PARAMETER(uint32, CurrentBatchProcessingMode)		
+		SHADER_PARAMETER(uint32, CurrentBatchProcessingMode)
 		SHADER_PARAMETER(int32, bDrawOnlyVSMInvalidatingGeometry)
 
 		SHADER_PARAMETER(int32, DynamicInstanceIdOffset)
@@ -275,7 +421,7 @@ const TRDGUniformBufferRef<FInstanceCullingGlobalUniforms> FInstanceCullingConte
 class FInstanceCullingDeferredContext : public FInstanceCullingMergedContext
 {
 public:
-	FInstanceCullingDeferredContext(ERHIFeatureLevel::Type InFeatureLevel, FInstanceCullingManager* InInstanceCullingManager = nullptr) 
+	FInstanceCullingDeferredContext(ERHIFeatureLevel::Type InFeatureLevel, FInstanceCullingManager* InInstanceCullingManager = nullptr)
 		: FInstanceCullingMergedContext(InFeatureLevel)
 		, InstanceCullingManager(InInstanceCullingManager)
 	{}
@@ -298,7 +444,7 @@ static FRDGBufferDesc CreateInstanceIdsBufferDesc(ERHIFeatureLevel::Type Feature
 	{
 		// Mobile writes to this buffer from compute and then uses as a vertex input
 		// D3D does not allow a storage buffer with vertex buffer usage
-		BufferDesc.Usage|= (IsD3DPlatform(GMaxRHIShaderPlatform) ? BUF_None : BUF_VertexBuffer);
+		BufferDesc.Usage |= (IsD3DPlatform(GMaxRHIShaderPlatform) ? BUF_None : BUF_VertexBuffer);
 	}
 	return BufferDesc;
 }
@@ -346,13 +492,49 @@ void FInstanceCullingContext::BuildRenderingCommands(
 
 	RDG_EVENT_SCOPE(GraphBuilder, "BuildRenderingCommands(Culling=%s)", bCullInstances ? TEXT("On") : TEXT("Off"));
 
+	const bool bOrderPreservationEnabled = IsInstanceOrderPreservationEnabled();
+	const uint32 NumCompactionBlocks = uint32(CompactionBlockDataIndices.Num());
+	FRDGBufferRef CompactInstanceIdsBuffer = nullptr;
+	FRDGBufferUAVRef CompactInstanceIdsUAV = nullptr;
+	FRDGBufferRef CompactionBlockCountsBuffer = nullptr;
+	FRDGBufferUAVRef CompactionBlockCountsUAV = nullptr;
+	FRDGBufferSRVRef DrawCommandCompactionDataSRV = nullptr;	
+	
+	if (bOrderPreservationEnabled)
+	{
+		// Create buffers for compacting instances for draw commands that need it
+		CompactInstanceIdsBuffer = GraphBuilder.CreateBuffer(FRDGBufferDesc::CreateStructuredDesc(sizeof(uint32), FMath::Max(NumCompactionInstances, 1u)), TEXT("InstanceCulling.Compaction.TempInstanceIdsBuffer"));
+		CompactInstanceIdsUAV = GraphBuilder.CreateUAV(CompactInstanceIdsBuffer);
+		CompactionBlockCountsBuffer = GraphBuilder.CreateBuffer(FRDGBufferDesc::CreateStructuredDesc(sizeof(uint32), FMath::Max(NumCompactionBlocks, 1u)), TEXT("InstanceCulling.Compaction.BlockInstanceCounts"));
+		CompactionBlockCountsUAV = GraphBuilder.CreateUAV(CompactionBlockCountsBuffer);
+
+		FRDGBufferRef DrawCommandCompactionDataBuffer = nullptr;
+		if (DrawCommandCompactionData.Num() > 0)
+		{
+			DrawCommandCompactionDataBuffer = CreateStructuredBuffer(GraphBuilder, TEXT("InstanceCulling.DrawCommandCompactionData"), DrawCommandCompactionData);
+		}
+		else
+		{
+			DrawCommandCompactionDataBuffer = GSystemTextures.GetDefaultStructuredBuffer(GraphBuilder, sizeof(FCompactionData));
+		}
+		DrawCommandCompactionDataSRV = GraphBuilder.CreateSRV(DrawCommandCompactionDataBuffer);
+
+		if (NumCompactionBlocks > 0)
+		{
+			ensure(NumCompactionInstances > 0);
+
+			// We must clear the block counts buffer, as it will be written to using atomic increments
+			AddClearUAVPass(GraphBuilder, CompactionBlockCountsUAV, 0);
+		}
+	}
+
 	FGlobalShaderMap* ShaderMap = GetGlobalShaderMap(FeatureLevel);
 
 	// Add any other conditions that needs debug code running here.
-	const bool bUseDebugMode = bDrawOnlyVSMInvalidatingGeometry;
+	const bool bUseDebugMode = EnumHasAnyFlags(Flags, EInstanceCullingFlags::DrawOnlyVSMInvalidatingGeometry);
 
 	FRDGBufferRef ViewIdsBuffer = CreateStructuredBuffer(GraphBuilder, TEXT("InstanceCulling.ViewIds"), ViewIds);
-	
+
 	FRDGBufferRef InstanceIdsBuffer = GraphBuilder.CreateBuffer(CreateInstanceIdsBufferDesc(FeatureLevel, InstanceIdBufferSize), TEXT("InstanceCulling.InstanceIdsBuffer"));
 	FRDGBufferUAVRef InstanceIdsBufferUAV = GraphBuilder.CreateUAV(InstanceIdsBuffer, ERDGUnorderedAccessViewFlags::SkipBarrier);
 
@@ -360,6 +542,8 @@ void FInstanceCullingContext::BuildRenderingCommands(
 
 	PassParametersTmp.DrawCommandDescs = GraphBuilder.CreateSRV(CreateStructuredBuffer(GraphBuilder, TEXT("InstanceCulling.DrawCommandDescs"), DrawCommandDescs));
 
+	PassParametersTmp.InstanceCullingPayloads = GraphBuilder.CreateSRV(CreateStructuredBuffer(GraphBuilder, TEXT("InstanceCulling.PayloadData"), PayloadData));
+
 	// Because the view uniforms are not set up by the time this runs
 	// PassParametersTmp.View = View.ViewUniformBuffer;
 	// Set up global GPU-scene data instead...
@@ -375,6 +559,10 @@ void FInstanceCullingContext::BuildRenderingCommands(
 	PassParametersTmp.DynamicInstanceIdOffset = DynamicInstanceIdOffset;
 	PassParametersTmp.DynamicInstanceIdMax = DynamicInstanceIdOffset + DynamicInstanceIdNum;
 
+	// Compaction parameters
+	PassParametersTmp.DrawCommandCompactionData = DrawCommandCompactionDataSRV;
+	PassParametersTmp.CompactInstanceIdsBufferOut = CompactInstanceIdsUAV;
+	PassParametersTmp.CompactionBlockCounts = CompactionBlockCountsUAV;
 
 	// Create buffer for indirect args and upload draw arg data, also clears the instance to zero
 	FRDGBufferDesc IndirectArgsDesc = FRDGBufferDesc::CreateIndirectDesc(IndirectArgsNumWords * IndirectArgs.Num());
@@ -404,7 +592,7 @@ void FInstanceCullingContext::BuildRenderingCommands(
 		PassParametersTmp.NumCullingViews = InstanceCullingManager->CullingIntermediate.NumViews;
 	}
 	PassParametersTmp.NumViewIds = ViewIds.Num();
-	PassParametersTmp.bDrawOnlyVSMInvalidatingGeometry = bDrawOnlyVSMInvalidatingGeometry;
+	PassParametersTmp.bDrawOnlyVSMInvalidatingGeometry = EnumHasAnyFlags(Flags, EInstanceCullingFlags::DrawOnlyVSMInvalidatingGeometry);
 	// only one of these will be used in the shader
 	PassParametersTmp.InstanceIdsBufferOut = InstanceIdsBufferUAV;
 	PassParametersTmp.InstanceIdsBufferOutMobile = InstanceIdsBufferUAV;
@@ -442,6 +630,7 @@ void FInstanceCullingContext::BuildRenderingCommands(
 			PermutationVector.Set<FBuildInstanceIdBufferAndCommandsFromPrimitiveIdsCs::FStereoModeDim>(InstanceCullingMode == EInstanceCullingMode::Stereo);
 			PermutationVector.Set<FBuildInstanceIdBufferAndCommandsFromPrimitiveIdsCs::FDebugModeDim>(bUseDebugMode);
 			PermutationVector.Set<FBuildInstanceIdBufferAndCommandsFromPrimitiveIdsCs::FBatchedDim>(false);
+			PermutationVector.Set<FBuildInstanceIdBufferAndCommandsFromPrimitiveIdsCs::FInstanceCompactionDim>(bOrderPreservationEnabled);
 
 			auto ComputeShader = ShaderMap->GetShader<FBuildInstanceIdBufferAndCommandsFromPrimitiveIdsCs>(PermutationVector);
 
@@ -454,8 +643,56 @@ void FInstanceCullingContext::BuildRenderingCommands(
 			);
 		}
 	}
+
+	if (bOrderPreservationEnabled && NumCompactionBlocks > 0)
+	{
+		FRDGBufferRef BlockDestInstanceOffsets = GraphBuilder.CreateBuffer(FRDGBufferDesc::CreateStructuredDesc(sizeof(uint32), NumCompactionBlocks), TEXT("InstanceCulling.Compaction.BlockDestInstanceOffsets"));
+
+		// Compaction phase one - prefix sum of the compaction "blocks"
+		{
+			auto PassParameters = GraphBuilder.AllocParameters<FCalculateCompactBlockInstanceOffsetsCs::FParameters>();
+			PassParameters->DrawCommandCompactionData = DrawCommandCompactionDataSRV;
+			PassParameters->BlockInstanceCounts = GraphBuilder.CreateSRV(CompactionBlockCountsBuffer);
+			PassParameters->BlockDestInstanceOffsetsOut = GraphBuilder.CreateUAV(BlockDestInstanceOffsets);
+			PassParameters->DrawIndirectArgsBufferOut = PassParametersTmp.DrawIndirectArgsBufferOut;
+
+			auto ComputeShader = ShaderMap->GetShader<FCalculateCompactBlockInstanceOffsetsCs>();
+
+			FComputeShaderUtils::AddPass(
+				GraphBuilder,
+				RDG_EVENT_NAME("Instance Compaction Phase 1"),
+				ComputeShader,
+				PassParameters,
+				FComputeShaderUtils::GetGroupCountWrapped(DrawCommandCompactionData.Num())
+			);
+		}
+
+		// Compaction phase two - write instances to compact final location
+		{
+			FRDGBufferRef BlockDrawCommandIndices = CreateStructuredBuffer(GraphBuilder, TEXT("InstanceCulling.Compaction.BlockDrawCommandIndices"), CompactionBlockDataIndices);
+
+			auto PassParameters = GraphBuilder.AllocParameters<FCompactVisibleInstancesCs::FParameters>();
+			PassParameters->DrawCommandCompactionData = DrawCommandCompactionDataSRV;
+			PassParameters->BlockDrawCommandIndices = GraphBuilder.CreateSRV(BlockDrawCommandIndices);
+			PassParameters->InstanceIdsBufferIn = GraphBuilder.CreateSRV(CompactInstanceIdsBuffer);
+			PassParameters->BlockDestInstanceOffsets = GraphBuilder.CreateSRV(BlockDestInstanceOffsets);			
+			PassParameters->InstanceIdsBufferOut = InstanceIdsBufferUAV;
+			PassParameters->InstanceIdsBufferOutMobile = InstanceIdsBufferUAV;			
+
+			auto ComputeShader = ShaderMap->GetShader<FCompactVisibleInstancesCs>();
+
+			FComputeShaderUtils::AddPass(
+				GraphBuilder,
+				RDG_EVENT_NAME("Instance Compaction Phase 2"),
+				ComputeShader,
+				PassParameters,
+				FComputeShaderUtils::GetGroupCountWrapped(NumCompactionBlocks)
+			);
+		}
+	}
+
 	Results.DrawIndirectArgsBuffer = DrawIndirectArgsRDG;
-	
+
 	if (FeatureLevel == ERHIFeatureLevel::ES3_1)
 	{
 		Results.InstanceDataBuffer = InstanceIdsBuffer;
@@ -561,9 +798,43 @@ FInstanceCullingDeferredContext *FInstanceCullingContext::CreateDeferredContext(
 		PassParameters[Mode] = GraphBuilder.AllocParameters<FBuildInstanceIdBufferAndCommandsFromPrimitiveIdsCs::FParameters>();
 	}
 
+	// Create buffers for compacting instances for draw commands that need it
+	const bool bEnableInstanceCompaction = IsInstanceOrderPreservationAllowed(FeatureLevel);
+	FRDGBufferSRVRef DrawCommandCompactionDataSRV = nullptr;
+	FRDGBufferRef CompactInstanceIdsBuffer = nullptr;
+	FRDGBufferUAVRef CompactInstanceIdsUAV = nullptr;
+	FRDGBufferRef CompactionBlockCountsBuffer = nullptr;
+	FRDGBufferUAVRef CompactionBlockCountsUAV = nullptr;
+
+	if (bEnableInstanceCompaction)
+	{
+		DrawCommandCompactionDataSRV = GraphBuilder.CreateSRV(CreateStructuredBuffer(INST_CULL_CREATE_STRUCT_BUFF_ARGS(DrawCommandCompactionData)));
+		CompactInstanceIdsBuffer = CreateStructuredBuffer(
+			GraphBuilder,
+			TEXT("InstanceCulling.Compaction.TempInstanceIdsBuffer"),
+			sizeof(uint32),
+			INST_CULL_CALLBACK(FMath::Max(DeferredContext->TotalCompactionInstances, 1)),
+			INST_CULL_CALLBACK(nullptr),
+			INST_CULL_CALLBACK(0));
+		CompactInstanceIdsUAV = GraphBuilder.CreateUAV(CompactInstanceIdsBuffer);
+		CompactionBlockCountsBuffer = CreateStructuredBuffer(
+			GraphBuilder,
+			TEXT("InstanceCulling.Compaction.BlockInstanceCounts"),
+			sizeof(uint32),
+			INST_CULL_CALLBACK(FMath::Max(DeferredContext->TotalCompactionBlocks, 1)),
+			INST_CULL_CALLBACK(nullptr),
+			INST_CULL_CALLBACK(0));
+		CompactionBlockCountsUAV = GraphBuilder.CreateUAV(CompactionBlockCountsBuffer);
+
+		// We must clear the block counts buffer, as they will be written to using atomic increments
+		// TODO: Come up with a clever way to cull this pass when no compaction is needed (currently can't know until the batch is complete on the RDG execution timeline).
+		AddClearUAVPass(GraphBuilder, CompactionBlockCountsUAV, 0);
+	}
+
 	FBuildInstanceIdBufferAndCommandsFromPrimitiveIdsCs::FParameters PassParametersTmp;
 
 	FRDGBufferRef DrawCommandDescsRDG = CreateStructuredBuffer(INST_CULL_CREATE_STRUCT_BUFF_ARGS(DrawCommandDescs));
+	FRDGBufferRef InstanceCullingPayloadsRDG = CreateStructuredBuffer(INST_CULL_CREATE_STRUCT_BUFF_ARGS(PayloadData));
 	FRDGBufferRef ViewIdsRDG = CreateStructuredBuffer(INST_CULL_CREATE_STRUCT_BUFF_ARGS(ViewIds));
 	FRDGBufferRef BatchInfosRDG = CreateStructuredBuffer(INST_CULL_CREATE_STRUCT_BUFF_ARGS(BatchInfos));
 
@@ -604,6 +875,7 @@ FInstanceCullingDeferredContext *FInstanceCullingContext::CreateDeferredContext(
 	PassParametersTmp.GPUSceneNumLightmapDataItems = GPUScene.GetNumLightmapDataItems();
 
 	PassParametersTmp.DrawCommandDescs = GraphBuilder.CreateSRV(DrawCommandDescsRDG);
+	PassParametersTmp.InstanceCullingPayloads = GraphBuilder.CreateSRV(InstanceCullingPayloadsRDG);
 	PassParametersTmp.BatchInfos = GraphBuilder.CreateSRV(BatchInfosRDG);
 	PassParametersTmp.ViewIds = GraphBuilder.CreateSRV(ViewIdsRDG);
 	// only one of these will be used in the shader
@@ -617,6 +889,12 @@ FInstanceCullingDeferredContext *FInstanceCullingContext::CreateDeferredContext(
 		PassParametersTmp.InViews = GraphBuilder.CreateSRV(InstanceCullingManager->CullingIntermediate.CullingViews);
 		PassParametersTmp.NumCullingViews = InstanceCullingManager->CullingIntermediate.NumViews;
 	}
+
+	// Compaction parameters
+	PassParametersTmp.DrawCommandCompactionData = DrawCommandCompactionDataSRV;
+	PassParametersTmp.CompactInstanceIdsBufferOut = CompactInstanceIdsUAV;
+	PassParametersTmp.CompactionBlockCounts = CompactionBlockCountsUAV;
+
 	// Record the number of culling views to be able to check that no views referencing out-of bounds views are queued up
 	DeferredContext->NumCullingViews = InstanceCullingManager->CullingIntermediate.NumViews;
 
@@ -663,6 +941,7 @@ FInstanceCullingDeferredContext *FInstanceCullingContext::CreateDeferredContext(
 		PermutationVector.Set<FBuildInstanceIdBufferAndCommandsFromPrimitiveIdsCs::FSingleInstanceModeDim>(EBatchProcessingMode(Mode) == EBatchProcessingMode::UnCulled);
 		PermutationVector.Set<FBuildInstanceIdBufferAndCommandsFromPrimitiveIdsCs::FCullInstancesDim>(bCullInstances && EBatchProcessingMode(Mode) != EBatchProcessingMode::UnCulled);
 		PermutationVector.Set<FBuildInstanceIdBufferAndCommandsFromPrimitiveIdsCs::FOcclusionCullInstancesDim>(bOcclusionCullInstances);
+		PermutationVector.Set<FBuildInstanceIdBufferAndCommandsFromPrimitiveIdsCs::FInstanceCompactionDim>(bEnableInstanceCompaction);
 
 		auto ComputeShader = ShaderMap->GetShader<FBuildInstanceIdBufferAndCommandsFromPrimitiveIdsCs>(PermutationVector);
 
@@ -673,7 +952,66 @@ FInstanceCullingDeferredContext *FInstanceCullingContext::CreateDeferredContext(
 			PassParameters[Mode],
 			INST_CULL_CALLBACK_MODE(DeferredContext->LoadBalancers[Mode].GetWrappedCsGroupCount()));
 	}
-	
+
+	// TODO: Come up with a way to cull these passes when no compaction is needed. The group count resulting in (0, 0, 0) causes the pass lambdas to not execute,
+	// but currently cannot cull resource transitions
+	if (bEnableInstanceCompaction)
+	{
+		FRDGBufferRef BlockDestInstanceOffsets = CreateStructuredBuffer(
+			GraphBuilder,
+			TEXT("InstanceCulling.Compaction.BlockDestInstanceOffsets"),
+			sizeof(uint32),
+			INST_CULL_CALLBACK(FMath::Max<uint32>(DeferredContext->TotalCompactionBlocks, 1U)),
+			INST_CULL_CALLBACK(nullptr),
+			INST_CULL_CALLBACK(0));
+
+		// Compaction phase one - prefix sum of the compaction "blocks"
+		{
+			auto PassParameters2 = GraphBuilder.AllocParameters<FCalculateCompactBlockInstanceOffsetsCs::FParameters>();
+			PassParameters2->DrawCommandCompactionData = DrawCommandCompactionDataSRV;
+			PassParameters2->BlockInstanceCounts = GraphBuilder.CreateSRV(CompactionBlockCountsBuffer);
+			PassParameters2->BlockDestInstanceOffsetsOut = GraphBuilder.CreateUAV(BlockDestInstanceOffsets);
+			PassParameters2->DrawIndirectArgsBufferOut = PassParametersTmp.DrawIndirectArgsBufferOut;
+
+			auto ComputeShader = ShaderMap->GetShader<FCalculateCompactBlockInstanceOffsetsCs>();
+
+			FComputeShaderUtils::AddPass(
+				GraphBuilder,
+				RDG_EVENT_NAME("Instance Compaction Phase 1"),
+				ComputeShader,
+				PassParameters2,
+				[DeferredContext]()
+			{
+				return FComputeShaderUtils::GetGroupCountWrapped(DeferredContext->TotalCompactionDrawCommands);
+			});
+		}
+
+		// Compaction phase two - write instances to compact final location
+		{
+			FRDGBufferRef BlockDrawCommandIndices = CreateStructuredBuffer(INST_CULL_CREATE_STRUCT_BUFF_ARGS(CompactionBlockDataIndices));
+
+			auto PassParameters2 = GraphBuilder.AllocParameters<FCompactVisibleInstancesCs::FParameters>();
+			PassParameters2->DrawCommandCompactionData = DrawCommandCompactionDataSRV;
+			PassParameters2->BlockDrawCommandIndices = GraphBuilder.CreateSRV(BlockDrawCommandIndices);
+			PassParameters2->InstanceIdsBufferIn = GraphBuilder.CreateSRV(CompactInstanceIdsBuffer);
+			PassParameters2->BlockDestInstanceOffsets = GraphBuilder.CreateSRV(BlockDestInstanceOffsets);
+			PassParameters2->InstanceIdsBufferOut = InstanceIdsBufferUAV;
+			PassParameters2->InstanceIdsBufferOutMobile = InstanceIdsBufferUAV;			
+
+			auto ComputeShader = ShaderMap->GetShader<FCompactVisibleInstancesCs>();
+
+			FComputeShaderUtils::AddPass(
+				GraphBuilder,
+				RDG_EVENT_NAME("Instance Compaction Phase 2"),
+				ComputeShader,
+				PassParameters2,
+				[PassParameters2, DeferredContext]()
+			{				
+				return FComputeShaderUtils::GetGroupCountWrapped(DeferredContext->TotalCompactionBlocks);
+			});
+		}
+	}
+
 	if (FeatureLevel > ERHIFeatureLevel::ES3_1)
 	{
 		FInstanceCullingGlobalUniforms* UniformParameters = GraphBuilder.AllocParameters<FInstanceCullingGlobalUniforms>();
@@ -714,6 +1052,7 @@ public:
 
 		OutEnvironment.SetDefine(TEXT("INDIRECT_ARGS_NUM_WORDS"), FInstanceCullingContext::IndirectArgsNumWords);
 		OutEnvironment.SetDefine(TEXT("NUM_THREADS_PER_GROUP"), NumThreadsPerGroup);
+		OutEnvironment.SetDefine(TEXT("COMPACTION_BLOCK_NUM_INSTANCES"), FCompactVisibleInstancesBaseCs::CompactionBlockNumInstances);
 	}
 
 	BEGIN_SHADER_PARAMETER_STRUCT(FParameters, )
@@ -826,6 +1165,7 @@ void FInstanceCullingContext::SetupDrawCommands(
 	const int32 NumViews = ViewIds.Num();
 	const bool bAlwaysUseIndirectDraws = (SingleInstanceProcessingMode != EBatchProcessingMode::UnCulled);
 	const uint32 InstanceIdBufferStride = GetInstanceIdBufferStride(FeatureLevel);
+	const bool bOrderPreservationEnabled = IsInstanceOrderPreservationEnabled();
 
 	// Allocate conservatively for all commands, may not use all.
 	for (int32 DrawCommandIndex = 0; DrawCommandIndex < NumDrawCommandsIn; ++DrawCommandIndex)
@@ -836,6 +1176,7 @@ void FInstanceCullingContext::SetupDrawCommands(
 		const bool bSupportsGPUSceneInstancing = EnumHasAnyFlags(VisibleMeshDrawCommand.Flags, EFVisibleMeshDrawCommandFlags::HasPrimitiveIdStreamIndex);
 		const bool bMaterialMayModifyPosition = EnumHasAnyFlags(VisibleMeshDrawCommand.Flags, EFVisibleMeshDrawCommandFlags::MaterialMayModifyPosition);
 		const bool bForceInstanceCulling = EnumHasAnyFlags(VisibleMeshDrawCommand.Flags, EFVisibleMeshDrawCommandFlags::ForceInstanceCulling);
+		const bool bPreserveInstanceOrder = bOrderPreservationEnabled && EnumHasAnyFlags(VisibleMeshDrawCommand.Flags, EFVisibleMeshDrawCommandFlags::PreserveInstanceOrder);
 		const bool bUseIndirectDraw = bAlwaysUseIndirectDraws || bForceInstanceCulling || (VisibleMeshDrawCommand.NumRuns > 0 || MeshDrawCommand->NumInstances > 1);
 
 		if (bCompactIdenticalCommands && CurrentStateBucketId != -1 && VisibleMeshDrawCommand.StateBucketId == CurrentStateBucketId)
@@ -903,15 +1244,54 @@ void FInstanceCullingContext::SetupDrawCommands(
 
 		if (bSupportsGPUSceneInstancing)
 		{
+			EInstanceFlags InstanceFlags = EInstanceFlags::None;
+			if (VisibleMeshDrawCommand.PrimitiveIdInfo.bIsDynamicPrimitive)
+			{
+				EnumAddFlags(InstanceFlags, EInstanceFlags::DynamicInstanceDataOffset);
+			}
+			if (bForceInstanceCulling)
+			{
+				EnumAddFlags(InstanceFlags, EInstanceFlags::ForceInstanceCulling);
+			}
+			if (bPreserveInstanceOrder)
+			{
+				EnumAddFlags(InstanceFlags, EInstanceFlags::PreserveInstanceOrder);
+			}
+
+			const uint32 InstanceOffset = TotalInstances;
+
 			// append 'culling command' targeting the current slot
 			// This will cause all instances belonging to the Primitive to be added to the command, if they are visible etc (GPU-Scene knows all - sees all)
 			if (VisibleMeshDrawCommand.RunArray)
 			{
-				AddInstanceRunsToDrawCommand(CurrentIndirectArgsOffset, VisibleMeshDrawCommand.PrimitiveIdInfo.InstanceSceneDataOffset, VisibleMeshDrawCommand.PrimitiveIdInfo.bIsDynamicPrimitive, bForceInstanceCulling, VisibleMeshDrawCommand.RunArray, VisibleMeshDrawCommand.NumRuns);
+				AddInstanceRunsToDrawCommand(CurrentIndirectArgsOffset, VisibleMeshDrawCommand.PrimitiveIdInfo.InstanceSceneDataOffset, VisibleMeshDrawCommand.RunArray, VisibleMeshDrawCommand.NumRuns, InstanceFlags);
+			}
+			else
+			{
+				AddInstancesToDrawCommand(CurrentIndirectArgsOffset, VisibleMeshDrawCommand.PrimitiveIdInfo.InstanceSceneDataOffset, 0, VisibleMeshDrawCommand.MeshDrawCommand->NumInstances, InstanceFlags);
 			}
-			else 
+
+			const uint32 NumInstancesAdded = TotalInstances - InstanceOffset;
+			if (bPreserveInstanceOrder && NumInstancesAdded > 0)
 			{
-				AddInstancesToDrawCommand(CurrentIndirectArgsOffset, VisibleMeshDrawCommand.PrimitiveIdInfo.InstanceSceneDataOffset, VisibleMeshDrawCommand.PrimitiveIdInfo.bIsDynamicPrimitive, bForceInstanceCulling, VisibleMeshDrawCommand.MeshDrawCommand->NumInstances);
+				const uint32 CompactionDataIndex = uint32(DrawCommandCompactionData.Num());
+				DrawCommandCompactionData.Emplace(
+					NumInstancesAdded,
+					NumViews,
+					uint32(CompactionBlockDataIndices.Num()),
+					CurrentIndirectArgsOffset,
+					NumCompactionInstances,
+					InstanceIdOffsets.Last());
+
+				const int32 FirstBlock = CompactionBlockDataIndices.Num();
+				const uint32 NumCompactionBlocksThisCommand = FMath::DivideAndRoundUp(NumInstancesAdded, CompactionBlockNumInstances);
+				CompactionBlockDataIndices.AddUninitialized(NumCompactionBlocksThisCommand);
+				for (int32 Block = FirstBlock; Block < CompactionBlockDataIndices.Num(); ++Block)
+				{
+					CompactionBlockDataIndices[Block] = CompactionDataIndex;
+				}
+
+				NumCompactionInstances += NumInstancesAdded * NumViews;
 			}
 		}
 	}
diff --git a/Engine/Source/Runtime/Renderer/Private/InstanceCulling/InstanceCullingMergedContext.cpp b/Engine/Source/Runtime/Renderer/Private/InstanceCulling/InstanceCullingMergedContext.cpp
--- a/Engine/Source/Runtime/Renderer/Private/InstanceCulling/InstanceCullingMergedContext.cpp
+++ b/Engine/Source/Runtime/Renderer/Private/InstanceCulling/InstanceCullingMergedContext.cpp
@@ -13,11 +13,15 @@ void FInstanceCullingMergedContext::MergeBatches()
 	IndirectArgs.Empty(TotalIndirectArgs);
 	DrawCommandDescs.Empty(TotalIndirectArgs);
 	InstanceIdOffsets.Empty(TotalIndirectArgs);
+	PayloadData.Empty(TotalPayloads);
 	ViewIds.Empty(TotalViewIds);
+	DrawCommandCompactionData.Empty(TotalCompactionDrawCommands);
+	CompactionBlockDataIndices.Empty(TotalCompactionBlocks);
 
 	BatchInfos.AddDefaulted(Batches.Num());
 	uint32 InstanceIdBufferOffset = 0U;
 	uint32 InstanceDataByteOffset = 0U;
+	uint32 TempCompactionInstanceOffset = 0U;
 	const uint32 InstanceIdBufferStride = FInstanceCullingContext::GetInstanceIdBufferStride(FeatureLevel);
 
 	// Index that maps from each command to the corresponding batch - maybe not the utmost efficiency
@@ -35,6 +39,9 @@ void FInstanceCullingMergedContext::MergeBatches()
 		check(InstanceCullingContext.DrawCommandDescs.Num() == InstanceCullingContext.IndirectArgs.Num());
 		DrawCommandDescs.Append(InstanceCullingContext.DrawCommandDescs);
 
+		BatchInfo.PayloadDataOffset = PayloadData.Num();
+		PayloadData.Append(InstanceCullingContext.PayloadData);
+
 		check(InstanceCullingContext.InstanceIdOffsets.Num() == InstanceCullingContext.IndirectArgs.Num());
 		InstanceIdOffsets.AddDefaulted(InstanceCullingContext.InstanceIdOffsets.Num());
 		// TODO: perform offset on GPU
@@ -79,8 +86,26 @@ void FInstanceCullingMergedContext::MergeBatches()
 		Result.IndirectArgsByteOffset = BatchInfo.IndirectArgsOffset * FInstanceCullingContext::IndirectArgsNumWords * sizeof(uint32);
 
 		BatchInfo.InstanceDataWriteOffset = InstanceIdBufferOffset;
+
+		// Append the compaction data, but fix up the offsets for the batch
+		BatchInfo.CompactionDataOffset = DrawCommandCompactionData.Num();
+		const uint32 CompactionBlockOffset = CompactionBlockDataIndices.Num();
+		for (auto CompactionData : InstanceCullingContext.DrawCommandCompactionData)
+		{			
+			CompactionData.BlockOffset += CompactionBlockOffset;
+			CompactionData.IndirectArgsIndex += BatchInfo.IndirectArgsOffset;
+			CompactionData.SrcInstanceIdOffset += TempCompactionInstanceOffset;
+			CompactionData.DestInstanceIdOffset += InstanceIdBufferOffset;
+			DrawCommandCompactionData.Add(CompactionData);
+		}
+		for (uint32 CompactionDataIndex : InstanceCullingContext.CompactionBlockDataIndices)
+		{
+			CompactionBlockDataIndices.Add(CompactionDataIndex + BatchInfo.CompactionDataOffset);
+		}
+		TempCompactionInstanceOffset += InstanceCullingContext.NumCompactionInstances;
+
+		// Advance offset into instance ID and per-instance buffer
 		InstanceIdBufferOffset += BatchTotalInstances;
-		// Advance offset into per-instance buffer
 		InstanceDataByteOffset += FInstanceCullingContext::StepInstanceDataOffset(FeatureLevel, BatchTotalInstances, BatchTotalDraws) * InstanceIdBufferStride;
 	}
 }
@@ -119,7 +144,11 @@ void FInstanceCullingMergedContext::AddBatch(FRDGBuilder& GraphBuilder, const FI
 #endif 
 
 	TotalIndirectArgs += Context->IndirectArgs.Num();
+	TotalPayloads += Context->PayloadData.Num();
 	TotalViewIds += Context->ViewIds.Num();
 	InstanceIdBufferSize += Context->TotalInstances * Context->ViewIds.Num();
 	TotalInstances += Context->TotalInstances;
+	TotalCompactionDrawCommands += Context->DrawCommandCompactionData.Num();
+	TotalCompactionBlocks += Context->CompactionBlockDataIndices.Num();
+	TotalCompactionInstances += Context->NumCompactionInstances;
 }
\ No newline at end of file
diff --git a/Engine/Source/Runtime/Renderer/Private/InstanceCulling/InstanceCullingMergedContext.h b/Engine/Source/Runtime/Renderer/Private/InstanceCulling/InstanceCullingMergedContext.h
--- a/Engine/Source/Runtime/Renderer/Private/InstanceCulling/InstanceCullingMergedContext.h
+++ b/Engine/Source/Runtime/Renderer/Private/InstanceCulling/InstanceCullingMergedContext.h
@@ -27,6 +27,8 @@ public:
 	{
 		uint32 IndirectArgsOffset;
 		uint32 InstanceDataWriteOffset;
+		uint32 PayloadDataOffset;
+		uint32 CompactionDataOffset;
 		uint32 ViewIdsOffset;
 		uint32 NumViewIds;
 		uint32 DynamicInstanceIdOffset;
@@ -45,7 +47,10 @@ public:
 	//TArray<FMeshDrawCommandInfo, SceneRenderingAllocator> MeshDrawCommandInfos;
 	TArray<FRHIDrawIndexedIndirectParameters, SceneRenderingAllocator> IndirectArgs;
 	TArray<FInstanceCullingContext::FDrawCommandDesc, SceneRenderingAllocator> DrawCommandDescs;
+	TArray<FInstanceCullingContext::FPayloadData, SceneRenderingAllocator> PayloadData;
 	TArray<uint32, SceneRenderingAllocator> InstanceIdOffsets;
+	TArray<FInstanceCullingContext::FCompactionData, SceneRenderingAllocator> DrawCommandCompactionData;
+	TArray<uint32, SceneRenderingAllocator> CompactionBlockDataIndices;	
 
 	TStaticArray<TInstanceCullingLoadBalancer<SceneRenderingAllocator>, static_cast<uint32>(EBatchProcessingMode::Num)> LoadBalancers;
 	TStaticArray<TArray<uint32, SceneRenderingAllocator>, static_cast<uint32>(EBatchProcessingMode::Num)> BatchInds;
@@ -57,8 +62,12 @@ public:
 	TStaticArray<int32, uint32(EBatchProcessingMode::Num)> TotalBatches = TStaticArray<int32, uint32(EBatchProcessingMode::Num)>(InPlace, 0);
 	TStaticArray<int32, uint32(EBatchProcessingMode::Num)> TotalItems = TStaticArray<int32, uint32(EBatchProcessingMode::Num)>(InPlace, 0);
 	int32 TotalIndirectArgs = 0;
+	int32 TotalPayloads = 0;
 	int32 TotalViewIds = 0;
 	int32 TotalInstances = 0;
+	int32 TotalCompactionDrawCommands = 0;
+	int32 TotalCompactionBlocks = 0;
+	int32 TotalCompactionInstances = 0;
 
 	// Single Previous frame HZB which is shared among all batched contexts, thus only one is allowed (but the same can be used in multiple passes). (Needs atlas or bindless to expand).
 	FRDGTextureRef PrevHZB = nullptr;
diff --git a/Engine/Source/Runtime/Renderer/Private/SceneRendering.cpp b/Engine/Source/Runtime/Renderer/Private/SceneRendering.cpp
--- a/Engine/Source/Runtime/Renderer/Private/SceneRendering.cpp
+++ b/Engine/Source/Runtime/Renderer/Private/SceneRendering.cpp
@@ -3645,12 +3645,17 @@ void FSceneRenderer::SetupMeshPass(FViewInfo& View, FExclusiveDepthStencil::Type
 				check(View.GetInstancedView() != nullptr);
 				ViewIds.Add(View.GetInstancedView()->GPUSceneViewId);
 			}
-			const bool bDrawOnlyVSMInvalidatingGeometry = ViewFamily.EngineShowFlags.DrawOnlyVSMInvalidatingGeo != 0;
+			
+			EInstanceCullingFlags CullingFlags = EInstanceCullingFlags::None;
+			if (ViewFamily.EngineShowFlags.DrawOnlyVSMInvalidatingGeo != 0)
+			{
+				EnumAddFlags(CullingFlags, EInstanceCullingFlags::DrawOnlyVSMInvalidatingGeometry);
+			}
 
 			Pass.DispatchPassSetup(
 				Scene,
 				View,
-				FInstanceCullingContext(FeatureLevel, &InstanceCullingManager, ViewIds, View.PrevViewInfo.HZB, InstanceCullingMode, bDrawOnlyVSMInvalidatingGeometry),
+				FInstanceCullingContext(FeatureLevel, &InstanceCullingManager, ViewIds, View.PrevViewInfo.HZB, InstanceCullingMode, CullingFlags),
 				PassType,
 				BasePassDepthStencilAccess,
 				MeshPassProcessor,
diff --git a/Engine/Source/Runtime/Renderer/Private/ShadowSetup.cpp b/Engine/Source/Runtime/Renderer/Private/ShadowSetup.cpp
--- a/Engine/Source/Runtime/Renderer/Private/ShadowSetup.cpp
+++ b/Engine/Source/Runtime/Renderer/Private/ShadowSetup.cpp
@@ -2179,7 +2179,7 @@ void FProjectedShadowInfo::SetupMeshDrawCommandsForShadowDepth(FSceneRenderer& R
 	ShadowDepthPass.DispatchPassSetup(
 		Renderer.Scene,
 		*ShadowDepthView,
-		FInstanceCullingContext(Renderer.FeatureLevel, &InstanceCullingManager, ViewIds, nullptr, EInstanceCullingMode::Normal, false, SingleInstanceProcessingMode),
+		FInstanceCullingContext(Renderer.FeatureLevel, &InstanceCullingManager, ViewIds, nullptr, EInstanceCullingMode::Normal, EInstanceCullingFlags::NoInstanceOrderPreservation, SingleInstanceProcessingMode),
 		EMeshPass::Num,
 		FExclusiveDepthStencil::DepthNop_StencilNop,
 		MeshPassProcessor,
diff --git a/Engine/Source/Runtime/Renderer/Public/InstanceCulling/InstanceCullingContext.h b/Engine/Source/Runtime/Renderer/Public/InstanceCulling/InstanceCullingContext.h
--- a/Engine/Source/Runtime/Renderer/Public/InstanceCulling/InstanceCullingContext.h
+++ b/Engine/Source/Runtime/Renderer/Public/InstanceCulling/InstanceCullingContext.h
@@ -38,6 +38,13 @@ enum class EInstanceCullingMode
 	Stereo,
 };
 
+enum class EInstanceCullingFlags : uint8
+{
+	None 							= 0,			
+	DrawOnlyVSMInvalidatingGeometry	= 1 << 0,
+	NoInstanceOrderPreservation		= 1 << 1,
+};
+ENUM_CLASS_FLAGS(EInstanceCullingFlags)
 
 // Enumeration of the specialized command processing variants
 enum class EBatchProcessingMode : uint32
@@ -58,7 +65,16 @@ class FInstanceProcessingGPULoadBalancer;
 class FInstanceCullingContext
 {
 public:
+	enum class EInstanceFlags : uint8
+	{
+		None 						= 0,			
+		DynamicInstanceDataOffset 	= 1 << 0,
+		ForceInstanceCulling 		= 1 << 1,
+		PreserveInstanceOrder		= 1 << 2
+	};
+
 	static constexpr uint32 IndirectArgsNumWords = 5;
+	static constexpr uint32 CompactionBlockNumInstances = 64;
 	RENDERER_API static uint32 GetInstanceIdBufferStride(ERHIFeatureLevel::Type FeatureLevel);
 	RENDERER_API static uint32 StepInstanceDataOffset(ERHIFeatureLevel::Type FeatureLevel, uint32 NumStepInstances, uint32 NumStepDraws);
 
@@ -74,8 +90,8 @@ public:
 		FInstanceCullingManager* InInstanceCullingManager, 
 		TArrayView<const int32> InViewIds, 
 		const TRefCountPtr<IPooledRenderTarget>& InPrevHZB,
-		EInstanceCullingMode InInstanceCullingMode = EInstanceCullingMode::Normal, 
-		bool bInDrawOnlyVSMInvalidatingGeometry = false, 
+		EInstanceCullingMode InInstanceCullingMode = EInstanceCullingMode::Normal,
+		EInstanceCullingFlags InFlags = EInstanceCullingFlags::None,
 		EBatchProcessingMode InSingleInstanceProcessingMode = EBatchProcessingMode::UnCulled
 	);
 	RENDERER_API ~FInstanceCullingContext();
@@ -83,6 +99,7 @@ public:
 	static RENDERER_API const TRDGUniformBufferRef<FInstanceCullingGlobalUniforms> CreateDummyInstanceCullingUniformBuffer(FRDGBuilder& GraphBuilder);
 
 	static bool IsOcclusionCullingEnabled();
+	
 
 	/**
 	 * Call to empty out the culling commands & other culling data.
@@ -91,17 +108,19 @@ public:
 
 	bool IsEnabled() const { return bIsEnabled; }
 
+	bool IsInstanceOrderPreservationEnabled() const;
+
 	/**
 	 * Add command to cull a range of instances for the given mesh draw command index.
 	 * Multiple commands may add to the same slot, ordering is not preserved.
 	 */
-	void AddInstancesToDrawCommand(uint32 IndirectArgsOffset, int32 InstanceDataOffset, bool bDynamicInstanceDataOffset, bool bForceInstanceCulling, uint32 NumInstances);
+	void AddInstancesToDrawCommand(uint32 IndirectArgsOffset, int32 InstanceDataOffset, uint32 RunOffset, uint32 NumInstances, EInstanceFlags InstanceFlags);
 
 	/**
 	 * Command that is executed in the per-view, post-cull pass to gather up the instances belonging to this primitive.
 	 * Multiple commands may add to the same slot, ordering is not preserved.
 	 */
-	void AddInstanceRunsToDrawCommand(uint32 IndirectArgsOffset, int32 InstanceDataOffset, bool bDynamicInstanceDataOffset, bool bForceInstanceCulling, const uint32* Runs, uint32 NumRuns);
+	void AddInstanceRunsToDrawCommand(uint32 IndirectArgsOffset, int32 InstanceDataOffset, const uint32* Runs, uint32 NumRuns, EInstanceFlags InstanceFlags);
 
 	/*
 	 * Allocate space for indirect draw call argumens for a given MeshDrawCommand and initialize with draw command data.
@@ -163,7 +182,7 @@ public:
 	TRefCountPtr<IPooledRenderTarget> PrevHZB = nullptr;
 	bool bIsEnabled = false;
 	EInstanceCullingMode InstanceCullingMode = EInstanceCullingMode::Normal;
-	bool bDrawOnlyVSMInvalidatingGeometry = false;
+	EInstanceCullingFlags Flags = EInstanceCullingFlags::None;
 
 	uint32 TotalInstances = 0U;
 
@@ -185,11 +204,68 @@ public:
 		uint32 bMaterialMayModifyPosition;
 	};
 
+	struct FPayloadData
+	{
+		uint32 bDynamicInstanceDataOffset_IndirectArgsIndex;
+		uint32 InstanceDataOffset;
+		uint32 RunInstanceOffset;
+		uint32 CompactionDataIndex;
+
+		FPayloadData() = default;
+		FPayloadData(
+			bool bInDynamicInstanceDataOffset,
+			uint32 InIndirectArgsIndex,
+			uint32 InInstanceDataOffset,
+			uint32 InRunInstanceOffset,
+			uint32 InCompactionDataIndex)
+			: bDynamicInstanceDataOffset_IndirectArgsIndex(InIndirectArgsIndex | (bInDynamicInstanceDataOffset ? (1u << 31u) : 0u))
+			, InstanceDataOffset(InInstanceDataOffset)
+			, RunInstanceOffset(InRunInstanceOffset)
+			, CompactionDataIndex(InCompactionDataIndex)
+		{
+			checkSlow(InIndirectArgsIndex < (1u << 31u));
+		}
+	};
+
+	struct FCompactionData
+	{
+		static const uint32 NumViewBits = 8;
+
+		uint32 NumInstances_NumViews;
+		uint32 BlockOffset;
+		uint32 IndirectArgsIndex;
+		uint32 SrcInstanceIdOffset;
+		uint32 DestInstanceIdOffset;
+
+		FCompactionData() = default;
+		FCompactionData(
+			uint32 InNumInstances,
+			uint32 InNumViews,
+			uint32 InBlockOffset,
+			uint32 InIndirectArgsIndex,
+			uint32 InSrcInstanceIdOffset,
+			uint32 InDestInstanceIdOffset)
+			: NumInstances_NumViews(InNumViews | (InNumInstances << NumViewBits))
+			, BlockOffset(InBlockOffset)
+			, IndirectArgsIndex(InIndirectArgsIndex)
+			, SrcInstanceIdOffset(InSrcInstanceIdOffset)
+			, DestInstanceIdOffset(InDestInstanceIdOffset)
+		{
+			checkSlow(InNumViews < (1u << NumViewBits));
+			checkSlow(InNumInstances < (1u << (32 - NumViewBits)));
+		}
+	};
+
 	TArray<FMeshDrawCommandInfo> MeshDrawCommandInfos;
 	TArray<FRHIDrawIndexedIndirectParameters> IndirectArgs;
 	TArray<FDrawCommandDesc> DrawCommandDescs;
+	TArray<FPayloadData> PayloadData;
 	TArray<uint32> InstanceIdOffsets;
 
+	TArray<FCompactionData> DrawCommandCompactionData;	
+	TArray<uint32> CompactionBlockDataIndices;	
+	uint32 NumCompactionInstances = 0U;
+
 	using LoadBalancerArray = TStaticArray<FInstanceProcessingGPULoadBalancer*, static_cast<uint32>(EBatchProcessingMode::Num)>;
 	// Driver for collecting items using one mode of processing
 	LoadBalancerArray LoadBalancers = LoadBalancerArray(InPlace, nullptr);
@@ -201,3 +277,5 @@ public:
 	EBatchProcessingMode SingleInstanceProcessingMode = EBatchProcessingMode::UnCulled;
 };
 
+ENUM_CLASS_FLAGS(FInstanceCullingContext::EInstanceFlags)
+
diff --git a/Engine/Source/Runtime/Renderer/Public/MeshPassProcessor.h b/Engine/Source/Runtime/Renderer/Public/MeshPassProcessor.h
--- a/Engine/Source/Runtime/Renderer/Public/MeshPassProcessor.h
+++ b/Engine/Source/Runtime/Renderer/Public/MeshPassProcessor.h
@@ -1262,10 +1262,13 @@ enum class EFVisibleMeshDrawCommandFlags : uint8
 	HasPrimitiveIdStreamIndex = 1U << 1U,
 
 	/** If set, forces individual instances to always be culled independently from the primitive */
-	ForceInstanceCulling = 1U << 2U,	
+	ForceInstanceCulling = 1U << 2U,
 
-	All = MaterialMayModifyPosition | HasPrimitiveIdStreamIndex | ForceInstanceCulling,
-	NumBits = 3U
+	/** If set, requires that instances preserve their original draw order in the draw command */
+	PreserveInstanceOrder = 1U << 3U,
+
+	All = MaterialMayModifyPosition | HasPrimitiveIdStreamIndex | ForceInstanceCulling | PreserveInstanceOrder,
+	NumBits = 4U
 };
 ENUM_CLASS_FLAGS(EFVisibleMeshDrawCommandFlags);
 static_assert(uint32(EFVisibleMeshDrawCommandFlags::All) < (1U << uint32(EFVisibleMeshDrawCommandFlags::NumBits)), "EFVisibleMeshDrawCommandFlags::NumBits too small to represent all flags in EFVisibleMeshDrawCommandFlags.");
diff --git a/Engine/Source/Runtime/Renderer/Public/MeshPassProcessor.inl b/Engine/Source/Runtime/Renderer/Public/MeshPassProcessor.inl
--- a/Engine/Source/Runtime/Renderer/Public/MeshPassProcessor.inl
+++ b/Engine/Source/Runtime/Renderer/Public/MeshPassProcessor.inl
@@ -121,12 +121,20 @@ void FMeshPassProcessor::BuildMeshDrawCommands(
 		{
 			const FMeshBatchElement& BatchElement = MeshBatch.Elements[BatchElementIndex];
 			FMeshDrawCommand& MeshDrawCommand = DrawListContext->AddCommand(SharedMeshDrawCommand, NumElements);
+			
 			EFVisibleMeshDrawCommandFlags Flags = SharedFlags;
-
-			if (BatchElement.DynamicPrimitiveData != nullptr && BatchElement.DynamicPrimitiveData->bForceInstanceCulling)
+			if (BatchElement.bForceInstanceCulling)
 			{
 				Flags |= EFVisibleMeshDrawCommandFlags::ForceInstanceCulling;
 			}
+			if (BatchElement.bPreserveInstanceOrder)
+			{
+				// TODO: add support for bPreserveInstanceOrder on mobile
+				if (ensureMsgf(FeatureLevel > ERHIFeatureLevel::ES3_1, TEXT("FMeshBatchElement::bPreserveInstanceOrder is currently only supported on non-mobile platforms.")))
+				{					
+					Flags |= EFVisibleMeshDrawCommandFlags::PreserveInstanceOrder;
+				}
+			}
 
 			DataOffset = 0;
 			if (PassShaders.VertexShader.IsValid())
