Remove RENDERER_API from force-inlined function to fix clang linker error

diff --git a/Engine/Source/Runtime/Renderer/Public/PrimitiveSceneInfo.h b/Engine/Source/Runtime/Renderer/Public/PrimitiveSceneInfo.h
--- a/Engine/Source/Runtime/Renderer/Public/PrimitiveSceneInfo.h
+++ b/Engine/Source/Runtime/Renderer/Public/PrimitiveSceneInfo.h
@@ -535,7 +535,7 @@ public:
 	 * In the future, the index will likely be refactored to persist for the lifetime of the component (to facilitate tracking data past proxy re-creates).
 	 * Note: It is not currently used to index any of the current FScene primitive arrays (use PackedIndex), though this is intended to change.
 	 */
-	RENDERER_API FORCEINLINE FPersistentPrimitiveIndex GetPersistentIndex() const { return PersistentIndex; }
+	FORCEINLINE FPersistentPrimitiveIndex GetPersistentIndex() const { return PersistentIndex; }
 
 	/**
 	 * Shifts primitive position and all relevant data by an arbitrary delta.
