Fix non-unity compile errors

diff --git a/Engine/Source/Runtime/Renderer/Public/FXRenderingUtils.h b/Engine/Source/Runtime/Renderer/Public/FXRenderingUtils.h
--- a/Engine/Source/Runtime/Renderer/Public/FXRenderingUtils.h
+++ b/Engine/Source/Runtime/Renderer/Public/FXRenderingUtils.h
@@ -3,6 +3,7 @@
 #pragma once
 
 class FSceneViewFamily;
+class FPrimitiveSceneProxy;
 class FMaterial;
 
 /**
