Initialize PerlinNoise3D texture on mobile platforms

diff --git a/Engine/Source/Runtime/Renderer/Private/SystemTextures.cpp b/Engine/Source/Runtime/Renderer/Private/SystemTextures.cpp
--- a/Engine/Source/Runtime/Renderer/Private/SystemTextures.cpp
+++ b/Engine/Source/Runtime/Renderer/Private/SystemTextures.cpp
@@ -462,7 +462,6 @@ void FSystemTextures::InitializeFeatureLevelDependentTextures(FRHICommandListImm
 		PreintegratedGF = CreateRenderTarget(Texture2D, CreateInfo.DebugName);
 	}
 
-	if (CurrentFeatureLevel < ERHIFeatureLevel::SM5 && InFeatureLevel >= ERHIFeatureLevel::SM5)
 	{
 	    // Create the PerlinNoise3D texture (similar to http://prettyprocs.wordpress.com/2012/10/20/fast-perlin-noise/)
 	    {
