LWC  fix for deferred decal.
Also move some math from shader to CPU for matrix inverse scale calculation.

diff --git a/Engine/Shaders/Private/DeferredDecal.usf b/Engine/Shaders/Private/DeferredDecal.usf
--- a/Engine/Shaders/Private/DeferredDecal.usf
+++ b/Engine/Shaders/Private/DeferredDecal.usf
@@ -6,12 +6,13 @@
 
 #include "DecalCommon.ush"
 
+// large world coordinate tile position
+float4 DecalTilePosition;
 // from SvPosition to decal space (for position), used like this: mul(float4(SvPosition.xyz, 1), SvPositionToDecal);
 float4x4 SvPositionToDecal;
 // to transform position from decal space to world space, for normals use transpose of inverse
 float4x4 DecalToWorld;
-// to transform position from world space to decal space, for normals use transpose of inverse
-float4x4 WorldToDecal;
+float3 DecalToWorldInvScale;
 // decal orientation in world space, X axis. Used for ObjectOrientation material expression
 float3 DecalOrientation;
 // x - Fade alpha from screen size fading
diff --git a/Engine/Shaders/Private/MaterialTemplate.ush b/Engine/Shaders/Private/MaterialTemplate.ush
--- a/Engine/Shaders/Private/MaterialTemplate.ush
+++ b/Engine/Shaders/Private/MaterialTemplate.ush
@@ -964,7 +964,7 @@ MaterialFloat2 GetDefaultSceneTextureUV(FMaterialPixelParameters Parameters, con
 	FLWCVector3 GetActorWorldPosition(FMaterialVertexParameters Parameters)
 	{
 	#if DECAL_PRIMITIVE
-		return LWCPromote(DecalToWorld[3].xyz);
+		return MakeLWCVector3(DecalTilePosition, DecalToWorld[3].xyz);
 	#else
 		return GetPrimitiveData(Parameters).ActorWorldPosition;
 	#endif
@@ -973,7 +973,7 @@ MaterialFloat2 GetDefaultSceneTextureUV(FMaterialPixelParameters Parameters, con
 	FLWCVector3 GetActorWorldPosition(FMaterialPixelParameters Parameters)
 	{
 	#if DECAL_PRIMITIVE
-		return LWCPromote(DecalToWorld[3].xyz);
+		return MakeLWCVector3(DecalTilePosition, DecalToWorld[3].xyz);
 	#else
 		return GetPrimitiveData(Parameters).ActorWorldPosition;
 	#endif
@@ -2782,14 +2782,9 @@ float3 TransformTangentNormalToWorld(MaterialFloat3x3 TangentToWorld, float3 Tan
 	// To transform the normals use tranpose(Inverse(DecalToWorld)) = transpose(WorldToDecal)
 	// But we want to only rotate the normals (we don't want to non-uniformaly scale them).
 	// We assume the matrix is only a scale and rotation, and we remove non-uniform scale:
-	float3 lengthSqr = { length2(DecalToWorld._m00_m01_m02),
-							length2(DecalToWorld._m10_m11_m12),
-							length2(DecalToWorld._m20_m21_m22) };
-
-	float3 scale = rsqrt(lengthSqr);
 				
 	// Pre-multiply by the inverse of the non-uniform scale in DecalToWorld
-	float4 ScaledNormal = float4(-TangentNormal.z * scale.x, TangentNormal.y * scale.y, TangentNormal.x * scale.z, 0.f);
+	float4 ScaledNormal = float4(-TangentNormal.z * DecalToWorldInvScale.x, TangentNormal.y * DecalToWorldInvScale.y, TangentNormal.x * DecalToWorldInvScale.z, 0.f);
 
 	// Compute the normal 
 	return normalize(mul(ScaledNormal, DecalToWorld).xyz);
diff --git a/Engine/Source/Runtime/Renderer/Private/DecalRenderingShared.cpp b/Engine/Source/Runtime/Renderer/Private/DecalRenderingShared.cpp
--- a/Engine/Source/Runtime/Renderer/Private/DecalRenderingShared.cpp
+++ b/Engine/Source/Runtime/Renderer/Private/DecalRenderingShared.cpp
@@ -72,10 +72,11 @@ public:
 	FDeferredDecalPS(const ShaderMetaType::CompiledShaderInitializerType& Initializer)
 		: FMaterialShader(Initializer)
 	{
+		DecalTilePosition.Bind(Initializer.ParameterMap, TEXT("DecalTilePosition"));
 		SvPositionToDecal.Bind(Initializer.ParameterMap,TEXT("SvPositionToDecal"));
 		DecalToWorld.Bind(Initializer.ParameterMap,TEXT("DecalToWorld"));
-		WorldToDecal.Bind(Initializer.ParameterMap,TEXT("WorldToDecal"));
-		DecalOrientation.Bind(Initializer.ParameterMap,TEXT("DecalOrientation"));
+		DecalToWorldInvScale.Bind(Initializer.ParameterMap, TEXT("DecalToWorldInvScale"));
+		DecalOrientation.Bind(Initializer.ParameterMap, TEXT("DecalOrientation"));
 		DecalParams.Bind(Initializer.ParameterMap, TEXT("DecalParams"));
 	}
 
@@ -87,19 +88,24 @@ public:
 		const FMaterial& Material = MaterialProxy->GetMaterialWithFallback(View.GetFeatureLevel(), MaterialProxyForRendering);
 		FMaterialShader::SetParameters(RHICmdList, ShaderRHI, MaterialProxyForRendering, Material, View);
 
-		FTransform ComponentTrans = DecalProxy.ComponentTrans;
+		const FLargeWorldRenderPosition AbsoluteOrigin(View.ViewMatrices.GetInvViewMatrix().GetOrigin());
+		const FVector3f TilePosition = AbsoluteOrigin.GetTile();
+		const FMatrix WorldToDecalMatrix = DecalProxy.ComponentTrans.ToInverseMatrixWithScale();
+		const FMatrix44f RelativeWorldToDecalMatrix = AbsoluteOrigin.MakeFromRelativeWorldMatrix(WorldToDecalMatrix);
+		const FMatrix DecalToWorldMatrix = DecalProxy.ComponentTrans.ToMatrixWithScale();
+		const FMatrix44f RelativeDecalToWorldMatrix = AbsoluteOrigin.MakeToRelativeWorldMatrix(DecalToWorldMatrix);
+		const FVector3f OrientationVector = DecalProxy.ComponentTrans.GetUnitAxis(EAxis::X);
 
-		FMatrix44f WorldToComponent = ComponentTrans.ToInverseMatrixWithScale();
-
-		// Set the transform from screen space to light space.
-		if(SvPositionToDecal.IsBound())
+		if (DecalTilePosition.IsBound())
+		{
+			SetShaderValue(RHICmdList, ShaderRHI, DecalTilePosition, TilePosition);
+		}
+		if (SvPositionToDecal.IsBound())
 		{
 			FVector2D InvViewSize = FVector2D(1.0f / View.ViewRect.Width(), 1.0f / View.ViewRect.Height());
 
 			// setup a matrix to transform float4(SvPosition.xyz,1) directly to Decal (quality, performance as we don't need to convert or use interpolator)
-
 			//	new_xy = (xy - ViewRectMin.xy) * ViewSizeAndInvSize.zw * float2(2,-2) + float2(-1, 1);
-
 			//  transformed into one MAD:  new_xy = xy * ViewSizeAndInvSize.zw * float2(2,-2)      +       (-ViewRectMin.xy) * ViewSizeAndInvSize.zw * float2(2,-2) + float2(-1, 1);
 
 			float Mx = 2.0f * InvViewSize.X;
@@ -114,25 +120,21 @@ public:
 					FPlane( 0, My,   0,  0),
 					FPlane( 0,  0,   1,  0),
 					FPlane(Ax, Ay,   0,  1)
-				) * View.ViewMatrices.GetInvViewProjectionMatrix() * WorldToComponent;
+				) * View.ViewMatrices.GetInvViewProjectionMatrix() * WorldToDecalMatrix;
 
 			SetShaderValue(RHICmdList, ShaderRHI, SvPositionToDecal, SvPositionToDecalValue);
 		}
-
-		// Set the transform from light space to world space
-		if(DecalToWorld.IsBound())
+		if (DecalToWorld.IsBound())
 		{
-			const FMatrix44f DecalToWorldValue = ComponentTrans.ToMatrixWithScale();
-			
-			SetShaderValue(RHICmdList, ShaderRHI, DecalToWorld, DecalToWorldValue);
+			SetShaderValue(RHICmdList, ShaderRHI, DecalToWorld, RelativeDecalToWorldMatrix);
+		}
+		if (DecalToWorldInvScale.IsBound())
+		{
+			SetShaderValue(RHICmdList, ShaderRHI, DecalToWorldInvScale, RelativeDecalToWorldMatrix.GetScaleVector().Reciprocal());
 		}
-
-		SetShaderValue(RHICmdList, ShaderRHI, WorldToDecal, WorldToComponent);
-
 		if (DecalOrientation.IsBound())
 		{
-			// can get DecalOrientation form DecalToWorld matrix, but it will require binding whole matrix and normalizing axis in the shader
-			SetShaderValue(RHICmdList, ShaderRHI, DecalOrientation, FVector3f(ComponentTrans.GetUnitAxis(EAxis::X)));
+			SetShaderValue(RHICmdList, ShaderRHI, DecalOrientation, OrientationVector);
 		}
 		
 		float LifetimeAlpha = 1.0f;
@@ -148,8 +150,9 @@ public:
 
 private:
 	LAYOUT_FIELD(FShaderParameter, SvPositionToDecal);
+	LAYOUT_FIELD(FShaderParameter, DecalTilePosition);
 	LAYOUT_FIELD(FShaderParameter, DecalToWorld);
-	LAYOUT_FIELD(FShaderParameter, WorldToDecal);
+	LAYOUT_FIELD(FShaderParameter, DecalToWorldInvScale);
 	LAYOUT_FIELD(FShaderParameter, DecalOrientation);
 	LAYOUT_FIELD(FShaderParameter, DecalParams);
 };
