Volumetric Fog results in Reflection Capture are not correct.

Add a new CVar �r.ReflectionSkyLightCapture.EnableLightFunctions� to resolve this issue.

diff --git a/Engine/Source/Runtime/Renderer/Private/ReflectionEnvironmentCapture.cpp b/Engine/Source/Runtime/Renderer/Private/ReflectionEnvironmentCapture.cpp
--- a/Engine/Source/Runtime/Renderer/Private/ReflectionEnvironmentCapture.cpp
+++ b/Engine/Source/Runtime/Renderer/Private/ReflectionEnvironmentCapture.cpp
@@ -99,6 +99,21 @@ static FAutoConsoleVariableRef CVarFreeReflectionScratchAfterUse(
 	GFreeReflectionScratchAfterUse,
 	TEXT("Free reflection scratch render targets after use."));
 
+/**
+* This CVar might affect the quality and performance for: 
+* (1) Captured reflection: Increase volume and light function quality and the cost at lighting build time.
+* (2) Sky light scene capture (non real time): Increase quality and the cost at the start of a level when the scene is captured.
+* (3) Real time sky light capture: this one does not render volumetric fog or anything that reads a light function. It increases the cost only.
+*
+* It might also create mismatch when light function is time dependent (e.g., sun light simulating time-varying cloud shadows). Different level building can look inconsistent builds after builds.
+*/
+static int32 GReflectionCaptureEnableLightFunctions = 0;
+static FAutoConsoleVariableRef CVarReflectionCaptureEnableLightFunctions(
+	TEXT("r.ReflectionCapture.EnableLightFunctions"),
+	GReflectionCaptureEnableLightFunctions,
+	TEXT("0. Disable light functions in reflection/sky light capture (default).\n")
+	TEXT("Others. Enable light functions."));
+
 TGlobalResource<FReflectionScratchCubemaps> GReflectionScratchCubemaps;
 
 bool DoGPUArrayCopy()
@@ -1391,9 +1406,9 @@ void CaptureSceneIntoScratchCubemap(
 		ViewFamily.EngineShowFlags.MotionBlur = 0;
 		ViewFamily.EngineShowFlags.SetOnScreenDebug(false);
 		ViewFamily.EngineShowFlags.HMDDistortion = 0;
-		// Exclude particles and light functions as they are usually dynamic, and can't be captured well
+		// Conditionally exclude particles and light functions as they are usually dynamic, and can't be captured well
 		ViewFamily.EngineShowFlags.Particles = 0;
-		ViewFamily.EngineShowFlags.LightFunctions = 0;
+		ViewFamily.EngineShowFlags.LightFunctions = abs(GReflectionCaptureEnableLightFunctions) ? 1 : 0;
 		ViewFamily.EngineShowFlags.SetCompositeEditorPrimitives(false);
 		// These are highly dynamic and can't be captured effectively
 		ViewFamily.EngineShowFlags.LightShafts = 0;
