Fix compile error from changed API

diff --git a/Engine/Source/Runtime/Renderer/Private/DecalRenderingShared.cpp b/Engine/Source/Runtime/Renderer/Private/DecalRenderingShared.cpp
--- a/Engine/Source/Runtime/Renderer/Private/DecalRenderingShared.cpp
+++ b/Engine/Source/Runtime/Renderer/Private/DecalRenderingShared.cpp
@@ -91,9 +91,9 @@ public:
 		const FLargeWorldRenderPosition AbsoluteOrigin(View.ViewMatrices.GetInvViewMatrix().GetOrigin());
 		const FVector3f TilePosition = AbsoluteOrigin.GetTile();
 		const FMatrix WorldToDecalMatrix = DecalProxy.ComponentTrans.ToInverseMatrixWithScale();
-		const FMatrix44f RelativeWorldToDecalMatrix = AbsoluteOrigin.MakeFromRelativeWorldMatrix(WorldToDecalMatrix);
+		const FMatrix44f RelativeWorldToDecalMatrix = FLargeWorldRenderScalar::MakeFromRelativeWorldMatrix(AbsoluteOrigin.GetTileOffset(), WorldToDecalMatrix);
 		const FMatrix DecalToWorldMatrix = DecalProxy.ComponentTrans.ToMatrixWithScale();
-		const FMatrix44f RelativeDecalToWorldMatrix = AbsoluteOrigin.MakeToRelativeWorldMatrix(DecalToWorldMatrix);
+		const FMatrix44f RelativeDecalToWorldMatrix = FLargeWorldRenderScalar::MakeToRelativeWorldMatrix(AbsoluteOrigin.GetTileOffset(), DecalToWorldMatrix);
 		const FVector3f OrientationVector = DecalProxy.ComponentTrans.GetUnitAxis(EAxis::X);
 
 		if (DecalTilePosition.IsBound())
