FEM - Optimize SceneUpdates for a small MVP improvement

diff --git a/Engine/Source/Runtime/Core/Public/Experimental/Containers/RobinHoodHashTable.h b/Engine/Source/Runtime/Core/Public/Experimental/Containers/RobinHoodHashTable.h
--- a/Engine/Source/Runtime/Core/Public/Experimental/Containers/RobinHoodHashTable.h
+++ b/Engine/Source/Runtime/Core/Public/Experimental/Containers/RobinHoodHashTable.h
@@ -455,6 +455,7 @@ namespace RobinHoodHashTable_Private
 				KeyVals.SetNumUnsafeInternal(0);
 				KeyVals.Empty();
 				FreeList.Empty();
+				Hashes.Empty();
 			}
 
 			void Reserve(SizeType ReserveNum)
diff --git a/Engine/Source/Runtime/Renderer/Private/RendererScene.cpp b/Engine/Source/Runtime/Renderer/Private/RendererScene.cpp
--- a/Engine/Source/Runtime/Renderer/Private/RendererScene.cpp
+++ b/Engine/Source/Runtime/Renderer/Private/RendererScene.cpp
@@ -1017,11 +1017,11 @@ void FScene::AddPrimitiveSceneInfo_RenderThread(FPrimitiveSceneInfo* PrimitiveSc
 {
 	check(IsInRenderingThread());
 	check(PrimitiveSceneInfo->PackedIndex == INDEX_NONE);
-	check(!AddedPrimitiveSceneInfos.Contains(PrimitiveSceneInfo));
-	AddedPrimitiveSceneInfos.Add(PrimitiveSceneInfo);
+	check(AddedPrimitiveSceneInfos.Find(PrimitiveSceneInfo) == nullptr);
+	AddedPrimitiveSceneInfos.FindOrAdd(PrimitiveSceneInfo);
 	if (PreviousTransform.IsSet())
 	{
-		OverridenPreviousTransforms.Add(PrimitiveSceneInfo, PreviousTransform.GetValue().ToMatrixWithScale());
+		OverridenPreviousTransforms.FindOrAdd(PrimitiveSceneInfo, PreviousTransform.GetValue().ToMatrixWithScale());
 	}
 }
 
@@ -1413,7 +1413,7 @@ void FScene::UpdatePrimitiveTransform_RenderThread(FPrimitiveSceneProxy* Primiti
 	check(IsInRenderingThread());
 
 #if VALIDATE_PRIMITIVE_PACKED_INDEX
-	if (AddedPrimitiveSceneInfos.Contains(PrimitiveSceneProxy->GetPrimitiveSceneInfo()))
+	if (AddedPrimitiveSceneInfos.Find(PrimitiveSceneProxy->GetPrimitiveSceneInfo()) != nullptr)
 	{
 		check(PrimitiveSceneProxy->GetPrimitiveSceneInfo()->PackedIndex == INDEX_NONE);
 	}
@@ -1422,14 +1422,14 @@ void FScene::UpdatePrimitiveTransform_RenderThread(FPrimitiveSceneProxy* Primiti
 		check(PrimitiveSceneProxy->GetPrimitiveSceneInfo()->PackedIndex != INDEX_NONE);
 	}
 
-	check(!RemovedPrimitiveSceneInfos.Contains(PrimitiveSceneProxy->GetPrimitiveSceneInfo()));
+	check(RemovedPrimitiveSceneInfos.Find(PrimitiveSceneProxy->GetPrimitiveSceneInfo()) == nullptr);
 #endif
 
-	UpdatedTransforms.Add(PrimitiveSceneProxy, { WorldBounds, LocalBounds, LocalToWorld, AttachmentRootPosition });
+	UpdatedTransforms.FindOrAdd(PrimitiveSceneProxy, { WorldBounds, LocalBounds, LocalToWorld, AttachmentRootPosition });
 
 	if (PreviousTransform.IsSet())
 	{
-		OverridenPreviousTransforms.Add(PrimitiveSceneProxy->GetPrimitiveSceneInfo(), PreviousTransform.GetValue().ToMatrixWithScale());
+		OverridenPreviousTransforms.FindOrAdd(PrimitiveSceneProxy->GetPrimitiveSceneInfo(), PreviousTransform.GetValue().ToMatrixWithScale());
 	}
 }
 
@@ -1438,7 +1438,7 @@ void FScene::UpdatePrimitiveOcclusionBoundsSlack_RenderThread(const FPrimitiveSc
 	check(IsInRenderingThread());
 
 #if VALIDATE_PRIMITIVE_PACKED_INDEX
-	if (AddedPrimitiveSceneInfos.Contains(PrimitiveSceneProxy->GetPrimitiveSceneInfo()))
+	if (AddedPrimitiveSceneInfos.Find(PrimitiveSceneProxy->GetPrimitiveSceneInfo()) != nullptr)
 	{
 		check(PrimitiveSceneProxy->GetPrimitiveSceneInfo()->PackedIndex == INDEX_NONE);
 	}
@@ -1447,10 +1447,10 @@ void FScene::UpdatePrimitiveOcclusionBoundsSlack_RenderThread(const FPrimitiveSc
 		check(PrimitiveSceneProxy->GetPrimitiveSceneInfo()->PackedIndex != INDEX_NONE);
 	}
 
-	check(!RemovedPrimitiveSceneInfos.Contains(PrimitiveSceneProxy->GetPrimitiveSceneInfo()));
+	check(RemovedPrimitiveSceneInfos.Find(PrimitiveSceneProxy->GetPrimitiveSceneInfo()) == nullptr);
 #endif
 
-	UpdatedOcclusionBoundsSlacks.Add(PrimitiveSceneProxy, NewSlack);
+	UpdatedOcclusionBoundsSlacks.FindOrAdd(PrimitiveSceneProxy, NewSlack);
 }
 
 void FScene::UpdatePrimitiveTransform(UPrimitiveComponent* Primitive)
@@ -1670,7 +1670,7 @@ void FScene::UpdatePrimitiveLightingAttachmentRoot(UPrimitiveComponent* Primitiv
 			[Scene, Proxy, NewComponentId](FRHICommandList&)
 			{
 				FPrimitiveSceneInfo* PrimitiveInfo = Proxy->GetPrimitiveSceneInfo();
-				Scene->UpdatedAttachmentRoots.Add(PrimitiveInfo, NewComponentId);
+				Scene->UpdatedAttachmentRoots.FindOrAdd(PrimitiveInfo, NewComponentId);
 			});
 	}
 }
@@ -1724,7 +1724,7 @@ void FScene::UpdateCustomPrimitiveData(UPrimitiveComponent* Primitive)
 		ENQUEUE_RENDER_COMMAND(UpdateCustomPrimitiveDataCommand)(
 			[UpdateParams](FRHICommandListImmediate& RHICmdList)
 			{
-				UpdateParams.Scene->UpdatedCustomPrimitiveParams.Add(UpdateParams.PrimitiveSceneProxy, UpdateParams.CustomPrimitiveData);
+				UpdateParams.Scene->UpdatedCustomPrimitiveParams.FindOrAdd(UpdateParams.PrimitiveSceneProxy, UpdateParams.CustomPrimitiveData);
 			});
 	}
 }
@@ -1743,7 +1743,7 @@ void FScene::UpdatePrimitiveDistanceFieldSceneData_GameThread(UPrimitiveComponen
 				if (PrimitiveSceneProxy && PrimitiveSceneProxy->GetPrimitiveSceneInfo())
 				{
 					FPrimitiveSceneInfo* Info = PrimitiveSceneProxy->GetPrimitiveSceneInfo();
-					this->DistanceFieldSceneDataUpdates.Add(Info);
+					this->DistanceFieldSceneDataUpdates.FindOrAdd(Info);
 				}
 			});
 	}
@@ -1790,8 +1790,8 @@ void FScene::RemovePrimitiveSceneInfo_RenderThread(FPrimitiveSceneInfo* Primitiv
 	else
 	{
 		check(PrimitiveSceneInfo->PackedIndex != INDEX_NONE);
-		check(!RemovedPrimitiveSceneInfos.Contains(PrimitiveSceneInfo));
-		RemovedPrimitiveSceneInfos.Add(PrimitiveSceneInfo);
+		check(RemovedPrimitiveSceneInfos.Find(PrimitiveSceneInfo) == nullptr);
+		RemovedPrimitiveSceneInfos.FindOrAdd(PrimitiveSceneInfo);
 	}
 }
 
@@ -4205,10 +4205,15 @@ void FScene::UpdateAllPrimitiveSceneInfos(FRDGBuilder& GraphBuilder, bool bAsync
 	UpdateRayTracingGroupBounds_AddPrimitives(AddedPrimitiveSceneInfos);
 #endif
 
-	TArray<FPrimitiveSceneInfo*> RemovedLocalPrimitiveSceneInfos(RemovedPrimitiveSceneInfos.Array());
+	TArray<FPrimitiveSceneInfo*> RemovedLocalPrimitiveSceneInfos;
+	RemovedLocalPrimitiveSceneInfos.Reserve(RemovedPrimitiveSceneInfos.Num());
+	for (FPrimitiveSceneInfo* SceneInfo : RemovedPrimitiveSceneInfos)
+	{
+		RemovedLocalPrimitiveSceneInfos.Add(SceneInfo);
+	}
 	// NOTE: We clear this early because IsPrimitiveBeingRemoved gets called from the CreateLightPrimitiveInteraction (to make sure that old primitives are not accessed) 
 	// we cannot safely kick off the AsyncCreateLightPrimitiveInteractionsTask before the RemovedPrimitiveSceneInfos has been cleared.
-	RemovedPrimitiveSceneInfos.Reset();
+	RemovedPrimitiveSceneInfos.Empty();
 
 	RemovedLocalPrimitiveSceneInfos.Sort(FPrimitiveArraySortKey());
 
@@ -4234,7 +4239,13 @@ void FScene::UpdateAllPrimitiveSceneInfos(FRDGBuilder& GraphBuilder, bool bAsync
 
 		VirtualShadowMapArrayCacheManager->ProcessRemovedOrUpdatedPrimitives(GraphBuilder, GPUScene, InvalidatingPrimitiveCollector);
 	}
-	TArray<FPrimitiveSceneInfo*> AddedLocalPrimitiveSceneInfos(AddedPrimitiveSceneInfos.Array());
+	TArray<FPrimitiveSceneInfo*> AddedLocalPrimitiveSceneInfos;
+	AddedLocalPrimitiveSceneInfos.Reserve(AddedPrimitiveSceneInfos.Num());
+	for (FPrimitiveSceneInfo* SceneInfo : AddedPrimitiveSceneInfos)
+	{
+		AddedLocalPrimitiveSceneInfos.Add(SceneInfo);
+	}
+
 	AddedLocalPrimitiveSceneInfos.Sort(FPrimitiveArraySortKey());
 
 	TSet<FPrimitiveSceneInfo*> DeletedSceneInfos;
@@ -5063,14 +5074,14 @@ void FScene::UpdateAllPrimitiveSceneInfos(FRDGBuilder& GraphBuilder, bool bAsync
 		}
 	}
 	
-	UpdatedAttachmentRoots.Reset();
-	UpdatedTransforms.Reset();
+	UpdatedAttachmentRoots.Empty();
+	UpdatedTransforms.Empty();
 	UpdatedInstances.Reset();
-	UpdatedCustomPrimitiveParams.Reset();
-	OverridenPreviousTransforms.Reset();
-	UpdatedOcclusionBoundsSlacks.Reset();
-	DistanceFieldSceneDataUpdates.Reset();
-	AddedPrimitiveSceneInfos.Reset();
+	UpdatedCustomPrimitiveParams.Empty();
+	OverridenPreviousTransforms.Empty();
+	UpdatedOcclusionBoundsSlacks.Empty();
+	DistanceFieldSceneDataUpdates.Empty();
+	AddedPrimitiveSceneInfos.Empty();
 }
 
 void FScene::CreateLightPrimitiveInteractionsForPrimitive(FPrimitiveSceneInfo* PrimitiveInfo, bool bAsyncCreateLPIs)
@@ -5096,7 +5107,7 @@ void FScene::CreateLightPrimitiveInteractionsForPrimitive(FPrimitiveSceneInfo* P
 bool FScene::IsPrimitiveBeingRemoved(FPrimitiveSceneInfo* PrimitiveSceneInfo) const
 {
 	check(IsInParallelRenderingThread() || IsInRenderingThread());
-	return RemovedPrimitiveSceneInfos.Contains(PrimitiveSceneInfo);
+	return RemovedPrimitiveSceneInfos.Find(PrimitiveSceneInfo) != nullptr;
 }
 
 #if RHI_RAYTRACING
diff --git a/Engine/Source/Runtime/Renderer/Private/ScenePrivate.h b/Engine/Source/Runtime/Renderer/Private/ScenePrivate.h
--- a/Engine/Source/Runtime/Renderer/Private/ScenePrivate.h
+++ b/Engine/Source/Runtime/Renderer/Private/ScenePrivate.h
@@ -60,6 +60,7 @@
 #include "PrimitiveInstanceUpdateCommand.h"
 #include "OIT/OIT.h"
 #include "ShadingEnergyConservation.h"
+#include "Experimental/Containers/RobinHoodHashTable.h"
 #include "SpanAllocator.h"
 
 /** Factor by which to grow occlusion tests **/
@@ -3497,15 +3498,15 @@ private:
 		FBoxSphereBounds StaticMeshBounds;
 	};
 
-	TMap<FPrimitiveSceneInfo*, FPrimitiveComponentId> UpdatedAttachmentRoots;
-	TMap<FPrimitiveSceneProxy*, FCustomPrimitiveData> UpdatedCustomPrimitiveParams;
-	TMap<FPrimitiveSceneProxy*, FUpdateTransformCommand> UpdatedTransforms;
+	Experimental::TRobinHoodHashMap<FPrimitiveSceneInfo*, FPrimitiveComponentId> UpdatedAttachmentRoots;
+	Experimental::TRobinHoodHashMap<FPrimitiveSceneProxy*, FCustomPrimitiveData> UpdatedCustomPrimitiveParams;
+	Experimental::TRobinHoodHashMap<FPrimitiveSceneProxy*, FUpdateTransformCommand> UpdatedTransforms;
 	TMap<FPrimitiveSceneProxy*, FUpdateInstanceCommand> UpdatedInstances;
-	TMap<FPrimitiveSceneInfo*, FMatrix> OverridenPreviousTransforms;
-	TMap<const FPrimitiveSceneProxy*, float> UpdatedOcclusionBoundsSlacks;
-	TSet<FPrimitiveSceneInfo*> AddedPrimitiveSceneInfos;
-	TSet<FPrimitiveSceneInfo*> RemovedPrimitiveSceneInfos;
-	TSet<FPrimitiveSceneInfo*> DistanceFieldSceneDataUpdates;
+	Experimental::TRobinHoodHashMap<FPrimitiveSceneInfo*, FMatrix> OverridenPreviousTransforms;
+	Experimental::TRobinHoodHashMap<const FPrimitiveSceneProxy*, float> UpdatedOcclusionBoundsSlacks;
+	Experimental::TRobinHoodHashMap<FPrimitiveSceneInfo*> AddedPrimitiveSceneInfos;
+	Experimental::TRobinHoodHashMap<FPrimitiveSceneInfo*> RemovedPrimitiveSceneInfos;
+	Experimental::TRobinHoodHashMap<FPrimitiveSceneInfo*> DistanceFieldSceneDataUpdates;
 
 	/** 
 	 * The number of visible lights in the scene
