REL - merge fixes

diff --git a/Engine/Source/Runtime/Renderer/Private/ScenePrivate.h b/Engine/Source/Runtime/Renderer/Private/ScenePrivate.h
--- a/Engine/Source/Runtime/Renderer/Private/ScenePrivate.h
+++ b/Engine/Source/Runtime/Renderer/Private/ScenePrivate.h
@@ -3504,9 +3504,9 @@ private:
 	TMap<FPrimitiveSceneProxy*, FUpdateInstanceCommand> UpdatedInstances;
 	Experimental::TRobinHoodHashMap<FPrimitiveSceneInfo*, FMatrix> OverridenPreviousTransforms;
 	Experimental::TRobinHoodHashMap<const FPrimitiveSceneProxy*, float> UpdatedOcclusionBoundsSlacks;
-	Experimental::TRobinHoodHashMap<FPrimitiveSceneInfo*> AddedPrimitiveSceneInfos;
-	Experimental::TRobinHoodHashMap<FPrimitiveSceneInfo*> RemovedPrimitiveSceneInfos;
-	Experimental::TRobinHoodHashMap<FPrimitiveSceneInfo*> DistanceFieldSceneDataUpdates;
+	Experimental::TRobinHoodHashSet<FPrimitiveSceneInfo*> AddedPrimitiveSceneInfos;
+	Experimental::TRobinHoodHashSet<FPrimitiveSceneInfo*> RemovedPrimitiveSceneInfos;
+	Experimental::TRobinHoodHashSet<FPrimitiveSceneInfo*> DistanceFieldSceneDataUpdates;
 
 	/** 
 	 * The number of visible lights in the scene
