REL - fix raytracing updates

diff --git a/Engine/Source/Runtime/Core/Public/Experimental/Containers/RobinHoodHashTable.h b/Engine/Source/Runtime/Core/Public/Experimental/Containers/RobinHoodHashTable.h
--- a/Engine/Source/Runtime/Core/Public/Experimental/Containers/RobinHoodHashTable.h
+++ b/Engine/Source/Runtime/Core/Public/Experimental/Containers/RobinHoodHashTable.h
@@ -208,6 +208,7 @@ namespace RobinHoodHashTable_Private
 			FreeList.Reset(1);
 			//push a dummy
 			FreeList.Push(FSpan{ INDEX_NONE, INDEX_NONE });
+			NumElements = 0;
 		}
 
 		int Num() const
diff --git a/Engine/Source/Runtime/Renderer/Private/RendererScene.cpp b/Engine/Source/Runtime/Renderer/Private/RendererScene.cpp
--- a/Engine/Source/Runtime/Renderer/Private/RendererScene.cpp
+++ b/Engine/Source/Runtime/Renderer/Private/RendererScene.cpp
@@ -1625,7 +1625,7 @@
 				check(!RemovedPrimitiveSceneInfos.Contains(UpdateParams.PrimitiveSceneProxy->GetPrimitiveSceneInfo()));
 #endif
 				FScopeCycleCounter Context(UpdateParams.PrimitiveSceneProxy->GetStatId());
-				UpdatedInstances.Add(UpdateParams.PrimitiveSceneProxy, UpdateParams);
+				UpdatedInstances.FindOrAdd(UpdateParams.PrimitiveSceneProxy, UpdateParams);
 			}
 		);
 	}
@@ -4188,6 +4188,95 @@
 	return bPlatformSupportsVelocityRendering && bShouldPrimitiveOutputVelocity;
 }
 
+#if RHI_RAYTRACING
+void FScene::UpdateRayTracingGroupBounds_AddPrimitives(const Experimental::TRobinHoodHashSet<FPrimitiveSceneInfo*>& PrimitiveSceneInfos)
+{
+	for (FPrimitiveSceneInfo* const PrimitiveSceneInfo : PrimitiveSceneInfos)
+	{
+		const int32 GroupId = PrimitiveSceneInfo->Proxy->GetRayTracingGroupId();
+		if (GroupId != -1)
+		{
+			bool bInMap = false;
+			static const FRayTracingCullingGroup DefaultGroup;
+			FRayTracingCullingGroup* const Group = PrimitiveRayTracingGroups.FindOrAdd(GroupId, DefaultGroup, bInMap);
+			if (bInMap)
+			{
+				Group->Bounds = Group->Bounds + PrimitiveSceneInfo->Proxy->GetBounds();
+			}
+			else
+			{
+				Group->Bounds = PrimitiveSceneInfo->Proxy->GetBounds();
+			}
+			Group->Primitives.Add(PrimitiveSceneInfo);
+		}
+	}
+}
+
+static void UpdateRayTracingGroupBounds(Experimental::TRobinHoodHashSet<FScene::FRayTracingCullingGroup*>& GroupsToUpdate)
+{
+	for (FScene::FRayTracingCullingGroup* const Group : GroupsToUpdate)
+	{
+		bool bFirstBounds = false;
+		for (FPrimitiveSceneInfo* const Primitive : Group->Primitives)
+		{
+			if (!bFirstBounds)
+			{
+				Group->Bounds = Primitive->Proxy->GetBounds();
+				bFirstBounds = true;
+			}
+			else
+			{
+				Group->Bounds = Group->Bounds + Primitive->Proxy->GetBounds();
+			}
+		}
+	}
+}
+
+void FScene::UpdateRayTracingGroupBounds_RemovePrimitives(const Experimental::TRobinHoodHashSet<FPrimitiveSceneInfo*>& PrimitiveSceneInfos)
+{
+	Experimental::TRobinHoodHashSet<FRayTracingCullingGroup*> GroupsToUpdate;
+	for (FPrimitiveSceneInfo* const PrimitiveSceneInfo : PrimitiveSceneInfos)
+	{
+		const int32 RayTracingGroupId = PrimitiveSceneInfo->Proxy->GetRayTracingGroupId();
+		const Experimental::FHashElementId GroupId = (RayTracingGroupId != -1) ? PrimitiveRayTracingGroups.FindId(RayTracingGroupId) : Experimental::FHashElementId();
+		if (GroupId.IsValid())
+		{
+			FRayTracingCullingGroup& Group = PrimitiveRayTracingGroups.GetByElementId(GroupId).Value;
+			Group.Primitives.RemoveSingleSwap(PrimitiveSceneInfo);
+			if (Group.Primitives.Num() == 0)
+			{
+				PrimitiveRayTracingGroups.RemoveByElementId(GroupId);
+			}
+			else
+			{
+				GroupsToUpdate.FindOrAdd(&Group);
+			}
+		}
+	}
+
+	UpdateRayTracingGroupBounds(GroupsToUpdate);
+}
+
+template<typename ValueType>
+inline void FScene::UpdateRayTracingGroupBounds_UpdatePrimitives(const Experimental::TRobinHoodHashMap<FPrimitiveSceneProxy*, ValueType>& UpdatedTransforms)
+{
+	Experimental::TRobinHoodHashSet<FRayTracingCullingGroup*> GroupsToUpdate;
+	for (const auto& Transform : UpdatedTransforms)
+	{
+		FPrimitiveSceneProxy* const PrimitiveSceneProxy = Transform.Key;
+		const int32 RayTracingGroupId = PrimitiveSceneProxy->GetRayTracingGroupId();
+		const Experimental::FHashElementId GroupId = (RayTracingGroupId != -1) ? PrimitiveRayTracingGroups.FindId(RayTracingGroupId) : Experimental::FHashElementId();
+		if (GroupId.IsValid())
+		{
+			FRayTracingCullingGroup& Group = PrimitiveRayTracingGroups.GetByElementId(GroupId).Value;
+			GroupsToUpdate.FindOrAdd(&Group);
+		}
+	}
+
+	UpdateRayTracingGroupBounds(GroupsToUpdate);
+}
+#endif
+
 void FScene::UpdateAllPrimitiveSceneInfos(FRDGBuilder& GraphBuilder, bool bAsyncCreateLPIs)
 {
 	TRACE_CPUPROFILER_EVENT_SCOPE(Scene::UpdateAllPrimitiveSceneInfos);
@@ -4855,9 +4944,7 @@
 		}
 #if RHI_RAYTRACING
 		{
-			TSet<FPrimitiveSceneProxy*> UpdatedPrimitives;
-			UpdatedTransforms.GetKeys(UpdatedPrimitives);
-			UpdateRayTracingGroupBounds_UpdatePrimitives(UpdatedPrimitives);
+			UpdateRayTracingGroupBounds_UpdatePrimitives(UpdatedTransforms);
 		}
 #endif
 
@@ -4969,9 +5056,7 @@
 
 #if RHI_RAYTRACING
 		{
-			TSet<FPrimitiveSceneProxy*> UpdatedPrimitives;
-			UpdatedInstances.GetKeys(UpdatedPrimitives);
-			UpdateRayTracingGroupBounds_UpdatePrimitives(UpdatedPrimitives);
+			UpdateRayTracingGroupBounds_UpdatePrimitives(UpdatedInstances);
 		}
 #endif
 
@@ -5076,7 +5161,7 @@
 	
 	UpdatedAttachmentRoots.Empty();
 	UpdatedTransforms.Empty();
-	UpdatedInstances.Reset();
+	UpdatedInstances.Empty();
 	UpdatedCustomPrimitiveParams.Empty();
 	OverridenPreviousTransforms.Empty();
 	UpdatedOcclusionBoundsSlacks.Empty();
@@ -5110,93 +5195,6 @@
 	return RemovedPrimitiveSceneInfos.Find(PrimitiveSceneInfo) != nullptr;
 }
 
-#if RHI_RAYTRACING
-void FScene::UpdateRayTracingGroupBounds_AddPrimitives(TSet<FPrimitiveSceneInfo*>& PrimitiveSceneInfos)
-{
-	for (FPrimitiveSceneInfo* const PrimitiveSceneInfo : PrimitiveSceneInfos)
-	{
-		const int32 GroupId = PrimitiveSceneInfo->Proxy->GetRayTracingGroupId();
-		if (GroupId != -1)
-		{
-			bool bInMap = false;
-			static const FRayTracingCullingGroup DefaultGroup;
-			FRayTracingCullingGroup* const Group = PrimitiveRayTracingGroups.FindOrAdd(GroupId, DefaultGroup, bInMap);
-			if (bInMap)
-			{
-				Group->Bounds = Group->Bounds + PrimitiveSceneInfo->Proxy->GetBounds();
-			}
-			else
-			{
-				Group->Bounds = PrimitiveSceneInfo->Proxy->GetBounds();
-			}
-			Group->Primitives.Add(PrimitiveSceneInfo);
-		}
-	}
-}
-
-static void UpdateRayTracingGroupBounds(Experimental::TRobinHoodHashSet<FScene::FRayTracingCullingGroup*>& GroupsToUpdate)
-{
-	for (FScene::FRayTracingCullingGroup* const Group : GroupsToUpdate)
-	{
-		bool bFirstBounds = false;
-		for (FPrimitiveSceneInfo* const Primitive : Group->Primitives)
-		{
-			if (!bFirstBounds)
-			{
-				Group->Bounds = Primitive->Proxy->GetBounds();
-				bFirstBounds = true;
-			}
-			else
-			{
-				Group->Bounds = Group->Bounds + Primitive->Proxy->GetBounds();
-			}
-		}
-	}
-}
-
-void FScene::UpdateRayTracingGroupBounds_RemovePrimitives(TSet<FPrimitiveSceneInfo*>& PrimitiveSceneInfos)
-{
-	Experimental::TRobinHoodHashSet<FRayTracingCullingGroup*> GroupsToUpdate;
-	for (FPrimitiveSceneInfo* const PrimitiveSceneInfo : PrimitiveSceneInfos)
-	{
-		const int32 RayTracingGroupId = PrimitiveSceneInfo->Proxy->GetRayTracingGroupId();
-		const Experimental::FHashElementId GroupId = (RayTracingGroupId != -1) ? PrimitiveRayTracingGroups.FindId(RayTracingGroupId) : Experimental::FHashElementId();
-		if (GroupId.IsValid())
-		{
-			FRayTracingCullingGroup& Group = PrimitiveRayTracingGroups.GetByElementId(GroupId).Value;
-			Group.Primitives.RemoveSingleSwap(PrimitiveSceneInfo);
-			if (Group.Primitives.Num() == 0)
-			{
-				PrimitiveRayTracingGroups.RemoveByElementId(GroupId);
-			}
-			else
-			{
-				GroupsToUpdate.FindOrAdd(&Group);
-			}
-		}
-	}
-
-	UpdateRayTracingGroupBounds(GroupsToUpdate);
-}
-
-void FScene::UpdateRayTracingGroupBounds_UpdatePrimitives(TSet<FPrimitiveSceneProxy*>& PrimitiveProxies)
-{
-	Experimental::TRobinHoodHashSet<FRayTracingCullingGroup*> GroupsToUpdate;
-	for (FPrimitiveSceneProxy* const PrimitiveSceneProxy : PrimitiveProxies)
-	{
-		const int32 RayTracingGroupId = PrimitiveSceneProxy->GetRayTracingGroupId();
-		const Experimental::FHashElementId GroupId = (RayTracingGroupId != -1) ? PrimitiveRayTracingGroups.FindId(RayTracingGroupId) : Experimental::FHashElementId();
-		if (GroupId.IsValid())
-		{
-			FRayTracingCullingGroup& Group = PrimitiveRayTracingGroups.GetByElementId(GroupId).Value;
-			GroupsToUpdate.FindOrAdd(&Group);
-		}
-	}
-
-	UpdateRayTracingGroupBounds(GroupsToUpdate);
-}
-#endif
-
 /**
  * Dummy NULL scene interface used by dedicated servers.
  */
diff --git a/Engine/Source/Runtime/Renderer/Private/ScenePrivate.h b/Engine/Source/Runtime/Renderer/Private/ScenePrivate.h
--- a/Engine/Source/Runtime/Renderer/Private/ScenePrivate.h
+++ b/Engine/Source/Runtime/Renderer/Private/ScenePrivate.h
@@ -3474,13 +3474,14 @@ private:
 	void ProcessAtmosphereLightRemoval_RenderThread(FLightSceneInfo* LightSceneInfo);
 	void ProcessAtmosphereLightAddition_RenderThread(FLightSceneInfo* LightSceneInfo);
 
+private:
 #if RHI_RAYTRACING
-	void UpdateRayTracingGroupBounds_AddPrimitives(TSet<FPrimitiveSceneInfo*>& PrimitiveSceneInfos);
-	void UpdateRayTracingGroupBounds_RemovePrimitives(TSet<FPrimitiveSceneInfo*>& PrimitiveSceneInfos);
-	void UpdateRayTracingGroupBounds_UpdatePrimitives(TSet<FPrimitiveSceneProxy*>& PrimitiveProxies);
+	void UpdateRayTracingGroupBounds_AddPrimitives(const Experimental::TRobinHoodHashSet<FPrimitiveSceneInfo*>& PrimitiveSceneInfos);
+	void UpdateRayTracingGroupBounds_RemovePrimitives(const Experimental::TRobinHoodHashSet<FPrimitiveSceneInfo*>& PrimitiveSceneInfos);
+	template<typename ValueType>
+	inline void UpdateRayTracingGroupBounds_UpdatePrimitives(const Experimental::TRobinHoodHashMap<FPrimitiveSceneProxy*, ValueType>& UpdatedTransforms);
 #endif
 
-private:
 	struct FUpdateTransformCommand
 	{
 		FBoxSphereBounds WorldBounds;
@@ -3501,7 +3502,7 @@ private:
 	Experimental::TRobinHoodHashMap<FPrimitiveSceneInfo*, FPrimitiveComponentId> UpdatedAttachmentRoots;
 	Experimental::TRobinHoodHashMap<FPrimitiveSceneProxy*, FCustomPrimitiveData> UpdatedCustomPrimitiveParams;
 	Experimental::TRobinHoodHashMap<FPrimitiveSceneProxy*, FUpdateTransformCommand> UpdatedTransforms;
-	TMap<FPrimitiveSceneProxy*, FUpdateInstanceCommand> UpdatedInstances;
+	Experimental::TRobinHoodHashMap<FPrimitiveSceneProxy*, FUpdateInstanceCommand> UpdatedInstances;
 	Experimental::TRobinHoodHashMap<FPrimitiveSceneInfo*, FMatrix> OverridenPreviousTransforms;
 	Experimental::TRobinHoodHashMap<const FPrimitiveSceneProxy*, float> UpdatedOcclusionBoundsSlacks;
 	Experimental::TRobinHoodHashSet<FPrimitiveSceneInfo*> AddedPrimitiveSceneInfos;
