Fix shadowing warning
#rnx

diff --git a/Engine/Source/Runtime/Renderer/Private/RendererScene.cpp b/Engine/Source/Runtime/Renderer/Private/RendererScene.cpp
--- a/Engine/Source/Runtime/Renderer/Private/RendererScene.cpp
+++ b/Engine/Source/Runtime/Renderer/Private/RendererScene.cpp
@@ -4258,10 +4258,10 @@ void FScene::UpdateRayTracingGroupBounds_RemovePrimitives(const Experimental::TR
 }
 
 template<typename ValueType>
-inline void FScene::UpdateRayTracingGroupBounds_UpdatePrimitives(const Experimental::TRobinHoodHashMap<FPrimitiveSceneProxy*, ValueType>& UpdatedTransforms)
+inline void FScene::UpdateRayTracingGroupBounds_UpdatePrimitives(const Experimental::TRobinHoodHashMap<FPrimitiveSceneProxy*, ValueType>& InUpdatedTransforms)
 {
 	Experimental::TRobinHoodHashSet<FRayTracingCullingGroup*> GroupsToUpdate;
-	for (const auto& Transform : UpdatedTransforms)
+	for (const auto& Transform : InUpdatedTransforms)
 	{
 		FPrimitiveSceneProxy* const PrimitiveSceneProxy = Transform.Key;
 		const int32 RayTracingGroupId = PrimitiveSceneProxy->GetRayTracingGroupId();
