Fixup new code in light of CL# 18702046
#rnx

diff --git a/Engine/Source/Runtime/Renderer/Private/RendererScene.cpp b/Engine/Source/Runtime/Renderer/Private/RendererScene.cpp
--- a/Engine/Source/Runtime/Renderer/Private/RendererScene.cpp
+++ b/Engine/Source/Runtime/Renderer/Private/RendererScene.cpp
@@ -1613,7 +1613,7 @@ void FScene::UpdatePrimitiveInstances(UInstancedStaticMeshComponent* Primitive)
 			[this, UpdateParams = MoveTemp(UpdateParams)](FRHICommandListImmediate& RHICmdList)
 			{
 #if VALIDATE_PRIMITIVE_PACKED_INDEX
-				if (AddedPrimitiveSceneInfos.Contains(UpdateParams.PrimitiveSceneProxy->GetPrimitiveSceneInfo()))
+				if (AddedPrimitiveSceneInfos.Find(UpdateParams.PrimitiveSceneProxy->GetPrimitiveSceneInfo()) != nullptr)
 				{
 					check(UpdateParams.PrimitiveSceneProxy->GetPrimitiveSceneInfo()->PackedIndex == INDEX_NONE);
 				}
@@ -1622,7 +1622,7 @@ void FScene::UpdatePrimitiveInstances(UInstancedStaticMeshComponent* Primitive)
 					check(UpdateParams.PrimitiveSceneProxy->GetPrimitiveSceneInfo()->PackedIndex != INDEX_NONE);
 				}
 
-				check(!RemovedPrimitiveSceneInfos.Contains(UpdateParams.PrimitiveSceneProxy->GetPrimitiveSceneInfo()));
+				check(RemovedPrimitiveSceneInfos.Find(UpdateParams.PrimitiveSceneProxy->GetPrimitiveSceneInfo()) == nullptr);
 #endif
 				FScopeCycleCounter Context(UpdateParams.PrimitiveSceneProxy->GetStatId());
 				UpdatedInstances.FindOrAdd(UpdateParams.PrimitiveSceneProxy, UpdateParams);
