Localization: Resolving conflicting Ids in LOC macros

diff --git a/Engine/Source/Developer/AnimationDataController/Private/AnimDataController.cpp b/Engine/Source/Developer/AnimationDataController/Private/AnimDataController.cpp
--- a/Engine/Source/Developer/AnimationDataController/Private/AnimDataController.cpp
+++ b/Engine/Source/Developer/AnimationDataController/Private/AnimDataController.cpp
@@ -268,7 +268,7 @@ void UAnimDataController::UpdateCurveNamesFromSkeleton(const USkeleton* Skeleton
 	}
 	else
 	{
-		ReportError(LOCTEXT("InvalidSkeletonError", "Invalid USkeleton supplied"));
+		ReportError(LOCTEXT("UpdateCurveInvalidSkeletonError", "Invalid USkeleton supplied"));
 	}
 }
 
@@ -323,7 +323,7 @@ void UAnimDataController::FindOrAddCurveNamesOnSkeleton(USkeleton* Skeleton, ERa
 	}
 	else
 	{
-		ReportError(LOCTEXT("InvalidSkeletonError", "Invalid USkeleton supplied "));
+		ReportError(LOCTEXT("FindOrAddCurveInvalidSkeletonError", "Invalid USkeleton supplied "));
 	}
 }
 
@@ -721,7 +721,7 @@ bool UAnimDataController::SetCurveFlags(const FAnimationCurveIdentifier& CurveId
 
 	if (Curve)
 	{
-		CONDITIONAL_TRANSACTION(LOCTEXT("SetCurveFlag", "Setting Raw Curve Flags"));
+		CONDITIONAL_TRANSACTION(LOCTEXT("SetRawCurveFlag", "Setting Raw Curve Flags"));
 
 		const int32 CurrentFlags = Curve->GetCurveTypeFlags();
 
diff --git a/Engine/Source/Developer/AnimationDataController/Private/AnimDataControllerActions.cpp b/Engine/Source/Developer/AnimationDataController/Private/AnimDataControllerActions.cpp
--- a/Engine/Source/Developer/AnimationDataController/Private/AnimDataControllerActions.cpp
+++ b/Engine/Source/Developer/AnimationDataController/Private/AnimDataControllerActions.cpp
@@ -270,7 +270,7 @@ TUniquePtr<FChange> FAddFloatCurveAction::ExecuteInternal(UAnimDataModel* Model,
 
 FString FAddFloatCurveAction::ToStringInternal() const
 {
-	return FText::Format(LOCTEXT("AddFloatCurveAction_Description", "Adding float curve '{0}'."), FText::FromName(CurveId.InternalName.DisplayName)).ToString();
+	return FText::Format(LOCTEXT("AddFloatCurveAction_Format", "Adding float curve '{0}'."), FText::FromName(CurveId.InternalName.DisplayName)).ToString();
 }
 
 FAddTransformCurveAction::FAddTransformCurveAction(const FAnimationCurveIdentifier& InCurveId, int32 InFlags, const FTransformCurve& InTransformCurve) : CurveId(InCurveId), Flags(InFlags)
@@ -311,7 +311,7 @@ TUniquePtr<FChange> FAddTransformCurveAction::ExecuteInternal(UAnimDataModel* Mo
 
 FString FAddTransformCurveAction::ToStringInternal() const
 {
-	return FText::Format(LOCTEXT("AddFloatCurveAction_Description", "Adding transform curve '{0}'."), FText::FromName(CurveId.InternalName.DisplayName)).ToString();
+	return FText::Format(LOCTEXT("AddTransformCurveAction_Description", "Adding transform curve '{0}'."), FText::FromName(CurveId.InternalName.DisplayName)).ToString();
 }
 
 TUniquePtr<FChange> FAddRichCurveKeyAction::ExecuteInternal(UAnimDataModel* Model, UAnimDataController* Controller)
diff --git a/Engine/Source/Developer/AssetTools/Private/AssetTypeActions/AssetTypeActions_TextureCube.cpp b/Engine/Source/Developer/AssetTools/Private/AssetTypeActions/AssetTypeActions_TextureCube.cpp
--- a/Engine/Source/Developer/AssetTools/Private/AssetTypeActions/AssetTypeActions_TextureCube.cpp
+++ b/Engine/Source/Developer/AssetTools/Private/AssetTypeActions/AssetTypeActions_TextureCube.cpp
@@ -25,7 +25,7 @@ void FAssetTypeActions_TextureCube::GetActions(const TArray<UObject*>& InObjects
 		Section.AddMenuEntry(
 			"Texture_TextureCuberray",
 			LOCTEXT("Texture_TextureCuberray", "Create Texture Cube Array"),
-			LOCTEXT("Texture_CreateTexture2DArrayTooltip", "Creates a new texture cube array."),
+			LOCTEXT("Texture_CreateTextureCubeArrayTooltip", "Creates a new texture cube array."),
 			FSlateIcon(FEditorStyle::GetStyleSetName(), "ClassIcon.TextureCube"),
 			FUIAction(FExecuteAction::CreateSP(this, &FAssetTypeActions_TextureCube::ExecuteCreateTextureArray, Textures), FCanExecuteAction())
 		);
diff --git a/Engine/Source/Developer/SourceControl/Private/ISourceControlProvider.cpp b/Engine/Source/Developer/SourceControl/Private/ISourceControlProvider.cpp
--- a/Engine/Source/Developer/SourceControl/Private/ISourceControlProvider.cpp
+++ b/Engine/Source/Developer/SourceControl/Private/ISourceControlProvider.cpp
@@ -111,7 +111,7 @@ bool ISourceControlProvider::TryToDownloadFileFromBackgroundThread(const TShared
 	// TryToDownloadFileFromBackgroundThread is unsupported by this source control provider
 	FFormatNamedArguments Arguments;
 	Arguments.Add(TEXT("ProviderName"), FText::FromName(GetName()));
-	FText Message = FText::Format(LOCTEXT("UnsupportedOperation", "TryToDownloadFileFromBackgroundThread is not supported by source control provider '{ProviderName}'"), Arguments);
+	FText Message = FText::Format(LOCTEXT("UnsupportedOperationTryToDownloadFileFromBackgroundThread", "TryToDownloadFileFromBackgroundThread is not supported by source control provider '{ProviderName}'"), Arguments);
 
 	InOperation->AddErrorMessge(Message);
 
@@ -122,7 +122,7 @@ ECommandResult::Type ISourceControlProvider::SwitchWorkspace(FStringView NewWork
 {
 	FFormatNamedArguments Arguments;
 	Arguments.Add(TEXT("ProviderName"), FText::FromName(GetName()));
-	FText Message = FText::Format(LOCTEXT("UnsupportedOperation", "SwitchWorkspace is not supported by source control provider '{ProviderName}'"), Arguments);
+	FText Message = FText::Format(LOCTEXT("UnsupportedOperationSwitchWorkspace", "SwitchWorkspace is not supported by source control provider '{ProviderName}'"), Arguments);
 
 	OutResultInfo.ErrorMessages.Add(Message);
 
diff --git a/Engine/Source/Developer/TraceInsights/Private/Insights/MemoryProfiler/Widgets/SMemAllocTableTreeView.cpp b/Engine/Source/Developer/TraceInsights/Private/Insights/MemoryProfiler/Widgets/SMemAllocTableTreeView.cpp
--- a/Engine/Source/Developer/TraceInsights/Private/Insights/MemoryProfiler/Widgets/SMemAllocTableTreeView.cpp
+++ b/Engine/Source/Developer/TraceInsights/Private/Insights/MemoryProfiler/Widgets/SMemAllocTableTreeView.cpp
@@ -957,7 +957,7 @@ END_SLATE_FUNCTION_BUILD_OPTIMIZATION
 
 FText SMemAllocTableTreeView::ViewPreset_GetSelectedText() const
 {
-	return SelectedViewPreset ? SelectedViewPreset->GetName() : LOCTEXT("CustomPreset_ToolTip", "Custom");
+	return SelectedViewPreset ? SelectedViewPreset->GetName() : LOCTEXT("Custom_ToolTip", "Custom");
 }
 
 ////////////////////////////////////////////////////////////////////////////////////////////////////
@@ -1092,16 +1092,16 @@ void SMemAllocTableTreeView::UpdateQueryInfo()
 		switch (NumMarkers)
 		{
 		case 1:
-			TimeMarkersText = FText::Format(LOCTEXT("TimeMarkersFmt", "A={0}"), FText::AsNumber(TimeMarkers[0]));
+			TimeMarkersText = FText::Format(LOCTEXT("OneTimeMarkersFmt", "A={0}"), FText::AsNumber(TimeMarkers[0]));
 			break;
 		case 2:
-			TimeMarkersText = FText::Format(LOCTEXT("TimeMarkersFmt", "A={0}  B={1}"), FText::AsNumber(TimeMarkers[0]), FText::AsNumber(TimeMarkers[1]));
+			TimeMarkersText = FText::Format(LOCTEXT("TwoTimeMarkersFmt", "A={0}  B={1}"), FText::AsNumber(TimeMarkers[0]), FText::AsNumber(TimeMarkers[1]));
 			break;
 		case 3:
-			TimeMarkersText = FText::Format(LOCTEXT("TimeMarkersFmt", "A={0}  B={1}  C={2}"), FText::AsNumber(TimeMarkers[0]), FText::AsNumber(TimeMarkers[1]), FText::AsNumber(TimeMarkers[2]));
+			TimeMarkersText = FText::Format(LOCTEXT("ThreeTimeMarkersFmt", "A={0}  B={1}  C={2}"), FText::AsNumber(TimeMarkers[0]), FText::AsNumber(TimeMarkers[1]), FText::AsNumber(TimeMarkers[2]));
 			break;
 		case 4:
-			TimeMarkersText = FText::Format(LOCTEXT("TimeMarkersFmt", "A={0}  B={1}  C={2}  D={3}"), FText::AsNumber(TimeMarkers[0]), FText::AsNumber(TimeMarkers[1]), FText::AsNumber(TimeMarkers[2]), FText::AsNumber(TimeMarkers[3]));
+			TimeMarkersText = FText::Format(LOCTEXT("FourTimeMarkersFmt", "A={0}  B={1}  C={2}  D={3}"), FText::AsNumber(TimeMarkers[0]), FText::AsNumber(TimeMarkers[1]), FText::AsNumber(TimeMarkers[2]), FText::AsNumber(TimeMarkers[3]));
 			break;
 		default:
 			// Unhandled value
@@ -1133,7 +1133,7 @@ void SMemAllocTableTreeView::AddCustomAdvancedFilters()
 {
 	TSharedPtr<TArray<TSharedPtr<struct FFilter>>>& AvailableFilters = FilterConfigurator->GetAvailableFilters();
 
-	AvailableFilters->Add(MakeShared<FFilter>(FullCallStackIndex, LOCTEXT("FullCallstack", "Full Callstack"), LOCTEXT("FullCallstack", "Search in all the callstack frames"), EFilterDataType::String, FFilterService::Get()->GetStringOperators()));
+	AvailableFilters->Add(MakeShared<FFilter>(FullCallStackIndex, LOCTEXT("FullCallstack", "Full Callstack"), LOCTEXT("SearchFullCallstack", "Search in all the callstack frames"), EFilterDataType::String, FFilterService::Get()->GetStringOperators()));
 	Context.AddFilterData<FString>(FullCallStackIndex, FString());
 }
 
diff --git a/Engine/Source/Runtime/ApplicationCore/Private/Mac/MacPlatformSplash.cpp b/Engine/Source/Runtime/ApplicationCore/Private/Mac/MacPlatformSplash.cpp
--- a/Engine/Source/Runtime/ApplicationCore/Private/Mac/MacPlatformSplash.cpp
+++ b/Engine/Source/Runtime/ApplicationCore/Private/Mac/MacPlatformSplash.cpp
@@ -222,7 +222,7 @@ void FMacPlatformSplash::Show()
 					}
 					else
 					{
-						VersionInfo = FText::Format( NSLOCTEXT( "UnrealEd", "UnrealEdTitleWithVersion_F", "Unreal Editor {0}  -  {1}" ), Version, GameName );
+						VersionInfo = FText::Format( NSLOCTEXT( "UnrealEd", "UnrealEdTitleWithVersionWithGameName_F", "Unreal Editor {0}  -  {1}" ), Version, GameName );
 						AppName = FText::Format( NSLOCTEXT( "UnrealEd", "UnrealEdTitle_F", "Unreal Editor - {0}" ), GameName );
 					}
 
diff --git a/Engine/Source/Runtime/Core/Private/Logging/MessageLog.cpp b/Engine/Source/Runtime/Core/Private/Logging/MessageLog.cpp
--- a/Engine/Source/Runtime/Core/Private/Logging/MessageLog.cpp
+++ b/Engine/Source/Runtime/Core/Private/Logging/MessageLog.cpp
@@ -43,14 +43,14 @@ public:
 	{
 		FFormatNamedArguments Arguments;
 		Arguments.Add(TEXT("PageTitle"), Title);
-		FMsg::Logf(__FILE__, __LINE__, LogName, ELogVerbosity::Log, TEXT("%s"), *FText::Format(LOCTEXT("BasicMessageLog_SetCurrentPage", "Set Current Page: {PageTitle}"), Arguments).ToString());
+		FMsg::Logf(__FILE__, __LINE__, LogName, ELogVerbosity::Log, TEXT("%s"), *FText::Format(LOCTEXT("BasicMessageLog_SetCurrentPageTitle", "Set Current Page: {PageTitle}"), Arguments).ToString());
 	}
 
 	virtual void SetCurrentPage( const uint32 PageIndex) override
 	{
 		FFormatNamedArguments Arguments;
 		Arguments.Add(TEXT("PageIndex"), PageIndex);
-		FMsg::Logf(__FILE__, __LINE__, LogName, ELogVerbosity::Log, TEXT("%s"), *FText::Format(LOCTEXT("BasicMessageLog_SetCurrentPage", "Set Current Page: {PageIndex}"), Arguments).ToString());
+		FMsg::Logf(__FILE__, __LINE__, LogName, ELogVerbosity::Log, TEXT("%s"), *FText::Format(LOCTEXT("BasicMessageLog_SetCurrentPageIndex", "Set Current Page: {PageIndex}"), Arguments).ToString());
 	}
 
 	virtual void NotifyIfAnyMessages( const FText& Message, EMessageSeverity::Type SeverityFilter, bool bForce ) override
diff --git a/Engine/Source/Runtime/CoreUObject/Private/Misc/PackageName.cpp b/Engine/Source/Runtime/CoreUObject/Private/Misc/PackageName.cpp
--- a/Engine/Source/Runtime/CoreUObject/Private/Misc/PackageName.cpp
+++ b/Engine/Source/Runtime/CoreUObject/Private/Misc/PackageName.cpp
@@ -2487,7 +2487,7 @@ FText FPackageName::FormatErrorAsText(FStringView InPath, EErrorCode ErrorCode)
 		return FText::Format(NSLOCTEXT("Core", "LongPackageNames_PathWithTrailingSlash", "Input '{InPath}' ends with a '/', which is invalid for LongPackageNames."), Args);
 	default:
 		UE_LOG(LogPackageName, Warning, TEXT("FPackageName::FormatErrorAsText: Invalid ErrorCode %d"), static_cast<int32>(ErrorCode));
-		return FText::Format(NSLOCTEXT("Core", "PackageNameUnknownError", "Input '{InPath} caused undocumented internal error."), Args);
+		return FText::Format(NSLOCTEXT("Core", "PackageNameUndocumentedError", "Input '{InPath} caused undocumented internal error."), Args);
 	}
 }
 
diff --git a/Engine/Source/Runtime/CoreUObject/Private/UObject/Linker.cpp b/Engine/Source/Runtime/CoreUObject/Private/UObject/Linker.cpp
--- a/Engine/Source/Runtime/CoreUObject/Private/UObject/Linker.cpp
+++ b/Engine/Source/Runtime/CoreUObject/Private/UObject/Linker.cpp
@@ -759,7 +759,7 @@ FLinkerLoad* GetPackageLinker
 			PackageFName = FName(*PackageName);
 			if (!FPackagePath::TryFromMountedName(PackageName, PackagePath))
 			{
-				LogGetPackageLinkerError(InExistingContext, InPackagePath, LOCTEXT("InvalidPackageRedirect", "GetDelegateResolvedPackagePath or GetLocalizedPackagePath returned path to unmounted contentroot"), InOuter, LoadFlags);
+				LogGetPackageLinkerError(InExistingContext, InPackagePath, LOCTEXT("UnmountedContentRoot", "GetDelegateResolvedPackagePath or GetLocalizedPackagePath returned path to unmounted contentroot"), InOuter, LoadFlags);
 				return nullptr;
 			}
 		}
diff --git a/Engine/Source/Runtime/Engine/Private/Animation/BlendSpace.cpp b/Engine/Source/Runtime/Engine/Private/Animation/BlendSpace.cpp
--- a/Engine/Source/Runtime/Engine/Private/Animation/BlendSpace.cpp
+++ b/Engine/Source/Runtime/Engine/Private/Animation/BlendSpace.cpp
@@ -1188,32 +1188,33 @@ void UBlendSpace::ValidateSampleData()
 				else
 				{
 					static FName NAME_LoadErrors("LoadErrors");
+					static const FText HasSampleText = LOCTEXT("EmptyAnimationData_HasSampleShort", "has sample");
 					FMessageLog LoadErrors(NAME_LoadErrors);
 
 					TSharedRef<FTokenizedMessage> Message = LoadErrors.Error();
-					Message->AddToken(FTextToken::Create(LOCTEXT("EmptyAnimationData1", "BlendSpace")));
+					Message->AddToken(FTextToken::Create(LOCTEXT("EmptyAnimationData_BlendSpace", "BlendSpace")));
 					Message->AddToken(FAssetNameToken::Create(GetPathName(), FText::FromString(GetName())));
 					if (!bAnimationExists)
 					{
-						Message->AddToken(FTextToken::Create(LOCTEXT("EmptyAnimationData2", "has a sample with no/invalid animation. Recommend to remove sample point or set new animation.")));
+						Message->AddToken(FTextToken::Create(LOCTEXT("EmptyAnimationData_HasSample", "has a sample with no/invalid animation. Recommend to remove sample point or set new animation.")));
 					}
 					else if (!bSampleInBounds)
 					{
-						Message->AddToken(FTextToken::Create(LOCTEXT("EmptyAnimationData2", "has sample")));
+						Message->AddToken(FTextToken::Create(HasSampleText));
 						Message->AddToken(FAssetNameToken::Create(Sample.Animation->GetPathName(), FText::FromString(Sample.Animation->GetName())));
-						Message->AddToken(FTextToken::Create(LOCTEXT("EmptyAnimationData3", "that is invalid due to being out of bounds. Recommend adjusting it.")));
+						Message->AddToken(FTextToken::Create(LOCTEXT("EmptyAnimationData_SampleNotInBounds", "that is invalid due to being out of bounds. Recommend adjusting it.")));
 					}
 					else if (!bSampleIsUnique)
 					{
-						Message->AddToken(FTextToken::Create(LOCTEXT("EmptyAnimationData2", "has sample")));
+						Message->AddToken(FTextToken::Create(HasSampleText));
 						Message->AddToken(FAssetNameToken::Create(Sample.Animation->GetPathName(), FText::FromString(Sample.Animation->GetName())));
-						Message->AddToken(FTextToken::Create(LOCTEXT("EmptyAnimationData3", "that is not unique. Recommend adjusting it.")));
+						Message->AddToken(FTextToken::Create(LOCTEXT("EmptyAnimationData_NotUnique", "that is not unique. Recommend adjusting it.")));
 					}
 					else
 					{
 						// Shouldn't get here
 						Message->AddToken(FAssetNameToken::Create(Sample.Animation->GetPathName(), FText::FromString(Sample.Animation->GetName())));
-						Message->AddToken(FTextToken::Create(LOCTEXT("EmptyAnimationData2", "is invalid.")));
+						Message->AddToken(FTextToken::Create(LOCTEXT("EmptyAnimationData_IsInvalid", "is invalid.")));
 					}
 					LoadErrors.Notify();
 				}
diff --git a/Engine/Source/Runtime/Engine/Private/StaticMesh.cpp b/Engine/Source/Runtime/Engine/Private/StaticMesh.cpp
--- a/Engine/Source/Runtime/Engine/Private/StaticMesh.cpp
+++ b/Engine/Source/Runtime/Engine/Private/StaticMesh.cpp
@@ -5684,7 +5684,7 @@ void UStaticMesh::ExecutePostLoadInternal(FStaticMeshPostLoadContext& Context)
 					{
 						FMessageLog("LoadErrors").Warning()
 							->AddToken(Token)
-							->AddToken(FTextToken::Create(FText::Format(LOCTEXT("LoadError_BadMinLOD", "Min LOD value of {DefaultMinLOD} is out of range 0..{MinAvailLOD} and has been adjusted to {MinAvailLOD}. Please verify and resave the asset."), Arguments)));
+							->AddToken(FTextToken::Create(FText::Format(LOCTEXT("LoadError_DefaultMinLODOutOfRange", "Min LOD value of {DefaultMinLOD} is out of range 0..{MinAvailLOD} and has been adjusted to {MinAvailLOD}. Please verify and resave the asset."), Arguments)));
 					}	
 				);
 				QualityLocalMinLOD.Default = MinAvailableLOD;
@@ -5707,7 +5707,7 @@ void UStaticMesh::ExecutePostLoadInternal(FStaticMeshPostLoadContext& Context)
 						{
 							FMessageLog("LoadErrors").Warning()
 								->AddToken(Token)
-								->AddToken(FTextToken::Create(FText::Format(LOCTEXT("LoadError_BadMinLODOverride", "Min LOD override of {QualityLevelMinLOD} for {QualityLevel} is out of range 0..{MinAvailLOD} and has been adjusted to {MinAvailLOD}. Please verify and resave the asset."), Arguments)));
+								->AddToken(FTextToken::Create(FText::Format(LOCTEXT("LoadError_MinLODOverrideForQualityLevel", "Min LOD override of {QualityLevelMinLOD} for {QualityLevel} is out of range 0..{MinAvailLOD} and has been adjusted to {MinAvailLOD}. Please verify and resave the asset."), Arguments)));
 						}
 					);
 					It.Value() = MinAvailableLOD;
diff --git a/Engine/Source/Runtime/Renderer/Private/SceneRendering.cpp b/Engine/Source/Runtime/Renderer/Private/SceneRendering.cpp
--- a/Engine/Source/Runtime/Renderer/Private/SceneRendering.cpp
+++ b/Engine/Source/Runtime/Renderer/Private/SceneRendering.cpp
@@ -3306,7 +3306,7 @@ void FSceneRenderer::RenderFinish(FRDGBuilder& GraphBuilder, FRDGTextureRef View
 						}
 						if (bShowSkylightWarning)
 						{
-							static const FText Message = NSLOCTEXT("Renderer", "Skylight", "PROJECT DOES NOT SUPPORT STATIONARY SKYLIGHT: ");
+							static const FText Message = NSLOCTEXT("Renderer", "SkylightNotSuppported", "PROJECT DOES NOT SUPPORT STATIONARY SKYLIGHT: ");
 							Canvas.DrawShadowedText(10, Y, Message, GetStatsFont(), FLinearColor(1.0, 0.05, 0.05, 1.0));
 							Y += 14;
 						}
@@ -3403,18 +3403,14 @@ void FSceneRenderer::RenderFinish(FRDGBuilder& GraphBuilder, FRDGTextureRef View
 
 						if (bLumenEnabledButHasNoDataForTracing)
 						{
-							static const FText Message = NSLOCTEXT("Renderer", "LumenCantDisplay", 
-								"Lumen is enabled, but has no ray tracing data and won't operate correctly.\n"
-								"Either configure Lumen to use software distance field ray tracing and enable 'Generate Mesh Distancefields' in project settings\n"
-								"or configure Lumen to use Hardware Ray Tracing and enable 'Support Hardware Ray Tracing' in project settings.");
+							static const FText Message = NSLOCTEXT("Renderer", "LumenCantDisplay", "Lumen is enabled, but has no ray tracing data and won't operate correctly.\nEither configure Lumen to use software distance field ray tracing and enable 'Generate Mesh Distancefields' in project settings\nor configure Lumen to use Hardware Ray Tracing and enable 'Support Hardware Ray Tracing' in project settings.");
 							Canvas.DrawShadowedText(10, Y, Message, GetStatsFont(), FLinearColor(1.0, 0.05, 0.05, 1.0));
 							Y += 14;
 						}
 
 						if (bLumenEnabledButDisabledForTheProject)
 						{
-							static const FText Message = NSLOCTEXT("Renderer", "LumenDisabledForProject", 
-								"Lumen is enabled but cannot render, because the project has Lumen disabled in an ini (r.Lumen.Supported = 0)");
+							static const FText Message = NSLOCTEXT("Renderer", "LumenDisabledForProject", "Lumen is enabled but cannot render, because the project has Lumen disabled in an ini (r.Lumen.Supported = 0)");
 							Canvas.DrawShadowedText(10, Y, Message, GetStatsFont(), FLinearColor(1.0, 0.05, 0.05, 1.0));
 							Y += 14;
 						}
@@ -3441,8 +3437,7 @@ void FSceneRenderer::RenderFinish(FRDGBuilder& GraphBuilder, FRDGTextureRef View
 
 						if (bNaniteEnabledButDisabledInProject)
 						{
-							static const FText Message = NSLOCTEXT("Renderer", "NaniteDisabledForProject",
-								"Nanite is enabled but cannot render, because the project has Nanite disabled in an ini (r.Nanite.ProjectEnabled = 0)");
+							static const FText Message = NSLOCTEXT("Renderer", "NaniteDisabledForProject", "Nanite is enabled but cannot render, because the project has Nanite disabled in an ini (r.Nanite.ProjectEnabled = 0)");
 							Canvas.DrawShadowedText(10, Y, Message, GetStatsFont(), FLinearColor(1.0, 0.05, 0.05, 1.0));
 							Y += 14;
 						}
