Fix for RDG validation assert when enabling Screen Space GI.  Need to pull the texture reference for temporal effect handling from the UAV, since that's flagged as accessible, rather than the texture itself.

#rnx

diff --git a/Engine/Source/Runtime/Renderer/Private/ScreenSpaceDenoise.cpp b/Engine/Source/Runtime/Renderer/Private/ScreenSpaceDenoise.cpp
--- a/Engine/Source/Runtime/Renderer/Private/ScreenSpaceDenoise.cpp
+++ b/Engine/Source/Runtime/Renderer/Private/ScreenSpaceDenoise.cpp
@@ -2059,14 +2059,6 @@ static void DenoiseSignalAtConstantPixelDensity(
 #if WITH_MGPU
 		{
 			FName EffectName = Settings.EffectName;
-			TArray<FRDGTexture*, TFixedAllocator<4>> SignalOutputTextures;
-			for (FRDGTexture* Texture : SignalOutput.Textures)
-			{
-				if (Texture)
-				{
-					SignalOutputTextures.Add(Texture);
-				}
-			}
 
 			FIntVector GroupCount = FComputeShaderUtils::GetGroupCount(Viewport.Size(), FComputeShaderUtils::kGolden2DGroupSize);
 			FComputeShaderUtils::ValidateGroupCount(GroupCount);
@@ -2079,16 +2071,19 @@ static void DenoiseSignalAtConstantPixelDensity(
 				ParametersMetadata,
 				PassParameters,
 				ERDGPassFlags::Compute,
-				[ParametersMetadata, PassParameters, ComputeShader, GroupCount, EffectName, SignalOutputTextures](FRHIComputeCommandList& RHICmdList)
+				[ParametersMetadata, PassParameters, ComputeShader, GroupCount, EffectName](FRHIComputeCommandList& RHICmdList)
 				{
 					RHICmdList.WaitForTemporalEffect(EffectName);
 
 					FComputeShaderUtils::Dispatch(RHICmdList, ComputeShader, ParametersMetadata, *PassParameters, GroupCount);
 
-					TArray<FRHITexture*, TFixedAllocator<4>> SignalOutputTexturesRHI;
-					for (FRDGTexture* Texture : SignalOutputTextures)
+					TArray<FRHITexture*, TFixedAllocator<kMaxBufferProcessingCount>> SignalOutputTexturesRHI;
+					for (FRDGTextureUAV* TextureUAV : PassParameters->SignalHistoryOutput.UAVs)
 					{
-						SignalOutputTexturesRHI.Add(Texture->GetRHI());
+						if (TextureUAV)
+						{
+							SignalOutputTexturesRHI.Add(TextureUAV->GetParentRHI());
+						}
 					}
 
 					RHICmdList.BroadcastTemporalEffect(
