Fix incorrect shadow rendering viewport for ISR

diff --git a/Engine/Source/Runtime/Renderer/Private/SceneRendering.cpp b/Engine/Source/Runtime/Renderer/Private/SceneRendering.cpp
--- a/Engine/Source/Runtime/Renderer/Private/SceneRendering.cpp
+++ b/Engine/Source/Runtime/Renderer/Private/SceneRendering.cpp
@@ -2030,7 +2030,8 @@ const FViewInfo* FViewInfo::GetPrimaryView() const
 
 const FViewInfo* FViewInfo::GetInstancedView() const
 {
-	if ((IsInstancedStereoPass() || bIsMobileMultiViewEnabled))
+	// Extra checks are needed because some code relies on this function to return NULL if ISR is disabled.
+	if (bIsInstancedStereoEnabled || bIsMobileMultiViewEnabled)
 	{
 		return static_cast<const FViewInfo*>(GetInstancedSceneView());
 	}
diff --git a/Engine/Source/Runtime/Renderer/Private/ShadowRendering.cpp b/Engine/Source/Runtime/Renderer/Private/ShadowRendering.cpp
--- a/Engine/Source/Runtime/Renderer/Private/ShadowRendering.cpp
+++ b/Engine/Source/Runtime/Renderer/Private/ShadowRendering.cpp
@@ -910,7 +910,9 @@ void FProjectedShadowInfo::SetupProjectionStencilMask(
 		{
 			ensure(ProjectionStencilingPasses[View->PrimaryViewIndex]->GetInstanceCullingMode() == EInstanceCullingMode::Stereo);
 
-			RHICmdList.SetViewport(View->ViewRect.Min.X, View->ViewRect.Min.Y, 0.0f, SceneRender->InstancedStereoWidth, View->ViewRect.Max.Y, 1.0f);
+			const FViewInfo* PrimaryView = View->GetPrimaryView();
+			const FViewInfo* InstancedView = View->GetInstancedView();
+			RHICmdList.SetViewport(PrimaryView->ViewRect.Min.X, PrimaryView->ViewRect.Min.Y, 0.0f, InstancedView->ViewRect.Max.X, InstancedView->ViewRect.Max.Y, 1.0f);
 			RHICmdList.SetScissorRect(true, View->ViewRect.Min.X, View->ViewRect.Min.Y, View->ViewRect.Max.X, View->ViewRect.Max.Y);
 			// Submit the first (and only pass - we share that at least) as the pass is set up for stereo.
 			ProjectionStencilingPasses[View->PrimaryViewIndex]->SubmitDraw(RHICmdList, InstanceCullingDrawParams);
