Uses BloomIntensity for FFT bloom and allows dirt mask to work

diff --git a/Engine/Source/Runtime/Core/Public/UObject/UE5ReleaseStreamObjectVersion.h b/Engine/Source/Runtime/Core/Public/UObject/UE5ReleaseStreamObjectVersion.h
--- a/Engine/Source/Runtime/Core/Public/UObject/UE5ReleaseStreamObjectVersion.h
+++ b/Engine/Source/Runtime/Core/Public/UObject/UE5ReleaseStreamObjectVersion.h
@@ -87,6 +87,9 @@ struct CORE_API FUE5ReleaseStreamObjectVersion
 		// Pose asset GUID generated from source AnimationSequence
 		PoseAssetRawDataGUID,
 
+		// Convolution bloom now take into account FPostProcessSettings::BloomIntensity for scatter dispersion.
+		ConvolutionBloomIntensity,
+
 		// -----<new versions can be added above this line>-------------------------------------------------
 		VersionPlusOne,
 		LatestVersion = VersionPlusOne - 1
diff --git a/Engine/Source/Runtime/Engine/Private/PostProcessVolume.cpp b/Engine/Source/Runtime/Engine/Private/PostProcessVolume.cpp
--- a/Engine/Source/Runtime/Engine/Private/PostProcessVolume.cpp
+++ b/Engine/Source/Runtime/Engine/Private/PostProcessVolume.cpp
@@ -165,7 +165,6 @@ bool APostProcessVolume::CanEditChange(const FProperty* InProperty) const
 
 			// Parameters that are only used for the Sum of Gaussian bloom / not the texture based fft bloom
 			if (PropertyName == GET_MEMBER_NAME_STRING_CHECKED(FPostProcessSettings, BloomThreshold) ||
-				PropertyName == GET_MEMBER_NAME_STRING_CHECKED(FPostProcessSettings, BloomIntensity) ||
 				PropertyName == GET_MEMBER_NAME_STRING_CHECKED(FPostProcessSettings, BloomSizeScale) ||
 				PropertyName == GET_MEMBER_NAME_STRING_CHECKED(FPostProcessSettings, Bloom1Size) ||
 				PropertyName == GET_MEMBER_NAME_STRING_CHECKED(FPostProcessSettings, Bloom2Size) ||
diff --git a/Engine/Source/Runtime/Engine/Private/Scene.cpp b/Engine/Source/Runtime/Engine/Private/Scene.cpp
--- a/Engine/Source/Runtime/Engine/Private/Scene.cpp
+++ b/Engine/Source/Runtime/Engine/Private/Scene.cpp
@@ -1223,6 +1223,16 @@ void FPostProcessSettings::PostSerialize(const FArchive& Ar)
 				}
 			}
 		}
+
+		// Before, the convolution bloom would ignore the bloom intensity, but activate when BloomIntensity > 0.0. UE 5.0 changed so that the 
+		// BloomIntensity also controls the scatter dispersion of the convolution bloom.
+		if (Ar.CustomVer(FUE5ReleaseStreamObjectVersion::GUID) < FUE5ReleaseStreamObjectVersion::ConvolutionBloomIntensity)
+		{
+			if (BloomMethod == BM_FFT && BloomIntensity > 0.0)
+			{
+				BloomIntensity = 1.0f;
+			}
+		}
 	}
 }
 #endif
diff --git a/Engine/Source/Runtime/Renderer/Private/PostProcess/PostProcessTonemap.cpp b/Engine/Source/Runtime/Renderer/Private/PostProcess/PostProcessTonemap.cpp
--- a/Engine/Source/Runtime/Renderer/Private/PostProcess/PostProcessTonemap.cpp
+++ b/Engine/Source/Runtime/Renderer/Private/PostProcess/PostProcessTonemap.cpp
@@ -792,7 +792,7 @@ FScreenPassTexture AddTonemapPass(FRDGBuilder& GraphBuilder, const FViewInfo& Vi
 	CommonParameters.LensPrincipalPointOffsetScale = View.LensPrincipalPointOffsetScale;
 
 	// TODO: PostProcessSettings.BloomDirtMask->GetResource() is not thread safe
-	if (PostProcessSettings.BloomDirtMask && PostProcessSettings.BloomDirtMask->GetResource() && Inputs.bAllowDirtMask && PostProcessSettings.BloomDirtMaskIntensity > 0 && PostProcessSettings.BloomIntensity > 0.0)
+	if (PostProcessSettings.BloomDirtMask && PostProcessSettings.BloomDirtMask->GetResource() && PostProcessSettings.BloomDirtMaskIntensity > 0 && PostProcessSettings.BloomIntensity > 0.0)
 	{
 		CommonParameters.BloomDirtMaskTint = PostProcessSettings.BloomDirtMaskTint * PostProcessSettings.BloomDirtMaskIntensity / PostProcessSettings.BloomIntensity;
 		CommonParameters.BloomDirtMaskTexture = PostProcessSettings.BloomDirtMask->GetResource()->TextureRHI;
diff --git a/Engine/Source/Runtime/Renderer/Private/PostProcess/PostProcessTonemap.h b/Engine/Source/Runtime/Renderer/Private/PostProcess/PostProcessTonemap.h
--- a/Engine/Source/Runtime/Renderer/Private/PostProcess/PostProcessTonemap.h
+++ b/Engine/Source/Runtime/Renderer/Private/PostProcess/PostProcessTonemap.h
@@ -74,9 +74,6 @@ struct FTonemapInputs
 	// [Optional, ES31] Eye adaptation buffer used to compute exposure. 
 	FRDGBufferRef EyeAdaptationBuffer = nullptr;
 
-	// Weather should allow dirk mask.
-	bool bAllowDirtMask = true;
-
 	// [Raster Only, Mobile] Flips the image vertically on output.
 	bool bFlipYAxis = false;
 
diff --git a/Engine/Source/Runtime/Renderer/Private/PostProcess/PostProcessing.cpp b/Engine/Source/Runtime/Renderer/Private/PostProcess/PostProcessing.cpp
--- a/Engine/Source/Runtime/Renderer/Private/PostProcess/PostProcessing.cpp
+++ b/Engine/Source/Runtime/Renderer/Private/PostProcess/PostProcessing.cpp
@@ -1054,7 +1054,6 @@ void AddPostProcessingPasses(FRDGBuilder& GraphBuilder, const FViewInfo& View, c
 				PassInputs.EyeAdaptationParameters = &EyeAdaptationParameters;
 				PassInputs.EyeAdaptationTexture = EyeAdaptationTexture;
 				PassInputs.ColorGradingTexture = ColorGradingTexture;
-				PassInputs.bAllowDirtMask = !bFFTBloomEnabled;
 				PassInputs.bWriteAlphaChannel = AntiAliasingMethod == AAM_FXAA || IsPostProcessingWithAlphaChannelSupported();
 				PassInputs.bOutputInHDR = bTonemapOutputInHDR;
 
