A file of 18721597 didn't shelve properly for horde submission

diff --git a/Engine/Source/Runtime/Renderer/Private/PostProcess/PostProcessFFTBloom.cpp b/Engine/Source/Runtime/Renderer/Private/PostProcess/PostProcessFFTBloom.cpp
--- a/Engine/Source/Runtime/Renderer/Private/PostProcess/PostProcessFFTBloom.cpp
+++ b/Engine/Source/Runtime/Renderer/Private/PostProcess/PostProcessFFTBloom.cpp
@@ -754,7 +754,7 @@ FFFTBloomOutput AddFFTBloomPass(FRDGBuilder& GraphBuilder, const FViewInfo& View
 			TEXT("Bloom.FFT.FFTMulitplyParameters"));
 
 		FBloomFinalizeApplyConstantsCS::FParameters* PassParameters = GraphBuilder.AllocParameters<FBloomFinalizeApplyConstantsCS::FParameters>();
-		PassParameters->ScatterDispersionIntensity = View.FinalPostProcessSettings.BloomConvolutionScatterDispersion;
+		PassParameters->ScatterDispersionIntensity = View.FinalPostProcessSettings.BloomConvolutionScatterDispersion * View.FinalPostProcessSettings.BloomIntensity;
 		PassParameters->KernelConstantsBuffer = GraphBuilder.CreateSRV(KernelConstantsBuffer);
 		PassParameters->SceneColorApplyOutput = GraphBuilder.CreateUAV(BloomOutput.SceneColorApplyParameters);
 		PassParameters->FFTMulitplyOutput = GraphBuilder.CreateUAV(FFTMulitplyParameters);
