Ensure stencil buffer is always clear for sub-pixel shadow mask (i.e. for hair strands).

diff --git a/Engine/Source/Runtime/Renderer/Private/ShadowRendering.cpp b/Engine/Source/Runtime/Renderer/Private/ShadowRendering.cpp
--- a/Engine/Source/Runtime/Renderer/Private/ShadowRendering.cpp
+++ b/Engine/Source/Runtime/Renderer/Private/ShadowRendering.cpp
@@ -1250,10 +1250,13 @@ void FProjectedShadowInfo::RenderProjectionInternal(
 	}
 
 	// Mark stencil so that only hair pixel within volume bound will be affected by the pre-shadow mask
-	if (bSubPixelSupport && !IsWholeSceneDirectionalShadow())
+	if (bSubPixelSupport)
 	{
 		DrawClearQuad(RHICmdList, false, FLinearColor::Transparent, false, 0, true, 0);
-		SetupProjectionStencilMaskForHair(RHICmdList, View);
+		if (!IsWholeSceneDirectionalShadow())
+		{
+			SetupProjectionStencilMaskForHair(RHICmdList, View);
+		}
 	}
 
 	// solid rasterization w/ back-face culling.
