Fix for broken/shuffled Nanite materials when commands like FUpdateStaticMeshesForMaterials call FPrimitiveSceneInfo::UpdateStaticMeshes(), which recaches Nanite mesh draw commands (and deriving new material table indices), but no other change will cause GPU Scene to perform a per-primitive upload (which would also update the global material tables). This change modifies FPrimitiveSceneInfo::UpdateStaticMeshes() to trigger a GPU Scene update so the MDCs and the global material table stay synchronized. Normally adding a primitive to the scene will mark GPU Scene dirty, but some cases of construction scripts adding to the scene, flushing the upload, but then changing just the static meshes or material assignments wouldn't flush an update prior to this change.

diff --git a/Engine/Source/Runtime/Renderer/Private/PrimitiveSceneInfo.cpp b/Engine/Source/Runtime/Renderer/Private/PrimitiveSceneInfo.cpp
--- a/Engine/Source/Runtime/Renderer/Private/PrimitiveSceneInfo.cpp
+++ b/Engine/Source/Runtime/Renderer/Private/PrimitiveSceneInfo.cpp
@@ -1751,6 +1751,12 @@ void FPrimitiveSceneInfo::UpdateStaticMeshes(FRHICommandListImmediate& RHICmdLis
 			SceneInfo->RemoveCachedRayTracingPrimitives();
 		}
 	#endif
+
+		if (SceneInfo->Proxy && SceneInfo->Proxy->IsNaniteMesh())
+		{
+			// Make sure material table indirections are kept in sync with GPU Scene and cached Nanite MDCs
+			SceneInfo->RequestGPUSceneUpdate(EPrimitiveDirtyState::ChangedOther);
+		}
 	}
 
 	if (bReAddToDrawLists)
