Added an option to apply height and atmospheric fog as a separate pass on mobile (disabled by default)
This helps to cut fog shader code and state setup from most of the draws

diff --git a/Engine/Shaders/Private/HeightFogCommon.ush b/Engine/Shaders/Private/HeightFogCommon.ush
--- a/Engine/Shaders/Private/HeightFogCommon.ush
+++ b/Engine/Shaders/Private/HeightFogCommon.ush
@@ -16,6 +16,14 @@
 #define PERMUTATION_SUPPORT_VOLUMETRIC_FOG (FEATURE_LEVEL >= FEATURE_LEVEL_SM4)
 #endif
 
+#ifndef PERMUTATION_SUPPORT_FOG_START_DISTANCE
+#define PERMUTATION_SUPPORT_FOG_START_DISTANCE 1
+#endif
+
+#ifndef PERMUTATION_SUPPORT_FOG_SECOND_TERM
+#define PERMUTATION_SUPPORT_FOG_SECOND_TERM 1
+#endif
+
 #if INSTANCED_STEREO
 	#if MATERIALBLENDING_ANY_TRANSLUCENT
 		#define FogStructISR TranslucentBasePass.Shared.FogISR
@@ -89,9 +97,6 @@ half4 GetExponentialHeightFog(float3 WorldPositionRelativeToCamera, float Exclud
 	float RayLength = CameraToReceiverLength;
 	float RayDirectionZ = CameraToReceiver.z;
 
-	// Factor in StartDistance
-	ExcludeDistance = max(ExcludeDistance, FogStruct.ExponentialFogParameters.w);
-	
 #if USE_GLOBAL_CLIP_PLANE
 
 	BRANCH
@@ -110,6 +115,9 @@ half4 GetExponentialHeightFog(float3 WorldPositionRelativeToCamera, float Exclud
 
 #endif
 
+	// Factor in StartDistance
+#if PERMUTATION_SUPPORT_FOG_START_DISTANCE
+	ExcludeDistance = max(ExcludeDistance, FogStruct.ExponentialFogParameters.w);
 	if (ExcludeDistance > 0)
 	{
 		float ExcludeIntersectionTime = ExcludeDistance * CameraToReceiverLengthInv;
@@ -127,9 +135,13 @@ half4 GetExponentialHeightFog(float3 WorldPositionRelativeToCamera, float Exclud
 		float ExponentSecond = max(-127.0f, FogStruct.ExponentialFogParameters2.y * (ExclusionIntersectionZ - FogStruct.ExponentialFogParameters2.w)); 	 
 		RayOriginTermsSecond = FogStruct.ExponentialFogParameters2.z * exp2(-ExponentSecond);
 	}
+#endif
 
 	// Calculate the "shared" line integral (this term is also used for the directional light inscattering) by adding the two line integrals together (from two different height falloffs and densities)
-	float ExponentialHeightLineIntegralShared = CalculateLineIntegralShared(FogStruct.ExponentialFogParameters.y, RayDirectionZ, RayOriginTerms) + CalculateLineIntegralShared(FogStruct.ExponentialFogParameters2.y, RayDirectionZ, RayOriginTermsSecond);
+	float ExponentialHeightLineIntegralShared = CalculateLineIntegralShared(FogStruct.ExponentialFogParameters.y, RayDirectionZ, RayOriginTerms);
+#if PERMUTATION_SUPPORT_FOG_SECOND_TERM
+	ExponentialHeightLineIntegralShared+= CalculateLineIntegralShared(FogStruct.ExponentialFogParameters2.y, RayDirectionZ, RayOriginTermsSecond);
+#endif
 
 	float ExponentialHeightLineIntegral = ExponentialHeightLineIntegralShared * RayLength;
 
diff --git a/Engine/Shaders/Private/MobileBasePassCommon.ush b/Engine/Shaders/Private/MobileBasePassCommon.ush
--- a/Engine/Shaders/Private/MobileBasePassCommon.ush
+++ b/Engine/Shaders/Private/MobileBasePassCommon.ush
@@ -1,7 +1,7 @@
 // Copyright Epic Games, Inc. All Rights Reserved.
 
 /*=============================================================================
-	MobileBasePassCommon.usf: Base pass definitions used by both vertex and pixel shader
+	MobileBasePassCommon.ush: Base pass definitions used by both vertex and pixel shader
 =============================================================================*/
 
 #define FogStruct MobileBasePass.Fog
@@ -16,7 +16,9 @@
 	#define PROJECT_MOBILE_DISABLE_VERTEX_FOG 0
 #endif
 
-#define USE_VERTEX_FOG					(!PROJECT_MOBILE_DISABLE_VERTEX_FOG) && MATERIAL_ENABLE_TRANSLUCENCY_FOGGING
+// Use vertex fogging for translucency and opaque if project does not use per-pixel fog
+// And always use vertex fogging on sky for perf
+#define USE_VERTEX_FOG	(!PROJECT_MOBILE_DISABLE_VERTEX_FOG || (MATERIAL_ENABLE_TRANSLUCENCY_FOGGING && MATERIALBLENDING_ANY_TRANSLUCENT) || MATERIAL_IS_SKY)
 
 #define PACK_INTERPOLANTS (USE_VERTEX_FOG && NUM_VF_PACKED_INTERPOLANTS > 0 && (ES3_1_PROFILE))
 #define LANDSCAPE_BUG_WORKAROUND (IOS && IS_MOBILE_BASEPASS_VERTEX_SHADER && PACK_INTERPOLANTS)
diff --git a/Engine/Shaders/Private/MobileBasePassPixelShader.usf b/Engine/Shaders/Private/MobileBasePassPixelShader.usf
--- a/Engine/Shaders/Private/MobileBasePassPixelShader.usf
+++ b/Engine/Shaders/Private/MobileBasePassPixelShader.usf
@@ -106,8 +106,8 @@ int NumDynamicPointLights;
 #define ALLOW_CUBE_REFLECTIONS HQ_REFLECTIONS
 
 #include "SHCommon.ush"
-#include "MobileBasePassCommon.ush"
 #include "/Engine/Generated/Material.ush"
+#include "MobileBasePassCommon.ush"
 #include "/Engine/Generated/VertexFactory.ush"
 #include "ReflectionEnvironmentShared.ush"
 #include "LightmapCommon.ush"  
@@ -244,7 +244,7 @@ half3 GetLookupVectorForSphereCaptureMobile(half3 ReflectionVector, half3 WorldP
 half3 GetMobileSkyLightReflection(half3 ReflectionVector, half Roughness, half CubemapMaxMip)
 {
 	half AbsoluteSpecularMip = ComputeReflectionCaptureMipFromRoughness(Roughness, CubemapMaxMip);
-	half3 Reflection = MobileReflectionCapture.Texture.SampleLevel(MobileReflectionCapture.TextureSampler, ReflectionVector, AbsoluteSpecularMip);
+	half3 Reflection = MobileReflectionCapture.Texture.SampleLevel(MobileReflectionCapture.TextureSampler, ReflectionVector, AbsoluteSpecularMip).rgb;
 	return Reflection * ResolvedView.SkyLightColor.rgb;
 }
 
diff --git a/Engine/Shaders/Private/MobileFog.usf b/Engine/Shaders/Private/MobileFog.usf
new file mode 100644
--- /dev/null
+++ b/Engine/Shaders/Private/MobileFog.usf
@@ -0,0 +1,75 @@
+// Copyright Epic Games, Inc. All Rights Reserved.
+#include "Common.ush"
+
+// Reroute MobileSceneTextures uniform buffer references to the base pass uniform buffer
+#define MobileSceneTextures MobileBasePass.SceneTextures
+#define FogStruct MobileBasePass.Fog
+
+#include "SceneTexturesCommon.ush"
+#include "HeightFogCommon.ush"
+#include "SkyAtmosphereCommon.ush"
+
+float StartDepthZ;
+
+void MobileFogVS(
+	in float2 InPosition : ATTRIBUTE0,
+	out float2 OutTexCoord : TEXCOORD0,
+	out float3 OutScreenVector : TEXCOORD1,
+	out float4 OutPosition : SV_POSITION
+	)
+{
+	ResolvedView = ResolveView();
+
+	// screenspace position from vb
+	OutPosition = float4(InPosition,StartDepthZ,1);
+	// texture coord from vb
+	OutTexCoord = InPosition * ResolvedView.ScreenPositionScaleBias.xy + ResolvedView.ScreenPositionScaleBias.wz;
+	// deproject to world space
+	OutScreenVector = mul(float4(InPosition,1,0), ResolvedView.ScreenToTranslatedWorld).xyz;
+}
+
+void MobileFogPS(
+	float2 TexCoord : TEXCOORD0,
+	float3 ScreenVector : TEXCOORD1,
+	float4 SVPos : SV_POSITION,
+	out float4 OutColor : SV_Target0
+	)
+{
+	ResolvedView = ResolveView();
+	half4 FinalFog = half4(0,0,0,1);
+	float SceneDepth = CalcSceneDepth(TexCoord);
+	float3 WorldPositionRelativeToCamera = ScreenVector.xyz * SceneDepth;
+
+// at least one of these is active
+#if PERMUTATION_SUPPORT_HEIGHT_FOG
+	FinalFog = CalculateHeightFog(WorldPositionRelativeToCamera);
+#endif
+
+#if PERMUTATION_SUPPORT_AERIAL_PERSPECTIVE
+	if (ResolvedView.SkyAtmosphereApplyCameraAerialPerspectiveVolume > 0.0f)
+	{
+		float3 WorldPosition = WorldPositionRelativeToCamera + ResolvedView.TranslatedWorldCameraOrigin - LWCHackToFloat(ResolvedView.PreViewTranslation);
+		float4 NDCPosition = LWCMultiply(MakeLWCVector4(WorldPosition, 1.0f), ResolvedView.WorldToClip);
+				
+		// Sample the aerial perspective (AP). It is also blended under the VertexFog parameter.
+		FinalFog = GetAerialPerspectiveLuminanceTransmittanceWithFogOver(
+			ResolvedView.RealTimeReflectionCapture, 
+			ResolvedView.SkyAtmosphereCameraAerialPerspectiveVolumeSizeAndInvSize,
+			NDCPosition, 
+			WorldPosition.xyz*CM_TO_SKY_UNIT, 
+			ResolvedView.TranslatedWorldCameraOrigin*CM_TO_SKY_UNIT,
+			View.CameraAerialPerspectiveVolume, 
+			View.CameraAerialPerspectiveVolumeSampler,
+			ResolvedView.SkyAtmosphereCameraAerialPerspectiveVolumeDepthResolutionInv,
+			ResolvedView.SkyAtmosphereCameraAerialPerspectiveVolumeDepthResolution,
+			ResolvedView.SkyAtmosphereAerialPerspectiveStartDepthKm,
+			ResolvedView.SkyAtmosphereCameraAerialPerspectiveVolumeDepthSliceLengthKm,
+			ResolvedView.SkyAtmosphereCameraAerialPerspectiveVolumeDepthSliceLengthKmInv,
+			ResolvedView.OneOverPreExposure, 
+			FinalFog);
+	}
+#endif// PERMUTATION_SUPPORT_AERIAL_PERSPECTIVE
+
+	FinalFog.rgb *= ResolvedView.PreExposure;
+	OutColor = RETURN_COLOR(half4(FinalFog.rgb, FinalFog.w)); 
+}
diff --git a/Engine/Source/Runtime/Renderer/Private/MobileFogRendering.cpp b/Engine/Source/Runtime/Renderer/Private/MobileFogRendering.cpp
new file mode 100644
--- /dev/null
+++ b/Engine/Source/Runtime/Renderer/Private/MobileFogRendering.cpp
@@ -0,0 +1,202 @@
+// Copyright Epic Games, Inc. All Rights Reserved.
+
+/*=============================================================================
+	MobileTranslucentRendering.cpp: translucent rendering implementation.
+=============================================================================*/
+
+#include "CoreMinimal.h"
+#include "HAL/IConsoleManager.h"
+#include "RHI.h"
+#include "Shader.h"
+#include "GlobalShader.h"
+#include "RenderUtils.h"
+#include "SceneRendering.h"
+#include "ScenePrivate.h"
+#include "FogRendering.h"
+#include "ProfilingDebugging/RealtimeGPUProfiler.h"
+
+static TAutoConsoleVariable<int32> CVarPixelFogQuality(
+	TEXT("r.Mobile.PixelFogQuality"),
+	1,
+	TEXT("Exponentional height fog rendering quality.\n")
+	TEXT("0 - basic per-pixel fog")
+	TEXT("1 - all per-pixel fog features (second fog, directional inscattering, aerial perspective)"),
+	ECVF_RenderThreadSafe);
+
+static TAutoConsoleVariable<int32> CVarPixelFogDepthTest(
+	TEXT("r.Mobile.PixelFogDepthTest"),
+	1,
+	TEXT("Whether to use depth and stencil tests for fog rendering"),
+	ECVF_RenderThreadSafe);
+
+class FMobileFogVS : public FGlobalShader
+{
+	DECLARE_GLOBAL_SHADER(FMobileFogVS);
+	SHADER_USE_PARAMETER_STRUCT(FMobileFogVS, FGlobalShader);
+
+	static bool ShouldCompilePermutation(const FGlobalShaderPermutationParameters& Parameters) 
+	{
+		return IsMobilePlatform(Parameters.Platform);
+	}
+
+	BEGIN_SHADER_PARAMETER_STRUCT(FParameters, )
+		SHADER_PARAMETER_STRUCT_INCLUDE(FViewShaderParameters, View)
+		SHADER_PARAMETER(float, StartDepthZ)
+	END_SHADER_PARAMETER_STRUCT()
+};
+
+IMPLEMENT_GLOBAL_SHADER(FMobileFogVS, "/Engine/Private/MobileFog.usf", "MobileFogVS", SF_Vertex);
+
+class FMobileFogPS : public FGlobalShader
+{
+	DECLARE_GLOBAL_SHADER(FMobileFogPS);
+	SHADER_USE_PARAMETER_STRUCT(FMobileFogPS, FGlobalShader);
+	
+	class FSupportHeightFog							: SHADER_PERMUTATION_BOOL("PERMUTATION_SUPPORT_HEIGHT_FOG");
+	class FSupportFogStartDistance					: SHADER_PERMUTATION_BOOL("PERMUTATION_SUPPORT_FOG_START_DISTANCE");
+	class FSupportFogInScatteringTexture			: SHADER_PERMUTATION_BOOL("PERMUTATION_SUPPORT_FOG_INSCATTERING_TEXTURE");
+	class FSupportFogSecondTerm						: SHADER_PERMUTATION_BOOL("PERMUTATION_SUPPORT_FOG_SECOND_TERM");
+	class FSupportFogDirectionalLightInScattering	: SHADER_PERMUTATION_BOOL("PERMUTATION_SUPPORT_FOG_DIRECTIONAL_LIGHT_INSCATTERING");
+	class FSupportAerialPerspective					: SHADER_PERMUTATION_BOOL("PERMUTATION_SUPPORT_AERIAL_PERSPECTIVE");
+	
+	using FPermutationDomain = TShaderPermutationDomain< 
+		FSupportHeightFog, 
+		FSupportFogStartDistance,
+		FSupportFogInScatteringTexture, 
+		FSupportFogSecondTerm,
+		FSupportFogDirectionalLightInScattering,
+		FSupportAerialPerspective 
+	>;
+
+	static FPermutationDomain RemapPermutation(FPermutationDomain PermutationVector)
+	{
+		if (PermutationVector.Get<FSupportHeightFog>() == false)
+		{
+			PermutationVector.Set<FSupportFogInScatteringTexture>(false);
+			PermutationVector.Set<FSupportFogDirectionalLightInScattering>(false);
+			PermutationVector.Set<FSupportFogStartDistance>(false);
+			PermutationVector.Set<FSupportFogSecondTerm>(false);
+		}
+		return PermutationVector;
+	}
+
+	static bool ShouldCompilePermutation(const FGlobalShaderPermutationParameters& Parameters)
+	{
+		if (!IsMobilePlatform(Parameters.Platform))
+		{
+			return false;
+		}
+
+		FPermutationDomain PermutationVector(Parameters.PermutationId);
+		if (RemapPermutation(PermutationVector) != PermutationVector)
+		{
+			return false;
+		}
+		
+		return true;
+	}
+
+	BEGIN_SHADER_PARAMETER_STRUCT(FParameters, )
+		SHADER_PARAMETER_STRUCT_INCLUDE(FViewShaderParameters, View)
+	END_SHADER_PARAMETER_STRUCT()
+};
+
+IMPLEMENT_SHADER_TYPE(, FMobileFogPS, TEXT("/Engine/Private/MobileFog.usf"), TEXT("MobileFogPS"), SF_Pixel);
+
+void FMobileSceneRenderer::RenderFog(FRHICommandListImmediate& RHICmdList, const FViewInfo& View)
+{	
+	static const auto* CVarDisableVertexFog = IConsoleManager::Get().FindTConsoleVariableDataInt(TEXT("r.Mobile.DisableVertexFog"));
+	if (CVarDisableVertexFog && CVarDisableVertexFog->GetValueOnRenderThread() == 0)
+	{
+		// Project uses only vertex fogging
+		return;
+	}
+
+	const int32 PixelFogQuality = CVarPixelFogQuality.GetValueOnRenderThread();
+	const bool bUseHeightFog = Scene->ExponentialFogs.Num() > 0 && ShouldRenderFog(*View.Family);
+	const bool bUseAerialPerspective = PixelFogQuality > 0 && ShouldRenderSkyAtmosphere(Scene, View.Family->EngineShowFlags);
+
+	if (!bUseAerialPerspective && !bUseHeightFog)
+	{
+		return;
+	}
+
+	const bool bUseDepthTest = CVarPixelFogDepthTest.GetValueOnRenderThread() != 0;
+		
+	SCOPED_DRAW_EVENT(RHICmdList, Fog);
+		
+	RHICmdList.SetViewport(View.ViewRect.Min.X, View.ViewRect.Min.Y, 0.0f, View.ViewRect.Max.X, View.ViewRect.Max.Y, 1.0f);
+	
+	FGraphicsPipelineStateInitializer GraphicsPSOInit;
+	RHICmdList.ApplyCachedRenderTargets(GraphicsPSOInit);
+	GraphicsPSOInit.RasterizerState = TStaticRasterizerState<FM_Solid, CM_None>::GetRHI();
+	GraphicsPSOInit.PrimitiveType = PT_TriangleList;
+	
+	if (bUseDepthTest)
+	{
+		GraphicsPSOInit.DepthStencilState = TStaticDepthStencilState<
+			false, CF_DepthNearOrEqual,
+			true, CF_Equal, SO_Keep, SO_Keep, SO_Keep,
+			false, CF_Always, SO_Keep, SO_Keep, SO_Keep,
+			STENCIL_MOBILE_SKY_MASK, 0x00>::GetRHI();
+	}
+	else
+	{
+		GraphicsPSOInit.DepthStencilState = TStaticDepthStencilState<false, CF_Always>::GetRHI();
+	}
+	
+	GraphicsPSOInit.BlendState = TStaticBlendState<CW_RGB, BO_Add, BF_One, BF_SourceAlpha>::GetRHI();
+		
+	TShaderMapRef<FMobileFogVS> VertexShader(View.ShaderMap);
+
+	const bool bUseFogStartDistance = View.ExponentialFogParameters.W > 0;
+	const bool bUseFogInscatteringColorCubemap = View.FogInscatteringColorCubemap != nullptr;
+	const bool bUseFogSecondTerm = PixelFogQuality > 0 && (View.ExponentialFogParameters2.X > 0);
+	const bool bUseFogDirectionalInscatering = PixelFogQuality > 0 && !bUseFogInscatteringColorCubemap && (View.DirectionalInscatteringColor.GetLuminance() > 0 || View.bUseDirectionalInscattering);
+	
+	FMobileFogPS::FPermutationDomain PsPermutationVector;
+	PsPermutationVector.Set<FMobileFogPS::FSupportHeightFog>(bUseHeightFog);
+	PsPermutationVector.Set<FMobileFogPS::FSupportFogStartDistance>(bUseFogStartDistance);
+	PsPermutationVector.Set<FMobileFogPS::FSupportFogInScatteringTexture>(bUseFogInscatteringColorCubemap);
+	PsPermutationVector.Set<FMobileFogPS::FSupportFogSecondTerm>(bUseFogSecondTerm);
+	PsPermutationVector.Set<FMobileFogPS::FSupportAerialPerspective>(bUseAerialPerspective);
+	PsPermutationVector.Set<FMobileFogPS::FSupportFogDirectionalLightInScattering>(bUseFogDirectionalInscatering);
+	
+	TShaderMapRef<FMobileFogPS> PixelShader(View.ShaderMap, PsPermutationVector);
+		
+	GraphicsPSOInit.BoundShaderState.VertexDeclarationRHI = GTileVertexDeclaration.VertexDeclarationRHI;
+	GraphicsPSOInit.BoundShaderState.VertexShaderRHI = VertexShader.GetVertexShader();
+	GraphicsPSOInit.BoundShaderState.PixelShaderRHI = PixelShader.GetPixelShader();
+	
+	SetGraphicsPipelineState(RHICmdList, GraphicsPSOInit, 0);
+	
+	// Use height fog start distance by default and fallback to AP distance
+	float FogStartDistance = View.ExponentialFogParameters.W;
+	if (!bUseHeightFog)
+	{
+		FSkyAtmosphereRenderSceneInfo& SkyInfo = *Scene->GetSkyAtmosphereSceneInfo();
+		const FSkyAtmosphereSceneProxy& SkyAtmosphereSceneProxy = SkyInfo.GetSkyAtmosphereSceneProxy();
+		FogStartDistance = GetValidAerialPerspectiveStartDepthInCm(View, SkyAtmosphereSceneProxy);
+	}
+		
+	float StartDepthZ = 0.1;	
+	if (bUseDepthTest && FogStartDistance > 0.f)
+	{
+		const FMatrix ProjectionMatrix = View.ViewMatrices.GetProjectionMatrix();
+		float HalfHorizontalFOV = FMath::Atan(1.0f / ProjectionMatrix.M[0][0]);
+		float HalfVerticalFOV = FMath::Atan(1.0f / ProjectionMatrix.M[1][1]);
+		float StartDepthViewCm = FMath::Cos(FMath::Max(HalfHorizontalFOV, HalfVerticalFOV)) * FogStartDistance;
+		StartDepthViewCm = FMath::Max(StartDepthViewCm, View.NearClippingDistance); // In any case, we need to limit the distance to frustum near plane to not be clipped away.
+		const FVector4f Projected = ProjectionMatrix.TransformFVector4(FVector4f(0.0f, 0.0f, StartDepthViewCm, 1.0f));
+		StartDepthZ = Projected.Z / Projected.W;
+	}
+
+	FMobileFogVS::FParameters VSParameters;
+	VSParameters.View = View.GetShaderParameters();
+	VSParameters.StartDepthZ = StartDepthZ;
+	SetShaderParameters(RHICmdList, VertexShader, VertexShader.GetVertexShader(), VSParameters);
+
+	// Draw a quad covering the view.
+	RHICmdList.SetStreamSource(0, GScreenSpaceVertexBuffer.VertexBufferRHI, 0);
+	RHICmdList.DrawIndexedPrimitive(GTwoTrianglesIndexBuffer.IndexBufferRHI, 0, 0, 4, 0, 2, 1);
+}
diff --git a/Engine/Source/Runtime/Renderer/Private/MobileShadingRenderer.cpp b/Engine/Source/Runtime/Renderer/Private/MobileShadingRenderer.cpp
--- a/Engine/Source/Runtime/Renderer/Private/MobileShadingRenderer.cpp
+++ b/Engine/Source/Runtime/Renderer/Private/MobileShadingRenderer.cpp
@@ -1095,6 +1095,7 @@ void FMobileSceneRenderer::RenderForwardSinglePass(FRDGBuilder& GraphBuilder, FM
 		RHICmdList.SetCurrentStat(GET_STATID(STAT_CLMM_Translucency));
 		RenderDecals(RHICmdList, View);
 		RenderModulatedShadowProjections(RHICmdList, ViewIndex, View);
+		RenderFog(RHICmdList, View);
 		// Draw translucency.
 		RenderTranslucency(RHICmdList, View);
 		if (bAdrenoOcclusionMode)
@@ -1174,6 +1175,7 @@ void FMobileSceneRenderer::RenderForwardMultiPass(FRDGBuilder& GraphBuilder, FMo
 		RHICmdList.SetCurrentStat(GET_STATID(STAT_CLMM_Translucency));
 		RenderDecals(RHICmdList, View);
 		RenderModulatedShadowProjections(RHICmdList, ViewIndex, View);
+		RenderFog(RHICmdList, View);
 		// Draw translucency.
 		RenderTranslucency(RHICmdList, View);
 		// Pre-tonemap before MSAA resolve (iOS only)
@@ -1378,6 +1380,7 @@ void FMobileSceneRenderer::RenderDeferredSinglePass(FRDGBuilder& GraphBuilder, c
 			RHICmdList.NextSubpass();
 		}
 		MobileDeferredShadingPass(RHICmdList, ViewIndex, NumViews, View, *Scene, SortedLightSet);
+		RenderFog(RHICmdList, View);
 		// Draw translucency.
 		RenderTranslucency(RHICmdList, View);
 	});
@@ -1457,6 +1460,7 @@ void FMobileSceneRenderer::RenderDeferredMultiPass(FRDGBuilder& GraphBuilder, cl
 	{
 		RHICmdList.NextSubpass();
 		MobileDeferredShadingPass(RHICmdList, ViewIndex, NumViews, View, *Scene, SortedLightSet);
+		RenderFog(RHICmdList, View);
 		// Draw translucency.
 		RenderTranslucency(RHICmdList, View);
 	});
diff --git a/Engine/Source/Runtime/Renderer/Private/PostProcess/SceneRenderTargets.h b/Engine/Source/Runtime/Renderer/Private/PostProcess/SceneRenderTargets.h
--- a/Engine/Source/Runtime/Renderer/Private/PostProcess/SceneRenderTargets.h
+++ b/Engine/Source/Runtime/Renderer/Private/PostProcess/SceneRenderTargets.h
@@ -41,6 +41,8 @@
 #define STENCIL_LIGHTING_CHANNELS_MASK(Value) uint8(((Value) & 0x7) << STENCIL_LIGHTING_CHANNELS_BIT_ID)
 
 // Mobile specific
+// Sky material mask
+#define STENCIL_MOBILE_SKY_MASK GET_STENCIL_BIT_MASK(SANDBOX,1)
 // Store shading model into stencil [1-3] bits
 #define GET_STENCIL_MOBILE_SM_MASK(Value) uint8(((Value) & 0x7) << 1)
 
diff --git a/Engine/Source/Runtime/Renderer/Private/SceneRendering.h b/Engine/Source/Runtime/Renderer/Private/SceneRendering.h
--- a/Engine/Source/Runtime/Renderer/Private/SceneRendering.h
+++ b/Engine/Source/Runtime/Renderer/Private/SceneRendering.h
@@ -2336,6 +2336,9 @@ protected:
 	/** Renders decals. */
 	void RenderDecals(FRHICommandListImmediate& RHICmdList, const FViewInfo& View);
 
+	/** Renders the atmospheric and height fog */
+	void RenderFog(FRHICommandListImmediate& RHICmdList, const FViewInfo& View);
+
 	/** Renders the base pass for translucency. */
 	void RenderTranslucency(FRHICommandListImmediate& RHICmdList, const FViewInfo& View);
 
diff --git a/Engine/Source/Runtime/Renderer/Private/SkyAtmosphereRendering.cpp b/Engine/Source/Runtime/Renderer/Private/SkyAtmosphereRendering.cpp
--- a/Engine/Source/Runtime/Renderer/Private/SkyAtmosphereRendering.cpp
+++ b/Engine/Source/Runtime/Renderer/Private/SkyAtmosphereRendering.cpp
@@ -601,6 +601,11 @@ class FRenderSkyAtmosphereVS : public FGlobalShader
 
 	static bool ShouldCompilePermutation(const FGlobalShaderPermutationParameters& Parameters)
 	{
+		if (IsMobilePlatform(Parameters.Platform))
+		{
+			// Mobile must use a skydome mesh with a sky material to achieve good GPU performance
+			return false;
+		}
 		return ShouldPipelineCompileSkyAtmosphereShader(Parameters.Platform);
 	}
 
@@ -674,6 +679,12 @@ class FRenderSkyAtmospherePS : public FGlobalShader
 
 	static bool ShouldCompilePermutation(const FGlobalShaderPermutationParameters& Parameters)
 	{
+		if (IsMobilePlatform(Parameters.Platform))
+		{
+			// Mobile must use a skydome mesh with a sky material to achieve good GPU performance
+			return false;
+		}
+
 		FPermutationDomain PermutationVector(Parameters.PermutationId);
 
 		// If not rendering the sky, ignore the FFastSky permutation
diff --git a/Engine/Source/Runtime/Renderer/Private/SkyPassRendering.cpp b/Engine/Source/Runtime/Renderer/Private/SkyPassRendering.cpp
--- a/Engine/Source/Runtime/Renderer/Private/SkyPassRendering.cpp
+++ b/Engine/Source/Runtime/Renderer/Private/SkyPassRendering.cpp
@@ -125,6 +125,15 @@ bool FSkyPassMeshProcessor::Process(
 			return false;
 		}
 
+		// Mask sky pixels so we can skip them when rendering per-pixel fog
+		PassDrawRenderState.SetDepthStencilState(TStaticDepthStencilState<
+			false, CF_DepthNearOrEqual,
+			true, CF_Always, SO_Keep, SO_Keep, SO_Replace,
+			false, CF_Always, SO_Keep, SO_Keep, SO_Keep,
+			0x00, STENCIL_MOBILE_SKY_MASK>::GetRHI());
+		
+		PassDrawRenderState.SetStencilRef(1); 
+		
 		TMobileBasePassShaderElementData<LightMapPolicyType> ShaderElementData(nullptr, false);
 		ShaderElementData.InitializeMeshMaterialData(ViewIfDynamicMeshCommand, PrimitiveSceneProxy, MeshBatch, StaticMeshId, false);
 
