Added BlackArrayDummy to system textures (single black pixel slice-texture array) because Metal requires texture arrays to be bound to texture array parameters

diff --git a/Engine/Source/Runtime/Renderer/Private/SystemTextures.cpp b/Engine/Source/Runtime/Renderer/Private/SystemTextures.cpp
--- a/Engine/Source/Runtime/Renderer/Private/SystemTextures.cpp
+++ b/Engine/Source/Runtime/Renderer/Private/SystemTextures.cpp
@@ -87,6 +87,15 @@ void SetDummyTextureData(FRHITexture2D* Texture2D, const DataType& DummyData)
 	RHIUnlockTexture2D(Texture2D, 0, false);
 }
 
+template <typename DataType>
+void SetDummyTextureArrayData(FRHITexture2DArray* Texture2DArray, const DataType& DummyData)
+{
+	uint32 DestStride;
+	DataType* DestBuffer = (DataType*)RHILockTexture2DArray(Texture2DArray, 0, 0, RLM_WriteOnly, DestStride, false);
+	*DestBuffer = DummyData;
+	RHIUnlockTexture2DArray(Texture2DArray, 0, 0, false);
+}
+
 void FSystemTextures::InitializeCommonTextures(FRHICommandListImmediate& RHICmdList)
 {
 	// First initialize textures that are common to all feature levels. This is always done the first time we come into this function, as doesn't care about the
@@ -106,6 +115,14 @@ void FSystemTextures::InitializeCommonTextures(FRHICommandListImmediate& RHICmdL
 		BlackDummy = CreateRenderTarget(Texture2D, CreateInfo.DebugName);
 	}
 	
+	// Create a texture array that has a single black dummy slice 
+	{
+		FRHIResourceCreateInfo CreateInfo(TEXT("BlackArrayDummy"));
+		FTexture2DArrayRHIRef Texture2DArray = RHICreateTexture2DArray(1, 1, 1, PF_B8G8R8A8, 1, 1, TexCreate_ShaderResource, CreateInfo);
+		SetDummyTextureArrayData<FColor>(Texture2DArray, FColor(0, 0, 0, 0));
+		BlackArrayDummy = CreateRenderTarget(Texture2DArray, CreateInfo.DebugName);
+	}
+
 	// Create a texture that is a single UInt32 value set to 0
 	{	
 		FRHIResourceCreateInfo CreateInfo(TEXT("ZeroUIntDummy"));
@@ -329,7 +346,7 @@ void FSystemTextures::InitializeFeatureLevelDependentTextures(FRHICommandListImm
 		RHICmdList.UnlockTexture2D((FTexture2DRHIRef&)AsciiTexture->GetRenderTargetItem().ShaderResourceTexture, 0, false);
 	}
 
-	// The PreintegratedGF maybe used on forward shading inluding mobile platorm, intialize it anyway.
+	// The PreintegratedGF maybe used on forward shading including mobile platform, initialize it anyway.
 	{
 		// for testing, with 128x128 R8G8 we are very close to the reference (if lower res is needed we might have to add an offset to counter the 0.5f texel shift)
 		const bool bReference = false;
@@ -799,6 +816,7 @@ void FSystemTextures::ReleaseDynamicRHI()
 	WhiteDummySRV.SafeRelease();
 	WhiteDummy.SafeRelease();
 	BlackDummy.SafeRelease();
+	BlackArrayDummy.SafeRelease();
 	BlackAlphaOneDummy.SafeRelease();
 	PerlinNoiseGradient.SafeRelease();
 	PerlinNoise3D.SafeRelease();
@@ -847,6 +865,11 @@ FRDGTextureRef FSystemTextures::GetBlackAlphaOneDummy(FRDGBuilder& GraphBuilder)
 	return GraphBuilder.RegisterExternalTexture(BlackAlphaOneDummy, TEXT("BlackAlphaOneDummy"));
 }
 
+FRDGTextureRef FSystemTextures::GetBlackArrayDummy(FRDGBuilder& GraphBuilder) const
+{
+	return GraphBuilder.RegisterExternalTexture(BlackArrayDummy, TEXT("BlackArrayDummy"));
+}
+
 FRDGTextureRef FSystemTextures::GetWhiteDummy(FRDGBuilder& GraphBuilder) const
 {
 	return GraphBuilder.RegisterExternalTexture(WhiteDummy, TEXT("WhiteDummy"));
diff --git a/Engine/Source/Runtime/Renderer/Private/SystemTextures.h b/Engine/Source/Runtime/Renderer/Private/SystemTextures.h
--- a/Engine/Source/Runtime/Renderer/Private/SystemTextures.h
+++ b/Engine/Source/Runtime/Renderer/Private/SystemTextures.h
@@ -99,6 +99,8 @@ public:
 	TRefCountPtr<IPooledRenderTarget> BlackDummy;
 	// float4(0,0,0,1)
 	TRefCountPtr<IPooledRenderTarget> BlackAlphaOneDummy;
+	// Dummy texture array with 1 black slice 
+	TRefCountPtr<IPooledRenderTarget> BlackArrayDummy;
 	// used by the material expression Noise
 	TRefCountPtr<IPooledRenderTarget> PerlinNoiseGradient;
 	// used by the material expression Noise (faster version, should replace old version), todo: move out of SceneRenderTargets
@@ -152,6 +154,7 @@ public:
 	// Create simple default texture
 	FRDGTextureRef RENDERER_API GetWhiteDummy(FRDGBuilder& GraphBuilder) const;
 	FRDGTextureRef RENDERER_API GetBlackDummy(FRDGBuilder& GraphBuilder) const;
+	FRDGTextureRef RENDERER_API GetBlackArrayDummy(FRDGBuilder& GraphBuilder) const;
 	FRDGTextureRef RENDERER_API GetZeroUIntDummy(FRDGBuilder& GraphBuilder) const;
 	FRDGTextureRef RENDERER_API GetZeroUShort4Dummy(FRDGBuilder& GraphBuilder) const;
 	FRDGTextureRef RENDERER_API GetBlackAlphaOneDummy(FRDGBuilder& GraphBuilder) const;
