Use UInt Dummy for GlobalDistanceFieldPageTableTexture to fix shader validation errors on macOS.

diff --git a/Engine/Source/Runtime/Renderer/Private/GlobalDistanceField.cpp b/Engine/Source/Runtime/Renderer/Private/GlobalDistanceField.cpp
--- a/Engine/Source/Runtime/Renderer/Private/GlobalDistanceField.cpp
+++ b/Engine/Source/Runtime/Renderer/Private/GlobalDistanceField.cpp
@@ -218,7 +218,7 @@ FGlobalDistanceFieldParameters2 SetupGlobalDistanceFieldParameters(const FGlobal
 	FGlobalDistanceFieldParameters2 ShaderParameters;
 
 	ShaderParameters.GlobalDistanceFieldPageAtlasTexture = OrBlack3DIfNull(ParameterData.PageAtlasTexture);
-	ShaderParameters.GlobalDistanceFieldPageTableTexture = OrBlack3DIfNull(ParameterData.PageTableTexture);
+	ShaderParameters.GlobalDistanceFieldPageTableTexture = OrBlack3DUintIfNull(ParameterData.PageTableTexture);
 	ShaderParameters.GlobalDistanceFieldMipTexture = OrBlack3DIfNull(ParameterData.MipTexture);
 
 	for (int32 Index = 0; Index < GMaxGlobalDistanceFieldClipmaps; Index++)
@@ -1068,7 +1068,7 @@ void FViewInfo::SetupDefaultGlobalDistanceFieldUniformBufferParameters(FViewUnif
 	ViewUniformShaderParameters.NumGlobalSDFClipmaps = 0;
 
 	ViewUniformShaderParameters.GlobalDistanceFieldPageAtlasTexture = OrBlack3DIfNull(GBlackVolumeTexture->TextureRHI.GetReference());
-	ViewUniformShaderParameters.GlobalDistanceFieldPageTableTexture = OrBlack3DIfNull(GBlackUintVolumeTexture->TextureRHI.GetReference());
+	ViewUniformShaderParameters.GlobalDistanceFieldPageTableTexture = OrBlack3DUintIfNull(GBlackUintVolumeTexture->TextureRHI.GetReference());
 	ViewUniformShaderParameters.GlobalDistanceFieldMipTexture = OrBlack3DIfNull(GBlackVolumeTexture->TextureRHI.GetReference());
 }
 
@@ -1093,7 +1093,7 @@ void FViewInfo::SetupGlobalDistanceFieldUniformBufferParameters(FViewUniformShad
 	ViewUniformShaderParameters.NumGlobalSDFClipmaps = GlobalDistanceFieldInfo.ParameterData.NumGlobalSDFClipmaps;
 
 	ViewUniformShaderParameters.GlobalDistanceFieldPageAtlasTexture = OrBlack3DIfNull(GlobalDistanceFieldInfo.ParameterData.PageAtlasTexture);
-	ViewUniformShaderParameters.GlobalDistanceFieldPageTableTexture = OrBlack3DIfNull(GlobalDistanceFieldInfo.ParameterData.PageTableTexture);
+	ViewUniformShaderParameters.GlobalDistanceFieldPageTableTexture = OrBlack3DUintIfNull(GlobalDistanceFieldInfo.ParameterData.PageTableTexture);
 	ViewUniformShaderParameters.GlobalDistanceFieldMipTexture = OrBlack3DIfNull(GlobalDistanceFieldInfo.ParameterData.MipTexture);
 }
 
@@ -2368,4 +2368,4 @@ void GlobalDistanceField::ExpandDistanceFieldUpdateTrackingBounds(const FSceneVi
 		const FBox ClipmapInfluenceBounds(ClipmapCenter - ClipmapExtent, ClipmapCenter + ClipmapExtent);
 		UpdateTrackingBounds.GlobalDistanceFieldBounds += ClipmapInfluenceBounds;
 	}
-}
\ No newline at end of file
+}
