Removing invalid assert from post processing for null custom depth texture.

diff --git a/Engine/Source/Runtime/Renderer/Private/PostProcess/PostProcessing.h b/Engine/Source/Runtime/Renderer/Private/PostProcess/PostProcessing.h
--- a/Engine/Source/Runtime/Renderer/Private/PostProcess/PostProcessing.h
+++ b/Engine/Source/Runtime/Renderer/Private/PostProcess/PostProcessing.h
@@ -36,7 +36,6 @@ struct FPostProcessingInputs
 	{
 		check(SceneTextures);
 		check(ViewFamilyTexture);
-		check(CustomDepthTexture);
 		check(SeparateTranslucencyTextures);
 	}
 };
