Fix hair strands DOM debug info to be output several time

diff --git a/Engine/Source/Runtime/Renderer/Private/HairStrands/HairStrandsDebug.cpp b/Engine/Source/Runtime/Renderer/Private/HairStrands/HairStrandsDebug.cpp
--- a/Engine/Source/Runtime/Renderer/Private/HairStrands/HairStrandsDebug.cpp
+++ b/Engine/Source/Runtime/Renderer/Private/HairStrands/HairStrandsDebug.cpp
@@ -605,9 +605,7 @@ static void AddDeepShadowInfoPass(
 	Parameters->OutputTexture = GraphBuilder.CreateUAV(OutputTexture);
 
 	TShaderMapRef<FDeepShadowInfoCS> ComputeShader(View.ShaderMap);
-
-	const FIntVector DispatchCount = DispatchCount.DivideAndRoundUp(FIntVector(OutputTexture->Desc.Extent.X, OutputTexture->Desc.Extent.Y, 1), FIntVector(8, 8, 1));
-	FComputeShaderUtils::AddPass(GraphBuilder, RDG_EVENT_NAME("HairStrands::DeepShadowDebugInfo"), ComputeShader, Parameters, DispatchCount);
+	FComputeShaderUtils::AddPass(GraphBuilder, RDG_EVENT_NAME("HairStrands::DeepShadowDebugInfo"), ComputeShader, Parameters, FIntVector(1, 1, 1));
 }
 
 ///////////////////////////////////////////////////////////////////////////////////////////////////
