Add VSM projection visualizations to the editor menu

diff --git a/Engine/Shaders/Private/VirtualShadowMaps/VirtualShadowMapProjection.usf b/Engine/Shaders/Private/VirtualShadowMaps/VirtualShadowMapProjection.usf
--- a/Engine/Shaders/Private/VirtualShadowMaps/VirtualShadowMapProjection.usf
+++ b/Engine/Shaders/Private/VirtualShadowMaps/VirtualShadowMapProjection.usf
@@ -18,6 +18,7 @@
 #include "ProjectionDirectional.ush"
 #include "ProjectionSpot.ush"
 #include "TransmissionCommon.ush"
+#include "Visualize.ush"
 
 #if HAS_HAIR_STRANDS
 #define VOXEL_TRAVERSAL_TYPE VOXEL_TRAVERSAL_LINEAR_MIPMAP
@@ -30,8 +31,8 @@
 #define DIRECTIONAL_LIGHT 0
 #endif
 
-#ifndef DEBUG_OUTPUT
-#define DEBUG_OUTPUT 0
+#ifndef VISUALIZE_OUTPUT
+#define VISUALIZE_OUTPUT 0
 #endif
 
 #define CONTACT_SHADOW_SAMPLES 8
@@ -143,18 +144,18 @@ float SMRTCotMaxRayAngleFromLight;	// Spot/point lights
 float SMRTTexelDitherScale;			// Currently only directional lights
 
 // One pass projection
-RWTexture2D< uint4 > RWShadowMaskBits;
+RWTexture2D< uint4 >OutShadowMaskBits;
 
 // Single light per pass
 // Light parameters loaded via GetRootLightShaderParameters();
 int LightUniformVirtualShadowMapId;
-RWTexture2D< float2 > RWShadowFactor;
+RWTexture2D< float2 > OutShadowFactor;
 
-// Debug output pass
+// Visualization output
 StructuredBuffer< FPhysicalPageMetaData > PhysicalPageMetaData;
-RWTexture2D< float4 > RWDebug;
-int DebugOutputType;
-int DebugVirtualShadowMapId;
+RWTexture2D< float4 > OutVisualize;
+int VisualizeModeId;
+int VisualizeVirtualShadowMapId;
 
 // Type of input data consume by the page allocation (i.e., data read from the source buffer: Gbuffer, HairStrands data, ...)
 #define INPUT_TYPE_GBUFFER 0
@@ -355,7 +356,7 @@ FVirtualShadowMapSampleResult ProjectLight(
 	return Result;
 }
 
-void OutputDebug( uint2 PixelPos, FVirtualShadowMapSampleResult Result, uint LightCount )
+void OutputVisualize( uint2 PixelPos, FVirtualShadowMapSampleResult Result )
 {
 	float3 Output = float3( 1, 0, 1 );
 	bool bWriteOutput = false;
@@ -367,61 +368,65 @@ void OutputDebug( uint2 PixelPos, FVirtualShadowMapSampleResult Result, uint Lig
 		uint2 vPage = Result.VirtualTexelAddress.xy >> VSM_LOG2_PAGE_SIZE;
 		uint2 pPage = Result.PhysicalTexelAddress.xy >> VSM_LOG2_PAGE_SIZE;
 		bool bValidPageAddress = all( pPage < uint2(VirtualShadowMap.PhysicalPoolSizePages) );
+		
+		FPhysicalPageMetaData pPageMeta = (FPhysicalPageMetaData)0;
+		if (bValidPageAddress)
+		{
+			pPageMeta = PhysicalPageMetaData[VSMPhysicalPageAddressToIndex(pPage)];
+		}
 
-		if ( DebugOutputType == 1 )			// ShadowFactor
+		if ( VisualizeModeId == VISUALIZE_SHADOW_FACTOR )
 		{
 			Output = Result.ShadowFactor.xxx;
 		}
-		if ( DebugOutputType == 2 )			// ClipmapIndexOrMip (mix in a bit of shadow for clarity)
+		else if ( VisualizeModeId == VISUALIZE_CLIPMAP_OR_MIP )
 		{
 			float3 Color = IntToColor( Result.ClipmapIndexOrMipLevel );
-			Output = lerp( Color, Result.ShadowFactor.xxx, 0.15f );
+			Output = 0.8f*Color + 0.2f*Result.ShadowFactor.xxx;
 		}
-		else if ( DebugOutputType == 3 && bValidPageAddress )	// Virtual page (mix in a bit of shadow for clarity)
+		else if ( VisualizeModeId == VISUALIZE_VIRTUAL_PAGE && bValidPageAddress )
 		{
-			float3 Color = IntToColor( vPage.x + ( vPage.y << 10U ) + (Result.ClipmapIndexOrMipLevel << 20U) );
-			Output = lerp( Color, Result.ShadowFactor.xxx, 0.15f );
-		}
-		else if ( DebugOutputType == 4 && bValidPageAddress )	// Physical page age (mix in a bit of shadow for clarity)
+			float3 PageColor = IntToColor( vPage.x + ( vPage.y << 10U ) );
+			float3 MipColor = IntToColor( Result.ClipmapIndexOrMipLevel );
+			Output = 0.4f*PageColor + 0.4f*MipColor + 0.2f*Result.ShadowFactor.xxx;
+		}		
+		else if ( VisualizeModeId == VISUALIZE_SMRT_RAY_COUNT )
 		{
-			uint Age = PhysicalPageMetaData[ VSMPhysicalPageAddressToIndex(pPage) ].Age;
-			float3 Color = GreenToRedTurbo( 1.0f - 0.1f * float(Age) );
-			Output = lerp( Color, Result.ShadowFactor.xxx, 0.15f );
+			Output = GreenToRedTurbo( float( Result.RayCount ) / float( SMRTRayCount ) );
 		}
-		else if ( DebugOutputType == 5 )	// RayCount
+		else if ( VisualizeModeId == VISUALIZE_CACHED_PAGE && bValidPageAddress )
 		{
-			Output = GreenToRedTurbo( float( Result.RayCount ) / float( SMRTRayCount ) );
+			bool bUncachedDynamic = (pPageMeta.Flags & VSM_DYNAMIC_UNCACHED_FLAG) != 0;
+			bool bUncachedStatic = (pPageMeta.Flags & VSM_STATIC_UNCACHED_FLAG) != 0;
+			// Red = both uncached, blue = only static cached, green otherwise
+			float3 CacheColor = float3(0, 1, 0);
+			if (bUncachedDynamic)
+			{
+				CacheColor = bUncachedStatic ? float3(1, 0, 0) : float3(0, 0, 1);
+			}
+			float3 PageColor = IntToColor(vPage.x + (vPage.y << 10U));
+			Output = 0.55f*CacheColor + 0.25f*PageColor + 0.2f*Result.ShadowFactor.xxx;
 		}
-		else if ( DebugOutputType == 6 )	// OccluderDistance
+		/*
+		else if ( VisualizeModeId == VISUALIZE_CACHED_PAGE_AGE && bValidPageAddress )
 		{
-			Output = ColorMapViridis( 0.001f * Result.OccluderDistance );
-			//Output = Result.OccluderDistance > 1000.0f ? float3(1, 1, 1) : ColorMapViridis( 0.001f * Result.OccluderDistance );
+			float3 Color = GreenToRedTurbo( 1.0f - 0.1f * float(pPageMeta.Age) );
+			Output = lerp( Color, Result.ShadowFactor.xxx, 0.15f );
 		}
-		else if ( DebugOutputType == 7 )	// Cached or not?
+		else if ( VisualizeModeId == 6 )	// OccluderDistance
 		{
-			bool bUncached = (PhysicalPageMetaData[VSMPhysicalPageAddressToIndex(pPage)].Flags & VSM_DYNAMIC_UNCACHED_FLAG) != 0;
-			float3 Color = GreenToRedTurbo(float(bUncached));
-			float3 Color2 = IntToColor(vPage.x + (vPage.y << 10U));
-
-			Output = lerp(Color, Color2, 0.35f);
+			Output = ColorMapViridis( 0.001f * Result.OccluderDistance );
 		}
+		*/
 		else
 		{
 			bWriteOutput = false;
 		}
 	}
 
-#if ONE_PASS_PROJECTION
-	if ( DebugOutputType == 8 )	// Light count (only useful for one pass projection)
-	{
-		Output = GreenToRedTurbo( float( LightCount ) / float( 50 ) );
-		bWriteOutput = true;
-	}
-#endif
-
 	if (bWriteOutput)
 	{
-		RWDebug[PixelPos] = float4(Output, 1.0f);
+		OutVisualize[PixelPos] = float4(Output, 1.0f);
 	}
 }
 
@@ -464,17 +469,14 @@ void VirtualShadowMapProjection(
 	const half3 GBufferNormal = GetGBufferDataUint( PixelPos, true ).WorldNormal;
 	const half SubsurfaceOpacity = bIsHairInput ? 1.0f : GetSubsurfaceOpacityFromGbuffer( PixelPos );
 
-	FVirtualShadowMapSampleResult DebugResult = InitVirtualShadowMapSampleResult();
-
-	uint LightCountDebug = 1;
+	FVirtualShadowMapSampleResult VisualizeResult = InitVirtualShadowMapSampleResult();
 
 #if ONE_PASS_PROJECTION
 	uint EyeIndex = 0; // TODO: Instanced stereo
 	uint GridLinearIndex = ComputeLightGridCellIndex( LocalPixelPos, SceneDepth, EyeIndex );
 	const FCulledLightsGridData CulledLightGridData = GetCulledLightsGrid( GridLinearIndex, EyeIndex );
 
-	// We can only handle 10 lights in our output encoding right now, so no purpose in computing more
-	LightCountDebug = CulledLightGridData.NumLocalLights;
+	// We can only handle so many lights in our output encoding right now, so no purpose in computing more
 	uint LightCount = min(GetPackedShadowMaskMaxLightCount(), CulledLightGridData.NumLocalLights);
 
 	uint4 ShadowMask = InitializePackedShadowMask();
@@ -503,14 +505,14 @@ void VirtualShadowMapProjection(
 
 			PackShadowMask(ShadowMask, Result.ShadowFactor, Index);
 
-			if (DebugVirtualShadowMapId == VirtualShadowMapId)
+			if (VisualizeVirtualShadowMapId == VirtualShadowMapId)
 			{
-				DebugResult = Result;
+				VisualizeResult = Result;
 			}
 		}
 	}
 
-	RWShadowMaskBits[ PixelPos ] = ~ShadowMask;
+	OutShadowMaskBits[ PixelPos ] = ~ShadowMask;
 
 #else // !ONE_PASS_PROJECTION
 	{
@@ -536,19 +538,19 @@ void VirtualShadowMapProjection(
 				SubsurfaceOpacity);
 		}
 
-		RWShadowFactor[ PixelPos ] = half2( Result.ShadowFactor, SSSTransmission );
+		OutShadowFactor[ PixelPos ] = half2( Result.ShadowFactor, SSSTransmission );
 
-		if (DebugVirtualShadowMapId == VirtualShadowMapId)
+		if (VisualizeVirtualShadowMapId == VirtualShadowMapId)
 		{
-			DebugResult = Result;
+			VisualizeResult = Result;
 		}
 	}
 #endif // ONE_PASS_PROJECTION
 
-#if DEBUG_OUTPUT
+#if VISUALIZE_OUTPUT
 	if (InputType == INPUT_TYPE_GBUFFER)
 	{
-		OutputDebug( PixelPos, DebugResult, LightCountDebug );
+		OutputVisualize( PixelPos, VisualizeResult );
 	}
-#endif // DEBUG_OUTPUT
+#endif // VISUALIZE_OUTPUT
 }
diff --git a/Engine/Shaders/Private/VirtualShadowMaps/Visualize.ush b/Engine/Shaders/Private/VirtualShadowMaps/Visualize.ush
new file mode 100644
--- /dev/null
+++ b/Engine/Shaders/Private/VirtualShadowMaps/Visualize.ush
@@ -0,0 +1,11 @@
+// Copyright Epic Games, Inc. All Rights Reserved.
+
+#pragma once
+
+// Must match the values in VirtualShadowMapVisualizationData.cpp
+#define VISUALIZE_NONE						0
+#define VISUALIZE_SHADOW_FACTOR				(1 << 0)
+#define VISUALIZE_CLIPMAP_OR_MIP			(1 << 1)
+#define VISUALIZE_VIRTUAL_PAGE				(1 << 2)
+#define VISUALIZE_CACHED_PAGE				(1 << 3)
+#define VISUALIZE_SMRT_RAY_COUNT			(1 << 4)
diff --git a/Engine/Source/Editor/CommonMenuExtensions/Private/CommonMenuExtensionsModule.cpp b/Engine/Source/Editor/CommonMenuExtensions/Private/CommonMenuExtensionsModule.cpp
--- a/Engine/Source/Editor/CommonMenuExtensions/Private/CommonMenuExtensionsModule.cpp
+++ b/Engine/Source/Editor/CommonMenuExtensions/Private/CommonMenuExtensionsModule.cpp
@@ -4,6 +4,7 @@
 #include "BufferVisualizationMenuCommands.h"
 #include "NaniteVisualizationMenuCommands.h"
 #include "LumenVisualizationMenuCommands.h"
+#include "VirtualShadowMapVisualizationMenuCommands.h"
 #include "ShowFlagMenuCommands.h"
 #include "Modules/ModuleManager.h"
 
@@ -14,12 +15,14 @@ void FCommonMenuExtensionsModule::StartupModule()
 	FBufferVisualizationMenuCommands::Register();
 	FNaniteVisualizationMenuCommands::Register();
 	FLumenVisualizationMenuCommands::Register();
+	FVirtualShadowMapVisualizationMenuCommands::Register();
 	FShowFlagMenuCommands::Register();
 }
 
 void FCommonMenuExtensionsModule::ShutdownModule()
 {
 	FShowFlagMenuCommands::Unregister();
+	FVirtualShadowMapVisualizationMenuCommands::Unregister();
 	FNaniteVisualizationMenuCommands::Unregister();
 	FLumenVisualizationMenuCommands::Unregister();
 	FBufferVisualizationMenuCommands::Unregister();
diff --git a/Engine/Source/Editor/CommonMenuExtensions/Private/VirtualShadowMapVisualizationMenuCommands.cpp b/Engine/Source/Editor/CommonMenuExtensions/Private/VirtualShadowMapVisualizationMenuCommands.cpp
new file mode 100644
--- /dev/null
+++ b/Engine/Source/Editor/CommonMenuExtensions/Private/VirtualShadowMapVisualizationMenuCommands.cpp
@@ -0,0 +1,149 @@
+// Copyright Epic Games, Inc. All Rights Reserved.
+
+#include "VirtualShadowMapVisualizationMenuCommands.h"
+#include "VirtualShadowMapVisualizationData.h"
+#include "Containers/UnrealString.h"
+#include "Framework/Commands/InputChord.h"
+#include "Materials/Material.h"
+#include "Internationalization/Text.h"
+#include "Templates/Function.h"
+#include "EditorStyleSet.h"
+#include "EditorViewportClient.h"
+
+int32 GVirtualShadowMapVisualizeAdvanced = 0;
+static FAutoConsoleVariableRef CVarVirtualShadowMapVisualizeAdvanced(
+	TEXT("r.Shadow.Virtual.Visualize.Advanced"),
+	GVirtualShadowMapVisualizeAdvanced,
+	TEXT("")
+);
+
+#define LOCTEXT_NAMESPACE "VirtualShadowMapVisualizationMenuCommands"
+
+FVirtualShadowMapVisualizationMenuCommands::FVirtualShadowMapVisualizationMenuCommands()
+	: TCommands<FVirtualShadowMapVisualizationMenuCommands>
+	(
+		TEXT("VirtualShadowMapVisualizationMenu"), // Context name for fast lookup
+		NSLOCTEXT("Contexts", "VirtualShadowMapVisualizationMenu", "VirtualShadowMa Visualization"), // Localized context name for displaying
+		NAME_None, // Parent context name.  
+		FEditorStyle::GetStyleSetName() // Icon Style Set
+	),
+	CommandMap()
+{
+}
+
+void FVirtualShadowMapVisualizationMenuCommands::BuildCommandMap()
+{
+	const FVirtualShadowMapVisualizationData& VisualizationData = GetVirtualShadowMapVisualizationData();
+	const FVirtualShadowMapVisualizationData::TModeMap& ModeMap = VisualizationData.GetModeMap();
+
+	CommandMap.Empty();
+	for (FVirtualShadowMapVisualizationData::TModeMap::TConstIterator It = ModeMap.CreateConstIterator(); It; ++It)
+	{
+		const FVirtualShadowMapVisualizationData::FModeRecord& Entry = It.Value();
+		FVirtualShadowMapVisualizationRecord& Record = CommandMap.Add(Entry.ModeName, FVirtualShadowMapVisualizationRecord());
+		Record.Name = Entry.ModeName;
+		Record.Command = FUICommandInfoDecl(
+			this->AsShared(),
+			Entry.ModeName,
+			Entry.ModeText,
+			Entry.ModeDesc)
+			.UserInterfaceType(EUserInterfaceActionType::RadioButton)
+			.DefaultChord(FInputChord());
+
+		switch (Entry.ModeType)
+		{
+		default:
+		case FVirtualShadowMapVisualizationData::FModeType::ProjectionStandard:
+			Record.Type = FVirtualShadowMapVisualizationType::ProjectionStandard;
+			break;
+
+		case FVirtualShadowMapVisualizationData::FModeType::ProjectionAdvanced:
+			Record.Type = FVirtualShadowMapVisualizationType::ProjectionAdvanced;
+			break;
+
+		case FVirtualShadowMapVisualizationData::FModeType::DebugStandard:
+			Record.Type = FVirtualShadowMapVisualizationType::DebugStandard;
+			break;
+		}
+	}
+}
+
+void FVirtualShadowMapVisualizationMenuCommands::BuildVisualisationSubMenu(FMenuBuilder& Menu)
+{
+	const bool bShowAdvanced = GVirtualShadowMapVisualizeAdvanced != 0;
+
+	const FVirtualShadowMapVisualizationMenuCommands& Commands = FVirtualShadowMapVisualizationMenuCommands::Get();
+	if (Commands.IsPopulated())
+	{
+		Menu.BeginSection("LevelViewportVirtualShadowMapVisualizationMode", LOCTEXT("VirtualShadowMapVisualizationHeader", "Visualization Mode"));
+
+		Commands.AddCommandTypeToMenu(Menu, FVirtualShadowMapVisualizationType::ProjectionStandard);
+		if (bShowAdvanced)
+		{
+			Menu.AddMenuSeparator();
+			Commands.AddCommandTypeToMenu(Menu, FVirtualShadowMapVisualizationType::ProjectionAdvanced);
+		}
+
+		//Menu.AddMenuSeparator();
+		//Commands.AddCommandTypeToMenu(Menu, FVirtualShadowMapVisualizationType::DebugStandard);
+
+		Menu.EndSection();
+	}
+}
+
+bool FVirtualShadowMapVisualizationMenuCommands::AddCommandTypeToMenu(FMenuBuilder& Menu, const FVirtualShadowMapVisualizationType Type) const
+{
+	bool bAddedCommands = false;
+
+	const TVirtualShadowMapVisualizationModeCommandMap& Commands = CommandMap;
+	for (TCommandConstIterator It = CreateCommandConstIterator(); It; ++It)
+	{
+		const FVirtualShadowMapVisualizationRecord& Record = It.Value();
+		if (Record.Type == Type)
+		{
+			Menu.AddMenuEntry(Record.Command, NAME_None, Record.Command->GetLabel());
+			bAddedCommands = true;
+		}
+	}
+
+	return bAddedCommands;
+}
+
+FVirtualShadowMapVisualizationMenuCommands::TCommandConstIterator FVirtualShadowMapVisualizationMenuCommands::CreateCommandConstIterator() const
+{
+	return CommandMap.CreateConstIterator();
+}
+
+void FVirtualShadowMapVisualizationMenuCommands::RegisterCommands()
+{
+	BuildCommandMap();
+}
+
+void FVirtualShadowMapVisualizationMenuCommands::BindCommands(FUICommandList& CommandList, const TSharedPtr<FEditorViewportClient>& Client) const
+{
+	// Map Virtual shadow map visualization mode actions
+	for (FVirtualShadowMapVisualizationMenuCommands::TCommandConstIterator It = FVirtualShadowMapVisualizationMenuCommands::Get().CreateCommandConstIterator(); It; ++It)
+	{
+		const FVirtualShadowMapVisualizationMenuCommands::FVirtualShadowMapVisualizationRecord& Record = It.Value();
+		CommandList.MapAction(
+			Record.Command,
+			FExecuteAction::CreateStatic<const TSharedPtr<FEditorViewportClient>&>(&FVirtualShadowMapVisualizationMenuCommands::ChangeVirtualShadowMapVisualizationMode, Client, Record.Name),
+			FCanExecuteAction(),
+			FIsActionChecked::CreateStatic<const TSharedPtr<FEditorViewportClient>&>(&FVirtualShadowMapVisualizationMenuCommands::IsVirtualShadowMapVisualizationModeSelected, Client, Record.Name)
+		);
+	}
+}
+
+void FVirtualShadowMapVisualizationMenuCommands::ChangeVirtualShadowMapVisualizationMode(const TSharedPtr<FEditorViewportClient>& Client, FName InName)
+{
+	check(Client.IsValid());
+	Client->ChangeVirtualShadowMapVisualizationMode(InName);
+}
+
+bool FVirtualShadowMapVisualizationMenuCommands::IsVirtualShadowMapVisualizationModeSelected(const TSharedPtr<FEditorViewportClient>& Client, FName InName)
+{
+	check(Client.IsValid());
+	return Client->IsVirtualShadowMapVisualizationModeSelected(InName);
+}
+
+#undef LOCTEXT_NAMESPACE
diff --git a/Engine/Source/Editor/CommonMenuExtensions/Public/VirtualShadowMapVisualizationMenuCommands.h b/Engine/Source/Editor/CommonMenuExtensions/Public/VirtualShadowMapVisualizationMenuCommands.h
new file mode 100644
--- /dev/null
+++ b/Engine/Source/Editor/CommonMenuExtensions/Public/VirtualShadowMapVisualizationMenuCommands.h
@@ -0,0 +1,66 @@
+// Copyright Epic Games, Inc. All Rights Reserved.
+
+#pragma once
+
+#include "Containers/Map.h"
+#include "UObject/NameTypes.h"
+#include "Templates/SharedPointer.h"
+#include "Framework/Commands/UICommandInfo.h"
+#include "Framework/MultiBox/MultiBoxBuilder.h"
+#include "Framework/Commands/Commands.h"
+
+class FEditorViewportClient;
+
+class COMMONMENUEXTENSIONS_API FVirtualShadowMapVisualizationMenuCommands : public TCommands<FVirtualShadowMapVisualizationMenuCommands>
+{
+public:
+	// Re-sort into categories for the menu
+	enum class FVirtualShadowMapVisualizationType : uint8
+	{
+		ProjectionStandard,
+		ProjectionAdvanced,
+		DebugStandard,
+	};
+
+	struct FVirtualShadowMapVisualizationRecord
+	{
+		FName Name;
+		TSharedPtr<FUICommandInfo> Command;
+		FVirtualShadowMapVisualizationType Type;
+
+		FVirtualShadowMapVisualizationRecord()
+		: Name()
+		, Command()
+		, Type(FVirtualShadowMapVisualizationType::ProjectionStandard)
+		{
+		}
+	};
+
+	typedef TMultiMap<FName, FVirtualShadowMapVisualizationRecord> TVirtualShadowMapVisualizationModeCommandMap;
+	typedef TVirtualShadowMapVisualizationModeCommandMap::TConstIterator TCommandConstIterator;
+
+	FVirtualShadowMapVisualizationMenuCommands();
+
+	TCommandConstIterator CreateCommandConstIterator() const;
+
+	static void BuildVisualisationSubMenu(FMenuBuilder& Menu);
+
+	virtual void RegisterCommands() override;
+
+	void BindCommands(FUICommandList& CommandList, const TSharedPtr<FEditorViewportClient>& Client) const;
+	
+	inline bool IsPopulated() const
+	{
+		return CommandMap.Num() > 0;
+	}
+
+private:
+	void BuildCommandMap();
+	bool AddCommandTypeToMenu(FMenuBuilder& Menu, const FVirtualShadowMapVisualizationType Type) const;
+
+	static void ChangeVirtualShadowMapVisualizationMode(const TSharedPtr<FEditorViewportClient>& Client, FName InName);
+	static bool IsVirtualShadowMapVisualizationModeSelected(const TSharedPtr<FEditorViewportClient>& Client, FName InName);
+
+private:
+	TVirtualShadowMapVisualizationModeCommandMap CommandMap;
+};
diff --git a/Engine/Source/Editor/DetailCustomizations/Private/DebugCameraControllerSettingsCustomization.cpp b/Engine/Source/Editor/DetailCustomizations/Private/DebugCameraControllerSettingsCustomization.cpp
--- a/Engine/Source/Editor/DetailCustomizations/Private/DebugCameraControllerSettingsCustomization.cpp
+++ b/Engine/Source/Editor/DetailCustomizations/Private/DebugCameraControllerSettingsCustomization.cpp
@@ -45,6 +45,7 @@ void FDebugCameraControllerSettingsViewModeIndexCustomization::CustomizeHeader(T
 	EnumRestriction->AddHiddenValue(ViewModeIndexEnum->GetNameStringByValue((uint8)EViewModeIndex::VMI_VisualizeBuffer));
 	EnumRestriction->AddHiddenValue(ViewModeIndexEnum->GetNameStringByValue((uint8)EViewModeIndex::VMI_VisualizeNanite));
 	EnumRestriction->AddHiddenValue(ViewModeIndexEnum->GetNameStringByValue((uint8)EViewModeIndex::VMI_VisualizeLumen));
+	EnumRestriction->AddHiddenValue(ViewModeIndexEnum->GetNameStringByValue((uint8)EViewModeIndex::VMI_VisualizeVirtualShadowMap));
 	EnumRestriction->AddHiddenValue(ViewModeIndexEnum->GetNameStringByValue((uint8)EViewModeIndex::VMI_StationaryLightOverlap));
 #if RHI_RAYTRACING
 	if (!GRHISupportsRayTracing || !GRHISupportsRayTracingShaders)
diff --git a/Engine/Source/Editor/EditorStyle/Private/StarshipStyle.cpp b/Engine/Source/Editor/EditorStyle/Private/StarshipStyle.cpp
--- a/Engine/Source/Editor/EditorStyle/Private/StarshipStyle.cpp
+++ b/Engine/Source/Editor/EditorStyle/Private/StarshipStyle.cpp
@@ -4395,6 +4395,7 @@ void FStarshipEditorStyle::FStyle::SetupLevelEditorStyle()
 		Set( "EditorViewport.VisualizeBufferMode",    new IMAGE_BRUSH_SVG("Starship/Common/BufferVisualization", Icon16x16 ) );
 		Set( "EditorViewport.VisualizeNaniteMode",    new IMAGE_BRUSH_SVG("Starship/Common/BufferVisualization", Icon16x16 ) );
 		Set( "EditorViewport.VisualizeLumenMode",    new IMAGE_BRUSH_SVG("Starship/Common/BufferVisualization", Icon16x16 ) );
+		Set( "EditorViewport.VisualizeVirtualShadowMapMode", new IMAGE_BRUSH_SVG("Starship/Common/BufferVisualization", Icon16x16 ) );
 		Set( "EditorViewport.CollisionPawn",          new IMAGE_BRUSH_SVG("Starship/Common/PlayerCollision", Icon16x16 ) );
 		Set( "EditorViewport.CollisionVisibility",    new IMAGE_BRUSH_SVG("Starship/Common/VisibilityCollision", Icon16x16 ) );
 		
diff --git a/Engine/Source/Editor/LevelEditor/Private/LevelViewportActions.cpp b/Engine/Source/Editor/LevelEditor/Private/LevelViewportActions.cpp
--- a/Engine/Source/Editor/LevelEditor/Private/LevelViewportActions.cpp
+++ b/Engine/Source/Editor/LevelEditor/Private/LevelViewportActions.cpp
@@ -11,6 +11,7 @@
 #include "BufferVisualizationData.h"
 #include "NaniteVisualizationData.h"
 #include "LumenVisualizationData.h"
+#include "VirtualShadowMapVisualizationData.h"
 #include "Bookmarks/BookmarkUI.h"
 #include "UnrealEdGlobals.h"
 
diff --git a/Engine/Source/Editor/LevelEditor/Private/SLevelViewport.cpp b/Engine/Source/Editor/LevelEditor/Private/SLevelViewport.cpp
--- a/Engine/Source/Editor/LevelEditor/Private/SLevelViewport.cpp
+++ b/Engine/Source/Editor/LevelEditor/Private/SLevelViewport.cpp
@@ -60,6 +60,7 @@
 #include "BufferVisualizationData.h"
 #include "NaniteVisualizationData.h"
 #include "LumenVisualizationData.h"
+#include "VirtualShadowMapVisualizationData.h"
 #include "Framework/Notifications/NotificationManager.h"
 #include "Widgets/Notifications/SNotificationList.h"
 #include "SActorPilotViewportToolbar.h"
@@ -73,6 +74,7 @@
 #include "BufferVisualizationMenuCommands.h"
 #include "NaniteVisualizationMenuCommands.h"
 #include "LumenVisualizationMenuCommands.h"
+#include "VirtualShadowMapVisualizationMenuCommands.h"
 #include "EditorLevelUtils.h"
 #include "Engine/LevelStreaming.h"
 #include "Editor/WorldBrowser/Public/WorldBrowserModule.h"
@@ -484,6 +486,7 @@ void SLevelViewport::ConstructLevelEditorViewportClient(FLevelEditorViewportInst
 	LevelViewportClient->CurrentBufferVisualizationMode = ViewportInstanceSettings.BufferVisualizationMode;
 	LevelViewportClient->CurrentNaniteVisualizationMode = ViewportInstanceSettings.NaniteVisualizationMode;
 	LevelViewportClient->CurrentLumenVisualizationMode = ViewportInstanceSettings.LumenVisualizationMode;
+	LevelViewportClient->CurrentVirtualShadowMapVisualizationMode = ViewportInstanceSettings.VirtualShadowMapVisualizationMode;
 	LevelViewportClient->CurrentRayTracingDebugVisualizationMode = ViewportInstanceSettings.RayTracingDebugVisualizationMode;
 	LevelViewportClient->ExposureSettings = ViewportInstanceSettings.ExposureSettings;
 	if(ViewportInstanceSettings.ViewportType == LVT_Perspective)
@@ -1575,6 +1578,7 @@ void SLevelViewport::BindViewCommands( FUICommandList& OutCommandList )
 	FBufferVisualizationMenuCommands::Get().BindCommands(OutCommandList, Client);
 	FNaniteVisualizationMenuCommands::Get().BindCommands(OutCommandList, Client);
 	FLumenVisualizationMenuCommands::Get().BindCommands(OutCommandList, Client);
+	FVirtualShadowMapVisualizationMenuCommands::Get().BindCommands(OutCommandList, Client);
 }
 
 
@@ -2273,6 +2277,7 @@ void SLevelViewport::SaveConfig(const FString& ConfigName) const
 		ViewportInstanceSettings.BufferVisualizationMode = LevelViewportClient->CurrentBufferVisualizationMode;
 		ViewportInstanceSettings.NaniteVisualizationMode = LevelViewportClient->CurrentNaniteVisualizationMode;
 		ViewportInstanceSettings.LumenVisualizationMode = LevelViewportClient->CurrentLumenVisualizationMode;
+		ViewportInstanceSettings.VirtualShadowMapVisualizationMode = LevelViewportClient->CurrentVirtualShadowMapVisualizationMode;
 		ViewportInstanceSettings.RayTracingDebugVisualizationMode = LevelViewportClient->CurrentRayTracingDebugVisualizationMode;
 		ViewportInstanceSettings.ExposureSettings = LevelViewportClient->ExposureSettings;
 		ViewportInstanceSettings.FOVAngle = LevelViewportClient->FOVAngle;
diff --git a/Engine/Source/Editor/LevelEditor/Private/SLevelViewportToolBar.cpp b/Engine/Source/Editor/LevelEditor/Private/SLevelViewportToolBar.cpp
--- a/Engine/Source/Editor/LevelEditor/Private/SLevelViewportToolBar.cpp
+++ b/Engine/Source/Editor/LevelEditor/Private/SLevelViewportToolBar.cpp
@@ -32,6 +32,7 @@
 #include "BufferVisualizationData.h"
 #include "NaniteVisualizationData.h"
 #include "LumenVisualizationData.h"
+#include "VirtualShadowMapVisualizationData.h"
 #include "FoliageType.h"
 #include "ShowFlagMenuCommands.h"
 #include "Bookmarks/BookmarkUI.h"
@@ -1599,6 +1600,30 @@ void SLevelViewportToolBar::FillViewMenu(UToolMenu* Menu)
 						);
 	}
 
+	{
+		FToolMenuSection& Section = Menu->FindOrAddSection("ViewMode");
+		Section.AddSubMenu(
+			"VisualizeVirtualShadowMapViewMode",
+			LOCTEXT("VisualizeVirtualShadowMapViewModeDisplayName", "Virtual Shadow Map"),
+			LOCTEXT("VirtualShadowMapVisualizationMenu_ToolTip", "Select a mode for virtual shadow map visualization. Select a light component in the world outliner to visualize that light."),
+			FNewMenuDelegate::CreateStatic(&FVirtualShadowMapVisualizationMenuCommands::BuildVisualisationSubMenu),
+			FUIAction(
+				FExecuteAction(),
+				FCanExecuteAction(),
+				FIsActionChecked::CreateLambda([this]()
+					{
+						const TSharedRef<SEditorViewport> ViewportRef = Viewport.Pin().ToSharedRef();
+						const TSharedPtr<FEditorViewportClient> ViewportClient = ViewportRef->GetViewportClient();
+						check(ViewportClient.IsValid());
+						return ViewportClient->IsViewModeEnabled(VMI_VisualizeVirtualShadowMap);
+					})
+			),
+			EUserInterfaceActionType::RadioButton,
+			/* bInOpenSubMenuOnClick = */ false,
+			FSlateIcon(FEditorStyle::GetStyleSetName(), "EditorViewport.VisualizeVirtualShadowMapMode")
+			);
+	}
+
 	{
 		FToolMenuSection& Section = Menu->AddSection("LevelViewportLandscape", LOCTEXT("LandscapeHeader", "Landscape"), InsertPosition);
 
diff --git a/Engine/Source/Editor/LevelEditor/Public/LevelViewportActions.h b/Engine/Source/Editor/LevelEditor/Public/LevelViewportActions.h
--- a/Engine/Source/Editor/LevelEditor/Public/LevelViewportActions.h
+++ b/Engine/Source/Editor/LevelEditor/Public/LevelViewportActions.h
@@ -10,6 +10,7 @@
 #include "BufferVisualizationMenuCommands.h"
 #include "NaniteVisualizationMenuCommands.h"
 #include "LumenVisualizationMenuCommands.h"
+#include "VirtualShadowMapVisualizationMenuCommands.h"
 
 /**
  * Public identifiers for the viewport layouts available in LevelViewportLayoutX.h files
diff --git a/Engine/Source/Editor/UnrealEd/Classes/Settings/LevelEditorViewportSettings.h b/Engine/Source/Editor/UnrealEd/Classes/Settings/LevelEditorViewportSettings.h
--- a/Engine/Source/Editor/UnrealEd/Classes/Settings/LevelEditorViewportSettings.h
+++ b/Engine/Source/Editor/UnrealEd/Classes/Settings/LevelEditorViewportSettings.h
@@ -88,6 +88,7 @@ struct UNREALED_API FLevelEditorViewportInstanceSettings
 		, BufferVisualizationMode()
 		, NaniteVisualizationMode()
 		, LumenVisualizationMode()
+		, VirtualShadowMapVisualizationMode()
 		, ExposureSettings()
 		, FOVAngle(EditorViewportDefs::DefaultPerspectiveFOVAngle)
 		, FarViewPlane(0)
@@ -136,6 +137,10 @@ struct UNREALED_API FLevelEditorViewportInstanceSettings
 	UPROPERTY(config)
 	FName LumenVisualizationMode;
 
+	/** The virtual shadow map visualization mode for the viewport. */
+	UPROPERTY(config)
+	FName VirtualShadowMapVisualizationMode;
+
 	/** The buffer visualization mode for the viewport. */
 	UPROPERTY(config)
 	FName RayTracingDebugVisualizationMode;
diff --git a/Engine/Source/Editor/UnrealEd/Private/EditorViewportClient.cpp b/Engine/Source/Editor/UnrealEd/Private/EditorViewportClient.cpp
--- a/Engine/Source/Editor/UnrealEd/Private/EditorViewportClient.cpp
+++ b/Engine/Source/Editor/UnrealEd/Private/EditorViewportClient.cpp
@@ -55,6 +55,7 @@
 #include "BufferVisualizationData.h"
 #include "NaniteVisualizationData.h"
 #include "LumenVisualizationData.h"
+#include "VirtualShadowMapVisualizationData.h"
 #include "UnrealWidget.h"
 #include "EdModeInteractiveToolsContext.h"
 
@@ -385,6 +386,7 @@ FEditorViewportClient::FEditorViewportClient(FEditorModeTools* InModeTools, FPre
 	, CurrentBufferVisualizationMode(NAME_None)
 	, CurrentNaniteVisualizationMode(NAME_None)
 	, CurrentLumenVisualizationMode(NAME_None)
+	, CurrentVirtualShadowMapVisualizationMode(NAME_None)
 	, CurrentRayTracingDebugVisualizationMode(NAME_None)
 	, FramesSinceLastDraw(0)
 	, ViewIndex(INDEX_NONE)
@@ -2644,6 +2646,23 @@ FText FEditorViewportClient::GetCurrentLumenVisualizationModeDisplayName() const
 	return GetLumenVisualizationData().GetModeDisplayName(CurrentLumenVisualizationMode);
 }
 
+void FEditorViewportClient::ChangeVirtualShadowMapVisualizationMode(FName InName)
+{
+	SetViewMode(VMI_VisualizeVirtualShadowMap);
+	CurrentVirtualShadowMapVisualizationMode = InName;
+}
+
+bool FEditorViewportClient::IsVirtualShadowMapVisualizationModeSelected(FName InName) const
+{
+	return IsViewModeEnabled(VMI_VisualizeVirtualShadowMap) && CurrentVirtualShadowMapVisualizationMode == InName;
+}
+
+FText FEditorViewportClient::GetCurrentVirtualShadowMapVisualizationModeDisplayName() const
+{
+	checkf(IsViewModeEnabled(VMI_VisualizeVirtualShadowMap), TEXT("In order to call GetCurrentVirtualShadowMapVisualizationMode(), first you must set ViewMode to VMI_VisualizeVirtualShadowMap."));
+	return GetVirtualShadowMapVisualizationData().GetModeDisplayName(CurrentVirtualShadowMapVisualizationMode);
+}
+
 bool FEditorViewportClient::IsVisualizeCalibrationMaterialEnabled() const
 {
 	// Get the list of requested buffers from the console
@@ -2689,7 +2708,7 @@ bool FEditorViewportClient::SupportsPreviewResolutionFraction() const
 	}
 
 	// Don't do preview screen percentage in certain cases.
-	if (EngineShowFlags.VisualizeBuffer || EngineShowFlags.VisualizeNanite || IsVisualizeCalibrationMaterialEnabled())
+	if (EngineShowFlags.VisualizeBuffer || EngineShowFlags.VisualizeNanite || EngineShowFlags.VisualizeVirtualShadowMap || IsVisualizeCalibrationMaterialEnabled())
 	{
 		return false;
 	}
@@ -3754,6 +3773,7 @@ void FEditorViewportClient::SetupViewForRendering(FSceneViewFamily& ViewFamily,
 	View.CurrentBufferVisualizationMode = CurrentBufferVisualizationMode;
 	View.CurrentNaniteVisualizationMode = CurrentNaniteVisualizationMode;
 	View.CurrentLumenVisualizationMode = CurrentLumenVisualizationMode;
+	View.CurrentVirtualShadowMapVisualizationMode = CurrentVirtualShadowMapVisualizationMode;
 #if RHI_RAYTRACING
 	View.CurrentRayTracingDebugVisualizationMode = CurrentRayTracingDebugVisualizationMode;
 #endif
@@ -3880,8 +3900,9 @@ void FEditorViewportClient::Draw(FViewport* InViewport, FCanvas* Canvas)
 
 	const bool bVisualizeBufferEnabled = CurrentViewMode == VMI_VisualizeBuffer && CurrentBufferVisualizationMode != NAME_None;
 	const bool bVisualizeNaniteEnabled = CurrentViewMode == VMI_VisualizeNanite && CurrentNaniteVisualizationMode != NAME_None;
+	const bool bVisualizeVirtualShadowMapEnabled = CurrentViewMode == VMI_VisualizeVirtualShadowMap && CurrentVirtualShadowMapVisualizationMode != NAME_None;
 	const bool bRayTracingDebugEnabled = CurrentViewMode == VMI_RayTracingDebug && CurrentRayTracingDebugVisualizationMode != NAME_None;
-	const bool bCanDisableTonemapper = bVisualizeBufferEnabled || bVisualizeNaniteEnabled || (bRayTracingDebugEnabled && !FRayTracingDebugVisualizationMenuCommands::DebugModeShouldBeTonemapped(CurrentRayTracingDebugVisualizationMode));
+	const bool bCanDisableTonemapper = bVisualizeBufferEnabled || bVisualizeNaniteEnabled || bVisualizeVirtualShadowMapEnabled || (bRayTracingDebugEnabled && !FRayTracingDebugVisualizationMenuCommands::DebugModeShouldBeTonemapped(CurrentRayTracingDebugVisualizationMode));
 	
 	EngineShowFlagOverride(ESFIM_Editor, ViewFamily.ViewMode, ViewFamily.EngineShowFlags, bCanDisableTonemapper);
 	EngineShowFlagOrthographicOverride(IsPerspective(), ViewFamily.EngineShowFlags);
diff --git a/Engine/Source/Editor/UnrealEd/Private/EditorViewportCommands.cpp b/Engine/Source/Editor/UnrealEd/Private/EditorViewportCommands.cpp
--- a/Engine/Source/Editor/UnrealEd/Private/EditorViewportCommands.cpp
+++ b/Engine/Source/Editor/UnrealEd/Private/EditorViewportCommands.cpp
@@ -117,6 +117,7 @@ void FEditorViewportCommands::RegisterCommands()
 	UI_COMMAND( VisualizeBufferMode, "Buffer Visualization View Mode", "Renders a set of selected post process materials, which visualize various intermediate render buffers (material attributes)", EUserInterfaceActionType::RadioButton, FInputChord());
 	UI_COMMAND( VisualizeNaniteMode, "Nanite Visualization View Mode", "Visualizes various rendering aspects of Nanite.", EUserInterfaceActionType::RadioButton, FInputChord());
 	UI_COMMAND( VisualizeLumenMode, "Lumen Visualization View Mode", "Visualizes Lumen debug views.", EUserInterfaceActionType::RadioButton, FInputChord());
+	UI_COMMAND( VisualizeVirtualShadowMapMode, "Virtual Shadow Map Visualization View Mode", "Visualizes various rendering aspects of virtual shadow maps.", EUserInterfaceActionType::RadioButton, FInputChord());
 
 	UI_COMMAND( ReflectionOverrideMode, "Reflections View Mode", "Renders the scene with reflections only", EUserInterfaceActionType::RadioButton, FInputChord() );
 	UI_COMMAND( CollisionPawn, "Player Collision", "Renders player collision visualization", EUserInterfaceActionType::RadioButton, FInputChord() );
diff --git a/Engine/Source/Editor/UnrealEd/Private/SEditorViewport.cpp b/Engine/Source/Editor/UnrealEd/Private/SEditorViewport.cpp
--- a/Engine/Source/Editor/UnrealEd/Private/SEditorViewport.cpp
+++ b/Engine/Source/Editor/UnrealEd/Private/SEditorViewport.cpp
@@ -428,6 +428,7 @@ void SEditorViewport::BindCommands()
 	MAP_VIEWMODE_ACTION( Commands.VisualizeBufferMode, VMI_VisualizeBuffer );
 	MAP_VIEWMODE_ACTION( Commands.VisualizeNaniteMode, VMI_VisualizeNanite );
 	MAP_VIEWMODE_ACTION( Commands.VisualizeLumenMode, VMI_VisualizeLumen );
+	MAP_VIEWMODE_ACTION( Commands.VisualizeVirtualShadowMapMode, VMI_VisualizeVirtualShadowMap );
 	MAP_VIEWMODE_ACTION( Commands.CollisionPawn, VMI_CollisionPawn);
 	MAP_VIEWMODE_ACTION( Commands.CollisionVisibility, VMI_CollisionVisibility);
 
diff --git a/Engine/Source/Editor/UnrealEd/Private/SEditorViewportViewMenu.cpp b/Engine/Source/Editor/UnrealEd/Private/SEditorViewportViewMenu.cpp
--- a/Engine/Source/Editor/UnrealEd/Private/SEditorViewportViewMenu.cpp
+++ b/Engine/Source/Editor/UnrealEd/Private/SEditorViewportViewMenu.cpp
@@ -53,6 +53,10 @@ FText SEditorViewportViewMenu::GetViewMenuLabel() const
 		{
 			Label = ViewportClient->GetCurrentLumenVisualizationModeDisplayName();
 		}
+		else if (ViewMode == VMI_VisualizeVirtualShadowMap)
+		{
+			Label = ViewportClient->GetCurrentVirtualShadowMapVisualizationModeDisplayName();
+		}
 		// For any other category, return its own name
 		else
 		{
diff --git a/Engine/Source/Editor/UnrealEd/Public/EditorViewportClient.h b/Engine/Source/Editor/UnrealEd/Public/EditorViewportClient.h
--- a/Engine/Source/Editor/UnrealEd/Public/EditorViewportClient.h
+++ b/Engine/Source/Editor/UnrealEd/Public/EditorViewportClient.h
@@ -1279,6 +1279,26 @@ public:
 	 */
 	FText GetCurrentLumenVisualizationModeDisplayName() const;
 
+	/**
+	* Changes the virtual shadow map visualization mode for this viewport.
+	*
+	* @param InName	The ID of the required visualization mode
+	*/
+	void ChangeVirtualShadowMapVisualizationMode(FName InName);
+
+	/**
+	* Checks if a virtual shadow map visualization mode is selected.
+	*
+	* @param InName	The ID of the required visualization mode
+	* @return	true if the supplied virtual shadow map visualization mode is checked
+	*/
+	bool IsVirtualShadowMapVisualizationModeSelected(FName InName) const;
+
+	/**
+	* Returns the FText display name associated with CurrentVirtualShadowMapVisualizationMode.
+	*/
+	FText GetCurrentVirtualShadowMapVisualizationModeDisplayName() const;
+
 	/**
 	* Returns whether visualize debug material is enabled.
 	*/
@@ -1601,6 +1621,7 @@ public:
 	FName CurrentBufferVisualizationMode;
 	FName CurrentNaniteVisualizationMode;
 	FName CurrentLumenVisualizationMode;
+	FName CurrentVirtualShadowMapVisualizationMode;
 
 	FName CurrentRayTracingDebugVisualizationMode;
 
diff --git a/Engine/Source/Editor/UnrealEd/Public/EditorViewportCommands.h b/Engine/Source/Editor/UnrealEd/Public/EditorViewportCommands.h
--- a/Engine/Source/Editor/UnrealEd/Public/EditorViewportCommands.h
+++ b/Engine/Source/Editor/UnrealEd/Public/EditorViewportCommands.h
@@ -136,6 +136,9 @@ public:
 	/** Changes the viewport to visualize Lumen debug views */
 	TSharedPtr< FUICommandInfo > VisualizeLumenMode;
 
+	/** Changes the viewport to visualize virtual shadow map content */
+	TSharedPtr< FUICommandInfo > VisualizeVirtualShadowMapMode;
+
 	/** Collision Draw Mode */
 	TSharedPtr< FUICommandInfo > CollisionPawn;
 	TSharedPtr< FUICommandInfo > CollisionVisibility;
diff --git a/Engine/Source/Runtime/Engine/Classes/Engine/EngineBaseTypes.h b/Engine/Source/Runtime/Engine/Classes/Engine/EngineBaseTypes.h
--- a/Engine/Source/Runtime/Engine/Classes/Engine/EngineBaseTypes.h
+++ b/Engine/Source/Runtime/Engine/Classes/Engine/EngineBaseTypes.h
@@ -913,6 +913,9 @@ enum EViewModeIndex
 	/** Visualize Lumen debug views */
 	VMI_VisualizeLumen = 31 UMETA(DisplayName = "Lumen Visualization"),
 
+	/** Visualize virtual shadow map */
+	VMI_VisualizeVirtualShadowMap = 32 UMETA(DisplayName = "Virtual Shadow Map Visualization"),
+
 	VMI_Max UMETA(Hidden),
 
 	// VMI_Unknown - The value assigned to VMI_Unknown must be the highest possible of any member of EViewModeIndex, or GetViewModeName might seg-fault
diff --git a/Engine/Source/Runtime/Engine/Classes/Engine/GameViewportClient.h b/Engine/Source/Runtime/Engine/Classes/Engine/GameViewportClient.h
--- a/Engine/Source/Runtime/Engine/Classes/Engine/GameViewportClient.h
+++ b/Engine/Source/Runtime/Engine/Classes/Engine/GameViewportClient.h
@@ -650,6 +650,9 @@ protected:
 	void SetCurrentLumenVisualizationMode(FName NewLumenVisualizationMode) { CurrentLumenVisualizationMode = NewLumenVisualizationMode; }
 	FName GetCurrentLumenVisualizationMode() const { return CurrentLumenVisualizationMode; }
 
+	void SetCurrentVirtualShadowMapVisualizationMode(FName NewVirtualShadowMapVisualizationMode) { CurrentVirtualShadowMapVisualizationMode = NewVirtualShadowMapVisualizationMode; }
+	FName GetCurrentVirtualShadowMapVisualizationMode() const { return CurrentVirtualShadowMapVisualizationMode; }
+
 	bool HasAudioFocus() const { return bHasAudioFocus; }
 
 	/** Updates CSVProfiler camera stats */
@@ -957,6 +960,9 @@ private:
 	/** Current Lumen visualization mode for this game viewport */
 	FName CurrentLumenVisualizationMode;
 
+	/** Current virtual shadow map visualization mode for this game viewport */
+	FName CurrentVirtualShadowMapVisualizationMode;
+
 	/** Weak pointer to the highres screenshot dialog if it's open */
 	TWeakPtr<SWindow> HighResScreenshotDialog;
 
diff --git a/Engine/Source/Runtime/Engine/Private/GameViewportClient.cpp b/Engine/Source/Runtime/Engine/Private/GameViewportClient.cpp
--- a/Engine/Source/Runtime/Engine/Private/GameViewportClient.cpp
+++ b/Engine/Source/Runtime/Engine/Private/GameViewportClient.cpp
@@ -45,6 +45,7 @@
 #include "BufferVisualizationData.h"
 #include "NaniteVisualizationData.h"
 #include "LumenVisualizationData.h"
+#include "VirtualShadowMapVisualizationData.h"
 #include "GameFramework/InputSettings.h"
 #include "Components/LineBatchComponent.h"
 #include "Debug/DebugDrawService.h"
@@ -230,6 +231,7 @@ UGameViewportClient::UGameViewportClient(const FObjectInitializer& ObjectInitial
 	, CurrentBufferVisualizationMode(NAME_None)
 	, CurrentNaniteVisualizationMode(NAME_None)
 	, CurrentLumenVisualizationMode(NAME_None)
+	, CurrentVirtualShadowMapVisualizationMode(NAME_None)
 	, HighResScreenshotDialog(nullptr)
 	, bUseSoftwareCursorWidgets(true)
 	, bIgnoreInput(false)
@@ -323,6 +325,7 @@ UGameViewportClient::UGameViewportClient(FVTableHelper& Helper)
 	, CurrentBufferVisualizationMode(NAME_None)
 	, CurrentNaniteVisualizationMode(NAME_None)
 	, CurrentLumenVisualizationMode(NAME_None)
+	, CurrentVirtualShadowMapVisualizationMode(NAME_None)
 	, HighResScreenshotDialog(nullptr)
 	, bIgnoreInput(false)
 	, MouseCaptureMode(EMouseCaptureMode::CapturePermanently)
@@ -1495,6 +1498,7 @@ void UGameViewportClient::Draw(FViewport* InViewport, FCanvas* SceneCanvas)
 					View->CurrentBufferVisualizationMode = CurrentBufferVisualizationMode;
 					View->CurrentNaniteVisualizationMode = CurrentNaniteVisualizationMode;
 					View->CurrentLumenVisualizationMode = CurrentLumenVisualizationMode;
+					View->CurrentVirtualShadowMapVisualizationMode = CurrentVirtualShadowMapVisualizationMode;
 
 					View->CameraConstrainedViewRect = View->UnscaledViewRect;
 
diff --git a/Engine/Source/Runtime/Engine/Private/SceneView.cpp b/Engine/Source/Runtime/Engine/Private/SceneView.cpp
--- a/Engine/Source/Runtime/Engine/Private/SceneView.cpp
+++ b/Engine/Source/Runtime/Engine/Private/SceneView.cpp
@@ -17,6 +17,7 @@
 #include "BufferVisualizationData.h"
 #include "NaniteVisualizationData.h"
 #include "LumenVisualizationData.h"
+#include "VirtualShadowMapVisualizationData.h"
 #include "Interfaces/Interface_PostProcessVolume.h"
 #include "Engine/TextureCube.h"
 #include "StereoRendering.h"
@@ -33,6 +34,7 @@
 DEFINE_LOG_CATEGORY(LogBufferVisualization);
 DEFINE_LOG_CATEGORY(LogNaniteVisualization);
 DEFINE_LOG_CATEGORY(LogLumenVisualization);
+DEFINE_LOG_CATEGORY(LogVirtualShadowMapVisualization);
 DEFINE_LOG_CATEGORY(LogMultiView);
 
 DECLARE_CYCLE_STAT(TEXT("StartFinalPostprocessSettings"), STAT_StartFinalPostprocessSettings, STATGROUP_Engine);
diff --git a/Engine/Source/Runtime/Engine/Private/ShowFlags.cpp b/Engine/Source/Runtime/Engine/Private/ShowFlags.cpp
--- a/Engine/Source/Runtime/Engine/Private/ShowFlags.cpp
+++ b/Engine/Source/Runtime/Engine/Private/ShowFlags.cpp
@@ -335,6 +335,9 @@ void ApplyViewMode(EViewModeIndex ViewModeIndex, bool bPerspective, FEngineShowF
 		case VMI_VisualizeLumen:
 			bPostProcessing = true;
 			break;
+		case VMI_VisualizeVirtualShadowMap:
+			bPostProcessing = true;
+			break;
 		case VMI_ReflectionOverride:
 			bPostProcessing = true;
 			break;
@@ -366,6 +369,7 @@ void ApplyViewMode(EViewModeIndex ViewModeIndex, bool bPerspective, FEngineShowF
 	EngineShowFlags.SetVisualizeBuffer(ViewModeIndex == VMI_VisualizeBuffer);
 	EngineShowFlags.SetVisualizeNanite(ViewModeIndex == VMI_VisualizeNanite);
 	EngineShowFlags.SetVisualizeLumen(ViewModeIndex == VMI_VisualizeLumen);
+	EngineShowFlags.SetVisualizeVirtualShadowMap(ViewModeIndex == VMI_VisualizeVirtualShadowMap);
 	EngineShowFlags.SetVisualizeLightCulling(ViewModeIndex == VMI_LightComplexity);
 	EngineShowFlags.SetShaderComplexity(ViewModeIndex == VMI_ShaderComplexity || ViewModeIndex == VMI_QuadOverdraw || ViewModeIndex == VMI_ShaderComplexityWithQuadOverdraw);
 	EngineShowFlags.SetQuadOverdraw(ViewModeIndex == VMI_QuadOverdraw);
@@ -683,6 +687,10 @@ EViewModeIndex FindViewMode(const FEngineShowFlags& EngineShowFlags)
 	{
 		return VMI_VisualizeLumen;
 	}
+	else if (EngineShowFlags.VisualizeVirtualShadowMap)
+	{
+		return VMI_VisualizeVirtualShadowMap;
+	}
 	else if (EngineShowFlags.StationaryLightOverlap)
 	{
 		return VMI_StationaryLightOverlap;
@@ -817,6 +825,7 @@ const TCHAR* GetViewModeName(EViewModeIndex ViewModeIndex)
 		case VMI_VisualizeBuffer:			return TEXT("VisualizeBuffer");
 		case VMI_VisualizeNanite:			return TEXT("VisualizeNanite");
 		case VMI_VisualizeLumen:			return TEXT("VisualizeLumen");
+		case VMI_VisualizeVirtualShadowMap:	return TEXT("VisualizeVirtualShadowMap");
 		case VMI_RayTracingDebug:			return TEXT("RayTracingDebug");
 		case VMI_PathTracing:				return TEXT("PathTracing");
 		case VMI_CollisionPawn:				return TEXT("CollisionPawn");
diff --git a/Engine/Source/Runtime/Engine/Private/UnrealEngine.cpp b/Engine/Source/Runtime/Engine/Private/UnrealEngine.cpp
--- a/Engine/Source/Runtime/Engine/Private/UnrealEngine.cpp
+++ b/Engine/Source/Runtime/Engine/Private/UnrealEngine.cpp
@@ -173,6 +173,7 @@ UnrealEngine.cpp: Implements the UEngine class and helpers.
 #include "UnrealExporter.h"
 #include "BufferVisualizationData.h"
 #include "NaniteVisualizationData.h"
+#include "VirtualShadowMapVisualizationData.h"
 #include "Misc/HotReloadInterface.h"
 #include "Widgets/Testing/STestSuite.h"
 #include "Engine/DemoNetDriver.h"
@@ -1990,6 +1991,7 @@ void UEngine::Init(IEngineLoop* InEngineLoop)
 	// Initialise buffer visualization system data
 	GetBufferVisualizationData().Initialize();
 	GetNaniteVisualizationData().Initialize();
+	GetVirtualShadowMapVisualizationData().Initialize();
 
 	// Initialize Portal services
 	if (!IsRunningCommandlet() && !IsRunningDedicatedServer())
diff --git a/Engine/Source/Runtime/Engine/Private/ViewModeNames.cpp b/Engine/Source/Runtime/Engine/Private/ViewModeNames.cpp
--- a/Engine/Source/Runtime/Engine/Private/ViewModeNames.cpp
+++ b/Engine/Source/Runtime/Engine/Private/ViewModeNames.cpp
@@ -163,6 +163,11 @@ TArray<FText> FillViewModeDisplayNames()
 			ViewModeDisplayNames.Emplace(LOCTEXT("UViewModeUtils_VMI_VisualizeLumen", "Lumen Visualization"));
 		}
 
+		else if (ViewModeIndex == VMI_VisualizeVirtualShadowMap)
+		{
+			ViewModeDisplayNames.Emplace(LOCTEXT("UViewModeUtils_VMI_VisualizeVirtualShadowMap", "Virtual Shadow Map Visualization"));
+		}
+
 		// VMI_Max
 		else if (ViewModeIndex == VMI_Max)
 		{
@@ -270,6 +275,10 @@ TArray<const FSlateBrush*> FillViewModeDisplayIcons()
 		{
 			ViewModeDisplayIcons.Emplace(FAppStyle::Get().GetBrush("EditorViewport.VisualizeLumenMode"));
 		}
+		else if (ViewModeIndex == VMI_VisualizeVirtualShadowMap)
+		{
+			ViewModeDisplayIcons.Emplace(FAppStyle::Get().GetBrush("EditorViewport.VisualizeVirtualShadowMapMode"));
+		}
 
 		//	VMI_VoxelLighting = 13,
 	
diff --git a/Engine/Source/Runtime/Engine/Private/VirtualShadowMapVisualizationData.cpp b/Engine/Source/Runtime/Engine/Private/VirtualShadowMapVisualizationData.cpp
new file mode 100644
--- /dev/null
+++ b/Engine/Source/Runtime/Engine/Private/VirtualShadowMapVisualizationData.cpp
@@ -0,0 +1,215 @@
+// Copyright Epic Games, Inc. All Rights Reserved.
+
+#include "VirtualShadowMapVisualizationData.h"
+#include "HAL/IConsoleManager.h"
+#include "Materials/Material.h"
+#include "Misc/ConfigCacheIni.h"
+
+#define LOCTEXT_NAMESPACE "FVirtualShadowMapVisualizationData"
+
+static FVirtualShadowMapVisualizationData GVirtualShadowMapVisualizationData;
+
+// Must match values in Shadows/VirtualShadowMaps/Visualize.ush
+#define VISUALIZE_NONE						0
+#define VISUALIZE_SHADOW_FACTOR				(1 << 0)
+#define VISUALIZE_CLIPMAP_OR_MIP			(1 << 1)
+#define VISUALIZE_VIRTUAL_PAGE				(1 << 2)
+#define VISUALIZE_CACHED_PAGE				(1 << 3)
+#define VISUALIZE_SMRT_RAY_COUNT			(1 << 4)
+
+
+void FVirtualShadowMapVisualizationData::Initialize()
+{
+	if (!bIsInitialized)
+	{
+		// NOTE: The first parameter determines the console command parameter. "none", "off" and "list" are reserved
+		AddVisualizationMode(
+			TEXT("mask"),
+			LOCTEXT("ShadowMask", "Shadow Mask"),
+			LOCTEXT("ShadowMaskDesc", "The final shadow mask that is used by shading"),
+			FModeType::ProjectionStandard,
+			VISUALIZE_SHADOW_FACTOR);
+
+		AddVisualizationMode(
+			TEXT("mip"),
+			LOCTEXT("ClipmapOrMip", "Clipmap/Mip Level"),
+			LOCTEXT("ClipmapOrMipDesc", "The chosen clipmap (for directional lights) or mip (for local lights) level"),
+			FModeType::ProjectionStandard,
+			VISUALIZE_CLIPMAP_OR_MIP);
+
+		AddVisualizationMode(
+			TEXT("vpage"),
+			LOCTEXT("VirtualPage", "Virtual Page"),
+			LOCTEXT("VirtualPageDesc", "Visualization of the virtual page address"),
+			FModeType::ProjectionStandard,
+			VISUALIZE_VIRTUAL_PAGE);
+
+		AddVisualizationMode(
+			TEXT("cache"),
+			LOCTEXT("CachedPage", "Cached Page"),
+			LOCTEXT("CachedPageDesc", "Cached pages are tinted green, uncached are red. Pages where only the static page is cached (dynamic uncached) are blue."),
+			FModeType::ProjectionStandard,
+			VISUALIZE_CACHED_PAGE);
+
+		AddVisualizationMode(
+			TEXT("raycount"),
+			LOCTEXT("SMRTRayCount", "SMRT Ray Count"),
+			LOCTEXT("SMRTRayCountDesc", "Rays evaluated per pixel: red is more, green is fewer. Penumbra regions require more rays and are more expensive."),
+			FModeType::ProjectionAdvanced,
+			VISUALIZE_SMRT_RAY_COUNT);
+
+		ConfigureConsoleCommand();
+
+		bIsInitialized = true;
+	}
+}
+
+void FVirtualShadowMapVisualizationData::ConfigureConsoleCommand()
+{
+	FString AvailableVisualizationModes;
+	for (TModeMap::TConstIterator It = ModeMap.CreateConstIterator(); It; ++It)
+	{
+		const FModeRecord& Record = It.Value();
+		AvailableVisualizationModes += FString(TEXT("\n  "));
+		AvailableVisualizationModes += Record.ModeString;
+	}
+
+	ConsoleDocumentationVisualizationMode = TEXT("When the viewport view-mode is set to 'Virtual Shadow Map Visualization', this command specifies which of the various channels to display. Values entered other than the allowed values shown below will be ignored.");
+	ConsoleDocumentationVisualizationMode += AvailableVisualizationModes;
+
+	IConsoleManager::Get().RegisterConsoleVariable(
+		GetVisualizeConsoleCommandName(),
+		TEXT(""),
+		*ConsoleDocumentationVisualizationMode,
+		ECVF_Cheat);
+}
+
+void FVirtualShadowMapVisualizationData::SetActiveMode(int32 ModeID, const FName& ModeName)
+{
+	ActiveVisualizationModeID = ModeID;
+	ActiveVisualizationModeName = ModeName;
+}
+
+bool FVirtualShadowMapVisualizationData::IsActive() const
+{
+	if (!IsInitialized())
+	{
+		return false;
+	}
+
+	if (GetActiveModeID() == INDEX_NONE)
+	{
+		return false;
+	}
+
+	return true;
+}
+
+bool FVirtualShadowMapVisualizationData::Update(const FName& InViewMode)
+{
+	bool bForceShowFlag = false;
+
+	if (!IsInitialized())
+	{
+		return false;
+	}
+
+	SetActiveMode(INDEX_NONE, NAME_None);
+			
+	// Check if the console command is set (overrides the editor)
+	{
+		static IConsoleVariable* ICVarVisualize = IConsoleManager::Get().FindConsoleVariable(GetVisualizeConsoleCommandName());
+		if (ICVarVisualize)
+		{
+			const FString ConsoleVisualizationMode = ICVarVisualize->GetString();
+			if (!ConsoleVisualizationMode.IsEmpty())
+			{
+				 if (ConsoleVisualizationMode == TEXT("off") || ConsoleVisualizationMode == TEXT("none"))
+				{
+					// Disable visualization
+				}
+				else
+				{
+					const FName  ModeName = FName(*ConsoleVisualizationMode);
+					const int32  ModeID   = GetModeID(ModeName);
+					if (ModeID == INDEX_NONE)
+					{
+						UE_LOG(LogVirtualShadowMapVisualization, Warning, TEXT("Unknown virtual shadow map visualization mode '%s'"), *ConsoleVisualizationMode);
+					}
+					else
+					{
+						SetActiveMode(ModeID, ModeName);
+						bForceShowFlag = true;
+					}
+				}
+			}
+		}
+	}
+
+	// Check the view mode state (set by editor).
+	if (ActiveVisualizationModeID == INDEX_NONE && InViewMode != NAME_None)
+	{
+		const int32 ModeID = GetModeID(InViewMode);
+		if (ModeID != INDEX_NONE)
+		{
+			SetActiveMode(ModeID, InViewMode);
+		}
+	}
+
+	return bForceShowFlag;
+}
+
+void FVirtualShadowMapVisualizationData::AddVisualizationMode(
+	const TCHAR* ModeString,
+	const FText& ModeText,
+	const FText& ModeDesc,
+	const FModeType ModeType,
+	int32 ModeID
+)
+{
+	const FName ModeName = FName(ModeString);
+
+	FModeRecord& Record	= ModeMap.Emplace(ModeName);
+	Record.ModeString			= FString(ModeString);
+	Record.ModeName				= ModeName;
+	Record.ModeText				= ModeText;
+	Record.ModeDesc				= ModeDesc;
+	Record.ModeType				= ModeType;
+	Record.ModeID				= ModeID;
+}
+
+FText FVirtualShadowMapVisualizationData::GetModeDisplayName(const FName& InModeName) const
+{
+	if (const FModeRecord* Record = ModeMap.Find(InModeName))
+	{
+		return Record->ModeText;
+	}
+	else
+	{
+		return FText::GetEmpty();
+	}
+}
+
+int32 FVirtualShadowMapVisualizationData::GetModeID(const FName& InModeName) const
+{
+	if (const FModeRecord* Record = ModeMap.Find(InModeName))
+	{
+		return Record->ModeID;
+	}
+	else
+	{
+		return INDEX_NONE;
+	}
+}
+
+FVirtualShadowMapVisualizationData& GetVirtualShadowMapVisualizationData()
+{
+	if (!GVirtualShadowMapVisualizationData.IsInitialized())
+	{
+		GVirtualShadowMapVisualizationData.Initialize();
+	}
+
+	return GVirtualShadowMapVisualizationData;
+}
+
+#undef LOCTEXT_NAMESPACE
diff --git a/Engine/Source/Runtime/Engine/Public/SceneManagement.h b/Engine/Source/Runtime/Engine/Public/SceneManagement.h
--- a/Engine/Source/Runtime/Engine/Public/SceneManagement.h
+++ b/Engine/Source/Runtime/Engine/Public/SceneManagement.h
@@ -55,6 +55,7 @@ class FGPUScenePrimitiveCollector;
 DECLARE_LOG_CATEGORY_EXTERN(LogBufferVisualization, Log, All);
 DECLARE_LOG_CATEGORY_EXTERN(LogNaniteVisualization, Log, All);
 DECLARE_LOG_CATEGORY_EXTERN(LogLumenVisualization, Log, All);
+DECLARE_LOG_CATEGORY_EXTERN(LogVirtualShadowMapVisualization, Log, All);
 DECLARE_LOG_CATEGORY_EXTERN(LogMultiView, Log, All);
 
 // -----------------------------------------------------------------------------
diff --git a/Engine/Source/Runtime/Engine/Public/SceneView.h b/Engine/Source/Runtime/Engine/Public/SceneView.h
--- a/Engine/Source/Runtime/Engine/Public/SceneView.h
+++ b/Engine/Source/Runtime/Engine/Public/SceneView.h
@@ -1032,6 +1032,9 @@ public:
 	/** Current Nanite visualization mode */
 	FName CurrentLumenVisualizationMode;
 
+	/** Current Nanite visualization mode */
+	FName CurrentVirtualShadowMapVisualizationMode;
+
 	/** Current visualize calibration color material name */
 	FName CurrentVisualizeCalibrationColorMaterialName;
 
diff --git a/Engine/Source/Runtime/Engine/Public/ShowFlags.h b/Engine/Source/Runtime/Engine/Public/ShowFlags.h
--- a/Engine/Source/Runtime/Engine/Public/ShowFlags.h
+++ b/Engine/Source/Runtime/Engine/Public/ShowFlags.h
@@ -374,6 +374,7 @@ private:
 		SetVisualizeBuffer(false);
 		SetVisualizeNanite(false);
 		SetVisualizeLumen(false);
+		SetVisualizeVirtualShadowMap(false);
 		SetVectorFields(false);
 		SetGBufferHints(false);
 		SetCompositeEditorPrimitives(InitMode == ESFIM_Editor || InitMode == ESFIM_VREditing);
diff --git a/Engine/Source/Runtime/Engine/Public/ShowFlagsValues.inl b/Engine/Source/Runtime/Engine/Public/ShowFlagsValues.inl
--- a/Engine/Source/Runtime/Engine/Public/ShowFlagsValues.inl
+++ b/Engine/Source/Runtime/Engine/Public/ShowFlagsValues.inl
@@ -66,6 +66,8 @@ SHOWFLAG_FIXED_IN_SHIPPING(0, VisualizeBuffer, SFG_Hidden, NSLOCTEXT("UnrealEd",
 SHOWFLAG_ALWAYS_ACCESSIBLE(VisualizeNanite, SFG_Hidden, NSLOCTEXT("UnrealEd", "VisualizeNaniteSF", "Nanite Visualization"))
 /** Needed for VMI_VisualizeLumen, Whether to enable the Lumen visualization mode. */
 SHOWFLAG_ALWAYS_ACCESSIBLE(VisualizeLumen, SFG_Hidden, NSLOCTEXT("UnrealEd", "VisualizeLumenSF", "Lumen Visualization"))
+/** Needed for VMI_VisualizeVirtualShadowMap, Whether to enable the virtual shadow map visualization mode. */
+SHOWFLAG_ALWAYS_ACCESSIBLE(VisualizeVirtualShadowMap, SFG_Hidden, NSLOCTEXT("UnrealEd", "VisualizeVirtualShadowMapSF", "Virtual Shadow Map Visualization"))
 /** Allows to disable all direct lighting (does not affect indirect light) */
 SHOWFLAG_FIXED_IN_SHIPPING(1, DirectLighting, SFG_LightingComponents, NSLOCTEXT("UnrealEd", "DirectLightingSF", "Direct Lighting"))
 /** Allows to disable lighting from Directional Lights */
diff --git a/Engine/Source/Runtime/Engine/Public/VirtualShadowMapVisualizationData.h b/Engine/Source/Runtime/Engine/Public/VirtualShadowMapVisualizationData.h
new file mode 100644
--- /dev/null
+++ b/Engine/Source/Runtime/Engine/Public/VirtualShadowMapVisualizationData.h
@@ -0,0 +1,103 @@
+// Copyright Epic Games, Inc. All Rights Reserved.
+
+#pragma once
+
+#include "CoreMinimal.h"
+
+class FVirtualShadowMapVisualizationData
+{
+public:
+	enum class FModeType : uint8
+	{
+		ProjectionStandard,
+		ProjectionAdvanced,
+		DebugStandard,
+	};
+
+	/** Describes a single available visualization mode. */
+	struct FModeRecord
+	{
+		FString   ModeString;
+		FName     ModeName;
+		FText     ModeText;
+		FText     ModeDesc;
+		FModeType ModeType;
+		int32     ModeID;
+	};
+
+	/** Mapping of FName to a visualization mode record. */
+	typedef TMultiMap<FName, FModeRecord> TModeMap;
+
+public:
+	FVirtualShadowMapVisualizationData()
+	: bIsInitialized(false)
+	{
+	}
+
+	/** Initialize the system. */
+	void Initialize();
+
+	/** Check if system was initialized. */
+	inline bool IsInitialized() const { return bIsInitialized; }
+
+	/** Check if visualization is active. */
+	ENGINE_API bool IsActive() const;
+
+	/** Update state and check if visualization is active. */
+	ENGINE_API bool Update(const FName& InViewMode);
+
+	/** Get the display name of a named mode from the available mode map. **/
+	ENGINE_API FText GetModeDisplayName(const FName& InModeName) const;
+
+	ENGINE_API int32 GetModeID(const FName& InModeName) const;
+
+	inline const TModeMap& GetModeMap() const
+	{
+		return ModeMap;
+	}
+
+	inline int32 GetActiveModeID() const
+	{
+		return ActiveVisualizationModeID;
+	}
+
+	inline const FName& GetActiveModeName() const
+	{
+		return ActiveVisualizationModeName;
+	}
+
+	/** Return the console command name for enabling single mode visualization. */
+	static const TCHAR* GetVisualizeConsoleCommandName()
+	{
+		return TEXT("r.Shadow.Virtual.Visualize");
+	}
+
+private:
+	/** Internal helper function for creating the VirtualShadowMap visualization system console commands. */
+	void ConfigureConsoleCommand();
+
+	void AddVisualizationMode(
+		const TCHAR* ModeString,
+		const FText& ModeText,
+		const FText& ModeDesc,
+		const FModeType ModeType,
+		int32 ModeID
+	);
+
+	void SetActiveMode(int32 ModeID, const FName& ModeName);
+
+private:
+	/** The name->mode mapping table */
+	TModeMap ModeMap;
+
+	int32 ActiveVisualizationModeID = INDEX_NONE;
+	FName ActiveVisualizationModeName = NAME_None;
+
+	/** Storage for console variable documentation strings. **/
+	FString ConsoleDocumentationVisualizationMode;
+
+	/** Flag indicating if system is initialized. **/
+	bool bIsInitialized;
+};
+
+ENGINE_API FVirtualShadowMapVisualizationData& GetVirtualShadowMapVisualizationData();
diff --git a/Engine/Source/Runtime/Renderer/Private/DeferredShadingRenderer.cpp b/Engine/Source/Runtime/Renderer/Private/DeferredShadingRenderer.cpp
--- a/Engine/Source/Runtime/Renderer/Private/DeferredShadingRenderer.cpp
+++ b/Engine/Source/Runtime/Renderer/Private/DeferredShadingRenderer.cpp
@@ -2706,6 +2706,8 @@ void FDeferredShadingSceneRenderer::Render(FRDGBuilder& GraphBuilder)
 				NaniteRasterResults
 			);
 		}
+
+		// VisualizeVirtualShadowMap TODO
 	}
 
 	if (ViewFamily.EngineShowFlags.VisualizeLightCulling)
@@ -2766,7 +2768,7 @@ void FDeferredShadingSceneRenderer::Render(FRDGBuilder& GraphBuilder)
 			if (VirtualShadowMapArray.IsEnabled())
 			{
 				ensureMsgf(AreLightsInLightGrid(), TEXT("Virtual shadow map setup requires local lights to be injected into the light grid (this may be caused by 'r.LightCulling.Quality=0')."));
-				VirtualShadowMapArray.BuildPageAllocations(GraphBuilder, SceneTextures, Views, SortedLightSet, VisibleLightInfos, NaniteRasterResults, *Scene);
+				VirtualShadowMapArray.BuildPageAllocations(GraphBuilder, SceneTextures, Views, ViewFamily.EngineShowFlags, SortedLightSet, VisibleLightInfos, NaniteRasterResults, *Scene);
 			}
 
 			RenderShadowDepthMaps(GraphBuilder, InstanceCullingManager);
@@ -3312,7 +3314,7 @@ void FDeferredShadingSceneRenderer::Render(FRDGBuilder& GraphBuilder)
 				else
 #endif
 				{
-					AddPostProcessingPasses(GraphBuilder, View, PostProcessingInputs, NaniteResults, InstanceCullingManager);
+					AddPostProcessingPasses(GraphBuilder, View, PostProcessingInputs, NaniteResults, InstanceCullingManager, &VirtualShadowMapArray);
 				}
 			}
 		}
diff --git a/Engine/Source/Runtime/Renderer/Private/PostProcess/PostProcessEyeAdaptation.cpp b/Engine/Source/Runtime/Renderer/Private/PostProcess/PostProcessEyeAdaptation.cpp
--- a/Engine/Source/Runtime/Renderer/Private/PostProcess/PostProcessEyeAdaptation.cpp
+++ b/Engine/Source/Runtime/Renderer/Private/PostProcess/PostProcessEyeAdaptation.cpp
@@ -215,6 +215,7 @@ bool IsAutoExposureDebugMode(const FViewInfo& View)
 		!EngineShowFlags.Lighting ||
 		(EngineShowFlags.VisualizeBuffer && View.CurrentBufferVisualizationMode != NAME_None) ||
 		(EngineShowFlags.VisualizeNanite && View.CurrentNaniteVisualizationMode != NAME_None) ||
+		(EngineShowFlags.VisualizeVirtualShadowMap && View.CurrentVirtualShadowMapVisualizationMode != NAME_None) ||
 		EngineShowFlags.RayTracingDebug ||
 		EngineShowFlags.VisualizeDistanceFieldAO ||
 		EngineShowFlags.VisualizeGlobalDistanceField ||
@@ -944,7 +945,8 @@ void FSceneViewState::UpdatePreExposure(FViewInfo& View)
 		!ViewFamily.EngineShowFlags.HLODColoration &&
 		!ViewFamily.EngineShowFlags.LevelColoration &&
 		((!ViewFamily.EngineShowFlags.VisualizeBuffer) || View.CurrentBufferVisualizationMode != NAME_None) && // disable pre-exposure for the buffer visualization modes
-		((!ViewFamily.EngineShowFlags.VisualizeNanite) || View.CurrentNaniteVisualizationMode != NAME_None);   // disable pre-exposure for the Nanite visualization modes
+		((!ViewFamily.EngineShowFlags.VisualizeNanite) || View.CurrentNaniteVisualizationMode != NAME_None) && // disable pre-exposure for the Nanite visualization modes
+		((!ViewFamily.EngineShowFlags.VisualizeVirtualShadowMap) || View.CurrentVirtualShadowMapVisualizationMode != NAME_None); // disable pre-exposure for the virtual shadow map visualization modes
 
 	PreExposure = 1.f;
 	bUpdateLastExposure = false;
diff --git a/Engine/Source/Runtime/Renderer/Private/PostProcess/PostProcessing.cpp b/Engine/Source/Runtime/Renderer/Private/PostProcess/PostProcessing.cpp
--- a/Engine/Source/Runtime/Renderer/Private/PostProcess/PostProcessing.cpp
+++ b/Engine/Source/Runtime/Renderer/Private/PostProcess/PostProcessing.cpp
@@ -55,6 +55,7 @@
 #include "FXSystem.h"
 #include "SkyAtmosphereRendering.h"
 #include "Strata/Strata.h"
+#include "VirtualShadowMaps/VirtualShadowMapArray.h"
 
 bool IsMobileEyeAdaptationEnabled(const FViewInfo& View);
 
@@ -339,7 +340,13 @@ void AddTemporalAA2Passes(
 
 bool ComposeSeparateTranslucencyInTSR(const FViewInfo& View);
 
-void AddPostProcessingPasses(FRDGBuilder& GraphBuilder, const FViewInfo& View, const FPostProcessingInputs& Inputs, const Nanite::FRasterResults* NaniteRasterResults, FInstanceCullingManager& InstanceCullingManager)
+void AddPostProcessingPasses(
+	FRDGBuilder& GraphBuilder,
+	const FViewInfo& View,
+	const FPostProcessingInputs& Inputs,
+	const Nanite::FRasterResults* NaniteRasterResults,
+	FInstanceCullingManager& InstanceCullingManager,
+	FVirtualShadowMapArray* VirtualShadowMapArray)
 {
 	RDG_CSV_STAT_EXCLUSIVE_SCOPE(GraphBuilder, RenderPostProcessing);
 	QUICK_SCOPE_CYCLE_COUNTER(STAT_PostProcessing_Process);
@@ -1437,6 +1444,11 @@ void AddPostProcessingPasses(FRDGBuilder& GraphBuilder, const FViewInfo& View, c
 			AddVisualizeNanitePass(GraphBuilder, View, SceneColor, *NaniteRasterResults);
 		}
 
+		if (EngineShowFlags.VisualizeVirtualShadowMap && VirtualShadowMapArray != nullptr)
+		{
+			VirtualShadowMapArray->AddVisualizePass(GraphBuilder, View, SceneColor);
+		}
+
 		if (ShaderDrawDebug::IsEnabled(View))
 		{
 			ShaderDrawDebug::DrawView(GraphBuilder, View, SceneColor.Texture, SceneDepth.Texture);
diff --git a/Engine/Source/Runtime/Renderer/Private/PostProcess/PostProcessing.h b/Engine/Source/Runtime/Renderer/Private/PostProcess/PostProcessing.h
--- a/Engine/Source/Runtime/Renderer/Private/PostProcess/PostProcessing.h
+++ b/Engine/Source/Runtime/Renderer/Private/PostProcess/PostProcessing.h
@@ -8,6 +8,7 @@
 #include "ScenePrivate.h"
 
 class FSceneTextureParameters;
+class FVirtualShadowMapArray;
 
 namespace Nanite
 {
@@ -40,7 +41,13 @@ struct FPostProcessingInputs
 	}
 };
 
-void AddPostProcessingPasses(FRDGBuilder& GraphBuilder, const FViewInfo& View, const FPostProcessingInputs& Inputs, const Nanite::FRasterResults* NaniteRasterResults, FInstanceCullingManager& InstanceCullingManager);
+void AddPostProcessingPasses(
+	FRDGBuilder& GraphBuilder,
+	const FViewInfo& View,
+	const FPostProcessingInputs& Inputs,
+	const Nanite::FRasterResults* NaniteRasterResults,
+	FInstanceCullingManager& InstanceCullingManager,
+	FVirtualShadowMapArray* VirtualShadowMapArray);
 
 void AddDebugViewPostProcessingPasses(FRDGBuilder& GraphBuilder, const FViewInfo& View, const FPostProcessingInputs& Inputs, const Nanite::FRasterResults* NaniteRasterResults);
 
diff --git a/Engine/Source/Runtime/Renderer/Private/VirtualShadowMaps/VirtualShadowMapArray.cpp b/Engine/Source/Runtime/Renderer/Private/VirtualShadowMaps/VirtualShadowMapArray.cpp
--- a/Engine/Source/Runtime/Renderer/Private/VirtualShadowMaps/VirtualShadowMapArray.cpp
+++ b/Engine/Source/Runtime/Renderer/Private/VirtualShadowMaps/VirtualShadowMapArray.cpp
@@ -3,7 +3,9 @@
 	VirtualShadowMap.h:
 =============================================================================*/
 #include "VirtualShadowMapArray.h"
+#include "VirtualShadowMapVisualizationData.h"
 #include "../BasePassRendering.h"
+#include "../ScreenPass.h"
 #include "Components/LightComponent.h"
 #include "RendererModule.h"
 #include "Rendering/NaniteResources.h"
@@ -188,6 +190,7 @@ TAutoConsoleVariable<int32> CVarDoNonNaniteBatching(
 	TEXT("."),
 	ECVF_RenderThreadSafe
 );
+
 #if !UE_BUILD_SHIPPING
 bool GDumpVSMLightNames = false;
 void DumpVSMLightNames()
@@ -200,23 +203,26 @@ void DumpVSMLightNames()
 }
 
 FAutoConsoleCommand CmdDumpVSMLightNames(
-	TEXT("r.Shadow.Virtual.DumpLightNames"),
-	TEXT("Dump light names with Virtual Shadow Maps (for developer only, only for non shiping build)"),
+	TEXT("r.Shadow.Virtual.Visualize.DumpLightNames"),
+	TEXT("Dump light names with virtual shadow maps (for developer use in non-shiping builds)"),
 	FConsoleCommandDelegate::CreateStatic(DumpVSMLightNames)
 );
 
-FString GVirtualShadowMapDebugLight;
-FAutoConsoleVariableRef CVarDebugLight(
-	TEXT("r.Shadow.Virtual.DebugLight"),
-	GVirtualShadowMapDebugLight,
-	TEXT("Sets the name of a specific light to debug."),
+FString GVirtualShadowMapVisualizeLightName;
+FAutoConsoleVariableRef CVarVisualizeLightName(
+	TEXT("r.Shadow.Virtual.Visualize.LightName"),
+	GVirtualShadowMapVisualizeLightName,
+	TEXT("Sets the name of a specific light to visualize (for developer use in non-shiping builds)"),
 	ECVF_RenderThreadSafe
 );
 
-static TAutoConsoleVariable<int32> CVarVirtualShadowMapDebugProjection(
-	TEXT("r.Shadow.Virtual.DebugProjection"),
+static TAutoConsoleVariable<int32> CVarVisualizeLayout(
+	TEXT("r.Shadow.Virtual.Visualize.Layout"),
 	0,
-	TEXT("Projection pass debug output visualization for use with 'vis Shadow.Virtual.DebugProjection'."),
+	TEXT("Overlay layout when virtual shadow map visualization is enabled:\n")
+	TEXT("  0: Full screen\n")
+	TEXT("  1: Thumbnail\n")
+	TEXT("  2: Split screen"),
 	ECVF_RenderThreadSafe
 );
 
@@ -344,7 +350,7 @@ void FVirtualShadowMapArray::Initialize(FRDGBuilder& GraphBuilder, FVirtualShado
 
 	// Fixed physical page pool width, we adjust the height to accomodate the requested maximum
 	// NOTE: Row size in pages has to be POT since we use mask & shift in place of integer ops
-	const uint32 PhysicalPagesX = FMath::DivideAndRoundUp(8192U, FVirtualShadowMap::PageSize);
+	const uint32 PhysicalPagesX = FMath::DivideAndRoundDown(GetMax2DTextureDimension(), FVirtualShadowMap::PageSize);
 	check(FMath::IsPowerOfTwo(PhysicalPagesX));
 	uint32 PhysicalPagesY = FMath::DivideAndRoundUp((uint32)FMath::Max(1, CVarMaxPhysicalPages.GetValueOnRenderThread()), PhysicalPagesX);	
 
@@ -367,6 +373,11 @@ void FVirtualShadowMapArray::Initialize(FRDGBuilder& GraphBuilder, FVirtualShado
 	uint32 PhysicalX = PhysicalPagesX * FVirtualShadowMap::PageSize;
 	uint32 PhysicalY = PhysicalPagesY * FVirtualShadowMap::PageSize;
 
+	// TODO: Some sort of better fallback with warning?
+	// All supported platforms support at least 16384 texture dimensions which translates to 16384 max pages with default 128x128 page size
+	check(PhysicalX <= GetMax2DTextureDimension());
+	check(PhysicalY <= GetMax2DTextureDimension());
+
 	UniformParameters.PhysicalPageRowMask = (PhysicalPagesX - 1);
 	UniformParameters.PhysicalPageRowShift = FMath::FloorLog2( PhysicalPagesX );
 	UniformParameters.RecPhysicalPoolSize = FVector4f( 1.0f / PhysicalX, 1.0f / PhysicalY, 1.0f, 1.0f );
@@ -940,50 +951,36 @@ class FVirtualSmFeedbackStatusCS : public FVirtualPageManagementShader
 };
 IMPLEMENT_GLOBAL_SHADER(FVirtualSmFeedbackStatusCS, "/Engine/Private/VirtualShadowMaps/PageManagement.usf", "FeedbackStatusCS", SF_Compute);
 
-
-
-static FString GetLightNameForDebug(const FLightSceneProxy* Proxy)
+void FVirtualShadowMapVisualizeLightSearch::CheckLight(const FLightSceneProxy* CheckProxy, int CheckVirtualShadowMapId)
 {
-	return Proxy->GetOwnerNameOrLabel();
-}
-
 #if !UE_BUILD_SHIPPING
-struct FDebugLightSearch
-{
-public:
-	void CheckDebugLight(const FLightSceneProxy* Proxy, int CheckVirtualShadowMapId)
+	FString CheckLightName = CheckProxy->GetOwnerNameOrLabel();
+	if (GDumpVSMLightNames)
 	{
-		if (bFoundExactMatch)
-		{
-			return;
-		}
-
-		FString LightName = GetLightNameForDebug(Proxy);
-		if (GDumpVSMLightNames)
-		{
-			UE_LOG(LogRenderer, Display, TEXT("%s"), *LightName);
-		}
-		
-		bFoundExactMatch = (LightName == GVirtualShadowMapDebugLight);
-		bool bPartialMatch = LightName.Contains(GVirtualShadowMapDebugLight);
-		bool bDirectionalLight = Proxy->GetLightType() == LightType_Directional;
-
-		// Priority: exact match, partial match, directional light, anything
-		if (bFoundExactMatch ||
-			VirtualShadowMapId == INDEX_NONE ||
-			(!bFoundPartialMatch && (bPartialMatch || (!bFoundDirectionalLight && bDirectionalLight))))
-		{
-			bFoundDirectionalLight = bDirectionalLight;
-			bFoundPartialMatch = bPartialMatch;
-			VirtualShadowMapId = CheckVirtualShadowMapId;
-		}
+		UE_LOG(LogRenderer, Display, TEXT("%s"), *CheckLightName);
 	}
 
-	bool bFoundDirectionalLight = false;
-	bool bFoundPartialMatch = false;
-	bool bFoundExactMatch = false;
-	int VirtualShadowMapId = INDEX_NONE;
-};
+	const ULightComponent* Component = CheckProxy->GetLightComponent();
+	check(Component);
+
+	// Fill out new sort key and compare to our best found so far
+	SortKey CheckKey;
+	CheckKey.Packed = 0;
+	CheckKey.Fields.bExactNameMatch = (CheckLightName == GVirtualShadowMapVisualizeLightName);
+	CheckKey.Fields.bPartialNameMatch = CheckKey.Fields.bExactNameMatch || CheckLightName.Contains(GVirtualShadowMapVisualizeLightName);
+	CheckKey.Fields.bSelected = Component->IsSelected();
+	CheckKey.Fields.bOwnerSelected = Component->IsOwnerSelected();
+	CheckKey.Fields.bDirectionalLight = CheckProxy->GetLightType() == LightType_Directional;
+	CheckKey.Fields.bExists = 1;
+
+	if (CheckKey.Packed > FoundKey.Packed)
+	{
+		FoundKey = CheckKey;
+		FoundProxy = CheckProxy;
+		FoundVirtualShadowMapId = CheckVirtualShadowMapId;
+	}
+#endif
+}
 
 static FRDGTextureRef CreateDebugOutputTexture(FRDGBuilder& GraphBuilder, FIntPoint Extent)
 {
@@ -999,523 +996,531 @@ static FRDGTextureRef CreateDebugOutputTexture(FRDGBuilder& GraphBuilder, FIntPo
 	AddClearUAVPass(GraphBuilder, GraphBuilder.CreateUAV(Texture), ClearColor);
 	return Texture;
 }
-#endif // !UE_BUILD_SHIPPING
 
 void FVirtualShadowMapArray::BuildPageAllocations(
 	FRDGBuilder& GraphBuilder,
 	const FMinimalSceneTextures& SceneTextures,
 	const TArray<FViewInfo>& Views,
+	const FEngineShowFlags& EngineShowFlags,
 	const FSortedLightSetSceneInfo& SortedLightsInfo,
 	const TArray<FVisibleLightInfo, SceneRenderingAllocator>& VisibleLightInfos,
 	const TArray<Nanite::FRasterResults, TInlineAllocator<2>>& NaniteRasterResults,
 	FScene& Scene)
 {
 	check(IsEnabled());
+
+	if (ShadowMaps.Num() == 0 || Views.Num() == 0)
+	{
+		// Nothing to do
+		return;
+	}
+
 	RDG_EVENT_SCOPE(GraphBuilder, "FVirtualShadowMapArray::BuildPageAllocation");
 
-#if !UE_BUILD_SHIPPING
 	bool bDebugOutputEnabled = false;
+	VisualizeLight.Reset();
+
+#if !UE_BUILD_SHIPPING
 	if (GDumpVSMLightNames)
 	{
 		bDebugOutputEnabled = true;
 		UE_LOG(LogRenderer, Display, TEXT("Lights with Virtual Shadow Maps:"));
 	}
 
-	// Setup debug output
-	DebugOutputType = CVarVirtualShadowMapDebugProjection.GetValueOnRenderThread();
-	FDebugLightSearch DebugLightSearch;
-	if (DebugOutputType > 0)
+	// Setup debug visualization/output if enabled
 	{
-		bDebugOutputEnabled = true;
-		DebugVisualizationProjectionOutput = CreateDebugOutputTexture(GraphBuilder, SceneTextures.Config.Extent);
+		FVirtualShadowMapVisualizationData& VisualizationData = GetVirtualShadowMapVisualizationData();
+	
+		// TODO: Support more than one view/debug output
+		const FName& VisualizationMode = Views[0].CurrentVirtualShadowMapVisualizationMode;
+		if (VisualizationData.Update(VisualizationMode))
+		{
+			// TODO - automatically enable the show flag when set from command line?
+			//EngineShowFlags.SetVisualizeVirtualShadowMap(true);
+		}
+
+		if (VisualizationData.IsActive() && EngineShowFlags.VisualizeVirtualShadowMap)
+		{
+			bDebugOutputEnabled = true;
+			DebugVisualizationOutput = CreateDebugOutputTexture(GraphBuilder, SceneTextures.Config.Extent);
+		}
 	}
 #endif //!UE_BUILD_SHIPPING
+		
+	// Store shadow map projection data for each virtual shadow map
+	TArray<FVirtualShadowMapProjectionShaderData, SceneRenderingAllocator> ShadowMapProjectionData;
+	ShadowMapProjectionData.AddDefaulted(ShadowMaps.Num());
 
-	const TArray<FSortedLightSceneInfo, SceneRenderingAllocator> &SortedLights = SortedLightsInfo.SortedLights;
-	if (ShadowMaps.Num())
+	// Gather directional light virtual shadow maps
+	TArray<int32, SceneRenderingAllocator> DirectionalLightIds;
+	for (const FVisibleLightInfo& VisibleLightInfo : VisibleLightInfos)
 	{
-		// Store shadow map projection data for each virtual shadow map
-		TArray<FVirtualShadowMapProjectionShaderData, SceneRenderingAllocator> ShadowMapProjectionData;
-		ShadowMapProjectionData.AddDefaulted(ShadowMaps.Num());
+		for (const TSharedPtr<FVirtualShadowMapClipmap>& Clipmap : VisibleLightInfo.VirtualShadowMapClipmaps)
+		{
+			// NOTE: Shader assumes all levels from a given clipmap are contiguous
+			int32 ClipmapID = Clipmap->GetVirtualShadowMap()->ID;
+			DirectionalLightIds.Add(ClipmapID);
+			for (int32 ClipmapLevel = 0; ClipmapLevel < Clipmap->GetLevelCount(); ++ClipmapLevel)
+			{
+				ShadowMapProjectionData[ClipmapID + ClipmapLevel] = Clipmap->GetProjectionShaderData(ClipmapLevel);
+			}
+
+			if (bDebugOutputEnabled)
+			{
+				VisualizeLight.CheckLight(Clipmap->GetLightSceneInfo().Proxy, ClipmapID);
+			}
+		}
 
-		// Gather directional light virtual shadow maps
-		TArray<int32, SceneRenderingAllocator> DirectionalLightIds;
-		for (const FVisibleLightInfo& VisibleLightInfo : VisibleLightInfos)
+		for (FProjectedShadowInfo* ProjectedShadowInfo : VisibleLightInfo.AllProjectedShadows)
 		{
-			for (const TSharedPtr<FVirtualShadowMapClipmap>& Clipmap : VisibleLightInfo.VirtualShadowMapClipmaps)
+			if (ProjectedShadowInfo->HasVirtualShadowMap())
 			{
-				// NOTE: Shader assumes all levels from a given clipmap are contiguous
-				int32 ClipmapID = Clipmap->GetVirtualShadowMap()->ID;
-				DirectionalLightIds.Add(ClipmapID);
-				for (int32 ClipmapLevel = 0; ClipmapLevel < Clipmap->GetLevelCount(); ++ClipmapLevel)
+				check(ProjectedShadowInfo->CascadeSettings.ShadowSplitIndex == INDEX_NONE);		// We use clipmaps for virtual shadow maps, not cascades
+
+				// NOTE: Virtual shadow maps are never atlased, but verify our assumptions
 				{
-					ShadowMapProjectionData[ClipmapID + ClipmapLevel] = Clipmap->GetProjectionShaderData(ClipmapLevel);
+					const FVector4f ClipToShadowUV = ProjectedShadowInfo->GetClipToShadowBufferUvScaleBias();
+					check(ProjectedShadowInfo->BorderSize == 0);
+					check(ProjectedShadowInfo->X == 0);
+					check(ProjectedShadowInfo->Y == 0);
+					const FIntRect ShadowViewRect = ProjectedShadowInfo->GetInnerViewRect();
+					check(ShadowViewRect.Min.X == 0);
+					check(ShadowViewRect.Min.Y == 0);
+					check(ShadowViewRect.Max.X == FVirtualShadowMap::VirtualMaxResolutionXY);
+					check(ShadowViewRect.Max.Y == FVirtualShadowMap::VirtualMaxResolutionXY);
 				}
 
-			#if !UE_BUILD_SHIPPING
-				if (bDebugOutputEnabled)
+				int32 NumMaps = ProjectedShadowInfo->bOnePassPointLightShadow ? 6 : 1;
+				for( int32 i = 0; i < NumMaps; i++ )
 				{
-					DebugLightSearch.CheckDebugLight(Clipmap->GetLightSceneInfo().Proxy, ClipmapID);
+					int32 ID = ProjectedShadowInfo->VirtualShadowMaps[i]->ID;
+
+					FViewMatrices ViewMatrices = ProjectedShadowInfo->GetShadowDepthRenderingViewMatrices( i, true );
+
+					FVirtualShadowMapProjectionShaderData& Data = ShadowMapProjectionData[ ID ];
+					Data.TranslatedWorldToShadowViewMatrix		= FMatrix44f(ViewMatrices.GetTranslatedViewMatrix());	// LWC_TODO: Precision loss?
+					Data.ShadowViewToClipMatrix					= FMatrix44f(ViewMatrices.GetProjectionMatrix());
+					Data.TranslatedWorldToShadowUVMatrix		= FMatrix44f(CalcTranslatedWorldToShadowUVMatrix( ViewMatrices.GetTranslatedViewMatrix(), ViewMatrices.GetProjectionMatrix() ));
+					Data.TranslatedWorldToShadowUVNormalMatrix	= FMatrix44f(CalcTranslatedWorldToShadowUVNormalMatrix( ViewMatrices.GetTranslatedViewMatrix(), ViewMatrices.GetProjectionMatrix() ));
+					Data.ShadowPreViewTranslation				= FVector(ProjectedShadowInfo->PreShadowTranslation);
+					Data.LightType								= ProjectedShadowInfo->GetLightSceneInfo().Proxy->GetLightType();
+					Data.LightSourceRadius						= ProjectedShadowInfo->GetLightSceneInfo().Proxy->GetSourceRadius();
 				}
-			#endif // !UE_BUILD_SHIPPING
-			}
 
-			for (FProjectedShadowInfo* ProjectedShadowInfo : VisibleLightInfo.AllProjectedShadows)
-			{
-				if (ProjectedShadowInfo->HasVirtualShadowMap())
+				if (bDebugOutputEnabled)
 				{
-					check(ProjectedShadowInfo->CascadeSettings.ShadowSplitIndex == INDEX_NONE);		// We use clipmaps for virtual shadow maps, not cascades
-
-					// NOTE: Virtual shadow maps are never atlased, but verify our assumptions
-					{
-						const FVector4f ClipToShadowUV = ProjectedShadowInfo->GetClipToShadowBufferUvScaleBias();
-						check(ProjectedShadowInfo->BorderSize == 0);
-						check(ProjectedShadowInfo->X == 0);
-						check(ProjectedShadowInfo->Y == 0);
-						const FIntRect ShadowViewRect = ProjectedShadowInfo->GetInnerViewRect();
-						check(ShadowViewRect.Min.X == 0);
-						check(ShadowViewRect.Min.Y == 0);
-						check(ShadowViewRect.Max.X == FVirtualShadowMap::VirtualMaxResolutionXY);
-						check(ShadowViewRect.Max.Y == FVirtualShadowMap::VirtualMaxResolutionXY);
-					}
-
-					int32 NumMaps = ProjectedShadowInfo->bOnePassPointLightShadow ? 6 : 1;
-					for( int32 i = 0; i < NumMaps; i++ )
-					{
-						int32 ID = ProjectedShadowInfo->VirtualShadowMaps[i]->ID;
-
-						FViewMatrices ViewMatrices = ProjectedShadowInfo->GetShadowDepthRenderingViewMatrices( i, true );
-
-						FVirtualShadowMapProjectionShaderData& Data = ShadowMapProjectionData[ ID ];
-						Data.TranslatedWorldToShadowViewMatrix		= FMatrix44f(ViewMatrices.GetTranslatedViewMatrix());	// LWC_TODO: Precision loss?
-						Data.ShadowViewToClipMatrix					= FMatrix44f(ViewMatrices.GetProjectionMatrix());
-						Data.TranslatedWorldToShadowUVMatrix		= FMatrix44f(CalcTranslatedWorldToShadowUVMatrix( ViewMatrices.GetTranslatedViewMatrix(), ViewMatrices.GetProjectionMatrix() ));
-						Data.TranslatedWorldToShadowUVNormalMatrix	= FMatrix44f(CalcTranslatedWorldToShadowUVNormalMatrix( ViewMatrices.GetTranslatedViewMatrix(), ViewMatrices.GetProjectionMatrix() ));
-						Data.ShadowPreViewTranslation				= FVector(ProjectedShadowInfo->PreShadowTranslation);
-						Data.LightType								= ProjectedShadowInfo->GetLightSceneInfo().Proxy->GetLightType();
-						Data.LightSourceRadius						= ProjectedShadowInfo->GetLightSceneInfo().Proxy->GetSourceRadius();
-					}
-
-				#if !UE_BUILD_SHIPPING
-					if (bDebugOutputEnabled)
-					{
-						DebugLightSearch.CheckDebugLight(ProjectedShadowInfo->GetLightSceneInfo().Proxy, ProjectedShadowInfo->VirtualShadowMaps[0]->ID);
-					}
-				#endif // !UE_BUILD_SHIPPING
+					VisualizeLight.CheckLight(ProjectedShadowInfo->GetLightSceneInfo().Proxy, ProjectedShadowInfo->VirtualShadowMaps[0]->ID);
 				}
 			}
 		}
-	#if !UE_BUILD_SHIPPING
-		DebugVirtualShadowMapId = DebugLightSearch.VirtualShadowMapId;
-	#endif
+	}
 
-		UniformParameters.NumShadowMaps = ShadowMaps.Num();
-		UniformParameters.NumDirectionalLights = DirectionalLightIds.Num();
+	UniformParameters.NumShadowMaps = ShadowMaps.Num();
+	UniformParameters.NumDirectionalLights = DirectionalLightIds.Num();
 
-		ShadowMapProjectionDataRDG = CreateProjectionDataBuffer(GraphBuilder, TEXT("Shadow.Virtual.ProjectionData"), ShadowMapProjectionData);
+	ShadowMapProjectionDataRDG = CreateProjectionDataBuffer(GraphBuilder, TEXT("Shadow.Virtual.ProjectionData"), ShadowMapProjectionData);
 
-		UniformParameters.ProjectionData = GraphBuilder.CreateSRV(ShadowMapProjectionDataRDG);
+	UniformParameters.ProjectionData = GraphBuilder.CreateSRV(ShadowMapProjectionDataRDG);
 
-		if (CVarShowStats.GetValueOnRenderThread() || CacheManager->IsAccumulatingStats())
-		{
-			StatsBufferRDG = GraphBuilder.CreateBuffer(FRDGBufferDesc::CreateStructuredDesc(sizeof(uint32), NumStats), TEXT("Shadow.Virtual.StatsBuffer"));
-			AddClearUAVPass(GraphBuilder, GraphBuilder.CreateUAV(StatsBufferRDG), 0);
-		}
+	if (CVarShowStats.GetValueOnRenderThread() || CacheManager->IsAccumulatingStats())
+	{
+		StatsBufferRDG = GraphBuilder.CreateBuffer(FRDGBufferDesc::CreateStructuredDesc(sizeof(uint32), NumStats), TEXT("Shadow.Virtual.StatsBuffer"));
+		AddClearUAVPass(GraphBuilder, GraphBuilder.CreateUAV(StatsBufferRDG), 0);
+	}
 		
-		// Create and clear the requested page flags
-		const uint32 NumPageFlags = ShadowMaps.Num() * FVirtualShadowMap::PageTableSize;
-		FRDGBufferRef PageRequestFlagsRDG = GraphBuilder.CreateBuffer(FRDGBufferDesc::CreateStructuredDesc(sizeof(uint32), NumPageFlags), TEXT("Shadow.Virtual.PageRequestFlags"));
-		AddClearUAVPass(GraphBuilder, GraphBuilder.CreateUAV(PageRequestFlagsRDG), 0);
+	// Create and clear the requested page flags
+	const uint32 NumPageFlags = ShadowMaps.Num() * FVirtualShadowMap::PageTableSize;
+	FRDGBufferRef PageRequestFlagsRDG = GraphBuilder.CreateBuffer(FRDGBufferDesc::CreateStructuredDesc(sizeof(uint32), NumPageFlags), TEXT("Shadow.Virtual.PageRequestFlags"));
+	AddClearUAVPass(GraphBuilder, GraphBuilder.CreateUAV(PageRequestFlagsRDG), 0);
 		
-		// TODO: Remove/move to next frame and make temporary OR replace with direct page table manipulation?
-		DynamicCasterPageFlagsRDG = GraphBuilder.CreateBuffer(FRDGBufferDesc::CreateStructuredDesc(sizeof(uint32), NumPageFlags), TEXT("Shadow.Virtual.DynamicCasterPageFlags"));
-		AddClearUAVPass(GraphBuilder, GraphBuilder.CreateUAV(DynamicCasterPageFlagsRDG), 0);
+	// TODO: Remove/move to next frame and make temporary OR replace with direct page table manipulation?
+	DynamicCasterPageFlagsRDG = GraphBuilder.CreateBuffer(FRDGBufferDesc::CreateStructuredDesc(sizeof(uint32), NumPageFlags), TEXT("Shadow.Virtual.DynamicCasterPageFlags"));
+	AddClearUAVPass(GraphBuilder, GraphBuilder.CreateUAV(DynamicCasterPageFlagsRDG), 0);
 		
-		// Record the number of instances the buffer has capactiy for, should anything change (it shouldn't!)
-		NumInvalidatingInstanceSlots = Scene.GPUScene.GetNumInstances();
-		// Allocate space for counter, worst case ID storage, and flags.
-		int32 InstanceInvalidationBufferSize = 1 + NumInvalidatingInstanceSlots + FMath::DivideAndRoundUp(NumInvalidatingInstanceSlots, 32);
-		InvalidatingInstancesRDG = GraphBuilder.CreateBuffer(FRDGBufferDesc::CreateStructuredDesc(sizeof(uint32), InstanceInvalidationBufferSize), TEXT("Shadow.Virtual.InvalidatingInstances"));
-		// Clear to zero, technically only need to clear first Scene.GPUScene.GetNumInstances()  + 1 uints
-		AddClearUAVPass(GraphBuilder, GraphBuilder.CreateUAV(InvalidatingInstancesRDG), 0);
-
-		const uint32 NumPageRects = UniformParameters.NumShadowMaps * FVirtualShadowMap::MaxMipLevels;
-		PageRectBoundsRDG = GraphBuilder.CreateBuffer(FRDGBufferDesc::CreateStructuredDesc(sizeof(FIntVector4), NumPageRects), TEXT("Shadow.Virtual.PageRectBounds"));
-		{
-			FInitPageRectBoundsCS::FParameters* PassParameters = GraphBuilder.AllocParameters< FInitPageRectBoundsCS::FParameters >();
-			PassParameters->VirtualShadowMap = GetUniformBuffer(GraphBuilder);
-			PassParameters->OutPageRectBounds = GraphBuilder.CreateUAV(PageRectBoundsRDG);
+	// Record the number of instances the buffer has capactiy for, should anything change (it shouldn't!)
+	NumInvalidatingInstanceSlots = Scene.GPUScene.GetNumInstances();
+	// Allocate space for counter, worst case ID storage, and flags.
+	int32 InstanceInvalidationBufferSize = 1 + NumInvalidatingInstanceSlots + FMath::DivideAndRoundUp(NumInvalidatingInstanceSlots, 32);
+	InvalidatingInstancesRDG = GraphBuilder.CreateBuffer(FRDGBufferDesc::CreateStructuredDesc(sizeof(uint32), InstanceInvalidationBufferSize), TEXT("Shadow.Virtual.InvalidatingInstances"));
+	// Clear to zero, technically only need to clear first Scene.GPUScene.GetNumInstances()  + 1 uints
+	AddClearUAVPass(GraphBuilder, GraphBuilder.CreateUAV(InvalidatingInstancesRDG), 0);
+
+	const uint32 NumPageRects = UniformParameters.NumShadowMaps * FVirtualShadowMap::MaxMipLevels;
+	PageRectBoundsRDG = GraphBuilder.CreateBuffer(FRDGBufferDesc::CreateStructuredDesc(sizeof(FIntVector4), NumPageRects), TEXT("Shadow.Virtual.PageRectBounds"));
+	{
+		FInitPageRectBoundsCS::FParameters* PassParameters = GraphBuilder.AllocParameters< FInitPageRectBoundsCS::FParameters >();
+		PassParameters->VirtualShadowMap = GetUniformBuffer(GraphBuilder);
+		PassParameters->OutPageRectBounds = GraphBuilder.CreateUAV(PageRectBoundsRDG);
 
-			auto ComputeShader = GetGlobalShaderMap(GMaxRHIFeatureLevel)->GetShader<FInitPageRectBoundsCS>();
-			ClearUnusedGraphResources(ComputeShader, PassParameters);
+		auto ComputeShader = GetGlobalShaderMap(GMaxRHIFeatureLevel)->GetShader<FInitPageRectBoundsCS>();
+		ClearUnusedGraphResources(ComputeShader, PassParameters);
 
-			FComputeShaderUtils::AddPass(
-				GraphBuilder,
-				RDG_EVENT_NAME("InitPageRectBounds"),
-				ComputeShader,
-				PassParameters,
-				FIntVector(FMath::DivideAndRoundUp(NumPageRects, FInitPageRectBoundsCS::DefaultCSGroupX), 1, 1)
-			);
+		FComputeShaderUtils::AddPass(
+			GraphBuilder,
+			RDG_EVENT_NAME("InitPageRectBounds"),
+			ComputeShader,
+			PassParameters,
+			FIntVector(FMath::DivideAndRoundUp(NumPageRects, FInitPageRectBoundsCS::DefaultCSGroupX), 1, 1)
+		);
+	}
+
+	for (int32 ViewIndex = 0; ViewIndex < Views.Num(); ++ViewIndex)
+	{
+		const FViewInfo &View = Views[ViewIndex];
+
+		// This view contained no local lights (that were stored in the light grid), and no directional lights, so nothing to do.
+		if (View.ForwardLightingResources.LocalLightVisibleLightInfosIndex.Num() + DirectionalLightIds.Num() == 0)
+		{
+			continue;
 		}
 
-		for (int32 ViewIndex = 0; ViewIndex < Views.Num(); ++ViewIndex)
+		FRDGBufferRef DirectionalLightIdsRDG = CreateStructuredBuffer(GraphBuilder, TEXT("Shadow.Virtual.DirectionalLightIds"), DirectionalLightIds);
+
+		const FRDGSystemTextures& SystemTextures = FRDGSystemTextures::Get(GraphBuilder);
+
+		FRDGBufferRef ScreenSpaceGridBoundsRDG = nullptr;
+			
 		{
-			const FViewInfo &View = Views[ViewIndex];
+			// It's safe to overlap these passes that all write to page request flags
+			FRDGBufferUAVRef PageRequestFlagsUAV = GraphBuilder.CreateUAV(PageRequestFlagsRDG, ERDGUnorderedAccessViewFlags::SkipBarrier);
 
-			// This view contained no local lights (that were stored in the light grid), and no directional lights, so nothing to do.
-			if (View.ForwardLightingResources.LocalLightVisibleLightInfosIndex.Num() + DirectionalLightIds.Num() == 0)
+			// Mark pages based on projected depth buffer pixels
+			if (CVarMarkPixelPages.GetValueOnRenderThread() != 0)
 			{
-				continue;
-			}
+				auto GeneratePageFlags = [&](bool bHairPass)
+				{
+					const uint32 InputType = bHairPass ? 1U : 0U; // HairStrands or GBuffer
 
-			FRDGBufferRef DirectionalLightIdsRDG = CreateStructuredBuffer(GraphBuilder, TEXT("Shadow.Virtual.DirectionalLightIds"), DirectionalLightIds);
+					FGeneratePageFlagsFromPixelsCS::FPermutationDomain PermutationVector;
+					PermutationVector.Set<FGeneratePageFlagsFromPixelsCS::FInputType>(InputType);
+					FGeneratePageFlagsFromPixelsCS::FParameters* PassParameters = GraphBuilder.AllocParameters< FGeneratePageFlagsFromPixelsCS::FParameters >();
+					PassParameters->VirtualShadowMap = GetUniformBuffer(GraphBuilder);
 
-			const FRDGSystemTextures& SystemTextures = FRDGSystemTextures::Get(GraphBuilder);
+					PassParameters->SceneTexturesStruct = SceneTextures.UniformBuffer;
 
-			FRDGBufferRef ScreenSpaceGridBoundsRDG = nullptr;
-			
-			{
-				// It's safe to overlap these passes that all write to page request flags
-				FRDGBufferUAVRef PageRequestFlagsUAV = GraphBuilder.CreateUAV(PageRequestFlagsRDG, ERDGUnorderedAccessViewFlags::SkipBarrier);
+					PassParameters->HairStrands = HairStrands::BindHairStrandsViewUniformParameters(View);
+					PassParameters->View = View.ViewUniformBuffer;
+					PassParameters->OutPageRequestFlags = PageRequestFlagsUAV;
+					PassParameters->ForwardLightData = View.ForwardLightingResources.ForwardLightUniformBuffer;
+					PassParameters->DirectionalLightIds = GraphBuilder.CreateSRV(DirectionalLightIdsRDG);
+					PassParameters->ResolutionLodBiasLocal = CVarResolutionLodBiasLocal.GetValueOnRenderThread();
+					PassParameters->PageDilationBorderSizeLocal = CVarPageDilationBorderSizeLocal.GetValueOnRenderThread();
+					PassParameters->PageDilationBorderSizeDirectional = CVarPageDilationBorderSizeDirectional.GetValueOnRenderThread();
+					PassParameters->bCullBackfacingPixels = ShouldCullBackfacingPixels() ? 1 : 0;
 
-				// Mark pages based on projected depth buffer pixels
-				if (CVarMarkPixelPages.GetValueOnRenderThread() != 0)
-				{
-					auto GeneratePageFlags = [&](bool bHairPass)
+					auto ComputeShader = View.ShaderMap->GetShader<FGeneratePageFlagsFromPixelsCS>(PermutationVector);
+
+					static_assert((FVirtualPageManagementShader::DefaultCSGroupXY % 2) == 0, "GeneratePageFlagsFromPixels requires even-sized CS groups for quad swizzling.");
+					const FIntPoint GridSize = FIntPoint::DivideAndRoundUp(View.ViewRect.Size(), FVirtualPageManagementShader::DefaultCSGroupXY);
+
+					if (bHairPass && View.HairStrandsViewData.VisibilityData.TileData.IsValid())
 					{
-						const uint32 InputType = bHairPass ? 1U : 0U; // HairStrands or GBuffer
-
-						FGeneratePageFlagsFromPixelsCS::FPermutationDomain PermutationVector;
-						PermutationVector.Set<FGeneratePageFlagsFromPixelsCS::FInputType>(InputType);
-						FGeneratePageFlagsFromPixelsCS::FParameters* PassParameters = GraphBuilder.AllocParameters< FGeneratePageFlagsFromPixelsCS::FParameters >();
-						PassParameters->VirtualShadowMap = GetUniformBuffer(GraphBuilder);
-
-						PassParameters->SceneTexturesStruct = SceneTextures.UniformBuffer;
-
-						PassParameters->HairStrands = HairStrands::BindHairStrandsViewUniformParameters(View);
-						PassParameters->View = View.ViewUniformBuffer;
-						PassParameters->OutPageRequestFlags = PageRequestFlagsUAV;
-						PassParameters->ForwardLightData = View.ForwardLightingResources.ForwardLightUniformBuffer;
-						PassParameters->DirectionalLightIds = GraphBuilder.CreateSRV(DirectionalLightIdsRDG);
-						PassParameters->ResolutionLodBiasLocal = CVarResolutionLodBiasLocal.GetValueOnRenderThread();
-						PassParameters->PageDilationBorderSizeLocal = CVarPageDilationBorderSizeLocal.GetValueOnRenderThread();
-						PassParameters->PageDilationBorderSizeDirectional = CVarPageDilationBorderSizeDirectional.GetValueOnRenderThread();
-						PassParameters->bCullBackfacingPixels = ShouldCullBackfacingPixels() ? 1 : 0;
-
-						auto ComputeShader = View.ShaderMap->GetShader<FGeneratePageFlagsFromPixelsCS>(PermutationVector);
-
-						static_assert((FVirtualPageManagementShader::DefaultCSGroupXY % 2) == 0, "GeneratePageFlagsFromPixels requires even-sized CS groups for quad swizzling.");
-						const FIntPoint GridSize = FIntPoint::DivideAndRoundUp(View.ViewRect.Size(), FVirtualPageManagementShader::DefaultCSGroupXY);
-
-						if (bHairPass && View.HairStrandsViewData.VisibilityData.TileData.IsValid())
-						{
-							PassParameters->IndirectBufferArgs = View.HairStrandsViewData.VisibilityData.TileData.TileIndirectDispatchBuffer;
-							FComputeShaderUtils::AddPass(
-								GraphBuilder,
-								RDG_EVENT_NAME("GeneratePageFlagsFromPixels(HairStrands,Tile)"),
-								ComputeShader,
-								PassParameters,
-								View.HairStrandsViewData.VisibilityData.TileData.TileIndirectDispatchBuffer,
-								View.HairStrandsViewData.VisibilityData.TileData.GetIndirectDispatchArgOffset(FHairStrandsTiles::ETileType::HairAll));
-						}
-						else
-						{
-							FComputeShaderUtils::AddPass(
-								GraphBuilder,
-								RDG_EVENT_NAME("GeneratePageFlagsFromPixels(%s)", bHairPass ? TEXT("HairStrands") : TEXT("GBuffer")),
-								ComputeShader,
-								PassParameters,
-								FIntVector(GridSize.X, GridSize.Y, 1));
-						}
-					};
-
-					GeneratePageFlags(false);
-					if (HairStrands::HasViewHairStrandsData(View))
+						PassParameters->IndirectBufferArgs = View.HairStrandsViewData.VisibilityData.TileData.TileIndirectDispatchBuffer;
+						FComputeShaderUtils::AddPass(
+							GraphBuilder,
+							RDG_EVENT_NAME("GeneratePageFlagsFromPixels(HairStrands,Tile)"),
+							ComputeShader,
+							PassParameters,
+							View.HairStrandsViewData.VisibilityData.TileData.TileIndirectDispatchBuffer,
+							View.HairStrandsViewData.VisibilityData.TileData.GetIndirectDispatchArgOffset(FHairStrandsTiles::ETileType::HairAll));
+					}
+					else
 					{
-						GeneratePageFlags(true);
+						FComputeShaderUtils::AddPass(
+							GraphBuilder,
+							RDG_EVENT_NAME("GeneratePageFlagsFromPixels(%s)", bHairPass ? TEXT("HairStrands") : TEXT("GBuffer")),
+							ComputeShader,
+							PassParameters,
+							FIntVector(GridSize.X, GridSize.Y, 1));
 					}
-				}
-				// Mark coarse pages
-				bool bMarkCoarsePagesDirectional = CVarMarkCoarsePagesDirectional.GetValueOnRenderThread() != 0;
-				bool bMarkCoarsePagesLocal = CVarMarkCoarsePagesLocal.GetValueOnRenderThread() != 0;
-				if (bMarkCoarsePagesDirectional || bMarkCoarsePagesLocal)
+				};
+
+				GeneratePageFlags(false);
+				if (HairStrands::HasViewHairStrandsData(View))
 				{
-					FMarkCoarsePagesCS::FParameters* PassParameters = GraphBuilder.AllocParameters< FMarkCoarsePagesCS::FParameters >();
-					PassParameters->VirtualShadowMap = GetUniformBuffer(GraphBuilder);
-					PassParameters->OutPageRequestFlags = PageRequestFlagsUAV;
-					PassParameters->bMarkCoarsePagesLocal = bMarkCoarsePagesLocal ? 1 : 0;
-					PassParameters->ClipmapIndexMask = bMarkCoarsePagesDirectional ? FVirtualShadowMapClipmap::GetCoarsePageClipmapIndexMask() : 0;
-					PassParameters->bIncludeNonNaniteGeometry = CVarCoarsePagesIncludeNonNanite.GetValueOnRenderThread();
-
-					auto ComputeShader = View.ShaderMap->GetShader<FMarkCoarsePagesCS>();
-
-					FComputeShaderUtils::AddPass(
-						GraphBuilder,
-						RDG_EVENT_NAME("MarkCoarsePages"),
-						ComputeShader,
-						PassParameters,
-						FIntVector(FMath::DivideAndRoundUp(uint32(ShadowMaps.Num()), FMarkCoarsePagesCS::DefaultCSGroupX), 1, 1)
-					);
+					GeneratePageFlags(true);
 				}
 			}
+			// Mark coarse pages
+			bool bMarkCoarsePagesDirectional = CVarMarkCoarsePagesDirectional.GetValueOnRenderThread() != 0;
+			bool bMarkCoarsePagesLocal = CVarMarkCoarsePagesLocal.GetValueOnRenderThread() != 0;
+			if (bMarkCoarsePagesDirectional || bMarkCoarsePagesLocal)
+			{
+				FMarkCoarsePagesCS::FParameters* PassParameters = GraphBuilder.AllocParameters< FMarkCoarsePagesCS::FParameters >();
+				PassParameters->VirtualShadowMap = GetUniformBuffer(GraphBuilder);
+				PassParameters->OutPageRequestFlags = PageRequestFlagsUAV;
+				PassParameters->bMarkCoarsePagesLocal = bMarkCoarsePagesLocal ? 1 : 0;
+				PassParameters->ClipmapIndexMask = bMarkCoarsePagesDirectional ? FVirtualShadowMapClipmap::GetCoarsePageClipmapIndexMask() : 0;
+				PassParameters->bIncludeNonNaniteGeometry = CVarCoarsePagesIncludeNonNanite.GetValueOnRenderThread();
+
+				auto ComputeShader = View.ShaderMap->GetShader<FMarkCoarsePagesCS>();
+
+				FComputeShaderUtils::AddPass(
+					GraphBuilder,
+					RDG_EVENT_NAME("MarkCoarsePages"),
+					ComputeShader,
+					PassParameters,
+					FIntVector(FMath::DivideAndRoundUp(uint32(ShadowMaps.Num()), FMarkCoarsePagesCS::DefaultCSGroupX), 1, 1)
+				);
+			}
 		}
+	}
 
-		PageTableRDG = GraphBuilder.CreateBuffer(FRDGBufferDesc::CreateStructuredDesc(sizeof(uint32), NumPageFlags), TEXT("Shadow.Virtual.PageTable"));		
-		// Note: these are passed to the rendering and are not identical to the PageRequest flags coming in from GeneratePageFlagsFromPixels 
-		PageFlagsRDG = GraphBuilder.CreateBuffer(FRDGBufferDesc::CreateStructuredDesc(sizeof(uint32), NumPageFlags), TEXT("Shadow.Virtual.PageFlags"));
+	PageTableRDG = GraphBuilder.CreateBuffer(FRDGBufferDesc::CreateStructuredDesc(sizeof(uint32), NumPageFlags), TEXT("Shadow.Virtual.PageTable"));		
+	// Note: these are passed to the rendering and are not identical to the PageRequest flags coming in from GeneratePageFlagsFromPixels 
+	PageFlagsRDG = GraphBuilder.CreateBuffer(FRDGBufferDesc::CreateStructuredDesc(sizeof(uint32), NumPageFlags), TEXT("Shadow.Virtual.PageFlags"));
 
-		// One additional element as the last element is used as an atomic counter
-		FRDGBufferRef FreePhysicalPagesRDG = GraphBuilder.CreateBuffer(FRDGBufferDesc::CreateStructuredDesc(sizeof(int32), GetMaxPhysicalPages() + 1), TEXT("Shadow.Virtual.FreePhysicalPages"));
+	// One additional element as the last element is used as an atomic counter
+	FRDGBufferRef FreePhysicalPagesRDG = GraphBuilder.CreateBuffer(FRDGBufferDesc::CreateStructuredDesc(sizeof(int32), GetMaxPhysicalPages() + 1), TEXT("Shadow.Virtual.FreePhysicalPages"));
 		
-		// Enough space for all physical pages that might be allocated
-		PhysicalPageMetaDataRDG = GraphBuilder.CreateBuffer(FRDGBufferDesc::CreateStructuredDesc(sizeof(FPhysicalPageMetaData), GetMaxPhysicalPages()), TEXT("Shadow.Virtual.PhysicalPageMetaData"));
+	// Enough space for all physical pages that might be allocated
+	PhysicalPageMetaDataRDG = GraphBuilder.CreateBuffer(FRDGBufferDesc::CreateStructuredDesc(sizeof(FPhysicalPageMetaData), GetMaxPhysicalPages()), TEXT("Shadow.Virtual.PhysicalPageMetaData"));
 
-		AllocatedPageRectBoundsRDG = GraphBuilder.CreateBuffer(FRDGBufferDesc::CreateStructuredDesc(sizeof(FIntVector4), NumPageRects), TEXT("Shadow.Virtual.AllocatedPageRectBounds"));
+	AllocatedPageRectBoundsRDG = GraphBuilder.CreateBuffer(FRDGBufferDesc::CreateStructuredDesc(sizeof(FIntVector4), NumPageRects), TEXT("Shadow.Virtual.AllocatedPageRectBounds"));
 
-		{
-			FInitPhysicalPageMetaData::FParameters* PassParameters = GraphBuilder.AllocParameters<FInitPhysicalPageMetaData::FParameters>();
-			PassParameters->VirtualShadowMap		= GetUniformBuffer(GraphBuilder);
-			PassParameters->OutPhysicalPageMetaData = GraphBuilder.CreateUAV(PhysicalPageMetaDataRDG);
-			PassParameters->OutFreePhysicalPages	= GraphBuilder.CreateUAV(FreePhysicalPagesRDG);
+	{
+		FInitPhysicalPageMetaData::FParameters* PassParameters = GraphBuilder.AllocParameters<FInitPhysicalPageMetaData::FParameters>();
+		PassParameters->VirtualShadowMap		= GetUniformBuffer(GraphBuilder);
+		PassParameters->OutPhysicalPageMetaData = GraphBuilder.CreateUAV(PhysicalPageMetaDataRDG);
+		PassParameters->OutFreePhysicalPages	= GraphBuilder.CreateUAV(FreePhysicalPagesRDG);
 
-			auto ComputeShader = Views[0].ShaderMap->GetShader<FInitPhysicalPageMetaData>();
-			FComputeShaderUtils::AddPass(
-				GraphBuilder,
-				RDG_EVENT_NAME("InitPhysicalPageMetaData"),
-				ComputeShader,
-				PassParameters,
-				FIntVector(FMath::DivideAndRoundUp(GetMaxPhysicalPages(), FInitPhysicalPageMetaData::DefaultCSGroupX), 1, 1)
-			);
-		}
+		auto ComputeShader = Views[0].ShaderMap->GetShader<FInitPhysicalPageMetaData>();
+		FComputeShaderUtils::AddPass(
+			GraphBuilder,
+			RDG_EVENT_NAME("InitPhysicalPageMetaData"),
+			ComputeShader,
+			PassParameters,
+			FIntVector(FMath::DivideAndRoundUp(GetMaxPhysicalPages(), FInitPhysicalPageMetaData::DefaultCSGroupX), 1, 1)
+		);
+	}
 		
-		// Start by marking any physical pages that we are going to keep due to caching
-		// NOTE: We run this pass even with no caching since we still need to initialize the metadata
-		{
-			FCreateCachedPageMappingsCS::FParameters* PassParameters = GraphBuilder.AllocParameters< FCreateCachedPageMappingsCS::FParameters >();
-			PassParameters->VirtualShadowMap		 = GetUniformBuffer(GraphBuilder);
-			PassParameters->PageRequestFlags		 = GraphBuilder.CreateSRV(PageRequestFlagsRDG);
-			PassParameters->OutPageTable			 = GraphBuilder.CreateUAV(PageTableRDG);
-			PassParameters->OutPhysicalPageMetaData  = GraphBuilder.CreateUAV(PhysicalPageMetaDataRDG);
-			PassParameters->OutPageFlags			 = GraphBuilder.CreateUAV(PageFlagsRDG);
-			PassParameters->bDynamicPageInvalidation = 1;
+	// Start by marking any physical pages that we are going to keep due to caching
+	// NOTE: We run this pass even with no caching since we still need to initialize the metadata
+	{
+		FCreateCachedPageMappingsCS::FParameters* PassParameters = GraphBuilder.AllocParameters< FCreateCachedPageMappingsCS::FParameters >();
+		PassParameters->VirtualShadowMap		 = GetUniformBuffer(GraphBuilder);
+		PassParameters->PageRequestFlags		 = GraphBuilder.CreateSRV(PageRequestFlagsRDG);
+		PassParameters->OutPageTable			 = GraphBuilder.CreateUAV(PageTableRDG);
+		PassParameters->OutPhysicalPageMetaData  = GraphBuilder.CreateUAV(PhysicalPageMetaDataRDG);
+		PassParameters->OutPageFlags			 = GraphBuilder.CreateUAV(PageFlagsRDG);
+		PassParameters->bDynamicPageInvalidation = 1;
 #if !UE_BUILD_SHIPPING
-			PassParameters->bDynamicPageInvalidation = CVarDebugSkipDynamicPageInvalidation.GetValueOnRenderThread() == 0 ? 1 : 0;
+		PassParameters->bDynamicPageInvalidation = CVarDebugSkipDynamicPageInvalidation.GetValueOnRenderThread() == 0 ? 1 : 0;
 #endif
 
-			bool bCacheEnabled = CacheManager->IsValid();
-			if (bCacheEnabled)
-			{
-				SetCacheDataShaderParameters(GraphBuilder, ShadowMaps, CacheManager, PassParameters->CacheDataParameters);
-			}
-			bool bGenerateStats = StatsBufferRDG != nullptr;
-			if (bGenerateStats)
-			{
-				PassParameters->OutStatsBuffer = GraphBuilder.CreateUAV(StatsBufferRDG);
-			}
+		bool bCacheEnabled = CacheManager->IsValid();
+		if (bCacheEnabled)
+		{
+			SetCacheDataShaderParameters(GraphBuilder, ShadowMaps, CacheManager, PassParameters->CacheDataParameters);
+		}
+		bool bGenerateStats = StatsBufferRDG != nullptr;
+		if (bGenerateStats)
+		{
+			PassParameters->OutStatsBuffer = GraphBuilder.CreateUAV(StatsBufferRDG);
+		}
 
-			FCreateCachedPageMappingsCS::FPermutationDomain PermutationVector;
-			PermutationVector.Set<FCreateCachedPageMappingsCS::FHasCacheDataDim>(bCacheEnabled);
-			PermutationVector.Set<FCreateCachedPageMappingsCS::FGenerateStatsDim>(bGenerateStats);
-			auto ComputeShader = Views[0].ShaderMap->GetShader<FCreateCachedPageMappingsCS>(PermutationVector);
+		FCreateCachedPageMappingsCS::FPermutationDomain PermutationVector;
+		PermutationVector.Set<FCreateCachedPageMappingsCS::FHasCacheDataDim>(bCacheEnabled);
+		PermutationVector.Set<FCreateCachedPageMappingsCS::FGenerateStatsDim>(bGenerateStats);
+		auto ComputeShader = Views[0].ShaderMap->GetShader<FCreateCachedPageMappingsCS>(PermutationVector);
 
-			FComputeShaderUtils::AddPass(
-				GraphBuilder,
-				RDG_EVENT_NAME("CreateCachedPageMappings"),
-				ComputeShader,
-				PassParameters,
-				FIntVector(FMath::DivideAndRoundUp(FVirtualShadowMap::PageTableSize, FCreateCachedPageMappingsCS::DefaultCSGroupX), ShadowMaps.Num(), 1)
-			);
-		}
+		FComputeShaderUtils::AddPass(
+			GraphBuilder,
+			RDG_EVENT_NAME("CreateCachedPageMappings"),
+			ComputeShader,
+			PassParameters,
+			FIntVector(FMath::DivideAndRoundUp(FVirtualShadowMap::PageTableSize, FCreateCachedPageMappingsCS::DefaultCSGroupX), ShadowMaps.Num(), 1)
+		);
+	}
 
-		// After we've marked any cached pages, collect all the remaining free pages into a list
-		// NOTE: We could optimize this more in the case where there's no caching of course; TBD priority
-		{
-			FPackFreePagesCS::FParameters* PassParameters = GraphBuilder.AllocParameters<FPackFreePagesCS::FParameters>();
-			PassParameters->VirtualShadowMap		= GetUniformBuffer(GraphBuilder);
-			PassParameters->PhysicalPageMetaData	= GraphBuilder.CreateSRV(PhysicalPageMetaDataRDG);
-			PassParameters->OutFreePhysicalPages	= GraphBuilder.CreateUAV(FreePhysicalPagesRDG);
+	// After we've marked any cached pages, collect all the remaining free pages into a list
+	// NOTE: We could optimize this more in the case where there's no caching of course; TBD priority
+	{
+		FPackFreePagesCS::FParameters* PassParameters = GraphBuilder.AllocParameters<FPackFreePagesCS::FParameters>();
+		PassParameters->VirtualShadowMap		= GetUniformBuffer(GraphBuilder);
+		PassParameters->PhysicalPageMetaData	= GraphBuilder.CreateSRV(PhysicalPageMetaDataRDG);
+		PassParameters->OutFreePhysicalPages	= GraphBuilder.CreateUAV(FreePhysicalPagesRDG);
 			
-			auto ComputeShader = Views[0].ShaderMap->GetShader<FPackFreePagesCS>();
-			FComputeShaderUtils::AddPass(
-				GraphBuilder,
-				RDG_EVENT_NAME("PackFreePages"),
-				ComputeShader,
-				PassParameters,
-				FIntVector(FMath::DivideAndRoundUp(GetMaxPhysicalPages(), FPackFreePagesCS::DefaultCSGroupX), 1, 1)
-			);
-		}
+		auto ComputeShader = Views[0].ShaderMap->GetShader<FPackFreePagesCS>();
+		FComputeShaderUtils::AddPass(
+			GraphBuilder,
+			RDG_EVENT_NAME("PackFreePages"),
+			ComputeShader,
+			PassParameters,
+			FIntVector(FMath::DivideAndRoundUp(GetMaxPhysicalPages(), FPackFreePagesCS::DefaultCSGroupX), 1, 1)
+		);
+	}
 
-		// Allocate any new physical pages that were not cached from the free list
-		{
-			FAllocateNewPageMappingsCS::FParameters* PassParameters = GraphBuilder.AllocParameters<FAllocateNewPageMappingsCS::FParameters>();
-			PassParameters->VirtualShadowMap		= GetUniformBuffer(GraphBuilder);
-			PassParameters->PageRequestFlags		= GraphBuilder.CreateSRV(PageRequestFlagsRDG);
-			PassParameters->OutPageTable			= GraphBuilder.CreateUAV(PageTableRDG);
-			PassParameters->OutPageFlags			= GraphBuilder.CreateUAV(PageFlagsRDG);
-			PassParameters->OutFreePhysicalPages	= GraphBuilder.CreateUAV(FreePhysicalPagesRDG);
-			PassParameters->OutPhysicalPageMetaData = GraphBuilder.CreateUAV(PhysicalPageMetaDataRDG);
+	// Allocate any new physical pages that were not cached from the free list
+	{
+		FAllocateNewPageMappingsCS::FParameters* PassParameters = GraphBuilder.AllocParameters<FAllocateNewPageMappingsCS::FParameters>();
+		PassParameters->VirtualShadowMap		= GetUniformBuffer(GraphBuilder);
+		PassParameters->PageRequestFlags		= GraphBuilder.CreateSRV(PageRequestFlagsRDG);
+		PassParameters->OutPageTable			= GraphBuilder.CreateUAV(PageTableRDG);
+		PassParameters->OutPageFlags			= GraphBuilder.CreateUAV(PageFlagsRDG);
+		PassParameters->OutFreePhysicalPages	= GraphBuilder.CreateUAV(FreePhysicalPagesRDG);
+		PassParameters->OutPhysicalPageMetaData = GraphBuilder.CreateUAV(PhysicalPageMetaDataRDG);
 			
-			bool bGenerateStats = StatsBufferRDG != nullptr;
-			if (bGenerateStats)
-			{
-				PassParameters->OutStatsBuffer = GraphBuilder.CreateUAV(StatsBufferRDG);
-			}
+		bool bGenerateStats = StatsBufferRDG != nullptr;
+		if (bGenerateStats)
+		{
+			PassParameters->OutStatsBuffer = GraphBuilder.CreateUAV(StatsBufferRDG);
+		}
 
-			FAllocateNewPageMappingsCS::FPermutationDomain PermutationVector;
-			PermutationVector.Set<FAllocateNewPageMappingsCS::FGenerateStatsDim>(bGenerateStats);
-			auto ComputeShader = Views[0].ShaderMap->GetShader<FAllocateNewPageMappingsCS>(PermutationVector);
+		FAllocateNewPageMappingsCS::FPermutationDomain PermutationVector;
+		PermutationVector.Set<FAllocateNewPageMappingsCS::FGenerateStatsDim>(bGenerateStats);
+		auto ComputeShader = Views[0].ShaderMap->GetShader<FAllocateNewPageMappingsCS>(PermutationVector);
 
-			FComputeShaderUtils::AddPass(
-				GraphBuilder,
-				RDG_EVENT_NAME("AllocateNewPageMappings"),
-				ComputeShader,
-				PassParameters,
-				FIntVector(FMath::DivideAndRoundUp(FVirtualShadowMap::PageTableSize, FAllocateNewPageMappingsCS::DefaultCSGroupX), ShadowMaps.Num(), 1)
-			);
-		}
+		FComputeShaderUtils::AddPass(
+			GraphBuilder,
+			RDG_EVENT_NAME("AllocateNewPageMappings"),
+			ComputeShader,
+			PassParameters,
+			FIntVector(FMath::DivideAndRoundUp(FVirtualShadowMap::PageTableSize, FAllocateNewPageMappingsCS::DefaultCSGroupX), ShadowMaps.Num(), 1)
+		);
+	}
 
-		{
-			// Run pass building hierarchical page flags to make culling acceptable performance.
-			FGenerateHierarchicalPageFlagsCS::FParameters* PassParameters = GraphBuilder.AllocParameters<FGenerateHierarchicalPageFlagsCS::FParameters>();
-			PassParameters->VirtualShadowMap = GetUniformBuffer(GraphBuilder);
-			PassParameters->OutPageFlags = GraphBuilder.CreateUAV(PageFlagsRDG);
-			PassParameters->OutPageRectBounds = GraphBuilder.CreateUAV(PageRectBoundsRDG);
+	{
+		// Run pass building hierarchical page flags to make culling acceptable performance.
+		FGenerateHierarchicalPageFlagsCS::FParameters* PassParameters = GraphBuilder.AllocParameters<FGenerateHierarchicalPageFlagsCS::FParameters>();
+		PassParameters->VirtualShadowMap = GetUniformBuffer(GraphBuilder);
+		PassParameters->OutPageFlags = GraphBuilder.CreateUAV(PageFlagsRDG);
+		PassParameters->OutPageRectBounds = GraphBuilder.CreateUAV(PageRectBoundsRDG);
 
-			auto ComputeShader = Views[0].ShaderMap->GetShader<FGenerateHierarchicalPageFlagsCS>();
-			FComputeShaderUtils::AddPass(
-				GraphBuilder,
-				RDG_EVENT_NAME("GenerateHierarchicalPageFlags"),
-				ComputeShader,
-				PassParameters,
-				FIntVector(FMath::DivideAndRoundUp(FVirtualShadowMap::PageTableSize, FGenerateHierarchicalPageFlagsCS::DefaultCSGroupX), ShadowMaps.Num(), 1)
-			);
-		}
+		auto ComputeShader = Views[0].ShaderMap->GetShader<FGenerateHierarchicalPageFlagsCS>();
+		FComputeShaderUtils::AddPass(
+			GraphBuilder,
+			RDG_EVENT_NAME("GenerateHierarchicalPageFlags"),
+			ComputeShader,
+			PassParameters,
+			FIntVector(FMath::DivideAndRoundUp(FVirtualShadowMap::PageTableSize, FGenerateHierarchicalPageFlagsCS::DefaultCSGroupX), ShadowMaps.Num(), 1)
+		);
+	}
 
-		// NOTE: We could skip this (in shader) for shadow maps that only have 1 mip (ex. clipmaps)
-		{
-			// Propagate mapped mips down the hierarchy to allow O(1) lookup of coarser mapped pages
-			FPropagateMappedMipsCS::FParameters* PassParameters = GraphBuilder.AllocParameters<FPropagateMappedMipsCS::FParameters>();
-			PassParameters->VirtualShadowMap = GetUniformBuffer(GraphBuilder);
-			PassParameters->OutPageTable	= GraphBuilder.CreateUAV(PageTableRDG);
+	// NOTE: We could skip this (in shader) for shadow maps that only have 1 mip (ex. clipmaps)
+	{
+		// Propagate mapped mips down the hierarchy to allow O(1) lookup of coarser mapped pages
+		FPropagateMappedMipsCS::FParameters* PassParameters = GraphBuilder.AllocParameters<FPropagateMappedMipsCS::FParameters>();
+		PassParameters->VirtualShadowMap = GetUniformBuffer(GraphBuilder);
+		PassParameters->OutPageTable	= GraphBuilder.CreateUAV(PageTableRDG);
 
-			auto ComputeShader = Views[0].ShaderMap->GetShader<FPropagateMappedMipsCS>();
-			FComputeShaderUtils::AddPass(
-				GraphBuilder,
-				RDG_EVENT_NAME("PropagateMappedMips"),
-				ComputeShader,
-				PassParameters,
-				FIntVector(FMath::DivideAndRoundUp(FMath::Square(FVirtualShadowMap::Level0DimPagesXY), FPropagateMappedMipsCS::DefaultCSGroupX), ShadowMaps.Num(), 1)
-			);
-		}
+		auto ComputeShader = Views[0].ShaderMap->GetShader<FPropagateMappedMipsCS>();
+		FComputeShaderUtils::AddPass(
+			GraphBuilder,
+			RDG_EVENT_NAME("PropagateMappedMips"),
+			ComputeShader,
+			PassParameters,
+			FIntVector(FMath::DivideAndRoundUp(FMath::Square(FVirtualShadowMap::Level0DimPagesXY), FPropagateMappedMipsCS::DefaultCSGroupX), ShadowMaps.Num(), 1)
+		);
+	}
 
-		// Initialize the physical page pool
-		check(PhysicalPagePoolRDG != nullptr);
+	// Initialize the physical page pool
+	check(PhysicalPagePoolRDG != nullptr);
+	{
+		RDG_EVENT_SCOPE( GraphBuilder, "InitializePhysicalPages" );
+		if (CVarInitializePhysicalUsingIndirect.GetValueOnRenderThread() != 0)
 		{
-			RDG_EVENT_SCOPE( GraphBuilder, "InitializePhysicalPages" );
-			if (CVarInitializePhysicalUsingIndirect.GetValueOnRenderThread() != 0)
-			{
-				FRDGBufferRef InitializePagesIndirectArgsRDG = GraphBuilder.CreateBuffer(FRDGBufferDesc::CreateIndirectDesc(3), TEXT("Shadow.Virtual.InitializePagesIndirectArgs"));
-				// Note: We use GetTotalAllocatedPhysicalPages() to size the buffer as the selection shader emits both static/dynamic pages separately when enabled.
-				FRDGBufferRef PhysicalPagesToInitializeRDG = GraphBuilder.CreateBuffer(FRDGBufferDesc::CreateStructuredDesc(sizeof(int32), GetTotalAllocatedPhysicalPages() + 1), TEXT("Shadow.Virtual.PhysicalPagesToInitialize"));
+			FRDGBufferRef InitializePagesIndirectArgsRDG = GraphBuilder.CreateBuffer(FRDGBufferDesc::CreateIndirectDesc(3), TEXT("Shadow.Virtual.InitializePagesIndirectArgs"));
+			// Note: We use GetTotalAllocatedPhysicalPages() to size the buffer as the selection shader emits both static/dynamic pages separately when enabled.
+			FRDGBufferRef PhysicalPagesToInitializeRDG = GraphBuilder.CreateBuffer(FRDGBufferDesc::CreateStructuredDesc(sizeof(int32), GetTotalAllocatedPhysicalPages() + 1), TEXT("Shadow.Virtual.PhysicalPagesToInitialize"));
 
-				// 1. Initialize the indirect args buffer
-				AddClearIndirectDispatchArgs1DPass(GraphBuilder, InitializePagesIndirectArgsRDG);
-				// 2. Filter the relevant physical pages and set up the indirect args
+			// 1. Initialize the indirect args buffer
+			AddClearIndirectDispatchArgs1DPass(GraphBuilder, InitializePagesIndirectArgsRDG);
+			// 2. Filter the relevant physical pages and set up the indirect args
+			{
+				FSelectPagesToInitializeCS::FParameters* PassParameters = GraphBuilder.AllocParameters<FSelectPagesToInitializeCS::FParameters>();
+				PassParameters->VirtualShadowMap = GetUniformBuffer(GraphBuilder);
+				PassParameters->PhysicalPageMetaData = GraphBuilder.CreateSRV(PhysicalPageMetaDataRDG);
+				PassParameters->OutInitializePagesIndirectArgsBuffer = GraphBuilder.CreateUAV(InitializePagesIndirectArgsRDG);
+				PassParameters->OutPhysicalPagesToInitialize = GraphBuilder.CreateUAV(PhysicalPagesToInitializeRDG);
+				bool bGenerateStats = StatsBufferRDG != nullptr;
+				if (bGenerateStats)
 				{
-					FSelectPagesToInitializeCS::FParameters* PassParameters = GraphBuilder.AllocParameters<FSelectPagesToInitializeCS::FParameters>();
-					PassParameters->VirtualShadowMap = GetUniformBuffer(GraphBuilder);
-					PassParameters->PhysicalPageMetaData = GraphBuilder.CreateSRV(PhysicalPageMetaDataRDG);
-					PassParameters->OutInitializePagesIndirectArgsBuffer = GraphBuilder.CreateUAV(InitializePagesIndirectArgsRDG);
-					PassParameters->OutPhysicalPagesToInitialize = GraphBuilder.CreateUAV(PhysicalPagesToInitializeRDG);
-					bool bGenerateStats = StatsBufferRDG != nullptr;
-					if (bGenerateStats)
-					{
-						PassParameters->OutStatsBuffer = GraphBuilder.CreateUAV(StatsBufferRDG);
-					}
-					FSelectPagesToInitializeCS::FPermutationDomain PermutationVector;
-					PermutationVector.Set<FSelectPagesToInitializeCS::FGenerateStatsDim>(bGenerateStats);
+					PassParameters->OutStatsBuffer = GraphBuilder.CreateUAV(StatsBufferRDG);
+				}
+				FSelectPagesToInitializeCS::FPermutationDomain PermutationVector;
+				PermutationVector.Set<FSelectPagesToInitializeCS::FGenerateStatsDim>(bGenerateStats);
 
-					auto ComputeShader = GetGlobalShaderMap(GMaxRHIFeatureLevel)->GetShader<FSelectPagesToInitializeCS>(PermutationVector);
+				auto ComputeShader = GetGlobalShaderMap(GMaxRHIFeatureLevel)->GetShader<FSelectPagesToInitializeCS>(PermutationVector);
 
-					FComputeShaderUtils::AddPass(
-						GraphBuilder,
-						RDG_EVENT_NAME("SelectPagesToInitialize"),
-						ComputeShader,
-						PassParameters,
-						FIntVector(FMath::DivideAndRoundUp(GetMaxPhysicalPages(), FSelectPagesToInitializeCS::DefaultCSGroupX), 1, 1)
-					);
+				FComputeShaderUtils::AddPass(
+					GraphBuilder,
+					RDG_EVENT_NAME("SelectPagesToInitialize"),
+					ComputeShader,
+					PassParameters,
+					FIntVector(FMath::DivideAndRoundUp(GetMaxPhysicalPages(), FSelectPagesToInitializeCS::DefaultCSGroupX), 1, 1)
+				);
 
-				}
-				// 3. Indirect dispatch to clear the selected pages
-				{
-					FInitializePhysicalPagesIndirectCS::FParameters* PassParameters = GraphBuilder.AllocParameters<FInitializePhysicalPagesIndirectCS::FParameters>();
-					PassParameters->VirtualShadowMap = GetUniformBuffer(GraphBuilder);
-					PassParameters->OutPhysicalPagePool = GraphBuilder.CreateUAV(PhysicalPagePoolRDG);
-					PassParameters->IndirectArgs = InitializePagesIndirectArgsRDG;
-					PassParameters->PhysicalPagesToInitialize = GraphBuilder.CreateSRV(PhysicalPagesToInitializeRDG);
-					auto ComputeShader = GetGlobalShaderMap(GMaxRHIFeatureLevel)->GetShader<FInitializePhysicalPagesIndirectCS>();
-
-					FComputeShaderUtils::AddPass(
-						GraphBuilder,
-						RDG_EVENT_NAME("InitializePhysicalMemoryIndirect"),
-						ComputeShader,
-						PassParameters,
-						PassParameters->IndirectArgs,
-						0
-					);
-				}
 			}
-			else
+			// 3. Indirect dispatch to clear the selected pages
 			{
-				FInitializePhysicalPagesCS::FParameters* PassParameters = GraphBuilder.AllocParameters<FInitializePhysicalPagesCS::FParameters>();
+				FInitializePhysicalPagesIndirectCS::FParameters* PassParameters = GraphBuilder.AllocParameters<FInitializePhysicalPagesIndirectCS::FParameters>();
 				PassParameters->VirtualShadowMap = GetUniformBuffer(GraphBuilder);
-				PassParameters->PhysicalPageMetaData = GraphBuilder.CreateSRV(PhysicalPageMetaDataRDG);
 				PassParameters->OutPhysicalPagePool = GraphBuilder.CreateUAV(PhysicalPagePoolRDG);
-
-				auto ComputeShader = GetGlobalShaderMap(GMaxRHIFeatureLevel)->GetShader<FInitializePhysicalPagesCS>();
-
-				// Shader contains logic to deal with static cached pages if enabled
-				// We only need to launch one per page, even if there are multiple cached pages per page
-				FIntPoint DynamicPoolSize = GetDynamicPhysicalPoolSize();
+				PassParameters->IndirectArgs = InitializePagesIndirectArgsRDG;
+				PassParameters->PhysicalPagesToInitialize = GraphBuilder.CreateSRV(PhysicalPagesToInitializeRDG);
+				auto ComputeShader = GetGlobalShaderMap(GMaxRHIFeatureLevel)->GetShader<FInitializePhysicalPagesIndirectCS>();
 
 				FComputeShaderUtils::AddPass(
 					GraphBuilder,
-					RDG_EVENT_NAME("InitializePhysicalPages"),
+					RDG_EVENT_NAME("InitializePhysicalMemoryIndirect"),
 					ComputeShader,
 					PassParameters,
-					FIntVector(
-						FMath::DivideAndRoundUp(DynamicPoolSize.X, 16),
-						FMath::DivideAndRoundUp(DynamicPoolSize.Y, 16),
-						1)
+					PassParameters->IndirectArgs,
+					0
 				);
 			}
 		}
-
-		UniformParameters.PageTable = GraphBuilder.CreateSRV(PageTableRDG);
-		UniformParameters.PageFlags = GraphBuilder.CreateSRV(PageFlagsRDG);
-		UniformParameters.PageRectBounds = GraphBuilder.CreateSRV(PageRectBoundsRDG);
-
-		// Add pass to pipe back important stats
+		else
 		{
-
-			FVirtualSmFeedbackStatusCS::FParameters* PassParameters = GraphBuilder.AllocParameters<FVirtualSmFeedbackStatusCS::FParameters>();
-			PassParameters->FreePhysicalPages = GraphBuilder.CreateSRV(FreePhysicalPagesRDG);
-			PassParameters->GPUMessageParams = GPUMessage::GetShaderParameters(GraphBuilder);
-			PassParameters->StatusMessageId = CacheManager->StatusFeedbackSocket.GetMessageId().GetIndex();
+			FInitializePhysicalPagesCS::FParameters* PassParameters = GraphBuilder.AllocParameters<FInitializePhysicalPagesCS::FParameters>();
 			PassParameters->VirtualShadowMap = GetUniformBuffer(GraphBuilder);
+			PassParameters->PhysicalPageMetaData = GraphBuilder.CreateSRV(PhysicalPageMetaDataRDG);
+			PassParameters->OutPhysicalPagePool = GraphBuilder.CreateUAV(PhysicalPagePoolRDG);
+
+			auto ComputeShader = GetGlobalShaderMap(GMaxRHIFeatureLevel)->GetShader<FInitializePhysicalPagesCS>();
 
-			auto ComputeShader = GetGlobalShaderMap(GMaxRHIFeatureLevel)->GetShader<FVirtualSmFeedbackStatusCS>();
+			// Shader contains logic to deal with static cached pages if enabled
+			// We only need to launch one per page, even if there are multiple cached pages per page
+			FIntPoint DynamicPoolSize = GetDynamicPhysicalPoolSize();
 
 			FComputeShaderUtils::AddPass(
 				GraphBuilder,
-				RDG_EVENT_NAME("Feedback Status"),
+				RDG_EVENT_NAME("InitializePhysicalPages"),
 				ComputeShader,
 				PassParameters,
-				FIntVector(1, 1, 1)
+				FIntVector(
+					FMath::DivideAndRoundUp(DynamicPoolSize.X, 16),
+					FMath::DivideAndRoundUp(DynamicPoolSize.Y, 16),
+					1)
 			);
 		}
 	}
 
+	UniformParameters.PageTable = GraphBuilder.CreateSRV(PageTableRDG);
+	UniformParameters.PageFlags = GraphBuilder.CreateSRV(PageFlagsRDG);
+	UniformParameters.PageRectBounds = GraphBuilder.CreateSRV(PageRectBoundsRDG);
+
+	// Add pass to pipe back important stats
+	{
+
+		FVirtualSmFeedbackStatusCS::FParameters* PassParameters = GraphBuilder.AllocParameters<FVirtualSmFeedbackStatusCS::FParameters>();
+		PassParameters->FreePhysicalPages = GraphBuilder.CreateSRV(FreePhysicalPagesRDG);
+		PassParameters->GPUMessageParams = GPUMessage::GetShaderParameters(GraphBuilder);
+		PassParameters->StatusMessageId = CacheManager->StatusFeedbackSocket.GetMessageId().GetIndex();
+		PassParameters->VirtualShadowMap = GetUniformBuffer(GraphBuilder);
+
+		auto ComputeShader = GetGlobalShaderMap(GMaxRHIFeatureLevel)->GetShader<FVirtualSmFeedbackStatusCS>();
+
+		FComputeShaderUtils::AddPass(
+			GraphBuilder,
+			RDG_EVENT_NAME("Feedback Status"),
+			ComputeShader,
+			PassParameters,
+			FIntVector(1, 1, 1)
+		);
+	}
+
 #if !UE_BUILD_SHIPPING
 	// Only dump one frame of light data
 	GDumpVSMLightNames = false;
@@ -1615,7 +1620,8 @@ void FVirtualShadowMapArray::RenderDebugInfo(FRDGBuilder& GraphBuilder)
 			FComputeShaderUtils::GetGroupCount(FIntPoint(DebugTargetWidth, DebugTargetHeight), FVirtualPageManagementShader::DefaultCSGroupXY)
 		);
 
-		DebugVisualizationOutput = GraphBuilder.ConvertToExternalTexture(DebugOutput);
+		// TODO!
+		// DebugVisualizationOutput = GraphBuilder.ConvertToExternalTexture(DebugOutput);		
 	}
 }
 
@@ -2808,3 +2814,63 @@ uint32 FVirtualShadowMapArray::AddRenderViews(const FProjectedShadowInfo* Projec
 
 	return uint32(NumMaps);
 }
+
+void FVirtualShadowMapArray::AddVisualizePass(FRDGBuilder& GraphBuilder, const FViewInfo& View, FScreenPassTexture Output)
+{
+#if !UE_BUILD_SHIPPING
+	if (!IsAllocated() || DebugVisualizationOutput == nullptr)
+	{
+		return;
+	}
+
+	const FVirtualShadowMapVisualizationData& VisualizationData = GetVirtualShadowMapVisualizationData();
+	if (VisualizationData.IsActive() && VisualizeLight.IsValid())
+	{	
+		FCopyRectPS::FParameters* Parameters = GraphBuilder.AllocParameters<FCopyRectPS::FParameters>();
+		Parameters->InputTexture = DebugVisualizationOutput;
+		Parameters->InputSampler = TStaticSamplerState<SF_Bilinear, AM_Clamp, AM_Clamp, AM_Clamp>::GetRHI();
+		Parameters->RenderTargets[0] = FRenderTargetBinding(Output.Texture, ERenderTargetLoadAction::ENoAction);
+
+		TShaderMapRef<FCopyRectPS> PixelShader(View.ShaderMap);
+
+		FScreenPassTextureViewport InputViewport(DebugVisualizationOutput->Desc.Extent);
+		FScreenPassTextureViewport OutputViewport(Output);
+
+		// See CVarVisualizeLayout documentation
+		const int32 VisualizeLayout = CVarVisualizeLayout.GetValueOnRenderThread();
+		if (VisualizeLayout == 1)		// Thumbnail
+		{
+			const int32 TileWidth  = View.UnscaledViewRect.Width() / 3;
+			const int32 TileHeight = View.UnscaledViewRect.Height() / 3;
+
+			OutputViewport.Rect.Min = FIntPoint(0, 0);
+			OutputViewport.Rect.Max = OutputViewport.Rect.Min + FIntPoint(TileWidth, TileHeight);
+		}
+		else if (VisualizeLayout == 2)	// Split screen
+		{
+			InputViewport.Rect.Max = InputViewport.Rect.Min + (InputViewport.Rect.Width() / 2);
+			OutputViewport.Rect.Max = OutputViewport.Rect.Min + (OutputViewport.Rect.Width() / 2);
+		}
+
+		// Use separate input and output viewports w/ bilinear sampling to properly support dynamic resolution scaling
+		AddDrawScreenPass(GraphBuilder, RDG_EVENT_NAME("DrawTexture"), View, OutputViewport, InputViewport, PixelShader, Parameters, EScreenPassDrawFlags::None);
+		
+		// Visualization light name
+		{
+			FScreenPassRenderTarget OutputTarget(Output.Texture, View.UnscaledViewRect, ERenderTargetLoadAction::ELoad);
+
+			AddDrawCanvasPass(GraphBuilder, RDG_EVENT_NAME("Labels"), View, OutputTarget,
+				[&VisualizeLight=VisualizeLight, &OutputViewport=OutputViewport](FCanvas& Canvas)
+			{
+				const FLinearColor LabelColor(1, 1, 0);
+				Canvas.DrawShadowedString(
+					OutputViewport.Rect.Min.X + 8,
+					OutputViewport.Rect.Max.Y - 19,
+					*VisualizeLight.GetLightName(),
+					GetStatsFont(),
+					LabelColor);
+			});
+		}
+	}
+#endif
+}
diff --git a/Engine/Source/Runtime/Renderer/Private/VirtualShadowMaps/VirtualShadowMapArray.h b/Engine/Source/Runtime/Renderer/Private/VirtualShadowMaps/VirtualShadowMapArray.h
--- a/Engine/Source/Runtime/Renderer/Private/VirtualShadowMaps/VirtualShadowMapArray.h
+++ b/Engine/Source/Runtime/Renderer/Private/VirtualShadowMaps/VirtualShadowMapArray.h
@@ -17,6 +17,7 @@ class FVirtualShadowMapCacheEntry;
 class FVirtualShadowMapArrayCacheManager;
 struct FSortedLightSetSceneInfo;
 class FVirtualShadowMapClipmap;
+struct FScreenPassTexture;
 
 // TODO: does this exist?
 constexpr uint32 ILog2Const(uint32 n)
@@ -137,6 +138,50 @@ END_SHADER_PARAMETER_STRUCT()
 FMatrix CalcTranslatedWorldToShadowUVMatrix(const FMatrix& TranslatedWorldToShadowView, const FMatrix& ViewToClip);
 FMatrix CalcTranslatedWorldToShadowUVNormalMatrix(const FMatrix& TranslatedWorldToShadowView, const FMatrix& ViewToClip);
 
+struct FVirtualShadowMapVisualizeLightSearch
+{
+public:
+	FVirtualShadowMapVisualizeLightSearch()
+	{
+		Reset();
+	}
+	
+	void Reset()
+	{
+		FoundKey.Packed = 0;
+		FoundProxy = nullptr;
+		FoundVirtualShadowMapId = INDEX_NONE;
+	}
+
+	void CheckLight(const FLightSceneProxy* CheckProxy, int CheckVirtualShadowMapId);
+
+	bool IsValid() const { return FoundProxy != nullptr; }
+
+	int GetVirtualShadowMapId() const { return FoundVirtualShadowMapId; }
+	const FLightSceneProxy* GetProxy() const { return FoundProxy; }
+	const FString& GetLightName() const { return FoundProxy->GetOwnerNameOrLabel(); }
+
+private:
+	union SortKey
+	{
+		struct
+		{
+			// NOTE: Lowest to highest priority
+			uint32 bExists : 1;				// Catch-all
+			uint32 bDirectionalLight : 1;
+			uint32 bOwnerSelected : 1;		// In editor
+			uint32 bSelected : 1;			// In editor
+			uint32 bPartialNameMatch : 1;
+			uint32 bExactNameMatch : 1;
+		} Fields;
+		uint32 Packed;
+	};
+
+	SortKey FoundKey;
+	const FLightSceneProxy* FoundProxy = nullptr;
+	int FoundVirtualShadowMapId = INDEX_NONE;
+};
+
 class FVirtualShadowMapArray
 {
 public:	
@@ -180,6 +225,7 @@ public:
 		FRDGBuilder& GraphBuilder,
 		const FMinimalSceneTextures& SceneTextures,
 		const TArray<FViewInfo> &Views, 
+		const FEngineShowFlags& EngineShowFlags,
 		const FSortedLightSetSceneInfo& SortedLights, 
 		const TArray<FVisibleLightInfo, SceneRenderingAllocator> &VisibleLightInfos, 
 		const TArray<Nanite::FRasterResults, TInlineAllocator<2>> &NaniteRasterResults, 
@@ -225,6 +271,9 @@ public:
 	uint32 AddRenderViews(const TSharedPtr<FVirtualShadowMapClipmap>& Clipmap, float LODScaleFactor, bool bSetHzbParams, bool bUpdateHZBMetaData, TArray<Nanite::FPackedView, SceneRenderingAllocator>& OutVirtualShadowViews);
 	uint32 AddRenderViews(const FProjectedShadowInfo* ProjectedShadowInfo, float LODScaleFactor, bool bSetHzbParams, bool bUpdateHZBMetaData, TArray<Nanite::FPackedView, SceneRenderingAllocator>& OutVirtualShadowViews);
 
+	// Add visualization composite pass, if enabled
+	void AddVisualizePass(FRDGBuilder& GraphBuilder, const FViewInfo& View, FScreenPassTexture Output);
+
 	// We keep a reference to the cache manager that was used to initialize this frame as it owns some of the buffers
 	FVirtualShadowMapArrayCacheManager* CacheManager = nullptr;
 
@@ -274,11 +323,9 @@ public:
 
 	FRDGBufferRef StatsBufferRDG = nullptr;
 
-	TRefCountPtr<IPooledRenderTarget> DebugVisualizationOutput;
-	int DebugOutputType = 0;	// 0 = Disabled
-	// Base ID of the light that the user has selected for debug output (if present)
-	int DebugVirtualShadowMapId = INDEX_NONE;
-	FRDGTextureRef DebugVisualizationProjectionOutput = nullptr;
+	// Debug visualization
+	FRDGTextureRef DebugVisualizationOutput = nullptr;
+	FVirtualShadowMapVisualizeLightSearch VisualizeLight;
 
 private:
 	bool bInitialized = false;
diff --git a/Engine/Source/Runtime/Renderer/Private/VirtualShadowMaps/VirtualShadowMapProjection.cpp b/Engine/Source/Runtime/Renderer/Private/VirtualShadowMaps/VirtualShadowMapProjection.cpp
--- a/Engine/Source/Runtime/Renderer/Private/VirtualShadowMaps/VirtualShadowMapProjection.cpp
+++ b/Engine/Source/Runtime/Renderer/Private/VirtualShadowMaps/VirtualShadowMapProjection.cpp
@@ -5,6 +5,7 @@
 =============================================================================*/
 
 #include "VirtualShadowMapProjection.h"
+#include "VirtualShadowMapVisualizationData.h"
 #include "CoreMinimal.h"
 #include "Stats/Stats.h"
 #include "RHI.h"
@@ -108,6 +109,13 @@ static TAutoConsoleVariable<float> CVarSMRTTexelDitherScale(
 	ECVF_Scalability | ECVF_RenderThreadSafe
 );
 
+static TAutoConsoleVariable<int32> CVarForcePerLightShadowMaskClear(
+	TEXT( "r.Shadow.Virtual.ForcePerLightShadowMaskClear" ),
+	0,
+	TEXT( "" ),
+	ECVF_RenderThreadSafe
+);
+
 
 // Composite denoised shadow projection mask onto the light's shadow mask
 // Basically just a copy shader with a special blend mode
@@ -178,14 +186,14 @@ class FVirtualShadowMapProjectionCS : public FGlobalShader
 	class FSMRTAdaptiveRayCountDim	: SHADER_PERMUTATION_BOOL("SMRT_ADAPTIVE_RAY_COUNT");
 	class FOnePassProjectionDim		: SHADER_PERMUTATION_BOOL("ONE_PASS_PROJECTION");
 	class FHairStrandsDim			: SHADER_PERMUTATION_BOOL("HAS_HAIR_STRANDS");
-	class FDebugOutputDim			: SHADER_PERMUTATION_BOOL("DEBUG_OUTPUT");
+	class FVisualizeOutputDim		: SHADER_PERMUTATION_BOOL("VISUALIZE_OUTPUT");
 
 	using FPermutationDomain = TShaderPermutationDomain<
 		FDirectionalLightDim,
 		FOnePassProjectionDim,
 		FSMRTAdaptiveRayCountDim,
 		FHairStrandsDim,
-		FDebugOutputDim>;
+		FVisualizeOutputDim>;
 
 	BEGIN_SHADER_PARAMETER_STRUCT(FParameters, )
 		SHADER_PARAMETER_STRUCT_INCLUDE(FVirtualShadowMapSamplingParameters, SamplingParameters)
@@ -205,16 +213,16 @@ class FVirtualShadowMapProjectionCS : public FGlobalShader
 		SHADER_PARAMETER(uint32, bCullBackfacingPixels)
 		// One pass projection parameters
 		SHADER_PARAMETER_RDG_UNIFORM_BUFFER(FForwardLightData, ForwardLightData)
-		SHADER_PARAMETER_RDG_TEXTURE_UAV(RWTexture2D, RWShadowMaskBits)
+		SHADER_PARAMETER_RDG_TEXTURE_UAV(RWTexture2D, OutShadowMaskBits)
 		// Pass per light parameters
 		SHADER_PARAMETER_STRUCT(FLightShaderParameters, Light)
 		SHADER_PARAMETER(int32, LightUniformVirtualShadowMapId)
-		SHADER_PARAMETER_RDG_TEXTURE_UAV(RWTexture2D, RWShadowFactor)
-		// Debug output
+		SHADER_PARAMETER_RDG_TEXTURE_UAV(RWTexture2D, OutShadowFactor)
+		// Visualization output
 		SHADER_PARAMETER_RDG_BUFFER_SRV(StructuredBuffer< FPhysicalPageMetaData >, PhysicalPageMetaData)
-		SHADER_PARAMETER(int32, DebugOutputType)
-		SHADER_PARAMETER(int32, DebugVirtualShadowMapId)
-		SHADER_PARAMETER_RDG_TEXTURE_UAV(RWTexture2D, RWDebug)
+		SHADER_PARAMETER(int32, VisualizeModeId)
+		SHADER_PARAMETER(int32, VisualizeVirtualShadowMapId)
+		SHADER_PARAMETER_RDG_TEXTURE_UAV(RWTexture2D, OutVisualize)
 	END_SHADER_PARAMETER_STRUCT()
 
 	static void ModifyCompilationEnvironment( const FGlobalShaderPermutationParameters& Parameters, FShaderCompilerEnvironment& OutEnvironment )
@@ -293,7 +301,7 @@ static void RenderVirtualShadowMapProjectionCommon(
 	{
 		// One pass projection
 		PassParameters->ForwardLightData = View.ForwardLightingResources.ForwardLightUniformBuffer;
-		PassParameters->RWShadowMaskBits = GraphBuilder.CreateUAV( OutputTexture );	
+		PassParameters->OutShadowMaskBits = GraphBuilder.CreateUAV( OutputTexture );	
 	}
 	else
 	{
@@ -303,7 +311,7 @@ static void RenderVirtualShadowMapProjectionCommon(
 		LightProxy->GetLightShaderParameters(LightParameters);
 		LightParameters.MakeShaderParameters(View.ViewMatrices, PassParameters->Light);
 		PassParameters->LightUniformVirtualShadowMapId = VirtualShadowMapId;
-		PassParameters->RWShadowFactor = GraphBuilder.CreateUAV( OutputTexture );
+		PassParameters->OutShadowFactor = GraphBuilder.CreateUAV( OutputTexture );
 	}
  
 	if (bDirectionalLight)
@@ -323,13 +331,15 @@ static void RenderVirtualShadowMapProjectionCommon(
 	
 	bool bDebugOutput = false;
 #if !UE_BUILD_SHIPPING
-	if ( VirtualShadowMapArray.DebugVisualizationProjectionOutput && InputType == EVirtualShadowMapProjectionInputType::GBuffer )
+	if ( VirtualShadowMapArray.DebugVisualizationOutput && InputType == EVirtualShadowMapProjectionInputType::GBuffer )
 	{
+		const FVirtualShadowMapVisualizationData& VisualizationData = GetVirtualShadowMapVisualizationData();
+
 		bDebugOutput = true;
-		PassParameters->DebugOutputType = VirtualShadowMapArray.DebugOutputType;
-		PassParameters->DebugVirtualShadowMapId = VirtualShadowMapArray.DebugVirtualShadowMapId;
+		PassParameters->VisualizeModeId = VisualizationData.GetActiveModeID();
+		PassParameters->VisualizeVirtualShadowMapId = VirtualShadowMapArray.VisualizeLight.GetVirtualShadowMapId();
 		PassParameters->PhysicalPageMetaData = GraphBuilder.CreateSRV( VirtualShadowMapArray.PhysicalPageMetaDataRDG );
-		PassParameters->RWDebug = GraphBuilder.CreateUAV( VirtualShadowMapArray.DebugVisualizationProjectionOutput );
+		PassParameters->OutVisualize = GraphBuilder.CreateUAV( VirtualShadowMapArray.DebugVisualizationOutput );
 	}
 #endif
 
@@ -338,7 +348,7 @@ static void RenderVirtualShadowMapProjectionCommon(
 	PermutationVector.Set< FVirtualShadowMapProjectionCS::FOnePassProjectionDim >( bOnePassProjection );
 	PermutationVector.Set< FVirtualShadowMapProjectionCS::FSMRTAdaptiveRayCountDim >( bAdaptiveRayCount );
 	PermutationVector.Set< FVirtualShadowMapProjectionCS::FHairStrandsDim >( bHasHairStrandsData ? 1 : 0 );
-	PermutationVector.Set< FVirtualShadowMapProjectionCS::FDebugOutputDim >( bDebugOutput );
+	PermutationVector.Set< FVirtualShadowMapProjectionCS::FVisualizeOutputDim >( bDebugOutput );
 
 	auto ComputeShader = View.ShaderMap->GetShader< FVirtualShadowMapProjectionCS >( PermutationVector );
 	ClearUnusedGraphResources( ComputeShader, PassParameters );
@@ -386,7 +396,7 @@ FRDGTextureRef RenderVirtualShadowMapProjectionOnePass(
 	return ShadowMaskBits;
 }
 
-static FRDGTextureRef CreateShadowFactorTexture(FRDGBuilder& GraphBuilder, FIntPoint Extent)
+static FRDGTextureRef CreateShadowMaskTexture(FRDGBuilder& GraphBuilder, FIntPoint Extent)
 {
 	const FLinearColor ClearColor(0.0f, 0.0f, 0.0f, 0.0f);
 
@@ -396,8 +406,14 @@ static FRDGTextureRef CreateShadowFactorTexture(FRDGBuilder& GraphBuilder, FIntP
 		FClearValueBinding(ClearColor),
 		TexCreate_ShaderResource | TexCreate_UAV);
 
-	FRDGTextureRef Texture = GraphBuilder.CreateTexture(Desc, TEXT("Shadow.Virtual.ShadowFactor"));
-	AddClearUAVPass(GraphBuilder, GraphBuilder.CreateUAV(Texture), ClearColor);
+	FRDGTextureRef Texture = GraphBuilder.CreateTexture(Desc, TEXT("Shadow.Virtual.ShadowMask"));
+
+	// NOTE: Projection pass writes all relevant pixels, so should not need to clear here
+	if (CVarForcePerLightShadowMaskClear.GetValueOnRenderThread() != 0)
+	{
+		AddClearUAVPass(GraphBuilder, GraphBuilder.CreateUAV(Texture), ClearColor);
+	}
+
 	return Texture;
 }
 
@@ -411,7 +427,7 @@ void RenderVirtualShadowMapProjection(
 	FProjectedShadowInfo* ShadowInfo,
 	FRDGTextureRef OutputShadowMaskTexture)
 {
-	FRDGTextureRef VirtualShadowMaskTexture = CreateShadowFactorTexture(GraphBuilder, SceneTextures.Config.Extent);
+	FRDGTextureRef VirtualShadowMaskTexture = CreateShadowMaskTexture(GraphBuilder, SceneTextures.Config.Extent);
 
 	RenderVirtualShadowMapProjectionCommon(
 		GraphBuilder,
@@ -442,7 +458,7 @@ void RenderVirtualShadowMapProjection(
 	const TSharedPtr<FVirtualShadowMapClipmap>& Clipmap,
 	FRDGTextureRef OutputShadowMaskTexture)
 {
-	FRDGTextureRef VirtualShadowMaskTexture = CreateShadowFactorTexture(GraphBuilder, SceneTextures.Config.Extent);
+	FRDGTextureRef VirtualShadowMaskTexture = CreateShadowMaskTexture(GraphBuilder, SceneTextures.Config.Extent);
 
 	RenderVirtualShadowMapProjectionCommon(
 		GraphBuilder,
