Fix missing hair shadow when camera is within the light volume.

diff --git a/Engine/Source/Runtime/Renderer/Private/ShadowRendering.cpp b/Engine/Source/Runtime/Renderer/Private/ShadowRendering.cpp
--- a/Engine/Source/Runtime/Renderer/Private/ShadowRendering.cpp
+++ b/Engine/Source/Runtime/Renderer/Private/ShadowRendering.cpp
@@ -1242,7 +1242,7 @@ void FProjectedShadowInfo::RenderProjectionInternal(
 	const bool bSubPixelSupport = HairStrandsUniformBuffer != nullptr;// HairStrands::HasViewHairStrandsData(*View);
 	const bool bStencilTestEnabled = !bSubPixelSupport && GShadowStencilCulling;
 	const bool bDepthBoundsTestEnabled = IsWholeSceneDirectionalShadow() && GSupportsDepthBoundsTest && CVarCSMDepthBoundsTest.GetValueOnRenderThread() != 0 && !bSubPixelSupport;
-	const uint32 StencilRef = bSubPixelSupport && !IsWholeSceneDirectionalShadow() ? 1u : 0u;
+	const uint32 StencilRef = bSubPixelSupport && !IsWholeSceneDirectionalShadow() && !bCameraInsideShadowFrustum ? 1u : 0u;
 
 	if (!bDepthBoundsTestEnabled && bStencilTestEnabled)
 	{
