Fix compile issue

diff --git a/Engine/Source/Runtime/Renderer/Private/VirtualShadowMaps/VirtualShadowMapArray.h b/Engine/Source/Runtime/Renderer/Private/VirtualShadowMaps/VirtualShadowMapArray.h
--- a/Engine/Source/Runtime/Renderer/Private/VirtualShadowMaps/VirtualShadowMapArray.h
+++ b/Engine/Source/Runtime/Renderer/Private/VirtualShadowMaps/VirtualShadowMapArray.h
@@ -159,7 +159,7 @@ public:
 
 	int GetVirtualShadowMapId() const { return FoundVirtualShadowMapId; }
 	const FLightSceneProxy* GetProxy() const { return FoundProxy; }
-	const FString& GetLightName() const { return FoundProxy->GetOwnerNameOrLabel(); }
+	const FString GetLightName() const { return FoundProxy->GetOwnerNameOrLabel(); }
 
 private:
 	union SortKey
