Support Reflection Captures in projects with r.AllowStaticLighting=0.  This allows them to be used on Medium scalability, when Lumen GI is disabled.
 * ReflectionEnvironmentAndSky is now run on Reflection Captures to apply Movable skylight
 * Movable meshes and lights are now rendered into Reflection Captures when r.AllowStaticLighting=0

diff --git a/Engine/Shaders/Private/ReflectionEnvironmentComposite.ush b/Engine/Shaders/Private/ReflectionEnvironmentComposite.ush
--- a/Engine/Shaders/Private/ReflectionEnvironmentComposite.ush
+++ b/Engine/Shaders/Private/ReflectionEnvironmentComposite.ush
@@ -157,8 +157,8 @@ float3 CompositeReflectionCapturesAndSkylightTWS(
 #endif
 
 	// Apply indirect lighting scale while we have only accumulated reflection captures
-	ImageBasedReflections.rgb *= View.PrecomputedIndirectLightingColorScale;
-	CompositedAverageBrightness.x *= Luminance( View.PrecomputedIndirectLightingColorScale );
+	ImageBasedReflections.rgb *= View.PrecomputedIndirectSpecularColorScale;
+	CompositedAverageBrightness.x *= Luminance( View.PrecomputedIndirectSpecularColorScale );
 
 #if ENABLE_SKY_LIGHT
 
diff --git a/Engine/Source/Runtime/Engine/Public/SceneView.h b/Engine/Source/Runtime/Engine/Public/SceneView.h
--- a/Engine/Source/Runtime/Engine/Public/SceneView.h
+++ b/Engine/Source/Runtime/Engine/Public/SceneView.h
@@ -693,6 +693,7 @@ enum ETranslucencyVolumeCascade
 	VIEW_UNIFORM_BUFFER_MEMBER(float, DecalDepthBias) \
 	VIEW_UNIFORM_BUFFER_MEMBER(FVector3f, IndirectLightingColorScale) \
 	VIEW_UNIFORM_BUFFER_MEMBER(FVector3f, PrecomputedIndirectLightingColorScale) \
+	VIEW_UNIFORM_BUFFER_MEMBER(FVector3f, PrecomputedIndirectSpecularColorScale) \
 	VIEW_UNIFORM_BUFFER_MEMBER_ARRAY(FVector4f, AtmosphereLightDirection, [NUM_ATMOSPHERE_LIGHTS]) \
 	VIEW_UNIFORM_BUFFER_MEMBER_ARRAY(FLinearColor, AtmosphereLightIlluminanceOnGroundPostTransmittance, [NUM_ATMOSPHERE_LIGHTS]) \
 	VIEW_UNIFORM_BUFFER_MEMBER_ARRAY(FLinearColor, AtmosphereLightIlluminanceOuterSpace, [NUM_ATMOSPHERE_LIGHTS]) \
diff --git a/Engine/Source/Runtime/Renderer/Private/IndirectLightRendering.cpp b/Engine/Source/Runtime/Renderer/Private/IndirectLightRendering.cpp
--- a/Engine/Source/Runtime/Renderer/Private/IndirectLightRendering.cpp
+++ b/Engine/Source/Runtime/Renderer/Private/IndirectLightRendering.cpp
@@ -1376,8 +1376,6 @@ void FDeferredShadingSceneRenderer::RenderDeferredReflectionsAndSkyLighting(
 		return;
 	}
 
-	// If we're currently capturing a reflection capture, output SpecularColor * IndirectIrradiance for metals so they are not black in reflections,
-	// Since we don't have multiple bounce specular reflections
 	bool bReflectionCapture = false;
 	for (int32 ViewIndex = 0, Num = Views.Num(); ViewIndex < Num; ViewIndex++)
 	{
@@ -1385,12 +1383,6 @@ void FDeferredShadingSceneRenderer::RenderDeferredReflectionsAndSkyLighting(
 		bReflectionCapture = bReflectionCapture || View.bIsReflectionCapture;
 	}
 
-	if (bReflectionCapture)
-	{
-		// if we are rendering a reflection capture then we can skip this pass entirely (no reflection and no sky contribution evaluated in this pass)
-		return;
-	}
-
 	// The specular sky light contribution is also needed by RT Reflections as a fallback.
 	const bool bSkyLight = Scene->SkyLight
 		&& (Scene->SkyLight->ProcessedTexture || Scene->SkyLight->bRealTimeCaptureEnabled)
@@ -1408,7 +1400,8 @@ void FDeferredShadingSceneRenderer::RenderDeferredReflectionsAndSkyLighting(
 			&& !GDistanceFieldAOApplyToStaticIndirect
 			&& ShouldRenderDistanceFieldAO()
 			&& ShouldRenderDistanceFieldLighting()
-			&& ViewFamily.EngineShowFlags.AmbientOcclusion)
+			&& ViewFamily.EngineShowFlags.AmbientOcclusion
+			&& !bReflectionCapture)
 		{
 			bApplySkyShadowing = true;
 			FDistanceFieldAOParameters Parameters(Scene->SkyLight->OcclusionMaxDistance, Scene->SkyLight->Contrast);
diff --git a/Engine/Source/Runtime/Renderer/Private/RayTracing/RayTracingReflections.cpp b/Engine/Source/Runtime/Renderer/Private/RayTracing/RayTracingReflections.cpp
--- a/Engine/Source/Runtime/Renderer/Private/RayTracing/RayTracingReflections.cpp
+++ b/Engine/Source/Runtime/Renderer/Private/RayTracing/RayTracingReflections.cpp
@@ -399,7 +399,7 @@ bool ShouldRenderRayTracingReflections(const FViewInfo& View)
 		? bThisViewHasRaytracingReflections 
 		: GRayTracingReflections != 0;
 
-	const bool bReflectionPassEnabled = bReflectionsCvarEnabled && (GetRayTracingReflectionsSamplesPerPixel(View) > 0);
+	const bool bReflectionPassEnabled = bReflectionsCvarEnabled && (GetRayTracingReflectionsSamplesPerPixel(View) > 0) && !View.bIsReflectionCapture;
 		
 	return ShouldRenderRayTracingEffect(bReflectionPassEnabled, ERayTracingPipelineCompatibilityFlags::FullPipeline);
 }
diff --git a/Engine/Source/Runtime/Renderer/Private/ReflectionEnvironmentCapture.cpp b/Engine/Source/Runtime/Renderer/Private/ReflectionEnvironmentCapture.cpp
--- a/Engine/Source/Runtime/Renderer/Private/ReflectionEnvironmentCapture.cpp
+++ b/Engine/Source/Runtime/Renderer/Private/ReflectionEnvironmentCapture.cpp
@@ -1577,7 +1577,12 @@ void FScene::CaptureOrUploadReflectionCapture(UReflectionCaptureComponent* Captu
 			
 			if (CaptureComponent->ReflectionSourceType == EReflectionSourceType::CapturedScene)
 			{
-				bool const bCaptureStaticSceneOnly = CVarReflectionCaptureStaticSceneOnly.GetValueOnGameThread() != 0;
+				static const auto* AllowStaticLightingVar = IConsoleManager::Get().FindTConsoleVariableDataInt(TEXT("r.AllowStaticLighting"));
+				const bool bAllowStaticLighting = (!AllowStaticLightingVar || AllowStaticLightingVar->GetValueOnAnyThread() != 0);
+
+				// Reflection Captures are a form of static lighting, so only capture scene elements that are static
+				// However if the project has static lighting disabled, Reflection Captures can still be made to work by capturing Movable lights
+				bool const bCaptureStaticSceneOnly = CVarReflectionCaptureStaticSceneOnly.GetValueOnGameThread() != 0 && bAllowStaticLighting;
 				CaptureSceneIntoScratchCubemap(this, CaptureComponent->GetComponentLocation() + CaptureComponent->CaptureOffset, ReflectionCaptureSize, false, bCaptureStaticSceneOnly, 0, false, false, FLinearColor(), bCapturingForMobile);
 			}
 			else if (CaptureComponent->ReflectionSourceType == EReflectionSourceType::SpecifiedCubemap)
diff --git a/Engine/Source/Runtime/Renderer/Private/SceneRendering.cpp b/Engine/Source/Runtime/Renderer/Private/SceneRendering.cpp
--- a/Engine/Source/Runtime/Renderer/Private/SceneRendering.cpp
+++ b/Engine/Source/Runtime/Renderer/Private/SceneRendering.cpp
@@ -1672,11 +1672,20 @@ void FViewInfo::SetupUniformBufferParameters(
 
 	// If Lumen Dynamic GI is enabled then we don't want GI from Lightmaps
 	// Note: this has the side effect of removing direct lighting from Static Lights
-	if (FinalPostProcessSettings.DynamicGlobalIlluminationMethod == EDynamicGlobalIlluminationMethod::Lumen && DoesPlatformSupportLumenGI(ShaderPlatform))
+	if (ShouldRenderLumenDiffuseGI(Scene, *this))
 	{
 		ViewUniformShaderParameters.PrecomputedIndirectLightingColorScale = FVector::ZeroVector;
 	}
 
+	ViewUniformShaderParameters.PrecomputedIndirectSpecularColorScale = ViewUniformShaderParameters.IndirectLightingColorScale;
+
+	// If Lumen Reflections are enabled then we don't want precomputed reflections from reflection captures
+	// Note: this has the side effect of removing direct specular from Static Lights
+	if (ShouldRenderLumenReflections(*this))
+	{
+		ViewUniformShaderParameters.PrecomputedIndirectSpecularColorScale = FVector::ZeroVector;
+	}
+
 	ViewUniformShaderParameters.NormalCurvatureToRoughnessScaleBias.X = FMath::Clamp(CVarNormalCurvatureToRoughnessScale.GetValueOnAnyThread(), 0.0f, 2.0f);
 	ViewUniformShaderParameters.NormalCurvatureToRoughnessScaleBias.Y = FMath::Clamp(CVarNormalCurvatureToRoughnessBias.GetValueOnAnyThread(), -1.0f, 1.0f);
 	ViewUniformShaderParameters.NormalCurvatureToRoughnessScaleBias.Z = FMath::Clamp(CVarNormalCurvatureToRoughnessExponent.GetValueOnAnyThread(), .05f, 20.0f);
diff --git a/Engine/Source/Runtime/Renderer/Private/SceneVisibility.cpp b/Engine/Source/Runtime/Renderer/Private/SceneVisibility.cpp
--- a/Engine/Source/Runtime/Renderer/Private/SceneVisibility.cpp
+++ b/Engine/Source/Runtime/Renderer/Private/SceneVisibility.cpp
@@ -4977,7 +4977,9 @@ void FSceneRenderer::SetupSceneReflectionCaptureBuffer(FRHICommandListImmediate&
 		View.NumSphereReflectionCaptures = 0;
 		View.FurthestReflectionCaptureDistance = 0.0f;
 
-		if (View.Family->EngineShowFlags.ReflectionEnvironment)
+		if (View.Family->EngineShowFlags.ReflectionEnvironment 
+			// Avoid feedback
+			&& !View.bIsReflectionCapture)
 		{
 			View.NumBoxReflectionCaptures = Scene->ReflectionSceneData.NumBoxCaptures;
 			View.NumSphereReflectionCaptures = Scene->ReflectionSceneData.NumSphereCaptures;
diff --git a/Engine/Source/Runtime/Renderer/Private/ScreenSpaceRayTracing.cpp b/Engine/Source/Runtime/Renderer/Private/ScreenSpaceRayTracing.cpp
--- a/Engine/Source/Runtime/Renderer/Private/ScreenSpaceRayTracing.cpp
+++ b/Engine/Source/Runtime/Renderer/Private/ScreenSpaceRayTracing.cpp
@@ -121,7 +121,10 @@ bool ShouldKeepBleedFreeSceneColor(const FViewInfo& View)
 
 bool ShouldRenderScreenSpaceReflections(const FViewInfo& View)
 {
-	if(!View.Family->EngineShowFlags.ScreenSpaceReflections || View.FinalPostProcessSettings.ReflectionMethod != EReflectionMethod::ScreenSpace || HasRayTracedOverlay(*View.Family))
+	if(!View.Family->EngineShowFlags.ScreenSpaceReflections 
+		|| View.FinalPostProcessSettings.ReflectionMethod != EReflectionMethod::ScreenSpace 
+		|| HasRayTracedOverlay(*View.Family)
+		|| View.bIsReflectionCapture)
 	{
 		return false;
 	}
