Fix over-allocation of shadow mask/visualization textures causing alignment issues when viewport was resized down in editor.

diff --git a/Engine/Source/Runtime/Renderer/Private/VirtualShadowMaps/VirtualShadowMapArray.cpp b/Engine/Source/Runtime/Renderer/Private/VirtualShadowMaps/VirtualShadowMapArray.cpp
--- a/Engine/Source/Runtime/Renderer/Private/VirtualShadowMaps/VirtualShadowMapArray.cpp
+++ b/Engine/Source/Runtime/Renderer/Private/VirtualShadowMaps/VirtualShadowMapArray.cpp
@@ -1033,6 +1033,7 @@ void FVirtualShadowMapArray::BuildPageAllocations(
 	
 		// TODO: Support more than one view/debug output
 		const FName& VisualizationMode = Views[0].CurrentVirtualShadowMapVisualizationMode;
+		FIntPoint Extent = Views[0].ViewRect.Max;
 		if (VisualizationData.Update(VisualizationMode))
 		{
 			// TODO - automatically enable the show flag when set from command line?
@@ -1042,7 +1043,7 @@ void FVirtualShadowMapArray::BuildPageAllocations(
 		if (VisualizationData.IsActive() && EngineShowFlags.VisualizeVirtualShadowMap)
 		{
 			bDebugOutputEnabled = true;
-			DebugVisualizationOutput = CreateDebugOutputTexture(GraphBuilder, SceneTextures.Config.Extent);
+			DebugVisualizationOutput = CreateDebugOutputTexture(GraphBuilder, Extent);
 		}
 	}
 #endif //!UE_BUILD_SHIPPING
@@ -2848,8 +2849,8 @@ void FVirtualShadowMapArray::AddVisualizePass(FRDGBuilder& GraphBuilder, const F
 		}
 		else if (VisualizeLayout == 2)	// Split screen
 		{
-			InputViewport.Rect.Max = InputViewport.Rect.Min + (InputViewport.Rect.Width() / 2);
-			OutputViewport.Rect.Max = OutputViewport.Rect.Min + (OutputViewport.Rect.Width() / 2);
+			InputViewport.Rect.Max.X = InputViewport.Rect.Min.X + (InputViewport.Rect.Width() / 2);
+			OutputViewport.Rect.Max.X = OutputViewport.Rect.Min.X + (OutputViewport.Rect.Width() / 2);
 		}
 
 		// Use separate input and output viewports w/ bilinear sampling to properly support dynamic resolution scaling
diff --git a/Engine/Source/Runtime/Renderer/Private/VirtualShadowMaps/VirtualShadowMapProjection.cpp b/Engine/Source/Runtime/Renderer/Private/VirtualShadowMaps/VirtualShadowMapProjection.cpp
--- a/Engine/Source/Runtime/Renderer/Private/VirtualShadowMaps/VirtualShadowMapProjection.cpp
+++ b/Engine/Source/Runtime/Renderer/Private/VirtualShadowMaps/VirtualShadowMapProjection.cpp
@@ -427,7 +427,7 @@ void RenderVirtualShadowMapProjection(
 	FProjectedShadowInfo* ShadowInfo,
 	FRDGTextureRef OutputShadowMaskTexture)
 {
-	FRDGTextureRef VirtualShadowMaskTexture = CreateShadowMaskTexture(GraphBuilder, SceneTextures.Config.Extent);
+	FRDGTextureRef VirtualShadowMaskTexture = CreateShadowMaskTexture(GraphBuilder, View.ViewRect.Max);
 
 	RenderVirtualShadowMapProjectionCommon(
 		GraphBuilder,
@@ -458,7 +458,7 @@ void RenderVirtualShadowMapProjection(
 	const TSharedPtr<FVirtualShadowMapClipmap>& Clipmap,
 	FRDGTextureRef OutputShadowMaskTexture)
 {
-	FRDGTextureRef VirtualShadowMaskTexture = CreateShadowMaskTexture(GraphBuilder, SceneTextures.Config.Extent);
+	FRDGTextureRef VirtualShadowMaskTexture = CreateShadowMaskTexture(GraphBuilder, View.ViewRect.Max);
 
 	RenderVirtualShadowMapProjectionCommon(
 		GraphBuilder,
