Remove Geometry Shader path for point-light shadow rendering, disable point light shadows if VertexShaderLayer is not supported (at runtime).
- Add r.Shadow.DetectVertexShaderLayerAtRuntime to make it possible to force (for DX11/12/SM5/Vulkan) compilation of vslayer despite RHISupportsVertexShaderLayer being false.
- Set r.Shadow.DetectVertexShaderLayerAtRuntime in WindowsEngine.ini,WinGDKEngine.ini & LinuxEngine.ini.
- Add DoesRuntimeSupportOnePassPointLightShadows to determine support at runtime.
- Fix bug excluding Nanite meshes from DF shadows for local lights if VSM is enabled (now only affects directional lights).
- Add OnGetOnScreenMessages to enable modular generation of on-screen messages in the scene renderer (aimed at transient rendering processes).

diff --git a/Engine/Config/Linux/LinuxEngine.ini b/Engine/Config/Linux/LinuxEngine.ini
--- a/Engine/Config/Linux/LinuxEngine.ini
+++ b/Engine/Config/Linux/LinuxEngine.ini
@@ -7,6 +7,7 @@ r.setres=1280x720
 fx.NiagaraAllowRuntimeScalabilityChanges=1
 ;Perplatform to PerQualityLevel conversion mapping for platform groups
 QualityLevelMapping="high"
+r.Shadow.DetectVertexShaderLayerAtRuntime=1
 
 [SystemSettingsEditor]
 ;r.RHICmdBypass=1
diff --git a/Engine/Config/Windows/WindowsEngine.ini b/Engine/Config/Windows/WindowsEngine.ini
--- a/Engine/Config/Windows/WindowsEngine.ini
+++ b/Engine/Config/Windows/WindowsEngine.ini
@@ -19,6 +19,7 @@ fx.NiagaraAllowRuntimeScalabilityChanges=1
 QualityLevelMapping="high"
 ; Enable single RHI thread stall for D3D11
 r.Occlusion.SingleRHIThreadStall=1
+r.Shadow.DetectVertexShaderLayerAtRuntime=1
 
 [PlatformCrypto]
 PlatformRequiresDataCrypto=True
diff --git a/Engine/Shaders/Private/ShadowDepthCommon.ush b/Engine/Shaders/Private/ShadowDepthCommon.ush
--- a/Engine/Shaders/Private/ShadowDepthCommon.ush
+++ b/Engine/Shaders/Private/ShadowDepthCommon.ush
@@ -30,15 +30,3 @@ struct FShadowDepthVSToPS
 	float3 PixelPosition	: TEXCOORD7;
 #endif
 };
-
-#if ONEPASS_POINTLIGHT_SHADOW
-//only for one pass pointlights
-struct FShadowDepthVSToGS : FShadowDepthVSToPS
-{
-	float4 GSPosition : TEXCOORD6;
-#if ENABLE_FALLBACK_POINTLIGHT_SHADOW_GS
-	nointerpolation uint ViewIndex : TEXCOORD9;
-#endif // ENABLE_FALLBACK_POINTLIGHT_SHADOW_GS
-};
-#endif
-
diff --git a/Engine/Shaders/Private/ShadowDepthVertexShader.usf b/Engine/Shaders/Private/ShadowDepthVertexShader.usf
--- a/Engine/Shaders/Private/ShadowDepthVertexShader.usf
+++ b/Engine/Shaders/Private/ShadowDepthVertexShader.usf
@@ -156,103 +156,15 @@ void TransformToVirtualSmPage(inout float4 PointClip, inout float4 ClipPlanesInO
 
 #endif // VIRTUAL_SM_ENABLED
 
-
-#if FEATURE_LEVEL >= FEATURE_LEVEL_SM4 && ONEPASS_POINTLIGHT_SHADOW
-
-// Layer Id per draw call for Vertex-Shader-Layer implementation
-#if VERTEXSHADER && USING_VERTEX_SHADER_LAYER
-uint LayerId;
-#endif
-
 uint bUseGpuSceneInstancing;
-
-struct FShadowDepthGSToPS
-{
-#if INTERPOLATE_VF_ATTRIBUTES
-	FShadowDepthVSToPS PSInputs;
-#endif
-	float4 OutPosition : SV_POSITION;
-
-	/** Controls which of the cube map faces to rasterize the primitive to, only the value from the first vertex is used. */
-	uint RTIndex : SV_RenderTargetArrayIndex;
-};
-
-bool IsTriangleInFrustum(float4 ClipSpacePositions[3])
-{
-	float4 FrustumTests0 = saturate(ClipSpacePositions[0].xyxy * float4(-1, -1, 1, 1) - ClipSpacePositions[0].w);
-	float4 FrustumTests1 = saturate(ClipSpacePositions[1].xyxy * float4(-1, -1, 1, 1) - ClipSpacePositions[1].w);
-	float4 FrustumTests2 = saturate(ClipSpacePositions[2].xyxy * float4(-1, -1, 1, 1) - ClipSpacePositions[2].w);
-	float4 FrustumTests = FrustumTests0 * FrustumTests1 * FrustumTests2;
-
-	float2 FrustumTestsZ0 = saturate(ClipSpacePositions[0].zz * float2(-1, 1) - ClipSpacePositions[0].w);
-	float2 FrustumTestsZ1 = saturate(ClipSpacePositions[1].zz * float2(-1, 1) - ClipSpacePositions[1].w);
-	float2 FrustumTestsZ2 = saturate(ClipSpacePositions[2].zz * float2(-1, 1) - ClipSpacePositions[2].w);
-	float2 FrustumTestsZ = FrustumTestsZ0 * FrustumTestsZ1 * FrustumTestsZ2;
-	return !( any(FrustumTests != 0) || any(FrustumTestsZ != 0) );
-}
-
-/** Allocate space for cloning to all 6 faces which is the worst case (no GPU-scene active). */
-/** GPU-Scene case: Instancing takes care of the clone, we just route to a layer. */
-[maxvertexcount(18)]
-void MainOnePassPointLightGS(triangle FShadowDepthVSToGS Input[3], inout TriangleStream<FShadowDepthGSToPS> OutStream)
-{
-#if ENABLE_FALLBACK_POINTLIGHT_SHADOW_GS
-	uint CubeFaceStartIndex = bUseGpuSceneInstancing ? Input[0].ViewIndex : 0;
-	uint CubeFaceEndIndex = bUseGpuSceneInstancing ? Input[0].ViewIndex + 1: 6;
-
-	// Clone the triangle to each face
-	for (int CubeFaceIndex = CubeFaceStartIndex; CubeFaceIndex < CubeFaceEndIndex; CubeFaceIndex++)
-#else //!ENABLE_FALLBACK_POINTLIGHT_SHADOW_GS
-	UNROLL
-	// Clone the triangle to each face
-	for (int CubeFaceIndex = 0; CubeFaceIndex < 6; CubeFaceIndex++)
-#endif // ENABLE_FALLBACK_POINTLIGHT_SHADOW_GS
-	{
-		float4 ClipSpacePositions[3];
-		UNROLL
-		for (int VertexIndex = 0; VertexIndex < 3; VertexIndex++)
-		{
-			// Calculate the clip space position for each cube face
-			// Note: We don't apply any constant/slop bias here. A constant bias is applied in the shadow projection 
-			// code (see CubemapHardwarePCF() in ShadowProjectionCommon.ush)
-			ClipSpacePositions[VertexIndex] = mul(Input[VertexIndex].GSPosition, PassStruct.ShadowViewProjectionMatrices[CubeFaceIndex]);
-		}
-		// Frustum culling, saves GPU time with high poly meshes
-		BRANCH
-		if (IsTriangleInFrustum(ClipSpacePositions))
-		{
-			FShadowDepthGSToPS Output;
-			Output.RTIndex = CubeFaceIndex;
-
-			UNROLL
-			for (int VertexIndex = 0; VertexIndex < 3; VertexIndex++)
-			{
-				Output.OutPosition = ClipSpacePositions[VertexIndex];
-
-			#if INTERPOLATE_VF_ATTRIBUTES
-				Output.PSInputs.FactoryInterpolants = Input[VertexIndex].FactoryInterpolants;
-			#endif
-
-			#if INTERPOLATE_POSITION
-				Output.PSInputs.PixelPosition = Input[VertexIndex].GSPosition.xyz;
-			#endif
-				OutStream.Append(Output);
-			}
-			OutStream.RestartStrip();
-		}
-	}
-}
-
+#if VERTEXSHADER && ONEPASS_POINTLIGHT_SHADOW
+uint LayerId;
 #endif
 
 void Main(
 	FVertexFactoryInput Input,
 #if ONEPASS_POINTLIGHT_SHADOW
-  #if IS_FOR_GEOMETRY_SHADER
-	out FShadowDepthVSToGS OutParameters,
-  #else
 	out FShadowDepthVSToPS OutParameters,
-  #endif
 #else
 	out FShadowDepthVSToPS OutParameters,
 #endif
@@ -260,8 +172,8 @@ void Main(
 	out nointerpolation uint PackedPageInfo : TEXCOORD8,
 #endif
 	out float4 OutPosition : SV_POSITION
-#if ONEPASS_POINTLIGHT_SHADOW && FEATURE_LEVEL >= FEATURE_LEVEL_SM4 && USING_VERTEX_SHADER_LAYER
-	,out uint LayerIndex : SV_RenderTargetArrayIndex
+#if ONEPASS_POINTLIGHT_SHADOW
+	, out uint LayerIndex : SV_RenderTargetArrayIndex
 #endif
 #if VIRTUAL_SM_ENABLED	
 	// OLA-TODO: this collides with instanced stereo, which thankfully is not used with shadow maps, so should be fine, presumably.
@@ -284,13 +196,6 @@ void Main(
 
 	OutPosition = WorldPos;
 
-	#if IS_FOR_GEOMETRY_SHADER
-		OutParameters.GSPosition = WorldPos;
-		#if ENABLE_FALLBACK_POINTLIGHT_SHADOW_GS
-			OutParameters.ViewIndex = VertexFactoryGetViewIndex(VFIntermediates);
-		#endif // ENABLE_FALLBACK_POINTLIGHT_SHADOW_GS
-	#endif // IS_FOR_GEOMETRY_SHADER
-
 	#if INTERPOLATE_VF_ATTRIBUTES
 		// Masked materials need texture coords to clip
 		OutParameters.FactoryInterpolants = VertexFactoryGetInterpolantsVSToPS(Input, VFIntermediates, VertexParameters);
@@ -300,13 +205,8 @@ void Main(
 		OutParameters.PixelPosition = WorldPos.xyz;
 	#endif
 
-	#if FEATURE_LEVEL >= FEATURE_LEVEL_SM4 && USING_VERTEX_SHADER_LAYER && VERTEXSHADER
-		LayerIndex = bUseGpuSceneInstancing ? VertexFactoryGetViewIndex(VFIntermediates) : LayerId;
-		OutPosition = mul(WorldPos, PassStruct.ShadowViewProjectionMatrices[LayerIndex]);
-		#if INTERPOLATE_VF_ATTRIBUTES && IS_FOR_GEOMETRY_SHADER
-			OutParameters.GSPosition = OutPosition;
-		#endif
-	#endif
+	LayerIndex = bUseGpuSceneInstancing ? VertexFactoryGetViewIndex(VFIntermediates) : LayerId;
+	OutPosition = mul(WorldPos, PassStruct.ShadowViewProjectionMatrices[LayerIndex]);
 
 #else
 	float Dummy;
@@ -357,54 +257,14 @@ void Main(
 #endif // VIRTUAL_SM_ENABLED
 }
 
-#if IS_FOR_GEOMETRY_SHADER
-
-// Vertex shader exports must match GS inputs.
-// This is a requirement on some platforms.
-void MainForGS(
-	FVertexFactoryInput Input,
-#if ONEPASS_POINTLIGHT_SHADOW
-	out FShadowDepthVSToGS OutParameters
-	#if FEATURE_LEVEL >= FEATURE_LEVEL_SM4 && USING_VERTEX_SHADER_LAYER
-		, out uint LayerIndex : SV_RenderTargetArrayIndex
-	#endif
-#else
-	#error Invalid permutation for MainForGS
-#endif
-#if COMPILER_METAL
-#define SKIP_UNUSED_POSITION
-	// Metal needs a position output
-	, out float4 UnusedPosition : SV_POSITION
-#endif
-	)
-{
-#ifndef SKIP_UNUSED_POSITION
-	float4 UnusedPosition;
-#endif
-
-	Main(Input,
-		OutParameters,
-#if FEATURE_LEVEL >= FEATURE_LEVEL_SM4 && ONEPASS_POINTLIGHT_SHADOW && USING_VERTEX_SHADER_LAYER
-		LayerIndex,
-#endif
-		UnusedPosition);
-}
-
-#endif // IS_FOR_GEOMETRY_SHADER
-
-#define USING_ONEPASS_POINT_VERTEX_LAYER	(FEATURE_LEVEL >= FEATURE_LEVEL_SM4 && ONEPASS_POINTLIGHT_SHADOW && USING_VERTEX_SHADER_LAYER)
-
 #if POSITION_ONLY
 void PositionOnlyMain(
 	in FPositionAndNormalOnlyVertexFactoryInput Input,
-#if USING_ONEPASS_POINT_VERTEX_LAYER
+#if ONEPASS_POINTLIGHT_SHADOW
 	out uint LayerIndex : SV_RenderTargetArrayIndex,
 #endif
-#if ONEPASS_POINTLIGHT_SHADOW
-	out FShadowDepthVSToGS OutParameters,
-#else
 	out FShadowDepthVSToPS OutParameters,
-#endif
+
 #if VIRTUAL_SM_ENABLED
 	out nointerpolation uint PackedPageInfo : TEXCOORD8,
 #endif
@@ -426,20 +286,10 @@ void PositionOnlyMain(
 	float3 WorldNormal = VertexFactoryGetWorldNormal(Input);
 
 #if ONEPASS_POINTLIGHT_SHADOW
-	OutPosition = OutParameters.GSPosition = WorldPos;
-#if ENABLE_FALLBACK_POINTLIGHT_SHADOW_GS
-	OutParameters.ViewIndex = VertexFactoryGetViewIndex(Input);
-#endif // !ENABLE_FALLBACK_POINTLIGHT_SHADOW_GS
-
-	#if USING_ONEPASS_POINT_VERTEX_LAYER
-		LayerIndex = bUseGpuSceneInstancing ? VertexFactoryGetViewIndex(Input) : LayerId;
-		OutPosition = mul(WorldPos, PassStruct.ShadowViewProjectionMatrices[LayerIndex]);
-		#if INTERPOLATE_VF_ATTRIBUTES
-			OutParameters.GSPosition = OutPosition;
-		#endif
-	#endif
+	LayerIndex = bUseGpuSceneInstancing ? VertexFactoryGetViewIndex(Input) : LayerId;
+	OutPosition = mul(WorldPos, PassStruct.ShadowViewProjectionMatrices[LayerIndex]);
 
-#else // #if ONEPASS_POINTLIGHT_SHADOW
+#else // !ONEPASS_POINTLIGHT_SHADOW
 	float ShadowDepth;
 	SetShadowDepthOutputs(
 		PassStruct.ProjectionMatrix, 
@@ -458,7 +308,7 @@ void PositionOnlyMain(
 	#if !PERSPECTIVE_CORRECT_DEPTH && !COMPILER_SUPPORTS_EMPTY_STRUCTS
 		OutParameters.Dummy = 0;
 	#endif
-#endif // #if ONEPASS_POINTLIGHT_SHADOW
+#endif // ONEPASS_POINTLIGHT_SHADOW
 
 #if INTERPOLATE_POSITION
 	OutParameters.PixelPosition = WorldPos.xyz;
@@ -482,40 +332,4 @@ void PositionOnlyMain(
 }
 
 
-#if IS_FOR_GEOMETRY_SHADER
-
-// Vertex shader exports must match GS inputs.
-// This is a requirement on some platforms.
-void PositionOnlyMainForGS(
-	in FPositionAndNormalOnlyVertexFactoryInput Input,
-#if USING_ONEPASS_POINT_VERTEX_LAYER
-	out uint LayerIndex : SV_RenderTargetArrayIndex,
-#endif
-#if ONEPASS_POINTLIGHT_SHADOW
-	out FShadowDepthVSToGS OutParameters
-#else
-	#error Invalid permutation for PositionOnlyMainForGS
-#endif
-#if COMPILER_METAL
-	#ifndef SKIP_UNUSED_POSITION
-		#define SKIP_UNUSED_POSITION
-	#endif
-	// Metal needs a position output
-	, out float4 UnusedPosition : SV_POSITION
-#endif
-	)
-{
-#ifndef SKIP_UNUSED_POSITION
-	float4 UnusedPosition;
-#endif
-
-	PositionOnlyMain(Input,
-#if USING_ONEPASS_POINT_VERTEX_LAYER
-		LayerIndex,
-#endif
-		OutParameters, UnusedPosition);
-}
-
-#endif // IS_FOR_GEOMETRY_SHADER
-
 #endif // POSITION_ONLY
diff --git a/Engine/Source/Runtime/Engine/Private/Components/PointLightComponent.cpp b/Engine/Source/Runtime/Engine/Private/Components/PointLightComponent.cpp
--- a/Engine/Source/Runtime/Engine/Private/Components/PointLightComponent.cpp
+++ b/Engine/Source/Runtime/Engine/Private/Components/PointLightComponent.cpp
@@ -67,9 +67,7 @@ bool FPointLightSceneProxy::GetWholeSceneProjectedShadowInitializer(const FScene
 		OutInitializer.WAxis = FVector4(0,0,1,0);
 		OutInitializer.MinLightW = 0.1f;
 		OutInitializer.MaxDistanceToCastInLightW = Radius;
-
-		bool bSupportsGeometryShaders = RHISupportsGeometryShaders(GShaderPlatformForFeatureLevel[ViewFamily.GetFeatureLevel()]) || RHISupportsVertexShaderLayer(ViewFamily.GetShaderPlatform());
-		OutInitializer.bOnePassPointLightShadow = bSupportsGeometryShaders;
+		OutInitializer.bOnePassPointLightShadow = true;
 
 		OutInitializer.bRayTracedDistanceField = UseRayTracedDistanceFieldShadows() && DoesPlatformSupportDistanceFieldShadowing(ViewFamily.GetShaderPlatform());
 		return true;
diff --git a/Engine/Source/Runtime/Engine/Private/Components/RectLightComponent.cpp b/Engine/Source/Runtime/Engine/Private/Components/RectLightComponent.cpp
--- a/Engine/Source/Runtime/Engine/Private/Components/RectLightComponent.cpp
+++ b/Engine/Source/Runtime/Engine/Private/Components/RectLightComponent.cpp
@@ -257,9 +257,7 @@ bool FRectLightSceneProxy::GetWholeSceneProjectedShadowInitializer(const FSceneV
 		OutInitializer.WAxis = FVector4(0, 0, 1, 0);
 		OutInitializer.MinLightW = 0.1f;
 		OutInitializer.MaxDistanceToCastInLightW = Radius;
-		
-		bool bSupportsGeometryShaders = RHISupportsGeometryShaders(GShaderPlatformForFeatureLevel[ViewFamily.GetFeatureLevel()]) || RHISupportsVertexShaderLayer(ViewFamily.GetShaderPlatform());
-		OutInitializer.bOnePassPointLightShadow = bSupportsGeometryShaders;
+		OutInitializer.bOnePassPointLightShadow = true;
 
 		OutInitializer.bRayTracedDistanceField = UseRayTracedDistanceFieldShadows() && DoesPlatformSupportDistanceFieldShadowing(ViewFamily.GetShaderPlatform());
 		return true;
diff --git a/Engine/Source/Runtime/RenderCore/Private/RenderUtils.cpp b/Engine/Source/Runtime/RenderCore/Private/RenderUtils.cpp
--- a/Engine/Source/Runtime/RenderCore/Private/RenderUtils.cpp
+++ b/Engine/Source/Runtime/RenderCore/Private/RenderUtils.cpp
@@ -1707,3 +1707,11 @@ RENDERCORE_API bool AreSkinCacheShadersEnabled(EShaderPlatform Platform)
 	static FShaderPlatformCachedIniValue<bool> PerPlatformCVar(TEXT("r.SkinCache.CompileShaders"));
 	return (PerPlatformCVar.Get(Platform) != 0);
 }
+
+RENDERCORE_API bool DoesRuntimeSupportOnePassPointLightShadows(EShaderPlatform Platform)
+{
+	static const auto CVar = IConsoleManager::Get().FindTConsoleVariableDataInt(TEXT("r.Shadow.DetectVertexShaderLayerAtRuntime"));
+
+	return RHISupportsVertexShaderLayer(Platform)
+		|| (CVar->GetValueOnAnyThread() != 0 && GRHISupportsArrayIndexFromAnyShader != 0);
+}
diff --git a/Engine/Source/Runtime/RenderCore/Public/RenderUtils.h b/Engine/Source/Runtime/RenderCore/Public/RenderUtils.h
--- a/Engine/Source/Runtime/RenderCore/Public/RenderUtils.h
+++ b/Engine/Source/Runtime/RenderCore/Public/RenderUtils.h
@@ -744,3 +744,10 @@ RENDERCORE_API bool ExcludeNonPipelinedShaderTypes(EShaderPlatform ShaderPlatfor
  *   Checks if skin cache shaders are enabled for the platform (via r.SkinCache.CompileShaders)
  */
 RENDERCORE_API bool AreSkinCacheShadersEnabled(EShaderPlatform Platform);
+
+
+/*
+ * Detect (at runtime) if the runtime supports rendering one-pass point light shadows (i.e., cube maps)
+ */
+RENDERCORE_API bool DoesRuntimeSupportOnePassPointLightShadows(EShaderPlatform Platform);
+
diff --git a/Engine/Source/Runtime/Renderer/Private/DistanceFieldShadowing.cpp b/Engine/Source/Runtime/Renderer/Private/DistanceFieldShadowing.cpp
--- a/Engine/Source/Runtime/Renderer/Private/DistanceFieldShadowing.cpp
+++ b/Engine/Source/Runtime/Renderer/Private/DistanceFieldShadowing.cpp
@@ -560,7 +560,8 @@ void CullDistanceFieldObjectsForLight(
 		PassParameters->ObjectExpandScale = bIsHeightfield ? 0.f : WorldToShadowValue.GetMaximumAxisScale();
 		PassParameters->NumShadowHullPlanes = NumPlanes;
 		PassParameters->ShadowBoundingSphere = (FVector4f)ShadowBoundingSphereValue; // LWC_TODO: Precision loss
-		PassParameters->bDrawNaniteMeshes = !LightSceneProxy->UseVirtualShadowMaps();
+		// Disable Nanite meshes for directional lights that use VSM since they draw into the VSM unconditionally (and would get double shadow)
+		PassParameters->bDrawNaniteMeshes = !(LightSceneProxy->UseVirtualShadowMaps() && LightSceneProxy->GetLightType() == LightType_Directional);
 
 		check(NumPlanes <= 12);
 
@@ -865,7 +866,7 @@ void FProjectedShadowInfo::BeginRenderRayTracedDistanceFieldProjection(
 				NumPlanes = CascadeSettings.ShadowBoundsAccurate.Planes.Num();
 				PlaneData = CascadeSettings.ShadowBoundsAccurate.Planes.GetData();
 			}
-			else if (bOnePassPointLightShadow)
+			else if (IsWholeScenePointLightShadow())
 			{
 				ShadowBoundingSphereValue = FVector4(ShadowBounds.Center.X, ShadowBounds.Center.Y, ShadowBounds.Center.Z, ShadowBounds.W);
 			}
diff --git a/Engine/Source/Runtime/Renderer/Private/SceneOcclusion.cpp b/Engine/Source/Runtime/Renderer/Private/SceneOcclusion.cpp
--- a/Engine/Source/Runtime/Renderer/Private/SceneOcclusion.cpp
+++ b/Engine/Source/Runtime/Renderer/Private/SceneOcclusion.cpp
@@ -1106,7 +1106,7 @@ static FViewOcclusionQueriesPerView AllocateOcclusionTests(const FScene* Scene,
 							// Skip virtual SMs, as they overlay the same light as physical SM and have their own culling methods anyway.
 							continue;
 						}
-						if (ProjectedShadowInfo.bOnePassPointLightShadow)
+						if (ProjectedShadowInfo.IsWholeScenePointLightShadow())
 						{
 							FRHIRenderQuery* ShadowOcclusionQuery;
 							if (AllocateProjectedShadowOcclusionQuery(View, ProjectedShadowInfo, NumBufferedFrames, SOQ_LightInfluenceSphere, ShadowOcclusionQuery))
diff --git a/Engine/Source/Runtime/Renderer/Private/SceneRendering.cpp b/Engine/Source/Runtime/Renderer/Private/SceneRendering.cpp
--- a/Engine/Source/Runtime/Renderer/Private/SceneRendering.cpp
+++ b/Engine/Source/Runtime/Renderer/Private/SceneRendering.cpp
@@ -3080,6 +3080,19 @@ FSceneRenderer::~FSceneRenderer()
 	Views.Empty();
 }
 
+
+FSceneRenderer::FScreenMessageWriter::FScreenMessageWriter(FCanvas& InCanvas, int32 InY)
+	: Canvas(InCanvas)
+	, Y(InY)
+{
+}
+
+void FSceneRenderer::FScreenMessageWriter::DrawLine(const FText& Message, int32 X, const FLinearColor& Color)
+{
+	Canvas.DrawShadowedText(X, Y, Message, GetStatsFont(), Color);
+	Y += 14;
+}
+
 /** 
 * Finishes the view family rendering.
 */
@@ -3224,11 +3237,13 @@ void FSceneRenderer::RenderFinish(FRDGBuilder& GraphBuilder, FRDGTextureRef View
 		FFXSystemInterface* FXInterface = Scene->GetFXSystem();
 		const bool bFxDebugDraw = FXInterface && FXInterface->ShouldDebugDraw_RenderThread();
 
+		const bool bHasDelegateWarnings = OnGetOnScreenMessages.IsBound();
+
 		const bool bAnyWarning = bShowPrecomputedVisibilityWarning || bShowDemotedLocalMemoryWarning || bShowGlobalClipPlaneWarning || bShowSkylightWarning || bShowPointLightWarning
 			|| bShowDFAODisabledWarning || bShowShadowedLightOverflowWarning || bShowMobileDynamicCSMWarning || bShowMobileLowQualityLightmapWarning || bShowMobileMovableDirectionalLightWarning
 			|| bMobileShowVertexFogWarning || bMobileMissingSkyMaterial || bShowSkinCacheOOM || bSingleLayerWaterWarning || bShowDFDisabledWarning || bShowNoSkyAtmosphereComponentWarning || bFxDebugDraw 
 			|| bLumenEnabledButHasNoDataForTracing || bLumenEnabledButDisabledForTheProject || bNaniteEnabledButNoAtomics || bNaniteEnabledButDisabledInProject || bRealTimeSkyCaptureButNothingToCapture || bShowWaitingSkylight
-			|| bShowLocalExposureDisabledWarning;
+			|| bShowLocalExposureDisabledWarning || bHasDelegateWarnings;
 
 		for(int32 ViewIndex = 0;ViewIndex < Views.Num();ViewIndex++)
 		{	
@@ -3272,7 +3287,8 @@ void FSceneRenderer::RenderFinish(FRDGBuilder& GraphBuilder, FRDGTextureRef View
 						(FCanvas& Canvas)
 					{
 						// so it can get the screen size
-						int32 Y = 130;
+						FScreenMessageWriter Writer(Canvas, 130);
+
 						// Make sure draws to the canvas are not rendered upside down.
 						Canvas.SetAllowSwitchVerticalAxis(true);
 						if (bViewParentOrFrozen)
@@ -3282,149 +3298,125 @@ void FSceneRenderer::RenderFinish(FRDGBuilder& GraphBuilder, FRDGTextureRef View
 								NSLOCTEXT("SceneRendering", "RenderingFrozen", "Rendering frozen...")
 								:
 								NSLOCTEXT("SceneRendering", "OcclusionChild", "Occlusion Child");
-							Canvas.DrawShadowedText(10, Y, StateText, GetStatsFont(), FLinearColor(0.8, 1.0, 0.2, 1.0));
-							Y += 14;
+							Writer.DrawLine(StateText, 10, FLinearColor(0.8, 1.0, 0.2, 1.0));
 						}
 						if (bShowPrecomputedVisibilityWarning)
 						{
 							static const FText Message = NSLOCTEXT("Renderer", "NoPrecomputedVisibility", "NO PRECOMPUTED VISIBILITY");
-							Canvas.DrawShadowedText(10, Y, Message, GetStatsFont(), FLinearColor(1.0, 0.05, 0.05, 1.0));
-							Y += 14;
+							Writer.DrawLine(Message);
 						}
 						if (bShowGlobalClipPlaneWarning)
 						{
 							static const FText Message = NSLOCTEXT("Renderer", "NoGlobalClipPlane", "PLANAR REFLECTION REQUIRES GLOBAL CLIP PLANE PROJECT SETTING ENABLED TO WORK PROPERLY");
-							Canvas.DrawShadowedText(10, Y, Message, GetStatsFont(), FLinearColor(1.0, 0.05, 0.05, 1.0));
-							Y += 14;
+							Writer.DrawLine(Message);
 						}
 						if (bShowDFAODisabledWarning)
 						{
 							static const FText Message = NSLOCTEXT("Renderer", "DFAODisabled", "Distance Field AO is disabled through scalability");
-							Canvas.DrawShadowedText(10, Y, Message, GetStatsFont(), FLinearColor(1.0, 0.05, 0.05, 1.0));
-							Y += 14;
+							Writer.DrawLine(Message);
 						}
 						if (bShowDFDisabledWarning)
 						{
 							static const FText Message = NSLOCTEXT("Renderer", "DFDisabled", "Mesh distance fields generation is disabled by project settings, cannot visualize DFAO, mesh or global distance field.");
-							Canvas.DrawShadowedText(10, Y, Message, GetStatsFont(), FLinearColor(1.0, 0.05, 0.05, 1.0));
-							Y += 14;
+							Writer.DrawLine(Message);
 						}
 
 						if (bShowNoSkyAtmosphereComponentWarning)
 						{
 							static const FText Message = NSLOCTEXT("Renderer", "SkyAtmosphere", "There is no SkyAtmosphere component to visualize.");
-							Canvas.DrawShadowedText(10, Y, Message, GetStatsFont(), FLinearColor(1.0, 0.05, 0.05, 1.0));
-							Y += 14;						
+							Writer.DrawLine(Message);
 						}
 						if (bShowSkylightWarning)
 						{
 							static const FText Message = NSLOCTEXT("Renderer", "SkylightNotSuppported", "PROJECT DOES NOT SUPPORT STATIONARY SKYLIGHT: ");
-							Canvas.DrawShadowedText(10, Y, Message, GetStatsFont(), FLinearColor(1.0, 0.05, 0.05, 1.0));
-							Y += 14;
+							Writer.DrawLine(Message);
 						}
 						if (bRealTimeSkyCaptureButNothingToCapture)
 						{
 							static const FText Message = NSLOCTEXT("Renderer", "Skylight", "A sky light with real-time capture enable is in the scene. It requires at least a SkyAtmosphere component, A volumetricCloud component or a mesh with a material tagged as IsSky. Otherwise it will be black.");
-							Canvas.DrawShadowedText(10, Y, Message, GetStatsFont(), FLinearColor(1.0, 0.05, 0.05, 1.0));
-							Y += 14;
+							Writer.DrawLine(Message);
 						}
 						if (bShowPointLightWarning)
 						{
 							static const FText Message = NSLOCTEXT("Renderer", "PointLight", "PROJECT DOES NOT SUPPORT WHOLE SCENE POINT LIGHT SHADOWS: ");
-							Canvas.DrawShadowedText(10, Y, Message, GetStatsFont(), FLinearColor(1.0, 0.05, 0.05, 1.0));
-							Y += 14;
-
-							for (const FString &LightName : UsedWholeScenePointLightNames)
+							Writer.DrawLine(Message);
+							for (const FString& LightName : UsedWholeScenePointLightNames)
 							{
-								Canvas.DrawShadowedText(10, Y, FText::FromString(LightName), GetStatsFont(), FLinearColor(1.0, 0.05, 0.05, 1.0));
-								Y += 14;
+								Writer.DrawLine(FText::FromString(LightName), 35);
 							}
 						}
 						if (bShowShadowedLightOverflowWarning)
 						{
 							static const FText Message = NSLOCTEXT("Renderer", "ShadowedLightOverflow", "TOO MANY OVERLAPPING SHADOWED MOVABLE LIGHTS, SHADOW CASTING DISABLED: ");
-							Canvas.DrawShadowedText(10, Y, Message, GetStatsFont(), FLinearColor(1.0, 0.05, 0.05, 1.0));
-							Y += 14;
+							Writer.DrawLine(Message);
 
 							for (const FString& LightName : Scene->OverflowingDynamicShadowedLights)
 							{
-								Canvas.DrawShadowedText(10, Y, FText::FromString(LightName), GetStatsFont(), FLinearColor(1.0, 0.05, 0.05, 1.0));
-								Y += 14;
+								Writer.DrawLine(FText::FromString(LightName));
 							}
 						}
 						if (bShowMobileLowQualityLightmapWarning)
 						{
 							static const FText Message = NSLOCTEXT("Renderer", "MobileLQLightmap", "MOBILE PROJECTS SUPPORTING STATIC LIGHTING MUST HAVE LQ LIGHTMAPS ENABLED");
-							Canvas.DrawShadowedText(10, Y, Message, GetStatsFont(), FLinearColor(1.0, 0.05, 0.05, 1.0));
-							Y += 14;
+							Writer.DrawLine(Message);
 						}
 						if (bShowMobileMovableDirectionalLightWarning)
 						{
 							static const FText Message = NSLOCTEXT("Renderer", "MobileMovableDirectional", "PROJECT HAS MOVABLE DIRECTIONAL LIGHTS ON MOBILE DISABLED");
-							Canvas.DrawShadowedText(10, Y, Message, GetStatsFont(), FLinearColor(1.0, 0.05, 0.05, 1.0));
-							Y += 14;
+							Writer.DrawLine(Message);
 						}
 						if (bShowMobileDynamicCSMWarning)
 						{
 							static const FText Message = (!ReadOnlyCVARCache.bMobileEnableStaticAndCSMShadowReceivers)
 								? NSLOCTEXT("Renderer", "MobileDynamicCSM", "PROJECT HAS MOBILE CSM SHADOWS FROM STATIONARY DIRECTIONAL LIGHTS DISABLED")
 								: NSLOCTEXT("Renderer", "MobileDynamicCSMDistFieldShadows", "MOBILE CSM+STATIC REQUIRES DISTANCE FIELD SHADOWS ENABLED FOR PROJECT");
-							Canvas.DrawShadowedText(10, Y, Message, GetStatsFont(), FLinearColor(1.0, 0.05, 0.05, 1.0));
-							Y += 14;
+							Writer.DrawLine(Message);
 						}
 						if (bMobileShowVertexFogWarning)
 						{
 							static const FText Message = NSLOCTEXT("Renderer", "MobileVertexFog", "PROJECT HAS VERTEX FOG ON MOBILE DISABLED");
-							Canvas.DrawShadowedText(10, Y, Message, GetStatsFont(), FLinearColor(1.0, 0.05, 0.05, 1.0));
-							Y += 14;
+							Writer.DrawLine(Message);
 						}
 
 						if (bMobileMissingSkyMaterial)
 						{
 							static const FText Message = NSLOCTEXT("Renderer", "MobileMissingSkyMaterial", "On mobile the SkyAtmosphere component needs a mesh with a material tagged as IsSky and using the SkyAtmosphere nodes to visualize the Atmosphere.");
-							Canvas.DrawShadowedText(10, Y, Message, GetStatsFont(), FLinearColor(1.0, 0.05, 0.05, 1.0));
-							Y += 14;
+							Writer.DrawLine(Message);
 						}
 
 						if (bShowSkinCacheOOM)
 						{
 							FString String = FString::Printf(TEXT("OUT OF MEMORY FOR SKIN CACHE, REQUIRES %.3f extra MB (currently at %.3f)"), (float)GPUSkinCacheExtraRequiredMemory / 1048576.0f, CVarSkinCacheOOM->GetValueOnAnyThread());
-							Canvas.DrawShadowedText(10, Y, FText::FromString(String), GetStatsFont(), FLinearColor(1.0, 0.05, 0.05, 1.0));
-							Y += 14;
+							Writer.DrawLine(FText::FromString(String));
 						}
 						if (bShowLocalExposureDisabledWarning)
 						{
 							static const FText Message = NSLOCTEXT("Renderer", "LocalExposureDisabled", "Local Exposure is disabled.");
-							Canvas.DrawShadowedText(10, Y, Message, GetStatsFont(), FLinearColor(1.0, 0.05, 0.05, 1.0));
-							Y += 14;
+							Writer.DrawLine(Message);
 						}
 
 						if (bLocked)
 						{
 							static const FText Message = NSLOCTEXT("Renderer", "ViewLocked", "VIEW LOCKED");
-							Canvas.DrawShadowedText(10, Y, Message, GetStatsFont(), FLinearColor(0.8, 1.0, 0.2, 1.0));
-							Y += 14;
+							Writer.DrawLine(Message, 10, FLinearColor(0.8, 1.0, 0.2, 1.0));
 						}
 
 						if (bSingleLayerWaterWarning)
 						{
 							static const FText Message = NSLOCTEXT("Renderer", "SingleLayerWater", "r.Water.SingleLayer rendering is disabled with a view containing mesh(es) using water material. Meshes are not visible.");
-							Canvas.DrawShadowedText(10, Y, Message, GetStatsFont(), FLinearColor(1.0, 0.05, 0.05, 1.0));
-							Y += 14;
+							Writer.DrawLine(Message);
 						}
 
 						if (bLumenEnabledButHasNoDataForTracing)
 						{
 							static const FText Message = NSLOCTEXT("Renderer", "LumenCantDisplay", "Lumen is enabled, but has no ray tracing data and won't operate correctly.\nEither configure Lumen to use software distance field ray tracing and enable 'Generate Mesh Distancefields' in project settings\nor configure Lumen to use Hardware Ray Tracing and enable 'Support Hardware Ray Tracing' in project settings.");
-							Canvas.DrawShadowedText(10, Y, Message, GetStatsFont(), FLinearColor(1.0, 0.05, 0.05, 1.0));
-							Y += 14;
+							Writer.DrawLine(Message);
 						}
 
 						if (bLumenEnabledButDisabledForTheProject)
 						{
 							static const FText Message = NSLOCTEXT("Renderer", "LumenDisabledForProject", "Lumen is enabled but cannot render, because the project has Lumen disabled in an ini (r.Lumen.Supported = 0)");
-							Canvas.DrawShadowedText(10, Y, Message, GetStatsFont(), FLinearColor(1.0, 0.05, 0.05, 1.0));
-							Y += 14;
+							Writer.DrawLine(Message);
 						}
 
 						if (bNaniteEnabledButNoAtomics)
@@ -3443,29 +3435,25 @@ void FSceneRenderer::RenderFinish(FRDGBuilder& GraphBuilder, FRDGTextureRef View
 							#endif
 							}
 
-							Canvas.DrawShadowedText(10, Y, FText::FromString(NaniteError), GetStatsFont(), FLinearColor(1.0, 0.05, 0.05, 1.0));
-							Y += 14;
+							Writer.DrawLine(FText::FromString(NaniteError));
 						}
 
 						if (bNaniteEnabledButDisabledInProject)
 						{
 							static const FText Message = NSLOCTEXT("Renderer", "NaniteDisabledForProject", "Nanite is enabled but cannot render, because the project has Nanite disabled in an ini (r.Nanite.ProjectEnabled = 0)");
-							Canvas.DrawShadowedText(10, Y, Message, GetStatsFont(), FLinearColor(1.0, 0.05, 0.05, 1.0));
-							Y += 14;
+							Writer.DrawLine(Message);
 						}
 
 						if (bShowDemotedLocalMemoryWarning)
 						{
 							FString String = FString::Printf(TEXT("Video memory has been exhausted (%.3f MB over budget). Expect extremely poor performance."), float(GDemotedLocalMemorySize) / 1048576.0f);
-							Canvas.DrawShadowedText(10, Y, FText::FromString(String), GetStatsFont(), FLinearColor(1.0, 0.05, 0.05, 1.0));
-							Y += 14;
+							Writer.DrawLine(FText::FromString(String));
 						}
 
 						if (bShowAmbientCubemapMipGenSettingsWarning)
 						{
 							static const FText Message = NSLOCTEXT("Renderer", "AmbientCubemapMipGenSettings", "Ambient cubemaps should use 'Angular' Mip Gen Settings.");
-							Canvas.DrawShadowedText(10, Y, Message, GetStatsFont(), FLinearColor(1.0, 0.05, 0.05, 1.0));
-							Y += 14;
+							Writer.DrawLine(Message);
 						}
 
 #if WITH_EDITOR
@@ -3493,10 +3481,10 @@ void FSceneRenderer::RenderFinish(FRDGBuilder& GraphBuilder, FRDGTextureRef View
 								String += TEXT("Meshes, Textures");
 							}
 							String += TEXT(" for final capture.");
-							Canvas.DrawShadowedText(10, Y, FText::FromString(String), GetStatsFont(), OrangeColor);
-							Y += 14;
+							Writer.DrawLine(FText::FromString(String), 10, OrangeColor);
 						}
 #endif
+						OnGetOnScreenMessages.Broadcast(Writer);
 					});
 					if (bFxDebugDraw)
 					{
diff --git a/Engine/Source/Runtime/Renderer/Private/SceneRendering.h b/Engine/Source/Runtime/Renderer/Private/SceneRendering.h
--- a/Engine/Source/Runtime/Renderer/Private/SceneRendering.h
+++ b/Engine/Source/Runtime/Renderer/Private/SceneRendering.h
@@ -2239,6 +2239,19 @@ private:
 #endif
 	bool bShadowDepthRenderCompleted;
 
+	struct FScreenMessageWriter
+	{
+		FScreenMessageWriter(FCanvas& InCanvas, int32 InY);
+
+		void DrawLine(const FText& Message, int32 X = 10, const FLinearColor& Color = FLinearColor(1.0, 0.05, 0.05, 1.0));
+
+		FCanvas& Canvas;
+		int32 Y;
+	};
+
+	DECLARE_MULTICAST_DELEGATE_OneParam(FSceneOnScreenMessagesDelegate, FScreenMessageWriter&);
+	FSceneOnScreenMessagesDelegate OnGetOnScreenMessages;
+
 	/** Distance field shadows to project. Used to avoid iterating through the scene lights array. */
 	TArray<FProjectedShadowInfo*, TInlineAllocator<2, SceneRenderingAllocator>> ProjectedDistanceFieldShadows;
 
diff --git a/Engine/Source/Runtime/Renderer/Private/ShadowDepthRendering.cpp b/Engine/Source/Runtime/Renderer/Private/ShadowDepthRendering.cpp
--- a/Engine/Source/Runtime/Renderer/Private/ShadowDepthRendering.cpp
+++ b/Engine/Source/Runtime/Renderer/Private/ShadowDepthRendering.cpp
@@ -265,7 +265,6 @@ enum EShadowDepthVertexShaderMode
 	VertexShadowDepth_PerspectiveCorrect,
 	VertexShadowDepth_OutputDepth,
 	VertexShadowDepth_OnePassPointLight,
-	VertexShadowDepth_VSLayer,
 	VertexShadowDepth_VirtualShadowMap,
 };
 
@@ -283,10 +282,19 @@ static bool MobileUsesPerspectiveCorrectShadowPermutation(EShaderPlatform Shader
 	return bMobileEnableMovableSpotlightsShadow;
 }
 
+
+static TAutoConsoleVariable<int32> CVarDetectVertexShaderLayerAtRuntime(
+	TEXT("r.Shadow.DetectVertexShaderLayerAtRuntime"),
+	0,
+	TEXT("Forces the compilation of the vslayer shader permutation even if the platform (RHI) does not declare compile-time support through RHISupportsVertexShaderLayer.")
+	TEXT("Enabled by default for windows/SM5 as DX11 almost universally supports this at runtime."),
+	ECVF_ReadOnly | ECVF_RenderThreadSafe);
+
+
 /**
 * A vertex shader for rendering the depth of a mesh.
 */
-template <EShadowDepthVertexShaderMode ShaderMode, bool bUsePositionOnlyStream, bool bIsForGeometryShader = false>
+template <EShadowDepthVertexShaderMode ShaderMode, bool bUsePositionOnlyStream>
 class TShadowDepthVS : public FShadowDepthVS
 {
 public:
@@ -299,7 +307,6 @@ public:
 		static const auto SupportAllShaderPermutationsVar = IConsoleManager::Get().FindTConsoleVariableDataInt(TEXT("r.SupportAllShaderPermutations"));
 		const bool bForceAllPermutations = SupportAllShaderPermutationsVar && SupportAllShaderPermutationsVar->GetValueOnAnyThread() != 0;
 		const bool bSupportPointLightWholeSceneShadows = CVarSupportPointLightWholeSceneShadows.GetValueOnAnyThread() != 0 || bForceAllPermutations;
-		const bool bRHISupportsShadowCastingPointLights = RHISupportsGeometryShaders(Platform) || RHISupportsVertexShaderLayer(Platform);
 
 		// Mobile only needs OutputDepth, and optionally PerspectiveCorrect
 		if (!IsFeatureLevelSupported(Platform, ERHIFeatureLevel::SM5) &&
@@ -308,48 +315,28 @@ public:
 		{
 			return false;
 		}
-		
-		if (bIsForGeometryShader && ShaderMode == VertexShadowDepth_VSLayer)
-		{
-			return false;
-		}
 
-		if (bIsForGeometryShader && (!bSupportPointLightWholeSceneShadows || !bRHISupportsShadowCastingPointLights))
+		// Compile VS layer permutation if RHI supports it unconditionally OR we have forced it on (default for DX11&12 at SM5)
+		static FShaderPlatformCachedIniValue<bool> DetectVertexShaderLayerRuntimeIniValue(TEXT("r.Shadow.DetectVertexShaderLayerAtRuntime"));
+		const bool bShouldCompileVSLayer = RHISupportsVertexShaderLayer(Platform) || DetectVertexShaderLayerRuntimeIniValue.Get(Platform) != 0;
+		if (ShaderMode == VertexShadowDepth_OnePassPointLight && !bShouldCompileVSLayer)
 		{
 			return false;
 		}
 
 		if (ShaderMode == VertexShadowDepth_VirtualShadowMap &&
-			(!IsFeatureLevelSupported(Platform, ERHIFeatureLevel::SM5) || !UseGPUScene(Platform) || GEnableNonNaniteVSM == 0 || bIsForGeometryShader))
+			(!IsFeatureLevelSupported(Platform, ERHIFeatureLevel::SM5) || !UseGPUScene(Platform) || GEnableNonNaniteVSM == 0))
 		{
 			return false;
 		}
 
-		// VSLayer and OnePassPointLight are mutually exclusive (see GetShadowDepthPassShaders())
-		if (ShaderMode == VertexShadowDepth_OnePassPointLight || ShaderMode == VertexShadowDepth_VSLayer)
-		{
-			if (RHISupportsVertexShaderLayer(Platform))
-			{
-				if (ShaderMode == VertexShadowDepth_OnePassPointLight)
-				{
-					return false;
-				}
-			}
-			else
-			{
-				if (ShaderMode == VertexShadowDepth_VSLayer)
-				{
-					return false;
-				}
-			}
-		}
 
 		//Note: This logic needs to stay in sync with OverrideWithDefaultMaterialForShadowDepth!
 		return (Parameters.MaterialParameters.bIsSpecialEngineMaterial
 			// Masked and WPO materials need their shaders but cannot be used with a position only stream.
 			|| ((!Parameters.MaterialParameters.bWritesEveryPixelShadowPass || Parameters.MaterialParameters.bMaterialMayModifyMeshPosition) && !bUsePositionOnlyStream))
 			// Only compile one pass point light shaders for feature levels >= SM5
-				&& ((ShaderMode != VertexShadowDepth_OnePassPointLight && ShaderMode != VertexShadowDepth_VSLayer) || IsFeatureLevelSupported(Platform, ERHIFeatureLevel::SM5))
+				&& (ShaderMode != VertexShadowDepth_OnePassPointLight || IsFeatureLevelSupported(Platform, ERHIFeatureLevel::SM5))
 			// Only compile position-only shaders for vertex factories that support it. (Note: this assumes that a vertex factor which supports PositionOnly, supports also PositionAndNormalOnly)
 			&& (!bUsePositionOnlyStream || Parameters.VertexFactoryType->SupportsPositionOnly())
 			// Don't render ShadowDepth for translucent unlit materials
@@ -360,11 +347,8 @@ public:
 	{
 		FShadowDepthVS::ModifyCompilationEnvironment(Parameters, OutEnvironment);
 		OutEnvironment.SetDefine(TEXT("PERSPECTIVE_CORRECT_DEPTH"), (uint32)(ShaderMode == VertexShadowDepth_PerspectiveCorrect));
-		OutEnvironment.SetDefine(TEXT("ONEPASS_POINTLIGHT_SHADOW"), (uint32)(ShaderMode == VertexShadowDepth_OnePassPointLight || ShaderMode == VertexShadowDepth_VSLayer));
-		OutEnvironment.SetDefine(TEXT("USING_VERTEX_SHADER_LAYER"), (uint32)(ShaderMode == VertexShadowDepth_VSLayer));
+		OutEnvironment.SetDefine(TEXT("ONEPASS_POINTLIGHT_SHADOW"), (uint32)(ShaderMode == VertexShadowDepth_OnePassPointLight));
 		OutEnvironment.SetDefine(TEXT("POSITION_ONLY"), (uint32)bUsePositionOnlyStream);
-		OutEnvironment.SetDefine(TEXT("IS_FOR_GEOMETRY_SHADER"), (uint32)bIsForGeometryShader);
-		OutEnvironment.SetDefine(TEXT("ENABLE_FALLBACK_POINTLIGHT_SHADOW_GS"), UseGPUScene(Parameters.Platform) ? 1 : 0);
 
 		bool bEnableNonNaniteVSM = (ShaderMode == VertexShadowDepth_VirtualShadowMap);
 		OutEnvironment.SetDefine(TEXT("ENABLE_NON_NANITE_VSM"), bEnableNonNaniteVSM ? 1 : 0);
@@ -373,11 +357,7 @@ public:
 			FVirtualShadowMapArray::SetShaderDefines(OutEnvironment);
 		}
 
-		if (bIsForGeometryShader)
-		{
-			OutEnvironment.CompilerFlags.Add(CFLAG_VertexToGeometryShader);
-		}
-		else if (ShaderMode == VertexShadowDepth_VSLayer)
+		if (ShaderMode == VertexShadowDepth_OnePassPointLight)
 		{
 			OutEnvironment.CompilerFlags.Add(CFLAG_VertexUseAutoCulling);
 		}
@@ -389,62 +369,6 @@ public:
 	{}
 };
 
-/** Geometry shader that allows one pass point light shadows by cloning triangles to all faces of the cube map. */
-class FOnePassPointShadowDepthGS : public FMeshMaterialShader
-{
-	DECLARE_SHADER_TYPE(FOnePassPointShadowDepthGS, MeshMaterial);
-public:
-
-	void GetShaderBindings(
-		const FScene* Scene,
-		ERHIFeatureLevel::Type FeatureLevel,
-		const FPrimitiveSceneProxy* PrimitiveSceneProxy,
-		const FMaterialRenderProxy& MaterialRenderProxy,
-		const FMaterial& Material,
-		const FMeshPassProcessorRenderState& DrawRenderState,
-		const FShadowDepthShaderElementData& ShaderElementData,
-		FMeshDrawSingleShaderBindings& ShaderBindings) const
-	{
-		FMeshMaterialShader::GetShaderBindings(Scene, FeatureLevel, PrimitiveSceneProxy, MaterialRenderProxy, Material, DrawRenderState, ShaderElementData, ShaderBindings);
-		ShaderBindings.Add(bUseGpuSceneInstancing, ShaderElementData.bUseGpuSceneInstancing);
-	}
-
-	static bool ShouldCompilePermutation(const FMeshMaterialShaderPermutationParameters& Parameters)
-	{
-		return RHISupportsGeometryShaders(Parameters.Platform) && TShadowDepthVS<VertexShadowDepth_OnePassPointLight, false, true>::ShouldCompilePermutation(Parameters);
-	}
-
-	static void ModifyCompilationEnvironment(const FMaterialShaderPermutationParameters& Parameters, FShaderCompilerEnvironment& OutEnvironment)
-	{
-		FMeshMaterialShader::ModifyCompilationEnvironment(Parameters, OutEnvironment);
-		OutEnvironment.SetDefine(TEXT("ONEPASS_POINTLIGHT_SHADOW"), 1);
-		TShadowDepthVS<VertexShadowDepth_OnePassPointLight, false, true>::ModifyCompilationEnvironment(Parameters, OutEnvironment);
-		OutEnvironment.SetDefine(TEXT("ENABLE_FALLBACK_POINTLIGHT_SHADOW_GS"), UseGPUScene(Parameters.Platform) ? 1 : 0);
-	}
-
-	FOnePassPointShadowDepthGS(const ShaderMetaType::CompiledShaderInitializerType& Initializer)
-		: FMeshMaterialShader(Initializer)
-	{
-		const ERHIFeatureLevel::Type FeatureLevel = GetMaxSupportedFeatureLevel((EShaderPlatform)Initializer.Target.Platform);
-
-		if (FSceneInterface::GetShadingPath(FeatureLevel) == EShadingPath::Deferred)
-		{
-			PassUniformBuffer.Bind(Initializer.ParameterMap, FShadowDepthPassUniformParameters::StaticStructMetadata.GetShaderVariableName());
-		}
-
-		if (FSceneInterface::GetShadingPath(FeatureLevel) == EShadingPath::Mobile)
-		{
-			PassUniformBuffer.Bind(Initializer.ParameterMap, FMobileShadowDepthPassUniformParameters::StaticStructMetadata.GetShaderVariableName());
-		}
-		bUseGpuSceneInstancing.Bind(Initializer.ParameterMap, TEXT("bUseGpuSceneInstancing"));
-	}
-
-	FOnePassPointShadowDepthGS() {}
-	LAYOUT_FIELD(FShaderParameter, bUseGpuSceneInstancing);
-};
-
-IMPLEMENT_SHADER_TYPE(, FOnePassPointShadowDepthGS, TEXT("/Engine/Private/ShadowDepthVertexShader.usf"), TEXT("MainOnePassPointLightGS"), SF_Geometry);
-
 #define IMPLEMENT_SHADOW_DEPTH_SHADERMODE_SHADERS(ShaderMode) \
 	typedef TShadowDepthVS<ShaderMode, false> TShadowDepthVS##ShaderMode; \
 	IMPLEMENT_MATERIAL_SHADER_TYPE(template<>, TShadowDepthVS##ShaderMode, TEXT("/Engine/Private/ShadowDepthVertexShader.usf"), TEXT("Main"),       SF_Vertex);
@@ -464,24 +388,6 @@ IMPLEMENT_MATERIAL_SHADER_TYPE(template<>, TShadowDepthVSVertexShadowDepth_Outpu
 IMPLEMENT_MATERIAL_SHADER_TYPE(template<>, TShadowDepthVSVertexShadowDepth_OnePassPointLightPositionOnly,  TEXT("/Engine/Private/ShadowDepthVertexShader.usf"), TEXT("PositionOnlyMain"), SF_Vertex);
 IMPLEMENT_MATERIAL_SHADER_TYPE(template<>, TShadowDepthVSVertexShadowDepth_VirtualShadowMapPositionOnly,   TEXT("/Engine/Private/ShadowDepthVertexShader.usf"), TEXT("PositionOnlyMain"), SF_Vertex);
 
-// One pass point light VS for GS shaders.
-typedef TShadowDepthVS<VertexShadowDepth_OnePassPointLight, false, true> TShadowDepthVSForGSVertexShadowDepth_OnePassPointLight;
-typedef TShadowDepthVS<VertexShadowDepth_OnePassPointLight, true,  true> TShadowDepthVSForGSVertexShadowDepth_OnePassPointLightPositionOnly;
-IMPLEMENT_MATERIAL_SHADER_TYPE(template<>, TShadowDepthVSForGSVertexShadowDepth_OnePassPointLight,             TEXT("/Engine/Private/ShadowDepthVertexShader.usf"), TEXT("MainForGS"),             SF_Vertex);
-IMPLEMENT_MATERIAL_SHADER_TYPE(template<>, TShadowDepthVSForGSVertexShadowDepth_OnePassPointLightPositionOnly, TEXT("/Engine/Private/ShadowDepthVertexShader.usf"), TEXT("PositionOnlyMainForGS"), SF_Vertex);
-
-// One pass point light with vertex shader layer output.
-//                                                       bUsePositionOnlyStream
-//                                                            | bIsForGeometryShader
-//                                                            |      |
-typedef TShadowDepthVS<VertexShadowDepth_VSLayer, false, false> TShadowDepthVSVertexShadowDepth_VSLayer;
-typedef TShadowDepthVS<VertexShadowDepth_VSLayer, true,  false> TShadowDepthVSVertexShadowDepth_VSLayerPositionOnly;
-typedef TShadowDepthVS<VertexShadowDepth_VSLayer, false, true>  TShadowDepthVSVertexShadowDepth_VSLayerGS; // not used
-typedef TShadowDepthVS<VertexShadowDepth_VSLayer, true,  true>  TShadowDepthVSVertexShadowDepth_VSLayerGSPositionOnly; // not used
-IMPLEMENT_MATERIAL_SHADER_TYPE(template<>, TShadowDepthVSVertexShadowDepth_VSLayer,               TEXT("/Engine/Private/ShadowDepthVertexShader.usf"), TEXT("Main"),             SF_Vertex);
-IMPLEMENT_MATERIAL_SHADER_TYPE(template<>, TShadowDepthVSVertexShadowDepth_VSLayerPositionOnly,   TEXT("/Engine/Private/ShadowDepthVertexShader.usf"), TEXT("PositionOnlyMain"), SF_Vertex);
-IMPLEMENT_MATERIAL_SHADER_TYPE(template<>, TShadowDepthVSVertexShadowDepth_VSLayerGS,             TEXT("/Engine/Private/ShadowDepthVertexShader.usf"), TEXT("Main"),             SF_Vertex); // not used
-IMPLEMENT_MATERIAL_SHADER_TYPE(template<>, TShadowDepthVSVertexShadowDepth_VSLayerGSPositionOnly, TEXT("/Engine/Private/ShadowDepthVertexShader.usf"), TEXT("PositionOnlyMain"), SF_Vertex); // not used
 
 /**
 * A pixel shader for rendering the depth of a mesh.
@@ -624,8 +530,7 @@ bool GetShadowDepthPassShaders(
 	bool bUsePerspectiveCorrectShadowDepths,
 	bool bVirtualShadowMap,
 	TShaderRef<FShadowDepthVS>& VertexShader,
-	TShaderRef<FShadowDepthBasePS>& PixelShader,
-	TShaderRef<FOnePassPointShadowDepthGS>& GeometryShader)
+	TShaderRef<FShadowDepthBasePS>& PixelShader)
 {
 	const FVertexFactoryType* VFType = VertexFactory->GetType();
 
@@ -634,33 +539,20 @@ bool GetShadowDepthPassShaders(
 	// Vertex related shaders
 	if (bOnePassPointLightShadow)
 	{
-		if (!RHISupportsVertexShaderLayer(GShaderPlatformForFeatureLevel[FeatureLevel]))
+		if (DoesRuntimeSupportOnePassPointLightShadows(GShaderPlatformForFeatureLevel[FeatureLevel]))
 		{
 			if (bPositionOnlyVS)
 			{
-				ShaderTypes.AddShaderType<TShadowDepthVS<VertexShadowDepth_OnePassPointLight, true, true>>();
+				ShaderTypes.AddShaderType<TShadowDepthVS<VertexShadowDepth_OnePassPointLight, true>>();
 			}
 			else
 			{
-				ShaderTypes.AddShaderType<TShadowDepthVS<VertexShadowDepth_OnePassPointLight, false, true>>();
-			}
-
-			if (RHISupportsGeometryShaders(GShaderPlatformForFeatureLevel[FeatureLevel]))
-			{
-				// Use the geometry shader which will clone output triangles to all faces of the cube map
-				ShaderTypes.AddShaderType<FOnePassPointShadowDepthGS>();
+				ShaderTypes.AddShaderType<TShadowDepthVS<VertexShadowDepth_OnePassPointLight, false>>();
 			}
 		}
 		else
 		{
-			if (bPositionOnlyVS)
-			{
-				ShaderTypes.AddShaderType<TShadowDepthVS<VertexShadowDepth_VSLayer, true, false>>();
-			}
-			else
-			{
-				ShaderTypes.AddShaderType<TShadowDepthVS<VertexShadowDepth_VSLayer, false, false>>();
-			}
+			return false;
 		}
 	}
 	else if (bVirtualShadowMap)
@@ -725,7 +617,6 @@ bool GetShadowDepthPassShaders(
 		return false;
 	}
 
-	Shaders.TryGetGeometryShader(GeometryShader);
 	Shaders.TryGetVertexShader(VertexShader);
 	Shaders.TryGetPixelShader(PixelShader);
 	return true;
@@ -2180,8 +2071,7 @@ bool FShadowDepthPassMeshProcessor::Process(
 
 	TMeshProcessorShaders<
 		FShadowDepthVS,
-		FShadowDepthBasePS,
-		FOnePassPointShadowDepthGS> ShadowDepthPassShaders;
+		FShadowDepthBasePS> ShadowDepthPassShaders;
 
 	const bool bUsePositionOnlyVS =
 		   VertexFactory->SupportsPositionAndNormalOnlyStream()
@@ -2212,8 +2102,7 @@ bool FShadowDepthPassMeshProcessor::Process(
 		bUsePerspectiveCorrectShadowDepths,
 		bVirtualShadowMap,
 		ShadowDepthPassShaders.VertexShader,
-		ShadowDepthPassShaders.PixelShader,
-		ShadowDepthPassShaders.GeometryShader))
+		ShadowDepthPassShaders.PixelShader))
 	{
 		return false;
 	}
@@ -2223,12 +2112,11 @@ bool FShadowDepthPassMeshProcessor::Process(
 
 	const FMeshDrawCommandSortKey SortKey = CalculateMeshStaticSortKey(ShadowDepthPassShaders.VertexShader, ShadowDepthPassShaders.PixelShader);
 
-	const bool bUseGeometryShader = !RHISupportsVertexShaderLayer(GShaderPlatformForFeatureLevel[FeatureLevel]) && RHISupportsGeometryShaders(GShaderPlatformForFeatureLevel[FeatureLevel]);
-
-	const bool bUseGpuSceneInstancing = UseGPUScene(GShaderPlatformForFeatureLevel[FeatureLevel], FeatureLevel) 
-		&& VertexFactory->GetPrimitiveIdStreamIndex(FeatureLevel, bUsePositionOnlyVS ? EVertexInputStreamType::PositionAndNormalOnly : EVertexInputStreamType::Default) != INDEX_NONE;
-
-	const uint32 InstanceFactor = bUseGpuSceneInstancing || !ShadowDepthType.bOnePassPointLightShadow || bUseGeometryShader ? 1 : 6;
+	const bool bUseGpuSceneInstancing = UseGPUScene(GShaderPlatformForFeatureLevel[FeatureLevel], FeatureLevel) && VertexFactory->SupportsGPUScene(FeatureLevel);
+	
+	// Need to replicate for cube faces on host if GPU-scene is not available (for this draw).
+	const bool bPerformHostCubeFaceReplication = ShadowDepthType.bOnePassPointLightShadow && !bUseGpuSceneInstancing;
+	const uint32 InstanceFactor = bPerformHostCubeFaceReplication ? 6 : 1;
 
 	for (uint32 i = 0; i < InstanceFactor; i++)
 	{
diff --git a/Engine/Source/Runtime/Renderer/Private/ShadowSetup.cpp b/Engine/Source/Runtime/Renderer/Private/ShadowSetup.cpp
--- a/Engine/Source/Runtime/Renderer/Private/ShadowSetup.cpp
+++ b/Engine/Source/Runtime/Renderer/Private/ShadowSetup.cpp
@@ -2124,7 +2124,7 @@ void FProjectedShadowInfo::SetupMeshDrawCommandsForShadowDepth(FSceneRenderer& R
 	}
 	
 	ViewIds.Reset();
-	if (GetShadowDepthType().bOnePassPointLightShadow)
+	if (bOnePassPointLightShadow)
 	{
 		ViewIds.AddDefaulted(6);
 		for (int32 CubemapFaceIndex = 0; CubemapFaceIndex < 6; CubemapFaceIndex++)
@@ -2169,9 +2169,9 @@ void FProjectedShadowInfo::SetupMeshDrawCommandsForShadowDepth(FSceneRenderer& R
 	// GPUCULL_TODO: Pass along any custom culling planes or whatever here (e.g., cacade bounds):
 	// GPUCULL_TODO: Add debug tags to context and views (so compute passes can be understood)
 	// GPUCULL_TODO: Needed to support legacy, non-GPU-Scene culled, primitives, this is merely used to allocate enough space for CPU-side replication.
-	const bool bUseGeometryShader = !RHISupportsVertexShaderLayer(Renderer.Scene->GetShaderPlatform()) && RHISupportsGeometryShaders(Renderer.Scene->GetShaderPlatform());
+	const bool bMayUseHostCubeFaceReplication = bOnePassPointLightShadow && !HasVirtualShadowMap();
 	// Note: Iteracts with FShadowDepthPassMeshProcessor::Process and must be an overestimate of the actual replication done there.
-	const uint32 InstanceFactor = !GetShadowDepthType().bOnePassPointLightShadow || bUseGeometryShader ? 1 : 6;
+	const uint32 InstanceFactor = bMayUseHostCubeFaceReplication ? 6 : 1;
 
 	// Ensure all work goes down the one path to simplify processing
 	EBatchProcessingMode SingleInstanceProcessingMode = (HasVirtualShadowMap() || VirtualShadowMapClipmap.IsValid()) ? EBatchProcessingMode::Generic : EBatchProcessingMode::UnCulled;
@@ -4921,6 +4921,20 @@ void FSceneRenderer::AddViewDependentWholeSceneShadowsForView(
 
 void FSceneRenderer::AllocateShadowDepthTargets(FRHICommandListImmediate& RHICmdList)
 {
+#if !(UE_BUILD_SHIPPING || UE_BUILD_TEST)
+	TBitArray<SceneRenderingAllocator> OnePassShadowUnsupportedLights;
+	auto MarkUnsupportedOnePassShadow = [&OnePassShadowUnsupportedLights](int32 LightId)
+	{
+		if (OnePassShadowUnsupportedLights.Num() <= LightId)
+		{
+			OnePassShadowUnsupportedLights.Add(false, 1 + LightId - OnePassShadowUnsupportedLights.Num());
+		}
+		OnePassShadowUnsupportedLights[LightId] = true;
+	};
+#else 
+	auto MarkUnsupportedOnePassShadow = [](int32 LightId) {};
+#endif
+
 	// Sort visible shadows based on their allocation needs
 	// 2d shadowmaps for this frame only that can be atlased across lights
 	TArray<FProjectedShadowInfo*, SceneRenderingAllocator> Shadows;
@@ -5003,6 +5017,18 @@ void FSceneRenderer::AllocateShadowDepthTargets(FRHICommandListImmediate& RHICmd
 				bShadowIsVisible = false;
 			}
 
+			// Skip one-pass point light shadows if rendering to them is not suported, unless VSM or DF
+			if (bShadowIsVisible
+				&& !ProjectedShadowInfo->HasVirtualShadowMap()
+				&& !ProjectedShadowInfo->bRayTracedDistanceField
+				&& ProjectedShadowInfo->bOnePassPointLightShadow 
+				&& !DoesRuntimeSupportOnePassPointLightShadows(GShaderPlatformForFeatureLevel[FeatureLevel]))
+			{ 
+				// Also note this to produce a log message & on-screen warning
+				MarkUnsupportedOnePassShadow(ProjectedShadowInfo->GetLightSceneInfo().Id);
+				bShadowIsVisible = false;
+			}
+
 			if (bShadowIsVisible)
 			{
 				// Visible shadow stats
@@ -5207,6 +5233,24 @@ void FSceneRenderer::AllocateShadowDepthTargets(FRHICommandListImmediate& RHICmd
 		}
 	}
 
+#if !(UE_BUILD_SHIPPING || UE_BUILD_TEST)
+	if (!OnePassShadowUnsupportedLights.IsEmpty())
+	{
+		OnGetOnScreenMessages.AddLambda([Scene = Scene, OnePassShadowUnsupportedLights = MoveTemp(OnePassShadowUnsupportedLights)](FScreenMessageWriter& ScreenMessageWriter)->void
+		{
+			static const FText Message = NSLOCTEXT("Renderer", "PointLightVsLayer", "RUNTIME DOES NOT SUPPORT WHOLE SCENE POINT LIGHT SHADOWS (Missing Vertexshader Layer support): ");
+			ScreenMessageWriter.DrawLine(Message);
+			for (TConstSetBitIterator<SceneRenderingAllocator> It(OnePassShadowUnsupportedLights); It; ++It)
+			{
+				int32 LightId = It.GetIndex();
+				if (Scene->Lights.IsValidIndex(LightId))
+				{
+					ScreenMessageWriter.DrawLine(FText::FromString(Scene->Lights[LightId].LightSceneInfo->Proxy->GetOwnerNameOrLabel()), 35);
+				}
+			}
+		});
+	}
+#endif
 	SET_MEMORY_STAT(STAT_CachedShadowmapMemory, Scene->GetCachedWholeSceneShadowMapsSize());
 	SET_MEMORY_STAT(STAT_ShadowmapAtlasMemory, SortedShadowsForShadowDepthPass.ComputeMemorySize());
 }
