Fixes saparate translucency screen percentage

diff --git a/Engine/Shaders/Private/ComposeSeparateTranslucency.usf b/Engine/Shaders/Private/ComposeSeparateTranslucency.usf
--- a/Engine/Shaders/Private/ComposeSeparateTranslucency.usf
+++ b/Engine/Shaders/Private/ComposeSeparateTranslucency.usf
@@ -3,68 +3,83 @@
 #include "Common.ush"
 #include "ScreenPass.ush"
 
-Texture2D SceneColor;
+#if PERMUTATION_NEARESTDEPTHNEIGHBOR 
+	#include "SeparateTranslucency.ush"
+#endif // PERMUTATION_NEARESTDEPTHNEIGHBOR
+
+
+//------------------------------------------------------- PARAMETERS
+
+Texture2D    SceneColorTexture;
 SamplerState SceneColorSampler;
 
-Texture2D SeparateTranslucency;
-SamplerState SeparateTranslucencySampler;
+Texture2D<float> LowResDepthTexture;
+Texture2D<float> FullResDepthTexture;
 
-Texture2D SeparateModulation;
-SamplerState SeparateModulationSampler;
+Texture2D    SeparateTranslucencyPointTexture;
+SamplerState SeparateTranslucencyPointSampler;
+Texture2D    SeparateModulationPointTexture;
+SamplerState SeparateModulationPointSampler;
 
-SCREEN_PASS_TEXTURE_VIEWPORT(Color)
-SCREEN_PASS_TEXTURE_VIEWPORT(Translucency)
-FScreenTransform ColorToTranslucency;
+Texture2D    SeparateTranslucencyBilinearTexture;
+SamplerState SeparateTranslucencyBilinearSampler;
+Texture2D    SeparateModulationBilinearTexture;
+SamplerState SeparateModulationBilinearSampler;
 
-float2 ScreenPosToSceneColorUV;
-float2 ScreenPosToSeparateTranslucencyUV;
+FScreenTransform ScreenPosToSceneColorUV;
+FScreenTransform ScreenPosToSeparateTranslucencyUV;
+float2 SeparateTranslucencyUVMin;
 float2 SeparateTranslucencyUVMax;
 float2 SeparateTranslucencyExtentInverse;
 
-#if PERMUTATION_NEARESTDEPTHNEIGHBOR 
-
-#include "SeparateTranslucency.ush"
 
-#endif // PERMUTATION_NEARESTDEPTHNEIGHBOR
+//------------------------------------------------------- ENTRY POINT
 
 void MainPS(
 	float4 SvPosition : SV_POSITION,
 	out float4 OutColor : SV_Target0)
 {
-	float2 ColorUV = SvPosition.xy * Color_ExtentInverse.xy;
-	float2 SeparateTranslucencyUV = clamp(ApplyScreenTransform(ColorUV, ColorToTranslucency), Translucency_UVViewportBilinearMin, Translucency_UVViewportBilinearMax);
+	float4 Debug = 0;
+
+	float2 SceneColorUV = ApplyScreenTransform(SvPosition.xy, ScreenPosToSceneColorUV);
+	float2 SeparateTranslucencyUV = ApplyScreenTransform(SvPosition.xy, ScreenPosToSeparateTranslucencyUV);
 
-	float4 SceneColorSample = SceneColor.SampleLevel(SceneColorSampler, ColorUV, 0);
+	float4 SceneColorSample = SceneColorTexture.SampleLevel(SceneColorSampler, SceneColorUV, 0);
+	
+	SeparateTranslucencyUV = clamp(SeparateTranslucencyUV, SeparateTranslucencyUVMin, SeparateTranslucencyUVMax);
 
 #if PERMUTATION_NEARESTDEPTHNEIGHBOR 
 
-	NearestDepthNeighborUpsamplingResult UpsampleResult = NearestDepthNeighborUpsampling(SvPosition.xy, ColorUV, Translucency_ExtentInverse);
+	NearestDepthNeighborUpsamplingResult UpsampleResult = NearestDepthNeighborUpsampling(
+		LowResDepthTexture,
+		FullResDepthTexture,
+		SvPosition.xy, SeparateTranslucencyUV, SeparateTranslucencyExtentInverse);
 
 	float4 SeparateTranslucencySample = 0;
 	float4 SeparateModulationSample = 0;
+	
+	float2 SeparateTranslucencySampleUV = clamp(UpsampleResult.UV, SeparateTranslucencyUVMin, SeparateTranslucencyUVMax);
+
+
 	if (UpsampleResult.bUsePointSampler)
 	{
-		SeparateTranslucencySample	= SeparateTranslucency.SampleLevel(	GlobalPointClampedSampler, UpsampleResult.UV, 0);
-		SeparateModulationSample	= SeparateModulation.SampleLevel(	GlobalPointClampedSampler, UpsampleResult.UV, 0);
+		SeparateTranslucencySample	= SeparateTranslucencyPointTexture.SampleLevel(	SeparateTranslucencyPointSampler, SeparateTranslucencySampleUV, 0);
+		SeparateModulationSample	= SeparateModulationPointTexture.SampleLevel(	SeparateModulationPointSampler, SeparateTranslucencySampleUV, 0);
 	}
 	else
 	{
-		SeparateTranslucencySample	= SeparateTranslucency.SampleLevel(	GlobalBilinearClampedSampler, UpsampleResult.UV, 0);
-		SeparateModulationSample	= SeparateModulation.SampleLevel(	GlobalBilinearClampedSampler, UpsampleResult.UV, 0);
+		SeparateTranslucencySample	= SeparateTranslucencyBilinearTexture.SampleLevel(	SeparateTranslucencyBilinearSampler, SeparateTranslucencySampleUV, 0);
+		SeparateModulationSample	= SeparateModulationBilinearTexture.SampleLevel(	SeparateModulationBilinearSampler, SeparateTranslucencySampleUV, 0);
 	}
 
 #else // PERMUTATION_NEARESTDEPTHNEIGHBOR
 	
-	float4 SeparateTranslucencySample = SeparateTranslucency.SampleLevel(SeparateTranslucencySampler, SeparateTranslucencyUV, 0);
-	float4 SeparateModulationSample   = SeparateModulation.SampleLevel(	 SeparateModulationSampler, SeparateTranslucencyUV, 0);
+	float4 SeparateTranslucencySample = SeparateTranslucencyBilinearTexture.SampleLevel(SeparateTranslucencyBilinearSampler, SeparateTranslucencyUV, 0);
+	float4 SeparateModulationSample   = SeparateModulationBilinearTexture.SampleLevel(	SeparateModulationBilinearSampler, SeparateTranslucencyUV, 0);
 
 #endif // PERMUTATION_NEARESTDEPTHNEIGHBOR
 
 	// Final composition
 	OutColor.rgb = SceneColorSample.rgb * SeparateTranslucencySample.a * SeparateModulationSample.rgb + SeparateTranslucencySample.rgb;
 	OutColor.a   = SceneColorSample.a   * SeparateTranslucencySample.a;
-
-#if !POST_PROCESS_ALPHA
-	OutColor.a = 0;
-#endif
 }
diff --git a/Engine/Shaders/Private/DiaphragmDOF/DOFRecombine.usf b/Engine/Shaders/Private/DiaphragmDOF/DOFRecombine.usf
--- a/Engine/Shaders/Private/DiaphragmDOF/DOFRecombine.usf
+++ b/Engine/Shaders/Private/DiaphragmDOF/DOFRecombine.usf
@@ -148,6 +148,9 @@ Texture2D SceneColorInput;
 Texture2D SceneDepthTexture;
 Texture2D SceneSeparateCoc;
 
+Texture2D<float> LowResDepthTexture;
+Texture2D<float> FullResDepthTexture;
+
 Texture2D SceneSeparateTranslucency;
 Texture2D SceneSeparateTranslucencyModulateColor;
 
@@ -1095,7 +1098,12 @@ void RecombineMainCS(
 		else
 		{
 			float2 PixelPos = (View.ViewRectMin.xy + DispatchThreadId) + 0.5;
-			UpsampleResult = NearestDepthNeighborUpsampling(PixelPos, InputParameters.SceneBufferUV, SeparateTranslucencyTextureLowResExtentInverse);
+			UpsampleResult = NearestDepthNeighborUpsampling(
+				LowResDepthTexture,
+				FullResDepthTexture,
+				PixelPos,
+				InputParameters.SceneBufferUV,
+				SeparateTranslucencyTextureLowResExtentInverse);
 		}
 		
 		UpsampleResult.UV = clamp(UpsampleResult.UV, SeparateTranslucencyBilinearUVMinMax.xy, SeparateTranslucencyBilinearUVMinMax.zw);
diff --git a/Engine/Shaders/Private/SeparateTranslucency.ush b/Engine/Shaders/Private/SeparateTranslucency.ush
--- a/Engine/Shaders/Private/SeparateTranslucency.ush
+++ b/Engine/Shaders/Private/SeparateTranslucency.ush
@@ -6,9 +6,6 @@
 	SeparateTranslucency.usf: 
 =============================================================================*/
 
-Texture2D<float> LowResDepthTexture;
-Texture2D<float> FullResDepthTexture;
-
 void UpdateNearestSample(float Z, float2 UV, float FullResZ, inout float MinDist, inout float2 NearestUV)
 {
     float DepthDelta = abs(Z - FullResZ);
@@ -21,17 +18,16 @@ void UpdateNearestSample(float Z, float2 UV, float FullResZ, inout float MinDist
     }
 }
 
-float4 BilinearUpsampling(float2 UV, Texture2D LowResTex)
-{
-	return Texture2DSample(LowResTex, GlobalBilinearClampedSampler, UV);
-}
-
 struct NearestDepthNeighborUpsamplingResult
 {
 	bool bUsePointSampler;
 	float2 UV;
 };
-NearestDepthNeighborUpsamplingResult NearestDepthNeighborUpsampling(float2 Position, float2 UV, float2 LowResTexelSize)
+NearestDepthNeighborUpsamplingResult NearestDepthNeighborUpsampling(
+	Texture2D<float> LowResDepthTexture,
+	Texture2D<float> FullResDepthTexture,
+	float2 Position,
+	float2 UV, float2 LowResTexelSize)
 {
 	NearestDepthNeighborUpsamplingResult UpsampleResult;
 
diff --git a/Engine/Shaders/Private/TranslucencyUpsampling.usf b/Engine/Shaders/Private/TranslucencyUpsampling.usf
--- a/Engine/Shaders/Private/TranslucencyUpsampling.usf
+++ b/Engine/Shaders/Private/TranslucencyUpsampling.usf
@@ -5,26 +5,29 @@
 =============================================================================*/
 
 #include "Common.ush"
-#include "SeparateTranslucency.ush"
+#include "ScreenPass.ush"
 
-Texture2D LowResColorTexture;
-float2 LowResExtentInverse;
 
-void SimpleUpsamplingPS(FScreenVertexOutput Input, out float4 OutColor : SV_Target0)
-{
-	OutColor = BilinearUpsampling(Input.UV, LowResColorTexture);
-}
+uint2 StencilPixelPosMin;
+uint2 StencilPixelPosMax;
+FScreenTransform SvPositionToStencilPixelCoord;
+uint StencilMask;
+Texture2D<uint2> StencilTexture;
 
-void NearestDepthNeighborUpsamplingPS(FScreenVertexOutput Input, out float4 OutColor : SV_Target0)
+void UpsampleResponsiveAAPS(FScreenVertexOutput Input, out float4 OutColor : SV_Target0)
 {
-	NearestDepthNeighborUpsamplingResult UpsampleResult = NearestDepthNeighborUpsampling(Input.Position.xy, Input.UV, LowResExtentInverse);
+	float2 SvPosition = Input.Position.xy;
 
-	if (UpsampleResult.bUsePointSampler)
-	{
-		OutColor = Texture2DSampleLevel(LowResColorTexture, GlobalPointClampedSampler, UpsampleResult.UV, 0);
-	}
-	else
+	uint2 StencilPixelPos = uint2(ApplyScreenTransform(SvPosition, SvPositionToStencilPixelCoord));
+	StencilPixelPos = clamp(StencilPixelPos, StencilPixelPosMin, StencilPixelPosMax);
+
+	uint StencilValue = StencilTexture.Load(int3(StencilPixelPos, 0)) STENCIL_COMPONENT_SWIZZLE;
+
+	// If the responsive AA stencil was not set, discard the pixel to not write the responsive AA in the full res depth buffer.
+	if ((StencilValue & StencilMask) == 0)
 	{
-		OutColor = Texture2DSampleLevel(LowResColorTexture, GlobalBilinearClampedSampler, UpsampleResult.UV, 0);
+		discard;
 	}
+
+	OutColor = float(0.0).xxxx;
 }
diff --git a/Engine/Source/Runtime/Renderer/Private/DeferredShadingRenderer.cpp b/Engine/Source/Runtime/Renderer/Private/DeferredShadingRenderer.cpp
--- a/Engine/Source/Runtime/Renderer/Private/DeferredShadingRenderer.cpp
+++ b/Engine/Source/Runtime/Renderer/Private/DeferredShadingRenderer.cpp
@@ -3058,7 +3058,7 @@ void FDeferredShadingSceneRenderer::Render(FRDGBuilder& GraphBuilder)
 		RenderHairComposition(GraphBuilder, Views, SceneTextures.Color.Target, SceneTextures.Depth.Target);
 	}
 
-	FSeparateTranslucencyTextures SeparateTranslucencyTextures(SeparateTranslucencyDimensions);
+	FTranslucencyPassResourcesMap TranslucencyResourceMap(Views.Num());
 
 	// Draw translucency.
 	if (!bHasRayTracedOverlay && TranslucencyViewsToRender != ETranslucencyView::None)
@@ -3066,6 +3066,8 @@ void FDeferredShadingSceneRenderer::Render(FRDGBuilder& GraphBuilder)
 		RDG_CSV_STAT_EXCLUSIVE_SCOPE(GraphBuilder, RenderTranslucency);
 		SCOPE_CYCLE_COUNTER(STAT_TranslucencyDrawTime);
 
+		RDG_EVENT_SCOPE(GraphBuilder, "Translucency");
+
 		// Raytracing doesn't need the distortion effect.
 		const bool bShouldRenderDistortion = TranslucencyViewsToRender != ETranslucencyView::RayTracing;
 
@@ -3087,7 +3089,7 @@ void FDeferredShadingSceneRenderer::Render(FRDGBuilder& GraphBuilder)
 
 		// Render all remaining translucency views.
 		GraphBuilder.SetCommandListStat(GET_STATID(STAT_CLM_Translucency));
-		RenderTranslucency(GraphBuilder, SceneTextures, TranslucencyLightingVolumeTextures, &SeparateTranslucencyTextures, TranslucencyViewsToRender, InstanceCullingManager);
+		RenderTranslucency(GraphBuilder, SceneTextures, TranslucencyLightingVolumeTextures, &TranslucencyResourceMap, TranslucencyViewsToRender, InstanceCullingManager);
 		TranslucencyViewsToRender = ETranslucencyView::None;
 
 		// Compose hair before velocity/distortion pass since these pass write depth value, 
@@ -3157,7 +3159,7 @@ void FDeferredShadingSceneRenderer::Render(FRDGBuilder& GraphBuilder)
 	{
 		SCOPE_CYCLE_COUNTER(STAT_FDeferredShadingSceneRenderer_RenderLightShaftBloom);
 		GraphBuilder.SetCommandListStat(GET_STATID(STAT_CLM_LightShaftBloom));
-		RenderLightShaftBloom(GraphBuilder, SceneTextures, SeparateTranslucencyTextures);
+		RenderLightShaftBloom(GraphBuilder, SceneTextures, /* inout */ TranslucencyResourceMap);
 	}
 
 	if (bUseVirtualTexturing)
@@ -3271,7 +3273,6 @@ void FDeferredShadingSceneRenderer::Render(FRDGBuilder& GraphBuilder)
 		FPostProcessingInputs PostProcessingInputs;
 		PostProcessingInputs.ViewFamilyTexture = ViewFamilyTexture;
 		PostProcessingInputs.CustomDepthTexture = SceneTextures.CustomDepth.Depth;
-		PostProcessingInputs.SeparateTranslucencyTextures = &SeparateTranslucencyTextures;
 		PostProcessingInputs.SceneTextures = SceneTextures.UniformBuffer;
 
 		if (ViewFamily.UseDebugViewPS())
@@ -3282,6 +3283,7 @@ void FDeferredShadingSceneRenderer::Render(FRDGBuilder& GraphBuilder)
    				const Nanite::FRasterResults* NaniteResults = bNaniteEnabled ? &NaniteRasterResults[ViewIndex] : nullptr;
 				RDG_GPU_MASK_SCOPE(GraphBuilder, View.GPUMask);
 				RDG_EVENT_SCOPE_CONDITIONAL(GraphBuilder, Views.Num() > 1, "View%d", ViewIndex);
+				PostProcessingInputs.TranslucencyViewResourcesMap = FTranslucencyViewResourcesMap(TranslucencyResourceMap, ViewIndex);
 				AddDebugViewPostProcessingPasses(GraphBuilder, View, PostProcessingInputs, NaniteResults);
 			}
 		}
@@ -3293,6 +3295,7 @@ void FDeferredShadingSceneRenderer::Render(FRDGBuilder& GraphBuilder)
 				{
 					FViewInfo& View = Views[ViewIndex];
 					RDG_GPU_MASK_SCOPE(GraphBuilder, View.GPUMask);
+					PostProcessingInputs.TranslucencyViewResourcesMap = FTranslucencyViewResourcesMap(TranslucencyResourceMap, ViewIndex);
 					ViewFamily.ViewExtensions[ViewExt]->PrePostProcessPass_RenderThread(GraphBuilder, View, PostProcessingInputs);
 				}
 			}
@@ -3303,6 +3306,8 @@ void FDeferredShadingSceneRenderer::Render(FRDGBuilder& GraphBuilder)
 				RDG_GPU_MASK_SCOPE(GraphBuilder, View.GPUMask);
 				RDG_EVENT_SCOPE_CONDITIONAL(GraphBuilder, Views.Num() > 1, "View%d", ViewIndex);
 
+				PostProcessingInputs.TranslucencyViewResourcesMap = FTranslucencyViewResourcesMap(TranslucencyResourceMap, ViewIndex);
+
 #if !(UE_BUILD_SHIPPING)
 				if (IsPostProcessVisualizeCalibrationMaterialEnabled(View))
 				{
diff --git a/Engine/Source/Runtime/Renderer/Private/DeferredShadingRenderer.h b/Engine/Source/Runtime/Renderer/Private/DeferredShadingRenderer.h
--- a/Engine/Source/Runtime/Renderer/Private/DeferredShadingRenderer.h
+++ b/Engine/Source/Runtime/Renderer/Private/DeferredShadingRenderer.h
@@ -588,7 +588,7 @@ private:
 		FRDGBuilder& GraphBuilder,
 		const FMinimalSceneTextures& SceneTextures,
 		const FTranslucencyLightingVolumeTextures& TranslucencyLightingVolumeTextures,
-		FSeparateTranslucencyTextures* OutSeparateTranslucencyTextures,
+		FTranslucencyPassResourcesMap* OutTranslucencyResourceMap,
 		ETranslucencyView ViewsToRender,
 		FInstanceCullingManager& InstanceCullingManager);
 
@@ -597,7 +597,8 @@ private:
 		FRDGBuilder& GraphBuilder,
 		const FMinimalSceneTextures& SceneTextures,
 		const FTranslucencyLightingVolumeTextures& TranslucencyLightingVolumeTextures,
-		FSeparateTranslucencyTextures* OutSeparateTranslucencyTextures,
+		FTranslucencyPassResourcesMap* OutTranslucencyResourceMap,
+		FRDGTextureMSAA SharedDepthTexture,
 		ETranslucencyView ViewsToRender,
 		FRDGTextureRef SceneColorCopyTexture,
 		ETranslucencyPass::Type TranslucencyPass,
@@ -611,7 +612,7 @@ private:
 	void RenderLightShaftBloom(
 		FRDGBuilder& GraphBuilder,
 		const FMinimalSceneTextures& SceneTextures,
-		FSeparateTranslucencyTextures& OutSeparateTranslucencyTextures);
+		FTranslucencyPassResourcesMap& OutTranslucencyResourceMap);
 
 	bool ShouldRenderDistortion() const;
 	void RenderDistortion(FRDGBuilder& GraphBuilder, FRDGTextureRef SceneColorTexture, FRDGTextureRef SceneDepthTexture);
diff --git a/Engine/Source/Runtime/Renderer/Private/LightShaftRendering.cpp b/Engine/Source/Runtime/Renderer/Private/LightShaftRendering.cpp
--- a/Engine/Source/Runtime/Renderer/Private/LightShaftRendering.cpp
+++ b/Engine/Source/Runtime/Renderer/Private/LightShaftRendering.cpp
@@ -529,7 +529,7 @@ void AddLightShaftBloomPass(
 void FDeferredShadingSceneRenderer::RenderLightShaftBloom(
 	FRDGBuilder& GraphBuilder,
 	const FMinimalSceneTextures& SceneTextures,
-	FSeparateTranslucencyTextures& SeparateTranslucencyTextures)
+	FTranslucencyPassResourcesMap& OutTranslucencyResourceMap)
 {
 	if (ShouldRenderLightShafts(ViewFamily))
 	{
@@ -564,13 +564,32 @@ void FDeferredShadingSceneRenderer::RenderLightShaftBloom(
 
 					// Reset to point at separate translucency if enabled.
 					const ELightShaftBloomOutput BloomOutput = GetLightShaftBloomOutput(ViewFamily);
+					bool bUpdateViewsTranslucency = false;
 					if (BloomOutput == ELightShaftBloomOutput::SeparateTranslucency)
 					{
-						if (!SeparateTranslucencyTextures.IsColorValid())
+						FTranslucencyPassResources& TranslucencyPassResources = OutTranslucencyResourceMap.Get(/* ViewIndex = */ 0, ETranslucencyPass::TPT_TranslucencyAfterDOF);
+
+						if (!TranslucencyPassResources.IsValid())
 						{
 							OutputLoadAction = ERenderTargetLoadAction::EClear;
+
+							const FRDGTextureDesc Desc = FRDGTextureDesc::Create2D(
+								SeparateTranslucencyDimensions.Extent,
+								PF_FloatRGBA,
+								FClearValueBinding::Black,
+								TexCreate_RenderTargetable | TexCreate_ShaderResource,
+								1,
+								SeparateTranslucencyDimensions.NumSamples);
+
+							OutputTexture = CreateTextureMSAA(
+								GraphBuilder, Desc,
+								TEXT("Translucency.LightShaftBloom"),
+								GFastVRamConfig.SeparateTranslucency);
+						}
+						else
+						{
+							OutputTexture = TranslucencyPassResources.ColorTexture;
 						}
-						OutputTexture = SeparateTranslucencyTextures.GetColorForWrite(GraphBuilder);
 						OutputViewRectScale = SeparateTranslucencyDimensions.Scale;
 					}
 
@@ -620,6 +639,18 @@ void FDeferredShadingSceneRenderer::RenderLightShaftBloom(
 
 							AddLightShaftBloomPass(GraphBuilder, View, LightShaftParameters, LightShafts, OutputViewport, OutputBinding);
 							OutputLoadAction = ERenderTargetLoadAction::ELoad;
+
+							if (bUpdateViewsTranslucency)
+							{
+								FTranslucencyPassResources& TranslucencyPassResources = OutTranslucencyResourceMap.Get(ViewIndex, ETranslucencyPass::TPT_TranslucencyAfterDOF);
+								TranslucencyPassResources.ViewRect = OutputViewport.Rect;
+								TranslucencyPassResources.ColorTexture = OutputTexture;
+							}
+							else
+							{
+								FTranslucencyPassResources& TranslucencyPassResources = OutTranslucencyResourceMap.Get(ViewIndex, ETranslucencyPass::TPT_TranslucencyAfterDOF);
+								ensure(TranslucencyPassResources.ViewRect == OutputViewport.Rect);
+							}
 						}
 					}
 				}
diff --git a/Engine/Source/Runtime/Renderer/Private/PostProcess/DiaphragmDOF.cpp b/Engine/Source/Runtime/Renderer/Private/PostProcess/DiaphragmDOF.cpp
--- a/Engine/Source/Runtime/Renderer/Private/PostProcess/DiaphragmDOF.cpp
+++ b/Engine/Source/Runtime/Renderer/Private/PostProcess/DiaphragmDOF.cpp
@@ -1321,7 +1321,7 @@ FRDGTextureRef DiaphragmDOF::AddPasses(
 	const FSceneTextureParameters& SceneTextures,
 	const FViewInfo& View,
 	FRDGTextureRef InputSceneColor,
-	const FSeparateTranslucencyTextures& SeparateTranslucencyTextures)
+	const FTranslucencyPassResources& TranslucencyPassResources)
 {
 	if (View.Family->EngineShowFlags.VisualizeDOF)
 	{
@@ -2592,9 +2592,17 @@ FRDGTextureRef DiaphragmDOF::AddPasses(
 			NewSceneColor = GraphBuilder.CreateTexture(Desc, TEXT("DOF.Recombine"));
 		}
 
-		const FSeparateTranslucencyDimensions SeparateTranslucencyDimensions = SeparateTranslucencyTextures.GetDimensions();
-		const FIntRect SeparateTranslucencyRect = SeparateTranslucencyDimensions.GetViewport(View.ViewRect).Rect;
-		const bool bScaleSeparateTranslucency = SeparateTranslucencyDimensions.Scale != 1.0f;
+		bool bScaleSeparateTranslucency = false;
+		FScreenPassTextureViewport SeparateTranslucencyViewport;
+		if (TranslucencyPassResources.IsValid())
+		{
+			SeparateTranslucencyViewport = TranslucencyPassResources.GetTextureViewport();
+			bScaleSeparateTranslucency = SeparateTranslucencyViewport.Rect.Size() != FullResViewSize;
+		}
+		else
+		{
+			SeparateTranslucencyViewport = FScreenPassTextureViewport(FIntPoint(0, 0), FIntRect(0, 0, 1, 1));
+		}
 
 		FIntRect PassViewRect = View.ViewRect;
 
@@ -2620,9 +2628,9 @@ FRDGTextureRef DiaphragmDOF::AddPasses(
 			PermutationVector.Set<FDDOFBokehSimulationDim>(BokehSimulation);
 		PermutationVector.Set<FDiaphragmDOFRecombineCS::FQualityDim>(RecombineQuality);
 
-		FRDGTextureRef SeparateTranslucency = SeparateTranslucencyTextures.GetColorForRead(GraphBuilder);
-		FRDGTextureRef SeparateTranslucencyDepth = SeparateTranslucencyTextures.GetDepthForRead(GraphBuilder);
-		FRDGTextureRef SeparateTranslucencyModulateColor = SeparateTranslucencyTextures.GetColorModulateForRead(GraphBuilder);
+		FRDGTextureRef SeparateTranslucency = TranslucencyPassResources.GetColorForRead(GraphBuilder);
+		FRDGTextureRef SeparateTranslucencyDepth = TranslucencyPassResources.GetDepthForRead(GraphBuilder);
+		FRDGTextureRef SeparateTranslucencyModulateColor = TranslucencyPassResources.GetColorModulateForRead(GraphBuilder);
 
 		FDiaphragmDOFRecombineCS::FParameters* PassParameters = GraphBuilder.AllocParameters<FDiaphragmDOFRecombineCS::FParameters>();
 		PassParameters->CommonParameters = CommonParameters;
@@ -2635,11 +2643,11 @@ FRDGTextureRef DiaphragmDOF::AddPasses(
 			(GatheringViewSize.X - 0.5f) / float(RefBufferSize.X),
 			(GatheringViewSize.Y - 0.5f) / float(RefBufferSize.Y));
 
-		PassParameters->SeparateTranslucencyBilinearUVMinMax.X = (SeparateTranslucencyRect.Min.X + 0.5f) / float(SeparateTranslucencyDimensions.Extent.X);
-		PassParameters->SeparateTranslucencyBilinearUVMinMax.Y = (SeparateTranslucencyRect.Min.Y + 0.5f) / float(SeparateTranslucencyDimensions.Extent.Y);
-		PassParameters->SeparateTranslucencyBilinearUVMinMax.Z = (SeparateTranslucencyRect.Max.X - 0.5f) / float(SeparateTranslucencyDimensions.Extent.X);
-		PassParameters->SeparateTranslucencyBilinearUVMinMax.W = (SeparateTranslucencyRect.Max.Y - 0.5f) / float(SeparateTranslucencyDimensions.Extent.Y);
-		PassParameters->SeparateTranslucencyUpscaling = SeparateTranslucencyRect.Size() != FullResViewSize ? 1 : 0;
+		PassParameters->SeparateTranslucencyBilinearUVMinMax.X = (SeparateTranslucencyViewport.Rect.Min.X + 0.5f) / float(SeparateTranslucencyViewport.Extent.X);
+		PassParameters->SeparateTranslucencyBilinearUVMinMax.Y = (SeparateTranslucencyViewport.Rect.Min.Y + 0.5f) / float(SeparateTranslucencyViewport.Extent.Y);
+		PassParameters->SeparateTranslucencyBilinearUVMinMax.Z = (SeparateTranslucencyViewport.Rect.Max.X - 0.5f) / float(SeparateTranslucencyViewport.Extent.X);
+		PassParameters->SeparateTranslucencyBilinearUVMinMax.W = (SeparateTranslucencyViewport.Rect.Max.Y - 0.5f) / float(SeparateTranslucencyViewport.Extent.Y);
+		PassParameters->SeparateTranslucencyUpscaling = bScaleSeparateTranslucency ? 1 : 0;
 
 		PassParameters->SceneColorInput = FullResGatherInputTextures.SceneColor;
 		PassParameters->SceneDepthTexture = SceneTextures.SceneDepthTexture;
diff --git a/Engine/Source/Runtime/Renderer/Private/PostProcess/DiaphragmDOF.h b/Engine/Source/Runtime/Renderer/Private/PostProcess/DiaphragmDOF.h
--- a/Engine/Source/Runtime/Renderer/Private/PostProcess/DiaphragmDOF.h
+++ b/Engine/Source/Runtime/Renderer/Private/PostProcess/DiaphragmDOF.h
@@ -12,7 +12,7 @@
 
 class FViewInfo;
 class FSceneTextureParameters;
-class FSeparateTranslucencyTextures;
+struct FTranslucencyPassResources;
 struct FTemporalAAHistory;
 
 
@@ -129,6 +129,6 @@ RENDERER_API FRDGTextureRef AddPasses(
 	const FSceneTextureParameters& SceneTextures,
 	const FViewInfo& View,
 	FRDGTextureRef InputSceneColor,
-	const FSeparateTranslucencyTextures& SeparateTranslucencyTextures);
+	const FTranslucencyPassResources& TranslucencyViewResources);
 
 } // namespace DiaphragmDOF
diff --git a/Engine/Source/Runtime/Renderer/Private/PostProcess/PostProcessMotionBlur.cpp b/Engine/Source/Runtime/Renderer/Private/PostProcess/PostProcessMotionBlur.cpp
--- a/Engine/Source/Runtime/Renderer/Private/PostProcess/PostProcessMotionBlur.cpp
+++ b/Engine/Source/Runtime/Renderer/Private/PostProcess/PostProcessMotionBlur.cpp
@@ -1219,7 +1219,7 @@ FMotionBlurOutputs AddMotionBlurPass(FRDGBuilder& GraphBuilder, const FViewInfo&
 			MotionBlurFilterTexture,
 			VelocityFlatTexture,
 			VelocityTileTextures,
-			Inputs.PostMotionBlurTranslucency.Texture,
+			Inputs.PostMotionBlurTranslucency.ColorTexture.Resolve,
 			Inputs.PostMotionBlurTranslucency.ViewRect.Size(),
 			EMotionBlurFilterPass::Separable1,
 			Inputs.Quality);
@@ -1234,7 +1234,7 @@ FMotionBlurOutputs AddMotionBlurPass(FRDGBuilder& GraphBuilder, const FViewInfo&
 			Inputs.SceneColor.Texture,
 			VelocityFlatTexture,
 			VelocityTileTextures,
-			Inputs.PostMotionBlurTranslucency.Texture,
+			Inputs.PostMotionBlurTranslucency.ColorTexture.Resolve,
 			Inputs.PostMotionBlurTranslucency.ViewRect.Size(),
 			EMotionBlurFilterPass::Unified,
 			Inputs.Quality);
diff --git a/Engine/Source/Runtime/Renderer/Private/PostProcess/PostProcessMotionBlur.h b/Engine/Source/Runtime/Renderer/Private/PostProcess/PostProcessMotionBlur.h
--- a/Engine/Source/Runtime/Renderer/Private/PostProcess/PostProcessMotionBlur.h
+++ b/Engine/Source/Runtime/Renderer/Private/PostProcess/PostProcessMotionBlur.h
@@ -7,6 +7,8 @@
 #pragma once
 
 #include "ScreenPass.h"
+#include "TranslucentRendering.h"
+
 
 // Returns whether motion blur is enabled for the requested view.
 bool IsMotionBlurEnabled(const FViewInfo& View);
@@ -86,7 +88,7 @@ struct FMotionBlurInputs
 	FScreenPassTexture SceneVelocity;
 
 	// [Optional] The separate translucency buffer to be composited after motion blur
-	FScreenPassTexture PostMotionBlurTranslucency;
+	FTranslucencyPassResources PostMotionBlurTranslucency;
 
 	// [Required] Quality to use when processing motion blur.
 	EMotionBlurQuality Quality = EMotionBlurQuality::VeryHigh;
diff --git a/Engine/Source/Runtime/Renderer/Private/PostProcess/PostProcessing.cpp b/Engine/Source/Runtime/Renderer/Private/PostProcess/PostProcessing.cpp
--- a/Engine/Source/Runtime/Renderer/Private/PostProcess/PostProcessing.cpp
+++ b/Engine/Source/Runtime/Renderer/Private/PostProcess/PostProcessing.cpp
@@ -169,167 +169,10 @@ bool IsPostProcessingWithAlphaChannelSupported()
 	return CVarPostProcessingPropagateAlpha.GetValueOnAnyThread() != 0;
 }
 
-class FComposeSeparateTranslucencyPS : public FGlobalShader
-{
-	DECLARE_GLOBAL_SHADER(FComposeSeparateTranslucencyPS);
-	SHADER_USE_PARAMETER_STRUCT(FComposeSeparateTranslucencyPS, FGlobalShader);
-
-	class FNearestDepthNeighborUpsampling : SHADER_PERMUTATION_BOOL("PERMUTATION_NEARESTDEPTHNEIGHBOR");
-	using FPermutationDomain = TShaderPermutationDomain<FNearestDepthNeighborUpsampling>;
-
-	BEGIN_SHADER_PARAMETER_STRUCT(FParameters, )
-		SHADER_PARAMETER_STRUCT(FScreenPassTextureViewportParameters, Color)
-		SHADER_PARAMETER_STRUCT(FScreenPassTextureViewportParameters, Translucency)
-		SHADER_PARAMETER(FScreenTransform, ColorToTranslucency)
-		SHADER_PARAMETER_RDG_TEXTURE(Texture2D, SceneColor)
-		SHADER_PARAMETER_SAMPLER(SamplerState, SceneColorSampler)
-		SHADER_PARAMETER_RDG_TEXTURE(Texture2D, SeparateTranslucency)
-		SHADER_PARAMETER_SAMPLER(SamplerState, SeparateTranslucencySampler)
-		SHADER_PARAMETER_RDG_TEXTURE(Texture2D, SeparateModulation)
-		SHADER_PARAMETER_SAMPLER(SamplerState, SeparateModulationSampler)
-		SHADER_PARAMETER_RDG_TEXTURE(Texture2D, LowResDepthTexture)
-		SHADER_PARAMETER_SAMPLER(SamplerState, LowResDepthSampler)
-		SHADER_PARAMETER_RDG_TEXTURE(Texture2D, FullResDepthTexture)
-		SHADER_PARAMETER_SAMPLER(SamplerState, FullResDepthSampler)
-		SHADER_PARAMETER_STRUCT_REF(FViewUniformShaderParameters, ViewUniformBuffer)
-		RENDER_TARGET_BINDING_SLOTS()
-	END_SHADER_PARAMETER_STRUCT()
-
-	static bool ShouldCompilePermutation(const FGlobalShaderPermutationParameters& Parameters)
-	{
-		return IsFeatureLevelSupported(Parameters.Platform, ERHIFeatureLevel::SM5);
-	}
-};
-
-IMPLEMENT_GLOBAL_SHADER(FComposeSeparateTranslucencyPS, "/Engine/Private/ComposeSeparateTranslucency.usf", "MainPS", SF_Pixel);
-
-extern bool GetUseTranslucencyNearestDepthNeighborUpsample(float DownsampleScale);
-
 #if DEBUG_POST_PROCESS_VOLUME_ENABLE
 FScreenPassTexture AddFinalPostProcessDebugInfoPasses(FRDGBuilder& GraphBuilder, const FViewInfo& View, FScreenPassTexture& ScreenPassSceneColor);
 #endif
 
-FRDGTextureRef AddTranslucencyCompositionPass(
-	FRDGBuilder& GraphBuilder,
-	const FViewInfo& View,
-	const FScreenPassTexture& SceneColor,
-	const FScreenPassTexture& SceneDepth,
-	const FSeparateTranslucencyTextures& SeparateTranslucencyTextures,
-	bool bPostMotionBlur,
-	bool bApplyModulateOnly)
-{
-	FRDGTextureRef SeparateModulationTexture;
-	FRDGTextureRef SeparateTranslucencyTexture;
-
-	if (bPostMotionBlur)
-	{
-		check(!bApplyModulateOnly);
-
-		// if nothing is rendered into the separate translucency, then just return the existing Scenecolor
-		if (!SeparateTranslucencyTextures.IsPostMotionBlurColorValid())
-		{
-			return SceneColor.Texture;
-		}
-
-		SeparateTranslucencyTexture = SeparateTranslucencyTextures.GetPostMotionBlurColorForRead(GraphBuilder);
-		SeparateModulationTexture = GraphBuilder.RegisterExternalTexture(GSystemTextures.WhiteDummy);
-	}
-	else
-	{
-		// if nothing is rendered into the separate translucency, then just return the existing Scenecolor
-		if (!SeparateTranslucencyTextures.IsColorValid() && !SeparateTranslucencyTextures.IsColorModulateValid())
-		{
-			return SceneColor.Texture;
-		}
-
-		SeparateTranslucencyTexture = SeparateTranslucencyTextures.GetColorForRead(GraphBuilder);
-		SeparateModulationTexture = SeparateTranslucencyTextures.GetColorModulateForRead(GraphBuilder);
-
-		if (bApplyModulateOnly)
-		{
-			if (!SeparateTranslucencyTextures.IsColorModulateValid())
-			{
-				return SceneColor.Texture;
-			}
-
-			SeparateTranslucencyTexture = GraphBuilder.RegisterExternalTexture(GSystemTextures.BlackAlphaOneDummy);
-		}
-	}
-	
-	FRDGTextureRef NewSceneColor;
-	{
-		FRDGTextureDesc OutputDesc = SceneColor.Texture->Desc;
-		OutputDesc.Reset();
-
-		if (OutputDesc.Format == PF_FloatRGBA && !IsPostProcessingWithAlphaChannelSupported())
-		{
-			OutputDesc.Format = PF_FloatRGB;
-		}
-
-		NewSceneColor = GraphBuilder.CreateTexture(
-			OutputDesc,
-			bPostMotionBlur ? TEXT("PostMotionBlurTranslucency.SceneColor") : TEXT("PostDOFTranslucency.SceneColor"));
-	}
-
-	const FVector2f SceneColorSize(SceneColor.ViewRect.Size());
-	const FVector2f SceneColorSizeInv = FVector2f(1.0f, 1.0f) / SceneColorSize;	
-	const FVector2f SceneColorExtent(NewSceneColor->Desc.Extent);
-	const FVector2f SceneColorExtentInv = FVector2f(1.0f, 1.0f) / SceneColorExtent;
-	
-	const FIntRect SeparateTranslucencyRect = SeparateTranslucencyTextures.GetDimensions().GetViewport(View.ViewRect).Rect;
-	const FVector2f SeparateTranslucencySize(SeparateTranslucencyRect.Size());
-	const FVector2f SeparateTranslucencyExtent((bApplyModulateOnly ? SeparateModulationTexture : SeparateTranslucencyTexture)->Desc.Extent);
-	const FVector2f SeparateTranslucencyExtentInv = FVector2f(1.0f, 1.0f) / SeparateTranslucencyExtent;
-
-	const FScreenPassTextureViewport SceneColorViewport(SceneColor);
-	const FScreenPassTextureViewport TranslucencyViewport = SeparateTranslucencyTextures.GetDimensions().GetViewport(View.ViewRect);
-
-	const float DownsampleScale = SeparateTranslucencyTextures.GetDimensions().Scale;
-
-	const bool bScaleSeparateTranslucency = DownsampleScale != 1.0f;
-	const bool DepthUpscampling = !bPostMotionBlur && GetUseTranslucencyNearestDepthNeighborUpsample(DownsampleScale) ? 1 : 0;
-
-	FComposeSeparateTranslucencyPS::FParameters* PassParameters = GraphBuilder.AllocParameters<FComposeSeparateTranslucencyPS::FParameters>();
-	PassParameters->Color = GetScreenPassTextureViewportParameters(SceneColorViewport);
-	PassParameters->Translucency = GetScreenPassTextureViewportParameters(TranslucencyViewport);
-	PassParameters->ColorToTranslucency = FScreenTransform::ChangeTextureUVCoordinateFromTo(SceneColorViewport, TranslucencyViewport);
-	PassParameters->SceneColor = SceneColor.Texture;
-	PassParameters->SceneColorSampler = TStaticSamplerState<SF_Point>::GetRHI();
-	PassParameters->SeparateTranslucency = SeparateTranslucencyTexture;
-	PassParameters->SeparateTranslucencySampler = bScaleSeparateTranslucency ? TStaticSamplerState<SF_Bilinear>::GetRHI() : TStaticSamplerState<SF_Point>::GetRHI();
-	PassParameters->SeparateModulation = SeparateModulationTexture;
-	PassParameters->SeparateModulationSampler = bScaleSeparateTranslucency ? TStaticSamplerState<SF_Bilinear>::GetRHI() : TStaticSamplerState<SF_Point>::GetRHI();
-	PassParameters->ViewUniformBuffer = View.ViewUniformBuffer;
-	PassParameters->RenderTargets[0] = FRenderTargetBinding(NewSceneColor, ERenderTargetLoadAction::ENoAction);
-
-	if (DepthUpscampling)
-	{
-		PassParameters->LowResDepthTexture = SeparateTranslucencyTextures.GetDepthForRead(GraphBuilder);
-		PassParameters->LowResDepthSampler = TStaticSamplerState<SF_Point>::GetRHI();
-		PassParameters->FullResDepthTexture = SceneDepth.Texture;
-		PassParameters->FullResDepthSampler = TStaticSamplerState<SF_Point>::GetRHI();
-	}
-
-	FComposeSeparateTranslucencyPS::FPermutationDomain PermutationVector;
-	PermutationVector.Set<FComposeSeparateTranslucencyPS::FNearestDepthNeighborUpsampling>(DepthUpscampling);
-
-	TShaderMapRef<FComposeSeparateTranslucencyPS> PixelShader(View.ShaderMap, PermutationVector);
-	FPixelShaderUtils::AddFullscreenPass(
-		GraphBuilder,
-		View.ShaderMap,
-		RDG_EVENT_NAME(
-			"ComposeTranslucency(%s%s%s) %dx%d",
-			bPostMotionBlur ? TEXT("PostMB") : TEXT("PostDOF"),
-			bApplyModulateOnly ? TEXT(" ModulateOnly") : TEXT(""),
-			bScaleSeparateTranslucency ? TEXT(" Rescale") : TEXT(""),
-			SceneColor.ViewRect.Width(), SceneColor.ViewRect.Height()),
-		PixelShader,
-		PassParameters,
-		SceneColor.ViewRect);
-
-	return NewSceneColor;
-}
-
 void AddTemporalAA2Passes(
 	FRDGBuilder& GraphBuilder,
 	const FSceneTextureParameters& SceneTextures,
@@ -358,23 +201,17 @@ void AddPostProcessingPasses(
 	FScene* Scene = View.Family->Scene->GetRenderScene();
 
 	const FIntRect PrimaryViewRect = View.ViewRect;
-	const FIntRect SeparateTranslucencyRect = Inputs.SeparateTranslucencyTextures->GetDimensions().GetViewport(PrimaryViewRect).Rect;
 
 	const FSceneTextureParameters SceneTextureParameters = GetSceneTextureParameters(GraphBuilder, Inputs.SceneTextures);
 
 	const FScreenPassRenderTarget ViewFamilyOutput = FScreenPassRenderTarget::CreateViewFamilyOutput(Inputs.ViewFamilyTexture, View);
 	const FScreenPassTexture SceneDepth(SceneTextureParameters.SceneDepthTexture, PrimaryViewRect);
-	const FScreenPassTexture SeparateTranslucency(Inputs.SeparateTranslucencyTextures->GetColorForRead(GraphBuilder), PrimaryViewRect);
 	const FScreenPassTexture CustomDepth(Inputs.CustomDepthTexture, PrimaryViewRect);
 	const FScreenPassTexture Velocity(SceneTextureParameters.GBufferVelocityTexture, PrimaryViewRect);
 	const FScreenPassTexture BlackDummy(GSystemTextures.GetBlackDummy(GraphBuilder));
 	
-	// Post-MotionBlur translucency need only be valid if we rendered it
-	FScreenPassTexture PostMotionBlurTranslucency;
-	if (Inputs.SeparateTranslucencyTextures->IsPostMotionBlurColorValid())
-	{
-		PostMotionBlurTranslucency = FScreenPassTexture(Inputs.SeparateTranslucencyTextures->GetPostMotionBlurColorForRead(GraphBuilder), SeparateTranslucencyRect);
-	}
+	const FTranslucencyPassResources& PostDOFTranslucencyResources = Inputs.TranslucencyViewResourcesMap.Get(ETranslucencyPass::TPT_TranslucencyAfterDOF);
+	const FTranslucencyPassResources& PostMotionBlurTranslucencyResources = Inputs.TranslucencyViewResourcesMap.Get(ETranslucencyPass::TPT_TranslucencyAfterMotionBlur);
 
 	// Scene color is updated incrementally through the post process pipeline.
 	FScreenPassTexture SceneColor((*Inputs.SceneTextures)->SceneColorTexture, PrimaryViewRect);
@@ -522,7 +359,7 @@ void AddPostProcessingPasses(
 		FPostProcessMaterialInputs PostProcessMaterialInputs;
 
 		PostProcessMaterialInputs.SetInput(EPostProcessMaterialInput::SceneColor, InSceneColor);
-		PostProcessMaterialInputs.SetInput(EPostProcessMaterialInput::SeparateTranslucency, SeparateTranslucency);
+		PostProcessMaterialInputs.SetInput(EPostProcessMaterialInput::SeparateTranslucency, FScreenPassTexture(PostDOFTranslucencyResources.GetColorForRead(GraphBuilder), PostDOFTranslucencyResources.ViewRect));
 		PostProcessMaterialInputs.SetInput(EPostProcessMaterialInput::Velocity, Velocity);
 		PostProcessMaterialInputs.SceneTextures = GetSceneTextureShaderParameters(Inputs.SceneTextures);
 		PostProcessMaterialInputs.CustomDepthTexture = CustomDepth.Texture;
@@ -666,28 +503,28 @@ void AddPostProcessingPasses(
 
 			if (bDepthOfFieldEnabled)
 			{
-				FSeparateTranslucencyDimensions DummyTranslucencyDimensions;
-				FSeparateTranslucencyTextures DummyTranslucency(DummyTranslucencyDimensions);
+				FTranslucencyPassResources DummyTranslucency;
 
 				SceneColor.Texture = DiaphragmDOF::AddPasses(
 					GraphBuilder,
 					SceneTextureParameters,
 					View,
 					SceneColor.Texture,
-					bComposeSeparateTranslucencyInTSR ? DummyTranslucencyDimensions : *Inputs.SeparateTranslucencyTextures);
+					bComposeSeparateTranslucencyInTSR ? DummyTranslucency : PostDOFTranslucencyResources);
 			}
 
 			// DOF passes were not added, therefore need to compose Separate translucency manually.
-			if (SceneColor.Texture == InputSceneColorTexture || bComposeSeparateTranslucencyInTSR)
+			if ((SceneColor.Texture == InputSceneColorTexture || bComposeSeparateTranslucencyInTSR) && PostDOFTranslucencyResources.IsValid())
 			{
-				SceneColor.Texture = AddTranslucencyCompositionPass(
-					GraphBuilder,
-					View,
-					SceneColor,
-					SceneDepth,
-					*Inputs.SeparateTranslucencyTextures,
-					/* bPostMotionBlur = */ false,
-					/* bApplyModulateOnly = */ bComposeSeparateTranslucencyInTSR);
+				FTranslucencyComposition TranslucencyComposition;
+				TranslucencyComposition.Operation = FTranslucencyComposition::EOperation::ComposeToNewSceneColor;
+				TranslucencyComposition.bApplyModulateOnly = bComposeSeparateTranslucencyInTSR;
+				TranslucencyComposition.SceneColor = SceneColor;
+				TranslucencyComposition.SceneDepth = SceneDepth;
+				TranslucencyComposition.OutputViewport = FScreenPassTextureViewport(SceneColor);
+
+				SceneColor = TranslucencyComposition.AddPass(
+					GraphBuilder, View, PostDOFTranslucencyResources);
 			}
 
 			if (GetHairStrandsComposition() == EHairStrandsCompositionType::AfterSeparateTranslucent)
@@ -744,7 +581,7 @@ void AddPostProcessingPasses(
 			UpscalerPassInputs.SceneColorTexture = SceneColor.Texture;
 			UpscalerPassInputs.SceneDepthTexture = SceneDepth.Texture;
 			UpscalerPassInputs.SceneVelocityTexture = Velocity.Texture;
-			UpscalerPassInputs.SeparateTranslucencyTextures = Inputs.SeparateTranslucencyTextures;
+			UpscalerPassInputs.PostDOFTranslucencyResources = PostDOFTranslucencyResources;
 
 			ITemporalUpscaler::FOutputs Outputs = UpscalerToUse->AddPasses(
 				GraphBuilder,
@@ -801,7 +638,7 @@ void AddPostProcessingPasses(
 			PassInputs.SceneColor = SceneColor;
 			PassInputs.SceneDepth = SceneDepth;
 			PassInputs.SceneVelocity = Velocity;
-			PassInputs.PostMotionBlurTranslucency = PostMotionBlurTranslucency;
+			PassInputs.PostMotionBlurTranslucency = PostMotionBlurTranslucencyResources;
 			PassInputs.Quality = GetMotionBlurQuality();
 			PassInputs.Filter = GetMotionBlurFilter();
 			PassInputs.VelocityFlattenTextures = VelocityFlattenTextures;
@@ -819,18 +656,16 @@ void AddPostProcessingPasses(
 				QuarterResSceneColor = PassOutputs.QuarterRes;
 			}
 		}
-		else
+		else if (PostMotionBlurTranslucencyResources.IsValid())
 		{
-			// Compose Post-MotionBlur translucency
+			// Compose Post-MotionBlur translucency in a new scene color to ensure it's not writing out in TAA's output that is also the history.
+			FTranslucencyComposition TranslucencyComposition;
+			TranslucencyComposition.Operation = FTranslucencyComposition::EOperation::ComposeToNewSceneColor;
+			TranslucencyComposition.SceneColor = SceneColor;
+			TranslucencyComposition.OutputViewport = FScreenPassTextureViewport(SceneColor);
 
-			SceneColor.Texture = AddTranslucencyCompositionPass(
-				GraphBuilder,
-				View,
-				SceneColor,
-				SceneDepth,
-				*Inputs.SeparateTranslucencyTextures,
-				/* bPostMotionBlur = */ true,
-				/* bApplyModulateOnly = */ false);
+			SceneColor = TranslucencyComposition.AddPass(
+				GraphBuilder, View, PostMotionBlurTranslucencyResources);
 		}
 
 		{
@@ -1025,7 +860,7 @@ void AddPostProcessingPasses(
 				FPostProcessMaterialInputs PassInputs;
 				PassSequence.AcceptOverrideIfLastPass(EPass::Tonemap, PassInputs.OverrideOutput);
 				PassInputs.SetInput(EPostProcessMaterialInput::SceneColor, SceneColor);
-				PassInputs.SetInput(EPostProcessMaterialInput::SeparateTranslucency, SeparateTranslucency);
+				//PassInputs.SetInput(EPostProcessMaterialInput::SeparateTranslucency, SeparateTranslucency);
 				PassInputs.SetInput(EPostProcessMaterialInput::CombinedBloom, Bloom);
 				PassInputs.SceneTextures = GetSceneTextureShaderParameters(Inputs.SceneTextures);
 				PassInputs.CustomDepthTexture = CustomDepth.Texture;
@@ -1120,22 +955,26 @@ void AddPostProcessingPasses(
 		PassSequence.Finalize();
 
 		// Compose separate translucency passes
-		SceneColor.Texture = AddTranslucencyCompositionPass(
-			GraphBuilder,
-			View,
-			SceneColor,
-			SceneDepth,
-			*Inputs.SeparateTranslucencyTextures,
-			/* bPostMotionBlur = */ false,
-			/* bApplyModulateOnly = */ false);
-		SceneColor.Texture = AddTranslucencyCompositionPass(
-			GraphBuilder,
-			View,
-			SceneColor,
-			SceneDepth,
-			*Inputs.SeparateTranslucencyTextures,
-			/* bPostMotionBlur = */ true,
-			/* bApplyModulateOnly = */ false);
+		{
+			FTranslucencyComposition TranslucencyComposition;
+			TranslucencyComposition.Operation = FTranslucencyComposition::EOperation::ComposeToNewSceneColor;
+			TranslucencyComposition.SceneColor = SceneColor;
+			TranslucencyComposition.OutputViewport = FScreenPassTextureViewport(SceneColor);
+
+			if (PostDOFTranslucencyResources.IsValid())
+			{
+				TranslucencyComposition.SceneColor = TranslucencyComposition.AddPass(
+					GraphBuilder, View, PostDOFTranslucencyResources);
+			}
+
+			if (PostMotionBlurTranslucencyResources.IsValid())
+			{
+				TranslucencyComposition.SceneColor = TranslucencyComposition.AddPass(
+					GraphBuilder, View, PostMotionBlurTranslucencyResources);
+			}
+
+			SceneColor = TranslucencyComposition.SceneColor;
+		}
 
 		SceneColorBeforeTonemap = SceneColor;
 
@@ -1293,7 +1132,7 @@ void AddPostProcessingPasses(
 		PassInputs.SceneColor = SceneColor;
 		PassInputs.SceneColorBeforeTonemap = SceneColorBeforeTonemap;
 		PassInputs.SceneColorAfterTonemap = SceneColorAfterTonemap;
-		PassInputs.SeparateTranslucency = SeparateTranslucency;
+		PassInputs.SeparateTranslucency = FScreenPassTexture(PostDOFTranslucencyResources.GetColorForRead(GraphBuilder), PostDOFTranslucencyResources.ViewRect); // TODO
 		PassInputs.Velocity = Velocity;
 		PassInputs.SceneTextures = GetSceneTextureShaderParameters(Inputs.SceneTextures);
 		PassInputs.bOverview = bVisualizeGBufferOverview;
diff --git a/Engine/Source/Runtime/Renderer/Private/PostProcess/PostProcessing.h b/Engine/Source/Runtime/Renderer/Private/PostProcess/PostProcessing.h
--- a/Engine/Source/Runtime/Renderer/Private/PostProcess/PostProcessing.h
+++ b/Engine/Source/Runtime/Renderer/Private/PostProcess/PostProcessing.h
@@ -31,13 +31,13 @@ struct FPostProcessingInputs
 	TRDGUniformBufferRef<FSceneTextureUniformParameters> SceneTextures = nullptr;
 	FRDGTextureRef ViewFamilyTexture = nullptr;
 	FRDGTextureRef CustomDepthTexture = nullptr;
-	const FSeparateTranslucencyTextures* SeparateTranslucencyTextures = nullptr;
+	FTranslucencyViewResourcesMap TranslucencyViewResourcesMap;
 
 	void Validate() const
 	{
 		check(SceneTextures);
 		check(ViewFamilyTexture);
-		check(SeparateTranslucencyTextures);
+		check(TranslucencyViewResourcesMap.IsValid());
 	}
 };
 
diff --git a/Engine/Source/Runtime/Renderer/Private/PostProcess/TemporalAA.h b/Engine/Source/Runtime/Renderer/Private/PostProcess/TemporalAA.h
--- a/Engine/Source/Runtime/Renderer/Private/PostProcess/TemporalAA.h
+++ b/Engine/Source/Runtime/Renderer/Private/PostProcess/TemporalAA.h
@@ -6,7 +6,7 @@
 #include "PostProcessMotionBlur.h"
 
 struct FTemporalAAHistory;
-class FSeparateTranslucencyTextures;
+struct FTranslucencyPassResources;
 
 
 /** Configuration of the main temporal AA pass. */
@@ -185,7 +185,7 @@ public:
 		FRDGTextureRef SceneColorTexture = nullptr;
 		FRDGTextureRef SceneDepthTexture = nullptr;
 		FRDGTextureRef SceneVelocityTexture = nullptr;
-		const FSeparateTranslucencyTextures* SeparateTranslucencyTextures = nullptr;
+		FTranslucencyPassResources PostDOFTranslucencyResources;
 	};
 
 	struct FOutputs
diff --git a/Engine/Source/Runtime/Renderer/Private/PostProcess/TemporalSuperResolution.cpp b/Engine/Source/Runtime/Renderer/Private/PostProcess/TemporalSuperResolution.cpp
--- a/Engine/Source/Runtime/Renderer/Private/PostProcess/TemporalSuperResolution.cpp
+++ b/Engine/Source/Runtime/Renderer/Private/PostProcess/TemporalSuperResolution.cpp
@@ -592,7 +592,7 @@ ITemporalUpscaler::FOutputs AddTemporalSuperResolutionPasses(
 
 	bool bHalfResLowFrequency = CVarTSRHalfResShadingRejection.GetValueOnRenderThread() != 0;
 
-	bool bIsSeperateTranslucyTexturesValid = PassInputs.SeparateTranslucencyTextures != nullptr && PassInputs.SeparateTranslucencyTextures->IsColorValid();
+	bool bIsSeperateTranslucyTexturesValid = PassInputs.PostDOFTranslucencyResources.IsValid();
 
 	bool bRejectSeparateTranslucency = bIsSeperateTranslucyTexturesValid && CVarTSRTranslucencyPreviousFrameRejection.GetValueOnRenderThread() != 0;
 
@@ -687,10 +687,17 @@ ITemporalUpscaler::FOutputs AddTemporalSuperResolutionPasses(
 
 	FRDGTextureRef BlackUintDummy = GSystemTextures.GetZeroUIntDummy(GraphBuilder);
 	FRDGTextureRef BlackDummy = GSystemTextures.GetBlackDummy(GraphBuilder);
+	FRDGTextureRef BlackAlphaOneDummy = GSystemTextures.GetBlackAlphaOneDummy(GraphBuilder);
 	FRDGTextureRef WhiteDummy = GSystemTextures.GetWhiteDummy(GraphBuilder);
 
-	FRDGTextureRef SeparateTranslucencyTexture = (bRejectSeparateTranslucency || bAccumulateSeparateTranslucency) ?
-		PassInputs.SeparateTranslucencyTextures->GetColorForRead(GraphBuilder) : nullptr;
+	FIntRect SeparateTranslucencyRect;
+	FRDGTextureRef SeparateTranslucencyTexture = nullptr;
+	if (bRejectSeparateTranslucency || bAccumulateSeparateTranslucency)
+	{
+		check(PassInputs.PostDOFTranslucencyResources.IsValid());
+		SeparateTranslucencyTexture = PassInputs.PostDOFTranslucencyResources.ColorTexture.Resolve;
+		SeparateTranslucencyRect = PassInputs.PostDOFTranslucencyResources.ViewRect;
+	}
 
 	FTSRCommonParameters CommonParameters;
 	{
@@ -1037,12 +1044,12 @@ ITemporalUpscaler::FOutputs AddTemporalSuperResolutionPasses(
 		{
 
 			PrevTranslucencyTexture = GraphBuilder.RegisterExternalTexture(View.PrevViewInfo.SeparateTranslucency);
-			PrevTranslucencyViewport = FScreenPassTextureViewport(PrevTranslucencyTexture->Desc.Extent, View.PrevViewInfo.ViewRect);
+			PrevTranslucencyViewport = FScreenPassTextureViewport(PrevTranslucencyTexture->Desc.Extent, View.PrevViewInfo.SeparateTranslucencyViewRect);
 
 		}
 		else
 		{
-			PrevTranslucencyTexture = BlackDummy;
+			PrevTranslucencyTexture = BlackAlphaOneDummy;
 			PrevTranslucencyViewport = FScreenPassTextureViewport(FIntPoint(1, 1), FIntRect(FIntPoint(0, 0), FIntPoint(1, 1)));
 		}
 
@@ -1060,7 +1067,7 @@ ITemporalUpscaler::FOutputs AddTemporalSuperResolutionPasses(
 		PassParameters->CommonParameters = CommonParameters;
 
 		PassParameters->TranslucencyInfo = GetScreenPassTextureViewportParameters(FScreenPassTextureViewport(
-			InputExtent, InputRect));
+			SeparateTranslucencyTexture->Desc.Extent, SeparateTranslucencyRect));
 		PassParameters->PrevTranslucencyInfo = GetScreenPassTextureViewportParameters(PrevTranslucencyViewport);
 		PassParameters->PrevTranslucencyPreExposureCorrection = PrevHistoryParameters.HistoryPreExposureCorrection;
 		PassParameters->TranslucencyHighlightLuminance = CVarTSRTranslucencyHighlightLuminance.GetValueOnRenderThread();
@@ -1358,6 +1365,8 @@ ITemporalUpscaler::FOutputs AddTemporalSuperResolutionPasses(
 
 	// Update temporal history.
 	{
+		ensure(SeparateTranslucencyRect.Size() == InputRect.Size() || SeparateTranslucencyTexture == nullptr);
+
 		static const TCHAR* const kUpdateQualityNames[] = {
 			TEXT("Low"),
 			TEXT("Medium"),
@@ -1516,6 +1525,7 @@ ITemporalUpscaler::FOutputs AddTemporalSuperResolutionPasses(
 		{
 			GraphBuilder.QueueTextureExtraction(
 				SeparateTranslucencyTexture, &View.ViewState->PrevFrameViewInfo.SeparateTranslucency);
+			View.ViewState->PrevFrameViewInfo.SeparateTranslucencyViewRect = InputRect;
 		}
 
 		// Extract the output for next frame SSR so that separate translucency shows up in SSR.
diff --git a/Engine/Source/Runtime/Renderer/Private/SceneRendering.h b/Engine/Source/Runtime/Renderer/Private/SceneRendering.h
--- a/Engine/Source/Runtime/Renderer/Private/SceneRendering.h
+++ b/Engine/Source/Runtime/Renderer/Private/SceneRendering.h
@@ -1005,6 +1005,7 @@ struct FPreviousViewInfo
 	TRefCountPtr<IPooledRenderTarget> ScreenSpaceRayTracingInput;
 
 	// Separate translucency of previous frame to detect translucency change in TAA
+	FIntRect SeparateTranslucencyViewRect;
 	TRefCountPtr<IPooledRenderTarget> SeparateTranslucency;
 
 	// Temporal AA result of last frame
diff --git a/Engine/Source/Runtime/Renderer/Private/ScreenPass.cpp b/Engine/Source/Runtime/Renderer/Private/ScreenPass.cpp
--- a/Engine/Source/Runtime/Renderer/Private/ScreenPass.cpp
+++ b/Engine/Source/Runtime/Renderer/Private/ScreenPass.cpp
@@ -225,16 +225,23 @@ void AddDownsampleDepthPass(
 
 	FRHIDepthStencilState* DepthStencilState = TStaticDepthStencilState<true, CF_Always>::GetRHI();
 
-	auto GetDownsampleDepthPassName = [](EDownsampleDepthFilter DownsampleDepthFilter)
-	{
-		switch (DownsampleDepthFilter)
-		{
-		case EDownsampleDepthFilter::Point:						return RDG_EVENT_NAME("DownsampleDepth-Point");
-		case EDownsampleDepthFilter::Max:						return RDG_EVENT_NAME("DownsampleDepth-Max");
-		case EDownsampleDepthFilter::Checkerboard:				return RDG_EVENT_NAME("DownsampleDepth-CheckerMinMax");
-		default:												return RDG_EVENT_NAME("MissingName");
-		}
+	static const TCHAR* kFilterNames[] = {
+		TEXT("Point"),
+		TEXT("Max"),
+		TEXT("CheckerMinMax"),
 	};
 
-	AddDrawScreenPass(GraphBuilder, GetDownsampleDepthPassName(DownsampleDepthFilter), View, OutputViewport, InputViewport, VertexShader, PixelShader, DepthStencilState, PassParameters);
+	AddDrawScreenPass(
+		GraphBuilder,
+		RDG_EVENT_NAME("DownsampleDepth(%s) %dx%dx -> %dx%d",
+			kFilterNames[int32(DownsampleDepthFilter)],
+			InputViewport.Rect.Width(),
+			InputViewport.Rect.Height(),
+			OutputViewport.Rect.Width(),
+			OutputViewport.Rect.Height()),
+		View,
+		OutputViewport, InputViewport,
+		VertexShader, PixelShader,
+		DepthStencilState,
+		PassParameters);
 }
\ No newline at end of file
diff --git a/Engine/Source/Runtime/Renderer/Private/ScreenPass.h b/Engine/Source/Runtime/Renderer/Private/ScreenPass.h
--- a/Engine/Source/Runtime/Renderer/Private/ScreenPass.h
+++ b/Engine/Source/Runtime/Renderer/Private/ScreenPass.h
@@ -313,6 +313,13 @@ struct FScreenTransform
 		return ChangeTextureBasisFromTo(TextureViewport.Extent, TextureViewport.Rect, SrcBasis, DestBasis);
 	}
 
+	static inline FScreenTransform ChangeTextureBasisFromTo(
+		const FScreenPassTexture& Texture,
+		ETextureBasis SrcBasis, ETextureBasis DestBasis)
+	{
+		return ChangeTextureBasisFromTo(FScreenPassTextureViewport(Texture.Texture, Texture.ViewRect), SrcBasis, DestBasis);
+	}
+
 	/** Change TextureUV coordinate from one texture to another, taking into account change in texture extent too. */
 	static FScreenTransform ChangeTextureUVCoordinateFromTo(
 		const FScreenPassTextureViewport& SrcViewport,
diff --git a/Engine/Source/Runtime/Renderer/Private/TranslucentRendering.cpp b/Engine/Source/Runtime/Renderer/Private/TranslucentRendering.cpp
--- a/Engine/Source/Runtime/Renderer/Private/TranslucentRendering.cpp
+++ b/Engine/Source/Runtime/Renderer/Private/TranslucentRendering.cpp
@@ -16,6 +16,7 @@
 #include "VirtualShadowMaps/VirtualShadowMapArray.h"
 #include "Strata/Strata.h"
 #include "HairStrands/HairStrandsUtils.h"
+#include "PixelShaderUtils.h"
 
 DECLARE_CYCLE_STAT(TEXT("TranslucencyTimestampQueryFence Wait"), STAT_TranslucencyTimestampQueryFence_Wait, STATGROUP_SceneRendering);
 DECLARE_CYCLE_STAT(TEXT("TranslucencyTimestampQuery Wait"), STAT_TranslucencyTimestampQuery_Wait, STATGROUP_SceneRendering);
@@ -76,23 +77,19 @@ static TAutoConsoleVariable<int32> CVarParallelTranslucency(
 	TEXT("Toggles parallel translucency rendering. Parallel rendering must be enabled for this to have an effect."),
 	ECVF_RenderThreadSafe);
 
+static const TCHAR* kTranslucencyPassName[] = {
+	TEXT("BeforeDistortion"),
+	TEXT("AfterDOF"),
+	TEXT("AfterDOFModulate"),
+	TEXT("AfterMotionBlur"),
+	TEXT("All"),
+};
+static_assert(UE_ARRAY_COUNT(kTranslucencyPassName) == int32(ETranslucencyPass::TPT_MAX), "Fix me");
+
+
 static const TCHAR* TranslucencyPassToString(ETranslucencyPass::Type TranslucencyPass)
 {
-	switch (TranslucencyPass)
-	{
-	case ETranslucencyPass::TPT_StandardTranslucency:
-		return TEXT("Standard");
-	case ETranslucencyPass::TPT_TranslucencyAfterDOF:
-		return TEXT("AfterDOF");
-	case ETranslucencyPass::TPT_TranslucencyAfterDOFModulate:
-		return TEXT("AfterDOFModulate");
-	case ETranslucencyPass::TPT_TranslucencyAfterMotionBlur:
-		return TEXT("AfterMotionBlur");
-	case ETranslucencyPass::TPT_AllTranslucency:
-		return TEXT("All");
-	}
-	checkNoEntry();
-	return TEXT("");
+	return kTranslucencyPassName[TranslucencyPass];
 }
 
 EMeshPass::Type TranslucencyPassToMeshPass(ETranslucencyPass::Type TranslucencyPass)
@@ -379,96 +376,18 @@ FSeparateTranslucencyDimensions UpdateTranslucencyTimers(FRHICommandListImmediat
 	return Dimensions;
 }
 
-FRDGTextureMSAA FSeparateTranslucencyTextures::GetColorForWrite(FRDGBuilder& GraphBuilder)
+FTranslucencyPassResourcesMap::FTranslucencyPassResourcesMap(int32 NumViews)
 {
-	if (!ColorTexture.IsValid())
-	{
-		const FRDGTextureDesc Desc = FRDGTextureDesc::Create2D(Dimensions.Extent, PF_FloatRGBA, FClearValueBinding::Black, TexCreate_RenderTargetable | TexCreate_ShaderResource, 1, Dimensions.NumSamples);
-		ColorTexture = CreateTextureMSAA(GraphBuilder, Desc, TEXT("SeparateTranslucencyColor"), GFastVRamConfig.SeparateTranslucency);
-	}
-	return ColorTexture;
-}
+	Array.SetNum(NumViews);
 
-FRDGTextureRef FSeparateTranslucencyTextures::GetColorForRead(FRDGBuilder& GraphBuilder) const
-{
-	if (ColorTexture.IsValid())
+	for (int32 ViewIndex = 0; ViewIndex < NumViews; ViewIndex++)
 	{
-		return ColorTexture.Resolve;
-	}
-	return GraphBuilder.RegisterExternalTexture(GSystemTextures.BlackAlphaOneDummy);
-}
-
-FRDGTextureMSAA FSeparateTranslucencyTextures::GetColorModulateForWrite(FRDGBuilder& GraphBuilder)
-{
-	if (!ColorModulateTexture.IsValid())
-	{
-		const FRDGTextureDesc Desc = FRDGTextureDesc::Create2D(Dimensions.Extent, PF_FloatR11G11B10, FClearValueBinding::White, TexCreate_RenderTargetable | TexCreate_ShaderResource, 1, Dimensions.NumSamples);
-		ColorModulateTexture = CreateTextureMSAA(GraphBuilder, Desc, TEXT("SeparateTranslucencyModulateColor"), GFastVRamConfig.SeparateTranslucencyModulate);
-	}
-	return ColorModulateTexture;
-}
-
-FRDGTextureRef FSeparateTranslucencyTextures::GetColorModulateForRead(FRDGBuilder& GraphBuilder) const
-{
-	if (ColorModulateTexture.IsValid())
-	{
-		return ColorModulateTexture.Resolve;
-	}
-	return GraphBuilder.RegisterExternalTexture(GSystemTextures.WhiteDummy);
-}
-
-FRDGTextureMSAA FSeparateTranslucencyTextures::GetPostMotionBlurColorForWrite(FRDGBuilder& GraphBuilder)
-{
-	if (!PostMotionBlurColorTexture.IsValid())
-	{
-		const FRDGTextureDesc Desc = FRDGTextureDesc::Create2D(Dimensions.Extent, PF_FloatRGBA, FClearValueBinding::Black, TexCreate_RenderTargetable | TexCreate_ShaderResource, 1, Dimensions.NumSamples);
-		PostMotionBlurColorTexture = CreateTextureMSAA(GraphBuilder, Desc, TEXT("PostMotionBlurTranslucencyColor"), GFastVRamConfig.SeparateTranslucency);
-	}
-	return PostMotionBlurColorTexture;
-}
-
-FRDGTextureRef FSeparateTranslucencyTextures::GetPostMotionBlurColorForRead(FRDGBuilder& GraphBuilder) const
-{
-	if (PostMotionBlurColorTexture.IsValid())
-	{
-		return PostMotionBlurColorTexture.Resolve;
-	}
-	return GraphBuilder.RegisterExternalTexture(GSystemTextures.BlackAlphaOneDummy);
-}
-
-FRDGTextureMSAA FSeparateTranslucencyTextures::GetDepthForWrite(FRDGBuilder& GraphBuilder)
-{
-	if (!DepthTexture.IsValid())
-	{
-		const FRDGTextureDesc Desc = FRDGTextureDesc::Create2D(Dimensions.Extent, PF_DepthStencil, FClearValueBinding::DepthFar, TexCreate_DepthStencilTargetable | TexCreate_ShaderResource, 1, Dimensions.NumSamples);
-		DepthTexture = CreateTextureMSAA(GraphBuilder, Desc, TEXT("SeparateTranslucencyDepth"), GFastVRamConfig.SeparateTranslucencyModulate);
-	}
-	return DepthTexture;
-}
-
-FRDGTextureRef FSeparateTranslucencyTextures::GetDepthForRead(FRDGBuilder& GraphBuilder) const
-{
-	if (DepthTexture.IsValid())
-	{
-		return DepthTexture.Resolve;
+		for (int32 i = 0; i < int32(ETranslucencyPass::TPT_MAX); i++)
+		{
+			Array[ViewIndex][i].Pass = ETranslucencyPass::Type(i);
+		}
 	}
-	return GraphBuilder.RegisterExternalTexture(GSystemTextures.MaxFP16Depth);
 }
-
-FRDGTextureMSAA FSeparateTranslucencyTextures::GetForWrite(FRDGBuilder& GraphBuilder, ETranslucencyPass::Type TranslucencyPass)
-{
-	switch (TranslucencyPass)
-	{
-	default:
-	case ETranslucencyPass::TPT_TranslucencyAfterDOF:
-		return GetColorForWrite(GraphBuilder);
-	case ETranslucencyPass::TPT_TranslucencyAfterDOFModulate:
-		return GetColorModulateForWrite(GraphBuilder);
-	case ETranslucencyPass::TPT_TranslucencyAfterMotionBlur:
-		return GetPostMotionBlurColorForWrite(GraphBuilder);		
-	}	
-}
-
 /** Pixel shader used to copy scene color into another texture so that materials can read from scene color with a node. */
 class FCopySceneColorPS : public FGlobalShader
 {
@@ -552,16 +471,41 @@ static FRDGTextureRef AddCopySceneColorPass(FRDGBuilder& GraphBuilder, TArrayVie
 	return SceneColorCopyTexture;
 }
 
-/** Pixel shader to upsample separate translucency. */
-class FTranslucencyUpsamplePS : public FGlobalShader
+class FComposeSeparateTranslucencyPS : public FGlobalShader
 {
-public:
+	DECLARE_GLOBAL_SHADER(FComposeSeparateTranslucencyPS);
+	SHADER_USE_PARAMETER_STRUCT(FComposeSeparateTranslucencyPS, FGlobalShader);
+
+	class FNearestDepthNeighborUpsampling : SHADER_PERMUTATION_BOOL("PERMUTATION_NEARESTDEPTHNEIGHBOR");
+	using FPermutationDomain = TShaderPermutationDomain<FNearestDepthNeighborUpsampling>;
+
 	BEGIN_SHADER_PARAMETER_STRUCT(FParameters, )
-		SHADER_PARAMETER_STRUCT_REF(FViewUniformShaderParameters, View)
-		SHADER_PARAMETER_RDG_TEXTURE(Texture2D, LowResColorTexture)
-		SHADER_PARAMETER_RDG_TEXTURE_SRV(Texture2D<float>, LowResDepthTexture)
-		SHADER_PARAMETER_RDG_TEXTURE_SRV(Texture2D<float>, FullResDepthTexture)
-		SHADER_PARAMETER(FVector2f, LowResExtentInverse)
+		SHADER_PARAMETER(FScreenTransform, ScreenPosToSceneColorUV)
+		SHADER_PARAMETER(FScreenTransform, ScreenPosToSeparateTranslucencyUV)
+		SHADER_PARAMETER(FVector2f, SeparateTranslucencyUVMin)
+		SHADER_PARAMETER(FVector2f, SeparateTranslucencyUVMax)
+		SHADER_PARAMETER(FVector2f, SeparateTranslucencyExtentInverse)
+
+		SHADER_PARAMETER_RDG_TEXTURE(Texture2D, SceneColorTexture)
+		SHADER_PARAMETER_SAMPLER(SamplerState,  SceneColorSampler)
+
+		SHADER_PARAMETER_RDG_TEXTURE(Texture2D, SeparateTranslucencyPointTexture)
+		SHADER_PARAMETER_SAMPLER(SamplerState,  SeparateTranslucencyPointSampler)
+		SHADER_PARAMETER_RDG_TEXTURE(Texture2D, SeparateModulationPointTexture)
+		SHADER_PARAMETER_SAMPLER(SamplerState,  SeparateModulationPointSampler)
+
+		SHADER_PARAMETER_RDG_TEXTURE(Texture2D, SeparateTranslucencyBilinearTexture)
+		SHADER_PARAMETER_SAMPLER(SamplerState,  SeparateTranslucencyBilinearSampler)
+		SHADER_PARAMETER_RDG_TEXTURE(Texture2D, SeparateModulationBilinearTexture)
+		SHADER_PARAMETER_SAMPLER(SamplerState,  SeparateModulationBilinearSampler)
+
+		SHADER_PARAMETER_RDG_TEXTURE(Texture2D, LowResDepthTexture)
+		SHADER_PARAMETER_SAMPLER(SamplerState, LowResDepthSampler)
+
+		SHADER_PARAMETER_RDG_TEXTURE(Texture2D, FullResDepthTexture)
+		SHADER_PARAMETER_SAMPLER(SamplerState, FullResDepthSampler)
+
+		SHADER_PARAMETER_STRUCT_REF(FViewUniformShaderParameters, ViewUniformBuffer)
 		RENDER_TARGET_BINDING_SLOTS()
 	END_SHADER_PARAMETER_STRUCT()
 
@@ -569,83 +513,252 @@ public:
 	{
 		return IsFeatureLevelSupported(Parameters.Platform, ERHIFeatureLevel::SM5);
 	}
-
-	FTranslucencyUpsamplePS() = default;
-	FTranslucencyUpsamplePS(const FGlobalShaderType::CompiledShaderInitializerType& Initializer)
-		: FGlobalShader(Initializer)
-	{}
 };
 
-class FTranslucencySimpleUpsamplePS : public FTranslucencyUpsamplePS
+class FTranslucencyUpsampleResponsiveAAPS : public FGlobalShader
 {
-protected:
-	DECLARE_GLOBAL_SHADER(FTranslucencySimpleUpsamplePS);
-	SHADER_USE_PARAMETER_STRUCT(FTranslucencySimpleUpsamplePS, FTranslucencyUpsamplePS);
-};
+	DECLARE_GLOBAL_SHADER(FTranslucencyUpsampleResponsiveAAPS);
+	SHADER_USE_PARAMETER_STRUCT(FTranslucencyUpsampleResponsiveAAPS, FGlobalShader);
 
-IMPLEMENT_GLOBAL_SHADER(FTranslucencySimpleUpsamplePS, "/Engine/Private/TranslucencyUpsampling.usf", "SimpleUpsamplingPS", SF_Pixel);
+	BEGIN_SHADER_PARAMETER_STRUCT(FParameters, )
+		SHADER_PARAMETER(FIntPoint, StencilPixelPosMin)
+		SHADER_PARAMETER(FIntPoint, StencilPixelPosMax)
+		SHADER_PARAMETER(FScreenTransform, SvPositionToStencilPixelCoord)
+		SHADER_PARAMETER(int32, StencilMask)
+		SHADER_PARAMETER_RDG_TEXTURE_SRV(Texture2D, StencilTexture)
+		RENDER_TARGET_BINDING_SLOTS()
+	END_SHADER_PARAMETER_STRUCT()
 
-class FTranslucencyNearestDepthNeighborUpsamplePS : public FTranslucencyUpsamplePS
-{
-public:
-	DECLARE_GLOBAL_SHADER(FTranslucencyNearestDepthNeighborUpsamplePS);
-	SHADER_USE_PARAMETER_STRUCT(FTranslucencyNearestDepthNeighborUpsamplePS, FTranslucencyUpsamplePS);
+	static bool ShouldCompilePermutation(const FGlobalShaderPermutationParameters& Parameters)
+	{
+		return IsFeatureLevelSupported(Parameters.Platform, ERHIFeatureLevel::SM5);
+	}
 };
 
-IMPLEMENT_GLOBAL_SHADER(FTranslucencyNearestDepthNeighborUpsamplePS, "/Engine/Private/TranslucencyUpsampling.usf", "NearestDepthNeighborUpsamplingPS", SF_Pixel);
+IMPLEMENT_GLOBAL_SHADER(FComposeSeparateTranslucencyPS, "/Engine/Private/ComposeSeparateTranslucency.usf", "MainPS", SF_Pixel);
+IMPLEMENT_GLOBAL_SHADER(FTranslucencyUpsampleResponsiveAAPS, "/Engine/Private/TranslucencyUpsampling.usf", "UpsampleResponsiveAAPS", SF_Pixel);
 
-bool GetUseTranslucencyNearestDepthNeighborUpsample(float DownsampleScale)
-{
-	const bool bHalfResDownsample = FMath::IsNearlyEqual(DownsampleScale, 0.5f);
-	const bool bUseNearestDepthNeighborUpsample = GSeparateTranslucencyUpsampleMode > 0 && bHalfResDownsample;
-	return bUseNearestDepthNeighborUpsample;
-}
 
-static void AddTranslucencyUpsamplePass(
+FScreenPassTexture FTranslucencyComposition::AddPass(
 	FRDGBuilder& GraphBuilder,
 	const FViewInfo& View,
-	FScreenPassRenderTarget Output,
-	FScreenPassTexture DownsampledTranslucencyColor,
-	FRDGTextureRef DownsampledTranslucencyDepthTexture,
-	FRDGTextureRef SceneDepthTexture,
-	ETranslucencyPass::Type TranslucencyPass,
-	float DownsampleScale)
+	const FTranslucencyPassResources& TranslucencyTextures) const
 {
-	TShaderMapRef<FScreenVS> VertexShader(View.ShaderMap);
-	TShaderRef<FTranslucencyUpsamplePS> PixelShader;
-	FRHIBlendState* BlendState = TStaticBlendState<CW_RGB, BO_Add, BF_One, BF_SourceAlpha>::GetRHI();
+	// if nothing is rendered into the separate translucency, then just return the existing Scenecolor
+	ensure(TranslucencyTextures.IsValid());
+	if (!TranslucencyTextures.IsValid())
+	{
+		return SceneColor;
+	}
+
+	FRDGTextureRef SeparateModulationTexture = TranslucencyTextures.GetColorModulateForRead(GraphBuilder);
+	FRDGTextureRef SeparateTranslucencyTexture = TranslucencyTextures.GetColorForRead(GraphBuilder);
 
-	if (GetUseTranslucencyNearestDepthNeighborUpsample(DownsampleScale))
+	FScreenPassTextureViewport SceneColorViewport(FIntPoint(1, 1), FIntRect(0, 0, 1, 1));
+	if (SceneColor.IsValid())
 	{
-		PixelShader = TShaderMapRef<FTranslucencyNearestDepthNeighborUpsamplePS>(View.ShaderMap);
+		SceneColorViewport = FScreenPassTextureViewport(SceneColor);
+	}
+
+	FScreenPassTextureViewport TranslucencyViewport(FIntPoint(1, 1), FIntRect(0, 0, 1, 1));
+	if (TranslucencyTextures.ColorTexture.IsValid())
+	{
+		TranslucencyViewport = FScreenPassTextureViewport(TranslucencyTextures.ColorTexture.Resolve, TranslucencyTextures.ViewRect);
+	}
+	else if (TranslucencyTextures.ColorModulateTexture.IsValid())
+	{
+		TranslucencyViewport = FScreenPassTextureViewport(TranslucencyTextures.ColorModulateTexture.Resolve, TranslucencyTextures.ViewRect);
+	}
+
+	bool bPostMotionBlur = TranslucencyTextures.Pass == ETranslucencyPass::TPT_TranslucencyAfterMotionBlur;
+	if (bPostMotionBlur)
+	{
+		check(!bApplyModulateOnly);
+	}
+	else if (bApplyModulateOnly)
+	{
+		if (!TranslucencyTextures.ColorModulateTexture.IsValid())
+		{
+			return SceneColor;
+		}
+
+		SeparateTranslucencyTexture = GraphBuilder.RegisterExternalTexture(GSystemTextures.BlackAlphaOneDummy);
+	}
+
+	const TCHAR* OpName = nullptr;
+	FRHIBlendState* BlendState = nullptr;
+	FRDGTextureRef NewSceneColor = nullptr;
+	if (Operation == EOperation::UpscaleOnly)
+	{
+		check(!SceneColor.IsValid());
+		ensure(!TranslucencyTextures.ColorModulateTexture.IsValid());
+
+		OpName = TEXT("UpscaleTranslucency");
+
+		FRDGTextureDesc OutputDesc = FRDGTextureDesc::Create2D(
+			OutputViewport.Extent,
+			PF_FloatRGBA,
+			FClearValueBinding::Black,
+			TexCreate_RenderTargetable | TexCreate_ShaderResource);
+
+		NewSceneColor = GraphBuilder.CreateTexture(
+			OutputDesc,
+			bPostMotionBlur ? TEXT("PostMotionBlurTranslucency.SceneColor") : TEXT("PostDOFTranslucency.SceneColor"));
+	}
+	else if (Operation == EOperation::ComposeToExistingSceneColor)
+	{
+		check(SceneColor.IsValid());
+		ensure(!TranslucencyTextures.ColorModulateTexture.IsValid());
+
+		OpName = TEXT("ComposeTranslucencyToExistingColor");
+		BlendState = TStaticBlendState<CW_RGB, BO_Add, BF_One, BF_SourceAlpha>::GetRHI();
+
+		ensure(SceneColor.Texture->Desc.Flags & TexCreate_RenderTargetable);
+		NewSceneColor = SceneColor.Texture;
+	}
+	else if (Operation == EOperation::ComposeToNewSceneColor)
+	{
+		check(SceneColor.IsValid());
+
+		OpName = TEXT("ComposeTranslucencyToNewSceneColor");
+
+		FRDGTextureDesc OutputDesc = FRDGTextureDesc::Create2D(
+			OutputViewport.Extent,
+			PF_FloatRGB,
+			FClearValueBinding::Black,
+			TexCreate_RenderTargetable | TexCreate_ShaderResource);
+
+		NewSceneColor = GraphBuilder.CreateTexture(
+			OutputDesc,
+			bPostMotionBlur ? TEXT("PostMotionBlurTranslucency.SceneColor") : TEXT("PostDOFTranslucency.SceneColor"));
 	}
 	else
 	{
-		PixelShader = TShaderMapRef<FTranslucencySimpleUpsamplePS>(View.ShaderMap);
+		unimplemented();
 	}
 
-	const FScreenPassTextureViewport OutputViewport(Output);
-	const FScreenPassTextureViewport InputViewport(DownsampledTranslucencyColor);
-	const FIntPoint LowResExtent = DownsampledTranslucencyColor.Texture->Desc.Extent;
+	const FVector2f SeparateTranslucencyExtentInv = FVector2f(1.0f, 1.0f) / FVector2f(TranslucencyViewport.Extent);
 
-	auto* PassParameters = GraphBuilder.AllocParameters<FTranslucencyUpsamplePS::FParameters>();
-	PassParameters->View = View.ViewUniformBuffer;
-	if (TranslucencyPass != ETranslucencyPass::TPT_TranslucencyAfterMotionBlur)
+	const bool bScaleSeparateTranslucency = OutputViewport.Rect.Size() != TranslucencyTextures.ViewRect.Size();
+	const float DownsampleScale = float(TranslucencyTextures.ViewRect.Width()) / float(OutputViewport.Rect.Width());
+	const bool DepthUpscampling = (
+		bScaleSeparateTranslucency &&
+		TranslucencyTextures.DepthTexture.IsValid() &&
+		SceneDepth.IsValid() && 
+		FMath::IsNearlyEqual(DownsampleScale, 0.5f) &&
+		GSeparateTranslucencyUpsampleMode > 0);
+
+	FScreenTransform SvPositionToViewportUV = FScreenTransform::SvPositionToViewportUV(OutputViewport.Rect);
+
+	FComposeSeparateTranslucencyPS::FParameters* PassParameters = GraphBuilder.AllocParameters<FComposeSeparateTranslucencyPS::FParameters>();
+	PassParameters->ScreenPosToSceneColorUV = SvPositionToViewportUV * FScreenTransform::ChangeTextureBasisFromTo(
+		SceneColorViewport, FScreenTransform::ETextureBasis::ViewportUV, FScreenTransform::ETextureBasis::TextureUV);
+	PassParameters->ScreenPosToSeparateTranslucencyUV = SvPositionToViewportUV * FScreenTransform::ChangeTextureBasisFromTo(
+		TranslucencyViewport, FScreenTransform::ETextureBasis::ViewportUV, FScreenTransform::ETextureBasis::TextureUV);
+
+	PassParameters->SeparateTranslucencyUVMin = (FVector2f(TranslucencyViewport.Rect.Min) + FVector2f(0.5f, 0.5f)) * SeparateTranslucencyExtentInv;
+	PassParameters->SeparateTranslucencyUVMax = (FVector2f(TranslucencyViewport.Rect.Max) - FVector2f(0.5f, 0.5f)) * SeparateTranslucencyExtentInv;
+	PassParameters->SeparateTranslucencyExtentInverse = SeparateTranslucencyExtentInv;
+	
+	PassParameters->SceneColorTexture = Operation == EOperation::ComposeToNewSceneColor
+		? SceneColor.Texture
+		: GraphBuilder.RegisterExternalTexture(GSystemTextures.BlackAlphaOneDummy);
+	PassParameters->SceneColorSampler = TStaticSamplerState<SF_Point>::GetRHI();
+
+	PassParameters->SeparateTranslucencyPointTexture = SeparateTranslucencyTexture;
+	PassParameters->SeparateTranslucencyPointSampler = TStaticSamplerState<SF_Point>::GetRHI();
+	
+	PassParameters->SeparateModulationPointTexture = SeparateModulationTexture;
+	PassParameters->SeparateModulationPointSampler = TStaticSamplerState<SF_Point>::GetRHI();
+
+	PassParameters->SeparateTranslucencyBilinearTexture = SeparateTranslucencyTexture;
+	PassParameters->SeparateTranslucencyBilinearSampler = TStaticSamplerState<SF_Bilinear>::GetRHI();
+
+	PassParameters->SeparateModulationBilinearTexture = SeparateModulationTexture;
+	PassParameters->SeparateModulationBilinearSampler = TStaticSamplerState<SF_Bilinear>::GetRHI();
+
+	PassParameters->ViewUniformBuffer = View.ViewUniformBuffer;
+
+	if (Operation == EOperation::ComposeToExistingSceneColor)
 	{
-		PassParameters->FullResDepthTexture = GraphBuilder.CreateSRV(FRDGTextureSRVDesc::CreateForMetaData(SceneDepthTexture, ERDGTextureMetaDataAccess::Depth));
-		PassParameters->LowResDepthTexture = GraphBuilder.CreateSRV(FRDGTextureSRVDesc::CreateForMetaData(DownsampledTranslucencyDepthTexture, ERDGTextureMetaDataAccess::Depth));
+		PassParameters->RenderTargets[0] = FRenderTargetBinding(NewSceneColor, ERenderTargetLoadAction::ELoad);
 	}
 	else
 	{
-		// no depth test in post-motionblur transparents
-		const FRDGSystemTextures& SystemTextures = FRDGSystemTextures::Get(GraphBuilder);
-		PassParameters->FullResDepthTexture = PassParameters->LowResDepthTexture = GraphBuilder.CreateSRV(FRDGTextureSRVDesc::Create(SystemTextures.White));
+		PassParameters->RenderTargets[0] = FRenderTargetBinding(NewSceneColor, ERenderTargetLoadAction::ENoAction);
 	}
-	PassParameters->LowResColorTexture = DownsampledTranslucencyColor.Texture;
-	PassParameters->LowResExtentInverse = FVector2D(1.0f / LowResExtent.X, 1.0f / LowResExtent.Y);
-	PassParameters->RenderTargets[0] = Output.GetRenderTargetBinding();
 
-	AddDrawScreenPass(GraphBuilder, RDG_EVENT_NAME("TranslucencyUpsample"), View, OutputViewport, InputViewport, VertexShader, PixelShader, BlendState, PassParameters);
+	if (DepthUpscampling)
+	{
+		PassParameters->LowResDepthTexture = TranslucencyTextures.GetDepthForRead(GraphBuilder);
+		PassParameters->LowResDepthSampler = TStaticSamplerState<SF_Point>::GetRHI();
+		PassParameters->FullResDepthTexture = SceneDepth.Texture;
+		PassParameters->FullResDepthSampler = TStaticSamplerState<SF_Point>::GetRHI();
+	}
+
+	FComposeSeparateTranslucencyPS::FPermutationDomain PermutationVector;
+	PermutationVector.Set<FComposeSeparateTranslucencyPS::FNearestDepthNeighborUpsampling>(DepthUpscampling);
+
+	TShaderMapRef<FComposeSeparateTranslucencyPS> PixelShader(View.ShaderMap, PermutationVector);
+	FPixelShaderUtils::AddFullscreenPass(
+		GraphBuilder,
+		View.ShaderMap,
+		RDG_EVENT_NAME(
+			"%s(%s%s%s) %dx%d -> %dx%d",
+			OpName,
+			kTranslucencyPassName[int32(TranslucencyTextures.Pass)],
+			bApplyModulateOnly ? TEXT(" ModulateOnly") : TEXT(""),
+			DepthUpscampling ? TEXT(" DepthUpscampling") : TEXT(""),
+			TranslucencyTextures.ViewRect.Width(), TranslucencyTextures.ViewRect.Height(),
+			OutputViewport.Rect.Width(), OutputViewport.Rect.Height()),
+		PixelShader,
+		PassParameters,
+		OutputViewport.Rect,
+		BlendState);
+
+	return FScreenPassTexture(NewSceneColor, OutputViewport.Rect);
+}
+
+static void AddUpsampleResponsiveAAPass(
+	FRDGBuilder& GraphBuilder,
+	const FViewInfo& View,
+	FScreenPassTexture DownsampledTranslucencyDepth,
+	FRDGTextureRef OutputDepthTexture)
+{
+	FTranslucencyUpsampleResponsiveAAPS::FParameters* PassParameters = GraphBuilder.AllocParameters<FTranslucencyUpsampleResponsiveAAPS::FParameters>();
+	PassParameters->StencilPixelPosMin = DownsampledTranslucencyDepth.ViewRect.Min;
+	PassParameters->StencilPixelPosMax = DownsampledTranslucencyDepth.ViewRect.Max - 1;
+	PassParameters->SvPositionToStencilPixelCoord = (FScreenTransform::Identity - View.ViewRect.Min) * (FVector2f(DownsampledTranslucencyDepth.ViewRect.Size()) / FVector2D(View.ViewRect.Size())) + DownsampledTranslucencyDepth.ViewRect.Min;
+	PassParameters->StencilMask = STENCIL_TEMPORAL_RESPONSIVE_AA_MASK;
+	PassParameters->StencilTexture = GraphBuilder.CreateSRV(FRDGTextureSRVDesc::CreateWithPixelFormat(DownsampledTranslucencyDepth.Texture, PF_X24_G8));
+	PassParameters->RenderTargets.DepthStencil = FDepthStencilBinding(OutputDepthTexture, ERenderTargetLoadAction::ELoad, ERenderTargetLoadAction::ELoad, FExclusiveDepthStencil::DepthNop_StencilWrite);
+
+	TShaderMapRef<FScreenVS> VertexShader(View.ShaderMap);
+	TShaderMapRef<FTranslucencyUpsampleResponsiveAAPS> PixelShader(View.ShaderMap);
+
+	FRHIDepthStencilState* DepthStencilState = TStaticDepthStencilState<false, CF_Always,
+		true, CF_Always, SO_Keep, SO_Keep, SO_Replace,
+		false, CF_Always, SO_Keep, SO_Keep, SO_Keep,
+		0x00, STENCIL_TEMPORAL_RESPONSIVE_AA_MASK>::GetRHI();
+	FRHIBlendState* BlendState = TStaticBlendState<CW_NONE>::GetRHI();
+
+	const FScreenPassPipelineState PipelineState(VertexShader, PixelShader, BlendState, DepthStencilState, /* StencilRef = */ STENCIL_TEMPORAL_RESPONSIVE_AA_MASK);
+
+	ClearUnusedGraphResources(PixelShader, PassParameters);
+	GraphBuilder.AddPass(
+		RDG_EVENT_NAME("UpsampleResponsiveAA %dx%d -> %dx%d",
+			DownsampledTranslucencyDepth.ViewRect.Width(), DownsampledTranslucencyDepth.ViewRect.Height(),
+			View.ViewRect.Width(), View.ViewRect.Height()),
+		PassParameters,
+		ERDGPassFlags::Raster,
+		[&View, PipelineState, PixelShader, PassParameters](FRHICommandList& RHICmdList)
+	{
+		FScreenPassTextureViewport OutputViewport(PassParameters->RenderTargets.DepthStencil.GetTexture()->Desc.Extent, View.ViewRect);
+		DrawScreenPass(RHICmdList, View, OutputViewport, OutputViewport, PipelineState, EScreenPassDrawFlags::None, [&](FRHICommandList&)
+		{
+			SetShaderParameters(RHICmdList, PixelShader, PixelShader.GetPixelShader(), *PassParameters);
+		});
+	});
 }
 
 bool FSceneRenderer::ShouldRenderTranslucency() const
@@ -811,7 +924,7 @@ TRDGUniformBufferRef<FTranslucentBasePassUniformParameters> CreateTranslucentBas
 			BasePassParameters.PrevSceneColorBilinearUVMin = PrevSceneColorParameters.UVViewportBilinearMin;
 			BasePassParameters.PrevSceneColorBilinearUVMax = PrevSceneColorParameters.UVViewportBilinearMax;
 
-			const FVector2D HZBUvFactor(
+			const FVector2f HZBUvFactor(
 				float(View.ViewRect.Width()) / float(2 * View.HZBMipmap0Size.X),
 				float(View.ViewRect.Height()) / float(2 * View.HZBMipmap0Size.Y)
 			);
@@ -890,7 +1003,7 @@ TRDGUniformBufferRef<FTranslucentBasePassUniformParameters> CreateTranslucentBas
 			EffectiveBufferSize = View.PrevViewInfo.ScreenSpaceRayTracingInput->GetDesc().Extent;
 		}
 
-		FVector2D InvBufferSize(1.0f / float(EffectiveBufferSize.X), 1.0f / float(EffectiveBufferSize.Y));
+		FVector2f InvBufferSize(1.0f / float(EffectiveBufferSize.X), 1.0f / float(EffectiveBufferSize.Y));
 
 		FVector4f ScreenPosToPixelValue(
 			ViewportExtent.X * 0.5f * InvBufferSize.X,
@@ -1137,7 +1250,10 @@ static void RenderTranslucencyViewInner(
 	if (bRenderInParallel)
 	{
 		GraphBuilder.AddPass(
-			RDG_EVENT_NAME("SeparateTranslucencyParallel"),
+			RDG_EVENT_NAME("Translucency(%s Parallel) %dx%d",
+				TranslucencyPassToString(TranslucencyPass),
+				int32(View.ViewRect.Width() * ViewportScale),
+				int32(View.ViewRect.Height() * ViewportScale)),
 			PassParameters,
 			ERDGPassFlags::Raster | ERDGPassFlags::SkipRenderPass,
 			[&SceneRenderer, &View, PassParameters, ViewportScale, Viewport, TranslucencyPass](FRHICommandListImmediate& RHICmdList)
@@ -1149,7 +1265,10 @@ static void RenderTranslucencyViewInner(
 	else
 	{
 		GraphBuilder.AddPass(
-			RDG_EVENT_NAME("SeparateTranslucency"),
+			RDG_EVENT_NAME("Translucency(%s) %dx%d",
+				TranslucencyPassToString(TranslucencyPass),
+				int32(View.ViewRect.Width() * ViewportScale),
+				int32(View.ViewRect.Height() * ViewportScale)),
 			PassParameters,
 			ERDGPassFlags::Raster,
 			[&SceneRenderer, &View, ViewportScale, Viewport, TranslucencyPass, PassParameters](FRHICommandListImmediate& RHICmdList)
@@ -1168,7 +1287,8 @@ void FDeferredShadingSceneRenderer::RenderTranslucencyInner(
 	FRDGBuilder& GraphBuilder,
 	const FMinimalSceneTextures& SceneTextures,
 	const FTranslucencyLightingVolumeTextures& TranslucentLightingVolumeTextures,
-	FSeparateTranslucencyTextures* OutSeparateTranslucencyTextures,
+	FTranslucencyPassResourcesMap* OutTranslucencyResourceMap,
+	FRDGTextureMSAA SharedDepthTexture,
 	ETranslucencyView ViewsToRender,
 	FRDGTextureRef SceneColorCopyTexture,
 	ETranslucencyPass::Type TranslucencyPass,
@@ -1183,7 +1303,10 @@ void FDeferredShadingSceneRenderer::RenderTranslucencyInner(
 	RDG_GPU_STAT_SCOPE(GraphBuilder, Translucency);
 	RDG_WAIT_FOR_TASKS_CONDITIONAL(GraphBuilder, IsTranslucencyWaitForTasksEnabled());
 
+	const bool bIsModulate = TranslucencyPass == ETranslucencyPass::TPT_TranslucencyAfterDOFModulate;
+	const bool bDepthTest = TranslucencyPass != ETranslucencyPass::TPT_TranslucencyAfterMotionBlur;
 	const bool bRenderInParallel = IsParallelTranslucencyEnabled();
+	const bool bIsScalingTranslucency = SeparateTranslucencyDimensions.Scale < 1.0f;
 	const bool bRenderInSeparateTranslucency = IsSeparateTranslucencyEnabled(TranslucencyPass, SeparateTranslucencyDimensions.Scale);
 
 	const auto ShouldRenderView = [&](const FViewInfo& View, ETranslucencyView TranslucencyView)
@@ -1198,9 +1321,32 @@ void FDeferredShadingSceneRenderer::RenderTranslucencyInner(
 	if (bRenderInSeparateTranslucency)
 	{
 		// Create resources shared by each view (each view data is tiled into each of the render target resources)
-		FSeparateTranslucencyTextures LocalSeparateTranslucencyTextures(SeparateTranslucencyDimensions);
+		FRDGTextureMSAA SharedColorTexture;
+		{
+			static const TCHAR* kTranslucencyColorTextureName[] = {
+				TEXT("Translucency.BeforeDistortion.Color"),
+				TEXT("Translucency.AfterDOF.Color"),
+				TEXT("Translucency.AfterDOF.Modulate"),
+				TEXT("Translucency.AfterMotionBlur.Color"),
+				TEXT("Translucency.All.Color"),
+			};
+			static_assert(UE_ARRAY_COUNT(kTranslucencyColorTextureName) == int32(ETranslucencyPass::TPT_MAX), "Fix me");
+
+			const FRDGTextureDesc Desc = FRDGTextureDesc::Create2D(
+				SeparateTranslucencyDimensions.Extent,
+				bIsModulate ? PF_FloatR11G11B10 : PF_FloatRGBA,
+				bIsModulate ? FClearValueBinding::White : FClearValueBinding::Black,
+				TexCreate_RenderTargetable | TexCreate_ShaderResource,
+				1,
+				SeparateTranslucencyDimensions.NumSamples);
+
+			SharedColorTexture = CreateTextureMSAA(
+				GraphBuilder, Desc,
+				kTranslucencyColorTextureName[int32(TranslucencyPass)],
+				bIsModulate ? GFastVRamConfig.SeparateTranslucencyModulate : GFastVRamConfig.SeparateTranslucency);
+		}
 
-		for (int32 ViewIndex = 0, NumProcessedViews = 0; ViewIndex < Views.Num(); ++ViewIndex)
+		for (int32 ViewIndex = 0, NumProcessedViews = 0; ViewIndex < Views.Num(); ++ViewIndex, ++NumProcessedViews)
 		{
 			FViewInfo& View = Views[ViewIndex];
 			const ETranslucencyView TranslucencyView = GetTranslucencyView(View);
@@ -1213,49 +1359,22 @@ void FDeferredShadingSceneRenderer::RenderTranslucencyInner(
 			RDG_GPU_MASK_SCOPE(GraphBuilder, View.GPUMask);
 			RDG_EVENT_SCOPE_CONDITIONAL(GraphBuilder, Views.Num() > 1, "View%d", ViewIndex);
 
+			FIntRect ScaledViewRect = GetScaledRect(View.ViewRect, SeparateTranslucencyDimensions.Scale);
+
 			const FScreenPassTextureViewport SeparateTranslucencyViewport = SeparateTranslucencyDimensions.GetInstancedStereoViewport(View, InstancedStereoWidth);
 			const bool bCompositeBackToSceneColor = IsMainTranslucencyPass(TranslucencyPass) || EnumHasAnyFlags(TranslucencyView, ETranslucencyView::UnderWater);
-			checkf(bCompositeBackToSceneColor || OutSeparateTranslucencyTextures, TEXT("OutSeparateTranslucencyTextures is null, but we aren't compositing immediately back to scene color."));
-
+			
 			/** Separate translucency color is either composited immediately or later during post processing. If done immediately, it's because the view doesn't support
 			 *  compositing (e.g. we're rendering an underwater view) or because we're downsampling the main translucency pass. In this case, we use a local set of
 			 *  textures instead of the external ones passed in.
 			 */
-			FRDGTextureMSAA SeparateTranslucencyColorTexture;
-			if (bCompositeBackToSceneColor)
-			{
-				SeparateTranslucencyColorTexture = LocalSeparateTranslucencyTextures.GetForWrite(GraphBuilder, TranslucencyPass);
-			}
-			else
-			{
-				SeparateTranslucencyColorTexture = OutSeparateTranslucencyTextures->GetForWrite(GraphBuilder, TranslucencyPass);
-			}
+			FRDGTextureMSAA SeparateTranslucencyColorTexture = SharedColorTexture;
 
-			// NOTE: No depth read on post-motionblur translucency
+			// NOTE: No depth test on post-motionblur translucency
 			FRDGTextureMSAA SeparateTranslucencyDepthTexture;
-			if (TranslucencyPass != ETranslucencyPass::TPT_TranslucencyAfterMotionBlur)
+			if (bDepthTest)
 			{
-				// When rendering to a 1-to-1 separate translucency target, use the existing scene depth.
-				SeparateTranslucencyDepthTexture = SceneTextures.Depth;
-
-				// Rendering to a downscaled target; allocate a new depth texture and downsample depth.
-				if (SeparateTranslucencyDimensions.Scale < 1.0f)
-				{
-					if (bCompositeBackToSceneColor)
-					{
-						SeparateTranslucencyDepthTexture = LocalSeparateTranslucencyTextures.GetDepthForWrite(GraphBuilder); 
-					}
-					else
-					{
-						SeparateTranslucencyDepthTexture = OutSeparateTranslucencyTextures->GetDepthForWrite(GraphBuilder);
-					}
-
-					AddDownsampleDepthPass(
-						GraphBuilder, View,
-						FScreenPassTexture(SceneTextures.Depth.Resolve, View.ViewRect),
-						FScreenPassRenderTarget(SeparateTranslucencyDepthTexture.Target, SeparateTranslucencyViewport.Rect, ERenderTargetLoadAction::ENoAction),
-						EDownsampleDepthFilter::Point);
-				}
+				SeparateTranslucencyDepthTexture = SharedDepthTexture;
 			}
 
 			AddBeginSeparateTranslucencyTimerPass(GraphBuilder, View, TranslucencyPass);
@@ -1279,6 +1398,13 @@ void FDeferredShadingSceneRenderer::RenderTranslucencyInner(
 				bRenderInParallel,
 				InstanceCullingManager);
 
+			{
+				FTranslucencyPassResources& TranslucencyPassResources = OutTranslucencyResourceMap->Get(ViewIndex, TranslucencyPass);
+				TranslucencyPassResources.ViewRect = ScaledViewRect;
+				TranslucencyPassResources.ColorTexture = SharedColorTexture;
+				TranslucencyPassResources.DepthTexture = SharedDepthTexture;
+			}
+
 			if (bCompositeBackToSceneColor)
 			{
 				FRDGTextureRef SeparateTranslucencyDepthResolve = nullptr;
@@ -1291,14 +1417,36 @@ void FDeferredShadingSceneRenderer::RenderTranslucencyInner(
 					SceneDepthResolve = SceneTextures.Depth.Resolve;
 				}
 
-				AddTranslucencyUpsamplePass(
-					GraphBuilder, View,
-					FScreenPassRenderTarget(SceneTextures.Color.Target, View.ViewRect, ERenderTargetLoadAction::ELoad),
-					FScreenPassTexture(SeparateTranslucencyColorTexture.Resolve, SeparateTranslucencyViewport.Rect),
-					SeparateTranslucencyDepthResolve,
-					SceneDepthResolve,
-					TranslucencyPass,
-					SeparateTranslucencyDimensions.Scale);
+				FTranslucencyPassResources& TranslucencyPassResources = OutTranslucencyResourceMap->Get(ViewIndex, TranslucencyPass);
+
+				FTranslucencyComposition TranslucencyComposition;
+				TranslucencyComposition.Operation = FTranslucencyComposition::EOperation::ComposeToExistingSceneColor;
+				TranslucencyComposition.SceneColor = FScreenPassTexture(SceneTextures.Color.Target, View.ViewRect);
+				TranslucencyComposition.SceneDepth = FScreenPassTexture(SceneTextures.Depth.Resolve, View.ViewRect);
+				TranslucencyComposition.OutputViewport = FScreenPassTextureViewport(SceneTextures.Depth.Resolve, View.ViewRect);
+
+				FScreenPassTexture UpscaledTranslucency = TranslucencyComposition.AddPass(
+					GraphBuilder, View, TranslucencyPassResources);
+
+				ensure(View.ViewRect == UpscaledTranslucency.ViewRect);
+				ensure(UpscaledTranslucency.Texture == SceneTextures.Color.Target);
+
+				//Invalidate.
+				TranslucencyPassResources = FTranslucencyPassResources();
+			}
+			else
+			{
+				if (TranslucencyPass == ETranslucencyPass::TPT_TranslucencyAfterDOFModulate)
+				{
+					FTranslucencyPassResources& TranslucencyPassResources = OutTranslucencyResourceMap->Get(ViewIndex, ETranslucencyPass::TPT_TranslucencyAfterDOF);
+					ensure(TranslucencyPassResources.ViewRect == ScaledViewRect);
+					ensure(TranslucencyPassResources.DepthTexture == SharedDepthTexture);
+					TranslucencyPassResources.ColorModulateTexture = SharedColorTexture;
+				}
+				else
+				{
+					check(!bIsModulate);
+				}
 			}
 
 			AddEndSeparateTranslucencyTimerPass(GraphBuilder, View, TranslucencyPass);
@@ -1351,7 +1499,7 @@ void FDeferredShadingSceneRenderer::RenderTranslucency(
 	FRDGBuilder& GraphBuilder,
 	const FMinimalSceneTextures& SceneTextures,
 	const FTranslucencyLightingVolumeTextures& TranslucentLightingVolumeTextures,
-	FSeparateTranslucencyTextures* OutSeparateTranslucencyTextures,
+	FTranslucencyPassResourcesMap* OutTranslucencyResourceMap,
 	ETranslucencyView ViewsToRender,
 	FInstanceCullingManager& InstanceCullingManager)
 {
@@ -1360,8 +1508,6 @@ void FDeferredShadingSceneRenderer::RenderTranslucency(
 		return;
 	}
 
-	RDG_EVENT_SCOPE(GraphBuilder, "Translucency");
-
 	FRDGTextureRef SceneColorCopyTexture = nullptr;
 
 	if (EnumHasAnyFlags(ViewsToRender, ETranslucencyView::AboveWater))
@@ -1369,19 +1515,123 @@ void FDeferredShadingSceneRenderer::RenderTranslucency(
 		SceneColorCopyTexture = AddCopySceneColorPass(GraphBuilder, Views, SceneTextures.Color);
 	}
 
+	const auto ShouldRenderView = [&](const FViewInfo& View, ETranslucencyView TranslucencyView)
+	{
+		return View.ShouldRenderView() && EnumHasAnyFlags(TranslucencyView, ViewsToRender);
+	};
+
+	// Create a shared depth texture at the correct resolution.
+	FRDGTextureMSAA SharedDepthTexture;
+	const bool bIsScalingTranslucency = SeparateTranslucencyDimensions.Scale < 1.0f;
+	if (bIsScalingTranslucency)
+	{
+		const FRDGTextureDesc Desc = FRDGTextureDesc::Create2D(
+			SeparateTranslucencyDimensions.Extent,
+			PF_DepthStencil,
+			FClearValueBinding::DepthFar,
+			TexCreate_DepthStencilTargetable | TexCreate_ShaderResource,
+			1,
+			SeparateTranslucencyDimensions.NumSamples);
+
+		SharedDepthTexture = CreateTextureMSAA(
+			GraphBuilder, Desc,
+			TEXT("Translucency.Depth"),
+			GFastVRamConfig.SeparateTranslucencyModulate); // TODO: this should be SeparateTranslucency, but is what the code was doing
+
+		// Downscale the depth buffer for each individual view, but shared accross all translucencies.
+		for (int32 ViewIndex = 0; ViewIndex < Views.Num(); ++ViewIndex)
+		{
+			FViewInfo& View = Views[ViewIndex];
+			const ETranslucencyView TranslucencyView = GetTranslucencyView(View);
+
+			if (!ShouldRenderView(View, TranslucencyView))
+			{
+				continue;
+			}
+
+			const FScreenPassTextureViewport SeparateTranslucencyViewport = SeparateTranslucencyDimensions.GetInstancedStereoViewport(View, InstancedStereoWidth);
+			AddDownsampleDepthPass(
+				GraphBuilder, View,
+				FScreenPassTexture(SceneTextures.Depth.Resolve, View.ViewRect),
+				FScreenPassRenderTarget(SharedDepthTexture.Target, SeparateTranslucencyViewport.Rect, ViewIndex == 0 ? ERenderTargetLoadAction::EClear : ERenderTargetLoadAction::ELoad),
+				EDownsampleDepthFilter::Point);
+		}
+	}
+	else
+	{
+		// Uses the existing depth buffer for depth testing the translucency.
+		SharedDepthTexture = SceneTextures.Depth;
+	}
+
 	if (ViewFamily.AllowTranslucencyAfterDOF())
 	{
-		RenderTranslucencyInner(GraphBuilder, SceneTextures, TranslucentLightingVolumeTextures, OutSeparateTranslucencyTextures, ViewsToRender, SceneColorCopyTexture, ETranslucencyPass::TPT_StandardTranslucency, InstanceCullingManager);
+		RenderTranslucencyInner(GraphBuilder, SceneTextures, TranslucentLightingVolumeTextures, OutTranslucencyResourceMap, SharedDepthTexture, ViewsToRender, SceneColorCopyTexture, ETranslucencyPass::TPT_StandardTranslucency, InstanceCullingManager);
 		if (GetHairStrandsComposition() == EHairStrandsCompositionType::AfterTranslucentBeforeTranslucentAfterDOF)
 		{
 			RenderHairComposition(GraphBuilder, Views, SceneTextures.Color.Target, SceneTextures.Depth.Target);
 		}
-		RenderTranslucencyInner(GraphBuilder, SceneTextures, TranslucentLightingVolumeTextures, OutSeparateTranslucencyTextures, ViewsToRender, SceneColorCopyTexture, ETranslucencyPass::TPT_TranslucencyAfterDOF, InstanceCullingManager);
-		RenderTranslucencyInner(GraphBuilder, SceneTextures, TranslucentLightingVolumeTextures, OutSeparateTranslucencyTextures, ViewsToRender, SceneColorCopyTexture, ETranslucencyPass::TPT_TranslucencyAfterDOFModulate, InstanceCullingManager);
-		RenderTranslucencyInner(GraphBuilder, SceneTextures, TranslucentLightingVolumeTextures, OutSeparateTranslucencyTextures, ViewsToRender, SceneColorCopyTexture, ETranslucencyPass::TPT_TranslucencyAfterMotionBlur, InstanceCullingManager);
+		RenderTranslucencyInner(GraphBuilder, SceneTextures, TranslucentLightingVolumeTextures, OutTranslucencyResourceMap, SharedDepthTexture, ViewsToRender, SceneColorCopyTexture, ETranslucencyPass::TPT_TranslucencyAfterDOF, InstanceCullingManager);
+		RenderTranslucencyInner(GraphBuilder, SceneTextures, TranslucentLightingVolumeTextures, OutTranslucencyResourceMap, SharedDepthTexture, ViewsToRender, SceneColorCopyTexture, ETranslucencyPass::TPT_TranslucencyAfterDOFModulate, InstanceCullingManager);
+		RenderTranslucencyInner(GraphBuilder, SceneTextures, TranslucentLightingVolumeTextures, OutTranslucencyResourceMap, SharedDepthTexture, ViewsToRender, SceneColorCopyTexture, ETranslucencyPass::TPT_TranslucencyAfterMotionBlur, InstanceCullingManager);
 	}
 	else // Otherwise render translucent primitives in a single bucket.
 	{
-		RenderTranslucencyInner(GraphBuilder, SceneTextures, TranslucentLightingVolumeTextures, OutSeparateTranslucencyTextures, ViewsToRender, SceneColorCopyTexture, ETranslucencyPass::TPT_AllTranslucency, InstanceCullingManager);
+		RenderTranslucencyInner(GraphBuilder, SceneTextures, TranslucentLightingVolumeTextures, OutTranslucencyResourceMap, SharedDepthTexture, ViewsToRender, SceneColorCopyTexture, ETranslucencyPass::TPT_AllTranslucency, InstanceCullingManager);
+	}
+
+	bool bUpscalePostDOFTranslucency = true;
+	FRDGTextureRef SharedUpscaledPostDOFTranslucencyColor = nullptr;
+	if (bUpscalePostDOFTranslucency)
+	{
+		const FRDGTextureDesc Desc = FRDGTextureDesc::Create2D(
+			SceneTextures.Color.Resolve->Desc.Extent,
+			PF_FloatRGBA,
+			FClearValueBinding::Black,
+			TexCreate_RenderTargetable | TexCreate_ShaderResource);
+
+		SharedUpscaledPostDOFTranslucencyColor = GraphBuilder.CreateTexture(
+			Desc, TEXT("Translucency.PostDOF.UpscaledColor"));
+	}
+
+	// Upscale to full res.
+	for (int32 ViewIndex = 0; ViewIndex < Views.Num(); ++ViewIndex)
+	{
+		FViewInfo& View = Views[ViewIndex];
+		const ETranslucencyView TranslucencyView = GetTranslucencyView(View);
+
+		if (!ShouldRenderView(View, TranslucencyView))
+		{
+			continue;
+		}
+
+		// Upscale the responsive AA into original depth buffer.
+		bool bUpscaleResponsiveAA = (
+			IsTemporalAccumulationBasedMethod(View.AntiAliasingMethod) &&
+			SharedDepthTexture.Target != SceneTextures.Depth.Target);
+		if (bUpscaleResponsiveAA)
+		{
+			const FScreenPassTextureViewport SeparateTranslucencyViewport = SeparateTranslucencyDimensions.GetInstancedStereoViewport(View, InstancedStereoWidth);
+			AddUpsampleResponsiveAAPass(
+				GraphBuilder,
+				View,
+				FScreenPassTexture(SharedDepthTexture.Target, SeparateTranslucencyViewport.Rect),
+				/* OutputDepthTexture = */ SceneTextures.Depth.Target);
+		}
+
+		FTranslucencyPassResources& TranslucencyPassResources = OutTranslucencyResourceMap->Get(ViewIndex, ETranslucencyPass::TPT_TranslucencyAfterDOF);
+		if (SharedUpscaledPostDOFTranslucencyColor && TranslucencyPassResources.IsValid() && TranslucencyPassResources.ViewRect.Size() != View.ViewRect.Size())
+		{
+			FTranslucencyComposition TranslucencyComposition;
+			TranslucencyComposition.Operation = FTranslucencyComposition::EOperation::UpscaleOnly;
+			TranslucencyComposition.SceneDepth = FScreenPassTexture(SceneTextures.Depth.Resolve, View.ViewRect);
+			TranslucencyComposition.OutputViewport = FScreenPassTextureViewport(SceneTextures.Depth.Resolve, View.ViewRect);
+
+			FScreenPassTexture UpscaledTranslucency = TranslucencyComposition.AddPass(
+				GraphBuilder, View, TranslucencyPassResources);
+
+			TranslucencyPassResources.ViewRect = UpscaledTranslucency.ViewRect;
+			TranslucencyPassResources.ColorTexture = FRDGTextureMSAA(UpscaledTranslucency.Texture);
+			TranslucencyPassResources.DepthTexture = FRDGTextureMSAA();
+		}
 	}
 }
diff --git a/Engine/Source/Runtime/Renderer/Private/TranslucentRendering.h b/Engine/Source/Runtime/Renderer/Private/TranslucentRendering.h
--- a/Engine/Source/Runtime/Renderer/Private/TranslucentRendering.h
+++ b/Engine/Source/Runtime/Renderer/Private/TranslucentRendering.h
@@ -8,6 +8,7 @@
 #include "SceneRendering.h"
 #include "VolumeRendering.h"
 #include "ScreenPass.h"
+#include "SystemTextures.h"
 
 struct FSeparateTranslucencyDimensions
 {
@@ -28,60 +29,102 @@ struct FSeparateTranslucencyDimensions
 	uint32 NumSamples = 1;
 };
 
-class FSeparateTranslucencyTextures
+/** Resources of a translucency pass. */
+struct FTranslucencyPassResources
 {
-public:
-	FSeparateTranslucencyTextures(FSeparateTranslucencyDimensions InDimensions)
-		: Dimensions(InDimensions)
-	{}
+	ETranslucencyPass::Type Pass = ETranslucencyPass::TPT_MAX;
+	FIntRect ViewRect = FIntRect(0, 0, 0, 0);
+	FRDGTextureMSAA ColorTexture;
+	FRDGTextureMSAA ColorModulateTexture;
+	FRDGTextureMSAA DepthTexture;
 
-	bool IsColorValid() const
+	inline bool IsValid() const
 	{
-		return ColorTexture.IsValid();
+		return ViewRect.Width() > 0 && ViewRect.Height() > 0;
 	}
 
-	FRDGTextureMSAA GetColorForWrite(FRDGBuilder& GraphBuilder);
-	FRDGTextureRef  GetColorForRead(FRDGBuilder& GraphBuilder) const;
+	inline FRDGTextureRef GetColorForRead(FRDGBuilder& GraphBuilder) const
+	{
+		if (!ColorTexture.IsValid())
+		{
+			return GSystemTextures.GetBlackAlphaOneDummy(GraphBuilder);
+		}
+		return ColorTexture.Resolve;
+	}
 
-	bool IsColorModulateValid() const
+	inline FRDGTextureRef GetColorModulateForRead(FRDGBuilder& GraphBuilder) const
 	{
-		return ColorModulateTexture.IsValid();
+		if (!ColorModulateTexture.IsValid())
+		{
+			return GSystemTextures.GetWhiteDummy(GraphBuilder);
+		}
+		return ColorModulateTexture.Resolve;
 	}
 
-	FRDGTextureMSAA GetColorModulateForWrite(FRDGBuilder& GraphBuilder);
-	FRDGTextureRef  GetColorModulateForRead(FRDGBuilder& GraphBuilder) const;
+	inline FRDGTextureRef GetDepthForRead(FRDGBuilder& GraphBuilder) const
+	{
+		if (!DepthTexture.IsValid())
+		{
+			return GSystemTextures.GetMaxFP16Depth(GraphBuilder);
+		}
+		return DepthTexture.Resolve;
+	}
 
-	bool IsPostMotionBlurColorValid() const
+	inline FScreenPassTextureViewport GetTextureViewport() const
 	{
-		return PostMotionBlurColorTexture.IsValid();
+		check(IsValid());
+		return FScreenPassTextureViewport(ColorTexture.Target->Desc.Extent, ViewRect);
 	}
-	
-	FRDGTextureMSAA GetPostMotionBlurColorForWrite(FRDGBuilder& GraphBuilder);
-	FRDGTextureRef  GetPostMotionBlurColorForRead(FRDGBuilder& GraphBuilder) const;
+};
 
-	FRDGTextureMSAA GetDepthForWrite(FRDGBuilder& GraphBuilder);
-	FRDGTextureRef  GetDepthForRead(FRDGBuilder& GraphBuilder) const;
+/** All resources of all translucency passes for a view family. */
+struct FTranslucencyPassResourcesMap
+{
+	FTranslucencyPassResourcesMap(int32 NumViews);
 
-	FRDGTextureMSAA GetForWrite(FRDGBuilder& GraphBuilder, ETranslucencyPass::Type TranslucencyPass);
+	inline FTranslucencyPassResources& Get(int32 ViewIndex, ETranslucencyPass::Type Translucency)
+	{
+		return Array[ViewIndex][int32(Translucency)];
+	};
 
-	const FSeparateTranslucencyDimensions& GetDimensions() const
+	inline const FTranslucencyPassResources& Get(int32 ViewIndex, ETranslucencyPass::Type Translucency) const
 	{
-		return Dimensions;
+		check(ViewIndex < Array.Num());
+		return Array[ViewIndex][int32(Translucency)];
+	};
+
+private:
+	TArray<TStaticArray<FTranslucencyPassResources, ETranslucencyPass::TPT_MAX>, TInlineAllocator<4>> Array;
+};
+
+/** All resources of all translucency for one view. */
+struct FTranslucencyViewResourcesMap
+{
+	FTranslucencyViewResourcesMap() = default;
+
+	FTranslucencyViewResourcesMap(const FTranslucencyPassResourcesMap& InTranslucencyPassResourcesMap, int32 InViewIndex)
+		: TranslucencyPassResourcesMap(&InTranslucencyPassResourcesMap)
+		, ViewIndex(InViewIndex)
+	{ }
+
+	bool IsValid() const
+	{
+		return TranslucencyPassResourcesMap != nullptr;
 	}
 
+	inline const FTranslucencyPassResources& Get(ETranslucencyPass::Type Translucency) const
+	{
+		check(IsValid());
+		return TranslucencyPassResourcesMap->Get(ViewIndex, Translucency);
+	};
+
 private:
-	FSeparateTranslucencyDimensions Dimensions;
-	FRDGTextureMSAA ColorTexture;
-	FRDGTextureMSAA ColorModulateTexture;
-	FRDGTextureMSAA PostMotionBlurColorTexture;
-	FRDGTextureMSAA DepthTexture;
+	const FTranslucencyPassResourcesMap* TranslucencyPassResourcesMap = nullptr;
+	int32 ViewIndex = 0;
 };
 
 DECLARE_GPU_DRAWCALL_STAT_EXTERN(Translucency);
 
-/** Creates separate translucency textures. */
-FSeparateTranslucencyTextures CreateSeparateTranslucencyTextures(FRDGBuilder& GraphBuilder, FSeparateTranslucencyDimensions Dimensions);
-
 /** Converts the the translucency pass into the respective mesh pass. */
 EMeshPass::Type TranslucencyPassToMeshPass(ETranslucencyPass::Type TranslucencyPass);
 
@@ -95,4 +138,29 @@ ETranslucencyView GetTranslucencyViews(TArrayView<const FViewInfo> Views);
 FSeparateTranslucencyDimensions UpdateTranslucencyTimers(FRHICommandListImmediate& RHICmdList, TArrayView<const FViewInfo> Views);
 
 /** Returns whether the view family is requesting to render translucency. */
-bool ShouldRenderTranslucency(const FSceneViewFamily& ViewFamily);
\ No newline at end of file
+bool ShouldRenderTranslucency(const FSceneViewFamily& ViewFamily);
+
+
+/** Add a pass to compose separate translucency. */
+struct FTranslucencyComposition
+{
+	enum class EOperation
+	{
+		UpscaleOnly,
+		ComposeToExistingSceneColor,
+		ComposeToNewSceneColor,
+	};
+
+	EOperation Operation = EOperation::UpscaleOnly;
+	bool bApplyModulateOnly = false;
+
+	FScreenPassTexture SceneColor;
+	FScreenPassTexture SceneDepth;
+
+	FScreenPassTextureViewport OutputViewport;
+
+	FScreenPassTexture AddPass(
+		FRDGBuilder& GraphBuilder,
+		const FViewInfo& View,
+		const FTranslucencyPassResources& TranslucencyTextures) const;
+};
