Enables r.TSR.Translucency.PreviousFrameRejection by default

diff --git a/Engine/Source/Runtime/Renderer/Private/PostProcess/TemporalSuperResolution.cpp b/Engine/Source/Runtime/Renderer/Private/PostProcess/TemporalSuperResolution.cpp
--- a/Engine/Source/Runtime/Renderer/Private/PostProcess/TemporalSuperResolution.cpp
+++ b/Engine/Source/Runtime/Renderer/Private/PostProcess/TemporalSuperResolution.cpp
@@ -55,7 +55,7 @@ TAutoConsoleVariable<float> CVarTSRTranslucencyHighlightLuminance(
 	ECVF_RenderThreadSafe);
 
 TAutoConsoleVariable<int32> CVarTSRTranslucencyPreviousFrameRejection(
-	TEXT("r.TSR.Translucency.PreviousFrameRejection"), 0,
+	TEXT("r.TSR.Translucency.PreviousFrameRejection"), 1,
 	TEXT("Enable heuristic to reject Separate translucency based on previous frame translucency."),
 	ECVF_RenderThreadSafe);
 
