Fix the bug that half resolution SSS brokes the second view and makes it pale

diff --git a/Engine/Source/Runtime/Renderer/Private/PostProcess/PostProcessSubsurface.cpp b/Engine/Source/Runtime/Renderer/Private/PostProcess/PostProcessSubsurface.cpp
--- a/Engine/Source/Runtime/Renderer/Private/PostProcess/PostProcessSubsurface.cpp
+++ b/Engine/Source/Runtime/Renderer/Private/PostProcess/PostProcessSubsurface.cpp
@@ -1303,7 +1303,7 @@ void AddSubsurfaceViewPass(
 			PassParameters->Subsurface = SubsurfaceCommonParameters;
 			if (SubsurfaceMode != ESubsurfaceMode::Bypass)
 			{
-				PassParameters->TileParameters = GetSubsurfaceTileParameters(View, Tiles, TileType, SubsurfaceViewport.Extent);
+				PassParameters->TileParameters = GetSubsurfaceTileParameters(SubsurfaceViewport, Tiles, TileType);
 			}
 			PassParameters->RenderTargets[0] = FRenderTargetBinding(SubsurfaceIntermediateTexture, SceneColorTextureLoadAction);
 			PassParameters->SubsurfaceInput0 = GetSubsurfaceInput(SceneColorTexture, SceneViewportParameters);
@@ -1364,7 +1364,7 @@ void AddSubsurfaceViewPass(
 			PassParameters->Subsurface = SubsurfaceCommonParameters;
 			if (SubsurfaceMode != ESubsurfaceMode::Bypass)
 			{
-				PassParameters->TileParameters = GetSubsurfaceTileParameters(View, Tiles, TileType, SubsurfaceViewport.Extent);
+				PassParameters->TileParameters = GetSubsurfaceTileParameters(SubsurfaceViewport, Tiles, TileType);
 			}
 			PassParameters->RenderTargets[0] = FRenderTargetBinding(SceneColorTexture, SceneColorTextureLoadAction);
 			PassParameters->SubsurfaceInput0 = GetSubsurfaceInput(SubsurfaceIntermediateTexture, SceneViewportParameters);
diff --git a/Engine/Source/Runtime/Renderer/Private/SubsurfaceTiles.cpp b/Engine/Source/Runtime/Renderer/Private/SubsurfaceTiles.cpp
--- a/Engine/Source/Runtime/Renderer/Private/SubsurfaceTiles.cpp
+++ b/Engine/Source/Runtime/Renderer/Private/SubsurfaceTiles.cpp
@@ -22,14 +22,14 @@ bool FSubsurfaceTileFallbackScreenPassVS::ShouldCompilePermutation(const FGlobal
 	return IsFeatureLevelSupported(Parameters.Platform, ERHIFeatureLevel::SM5);
 }
 
-FSubsurfaceTilePassVS::FParameters GetSubsurfaceTileParameters(const FViewInfo& InView,const FSubsurfaceTiles& InTile, FSubsurfaceTiles::ETileType TileType, FIntPoint BufferExtent)
+FSubsurfaceTilePassVS::FParameters GetSubsurfaceTileParameters(const FScreenPassTextureViewport& TileViewport, const FSubsurfaceTiles& InTile, FSubsurfaceTiles::ETileType TileType)
 {
 	FSubsurfaceTilePassVS::FParameters Out;
 	Out.TileType = uint32(TileType);
 	Out.bRectPrimitive = InTile.bRectPrimitive ? 1 : 0;
-	Out.ViewMin = InView.ViewRect.Min;
-	Out.ViewMax = InView.ViewRect.Max;
-	Out.ExtentInverse = FVector2f(1.f / BufferExtent.X, 1.f / BufferExtent.Y);
+	Out.ViewMin = TileViewport.Rect.Min;
+	Out.ViewMax = TileViewport.Rect.Max;
+	Out.ExtentInverse = FVector2f(1.f / TileViewport.Extent.X, 1.f / TileViewport.Extent.Y);
 	Out.TileDataBuffer = InTile.GetTileBufferSRV(TileType);
 	Out.TileIndirectBuffer = InTile.TileIndirectDrawBuffer;
 	return Out;
diff --git a/Engine/Source/Runtime/Renderer/Private/SubsurfaceTiles.h b/Engine/Source/Runtime/Renderer/Private/SubsurfaceTiles.h
--- a/Engine/Source/Runtime/Renderer/Private/SubsurfaceTiles.h
+++ b/Engine/Source/Runtime/Renderer/Private/SubsurfaceTiles.h
@@ -91,7 +91,7 @@ public:
 };
 
 
-FSubsurfaceTilePassVS::FParameters GetSubsurfaceTileParameters(const FViewInfo& InView, const FSubsurfaceTiles& InTile, FSubsurfaceTiles::ETileType TileType, FIntPoint BufferExtent);
+FSubsurfaceTilePassVS::FParameters GetSubsurfaceTileParameters(const FScreenPassTextureViewport& TileViewport, const FSubsurfaceTiles& InTile, FSubsurfaceTiles::ETileType TileType);
 
 template <typename FSubsurfacePassPS, typename FSubsurfacePassVS>
 void AddSubsurfaceTiledScreenPass(
