Fix RVT ortho page rendering to work with large world coords.
Also removed an RVT matrix shader parameter which was completely unused but wouldn't have been LWC friendly.

diff --git a/Engine/Source/Runtime/Engine/Public/SceneView.h b/Engine/Source/Runtime/Engine/Public/SceneView.h
--- a/Engine/Source/Runtime/Engine/Public/SceneView.h
+++ b/Engine/Source/Runtime/Engine/Public/SceneView.h
@@ -760,7 +760,6 @@ enum ETranslucencyVolumeCascade
 	VIEW_UNIFORM_BUFFER_MEMBER(float, StereoIPD) \
 	VIEW_UNIFORM_BUFFER_MEMBER(float, IndirectLightingCacheShowFlag) \
 	VIEW_UNIFORM_BUFFER_MEMBER(float, EyeToPixelSpreadAngle) \
-	VIEW_UNIFORM_BUFFER_MEMBER(FMatrix44f, WorldToVirtualTexture) \
 	VIEW_UNIFORM_BUFFER_MEMBER_ARRAY(FVector4f, XRPassthroughCameraUVs, [2]) \
 	VIEW_UNIFORM_BUFFER_MEMBER(float, GlobalVirtualTextureMipBias) \
 	VIEW_UNIFORM_BUFFER_MEMBER(uint32, VirtualTextureFeedbackShift) \
diff --git a/Engine/Source/Runtime/Renderer/Private/VT/RuntimeVirtualTextureRender.cpp b/Engine/Source/Runtime/Renderer/Private/VT/RuntimeVirtualTextureRender.cpp
--- a/Engine/Source/Runtime/Renderer/Private/VT/RuntimeVirtualTextureRender.cpp
+++ b/Engine/Source/Runtime/Renderer/Private/VT/RuntimeVirtualTextureRender.cpp
@@ -1141,6 +1141,9 @@ namespace RuntimeVirtualTexture
 		const float ZOffset = -NearPlane;
 		ViewInitOptions.ProjectionMatrix = FReversedZOrthoMatrix(OrthoWidth, OrthoHeight, ZScale, ZOffset);
 
+		// Ortho views need to set this to have a ViewOrigin that works with large world coordinates.
+		ViewInitOptions.bUseFauxOrthoViewPos = true;
+
 		const FVector4f MipLevelParameter = FVector4f((float)vLevel, (float)MaxLevel, OrthoWidth / (float)TextureSize.X, OrthoHeight / (float)TextureSize.Y);
 		
 		const float HeightRange = FMath::Max<float>(WorldBounds.Max.Z - WorldBounds.Min.Z, 1.f);
@@ -1156,7 +1159,6 @@ namespace RuntimeVirtualTexture
 		View->ViewRect = View->UnconstrainedViewRect;
 		View->CachedViewUniformShaderParameters = MakeUnique<FViewUniformShaderParameters>();
 		View->SetupUniformBufferParameters(nullptr, 0, *View->CachedViewUniformShaderParameters);
-		View->CachedViewUniformShaderParameters->WorldToVirtualTexture = FMatrix44f(WorldToUVRotate.ToMatrixNoScale());	// LWC_TODO: Precision loss
 		View->CachedViewUniformShaderParameters->RuntimeVirtualTextureMipLevel = MipLevelParameter;
 		View->CachedViewUniformShaderParameters->RuntimeVirtualTexturePackHeight = WorldHeightPackParameter;
 		View->CachedViewUniformShaderParameters->RuntimeVirtualTextureDebugParams = FVector4f(DebugType == ERuntimeVirtualTextureDebugType::Debug ? 1.f : 0.f, 0.f, 0.f, 0.f);
