DisplayCluster perf:  made it so scene occlusion queries are added to a queue based on the UniqueID of the first View in the ViewState, so in scenarios where there are multiple view families rendering, a given view family can wait on its own query fence (next frame), rather than the fence for the previously rendered view family in the current frame.  This avoids an unnecessary stall waiting on results that aren't actually going to be used, saving up to 8 ms in stalls in a Virtual Production test case.  If the queue fills up, one fence is automatically pulled off the queue and waited on, so there is room to push a new fence.

#rnx

diff --git a/Engine/Source/Runtime/Renderer/Private/SceneOcclusion.cpp b/Engine/Source/Runtime/Renderer/Private/SceneOcclusion.cpp
--- a/Engine/Source/Runtime/Renderer/Private/SceneOcclusion.cpp
+++ b/Engine/Source/Runtime/Renderer/Private/SceneOcclusion.cpp
@@ -1464,16 +1464,28 @@ void FMobileSceneRenderer::RenderOcclusion(FRHICommandListImmediate& RHICmdList,
 DECLARE_CYCLE_STAT(TEXT("OcclusionSubmittedFence Dispatch"), STAT_OcclusionSubmittedFence_Dispatch, STATGROUP_SceneRendering);
 DECLARE_CYCLE_STAT(TEXT("OcclusionSubmittedFence Wait"), STAT_OcclusionSubmittedFence_Wait, STATGROUP_SceneRendering);
 
+static uint32 GetViewStateUniqueID(const FSceneRenderer* SceneRenderer)
+{
+	return SceneRenderer->Views.Num() && SceneRenderer->Views[0].ViewState ? SceneRenderer->Views[0].ViewState->UniqueID : 0;
+}
+
 void FSceneRenderer::FenceOcclusionTestsInternal(FRHICommandListImmediate& RHICmdList)
 {
 	SCOPE_CYCLE_COUNTER(STAT_OcclusionSubmittedFence_Dispatch);
-	int32 NumFrames = FOcclusionQueryHelpers::GetNumBufferedFrames(FeatureLevel);
-	for (int32 Dest = NumFrames - 1; Dest >= 1; Dest--)
+
+	// Push a new fence into the queue of buffered fences.  We assume that the queue isn't full, since WaitOcclusionTests
+	// (called earlier in the frame) will always pop at least one fence if the queue is full.
+	check(OcclusionSubmittedFence[FOcclusionQueryHelpers::MaxBufferedOcclusionFrames - 1].Fence == nullptr);
+
+	for (int32 Dest = FOcclusionQueryHelpers::MaxBufferedOcclusionFrames - 1; Dest >= 1; Dest--)
 	{
 		CA_SUPPRESS(6385);
 		OcclusionSubmittedFence[Dest] = OcclusionSubmittedFence[Dest - 1];
 	}
-	OcclusionSubmittedFence[0] = RHICmdList.RHIThreadFence();
+
+	OcclusionSubmittedFence[0].Fence = RHICmdList.RHIThreadFence();
+	OcclusionSubmittedFence[0].ViewStateUniqueID = GetViewStateUniqueID(this);
+
 	RHICmdList.ImmediateFlush(EImmediateFlushType::DispatchToRHIThread);
 	RHICmdList.PollRenderQueryResults();
 }
@@ -1494,8 +1506,49 @@ void FSceneRenderer::WaitOcclusionTests(FRHICommandListImmediate& RHICmdList)
 	if (IsRunningRHIInSeparateThread())
 	{
 		SCOPE_CYCLE_COUNTER(STAT_OcclusionSubmittedFence_Wait);
-		int32 BlockFrame = FOcclusionQueryHelpers::GetNumBufferedFrames(FeatureLevel) - 1;
-		FRHICommandListExecutor::WaitOnRHIThreadFence(OcclusionSubmittedFence[BlockFrame]);
-		OcclusionSubmittedFence[BlockFrame] = nullptr;
+
+		// Count the number of fences in the queue for the view in question
+		const uint32 ViewStateUniqueID = GetViewStateUniqueID(this);
+
+		int32 ViewStateFenceCount = 0;
+		for (int32 FenceIndex = 0; FenceIndex < FOcclusionQueryHelpers::MaxBufferedOcclusionFrames && OcclusionSubmittedFence[FenceIndex].Fence; FenceIndex++)
+		{
+			if (OcclusionSubmittedFence[FenceIndex].Fence && (OcclusionSubmittedFence[FenceIndex].ViewStateUniqueID == ViewStateUniqueID))
+			{
+				ViewStateFenceCount++;
+			}
+		}
+
+		// Figure out how many fences are allowed to remain in the queue after the wait, which is one less than the number of buffered frames
+		const int32 FencesAllowedInQueue = FOcclusionQueryHelpers::GetNumBufferedFrames(FeatureLevel) - 1;
+		check(FencesAllowedInQueue >= 0);
+
+		// Scan the fences in the queue in order (oldest will be at the end), and wait on them as appropriate.  If the queue is full,
+		// we always want to wait on at least one fence, so we have room to push a new fence onto the queue if necessary.  Thus the
+		// loop always runs at least once before breaking out.
+		int32 FenceToWaitOn = FOcclusionQueryHelpers::MaxBufferedOcclusionFrames - 1;
+		do
+		{
+			// If our fence counting logic is correct, we should already have broken out of the loop
+			check(FenceToWaitOn >= 0);
+
+			// Is there a fence in this queue entry?
+			if (OcclusionSubmittedFence[FenceToWaitOn].Fence)
+			{
+				// If this is one of the fences for the current scene's view state, mark that we're waiting on it
+				if (OcclusionSubmittedFence[FenceToWaitOn].ViewStateUniqueID == ViewStateUniqueID)
+				{
+					ViewStateFenceCount--;
+				}
+
+				// Wait on the current fence and clear it out
+				FRHICommandListExecutor::WaitOnRHIThreadFence(OcclusionSubmittedFence[FenceToWaitOn].Fence);
+				OcclusionSubmittedFence[FenceToWaitOn].Fence = nullptr;
+				OcclusionSubmittedFence[FenceToWaitOn].ViewStateUniqueID = 0;
+			}
+
+			FenceToWaitOn--;
+
+		} while (ViewStateFenceCount > FencesAllowedInQueue);
 	}
 }
diff --git a/Engine/Source/Runtime/Renderer/Private/SceneRendering.cpp b/Engine/Source/Runtime/Renderer/Private/SceneRendering.cpp
--- a/Engine/Source/Runtime/Renderer/Private/SceneRendering.cpp
+++ b/Engine/Source/Runtime/Renderer/Private/SceneRendering.cpp
@@ -382,7 +382,7 @@ static TAutoConsoleVariable<int32> CVarTestSecondaryUpscaleOverride(
 #endif
 
 static FParallelCommandListSet* GOutstandingParallelCommandListSet = nullptr;
-FGraphEventRef FSceneRenderer::OcclusionSubmittedFence[FOcclusionQueryHelpers::MaxBufferedOcclusionFrames];
+FOcclusionSubmittedFenceState FSceneRenderer::OcclusionSubmittedFence[FOcclusionQueryHelpers::MaxBufferedOcclusionFrames];
 
 extern int32 GetTranslucencyLightingVolumeDim();
 
diff --git a/Engine/Source/Runtime/Renderer/Private/SceneRendering.h b/Engine/Source/Runtime/Renderer/Private/SceneRendering.h
--- a/Engine/Source/Runtime/Renderer/Private/SceneRendering.h
+++ b/Engine/Source/Runtime/Renderer/Private/SceneRendering.h
@@ -1818,6 +1818,12 @@ struct FSortedShadowMaps
 	}
 };
 
+struct FOcclusionSubmittedFenceState
+{
+	FGraphEventRef	Fence;
+	uint32			ViewStateUniqueID;
+};
+
 /**
  * Used as the scope for scene rendering functions.
  * It is initialized in the game thread by FSceneViewFamily::BeginRender, and then passed to the rendering thread.
@@ -1907,7 +1913,7 @@ public:
 	void WaitOcclusionTests(FRHICommandListImmediate& RHICmdList);
 
 	// fences to make sure the rhi thread has digested the occlusion query renders before we attempt to read them back async
-	static FGraphEventRef OcclusionSubmittedFence[FOcclusionQueryHelpers::MaxBufferedOcclusionFrames];
+	static FOcclusionSubmittedFenceState OcclusionSubmittedFence[FOcclusionQueryHelpers::MaxBufferedOcclusionFrames];
 
 	bool ShouldDumpMeshDrawCommandInstancingStats() const { return bDumpMeshDrawCommandInstancingStats; }
 
