Update translucent light cascade computation to use doubles
Fix translucent light functions

diff --git a/Engine/Shaders/Private/TranslucentLightInjectionShaders.usf b/Engine/Shaders/Private/TranslucentLightInjectionShaders.usf
--- a/Engine/Shaders/Private/TranslucentLightInjectionShaders.usf
+++ b/Engine/Shaders/Private/TranslucentLightInjectionShaders.usf
@@ -108,11 +108,11 @@ float GetLightFunctionShadowFactor(float3 TranslatedWorldPositionForLighting)
 	// Apply light function after edge fading, so that a black light function at the edges can cause distant translucency to also be black
 	#if APPLY_LIGHT_FUNCTION
 		float4 LightVector = mul(float4(TranslatedWorldPositionForLighting, 1),LightFunctionTranslatedWorldToLight);
-	LightVector.xyz /= LightVector.w;
-	float3 LightFunction = GetLightFunctionColor(LightVector.xyz, TranslatedWorldPositionForLighting);
+		LightVector.xyz /= LightVector.w;
+		float3 LightFunction = GetLightFunctionColor(LightVector.xyz, TranslatedWorldPositionForLighting);
 
-	// We only suport monochrome light functions
-	ShadowFactor = dot(LightFunction, .3333f).x;
+		// We only suport monochrome light functions
+		ShadowFactor = dot(LightFunction, .3333f).x;
 	#endif
 
 	return ShadowFactor;
@@ -203,7 +203,7 @@ void InjectMainPS(
 #endif
 
 		// Apply light function (it will also fade out at the edge of the volume so overall dimming of light using light function is not a recommended workflow)
-		ShadowFactor *= GetLightFunctionShadowFactor(WorldPositionForLighting);
+		ShadowFactor *= GetLightFunctionShadowFactor(TranslatedWorldPositionForLighting);
 
 		#if RADIAL_ATTENUATION==0
 		// Apply cloud shadow for atmosphere directional lights only if needed
diff --git a/Engine/Source/Runtime/Renderer/Private/TranslucentLighting.cpp b/Engine/Source/Runtime/Renderer/Private/TranslucentLighting.cpp
--- a/Engine/Source/Runtime/Renderer/Private/TranslucentLighting.cpp
+++ b/Engine/Source/Runtime/Renderer/Private/TranslucentLighting.cpp
@@ -121,36 +121,36 @@ void FViewInfo::CalcTranslucencyLightingVolumeBounds(FBox* InOutCascadeBoundsArr
 {
 	for (int32 CascadeIndex = 0; CascadeIndex < NumCascades; CascadeIndex++)
 	{
-		float InnerDistance = CVarTranslucencyLightingVolumeInnerDistance.GetValueOnRenderThread();
-		float OuterDistance = CVarTranslucencyLightingVolumeOuterDistance.GetValueOnRenderThread();
+		double InnerDistance = CVarTranslucencyLightingVolumeInnerDistance.GetValueOnRenderThread();
+		double OuterDistance = CVarTranslucencyLightingVolumeOuterDistance.GetValueOnRenderThread();
 
-		const float FrustumStartDistance = CascadeIndex == 0 ? 0 : InnerDistance;
-		const float FrustumEndDistance = CascadeIndex == 0 ? InnerDistance : OuterDistance;
+		const double FrustumStartDistance = CascadeIndex == 0 ? 0 : InnerDistance;
+		const double FrustumEndDistance = CascadeIndex == 0 ? InnerDistance : OuterDistance;
 
-		float FieldOfView = PI / 4.0f;
-		float AspectRatio = 1.0f;
+		double FieldOfView = DOUBLE_PI / 4.0;
+		double AspectRatio = 1.0;
 
 		if (IsPerspectiveProjection())
 		{
 			// Derive FOV and aspect ratio from the perspective projection matrix
-			FieldOfView = FMath::Atan(1.0f / ShadowViewMatrices.GetProjectionMatrix().M[0][0]);
+			FieldOfView = FMath::Atan(1.0 / ShadowViewMatrices.GetProjectionMatrix().M[0][0]);
 			// Clamp to prevent shimmering when zooming in
-			FieldOfView = FMath::Max(FieldOfView, GTranslucentVolumeMinFOV * (float)PI / 180.0f);
-			const float RoundFactorRadians = GTranslucentVolumeFOVSnapFactor * (float)PI / 180.0f;
+			FieldOfView = FMath::Max(FieldOfView, GTranslucentVolumeMinFOV * DOUBLE_PI / 180.0);
+			const double RoundFactorRadians = GTranslucentVolumeFOVSnapFactor * DOUBLE_PI / 180.0;
 			// Round up to a fixed factor
 			// This causes the volume lighting to make discreet jumps as the FOV animates, instead of slowly crawling over a long period
 			FieldOfView = FieldOfView + RoundFactorRadians - FMath::Fmod(FieldOfView, RoundFactorRadians);
 			AspectRatio = ShadowViewMatrices.GetProjectionMatrix().M[1][1] / ShadowViewMatrices.GetProjectionMatrix().M[0][0];
 		}
 
-		const float StartHorizontalLength = FrustumStartDistance * FMath::Tan(FieldOfView);
+		const double StartHorizontalLength = FrustumStartDistance * FMath::Tan(FieldOfView);
 		const FVector StartCameraRightOffset = ShadowViewMatrices.GetViewMatrix().GetColumn(0) * StartHorizontalLength;
-		const float StartVerticalLength = StartHorizontalLength / AspectRatio;
+		const double StartVerticalLength = StartHorizontalLength / AspectRatio;
 		const FVector StartCameraUpOffset = ShadowViewMatrices.GetViewMatrix().GetColumn(1) * StartVerticalLength;
 
-		const float EndHorizontalLength = FrustumEndDistance * FMath::Tan(FieldOfView);
+		const double EndHorizontalLength = FrustumEndDistance * FMath::Tan(FieldOfView);
 		const FVector EndCameraRightOffset = ShadowViewMatrices.GetViewMatrix().GetColumn(0) * EndHorizontalLength;
-		const float EndVerticalLength = EndHorizontalLength / AspectRatio;
+		const double EndVerticalLength = EndHorizontalLength / AspectRatio;
 		const FVector EndCameraUpOffset = ShadowViewMatrices.GetViewMatrix().GetColumn(1) * EndVerticalLength;
 
 		FVector SplitVertices[8];
@@ -169,14 +169,14 @@ void FViewInfo::CalcTranslucencyLightingVolumeBounds(FBox* InOutCascadeBoundsArr
 		FVector Center(0,0,0);
 		// Weight the far vertices more so that the bounding sphere will be further from the camera
 		// This minimizes wasted shadowmap space behind the viewer
-		const float FarVertexWeightScale = 10.0f;
+		const double FarVertexWeightScale = 10.0;
 		for (int32 VertexIndex = 0; VertexIndex < 8; VertexIndex++)
 		{
-			const float Weight = VertexIndex > 3 ? 1 / (4.0f + 4.0f / FarVertexWeightScale) : 1 / (4.0f + 4.0f * FarVertexWeightScale);
+			const double Weight = VertexIndex > 3 ? 1 / (4.0 + 4.0 / FarVertexWeightScale) : 1 / (4.0 + 4.0 * FarVertexWeightScale);
 			Center += SplitVertices[VertexIndex] * Weight;
 		}
 
-		float RadiusSquared = 0;
+		double RadiusSquared = 0;
 		for (int32 VertexIndex = 0; VertexIndex < 8; VertexIndex++)
 		{
 			RadiusSquared = FMath::Max(RadiusSquared, (Center - SplitVertices[VertexIndex]).SizeSquared());
