Fixed raytracing scene captures were using the reflection captures visibility flag instead of the scene captures visibility flag

diff --git a/Engine/Source/Runtime/Renderer/Private/DeferredShadingRenderer.cpp b/Engine/Source/Runtime/Renderer/Private/DeferredShadingRenderer.cpp
--- a/Engine/Source/Runtime/Renderer/Private/DeferredShadingRenderer.cpp
+++ b/Engine/Source/Runtime/Renderer/Private/DeferredShadingRenderer.cpp
@@ -637,7 +637,7 @@ bool FDeferredShadingSceneRenderer::GatherRayTracingWorldInstancesForView(FRDGBu
 			bool bShouldRayTraceSceneCapture = GRayTracingSceneCaptures > 0
 				|| (GRayTracingSceneCaptures == -1 && View.bSceneCaptureUsesRayTracing);
 
-			if (View.bIsSceneCapture && (!bShouldRayTraceSceneCapture || !SceneInfo->bIsVisibleInReflectionCaptures))
+			if (View.bIsSceneCapture && (!bShouldRayTraceSceneCapture || !SceneInfo->bIsVisibleInSceneCaptures))
 			{
 				continue;
 			}
diff --git a/Engine/Source/Runtime/Renderer/Private/PrimitiveSceneInfo.cpp b/Engine/Source/Runtime/Renderer/Private/PrimitiveSceneInfo.cpp
--- a/Engine/Source/Runtime/Renderer/Private/PrimitiveSceneInfo.cpp
+++ b/Engine/Source/Runtime/Renderer/Private/PrimitiveSceneInfo.cpp
@@ -190,7 +190,7 @@ FPrimitiveSceneInfo::FPrimitiveSceneInfo(UPrimitiveComponent* InComponent,FScene
 	bVisibleInRealTimeSkyCapture(InComponent->SceneProxy->IsVisibleInRealTimeSkyCaptures()),
 #if RHI_RAYTRACING
 	bDrawInGame(Proxy->IsDrawnInGame()),
-	bIsVisibleInReflectionCaptures(InComponent->SceneProxy->IsVisibleInReflectionCaptures()),
+	bIsVisibleInSceneCaptures(!InComponent->SceneProxy->IsHiddenInSceneCapture()),
 	bIsRayTracingRelevant(InComponent->SceneProxy->IsRayTracingRelevant()),
 	bIsRayTracingStaticRelevant(InComponent->SceneProxy->IsRayTracingStaticRelevant()),
 	bIsVisibleInRayTracing(InComponent->SceneProxy->IsVisibleInRayTracing()),
diff --git a/Engine/Source/Runtime/Renderer/Public/PrimitiveSceneInfo.h b/Engine/Source/Runtime/Renderer/Public/PrimitiveSceneInfo.h
--- a/Engine/Source/Runtime/Renderer/Public/PrimitiveSceneInfo.h
+++ b/Engine/Source/Runtime/Renderer/Public/PrimitiveSceneInfo.h
@@ -398,7 +398,7 @@ public:
 
 #if RHI_RAYTRACING
 	bool bDrawInGame : 1;
-	bool bIsVisibleInReflectionCaptures : 1;
+	bool bIsVisibleInSceneCaptures : 1;
 	bool bIsRayTracingRelevant : 1;
 	bool bIsRayTracingStaticRelevant : 1;
 	bool bIsVisibleInRayTracing : 1;
