Fix groom voxel adaptive allocation failling to compensate for over-allocation due to rounding error.

This CL also ensures that voxel/DOM debug mode auto-enable ShaderDraw/ShaderPrint automatically.

diff --git a/Engine/Shaders/Private/HairStrands/HairStrandsVoxelPageRayMarching.usf b/Engine/Shaders/Private/HairStrands/HairStrandsVoxelPageRayMarching.usf
--- a/Engine/Shaders/Private/HairStrands/HairStrandsVoxelPageRayMarching.usf
+++ b/Engine/Shaders/Private/HairStrands/HairStrandsVoxelPageRayMarching.usf
@@ -97,8 +97,30 @@ void MainCS(uint3 DispatchThreadId : SV_DispatchThreadID)
 			float2 StartPos = float2(50.0f, 100.0f ) / OutputResolution;
 			float2 Pos = StartPos;
 
+			// Adaptive voxels
+			Pos = ShaderPrintText(Pos, TEXT("Adapative "), FontOrange);
+			if (VirtualVoxel.AdaptiveEnable)
+			{
+				Pos = ShaderPrintText(Pos, TEXT("True"), FontGreen);
+			}
+			else
+			{
+				Pos = ShaderPrintText(Pos, TEXT("False"), FontRed);
+			}
+			Pos = ShaderPrintNewline(Pos); Pos.x = StartPos.x;
+
+			Pos = ShaderPrintText(Pos, TEXT("Requested Voxel Size "));
+			Pos = ShaderPrint(Pos, VirtualVoxel.AdaptiveRequestedVoxelWorldSize);
+			Pos = ShaderPrintNewline(Pos); Pos.x = StartPos.x;
+
+			Pos = ShaderPrintText(Pos, TEXT("Target Voxel Size "));
+			Pos = ShaderPrint(Pos, VirtualVoxel.AdaptiveTargetVoxelWorldSize);
+			Pos = ShaderPrintNewline(Pos); Pos.x = StartPos.x;
+			Pos = ShaderPrintNewline(Pos); Pos.x = StartPos.x;
+
+			// Actual voxel/pages size & resolution
 			Pos = ShaderPrintText(Pos, TEXT("Voxel World Size "), FontOrange);
-			Pos = ShaderPrint(Pos, VirtualVoxel.VoxelWorldSize, FontOrange);
+			Pos = ShaderPrint(Pos, VirtualVoxel.VoxelWorldSize, FontEmerald);
 			Pos = ShaderPrintNewline(Pos); Pos.x = StartPos.x;
 			
 			Pos = ShaderPrintText(Pos, TEXT("Page Res         "));
@@ -111,7 +133,7 @@ void MainCS(uint3 DispatchThreadId : SV_DispatchThreadID)
 			Pos = ShaderPrint(Pos, VirtualVoxel.PageCountResolution.z);
 			Pos = ShaderPrintNewline(Pos); Pos.x = StartPos.x;
 
-			Pos = ShaderPrintText(Pos, TEXT("Page Texture Res  "), FontEmerald);
+			Pos = ShaderPrintText(Pos, TEXT("Page Texture Res "), FontEmerald);
 			Pos = ShaderPrint(Pos, VirtualVoxel.PageTextureResolution.x, FontEmerald);
 			Pos = ShaderPrint(Pos, VirtualVoxel.PageTextureResolution.y, FontEmerald);
 			Pos = ShaderPrint(Pos, VirtualVoxel.PageTextureResolution.z, FontEmerald);
@@ -129,7 +151,7 @@ void MainCS(uint3 DispatchThreadId : SV_DispatchThreadID)
 			Pos = ShaderPrintNewline(Pos); Pos.x = StartPos.x;
 
 			const float PageIndexOccupancy = TotalAllocatedPageCount / float(TotalPageIndexCount);
-			Pos = ShaderPrintText(Pos, TEXT("Page Index Occupancy         "), FontOrange);
+			Pos = ShaderPrintText(Pos, TEXT("Page Index Occupancy        "), FontOrange);
 			Pos = ShaderPrint(Pos, PageIndexOccupancy * 100.f, GetOccupancyColor(PageIndexOccupancy));
 			Pos = ShaderPrintNewline(Pos); Pos.x = StartPos.x;
 
diff --git a/Engine/Source/Runtime/Renderer/Private/HairStrands/HairStrandsData.h b/Engine/Source/Runtime/Renderer/Private/HairStrands/HairStrandsData.h
--- a/Engine/Source/Runtime/Renderer/Private/HairStrands/HairStrandsData.h
+++ b/Engine/Source/Runtime/Renderer/Private/HairStrands/HairStrandsData.h
@@ -179,6 +179,10 @@ BEGIN_SHADER_PARAMETER_STRUCT(FHairStrandsVoxelCommonParameters, )
 
 	SHADER_PARAMETER(FVector3f, TranslatedWorldOffset) // For debug purpose
 
+	SHADER_PARAMETER(uint32, AdaptiveEnable)
+	SHADER_PARAMETER(float, AdaptiveRequestedVoxelWorldSize)
+	SHADER_PARAMETER(float, AdaptiveTargetVoxelWorldSize)
+
 	SHADER_PARAMETER_RDG_BUFFER_SRV(Buffer<uint>, PageIndexBuffer)
 	SHADER_PARAMETER_RDG_BUFFER_SRV(Buffer<uint2>, PageIndexOccupancyBuffer)
 	SHADER_PARAMETER_RDG_BUFFER_SRV(Buffer<uint>, PageIndexCoordBuffer)
diff --git a/Engine/Source/Runtime/Renderer/Private/HairStrands/HairStrandsDebug.cpp b/Engine/Source/Runtime/Renderer/Private/HairStrands/HairStrandsDebug.cpp
--- a/Engine/Source/Runtime/Renderer/Private/HairStrands/HairStrandsDebug.cpp
+++ b/Engine/Source/Runtime/Renderer/Private/HairStrands/HairStrandsDebug.cpp
@@ -71,6 +71,28 @@ static FAutoConsoleVariableRef CVarHairVirtualVoxel_ForceMipLevel(TEXT("r.HairSt
 static int32 GHairVirtualVoxel_DebugTraversalType = 0;
 static FAutoConsoleVariableRef CVarHairVirtualVoxel_DebugTraversalType(TEXT("r.HairStrands.Voxelization.Virtual.DebugTraversalType"), GHairVirtualVoxel_DebugTraversalType, TEXT("Traversal mode (0:linear, 1:mip) for debug voxel visualization."));
 
+static bool TryEnableShaderDrawAndShaderPrint(const FViewInfo& View, uint32 ResquestedShaderDrawElements, uint32 RequestedShaderPrintElements)
+{
+	const EShaderPlatform Platform = View.Family->GetShaderPlatform();
+	if (!ShaderDrawDebug::IsSupported(Platform) || !ShaderPrint::IsSupported(Platform))
+	{
+		return false;
+	}
+
+	if (!ShaderPrint::IsEnabled(View))
+	{
+		ShaderPrint::SetEnabled(true);
+	}
+	ShaderPrint::RequestSpaceForCharacters(RequestedShaderPrintElements);
+
+	if (!ShaderDrawDebug::IsEnabled(View))
+	{
+		ShaderDrawDebug::SetEnabled(true);
+	}
+	ShaderDrawDebug::RequestSpaceForElements(ResquestedShaderDrawElements);
+	return true;
+}
+
 static bool IsDebugDrawAndDebugPrintEnabled(const FViewInfo& View)
 {
 	return ShaderDrawDebug::IsEnabled(View) && ShaderPrint::IsEnabled(View);
@@ -652,6 +674,11 @@ static void AddDeepShadowInfoPass(
 		return;
 	}
 
+	if (!TryEnableShaderDrawAndShaderPrint(View, DeepShadowResources.TotalAtlasSlotCount * 64, 2000))
+	{
+		return;
+	}
+
 	FSceneTextureParameters SceneTextures = GetSceneTextureParameters(GraphBuilder);
 
 	const FIntPoint Resolution(OutputTexture->Desc.Extent);
@@ -715,7 +742,7 @@ static void AddVoxelPageRaymarchingPass(
 	const FHairStrandsVoxelResources& VoxelResources,
 	FRDGTextureRef& OutputTexture)
 {
-	if (!IsDebugDrawAndDebugPrintEnabled(View))
+	if (!TryEnableShaderDrawAndShaderPrint(View, 4000, 2000))
 	{
 		return;
 	}
@@ -1071,11 +1098,6 @@ static void AddDrawDebugClusterPass(
 	const FHairStrandClusterData& HairClusterData,
 	const FViewInfo& View)
 {
-	if (!IsDebugDrawAndDebugPrintEnabled(View))
-	{
-		return;
-	}
-
 	const bool bDebugEnable = IsHairStrandsClusterDebugAABBEnable();
 	const bool bCullingEnable = IsHairStrandsClusterCullingEnable();
 	if (!bDebugEnable || !bCullingEnable)
@@ -1083,6 +1105,11 @@ static void AddDrawDebugClusterPass(
 		return;
 	}
 	
+	if (!TryEnableShaderDrawAndShaderPrint(View, 5000, 2000))
+	{
+		return;
+	}
+
 	{
 		{
 			for (const FHairStrandsMacroGroupData& MacroGroupData : View.HairStrandsViewData.MacroGroupDatas)
diff --git a/Engine/Source/Runtime/Renderer/Private/HairStrands/HairStrandsVoxelization.cpp b/Engine/Source/Runtime/Renderer/Private/HairStrands/HairStrandsVoxelization.cpp
--- a/Engine/Source/Runtime/Renderer/Private/HairStrands/HairStrandsVoxelization.cpp
+++ b/Engine/Source/Runtime/Renderer/Private/HairStrands/HairStrandsVoxelization.cpp
@@ -901,7 +901,7 @@ static void AddAllocateVoxelPagesPass(
 static float RoundHairVoxeliSize(float In)
 {
 	// Round voxel size to 0.01f to avoid oscillation issue
-	return FMath::RoundToFloat(In * 100.f) * 0.01f;
+	return FMath::RoundToFloat(In * 1000.f) * 0.001f;
 }
 
 static FHairStrandsVoxelResources AllocateVirtualVoxelResources(
@@ -923,7 +923,9 @@ static FHairStrandsVoxelResources AllocateVirtualVoxelResources(
 	const float MinVoxelWorldSize = 0.01f;
 	const float MaxVoxelWorldSize = 10.f;
 
-	float VoxelWorldSize = RoundHairVoxeliSize(FMath::Clamp(GHairVirtualVoxel_VoxelWorldSize, MinVoxelWorldSize, MaxVoxelWorldSize));
+	const float RequestedVoxelWorldSize = RoundHairVoxeliSize(FMath::Clamp(GHairVirtualVoxel_VoxelWorldSize, MinVoxelWorldSize, MaxVoxelWorldSize));
+	float TargetVoxelWorldSize = RequestedVoxelWorldSize;
+	float VoxelWorldSize = RequestedVoxelWorldSize;
 
 	// Readback allocated value to adapt the voxel size in order to fit max page allocation
 	const bool bAdaptiveResolution = OutViewStateData && OutViewStateData->IsInit();
@@ -979,7 +981,7 @@ static FHairStrandsVoxelResources AllocateVirtualVoxelResources(
 			// If the page pool is large enough but the voxel are larger than the requested size decrease voxel size
 			else if (AllocatedPageCount < PageCount && PrevWorldVoxelSize > VoxelWorldSize)
 			{
-				const float TargetVoxelWorldSize = PrevWorldVoxelSize * LinearRatio_Thres;
+				TargetVoxelWorldSize = PrevWorldVoxelSize * LinearRatio_Thres;
 				VoxelWorldSize = FMath::Max(VoxelWorldSize, FMath::Lerp(PrevWorldVoxelSize, TargetVoxelWorldSize, Factor));
 			}
 			else
@@ -1027,12 +1029,17 @@ static FHairStrandsVoxelResources AllocateVirtualVoxelResources(
 	Out.Parameters.Common.SteppingScale_Environment	  = FMath::Clamp(GHairStransVoxelRaymarchingSteppingScale_Environment		>= 0 ? GHairStransVoxelRaymarchingSteppingScale_Environment		: GHairStransVoxelRaymarchingSteppingScale, 1.f, 10.f);
 	Out.Parameters.Common.SteppingScale_Raytracing	  = FMath::Clamp(GHairStransVoxelRaymarchingSteppingScale_Raytracing		>= 0 ? GHairStransVoxelRaymarchingSteppingScale_Raytracing		: GHairStransVoxelRaymarchingSteppingScale, 1.f, 10.f);
 
-	Out.Parameters.Common.NodeDescCount				= MacroGroupDatas.Num();
-	Out.Parameters.Common.IndirectDispatchGroupSize = 64;	
+	Out.Parameters.Common.NodeDescCount							= MacroGroupDatas.Num();
+	Out.Parameters.Common.IndirectDispatchGroupSize				= 64;	
 	Out.Parameters.Common.Raytracing_ShadowOcclusionThreshold	= FMath::Max(0.f, GHairVirtualVoxelRaytracing_ShadowOcclusionThreshold);
 	Out.Parameters.Common.Raytracing_SkyOcclusionThreshold		= FMath::Max(0.f, GHairVirtualVoxelRaytracing_SkyOcclusionThreshold);
-	Out.Parameters.Common.HairCoveragePixelRadiusAtDepth1	= ComputeMinStrandRadiusAtDepth1(FIntPoint(View.ViewRect.Width(), View.ViewRect.Height()), View.FOV, 1/*SampleCount*/, 1/*RasterizationScale*/).Primary;
-	Out.Parameters.Common.TranslatedWorldOffset		= FVector3f(View.ViewMatrices.GetPreViewTranslation());
+	Out.Parameters.Common.HairCoveragePixelRadiusAtDepth1		= ComputeMinStrandRadiusAtDepth1(FIntPoint(View.ViewRect.Width(), View.ViewRect.Height()), View.FOV, 1/*SampleCount*/, 1/*RasterizationScale*/).Primary;
+	Out.Parameters.Common.TranslatedWorldOffset					= FVector3f(View.ViewMatrices.GetPreViewTranslation());
+
+	// For debug purpose
+	Out.Parameters.Common.AdaptiveEnable					= bAdaptiveResolution ? 1u : 0u;
+	Out.Parameters.Common.AdaptiveRequestedVoxelWorldSize	= RequestedVoxelWorldSize;
+	Out.Parameters.Common.AdaptiveTargetVoxelWorldSize		= TargetVoxelWorldSize;
 
 	FRDGBufferRef TotalRequestedPageAllocationBuffer;
 	AddAllocateVoxelPagesPass(
