Fix CIS issue.

diff --git a/Engine/Source/Runtime/Renderer/Private/HairStrands/HairStrandsEnvironment.cpp b/Engine/Source/Runtime/Renderer/Private/HairStrands/HairStrandsEnvironment.cpp
--- a/Engine/Source/Runtime/Renderer/Private/HairStrands/HairStrandsEnvironment.cpp
+++ b/Engine/Source/Runtime/Renderer/Private/HairStrands/HairStrandsEnvironment.cpp
@@ -341,6 +341,10 @@ static void AddHairStrandsEnvironmentLightingPassPS(
 	const EHairLightingSourceType LightingType,
 	const FHairStrandsDebugData::Data* DebugData)
 {
+	if (!Scene)
+	{
+		return;
+	}
 
 	EHairLightingIntegrationType IntegrationType = EHairLightingIntegrationType::AdHoc;
 	if (LightingType == EHairLightingSourceType::SceneColor)
