Fix GlobalDistanceField not getting correctly updated when scene only contains heightfields.
- Multiple branches that only checked for mesh SDFs needed to be updated.
- New permutation of AllocatePagesCS to skip mesh SDFs when none are in the scene.
- New permutations of ComposeHeightfieldsIntoPagesCS to support composing with None/Previous/Parent.
    - Previous - compose with mesh SDF pages.
    - Parent - compose with mostly static pages.
- Improvements to distance field visualization mode conditions.

diff --git a/Engine/Shaders/Private/GlobalDistanceField.usf b/Engine/Shaders/Private/GlobalDistanceField.usf
--- a/Engine/Shaders/Private/GlobalDistanceField.usf
+++ b/Engine/Shaders/Private/GlobalDistanceField.usf
@@ -242,6 +242,10 @@ uint3 PageGridCoordToPageTableTextureCoord(uint3 PageGridCoord)
 	return PageTableTextureCoord;
 }
 
+#ifndef PROCESS_DISTANCE_FIELDS
+#define PROCESS_DISTANCE_FIELDS 0
+#endif // PROCESS_DISTANCE_FIELDS
+
 [numthreads(THREADGROUP_SIZE_X, 1, 1)]
 void AllocatePagesCS(
 	uint3 GroupId : SV_GroupID,
@@ -265,6 +269,7 @@ void AllocatePagesCS(
 	uint3 PageTableTextureCoord = PageGridCoordToPageTableTextureCoord(PageGridCoord);
 	float3 PageWorldCenter = PageGridCoord * PageCoordToPageWorldCenterScale + PageCoordToPageWorldCenterBias;
 
+#if PROCESS_DISTANCE_FIELDS
 	// Loop through objects to find out if the current page contains any objects and should be allocated
 	if (!bAnyObjectsInPage)
 	{
@@ -330,15 +335,18 @@ void AllocatePagesCS(
 	}
 
 	GroupMemoryBarrierWithGroupSync();
+#endif // PROCESS_DISTANCE_FIELDS
 
 	if (ThreadIndex == 0)
 	{
+#if PROCESS_DISTANCE_FIELDS
 		float MinPageDistance = min(GroupMinPageDistance[0], GroupMinPageDistance[1]);
 
 		if (MinPageDistance + PageWorldRadius >= -ClipmapInfluenceRadius && MinPageDistance - PageWorldRadius <= ClipmapInfluenceRadius)
 		{
 			bAnyObjectsInPage = true;
 		}
+#endif // PROCESS_DISTANCE_FIELDS
 
 		uint PageId = RWPageTableLayerTexture[PageTableTextureCoord];
 
diff --git a/Engine/Shaders/Private/GlobalDistanceFieldHeightfields.usf b/Engine/Shaders/Private/GlobalDistanceFieldHeightfields.usf
--- a/Engine/Shaders/Private/GlobalDistanceFieldHeightfields.usf
+++ b/Engine/Shaders/Private/GlobalDistanceFieldHeightfields.usf
@@ -11,6 +11,10 @@
 #define GLOBAL_DISTANCE_FIELD_PAGE_HEIGHFIELD_GROUP_X (GLOBAL_DISTANCE_FIELD_PAGE_RESOLUTION_IN_ATLAS / 8)
 #define GLOBAL_DISTANCE_FIELD_PAGE_HEIGHFIELD_GROUP_NUM (GLOBAL_DISTANCE_FIELD_PAGE_HEIGHFIELD_GROUP_X * GLOBAL_DISTANCE_FIELD_PAGE_HEIGHFIELD_GROUP_X)
 
+#define COMPOSE_DISTANCE_FIELD_NONE 0
+#define COMPOSE_DISTANCE_FIELD_PARENT 1
+#define COMPOSE_DISTANCE_FIELD_PREVIOUS 2
+
 RWStructuredBuffer<uint> RWMarkedHeightfieldPageBuffer;
 RWTexture3D<UNORM float> RWPageAtlasTexture;
 
@@ -18,6 +22,7 @@ StructuredBuffer<uint> MarkedHeightfieldPageBuffer;
 StructuredBuffer<uint> PageUpdateTileBuffer;
 StructuredBuffer<uint> ComposeTileBuffer;
 Texture3D<uint> PageTableLayerTexture;
+Texture3D<uint> ParentPageTableLayerTexture;
 
 float3 PageCoordToPageWorldCenterScale;
 float3 PageCoordToPageWorldCenterBias;
@@ -296,11 +301,26 @@ void ComposeHeightfieldsIntoPagesCS(
 					DistanceToHeightfieldPlane = MinInteriorDistance - DistanceToHeightfieldPlane;
 				}
 
-				const float MinDistance = clamp(DistanceToHeightfieldPlane, -InfluenceRadius, InfluenceRadius);
+				float MinDistance = clamp(DistanceToHeightfieldPlane, -InfluenceRadius, InfluenceRadius);
+
+				const uint3 PageAtlasCoordZ = PageAtlasCoord + uint3(0, 0, ZIndex);
+
+#if COMPOSE_DISTANCE_FIELD_MODE == COMPOSE_DISTANCE_FIELD_PARENT
+				uint ParentPageId = ParentPageTableLayerTexture.Load(int4(PageTableTextureCoord, 0));
+				if (ParentPageId != GLOBAL_DISTANCE_FIELD_INVALID_PAGE_ID)
+				{
+					uint3 ParentPageAtlasCoord = GlobalDistanceFieldPageLinearIndexToPageAtlasCoord(ParentPageId);
+					ParentPageAtlasCoord += TexelCoordInPage;
+
+					float ParentDistanceField = DecodeGlobalDistanceFieldPageDistance(RWPageAtlasTexture[ParentPageAtlasCoord].x, InfluenceRadius);
+					MinDistance = min(MinDistance, ParentDistanceField);
+				}
+#elif COMPOSE_DISTANCE_FIELD_MODE == COMPOSE_DISTANCE_FIELD_PREVIOUS
+				float PreviousDistanceField = DecodeGlobalDistanceFieldPageDistance(RWPageAtlasTexture[PageAtlasCoordZ], InfluenceRadius);
+				MinDistance = min(MinDistance, PreviousDistanceField);
+#endif
 
-				uint3 PageAtlasCoordZ = PageAtlasCoord + uint3(0, 0, ZIndex);
-				float PrevMinDistance = DecodeGlobalDistanceFieldPageDistance(RWPageAtlasTexture[PageAtlasCoordZ], InfluenceRadius);
-				RWPageAtlasTexture[PageAtlasCoordZ] = EncodeGlobalDistanceFieldPageDistance(min(MinDistance, PrevMinDistance), InfluenceRadius);
+				RWPageAtlasTexture[PageAtlasCoordZ] = EncodeGlobalDistanceFieldPageDistance(MinDistance, InfluenceRadius);
 			}
 		}
 	}
diff --git a/Engine/Source/Runtime/Renderer/Private/DistanceFieldVisualization.cpp b/Engine/Source/Runtime/Renderer/Private/DistanceFieldVisualization.cpp
--- a/Engine/Source/Runtime/Renderer/Private/DistanceFieldVisualization.cpp
+++ b/Engine/Source/Runtime/Renderer/Private/DistanceFieldVisualization.cpp
@@ -88,10 +88,15 @@ void FDeferredShadingSceneRenderer::RenderMeshDistanceFieldVisualization(
 {
 	const FViewInfo& FirstView = Views[0];
 
-	if (!UseDistanceFieldAO()
-		|| FeatureLevel < ERHIFeatureLevel::SM5
-		|| !DoesPlatformSupportDistanceFieldAO(FirstView.GetShaderPlatform())
-		|| Scene->DistanceFieldSceneData.NumObjectsInBuffer == 0)
+	const bool bAnyMeshSDFs = Scene->DistanceFieldSceneData.NumObjectsInBuffer > 0;
+	const bool bAnyHeightFields = Scene->DistanceFieldSceneData.HeightfieldPrimitives.Num() > 0;
+
+	const bool bVisualizeGlobalDistanceField = FirstView.Family->EngineShowFlags.VisualizeGlobalDistanceField && UseGlobalDistanceField(Parameters) && (bAnyMeshSDFs || bAnyHeightFields);
+	const bool bVisualizeMeshDistanceFields = FirstView.Family->EngineShowFlags.VisualizeMeshDistanceFields && bAnyMeshSDFs;
+
+	if (FeatureLevel < ERHIFeatureLevel::SM5
+		|| !DoesPlatformSupportDistanceFields(FirstView.GetShaderPlatform())
+		|| !(bVisualizeGlobalDistanceField || bVisualizeMeshDistanceFields))
 	{
 		return;
 	}
@@ -100,12 +105,9 @@ void FDeferredShadingSceneRenderer::RenderMeshDistanceFieldVisualization(
 
 	QUICK_SCOPE_CYCLE_COUNTER(STAT_AOIssueGPUWork);
 
-	const bool bUseGlobalDistanceField = UseGlobalDistanceField(Parameters) && FirstView.Family->EngineShowFlags.VisualizeGlobalDistanceField;
-
 	RDG_EVENT_SCOPE(GraphBuilder, "VisualizeMeshDistanceFields");
 
 	FRDGTextureRef VisualizeResultTexture = nullptr;
-
 	{
 		const FIntPoint BufferSize = GetBufferSizeForAO();
 		const FRDGTextureDesc Desc(FRDGTextureDesc::Create2D(BufferSize, PF_FloatRGBA, FClearValueBinding::None, TexCreate_RenderTargetable | TexCreate_UAV));
@@ -113,25 +115,30 @@ void FDeferredShadingSceneRenderer::RenderMeshDistanceFieldVisualization(
 	}
 
 	FVisualizeMeshDistanceFieldCS::FPermutationDomain PermutationVector;
-	PermutationVector.Set<FVisualizeMeshDistanceFieldCS::FUseGlobalDistanceFieldDim>(bUseGlobalDistanceField);
+	PermutationVector.Set<FVisualizeMeshDistanceFieldCS::FUseGlobalDistanceFieldDim>(bVisualizeGlobalDistanceField);
 	TShaderMapRef<FVisualizeMeshDistanceFieldCS> ComputeShader(FirstView.ShaderMap, PermutationVector);
 
 	for (const FViewInfo& View : Views)
 	{
-		check(!bUseGlobalDistanceField || View.GlobalDistanceFieldInfo.Clipmaps.Num() > 0);
-
 		FRDGBufferRef ObjectIndirectArguments = nullptr;
 		FDistanceFieldCulledObjectBufferParameters CulledObjectBufferParameters;
 
-		AllocateDistanceFieldCulledObjectBuffers(
-			GraphBuilder,
-			false,
-			FMath::DivideAndRoundUp(Scene->DistanceFieldSceneData.NumObjectsInBuffer, 256) * 256,
-			DFPT_SignedDistanceField,
-			ObjectIndirectArguments,
-			CulledObjectBufferParameters);
-
-		CullObjectsToView(GraphBuilder, Scene, View, Parameters, CulledObjectBufferParameters);
+		if (bVisualizeGlobalDistanceField)
+		{
+			check(View.GlobalDistanceFieldInfo.Clipmaps.Num() > 0);
+		}
+		else
+		{
+			AllocateDistanceFieldCulledObjectBuffers(
+				GraphBuilder,
+				false,
+				FMath::DivideAndRoundUp(Scene->DistanceFieldSceneData.NumObjectsInBuffer, 256) * 256,
+				DFPT_SignedDistanceField,
+				ObjectIndirectArguments,
+				CulledObjectBufferParameters);
+
+			CullObjectsToView(GraphBuilder, Scene, View, Parameters, CulledObjectBufferParameters);
+		}
 
 		uint32 GroupSizeX = FMath::DivideAndRoundUp(View.ViewRect.Size().X / GAODownsampleFactor, GDistanceFieldAOTileSizeX);
 		uint32 GroupSizeY = FMath::DivideAndRoundUp(View.ViewRect.Size().Y / GAODownsampleFactor, GDistanceFieldAOTileSizeY);
diff --git a/Engine/Source/Runtime/Renderer/Private/GlobalDistanceField.cpp b/Engine/Source/Runtime/Renderer/Private/GlobalDistanceField.cpp
--- a/Engine/Source/Runtime/Renderer/Private/GlobalDistanceField.cpp
+++ b/Engine/Source/Runtime/Renderer/Private/GlobalDistanceField.cpp
@@ -1401,9 +1401,10 @@ class FAllocatePagesCS : public FGlobalShader
 		SHADER_PARAMETER_STRUCT_INCLUDE(FDistanceFieldAtlasParameters, DistanceFieldAtlas)
 	END_SHADER_PARAMETER_STRUCT()
 
+	class FProcessDistanceFields : SHADER_PERMUTATION_BOOL("PROCESS_DISTANCE_FIELDS");
 	class FMarkedHeightfieldPageBuffer : SHADER_PERMUTATION_BOOL("MARKED_HEIGHTFIELD_PAGE_BUFFER");
 	class FComposeParentDistanceField : SHADER_PERMUTATION_BOOL("COMPOSE_PARENT_DISTANCE_FIELD");
-	using FPermutationDomain = TShaderPermutationDomain<FMarkedHeightfieldPageBuffer, FComposeParentDistanceField>;
+	using FPermutationDomain = TShaderPermutationDomain<FProcessDistanceFields, FMarkedHeightfieldPageBuffer, FComposeParentDistanceField>;
 	
 	static FPermutationDomain RemapPermutation(FPermutationDomain PermutationVector)
 	{
@@ -1569,7 +1570,7 @@ void UpdateGlobalDistanceFieldVolume(
 
 	UpdateGlobalDistanceFieldViewOrigin(View, bLumenEnabled);
 
-	if (Scene->DistanceFieldSceneData.NumObjectsInBuffer > 0)
+	if (DistanceFieldSceneData.NumObjectsInBuffer > 0 || DistanceFieldSceneData.HeightfieldPrimitives.Num() > 0)
 	{
 		const int32 NumClipmaps = FMath::Clamp<int32>(GetNumGlobalDistanceFieldClipmaps(bLumenEnabled, View.FinalPostProcessSettings.LumenSceneViewDistance), 0, GMaxGlobalDistanceFieldClipmaps);
 		ComputeUpdateRegionsAndUpdateViewState(GraphBuilder.RHICmdList, View, Scene, GlobalDistanceFieldInfo, NumClipmaps, MaxOcclusionDistance, bLumenEnabled);
@@ -1808,6 +1809,7 @@ void UpdateGlobalDistanceFieldVolume(
 					if (NumUpdateBounds > 0 && PageAtlasTexture)
 					{
 						// Cull the global objects to the update regions
+						if (Scene->DistanceFieldSceneData.NumObjectsInBuffer > 0)
 						{
 							uint32 AcceptOftenMovingObjectsOnlyValue = 0;
 
@@ -2019,6 +2021,7 @@ void UpdateGlobalDistanceFieldVolume(
 						FDistanceFieldAtlasParameters DistanceFieldAtlas = DistanceField::SetupAtlasParameters(DistanceFieldSceneData);
 
 						// Cull objects into a cull grid
+						if (Scene->DistanceFieldSceneData.NumObjectsInBuffer > 0)
 						{
 							AddClearUAVPass(GraphBuilder, GraphBuilder.CreateUAV(CullGridAllocator, PF_R32_UINT), 0);
 							AddClearUAVPass(GraphBuilder, GraphBuilder.CreateUAV(CullGridObjectHeader, PF_R32_UINT), 0);
@@ -2093,6 +2096,7 @@ void UpdateGlobalDistanceFieldVolume(
 								PassParameters->DistanceFieldAtlas = DistanceFieldAtlas;
 
 								FAllocatePagesCS::FPermutationDomain PermutationVector;
+								PermutationVector.Set<FAllocatePagesCS::FProcessDistanceFields>(Scene->DistanceFieldSceneData.NumObjectsInBuffer > 0);
 								PermutationVector.Set<FAllocatePagesCS::FMarkedHeightfieldPageBuffer>(MarkedHeightfieldPageBuffer != nullptr);
 								PermutationVector.Set<FAllocatePagesCS::FComposeParentDistanceField>(ParentPageTableLayerTexture != nullptr);
 								auto ComputeShader = View.ShaderMap->GetShader<FAllocatePagesCS>(PermutationVector);
@@ -2147,6 +2151,8 @@ void UpdateGlobalDistanceFieldVolume(
 						}
 
 						// Compose the mesh SDFs into allocated pages
+						const bool bComposeMeshSDFsIntoPages = Scene->DistanceFieldSceneData.NumObjectsInBuffer > 0;
+						if (bComposeMeshSDFsIntoPages)
 						{
 							const FVector PageVoxelExtent = 0.5f * ClipmapSize / FVector(ClipmapResolution);
 							const FVector PageCoordToVoxelCenterScale = ClipmapSize / FVector(ClipmapResolution);
@@ -2227,6 +2233,7 @@ void UpdateGlobalDistanceFieldVolume(
 									PassParameters->ComposeIndirectArgBuffer = PageComposeHeightfieldIndirectArgBuffer;
 									PassParameters->ComposeTileBuffer = GraphBuilder.CreateSRV(PageComposeHeightfieldTileBuffer, PF_R32_UINT);
 									PassParameters->PageTableLayerTexture = PageTableLayerTexture;
+									PassParameters->ParentPageTableLayerTexture = ParentPageTableLayerTexture;
 									PassParameters->InfluenceRadius = ClipmapInfluenceRadius;
 									PassParameters->PageCoordToVoxelCenterScale = PageCoordToVoxelCenterScale;
 									PassParameters->PageCoordToVoxelCenterBias = PageCoordToVoxelCenterBias;
@@ -2246,7 +2253,20 @@ void UpdateGlobalDistanceFieldVolume(
 									PassParameters->VisibilitySampler = TStaticSamplerState<SF_Bilinear>::GetRHI();
 									PassParameters->HeightfieldDescriptions = GraphBuilder.CreateSRV(HeightfieldDescriptionBuffer, EPixelFormat::PF_A32B32G32R32F);
 
-									auto ComputeShader = View.ShaderMap->GetShader<FComposeHeightfieldsIntoPagesCS>();
+									EComposeDistanceFieldMode ComposeMode = EComposeDistanceFieldMode::None;
+
+									if (bComposeMeshSDFsIntoPages)
+									{
+										ComposeMode = EComposeDistanceFieldMode::Previous;
+									}
+									else if (ParentPageTableLayerTexture != nullptr)
+									{
+										ComposeMode = EComposeDistanceFieldMode::Parent;
+									}
+
+									FComposeHeightfieldsIntoPagesCS::FPermutationDomain PermutationVector;
+									PermutationVector.Set<FComposeHeightfieldsIntoPagesCS::FComposeDistanceFieldMode>(ComposeMode);
+									auto ComputeShader = View.ShaderMap->GetShader<FComposeHeightfieldsIntoPagesCS>(PermutationVector);
 
 									FComputeShaderUtils::AddPass(
 										GraphBuilder,
diff --git a/Engine/Source/Runtime/Renderer/Private/GlobalDistanceFieldHeightfields.h b/Engine/Source/Runtime/Renderer/Private/GlobalDistanceFieldHeightfields.h
--- a/Engine/Source/Runtime/Renderer/Private/GlobalDistanceFieldHeightfields.h
+++ b/Engine/Source/Runtime/Renderer/Private/GlobalDistanceFieldHeightfields.h
@@ -116,6 +116,16 @@ class FBuildHeightfieldComposeTilesCS : public FGlobalShader
 
 IMPLEMENT_GLOBAL_SHADER(FBuildHeightfieldComposeTilesCS, "/Engine/Private/GlobalDistanceFieldHeightfields.usf", "BuildHeightfieldComposeTilesCS", SF_Compute);
 
+// Must match COMPOSE_DISTANCE_FIELD_ in usf
+enum class EComposeDistanceFieldMode
+{
+	None,
+	Parent,
+	Previous,
+
+	MAX
+};
+
 class FComposeHeightfieldsIntoPagesCS : public FGlobalShader
 {
 	DECLARE_GLOBAL_SHADER(FComposeHeightfieldsIntoPagesCS);
@@ -127,6 +137,7 @@ class FComposeHeightfieldsIntoPagesCS : public FGlobalShader
 		RDG_BUFFER_ACCESS(ComposeIndirectArgBuffer, ERHIAccess::IndirectArgs)
 		SHADER_PARAMETER_RDG_BUFFER_SRV(StructuredBuffer<uint>, ComposeTileBuffer)
 		SHADER_PARAMETER_RDG_TEXTURE(Texture3D<uint>, PageTableLayerTexture)
+		SHADER_PARAMETER_RDG_TEXTURE(Texture3D<uint>, ParentPageTableLayerTexture)
 		SHADER_PARAMETER(FVector3f, InvPageGridResolution)
 		SHADER_PARAMETER(FIntVector, PageGridResolution)
 		SHADER_PARAMETER(FVector3f, PageCoordToVoxelCenterScale)
@@ -146,6 +157,9 @@ class FComposeHeightfieldsIntoPagesCS : public FGlobalShader
 		SHADER_PARAMETER(float, HeightfieldThickness)
 	END_SHADER_PARAMETER_STRUCT()
 
+	class FComposeDistanceFieldMode : SHADER_PERMUTATION_ENUM_CLASS("COMPOSE_DISTANCE_FIELD_MODE", EComposeDistanceFieldMode);
+	using FPermutationDomain = TShaderPermutationDomain<FComposeDistanceFieldMode>;
+
 	static bool ShouldCompilePermutation(const FGlobalShaderPermutationParameters& Parameters)
 	{
 		return ShouldCompileDistanceFieldShaders(Parameters.Platform);
