Fix single-layer water rendering crash due to missing TranslucencyResourceMap.

diff --git a/Engine/Source/Runtime/Renderer/Private/DeferredShadingRenderer.cpp b/Engine/Source/Runtime/Renderer/Private/DeferredShadingRenderer.cpp
--- a/Engine/Source/Runtime/Renderer/Private/DeferredShadingRenderer.cpp
+++ b/Engine/Source/Runtime/Renderer/Private/DeferredShadingRenderer.cpp
@@ -2984,6 +2984,8 @@ void FDeferredShadingSceneRenderer::Render(FRDGBuilder& GraphBuilder)
 	// Union of all translucency view render flags.
 	ETranslucencyView TranslucencyViewsToRender = bShouldRenderTranslucency ? GetTranslucencyViews(Views) : ETranslucencyView::None;
 
+	FTranslucencyPassResourcesMap TranslucencyResourceMap(Views.Num());
+
 	const bool bShouldRenderSingleLayerWater = !bHasRayTracedOverlay && ShouldRenderSingleLayerWater(Views);
 	FSceneWithoutWaterTextures SceneWithoutWaterTextures;
 	if (bShouldRenderSingleLayerWater)
@@ -2993,7 +2995,7 @@ void FDeferredShadingSceneRenderer::Render(FRDGBuilder& GraphBuilder)
 			RDG_CSV_STAT_EXCLUSIVE_SCOPE(GraphBuilder, RenderTranslucency);
 			SCOPE_CYCLE_COUNTER(STAT_TranslucencyDrawTime);
 			GraphBuilder.SetCommandListStat(GET_STATID(STAT_CLM_Translucency));
-			RenderTranslucency(GraphBuilder, SceneTextures, TranslucencyLightingVolumeTextures, nullptr, ETranslucencyView::UnderWater, InstanceCullingManager);
+			RenderTranslucency(GraphBuilder, SceneTextures, TranslucencyLightingVolumeTextures, &TranslucencyResourceMap, ETranslucencyView::UnderWater, InstanceCullingManager);
 			EnumRemoveFlags(TranslucencyViewsToRender, ETranslucencyView::UnderWater);
 		}
 
@@ -3058,8 +3060,6 @@ void FDeferredShadingSceneRenderer::Render(FRDGBuilder& GraphBuilder)
 		RenderHairComposition(GraphBuilder, Views, SceneTextures.Color.Target, SceneTextures.Depth.Target);
 	}
 
-	FTranslucencyPassResourcesMap TranslucencyResourceMap(Views.Num());
-
 	// Draw translucency.
 	if (!bHasRayTracedOverlay && TranslucencyViewsToRender != ETranslucencyView::None)
 	{
