Fix crash when using shader draw debug, while it is not enabled.

diff --git a/Engine/Shaders/Private/ShaderDrawDebug.ush b/Engine/Shaders/Private/ShaderDrawDebug.ush
--- a/Engine/Shaders/Private/ShaderDrawDebug.ush
+++ b/Engine/Shaders/Private/ShaderDrawDebug.ush
@@ -254,11 +254,6 @@ void AddLineTriangle(float3 P0, float3 P1, float3 P2, float4 Color)		{ AddLineTr
 
 void AddOBBTWS(float3 Min, float3 Max, float4 Color, float4x4 Transform)
 {
-	if (ShaderDrawMaxElementCount == 0)
-	{
-		return;
-	}
-
 	uint Offset = 0;
 	if (AllocateShaderDrawElement(12, Offset))
 	{
diff --git a/Engine/Source/Runtime/Renderer/Private/ShaderDebug.cpp b/Engine/Source/Runtime/Renderer/Private/ShaderDebug.cpp
--- a/Engine/Source/Runtime/Renderer/Private/ShaderDebug.cpp
+++ b/Engine/Source/Runtime/Renderer/Private/ShaderDebug.cpp
@@ -291,46 +291,47 @@ namespace ShaderDrawDebug
 	void BeginView(FRDGBuilder& GraphBuilder, FViewInfo& View)
 	{
 		View.ShaderDrawData = FShaderDrawDebugData();
-		if (!IsEnabled(View))
-		{
-			// Default resources
-			View.ShaderDrawData.Buffer = GraphBuilder.CreateBuffer(FRDGBufferDesc::CreateStructuredDesc(4, 8), TEXT("ShaderDraw.DataBuffer(Dummy)"), ERDGBufferFlags::None);
-			return;
-		}
+		View.ShaderDrawData.ShaderDrawTranslatedWorldOffset = View.ViewMatrices.GetPreViewTranslation();
+		View.ShaderDrawData.CursorPosition = View.CursorPos;
+		View.ShaderDrawData.MaxElementCount = 0;
 
-		// Update max element count from the request count and reset
-		GShaderDrawDebug_MaxElementCount = FMath::Max3(1, GShaderDrawDebug_MaxElementCount, int32(GElementRequestCount));
-		View.ShaderDrawData.MaxElementCount = GShaderDrawDebug_MaxElementCount;
-		GElementRequestCount = 0;
+		if (IsEnabled(View))
+		{
+			// Update max element count from the request count and reset
+			GShaderDrawDebug_MaxElementCount = FMath::Max3(1, GShaderDrawDebug_MaxElementCount, int32(GElementRequestCount));
+			View.ShaderDrawData.MaxElementCount = GShaderDrawDebug_MaxElementCount;
+			GElementRequestCount = 0;
 
-		FSceneViewState* ViewState = View.ViewState;
-		const bool bLockBufferThisFrame = IsShaderDrawLocked() && ViewState && !ViewState->ShaderDrawDebugStateData.bIsLocked;
-		ERDGBufferFlags Flags = bLockBufferThisFrame ? ERDGBufferFlags::MultiFrame : ERDGBufferFlags::None;
+			FSceneViewState* ViewState = View.ViewState;
+			const bool bLockBufferThisFrame = IsShaderDrawLocked() && ViewState && !ViewState->ShaderDrawDebugStateData.bIsLocked;
+			ERDGBufferFlags Flags = bLockBufferThisFrame ? ERDGBufferFlags::MultiFrame : ERDGBufferFlags::None;
 
-		FRDGBufferRef DataBuffer = GraphBuilder.CreateBuffer(FRDGBufferDesc::CreateStructuredDesc(4, 8*View.ShaderDrawData.MaxElementCount), TEXT("ShaderDraw.DataBuffer"), Flags);
-		AddShaderDrawDebugClearPass(GraphBuilder, View, DataBuffer);
+			View.ShaderDrawData.Buffer = GraphBuilder.CreateBuffer(FRDGBufferDesc::CreateStructuredDesc(4, 8*View.ShaderDrawData.MaxElementCount), TEXT("ShaderDraw.DataBuffer"), Flags);
+			AddShaderDrawDebugClearPass(GraphBuilder, View, View.ShaderDrawData.Buffer);
 
-		View.ShaderDrawData.Buffer = DataBuffer;
-		View.ShaderDrawData.CursorPosition = View.CursorPos;
-		View.ShaderDrawData.ShaderDrawTranslatedWorldOffset = View.ViewMatrices.GetPreViewTranslation();
+			if (IsShaderDrawLocked() && ViewState && !ViewState->ShaderDrawDebugStateData.bIsLocked)
+			{
+				ViewState->ShaderDrawDebugStateData.Buffer = GraphBuilder.ConvertToExternalBuffer(View.ShaderDrawData.Buffer);
+				ViewState->ShaderDrawDebugStateData.bIsLocked = true;
+			}
 
-		if (IsShaderDrawLocked() && ViewState && !ViewState->ShaderDrawDebugStateData.bIsLocked)
-		{
-			ViewState->ShaderDrawDebugStateData.Buffer = GraphBuilder.ConvertToExternalBuffer(View.ShaderDrawData.Buffer);
-			ViewState->ShaderDrawDebugStateData.bIsLocked = true;
-		}
+			if (!IsShaderDrawLocked() && ViewState && ViewState->ShaderDrawDebugStateData.bIsLocked)
+			{
+				ViewState->ShaderDrawDebugStateData.Buffer = nullptr;
+				ViewState->ShaderDrawDebugStateData.bIsLocked = false;
+			}
 
-		if (!IsShaderDrawLocked() && ViewState && ViewState->ShaderDrawDebugStateData.bIsLocked)
-		{
-			ViewState->ShaderDrawDebugStateData.Buffer = nullptr;
-			ViewState->ShaderDrawDebugStateData.bIsLocked = false;
+			// Invalid to call begin twice for the same view.
+			check(GDefaultView != &View);
+			if (GDefaultView == nullptr)
+			{
+				GDefaultView = &View;
+			}
 		}
-
-		// Invalid to call begin twice for the same view.
-		check(GDefaultView != &View);
-		if (GDefaultView == nullptr)
+		else
 		{
-			GDefaultView = &View;
+			// Default resources
+			View.ShaderDrawData.Buffer = GraphBuilder.CreateBuffer(FRDGBufferDesc::CreateStructuredDesc(4, 8), TEXT("ShaderDraw.DataBuffer(Dummy)"), ERDGBufferFlags::None);
 		}
 	}
 
@@ -365,14 +366,17 @@ namespace ShaderDrawDebug
 
 	void SetParameters(FRDGBuilder& GraphBuilder, const FShaderDrawDebugData& Data, FShaderParameters& OutParameters)
 	{
-		// Early out if debug rendering is not enabled/supported
-		if (Data.Buffer == nullptr)
-			return;
-
 		OutParameters.ShaderDrawCursorPos				= Data.CursorPosition;
 		OutParameters.ShaderDrawMaxElementCount			= Data.MaxElementCount;
 		OutParameters.ShaderDrawTranslatedWorldOffset	= Data.ShaderDrawTranslatedWorldOffset;
-		OutParameters.OutShaderDrawPrimitive			= GraphBuilder.CreateUAV(Data.Buffer);
+
+		check(Data.Buffer != nullptr);
+		OutParameters.OutShaderDrawPrimitive = GraphBuilder.CreateUAV(Data.Buffer);		
+	}
+
+	void SetParameters(FRDGBuilder& GraphBuilder, const FViewInfo& View, FShaderParameters& OutParameters)
+	{
+		SetParameters(GraphBuilder, View.ShaderDrawData, OutParameters);
 	}
 
 	// Returns true if the default view exists and has shader debug rendering enabled (this needs to be checked before using a permutation that requires the shader draw parameters)
diff --git a/Engine/Source/Runtime/Renderer/Public/ShaderDebug.h b/Engine/Source/Runtime/Renderer/Public/ShaderDebug.h
--- a/Engine/Source/Runtime/Renderer/Public/ShaderDebug.h
+++ b/Engine/Source/Runtime/Renderer/Public/ShaderDebug.h
@@ -57,8 +57,12 @@ namespace ShaderDrawDebug
 	// Call this to fill the FShaderDrawParameters
 	RENDERER_API void SetParameters(FRDGBuilder& GraphBuilder, const FShaderDrawDebugData& Data, FShaderParameters& OutParameters);
 
+	// Call this to fill the FShaderDrawParameters
+	void SetParameters(FRDGBuilder& GraphBuilder, const FViewInfo& Data, FShaderParameters& OutParameters);
+
 	// Returns true if the default view exists and has shader debug rendering enabled (this needs to be checked before using a permutation that requires the shader draw parameters)
 	bool IsDefaultViewEnabled();
+
 	// Call this to fill the FShaderDrawParameters using the default view (the first one for which Begin was called in case of stereo or similar)
 	RENDERER_API void SetParameters(FRDGBuilder& GraphBuilder, FShaderParameters& OutParameters);
 }
