Fix hair env. lighting caused by incorrect view SH irradiance map parameters binding.

diff --git a/Engine/Source/Runtime/Renderer/Private/HairStrands/HairStrandsEnvironment.cpp b/Engine/Source/Runtime/Renderer/Private/HairStrands/HairStrandsEnvironment.cpp
--- a/Engine/Source/Runtime/Renderer/Private/HairStrands/HairStrandsEnvironment.cpp
+++ b/Engine/Source/Runtime/Renderer/Private/HairStrands/HairStrandsEnvironment.cpp
@@ -184,7 +184,7 @@
 	ClearUnusedGraphResources(PixelShader, PassParameters);
 
 	GraphBuilder.AddPass(
-		RDG_EVENT_NAME("HairStrandsAO %dx%d", ViewRect.Width(), ViewRect.Height()),
+		RDG_EVENT_NAME("HairStrands::AO %dx%d", ViewRect.Width(), ViewRect.Height()),
 		PassParameters,
 		ERDGPassFlags::Raster,
 		[PassParameters, &View, PixelShader, ViewRect](FRHICommandList& InRHICmdList)
@@ -202,54 +202,44 @@
 
 ///////////////////////////////////////////////////////////////////////////////////////////////////
 
-class FHairEnvironmentLighting
+class FHairEnvironmentLightingVS : public FGlobalShader
 {
-public:
+	DECLARE_GLOBAL_SHADER(FHairEnvironmentLightingVS);
+	SHADER_USE_PARAMETER_STRUCT(FHairEnvironmentLightingVS, FGlobalShader)
+
+	BEGIN_SHADER_PARAMETER_STRUCT(FParameters, )
+		SHADER_PARAMETER_RDG_UNIFORM_BUFFER(FHairStrandsViewUniformParameters, HairStrands)
+	END_SHADER_PARAMETER_STRUCT()
+
+	static void ModifyCompilationEnvironment(const FGlobalShaderPermutationParameters& Parameters, FShaderCompilerEnvironment& OutEnvironment)
+	{
+		OutEnvironment.SetDefine(TEXT("LIGHTING_VS"), 1);
+	}
+
+	static bool ShouldCompilePermutation(const FGlobalShaderPermutationParameters& Parameters)
+	{
+		return IsHairStrandsSupported(EHairStrandsShaderType::Strands, Parameters.Platform);
+	}
+};
+
+class FHairEnvironmentLightingPS : public FGlobalShader
+{
+	DECLARE_GLOBAL_SHADER(FHairEnvironmentLightingPS);
+	SHADER_USE_PARAMETER_STRUCT(FHairEnvironmentLightingPS, FGlobalShader)
+
 	class FIntegrationType	: SHADER_PERMUTATION_INT("PERMUTATION_INTEGRATION_TYPE", uint32(EHairLightingIntegrationType::Count));
 	class FDebug			: SHADER_PERMUTATION_BOOL("PERMUTATION_DEBUG"); 
 	class FLumen			: SHADER_PERMUTATION_BOOL("PERMUTATION_LUMEN"); 
 	class FStaticLighting	: SHADER_PERMUTATION_BOOL("PERMUTATION_STATIC_LIGHTING"); 
 	using FPermutationDomain = TShaderPermutationDomain<FIntegrationType, FLumen, FDebug, FStaticLighting>;
 
-	static void ModifyCompilationEnvironment(const FGlobalShaderPermutationParameters& Parameters, FShaderCompilerEnvironment& OutEnvironment)
-	{
-		FForwardLightingParameters::ModifyCompilationEnvironment(Parameters.Platform, OutEnvironment);
-	}
+	BEGIN_SHADER_PARAMETER_STRUCT(FParameters, )
+		SHADER_PARAMETER_STRUCT_INCLUDE(FSkyDiffuseLightingParameters, SkyDiffuseLighting)
+		SHADER_PARAMETER_STRUCT_INCLUDE(FSceneTextureParameters, SceneTextures)
+		SHADER_PARAMETER_STRUCT_INCLUDE(LumenRadianceCache::FRadianceCacheInterpolationParameters, RadianceCache)
+		SHADER_PARAMETER_STRUCT_INCLUDE(ShaderDrawDebug::FShaderParameters, ShaderDrawParameters)
+		SHADER_PARAMETER_STRUCT_INCLUDE(FHairStrandsDebugData::FWriteParameters, DebugData)
 
-	static bool ShouldCompilePermutation(const FGlobalShaderPermutationParameters& Parameters)
-	{
-		// Compile debug permutation only for the uniform integrator
-		FPermutationDomain PermutationVector(Parameters.PermutationId);
-		if (PermutationVector.Get<FIntegrationType>() != int32(EHairLightingIntegrationType::Uniform) && PermutationVector.Get<FDebug>())
-		{
-			return false;
-		}
-
-		// If using Lumen, no static lighting
-		if (PermutationVector.Get<FStaticLighting>() && PermutationVector.Get<FLumen>())
-		{
-			return false;
-		}
-
-		// If using another integrator than SH integrator, no static lighting
-		if (PermutationVector.Get<FStaticLighting>() && PermutationVector.Get<FIntegrationType>() != int32(EHairLightingIntegrationType::SH))
-		{
-			return false;
-		}
-
-		return IsHairStrandsSupported(EHairStrandsShaderType::Strands, Parameters.Platform);
-	}
-
-	static FPermutationDomain RemapPermutation(FPermutationDomain PermutationVector)
-	{
-		if (PermutationVector.Get<FIntegrationType>() != int32(EHairLightingIntegrationType::Uniform))
-		{
-			PermutationVector.Set<FDebug>(false);
-		}
-		return PermutationVector;
-	}
-
-	BEGIN_SHADER_PARAMETER_STRUCT(FParameters, )
 		SHADER_PARAMETER(uint32, bHasStaticLighting)
 		SHADER_PARAMETER(uint32, MultipleScatterSampleCount)
 		SHADER_PARAMETER(float,  HairDualScatteringRoughnessOverride)
@@ -260,69 +250,56 @@
 		SHADER_PARAMETER_TEXTURE(Texture2D, PreIntegratedGF)
 		SHADER_PARAMETER_SAMPLER(SamplerState, PreIntegratedGFSampler)
 
-		SHADER_PARAMETER_STRUCT_INCLUDE(FSkyDiffuseLightingParameters, SkyDiffuseLighting)
-		SHADER_PARAMETER_STRUCT_INCLUDE(FSceneTextureParameters, SceneTextures)
-
-		SHADER_PARAMETER_RDG_UNIFORM_BUFFER(FHairStrandsViewUniformParameters, HairStrands)
-		SHADER_PARAMETER_RDG_UNIFORM_BUFFER(FVirtualVoxelParameters, VirtualVoxel)
-
 		SHADER_PARAMETER_RDG_TEXTURE(Texture2D, SceneColorTexture)
 		SHADER_PARAMETER_RDG_TEXTURE(Texture2D, HairEnergyLUTTexture)
 
-		SHADER_PARAMETER_RDG_BUFFER_UAV(RWBuffer, OutLightingBuffer)
-
 		SHADER_PARAMETER_STRUCT_REF(FViewUniformShaderParameters, ViewUniformBuffer)
 		SHADER_PARAMETER_STRUCT_REF(FReflectionUniformParameters, ReflectionsParameters)
 		SHADER_PARAMETER_STRUCT_REF(FReflectionCaptureShaderData, ReflectionCaptureData)
 		SHADER_PARAMETER_RDG_UNIFORM_BUFFER(FForwardLightData, ForwardLightData)
-	END_SHADER_PARAMETER_STRUCT()
-};
-
-class FHairEnvironmentLightingVS : public FGlobalShader
-{
-	DECLARE_GLOBAL_SHADER(FHairEnvironmentLightingVS);
-	SHADER_USE_PARAMETER_STRUCT(FHairEnvironmentLightingVS, FGlobalShader)
-	//using FPermutationDomain = FHairEnvironmentLighting::FPermutationDomain;
-
-	BEGIN_SHADER_PARAMETER_STRUCT(FParameters, )
-		SHADER_PARAMETER_STRUCT_INCLUDE(FHairEnvironmentLighting::FParameters, Common)
-	END_SHADER_PARAMETER_STRUCT()
-
-	static void ModifyCompilationEnvironment(const FGlobalShaderPermutationParameters& Parameters, FShaderCompilerEnvironment& OutEnvironment)
-	{
-		FHairEnvironmentLighting::ModifyCompilationEnvironment(Parameters, OutEnvironment);
-		OutEnvironment.SetDefine(TEXT("LIGHTING_VS"), 1);
-	}
-
-	static bool ShouldCompilePermutation(const FGlobalShaderPermutationParameters& Parameters)
-	{
-		return IsHairStrandsSupported(EHairStrandsShaderType::Strands, Parameters.Platform);
-	}
-};
-
-class FHairEnvironmentLightingPS : public FGlobalShader
-{
-	DECLARE_GLOBAL_SHADER(FHairEnvironmentLightingPS);
-	SHADER_USE_PARAMETER_STRUCT(FHairEnvironmentLightingPS, FGlobalShader)
-	using FPermutationDomain = FHairEnvironmentLighting::FPermutationDomain;
-
-	BEGIN_SHADER_PARAMETER_STRUCT(FParameters, )
-		SHADER_PARAMETER_STRUCT_INCLUDE(FHairEnvironmentLighting::FParameters, Common)
-		SHADER_PARAMETER_STRUCT_INCLUDE(LumenRadianceCache::FRadianceCacheInterpolationParameters, RadianceCache)
-		SHADER_PARAMETER_STRUCT_INCLUDE(ShaderDrawDebug::FShaderParameters, ShaderDrawParameters)
-		SHADER_PARAMETER_STRUCT_INCLUDE(FHairStrandsDebugData::FWriteParameters, DebugData)
+		SHADER_PARAMETER_RDG_UNIFORM_BUFFER(FHairStrandsViewUniformParameters, HairStrands)
+		SHADER_PARAMETER_RDG_UNIFORM_BUFFER(FVirtualVoxelParameters, VirtualVoxel)
+		SHADER_PARAMETER_RDG_BUFFER_UAV(RWBuffer, OutLightingBuffer)
 		RENDER_TARGET_BINDING_SLOTS()
 	END_SHADER_PARAMETER_STRUCT()
 
 	static void ModifyCompilationEnvironment(const FGlobalShaderPermutationParameters& Parameters, FShaderCompilerEnvironment& OutEnvironment)
 	{
-		FHairEnvironmentLighting::ModifyCompilationEnvironment(Parameters, OutEnvironment);
+		FForwardLightingParameters::ModifyCompilationEnvironment(Parameters.Platform, OutEnvironment);
 		OutEnvironment.SetDefine(TEXT("LIGHTING_PS"), 1);
 	}
 
 	static bool ShouldCompilePermutation(const FGlobalShaderPermutationParameters& Parameters)
 	{
-		return FHairEnvironmentLighting::ShouldCompilePermutation(Parameters);
+		// Compile debug permutation only for the uniform integrator
+		FPermutationDomain PermutationVector(Parameters.PermutationId);
+		if (PermutationVector.Get<FIntegrationType>() != int32(EHairLightingIntegrationType::Uniform) && PermutationVector.Get<FDebug>())
+		{
+			return false;
+		}
+
+		// If using Lumen, no static lighting
+		if (PermutationVector.Get<FStaticLighting>() && PermutationVector.Get<FLumen>())
+		{
+			return false;
+		}
+
+		// If using another integrator than SH integrator, no static lighting
+		if (PermutationVector.Get<FStaticLighting>() && PermutationVector.Get<FIntegrationType>() != int32(EHairLightingIntegrationType::SH))
+		{
+			return false;
+		}
+
+		return IsHairStrandsSupported(EHairStrandsShaderType::Strands, Parameters.Platform);
+	}
+
+	static FPermutationDomain RemapPermutation(FPermutationDomain PermutationVector)
+	{
+		if (PermutationVector.Get<FIntegrationType>() != int32(EHairLightingIntegrationType::Uniform))
+		{
+			PermutationVector.Set<FDebug>(false);
+		}
+		return PermutationVector;
 	}
 };
 
@@ -387,28 +364,27 @@
 	const bool bHasBoxCaptures = (View.NumBoxReflectionCaptures > 0);
 	const bool bHasSphereCaptures = (View.NumSphereReflectionCaptures > 0);
 	FHairEnvironmentLightingPS::FParameters* ParametersPS = GraphBuilder.AllocParameters<FHairEnvironmentLightingPS::FParameters>();
-	FHairEnvironmentLighting::FParameters* PassParameters = &ParametersPS->Common;
 
-	PassParameters->SceneColorTexture = SceneColorTexture;
-	PassParameters->HairEnergyLUTTexture = GetHairLUT(GraphBuilder, View, HairLUTType_MeanEnergy);
-	PassParameters->bHasStaticLighting = bHasStaticLighting ? 1 : 0;
-	PassParameters->SkyDiffuseLighting = GetSkyDiffuseLightingParameters(Scene->SkyLight, 1.0f /*DynamicBentNormalAO*/);
-	PassParameters->HairDistanceThreshold = FMath::Max(GHairStrandsSkyLightingDistanceThreshold, 1.f);
-	PassParameters->bHairUseViewHairCount = VisibilityData.ViewHairCountTexture && GHairStrandsSkyLightingUseHairCountTexture ? 1 : 0;
-	PassParameters->MultipleScatterSampleCount = FMath::Max(uint32(GHairStrandsSkyLightingSampleCount), 1u);
-	PassParameters->HairDualScatteringRoughnessOverride = GetHairDualScatteringRoughnessOverride();
-	PassParameters->TransmissionDensityScaleFactor = FMath::Max(0.f, GHairStrandsTransmissionDensityScaleFactor);
-	PassParameters->PreIntegratedGF = GSystemTextures.PreintegratedGF->GetRenderTargetItem().ShaderResourceTexture;
-	PassParameters->PreIntegratedGFSampler = TStaticSamplerState<SF_Bilinear, AM_Clamp, AM_Clamp, AM_Clamp>::GetRHI();
-	PassParameters->SceneTextures = SceneTextures;
-	PassParameters->VirtualVoxel = VirtualVoxelResources.UniformBuffer;
-	PassParameters->HairStrands = HairStrands::BindHairStrandsViewUniformParameters(View);
-	PassParameters->ViewUniformBuffer = View.ViewUniformBuffer;
-	PassParameters->ReflectionCaptureData = View.ReflectionCaptureUniformBuffer;
+	ParametersPS->SceneColorTexture = SceneColorTexture;
+	ParametersPS->HairEnergyLUTTexture = GetHairLUT(GraphBuilder, View, HairLUTType_MeanEnergy);
+	ParametersPS->bHasStaticLighting = bHasStaticLighting ? 1 : 0;
+	ParametersPS->SkyDiffuseLighting = GetSkyDiffuseLightingParameters(Scene->SkyLight, 1.0f /*DynamicBentNormalAO*/);
+	ParametersPS->HairDistanceThreshold = FMath::Max(GHairStrandsSkyLightingDistanceThreshold, 1.f);
+	ParametersPS->bHairUseViewHairCount = VisibilityData.ViewHairCountTexture && GHairStrandsSkyLightingUseHairCountTexture ? 1 : 0;
+	ParametersPS->MultipleScatterSampleCount = FMath::Max(uint32(GHairStrandsSkyLightingSampleCount), 1u);
+	ParametersPS->HairDualScatteringRoughnessOverride = GetHairDualScatteringRoughnessOverride();
+	ParametersPS->TransmissionDensityScaleFactor = FMath::Max(0.f, GHairStrandsTransmissionDensityScaleFactor);
+	ParametersPS->PreIntegratedGF = GSystemTextures.PreintegratedGF->GetRenderTargetItem().ShaderResourceTexture;
+	ParametersPS->PreIntegratedGFSampler = TStaticSamplerState<SF_Bilinear, AM_Clamp, AM_Clamp, AM_Clamp>::GetRHI();
+	ParametersPS->SceneTextures = SceneTextures;
+	ParametersPS->VirtualVoxel = VirtualVoxelResources.UniformBuffer;
+	ParametersPS->HairStrands = HairStrands::BindHairStrandsViewUniformParameters(View);
+	ParametersPS->ViewUniformBuffer = View.ViewUniformBuffer;
+	ParametersPS->ReflectionCaptureData = View.ReflectionCaptureUniformBuffer;
 	{
 		FReflectionUniformParameters ReflectionUniformParameters;
 		SetupReflectionUniformParameters(View, ReflectionUniformParameters);
-		PassParameters->ReflectionsParameters = CreateUniformBufferImmediate(ReflectionUniformParameters, UniformBuffer_SingleDraw);
+		ParametersPS->ReflectionsParameters = CreateUniformBufferImmediate(ReflectionUniformParameters, UniformBuffer_SingleDraw);
 	}
 
 	if (LightingType == EHairLightingSourceType::Lumen)
@@ -417,8 +393,8 @@
 		const LumenRadianceCache::FRadianceCacheInputs RadianceCacheInputs = LumenScreenProbeGatherRadianceCache::SetupRadianceCacheInputs(View);
 		LumenRadianceCache::GetInterpolationParameters(View, GraphBuilder, RadianceCacheState, RadianceCacheInputs, ParametersPS->RadianceCache);
 	}
-	PassParameters->ForwardLightData = View.ForwardLightingResources.ForwardLightUniformBuffer;
-	PassParameters->OutLightingBuffer = nullptr;
+	ParametersPS->ForwardLightData = View.ForwardLightingResources.ForwardLightUniformBuffer;
+	ParametersPS->OutLightingBuffer = nullptr;
 
 	if (ShaderDrawDebug::IsEnabled(View))
 	{
@@ -431,11 +407,11 @@
 	}
 
 	FHairEnvironmentLightingPS::FPermutationDomain PermutationVector;
-	PermutationVector.Set<FHairEnvironmentLighting::FIntegrationType>(uint32(IntegrationType));
-	PermutationVector.Set<FHairEnvironmentLighting::FLumen>(LightingType == EHairLightingSourceType::Lumen);
-	PermutationVector.Set<FHairEnvironmentLighting::FDebug>(DebugData != nullptr);
-	PermutationVector.Set<FHairEnvironmentLighting::FStaticLighting>(bHasStaticLighting);
-	PermutationVector = FHairEnvironmentLighting::RemapPermutation(PermutationVector);
+	PermutationVector.Set<FHairEnvironmentLightingPS::FIntegrationType>(uint32(IntegrationType));
+	PermutationVector.Set<FHairEnvironmentLightingPS::FLumen>(LightingType == EHairLightingSourceType::Lumen);
+	PermutationVector.Set<FHairEnvironmentLightingPS::FDebug>(DebugData != nullptr);
+	PermutationVector.Set<FHairEnvironmentLightingPS::FStaticLighting>(bHasStaticLighting);
+	PermutationVector = FHairEnvironmentLightingPS::RemapPermutation(PermutationVector);
 
 	FIntPoint ViewportResolution = VisibilityData.SampleLightingViewportResolution;
 	const FViewInfo* CapturedView = &View;
@@ -446,13 +422,13 @@
 	ParametersPS->RenderTargets[0] = FRenderTargetBinding(VisibilityData.SampleLightingTexture, ERenderTargetLoadAction::ELoad);
 
 	GraphBuilder.AddPass(
-		RDG_EVENT_NAME("HairEnvLightingPS(%s)", LightingType == EHairLightingSourceType::SceneColor ? TEXT("SceneScatter") : (LightingType == EHairLightingSourceType::Lumen ? TEXT("Lumen") : TEXT("ReflectionProbe")) ),
+		RDG_EVENT_NAME("HairStrands::EnvLightingPS(%s)", LightingType == EHairLightingSourceType::SceneColor ? TEXT("SceneScatter") : (LightingType == EHairLightingSourceType::Lumen ? TEXT("Lumen") : TEXT("ReflectionProbe")) ),
 		ParametersPS,
 		ERDGPassFlags::Raster,
 		[ParametersPS, VertexShader, PixelShader, ViewportResolution, CapturedView](FRHICommandList& RHICmdList)
 	{
 		FHairEnvironmentLightingVS::FParameters ParametersVS;
-		ParametersVS.Common = ParametersPS->Common;
+		ParametersVS.HairStrands = ParametersPS->HairStrands;
 
 		FGraphicsPipelineStateInitializer GraphicsPSOInit;
 		RHICmdList.ApplyCachedRenderTargets(GraphicsPSOInit);
