Fix ShaderDraw / ShaderPrint in stereo mode.

diff --git a/Engine/Source/Runtime/Renderer/Private/ShaderDebug.cpp b/Engine/Source/Runtime/Renderer/Private/ShaderDebug.cpp
--- a/Engine/Source/Runtime/Renderer/Private/ShaderDebug.cpp
+++ b/Engine/Source/Runtime/Renderer/Private/ShaderDebug.cpp
@@ -258,11 +258,12 @@ namespace ShaderDrawDebug
 		ValidateShaderParameters(VertexShader, PassParameters->ShaderDrawVSParameters);
 		ClearUnusedGraphResources(VertexShader, &PassParameters->ShaderDrawVSParameters, { IndirectBuffer });
 
+		const FIntRect Viewport = View.ViewRect;
 		GraphBuilder.AddPass(
 			RDG_EVENT_NAME("ShaderDebug::Draw"),
 			PassParameters,
 			ERDGPassFlags::Raster,
-			[VertexShader, PixelShader, PassParameters, IndirectBuffer](FRHICommandList& RHICmdList)
+			[VertexShader, PixelShader, PassParameters, IndirectBuffer, Viewport](FRHICommandList& RHICmdList)
 			{
 				// Marks the indirect draw parameter as used by the pass, given it's not used directly by any of the shaders.
 				PassParameters->ShaderDrawVSParameters.IndirectBuffer->MarkResourceAsUsed();
@@ -278,6 +279,7 @@ namespace ShaderDrawDebug
 				GraphicsPSOInit.BoundShaderState.PixelShaderRHI = PixelShader.GetPixelShader();
 				SetGraphicsPipelineState(RHICmdList, GraphicsPSOInit, 0);
 
+				RHICmdList.SetViewport(Viewport.Min.X, Viewport.Min.Y, 0.0f, Viewport.Max.X, Viewport.Max.Y, 1.0f);
 				SetShaderParameters(RHICmdList, VertexShader, VertexShader.GetVertexShader(), PassParameters->ShaderDrawVSParameters);
 				SetShaderParameters(RHICmdList, PixelShader, PixelShader.GetPixelShader(), PassParameters->ShaderDrawPSParameters);
 
diff --git a/Engine/Source/Runtime/Renderer/Private/ShaderPrint.cpp b/Engine/Source/Runtime/Renderer/Private/ShaderPrint.cpp
--- a/Engine/Source/Runtime/Renderer/Private/ShaderPrint.cpp
+++ b/Engine/Source/Runtime/Renderer/Private/ShaderPrint.cpp
@@ -336,7 +336,7 @@ namespace ShaderPrint
 
 		FComputeShaderUtils::AddPass(
 			GraphBuilder,
-			RDG_EVENT_NAME("ShaderPrintBeginView"),
+			RDG_EVENT_NAME("ShaderPrint::BeginView"),
 			ComputeShader,
 			PassParameters,
 			FIntVector(1, 1, 1));
@@ -346,9 +346,9 @@ namespace ShaderPrint
 	{
 		check(OutputTexture.IsValid());
 
-		RDG_EVENT_SCOPE(GraphBuilder, "ShaderPrintDrawView");
+		RDG_EVENT_SCOPE(GraphBuilder, "ShaderPrint::DrawView");
 
-		FIntRect OutputRect = OutputTexture.ViewRect;
+		const FIntRect Viewport = OutputTexture.ViewRect;
 	
 		// Initialize graph managed resources
 		// Symbols buffer contains Count + 1 elements. The first element is only used as a counter.
@@ -377,7 +377,7 @@ namespace ShaderPrint
 
 			FComputeShaderUtils::AddPass(
 				GraphBuilder, 
-				RDG_EVENT_NAME("BuildIndirectDispatchArgs"), 
+				RDG_EVENT_NAME("ShaderPrint::BuildIndirectDispatchArgs"), 
 				ComputeShader, PassParameters,
 				FIntVector(1, 1, 1));
 		}
@@ -395,7 +395,7 @@ namespace ShaderPrint
 
 			FComputeShaderUtils::AddPass(
 				GraphBuilder,
-				RDG_EVENT_NAME("BuildSymbolBuffer"),
+				RDG_EVENT_NAME("ShaderPrint::BuildSymbolBuffer"),
 				ComputeShader, PassParameters,
 				IndirectDispatchArgsBuffer, 0);
 		}
@@ -412,7 +412,7 @@ namespace ShaderPrint
 
 			FComputeShaderUtils::AddPass(
 				GraphBuilder,
-				RDG_EVENT_NAME("BuildIndirectDrawArgs"),
+				RDG_EVENT_NAME("ShaderPrint::BuildIndirectDrawArgs"),
 				ComputeShader, PassParameters,
 				FIntVector(1, 1, 1));
 		}
@@ -431,12 +431,11 @@ namespace ShaderPrint
 			PassParameters->IndirectDrawArgsBuffer = IndirectDrawArgsBuffer;
 
 			GraphBuilder.AddPass(
-				RDG_EVENT_NAME("DrawSymbols"),
+				RDG_EVENT_NAME("ShaderPrint::DrawSymbols"),
 				PassParameters,
 				ERDGPassFlags::Raster,
-				[VertexShader, PixelShader, PassParameters, OutputRect](FRHICommandList& RHICmdList)
+				[VertexShader, PixelShader, PassParameters, Viewport](FRHICommandList& RHICmdList)
 			{
-				RHICmdList.SetViewport(OutputRect.Min.X, OutputRect.Min.Y, 0.0f, OutputRect.Min.X, OutputRect.Min.Y, 1.0f);
 				
 				FGraphicsPipelineStateInitializer GraphicsPSOInit;
 				RHICmdList.ApplyCachedRenderTargets(GraphicsPSOInit);
@@ -449,6 +448,7 @@ namespace ShaderPrint
 				GraphicsPSOInit.BoundShaderState.PixelShaderRHI = PixelShader.GetPixelShader();
 				SetGraphicsPipelineState(RHICmdList, GraphicsPSOInit, 0);
 
+				RHICmdList.SetViewport(Viewport.Min.X, Viewport.Min.Y, 0.0f, Viewport.Max.X, Viewport.Max.Y, 1.0f);
 				SetShaderParameters(RHICmdList, VertexShader, VertexShader.GetVertexShader(), *PassParameters);
 				SetShaderParameters(RHICmdList, PixelShader, PixelShader.GetPixelShader(), *PassParameters);
 
