Ignore false positive static analysis warning

diff --git a/Engine/Source/Runtime/Renderer/Private/VirtualShadowMaps/VirtualShadowMapArray.cpp b/Engine/Source/Runtime/Renderer/Private/VirtualShadowMaps/VirtualShadowMapArray.cpp
--- a/Engine/Source/Runtime/Renderer/Private/VirtualShadowMaps/VirtualShadowMapArray.cpp
+++ b/Engine/Source/Runtime/Renderer/Private/VirtualShadowMaps/VirtualShadowMapArray.cpp
@@ -973,7 +973,7 @@ void FVirtualShadowMapVisualizeLightSearch::CheckLight(const FLightSceneProxy* C
 	CheckKey.Fields.bDirectionalLight = CheckProxy->GetLightType() == LightType_Directional;
 	CheckKey.Fields.bExists = 1;
 
-	if (CheckKey.Packed > FoundKey.Packed)
+	if (CheckKey.Packed > FoundKey.Packed)		//-V547
 	{
 		FoundKey = CheckKey;
 		FoundProxy = CheckProxy;
