Fix missing shadow when rendering hair strands in stereo mode.

HairGroup/Voxelization/DeepShadow are not run only once for all views.

diff --git a/Engine/Source/Runtime/Renderer/Private/HairStrands/HairStrandsRendering.cpp b/Engine/Source/Runtime/Renderer/Private/HairStrands/HairStrandsRendering.cpp
--- a/Engine/Source/Runtime/Renderer/Private/HairStrands/HairStrandsRendering.cpp
+++ b/Engine/Source/Runtime/Renderer/Private/HairStrands/HairStrandsRendering.cpp
@@ -97,13 +97,25 @@ void RenderHairPrePass(
 	TArray<FViewInfo>& Views,
 	FInstanceCullingManager& InstanceCullingManager)
 {
-	// #hair_todo: Add multi-view
 	for (FViewInfo& View : Views)
 	{
 		const bool bIsViewCompatible = IsHairStrandsEnabled(EHairStrandsShaderType::Strands, View.GetShaderPlatform());
 		if (!View.Family || !bIsViewCompatible)
 			continue;
 
+		// For stereo rendering, hair groups/voxelization/deep-shadow are only produced once
+		if (IStereoRendering::IsStereoEyeView(View))
+		{
+			if (IStereoRendering::IsASecondaryView(View))
+			{
+				// No need to copy the view state (i.e., HairStrandsViewStateData) as it is only used for 
+				// voxelization feedback (only done for the first view in stereo) and Path-Tracer invalidation 
+				// (not supporting stereo)
+				Views[1].HairStrandsViewData = Views[0].HairStrandsViewData;
+				return;
+			}
+		}
+
 		const ERHIFeatureLevel::Type FeatureLevel = Scene->GetFeatureLevel();
 
 		// Allocate voxel page allocation readback buffers
@@ -134,7 +146,7 @@ void RenderHairPrePass(
 			AddMeshDrawTransitionPass(GraphBuilder, View, View.HairStrandsViewData.MacroGroupDatas);
 		}
 		RenderHairStrandsDeepShadows(GraphBuilder, Scene, View, InstanceCullingManager);
-		GraphBuilder.AddDispatchHint();;
+		GraphBuilder.AddDispatchHint();
 	}
 }
 
