Fix unnecessary invalidation of VSMs for dynamic primitives that are culled in the shadow pass.
Most relevant case is in editor with the "arrow" primitives causing invalidations even though they do not cast shadows.

diff --git a/Engine/Source/Runtime/Renderer/Private/GPUScene.cpp b/Engine/Source/Runtime/Renderer/Private/GPUScene.cpp
--- a/Engine/Source/Runtime/Renderer/Private/GPUScene.cpp
+++ b/Engine/Source/Runtime/Renderer/Private/GPUScene.cpp
@@ -1640,7 +1640,7 @@ struct FUploadDataSourceAdapterDynamicPrimitives
 	TArray<uint32, SceneRenderingAllocator> PrimitivesIds;
 };
 
-void FGPUScene::UploadDynamicPrimitiveShaderDataForViewInternal(FRDGBuilder& GraphBuilder, FScene *Scene, FViewInfo& View)
+void FGPUScene::UploadDynamicPrimitiveShaderDataForViewInternal(FRDGBuilder& GraphBuilder, FScene *Scene, FViewInfo& View, bool bIsShadowView)
 {
 	LLM_SCOPE_BYTAG(GPUScene);
 
@@ -1678,7 +1678,7 @@ void FGPUScene::UploadDynamicPrimitiveShaderDataForViewInternal(FRDGBuilder& Gra
 		ensure(UploadIdStart < DynamicPrimitivesOffset);
 		ensure(InstanceIdStart != INDEX_NONE);
 
-		if (Scene != nullptr && Scene->VirtualShadowMapArrayCacheManager != nullptr)
+		if (bIsShadowView && Scene != nullptr && Scene->VirtualShadowMapArrayCacheManager != nullptr)
 		{
 			// Enqueue cache invalidations for all dynamic primitives' instances, as they will be removed this frame and are not associated
 			// with any particular FPrimitiveSceneInfo. Will occur on the next call to UpdateAllPrimitiveSceneInfos
@@ -1822,11 +1822,11 @@ void FGPUScene::Update(FRDGBuilder& GraphBuilder, FScene& Scene)
 	}
 }
 
-void FGPUScene::UploadDynamicPrimitiveShaderDataForView(FRDGBuilder& GraphBuilder, FScene *Scene, FViewInfo& View)
+void FGPUScene::UploadDynamicPrimitiveShaderDataForView(FRDGBuilder& GraphBuilder, FScene *Scene, FViewInfo& View, bool bIsShadowView)
 {
 	if (bIsEnabled)
 	{
-		UploadDynamicPrimitiveShaderDataForViewInternal(GraphBuilder, Scene, View);
+		UploadDynamicPrimitiveShaderDataForViewInternal(GraphBuilder, Scene, View, bIsShadowView);
 	}
 }
 
diff --git a/Engine/Source/Runtime/Renderer/Private/GPUScene.h b/Engine/Source/Runtime/Renderer/Private/GPUScene.h
--- a/Engine/Source/Runtime/Renderer/Private/GPUScene.h
+++ b/Engine/Source/Runtime/Renderer/Private/GPUScene.h
@@ -211,7 +211,7 @@ public:
 	/**
 	 * Upload primitives from View.DynamicPrimitiveCollector.
 	 */
-	void UploadDynamicPrimitiveShaderDataForView(FRDGBuilder& GraphBuilder, FScene* Scene, FViewInfo& View);
+	void UploadDynamicPrimitiveShaderDataForView(FRDGBuilder& GraphBuilder, FScene* Scene, FViewInfo& View, bool bIsShadowView = false);
 
 	/**
 	 * Pull all pending updates from Scene and upload primitive & instance data.
@@ -398,7 +398,7 @@ private:
 	template<typename FUploadDataSourceAdapter>
 	void UploadGeneral(FRHICommandListImmediate& RHICmdList, FScene* Scene, const FUploadDataSourceAdapter& UploadDataSourceAdapter, const FGPUSceneBufferState &BufferState);
 
-	void UploadDynamicPrimitiveShaderDataForViewInternal(FRDGBuilder& GraphBuilder, FScene* Scene, FViewInfo& View);
+	void UploadDynamicPrimitiveShaderDataForViewInternal(FRDGBuilder& GraphBuilder, FScene* Scene, FViewInfo& View, bool bIsShadowView);
 
 	void UpdateInternal(FRDGBuilder& GraphBuilder, FScene& Scene);
 
diff --git a/Engine/Source/Runtime/Renderer/Private/ShadowDepthRendering.cpp b/Engine/Source/Runtime/Renderer/Private/ShadowDepthRendering.cpp
--- a/Engine/Source/Runtime/Renderer/Private/ShadowDepthRendering.cpp
+++ b/Engine/Source/Runtime/Renderer/Private/ShadowDepthRendering.cpp
@@ -1785,32 +1785,32 @@ void FSceneRenderer::RenderShadowDepthMaps(FRDGBuilder& GraphBuilder, FInstanceC
 	{
 		for (FProjectedShadowInfo* ProjectedShadowInfo : ShadowMapAtlas.Shadows)
 		{
-			Scene->GPUScene.UploadDynamicPrimitiveShaderDataForView(GraphBuilder, Scene, *ProjectedShadowInfo->ShadowDepthView);
+			Scene->GPUScene.UploadDynamicPrimitiveShaderDataForView(GraphBuilder, Scene, *ProjectedShadowInfo->ShadowDepthView, true);
 		}
 	}
 	for (FSortedShadowMapAtlas& ShadowMap : SortedShadowsForShadowDepthPass.ShadowMapCubemaps)
 	{
 		check(ShadowMap.Shadows.Num() == 1);
 		FProjectedShadowInfo* ProjectedShadowInfo = ShadowMap.Shadows[0];
-		Scene->GPUScene.UploadDynamicPrimitiveShaderDataForView(GraphBuilder, Scene, *ProjectedShadowInfo->ShadowDepthView);
+		Scene->GPUScene.UploadDynamicPrimitiveShaderDataForView(GraphBuilder, Scene, *ProjectedShadowInfo->ShadowDepthView, true);
 	}
 	for (FProjectedShadowInfo* ProjectedShadowInfo : SortedShadowsForShadowDepthPass.PreshadowCache.Shadows)
 	{
 		if (!ProjectedShadowInfo->bDepthsCached)
 		{
-			Scene->GPUScene.UploadDynamicPrimitiveShaderDataForView(GraphBuilder, Scene, *ProjectedShadowInfo->ShadowDepthView);
+			Scene->GPUScene.UploadDynamicPrimitiveShaderDataForView(GraphBuilder, Scene, *ProjectedShadowInfo->ShadowDepthView, true);
 		}
 	}
 	for (const FSortedShadowMapAtlas& ShadowMapAtlas : SortedShadowsForShadowDepthPass.TranslucencyShadowMapAtlases)
 	{
 		for (FProjectedShadowInfo* ProjectedShadowInfo : ShadowMapAtlas.Shadows)
 		{
-			Scene->GPUScene.UploadDynamicPrimitiveShaderDataForView(GraphBuilder, Scene, *ProjectedShadowInfo->ShadowDepthView);
+			Scene->GPUScene.UploadDynamicPrimitiveShaderDataForView(GraphBuilder, Scene, *ProjectedShadowInfo->ShadowDepthView, true);
 		}
 	}
 	for (FProjectedShadowInfo* ProjectedShadowInfo : SortedShadowsForShadowDepthPass.VirtualShadowMapShadows)
 	{
-		Scene->GPUScene.UploadDynamicPrimitiveShaderDataForView(GraphBuilder, Scene, *ProjectedShadowInfo->ShadowDepthView);
+		Scene->GPUScene.UploadDynamicPrimitiveShaderDataForView(GraphBuilder, Scene, *ProjectedShadowInfo->ShadowDepthView, true);
 	}
 
 	// Begin new deferred culling bacthing scope to catch shadow render passes, as there can use dynamic primitives that have not been uploaded before 
