Make the fetching of the Stencil/depth in the PostProcessMaterial datadriven instead of just testing the feature level. This enable some consoles using the mobile renderer to still have the possibility to use the HiStencil and fetch directly the stencil buffer in the shader.

diff --git a/Engine/Plugins/FX/Niagara/Shaders/Private/NiagaraDataInterfaceGBuffer.ush b/Engine/Plugins/FX/Niagara/Shaders/Private/NiagaraDataInterfaceGBuffer.ush
--- a/Engine/Plugins/FX/Niagara/Shaders/Private/NiagaraDataInterfaceGBuffer.ush
+++ b/Engine/Plugins/FX/Niagara/Shaders/Private/NiagaraDataInterfaceGBuffer.ush
@@ -224,7 +224,11 @@ void DIGBuffer_DecodeCustomStencil(in float2 ScreenUV, in bool bApplyViewportOff
 	OutValue = GBufferData.CustomStencil;
 #elif NDIGBUFFER_IS_MOBILE
 	IsValid = true;
-	OutValue = Texture2DSampleLevel(MobileSceneTextures.MobileCustomStencilTexture, MobileSceneTextures.MobileCustomStencilTextureSampler, ScreenUV, 0).r * 255.0f;
+	#if MOBILE_SUPPORT_FETCH_BINDED_CUSTOM_STENCIL_BUFFER
+		OutValue = MobileSceneTextures.CustomStencilTexture.Load(int3(Parameters.SvPosition.xy, 0)).x;
+	#else
+		OutValue = Texture2DSampleLevel(MobileSceneTextures.MobileCustomStencilTexture, MobileSceneTextures.MobileCustomStencilTextureSampler, ScreenUV, 0).r * 255.0f;
+	#endif
 #else
 	IsValid = false;
 	OutValue = 0;
diff --git a/Engine/Shaders/Private/MaterialTemplate.ush b/Engine/Shaders/Private/MaterialTemplate.ush
--- a/Engine/Shaders/Private/MaterialTemplate.ush
+++ b/Engine/Shaders/Private/MaterialTemplate.ush
@@ -2158,9 +2158,13 @@ MaterialFloat4 MobileSceneTextureLookup(inout FMaterialPixelParameters Parameter
 	}
 	else if (SceneTextureId == PPI_CustomStencil)
 	{
-		MaterialFloat Stencil = Texture2DSample(MobileSceneTextures.MobileCustomStencilTexture, MobileSceneTextures.MobileCustomStencilTextureSampler, UV).r*255.0;
+#if MOBILE_SUPPORT_FETCH_BINDED_CUSTOM_STENCIL_BUFFER
+		MaterialFloat Stencil = MobileSceneTextures.CustomStencilTexture.Load(int3(Parameters.SvPosition.xy, 0)).x;
+#else
+		MaterialFloat Stencil = Texture2DSample(MobileSceneTextures.MobileCustomStencilTexture, MobileSceneTextures.MobileCustomStencilTextureSampler, UV).r;
 		Stencil = floor(Stencil + 0.5);
-		return MaterialFloat4(Stencil.rrr, 0.f);
+#endif
+	return MaterialFloat4(Stencil.rrr, 0.f);
 	}
 #endif// SCENE_TEXTURES_DISABLED
 
diff --git a/Engine/Shaders/Private/PostProcessMaterialShaders.usf b/Engine/Shaders/Private/PostProcessMaterialShaders.usf
--- a/Engine/Shaders/Private/PostProcessMaterialShaders.usf
+++ b/Engine/Shaders/Private/PostProcessMaterialShaders.usf
@@ -22,8 +22,8 @@
 #error POST_PROCESS_MATERIAL must be defined to non-zero in the shader compilation environment.
 #endif
 
-#if (FEATURE_LEVEL <= FEATURE_LEVEL_ES3_1) != POST_PROCESS_MATERIAL_MOBILE
-#error POST_PROCESS_MATERIAL_MOBILE define does not match feature level.
+#ifndef MOBILE_SUPPORT_FETCH_BINDED_CUSTOM_STENCIL_BUFFER
+#define MOBILE_SUPPORT_FETCH_BINDED_CUSTOM_STENCIL_BUFFER 0
 #endif
 
 // Must match ESceneTextureId
@@ -119,7 +119,7 @@ FMaterialVertexParameters GetPostProcessMaterialVSParameters(float2 UV)
 	return Result;
 }
 
-#if (FEATURE_LEVEL > FEATURE_LEVEL_ES3_1)
+#if (FEATURE_LEVEL > FEATURE_LEVEL_ES3_1 || MOBILE_SUPPORT_FETCH_BINDED_CUSTOM_STENCIL_BUFFER)
 
 #if VERTEXSHADER
 
@@ -213,6 +213,8 @@ void MainPS(
 	// Grab emissive colour as output
   #if MATERIAL_OUTPUT_OPACITY_AS_ALPHA
 	Alpha = UnlitBSDF.Coverage;
+  #elif (FEATURE_LEVEL <= FEATURE_LEVEL_ES3_1)
+	Alpha = Parameters.BackupSceneColorAlpha;
   #endif
 	const float3 Emissive = BSDF_GETEMISSIVE(UnlitBSDF);
 #else
@@ -270,7 +272,7 @@ void MainPS_VideoOverlay(
 
 #endif
 
-#else // FEATURE_LEVEL_ES3_1
+#else // FEATURE_LEVEL > FEATURE_LEVEL_ES3_1 || MOBILE_SUPPORT_FETCH_BINDED_CUSTOM_STENCIL_BUFFER
 
 Texture2D MobileCustomStencilTexture;
 SamplerState MobileCustomStencilTextureSampler;
@@ -419,4 +421,4 @@ void MainPS(
 #endif
 	OutColor = FullSceneColor;
 }
-#endif //(FEATURE_LEVEL > FEATURE_LEVEL_ES3_1)
+#endif // FEATURE_LEVEL > FEATURE_LEVEL_ES3_1 || MOBILE_SUPPORT_FETCH_BINDED_CUSTOM_STENCIL_BUFFER
\ No newline at end of file
diff --git a/Engine/Source/Runtime/RHI/Private/RHI.cpp b/Engine/Source/Runtime/RHI/Private/RHI.cpp
--- a/Engine/Source/Runtime/RHI/Private/RHI.cpp
+++ b/Engine/Source/Runtime/RHI/Private/RHI.cpp
@@ -2253,6 +2253,7 @@ void FGenericDataDrivenShaderPlatformInfo::SetDefaultValues()
 	bSupportsAsyncPipelineCompilation = true;
 	bSupportsGPUSkinCache = true;
 	bSupportsManualVertexFetch = true;
+	bMobileSupportFetchBindedCustomStencilBuffer = false;
 }
 
 void FGenericDataDrivenShaderPlatformInfo::ParseDataDrivenShaderInfo(const FConfigSection& Section, FGenericDataDrivenShaderPlatformInfo& Info)
@@ -2340,6 +2341,7 @@ void FGenericDataDrivenShaderPlatformInfo::ParseDataDrivenShaderInfo(const FConf
 	GET_SECTION_BOOL_HELPER(bSupportsMobileDistanceField);
 	GET_SECTION_BOOL_HELPER(bSupportsFFTBloom);
 	GET_SECTION_BOOL_HELPER(bSupportsVertexShaderLayer);
+	GET_SECTION_BOOL_HELPER(bMobileSupportFetchBindedCustomStencilBuffer);
 	GET_SECTION_BOOL_HELPER(bSupportsBindless);
 #undef GET_SECTION_BOOL_HELPER
 #undef GET_SECTION_INT_HELPER
diff --git a/Engine/Source/Runtime/RHI/Public/RHIDefinitions.h b/Engine/Source/Runtime/RHI/Public/RHIDefinitions.h
--- a/Engine/Source/Runtime/RHI/Public/RHIDefinitions.h
+++ b/Engine/Source/Runtime/RHI/Public/RHIDefinitions.h
@@ -381,6 +381,7 @@ class RHI_API FGenericDataDrivenShaderPlatformInfo
 	uint32 bSupportsInlineRayTracing : 1;
 	uint32 bSupportsRayTracingShaders : 1;
 	uint32 bSupportsVertexShaderLayer : 1;
+	uint32 bMobileSupportFetchBindedCustomStencilBuffer : 1;
 	uint32 bSupportsBindless : 1;
 
 		
@@ -889,6 +890,12 @@ public:
 		check(IsValid(Platform));
 		return Infos[Platform].bSupportsVertexShaderLayer;
 	}
+	
+	static FORCEINLINE_DEBUGGABLE const bool GetMobileSupportFetchBindedCustomStencilBuffer(const FStaticShaderPlatform Platform)
+	{
+		check(IsValid(Platform));
+		return Infos[Platform].bMobileSupportFetchBindedCustomStencilBuffer;
+	}
 
 	static FORCEINLINE_DEBUGGABLE const bool GetSupportsBindless(const FStaticShaderPlatform Platform)
 	{
@@ -2519,6 +2526,7 @@ inline ERHIResourceType GetRHIResourceType(ETextureDimension Dimension)
 	return ERHIResourceType::RRT_None;
 }
 
+
 #if PLATFORM_SUPPORTS_GEOMETRY_SHADERS
 	#define GEOMETRY_SHADER(GeometryShader)	(GeometryShader)
 #else
diff --git a/Engine/Source/Runtime/Renderer/Private/CustomDepthRendering.cpp b/Engine/Source/Runtime/Renderer/Private/CustomDepthRendering.cpp
--- a/Engine/Source/Runtime/Renderer/Private/CustomDepthRendering.cpp
+++ b/Engine/Source/Runtime/Renderer/Private/CustomDepthRendering.cpp
@@ -60,21 +60,35 @@ ECustomDepthMode GetCustomDepthMode()
 	return ECustomDepthMode::Disabled;
 }
 
-bool IsCustomDepthPassWritingStencil(ERHIFeatureLevel::Type FeatureLevel)
+bool IsMobileSeparateDepthStencilRenderTargets(EShaderPlatform Platform)
+{
+	return IsMobilePlatform(Platform) && !FDataDrivenShaderPlatformInfo::GetMobileSupportFetchBindedCustomStencilBuffer(Platform);
+}
+
+bool IsMobileSupportFetchBindedCustomStencilBuffer(EShaderPlatform Platform)
+{
+	return IsMobilePlatform(Platform) && FDataDrivenShaderPlatformInfo::GetMobileSupportFetchBindedCustomStencilBuffer(Platform);
+}
+
+bool IsCustomDepthPassWritingStencil(EShaderPlatform Platform)
 {
 	switch (GetCustomDepthMode())
 	{
-	case ECustomDepthMode::Disabled:
-		return false;
-	case ECustomDepthMode::Enabled:
-		return FeatureLevel <= ERHIFeatureLevel::ES3_1;
+		case ECustomDepthMode::Disabled:
+		{
+			return false;
+		}
+		case ECustomDepthMode::Enabled:
+		{
+			return IsMobileSeparateDepthStencilRenderTargets(Platform);
+		}
 	}
 	return true;
 }
 
-uint32 GetCustomDepthDownsampleFactor(ERHIFeatureLevel::Type FeatureLevel)
+uint32 GetCustomDepthDownsampleFactor(EShaderPlatform Platform)
 {
-	return FeatureLevel <= ERHIFeatureLevel::ES3_1 && CVarMobileCustomDepthDownSample.GetValueOnRenderThread() > 0 ? 2 : 1;
+	return IsMobileSeparateDepthStencilRenderTargets(Platform) && CVarMobileCustomDepthDownSample.GetValueOnRenderThread() > 0 ? 2 : 1;
 }
 
 FCustomDepthTextures FCustomDepthTextures::Create(FRDGBuilder& GraphBuilder, FIntPoint Extent, ERHIFeatureLevel::Type FeatureLevel, uint32 DownsampleFactor)
@@ -86,12 +100,13 @@ FCustomDepthTextures FCustomDepthTextures::Create(FRDGBuilder& GraphBuilder, FIn
 		return {};
 	}
 
-	const bool bWritesCustomStencil = IsCustomDepthPassWritingStencil(FeatureLevel);
+	EShaderPlatform Platform = GShaderPlatformForFeatureLevel[FeatureLevel];
+	const bool bWritesCustomStencil = IsCustomDepthPassWritingStencil(Platform);
 	const FIntPoint CustomDepthExtent = FIntPoint::DivideAndRoundUp(Extent, DownsampleFactor);
 
 	FCustomDepthTextures CustomDepthTextures;
-
-	if (FeatureLevel <= ERHIFeatureLevel::ES3_1)
+	
+	if (IsMobileSeparateDepthStencilRenderTargets(Platform))
 	{
 		const float DepthFar = (float)ERHIZBuffer::FarPlane;
 		const FClearValueBinding DepthFarColor = FClearValueBinding(FLinearColor(DepthFar, DepthFar, DepthFar, DepthFar));
@@ -170,9 +185,6 @@ bool FSceneRenderer::RenderCustomDepthPass(FRDGBuilder& GraphBuilder, const FCus
 	CSV_SCOPED_TIMING_STAT_EXCLUSIVE(RenderCustomDepthPass);
 	RDG_GPU_STAT_SCOPE(GraphBuilder, CustomDepth);
 
-	const bool bMobilePath = (FeatureLevel <= ERHIFeatureLevel::ES3_1);
-	const bool bWritesCustomStencilValues = IsCustomDepthPassWritingStencil(FeatureLevel);
-
 	bool bCustomDepthRendered = false;
 	for (int32 ViewIndex = 0; ViewIndex < Views.Num(); ++ViewIndex)
 	{
@@ -200,7 +212,7 @@ bool FSceneRenderer::RenderCustomDepthPass(FRDGBuilder& GraphBuilder, const FCus
 			const ERenderTargetLoadAction DepthLoadAction = GetLoadActionIfProduced(CustomDepthTextures.Depth, CustomDepthTextures.DepthAction);
 			const ERenderTargetLoadAction StencilLoadAction = GetLoadActionIfProduced(CustomDepthTextures.Depth, CustomDepthTextures.StencilAction);
 
-			if (bMobilePath)
+			if (IsMobileSeparateDepthStencilRenderTargets(View.GetShaderPlatform()))
 			{
 				PassParameters->RenderTargets[0] = FRenderTargetBinding(CustomDepthTextures.MobileDepth, DepthLoadAction);
 				PassParameters->RenderTargets[1] = FRenderTargetBinding(CustomDepthTextures.MobileStencil, StencilLoadAction);
@@ -308,11 +320,12 @@ bool FCustomDepthPassMeshProcessor::TryAddMeshBatch(
 {
 	// Determine the mesh's material and blend mode.
 	const EBlendMode BlendMode = Material.GetBlendMode();
+	EShaderPlatform Platform = GShaderPlatformForFeatureLevel[FeatureLevel];
 	const FMeshDrawingPolicyOverrideSettings OverrideSettings = ComputeMeshOverrideSettings(MeshBatch);
 	const ERasterizerFillMode MeshFillMode = ComputeMeshFillMode(MeshBatch, Material, OverrideSettings);
 	const ERasterizerCullMode MeshCullMode = ComputeMeshCullMode(MeshBatch, Material, OverrideSettings);
 	const bool bIsTranslucent = IsTranslucentBlendMode(BlendMode);
-	const bool bWriteCustomStencilValues = IsCustomDepthPassWritingStencil(FeatureLevel);
+	const bool bWriteCustomStencilValues = IsCustomDepthPassWritingStencil(Platform);
 	float MobileColorValue = 0.0f;
 
 	if (bWriteCustomStencilValues)
@@ -336,18 +349,15 @@ bool FCustomDepthPassMeshProcessor::TryAddMeshBatch(
 		PassDrawRenderState.SetDepthStencilState(StencilStates[(int32)PrimitiveSceneProxy->GetStencilWriteMask()]);
 		PassDrawRenderState.SetStencilRef(CustomDepthStencilValue);
 
-		if (FeatureLevel <= ERHIFeatureLevel::ES3_1)
-		{
-			// On mobile platforms write custom stencil value to color target
-			MobileColorValue = CustomDepthStencilValue / 255.0f;
-		}
+		// On mobile platforms write custom stencil value to color target
+		MobileColorValue = CustomDepthStencilValue / 255.0f;
 	}
 	else
 	{
 		PassDrawRenderState.SetDepthStencilState(TStaticDepthStencilState<true, CF_DepthNearOrEqual>::GetRHI());
 	}
 
-	const bool bUsesMobileColorValue = FeatureLevel <= ERHIFeatureLevel::ES3_1;
+	const bool bUsesMobileColorValue = IsMobileSeparateDepthStencilRenderTargets(Platform);
 
 	bool bResult = true;
 	if (BlendMode == BLEND_Opaque
@@ -447,4 +457,4 @@ FMeshPassProcessor* CreateCustomDepthPassProcessor(const FScene* Scene, const FS
 }
 
 FRegisterPassProcessorCreateFunction RegisterCustomDepthPass(&CreateCustomDepthPassProcessor, EShadingPath::Deferred, EMeshPass::CustomDepth, EMeshPassFlags::MainView);
-FRegisterPassProcessorCreateFunction RegisterMobileCustomDepthPass(&CreateCustomDepthPassProcessor, EShadingPath::Mobile, EMeshPass::CustomDepth, EMeshPassFlags::MainView);
\ No newline at end of file
+FRegisterPassProcessorCreateFunction RegisterMobileCustomDepthPass(&CreateCustomDepthPassProcessor, EShadingPath::Mobile, EMeshPass::CustomDepth, EMeshPassFlags::MainView);
diff --git a/Engine/Source/Runtime/Renderer/Private/CustomDepthRendering.h b/Engine/Source/Runtime/Renderer/Private/CustomDepthRendering.h
--- a/Engine/Source/Runtime/Renderer/Private/CustomDepthRendering.h
+++ b/Engine/Source/Runtime/Renderer/Private/CustomDepthRendering.h
@@ -32,13 +32,19 @@ enum class ECustomDepthMode : uint32
 extern ECustomDepthMode GetCustomDepthMode();
 
 // Returns the requested downsample factor for custom depth textures.
-extern uint32 GetCustomDepthDownsampleFactor(ERHIFeatureLevel::Type InFeatureLevel);
+extern uint32 GetCustomDepthDownsampleFactor(EShaderPlatform Platform);
 
 inline bool IsCustomDepthPassEnabled()
 {
 	return GetCustomDepthMode() != ECustomDepthMode::Disabled;
 }
 
+// Check if the mobile uses separate depth and stencil render targets for the custom renderer
+extern bool IsMobileSeparateDepthStencilRenderTargets(EShaderPlatform Platform);
+
+// Check if the mobile support to directly fetch from the depthStencil buffer
+extern bool IsMobileSupportFetchBindedCustomStencilBuffer(EShaderPlatform Platform);
+
 struct FCustomDepthTextures
 {
 	static FCustomDepthTextures Create(FRDGBuilder& GraphBuilder, FIntPoint Extent, ERHIFeatureLevel::Type FeatureLevel, uint32 DownsampleFactor);
diff --git a/Engine/Source/Runtime/Renderer/Private/PostProcess/PostProcessMaterial.cpp b/Engine/Source/Runtime/Renderer/Private/PostProcess/PostProcessMaterial.cpp
--- a/Engine/Source/Runtime/Renderer/Private/PostProcess/PostProcessMaterial.cpp
+++ b/Engine/Source/Runtime/Renderer/Private/PostProcess/PostProcessMaterial.cpp
@@ -149,7 +149,7 @@ public:
 			}
 			else
 			{
-				return IsFeatureLevelSupported(Parameters.Platform, ERHIFeatureLevel::SM5);
+				return IsFeatureLevelSupported(Parameters.Platform, ERHIFeatureLevel::SM5) || IsMobileSupportFetchBindedCustomStencilBuffer(Parameters.Platform);
 			}
 		}
 		return false;
@@ -169,6 +169,12 @@ public:
 			OutEnvironment.SetDefine(TEXT("POST_PROCESS_MATERIAL_BEFORE_TONEMAP"), (Parameters.MaterialParameters.BlendableLocation != BL_AfterTonemapping) ? 1 : 0);
 		}
 
+
+		if (IsMobileSupportFetchBindedCustomStencilBuffer(Parameters.Platform))
+		{
+			OutEnvironment.SetDefine(TEXT("MOBILE_SUPPORT_FETCH_BINDED_CUSTOM_STENCIL_BUFFER"), 1);
+		}
+
 		OutEnvironment.SetDefine(TEXT("STRATA_ENABLED"), Strata::IsStrataEnabled() ? 1u : 0u);
 	}
 
@@ -257,9 +263,9 @@ void GetMaterialInfo(
 {
 	FMaterialShaderTypes ShaderTypes;
 	{
-		const bool bIsMobile = InFeatureLevel <= ERHIFeatureLevel::ES3_1;
+		const bool bUseMobilePath = IsMobileSeparateDepthStencilRenderTargets(GShaderPlatformForFeatureLevel[InFeatureLevel]);
 		FPostProcessMaterialShader::FPermutationDomain PermutationVector;
-		PermutationVector.Set<FPostProcessMaterialShader::FMobileDimension>(bIsMobile);
+		PermutationVector.Set<FPostProcessMaterialShader::FMobileDimension>(bUseMobilePath);
 
 		ShaderTypes.AddShaderType< FPostProcessMaterialVS>(PermutationVector.ToDimensionValueId());
 		ShaderTypes.AddShaderType< FPostProcessMaterialPS>(PermutationVector.ToDimensionValueId());
@@ -496,7 +502,7 @@ FScreenPassTexture AddPostProcessMaterialPass(
 	// The target color will be decoded if bForceIntermediateTarget is true in any case, but we might still need to decode the input color
 	PostProcessMaterialParameters->bMetalMSAAHDRDecode = Inputs.bMetalMSAAHDRDecode ? 1 : 0;
 
-	if (DepthStencilTexture && !bMobilePlatform)
+	if (DepthStencilTexture && (!bMobilePlatform || IsMobileSupportFetchBindedCustomStencilBuffer(View.GetShaderPlatform())))
 	{
 		PostProcessMaterialParameters->RenderTargets.DepthStencil = FDepthStencilBinding(
 			DepthStencilTexture,
@@ -504,7 +510,7 @@ FScreenPassTexture AddPostProcessMaterialPass(
 			ERenderTargetLoadAction::ELoad,
 			FExclusiveDepthStencil::DepthRead_StencilRead);
 	}
-	else if (!DepthStencilTexture && bMobilePlatform && Material->IsStencilTestEnabled()) // we have to set a default texture for MobileStencilTexture and override the MobileStencilValueRef to make all function to pass the stencil test
+	else if (!DepthStencilTexture && IsMobileSeparateDepthStencilRenderTargets(View.GetShaderPlatform()) && Material->IsStencilTestEnabled()) // we have to set a default texture for MobileStencilTexture and override the MobileStencilValueRef to make all function to pass the stencil test
 	{
 		PostProcessMaterialParameters->MobileCustomStencilTexture = GSystemTextures.GetBlackDummy(GraphBuilder);
 		
diff --git a/Engine/Source/Runtime/Renderer/Private/SceneTextures.cpp b/Engine/Source/Runtime/Renderer/Private/SceneTextures.cpp
--- a/Engine/Source/Runtime/Renderer/Private/SceneTextures.cpp
+++ b/Engine/Source/Runtime/Renderer/Private/SceneTextures.cpp
@@ -533,7 +533,7 @@ FSceneTexturesConfig FSceneTexturesConfig::Create(const FSceneViewFamily& ViewFa
 	Config.ColorFormat = PF_Unknown;
 	Config.ColorClearValue = FClearValueBinding::Black;
 	Config.DepthClearValue = FClearValueBinding::DepthFar;
-	Config.CustomDepthDownsampleFactor = GetCustomDepthDownsampleFactor(Config.FeatureLevel);
+	Config.CustomDepthDownsampleFactor = GetCustomDepthDownsampleFactor(Config.ShaderPlatform);
 	Config.bRequireMultiView = ViewFamily.bRequireMultiView;
 	Config.bIsUsingGBuffers = IsUsingGBuffers(Config.ShaderPlatform);
 
@@ -1096,6 +1096,7 @@ void SetupMobileSceneTextureUniformParameters(
 	FMobileSceneTextureUniformParameters& SceneTextureParameters)
 {
 	const FRDGSystemTextures& SystemTextures = FRDGSystemTextures::Get(GraphBuilder);
+	bool bMobileSupportFetchBindedCustomStencilBuffer = FDataDrivenShaderPlatformInfo::GetMobileSupportFetchBindedCustomStencilBuffer(GShaderPlatformForFeatureLevel[GMaxRHIFeatureLevel]);
 
 	SceneTextureParameters.SceneColorTexture = SystemTextures.Black;
 	SceneTextureParameters.SceneColorTextureSampler = TStaticSamplerState<>::GetRHI();
@@ -1104,6 +1105,7 @@ void SetupMobileSceneTextureUniformParameters(
 	// CustomDepthTexture is a color texture on mobile, with DeviceZ values
 	SceneTextureParameters.CustomDepthTexture = SystemTextures.Black;
 	SceneTextureParameters.CustomDepthTextureSampler = TStaticSamplerState<>::GetRHI();
+	SceneTextureParameters.CustomStencilTexture = SystemTextures.StencilDummySRV;
 	SceneTextureParameters.MobileCustomStencilTexture = SystemTextures.Black;
 	SceneTextureParameters.MobileCustomStencilTextureSampler = TStaticSamplerState<>::GetRHI();
 	SceneTextureParameters.SceneVelocityTexture = SystemTextures.Black;
@@ -1169,14 +1171,23 @@ void SetupMobileSceneTextureUniformParameters(
 		{
 			const FCustomDepthTextures& CustomDepthTextures = SceneTextures->CustomDepth;
 
-			if (HasBeenProduced(CustomDepthTextures.MobileDepth))
+			if (bMobileSupportFetchBindedCustomStencilBuffer)
 			{
-				SceneTextureParameters.CustomDepthTexture = CustomDepthTextures.MobileDepth;
-			}
-
-			if (HasBeenProduced(CustomDepthTextures.MobileStencil) && !EnumHasAnyFlags(CustomDepthTextures.MobileStencil->Desc.Flags, TexCreate_Memoryless))
+				bool bCustomDepthProduced = HasBeenProduced(CustomDepthTextures.Depth);
+				SceneTextureParameters.CustomDepthTexture = bCustomDepthProduced ? CustomDepthTextures.Depth : SystemTextures.DepthDummy;
+				SceneTextureParameters.CustomStencilTexture = bCustomDepthProduced ? CustomDepthTextures.Stencil : SystemTextures.StencilDummySRV;
+			}				
+			else
 			{
-				SceneTextureParameters.MobileCustomStencilTexture = CustomDepthTextures.MobileStencil;
+				if (HasBeenProduced(CustomDepthTextures.MobileDepth))
+				{
+					SceneTextureParameters.CustomDepthTexture = CustomDepthTextures.MobileDepth;
+				}
+
+				if (HasBeenProduced(CustomDepthTextures.MobileStencil) && !EnumHasAnyFlags(CustomDepthTextures.MobileStencil->Desc.Flags, TexCreate_Memoryless))
+				{
+					SceneTextureParameters.MobileCustomStencilTexture = CustomDepthTextures.MobileStencil;
+				}
 			}
 		}
 
@@ -1234,4 +1245,4 @@ ERHIFeatureLevel::Type GetSceneTextureFeatureLevel()
 void CreateSystemTextures(FRDGBuilder& GraphBuilder)
 {
 	FRDGSystemTextures::Create(GraphBuilder);
-}
\ No newline at end of file
+}
diff --git a/Engine/Source/Runtime/Renderer/Public/SceneRenderTargetParameters.h b/Engine/Source/Runtime/Renderer/Public/SceneRenderTargetParameters.h
--- a/Engine/Source/Runtime/Renderer/Public/SceneRenderTargetParameters.h
+++ b/Engine/Source/Runtime/Renderer/Public/SceneRenderTargetParameters.h
@@ -69,6 +69,7 @@ BEGIN_GLOBAL_SHADER_PARAMETER_STRUCT(FMobileSceneTextureUniformParameters, RENDE
 	SHADER_PARAMETER_SAMPLER(SamplerState, SceneDepthTextureSampler)
 	SHADER_PARAMETER_RDG_TEXTURE(Texture2D, CustomDepthTexture)
 	SHADER_PARAMETER_SAMPLER(SamplerState, CustomDepthTextureSampler)
+	SHADER_PARAMETER_RDG_TEXTURE_SRV(Texture2D<uint2>, CustomStencilTexture)
 	SHADER_PARAMETER_RDG_TEXTURE(Texture2D, MobileCustomStencilTexture)
 	SHADER_PARAMETER_SAMPLER(SamplerState, MobileCustomStencilTextureSampler)
 	SHADER_PARAMETER_RDG_TEXTURE(Texture2D, SceneVelocityTexture)
