Add r.VelocityOutputPass render setting and deprecated r.BasePassOutputsVelocity and r.DepthPassMergedWithVelocity.
r.VelocityOutputPass supports the 3 valid settings: depth pass, base pass, after base pass.

diff --git a/Engine/Source/Developer/TargetPlatform/Private/TargetPlatformBase.cpp b/Engine/Source/Developer/TargetPlatform/Private/TargetPlatformBase.cpp
--- a/Engine/Source/Developer/TargetPlatform/Private/TargetPlatformBase.cpp
+++ b/Engine/Source/Developer/TargetPlatform/Private/TargetPlatformBase.cpp
@@ -27,8 +27,8 @@ bool FTargetPlatformBase::UsesDBuffer() const
 
 bool FTargetPlatformBase::UsesBasePassVelocity() const
 {
-	static IConsoleVariable* CVar = IConsoleManager::Get().FindConsoleVariable(TEXT("r.BasePassOutputsVelocity"));
-	return CVar ? (CVar->GetInt() != 0) : false;
+	static IConsoleVariable* CVar = IConsoleManager::Get().FindConsoleVariable(TEXT("r.VelocityOutputPass"));
+	return CVar ? (CVar->GetInt() == 1) : false;
 }
 
 bool FTargetPlatformBase::VelocityEncodeDepth() const
diff --git a/Engine/Source/Runtime/Engine/Classes/Engine/RendererSettings.h b/Engine/Source/Runtime/Engine/Classes/Engine/RendererSettings.h
--- a/Engine/Source/Runtime/Engine/Classes/Engine/RendererSettings.h
+++ b/Engine/Source/Runtime/Engine/Classes/Engine/RendererSettings.h
@@ -78,6 +78,20 @@ namespace EEarlyZPass
 	};
 }
 
+/**
+ * Enumerates available options for early Z-passes.
+ */
+UENUM()
+namespace EVelocityOutputPass
+{
+	enum Type
+	{
+		DepthPass = 0 UMETA(DisplayName = "Write during depth pass"),
+		BasePass = 1 UMETA(DisplayName = "Write during base pass"),
+		AfterBasePass = 2 UMETA(DisplayName = "Write after base pass"),
+	};
+}
+
 /**
  * Enumerates available options for alpha channel through post processing. The renderer will always generate premultiplied RGBA
  * with alpha as translucency (0 = fully opaque; 1 = fully translucent).
@@ -689,19 +703,17 @@ class ENGINE_API URendererSettings : public UDeveloperSettings
 		ToolTip="Select how the g-buffer is cleared in game mode (only affects deferred shading)."))
 	TEnumAsByte<EClearSceneOptions::Type> ClearSceneMethod;
 
-	UPROPERTY(config, EditAnywhere, Category=Optimizations, meta=(
-		DisplayName="Output velocities during base pass",
-		ConsoleVariable="r.BasePassOutputsVelocity",
-		ConfigRestartRequired = true,
-		ToolTip="Enables emitting velocity during Base Pass rendering. Changing this setting requires restarting the editor.\nNote: enabling this behaves as if 'Output velocities due to vertex deformation' (r.VertexDeformationOutputsVelocity) is also enabled."
-		))
-	uint32 bBasePassOutputsVelocity:1;
+	UPROPERTY(config, EditAnywhere, Category = Optimizations, meta = (
+		ConsoleVariable = "r.VelocityOutputPass", DisplayName = "Velocity Pass",
+		ToolTip = "When to write velocity buffer. Changing this setting requires restarting the editor.",
+		ConfigRestartRequired = true))
+	TEnumAsByte<EVelocityOutputPass::Type> VelocityPass;
 
 	UPROPERTY(config, EditAnywhere, Category=Optimizations, meta=(
 		DisplayName="Output velocities due to vertex deformation",
 		ConsoleVariable="r.VertexDeformationOutputsVelocity",
-		EditCondition = "!bBasePassOutputsVelocity",
-		ToolTip="Enables materials with World Position Offset and/or World Displacement to output velocities during velocity pass even when the actor has not moved. This incurs a performance cost and can be quite significant if many objects are using WPO, such as a forest of trees - in that case consider 'Output velocities during base pass' (r.BasePassOutputsVelocity) and disabling this option."
+		EditCondition = "VelocityPass == EVelocityOutputPass::AfterBasePass",
+		ToolTip="Enables materials with World Position Offset and/or World Displacement to output velocities during velocity pass even when the actor has not moved. This incurs a performance cost and can be quite significant if many objects are using WPO, such as a forest of trees. Instead consider outputting velocities during depth or base pass (r.VelocityOutputPass)."
 		))
 	uint32 bVertexDeformationOutputsVelocity:1;
 
diff --git a/Engine/Source/Runtime/RenderCore/Private/RenderUtils.cpp b/Engine/Source/Runtime/RenderCore/Private/RenderUtils.cpp
--- a/Engine/Source/Runtime/RenderCore/Private/RenderUtils.cpp
+++ b/Engine/Source/Runtime/RenderCore/Private/RenderUtils.cpp
@@ -1281,8 +1281,8 @@ RENDERCORE_API void RenderUtilsInit()
 		GDBufferPlatformMask = ~0ull;
 	}
 
-	static IConsoleVariable* BasePassVelocityCVar = IConsoleManager::Get().FindConsoleVariable(TEXT("r.BasePassOutputsVelocity"));
-	if (BasePassVelocityCVar && BasePassVelocityCVar->GetInt())
+	static IConsoleVariable* VelocityPassCVar = IConsoleManager::Get().FindConsoleVariable(TEXT("r.VelocityOutputPass"));
+	if (VelocityPassCVar && VelocityPassCVar->GetInt() == 1)
 	{
 		GBasePassVelocityPlatformMask = ~0ull;
 	}
diff --git a/Engine/Source/Runtime/Renderer/Private/CompositionLighting/PostProcessDeferredDecals.cpp b/Engine/Source/Runtime/Renderer/Private/CompositionLighting/PostProcessDeferredDecals.cpp
--- a/Engine/Source/Runtime/Renderer/Private/CompositionLighting/PostProcessDeferredDecals.cpp
+++ b/Engine/Source/Runtime/Renderer/Private/CompositionLighting/PostProcessDeferredDecals.cpp
@@ -9,6 +9,7 @@
 #include "RendererUtils.h"
 #include "SceneUtils.h"
 #include "ScenePrivate.h"
+#include "VelocityRendering.h"
 #include "VisualizeTexture.h"
 
 static TAutoConsoleVariable<float> CVarStencilSizeThreshold(
@@ -44,12 +45,9 @@ static TAutoConsoleVariable<bool> CVarDBufferDecalNormalReprojectionEnabled(
 	TEXT("r.Decal.NormalReprojectionEnabled"),
 	true, 
 	TEXT("If true, normal reprojection from the previous frame is allowed in SceneTexture nodes on DBuffer decals, provided that motion ")
-	TEXT("in depth prepass is enabled as well (r.DepthPassMergedWithVelocity). Otherwise the fallback is the normal extracted from the depth buffer."),
+	TEXT("in depth prepass is enabled as well (r.VelocityOutputPass=0). Otherwise the fallback is the normal extracted from the depth buffer."),
 	ECVF_RenderThreadSafe);
 
-// This is a temporary extern while testing the Velocity changes. Will be removed once a decision is made.
-extern bool IsVelocityMergedWithDepthPass();
-
 bool IsDBufferEnabled(const FSceneViewFamily& ViewFamily, EShaderPlatform ShaderPlatform)
 {
 	return !ViewFamily.EngineShowFlags.ShaderComplexity && ViewFamily.EngineShowFlags.Decals && IsUsingDBuffers(ShaderPlatform);
@@ -167,7 +165,7 @@ void GetDeferredDecalPassParameters(
 
 TUniformBufferRef<FDeferredDecalUniformParameters> CreateDeferredDecalUniformBuffer(const FViewInfo& View)
 {
-	bool bIsMotionInDepth = IsVelocityMergedWithDepthPass();
+	bool bIsMotionInDepth = FVelocityRendering::DepthPassCanOutputVelocity();
 	// if we have early motion vectors (bIsMotionInDepth) and the cvar is enabled and we actually have a buffer from the previous frame (View.PrevViewInfo.GBufferA.IsValid())
 	bool bIsNormalReprojectionEnabled = (bIsMotionInDepth && CVarDBufferDecalNormalReprojectionEnabled.GetValueOnRenderThread() && View.PrevViewInfo.GBufferA.IsValid());
 
@@ -566,7 +564,7 @@ void ExtractNormalsForNextFrameReprojection(FRDGBuilder& GraphBuilder, const FSc
 	static auto CVarNormalReprojectionEnabled = IConsoleManager::Get().FindConsoleVariable(TEXT("r.Decal.NormalReprojectionEnabled"));
 
 	// save the previous frame if early motion vectors are enabled and normal reprojection is enabled, so there should be no cost if these options are off
-	bool bApplyReproject = IsVelocityMergedWithDepthPass() && CVarNormalReprojectionEnabled && CVarNormalReprojectionEnabled->GetInt() > 0;
+	bool bApplyReproject = FVelocityRendering::DepthPassCanOutputVelocity() && CVarNormalReprojectionEnabled && CVarNormalReprojectionEnabled->GetInt() > 0;
 
 	if (bApplyReproject)
 	{
diff --git a/Engine/Source/Runtime/Renderer/Private/DepthRendering.cpp b/Engine/Source/Runtime/Renderer/Private/DepthRendering.cpp
--- a/Engine/Source/Runtime/Renderer/Private/DepthRendering.cpp
+++ b/Engine/Source/Runtime/Renderer/Private/DepthRendering.cpp
@@ -69,11 +69,6 @@ static TAutoConsoleVariable<int32> CVarStencilForLODDither(
 	TEXT("Forces a full prepass when enabled."),
 	ECVF_RenderThreadSafe | ECVF_ReadOnly);
 
-static TAutoConsoleVariable<int32> CVarDepthPassMergedWithVelocity(
-	TEXT("r.DepthPassMergedWithVelocity"),
-	0,
-	TEXT("If enabled, and we are doing a full depth pass, then the depth pass will ignore movable objects and the velocity pass will write depth directly after the depth pass. After the velocity pass is finished, a full opaque depth-only texture is ready."));
-
 extern bool IsHMDHiddenAreaMaskActive();
 
 FDepthPassInfo GetDepthPassInfo(const FScene* Scene)
diff --git a/Engine/Source/Runtime/Renderer/Private/PostProcess/PostProcessMotionBlur.cpp b/Engine/Source/Runtime/Renderer/Private/PostProcess/PostProcessMotionBlur.cpp
--- a/Engine/Source/Runtime/Renderer/Private/PostProcess/PostProcessMotionBlur.cpp
+++ b/Engine/Source/Runtime/Renderer/Private/PostProcess/PostProcessMotionBlur.cpp
@@ -6,6 +6,7 @@
 #include "RenderTargetTemp.h"
 #include "SpriteIndexBuffer.h"
 #include "PostProcessing.h"
+#include "VelocityRendering.h"
 
 namespace
 {
@@ -1152,9 +1153,7 @@ FScreenPassTexture AddVisualizeMotionBlurPass(FRDGBuilder& GraphBuilder, const F
 		static const auto VelocityTestVar = IConsoleManager::Get().FindTConsoleVariableDataInt(TEXT("r.VelocityTest"));
 		const int32 VelocityTest = VelocityTestVar ? VelocityTestVar->GetValueOnRenderThread() : 0;
 
-		extern bool IsParallelVelocity(EShaderPlatform ShaderPlatform);
-
-		Line = FString::Printf(TEXT("%d, %d, %d"), View.Family->bWorldIsPaused, VelocityTest, IsParallelVelocity(View.GetShaderPlatform()));
+		Line = FString::Printf(TEXT("%d, %d, %d"), View.Family->bWorldIsPaused, VelocityTest, FVelocityRendering::IsParallelVelocity(View.GetShaderPlatform()));
 		Canvas.DrawShadowedString(X, Y += YStep, TEXT("Paused, r.VelocityTest, Parallel:"), GetStatsFont(), FLinearColor(1, 1, 0));
 		Canvas.DrawShadowedString(X + ColumnWidth, Y, *Line, GetStatsFont(), FLinearColor(1, 1, 0));
 
diff --git a/Engine/Source/Runtime/Renderer/Private/RendererScene.cpp b/Engine/Source/Runtime/Renderer/Private/RendererScene.cpp
--- a/Engine/Source/Runtime/Renderer/Private/RendererScene.cpp
+++ b/Engine/Source/Runtime/Renderer/Private/RendererScene.cpp
@@ -62,6 +62,7 @@
 #include "HAL/LowLevelMemTracker.h"
 #include "VT/RuntimeVirtualTextureSceneProxy.h"
 #include "HairStrandsInterface.h"
+#include "VelocityRendering.h"
 
 #if RHI_RAYTRACING
 #include "RayTracingDynamicGeometryCollection.h"
@@ -145,8 +146,6 @@ TGlobalResource< FGlobalDitherUniformBuffer > GDitherFadedInUniformBuffer;
 
 static FThreadSafeCounter FSceneViewState_UniqueID;
 
-extern bool IsVelocityMergedWithDepthPass();
-
 #define ENABLE_LOG_PRIMITIVE_INSTANCE_ID_STATS_TO_CSV 0
 
 #if ENABLE_LOG_PRIMITIVE_INSTANCE_ID_STATS_TO_CSV
@@ -3777,8 +3776,8 @@ void FScene::UpdateEarlyZPassMode()
 		if (ShouldForceFullDepthPass(ShaderPlatform))
 		{
 			// DBuffer decals and stencil LOD dithering force a full prepass
-			EarlyZPassMode = IsVelocityMergedWithDepthPass() ? DDM_AllOpaqueNoVelocity : DDM_AllOpaque;
-			bEarlyZPassMovable = IsVelocityMergedWithDepthPass() ? false : true;
+			EarlyZPassMode = FVelocityRendering::DepthPassCanOutputVelocity() ? DDM_AllOpaqueNoVelocity : DDM_AllOpaque;
+			bEarlyZPassMovable = FVelocityRendering::DepthPassCanOutputVelocity() ? false : true;
 		}
 
 		if ((EarlyZPassMode == DDM_AllOpaque || EarlyZPassMode == DDM_AllOpaqueNoVelocity)
diff --git a/Engine/Source/Runtime/Renderer/Private/VelocityRendering.cpp b/Engine/Source/Runtime/Renderer/Private/VelocityRendering.cpp
--- a/Engine/Source/Runtime/Renderer/Private/VelocityRendering.cpp
+++ b/Engine/Source/Runtime/Renderer/Private/VelocityRendering.cpp
@@ -22,14 +22,22 @@
 #include "VisualizeTexture.h"
 #include "MeshPassProcessor.inl"
 #include "DebugProbeRendering.h"
+#include "RendererModule.h"
 
 // Changing this causes a full shader recompile
+static TAutoConsoleVariable<int32> CVarVelocityOutputPass(
+	TEXT("r.VelocityOutputPass"),
+	0,
+	TEXT("When to write velocity buffer.\n") \
+	TEXT(" 0: Renders during the depth pass. This splits the depth pass into 2 phases: with and without velocity.\n") \
+	TEXT(" 1: Renders during the regular base pass. This adds an extra GBuffer target during base pass rendering.") \
+	TEXT(" 2: Renders after the regular base pass.\n"), \
+	ECVF_ReadOnly | ECVF_RenderThreadSafe);
+
 static TAutoConsoleVariable<int32> CVarBasePassOutputsVelocity(
 	TEXT("r.BasePassOutputsVelocity"),
-	1,
-	TEXT("Enables rendering WPO velocities on the base pass.\n") \
-	TEXT(" 0: Renders in a separate pass/rendertarget, all movable static meshes + dynamic.\n") \
-	TEXT(" 1: Renders during the regular base pass adding an extra GBuffer, but allowing motion blur on materials with Time-based WPO."),
+	-1,
+	TEXT("Deprecated CVar. Use r.VelocityOutputPass instead.\n"),
 	ECVF_ReadOnly | ECVF_RenderThreadSafe);
 
 static TAutoConsoleVariable<int32> CVarParallelVelocity(
@@ -42,8 +50,9 @@ static TAutoConsoleVariable<int32> CVarParallelVelocity(
 static TAutoConsoleVariable<int32> CVarVertexDeformationOutputsVelocity(
 	TEXT("r.VertexDeformationOutputsVelocity"),
 	0,
-	TEXT("Enables materials with World Position Offset and/or World Displacement to output velocities during velocity pass even when the actor has not moved. ")
-	TEXT("This incurs a performance cost and can be quite significant if many objects are using WPO, such as a forest of trees - in that case consider r.BasePassOutputsVelocity and disabling this option."));
+	TEXT("Enables materials with World Position Offset and/or World Displacement to output velocities during velocity pass even when the actor has not moved. \n")
+	TEXT("This only has an impact if r.VelocityOutputPass=2. \n")
+	TEXT("This will incur a performance cost that can be quite significant if many objects are using WPO, such as a forest of trees."));
 
 static TAutoConsoleVariable<int32> CVarRHICmdFlushRenderThreadTasksVelocityPass(
 	TEXT("r.RHICmdFlushRenderThreadTasksVelocityPass"),
@@ -52,23 +61,21 @@ static TAutoConsoleVariable<int32> CVarRHICmdFlushRenderThreadTasksVelocityPass(
 
 DECLARE_GPU_STAT_NAMED(RenderVelocities, TEXT("Render Velocities"));
 
-bool IsParallelVelocity(EShaderPlatform ShaderPlatform)
-{
-	return GRHICommandList.UseParallelAlgorithms() && CVarParallelVelocity.GetValueOnRenderThread()
-		// Parallel dispatch is not supported on mobile platform
-		&& !IsMobilePlatform(ShaderPlatform);
-}
-
-bool IsVelocityWaitForTasksEnabled(EShaderPlatform ShaderPlatform)
-{
-	return IsParallelVelocity(ShaderPlatform) && (CVarRHICmdFlushRenderThreadTasksVelocityPass.GetValueOnRenderThread() > 0 || CVarRHICmdFlushRenderThreadTasks.GetValueOnRenderThread() > 0);
-}
-
-bool IsVelocityMergedWithDepthPass()
+/** Validate that deprecated CVars are no longer set. */
+inline void ValidateVelocityCVars()
 {
-	static const auto CVarMergeDepth = IConsoleManager::Get().FindTConsoleVariableDataInt(TEXT("r.DepthPassMergedWithVelocity"));
-	bool bIsMerged = (CVarMergeDepth && CVarMergeDepth->GetValueOnAnyThread() != 0);
-	return bIsMerged;
+#if !UE_BUILD_SHIPPING
+	static bool bHasValidatedCVars = false;
+	if (!bHasValidatedCVars)
+	{
+		const int32 Value = CVarBasePassOutputsVelocity.GetValueOnAnyThread();
+		if (Value != -1)
+		{
+			UE_LOG(LogRenderer, Warning, TEXT("Deprectaed CVar r.BasePassOutputsVelocity is set to %d. Use r.VelocityOutputPass instead."), Value);
+		}
+		bHasValidatedCVars = true;
+	}
+#endif
 }
 
 class FVelocityVS : public FMeshMaterialShader
@@ -158,7 +165,7 @@ DECLARE_CYCLE_STAT(TEXT("Velocity"), STAT_CLP_Velocity, STATGROUP_ParallelComman
 
 bool FDeferredShadingSceneRenderer::ShouldRenderVelocities() const
 {
-	if (!FVelocityRendering::IsSeparateVelocityPassSupported(ShaderPlatform) || ViewFamily.UseDebugViewPS())
+	if (!FVelocityRendering::IsVelocityPassSupported(ShaderPlatform) || ViewFamily.UseDebugViewPS())
 	{
 		return false;
 	}
@@ -191,7 +198,7 @@ bool FDeferredShadingSceneRenderer::ShouldRenderVelocities() const
 
 bool FMobileSceneRenderer::ShouldRenderVelocities() const
 {
-	if (!FVelocityRendering::IsSeparateVelocityPassSupported(ShaderPlatform) || ViewFamily.UseDebugViewPS() || !PlatformSupportsVelocityRendering(ShaderPlatform))
+	if (!FVelocityRendering::IsVelocityPassSupported(ShaderPlatform) || ViewFamily.UseDebugViewPS() || !PlatformSupportsVelocityRendering(ShaderPlatform))
 	{
 		return false;
 	}
@@ -236,7 +243,7 @@ void FSceneRenderer::RenderVelocities(
 		: ERenderTargetLoadAction::EClear;
 
 	RDG_GPU_STAT_SCOPE(GraphBuilder, RenderVelocities);
-	RDG_WAIT_FOR_TASKS_CONDITIONAL(GraphBuilder, IsVelocityWaitForTasksEnabled(ShaderPlatform));
+	RDG_WAIT_FOR_TASKS_CONDITIONAL(GraphBuilder, FVelocityRendering::IsVelocityWaitForTasksEnabled(ShaderPlatform));
 
 	const EMeshPass::Type MeshPass = GetMeshPassFromVelocityPass(VelocityPass);
 	FExclusiveDepthStencil ExclusiveDepthStencil = (VelocityPass == EVelocityPass::Opaque && !(Scene->EarlyZPassMode == DDM_AllOpaqueNoVelocity))
@@ -259,7 +266,7 @@ void FSceneRenderer::RenderVelocities(
 
 			RDG_GPU_MASK_SCOPE(GraphBuilder, View.GPUMask);
 
-			const bool bIsParallelVelocity = IsParallelVelocity(ShaderPlatform);
+			const bool bIsParallelVelocity = FVelocityRendering::IsParallelVelocity(ShaderPlatform);
 
 			// Clear velocity render target explicitly when velocity rendering in parallel or no draw but force to.
 			// Avoid adding a separate clear pass in non parallel rendering.
@@ -357,11 +364,19 @@ FRDGTextureDesc FVelocityRendering::GetRenderTargetDesc(EShaderPlatform ShaderPl
 	return FRDGTextureDesc::Create2D(Extent, GetFormat(ShaderPlatform), FClearValueBinding::Transparent, TexCreate_RenderTargetable | TexCreate_UAV | TexCreate_ShaderResource | FastVRamFlag);
 }
 
-bool FVelocityRendering::IsSeparateVelocityPassSupported(EShaderPlatform ShaderPlatform)
+bool FVelocityRendering::IsVelocityPassSupported(EShaderPlatform ShaderPlatform)
 {
+	ValidateVelocityCVars();
+
 	return GPixelFormats[GetFormat(ShaderPlatform)].Supported;
 }
 
+bool FVelocityRendering::DepthPassCanOutputVelocity()
+{
+	static bool bIsMerged = CVarVelocityOutputPass.GetValueOnAnyThread() == 0;
+	return bIsMerged;
+}
+
 bool FVelocityRendering::BasePassCanOutputVelocity(EShaderPlatform ShaderPlatform)
 {
 	return IsUsingBasePassVelocity(ShaderPlatform);
@@ -384,6 +399,18 @@ bool FVelocityRendering::IsSeparateVelocityPassRequiredByVertexFactory(EShaderPl
 	return bBasePassVelocityNotSupported || bVertexFactoryRequiresSeparateVelocityPass;
 }
 
+bool FVelocityRendering::IsParallelVelocity(EShaderPlatform ShaderPlatform)
+{
+	return GRHICommandList.UseParallelAlgorithms() && CVarParallelVelocity.GetValueOnRenderThread()
+		// Parallel dispatch is not supported on mobile platform
+		&& !IsMobilePlatform(ShaderPlatform);
+}
+
+bool FVelocityRendering::IsVelocityWaitForTasksEnabled(EShaderPlatform ShaderPlatform)
+{
+	return FVelocityRendering::IsParallelVelocity(ShaderPlatform) && (CVarRHICmdFlushRenderThreadTasksVelocityPass.GetValueOnRenderThread() > 0 || CVarRHICmdFlushRenderThreadTasks.GetValueOnRenderThread() > 0);
+}
+
 bool FVelocityMeshProcessor::PrimitiveHasVelocityForView(const FViewInfo& View, const FPrimitiveSceneProxy* PrimitiveSceneProxy)
 {
 	// Skip camera cuts which effectively reset velocity for the new frame.
@@ -415,7 +442,7 @@ bool FVelocityMeshProcessor::PrimitiveHasVelocityForView(const FViewInfo& View,
 
 bool FOpaqueVelocityMeshProcessor::PrimitiveCanHaveVelocity(EShaderPlatform ShaderPlatform, const FPrimitiveSceneProxy* PrimitiveSceneProxy)
 {
-	if (!FVelocityRendering::IsSeparateVelocityPassSupported(ShaderPlatform) || !PlatformSupportsVelocityRendering(ShaderPlatform))
+	if (!FVelocityRendering::IsVelocityPassSupported(ShaderPlatform) || !PlatformSupportsVelocityRendering(ShaderPlatform))
 	{
 		return false;
 	}
@@ -553,7 +580,7 @@ bool FTranslucentVelocityMeshProcessor::PrimitiveCanHaveVelocity(EShaderPlatform
 	 * Therefore, the primitive can't be filtered based on motion, or it will break post
 	 * effects like depth of field which rely on depth information.
 	 */
-	return FVelocityRendering::IsSeparateVelocityPassSupported(ShaderPlatform) && PlatformSupportsVelocityRendering(ShaderPlatform);
+	return FVelocityRendering::IsVelocityPassSupported(ShaderPlatform) && PlatformSupportsVelocityRendering(ShaderPlatform);
 }
 
 bool FTranslucentVelocityMeshProcessor::PrimitiveHasVelocityForFrame(const FPrimitiveSceneProxy* PrimitiveSceneProxy)
diff --git a/Engine/Source/Runtime/Renderer/Private/VelocityRendering.h b/Engine/Source/Runtime/Renderer/Private/VelocityRendering.h
--- a/Engine/Source/Runtime/Renderer/Private/VelocityRendering.h
+++ b/Engine/Source/Runtime/Renderer/Private/VelocityRendering.h
@@ -32,11 +32,17 @@ EMeshPass::Type GetMeshPassFromVelocityPass(EVelocityPass VelocityPass);
 // Group Velocity Rendering accessors, types, etc.
 struct FVelocityRendering
 {
+	/** Returns the texture format for the velocity buffer. */
 	static EPixelFormat GetFormat(EShaderPlatform ShaderPlatform);
+	
+	/** Returns the render target description for the velocity buffer. */
 	static FRDGTextureDesc GetRenderTargetDesc(EShaderPlatform ShaderPlatform, FIntPoint Extent);
 
-	/** Returns true if the separate velocity pass is enabled. */
-	static bool IsSeparateVelocityPassSupported(EShaderPlatform ShaderPlatform);
+	/** Returns true if a velocity pass is supported. */
+	static bool IsVelocityPassSupported(EShaderPlatform ShaderPlatform);
+
+	/** Returns true if the velocity can be output during depth pass. */
+	static bool DepthPassCanOutputVelocity();
 
 	/** Returns true if the velocity can be output in the BasePass. */
 	static bool BasePassCanOutputVelocity(EShaderPlatform ShaderPlatform);
@@ -46,6 +52,12 @@ struct FVelocityRendering
 
 	/** Returns true if a separate velocity pass is required (i.e. not rendered by the base pass) given the provided vertex factory settings. */
 	static bool IsSeparateVelocityPassRequiredByVertexFactory(EShaderPlatform ShaderPlatform, bool bVertexFactoryUsesStaticLighting);
+
+	/** Returns true if the velocity pass is using parallel dispatch. */
+	static bool IsParallelVelocity(EShaderPlatform ShaderPlatform);
+
+	/** Returns true if we wait for outstanding tasks in velocity pass. */
+	static bool IsVelocityWaitForTasksEnabled(EShaderPlatform ShaderPlatform);
 };
 
 /**
