Fix for switch hanging on table being built.

diff --git a/Engine/Source/Runtime/Renderer/Private/ShadingEnergyConservation.cpp b/Engine/Source/Runtime/Renderer/Private/ShadingEnergyConservation.cpp
--- a/Engine/Source/Runtime/Renderer/Private/ShadingEnergyConservation.cpp
+++ b/Engine/Source/Runtime/Renderer/Private/ShadingEnergyConservation.cpp
@@ -182,13 +182,16 @@ void Init(FRDGBuilder& GraphBuilder, FViewInfo& View)
 	FRDGTextureRef DiffuseEnergyTexture = nullptr;
 
 	// Enabled based on settings
+	const bool bMaterialEnergyConservationEnabled = CVarMaterialEnergyConservation.GetValueOnRenderThread() > 0;
 	const bool bIsEnergyConservationEnabled = CVarShadingEnergyConservation.GetValueOnRenderThread() > 0;
 	const bool bIsEnergyPreservationEnabled = CVarShadingEnergyConservation_Preservation.GetValueOnRenderThread() > 0;	
 
 	// Build/bind table if energy conservation is enabled or if strata is enabled in order to have 
 	// the correct tables built & bound. Even if we are not using energy conservation, we want to 
 	// have access to directional albedo information for env. lighting for instance)
-	const bool bBindEnergyData = (View.ViewState != nullptr) && (bIsEnergyPreservationEnabled || bIsEnergyConservationEnabled || Strata::IsStrataEnabled() || View.Family->EngineShowFlags.PathTracing);
+	// TODO we do disable LUT generation when bMaterialEnergyConservationEnabled is false in order avoid having
+	//      some platofrms hanging when generating the luts on startup. We will revisit with LUTs that are simply loaded.
+	const bool bBindEnergyData = (View.ViewState != nullptr) && bMaterialEnergyConservationEnabled && (bIsEnergyPreservationEnabled || bIsEnergyConservationEnabled || Strata::IsStrataEnabled() || View.Family->EngineShowFlags.PathTracing);
 	if (bBindEnergyData)
 	{
 		const EPixelFormat Format = CVarShadingFurnaceTest_TableFormat.GetValueOnRenderThread() > 0 ? PF_G32R32F : PF_G16R16F;
