Fix MipLevel sample mode with Adaptive Virtual Texture.

diff --git a/Engine/Source/Runtime/RenderCore/Public/VirtualTexturing.h b/Engine/Source/Runtime/RenderCore/Public/VirtualTexturing.h
--- a/Engine/Source/Runtime/RenderCore/Public/VirtualTexturing.h
+++ b/Engine/Source/Runtime/RenderCore/Public/VirtualTexturing.h
@@ -53,12 +53,12 @@ struct FAllocatedVTDescription
 
 	uint32 TileSize = 0u;
 	uint32 TileBorderSize = 0u;
+	uint32 MaxSpaceSize = 0u;
+	uint32 IndirectionTextureSize = 0u;
 	uint8 Dimensions = 0u;
 	uint8 NumTextureLayers = 0u;
-	
-	uint32 MaxSpaceSize = 0u;
 	uint8 ForceSpaceID = 0xff;
-	uint32 IndirectionTextureSize = 0u;
+	uint8 AdaptiveLevelBias = 0u;
 
 	/** Producer for each texture layer. */
 	FVirtualTextureProducerHandle ProducerHandle[VIRTUALTEXTURE_SPACE_MAXLAYERS];
diff --git a/Engine/Source/Runtime/Renderer/Private/VT/AdaptiveVirtualTexture.cpp b/Engine/Source/Runtime/Renderer/Private/VT/AdaptiveVirtualTexture.cpp
--- a/Engine/Source/Runtime/Renderer/Private/VT/AdaptiveVirtualTexture.cpp
+++ b/Engine/Source/Runtime/Renderer/Private/VT/AdaptiveVirtualTexture.cpp
@@ -140,6 +140,7 @@ namespace
 		AllocatedDesc.bPrivateSpace = true;
 		AllocatedDesc.ForceSpaceID = InForcedSpaceID;
 		AllocatedDesc.IndirectionTextureSize = FMath::Max(InGridSize.X, InGridSize.Y);
+		AllocatedDesc.AdaptiveLevelBias = InLevelOffset;
 
 		for (int32 LayerIndex = 0; LayerIndex < InAllocatedDesc.NumTextureLayers; ++LayerIndex)
 		{
diff --git a/Engine/Source/Runtime/Renderer/Private/VT/AllocatedVirtualTexture.cpp b/Engine/Source/Runtime/Renderer/Private/VT/AllocatedVirtualTexture.cpp
--- a/Engine/Source/Runtime/Renderer/Private/VT/AllocatedVirtualTexture.cpp
+++ b/Engine/Source/Runtime/Renderer/Private/VT/AllocatedVirtualTexture.cpp
@@ -451,14 +451,13 @@ void FAllocatedVirtualTexture::GetPackedPageTableUniform(FUintVector4* OutUnifor
 	const uint32 WidthInPages = GetWidthInTiles();
 	const uint32 HeightInPages = GetHeightInTiles();
 	const uint32 vPageTableMipBias = FMath::FloorLog2(vPageSize);
+	const uint32 AdaptiveLevelBias = Description.AdaptiveLevelBias;
 
 	// Here MaxAnisotropy only controls the VT mip level selection
 	// We don't need to limit this value based on border size, and we can add this factor in even if HW anisotropic filtering is disabled
 	const uint32 MaxAnisotropy = VirtualTextureScalability::GetMaxAnisotropy();
 	const uint32 MaxAnisotropyLog2 = (MaxAnisotropy > 0u) ? FMath::FloorLog2(MaxAnisotropy) : 0u;
 
-	const uint32 AdaptiveLevelBias = 0; // todo[vt]: Required for handling SampleLevel correctly on adaptive page tables.
-
 	// make sure everything fits in the allocated number of bits
 	checkSlow(vPageX < 4096u);
 	checkSlow(vPageY < 4096u);
