Fixes crash on r.BloomQuality=6

diff --git a/Engine/Source/Runtime/Renderer/Private/PostProcess/PostProcessBloomSetup.cpp b/Engine/Source/Runtime/Renderer/Private/PostProcess/PostProcessBloomSetup.cpp
--- a/Engine/Source/Runtime/Renderer/Private/PostProcess/PostProcessBloomSetup.cpp
+++ b/Engine/Source/Runtime/Renderer/Private/PostProcess/PostProcessBloomSetup.cpp
@@ -153,7 +153,7 @@ EBloomQuality GetBloomQuality()
 	return static_cast<EBloomQuality>(FMath::Clamp(
 		CVar->GetValueOnRenderThread(),
 		static_cast<int32>(EBloomQuality::Disabled),
-		static_cast<int32>(EBloomQuality::MAX)));
+		static_cast<int32>(EBloomQuality::MAX) - 1));
 }
 
 static_assert(
@@ -211,6 +211,8 @@ FScreenPassTexture AddGaussianBloomPasses(FRDGBuilder& GraphBuilder, const FView
 		static_assert(UE_ARRAY_COUNT(BloomStages) == BloomQualityCountMax, "Array must be one less than the number of bloom quality entries.");
 		static_assert(UE_ARRAY_COUNT(BloomQualityToSceneDownsampleStage) == BloomQualityCountMax, "Array must be one less than the number of bloom quality entries.");
 
+		check(BloomQualityIndex < BloomQualityCountMax);
+
 		// Use bloom quality to select the number of downsample stages to use for bloom.
 		const uint32 BloomStageCount = BloomQualityToSceneDownsampleStage[BloomQualityIndex];
 
