Removed the used of custom vertex shader indexing using VertexId for clouds because some platforms are registered as SM5 but do not supprot VertexId.

diff --git a/Engine/Shaders/Private/VolumetricCloud.usf b/Engine/Shaders/Private/VolumetricCloud.usf
--- a/Engine/Shaders/Private/VolumetricCloud.usf
+++ b/Engine/Shaders/Private/VolumetricCloud.usf
@@ -179,7 +179,7 @@ void MainVS(
 	ResolvedView = ResolveView();
 
 	// Full screen triangle
-	Interpolants.Position = float4(FactoryInput.VertexId == 2 || FactoryInput.VertexId == 1 ? 1.0 : -3.0, FactoryInput.VertexId == 1 ? -3.0 : 1.0, 0.5f, 1.0f);
+	Interpolants.Position = FactoryInput.Position;
 
 	const float3 WorldPosition = 0.0f;
 	FVertexFactoryIntermediates VFIntermediates = GetVertexFactoryIntermediates(FactoryInput);
diff --git a/Engine/Source/Runtime/Renderer/Private/VolumetricCloudRendering.cpp b/Engine/Source/Runtime/Renderer/Private/VolumetricCloudRendering.cpp
--- a/Engine/Source/Runtime/Renderer/Private/VolumetricCloudRendering.cpp
+++ b/Engine/Source/Runtime/Renderer/Private/VolumetricCloudRendering.cpp
@@ -871,9 +871,9 @@ public:
 		TArray<FDynamicMeshVertex> Vertices;
 
 		// Vertex position constructed in the shader
-		Vertices.Add(FDynamicMeshVertex(FVector3f(0.0f, 0.0f, 0.0f)));
-		Vertices.Add(FDynamicMeshVertex(FVector3f(0.0f, 0.0f, 0.0f)));
-		Vertices.Add(FDynamicMeshVertex(FVector3f(0.0f, 0.0f, 0.0f)));
+		Vertices.Add(FDynamicMeshVertex(FVector3f(-3.0f,  1.0f, 0.5f)));
+		Vertices.Add(FDynamicMeshVertex(FVector3f( 1.0f, -3.0f, 0.5f)));
+		Vertices.Add(FDynamicMeshVertex(FVector3f( 1.0f,  1.0f, 0.5f)));
 
 		Buffers.PositionVertexBuffer.Init(Vertices.Num());
 		Buffers.StaticMeshVertexBuffer.Init(Vertices.Num(), 1);
@@ -946,7 +946,7 @@ static void GetSingleTriangleMeshBatch(FMeshBatch& LocalSingleTriangleMesh, cons
 
 	LocalSingleTriangleMesh.VertexFactory = GSingleTriangleMeshVertexFactory;
 	LocalSingleTriangleMesh.MaterialRenderProxy = CloudVolumeMaterialProxy;
-	LocalSingleTriangleMesh.Elements[0].IndexBuffer = nullptr;
+	LocalSingleTriangleMesh.Elements[0].IndexBuffer = &GScreenRectangleIndexBuffer;
 	LocalSingleTriangleMesh.Elements[0].FirstIndex = 0;
 	LocalSingleTriangleMesh.Elements[0].NumPrimitives = 1;
 	LocalSingleTriangleMesh.Elements[0].MinVertexIndex = 0;
