Added ForwardShadingPriority to help artist select the single directional light that be used for forward shading, translucent, water, volumetric fog.
Compiled game to make sure WITH_EDITOR defines are correctly positioned.
If the MultipleDirLightsConflictForForwardShading message is annoying and pops up to ooften we will remove it. At least it is clear what is happening.

diff --git a/Engine/Source/Runtime/Engine/Classes/Components/DirectionalLightComponent.h b/Engine/Source/Runtime/Engine/Classes/Components/DirectionalLightComponent.h
--- a/Engine/Source/Runtime/Engine/Classes/Components/DirectionalLightComponent.h
+++ b/Engine/Source/Runtime/Engine/Classes/Components/DirectionalLightComponent.h
@@ -123,6 +123,13 @@ class ENGINE_API UDirectionalLightComponent : public ULightComponent
 	UPROPERTY(EditAnywhere, BlueprintReadOnly, Category=DistanceFieldShadows, meta=(UIMin = "0", UIMax = "100000"), DisplayName = "DistanceField Shadow Distance")
 	float DistanceFieldShadowDistance;
 
+	/**
+	 * Forward lighting priority for the single directional light that will be used for forward shading, translucent, single layer water and volumetric fog.
+	 * When two lights have equal priorities, the selection will be based on their overall brightness as a fallback.
+	 */
+	UPROPERTY(EditAnywhere, AdvancedDisplay, BlueprintReadOnly, Category = Light, meta = (UIMin = "0", ClampMin = "0"))
+	int32 ForwardShadingPriority;
+
 	/** 
 	 * Angle subtended by light source in degrees (also known as angular diameter).
 	 * Defaults to 0.5357 which is the angle for our sun.
diff --git a/Engine/Source/Runtime/Engine/Private/Components/DirectionalLightComponent.cpp b/Engine/Source/Runtime/Engine/Private/Components/DirectionalLightComponent.cpp
--- a/Engine/Source/Runtime/Engine/Private/Components/DirectionalLightComponent.cpp
+++ b/Engine/Source/Runtime/Engine/Private/Components/DirectionalLightComponent.cpp
@@ -257,6 +257,9 @@ public:
 	/** If greater than WholeSceneDynamicShadowRadius, a cascade will be created to support ray traced distance field shadows covering up to this distance. */
 	float DistanceFieldShadowDistance;
 
+	/** Forward lighting priority for the single slot available for a single directional light.*/
+	int32 ForwardShadingPriority;
+
 	/** Light source angle in degrees. */
 	float LightSourceAngle;
 
@@ -299,6 +302,7 @@ public:
 		CascadeTransitionFraction(Component->CascadeTransitionFraction),
 		ShadowDistanceFadeoutFraction(Component->ShadowDistanceFadeoutFraction),
 		DistanceFieldShadowDistance(Component->bUseRayTracedDistanceFieldShadows ? Component->DistanceFieldShadowDistance : 0),
+		ForwardShadingPriority(Component->ForwardShadingPriority),
 		LightSourceAngle(Component->LightSourceAngle),
 		LightSourceSoftAngle(Component->LightSourceSoftAngle),
 		ShadowSourceAngleFactor(Component->ShadowSourceAngleFactor),
@@ -504,6 +508,11 @@ public:
 		return FVector2D(NearDistance, 1.0f / FMath::Max<float>(FarDistance - NearDistance, KINDA_SMALL_NUMBER));
 	}
 
+	virtual int32 GetDirectionalLightForwardShadingPriority() const override
+	{
+		return ForwardShadingPriority;
+	}
+
 	virtual bool GetPerObjectProjectedShadowInitializer(const FBoxSphereBounds& SubjectBounds,FPerObjectProjectedShadowInitializer& OutInitializer) const override
 	{
 		OutInitializer.PreShadowTranslation = -SubjectBounds.Origin;
@@ -968,6 +977,7 @@ UDirectionalLightComponent::UDirectionalLightComponent(const FObjectInitializer&
 	DynamicShadowDistanceMovableLight = 20000.0f;
 	DynamicShadowDistanceStationaryLight = 0.f;
 
+	ForwardShadingPriority = 0;
 	DistanceFieldShadowDistance = 30000.0f;
 	TraceDistance = 10000.0f;
 	FarShadowDistance = 300000.0f;
diff --git a/Engine/Source/Runtime/Engine/Public/SceneManagement.h b/Engine/Source/Runtime/Engine/Public/SceneManagement.h
--- a/Engine/Source/Runtime/Engine/Public/SceneManagement.h
+++ b/Engine/Source/Runtime/Engine/Public/SceneManagement.h
@@ -1447,6 +1447,11 @@ public:
 		return FVector2D(0, 0);
 	}
 
+	virtual int32 GetDirectionalLightForwardShadingPriority() const
+	{
+		return 0;
+	}
+
 	virtual bool GetLightShaftOcclusionParameters(float& OutOcclusionMaskDarkness, float& OutOcclusionDepthRange) const
 	{
 		OutOcclusionMaskDarkness = 0;
diff --git a/Engine/Source/Runtime/Renderer/Private/LightGridInjection.cpp b/Engine/Source/Runtime/Renderer/Private/LightGridInjection.cpp
--- a/Engine/Source/Runtime/Renderer/Private/LightGridInjection.cpp
+++ b/Engine/Source/Runtime/Renderer/Private/LightGridInjection.cpp
@@ -288,6 +288,8 @@ void FSceneRenderer::ComputeLightGrid(FRDGBuilder& GraphBuilder, bool bCullLight
 
 		float FurthestLight = 1000;
 
+		int32 ConflictingLightCountForForwardShading = 0;
+
 		// Track the end markers for different types
 		int32 SimpleLightsEnd = 0;
 		int32 ClusteredSupportedEnd = 0;
@@ -357,6 +359,7 @@ void FSceneRenderer::ComputeLightGrid(FRDGBuilder& GraphBuilder, bool bCullLight
 			}
 
 			float SelectedForwardDirectionalLightIntensitySq = 0.0f;
+			int32 SelectedForwardDirectionalLightPriority = -1;
 			const TArray<FSortedLightSceneInfo, SceneRenderingAllocator>& SortedLights = SortedLightSet.SortedLights;
 			ClusteredSupportedEnd = SimpleLightsEnd;
 			// Next add all the other lights, track the end index for clustered supporting lights
@@ -475,8 +478,24 @@ void FSceneRenderer::ComputeLightGrid(FRDGBuilder& GraphBuilder, bool bCullLight
 						// Also some people noticed that depending on the order a two directional lights are made visible in a level, the selected light for volumetric fog lighting will be different.
 						// So to be clear and avoid such issue, we select the most intense directional light for forward shading and volumetric lighting.
 						const float LightIntensitySq = FVector3f(LightParameters.Color).SizeSquared();
-						if (LightIntensitySq > SelectedForwardDirectionalLightIntensitySq)
+						const int32 LightForwardShadingPriority = LightProxy->GetDirectionalLightForwardShadingPriority();
+#if WITH_EDITOR
+						if (LightForwardShadingPriority > SelectedForwardDirectionalLightPriority)
+						{
+							// Reset the count if the new light has a higher priority than the previous one.
+							ConflictingLightCountForForwardShading = 1;
+						}
+						else if (LightForwardShadingPriority == SelectedForwardDirectionalLightPriority)
 						{
+							// Accumulate new light if also has the highest priority value.
+							ConflictingLightCountForForwardShading++;
+						}
+#endif
+						if (LightForwardShadingPriority > SelectedForwardDirectionalLightPriority
+							|| (LightForwardShadingPriority == SelectedForwardDirectionalLightPriority && LightIntensitySq > SelectedForwardDirectionalLightIntensitySq))
+						{
+
+							SelectedForwardDirectionalLightPriority = LightForwardShadingPriority;
 							SelectedForwardDirectionalLightIntensitySq = LightIntensitySq;
 							View.ForwardLightingResources.SelectedForwardDirectionalLightProxy = LightProxy;
 
@@ -555,6 +574,10 @@ void FSceneRenderer::ComputeLightGrid(FRDGBuilder& GraphBuilder, bool bCullLight
 			}
 		}
 
+#if WITH_EDITOR
+		// For any views, if there are more than two light that compete for the forward shaded light, we report it.
+		bMultipleDirLightsConflictForForwardShading |= ConflictingLightCountForForwardShading >= 2;
+#endif
 
 		// Store off the number of lights before we add a fake entry
 		const int32 NumLocalLightsFinal = ForwardLocalLightData.Num();
diff --git a/Engine/Source/Runtime/Renderer/Private/SceneRendering.cpp b/Engine/Source/Runtime/Renderer/Private/SceneRendering.cpp
--- a/Engine/Source/Runtime/Renderer/Private/SceneRendering.cpp
+++ b/Engine/Source/Runtime/Renderer/Private/SceneRendering.cpp
@@ -2416,6 +2416,9 @@ FSceneRenderer::FSceneRenderer(const FSceneViewFamily* InViewFamily,FHitProxyCon
 ,	bHasRequestedToggleFreeze(false)
 ,	bUsedPrecomputedVisibility(false)
 ,	InstancedStereoWidth(0)
+#if WITH_EDITOR
+,	bMultipleDirLightsConflictForForwardShading(false)
+#endif
 ,	FamilySize(0, 0)
 ,	GPUSceneDynamicContext(CheckPointer(Scene)->GPUScene)
 ,	bShadowDepthRenderCompleted(false)
@@ -3244,7 +3247,11 @@ void FSceneRenderer::RenderFinish(FRDGBuilder& GraphBuilder, FRDGTextureRef View
 			|| bShowDFAODisabledWarning || bShowShadowedLightOverflowWarning || bShowMobileDynamicCSMWarning || bShowMobileLowQualityLightmapWarning || bShowMobileMovableDirectionalLightWarning
 			|| bMobileShowVertexFogWarning || bMobileMissingSkyMaterial || bShowSkinCacheOOM || bSingleLayerWaterWarning || bShowDFDisabledWarning || bShowNoSkyAtmosphereComponentWarning || bFxDebugDraw 
 			|| bLumenEnabledButHasNoDataForTracing || bLumenEnabledButDisabledForTheProject || bNaniteEnabledButNoAtomics || bNaniteEnabledButDisabledInProject || bRealTimeSkyCaptureButNothingToCapture || bShowWaitingSkylight
-			|| bShowLocalExposureDisabledWarning || bHasDelegateWarnings;
+			|| bShowLocalExposureDisabledWarning || bHasDelegateWarnings
+#if WITH_EDITOR
+			|| bMultipleDirLightsConflictForForwardShading
+#endif
+			;
 
 		for(int32 ViewIndex = 0;ViewIndex < Views.Num();ViewIndex++)
 		{	
@@ -3458,6 +3465,12 @@ void FSceneRenderer::RenderFinish(FRDGBuilder& GraphBuilder, FRDGTextureRef View
 						}
 
 #if WITH_EDITOR
+						if (bMultipleDirLightsConflictForForwardShading)
+						{
+							static const FText Message = NSLOCTEXT("Renderer", "MultipleDirLightsConflictForForwardShading", "Multiple directional lights are competing to be the single one used for forward shading, translucent, water or volumetric fog. Please adjust their ForwardShadingPriority.\nAs a fallback, the main directional light will be selected based on overall brightness.");
+							Writer.DrawLine(Message, 10, FColor::Orange);
+						}
+
 						FSkyLightSceneProxy* SkyLight = Scene->SkyLight;
 						if (bShowWaitingSkylight && SkyLight)
 						{
diff --git a/Engine/Source/Runtime/Renderer/Private/SceneRendering.h b/Engine/Source/Runtime/Renderer/Private/SceneRendering.h
--- a/Engine/Source/Runtime/Renderer/Private/SceneRendering.h
+++ b/Engine/Source/Runtime/Renderer/Private/SceneRendering.h
@@ -1876,6 +1876,13 @@ public:
 	 */
 	uint32 InstancedStereoWidth;
 
+#if WITH_EDITOR
+	/**
+	 * True when multiple direcitonal lights are in competiation for the single slot available for forward shading (many have the same highest ForwardShadingPriority)
+	 */
+	bool bMultipleDirLightsConflictForForwardShading;
+#endif
+
 public:
 
 	FSceneRenderer(const FSceneViewFamily* InViewFamily,FHitProxyConsumer* HitProxyConsumer);
