Add single pass permutation for computing hair shadow mask onto opaque geometry.

Save 30% when there is several hair macrogroups on screen.

diff --git a/Engine/Shaders/Private/HairStrands/HairStrandsDeepShadowMask.usf b/Engine/Shaders/Private/HairStrands/HairStrandsDeepShadowMask.usf
--- a/Engine/Shaders/Private/HairStrands/HairStrandsDeepShadowMask.usf
+++ b/Engine/Shaders/Private/HairStrands/HairStrandsDeepShadowMask.usf
@@ -394,11 +394,10 @@ void MainPS(
 #endif
 
 float4				Voxel_TranslatedLightPosition_LightDirection;
+uint				Voxel_MacroGroupCount;
 uint				Voxel_MacroGroupId;
 uint				Voxel_RandomType;
 Texture2D<float4>	RayMarchMaskTexture;
-Texture2D<float>	DeepShadow_FrontDepthTexture;
-Texture2D<float4>	DeepShadow_DomTexture;
 
 // The hair shadow mask is rendered by tracing ray/cone through the voxel structure. The traversing is done with 
 // point sampling and jittering to get the smooth aspect, and avoid expensive trilinear sampling in a sparse 
@@ -438,17 +437,18 @@ void MainPS(
 	const uint2 PixelCoord = floor(Input.Position.xy);
 	const float SceneDepth = ConvertFromDeviceZ(SceneDepthTexture.Load(uint3(PixelCoord, 0)).x);
 
+	const float DistanceThreshold = 100000.0f;
+	const float CoverageThreshold = 0.995f; // When Coverage is high, we do not trace shadow on opaque since hair/fur is covering the background.
 	const float HairPixelCoverage = HairStrands.HairCoverageTexture.Load(uint3(PixelCoord, 0));
-	const bool bIsHairPixel = HairPixelCoverage >= 1;
-
+	const bool bIsValidOpaquePixel = HairPixelCoverage < CoverageThreshold;
+	
 	float3 TranslatedWorldPosition = GetTranslatedWorldPosition(UV, SceneDepth);
 
 	// PCF or simple point sampler
 	float Visibility = 1;
-	if (!bIsHairPixel)
+	bool bHasValidOutput = false;
+	if (bIsValidOpaquePixel)
 	{
-		const float VoxelShadowMaxDistance  = 100000.0; // 1 km shadow distance
-		const bool bIsDirectional			= IsLightDirectional(Voxel_TranslatedLightPosition_LightDirection);
 		const float3 TranslatedLightPosition= GetTranslatedLightPosition(Voxel_TranslatedLightPosition_LightDirection, TranslatedWorldPosition);
 		const float3 LightDirection			= GetTranslatedLightDirection(Voxel_TranslatedLightPosition_LightDirection, TranslatedWorldPosition);
 
@@ -463,18 +463,11 @@ void MainPS(
 
 		FHairTraversalResult Result = InitHairTraversalResult();
 
-		const float DistanceThreshold = 100000.0f;
-		const float CoverageThreshold = 0.995f; // When Coverage is high, we do not trace shadow on opaque since hair/fur is covering the background.
-
 		FVirtualVoxelCommonDesc CommonDesc;
-		CommonDesc.PageCountResolution = VirtualVoxel.PageCountResolution;
-		CommonDesc.VoxelWorldSize = VirtualVoxel.VoxelWorldSize;
-		CommonDesc.PageTextureResolution = VirtualVoxel.PageTextureResolution;
-		CommonDesc.PageResolution = VirtualVoxel.PageResolution;
-
-		const FPackedVirtualVoxelNodeDesc PackedNode = VirtualVoxel.NodeDescBuffer[Voxel_MacroGroupId];
-		const FVirtualVoxelNodeDesc NodeDesc = UnpackVoxelNode(PackedNode, VirtualVoxel.PageResolution);
-		const bool bIsValid = all(NodeDesc.PageIndexResolution != 0);
+		CommonDesc.PageCountResolution	= VirtualVoxel.PageCountResolution;
+		CommonDesc.VoxelWorldSize		= VirtualVoxel.VoxelWorldSize;
+		CommonDesc.PageTextureResolution= VirtualVoxel.PageTextureResolution;
+		CommonDesc.PageResolution		= VirtualVoxel.PageResolution;
 
 		#if VOXEL_TRAVERSAL_DEBUG
 		const bool bDebugEnabled = PixelCoord.x == GetCursorPos().x && PixelCoord.y == GetCursorPos().y;
@@ -482,34 +475,51 @@ void MainPS(
 		const bool bDebugEnabled = false;
 		#endif	
 
-		FHairTraversalSettings TraversalSettings = InitHairTraversalSettings();
-		TraversalSettings.DensityScale = VirtualVoxel.DensityScale_Shadow;
-		TraversalSettings.CountThreshold = 1;
-		TraversalSettings.DistanceThreshold = DistanceThreshold;
-		TraversalSettings.bDebugEnabled = bDebugEnabled;
-		TraversalSettings.SteppingScale = VirtualVoxel.SteppingScale_Shadow;
-		TraversalSettings.Random = Random;
-		TraversalSettings.PixelRadius = SceneDepth * VirtualVoxel.HairCoveragePixelRadiusAtDepth1;
-		TraversalSettings.bCastShadow = true;
-
-		if (bIsValid && RayMarchMaskTexture.Load(uint3(PixelCoord, 0)).r > 0.0f && HairPixelCoverage < CoverageThreshold)
+		const bool bHasPixelVisibleLight = RayMarchMaskTexture.Load(uint3(PixelCoord, 0)).r > 0.0f;
+		if (bHasPixelVisibleLight)
 		{
-			Result = ComputeHairCountVirtualVoxel(
-				TranslatedWorldPosition,
-				TranslatedLightPosition,
-				CommonDesc,
-				NodeDesc,
-				VirtualVoxel.PageIndexBuffer,
-				VirtualVoxel.PageIndexOccupancyBuffer,
-				VirtualVoxel.PageTexture,
-				TraversalSettings);
-		}
-		else
-		{
-			discard;
+			#if PERMUTATION_USE_ONEPASS
+			for (uint GroupIt=0;GroupIt< Voxel_MacroGroupCount;++GroupIt)
+			#else
+			const uint GroupIt = Voxel_MacroGroupId;
+			#endif
+			{
+
+				const FPackedVirtualVoxelNodeDesc PackedNode = VirtualVoxel.NodeDescBuffer[GroupIt];
+				const FVirtualVoxelNodeDesc NodeDesc = UnpackVoxelNode(PackedNode, VirtualVoxel.PageResolution);
+				const bool bIsValid = all(NodeDesc.PageIndexResolution != 0);
+				if (bIsValid)
+				{
+					FHairTraversalSettings TraversalSettings = InitHairTraversalSettings();
+					TraversalSettings.DensityScale = VirtualVoxel.DensityScale_Shadow;
+					TraversalSettings.CountThreshold = 1;
+					TraversalSettings.DistanceThreshold = DistanceThreshold;
+					TraversalSettings.bDebugEnabled = bDebugEnabled;
+					TraversalSettings.SteppingScale = VirtualVoxel.SteppingScale_Shadow;
+					TraversalSettings.Random = Random;
+					TraversalSettings.PixelRadius = SceneDepth * VirtualVoxel.HairCoveragePixelRadiusAtDepth1;
+					TraversalSettings.bCastShadow = true;
+
+					Result = ComputeHairCountVirtualVoxel(
+						TranslatedWorldPosition,
+						TranslatedLightPosition,
+						CommonDesc,
+						NodeDesc,
+						VirtualVoxel.PageIndexBuffer,
+						VirtualVoxel.PageIndexOccupancyBuffer,
+						VirtualVoxel.PageTexture,
+						TraversalSettings);
+
+					Visibility = min(Visibility, saturate(1 - Result.HairCount));
+					bHasValidOutput = true;
+				}
+			}
 		}
+	}
 
-		Visibility = saturate(1 - Result.HairCount);
+	if (!bHasValidOutput)
+	{
+		discard;
 	}
 	OutColor = WriteShadowMaskOutput(Visibility);
 }
diff --git a/Engine/Source/Runtime/Renderer/Private/HairStrands/HairStrandsTransmittance.cpp b/Engine/Source/Runtime/Renderer/Private/HairStrands/HairStrandsTransmittance.cpp
--- a/Engine/Source/Runtime/Renderer/Private/HairStrands/HairStrandsTransmittance.cpp
+++ b/Engine/Source/Runtime/Renderer/Private/HairStrands/HairStrandsTransmittance.cpp
@@ -50,6 +50,9 @@ static FAutoConsoleVariableRef CVarHairStrandsTransmittanceMaskUseMipTraversal(T
 static int32 GHairStrandsShadowRandomTraversalType = 2;
 static FAutoConsoleVariableRef CVarHairStrandsShadowRandomTraversalType(TEXT("r.HairStrands.DeepShadow.RandomType"), GHairStrandsShadowRandomTraversalType, TEXT("Change how traversal jittering is initialized. Valid value are 0, 1, and 2. Each type makes different type of tradeoff."), ECVF_Scalability | ECVF_RenderThreadSafe);
 
+static int32 GHairStrandsShadow_ShadowMaskPassType = 1;
+static FAutoConsoleVariableRef CVarHairStrandsShadow_ShadowMaskPassType(TEXT("r.HairStrands.DeepShadow.ShadowMaskPassType"), GHairStrandsShadow_ShadowMaskPassType, TEXT("Change how shadow mask from hair onto opaque geometry is generated. 0: one pass per hair group, 1: one pass for all groups."), ECVF_Scalability | ECVF_RenderThreadSafe);
+
 static float GetDeepShadowDensityScale() { return FMath::Max(0.0f, GDeepShadowDensityScale); }
 static float GetDeepShadowDepthBiasScale() { return FMath::Max(0.0f, GDeepShadowDepthBiasScale); }
 ///////////////////////////////////////////////////////////////////////////////////////////////////
@@ -440,10 +443,14 @@ class FHairStrandsVoxelShadowMaskPS : public FGlobalShader
 	DECLARE_GLOBAL_SHADER(FHairStrandsVoxelShadowMaskPS);
 	SHADER_USE_PARAMETER_STRUCT(FHairStrandsVoxelShadowMaskPS, FGlobalShader);
 
+	class FOnePass : SHADER_PERMUTATION_BOOL("PERMUTATION_USE_ONEPASS");
+	using FPermutationDomain = TShaderPermutationDomain<FOnePass>;
+
 	BEGIN_SHADER_PARAMETER_STRUCT(FParameters, )
 		SHADER_PARAMETER_STRUCT_INCLUDE(ShaderDrawDebug::FShaderParameters, ShaderDrawParameters)
 		
 		SHADER_PARAMETER(FVector4f, Voxel_TranslatedLightPosition_LightDirection)
+		SHADER_PARAMETER(uint32, Voxel_MacroGroupCount)
 		SHADER_PARAMETER(uint32, Voxel_MacroGroupId)
 		SHADER_PARAMETER(uint32, Voxel_RandomType)
 		SHADER_PARAMETER(uint32, bIsWholeSceneLight)
@@ -471,43 +478,20 @@ public:
 
 IMPLEMENT_GLOBAL_SHADER(FHairStrandsVoxelShadowMaskPS, "/Engine/Private/HairStrands/HairStrandsDeepShadowMask.usf", "MainPS", SF_Pixel);
 
-struct FHairStrandsVoxelShadowParams
-{
-	FVector4f		Voxel_TranslatedLightPosition_LightDirection = FVector4f(0, 0, 0, 0);
-	uint32			OutputChannel = ~0;
-	uint32			Voxel_MacroGroupId;
-	bool			bIsWholeSceneLight = false;
-};
-
-// Opaque mask from voxel
+// Opaque mask from voxels
 static void AddHairStrandsVoxelShadowMaskPass(
 	FRDGBuilder& GraphBuilder,
 	FRDGTextureRef SceneDepthTexture,
 	const FViewInfo& View,
-	const FHairStrandsVoxelShadowParams& Params,
+	const FLightSceneInfo* LightSceneInfo,
+	const bool bProjectingForForwardShading,
 	FRDGTextureRef& OutShadowMask)
 {
+	check(LightSceneInfo);
 	check(OutShadowMask);
 	check(HairStrands::HasViewHairStrandsVoxelData(View));
 
-	FHairStrandsVoxelShadowMaskPS::FParameters* Parameters = GraphBuilder.AllocParameters<FHairStrandsVoxelShadowMaskPS::FParameters>();
-	Parameters->ViewUniformBuffer = View.ViewUniformBuffer;
-	Parameters->SceneDepthTexture = SceneDepthTexture;	
-	Parameters->HairStrands = HairStrands::BindHairStrandsViewUniformParameters(View);
-	Parameters->VirtualVoxel = HairStrands::BindHairStrandsVoxelUniformParameters(View);
-	Parameters->LinearSampler = TStaticSamplerState<SF_Bilinear, AM_Clamp, AM_Clamp, AM_Clamp>::GetRHI();
-	Parameters->ShadowSampler = TStaticSamplerState<SF_Point, AM_Clamp, AM_Clamp, AM_Clamp>::GetRHI();
-	Parameters->bIsWholeSceneLight = Params.bIsWholeSceneLight ? 1 : 0;
-	Parameters->Voxel_TranslatedLightPosition_LightDirection = Params.Voxel_TranslatedLightPosition_LightDirection;
-	Parameters->Voxel_MacroGroupId		= Params.Voxel_MacroGroupId;
-	Parameters->Voxel_RandomType = FMath::Clamp(GHairStrandsShadowRandomTraversalType, 0, 2);	
-	Parameters->RenderTargets[0] = FRenderTargetBinding(OutShadowMask, ERenderTargetLoadAction::ELoad);
-
-	if (ShaderDrawDebug::IsEnabled(View))
-	{
-		ShaderDrawDebug::SetParameters(GraphBuilder, View.ShaderDrawData, Parameters->ShaderDrawParameters);
-	}
-
+	// Copy the shadow mask texture to read its content, and early out voxel traversal
 	FRDGTextureRef RayMarchMask = nullptr;
 	{
 		FRDGTextureDesc Desc = OutShadowMask->Desc;
@@ -517,61 +501,88 @@ static void AddHairStrandsVoxelShadowMaskPass(
 		CopyInfo.Size = OutShadowMask->Desc.GetSize();
 		AddCopyTexturePass(GraphBuilder, OutShadowMask, RayMarchMask, CopyInfo);
 	}
-	Parameters->RayMarchMaskTexture = RayMarchMask;
-
 
+	const bool bOnePass = GHairStrandsShadow_ShadowMaskPassType > 0 && View.HairStrandsViewData.MacroGroupDatas.Num() > 1;
+	const FIntPoint OutputResolution = SceneDepthTexture->Desc.Extent;
 	FHairStrandsVoxelShadowMaskPS::FPermutationDomain PermutationVector;
+	PermutationVector.Set<FHairStrandsVoxelShadowMaskPS::FOnePass>(bOnePass);
 
-	const FIntPoint OutputResolution = SceneDepthTexture->Desc.Extent;
 	TShaderMapRef<FPostProcessVS> VertexShader(View.ShaderMap);
 	TShaderMapRef<FHairStrandsVoxelShadowMaskPS> PixelShader(View.ShaderMap, PermutationVector);
 	const FGlobalShaderMap* GlobalShaderMap = View.ShaderMap;
 	const FIntRect Viewport = View.ViewRect;
-	const uint32 OutputChannel = Params.OutputChannel;
-
-	ClearUnusedGraphResources(PixelShader, Parameters);
-	FIntPoint Resolution = OutShadowMask->Desc.Extent;
+	const uint32 OutputChannel = bProjectingForForwardShading ? LightSceneInfo->GetDynamicShadowMapChannel() : ~0;
+	const FIntPoint Resolution = OutShadowMask->Desc.Extent;
+	const bool bIsWholeSceneLight = LightSceneInfo->Proxy->GetLightType() == LightType_Directional;
 
-	GraphBuilder.AddPass(
-		RDG_EVENT_NAME("HairStrands::ShadowMask(Voxel)"),
-		Parameters,
-		ERDGPassFlags::Raster,
-		[Parameters, VertexShader, PixelShader, Viewport, Resolution, OutputChannel](FRHICommandList& RHICmdList)
+	for (int32 GroupIt=0, GroupCount=View.HairStrandsViewData.MacroGroupDatas.Num(); GroupIt < GroupCount; ++GroupIt)
 	{
-		FGraphicsPipelineStateInitializer GraphicsPSOInit;
-		RHICmdList.ApplyCachedRenderTargets(GraphicsPSOInit);
-
-		switch (OutputChannel)
+		if (bOnePass && GroupIt > 0)
 		{
-		case 0:  GraphicsPSOInit.BlendState = TStaticBlendState<CW_RED,   BO_Min, BF_One, BF_One, BO_Min, BF_One, BF_One>::GetRHI(); break; // Min Operator
-		case 1:  GraphicsPSOInit.BlendState = TStaticBlendState<CW_GREEN, BO_Min, BF_One, BF_One, BO_Min, BF_One, BF_One>::GetRHI(); break; // Min Operator
-		case 2:  GraphicsPSOInit.BlendState = TStaticBlendState<CW_BLUE,  BO_Min, BF_One, BF_One, BO_Min, BF_One, BF_One>::GetRHI(); break; // Min Operator
-		case 3:  GraphicsPSOInit.BlendState = TStaticBlendState<CW_ALPHA, BO_Min, BF_One, BF_One, BO_Min, BF_One, BF_One>::GetRHI(); break; // Min Operator
-		default: GraphicsPSOInit.BlendState = TStaticBlendState<CW_RGBA,  BO_Min, BF_One, BF_One, BO_Min, BF_One, BF_One>::GetRHI(); break; // Min Operator
+			return;
 		}
-		GraphicsPSOInit.RasterizerState = TStaticRasterizerState<>::GetRHI();
-		GraphicsPSOInit.DepthStencilState = TStaticDepthStencilState<false, CF_Always>::GetRHI();
 
-		GraphicsPSOInit.BoundShaderState.VertexDeclarationRHI = GFilterVertexDeclaration.VertexDeclarationRHI;
-		GraphicsPSOInit.BoundShaderState.VertexShaderRHI = VertexShader.GetVertexShader();
-		GraphicsPSOInit.BoundShaderState.PixelShaderRHI = PixelShader.GetPixelShader();
-		GraphicsPSOInit.PrimitiveType = PT_TriangleList;
-		SetGraphicsPipelineState(RHICmdList, GraphicsPSOInit, 0);
+		const FHairStrandsMacroGroupData& MacroGroupData = View.HairStrandsViewData.MacroGroupDatas[GroupIt];
+
+		FHairStrandsVoxelShadowMaskPS::FParameters* Parameters = GraphBuilder.AllocParameters<FHairStrandsVoxelShadowMaskPS::FParameters>();
+		Parameters->ViewUniformBuffer = View.ViewUniformBuffer;
+		Parameters->SceneDepthTexture = SceneDepthTexture;	
+		Parameters->HairStrands = HairStrands::BindHairStrandsViewUniformParameters(View);
+		Parameters->VirtualVoxel = HairStrands::BindHairStrandsVoxelUniformParameters(View);
+		Parameters->LinearSampler = TStaticSamplerState<SF_Bilinear, AM_Clamp, AM_Clamp, AM_Clamp>::GetRHI();
+		Parameters->ShadowSampler = TStaticSamplerState<SF_Point, AM_Clamp, AM_Clamp, AM_Clamp>::GetRHI();
+		Parameters->bIsWholeSceneLight = bIsWholeSceneLight ? 1 : 0;
+		Parameters->Voxel_TranslatedLightPosition_LightDirection = GetLightTranslatedWorldPositionAndDirection(View, LightSceneInfo);
+		Parameters->Voxel_MacroGroupId	= MacroGroupData.MacroGroupId;
+		Parameters->Voxel_MacroGroupCount = GroupCount;
+		Parameters->Voxel_RandomType = FMath::Clamp(GHairStrandsShadowRandomTraversalType, 0, 2);	
+		Parameters->RenderTargets[0] = FRenderTargetBinding(OutShadowMask, ERenderTargetLoadAction::ELoad);
+		ShaderDrawDebug::SetParameters(GraphBuilder, View.ShaderDrawData, Parameters->ShaderDrawParameters);
+		Parameters->RayMarchMaskTexture = RayMarchMask;
 
-		RHICmdList.SetViewport(Viewport.Min.X, Viewport.Min.Y, 0.0f, Viewport.Max.X, Viewport.Max.Y, 1.0f);
-		SetShaderParameters(RHICmdList, PixelShader, PixelShader.GetPixelShader(), *Parameters);
+		ClearUnusedGraphResources(PixelShader, Parameters);
 
-		DrawRectangle(
-			RHICmdList,
-			0, 0,
-			Viewport.Width(), Viewport.Height(),
-			Viewport.Min.X, Viewport.Min.Y,
-			Viewport.Width(), Viewport.Height(),
-			Viewport.Size(),
-			Resolution,
-			VertexShader,
-			EDRF_UseTriangleOptimization);
-	});
+		GraphBuilder.AddPass(
+			RDG_EVENT_NAME("HairStrands::ShadowMask(Voxel,%s)", bOnePass ? TEXT("OnePass") : TEXT("PerGroup")),
+			Parameters,
+			ERDGPassFlags::Raster,
+			[Parameters, VertexShader, PixelShader, Viewport, Resolution, OutputChannel](FRHICommandList& RHICmdList)
+		{
+			FGraphicsPipelineStateInitializer GraphicsPSOInit;
+			RHICmdList.ApplyCachedRenderTargets(GraphicsPSOInit);
+
+			switch (OutputChannel)
+			{
+			case 0:  GraphicsPSOInit.BlendState = TStaticBlendState<CW_RED,   BO_Min, BF_One, BF_One, BO_Min, BF_One, BF_One>::GetRHI(); break; // Min Operator
+			case 1:  GraphicsPSOInit.BlendState = TStaticBlendState<CW_GREEN, BO_Min, BF_One, BF_One, BO_Min, BF_One, BF_One>::GetRHI(); break; // Min Operator
+			case 2:  GraphicsPSOInit.BlendState = TStaticBlendState<CW_BLUE,  BO_Min, BF_One, BF_One, BO_Min, BF_One, BF_One>::GetRHI(); break; // Min Operator
+			case 3:  GraphicsPSOInit.BlendState = TStaticBlendState<CW_ALPHA, BO_Min, BF_One, BF_One, BO_Min, BF_One, BF_One>::GetRHI(); break; // Min Operator
+			default: GraphicsPSOInit.BlendState = TStaticBlendState<CW_RGBA,  BO_Min, BF_One, BF_One, BO_Min, BF_One, BF_One>::GetRHI(); break; // Min Operator
+			}
+			GraphicsPSOInit.RasterizerState = TStaticRasterizerState<>::GetRHI();
+			GraphicsPSOInit.DepthStencilState = TStaticDepthStencilState<false, CF_Always>::GetRHI();
+
+			GraphicsPSOInit.BoundShaderState.VertexDeclarationRHI = GFilterVertexDeclaration.VertexDeclarationRHI;
+			GraphicsPSOInit.BoundShaderState.VertexShaderRHI = VertexShader.GetVertexShader();
+			GraphicsPSOInit.BoundShaderState.PixelShaderRHI = PixelShader.GetPixelShader();
+			GraphicsPSOInit.PrimitiveType = PT_TriangleList;
+			SetGraphicsPipelineState(RHICmdList, GraphicsPSOInit, 0);
+
+			RHICmdList.SetViewport(Viewport.Min.X, Viewport.Min.Y, 0.0f, Viewport.Max.X, Viewport.Max.Y, 1.0f);
+			SetShaderParameters(RHICmdList, PixelShader, PixelShader.GetPixelShader(), *Parameters);
+
+			DrawRectangle(
+				RHICmdList,
+				0, 0,
+				Viewport.Width(), Viewport.Height(),
+				Viewport.Min.X, Viewport.Min.Y,
+				Viewport.Width(), Viewport.Height(),
+				Viewport.Size(),
+				Resolution,
+				VertexShader,
+				EDRF_UseTriangleOptimization);
+		});
+	}
 }
 
 ///////////////////////////////////////////////////////////////////////////////////////////////////
@@ -762,7 +773,7 @@ static FHairStrandsTransmittanceMaskData InternalRenderHairStrandsTransmittanceM
 		return Out;
 
 	DECLARE_GPU_STAT(HairStrandsTransmittanceMask);
-	RDG_EVENT_SCOPE(GraphBuilder, "HairStrandsTransmittanceMask");
+	RDG_EVENT_SCOPE(GraphBuilder, "HairStrands::TransmittanceMask");
 	RDG_GPU_STAT_SCOPE(GraphBuilder, HairStrandsTransmittanceMask);
 
 	// Note: GbufferB.a store the shading model on the 4 lower bits (MATERIAL_SHADINGMODEL_HAIR)
@@ -856,7 +867,7 @@ FHairStrandsTransmittanceMaskData RenderHairStrandsOnePassTransmittanceMask(
 	if (HairStrands::HasViewHairStrandsData(View) && View.HairStrandsViewData.MacroGroupDatas.Num() > 0)
 	{
 		DECLARE_GPU_STAT(HairStrandsOnePassTransmittanceMask);
-		RDG_EVENT_SCOPE(GraphBuilder, "HairStrandsTransmittanceMask(OnePass)");
+		RDG_EVENT_SCOPE(GraphBuilder, "HairStrands::TransmittanceMask(OnePass)");
 		RDG_GPU_STAT_SCOPE(GraphBuilder, HairStrandsOnePassTransmittanceMask);
 
 		if (HairStrands::HasViewHairStrandsVoxelData(View))
@@ -924,7 +935,7 @@ static void InternalRenderHairStrandsShadowMask(
 		return;
 
 	DECLARE_GPU_STAT(HairStrandsOpaqueMask);
-	RDG_EVENT_SCOPE(GraphBuilder, "HairStrandsOpaqueMask");
+	RDG_EVENT_SCOPE(GraphBuilder, "HairStrands::OpaqueShadowMask");
 	RDG_GPU_STAT_SCOPE(GraphBuilder, HairStrandsOpaqueMask);
 	const FMinimalSceneTextures& SceneTextures = FSceneTextures::Get(GraphBuilder);
 
@@ -973,28 +984,16 @@ static void InternalRenderHairStrandsShadowMask(
 		}
 	}
 
-	// Code is disabled for now until we have the full DOM/voxel fallback logic
 	// If there is no deep shadow for this light, fallback on the voxel representation
 	if (!bHasDeepShadow && HairStrands::HasViewHairStrandsVoxelData(View))
 	{
-		// TODO: Change this to be a single pass with virtual voxel?
-		for (const FHairStrandsMacroGroupData& MacroGroupData : InMacroGroupDatas)
-		{
-			const bool bIsWholeSceneLight = LightSceneInfo->Proxy->GetLightType() == LightType_Directional;
-
-			FHairStrandsVoxelShadowParams Params;
-			Params.Voxel_TranslatedLightPosition_LightDirection = GetLightTranslatedWorldPositionAndDirection(View, LightSceneInfo);
-			Params.Voxel_MacroGroupId = MacroGroupData.MacroGroupId;
-			Params.bIsWholeSceneLight = bIsWholeSceneLight ? 1 : 0;
-			Params.OutputChannel = bProjectingForForwardShading ? LightSceneInfo->GetDynamicShadowMapChannel() : ~0;
-
-			AddHairStrandsVoxelShadowMaskPass(
-				GraphBuilder,
-				SceneTextures.Depth.Resolve,
-				View,
-				Params,
-				OutShadowMask);
-		}
+		AddHairStrandsVoxelShadowMaskPass(
+			GraphBuilder,
+			SceneTextures.Depth.Resolve,
+			View,
+			LightSceneInfo,
+			bProjectingForForwardShading,
+			OutShadowMask);
 	}
 }
 
