PostProcessing: fix for sequence of after-passes where the next pass wasn't consuming the last result.

diff --git a/Engine/Source/Runtime/Renderer/Private/PostProcess/PostProcessing.cpp b/Engine/Source/Runtime/Renderer/Private/PostProcess/PostProcessing.cpp
--- a/Engine/Source/Runtime/Renderer/Private/PostProcess/PostProcessing.cpp
+++ b/Engine/Source/Runtime/Renderer/Private/PostProcess/PostProcessing.cpp
@@ -379,6 +379,8 @@ void AddPostProcessingPasses(
 
 			for (int32 AfterPassCallbackIndex = 0; AfterPassCallbackIndex < PassCallbacks.Num(); AfterPassCallbackIndex++)
 			{
+				InOutPostProcessAfterPassInputs.SetInput(EPostProcessMaterialInput::SceneColor, InSceneColor);
+
 				FAfterPassCallbackDelegate& AfterPassCallback = PassCallbacks[AfterPassCallbackIndex];
 				PassSequence.AcceptOverrideIfLastPass(InPass, InOutPostProcessAfterPassInputs.OverrideOutput, AfterPassCallbackIndex);
 				InSceneColor = AfterPassCallback.Execute(GraphBuilder, View, InOutPostProcessAfterPassInputs);
