Fix GlobalDistanceField corruption when scene only contains heightfields.
Issue:
- Can't rely on FComposeHeightfieldsIntoPagesCS to initialize pages in its current form
- Multiple heightfield components can write to the same pages.
- FComposeHeightfieldsIntoPagesCS only runs on pages that overlap with heightfield components.

Fix:
- Run FComposeObjectsIntoPagesCS even when scene doesn't contain any mesh SDFs.
- It will initialize all pages that require it.
- Added FProcessDistanceFields permutation to skip iterating through mesh SDF lists when scene doesn't contain any.

diff --git a/Engine/Shaders/Private/GlobalDistanceField.usf b/Engine/Shaders/Private/GlobalDistanceField.usf
--- a/Engine/Shaders/Private/GlobalDistanceField.usf
+++ b/Engine/Shaders/Private/GlobalDistanceField.usf
@@ -595,6 +595,7 @@ void ComposeObjectsIntoPagesCS(
 
 	float MinDistance = InfluenceRadius;
 
+#if PROCESS_DISTANCE_FIELDS
 #undef USE_OBJECT_COMPOSITING_TILE_CULLING
 #define USE_OBJECT_COMPOSITING_TILE_CULLING 1
 #if USE_OBJECT_COMPOSITING_TILE_CULLING
@@ -654,6 +655,7 @@ void ComposeObjectsIntoPagesCS(
 	}
 
 #endif
+#endif // PROCESS_DISTANCE_FIELDS
 
 	if (all(TexelCoordInPage < GLOBAL_DISTANCE_FIELD_PAGE_RESOLUTION_IN_ATLAS))
 	{
diff --git a/Engine/Shaders/Private/GlobalDistanceFieldHeightfields.usf b/Engine/Shaders/Private/GlobalDistanceFieldHeightfields.usf
--- a/Engine/Shaders/Private/GlobalDistanceFieldHeightfields.usf
+++ b/Engine/Shaders/Private/GlobalDistanceFieldHeightfields.usf
@@ -11,10 +11,6 @@
 #define GLOBAL_DISTANCE_FIELD_PAGE_HEIGHFIELD_GROUP_X (GLOBAL_DISTANCE_FIELD_PAGE_RESOLUTION_IN_ATLAS / 8)
 #define GLOBAL_DISTANCE_FIELD_PAGE_HEIGHFIELD_GROUP_NUM (GLOBAL_DISTANCE_FIELD_PAGE_HEIGHFIELD_GROUP_X * GLOBAL_DISTANCE_FIELD_PAGE_HEIGHFIELD_GROUP_X)
 
-#define COMPOSE_DISTANCE_FIELD_NONE 0
-#define COMPOSE_DISTANCE_FIELD_PARENT 1
-#define COMPOSE_DISTANCE_FIELD_PREVIOUS 2
-
 RWStructuredBuffer<uint> RWMarkedHeightfieldPageBuffer;
 RWTexture3D<UNORM float> RWPageAtlasTexture;
 
@@ -231,7 +227,6 @@ void BuildHeightfieldComposeTilesCS(
 [numthreads(THREADGROUP_SIZE_X, THREADGROUP_SIZE_Y, 1)]
 void ComposeHeightfieldsIntoPagesCS(
 	uint3 GroupId : SV_GroupID,
-	uint3 DispatchThreadId : SV_DispatchThreadID,
 	uint3 GroupThreadId : SV_GroupThreadID)
 {
 	uint PackedPageTile = ComposeTileBuffer[GroupId.x / GLOBAL_DISTANCE_FIELD_PAGE_HEIGHFIELD_GROUP_NUM];
@@ -305,20 +300,8 @@ void ComposeHeightfieldsIntoPagesCS(
 
 				const uint3 PageAtlasCoordZ = PageAtlasCoord + uint3(0, 0, ZIndex);
 
-#if COMPOSE_DISTANCE_FIELD_MODE == COMPOSE_DISTANCE_FIELD_PARENT
-				uint ParentPageId = ParentPageTableLayerTexture.Load(int4(PageTableTextureCoord, 0));
-				if (ParentPageId != GLOBAL_DISTANCE_FIELD_INVALID_PAGE_ID)
-				{
-					uint3 ParentPageAtlasCoord = GlobalDistanceFieldPageLinearIndexToPageAtlasCoord(ParentPageId);
-					ParentPageAtlasCoord += TexelCoordInPage;
-
-					float ParentDistanceField = DecodeGlobalDistanceFieldPageDistance(RWPageAtlasTexture[ParentPageAtlasCoord].x, InfluenceRadius);
-					MinDistance = min(MinDistance, ParentDistanceField);
-				}
-#elif COMPOSE_DISTANCE_FIELD_MODE == COMPOSE_DISTANCE_FIELD_PREVIOUS
 				float PreviousDistanceField = DecodeGlobalDistanceFieldPageDistance(RWPageAtlasTexture[PageAtlasCoordZ], InfluenceRadius);
 				MinDistance = min(MinDistance, PreviousDistanceField);
-#endif
 
 				RWPageAtlasTexture[PageAtlasCoordZ] = EncodeGlobalDistanceFieldPageDistance(MinDistance, InfluenceRadius);
 			}
diff --git a/Engine/Source/Runtime/Renderer/Private/GlobalDistanceField.cpp b/Engine/Source/Runtime/Renderer/Private/GlobalDistanceField.cpp
--- a/Engine/Source/Runtime/Renderer/Private/GlobalDistanceField.cpp
+++ b/Engine/Source/Runtime/Renderer/Private/GlobalDistanceField.cpp
@@ -1307,7 +1307,8 @@ class FComposeObjectsIntoPagesCS : public FGlobalShader
 	END_SHADER_PARAMETER_STRUCT()
 
 	class FComposeParentDistanceField : SHADER_PERMUTATION_BOOL("COMPOSE_PARENT_DISTANCE_FIELD");
-	using FPermutationDomain = TShaderPermutationDomain<FComposeParentDistanceField>;
+	class FProcessDistanceFields : SHADER_PERMUTATION_BOOL("PROCESS_DISTANCE_FIELDS");
+	using FPermutationDomain = TShaderPermutationDomain<FComposeParentDistanceField, FProcessDistanceFields>;
 
 	static bool ShouldCompilePermutation(const FGlobalShaderPermutationParameters& Parameters)
 	{
@@ -2146,9 +2147,8 @@ void UpdateGlobalDistanceFieldVolume(
 							}
 						}
 
-						// Compose the mesh SDFs into allocated pages
-						const bool bComposeMeshSDFsIntoPages = Scene->DistanceFieldSceneData.NumObjectsInBuffer > 0;
-						if (bComposeMeshSDFsIntoPages)
+						// Initialize pages and compose the mesh SDFs into allocated pages
+						if(Scene->DistanceFieldSceneData.NumObjectsInBuffer > 0 || UpdateRegionHeightfield.ComponentDescriptions.Num() > 0)
 						{
 							const FVector PageVoxelExtent = 0.5f * ClipmapSize / FVector(ClipmapResolution);
 							const FVector PageCoordToVoxelCenterScale = ClipmapSize / FVector(ClipmapResolution);
@@ -2188,6 +2188,7 @@ void UpdateGlobalDistanceFieldVolume(
 
 							FComposeObjectsIntoPagesCS::FPermutationDomain PermutationVector;
 							PermutationVector.Set<FComposeObjectsIntoPagesCS::FComposeParentDistanceField>(ParentPageTableLayerTexture != nullptr);
+							PermutationVector.Set<FComposeObjectsIntoPagesCS::FProcessDistanceFields>(Scene->DistanceFieldSceneData.NumObjectsInBuffer > 0);
 							auto ComputeShader = View.ShaderMap->GetShader<FComposeObjectsIntoPagesCS>(PermutationVector);
 
 							FComputeShaderUtils::AddPass(
@@ -2249,20 +2250,7 @@ void UpdateGlobalDistanceFieldVolume(
 									PassParameters->VisibilitySampler = TStaticSamplerState<SF_Bilinear>::GetRHI();
 									PassParameters->HeightfieldDescriptions = GraphBuilder.CreateSRV(HeightfieldDescriptionBuffer, EPixelFormat::PF_A32B32G32R32F);
 
-									EComposeDistanceFieldMode ComposeMode = EComposeDistanceFieldMode::None;
-
-									if (bComposeMeshSDFsIntoPages)
-									{
-										ComposeMode = EComposeDistanceFieldMode::Previous;
-									}
-									else if (ParentPageTableLayerTexture != nullptr)
-									{
-										ComposeMode = EComposeDistanceFieldMode::Parent;
-									}
-
-									FComposeHeightfieldsIntoPagesCS::FPermutationDomain PermutationVector;
-									PermutationVector.Set<FComposeHeightfieldsIntoPagesCS::FComposeDistanceFieldMode>(ComposeMode);
-									auto ComputeShader = View.ShaderMap->GetShader<FComposeHeightfieldsIntoPagesCS>(PermutationVector);
+									auto ComputeShader = View.ShaderMap->GetShader<FComposeHeightfieldsIntoPagesCS>();
 
 									FComputeShaderUtils::AddPass(
 										GraphBuilder,
diff --git a/Engine/Source/Runtime/Renderer/Private/GlobalDistanceFieldHeightfields.h b/Engine/Source/Runtime/Renderer/Private/GlobalDistanceFieldHeightfields.h
--- a/Engine/Source/Runtime/Renderer/Private/GlobalDistanceFieldHeightfields.h
+++ b/Engine/Source/Runtime/Renderer/Private/GlobalDistanceFieldHeightfields.h
@@ -116,16 +116,6 @@ class FBuildHeightfieldComposeTilesCS : public FGlobalShader
 
 IMPLEMENT_GLOBAL_SHADER(FBuildHeightfieldComposeTilesCS, "/Engine/Private/GlobalDistanceFieldHeightfields.usf", "BuildHeightfieldComposeTilesCS", SF_Compute);
 
-// Must match COMPOSE_DISTANCE_FIELD_ in usf
-enum class EComposeDistanceFieldMode
-{
-	None,
-	Parent,
-	Previous,
-
-	MAX
-};
-
 class FComposeHeightfieldsIntoPagesCS : public FGlobalShader
 {
 	DECLARE_GLOBAL_SHADER(FComposeHeightfieldsIntoPagesCS);
@@ -157,9 +147,6 @@ class FComposeHeightfieldsIntoPagesCS : public FGlobalShader
 		SHADER_PARAMETER(float, HeightfieldThickness)
 	END_SHADER_PARAMETER_STRUCT()
 
-	class FComposeDistanceFieldMode : SHADER_PERMUTATION_ENUM_CLASS("COMPOSE_DISTANCE_FIELD_MODE", EComposeDistanceFieldMode);
-	using FPermutationDomain = TShaderPermutationDomain<FComposeDistanceFieldMode>;
-
 	static bool ShouldCompilePermutation(const FGlobalShaderPermutationParameters& Parameters)
 	{
 		return ShouldCompileDistanceFieldShaders(Parameters.Platform);
