LightGrid in translated world space
- not yet taking full advantage of LWC since light positions are still sent to shader as world space floats.

diff --git a/Engine/Plugins/Experimental/GPULightmass/Source/GPULightmass/Private/LightmapRayTracing.h b/Engine/Plugins/Experimental/GPULightmass/Source/GPULightmass/Private/LightmapRayTracing.h
--- a/Engine/Plugins/Experimental/GPULightmass/Source/GPULightmass/Private/LightmapRayTracing.h
+++ b/Engine/Plugins/Experimental/GPULightmass/Source/GPULightmass/Private/LightmapRayTracing.h
@@ -15,8 +15,8 @@
 
 BEGIN_SHADER_PARAMETER_STRUCT(FPathTracingLightGrid, RENDERER_API)
 	SHADER_PARAMETER(uint32, SceneInfiniteLightCount)
-	SHADER_PARAMETER(FVector3f, SceneLightsBoundMin)
-	SHADER_PARAMETER(FVector3f, SceneLightsBoundMax)
+	SHADER_PARAMETER(FVector3f, SceneLightsTranslatedBoundMin)
+	SHADER_PARAMETER(FVector3f, SceneLightsTranslatedBoundMax)
 	SHADER_PARAMETER_RDG_TEXTURE(Texture2D, LightGrid)
 	SHADER_PARAMETER_RDG_BUFFER_SRV(Buffer<uint>, LightGridData)
 	SHADER_PARAMETER(unsigned, LightGridResolution)
diff --git a/Engine/Plugins/Experimental/GPULightmass/Source/GPULightmass/Private/LightmapRenderer.cpp b/Engine/Plugins/Experimental/GPULightmass/Source/GPULightmass/Private/LightmapRenderer.cpp
--- a/Engine/Plugins/Experimental/GPULightmass/Source/GPULightmass/Private/LightmapRenderer.cpp
+++ b/Engine/Plugins/Experimental/GPULightmass/Source/GPULightmass/Private/LightmapRenderer.cpp
@@ -2318,7 +2318,7 @@ void FLightmapRenderer::Finalize(FRDGBuilder& GraphBuilder)
 									PassParameters->ViewUniformBuffer = Scene->ReferenceView->ViewUniformBuffer;
 									PassParameters->IrradianceCachingParameters = Scene->IrradianceCache->IrradianceCachingParametersUniformBuffer;
 
-									SetupPathTracingLightParameters(Scene->LightSceneRenderState, GraphBuilder, PassParameters);
+									SetupPathTracingLightParameters(Scene->LightSceneRenderState, GraphBuilder, *Scene->ReferenceView, PassParameters);
 
 									// TODO: find a way to share IES atlas with path tracer ...
 									PassParameters->IESTexture = GraphBuilder.RegisterExternalTexture(GSystemTextures.WhiteDummy, TEXT("IESTexture"));
diff --git a/Engine/Plugins/Experimental/GPULightmass/Source/GPULightmass/Private/PathTracingLightParameters.inl b/Engine/Plugins/Experimental/GPULightmass/Source/GPULightmass/Private/PathTracingLightParameters.inl
--- a/Engine/Plugins/Experimental/GPULightmass/Source/GPULightmass/Private/PathTracingLightParameters.inl
+++ b/Engine/Plugins/Experimental/GPULightmass/Source/GPULightmass/Private/PathTracingLightParameters.inl
@@ -7,12 +7,13 @@
 
 RENDERER_API FRDGTexture* PrepareIESAtlas(const TMap<FTexture*, int>& InIESLightProfilesMap, FRDGBuilder& GraphBuilder);
 
-RENDERER_API void PrepareLightGrid(FRDGBuilder& GraphBuilder, FPathTracingLightGrid* LightGridParameters, const FPathTracingLight* Lights, uint32 NumLights, uint32 NumInfiniteLights, FRDGBufferSRV* LightsSRV);
+RENDERER_API void PrepareLightGrid(FRDGBuilder& GraphBuilder, const FViewInfo& View, FPathTracingLightGrid* LightGridParameters, const FPathTracingLight* Lights, uint32 NumLights, uint32 NumInfiniteLights, FRDGBufferSRV* LightsSRV);
 
 template<typename PassParameterType>
 void SetupPathTracingLightParameters(
 	const GPULightmass::FLightSceneRenderState& LightScene,
 	FRDGBuilder& GraphBuilder,
+	const FViewInfo& View,
 	PassParameterType* PassParameters)
 {
 	TArray<FPathTracingLight> Lights;
@@ -223,5 +224,5 @@ void SetupPathTracingLightParameters(
 		PassParameters->IESTexture = GraphBuilder.RegisterExternalTexture(GSystemTextures.WhiteDummy);
 	}
 
-	PrepareLightGrid(GraphBuilder, &PassParameters->LightGridParameters, Lights.GetData(), PassParameters->SceneLightCount, NumInfiniteLights, PassParameters->SceneLights);
+	PrepareLightGrid(GraphBuilder, View, &PassParameters->LightGridParameters, Lights.GetData(), PassParameters->SceneLightCount, NumInfiniteLights, PassParameters->SceneLights);
 }
diff --git a/Engine/Plugins/Experimental/GPULightmass/Source/GPULightmass/Private/VolumetricLightmap.cpp b/Engine/Plugins/Experimental/GPULightmass/Source/GPULightmass/Private/VolumetricLightmap.cpp
--- a/Engine/Plugins/Experimental/GPULightmass/Source/GPULightmass/Private/VolumetricLightmap.cpp
+++ b/Engine/Plugins/Experimental/GPULightmass/Source/GPULightmass/Private/VolumetricLightmap.cpp
@@ -75,7 +75,7 @@ void ReleaseBrickData(FVolumetricLightmapBrickData& BrickData)
 	BrickData.DirectionalLightShadowing.UAV.SafeRelease();
 }
 
-void SetupPathTracingLightParameters(const GPULightmass::FLightSceneRenderState& LightScene, FRDGBuilder& GraphBuilder, FRDGBufferSRV** OutLightBuffer, uint32* OutLightCount);
+void SetupPathTracingLightParameters(const GPULightmass::FLightSceneRenderState& LightScene, FRDGBuilder& GraphBuilder, const FViewInfo& View, FRDGBufferSRV** OutLightBuffer, uint32* OutLightCount);
 
 namespace GPULightmass
 {
@@ -571,7 +571,7 @@ void FVolumetricLightmapRenderer::BackgroundTick()
 					PassParameters->ViewUniformBuffer = Scene->ReferenceView->ViewUniformBuffer;
 					PassParameters->IrradianceCachingParameters = Scene->IrradianceCache->IrradianceCachingParametersUniformBuffer;
 
-					SetupPathTracingLightParameters(Scene->LightSceneRenderState, GraphBuilder, PassParameters);
+					SetupPathTracingLightParameters(Scene->LightSceneRenderState, GraphBuilder, *Scene->ReferenceView, PassParameters);
 
 
 					PassParameters->IESTextureSampler = TStaticSamplerState<SF_Bilinear, AM_Clamp, AM_Clamp, AM_Clamp>::GetRHI();
diff --git a/Engine/Shaders/Private/PathTracing/Light/PathTracingLightGrid.ush b/Engine/Shaders/Private/PathTracing/Light/PathTracingLightGrid.ush
--- a/Engine/Shaders/Private/PathTracing/Light/PathTracingLightGrid.ush
+++ b/Engine/Shaders/Private/PathTracing/Light/PathTracingLightGrid.ush
@@ -10,8 +10,8 @@
 #if USE_PATH_TRACING_LIGHT_GRID
 // extra shader parameters that must be defined when using the light grid
 uint SceneInfiniteLightCount;
-float3 SceneLightsBoundMin;
-float3 SceneLightsBoundMax;
+float3 SceneLightsTranslatedBoundMin;
+float3 SceneLightsTranslatedBoundMax;
 Texture2D<uint> LightGrid;
 Buffer<uint> LightGridData;
 uint LightGridResolution;
@@ -38,17 +38,17 @@ struct FLightLoopCount {
 
 #if USE_PATH_TRACING_LIGHT_GRID
 
-FLightLoopCount LightGridLookup(float3 WorldPos)
+FLightLoopCount LightGridLookup(float3 TranslatedWorldPos)
 {
 	FLightLoopCount Result;
 	Result.NumLights = SceneInfiniteLightCount;
 	Result.NumMISLights = SceneInfiniteLightCount;
 	Result.LightGridOffset = ~0u;
 
-	if (all(SceneLightsBoundMin <= WorldPos) && all(WorldPos <= SceneLightsBoundMax))
+	if (all(SceneLightsTranslatedBoundMin <= TranslatedWorldPos) && all(TranslatedWorldPos <= SceneLightsTranslatedBoundMax))
 	{
-		float3 P = WorldPos - SceneLightsBoundMin;
-		float3 D = SceneLightsBoundMax - SceneLightsBoundMin;
+		float3 P = TranslatedWorldPos - SceneLightsTranslatedBoundMin;
+		float3 D = SceneLightsTranslatedBoundMax - SceneLightsTranslatedBoundMin;
 		int2 UV = 0;
 		switch (LightGridAxis)
 		{
@@ -132,7 +132,7 @@ uint GetLightId(uint Index, FLightLoopCount LoopCount)
 
 #else
 
-FLightLoopCount LightGridLookup(float3 WorldPos)
+FLightLoopCount LightGridLookup(float3 TranslatedWorldPos)
 {
 	uint NumLights = min(SceneLightCount, RAY_TRACING_LIGHT_COUNT_MAXIMUM);
 	FLightLoopCount Result;
diff --git a/Engine/Shaders/Private/PathTracing/PathTracingBuildLightGrid.usf b/Engine/Shaders/Private/PathTracing/PathTracingBuildLightGrid.usf
--- a/Engine/Shaders/Private/PathTracing/PathTracingBuildLightGrid.usf
+++ b/Engine/Shaders/Private/PathTracing/PathTracingBuildLightGrid.usf
@@ -10,8 +10,9 @@
 #include "Light/PathTracingLightCommon.ush"
 
 uint SceneInfiniteLightCount;
-float3 SceneLightsBoundMin;
-float3 SceneLightsBoundMax;
+float3 SceneLightsTranslatedBoundMin;
+float3 SceneLightsTranslatedBoundMax;
+float3 PreViewTranslation;
 RWTexture2D<uint> RWLightGrid;
 RWBuffer<uint> RWLightGridData;
 uint LightGridResolution;
@@ -35,8 +36,8 @@ void PathTracingBuildLightGridCS(uint2 DispatchThreadId : SV_DispatchThreadID)
 		case 2: VoxelId.xy = DispatchThreadId; VoxelRes.xy = LightGridResolution; break;
 	}
 	// get bounding box of current voxel
-	float3 Lo = lerp(SceneLightsBoundMin, SceneLightsBoundMax, float3(VoxelId + 0) / float3(VoxelRes));
-	float3 Hi = lerp(SceneLightsBoundMin, SceneLightsBoundMax, float3(VoxelId + 1) / float3(VoxelRes));
+	float3 Lo = lerp(SceneLightsTranslatedBoundMin, SceneLightsTranslatedBoundMax, float3(VoxelId + 0) / float3(VoxelRes));
+	float3 Hi = lerp(SceneLightsTranslatedBoundMin, SceneLightsTranslatedBoundMax, float3(VoxelId + 1) / float3(VoxelRes));
 
 	uint NumVisible = 0;
 	uint MaxVisible = min(LightGridMaxCount, RAY_TRACING_LIGHT_COUNT_MAXIMUM - SceneInfiniteLightCount);
@@ -45,7 +46,7 @@ void PathTracingBuildLightGridCS(uint2 DispatchThreadId : SV_DispatchThreadID)
 	for (uint LightIndex = SceneInfiniteLightCount; LightIndex < SceneLightCount; LightIndex++)
 	{
 		uint LightType = SceneLights[LightIndex].Flags & PATHTRACER_FLAG_TYPE_MASK;
-		float3 Center = GetPosition(LightIndex);
+		float3 Center = GetPosition(LightIndex) + PreViewTranslation; // RT_LWC_TODO
 		float Radius = 1.0 / GetAttenuation(LightIndex);
 
 		float3 E = max(Lo - Center, 0) + max(Center - Hi, 0);
diff --git a/Engine/Shaders/Private/PathTracing/PathTracingCore.ush b/Engine/Shaders/Private/PathTracing/PathTracingCore.ush
--- a/Engine/Shaders/Private/PathTracing/PathTracingCore.ush
+++ b/Engine/Shaders/Private/PathTracing/PathTracingCore.ush
@@ -741,7 +741,7 @@ bool PathTracingKernel(inout FPathState PathState, int Bounce)
 		return false;
 	}
 
-	FLightLoopCount LightLoopCount = LightGridLookup(Payload.WorldPos);
+	FLightLoopCount LightLoopCount = LightGridLookup(Payload.WorldPos + LWCHackToFloat(PrimaryView.PreViewTranslation)); // RT_LWC_TODO
 
 	// visualize the # of lights
 	if (VisualizeLightGrid)
diff --git a/Engine/Source/Runtime/Renderer/Private/PathTracing.cpp b/Engine/Source/Runtime/Renderer/Private/PathTracing.cpp
--- a/Engine/Source/Runtime/Renderer/Private/PathTracing.cpp
+++ b/Engine/Source/Runtime/Renderer/Private/PathTracing.cpp
@@ -448,8 +448,9 @@ IMPLEMENT_SHADER_TYPE(, FPathTracingSkylightMISCompensationCS, TEXT("/Engine/Pri
 // this struct holds a light grid for both building or rendering
 BEGIN_SHADER_PARAMETER_STRUCT(FPathTracingLightGrid, RENDERER_API)
 	SHADER_PARAMETER(uint32, SceneInfiniteLightCount)
-	SHADER_PARAMETER(FVector3f, SceneLightsBoundMin)
-	SHADER_PARAMETER(FVector3f, SceneLightsBoundMax)
+	SHADER_PARAMETER(FVector3f, SceneLightsTranslatedBoundMin)
+	SHADER_PARAMETER(FVector3f, SceneLightsTranslatedBoundMax)
+	SHADER_PARAMETER(FVector3f, PreViewTranslation) // RT_LWC_TODO
 	SHADER_PARAMETER_RDG_TEXTURE(Texture2D, LightGrid)
 	SHADER_PARAMETER_RDG_BUFFER_SRV(Buffer<uint>, LightGridData)
 	SHADER_PARAMETER(unsigned, LightGridResolution)
@@ -897,12 +898,13 @@ bool PrepareSkyTexture(FRDGBuilder& GraphBuilder, FScene* Scene, const FViewInfo
 	return true;
 }
 
-RENDERER_API void PrepareLightGrid(FRDGBuilder& GraphBuilder, FPathTracingLightGrid* LightGridParameters, const FPathTracingLight* Lights, uint32 NumLights, uint32 NumInfiniteLights, FRDGBufferSRV* LightsSRV)
+RENDERER_API void PrepareLightGrid(FRDGBuilder& GraphBuilder, const FViewInfo& View, FPathTracingLightGrid* LightGridParameters, const FPathTracingLight* Lights, uint32 NumLights, uint32 NumInfiniteLights, FRDGBufferSRV* LightsSRV)
 {
 	const float Inf = std::numeric_limits<float>::infinity();
 	LightGridParameters->SceneInfiniteLightCount = NumInfiniteLights;
-	LightGridParameters->SceneLightsBoundMin = FVector(+Inf, +Inf, +Inf);
-	LightGridParameters->SceneLightsBoundMax = FVector(-Inf, -Inf, -Inf);
+	LightGridParameters->SceneLightsTranslatedBoundMin = FVector(+Inf, +Inf, +Inf);
+	LightGridParameters->SceneLightsTranslatedBoundMax = FVector(-Inf, -Inf, -Inf);
+	LightGridParameters->PreViewTranslation = View.ViewMatrices.GetPreViewTranslation(); // RT_LWC_TODO
 	LightGridParameters->LightGrid = nullptr;
 	LightGridParameters->LightGridData = nullptr;
 
@@ -910,15 +912,21 @@ RENDERER_API void PrepareLightGrid(FRDGBuilder& GraphBuilder, FPathTracingLightG
 	// if we have some finite lights -- build a light grid
 	if (NumFiniteLights > 0)
 	{
+		FVector SceneLightsBoundMin = FVector(+Inf, +Inf, +Inf);
+		FVector SceneLightsBoundMax = FVector(-Inf, -Inf, -Inf);
+
 		// get bounding box of all finite lights
 		const FPathTracingLight* FiniteLights = Lights + NumInfiniteLights;
 		for (int Index = 0; Index < NumFiniteLights; Index++)
 		{
 			const FPathTracingLight& Light = FiniteLights[Index];
-			LightGridParameters->SceneLightsBoundMin = FVector::Min(LightGridParameters->SceneLightsBoundMin, Light.BoundMin);
-			LightGridParameters->SceneLightsBoundMax = FVector::Max(LightGridParameters->SceneLightsBoundMax, Light.BoundMax);
+			SceneLightsBoundMin = FVector::Min(SceneLightsBoundMin, Light.BoundMin);
+			SceneLightsBoundMax = FVector::Max(SceneLightsBoundMax, Light.BoundMax);
 		}
 
+		LightGridParameters->SceneLightsTranslatedBoundMin = SceneLightsBoundMin + View.ViewMatrices.GetPreViewTranslation();
+		LightGridParameters->SceneLightsTranslatedBoundMax = SceneLightsBoundMax + View.ViewMatrices.GetPreViewTranslation();
+
 		const uint32 Resolution = FMath::RoundUpToPowerOfTwo(CVarPathTracingLightGridResolution.GetValueOnRenderThread());
 		const uint32 MaxCount = FMath::Clamp(
 			CVarPathTracingLightGridMaxCount.GetValueOnRenderThread(),
@@ -929,7 +937,7 @@ RENDERER_API void PrepareLightGrid(FRDGBuilder& GraphBuilder, FPathTracingLightG
 		LightGridParameters->LightGridMaxCount = MaxCount;
 
 		// pick the shortest axis
-		FVector Diag = LightGridParameters->SceneLightsBoundMax - LightGridParameters->SceneLightsBoundMin;
+		FVector Diag = SceneLightsBoundMax - SceneLightsBoundMin;
 		if (Diag.X < Diag.Y && Diag.X < Diag.Z)
 		{
 			LightGridParameters->LightGridAxis = 0;
@@ -1274,7 +1282,7 @@ void SetLightParameters(FRDGBuilder& GraphBuilder, FPathTracingRG::FParameters*
 	}
 
 
-	PrepareLightGrid(GraphBuilder, &PassParameters->LightGridParameters, Lights, NumLights, NumInfiniteLights, PassParameters->SceneLights);
+	PrepareLightGrid(GraphBuilder, View, &PassParameters->LightGridParameters, Lights, NumLights, NumInfiniteLights, PassParameters->SceneLights);
 }
 
 class FPathTracingCompositorPS : public FGlobalShader
