Port RayTracingScene pass to RDG async compute pass.

diff --git a/Engine/Source/Runtime/Renderer/Private/DeferredShadingRenderer.cpp b/Engine/Source/Runtime/Renderer/Private/DeferredShadingRenderer.cpp
--- a/Engine/Source/Runtime/Renderer/Private/DeferredShadingRenderer.cpp
+++ b/Engine/Source/Runtime/Renderer/Private/DeferredShadingRenderer.cpp
@@ -1528,115 +1528,58 @@ bool FDeferredShadingSceneRenderer::DispatchRayTracingWorldUpdates(FRDGBuilder&
 		OutDynamicGeometryScratchBuffer = GraphBuilder.CreateBuffer(ScratchBufferDesc, TEXT("DynamicGeometry.BLASSharedScratchBuffer"));
 	}
 
-	const bool bRayTracingAsyncBuild = CVarRayTracingAsyncBuild.GetValueOnRenderThread() != 0;
+	const bool bRayTracingAsyncBuild = CVarRayTracingAsyncBuild.GetValueOnRenderThread() != 0 && GRHISupportsRayTracingAsyncBuildAccelerationStructure;
+	const ERDGPassFlags ComputePassFlags = bRayTracingAsyncBuild ? ERDGPassFlags::AsyncCompute : ERDGPassFlags::Compute;
 
-	if (bRayTracingAsyncBuild && GRHISupportsRayTracingAsyncBuildAccelerationStructure)
 	{
+		RDG_GPU_STAT_SCOPE(GraphBuilder, RayTracingScene);
+
 		FBuildAccelerationStructurePassParams* PassParams = GraphBuilder.AllocParameters<FBuildAccelerationStructurePassParams>();
 		PassParams->RayTracingSceneScratchBuffer = Scene->RayTracingScene.BuildScratchBuffer;
 		PassParams->RayTracingSceneInstanceBuffer = Scene->RayTracingScene.InstanceBuffer;
 		PassParams->DynamicGeometryScratchBuffer = OutDynamicGeometryScratchBuffer;
 
-		GraphBuilder.AddPass(RDG_EVENT_NAME("RayTracingScene"), PassParams, ERDGPassFlags::Compute | ERDGPassFlags::NeverCull,
-			[this, PassParams](FRHICommandListImmediate& RHICmdList)
-			{
-				FRHIAsyncComputeCommandListImmediate& RHIAsyncCmdList = FRHICommandListExecutor::GetImmediateAsyncComputeCommandList();
-
-
-				const FRHITransition* RayTracingDynamicGeometryUpdateBeginTransition = RHICreateTransition(FRHITransitionCreateInfo(ERHIPipeline::Graphics, ERHIPipeline::AsyncCompute));
-				RHICmdList.BeginTransition(RayTracingDynamicGeometryUpdateBeginTransition);
-				RHIAsyncCmdList.EndTransition(RayTracingDynamicGeometryUpdateBeginTransition);
-
-				FRHIBuffer* DynamicGeometryScratchBuffer = PassParams->DynamicGeometryScratchBuffer ? PassParams->DynamicGeometryScratchBuffer->GetRHI() : nullptr;
-				Scene->GetRayTracingDynamicGeometryCollection()->DispatchUpdates(RHIAsyncCmdList, DynamicGeometryScratchBuffer);
-
-				FRHIRayTracingScene* RayTracingSceneRHI = Scene->RayTracingScene.GetRHIRayTracingSceneChecked();
-				FRHIBuffer* AccelerationStructureBuffer = Scene->RayTracingScene.GetBufferChecked();
-				FRHIBuffer* InstanceBuffer = PassParams->RayTracingSceneInstanceBuffer->GetRHI();
-				FRHIBuffer* ScratchBuffer = PassParams->RayTracingSceneScratchBuffer->GetRHI();
-
-				RHIAsyncCmdList.BindAccelerationStructureMemory(RayTracingSceneRHI, AccelerationStructureBuffer, 0);
-
-				{
-					SCOPED_DRAW_EVENT(RHIAsyncCmdList, RayTracingScene);
-
-					FRayTracingSceneBuildParams BuildParams;
-					BuildParams.Scene = RayTracingSceneRHI;
-					BuildParams.ScratchBuffer = ScratchBuffer;
-					BuildParams.ScratchBufferOffset = 0;
-					BuildParams.InstanceBuffer = InstanceBuffer;
-					BuildParams.InstanceBufferOffset = 0;
-
-					RHIAsyncCmdList.BuildAccelerationStructure(BuildParams);
-				}
-
-				check(RayTracingDynamicGeometryUpdateEndTransition == nullptr)
-				RayTracingDynamicGeometryUpdateEndTransition = RHICreateTransition(FRHITransitionCreateInfo(ERHIPipeline::AsyncCompute, ERHIPipeline::Graphics));
-				RHIAsyncCmdList.BeginTransition(RayTracingDynamicGeometryUpdateEndTransition);
-
-				FRHIAsyncComputeCommandListImmediate::ImmediateDispatch(RHIAsyncCmdList);
-			});
-	}
-	else
-	{
+		GraphBuilder.AddPass(RDG_EVENT_NAME("RayTracingScene"), PassParams, ComputePassFlags | ERDGPassFlags::NeverCull,
+			[this, PassParams, bRayTracingAsyncBuild](FRHIComputeCommandList& RHICmdList)
 		{
-			RDG_GPU_STAT_SCOPE(GraphBuilder, RayTracingGeometry);
-			
-			FBuildAccelerationStructurePassParams* PassParams = GraphBuilder.AllocParameters<FBuildAccelerationStructurePassParams>();
-			PassParams->RayTracingSceneScratchBuffer = nullptr;
-			PassParams->DynamicGeometryScratchBuffer = OutDynamicGeometryScratchBuffer;
-			GraphBuilder.AddPass(RDG_EVENT_NAME("RayTracingGeometry"), PassParams, ERDGPassFlags::Compute | ERDGPassFlags::NeverCull,
-				[this, PassParams](FRHICommandListImmediate& RHICmdList)
-			{
-				FRHIBuffer* DynamicGeometryScratchBuffer = PassParams->DynamicGeometryScratchBuffer ? PassParams->DynamicGeometryScratchBuffer->GetRHI() : nullptr;
-				Scene->GetRayTracingDynamicGeometryCollection()->DispatchUpdates(RHICmdList, DynamicGeometryScratchBuffer);
-			});
-		}
+			FRHIBuffer* DynamicGeometryScratchBuffer = PassParams->DynamicGeometryScratchBuffer ? PassParams->DynamicGeometryScratchBuffer->GetRHI() : nullptr;
+			Scene->GetRayTracingDynamicGeometryCollection()->DispatchUpdates(RHICmdList, DynamicGeometryScratchBuffer);
 
-		{
-			RDG_GPU_STAT_SCOPE(GraphBuilder, RayTracingScene);
+			FRHIRayTracingScene* RayTracingSceneRHI = Scene->RayTracingScene.GetRHIRayTracingSceneChecked();
+			FRHIBuffer* AccelerationStructureBuffer = Scene->RayTracingScene.GetBufferChecked();
+			FRHIBuffer* ScratchBuffer = PassParams->RayTracingSceneScratchBuffer->GetRHI();
+			FRHIBuffer* InstanceBuffer = PassParams->RayTracingSceneInstanceBuffer->GetRHI();
 
-			FBuildAccelerationStructurePassParams* PassParams = GraphBuilder.AllocParameters<FBuildAccelerationStructurePassParams>();
-			PassParams->RayTracingSceneScratchBuffer = Scene->RayTracingScene.BuildScratchBuffer;
-			PassParams->RayTracingSceneInstanceBuffer = Scene->RayTracingScene.InstanceBuffer;
-			PassParams->DynamicGeometryScratchBuffer = nullptr;
+			FRayTracingSceneBuildParams BuildParams;
+			BuildParams.Scene = RayTracingSceneRHI;
+			BuildParams.ScratchBuffer = ScratchBuffer;
+			BuildParams.ScratchBufferOffset = 0;
+			BuildParams.InstanceBuffer = InstanceBuffer;
+			BuildParams.InstanceBufferOffset = 0;
 
-			GraphBuilder.AddPass(RDG_EVENT_NAME("RayTracingScene"), PassParams, ERDGPassFlags::Compute | ERDGPassFlags::NeverCull,
-				[this, PassParams](FRHICommandListImmediate& RHICmdList)
+			// Sanity check acceleration structure buffer sizes
+		#if DO_CHECK
 			{
-				FRHIRayTracingScene* RayTracingSceneRHI = Scene->RayTracingScene.GetRHIRayTracingSceneChecked();
-				FRHIBuffer* AccelerationStructureBuffer = Scene->RayTracingScene.GetBufferChecked();
-				FRHIBuffer* ScratchBuffer = PassParams->RayTracingSceneScratchBuffer->GetRHI();
-				FRHIBuffer* InstanceBuffer = PassParams->RayTracingSceneInstanceBuffer->GetRHI();
-
-				FRayTracingSceneBuildParams BuildParams;
-				BuildParams.Scene = RayTracingSceneRHI;
-				BuildParams.ScratchBuffer = ScratchBuffer;
-				BuildParams.ScratchBufferOffset = 0;
-				BuildParams.InstanceBuffer = InstanceBuffer;
-				BuildParams.InstanceBufferOffset = 0;
-
-				// Sanity check acceleration structure buffer sizes
-			#if DO_CHECK
-				{
-					FRayTracingAccelerationStructureSize SizeInfo = RHICalcRayTracingSceneSize(
-						RayTracingSceneRHI->GetInitializer().NumNativeInstances, ERayTracingAccelerationStructureFlags::FastTrace);
+				FRayTracingAccelerationStructureSize SizeInfo = RHICalcRayTracingSceneSize(
+					RayTracingSceneRHI->GetInitializer().NumNativeInstances, ERayTracingAccelerationStructureFlags::FastTrace);
 
-					check(SizeInfo.ResultSize <= Scene->RayTracingScene.SizeInfo.ResultSize);
-					check(SizeInfo.BuildScratchSize <= Scene->RayTracingScene.SizeInfo.BuildScratchSize);
-					check(SizeInfo.ResultSize <= AccelerationStructureBuffer->GetSize());
-					check(SizeInfo.BuildScratchSize <= ScratchBuffer->GetSize());
-				}
-			#endif // DO_CHECK
+				check(SizeInfo.ResultSize <= Scene->RayTracingScene.SizeInfo.ResultSize);
+				check(SizeInfo.BuildScratchSize <= Scene->RayTracingScene.SizeInfo.BuildScratchSize);
+				check(SizeInfo.ResultSize <= AccelerationStructureBuffer->GetSize());
+				check(SizeInfo.BuildScratchSize <= ScratchBuffer->GetSize());
+			}
+		#endif // DO_CHECK
 
-				RHICmdList.BindAccelerationStructureMemory(RayTracingSceneRHI, AccelerationStructureBuffer, 0);
-				RHICmdList.BuildAccelerationStructure(BuildParams);
+			RHICmdList.BindAccelerationStructureMemory(RayTracingSceneRHI, AccelerationStructureBuffer, 0);
+			RHICmdList.BuildAccelerationStructure(BuildParams);
 
+			if (!bRayTracingAsyncBuild)
+			{
 				// Submit potentially expensive BVH build commands to the GPU as soon as possible.
 				// Avoids a GPU bubble in some CPU-limited cases.
 				RHICmdList.SubmitCommandsHint();
-			});
-		}
+			}
+		});
 	}
 
 	AddPass(GraphBuilder, RDG_EVENT_NAME("EndUpdate"), [this](FRHICommandListImmediate& RHICmdList)
