Store TranslatedLightPosition in FRTLightingData
- Not using TilePosition+RelativeWorldPosition to avoid increasing struct size.
- FRTLightingData arrays seem to be calculated every frame so shouldn't be a problem.

diff --git a/Engine/Shaders/Private/RayTracing/RayTracingLightingCommon.ush b/Engine/Shaders/Private/RayTracing/RayTracingLightingCommon.ush
--- a/Engine/Shaders/Private/RayTracing/RayTracingLightingCommon.ush
+++ b/Engine/Shaders/Private/RayTracing/RayTracingLightingCommon.ush
@@ -33,7 +33,7 @@ struct FRTLightingData
 	int LightProfileIndex;
 	uint RectLightTextureIndex;
 	uint Pad;
-	float3 LightPosition;
+	float3 TranslatedLightPosition;
 	float InvRadius;
 	float3 Direction;
 	float FalloffExponent;
@@ -73,7 +73,7 @@ FRTLightingData GetRayTracingLightData(int LightIndex)
 
 	Element = LightDataBuffer[LightIndex++];
 
-	Data.LightPosition = asfloat(Element.xyz);
+	Data.TranslatedLightPosition = asfloat(Element.xyz);
 	Data.InvRadius = asfloat(Element.w);
 
 	Element = LightDataBuffer[LightIndex++];
@@ -119,7 +119,8 @@ FDeferredLightData GetRayTracingDeferredLightData(int LightIndex,
 
 	const uint LightType = RayTracingLightData.Type;
 
-	LightData.TranslatedWorldPosition = RayTracingLightData.LightPosition;
+	// We are actually storing WorldPosition at the moment (but FDeferredLightData struct is shared with other systems)
+	LightData.TranslatedWorldPosition = RayTracingLightData.TranslatedLightPosition - LWCHackToFloat(PrimaryView.PreViewTranslation); // RT_LWC_TODO
 	LightData.InvRadius = RayTracingLightData.InvRadius;
 	LightData.Color = RayTracingLightData.LightColor;
 	LightData.FalloffExponent = RayTracingLightData.FalloffExponent;
diff --git a/Engine/Source/Runtime/Renderer/Private/RayTracing/RayTracingLighting.cpp b/Engine/Source/Runtime/Renderer/Private/RayTracing/RayTracingLighting.cpp
--- a/Engine/Source/Runtime/Renderer/Private/RayTracing/RayTracingLighting.cpp
+++ b/Engine/Source/Runtime/Renderer/Private/RayTracing/RayTracingLighting.cpp
@@ -272,14 +272,10 @@ static void SetupRaytracingLightDataPacked(
 		LightDataElement.LightProfileIndex = IESLightProfileIndex;
 		LightDataElement.RectLightTextureIndex = InvalidTextureIndex;
 
-		const FVector3f LightColor(LightParameters.Color);
-		for (int32 Element = 0; Element < 3; Element++)
-		{
-			LightDataElement.Direction[Element] = LightParameters.Direction[Element];
-			LightDataElement.LightPosition[Element] = (float)LightParameters.WorldPosition[Element]; // LWC_TODO
-			LightDataElement.LightColor[Element] = LightColor[Element];
-			LightDataElement.Tangent[Element] = LightParameters.Tangent[Element];
-		}
+		LightDataElement.Direction = LightParameters.Direction;
+		LightDataElement.TranslatedLightPosition = LightParameters.WorldPosition + View.ViewMatrices.GetPreViewTranslation();
+		LightDataElement.LightColor = FVector3f(LightParameters.Color);
+		LightDataElement.Tangent = LightParameters.Tangent;
 
 		// Ray tracing should compute fade parameters ignoring lightmaps
 		const FVector2D FadeParams = Light.LightSceneInfo->Proxy->GetDirectionalLightDistanceFadeParameters(View.GetFeatureLevel(), false, View.MaxShadowCascades);
diff --git a/Engine/Source/Runtime/Renderer/Private/RayTracing/RayTracingLighting.h b/Engine/Source/Runtime/Renderer/Private/RayTracing/RayTracingLighting.h
--- a/Engine/Source/Runtime/Renderer/Private/RayTracing/RayTracingLighting.h
+++ b/Engine/Source/Runtime/Renderer/Private/RayTracing/RayTracingLighting.h
@@ -43,13 +43,13 @@ struct FRTLightingData
 	// Force alignment before next vector
 	int32 Pad;
 
-	float LightPosition[3];
+	FVector3f TranslatedLightPosition;
 	float InvRadius;
-	float Direction[3];
+	FVector3f Direction;
 	float FalloffExponent;
-	float LightColor[3];
+	FVector3f LightColor;
 	float SpecularScale;
-	float Tangent[3];
+	FVector3f Tangent;
 	float SourceRadius;
 	float SpotAngles[2];
 	float SourceLength;
