Updated FPathTracingLight to use translated world space
- Store Position and Bounds in translated world space
- Not using TilePosition+RelativeWorldPosition to avoid increasing struct size.
- FPathTracingLight arrays seem to be calculated every frame so shouldn't be a problem.
- LightGrid fully built in translated world space.
- TraceLight/SampleLight/EstimateLight/InitLightPickingCdf now expect rays and positions in translated world space.

diff --git a/Engine/Plugins/Experimental/GPULightmass/Shaders/Private/LightmapPathTracing.usf b/Engine/Plugins/Experimental/GPULightmass/Shaders/Private/LightmapPathTracing.usf
--- a/Engine/Plugins/Experimental/GPULightmass/Shaders/Private/LightmapPathTracing.usf
+++ b/Engine/Plugins/Experimental/GPULightmass/Shaders/Private/LightmapPathTracing.usf
@@ -195,6 +195,7 @@ FPathTracingPayload TraceTransparentRay(RayDesc Ray, FRayCone RayCone, bool IsCa
 		{
 			RayDesc LightRay = Ray;
 			LightRay.TMax = PackedPayload.IsMiss() ? Ray.TMax : PackedPayload.HitT;
+			LightRay.Origin += LWCHackToFloat(PrimaryView.PreViewTranslation); // RT_LWC_TODO
 			float3 LightRadiance = TraceLight(LightRay, LightId).Radiance;
 			AccumulateRadiance(Radiance, PathThroughput * LightRadiance);
 		}
@@ -485,7 +486,7 @@ void PathTracingKernel(
 			AdjustMaterialRoughness(Payload, PathRoughness);
 		}
 		
-		FLightLoopCount LightLoopCount = LightGridLookup(Payload.WorldPos);
+		FLightLoopCount LightLoopCount = LightGridLookup(Payload.WorldPos + LWCHackToFloat(PrimaryView.PreViewTranslation)); // RT_LWC_TODO
 
 		// Choose a random number for both Light sampling and BxDF sampling
 		float4 RandSample = RandomSequence_GenerateSample4D(RandSequence);
@@ -501,7 +502,8 @@ void PathTracingKernel(
 		if (MISMode != 0 && bIsNeeValid && SceneLightCount > 0 && bDoLightLoop)
 		{
 			// Choose a light and sample it
-			float3 WorldPos = Payload.WorldPos;
+			float3 WorldPos = Payload.WorldPos; // RT_LWC_TODO
+			float3 TranslatedWorldPos = Payload.WorldPos + LWCHackToFloat(PrimaryView.PreViewTranslation); // RT_LWC_TODO
 			float3 WorldNormal = Payload.WorldNormal;
 			uint PrimitiveLightingChannelMask = Payload.PrimitiveLightingChannelMask;
 
@@ -510,7 +512,7 @@ void PathTracingKernel(
 			for (uint Index = 0, Num = LightLoopCount.NumLights; Index < Num; ++Index)
 			{
 				uint LightIndex = GetLightId(Index, LightLoopCount);
-				LightPickingCdfSum += EstimateLight(LightIndex, WorldPos, WorldNormal, PrimitiveLightingChannelMask, IsTransmissiveMaterial);
+				LightPickingCdfSum += EstimateLight(LightIndex, TranslatedWorldPos, WorldNormal, PrimitiveLightingChannelMask, IsTransmissiveMaterial);
 				LightPickingCdf[Index] = LightPickingCdfSum;
 			}
 
@@ -522,7 +524,7 @@ void PathTracingKernel(
 
 				SelectLight(RandSample.x * LightPickingCdfSum, LightLoopCount.NumLights, LightPickingCdf, LightId, LightPickPdf);
 
-				FLightSample LightSample = SampleLight(LightId, RandSample.yz, WorldPos, WorldNormal);
+				FLightSample LightSample = SampleLight(LightId, RandSample.yz, TranslatedWorldPos, WorldNormal);
 
 				LightPickPdf /= LightPickingCdfSum;
 
@@ -667,7 +669,9 @@ void PathTracingKernel(
 					continue;
 				}
 
-				FLightHit LightResult = TraceLight(Ray, LightId);
+				RayDesc TranslatedRay = Ray;
+				TranslatedRay.Origin += LWCHackToFloat(PrimaryView.PreViewTranslation); // RT_LWC_TODO
+				FLightHit LightResult = TraceLight(TranslatedRay, LightId);
 
 				if (LightResult.IsMiss())
 				{
diff --git a/Engine/Plugins/Experimental/GPULightmass/Source/GPULightmass/Private/PathTracingLightParameters.inl b/Engine/Plugins/Experimental/GPULightmass/Source/GPULightmass/Private/PathTracingLightParameters.inl
--- a/Engine/Plugins/Experimental/GPULightmass/Source/GPULightmass/Private/PathTracingLightParameters.inl
+++ b/Engine/Plugins/Experimental/GPULightmass/Source/GPULightmass/Private/PathTracingLightParameters.inl
@@ -73,7 +73,7 @@ void SetupPathTracingLightParameters(
 	{
 		FPathTracingLight& DestLight = Lights.AddDefaulted_GetRef();
 
-		DestLight.Position = Light.Position;
+		DestLight.TranslatedWorldPosition = Light.Position + View.ViewMatrices.GetPreViewTranslation();
 		DestLight.Color = FVector(Light.Color);
 		DestLight.Normal = Light.Direction;
 		DestLight.dPdu = FVector::CrossProduct(Light.Tangent, Light.Direction);
@@ -100,17 +100,17 @@ void SetupPathTracingLightParameters(
 		DestLight.Flags |= PATHTRACING_LIGHT_POINT;
 
 		float Radius = Light.AttenuationRadius;
-		FVector Center = DestLight.Position;
+		FVector3f Center = DestLight.TranslatedWorldPosition;
 		// simple sphere of influence
-		DestLight.BoundMin = Center - FVector(Radius, Radius, Radius);
-		DestLight.BoundMax = Center + FVector(Radius, Radius, Radius);
+		DestLight.TranslatedBoundMin = Center - FVector3f(Radius, Radius, Radius);
+		DestLight.TranslatedBoundMax = Center + FVector3f(Radius, Radius, Radius);
 	}
 
 	for (auto Light : LightScene.SpotLights.Elements)
 	{
 		FPathTracingLight& DestLight = Lights.AddDefaulted_GetRef();
 
-		DestLight.Position = Light.Position;
+		DestLight.TranslatedWorldPosition = Light.Position + View.ViewMatrices.GetPreViewTranslation();
 		DestLight.Normal = Light.Direction;
 		DestLight.dPdu = FVector::CrossProduct(Light.Tangent, Light.Direction);
 		DestLight.dPdv = Light.Tangent;
@@ -137,31 +137,31 @@ void SetupPathTracingLightParameters(
 		DestLight.Flags |= PATHTRACING_LIGHT_SPOT;
 
 		float Radius = Light.AttenuationRadius;
-		FVector Center = DestLight.Position;
-		FVector Normal = DestLight.Normal;
-		FVector Disc = FVector(
+		FVector3f Center = DestLight.TranslatedWorldPosition;
+		FVector3f Normal = DestLight.Normal;
+		FVector3f Disc = FVector3f(
 			FMath::Sqrt(FMath::Clamp(1 - Normal.X * Normal.X, 0.0f, 1.0f)),
 			FMath::Sqrt(FMath::Clamp(1 - Normal.Y * Normal.Y, 0.0f, 1.0f)),
 			FMath::Sqrt(FMath::Clamp(1 - Normal.Z * Normal.Z, 0.0f, 1.0f))
 		);
 		// box around ray from light center to tip of the cone
-		FVector Tip = Center + Normal * Radius;
-		DestLight.BoundMin = Center.ComponentMin(Tip);
-		DestLight.BoundMax = Center.ComponentMax(Tip);
+		FVector3f Tip = Center + Normal * Radius;
+		DestLight.TranslatedBoundMin = Center.ComponentMin(Tip);
+		DestLight.TranslatedBoundMax = Center.ComponentMax(Tip);
 		// expand by disc around the farthest part of the cone
 
 		float CosOuter = Light.SpotAngles.X;
 		float SinOuter = FMath::Sqrt(1.0f - CosOuter * CosOuter);
 
-		DestLight.BoundMin = DestLight.BoundMin.ComponentMin(Center + Radius * (Normal * CosOuter - Disc * SinOuter));
-		DestLight.BoundMax = DestLight.BoundMax.ComponentMax(Center + Radius * (Normal * CosOuter + Disc * SinOuter));
+		DestLight.TranslatedBoundMin = DestLight.TranslatedBoundMin.ComponentMin(Center + Radius * (Normal * CosOuter - Disc * SinOuter));
+		DestLight.TranslatedBoundMax = DestLight.TranslatedBoundMax.ComponentMax(Center + Radius * (Normal * CosOuter + Disc * SinOuter));
 	}
 
 	for (auto Light : LightScene.RectLights.Elements)
 	{
 		FPathTracingLight& DestLight = Lights.AddDefaulted_GetRef();
 
-		DestLight.Position = Light.Position;
+		DestLight.TranslatedWorldPosition = Light.Position + View.ViewMatrices.GetPreViewTranslation();
 		DestLight.Normal = Light.Direction;
 		DestLight.dPdu = FVector::CrossProduct(Light.Tangent, -Light.Direction);
 		DestLight.dPdv = Light.Tangent;
@@ -190,18 +190,18 @@ void SetupPathTracingLightParameters(
 		DestLight.Flags |= PATHTRACING_LIGHT_RECT;
 
 		float Radius = Light.AttenuationRadius;
-		FVector Center = DestLight.Position;
-		FVector Normal = DestLight.Normal;
-		FVector Disc = FVector(
+		FVector3f Center = DestLight.TranslatedWorldPosition;
+		FVector3f Normal = DestLight.Normal;
+		FVector3f Disc = FVector3f(
 			FMath::Sqrt(FMath::Clamp(1 - Normal.X * Normal.X, 0.0f, 1.0f)),
 			FMath::Sqrt(FMath::Clamp(1 - Normal.Y * Normal.Y, 0.0f, 1.0f)),
 			FMath::Sqrt(FMath::Clamp(1 - Normal.Z * Normal.Z, 0.0f, 1.0f))
 		);
 		// quad bbox is the bbox of the disc +  the tip of the hemisphere
 		// TODO: is it worth trying to account for barndoors? seems unlikely to cut much empty space since the volume _inside_ the barndoor receives light
-		FVector Tip = Center + Normal * Radius;
-		DestLight.BoundMin = Tip.ComponentMin(Center - Radius * Disc);
-		DestLight.BoundMax = Tip.ComponentMax(Center + Radius * Disc);
+		FVector3f Tip = Center + Normal * Radius;
+		DestLight.TranslatedBoundMin = Tip.ComponentMin(Center - Radius * Disc);
+		DestLight.TranslatedBoundMax = Tip.ComponentMax(Center + Radius * Disc);
 	}
 
 	PassParameters->SceneLightCount = Lights.Num();
diff --git a/Engine/Plugins/Experimental/GPULightmass/Source/GPULightmass/Private/VolumetricLightmap.cpp b/Engine/Plugins/Experimental/GPULightmass/Source/GPULightmass/Private/VolumetricLightmap.cpp
--- a/Engine/Plugins/Experimental/GPULightmass/Source/GPULightmass/Private/VolumetricLightmap.cpp
+++ b/Engine/Plugins/Experimental/GPULightmass/Source/GPULightmass/Private/VolumetricLightmap.cpp
@@ -75,8 +75,6 @@ void ReleaseBrickData(FVolumetricLightmapBrickData& BrickData)
 	BrickData.DirectionalLightShadowing.UAV.SafeRelease();
 }
 
-void SetupPathTracingLightParameters(const GPULightmass::FLightSceneRenderState& LightScene, FRDGBuilder& GraphBuilder, const FViewInfo& View, FRDGBufferSRV** OutLightBuffer, uint32* OutLightCount);
-
 namespace GPULightmass
 {
 
diff --git a/Engine/Shaders/Private/PathTracing/Light/PathTracingCapsuleLight.ush b/Engine/Shaders/Private/PathTracing/Light/PathTracingCapsuleLight.ush
--- a/Engine/Shaders/Private/PathTracing/Light/PathTracingCapsuleLight.ush
+++ b/Engine/Shaders/Private/PathTracing/Light/PathTracingCapsuleLight.ush
@@ -26,13 +26,13 @@ float GetCapsuleRadius(int LightId)
 
 FLightHit CapsuleLight_TraceLight(RayDesc Ray, int LightId)
 {
-	float3 LightPosition = GetPosition(LightId);
+	float3 TranslatedLightPosition = GetTranslatedPosition(LightId);
 	float LightRadius = GetCapsuleRadius(LightId);
 	float LightRadius2 = Pow2(LightRadius);
 	float SourceLength = GetSourceLength(LightId);
 
 #if 0 // DEBUG code: visualize the bounding rectangle against the capsule intersection test
-	FCapsuleSphericalBounds CapsuleBounds = CapsuleGetSphericalBounds(LightPosition - Ray.Origin, GetdPdv(LightId), LightRadius, SourceLength);
+	FCapsuleSphericalBounds CapsuleBounds = CapsuleGetSphericalBounds(TranslatedLightPosition - Ray.Origin, GetdPdv(LightId), LightRadius, SourceLength);
 
 	float ConeSolidAngle = CapsuleBounds.ConeSolidAngle;
 	float RectSolidAngle = CapsuleBounds.SphericalRect.SolidAngle;
@@ -42,7 +42,7 @@ FLightHit CapsuleLight_TraceLight(RayDesc Ray, int LightId)
 	{
 		// Is ray pointing inside the cone of directions?
 		float CosTheta = dot(normalize(Ray.Direction), CapsuleBounds.ConeAxis);
-		t = CosTheta >= SqrtOneMinusX(CapsuleBounds.ConeSinThetaMax2) ? length(LightPosition) : -1.0;
+		t = CosTheta >= SqrtOneMinusX(CapsuleBounds.ConeSinThetaMax2) ? length(TranslatedLightPosition) : -1.0;
 	}
 	else
 	{
@@ -72,7 +72,7 @@ FLightHit CapsuleLight_TraceLight(RayDesc Ray, int LightId)
 		float BodyArea = 0.5 * LightRadius * SourceLength;
 		float3 LightPower = GetColor(LightId) * ComputeIESAttenuation(LightId, Ray.Origin);
 		float3 LightRadiance = LightPower / (PI * (CapsArea + BodyArea));
-		float t = CapsuleIntersect(Ray.Direction, LightPosition - Ray.Origin, GetdPdv(LightId), LightRadius2, SourceLength);
+		float t = CapsuleIntersect(Ray.Direction, TranslatedLightPosition - Ray.Origin, GetdPdv(LightId), LightRadius2, SourceLength);
 		if (t > Ray.TMin && t < Ray.TMax)
 		{
 			LightRadiance *= float3(0, 1, 0);
@@ -88,7 +88,7 @@ FLightHit CapsuleLight_TraceLight(RayDesc Ray, int LightId)
 		float3 LightRadiance = LightPower / (PI * (CapsArea + BodyArea));
 		{
 			// exact
-			float t = CapsuleIntersect(Ray.Direction, LightPosition - Ray.Origin, GetdPdv(LightId), LightRadius2, SourceLength);
+			float t = CapsuleIntersect(Ray.Direction, TranslatedLightPosition - Ray.Origin, GetdPdv(LightId), LightRadius2, SourceLength);
 			if (t > Ray.TMin && t < Ray.TMax)
 			{
 				LightRadiance *= float3(1, 0, 0);
@@ -98,20 +98,20 @@ FLightHit CapsuleLight_TraceLight(RayDesc Ray, int LightId)
 	}
 #else
 	float3 Axis = GetdPdv(LightId);
-	float t = CapsuleIntersect(Ray.Direction, LightPosition - Ray.Origin, Axis, LightRadius2, SourceLength);
+	float t = CapsuleIntersect(Ray.Direction, TranslatedLightPosition - Ray.Origin, Axis, LightRadius2, SourceLength);
 	if (t > Ray.TMin && t < Ray.TMax)
 	{
 		float CapsArea = LightRadius2;
 		float BodyArea = 0.5 * LightRadius * SourceLength;
 		float3 LightPower = GetColor(LightId) * ComputeIESAttenuation(LightId, Ray.Origin);
 		float3 LightRadiance = LightPower / (PI * (CapsArea + BodyArea));
-		float3 LightDirection = LightPosition - Ray.Origin;
+		float3 LightDirection = TranslatedLightPosition - Ray.Origin;
 		float LightDistanceSquared = dot(LightDirection, LightDirection);
 		LightRadiance *= ComputeAttenuationFalloff(LightDistanceSquared, LightId);
 #if CAPSULE_SOLIDANGLE_SAMPLING
 		float Pdf = 1.0 / GetCapsuleBoundsSolidAngle(CapsuleGetSphericalBounds(LightDirection, Axis, LightRadius, SourceLength));
 #else
-		float3 Normal = CapsuleNormal(Ray.Origin + t * Ray.Direction, LightPosition, Axis * SourceLength);
+		float3 Normal = CapsuleNormal(Ray.Origin + t * Ray.Direction, TranslatedLightPosition, Axis * SourceLength);
 
 		float CosTheta = saturate(-dot(Normal, Ray.Direction));
 
@@ -127,7 +127,7 @@ FLightHit CapsuleLight_TraceLight(RayDesc Ray, int LightId)
 FLightSample CapsuleLight_SampleLight(
 	int LightId,
 	float2 RandSample,
-	float3 WorldPos,
+	float3 TranslatedWorldPos,
 	float3 WorldNormal
 )
 {
@@ -150,7 +150,7 @@ FLightSample CapsuleLight_SampleLight(
 	float3 LightPower = GetColor(LightId);
 	float3 LightRadiance = LightPower / (PI * (CapsArea + BodyArea));
 
-	float3 LightDirection = GetPosition(LightId) - WorldPos;
+	float3 LightDirection = GetTranslatedPosition(LightId) - TranslatedWorldPos;
 	FCapsuleSphericalBounds CapsuleBounds = CapsuleGetSphericalBounds(LightDirection, Axis, Radius, SourceLength);
 
 	float4 Result = SampleCapsuleBounds(CapsuleBounds, RandSample);
@@ -202,7 +202,7 @@ FLightSample CapsuleLight_SampleLight(
 	}
 	float3x3 CylinderBasis = GetTangentBasis(Axis);
 
-	float3 LightDirection = GetPosition(LightId) - WorldPos;
+	float3 LightDirection = GetTranslatedPosition(LightId) - TranslatedWorldPos;
 	float3 LocalLightDirection = mul(CylinderBasis, LightDirection); // World To Local
 
 	float CosTheta = saturate(dot(LightNormal, -normalize(LocalLightDirection + LightPoint)));
diff --git a/Engine/Shaders/Private/PathTracing/Light/PathTracingDirectionalLight.ush b/Engine/Shaders/Private/PathTracing/Light/PathTracingDirectionalLight.ush
--- a/Engine/Shaders/Private/PathTracing/Light/PathTracingDirectionalLight.ush
+++ b/Engine/Shaders/Private/PathTracing/Light/PathTracingDirectionalLight.ush
@@ -36,7 +36,6 @@ FLightHit DirectionalLight_TraceLight(RayDesc Ray, int LightId)
 FLightSample DirectionalLight_SampleLight(
 	int LightId,
 	float2 RandSample,
-	float3 WorldPos,
 	float3 WorldNormal
 )
 {
@@ -50,7 +49,6 @@ FLightSample DirectionalLight_SampleLight(
 
 float DirectionalLight_EstimateLight(
 	int LightId,
-	float3 WorldPos,
 	float3 WorldNormal,
 	bool IsTransmissiveMaterial
 )
diff --git a/Engine/Shaders/Private/PathTracing/Light/PathTracingLightCommon.ush b/Engine/Shaders/Private/PathTracing/Light/PathTracingLightCommon.ush
--- a/Engine/Shaders/Private/PathTracing/Light/PathTracingLightCommon.ush
+++ b/Engine/Shaders/Private/PathTracing/Light/PathTracingLightCommon.ush
@@ -88,9 +88,12 @@ bool IsPhysicalLight(int LightId)
 	return IsEnvironmentLight(LightId);
 }
 
-float3 GetPosition(int LightId)
+float3 GetTranslatedPosition(int LightId)
 {
-	return SceneLights[LightId].Position;
+	//const FLWCVector3 WorldPosition = MakeLWCVector3(SceneLights[LightId].TilePosition, SceneLights[LightId].RelativeWorldPosition);
+	//return LWCToFloat(LWCAdd(WorldPosition, PreViewTranslation));
+
+	return SceneLights[LightId].TranslatedWorldPosition;
 }
 
 float3 GetNormal(int LightId)
@@ -200,13 +203,13 @@ Texture2DArray<float> IESTexture;
 SamplerState          IESTextureSampler;
 #endif
 
-float ComputeIESAttenuation(int LightId, float3 WorldPosition)
+float ComputeIESAttenuation(int LightId, float3 TranslatedWorldPosition)
 {
 	float Result = 1.f;
 #if USE_IES_TERM
 	if (SceneLights[LightId].IESTextureSlice >= 0)
 	{
-		float3 LightDirection = normalize(GetPosition(LightId) - WorldPosition);
+		float3 LightDirection = normalize(GetTranslatedPosition(LightId) - TranslatedWorldPosition);
 
 		float DotProd = dot(LightDirection, GetNormal(LightId));
 		float Angle = asin(DotProd);
diff --git a/Engine/Shaders/Private/PathTracing/Light/PathTracingLightSampling.ush b/Engine/Shaders/Private/PathTracing/Light/PathTracingLightSampling.ush
--- a/Engine/Shaders/Private/PathTracing/Light/PathTracingLightSampling.ush
+++ b/Engine/Shaders/Private/PathTracing/Light/PathTracingLightSampling.ush
@@ -37,24 +37,24 @@ FLightHit TraceLight(RayDesc Ray, int LightId)
 FLightSample SampleLight(
 	int LightId,
 	float2 RandSample,
-	float3 WorldPos,
+	float3 TranslatedWorldPos,
 	float3 WorldNormal
 )
 {
 	switch (SceneLights[LightId].Flags & PATHTRACER_FLAG_TYPE_MASK)
 	{
-		case PATHTRACING_LIGHT_SKY:			return SkyLight_SampleLight(LightId, RandSample, WorldPos, WorldNormal);
-		case PATHTRACING_LIGHT_POINT:		return PointLight_SampleLight(LightId, RandSample, WorldPos, WorldNormal);
-		case PATHTRACING_LIGHT_DIRECTIONAL: return DirectionalLight_SampleLight(LightId, RandSample, WorldPos, WorldNormal);
-		case PATHTRACING_LIGHT_RECT:		return RectLight_SampleLight(LightId, RandSample, WorldPos, WorldNormal);
-		case PATHTRACING_LIGHT_SPOT:		return SpotLight_SampleLight(LightId, RandSample, WorldPos, WorldNormal);
+		case PATHTRACING_LIGHT_SKY:			return SkyLight_SampleLight(LightId, RandSample, WorldNormal);
+		case PATHTRACING_LIGHT_POINT:		return PointLight_SampleLight(LightId, RandSample, TranslatedWorldPos, WorldNormal);
+		case PATHTRACING_LIGHT_DIRECTIONAL: return DirectionalLight_SampleLight(LightId, RandSample, WorldNormal);
+		case PATHTRACING_LIGHT_RECT:		return RectLight_SampleLight(LightId, RandSample, TranslatedWorldPos, WorldNormal);
+		case PATHTRACING_LIGHT_SPOT:		return SpotLight_SampleLight(LightId, RandSample, TranslatedWorldPos, WorldNormal);
 		default:							return NullLightSample(); // unknown light type?
 	}
 }
 
 float EstimateLight(
 	int LightId,
-	float3 WorldPos,
+	float3 TranslatedWorldPos,
 	float3 WorldNormal,
 	uint PrimitiveLightingChannelMask,
 	bool IsTransmissiveMaterial
@@ -67,18 +67,18 @@ float EstimateLight(
 
 	switch (SceneLights[LightId].Flags & PATHTRACER_FLAG_TYPE_MASK)
 	{
-		case PATHTRACING_LIGHT_SKY:			return SkyLight_EstimateLight(LightId, WorldPos, WorldNormal, IsTransmissiveMaterial);
-		case PATHTRACING_LIGHT_POINT:		return PointLight_EstimateLight(LightId, WorldPos, WorldNormal, IsTransmissiveMaterial);
-		case PATHTRACING_LIGHT_DIRECTIONAL:	return DirectionalLight_EstimateLight(LightId, WorldPos, WorldNormal, IsTransmissiveMaterial);
-		case PATHTRACING_LIGHT_RECT:		return RectLight_EstimateLight(LightId, WorldPos, WorldNormal, IsTransmissiveMaterial);
-		case PATHTRACING_LIGHT_SPOT:		return SpotLight_EstimateLight(LightId, WorldPos, WorldNormal, IsTransmissiveMaterial);
+		case PATHTRACING_LIGHT_SKY:			return SkyLight_EstimateLight(LightId, WorldNormal, IsTransmissiveMaterial);
+		case PATHTRACING_LIGHT_POINT:		return PointLight_EstimateLight(LightId, TranslatedWorldPos, WorldNormal, IsTransmissiveMaterial);
+		case PATHTRACING_LIGHT_DIRECTIONAL:	return DirectionalLight_EstimateLight(LightId, WorldNormal, IsTransmissiveMaterial);
+		case PATHTRACING_LIGHT_RECT:		return RectLight_EstimateLight(LightId, TranslatedWorldPos, WorldNormal, IsTransmissiveMaterial);
+		case PATHTRACING_LIGHT_SPOT:		return SpotLight_EstimateLight(LightId, TranslatedWorldPos, WorldNormal, IsTransmissiveMaterial);
 		default:							return 0.0;
 	}
 }
 
 // #dxr_todo: passing LightPickingCdf as anything but inout causes performance to drop significantly -- why ?
 float InitLightPickingCdf(
-	float3 WorldPos,
+	float3 TranslatedWorldPos,
 	float3 WorldNormal,
 	uint PrimitiveLightingChannelMask,
 	bool IsTransmissiveMaterial,
@@ -93,7 +93,7 @@ float InitLightPickingCdf(
 	{
 		for (uint LightIndex = 0; LightIndex < SceneLightCount; ++LightIndex)
 		{
-			CdfSum += EstimateLight(LightIndex, WorldPos, WorldNormal, PrimitiveLightingChannelMask, IsTransmissiveMaterial);
+			CdfSum += EstimateLight(LightIndex, TranslatedWorldPos, WorldNormal, PrimitiveLightingChannelMask, IsTransmissiveMaterial);
 			LightPickingCdf[LightIndex] = CdfSum;
 		}
 	}
@@ -150,7 +150,7 @@ void SelectLight(float RandSample, int NumLights, inout float LightPickingCdf[RA
 
 bool SampleLight(
 	float3 RandSample,
-	float3 WorldPos,
+	float3 TranslatedWorldPos,
 	float3 WorldNormal,
 	uint PrimitiveLightingChannelMask,
 	bool IsTransmissiveMaterial,
@@ -175,7 +175,7 @@ bool SampleLight(
 
 	float LightPickingCdf[RAY_TRACING_LIGHT_COUNT_MAXIMUM];
 
-	float LightPickingCdfSum = InitLightPickingCdf(WorldPos, WorldNormal, PrimitiveLightingChannelMask, IsTransmissiveMaterial, LightPickingCdf);
+	float LightPickingCdfSum = InitLightPickingCdf(TranslatedWorldPos, WorldNormal, PrimitiveLightingChannelMask, IsTransmissiveMaterial, LightPickingCdf);
 
 	if (LightPickingCdfSum == 0.0)
 	{
@@ -189,7 +189,7 @@ bool SampleLight(
 
 	LightPickPdf /= LightPickingCdfSum;
 
-	FLightSample Result = SampleLight(OutLightId, RandSample.yz, WorldPos, WorldNormal);
+	FLightSample Result = SampleLight(OutLightId, RandSample.yz, TranslatedWorldPos, WorldNormal);
 	OutDirection = Result.Direction;
 	OutDistance = Result.Distance;
 	OutRadianceOverPdf = Result.RadianceOverPdf / LightPickPdf;
diff --git a/Engine/Shaders/Private/PathTracing/Light/PathTracingPointLight.ush b/Engine/Shaders/Private/PathTracing/Light/PathTracingPointLight.ush
--- a/Engine/Shaders/Private/PathTracing/Light/PathTracingPointLight.ush
+++ b/Engine/Shaders/Private/PathTracing/Light/PathTracingPointLight.ush
@@ -17,10 +17,10 @@ FLightHit PointLight_TraceLight(RayDesc Ray, int LightId)
 		return CapsuleLight_TraceLight(Ray, LightId);
 	}
 
-	float3 LightPosition = GetPosition(LightId);
+	float3 TranslatedLightPosition = GetTranslatedPosition(LightId);
 	float LightRadius = GetRadius(LightId);
 	float LightRadius2 = Pow2(LightRadius);
-	float3 oc = Ray.Origin - LightPosition;
+	float3 oc = Ray.Origin - TranslatedLightPosition;
 	float b = dot(oc, Ray.Direction);
 	// "Precision Improvements for Ray / Sphere Intersection" - Ray Tracing Gems (2019)
 	// https://link.springer.com/content/pdf/10.1007%2F978-1-4842-4427-2_7.pdf
@@ -53,17 +53,17 @@ FLightHit PointLight_TraceLight(RayDesc Ray, int LightId)
 FLightSample PointLight_SampleLight(
 	int LightId,
 	float2 RandSample,
-	float3 WorldPos,
+	float3 TranslatedWorldPos,
 	float3 WorldNormal
 )
 {
-	float3 LightDirection = GetPosition(LightId) - WorldPos;
+	float3 LightDirection = GetTranslatedPosition(LightId) - TranslatedWorldPos;
 	float LightDistanceSquared = dot(LightDirection, LightDirection);
 
 	FLightSample Result = NullLightSample();
 	if (GetSourceLength(LightId) > 0)
 	{
-		Result = CapsuleLight_SampleLight(LightId, RandSample, WorldPos, WorldNormal);
+		Result = CapsuleLight_SampleLight(LightId, RandSample, TranslatedWorldPos, WorldNormal);
 	}
 	else
 	{
@@ -97,19 +97,19 @@ FLightSample PointLight_SampleLight(
 
 	// NOTE: uses distance to center to keep fadeoff mask consistent for all samples
 	Result.RadianceOverPdf *= ComputeAttenuationFalloff(LightDistanceSquared, LightId);
-	Result.RadianceOverPdf *= ComputeIESAttenuation(LightId, WorldPos);
+	Result.RadianceOverPdf *= ComputeIESAttenuation(LightId, TranslatedWorldPos);
 	return Result;
 }
 
 float PointLight_EstimateLight(
 	int LightId,
-	float3 WorldPos,
+	float3 TranslatedWorldPos,
 	float3 WorldNormal,
 	bool IsTransmissiveMaterial
 )
 {
 	// Distance
-	float3 LightDirection = GetPosition(LightId) - WorldPos;
+	float3 LightDirection = GetTranslatedPosition(LightId) - TranslatedWorldPos;
 	float LightDistanceSquared = dot(LightDirection, LightDirection);
 
 	// Geometric term
diff --git a/Engine/Shaders/Private/PathTracing/Light/PathTracingRectLight.ush b/Engine/Shaders/Private/PathTracing/Light/PathTracingRectLight.ush
--- a/Engine/Shaders/Private/PathTracing/Light/PathTracingRectLight.ush
+++ b/Engine/Shaders/Private/PathTracing/Light/PathTracingRectLight.ush
@@ -91,9 +91,9 @@ float4 SampleApproxProjectionSphericalRect(float2 Rand, FSphericalRect Spherical
 
 FLightHit RectLight_TraceLight(RayDesc Ray, int LightId)
 {
-	float3 LightPosition = GetPosition(LightId);
+	float3 TranslatedLightPosition = GetTranslatedPosition(LightId);
 	float3 LightNormal = GetNormal(LightId);
-	float3 LightDirection = LightPosition - Ray.Origin;
+	float3 LightDirection = TranslatedLightPosition - Ray.Origin;
 	float DoN = dot(Ray.Direction, LightNormal);
 	float t = dot(LightDirection, LightNormal) / DoN;
 	// ray points toward the plane and intersect it?
@@ -139,11 +139,11 @@ FLightHit RectLight_TraceLight(RayDesc Ray, int LightId)
 FLightSample RectLight_SampleLight(
 	int LightId,
 	float2 RandSample,
-	float3 WorldPos,
+	float3 TranslatedWorldPos,
 	float3 WorldNormal
 )
 {
-	float3 LightPosition = GetPosition(LightId);
+	float3 TranslatedLightPosition = GetTranslatedPosition(LightId);
 	float3 LightNormal = GetNormal(LightId);
 	float3 LightdPdu = GetdPdu(LightId);
 	float3 LightdPdv = GetdPdv(LightId);
@@ -151,7 +151,7 @@ FLightSample RectLight_SampleLight(
 	float LightHeight = GetHeight(LightId);
 
 	// Define rectangle and compute solid angle
-	float3 LightDirection = LightPosition - WorldPos;
+	float3 LightDirection = TranslatedLightPosition - TranslatedWorldPos;
 	FRect Rect = GetRect(LightDirection,
 		-LightNormal,
 		LightdPdv,
@@ -174,7 +174,7 @@ FLightSample RectLight_SampleLight(
 	float3 Direction = UniformSampleSphericalRect(RandSample, SphericalRect);
 #endif
 
-	float3 RadianceOverPdf = GetColor(LightId) * SphericalRect.SolidAngle * ComputeIESAttenuation(LightId, WorldPos);
+	float3 RadianceOverPdf = GetColor(LightId) * SphericalRect.SolidAngle * ComputeIESAttenuation(LightId, TranslatedWorldPos);
 	RadianceOverPdf *= ComputeAttenuationFalloff(dot(LightDirection, LightDirection), LightId);
 #if USE_RECT_LIGHT_TEXTURES
 	float2 UV = 0.5 * ((2 * ClippedUV - 1) * Rect.Extent + Rect.Offset) / Rect.FullExtent + 0.5;
@@ -186,14 +186,14 @@ FLightSample RectLight_SampleLight(
 
 float RectLight_EstimateLight(
 	int LightId,
-	float3 WorldPos,
+	float3 TranslatedWorldPos,
 	float3 WorldNormal,
 	bool IsTransmissiveMaterial
 )
 {
 	// Distance to centroid
 	// #dxr_todo: UE-72533 Use closest point, instead
-	float3 LightDirection = GetPosition(LightId) - WorldPos;
+	float3 LightDirection = GetTranslatedPosition(LightId) - TranslatedWorldPos;
 	float LightDistanceSquared = dot(LightDirection, LightDirection);
 	float3 LightNormal = GetNormal(LightId);
 
diff --git a/Engine/Shaders/Private/PathTracing/Light/PathTracingSkyLight.ush b/Engine/Shaders/Private/PathTracing/Light/PathTracingSkyLight.ush
--- a/Engine/Shaders/Private/PathTracing/Light/PathTracingSkyLight.ush
+++ b/Engine/Shaders/Private/PathTracing/Light/PathTracingSkyLight.ush
@@ -17,7 +17,6 @@ SkyLight.usf: Light sampling functions for SkyLight implementation
 
 float SkyLight_EstimateLight(
 	int LightId,
-	float3 WorldPos,
 	float3 WorldNormal,
 	bool IsTransmissiveMaterial
 )
@@ -42,7 +41,6 @@ FLightHit SkyLight_TraceLight(RayDesc Ray, int LightId)
 FLightSample SkyLight_SampleLight(
 	int LightId,
 	float2 RandSample,
-	float3 WorldPos,
 	float3 WorldNormal
 )
 {
diff --git a/Engine/Shaders/Private/PathTracing/Light/PathTracingSpotLight.ush b/Engine/Shaders/Private/PathTracing/Light/PathTracingSpotLight.ush
--- a/Engine/Shaders/Private/PathTracing/Light/PathTracingSpotLight.ush
+++ b/Engine/Shaders/Private/PathTracing/Light/PathTracingSpotLight.ush
@@ -32,10 +32,10 @@ FLightHit SpotLight_TraceLight(RayDesc Ray, int LightId) {
 	FLightHit Result = PointLight_TraceLight(Ray, LightId);
 	if (Result.IsHit())
 	{
-		float3 LightPosition = GetPosition(LightId);
+		float3 TranslatedLightPosition = GetTranslatedPosition(LightId);
 		float3 LightNormal = GetNormal(LightId);
 		float2 CosConeAngles = GetCosConeAngles(LightId);
-		Result.Radiance *= SpotAttenuation(-normalize(Ray.Origin - LightPosition), LightNormal, CosConeAngles);
+		Result.Radiance *= SpotAttenuation(-normalize(Ray.Origin - TranslatedLightPosition), LightNormal, CosConeAngles);
 	}
 	return Result;
 }
@@ -43,34 +43,34 @@ FLightHit SpotLight_TraceLight(RayDesc Ray, int LightId) {
 FLightSample SpotLight_SampleLight(
 	int LightId,
 	float2 RandSample,
-	float3 WorldPos,
+	float3 TranslatedWorldPos,
 	float3 WorldNormal
 )
 {
-	FLightSample Result = PointLight_SampleLight(LightId, RandSample, WorldPos, WorldNormal);
+	FLightSample Result = PointLight_SampleLight(LightId, RandSample, TranslatedWorldPos, WorldNormal);
 
-	float3 LightPosition = GetPosition(LightId);
+	float3 TranslatedLightPosition = GetTranslatedPosition(LightId);
 	float3 LightNormal = GetNormal(LightId);
 	float2 CosConeAngles = GetCosConeAngles(LightId);
-	Result.RadianceOverPdf *= SpotAttenuation(-normalize(WorldPos - LightPosition), LightNormal, CosConeAngles);
+	Result.RadianceOverPdf *= SpotAttenuation(-normalize(TranslatedWorldPos - TranslatedLightPosition), LightNormal, CosConeAngles);
 
 	return Result;
 }
 
 float SpotLight_EstimateLight(
 	int LightId,
-	float3 WorldPos,
+	float3 TranslatedWorldPos,
 	float3 WorldNormal,
 	bool IsTransmissiveMaterial
 )
 {
-	float3 LightPosition = GetPosition(LightId);
+	float3 TranslatedLightPosition = GetTranslatedPosition(LightId);
 	float3 LightNormal = GetNormal(LightId);
 	float LightPower = Luminance(GetColor(LightId));
 	float Radius = GetRadius(LightId);
 	
 	// Distance
-	float3 LightDirection = GetPosition(LightId) - WorldPos;
+	float3 LightDirection = TranslatedLightPosition - TranslatedWorldPos;
 	float LightDistanceSquared = dot(LightDirection, LightDirection);
 
 	// Geometric term
diff --git a/Engine/Shaders/Private/PathTracing/PathTracingBuildLightGrid.usf b/Engine/Shaders/Private/PathTracing/PathTracingBuildLightGrid.usf
--- a/Engine/Shaders/Private/PathTracing/PathTracingBuildLightGrid.usf
+++ b/Engine/Shaders/Private/PathTracing/PathTracingBuildLightGrid.usf
@@ -12,7 +12,6 @@
 uint SceneInfiniteLightCount;
 float3 SceneLightsTranslatedBoundMin;
 float3 SceneLightsTranslatedBoundMax;
-float3 PreViewTranslation;
 RWTexture2D<uint> RWLightGrid;
 RWBuffer<uint> RWLightGridData;
 uint LightGridResolution;
@@ -46,7 +45,7 @@ void PathTracingBuildLightGridCS(uint2 DispatchThreadId : SV_DispatchThreadID)
 	for (uint LightIndex = SceneInfiniteLightCount; LightIndex < SceneLightCount; LightIndex++)
 	{
 		uint LightType = SceneLights[LightIndex].Flags & PATHTRACER_FLAG_TYPE_MASK;
-		float3 Center = GetPosition(LightIndex) + PreViewTranslation; // RT_LWC_TODO
+		float3 Center = GetTranslatedPosition(LightIndex);
 		float Radius = 1.0 / GetAttenuation(LightIndex);
 
 		float3 E = max(Lo - Center, 0) + max(Center - Hi, 0);
@@ -68,8 +67,8 @@ void PathTracingBuildLightGridCS(uint2 DispatchThreadId : SV_DispatchThreadID)
 
 		if (LightType == PATHTRACING_LIGHT_SPOT)
 		{
-			float3 LightBoundLo = SceneLights[LightIndex].BoundMin;
-			float3 LightBoundHi = SceneLights[LightIndex].BoundMax;
+			float3 LightBoundLo = SceneLights[LightIndex].TranslatedBoundMin;
+			float3 LightBoundHi = SceneLights[LightIndex].TranslatedBoundMax;
 
 			if (any(Hi < LightBoundLo) || any(Lo > LightBoundHi))
 			{
diff --git a/Engine/Shaders/Private/PathTracing/PathTracingCore.ush b/Engine/Shaders/Private/PathTracing/PathTracingCore.ush
--- a/Engine/Shaders/Private/PathTracing/PathTracingCore.ush
+++ b/Engine/Shaders/Private/PathTracing/PathTracingCore.ush
@@ -134,6 +134,7 @@ FPathTracingPayload TraceTransparentRay(inout FPathState PathState, bool bIsCame
 		{
 			RayDesc LightRay = PathState.Ray;
 			LightRay.TMax = PackedPayload.IsMiss() ? PathState.Ray.TMax : PackedPayload.HitT;
+			LightRay.Origin += LWCHackToFloat(PrimaryView.PreViewTranslation); // RT_LWC_TODO
 			float3 LightRadiance = TraceLight(LightRay, LightId).Radiance;
 			AccumulateRadiance(PathState.Radiance, PathState.PathThroughput * LightRadiance, bIsCameraRay);
 		}
@@ -767,6 +768,7 @@ bool PathTracingKernel(inout FPathState PathState, int Bounce)
 	{
 		// Choose a light and sample it
 		float3 WorldPos = Payload.WorldPos;
+		float3 TranslatedWorldPos = Payload.WorldPos + LWCHackToFloat(PrimaryView.PreViewTranslation); // RT_LWC_TODO
 		float3 WorldNormal = Payload.WorldNormal;
 		uint PrimitiveLightingChannelMask = Payload.PrimitiveLightingChannelMask;
 
@@ -775,7 +777,7 @@ bool PathTracingKernel(inout FPathState PathState, int Bounce)
 		for (uint Index = 0, Num = LightLoopCount.NumLights; Index < Num; ++Index)
 		{
 			uint LightIndex = GetLightId(Index, LightLoopCount);
-			LightPickingCdfSum += EstimateLight(LightIndex, WorldPos, WorldNormal, PrimitiveLightingChannelMask, IsTransmissiveMaterial);
+			LightPickingCdfSum += EstimateLight(LightIndex, TranslatedWorldPos, WorldNormal, PrimitiveLightingChannelMask, IsTransmissiveMaterial);
 			LightPickingCdf[Index] = LightPickingCdfSum;
 		}
 
@@ -789,7 +791,7 @@ bool PathTracingKernel(inout FPathState PathState, int Bounce)
 			LightPickPdf /= LightPickingCdfSum;
 
 			LightId = GetLightId(LightId, LightLoopCount);
-			FLightSample LightSample = SampleLight(LightId, RandSample.yz, WorldPos, WorldNormal);
+			FLightSample LightSample = SampleLight(LightId, RandSample.yz, TranslatedWorldPos, WorldNormal);
 
 			LightSample.RadianceOverPdf /= LightPickPdf;
 			LightSample.Pdf *= LightPickPdf;
@@ -925,7 +927,9 @@ bool PathTracingKernel(inout FPathState PathState, int Bounce)
 				continue;
 			}
 
-			FLightHit LightResult = TraceLight(PathState.Ray, LightId);
+			RayDesc TranslatedRay = PathState.Ray;
+			TranslatedRay.Origin += LWCHackToFloat(PrimaryView.PreViewTranslation); // RT_LWC_TODO
+			FLightHit LightResult = TraceLight(TranslatedRay, LightId);
 
 			if (LightResult.IsMiss())
 			{
diff --git a/Engine/Shaders/Private/RayTracing/RayTracingCreateGatherPointsRGS.usf b/Engine/Shaders/Private/RayTracing/RayTracingCreateGatherPointsRGS.usf
--- a/Engine/Shaders/Private/RayTracing/RayTracingCreateGatherPointsRGS.usf
+++ b/Engine/Shaders/Private/RayTracing/RayTracingCreateGatherPointsRGS.usf
@@ -388,7 +388,7 @@ RAY_TRACING_ENTRY_RAYGEN(RayTracingCreateGatherPointsRGS)
 
 				RayDesc LightRay;
 				if (SampleLight(RandSample3,
-								Position,
+								Position + LWCHackToFloat(PrimaryView.PreViewTranslation), // RT_LWC_TODO,
 								Normal,
 								Payload.PrimitiveLightingChannelMask,
 								ENABLE_TRANSMISSION,
diff --git a/Engine/Shaders/Private/RayTracing/RayTracingGlobalIlluminationRGS.usf b/Engine/Shaders/Private/RayTracing/RayTracingGlobalIlluminationRGS.usf
--- a/Engine/Shaders/Private/RayTracing/RayTracingGlobalIlluminationRGS.usf
+++ b/Engine/Shaders/Private/RayTracing/RayTracingGlobalIlluminationRGS.usf
@@ -184,6 +184,7 @@ RAY_TRACING_ENTRY_RAYGEN(GlobalIlluminationRGS)
 					uint SkyLightId = 0;
 					RayDesc LightRay = Ray;
 					LightRay.TMax = POSITIVE_INFINITY;
+					LightRay.Origin += LWCHackToFloat(PrimaryView.PreViewTranslation); // RT_LWC_TODO
 					float3 EnvironmentRadiance = SkyLight_TraceLight(Ray, SkyLightId).Radiance;
 					Irradiance += EnvironmentRadiance * RayThroughput / RayPdf;
 				}
@@ -233,7 +234,7 @@ RAY_TRACING_ENTRY_RAYGEN(GlobalIlluminationRGS)
 
 				RayDesc LightRay;
 				if (SampleLight(RandSample3,
-								Ray.Origin,
+								Ray.Origin + LWCHackToFloat(PrimaryView.PreViewTranslation), // RT_LWC_TODO,
 								Payload.WorldNormal,
 								Payload.PrimitiveLightingChannelMask,
 								ENABLE_TRANSMISSION,
diff --git a/Engine/Shaders/Shared/RayTracingTypes.h b/Engine/Shaders/Shared/RayTracingTypes.h
--- a/Engine/Shaders/Shared/RayTracingTypes.h
+++ b/Engine/Shaders/Shared/RayTracingTypes.h
@@ -12,7 +12,9 @@
 // C++ representation of a light for the path tracer
 // #dxr_todo: Unify this with FRTLightingData ?
 struct FPathTracingLight {
-	FVector3f Position;
+// 	FVector3f RelativeWorldPosition;
+// 	FVector3f TilePosition;
+	FVector3f TranslatedWorldPosition;
 	FVector3f Normal;
 	FVector3f dPdu;
 	FVector3f dPdv;
@@ -24,8 +26,8 @@ struct FPathTracingLight {
 	int32   IESTextureSlice;
 	uint32  Flags; // see defines PATHTRACER_FLAG_*
 	int32   RectLightTextureIndex;
-	FVector3f BoundMin;
-	FVector3f BoundMax;
+	FVector3f TranslatedBoundMin;
+	FVector3f TranslatedBoundMax;
 	float Padding; // keep structure aligned
 };
 
@@ -54,7 +56,9 @@ static_assert(sizeof(FPathTracingPackedPathState) == 88, "Packed Path State size
 // HLSL side of the structs above
 
 struct FPathTracingLight {
-	float3  Position;
+// 	float3  RelativeWorldPosition;
+// 	float3  TilePosition;
+	float3  TranslatedWorldPosition;
 	float3  Normal;
 	float3  dPdu;
 	float3  dPdv;
@@ -66,8 +70,8 @@ struct FPathTracingLight {
 	int     IESTextureSlice;
 	uint    Flags;
 	int	    RectLightTextureIndex;
-	float3  BoundMin;
-	float3  BoundMax;
+	float3  TranslatedBoundMin;
+	float3  TranslatedBoundMax;
 	float Padding;
 };
 
diff --git a/Engine/Source/Runtime/Renderer/Private/PathTracing.cpp b/Engine/Source/Runtime/Renderer/Private/PathTracing.cpp
--- a/Engine/Source/Runtime/Renderer/Private/PathTracing.cpp
+++ b/Engine/Source/Runtime/Renderer/Private/PathTracing.cpp
@@ -450,7 +450,6 @@ BEGIN_SHADER_PARAMETER_STRUCT(FPathTracingLightGrid, RENDERER_API)
 	SHADER_PARAMETER(uint32, SceneInfiniteLightCount)
 	SHADER_PARAMETER(FVector3f, SceneLightsTranslatedBoundMin)
 	SHADER_PARAMETER(FVector3f, SceneLightsTranslatedBoundMax)
-	SHADER_PARAMETER(FVector3f, PreViewTranslation) // RT_LWC_TODO
 	SHADER_PARAMETER_RDG_TEXTURE(Texture2D, LightGrid)
 	SHADER_PARAMETER_RDG_BUFFER_SRV(Buffer<uint>, LightGridData)
 	SHADER_PARAMETER(unsigned, LightGridResolution)
@@ -904,7 +903,6 @@ RENDERER_API void PrepareLightGrid(FRDGBuilder& GraphBuilder, const FViewInfo& V
 	LightGridParameters->SceneInfiniteLightCount = NumInfiniteLights;
 	LightGridParameters->SceneLightsTranslatedBoundMin = FVector(+Inf, +Inf, +Inf);
 	LightGridParameters->SceneLightsTranslatedBoundMax = FVector(-Inf, -Inf, -Inf);
-	LightGridParameters->PreViewTranslation = View.ViewMatrices.GetPreViewTranslation(); // RT_LWC_TODO
 	LightGridParameters->LightGrid = nullptr;
 	LightGridParameters->LightGridData = nullptr;
 
@@ -912,21 +910,15 @@ RENDERER_API void PrepareLightGrid(FRDGBuilder& GraphBuilder, const FViewInfo& V
 	// if we have some finite lights -- build a light grid
 	if (NumFiniteLights > 0)
 	{
-		FVector SceneLightsBoundMin = FVector(+Inf, +Inf, +Inf);
-		FVector SceneLightsBoundMax = FVector(-Inf, -Inf, -Inf);
-
 		// get bounding box of all finite lights
 		const FPathTracingLight* FiniteLights = Lights + NumInfiniteLights;
 		for (int Index = 0; Index < NumFiniteLights; Index++)
 		{
 			const FPathTracingLight& Light = FiniteLights[Index];
-			SceneLightsBoundMin = FVector::Min(SceneLightsBoundMin, Light.BoundMin);
-			SceneLightsBoundMax = FVector::Max(SceneLightsBoundMax, Light.BoundMax);
+			LightGridParameters->SceneLightsTranslatedBoundMin = FVector3f::Min(LightGridParameters->SceneLightsTranslatedBoundMin, Light.TranslatedBoundMin);
+			LightGridParameters->SceneLightsTranslatedBoundMax = FVector3f::Max(LightGridParameters->SceneLightsTranslatedBoundMax, Light.TranslatedBoundMax);
 		}
 
-		LightGridParameters->SceneLightsTranslatedBoundMin = SceneLightsBoundMin + View.ViewMatrices.GetPreViewTranslation();
-		LightGridParameters->SceneLightsTranslatedBoundMax = SceneLightsBoundMax + View.ViewMatrices.GetPreViewTranslation();
-
 		const uint32 Resolution = FMath::RoundUpToPowerOfTwo(CVarPathTracingLightGridResolution.GetValueOnRenderThread());
 		const uint32 MaxCount = FMath::Clamp(
 			CVarPathTracingLightGridMaxCount.GetValueOnRenderThread(),
@@ -937,7 +929,7 @@ RENDERER_API void PrepareLightGrid(FRDGBuilder& GraphBuilder, const FViewInfo& V
 		LightGridParameters->LightGridMaxCount = MaxCount;
 
 		// pick the shortest axis
-		FVector Diag = SceneLightsBoundMax - SceneLightsBoundMin;
+		FVector3f Diag = LightGridParameters->SceneLightsTranslatedBoundMax - LightGridParameters->SceneLightsTranslatedBoundMin;
 		if (Diag.X < Diag.Y && Diag.X < Diag.Z)
 		{
 			LightGridParameters->LightGridAxis = 0;
@@ -1030,8 +1022,8 @@ void SetLightParameters(FRDGBuilder& GraphBuilder, FPathTracingRG::FParameters*
 		DestLight.Flags |= PATHTRACING_LIGHT_SKY;
 		DestLight.Flags |= Scene->SkyLight->bCastShadows ? PATHTRACER_FLAG_CAST_SHADOW_MASK : 0;
 		DestLight.IESTextureSlice = -1;
-		DestLight.BoundMin = FVector(-Inf, -Inf, -Inf);
-		DestLight.BoundMax = FVector( Inf,  Inf,  Inf);
+		DestLight.TranslatedBoundMin = FVector(-Inf, -Inf, -Inf);
+		DestLight.TranslatedBoundMax = FVector( Inf,  Inf,  Inf);
 		if (Scene->SkyLight->bRealTimeCaptureEnabled || CVarPathTracingVisibleLights.GetValueOnRenderThread() == 2)
 		{
 			// When using the realtime capture system, always make the skylight visible
@@ -1073,7 +1065,7 @@ void SetLightParameters(FRDGBuilder& GraphBuilder, FPathTracingRG::FParameters*
 
 			// these mean roughly the same thing across all light types
 			DestLight.Color = FVector3f(LightParameters.Color);
-			DestLight.Position = LightParameters.WorldPosition; // LWC_TODO
+			DestLight.TranslatedWorldPosition = LightParameters.WorldPosition + View.ViewMatrices.GetPreViewTranslation();
 			DestLight.Normal = -LightParameters.Direction;
 			DestLight.dPdu = FVector::CrossProduct(LightParameters.Tangent, LightParameters.Direction);
 			DestLight.dPdv = LightParameters.Tangent;
@@ -1084,8 +1076,8 @@ void SetLightParameters(FRDGBuilder& GraphBuilder, FPathTracingRG::FParameters*
 			DestLight.Dimensions = FVector(LightParameters.SourceRadius, LightParameters.SoftSourceRadius, 0.0f);
 			DestLight.Flags |= PATHTRACING_LIGHT_DIRECTIONAL;
 
-			DestLight.BoundMin = FVector(-Inf, -Inf, -Inf);
-			DestLight.BoundMax = FVector( Inf,  Inf,  Inf);
+			DestLight.TranslatedBoundMin = FVector(-Inf, -Inf, -Inf);
+			DestLight.TranslatedBoundMax = FVector( Inf,  Inf,  Inf);
 		}
 	}
 
@@ -1138,7 +1130,7 @@ void SetLightParameters(FRDGBuilder& GraphBuilder, FPathTracingRG::FParameters*
 
 		// these mean roughly the same thing across all light types
 		DestLight.Color = FVector3f(LightParameters.Color);
-		DestLight.Position = LightParameters.WorldPosition; // LWC_TODO
+		DestLight.TranslatedWorldPosition = LightParameters.WorldPosition + View.ViewMatrices.GetPreViewTranslation();
 		DestLight.Normal = -LightParameters.Direction;
 		DestLight.dPdu = FVector::CrossProduct(LightParameters.Tangent, LightParameters.Direction);
 		DestLight.dPdv = LightParameters.Tangent;
@@ -1184,18 +1176,18 @@ void SetLightParameters(FRDGBuilder& GraphBuilder, FPathTracingRG::FParameters*
 				}
 
 				float Radius = 1.0f / LightParameters.InvRadius;
-				FVector Center = DestLight.Position;
-				FVector Normal = DestLight.Normal;
-				FVector Disc = FVector(
+				FVector3f Center = DestLight.TranslatedWorldPosition;
+				FVector3f Normal = DestLight.Normal;
+				FVector3f Disc = FVector3f(
 					FMath::Sqrt(FMath::Clamp(1 - Normal.X * Normal.X, 0.0f, 1.0f)),
 					FMath::Sqrt(FMath::Clamp(1 - Normal.Y * Normal.Y, 0.0f, 1.0f)),
 					FMath::Sqrt(FMath::Clamp(1 - Normal.Z * Normal.Z, 0.0f, 1.0f))
 				);
 				// quad bbox is the bbox of the disc +  the tip of the hemisphere
 				// TODO: is it worth trying to account for barndoors? seems unlikely to cut much empty space since the volume _inside_ the barndoor receives light
-				FVector Tip = Center + Normal * Radius;
-				DestLight.BoundMin = FVector::Min(Tip, Center - Radius * Disc);
-				DestLight.BoundMax = FVector::Max(Tip, Center + Radius * Disc);
+				FVector3f Tip = Center + Normal * Radius;
+				DestLight.TranslatedBoundMin = FVector3f::Min(Tip, Center - Radius * Disc);
+				DestLight.TranslatedBoundMax = FVector3f::Max(Tip, Center + Radius * Disc);
 				break;
 			}
 			case LightType_Spot:
@@ -1207,24 +1199,24 @@ void SetLightParameters(FRDGBuilder& GraphBuilder, FPathTracingRG::FParameters*
 				DestLight.Flags |= PATHTRACING_LIGHT_SPOT;
 
 				float Radius = 1.0f / LightParameters.InvRadius;
-				FVector Center = DestLight.Position;
-				FVector Normal = DestLight.Normal;
-				FVector Disc = FVector(
+				FVector3f Center = DestLight.TranslatedWorldPosition;
+				FVector3f Normal = DestLight.Normal;
+				FVector3f Disc = FVector3f(
 					FMath::Sqrt(FMath::Clamp(1 - Normal.X * Normal.X, 0.0f, 1.0f)),
 					FMath::Sqrt(FMath::Clamp(1 - Normal.Y * Normal.Y, 0.0f, 1.0f)),
 					FMath::Sqrt(FMath::Clamp(1 - Normal.Z * Normal.Z, 0.0f, 1.0f))
 				);
 				// box around ray from light center to tip of the cone
-				FVector Tip = Center + Normal * Radius;
-				DestLight.BoundMin = FVector::Min(Center, Tip);
-				DestLight.BoundMax = FVector::Max(Center, Tip);
+				FVector3f Tip = Center + Normal * Radius;
+				DestLight.TranslatedBoundMin = FVector3f::Min(Center, Tip);
+				DestLight.TranslatedBoundMax = FVector3f::Max(Center, Tip);
 				// expand by disc around the farthest part of the cone
 
 				float CosOuter = LightParameters.SpotAngles.X;
 				float SinOuter = FMath::Sqrt(1.0f - CosOuter * CosOuter);
 
-				DestLight.BoundMin = FVector::Min(DestLight.BoundMin, Center + Radius * (Normal * CosOuter - Disc * SinOuter));
-				DestLight.BoundMax = FVector::Max(DestLight.BoundMax, Center + Radius * (Normal * CosOuter + Disc * SinOuter));
+				DestLight.TranslatedBoundMin = FVector3f::Min(DestLight.TranslatedBoundMin, Center + Radius * (Normal * CosOuter - Disc * SinOuter));
+				DestLight.TranslatedBoundMax = FVector3f::Max(DestLight.TranslatedBoundMax, Center + Radius * (Normal * CosOuter + Disc * SinOuter));
 				break;
 			}
 			case LightType_Point:
@@ -1234,10 +1226,10 @@ void SetLightParameters(FRDGBuilder& GraphBuilder, FPathTracingRG::FParameters*
 				DestLight.Flags |= Light.LightSceneInfo->Proxy->IsInverseSquared() ? 0 : PATHTRACER_FLAG_NON_INVERSE_SQUARE_FALLOFF_MASK;
 				DestLight.Flags |= PATHTRACING_LIGHT_POINT;
 				float Radius = 1.0f / LightParameters.InvRadius;
-				FVector Center = DestLight.Position;
+				FVector3f Center = DestLight.TranslatedWorldPosition;
 				// simple sphere of influence
-				DestLight.BoundMin = Center - FVector(Radius, Radius, Radius);
-				DestLight.BoundMax = Center + FVector(Radius, Radius, Radius);
+				DestLight.TranslatedBoundMin = Center - FVector3f(Radius, Radius, Radius);
+				DestLight.TranslatedBoundMax = Center + FVector3f(Radius, Radius, Radius);
 				break;
 			}
 			default:
diff --git a/Engine/Source/Runtime/Renderer/Private/RayTracing/RayTracingGlobalIllumination.cpp b/Engine/Source/Runtime/Renderer/Private/RayTracing/RayTracingGlobalIllumination.cpp
--- a/Engine/Source/Runtime/Renderer/Private/RayTracing/RayTracingGlobalIllumination.cpp
+++ b/Engine/Source/Runtime/Renderer/Private/RayTracing/RayTracingGlobalIllumination.cpp
@@ -314,7 +314,7 @@ static void SetupLightParameters(
 		{
 			if (CVarRayTracingGlobalIlluminationRectLight.GetValueOnRenderThread() == 0) continue;
 
-			DestLight.Position = LightShaderParameters.WorldPosition; // LWC_TODO
+			DestLight.TranslatedWorldPosition = LightShaderParameters.WorldPosition + View.ViewMatrices.GetPreViewTranslation();
 			DestLight.Normal = -LightShaderParameters.Direction;
 			DestLight.dPdu = FVector::CrossProduct(LightShaderParameters.Direction, LightShaderParameters.Tangent);
 			DestLight.dPdv = LightShaderParameters.Tangent;
@@ -329,7 +329,7 @@ static void SetupLightParameters(
 		{
 			if (CVarRayTracingGlobalIlluminationPointLight.GetValueOnRenderThread() == 0) continue;
 
-			DestLight.Position = LightShaderParameters.WorldPosition; // LWC_TODO
+			DestLight.TranslatedWorldPosition = LightShaderParameters.WorldPosition + View.ViewMatrices.GetPreViewTranslation();
 			// #dxr_todo: UE-72556 define these differences from Lit..
 			DestLight.Color = FVector3f(LightShaderParameters.Color);
 			float SourceRadius = 0.0; // LightShaderParameters.SourceRadius causes too much noise for little pay off at this time
@@ -341,7 +341,7 @@ static void SetupLightParameters(
 		{
 			if (CVarRayTracingGlobalIlluminationSpotLight.GetValueOnRenderThread() == 0) continue;
 
-			DestLight.Position = LightShaderParameters.WorldPosition; // LWC_TODO
+			DestLight.TranslatedWorldPosition = LightShaderParameters.WorldPosition + View.ViewMatrices.GetPreViewTranslation();
 			DestLight.Normal = -LightShaderParameters.Direction;
 			// #dxr_todo: UE-72556 define these differences from Lit..
 			DestLight.Color = FVector3f(LightShaderParameters.Color);
