Fix invalidation of Adaptive Virtual Textures.
We need to gather all of the Allocated VTs in the AVT and invalidate the pagaes associated with their individual producers.

diff --git a/Engine/Source/Runtime/Renderer/Private/VT/AdaptiveVirtualTexture.cpp b/Engine/Source/Runtime/Renderer/Private/VT/AdaptiveVirtualTexture.cpp
--- a/Engine/Source/Runtime/Renderer/Private/VT/AdaptiveVirtualTexture.cpp
+++ b/Engine/Source/Runtime/Renderer/Private/VT/AdaptiveVirtualTexture.cpp
@@ -277,6 +277,66 @@ int32 FAdaptiveVirtualTexture::GetSpaceID() const
 	return AllocatedVirtualTextureLowMips->GetSpaceID();
 }
 
+void FAdaptiveVirtualTexture::GetProducers(FIntRect const& InTextureRegion, uint32 InMaxLevel, TArray<FProducerInfo>& OutProducerInfos)
+{
+	const uint32 NumProducers = AllocatedVirtualTextureLowMips->GetNumUniqueProducers();
+
+	OutProducerInfos.Reserve((NumAllocated + 1) * NumProducers);
+	
+	// Add producers from persistent allocated virtual texture.
+	{
+		const uint32 AdaptiveLevelBias = AllocatedVirtualTextureLowMips->GetDescription().AdaptiveLevelBias;
+		
+		// Only add to output array if we have some relevant mips under the InMaxLevel.
+		if (InMaxLevel >= AdaptiveLevelBias)
+		{
+			const int32 Divisor = 1 << AdaptiveLevelBias;
+			const FIntRect RemappedTextureRegion(
+				FIntPoint::DivideAndRoundDown(InTextureRegion.Min, Divisor), 
+				FIntPoint::DivideAndRoundUp(InTextureRegion.Max, Divisor));
+			const uint32 RemappedMaxLevel = InMaxLevel - AdaptiveLevelBias;
+
+			for (uint32 ProducerIndex = 0; ProducerIndex < NumProducers; ++ProducerIndex)
+			{
+		 		OutProducerInfos.Emplace(FProducerInfo{ AllocatedVirtualTextureLowMips->GetUniqueProducerHandle(ProducerIndex), RemappedTextureRegion, RemappedMaxLevel });
+			}
+		}
+	}
+
+	// Add producers from transient allocated virtual textures.
+	for (FAllocation const& Allocation : AllocationSlots)
+	{
+		if (Allocation.AllocatedVT != nullptr)
+		{
+			const uint32 AdaptiveLevelBias = Allocation.AllocatedVT->GetDescription().AdaptiveLevelBias;
+			if (InMaxLevel >= AdaptiveLevelBias)
+			{
+				// Get texture region in the full VT space for this allocated VT.
+				const uint32 X = Allocation.GridIndex % GridSize.X;
+				const uint32 Y = Allocation.GridIndex / GridSize.X;
+				const FIntPoint PageSize(AllocatedDesc.TileSize * AdaptiveDesc.TileCountX / GridSize.X, AllocatedDesc.TileSize * AdaptiveDesc.TileCountY / GridSize.Y);
+				const FIntPoint PageBase(PageSize.X * X, PageSize.Y * Y);
+				const FIntRect AllocationRegion(PageBase - AllocatedDesc.TileBorderSize, PageBase + PageSize + AllocatedDesc.TileBorderSize);
+
+				// Only add to output array if the texture region intersects this allocation region.
+				if (AllocationRegion.Intersect(InTextureRegion))
+				{
+					const int32 Divisor = 1 << AdaptiveLevelBias;
+					const FIntRect RemappedTextureRegion(
+							FIntPoint::DivideAndRoundDown(InTextureRegion.Min - PageBase, Divisor),
+							FIntPoint::DivideAndRoundUp(InTextureRegion.Max - PageBase, Divisor));
+					const uint32 RemappedMaxLevel = InMaxLevel - AdaptiveLevelBias;
+
+					for (uint32 ProducerIndex = 0; ProducerIndex < NumProducers; ++ProducerIndex)
+					{
+						OutProducerInfos.Emplace(FProducerInfo{ Allocation.AllocatedVT->GetUniqueProducerHandle(ProducerIndex), RemappedTextureRegion, RemappedMaxLevel });
+					}
+				}
+			}
+		}
+	}
+}
+
 /** Get hash key for the GridIndexMap. */
 static uint16 GetGridIndexHash(int32 InGridIndex)
 {
diff --git a/Engine/Source/Runtime/Renderer/Private/VT/AdaptiveVirtualTexture.h b/Engine/Source/Runtime/Renderer/Private/VT/AdaptiveVirtualTexture.h
--- a/Engine/Source/Runtime/Renderer/Private/VT/AdaptiveVirtualTexture.h
+++ b/Engine/Source/Runtime/Renderer/Private/VT/AdaptiveVirtualTexture.h
@@ -30,7 +30,7 @@ public:
 	uint32 GetPackedAllocationRequest(uint32 vAddress, uint32 vLevelPlusOne, uint32 Frame) const;
 	/** Queue a batch of allocation requests. These will be used to reallocate any virtual textures during the next call to UpdateAllocations(). */
 	void QueuePackedAllocationRequests(uint32 const* InRequests, uint32 InNumRequests, uint32 InFrame);
-	/** Quee a batch of allocation requests. This static function relays the global requests to the individual object queues. */
+	/** Queue a batch of allocation requests. This static function relays the global requests to the individual object queues. */
 	static void QueuePackedAllocationRequests(FVirtualTextureSystem* InSystem, uint32 const* InRequests, uint32 InNumRequests, uint32 InFrame);
 	/** Update any allocations based on recent requests. */
 	void UpdateAllocations(FVirtualTextureSystem* InSystem, FRHICommandListImmediate& RHICmdList, uint32 InFrame);
@@ -40,6 +40,16 @@ public:
 	virtual int32 GetSpaceID() const override;
 	//~ End IAdaptiveVirtualTexture Interface.
 
+	/** Information needed by GetProducers() for all producers for the internally allocated virtual textures. */
+	struct FProducerInfo
+	{
+		FVirtualTextureProducerHandle ProducerHandle;
+		FIntRect RemappedTextureRegion;
+		uint32 RemappedMaxLevel;
+	};
+	/** Get internal producers that touch the texture region. */
+	void GetProducers(FIntRect const& InTextureRegion, uint32 InMaxLevel, TArray<FProducerInfo>& OutProducerInfos);
+
 protected:
 	//~ Begin IAdaptiveVirtualTexture Interface.
 	virtual void Destroy(class FVirtualTextureSystem* InSystem) override;
@@ -61,7 +71,7 @@ private:
 	bool FreeLRU(FVirtualTextureSystem* InSystem, uint32 InFrame, uint32 InFrameUnusedThreshold);
 
 private:
-	/** Adaptive virtual texture descrition. */
+	/** Adaptive virtual texture description. */
 	FAdaptiveVTDescription AdaptiveDesc;
 	/** Allocated virtual texture description for the full virtual texture. Used internally to generate descriptions for the sub allocations. */
 	FAllocatedVTDescription AllocatedDesc;
diff --git a/Engine/Source/Runtime/Renderer/Private/VT/RuntimeVirtualTextureProducer.cpp b/Engine/Source/Runtime/Renderer/Private/VT/RuntimeVirtualTextureProducer.cpp
--- a/Engine/Source/Runtime/Renderer/Private/VT/RuntimeVirtualTextureProducer.cpp
+++ b/Engine/Source/Runtime/Renderer/Private/VT/RuntimeVirtualTextureProducer.cpp
@@ -43,11 +43,6 @@ void FRuntimeVirtualTextureFinalizer::InitProducer(const FVirtualTextureProducer
 		const uint32 VirtualTextureSceneIndex = RenderScene->GetRuntimeVirtualTextureSceneIndex(ProducerId);
 		RuntimeVirtualTextureMask = 1 << VirtualTextureSceneIndex;
 
-		// Store the ProducerHandle in the FRuntimeVirtualTextureSceneProxy object.
-		// This is a bit of a hack: the proxy needs to know the producer handle but can't know it on proxy creation because the producer registration is deferred to the render thread.
-		check(ProducerHandle.PackedValue != 0);
-		RenderScene->RuntimeVirtualTextures[VirtualTextureSceneIndex]->ProducerHandle = ProducerHandle;
-
 		//todo[vt]: 
 		// Add a slow render path inside RenderPage() when this check fails. 
 		// It will need to iterate the virtual textures on each primitive instead of using the RuntimeVirtualTextureMask.
diff --git a/Engine/Source/Runtime/Renderer/Private/VT/RuntimeVirtualTextureSceneProxy.cpp b/Engine/Source/Runtime/Renderer/Private/VT/RuntimeVirtualTextureSceneProxy.cpp
--- a/Engine/Source/Runtime/Renderer/Private/VT/RuntimeVirtualTextureSceneProxy.cpp
+++ b/Engine/Source/Runtime/Renderer/Private/VT/RuntimeVirtualTextureSceneProxy.cpp
@@ -88,6 +88,17 @@ FRuntimeVirtualTextureSceneProxy::FRuntimeVirtualTextureSceneProxy(URuntimeVirtu
 
 		// The Initialize() call will allocate the virtual texture by spawning work on the render thread.
 		VirtualTexture->Initialize(Producer, ProducerDesc, Transform, Bounds);
+
+		// Store the ProducerHandle and SpaceID immediately after virtual texture is initialized.
+		ENQUEUE_RENDER_COMMAND(GetProducerHandle)(
+			[this](FRHICommandList& RHICmdList)
+			{
+				if (VirtualTexture != nullptr)
+				{
+					ProducerHandle = VirtualTexture->GetProducerHandle();
+					SpaceID = VirtualTexture->GetAllocatedVirtualTexture()->GetSpaceID();
+				}
+			});
 	}
 }
 
@@ -102,6 +113,8 @@ FRuntimeVirtualTextureSceneProxy::~FRuntimeVirtualTextureSceneProxy()
 
 void FRuntimeVirtualTextureSceneProxy::Release()
 {
+	checkSlow(!IsInRenderingThread());
+
 	if (VirtualTexture != nullptr)
 	{
 		VirtualTexture->Release();
@@ -161,13 +174,13 @@ void FRuntimeVirtualTextureSceneProxy::FlushDirtyPages()
 
 			if (bCombinedFlush)
 			{
-				FVirtualTextureSystem::Get().FlushCache(ProducerHandle, CombinedDirtyRect, MaxDirtyLevel);
+				FVirtualTextureSystem::Get().FlushCache(ProducerHandle, SpaceID, CombinedDirtyRect, MaxDirtyLevel);
 			}
 			else
 			{
 				for (FIntRect Rect : DirtyRects)
 				{
-					FVirtualTextureSystem::Get().FlushCache(ProducerHandle, Rect, MaxDirtyLevel);
+					FVirtualTextureSystem::Get().FlushCache(ProducerHandle, SpaceID, Rect, MaxDirtyLevel);
 				}
 			}
 		}
diff --git a/Engine/Source/Runtime/Renderer/Private/VT/RuntimeVirtualTextureSceneProxy.h b/Engine/Source/Runtime/Renderer/Private/VT/RuntimeVirtualTextureSceneProxy.h
--- a/Engine/Source/Runtime/Renderer/Private/VT/RuntimeVirtualTextureSceneProxy.h
+++ b/Engine/Source/Runtime/Renderer/Private/VT/RuntimeVirtualTextureSceneProxy.h
@@ -33,9 +33,6 @@ public:
 	/** Unique ID for producer that this Proxy created. Used for finding this object and it's SceneIndex from the producer.  */
 	int32 ProducerId;
 
-	/** Handle for the producer that this Proxy initialized. This is only filled in by the render thread sometime after construction! */
-	FVirtualTextureProducerHandle ProducerHandle;
-
 	/** Pointer to linked URuntimeVirtualTexture. Not for dereferencing, just for pointer comparison. */
 	URuntimeVirtualTexture* VirtualTexture;
 
@@ -50,6 +47,11 @@ private:
 	/** Virtual texture size of the URuntimeVirtualTexture object. */
 	FIntPoint VirtualTextureSize;
 
+	/** Handle for the producer that this Proxy initialized. Used only for invalidation logic. */
+	FVirtualTextureProducerHandle ProducerHandle;
+	/** Space ID used by the virtual texture. Used only for invalidation logic. */
+	int32 SpaceID;
+
 	/** Maximum mip level to mark dirty. Can be less than the virtual texture's MaxLevel if we have streaming mips. */
 	int32 MaxDirtyLevel;
 	/** Array of dirty rectangles to process at the next flush. */
diff --git a/Engine/Source/Runtime/Renderer/Private/VT/TexturePagePool.cpp b/Engine/Source/Runtime/Renderer/Private/VT/TexturePagePool.cpp
--- a/Engine/Source/Runtime/Renderer/Private/VT/TexturePagePool.cpp
+++ b/Engine/Source/Runtime/Renderer/Private/VT/TexturePagePool.cpp
@@ -124,12 +124,15 @@ void FTexturePagePool::EvictPages(FVirtualTextureSystem* System, const FVirtualT
 
 void FTexturePagePool::EvictPages(FVirtualTextureSystem* System, FVirtualTextureProducerHandle const& ProducerHandle, FVTProducerDescription const& Desc, FIntRect const& TextureRegion, uint32 MaxLevel, TArray<union FVirtualTextureLocalTile>& OutLocked)
 {
-	for (uint32 i = NumReservedPages; i < NumPages; ++i)
+	TArray<uint32, TInlineAllocator<256>> ToEvict;
+	const uint32 Hash = MurmurFinalize32(ProducerHandle.PackedValue);
+	for (uint32 pAddress = ProducerToPageIndex.First(Hash); ProducerToPageIndex.IsValid(pAddress); pAddress = ProducerToPageIndex.Next(pAddress))
 	{
-		if (Pages[i].PackedProducerHandle == ProducerHandle.PackedValue)
+		const FPageEntry& PageEntry = Pages[pAddress];
+		if (PageEntry.PackedProducerHandle == ProducerHandle.PackedValue)
 		{
-			const uint32 vAddress = Pages[i].Local_vAddress;
-			const uint32 vLevel = Pages[i].Local_vLevel;
+			const uint32 vAddress = Pages[pAddress].Local_vAddress;
+			const uint32 vLevel = Pages[pAddress].Local_vLevel;
 
 			if (vLevel <= MaxLevel)
 			{
@@ -141,19 +144,24 @@ void FTexturePagePool::EvictPages(FVirtualTextureSystem* System, FVirtualTexture
 
 				if (!(PageRect.Min.X > TextureRegion.Max.X) && !(TextureRegion.Min.X > PageRect.Max.X) && !(PageRect.Min.Y > TextureRegion.Max.Y) && !(TextureRegion.Min.Y > PageRect.Max.Y))
 				{
-					if (!FreeHeap.IsPresent(i))
+					if (!FreeHeap.IsPresent(pAddress))
 					{
 						OutLocked.Add(FVirtualTextureLocalTile(ProducerHandle, vAddress, vLevel));
 					}
 					else
 					{
-						UnmapAllPages(System, i, true);
-						FreeHeap.Update(0, i);
+						ToEvict.Add(pAddress);
 					}
 				}
 			}
 		}
 	}
+
+	for (uint32 pAddress : ToEvict)
+	{
+		UnmapAllPages(System, pAddress, false);
+		FreeHeap.Update(0, pAddress);
+	}
 }
 
 void FTexturePagePool::GetAllLockedPages(FVirtualTextureSystem* System, TSet<FVirtualTextureLocalTile>& OutPages)
@@ -433,58 +441,68 @@ void FTexturePagePool::RemapPages(FVirtualTextureSystem* System, uint8 SpaceID,
 	const uint32 NewBaseX = FMath::ReverseMortonCode2(NewVirtualAddress);
 	const uint32 NewBaseY = FMath::ReverseMortonCode2(NewVirtualAddress >> 1);
 
-	for (uint32 pAddress = NumReservedPages; pAddress < NumPages; ++pAddress)
+	TArray<uint32, TInlineAllocator<256>> ToRemap;
+	const uint32 Hash = MurmurFinalize32(OldProducerHandle.PackedValue);
+	for (uint32 pAddress = ProducerToPageIndex.First(Hash); ProducerToPageIndex.IsValid(pAddress); pAddress = ProducerToPageIndex.Next(pAddress))
 	{
-		FPageEntry& PageEntry = Pages[pAddress];
+		const FPageEntry& PageEntry = Pages[pAddress];
 		if (PageEntry.PackedProducerHandle == OldProducerHandle.PackedValue)
 		{
-			if ((int32)PageEntry.Local_vLevel + vLevelBias < 0)
-			{
-				// Remap removes this level 
-				UnmapAllPages(System, pAddress, false);
-				// Queue page for recycling
-				FreeHeap.Update(0, pAddress);
-			}
-			else
-			{
-				// Directly modify page entry for new producer
-				PageHash.Remove(GetPageHash(PageEntry), pAddress);
-				PageEntry.PackedProducerHandle = NewProducerHandle.PackedValue;
-				PageEntry.Local_vLevel += vLevelBias;
-				PageHash.Add(GetPageHash(PageEntry), pAddress);
+			ToRemap.Add(pAddress);
+		}
+	}
 
-				ProducerToPageIndex.Remove(OldProducerHash, pAddress);
-				ProducerToPageIndex.Add(NewProducerHash, pAddress);
+	for (uint32 pAddress : ToRemap)
+	{
+		FPageEntry& PageEntry = Pages[pAddress];
+		check(PageEntry.PackedProducerHandle == OldProducerHandle.PackedValue);
 
-				if (FreeHeap.IsPresent(pAddress))
-				{
-					FreeHeap.Update((Frame << 4) + PageEntry.Local_vLevel, pAddress);
-				}
+		if ((int32)PageEntry.Local_vLevel + vLevelBias < 0)
+		{
+			// Remap removes this level 
+			UnmapAllPages(System, pAddress, false);
+			// Queue page for recycling
+			FreeHeap.Update(0, pAddress);
+		}
+		else
+		{
+			// Directly modify page entry for new producer
+			PageHash.Remove(GetPageHash(PageEntry), pAddress);
+			PageEntry.PackedProducerHandle = NewProducerHandle.PackedValue;
+			PageEntry.Local_vLevel += vLevelBias;
+			PageHash.Add(GetPageHash(PageEntry), pAddress);
 
-				// Go through mappings and modify directly.
-				uint32 MappingIndex = PageMapping[pAddress].NextIndex;
-				while (MappingIndex != pAddress)
-				{
-					FPageMapping& Mapping = PageMapping[MappingIndex];
+			ProducerToPageIndex.Remove(OldProducerHash, pAddress);
+			ProducerToPageIndex.Add(NewProducerHash, pAddress);
+
+			if (FreeHeap.IsPresent(pAddress))
+			{
+				FreeHeap.Update((Frame << 4) + PageEntry.Local_vLevel, pAddress);
+			}
+
+			// Go through mappings and modify directly.
+			uint32 MappingIndex = PageMapping[pAddress].NextIndex;
+			while (MappingIndex != pAddress)
+			{
+				FPageMapping& Mapping = PageMapping[MappingIndex];
 
-					FVirtualTextureSpace* Space = System->GetSpace(Mapping.SpaceID);
-					FTexturePageMap& PageMap = Space->GetPageMapForPageTableLayer(Mapping.PageTableLayerIndex);
+				FVirtualTextureSpace* Space = System->GetSpace(Mapping.SpaceID);
+				FTexturePageMap& PageMap = Space->GetPageMapForPageTableLayer(Mapping.PageTableLayerIndex);
 
-					PageMap.UnmapPage(System, Space, Mapping.vLogSize, Mapping.vAddress, false);
+				PageMap.UnmapPage(System, Space, Mapping.vLogSize, Mapping.vAddress, false);
 
-					const uint32 XLocal = FMath::ReverseMortonCode2(Mapping.vAddress) - OldBaseX;
-					const uint32 YLocal = FMath::ReverseMortonCode2(Mapping.vAddress >> 1) - OldBaseY;
-					const uint32 X = NewBaseX + ((vLevelBias >= 0) ? (XLocal << vLevelBias) : (XLocal >> -vLevelBias));
-					const uint32 Y = NewBaseY + ((vLevelBias >= 0) ? (YLocal << vLevelBias) : (YLocal >> -vLevelBias));
+				const uint32 XLocal = FMath::ReverseMortonCode2(Mapping.vAddress) - OldBaseX;
+				const uint32 YLocal = FMath::ReverseMortonCode2(Mapping.vAddress >> 1) - OldBaseY;
+				const uint32 X = NewBaseX + ((vLevelBias >= 0) ? (XLocal << vLevelBias) : (XLocal >> -vLevelBias));
+				const uint32 Y = NewBaseY + ((vLevelBias >= 0) ? (YLocal << vLevelBias) : (YLocal >> -vLevelBias));
 
-					Mapping.vAddress = FMath::MortonCode2(X) | (FMath::MortonCode2(Y) << 1);
-					Mapping.vLogSize = (int32)Mapping.vLogSize + vLevelBias;
+				Mapping.vAddress = FMath::MortonCode2(X) | (FMath::MortonCode2(Y) << 1);
+				Mapping.vLogSize = (int32)Mapping.vLogSize + vLevelBias;
 
-					const int32 vLevel = PageEntry.Local_vLevel; // Deal with any producer mip bias?
-					PageMap.MapPage(Space, PhysicalSpace, NewProducerHandle.PackedValue, Mapping.MaxLevel, Mapping.vLogSize, Mapping.vAddress, vLevel, pAddress);
+				const int32 vLevel = PageEntry.Local_vLevel; // Deal with any producer mip bias?
+				PageMap.MapPage(Space, PhysicalSpace, NewProducerHandle.PackedValue, Mapping.MaxLevel, Mapping.vLogSize, Mapping.vAddress, vLevel, pAddress);
 
-					MappingIndex = Mapping.NextIndex;
-				}
+				MappingIndex = Mapping.NextIndex;
 			}
 		}
 	}
diff --git a/Engine/Source/Runtime/Renderer/Private/VT/VirtualTextureSystem.cpp b/Engine/Source/Runtime/Renderer/Private/VT/VirtualTextureSystem.cpp
--- a/Engine/Source/Runtime/Renderer/Private/VT/VirtualTextureSystem.cpp
+++ b/Engine/Source/Runtime/Renderer/Private/VT/VirtualTextureSystem.cpp
@@ -370,7 +370,7 @@ void FVirtualTextureSystem::FlushCache()
 	bFlushCaches = true;
 }
 
-void FVirtualTextureSystem::FlushCache(FVirtualTextureProducerHandle const& ProducerHandle, FIntRect const& TextureRegion, uint32 MaxLevel)
+void FVirtualTextureSystem::FlushCache(FVirtualTextureProducerHandle const& ProducerHandle, int32 SpaceID, FIntRect const& TextureRegion, uint32 MaxLevel)
 {
 	checkSlow(IsInRenderingThread());
 
@@ -390,10 +390,30 @@ void FVirtualTextureSystem::FlushCache(FVirtualTextureProducerHandle const& Prod
 
 		check(TransientCollectedPages.Num() == 0);
 
+		// If this is an Adaptive VT we need to collect all of the associated Producers to flush.
+		TArray<FAdaptiveVirtualTexture::FProducerInfo> ProducerInfos;
+		if (AdaptiveVTs[SpaceID] != nullptr)
+		{
+			AdaptiveVTs[SpaceID]->GetProducers(TextureRegion, MaxLevel, ProducerInfos);
+		}
+
 		for (int32 i = 0; i < PhysicalSpacesForProducer.Num(); ++i)
 		{
 			FTexturePagePool& Pool = PhysicalSpacesForProducer[i]->GetPagePool();
-			Pool.EvictPages(this, ProducerHandle, ProducerDescription, TextureRegion, MaxLevel, TransientCollectedPages);
+		
+			if (ProducerInfos.Num())
+			{
+				// Adaptive VT flushes.
+				for (FAdaptiveVirtualTexture::FProducerInfo& Info : ProducerInfos)
+				{
+					Pool.EvictPages(this, Info.ProducerHandle, ProducerDescription, Info.RemappedTextureRegion, Info.RemappedMaxLevel, TransientCollectedPages);
+				}
+			}
+			else
+			{
+				// Regular flush.
+				Pool.EvictPages(this, ProducerHandle, ProducerDescription, TextureRegion, MaxLevel, TransientCollectedPages);
+			}
 		}
 
 		for (auto& Page : TransientCollectedPages)
diff --git a/Engine/Source/Runtime/Renderer/Private/VT/VirtualTextureSystem.h b/Engine/Source/Runtime/Renderer/Private/VT/VirtualTextureSystem.h
--- a/Engine/Source/Runtime/Renderer/Private/VT/VirtualTextureSystem.h
+++ b/Engine/Source/Runtime/Renderer/Private/VT/VirtualTextureSystem.h
@@ -81,7 +81,7 @@ public:
 	void RequestRecordedTiles(TArray<uint64>&& InPageRequests);
 
 	void FlushCache();
-	void FlushCache(FVirtualTextureProducerHandle const& ProducerHandle, FIntRect const& TextureRegion, uint32 MaxLevel);
+	void FlushCache(FVirtualTextureProducerHandle const& ProducerHandle, int32 SpaceID, FIntRect const& TextureRegion, uint32 MaxLevel);
 
 	float GetGlobalMipBias() const;
 
