fix issue with forward rendering + MSAA on console platforms

diff --git a/Engine/Source/Runtime/Renderer/Private/SceneTextures.cpp b/Engine/Source/Runtime/Renderer/Private/SceneTextures.cpp
--- a/Engine/Source/Runtime/Renderer/Private/SceneTextures.cpp
+++ b/Engine/Source/Runtime/Renderer/Private/SceneTextures.cpp
@@ -622,7 +622,7 @@ FSceneTextures& FMinimalSceneTextures::Create(FRDGBuilder& GraphBuilder, const F
 			Flags |= TexCreate_Memoryless;
 		}
 
-		if (GRHISupportsDepthUAV)
+		if (Config.NumSamples == 1 && GRHISupportsDepthUAV)
 		{
 			Flags |= TexCreate_UAV;
 		}
