Don't skip velocity for view conditions when velocity is in depth pass.

diff --git a/Engine/Source/Runtime/Renderer/Private/VelocityRendering.cpp b/Engine/Source/Runtime/Renderer/Private/VelocityRendering.cpp
--- a/Engine/Source/Runtime/Renderer/Private/VelocityRendering.cpp
+++ b/Engine/Source/Runtime/Renderer/Private/VelocityRendering.cpp
@@ -169,6 +169,14 @@ bool FDeferredShadingSceneRenderer::ShouldRenderVelocities() const
 	{
 		return false;
 	}
+	if (FVelocityRendering::DepthPassCanOutputVelocity())
+	{
+		// Always render velocity when it is part of the depth pass to avoid dropping things from the depth pass.
+		// This means that we will pay the cost of velocity in the pass even if we don't really need it according to the view logic below.
+		// But requiring velocity is by far the most common case.
+		// And the alternative approach is for the depth pass to also incorporate the logic below to avoid dropping velocity primitives.
+		return true;
+	}
 
 	bool bNeedsVelocity = false;
 	for (int32 ViewIndex = 0; ViewIndex < Views.Num(); ViewIndex++)
