Remove unused r.CachedRayTracingInstances.CacheLocalTransform/LazyUpdate
- Only keep default codepaths (dont cache local transform, use lazy update)

diff --git a/Engine/Source/Runtime/Renderer/Private/PrimitiveSceneInfo.cpp b/Engine/Source/Runtime/Renderer/Private/PrimitiveSceneInfo.cpp
--- a/Engine/Source/Runtime/Renderer/Private/PrimitiveSceneInfo.cpp
+++ b/Engine/Source/Runtime/Renderer/Private/PrimitiveSceneInfo.cpp
@@ -28,23 +28,6 @@
 extern int32 GGPUSceneInstanceClearList;
 extern int32 GGPUSceneInstanceBVH;
 
-#if RHI_RAYTRACING
-static int32 GCachedRayTracingInstancesCacheLocalTransform = 0;
-static FAutoConsoleVariableRef CVarCachedRayTracingInstancesUseInstanceData(
-	TEXT("r.CachedRayTracingInstances.CacheLocalTransform"),
-	GCachedRayTracingInstancesCacheLocalTransform,
-	TEXT("Cache Local Transform instead of using InstanceData (increases memory usage)."),
-	ECVF_ReadOnly);
-
-static int32 GCachedRayTracingInstancesLazyUpdate = 1;
-static FAutoConsoleVariableRef CVarCachedRayTracingInstancesLazyUpdate(
-	TEXT("r.CachedRayTracingInstances.LazyUpdate"),
-	GCachedRayTracingInstancesLazyUpdate,
-	TEXT("Lazy update cached ray tracing instances world transforms. \n")
-	TEXT("Reduces memory usage by only caching world transforms of primitives when necessary."),
-	ECVF_ReadOnly);
-#endif
-
 static int32 GMeshDrawCommandsCacheMultithreaded = 1;
 static FAutoConsoleVariableRef CVarDrawCommandsCacheMultithreaded(
 	TEXT("r.MeshDrawCommands.CacheMultithreaded"),
@@ -257,7 +240,6 @@ FPrimitiveSceneInfo::~FPrimitiveSceneInfo()
 	}
 
 #if RHI_RAYTRACING
-	DEC_MEMORY_STAT_BY(STAT_CachedRayTracingInstancesMemory, CachedRayTracingInstanceLocalTransforms.Num() * sizeof(FMatrix));
 	DEC_MEMORY_STAT_BY(STAT_CachedRayTracingInstancesMemory, CachedRayTracingInstanceWorldTransforms.Num() * sizeof(FMatrix));
 #endif
 }
@@ -935,7 +917,7 @@ void FPrimitiveSceneInfo::CacheRayTracingPrimitives(FScene* Scene, const TArrayV
 
 void FPrimitiveSceneInfo::UpdateCachedRayTracingInstanceWorldTransforms()
 {
-	if (bUpdateCachedRayTracingInstanceWorldTransforms && GCachedRayTracingInstancesLazyUpdate)
+	if (bUpdateCachedRayTracingInstanceWorldTransforms)
 	{
 		QUICK_SCOPE_CYCLE_COUNTER(STAT_UpdateCachedRayTracingInstanceWorldTransforms);
 
@@ -958,7 +940,7 @@ void FPrimitiveSceneInfo::UpdateCachedRayTracingInstanceWorldTransforms()
 
 		for (int32 Index = 0; Index < CachedRayTracingInstanceWorldTransforms.Num(); Index++)
 		{
-			FMatrix LocalTransform = GCachedRayTracingInstancesCacheLocalTransform ? CachedRayTracingInstanceLocalTransforms[Index] : InstanceSceneData[Index].LocalToPrimitive.ToMatrix();
+			FMatrix LocalTransform = InstanceSceneData[Index].LocalToPrimitive.ToMatrix();
 
 			CachedRayTracingInstanceWorldTransforms[Index] = LocalTransform * LocalToWorld;
 		}
@@ -972,23 +954,8 @@ void FPrimitiveSceneInfo::UpdateCachedRayTracingInstance(FPrimitiveSceneInfo* Sc
 {
 	if (EnumHasAnyFlags(Flags, ERayTracingPrimitiveFlags::CacheInstances))
 	{
-		if (GCachedRayTracingInstancesCacheLocalTransform)
-		{
-			// Cache a copy of local transforms so that they can be updated in the future
-			// TODO: this is actually not needed for static meshes with non-movable mobility (except in editor)
-			DEC_MEMORY_STAT_BY(STAT_CachedRayTracingInstancesMemory, SceneInfo->CachedRayTracingInstanceLocalTransforms.Num() * sizeof(FMatrix));
-			SceneInfo->CachedRayTracingInstanceLocalTransforms = CachedRayTracingInstance.InstanceTransforms;
-			INC_MEMORY_STAT_BY(STAT_CachedRayTracingInstancesMemory, SceneInfo->CachedRayTracingInstanceLocalTransforms.Num() * sizeof(FMatrix));
-		}
 		// TODO: allocate from FRayTracingScene & do better low-level caching
 		SceneInfo->CachedRayTracingInstance.NumTransforms = CachedRayTracingInstance.NumTransforms;
-		if (!GCachedRayTracingInstancesLazyUpdate)
-		{
-			DEC_MEMORY_STAT_BY(STAT_CachedRayTracingInstancesMemory, SceneInfo->CachedRayTracingInstanceWorldTransforms.Num() * sizeof(FMatrix));
-			SceneInfo->CachedRayTracingInstanceWorldTransforms.Empty();
-			SceneInfo->CachedRayTracingInstanceWorldTransforms.AddUninitialized(CachedRayTracingInstance.NumTransforms);
-			INC_MEMORY_STAT_BY(STAT_CachedRayTracingInstancesMemory, SceneInfo->CachedRayTracingInstanceWorldTransforms.Num() * sizeof(FMatrix));
-		}
 
 		// Apply local offset to far-field object
 		FMatrix LocalToWorld = SceneInfo->Proxy->GetLocalToWorld();
@@ -1002,12 +969,6 @@ void FPrimitiveSceneInfo::UpdateCachedRayTracingInstance(FPrimitiveSceneInfo* Sc
 
 		SceneInfo->UpdateCachedRayTracingInstanceTransforms(LocalToWorld);
 
-		if (!GCachedRayTracingInstancesLazyUpdate)
-		{
-			SceneInfo->CachedRayTracingInstance.Transforms = MakeArrayView(SceneInfo->CachedRayTracingInstanceWorldTransforms);
-			check(SceneInfo->CachedRayTracingInstance.NumTransforms >= uint32(SceneInfo->CachedRayTracingInstance.Transforms.Num()));
-		}
-
 		SceneInfo->CachedRayTracingInstance.GeometryRHI = CachedRayTracingInstance.Geometry->RayTracingGeometryRHI;
 
 		// At this point (in AddToScene()) PrimitiveIndex has been set
@@ -1067,13 +1028,9 @@ void FPrimitiveSceneInfo::UpdateCachedRayTracingInstanceTransforms(const FMatrix
 	{
 		const FRenderBounds& LocalBoundingBox = Proxy->GetInstanceLocalBounds(Index);
 
-		FMatrix LocalTransform = GCachedRayTracingInstancesCacheLocalTransform ? CachedRayTracingInstanceLocalTransforms[Index] : InstanceSceneData[Index].LocalToPrimitive.ToMatrix();
+		FMatrix LocalTransform = InstanceSceneData[Index].LocalToPrimitive.ToMatrix();
 
 		CachedRayTracingInstanceWorldBounds[Index] = LocalBoundingBox.TransformBy(LocalTransform * NewPrimitiveLocalToWorld).ToBoxSphereBounds();
-		if(!GCachedRayTracingInstancesLazyUpdate)
-		{
-			CachedRayTracingInstanceWorldTransforms[Index] = LocalTransform * NewPrimitiveLocalToWorld;
-		}
 		SmallestRayTracingInstanceWorldBoundsIndex = CachedRayTracingInstanceWorldBounds[Index].SphereRadius < CachedRayTracingInstanceWorldBounds[SmallestRayTracingInstanceWorldBoundsIndex].SphereRadius ? Index : SmallestRayTracingInstanceWorldBoundsIndex;
 	}
 
diff --git a/Engine/Source/Runtime/Renderer/Public/PrimitiveSceneInfo.h b/Engine/Source/Runtime/Renderer/Public/PrimitiveSceneInfo.h
--- a/Engine/Source/Runtime/Renderer/Public/PrimitiveSceneInfo.h
+++ b/Engine/Source/Runtime/Renderer/Public/PrimitiveSceneInfo.h
@@ -410,7 +410,6 @@ public:
 	TArray<uint64> CachedRayTracingMeshCommandsHashPerLOD;
 	// TODO: this should be placed in FRayTracingScene and we have a pointer/handle here. It's here for now for PoC
 	FRayTracingGeometryInstance CachedRayTracingInstance;
-	TArray<FMatrix> CachedRayTracingInstanceLocalTransforms;
 	TArray<FMatrix> CachedRayTracingInstanceWorldTransforms;
 	TArray<FBoxSphereBounds> CachedRayTracingInstanceWorldBounds;
 	int32 SmallestRayTracingInstanceWorldBoundsIndex;
