Fix for Sky light average irradiance to be valid on static meshes when it is stationary.
It  now follows the same rules as SH irradiance:
- uploaded per frame if updated from a static capture
- written directily from the GPU if running real time capture

Nothing to do on mobile because it does not properly use real time capture anyway with the SkyIrradianceEnvironmentMap structured buffer.

diff --git a/Engine/Shaders/Private/ReflectionEnvironmentShaders.usf b/Engine/Shaders/Private/ReflectionEnvironmentShaders.usf
--- a/Engine/Shaders/Private/ReflectionEnvironmentShaders.usf
+++ b/Engine/Shaders/Private/ReflectionEnvironmentShaders.usf
@@ -537,7 +537,6 @@ void ComputeSkyEnvMapDiffuseIrradianceCS(uint3 ThreadId : SV_DispatchThreadID)
 {
 	const uint LinearIndex = THREADGROUP_SIZE_X * ThreadId.y + ThreadId.x;
 
-
 #if 1
 
 	// For a 128x128 cubemap, sampling mip level 2 with only 8x8 samples matches closely the super sampled version.
@@ -652,6 +651,19 @@ void ComputeSkyEnvMapDiffuseIrradianceCS(uint3 ThreadId : SV_DispatchThreadID)
 		OutIrradianceEnvMapSH[6].y = Coefficient4 * SkyIrradiance.G.V2;//[8];
 		OutIrradianceEnvMapSH[6].z = Coefficient4 * SkyIrradiance.B.V2;//[8];
 		OutIrradianceEnvMapSH[6].w = 1.0f;
+
+		// We inverse the weights applied onto the non directional SH component to recover the average sky irradiance.
+		const float3 DummyVector = float3(0.0, 0.0, 1.0);
+		const float SHBasisL0Weight = SHBasisFunction3(DummyVector).V0.x;
+		const float InvSHBand0Weights = 1.0f / (SHBasisL0Weight * 4.0f * PI);
+		const float SkyAverageIrradianceR = SkyIrradiance.R.V0.x * InvSHBand0Weights;
+		const float SkyAverageIrradianceG = SkyIrradiance.G.V0.x * InvSHBand0Weights;
+		const float SkyAverageIrradianceB = SkyIrradiance.B.V0.x * InvSHBand0Weights;
+		const float SkyAverageIrradiance  = dot(float3(SkyAverageIrradianceR, SkyAverageIrradianceG, SkyAverageIrradianceB), .3333f);
+		OutIrradianceEnvMapSH[7].x = SkyAverageIrradiance;
+		OutIrradianceEnvMapSH[7].y = SkyAverageIrradiance;
+		OutIrradianceEnvMapSH[7].z = SkyAverageIrradiance;
+		OutIrradianceEnvMapSH[7].w = SkyAverageIrradiance;
 	}
 }
 #endif
diff --git a/Engine/Shaders/Private/ReflectionEnvironmentShared.ush b/Engine/Shaders/Private/ReflectionEnvironmentShared.ush
--- a/Engine/Shaders/Private/ReflectionEnvironmentShared.ush
+++ b/Engine/Shaders/Private/ReflectionEnvironmentShared.ush
@@ -12,6 +12,11 @@
 #define SkyIrradianceEnvironmentMap		View.SkyIrradianceEnvironmentMap
 #endif
 
+float GetSkyLightCubemapBrightness()
+{
+	return SkyIrradianceEnvironmentMap[7].x; // Refer to FSceneRenderer::UpdateSkyIrradianceGpuBuffer for more details.
+}
+
 #define REFLECTION_CAPTURE_ROUGHEST_MIP 1
 #define REFLECTION_CAPTURE_ROUGHNESS_MIP_SCALE 1.2
 
@@ -40,7 +45,7 @@ float3 GetSkyLightReflection(float3 ReflectionVector, float Roughness, out float
 	float AbsoluteSpecularMip = ComputeReflectionCaptureMipFromRoughness(Roughness, ReflectionStruct.SkyLightParameters.x);
 	float3 Reflection = TextureCubeSampleLevel(ReflectionStruct.SkyLightCubemap, ReflectionStruct.SkyLightCubemapSampler, ReflectionVector, AbsoluteSpecularMip).rgb;
 
-	OutSkyAverageBrightness = ReflectionStruct.SkyLightCubemapBrightness * Luminance( View.SkyLightColor.rgb );
+	OutSkyAverageBrightness = GetSkyLightCubemapBrightness() * Luminance(View.SkyLightColor.rgb);
 	return Reflection * View.SkyLightColor.rgb;
 }
 
diff --git a/Engine/Shaders/Public/ShaderVersion.ush b/Engine/Shaders/Public/ShaderVersion.ush
--- a/Engine/Shaders/Public/ShaderVersion.ush
+++ b/Engine/Shaders/Public/ShaderVersion.ush
@@ -2,5 +2,5 @@
 // Each time the console command is executed it generates a new GUID. As this file is included
 // in Platform.ush (which should be included in any shader) it allows to invalidate the shader DDC.
 // 
-// GUID = 03177423-197C-410E-8F15-CC6C2F7304E6
+// GUID = B6B2FB8E-934E-4972-B768-7BD0AAF512F1
 
diff --git a/Engine/Source/Runtime/Engine/Public/SceneView.h b/Engine/Source/Runtime/Engine/Public/SceneView.h
--- a/Engine/Source/Runtime/Engine/Public/SceneView.h
+++ b/Engine/Source/Runtime/Engine/Public/SceneView.h
@@ -588,6 +588,8 @@ enum ETranslucencyVolumeCascade
 	TVC_MAX,
 };
 
+#define SKY_IRRADIANCE_ENVIRONMENT_MAP_VEC4_COUNT 8
+
 //VIEW_UNIFORM_BUFFER_MEMBER(FMatrix44f, PrevProjection)
 //VIEW_UNIFORM_BUFFER_MEMBER(FMatrix44f, PrevViewProj)
 //VIEW_UNIFORM_BUFFER_MEMBER(FMatrix44f, PrevViewRotationProj)
@@ -725,7 +727,7 @@ enum ETranslucencyVolumeCascade
 	VIEW_UNIFORM_BUFFER_MEMBER(float, SkyLightAffectReflectionFlag) \
 	VIEW_UNIFORM_BUFFER_MEMBER(float, SkyLightAffectGlobalIlluminationFlag) \
 	VIEW_UNIFORM_BUFFER_MEMBER(FLinearColor, SkyLightColor) \
-	VIEW_UNIFORM_BUFFER_MEMBER_ARRAY(FVector4f, MobileSkyIrradianceEnvironmentMap, [7]) \
+	VIEW_UNIFORM_BUFFER_MEMBER_ARRAY(FVector4f, MobileSkyIrradianceEnvironmentMap, [SKY_IRRADIANCE_ENVIRONMENT_MAP_VEC4_COUNT]) \
 	VIEW_UNIFORM_BUFFER_MEMBER(float, MobilePreviewMode) \
 	VIEW_UNIFORM_BUFFER_MEMBER(float, HMDEyePaddingOffset) \
 	VIEW_UNIFORM_BUFFER_MEMBER_EX(float, ReflectionCubemapMaxMip, EShaderPrecisionModifier::Half) \
diff --git a/Engine/Source/Runtime/Renderer/Private/IndirectLightRendering.cpp b/Engine/Source/Runtime/Renderer/Private/IndirectLightRendering.cpp
--- a/Engine/Source/Runtime/Renderer/Private/IndirectLightRendering.cpp
+++ b/Engine/Source/Runtime/Renderer/Private/IndirectLightRendering.cpp
@@ -509,7 +509,6 @@ void SetupReflectionUniformParameters(const FViewInfo& View, FReflectionUniformP
 	OutParameters.SkyLightBlendDestinationCubemap = SkyLightBlendDestinationTextureResource->TextureRHI;
 	OutParameters.SkyLightBlendDestinationCubemapSampler = SkyLightBlendDestinationTextureResource->SamplerStateRHI;
 	OutParameters.SkyLightParameters = FVector4f(SkyMipCount - 1.0f, ApplySkyLightMask, bSkyLightIsDynamic ? 1.0f : 0.0f, BlendFraction);
-	OutParameters.SkyLightCubemapBrightness = SkyAverageBrightness;
 
 	// Note: GBlackCubeArrayTexture has an alpha of 0, which is needed to represent invalid data so the sky cubemap can still be applied
 	FRHITexture* CubeArrayTexture = (SupportsTextureCubeArray(View.FeatureLevel))? GBlackCubeArrayTexture->TextureRHI : GBlackTextureCube->TextureRHI;
diff --git a/Engine/Source/Runtime/Renderer/Private/ReflectionEnvironment.h b/Engine/Source/Runtime/Renderer/Private/ReflectionEnvironment.h
--- a/Engine/Source/Runtime/Renderer/Private/ReflectionEnvironment.h
+++ b/Engine/Source/Runtime/Renderer/Private/ReflectionEnvironment.h
@@ -16,7 +16,6 @@ extern bool IsReflectionCaptureAvailable();
 
 BEGIN_GLOBAL_SHADER_PARAMETER_STRUCT(FReflectionUniformParameters,)
 	SHADER_PARAMETER(FVector4f, SkyLightParameters)
-	SHADER_PARAMETER(float, SkyLightCubemapBrightness)
 	SHADER_PARAMETER_TEXTURE(TextureCube, SkyLightCubemap)
 	SHADER_PARAMETER_SAMPLER(SamplerState, SkyLightCubemapSampler)
 	SHADER_PARAMETER_TEXTURE(TextureCube, SkyLightBlendDestinationCubemap)
diff --git a/Engine/Source/Runtime/Renderer/Private/SceneRendering.cpp b/Engine/Source/Runtime/Renderer/Private/SceneRendering.cpp
--- a/Engine/Source/Runtime/Renderer/Private/SceneRendering.cpp
+++ b/Engine/Source/Runtime/Renderer/Private/SceneRendering.cpp
@@ -1729,7 +1729,7 @@ void FViewInfo::SetupUniformBufferParameters(
 	if (RHIFeatureLevel == ERHIFeatureLevel::ES3_1)
 	{
 		// Make sure there's no padding since we're going to cast to FVector4f*
-		static_assert(sizeof(ViewUniformShaderParameters.MobileSkyIrradianceEnvironmentMap) == sizeof(FVector4f) * 7, "unexpected sizeof ViewUniformShaderParameters.MobileSkyIrradianceEnvironmentMap");
+		static_assert(sizeof(ViewUniformShaderParameters.MobileSkyIrradianceEnvironmentMap) == sizeof(FVector4f) * SKY_IRRADIANCE_ENVIRONMENT_MAP_VEC4_COUNT, "unexpected sizeof ViewUniformShaderParameters.MobileSkyIrradianceEnvironmentMap");
 
 		const bool bSetupSkyIrradiance = Scene
 			&& Scene->SkyLight
@@ -1744,7 +1744,7 @@ void FViewInfo::SetupUniformBufferParameters(
 		}
 		else
 		{
-			FMemory::Memzero((FVector4f*)&ViewUniformShaderParameters.MobileSkyIrradianceEnvironmentMap, sizeof(FVector4f) * 7);
+			FMemory::Memzero((FVector4f*)&ViewUniformShaderParameters.MobileSkyIrradianceEnvironmentMap, sizeof(FVector4f) * SKY_IRRADIANCE_ENVIRONMENT_MAP_VEC4_COUNT);
 		}
 	}
 	else
@@ -5226,14 +5226,14 @@ void FSceneRenderer::UpdateSkyIrradianceGpuBuffer(FRHICommandListImmediate& RHIC
 
 	if (!Scene->SkyIrradianceEnvironmentMap.Buffer)
 	{
-		Scene->SkyIrradianceEnvironmentMap.Initialize(TEXT("SkyIrradianceEnvironmentMap"), sizeof(FVector4f), 7);
+		Scene->SkyIrradianceEnvironmentMap.Initialize(TEXT("SkyIrradianceEnvironmentMap"), sizeof(FVector4f), SKY_IRRADIANCE_ENVIRONMENT_MAP_VEC4_COUNT);
 	}
 
 	TRACE_CPUPROFILER_EVENT_SCOPE(UpdateSkyIrradianceGpuBuffer);
 
-	FVector4f OutSkyIrradianceEnvironmentMap[7];
+	FVector4f OutSkyIrradianceEnvironmentMap[SKY_IRRADIANCE_ENVIRONMENT_MAP_VEC4_COUNT];
 	// Make sure there's no padding since we're going to cast to FVector4f*
-	checkSlow(sizeof(OutSkyIrradianceEnvironmentMap) == sizeof(FVector4f) * 7);
+	checkSlow(sizeof(OutSkyIrradianceEnvironmentMap) == sizeof(FVector4f) * SKY_IRRADIANCE_ENVIRONMENT_MAP_VEC4_COUNT);
 	check(Scene);
 	const bool bUploadIrradiance = 
 		Scene->SkyLight
@@ -5247,10 +5247,14 @@ void FSceneRenderer::UpdateSkyIrradianceGpuBuffer(FRHICommandListImmediate& RHIC
 		const FSHVectorRGB3& SkyIrradiance = Scene->SkyLight->IrradianceEnvironmentMap;
 		SetupSkyIrradianceEnvironmentMapConstantsFromSkyIrradiance(OutSkyIrradianceEnvironmentMap, SkyIrradiance);
 
+		const float SkyLightAverageBrightness = Scene->SkyLight->AverageBrightness;
+		const FVector4f SkyAverageBrightness = FVector4f(SkyLightAverageBrightness, SkyLightAverageBrightness, SkyLightAverageBrightness, SkyLightAverageBrightness);
+
 		// Set the captured environment map data
-		void* DataPtr = RHICmdList.LockBuffer(Scene->SkyIrradianceEnvironmentMap.Buffer, 0, Scene->SkyIrradianceEnvironmentMap.NumBytes, RLM_WriteOnly);
+		FVector4f* DataPtr = (FVector4f*)RHICmdList.LockBuffer(Scene->SkyIrradianceEnvironmentMap.Buffer, 0, Scene->SkyIrradianceEnvironmentMap.NumBytes, RLM_WriteOnly);
 		checkSlow(Scene->SkyIrradianceEnvironmentMap.NumBytes == sizeof(OutSkyIrradianceEnvironmentMap));
-		FPlatformMemory::Memcpy(DataPtr, &OutSkyIrradianceEnvironmentMap, sizeof(OutSkyIrradianceEnvironmentMap));
+		FPlatformMemory::Memcpy(DataPtr    , &OutSkyIrradianceEnvironmentMap,	sizeof(FVector4f) * 7);
+		FPlatformMemory::Memcpy(DataPtr + 7, &SkyAverageBrightness,				sizeof(FVector4f) * 1);
 		RHICmdList.UnlockBuffer(Scene->SkyIrradianceEnvironmentMap.Buffer);
 	}
 	else if (Scene->SkyIrradianceEnvironmentMap.NumBytes == 0)
