Fixed: RDG validation assert related to mobile deferred shading pass parameters setup

diff --git a/Engine/Shaders/Private/MobileDeferredShading.usf b/Engine/Shaders/Private/MobileDeferredShading.usf
--- a/Engine/Shaders/Private/MobileDeferredShading.usf
+++ b/Engine/Shaders/Private/MobileDeferredShading.usf
@@ -5,6 +5,12 @@
 #include "Common.ush"
 #include "SHCommon.ush"
 #include "/Engine/Generated/Material.ush"
+
+// Reroute MobileSceneTextures uniform buffer references to the base pass uniform buffer
+#define MobileSceneTextures		MobileBasePass.SceneTextures
+#define ForwardLightData		MobileBasePass.Forward
+#define PlanarReflectionStruct	MobileBasePass.PlanarReflection
+
 #include "DynamicLightingCommon.ush"
 #include "ReflectionEnvironmentShared.ush"
 #include "LightGridCommon.ush"
@@ -13,9 +19,6 @@
 #include "LightFunctionCommon.ush"
 #include "PlanarReflectionShared.ush"
 
-// Reroute MobileSceneTextures uniform buffer references to the base pass uniform buffer
-#define MobileSceneTextures MobileBasePass.SceneTextures
-
 #if APPLY_SKY_REFLECTION
 #define ENABLE_SKY_LIGHT 1
 #endif
diff --git a/Engine/Source/Runtime/Renderer/Private/LightGridInjection.cpp b/Engine/Source/Runtime/Renderer/Private/LightGridInjection.cpp
--- a/Engine/Source/Runtime/Renderer/Private/LightGridInjection.cpp
+++ b/Engine/Source/Runtime/Renderer/Private/LightGridInjection.cpp
@@ -86,31 +86,35 @@ FAutoConsoleVariableRef CVarLightCullingMaxDistanceOverride(
 
 extern TAutoConsoleVariable<int32> CVarVirtualShadowOnePassProjection;
 
-TRDGUniformBufferRef<FForwardLightData> CreateDummyForwardLightUniformBuffer(FRDGBuilder& GraphBuilder)
+void SetupDummyForwardLightUniformParameters(FRDGBuilder& GraphBuilder, FForwardLightData& ForwardLightData)
 {
 	const FRDGSystemTextures& SystemTextures = FRDGSystemTextures::Get(GraphBuilder);
+	
+	ForwardLightData.DirectionalLightShadowmapAtlas = SystemTextures.Black;
+	ForwardLightData.DirectionalLightStaticShadowmap = GBlackTexture->TextureRHI;
 
-	FForwardLightData* ForwardLightData = GraphBuilder.AllocParameters<FForwardLightData>();
-	ForwardLightData->DirectionalLightShadowmapAtlas = SystemTextures.Black;
-	ForwardLightData->DirectionalLightStaticShadowmap = GBlackTexture->TextureRHI;
-
-	FRDGBufferRef ForwardLocalLightBuffer     = GSystemTextures.GetDefaultBuffer(GraphBuilder, sizeof(FVector4f));
-	ForwardLightData->ForwardLocalLightBuffer = GraphBuilder.CreateSRV(ForwardLocalLightBuffer, PF_A32B32G32R32F);
+	FRDGBufferRef ForwardLocalLightBuffer = GSystemTextures.GetDefaultBuffer(GraphBuilder, sizeof(FVector4f));
+	ForwardLightData.ForwardLocalLightBuffer = GraphBuilder.CreateSRV(ForwardLocalLightBuffer, PF_A32B32G32R32F);
 
-	FRDGBufferRef NumCulledLightsGrid     = GSystemTextures.GetDefaultBuffer(GraphBuilder, sizeof(uint32));
-	ForwardLightData->NumCulledLightsGrid = GraphBuilder.CreateSRV(NumCulledLightsGrid, PF_R32_UINT);
+	FRDGBufferRef NumCulledLightsGrid = GSystemTextures.GetDefaultBuffer(GraphBuilder, sizeof(uint32));
+	ForwardLightData.NumCulledLightsGrid = GraphBuilder.CreateSRV(NumCulledLightsGrid, PF_R32_UINT);
 
 	if (RHISupportsBufferLoadTypeConversion(GMaxRHIShaderPlatform))
 	{
-		FRDGBufferRef CulledLightDataGrid     = GSystemTextures.GetDefaultBuffer(GraphBuilder, sizeof(uint16));
-		ForwardLightData->CulledLightDataGrid = GraphBuilder.CreateSRV(CulledLightDataGrid, PF_R16_UINT);
+		FRDGBufferRef CulledLightDataGrid = GSystemTextures.GetDefaultBuffer(GraphBuilder, sizeof(uint16));
+		ForwardLightData.CulledLightDataGrid = GraphBuilder.CreateSRV(CulledLightDataGrid, PF_R16_UINT);
 	}
 	else
 	{
-		FRDGBufferRef CulledLightDataGrid     = GSystemTextures.GetDefaultBuffer(GraphBuilder, sizeof(uint32));
-		ForwardLightData->CulledLightDataGrid = GraphBuilder.CreateSRV(CulledLightDataGrid, PF_R32_UINT);
+		FRDGBufferRef CulledLightDataGrid = GSystemTextures.GetDefaultBuffer(GraphBuilder, sizeof(uint32));
+		ForwardLightData.CulledLightDataGrid = GraphBuilder.CreateSRV(CulledLightDataGrid, PF_R32_UINT);
 	}
+}
 
+TRDGUniformBufferRef<FForwardLightData> CreateDummyForwardLightUniformBuffer(FRDGBuilder& GraphBuilder)
+{
+	FForwardLightData* ForwardLightData = GraphBuilder.AllocParameters<FForwardLightData>();
+	SetupDummyForwardLightUniformParameters(GraphBuilder, *ForwardLightData);
 	return GraphBuilder.CreateUniformBuffer(ForwardLightData);
 }
 
diff --git a/Engine/Source/Runtime/Renderer/Private/MobileBasePassRendering.cpp b/Engine/Source/Runtime/Renderer/Private/MobileBasePassRendering.cpp
--- a/Engine/Source/Runtime/Renderer/Private/MobileBasePassRendering.cpp
+++ b/Engine/Source/Runtime/Renderer/Private/MobileBasePassRendering.cpp
@@ -166,6 +166,8 @@ FMobileBasePassMovableLightInfo::FMobileBasePassMovableLightInfo(const FPrimitiv
 	}
 }
 
+extern void SetupDummyForwardLightUniformParameters(FRDGBuilder& GraphBuilder, FForwardLightData& ForwardLightData);
+
 void SetupMobileBasePassUniformParameters(
 	FRDGBuilder& GraphBuilder,
 	const FViewInfo& View, 
@@ -176,6 +178,15 @@ void SetupMobileBasePassUniformParameters(
 {
 	SetupFogUniformParameters(GraphBuilder, View, BasePassParameters.Fog);
 
+	if (View.ForwardLightingResources.ForwardLightData)
+	{
+		BasePassParameters.Forward = *View.ForwardLightingResources.ForwardLightData;
+	}
+	else
+	{
+		SetupDummyForwardLightUniformParameters(GraphBuilder, BasePassParameters.Forward);
+	}
+
 	const FScene* Scene = View.Family->Scene ? View.Family->Scene->GetRenderScene() : nullptr;
 	const FPlanarReflectionSceneProxy* ReflectionSceneProxy = Scene ? Scene->GetForwardPassGlobalPlanarReflection() : nullptr;
 	SetupPlanarReflectionUniformParameters(View, ReflectionSceneProxy, BasePassParameters.PlanarReflection);
diff --git a/Engine/Source/Runtime/Renderer/Private/MobileBasePassRendering.h b/Engine/Source/Runtime/Renderer/Private/MobileBasePassRendering.h
--- a/Engine/Source/Runtime/Renderer/Private/MobileBasePassRendering.h
+++ b/Engine/Source/Runtime/Renderer/Private/MobileBasePassRendering.h
@@ -34,6 +34,7 @@ struct FMobileBasePassTextures
 BEGIN_GLOBAL_SHADER_PARAMETER_STRUCT(FMobileBasePassUniformParameters, )
 	SHADER_PARAMETER(float, AmbientOcclusionStaticFraction)
 	SHADER_PARAMETER_STRUCT(FFogUniformParameters, Fog)
+	SHADER_PARAMETER_STRUCT(FForwardLightData, Forward)
 	SHADER_PARAMETER_STRUCT(FPlanarReflectionUniformParameters, PlanarReflection) // Single global planar reflection for the forward pass.
 	SHADER_PARAMETER_STRUCT(FMobileSceneTextureUniformParameters, SceneTextures)
 	SHADER_PARAMETER_STRUCT(FDebugViewModeUniformParameters, DebugViewMode)
diff --git a/Engine/Source/Runtime/Renderer/Private/MobileDecalRendering.cpp b/Engine/Source/Runtime/Renderer/Private/MobileDecalRendering.cpp
--- a/Engine/Source/Runtime/Renderer/Private/MobileDecalRendering.cpp
+++ b/Engine/Source/Runtime/Renderer/Private/MobileDecalRendering.cpp
@@ -36,12 +36,14 @@ void FMobileSceneRenderer::RenderDecals(FRHICommandListImmediate& RHICmdList, co
 	// Deferred decals
 	if (Scene->Decals.Num() > 0)
 	{
+		SCOPED_DRAW_EVENT(RHICmdList, Decals);
 		RenderDeferredDecalsMobile(RHICmdList, *Scene, View, DecalRenderStage, RenderTargetMode);
 	}
 
 	// Mesh decals
 	if (View.MeshDecalBatches.Num() > 0)
 	{
+		SCOPED_DRAW_EVENT(RHICmdList, MeshDecals);
 		RenderMeshDecalsMobile(RHICmdList, View, DecalRenderStage, RenderTargetMode);
 	}
 }
diff --git a/Engine/Source/Runtime/Renderer/Private/MobileDeferredShadingPass.cpp b/Engine/Source/Runtime/Renderer/Private/MobileDeferredShadingPass.cpp
--- a/Engine/Source/Runtime/Renderer/Private/MobileDeferredShadingPass.cpp
+++ b/Engine/Source/Runtime/Renderer/Private/MobileDeferredShadingPass.cpp
@@ -42,11 +42,9 @@ class FMobileDirectLightFunctionPS : public FMaterialShader
 	using FPermutationDomain = TShaderPermutationDomain< FUseClustred, FApplySkyReflection, FApplyCSM, FApplyReflection, FShadowQuality>;
 
 	BEGIN_SHADER_PARAMETER_STRUCT(FParameters, )
-		SHADER_PARAMETER_RDG_UNIFORM_BUFFER(FForwardLightData, Forward)
 		SHADER_PARAMETER_STRUCT_REF(FMobileDirectionalLightShaderParameters, MobileDirectionalLight)
 		SHADER_PARAMETER_STRUCT_REF(FReflectionUniformParameters, ReflectionsParameters)
 		SHADER_PARAMETER_STRUCT_REF(FReflectionCaptureShaderData, ReflectionCaptureData)
-		SHADER_PARAMETER_STRUCT_REF(FPlanarReflectionUniformParameters, PlanarReflection) // Single global planar reflection.
 		SHADER_PARAMETER(FMatrix44f, TranslatedWorldToLight)
 		SHADER_PARAMETER(FVector4f, LightFunctionParameters)
 		SHADER_PARAMETER(FVector3f, LightFunctionParameters2)
@@ -217,6 +215,13 @@ static void RenderDirectLight(FRHICommandListImmediate& RHICmdList, const FScene
 		DirectionalLight = Scene.MobileDirectionalLights[ChannelIdx];
 	}
 
+	FString LightNameWithLevel = TEXT("DirectionalLight");
+	if (DirectionalLight)
+	{
+		FSceneRenderer::GetLightNameForDrawEvent(DirectionalLight->Proxy, LightNameWithLevel);
+	}
+	SCOPED_DRAW_EVENTF(RHICmdList, DirectionalLight, TEXT("%s"), *LightNameWithLevel);
+
 	FGraphicsPipelineStateInitializer GraphicsPSOInit;
 	RHICmdList.ApplyCachedRenderTargets(GraphicsPSOInit);
 	// Add to emissive in SceneColor
@@ -249,23 +254,12 @@ static void RenderDirectLight(FRHICommandListImmediate& RHICmdList, const FScene
 	SetGraphicsPipelineState(RHICmdList, GraphicsPSOInit, StencilRef);
 
 	FMobileDirectLightFunctionPS::FParameters PassParameters;
-	PassParameters.Forward = View.ForwardLightingResources.ForwardLightUniformBuffer;
 	PassParameters.MobileDirectionalLight = Scene.UniformBuffers.MobileDirectionalLightUniformBuffers[1];
 	PassParameters.ReflectionCaptureData = GetShaderBinding(View.ReflectionCaptureUniformBuffer);
 	FReflectionUniformParameters ReflectionUniformParameters;
 	SetupReflectionUniformParameters(View, ReflectionUniformParameters);
 	PassParameters.ReflectionsParameters = CreateUniformBufferImmediate(ReflectionUniformParameters, UniformBuffer_SingleDraw);
 	PassParameters.LightFunctionParameters = FVector4f(1.0f, 1.0f, 0.0f, 0.0f);
-
-	const FPlanarReflectionSceneProxy* ReflectionSceneProxy = Scene.GetForwardPassGlobalPlanarReflection();
-	FPlanarReflectionUniformParameters PlanarReflectionUniformParameters;
-	SetupPlanarReflectionUniformParameters(View, ReflectionSceneProxy, PlanarReflectionUniformParameters);
-	if (View.PrevViewInfo.MobilePixelProjectedReflection.IsValid())
-	{
-		PlanarReflectionUniformParameters.PlanarReflectionTexture = View.PrevViewInfo.MobilePixelProjectedReflection->GetRenderTargetItem().ShaderResourceTexture;
-	}
-	PassParameters.PlanarReflection = TUniformBufferRef<FPlanarReflectionUniformParameters>::CreateUniformBufferImmediate(PlanarReflectionUniformParameters, UniformBuffer_SingleDraw);
-
 	PassParameters.PreIntegratedGF = GSystemTextures.PreintegratedGF->GetRenderTargetItem().ShaderResourceTexture;
 	PassParameters.PreIntegratedGFSampler = TStaticSamplerState<SF_Bilinear, AM_Clamp, AM_Clamp, AM_Clamp>::GetRHI();
 
@@ -403,6 +397,10 @@ static void RenderLocalLight(
 	{
 		return;
 	}
+
+	FString LightNameWithLevel;
+	FSceneRenderer::GetLightNameForDrawEvent(LightSceneInfo.Proxy, LightNameWithLevel);
+	SCOPED_DRAW_EVENTF(RHICmdList, LocalLight, TEXT("%s"), *LightNameWithLevel);
 	
 	if (GMobileUseLightStencilCulling != 0)
 	{
@@ -512,6 +510,12 @@ static void RenderSimpleLights(
 	const FCachedLightMaterial& DefaultMaterial)
 {
 	const FSimpleLightArray& SimpleLights = SortedLightSet.SimpleLights;
+	if (SimpleLights.InstanceData.Num() == 0)
+	{
+		return;
+	}
+
+	SCOPED_DRAW_EVENT(RHICmdList, SimpleLights);
 
 	FDeferredLightVS::FPermutationDomain PermutationVectorVS;
 	PermutationVectorVS.Set<FDeferredLightVS::FRadialLight>(true);
@@ -598,6 +602,8 @@ void MobileDeferredShadingPass(
 	const FScene& Scene, 
 	const FSortedLightSetSceneInfo &SortedLightSet)
 {
+	SCOPED_DRAW_EVENT(RHICmdList, DeferredShading);
+	
 	RHICmdList.SetViewport(View.ViewRect.Min.X, View.ViewRect.Min.Y, 0.0f, View.ViewRect.Max.X, View.ViewRect.Max.Y, 1.0f);
 
 	// Default material for light rendering
