PR #8838: Fix Capsule Indirect Shadow incorrect angle (Contributed by [REDACTED_PERSON])

diff --git a/Engine/Source/Runtime/Renderer/Private/CapsuleShadowRendering.cpp b/Engine/Source/Runtime/Renderer/Private/CapsuleShadowRendering.cpp
--- a/Engine/Source/Runtime/Renderer/Private/CapsuleShadowRendering.cpp
+++ b/Engine/Source/Runtime/Renderer/Private/CapsuleShadowRendering.cpp
@@ -928,7 +928,7 @@ static IndirectCapsuleShadowsResources CreateIndirectCapsuleShadowsResources(
 			PassParameters->View = View.ViewUniformBuffer;
 			PassParameters->NumLightDirectionData = NumLightDataElements;
 			PassParameters->SkyLightMode = SkyLightMode;
-			PassParameters->CapsuleIndirectConeAngle = GCapsuleSkyAngleScale;
+			PassParameters->CapsuleIndirectConeAngle = GCapsuleIndirectConeAngle;
 			PassParameters->CapsuleSkyAngleScale = GCapsuleSkyAngleScale;
 			PassParameters->CapsuleMinSkyAngle = GCapsuleMinSkyAngle;
 			PassParameters->RWComputedLightDirectionData = ComputedLightDirectionUAV;
