Reduce the likelyhood of TDRs on Windows by adding a submit hint to the command list after the path tracing ray dispatch

Default value is 2 (flush after each tile) to minimize the overhead when path compaction is used.

diff --git a/Engine/Source/Runtime/Renderer/Private/PathTracing.cpp b/Engine/Source/Runtime/Renderer/Private/PathTracing.cpp
--- a/Engine/Source/Runtime/Renderer/Private/PathTracing.cpp
+++ b/Engine/Source/Runtime/Renderer/Private/PathTracing.cpp
@@ -41,6 +41,16 @@ TAutoConsoleVariable<int32> CVarPathTracingIndirectDispatch(
 	ECVF_RenderThreadSafe
 );
 
+TAutoConsoleVariable<int32> CVarPathTracingFlushDispatch(
+	TEXT("r.PathTracing.FlushDispatch"),
+	2,
+	TEXT("Enables flushing of the command list after dispatch to reduce the likelyhood of TDRs on Windows (default: 2)\n")
+	TEXT("0: off\n")
+	TEXT("1: flush after each dispatch\n")
+	TEXT("2: flush after each tile\n"),
+	ECVF_RenderThreadSafe
+);
+
 TAutoConsoleVariable<int32> CVarPathTracingDispatchSize(
 	TEXT("r.PathTracing.DispatchSize"),
 	2048,
@@ -1498,6 +1508,8 @@ void FDeferredShadingSceneRenderer::RenderPathTracing(
 		// should we use path compaction?
 		const int CompactionType = CVarPathTracingCompaction.GetValueOnRenderThread();
 		const bool bUseIndirectDispatch = GRHISupportsRayTracingDispatchIndirect && CVarPathTracingIndirectDispatch.GetValueOnRenderThread() != 0;
+		const int FlushRenderingCommands = CVarPathTracingFlushDispatch.GetValueOnRenderThread();
+		
 		FRDGBuffer* ActivePaths[2] = {};
 		FRDGBuffer* NumActivePaths[2] = {};
 		FRDGBuffer* PathStateData = nullptr;
@@ -1591,13 +1603,14 @@ void FDeferredShadingSceneRenderer::RenderPathTracing(
 						AddClearUAVPass(GraphBuilder, PassParameters->NumPathStates, 0);
 					}
 					ClearUnusedGraphResources(RayGenShader, PassParameters);
+					const bool bFlushRenderingCommands = FlushRenderingCommands == 1 || (FlushRenderingCommands == 2 && Bounce == MaxBounces);
 					GraphBuilder.AddPass(
 						CompactionType == 1
 							? RDG_EVENT_NAME("Path Tracer Compute (%d x %d) Tile=(%d,%d - %dx%d) Sample=%d/%d NumLights=%d (Bounce=%d%s)", ResX, ResY, TileX, TileY, DispatchSizeX, DispatchSizeY, PathTracingState->SampleIndex, MaxSPP, PassParameters->SceneLightCount, Bounce, bUseIndirectDispatch && Bounce > 0 ? TEXT(" indirect") : TEXT(""))
 							: RDG_EVENT_NAME("Path Tracer Compute (%d x %d) Tile=(%d,%d - %dx%d) Sample=%d/%d NumLights=%d", ResX, ResY, TileX, TileY, DispatchSizeX, DispatchSizeY, PathTracingState->SampleIndex, MaxSPP, PassParameters->SceneLightCount),
 						PassParameters,
 						ERDGPassFlags::Compute,
-						[PassParameters, RayGenShader, DispatchSizeX, DispatchSizeY, bUseIndirectDispatch, &View](FRHIRayTracingCommandList& RHICmdList)
+						[PassParameters, RayGenShader, DispatchSizeX, DispatchSizeY, bUseIndirectDispatch, bFlushRenderingCommands, &View](FRHIRayTracingCommandList& RHICmdList)
 					{
 						FRHIRayTracingScene* RayTracingSceneRHI = View.GetRayTracingSceneChecked();
 			
@@ -1622,6 +1635,10 @@ void FDeferredShadingSceneRenderer::RenderPathTracing(
 								DispatchSizeX, DispatchSizeY
 							);
 						}
+						if (bFlushRenderingCommands)
+						{
+							RHICmdList.SubmitCommandsHint();
+						}
 					});
 				}
 			}
