Fix shader draw debug when upscaler is enabled.

diff --git a/Engine/Shaders/Private/ShaderDrawDebug.usf b/Engine/Shaders/Private/ShaderDrawDebug.usf
--- a/Engine/Shaders/Private/ShaderDrawDebug.usf
+++ b/Engine/Shaders/Private/ShaderDrawDebug.usf
@@ -9,8 +9,11 @@
 StructuredBuffer<uint> ShaderDrawDebugPrimitive;
 
 Texture2D<float> DepthTexture;
-int2 DepthTextureResolution;
-float2 DepthTextureInvResolution;
+float2 OutputInvResolution; 
+float2 OriginalViewRectMin;
+float2 OriginalViewSize;
+float2 OriginalBufferInvSize;
+
 SamplerState DepthSampler;
 
 #if GPU_DEBUG_RENDERING_VS 
@@ -38,8 +41,9 @@ void ShaderDrawDebugPS(
 {
 	uint2 PixelCoord = Position.xy;
 
-	const float2 DepthUV = float2(PixelCoord) * DepthTextureInvResolution;
-	const float Depth = DepthTexture.SampleLevel(DepthSampler, DepthUV, 0);	
+	const float2 ViewportUV = float2(PixelCoord) * OutputInvResolution;
+	const float2 DepthUV    = (ViewportUV * OriginalViewSize.xy + OriginalViewRectMin) * OriginalBufferInvSize;
+	const float Depth = DepthTexture.SampleLevel(DepthSampler, DepthUV, 0);
 	const bool bIsHidden = Position.z < Depth; // Reverse-Z
 
 	const float ColorScale = bIsHidden ? 0.4f : 1;
diff --git a/Engine/Source/Runtime/Renderer/Private/ShaderDebug.cpp b/Engine/Source/Runtime/Renderer/Private/ShaderDebug.cpp
--- a/Engine/Source/Runtime/Renderer/Private/ShaderDebug.cpp
+++ b/Engine/Source/Runtime/Renderer/Private/ShaderDebug.cpp
@@ -162,9 +162,11 @@ namespace ShaderDrawDebug
 		SHADER_USE_PARAMETER_STRUCT(FShaderDrawDebugPS, FGlobalShader);
 		
 		BEGIN_SHADER_PARAMETER_STRUCT(FParameters, )
+			SHADER_PARAMETER(FVector2f, OutputInvResolution)
+			SHADER_PARAMETER(FVector2f, OriginalViewRectMin)
+			SHADER_PARAMETER(FVector2f, OriginalViewSize)
+			SHADER_PARAMETER(FVector2f, OriginalBufferInvSize)
 			SHADER_PARAMETER_RDG_TEXTURE(Texture2D, DepthTexture)
-			SHADER_PARAMETER(FIntPoint, DepthTextureResolution)
-			SHADER_PARAMETER(FVector2f, DepthTextureInvResolution)
 			SHADER_PARAMETER_SAMPLER(SamplerState, DepthSampler)
 			RENDER_TARGET_BINDING_SLOTS()
 		END_SHADER_PARAMETER_STRUCT()
@@ -188,8 +190,8 @@ namespace ShaderDrawDebug
 	IMPLEMENT_GLOBAL_SHADER(FShaderDrawDebugPS, "/Engine/Private/ShaderDrawDebug.usf", "ShaderDrawDebugPS", SF_Pixel);
 
 	BEGIN_SHADER_PARAMETER_STRUCT(FShaderDrawVSPSParameters , )
-		SHADER_PARAMETER_STRUCT_INCLUDE(FShaderDrawDebugVS::FParameters, ShaderDrawVSParameters)
-		SHADER_PARAMETER_STRUCT_INCLUDE(FShaderDrawDebugPS::FParameters, ShaderDrawPSParameters)
+		SHADER_PARAMETER_STRUCT_INCLUDE(FShaderDrawDebugVS::FParameters, VS)
+		SHADER_PARAMETER_STRUCT_INCLUDE(FShaderDrawDebugPS::FParameters, PS)
 	END_SHADER_PARAMETER_STRUCT()
 
 	//////////////////////////////////////////////////////////////////////////
@@ -243,22 +245,23 @@ namespace ShaderDrawDebug
 		TShaderMapRef<FShaderDrawDebugPS> PixelShader(View.ShaderMap);
 
 		FShaderDrawVSPSParameters* PassParameters = GraphBuilder.AllocParameters<FShaderDrawVSPSParameters >();
-		PassParameters->ShaderDrawPSParameters.RenderTargets[0] = FRenderTargetBinding(OutputTexture, ERenderTargetLoadAction::ELoad);
-		//PassParameters->ShaderDrawPSParameters.RenderTargets.DepthStencil = FDepthStencilBinding(DepthTexture, ERenderTargetLoadAction::ELoad, FExclusiveDepthStencil::DepthRead_StencilNop);
-		PassParameters->ShaderDrawPSParameters.DepthTexture = DepthTexture;
-		PassParameters->ShaderDrawPSParameters.DepthTextureResolution = FIntPoint(DepthTexture->Desc.Extent.X, DepthTexture->Desc.Extent.Y);
-		PassParameters->ShaderDrawPSParameters.DepthTextureInvResolution = FVector2f(1.f / DepthTexture->Desc.Extent.X, 1.f / DepthTexture->Desc.Extent.Y);
-		PassParameters->ShaderDrawPSParameters.DepthSampler = TStaticSamplerState<SF_Point, AM_Clamp, AM_Clamp, AM_Clamp>::GetRHI();
-		PassParameters->ShaderDrawVSParameters.View = View.ViewUniformBuffer;
-		PassParameters->ShaderDrawVSParameters.ShaderDrawDebugPrimitive = GraphBuilder.CreateSRV(DataBuffer);
-		PassParameters->ShaderDrawVSParameters.IndirectBuffer = IndirectBuffer;
-
-		ValidateShaderParameters(PixelShader, PassParameters->ShaderDrawPSParameters);
-		ClearUnusedGraphResources(PixelShader, &PassParameters->ShaderDrawPSParameters, { IndirectBuffer });
-		ValidateShaderParameters(VertexShader, PassParameters->ShaderDrawVSParameters);
-		ClearUnusedGraphResources(VertexShader, &PassParameters->ShaderDrawVSParameters, { IndirectBuffer });
-
-		const FIntRect Viewport = View.ViewRect;
+		PassParameters->VS.View = View.ViewUniformBuffer;
+		PassParameters->PS.RenderTargets[0] = FRenderTargetBinding(OutputTexture, ERenderTargetLoadAction::ELoad);
+		PassParameters->PS.OutputInvResolution = FVector2f(1.f / View.UnscaledViewRect.Width(), 1.f / View.UnscaledViewRect.Height());
+		PassParameters->PS.OriginalViewRectMin = FVector2f(View.ViewRect.Min);
+		PassParameters->PS.OriginalViewSize = FVector2f(View.ViewRect.Width(), View.ViewRect.Height());
+		PassParameters->PS.OriginalBufferInvSize = FVector2f(1.f / DepthTexture->Desc.Extent.X, 1.f / DepthTexture->Desc.Extent.Y);
+		PassParameters->PS.DepthTexture = DepthTexture;
+		PassParameters->PS.DepthSampler = TStaticSamplerState<SF_Point, AM_Clamp, AM_Clamp, AM_Clamp>::GetRHI();
+		PassParameters->VS.ShaderDrawDebugPrimitive = GraphBuilder.CreateSRV(DataBuffer);
+		PassParameters->VS.IndirectBuffer = IndirectBuffer;
+
+		ValidateShaderParameters(PixelShader, PassParameters->PS);
+		ClearUnusedGraphResources(PixelShader, &PassParameters->PS, { IndirectBuffer });
+		ValidateShaderParameters(VertexShader, PassParameters->VS);
+		ClearUnusedGraphResources(VertexShader, &PassParameters->VS, { IndirectBuffer });
+
+		const FIntRect Viewport = View.UnscaledViewRect;
 		GraphBuilder.AddPass(
 			RDG_EVENT_NAME("ShaderDebug::Draw"),
 			PassParameters,
@@ -266,7 +269,7 @@ namespace ShaderDrawDebug
 			[VertexShader, PixelShader, PassParameters, IndirectBuffer, Viewport](FRHICommandList& RHICmdList)
 			{
 				// Marks the indirect draw parameter as used by the pass, given it's not used directly by any of the shaders.
-				PassParameters->ShaderDrawVSParameters.IndirectBuffer->MarkResourceAsUsed();
+				PassParameters->VS.IndirectBuffer->MarkResourceAsUsed();
 
 				FGraphicsPipelineStateInitializer GraphicsPSOInit;
 				RHICmdList.ApplyCachedRenderTargets(GraphicsPSOInit);
@@ -280,11 +283,11 @@ namespace ShaderDrawDebug
 				SetGraphicsPipelineState(RHICmdList, GraphicsPSOInit, 0);
 
 				RHICmdList.SetViewport(Viewport.Min.X, Viewport.Min.Y, 0.0f, Viewport.Max.X, Viewport.Max.Y, 1.0f);
-				SetShaderParameters(RHICmdList, VertexShader, VertexShader.GetVertexShader(), PassParameters->ShaderDrawVSParameters);
-				SetShaderParameters(RHICmdList, PixelShader, PixelShader.GetPixelShader(), PassParameters->ShaderDrawPSParameters);
+				SetShaderParameters(RHICmdList, VertexShader, VertexShader.GetVertexShader(), PassParameters->VS);
+				SetShaderParameters(RHICmdList, PixelShader, PixelShader.GetPixelShader(), PassParameters->PS);
 
 				// Marks the indirect draw parameter as used by the pass, given it's not used directly by any of the shaders.
-				FRHIBuffer* IndirectBufferRHI = PassParameters->ShaderDrawVSParameters.IndirectBuffer->GetIndirectRHICallBuffer();
+				FRHIBuffer* IndirectBufferRHI = PassParameters->VS.IndirectBuffer->GetIndirectRHICallBuffer();
 				check(IndirectBufferRHI != nullptr);
 				RHICmdList.DrawPrimitiveIndirect(IndirectBufferRHI, 0);
 			});
