Add precomputed BSDF directional albedo path for non-editor target.

This avoids to run the BuildTable pass which is too long on low-powered platforms.

diff --git a/Engine/Source/Runtime/Renderer/Private/ShadingEnergyConservation.cpp b/Engine/Source/Runtime/Renderer/Private/ShadingEnergyConservation.cpp
--- a/Engine/Source/Runtime/Renderer/Private/ShadingEnergyConservation.cpp
+++ b/Engine/Source/Runtime/Renderer/Private/ShadingEnergyConservation.cpp
@@ -57,6 +57,45 @@ static TAutoConsoleVariable<int32> CVarMaterialEnergyConservation(
 
 ////////////////////////////////////////////////////////////////////////////////////////////////////////////////////////////////////////////////
 
+namespace ShadingEnergyConservationData
+{
+	#include "ShadingEnergyConservationData.h"
+
+	template<typename TDataType>
+	void LockCopyTexture2D(FRHICommandListImmediate& RHICmdList, FTexture2DRHIRef& Texture, const TDataType* InSrcBuffer, uint32 NumComponents)
+	{
+		const uint8* SrcBuffer = (const uint8*)InSrcBuffer;
+		const uint32 SrcBytesPerPixel = sizeof(TDataType) * NumComponents;
+
+		uint32 SrcStride = SrcBytesPerPixel * SHADING_ENERGY_CONSERVATION_TABLE_RESOLUTION;
+		uint32 DstStride;
+		uint8* DstBuffer = (uint8*)RHICmdList.LockTexture2D(Texture, 0, RLM_WriteOnly, DstStride, false);
+		for (uint32 y = 0; y < SHADING_ENERGY_CONSERVATION_TABLE_RESOLUTION; ++y)
+		{
+			const uint8* Src = SrcBuffer + y * SrcStride;
+			uint8* Dst = DstBuffer + y * DstStride;
+			FMemory::Memcpy(Dst, Src, SrcStride);
+		}
+		RHICmdList.UnlockTexture2D(Texture, 0, false);
+	}
+
+	template<typename TDataType>
+	void LockCopyTexture3D(FRHICommandListImmediate& RHICmdList, FTexture3DRHIRef& Texture, const TDataType* InSrcBuffer, uint32 NumComponents)
+	{
+		const uint8* SrcBuffer = (const uint8*)InSrcBuffer;
+		const uint32 SrcBytesPerPixel = sizeof(TDataType) * NumComponents;
+
+		const FIntVector Extent(SHADING_ENERGY_CONSERVATION_TABLE_RESOLUTION, SHADING_ENERGY_CONSERVATION_TABLE_RESOLUTION, SHADING_ENERGY_CONSERVATION_TABLE_RESOLUTION);
+		FUpdateTextureRegion3D Region(0, 0, 0, 0, 0, 0, Extent.X, Extent.Y, Extent.Z);
+		RHICmdList.UpdateTexture3D(Texture, 0, Region, Extent.X * SrcBytesPerPixel, Extent.X * Extent.Y * SrcBytesPerPixel, SrcBuffer);
+
+		// UpdateTexture3D before and after state is currently undefined
+		RHICmdList.Transition(FRHITransitionInfo(Texture, ERHIAccess::Unknown, ERHIAccess::SRVMask));
+	}
+}
+
+////////////////////////////////////////////////////////////////////////////////////////////////////////////////////////////////////////////////
+
 class FShadingFurnaceTestPassPS : public FGlobalShader
 {
 	DECLARE_GLOBAL_SHADER(FShadingFurnaceTestPassPS);
@@ -189,14 +228,20 @@ void Init(FRDGBuilder& GraphBuilder, FViewInfo& View)
 	// Build/bind table if energy conservation is enabled or if strata is enabled in order to have 
 	// the correct tables built & bound. Even if we are not using energy conservation, we want to 
 	// have access to directional albedo information for env. lighting for instance)
-	// TODO we do disable LUT generation when bMaterialEnergyConservationEnabled is false in order avoid having
-	//      some platofrms hanging when generating the luts on startup. We will revisit with LUTs that are simply loaded.
-	const bool bBindEnergyData = (View.ViewState != nullptr) && bMaterialEnergyConservationEnabled && (bIsEnergyPreservationEnabled || bIsEnergyConservationEnabled || Strata::IsStrataEnabled() || View.Family->EngineShowFlags.PathTracing);
+	const bool bBindEnergyData = (View.ViewState != nullptr) && (bMaterialEnergyConservationEnabled || Strata::IsStrataEnabled() || View.Family->EngineShowFlags.PathTracing) && (bIsEnergyPreservationEnabled || bIsEnergyConservationEnabled);
 	if (bBindEnergyData)
 	{
-		const EPixelFormat Format = CVarShadingFurnaceTest_TableFormat.GetValueOnRenderThread() > 0 ? PF_G32R32F : PF_G16R16F;
+		#if UE_EDITOR
+		const bool bRuntimeGeneration = true;
+		#else
+		const bool bRuntimeGeneration = false;
+		#endif
+
+		const int Size = SHADING_ENERGY_CONSERVATION_TABLE_RESOLUTION;
+		const EPixelFormat SpecFormat = bRuntimeGeneration && CVarShadingFurnaceTest_TableFormat.GetValueOnRenderThread() > 0 ? PF_G32R32F : PF_G16R16;
+		const EPixelFormat DiffFormat = PF_G16;
 		const bool bBuildTable = 
-			View.ViewState->ShadingEnergyConservationData.Format != Format ||
+			View.ViewState->ShadingEnergyConservationData.Format != SpecFormat ||
 			View.ViewState->ShadingEnergyConservationData.GGXSpecEnergyTexture == nullptr ||
 			View.ViewState->ShadingEnergyConservationData.GGXGlassEnergyTexture == nullptr ||
 			View.ViewState->ShadingEnergyConservationData.ClothEnergyTexture ==  nullptr ||
@@ -204,86 +249,100 @@ void Init(FRDGBuilder& GraphBuilder, FViewInfo& View)
 
 		if (bBuildTable)
 		{
-			const int Size = SHADING_ENERGY_CONSERVATION_TABLE_RESOLUTION;
-			GGXSpecEnergyTexture = GraphBuilder.CreateTexture(FRDGTextureDesc::Create2D(
-				FIntPoint(Size, Size),
-				Format,
-				FClearValueBinding::None,
-				TexCreate_ShaderResource | TexCreate_UAV), TEXT("Shading.GGXSpecEnergy"), ERDGTextureFlags::MultiFrame);
-			GGXGlassEnergyTexture = GraphBuilder.CreateTexture(FRDGTextureDesc::Create3D(
-				FIntVector(Size, Size, Size),
-				Format,
-				FClearValueBinding::None,
-				TexCreate_ShaderResource | TexCreate_UAV), TEXT("Shading.GGXGlassEnergy"), ERDGTextureFlags::MultiFrame);
-			ClothEnergyTexture = GraphBuilder.CreateTexture(FRDGTextureDesc::Create2D(
-				FIntPoint(Size, Size),
-				Format,
-				FClearValueBinding::None,
-				TexCreate_ShaderResource | TexCreate_UAV), TEXT("Shading.ClothSpecEnergy"), ERDGTextureFlags::MultiFrame);
-			DiffuseEnergyTexture = GraphBuilder.CreateTexture(FRDGTextureDesc::Create2D(
-				FIntPoint(Size, Size),
-				PF_R16F,
-				FClearValueBinding::None,
-				TexCreate_ShaderResource | TexCreate_UAV), TEXT("Shading.DiffuseEnergy"), ERDGTextureFlags::MultiFrame);
+			View.ViewState->ShadingEnergyConservationData.Format = SpecFormat;
+
+			if (bRuntimeGeneration)
+			{			
+				GGXSpecEnergyTexture	= GraphBuilder.CreateTexture(FRDGTextureDesc::Create2D(FIntPoint(Size, Size),        SpecFormat, FClearValueBinding::None, TexCreate_ShaderResource | TexCreate_UAV), TEXT("Shading.GGXSpecEnergy"),   ERDGTextureFlags::MultiFrame);
+				GGXGlassEnergyTexture	= GraphBuilder.CreateTexture(FRDGTextureDesc::Create3D(FIntVector(Size, Size, Size), SpecFormat, FClearValueBinding::None, TexCreate_ShaderResource | TexCreate_UAV), TEXT("Shading.GGXGlassEnergy"),  ERDGTextureFlags::MultiFrame);
+				ClothEnergyTexture		= GraphBuilder.CreateTexture(FRDGTextureDesc::Create2D(FIntPoint(Size, Size),        SpecFormat, FClearValueBinding::None, TexCreate_ShaderResource | TexCreate_UAV), TEXT("Shading.ClothSpecEnergy"), ERDGTextureFlags::MultiFrame);
+				DiffuseEnergyTexture	= GraphBuilder.CreateTexture(FRDGTextureDesc::Create2D(FIntPoint(Size, Size),        DiffFormat, FClearValueBinding::None, TexCreate_ShaderResource | TexCreate_UAV), TEXT("Shading.DiffuseEnergy"),   ERDGTextureFlags::MultiFrame);
 			
-			{
-				FBuildShadingEnergyConservationTableCS::FPermutationDomain PermutationVector;
-				PermutationVector.Set<FBuildShadingEnergyConservationTableCS::FEnergyTableDim>(FBuildShadingEnergyConservationTableCS::EEnergyTableType::GGXSpecular);
-				TShaderMapRef<FBuildShadingEnergyConservationTableCS> ComputeShader(View.ShaderMap, PermutationVector);
-				FBuildShadingEnergyConservationTableCS::FParameters* PassParameters = GraphBuilder.AllocParameters<FBuildShadingEnergyConservationTableCS::FParameters>();
-				PassParameters->Output2Texture2D = GraphBuilder.CreateUAV(FRDGTextureUAVDesc(GGXSpecEnergyTexture, 0));
-				FComputeShaderUtils::AddPass(
-					GraphBuilder,
-					RDG_EVENT_NAME("ShadingEnergyConservation::BuildTable(GGXSpec)"),
-					ComputeShader,
-					PassParameters,
-					FComputeShaderUtils::GetGroupCount(FIntPoint(Size, Size), FComputeShaderUtils::kGolden2DGroupSize));
+				// GGX
+				{
+					FBuildShadingEnergyConservationTableCS::FPermutationDomain PermutationVector;
+					PermutationVector.Set<FBuildShadingEnergyConservationTableCS::FEnergyTableDim>(FBuildShadingEnergyConservationTableCS::EEnergyTableType::GGXSpecular);
+					TShaderMapRef<FBuildShadingEnergyConservationTableCS> ComputeShader(View.ShaderMap, PermutationVector);
+					FBuildShadingEnergyConservationTableCS::FParameters* PassParameters = GraphBuilder.AllocParameters<FBuildShadingEnergyConservationTableCS::FParameters>();
+					PassParameters->Output2Texture2D = GraphBuilder.CreateUAV(FRDGTextureUAVDesc(GGXSpecEnergyTexture, 0));
+					FComputeShaderUtils::AddPass(
+						GraphBuilder,
+						RDG_EVENT_NAME("ShadingEnergyConservation::BuildTable(GGXSpec)"),
+						ComputeShader,
+						PassParameters,
+						FComputeShaderUtils::GetGroupCount(FIntPoint(Size, Size), FComputeShaderUtils::kGolden2DGroupSize));
+				}
+
+				// GGX (Reflection + Transmission) indexed by IOR
+				{
+					FBuildShadingEnergyConservationTableCS::FPermutationDomain PermutationVector;
+					PermutationVector.Set<FBuildShadingEnergyConservationTableCS::FEnergyTableDim>(FBuildShadingEnergyConservationTableCS::EEnergyTableType::GGXGlass);
+					TShaderMapRef<FBuildShadingEnergyConservationTableCS> ComputeShader(View.ShaderMap, PermutationVector);
+					FBuildShadingEnergyConservationTableCS::FParameters* PassParameters = GraphBuilder.AllocParameters<FBuildShadingEnergyConservationTableCS::FParameters>();
+					PassParameters->OutputTexture3D = GraphBuilder.CreateUAV(FRDGTextureUAVDesc(GGXGlassEnergyTexture, 0));
+					FComputeShaderUtils::AddPass(
+						GraphBuilder,
+						RDG_EVENT_NAME("ShadingEnergyConservation::BuildTable(GGXGlass)"),
+						ComputeShader,
+						PassParameters,
+						FComputeShaderUtils::GetGroupCount(FIntVector(Size, Size, Size), FIntVector(FComputeShaderUtils::kGolden2DGroupSize, FComputeShaderUtils::kGolden2DGroupSize, 1)));
+				}
+
+				// Cloth
+				{
+					FBuildShadingEnergyConservationTableCS::FPermutationDomain PermutationVector;
+					PermutationVector.Set<FBuildShadingEnergyConservationTableCS::FEnergyTableDim>(FBuildShadingEnergyConservationTableCS::EEnergyTableType::Cloth);
+					TShaderMapRef<FBuildShadingEnergyConservationTableCS> ComputeShader(View.ShaderMap, PermutationVector);
+					FBuildShadingEnergyConservationTableCS::FParameters* PassParameters = GraphBuilder.AllocParameters<FBuildShadingEnergyConservationTableCS::FParameters>();
+					PassParameters->Output2Texture2D = GraphBuilder.CreateUAV(FRDGTextureUAVDesc(ClothEnergyTexture, 0));
+					FComputeShaderUtils::AddPass(
+						GraphBuilder,
+						RDG_EVENT_NAME("ShadingEnergyConservation::BuildTable(Cloth)"),
+						ComputeShader,
+						PassParameters,
+						FComputeShaderUtils::GetGroupCount(FIntPoint(Size, Size), FComputeShaderUtils::kGolden2DGroupSize));
+				}
+
+				// Diffuse
+				{
+					FBuildShadingEnergyConservationTableCS::FPermutationDomain PermutationVector;
+					PermutationVector.Set<FBuildShadingEnergyConservationTableCS::FEnergyTableDim>(FBuildShadingEnergyConservationTableCS::EEnergyTableType::Diffuse);
+					TShaderMapRef<FBuildShadingEnergyConservationTableCS> ComputeShader(View.ShaderMap, PermutationVector);
+					FBuildShadingEnergyConservationTableCS::FParameters* PassParameters = GraphBuilder.AllocParameters<FBuildShadingEnergyConservationTableCS::FParameters>();
+					PassParameters->Output1Texture2D = GraphBuilder.CreateUAV(FRDGTextureUAVDesc(DiffuseEnergyTexture, 0));
+					FComputeShaderUtils::AddPass(
+						GraphBuilder,
+						RDG_EVENT_NAME("ShadingEnergyConservation::BuildTable(Diffuse)"),
+						ComputeShader,
+						PassParameters,
+						FComputeShaderUtils::GetGroupCount(FIntPoint(Size, Size), FComputeShaderUtils::kGolden2DGroupSize));
+				}
+
+				GraphBuilder.QueueTextureExtraction(GGXSpecEnergyTexture,	&View.ViewState->ShadingEnergyConservationData.GGXSpecEnergyTexture);
+				GraphBuilder.QueueTextureExtraction(GGXGlassEnergyTexture,	&View.ViewState->ShadingEnergyConservationData.GGXGlassEnergyTexture);
+				GraphBuilder.QueueTextureExtraction(ClothEnergyTexture,		&View.ViewState->ShadingEnergyConservationData.ClothEnergyTexture);
+				GraphBuilder.QueueTextureExtraction(DiffuseEnergyTexture,	&View.ViewState->ShadingEnergyConservationData.DiffuseEnergyTexture);
 			}
+			else
 			{
-				FBuildShadingEnergyConservationTableCS::FPermutationDomain PermutationVector;
-				PermutationVector.Set<FBuildShadingEnergyConservationTableCS::FEnergyTableDim>(FBuildShadingEnergyConservationTableCS::EEnergyTableType::GGXGlass);
-				TShaderMapRef<FBuildShadingEnergyConservationTableCS> ComputeShader(View.ShaderMap, PermutationVector);
-				FBuildShadingEnergyConservationTableCS::FParameters* PassParameters = GraphBuilder.AllocParameters<FBuildShadingEnergyConservationTableCS::FParameters>();
-				PassParameters->OutputTexture3D = GraphBuilder.CreateUAV(FRDGTextureUAVDesc(GGXGlassEnergyTexture, 0));
-				FComputeShaderUtils::AddPass(
-					GraphBuilder,
-					RDG_EVENT_NAME("ShadingEnergyConservation::BuildTable(GGXGlass)"),
-					ComputeShader,
-					PassParameters,
-					FComputeShaderUtils::GetGroupCount(FIntVector(Size, Size, Size), FIntVector(FComputeShaderUtils::kGolden2DGroupSize, FComputeShaderUtils::kGolden2DGroupSize, 1)));
+				// Precomputed data are stored as float16
+				check(SpecFormat == PF_G16R16);
+
+				View.ViewState->ShadingEnergyConservationData.GGXSpecEnergyTexture  = GRenderTargetPool.FindFreeElement(FRDGTextureDesc::Create2D(FIntPoint(Size, Size),        SpecFormat, FClearValueBinding::None, TexCreate_ShaderResource), TEXT("Shading.GGXSpecEnergy"));
+				View.ViewState->ShadingEnergyConservationData.GGXGlassEnergyTexture = GRenderTargetPool.FindFreeElement(FRDGTextureDesc::Create3D(FIntVector(Size, Size, Size), SpecFormat, FClearValueBinding::None, TexCreate_ShaderResource), TEXT("Shading.GGXGlassEnergy"));
+				View.ViewState->ShadingEnergyConservationData.ClothEnergyTexture    = GRenderTargetPool.FindFreeElement(FRDGTextureDesc::Create2D(FIntPoint(Size, Size),        SpecFormat, FClearValueBinding::None, TexCreate_ShaderResource), TEXT("Shading.ClothSpecEnergy"));
+				View.ViewState->ShadingEnergyConservationData.DiffuseEnergyTexture  = GRenderTargetPool.FindFreeElement(FRDGTextureDesc::Create2D(FIntPoint(Size, Size),        DiffFormat,	FClearValueBinding::None, TexCreate_ShaderResource), TEXT("Shading.DiffuseEnergy"));
+
+				ShadingEnergyConservationData::LockCopyTexture2D(GraphBuilder.RHICmdList, (FTexture2DRHIRef&)View.ViewState->ShadingEnergyConservationData.GGXSpecEnergyTexture->GetRenderTargetItem().ShaderResourceTexture,  ShadingEnergyConservationData::GGXSpecValues, 2);
+				ShadingEnergyConservationData::LockCopyTexture3D(GraphBuilder.RHICmdList, (FTexture3DRHIRef&)View.ViewState->ShadingEnergyConservationData.GGXGlassEnergyTexture->GetRenderTargetItem().ShaderResourceTexture, ShadingEnergyConservationData::GGXGlassValues, 2);
+				ShadingEnergyConservationData::LockCopyTexture2D(GraphBuilder.RHICmdList, (FTexture2DRHIRef&)View.ViewState->ShadingEnergyConservationData.ClothEnergyTexture->GetRenderTargetItem().ShaderResourceTexture,    ShadingEnergyConservationData::ClothSpecValues, 2);
+				ShadingEnergyConservationData::LockCopyTexture2D(GraphBuilder.RHICmdList, (FTexture2DRHIRef&)View.ViewState->ShadingEnergyConservationData.DiffuseEnergyTexture->GetRenderTargetItem().ShaderResourceTexture,  ShadingEnergyConservationData::DiffuseValues, 1);
+
+				GGXSpecEnergyTexture  = GraphBuilder.RegisterExternalTexture(View.ViewState->ShadingEnergyConservationData.GGXSpecEnergyTexture);
+				GGXGlassEnergyTexture = GraphBuilder.RegisterExternalTexture(View.ViewState->ShadingEnergyConservationData.GGXGlassEnergyTexture);
+				ClothEnergyTexture    = GraphBuilder.RegisterExternalTexture(View.ViewState->ShadingEnergyConservationData.ClothEnergyTexture);
+				DiffuseEnergyTexture  = GraphBuilder.RegisterExternalTexture(View.ViewState->ShadingEnergyConservationData.DiffuseEnergyTexture);
 			}
-			{
-				FBuildShadingEnergyConservationTableCS::FPermutationDomain PermutationVector;
-				PermutationVector.Set<FBuildShadingEnergyConservationTableCS::FEnergyTableDim>(FBuildShadingEnergyConservationTableCS::EEnergyTableType::Cloth);
-				TShaderMapRef<FBuildShadingEnergyConservationTableCS> ComputeShader(View.ShaderMap, PermutationVector);
-				FBuildShadingEnergyConservationTableCS::FParameters* PassParameters = GraphBuilder.AllocParameters<FBuildShadingEnergyConservationTableCS::FParameters>();
-				PassParameters->Output2Texture2D = GraphBuilder.CreateUAV(FRDGTextureUAVDesc(ClothEnergyTexture, 0));
-				FComputeShaderUtils::AddPass(
-					GraphBuilder,
-					RDG_EVENT_NAME("ShadingEnergyConservation::BuildTable(Cloth)"),
-					ComputeShader,
-					PassParameters,
-					FComputeShaderUtils::GetGroupCount(FIntPoint(Size, Size), FComputeShaderUtils::kGolden2DGroupSize));
-			}
-			{
-				FBuildShadingEnergyConservationTableCS::FPermutationDomain PermutationVector;
-				PermutationVector.Set<FBuildShadingEnergyConservationTableCS::FEnergyTableDim>(FBuildShadingEnergyConservationTableCS::EEnergyTableType::Diffuse);
-				TShaderMapRef<FBuildShadingEnergyConservationTableCS> ComputeShader(View.ShaderMap, PermutationVector);
-				FBuildShadingEnergyConservationTableCS::FParameters* PassParameters = GraphBuilder.AllocParameters<FBuildShadingEnergyConservationTableCS::FParameters>();
-				PassParameters->Output1Texture2D = GraphBuilder.CreateUAV(FRDGTextureUAVDesc(DiffuseEnergyTexture, 0));
-				FComputeShaderUtils::AddPass(
-					GraphBuilder,
-					RDG_EVENT_NAME("ShadingEnergyConservation::BuildTable(Diffuse)"),
-					ComputeShader,
-					PassParameters,
-					FComputeShaderUtils::GetGroupCount(FIntPoint(Size, Size), FComputeShaderUtils::kGolden2DGroupSize));
-			}
-
-			GraphBuilder.QueueTextureExtraction(GGXSpecEnergyTexture,	&View.ViewState->ShadingEnergyConservationData.GGXSpecEnergyTexture);
-			GraphBuilder.QueueTextureExtraction(GGXGlassEnergyTexture,	&View.ViewState->ShadingEnergyConservationData.GGXGlassEnergyTexture);
-			GraphBuilder.QueueTextureExtraction(ClothEnergyTexture,		&View.ViewState->ShadingEnergyConservationData.ClothEnergyTexture);
-			GraphBuilder.QueueTextureExtraction(DiffuseEnergyTexture,	&View.ViewState->ShadingEnergyConservationData.DiffuseEnergyTexture);
-			View.ViewState->ShadingEnergyConservationData.Format = Format;
 		}
 		else
 		{
