Fix CIS warning.

diff --git a/Engine/Source/Runtime/Renderer/Private/ShadingEnergyConservationData.h b/Engine/Source/Runtime/Renderer/Private/ShadingEnergyConservationData.h
--- a/Engine/Source/Runtime/Renderer/Private/ShadingEnergyConservationData.h
+++ b/Engine/Source/Runtime/Renderer/Private/ShadingEnergyConservationData.h
@@ -1,5 +1,7 @@
 // Copyright Epic Games, Inc. All Rights Reserved.
 
+#pragma once
+
 // These values were generated using FBuildShadingEnergyConservationTableCS
 
 // Specular directional albedo for GGX BSDF (2D)
