Add FarField flag to GPUScene InstanceSceneData
- Renamed PayloadDataFlags to InstanceFlags since it contains more than Payload related flags
- This is required to be able to identify instances in far field and apply the reference offset when building RT instance buffer in compute shader.

diff --git a/Engine/Shaders/Private/GPUScene/GPUSceneWriter.ush b/Engine/Shaders/Private/GPUScene/GPUSceneWriter.ush
--- a/Engine/Shaders/Private/GPUScene/GPUSceneWriter.ush
+++ b/Engine/Shaders/Private/GPUScene/GPUSceneWriter.ush
@@ -100,12 +100,11 @@ void WriteInstanceLocalToWorld(uint InstanceId, float4x4 LocalToWorld, bool Stor
  	uint InstanceId,
  	uint PrimitiveId,
  	uint InstanceRelativeId,
- 	uint PayloadDataFlags,
+ 	uint InstanceFlags, // combination of INSTANCE_SCENE_DATA_FLAG_
  	uint CustomDataCount,
  	float RandomId,
  	float4x4 LocalToWorld)
  {
- 	uint InstanceFlags = PayloadDataFlags;
  	SetInstanceDeterminantSignFlag(determinant((float3x3)LocalToWorld), InstanceFlags);
  	
  	float4 Element0 =
diff --git a/Engine/Shaders/Private/SceneData.ush b/Engine/Shaders/Private/SceneData.ush
--- a/Engine/Shaders/Private/SceneData.ush
+++ b/Engine/Shaders/Private/SceneData.ush
@@ -96,6 +96,7 @@
 #define INSTANCE_SCENE_DATA_FLAG_HAS_HIERARCHY_OFFSET		0x20
 #define INSTANCE_SCENE_DATA_FLAG_HAS_LOCAL_BOUNDS			0x40
 #define INSTANCE_SCENE_DATA_FLAG_HAS_EDITOR_DATA			0x80
+#define INSTANCE_SCENE_DATA_FLAG_IS_RAYTRACING_FAR_FIELD	0x100
 
 #if INSTANCE_SCENE_DATA_COMPRESSED_TRANSFORMS
 	static const uint InstanceTransformSizeFloat4Count = 2u; // compressed transform
diff --git a/Engine/Source/Runtime/Engine/Private/InstanceUniformShaderParameters.cpp b/Engine/Source/Runtime/Engine/Private/InstanceUniformShaderParameters.cpp
--- a/Engine/Source/Runtime/Engine/Private/InstanceUniformShaderParameters.cpp
+++ b/Engine/Source/Runtime/Engine/Private/InstanceUniformShaderParameters.cpp
@@ -8,20 +8,20 @@ void FInstanceSceneShaderData::Build
 (
 	uint32 PrimitiveId,
 	uint32 RelativeId,
-	uint32 PayloadDataFlags,
+	uint32 InstanceFlags,
 	uint32 LastUpdateFrame,
 	uint32 CustomDataCount,
 	float RandomID
 )
 {
-	BuildInternal(PrimitiveId, RelativeId, PayloadDataFlags, LastUpdateFrame, CustomDataCount, RandomID, FRenderTransform::Identity, FRenderTransform::Identity);
+	BuildInternal(PrimitiveId, RelativeId, InstanceFlags, LastUpdateFrame, CustomDataCount, RandomID, FRenderTransform::Identity, FRenderTransform::Identity);
 }
 
 void FInstanceSceneShaderData::Build
 (
 	uint32 PrimitiveId,
 	uint32 RelativeId,
-	uint32 PayloadDataFlags,
+	uint32 InstanceFlags,
 	uint32 LastUpdateFrame,
 	uint32 CustomDataCount,
 	float RandomID,
@@ -37,14 +37,14 @@ void FInstanceSceneShaderData::Build
 	LocalToWorld.Orthogonalize();
 	PrevLocalToWorld.Orthogonalize(); // TODO: Temporary PrevVelocityHack
 
-	BuildInternal(PrimitiveId, RelativeId, PayloadDataFlags, LastUpdateFrame, CustomDataCount, RandomID, LocalToWorld, PrevLocalToWorld);
+	BuildInternal(PrimitiveId, RelativeId, InstanceFlags, LastUpdateFrame, CustomDataCount, RandomID, LocalToWorld, PrevLocalToWorld);
 }
 
 void FInstanceSceneShaderData::BuildInternal
 (
 	uint32 PrimitiveId,
 	uint32 RelativeId,
-	uint32 PayloadDataFlags,
+	uint32 InstanceFlags,
 	uint32 LastUpdateFrame,
 	uint32 CustomDataCount,
 	float RandomID,
@@ -54,7 +54,6 @@ void FInstanceSceneShaderData::BuildInternal
 {
 	// Note: layout must match GetInstanceData in SceneData.ush and InitializeInstanceSceneData in GPUSceneWriter.ush
 
-	uint32 InstanceFlags = PayloadDataFlags;
 	if (LocalToWorld.RotDeterminant() < 0.0f)
 	{
 		InstanceFlags |= INSTANCE_SCENE_DATA_FLAG_DETERMINANT_SIGN;
diff --git a/Engine/Source/Runtime/Engine/Public/InstanceUniformShaderParameters.h b/Engine/Source/Runtime/Engine/Public/InstanceUniformShaderParameters.h
--- a/Engine/Source/Runtime/Engine/Public/InstanceUniformShaderParameters.h
+++ b/Engine/Source/Runtime/Engine/Public/InstanceUniformShaderParameters.h
@@ -20,6 +20,17 @@
 #define INSTANCE_SCENE_DATA_FLAG_HAS_HIERARCHY_OFFSET		0x20
 #define INSTANCE_SCENE_DATA_FLAG_HAS_LOCAL_BOUNDS			0x40
 #define INSTANCE_SCENE_DATA_FLAG_HAS_EDITOR_DATA			0x80
+#define INSTANCE_SCENE_DATA_FLAG_IS_RAYTRACING_FAR_FIELD	0x100
+
+#define INSTANCE_SCENE_DATA_FLAG_PAYLOAD_MASK ( \
+	INSTANCE_SCENE_DATA_FLAG_HAS_RANDOM \
+	| INSTANCE_SCENE_DATA_FLAG_HAS_CUSTOM_DATA \
+	| INSTANCE_SCENE_DATA_FLAG_HAS_DYNAMIC_DATA \
+	| INSTANCE_SCENE_DATA_FLAG_HAS_LIGHTSHADOW_UV_BIAS \
+	| INSTANCE_SCENE_DATA_FLAG_HAS_HIERARCHY_OFFSET \
+	| INSTANCE_SCENE_DATA_FLAG_HAS_LOCAL_BOUNDS \
+	| INSTANCE_SCENE_DATA_FLAG_HAS_EDITOR_DATA \
+)
 
 #define INVALID_PRIMITIVE_ID 0x000FFFFFu
 
@@ -101,7 +112,7 @@ struct FInstanceSceneShaderData
 	(
 		uint32 PrimitiveId,
 		uint32 RelativeId,
-		uint32 PayloadDataFlags,
+		uint32 InstanceFlags,
 		uint32 LastUpdateFrame,
 		uint32 CustomDataCount,
 		float RandomID
@@ -111,7 +122,7 @@ struct FInstanceSceneShaderData
 	(
 		uint32 PrimitiveId,
 		uint32 RelativeId,
-		uint32 PayloadDataFlags,
+		uint32 InstanceFlags,
 		uint32 LastUpdateFrame,
 		uint32 CustomDataCount,
 		float RandomID,
@@ -124,7 +135,7 @@ struct FInstanceSceneShaderData
 	(
 		uint32 PrimitiveId,
 		uint32 RelativeId,
-		uint32 PayloadDataFlags,
+		uint32 InstanceFlags,
 		uint32 LastUpdateFrame,
 		uint32 CustomDataCount,
 		float RandomID,
diff --git a/Engine/Source/Runtime/Engine/Public/PrimitiveSceneProxy.h b/Engine/Source/Runtime/Engine/Public/PrimitiveSceneProxy.h
--- a/Engine/Source/Runtime/Engine/Public/PrimitiveSceneProxy.h
+++ b/Engine/Source/Runtime/Engine/Public/PrimitiveSceneProxy.h
@@ -663,20 +663,21 @@ public:
 			bHasPerInstanceHierarchyOffset;
 	}
 
-	inline uint32 GetPayloadDataFlags()
+	inline uint32 GetInstanceSceneDataFlags()
 	{
-		uint32 PayloadDataFlags = 0x0;
-		PayloadDataFlags |= HasPerInstanceRandom()          ? INSTANCE_SCENE_DATA_FLAG_HAS_RANDOM				: 0u;
-		PayloadDataFlags |= HasPerInstanceCustomData()      ? INSTANCE_SCENE_DATA_FLAG_HAS_CUSTOM_DATA			: 0u;
-		PayloadDataFlags |= HasPerInstanceDynamicData()     ? INSTANCE_SCENE_DATA_FLAG_HAS_DYNAMIC_DATA			: 0u;
-		PayloadDataFlags |= HasPerInstanceLMSMUVBias()      ? INSTANCE_SCENE_DATA_FLAG_HAS_LIGHTSHADOW_UV_BIAS	: 0u;
-		PayloadDataFlags |= HasPerInstanceHierarchyOffset() ? INSTANCE_SCENE_DATA_FLAG_HAS_HIERARCHY_OFFSET		: 0u;
-		PayloadDataFlags |= HasPerInstanceLocalBounds()     ? INSTANCE_SCENE_DATA_FLAG_HAS_LOCAL_BOUNDS			: 0u;
+		uint32 Flags = 0x0;
+		Flags |= HasPerInstanceRandom()          ? INSTANCE_SCENE_DATA_FLAG_HAS_RANDOM              : 0u;
+		Flags |= HasPerInstanceCustomData()      ? INSTANCE_SCENE_DATA_FLAG_HAS_CUSTOM_DATA         : 0u;
+		Flags |= HasPerInstanceDynamicData()     ? INSTANCE_SCENE_DATA_FLAG_HAS_DYNAMIC_DATA        : 0u;
+		Flags |= HasPerInstanceLMSMUVBias()      ? INSTANCE_SCENE_DATA_FLAG_HAS_LIGHTSHADOW_UV_BIAS : 0u;
+		Flags |= HasPerInstanceHierarchyOffset() ? INSTANCE_SCENE_DATA_FLAG_HAS_HIERARCHY_OFFSET    : 0u;
+		Flags |= HasPerInstanceLocalBounds()     ? INSTANCE_SCENE_DATA_FLAG_HAS_LOCAL_BOUNDS        : 0u;
 #if WITH_EDITOR
-		PayloadDataFlags |= HasPerInstanceEditorData()      ? INSTANCE_SCENE_DATA_FLAG_HAS_EDITOR_DATA			: 0u;
+		Flags |= HasPerInstanceEditorData()      ? INSTANCE_SCENE_DATA_FLAG_HAS_EDITOR_DATA         : 0u;
 #endif
+		Flags |= IsRayTracingFarField()          ? INSTANCE_SCENE_DATA_FLAG_IS_RAYTRACING_FAR_FIELD : 0u;
 
-		return PayloadDataFlags;
+		return Flags;
 	}
 
 	inline bool UseEditorCompositing(const FSceneView* View) const { return GIsEditor && bUseEditorCompositing && !View->bIsGameView; }
diff --git a/Engine/Source/Runtime/Renderer/Private/GPUScene.cpp b/Engine/Source/Runtime/Renderer/Private/GPUScene.cpp
--- a/Engine/Source/Runtime/Renderer/Private/GPUScene.cpp
+++ b/Engine/Source/Runtime/Renderer/Private/GPUScene.cpp
@@ -307,7 +307,7 @@ struct FInstanceUploadInfo
 	FPrimitiveInstance DummyInstance;
 	FRenderBounds DummyLocalBounds;
 
-	uint32 PayloadDataFlags = 0x0;
+	uint32 InstanceFlags = 0x0;
 
 	FRenderTransform PrimitiveToWorld;
 	FRenderTransform PrevPrimitiveToWorld;
@@ -318,14 +318,14 @@ struct FInstanceUploadInfo
 void ValidateInstanceUploadInfo(const FInstanceUploadInfo& UploadInfo, const FGPUSceneBufferState& BufferState)
 {
 #if DO_CHECK
-	const bool bHasRandomID			= (UploadInfo.PayloadDataFlags & INSTANCE_SCENE_DATA_FLAG_HAS_RANDOM) != 0u;
-	const bool bHasCustomData		= (UploadInfo.PayloadDataFlags & INSTANCE_SCENE_DATA_FLAG_HAS_CUSTOM_DATA) != 0u;
-	const bool bHasDynamicData		= (UploadInfo.PayloadDataFlags & INSTANCE_SCENE_DATA_FLAG_HAS_DYNAMIC_DATA) != 0u;
-	const bool bHasLightShadowUVBias = (UploadInfo.PayloadDataFlags & INSTANCE_SCENE_DATA_FLAG_HAS_LIGHTSHADOW_UV_BIAS) != 0u;
-	const bool bHasHierarchyOffset	= (UploadInfo.PayloadDataFlags & INSTANCE_SCENE_DATA_FLAG_HAS_HIERARCHY_OFFSET) != 0u;
-	const bool bHasLocalBounds		= (UploadInfo.PayloadDataFlags & INSTANCE_SCENE_DATA_FLAG_HAS_LOCAL_BOUNDS) != 0u;
+	const bool bHasRandomID			= (UploadInfo.InstanceFlags & INSTANCE_SCENE_DATA_FLAG_HAS_RANDOM) != 0u;
+	const bool bHasCustomData		= (UploadInfo.InstanceFlags & INSTANCE_SCENE_DATA_FLAG_HAS_CUSTOM_DATA) != 0u;
+	const bool bHasDynamicData		= (UploadInfo.InstanceFlags & INSTANCE_SCENE_DATA_FLAG_HAS_DYNAMIC_DATA) != 0u;
+	const bool bHasLightShadowUVBias = (UploadInfo.InstanceFlags & INSTANCE_SCENE_DATA_FLAG_HAS_LIGHTSHADOW_UV_BIAS) != 0u;
+	const bool bHasHierarchyOffset	= (UploadInfo.InstanceFlags & INSTANCE_SCENE_DATA_FLAG_HAS_HIERARCHY_OFFSET) != 0u;
+	const bool bHasLocalBounds		= (UploadInfo.InstanceFlags & INSTANCE_SCENE_DATA_FLAG_HAS_LOCAL_BOUNDS) != 0u;
 #if WITH_EDITOR
-	const bool bHasEditorData		= (UploadInfo.PayloadDataFlags & INSTANCE_SCENE_DATA_FLAG_HAS_EDITOR_DATA) != 0u;
+	const bool bHasEditorData		= (UploadInfo.InstanceFlags & INSTANCE_SCENE_DATA_FLAG_HAS_EDITOR_DATA) != 0u;
 #endif
 
 	const int32 InstanceCount = UploadInfo.PrimitiveInstances.Num();
@@ -484,7 +484,7 @@ struct FUploadDataSourceAdapterScenePrimitives
 			InstanceUploadInfo.PrevPrimitiveToWorld = FLargeWorldRenderScalar::MakeClampedToRelativeWorldMatrix(AbsoluteOrigin.GetTileOffset(), PreviousLocalToWorld);;
 		}
 
-		InstanceUploadInfo.PayloadDataFlags = PrimitiveSceneProxy->GetPayloadDataFlags();
+		InstanceUploadInfo.InstanceFlags = PrimitiveSceneProxy->GetInstanceSceneDataFlags();
 		InstanceUploadInfo.InstanceLocalBounds = PrimitiveSceneProxy->GetInstanceLocalBounds();
 		if (InstanceUploadInfo.InstanceLocalBounds.Num() == 0)
 		{
@@ -507,7 +507,7 @@ struct FUploadDataSourceAdapterScenePrimitives
 		}
 		else
 		{
-			check(InstanceUploadInfo.PayloadDataFlags == 0x0);
+			checkf((InstanceUploadInfo.InstanceFlags & INSTANCE_SCENE_DATA_FLAG_PAYLOAD_MASK) == 0x0, TEXT("Proxy must support instance data buffer to use payload data"));
 			check(InstanceUploadInfo.InstancePayloadDataOffset == INDEX_NONE && InstanceUploadInfo.InstancePayloadDataStride == 0);
 
 			// We always create an instance to ensure that we can always use the same code paths in the shader
@@ -1299,13 +1299,13 @@ void FGPUScene::UploadGeneral(FRHICommandListImmediate& RHICmdList, FScene *Scen
 								const FPrimitiveInstance& SceneData = UploadInfo.PrimitiveInstances[InstanceIndex];
 
 								// Directly embedded in instance scene data
-								const float RandomID = (UploadInfo.PayloadDataFlags & INSTANCE_SCENE_DATA_FLAG_HAS_RANDOM) ? UploadInfo.InstanceRandomID[InstanceIndex] : 0.0f;
+								const float RandomID = (UploadInfo.InstanceFlags & INSTANCE_SCENE_DATA_FLAG_HAS_RANDOM) ? UploadInfo.InstanceRandomID[InstanceIndex] : 0.0f;
 
 								FInstanceSceneShaderData InstanceSceneData;
 								InstanceSceneData.Build(
 									UploadInfo.PrimitiveID,
 									InstanceIndex,
-									UploadInfo.PayloadDataFlags,
+									UploadInfo.InstanceFlags,
 									UploadInfo.LastUpdateSceneFrameNumber,
 									UploadInfo.InstanceCustomDataCount,
 									RandomID,
@@ -1334,16 +1334,16 @@ void FGPUScene::UploadGeneral(FRHICommandListImmediate& RHICmdList, FScene *Scen
 
 									int32 PayloadPosition = 0;
 
-									if (UploadInfo.PayloadDataFlags & (INSTANCE_SCENE_DATA_FLAG_HAS_HIERARCHY_OFFSET | INSTANCE_SCENE_DATA_FLAG_HAS_LOCAL_BOUNDS | INSTANCE_SCENE_DATA_FLAG_HAS_EDITOR_DATA))
+									if (UploadInfo.InstanceFlags & (INSTANCE_SCENE_DATA_FLAG_HAS_HIERARCHY_OFFSET | INSTANCE_SCENE_DATA_FLAG_HAS_LOCAL_BOUNDS | INSTANCE_SCENE_DATA_FLAG_HAS_EDITOR_DATA))
 									{
-										const uint32 InstanceHierarchyOffset = (UploadInfo.PayloadDataFlags & INSTANCE_SCENE_DATA_FLAG_HAS_HIERARCHY_OFFSET) ? UploadInfo.InstanceHierarchyOffset[InstanceIndex] : 0;
+										const uint32 InstanceHierarchyOffset = (UploadInfo.InstanceFlags & INSTANCE_SCENE_DATA_FLAG_HAS_HIERARCHY_OFFSET) ? UploadInfo.InstanceHierarchyOffset[InstanceIndex] : 0;
 										InstancePayloadData[PayloadPosition].X = *(const float*)&InstanceHierarchyOffset;
 
 #if WITH_EDITOR
-										const uint32 InstanceEditorData = (UploadInfo.PayloadDataFlags & INSTANCE_SCENE_DATA_FLAG_HAS_EDITOR_DATA) ? UploadInfo.InstanceEditorData[InstanceIndex] : 0;
+										const uint32 InstanceEditorData = (UploadInfo.InstanceFlags & INSTANCE_SCENE_DATA_FLAG_HAS_EDITOR_DATA) ? UploadInfo.InstanceEditorData[InstanceIndex] : 0;
 										InstancePayloadData[PayloadPosition].Y = *(const float*)&InstanceEditorData;
 #endif
-										if (UploadInfo.PayloadDataFlags & INSTANCE_SCENE_DATA_FLAG_HAS_LOCAL_BOUNDS)
+										if (UploadInfo.InstanceFlags & INSTANCE_SCENE_DATA_FLAG_HAS_LOCAL_BOUNDS)
 										{
 											check(UploadInfo.InstanceLocalBounds.Num() == UploadInfo.PrimitiveInstances.Num());
 											const FRenderBounds& InstanceLocalBounds = UploadInfo.InstanceLocalBounds[InstanceIndex];
@@ -1359,10 +1359,10 @@ void FGPUScene::UploadGeneral(FRHICommandListImmediate& RHICmdList, FScene *Scen
 											InstancePayloadData[PayloadPosition + 1].W = *(const float*)&BoundsExtent.Z;
 										}
 
-										PayloadPosition += (UploadInfo.PayloadDataFlags & INSTANCE_SCENE_DATA_FLAG_HAS_LOCAL_BOUNDS) ? 2 : 1;
+										PayloadPosition += (UploadInfo.InstanceFlags & INSTANCE_SCENE_DATA_FLAG_HAS_LOCAL_BOUNDS) ? 2 : 1;
 									}
 
-									if (UploadInfo.PayloadDataFlags & INSTANCE_SCENE_DATA_FLAG_HAS_DYNAMIC_DATA)
+									if (UploadInfo.InstanceFlags & INSTANCE_SCENE_DATA_FLAG_HAS_DYNAMIC_DATA)
 									{
 										check(UploadInfo.InstanceDynamicData.Num() == UploadInfo.PrimitiveInstances.Num());
 										const FRenderTransform PrevLocalToWorld = UploadInfo.InstanceDynamicData[InstanceIndex].ComputePrevLocalToWorld(UploadInfo.PrevPrimitiveToWorld);
@@ -1380,14 +1380,14 @@ void FGPUScene::UploadGeneral(FRHICommandListImmediate& RHICmdList, FScene *Scen
 									#endif
 									}
 
-									if (UploadInfo.PayloadDataFlags & INSTANCE_SCENE_DATA_FLAG_HAS_LIGHTSHADOW_UV_BIAS)
+									if (UploadInfo.InstanceFlags & INSTANCE_SCENE_DATA_FLAG_HAS_LIGHTSHADOW_UV_BIAS)
 									{
 										check(UploadInfo.InstanceLightShadowUVBias.Num() == UploadInfo.PrimitiveInstances.Num());
 										InstancePayloadData[PayloadPosition] = UploadInfo.InstanceLightShadowUVBias[InstanceIndex];
 										PayloadPosition += 1;
 									}
 
-									if (UploadInfo.PayloadDataFlags & INSTANCE_SCENE_DATA_FLAG_HAS_CUSTOM_DATA)
+									if (UploadInfo.InstanceFlags & INSTANCE_SCENE_DATA_FLAG_HAS_CUSTOM_DATA)
 									{
 										check(PayloadPosition + (UploadInfo.InstanceCustomDataCount >> 2u) <= InstancePayloadData.Num());
 										const float* CustomDataGetPtr = &UploadInfo.InstanceCustomData[(InstanceIndex * UploadInfo.InstanceCustomDataCount)];
@@ -1603,7 +1603,7 @@ struct FUploadDataSourceAdapterDynamicPrimitives
 			InstanceUploadInfo.InstancePayloadDataOffset	= PrimData.LocalPayloadDataOffset == INDEX_NONE ? INDEX_NONE : PayloadStartOffset + PrimData.LocalPayloadDataOffset;
 			InstanceUploadInfo.InstancePayloadDataStride	= PrimData.SourceData.GetPayloadFloat4Stride();
 			InstanceUploadInfo.InstanceCustomDataCount		= PrimData.SourceData.NumInstanceCustomDataFloats;
-			InstanceUploadInfo.PayloadDataFlags				= PrimData.SourceData.PayloadDataFlags;
+			InstanceUploadInfo.InstanceFlags                = PrimData.SourceData.PayloadDataFlags;
 			InstanceUploadInfo.PrimitiveInstances 			= PrimData.SourceData.InstanceSceneData;
 			InstanceUploadInfo.InstanceDynamicData 			= PrimData.SourceData.InstanceDynamicData;
 			InstanceUploadInfo.InstanceCustomData 			= PrimData.SourceData.InstanceCustomData;
