UE5:  Fix for occlusion query startup assert when using the mobile renderer. Reverting occlusion query fence implementation to original code for all use cases besides DisplayCluster, the case we needed to optimize for multiple view family rendering (the optimization saves 8 ms in the Virtual Production City demo).  Longer term, we might want a more robust refactor of the code to handle both cases, but we need to unblock mobile in the short term.

#rnx

diff --git a/Engine/Plugins/Runtime/nDisplay/Source/DisplayCluster/Private/Game/EngineClasses/Basics/DisplayClusterViewportClient.cpp b/Engine/Plugins/Runtime/nDisplay/Source/DisplayCluster/Private/Game/EngineClasses/Basics/DisplayClusterViewportClient.cpp
--- a/Engine/Plugins/Runtime/nDisplay/Source/DisplayCluster/Private/Game/EngineClasses/Basics/DisplayClusterViewportClient.cpp
+++ b/Engine/Plugins/Runtime/nDisplay/Source/DisplayCluster/Private/Game/EngineClasses/Basics/DisplayClusterViewportClient.cpp
@@ -744,6 +744,7 @@ void UDisplayClusterViewportClient::Draw(FViewport* InViewport, FCanvas* SceneCa
 				if (!bDisableWorldRendering && PlayerViewMap.Num() > 0 && FSlateApplication::Get().GetPlatformApplication()->IsAllowedToRender()) //-V560
 				{
 					ViewFamily.bIsRenderedImmediatelyAfterAnotherViewFamily = bIsRenderedImmediatelyAfterAnotherViewFamily;
+					ViewFamily.bIsMultipleViewFamily = true;
 
 					GetRendererModule().BeginRenderingViewFamily(SceneCanvas, &ViewFamily);
 
diff --git a/Engine/Source/Runtime/Engine/Public/SceneView.h b/Engine/Source/Runtime/Engine/Public/SceneView.h
--- a/Engine/Source/Runtime/Engine/Public/SceneView.h
+++ b/Engine/Source/Runtime/Engine/Public/SceneView.h
@@ -1716,6 +1716,9 @@ public:
 	/** Whether this view is rendered in BeginRenderingViewFamily() immediately after another render. */
 	bool bIsRenderedImmediatelyAfterAnotherViewFamily = false;
 
+	/** Whether this view is one of multiple view families rendered in a single frame.  Affects occlusion query synchronization logic. */
+	bool bIsMultipleViewFamily = false;
+
 	/** 
 	 * Which component of the scene rendering should be output to the final render target.
 	 * If SCS_FinalColorLDR this indicates do nothing.
diff --git a/Engine/Source/Runtime/Renderer/Private/SceneOcclusion.cpp b/Engine/Source/Runtime/Renderer/Private/SceneOcclusion.cpp
--- a/Engine/Source/Runtime/Renderer/Private/SceneOcclusion.cpp
+++ b/Engine/Source/Runtime/Renderer/Private/SceneOcclusion.cpp
@@ -1473,14 +1473,29 @@ void FSceneRenderer::FenceOcclusionTestsInternal(FRHICommandListImmediate& RHICm
 {
 	SCOPE_CYCLE_COUNTER(STAT_OcclusionSubmittedFence_Dispatch);
 
-	// Push a new fence into the queue of buffered fences.  We assume that the queue isn't full, since WaitOcclusionTests
-	// (called earlier in the frame) will always pop at least one fence if the queue is full.
-	check(OcclusionSubmittedFence[FOcclusionQueryHelpers::MaxBufferedOcclusionFrames - 1].Fence == nullptr);
+	if (ViewFamily.bIsMultipleViewFamily)
+	{
+		// If there are multiple view families, we implement a queue of buffered fences, so we can avoid waiting on queries for
+		// another view family that's rendering in the same frame.  Here we push a new fence into the queue of buffered fences.
+		// We assume that the queue isn't full, since WaitOcclusionTests (called earlier in the frame) will always pop at least
+		// one fence if the queue is full.
+		check(OcclusionSubmittedFence[FOcclusionQueryHelpers::MaxBufferedOcclusionFrames - 1].Fence == nullptr);
 
-	for (int32 Dest = FOcclusionQueryHelpers::MaxBufferedOcclusionFrames - 1; Dest >= 1; Dest--)
+		for (int32 Dest = FOcclusionQueryHelpers::MaxBufferedOcclusionFrames - 1; Dest >= 1; Dest--)
+		{
+			CA_SUPPRESS(6385);
+			OcclusionSubmittedFence[Dest] = OcclusionSubmittedFence[Dest - 1];
+		}
+	}
+	else
 	{
-		CA_SUPPRESS(6385);
-		OcclusionSubmittedFence[Dest] = OcclusionSubmittedFence[Dest - 1];
+		// Single view family implementation, fixed number of buffered frames.
+		int32 NumFrames = FOcclusionQueryHelpers::GetNumBufferedFrames(FeatureLevel);
+		for (int32 Dest = NumFrames - 1; Dest >= 1; Dest--)
+		{
+			CA_SUPPRESS(6385);
+			OcclusionSubmittedFence[Dest] = OcclusionSubmittedFence[Dest - 1];
+		}
 	}
 
 	OcclusionSubmittedFence[0].Fence = RHICmdList.RHIThreadFence();
@@ -1507,48 +1522,59 @@ void FSceneRenderer::WaitOcclusionTests(FRHICommandListImmediate& RHICmdList)
 	{
 		SCOPE_CYCLE_COUNTER(STAT_OcclusionSubmittedFence_Wait);
 
-		// Count the number of fences in the queue for the view in question
-		const uint32 ViewStateUniqueID = GetViewStateUniqueID(this);
-
-		int32 ViewStateFenceCount = 0;
-		for (int32 FenceIndex = 0; FenceIndex < FOcclusionQueryHelpers::MaxBufferedOcclusionFrames && OcclusionSubmittedFence[FenceIndex].Fence; FenceIndex++)
+		if (ViewFamily.bIsMultipleViewFamily)
 		{
-			if (OcclusionSubmittedFence[FenceIndex].Fence && (OcclusionSubmittedFence[FenceIndex].ViewStateUniqueID == ViewStateUniqueID))
+			// Count the number of fences in the queue for the view in question
+			const uint32 ViewStateUniqueID = GetViewStateUniqueID(this);
+
+			int32 ViewStateFenceCount = 0;
+			for (int32 FenceIndex = 0; FenceIndex < FOcclusionQueryHelpers::MaxBufferedOcclusionFrames && OcclusionSubmittedFence[FenceIndex].Fence; FenceIndex++)
 			{
-				ViewStateFenceCount++;
+				if (OcclusionSubmittedFence[FenceIndex].Fence && (OcclusionSubmittedFence[FenceIndex].ViewStateUniqueID == ViewStateUniqueID))
+				{
+					ViewStateFenceCount++;
+				}
 			}
-		}
 
-		// Figure out how many fences are allowed to remain in the queue after the wait, which is one less than the number of buffered frames
-		const int32 FencesAllowedInQueue = FOcclusionQueryHelpers::GetNumBufferedFrames(FeatureLevel) - 1;
-		check(FencesAllowedInQueue >= 0);
+			// Figure out how many fences are allowed to remain in the queue after the wait, which is one less than the number of buffered frames
+			const int32 FencesAllowedInQueue = FOcclusionQueryHelpers::GetNumBufferedFrames(FeatureLevel) - 1;
+			check(FencesAllowedInQueue >= 0);
 
-		// Scan the fences in the queue in order (oldest will be at the end), and wait on them as appropriate.  If the queue is full,
-		// we always want to wait on at least one fence, so we have room to push a new fence onto the queue if necessary.  Thus the
-		// loop always runs at least once before breaking out.
-		int32 FenceToWaitOn = FOcclusionQueryHelpers::MaxBufferedOcclusionFrames - 1;
-		do
-		{
-			// If our fence counting logic is correct, we should already have broken out of the loop
-			check(FenceToWaitOn >= 0);
-
-			// Is there a fence in this queue entry?
-			if (OcclusionSubmittedFence[FenceToWaitOn].Fence)
+			// Scan the fences in the queue in order (oldest will be at the end), and wait on them as appropriate.  If the queue is full,
+			// we always want to wait on at least one fence, so we have room to push a new fence onto the queue if necessary.  Thus the
+			// loop always runs at least once before breaking out.
+			int32 FenceToWaitOn = FOcclusionQueryHelpers::MaxBufferedOcclusionFrames - 1;
+			do
 			{
-				// If this is one of the fences for the current scene's view state, mark that we're waiting on it
-				if (OcclusionSubmittedFence[FenceToWaitOn].ViewStateUniqueID == ViewStateUniqueID)
+				// If our fence counting logic is correct, we should already have broken out of the loop
+				check(FenceToWaitOn >= 0);
+
+				// Is there a fence in this queue entry?
+				if (OcclusionSubmittedFence[FenceToWaitOn].Fence)
 				{
-					ViewStateFenceCount--;
-				}
+					// If this is one of the fences for the current scene's view state, mark that we're waiting on it
+					if (OcclusionSubmittedFence[FenceToWaitOn].ViewStateUniqueID == ViewStateUniqueID)
+					{
+						ViewStateFenceCount--;
+					}
 
-				// Wait on the current fence and clear it out
-				FRHICommandListExecutor::WaitOnRHIThreadFence(OcclusionSubmittedFence[FenceToWaitOn].Fence);
-				OcclusionSubmittedFence[FenceToWaitOn].Fence = nullptr;
-				OcclusionSubmittedFence[FenceToWaitOn].ViewStateUniqueID = 0;
-			}
+					// Wait on the current fence and clear it out
+					FRHICommandListExecutor::WaitOnRHIThreadFence(OcclusionSubmittedFence[FenceToWaitOn].Fence);
+					OcclusionSubmittedFence[FenceToWaitOn].Fence = nullptr;
+					OcclusionSubmittedFence[FenceToWaitOn].ViewStateUniqueID = 0;
+				}
 
-			FenceToWaitOn--;
+				FenceToWaitOn--;
 
-		} while (ViewStateFenceCount > FencesAllowedInQueue);
+			} while (ViewStateFenceCount > FencesAllowedInQueue);
+		}
+		else
+		{
+			// Single view family implementation, fixed number of buffered frames.
+			int32 BlockFrame = FOcclusionQueryHelpers::GetNumBufferedFrames(FeatureLevel) - 1;
+			FRHICommandListExecutor::WaitOnRHIThreadFence(OcclusionSubmittedFence[BlockFrame].Fence);
+			OcclusionSubmittedFence[BlockFrame].Fence = nullptr;
+			OcclusionSubmittedFence[BlockFrame].ViewStateUniqueID = 0;
+		}
 	}
 }
