Disable depth pass velocity and fall back to velocity after base pass when MSAA is enabled.
Velocity buffer isn't MSAA'd so can't have a combined depth/velocity pass.
Other than performance (which we can't get if we don't combine depth/velocity) the only reason that we might still want to keep early velocity is to reproject normals for DBuffer decals. But with forward MSAA we don't have a normal buffer available anyway.
#rnx

diff --git a/Engine/Source/Runtime/Renderer/Private/CompositionLighting/PostProcessDeferredDecals.cpp b/Engine/Source/Runtime/Renderer/Private/CompositionLighting/PostProcessDeferredDecals.cpp
--- a/Engine/Source/Runtime/Renderer/Private/CompositionLighting/PostProcessDeferredDecals.cpp
+++ b/Engine/Source/Runtime/Renderer/Private/CompositionLighting/PostProcessDeferredDecals.cpp
@@ -165,9 +165,9 @@ void GetDeferredDecalPassParameters(
 
 TUniformBufferRef<FDeferredDecalUniformParameters> CreateDeferredDecalUniformBuffer(const FViewInfo& View)
 {
-	bool bIsMotionInDepth = FVelocityRendering::DepthPassCanOutputVelocity();
+	const bool bIsMotionInDepth = FVelocityRendering::DepthPassCanOutputVelocity(View.GetFeatureLevel());
 	// if we have early motion vectors (bIsMotionInDepth) and the cvar is enabled and we actually have a buffer from the previous frame (View.PrevViewInfo.GBufferA.IsValid())
-	bool bIsNormalReprojectionEnabled = (bIsMotionInDepth && CVarDBufferDecalNormalReprojectionEnabled.GetValueOnRenderThread() && View.PrevViewInfo.GBufferA.IsValid());
+	const bool bIsNormalReprojectionEnabled = (bIsMotionInDepth && CVarDBufferDecalNormalReprojectionEnabled.GetValueOnRenderThread() && View.PrevViewInfo.GBufferA.IsValid());
 
 	FDeferredDecalUniformParameters UniformParameters;
 	UniformParameters.NormalReprojectionThresholdLow  = CVarDBufferDecalNormalReprojectionThresholdLow .GetValueOnRenderThread();
@@ -562,16 +562,18 @@ void AddDeferredDecalPass(
 void ExtractNormalsForNextFrameReprojection(FRDGBuilder& GraphBuilder, const FSceneTextures& SceneTextures, const TArray<FViewInfo>& Views)
 {
 	// save the previous frame if early motion vectors are enabled and normal reprojection is enabled, so there should be no cost if these options are off
-	bool bApplyReproject = FVelocityRendering::DepthPassCanOutputVelocity() && CVarDBufferDecalNormalReprojectionEnabled.GetValueOnRenderThread();
+	const bool bIsNormalReprojectionEnabled = CVarDBufferDecalNormalReprojectionEnabled.GetValueOnRenderThread();
 
-	if (bApplyReproject)
+	if (bIsNormalReprojectionEnabled)
 	{
-		bool bAnyViewWriteable = false;
 		for (int32 Index = 0; Index < Views.Num(); Index++)
 		{
-			if (Views[Index].bStatePrevViewInfoIsReadOnly == false)
+			if (FVelocityRendering::DepthPassCanOutputVelocity(Views[Index].GetFeatureLevel()))
 			{
-				GraphBuilder.QueueTextureExtraction(SceneTextures.GBufferA, &Views[Index].ViewState->PrevFrameViewInfo.GBufferA);
+				if (Views[Index].bStatePrevViewInfoIsReadOnly == false)
+				{
+					GraphBuilder.QueueTextureExtraction(SceneTextures.GBufferA, &Views[Index].ViewState->PrevFrameViewInfo.GBufferA);
+				}
 			}
 		}
 	}
diff --git a/Engine/Source/Runtime/Renderer/Private/RendererScene.cpp b/Engine/Source/Runtime/Renderer/Private/RendererScene.cpp
--- a/Engine/Source/Runtime/Renderer/Private/RendererScene.cpp
+++ b/Engine/Source/Runtime/Renderer/Private/RendererScene.cpp
@@ -3776,8 +3776,9 @@ void FScene::UpdateEarlyZPassMode()
 		if (ShouldForceFullDepthPass(ShaderPlatform))
 		{
 			// DBuffer decals and stencil LOD dithering force a full prepass
-			EarlyZPassMode = FVelocityRendering::DepthPassCanOutputVelocity() ? DDM_AllOpaqueNoVelocity : DDM_AllOpaque;
-			bEarlyZPassMovable = FVelocityRendering::DepthPassCanOutputVelocity() ? false : true;
+			const bool bDepthPassCanOutputVelocity = FVelocityRendering::DepthPassCanOutputVelocity(FeatureLevel);
+			EarlyZPassMode = bDepthPassCanOutputVelocity ? DDM_AllOpaqueNoVelocity : DDM_AllOpaque;
+			bEarlyZPassMovable = bDepthPassCanOutputVelocity ? false : true;
 		}
 
 		if ((EarlyZPassMode == DDM_AllOpaque || EarlyZPassMode == DDM_AllOpaqueNoVelocity)
diff --git a/Engine/Source/Runtime/Renderer/Private/VelocityRendering.cpp b/Engine/Source/Runtime/Renderer/Private/VelocityRendering.cpp
--- a/Engine/Source/Runtime/Renderer/Private/VelocityRendering.cpp
+++ b/Engine/Source/Runtime/Renderer/Private/VelocityRendering.cpp
@@ -162,7 +162,7 @@ bool FDeferredShadingSceneRenderer::ShouldRenderVelocities() const
 	{
 		return false;
 	}
-	if (FVelocityRendering::DepthPassCanOutputVelocity())
+	if (FVelocityRendering::DepthPassCanOutputVelocity(Scene->GetFeatureLevel()))
 	{
 		// Always render velocity when it is part of the depth pass to avoid dropping things from the depth pass.
 		// This means that we will pay the cost of velocity in the pass even if we don't really need it according to the view logic below.
@@ -372,10 +372,11 @@ bool FVelocityRendering::IsVelocityPassSupported(EShaderPlatform ShaderPlatform)
 	return GPixelFormats[GetFormat(ShaderPlatform)].Supported;
 }
 
-bool FVelocityRendering::DepthPassCanOutputVelocity()
+bool FVelocityRendering::DepthPassCanOutputVelocity(ERHIFeatureLevel::Type FeatureLevel)
 {
-	static bool bIsMerged = CVarVelocityOutputPass.GetValueOnAnyThread() == 0;
-	return bIsMerged;
+	static bool bRequestedDepthPassVelocity = CVarVelocityOutputPass.GetValueOnAnyThread() == 0;
+	const bool bMSAAEnabled = GetDefaultMSAACount(FeatureLevel) > 1;
+	return !bMSAAEnabled && bRequestedDepthPassVelocity;
 }
 
 bool FVelocityRendering::BasePassCanOutputVelocity(EShaderPlatform ShaderPlatform)
diff --git a/Engine/Source/Runtime/Renderer/Private/VelocityRendering.h b/Engine/Source/Runtime/Renderer/Private/VelocityRendering.h
--- a/Engine/Source/Runtime/Renderer/Private/VelocityRendering.h
+++ b/Engine/Source/Runtime/Renderer/Private/VelocityRendering.h
@@ -42,7 +42,7 @@ struct FVelocityRendering
 	static bool IsVelocityPassSupported(EShaderPlatform ShaderPlatform);
 
 	/** Returns true if the velocity can be output during depth pass. */
-	static bool DepthPassCanOutputVelocity();
+	static bool DepthPassCanOutputVelocity(ERHIFeatureLevel::Type FeatureLevel);
 
 	/** Returns true if the velocity can be output in the BasePass. */
 	static bool BasePassCanOutputVelocity(EShaderPlatform ShaderPlatform);
