Fix incorrect inverse projection parameters used in thickness calculation for cube-shadows transmission.

diff --git a/Engine/Source/Runtime/Renderer/Private/ShadowRendering.h b/Engine/Source/Runtime/Renderer/Private/ShadowRendering.h
--- a/Engine/Source/Runtime/Renderer/Private/ShadowRendering.h
+++ b/Engine/Source/Runtime/Renderer/Private/ShadowRendering.h
@@ -1500,13 +1500,11 @@ public:
 
 		SetShaderValue(RHICmdList, ShaderRHI, ShadowFadeFraction, ShadowInfo->FadeAlphas[ViewIndex]);
 		SetShaderValue(RHICmdList, ShaderRHI, ShadowSharpen, LightProxy.GetShadowSharpen() * 7.0f + 1.0f);
-		//Near is always 1? // TODO: validate
-		float Near = 1;
-		float Far = LightProxy.GetRadius();
-		FVector2f param = FVector2f(Far / (Far - Near), -Near * Far / (Far - Near));
-		FVector2f projParam = FVector2f(1.0f / param.Y, param.X / param.Y);
+
+		FVector2d ProjectionParams = FVector2d(ShadowInfo->OnePassShadowFaceProjectionMatrix.M[2][2], ShadowInfo->OnePassShadowFaceProjectionMatrix.M[3][2]);
+		FVector2f InverseProjParams = FVector2f(1.0 / ProjectionParams.Y, ProjectionParams.X / ProjectionParams.Y);
 		SetShaderValue(RHICmdList, ShaderRHI, PointLightDepthBias, FVector3f(ShadowInfo->GetShaderDepthBias(), ShadowInfo->GetShaderSlopeDepthBias(), ShadowInfo->GetShaderMaxSlopeDepthBias()));
-		SetShaderValue(RHICmdList, ShaderRHI, PointLightProjParameters, FVector2f(projParam.X, projParam.Y));
+		SetShaderValue(RHICmdList, ShaderRHI, PointLightProjParameters, InverseProjParams);
 
 		if (HairStrandsParameters.IsBound())
 		{
