Fix wrong logic around bloom separate transluceny resource creation and an ensure test crashing some tests.

diff --git a/Engine/Source/Runtime/Renderer/Private/LightShaftRendering.cpp b/Engine/Source/Runtime/Renderer/Private/LightShaftRendering.cpp
--- a/Engine/Source/Runtime/Renderer/Private/LightShaftRendering.cpp
+++ b/Engine/Source/Runtime/Renderer/Private/LightShaftRendering.cpp
@@ -562,9 +562,9 @@ void FDeferredShadingSceneRenderer::RenderLightShaftBloom(
 					ERenderTargetLoadAction OutputLoadAction = ERenderTargetLoadAction::ELoad;
 					float OutputViewRectScale = 1.0f;
 
-					// Reset to point at separate translucency if enabled.
+					// Render to separate translucency buffer instead of scene color if requested.
 					const ELightShaftBloomOutput BloomOutput = GetLightShaftBloomOutput(ViewFamily);
-					bool bUpdateViewsTranslucency = false;
+					bool bUpdateViewsSeparateTranslucency = false;
 					if (BloomOutput == ELightShaftBloomOutput::SeparateTranslucency)
 					{
 						FTranslucencyPassResources& TranslucencyPassResources = OutTranslucencyResourceMap.Get(/* ViewIndex = */ 0, ETranslucencyPass::TPT_TranslucencyAfterDOF);
@@ -585,6 +585,9 @@ void FDeferredShadingSceneRenderer::RenderLightShaftBloom(
 								GraphBuilder, Desc,
 								TEXT("Translucency.LightShaftBloom"),
 								GFastVRamConfig.SeparateTranslucency);
+
+							// We will need to update views separate transluceny buffers if we have just created them.
+							bUpdateViewsSeparateTranslucency = true;
 						}
 						else
 						{
@@ -640,15 +643,14 @@ void FDeferredShadingSceneRenderer::RenderLightShaftBloom(
 							AddLightShaftBloomPass(GraphBuilder, View, LightShaftParameters, LightShafts, OutputViewport, OutputBinding);
 							OutputLoadAction = ERenderTargetLoadAction::ELoad;
 
-							if (bUpdateViewsTranslucency)
+							FTranslucencyPassResources& TranslucencyPassResources = OutTranslucencyResourceMap.Get(ViewIndex, ETranslucencyPass::TPT_TranslucencyAfterDOF);
+							if (bUpdateViewsSeparateTranslucency)
 							{
-								FTranslucencyPassResources& TranslucencyPassResources = OutTranslucencyResourceMap.Get(ViewIndex, ETranslucencyPass::TPT_TranslucencyAfterDOF);
 								TranslucencyPassResources.ViewRect = OutputViewport.Rect;
 								TranslucencyPassResources.ColorTexture = OutputTexture;
 							}
-							else
+							else if(TranslucencyPassResources.IsValid())
 							{
-								FTranslucencyPassResources& TranslucencyPassResources = OutTranslucencyResourceMap.Get(ViewIndex, ETranslucencyPass::TPT_TranslucencyAfterDOF);
 								ensure(TranslucencyPassResources.ViewRect == OutputViewport.Rect);
 							}
 						}
