Translucent post dof texture creation is now using shared code for translucent and light shaft rendeirng.

diff --git a/Engine/Source/Runtime/Renderer/Private/LightShaftRendering.cpp b/Engine/Source/Runtime/Renderer/Private/LightShaftRendering.cpp
--- a/Engine/Source/Runtime/Renderer/Private/LightShaftRendering.cpp
+++ b/Engine/Source/Runtime/Renderer/Private/LightShaftRendering.cpp
@@ -573,18 +573,8 @@ void FDeferredShadingSceneRenderer::RenderLightShaftBloom(
 						{
 							OutputLoadAction = ERenderTargetLoadAction::EClear;
 
-							const FRDGTextureDesc Desc = FRDGTextureDesc::Create2D(
-								SeparateTranslucencyDimensions.Extent,
-								PF_FloatRGBA,
-								FClearValueBinding::Black,
-								TexCreate_RenderTargetable | TexCreate_ShaderResource,
-								1,
-								SeparateTranslucencyDimensions.NumSamples);
-
-							OutputTexture = CreateTextureMSAA(
-								GraphBuilder, Desc,
-								TEXT("Translucency.LightShaftBloom"),
-								GFastVRamConfig.SeparateTranslucency);
+							const bool bIsModulate = false;
+							OutputTexture = CreatePostDOFTranslucentTexture(GraphBuilder, ETranslucencyPass::TPT_TranslucencyAfterDOF, SeparateTranslucencyDimensions, bIsModulate);
 
 							// We will need to update views separate transluceny buffers if we have just created them.
 							bUpdateViewsSeparateTranslucency = true;
diff --git a/Engine/Source/Runtime/Renderer/Private/TranslucentRendering.cpp b/Engine/Source/Runtime/Renderer/Private/TranslucentRendering.cpp
--- a/Engine/Source/Runtime/Renderer/Private/TranslucentRendering.cpp
+++ b/Engine/Source/Runtime/Renderer/Private/TranslucentRendering.cpp
@@ -86,6 +86,14 @@ static const TCHAR* kTranslucencyPassName[] = {
 };
 static_assert(UE_ARRAY_COUNT(kTranslucencyPassName) == int32(ETranslucencyPass::TPT_MAX), "Fix me");
 
+static const TCHAR* kTranslucencyColorTextureName[] = {
+	TEXT("Translucency.BeforeDistortion.Color"),
+	TEXT("Translucency.AfterDOF.Color"),
+	TEXT("Translucency.AfterDOF.Modulate"),
+	TEXT("Translucency.AfterMotionBlur.Color"),
+	TEXT("Translucency.All.Color"),
+};
+static_assert(UE_ARRAY_COUNT(kTranslucencyColorTextureName) == int32(ETranslucencyPass::TPT_MAX), "Fix me");
 
 static const TCHAR* TranslucencyPassToString(ETranslucencyPass::Type TranslucencyPass)
 {
@@ -823,6 +831,26 @@ void SetupPostMotionBlurTranslucencyViewParameters(const FViewInfo& View, FViewU
 	View.SetupUniformBufferParameters(ModifiedViewMatrices, ModifiedViewMatrices, VolumeBounds, TVC_MAX, Parameters);
 }
 
+FRDGTextureMSAA CreatePostDOFTranslucentTexture(
+	FRDGBuilder& GraphBuilder,
+	ETranslucencyPass::Type TranslucencyPass,
+	FSeparateTranslucencyDimensions& SeparateTranslucencyDimensions,
+	bool bIsModulate)
+{
+	const FRDGTextureDesc Desc = FRDGTextureDesc::Create2D(
+		SeparateTranslucencyDimensions.Extent,
+		bIsModulate ? PF_FloatR11G11B10 : PF_FloatRGBA,
+		bIsModulate ? FClearValueBinding::White : FClearValueBinding::Black,
+		TexCreate_RenderTargetable | TexCreate_ShaderResource,
+		1,
+		SeparateTranslucencyDimensions.NumSamples);
+
+	return CreateTextureMSAA(
+		GraphBuilder, Desc,
+		kTranslucencyColorTextureName[int32(TranslucencyPass)],
+		bIsModulate ? GFastVRamConfig.SeparateTranslucencyModulate : GFastVRamConfig.SeparateTranslucency);
+}
+
 void SetupDownsampledTranslucencyViewParameters(
 	const FViewInfo& View,
 	FIntPoint TextureExtent,
@@ -1321,30 +1349,7 @@ void FDeferredShadingSceneRenderer::RenderTranslucencyInner(
 	if (bRenderInSeparateTranslucency)
 	{
 		// Create resources shared by each view (each view data is tiled into each of the render target resources)
-		FRDGTextureMSAA SharedColorTexture;
-		{
-			static const TCHAR* kTranslucencyColorTextureName[] = {
-				TEXT("Translucency.BeforeDistortion.Color"),
-				TEXT("Translucency.AfterDOF.Color"),
-				TEXT("Translucency.AfterDOF.Modulate"),
-				TEXT("Translucency.AfterMotionBlur.Color"),
-				TEXT("Translucency.All.Color"),
-			};
-			static_assert(UE_ARRAY_COUNT(kTranslucencyColorTextureName) == int32(ETranslucencyPass::TPT_MAX), "Fix me");
-
-			const FRDGTextureDesc Desc = FRDGTextureDesc::Create2D(
-				SeparateTranslucencyDimensions.Extent,
-				bIsModulate ? PF_FloatR11G11B10 : PF_FloatRGBA,
-				bIsModulate ? FClearValueBinding::White : FClearValueBinding::Black,
-				TexCreate_RenderTargetable | TexCreate_ShaderResource,
-				1,
-				SeparateTranslucencyDimensions.NumSamples);
-
-			SharedColorTexture = CreateTextureMSAA(
-				GraphBuilder, Desc,
-				kTranslucencyColorTextureName[int32(TranslucencyPass)],
-				bIsModulate ? GFastVRamConfig.SeparateTranslucencyModulate : GFastVRamConfig.SeparateTranslucency);
-		}
+		FRDGTextureMSAA SharedColorTexture = CreatePostDOFTranslucentTexture(GraphBuilder, TranslucencyPass, SeparateTranslucencyDimensions, bIsModulate);
 
 		for (int32 ViewIndex = 0, NumProcessedViews = 0; ViewIndex < Views.Num(); ++ViewIndex, ++NumProcessedViews)
 		{
diff --git a/Engine/Source/Runtime/Renderer/Private/TranslucentRendering.h b/Engine/Source/Runtime/Renderer/Private/TranslucentRendering.h
--- a/Engine/Source/Runtime/Renderer/Private/TranslucentRendering.h
+++ b/Engine/Source/Runtime/Renderer/Private/TranslucentRendering.h
@@ -140,6 +140,9 @@ FSeparateTranslucencyDimensions UpdateTranslucencyTimers(FRHICommandListImmediat
 /** Returns whether the view family is requesting to render translucency. */
 bool ShouldRenderTranslucency(const FSceneViewFamily& ViewFamily);
 
+/** Shared function used to create Post DOF translucent textures */
+FRDGTextureMSAA CreatePostDOFTranslucentTexture(FRDGBuilder& GraphBuilder, ETranslucencyPass::Type TranslucencyPass, FSeparateTranslucencyDimensions& SeparateTranslucencyDimensions, bool bIsModulate);
+
 
 /** Add a pass to compose separate translucency. */
 struct FTranslucencyComposition
