Fix CIS warnings/errors.

diff --git a/Engine/Source/Runtime/Renderer/Private/RectLightTextureManager.h b/Engine/Source/Runtime/Renderer/Private/RectLightTextureManager.h
--- a/Engine/Source/Runtime/Renderer/Private/RectLightTextureManager.h
+++ b/Engine/Source/Runtime/Renderer/Private/RectLightTextureManager.h
@@ -3,9 +3,11 @@
 
 #include "CoreMinimal.h"
 #include "RHIDefinitions.h"
+#include "RenderGraphDefinitions.h"
 
 class FRHITexture;
-class FRDGBuilder;
+class UTexture;
+class FViewInfo;
 
 namespace RectLightAtlas
 {
@@ -23,7 +25,7 @@ uint32 AddRectLightTexture(UTexture* Texture);
 // Remove a rect light source texture to the texture atlas
 void RemoveRectLightTexture(uint32 InSlotId);
 
-// Return the atlas texture coordinate for a particular slot
+// Return the atlas texture coordinate for a	 particular slot
 FAtlasSlotDesc GetRectLightAtlasSlot(uint32 InSlotId);
 
 // Return the atlas texture
