Fix invalid rect light free block height resolution.

diff --git a/Engine/Source/Runtime/Renderer/Private/RectLightTextureManager.cpp b/Engine/Source/Runtime/Renderer/Private/RectLightTextureManager.cpp
--- a/Engine/Source/Runtime/Renderer/Private/RectLightTextureManager.cpp
+++ b/Engine/Source/Runtime/Renderer/Private/RectLightTextureManager.cpp
@@ -825,7 +825,7 @@ static bool FindFreeRect(FAtlasLayout& Layout, FAtlasSlot& Slot, int32 Alignemen
 				FreeRect.Origin.X     = Slot.Rect.Origin.X + Slot.Rect.Resolution.X;
 				FreeRect.Origin.Y     = BestFreeRect.Origin.Y;
 				FreeRect.Resolution.X =(BestFreeRect.Origin.X + BestFreeRect.Resolution.X) - FreeRect.Origin.X;
-				FreeRect.Resolution.Y = FreeRect.Resolution.Y;
+				FreeRect.Resolution.Y = BestFreeRect.Resolution.Y;
 			}
 		}
 		#endif
