Fix crash when using the GPU benchmark function.

The code was trying to create a readback-enabled render target, which is no longer supported now that pooled render targets no longer distinguish between targetable and shader resources. The target wasn't needed anyway, it was used in some code which was trying to sync the GPU with a readback that's been commented out 7 years ago. This CL removes the target entirely.

#rnx

diff --git a/Engine/Source/Runtime/Renderer/Private/GPUBenchmark.cpp b/Engine/Source/Runtime/Renderer/Private/GPUBenchmark.cpp
--- a/Engine/Source/Runtime/Renderer/Private/GPUBenchmark.cpp
+++ b/Engine/Source/Runtime/Renderer/Private/GPUBenchmark.cpp
@@ -420,18 +420,13 @@ void RendererGPUBenchmark(FRHICommandListImmediate& RHICmdList, FSynthBenchmarkR
 	FRHIResourceCreateInfo CreateInfo(TEXT("RendererGPUBenchmark"), &Vertices);
 	FBufferRHIRef VertexBuffer = RHICreateVertexBuffer(GBenchmarkVertices * sizeof(FBenchmarkVertex), BUF_Static, CreateInfo);
 
-	// two RT to ping pong so we force the GPU to flush it's pipeline
-	TRefCountPtr<IPooledRenderTarget> RTItems[3];
+	// two RT to ping pong so we force the GPU to flush its pipeline
+	TRefCountPtr<IPooledRenderTarget> RTItems[2];
 	{
 		FPooledRenderTargetDesc Desc(FPooledRenderTargetDesc::Create2DDesc(FIntPoint(GBenchmarkResolution, GBenchmarkResolution), PF_B8G8R8A8, FClearValueBinding::None, TexCreate_None, TexCreate_RenderTargetable | TexCreate_ShaderResource, false));
 
 		GRenderTargetPool.FindFreeElement(RHICmdList, Desc, RTItems[0], TEXT("Benchmark0"));
 		GRenderTargetPool.FindFreeElement(RHICmdList, Desc, RTItems[1], TEXT("Benchmark1"));
-
-		Desc.Extent = FIntPoint(1, 1);
-		Desc.Flags = TexCreate_CPUReadback;	// needs TexCreate_ResolveTargetable?
-
-		GRenderTargetPool.FindFreeElement(RHICmdList, Desc, RTItems[2], TEXT("BenchmarkReadback"));
 	}
 
 	{
@@ -569,27 +564,6 @@ void RendererGPUBenchmark(FRHICommandListImmediate& RHICmdList, FSynthBenchmarkR
 				RHICmdList.EndRenderPass();
 				RHICmdList.CopyToResolveTarget(RTItems[DestRTIndex]->GetRenderTargetItem().TargetableTexture, RTItems[DestRTIndex]->GetRenderTargetItem().ShaderResourceTexture, FResolveParams());
 
-				/*if(bGPUCPUSync)
-				{
-					// more consistent timing but strangely much faster to the level that is unrealistic
-
-					FResolveParams Param;
-
-					Param.Rect = FResolveRect(0, 0, 1, 1);
-					RHICmdList.CopyToResolveTarget(
-						RTItems[DestRTIndex]->GetRenderTargetItem().TargetableTexture,
-						RTItems[2]->GetRenderTargetItem().ShaderResourceTexture,
-						false,
-						Param);
-
-					void* Data = 0;
-					int Width = 0;
-					int Height = 0;
-
-					RHIMapStagingSurface(RTItems[2]->GetRenderTargetItem().ShaderResourceTexture, Data, Width, Height);
-					RHIUnmapStagingSurface(RTItems[2]->GetRenderTargetItem().ShaderResourceTexture);
-				}*/
-
 				RHICmdList.EndRenderQuery(TimerQueries[QueryIndex].GetQuery());
 
 				// ping pong
