Fix nanite not being included in CSMs in Nanite-only scenes

diff --git a/Engine/Source/Runtime/Renderer/Private/ShadowSetup.cpp b/Engine/Source/Runtime/Renderer/Private/ShadowSetup.cpp
--- a/Engine/Source/Runtime/Renderer/Private/ShadowSetup.cpp
+++ b/Engine/Source/Runtime/Renderer/Private/ShadowSetup.cpp
@@ -5127,7 +5127,8 @@ void FSceneRenderer::AllocateShadowDepthTargets(FRHICommandListImmediate& RHICmd
 			AllocateCachedShadowDepthTargets(RHICmdList, CachedWholeSceneDirectionalShadows);
 
 			// Only allocate CSM targets if at least one of the SMs in the CSM has any primitives, but then always allocate all of them as some algorithms depend on this (notably the forward shading structures).
-			bool bAnyHasSubjectPrims = false;
+			// Don't perform this if VSM are not enabled.
+			bool bAnyHasSubjectPrims = !VirtualShadowMapArray.IsEnabled();
 			for (const FProjectedShadowInfo* ProjectedShadowInfo : WholeSceneDirectionalShadows)
 			{
 				bAnyHasSubjectPrims = bAnyHasSubjectPrims || ProjectedShadowInfo->HasSubjectPrims();
