Clean up old CPU driven voxelization for groom.

diff --git a/Engine/Shaders/Private/HairStrands/HairStrandsVoxelOpaque.usf b/Engine/Shaders/Private/HairStrands/HairStrandsVoxelOpaque.usf
--- a/Engine/Shaders/Private/HairStrands/HairStrandsVoxelOpaque.usf
+++ b/Engine/Shaders/Private/HairStrands/HairStrandsVoxelOpaque.usf
@@ -6,7 +6,6 @@
 #include "HairStrandsVoxelPageCommon.ush"
 
 #if SHADER_INJECTOPAQUE_VIRTUALVOXEL
-uint3 DispatchedPageIndexResolution;
 uint MacroGroupId;
 float2 SceneDepthResolution;
 uint VoxelBiasCount;
diff --git a/Engine/Source/Runtime/Renderer/Private/HairStrands/HairStrandsData.h b/Engine/Source/Runtime/Renderer/Private/HairStrands/HairStrandsData.h
--- a/Engine/Source/Runtime/Renderer/Private/HairStrands/HairStrandsData.h
+++ b/Engine/Source/Runtime/Renderer/Private/HairStrands/HairStrandsData.h
@@ -122,14 +122,6 @@ struct FHairStrandsVisibilityData
 ////////////////////////////////////////////////////////////////////////////////////
 // Voxel data
 
-struct FHairStrandsVoxelNodeDesc
-{
-	FVector3f TranslatedWorldMinAABB = FVector3f::ZeroVector;
-	FVector3f TranslatedWorldMaxAABB = FVector3f::ZeroVector;
-	FIntVector PageIndexResolution = FIntVector::ZeroValue;
-	FMatrix TranslatedWorldToClip;
-};
-
 struct FPackedVirtualVoxelNodeDesc
 {
 	// This is just a placeholder having the correct size. The actual definition is in HairStradsNVoxelPageCommon.ush
@@ -300,7 +292,6 @@ struct FHairStrandsMacroGroupData
 	};
 	typedef TArray<PrimitiveInfo, SceneRenderingAllocator> TPrimitiveInfos;
 
-	FHairStrandsVoxelNodeDesc VirtualVoxelNodeDesc;
 	FHairStrandsDeepShadowDatas DeepShadowDatas;
 	TPrimitiveInfos PrimitivesInfos;
 	FBoxSphereBounds Bounds;
diff --git a/Engine/Source/Runtime/Renderer/Private/HairStrands/HairStrandsDebug.cpp b/Engine/Source/Runtime/Renderer/Private/HairStrands/HairStrandsDebug.cpp
--- a/Engine/Source/Runtime/Renderer/Private/HairStrands/HairStrandsDebug.cpp
+++ b/Engine/Source/Runtime/Renderer/Private/HairStrands/HairStrandsDebug.cpp
@@ -363,7 +363,12 @@ static void AddDebugHairPrintPass(
 	const FHairStrandsMacroGroupResources& MacroGroupResources,
 	FRDGTextureSRVRef InStencilTexture)
 {
-	if (!View || !View->HairStrandsViewData.UniformBuffer || !InStencilTexture || !ShaderDrawDebug::IsEnabled(*View)) return;
+	if (!View || !View->HairStrandsViewData.UniformBuffer || !InStencilTexture) return;
+
+	if (!TryEnableShaderDrawAndShaderPrint(*View, MacroGroupResources.MacroGroupCount * 32u, 2000u))
+	{
+		return;
+	}
 
 	FRDGTextureRef ViewHairCountTexture = VisibilityData.ViewHairCountTexture ? VisibilityData.ViewHairCountTexture : GSystemTextures.GetBlackDummy(GraphBuilder);
 	FRDGTextureRef ViewHairCountUintTexture = VisibilityData.ViewHairCountUintTexture ? VisibilityData.ViewHairCountUintTexture : GSystemTextures.GetBlackDummy(GraphBuilder);
@@ -811,7 +816,10 @@ static void AddDebugHairTangentPass(
 	const FSceneTextures& SceneTextures,
 	FRDGTextureRef& OutputTexture)
 {
-	ShaderDrawDebug::SetEnabled(true);
+	if (!TryEnableShaderDrawAndShaderPrint(View, 8192u, 1000u))
+	{
+		return;
+	}
 
 	FDebugHairTangentCS::FParameters* Parameters = GraphBuilder.AllocParameters<FDebugHairTangentCS::FParameters>();
 	Parameters->ViewUniformBuffer		= View.ViewUniformBuffer;
@@ -1110,58 +1118,55 @@ static void AddDrawDebugClusterPass(
 		return;
 	}
 
+
+	for (const FHairStrandsMacroGroupData& MacroGroupData : View.HairStrandsViewData.MacroGroupDatas)
 	{
+		const bool bDebugAABB = IsHairStrandsClusterDebugAABBEnable();
+
+		for (const FHairStrandsMacroGroupData::PrimitiveInfo& PrimitiveInfo : MacroGroupData.PrimitivesInfos)
 		{
-			for (const FHairStrandsMacroGroupData& MacroGroupData : View.HairStrandsViewData.MacroGroupDatas)
+			check(PrimitiveInfo.Mesh && PrimitiveInfo.Mesh->Elements.Num() > 0);
+
+			for (int DataIndex = 0; DataIndex < HairClusterData.HairGroups.Num(); ++DataIndex)
 			{
-				const bool bDebugAABB = IsHairStrandsClusterDebugAABBEnable();
+				const FHairStrandClusterData::FHairGroup& HairGroupClusters = HairClusterData.HairGroups[DataIndex];
 
-				for (const FHairStrandsMacroGroupData::PrimitiveInfo& PrimitiveInfo : MacroGroupData.PrimitivesInfos)
-				{
-					check(PrimitiveInfo.Mesh && PrimitiveInfo.Mesh->Elements.Num() > 0);
+				// Find a better/less hacky way
+				if (PrimitiveInfo.PublicDataPtr != HairGroupClusters.HairGroupPublicPtr)
+					continue;
 
-					for (int DataIndex = 0; DataIndex < HairClusterData.HairGroups.Num(); ++DataIndex)
+				if (ShaderDrawDebug::IsEnabled(View) && HairGroupClusters.CulledClusterCountBuffer)
+				{
+					FRDGExternalBuffer& DrawIndirectBuffer = HairGroupClusters.HairGroupPublicPtr->GetDrawIndirectBuffer();
+
+					FDrawDebugClusterAABBCS::FPermutationDomain Permutation;
+					Permutation.Set<FDrawDebugClusterAABBCS::FDebugAABBBuffer>(bDebugAABB ? 1 : 0);
+					TShaderMapRef<FDrawDebugClusterAABBCS> ComputeShader(View.ShaderMap, Permutation);
+
+					FDrawDebugClusterAABBCS::FParameters* Parameters = GraphBuilder.AllocParameters<FDrawDebugClusterAABBCS::FParameters>();
+					Parameters->ViewUniformBuffer = View.ViewUniformBuffer;
+					Parameters->ClusterCount = HairGroupClusters.ClusterCount;
+					Parameters->TriangleCount = HairGroupClusters.VertexCount * 2; // VertexCount is actually the number of control points
+					Parameters->HairGroupId = DataIndex;
+					Parameters->ClusterDebugMode = GHairStrandsClusterDebug;
+					Parameters->ClusterAABBBuffer = HairGroupClusters.ClusterAABBBuffer->SRV;
+					Parameters->CulledDispatchIndirectParametersClusterCountBuffer = GraphBuilder.CreateSRV(HairGroupClusters.CulledClusterCountBuffer, EPixelFormat::PF_R32_UINT);
+					Parameters->CulledDrawIndirectParameters = DrawIndirectBuffer.SRV;
+					Parameters->GroupAABBBuffer = HairGroupClusters.GroupAABBBuffer->SRV;
+
+					if (HairGroupClusters.ClusterDebugInfoBuffer && bDebugAABB)
 					{
-						const FHairStrandClusterData::FHairGroup& HairGroupClusters = HairClusterData.HairGroups[DataIndex];
-
-						// Find a better/less hacky way
-						if (PrimitiveInfo.PublicDataPtr != HairGroupClusters.HairGroupPublicPtr)
-							continue;
-
-						if (ShaderDrawDebug::IsEnabled(View) && HairGroupClusters.CulledClusterCountBuffer)
-						{
-							FRDGExternalBuffer& DrawIndirectBuffer = HairGroupClusters.HairGroupPublicPtr->GetDrawIndirectBuffer();
-
-							FDrawDebugClusterAABBCS::FPermutationDomain Permutation;
-							Permutation.Set<FDrawDebugClusterAABBCS::FDebugAABBBuffer>(bDebugAABB ? 1 : 0);
-							TShaderMapRef<FDrawDebugClusterAABBCS> ComputeShader(View.ShaderMap, Permutation);
-
-							FDrawDebugClusterAABBCS::FParameters* Parameters = GraphBuilder.AllocParameters<FDrawDebugClusterAABBCS::FParameters>();
-							Parameters->ViewUniformBuffer = View.ViewUniformBuffer;
-							Parameters->ClusterCount = HairGroupClusters.ClusterCount;
-							Parameters->TriangleCount = HairGroupClusters.VertexCount * 2; // VertexCount is actually the number of control points
-							Parameters->HairGroupId = DataIndex;
-							Parameters->ClusterDebugMode = GHairStrandsClusterDebug;
-							Parameters->ClusterAABBBuffer = HairGroupClusters.ClusterAABBBuffer->SRV;
-							Parameters->CulledDispatchIndirectParametersClusterCountBuffer = GraphBuilder.CreateSRV(HairGroupClusters.CulledClusterCountBuffer, EPixelFormat::PF_R32_UINT);
-							Parameters->CulledDrawIndirectParameters = DrawIndirectBuffer.SRV;
-							Parameters->GroupAABBBuffer = HairGroupClusters.GroupAABBBuffer->SRV;
-
-							if (HairGroupClusters.ClusterDebugInfoBuffer && bDebugAABB)
-							{
-								FRDGBufferRef ClusterDebugInfoBuffer = GraphBuilder.RegisterExternalBuffer(HairGroupClusters.ClusterDebugInfoBuffer);
-								Parameters->ClusterDebugInfoBuffer = GraphBuilder.CreateSRV(ClusterDebugInfoBuffer);
-							}
-							ShaderDrawDebug::SetParameters(GraphBuilder, View.ShaderDrawData, Parameters->ShaderDrawParameters);
-							ShaderPrint::SetParameters(GraphBuilder, View, Parameters->ShaderPrintParameters);
-
-							check(Parameters->ClusterCount / 64 <= 65535);
-							const FIntVector DispatchCount = DispatchCount.DivideAndRoundUp(FIntVector(Parameters->ClusterCount, 1, 1), FIntVector(64, 1, 1));// FIX ME, this could get over 65535
-							FComputeShaderUtils::AddPass(
-								GraphBuilder, RDG_EVENT_NAME("DrawDebugClusterAABB"),
-								ComputeShader, Parameters, DispatchCount);
-						}
+						FRDGBufferRef ClusterDebugInfoBuffer = GraphBuilder.RegisterExternalBuffer(HairGroupClusters.ClusterDebugInfoBuffer);
+						Parameters->ClusterDebugInfoBuffer = GraphBuilder.CreateSRV(ClusterDebugInfoBuffer);
 					}
+					ShaderDrawDebug::SetParameters(GraphBuilder, View.ShaderDrawData, Parameters->ShaderDrawParameters);
+					ShaderPrint::SetParameters(GraphBuilder, View, Parameters->ShaderPrintParameters);
+
+					check(Parameters->ClusterCount / 64 <= 65535);
+					const FIntVector DispatchCount = DispatchCount.DivideAndRoundUp(FIntVector(Parameters->ClusterCount, 1, 1), FIntVector(64, 1, 1));// FIX ME, this could get over 65535
+					FComputeShaderUtils::AddPass(
+						GraphBuilder, RDG_EVENT_NAME("DrawDebugClusterAABB"),
+						ComputeShader, Parameters, DispatchCount);
 				}
 			}
 		}
@@ -1283,61 +1288,25 @@ static void InternalRenderHairStrandsDebugInfo(
 
 	if (HairDebugMode == EHairDebugMode::MacroGroups)
 	{
-		{		
-			AddDebugHairPrintPass(GraphBuilder, &View, HairDebugMode, HairData.VisibilityData, HairData.MacroGroupResources, SceneTextures.Stencil);
-		}
-
-		// CPU bound of macro groups
-		FViewElementPDI ShadowFrustumPDI(&View, nullptr, nullptr);
-		if (HairData.VirtualVoxelResources.IsValid())
-		{
-			for (const FHairStrandsMacroGroupData& MacroGroupData : HairData.MacroGroupDatas)
-			{
-				const FBox Bound(MacroGroupData.VirtualVoxelNodeDesc.TranslatedWorldMinAABB, MacroGroupData.VirtualVoxelNodeDesc.TranslatedWorldMaxAABB);
-				DrawWireBox(&ShadowFrustumPDI, Bound, FColor::Red, 0);
-			}
-		}
-		#if 0
-		AddDrawCanvasPass(GraphBuilder, {}, View, SceneColor, [ClusterY, &MacroGroupDatas, YStep] (FCanvas& Canvas)
-		{
-			float X = 20;
-			float Y = ClusterY;
-			FLinearColor InactiveColor(0.5, 0.5, 0.5);
-			FLinearColor DebugColor(1, 1, 0);
-			FString Line;
-
-			Line = FString::Printf(TEXT("----------------------------------------------------------------"));
-			Canvas.DrawShadowedString(X, Y += YStep, *Line, GetStatsFont(), DebugColor);
-
-			Line = FString::Printf(TEXT("Macro group count : %d"), HairData.MacroGroupDatas.Num());
-			Canvas.DrawShadowedString(X, Y += YStep, *Line, GetStatsFont(), DebugColor);
-			for (const FHairStrandsMacroGroupData& MacroGroupData : HairData.MacroGroupDatas)
-			{
-				Line = FString::Printf(TEXT(" %d - Bound Radus: %f.2m (%dx%d)"), MacroGroupData.MacroGroupId, MacroGroupData.Bounds.GetSphere().W);
-				Canvas.DrawShadowedString(X, Y += YStep, *Line, GetStatsFont(), DebugColor);
-			}
-		});
-		#endif
+		AddDebugHairPrintPass(GraphBuilder, &View, HairDebugMode, HairData.VisibilityData, HairData.MacroGroupResources, SceneTextures.Stencil);
 	}
 
 	if (HairDebugMode == EHairDebugMode::DeepOpacityMaps)
 	{
+		for (const FHairStrandsMacroGroupData& MacroGroup : HairData.MacroGroupDatas)
 		{
-			for (const FHairStrandsMacroGroupData& MacroGroup : HairData.MacroGroupDatas)
+			if (!HairData.DeepShadowResources.DepthAtlasTexture || !HairData.DeepShadowResources.LayersAtlasTexture)
 			{
-				if (!HairData.DeepShadowResources.DepthAtlasTexture || !HairData.DeepShadowResources.LayersAtlasTexture)
-				{
-					continue;
-				}
+				continue;
+			}
 
-				for (const FHairStrandsDeepShadowData& DeepShadowData : MacroGroup.DeepShadowDatas)
-				{
-					const uint32 DomIndex = GDeepShadowDebugIndex;
-					if (DeepShadowData.AtlasSlotIndex != DomIndex)
-						continue;
+			for (const FHairStrandsDeepShadowData& DeepShadowData : MacroGroup.DeepShadowDatas)
+			{
+				const uint32 DomIndex = GDeepShadowDebugIndex;
+				if (DeepShadowData.AtlasSlotIndex != DomIndex)
+					continue;
 
-					AddDebugDeepShadowTexturePass(GraphBuilder, &View, FIntRect(), &DeepShadowData, &HairData.DeepShadowResources, SceneColorTexture);
-				}
+				AddDebugDeepShadowTexturePass(GraphBuilder, &View, FIntRect(), &DeepShadowData, &HairData.DeepShadowResources, SceneColorTexture);
 			}
 		}
 	}
@@ -1345,15 +1314,13 @@ static void InternalRenderHairStrandsDebugInfo(
 	// View Rect
 	if (IsHairStrandsViewRectOptimEnable() && HairDebugMode == EHairDebugMode::MacroGroupScreenRect)
 	{
+		for (const FHairStrandsMacroGroupData& MacroGroupData : HairData.MacroGroupDatas)
 		{
-			for (const FHairStrandsMacroGroupData& MacroGroupData : HairData.MacroGroupDatas)
-			{
-				AddDebugDeepShadowTexturePass(GraphBuilder, &View, MacroGroupData.ScreenRect, nullptr, nullptr, SceneColorTexture);
-			}
-
-			const FIntRect TotalRect = ComputeVisibleHairStrandsMacroGroupsRect(View.ViewRect, HairData.MacroGroupDatas);
-			AddDebugDeepShadowTexturePass(GraphBuilder, &View, TotalRect, nullptr, nullptr, SceneColorTexture);
+			AddDebugDeepShadowTexturePass(GraphBuilder, &View, MacroGroupData.ScreenRect, nullptr, nullptr, SceneColorTexture);
 		}
+
+		const FIntRect TotalRect = ComputeVisibleHairStrandsMacroGroupsRect(View.ViewRect, HairData.MacroGroupDatas);
+		AddDebugDeepShadowTexturePass(GraphBuilder, &View, TotalRect, nullptr, nullptr, SceneColorTexture);
 	}
 	
 	const bool bIsVoxelMode = HairDebugMode == EHairDebugMode::VoxelsDensity || HairDebugMode == EHairDebugMode::VoxelsTangent || HairDebugMode == EHairDebugMode::VoxelsBaseColor || HairDebugMode == EHairDebugMode::VoxelsRoughness;
@@ -1362,68 +1329,11 @@ static void InternalRenderHairStrandsDebugInfo(
 	{
 		if ((HairDebugMode == EHairDebugMode::LightBounds || HairDebugMode == EHairDebugMode::DeepOpacityMaps))
 		{
+			if (HairData.DeepShadowResources.bIsGPUDriven)
 			{
-				if (HairData.DeepShadowResources.bIsGPUDriven)
-				{
-					AddDeepShadowInfoPass(GraphBuilder, View, HairData.DeepShadowResources, HairData.MacroGroupResources, SceneColorTexture);
-				}
-			}
-		}
-
-		FViewElementPDI ShadowFrustumPDI(&View, nullptr, nullptr);
-
-		// Do not draw the CPU boud. Only GPU bound are drawn using ShaderDrawDebug shader functions.
-		#if 0
-		// All DOMs
-		if (HairDebugMode == EHairDebugMode::LightBounds)
-		{
-			if (!HairData.DeepShadowResources.bIsGPUDriven)
-			{
-				for (const FHairStrandsMacroGroupData& MacroGroupData : HairData.MacroGroupDatas)
-				{
-					for (const FHairStrandsDeepShadowData& DomData : MacroGroupData.DeepShadowDatas)
-					{
-						DrawFrustumWireframe(&ShadowFrustumPDI, DomData.CPU_TranslatedWorldToLightTransform.Inverse(), FColor::Emerald, 0);
-						DrawWireBox(&ShadowFrustumPDI, DomData.Bounds.GetBox(), FColor::Yellow, 0);
-					}
-				}
-			}
-		}
-
-		// Current DOM
-		if (HairDebugMode == EHairDebugMode::DeepOpacityMaps)
-		{
-			if (!HairData.DeepShadowResources.bIsGPUDriven)
-			{
-				const int32 CurrentIndex = FMath::Max(0, GDeepShadowDebugIndex);
-				for (const FHairStrandsMacroGroupData& MacroGroupData : HairData.MacroGroupDatas)
-				{
-					for (const FHairStrandsDeepShadowData& DomData : MacroGroupData.DeepShadowDatas)
-					{
-						if (DomData.AtlasSlotIndex == CurrentIndex)
-						{
-							DrawFrustumWireframe(&ShadowFrustumPDI, DomData.CPU_TranslatedWorldToLightTransform.Inverse(), FColor::Emerald, 0);
-							DrawWireBox(&ShadowFrustumPDI, DomData.Bounds.GetBox(), FColor::Yellow, 0);
-						}
-					}
-				}
-			}
-		}
-
-		// Voxelization
-		if (bIsVoxelMode)
-		{
-			if (HairData.VirtualVoxelResources.IsValid())
-			{
-				for (const FHairStrandsMacroGroupData& MacroGroupData : HairData.MacroGroupDatas)
-				{
-					const FBox Bound(MacroGroupData.VirtualVoxelNodeDesc.TranslatedWorldMinAABB, MacroGroupData.VirtualVoxelNodeDesc.TranslatedWorldMaxAABB);
-					DrawWireBox(&ShadowFrustumPDI, Bound, FColor::Red, 0);
-					DrawFrustumWireframe(&ShadowFrustumPDI, MacroGroupData.VirtualVoxelNodeDesc.TranslatedWorldToClip.Inverse(), FColor::Purple, 0);
-				}
+				AddDeepShadowInfoPass(GraphBuilder, View, HairData.DeepShadowResources, HairData.MacroGroupResources, SceneColorTexture);
 			}
 		}
-		#endif
 	}
 	
 	const bool bRunDebugPass =
@@ -1438,10 +1348,8 @@ static void InternalRenderHairStrandsDebugInfo(
 		HairDebugMode == EHairDebugMode::MaterialTangent;
 	if (bRunDebugPass)
 	{
-		{
-			AddDebugHairPass(GraphBuilder, &View, HairDebugMode, HairData.VisibilityData, SceneTextures.Stencil, SceneColorTexture);
-			AddDebugHairPrintPass(GraphBuilder, &View, HairDebugMode, HairData.VisibilityData, HairData.MacroGroupResources, SceneTextures.Stencil);
-		}
+		AddDebugHairPass(GraphBuilder, &View, HairDebugMode, HairData.VisibilityData, SceneTextures.Stencil, SceneColorTexture);
+		AddDebugHairPrintPass(GraphBuilder, &View, HairDebugMode, HairData.VisibilityData, HairData.MacroGroupResources, SceneTextures.Stencil);
 	}
 	else if (HairDebugMode == EHairDebugMode::Tile && HairData.VisibilityData.TileData.IsValid())
 	{
@@ -1450,11 +1358,9 @@ static void InternalRenderHairStrandsDebugInfo(
 
 	if (bIsVoxelMode)
 	{
+		if (HairData.VirtualVoxelResources.IsValid())
 		{
-			if (HairData.VirtualVoxelResources.IsValid())
-			{
-				AddVoxelPageRaymarchingPass(GraphBuilder, View, HairData.MacroGroupDatas, HairData.VirtualVoxelResources, SceneColorTexture);
-			}
+			AddVoxelPageRaymarchingPass(GraphBuilder, View, HairData.MacroGroupDatas, HairData.VirtualVoxelResources, SceneColorTexture);
 		}
 	}
 
diff --git a/Engine/Source/Runtime/Renderer/Private/HairStrands/HairStrandsVoxelization.cpp b/Engine/Source/Runtime/Renderer/Private/HairStrands/HairStrandsVoxelization.cpp
--- a/Engine/Source/Runtime/Renderer/Private/HairStrands/HairStrandsVoxelization.cpp
+++ b/Engine/Source/Runtime/Renderer/Private/HairStrands/HairStrandsVoxelization.cpp
@@ -14,9 +14,6 @@
 #include "PostProcessing.h"
 #include "ScenePrivate.h"
 
-static float GStrandHairVoxelizationRasterizationScale = 1.0f;
-static FAutoConsoleVariableRef CVarStrandHairVoxelizationRasterizationScale(TEXT("r.HairStrands.VoxelizationRasterizationScale"), GStrandHairVoxelizationRasterizationScale, TEXT("Rasterization scale to snap strand to pixel for voxelization"));
-
 static int32 GHairVoxelizationEnable = 1;
 static FAutoConsoleVariableRef CVarGHairVoxelizationEnable(TEXT("r.HairStrands.Voxelization"), GHairVoxelizationEnable, TEXT("Enable hair voxelization for transmittance evaluation"));
 
@@ -149,7 +146,6 @@ class FVirtualVoxelInjectOpaqueCS : public FGlobalShader
 	BEGIN_SHADER_PARAMETER_STRUCT(FParameters, )
 		SHADER_PARAMETER_STRUCT_INCLUDE(FSceneTextureParameters, SceneTextures)
 		SHADER_PARAMETER_STRUCT(FHairStrandsVoxelCommonParameters, VirtualVoxelParams)
-		SHADER_PARAMETER(FIntVector, DispatchedPageIndexResolution)
 		SHADER_PARAMETER(uint32, MacroGroupId)
 		SHADER_PARAMETER(FVector2f, SceneDepthResolution)
 		SHADER_PARAMETER(uint32, VoxelBiasCount)
@@ -194,7 +190,6 @@ static void AddVirtualVoxelInjectOpaquePass(
 	Parameters->SceneTextures = SceneTextures;
 	Parameters->MacroGroupId = MacroGroup.MacroGroupId;
 	Parameters->OutPageTexture = GraphBuilder.CreateUAV(VoxelResources.PageTexture);
-	Parameters->DispatchedPageIndexResolution = MacroGroup.VirtualVoxelNodeDesc.PageIndexResolution;
 	Parameters->IndirectDispatchArgs = VoxelResources.IndirectArgsBuffer;
 	TShaderMapRef<FVirtualVoxelInjectOpaqueCS> ComputeShader(View.ShaderMap);
 	const FGlobalShaderMap* GlobalShaderMap = View.ShaderMap;
@@ -557,10 +552,6 @@ static void AddAllocateVoxelPagesPass(
 		Out.PageIndexOffset = OutTotalPageIndexCount;
 
 		OutTotalPageIndexCount += Out.PageIndexCount;
-
-		MacroGroup.VirtualVoxelNodeDesc.TranslatedWorldMinAABB = (FVector3f)Out.TranslatedMinAABB;
-		MacroGroup.VirtualVoxelNodeDesc.TranslatedWorldMaxAABB = (FVector3f)Out.TranslatedMaxAABB;
-		MacroGroup.VirtualVoxelNodeDesc.PageIndexResolution = Out.PageIndexResolution;
 	}
 
 	// Over-allocation (upper bound)
@@ -722,9 +713,9 @@ static void AddAllocateVoxelPagesPass(
 				FIntVector DispatchCount = bUseClusterAABB ?
 					FIntVector((Parameters->MaxClusterCount + GroupSize - 1) / GroupSize, 1, 1) :
 					FIntVector(
-						FMath::DivideAndRoundUp(MacroGroup.VirtualVoxelNodeDesc.PageIndexResolution.X, 8),
-						FMath::DivideAndRoundUp(MacroGroup.VirtualVoxelNodeDesc.PageIndexResolution.Y, 8),
-						FMath::DivideAndRoundUp(MacroGroup.VirtualVoxelNodeDesc.PageIndexResolution.Z, 8));
+						FMath::DivideAndRoundUp(CPUAllocationDesc.PageIndexResolution.X, 8),
+						FMath::DivideAndRoundUp(CPUAllocationDesc.PageIndexResolution.Y, 8),
+						FMath::DivideAndRoundUp(CPUAllocationDesc.PageIndexResolution.Z, 8));
 
 				check(DispatchCount.X < 65535);
 				FVoxelMarkValidPageIndex_PrepareCS::FPermutationDomain PermutationVector;
@@ -1147,10 +1138,6 @@ static FHairStrandsVoxelResources AllocateDummyVirtualVoxelResources(
 	Out.Parameters.Common.TranslatedWorldOffset = FVector3f(View.ViewMatrices.GetPreViewTranslation());
 	Out.Parameters.Common.HairCoveragePixelRadiusAtDepth1	= ComputeMinStrandRadiusAtDepth1(FIntPoint(View.ViewRect.Width(), View.ViewRect.Height()), View.FOV, 1/*SampleCount*/, 1/*RasterizationScale*/).Primary;
 	Out.Parameters.Common.PageIndexCount = 0;
-	for (FHairStrandsMacroGroupData& MacroGroup : MacroGroupDatas)
-	{
-		MacroGroup.VirtualVoxelNodeDesc.PageIndexResolution = FIntVector(0,0,0);
-	}
 
 	FRDGBufferRef Dummy4BytesBuffer = GSystemTextures.GetDefaultBuffer(GraphBuilder, 4, 0u);
 	FRDGBufferRef Dummy8BytesBuffer = GSystemTextures.GetDefaultBuffer(GraphBuilder, 8, 0u);
@@ -1278,8 +1265,8 @@ IMPLEMENT_GLOBAL_SHADER(FVoxelRasterComputeCS, "/Engine/Private/HairStrands/Hair
 static void AddVirtualVoxelizationComputeRasterPass(
 	FRDGBuilder& GraphBuilder,
 	const FViewInfo* ViewInfo,
-	FHairStrandsVoxelResources& VoxelResources,
-	FHairStrandsMacroGroupData& MacroGroup)
+	const FHairStrandsVoxelResources& VoxelResources,
+	const FHairStrandsMacroGroupData& MacroGroup)
 {
 	const bool bIsGPUDriven = GHairVirtualVoxelGPUDriven > 0;
 	if (!bIsGPUDriven)
@@ -1374,85 +1361,12 @@ static void AddVirtualVoxelizationRasterPass(
 	const FViewInfo* ViewInfo,
 	FHairStrandsVoxelResources& VoxelResources,
 	FInstanceCullingManager& InstanceCullingManager,
-	FHairStrandsMacroGroupData& MacroGroup)
+	const FHairStrandsMacroGroupData& MacroGroup)
 {
-	const bool bIsGPUDriven = GHairVirtualVoxelGPUDriven > 0;
-	const FHairStrandsMacroGroupData::TPrimitiveInfos& PrimitiveSceneInfo = MacroGroup.PrimitivesInfos;
 	DECLARE_GPU_STAT(HairStrandsVoxelize);
 	SCOPED_DRAW_EVENT(GraphBuilder.RHICmdList, HairStrandsVoxelize);
 	SCOPED_GPU_STAT(GraphBuilder.RHICmdList, HairStrandsVoxelize);
 
-	// Find the largest resolution and its dominant axis
-	FIntPoint RasterResolution(0, 0);
-	FVector	RasterProjectionSize = FVector::ZeroVector;
-	FVector RasterDirection = FVector::ZeroVector;
-	FVector RasterUp = FVector::ZeroVector;
-	const FIntVector TotalVoxelResolution = MacroGroup.VirtualVoxelNodeDesc.PageIndexResolution * VoxelResources.Parameters.Common.PageResolution;
-	{
-		FIntVector ReorderIndex(0, 0, 0);
-
-		const uint32 ResolutionXY = TotalVoxelResolution.X * TotalVoxelResolution.Y;
-		const uint32 ResolutionXZ = TotalVoxelResolution.X * TotalVoxelResolution.Z;
-		const uint32 ResolutionYZ = TotalVoxelResolution.Y * TotalVoxelResolution.Y;
-		if (ResolutionXY >= ResolutionXZ && ResolutionXY >= ResolutionYZ)
-		{
-			RasterResolution	= FIntPoint(TotalVoxelResolution.X, TotalVoxelResolution.Y);
-			RasterDirection		= FVector(0, 0, 1);			
-			ReorderIndex		= FIntVector(0, 1, 2);
-			RasterUp			= FVector(0, 1, 0);
-		}
-		else if (ResolutionXZ >= ResolutionXY && ResolutionXZ >= ResolutionYZ)
-		{
-			RasterResolution	= FIntPoint(TotalVoxelResolution.X, TotalVoxelResolution.Z);
-			RasterDirection		= FVector(0, -1, 0);
-			ReorderIndex		= FIntVector(0, 2, 1);
-			RasterUp			= FVector(0, 0, 1);
-		}
-		else
-		{
-			RasterResolution	= FIntPoint(TotalVoxelResolution.Y, TotalVoxelResolution.Z);
-			RasterDirection		= FVector(1, 0, 0);
-			ReorderIndex		= FIntVector(1, 2, 0);
-			RasterUp			= FVector(0, 0, 1);
-		}
-
-		FBox ProjRasterAABB;
-		ProjRasterAABB.Min.X = MacroGroup.VirtualVoxelNodeDesc.TranslatedWorldMinAABB[ReorderIndex[0]];
-		ProjRasterAABB.Min.Y = MacroGroup.VirtualVoxelNodeDesc.TranslatedWorldMinAABB[ReorderIndex[1]];
-		ProjRasterAABB.Min.Z = MacroGroup.VirtualVoxelNodeDesc.TranslatedWorldMinAABB[ReorderIndex[2]];
-
-		ProjRasterAABB.Max.X = MacroGroup.VirtualVoxelNodeDesc.TranslatedWorldMaxAABB[ReorderIndex[0]];
-		ProjRasterAABB.Max.Y = MacroGroup.VirtualVoxelNodeDesc.TranslatedWorldMaxAABB[ReorderIndex[1]];
-		ProjRasterAABB.Max.Z = MacroGroup.VirtualVoxelNodeDesc.TranslatedWorldMaxAABB[ReorderIndex[2]];
-
-		RasterProjectionSize = ProjRasterAABB.GetSize();
-	}
-
-	if (bIsGPUDriven)
-	{
-		RasterResolution = GPUDrivenViewportResolution;
-	}
-
-	// RasterAABB is in translated world space
-	const FBox RasterAABB(MacroGroup.VirtualVoxelNodeDesc.TranslatedWorldMinAABB, MacroGroup.VirtualVoxelNodeDesc.TranslatedWorldMaxAABB);
-	const FVector RasterAABBSize = RasterAABB.GetSize();
-	const FVector RasterAABBCenter = RasterAABB.GetCenter();
-	const FIntRect ViewportRect = FIntRect(0, 0, RasterResolution.X, RasterResolution.Y);
-
-	const float RadiusAtDepth1 = GStrandHairVoxelizationRasterizationScale * VoxelResources.Parameters.Common.MinVoxelWorldSize * 0.5f;
-	const bool bIsOrtho = true;
-	const FVector4f HairRenderInfo = PackHairRenderInfo(RadiusAtDepth1, RadiusAtDepth1, RadiusAtDepth1, 1);
-	const uint32 HairRenderInfoBits = PackHairRenderInfoBits(bIsOrtho, bIsGPUDriven);
-
-	// TODO REMOVE THIS
-	FMatrix TranslatedWorldToClip;
-	{
-		FReversedZOrthoMatrix OrthoMatrix(0.5f * RasterProjectionSize.X, 0.5f * RasterProjectionSize.Y, 1.0f / RasterProjectionSize.Z, 0);
-		FLookAtMatrix LookAt(RasterAABBCenter - RasterDirection * RasterProjectionSize.Z * 0.5f, RasterAABBCenter, RasterUp);
-		TranslatedWorldToClip = LookAt * OrthoMatrix;
-		MacroGroup.VirtualVoxelNodeDesc.TranslatedWorldToClip = TranslatedWorldToClip;
-	}
-
 	AddVirtualVoxelizationComputeRasterPass(GraphBuilder, ViewInfo, VoxelResources, MacroGroup);
 }
 
