Skip post process material pass when all the pixels will fail the stencil test of the material. 0.1ms save on switch and will save more on mobile that are doing the stencil test in the pixel shader with a clip.

diff --git a/Engine/Source/Runtime/Renderer/Private/PostProcess/PostProcessMaterial.cpp b/Engine/Source/Runtime/Renderer/Private/PostProcess/PostProcessMaterial.cpp
--- a/Engine/Source/Runtime/Renderer/Private/PostProcess/PostProcessMaterial.cpp
+++ b/Engine/Source/Runtime/Renderer/Private/PostProcess/PostProcessMaterial.cpp
@@ -128,6 +128,40 @@ FRHIBlendState* GetMaterialBlendState(const FMaterial* Material)
 	return BlendStates[Material->GetBlendMode()];
 }
 
+bool PostProceessStencilTest(uint32 StencilValue, uint32 StencilComp, uint32 StencilRef)
+{
+	bool bStencilTestPassed = true;
+
+	switch (StencilComp)
+	{
+	case EMaterialStencilCompare::MSC_Less:
+		bStencilTestPassed = (StencilValue < StencilRef);
+		break;
+	case EMaterialStencilCompare::MSC_LessEqual:
+		bStencilTestPassed = (StencilValue <= StencilRef);
+		break;
+	case EMaterialStencilCompare::MSC_GreaterEqual:
+		bStencilTestPassed = (StencilValue >= StencilRef);
+		break;
+	case EMaterialStencilCompare::MSC_Equal:
+		bStencilTestPassed = (StencilValue == StencilRef);
+		break;
+	case EMaterialStencilCompare::MSC_Greater:
+		bStencilTestPassed = (StencilValue > StencilRef);
+		break;
+	case EMaterialStencilCompare::MSC_NotEqual:
+		bStencilTestPassed = (StencilValue != StencilRef);
+		break;
+	case EMaterialStencilCompare::MSC_Never:
+		bStencilTestPassed = false;
+		break;
+	default:
+		break;
+	}
+
+	return !bStencilTestPassed;
+}
+
 class FPostProcessMaterialShader : public FMaterialShader
 {
 public:
@@ -574,43 +608,73 @@ FScreenPassTexture AddPostProcessMaterialPass(
 		ScreenPassFlags |= EScreenPassDrawFlags::FlipYAxis;
 	}
 
-	AddDrawScreenPass(
-		GraphBuilder,
-		RDG_EVENT_NAME("PostProcessMaterial"),
-		View,
-		OutputViewport,
-		SceneColorViewport,
-		FScreenPassPipelineState(VertexShader, PixelShader, BlendState, DepthStencilState, MaterialStencilRef),
-		PostProcessMaterialParameters,
-		ScreenPassFlags,
-		[&View, VertexShader, PixelShader, MaterialRenderProxy, Material, PostProcessMaterialParameters, MaterialStencilRef](FRHICommandList& RHICmdList)
-	{
-		FPostProcessMaterialVS::SetParameters(RHICmdList, VertexShader, View, MaterialRenderProxy, *Material, *PostProcessMaterialParameters);
-		FPostProcessMaterialPS::SetParameters(RHICmdList, PixelShader, View, MaterialRenderProxy, *Material, *PostProcessMaterialParameters);
-	});
+	// check if we can skip that draw call in case if all pixels will fail the stencil test of the material
+	bool bSkipPostProcess = false;
 
-	if (bForceIntermediateTarget && !bCompositeWithInputAndDecode)
+	if (Material->IsStencilTestEnabled() && IsPostProcessStencilTestAllowed())
 	{
-		if (!Inputs.bFlipYAxis)
+		bool bFailStencil = true;
+
+		// when the depthStencilTexture was not produced, check if all the pixels will fail the clear value
+		if (!DepthStencilTexture)
 		{
-			// We shouldn't get here unless we had an override target.
-			check(Inputs.OverrideOutput.IsValid());
-			AddDrawTexturePass(GraphBuilder, View, Output.Texture, Inputs.OverrideOutput.Texture);
-			Output = Inputs.OverrideOutput;
+			uint32 StencilClearValue = Inputs.CustomDepthTexture ? Inputs.CustomDepthTexture->Desc.ClearValue.Value.DSValue.Stencil : 0;
+			bFailStencil &= PostProceessStencilTest(StencilClearValue, Material->GetStencilCompare(), PostProcessMaterialParameters->MobileStencilValueRef);
 		}
-		else
+
+		for (const uint32& Value : View.CustomDepthStencilValues)
+		{
+			bFailStencil &= PostProceessStencilTest(Value, Material->GetStencilCompare(), PostProcessMaterialParameters->MobileStencilValueRef);
+
+			if (!bFailStencil)
+			{
+				break;
+			}
+		}
+
+		bSkipPostProcess = bFailStencil;
+	}
+
+	if (!bSkipPostProcess)
+	{
+		AddDrawScreenPass(
+			GraphBuilder,
+			RDG_EVENT_NAME("PostProcessMaterial"),
+			View,
+			OutputViewport,
+			SceneColorViewport,
+			FScreenPassPipelineState(VertexShader, PixelShader, BlendState, DepthStencilState, MaterialStencilRef),
+			PostProcessMaterialParameters,
+			ScreenPassFlags,
+			[&View, VertexShader, PixelShader, MaterialRenderProxy, Material, PostProcessMaterialParameters, MaterialStencilRef](FRHICommandList& RHICmdList)
+			{
+				FPostProcessMaterialVS::SetParameters(RHICmdList, VertexShader, View, MaterialRenderProxy, *Material, *PostProcessMaterialParameters);
+				FPostProcessMaterialPS::SetParameters(RHICmdList, PixelShader, View, MaterialRenderProxy, *Material, *PostProcessMaterialParameters);
+			});
+
+		if (bForceIntermediateTarget && !bCompositeWithInputAndDecode)
 		{
-			FScreenPassRenderTarget TempTarget = Output;
-			if (Inputs.OverrideOutput.IsValid())
+			if (!Inputs.bFlipYAxis)
 			{
+				// We shouldn't get here unless we had an override target.
+				check(Inputs.OverrideOutput.IsValid());
+				AddDrawTexturePass(GraphBuilder, View, Output.Texture, Inputs.OverrideOutput.Texture);
 				Output = Inputs.OverrideOutput;
 			}
 			else
 			{
-				Output = FScreenPassRenderTarget(SceneColor, ERenderTargetLoadAction::ENoAction);
+				FScreenPassRenderTarget TempTarget = Output;
+				if (Inputs.OverrideOutput.IsValid())
+				{
+					Output = Inputs.OverrideOutput;
+				}
+				else
+				{
+					Output = FScreenPassRenderTarget(SceneColor, ERenderTargetLoadAction::ENoAction);
+				}
+
+				AddCopyAndFlipTexturePass(GraphBuilder, View, TempTarget.Texture, Output.Texture);
 			}
-
-			AddCopyAndFlipTexturePass(GraphBuilder, View, TempTarget.Texture, Output.Texture);
 		}
 	}
 
diff --git a/Engine/Source/Runtime/Renderer/Private/SceneRendering.h b/Engine/Source/Runtime/Renderer/Private/SceneRendering.h
--- a/Engine/Source/Runtime/Renderer/Private/SceneRendering.h
+++ b/Engine/Source/Runtime/Renderer/Private/SceneRendering.h
@@ -1210,6 +1210,9 @@ public:
 	bool bHasDistortionPrimitives;
 	bool bHasCustomDepthPrimitives;
 
+    /** Get all stencil values written into the custom depth pass */
+	TSet<uint32> CustomDepthStencilValues;
+
 	/** Mesh batches with for mesh decal rendering. */
 	TArray<FMeshDecalBatch, SceneRenderingAllocator> MeshDecalBatches;
 
diff --git a/Engine/Source/Runtime/Renderer/Private/SceneVisibility.cpp b/Engine/Source/Runtime/Renderer/Private/SceneVisibility.cpp
--- a/Engine/Source/Runtime/Renderer/Private/SceneVisibility.cpp
+++ b/Engine/Source/Runtime/Renderer/Private/SceneVisibility.cpp
@@ -2247,6 +2247,7 @@ struct FRelevancePacket
 	TArray<FSkyMeshBatch> SkyMeshBatches;
 	TArray<FSortedTrianglesMeshBatch> SortedTrianglesMeshBatches;
 	FDrawCommandRelevancePacket DrawCommandPacket;
+	TSet<uint32> CustumDepthStencilValues;
 
 	struct FPrimitiveLODMask
 	{
@@ -2457,6 +2458,7 @@ struct FRelevancePacket
 			if (ViewRelevance.bRenderCustomDepth)
 			{
 				bHasCustomDepthPrimitives = true;
+				CustumDepthStencilValues.Add(PrimitiveSceneInfo->Proxy->GetCustomDepthStencilValue());
 			}
 
 			extern bool GUseTranslucencyShadowDepths;
@@ -2880,6 +2882,7 @@ struct FRelevancePacket
 		WriteView.TranslucentPrimCount.Append(TranslucentPrimCount);
 		WriteView.bHasDistortionPrimitives |= bHasDistortionPrimitives;
 		WriteView.bHasCustomDepthPrimitives |= bHasCustomDepthPrimitives;
+		WriteView.CustomDepthStencilValues.Append(CustumDepthStencilValues);
 		WriteView.bUsesCustomDepthStencilInTranslucentMaterials |= bUsesCustomDepthStencil;
 		DirtyIndirectLightingCacheBufferPrimitives.AppendTo(WriteView.DirtyIndirectLightingCacheBufferPrimitives);
 
