Remove unused data from HeightfieldDescriptions buffer.
Remove unused param from ComposeHeightfieldsIntoPagesCS.
Added const keyword to GetHeightfieldRepresentation(...)
Remove TexCreate_RenderTargetable flag from VisualizeDistanceField texture since it is only used as UAV.

diff --git a/Engine/Shaders/Private/GlobalDistanceFieldHeightfields.usf b/Engine/Shaders/Private/GlobalDistanceFieldHeightfields.usf
--- a/Engine/Shaders/Private/GlobalDistanceFieldHeightfields.usf
+++ b/Engine/Shaders/Private/GlobalDistanceFieldHeightfields.usf
@@ -18,7 +18,6 @@ StructuredBuffer<uint> MarkedHeightfieldPageBuffer;
 StructuredBuffer<uint> PageUpdateTileBuffer;
 StructuredBuffer<uint> ComposeTileBuffer;
 Texture3D<uint> PageTableLayerTexture;
-Texture3D<uint> ParentPageTableLayerTexture;
 
 float3 PageCoordToPageWorldCenterScale;
 float3 PageCoordToPageWorldCenterBias;
diff --git a/Engine/Shaders/Private/HeightfieldLightingShared.ush b/Engine/Shaders/Private/HeightfieldLightingShared.ush
--- a/Engine/Shaders/Private/HeightfieldLightingShared.ush
+++ b/Engine/Shaders/Private/HeightfieldLightingShared.ush
@@ -14,7 +14,7 @@ uint NumHeightfields;
 Buffer<float4> HeightfieldDescriptions;
 
 // In float4's, must match C++
-#define HEIGHTFIELD_DATA_STRIDE 12
+#define HEIGHTFIELD_DATA_STRIDE 10
 
 float4 GetHeightfieldScaleBias(uint Index, out bool bHasSubsections)
 {
@@ -30,47 +30,37 @@ float4 GetHeightfieldMinMaxUV(uint Index)
 	return HeightfieldDescriptions[Index * HEIGHTFIELD_DATA_STRIDE + 1];
 }
 
-float4 GetHeightfieldGlobalScaleBias(uint Index)
-{
-	return HeightfieldDescriptions[Index * HEIGHTFIELD_DATA_STRIDE + 2];
-}
-
 float2 GetHeightfieldSize(uint Index)
 {
-	return HeightfieldDescriptions[Index * HEIGHTFIELD_DATA_STRIDE + 3].xy;
+	return HeightfieldDescriptions[Index * HEIGHTFIELD_DATA_STRIDE + 2].xy;
 }
 
 float2 GetHeightfieldSizeInv(uint Index)
 {
-	return HeightfieldDescriptions[Index * HEIGHTFIELD_DATA_STRIDE + 3].zw;
-}
-
-float2 GetGlobalTexelSize(uint Index)
-{
-	return HeightfieldDescriptions[Index * HEIGHTFIELD_DATA_STRIDE + 4].xy;
+	return HeightfieldDescriptions[Index * HEIGHTFIELD_DATA_STRIDE + 2].zw;
 }
 
 float4x4 GetWorldToLocal(uint Index)
 {
-	float4 M0 = HeightfieldDescriptions[Index * HEIGHTFIELD_DATA_STRIDE + 5];
-	float4 M1 = HeightfieldDescriptions[Index * HEIGHTFIELD_DATA_STRIDE + 6];
-	float4 M2 = HeightfieldDescriptions[Index * HEIGHTFIELD_DATA_STRIDE + 7];
+	float4 M0 = HeightfieldDescriptions[Index * HEIGHTFIELD_DATA_STRIDE + 3];
+	float4 M1 = HeightfieldDescriptions[Index * HEIGHTFIELD_DATA_STRIDE + 4];
+	float4 M2 = HeightfieldDescriptions[Index * HEIGHTFIELD_DATA_STRIDE + 5];
 
 	return transpose(float4x4(M0, M1, M2, float4(0.0f, 0.0f, 0.0f, 1.0f)));
 }
 
 float4x4 GetLocalToWorld(uint Index)
 {
-	float4 M0 = HeightfieldDescriptions[Index * HEIGHTFIELD_DATA_STRIDE + 8];
-	float4 M1 = HeightfieldDescriptions[Index * HEIGHTFIELD_DATA_STRIDE + 9];
-	float4 M2 = HeightfieldDescriptions[Index * HEIGHTFIELD_DATA_STRIDE + 10];
+	float4 M0 = HeightfieldDescriptions[Index * HEIGHTFIELD_DATA_STRIDE + 6];
+	float4 M1 = HeightfieldDescriptions[Index * HEIGHTFIELD_DATA_STRIDE + 7];
+	float4 M2 = HeightfieldDescriptions[Index * HEIGHTFIELD_DATA_STRIDE + 8];
 
 	return transpose(float4x4(M0, M1, M2, float4(0.0f, 0.0f, 0.0f, 1.0f)));
 }
 
 float4 GetVisibilityMask(uint Index)
 {
-	return HeightfieldDescriptions[Index * HEIGHTFIELD_DATA_STRIDE + 11];
+	return HeightfieldDescriptions[Index * HEIGHTFIELD_DATA_STRIDE + 9];
 }
 
 float2 GetHeightfieldUV(uint HeightfieldIndex, float2 LocalPosition, out float4 MinMaxHeightfieldUV)
diff --git a/Engine/Source/Runtime/Engine/Public/PrimitiveSceneProxy.h b/Engine/Source/Runtime/Engine/Public/PrimitiveSceneProxy.h
--- a/Engine/Source/Runtime/Engine/Public/PrimitiveSceneProxy.h
+++ b/Engine/Source/Runtime/Engine/Public/PrimitiveSceneProxy.h
@@ -364,7 +364,7 @@ public:
 
 	virtual bool HeightfieldHasPendingStreaming() const { return false; }
 
-	virtual void GetHeightfieldRepresentation(UTexture2D*& OutHeightmapTexture, UTexture2D*& OutDiffuseColorTexture, UTexture2D*& OutVisibilityTexture, FHeightfieldComponentDescription& OutDescription)
+	virtual void GetHeightfieldRepresentation(UTexture2D*& OutHeightmapTexture, UTexture2D*& OutDiffuseColorTexture, UTexture2D*& OutVisibilityTexture, FHeightfieldComponentDescription& OutDescription) const
 	{
 		OutHeightmapTexture = nullptr;
 		OutDiffuseColorTexture = nullptr;
diff --git a/Engine/Source/Runtime/Landscape/Private/LandscapeRender.cpp b/Engine/Source/Runtime/Landscape/Private/LandscapeRender.cpp
--- a/Engine/Source/Runtime/Landscape/Private/LandscapeRender.cpp
+++ b/Engine/Source/Runtime/Landscape/Private/LandscapeRender.cpp
@@ -3999,7 +3999,7 @@ bool FLandscapeComponentSceneProxy::HeightfieldHasPendingStreaming() const
 	return HeightmapTexture && HeightmapTexture->bHasStreamingUpdatePending;
 }
 
-void FLandscapeComponentSceneProxy::GetHeightfieldRepresentation(UTexture2D*& OutHeightmapTexture, UTexture2D*& OutDiffuseColorTexture, UTexture2D*& OutVisibilityTexture, FHeightfieldComponentDescription& OutDescription)
+void FLandscapeComponentSceneProxy::GetHeightfieldRepresentation(UTexture2D*& OutHeightmapTexture, UTexture2D*& OutDiffuseColorTexture, UTexture2D*& OutVisibilityTexture, FHeightfieldComponentDescription& OutDescription) const
 {
 	OutHeightmapTexture = HeightmapTexture;
 	OutDiffuseColorTexture = BaseColorForGITexture;
diff --git a/Engine/Source/Runtime/Landscape/Public/LandscapeRender.h b/Engine/Source/Runtime/Landscape/Public/LandscapeRender.h
--- a/Engine/Source/Runtime/Landscape/Public/LandscapeRender.h
+++ b/Engine/Source/Runtime/Landscape/Public/LandscapeRender.h
@@ -802,7 +802,7 @@ public:
 
 	virtual bool HeightfieldHasPendingStreaming() const override;
 
-	virtual void GetHeightfieldRepresentation(UTexture2D*& OutHeightmapTexture, UTexture2D*& OutDiffuseColorTexture, UTexture2D*& OutVisibilityTexture, FHeightfieldComponentDescription& OutDescription) override;
+	virtual void GetHeightfieldRepresentation(UTexture2D*& OutHeightmapTexture, UTexture2D*& OutDiffuseColorTexture, UTexture2D*& OutVisibilityTexture, FHeightfieldComponentDescription& OutDescription) const override;
 
 	virtual void GetLCIs(FLCIArray& LCIs) override;
 
diff --git a/Engine/Source/Runtime/Renderer/Private/DistanceFieldVisualization.cpp b/Engine/Source/Runtime/Renderer/Private/DistanceFieldVisualization.cpp
--- a/Engine/Source/Runtime/Renderer/Private/DistanceFieldVisualization.cpp
+++ b/Engine/Source/Runtime/Renderer/Private/DistanceFieldVisualization.cpp
@@ -111,7 +111,7 @@ void FDeferredShadingSceneRenderer::RenderMeshDistanceFieldVisualization(
 	FRDGTextureRef VisualizeResultTexture = nullptr;
 	{
 		const FIntPoint BufferSize = GetBufferSizeForAO();
-		const FRDGTextureDesc Desc(FRDGTextureDesc::Create2D(BufferSize, PF_FloatRGBA, FClearValueBinding::None, TexCreate_RenderTargetable | TexCreate_UAV));
+		const FRDGTextureDesc Desc(FRDGTextureDesc::Create2D(BufferSize, PF_FloatRGBA, FClearValueBinding::None, TexCreate_UAV));
 		VisualizeResultTexture = GraphBuilder.CreateTexture(Desc, TEXT("VisualizeDistanceField"));
 	}
 
diff --git a/Engine/Source/Runtime/Renderer/Private/GlobalDistanceField.cpp b/Engine/Source/Runtime/Renderer/Private/GlobalDistanceField.cpp
--- a/Engine/Source/Runtime/Renderer/Private/GlobalDistanceField.cpp
+++ b/Engine/Source/Runtime/Renderer/Private/GlobalDistanceField.cpp
@@ -1784,17 +1784,6 @@ void UpdateGlobalDistanceFieldVolume(
 
 									if (HeightfieldTexture && HeightfieldTexture->GetResource() && HeightfieldTexture->GetResource()->TextureRHI)
 									{
-										const FIntPoint HeightfieldSize = NewComponentDescription.HeightfieldRect.Size();
-
-										if (UpdateRegionHeightfield.Rect.Area() == 0)
-										{
-											UpdateRegionHeightfield.Rect = NewComponentDescription.HeightfieldRect;
-										}
-										else
-										{
-											UpdateRegionHeightfield.Rect.Union(NewComponentDescription.HeightfieldRect);
-										}
-
 										TArray<FHeightfieldComponentDescription>& ComponentDescriptions = UpdateRegionHeightfield.ComponentDescriptions.FindOrAdd(FHeightfieldComponentTextures(HeightfieldTexture, DiffuseColorTexture, VisibilityTexture));
 										ComponentDescriptions.Add(NewComponentDescription);
 									}
@@ -1925,11 +1914,7 @@ void UpdateGlobalDistanceFieldVolume(
 
 								if (HeightfieldDescriptions.Num() > 0)
 								{
-									FRDGBufferRef HeightfieldDescriptionBuffer = UploadHeightfieldDescriptions(
-										GraphBuilder,
-										HeightfieldDescriptions,
-										FVector2D(1, 1),
-										1.0f / UpdateRegionHeightfield.DownsampleFactor);
+									FRDGBufferRef HeightfieldDescriptionBuffer = UploadHeightfieldDescriptions(GraphBuilder, HeightfieldDescriptions);
 
 									UTexture2D* HeightfieldTexture = It.Key().HeightAndNormal;
 									UTexture2D* VisibilityTexture = It.Key().Visibility;
@@ -2215,11 +2200,7 @@ void UpdateGlobalDistanceFieldVolume(
 
 								if (HeightfieldDescriptions.Num() > 0)
 								{
-									FRDGBufferRef HeightfieldDescriptionBuffer = UploadHeightfieldDescriptions(
-										GraphBuilder,
-										HeightfieldDescriptions,
-										FVector2D(1, 1),
-										1.0f / UpdateRegionHeightfield.DownsampleFactor);
+									FRDGBufferRef HeightfieldDescriptionBuffer = UploadHeightfieldDescriptions(GraphBuilder, HeightfieldDescriptions);
 
 									UTexture2D* HeightfieldTexture = It.Key().HeightAndNormal;
 									UTexture2D* VisibilityTexture = It.Key().Visibility;
@@ -2230,7 +2211,6 @@ void UpdateGlobalDistanceFieldVolume(
 									PassParameters->ComposeIndirectArgBuffer = PageComposeHeightfieldIndirectArgBuffer;
 									PassParameters->ComposeTileBuffer = GraphBuilder.CreateSRV(PageComposeHeightfieldTileBuffer, PF_R32_UINT);
 									PassParameters->PageTableLayerTexture = PageTableLayerTexture;
-									PassParameters->ParentPageTableLayerTexture = ParentPageTableLayerTexture;
 									PassParameters->InfluenceRadius = ClipmapInfluenceRadius;
 									PassParameters->PageCoordToVoxelCenterScale = (FVector3f)PageCoordToVoxelCenterScale;
 									PassParameters->PageCoordToVoxelCenterBias = (FVector3f)PageCoordToVoxelCenterBias;
diff --git a/Engine/Source/Runtime/Renderer/Private/GlobalDistanceFieldHeightfields.h b/Engine/Source/Runtime/Renderer/Private/GlobalDistanceFieldHeightfields.h
--- a/Engine/Source/Runtime/Renderer/Private/GlobalDistanceFieldHeightfields.h
+++ b/Engine/Source/Runtime/Renderer/Private/GlobalDistanceFieldHeightfields.h
@@ -127,7 +127,6 @@ class FComposeHeightfieldsIntoPagesCS : public FGlobalShader
 		RDG_BUFFER_ACCESS(ComposeIndirectArgBuffer, ERHIAccess::IndirectArgs)
 		SHADER_PARAMETER_RDG_BUFFER_SRV(StructuredBuffer<uint>, ComposeTileBuffer)
 		SHADER_PARAMETER_RDG_TEXTURE(Texture3D<uint>, PageTableLayerTexture)
-		SHADER_PARAMETER_RDG_TEXTURE(Texture3D<uint>, ParentPageTableLayerTexture)
 		SHADER_PARAMETER(FVector3f, InvPageGridResolution)
 		SHADER_PARAMETER(FIntVector, PageGridResolution)
 		SHADER_PARAMETER(FVector3f, PageCoordToVoxelCenterScale)
@@ -168,4 +167,4 @@ class FComposeHeightfieldsIntoPagesCS : public FGlobalShader
 
 IMPLEMENT_GLOBAL_SHADER(FComposeHeightfieldsIntoPagesCS, "/Engine/Private/GlobalDistanceFieldHeightfields.usf", "ComposeHeightfieldsIntoPagesCS", SF_Compute);
 
-extern FRDGBufferRef UploadHeightfieldDescriptions(FRDGBuilder& GraphBuilder, const TArray<FHeightfieldComponentDescription>& HeightfieldDescriptions, FVector2D InvLightingAtlasSize, float InvDownsampleFactor);
\ No newline at end of file
+extern FRDGBufferRef UploadHeightfieldDescriptions(FRDGBuilder& GraphBuilder, const TArray<FHeightfieldComponentDescription>& HeightfieldDescriptions);
\ No newline at end of file
diff --git a/Engine/Source/Runtime/Renderer/Private/HeightfieldLighting.cpp b/Engine/Source/Runtime/Renderer/Private/HeightfieldLighting.cpp
--- a/Engine/Source/Runtime/Renderer/Private/HeightfieldLighting.cpp
+++ b/Engine/Source/Runtime/Renderer/Private/HeightfieldLighting.cpp
@@ -18,12 +18,9 @@ HeightfieldLighting.cpp
 #include "PipelineStateCache.h"
 
 // In float4's, must match usf
-static const int32 HEIGHTFIELD_DATA_STRIDE = 12;
+static const int32 HEIGHTFIELD_DATA_STRIDE = 10;
 
-void FillHeightfieldDescriptionData(const TArray<FHeightfieldComponentDescription>& HeightfieldDescriptions, 
-	FVector2D InvLightingAtlasSize, 
-	float InvDownsampleFactor,
-	TArray<FVector4f, SceneRenderingAllocator>& HeightfieldDescriptionData)
+void FillHeightfieldDescriptionData(const TArray<FHeightfieldComponentDescription>& HeightfieldDescriptions, TArray<FVector4f, SceneRenderingAllocator>& HeightfieldDescriptionData)
 {
 	HeightfieldDescriptionData.Empty(HeightfieldDescriptions.Num() * HEIGHTFIELD_DATA_STRIDE);
 
@@ -43,16 +40,7 @@ void FillHeightfieldDescriptionData(const TArray<FHeightfieldComponentDescriptio
 		HeightfieldDescriptionData.Add(HeightfieldScaleBias);
 		HeightfieldDescriptionData.Add(Description.MinMaxUV);
 
-		const FVector4f LightingUVScaleBias(
-			InvLightingAtlasSize.X * InvDownsampleFactor,
-			InvLightingAtlasSize.Y * InvDownsampleFactor,
-			Description.LightingAtlasLocation.X * InvLightingAtlasSize.X,
-			Description.LightingAtlasLocation.Y * InvLightingAtlasSize.Y);
-
-		HeightfieldDescriptionData.Add(LightingUVScaleBias);
-
 		HeightfieldDescriptionData.Add(FVector4f(Description.HeightfieldRect.Size().X, Description.HeightfieldRect.Size().Y, 1.f / Description.HeightfieldRect.Size().X, 1.f / Description.HeightfieldRect.Size().Y));
-		HeightfieldDescriptionData.Add(FVector4f(InvLightingAtlasSize.X, InvLightingAtlasSize.Y, 0.f, 0.f));
 
 		const FMatrix44f LocalToWorldT = FMatrix44f(Description.LocalToWorld.GetTransposed());
 		const FMatrix44f WorldToLocalT = FMatrix44f(Description.LocalToWorld.Inverse().GetTransposed());
@@ -76,14 +64,11 @@ void FillHeightfieldDescriptionData(const TArray<FHeightfieldComponentDescriptio
 	check(HeightfieldDescriptionData.Num() % HEIGHTFIELD_DATA_STRIDE == 0);
 }
 
-FRDGBufferRef UploadHeightfieldDescriptions(FRDGBuilder& GraphBuilder, const TArray<FHeightfieldComponentDescription>& HeightfieldDescriptions, FVector2D InvLightingAtlasSize, float InvDownsampleFactor)
+FRDGBufferRef UploadHeightfieldDescriptions(FRDGBuilder& GraphBuilder, const TArray<FHeightfieldComponentDescription>& HeightfieldDescriptions)
 {
 	TArray<FVector4f, SceneRenderingAllocator> HeightfieldDescriptionData;
 
-	FillHeightfieldDescriptionData(HeightfieldDescriptions,
-		InvLightingAtlasSize,
-		InvDownsampleFactor,
-		/*out*/ HeightfieldDescriptionData);
+	FillHeightfieldDescriptionData(HeightfieldDescriptions, /*out*/ HeightfieldDescriptionData);
 
 	FRDGBufferRef HeightfieldDescriptionsBuffer = CreateUploadBuffer(
 		GraphBuilder,
diff --git a/Engine/Source/Runtime/Renderer/Private/HeightfieldLighting.h b/Engine/Source/Runtime/Renderer/Private/HeightfieldLighting.h
--- a/Engine/Source/Runtime/Renderer/Private/HeightfieldLighting.h
+++ b/Engine/Source/Runtime/Renderer/Private/HeightfieldLighting.h
@@ -43,15 +43,5 @@ public:
 class FHeightfieldDescription
 {
 public:
-	FIntRect Rect;
-	int32 DownsampleFactor;
-	FIntRect DownsampledRect;
-
 	TMap<FHeightfieldComponentTextures, TArray<FHeightfieldComponentDescription>> ComponentDescriptions;
-
-	FHeightfieldDescription() :
-		Rect(FIntRect(0, 0, 0, 0)),
-		DownsampleFactor(1),
-		DownsampledRect(FIntRect(0, 0, 0, 0))
-	{}
 };
\ No newline at end of file
