Fix result of last global distance field propagation step not being used.

diff --git a/Engine/Source/Runtime/Renderer/Private/GlobalDistanceField.cpp b/Engine/Source/Runtime/Renderer/Private/GlobalDistanceField.cpp
--- a/Engine/Source/Runtime/Renderer/Private/GlobalDistanceField.cpp
+++ b/Engine/Source/Runtime/Renderer/Private/GlobalDistanceField.cpp
@@ -2258,7 +2258,7 @@ void UpdateGlobalDistanceFieldVolume(
 								uint32 PrevClipmapOffsetZ = 0;
 								uint32 NextClipmapOffsetZ = ClipmapIndex * ClipmapMipResolution;
 
-								if (StepIndex % 2 == 0)
+								if (StepIndex % 2 == NumPropagationSteps % 2)
 								{
 									Swap(PrevTexture, NextTexture);
 									Swap(PrevClipmapOffsetZ, NextClipmapOffsetZ);
