Fix issues r.AOGlobalDistanceFieldCacheMostlyStaticSeparately is disabled.
- Don't allocate PageTableCombinedTexture and use PageTableLayerTextures[GDF_Full] instead.
- Accumulate all primitive modifications info GDF_Full cache when not caching mostly static objects separately.

diff --git a/Engine/Shaders/Private/GlobalDistanceFieldMip.usf b/Engine/Shaders/Private/GlobalDistanceFieldMip.usf
--- a/Engine/Shaders/Private/GlobalDistanceFieldMip.usf
+++ b/Engine/Shaders/Private/GlobalDistanceFieldMip.usf
@@ -11,7 +11,7 @@
 RWTexture3D<float> RWMipTexture;
 uint ClipmapMipResolution;
 float MipInfluenceRadius;
-Texture3D<uint> PageTableCombinedTexture;
+Texture3D<uint> PageTableTexture;
 Texture3D PageAtlasTexture;
 Texture3D PrevMipTexture;
 
@@ -54,7 +54,7 @@ float LoadPrevDistanceFieldValue(int3 MipCoord, int3 Offset)
 	ScrolledClipmapUV = frac(ScrolledClipmapUV + ClipmapUVScrollOffset);
 
 	int3 PageTableTextureCoord = saturate(ScrolledClipmapUV) * GlobalDistanceFieldClipmapSizeInPages + int3(0, 0, ClipmapIndex * GlobalDistanceFieldClipmapSizeInPages);
-	uint PageIndex = PageTableCombinedTexture.Load(int4(PageTableTextureCoord, 0));
+	uint PageIndex = PageTableTexture.Load(int4(PageTableTextureCoord, 0));
 
 	float DistanceFieldValue = 1.0f;
 	if (PageIndex < 0xFFFF)
diff --git a/Engine/Source/Runtime/Renderer/Private/GlobalDistanceField.cpp b/Engine/Source/Runtime/Renderer/Private/GlobalDistanceField.cpp
--- a/Engine/Source/Runtime/Renderer/Private/GlobalDistanceField.cpp
+++ b/Engine/Source/Runtime/Renderer/Private/GlobalDistanceField.cpp
@@ -402,8 +402,17 @@ void FGlobalDistanceFieldInfo::UpdateParameterData(float MaxOcclusionDistance, b
 
 		if (PageTableCombinedTexture)
 		{
+			checkf(GAOGlobalDistanceFieldCacheMostlyStaticSeparately, TEXT("PageTableCombinedTexture should only be allocated when caching mostly static objects separately."));
 			ParameterData.PageTableTexture = PageTableCombinedTexture->GetRenderTargetItem().ShaderResourceTexture;
 		}
+		else
+		{
+			checkf(!GAOGlobalDistanceFieldCacheMostlyStaticSeparately, TEXT("PageTableCombinedTexture should be allocated when caching mostly static objects separately."));
+			if (PageTableLayerTextures[GDF_Full])
+			{
+				ParameterData.PageTableTexture = PageTableLayerTextures[GDF_Full]->GetRenderTargetItem().ShaderResourceTexture;
+			}
+		}
 
 		if (MipTexture)
 		{
@@ -731,6 +740,7 @@ static void ComputeUpdateRegionsAndUpdateViewState(
 			GlobalDistanceFieldInfo.PageAtlasTexture = ViewState.GlobalDistanceFieldPageAtlasTexture;
 		}
 
+		if(GAOGlobalDistanceFieldCacheMostlyStaticSeparately)
 		{
 			const FIntVector PageTableTextureResolution = GlobalDistanceField::GetPageTableTextureResolution(bLumenEnabled, View.FinalPostProcessSettings.LumenSceneViewDistance);
 			TRefCountPtr<IPooledRenderTarget>& PageTableTexture = View.ViewState->GlobalDistanceFieldPageTableCombinedTexture;
@@ -845,8 +855,8 @@ static void ComputeUpdateRegionsAndUpdateViewState(
 			// Accumulate primitive modifications in the viewstate in case we don't update the clipmap this frame
 			for (uint32 CacheType = 0; CacheType < GDF_Num; CacheType++)
 			{
-				const uint32 SourceCacheType = GAOGlobalDistanceFieldCacheMostlyStaticSeparately ? CacheType : GDF_Full;
-				ClipmapViewState.Cache[CacheType].PrimitiveModifiedBounds.Append(Scene->DistanceFieldSceneData.PrimitiveModifiedBounds[SourceCacheType]);
+				const uint32 DestCacheType = GAOGlobalDistanceFieldCacheMostlyStaticSeparately ? CacheType : GDF_Full;
+				ClipmapViewState.Cache[DestCacheType].PrimitiveModifiedBounds.Append(Scene->DistanceFieldSceneData.PrimitiveModifiedBounds[CacheType]);
 			}
 
 			const bool bForceFullUpdate = bSharedDataReallocated
@@ -1511,7 +1521,7 @@ class FPropagateMipDistanceCS : public FGlobalShader
 		SHADER_PARAMETER_STRUCT_REF(FViewUniformShaderParameters, View)
 		SHADER_PARAMETER_RDG_TEXTURE_UAV(RWTexture3D<float>, RWMipTexture)
 		SHADER_PARAMETER_RDG_TEXTURE(Texture3D<float>, PrevMipTexture)
-		SHADER_PARAMETER_RDG_TEXTURE(Texture3D<uint>, PageTableCombinedTexture)
+		SHADER_PARAMETER_RDG_TEXTURE(Texture3D<uint>, PageTableTexture)
 		SHADER_PARAMETER_RDG_TEXTURE(Texture3D<float>, PageAtlasTexture)
 		SHADER_PARAMETER(FVector3f, GlobalDistanceFieldInvPageAtlasSize)
 		SHADER_PARAMETER(uint32, GlobalDistanceFieldClipmapSizeInPages)
@@ -2049,7 +2059,7 @@ void UpdateGlobalDistanceFieldVolume(
 								PassParameters->PageUpdateTileBuffer = GraphBuilder.CreateSRV(PageUpdateTileBuffer, PF_R32_UINT);
 								PassParameters->MarkedHeightfieldPageBuffer = MarkedHeightfieldPageBuffer ? GraphBuilder.CreateSRV(MarkedHeightfieldPageBuffer, PF_R32_UINT) : nullptr;
 
-								PassParameters->RWPageTableCombinedTexture = GraphBuilder.CreateUAV(PageTableCombinedTexture);
+								PassParameters->RWPageTableCombinedTexture = PageTableCombinedTexture ? GraphBuilder.CreateUAV(PageTableCombinedTexture) : nullptr;
 								PassParameters->RWPageTableLayerTexture = GraphBuilder.CreateUAV(PageTableLayerTexture);
 								PassParameters->RWPageFreeListAllocatorBuffer = GraphBuilder.CreateUAV(PageFreeListAllocatorBuffer, PF_R32_SINT);
 								PassParameters->PageFreeListBuffer = GraphBuilder.CreateSRV(PageFreeListBuffer, PF_R32_UINT);
@@ -2267,7 +2277,7 @@ void UpdateGlobalDistanceFieldVolume(
 								FPropagateMipDistanceCS::FParameters* PassParameters = GraphBuilder.AllocParameters<FPropagateMipDistanceCS::FParameters>();
 								PassParameters->View = View.ViewUniformBuffer;
 								PassParameters->RWMipTexture = GraphBuilder.CreateUAV(NextTexture);
-								PassParameters->PageTableCombinedTexture = PageTableCombinedTexture;
+								PassParameters->PageTableTexture = GAOGlobalDistanceFieldCacheMostlyStaticSeparately ? PageTableCombinedTexture : PageTableLayerTexture;
 								PassParameters->PageAtlasTexture = PageAtlasTexture;
 								PassParameters->GlobalDistanceFieldInvPageAtlasSize = FVector3f::OneVector / FVector3f(GlobalDistanceField::GetPageAtlasSize(bLumenEnabled, View.FinalPostProcessSettings.LumenSceneViewDistance));
 								PassParameters->GlobalDistanceFieldClipmapSizeInPages = GlobalDistanceField::GetPageTableTextureResolution(bLumenEnabled, View.FinalPostProcessSettings.LumenSceneViewDistance).X;
