U5 - fix copy and paste error from loadings refactor.

diff --git a/Engine/Source/Runtime/Renderer/Private/MeshPassProcessor.cpp b/Engine/Source/Runtime/Renderer/Private/MeshPassProcessor.cpp
--- a/Engine/Source/Runtime/Renderer/Private/MeshPassProcessor.cpp
+++ b/Engine/Source/Runtime/Renderer/Private/MeshPassProcessor.cpp
@@ -931,15 +931,15 @@ bool FMeshDrawShaderBindings::MatchesForDynamicInstancing(const FMeshDrawShaderB
 		}
 
 		FRHIResource* const* SrvBindings = SingleShaderBindings.GetSRVStart();
-		FRHIResource* const* OtherSrvBindings = SingleShaderBindings.GetSRVStart();
+		FRHIResource* const* OtherSrvBindings = OtherSingleShaderBindings.GetSRVStart();
 		for (int32 SrvIndex = 0; SrvIndex < SingleShaderBindings.ParameterMapInfo.SRVs.Num(); SrvIndex++)
 		{
 			const FRHIResource* Srv = SrvBindings[SrvIndex];
 			const FRHIResource* OtherSrv = OtherSrvBindings[SrvIndex];
 			if (Srv != OtherSrv)
-		{
-			return false;
-		}
+			{
+				return false;
+			}
 		}
 
 		FRHIUniformBuffer* const* UniformBufferBindings = SingleShaderBindings.GetUniformBufferStart();
