Change voxel size feedback to run only on GPU (i.e., without GPU->CPU->GPU roundtrip).

This simplifies the logic/code and reduce feedback latency.

diff --git a/Engine/Shaders/Private/HairStrands/HairStrandsAABB.usf b/Engine/Shaders/Private/HairStrands/HairStrandsAABB.usf
--- a/Engine/Shaders/Private/HairStrands/HairStrandsAABB.usf
+++ b/Engine/Shaders/Private/HairStrands/HairStrandsAABB.usf
@@ -10,7 +10,6 @@ uint MacroGroupId;
 uint MacroGroupValid;
 
 float PixelSizeAtDepth1;
-float MinVoxelWorldSize;
 float NumPixelPerVoxel;
 
 // Group data packed as {uint3 AABBMin, uint3 AABBMax}
@@ -58,6 +57,8 @@ void Main(uint2 DispatchThreadId : SV_DispatchThreadID)
 		const float PixelWorldSize = PixelSizeAtDepth1 * ViewBound.Min.z;
 		const float VoxelWorldSize = PixelWorldSize * NumPixelPerVoxel;
 		WriteHairAABB(MacroGroupId, Bound, OutMacroGroupAABBBuffer);
+
+		const float MinVoxelWorldSize = 0.01f; // 0.01mm
 		OutMacroGroupVoxelSizeBuffer[MacroGroupId] = max(MinVoxelWorldSize, VoxelWorldSize);
 	}
 }
diff --git a/Engine/Shaders/Private/HairStrands/HairStrandsVoxelPageAllocation.usf b/Engine/Shaders/Private/HairStrands/HairStrandsVoxelPageAllocation.usf
--- a/Engine/Shaders/Private/HairStrands/HairStrandsVoxelPageAllocation.usf
+++ b/Engine/Shaders/Private/HairStrands/HairStrandsVoxelPageAllocation.usf
@@ -6,6 +6,8 @@
 #include "HairStrandsAABBCommon.ush"
 #include "HairStrandsVisibilityCommon.ush"
 
+///////////////////////////////////////////////////////////////////////////////////////////////////
+
 #if SHADER_ALLOCATEPAGEINDEX 
 
 float				CPUPageWorldSize;
@@ -22,6 +24,7 @@ float3				CPU_TranslatedWorldMaxAABB;
 int3				CPU_PageIndexResolution;
 uint				CPU_bUseCPUData;
 
+Buffer<float>								GPUVoxelWorldSize;
 Buffer<float>								MacroGroupVoxelSizeBuffer;
 RWBuffer<int>								MacroGroupAABBBuffer;
 RWBuffer<uint4>								OutPageIndexResolutionAndOffsetBuffer;
@@ -49,7 +52,7 @@ void AllocatePageIndex(uint2 DispatchThreadId : SV_DispatchThreadID)
 	{
 		const bool bSupportVoxelization = (bDoesMacroGroupSupportVoxelization >> MacroGroupId) & 0x1;
 		Bound = ReadHairAABB(MacroGroupId, MacroGroupAABBBuffer);
-		const float VoxelWorldSize = max(CPUVoxelWorldSize, MacroGroupVoxelSizeBuffer[MacroGroupId]);
+		const float VoxelWorldSize = max(GPUVoxelWorldSize[0], MacroGroupVoxelSizeBuffer[MacroGroupId]);
 		PageWorldSize = VoxelWorldSize * PageResolution;
 
 		if (any(Bound.Min > Bound.Max) || !bSupportVoxelization)
@@ -208,6 +211,8 @@ void AllocatePageIndex(uint2 DispatchThreadId : SV_DispatchThreadID)
 }
 #endif
 
+///////////////////////////////////////////////////////////////////////////////////////////////////
+
 #if SHADER_MARKVALID_PREPARE || SHADER_MARKVALID_INDIRECTARG || SHADER_MARKVALID_SCATTER
 struct FScatterCoord
 {
@@ -264,6 +269,8 @@ void UpdatePackScatterCoord(inout uint2 Out, uint PartIndex)
 
 #endif // SHADER_MARKVALID_PREPARE || SHADER_MARKVALID_INDIRECTARG || SHADER_MARKVALID_SCATTER
 
+///////////////////////////////////////////////////////////////////////////////////////////////////
+
 #if SHADER_MARKVALID_PREPARE
 uint			MaxClusterCount;
 uint			MacroGroupId;
@@ -433,6 +440,8 @@ void MarkValid_BuildIndirectArgCS(uint2 DispatchThreadId : SV_DispatchThreadID)
 }
 #endif
 
+///////////////////////////////////////////////////////////////////////////////////////////////////
+
 #if SHADER_MARKVALID_SCATTER
 Buffer<int>		IndirectBufferArgs;
 Buffer<uint4>	PageIndexResolutionAndOffsetBuffer;
@@ -468,6 +477,8 @@ void MarkValid_ScatterCS(uint2 DispatchThreadId : SV_DispatchThreadID)
 }
 #endif
 
+///////////////////////////////////////////////////////////////////////////////////////////////////
+
 #if SHADER_MARKVALID
 uint			MaxClusterCount;
 int3			CPU_PageIndexResolution;
@@ -541,6 +552,8 @@ void MarkValidCS(uint2 DispatchThreadId : SV_DispatchThreadID)
 }
 #endif
 
+///////////////////////////////////////////////////////////////////////////////////////////////////
+
 #if SHADER_ALLOCATE
 uint MacroGroupId;
 uint PageCount;
@@ -629,6 +642,8 @@ void AllocateCS(uint GroupIndex : SV_GroupIndex, uint3 DispatchThreadId : SV_Dis
 }
 #endif
 
+///////////////////////////////////////////////////////////////////////////////////////////////////
+
 #if SHADER_ADDDESC
 
 float3 CPU_TranslatedWorldMinAABB;
@@ -638,6 +653,7 @@ uint CPU_PageIndexOffset;
 int3 CPU_PageIndexResolution;
 float CPU_VoxelWorldSize;
 
+Buffer<float> GPU_VoxelWorldSize;
 Buffer<int> MacroGroupAABBBuffer;
 Buffer<float> MacroGroupVoxelSizeBuffer;
 Buffer<uint4> PageIndexResolutionAndOffsetBuffer;
@@ -657,7 +673,7 @@ void AddDescCS(uint GroupIndex : SV_GroupIndex, uint3 DispatchThreadId : SV_Disp
 	Node.TranslatedWorldMaxAABB = TranslatedWorldBound.Max;
 	Node.PageIndexResolution = PageIndexResolutionAndOffset.xyz;
 	Node.PageIndexOffset = PageIndexResolutionAndOffset.w;
-	Node.VoxelWorldSize = max(CPU_VoxelWorldSize, VoxelWorldSize);
+	Node.VoxelWorldSize = max(GPU_VoxelWorldSize[0], VoxelWorldSize);
 #else
 	Node.TranslatedWorldMinAABB = CPU_TranslatedWorldMinAABB;
 	Node.TranslatedWorldMaxAABB = CPU_TranslatedWorldMaxAABB;
@@ -671,6 +687,7 @@ void AddDescCS(uint GroupIndex : SV_GroupIndex, uint3 DispatchThreadId : SV_Disp
 }
 #endif
 
+///////////////////////////////////////////////////////////////////////////////////////////////////
 
 #if SHADER_ADDINDIRECTBUFFER
 
@@ -708,6 +725,7 @@ void AddIndirectBufferCS(uint GroupIndex : SV_GroupIndex, uint3 DispatchThreadId
 
 #endif
 
+///////////////////////////////////////////////////////////////////////////////////////////////////
 
 #define IND_CLEAR_GROUP_SIZE_Z 64
 
@@ -734,6 +752,7 @@ void VoxelIndPageClearBufferGenCS(uint GroupIndex : SV_GroupIndex, uint3 Dispatc
 
 #endif
 
+///////////////////////////////////////////////////////////////////////////////////////////////////
 
 #if SHADER_INDPAGECLEAR
 
@@ -768,3 +787,128 @@ void VoxelIndPageClearCS(uint3 DispatchThreadId : SV_DispatchThreadID, uint3 Gro
 
 #endif
 
+#if SHADER_ADAPTIVE_FEEDBACK
+
+#define DEBUG_ENABLE 0
+
+#if DEBUG_ENABLE
+#include "../ShaderPrintCommon.ush"
+void AddNewLine(float2 OriginalPos, inout float2 Pos)
+{
+	Pos = ShaderPrintNewline(Pos); Pos.x = OriginalPos.x;
+}
+#endif
+
+uint  CPUAllocatedPageCount;
+float CPUMinVoxelWorldSize;
+float AdaptiveCorrectionThreshold;
+float AdaptiveCorrectionSpeed;
+
+Buffer<uint> TotalRequestedPageAllocationBuffer;
+Buffer<float> CurrGPUMinVoxelWorldSize;
+RWBuffer<float> NextGPUMinVoxelWorldSize;
+
+float RoundHairVoxelSize(float In)
+{
+	// Round voxel size to 0.01f to avoid oscillation issue
+	return floor(In * 1000.f + 0.5f) * 0.001f;
+}
+
+[numthreads(1, 1, 1)]
+void FeedbackCS(uint3 DispatchThreadId : SV_DispatchThreadID)
+{
+	const float CurrVoxelWorldSize = RoundHairVoxelSize(CurrGPUMinVoxelWorldSize[0]);
+	
+	// Voxel pages are represent a volume. To derive a better estimate of the ratio by which voxel size needs to be scale, 
+	// compute the cubic root of this ratio.
+	//
+	// AllocatedPage   AllocatedRes^3
+	// ------------- = --------------  = VolumeRatio = LinearRatio^3
+	//    MaxPage          MaxRes^3
+
+	// Ratio used for predicting voxel size increase
+	const uint GPUAllocatedPageCount = TotalRequestedPageAllocationBuffer[0];
+	const float VolumeRatio = float(GPUAllocatedPageCount) / float(CPUAllocatedPageCount);
+	const float LinearRatio = pow(VolumeRatio, 1.f / 3.f);
+
+	// Ratio used for predicting voxel size decrease (i.e. when requested allocation fit, 
+	// but the voxel size does not match the (more precise) target).
+	// In this case, we add a threshold/margin to to the target, so that there is no oscillation.
+	const float VolumeRatio_Thres = float(GPUAllocatedPageCount) / float(CPUAllocatedPageCount * AdaptiveCorrectionThreshold);
+	const float LinearRatio_Thres = pow(VolumeRatio_Thres, 1.f / 3.f);
+
+	// If the page pool is not large enough increase voxel size
+	float NextVoxelWorldSize = CPUMinVoxelWorldSize;
+	if (GPUAllocatedPageCount > CPUAllocatedPageCount)
+	{
+		//NextVoxelWorldSize = CurrVoxelWorldSize * LinearRatio;
+		NextVoxelWorldSize = CurrVoxelWorldSize * LinearRatio_Thres;
+	}
+	// If the page pool is large enough but the voxel are larger than the requested size decrease voxel size
+	else if (GPUAllocatedPageCount < CPUAllocatedPageCount && CurrVoxelWorldSize > CPUMinVoxelWorldSize)
+	{
+		const float TargetVoxelWorldSize = CurrVoxelWorldSize * LinearRatio_Thres;
+		NextVoxelWorldSize = max(CPUMinVoxelWorldSize, lerp(CurrVoxelWorldSize, TargetVoxelWorldSize, AdaptiveCorrectionSpeed));
+	}
+	//else if (GPUAllocatedPageCount > CPUAllocatedPageCount * AdaptiveCorrectionThreshold)
+	//{
+	//	const float TargetVoxelWorldSize = CurrVoxelWorldSize * LinearRatio_Thres;
+	//	NextVoxelWorldSize = max(CPUMinVoxelWorldSize, lerp(CurrVoxelWorldSize, TargetVoxelWorldSize, AdaptiveCorrectionSpeed));
+	//}
+	else
+	{
+		NextVoxelWorldSize = CPUMinVoxelWorldSize;
+	}
+
+	// Clamp voxel size into a reasonable amount (e.g. 0.1mm - 100mm)
+	const float ClampMinVoxelWorldSize = 0.01f;
+	const float ClampMaxVoxelWorldSize = 10.0f;
+	NextVoxelWorldSize = clamp(RoundHairVoxelSize(NextVoxelWorldSize), ClampMinVoxelWorldSize, ClampMaxVoxelWorldSize);
+
+	// Debug
+#if DEBUG_ENABLE
+	FFontColor CPUColor = FontEmerald;
+	FFontColor GPUColor = FontOrange;
+	FFontColor CstColor = FontSilver;
+
+	const float2 OriginalPos = ShaderPrintGetPos(uint2(700, 50));
+	float2 Pos = OriginalPos;
+
+	Pos = ShaderPrintText(Pos, TEXT(" ------------------------------- "), FontSilver); AddNewLine(OriginalPos, Pos);	
+	Pos = ShaderPrintText(Pos, TEXT("|          Allocations          |"), FontSilver); AddNewLine(OriginalPos, Pos);
+	Pos = ShaderPrintText(Pos, TEXT(" ------------------------------- "), FontSilver); AddNewLine(OriginalPos, Pos);
+
+	Pos = ShaderPrintText(Pos, TEXT("GPU Allocated      "), GPUColor);
+	Pos = ShaderPrint(Pos, GPUAllocatedPageCount, GPUColor);
+	AddNewLine(OriginalPos, Pos);
+
+	Pos = ShaderPrintText(Pos, TEXT("CPU Allocated      "), CPUColor);
+	Pos = ShaderPrint(Pos, CPUAllocatedPageCount, CPUColor);
+	AddNewLine(OriginalPos, Pos);
+
+	Pos = ShaderPrintText(Pos, TEXT("GPU Curr Min. Size "), GPUColor);
+	Pos = ShaderPrint(Pos, CurrVoxelWorldSize, GPUColor);
+	AddNewLine(OriginalPos, Pos);
+
+	Pos = ShaderPrintText(Pos, TEXT("GPU Next Min. Size "), GPUColor);
+	Pos = ShaderPrint(Pos, NextVoxelWorldSize, GPUColor);
+	AddNewLine(OriginalPos, Pos);
+
+	Pos = ShaderPrintText(Pos, TEXT("CPU Min. Size      "), CPUColor);
+	Pos = ShaderPrint(Pos, CPUMinVoxelWorldSize, CPUColor);
+	AddNewLine(OriginalPos, Pos);
+
+	Pos = ShaderPrintText(Pos, TEXT("Correction Thres.  "), CstColor);
+	Pos = ShaderPrint(Pos, AdaptiveCorrectionThreshold, CstColor);
+	AddNewLine(OriginalPos, Pos);
+
+	Pos = ShaderPrintText(Pos, TEXT("Correction Speed   "), CstColor);
+	Pos = ShaderPrint(Pos, AdaptiveCorrectionSpeed, CstColor);
+	AddNewLine(OriginalPos, Pos);
+#endif
+
+	// Update state data
+	NextGPUMinVoxelWorldSize[0] = RoundHairVoxelSize(NextVoxelWorldSize);
+}
+
+#endif // SHADER_ADAPTIVE_FEEDBACK
diff --git a/Engine/Shaders/Private/HairStrands/HairStrandsVoxelPageRayMarching.usf b/Engine/Shaders/Private/HairStrands/HairStrandsVoxelPageRayMarching.usf
--- a/Engine/Shaders/Private/HairStrands/HairStrandsVoxelPageRayMarching.usf
+++ b/Engine/Shaders/Private/HairStrands/HairStrandsVoxelPageRayMarching.usf
@@ -97,8 +97,8 @@ void MainCS(uint3 DispatchThreadId : SV_DispatchThreadID)
 			float2 Pos = StartPos;
 
 			// Adaptive voxels
-			Pos = ShaderPrintText(Pos, TEXT("Adapative "), FontOrange);
-			if (VirtualVoxel.AdaptiveEnable)
+			Pos = ShaderPrintText(Pos, TEXT("Alloc. Feedback  "), FontOrange);
+			if (VirtualVoxel.AllocationFeedbackEnable)
 			{
 				Pos = ShaderPrintText(Pos, TEXT("True"), FontGreen);
 			}
@@ -108,31 +108,30 @@ void MainCS(uint3 DispatchThreadId : SV_DispatchThreadID)
 			}
 			Pos = ShaderPrintNewline(Pos); Pos.x = StartPos.x;
 
-			Pos = ShaderPrintText(Pos, TEXT("Requested Voxel Size "));
-			Pos = ShaderPrint(Pos, VirtualVoxel.AdaptiveRequestedVoxelWorldSize);
+			// Actual voxel/pages size & resolution
+			Pos = ShaderPrintText(Pos, TEXT("GPU Curr. Voxel Size        "), FontOrange);
+			Pos = ShaderPrint(Pos, VirtualVoxel.CurrGPUMinVoxelSize[0], FontEmerald);
 			Pos = ShaderPrintNewline(Pos); Pos.x = StartPos.x;
 
-			Pos = ShaderPrintText(Pos, TEXT("Target Voxel Size "));
-			Pos = ShaderPrint(Pos, VirtualVoxel.AdaptiveTargetVoxelWorldSize);
-			Pos = ShaderPrintNewline(Pos); Pos.x = StartPos.x;
+			Pos = ShaderPrintText(Pos, TEXT("GPU Next. Voxel Size        "));
+			Pos = ShaderPrint(Pos, VirtualVoxel.NextGPUMinVoxelSize[0]);
 			Pos = ShaderPrintNewline(Pos); Pos.x = StartPos.x;
 
-			// Actual voxel/pages size & resolution
-			Pos = ShaderPrintText(Pos, TEXT("Min. Voxel Size  "), FontOrange);
-			Pos = ShaderPrint(Pos, VirtualVoxel.MinVoxelWorldSize, FontEmerald);
+			Pos = ShaderPrintText(Pos, TEXT("CPU Min. Voxel Size         "));
+			Pos = ShaderPrint(Pos, VirtualVoxel.CPUMinVoxelWorldSize);
 			Pos = ShaderPrintNewline(Pos); Pos.x = StartPos.x;
 			
-			Pos = ShaderPrintText(Pos, TEXT("Page Res         "));
+			Pos = ShaderPrintText(Pos, TEXT("Page Res                    "));
 			Pos = ShaderPrint(Pos, VirtualVoxel.PageResolution);
 			Pos = ShaderPrintNewline(Pos); Pos.x = StartPos.x;
 
-			Pos = ShaderPrintText(Pos, TEXT("Page Count Res   "));
+			Pos = ShaderPrintText(Pos, TEXT("Page Count Res              "));
 			Pos = ShaderPrint(Pos, VirtualVoxel.PageCountResolution.x);
 			Pos = ShaderPrint(Pos, VirtualVoxel.PageCountResolution.y);
 			Pos = ShaderPrint(Pos, VirtualVoxel.PageCountResolution.z);
 			Pos = ShaderPrintNewline(Pos); Pos.x = StartPos.x;
 
-			Pos = ShaderPrintText(Pos, TEXT("Page Texture Res "), FontEmerald);
+			Pos = ShaderPrintText(Pos, TEXT("Page Texture Res            "), FontEmerald);
 			Pos = ShaderPrint(Pos, VirtualVoxel.PageTextureResolution.x, FontEmerald);
 			Pos = ShaderPrint(Pos, VirtualVoxel.PageTextureResolution.y, FontEmerald);
 			Pos = ShaderPrint(Pos, VirtualVoxel.PageTextureResolution.z, FontEmerald);
@@ -208,6 +207,8 @@ void MainCS(uint3 DispatchThreadId : SV_DispatchThreadID)
 			Pos = ShaderPrintText(Pos, TEXT(" Page Index Offset "));
 			Pos = ShaderPrintNewline(Pos); Pos.x = StartPos.x;
 
+			Pos = ShaderPrintNewline(Pos); Pos.x = StartPos.x;
+
 			Pos = ShaderPrintText(Pos, TEXT(" Voxel World Size  "));
 			Pos = ShaderPrintNewline(Pos); Pos.x = StartPos.x;
 		}
@@ -237,9 +238,16 @@ void MainCS(uint3 DispatchThreadId : SV_DispatchThreadID)
 
 			Pos = ShaderPrintNewline(Pos); Pos.x = StartPos.x;
 
+			// Offset of 1 line for every other voxel size, to avoid super-position with large float number
+			const bool bAddLine = (MacroGroupId & 1) == 1;
+
+			if (bAddLine) { Pos = ShaderPrintNewline(Pos); Pos.x = StartPos.x; }
+
 			Pos = ShaderPrint(Pos, NodeDesc.PageIndexOffset, FontSilver);
 			Pos = ShaderPrintNewline(Pos); Pos.x = StartPos.x;
 
+			Pos = ShaderPrintNewline(Pos); Pos.x = StartPos.x;
+
 			Pos = ShaderPrint(Pos, NodeDesc.VoxelWorldSize, FontSilver);
 			Pos = ShaderPrintNewline(Pos); Pos.x = StartPos.x;
 		}
diff --git a/Engine/Source/Runtime/Renderer/Private/HairStrands/HairStrandsCluster.cpp b/Engine/Source/Runtime/Renderer/Private/HairStrands/HairStrandsCluster.cpp
--- a/Engine/Source/Runtime/Renderer/Private/HairStrands/HairStrandsCluster.cpp
+++ b/Engine/Source/Runtime/Renderer/Private/HairStrands/HairStrandsCluster.cpp
@@ -28,7 +28,6 @@ class FHairMacroGroupAABBCS : public FGlobalShader
 
 		SHADER_PARAMETER(float, PixelSizeAtDepth1)
 		SHADER_PARAMETER(float, NumPixelPerVoxel)
-		SHADER_PARAMETER(float, MinVoxelWorldSize)
 
 		SHADER_PARAMETER_RDG_BUFFER_SRV(Buffer, InGroupAABBBuffer0)
 		SHADER_PARAMETER_RDG_BUFFER_SRV(Buffer, InGroupAABBBuffer1)
@@ -71,7 +70,6 @@ static void AddHairMacroGroupAABBPass(
 	const float PixelSizeAtDepth1 = FMath::Tan(vFOV * 0.5f) / (0.5f * Resolution.Y);
 	
 	const float NumPixelPerVoxel = FMath::Clamp(GHairVirtualVoxel_NumPixelPerVoxel, 1.f, 50.f);
-	const float MinVoxelWorldSize = 0.01f;
 
 	const uint32 GroupPerPass = 8;
 	bool bNeedClear = true;
@@ -83,7 +81,6 @@ static void AddHairMacroGroupAABBPass(
 		Parameters->MacroGroupId = MacroGroupId;
 		Parameters->OutMacroGroupAABBBuffer = OutHairMacroGroupAABBBufferUAV;
 		Parameters->OutMacroGroupVoxelSizeBuffer = OutMacroGroupVoxelSizeBufferUAV;
-		Parameters->MinVoxelWorldSize = MinVoxelWorldSize;
 		Parameters->PixelSizeAtDepth1 = PixelSizeAtDepth1;
 		Parameters->NumPixelPerVoxel  = NumPixelPerVoxel;
 		Parameters->View			  = View.ViewUniformBuffer;
@@ -130,7 +127,7 @@ static void AddHairMacroGroupAABBPass(
 		TShaderMapRef<FHairMacroGroupAABBCS> ComputeShader(View.ShaderMap);
 		FComputeShaderUtils::AddPass(
 			GraphBuilder,
-			RDG_EVENT_NAME("HairStrandsMacroGroupAABBUpdate"),
+			RDG_EVENT_NAME("HairStrands::MacroGroupAABBUpdate"),
 			ComputeShader,
 			Parameters,
 			FIntVector(1,1,1));
diff --git a/Engine/Source/Runtime/Renderer/Private/HairStrands/HairStrandsData.h b/Engine/Source/Runtime/Renderer/Private/HairStrands/HairStrandsData.h
--- a/Engine/Source/Runtime/Renderer/Private/HairStrands/HairStrandsData.h
+++ b/Engine/Source/Runtime/Renderer/Private/HairStrands/HairStrandsData.h
@@ -140,7 +140,7 @@ struct FPackedVirtualVoxelNodeDesc
 // at the moment
 BEGIN_SHADER_PARAMETER_STRUCT(FHairStrandsVoxelCommonParameters, )
 	SHADER_PARAMETER(FIntVector, PageCountResolution)
-	SHADER_PARAMETER(float, MinVoxelWorldSize)
+	SHADER_PARAMETER(float, CPUMinVoxelWorldSize)
 	SHADER_PARAMETER(FIntVector, PageTextureResolution)
 	SHADER_PARAMETER(uint32, PageCount)
 	SHADER_PARAMETER(uint32, PageResolution)
@@ -172,14 +172,16 @@ BEGIN_SHADER_PARAMETER_STRUCT(FHairStrandsVoxelCommonParameters, )
 
 	SHADER_PARAMETER(FVector3f, TranslatedWorldOffset) // For debug purpose
 
-	SHADER_PARAMETER(uint32, AdaptiveEnable)
-	SHADER_PARAMETER(float, AdaptiveRequestedVoxelWorldSize)
-	SHADER_PARAMETER(float, AdaptiveTargetVoxelWorldSize)
+	SHADER_PARAMETER(uint32, AllocationFeedbackEnable)
 
 	SHADER_PARAMETER_RDG_BUFFER_SRV(Buffer<uint>, PageIndexBuffer)
 	SHADER_PARAMETER_RDG_BUFFER_SRV(Buffer<uint2>, PageIndexOccupancyBuffer)
 	SHADER_PARAMETER_RDG_BUFFER_SRV(Buffer<uint>, PageIndexCoordBuffer)
 	SHADER_PARAMETER_RDG_BUFFER_SRV(StructuredBuffer<FPackedVirtualVoxelNodeDesc>, NodeDescBuffer) // Packed into 2 x uint4
+
+	SHADER_PARAMETER_RDG_BUFFER_SRV(Buffer<float>, CurrGPUMinVoxelSize)
+	SHADER_PARAMETER_RDG_BUFFER_SRV(Buffer<float>, NextGPUMinVoxelSize)
+
 END_SHADER_PARAMETER_STRUCT()
 
 BEGIN_GLOBAL_SHADER_PARAMETER_STRUCT(FVirtualVoxelParameters, RENDERER_API)
@@ -395,21 +397,15 @@ struct FHairStrandsViewData
 	TArray<FSphere> VisibleShadowCastingBounds;
 };
 
-// View State data (i.e., persistent accross fram)
+// View State data (i.e., persistent across frame)
 struct FHairStrandsViewStateData
 {
-	FRHIGPUBufferReadback* GetBuffer() const { return VoxelPageAllocationCountReadback; }
-	bool IsReady() const { return VoxelPageAllocationCountReadback->IsReady(); }
-	bool IsInit() const { return VoxelPageAllocationCountReadback != nullptr; }
-
+	bool IsInit() const { return PositionsChangedDatas.Num() > 0; }
 	void Init();
 	void Release();
 
-	float MinVoxelWorldSize = 0; // Min voxel size used during the last frame allocation
-	uint32 VoxelAllocatedPageCount = 0; // Number of voxels allocated last frame
-
-	// Buffer used for reading back the number of voxels allocated on the GPU
-	FRHIGPUBufferReadback* VoxelPageAllocationCountReadback = nullptr;
+	// Buffer used for reading back the number of allocated voxels on the GPU
+	TRefCountPtr<FRDGPooledBuffer> VoxelFeedbackBuffer = nullptr;
 
 	// Buffer used for reading back hair strands position changed on the GPU
 	struct FPositionChangedData
diff --git a/Engine/Source/Runtime/Renderer/Private/HairStrands/HairStrandsRendering.cpp b/Engine/Source/Runtime/Renderer/Private/HairStrands/HairStrandsRendering.cpp
--- a/Engine/Source/Runtime/Renderer/Private/HairStrands/HairStrandsRendering.cpp
+++ b/Engine/Source/Runtime/Renderer/Private/HairStrands/HairStrandsRendering.cpp
@@ -85,7 +85,6 @@ static TRDGUniformBufferRef<FHairStrandsViewUniformParameters> InternalCreateHai
 IMPLEMENT_GLOBAL_SHADER_PARAMETER_STRUCT(FHairStrandsViewUniformParameters, "HairStrands");
 
 bool GetHairStrandsSkyLightingDebugEnable();
-bool IsHairStrandsAdaptiveVoxelAllocationEnable();
 void AddMeshDrawTransitionPass(
 	FRDGBuilder& GraphBuilder,
 	const FViewInfo& ViewInfo,
@@ -118,21 +117,10 @@ void RenderHairPrePass(
 
 		const ERHIFeatureLevel::Type FeatureLevel = Scene->GetFeatureLevel();
 
-		// Allocate voxel page allocation readback buffers
-		const bool bAdaptiveAllocationEnable = IsHairStrandsAdaptiveVoxelAllocationEnable();
-		if (View.ViewState)
+		// Allocate state/feedback data if needed
+		if (View.ViewState && !View.ViewState->HairStrandsViewStateData.IsInit())
 		{
-			const bool bIsInit = View.ViewState->HairStrandsViewStateData.IsInit();
-			// Init resources if the adaptive allocation is enabled, but the resources are not initialized yet.
-			if (bAdaptiveAllocationEnable && !bIsInit)
-			{
-				View.ViewState->HairStrandsViewStateData.Init();
-			}
-			// Release resources if the adaptive allocation is disabled, but the resources are initialized.
-			else if (!bAdaptiveAllocationEnable && bIsInit)
-			{
-				View.ViewState->HairStrandsViewStateData.Release();
-			}
+			View.ViewState->HairStrandsViewStateData.Init();
 		}
 
 		//SCOPED_GPU_STAT(RHICmdList, HairRendering);
@@ -199,12 +187,7 @@ void RenderHairBasePass(
 void FHairStrandsViewStateData::Init()
 {
 	// Voxel adaptive sizing
-	MinVoxelWorldSize = 0;
-	VoxelAllocatedPageCount = 0;
-	if (VoxelPageAllocationCountReadback == nullptr)
-	{
-		VoxelPageAllocationCountReadback = new FRHIGPUBufferReadback(TEXT("Hair.VoxelPageAllocationReadback"));
-	}
+	VoxelFeedbackBuffer = nullptr;
 	
 	// Track if hair strands positions has changed
 	PositionsChangedDatas.SetNum(4);
@@ -218,13 +201,7 @@ void FHairStrandsViewStateData::Init()
 void FHairStrandsViewStateData::Release()
 {
 	// Voxel adaptive sizing
-	MinVoxelWorldSize = 0;
-	VoxelAllocatedPageCount = 0;
-	if (VoxelPageAllocationCountReadback)
-	{
-		delete VoxelPageAllocationCountReadback;
-		VoxelPageAllocationCountReadback = nullptr;
-	}
+	VoxelFeedbackBuffer = nullptr;
 
 	// Track if hair strands positions has changed
 	for (FPositionChangedData& Data : PositionsChangedDatas)
diff --git a/Engine/Source/Runtime/Renderer/Private/HairStrands/HairStrandsVoxelization.cpp b/Engine/Source/Runtime/Renderer/Private/HairStrands/HairStrandsVoxelization.cpp
--- a/Engine/Source/Runtime/Renderer/Private/HairStrands/HairStrandsVoxelization.cpp
+++ b/Engine/Source/Runtime/Renderer/Private/HairStrands/HairStrandsVoxelization.cpp
@@ -10,6 +10,7 @@
 #include "ShaderParameters.h"
 #include "ShaderParameterStruct.h"
 #include "SceneTextureParameters.h"
+#include "ShaderPrintParameters.h"
 #include "RenderGraphUtils.h"
 #include "PostProcessing.h"
 #include "ScenePrivate.h"
@@ -118,7 +119,7 @@ static FAutoConsoleVariableRef CVarHairVirtualVoxelAdaptive_CorrectionThreshold(
 static int32 GHairVirtualVoxel_JitterMode = 1;
 static FAutoConsoleVariableRef CVarHairVirtualVoxel_JitterMode(TEXT("r.HairStrands.Voxelization.Virtual.Jitter"), GHairVirtualVoxel_JitterMode, TEXT("Change jittered for voxelization/traversal. 0: No jitter 1: Regular randomized jitter: 2: Constant Jitter (default = 1)"), ECVF_RenderThreadSafe);
 
-bool IsHairStrandsAdaptiveVoxelAllocationEnable()
+static bool IsHairStrandsAdaptiveVoxelAllocationEnable()
 {
 	return GHairVirtualVoxelAdaptive_Enable > 0;
 }
@@ -226,6 +227,7 @@ class FVoxelAllocatePageIndexCS : public FGlobalShader
 		SHADER_PARAMETER(FIntVector, CPU_PageIndexResolution)
 		SHADER_PARAMETER(uint32, CPU_bUseCPUData)
 
+		SHADER_PARAMETER_RDG_BUFFER_SRV(Buffer<float>, GPUVoxelWorldSize)
 		SHADER_PARAMETER_RDG_BUFFER_SRV(Buffer<float>, MacroGroupVoxelSizeBuffer)
 		SHADER_PARAMETER_RDG_BUFFER_UAV(RWBuffer<int>, MacroGroupAABBBuffer)
 		SHADER_PARAMETER_RDG_BUFFER_UAV(RWBuffer<uint4>, OutPageIndexResolutionAndOffsetBuffer)
@@ -391,6 +393,7 @@ class FVoxelAddNodeDescCS : public FGlobalShader
 		SHADER_PARAMETER(uint32, MacroGroupId)
 		SHADER_PARAMETER(FIntVector, CPU_PageIndexResolution)
 		SHADER_PARAMETER(float, CPU_VoxelWorldSize)
+		SHADER_PARAMETER_RDG_BUFFER_SRV(Buffer, GPU_VoxelWorldSize)
 		SHADER_PARAMETER_RDG_BUFFER_SRV(Buffer, MacroGroupAABBBuffer)
 		SHADER_PARAMETER_RDG_BUFFER_SRV(Buffer, MacroGroupVoxelSizeBuffer)
 		SHADER_PARAMETER_RDG_BUFFER_SRV(Buffer, PageIndexResolutionAndOffsetBuffer)
@@ -470,6 +473,33 @@ public:
 	}
 };
 
+class FVoxelAdaptiveFeedbackCS : public FGlobalShader
+{
+	DECLARE_GLOBAL_SHADER(FVoxelAdaptiveFeedbackCS);
+	SHADER_USE_PARAMETER_STRUCT(FVoxelAdaptiveFeedbackCS, FGlobalShader);
+
+	BEGIN_SHADER_PARAMETER_STRUCT(FParameters, )
+		SHADER_PARAMETER(uint32, CPUAllocatedPageCount)
+		SHADER_PARAMETER(float,  CPUMinVoxelWorldSize)
+		SHADER_PARAMETER(float,  AdaptiveCorrectionThreshold)
+		SHADER_PARAMETER(float,  AdaptiveCorrectionSpeed)
+
+		SHADER_PARAMETER_RDG_BUFFER_SRV(Buffer, TotalRequestedPageAllocationBuffer)
+		SHADER_PARAMETER_RDG_BUFFER_SRV(Buffer, CurrGPUMinVoxelWorldSize)
+		SHADER_PARAMETER_RDG_BUFFER_UAV(RWBuffer, NextGPUMinVoxelWorldSize)
+
+		SHADER_PARAMETER_STRUCT_INCLUDE(ShaderPrint::FShaderParameters, ShaderPrintUniformBuffer)
+	END_SHADER_PARAMETER_STRUCT()
+
+public:
+	static bool ShouldCompilePermutation(const FGlobalShaderPermutationParameters& Parameters) { return IsHairStrandsSupported(EHairStrandsShaderType::Strands, Parameters.Platform); }
+	static void ModifyCompilationEnvironment(const FGlobalShaderPermutationParameters& Parameters, FShaderCompilerEnvironment& OutEnvironment)
+	{
+		FGlobalShader::ModifyCompilationEnvironment(Parameters, OutEnvironment);
+		OutEnvironment.SetDefine(TEXT("SHADER_ADAPTIVE_FEEDBACK"), 1);
+	}
+};
+
 IMPLEMENT_GLOBAL_SHADER(FVoxelMarkValidPageIndex_PrepareCS, "/Engine/Private/HairStrands/HairStrandsVoxelPageAllocation.usf", "MarkValid_PrepareCS", SF_Compute);
 IMPLEMENT_GLOBAL_SHADER(FVoxelMarkValidPageIndex_IndirectArgsCS, "/Engine/Private/HairStrands/HairStrandsVoxelPageAllocation.usf", "MarkValid_BuildIndirectArgCS", SF_Compute);
 IMPLEMENT_GLOBAL_SHADER(FVoxelMarkValidPageIndex_ScatterCS, "/Engine/Private/HairStrands/HairStrandsVoxelPageAllocation.usf", "MarkValid_ScatterCS", SF_Compute);
@@ -481,6 +511,7 @@ IMPLEMENT_GLOBAL_SHADER(FVoxelAddNodeDescCS, "/Engine/Private/HairStrands/HairSt
 IMPLEMENT_GLOBAL_SHADER(FVoxelAddIndirectBufferCS, "/Engine/Private/HairStrands/HairStrandsVoxelPageAllocation.usf", "AddIndirectBufferCS", SF_Compute);
 IMPLEMENT_GLOBAL_SHADER(FVoxelIndPageClearBufferGenCS, "/Engine/Private/HairStrands/HairStrandsVoxelPageAllocation.usf", "VoxelIndPageClearBufferGenCS", SF_Compute);
 IMPLEMENT_GLOBAL_SHADER(FVoxelIndPageClearCS, "/Engine/Private/HairStrands/HairStrandsVoxelPageAllocation.usf", "VoxelIndPageClearCS", SF_Compute);
+IMPLEMENT_GLOBAL_SHADER(FVoxelAdaptiveFeedbackCS, "/Engine/Private/HairStrands/HairStrandsVoxelPageAllocation.usf", "FeedbackCS", SF_Compute);
 
 inline FIntVector CeilToInt(const FVector& V)
 {
@@ -494,7 +525,7 @@ static void AddAllocateVoxelPagesPass(
 	FHairStrandsMacroGroupResources& MacroGroupResources,
 	const FIntVector PageCountResolution,
 	const uint32 PageCount,
-	const float MinVoxelWorldSize,
+	const float CPUMinVoxelWorldSize,
 	const uint32 PageResolution,
 	const FIntVector PageTextureResolution,
 	const uint32 IndirectDispatchGroupSize,
@@ -505,9 +536,10 @@ static void AddAllocateVoxelPagesPass(
 	FRDGBufferRef& OutPageIndexCoordBuffer,
 	FRDGBufferRef& OutNodeDescBuffer,
 	FRDGBufferRef& OutIndirectArgsBuffer,
-	FRDGBufferRef& OutPageIndexGlobalCounter, 
+	FRDGBufferRef& OutPageIndexGlobalCounter,
 	FRDGBufferRef& OutVoxelizationViewInfoBuffer,
-	FRDGBufferRef& OutTotalRequestedPageAllocationBuffer)
+	FRDGBufferRef& CurrGPUMinVoxelSize,
+	FRDGBufferRef& NextGPUMinVoxelSize)
 {
 	const uint32 GroupSize = 32;
 	const bool bIsGPUDriven = GHairVirtualVoxelGPUDriven > 0;
@@ -524,7 +556,7 @@ static void AddAllocateVoxelPagesPass(
 		uint32		PageIndexOffset;
 		uint32		MacroGroupId;
 	};
-	const float CPUPageWorldSize = PageResolution * MinVoxelWorldSize;
+	const float CPUPageWorldSize = PageResolution * CPUMinVoxelWorldSize;
 
 	OutTotalPageIndexCount = 0;
 	TArray<FCPUMacroGroupAllocation> CPUAllocationDescs;	
@@ -587,16 +619,16 @@ static void AddAllocateVoxelPagesPass(
 	FRDGBufferRef PageIndexAllocationIndirectBufferArgs = GraphBuilder.CreateBuffer(FRDGBufferDesc::CreateIndirectDesc<FRHIDispatchIndirectParameters>(MacroGroupCount), TEXT("Hair.PageIndexAllocationIndirectBufferArgs"));
 
 	// Store the total requested page allocation (for feedback purpose)
-	FRDGBufferDesc TotalRequestDesc = FRDGBufferDesc::CreateBufferDesc(sizeof(uint32), 1);
-	TotalRequestDesc.Usage |= BUF_SourceCopy;
-	OutTotalRequestedPageAllocationBuffer = GraphBuilder.CreateBuffer(TotalRequestDesc, TEXT("TotalRequestedPageAllocationBuffer"));
-	FRDGBufferUAVRef TotalRequestedPageAllocationBufferUAV = GraphBuilder.CreateUAV(OutTotalRequestedPageAllocationBuffer, PF_R32_UINT);
+	FRDGBufferRef TotalRequestedPageAllocationBuffer = GraphBuilder.CreateBuffer(FRDGBufferDesc::CreateBufferDesc(sizeof(uint32), 1), TEXT("Hair.TotalRequestedPageAllocationBuffer"));
+	FRDGBufferUAVRef TotalRequestedPageAllocationBufferUAV = GraphBuilder.CreateUAV(TotalRequestedPageAllocationBuffer, PF_R32_UINT);
 	AddClearUAVPass(GraphBuilder, TotalRequestedPageAllocationBufferUAV, 0u);
 	
 	AddClearUAVPass(GraphBuilder, PageIndexBufferUAV, 0u);
 	AddClearUAVPass(GraphBuilder, PageIndexOccupancyBufferUAV, 0u);
 	AddClearUAVPass(GraphBuilder, PageIndexGlobalCounterUAV, 0u);
 
+	const FRDGBufferSRVRef CurrGPUMinVoxelSizeSRV = GraphBuilder.CreateSRV(CurrGPUMinVoxelSize, PF_R16F);
+
 	// Allocate page index for all instance group
 //	if (bIsGPUDriven)
 	{
@@ -611,7 +643,7 @@ static void AddAllocateVoxelPagesPass(
 
 		FVoxelAllocatePageIndexCS::FParameters* Parameters = GraphBuilder.AllocParameters<FVoxelAllocatePageIndexCS::FParameters>();
 		Parameters->CPUPageWorldSize = CPUPageWorldSize;
-		Parameters->CPUVoxelWorldSize = MinVoxelWorldSize;
+		Parameters->CPUVoxelWorldSize = CPUMinVoxelWorldSize;
 		Parameters->TotalPageIndexCount = OutTotalPageIndexCount;
 		Parameters->PageResolution = PageResolution;
 		Parameters->MacroGroupCount = MacroGroupCount;
@@ -623,6 +655,7 @@ static void AddAllocateVoxelPagesPass(
 		Parameters->OutVoxelizationViewInfoBuffer = GraphBuilder.CreateUAV(VoxelizationViewInfoBuffer);
 		Parameters->OutPageIndexAllocationIndirectBufferArgs = GraphBuilder.CreateUAV(PageIndexAllocationIndirectBufferArgs);
 		Parameters->CPU_bUseCPUData = GHairVirtualVoxelGPUDriven == 2 ? 1 : 0;
+		Parameters->GPUVoxelWorldSize = CurrGPUMinVoxelSizeSRV;
 		if (Parameters->CPU_bUseCPUData)
 		{
 			Parameters->CPU_TranslatedWorldMinAABB = CPUAllocationDescs.Num() > 0 ? (FVector3f)CPUAllocationDescs[0].TranslatedMinAABB : FVector3f::ZeroVector;
@@ -815,7 +848,8 @@ static void AddAllocateVoxelPagesPass(
 			Parameters->CPU_TranslatedWorldMaxAABB = (FVector3f)CPUAllocationDesc.TranslatedMaxAABB;
 			Parameters->CPU_PageIndexResolution = CPUAllocationDesc.PageIndexResolution;
 			Parameters->CPU_PageIndexOffset = CPUAllocationDesc.PageIndexOffset;
-			Parameters->CPU_VoxelWorldSize = MinVoxelWorldSize;
+			Parameters->CPU_VoxelWorldSize = CPUMinVoxelWorldSize;
+			Parameters->GPU_VoxelWorldSize = CurrGPUMinVoxelSizeSRV;
 			Parameters->MacroGroupVoxelSizeBuffer = GraphBuilder.CreateSRV(MacroGroupResources.MacroGroupVoxelSizeBuffer, PF_R16F);
 			Parameters->OutNodeDescBuffer = GraphBuilder.CreateUAV(NodeDescBuffer);
 
@@ -881,12 +915,33 @@ static void AddAllocateVoxelPagesPass(
 			Parameters->OutIndirectArgsBuffer = GraphBuilder.CreateUAV(IndirectArgsBuffer);
 			Parameters->OutTotalRequestedPageAllocationBuffer = TotalRequestedPageAllocationBufferUAV;
 
-			const FIntVector DispatchCount(1, 1, 1);
 			TShaderMapRef<FVoxelAddIndirectBufferCS> ComputeShader(View.ShaderMap);
-			FComputeShaderUtils::AddPass(GraphBuilder, RDG_EVENT_NAME("HairStrands::BuildVoxelIndirectArgs"), ComputeShader, Parameters, DispatchCount);
+			FComputeShaderUtils::AddPass(GraphBuilder, RDG_EVENT_NAME("HairStrands::BuildVoxelIndirectArgs"), ComputeShader, Parameters, FIntVector(1, 1, 1));
 		}
 	}
 
+	// Feedback allocation for next frame
+	const bool bAdaptiveVoxelEnable = IsHairStrandsAdaptiveVoxelAllocationEnable();
+	if (bAdaptiveVoxelEnable)
+	{
+		// Store the total requested page allocation (for feedback purpose)
+		const EPixelFormat Format = PF_R16F;
+		NextGPUMinVoxelSize = GraphBuilder.CreateBuffer(FRDGBufferDesc::CreateBufferDesc(sizeof(FFloat16), 1), TEXT("Hair.VoxelPageAllocationFeedbackBuffer"));
+
+		FVoxelAdaptiveFeedbackCS::FParameters* Parameters = GraphBuilder.AllocParameters<FVoxelAdaptiveFeedbackCS::FParameters>();
+		Parameters->CPUMinVoxelWorldSize = CPUMinVoxelWorldSize;
+		Parameters->CPUAllocatedPageCount = TotalPageCount;
+		Parameters->AdaptiveCorrectionThreshold = FMath::Clamp(GHairVirtualVoxelAdaptive_CorrectionThreshold, 0.f, 1.f);
+		Parameters->AdaptiveCorrectionSpeed = FMath::Clamp(GHairVirtualVoxelAdaptive_CorrectionSpeed, 0.f, 1.f);
+		Parameters->TotalRequestedPageAllocationBuffer = GraphBuilder.CreateSRV(TotalRequestedPageAllocationBuffer, PF_R32_UINT);
+		Parameters->CurrGPUMinVoxelWorldSize = CurrGPUMinVoxelSizeSRV;
+		Parameters->NextGPUMinVoxelWorldSize = GraphBuilder.CreateUAV(NextGPUMinVoxelSize, Format);
+		ShaderPrint::SetParameters(GraphBuilder, View, Parameters->ShaderPrintUniformBuffer);
+
+		TShaderMapRef<FVoxelAdaptiveFeedbackCS> ComputeShader(View.ShaderMap);
+		FComputeShaderUtils::AddPass(GraphBuilder, RDG_EVENT_NAME("HairStrands::RunVoxelSizeFeedback"), ComputeShader, Parameters, FIntVector(1, 1, 1));
+	}
+
 	OutPageIndexBuffer = PageIndexBuffer;
 	OutPageIndexOccupancyBuffer = PageIndexOccupancyBuffer;
 	OutPageToPageIndexBuffer = PageToPageIndexBuffer;
@@ -921,93 +976,14 @@ static FHairStrandsVoxelResources AllocateVirtualVoxelResources(
 	uint32 PageCount = PageCountResolution.X * PageCountResolution.Y * PageCountResolution.Z;
 	const float ClampMinVoxelWorldSize = 0.01f;
 	const float ClampMaxVoxelWorldSize = 10.f;
-
-	const float RequestedVoxelWorldSize = RoundHairVoxeliSize(FMath::Clamp(GHairVirtualVoxel_VoxelWorldSize, ClampMinVoxelWorldSize, ClampMaxVoxelWorldSize));
-	float TargetVoxelWorldSize = RequestedVoxelWorldSize;
-	float MinVoxelWorldSize = RequestedVoxelWorldSize;
-
-	// Readback allocated value to adapt the voxel size in order to fit max page allocation
-	const bool bAdaptiveResolution = OutViewStateData && OutViewStateData->IsInit();
-	FRHIGPUBufferReadback* ReadbackBuffer = nullptr;
-	bool bEnqueueNewReadbackQuery = false;
-	if (bAdaptiveResolution)
-	{
-		// First initialization (no query has been issued yet)
-		if (OutViewStateData->MinVoxelWorldSize == 0)
-		{
-			OutViewStateData->MinVoxelWorldSize = MinVoxelWorldSize;
-			bEnqueueNewReadbackQuery = true;
-		}
-
-		uint32 AllocatedPageCount = PageCount;
-		ReadbackBuffer = OutViewStateData->GetBuffer();
-		if (ReadbackBuffer->IsReady())
-		{
-			bEnqueueNewReadbackQuery = true;
-
-			ReadbackBuffer = OutViewStateData->GetBuffer();
-			AllocatedPageCount = *(uint32*)(ReadbackBuffer->Lock(sizeof(uint32)));
-			ReadbackBuffer->Unlock();
-
-			// ...
-			const float Threshold = FMath::Clamp(GHairVirtualVoxelAdaptive_CorrectionThreshold, 0.f, 1.f);
-			const float Factor = FMath::Clamp(GHairVirtualVoxelAdaptive_CorrectionSpeed, 0.f, 1.f);
-
-			// Voxel pages are represent a volume. To derive a better estimate of the ratio by which voxel size needs to be scale, 
-			// compute the cubic root of this ratio.
-			//
-			// AllocatedPage   AllocatedRes^3
-			// ------------- = --------------  = VolumeRatio = LinearRatio^3
-			//    MaxPage          MaxRes^3
-
-			// Ratio used for predicting voxel size increase
-			const float VolumeRatio = float(AllocatedPageCount) / float(PageCount);
-			const float LinearRatio = FMath::Pow(VolumeRatio, 1.f / 3.f);
-
-			// Ratio used for predicting voxel size decrease (i.e. when requested allocation fit, 
-			// but the voxel size does not match the (more precise) target).
-			// In this case, we add a threshold/margin to to the target, so that there is no oscillation.
-			const float VolumeRatio_Thres = float(AllocatedPageCount) / float(PageCount * Threshold);
-			const float LinearRatio_Thres = FMath::Pow(VolumeRatio_Thres, 1.f / 3.f);
-
-			const float PrevWorldVoxelSize = RoundHairVoxeliSize(OutViewStateData->MinVoxelWorldSize);
-			
-			// If the page pool is not large enough increase voxel size
-			if (AllocatedPageCount > PageCount)
-			{
-				MinVoxelWorldSize = PrevWorldVoxelSize * LinearRatio;
-			}
-			// If the page pool is large enough but the voxel are larger than the requested size decrease voxel size
-			else if (AllocatedPageCount < PageCount && PrevWorldVoxelSize > MinVoxelWorldSize)
-			{
-				TargetVoxelWorldSize = PrevWorldVoxelSize * LinearRatio_Thres;
-				MinVoxelWorldSize = FMath::Max(MinVoxelWorldSize, FMath::Lerp(PrevWorldVoxelSize, TargetVoxelWorldSize, Factor));
-			}
-			else
-			{
-				MinVoxelWorldSize = PrevWorldVoxelSize;
-			}
-
-			// Clamp voxel size into a reasonable amount (e.g. 0.1mm - 100mm)
-			MinVoxelWorldSize = FMath::Clamp(MinVoxelWorldSize, ClampMinVoxelWorldSize, ClampMaxVoxelWorldSize);
-		}
-		else
-		{
-			// Use previous frame prediction by default (a readback is currently in-flight, but not ready for this frame)
-			MinVoxelWorldSize = OutViewStateData->MinVoxelWorldSize;
-		}
-		MinVoxelWorldSize = RoundHairVoxeliSize(MinVoxelWorldSize);
-
-		// Update state data
-		OutViewStateData->MinVoxelWorldSize = MinVoxelWorldSize;
-		OutViewStateData->VoxelAllocatedPageCount = AllocatedPageCount;
-	}
+	const float CPUMinVoxelWorldSize = RoundHairVoxeliSize(FMath::Clamp(GHairVirtualVoxel_VoxelWorldSize, ClampMinVoxelWorldSize, ClampMaxVoxelWorldSize));
+	const bool bVoxelAllocationFeedbackEnable = OutViewStateData && IsHairStrandsAdaptiveVoxelAllocationEnable();
 
 	FHairStrandsVoxelResources Out;
 
 	Out.Parameters.Common.PageCountResolution		= PageCountResolution;
 	Out.Parameters.Common.PageCount					= PageCount;
-	Out.Parameters.Common.MinVoxelWorldSize			= MinVoxelWorldSize;
+	Out.Parameters.Common.CPUMinVoxelWorldSize		= CPUMinVoxelWorldSize;
 	Out.Parameters.Common.PageResolutionLog2		= FMath::CeilLogTwo(FMath::Clamp(GHairVirtualVoxel_PageResolution, 2, 256));
 	Out.Parameters.Common.PageResolution			= (1u << Out.Parameters.Common.PageResolutionLog2);
 	Out.Parameters.Common.PageTextureResolution		= Out.Parameters.Common.PageCountResolution * Out.Parameters.Common.PageResolution;
@@ -1037,11 +1013,16 @@ static FHairStrandsVoxelResources AllocateVirtualVoxelResources(
 	Out.Parameters.Common.TranslatedWorldOffset					= FVector3f(View.ViewMatrices.GetPreViewTranslation());
 
 	// For debug purpose
-	Out.Parameters.Common.AdaptiveEnable					= bAdaptiveResolution ? 1u : 0u;
-	Out.Parameters.Common.AdaptiveRequestedVoxelWorldSize	= RequestedVoxelWorldSize;
-	Out.Parameters.Common.AdaptiveTargetVoxelWorldSize		= TargetVoxelWorldSize;
+	Out.Parameters.Common.AllocationFeedbackEnable = bVoxelAllocationFeedbackEnable ? 1u : 0u;
+
+	FRDGBufferRef CurrGPUMinVoxelSize = nullptr;
+	FRDGBufferRef NextGPUMinVoxelSize = nullptr;
+	if (bVoxelAllocationFeedbackEnable && OutViewStateData->VoxelFeedbackBuffer.IsValid())
+	{
+		CurrGPUMinVoxelSize = GraphBuilder.RegisterExternalBuffer(OutViewStateData->VoxelFeedbackBuffer);
+	}
+	if (CurrGPUMinVoxelSize == nullptr) { CurrGPUMinVoxelSize = GSystemTextures.GetDefaultBuffer(GraphBuilder, 2, FFloat16(CPUMinVoxelWorldSize)); }
 
-	FRDGBufferRef TotalRequestedPageAllocationBuffer;
 	AddAllocateVoxelPagesPass(
 		GraphBuilder, 
 		View, 
@@ -1049,7 +1030,7 @@ static FHairStrandsVoxelResources AllocateVirtualVoxelResources(
 		MacroGroupResources,
 		Out.Parameters.Common.PageCountResolution,
 		Out.Parameters.Common.PageCount,
-		Out.Parameters.Common.MinVoxelWorldSize,
+		Out.Parameters.Common.CPUMinVoxelWorldSize,
 		Out.Parameters.Common.PageResolution,
 		Out.Parameters.Common.PageTextureResolution,
 		Out.Parameters.Common.IndirectDispatchGroupSize,
@@ -1062,12 +1043,12 @@ static FHairStrandsVoxelResources AllocateVirtualVoxelResources(
 		Out.IndirectArgsBuffer,
 		Out.PageIndexGlobalCounter,
 		Out.VoxelizationViewInfoBuffer,
-		TotalRequestedPageAllocationBuffer);
+		CurrGPUMinVoxelSize,
+		NextGPUMinVoxelSize);
 
-	// Enque next adaptive feedback buffer
-	if (bEnqueueNewReadbackQuery)
+	if (bVoxelAllocationFeedbackEnable && NextGPUMinVoxelSize)
 	{
-		AddEnqueueCopyPass(GraphBuilder, ReadbackBuffer, TotalRequestedPageAllocationBuffer, 4);
+		GraphBuilder.QueueBufferExtraction(NextGPUMinVoxelSize, &OutViewStateData->VoxelFeedbackBuffer);
 	}
 
 	{
@@ -1089,6 +1070,10 @@ static FHairStrandsVoxelResources AllocateVirtualVoxelResources(
 	Out.Parameters.Common.PageIndexOccupancyBuffer	= GraphBuilder.CreateSRV(Out.PageIndexOccupancyBuffer, PF_R32G32_UINT);
 	Out.Parameters.Common.PageIndexCoordBuffer		= GraphBuilder.CreateSRV(Out.PageIndexCoordBuffer, PF_R8G8B8A8_UINT);
 	Out.Parameters.Common.NodeDescBuffer			= GraphBuilder.CreateSRV(Out.NodeDescBuffer); 
+
+	Out.Parameters.Common.CurrGPUMinVoxelSize		= GraphBuilder.CreateSRV(CurrGPUMinVoxelSize, PF_R16F);
+	Out.Parameters.Common.NextGPUMinVoxelSize		= GraphBuilder.CreateSRV(NextGPUMinVoxelSize ? NextGPUMinVoxelSize : CurrGPUMinVoxelSize, PF_R16F);
+
 	Out.Parameters.PageTexture						= Out.PageTexture;
 
 	if (Out.PageIndexBuffer && Out.NodeDescBuffer)
@@ -1110,7 +1095,7 @@ static FHairStrandsVoxelResources AllocateDummyVirtualVoxelResources(
 
 	Out.Parameters.Common.PageCountResolution		= FIntVector(0,0,0);
 	Out.Parameters.Common.PageCount					= 0;
-	Out.Parameters.Common.MinVoxelWorldSize			= 1;
+	Out.Parameters.Common.CPUMinVoxelWorldSize		= 1;
 	Out.Parameters.Common.PageResolution			= 1;
 	Out.Parameters.Common.PageTextureResolution		= FIntVector(1, 1, 1);
 
