Add support Strata classification for simple lights.

diff --git a/Engine/Source/Runtime/Renderer/Private/LightRendering.cpp b/Engine/Source/Runtime/Renderer/Private/LightRendering.cpp
--- a/Engine/Source/Runtime/Renderer/Private/LightRendering.cpp
+++ b/Engine/Source/Runtime/Renderer/Private/LightRendering.cpp
@@ -2353,7 +2353,7 @@ static FSimpleLightsStandardDeferredParameters GetRenderLightSimpleParameters(
 	FRDGBufferRef BufferDummy = GSystemTextures.GetDefaultBuffer(GraphBuilder, 4, 0u);
 	FRDGBufferSRVRef BufferDummySRV = GraphBuilder.CreateSRV(BufferDummy, PF_R32_UINT);
 	
-	// PS - General parmaeters
+	// PS - General parameters
 	Out.PS.SceneTextures = SceneTextures.UniformBuffer;
 	Out.PS.HairStrands = View.HairStrandsViewData.UniformBuffer;
 	Out.PS.Strata = Strata::BindStrataGlobalUniformParameters(View.StrataSceneData);
@@ -2383,7 +2383,7 @@ static FSimpleLightsStandardDeferredParameters GetRenderLightSimpleParameters(
 		Out.PS.RenderTargets.DepthStencil = FDepthStencilBinding(SceneTextures.Depth.Target, ERenderTargetLoadAction::ELoad, ERenderTargetLoadAction::ELoad, FExclusiveDepthStencil::DepthRead_StencilWrite);
 	}
 
-	// VS - General paramters (dummy geometry, as the geometry is setup within the pass light loop)
+	// VS - General parameters (dummy geometry, as the geometry is setup within the pass light loop)
 	FSphere SphereLight;
 	SphereLight.Center = SimpleLightPosition; // Should we account for LWC Position+Tile here?
 	SphereLight.W = SimpleLight.Radius;
@@ -2392,102 +2392,120 @@ static FSimpleLightsStandardDeferredParameters GetRenderLightSimpleParameters(
 	return Out;
 }
 
-void FDeferredShadingSceneRenderer::RenderSimpleLightsStandardDeferred(
+static void InternalRenderSimpleLightsStandardDeferred(
 	FRDGBuilder& GraphBuilder,
+	const FScene* Scene,
+	const FViewInfo& View,
+	const uint32 ViewIndex,
+	const uint32 NumViews,
 	const FMinimalSceneTextures& SceneTextures,
-	const FSimpleLightArray& SimpleLights)
+	const FSimpleLightArray& SimpleLights,
+	EStrataTileMaterialType TileType)
 {
-	if (SimpleLights.InstanceData.Num() == 0)
-	{
-		return;
-	}
-
-	SCOPE_CYCLE_COUNTER(STAT_DirectLightRenderingTime);
-	INC_DWORD_STAT_BY(STAT_NumLightsUsingStandardDeferred, SimpleLights.InstanceData.Num());
-
-	for (int32 ViewIndex = 0, NumViews = Views.Num(); ViewIndex < NumViews; ViewIndex++)
-	{
-		const FViewInfo& View = Views[ViewIndex];
+	FSimpleLightsStandardDeferredParameters* PassParameters = GraphBuilder.AllocParameters<FSimpleLightsStandardDeferredParameters>();
+	*PassParameters = GetRenderLightSimpleParameters(
+		GraphBuilder,
+		Scene,
+		View,
+		SceneTextures,
+		SimpleLights.InstanceData[0], // Use a dummy light to create the PassParameter buffer. The light data will be
+		FVector(0, 0, 0));		  // update dynamically with the pass light loop for efficiency purpose
 
-		FSimpleLightsStandardDeferredParameters* PassParameters = GraphBuilder.AllocParameters<FSimpleLightsStandardDeferredParameters>();
-		*PassParameters = GetRenderLightSimpleParameters(
-			GraphBuilder,
-			Scene,
-			View,
-			SceneTextures,
-			SimpleLights.InstanceData[0], // Use a dummy light to create the PassParameter buffer. The light data will be
-			FVector(0, 0, 0));		  // update dynamically with the pass light loop for efficiency purpose
+	FDeferredLightPS::FPermutationDomain PermutationVector;
+	PermutationVector.Set< FDeferredLightPS::FSourceShapeDim >(ELightSourceShape::Capsule);
+	PermutationVector.Set< FDeferredLightPS::FIESProfileDim >(false);
+	PermutationVector.Set< FDeferredLightPS::FVisualizeCullingDim >(View.Family->EngineShowFlags.VisualizeLightCulling);
+	PermutationVector.Set< FDeferredLightPS::FLightingChannelsDim >(false);
+	PermutationVector.Set< FDeferredLightPS::FAnistropicMaterials >(false);
+	PermutationVector.Set< FDeferredLightPS::FTransmissionDim >(false);
+	PermutationVector.Set< FDeferredLightPS::FHairLighting>(0);
+	PermutationVector.Set< FDeferredLightPS::FAtmosphereTransmittance >(false);
+	PermutationVector.Set< FDeferredLightPS::FCloudTransmittance >(false);
+	PermutationVector.Set< FDeferredLightPS::FStrataTileType>(TileType != EStrataTileMaterialType::ECount ? TileType : 0);
+	TShaderMapRef<FDeferredLightPS> PixelShader(View.ShaderMap, PermutationVector);
 
-		FDeferredLightPS::FPermutationDomain PermutationVector;
-		PermutationVector.Set< FDeferredLightPS::FSourceShapeDim >(ELightSourceShape::Capsule);
-		PermutationVector.Set< FDeferredLightPS::FIESProfileDim >(false);
-		PermutationVector.Set< FDeferredLightPS::FVisualizeCullingDim >(View.Family->EngineShowFlags.VisualizeLightCulling);
-		PermutationVector.Set< FDeferredLightPS::FLightingChannelsDim >(false);
-		PermutationVector.Set< FDeferredLightPS::FAnistropicMaterials >(false);
-		PermutationVector.Set< FDeferredLightPS::FTransmissionDim >(false);
-		PermutationVector.Set< FDeferredLightPS::FHairLighting>(0);
-		PermutationVector.Set< FDeferredLightPS::FAtmosphereTransmittance >(false);
-		PermutationVector.Set< FDeferredLightPS::FCloudTransmittance >(false);
-		PermutationVector.Set< FDeferredLightPS::FStrataTileType>(0);
-		TShaderMapRef<FDeferredLightPS> PixelShader(View.ShaderMap, PermutationVector);
+	FDeferredLightVS::FPermutationDomain PermutationVectorVS;
+	PermutationVectorVS.Set<FDeferredLightVS::FRadialLight>(true);
+	TShaderMapRef<FDeferredLightVS> VertexShader(View.ShaderMap, PermutationVectorVS);
 
-		FDeferredLightVS::FPermutationDomain PermutationVectorVS;
-		PermutationVectorVS.Set<FDeferredLightVS::FRadialLight>(true);
-		TShaderMapRef<FDeferredLightVS> VertexShader(View.ShaderMap, PermutationVectorVS);
+	GraphBuilder.AddPass(
+		RDG_EVENT_NAME("Light::DeferredSimpleLights(Strata:%s,Tile:%s)", Strata::IsStrataEnabled() ? TEXT("True") : TEXT("False"), Strata::IsStrataEnabled() ? ToString(TileType) : TEXT("None")),
+		PassParameters,
+		ERDGPassFlags::Raster,
+		[&View, &SimpleLights, ViewIndex, NumViews, PassParameters, PixelShader, VertexShader, TileType](FRHICommandList& RHICmdList)
+	{
+		FGraphicsPipelineStateInitializer GraphicsPSOInit;
+		RHICmdList.ApplyCachedRenderTargets(GraphicsPSOInit);
 
-		// STRATA_TODO: add simple/complex tile support for simple lights
-		const EStrataTileMaterialType TileType = Strata::IsStrataEnabled() ? EStrataTileMaterialType::EComplex : EStrataTileMaterialType::ECount;
+		// Use additive blending for color
+		GraphicsPSOInit.BlendState = TStaticBlendState<CW_RGBA, BO_Add, BF_One, BF_One, BO_Add, BF_One, BF_One>::GetRHI();
+		GraphicsPSOInit.PrimitiveType = PT_TriangleList;
 
-		GraphBuilder.AddPass(
-			RDG_EVENT_NAME("StandardDeferredSimpleLights"),
-			PassParameters,
-			ERDGPassFlags::Raster,
-			[this, &View, &SimpleLights, ViewIndex, NumViews, PassParameters, PixelShader, VertexShader, TileType](FRHICommandList& RHICmdList)
+		for (int32 LightIndex = 0; LightIndex < SimpleLights.InstanceData.Num(); LightIndex++)
 		{
-			FGraphicsPipelineStateInitializer GraphicsPSOInit;
-			RHICmdList.ApplyCachedRenderTargets(GraphicsPSOInit);
+			const FSimpleLightEntry& SimpleLight = SimpleLights.InstanceData[LightIndex];
 
-			// Use additive blending for color
-			GraphicsPSOInit.BlendState = TStaticBlendState<CW_RGBA, BO_Add, BF_One, BF_One, BO_Add, BF_One, BF_One>::GetRHI();
-			GraphicsPSOInit.PrimitiveType = PT_TriangleList;
+			const FSimpleLightPerViewEntry& SimpleLightPerViewData = SimpleLights.GetViewDependentData(LightIndex, ViewIndex, NumViews);
+			const FSphere LightBounds(SimpleLightPerViewData.Position, SimpleLight.Radius);
 
-			for (int32 LightIndex = 0; LightIndex < SimpleLights.InstanceData.Num(); LightIndex++)
-			{
-				const FSimpleLightEntry& SimpleLight = SimpleLights.InstanceData[LightIndex];
 
-				const FSimpleLightPerViewEntry& SimpleLightPerViewData = SimpleLights.GetViewDependentData(LightIndex, ViewIndex, NumViews);
-				const FSphere LightBounds(SimpleLightPerViewData.Position, SimpleLight.Radius);
+			// Set the device viewport for the view.
+			RHICmdList.SetViewport(View.ViewRect.Min.X, View.ViewRect.Min.Y, 0.0f, View.ViewRect.Max.X, View.ViewRect.Max.Y, 1.0f);
 
+			const bool bCameraInsideLightGeometry = ((FVector)View.ViewMatrices.GetViewOrigin() - LightBounds.Center).SizeSquared() < FMath::Square(LightBounds.W * 1.05f + View.NearClippingDistance * 2.0f)
+							// Always draw backfaces in ortho
+							//@todo - accurate ortho camera / light intersection
+							|| !View.IsPerspectiveProjection();
 
-				// Set the device viewport for the view.
-				RHICmdList.SetViewport(View.ViewRect.Min.X, View.ViewRect.Min.Y, 0.0f, View.ViewRect.Max.X, View.ViewRect.Max.Y, 1.0f);
+			const uint32 StencilRef = SetBoundingGeometryRasterizerAndDepthState(GraphicsPSOInit, View, bCameraInsideLightGeometry, TileType);
+			GraphicsPSOInit.BoundShaderState.VertexDeclarationRHI = GetVertexDeclarationFVector4();
+			GraphicsPSOInit.BoundShaderState.VertexShaderRHI = VertexShader.GetVertexShader();
+			GraphicsPSOInit.BoundShaderState.PixelShaderRHI = PixelShader.GetPixelShader();
+			SetGraphicsPipelineState(RHICmdList, GraphicsPSOInit, StencilRef);
 
-				const bool bCameraInsideLightGeometry = ((FVector)View.ViewMatrices.GetViewOrigin() - LightBounds.Center).SizeSquared() < FMath::Square(LightBounds.W * 1.05f + View.NearClippingDistance * 2.0f)
-								// Always draw backfaces in ortho
-								//@todo - accurate ortho camera / light intersection
-								|| !View.IsPerspectiveProjection();
+			// Update the light parameters with a custom uniform buffer
+			FDeferredLightUniformStruct DeferredLightUniformsValue = GetSimpleDeferredLightParameters(View, SimpleLight, SimpleLightPerViewData);
+			SetShaderParameters(RHICmdList, PixelShader, PixelShader.GetPixelShader(), PassParameters->PS);
+			SetUniformBufferParameterImmediate(RHICmdList, RHICmdList.GetBoundPixelShader(), PixelShader->GetUniformBufferParameter<FDeferredLightUniformStruct>(), DeferredLightUniformsValue);
 
-				SetBoundingGeometryRasterizerAndDepthState(GraphicsPSOInit, View, bCameraInsideLightGeometry, TileType);
-				GraphicsPSOInit.BoundShaderState.VertexDeclarationRHI = GetVertexDeclarationFVector4();
-				GraphicsPSOInit.BoundShaderState.VertexShaderRHI = VertexShader.GetVertexShader();
-				GraphicsPSOInit.BoundShaderState.PixelShaderRHI = PixelShader.GetPixelShader();
-				SetGraphicsPipelineState(RHICmdList, GraphicsPSOInit, 0);
+			// Update vertex shader parameters with custom parameters/uniform buffer
+			FDeferredLightVS::FParameters ParametersVS = FDeferredLightVS::GetParameters(View, LightBounds);
+			SetShaderParameters(RHICmdList, VertexShader, VertexShader.GetVertexShader(), ParametersVS /*PassParameters->VS*/);
 
+			// Apply the point or spot light with some approximately bounding geometry,
+			// So we can get speedups from depth testing and not processing pixels outside of the light's influence.
+			StencilingGeometry::DrawSphere(RHICmdList);
+		}
+	});
+}
 
-				// Update the light parameters with a custom uniform buffer
-				FDeferredLightUniformStruct DeferredLightUniformsValue = GetSimpleDeferredLightParameters(View, SimpleLight, SimpleLightPerViewData);
-				SetShaderParameters(RHICmdList, PixelShader, PixelShader.GetPixelShader(), PassParameters->PS);
-				SetUniformBufferParameterImmediate(RHICmdList, RHICmdList.GetBoundPixelShader(), PixelShader->GetUniformBufferParameter<FDeferredLightUniformStruct>(), DeferredLightUniformsValue);
+void FDeferredShadingSceneRenderer::RenderSimpleLightsStandardDeferred(
+	FRDGBuilder& GraphBuilder,
+	const FMinimalSceneTextures& SceneTextures,
+	const FSimpleLightArray& SimpleLights)
+{
+	if (SimpleLights.InstanceData.Num() == 0)
+	{
+		return;
+	}
 
-				// Update vertex shader parameters with custom parameters/uniform buffer
-				FDeferredLightVS::FParameters ParametersVS = FDeferredLightVS::GetParameters(View, LightBounds);
-				SetShaderParameters(RHICmdList, VertexShader, VertexShader.GetVertexShader(), ParametersVS /*PassParameters->VS*/);
+	SCOPE_CYCLE_COUNTER(STAT_DirectLightRenderingTime);
+	INC_DWORD_STAT_BY(STAT_NumLightsUsingStandardDeferred, SimpleLights.InstanceData.Num());
 
-				// Apply the point or spot light with some approximately bounding geometry,
-				// So we can get speedups from depth testing and not processing pixels outside of the light's influence.
-				StencilingGeometry::DrawSphere(RHICmdList);
-			}
-		});
+	for (int32 ViewIndex = 0, NumViews = Views.Num(); ViewIndex < NumViews; ViewIndex++)
+	{
+		const FViewInfo& View = Views[ViewIndex];
+
+		if (Strata::IsStrataEnabled())
+		{
+			InternalRenderSimpleLightsStandardDeferred(GraphBuilder, Scene, View, ViewIndex, NumViews, SceneTextures, SimpleLights, EStrataTileMaterialType::ESimple);
+			InternalRenderSimpleLightsStandardDeferred(GraphBuilder, Scene, View, ViewIndex, NumViews, SceneTextures, SimpleLights, EStrataTileMaterialType::ESingle);
+			InternalRenderSimpleLightsStandardDeferred(GraphBuilder, Scene, View, ViewIndex, NumViews, SceneTextures, SimpleLights, EStrataTileMaterialType::EComplex);
+		}
+		else
+		{
+			InternalRenderSimpleLightsStandardDeferred(GraphBuilder, Scene, View, ViewIndex, NumViews, SceneTextures, SimpleLights, EStrataTileMaterialType::ECount);
+		}
 	}
 }
 
