Fixed: Incorrect on screen warning about vertex fogging on mobile

diff --git a/Engine/Source/Runtime/Engine/Classes/Engine/RendererSettings.h b/Engine/Source/Runtime/Engine/Classes/Engine/RendererSettings.h
--- a/Engine/Source/Runtime/Engine/Classes/Engine/RendererSettings.h
+++ b/Engine/Source/Runtime/Engine/Classes/Engine/RendererSettings.h
@@ -265,7 +265,7 @@ class ENGINE_API URendererSettings : public UDeveloperSettings
 
 	UPROPERTY(config, EditAnywhere, Category = Mobile, meta = (
 		ConsoleVariable = "r.Mobile.DisableVertexFog", DisplayName = "Disable vertex fogging in mobile shaders",
-		ToolTip = "If true, vertex fog will be omitted from all mobile shaders. If your game does not use fog, you should choose this setting to increase shading performance.",
+		ToolTip = "If true, vertex fog will be omitted from the most of the mobile base pass shaders. Instead, fog will be applied in a separate pass and only when scene has a fog component.",
 		ConfigRestartRequired = true))
 		uint32 bMobileDisableVertexFog : 1;
 
diff --git a/Engine/Source/Runtime/Renderer/Private/SceneRendering.cpp b/Engine/Source/Runtime/Renderer/Private/SceneRendering.cpp
--- a/Engine/Source/Runtime/Renderer/Private/SceneRendering.cpp
+++ b/Engine/Source/Runtime/Renderer/Private/SceneRendering.cpp
@@ -3223,22 +3223,7 @@ void FSceneRenderer::RenderFinish(FRDGBuilder& GraphBuilder, FRDGTextureRef View
 		const bool bShowMobileLowQualityLightmapWarning = bMobile && !ReadOnlyCVARCache.bEnableLowQualityLightmaps && ReadOnlyCVARCache.bAllowStaticLighting;
 		const bool bShowMobileDynamicCSMWarning = bMobile && Scene->NumMobileStaticAndCSMLights_RenderThread > 0 && !(ReadOnlyCVARCache.bMobileEnableStaticAndCSMShadowReceivers && ReadOnlyCVARCache.bMobileAllowDistanceFieldShadows);
 		const bool bShowMobileMovableDirectionalLightWarning = bMobile && Scene->NumMobileMovableDirectionalLights_RenderThread > 0 && !ReadOnlyCVARCache.bMobileAllowMovableDirectionalLights;
-		
-		bool bMobileMissingSkyMaterial = false;
-		bool bMobileShowVertexFogWarning = false;
-		if (bMobile && (Scene->ExponentialFogs.Num() > 0 || Scene->HasSkyAtmosphere()))
-		{
-			static const auto* CVarDisableVertexFog = IConsoleManager::Get().FindTConsoleVariableDataInt(TEXT("r.Mobile.DisableVertexFog"));
-			if (CVarDisableVertexFog && CVarDisableVertexFog->GetValueOnRenderThread() != 0)
-			{
-				bMobileShowVertexFogWarning = true;
-			}
-			
-			if (Scene->HasSkyAtmosphere() && (Views.Num() > 0 && !Views[0].bSceneHasSkyMaterial))
-			{
-				bMobileMissingSkyMaterial = true;
-			}
-		}
+		const bool bMobileMissingSkyMaterial = (bMobile && Scene->HasSkyAtmosphere() && (Views.Num() > 0 && !Views[0].bSceneHasSkyMaterial));
 
 		const bool bSingleLayerWaterWarning = ShouldRenderSingleLayerWaterSkippedRenderEditorNotification(Views);
 
@@ -3258,7 +3243,7 @@ void FSceneRenderer::RenderFinish(FRDGBuilder& GraphBuilder, FRDGTextureRef View
 
 		const bool bAnyWarning = bShowPrecomputedVisibilityWarning || bShowDemotedLocalMemoryWarning || bShowGlobalClipPlaneWarning || bShowSkylightWarning || bShowPointLightWarning
 			|| bShowDFAODisabledWarning || bShowShadowedLightOverflowWarning || bShowMobileDynamicCSMWarning || bShowMobileLowQualityLightmapWarning || bShowMobileMovableDirectionalLightWarning
-			|| bMobileShowVertexFogWarning || bMobileMissingSkyMaterial || bShowSkinCacheOOM || bSingleLayerWaterWarning || bShowDFDisabledWarning || bShowNoSkyAtmosphereComponentWarning || bFxDebugDraw 
+			|| bMobileMissingSkyMaterial || bShowSkinCacheOOM || bSingleLayerWaterWarning || bShowDFDisabledWarning || bShowNoSkyAtmosphereComponentWarning || bFxDebugDraw 
 			|| bLumenEnabledButHasNoDataForTracing || bLumenEnabledButDisabledForTheProject || bNaniteEnabledButNoAtomics || bNaniteEnabledButDisabledInProject || bRealTimeSkyCaptureButNothingToCapture || bShowWaitingSkylight
 			|| bShowLocalExposureDisabledWarning || bHasDelegateWarnings
 #if WITH_EDITOR
@@ -3302,7 +3287,7 @@ void FSceneRenderer::RenderFinish(FRDGBuilder& GraphBuilder, FRDGTextureRef View
 						[this, &ReadOnlyCVARCache, ViewState, GPUSkinCacheExtraRequiredMemory,
 						bLocked, bShowPrecomputedVisibilityWarning, bShowDemotedLocalMemoryWarning, bShowGlobalClipPlaneWarning, bShowDFAODisabledWarning, bShowDFDisabledWarning,
 						bViewParentOrFrozen, bShowSkylightWarning, bShowPointLightWarning, bShowShadowedLightOverflowWarning,
-						bShowMobileLowQualityLightmapWarning, bShowMobileMovableDirectionalLightWarning, bShowMobileDynamicCSMWarning, bMobileShowVertexFogWarning, bMobileMissingSkyMaterial, 
+						bShowMobileLowQualityLightmapWarning, bShowMobileMovableDirectionalLightWarning, bShowMobileDynamicCSMWarning, bMobileMissingSkyMaterial, 
 						bShowSkinCacheOOM, bSingleLayerWaterWarning, bShowNoSkyAtmosphereComponentWarning, bFxDebugDraw, FXInterface, bShowLocalExposureDisabledWarning,
 						bLumenEnabledButHasNoDataForTracing, bLumenEnabledButDisabledForTheProject, bNaniteEnabledButNoAtomics, bNaniteEnabledButDisabledInProject, bNaniteRequireAtomics, bRealTimeSkyCaptureButNothingToCapture, bShowWaitingSkylight, bShowAmbientCubemapMipGenSettingsWarning]
 						(FCanvas& Canvas)
@@ -3393,11 +3378,6 @@ void FSceneRenderer::RenderFinish(FRDGBuilder& GraphBuilder, FRDGTextureRef View
 								: NSLOCTEXT("Renderer", "MobileDynamicCSMDistFieldShadows", "MOBILE CSM+STATIC REQUIRES DISTANCE FIELD SHADOWS ENABLED FOR PROJECT");
 							Writer.DrawLine(Message);
 						}
-						if (bMobileShowVertexFogWarning)
-						{
-							static const FText Message = NSLOCTEXT("Renderer", "MobileVertexFog", "PROJECT HAS VERTEX FOG ON MOBILE DISABLED");
-							Writer.DrawLine(Message);
-						}
 
 						if (bMobileMissingSkyMaterial)
 						{
