MGPU:  Ray tracing startup crash, due to bug where GPU mask in the command list structure doesn't get updated during RHI command playback (it's only up to date during command recording).  Pull out the GPU mask and pass it in separately to the lambda.  Also added a FRHIGPUMask::All() in FRayTracingSkinnedGeometryUpdateQueue::Commit -- noticed this was missing while debugging.

#rnx

diff --git a/Engine/Source/Runtime/Engine/Private/RayTracingSkinnedGeometry.cpp b/Engine/Source/Runtime/Engine/Private/RayTracingSkinnedGeometry.cpp
--- a/Engine/Source/Runtime/Engine/Private/RayTracingSkinnedGeometry.cpp
+++ b/Engine/Source/Runtime/Engine/Private/RayTracingSkinnedGeometry.cpp
@@ -225,6 +225,8 @@ void FRayTracingSkinnedGeometryUpdateQueue::Commit(FRDGBuilder& GraphBuilder)
 	FSkinnedGeometryBLASUpdateParams* BLASUpdateParams = GraphBuilder.AllocParameters<FSkinnedGeometryBLASUpdateParams>();
 	BLASUpdateParams->SharedScratchBuffer = SharedScratchBuffer;
 
+	RDG_GPU_MASK_SCOPE(GraphBuilder, FRHIGPUMask::All());
+
 	GraphBuilder.AddPass(RDG_EVENT_NAME("CommitRayTracingSkinnedGeometryUpdates"), BLASUpdateParams, ERDGPassFlags::Compute | ERDGPassFlags::NeverCull,
 		[this, SharedScratchBuffer](FRHICommandListImmediate& RHICmdList)
 		{
diff --git a/Engine/Source/Runtime/Renderer/Private/RayTracing/RayTracingScene.cpp b/Engine/Source/Runtime/Renderer/Private/RayTracing/RayTracingScene.cpp
--- a/Engine/Source/Runtime/Renderer/Private/RayTracing/RayTracingScene.cpp
+++ b/Engine/Source/Runtime/Renderer/Private/RayTracing/RayTracingScene.cpp
@@ -177,12 +177,13 @@ void FRayTracingScene::Create(FRDGBuilder& GraphBuilder, const FGPUScene& GPUSce
 				RHICmdList.UnlockBuffer(InstanceUploadBuffer);
 				RHICmdList.UnlockBuffer(TransformUploadBuffer);
 
-				RHICmdList.EnqueueLambda([this, &SceneInitializer](FRHICommandListImmediate& RHICmdList)
+				// Pull this out here, because command list playback (where the lambda is executed) doesn't update the GPU mask
+				FRHIGPUMask IterateGPUMasks = RHICmdList.GetGPUMask();
+
+				RHICmdList.EnqueueLambda([this, &SceneInitializer, IterateGPUMasks](FRHICommandListImmediate& RHICmdList)
 					{
 						QUICK_SCOPE_CYCLE_COUNTER(GetAccelerationStructuresAddresses);
 
-						FRHIGPUMask IterateGPUMasks = RHICmdList.GetGPUMask();
-
 						for (uint32 GPUIndex : IterateGPUMasks)
 						{
 							FRayTracingAccelerationStructureAddress* AddressesPtr = (FRayTracingAccelerationStructureAddress*)RHICmdList.LockBufferMGPU(
