Add strata tile permutation to cluster shading for performance improvment.

diff --git a/Engine/Shaders/Private/ClusteredDeferredShadingPixelShader.usf b/Engine/Shaders/Private/ClusteredDeferredShadingPixelShader.usf
--- a/Engine/Shaders/Private/ClusteredDeferredShadingPixelShader.usf
+++ b/Engine/Shaders/Private/ClusteredDeferredShadingPixelShader.usf
@@ -31,6 +31,15 @@
 #include "VirtualShadowMaps/TransmissionCommon.ush"
 
 #if STRATA_ENABLED
+  #if STRATA_TILETYPE == 0
+	#define STRATA_FASTPATH 1
+  #elif STRATA_TILETYPE == 1
+	#define STRATA_SINGLEPATH 1
+  #elif STRATA_TILETYPE == 2
+  // COMPLEX PATH
+  #else
+	#error Strata tile type non-implemented
+  #endif
 #include "/Engine/Private/Strata/StrataEvaluation.ush"
 #include "/Engine/Private/Strata/StrataDeferredLighting.ush"
 #endif
@@ -412,12 +421,12 @@ float4 GetLightGridLocalLighting(
 
 #if STRATA_ENABLED && USE_HAIR_LIGHTING == 0
 void ClusteredShadingPixelShader(
-	in noperspective float4 UVAndScreenPos : TEXCOORD0,
+	float2 ScreenUV : TEXCOORD0,
+	float3 ScreenVector : TEXCOORD1,
 	in float4 SvPosition : SV_Position,
 	out float4 OutColor : SV_Target0)
 {
 	uint2 PixelPos = SvPosition.xy;
-	float2 ScreenUV = UVAndScreenPos.xy;
 
 	const float Dither = InterleavedGradientNoise(SvPosition.xy, View.StateFrameIndexMod8);
 	const uint EyeIndex = 0;
@@ -440,7 +449,7 @@ void ClusteredShadingPixelShader(
 	BRANCH
 	if (StrataPixelHeader.BSDFCount > 0 || CulledLightGridData.NumLocalLights == 0)
 	{
-		const float2 ScreenPosition = UVAndScreenPos.zw;
+		const float2 ScreenPosition = ViewportUVToScreenPos(ScreenUV);
 		const float3 TranslatedWorldPosition = mul(float4(ScreenPosition * SceneDepth, SceneDepth, 1), PrimaryView.ScreenToTranslatedWorld).xyz;
 		const float3 CameraVector = normalize(TranslatedWorldPosition - PrimaryView.TranslatedWorldCameraOrigin);
 		uint FirstNonSimpleLightIndex = 0;
diff --git a/Engine/Source/Runtime/Renderer/Private/ClusteredDeferredShadingPass.cpp b/Engine/Source/Runtime/Renderer/Private/ClusteredDeferredShadingPass.cpp
--- a/Engine/Source/Runtime/Renderer/Private/ClusteredDeferredShadingPass.cpp
+++ b/Engine/Source/Runtime/Renderer/Private/ClusteredDeferredShadingPass.cpp
@@ -94,7 +94,8 @@ class FClusteredShadingPS : public FGlobalShader
 
 	class FVisualizeLightCullingDim : SHADER_PERMUTATION_BOOL("VISUALIZE_LIGHT_CULLING");
 	class FHairStrandsLighting : SHADER_PERMUTATION_BOOL("USE_HAIR_LIGHTING");
-	using FPermutationDomain = TShaderPermutationDomain<FVisualizeLightCullingDim, FHairStrandsLighting>;
+	class FStrataTileType : SHADER_PERMUTATION_INT("STRATA_TILETYPE", 3);
+	using FPermutationDomain = TShaderPermutationDomain<FVisualizeLightCullingDim, FHairStrandsLighting, FStrataTileType>;
 
 	BEGIN_SHADER_PARAMETER_STRUCT(FParameters, )
 		SHADER_PARAMETER_RDG_UNIFORM_BUFFER(FForwardLightData, Forward)
@@ -107,12 +108,23 @@ class FClusteredShadingPS : public FGlobalShader
 		SHADER_PARAMETER_RDG_TEXTURE(Texture2D, ShadowMaskBits)
 		SHADER_PARAMETER_STRUCT_INCLUDE(FVirtualShadowMapSamplingParameters, VirtualShadowMapSamplingParameters)
 		SHADER_PARAMETER_RDG_BUFFER_SRV(Buffer<uint>, HairTransmittanceBuffer)
-
+		// Strata tiles
+		SHADER_PARAMETER_RDG_BUFFER_SRV(Buffer<uint>, TileListBufferSimple)
+		SHADER_PARAMETER_RDG_BUFFER_SRV(Buffer<uint>, TileListBufferSingle)
+		SHADER_PARAMETER_RDG_BUFFER_SRV(Buffer<uint>, TileListBufferComplex)
+		RDG_BUFFER_ACCESS(TileIndirectBufferSimple,  ERHIAccess::IndirectArgs)
+		RDG_BUFFER_ACCESS(TileIndirectBufferSingle,  ERHIAccess::IndirectArgs)
+		RDG_BUFFER_ACCESS(TileIndirectBufferComplex, ERHIAccess::IndirectArgs)
 		RENDER_TARGET_BINDING_SLOTS()
 	END_SHADER_PARAMETER_STRUCT()
 
 	static bool ShouldCompilePermutation(const FGlobalShaderPermutationParameters& Parameters)
 	{
+		FPermutationDomain PermutationVector(Parameters.PermutationId);
+		if (!Strata::IsStrataEnabled() && PermutationVector.Get<FStrataTileType>() != 0)
+		{
+			return false;
+		}
 		// OLATODO: what level do we actually need for this?
 		return IsFeatureLevelSupported(Parameters.Platform, ERHIFeatureLevel::SM5);
 	}
@@ -131,6 +143,7 @@ IMPLEMENT_GLOBAL_SHADER(FClusteredShadingPS, "/Engine/Private/ClusteredDeferredS
 enum class EClusterPassInputType : uint8
 {
 	GBuffer,
+	Strata,
 	HairStrands
 };
 
@@ -140,6 +153,7 @@ static void InternalAddClusteredDeferredShadingPass(
 	const FMinimalSceneTextures& SceneTextures,
 	const FSortedLightSetSceneInfo &SortedLightsSet,
 	EClusterPassInputType InputType,
+	EStrataTileMaterialType TileType,
 	FRDGTextureRef ShadowMaskBits,
 	FVirtualShadowMapArray& VirtualShadowMapArray,
 	FRDGBufferSRVRef HairTransmittanceBuffer)
@@ -147,6 +161,7 @@ static void InternalAddClusteredDeferredShadingPass(
 	check(SortedLightsSet.ClusteredSupportedEnd > 0);
 	const FIntPoint SceneTextureExtent = SceneTextures.Config.Extent;
 	const bool bHairStrands = InputType == EClusterPassInputType::HairStrands;
+	const bool bStrata = Strata::IsStrataEnabled();
 	
 	FClusteredShadingPS::FParameters *PassParameters = GraphBuilder.AllocParameters<FClusteredShadingPS::FParameters>();
 	PassParameters->View = View.ViewUniformBuffer;
@@ -166,19 +181,47 @@ static void InternalAddClusteredDeferredShadingPass(
 		PassParameters->RenderTargets[0] = FRenderTargetBinding(View.HairStrandsViewData.VisibilityData.SampleLightingTexture, ERenderTargetLoadAction::ELoad);
 	}
 
-	// STRATA_TODO: use classification tile
+	// VS - Strata tile parameters
+	if (bStrata)
+	{
+		PassParameters->TileListBufferSimple = View.StrataSceneData->ClassificationTileListBufferSRV[EStrataTileMaterialType::ESimple];
+		PassParameters->TileListBufferSingle = View.StrataSceneData->ClassificationTileListBufferSRV[EStrataTileMaterialType::ESingle];
+		PassParameters->TileListBufferComplex = View.StrataSceneData->ClassificationTileListBufferSRV[EStrataTileMaterialType::EComplex];
+		PassParameters->TileIndirectBufferSimple = View.StrataSceneData->ClassificationTileIndirectBuffer[EStrataTileMaterialType::ESimple];
+		PassParameters->TileIndirectBufferSingle = View.StrataSceneData->ClassificationTileIndirectBuffer[EStrataTileMaterialType::ESingle];
+		PassParameters->TileIndirectBufferComplex = View.StrataSceneData->ClassificationTileIndirectBuffer[EStrataTileMaterialType::EComplex];
+	}
+	else
+	{
+		FRDGBufferRef BufferDummy = GSystemTextures.GetDefaultBuffer(GraphBuilder, 4, 0u);
+		FRDGBufferSRVRef BufferDummySRV = GraphBuilder.CreateSRV(BufferDummy, PF_R32_UINT);
+		PassParameters->TileListBufferSimple = BufferDummySRV;
+		PassParameters->TileListBufferComplex = BufferDummySRV;
+		PassParameters->TileIndirectBufferSimple = BufferDummy;
+		PassParameters->TileIndirectBufferComplex = BufferDummy;
+	}
+
+	const TCHAR* TileTypeName = ToString(TileType);
 	GraphBuilder.AddPass(
-		RDG_EVENT_NAME("Light::ClusteredDeferredShading(%s,#Lights:%d)", bHairStrands ? TEXT("HairStrands") : (Strata::IsStrataEnabled() ? TEXT("Strata") : TEXT("GBuffer")), SortedLightsSet.ClusteredSupportedEnd),
+		RDG_EVENT_NAME("Light::ClusteredDeferredShading(%s,Lights:%d,Tile:%s)", bHairStrands ? TEXT("HairStrands") : (bStrata ? TEXT("Strata") : TEXT("GBuffer")), SortedLightsSet.ClusteredSupportedEnd, TileTypeName),
 		PassParameters,
 		ERDGPassFlags::Raster,
-		[PassParameters, &View, SceneTextureExtent, bHairStrands](FRHICommandListImmediate& InRHICmdList)
+		[PassParameters, &View, SceneTextureExtent, bHairStrands, bStrata, TileType](FRHICommandListImmediate& InRHICmdList)
 	{
 		TShaderMapRef<FClusteredShadingVS> HairVertexShader(View.ShaderMap);
 		TShaderMapRef<FPostProcessVS> VertexShader(View.ShaderMap);
 
+		Strata::FStrataTilePassVS::FPermutationDomain VSPermutationVector;
+		VSPermutationVector.Set<Strata::FStrataTilePassVS::FEnableDebug>(false);
+		VSPermutationVector.Set<Strata::FStrataTilePassVS::FEnableTexCoordScreenVector>(true);
+		TShaderMapRef<Strata::FStrataTilePassVS> TileVertexShader(View.ShaderMap, VSPermutationVector);
+		Strata::FStrataTilePassVS::FParameters VSParameters;
+
+
 		FClusteredShadingPS::FPermutationDomain PermutationVector;
 		PermutationVector.Set<FClusteredShadingPS::FVisualizeLightCullingDim>(View.Family->EngineShowFlags.VisualizeLightCulling);
 		PermutationVector.Set<FClusteredShadingPS::FHairStrandsLighting>(bHairStrands);
+		PermutationVector.Set<FClusteredShadingPS::FStrataTileType>(bStrata ? TileType : 0);
 		TShaderMapRef<FClusteredShadingPS> PixelShader(View.ShaderMap, PermutationVector);
 		{
 			FGraphicsPipelineStateInitializer GraphicsPSOInit;
@@ -191,9 +234,13 @@ static void InternalAddClusteredDeferredShadingPass(
 			GraphicsPSOInit.DepthStencilState = TStaticDepthStencilState<false, CF_Always>::GetRHI();
 
 			GraphicsPSOInit.BoundShaderState.VertexDeclarationRHI = GFilterVertexDeclaration.VertexDeclarationRHI;
-			GraphicsPSOInit.BoundShaderState.VertexShaderRHI = bHairStrands ? HairVertexShader.GetVertexShader() : VertexShader.GetVertexShader();
+			GraphicsPSOInit.BoundShaderState.VertexShaderRHI = bHairStrands ? HairVertexShader.GetVertexShader() : (bStrata ? TileVertexShader.GetVertexShader() : VertexShader.GetVertexShader());
 			GraphicsPSOInit.BoundShaderState.PixelShaderRHI = PixelShader.GetPixelShader();
 			GraphicsPSOInit.PrimitiveType = PT_TriangleList;
+			if (bStrata)
+			{
+				Strata::FillUpTiledPassData(TileType, View, VSParameters, GraphicsPSOInit.PrimitiveType);
+			}
 
 			SetGraphicsPipelineState(InRHICmdList, GraphicsPSOInit, 0);
 		}
@@ -212,6 +259,11 @@ static void InternalAddClusteredDeferredShadingPass(
 			InRHICmdList.SetStreamSource(0, nullptr, 0);
 			InRHICmdList.DrawPrimitive(0, 1, 1);
 		}
+		else if (bStrata)
+		{
+			SetShaderParameters(InRHICmdList, TileVertexShader, TileVertexShader.GetVertexShader(), VSParameters);
+			InRHICmdList.DrawPrimitiveIndirect(VSParameters.TileIndirectBuffer->GetIndirectRHICallBuffer(), 0);
+		}
 		else
 		{
 			InRHICmdList.SetViewport(View.ViewRect.Min.X, View.ViewRect.Min.Y, 0.0f, View.ViewRect.Max.X, View.ViewRect.Max.Y, 1.0f);
@@ -243,15 +295,54 @@ void FDeferredShadingSceneRenderer::AddClusteredDeferredShadingPass(
 		{
 			FViewInfo& View = Views[ViewIndex];
 
-			InternalAddClusteredDeferredShadingPass(
-				GraphBuilder,
-				View,
-				SceneTextures,
-				SortedLightsSet,
-				EClusterPassInputType::GBuffer,
-				ShadowMaskBits,
-				VirtualShadowMapArray,
-				nullptr);
+			if (Strata::IsStrataEnabled())
+			{
+				InternalAddClusteredDeferredShadingPass(
+					GraphBuilder,
+					View,
+					SceneTextures,
+					SortedLightsSet,
+					EClusterPassInputType::Strata,
+					EStrataTileMaterialType::ESimple,
+					ShadowMaskBits,
+					VirtualShadowMapArray,
+					nullptr);
+
+				InternalAddClusteredDeferredShadingPass(
+					GraphBuilder,
+					View,
+					SceneTextures,
+					SortedLightsSet,
+					EClusterPassInputType::Strata,
+					EStrataTileMaterialType::ESingle,
+					ShadowMaskBits,
+					VirtualShadowMapArray,
+					nullptr);
+
+				InternalAddClusteredDeferredShadingPass(
+					GraphBuilder,
+					View,
+					SceneTextures,
+					SortedLightsSet,
+					EClusterPassInputType::Strata,
+					EStrataTileMaterialType::EComplex,
+					ShadowMaskBits,
+					VirtualShadowMapArray,
+					nullptr);
+			}
+			else
+			{
+				InternalAddClusteredDeferredShadingPass(
+					GraphBuilder,
+					View,
+					SceneTextures,
+					SortedLightsSet,
+					EClusterPassInputType::GBuffer,
+					EStrataTileMaterialType::ECount,
+					ShadowMaskBits,
+					VirtualShadowMapArray,
+					nullptr);
+			}
 
 			if (HairStrands::HasViewHairStrandsData(View))
 			{
@@ -262,6 +353,7 @@ void FDeferredShadingSceneRenderer::AddClusteredDeferredShadingPass(
 					SceneTextures,
 					SortedLightsSet,
 					EClusterPassInputType::HairStrands,
+					EStrataTileMaterialType::ECount,
 					HairStrandsShadowMaskBits,
 					VirtualShadowMapArray,
 					GraphBuilder.CreateSRV(TransmittanceMask.TransmittanceMask, FHairStrandsTransmittanceMaskData::Format));
