Fix swapping of CachedRayTracingInstance.DefaultUserData.

diff --git a/Engine/Source/Runtime/Renderer/Private/RendererScene.cpp b/Engine/Source/Runtime/Renderer/Private/RendererScene.cpp
--- a/Engine/Source/Runtime/Renderer/Private/RendererScene.cpp
+++ b/Engine/Source/Runtime/Renderer/Private/RendererScene.cpp
@@ -4479,9 +4479,8 @@ void FScene::UpdateAllPrimitiveSceneInfos(FRDGBuilder& GraphBuilder, bool bAsync
 							GPUScene.RecordPrimitiveIdSwap(DestIndex, SourceIndex);
 
 						#if RHI_RAYTRACING
-							// Update cached PrimitiveIndex after an index swap
-							Primitives[SourceIndex]->CachedRayTracingInstance.DefaultUserData = SourceIndex;
-							Primitives[DestIndex]->CachedRayTracingInstance.DefaultUserData = DestIndex;
+							// Also need to swap DefaultUserData since it currently contains PrimitiveIndex
+							Swap(Primitives[SourceIndex]->CachedRayTracingInstance.DefaultUserData, Primitives[DestIndex]->CachedRayTracingInstance.DefaultUserData);
 						#endif
 
 							SourceIndex = DestIndex;
@@ -4817,8 +4816,8 @@ void FScene::UpdateAllPrimitiveSceneInfos(FRDGBuilder& GraphBuilder, bool bAsync
 							GPUScene.RecordPrimitiveIdSwap(DestIndex, SourceIndex);
 
 						#if RHI_RAYTRACING
-							// Update cached PrimitiveIndex after an index swap
-							Primitives[DestIndex]->CachedRayTracingInstance.DefaultUserData = DestIndex;
+							// Also need to swap DefaultUserData since it currently contains PrimitiveIndex
+							Swap(Primitives[SourceIndex]->CachedRayTracingInstance.DefaultUserData, Primitives[DestIndex]->CachedRayTracingInstance.DefaultUserData);
 						#endif
 						}
 					}
