Avoid asserts even when VSM is disabled/unsupported (in this case on Android devices with NPOT 2D texture size maxes)
Avoid a bunch of unnecessary math when VSM is disabled

diff --git a/Engine/Source/Runtime/Renderer/Private/VirtualShadowMaps/VirtualShadowMapArray.cpp b/Engine/Source/Runtime/Renderer/Private/VirtualShadowMaps/VirtualShadowMapArray.cpp
--- a/Engine/Source/Runtime/Renderer/Private/VirtualShadowMaps/VirtualShadowMapArray.cpp
+++ b/Engine/Source/Runtime/Renderer/Private/VirtualShadowMaps/VirtualShadowMapArray.cpp
@@ -334,52 +334,15 @@ void FVirtualShadowMapArray::Initialize(FRDGBuilder& GraphBuilder, FVirtualShado
 	bInitialized = true;
 	bEnabled = bInEnabled;
 	CacheManager = InCacheManager;
-	check(!bEnabled || CacheManager);
 
 	bCullBackfacingPixels = CVarCullBackfacingPixels.GetValueOnRenderThread() != 0;
 
 	UniformParameters.NumShadowMaps = 0;
 	UniformParameters.NumDirectionalLights = 0;
-
-	// Fixed physical page pool width, we adjust the height to accomodate the requested maximum
-	// NOTE: Row size in pages has to be POT since we use mask & shift in place of integer ops
-	const uint32 PhysicalPagesX = FMath::DivideAndRoundDown(GetMax2DTextureDimension(), FVirtualShadowMap::PageSize);
-	check(FMath::IsPowerOfTwo(PhysicalPagesX));
-	uint32 PhysicalPagesY = FMath::DivideAndRoundUp((uint32)FMath::Max(1, CVarMaxPhysicalPages.GetValueOnRenderThread()), PhysicalPagesX);	
-
-	UniformParameters.MaxPhysicalPages = PhysicalPagesX * PhysicalPagesY;
-
-	if (CacheManager->IsValid() && CVarCacheStaticSeparate.GetValueOnRenderThread() != 0)
-	{
-		// Store the static pages below the dynamic/merged pages
-		UniformParameters.StaticCachedPixelOffsetY = PhysicalPagesY * FVirtualShadowMap::PageSize;
-		// Offset to static pages in linear page index
-		UniformParameters.StaticPageIndexOffset = PhysicalPagesY * PhysicalPagesX;
-		PhysicalPagesY *= 2;
-	}
-	else
-	{
-		UniformParameters.StaticCachedPixelOffsetY = 0;
-		UniformParameters.StaticPageIndexOffset = 0;
-	}
-
-	uint32 PhysicalX = PhysicalPagesX * FVirtualShadowMap::PageSize;
-	uint32 PhysicalY = PhysicalPagesY * FVirtualShadowMap::PageSize;
-
-	// TODO: Some sort of better fallback with warning?
-	// All supported platforms support at least 16384 texture dimensions which translates to 16384 max pages with default 128x128 page size
-	check(PhysicalX <= GetMax2DTextureDimension());
-	check(PhysicalY <= GetMax2DTextureDimension());
-
-	UniformParameters.PhysicalPageRowMask = (PhysicalPagesX - 1);
-	UniformParameters.PhysicalPageRowShift = FMath::FloorLog2( PhysicalPagesX );
-	UniformParameters.RecPhysicalPoolSize = FVector4f( 1.0f / PhysicalX, 1.0f / PhysicalY, 1.0f, 1.0f );
-	UniformParameters.PhysicalPoolSize = FIntPoint( PhysicalX, PhysicalY );
-	UniformParameters.PhysicalPoolSizePages = FIntPoint( PhysicalPagesX, PhysicalPagesY );
-
-	// TODO: Parameterize this in a useful way; potentially modify it automatically
-	// when there are fewer lights in the scene and/or clustered shading settings differ.
-	UniformParameters.PackedShadowMaskMaxLightCount = FMath::Min(CVarVirtualShadowOnePassProjectionMaxLights.GetValueOnRenderThread(), 32);
+	UniformParameters.MaxPhysicalPages = 0;
+	UniformParameters.StaticCachedPixelOffsetY = 0;
+	UniformParameters.StaticPageIndexOffset = 0;
+	// NOTE: Most uniform values don't matter when VSM is disabled
 
 	// Reference dummy data in the UB initially
 	const uint32 DummyPageTableElement = 0xFFFFFFFF;
@@ -390,12 +353,48 @@ void FVirtualShadowMapArray::Initialize(FRDGBuilder& GraphBuilder, FVirtualShado
 
 	if (bEnabled)
 	{
+		// Fixed physical page pool width, we adjust the height to accomodate the requested maximum
+		// NOTE: Row size in pages has to be POT since we use mask & shift in place of integer ops
+		// NOTE: This assumes GetMax2DTextureDimension() is a power of two on supported platforms
+		const uint32 PhysicalPagesX = FMath::DivideAndRoundDown(GetMax2DTextureDimension(), FVirtualShadowMap::PageSize);
+		check(FMath::IsPowerOfTwo(PhysicalPagesX));
+		uint32 PhysicalPagesY = FMath::DivideAndRoundUp((uint32)FMath::Max(1, CVarMaxPhysicalPages.GetValueOnRenderThread()), PhysicalPagesX);	
+
+		UniformParameters.MaxPhysicalPages = PhysicalPagesX * PhysicalPagesY;
+
+		if (CacheManager->IsValid() && CVarCacheStaticSeparate.GetValueOnRenderThread() != 0)
+		{
+			// Store the static pages below the dynamic/merged pages
+			UniformParameters.StaticCachedPixelOffsetY = PhysicalPagesY * FVirtualShadowMap::PageSize;
+			// Offset to static pages in linear page index
+			UniformParameters.StaticPageIndexOffset = PhysicalPagesY * PhysicalPagesX;
+			PhysicalPagesY *= 2;
+		}
+
+		uint32 PhysicalX = PhysicalPagesX * FVirtualShadowMap::PageSize;
+		uint32 PhysicalY = PhysicalPagesY * FVirtualShadowMap::PageSize;
+
+		// TODO: Some sort of better fallback with warning?
+		// All supported platforms support at least 16384 texture dimensions which translates to 16384 max pages with default 128x128 page size
+		check(PhysicalX <= GetMax2DTextureDimension());
+		check(PhysicalY <= GetMax2DTextureDimension());
+
+		UniformParameters.PhysicalPageRowMask = (PhysicalPagesX - 1);
+		UniformParameters.PhysicalPageRowShift = FMath::FloorLog2( PhysicalPagesX );
+		UniformParameters.RecPhysicalPoolSize = FVector4f( 1.0f / PhysicalX, 1.0f / PhysicalY, 1.0f, 1.0f );
+		UniformParameters.PhysicalPoolSize = FIntPoint( PhysicalX, PhysicalY );
+		UniformParameters.PhysicalPoolSizePages = FIntPoint( PhysicalPagesX, PhysicalPagesY );
+
+		// TODO: Parameterize this in a useful way; potentially modify it automatically
+		// when there are fewer lights in the scene and/or clustered shading settings differ.
+		UniformParameters.PackedShadowMaskMaxLightCount = FMath::Min(CVarVirtualShadowOnePassProjectionMaxLights.GetValueOnRenderThread(), 32);
+
 		// If enabled, ensure we have a properly-sized physical page pool
 		// We can do this here since the pool is independent of the number of shadow maps
 		TRefCountPtr<IPooledRenderTarget> PhysicalPagePool = CacheManager->SetPhysicalPoolSize(GraphBuilder, GetPhysicalPoolSize());
 		PhysicalPagePoolRDG = GraphBuilder.RegisterExternalTexture(PhysicalPagePool);
 		UniformParameters.PhysicalPagePool = PhysicalPagePoolRDG;
-	}
+	}	
 	else
 	{
 		CacheManager->FreePhysicalPool();
