Clean up strata tile binding.

diff --git a/Engine/Source/Runtime/Renderer/Private/ClusteredDeferredShadingPass.cpp b/Engine/Source/Runtime/Renderer/Private/ClusteredDeferredShadingPass.cpp
--- a/Engine/Source/Runtime/Renderer/Private/ClusteredDeferredShadingPass.cpp
+++ b/Engine/Source/Runtime/Renderer/Private/ClusteredDeferredShadingPass.cpp
@@ -107,14 +107,8 @@ class FClusteredShadingPS : public FGlobalShader
 		SHADER_PARAMETER_RDG_UNIFORM_BUFFER(FSceneTextureUniformParameters, SceneTextures)
 		SHADER_PARAMETER_RDG_TEXTURE(Texture2D, ShadowMaskBits)
 		SHADER_PARAMETER_STRUCT_INCLUDE(FVirtualShadowMapSamplingParameters, VirtualShadowMapSamplingParameters)
+		SHADER_PARAMETER_STRUCT_INCLUDE(Strata::FStrataTilePassVS::FParameters, StrataTile)
 		SHADER_PARAMETER_RDG_BUFFER_SRV(Buffer<uint>, HairTransmittanceBuffer)
-		// Strata tiles
-		SHADER_PARAMETER_RDG_BUFFER_SRV(Buffer<uint>, TileListBufferSimple)
-		SHADER_PARAMETER_RDG_BUFFER_SRV(Buffer<uint>, TileListBufferSingle)
-		SHADER_PARAMETER_RDG_BUFFER_SRV(Buffer<uint>, TileListBufferComplex)
-		RDG_BUFFER_ACCESS(TileIndirectBufferSimple,  ERHIAccess::IndirectArgs)
-		RDG_BUFFER_ACCESS(TileIndirectBufferSingle,  ERHIAccess::IndirectArgs)
-		RDG_BUFFER_ACCESS(TileIndirectBufferComplex, ERHIAccess::IndirectArgs)
 		RENDER_TARGET_BINDING_SLOTS()
 	END_SHADER_PARAMETER_STRUCT()
 
@@ -182,31 +176,15 @@ static void InternalAddClusteredDeferredShadingPass(
 	}
 
 	// VS - Strata tile parameters
-	if (bStrata)
-	{
-		PassParameters->TileListBufferSimple = View.StrataSceneData->ClassificationTileListBufferSRV[EStrataTileMaterialType::ESimple];
-		PassParameters->TileListBufferSingle = View.StrataSceneData->ClassificationTileListBufferSRV[EStrataTileMaterialType::ESingle];
-		PassParameters->TileListBufferComplex = View.StrataSceneData->ClassificationTileListBufferSRV[EStrataTileMaterialType::EComplex];
-		PassParameters->TileIndirectBufferSimple = View.StrataSceneData->ClassificationTileIndirectBuffer[EStrataTileMaterialType::ESimple];
-		PassParameters->TileIndirectBufferSingle = View.StrataSceneData->ClassificationTileIndirectBuffer[EStrataTileMaterialType::ESingle];
-		PassParameters->TileIndirectBufferComplex = View.StrataSceneData->ClassificationTileIndirectBuffer[EStrataTileMaterialType::EComplex];
-	}
-	else
-	{
-		FRDGBufferRef BufferDummy = GSystemTextures.GetDefaultBuffer(GraphBuilder, 4, 0u);
-		FRDGBufferSRVRef BufferDummySRV = GraphBuilder.CreateSRV(BufferDummy, PF_R32_UINT);
-		PassParameters->TileListBufferSimple = BufferDummySRV;
-		PassParameters->TileListBufferComplex = BufferDummySRV;
-		PassParameters->TileIndirectBufferSimple = BufferDummy;
-		PassParameters->TileIndirectBufferComplex = BufferDummy;
-	}
-
+	EPrimitiveType PrimitiveType = PT_TriangleList;
+	PassParameters->StrataTile = Strata::SetTileParameters(GraphBuilder, View, TileType, PrimitiveType);
+	
 	const TCHAR* TileTypeName = ToString(TileType);
 	GraphBuilder.AddPass(
 		RDG_EVENT_NAME("Light::ClusteredDeferredShading(%s,Lights:%d,Tile:%s)", bHairStrands ? TEXT("HairStrands") : (bStrata ? TEXT("Strata") : TEXT("GBuffer")), SortedLightsSet.ClusteredSupportedEnd, TileTypeName),
 		PassParameters,
 		ERDGPassFlags::Raster,
-		[PassParameters, &View, SceneTextureExtent, bHairStrands, bStrata, TileType](FRHICommandListImmediate& InRHICmdList)
+		[PassParameters, &View, SceneTextureExtent, bHairStrands, bStrata, TileType, PrimitiveType](FRHICommandListImmediate& InRHICmdList)
 	{
 		TShaderMapRef<FClusteredShadingVS> HairVertexShader(View.ShaderMap);
 		TShaderMapRef<FPostProcessVS> VertexShader(View.ShaderMap);
@@ -215,8 +193,6 @@ static void InternalAddClusteredDeferredShadingPass(
 		VSPermutationVector.Set<Strata::FStrataTilePassVS::FEnableDebug>(false);
 		VSPermutationVector.Set<Strata::FStrataTilePassVS::FEnableTexCoordScreenVector>(true);
 		TShaderMapRef<Strata::FStrataTilePassVS> TileVertexShader(View.ShaderMap, VSPermutationVector);
-		Strata::FStrataTilePassVS::FParameters VSParameters;
-
 
 		FClusteredShadingPS::FPermutationDomain PermutationVector;
 		PermutationVector.Set<FClusteredShadingPS::FVisualizeLightCullingDim>(View.Family->EngineShowFlags.VisualizeLightCulling);
@@ -236,12 +212,7 @@ static void InternalAddClusteredDeferredShadingPass(
 			GraphicsPSOInit.BoundShaderState.VertexDeclarationRHI = GFilterVertexDeclaration.VertexDeclarationRHI;
 			GraphicsPSOInit.BoundShaderState.VertexShaderRHI = bHairStrands ? HairVertexShader.GetVertexShader() : (bStrata ? TileVertexShader.GetVertexShader() : VertexShader.GetVertexShader());
 			GraphicsPSOInit.BoundShaderState.PixelShaderRHI = PixelShader.GetPixelShader();
-			GraphicsPSOInit.PrimitiveType = PT_TriangleList;
-			if (bStrata)
-			{
-				Strata::FillUpTiledPassData(TileType, View, VSParameters, GraphicsPSOInit.PrimitiveType);
-			}
-
+			GraphicsPSOInit.PrimitiveType = PrimitiveType;
 			SetGraphicsPipelineState(InRHICmdList, GraphicsPSOInit, 0);
 		}
 
@@ -261,8 +232,8 @@ static void InternalAddClusteredDeferredShadingPass(
 		}
 		else if (bStrata)
 		{
-			SetShaderParameters(InRHICmdList, TileVertexShader, TileVertexShader.GetVertexShader(), VSParameters);
-			InRHICmdList.DrawPrimitiveIndirect(VSParameters.TileIndirectBuffer->GetIndirectRHICallBuffer(), 0);
+			SetShaderParameters(InRHICmdList, TileVertexShader, TileVertexShader.GetVertexShader(), PassParameters->StrataTile);
+			InRHICmdList.DrawPrimitiveIndirect(PassParameters->StrataTile.TileIndirectBuffer->GetIndirectRHICallBuffer(), 0);
 		}
 		else
 		{
diff --git a/Engine/Source/Runtime/Renderer/Private/IndirectLightRendering.cpp b/Engine/Source/Runtime/Renderer/Private/IndirectLightRendering.cpp
--- a/Engine/Source/Runtime/Renderer/Private/IndirectLightRendering.cpp
+++ b/Engine/Source/Runtime/Renderer/Private/IndirectLightRendering.cpp
@@ -1276,7 +1276,7 @@ static void AddSkyReflectionPass(
 	if (bStrataEnabled)
 	{
 		check(StrataTileMaterialType != EStrataTileMaterialType::ECount);
-		Strata::FillUpTiledPassData(StrataTileMaterialType, View, PassParameters->VS, StrataTilePrimitiveType);
+		PassParameters->VS = Strata::SetTileParameters(GraphBuilder, View, StrataTileMaterialType, StrataTilePrimitiveType);
 		ClearUnusedGraphResources(StrataTilePassVertexShader, &PassParameters->VS);
 	}
 
diff --git a/Engine/Source/Runtime/Renderer/Private/LightRendering.cpp b/Engine/Source/Runtime/Renderer/Private/LightRendering.cpp
--- a/Engine/Source/Runtime/Renderer/Private/LightRendering.cpp
+++ b/Engine/Source/Runtime/Renderer/Private/LightRendering.cpp
@@ -1883,7 +1883,7 @@ static void InternalRenderLight(
 			Strata::FStrataTilePassVS::FParameters VSParameters;
 			if (Strata::IsStrataEnabled())
 			{
-				Strata::FillUpTiledPassData(StrataTileMaterialType, View, VSParameters, GraphicsPSOInit.PrimitiveType);
+				VSParameters = Strata::SetTileParameters(View, StrataTileMaterialType, GraphicsPSOInit.PrimitiveType);
 			}
 
 			// Turn DBT back off
@@ -1985,12 +1985,9 @@ BEGIN_SHADER_PARAMETER_STRUCT(FRenderLightParameters, )
 	SHADER_PARAMETER_STRUCT_INCLUDE(FDeferredLightPS::FParameters, PS)
 	SHADER_PARAMETER_STRUCT_INCLUDE(FDeferredLightVS::FParameters, VS)
 	// Strata tiles
-	SHADER_PARAMETER_RDG_BUFFER_SRV(Buffer<uint>, TileListBufferSimple)
-	SHADER_PARAMETER_RDG_BUFFER_SRV(Buffer<uint>, TileListBufferSingle)
-	SHADER_PARAMETER_RDG_BUFFER_SRV(Buffer<uint>, TileListBufferComplex)
-	RDG_BUFFER_ACCESS(TileIndirectBufferSimple, ERHIAccess::IndirectArgs)
-	RDG_BUFFER_ACCESS(TileIndirectBufferSingle, ERHIAccess::IndirectArgs)
-	RDG_BUFFER_ACCESS(TileIndirectBufferComplex, ERHIAccess::IndirectArgs)
+	SHADER_PARAMETER_STRUCT_INCLUDE(FStrataTileParameter, StrataTileSimple)
+	SHADER_PARAMETER_STRUCT_INCLUDE(FStrataTileParameter, StrataTileSingle)
+	SHADER_PARAMETER_STRUCT_INCLUDE(FStrataTileParameter, StrataTileComplex)
 END_SHADER_PARAMETER_STRUCT()
 
 /** Shader parameters for Standard Deferred Light pass. */
@@ -2075,21 +2072,11 @@ static void RenderLight(
 		// VS - Strata tile parameters
 		if (Strata::IsStrataEnabled())
 		{
-			PassParameters->TileListBufferSimple = View.StrataSceneData->ClassificationTileListBufferSRV[EStrataTileMaterialType::ESimple];
-			PassParameters->TileListBufferSingle = View.StrataSceneData->ClassificationTileListBufferSRV[EStrataTileMaterialType::ESingle];
-			PassParameters->TileListBufferComplex = View.StrataSceneData->ClassificationTileListBufferSRV[EStrataTileMaterialType::EComplex];
-			PassParameters->TileIndirectBufferSimple = View.StrataSceneData->ClassificationTileIndirectBuffer[EStrataTileMaterialType::ESimple];
-			PassParameters->TileIndirectBufferSingle = View.StrataSceneData->ClassificationTileIndirectBuffer[EStrataTileMaterialType::ESingle];
-			PassParameters->TileIndirectBufferComplex = View.StrataSceneData->ClassificationTileIndirectBuffer[EStrataTileMaterialType::EComplex];
-		}
-		else
-		{
-			FRDGBufferRef BufferDummy = GSystemTextures.GetDefaultBuffer(GraphBuilder, 4, 0u);
-			FRDGBufferSRVRef BufferDummySRV = GraphBuilder.CreateSRV(BufferDummy, PF_R32_UINT);
-			PassParameters->TileListBufferSimple = BufferDummySRV;
-			PassParameters->TileListBufferComplex = BufferDummySRV;
-			PassParameters->TileIndirectBufferSimple = BufferDummy;
-			PassParameters->TileIndirectBufferComplex = BufferDummy;
+			// Note: we register all tile types here in order to have all resources tracked properly and being able 
+			//       to create a single pass parameters struct instead of created one for each tile types
+			PassParameters->StrataTileSimple = Strata::SetTileParameters(GraphBuilder, View, EStrataTileMaterialType::ESingle);
+			PassParameters->StrataTileSingle = Strata::SetTileParameters(GraphBuilder, View, EStrataTileMaterialType::ESimple);
+			PassParameters->StrataTileComplex = Strata::SetTileParameters(GraphBuilder, View, EStrataTileMaterialType::EComplex);
 		}
 
 		FDeferredLightPS::FPermutationDomain PermutationVector;
diff --git a/Engine/Source/Runtime/Renderer/Private/Strata/Strata.cpp b/Engine/Source/Runtime/Renderer/Private/Strata/Strata.cpp
--- a/Engine/Source/Runtime/Renderer/Private/Strata/Strata.cpp
+++ b/Engine/Source/Runtime/Renderer/Private/Strata/Strata.cpp
@@ -500,20 +500,62 @@ class FStrataMaterialStencilTaggingPassPS : public FGlobalShader
 IMPLEMENT_GLOBAL_SHADER(FStrataTilePassVS, "/Engine/Private/Strata/StrataTiling.usf", "StrataTilePassVS", SF_Vertex);
 IMPLEMENT_GLOBAL_SHADER(FStrataMaterialStencilTaggingPassPS, "/Engine/Private/Strata/StrataTiling.usf", "StencilTaggingMainPS", SF_Pixel);
 
-void FillUpTiledPassData(
-	EStrataTileMaterialType Type, 
-	const FViewInfo& View, 
-	FStrataTilePassVS::FParameters& ParametersVS,
+static FStrataTileParameter InternalSetTileParameters(FRDGBuilder* GraphBuilder, const FViewInfo& View, const EStrataTileMaterialType TileType)
+{
+	FStrataTileParameter Out;
+	if (View.StrataSceneData && TileType != EStrataTileMaterialType::ECount)
+	{
+		Out.TileListBuffer = View.StrataSceneData->ClassificationTileListBufferSRV[TileType];
+		Out.TileIndirectBuffer = View.StrataSceneData->ClassificationTileIndirectBuffer[TileType];
+	}
+	else if (GraphBuilder)
+	{
+		FRDGBufferRef BufferDummy = GSystemTextures.GetDefaultBuffer(*GraphBuilder, 4, 0u);
+		FRDGBufferSRVRef BufferDummySRV = GraphBuilder->CreateSRV(BufferDummy, PF_R32_UINT);
+		Out.TileListBuffer = BufferDummySRV;
+		Out.TileIndirectBuffer = BufferDummy;
+	}
+	return Out;
+}
+
+FStrataTilePassVS::FParameters SetTileParameters(
+	const FViewInfo& View,
+	const EStrataTileMaterialType TileType,
 	EPrimitiveType& PrimitiveType)
 {
-	ParametersVS.OutputViewSizeAndInvSize = View.CachedViewUniformShaderParameters->ViewSizeAndInvSize;
-	ParametersVS.OutputBufferSizeAndInvSize = View.CachedViewUniformShaderParameters->BufferSizeAndInvSize;
-	ParametersVS.ViewScreenToTranslatedWorld = View.CachedViewUniformShaderParameters->ScreenToTranslatedWorld;
+	FStrataTileParameter Temp = InternalSetTileParameters(nullptr, View, TileType);
+	PrimitiveType = GRHISupportsRectTopology ? PT_RectList : PT_TriangleList;
 
-	ParametersVS.TileListBuffer = View.StrataSceneData->ClassificationTileListBufferSRV[Type];
-	ParametersVS.TileIndirectBuffer = View.StrataSceneData->ClassificationTileIndirectBuffer[Type];
+	FStrataTilePassVS::FParameters Out;
+	Out.OutputViewSizeAndInvSize = View.CachedViewUniformShaderParameters->ViewSizeAndInvSize;
+	Out.OutputBufferSizeAndInvSize = View.CachedViewUniformShaderParameters->BufferSizeAndInvSize;
+	Out.ViewScreenToTranslatedWorld = View.CachedViewUniformShaderParameters->ScreenToTranslatedWorld;
+	Out.TileListBuffer = Temp.TileListBuffer;
+	Out.TileIndirectBuffer = Temp.TileIndirectBuffer;
+	return Out;
+}
 
+FStrataTilePassVS::FParameters SetTileParameters(
+	FRDGBuilder& GraphBuilder, 
+	const FViewInfo& View, 
+	const EStrataTileMaterialType TileType,
+	EPrimitiveType& PrimitiveType)
+{
+	FStrataTileParameter Temp = InternalSetTileParameters(&GraphBuilder, View, TileType);
 	PrimitiveType = GRHISupportsRectTopology ? PT_RectList : PT_TriangleList;
+
+	FStrataTilePassVS::FParameters Out;
+	Out.OutputViewSizeAndInvSize = View.CachedViewUniformShaderParameters->ViewSizeAndInvSize;
+	Out.OutputBufferSizeAndInvSize = View.CachedViewUniformShaderParameters->BufferSizeAndInvSize;
+	Out.ViewScreenToTranslatedWorld = View.CachedViewUniformShaderParameters->ScreenToTranslatedWorld;
+	Out.TileListBuffer = Temp.TileListBuffer;
+	Out.TileIndirectBuffer = Temp.TileIndirectBuffer;
+	return Out;
+}
+
+FStrataTileParameter SetTileParameters(FRDGBuilder& GraphBuilder, const FViewInfo& View, const EStrataTileMaterialType TileType)
+{
+	return InternalSetTileParameters(&GraphBuilder, View, TileType);
 }
 
 static void AddStrataInternalClassificationTilePass(
@@ -529,7 +571,7 @@ static void AddStrataInternalClassificationTilePass(
 	FVector4f OutputResolutionAndInv = FVector4f(OutputResolution.X, OutputResolution.Y, 1.0f / float(OutputResolution.X), 1.0f / float(OutputResolution.Y));
 
 	FStrataMaterialStencilTaggingPassPS::FParameters* ParametersPS = GraphBuilder.AllocParameters<FStrataMaterialStencilTaggingPassPS::FParameters>();
-	FillUpTiledPassData(TileMaterialType, View, ParametersPS->VS, StrataTilePrimitiveType);
+	ParametersPS->VS = Strata::SetTileParameters(GraphBuilder, View, TileMaterialType, StrataTilePrimitiveType);
 
 	FStrataTilePassVS::FPermutationDomain VSPermutationVector;
 	VSPermutationVector.Set< FStrataTilePassVS::FEnableDebug >(bDebug);
diff --git a/Engine/Source/Runtime/Renderer/Private/Strata/Strata.h b/Engine/Source/Runtime/Renderer/Private/Strata/Strata.h
--- a/Engine/Source/Runtime/Renderer/Private/Strata/Strata.h
+++ b/Engine/Source/Runtime/Renderer/Private/Strata/Strata.h
@@ -26,6 +26,11 @@ BEGIN_SHADER_PARAMETER_STRUCT(FStrataForwardPassUniformParameters, )
 	SHADER_PARAMETER(uint32, bRoughDiffuse)
 END_SHADER_PARAMETER_STRUCT()
 
+BEGIN_SHADER_PARAMETER_STRUCT(FStrataTileParameter, )
+	SHADER_PARAMETER_RDG_BUFFER_SRV(Buffer<uint>, TileListBuffer)
+	RDG_BUFFER_ACCESS(TileIndirectBuffer, ERHIAccess::IndirectArgs)
+END_SHADER_PARAMETER_STRUCT()
+
 // This paramater struct is declared with RENDERER_API even though it is not public. This is
 // to workaround other modules doing 'private include' of the Renderer module
 BEGIN_GLOBAL_SHADER_PARAMETER_STRUCT(FStrataGlobalUniformParameters, RENDERER_API)
@@ -109,7 +114,6 @@ void AddStrataStencilPass(FRDGBuilder& GraphBuilder, const TArray<FViewInfo>& Vi
 bool ShouldRenderStrataDebugPasses(const FViewInfo& View);
 FScreenPassTexture AddStrataDebugPasses(FRDGBuilder& GraphBuilder, const FViewInfo& View, FScreenPassTexture& ScreenPassSceneColor);
 
-
 class FStrataTilePassVS : public FGlobalShader
 {
 	DECLARE_GLOBAL_SHADER(FStrataTilePassVS);
@@ -142,6 +146,8 @@ class FStrataTilePassVS : public FGlobalShader
 	}
 };
 
-void FillUpTiledPassData(EStrataTileMaterialType Type, const FViewInfo& View, FStrataTilePassVS::FParameters& ParametersVS, EPrimitiveType& PrimitiveType);
+FStrataTileParameter SetTileParameters(FRDGBuilder& GraphBuilder, const FViewInfo& View, const EStrataTileMaterialType Type);
+FStrataTilePassVS::FParameters SetTileParameters(FRDGBuilder& GraphBuilder, const FViewInfo& View, const EStrataTileMaterialType Type, EPrimitiveType& PrimitiveType);
+FStrataTilePassVS::FParameters SetTileParameters(const FViewInfo& View, const EStrataTileMaterialType Type, EPrimitiveType& PrimitiveType);
 
 };
