Add strata support for virtual shadow map

diff --git a/Engine/Shaders/Private/Strata/Strata.ush b/Engine/Shaders/Private/Strata/Strata.ush
--- a/Engine/Shaders/Private/Strata/Strata.ush
+++ b/Engine/Shaders/Private/Strata/Strata.ush
@@ -1716,6 +1716,13 @@ float StrataGetLegacyShadingModels(in FStrataBSDF BSDF)
 	return 0.f;
 }
 
+// Legacy conversion function, which translated Strata sub-surface data into 'opacity' for shadow transmission purpose. 
+// Note: When SSS data uses DFMP data, the conversion to opacity is ad-hoc/hacky
+float StrataGetOpacityFromSubSurface(FStrataSubsurfaceData In)
+{
+	return In.bIsValid ? (In.bIsProfile ? In.ProfileRadiusScale : saturate(1.f - max3(In.DMFP.x, In.DMFP.y, In.DMFP.z))) : 1.f;
+}
+
 ///////////////////////////////////////////////////////////////////////////////
 // BSDFs packing/unpacking
 // 11G11B10F and 10F are unsigned float format
diff --git a/Engine/Shaders/Private/VirtualShadowMaps/PageManagement.usf b/Engine/Shaders/Private/VirtualShadowMaps/PageManagement.usf
--- a/Engine/Shaders/Private/VirtualShadowMaps/PageManagement.usf
+++ b/Engine/Shaders/Private/VirtualShadowMaps/PageManagement.usf
@@ -32,6 +32,11 @@
 #define INPUT_TYPE_GBUFFER 0
 #define INPUT_TYPE_HAIRSTRANDS 1
 
+#if STRATA_ENABLED 
+#define STRATA_INLINE_SHADING 0
+#include "/Engine/Private/Strata/Strata.ush"
+#endif
+
 // Flags generated by per-pixel pass to determine which pages are required to provide shadow for the visible geometry
 RWStructuredBuffer<uint> OutPageRequestFlags;
 
@@ -199,12 +204,18 @@ void GeneratePageFlagsFromPixels(
 		FScreenPixelData Pixel = GetScreenPixelData(PixelPos, LookupDeviceZ(CalcScreenUV(PixelPos)));
 		WorldPosition = Pixel.WorldPosition;
 		SceneDepth = Pixel.SceneDepth;
-
+	  #if STRATA_ENABLED
+		const FStrataTopLayerData TopLayerData = StrataUnpackTopLayerData(Strata.TopLayerTexture.Load(uint3(PixelPos, 0)));
+		WorldNormal = TopLayerData.WorldNormal;
+		bIsValid = IsStrataMaterial(TopLayerData);
+		// bBackfaceCull = bBackfaceCull && !IsSubsurfaceModel(GBuffer.ShadingModelID); STRATA_TODO
+	  #else
 		FGBufferData GBuffer = GetGBufferDataUint(PixelPos, true);
 		WorldNormal = GBuffer.WorldNormal;
 		// Excluding unlit to avoid including processing sky dome
 		bIsValid = GBuffer.ShadingModelID != SHADINGMODELID_UNLIT;
 		bBackfaceCull = bBackfaceCull && !IsSubsurfaceModel(GBuffer.ShadingModelID);
+	  #endif // STRATA_ENABLED
 	}
 	#endif
 	
diff --git a/Engine/Shaders/Private/VirtualShadowMaps/TransmissionCommon.ush b/Engine/Shaders/Private/VirtualShadowMaps/TransmissionCommon.ush
--- a/Engine/Shaders/Private/VirtualShadowMaps/TransmissionCommon.ush
+++ b/Engine/Shaders/Private/VirtualShadowMaps/TransmissionCommon.ush
@@ -10,32 +10,16 @@ VirtualShadowMaps/TransmissionCommon.ush:
 #include "../DeferredShadingCommon.ush"
 #include "ProjectionCommon.ush"
 
-// TODO: Support strata
 #ifndef STRATA_ENABLED
-#define STRATA_ENABLED 0
+#error STRATA_ENABLED needs to be defined
 #endif
 
 half GetSubsurfaceOpacityFromGbuffer(uint2 PixelPos)
 {
 	half Opacity = 1.0;
 
-	// TODO: Support strata
-	// TODO: Once strata is set up this function should be shared with ShadowProjectionPixelShader
-	// and possibly also the contact shadow trace
-#if STRATA_ENABLED && 0
-	// TODO: VSM projection shaders are not set up for strata yet, but added a skeleton of the logic
-	const FStrataClassification StrataData = StrataUnpackClassificationData(Strata.ClassificationTexture.Load(PixelPos));
-	if (IsStrataMaterial(StrataData))
-	{
-		if (HasSubsurface(StrataData))
-		{
-			const FStrataSubsurfaceData SSSData = StrataUnpackSSSData(Strata.SSSTexture.Load(uint3(SVPos.xy, 0)));
-			if (SSSData.IsValid())
-			{
-				Opacity = SSSData.bIsProfile ? SSSData.ProfileRadiusScale : 1.0;
-			}
-		}
-	}
+#if STRATA_ENABLED
+	Opacity = StrataGetOpacityFromSubSurface(StrataUnpackSSSData(Strata.SSSTexture.Load(uint3(PixelPos, 0))));
 #else // STRATA_ENABLED
 	FGBufferData GBufferData = GetGBufferDataUint(PixelPos);
 
diff --git a/Engine/Shaders/Private/VirtualShadowMaps/VirtualShadowMapProjection.usf b/Engine/Shaders/Private/VirtualShadowMaps/VirtualShadowMapProjection.usf
--- a/Engine/Shaders/Private/VirtualShadowMaps/VirtualShadowMapProjection.usf
+++ b/Engine/Shaders/Private/VirtualShadowMaps/VirtualShadowMapProjection.usf
@@ -7,6 +7,10 @@
 #define SUPPORT_CONTACT_SHADOWS 1
 #include "../Common.ush"
 #include "../SceneTexturesCommon.ush"
+#if STRATA_ENABLED 
+#define STRATA_INLINE_SHADING 0
+#include "/Engine/Private/Strata/Strata.ush"
+#endif
 #include "../DeferredShadingCommon.ush"
 #include "../DeferredLightingCommon.ush"
 #include "../LightShaderParameters.ush"
diff --git a/Engine/Source/Runtime/Renderer/Private/DeferredShadingRenderer.cpp b/Engine/Source/Runtime/Renderer/Private/DeferredShadingRenderer.cpp
--- a/Engine/Source/Runtime/Renderer/Private/DeferredShadingRenderer.cpp
+++ b/Engine/Source/Runtime/Renderer/Private/DeferredShadingRenderer.cpp
@@ -2696,6 +2696,13 @@ void FDeferredShadingSceneRenderer::Render(FRDGBuilder& GraphBuilder)
 		RenderHairBasePass(GraphBuilder, Scene, SceneTextures, Views, InstanceCullingManager);
 	}
 
+	// Post base pass for material classification
+	// This needs to run before virtual shadow map, in order to have ready&cleared classified SSS data
+	if (Strata::IsStrataEnabled())
+	{
+		Strata::AddStrataMaterialClassificationPass(GraphBuilder, SceneTextures, Views);
+	}
+
 	FLumenSceneFrameTemporaries LumenFrameTemporaries;
 
 	// Shadows, lumen and fog after base pass
@@ -2785,12 +2792,6 @@ void FDeferredShadingSceneRenderer::Render(FRDGBuilder& GraphBuilder)
 	// Copy lighting channels out of stencil before deferred decals which overwrite those values
 	FRDGTextureRef LightingChannelsTexture = CopyStencilToLightingChannelTexture(GraphBuilder, SceneTextures.Stencil);
 
-	// Post base pass for material classification
-	if (Strata::IsStrataEnabled())
-	{
-		Strata::AddStrataMaterialClassificationPass(GraphBuilder, SceneTextures, Views);
-	}
-
 	// Pre-lighting composition lighting stage
 	// e.g. deferred decals, SSAO
 	{
diff --git a/Engine/Source/Runtime/Renderer/Private/VirtualShadowMaps/VirtualShadowMapArray.cpp b/Engine/Source/Runtime/Renderer/Private/VirtualShadowMaps/VirtualShadowMapArray.cpp
--- a/Engine/Source/Runtime/Renderer/Private/VirtualShadowMaps/VirtualShadowMapArray.cpp
+++ b/Engine/Source/Runtime/Renderer/Private/VirtualShadowMaps/VirtualShadowMapArray.cpp
@@ -528,6 +528,7 @@ public:
 		OutEnvironment.SetDefine(TEXT("VSM_DEFAULT_CS_GROUP_XY"), DefaultCSGroupXY);
 		OutEnvironment.SetDefine(TEXT("VSM_GENERATE_PAGE_FLAGS_CS_GROUP_XYZ"), GeneratePageFlagsGroupXYZ);
 		OutEnvironment.SetDefine(TEXT("VSM_BUILD_EXPLICIT_BOUNDS_CS_XY"), BuildExplicitBoundsGroupXY);
+		OutEnvironment.SetDefine(TEXT("STRATA_ENABLED"), Strata::IsStrataEnabled() ? 1u : 0u);
 
 		FForwardLightingParameters::ModifyCompilationEnvironment(Parameters.Platform, OutEnvironment);
 		OutEnvironment.SetDefine(TEXT("VF_SUPPORTS_PRIMITIVE_SCENE_DATA"), 1);
@@ -552,6 +553,7 @@ class FGeneratePageFlagsFromPixelsCS : public FVirtualPageManagementShader
 		SHADER_PARAMETER_RDG_TEXTURE(Texture2D<uint2>, VisBuffer64)
 		SHADER_PARAMETER_RDG_BUFFER_UAV(RWStructuredBuffer<uint>, OutPageRequestFlags)
 		SHADER_PARAMETER_RDG_BUFFER_SRV(StructuredBuffer< uint >, DirectionalLightIds)
+		SHADER_PARAMETER_RDG_UNIFORM_BUFFER(FStrataGlobalUniformParameters, Strata)
 		RDG_BUFFER_ACCESS(IndirectBufferArgs, ERHIAccess::IndirectArgs)
 		SHADER_PARAMETER(uint32, InputType)
 		SHADER_PARAMETER(uint32, NumDirectionalLightSmInds)
@@ -1180,12 +1182,10 @@ void FVirtualShadowMapArray::BuildPageAllocations(
 			// Mark pages based on projected depth buffer pixels
 			if (CVarMarkPixelPages.GetValueOnRenderThread() != 0)
 			{
-				auto GeneratePageFlags = [&](bool bHairPass)
+				auto GeneratePageFlags = [&](const EVirtualShadowMapProjectionInputType InputType)
 				{
-					const uint32 InputType = bHairPass ? 1U : 0U; // HairStrands or GBuffer
-
 					FGeneratePageFlagsFromPixelsCS::FPermutationDomain PermutationVector;
-					PermutationVector.Set<FGeneratePageFlagsFromPixelsCS::FInputType>(InputType);
+					PermutationVector.Set<FGeneratePageFlagsFromPixelsCS::FInputType>(InputType == EVirtualShadowMapProjectionInputType::HairStrands ? 1u : 0u);
 					FGeneratePageFlagsFromPixelsCS::FParameters* PassParameters = GraphBuilder.AllocParameters< FGeneratePageFlagsFromPixelsCS::FParameters >();
 					PassParameters->VirtualShadowMap = GetUniformBuffer(GraphBuilder);
 
@@ -1200,13 +1200,14 @@ void FVirtualShadowMapArray::BuildPageAllocations(
 					PassParameters->PageDilationBorderSizeLocal = CVarPageDilationBorderSizeLocal.GetValueOnRenderThread();
 					PassParameters->PageDilationBorderSizeDirectional = CVarPageDilationBorderSizeDirectional.GetValueOnRenderThread();
 					PassParameters->bCullBackfacingPixels = ShouldCullBackfacingPixels() ? 1 : 0;
+					PassParameters->Strata = Strata::BindStrataGlobalUniformParameters(View.StrataSceneData);
 
 					auto ComputeShader = View.ShaderMap->GetShader<FGeneratePageFlagsFromPixelsCS>(PermutationVector);
 
 					static_assert((FVirtualPageManagementShader::DefaultCSGroupXY % 2) == 0, "GeneratePageFlagsFromPixels requires even-sized CS groups for quad swizzling.");
 					const FIntPoint GridSize = FIntPoint::DivideAndRoundUp(View.ViewRect.Size(), FVirtualPageManagementShader::DefaultCSGroupXY);
 
-					if (bHairPass && View.HairStrandsViewData.VisibilityData.TileData.IsValid())
+					if (InputType == EVirtualShadowMapProjectionInputType::HairStrands && View.HairStrandsViewData.VisibilityData.TileData.IsValid())
 					{
 						PassParameters->IndirectBufferArgs = View.HairStrandsViewData.VisibilityData.TileData.TileIndirectDispatchBuffer;
 						FComputeShaderUtils::AddPass(
@@ -1221,17 +1222,17 @@ void FVirtualShadowMapArray::BuildPageAllocations(
 					{
 						FComputeShaderUtils::AddPass(
 							GraphBuilder,
-							RDG_EVENT_NAME("GeneratePageFlagsFromPixels(%s)", bHairPass ? TEXT("HairStrands") : TEXT("GBuffer")),
+							RDG_EVENT_NAME("GeneratePageFlagsFromPixels(%s)", ToString(InputType)),
 							ComputeShader,
 							PassParameters,
 							FIntVector(GridSize.X, GridSize.Y, 1));
 					}
 				};
 
-				GeneratePageFlags(false);
+				GeneratePageFlags(EVirtualShadowMapProjectionInputType::GBuffer);
 				if (HairStrands::HasViewHairStrandsData(View))
 				{
-					GeneratePageFlags(true);
+					GeneratePageFlags(EVirtualShadowMapProjectionInputType::HairStrands);
 				}
 			}
 			// Mark coarse pages
diff --git a/Engine/Source/Runtime/Renderer/Private/VirtualShadowMaps/VirtualShadowMapProjection.cpp b/Engine/Source/Runtime/Renderer/Private/VirtualShadowMaps/VirtualShadowMapProjection.cpp
--- a/Engine/Source/Runtime/Renderer/Private/VirtualShadowMaps/VirtualShadowMapProjection.cpp
+++ b/Engine/Source/Runtime/Renderer/Private/VirtualShadowMaps/VirtualShadowMapProjection.cpp
@@ -116,6 +116,15 @@ static TAutoConsoleVariable<int32> CVarForcePerLightShadowMaskClear(
 	ECVF_RenderThreadSafe
 );
 
+const TCHAR* ToString(EVirtualShadowMapProjectionInputType In)
+{
+	switch (In)
+	{
+	case EVirtualShadowMapProjectionInputType::HairStrands: return TEXT("HairStrands");
+	case EVirtualShadowMapProjectionInputType::GBuffer:     return Strata::IsStrataEnabled() ? TEXT("Strata") : TEXT("GBuffer");
+	}
+	return TEXT("Invalid");
+}
 
 // Composite denoised shadow projection mask onto the light's shadow mask
 // Basically just a copy shader with a special blend mode
@@ -200,6 +209,7 @@ class FVirtualShadowMapProjectionCS : public FGlobalShader
 		SHADER_PARAMETER_RDG_UNIFORM_BUFFER(FSceneTextureUniformParameters, SceneTexturesStruct)
 		SHADER_PARAMETER_RDG_UNIFORM_BUFFER(FHairStrandsViewUniformParameters, HairStrands)
 		SHADER_PARAMETER_RDG_UNIFORM_BUFFER(FVirtualVoxelParameters, HairStrandsVoxel)
+		SHADER_PARAMETER_RDG_UNIFORM_BUFFER(FStrataGlobalUniformParameters, Strata)
 		SHADER_PARAMETER_STRUCT_REF(FViewUniformShaderParameters, View)
 		SHADER_PARAMETER(FIntVector4, ProjectionRect)
 		SHADER_PARAMETER(float, ContactShadowLength)
@@ -231,9 +241,9 @@ class FVirtualShadowMapProjectionCS : public FGlobalShader
 
 		FVirtualShadowMapArray::SetShaderDefines(OutEnvironment);
 		FForwardLightingParameters::ModifyCompilationEnvironment(Parameters.Platform, OutEnvironment);
+		OutEnvironment.SetDefine(TEXT("STRATA_ENABLED"), Strata::IsStrataEnabled() ? 1u : 0u);
 
 		FPermutationDomain PermutationVector(Parameters.PermutationId);
-
 		if (PermutationVector.Get<FSMRTAdaptiveRayCountDim>())
 		{
 			OutEnvironment.CompilerFlags.Add(CFLAG_WaveOperations);
@@ -289,6 +299,7 @@ static void RenderVirtualShadowMapProjectionCommon(
 	PassParameters->InputType = uint32(InputType);
 	PassParameters->bCullBackfacingPixels = VirtualShadowMapArray.ShouldCullBackfacingPixels() ? 1 : 0;
 	PassParameters->SMRTTexelDitherScale = CVarSMRTTexelDitherScale.GetValueOnRenderThread();
+	PassParameters->Strata = Strata::BindStrataGlobalUniformParameters(View.StrataSceneData);
 	if (bHasHairStrandsData)
 	{
 		PassParameters->HairStrands = HairStrands::BindHairStrandsViewUniformParameters(View);
@@ -359,7 +370,7 @@ static void RenderVirtualShadowMapProjectionCommon(
 		GraphBuilder,
 		RDG_EVENT_NAME("VirtualShadowMapProjection(RayCount:%s,Input:%s%s)", 
 			bAdaptiveRayCount ? TEXT("Adaptive") : TEXT("Static"), 
-			(InputType == EVirtualShadowMapProjectionInputType::HairStrands) ? TEXT("HairStrands") : TEXT("GBuffer"),
+			ToString(InputType),
 			bDebugOutput ? TEXT(",Debug") : TEXT("")),
 		ComputeShader,
 		PassParameters,
diff --git a/Engine/Source/Runtime/Renderer/Private/VirtualShadowMaps/VirtualShadowMapProjection.h b/Engine/Source/Runtime/Renderer/Private/VirtualShadowMaps/VirtualShadowMapProjection.h
--- a/Engine/Source/Runtime/Renderer/Private/VirtualShadowMaps/VirtualShadowMapProjection.h
+++ b/Engine/Source/Runtime/Renderer/Private/VirtualShadowMaps/VirtualShadowMapProjection.h
@@ -17,6 +17,7 @@ enum class EVirtualShadowMapProjectionInputType
 	GBuffer = 0,
 	HairStrands = 1
 };
+const TCHAR* ToString(EVirtualShadowMapProjectionInputType In);
 
 void RenderVirtualShadowMapProjection(
 	FRDGBuilder& GraphBuilder,
