Update EarlyZPassMode when anti-aliasing mode changes. It fixes skeletal mesh rendering broken in forward shading mode when switching from another anti-alising mode to MSAA.

diff --git a/Engine/Source/Runtime/Renderer/Private/RendererScene.cpp b/Engine/Source/Runtime/Renderer/Private/RendererScene.cpp
--- a/Engine/Source/Runtime/Renderer/Private/RendererScene.cpp
+++ b/Engine/Source/Runtime/Renderer/Private/RendererScene.cpp
@@ -843,15 +843,20 @@ static FAutoConsoleVariableSink CVarDoLazyStaticMeshUpdateSink(FConsoleCommandDe
 
 static void UpdateEarlyZPassModeCVarSinkFunction()
 {
+	static auto* CVarAntiAliasingMethod = IConsoleManager::Get().FindTConsoleVariableDataInt(TEXT("r.AntiAliasingMethod"));
+	static int32 CachedAntiAliasingMethod = CVarAntiAliasingMethod->GetValueOnGameThread();
 	static int32 CachedEarlyZPass = CVarEarlyZPass.GetValueOnGameThread();
 	static int32 CachedBasePassWriteDepthEvenWithFullPrepass = CVarBasePassWriteDepthEvenWithFullPrepass.GetValueOnGameThread();
 	static int32 CachedMobileEarlyZPass = CVarMobileEarlyZPass.GetValueOnGameThread();
 
+	const int32 AntiAliasingMethod = CVarAntiAliasingMethod->GetValueOnGameThread();
 	const int32 EarlyZPass = CVarEarlyZPass.GetValueOnGameThread();
 	const int32 BasePassWriteDepthEvenWithFullPrepass = CVarBasePassWriteDepthEvenWithFullPrepass.GetValueOnGameThread();
 	const int32 MobileEarlyZPass = CVarMobileEarlyZPass.GetValueOnGameThread();
 
-	if (EarlyZPass != CachedEarlyZPass
+	// Switching between MSAA and another AA in forward shading mode requires EarlyZPassMode to update.
+	if (AntiAliasingMethod != CachedAntiAliasingMethod
+		|| EarlyZPass != CachedEarlyZPass
 		|| BasePassWriteDepthEvenWithFullPrepass != CachedBasePassWriteDepthEvenWithFullPrepass
 		|| MobileEarlyZPass != CachedMobileEarlyZPass)
 	{
@@ -865,6 +870,7 @@ static void UpdateEarlyZPassModeCVarSinkFunction()
 			}
 		}
 
+		CachedAntiAliasingMethod = AntiAliasingMethod;
 		CachedEarlyZPass = EarlyZPass;
 		CachedBasePassWriteDepthEvenWithFullPrepass = BasePassWriteDepthEvenWithFullPrepass;
 		CachedMobileEarlyZPass = MobileEarlyZPass;
