Path Tracing should not be affected by the cvar r.RayTracing.ForceAllRayTracingEffects.

diff --git a/Engine/Source/Runtime/Renderer/Private/DeferredShadingRenderer.cpp b/Engine/Source/Runtime/Renderer/Private/DeferredShadingRenderer.cpp
--- a/Engine/Source/Runtime/Renderer/Private/DeferredShadingRenderer.cpp
+++ b/Engine/Source/Runtime/Renderer/Private/DeferredShadingRenderer.cpp
@@ -3116,8 +3116,7 @@ void FDeferredShadingSceneRenderer::Render(FRDGBuilder& GraphBuilder)
 		// Most of the ray tracing debug visualizations also require the full pipeline, but some support inline mode.
 		
 		if (ViewFamily.EngineShowFlags.PathTracing 
-			&& FDataDrivenShaderPlatformInfo::GetSupportsPathTracing(Scene->GetShaderPlatform())
-			&& ShouldRenderRayTracingEffect(ERayTracingPipelineCompatibilityFlags::FullPipeline))
+			&& FDataDrivenShaderPlatformInfo::GetSupportsPathTracing(Scene->GetShaderPlatform()))
 		{
 			for (const FViewInfo& View : Views)
 			{
