Fix for IsValid() crash when attempting to visualize the Separate Translucency buffer when r.SeparateTranslucency = 0

- If PostDOFTranslucencyResources is invalid we should set a 0,0,1,1 rect since we are using the BlackDummyAlpha buffer.
- This is similar to how DiaphragmDOF handles accessing the separate translucency buffers.

diff --git a/Engine/Source/Runtime/Renderer/Private/PostProcess/PostProcessing.cpp b/Engine/Source/Runtime/Renderer/Private/PostProcess/PostProcessing.cpp
--- a/Engine/Source/Runtime/Renderer/Private/PostProcess/PostProcessing.cpp
+++ b/Engine/Source/Runtime/Renderer/Private/PostProcess/PostProcessing.cpp
@@ -1159,7 +1159,14 @@ void AddPostProcessingPasses(
 		PassInputs.SceneColor = SceneColor;
 		PassInputs.SceneColorBeforeTonemap = SceneColorBeforeTonemap;
 		PassInputs.SceneColorAfterTonemap = SceneColorAfterTonemap;
-		PassInputs.SeparateTranslucency = FScreenPassTexture(PostDOFTranslucencyResources.GetColorForRead(GraphBuilder), PostDOFTranslucencyResources.ViewRect); // TODO
+
+		FIntRect ViewRect{ 0, 0, 1, 1 };
+		if (PostDOFTranslucencyResources.IsValid())
+		{
+			ViewRect = PostDOFTranslucencyResources.ViewRect;
+		}
+
+		PassInputs.SeparateTranslucency = FScreenPassTexture(PostDOFTranslucencyResources.GetColorForRead(GraphBuilder), ViewRect); // TODO
 		PassInputs.Velocity = Velocity;
 		PassInputs.SceneTextures = GetSceneTextureShaderParameters(Inputs.SceneTextures);
 		PassInputs.bOverview = bVisualizeGBufferOverview;
