Fixed assert when texture can be "not written to" by any pass when local lights passes get culled by other logic (bounding boxes are not visible). In this case we keep default dummy textures.

Ran volumetric fog engine test.

diff --git a/Engine/Source/Runtime/Renderer/Private/VolumetricFog.cpp b/Engine/Source/Runtime/Renderer/Private/VolumetricFog.cpp
--- a/Engine/Source/Runtime/Renderer/Private/VolumetricFog.cpp
+++ b/Engine/Source/Runtime/Renderer/Private/VolumetricFog.cpp
@@ -485,7 +485,6 @@ void FDeferredShadingSceneRenderer::RenderLocalLightsForVolumetricFog(
 	// Now voxelise all the light we have just gathered.
 	if (LightsToInject.Num() > 0)
 	{
-		OutLocalShadowedLightScattering = GraphBuilder.CreateTexture(VolumeDesc, TEXT("VolumetricFog.LocalShadowedLightScattering"));
 		TMap<FLightSceneInfo*, FVolumetricFogLocalLightFunctionInfo>& LocalLightFunctionData = View.VolumetricFogResources.LocalLightFunctionData;
 
 		bool bClearExecuted = false;
@@ -526,6 +525,11 @@ void FDeferredShadingSceneRenderer::RenderLocalLightsForVolumetricFog(
 					PassParameters->LightFunctionAtlasTileMinMaxUvBound = LightFunctionData->AtlasTile.MinMaxUvBound;
 				}
 
+				// We need to delay the real texture creation replacing the dummy on first use after all the culling logic, 
+				// otherwise it will be reported as not written by any pass (in this case we need to keep the default dummy).
+				const bool bHasTextureBeenCreated = bClearExecuted == true;
+				OutLocalShadowedLightScattering = bHasTextureBeenCreated ? OutLocalShadowedLightScattering : GraphBuilder.CreateTexture(VolumeDesc, TEXT("VolumetricFog.LocalShadowedLightScattering"));
+
 				PassParameters->RenderTargets[0] = FRenderTargetBinding(OutLocalShadowedLightScattering, bClearExecuted ? ERenderTargetLoadAction::ELoad : ERenderTargetLoadAction::EClear);
 				bClearExecuted = true;
 				PassParameters->VirtualShadowMapSamplingParameters = VirtualShadowMapArray.GetSamplingParameters(GraphBuilder);
