Fix strata complex tile during local light pass.

diff --git a/Engine/Source/Runtime/Renderer/Private/LightRendering.cpp b/Engine/Source/Runtime/Renderer/Private/LightRendering.cpp
--- a/Engine/Source/Runtime/Renderer/Private/LightRendering.cpp
+++ b/Engine/Source/Runtime/Renderer/Private/LightRendering.cpp
@@ -1700,7 +1700,7 @@ static uint32 InternalSetBoundingGeometryDepthState(FGraphicsPipelineStateInitia
 		{
 		case EStrataTileMaterialType::ESimple : StencilRef = Strata::StencilBit_Fast;    GraphicsPSOInit.DepthStencilState = TStaticDepthStencilState<false, CompareFunction, true, CF_Equal, SO_Keep, SO_Keep, SO_Keep, true, CF_Equal, SO_Keep, SO_Keep, SO_Keep, Strata::StencilBit_Fast, 0x0>::GetRHI(); break;
 		case EStrataTileMaterialType::ESingle : StencilRef = Strata::StencilBit_Single;  GraphicsPSOInit.DepthStencilState = TStaticDepthStencilState<false, CompareFunction, true, CF_Equal, SO_Keep, SO_Keep, SO_Keep, true, CF_Equal, SO_Keep, SO_Keep, SO_Keep, Strata::StencilBit_Single, 0x0>::GetRHI(); break;
-		case EStrataTileMaterialType::EComplex: StencilRef = Strata::StencilBit_Complex; GraphicsPSOInit.DepthStencilState = TStaticDepthStencilState<false, CompareFunction, true, CF_Equal, SO_Keep, SO_Keep, SO_Keep, true, CF_Equal, SO_Keep, SO_Keep, SO_Keep, Strata::StencilBit_Fast | Strata::StencilBit_Single, 0x0>::GetRHI(); break;
+		case EStrataTileMaterialType::EComplex: StencilRef = Strata::StencilBit_Complex; GraphicsPSOInit.DepthStencilState = TStaticDepthStencilState<false, CompareFunction, true, CF_Equal, SO_Keep, SO_Keep, SO_Keep, true, CF_Equal, SO_Keep, SO_Keep, SO_Keep, Strata::StencilBit_Complex, 0x0>::GetRHI(); break;
 		default: check(false);
 		}
 	}
diff --git a/Engine/Source/Runtime/Renderer/Private/PostProcess/SceneRenderTargets.h b/Engine/Source/Runtime/Renderer/Private/PostProcess/SceneRenderTargets.h
--- a/Engine/Source/Runtime/Renderer/Private/PostProcess/SceneRenderTargets.h
+++ b/Engine/Source/Runtime/Renderer/Private/PostProcess/SceneRenderTargets.h
@@ -24,9 +24,10 @@
 #define STENCIL_TEMPORAL_RESPONSIVE_AA_BIT_ID 3
 #define STENCIL_LIGHTING_CHANNELS_BIT_ID	4
 #define STENCIL_RECEIVE_DECAL_BIT_ID		7
-// Used only during the lighting pass - alias/reuse decal bit
-#define STENCIL_STRATA_FASTPATH				7 
-#define STENCIL_STRATA_SINGLEPATH			1
+// Used only during the lighting pass - alias/reuse light channels (which copied from stencil to a texture prior to lighting pass)
+#define STENCIL_STRATA_FASTPATH				4 
+#define STENCIL_STRATA_SINGLEPATH			5
+#define STENCIL_STRATA_COMPLEX				6
 
 // Outputs a compile-time constant stencil's bit mask ready to be used
 // in TStaticDepthStencilState<> template parameter. It also takes care
diff --git a/Engine/Source/Runtime/Renderer/Private/Strata/Strata.cpp b/Engine/Source/Runtime/Renderer/Private/Strata/Strata.cpp
--- a/Engine/Source/Runtime/Renderer/Private/Strata/Strata.cpp
+++ b/Engine/Source/Runtime/Renderer/Private/Strata/Strata.cpp
@@ -622,12 +622,14 @@ static void AddStrataInternalClassificationTilePass(
 			}
 			else
 			{
-				check(TileMaterialType == EStrataTileMaterialType::ESimple || TileMaterialType == EStrataTileMaterialType::ESingle);
+				check(TileMaterialType != EStrataTileMaterialType::ECount);
 
 				// No blending and no pixel shader required. Stencil will be writen to.
 				GraphicsPSOInit.BoundShaderState.PixelShaderRHI = nullptr;
 				GraphicsPSOInit.BlendState = TStaticBlendState<>::GetRHI();
-				if (TileMaterialType == EStrataTileMaterialType::ESimple)
+				switch (TileMaterialType)
+				{
+				case EStrataTileMaterialType::ESimple:
 				{
 					GraphicsPSOInit.DepthStencilState = TStaticDepthStencilState<
 						false, CF_Always,
@@ -636,7 +638,8 @@ static void AddStrataInternalClassificationTilePass(
 						0xFF, StencilBit_Fast>::GetRHI();
 					StencilRef = StencilBit_Fast;
 				}
-				else // EStrataTileMaterialType::ESingle
+				break;
+				case EStrataTileMaterialType::ESingle:
 				{
 					GraphicsPSOInit.DepthStencilState = TStaticDepthStencilState<
 						false, CF_Always,
@@ -645,6 +648,18 @@ static void AddStrataInternalClassificationTilePass(
 						0xFF, StencilBit_Single>::GetRHI();
 					StencilRef = StencilBit_Single;
 				}
+				break;
+				case EStrataTileMaterialType::EComplex:
+				{
+					GraphicsPSOInit.DepthStencilState = TStaticDepthStencilState<
+						false, CF_Always,
+						true, CF_Always, SO_Keep, SO_Keep, SO_Replace,
+						false, CF_Always, SO_Keep, SO_Keep, SO_Keep,
+						0xFF, StencilBit_Complex>::GetRHI();
+					StencilRef = StencilBit_Complex;
+				}
+				break;
+				}
 			}
 			GraphicsPSOInit.BoundShaderState.VertexDeclarationRHI = GFilterVertexDeclaration.VertexDeclarationRHI;
 			GraphicsPSOInit.BoundShaderState.VertexShaderRHI = VertexShader.GetVertexShader();
@@ -662,15 +677,6 @@ static void AddStrataInternalClassificationTilePass(
 		});
 }
 
-void AddStrataStencilPass(
-	FRDGBuilder& GraphBuilder,
-	const FViewInfo& View,
-	const FMinimalSceneTextures& SceneTextures)
-{
-	AddStrataInternalClassificationTilePass(GraphBuilder, View, &SceneTextures.Depth.Target, nullptr, EStrataTileMaterialType::ESimple);
-	AddStrataInternalClassificationTilePass(GraphBuilder, View, &SceneTextures.Depth.Target, nullptr, EStrataTileMaterialType::ESingle);
-}
-
 void AddStrataStencilPass(
 	FRDGBuilder& GraphBuilder,
 	const TArray<FViewInfo>& Views,
@@ -681,6 +687,7 @@ void AddStrataStencilPass(
 		const FViewInfo& View = Views[i];
 		AddStrataInternalClassificationTilePass(GraphBuilder, View, &SceneTextures.Depth.Target, nullptr, EStrataTileMaterialType::ESimple);
 		AddStrataInternalClassificationTilePass(GraphBuilder, View, &SceneTextures.Depth.Target, nullptr, EStrataTileMaterialType::ESingle);
+		AddStrataInternalClassificationTilePass(GraphBuilder, View, &SceneTextures.Depth.Target, nullptr, EStrataTileMaterialType::EComplex);
 	}
 }
 
diff --git a/Engine/Source/Runtime/Renderer/Private/Strata/Strata.h b/Engine/Source/Runtime/Renderer/Private/Strata/Strata.h
--- a/Engine/Source/Runtime/Renderer/Private/Strata/Strata.h
+++ b/Engine/Source/Runtime/Renderer/Private/Strata/Strata.h
@@ -89,9 +89,9 @@ struct FStrataSceneData
 
 namespace Strata
 {
-constexpr uint32 StencilBit_Complex= 0x00; // No stencil bit
-constexpr uint32 StencilBit_Single = 0x02; // In sync with SceneRenderTargets.h - GET_STENCIL_BIT_MASK(STENCIL_STRATA_SINGLEPATH)
-constexpr uint32 StencilBit_Fast   = 0x80; // In sync with SceneRenderTargets.h - GET_STENCIL_BIT_MASK(STENCIL_STRATA_FASTPATH)
+constexpr uint32 StencilBit_Fast   = 0x08; // In sync with SceneRenderTargets.h - GET_STENCIL_BIT_MASK(STENCIL_STRATA_FASTPATH)
+constexpr uint32 StencilBit_Single = 0x10; // In sync with SceneRenderTargets.h - GET_STENCIL_BIT_MASK(STENCIL_STRATA_SINGLEPATH)
+constexpr uint32 StencilBit_Complex= 0x20; // In sync with SceneRenderTargets.h - GET_STENCIL_BIT_MASK(STENCIL_STRATA_COMPLEX)
 
 bool IsStrataEnabled();
 
@@ -108,7 +108,6 @@ TRDGUniformBufferRef<FStrataGlobalUniformParameters> BindStrataGlobalUniformPara
 
 void AddStrataMaterialClassificationPass(FRDGBuilder& GraphBuilder, const FMinimalSceneTextures& SceneTextures, const TArray<FViewInfo>& Views);
 
-void AddStrataStencilPass(FRDGBuilder& GraphBuilder, const FViewInfo& View, const FMinimalSceneTextures& SceneTextures);
 void AddStrataStencilPass(FRDGBuilder& GraphBuilder, const TArray<FViewInfo>& Views, const FMinimalSceneTextures& SceneTextures);
 
 bool ShouldRenderStrataDebugPasses(const FViewInfo& View);
