Generate Raytracing light culling structure in TranslatedWorldSpace.
- Calculate light position in translated world space.
- Updated FCulledLightList::Create(...) to calculate volume cell position using TranslatedWorldSpace.

diff --git a/Engine/Shaders/Private/RayTracing/GenerateCulledLightListCS.usf b/Engine/Shaders/Private/RayTracing/GenerateCulledLightListCS.usf
--- a/Engine/Shaders/Private/RayTracing/GenerateCulledLightListCS.usf
+++ b/Engine/Shaders/Private/RayTracing/GenerateCulledLightListCS.usf
@@ -4,7 +4,7 @@
 #define CACHED_LIGHT_COUNT  16
 
 // inputs
-float3                    WorldPos;
+float3                    TranslatedWorldPos;
 uint                      NumLightsToUse;
 uint                      CellCount;
 float                     CellScale;
@@ -74,7 +74,7 @@ void GenerateCulledLightListCS(
 
 	float3 Min, Max;
 	ComputeCoords(Min, Max, GroupId);
-	float3 CellCenter = (Min + Max) * 0.5 + WorldPos;
+	float3 CellCenter = (Min + Max) * 0.5 + TranslatedWorldPos;
 	float3 CellExtents = (Max - Min) * 0.5;
 
 	uint LightIndexBufferOffset = (GroupId.x + CellCount * (GroupId.y + (GroupId.z * CellCount))) * NumLightsToUse;
diff --git a/Engine/Shaders/Private/RayTracing/RayTracingLightCullingCommon.ush b/Engine/Shaders/Private/RayTracing/RayTracingLightCullingCommon.ush
--- a/Engine/Shaders/Private/RayTracing/RayTracingLightCullingCommon.ush
+++ b/Engine/Shaders/Private/RayTracing/RayTracingLightCullingCommon.ush
@@ -6,9 +6,6 @@
 
 #pragma once
 
-#define CULL_VOLUME_STRUCTURED_BUFFER 1
-
-
 // *************** LIGHT CULLING *********************************
 // Use view-centered, world-aligned volume to store references to lights impacting cells
 // Volume cells store light indices directly if space permits, otherwise indirect into
@@ -23,12 +20,12 @@ struct FCulledLightList
 	static const uint PackedLightMask = (1u << PackedLightBits) - 1;
 	static const uint PackedLightsMax = 11;
 
-	static FCulledLightList Create(float3 WorldPos)
+	static FCulledLightList Create(float3 TranslatedWorldPosition)
 	{
 		FCulledLightList Result = (FCulledLightList)0;
 
 		// volume is centered on the viewer
-		float3 Position = WorldPos - LWCHackToFloat(PrimaryView.WorldViewOrigin);
+		float3 Position = TranslatedWorldPosition - PrimaryView.TranslatedWorldCameraOrigin;
 
 		const float Scale = 100.0f;
 		const int Dim = RaytracingLightsDataPacked.CellCount;
diff --git a/Engine/Shaders/Private/RayTracing/RayTracingLightingCommon.ush b/Engine/Shaders/Private/RayTracing/RayTracingLightingCommon.ush
--- a/Engine/Shaders/Private/RayTracing/RayTracingLightingCommon.ush
+++ b/Engine/Shaders/Private/RayTracing/RayTracingLightingCommon.ush
@@ -400,7 +400,7 @@ float3 ComputeDirectLighting(
 	Payload.SetRadiance(float3(0, 0, 0));
 	Payload.SetIndirectIrradiance(ViewDirection);
 
-	FCulledLightList CullData = FCulledLightList::Create(WorldPosition);
+	FCulledLightList CullData = FCulledLightList::Create(WorldPosition + LWCHackToFloat(PrimaryView.PreViewTranslation)); // RT_LWC_TODO
 	const uint LightCount = CullData.NumLights();
 
 	for (uint Index = 0; WaveActiveAnyTrue(Index < LightCount); Index++)
@@ -416,7 +416,7 @@ float3 ComputeDirectLighting(
 
 		float3 ShadowRayDirection;
 		// ToLight should not be normalized because its length is used to compute the shadow ray TMax 
-		float3 ToLight = LightData.TranslatedWorldPosition - WorldPosition; // LWC_TODO
+		float3 ToLight = LightData.TranslatedWorldPosition - WorldPosition; // RT_LWC_TODO
 		float LightMask = 1.0;
 
 		if (LightType == LIGHT_TYPE_DIRECTIONAL)
diff --git a/Engine/Source/Runtime/Renderer/Private/RayTracing/RayTracingLighting.cpp b/Engine/Source/Runtime/Renderer/Private/RayTracing/RayTracingLighting.cpp
--- a/Engine/Source/Runtime/Renderer/Private/RayTracing/RayTracingLighting.cpp
+++ b/Engine/Source/Runtime/Renderer/Private/RayTracing/RayTracingLighting.cpp
@@ -50,7 +50,7 @@ class FSetupRayTracingLightCullData : public FGlobalShader
 
 	BEGIN_SHADER_PARAMETER_STRUCT(FParameters, )
 		SHADER_PARAMETER_SRV(StructuredBuffer<float4>, RankedLights)
-		SHADER_PARAMETER(FVector3f, WorldPos)
+		SHADER_PARAMETER(FVector3f, TranslatedWorldPos)
 		SHADER_PARAMETER(uint32, NumLightsToUse)
 		SHADER_PARAMETER(uint32, CellCount)
 		SHADER_PARAMETER(float, CellScale)
@@ -114,15 +114,15 @@ static void CreateRaytracingLightCullingStructure(
 	// setup light vector array sorted by rank
 	for (int32 LightIndex = 0; LightIndex < NumLightsToUse; LightIndex++)
 	{
-		VectorRegister BoundingSphere = Lights[LightIndices[LightIndex]].BoundingSphereVector;
-		RankedLights.Push(
-			FVector4f(
-				VectorGetComponentImpl<0>(BoundingSphere), 
-				VectorGetComponentImpl<1>(BoundingSphere), 
-				VectorGetComponentImpl<2>(BoundingSphere), 
-				VectorGetComponentImpl<3>(BoundingSphere)
-			)
+		VectorRegister BoundingSphereRegister = Lights[LightIndices[LightIndex]].BoundingSphereVector;
+		FVector4 BoundingSphere = FVector4(
+			VectorGetComponentImpl<0>(BoundingSphereRegister),
+			VectorGetComponentImpl<1>(BoundingSphereRegister),
+			VectorGetComponentImpl<2>(BoundingSphereRegister),
+			VectorGetComponentImpl<3>(BoundingSphereRegister)
 		);
+		FVector4f TranslatedBoundingSphere = FVector4f(BoundingSphere + View.ViewMatrices.GetPreViewTranslation());
+		RankedLights.Push(TranslatedBoundingSphere);
 	}
 
 	// push null vector to prevent failure in RHICreateStructuredBuffer due to requesting a zero sized allocation
@@ -162,7 +162,7 @@ static void CreateRaytracingLightCullingStructure(
 			FSetupRayTracingLightCullData::FParameters Params;
 			Params.RankedLights = RankedLightsSRV;
 
-			Params.WorldPos = (FVector3f)View.ViewMatrices.GetViewOrigin(); // View.ViewLocation; // LWC_TODO: Precision Loss
+			Params.TranslatedWorldPos = (FVector3f)(View.ViewMatrices.GetViewOrigin() + View.ViewMatrices.GetPreViewTranslation());
 			Params.NumLightsToUse = NumLightsToUse;
 			Params.LightCullingVolume = LightCullVolumeUAV;
 			Params.LightIndices = OutLightingData.LightIndices.UAV;
