Fix local lighting with strata.

diff --git a/Engine/Source/Runtime/Renderer/Private/ShadowRendering.cpp b/Engine/Source/Runtime/Renderer/Private/ShadowRendering.cpp
--- a/Engine/Source/Runtime/Renderer/Private/ShadowRendering.cpp
+++ b/Engine/Source/Runtime/Renderer/Private/ShadowRendering.cpp
@@ -237,8 +237,8 @@ static TAutoConsoleVariable<int32> CVarHairStrandsCullPerObjectShadowCaster(
 
 DEFINE_GPU_DRAWCALL_STAT(ShadowProjection);
 
-// Use Shadow stencil mask is set to 0x7F instead of 0xFF so that that last bit can be used for strata classification without clearing the stencil bit for pre-shadow/per-object static shadow mask
-constexpr uint32 ShadowStencilMask = 0x7f;
+// Use Shadow stencil mask is set to 0x07u instead of 0xFF so that that last bit can be used for strata classification without clearing the stencil bit for pre-shadow/per-object static shadow mask
+constexpr uint32 ShadowStencilMask = 0x07u;
 
 // 0:off, 1:low, 2:med, 3:high, 4:very high, 5:max
 uint32 GetShadowQuality()
diff --git a/Engine/Source/Runtime/Renderer/Private/Strata/Strata.cpp b/Engine/Source/Runtime/Renderer/Private/Strata/Strata.cpp
--- a/Engine/Source/Runtime/Renderer/Private/Strata/Strata.cpp
+++ b/Engine/Source/Runtime/Renderer/Private/Strata/Strata.cpp
@@ -558,6 +558,12 @@ FStrataTileParameter SetTileParameters(FRDGBuilder& GraphBuilder, const FViewInf
 	return InternalSetTileParameters(&GraphBuilder, View, TileType);
 }
 
+// Add additionnaly bits for filling/clearing stencil to ensure that the 'Strata' bits are not corrupted by the stencil shadows 
+// when generating shadow mask. Withouth these 'trailing' bits, the incr./decr. operation would change/corrupt the 'Strata' bits
+constexpr uint32 StencilBit_Fast_1	  = 0x07u | StencilBit_Fast;
+constexpr uint32 StencilBit_Single_1  = 0x07u | StencilBit_Single;
+constexpr uint32 StencilBit_Complex_1 = 0x07u | StencilBit_Complex; 
+
 static void AddStrataInternalClassificationTilePass(
 	FRDGBuilder& GraphBuilder,
 	const FViewInfo& View,
@@ -635,8 +641,8 @@ static void AddStrataInternalClassificationTilePass(
 						false, CF_Always,
 						true, CF_Always, SO_Keep, SO_Keep, SO_Replace,
 						false, CF_Always, SO_Keep, SO_Keep, SO_Keep,
-						0xFF, StencilBit_Fast>::GetRHI();
-					StencilRef = StencilBit_Fast;
+						0xFF, StencilBit_Fast_1>::GetRHI();
+					StencilRef = StencilBit_Fast_1;
 				}
 				break;
 				case EStrataTileMaterialType::ESingle:
@@ -645,8 +651,8 @@ static void AddStrataInternalClassificationTilePass(
 						false, CF_Always,
 						true, CF_Always, SO_Keep, SO_Keep, SO_Replace,
 						false, CF_Always, SO_Keep, SO_Keep, SO_Keep,
-						0xFF, StencilBit_Single>::GetRHI();
-					StencilRef = StencilBit_Single;
+						0xFF, StencilBit_Single_1>::GetRHI();
+					StencilRef = StencilBit_Single_1;
 				}
 				break;
 				case EStrataTileMaterialType::EComplex:
@@ -655,8 +661,8 @@ static void AddStrataInternalClassificationTilePass(
 						false, CF_Always,
 						true, CF_Always, SO_Keep, SO_Keep, SO_Replace,
 						false, CF_Always, SO_Keep, SO_Keep, SO_Keep,
-						0xFF, StencilBit_Complex>::GetRHI();
-					StencilRef = StencilBit_Complex;
+						0xFF, StencilBit_Complex_1>::GetRHI();
+					StencilRef = StencilBit_Complex_1;
 				}
 				break;
 				}
