Fix incompatible serialization of change CVar r.VertexDeformationOutputsVelocity.
Converted it to new CVar r.Velocity.EnableVertexDeformation and added warning if old CVar is found.
#rnx

diff --git a/Engine/Source/Runtime/Engine/Classes/Engine/RendererSettings.h b/Engine/Source/Runtime/Engine/Classes/Engine/RendererSettings.h
--- a/Engine/Source/Runtime/Engine/Classes/Engine/RendererSettings.h
+++ b/Engine/Source/Runtime/Engine/Classes/Engine/RendererSettings.h
@@ -718,7 +718,7 @@ class ENGINE_API URendererSettings : public UDeveloperSettings
 
 	UPROPERTY(config, EditAnywhere, Category=Optimizations, meta=(
 		DisplayName="Output velocities due to vertex deformation",
-		ConsoleVariable="r.VertexDeformationOutputsVelocity",
+		ConsoleVariable="r.Velocity.EnableVertexDeformation",
 		ToolTip="Enables materials with World Position Offset and/or World Displacement to output velocities during the velocity pass even when the actor has not moved. \nIf the VelocityPass is set to 'Write after base pass' this can incur a performance cost due to additional draw calls. \nThat performance cost is higher if many objects are using World Position Offset. A forest of trees for example." ))
 	TEnumAsByte<EVertexDeformationOutputsVelocity::Type> VertexDeformationOutputsVelocity;
 
diff --git a/Engine/Source/Runtime/Engine/Private/PrimitiveSceneProxy.cpp b/Engine/Source/Runtime/Engine/Private/PrimitiveSceneProxy.cpp
--- a/Engine/Source/Runtime/Engine/Private/PrimitiveSceneProxy.cpp
+++ b/Engine/Source/Runtime/Engine/Private/PrimitiveSceneProxy.cpp
@@ -38,8 +38,8 @@ static TAutoConsoleVariable<int32> CVarCacheWPOPrimitives(
 	ECVF_RenderThreadSafe | ECVF_Scalability
 	);
 
-static TAutoConsoleVariable<int32> CVarVertexDeformationOutputsVelocity(
-	TEXT("r.VertexDeformationOutputsVelocity"),
+static TAutoConsoleVariable<int32> CVarVelocityEnableVertexDeformation(
+	TEXT("r.Velocity.EnableVertexDeformation"),
 	2,
 	TEXT("Enables materials with World Position Offset and/or World Displacement to output velocities during velocity pass even when the actor has not moved. \n")
 	TEXT("0=Off, 1=On, 2=Auto(Default). \n")
@@ -110,7 +110,7 @@ static bool VertexDeformationOutputsVelocity()
 {
 	static const auto CVarVelocityOutputPass = IConsoleManager::Get().FindConsoleVariable(TEXT("r.VelocityOutputPass"));
 	const bool bVertexDeformationOutputsVelocityDefault = CVarVelocityOutputPass && CVarVelocityOutputPass->GetInt() != 2;
-	const int32 VertexDeformationOutputsVelocity = CVarVertexDeformationOutputsVelocity.GetValueOnAnyThread();
+	const int32 VertexDeformationOutputsVelocity = CVarVelocityEnableVertexDeformation.GetValueOnAnyThread();
 	return VertexDeformationOutputsVelocity == 1 || (VertexDeformationOutputsVelocity == 2 && bVertexDeformationOutputsVelocityDefault);
 }
 
diff --git a/Engine/Source/Runtime/Renderer/Private/SceneRendering.cpp b/Engine/Source/Runtime/Renderer/Private/SceneRendering.cpp
--- a/Engine/Source/Runtime/Renderer/Private/SceneRendering.cpp
+++ b/Engine/Source/Runtime/Renderer/Private/SceneRendering.cpp
@@ -4213,7 +4213,7 @@ FRendererModule::FRendererModule()
 	static auto EarlyZPassVar = IConsoleManager::Get().FindConsoleVariable(TEXT("r.EarlyZPass"));
 	EarlyZPassVar->SetOnChangedCallback(FConsoleVariableDelegate::CreateStatic(&OnChangeCVarRequiringRecreateRenderState));
 
-	static auto CVarVertexDeformationOutputsVelocity = IConsoleManager::Get().FindConsoleVariable(TEXT("r.VertexDeformationOutputsVelocity"));
+	static auto CVarVertexDeformationOutputsVelocity = IConsoleManager::Get().FindConsoleVariable(TEXT("r.Velocity.EnableVertexDeformation"));
 	CVarVertexDeformationOutputsVelocity->SetOnChangedCallback(FConsoleVariableDelegate::CreateStatic(&OnChangeCVarRequiringRecreateRenderState));
 
 #if !(UE_BUILD_SHIPPING || UE_BUILD_TEST)
diff --git a/Engine/Source/Runtime/Renderer/Private/VelocityRendering.cpp b/Engine/Source/Runtime/Renderer/Private/VelocityRendering.cpp
--- a/Engine/Source/Runtime/Renderer/Private/VelocityRendering.cpp
+++ b/Engine/Source/Runtime/Renderer/Private/VelocityRendering.cpp
@@ -40,6 +40,12 @@ static TAutoConsoleVariable<int32> CVarBasePassOutputsVelocity(
 	TEXT("Deprecated CVar. Use r.VelocityOutputPass instead.\n"),
 	ECVF_ReadOnly | ECVF_RenderThreadSafe);
 
+static TAutoConsoleVariable<int32> CVarVertexDeformationOutputsVelocity(
+	TEXT("r.VertexDeformationOutputsVelocity"),
+	-1,
+	TEXT("Deprecated CVar. Use r.Velocity.EnableVertexDeformation instead.\n"),
+	ECVF_ReadOnly | ECVF_RenderThreadSafe);
+
 static TAutoConsoleVariable<int32> CVarParallelVelocity(
 	TEXT("r.ParallelVelocity"),
 	1,  
@@ -61,10 +67,19 @@ inline void ValidateVelocityCVars()
 	static bool bHasValidatedCVars = false;
 	if (!bHasValidatedCVars)
 	{
-		const int32 Value = CVarBasePassOutputsVelocity.GetValueOnAnyThread();
-		if (Value != -1)
 		{
-			UE_LOG(LogRenderer, Warning, TEXT("Deprectaed CVar r.BasePassOutputsVelocity is set to %d. Use r.VelocityOutputPass instead."), Value);
+			const int32 Value = CVarBasePassOutputsVelocity.GetValueOnAnyThread();
+			if (Value != -1)
+			{
+				UE_LOG(LogRenderer, Warning, TEXT("Deprectaed CVar r.BasePassOutputsVelocity is set to %d. Remove and use r.VelocityOutputPass instead."), Value);
+			}
+		}
+		{
+			const int32 Value = CVarVertexDeformationOutputsVelocity.GetValueOnAnyThread();
+			if (Value != -1)
+			{
+				UE_LOG(LogRenderer, Warning, TEXT("Deprectaed CVar r.VertexDeformationOutputsVelocity is set to %d. Remove and use r.Velocity.EnableVertexDeformation instead."), Value);
+			}
 		}
 		bHasValidatedCVars = true;
 	}
