Fixes alpha channel through post processing when translucent object is in the scene

diff --git a/Engine/Shaders/Private/ComposeSeparateTranslucency.usf b/Engine/Shaders/Private/ComposeSeparateTranslucency.usf
--- a/Engine/Shaders/Private/ComposeSeparateTranslucency.usf
+++ b/Engine/Shaders/Private/ComposeSeparateTranslucency.usf
@@ -81,5 +81,7 @@ void MainPS(
 
 	// Final composition
 	OutColor.rgb = SceneColorSample.rgb * SeparateTranslucencySample.a * SeparateModulationSample.rgb + SeparateTranslucencySample.rgb;
-	OutColor.a   = SceneColorSample.a   * SeparateTranslucencySample.a;
+
+	float GreyScaleModulateColorBackgroundVisibility = dot(SeparateModulationSample.rgb, float3(1.0f / 3.0f, 1.0f / 3.0f, 1.0f / 3.0f));
+	OutColor.a = SceneColorSample.a * SeparateTranslucencySample.a * GreyScaleModulateColorBackgroundVisibility;
 }
diff --git a/Engine/Source/Runtime/Renderer/Private/PostProcess/PostProcessing.cpp b/Engine/Source/Runtime/Renderer/Private/PostProcess/PostProcessing.cpp
--- a/Engine/Source/Runtime/Renderer/Private/PostProcess/PostProcessing.cpp
+++ b/Engine/Source/Runtime/Renderer/Private/PostProcess/PostProcessing.cpp
@@ -217,6 +217,10 @@ void AddPostProcessingPasses(
 	const FTranslucencyPassResources& PostDOFTranslucencyResources = Inputs.TranslucencyViewResourcesMap.Get(ETranslucencyPass::TPT_TranslucencyAfterDOF);
 	const FTranslucencyPassResources& PostMotionBlurTranslucencyResources = Inputs.TranslucencyViewResourcesMap.Get(ETranslucencyPass::TPT_TranslucencyAfterMotionBlur);
 
+	// Whether should process the alpha channel of the scene color.
+	const bool bProcessSceneColorAlpha = IsPostProcessingWithAlphaChannelSupported();
+	const EPixelFormat SceneColorFormat = bProcessSceneColorAlpha ? PF_FloatRGBA : PF_FloatR11G11B10;
+
 	// Scene color is updated incrementally through the post process pipeline.
 	FScreenPassTexture SceneColor((*Inputs.SceneTextures)->SceneColorTexture, PrimaryViewRect);
 
@@ -536,6 +540,7 @@ void AddPostProcessingPasses(
 				TranslucencyComposition.SceneColor = SceneColor;
 				TranslucencyComposition.SceneDepth = SceneDepth;
 				TranslucencyComposition.OutputViewport = FScreenPassTextureViewport(SceneColor);
+				TranslucencyComposition.OutputPixelFormat = SceneColorFormat;
 
 				SceneColor = TranslucencyComposition.AddPass(
 					GraphBuilder, View, PostDOFTranslucencyResources);
@@ -677,6 +682,7 @@ void AddPostProcessingPasses(
 			TranslucencyComposition.Operation = FTranslucencyComposition::EOperation::ComposeToNewSceneColor;
 			TranslucencyComposition.SceneColor = SceneColor;
 			TranslucencyComposition.OutputViewport = FScreenPassTextureViewport(SceneColor);
+			TranslucencyComposition.OutputPixelFormat = SceneColorFormat;
 
 			SceneColor = TranslucencyComposition.AddPass(
 				GraphBuilder, View, PostMotionBlurTranslucencyResources);
@@ -910,7 +916,7 @@ void AddPostProcessingPasses(
 				PassInputs.EyeAdaptationParameters = &EyeAdaptationParameters;
 				PassInputs.EyeAdaptationTexture = EyeAdaptationTexture;
 				PassInputs.ColorGradingTexture = ColorGradingTexture;
-				PassInputs.bWriteAlphaChannel = AntiAliasingMethod == AAM_FXAA || IsPostProcessingWithAlphaChannelSupported();
+				PassInputs.bWriteAlphaChannel = AntiAliasingMethod == AAM_FXAA || bProcessSceneColorAlpha;
 				PassInputs.bOutputInHDR = bTonemapOutputInHDR;
 
 				SceneColor = AddTonemapPass(GraphBuilder, View, PassInputs);
@@ -987,6 +993,7 @@ void AddPostProcessingPasses(
 			TranslucencyComposition.Operation = FTranslucencyComposition::EOperation::ComposeToNewSceneColor;
 			TranslucencyComposition.SceneColor = SceneColor;
 			TranslucencyComposition.OutputViewport = FScreenPassTextureViewport(SceneColor);
+			TranslucencyComposition.OutputPixelFormat = SceneColorFormat;
 
 			if (PostDOFTranslucencyResources.IsValid())
 			{
diff --git a/Engine/Source/Runtime/Renderer/Private/TranslucentRendering.cpp b/Engine/Source/Runtime/Renderer/Private/TranslucentRendering.cpp
--- a/Engine/Source/Runtime/Renderer/Private/TranslucentRendering.cpp
+++ b/Engine/Source/Runtime/Renderer/Private/TranslucentRendering.cpp
@@ -632,7 +632,7 @@ FScreenPassTexture FTranslucencyComposition::AddPass(
 
 		FRDGTextureDesc OutputDesc = FRDGTextureDesc::Create2D(
 			OutputViewport.Extent,
-			PF_FloatRGB,
+			OutputPixelFormat != PF_Unknown ? OutputPixelFormat : SceneColor.Texture->Desc.Format,
 			FClearValueBinding::Black,
 			TexCreate_RenderTargetable | TexCreate_ShaderResource);
 
diff --git a/Engine/Source/Runtime/Renderer/Private/TranslucentRendering.h b/Engine/Source/Runtime/Renderer/Private/TranslucentRendering.h
--- a/Engine/Source/Runtime/Renderer/Private/TranslucentRendering.h
+++ b/Engine/Source/Runtime/Renderer/Private/TranslucentRendering.h
@@ -161,6 +161,7 @@ struct FTranslucencyComposition
 	FScreenPassTexture SceneDepth;
 
 	FScreenPassTextureViewport OutputViewport;
+	EPixelFormat OutputPixelFormat = PF_Unknown;
 
 	FScreenPassTexture AddPass(
 		FRDGBuilder& GraphBuilder,
