Fix assert in GlobalDistanceField.cpp when View.ViewState == nullptr.

diff --git a/Engine/Source/Runtime/Renderer/Private/GlobalDistanceField.cpp b/Engine/Source/Runtime/Renderer/Private/GlobalDistanceField.cpp
--- a/Engine/Source/Runtime/Renderer/Private/GlobalDistanceField.cpp
+++ b/Engine/Source/Runtime/Renderer/Private/GlobalDistanceField.cpp
@@ -402,16 +402,13 @@ void FGlobalDistanceFieldInfo::UpdateParameterData(float MaxOcclusionDistance, b
 
 		if (PageTableCombinedTexture)
 		{
-			checkf(GAOGlobalDistanceFieldCacheMostlyStaticSeparately, TEXT("PageTableCombinedTexture should only be allocated when caching mostly static objects separately."));
+			ensureMsgf(GAOGlobalDistanceFieldCacheMostlyStaticSeparately, TEXT("PageTableCombinedTexture should only be allocated when caching mostly static objects separately."));
 			ParameterData.PageTableTexture = PageTableCombinedTexture->GetRenderTargetItem().ShaderResourceTexture;
 		}
-		else
+		else if (PageTableLayerTextures[GDF_Full])
 		{
-			checkf(!GAOGlobalDistanceFieldCacheMostlyStaticSeparately, TEXT("PageTableCombinedTexture should be allocated when caching mostly static objects separately."));
-			if (PageTableLayerTextures[GDF_Full])
-			{
-				ParameterData.PageTableTexture = PageTableLayerTextures[GDF_Full]->GetRenderTargetItem().ShaderResourceTexture;
-			}
+			ensureMsgf(!GAOGlobalDistanceFieldCacheMostlyStaticSeparately, TEXT("PageTableCombinedTexture should be allocated when caching mostly static objects separately."));
+			ParameterData.PageTableTexture = PageTableLayerTextures[GDF_Full]->GetRenderTargetItem().ShaderResourceTexture;
 		}
 
 		FIntVector MipTextureResolution(1, 1, 1);
