Adds split screen and VR support to TSR

diff --git a/Engine/Shaders/Private/TemporalSuperResolution/TSRCommon.ush b/Engine/Shaders/Private/TemporalSuperResolution/TSRCommon.ush
--- a/Engine/Shaders/Private/TemporalSuperResolution/TSRCommon.ush
+++ b/Engine/Shaders/Private/TemporalSuperResolution/TSRCommon.ush
@@ -62,7 +62,7 @@
 	// COMPRESS_PREV_USE_COUNT compress the use count in 8bit, so need to sum of a 2x2 to not overflow these 8bits
 	#define PREV_USE_COUNT_QUANTIZATION 63
 #else
-	#define PREV_USE_COUNT_QUANTIZATION 256
+	#define PREV_USE_COUNT_QUANTIZATION 255
 #endif
 
 #define PARALLAX_REJECTION_MASK_THRESHOLD 0.5
@@ -1057,8 +1057,10 @@ void ComputeCompressedUseCountPixelCoordinates(uint2 PixelPos, out uint2 Compres
 	// Overlap 2x2 consecutive blocks of 8x8 pixel to reduce atomic contention in the scattering pass
 	const uint kTileSize = 8;
 
-	CompressedUseCountPixelPos = (PixelPos % kTileSize) | ((PixelPos >> 1) & ~(kTileSize - 1));
-	Shift = ((PixelPos.x / kTileSize) % 2) | (((PixelPos.y / kTileSize) % 2) << 1);
+	uint2 PixelPosOffset = PixelPos - InputInfo_ViewportMin;
+
+	CompressedUseCountPixelPos = ((PixelPosOffset % kTileSize) | ((PixelPosOffset >> 1) & ~(kTileSize - 1))) + InputInfo_ViewportMin;
+	Shift = ((PixelPosOffset.x / kTileSize) % 2) | (((PixelPosOffset.y / kTileSize) % 2) << 1);
 }
 	
 #endif
diff --git a/Engine/Shaders/Private/TemporalSuperResolution/TSRCompareTranslucency.usf b/Engine/Shaders/Private/TemporalSuperResolution/TSRCompareTranslucency.usf
--- a/Engine/Shaders/Private/TemporalSuperResolution/TSRCompareTranslucency.usf
+++ b/Engine/Shaders/Private/TemporalSuperResolution/TSRCompareTranslucency.usf
@@ -105,9 +105,11 @@ void MainCS(
 {
 	float4 Debug = 0;
 	
-	tsr_short2 InputPixelPos = tsr_short2(
+	tsr_short2 DispatchThreadId = tsr_short2(
 		ZOrder2D(GroupThreadIndex, uint(log2(float(TILE_SIZE)))) +
 		GroupId * uint2(TILE_SIZE, TILE_SIZE));
+		
+	tsr_short2 InputPixelPos = tsr_short2(InputInfo_ViewportMin) + DispatchThreadId;
 
 	float2 ScreenPos = ApplyScreenTransform(float2(InputPixelPos), InputPixelPosToScreenPos);
 	float2 EncodedVelocity = DilatedVelocityTexture[InputPixelPos];
diff --git a/Engine/Shaders/Private/TemporalSuperResolution/TSRFilterFrequencies.usf b/Engine/Shaders/Private/TemporalSuperResolution/TSRFilterFrequencies.usf
--- a/Engine/Shaders/Private/TemporalSuperResolution/TSRFilterFrequencies.usf
+++ b/Engine/Shaders/Private/TemporalSuperResolution/TSRFilterFrequencies.usf
@@ -70,10 +70,11 @@ void MainCS(
 
 	float4 Debug = 0.0;
 	
-	tsr_short2 InputPixelPos = (
-		tsr_short2(InputPixelPosMin) +
+	tsr_short2 DispatchThreadId = (
 		tsr_short2(GroupId) * tsr_short2(TILE_SIZE, TILE_SIZE) +
 		Map16x16Tile2x2Lane(GroupThreadIndex));
+		
+	tsr_short2 InputPixelPos = tsr_short2(InputPixelPosMin) + DispatchThreadId;
 
 	// Fetch center of the kernel
 	tsr_halfC CenterInput;
diff --git a/Engine/Source/Runtime/Renderer/Private/PostProcess/TemporalSuperResolution.cpp b/Engine/Source/Runtime/Renderer/Private/PostProcess/TemporalSuperResolution.cpp
--- a/Engine/Source/Runtime/Renderer/Private/PostProcess/TemporalSuperResolution.cpp
+++ b/Engine/Source/Runtime/Renderer/Private/PostProcess/TemporalSuperResolution.cpp
@@ -621,7 +621,7 @@ ITemporalUpscaler::FOutputs AddTemporalSuperResolutionPasses(
 	FIntRect InputRect = View.ViewRect;
 
 	FIntPoint LowFrequencyExtent = InputExtent;
-	FIntRect LowFrequencyRect = FIntRect(FIntPoint(0, 0), InputRect.Size());
+	FIntRect LowFrequencyRect = InputRect;
 
 	if (bHalfResLowFrequency)
 	{
@@ -705,7 +705,8 @@ ITemporalUpscaler::FOutputs AddTemporalSuperResolutionPasses(
 			InputExtent, InputRect));
 		CommonParameters.InputPixelPosMin = CommonParameters.InputInfo.ViewportMin;
 		CommonParameters.InputPixelPosMax = CommonParameters.InputInfo.ViewportMax - 1;
-		CommonParameters.InputPixelPosToScreenPos = (FScreenTransform::Identity + 0.5) * CommonParameters.InputInfo.ViewportSizeInverse * FScreenTransform::ViewportUVToScreenPos;
+		CommonParameters.InputPixelPosToScreenPos = (FScreenTransform::Identity + 0.5f) * FScreenTransform::ChangeTextureBasisFromTo(FScreenPassTextureViewport(
+			InputExtent, InputRect), FScreenTransform::ETextureBasis::TexelPosition, FScreenTransform::ETextureBasis::ScreenPosition);
 		CommonParameters.ScreenVelocityToInputPixelVelocity = (FScreenTransform::Identity / CommonParameters.InputPixelPosToScreenPos).Scale;
 		CommonParameters.InputPixelVelocityToScreenVelocity = CommonParameters.InputPixelPosToScreenPos.Scale.GetAbs();
 
@@ -1297,7 +1298,7 @@ ITemporalUpscaler::FOutputs AddTemporalSuperResolutionPasses(
 	if (PostFilter != ERejectionPostFilter::Disabled)
 	{
 		bool bOutputHalfRes = PostFilter == ERejectionPostFilter::PreRejectionDownsample;
-		FIntRect Rect = bOutputHalfRes ? LowFrequencyRect : RejectionRect;
+		FIntRect Rect = bOutputHalfRes ? FIntRect(FIntPoint(0, 0), LowFrequencyRect.Size()) : RejectionRect;
 
 		FRDGTextureRef FilteredHistoryRejectionTexture;
 		{
