Fix typo in stat label (Tranfer -> Transfer)

diff --git a/Engine/Source/Runtime/Renderer/Private/SceneRendering.cpp b/Engine/Source/Runtime/Renderer/Private/SceneRendering.cpp
--- a/Engine/Source/Runtime/Renderer/Private/SceneRendering.cpp
+++ b/Engine/Source/Runtime/Renderer/Private/SceneRendering.cpp
@@ -2977,7 +2977,7 @@ FRHIGPUMask FSceneRenderer::ComputeGPUMasks(FRHICommandListImmediate& RHICmdList
 #endif // WITH_MGPU
 
 #if WITH_MGPU
-DECLARE_GPU_STAT_NAMED(CrossGPUTransfers, TEXT("Cross GPU Tranfer"));
+DECLARE_GPU_STAT_NAMED(CrossGPUTransfers, TEXT("Cross GPU Transfer"));
 #endif // WITH_MGPU
 
 void FSceneRenderer::DoCrossGPUTransfers(FRDGBuilder& GraphBuilder, FRHIGPUMask RenderTargetGPUMask, FRDGTextureRef ViewFamilyTexture)
