Fix typos

diff --git a/Engine/Source/Runtime/Renderer/Private/PostProcess/PostProcessMaterial.cpp b/Engine/Source/Runtime/Renderer/Private/PostProcess/PostProcessMaterial.cpp
--- a/Engine/Source/Runtime/Renderer/Private/PostProcess/PostProcessMaterial.cpp
+++ b/Engine/Source/Runtime/Renderer/Private/PostProcess/PostProcessMaterial.cpp
@@ -128,7 +128,7 @@ FRHIBlendState* GetMaterialBlendState(const FMaterial* Material)
 	return BlendStates[Material->GetBlendMode()];
 }
 
-bool PostProceessStencilTest(uint32 StencilValue, uint32 StencilComp, uint32 StencilRef)
+bool PostProcessStencilTest(uint32 StencilValue, uint32 StencilComp, uint32 StencilRef)
 {
 	bool bStencilTestPassed = true;
 
@@ -619,12 +619,12 @@ FScreenPassTexture AddPostProcessMaterialPass(
 		if (!DepthStencilTexture)
 		{
 			uint32 StencilClearValue = Inputs.CustomDepthTexture ? Inputs.CustomDepthTexture->Desc.ClearValue.Value.DSValue.Stencil : 0;
-			bFailStencil &= PostProceessStencilTest(StencilClearValue, Material->GetStencilCompare(), PostProcessMaterialParameters->MobileStencilValueRef);
+			bFailStencil &= PostProcessStencilTest(StencilClearValue, Material->GetStencilCompare(), PostProcessMaterialParameters->MobileStencilValueRef);
 		}
 
 		for (const uint32& Value : View.CustomDepthStencilValues)
 		{
-			bFailStencil &= PostProceessStencilTest(Value, Material->GetStencilCompare(), PostProcessMaterialParameters->MobileStencilValueRef);
+			bFailStencil &= PostProcessStencilTest(Value, Material->GetStencilCompare(), PostProcessMaterialParameters->MobileStencilValueRef);
 
 			if (!bFailStencil)
 			{
diff --git a/Engine/Source/Runtime/Renderer/Private/SceneVisibility.cpp b/Engine/Source/Runtime/Renderer/Private/SceneVisibility.cpp
--- a/Engine/Source/Runtime/Renderer/Private/SceneVisibility.cpp
+++ b/Engine/Source/Runtime/Renderer/Private/SceneVisibility.cpp
@@ -2247,7 +2247,7 @@ struct FRelevancePacket
 	TArray<FSkyMeshBatch> SkyMeshBatches;
 	TArray<FSortedTrianglesMeshBatch> SortedTrianglesMeshBatches;
 	FDrawCommandRelevancePacket DrawCommandPacket;
-	TSet<uint32> CustumDepthStencilValues;
+	TSet<uint32> CustomDepthStencilValues;
 
 	struct FPrimitiveLODMask
 	{
@@ -2458,7 +2458,7 @@ struct FRelevancePacket
 			if (ViewRelevance.bRenderCustomDepth)
 			{
 				bHasCustomDepthPrimitives = true;
-				CustumDepthStencilValues.Add(PrimitiveSceneInfo->Proxy->GetCustomDepthStencilValue());
+				CustomDepthStencilValues.Add(PrimitiveSceneInfo->Proxy->GetCustomDepthStencilValue());
 			}
 
 			extern bool GUseTranslucencyShadowDepths;
@@ -2882,7 +2882,7 @@ struct FRelevancePacket
 		WriteView.TranslucentPrimCount.Append(TranslucentPrimCount);
 		WriteView.bHasDistortionPrimitives |= bHasDistortionPrimitives;
 		WriteView.bHasCustomDepthPrimitives |= bHasCustomDepthPrimitives;
-		WriteView.CustomDepthStencilValues.Append(CustumDepthStencilValues);
+		WriteView.CustomDepthStencilValues.Append(CustomDepthStencilValues);
 		WriteView.bUsesCustomDepthStencilInTranslucentMaterials |= bUsesCustomDepthStencil;
 		DirtyIndirectLightingCacheBufferPrimitives.AppendTo(WriteView.DirtyIndirectLightingCacheBufferPrimitives);
 
