Fixes crashs in DOF

diff --git a/Engine/Source/Runtime/Renderer/Private/PostProcess/DiaphragmDOF.cpp b/Engine/Source/Runtime/Renderer/Private/PostProcess/DiaphragmDOF.cpp
--- a/Engine/Source/Runtime/Renderer/Private/PostProcess/DiaphragmDOF.cpp
+++ b/Engine/Source/Runtime/Renderer/Private/PostProcess/DiaphragmDOF.cpp
@@ -1998,9 +1998,10 @@ FRDGTextureRef DiaphragmDOF::AddPasses(
 		// Complete the reduction.
 		while (ReducedMipLevelCount < MipLevelCount)
 		{
-			int32 ProcessingMipLevelCount = FMath::Min(MipLevelCount - ReducedMipLevelCount, MaxReductionMipLevelCount);
+			int32 ProcessingMipLevelCount = FMath::Min(MipLevelCount - ReducedMipLevelCount, IsPostProcessingWithAlphaChannelSupported() ? 2 : (kMaxMipLevelCount - 1));
 			int32 InputMipLevel = ReducedMipLevelCount - 1;
 
+			FIntPoint GatherInputExtent = ReducedGatherInputTextures.SceneColor->Desc.Extent;
 			FIntPoint PassViewSize = FIntPoint::DivideAndRoundUp(PreprocessViewSize, 1 << InputMipLevel);
 
 			FDiaphragmDOFReduceCS::FPermutationDomain PermutationVector;
@@ -2011,14 +2012,14 @@ FRDGTextureRef DiaphragmDOF::AddPasses(
 			FDiaphragmDOFReduceCS::FParameters* PassParameters = GraphBuilder.AllocParameters<FDiaphragmDOFReduceCS::FParameters>();
 			PassParameters->ViewportRect = FIntRect(0, 0, PassViewSize.X, PassViewSize.Y);
 			PassParameters->MaxInputBufferUV = FVector2f(
-				(PreprocessViewSize.X - 0.5f) / SrcSize.X,
-				(PreprocessViewSize.Y - 0.5f) / SrcSize.Y);
+				(PreprocessViewSize.X - 0.5f) / GatherInputExtent.X,
+				(PreprocessViewSize.Y - 0.5f) / GatherInputExtent.Y);
 
 			PassParameters->EyeAdaptationTexture = GetEyeAdaptationTexture(GraphBuilder, View);
 			PassParameters->CommonParameters = CommonParameters;
 
 			float InputMipLevelPow2 = 1 << InputMipLevel;
-			PassParameters->GatherInputSize = FVector4f(SrcSize.X / InputMipLevelPow2, SrcSize.Y / InputMipLevelPow2, InputMipLevelPow2 / SrcSize.X, InputMipLevelPow2 / SrcSize.Y);
+			PassParameters->GatherInputSize = FVector4f(GatherInputExtent.X / InputMipLevelPow2, GatherInputExtent.Y / InputMipLevelPow2, InputMipLevelPow2 / GatherInputExtent.X, InputMipLevelPow2 / GatherInputExtent.Y);
 			PassParameters->GatherInput = CreateSRVs(GraphBuilder, ReducedGatherInputTextures, InputMipLevel);
 
 			for (int32 MipLevel = 0; MipLevel < ProcessingMipLevelCount; MipLevel++)
@@ -2027,19 +2028,16 @@ FRDGTextureRef DiaphragmDOF::AddPasses(
 			}
 
 			TShaderMapRef<FDiaphragmDOFReduceCS> ComputeShader(View.ShaderMap, PermutationVector);
-			ClearUnusedGraphResources(ComputeShader, PassParameters);
-			GraphBuilder.AddPass(
+			FComputeShaderUtils::AddPass(
+				GraphBuilder,
 				RDG_EVENT_NAME("DOF Reduce(Mips=[%d;%d]%s) %dx%d",
 					ReducedMipLevelCount,
 					ReducedMipLevelCount + ProcessingMipLevelCount - 1,
 					bRGBBufferSeparateCocBuffer ? TEXT(" R11G11B10") : TEXT(""),
 					PassViewSize.X, PassViewSize.Y),
+				ComputeShader,
 				PassParameters,
-				ERDGPassFlags::Compute,
-				[PassParameters, ComputeShader, PassViewSize](FRHICommandList& RHICmdList)
-			{
-				FComputeShaderUtils::Dispatch(RHICmdList, ComputeShader, *PassParameters, FComputeShaderUtils::GetGroupCount(PassViewSize, kDefaultGroupSize));
-			});
+				FComputeShaderUtils::GetGroupCount(PassViewSize, kDefaultGroupSize));
 
 			ReducedMipLevelCount += ProcessingMipLevelCount;
 		}
