Fix draw debug issue with translated world space.

diff --git a/Engine/Shaders/Private/ShaderDrawDebug.usf b/Engine/Shaders/Private/ShaderDrawDebug.usf
--- a/Engine/Shaders/Private/ShaderDrawDebug.usf
+++ b/Engine/Shaders/Private/ShaderDrawDebug.usf
@@ -5,6 +5,8 @@
 
 #include "ShaderDrawDebug.ush"
 
+float3 TranslatedWorldOffsetConversion;
+
 // The shader used to render the debug element on view
 StructuredBuffer<uint> ShaderDrawDebugPrimitive;
 
@@ -28,6 +30,11 @@ void ShaderDrawDebugVS(
 	float4 VertexPositionTWS = float4(VertexId == 0 ? Item.Pos0 : Item.Pos1, 1.0f);
 	float4 VertexColor = VertexId == 0 ? Item.Color0 : Item.Color1;
 
+	// Optional world offset. For 'locked' debug drawing, this offset translates the world 
+	// position recorded for a given translated world offset, into the current translated 
+	// world offset
+	VertexPositionTWS.xyz += TranslatedWorldOffsetConversion;
+
 	Position = mul(VertexPositionTWS, PrimaryView.TranslatedWorldToClip);
 	ColorAlpha = VertexColor;
 }
diff --git a/Engine/Source/Runtime/Renderer/Private/ScenePrivate.h b/Engine/Source/Runtime/Renderer/Private/ScenePrivate.h
--- a/Engine/Source/Runtime/Renderer/Private/ScenePrivate.h
+++ b/Engine/Source/Runtime/Renderer/Private/ScenePrivate.h
@@ -692,6 +692,7 @@ struct FHLODSceneNodeVisibilityState
 struct FShaderDrawDebugStateData
 {
 	TRefCountPtr<FRDGPooledBuffer> Buffer;
+	FVector PreViewTranslation = FVector::ZeroVector;
 	bool bIsLocked = false;
 
 	void Release()
diff --git a/Engine/Source/Runtime/Renderer/Private/ShaderDebug.cpp b/Engine/Source/Runtime/Renderer/Private/ShaderDebug.cpp
--- a/Engine/Source/Runtime/Renderer/Private/ShaderDebug.cpp
+++ b/Engine/Source/Runtime/Renderer/Private/ShaderDebug.cpp
@@ -132,6 +132,7 @@ namespace ShaderDrawDebug
 		SHADER_USE_PARAMETER_STRUCT(FShaderDrawDebugVS, FGlobalShader);
 
 		BEGIN_SHADER_PARAMETER_STRUCT(FParameters, )
+			SHADER_PARAMETER(FVector3f, TranslatedWorldOffsetConversion)
 			SHADER_PARAMETER_STRUCT_REF(FViewUniformShaderParameters, View)
 			SHADER_PARAMETER_SRV(StructuredBuffer, LockedShaderDrawDebugPrimitive)
 			SHADER_PARAMETER_RDG_BUFFER_SRV(StructuredBuffer, ShaderDrawDebugPrimitive)
@@ -234,6 +235,7 @@ namespace ShaderDrawDebug
 	static void InternalDrawView(
 		FRDGBuilder& GraphBuilder,
 		const FViewInfo& View,
+		const FVector& TranslatedWorldOffsetConversion,
 		FRDGBufferRef DataBuffer,
 		FRDGTextureRef OutputTexture,
 		FRDGTextureRef DepthTexture)
@@ -246,6 +248,7 @@ namespace ShaderDrawDebug
 
 		FShaderDrawVSPSParameters* PassParameters = GraphBuilder.AllocParameters<FShaderDrawVSPSParameters >();
 		PassParameters->VS.View = View.ViewUniformBuffer;
+		PassParameters->VS.TranslatedWorldOffsetConversion = FVector3f(TranslatedWorldOffsetConversion);
 		PassParameters->PS.RenderTargets[0] = FRenderTargetBinding(OutputTexture, ERenderTargetLoadAction::ELoad);
 		PassParameters->PS.OutputInvResolution = FVector2f(1.f / View.UnscaledViewRect.Width(), 1.f / View.UnscaledViewRect.Height());
 		PassParameters->PS.OriginalViewRectMin = FVector2f(View.ViewRect.Min);
@@ -317,12 +320,14 @@ namespace ShaderDrawDebug
 			if (IsShaderDrawLocked() && ViewState && !ViewState->ShaderDrawDebugStateData.bIsLocked)
 			{
 				ViewState->ShaderDrawDebugStateData.Buffer = GraphBuilder.ConvertToExternalBuffer(View.ShaderDrawData.Buffer);
+				ViewState->ShaderDrawDebugStateData.PreViewTranslation = View.ViewMatrices.GetPreViewTranslation();
 				ViewState->ShaderDrawDebugStateData.bIsLocked = true;
 			}
 
 			if (!IsShaderDrawLocked() && ViewState && ViewState->ShaderDrawDebugStateData.bIsLocked)
 			{
 				ViewState->ShaderDrawDebugStateData.Buffer = nullptr;
+				ViewState->ShaderDrawDebugStateData.PreViewTranslation = FVector::ZeroVector;
 				ViewState->ShaderDrawDebugStateData.bIsLocked = false;
 			}
 
@@ -349,13 +354,14 @@ namespace ShaderDrawDebug
 
 		{
 			FRDGBufferRef DataBuffer = View.ShaderDrawData.Buffer;
-			InternalDrawView(GraphBuilder, View, DataBuffer, OutputTexture, DepthTexture);
+			InternalDrawView(GraphBuilder, View, FVector::ZeroVector, DataBuffer, OutputTexture, DepthTexture);
 		}
 
 		if (View.ViewState && View.ViewState->ShaderDrawDebugStateData.bIsLocked)
 		{
+			const FVector LockedToCurrentTranslatedOffset = View.ViewMatrices.GetPreViewTranslation() - View.ViewState->ShaderDrawDebugStateData.PreViewTranslation;
 			FRDGBufferRef DataBuffer = GraphBuilder.RegisterExternalBuffer(View.ViewState->ShaderDrawDebugStateData.Buffer);
-			InternalDrawView(GraphBuilder, View, DataBuffer, OutputTexture, DepthTexture);
+			InternalDrawView(GraphBuilder, View, LockedToCurrentTranslatedOffset, DataBuffer, OutputTexture, DepthTexture);
 		}
 	}
 
