GI plugin interface ported from UE4.27

#pr 8852
#8852

diff --git a/Engine/Source/Runtime/Engine/Classes/Engine/EngineTypes.h b/Engine/Source/Runtime/Engine/Classes/Engine/EngineTypes.h
--- a/Engine/Source/Runtime/Engine/Classes/Engine/EngineTypes.h
+++ b/Engine/Source/Runtime/Engine/Classes/Engine/EngineTypes.h
@@ -397,6 +397,9 @@ namespace EDynamicGlobalIlluminationMethod
 
 		/** Standalone Ray Traced Global Illumination technique.  Deprecated, use Lumen Global Illumination instead. */
 		RayTraced UMETA(DisplayName="Standalone Ray Traced (Deprecated)"),
+
+		/** Use a plugin for Global Illumination */
+		Plugin UMETA(DisplayName="Plugin"),
 	};
 }
 
diff --git a/Engine/Source/Runtime/Renderer/Private/DeferredShadingRenderer.cpp b/Engine/Source/Runtime/Renderer/Private/DeferredShadingRenderer.cpp
--- a/Engine/Source/Runtime/Renderer/Private/DeferredShadingRenderer.cpp
+++ b/Engine/Source/Runtime/Renderer/Private/DeferredShadingRenderer.cpp
@@ -297,6 +297,36 @@ DECLARE_GPU_STAT(PostOpaqueExtensions);
 
 CSV_DEFINE_CATEGORY(LightCount, true);
 
+/*-----------------------------------------------------------------------------
+	Global Illumination Plugin Function Delegates
+-----------------------------------------------------------------------------*/
+
+static FGlobalIlluminationPluginDelegates::FAnyRayTracingPassEnabled GIPluginAnyRaytracingPassEnabledDelegate;
+FGlobalIlluminationPluginDelegates::FAnyRayTracingPassEnabled& FGlobalIlluminationPluginDelegates::AnyRayTracingPassEnabled()
+{
+	return GIPluginAnyRaytracingPassEnabledDelegate;
+}
+
+static FGlobalIlluminationPluginDelegates::FPrepareRayTracing GIPluginPrepareRayTracingDelegate;
+FGlobalIlluminationPluginDelegates::FPrepareRayTracing& FGlobalIlluminationPluginDelegates::PrepareRayTracing()
+{
+	return GIPluginPrepareRayTracingDelegate;
+}
+
+static FGlobalIlluminationPluginDelegates::FRenderDiffuseIndirectLight GIPluginRenderDiffuseIndirectLightDelegate;
+FGlobalIlluminationPluginDelegates::FRenderDiffuseIndirectLight& FGlobalIlluminationPluginDelegates::RenderDiffuseIndirectLight()
+{
+	return GIPluginRenderDiffuseIndirectLightDelegate;
+}
+
+#if !(UE_BUILD_SHIPPING || UE_BUILD_TEST)
+static FGlobalIlluminationPluginDelegates::FRenderDiffuseIndirectVisualizations GIPluginRenderDiffuseIndirectVisualizationsDelegate;
+FGlobalIlluminationPluginDelegates::FRenderDiffuseIndirectVisualizations& FGlobalIlluminationPluginDelegates::RenderDiffuseIndirectVisualizations()
+{
+	return GIPluginRenderDiffuseIndirectVisualizationsDelegate;
+}
+#endif //!(UE_BUILD_SHIPPING || UE_BUILD_TEST)
+
 const TCHAR* GetDepthPassReason(bool bDitheredLODTransitionsUseStencil, EShaderPlatform ShaderPlatform)
 {
 	if (IsForwardShadingEnabled(ShaderPlatform))
@@ -1418,6 +1448,7 @@ bool FDeferredShadingSceneRenderer::SetupRayTracingPipelineStates(FRHICommandLis
 					PrepareRayTracingAmbientOcclusion(View, RayGenShaders);
 					PrepareRayTracingSkyLight(View, *Scene, RayGenShaders);
 					PrepareRayTracingGlobalIllumination(View, RayGenShaders);
+					PrepareRayTracingGlobalIlluminationPlugin(View, RayGenShaders);
 					PrepareRayTracingTranslucency(View, RayGenShaders);
 
 					if (DoesPlatformSupportLumenGI(ShaderPlatform) && Lumen::UseHardwareRayTracing())
@@ -2868,6 +2899,11 @@ void FDeferredShadingSceneRenderer::Render(FRDGBuilder& GraphBuilder)
 
 		// Render diffuse sky lighting and reflections that only operate on opaque pixels
 		RenderDeferredReflectionsAndSkyLighting(GraphBuilder, SceneTextures, DynamicBentNormalAOTexture);
+		
+#if !(UE_BUILD_SHIPPING || UE_BUILD_TEST)
+		// Renders debug visualizations for global illumination plugins
+		RenderGlobalIlluminationPluginVisualizations(GraphBuilder, LightingChannelsTexture);
+#endif
 
 		AddSubsurfacePass(GraphBuilder, SceneTextures, Views);
 
@@ -3324,6 +3360,7 @@ bool AnyRayTracingPassEnabled(const FScene* Scene, const FViewInfo& View)
 		|| ShouldRenderRayTracingSkyLight(Scene->SkyLight)
 		|| ShouldRenderRayTracingShadows()
 		|| Scene->RayTracedLights.Num() > 0
+		|| ShouldRenderPluginRayTracingGlobalIllumination(View)
         || Lumen::AnyLumenHardwareRayTracingPassEnabled(Scene, View)
 		|| HasRayTracedOverlay(*View.Family);
 }
diff --git a/Engine/Source/Runtime/Renderer/Private/DeferredShadingRenderer.h b/Engine/Source/Runtime/Renderer/Private/DeferredShadingRenderer.h
--- a/Engine/Source/Runtime/Renderer/Private/DeferredShadingRenderer.h
+++ b/Engine/Source/Runtime/Renderer/Private/DeferredShadingRenderer.h
@@ -153,6 +153,41 @@ private:
 };
 
 
+/**
+ * Encapsulates the resources and render targets used by global illumination plugins.
+ */
+class RENDERER_API FGlobalIlluminationPluginResources : public FRenderResource
+
+{
+public:
+	FRDGTextureRef GBufferA;
+	FRDGTextureRef GBufferB;
+	FRDGTextureRef GBufferC;
+	FRDGTextureRef SceneDepthZ;
+	FRDGTextureRef SceneColor;
+	FRDGTextureRef LightingChannelsTexture;
+};
+
+/**
+ * Delegate callback used by global illumination plugins
+ */
+class RENDERER_API FGlobalIlluminationPluginDelegates
+{
+public:
+	DECLARE_MULTICAST_DELEGATE_OneParam(FAnyRayTracingPassEnabled, bool& /*bAnyRayTracingPassEnabled*/);
+	DECLARE_MULTICAST_DELEGATE_TwoParams(FPrepareRayTracing, const FViewInfo& /*View*/, TArray<FRHIRayTracingShader*>& /*OutRayGenShaders*/);
+	DECLARE_MULTICAST_DELEGATE_FourParams(FRenderDiffuseIndirectLight, const FScene& /*Scene*/, const FViewInfo& /*View*/, FRDGBuilder& /*GraphBuilder*/, FGlobalIlluminationPluginResources& /*Resources*/);
+
+	static FAnyRayTracingPassEnabled& AnyRayTracingPassEnabled();
+	static FPrepareRayTracing& PrepareRayTracing();
+	static FRenderDiffuseIndirectLight& RenderDiffuseIndirectLight();
+
+#if !(UE_BUILD_SHIPPING || UE_BUILD_TEST)
+	DECLARE_MULTICAST_DELEGATE_FourParams(FRenderDiffuseIndirectVisualizations, const FScene& /*Scene*/, const FViewInfo& /*View*/, FRDGBuilder& /*GraphBuilder*/, FGlobalIlluminationPluginResources& /*Resources*/);
+	static FRenderDiffuseIndirectVisualizations& RenderDiffuseIndirectVisualizations();
+#endif //!(UE_BUILD_SHIPPING || UE_BUILD_TEST)
+};
+
 /**
  * Scene renderer that implements a deferred shading pipeline and associated features.
  */
@@ -265,6 +300,7 @@ private:
 		SSGI,
 		RTGI,
 		Lumen,
+		Plugin,
 	};
 
 	enum class EAmbientOcclusionMethod
@@ -464,6 +500,11 @@ private:
 
 	void RenderDeferredReflectionsAndSkyLightingHair(FRDGBuilder& GraphBuilder);
 
+#if !(UE_BUILD_SHIPPING || UE_BUILD_TEST)
+	/** Renders debug visualizations for global illumination plugins. */
+	void RenderGlobalIlluminationPluginVisualizations(FRDGBuilder& GraphBuilder, FRDGTextureRef LightingChannelsTexture);
+#endif
+
 	/** Computes DFAO, modulates it to scene color (which is assumed to contain diffuse indirect lighting), and stores the output bent normal for use occluding specular. */
 	void RenderDFAOAsIndirectShadowing(
 		FRDGBuilder& GraphBuilder,
@@ -916,6 +957,7 @@ private:
 	static void PrepareRayTracingAmbientOcclusion(const FViewInfo& View, TArray<FRHIRayTracingShader*>& OutRayGenShaders);
 	static void PrepareRayTracingSkyLight(const FViewInfo& View, const FScene& Scene, TArray<FRHIRayTracingShader*>& OutRayGenShaders);
 	static void PrepareRayTracingGlobalIllumination(const FViewInfo& View, TArray<FRHIRayTracingShader*>& OutRayGenShaders);
+	static void PrepareRayTracingGlobalIlluminationPlugin(const FViewInfo& View, TArray<FRHIRayTracingShader*>& OutRayGenShaders);
 	static void PrepareRayTracingTranslucency(const FViewInfo& View, TArray<FRHIRayTracingShader*>& OutRayGenShaders);
 	static void PrepareRayTracingDebug(const FSceneViewFamily& ViewFamily, TArray<FRHIRayTracingShader*>& OutRayGenShaders);
 	static void PreparePathTracing(const FSceneViewFamily& ViewFamily, TArray<FRHIRayTracingShader*>& OutRayGenShaders);
diff --git a/Engine/Source/Runtime/Renderer/Private/IndirectLightRendering.cpp b/Engine/Source/Runtime/Renderer/Private/IndirectLightRendering.cpp
--- a/Engine/Source/Runtime/Renderer/Private/IndirectLightRendering.cpp
+++ b/Engine/Source/Runtime/Renderer/Private/IndirectLightRendering.cpp
@@ -32,7 +32,8 @@ static TAutoConsoleVariable<int32> CVarDynamicGlobalIlluminationMethod(
 	TEXT("0 - None.  Global Illumination can be baked into Lightmaps but no technique will be used for Dynamic Global Illumination.\n")
 	TEXT("1 - Lumen.  Use Lumen Global Illumination for all lights, emissive materials casting light and SkyLight Occlusion.  Requires 'Generate Mesh Distance Fields' enabled for Software Ray Tracing and 'Support Hardware Ray Tracing' enabled for Hardware Ray Tracing.\n")
 	TEXT("2 - SSGI.  Standalone Screen Space Global Illumination.  Low cost, but limited by screen space information.\n")
-	TEXT("3 - RTGI.  Ray Traced Global Illumination technique.  Deprecated, use Lumen Global Illumination instead."),
+	TEXT("3 - RTGI.  Ray Traced Global Illumination technique.  Deprecated, use Lumen Global Illumination instead.\n")
+	TEXT("4 - Plugin.  Use a plugin for Global Illumination."),
 	ECVF_RenderThreadSafe | ECVF_Scalability);
 
 // Must match EReflectionMethod
@@ -391,6 +392,11 @@ void FDeferredShadingSceneRenderer::CommitIndirectLightingState()
 			DiffuseIndirectMethod = EDiffuseIndirectMethod::RTGI;
 			DiffuseIndirectDenoiser = IScreenSpaceDenoiser::GetDenoiserMode(CVarDiffuseIndirectDenoiser);
 		}
+		else if (ShouldRenderPluginRayTracingGlobalIllumination(View))
+		{
+			DiffuseIndirectMethod = EDiffuseIndirectMethod::Plugin;
+			DiffuseIndirectDenoiser = IScreenSpaceDenoiser::GetDenoiserMode(CVarDiffuseIndirectDenoiser);
+		}
 		
 		if (DiffuseIndirectMethod == EDiffuseIndirectMethod::Disabled && ScreenSpaceRayTracing::IsScreenSpaceDiffuseIndirectSupported(View))
 		{
@@ -528,6 +534,29 @@ TUniformBufferRef<FReflectionUniformParameters> CreateReflectionUniformBuffer(co
 	return CreateUniformBufferImmediate(ReflectionStruct, Usage);
 }
 
+#if RHI_RAYTRACING
+bool ShouldRenderPluginRayTracingGlobalIllumination(const FViewInfo& View)
+{
+	if (View.FinalPostProcessSettings.DynamicGlobalIlluminationMethod != EDynamicGlobalIlluminationMethod::Plugin)
+	{
+		return false;
+	}
+
+	bool bAnyRayTracingPassEnabled = false;
+	FGlobalIlluminationPluginDelegates::FAnyRayTracingPassEnabled& Delegate = FGlobalIlluminationPluginDelegates::AnyRayTracingPassEnabled();
+	Delegate.Broadcast(bAnyRayTracingPassEnabled);
+
+	return ShouldRenderRayTracingEffect(bAnyRayTracingPassEnabled, ERayTracingPipelineCompatibilityFlags::FullPipeline);
+}
+
+void FDeferredShadingSceneRenderer::PrepareRayTracingGlobalIlluminationPlugin(const FViewInfo& View, TArray<FRHIRayTracingShader*>& OutRayGenShaders)
+{
+	// Call the GI plugin delegate function to prepare ray tracing
+	FGlobalIlluminationPluginDelegates::FPrepareRayTracing& Delegate = FGlobalIlluminationPluginDelegates::PrepareRayTracing();
+	Delegate.Broadcast(View, OutRayGenShaders);
+}
+#endif
+
 static const FVector SampleArray4x4x6[96] = {
 	FVector(0.72084325551986694, -0.44043412804603577, -0.53516626358032227),
 	FVector(-0.51286971569061279, 0.57541996240615845, 0.63706874847412109),
@@ -851,6 +880,20 @@ void FDeferredShadingSceneRenderer::RenderDiffuseIndirectAndAmbientOcclusion(
 				DenoiserOutputs.Textures[2] = DenoiserOutputs.Textures[1];
 			}
 		}
+		else if (ViewPipelineState.DiffuseIndirectMethod == EDiffuseIndirectMethod::Plugin)
+		{
+			// Get the resources and call the GI plugin's rendering function delegate
+			FGlobalIlluminationPluginResources GIPluginResources;
+			GIPluginResources.GBufferA = SceneTextures.GBufferA;
+			GIPluginResources.GBufferB = SceneTextures.GBufferB;
+			GIPluginResources.GBufferC = SceneTextures.GBufferC;
+			GIPluginResources.LightingChannelsTexture = LightingChannelsTexture;
+			GIPluginResources.SceneDepthZ = SceneTextures.Depth.Target;
+			GIPluginResources.SceneColor = SceneTextures.Color.Target;
+
+			FGlobalIlluminationPluginDelegates::FRenderDiffuseIndirectLight& Delegate = FGlobalIlluminationPluginDelegates::RenderDiffuseIndirectLight();
+			Delegate.Broadcast(*Scene, View, GraphBuilder, GIPluginResources);
+		}
 
 		FRDGTextureRef AmbientOcclusionMask = DenoiserInputs.AmbientOcclusionMask;
 
@@ -1650,3 +1693,43 @@ void FDeferredShadingSceneRenderer::RenderDeferredReflectionsAndSkyLightingHair(
 	}
 
 }
+
+#if !(UE_BUILD_SHIPPING || UE_BUILD_TEST)
+void FDeferredShadingSceneRenderer::RenderGlobalIlluminationPluginVisualizations(
+	FRDGBuilder& GraphBuilder,
+	FRDGTextureRef LightingChannelsTexture)
+{
+	// Early out if GI plugins aren't enabled
+	bool bGIPluginEnabled = false;
+	for (const FViewInfo& View : Views)
+	{
+		if (View.FinalPostProcessSettings.DynamicGlobalIlluminationMethod == EDynamicGlobalIlluminationMethod::Plugin)
+		{
+			bGIPluginEnabled = true;
+			break;
+		}
+	}
+	if (!bGIPluginEnabled)
+	{
+		return;
+	}
+
+	const FSceneTextures& SceneTextures = FSceneTextures::Get(GraphBuilder);
+
+	// Get the resources passed to GI plugins
+	FGlobalIlluminationPluginResources GIPluginResources;
+	GIPluginResources.GBufferA = SceneTextures.GBufferA;
+	GIPluginResources.GBufferB = SceneTextures.GBufferB;
+	GIPluginResources.GBufferC = SceneTextures.GBufferC;
+	GIPluginResources.LightingChannelsTexture = LightingChannelsTexture;
+	GIPluginResources.SceneDepthZ = SceneTextures.Depth.Target;
+	GIPluginResources.SceneColor = SceneTextures.Color.Target;
+
+	// Render visualizations to all views by calling the GI plugin's delegate
+	FGlobalIlluminationPluginDelegates::FRenderDiffuseIndirectVisualizations& PRVDelegate = FGlobalIlluminationPluginDelegates::RenderDiffuseIndirectVisualizations();
+	for (int32 ViewIndexZ = 0; ViewIndexZ < Views.Num(); ViewIndexZ++)
+	{
+		PRVDelegate.Broadcast(*Scene, Views[ViewIndexZ], GraphBuilder, GIPluginResources);
+	}
+}
+#endif //!(UE_BUILD_SHIPPING || UE_BUILD_TEST)
\ No newline at end of file
diff --git a/Engine/Source/Runtime/Renderer/Private/RayTracing/RayTracingLighting.h b/Engine/Source/Runtime/Renderer/Private/RayTracing/RayTracingLighting.h
--- a/Engine/Source/Runtime/Renderer/Private/RayTracing/RayTracingLighting.h
+++ b/Engine/Source/Runtime/Renderer/Private/RayTracing/RayTracingLighting.h
@@ -12,7 +12,7 @@
 
 #if RHI_RAYTRACING
 
-BEGIN_GLOBAL_SHADER_PARAMETER_STRUCT(FRaytracingLightDataPacked, )
+BEGIN_GLOBAL_SHADER_PARAMETER_STRUCT(FRaytracingLightDataPacked, RENDERER_API)
 	SHADER_PARAMETER(uint32, Count)
 	SHADER_PARAMETER(float, IESLightProfileInvCount)
 	SHADER_PARAMETER(uint32, CellCount)
diff --git a/Engine/Source/Runtime/Renderer/Private/RayTracing/RaytracingOptions.h b/Engine/Source/Runtime/Renderer/Private/RayTracing/RaytracingOptions.h
--- a/Engine/Source/Runtime/Renderer/Private/RayTracing/RaytracingOptions.h
+++ b/Engine/Source/Runtime/Renderer/Private/RayTracing/RaytracingOptions.h
@@ -68,6 +68,7 @@ extern bool ShouldRenderRayTracingTranslucency(const FViewInfo& View);
 extern bool ShouldRenderRayTracingShadows();
 extern bool ShouldRenderRayTracingShadowsForLight(const FLightSceneProxy& LightProxy);
 extern bool ShouldRenderRayTracingShadowsForLight(const FLightSceneInfoCompact& LightInfo);
+extern bool ShouldRenderPluginRayTracingGlobalIllumination(const FViewInfo& View);
 extern bool HasRayTracedOverlay(const FSceneViewFamily& ViewFamily);
 
 extern bool EnableRayTracingShadowTwoSidedGeometry();
@@ -129,6 +130,11 @@ FORCEINLINE bool ShouldRenderRayTracingShadowsForLight(const FLightSceneInfoComp
 	return false;
 }
 
+FORCEINLINE bool ShouldRenderPluginRayTracingGlobalIllumination(const FViewInfo& View)
+{
+	return false;
+}
+
 FORCEINLINE bool HasRayTracedOverlay(const FSceneViewFamily& ViewFamily)
 {
 	return false;
