Add missing include for GlobalDistanceFieldParameters to prevent non-unity build failiure

diff --git a/Engine/Source/Runtime/Renderer/Public/GlobalDistanceFieldParameters.h b/Engine/Source/Runtime/Renderer/Public/GlobalDistanceFieldParameters.h
--- a/Engine/Source/Runtime/Renderer/Public/GlobalDistanceFieldParameters.h
+++ b/Engine/Source/Runtime/Renderer/Public/GlobalDistanceFieldParameters.h
@@ -7,6 +7,7 @@
 #pragma once
 
 #include "CoreMinimal.h"
+#include "ShaderParameterMacros.h"
 #include "ShaderParameters.h"
 #include "RenderUtils.h"
 #include "RHIStaticStates.h"
