Allow caching of ray tracing instances for preview scenes

This enables the use of path tracing for the material editor and static mesh editors.

diff --git a/Engine/Source/Editor/AdvancedPreviewScene/Private/AdvancedPreviewScene.cpp b/Engine/Source/Editor/AdvancedPreviewScene/Private/AdvancedPreviewScene.cpp
--- a/Engine/Source/Editor/AdvancedPreviewScene/Private/AdvancedPreviewScene.cpp
+++ b/Engine/Source/Editor/AdvancedPreviewScene/Private/AdvancedPreviewScene.cpp
@@ -53,7 +53,6 @@ FAdvancedPreviewScene::FAdvancedPreviewScene(ConstructionValues CVS, float InFlo
 	check(SkySphere);
 	SkyComponent->SetStaticMesh(SkySphere);
 	SkyComponent->SetCollisionEnabled(ECollisionEnabled::NoCollision);
-	SkyComponent->bVisibleInRayTracing = false;
 
 	UMaterial* SkyMaterial = LoadObject<UMaterial>(NULL, TEXT("/Engine/EditorMaterials/AssetViewer/M_SkyBox.M_SkyBox"), NULL, LOAD_None, NULL);
 	check(SkyMaterial);
diff --git a/Engine/Source/Runtime/Renderer/Private/PrimitiveSceneInfo.cpp b/Engine/Source/Runtime/Renderer/Private/PrimitiveSceneInfo.cpp
--- a/Engine/Source/Runtime/Renderer/Private/PrimitiveSceneInfo.cpp
+++ b/Engine/Source/Runtime/Renderer/Private/PrimitiveSceneInfo.cpp
@@ -706,7 +706,7 @@ void FScene::RefreshRayTracingInstances()
 
 void FPrimitiveSceneInfo::UpdateCachedRayTracingInstances(FScene* Scene, const TArrayView<FPrimitiveSceneInfo*>& SceneInfos)
 {
-	if (IsRayTracingEnabled() && !(Scene->World->WorldType == EWorldType::EditorPreview || Scene->World->WorldType == EWorldType::GamePreview))
+	if (IsRayTracingEnabled())
 	{
 		checkf(GRHISupportsMultithreadedShaderCreation, TEXT("Raytracing code needs the ability to create shaders from task threads."));
 
@@ -864,7 +864,7 @@ void CacheRayTracingPrimitive(
 
 void FPrimitiveSceneInfo::CacheRayTracingPrimitives(FScene* Scene, const TArrayView<FPrimitiveSceneInfo*>& SceneInfos)
 {
-	if (IsRayTracingEnabled() && !(Scene->World->WorldType == EWorldType::EditorPreview || Scene->World->WorldType == EWorldType::GamePreview))
+	if (IsRayTracingEnabled())
 	{
 		CSV_SCOPED_TIMING_STAT_EXCLUSIVE(FPrimitiveSceneInfo_CacheRayTracingPrimitives)
 		SCOPED_NAMED_EVENT(FPrimitiveSceneInfo_CacheRayTracingPrimitives, FColor::Emerald);
