Add hair GPU scene bound visualization.

diff --git a/Engine/Shaders/Private/HairStrands/HairStrandsDebugPrint.usf b/Engine/Shaders/Private/HairStrands/HairStrandsDebugPrint.usf
--- a/Engine/Shaders/Private/HairStrands/HairStrandsDebugPrint.usf
+++ b/Engine/Shaders/Private/HairStrands/HairStrandsDebugPrint.usf
@@ -10,7 +10,6 @@
 #include "HairStrandsVisibilityCommon.ush"
 #include "HairStrandsAABBCommon.ush" 
 
-
 ////////////////////////////////////////////////////////////////////////////////////////////////
 
 void AddNewLine(float2 OriginalPos, inout float2 Pos)
@@ -22,6 +21,8 @@ void AddNewLine(float2 OriginalPos, inout float2 Pos)
 
 #if SHADER_PRINT
 
+#include "../SceneData.ush"
+
 int2 PixelCoord;
 int2 MaxResolution;
 uint FastResolveMask;
@@ -29,7 +30,12 @@ uint HairMacroGroupCount;
 uint2 GroupSize; 
 uint HairVisibilityNodeGroupSize;
 uint AllocatedSampleCount;
+uint HairInstanceCount;
+
+uint NumInstances;
+uint InstanceDataSOAStride;
 
+Buffer<uint> HairInstanceIDs;
 Texture2D<float> HairCountTexture;
 Texture2D<uint> HairCountUintTexture;
 
@@ -48,6 +54,31 @@ FFontColor GetOccupancyColor(float In)
 	return InitFontColor(Color);
 }
 
+void DrawInstanceAABB(uint PrimitiveId)
+{
+	// Since we don't have the actual instance ID (they are transient instance vs. persistent instance), 
+	// we find the primitive with the matching PrimitiveID
+	uint InstanceID = ~0;
+	const uint TotalInstanceCount = min(8192u, NumInstances);
+	for (uint LocalInstanceID = 0; LocalInstanceID < TotalInstanceCount; ++LocalInstanceID)
+	{
+		FInstanceSceneData InstanceData = GetInstanceSceneData(LocalInstanceID, InstanceDataSOAStride);
+		if (InstanceData.PrimitiveId == PrimitiveId)
+		{
+			InstanceID = LocalInstanceID;
+			break;
+		}
+	}
+
+	if (InstanceID != ~0)
+	{
+		FInstanceSceneData InstanceData = GetInstanceSceneData(InstanceID, InstanceDataSOAStride);
+		const float4 InstanceBoundColor = ColorLightGreen;
+		AddOBBWS(InstanceData.LocalBoundsCenter - InstanceData.LocalBoundsExtent, InstanceData.LocalBoundsCenter + InstanceData.LocalBoundsExtent, InstanceBoundColor, LWCHackToFloat(InstanceData.LocalToWorld));
+		AddReferentialWS(LWCHackToFloat(InstanceData.LocalToWorld), 50.f);
+	}
+}
+
 [numthreads(1, 1, 1)]
 void MainCS(uint GroupIndex : SV_GroupIndex, uint3 DispatchThreadId : SV_DispatchThreadID)
 {	
@@ -62,6 +93,16 @@ void MainCS(uint GroupIndex : SV_GroupIndex, uint3 DispatchThreadId : SV_Dispatc
 			AddAABBTWS(Bound.Min, Bound.Max, InstanceBoundColor);
 		}
 
+		// Hair instance bounds
+		for (uint HairInstanceIndex = 0; HairInstanceIndex < HairInstanceCount; HairInstanceIndex++)
+		{
+			const uint InstanceID = HairInstanceIDs[HairInstanceIndex];
+			if (InstanceID != ~0)
+			{
+				DrawInstanceAABB(InstanceID);
+			}
+		}
+
 		const uint DispatchX = HairVisibilityIndirectArgsBuffer[0];
 		const uint DispatchY = HairVisibilityIndirectArgsBuffer[1];
 		const uint UsedSampleCount = HairStrands.HairSampleCount.Load(uint3(0, 0, 0));
diff --git a/Engine/Source/Runtime/Renderer/Private/GPUScene.cpp b/Engine/Source/Runtime/Renderer/Private/GPUScene.cpp
--- a/Engine/Source/Runtime/Renderer/Private/GPUScene.cpp
+++ b/Engine/Source/Runtime/Renderer/Private/GPUScene.cpp
@@ -1905,15 +1905,9 @@ class FGPUSceneDebugRenderCS : public FGlobalShader
 	SHADER_USE_PARAMETER_STRUCT(FGPUSceneDebugRenderCS, FGlobalShader);
 
 	BEGIN_SHADER_PARAMETER_STRUCT(FParameters, )
+		SHADER_PARAMETER_STRUCT_INCLUDE(FGPUSceneResourceParameters, GPUSceneResource)
 		SHADER_PARAMETER_STRUCT_INCLUDE(ShaderDrawDebug::FShaderParameters, ShaderDrawUniformBuffer)
 		SHADER_PARAMETER_STRUCT_INCLUDE(ShaderPrint::FShaderParameters, ShaderPrintUniformBuffer)
-		SHADER_PARAMETER_SRV(StructuredBuffer<float4>, GPUSceneInstanceSceneData)
-		SHADER_PARAMETER_SRV(StructuredBuffer<float4>, GPUSceneInstancePayloadData)
-		SHADER_PARAMETER_SRV(StructuredBuffer<float4>, GPUScenePrimitiveSceneData)
-		SHADER_PARAMETER(uint32, InstanceDataSOAStride)
-		SHADER_PARAMETER(uint32, GPUSceneFrameNumber)
-		SHADER_PARAMETER(int32, NumInstances)
-		SHADER_PARAMETER(int32, NumScenePrimitives)
 		SHADER_PARAMETER(int32, bDrawAll)
 		SHADER_PARAMETER(int32, bDrawUpdatedOnly)
 		SHADER_PARAMETER(int32, SelectedNameInfoCount)
@@ -1951,6 +1945,19 @@ public:
 IMPLEMENT_GLOBAL_SHADER(FGPUSceneDebugRenderCS, "/Engine/Private/GPUSceneDebugRender.usf", "GPUSceneDebugRenderCS", SF_Compute);
 
 
+FGPUSceneResourceParameters FGPUScene::SetParameters(FRDGBuilder& GraphBuilder) const
+{
+	FGPUSceneResourceParameters Out;
+	Out.GPUSceneInstanceSceneData = InstanceSceneDataBuffer.SRV;
+	Out.GPUSceneInstancePayloadData = InstancePayloadDataBuffer.SRV;
+	Out.GPUScenePrimitiveSceneData = PrimitiveBuffer.SRV;
+	Out.InstanceDataSOAStride = InstanceSceneDataSOAStride;
+	Out.NumScenePrimitives = NumScenePrimitives;
+	Out.NumInstances = InstanceSceneDataAllocator.GetMaxSize();
+	Out.GPUSceneFrameNumber = GetSceneFrameNumber();
+	return Out;
+}
+
 void FGPUScene::DebugRender(FRDGBuilder& GraphBuilder, FScene& Scene, FViewInfo& View)
 {
 	int32 DebugMode = CVarGPUSceneDebugMode.GetValueOnRenderThread();
@@ -2027,20 +2034,14 @@ void FGPUScene::DebugRender(FRDGBuilder& GraphBuilder, FScene& Scene, FViewInfo&
 			FGPUSceneDebugRenderCS::FParameters* PassParameters = GraphBuilder.AllocParameters<FGPUSceneDebugRenderCS::FParameters>();
 			ShaderDrawDebug::SetParameters(GraphBuilder, View.ShaderDrawData, PassParameters->ShaderDrawUniformBuffer);
 			ShaderPrint::SetParameters(GraphBuilder, View, PassParameters->ShaderPrintUniformBuffer);
-			PassParameters->GPUSceneInstanceSceneData = InstanceSceneDataBuffer.SRV;
-			PassParameters->GPUSceneInstancePayloadData = InstancePayloadDataBuffer.SRV;
-			PassParameters->GPUScenePrimitiveSceneData = PrimitiveBuffer.SRV;
-			PassParameters->InstanceDataSOAStride = InstanceSceneDataSOAStride;
-			PassParameters->GPUSceneFrameNumber = GetSceneFrameNumber();
+			PassParameters->GPUSceneResource = SetParameters(GraphBuilder);
 			PassParameters->bDrawUpdatedOnly = DebugMode == 3;
 			PassParameters->bDrawAll = DebugMode != 2;
-			PassParameters->NumInstances = NumInstances;
 			PassParameters->SelectedNameInfoCount = SelectedCount;
 			PassParameters->SelectedNameCharacterCount = SelectedCount > 0 ? SelectedNames.Num() : 0;
 			PassParameters->SelectedPrimitiveFlags = GraphBuilder.CreateSRV(SelectedPrimitiveFlagsRDG);
 			PassParameters->SelectedPrimitiveNameInfos = GraphBuilder.CreateSRV(SelectedPrimitiveNameInfos);
 			PassParameters->SelectedPrimitiveNames = GraphBuilder.CreateSRV(SelectedPrimitiveNames, PF_R8_UINT);
-			PassParameters->NumScenePrimitives = NumScenePrimitives;
 			PassParameters->DrawRange = CVarGPUSceneDebugDrawRange.GetValueOnRenderThread();
 			PassParameters->RWDrawCounter = DrawCounterUAV;
 
diff --git a/Engine/Source/Runtime/Renderer/Private/GPUScene.h b/Engine/Source/Runtime/Renderer/Private/GPUScene.h
--- a/Engine/Source/Runtime/Renderer/Private/GPUScene.h
+++ b/Engine/Source/Runtime/Renderer/Private/GPUScene.h
@@ -21,6 +21,15 @@ extern RENDERER_API void AddPrimitiveToUpdateGPU(FScene& Scene, int32 PrimitiveI
 class FGPUScene;
 class FGPUSceneDynamicContext;
 
+BEGIN_SHADER_PARAMETER_STRUCT(FGPUSceneResourceParameters, )
+	SHADER_PARAMETER_SRV(StructuredBuffer<float4>, GPUSceneInstanceSceneData)
+	SHADER_PARAMETER_SRV(StructuredBuffer<float4>, GPUSceneInstancePayloadData)
+	SHADER_PARAMETER_SRV(StructuredBuffer<float4>, GPUScenePrimitiveSceneData)
+	SHADER_PARAMETER(uint32, InstanceDataSOAStride)
+	SHADER_PARAMETER(uint32, GPUSceneFrameNumber)
+	SHADER_PARAMETER(int32, NumInstances)
+	SHADER_PARAMETER(int32, NumScenePrimitives)
+END_SHADER_PARAMETER_STRUCT()
 
 /**
  * Used to manage dynamic primitives for a given view, during InitViews the data is collected and then can be committed to the GPU-Scene. 
@@ -293,6 +302,11 @@ public:
 
 	const FGrowOnlySpanAllocator& GetInstanceSceneDataAllocator() const { return InstanceSceneDataAllocator; }
 
+	/**
+	 * Return the GPU scene resource
+	 */
+	FGPUSceneResourceParameters SetParameters(FRDGBuilder& GraphBuilder) const;
+
 	/**
 	 * Draw GPU-Scene debug info, such as bounding boxes. Call once per view at some point in the frame after GPU scene has been updated fully.
 	 * What is drawn is controlled by the CVar: r.GPUScene.DebugMode. Enabling this cvar causes ShaderDraw to be being active (if supported). 
diff --git a/Engine/Source/Runtime/Renderer/Private/HairStrands/HairStrandsDebug.cpp b/Engine/Source/Runtime/Renderer/Private/HairStrands/HairStrandsDebug.cpp
--- a/Engine/Source/Runtime/Renderer/Private/HairStrands/HairStrandsDebug.cpp
+++ b/Engine/Source/Runtime/Renderer/Private/HairStrands/HairStrandsDebug.cpp
@@ -330,12 +330,15 @@ class FHairDebugPrintCS : public FGlobalShader
 		SHADER_PARAMETER(uint32, HairMacroGroupCount)
 		SHADER_PARAMETER(uint32, HairVisibilityNodeGroupSize)
 		SHADER_PARAMETER(uint32, AllocatedSampleCount)
+		SHADER_PARAMETER(uint32, HairInstanceCount)
+		SHADER_PARAMETER_RDG_BUFFER_SRV(Buffer<uint>, HairInstanceIDs)
 		SHADER_PARAMETER_RDG_TEXTURE(Texture2D, HairCountTexture)
 		SHADER_PARAMETER_RDG_TEXTURE(Texture2D, HairCountUintTexture)
 		SHADER_PARAMETER_RDG_BUFFER_SRV(Buffer<uint>, HairVisibilityIndirectArgsBuffer)
 		SHADER_PARAMETER_RDG_BUFFER_SRV(Buffer, HairMacroGroupAABBBuffer)
 		SHADER_PARAMETER_RDG_TEXTURE_SRV(Texture2D, StencilTexture)
 		SHADER_PARAMETER_SAMPLER(SamplerState, LinearSampler)
+		SHADER_PARAMETER_STRUCT_INCLUDE(FGPUSceneResourceParameters, GPUSceneResource)
 		SHADER_PARAMETER_STRUCT_REF(FViewUniformShaderParameters, ViewUniformBuffer)
 		SHADER_PARAMETER_STRUCT_INCLUDE(ShaderPrint::FShaderParameters, ShaderPrintUniformBuffer)
 		SHADER_PARAMETER_STRUCT_INCLUDE(ShaderDrawDebug::FShaderParameters, ShaderDrawUniformBuffer)
@@ -350,6 +353,8 @@ public:
 		FGlobalShader::ModifyCompilationEnvironment(Parameters, OutEnvironment);
 		OutEnvironment.CompilerFlags.Add(CFLAG_Debug);
 		OutEnvironment.SetDefine(TEXT("SHADER_PRINT"), 1);
+		OutEnvironment.SetDefine(TEXT("VF_SUPPORTS_PRIMITIVE_SCENE_DATA"), 1);
+		OutEnvironment.SetDefine(TEXT("USE_GLOBAL_GPU_SCENE_DATA"), 1);
 	}
 };
 
@@ -357,19 +362,41 @@ IMPLEMENT_GLOBAL_SHADER(FHairDebugPrintCS, "/Engine/Private/HairStrands/HairStra
 
 static void AddDebugHairPrintPass(
 	FRDGBuilder& GraphBuilder,
+	const FScene* Scene,
 	const FViewInfo* View,
 	const EHairDebugMode InDebugMode,
 	const FHairStrandsVisibilityData& VisibilityData,
+	const FHairStrandsMacroGroupDatas& MacroGroupDatas,
 	const FHairStrandsMacroGroupResources& MacroGroupResources,
 	FRDGTextureSRVRef InStencilTexture)
 {
-	if (!View || !View->HairStrandsViewData.UniformBuffer || !InStencilTexture) return;
+	if (!Scene || !View || !View->HairStrandsViewData.UniformBuffer || !InStencilTexture) return;
 
 	if (!TryEnableShaderDrawAndShaderPrint(*View, MacroGroupResources.MacroGroupCount * 32u, 2000u))
 	{
 		return;
 	}
 
+	// Build mapping Instance -> PrimitiveID to fetch primitive data (i.e., transform & co)
+	const uint32 MacroGroupCount = MacroGroupDatas.Num();
+	TArray<uint32> InstanceIDs;
+	InstanceIDs.Reserve(MacroGroupCount * 4u);
+	for (uint32 MacroGroupIndex = 0; MacroGroupIndex < MacroGroupCount; ++MacroGroupIndex)
+	{
+		const FHairStrandsMacroGroupData& MacroGroup = MacroGroupDatas[MacroGroupIndex];
+		for (const FHairStrandsMacroGroupData::PrimitiveInfo& PrimitiveInfo : MacroGroup.PrimitivesInfos)
+		{
+			uint32 PrimitiveID = ~0u;
+			if (const FPrimitiveSceneInfo* SceneInfo = PrimitiveInfo.PrimitiveSceneProxy->GetPrimitiveSceneInfo())
+			{
+				PrimitiveID = SceneInfo->GetIndex();
+			}
+			InstanceIDs.Add(PrimitiveID);
+		}
+	}
+	FRDGBufferRef InstancesIDBuffer = CreateVertexBuffer(GraphBuilder, TEXT("Hair.Debug.InstanceIDs"), FRDGBufferDesc::CreateBufferDesc(4, InstanceIDs.Num()), InstanceIDs.GetData(), 4u * InstanceIDs.Num());
+
+
 	FRDGTextureRef ViewHairCountTexture = VisibilityData.ViewHairCountTexture ? VisibilityData.ViewHairCountTexture : GSystemTextures.GetBlackDummy(GraphBuilder);
 	FRDGTextureRef ViewHairCountUintTexture = VisibilityData.ViewHairCountUintTexture ? VisibilityData.ViewHairCountUintTexture : GSystemTextures.GetBlackDummy(GraphBuilder);
 
@@ -377,6 +404,9 @@ static void AddDebugHairPrintPass(
 	const FIntPoint Resolution(Viewport.Width(), Viewport.Height());
 
 	FHairDebugPrintCS::FParameters* Parameters = GraphBuilder.AllocParameters<FHairDebugPrintCS::FParameters>();
+	Parameters->GPUSceneResource = Scene->GPUScene.SetParameters(GraphBuilder);
+	Parameters->HairInstanceCount = InstanceIDs.Num();
+	Parameters->HairInstanceIDs = GraphBuilder.CreateSRV(InstancesIDBuffer, PF_R32_UINT);
 	Parameters->GroupSize = GetVendorOptimalGroupSize2D();
 	Parameters->ViewUniformBuffer = View->ViewUniformBuffer;
 	Parameters->MaxResolution = VisibilityData.CoverageTexture ? VisibilityData.CoverageTexture->Desc.Extent : FIntPoint(0,0);
@@ -1288,7 +1318,7 @@ static void InternalRenderHairStrandsDebugInfo(
 
 	if (HairDebugMode == EHairDebugMode::MacroGroups)
 	{
-		AddDebugHairPrintPass(GraphBuilder, &View, HairDebugMode, HairData.VisibilityData, HairData.MacroGroupResources, SceneTextures.Stencil);
+		AddDebugHairPrintPass(GraphBuilder, Scene, &View, HairDebugMode, HairData.VisibilityData, HairData.MacroGroupDatas, HairData.MacroGroupResources, SceneTextures.Stencil);
 	}
 
 	if (HairDebugMode == EHairDebugMode::DeepOpacityMaps)
@@ -1349,7 +1379,7 @@ static void InternalRenderHairStrandsDebugInfo(
 	if (bRunDebugPass)
 	{
 		AddDebugHairPass(GraphBuilder, &View, HairDebugMode, HairData.VisibilityData, SceneTextures.Stencil, SceneColorTexture);
-		AddDebugHairPrintPass(GraphBuilder, &View, HairDebugMode, HairData.VisibilityData, HairData.MacroGroupResources, SceneTextures.Stencil);
+		AddDebugHairPrintPass(GraphBuilder, Scene, &View, HairDebugMode, HairData.VisibilityData, HairData.MacroGroupDatas, HairData.MacroGroupResources, SceneTextures.Stencil);
 	}
 	else if (HairDebugMode == EHairDebugMode::Tile && HairData.VisibilityData.TileData.IsValid())
 	{
