Fix crash when hair are rendered without any voxelization.

diff --git a/Engine/Source/Runtime/Renderer/Private/HairStrands/HairStrandsVoxelization.cpp b/Engine/Source/Runtime/Renderer/Private/HairStrands/HairStrandsVoxelization.cpp
--- a/Engine/Source/Runtime/Renderer/Private/HairStrands/HairStrandsVoxelization.cpp
+++ b/Engine/Source/Runtime/Renderer/Private/HairStrands/HairStrandsVoxelization.cpp
@@ -1124,6 +1124,7 @@ static FHairStrandsVoxelResources AllocateDummyVirtualVoxelResources(
 	Out.Parameters.Common.HairCoveragePixelRadiusAtDepth1	= ComputeMinStrandRadiusAtDepth1(FIntPoint(View.ViewRect.Width(), View.ViewRect.Height()), View.FOV, 1/*SampleCount*/, 1/*RasterizationScale*/).Primary;
 	Out.Parameters.Common.PageIndexCount = 0;
 
+	FRDGBufferRef Dummy2BytesBuffer = GSystemTextures.GetDefaultBuffer(GraphBuilder, 2, 0u);
 	FRDGBufferRef Dummy4BytesBuffer = GSystemTextures.GetDefaultBuffer(GraphBuilder, 4, 0u);
 	FRDGBufferRef Dummy8BytesBuffer = GSystemTextures.GetDefaultBuffer(GraphBuilder, 8, 0u);
 	FRDGBufferRef DummyStructBuffer = GSystemTextures.GetDefaultStructuredBuffer(GraphBuilder, sizeof(FPackedVirtualVoxelNodeDesc), 0u);
@@ -1141,6 +1142,8 @@ static FHairStrandsVoxelResources AllocateDummyVirtualVoxelResources(
 	Out.Parameters.Common.PageIndexOccupancyBuffer	= GraphBuilder.CreateSRV(Out.PageIndexOccupancyBuffer, PF_R32G32_UINT);
 	Out.Parameters.Common.PageIndexCoordBuffer		= GraphBuilder.CreateSRV(Out.PageIndexCoordBuffer, PF_R8G8B8A8_UINT);
 	Out.Parameters.Common.NodeDescBuffer			= GraphBuilder.CreateSRV(Out.NodeDescBuffer); 
+	Out.Parameters.Common.CurrGPUMinVoxelSize		= GraphBuilder.CreateSRV(Dummy2BytesBuffer, PF_R16F);
+	Out.Parameters.Common.NextGPUMinVoxelSize		= Out.Parameters.Common.CurrGPUMinVoxelSize;
 	Out.Parameters.PageTexture						= Out.PageTexture;
 	
 	if (Out.PageIndexBuffer && Out.NodeDescBuffer)
