Fix shader parameter typo.

diff --git a/Engine/Source/Runtime/Renderer/Public/ShaderPrintParameters.h b/Engine/Source/Runtime/Renderer/Public/ShaderPrintParameters.h
--- a/Engine/Source/Runtime/Renderer/Public/ShaderPrintParameters.h
+++ b/Engine/Source/Runtime/Renderer/Public/ShaderPrintParameters.h
@@ -46,7 +46,7 @@ namespace ShaderPrint
 	// ShaderPrint parameter struct declaration
 	BEGIN_SHADER_PARAMETER_STRUCT(FShaderParameters, )
 		SHADER_PARAMETER_STRUCT_REF(FUniformBufferParameters, UniformBufferParameters)
-		SHADER_PARAMETER_RDG_BUFFER_SRV(RWStructuredBuffer<uint2>, StateBuffer)
+		SHADER_PARAMETER_RDG_BUFFER_SRV(StructuredBuffer<uint2>, StateBuffer)
 		SHADER_PARAMETER_RDG_BUFFER_UAV(RWStructuredBuffer<ShaderPrintItem>, RWValuesBuffer)
 	END_SHADER_PARAMETER_STRUCT()
 
