Fixed missing SceneTextures uniform buffer for selection outlline pass

diff --git a/Engine/Source/Runtime/Renderer/Private/PostProcess/PostProcessing.cpp b/Engine/Source/Runtime/Renderer/Private/PostProcess/PostProcessing.cpp
--- a/Engine/Source/Runtime/Renderer/Private/PostProcess/PostProcessing.cpp
+++ b/Engine/Source/Runtime/Renderer/Private/PostProcess/PostProcessing.cpp
@@ -1535,6 +1535,7 @@ void AddDebugViewPostProcessingPasses(FRDGBuilder& GraphBuilder, const FViewInfo
 		PassSequence.AcceptOverrideIfLastPass(EPass::SelectionOutline, PassInputs.OverrideOutput);
 		PassInputs.SceneColor = SceneColor;
 		PassInputs.SceneDepth = SceneDepth;
+		PassInputs.SceneTextures.SceneTextures = Inputs.SceneTextures;
 
 		SceneColor = AddSelectionOutlinePass(GraphBuilder, View, PassInputs, NaniteRasterResults);
 	}
@@ -2340,6 +2341,7 @@ void AddMobilePostProcessingPasses(FRDGBuilder& GraphBuilder, FScene* Scene, con
 		PassSequence.AcceptOverrideIfLastPass(EPass::SelectionOutline, PassInputs.OverrideOutput);
 		PassInputs.SceneColor = SceneColor;
 		PassInputs.SceneDepth = SceneDepth;
+		PassInputs.SceneTextures = GetSceneTextureShaderParameters(Inputs.SceneTextures);
 		PassInputs.OverrideOutput.LoadAction = View.IsFirstInFamily() ? ERenderTargetLoadAction::EClear : ERenderTargetLoadAction::ELoad;
 
 		// TODO: Nanite - pipe through results
