REL - Added Update to the RobinHoodHashTable and used it to fix a recent Bug that would otherwise not Update the RenderData properly.

diff --git a/Engine/Source/Runtime/Renderer/Private/RendererScene.cpp b/Engine/Source/Runtime/Renderer/Private/RendererScene.cpp
--- a/Engine/Source/Runtime/Renderer/Private/RendererScene.cpp
+++ b/Engine/Source/Runtime/Renderer/Private/RendererScene.cpp
@@ -81,7 +81,7 @@
 #include "Rendering/StaticLightingSystemInterface.h"
 #endif
 
-#define VALIDATE_PRIMITIVE_PACKED_INDEX UE_BUILD_DEBUG
+#define VALIDATE_PRIMITIVE_PACKED_INDEX 0
 
 /** Affects BasePassPixelShader.usf so must relaunch editor to recompile shaders. */
 static TAutoConsoleVariable<int32> CVarEarlyZPassOnlyMaterialMasking(
@@ -1023,7 +1023,7 @@ void FScene::AddPrimitiveSceneInfo_RenderThread(FPrimitiveSceneInfo* PrimitiveSc
 	AddedPrimitiveSceneInfos.FindOrAdd(PrimitiveSceneInfo);
 	if (PreviousTransform.IsSet())
 	{
-		OverridenPreviousTransforms.FindOrAdd(PrimitiveSceneInfo, PreviousTransform.GetValue().ToMatrixWithScale());
+		OverridenPreviousTransforms.Update(PrimitiveSceneInfo, PreviousTransform.GetValue().ToMatrixWithScale());
 	}
 }
 
@@ -1428,11 +1428,11 @@ void FScene::UpdatePrimitiveTransform_RenderThread(FPrimitiveSceneProxy* Primiti
 	check(RemovedPrimitiveSceneInfos.Find(PrimitiveSceneProxy->GetPrimitiveSceneInfo()) == nullptr);
 #endif
 
-	UpdatedTransforms.FindOrAdd(PrimitiveSceneProxy, { WorldBounds, LocalBounds, LocalToWorld, AttachmentRootPosition });
+	UpdatedTransforms.Update(PrimitiveSceneProxy, { WorldBounds, LocalBounds, LocalToWorld, AttachmentRootPosition });
 
 	if (PreviousTransform.IsSet())
 	{
-		OverridenPreviousTransforms.FindOrAdd(PrimitiveSceneProxy->GetPrimitiveSceneInfo(), PreviousTransform.GetValue().ToMatrixWithScale());
+		OverridenPreviousTransforms.Update(PrimitiveSceneProxy->GetPrimitiveSceneInfo(), PreviousTransform.GetValue().ToMatrixWithScale());
 	}
 }
 
@@ -1453,7 +1453,7 @@ void FScene::UpdatePrimitiveOcclusionBoundsSlack_RenderThread(const FPrimitiveSc
 	check(RemovedPrimitiveSceneInfos.Find(PrimitiveSceneProxy->GetPrimitiveSceneInfo()) == nullptr);
 #endif
 
-	UpdatedOcclusionBoundsSlacks.FindOrAdd(PrimitiveSceneProxy, NewSlack);
+	UpdatedOcclusionBoundsSlacks.Update(PrimitiveSceneProxy, NewSlack);
 }
 
 void FScene::UpdatePrimitiveTransform(UPrimitiveComponent* Primitive)
@@ -1628,7 +1628,7 @@ void FScene::UpdatePrimitiveInstances(UInstancedStaticMeshComponent* Primitive)
 				check(RemovedPrimitiveSceneInfos.Find(UpdateParams.PrimitiveSceneProxy->GetPrimitiveSceneInfo()) == nullptr);
 #endif
 				FScopeCycleCounter Context(UpdateParams.PrimitiveSceneProxy->GetStatId());
-				UpdatedInstances.FindOrAdd(UpdateParams.PrimitiveSceneProxy, UpdateParams);
+				UpdatedInstances.Update(UpdateParams.PrimitiveSceneProxy, UpdateParams);
 			}
 		);
 	}
@@ -1673,7 +1673,7 @@ void FScene::UpdatePrimitiveLightingAttachmentRoot(UPrimitiveComponent* Primitiv
 			[Scene, Proxy, NewComponentId](FRHICommandList&)
 			{
 				FPrimitiveSceneInfo* PrimitiveInfo = Proxy->GetPrimitiveSceneInfo();
-				Scene->UpdatedAttachmentRoots.FindOrAdd(PrimitiveInfo, NewComponentId);
+				Scene->UpdatedAttachmentRoots.Update(PrimitiveInfo, NewComponentId);
 			});
 	}
 }
@@ -1727,7 +1727,7 @@ void FScene::UpdateCustomPrimitiveData(UPrimitiveComponent* Primitive)
 		ENQUEUE_RENDER_COMMAND(UpdateCustomPrimitiveDataCommand)(
 			[UpdateParams](FRHICommandListImmediate& RHICmdList)
 			{
-				UpdateParams.Scene->UpdatedCustomPrimitiveParams.FindOrAdd(UpdateParams.PrimitiveSceneProxy, UpdateParams.CustomPrimitiveData);
+				UpdateParams.Scene->UpdatedCustomPrimitiveParams.Update(UpdateParams.PrimitiveSceneProxy, UpdateParams.CustomPrimitiveData);
 			});
 	}
 }
@@ -5438,20 +5438,20 @@ void UpdateStaticMeshesForMaterials(const TArray<const FMaterial*>& MaterialReso
 			}
 
 			if (UsedMaterialsDependencies.Num() > 0)
-					{
+			{
 				for (const FMaterial* MaterialResourceToUpdate : MaterialResourcesToUpdate)
-						{
+				{
 					UMaterialInterface* UpdatedMaterialInterface = MaterialResourceToUpdate->GetMaterialInterface();
 
 					if (UpdatedMaterialInterface)
-							{
+					{
 						if (UsedMaterialsDependencies.Contains(UpdatedMaterialInterface))
-				{
-					FPrimitiveSceneProxy* SceneProxy = PrimitiveComponent->SceneProxy;
-					FPrimitiveSceneInfo* SceneInfo = SceneProxy->GetPrimitiveSceneInfo();
-					FScene* Scene = SceneInfo->Scene;
-					TArray<FPrimitiveSceneInfo*>& SceneInfos = UsedPrimitives.FindOrAdd(Scene);
-					SceneInfos.Add(SceneInfo);
+						{
+							FPrimitiveSceneProxy* SceneProxy = PrimitiveComponent->SceneProxy;
+							FPrimitiveSceneInfo* SceneInfo = SceneProxy->GetPrimitiveSceneInfo();
+							FScene* Scene = SceneInfo->Scene;
+							TArray<FPrimitiveSceneInfo*>& SceneInfos = UsedPrimitives.FindOrAdd(Scene);
+							SceneInfos.Add(SceneInfo);
 							break;
 						}
 					}
