Strata - new blending mode.
- works for starta and legacy conversion material
- legacy converted to strata setup the strata blending mode
- strata material also affect legacy blending mode because it is used in multiple places in the engine.
- works for lit and unlit legacy material
- to be fixed:
    - lit with premultiplied alpha
    - materials with multiple shading models

diff --git a/Engine/Shaders/Private/BasePassPixelShader.usf b/Engine/Shaders/Private/BasePassPixelShader.usf
--- a/Engine/Shaders/Private/BasePassPixelShader.usf
+++ b/Engine/Shaders/Private/BasePassPixelShader.usf
@@ -130,7 +130,7 @@ float NormalCurvatureToRoughness(float3 WorldNormal)
 #define FORCE_FULLY_ROUGH (SIMPLE_FORWARD_SHADING || MATERIAL_FULLY_ROUGH)
 #define EDITOR_ALPHA2COVERAGE (USE_EDITOR_COMPOSITING && SUPPORTS_PIXEL_COVERAGE)
 #define POST_PROCESS_SUBSURFACE ((MATERIAL_SHADINGMODEL_SUBSURFACE_PROFILE || MATERIAL_SHADINGMODEL_EYE) && USES_GBUFFER)
-#define MATERIAL_WORKS_WITH_DUAL_SOURCE_COLOR_BLENDING (MATERIAL_SHADINGMODEL_THIN_TRANSLUCENT || STRATA_TRANSLUCENT_MATERIAL )
+#define MATERIAL_WORKS_WITH_DUAL_SOURCE_COLOR_BLENDING (MATERIAL_SHADINGMODEL_THIN_TRANSLUCENT || STRATA_BLENDING_TRANSLUCENT_COLOREDTRANSMITTANCE )
 
 #if  MATERIAL_SHADINGMODEL_SINGLELAYERWATER
 
@@ -1582,7 +1582,7 @@ void FPixelShaderInOut_MainPS(
 	#endif
 #endif
 
-#if MATERIAL_WORKS_WITH_DUAL_SOURCE_COLOR_BLENDING
+#if MATERIAL_WORKS_WITH_DUAL_SOURCE_COLOR_BLENDING || STRATA_TRANSLUCENT_MATERIAL
 	float3 DualBlendSurfaceLuminancePostCoverage	= 0.0f;
 	float3 DualBlendSurfaceTransmittancePreCoverage = 1.0f;
 	float DualBlendSurfaceCoverage					= 1.0f;
@@ -1799,6 +1799,7 @@ void FPixelShaderInOut_MainPS(
 			Out.MRT[0] = half4(Color, 0);
 		#endif
 		Out.MRT[0] = RETURN_COLOR(Out.MRT[0]);
+
 	// MATERIAL_SHADINGMODEL_THIN_TRANSLUCENT must come first because it also has MATERIALBLENDING_TRANSLUCENT defined
 	#elif MATERIAL_WORKS_WITH_DUAL_SOURCE_COLOR_BLENDING
 
@@ -1817,6 +1818,32 @@ void FPixelShaderInOut_MainPS(
 			Out.MRT[0] = half4(AdjustedDualBlendAdd,AdjustedAlpha);
 			Out.MRT[0] = RETURN_COLOR(Out.MRT[0]);
 		#endif
+
+	#elif STRATA_TRANSLUCENT_MATERIAL
+
+			// Add fog luminance according to surfacecoverage and reduce surface luminance according to fog coverage.
+			float3 AdjustedDualBlendAdd = DualBlendSurfaceCoverage * Fogging.rgb + Fogging.a * DualBlendSurfaceLuminancePostCoverage;
+			// Fade the surface color transmittance out to 1 according to the surface coverage, and take into account the fog coverage to the surface.
+			float3 AdjustedDualBlendMul = lerp(1.0f, Fogging.a * DualBlendSurfaceTransmittancePreCoverage, DualBlendSurfaceCoverage);
+
+		#if STRATA_BLENDING_COLOREDTRANSMITTANCEONLY
+			// RETURN_COLOR not needed with modulative blending
+			half3 FoggedColor = lerp(float3(1, 1, 1), DualBlendSurfaceTransmittancePreCoverage, Fogging.aaa * DualBlendSurfaceCoverage);
+			Out.MRT[0] = half4(FoggedColor, 1.0f);
+		#elif STRATA_BLENDING_ALPHAHOLDOUT
+			Out.MRT[0] = half4(0.0f.xxx, Opacity);
+			Out.MRT[0] = RETURN_COLOR(Out.MRT[0]);
+		#else
+			#if STRATA_USES_CONVERSION_FROM_LEGACY && MATERIALBLENDING_ADDITIVE
+				float AdjustedAlpha = 0.0f;
+			#else
+				float AdjustedAlpha = saturate(1 - dot(AdjustedDualBlendMul, float3(1.0f, 1.0f, 1.0f) / 3.0f));
+			#endif
+			// Pre multipled alpha blending
+			Out.MRT[0] = half4(AdjustedDualBlendAdd, AdjustedAlpha);
+			Out.MRT[0] = RETURN_COLOR(Out.MRT[0]);
+		#endif
+
 	#elif MATERIALBLENDING_ALPHAHOLDOUT
 		// not implemented for holdout
 		Out.MRT[0] = half4(Color * Fogging.a + Fogging.rgb * Opacity, Opacity);
diff --git a/Engine/Shaders/Private/Definitions.usf b/Engine/Shaders/Private/Definitions.usf
--- a/Engine/Shaders/Private/Definitions.usf
+++ b/Engine/Shaders/Private/Definitions.usf
@@ -46,6 +46,30 @@
 #define MATERIALBLENDING_ALPHAHOLDOUT					0
 #endif
 
+#ifndef STRATA_BLENDING_OPAQUE
+#define STRATA_BLENDING_OPAQUE							0
+#endif
+
+#ifndef STRATA_BLENDING_MASKED
+#define STRATA_BLENDING_MASKED							0
+#endif
+
+#ifndef STRATA_BLENDING_TRANSLUCENT_GREYTRANSMITTANCE
+#define STRATA_BLENDING_TRANSLUCENT_GREYTRANSMITTANCE	0
+#endif
+
+#ifndef STRATA_BLENDING_TRANSLUCENT_COLOREDTRANSMITTANCE
+#define STRATA_BLENDING_TRANSLUCENT_COLOREDTRANSMITTANCE 0
+#endif
+
+#ifndef STRATA_BLENDING_COLOREDTRANSMITTANCEONLY
+#define STRATA_BLENDING_COLOREDTRANSMITTANCEONLY		0
+#endif
+
+#ifndef STRATA_BLENDING_ALPHAHOLDOUT
+#define STRATA_BLENDING_ALPHAHOLDOUT					0
+#endif
+
 #ifndef MATERIAL_SHADINGMODEL_DEFAULT_LIT
 #define MATERIAL_SHADINGMODEL_DEFAULT_LIT				0
 #endif
diff --git a/Engine/Shaders/Private/MaterialTemplate.ush b/Engine/Shaders/Private/MaterialTemplate.ush
--- a/Engine/Shaders/Private/MaterialTemplate.ush
+++ b/Engine/Shaders/Private/MaterialTemplate.ush
@@ -122,8 +122,8 @@
 //Materials also have to opt in to these features.
 #define USE_EDITOR_COMPOSITING (USE_EDITOR_SHADERS && EDITOR_PRIMITIVE_MATERIAL)
 
-#define MATERIALBLENDING_ANY_TRANSLUCENT (MATERIALBLENDING_TRANSLUCENT || MATERIALBLENDING_ADDITIVE || MATERIALBLENDING_MODULATE)
-
+#define MATERIALBLENDING_ANY_TRANSLUCENT (MATERIALBLENDING_TRANSLUCENT || MATERIALBLENDING_ADDITIVE || MATERIALBLENDING_MODULATE || STRATA_BLENDING_TRANSLUCENT_GREYTRANSMITTANCE || STRATA_BLENDING_TRANSLUCENT_COLOREDTRANSMITTANCE || STRATA_BLENDING_COLOREDTRANSMITTANCEONLY)
+ 
 #define IS_MATERIAL_TRANSLUCENT_AND_LIT  (MATERIALBLENDING_TRANSLUCENT)
 
 #define IS_MESHPARTICLE_FACTORY (PARTICLE_MESH_FACTORY || NIAGARA_MESH_FACTORY)
diff --git a/Engine/Shaders/Private/Strata/StrataLegacyConversion.ush b/Engine/Shaders/Private/Strata/StrataLegacyConversion.ush
--- a/Engine/Shaders/Private/Strata/StrataLegacyConversion.ush
+++ b/Engine/Shaders/Private/Strata/StrataLegacyConversion.ush
@@ -495,8 +495,32 @@ FStrataData StrataConvertLegacyMaterialStatic(
 
 	if (ShadingModel == SHADINGMODELID_UNLIT)
 	{
-		Out = GetStrataUnlitBSDF(Emissive, Zeros);
-//		Out = StrataWeight(Out, Opacity);				// STRATA_TODO
+		const float GreyTransmittance = saturate(1.0f - Opacity);
+	#if MATERIALBLENDING_OPAQUE || MATERIALBLENDING_MASKED
+		// Uses opaque blending
+		Out = GetStrataUnlitBSDF(Emissive, 0.0f);
+	#elif MATERIALBLENDING_ALPHACOMPOSITE
+		// Uses (Luminance, grey transmittance) blending
+		Out = GetStrataUnlitBSDF(Emissive, 0.0f);
+		#if STRATA_INLINE_SHADING
+			Out.InlinedBSDF.Coverage = Opacity;
+		#endif
+	#elif MATERIALBLENDING_TRANSLUCENT
+		// Uses (Luminance, grey transmittance) blending
+		Out = GetStrataUnlitBSDF(Emissive * Opacity, 0.0f);
+		#if STRATA_INLINE_SHADING
+			Out.InlinedBSDF.Coverage = Opacity;
+		#endif
+	#elif MATERIALBLENDING_ALPHAHOLDOUT
+		Out = GetStrataUnlitBSDF(0.0f, 1.0f);
+	#elif MATERIALBLENDING_ADDITIVE
+		Out = GetStrataUnlitBSDF(Emissive * Opacity, 1.0f);
+	#elif MATERIALBLENDING_MODULATE
+		Out = GetStrataUnlitBSDF(0.0f, Emissive);
+	#else
+		// Opaque and masked
+		Out = GetStrataUnlitBSDF(Emissive, 0.0f);
+	#endif
 	}
 	else if (ShadingModel == SHADINGMODELID_HAIR)
 	{
diff --git a/Engine/Source/Developer/MaterialBaking/Private/ExportMaterialProxy.h b/Engine/Source/Developer/MaterialBaking/Private/ExportMaterialProxy.h
--- a/Engine/Source/Developer/MaterialBaking/Private/ExportMaterialProxy.h
+++ b/Engine/Source/Developer/MaterialBaking/Private/ExportMaterialProxy.h
@@ -500,6 +500,7 @@ public:
 	}
 	virtual bool IsMasked() const override { return false; }
 	virtual enum EBlendMode GetBlendMode() const override { return BLEND_Opaque; }
+	virtual enum EStrataBlendMode GetStrataBlendMode() const override { return EStrataBlendMode::SBM_Opaque; }
 	virtual FMaterialShadingModelField GetShadingModels() const override { return MSM_DefaultLit; }
 	virtual bool IsShadingModelFromMaterialExpression() const override { return false; }
 	virtual float GetOpacityMaskClipValue() const override { return 0.5f; }
diff --git a/Engine/Source/Developer/MaterialUtilities/Private/MaterialUtilities.cpp b/Engine/Source/Developer/MaterialUtilities/Private/MaterialUtilities.cpp
--- a/Engine/Source/Developer/MaterialUtilities/Private/MaterialUtilities.cpp
+++ b/Engine/Source/Developer/MaterialUtilities/Private/MaterialUtilities.cpp
@@ -681,6 +681,7 @@ public:
 	}
 	virtual bool IsMasked() const override									{ return false; }
 	virtual enum EBlendMode GetBlendMode() const override					{ return BLEND_Opaque; }
+	virtual enum EStrataBlendMode GetStrataBlendMode() const override		{ return EStrataBlendMode::SBM_Opaque; }
 	virtual FMaterialShadingModelField GetShadingModels() const override	{ return MSM_Unlit; }
 	virtual bool IsShadingModelFromMaterialExpression() const override		{ return false; }
 	virtual float GetOpacityMaskClipValue() const override					{ return 0.5f; }
diff --git a/Engine/Source/Editor/MaterialEditor/Private/MaterialEditor.h b/Engine/Source/Editor/MaterialEditor/Private/MaterialEditor.h
--- a/Engine/Source/Editor/MaterialEditor/Private/MaterialEditor.h
+++ b/Engine/Source/Editor/MaterialEditor/Private/MaterialEditor.h
@@ -153,6 +153,7 @@ public:
 	virtual bool IsWireframe() const override { return false; }
 	virtual bool IsMasked() const override { return false; }
 	virtual enum EBlendMode GetBlendMode() const override { return BLEND_Translucent; }
+	virtual enum EStrataBlendMode GetStrataBlendMode() const override { return EStrataBlendMode::SBM_TranslucentGreyTransmittance; }
 	virtual FMaterialShadingModelField GetShadingModels() const override { return MSM_Unlit; }
 	virtual bool IsShadingModelFromMaterialExpression() const override { return false; }
 	virtual float GetOpacityMaskClipValue() const override { return 0.5f; }
diff --git a/Engine/Source/Editor/UnrealEd/Private/Lightmass/LightmassRender.cpp b/Engine/Source/Editor/UnrealEd/Private/Lightmass/LightmassRender.cpp
--- a/Engine/Source/Editor/UnrealEd/Private/Lightmass/LightmassRender.cpp
+++ b/Engine/Source/Editor/UnrealEd/Private/Lightmass/LightmassRender.cpp
@@ -532,6 +532,7 @@ public:
 	}
 	virtual bool IsMasked() const override									{ return false; }
 	virtual enum EBlendMode GetBlendMode() const override					{ return BLEND_Opaque; }
+	virtual enum EStrataBlendMode GetStrataBlendMode() const override		{ return EStrataBlendMode::SBM_Opaque; }
 	virtual FMaterialShadingModelField GetShadingModels() const override	{ return MSM_Unlit; }
 	virtual bool IsShadingModelFromMaterialExpression() const override		{ return false; }
 	virtual float GetOpacityMaskClipValue() const override					{ return 0.5f; }
diff --git a/Engine/Source/Runtime/Engine/Classes/Engine/EngineTypes.h b/Engine/Source/Runtime/Engine/Classes/Engine/EngineTypes.h
--- a/Engine/Source/Runtime/Engine/Classes/Engine/EngineTypes.h
+++ b/Engine/Source/Runtime/Engine/Classes/Engine/EngineTypes.h
@@ -250,6 +250,21 @@ enum EBlendMode
 	BLEND_MAX,
 };
 
+/**
+ * The blending mode for Strata materials
+ */
+UENUM(BlueprintType)
+enum EStrataBlendMode
+{
+	SBM_Opaque UMETA(DisplayName = "Opaque"),
+	SBM_Masked UMETA(DisplayName = "Masked"),
+	SBM_TranslucentGreyTransmittance UMETA(DisplayName = "Translucent - Grey Transmittance"),
+	SBM_TranslucentColoredTransmittance UMETA(DisplayName = "Translucent - Colored Transmittance"),
+	SBM_ColoredTransmittanceOnly UMETA(DisplayName = "Colored Transmittance Only"),
+	SBM_AlphaHoldout UMETA(DisplayName = "Alpha Holdout"),
+	SBM_MAX,
+};
+
 /** The default float precision for material's pixel shaders on mobile devices*/
 UENUM()
 enum EMaterialFloatPrecisionMode
diff --git a/Engine/Source/Runtime/Engine/Classes/Materials/Material.h b/Engine/Source/Runtime/Engine/Classes/Materials/Material.h
--- a/Engine/Source/Runtime/Engine/Classes/Materials/Material.h
+++ b/Engine/Source/Runtime/Engine/Classes/Materials/Material.h
@@ -389,6 +389,10 @@ class UMaterial : public UMaterialInterface
 	UPROPERTY(BlueprintReadOnly, EditAnywhere, Category=Material, AssetRegistrySearchable)
 	TEnumAsByte<enum EBlendMode> BlendMode;
 
+	/** Determines how the material's color is blended with background colors. */
+	UPROPERTY(BlueprintReadOnly, EditAnywhere, Category = Material, AssetRegistrySearchable)
+	TEnumAsByte<enum EStrataBlendMode> StrataBlendMode;
+
 	UPROPERTY(meta = (DeprecatedProperty, DeprecationMessage = "No longer used."))
 	TEnumAsByte<enum EDecalBlendMode> DecalBlendMode;
 
@@ -1065,6 +1069,7 @@ public:
 	ENGINE_API virtual float GetOpacityMaskClipValue() const override;
 	ENGINE_API virtual bool GetCastDynamicShadowAsMasked() const override;
 	ENGINE_API virtual EBlendMode GetBlendMode() const override;
+	ENGINE_API virtual EStrataBlendMode GetStrataBlendMode() const override;
 	ENGINE_API virtual FMaterialShadingModelField GetShadingModels() const override;
 	ENGINE_API virtual bool IsShadingModelFromMaterialExpression() const override;
 	ENGINE_API virtual bool IsTwoSided() const override;
diff --git a/Engine/Source/Runtime/Engine/Classes/Materials/MaterialInterface.h b/Engine/Source/Runtime/Engine/Classes/Materials/MaterialInterface.h
--- a/Engine/Source/Runtime/Engine/Classes/Materials/MaterialInterface.h
+++ b/Engine/Source/Runtime/Engine/Classes/Materials/MaterialInterface.h
@@ -732,6 +732,7 @@ public:
 	ENGINE_API virtual float GetOpacityMaskClipValue() const;
 	ENGINE_API virtual bool GetCastDynamicShadowAsMasked() const;
 	ENGINE_API virtual EBlendMode GetBlendMode() const;
+	ENGINE_API virtual EStrataBlendMode GetStrataBlendMode() const;
 	ENGINE_API virtual FMaterialShadingModelField GetShadingModels() const;
 	ENGINE_API virtual bool IsShadingModelFromMaterialExpression() const;
 	ENGINE_API virtual bool IsTwoSided() const;
diff --git a/Engine/Source/Runtime/Engine/Private/Materials/HLSLMaterialTranslator.cpp b/Engine/Source/Runtime/Engine/Private/Materials/HLSLMaterialTranslator.cpp
--- a/Engine/Source/Runtime/Engine/Private/Materials/HLSLMaterialTranslator.cpp
+++ b/Engine/Source/Runtime/Engine/Private/Materials/HLSLMaterialTranslator.cpp
@@ -1147,14 +1147,6 @@ bool FHLSLMaterialTranslator::Translate()
 			}
 		}
 
-		if (Material->IsStrataMaterial())
-		{
-			if (Domain!= MD_Volume && !(BlendMode == BLEND_Translucent || BlendMode == BLEND_Opaque || BlendMode == BLEND_Masked))
-		{
-				Errorf(TEXT("Strata materials must be an opaque or translucent surface, or a volume."));
-			}
-		}
-
 		if (Domain == MD_DeferredDecal && !(BlendMode == BLEND_Translucent || BlendMode == BLEND_AlphaComposite || BlendMode == BLEND_Modulate))
 		{
 			// We could make the change for the user but it would be confusing when going to DeferredDecal and back
@@ -1865,9 +1857,8 @@ void FHLSLMaterialTranslator::GetMaterialEnvironment(EShaderPlatform InPlatform,
 
 	OutEnvironment.SetDefine(TEXT("MATERIAL_IS_STRATA"), bMaterialIsStrata ? TEXT("1") : TEXT("0"));
 
-	// STRATA_TODO do not request DualSourceBlending if gray scale transmittance is selected.
-	// bMaterialRequestsDualSourceBlending this base on limited set of blend mode: Opaque, Masked, TransmittanceCoverage, TransmittanceColored;
-	bMaterialRequestsDualSourceBlending |= bMaterialIsStrata;
+	// Strata requests dual source blending only for SBM_TranslucentColoredTransmittance
+	bMaterialRequestsDualSourceBlending |= bMaterialIsStrata && Material->GetStrataBlendMode() == EStrataBlendMode::SBM_TranslucentColoredTransmittance;
 
 	// if duals source blending (colored transmittance) is not supported on a platform, it will fall back to standard alpha blending (grey scale transmittance)
 	OutEnvironment.SetDefine(TEXT("DUAL_SOURCE_COLOR_BLENDING_ENABLED"), bMaterialRequestsDualSourceBlending && Material->IsDualBlendingEnabled(Platform) ? TEXT("1") : TEXT("0"));
diff --git a/Engine/Source/Runtime/Engine/Private/Materials/Material.cpp b/Engine/Source/Runtime/Engine/Private/Materials/Material.cpp
--- a/Engine/Source/Runtime/Engine/Private/Materials/Material.cpp
+++ b/Engine/Source/Runtime/Engine/Private/Materials/Material.cpp
@@ -2798,6 +2798,39 @@ void UMaterial::ConvertMaterialToStrataMaterial()
 		}
 	};
 
+	auto ConvertLegacyToStrataBlendMode = [&]()
+	{
+		if (ShadingModels.CountShadingModels() == 1 && ShadingModels.GetFirstShadingModel() == EMaterialShadingModel::MSM_ThinTranslucent)
+		{
+			StrataBlendMode = SBM_TranslucentColoredTransmittance;
+		}
+		else
+		{
+			switch (BlendMode)
+			{
+			case BLEND_Opaque:
+				StrataBlendMode = SBM_Opaque;
+				break;
+			case BLEND_Masked:
+				StrataBlendMode = SBM_Masked;
+				break;
+			case BLEND_Translucent:
+			case BLEND_AlphaComposite:
+			case BLEND_Additive:
+				StrataBlendMode = SBM_TranslucentGreyTransmittance;
+				break;
+			case BLEND_Modulate:
+				StrataBlendMode = SBM_ColoredTransmittanceOnly;
+				break;
+			case BLEND_AlphaHoldout:
+				StrataBlendMode = SBM_AlphaHoldout;
+				break;
+			default:
+				UE_LOG(LogMaterial, Error, TEXT("%s: Material blend mode could not be converted to Starta."), *GetName());
+			}
+		}
+	};
+
 	// SSS Profile
 	const bool bHasShadingModelMixture		= ShadingModels.CountShadingModels() > 1;
 	const bool bRequireSubsurfacePasses		= ShadingModels.HasShadingModel(MSM_SubsurfaceProfile) || ShadingModels.HasShadingModel(MSM_Subsurface) || ShadingModels.HasShadingModel(MSM_PreintegratedSkin) || ShadingModels.HasShadingModel(MSM_Eye);
@@ -2861,6 +2894,7 @@ void UMaterial::ConvertMaterialToStrataMaterial()
 			check(ConvertNode->ConvertedStrataMaterialInfo.CountShadingModels() == 1);
 		}
 
+		ConvertLegacyToStrataBlendMode();
 		bInvalidateShader = true;
 	}
 	else if (!bUseMaterialAttributes && !FrontMaterial.IsConnected())
@@ -3051,6 +3085,8 @@ void UMaterial::ConvertMaterialToStrataMaterial()
 			FrontMaterial.Connect(0, ConvertNode);
 			bInvalidateShader = true;
 		}
+
+		ConvertLegacyToStrataBlendMode();
 	}
 
 	if (bInvalidateShader)
@@ -3588,12 +3624,17 @@ bool UMaterial::CanEditChange(const FProperty* InProperty) const
 
 		if (PropertyName == GET_MEMBER_NAME_STRING_CHECKED(UMaterial, BlendMode))
 		{
-			return MaterialDomain == MD_DeferredDecal || MaterialDomain == MD_Surface || MaterialDomain == MD_Volume || MaterialDomain == MD_UI || (MaterialDomain == MD_PostProcess && BlendableOutputAlpha);
+			return !Engine_IsStrataEnabled() && (MaterialDomain == MD_DeferredDecal || MaterialDomain == MD_Surface || MaterialDomain == MD_Volume || MaterialDomain == MD_UI || (MaterialDomain == MD_PostProcess && BlendableOutputAlpha));
+		}
+
+		if (PropertyName == GET_MEMBER_NAME_STRING_CHECKED(UMaterial, StrataBlendMode))
+		{
+			return Engine_IsStrataEnabled();
 		}
 	
 		if (PropertyName == GET_MEMBER_NAME_STRING_CHECKED(UMaterial, ShadingModel))
 		{
-			return MaterialDomain == MD_Surface;
+			return !Engine_IsStrataEnabled() && MaterialDomain == MD_Surface;
 		}
 
 		if (FCString::Strncmp(*PropertyName, TEXT("bUsedWith"), 9) == 0)
@@ -3740,6 +3781,29 @@ void UMaterial::PostEditChangePropertyInternal(FPropertyChangedEvent& PropertyCh
 		}
 	}
 
+	// If the strata blending mode is changed, make sure we update the legacy blend mode too for the shaders to be filtered to passes correctly at runtime.
+	if (Engine_IsStrataEnabled() && PropertyChangedEvent.GetPropertyName() == GET_MEMBER_NAME_CHECKED(UMaterial, StrataBlendMode))
+	{
+		switch (GetStrataBlendMode())
+		{
+		case SBM_Opaque:
+			BlendMode = BLEND_Opaque;
+			break;
+		case SBM_Masked:
+			BlendMode = BLEND_Masked;
+			break;
+		case SBM_TranslucentGreyTransmittance:
+		case SBM_TranslucentColoredTransmittance:
+		case SBM_ColoredTransmittanceOnly:
+			BlendMode = BLEND_Translucent;
+			break;
+		case SBM_AlphaHoldout:
+			BlendMode = BLEND_AlphaHoldout;
+			break;
+		}
+	}
+
+
 	//If we can be sure this material would be the same opaque as it is masked then allow it to be assumed opaque.
 	bCanMaskedBeAssumedOpaque = !OpacityMask.Expression && !(OpacityMask.UseConstant && OpacityMask.Constant < 0.999f) && !bUseMaterialAttributes;
 
@@ -4043,13 +4107,6 @@ void UMaterial::RebuildShadingModelField()
 				}
 			}
 
-			// Blend mode
-			// Unclear what best fallback it should be
-			if (BlendMode != EBlendMode::BLEND_Opaque && BlendMode != EBlendMode::BLEND_Masked)
-			{
-				BlendMode = EBlendMode::BLEND_Translucent; // This is to be able to use dual-source blending
-			}
-
 			// Subsurface profil
 			if (StrataMaterialInfo.HasShadingModel(SSM_SubsurfaceLit) && StrataMaterialInfo.CountSubsurfaceProfiles() > 0)
 			{
@@ -4070,10 +4127,6 @@ void UMaterial::RebuildShadingModelField()
 					MaterialDomain = EMaterialDomain::MD_Surface;
 				}
 				ShadingModel = MSM_Unlit;
-				if (BlendMode != EBlendMode::BLEND_Opaque && BlendMode != EBlendMode::BLEND_Masked)
-				{
-					BlendMode = EBlendMode::BLEND_Translucent; // This is to be able to use dual-source blending
-				}
 			}
 			else if (StrataMaterialInfo.HasOnlyShadingModel(SSM_SubsurfaceLit))
 			{
@@ -4092,10 +4145,6 @@ void UMaterial::RebuildShadingModelField()
 					MaterialDomain = EMaterialDomain::MD_Surface;
 				}
 				ShadingModel = MSM_DefaultLit;
-				if (BlendMode != EBlendMode::BLEND_Opaque && BlendMode != EBlendMode::BLEND_Masked)
-				{
-					BlendMode = EBlendMode::BLEND_Translucent; // This is to be able to use dual-source blending
-				}
 			}
 			else if (StrataMaterialInfo.HasOnlyShadingModel(SSM_VolumetricFogCloud))
 			{
@@ -5315,6 +5364,25 @@ EBlendMode UMaterial::GetBlendMode() const
 	}
 }
 
+EStrataBlendMode UMaterial::GetStrataBlendMode() const
+{
+	if (EStrataBlendMode(StrataBlendMode) == SBM_Masked)
+	{
+		if (bCanMaskedBeAssumedOpaque)
+		{
+			return SBM_Opaque;
+		}
+		else
+		{
+			return SBM_Masked;
+		}
+	}
+	else
+	{
+		return StrataBlendMode;
+	}
+}
+
 FMaterialShadingModelField UMaterial::GetShadingModels() const
 {
 	switch (MaterialDomain)
diff --git a/Engine/Source/Runtime/Engine/Private/Materials/MaterialInterface.cpp b/Engine/Source/Runtime/Engine/Private/Materials/MaterialInterface.cpp
--- a/Engine/Source/Runtime/Engine/Private/Materials/MaterialInterface.cpp
+++ b/Engine/Source/Runtime/Engine/Private/Materials/MaterialInterface.cpp
@@ -282,6 +282,7 @@ FMaterialRelevance UMaterialInterface::GetRelevance_Internal(const UMaterial* Ma
 			MaterialResource->MaterialUsesAnisotropy_GameThread();
 
 		const EBlendMode BlendMode = (EBlendMode)GetBlendMode();
+		const EStrataBlendMode StrataBlendMode = (EStrataBlendMode)GetStrataBlendMode();
 		const bool bIsTranslucent = IsTranslucentBlendMode(BlendMode) || IsSinglePassWaterTranslucent || bIsMobilePixelProjectedTranslucent; // We want meshes with water materials to be scheduled for translucent pass on mobile. And we also have to render the meshes used for mobile pixel projection reflection in translucent pass.
 
 		EMaterialDomain Domain = (EMaterialDomain)MaterialResource->GetMaterialDomain();
@@ -328,8 +329,8 @@ FMaterialRelevance UMaterialInterface::GetRelevance_Internal(const UMaterial* Ma
 			MaterialRelevance.bMasked = IsMasked();
 			MaterialRelevance.bDistortion = MaterialResource->IsDistorted();
 			MaterialRelevance.bHairStrands = IsCompatibleWithHairStrands(MaterialResource, InFeatureLevel);
-			MaterialRelevance.bSeparateTranslucency = (TranslucencyPass == MTP_AfterDOF);
-			MaterialRelevance.bSeparateTranslucencyModulate = bMaterialSeparateModulation;
+			MaterialRelevance.bSeparateTranslucency = (TranslucencyPass == MTP_AfterDOF) && StrataBlendMode != SBM_ColoredTransmittanceOnly;
+			MaterialRelevance.bSeparateTranslucencyModulate = bMaterialSeparateModulation || StrataBlendMode == SBM_ColoredTransmittanceOnly;
 			MaterialRelevance.bPostMotionBlurTranslucency = (TranslucencyPass == MTP_AfterMotionBlur);
 			MaterialRelevance.bNormalTranslucency = bIsTranslucent && (TranslucencyPass == MTP_BeforeDOF);
 			MaterialRelevance.bDisableDepthTest = bIsTranslucent && Material->bDisableDepthTest;		
@@ -942,6 +943,11 @@ EBlendMode UMaterialInterface::GetBlendMode() const
 	return BLEND_Opaque;
 }
 
+EStrataBlendMode UMaterialInterface::GetStrataBlendMode() const
+{
+	return SBM_Opaque;
+}
+
 bool UMaterialInterface::IsTwoSided() const
 {
 	return false;
diff --git a/Engine/Source/Runtime/Engine/Private/Materials/MaterialShared.cpp b/Engine/Source/Runtime/Engine/Private/Materials/MaterialShared.cpp
--- a/Engine/Source/Runtime/Engine/Private/Materials/MaterialShared.cpp
+++ b/Engine/Source/Runtime/Engine/Private/Materials/MaterialShared.cpp
@@ -83,6 +83,8 @@ FAutoConsoleVariableRef CVarDeferUniformExpressionCaching(
 	ECVF_RenderThreadSafe
 	);
 
+bool Engine_IsStrataEnabled();
+
 struct FAllowCachingStaticParameterValues
 {
 	FAllowCachingStaticParameterValues(FMaterial& InMaterial)
@@ -1461,7 +1463,12 @@ bool FMaterialResource::IsTranslucencyAfterMotionBlurEnabled() const
 
 bool FMaterialResource::IsDualBlendingEnabled(EShaderPlatform Platform) const
 {
-	const bool bMaterialRequestsDualSourceBlending = Material->ShadingModel == MSM_ThinTranslucent || IsStrataMaterial();
+	bool bMaterialRequestsDualSourceBlending = Material->ShadingModel == MSM_ThinTranslucent;
+	if (IsStrataMaterial())
+	{
+		EStrataBlendMode StrataBlendMode = GetStrataBlendMode();
+		bMaterialRequestsDualSourceBlending = StrataBlendMode == EStrataBlendMode::SBM_TranslucentColoredTransmittance;
+	}
 	const bool bIsPlatformSupported = RHISupportsDualSourceBlending(Platform);
 	return bMaterialRequestsDualSourceBlending && bIsPlatformSupported;
 }
@@ -1558,6 +1565,11 @@ EBlendMode FMaterialResource::GetBlendMode() const
 	return MaterialInstance ? MaterialInstance->GetBlendMode() : Material->GetBlendMode();
 }
 
+EStrataBlendMode FMaterialResource::GetStrataBlendMode() const
+{
+	return MaterialInstance ? MaterialInstance->GetStrataBlendMode() : Material->GetStrataBlendMode();
+}
+
 ERefractionMode FMaterialResource::GetRefractionMode() const
 {
 	return Material->RefractionMode;
@@ -1678,16 +1690,11 @@ bool FMaterialResource::HasAmbientOcclusionConnected() const
 
 bool FMaterialResource::IsStrataMaterial() const
 {
-	// STRATA_TODO IsStrataMaterial should go away once Strata implementation is finished
-	static const auto CVar = IConsoleManager::Get().FindTConsoleVariableDataInt(TEXT("r.Strata"));
-	const bool bStrataEnabled = CVar && CVar->GetValueOnAnyThread() > 0;
 	// We no longer support both types of material (Strata and non strata) so no need to check if FrontMaterial is plugged in.
 	// We simply consider all material as Strata when Strata is enabled.
-	return bStrataEnabled;
+	return Engine_IsStrataEnabled();
 }
 
-bool Engine_IsStrataEnabled();
-
 bool FMaterialResource::HasMaterialPropertyConnected(EMaterialProperty In) const
 {
 	// STRATA_TODO: temporary validation until we have converted all domains
@@ -2076,6 +2083,33 @@ void FMaterial::SetupMaterialEnvironment(
 		OutEnvironment.SetDefine(TEXT("MATERIALBLENDING_SOLID"),TEXT("1"));
 	}
 
+	if (Engine_IsStrataEnabled())
+	{
+		switch (GetStrataBlendMode())
+		{
+		case SBM_Opaque:
+			OutEnvironment.SetDefine(TEXT("STRATA_BLENDING_OPAQUE"), TEXT("1"));
+			break;
+		case SBM_Masked:
+			OutEnvironment.SetDefine(TEXT("STRATA_BLENDING_MASKED"), TEXT("1"));
+			break;
+		case SBM_TranslucentGreyTransmittance:
+			OutEnvironment.SetDefine(TEXT("STRATA_BLENDING_TRANSLUCENT_GREYTRANSMITTANCE"), TEXT("1"));
+			break;
+		case SBM_TranslucentColoredTransmittance:
+			OutEnvironment.SetDefine(TEXT("STRATA_BLENDING_TRANSLUCENT_COLOREDTRANSMITTANCE"), TEXT("1"));
+			break;
+		case SBM_ColoredTransmittanceOnly:
+			OutEnvironment.SetDefine(TEXT("STRATA_BLENDING_COLOREDTRANSMITTANCEONLY"), TEXT("1"));
+			break;
+		case SBM_AlphaHoldout:
+			OutEnvironment.SetDefine(TEXT("STRATA_BLENDING_ALPHAHOLDOUT"), TEXT("1"));
+			break;
+		default:
+			UE_LOG(LogMaterial, Error, TEXT("%s: unkown strata material blend mode could not be converted to Starta. (Asset: %s)"), *GetFriendlyName(), *GetAssetName());
+		}
+	}
+
 	{
 		EMaterialDecalResponse MaterialDecalResponse = (EMaterialDecalResponse)GetMaterialDecalResponse();
 
diff --git a/Engine/Source/Runtime/Engine/Private/ShaderCompiler/ShaderGenerationUtil.cpp b/Engine/Source/Runtime/Engine/Private/ShaderCompiler/ShaderGenerationUtil.cpp
--- a/Engine/Source/Runtime/Engine/Private/ShaderCompiler/ShaderGenerationUtil.cpp
+++ b/Engine/Source/Runtime/Engine/Private/ShaderCompiler/ShaderGenerationUtil.cpp
@@ -170,6 +170,13 @@ void FShaderCompileUtilities::ApplyFetchEnvironment(FShaderMaterialPropertyDefin
 	FETCH_COMPILE_BOOL(MATERIALBLENDING_MODULATE);
 	FETCH_COMPILE_BOOL(MATERIALBLENDING_ALPHAHOLDOUT);
 
+	FETCH_COMPILE_BOOL(STRATA_BLENDING_OPAQUE);
+	FETCH_COMPILE_BOOL(STRATA_BLENDING_MASKED);
+	FETCH_COMPILE_BOOL(STRATA_BLENDING_TRANSLUCENT_GREYTRANSMITTANCE);
+	FETCH_COMPILE_BOOL(STRATA_BLENDING_TRANSLUCENT_COLOREDTRANSMITTANCE);
+	FETCH_COMPILE_BOOL(STRATA_BLENDING_COLOREDTRANSMITTANCEONLY);
+	FETCH_COMPILE_BOOL(STRATA_BLENDING_ALPHAHOLDOUT);
+
 	FETCH_COMPILE_INT(MATERIALDECALRESPONSEMASK);
 
 	FETCH_COMPILE_BOOL(REFRACTION_USE_INDEX_OF_REFRACTION);
diff --git a/Engine/Source/Runtime/Engine/Public/MaterialShared.h b/Engine/Source/Runtime/Engine/Public/MaterialShared.h
--- a/Engine/Source/Runtime/Engine/Public/MaterialShared.h
+++ b/Engine/Source/Runtime/Engine/Public/MaterialShared.h
@@ -1783,6 +1783,7 @@ public:
 	virtual bool IsDitherMasked() const { return false; }
 	virtual bool AllowNegativeEmissiveColor() const { return false; }
 	virtual enum EBlendMode GetBlendMode() const = 0;
+	virtual enum EStrataBlendMode GetStrataBlendMode() const = 0;
 	ENGINE_API virtual enum ERefractionMode GetRefractionMode() const;
 	virtual FMaterialShadingModelField GetShadingModels() const = 0;
 	virtual bool IsShadingModelFromMaterialExpression() const = 0;
@@ -2712,6 +2713,7 @@ public:
 	ENGINE_API virtual bool UseLmDirectionality() const override;
 	ENGINE_API virtual bool IsMobileHighQualityBRDFEnabled() const override;
 	ENGINE_API virtual enum EBlendMode GetBlendMode() const override;
+	ENGINE_API virtual enum EStrataBlendMode GetStrataBlendMode() const override;
 	ENGINE_API virtual enum ERefractionMode GetRefractionMode() const override;
 	ENGINE_API virtual uint32 GetMaterialDecalResponse() const override;
 	ENGINE_API virtual bool HasBaseColorConnected() const override;
diff --git a/Engine/Source/Runtime/RenderCore/Private/ShaderMaterialDerivedHelpers.cpp b/Engine/Source/Runtime/RenderCore/Private/ShaderMaterialDerivedHelpers.cpp
--- a/Engine/Source/Runtime/RenderCore/Private/ShaderMaterialDerivedHelpers.cpp
+++ b/Engine/Source/Runtime/RenderCore/Private/ShaderMaterialDerivedHelpers.cpp
@@ -78,7 +78,8 @@ FShaderMaterialDerivedDefines RENDERCORE_API CalculateDerivedMaterialParameters(
 
 	//Materials also have to opt in to these features.
 	Dst.USE_EDITOR_COMPOSITING = (Dst.USE_EDITOR_SHADERS && Mat.EDITOR_PRIMITIVE_MATERIAL);
-	Dst.MATERIALBLENDING_ANY_TRANSLUCENT = (Mat.MATERIALBLENDING_TRANSLUCENT || Mat.MATERIALBLENDING_ADDITIVE || Mat.MATERIALBLENDING_MODULATE);
+	Dst.MATERIALBLENDING_ANY_TRANSLUCENT = (Mat.MATERIALBLENDING_TRANSLUCENT || Mat.MATERIALBLENDING_ADDITIVE || Mat.MATERIALBLENDING_MODULATE
+		|| Mat.STRATA_BLENDING_TRANSLUCENT_GREYTRANSMITTANCE || Mat.STRATA_BLENDING_TRANSLUCENT_COLOREDTRANSMITTANCE || Mat.STRATA_BLENDING_COLOREDTRANSMITTANCEONLY);
 
 	Dst.IS_MESHPARTICLE_FACTORY = (Lightmap.PARTICLE_MESH_FACTORY || Lightmap.NIAGARA_MESH_FACTORY);
 
@@ -97,7 +98,7 @@ FShaderMaterialDerivedDefines RENDERCORE_API CalculateDerivedMaterialParameters(
 	Dst.STRATA_ENABLED = Mat.PROJECT_STRATA | Mat.MATERIAL_IS_STRATA;
 	Dst.SHADER_STRATA_TRANSLUCENT_ENABLED = Dst.STRATA_ENABLED && Dst.MATERIALBLENDING_ANY_TRANSLUCENT;
 	// As of today, all translucent Strata material forces dual source color blending. STRATA_TODO: optimize out dual color blending when not needed or requested
-	Dst.MATERIAL_WORKS_WITH_DUAL_SOURCE_COLOR_BLENDING = (Mat.MATERIAL_SHADINGMODEL_THIN_TRANSLUCENT || Dst.SHADER_STRATA_TRANSLUCENT_ENABLED);
+	Dst.MATERIAL_WORKS_WITH_DUAL_SOURCE_COLOR_BLENDING = (Mat.MATERIAL_SHADINGMODEL_THIN_TRANSLUCENT || Mat.STRATA_BLENDING_TRANSLUCENT_COLOREDTRANSMITTANCE);
 
 	// There are 4 different ways of setting MRTs depending on which .usf file is the base shader, which sets defines which includes a different include file
 	// which may or may not override the current defines. Here is my best attempt at figuring out the logic.
diff --git a/Engine/Source/Runtime/RenderCore/Public/ShaderMaterial.h b/Engine/Source/Runtime/RenderCore/Public/ShaderMaterial.h
--- a/Engine/Source/Runtime/RenderCore/Public/ShaderMaterial.h
+++ b/Engine/Source/Runtime/RenderCore/Public/ShaderMaterial.h
@@ -123,6 +123,13 @@ struct FShaderMaterialPropertyDefines
 	uint8 MATERIALBLENDING_MODULATE : 1;
 	uint8 MATERIALBLENDING_ALPHAHOLDOUT : 1;
 
+	uint8 STRATA_BLENDING_OPAQUE : 1;
+	uint8 STRATA_BLENDING_MASKED : 1;
+	uint8 STRATA_BLENDING_TRANSLUCENT_GREYTRANSMITTANCE : 1;
+	uint8 STRATA_BLENDING_TRANSLUCENT_COLOREDTRANSMITTANCE : 1;
+	uint8 STRATA_BLENDING_COLOREDTRANSMITTANCEONLY : 1;
+	uint8 STRATA_BLENDING_ALPHAHOLDOUT : 1;
+
 	uint8 USES_EMISSIVE_COLOR : 1;
 	
 	uint8 REFRACTION_USE_INDEX_OF_REFRACTION : 1;
diff --git a/Engine/Source/Runtime/Renderer/Private/BasePassRendering.cpp b/Engine/Source/Runtime/Renderer/Private/BasePassRendering.cpp
--- a/Engine/Source/Runtime/Renderer/Private/BasePassRendering.cpp
+++ b/Engine/Source/Runtime/Renderer/Private/BasePassRendering.cpp
@@ -142,8 +142,58 @@ static bool IsBasePassWaitForTasksEnabled()
 
 void SetTranslucentRenderState(FMeshPassProcessorRenderState& DrawRenderState, const FMaterial& Material, const EShaderPlatform Platform, ETranslucencyPass::Type InTranslucencyPassType)
 {
-	// STRATA_TODO IsStrataMaterial will be removed when all materials are unified to use premultipled alpha
-	if (Material.GetShadingModels().HasShadingModel(MSM_ThinTranslucent) || Material.IsStrataMaterial())
+	if (Material.IsStrataMaterial())
+	{
+		if (Material.IsDualBlendingEnabled(Platform))
+		{
+			if (InTranslucencyPassType == ETranslucencyPass::TPT_StandardTranslucency || InTranslucencyPassType == ETranslucencyPass::TPT_AllTranslucency)
+			{
+				// If we are in the transparancy pass (before DoF) we do standard dual blending, and the alpha gets ignored
+				// Blend by putting add in target 0 and multiply by background in target 1.
+				DrawRenderState.SetBlendState(TStaticBlendState<CW_RGBA, BO_Add, BF_One, BF_Source1Color, BO_Add, BF_One, BF_Source1Alpha>::GetRHI());
+			}
+			else if (InTranslucencyPassType == ETranslucencyPass::TPT_TranslucencyAfterDOF)
+			{
+				// In the separate pass (after DoF), we want let alpha pass through, and then multiply our color modulation in the after DoF Modulation pass.
+				// Alpha is BF_Zero for source and BF_One for dest, which leaves alpha unchanged
+				DrawRenderState.SetBlendState(TStaticBlendState<CW_RGBA, BO_Add, BF_One, BF_Source1Color, BO_Add, BF_Zero, BF_One>::GetRHI());
+			}
+			else if (InTranslucencyPassType == ETranslucencyPass::TPT_TranslucencyAfterDOFModulate)
+			{
+				// In the separate pass (after DoF) modulate, we want to only darken the target by our multiplication term, and ignore the addition term.
+				// For regular dual blending, our function is:
+				//     FrameBuffer = MRT0 + MRT1 * FrameBuffer;
+				// So we can just remove the MRT0 component and it will modulate as expected.
+				// Alpha we will leave unchanged.
+				DrawRenderState.SetBlendState(TStaticBlendState<CW_RGBA, BO_Add, BF_Zero, BF_Source1Color, BO_Add, BF_Zero, BF_One>::GetRHI());
+			}
+			else if (InTranslucencyPassType == ETranslucencyPass::TPT_TranslucencyAfterMotionBlur)
+			{
+				// We don't actually currently support color modulation in the post-motion blur pass at the moment, so just do the same as post-DOF for now
+				DrawRenderState.SetBlendState(TStaticBlendState<CW_RGBA, BO_Add, BF_One, BF_Source1Color, BO_Add, BF_Zero, BF_One>::GetRHI());
+			}
+		}
+		else
+		{
+			if (Material.GetStrataBlendMode() == SBM_ColoredTransmittanceOnly)
+			{
+				// Modulate with the existing scene color, preserve destination alpha.
+				DrawRenderState.SetBlendState(TStaticBlendState<CW_RGB, BO_Add, BF_DestColor, BF_Zero>::GetRHI());
+			}
+			else if (Material.GetStrataBlendMode() == SBM_AlphaHoldout)
+			{
+				// Blend by holding out the matte shape of the source alpha
+				DrawRenderState.SetBlendState(TStaticBlendState<CW_RGBA, BO_Add, BF_Zero, BF_InverseSourceAlpha, BO_Add, BF_One, BF_InverseSourceAlpha>::GetRHI());
+			}
+			else
+			{
+				// We always use premultipled alpha for translucent rendering.
+				// If a material was requesting dual source blending, the shader will use static platofm knowledge to convert colored transmittance to a grey scale transmittance.
+				DrawRenderState.SetBlendState(TStaticBlendState<CW_RGBA, BO_Add, BF_One, BF_InverseSourceAlpha, BO_Add, BF_Zero, BF_InverseSourceAlpha>::GetRHI());
+			}
+		}
+	}
+	else if (Material.GetShadingModels().HasShadingModel(MSM_ThinTranslucent))
 	{
 		// Special case for dual blending, which is not exposed as a parameter in the material editor
 		if (Material.IsDualBlendingEnabled(Platform))
@@ -178,8 +228,7 @@ void SetTranslucentRenderState(FMeshPassProcessorRenderState& DrawRenderState, c
 		}
 		else
 		{
-			// If unsupported, use the same as translucent, but with color multiplied by BF_One instead of BF_SourceAlpha.
-			// The shader will use the variation that approximates color modulation using alpha
+			// If unsupported, we still use premultipled alpha but the shader will use the variation converting color transmittance to a grey scale transmittance.
 			DrawRenderState.SetBlendState(TStaticBlendState<CW_RGBA, BO_Add, BF_One, BF_InverseSourceAlpha, BO_Add, BF_Zero, BF_InverseSourceAlpha>::GetRHI());
 		}
 	}
@@ -1600,6 +1649,7 @@ bool FBasePassMeshProcessor::TryAddMeshBatch(const FMeshBatch& RESTRICT MeshBatc
 {
 	// Determine the mesh's material and blend mode.
 	const EBlendMode BlendMode = Material.GetBlendMode();
+	const EStrataBlendMode StrataBlendMode = Material.GetStrataBlendMode();
 	const FMaterialShadingModelField ShadingModels = Material.GetShadingModels();
 	const bool bIsTranslucent = IsTranslucentBlendMode(BlendMode);
 	const FMeshDrawingPolicyOverrideSettings OverrideSettings = ComputeMeshOverrideSettings(MeshBatch);
@@ -1619,12 +1669,13 @@ bool FBasePassMeshProcessor::TryAddMeshBatch(const FMeshBatch& RESTRICT MeshBatc
 				break;
 
 			case ETranslucencyPass::TPT_TranslucencyAfterDOF:
-				bShouldDraw = Material.IsTranslucencyAfterDOFEnabled();
+				bShouldDraw = Material.IsTranslucencyAfterDOFEnabled() && StrataBlendMode != SBM_ColoredTransmittanceOnly;
 				break;
 
 			// only dual blended or modulate surfaces need background modulation
 			case ETranslucencyPass::TPT_TranslucencyAfterDOFModulate:
-				bShouldDraw = Material.IsTranslucencyAfterDOFEnabled() && (Material.IsDualBlendingEnabled(GetFeatureLevelShaderPlatform(FeatureLevel)) || BlendMode == BLEND_Modulate);
+				bShouldDraw = Material.IsTranslucencyAfterDOFEnabled() && (Material.IsDualBlendingEnabled(GetFeatureLevelShaderPlatform(FeatureLevel)) || 
+					BlendMode == BLEND_Modulate || StrataBlendMode == SBM_ColoredTransmittanceOnly);
 				break;
 
 			case ETranslucencyPass::TPT_TranslucencyAfterMotionBlur:
