Rename cvar for consistency.

diff --git a/Engine/Source/Runtime/Renderer/Private/VirtualShadowMaps/VirtualShadowMapArray.cpp b/Engine/Source/Runtime/Renderer/Private/VirtualShadowMaps/VirtualShadowMapArray.cpp
--- a/Engine/Source/Runtime/Renderer/Private/VirtualShadowMaps/VirtualShadowMapArray.cpp
+++ b/Engine/Source/Runtime/Renderer/Private/VirtualShadowMaps/VirtualShadowMapArray.cpp
@@ -118,10 +118,10 @@ TAutoConsoleVariable<int32> CVarMarkCoarsePagesLocal(
 );
 
 TAutoConsoleVariable<int32> CVarCoarsePagesIncludeNonNanite(
-	TEXT("r.Shadow.Virtual.CoarsePagesIncludeNonNanite"),
+	TEXT("r.Shadow.Virtual.NonNanite.IncludeInCoarsePages"),
 	1,
-	TEXT("Include non-nanite geometry in coarse pages.")
-	TEXT("Rendering non-nanite geometry into large coarse pages can be expensive; disabling this can be a significant performance win."),
+	TEXT("Include non-Nanite geometry in coarse pages.")
+	TEXT("Rendering non-Nanite geometry into large coarse pages can be expensive; disabling this can be a significant performance win."),
 	ECVF_Scalability | ECVF_RenderThreadSafe
 );
 
diff --git a/Engine/Source/Runtime/Renderer/Private/VirtualShadowMaps/VirtualShadowMapCacheManager.cpp b/Engine/Source/Runtime/Renderer/Private/VirtualShadowMaps/VirtualShadowMapCacheManager.cpp
--- a/Engine/Source/Runtime/Renderer/Private/VirtualShadowMaps/VirtualShadowMapCacheManager.cpp
+++ b/Engine/Source/Runtime/Renderer/Private/VirtualShadowMaps/VirtualShadowMapCacheManager.cpp
@@ -30,7 +30,7 @@ static TAutoConsoleVariable<int32> CVarCacheVirtualSMs(
 );
 
 static TAutoConsoleVariable<int32> CVarDrawInvalidatingBounds(
-	TEXT("r.Shadow.Virtual.DrawInvalidatingBounds"),
+	TEXT("r.Shadow.Virtual.Cache.DrawInvalidatingBounds"),
 	0,
 	TEXT("Turn on debug render cache invalidating instance bounds, heat mapped by number of pages invalidated.\n")
 	TEXT("   1  = Draw all bounds.\n")
