Disable HMD extensions for reflection capture rendering.

diff --git a/Engine/Source/Runtime/Renderer/Private/ReflectionEnvironmentCapture.cpp b/Engine/Source/Runtime/Renderer/Private/ReflectionEnvironmentCapture.cpp
--- a/Engine/Source/Runtime/Renderer/Private/ReflectionEnvironmentCapture.cpp
+++ b/Engine/Source/Runtime/Renderer/Private/ReflectionEnvironmentCapture.cpp
@@ -1441,13 +1441,6 @@ void CaptureSceneIntoScratchCubemap(
 
 		FSceneView* View = new FSceneView(ViewInitOptions);
 
-		ViewFamily.ViewExtensions = GEngine->ViewExtensions->GatherActiveExtensions(FSceneViewExtensionContext(Scene));
-		for (const FSceneViewExtensionRef& Extension : ViewFamily.ViewExtensions)
-		{
-			Extension->SetupViewFamily(ViewFamily);
-			Extension->SetupView(ViewFamily, *View);
-		}
-
 		// Force all surfaces diffuse
 		View->RoughnessOverrideParameter = FVector2D( 1.0f, 0.0f );
 
@@ -1467,7 +1460,9 @@ void CaptureSceneIntoScratchCubemap(
 		ViewFamily.SetScreenPercentageInterface(new FLegacyScreenPercentageDriver(
 			ViewFamily, /* GlobalResolutionFraction = */ 1.0f));
 
-		ViewFamily.ViewExtensions = GEngine->ViewExtensions->GatherActiveExtensions(FSceneViewExtensionContext(Scene));
+		FSceneViewExtensionContext ViewExtensionContext(Scene);
+		ViewExtensionContext.bStereoDisabled = true;
+		ViewFamily.ViewExtensions = GEngine->ViewExtensions->GatherActiveExtensions(ViewExtensionContext);
 
 		for (const FSceneViewExtensionRef& Extension : ViewFamily.ViewExtensions)
 		{
