Fixed shader print related crash in strata debug vis code.

diff --git a/Engine/Source/Runtime/Renderer/Private/Strata/Strata.cpp b/Engine/Source/Runtime/Renderer/Private/Strata/Strata.cpp
--- a/Engine/Source/Runtime/Renderer/Private/Strata/Strata.cpp
+++ b/Engine/Source/Runtime/Renderer/Private/Strata/Strata.cpp
@@ -365,11 +365,7 @@ static void AddVisualizeMaterialPasses(FRDGBuilder& GraphBuilder, const FViewInf
 		PassParameters->MiniFontTexture = GetMiniFontTexture();
 		PassParameters->SceneTextures = GetSceneTextureParameters(GraphBuilder);
 		PassParameters->RenderTargets[0] = FRenderTargetBinding(SceneColorTexture, ERenderTargetLoadAction::ELoad);
-
-		if (ShaderPrint::IsEnabled())
-		{
-			ShaderPrint::SetParameters(GraphBuilder, View.ShaderPrintData, PassParameters->ShaderPrintParameters);
-		}
+		ShaderPrint::SetParameters(GraphBuilder, View.ShaderPrintData, PassParameters->ShaderPrintParameters);
 
 		for (uint32 j = 0; j < VISUALIZE_MATERIAL_PASS_COUNT; ++j)
 		{
