Fix issues when ShowFlag.Rendering is disabled.

diff --git a/Engine/Source/Runtime/Renderer/Private/DeferredShadingRenderer.cpp b/Engine/Source/Runtime/Renderer/Private/DeferredShadingRenderer.cpp
--- a/Engine/Source/Runtime/Renderer/Private/DeferredShadingRenderer.cpp
+++ b/Engine/Source/Runtime/Renderer/Private/DeferredShadingRenderer.cpp
@@ -1955,6 +1955,15 @@ void FDeferredShadingSceneRenderer::Render(FRDGBuilder& GraphBuilder)
 		ShaderPrint::BeginView(GraphBuilder, View);
 		ShadingEnergyConservation::Init(GraphBuilder, View);
 	}
+	
+	ON_SCOPE_EXIT
+	{
+		for (FViewInfo& View : Views)
+		{
+			ShaderPrint::EndView(View);
+		}
+	};
+
 	Scene->UpdateAllPrimitiveSceneInfos(GraphBuilder, true);
 
 	FGPUSceneScopeBeginEndHelper GPUSceneScopeBeginEndHelper(Scene->GPUScene, GPUSceneDynamicContext, Scene);
@@ -3309,11 +3318,6 @@ void FDeferredShadingSceneRenderer::Render(FRDGBuilder& GraphBuilder)
 		FinishGatheringLumenSurfaceCacheFeedback(GraphBuilder, Views[0], LumenFrameTemporaries);
 	}
 
-	for (FViewInfo& View : Views)
-	{
-		ShaderPrint::EndView(View);
-	}
-
 	GEngine->GetPostRenderDelegateEx().Broadcast(GraphBuilder);
 
 #if RHI_RAYTRACING
diff --git a/Engine/Source/Runtime/Renderer/Private/GPUScene.cpp b/Engine/Source/Runtime/Renderer/Private/GPUScene.cpp
--- a/Engine/Source/Runtime/Renderer/Private/GPUScene.cpp
+++ b/Engine/Source/Runtime/Renderer/Private/GPUScene.cpp
@@ -2176,8 +2176,10 @@ void FGPUSceneDynamicContext::Release()
 {
 	for (auto UploadData : DymamicPrimitiveUploadData)
 	{
-		check(UploadData->InstanceSceneDataOffset != INDEX_NONE);
-		GPUScene.FreeInstanceSceneDataSlots(UploadData->InstanceSceneDataOffset, UploadData->TotalInstanceCount);
+		if (UploadData->InstanceSceneDataOffset != INDEX_NONE)
+		{
+			GPUScene.FreeInstanceSceneDataSlots(UploadData->InstanceSceneDataOffset, UploadData->TotalInstanceCount);
+		}
 		if (UploadData->InstancePayloadDataOffset != INDEX_NONE)
 		{
 			GPUScene.FreeInstancePayloadDataSlots(UploadData->InstancePayloadDataOffset, UploadData->InstancePayloadDataFloat4Count);
