Mark r.Mobile.PixelFogQuality as scalability cvar

diff --git a/Engine/Source/Runtime/Renderer/Private/MobileFogRendering.cpp b/Engine/Source/Runtime/Renderer/Private/MobileFogRendering.cpp
--- a/Engine/Source/Runtime/Renderer/Private/MobileFogRendering.cpp
+++ b/Engine/Source/Runtime/Renderer/Private/MobileFogRendering.cpp
@@ -21,7 +21,7 @@ static TAutoConsoleVariable<int32> CVarPixelFogQuality(
 	TEXT("Exponentional height fog rendering quality.\n")
 	TEXT("0 - basic per-pixel fog")
 	TEXT("1 - all per-pixel fog features (second fog, directional inscattering, aerial perspective)"),
-	ECVF_RenderThreadSafe);
+	ECVF_Scalability | ECVF_RenderThreadSafe);
 
 static TAutoConsoleVariable<int32> CVarPixelFogDepthTest(
 	TEXT("r.Mobile.PixelFogDepthTest"),
