Enable the voxelization from height field to global distance field for metal SM5.

diff --git a/Engine/Source/Runtime/Renderer/Private/GlobalDistanceField.cpp b/Engine/Source/Runtime/Renderer/Private/GlobalDistanceField.cpp
--- a/Engine/Source/Runtime/Renderer/Private/GlobalDistanceField.cpp
+++ b/Engine/Source/Runtime/Renderer/Private/GlobalDistanceField.cpp
@@ -1785,7 +1785,6 @@ void UpdateGlobalDistanceFieldVolume(
 							&& NumHeightfieldPrimitives > 0
 							&& GAOGlobalDistanceFieldRepresentHeightfields
 							&& SupportsDistanceFieldAO(Scene->GetFeatureLevel(), Scene->GetShaderPlatform())
-							&& !IsMetalPlatform(Scene->GetShaderPlatform())
 							&& !IsVulkanMobileSM5Platform(Scene->GetShaderPlatform()))
 						{
 							for (int32 HeightfieldPrimitiveIndex = 0; HeightfieldPrimitiveIndex < NumHeightfieldPrimitives; HeightfieldPrimitiveIndex++)
