Fixes TSR breaking the ICVFX camera motion blur setting in ndisplay

diff --git a/Engine/Shaders/Private/MotionBlur/MotionBlurVelocityFlatten.usf b/Engine/Shaders/Private/MotionBlur/MotionBlurVelocityFlatten.usf
--- a/Engine/Shaders/Private/MotionBlur/MotionBlurVelocityFlatten.usf
+++ b/Engine/Shaders/Private/MotionBlur/MotionBlurVelocityFlatten.usf
@@ -20,9 +20,10 @@
 
 SCREEN_PASS_TEXTURE_VIEWPORT(Velocity)
 
-#if CAMERA_MOTION_BLUR_MODE == 2 // TODO: move into MotionBlurVelocityFlatten.ush
 float4x4 ClipToPrevClipOverride;
-#endif
+
+uint bCancelCameraMotion;
+uint bAddCustomCameraMotion;
 
 Texture2D VelocityTexture;
 Texture2D DepthTexture;
@@ -47,6 +48,23 @@ float2 GetCameraMotionVelocity(uint2 DispatchThreadId, float Depth, float4x4 Cli
 	return ScreenPos - PrevScreen;
 }
 
+float2 PreprocessVelocityForMotionBlur(uint2 DispatchThreadId, float DeviceZ, float2 Velocity, float2 CameraMotionVelocity)
+{
+	FLATTEN
+	if (bCancelCameraMotion)
+	{
+		Velocity -= CameraMotionVelocity;
+	}
+
+	BRANCH
+	if (bAddCustomCameraMotion)
+	{
+		Velocity += GetCameraMotionVelocity(DispatchThreadId, DeviceZ, ClipToPrevClipOverride);
+	}
+
+	return Velocity;
+}
+
 
 //------------------------------------------------------- ENTRY POINT
 
@@ -97,31 +115,15 @@ void MainCS(
 	float4 EncodedVelocity = VelocityTexture[ PixelPos + DepthOffset ];
 	float DeviceZ = DepthTexture[ PixelPos + DepthOffset ].x;
 
-	float2 Velocity = 0.0f;
+	float2 CameraMotionVelocity = GetCameraMotionVelocity(DispatchThreadId, DeviceZ, View.ClipToPrevClip);
+	float2 Velocity = CameraMotionVelocity;
+	FLATTEN
 	if (EncodedVelocity.x > 0.0)
 	{
 		Velocity = DecodeVelocityFromTexture(EncodedVelocity).xy;
-#if CAMERA_MOTION_BLUR_MODE != 1
-		// Cancel motion blur caused by the current camera movements
-		Velocity -= GetCameraMotionVelocity(DispatchThreadId, DeviceZ, View.ClipToPrevClip);
-#endif
-#if CAMERA_MOTION_BLUR_MODE == 2
-		// Apply motion blur caused by the custom camera movements
-		Velocity += GetCameraMotionVelocity(DispatchThreadId, DeviceZ, ClipToPrevClipOverride);
-#endif
 	}
-#if CAMERA_MOTION_BLUR_MODE > 0
-	else
-	{
-#if CAMERA_MOTION_BLUR_MODE == 2
-		// Apply motion blur caused by the custom camera movements
-		Velocity = GetCameraMotionVelocity(DispatchThreadId, DeviceZ, ClipToPrevClipOverride);
-#else
-		// Apply motion blur caused by the current camera movements
-		Velocity = GetCameraMotionVelocity(DispatchThreadId, DeviceZ, View.ClipToPrevClip);
-#endif
-	}
-#endif
+
+	Velocity = PreprocessVelocityForMotionBlur(DispatchThreadId, DeviceZ, Velocity, CameraMotionVelocity);
 
 	FVelocityRange VelocityPolarRange;
 	float2 VelocityPolar;
diff --git a/Engine/Source/Runtime/Renderer/Private/PostProcess/PostProcessMotionBlur.cpp b/Engine/Source/Runtime/Renderer/Private/PostProcess/PostProcessMotionBlur.cpp
--- a/Engine/Source/Runtime/Renderer/Private/PostProcess/PostProcessMotionBlur.cpp
+++ b/Engine/Source/Runtime/Renderer/Private/PostProcess/PostProcessMotionBlur.cpp
@@ -182,8 +182,17 @@ int32 GetMotionBlurDirections()
 }
 
 // static
-bool FVelocityFlattenTextures::AllowExternal()
+bool FVelocityFlattenTextures::AllowExternal(const FViewInfo& View)
 {
+	const bool bEnableCameraMotionBlur = View.bCameraMotionBlur.Get(true);
+	const bool bOverrideCameraMotionBlur = View.ClipToPrevClipOverride.IsSet();
+
+	// Do not use external velocity flatten passes if the camera motion blur needs to be modified.
+	if (!bEnableCameraMotionBlur || bOverrideCameraMotionBlur)
+	{
+		return false;
+	}
+
 	return CVarMotionBlurAllowExternalVelocityFlatten.GetValueOnRenderThread() != 0;
 }
 
@@ -298,8 +307,7 @@ class FMotionBlurVelocityFlattenCS : public FMotionBlurShader
 public:
 	static const uint32 ThreadGroupSize = 16;
 
-	class FCameraMotionBlurMode : SHADER_PERMUTATION_INT("CAMERA_MOTION_BLUR_MODE", 3);
-	using FPermutationDomain = TShaderPermutationDomain<FCameraMotionBlurMode, FMotionBlurDirections>;
+	using FPermutationDomain = TShaderPermutationDomain<FMotionBlurDirections>;
 
 	DECLARE_GLOBAL_SHADER(FMotionBlurVelocityFlattenCS);
 	SHADER_USE_PARAMETER_STRUCT(FMotionBlurVelocityFlattenCS, FMotionBlurShader);
@@ -309,6 +317,8 @@ public:
 		SHADER_PARAMETER_STRUCT(FScreenPassTextureViewportParameters, Velocity)
 		SHADER_PARAMETER_STRUCT_INCLUDE(FVelocityFlattenParameters, VelocityFlattenParameters)
 		SHADER_PARAMETER(FMatrix44f, ClipToPrevClipOverride)
+		SHADER_PARAMETER(int32, bCancelCameraMotion)
+		SHADER_PARAMETER(int32, bAddCustomCameraMotion)
 
 		SHADER_PARAMETER_RDG_TEXTURE(Texture2D, VelocityTexture)
 		SHADER_PARAMETER_RDG_TEXTURE(Texture2D, DepthTexture)
@@ -640,9 +650,11 @@ void AddMotionBlurVelocityPass(
 		FMotionBlurVelocityFlattenCS::FParameters* PassParameters = GraphBuilder.AllocParameters<FMotionBlurVelocityFlattenCS::FParameters>();
 		PassParameters->View = View.ViewUniformBuffer;
 		PassParameters->VelocityFlattenParameters = GetVelocityFlattenParameters(View);
-		if (bOverrideCameraMotionBlur)
+		PassParameters->bCancelCameraMotion = !bEnableCameraMotionBlur || bOverrideCameraMotionBlur;
+		PassParameters->bAddCustomCameraMotion = bOverrideCameraMotionBlur;
+		if (PassParameters->bAddCustomCameraMotion)
 		{
-			PassParameters->ClipToPrevClipOverride = FMatrix44f(View.ClipToPrevClipOverride.GetValue());		// LWC_TODO: Precision loss?
+			PassParameters->ClipToPrevClipOverride = FMatrix44f(View.ClipToPrevClipOverride.GetValue());
 		}
 
 		PassParameters->Velocity = Viewports.VelocityParameters;
@@ -654,7 +666,6 @@ void AddMotionBlurVelocityPass(
 		PassParameters->DebugOutput = CreateDebugUAV(GraphBuilder, VelocityFlatTexture->Desc.Extent, TEXT("Debug.MotionBlur.Flatten"));
 
 		FMotionBlurVelocityFlattenCS::FPermutationDomain PermutationVector;
-		PermutationVector.Set<FMotionBlurVelocityFlattenCS::FCameraMotionBlurMode>(bEnableCameraMotionBlur ? (bOverrideCameraMotionBlur ? 2 : 1) : 0);
 		PermutationVector.Set<FMotionBlurDirections>(BlurDirections);
 
 		TShaderMapRef<FMotionBlurVelocityFlattenCS> ComputeShader(View.ShaderMap, PermutationVector);
diff --git a/Engine/Source/Runtime/Renderer/Private/PostProcess/PostProcessMotionBlur.h b/Engine/Source/Runtime/Renderer/Private/PostProcess/PostProcessMotionBlur.h
--- a/Engine/Source/Runtime/Renderer/Private/PostProcess/PostProcessMotionBlur.h
+++ b/Engine/Source/Runtime/Renderer/Private/PostProcess/PostProcessMotionBlur.h
@@ -67,7 +67,7 @@ struct FVelocityFlattenTextures
 	}
 
 	// returns whether FVelocityFlattenTextures can be generated in external system.
-	static bool AllowExternal();
+	static bool AllowExternal(const FViewInfo& View);
 };
 
 struct FMotionBlurInputs
diff --git a/Engine/Source/Runtime/Renderer/Private/PostProcess/PostProcessing.cpp b/Engine/Source/Runtime/Renderer/Private/PostProcess/PostProcessing.cpp
--- a/Engine/Source/Runtime/Renderer/Private/PostProcess/PostProcessing.cpp
+++ b/Engine/Source/Runtime/Renderer/Private/PostProcess/PostProcessing.cpp
@@ -594,7 +594,7 @@ void AddPostProcessingPasses(
 
 			UpscalerPassInputs.bAllowDownsampleSceneColor = bAllowSceneDownsample;
 			UpscalerPassInputs.bGenerateOutputMip1 = bMotionBlurNeedsHalfResInput;
-			UpscalerPassInputs.bGenerateVelocityFlattenTextures = PassSequence.IsEnabled(EPass::MotionBlur) && FVelocityFlattenTextures::AllowExternal() && !bVisualizeMotionBlur;
+			UpscalerPassInputs.bGenerateVelocityFlattenTextures = PassSequence.IsEnabled(EPass::MotionBlur) && FVelocityFlattenTextures::AllowExternal(View) && !bVisualizeMotionBlur;
 			UpscalerPassInputs.DownsampleOverrideFormat = DownsampleOverrideFormat;
 			UpscalerPassInputs.SceneColorTexture = SceneColor.Texture;
 			UpscalerPassInputs.SceneDepthTexture = SceneDepth.Texture;
