Ray Tracing - fixed primitive index in CachedRayTracingInstance. PrimitiveIndex in CachedRayTracingInstance can't be swapped during a primitive swap, as it may contain a null CachedRayTracingInstance.

diff --git a/Engine/Source/Runtime/Renderer/Private/RendererScene.cpp b/Engine/Source/Runtime/Renderer/Private/RendererScene.cpp
--- a/Engine/Source/Runtime/Renderer/Private/RendererScene.cpp
+++ b/Engine/Source/Runtime/Renderer/Private/RendererScene.cpp
@@ -4474,8 +4474,9 @@ void FScene::UpdateAllPrimitiveSceneInfos(FRDGBuilder& GraphBuilder, bool bAsync
 							GPUScene.RecordPrimitiveIdSwap(DestIndex, SourceIndex);
 
 						#if RHI_RAYTRACING
-							// Also need to swap DefaultUserData since it currently contains PrimitiveIndex
-							Swap(Primitives[SourceIndex]->CachedRayTracingInstance.DefaultUserData, Primitives[DestIndex]->CachedRayTracingInstance.DefaultUserData);
+							// Update cached PrimitiveIndex after an index swap
+							Primitives[SourceIndex]->CachedRayTracingInstance.DefaultUserData = SourceIndex;
+							Primitives[DestIndex]->CachedRayTracingInstance.DefaultUserData = DestIndex;
 						#endif
 
 							SourceIndex = DestIndex;
@@ -4811,8 +4812,9 @@ void FScene::UpdateAllPrimitiveSceneInfos(FRDGBuilder& GraphBuilder, bool bAsync
 							GPUScene.RecordPrimitiveIdSwap(DestIndex, SourceIndex);
 
 						#if RHI_RAYTRACING
-							// Also need to swap DefaultUserData since it currently contains PrimitiveIndex
-							Swap(Primitives[SourceIndex]->CachedRayTracingInstance.DefaultUserData, Primitives[DestIndex]->CachedRayTracingInstance.DefaultUserData);
+							// Update cached PrimitiveIndex after an index swap
+							Primitives[SourceIndex]->CachedRayTracingInstance.DefaultUserData = SourceIndex;
+							Primitives[DestIndex]->CachedRayTracingInstance.DefaultUserData = DestIndex;
 						#endif
 						}
 					}
