Fix DFAO ghosting in mesh preview.
- By default mesh preview doesn't update in real time so temporal ghosting doesn't clear.
- Don't use temporal history when View.Family->bRealtimeUpdate == false.
- This matches what other techniques such as TAA and Motion Blur do.

diff --git a/Engine/Source/Runtime/Renderer/Private/DistanceFieldLightingPost.cpp b/Engine/Source/Runtime/Renderer/Private/DistanceFieldLightingPost.cpp
--- a/Engine/Source/Runtime/Renderer/Private/DistanceFieldLightingPost.cpp
+++ b/Engine/Source/Runtime/Renderer/Private/DistanceFieldLightingPost.cpp
@@ -306,6 +306,7 @@ void UpdateHistory(
 		if (*BentNormalHistoryState 
 			&& !View.bCameraCut 
 			&& !View.bPrevTransformsReset
+			&& View.Family->bRealtimeUpdate
 			&& !GAOClearHistory
 			// If the scene render targets reallocate, toss the history so we don't read uninitialized data
 			&& (*BentNormalHistoryState)->GetDesc().Extent == BufferSize)
