Fix incorrect heightfield data in Global Distance Field

Context:
- When texture is compiling, renderer uses a default texture instead.
- Compiling state is not reflected in bHasStreamingUpdatePending
- Distance to heightfield calculation involves a projection onto the heightfield normal which is incorrect when using default texture, resulting in distance to be 0 everywhere inside the heightfield xy bounds.

Change:
- check if HeightmapTexture->IsCompiling() in FLandscapeComponentSceneProxy::HeightfieldHasPendingStreaming().
- Only add primitive heightfield to UpdateRegionHeightfield if texture is not streaming or compiling.

diff --git a/Engine/Source/Runtime/Landscape/Private/LandscapeRender.cpp b/Engine/Source/Runtime/Landscape/Private/LandscapeRender.cpp
--- a/Engine/Source/Runtime/Landscape/Private/LandscapeRender.cpp
+++ b/Engine/Source/Runtime/Landscape/Private/LandscapeRender.cpp
@@ -3995,7 +3995,27 @@ void FLandscapeComponentSceneProxy::ChangeComponentScreenSizeToUseSubSections_Re
 
 bool FLandscapeComponentSceneProxy::HeightfieldHasPendingStreaming() const
 {
-	return HeightmapTexture && HeightmapTexture->bHasStreamingUpdatePending;
+	bool bHeightmapTextureStreaming = false;
+
+	if (HeightmapTexture)
+	{
+		bHeightmapTextureStreaming |= HeightmapTexture->bHasStreamingUpdatePending;
+#if WITH_EDITOR
+		bHeightmapTextureStreaming |= HeightmapTexture->IsCompiling();
+#endif
+	}
+
+	bool bVisibilityTextureStreaming = false;
+
+	if (VisibilityWeightmapTexture)
+	{
+		bVisibilityTextureStreaming |= VisibilityWeightmapTexture->bHasStreamingUpdatePending;
+#if WITH_EDITOR
+		bVisibilityTextureStreaming |= VisibilityWeightmapTexture->IsCompiling();
+#endif
+	}
+
+	return bHeightmapTextureStreaming || bVisibilityTextureStreaming;
 }
 
 void FLandscapeComponentSceneProxy::GetHeightfieldRepresentation(UTexture2D*& OutHeightmapTexture, UTexture2D*& OutDiffuseColorTexture, UTexture2D*& OutVisibilityTexture, FHeightfieldComponentDescription& OutDescription) const
diff --git a/Engine/Source/Runtime/Renderer/Private/GlobalDistanceField.cpp b/Engine/Source/Runtime/Renderer/Private/GlobalDistanceField.cpp
--- a/Engine/Source/Runtime/Renderer/Private/GlobalDistanceField.cpp
+++ b/Engine/Source/Runtime/Renderer/Private/GlobalDistanceField.cpp
@@ -1789,11 +1789,16 @@ void UpdateGlobalDistanceFieldVolume(
 						{
 							for (int32 HeightfieldPrimitiveIndex = 0; HeightfieldPrimitiveIndex < NumHeightfieldPrimitives; HeightfieldPrimitiveIndex++)
 							{
-								const FPrimitiveSceneInfo* HeightfieldPrimitive = Scene->DistanceFieldSceneData.HeightfieldPrimitives[HeightfieldPrimitiveIndex];
-								const FBoxSphereBounds& PrimitiveBounds = HeightfieldPrimitive->Proxy->GetBounds();
+								const FPrimitiveSceneProxy* HeightfieldPrimitiveProxy = Scene->DistanceFieldSceneData.HeightfieldPrimitives[HeightfieldPrimitiveIndex]->Proxy;
+								const FBoxSphereBounds& PrimitiveBounds = HeightfieldPrimitiveProxy->GetBounds();
+
+								if (HeightfieldPrimitiveProxy->HeightfieldHasPendingStreaming())
+								{
+									continue;
+								}
 
 								// Expand bounding box by a SDF max influence distance (only in local Z axis, as distance is computed from a top down projected heightmap point).
-								const FVector QueryInfluenceExpand = HeightfieldPrimitive->Proxy->GetLocalToWorld().GetUnitAxis(EAxis::Z) * FVector(0.0f, 0.0f, ClipmapInfluenceRadius);
+								const FVector QueryInfluenceExpand = HeightfieldPrimitiveProxy->GetLocalToWorld().GetUnitAxis(EAxis::Z) * FVector(0.0f, 0.0f, ClipmapInfluenceRadius);
 								const FBox HeightfieldInfluenceBox = PrimitiveBounds.GetBox().ExpandBy(QueryInfluenceExpand, QueryInfluenceExpand);
 
 								if (Clipmap.Bounds.Intersect(HeightfieldInfluenceBox))
@@ -1801,8 +1806,8 @@ void UpdateGlobalDistanceFieldVolume(
 									UTexture2D* HeightfieldTexture = nullptr;
 									UTexture2D* DiffuseColorTexture = nullptr;
 									UTexture2D* VisibilityTexture = nullptr;
-									FHeightfieldComponentDescription NewComponentDescription(HeightfieldPrimitive->Proxy->GetLocalToWorld());
-									HeightfieldPrimitive->Proxy->GetHeightfieldRepresentation(HeightfieldTexture, DiffuseColorTexture, VisibilityTexture, NewComponentDescription);
+									FHeightfieldComponentDescription NewComponentDescription(HeightfieldPrimitiveProxy->GetLocalToWorld());
+									HeightfieldPrimitiveProxy->GetHeightfieldRepresentation(HeightfieldTexture, DiffuseColorTexture, VisibilityTexture, NewComponentDescription);
 
 									if (HeightfieldTexture && HeightfieldTexture->GetResource() && HeightfieldTexture->GetResource()->TextureRHI)
 									{
