Fixed volumetric cloud shadow map flickering in volumetric fog: CloudShadowmapTranslatedWorldToLightClipMatrix parameter is now assigned.
Fixed trailing in volumetric fog: disabled r.VolumetricFog.ConservativeDepth by default and marked it as experimental.

diff --git a/Engine/Source/Runtime/Renderer/Private/VolumetricFog.cpp b/Engine/Source/Runtime/Renderer/Private/VolumetricFog.cpp
--- a/Engine/Source/Runtime/Renderer/Private/VolumetricFog.cpp
+++ b/Engine/Source/Runtime/Renderer/Private/VolumetricFog.cpp
@@ -114,11 +114,11 @@ FAutoConsoleVariableRef CVarVolumetricFogEmissive(
 	ECVF_Scalability | ECVF_RenderThreadSafe
 );
 
-int32 GVolumetricFogConservativeDepth = 1;
+int32 GVolumetricFogConservativeDepth = 0;
 FAutoConsoleVariableRef CVarVolumetricFogConservativeDepth(
 	TEXT("r.VolumetricFog.ConservativeDepth"),
 	GVolumetricFogConservativeDepth,
-	TEXT("Whether to allow the volumetric to use conservative depth to accelerate computations."),
+	TEXT("[Experimental] Whether to allow the volumetric to use conservative depth to accelerate computations."),
 	ECVF_Scalability | ECVF_RenderThreadSafe
 );
 
@@ -1208,6 +1208,7 @@ void FDeferredShadingSceneRenderer::ComputeVolumetricFog(FRDGBuilder& GraphBuild
 			PassParameters->CloudShadowmapSampler = TStaticSamplerState<SF_Bilinear, AM_Clamp, AM_Clamp, AM_Clamp>::GetRHI();
 			PassParameters->CloudShadowmapFarDepthKm = CloudShadowmap_FarDepthKm;
 			PassParameters->CloudShadowmapStrength = CloudShadowmap_Strength;
+			PassParameters->CloudShadowmapTranslatedWorldToLightClipMatrix = CloudWorldToLightClipShadowMatrix;
 
 			const bool bUseLumenGI = View.LumenTranslucencyGIVolume.Texture0 != nullptr;
 			const bool bUseGlobalDistanceField = UseGlobalDistanceField() && Scene->DistanceFieldSceneData.NumObjectsInBuffer > 0;
