Adds VisualizeMotionVectors show flag

diff --git a/Engine/Shaders/Private/MotionBlur/VisualizeMotionVectors.usf b/Engine/Shaders/Private/MotionBlur/VisualizeMotionVectors.usf
new file mode 100644
--- /dev/null
+++ b/Engine/Shaders/Private/MotionBlur/VisualizeMotionVectors.usf
@@ -0,0 +1,88 @@
+// Copyright Epic Games, Inc. All Rights Reserved.
+
+#include "../Common.ush"
+#include "../ScreenPass.ush"
+
+
+//------------------------------------------------------- PARAMETERS
+
+Texture2D ColorTexture;
+Texture2D VelocityTexture;
+Texture2D DepthTexture;
+Texture2D PrevColorTexture;
+
+SCREEN_PASS_TEXTURE_VIEWPORT(Color)
+SCREEN_PASS_TEXTURE_VIEWPORT(Velocity)
+SCREEN_PASS_TEXTURE_VIEWPORT(PrevColor)
+
+SamplerState ColorSampler;
+SamplerState VelocitySampler;
+SamplerState DepthSampler;
+SamplerState PrevColorSampler;
+
+FScreenTransform SvPositionToVelocity;
+FScreenTransform SvPositionToColor;
+FScreenTransform SvPositionToScreenPos;
+FScreenTransform PrevScreenPosToPrevColor;
+
+
+//------------------------------------------------------- FUNCTIONS
+
+float2 ComputeStaticVelocity(float2 ScreenPos, float DeviceZ)
+{
+	float3 PosN = float3(ScreenPos, DeviceZ);
+
+	float4 ThisClip = float4(PosN, 1);
+	float4 PrevClip = mul( ThisClip, View.ClipToPrevClip );
+	float2 PrevScreen = PrevClip.xy / PrevClip.w;
+	return PosN.xy - PrevScreen;
+}
+
+
+//------------------------------------------------------- ENTRY POINT
+
+void MainPS(
+	float4 SvPosition : SV_POSITION,
+	out float4 OutColor : SV_Target0)
+{
+	float2 ColorTextureUV = ApplyScreenTransform(SvPosition.xy, SvPositionToColor);
+	float2 VelocityTextureUV = ApplyScreenTransform(SvPosition.xy, SvPositionToVelocity);
+	float2 ScreenPos = ApplyScreenTransform(SvPosition.xy, SvPositionToScreenPos);
+
+	ColorTextureUV = clamp(ColorTextureUV, Color_UVViewportBilinearMin, Color_UVViewportBilinearMax);
+	VelocityTextureUV = clamp(VelocityTextureUV, Velocity_UVViewportBilinearMin, Velocity_UVViewportBilinearMax);
+
+	float4 Color = ColorTexture.SampleLevel(ColorSampler, ColorTextureUV, 0.0);
+	float4 EncodedVelocity = VelocityTexture.SampleLevel(VelocitySampler, VelocityTextureUV, 0.0);
+	float DeviceZ = DepthTexture.SampleLevel(DepthSampler, VelocityTextureUV, 0.0).r;
+	
+	float2 ScreenVelocity = ComputeStaticVelocity(ScreenPos, DeviceZ);
+	bool bIsRenderedVelocity = EncodedVelocity.x > 0.0;
+	if (bIsRenderedVelocity)
+	{
+		ScreenVelocity = DecodeVelocityFromTexture(EncodedVelocity).xy;
+	}
+
+	float2 PrevScreenPos = ScreenPos - ScreenVelocity;
+	float2 PrevColorTextureUV = ApplyScreenTransform(PrevScreenPos, PrevScreenPosToPrevColor);
+	PrevColorTextureUV = clamp(PrevColorTextureUV, Color_UVViewportBilinearMin, Color_UVViewportBilinearMax);
+
+	float4 PrevColor = PrevColorTexture.SampleLevel(PrevColorSampler, PrevColorTextureUV, 0.0);
+
+	float Luma = Luminance(Color.rgb);
+	float PrevLuma = Luminance(PrevColor.rgb);
+
+	uint2 SubTileCoord = (uint2(SvPosition.xy) >> 4) & 0x1;
+	float Checker = SubTileCoord.x == SubTileCoord.y ? 1.0 : 0.0;
+
+	if (bIsRenderedVelocity)
+	{
+		OutColor.rgb = (Luma + PrevLuma).xxx * 0.5 + abs(Color - PrevColor);
+		//OutColor.rgb = (Color + PrevColor) * 0.5;
+		OutColor.a = min(Color.a, PrevColor.a);
+	}
+	else
+	{
+		OutColor = float4(Luma.xxx * (0.25 + 0.25 * Checker), Color.a);
+	}
+}
diff --git a/Engine/Source/Runtime/Engine/Public/ShowFlags.h b/Engine/Source/Runtime/Engine/Public/ShowFlags.h
--- a/Engine/Source/Runtime/Engine/Public/ShowFlags.h
+++ b/Engine/Source/Runtime/Engine/Public/ShowFlags.h
@@ -384,6 +384,7 @@ private:
 		SetVertexColors(false);
 		SetPhysicalMaterialMasks(false);
 		SetVisualizeMotionBlur(false);
+		SetVisualizeMotionVectors(false);
 		SetEditingLevelInstance(false);
 		SetSelectionOutline(false);
 		SetDebugAI(false);
diff --git a/Engine/Source/Runtime/Engine/Public/ShowFlagsValues.inl b/Engine/Source/Runtime/Engine/Public/ShowFlagsValues.inl
--- a/Engine/Source/Runtime/Engine/Public/ShowFlagsValues.inl
+++ b/Engine/Source/Runtime/Engine/Public/ShowFlagsValues.inl
@@ -114,6 +114,8 @@ SHOWFLAG_ALWAYS_ACCESSIBLE(SeparateTranslucency, SFG_Advanced, NSLOCTEXT("Unreal
 SHOWFLAG_ALWAYS_ACCESSIBLE(ScreenPercentage, SFG_Hidden, NSLOCTEXT("UnrealEd", "ScreenPercentageSF", "Screen Percentage"))
 /** Helper to tweak motion blur settings */
 SHOWFLAG_FIXED_IN_SHIPPING(0, VisualizeMotionBlur, SFG_Visualize, NSLOCTEXT("UnrealEd", "VisualizeMotionBlurSF", "Motion Blur"))
+/** Helper to tweak motion vectors */
+SHOWFLAG_FIXED_IN_SHIPPING(0, VisualizeMotionVectors, SFG_Visualize, NSLOCTEXT("UnrealEd", "VisualizeMotionVectorsSF", "Motion Vectors"))
 /** Whether to display the Reflection Environment feature, which has local reflections from Reflection Capture actors, for now SHOWFLAG_ALWAYS_ACCESSIBLE because it's exposed in SceneCapture */
 SHOWFLAG_ALWAYS_ACCESSIBLE(ReflectionEnvironment, SFG_LightingFeatures, NSLOCTEXT("UnrealEd", "ReflectionEnvironmentSF", "Reflection Environment"))
 /** Visualize pixels that are outside of their object's bounding box (content error). */
diff --git a/Engine/Source/Runtime/Renderer/Private/PostProcess/PostProcessing.cpp b/Engine/Source/Runtime/Renderer/Private/PostProcess/PostProcessing.cpp
--- a/Engine/Source/Runtime/Renderer/Private/PostProcess/PostProcessing.cpp
+++ b/Engine/Source/Runtime/Renderer/Private/PostProcess/PostProcessing.cpp
@@ -39,6 +39,7 @@
 #include "PostProcess/PostProcessFFTBloom.h"
 #include "PostProcess/PostProcessStreamingAccuracyLegend.h"
 #include "PostProcess/PostProcessSubsurface.h"
+#include "PostProcess/VisualizeMotionVectors.h"
 #include "Rendering/MotionVectorSimulation.h"
 #include "ShaderPrint.h"
 #include "HighResScreenshot.h"
@@ -276,6 +277,7 @@ void AddPostProcessingPasses(
 		VisualizeLumenSceneOverview,
 		VisualizeHDR,
 		VisualizeLocalExposure,
+		VisualizeMotionVectors,
 		PixelInspector,
 		HMDDistortion,
 		HighResolutionScreenshotMask,
@@ -322,6 +324,7 @@ void AddPostProcessingPasses(
 		TEXT("VisualizeLumenSceneOverview"),
 		TEXT("VisualizeHDR"),
 		TEXT("VisualizeLocalExposure"),
+		TEXT("VisualizeMotionVectors"),
 		TEXT("PixelInspector"),
 		TEXT("HMDDistortion"),
 		TEXT("HighResolutionScreenshotMask"),
@@ -359,6 +362,7 @@ void AddPostProcessingPasses(
 	PassSequence.SetEnabled(EPass::VisualizeGBufferOverview, bVisualizeGBufferOverview || bVisualizeGBufferDumpToFile || bVisualizeGBufferDumpToPIpe);
 	PassSequence.SetEnabled(EPass::VisualizeLumenSceneOverview, VisualizeMode == VISUALIZE_MODE_OVERVIEW && bPostProcessingEnabled);
 	PassSequence.SetEnabled(EPass::VisualizeHDR, EngineShowFlags.VisualizeHDR);
+	PassSequence.SetEnabled(EPass::VisualizeMotionVectors, EngineShowFlags.VisualizeMotionVectors);
 #if WITH_EDITOR
 	PassSequence.SetEnabled(EPass::PixelInspector, View.bUsePixelInspector);
 #else
@@ -1222,6 +1226,17 @@ void AddPostProcessingPasses(
 		SceneColor = AddVisualizeLocalExposurePass(GraphBuilder, View, PassInputs);
 	}
 
+	if (PassSequence.IsEnabled(EPass::VisualizeMotionVectors))
+	{
+		FVisualizeMotionVectorsInputs PassInputs;
+		PassSequence.AcceptOverrideIfLastPass(EPass::VisualizeMotionVectors, PassInputs.OverrideOutput);
+		PassInputs.SceneColor = SceneColor;
+		PassInputs.SceneDepth = SceneDepth;
+		PassInputs.SceneVelocity = Velocity;
+
+		SceneColor = AddVisualizeMotionVectorsPass(GraphBuilder, View, PassInputs);
+	}
+
 #if WITH_EDITOR
 	if (PassSequence.IsEnabled(EPass::PixelInspector))
 	{
diff --git a/Engine/Source/Runtime/Renderer/Private/PostProcess/VisualizeMotionVectors.cpp b/Engine/Source/Runtime/Renderer/Private/PostProcess/VisualizeMotionVectors.cpp
new file mode 100644
--- /dev/null
+++ b/Engine/Source/Runtime/Renderer/Private/PostProcess/VisualizeMotionVectors.cpp
@@ -0,0 +1,116 @@
+// Copyright Epic Games, Inc. All Rights Reserved.
+
+#include "PostProcess/VisualizeMotionVectors.h"
+#include "PostProcessing.h"
+#include "ScreenPass.h"
+#include "PixelShaderUtils.h"
+
+
+class FVisualizeMotionVectorsPS : public FGlobalShader
+{
+	DECLARE_GLOBAL_SHADER(FVisualizeMotionVectorsPS);
+	SHADER_USE_PARAMETER_STRUCT(FVisualizeMotionVectorsPS, FGlobalShader);
+
+	BEGIN_SHADER_PARAMETER_STRUCT(FParameters, )
+		SHADER_PARAMETER_STRUCT_REF(FViewUniformShaderParameters, View)
+		SHADER_PARAMETER_RDG_TEXTURE(Texture2D, ColorTexture)
+		SHADER_PARAMETER_RDG_TEXTURE(Texture2D, DepthTexture)
+		SHADER_PARAMETER_RDG_TEXTURE(Texture2D, VelocityTexture)
+		SHADER_PARAMETER_RDG_TEXTURE(Texture2D, PrevColorTexture)
+
+		SHADER_PARAMETER_STRUCT(FScreenPassTextureViewportParameters, Color)
+		SHADER_PARAMETER_STRUCT(FScreenPassTextureViewportParameters, Velocity)
+		SHADER_PARAMETER_STRUCT(FScreenPassTextureViewportParameters, PrevColor)
+
+		SHADER_PARAMETER_SAMPLER(SamplerState, ColorSampler)
+		SHADER_PARAMETER_SAMPLER(SamplerState, VelocitySampler)
+		SHADER_PARAMETER_SAMPLER(SamplerState, DepthSampler)
+		SHADER_PARAMETER_SAMPLER(SamplerState, PrevColorSampler)
+
+		SHADER_PARAMETER(FScreenTransform, SvPositionToVelocity)
+		SHADER_PARAMETER(FScreenTransform, SvPositionToColor)
+		SHADER_PARAMETER(FScreenTransform, SvPositionToScreenPos)
+		SHADER_PARAMETER(FScreenTransform, PrevScreenPosToPrevColor)
+
+		RENDER_TARGET_BINDING_SLOTS()
+	END_SHADER_PARAMETER_STRUCT()
+};
+
+IMPLEMENT_GLOBAL_SHADER(FVisualizeMotionVectorsPS, "/Engine/Private/MotionBlur/VisualizeMotionVectors.usf", "MainPS", SF_Pixel);
+
+FScreenPassTexture AddVisualizeMotionVectorsPass(FRDGBuilder& GraphBuilder, const FViewInfo& View, const FVisualizeMotionVectorsInputs& Inputs)
+{
+	check(Inputs.SceneColor.IsValid());
+	check(Inputs.SceneDepth.IsValid());
+	check(Inputs.SceneVelocity.IsValid());
+	checkf(Inputs.SceneDepth.ViewRect == Inputs.SceneVelocity.ViewRect, TEXT("The implementation requires that depth and velocity have the same viewport."));
+
+	FScreenPassTexture PrevColor;
+	if (View.PrevViewInfo.VisualizeMotionVectors.IsValid())
+	{
+		PrevColor.Texture = GraphBuilder.RegisterExternalTexture(View.PrevViewInfo.VisualizeMotionVectors);
+		PrevColor.ViewRect = View.PrevViewInfo.VisualizeMotionVectorsRect;
+	}
+	else
+	{
+		PrevColor.Texture = GSystemTextures.GetBlackDummy(GraphBuilder);
+		PrevColor.ViewRect = FIntRect(0, 0, 1, 1);
+	}
+
+	FScreenPassRenderTarget Output = Inputs.OverrideOutput;
+	{
+		if (!Output.IsValid())
+		{
+			Output = FScreenPassRenderTarget::CreateFromInput(GraphBuilder, Inputs.SceneColor, View.GetOverwriteLoadAction(), TEXT("MotionVectors.Visualize"));
+		}
+
+		FVisualizeMotionVectorsPS::FParameters* PassParameters = GraphBuilder.AllocParameters<FVisualizeMotionVectorsPS::FParameters>();
+		PassParameters->View = View.ViewUniformBuffer;
+
+		PassParameters->ColorTexture = Inputs.SceneColor.Texture;
+		PassParameters->DepthTexture = Inputs.SceneDepth.Texture;
+		PassParameters->VelocityTexture = Inputs.SceneVelocity.Texture;
+		PassParameters->PrevColorTexture = PrevColor.Texture;
+
+		PassParameters->Color = GetScreenPassTextureViewportParameters(FScreenPassTextureViewport(Inputs.SceneColor));
+		PassParameters->Velocity = GetScreenPassTextureViewportParameters(FScreenPassTextureViewport(Inputs.SceneDepth));
+		PassParameters->PrevColor = GetScreenPassTextureViewportParameters(FScreenPassTextureViewport(PrevColor));
+
+		PassParameters->ColorSampler     = TStaticSamplerState<SF_Point, AM_Clamp, AM_Clamp, AM_Clamp>::GetRHI();
+		PassParameters->VelocitySampler  = TStaticSamplerState<SF_Point, AM_Clamp, AM_Clamp, AM_Clamp>::GetRHI();
+		PassParameters->DepthSampler     = TStaticSamplerState<SF_Point, AM_Clamp, AM_Clamp, AM_Clamp>::GetRHI();
+		PassParameters->PrevColorSampler = TStaticSamplerState<SF_Bilinear, AM_Clamp, AM_Clamp, AM_Clamp>::GetRHI();
+
+		FScreenTransform SvPositionToViewportUV = FScreenTransform::SvPositionToViewportUV(Output.ViewRect);
+		PassParameters->SvPositionToColor = (
+			SvPositionToViewportUV *
+			FScreenTransform::ChangeTextureBasisFromTo(Inputs.SceneColor, FScreenTransform::ETextureBasis::ViewportUV, FScreenTransform::ETextureBasis::TextureUV));
+		PassParameters->SvPositionToVelocity = (
+			SvPositionToViewportUV *
+			FScreenTransform::ChangeTextureBasisFromTo(Inputs.SceneDepth, FScreenTransform::ETextureBasis::ViewportUV, FScreenTransform::ETextureBasis::TextureUV));
+		PassParameters->SvPositionToScreenPos = (
+			SvPositionToViewportUV *
+			FScreenTransform::ViewportUVToScreenPos);
+		PassParameters->PrevScreenPosToPrevColor = (
+			FScreenTransform::ChangeTextureBasisFromTo(PrevColor, FScreenTransform::ETextureBasis::ScreenPosition, FScreenTransform::ETextureBasis::TextureUV));
+
+		PassParameters->RenderTargets[0] = Output.GetRenderTargetBinding();
+
+		TShaderMapRef<FVisualizeMotionVectorsPS> PixelShader(View.ShaderMap);
+		FPixelShaderUtils::AddFullscreenPass(
+			GraphBuilder,
+			View.ShaderMap,
+			RDG_EVENT_NAME("VisualizeMotionVectors %dx%d", Output.ViewRect.Width(), Output.ViewRect.Height()),
+			PixelShader,
+			PassParameters,
+			Output.ViewRect);
+	}
+
+	if (View.ViewState && !View.bStatePrevViewInfoIsReadOnly)
+	{
+		GraphBuilder.QueueTextureExtraction(Inputs.SceneColor.Texture, &View.ViewState->PrevFrameViewInfo.VisualizeMotionVectors);
+		View.ViewState->PrevFrameViewInfo.VisualizeMotionVectorsRect = Inputs.SceneColor.ViewRect;
+	}
+
+	return MoveTemp(Output);
+}
diff --git a/Engine/Source/Runtime/Renderer/Private/PostProcess/VisualizeMotionVectors.h b/Engine/Source/Runtime/Renderer/Private/PostProcess/VisualizeMotionVectors.h
new file mode 100644
--- /dev/null
+++ b/Engine/Source/Runtime/Renderer/Private/PostProcess/VisualizeMotionVectors.h
@@ -0,0 +1,18 @@
+// Copyright Epic Games, Inc. All Rights Reserved.
+
+#pragma once
+
+#include "ScreenPass.h"
+
+
+struct FVisualizeMotionVectorsInputs
+{
+	// [Optional] Render to the specified output. If invalid, a new texture is created and returned.
+	FScreenPassRenderTarget OverrideOutput;
+
+	FScreenPassTexture SceneColor;
+	FScreenPassTexture SceneDepth;
+	FScreenPassTexture SceneVelocity;
+};
+
+FScreenPassTexture AddVisualizeMotionVectorsPass(FRDGBuilder& GraphBuilder, const FViewInfo& View, const FVisualizeMotionVectorsInputs& Inputs);
\ No newline at end of file
diff --git a/Engine/Source/Runtime/Renderer/Private/SceneRendering.h b/Engine/Source/Runtime/Renderer/Private/SceneRendering.h
--- a/Engine/Source/Runtime/Renderer/Private/SceneRendering.h
+++ b/Engine/Source/Runtime/Renderer/Private/SceneRendering.h
@@ -1066,6 +1066,10 @@ struct FPreviousViewInfo
 
 	// Mobile ambient occlusion texture used for next frame.
 	TRefCountPtr<IPooledRenderTarget> MobileAmbientOcclusion = nullptr;
+
+	// Scene color used for reprojecting next frame to verify the motion vector reprojects correctly.
+	TRefCountPtr<IPooledRenderTarget> VisualizeMotionVectors;
+	FIntRect VisualizeMotionVectorsRect;
 };
 
 class FViewCommands
