Rolls back responsive AA behavior change

diff --git a/Engine/Shaders/Private/TemporalSuperResolution/TSRUpdateHistory.usf b/Engine/Shaders/Private/TemporalSuperResolution/TSRUpdateHistory.usf
--- a/Engine/Shaders/Private/TemporalSuperResolution/TSRUpdateHistory.usf
+++ b/Engine/Shaders/Private/TemporalSuperResolution/TSRUpdateHistory.usf
@@ -40,10 +40,10 @@
 #if DIM_SEPARATE_TRANSLUCENCY == 1
 	#define CONFIG_ACCUMULATE_TRANSLUCENCY_SEPARATLY 1
 	#define CONFIG_VFX_SURVIVAL 0
-	#define CONFIG_REJECT_TRAMSLUCENCY_ON_RESPONSIVE_AA_ONLY 1
+	#define CONFIG_REJECT_TRAMSLUCENCY_ON_RESPONSIVE_AA_ONLY 0
 #else
 	#define CONFIG_ACCUMULATE_TRANSLUCENCY_SEPARATLY 0
-	#define CONFIG_VFX_SURVIVAL 1
+	#define CONFIG_VFX_SURVIVAL 0
 	#define CONFIG_REJECT_TRAMSLUCENCY_ON_RESPONSIVE_AA_ONLY 0
 #endif
 
@@ -57,6 +57,8 @@
 #endif
 
 #define CONFIG_RESPONSIVE_STENCIL 1
+#define CONFIG_TAA_RESPONSIVE_BEHAVIOR 1
+#define CONFIG_LOAD_TRANSLUCENCY_REJECTION 0
 
 #if CONFIG_COMPILE_FP16 && !CONFIG_SCENE_COLOR_APLHA
 	#define CONFIG_MANUAL_LDS_SPILL 1
@@ -1258,7 +1260,15 @@ void MainCS(
 		}
 		#endif
 		
-		BlendFinal = max(BlendFinal, (tsr_half(1.0) - TranslucencyRejection));
+		#if CONFIG_TAA_RESPONSIVE_BEHAVIOR
+			// Responsive forces 1/4 of new frame.
+			BlendFinal[0] = bIsResponsiveAAPixel[0] ? tsr_half(1.0/4.0) : BlendFinal[0];
+			BlendFinal[1] = bIsResponsiveAAPixel[1] ? tsr_half(1.0/4.0) : BlendFinal[1];
+		#endif
+		
+		#if CONFIG_LOAD_TRANSLUCENCY_REJECTION
+			BlendFinal = max(BlendFinal, (tsr_half(1.0) - TranslucencyRejection));
+		#endif
 		
 		FLATTEN
 		if (bCameraCut)
diff --git a/Engine/Source/Runtime/Renderer/Private/PostProcess/TemporalSuperResolution.cpp b/Engine/Source/Runtime/Renderer/Private/PostProcess/TemporalSuperResolution.cpp
--- a/Engine/Source/Runtime/Renderer/Private/PostProcess/TemporalSuperResolution.cpp
+++ b/Engine/Source/Runtime/Renderer/Private/PostProcess/TemporalSuperResolution.cpp
@@ -54,10 +54,10 @@ TAutoConsoleVariable<float> CVarTSRTranslucencyHighlightLuminance(
 	TEXT("Sets the liminance at which translucency is considered an highlights (default=-1.0)."),
 	ECVF_RenderThreadSafe);
 
-TAutoConsoleVariable<int32> CVarTSRTranslucencyPreviousFrameRejection(
-	TEXT("r.TSR.Translucency.PreviousFrameRejection"), 1,
-	TEXT("Enable heuristic to reject Separate translucency based on previous frame translucency."),
-	ECVF_RenderThreadSafe);
+//TAutoConsoleVariable<int32> CVarTSRTranslucencyPreviousFrameRejection(
+//	TEXT("r.TSR.Translucency.PreviousFrameRejection"), 0,
+//	TEXT("Enable heuristic to reject Separate translucency based on previous frame translucency."),
+//	ECVF_RenderThreadSafe);
 
 TAutoConsoleVariable<int32> CVarTSRTranslucencySeparateTemporalAccumulation(
 	TEXT("r.TSR.Translucency.SeparateTemporalAccumulation"), 1,
@@ -594,7 +594,7 @@ ITemporalUpscaler::FOutputs AddTemporalSuperResolutionPasses(
 
 	bool bIsSeperateTranslucyTexturesValid = PassInputs.PostDOFTranslucencyResources.IsValid();
 
-	bool bRejectSeparateTranslucency = bIsSeperateTranslucyTexturesValid && CVarTSRTranslucencyPreviousFrameRejection.GetValueOnRenderThread() != 0;
+	const bool bRejectSeparateTranslucency = false; // bIsSeperateTranslucyTexturesValid && CVarTSRTranslucencyPreviousFrameRejection.GetValueOnRenderThread() != 0;
 
 	bool bAccumulateSeparateTranslucency = bIsSeperateTranslucyTexturesValid && CVarTSRTranslucencySeparateTemporalAccumulation.GetValueOnRenderThread() != 0;
 
