Fix static analysis

diff --git a/Engine/Source/Runtime/Renderer/Private/DynamicBVH.h b/Engine/Source/Runtime/Renderer/Private/DynamicBVH.h
--- a/Engine/Source/Runtime/Renderer/Private/DynamicBVH.h
+++ b/Engine/Source/Runtime/Renderer/Private/DynamicBVH.h
@@ -592,7 +592,7 @@ uint32 FDynamicBVH< MaxChildren, FRootPolicy, FDirtyPolicy, FCostMetric >::FindB
 		if( Node.IsFull() )
 		{
 			float	BestChildDist = MAX_flt;
-			uint32	BestChildIndex = ~0u;
+			uint32	BestChildIndex = 0;
 
 			for( uint32 i = 0; i < Node.NumChildren; i += 4 )
 			{
@@ -1112,7 +1112,7 @@ uint32 FDynamicBVH< MaxChildren, FRootPolicy, FDirtyPolicy, FCostMetric >::FindC
 			FVector3f RelativePosition = Root.ToRelative( Position );
 
 			float	NodeDistSqr	= Root.Bounds.DistSqr( RelativePosition );
-			uint32	NodeIndex	= Root.NodeIndex;
+			uint32	NodeIndex	= Root.FirstChild;
 
 			while( NodeDistSqr < ClosestDistSqr )
 			{
