MGPU:  Crash on Valley_Lumen_LARD level.  Make it so we no longer attempt to render shadows across all GPUs, which was done to handle the case of cached shadow maps.

Our plan going forward will be to copy cached shadow maps cross GPU after they render, rather than the fragile approach of intricately placed GPU mask blocks for different parts of the shadow code (which has created numerous bugs, typically difficult to diagnose GPU hangs that only occur on MGPU machines).  A utility function has been added to preserve the logic that identifies projected lights using cached shadow maps, which will be used in the future for that purpose.  Until that's solved, artifacts will result for any MGPU scenarios using cached shadow maps, but artifacts are preferable to the crash...  A new Jira will be created for fixing artifacts.

Bonus fix for validation error (ParticleStatePosition_1 (GPU 0)') (subresource: 0) is invalid for use as a destination resource) mentioned in the bug description:  in RHITransferResources, we need to transition both the source and destination resource for the copy on both command queues.

#rnx

diff --git a/Engine/Source/Runtime/D3D12RHI/Private/D3D12CommandContext.cpp b/Engine/Source/Runtime/D3D12RHI/Private/D3D12CommandContext.cpp
--- a/Engine/Source/Runtime/D3D12RHI/Private/D3D12CommandContext.cpp
+++ b/Engine/Source/Runtime/D3D12RHI/Private/D3D12CommandContext.cpp
@@ -1128,6 +1128,8 @@ void FD3D12CommandContextRedirector::RHITransferResources(const TArrayView<const
 		}
 
 		FD3D12DynamicRHI::TransitionResource(SrcContext->CommandListHandle, SrcResource, D3D12_RESOURCE_STATE_TBD, D3D12_RESOURCE_STATE_COPY_SOURCE, 0, FD3D12DynamicRHI::ETransitionMode::Validate);
+		FD3D12DynamicRHI::TransitionResource(SrcContext->CommandListHandle, DestResource, D3D12_RESOURCE_STATE_TBD, D3D12_RESOURCE_STATE_COPY_DEST, 0, FD3D12DynamicRHI::ETransitionMode::Validate);
+		FD3D12DynamicRHI::TransitionResource(DestContext->CommandListHandle, SrcResource, D3D12_RESOURCE_STATE_TBD, D3D12_RESOURCE_STATE_COPY_SOURCE, 0, FD3D12DynamicRHI::ETransitionMode::Validate);
 		FD3D12DynamicRHI::TransitionResource(DestContext->CommandListHandle, DestResource, D3D12_RESOURCE_STATE_TBD, D3D12_RESOURCE_STATE_COPY_DEST, 0, FD3D12DynamicRHI::ETransitionMode::Validate);
 	}
 
diff --git a/Engine/Source/Runtime/Renderer/Private/SceneRendering.h b/Engine/Source/Runtime/Renderer/Private/SceneRendering.h
--- a/Engine/Source/Runtime/Renderer/Private/SceneRendering.h
+++ b/Engine/Source/Runtime/Renderer/Private/SceneRendering.h
@@ -2028,6 +2028,7 @@ protected:
 
 #if WITH_MGPU
 	FRHIGPUMask AllViewsGPUMask;
+	bool IsShadowCached(FProjectedShadowInfo* ProjectedShadowInfo) const;
 	FRHIGPUMask GetGPUMaskForShadow(FProjectedShadowInfo* ProjectedShadowInfo) const;
 #endif
 
diff --git a/Engine/Source/Runtime/Renderer/Private/ShadowDepthRendering.cpp b/Engine/Source/Runtime/Renderer/Private/ShadowDepthRendering.cpp
--- a/Engine/Source/Runtime/Renderer/Private/ShadowDepthRendering.cpp
+++ b/Engine/Source/Runtime/Renderer/Private/ShadowDepthRendering.cpp
@@ -1248,37 +1248,36 @@ void FProjectedShadowInfo::GetShadowTypeNameForDrawEvent(FString& TypeName) cons
 }
 
 #if WITH_MGPU
-FRHIGPUMask FSceneRenderer::GetGPUMaskForShadow(FProjectedShadowInfo* ProjectedShadowInfo) const
+// Shadows that are cached need to be copied to other GPUs after they render
+bool FSceneRenderer::IsShadowCached(FProjectedShadowInfo* ProjectedShadowInfo) const
 {
-	// Preshadows that are going to be cached this frame should render on all GPUs.
+	// Preshadows that are going to be cached this frame should be copied to other GPUs.
 	if (ProjectedShadowInfo->bPreShadow)
 	{
-		// Multi-GPU support : Updating on all GPUs may be inefficient for AFR. Work is
-		// wasted for any shadows that re-cache on consecutive frames.
-		return !ProjectedShadowInfo->bDepthsCached && ProjectedShadowInfo->bAllocatedInPreshadowCache ? FRHIGPUMask::All() : AllViewsGPUMask;
+		return !ProjectedShadowInfo->bDepthsCached && ProjectedShadowInfo->bAllocatedInPreshadowCache;
 	}
-	// SDCM_StaticPrimitivesOnly shadows don't update every frame so we need to render
-	// their depths on all possible GPUs.
+	// SDCM_StaticPrimitivesOnly shadows don't update every frame so we need to copy
+	// their depths to all possible GPUs.
 	else if (!ProjectedShadowInfo->IsWholeSceneDirectionalShadow() && ProjectedShadowInfo->CacheMode == SDCM_StaticPrimitivesOnly)
 	{
 		// Cached whole scene shadows shouldn't be view dependent.
 		checkSlow(ProjectedShadowInfo->DependentView == nullptr);
 
-		// Multi-GPU support : Updating on all GPUs may be inefficient for AFR. Work is
-		// wasted for any shadows that re-cache on consecutive frames.
-		return FRHIGPUMask::All();
+		return true;
+	}
+	return false;
+}
+
+FRHIGPUMask FSceneRenderer::GetGPUMaskForShadow(FProjectedShadowInfo* ProjectedShadowInfo) const
+{
+	// View dependent shadows only need to render depths on their view's GPUs.
+	if (ProjectedShadowInfo->DependentView != nullptr)
+	{
+		return ProjectedShadowInfo->DependentView->GPUMask;
 	}
 	else
 	{
-		// View dependent shadows only need to render depths on their view's GPUs.
-		if (ProjectedShadowInfo->DependentView != nullptr)
-		{
-			return ProjectedShadowInfo->DependentView->GPUMask;
-		}
-		else
-		{
-			return AllViewsGPUMask;
-		}
+		return AllViewsGPUMask;
 	}
 }
 #endif // WITH_MGPU
@@ -1489,10 +1488,6 @@ bool IsParallelDispatchEnabled(const FProjectedShadowInfo* ProjectedShadowInfo,
 
 void FSceneRenderer::RenderShadowDepthMapAtlases(FRDGBuilder& GraphBuilder)
 {
-	// Perform setup work on all GPUs in case any cached shadows are being updated this
-	// frame. We revert to the AllViewsGPUMask for uncached shadows.
-	RDG_GPU_MASK_SCOPE(GraphBuilder, FRHIGPUMask::All());
-
 	const bool bNaniteEnabled = 
 		UseNanite(ShaderPlatform) &&
 		ViewFamily.EngineShowFlags.NaniteMeshes &&
@@ -1821,13 +1816,6 @@ void FSceneRenderer::RenderShadowDepthMaps(FRDGBuilder& GraphBuilder, FInstanceC
 	// culling.
 	InstanceCullingManager.BeginDeferredCulling(GraphBuilder, Scene->GPUScene);
 
-	// Perform setup work on all GPUs in case any cached shadows are being updated this
-	// frame. We revert to the AllViewsGPUMask for uncached shadows.
-#if WITH_MGPU
-	ensure(GraphBuilder.RHICmdList.GetGPUMask() == AllViewsGPUMask);
-#endif
-	RDG_GPU_MASK_SCOPE(GraphBuilder, FRHIGPUMask::All());
-
 	const bool bNaniteEnabled = 
 		UseNanite(ShaderPlatform) &&
 		ViewFamily.EngineShowFlags.NaniteMeshes &&
