Revert Octree bounds to WORLD_MAX

diff --git a/Engine/Source/Runtime/Renderer/Private/RendererScene.cpp b/Engine/Source/Runtime/Renderer/Private/RendererScene.cpp
--- a/Engine/Source/Runtime/Renderer/Private/RendererScene.cpp
+++ b/Engine/Source/Runtime/Renderer/Private/RendererScene.cpp
@@ -1166,8 +1166,8 @@ FScene::FScene(UWorld* InWorld, bool bInRequiresHitProxies, bool bInIsEditorScen
 ,	SkyAtmosphere(NULL)
 ,	VolumetricCloud(NULL)
 ,	PrecomputedVisibilityHandler(NULL)
-,	LocalShadowCastingLightOctree(FVector::ZeroVector, UE_LARGE_HALF_WORLD_MAX)
-,	PrimitiveOctree(FVector::ZeroVector, UE_LARGE_HALF_WORLD_MAX)
+,	LocalShadowCastingLightOctree(FVector::ZeroVector, WORLD_MAX)
+,	PrimitiveOctree(FVector::ZeroVector, WORLD_MAX)
 ,	bRequiresHitProxies(bInRequiresHitProxies)
 ,	bIsEditorScene(bInIsEditorScene)
 ,	NumUncachedStaticLightingInteractions(0)
