Fix path tracer updates not working over nDisplay and in the editor's Picture-in-picture camera previews.

 * Invalidate Path Traced output when component transforms change (this catches cases that don't re-create a primitive info object)
 * Invalidate Path Traced output when light color or intensity changes
 * Fix invalidation in a scene only affecting a single view (or view family). Instead of a single boolean, the scene now keeps a counter that ticks on changes and each view remembers the last value of this counter to detect changes.

diff --git a/Engine/Source/Runtime/Engine/Private/Components/SceneComponent.cpp b/Engine/Source/Runtime/Engine/Private/Components/SceneComponent.cpp
--- a/Engine/Source/Runtime/Engine/Private/Components/SceneComponent.cpp
+++ b/Engine/Source/Runtime/Engine/Private/Components/SceneComponent.cpp
@@ -786,6 +786,15 @@ void USceneComponent::PropagateTransformUpdate(bool bTransformChanged, EUpdateTr
 
 			// Flag render transform as dirty
 			MarkRenderTransformDirty();
+
+			// Invalidate the path tracer because we know something has changed about the transform which most likely
+			// has invalidated the image. If we hooked into these updates at a deeper level, we would also be notified
+			// when the transforms are marked dirty even when they haven't changed.
+			const UWorld* World = GetWorld();
+			if (World != nullptr && World->Scene != nullptr)
+			{
+				World->Scene->InvalidatePathTracedOutput();
+			}
 		}
 		
 		{
diff --git a/Engine/Source/Runtime/Renderer/Private/RendererScene.cpp b/Engine/Source/Runtime/Renderer/Private/RendererScene.cpp
--- a/Engine/Source/Runtime/Renderer/Private/RendererScene.cpp
+++ b/Engine/Source/Runtime/Renderer/Private/RendererScene.cpp
@@ -1146,7 +1146,7 @@ FScene::FScene(UWorld* InWorld, bool bInRequiresHitProxies, bool bInIsEditorScen
 ,	World(InWorld)
 ,	FXSystem(nullptr)
 ,	bScenesPrimitivesNeedStaticMeshElementUpdate(false)
-,	bPathTracingNeedsInvalidation(true)
+,   PathTracingInvalidationCounter(0)
 #if RHI_RAYTRACING
 ,   CachedRayTracingMeshCommandsMode(ERayTracingMeshCommandsMode::RAY_TRACING)
 #endif
@@ -1966,6 +1966,7 @@ void FScene::AddLightSceneInfo_RenderThread(FLightSceneInfo* LightSceneInfo)
 	{
 		RayTracedLights.Add(LightSceneInfo);
 	}
+	InvalidatePathTracedOutput();
 #endif 
 
 	// Add the light to the scene.
@@ -2926,7 +2927,7 @@ void FScene::InvalidateRuntimeVirtualTexture(class URuntimeVirtualTextureCompone
 void FScene::InvalidatePathTracedOutput()
 {
 	// NOTE: this is an atomic, so this function is ok to call from any thread
-	bPathTracingNeedsInvalidation = true;
+	++PathTracingInvalidationCounter;
 }
 
 void FScene::FlushDirtyRuntimeVirtualTextures()
@@ -3072,7 +3073,7 @@ void FScene::UpdateLightColorAndBrightness(ULightComponent* Light)
 		{
 			 NewParameters.NewColor *= FLinearColor::MakeFromColorTemperature(Light->Temperature);
 		}
-	
+
 		FScene* Scene = this;
 		FLightSceneInfo* LightSceneInfo = Light->SceneProxy->GetLightSceneInfo();
 		ENQUEUE_RENDER_COMMAND(UpdateLightColorAndBrightness)(
@@ -3088,6 +3089,9 @@ void FScene::UpdateLightColorAndBrightness(ULightComponent* Light)
 						( Scene->GetShadingPath() == EShadingPath::Mobile 
 						&& NewParameters.NewColor.IsAlmostBlack() != LightSceneInfo->Proxy->GetColor().IsAlmostBlack() );
 
+					// Path Tracing: something about the light has changed, restart path traced accumulation
+					Scene->InvalidatePathTracedOutput();
+
 					LightSceneInfo->Proxy->SetColor(NewParameters.NewColor);
 					LightSceneInfo->Proxy->IndirectLightingScale = NewParameters.NewIndirectLightingScale;
 					LightSceneInfo->Proxy->VolumetricScatteringIntensity = NewParameters.NewVolumetricScatteringIntensity;
@@ -3097,7 +3101,6 @@ void FScene::UpdateLightColorAndBrightness(ULightComponent* Light)
 					{
 						Scene->Lights[ LightSceneInfo->Id ].Color = NewParameters.NewColor;
 					}
-
 					LightSceneInfo->Proxy->SetMobileMovablePointLightUniformBufferNeedsUpdate(true);
 				}
 			});
@@ -3195,6 +3198,8 @@ void FScene::RemoveLightSceneInfo_RenderThread(FLightSceneInfo* LightSceneInfo)
 		{
 			OverflowingDynamicShadowedLights.Remove(LightSceneInfo->Proxy->GetOwnerNameOrLabel());
 		}
+
+		InvalidatePathTracedOutput();
 	}
 	else
 	{
@@ -4898,9 +4903,6 @@ void FScene::UpdateAllPrimitiveSceneInfos(FRDGBuilder& GraphBuilder, bool bAsync
 					VelocityData.UpdateTransform(PrimitiveSceneInfo, PrimitiveTransforms[PrimitiveIndex], PrimitiveTransforms[PrimitiveIndex]);
 				}
 
-				// Invalidate PathTraced image because we added something to the scene
-				bPathTracingNeedsInvalidation = true;
-
 				DistanceFieldSceneData.AddPrimitive(PrimitiveSceneInfo);
 				LumenSceneData->AddPrimitive(PrimitiveSceneInfo);
 
@@ -4913,6 +4915,11 @@ void FScene::UpdateAllPrimitiveSceneInfos(FRDGBuilder& GraphBuilder, bool bAsync
 				// Update scene LOD tree
 				SceneLODHierarchy.UpdateNodeSceneInfo(PrimitiveSceneInfo->PrimitiveComponentId, PrimitiveSceneInfo);
 			}
+			if (StartIndex < AddedLocalPrimitiveSceneInfos.Num())
+			{
+				// Invalidate PathTraced image because we added something to the scene
+				InvalidatePathTracedOutput();
+			}
 			AddedLocalPrimitiveSceneInfos.RemoveAt(StartIndex, AddedLocalPrimitiveSceneInfos.Num() - StartIndex, false);
 		}
 	}
@@ -5196,9 +5203,11 @@ void FScene::UpdateAllPrimitiveSceneInfos(FRDGBuilder& GraphBuilder, bool bAsync
 			// free the primitive scene proxy.
 			delete PrimitiveSceneInfo->Proxy;
 			delete PrimitiveSceneInfo;
-
+		}
+		if (DeletedSceneInfos.Num() > 0)
+		{
 			// Invalidate PathTraced image because we removed something from the scene
-			bPathTracingNeedsInvalidation = true;
+			InvalidatePathTracedOutput();
 		}
 	}
 	UpdatedAttachmentRoots.Empty();
diff --git a/Engine/Source/Runtime/Renderer/Private/ScenePrivate.h b/Engine/Source/Runtime/Renderer/Private/ScenePrivate.h
--- a/Engine/Source/Runtime/Renderer/Private/ScenePrivate.h
+++ b/Engine/Source/Runtime/Renderer/Private/ScenePrivate.h
@@ -1045,6 +1045,7 @@ public:
 
 	// Keeps track of the internal path tracer state
 	TPimplPtr<FPathTracingState> PathTracingState;
+	uint32 PathTracingInvalidationCounter = 0;
 
 	// IES light profiles
 	FIESLightProfileResource IESLightProfileResources;
@@ -2800,8 +2801,8 @@ public:
 	/** True if a change to SkyLight / Lighting has occurred that requires static draw lists to be updated. */
 	bool bScenesPrimitivesNeedStaticMeshElementUpdate;
 
-	/** True if a change to the scene that requires to invalidate the path tracer buffers has happened. */
-	TAtomic<bool> bPathTracingNeedsInvalidation;
+	/** This counter will be incremented anytime something about the scene changed that should invalidate the path traced accumulation buffers. */
+	TAtomic<uint32> PathTracingInvalidationCounter;
 
 #if RHI_RAYTRACING
 	/** What mode where the cached RT commands prepared for last? */
diff --git a/Engine/Source/Runtime/Renderer/Private/SceneVisibility.cpp b/Engine/Source/Runtime/Renderer/Private/SceneVisibility.cpp
--- a/Engine/Source/Runtime/Renderer/Private/SceneVisibility.cpp
+++ b/Engine/Source/Runtime/Renderer/Private/SceneVisibility.cpp
@@ -3570,6 +3570,12 @@ void FSceneRenderer::PreVisibilityFrameSetup(FRDGBuilder& GraphBuilder, const FS
 	bool bFreezeTemporalSequences = bFreezeTemporalHistories || CVarFreezeTemporalSequences.GetValueOnRenderThread() != 0;
 #endif
 
+	// Load this field once so it has a consistent value for all views (and to avoid the atomic load in the loop).
+	// While the value may not be perfectly in sync when we render other view families, this is ok as this
+	// invalidation mechanism is only used for interactive rendering where we expect to be constantly drawing the scene.
+	// Therefore it is acceptable for some view families to be a frame or so behind others.
+	uint32 CurrentPathTracingInvalidationCounter = Scene->PathTracingInvalidationCounter.Load();
+
 	// Setup motion blur parameters (also check for camera movement thresholds)
 	for(int32 ViewIndex = 0;ViewIndex < Views.Num();ViewIndex++)
 	{
@@ -3858,11 +3864,10 @@ void FSceneRenderer::PreVisibilityFrameSetup(FRDGBuilder& GraphBuilder, const FS
 			{
 				// for interactive usage - any movement or scene change should restart the path tracer
 
-				// We use an atomic to signal path tracer invalidation easily from other threads (like the game thread). This means
-				// we need to be very careful to read and reset the atomic flag in a single step to avoid missing anything. This is
-				// also why we check this field first before looking at the other conditions.
-				bool bNeedsInvalidation = true;
-				if (Scene->bPathTracingNeedsInvalidation.CompareExchange(bNeedsInvalidation, false) ||
+				// For each view, we remember what the invalidation counter was set to last time we were here so we can catch all changes
+				bool bNeedsInvalidation = ViewState->PathTracingInvalidationCounter != CurrentPathTracingInvalidationCounter;
+				ViewState->PathTracingInvalidationCounter = CurrentPathTracingInvalidationCounter;
+				if (bNeedsInvalidation ||
 					bResetCamera ||
 					bIsProjMatrixDifferent ||
 					bIsCameraMove ||
