Distance Field LWC support #1
- Upload bounds tile + relative world position
- Upload object tile + relative world matrices
- Positions are still converted to float when loaded in shader.

diff --git a/Engine/Shaders/Private/DistanceFieldLightingShared.ush b/Engine/Shaders/Private/DistanceFieldLightingShared.ush
--- a/Engine/Shaders/Private/DistanceFieldLightingShared.ush
+++ b/Engine/Shaders/Private/DistanceFieldLightingShared.ush
@@ -141,8 +141,8 @@ uint GetCulledNumObjects()
 }
 
 // In float4's.  Must match equivalent C++ variables.
-#define DF_OBJECT_BOUNDS_STRIDE 2
-#define DF_OBJECT_DATA_STRIDE 9
+#define DF_OBJECT_BOUNDS_STRIDE 3
+#define DF_OBJECT_DATA_STRIDE 10
 #define HEIGHTFIELD_OBJECT_BOUNDS_STRIDE 2
 #define HEIGHTFIELD_OBJECT_DATA_STRIDE 6
 
@@ -170,13 +170,18 @@ FDFObjectBounds LoadDFObjectBounds(uint ObjectIndex)
 	FDFObjectBounds Bounds;
 
 	float4 Vector0 = SceneObjectBounds[ObjectIndex * DF_OBJECT_BOUNDS_STRIDE + 0];
-	Bounds.Center = Vector0.xyz;
-	Bounds.SphereRadius = Vector0.w;
+	float3 TilePosition = Vector0.xyz;
 
 	float4 Vector1 = SceneObjectBounds[ObjectIndex * DF_OBJECT_BOUNDS_STRIDE + 1];
-	Bounds.BoxExtent = Vector1.xyz;
+	float3 RelativeWorldPosition = Vector1.xyz;
+
+	Bounds.Center = LWCHackToFloat(MakeLWCVector3(TilePosition, RelativeWorldPosition)); // DF_LWC_TODO
+	Bounds.SphereRadius = Vector1.w;
+
+	float4 Vector2 = SceneObjectBounds[ObjectIndex * DF_OBJECT_BOUNDS_STRIDE + 2];
+	Bounds.BoxExtent = Vector2.xyz;
 
-	uint Flags = asuint(Vector1.w);
+	uint Flags = asuint(Vector2.w);
 	Bounds.OftenMoving = Flags & 1U;
 	Bounds.bCastShadow = (Flags & 2U) != 0U;
 	Bounds.bIsNaniteMesh = (Flags & 4U) != 0U;
@@ -239,29 +244,35 @@ struct FDFObjectData
 
 FDFObjectData LoadDFObjectDataFromBuffer(StructuredBuffer<float4> SourceBuffer, uint ObjectIndex)
 {
+	float3 TilePosition = SourceBuffer[ObjectIndex * DF_OBJECT_DATA_STRIDE + 0].xyz;
+
 	FDFObjectData Data;
 
-	float4 V0 = SourceBuffer[ObjectIndex * DF_OBJECT_DATA_STRIDE + 0];
-	float4 V1 = SourceBuffer[ObjectIndex * DF_OBJECT_DATA_STRIDE + 1];
-	float4 V2 = SourceBuffer[ObjectIndex * DF_OBJECT_DATA_STRIDE + 2];
-	Data.WorldToVolume = transpose(float4x4(V0, V1, V2, float4(0.0f, 0.0f, 0.0f, 1.0f)));
+	float4 V0 = SourceBuffer[ObjectIndex * DF_OBJECT_DATA_STRIDE + 1];
+	float4 V1 = SourceBuffer[ObjectIndex * DF_OBJECT_DATA_STRIDE + 2];
+	float4 V2 = SourceBuffer[ObjectIndex * DF_OBJECT_DATA_STRIDE + 3];
+	float4x4 RelativeWorldToVolume = transpose(float4x4(V0, V1, V2, float4(0.0f, 0.0f, 0.0f, 1.0f)));
+
+	Data.WorldToVolume = LWCHackToFloat(MakeLWCInverseMatrix4x3(TilePosition, RelativeWorldToVolume)); // DF_LWC_TODO
 
-	float4 Vector3 = SourceBuffer[ObjectIndex * DF_OBJECT_DATA_STRIDE + 3];
+	float4 Vector3 = SourceBuffer[ObjectIndex * DF_OBJECT_DATA_STRIDE + 4];
 	Data.VolumePositionExtent = Vector3.xyz;
 	Data.VolumeSurfaceBias = abs(Vector3.w);
 	Data.bMostlyTwoSided = Vector3.w < 0.0f;
 
-	float4 Vector4 = SourceBuffer[ObjectIndex * DF_OBJECT_DATA_STRIDE + 4];
+	float4 Vector4 = SourceBuffer[ObjectIndex * DF_OBJECT_DATA_STRIDE + 5];
 	Data.MinMaxDrawDistance2 = Vector4.xy;
 	Data.SelfShadowBias = Vector4.z;
 	Data.GPUSceneInstanceIndex = asuint(Vector4.w);
 	
-	V0 = SourceBuffer[ObjectIndex * DF_OBJECT_DATA_STRIDE + 5];
-	V1 = SourceBuffer[ObjectIndex * DF_OBJECT_DATA_STRIDE + 6];
-	V2 = SourceBuffer[ObjectIndex * DF_OBJECT_DATA_STRIDE + 7];
-	Data.VolumeToWorld = transpose(float4x4(V0, V1, V2, float4(0.0f, 0.0f, 0.0f, 1.0f)));
+	V0 = SourceBuffer[ObjectIndex * DF_OBJECT_DATA_STRIDE + 6];
+	V1 = SourceBuffer[ObjectIndex * DF_OBJECT_DATA_STRIDE + 7];
+	V2 = SourceBuffer[ObjectIndex * DF_OBJECT_DATA_STRIDE + 8];
+	float4x4 VolumeToRelativeWorld = transpose(float4x4(V0, V1, V2, float4(0.0f, 0.0f, 0.0f, 1.0f)));
+
+	Data.VolumeToWorld = LWCHackToFloat(MakeLWCMatrix(TilePosition, VolumeToRelativeWorld)); // DF_LWC_TODO
 
-	float4 Vector8 = SourceBuffer[ObjectIndex * DF_OBJECT_DATA_STRIDE + 8];
+	float4 Vector8 = SourceBuffer[ObjectIndex * DF_OBJECT_DATA_STRIDE + 9];
 	Data.VolumeToWorldScale = Vector8.xyz;
 	Data.VolumeScale = min(Data.VolumeToWorldScale.x, min(Data.VolumeToWorldScale.y, Data.VolumeToWorldScale.z));
 	Data.AssetIndex = asuint(Vector8.w);
diff --git a/Engine/Source/Runtime/Renderer/Private/DistanceFieldLightingShared.h b/Engine/Source/Runtime/Renderer/Private/DistanceFieldLightingShared.h
--- a/Engine/Source/Runtime/Renderer/Private/DistanceFieldLightingShared.h
+++ b/Engine/Source/Runtime/Renderer/Private/DistanceFieldLightingShared.h
@@ -44,8 +44,8 @@ enum EDistanceFieldPrimitiveType
 };
 
 // Must match equivalent shader defines
-static const int32 GDistanceFieldObjectDataStride = 9;
-static const int32 GDistanceFieldObjectBoundsStride = 2;
+static const int32 GDistanceFieldObjectDataStride = 10;
+static const int32 GDistanceFieldObjectBoundsStride = 3;
 
 static const int32 GHeightFieldObjectDataStride = 6;
 static const int32 GHeightFieldObjectBoundsStride = 2;
diff --git a/Engine/Source/Runtime/Renderer/Private/DistanceFieldObjectManagement.cpp b/Engine/Source/Runtime/Renderer/Private/DistanceFieldObjectManagement.cpp
--- a/Engine/Source/Runtime/Renderer/Private/DistanceFieldObjectManagement.cpp
+++ b/Engine/Source/Runtime/Renderer/Private/DistanceFieldObjectManagement.cpp
@@ -715,23 +715,29 @@ void FDistanceFieldSceneData::UpdateDistanceFieldObjectBuffers(
 									const FBox LocalSpaceMeshBounds = DistanceFieldData->LocalSpaceMeshBounds;
 			
 									const FMatrix LocalToWorld = PrimAndInst.LocalToWorld;
-									const FBox WorldSpaceMeshBounds = LocalSpaceMeshBounds.TransformBy(LocalToWorld);
 
-									const FVector4f ObjectBoundingSphere((FVector3f)WorldSpaceMeshBounds.GetCenter(), WorldSpaceMeshBounds.GetExtent().Size());
-									
-									UploadObjectBounds[0] = ObjectBoundingSphere;
+									{
+										const FBox WorldSpaceMeshBounds = LocalSpaceMeshBounds.TransformBy(LocalToWorld);
+
+										const FLargeWorldRenderPosition AbsoluteWorldPosition(WorldSpaceMeshBounds.GetCenter());
 
-									const FGlobalDFCacheType CacheType = PrimitiveSceneProxy->IsOftenMoving() ? GDF_Full : GDF_MostlyStatic;
-									const uint32 bOftenMoving = CacheType == GDF_Full;
-									const uint32 bCastShadow = PrimitiveSceneProxy->CastsDynamicShadow();
-									const uint32 bIsNaniteMesh = PrimitiveSceneProxy->IsNaniteMesh() ? 1U : 0U;
-									const uint32 bEmissiveLightSource = PrimitiveSceneProxy->IsEmissiveLightSource();
+										const FVector4f ObjectBoundingSphere(AbsoluteWorldPosition.GetOffset(), WorldSpaceMeshBounds.GetExtent().Size());
 
-									const uint32 Flags = bOftenMoving | (bCastShadow << 1U) | (bIsNaniteMesh << 2U) | (bEmissiveLightSource << 4U);
+										UploadObjectBounds[0] = AbsoluteWorldPosition.GetTile();
+										UploadObjectBounds[1] = ObjectBoundingSphere;
 
-									FVector4f ObjectWorldExtentAndFlags((FVector3f)WorldSpaceMeshBounds.GetExtent(), 0.0f);
-									ObjectWorldExtentAndFlags.W = *(const float*)&Flags;
-									UploadObjectBounds[1] = ObjectWorldExtentAndFlags;
+										const FGlobalDFCacheType CacheType = PrimitiveSceneProxy->IsOftenMoving() ? GDF_Full : GDF_MostlyStatic;
+										const uint32 bOftenMoving = CacheType == GDF_Full;
+										const uint32 bCastShadow = PrimitiveSceneProxy->CastsDynamicShadow();
+										const uint32 bIsNaniteMesh = PrimitiveSceneProxy->IsNaniteMesh() ? 1U : 0U;
+										const uint32 bEmissiveLightSource = PrimitiveSceneProxy->IsEmissiveLightSource();
+
+										const uint32 Flags = bOftenMoving | (bCastShadow << 1U) | (bIsNaniteMesh << 2U) | (bEmissiveLightSource << 4U);
+
+										FVector4f ObjectWorldExtentAndFlags((FVector3f)WorldSpaceMeshBounds.GetExtent(), 0.0f);
+										ObjectWorldExtentAndFlags.W = *(const float*)&Flags;
+										UploadObjectBounds[2] = ObjectWorldExtentAndFlags;
+									}
 
 									// Uniformly scale our Volume space to lie within [-1, 1] at the max extent
 									// This is mirrored in the SDF encoding
@@ -739,13 +745,22 @@ void FDistanceFieldSceneData::UpdateDistanceFieldObjectBuffers(
 
 									const FMatrix VolumeToWorld = FScaleMatrix(1.0f / LocalToVolumeScale) * FTranslationMatrix(LocalSpaceMeshBounds.GetCenter()) * LocalToWorld;
 
-									const FVector VolumePositionExtent = LocalSpaceMeshBounds.GetExtent() * LocalToVolumeScale;
+									const FLargeWorldRenderPosition WorldPosition(VolumeToWorld.GetOrigin());
+									const FVector TilePositionOffset = WorldPosition.GetTileOffset();
+
+									// Inverse on FMatrix44f can generate NaNs if the source matrix contains large scaling, so do it in double precision.
+									FMatrix LocalToRelativeWorld = FLargeWorldRenderScalar::MakeToRelativeWorldMatrixDouble(TilePositionOffset, VolumeToWorld);
+
+									// TilePosition
+									UploadObjectData[0] = WorldPosition.GetTile();
 
-									const FMatrix44f WorldToVolumeT = FMatrix44f(VolumeToWorld.Inverse().GetTransposed()); // DF_LWC_TODO
+									const FMatrix44f WorldToVolumeT = FMatrix44f(LocalToRelativeWorld.Inverse().GetTransposed());
 									// WorldToVolumeT
-									UploadObjectData[0] = (*(FVector4f*)&WorldToVolumeT.M[0]);
-									UploadObjectData[1] = (*(FVector4f*)&WorldToVolumeT.M[1]);
-									UploadObjectData[2] = (*(FVector4f*)&WorldToVolumeT.M[2]);
+									UploadObjectData[1] = (*(FVector4f*)&WorldToVolumeT.M[0]);
+									UploadObjectData[2] = (*(FVector4f*)&WorldToVolumeT.M[1]);
+									UploadObjectData[3] = (*(FVector4f*)&WorldToVolumeT.M[2]);
+
+									const FVector VolumePositionExtent = LocalSpaceMeshBounds.GetExtent() * LocalToVolumeScale;
 
 									// Minimal surface bias which increases chance that ray hit will a surface located between two texels
 									float ExpandSurfaceDistance = (GMeshSDFSurfaceBiasExpand * VolumePositionExtent / FVector(DistanceFieldData->Mips[0].IndirectionDimensions * DistanceField::UniqueDataBrickSize)).Size();
@@ -756,7 +771,7 @@ void FDistanceFieldSceneData::UpdateDistanceFieldObjectBuffers(
 									}
 
 									const float WSign = DistanceFieldData->bMostlyTwoSided ? -1 : 1;
-									UploadObjectData[3] = FVector4f((FVector3f)VolumePositionExtent, WSign * ExpandSurfaceDistance);
+									UploadObjectData[4] = FVector4f((FVector3f)VolumePositionExtent, WSign * ExpandSurfaceDistance);
 
 									const int32 PrimIdx = PrimAndInst.Primitive->GetIndex();
 									const FPrimitiveBounds& PrimBounds = PrimitiveBounds[PrimIdx];
@@ -775,12 +790,12 @@ void FDistanceFieldSceneData::UpdateDistanceFieldObjectBuffers(
 									Vector4.Y = MaxDrawDist * MaxDrawDist;
 									Vector4.Z = SelfShadowBias;
 									Vector4.W = *(const float*)&GPUSceneInstanceIndex;
-									UploadObjectData[4] = Vector4;
+									UploadObjectData[5] = Vector4;
 
-									const FMatrix44f VolumeToWorldT = FMatrix44f(VolumeToWorld.GetTransposed()); // DF_LWC_TODO
-									UploadObjectData[5] = *(FVector4f*)&VolumeToWorldT.M[0];
-									UploadObjectData[6] = *(FVector4f*)&VolumeToWorldT.M[1];
-									UploadObjectData[7] = *(FVector4f*)&VolumeToWorldT.M[2];
+									const FMatrix44f VolumeToWorldT = FMatrix44f(LocalToRelativeWorld.GetTransposed());
+									UploadObjectData[6] = *(FVector4f*)&VolumeToWorldT.M[0];
+									UploadObjectData[7] = *(FVector4f*)&VolumeToWorldT.M[1];
+									UploadObjectData[8] = *(FVector4f*)&VolumeToWorldT.M[2];
 
 									FVector4f FloatVector8(VolumeToWorld.GetScaleVector(), 0.0f);
 
@@ -790,7 +805,7 @@ void FDistanceFieldSceneData::UpdateDistanceFieldObjectBuffers(
 									const int32 AssetStateInt = AssetStateSetId.AsInteger();
 									FloatVector8.W = *(const float*)&AssetStateInt;
 									
-									UploadObjectData[8] = FloatVector8;
+									UploadObjectData[9] = FloatVector8;
 								}
 							}
 						},
