Fix wrong default resource being bound to LightScatteringHistory.
#rnx

Submitting on behalf of Mihnea Balta

diff --git a/Engine/Source/Runtime/Renderer/Private/VolumetricFog.cpp b/Engine/Source/Runtime/Renderer/Private/VolumetricFog.cpp
--- a/Engine/Source/Runtime/Renderer/Private/VolumetricFog.cpp
+++ b/Engine/Source/Runtime/Renderer/Private/VolumetricFog.cpp
@@ -1146,7 +1146,7 @@ void FDeferredShadingSceneRenderer::ComputeVolumetricFog(FRDGBuilder& GraphBuild
 			PassParameters->GlobalDistanceFieldParameters = SetupGlobalDistanceFieldParameters(View.GlobalDistanceFieldInfo.ParameterData);
 
 			FVolumetricCloudRenderSceneInfo* CloudInfo = Scene->GetVolumetricCloudSceneInfo();
-			FRDGTexture* LightScatteringHistoryRDGTexture = BlackDummyTexture;
+			FRDGTexture* LightScatteringHistoryRDGTexture = VolumetricBlackDummyTexture;
 			if (bUseTemporalReprojection && View.ViewState->LightScatteringHistory.IsValid())
 			{
 				LightScatteringHistoryRDGTexture = GraphBuilder.RegisterExternalTexture(View.ViewState->LightScatteringHistory);
