RayTracing: Add flag to optionally allow ray tracing per view.

diff --git a/Engine/Source/Runtime/Engine/Public/SceneView.h b/Engine/Source/Runtime/Engine/Public/SceneView.h
--- a/Engine/Source/Runtime/Engine/Public/SceneView.h
+++ b/Engine/Source/Runtime/Engine/Public/SceneView.h
@@ -1254,6 +1254,9 @@ public:
 
 #if RHI_RAYTRACING
 	FIESLightProfileResource* IESLightProfileResource;
+
+	/** Use to allow ray tracing on this view. */
+	bool bAllowRayTracing = false;
 #endif
 
 protected:
diff --git a/Engine/Source/Runtime/Renderer/Private/DeferredShadingRenderer.cpp b/Engine/Source/Runtime/Renderer/Private/DeferredShadingRenderer.cpp
--- a/Engine/Source/Runtime/Renderer/Private/DeferredShadingRenderer.cpp
+++ b/Engine/Source/Runtime/Renderer/Private/DeferredShadingRenderer.cpp
@@ -3361,7 +3361,7 @@ void FDeferredShadingSceneRenderer::Render(FRDGBuilder& GraphBuilder)
 
 bool AnyRayTracingPassEnabled(const FScene* Scene, const FViewInfo& View)
 {
-	if (!IsRayTracingEnabled() || Scene == nullptr)
+	if (!IsRayTracingEnabled() || Scene == nullptr || !View.bAllowRayTracing)
 	{
 		return false;
 	}
