REL - add extra flushes to prevent TDR on DX12

diff --git a/Engine/Source/Runtime/Renderer/Private/GPUBenchmark.cpp b/Engine/Source/Runtime/Renderer/Private/GPUBenchmark.cpp
--- a/Engine/Source/Runtime/Renderer/Private/GPUBenchmark.cpp
+++ b/Engine/Source/Runtime/Renderer/Private/GPUBenchmark.cpp
@@ -406,7 +406,7 @@ void RendererGPUBenchmark(FRHICommandListImmediate& RHICmdList, FSynthBenchmarkR
 		return;
 	}
 
-	
+	RHICmdList.ImmediateFlush(EImmediateFlushType::FlushRHIThread);
 	MeasureLongGPUTaskExecutionTime(RHICmdList);
 	RHICmdList.ImmediateFlush(EImmediateFlushType::FlushRHIThread);	
 
@@ -537,6 +537,10 @@ void RendererGPUBenchmark(FRHICommandListImmediate& RHICmdList, FSynthBenchmarkR
 		// multiple iterations to see how trust able the values are
 		for(uint32 Iteration = 0; Iteration < IterationCount; ++Iteration)
 		{
+			if(Iteration % 10 == 0)
+			{
+				RHICmdList.ImmediateFlush(EImmediateFlushType::FlushRHIThread);
+			}
 			SCOPED_DRAW_EVENTF(RHICmdList, Benchmark, TEXT("Iteration:%d"), Iteration);
 
 			for(uint32 MethodIterator = 0; MethodIterator < MethodCount; ++MethodIterator)
@@ -580,6 +584,11 @@ void RendererGPUBenchmark(FRHICommandListImmediate& RHICmdList, FSynthBenchmarkR
 
 			for(uint32 Iteration = 0; Iteration < IterationCount; ++Iteration)
 			{
+				if(Iteration % 10 == 0)
+				{
+					RHICmdList.ImmediateFlush(EImmediateFlushType::FlushRHIThread);
+				}
+
 				uint32 Results[MethodCount];
 
 				for(uint32 MethodId = 0; MethodId < MethodCount; ++MethodId)
