Handle missing shader in FRayTracingDynamicGeometryCollection::AddDynamicMeshBatchForGeometryUpdate:
* If during shader compilation we end up using a fallback material that does not support FRayTracingDynamicGeometryConverterCS we skip that MeshBatch.

diff --git a/Engine/Source/Runtime/Renderer/Private/RayTracing/RayTracingDynamicGeometry.cpp b/Engine/Source/Runtime/Renderer/Private/RayTracing/RayTracingDynamicGeometry.cpp
--- a/Engine/Source/Runtime/Renderer/Private/RayTracing/RayTracingDynamicGeometry.cpp
+++ b/Engine/Source/Runtime/Renderer/Private/RayTracing/RayTracingDynamicGeometry.cpp
@@ -227,7 +227,19 @@ void FRayTracingDynamicGeometryCollection::AddDynamicMeshBatchForGeometryUpdate(
 
 		FMeshComputeDispatchCommand DispatchCmd;
 
-		TShaderRef<FRayTracingDynamicGeometryConverterCS> Shader = Material.GetShader<FRayTracingDynamicGeometryConverterCS>(MeshBatch.VertexFactory->GetType());
+		
+		FMaterialShaderTypes ShaderTypes;
+		ShaderTypes.AddShaderType<FRayTracingDynamicGeometryConverterCS>();
+
+		FMaterialShaders MaterialShaders;
+		if (!Material.TryGetShaders(ShaderTypes, MeshBatch.VertexFactory->GetType(), MaterialShaders))
+		{
+			continue;
+		}
+
+		TShaderRef<FRayTracingDynamicGeometryConverterCS> Shader;
+		MaterialShaders.TryGetShader(SF_Compute, Shader);
+
 		DispatchCmd.MaterialShader = Shader;
 		FMeshDrawShaderBindings& ShaderBindings = DispatchCmd.ShaderBindings;
 
