Don't disable AutoExposure when VisualizeGlobalDistanceField is enabled.
- Visualization was updated to apply lighting and now requires auto exposure / postfx.

diff --git a/Engine/Source/Runtime/Renderer/Private/PostProcess/PostProcessEyeAdaptation.cpp b/Engine/Source/Runtime/Renderer/Private/PostProcess/PostProcessEyeAdaptation.cpp
--- a/Engine/Source/Runtime/Renderer/Private/PostProcess/PostProcessEyeAdaptation.cpp
+++ b/Engine/Source/Runtime/Renderer/Private/PostProcess/PostProcessEyeAdaptation.cpp
@@ -238,7 +238,6 @@ bool IsAutoExposureDebugMode(const FViewInfo& View)
 		(EngineShowFlags.VisualizeVirtualShadowMap && View.CurrentVirtualShadowMapVisualizationMode != NAME_None) ||
 		EngineShowFlags.RayTracingDebug ||
 		EngineShowFlags.VisualizeDistanceFieldAO ||
-		EngineShowFlags.VisualizeGlobalDistanceField ||
 		EngineShowFlags.VisualizeVolumetricCloudConservativeDensity ||
 		EngineShowFlags.CollisionVisibility ||
 		EngineShowFlags.CollisionPawn ||
diff --git a/Engine/Source/Runtime/Renderer/Private/PostProcess/PostProcessing.cpp b/Engine/Source/Runtime/Renderer/Private/PostProcess/PostProcessing.cpp
--- a/Engine/Source/Runtime/Renderer/Private/PostProcess/PostProcessing.cpp
+++ b/Engine/Source/Runtime/Renderer/Private/PostProcess/PostProcessing.cpp
@@ -151,7 +151,6 @@ bool IsPostProcessingEnabled(const FViewInfo& View)
 			 View.Family->EngineShowFlags.PostProcessing &&
 			!View.Family->EngineShowFlags.VisualizeDistanceFieldAO &&
 			!View.Family->EngineShowFlags.VisualizeShadingModels &&
-			!View.Family->EngineShowFlags.VisualizeGlobalDistanceField &&
 			!View.Family->EngineShowFlags.VisualizeVolumetricCloudConservativeDensity &&
 			!View.Family->EngineShowFlags.ShaderComplexity;
 	}
