MoviePipeline: Fixed a UI warning from Sky Materials showing up in the final rendered frames, instead of just in the viewport.

#rnx

diff --git a/Engine/Source/Runtime/Renderer/Private/ReflectionEnvironmentRealTimeCapture.cpp b/Engine/Source/Runtime/Renderer/Private/ReflectionEnvironmentRealTimeCapture.cpp
--- a/Engine/Source/Runtime/Renderer/Private/ReflectionEnvironmentRealTimeCapture.cpp
+++ b/Engine/Source/Runtime/Renderer/Private/ReflectionEnvironmentRealTimeCapture.cpp
@@ -244,6 +244,11 @@ void FScene::ValidateSkyLightRealTimeCapture(
 	FRDGTextureRef SceneColorTexture)
 {
 #if WITH_EDITOR
+	if (!GAreScreenMessagesEnabled)
+	{
+		return;
+	}
+
 	bool bSkyMeshInMainPassExist = false;
 	bool bSkyMeshInRealTimeSkyCaptureExtist = false;
 
