Fix shadow artifacts in the right eye (UE-136749).

diff --git a/Engine/Source/Runtime/Renderer/Private/MobileShadingRenderer.cpp b/Engine/Source/Runtime/Renderer/Private/MobileShadingRenderer.cpp
--- a/Engine/Source/Runtime/Renderer/Private/MobileShadingRenderer.cpp
+++ b/Engine/Source/Runtime/Renderer/Private/MobileShadingRenderer.cpp
@@ -548,7 +548,7 @@ void FMobileSceneRenderer::InitViews(FRDGBuilder& GraphBuilder, FSceneTexturesCo
 	}
 
 	// initialize per-view uniform buffer.  Pass in shadow info as necessary.
-	for (int32 ViewIndex = 0; ViewIndex < Views.Num(); ViewIndex++)
+	for (int32 ViewIndex = Views.Num() - 1; ViewIndex >= 0; --ViewIndex)
 	{
 		FViewInfo& View = Views[ViewIndex];
 		
diff --git a/Engine/Source/Runtime/Renderer/Private/SceneRendering.cpp b/Engine/Source/Runtime/Renderer/Private/SceneRendering.cpp
--- a/Engine/Source/Runtime/Renderer/Private/SceneRendering.cpp
+++ b/Engine/Source/Runtime/Renderer/Private/SceneRendering.cpp
@@ -1992,6 +1992,23 @@ void FViewInfo::InitRHIResources()
 		*CachedViewUniformShaderParameters);
 
 	ViewUniformBuffer = TUniformBufferRef<FViewUniformShaderParameters>::CreateUniformBufferImmediate(*CachedViewUniformShaderParameters, UniformBuffer_SingleFrame);
+	if (bShouldBindInstancedViewUB)
+	{
+		if (const FViewInfo* InstancedView = GetInstancedView())
+		{
+			checkf(InstancedView->CachedViewUniformShaderParameters.IsValid(), TEXT("Instanced view should have had its RHI resources initialized first. Check InitViews order."));
+			InstancedViewUniformBuffer = TUniformBufferRef<FInstancedViewUniformShaderParameters>::CreateUniformBufferImmediate(
+				reinterpret_cast<FInstancedViewUniformShaderParameters&>(*InstancedView->CachedViewUniformShaderParameters),
+				UniformBuffer_SingleFrame);
+		}
+		else
+		{
+			// If we don't render this view in stereo, we simply initialize with the existing contents.
+			InstancedViewUniformBuffer = TUniformBufferRef<FInstancedViewUniformShaderParameters>::CreateUniformBufferImmediate(
+				reinterpret_cast<FInstancedViewUniformShaderParameters&>(*CachedViewUniformShaderParameters),
+				UniformBuffer_SingleFrame);
+		}
+	}
 
 	const int32 TranslucencyLightingVolumeDim = GetTranslucencyLightingVolumeDim();
 
@@ -2014,23 +2031,6 @@ void FViewInfo::BeginRenderView() const
 	{
 		Extension->BeginRenderView(this, bShouldWaitForPersistentViewUniformBufferExtensionsJobs);
 	}
-
-	if (!InstancedViewUniformBuffer)
-	{
-		if (const FViewInfo* InstancedView = GetInstancedView())
-		{
-			InstancedViewUniformBuffer = TUniformBufferRef<FInstancedViewUniformShaderParameters>::CreateUniformBufferImmediate(
-				reinterpret_cast<FInstancedViewUniformShaderParameters&>(*InstancedView->CachedViewUniformShaderParameters),
-				UniformBuffer_SingleFrame);
-		}
-		else
-		{
-			// If we don't render this view in stereo, we simply initialize with the existing contents.
-			InstancedViewUniformBuffer = TUniformBufferRef<FInstancedViewUniformShaderParameters>::CreateUniformBufferImmediate(
-				reinterpret_cast<FInstancedViewUniformShaderParameters&>(*CachedViewUniformShaderParameters),
-				UniformBuffer_SingleFrame);
-		}
-	}
 }
 
 FViewShaderParameters FViewInfo::GetShaderParameters() const
@@ -2038,6 +2038,8 @@ FViewShaderParameters FViewInfo::GetShaderParameters() const
 	FViewShaderParameters Parameters;
 	Parameters.View = ViewUniformBuffer;
 	Parameters.InstancedView = InstancedViewUniformBuffer;
+	// if we're a part of the stereo pair, make sure that the pointer isn't bogus
+	checkf(InstancedViewUniformBuffer.IsValid() || !bShouldBindInstancedViewUB, TEXT("A view that is a part of the stereo pair has bogus state for InstancedView."));
 	return Parameters;
 }
 
diff --git a/Engine/Source/Runtime/Renderer/Private/SceneRendering.h b/Engine/Source/Runtime/Renderer/Private/SceneRendering.h
--- a/Engine/Source/Runtime/Renderer/Private/SceneRendering.h
+++ b/Engine/Source/Runtime/Renderer/Private/SceneRendering.h
@@ -1696,7 +1696,7 @@ private:
 	void CalcTranslucencyLightingVolumeBounds(FBox* InOutCascadeBoundsArray, int32 NumCascades) const;
 
 	/** Instanced view uniform buffer held on the primary view. */
-	mutable TUniformBufferRef<FInstancedViewUniformShaderParameters> InstancedViewUniformBuffer;
+	TUniformBufferRef<FInstancedViewUniformShaderParameters> InstancedViewUniformBuffer;
 };
 
 
diff --git a/Engine/Source/Runtime/Renderer/Private/SceneVisibility.cpp b/Engine/Source/Runtime/Renderer/Private/SceneVisibility.cpp
--- a/Engine/Source/Runtime/Renderer/Private/SceneVisibility.cpp
+++ b/Engine/Source/Runtime/Renderer/Private/SceneVisibility.cpp
@@ -4926,8 +4926,8 @@ void FDeferredShadingSceneRenderer::InitViews(FRDGBuilder& GraphBuilder, const F
 
 	{
 		QUICK_SCOPE_CYCLE_COUNTER(STAT_InitViews_InitRHIResources);
-		// initialize per-view uniform buffer.
-		for (int32 ViewIndex = 0; ViewIndex < Views.Num(); ViewIndex++)
+		// initialize per-view uniform buffer. Do it from back to front because secondary stereo view follows its primary one, but primary needs to know the instanced's params
+		for (int32 ViewIndex = Views.Num() - 1; ViewIndex >= 0; --ViewIndex)
 		{
 			FViewInfo& View = Views[ViewIndex];
 
diff --git a/Engine/Source/Runtime/Renderer/Private/ShadowRendering.cpp b/Engine/Source/Runtime/Renderer/Private/ShadowRendering.cpp
--- a/Engine/Source/Runtime/Renderer/Private/ShadowRendering.cpp
+++ b/Engine/Source/Runtime/Renderer/Private/ShadowRendering.cpp
@@ -1980,7 +1980,9 @@ void FSceneRenderer::RenderShadowProjections(
 		{
 			continue;
 		}
-		
+
+		View.BeginRenderView();
+
 		// Sanity check
 		if (bSubPixelShadow)
 		{
