Skip post process material pass when all the pixels will fail the stencil test of the material. Corrects certain stencil test operations from executing full screen when not needed. Saves performance on platforms that are doing the stencil test in the pixel shader with a clip.

#rnx

diff --git a/Engine/Source/Runtime/Renderer/Private/PostProcess/PostProcessMaterial.cpp b/Engine/Source/Runtime/Renderer/Private/PostProcess/PostProcessMaterial.cpp
--- a/Engine/Source/Runtime/Renderer/Private/PostProcess/PostProcessMaterial.cpp
+++ b/Engine/Source/Runtime/Renderer/Private/PostProcess/PostProcessMaterial.cpp
@@ -639,46 +639,46 @@ FScreenPassTexture AddPostProcessMaterialPass(
 
 	if (!bSkipPostProcess)
 	{
-		AddDrawScreenPass(
-			GraphBuilder,
-			RDG_EVENT_NAME("PostProcessMaterial"),
-			View,
-			OutputViewport,
-			SceneColorViewport,
-			FScreenPassPipelineState(VertexShader, PixelShader, BlendState, DepthStencilState, MaterialStencilRef),
-			PostProcessMaterialParameters,
-			ScreenPassFlags,
-			[&View, VertexShader, PixelShader, MaterialRenderProxy, Material, PostProcessMaterialParameters, MaterialStencilRef](FRHICommandList& RHICmdList)
-			{
-				FPostProcessMaterialVS::SetParameters(RHICmdList, VertexShader, View, MaterialRenderProxy, *Material, *PostProcessMaterialParameters);
-				FPostProcessMaterialPS::SetParameters(RHICmdList, PixelShader, View, MaterialRenderProxy, *Material, *PostProcessMaterialParameters);
-			});
-
-		if (bForceIntermediateTarget && !bCompositeWithInputAndDecode)
+	AddDrawScreenPass(
+		GraphBuilder,
+		RDG_EVENT_NAME("PostProcessMaterial"),
+		View,
+		OutputViewport,
+		SceneColorViewport,
+		FScreenPassPipelineState(VertexShader, PixelShader, BlendState, DepthStencilState, MaterialStencilRef),
+		PostProcessMaterialParameters,
+		ScreenPassFlags,
+		[&View, VertexShader, PixelShader, MaterialRenderProxy, Material, PostProcessMaterialParameters, MaterialStencilRef](FRHICommandList& RHICmdList)
+	{
+		FPostProcessMaterialVS::SetParameters(RHICmdList, VertexShader, View, MaterialRenderProxy, *Material, *PostProcessMaterialParameters);
+		FPostProcessMaterialPS::SetParameters(RHICmdList, PixelShader, View, MaterialRenderProxy, *Material, *PostProcessMaterialParameters);
+	});
+
+	if (bForceIntermediateTarget && !bCompositeWithInputAndDecode)
+	{
+		if (!Inputs.bFlipYAxis)
+		{
+			// We shouldn't get here unless we had an override target.
+			check(Inputs.OverrideOutput.IsValid());
+			AddDrawTexturePass(GraphBuilder, View, Output.Texture, Inputs.OverrideOutput.Texture);
+			Output = Inputs.OverrideOutput;
+		}
+		else
 		{
-			if (!Inputs.bFlipYAxis)
+			FScreenPassRenderTarget TempTarget = Output;
+			if (Inputs.OverrideOutput.IsValid())
 			{
-				// We shouldn't get here unless we had an override target.
-				check(Inputs.OverrideOutput.IsValid());
-				AddDrawTexturePass(GraphBuilder, View, Output.Texture, Inputs.OverrideOutput.Texture);
 				Output = Inputs.OverrideOutput;
 			}
 			else
 			{
-				FScreenPassRenderTarget TempTarget = Output;
-				if (Inputs.OverrideOutput.IsValid())
-				{
-					Output = Inputs.OverrideOutput;
-				}
-				else
-				{
-					Output = FScreenPassRenderTarget(SceneColor, ERenderTargetLoadAction::ENoAction);
-				}
-
-				AddCopyAndFlipTexturePass(GraphBuilder, View, TempTarget.Texture, Output.Texture);
+				Output = FScreenPassRenderTarget(SceneColor, ERenderTargetLoadAction::ENoAction);
 			}
+
+			AddCopyAndFlipTexturePass(GraphBuilder, View, TempTarget.Texture, Output.Texture);
 		}
 	}
+	}
 
 	return MoveTemp(Output);
 }
