[Backout] - CL19146706
Original CL Desc
-----------------------------------------------------------------
Skip post process material pass when all the pixels will fail the stencil test of the material. Corrects certain stencil test operations from executing full screen when not needed. Saves performance on platforms that are doing the stencil test in the pixel shader with a clip.

#rnx

diff --git a/Engine/Source/Runtime/Renderer/Private/PostProcess/PostProcessMaterial.cpp b/Engine/Source/Runtime/Renderer/Private/PostProcess/PostProcessMaterial.cpp
--- a/Engine/Source/Runtime/Renderer/Private/PostProcess/PostProcessMaterial.cpp
+++ b/Engine/Source/Runtime/Renderer/Private/PostProcess/PostProcessMaterial.cpp
@@ -128,40 +128,6 @@ FRHIBlendState* GetMaterialBlendState(const FMaterial* Material)
 	return BlendStates[Material->GetBlendMode()];
 }
 
-bool PostProcessStencilTest(uint32 StencilValue, uint32 StencilComp, uint32 StencilRef)
-{
-	bool bStencilTestPassed = true;
-
-	switch (StencilComp)
-	{
-	case EMaterialStencilCompare::MSC_Less:
-		bStencilTestPassed = (StencilValue < StencilRef);
-		break;
-	case EMaterialStencilCompare::MSC_LessEqual:
-		bStencilTestPassed = (StencilValue <= StencilRef);
-		break;
-	case EMaterialStencilCompare::MSC_GreaterEqual:
-		bStencilTestPassed = (StencilValue >= StencilRef);
-		break;
-	case EMaterialStencilCompare::MSC_Equal:
-		bStencilTestPassed = (StencilValue == StencilRef);
-		break;
-	case EMaterialStencilCompare::MSC_Greater:
-		bStencilTestPassed = (StencilValue > StencilRef);
-		break;
-	case EMaterialStencilCompare::MSC_NotEqual:
-		bStencilTestPassed = (StencilValue != StencilRef);
-		break;
-	case EMaterialStencilCompare::MSC_Never:
-		bStencilTestPassed = false;
-		break;
-	default:
-		break;
-	}
-
-	return !bStencilTestPassed;
-}
-
 class FPostProcessMaterialShader : public FMaterialShader
 {
 public:
@@ -610,8 +576,6 @@ FScreenPassTexture AddPostProcessMaterialPass(
 		ScreenPassFlags |= EScreenPassDrawFlags::FlipYAxis;
 	}
 
-	// check if we can skip that draw call in case if all pixels will fail the stencil test of the material
-	bool bSkipPostProcess = false;
 
 	if (Material->IsStencilTestEnabled() && IsPostProcessStencilTestAllowed())
 	{
diff --git a/Engine/Source/Runtime/Renderer/Private/SceneRendering.h b/Engine/Source/Runtime/Renderer/Private/SceneRendering.h
--- a/Engine/Source/Runtime/Renderer/Private/SceneRendering.h
+++ b/Engine/Source/Runtime/Renderer/Private/SceneRendering.h
@@ -1213,9 +1213,6 @@ public:
 	bool bHasDistortionPrimitives;
 	bool bHasCustomDepthPrimitives;
 
-    /** Get all stencil values written into the custom depth pass */
-	TSet<uint32> CustomDepthStencilValues;
-
 	/** Mesh batches with for mesh decal rendering. */
 	TArray<FMeshDecalBatch, SceneRenderingAllocator> MeshDecalBatches;
 
diff --git a/Engine/Source/Runtime/Renderer/Private/SceneVisibility.cpp b/Engine/Source/Runtime/Renderer/Private/SceneVisibility.cpp
--- a/Engine/Source/Runtime/Renderer/Private/SceneVisibility.cpp
+++ b/Engine/Source/Runtime/Renderer/Private/SceneVisibility.cpp
@@ -2247,7 +2247,6 @@ struct FRelevancePacket
 	TArray<FSkyMeshBatch> SkyMeshBatches;
 	TArray<FSortedTrianglesMeshBatch> SortedTrianglesMeshBatches;
 	FDrawCommandRelevancePacket DrawCommandPacket;
-	TSet<uint32> CustomDepthStencilValues;
 
 	struct FPrimitiveLODMask
 	{
@@ -2458,7 +2457,6 @@ struct FRelevancePacket
 			if (ViewRelevance.bRenderCustomDepth)
 			{
 				bHasCustomDepthPrimitives = true;
-				CustomDepthStencilValues.Add(PrimitiveSceneInfo->Proxy->GetCustomDepthStencilValue());
 			}
 
 			extern bool GUseTranslucencyShadowDepths;
@@ -2882,7 +2880,6 @@ struct FRelevancePacket
 		WriteView.TranslucentPrimCount.Append(TranslucentPrimCount);
 		WriteView.bHasDistortionPrimitives |= bHasDistortionPrimitives;
 		WriteView.bHasCustomDepthPrimitives |= bHasCustomDepthPrimitives;
-		WriteView.CustomDepthStencilValues.Append(CustomDepthStencilValues);
 		WriteView.bUsesCustomDepthStencilInTranslucentMaterials |= bUsesCustomDepthStencil;
 		DirtyIndirectLightingCacheBufferPrimitives.AppendTo(WriteView.DirtyIndirectLightingCacheBufferPrimitives);
 
