RayTracing: Fixed how allowance per view check is made.

There was a crash that could happen when a view had ray tracing disallowed and then forcing all ray tracing effects. This was happening because part of the code thought there were not supposed to be any RT passes and would not add a given shader, but then when the effect was prepared and didn't have an associated view to query its RT allowance then it would later assert when looking for the associated shader.

diff --git a/Engine/Source/Runtime/Renderer/Private/DeferredShadingRenderer.cpp b/Engine/Source/Runtime/Renderer/Private/DeferredShadingRenderer.cpp
--- a/Engine/Source/Runtime/Renderer/Private/DeferredShadingRenderer.cpp
+++ b/Engine/Source/Runtime/Renderer/Private/DeferredShadingRenderer.cpp
@@ -633,7 +633,7 @@ bool FDeferredShadingSceneRenderer::GatherRayTracingWorldInstancesForView(FRDGBu
 	DirtyCachedRayTracingPrimitives.Reserve(Scene->PrimitiveSceneProxies.Num());
 
 	int32 VisiblePrimitives = 0;
-	bool bPerformRayTracing = View.State != nullptr && !View.bIsReflectionCapture;
+	bool bPerformRayTracing = View.State != nullptr && !View.bIsReflectionCapture && View.bAllowRayTracing;
 	if (bPerformRayTracing)
 	{
 		TRACE_CPUPROFILER_EVENT_SCOPE(GatherRayTracingWorldInstances_RelevantPrimitives);
@@ -3361,7 +3361,7 @@ void FDeferredShadingSceneRenderer::Render(FRDGBuilder& GraphBuilder)
 
 bool AnyRayTracingPassEnabled(const FScene* Scene, const FViewInfo& View)
 {
-	if (!IsRayTracingEnabled() || Scene == nullptr || !View.bAllowRayTracing)
+	if (!IsRayTracingEnabled() || Scene == nullptr)
 	{
 		return false;
 	}
@@ -3378,9 +3378,9 @@ bool AnyRayTracingPassEnabled(const FScene* Scene, const FViewInfo& View)
 		|| HasRayTracedOverlay(*View.Family);
 }
 
-bool ShouldRenderRayTracingEffect(bool bEffectEnabled, ERayTracingPipelineCompatibilityFlags CompatibilityFlags)
+bool ShouldRenderRayTracingEffect(bool bEffectEnabled, ERayTracingPipelineCompatibilityFlags CompatibilityFlags, const FSceneView* View)
 {
-	if (!IsRayTracingEnabled())
+	if (!IsRayTracingEnabled() || (View && !View->bAllowRayTracing))
 	{
 		return false;
 	}
diff --git a/Engine/Source/Runtime/Renderer/Private/IndirectLightRendering.cpp b/Engine/Source/Runtime/Renderer/Private/IndirectLightRendering.cpp
--- a/Engine/Source/Runtime/Renderer/Private/IndirectLightRendering.cpp
+++ b/Engine/Source/Runtime/Renderer/Private/IndirectLightRendering.cpp
@@ -545,7 +545,7 @@ bool ShouldRenderPluginRayTracingGlobalIllumination(const FViewInfo& View)
 	FGlobalIlluminationPluginDelegates::FAnyRayTracingPassEnabled& Delegate = FGlobalIlluminationPluginDelegates::AnyRayTracingPassEnabled();
 	Delegate.Broadcast(bAnyRayTracingPassEnabled);
 
-	return ShouldRenderRayTracingEffect(bAnyRayTracingPassEnabled, ERayTracingPipelineCompatibilityFlags::FullPipeline);
+	return ShouldRenderRayTracingEffect(bAnyRayTracingPassEnabled, ERayTracingPipelineCompatibilityFlags::FullPipeline, &View);
 }
 
 void FDeferredShadingSceneRenderer::PrepareRayTracingGlobalIlluminationPlugin(const FViewInfo& View, TArray<FRHIRayTracingShader*>& OutRayGenShaders)
diff --git a/Engine/Source/Runtime/Renderer/Private/LightRendering.cpp b/Engine/Source/Runtime/Renderer/Private/LightRendering.cpp
--- a/Engine/Source/Runtime/Renderer/Private/LightRendering.cpp
+++ b/Engine/Source/Runtime/Renderer/Private/LightRendering.cpp
@@ -162,7 +162,7 @@ bool ShouldRenderRayTracingShadows()
 	const bool bIsStereo = GEngine->StereoRenderingDevice.IsValid() && GEngine->StereoRenderingDevice->IsStereoEnabled();
 	const bool bHairStrands = IsHairStrandsEnabled(EHairStrandsShaderType::Strands);
 
-	return ShouldRenderRayTracingEffect((CVarRayTracingOcclusion.GetValueOnRenderThread() > 0) && !(bIsStereo && bHairStrands), ERayTracingPipelineCompatibilityFlags::FullPipeline);
+	return ShouldRenderRayTracingEffect((CVarRayTracingOcclusion.GetValueOnRenderThread() > 0) && !(bIsStereo && bHairStrands), ERayTracingPipelineCompatibilityFlags::FullPipeline, nullptr);
 }
 
 bool ShouldRenderRayTracingShadowsForLight(const FLightSceneProxy& LightProxy)
diff --git a/Engine/Source/Runtime/Renderer/Private/RayTracing/RayTracingAmbientOcclusion.cpp b/Engine/Source/Runtime/Renderer/Private/RayTracing/RayTracingAmbientOcclusion.cpp
--- a/Engine/Source/Runtime/Renderer/Private/RayTracing/RayTracingAmbientOcclusion.cpp
+++ b/Engine/Source/Runtime/Renderer/Private/RayTracing/RayTracingAmbientOcclusion.cpp
@@ -68,7 +68,7 @@ bool ShouldRenderRayTracingAmbientOcclusion(const FViewInfo& View)
 
 	bEnabled &= (View.FinalPostProcessSettings.RayTracingAOIntensity > 0.0f);
 
-	return ShouldRenderRayTracingEffect(bEnabled, ERayTracingPipelineCompatibilityFlags::FullPipeline);
+	return ShouldRenderRayTracingEffect(bEnabled, ERayTracingPipelineCompatibilityFlags::FullPipeline, &View);
 }
 
 DECLARE_GPU_STAT_NAMED(RayTracingAmbientOcclusion, TEXT("Ray Tracing Ambient Occlusion"));
diff --git a/Engine/Source/Runtime/Renderer/Private/RayTracing/RayTracingGlobalIllumination.cpp b/Engine/Source/Runtime/Renderer/Private/RayTracing/RayTracingGlobalIllumination.cpp
--- a/Engine/Source/Runtime/Renderer/Private/RayTracing/RayTracingGlobalIllumination.cpp
+++ b/Engine/Source/Runtime/Renderer/Private/RayTracing/RayTracingGlobalIllumination.cpp
@@ -395,7 +395,7 @@ bool ShouldRenderRayTracingGlobalIllumination(const FViewInfo& View)
 		? CVarRayTracingGlobalIlluminationValue > 0
 		: View.FinalPostProcessSettings.RayTracingGIType > ERayTracingGlobalIlluminationType::Disabled;
 
-	return ShouldRenderRayTracingEffect(bEnabled, ERayTracingPipelineCompatibilityFlags::FullPipeline);
+	return ShouldRenderRayTracingEffect(bEnabled, ERayTracingPipelineCompatibilityFlags::FullPipeline, &View);
 }
 
 bool IsFinalGatherEnabled(const FViewInfo& View)
diff --git a/Engine/Source/Runtime/Renderer/Private/RayTracing/RayTracingReflections.cpp b/Engine/Source/Runtime/Renderer/Private/RayTracing/RayTracingReflections.cpp
--- a/Engine/Source/Runtime/Renderer/Private/RayTracing/RayTracingReflections.cpp
+++ b/Engine/Source/Runtime/Renderer/Private/RayTracing/RayTracingReflections.cpp
@@ -401,7 +401,7 @@ bool ShouldRenderRayTracingReflections(const FViewInfo& View)
 
 	const bool bReflectionPassEnabled = bReflectionsCvarEnabled && (GetRayTracingReflectionsSamplesPerPixel(View) > 0) && !View.bIsReflectionCapture;
 		
-	return ShouldRenderRayTracingEffect(bReflectionPassEnabled, ERayTracingPipelineCompatibilityFlags::FullPipeline);
+	return ShouldRenderRayTracingEffect(bReflectionPassEnabled, ERayTracingPipelineCompatibilityFlags::FullPipeline, &View);
 }
 
 static bool ShouldRayTracedReflectionsUseHybridReflections()
diff --git a/Engine/Source/Runtime/Renderer/Private/RayTracing/RayTracingTranslucency.cpp b/Engine/Source/Runtime/Renderer/Private/RayTracing/RayTracingTranslucency.cpp
--- a/Engine/Source/Runtime/Renderer/Private/RayTracing/RayTracingTranslucency.cpp
+++ b/Engine/Source/Runtime/Renderer/Private/RayTracing/RayTracingTranslucency.cpp
@@ -112,7 +112,7 @@ FRayTracingPrimaryRaysOptions GetRayTracingTranslucencyOptions(const FViewInfo&
 {
 	FRayTracingPrimaryRaysOptions Options;
 
-	Options.bEnabled = ShouldRenderRayTracingEffect(CVarRayTracingTranslucency.GetValueOnRenderThread() != 0, ERayTracingPipelineCompatibilityFlags::FullPipeline);
+	Options.bEnabled = ShouldRenderRayTracingEffect(CVarRayTracingTranslucency.GetValueOnRenderThread() != 0, ERayTracingPipelineCompatibilityFlags::FullPipeline, &View);
 	Options.SamplerPerPixel = GRayTracingTranslucencySamplesPerPixel >= 0 ? GRayTracingTranslucencySamplesPerPixel : View.FinalPostProcessSettings.RayTracingTranslucencySamplesPerPixel;
 	Options.ApplyHeightFog = GRayTracingTranslucencyHeightFog;
 	Options.PrimaryRayBias = GRayTracingTranslucencyPrimaryRayBias;
@@ -138,7 +138,7 @@ bool ShouldRenderRayTracingTranslucency(const FViewInfo& View)
 		? bViewWithRaytracingTranslucency
 		: RayTracingTranslucencyMode != 0;
 
-	return ShouldRenderRayTracingEffect(bTranslucencyEnabled, ERayTracingPipelineCompatibilityFlags::FullPipeline);
+	return ShouldRenderRayTracingEffect(bTranslucencyEnabled, ERayTracingPipelineCompatibilityFlags::FullPipeline, &View);
 }
 #endif // RHI_RAYTRACING
 
diff --git a/Engine/Source/Runtime/Renderer/Private/RayTracing/RaytracingOptions.h b/Engine/Source/Runtime/Renderer/Private/RayTracing/RaytracingOptions.h
--- a/Engine/Source/Runtime/Renderer/Private/RayTracing/RaytracingOptions.h
+++ b/Engine/Source/Runtime/Renderer/Private/RayTracing/RaytracingOptions.h
@@ -54,7 +54,7 @@ ENUM_CLASS_FLAGS(ERayTracingPipelineCompatibilityFlags);
 #if RHI_RAYTRACING
 
 // Whether a particular effect should be used, taking into account debug override
-extern bool ShouldRenderRayTracingEffect(bool bEffectEnabled, ERayTracingPipelineCompatibilityFlags CompatibilityFlags);
+extern bool ShouldRenderRayTracingEffect(bool bEffectEnabled, ERayTracingPipelineCompatibilityFlags CompatibilityFlags, const FSceneView* View);
 
 extern bool AnyRayTracingPassEnabled(const FScene* Scene, const FViewInfo& View);
 extern bool AnyInlineRayTracingPassEnabled(const FScene* Scene, const FViewInfo& View);
@@ -80,7 +80,7 @@ extern bool CanUseRayTracingAMDHitToken();
 
 #else // RHI_RAYTRACING
 
-FORCEINLINE bool ShouldRenderRayTracingEffect(bool bEffectEnabled, ERayTracingPipelineCompatibilityFlags CompatibilityFlags)
+FORCEINLINE bool ShouldRenderRayTracingEffect(bool bEffectEnabled, ERayTracingPipelineCompatibilityFlags CompatibilityFlags, const FSceneView* View)
 {
 	return false;
 }
@@ -159,5 +159,5 @@ FORCEINLINE bool CanUseRayTracingAMDHitToken()
 
 FORCEINLINE bool ShouldRenderRayTracingEffect(ERayTracingPipelineCompatibilityFlags CompatibilityFlags)
 {
-	return ShouldRenderRayTracingEffect(true, CompatibilityFlags);
+	return ShouldRenderRayTracingEffect(true, CompatibilityFlags, nullptr);
 }
\ No newline at end of file
