Fix height field shadowing
- Was incorrectly casting rows of FMatrix to FVector4f.
- Cast matrix to FMatrix44f first.

diff --git a/Engine/Source/Runtime/Renderer/Private/DistanceFieldObjectManagement.cpp b/Engine/Source/Runtime/Renderer/Private/DistanceFieldObjectManagement.cpp
--- a/Engine/Source/Runtime/Renderer/Private/DistanceFieldObjectManagement.cpp
+++ b/Engine/Source/Runtime/Renderer/Private/DistanceFieldObjectManagement.cpp
@@ -960,7 +960,7 @@ void FSceneRenderer::UpdateGlobalHeightFieldObjectBuffers(FRDGBuilder& GraphBuil
 
 								const FMatrix& LocalToWorld = HeightFieldCompDesc.LocalToWorld;
 								check(LocalToWorld.GetMaximumAxisScale() > 0.f);
-								const FMatrix WorldToLocalT = LocalToWorld.Inverse().GetTransposed();
+								const FMatrix44f WorldToLocalT = LocalToWorld.Inverse().GetTransposed(); // LWC_TODO
 								UploadObjectData[0] = *(const FVector4f*)&WorldToLocalT.M[0];
 								UploadObjectData[1] = *(const FVector4f*)&WorldToLocalT.M[1];
 								UploadObjectData[2] = *(const FVector4f*)&WorldToLocalT.M[2];
