Auto-enable shader print when activating Strata visualization

diff --git a/Engine/Source/Runtime/Renderer/Private/Strata/Strata.cpp b/Engine/Source/Runtime/Renderer/Private/Strata/Strata.cpp
--- a/Engine/Source/Runtime/Renderer/Private/Strata/Strata.cpp
+++ b/Engine/Source/Runtime/Renderer/Private/Strata/Strata.cpp
@@ -359,6 +359,10 @@ static void AddVisualizeMaterialPasses(FRDGBuilder& GraphBuilder, const FViewInf
 	FRHIBlendState* PreMultipliedColorTransmittanceBlend = TStaticBlendState<CW_RGB, BO_Add, BF_One, BF_SourceAlpha, BO_Add, BF_Zero, BF_One>::GetRHI();
 	if (View.Family->EngineShowFlags.VisualizeStrataMaterial)
 	{
+		if (!ShaderPrint::IsEnabled(View)) { ShaderPrint::SetEnabled(true); }
+		ShaderPrint::RequestSpaceForLines(64);
+
+		FSceneTextureParameters SceneTextureParameters = GetSceneTextureParameters(GraphBuilder);
 		FVisualizeMaterialPS::FParameters* PassParameters = GraphBuilder.AllocParameters<FVisualizeMaterialPS::FParameters>();
 		PassParameters->ViewUniformBuffer = View.ViewUniformBuffer;
 		PassParameters->Strata = Strata::BindStrataGlobalUniformParameters(View.StrataSceneData);
