Potential fix for water sampling depth outside of valid UV range after gather4.

diff --git a/Engine/Source/Runtime/Renderer/Private/SingleLayerWaterRendering.cpp b/Engine/Source/Runtime/Renderer/Private/SingleLayerWaterRendering.cpp
--- a/Engine/Source/Runtime/Renderer/Private/SingleLayerWaterRendering.cpp
+++ b/Engine/Source/Runtime/Renderer/Private/SingleLayerWaterRendering.cpp
@@ -322,10 +322,14 @@ static FSceneWithoutWaterTextures AddCopySceneWithoutWaterPass(
 		const FIntRect RefractionViewRect = FIntRect(FIntPoint::DivideAndRoundDown(View.ViewRect.Min, RefractionDownsampleFactor), FIntPoint::DivideAndRoundDown(View.ViewRect.Max, RefractionDownsampleFactor));
 
 		Textures.Views[ViewIndex].ViewRect   = RefractionViewRect;
-		Textures.Views[ViewIndex].MinMaxUV.X = (RefractionViewRect.Min.X + 0.5f) / RefractionResolution.X; 
-		Textures.Views[ViewIndex].MinMaxUV.Y = (RefractionViewRect.Min.Y + 0.5f) / RefractionResolution.Y;
-		Textures.Views[ViewIndex].MinMaxUV.Z = (RefractionViewRect.Max.X - 0.5f) / RefractionResolution.X; 
-		Textures.Views[ViewIndex].MinMaxUV.W = (RefractionViewRect.Max.Y - 0.5f) / RefractionResolution.Y;
+
+		// This is usually half a pixel. But it seems that when using Gather4, 0.5 is not conservative enough and can return pixel outside the guard band. 
+		// That is why it is a tiny bit higher than 0.5: for Gathre4 to always return pixels within the valid side of UVs (see EvaluateWaterVolumeLighting).
+		const float PixelSafeGuardBand = 0.55;
+		Textures.Views[ViewIndex].MinMaxUV.X = (RefractionViewRect.Min.X + PixelSafeGuardBand) / RefractionResolution.X;
+		Textures.Views[ViewIndex].MinMaxUV.Y = (RefractionViewRect.Min.Y + PixelSafeGuardBand) / RefractionResolution.Y;
+		Textures.Views[ViewIndex].MinMaxUV.Z = (RefractionViewRect.Max.X - PixelSafeGuardBand) / RefractionResolution.X;
+		Textures.Views[ViewIndex].MinMaxUV.W = (RefractionViewRect.Max.Y - PixelSafeGuardBand) / RefractionResolution.Y;
 
 		FPixelShaderUtils::AddFullscreenPass(
 			GraphBuilder,
