Fix not being able to disable 'adaptive voxel size' feature.

Even when the voxel size GPU->GPU feedback was disable, the value was still read from voxel size estimation based projected bound.

diff --git a/Engine/Shaders/Private/HairStrands/HairStrandsVoxelPageAllocation.usf b/Engine/Shaders/Private/HairStrands/HairStrandsVoxelPageAllocation.usf
--- a/Engine/Shaders/Private/HairStrands/HairStrandsVoxelPageAllocation.usf
+++ b/Engine/Shaders/Private/HairStrands/HairStrandsVoxelPageAllocation.usf
@@ -12,6 +12,7 @@
 
 float				CPUPageWorldSize;
 float				CPUVoxelWorldSize;
+uint				bUseCPUVoxelWorldSize;	// When adaptive voxel size is disabled, we use CPU voxel size value
 uint				TotalPageIndexCount; 	// This is the max page index count;
 uint				PageResolution; 		// Resolution of a page
 uint				MacroGroupCount;
@@ -52,7 +53,7 @@ void AllocatePageIndex(uint2 DispatchThreadId : SV_DispatchThreadID)
 	{
 		const bool bSupportVoxelization = (bDoesMacroGroupSupportVoxelization >> MacroGroupId) & 0x1;
 		Bound = ReadHairAABB(MacroGroupId, MacroGroupAABBBuffer);
-		const float VoxelWorldSize = max(GPUVoxelWorldSize[0], MacroGroupVoxelSizeBuffer[MacroGroupId]);
+		const float VoxelWorldSize = bUseCPUVoxelWorldSize ? CPUVoxelWorldSize : max(GPUVoxelWorldSize[0], MacroGroupVoxelSizeBuffer[MacroGroupId]);
 		PageWorldSize = VoxelWorldSize * PageResolution;
 
 		if (any(Bound.Min > Bound.Max) || !bSupportVoxelization)
@@ -652,6 +653,7 @@ float3 CPU_TranslatedWorldMaxAABB;
 uint CPU_PageIndexOffset;
 int3 CPU_PageIndexResolution;
 float CPU_VoxelWorldSize;
+uint bUseCPUVoxelWorldSize; // When adaptive voxel size is disabled, we use CPU voxel size value
 
 Buffer<float> GPU_VoxelWorldSize;
 Buffer<int> MacroGroupAABBBuffer;
@@ -673,7 +675,7 @@ void AddDescCS(uint GroupIndex : SV_GroupIndex, uint3 DispatchThreadId : SV_Disp
 	Node.TranslatedWorldMaxAABB = TranslatedWorldBound.Max;
 	Node.PageIndexResolution = PageIndexResolutionAndOffset.xyz;
 	Node.PageIndexOffset = PageIndexResolutionAndOffset.w;
-	Node.VoxelWorldSize = max(GPU_VoxelWorldSize[0], VoxelWorldSize);
+	Node.VoxelWorldSize = bUseCPUVoxelWorldSize ? CPU_VoxelWorldSize : max(GPU_VoxelWorldSize[0], VoxelWorldSize);
 #else
 	Node.TranslatedWorldMinAABB = CPU_TranslatedWorldMinAABB;
 	Node.TranslatedWorldMaxAABB = CPU_TranslatedWorldMaxAABB;
diff --git a/Engine/Source/Runtime/Renderer/Private/HairStrands/HairStrandsVoxelization.cpp b/Engine/Source/Runtime/Renderer/Private/HairStrands/HairStrandsVoxelization.cpp
--- a/Engine/Source/Runtime/Renderer/Private/HairStrands/HairStrandsVoxelization.cpp
+++ b/Engine/Source/Runtime/Renderer/Private/HairStrands/HairStrandsVoxelization.cpp
@@ -216,6 +216,7 @@ class FVoxelAllocatePageIndexCS : public FGlobalShader
 	BEGIN_SHADER_PARAMETER_STRUCT(FParameters, )
 		SHADER_PARAMETER(float, CPUPageWorldSize)
 		SHADER_PARAMETER(float, CPUVoxelWorldSize)
+		SHADER_PARAMETER(uint32, bUseCPUVoxelWorldSize)
 		SHADER_PARAMETER(uint32, TotalPageIndexCount)
 		SHADER_PARAMETER(uint32, PageResolution)
 		SHADER_PARAMETER(uint32, MacroGroupCount)
@@ -393,6 +394,7 @@ class FVoxelAddNodeDescCS : public FGlobalShader
 		SHADER_PARAMETER(uint32, MacroGroupId)
 		SHADER_PARAMETER(FIntVector, CPU_PageIndexResolution)
 		SHADER_PARAMETER(float, CPU_VoxelWorldSize)
+		SHADER_PARAMETER(uint32, bUseCPUVoxelWorldSize)
 		SHADER_PARAMETER_RDG_BUFFER_SRV(Buffer, GPU_VoxelWorldSize)
 		SHADER_PARAMETER_RDG_BUFFER_SRV(Buffer, MacroGroupAABBBuffer)
 		SHADER_PARAMETER_RDG_BUFFER_SRV(Buffer, MacroGroupVoxelSizeBuffer)
@@ -628,6 +630,7 @@ static void AddAllocateVoxelPagesPass(
 	AddClearUAVPass(GraphBuilder, PageIndexGlobalCounterUAV, 0u);
 
 	const FRDGBufferSRVRef CurrGPUMinVoxelSizeSRV = GraphBuilder.CreateSRV(CurrGPUMinVoxelSize, PF_R16F);
+	const bool bAdaptiveVoxelEnable = IsHairStrandsAdaptiveVoxelAllocationEnable();
 
 	// Allocate page index for all instance group
 //	if (bIsGPUDriven)
@@ -644,6 +647,7 @@ static void AddAllocateVoxelPagesPass(
 		FVoxelAllocatePageIndexCS::FParameters* Parameters = GraphBuilder.AllocParameters<FVoxelAllocatePageIndexCS::FParameters>();
 		Parameters->CPUPageWorldSize = CPUPageWorldSize;
 		Parameters->CPUVoxelWorldSize = CPUMinVoxelWorldSize;
+		Parameters->bUseCPUVoxelWorldSize = bAdaptiveVoxelEnable ? 0u : 1u;
 		Parameters->TotalPageIndexCount = OutTotalPageIndexCount;
 		Parameters->PageResolution = PageResolution;
 		Parameters->MacroGroupCount = MacroGroupCount;
@@ -850,6 +854,7 @@ static void AddAllocateVoxelPagesPass(
 			Parameters->CPU_PageIndexOffset = CPUAllocationDesc.PageIndexOffset;
 			Parameters->CPU_VoxelWorldSize = CPUMinVoxelWorldSize;
 			Parameters->GPU_VoxelWorldSize = CurrGPUMinVoxelSizeSRV;
+			Parameters->bUseCPUVoxelWorldSize = bAdaptiveVoxelEnable ? 0u : 1u;
 			Parameters->MacroGroupVoxelSizeBuffer = GraphBuilder.CreateSRV(MacroGroupResources.MacroGroupVoxelSizeBuffer, PF_R16F);
 			Parameters->OutNodeDescBuffer = GraphBuilder.CreateUAV(NodeDescBuffer);
 
@@ -921,7 +926,6 @@ static void AddAllocateVoxelPagesPass(
 	}
 
 	// Feedback allocation for next frame
-	const bool bAdaptiveVoxelEnable = IsHairStrandsAdaptiveVoxelAllocationEnable();
 	if (bAdaptiveVoxelEnable)
 	{
 		// Store the total requested page allocation (for feedback purpose)
