Cloud are now precompile through dxc. This accelrates cloud shaders compilation time by at least 6x.

diff --git a/Engine/Source/Runtime/Renderer/Private/VolumetricCloudRendering.cpp b/Engine/Source/Runtime/Renderer/Private/VolumetricCloudRendering.cpp
--- a/Engine/Source/Runtime/Renderer/Private/VolumetricCloudRendering.cpp
+++ b/Engine/Source/Runtime/Renderer/Private/VolumetricCloudRendering.cpp
@@ -754,6 +754,9 @@ public:
 		OutEnvironment.SetDefine(TEXT("CLOUD_SAMPLE_SECOND_LIGHT"), bSampleSecondLight ? TEXT("1") : TEXT("0"));
 
 		OutEnvironment.SetDefine(TEXT("CLOUD_SAMPLE_COUNT_DEBUG_MODE"), bSampleCountDebugMode);
+
+		// This shader takes a very long time to compile with FXC, so we pre-compile it with DXC first and then forward the optimized HLSL to FXC.
+		OutEnvironment.CompilerFlags.Add(CFLAG_PrecompileWithDXC);
 	}
 
 private:
@@ -853,6 +856,9 @@ class FRenderVolumetricCloudRenderViewCS : public FMeshMaterialShader
 		// This shader must support typed UAV load and we are testing if it is supported at runtime using RHIIsTypedUAVLoadSupported
 		OutEnvironment.CompilerFlags.Add(CFLAG_AllowTypedUAVLoads);
 
+		// This shader takes a very long time to compile with FXC, so we pre-compile it with DXC first and then forward the optimized HLSL to FXC.
+		OutEnvironment.CompilerFlags.Add(CFLAG_PrecompileWithDXC);
+
 		FForwardLightingParameters::ModifyCompilationEnvironment(Parameters.Platform, OutEnvironment);
 	}
 };
