Fixed simple light not being visible in volumetric fog.

diff --git a/Engine/Source/Runtime/Renderer/Private/LightGridInjection.cpp b/Engine/Source/Runtime/Renderer/Private/LightGridInjection.cpp
--- a/Engine/Source/Runtime/Renderer/Private/LightGridInjection.cpp
+++ b/Engine/Source/Runtime/Renderer/Private/LightGridInjection.cpp
@@ -337,7 +337,9 @@ void FSceneRenderer::ComputeLightGrid(FRDGBuilder& GraphBuilder, bool bCullLight
 
 					const FSimpleLightEntry& SimpleLight = SimpleLights.InstanceData[SimpleLightIndex];
 					const FSimpleLightPerViewEntry& SimpleLightPerViewData = SimpleLights.GetViewDependentData(SimpleLightIndex, ViewIndex, Views.Num());
-					LightData.LightPositionAndInvRadius = FVector4f((FVector3f)SimpleLightPerViewData.Position, 1.0f / FMath::Max(SimpleLight.Radius, KINDA_SMALL_NUMBER));
+
+					const FVector3f LightTranslatedWorldPosition(View.ViewMatrices.GetPreViewTranslation() + SimpleLightPerViewData.Position);
+					LightData.LightPositionAndInvRadius = FVector4f(LightTranslatedWorldPosition, 1.0f / FMath::Max(SimpleLight.Radius, KINDA_SMALL_NUMBER));
 					LightData.LightColorAndFalloffExponent = FVector4f((FVector3f)SimpleLight.Color, SimpleLight.Exponent);
 
 					// No shadowmap channels for simple lights
