Fixes crash in TSR cinematic scalability when disabling translucency

diff --git a/Engine/Shaders/Private/TemporalSuperResolution/TSRResolveHistory.usf b/Engine/Shaders/Private/TemporalSuperResolution/TSRResolveHistory.usf
--- a/Engine/Shaders/Private/TemporalSuperResolution/TSRResolveHistory.usf
+++ b/Engine/Shaders/Private/TemporalSuperResolution/TSRResolveHistory.usf
@@ -145,7 +145,7 @@ void MainCS(
 				tsr_half SampleHistoryValidity = SampleMetadata.y;
 				tsr_half InvSampleHistoryValidity = SafeRcp(SampleHistoryValidity);
 
-				#if 1
+				#if DIM_SEPARATE_TRANSLUCENCY
 				{
 					SampleLowFrequencyColor.rgb  = SampleLowFrequencyColor.rgb  * SampleTranslucencyAlpha + SampleTranslucencyColor.rgb;
 					SampleHighFrequencyColor.rgb = SampleHighFrequencyColor.rgb * SampleTranslucencyAlpha + SampleTranslucencyColor.rgb * SampleHistoryValidity;
diff --git a/Engine/Source/Runtime/Renderer/Private/PostProcess/TemporalSuperResolution.cpp b/Engine/Source/Runtime/Renderer/Private/PostProcess/TemporalSuperResolution.cpp
--- a/Engine/Source/Runtime/Renderer/Private/PostProcess/TemporalSuperResolution.cpp
+++ b/Engine/Source/Runtime/Renderer/Private/PostProcess/TemporalSuperResolution.cpp
@@ -489,6 +489,8 @@ class FTSRResolveHistoryCS : public FTSRShader
 	DECLARE_GLOBAL_SHADER(FTSRResolveHistoryCS);
 	SHADER_USE_PARAMETER_STRUCT(FTSRResolveHistoryCS, FTSRShader);
 
+	using FPermutationDomain = TShaderPermutationDomain<FTSRUpdateHistoryCS::FSeparateTranslucencyDim>;
+
 	BEGIN_SHADER_PARAMETER_STRUCT(FParameters, )
 		SHADER_PARAMETER_STRUCT_INCLUDE(FTSRCommonParameters, CommonParameters)
 		SHADER_PARAMETER(FScreenTransform, DispatchThreadToHistoryPixelPos)
@@ -1493,7 +1495,10 @@ ITemporalUpscaler::FOutputs AddTemporalSuperResolutionPasses(
 		}
 		PassParameters->DebugOutput = CreateDebugUAV(HistoryExtent, TEXT("Debug.TSR.ResolveHistory"));
 
-		TShaderMapRef<FTSRResolveHistoryCS> ComputeShader(View.ShaderMap);
+		FTSRResolveHistoryCS::FPermutationDomain PermutationVector;
+		PermutationVector.Set<FTSRUpdateHistoryCS::FSeparateTranslucencyDim>(bAccumulateSeparateTranslucency);
+
+		TShaderMapRef<FTSRResolveHistoryCS> ComputeShader(View.ShaderMap, PermutationVector);
 		FComputeShaderUtils::AddPass(
 			GraphBuilder,
 			RDG_EVENT_NAME("TSR ResolveHistory %dx%d", OutputRect.Width(), OutputRect.Height()),
