- Fix for translucent shadows not being rendered

#rnx

diff --git a/Engine/Source/Runtime/Renderer/Private/ShadowSetup.cpp b/Engine/Source/Runtime/Renderer/Private/ShadowSetup.cpp
--- a/Engine/Source/Runtime/Renderer/Private/ShadowSetup.cpp
+++ b/Engine/Source/Runtime/Renderer/Private/ShadowSetup.cpp
@@ -2094,6 +2094,7 @@ void FProjectedShadowInfo::FinalizeAddSubjectPrimitive(
 bool FProjectedShadowInfo::HasSubjectPrims() const
 {
 	return DynamicSubjectPrimitives.Num() > 0
+		|| SubjectTranslucentPrimitives.Num() > 0
 		|| ShadowDepthPass.HasAnyDraw()
 		|| SubjectMeshCommandBuildRequests.Num() > 0
 		|| ShadowDepthPassVisibleCommands.Num() > 0
