Fix black/flickering at borders for half resolution SSS due to bilinear filtering of uninitialized buffer.

diff --git a/Engine/Source/Runtime/Renderer/Private/PostProcess/PostProcessSubsurface.cpp b/Engine/Source/Runtime/Renderer/Private/PostProcess/PostProcessSubsurface.cpp
--- a/Engine/Source/Runtime/Renderer/Private/PostProcess/PostProcessSubsurface.cpp
+++ b/Engine/Source/Runtime/Renderer/Private/PostProcess/PostProcessSubsurface.cpp
@@ -1155,8 +1155,23 @@ void AddSubsurfaceViewPass(
 				Tiles.TileTypeCountBuffer,
 				ToIndex(FSubsurfaceTiles::ETileType::SEPARABLE));
 
-			SubsurfaceSubpassTwoTex = GraphBuilder.CreateTexture(SubsurfaceTextureDescriptor, TEXT("Subsurface.SubpassTwoTex"));
-			
+			// Pre-allocate black UAV for the second buffer if the separable tiles are rendered in half resolution.
+			if (bHalfRes)
+			{
+				SubsurfaceSubpassTwoTex = CreateBlackUAVTexture(
+					GraphBuilder,
+					SubsurfaceTextureDescriptor,
+					TEXT("Subsurface.SubpassTwoTex"),
+					RDG_EVENT_NAME("SSS::ClearUAV(%s, %s.TileCount > 0)", TEXT("Subsurface.SubpassTwoTex"), ToString(FSubsurfaceTiles::ETileType::SEPARABLE)),
+					SubsurfaceViewport,
+					Tiles.TileTypeCountBuffer,
+					ToIndex(FSubsurfaceTiles::ETileType::SEPARABLE));
+			}
+			else
+			{
+				SubsurfaceSubpassTwoTex = GraphBuilder.CreateTexture(SubsurfaceTextureDescriptor, TEXT("Subsurface.SubpassTwoTex"));
+			}
+
 			if (!bForceRunningInSeparable)
 			{
 				NewQualityHistoryTexture = GraphBuilder.CreateTexture(SubsurfaceTextureDescriptor, TEXT("Subsurface.QualityHistoryState"));
