Remove check that prevented per-object shadow setup to run if non-nanite VSM were active (which allows the per-object translucent shadows to work).

diff --git a/Engine/Source/Runtime/Renderer/Private/ShadowSetup.cpp b/Engine/Source/Runtime/Renderer/Private/ShadowSetup.cpp
--- a/Engine/Source/Runtime/Renderer/Private/ShadowSetup.cpp
+++ b/Engine/Source/Runtime/Renderer/Private/ShadowSetup.cpp
@@ -5631,7 +5631,6 @@ FDynamicShadowsTaskData* FSceneRenderer::BeginInitDynamicShadows(bool bRunningEa
 	const bool bProjectEnablePointLightShadows = Scene->ReadOnlyCVARCache.bEnablePointLightShadows && !bMobile; // Point light shadow is unsupported on mobile for now.
 	const bool bProjectEnableMovableDirectionLightShadows = !bMobile || Scene->ReadOnlyCVARCache.bMobileAllowMovableDirectionalLights;
 	const bool bProjectEnableMovableSpotLightShadows = !bMobile || Scene->ReadOnlyCVARCache.bMobileEnableMovableSpotlightsShadow;
-	const bool bProjectEnableNonNaniteVirtualShadowMaps = UseNonNaniteVirtualShadowMaps(ShaderPlatform, FeatureLevel);
 
 	uint32 NumPointShadowCachesUpdatedThisFrame = 0;
 	uint32 NumSpotShadowCachesUpdatedThisFrame = 0;
@@ -5755,8 +5754,7 @@ FDynamicShadowsTaskData* FSceneRenderer::BeginInitDynamicShadows(bool bRunningEa
 								AddViewDependentWholeSceneShadowsForView(ViewDependentWholeSceneShadows, ViewDependentWholeSceneShadowsThatNeedCulling, VisibleLightInfo, *LightSceneInfo, NumCSMCachesUpdatedThisFrame);
 							}
 
-							// Disable per-object shadows when non-nanite VSMs are enabled
-							if (!bProjectEnableNonNaniteVirtualShadowMaps && (!bMobile || (LightSceneInfo->Proxy->CastsModulatedShadows() && !LightSceneInfo->Proxy->UseCSMForDynamicObjects())))
+							if (!bMobile || (LightSceneInfo->Proxy->CastsModulatedShadows() && !LightSceneInfo->Proxy->UseCSMForDynamicObjects()))
 							{
 								const TArray<FLightPrimitiveInteraction*>* InteractionShadowPrimitives = LightSceneInfo->GetInteractionShadowPrimitives();
 
