Global Distance Field - fixed crash when enabling global distance field visualization, but global distance field isn't used

diff --git a/Engine/Source/Runtime/Renderer/Private/DistanceFieldVisualization.cpp b/Engine/Source/Runtime/Renderer/Private/DistanceFieldVisualization.cpp
--- a/Engine/Source/Runtime/Renderer/Private/DistanceFieldVisualization.cpp
+++ b/Engine/Source/Runtime/Renderer/Private/DistanceFieldVisualization.cpp
@@ -92,7 +92,11 @@ void FDeferredShadingSceneRenderer::RenderMeshDistanceFieldVisualization(
 	const bool bAnyMeshSDFs = Scene->DistanceFieldSceneData.NumObjectsInBuffer > 0;
 	const bool bAnyHeightFields = Scene->DistanceFieldSceneData.HeightfieldPrimitives.Num() > 0;
 
-	const bool bVisualizeGlobalDistanceField = FirstView.Family->EngineShowFlags.VisualizeGlobalDistanceField && UseGlobalDistanceField(Parameters) && (bAnyMeshSDFs || bAnyHeightFields);
+	const bool bVisualizeGlobalDistanceField = FirstView.Family->EngineShowFlags.VisualizeGlobalDistanceField 
+		&& ShouldPrepareGlobalDistanceField() 
+		&& UseGlobalDistanceField(Parameters) 
+		&& (bAnyMeshSDFs || bAnyHeightFields);
+
 	const bool bVisualizeMeshDistanceFields = FirstView.Family->EngineShowFlags.VisualizeMeshDistanceFields && bAnyMeshSDFs;
 
 	if (FeatureLevel < ERHIFeatureLevel::SM5
