Fix translucent rendering result not being composited in the scene color when ISR is enabled.

This was caused by an empty TranslucencyViewResourcesMap in the post-processor. To fix it we allow more of the code path to execute before we skip rendering the second view.

diff --git a/Engine/Source/Runtime/Renderer/Private/TranslucentRendering.cpp b/Engine/Source/Runtime/Renderer/Private/TranslucentRendering.cpp
--- a/Engine/Source/Runtime/Renderer/Private/TranslucentRendering.cpp
+++ b/Engine/Source/Runtime/Renderer/Private/TranslucentRendering.cpp
@@ -1255,6 +1255,11 @@ static void RenderTranslucencyViewInner(
 	bool bRenderInParallel,
 	FInstanceCullingManager& InstanceCullingManager)
 {
+	if (!View.ShouldRenderView())
+	{
+		return;
+	}
+
 	if (SceneColorLoadAction == ERenderTargetLoadAction::EClear)
 	{
 		AddClearRenderTargetPass(GraphBuilder, SceneColorTexture.Target);
@@ -1340,11 +1345,6 @@ void FDeferredShadingSceneRenderer::RenderTranslucencyInner(
 	const bool bIsScalingTranslucency = SeparateTranslucencyDimensions.Scale < 1.0f;
 	const bool bRenderInSeparateTranslucency = IsSeparateTranslucencyEnabled(TranslucencyPass, SeparateTranslucencyDimensions.Scale);
 
-	const auto ShouldRenderView = [&](const FViewInfo& View, ETranslucencyView TranslucencyView)
-	{
-		return View.ShouldRenderView() && EnumHasAnyFlags(TranslucencyView, ViewsToRender);
-	};
-
 	// Can't reference scene color in scene textures. Scene color copy is used instead.
 	ESceneTextureSetupMode SceneTextureSetupMode = ESceneTextureSetupMode::All;
 	EnumRemoveFlags(SceneTextureSetupMode, ESceneTextureSetupMode::SceneColor);
@@ -1359,7 +1359,7 @@ void FDeferredShadingSceneRenderer::RenderTranslucencyInner(
 			FViewInfo& View = Views[ViewIndex];
 			const ETranslucencyView TranslucencyView = GetTranslucencyView(View);
 
-			if (!ShouldRenderView(View, TranslucencyView))
+			if (!EnumHasAnyFlags(TranslucencyView, ViewsToRender))
 			{
 				continue;
 			}
@@ -1469,7 +1469,7 @@ void FDeferredShadingSceneRenderer::RenderTranslucencyInner(
 			FViewInfo& View = Views[ViewIndex];
 			const ETranslucencyView TranslucencyView = GetTranslucencyView(View);
 
-			if (!ShouldRenderView(View, TranslucencyView))
+			if (!EnumHasAnyFlags(TranslucencyView, ViewsToRender))
 			{
 				continue;
 			}
