Improve DistanceFields streaming CPU perf.
- Modified PendingThrottledOperations to be a TSet, avoiding expensive Remove operation. PendingThrottledOperations is appended to PendingAddOperations which is already a TSet, so order is not guaranteed anyway.
- Use RemoveAtSwap to remove completed requests from ReadRequests.

diff --git a/Engine/Source/Runtime/Renderer/Private/DistanceFieldStreaming.cpp b/Engine/Source/Runtime/Renderer/Private/DistanceFieldStreaming.cpp
--- a/Engine/Source/Runtime/Renderer/Private/DistanceFieldStreaming.cpp
+++ b/Engine/Source/Runtime/Renderer/Private/DistanceFieldStreaming.cpp
@@ -813,7 +813,7 @@ void FDistanceFieldSceneData::ProcessReadRequests(
 
 		if (bReady)
 		{
-			ReadRequests.RemoveAt(RequestIndex);
+			ReadRequests.RemoveAtSwap(RequestIndex);
 			RequestIndex--;
 
 			if (AssetStateArray.IsValidId(ReadRequest.AssetSetId) 
diff --git a/Engine/Source/Runtime/Renderer/Private/ScenePrivate.h b/Engine/Source/Runtime/Renderer/Private/ScenePrivate.h
--- a/Engine/Source/Runtime/Renderer/Private/ScenePrivate.h
+++ b/Engine/Source/Runtime/Renderer/Private/ScenePrivate.h
@@ -2101,7 +2101,7 @@ public:
 	TArray<FPrimitiveSceneInfo*> HeightfieldPrimitives;
 	/** Pending operations on the object buffers to be processed next frame. */
 	TSet<FPrimitiveSceneInfo*> PendingAddOperations;
-	TArray<FPrimitiveSceneInfo*> PendingThrottledOperations;
+	TSet<FPrimitiveSceneInfo*> PendingThrottledOperations;
 	TSet<FPrimitiveSceneInfo*> PendingUpdateOperations;
 	TArray<FPrimitiveRemoveInfo> PendingRemoveOperations;
 	TArray<FRenderBounds> PrimitiveModifiedBounds[GDF_Num];
