Fix the stencil write on mobile depth prepass for unlit materials.

diff --git a/Engine/Source/Runtime/Renderer/Private/DepthRendering.cpp b/Engine/Source/Runtime/Renderer/Private/DepthRendering.cpp
--- a/Engine/Source/Runtime/Renderer/Private/DepthRendering.cpp
+++ b/Engine/Source/Runtime/Renderer/Private/DepthRendering.cpp
@@ -703,7 +703,7 @@ FMeshDrawCommandSortKey CalculateDepthPassMeshStaticSortKey(EBlendMode BlendMode
 	return SortKey;
 }
 
-void SetMobileDepthPassRenderState(const FPrimitiveSceneProxy* RESTRICT PrimitiveSceneProxy, FMeshPassProcessorRenderState& DrawRenderState, const FMaterial& RESTRICT MaterialResource, bool bUsesDeferredShading)
+void SetMobileDepthPassRenderState(const FPrimitiveSceneProxy* RESTRICT PrimitiveSceneProxy, FMeshPassProcessorRenderState& DrawRenderState, const FMeshBatch& RESTRICT MeshBatch, bool bUsesDeferredShading)
 {
 	DrawRenderState.SetDepthStencilState(TStaticDepthStencilState<
 		true, CF_DepthNearOrEqual,
@@ -720,6 +720,7 @@ void SetMobileDepthPassRenderState(const FPrimitiveSceneProxy* RESTRICT Primitiv
 	if (bUsesDeferredShading)
 	{
 		// store into [1-3] bits
+		const FMaterial& MaterialResource = *MeshBatch.MaterialRenderProxy->GetMaterialNoFallback(ERHIFeatureLevel::ES3_1);
 		uint8 ShadingModel = MaterialResource.GetShadingModels().IsLit() ? MSM_DefaultLit : MSM_Unlit;
 		StencilValue |= GET_STENCIL_MOBILE_SM_MASK(ShadingModel);
 	}
@@ -768,7 +769,7 @@ bool FDepthPassMeshProcessor::Process(
 	// Use StencilMask for DecalOutput on mobile
 	if (FeatureLevel == ERHIFeatureLevel::ES3_1 && !bShadowProjection)
 	{
-		SetMobileDepthPassRenderState(PrimitiveSceneProxy, DrawRenderState, MaterialResource, IsMobileDeferredShadingEnabled(GetFeatureLevelShaderPlatform(FeatureLevel)));
+		SetMobileDepthPassRenderState(PrimitiveSceneProxy, DrawRenderState, MeshBatch, IsMobileDeferredShadingEnabled(GetFeatureLevelShaderPlatform(FeatureLevel)));
 	}
 
 	FDepthOnlyShaderElementData ShaderElementData(0.0f);
