Fixed: Android Package for TM-Shadermodels crashes on launch with Global Clip Plane enabled

diff --git a/Engine/Shaders/Private/MobileBasePassCommon.ush b/Engine/Shaders/Private/MobileBasePassCommon.ush
--- a/Engine/Shaders/Private/MobileBasePassCommon.ush
+++ b/Engine/Shaders/Private/MobileBasePassCommon.ush
@@ -43,7 +43,7 @@ struct FSharedMobileBasePassInterpolants
 #endif
 
 #if USE_GLOBAL_CLIP_PLANE
-	float OutClipDistance : OUTCLIPDIST;
+	float OutClipDistance : SV_ClipDistance;
 #endif
 
 #if MOBILE_MULTI_VIEW
diff --git a/Engine/Shaders/Public/Platform/GL/GLCommon.ush b/Engine/Shaders/Public/Platform/GL/GLCommon.ush
--- a/Engine/Shaders/Public/Platform/GL/GLCommon.ush
+++ b/Engine/Shaders/Public/Platform/GL/GLCommon.ush
@@ -8,3 +8,7 @@
 // GUID = E121241A-1AD3-416B-8F74-288560D5F0005
 
 #pragma once
+
+
+// GLES does not support gl_ClipDistance, use interpolator for it
+#define SV_ClipDistance OUTCLIPDIST
\ No newline at end of file
diff --git a/Engine/Source/Runtime/Renderer/Private/MobileBasePass.cpp b/Engine/Source/Runtime/Renderer/Private/MobileBasePass.cpp
--- a/Engine/Source/Runtime/Renderer/Private/MobileBasePass.cpp
+++ b/Engine/Source/Runtime/Renderer/Private/MobileBasePass.cpp
@@ -317,7 +317,7 @@ ELightMapPolicyType MobileBasePass::SelectMeshLightmapPolicy(
 			}
 			else if ((bHasValidVLM || bHasValidILC) && bPrimitiveUsesILC)
 			{
-				if ((bUseStaticAndCSM || bMovableWithCSM) && !bUsesDeferredShading)
+				if ((bUseStaticAndCSM || bMovableWithCSM) && !bUsesDeferredShading && ReadOnlyCVARCache.bMobileEnableStaticAndCSMShadowReceivers)
 				{
 					SelectedLightmapPolicy = LMP_MOBILE_DIRECTIONAL_LIGHT_CSM_AND_SH_INDIRECT;
 				}
