Global Distance Field Visualization only uses coverage atlas when Lumen is enabled

diff --git a/Engine/Shaders/Private/GlobalDistanceFieldUtils.ush b/Engine/Shaders/Private/GlobalDistanceFieldUtils.ush
--- a/Engine/Shaders/Private/GlobalDistanceFieldUtils.ush
+++ b/Engine/Shaders/Private/GlobalDistanceFieldUtils.ush
@@ -7,6 +7,10 @@
 #define TWO_SIDED_BAND_SIZE 4.0f
 #define DEFAULT_BAND_SIZE 4.0f
 
+#ifndef GLOBALSDF_USE_COVERAGE_BASED_EXPAND
+#define GLOBALSDF_USE_COVERAGE_BASED_EXPAND 1
+#endif
+
 uint ComputeGlobalDistanceFieldClipmapIndex(float3 WorldPosition)
 {
 	uint FoundClipmapIndex = 0;
@@ -131,8 +135,7 @@ FGlobalSDFTraceResult RayTraceGlobalDistanceField(FGlobalSDFTraceInput TraceInpu
 
 				if (PageIndex < GLOBAL_DISTANCE_FIELD_INVALID_PAGE_ID && DistanceFieldMipValue < GlobalDistanceFieldMipTransition)
 				{
-#define USE_COVERAGE_BASED_EXPAND 1
-#if USE_COVERAGE_BASED_EXPAND
+#if GLOBALSDF_USE_COVERAGE_BASED_EXPAND
 					float3 PageUV;
 					float3 CoveragePageUV;
 					ComputeGlobalDistanceFieldPageUV(ClipmapVolumeUV, PageIndex, PageUV, CoveragePageUV);
diff --git a/Engine/Source/Runtime/Renderer/Private/DistanceFieldVisualization.cpp b/Engine/Source/Runtime/Renderer/Private/DistanceFieldVisualization.cpp
--- a/Engine/Source/Runtime/Renderer/Private/DistanceFieldVisualization.cpp
+++ b/Engine/Source/Runtime/Renderer/Private/DistanceFieldVisualization.cpp
@@ -24,7 +24,8 @@ public:
 	SHADER_USE_PARAMETER_STRUCT(FVisualizeMeshDistanceFieldCS, FGlobalShader);
 
 	class FUseGlobalDistanceFieldDim : SHADER_PERMUTATION_BOOL("USE_GLOBAL_DISTANCE_FIELD");
-	using FPermutationDomain = TShaderPermutationDomain<FUseGlobalDistanceFieldDim>;
+	class FCoverageBasedExpand : SHADER_PERMUTATION_BOOL("GLOBALSDF_USE_COVERAGE_BASED_EXPAND");
+	using FPermutationDomain = TShaderPermutationDomain<FUseGlobalDistanceFieldDim, FCoverageBasedExpand>;
 
 	BEGIN_SHADER_PARAMETER_STRUCT(FParameters, )
 		SHADER_PARAMETER_STRUCT_REF(FViewUniformShaderParameters, View)
@@ -121,6 +122,7 @@ void FDeferredShadingSceneRenderer::RenderMeshDistanceFieldVisualization(
 
 	FVisualizeMeshDistanceFieldCS::FPermutationDomain PermutationVector;
 	PermutationVector.Set<FVisualizeMeshDistanceFieldCS::FUseGlobalDistanceFieldDim>(bVisualizeGlobalDistanceField);
+	PermutationVector.Set<FVisualizeMeshDistanceFieldCS::FCoverageBasedExpand>(IsLumenEnabled(FirstView));
 	TShaderMapRef<FVisualizeMeshDistanceFieldCS> ComputeShader(FirstView.ShaderMap, PermutationVector);
 
 	for (const FViewInfo& View : Views)
