Cloud component samples now all default to 1.
Configurable real time capture with respect to cost/latency: sky and cloud captures can be timesliced per face now.
Fixed potential issue in deprecated SetXXX fucntions.
Removed advanced tag from sample count because those are too important to hide by default.

diff --git a/Engine/Source/Runtime/Core/Public/UObject/UE5MainStreamObjectVersion.h b/Engine/Source/Runtime/Core/Public/UObject/UE5MainStreamObjectVersion.h
--- a/Engine/Source/Runtime/Core/Public/UObject/UE5MainStreamObjectVersion.h
+++ b/Engine/Source/Runtime/Core/Public/UObject/UE5MainStreamObjectVersion.h
@@ -191,6 +191,9 @@ struct CORE_API FUE5MainStreamObjectVersion
 		// Add Texture DoScaleMipsForAlphaCoverage
 		TextureDoScaleMipsForAlphaCoverage,
 
+		// Fixed default value of volumetric cloud to be exact match with main view, more expenssive but we let user choosing how to lower the quality.
+		VolumetricCloudReflectionSampleCountDefaultUpdate,
+
 		// -----<new versions can be added above this line>-------------------------------------------------
 		VersionPlusOne,
 		LatestVersion = VersionPlusOne - 1
diff --git a/Engine/Source/Runtime/Engine/Classes/Components/VolumetricCloudComponent.h b/Engine/Source/Runtime/Engine/Classes/Components/VolumetricCloudComponent.h
--- a/Engine/Source/Runtime/Engine/Classes/Components/VolumetricCloudComponent.h
+++ b/Engine/Source/Runtime/Engine/Classes/Components/VolumetricCloudComponent.h
@@ -70,14 +70,16 @@ class UVolumetricCloudComponent : public USceneComponent
 	 * Scale the tracing sample count in primary views. Quality level scalability CVARs affect the maximum range.
 	 * The sample count resolution is still clamped according to scalability setting to 'r.VolumetricCloud.ViewRaySampleCountMax'.
 	 */
-	UPROPERTY(EditAnywhere, BlueprintReadOnly, Category = "Cloud Tracing", AdvancedDisplay, meta = (UIMin = "0.05", UIMax = "8", ClampMin = "0.05", SliderExponent = 1.0))
+	UPROPERTY(EditAnywhere, BlueprintReadOnly, Category = "Cloud Tracing", meta = (UIMin = "0.05", UIMax = "8", ClampMin = "0.05", SliderExponent = 1.0))
 	float ViewSampleCountScale;
 	/**
 	 * Scale the tracing sample count in reflection views. Quality level scalability CVARs affect the maximum range.
 	 * The sample count resolution is still clamped according to scalability setting to 'r.VolumetricCloud.ReflectionRaySampleMaxCount'.
 	 */
-	UPROPERTY(EditAnywhere, BlueprintReadOnly, Category = "Cloud Tracing", AdvancedDisplay, meta = (UIMin = "0.05", UIMax = "8", ClampMin = "0.05", SliderExponent = 1.0))
-	float ReflectionViewSampleCountScale;
+	UPROPERTY(EditAnywhere, BlueprintReadOnly, Category = "Cloud Tracing", meta = (UIMin = "0.05", UIMax = "8", ClampMin = "0.05", SliderExponent = 1.0))
+	float ReflectionViewSampleCountScaleValue;
+	UPROPERTY()
+	float ReflectionViewSampleCountScale_DEPRECATED;
 	UPROPERTY()
 	float ReflectionSampleCountScale_DEPRECATED;
 
@@ -85,27 +87,29 @@ class UVolumetricCloudComponent : public USceneComponent
 	 * Scale the shadow tracing sample count in primary views, only used with Advanced Output ray marched shadows. Quality level scalability CVARs affect the maximum range.
 	 * The sample count resolution is still clamped according to scalability setting to 'r.VolumetricCloud.Shadow.ViewRaySampleMaxCount'.
 	 */
-	UPROPERTY(EditAnywhere, BlueprintReadOnly, Category = "Cloud Tracing", AdvancedDisplay, meta = (UIMin = "0.05", UIMax = "8", ClampMin = "0.05", SliderExponent = 1.0))
+	UPROPERTY(EditAnywhere, BlueprintReadOnly, Category = "Cloud Tracing", meta = (UIMin = "0.05", UIMax = "8", ClampMin = "0.05", SliderExponent = 1.0))
 	float ShadowViewSampleCountScale;
 	/**
 	 * Scale the shadow tracing sample count in reflection views, only used with Advanced Output ray marched shadows. Quality level scalability CVARs affect the maximum range.
 	 * The sample count resolution is still clamped according to scalability setting to 'r.VolumetricCloud.Shadow.ReflectionRaySampleMaxCount'.
 	 */
-	UPROPERTY(EditAnywhere, BlueprintReadOnly, Category = "Cloud Tracing", AdvancedDisplay, meta = (UIMin = "0.05", UIMax = "8", ClampMin = "0.05", SliderExponent = 1.0))
-	float ShadowReflectionViewSampleCountScale;
+	UPROPERTY(EditAnywhere, BlueprintReadOnly, Category = "Cloud Tracing", meta = (UIMin = "0.05", UIMax = "8", ClampMin = "0.05", SliderExponent = 1.0))
+	float ShadowReflectionViewSampleCountScaleValue;
+	UPROPERTY()
+	float ShadowReflectionViewSampleCountScale_DEPRECATED;
 	UPROPERTY()
 	float ShadowReflectionSampleCountScale_DEPRECATED;
 
 	/**
 	 * The shadow tracing distance in kilometers, only used with Advanced Output ray marched shadows.
 	 */
-	UPROPERTY(EditAnywhere, BlueprintReadOnly, Category = "Cloud Tracing", AdvancedDisplay, meta = (UIMin = "0.1", UIMax = "50", ClampMin = "0.01", SliderExponent = 3.0))
+	UPROPERTY(EditAnywhere, BlueprintReadOnly, Category = "Cloud Tracing", meta = (UIMin = "0.1", UIMax = "50", ClampMin = "0.01", SliderExponent = 3.0))
 	float ShadowTracingDistance;
 
 	/**
 	 * When the mean transmittance is below this threshold, we stop tracing. This is a good way to reduce the ray marched sample count, and thus to increase performance.
 	 */
-	UPROPERTY(EditAnywhere, BlueprintReadOnly, Category = "Cloud Tracing", AdvancedDisplay, meta = (UIMin = "0.0", UIMax = "1.0", ClampMin = "0.0", ClampMax = "1.0", SliderExponent = 5.0))
+	UPROPERTY(EditAnywhere, BlueprintReadOnly, Category = "Cloud Tracing", meta = (UIMin = "0.0", UIMax = "1.0", ClampMin = "0.0", ClampMax = "1.0", SliderExponent = 5.0))
 	float StopTracingTransmittanceThreshold;
 
 
diff --git a/Engine/Source/Runtime/Engine/Private/Components/VolumetricCloudComponent.cpp b/Engine/Source/Runtime/Engine/Private/Components/VolumetricCloudComponent.cpp
--- a/Engine/Source/Runtime/Engine/Private/Components/VolumetricCloudComponent.cpp
+++ b/Engine/Source/Runtime/Engine/Private/Components/VolumetricCloudComponent.cpp
@@ -15,6 +15,7 @@
 #include "UObject/ConstructorHelpers.h"
 #include "Materials/MaterialInstance.h"
 #include "UObject/UE5ReleaseStreamObjectVersion.h"
+#include "UObject/UE5MainStreamObjectVersion.h"
 
 #if WITH_EDITOR
 #include "ObjectEditorUtils.h"
@@ -40,10 +41,12 @@ UVolumetricCloudComponent::UVolumetricCloudComponent(const FObjectInitializer& O
 	, bUsePerSampleAtmosphericLightTransmittance(false)
 	, SkyLightCloudBottomOcclusion(0.5f)
 	, ViewSampleCountScale(1.0f)
-	, ReflectionViewSampleCountScale(0.15f)			// Roughly equivalent to previous default 1.0f, scaled by OldToNewReflectionViewRaySampleCount
+	, ReflectionViewSampleCountScaleValue(1.0f)
+	, ReflectionViewSampleCountScale_DEPRECATED(0.15f)		// Roughly equivalent to previous default 1.0f, scaled by OldToNewReflectionViewRaySampleCount
 	, ReflectionSampleCountScale_DEPRECATED(1.0f)
 	, ShadowViewSampleCountScale(1.0f)
-	, ShadowReflectionViewSampleCountScale(0.3f)	// Roughly equivalent to previous default 1.0f, scaled by OldToNewReflectionShadowRaySampleCount
+	, ShadowReflectionViewSampleCountScaleValue(1.0f)
+	, ShadowReflectionViewSampleCountScale_DEPRECATED(0.3f)	// Roughly equivalent to previous default 1.0f, scaled by OldToNewReflectionShadowRaySampleCount
 	, ShadowReflectionSampleCountScale_DEPRECATED(1.0f)
 	, ShadowTracingDistance(15.0f)
 	, StopTracingTransmittanceThreshold(0.005f)
@@ -126,11 +129,18 @@ void UVolumetricCloudComponent::Serialize(FArchive& Ar)
 	Super::Serialize(Ar);
 
 	Ar.UsingCustomVersion(FUE5ReleaseStreamObjectVersion::GUID);
+	Ar.UsingCustomVersion(FUE5MainStreamObjectVersion::GUID);
 
 	if (Ar.CustomVer(FUE5ReleaseStreamObjectVersion::GUID) < FUE5ReleaseStreamObjectVersion::VolumetricCloudSampleCountUnification && Ar.IsLoading())
 	{
-		ReflectionViewSampleCountScale		= ReflectionSampleCountScale_DEPRECATED		  * UVolumetricCloudComponent::OldToNewReflectionViewRaySampleCount;
-		ShadowReflectionViewSampleCountScale= ShadowReflectionSampleCountScale_DEPRECATED * UVolumetricCloudComponent::OldToNewReflectionShadowRaySampleCount;
+		ReflectionViewSampleCountScale_DEPRECATED		= ReflectionSampleCountScale_DEPRECATED		  * UVolumetricCloudComponent::OldToNewReflectionViewRaySampleCount;
+		ShadowReflectionViewSampleCountScale_DEPRECATED = ShadowReflectionSampleCountScale_DEPRECATED * UVolumetricCloudComponent::OldToNewReflectionShadowRaySampleCount;
+	}
+
+	if (Ar.CustomVer(FUE5MainStreamObjectVersion::GUID) < FUE5MainStreamObjectVersion::VolumetricCloudReflectionSampleCountDefaultUpdate && Ar.IsLoading())
+	{
+		ReflectionViewSampleCountScaleValue = ReflectionViewSampleCountScale_DEPRECATED;
+		ShadowReflectionViewSampleCountScaleValue = ShadowReflectionViewSampleCountScale_DEPRECATED;
 	}
 }
 
@@ -153,26 +163,41 @@ CLOUD_DECLARE_BLUEPRINT_SETFUNCTION(FColor, GroundAlbedo);
 CLOUD_DECLARE_BLUEPRINT_SETFUNCTION(bool, bUsePerSampleAtmosphericLightTransmittance);
 CLOUD_DECLARE_BLUEPRINT_SETFUNCTION(float, SkyLightCloudBottomOcclusion);
 CLOUD_DECLARE_BLUEPRINT_SETFUNCTION(float, ViewSampleCountScale);
-CLOUD_DECLARE_BLUEPRINT_SETFUNCTION(float, ReflectionViewSampleCountScale);
 CLOUD_DECLARE_BLUEPRINT_SETFUNCTION(float, ShadowViewSampleCountScale);
-CLOUD_DECLARE_BLUEPRINT_SETFUNCTION(float, ShadowReflectionViewSampleCountScale);
 CLOUD_DECLARE_BLUEPRINT_SETFUNCTION(float, ShadowTracingDistance);
 CLOUD_DECLARE_BLUEPRINT_SETFUNCTION(float, StopTracingTransmittanceThreshold);
 CLOUD_DECLARE_BLUEPRINT_SETFUNCTION(UMaterialInterface*, Material);
 
+void UVolumetricCloudComponent::SetReflectionViewSampleCountScale(float NewValue)
+{
+	if (AreDynamicDataChangesAllowed() && ReflectionViewSampleCountScaleValue != NewValue)
+	{
+		ReflectionViewSampleCountScaleValue = NewValue;
+		MarkRenderStateDirty();
+	}
+}
+void UVolumetricCloudComponent::SetShadowReflectionViewSampleCountScale(float NewValue)
+{
+	if (AreDynamicDataChangesAllowed() && ShadowReflectionViewSampleCountScaleValue != NewValue)
+	{
+		ShadowReflectionViewSampleCountScaleValue = NewValue;
+		MarkRenderStateDirty();
+	}
+}
+
 void UVolumetricCloudComponent::SetReflectionSampleCountScale(float NewValue)
 {
-	if (AreDynamicDataChangesAllowed() && ReflectionViewSampleCountScale != NewValue)
+	if (AreDynamicDataChangesAllowed() && ReflectionViewSampleCountScaleValue != NewValue * UVolumetricCloudComponent::OldToNewReflectionViewRaySampleCount)
 	{
-		ReflectionViewSampleCountScale = NewValue * UVolumetricCloudComponent::OldToNewReflectionViewRaySampleCount;
+		ReflectionViewSampleCountScaleValue = NewValue * UVolumetricCloudComponent::OldToNewReflectionViewRaySampleCount;
 		MarkRenderStateDirty();
 	}
 }
 void UVolumetricCloudComponent::SetShadowReflectionSampleCountScale(float NewValue)
 {
-	if (AreDynamicDataChangesAllowed() && ShadowReflectionViewSampleCountScale != NewValue)
+	if (AreDynamicDataChangesAllowed() && ShadowReflectionViewSampleCountScaleValue != NewValue * UVolumetricCloudComponent::OldToNewReflectionShadowRaySampleCount)
 	{
-		ShadowReflectionViewSampleCountScale = NewValue * UVolumetricCloudComponent::OldToNewReflectionShadowRaySampleCount;
+		ShadowReflectionViewSampleCountScaleValue = NewValue * UVolumetricCloudComponent::OldToNewReflectionShadowRaySampleCount;
 		MarkRenderStateDirty();
 	}
 }
diff --git a/Engine/Source/Runtime/Engine/Private/VolumetricCloudProxy.cpp b/Engine/Source/Runtime/Engine/Private/VolumetricCloudProxy.cpp
--- a/Engine/Source/Runtime/Engine/Private/VolumetricCloudProxy.cpp
+++ b/Engine/Source/Runtime/Engine/Private/VolumetricCloudProxy.cpp
@@ -19,9 +19,9 @@ FVolumetricCloudSceneProxy::FVolumetricCloudSceneProxy(const UVolumetricCloudCom
 	, bUsePerSampleAtmosphericLightTransmittance(InComponent->bUsePerSampleAtmosphericLightTransmittance)
 	, SkyLightCloudBottomOcclusion(InComponent->SkyLightCloudBottomOcclusion)
 	, ViewSampleCountScale(InComponent->ViewSampleCountScale)
-	, ReflectionViewSampleCountScale(InComponent->ReflectionViewSampleCountScale)
+	, ReflectionViewSampleCountScale(InComponent->ReflectionViewSampleCountScaleValue)
 	, ShadowViewSampleCountScale(InComponent->ShadowViewSampleCountScale)
-	, ShadowReflectionViewSampleCountScale(InComponent->ShadowReflectionViewSampleCountScale)
+	, ShadowReflectionViewSampleCountScale(InComponent->ShadowReflectionViewSampleCountScaleValue)
 	, ShadowTracingDistance(InComponent->ShadowTracingDistance)
 	, StopTracingTransmittanceThreshold(InComponent->StopTracingTransmittanceThreshold)
 	, CloudVolumeMaterial(InComponent->Material)
diff --git a/Engine/Source/Runtime/Renderer/Private/ReflectionEnvironmentRealTimeCapture.cpp b/Engine/Source/Runtime/Renderer/Private/ReflectionEnvironmentRealTimeCapture.cpp
--- a/Engine/Source/Runtime/Renderer/Private/ReflectionEnvironmentRealTimeCapture.cpp
+++ b/Engine/Source/Runtime/Renderer/Private/ReflectionEnvironmentRealTimeCapture.cpp
@@ -31,6 +31,11 @@ DECLARE_GPU_STAT(CaptureConvolveSkyEnvMap);
 
 static TAutoConsoleVariable<int32> CVarRealTimeReflectionCaptureTimeSlicing(
 	TEXT("r.SkyLight.RealTimeReflectionCapture.TimeSlice"), 1,
+	TEXT("When enabled, the real-time sky light capture and convolutions will by distributed over several frames to lower the per-frame cost. Value in [1,6]."),
+	ECVF_RenderThreadSafe);
+
+static TAutoConsoleVariable<int32> CVarRealTimeReflectionCaptureTimeSlicingSkyCloudCubeFacePerFrame(
+	TEXT("r.SkyLight.RealTimeReflectionCapture.TimeSlice.SkyCloudCubeFacePerFrame"), 6,
 	TEXT("When enabled, the real-time sky light capture and convolutions will by distributed over several frames to lower the per-frame cost."),
 	ECVF_RenderThreadSafe);
 
@@ -368,7 +373,7 @@ void FScene::AllocateAndCaptureFrameSkyEnvMap(
 			});
 	};
 
-	auto RenderCubeFaces_SkyCloud = [&](bool bExecuteSky, bool bExecuteCloud, TRefCountPtr<IPooledRenderTarget>& SkyRenderTarget)
+	auto RenderCubeFaces_SkyCloud = [&](bool bExecuteSky, bool bExecuteCloud, TRefCountPtr<IPooledRenderTarget>& SkyRenderTarget, int32 StartCubeFace, int32 EndCubeFace)
 	{
 		FScene* Scene = MainView.Family->Scene->GetRenderScene();
 
@@ -467,8 +472,10 @@ void FScene::AllocateAndCaptureFrameSkyEnvMap(
 			}
 
 
-			for (int32 CubeFace = 0; CubeFace < CubeFace_MAX; CubeFace++)
+			for (int32 CubeFace = StartCubeFace; CubeFace < EndCubeFace; CubeFace++)
 			{
+				RDG_EVENT_SCOPE(GraphBuilder, "Capture Face=%d", CubeFace);
+
 				SkyRC.RenderTargets[0] = FRenderTargetBinding(SkyCubeTexture, ERenderTargetLoadAction::ENoAction, 0, CubeFace);
 
 				const FMatrix CubeViewRotationMatrix = CalcCubeFaceViewRotationMatrix((ECubeFace)CubeFace);
@@ -572,6 +579,7 @@ void FScene::AllocateAndCaptureFrameSkyEnvMap(
 						// If there are any mesh tagged as IsSky then we render them only, otherwise we simply render the sky atmosphere itself.
 						if (MainView.bSceneHasSkyMaterial)
 						{
+							RDG_EVENT_SCOPE(GraphBuilder, "Capture Sky Materials", CubeFace);
 							auto* PassParameters = GraphBuilder.AllocParameters<FCaptureSkyMeshReflectionPassParameters>();
 							PassParameters->View = CubeView.GetShaderParameters();
 							PassParameters->RenderTargets = SkyRC.RenderTargets;
@@ -616,6 +624,7 @@ void FScene::AllocateAndCaptureFrameSkyEnvMap(
 						}
 						else
 						{
+							RDG_EVENT_SCOPE(GraphBuilder, "Capture Sky Raw", CubeFace);
 							FSceneTextureShaderParameters SceneTextures = CreateSceneTextureShaderParameters(GraphBuilder, SceneRenderer.FeatureLevel, ESceneTextureSetupMode::SceneDepth);
 							SceneRenderer.RenderSkyAtmosphereInternal(GraphBuilder, SceneTextures, SkyRC);
 						}
@@ -883,7 +892,7 @@ void FScene::AllocateAndCaptureFrameSkyEnvMap(
 		ConvolvedSkyRenderTargetReadyIndex = 0;
 
 		// 0.60ms (0.12ms for faces with the most clouds)
-		RenderCubeFaces_SkyCloud(true, true, CapturedSkyRenderTarget);
+		RenderCubeFaces_SkyCloud(true, true, CapturedSkyRenderTarget, 0, CubeFace_MAX);
 
 		// 0.05ms
 		RenderCubeFaces_GenCubeMips(1, LastMipLevel, CapturedSkyRenderTarget);
@@ -911,32 +920,49 @@ void FScene::AllocateAndCaptureFrameSkyEnvMap(
 		const int32 ConvolvedSkyRenderTargetWorkIndex = 1 - ConvolvedSkyRenderTargetReadyIndex;
 		const int32 TimeSliceCount = 12;
 
-		// Update the current time-slicing state if this is a new frame
+#define DEBUG_TIME_SLICE 0
+#if DEBUG_TIME_SLICE
+		RealTimeSlicedReflectionCaptureState = -1;
+		RealTimeSlicedReflectionCaptureStateStep = 0;
+		while(true)
+		{
+			if (RealTimeSlicedReflectionCaptureState+1 >= TimeSliceCount)
+			{
+				break;
+			}
+#endif 
+
+		// Update the current time-slicing state if this is a new frame and if the current step is done.
 		// Note: RealTimeSlicedReflectionCaptureState will initially be -1.
-		if (bIsNewFrame)
+		if (bIsNewFrame && bRealTimeSlicedReflectionCaptureStateStepDone)
 		{
 			if (++RealTimeSlicedReflectionCaptureState >= TimeSliceCount)
 			{
 				RealTimeSlicedReflectionCaptureState = 0;
+				RealTimeSlicedReflectionCaptureStateStep = 0;
 			}
 		}
+		bRealTimeSlicedReflectionCaptureStateStepDone = true;
 
-#define DEBUG_TIME_SLICE 0
-#if DEBUG_TIME_SLICE
-		RealTimeSlicedReflectionCaptureState = 0;
-		for(int i=0; i<TimeSliceCount; ++i)
-		{
-#endif 
+		const int32 SkyCloudFrameStepCount	= FMath::Clamp(CVarRealTimeReflectionCaptureTimeSlicingSkyCloudCubeFacePerFrame.GetValueOnRenderThread(), int32(1), int32(CubeFace_MAX));
+		const int32 SkyCloudStartSubStep	= FMath::Clamp(RealTimeSlicedReflectionCaptureStateStep, int32(0), int32(CubeFace_MAX - 1));
+		const int32 SkyCloudEndSubStep		= FMath::Clamp(RealTimeSlicedReflectionCaptureStateStep + SkyCloudFrameStepCount, int32(0), int32(CubeFace_MAX));
 
 		if (RealTimeSlicedReflectionCaptureState <= 0)
 		{
-			RDG_EVENT_SCOPE(GraphBuilder, "RenderSky");
-			RenderCubeFaces_SkyCloud(true, false, CapturedSkyRenderTarget);
+			RDG_EVENT_SCOPE(GraphBuilder, "RenderSky StartFace=%d EndFace=%d", SkyCloudStartSubStep, SkyCloudEndSubStep);
+			RenderCubeFaces_SkyCloud(true, false, CapturedSkyRenderTarget, SkyCloudStartSubStep, SkyCloudEndSubStep);
+
+			bRealTimeSlicedReflectionCaptureStateStepDone = SkyCloudEndSubStep >= CubeFace_MAX;
+			RealTimeSlicedReflectionCaptureStateStep = bRealTimeSlicedReflectionCaptureStateStepDone ? 0 : SkyCloudEndSubStep;
 		}
 		else if (RealTimeSlicedReflectionCaptureState == 1)
 		{
-			RDG_EVENT_SCOPE(GraphBuilder, "RenderCloud");
-			RenderCubeFaces_SkyCloud(false, true, CapturedSkyRenderTarget);
+			RDG_EVENT_SCOPE(GraphBuilder, "RenderCloud StartFace=%d EndFace=%d", SkyCloudStartSubStep, SkyCloudEndSubStep);
+			RenderCubeFaces_SkyCloud(false, true, CapturedSkyRenderTarget, SkyCloudStartSubStep, SkyCloudEndSubStep);
+
+			bRealTimeSlicedReflectionCaptureStateStepDone = SkyCloudEndSubStep >= CubeFace_MAX;
+			RealTimeSlicedReflectionCaptureStateStep = bRealTimeSlicedReflectionCaptureStateStepDone ? 0 : SkyCloudEndSubStep;
 		}
 		else if (RealTimeSlicedReflectionCaptureState == 2)
 		{
diff --git a/Engine/Source/Runtime/Renderer/Private/RendererScene.cpp b/Engine/Source/Runtime/Renderer/Private/RendererScene.cpp
--- a/Engine/Source/Runtime/Renderer/Private/RendererScene.cpp
+++ b/Engine/Source/Runtime/Renderer/Private/RendererScene.cpp
@@ -1154,6 +1154,8 @@ FScene::FScene(UWorld* InWorld, bool bInRequiresHitProxies, bool bInIsEditorScen
 ,	ConvolvedSkyRenderTargetReadyIndex(-1)
 ,	RealTimeSlicedReflectionCaptureFirstFrameState(ERealTimeSlicedReflectionCaptureFirstFrameState::INIT)
 ,	RealTimeSlicedReflectionCaptureState(-1)
+,	RealTimeSlicedReflectionCaptureStateStep(0)
+,	bRealTimeSlicedReflectionCaptureStateStepDone(true)
 ,	RealTimeSlicedReflectionCaptureFrameNumber(0)
 ,   PathTracingSkylightColor(0, 0, 0, 0)
 ,	SimpleDirectionalLight(NULL)
diff --git a/Engine/Source/Runtime/Renderer/Private/ScenePrivate.h b/Engine/Source/Runtime/Renderer/Private/ScenePrivate.h
--- a/Engine/Source/Runtime/Renderer/Private/ScenePrivate.h
+++ b/Engine/Source/Runtime/Renderer/Private/ScenePrivate.h
@@ -2841,6 +2841,10 @@ public:
 
 	/** The current progress of the real time reflection capture when time sliced. */
 	int32 RealTimeSlicedReflectionCaptureState;
+	/** The current progress of each sub step of a state capture of the real time reflection capture when time sliced. */
+	int32 RealTimeSlicedReflectionCaptureStateStep;
+	/** When a state is sub divided, StepDone tells us if we can move onto the next step*/
+	bool bRealTimeSlicedReflectionCaptureStateStepDone;
 
 	/** Cache of Frame Number, used to detect first viewfamily */
 	uint64 RealTimeSlicedReflectionCaptureFrameNumber;
