Add HZB parameters & sampler to the HairStrands uniform buffer for easier access.

diff --git a/Engine/Source/Runtime/Renderer/Private/HairStrands/HairStrandsData.h b/Engine/Source/Runtime/Renderer/Private/HairStrands/HairStrandsData.h
--- a/Engine/Source/Runtime/Renderer/Private/HairStrands/HairStrandsData.h
+++ b/Engine/Source/Runtime/Renderer/Private/HairStrands/HairStrandsData.h
@@ -28,10 +28,12 @@ BEGIN_GLOBAL_SHADER_PARAMETER_STRUCT(FHairStrandsViewUniformParameters, )
 	SHADER_PARAMETER(float, HairDualScatteringRoughnessOverride)						// Override the roughness used for dual scattering (for hack/test purpose only)
 	SHADER_PARAMETER(FIntPoint, HairSampleViewportResolution)							// Maximum viewport resolution of the sample space
 	SHADER_PARAMETER(uint32, bHairTileValid)											// True if tile data are valid
+	SHADER_PARAMETER(FVector4f, HairOnlyDepthHZBParameters)								// HZB parameters
 	SHADER_PARAMETER_RDG_TEXTURE(Texture2D<float>, HairCoverageTexture)					// Hair pixel's coverage
 	SHADER_PARAMETER_RDG_TEXTURE(Texture2D, HairOnlyDepthTexture)						// Depth texture containing only hair depth (not strongly typed for legacy shader code reason)
 	SHADER_PARAMETER_RDG_TEXTURE(Texture2D, HairOnlyDepthClosestHZBTexture)				// HZB closest depth texture containing only hair depth (not strongly typed for legacy shader code reason)
 	SHADER_PARAMETER_RDG_TEXTURE(Texture2D, HairOnlyDepthFurthestHZBTexture)			// HZB furthest depth texture containing only hair depth (not strongly typed for legacy shader code reason)
+	SHADER_PARAMETER_SAMPLER(SamplerState, HairOnlyDepthHZBSampler)						// HZB sampler
 	SHADER_PARAMETER_RDG_TEXTURE(Texture2D<uint>, HairSampleOffset)						// Offset & count, for accessing pixel's samples, based on screen pixel position
 	SHADER_PARAMETER_RDG_TEXTURE(Texture2D<uint>, HairSampleCount)						// Total count of hair sample, in sample space
 	SHADER_PARAMETER_RDG_BUFFER_SRV(StructuredBuffer<FPackedHairSample>, HairSampleData)// Sample data (coverage, tangent, base color, ...), in sample space // HAIRSTRANDS_TODO: change this to be a uint4 so that we don't have to include the type for generated contant buffer
@@ -106,6 +108,7 @@ struct FHairStrandsVisibilityData
 	FRDGBufferRef	NodeCoord = nullptr;
 	FRDGBufferRef	NodeIndirectArg = nullptr;
 	uint32			NodeGroupSize = 0;
+	FVector4f		HairOnlyDepthHZBParameters = FVector4f::Zero();
 
 	FHairStrandsTiles TileData;
 
diff --git a/Engine/Source/Runtime/Renderer/Private/HairStrands/HairStrandsRendering.cpp b/Engine/Source/Runtime/Renderer/Private/HairStrands/HairStrandsRendering.cpp
--- a/Engine/Source/Runtime/Renderer/Private/HairStrands/HairStrandsRendering.cpp
+++ b/Engine/Source/Runtime/Renderer/Private/HairStrands/HairStrandsRendering.cpp
@@ -16,8 +16,10 @@ static TRDGUniformBufferRef<FHairStrandsViewUniformParameters> InternalCreateHai
 	{
 		Parameters->HairCoverageTexture = In->CoverageTexture;
 		Parameters->HairOnlyDepthTexture = In->HairOnlyDepthTexture;
+		Parameters->HairOnlyDepthHZBParameters = In->HairOnlyDepthHZBParameters;
 		Parameters->HairOnlyDepthClosestHZBTexture = In->HairOnlyDepthClosestHZBTexture;
 		Parameters->HairOnlyDepthFurthestHZBTexture = In->HairOnlyDepthFurthestHZBTexture;
+		Parameters->HairOnlyDepthHZBSampler = TStaticSamplerState<SF_Point>::GetRHI();
 		Parameters->HairSampleOffset = In->NodeIndex;
 		Parameters->HairSampleData = GraphBuilder.CreateSRV(In->NodeData);
 		Parameters->HairSampleCoords = GraphBuilder.CreateSRV(In->NodeCoord, FHairStrandsVisibilityData::NodeCoordFormat);
@@ -46,6 +48,7 @@ static TRDGUniformBufferRef<FHairStrandsViewUniformParameters> InternalCreateHai
 		if (!Parameters->HairOnlyDepthFurthestHZBTexture)
 		{
 			FRDGTextureRef BlackTexture = GSystemTextures.GetBlackDummy(GraphBuilder);
+			Parameters->HairOnlyDepthHZBParameters = FVector4f::Zero();
 			Parameters->HairOnlyDepthFurthestHZBTexture = BlackTexture;
 			Parameters->HairOnlyDepthClosestHZBTexture = BlackTexture;
 		}
@@ -64,8 +67,10 @@ static TRDGUniformBufferRef<FHairStrandsViewUniformParameters> InternalCreateHai
 		FRDGBufferSRVRef DummyBufferRG16SRV = GraphBuilder.CreateSRV(DummyBuffer, PF_R16G16_UINT);
 
 		Parameters->HairOnlyDepthTexture = FarDepth;
+		Parameters->HairOnlyDepthHZBParameters = FVector4f::Zero();
 		Parameters->HairOnlyDepthFurthestHZBTexture = BlackTexture;
 		Parameters->HairOnlyDepthClosestHZBTexture = Parameters->HairOnlyDepthFurthestHZBTexture;
+		Parameters->HairOnlyDepthHZBSampler = TStaticSamplerState<SF_Point>::GetRHI();
 		Parameters->HairCoverageTexture = BlackTexture;
 		Parameters->HairSampleCount = ZeroR32_UINT;
 		Parameters->HairSampleOffset = ZeroR32_UINT;
diff --git a/Engine/Source/Runtime/Renderer/Private/HairStrands/HairStrandsVisibility.cpp b/Engine/Source/Runtime/Renderer/Private/HairStrands/HairStrandsVisibility.cpp
--- a/Engine/Source/Runtime/Renderer/Private/HairStrands/HairStrandsVisibility.cpp
+++ b/Engine/Source/Runtime/Renderer/Private/HairStrands/HairStrandsVisibility.cpp
@@ -2779,9 +2779,20 @@ static void AddHairOnlyHZBPass(
 	FRDGBuilder& GraphBuilder,
 	const FViewInfo& View,
 	FRDGTextureRef HairDepthTexture,
+	FVector4f& OutHZBUvFactorAndInvFactor,
 	FRDGTextureRef& OutClosestHZBTexture,
 	FRDGTextureRef& OutFurthestHZBTexture)
 {
+	const FVector2D ViewportUVToHZBBufferUV(
+		float(View.ViewRect.Width()) / float(2 * View.HZBMipmap0Size.X),
+		float(View.ViewRect.Height()) / float(2 * View.HZBMipmap0Size.Y));
+
+	OutHZBUvFactorAndInvFactor = FVector4f(
+		ViewportUVToHZBBufferUV.X,
+		ViewportUVToHZBBufferUV.Y,
+		1.0f / ViewportUVToHZBBufferUV.X,
+		1.0f / ViewportUVToHZBBufferUV.Y);
+
 	BuildHZB(
 		GraphBuilder,
 		HairDepthTexture,
@@ -3732,6 +3743,7 @@ void RenderHairStrandsVisibilityBuffer(
 						GraphBuilder,
 						View,
 						HairOnlyDepthTexture,
+						VisibilityData.HairOnlyDepthHZBParameters,
 						VisibilityData.HairOnlyDepthClosestHZBTexture,
 						VisibilityData.HairOnlyDepthFurthestHZBTexture);
 				}
@@ -3930,6 +3942,7 @@ void RenderHairStrandsVisibilityBuffer(
 						GraphBuilder,
 						View,
 						HairOnlyDepthTexture,
+						VisibilityData.HairOnlyDepthHZBParameters,
 						VisibilityData.HairOnlyDepthClosestHZBTexture,
 						VisibilityData.HairOnlyDepthFurthestHZBTexture);
 				}
