Distance Field LWC support #2
- Compute wanted mips in translated world space.

diff --git a/Engine/Shaders/Private/DistanceFieldLightingShared.ush b/Engine/Shaders/Private/DistanceFieldLightingShared.ush
--- a/Engine/Shaders/Private/DistanceFieldLightingShared.ush
+++ b/Engine/Shaders/Private/DistanceFieldLightingShared.ush
@@ -148,9 +148,10 @@ uint GetCulledNumObjects()
 
 struct FDFObjectBounds
 {
-	float3 Center;		// World space bounds center
-	float SphereRadius;	// World space bounding sphere extent
-	float3 BoxExtent;	// World space AABB extent
+	float3 TranslatedCenter;	// Translated world space bounds center
+	float3 Center;				// World space bounds center (DF_LWC_TODO)
+	float SphereRadius;			// World space bounding sphere extent
+	float3 BoxExtent;			// World space AABB extent
 	uint OftenMoving;
 	uint bCastShadow;
 	bool bIsNaniteMesh;
@@ -165,7 +166,7 @@ uint NumSceneHeightfieldObjects;
 
 SamplerState DistanceFieldSampler;
 
-FDFObjectBounds LoadDFObjectBounds(uint ObjectIndex)
+FDFObjectBounds LoadDFObjectBounds(uint ObjectIndex, FLWCVector3 PreViewTranslation)
 {
 	FDFObjectBounds Bounds;
 
@@ -175,6 +176,7 @@ FDFObjectBounds LoadDFObjectBounds(uint ObjectIndex)
 	float4 Vector1 = SceneObjectBounds[ObjectIndex * DF_OBJECT_BOUNDS_STRIDE + 1];
 	float3 RelativeWorldPosition = Vector1.xyz;
 
+	Bounds.TranslatedCenter = LWCToFloat(LWCAdd(MakeLWCVector3(TilePosition, RelativeWorldPosition), PreViewTranslation));
 	Bounds.Center = LWCHackToFloat(MakeLWCVector3(TilePosition, RelativeWorldPosition)); // DF_LWC_TODO
 	Bounds.SphereRadius = Vector1.w;
 
@@ -190,6 +192,11 @@ FDFObjectBounds LoadDFObjectBounds(uint ObjectIndex)
 	return Bounds;
 }
 
+FDFObjectBounds LoadDFObjectBounds(uint ObjectIndex)
+{
+	return LoadDFObjectBounds(ObjectIndex, MakeLWCVector3(0, 0)); // DF_LWC_TODO
+}
+
 struct FHeightfieldObjectBounds
 {
 	float3 BoxOrigin;	// World space AABB center
diff --git a/Engine/Shaders/Private/DistanceFieldStreaming.usf b/Engine/Shaders/Private/DistanceFieldStreaming.usf
--- a/Engine/Shaders/Private/DistanceFieldStreaming.usf
+++ b/Engine/Shaders/Private/DistanceFieldStreaming.usf
@@ -94,9 +94,9 @@ void CopyDistanceFieldIndirectionAtlasCS(
 RWStructuredBuffer<uint> RWDistanceFieldAssetWantedNumMips;
 RWStructuredBuffer<uint> RWDistanceFieldAssetStreamingRequests;
 
-float3 Mip1WorldCenter;
+float3 Mip1WorldTranslatedCenter;
 float3 Mip1WorldExtent;
-float3 Mip2WorldCenter;
+float3 Mip2WorldTranslatedCenter;
 float3 Mip2WorldExtent;
 
 float ComputeSquaredDistanceBetweenAABBs(float3 CenterA, float3 ExtentA, float3 CenterB, float3 ExtentB)
@@ -127,12 +127,12 @@ void ComputeDistanceFieldAssetWantedMipsCS(
 		uint WantedNumMips = 1;
 		if (DebugForceNumMips == 0)
 		{
-			FDFObjectBounds DFObjectBounds = LoadDFObjectBounds(ObjectIndex);
-			float OuterDistanceSq = ComputeSquaredDistanceBetweenAABBs(Mip1WorldCenter, Mip1WorldExtent, DFObjectBounds.Center, DFObjectBounds.BoxExtent);
+			FDFObjectBounds DFObjectBounds = LoadDFObjectBounds(ObjectIndex, PrimaryView.PreViewTranslation);
+			float OuterDistanceSq = ComputeSquaredDistanceBetweenAABBs(Mip1WorldTranslatedCenter, Mip1WorldExtent, DFObjectBounds.TranslatedCenter, DFObjectBounds.BoxExtent);
  
 			if (OuterDistanceSq <= 0)
 			{
-				float InnerDistanceSq = ComputeSquaredDistanceBetweenAABBs(Mip2WorldCenter, Mip2WorldExtent, DFObjectBounds.Center, DFObjectBounds.BoxExtent);
+				float InnerDistanceSq = ComputeSquaredDistanceBetweenAABBs(Mip2WorldTranslatedCenter, Mip2WorldExtent, DFObjectBounds.TranslatedCenter, DFObjectBounds.BoxExtent);
 				WantedNumMips = InnerDistanceSq <= 0 ? 3 : 2;
 			}
 		}
diff --git a/Engine/Source/Runtime/Renderer/Private/DistanceFieldStreaming.cpp b/Engine/Source/Runtime/Renderer/Private/DistanceFieldStreaming.cpp
--- a/Engine/Source/Runtime/Renderer/Private/DistanceFieldStreaming.cpp
+++ b/Engine/Source/Runtime/Renderer/Private/DistanceFieldStreaming.cpp
@@ -234,13 +234,14 @@ class FComputeDistanceFieldAssetWantedMipsCS : public FGlobalShader
 	SHADER_USE_PARAMETER_STRUCT(FComputeDistanceFieldAssetWantedMipsCS, FGlobalShader)
 
 	BEGIN_SHADER_PARAMETER_STRUCT(FParameters, )
+		SHADER_PARAMETER_STRUCT_REF(FViewUniformShaderParameters, View)
 		SHADER_PARAMETER_RDG_BUFFER_UAV(RWStructuredBuffer<uint>, RWDistanceFieldAssetWantedNumMips)
 		SHADER_PARAMETER_RDG_BUFFER_UAV(RWStructuredBuffer<uint>, RWDistanceFieldAssetStreamingRequests)
 		SHADER_PARAMETER_STRUCT_INCLUDE(FDistanceFieldObjectBufferParameters, DistanceFieldObjectBuffers)
 		SHADER_PARAMETER(int32, DebugForceNumMips)
-		SHADER_PARAMETER(FVector3f, Mip1WorldCenter)
+		SHADER_PARAMETER(FVector3f, Mip1WorldTranslatedCenter)
 		SHADER_PARAMETER(FVector3f, Mip1WorldExtent)
-		SHADER_PARAMETER(FVector3f, Mip2WorldCenter)
+		SHADER_PARAMETER(FVector3f, Mip2WorldTranslatedCenter)
 		SHADER_PARAMETER(FVector3f, Mip2WorldExtent)
 	END_SHADER_PARAMETER_STRUCT()
 
@@ -1089,7 +1090,7 @@ void FDistanceFieldSceneData::GenerateStreamingRequests(
 
 		{
 			FComputeDistanceFieldAssetWantedMipsCS::FParameters* PassParameters = GraphBuilder.AllocParameters<FComputeDistanceFieldAssetWantedMipsCS::FParameters>();
-
+			PassParameters->View = View.ViewUniformBuffer;
 			checkf(DistanceField::NumMips == 3, TEXT("Shader needs to be updated"));
 			PassParameters->RWDistanceFieldAssetWantedNumMips = GraphBuilder.CreateUAV(WantedNumMips);
 			PassParameters->RWDistanceFieldAssetStreamingRequests = GraphBuilder.CreateUAV(StreamingRequestsBuffer);
@@ -1097,9 +1098,9 @@ void FDistanceFieldSceneData::GenerateStreamingRequests(
 			PassParameters->DebugForceNumMips = CVarDebugForceNumMips.GetValueOnRenderThread();
 			extern int32 GAOGlobalDistanceFieldNumClipmaps;
 			// Request Mesh SDF mips based off of the Global SDF clipmaps
-			PassParameters->Mip1WorldCenter = (FVector3f)View.ViewMatrices.GetViewOrigin();
+			PassParameters->Mip1WorldTranslatedCenter = FVector3f(View.ViewMatrices.GetViewOrigin() + View.ViewMatrices.GetPreViewTranslation());
 			PassParameters->Mip1WorldExtent = FVector3f(GlobalDistanceField::GetClipmapExtent(GAOGlobalDistanceFieldNumClipmaps - 1, Scene, bLumenEnabled));
-			PassParameters->Mip2WorldCenter = (FVector3f)View.ViewMatrices.GetViewOrigin();
+			PassParameters->Mip2WorldTranslatedCenter = FVector3f(View.ViewMatrices.GetViewOrigin() + View.ViewMatrices.GetPreViewTranslation());
 			PassParameters->Mip2WorldExtent = FVector3f(GlobalDistanceField::GetClipmapExtent(FMath::Max<int32>(GAOGlobalDistanceFieldNumClipmaps / 2 - 1, 0), Scene, bLumenEnabled));
 
 			auto ComputeShader = GlobalShaderMap->GetShader<FComputeDistanceFieldAssetWantedMipsCS>();
