Clamp adaptive voxel size to be aligned macro-group AABB onto voxel page size, when reaching growing.

This reduces the size of the macro-group AABBs, and thus avoid false-positive intersection when traversing all the macro-groups during tracing.

diff --git a/Engine/Shaders/Private/HairStrands/HairStrandsAABB.usf b/Engine/Shaders/Private/HairStrands/HairStrandsAABB.usf
--- a/Engine/Shaders/Private/HairStrands/HairStrandsAABB.usf
+++ b/Engine/Shaders/Private/HairStrands/HairStrandsAABB.usf
@@ -11,6 +11,7 @@ uint MacroGroupValid;
 
 float PixelSizeAtDepth1;
 float NumPixelPerVoxel;
+uint  VoxelPageResolution;
 
 // Group data packed as {uint3 AABBMin, uint3 AABBMax}
 Buffer<int> InGroupAABBBuffer0;
@@ -55,9 +56,18 @@ void Main(uint2 DispatchThreadId : SV_DispatchThreadID)
 
 		// Compute voxel size to have 1px = 1 voxels
 		const float PixelWorldSize = PixelSizeAtDepth1 * ViewBound.Min.z;
-		const float VoxelWorldSize = PixelWorldSize * NumPixelPerVoxel;
+		float VoxelWorldSize = PixelWorldSize * NumPixelPerVoxel;
 		WriteHairAABB(MacroGroupId, Bound, OutMacroGroupAABBBuffer);
 
+		// Ensure the MacroGroup AABB has, at minimum, the size of a voxel page
+		const float3 Extend = Bound.Max - Bound.Min;
+		const float MaxExtent = max3(Extend.x, Extend.y, Extend.z);
+		const float BoundVoxelCount = MaxExtent / VoxelWorldSize;
+		if (BoundVoxelCount < float(VoxelPageResolution))
+		{
+			VoxelWorldSize = MaxExtent / float(VoxelPageResolution);
+		}
+
 		const float MinVoxelWorldSize = 0.01f; // 0.01mm
 		OutMacroGroupVoxelSizeBuffer[MacroGroupId] = max(MinVoxelWorldSize, VoxelWorldSize);
 	}
diff --git a/Engine/Source/Runtime/Renderer/Private/HairStrands/HairStrandsCluster.cpp b/Engine/Source/Runtime/Renderer/Private/HairStrands/HairStrandsCluster.cpp
--- a/Engine/Source/Runtime/Renderer/Private/HairStrands/HairStrandsCluster.cpp
+++ b/Engine/Source/Runtime/Renderer/Private/HairStrands/HairStrandsCluster.cpp
@@ -28,6 +28,7 @@ class FHairMacroGroupAABBCS : public FGlobalShader
 
 		SHADER_PARAMETER(float, PixelSizeAtDepth1)
 		SHADER_PARAMETER(float, NumPixelPerVoxel)
+		SHADER_PARAMETER(uint32, VoxelPageResolution)
 
 		SHADER_PARAMETER_RDG_BUFFER_SRV(Buffer, InGroupAABBBuffer0)
 		SHADER_PARAMETER_RDG_BUFFER_SRV(Buffer, InGroupAABBBuffer1)
@@ -53,6 +54,8 @@ public:
 
 IMPLEMENT_GLOBAL_SHADER(FHairMacroGroupAABBCS, "/Engine/Private/HairStrands/HairStrandsAABB.usf", "Main", SF_Compute);
 
+void GetVoxelPageResolution(uint32& OutPageResolution, uint32& OutPageResolutionLog2);
+
 static void AddHairMacroGroupAABBPass(
 	FRDGBuilder& GraphBuilder,
 	const FViewInfo& View, 
@@ -69,6 +72,10 @@ static void AddHairMacroGroupAABBPass(
 	const float vFOV = FMath::DegreesToRadians(View.FOV);
 	const float PixelSizeAtDepth1 = FMath::Tan(vFOV * 0.5f) / (0.5f * Resolution.Y);
 	
+	uint32 VoxelPageResolution = 0;
+	uint32 VoxelPageResolutionLog2 = 1;
+	GetVoxelPageResolution(VoxelPageResolution, VoxelPageResolutionLog2);
+
 	const float NumPixelPerVoxel = FMath::Clamp(GHairVirtualVoxel_NumPixelPerVoxel, 1.f, 50.f);
 
 	const uint32 GroupPerPass = 8;
@@ -81,9 +88,10 @@ static void AddHairMacroGroupAABBPass(
 		Parameters->MacroGroupId = MacroGroupId;
 		Parameters->OutMacroGroupAABBBuffer = OutHairMacroGroupAABBBufferUAV;
 		Parameters->OutMacroGroupVoxelSizeBuffer = OutMacroGroupVoxelSizeBufferUAV;
-		Parameters->PixelSizeAtDepth1 = PixelSizeAtDepth1;
-		Parameters->NumPixelPerVoxel  = NumPixelPerVoxel;
-		Parameters->View			  = View.ViewUniformBuffer;
+		Parameters->PixelSizeAtDepth1   = PixelSizeAtDepth1;
+		Parameters->NumPixelPerVoxel    = NumPixelPerVoxel;
+		Parameters->VoxelPageResolution = VoxelPageResolution;
+		Parameters->View			    = View.ViewUniformBuffer;
 
 		uint32 MacroGroupValid = 1;
 		uint32 CurrentGroupIt = 1;
diff --git a/Engine/Source/Runtime/Renderer/Private/HairStrands/HairStrandsVoxelization.cpp b/Engine/Source/Runtime/Renderer/Private/HairStrands/HairStrandsVoxelization.cpp
--- a/Engine/Source/Runtime/Renderer/Private/HairStrands/HairStrandsVoxelization.cpp
+++ b/Engine/Source/Runtime/Renderer/Private/HairStrands/HairStrandsVoxelization.cpp
@@ -119,6 +119,12 @@ static FAutoConsoleVariableRef CVarHairVirtualVoxelAdaptive_CorrectionThreshold(
 static int32 GHairVirtualVoxel_JitterMode = 1;
 static FAutoConsoleVariableRef CVarHairVirtualVoxel_JitterMode(TEXT("r.HairStrands.Voxelization.Virtual.Jitter"), GHairVirtualVoxel_JitterMode, TEXT("Change jittered for voxelization/traversal. 0: No jitter 1: Regular randomized jitter: 2: Constant Jitter (default = 1)"), ECVF_RenderThreadSafe);
 
+void GetVoxelPageResolution(uint32& OutPageResolution, uint32& OutPageResolutionLog2)
+{
+	OutPageResolutionLog2 = FMath::CeilLogTwo(FMath::Clamp(GHairVirtualVoxel_PageResolution, 2, 256));
+	OutPageResolution = (1u << OutPageResolutionLog2);
+}
+
 static bool IsHairStrandsAdaptiveVoxelAllocationEnable()
 {
 	return GHairVirtualVoxelAdaptive_Enable > 0;
@@ -985,11 +991,11 @@ static FHairStrandsVoxelResources AllocateVirtualVoxelResources(
 
 	FHairStrandsVoxelResources Out;
 
+	GetVoxelPageResolution(Out.Parameters.Common.PageResolution, Out.Parameters.Common.PageResolutionLog2);
+
 	Out.Parameters.Common.PageCountResolution		= PageCountResolution;
 	Out.Parameters.Common.PageCount					= PageCount;
 	Out.Parameters.Common.CPUMinVoxelWorldSize		= CPUMinVoxelWorldSize;
-	Out.Parameters.Common.PageResolutionLog2		= FMath::CeilLogTwo(FMath::Clamp(GHairVirtualVoxel_PageResolution, 2, 256));
-	Out.Parameters.Common.PageResolution			= (1u << Out.Parameters.Common.PageResolutionLog2);
 	Out.Parameters.Common.PageTextureResolution		= Out.Parameters.Common.PageCountResolution * Out.Parameters.Common.PageResolution;
 	Out.Parameters.Common.JitterMode				= FMath::Clamp(GHairVirtualVoxel_JitterMode, 0, 2);
 
