Distance Field LWC support #3
- Upload heightfield bounds tile + relative world position
- Upload heightfield tile + relative world matrices
- Positions are still converted to float when loaded in shader.

diff --git a/Engine/Shaders/Private/DistanceFieldLightingShared.ush b/Engine/Shaders/Private/DistanceFieldLightingShared.ush
--- a/Engine/Shaders/Private/DistanceFieldLightingShared.ush
+++ b/Engine/Shaders/Private/DistanceFieldLightingShared.ush
@@ -143,8 +143,8 @@ uint GetCulledNumObjects()
 // In float4's.  Must match equivalent C++ variables.
 #define DF_OBJECT_BOUNDS_STRIDE 3
 #define DF_OBJECT_DATA_STRIDE 10
-#define HEIGHTFIELD_OBJECT_BOUNDS_STRIDE 2
-#define HEIGHTFIELD_OBJECT_DATA_STRIDE 6
+#define HEIGHTFIELD_OBJECT_BOUNDS_STRIDE 3
+#define HEIGHTFIELD_OBJECT_DATA_STRIDE 7
 
 struct FDFObjectBounds
 {
@@ -199,8 +199,8 @@ FDFObjectBounds LoadDFObjectBounds(uint ObjectIndex)
 
 struct FHeightfieldObjectBounds
 {
-	float3 BoxOrigin;	// World space AABB center
-	float3 BoxExtent;	// World space AABB extent
+	FLWCVector3 BoxOrigin;	// World space AABB center
+	float3 BoxExtent;		// World space AABB extent
 };
 
 FHeightfieldObjectBounds LoadHeightfieldObjectBounds(uint ObjectIndex)
@@ -208,10 +208,15 @@ FHeightfieldObjectBounds LoadHeightfieldObjectBounds(uint ObjectIndex)
 	FHeightfieldObjectBounds Bounds;
 
 	float4 Vector0 = SceneHeightfieldObjectBounds[ObjectIndex * HEIGHTFIELD_OBJECT_BOUNDS_STRIDE + 0];
-	Bounds.BoxOrigin = Vector0.xyz;
-
 	float4 Vector1 = SceneHeightfieldObjectBounds[ObjectIndex * HEIGHTFIELD_OBJECT_BOUNDS_STRIDE + 1];
-	Bounds.BoxExtent = Vector1.xyz;
+
+	float3 TilePosition = Vector0.xyz;
+	float3 RelativeWorldPosition = Vector1.xyz;
+
+	Bounds.BoxOrigin = MakeLWCVector3(TilePosition, RelativeWorldPosition);
+
+	float4 Vector2 = SceneHeightfieldObjectBounds[ObjectIndex * HEIGHTFIELD_OBJECT_BOUNDS_STRIDE + 2];
+	Bounds.BoxExtent = Vector2.xyz;
 
 	return Bounds;
 }
@@ -289,7 +294,7 @@ FDFObjectData LoadDFObjectDataFromBuffer(StructuredBuffer<float4> SourceBuffer,
 
 struct FHeightfieldObjectData
 {
-	float4x4 WorldToLocal;
+	FLWCInverseMatrix WorldToLocal;
 	float4 SizeScale;
 	float4 AtlasUVScaleBias;
 	float4 VisibilityAtlasUVScaleBias;
@@ -299,14 +304,18 @@ FHeightfieldObjectData LoadHeightfieldObjectDataFromBuffer(StructuredBuffer<floa
 {
 	FHeightfieldObjectData Data;
 
-	float4 V0 = SourceBuffer[ObjectIndex * HEIGHTFIELD_OBJECT_DATA_STRIDE + 0];
-	float4 V1 = SourceBuffer[ObjectIndex * HEIGHTFIELD_OBJECT_DATA_STRIDE + 1];
-	float4 V2 = SourceBuffer[ObjectIndex * HEIGHTFIELD_OBJECT_DATA_STRIDE + 2];
-	Data.WorldToLocal = transpose(float4x4(V0, V1, V2, float4(0.0f, 0.0f, 0.0f, 1.0f)));
+	float3 TilePosition = SourceBuffer[ObjectIndex * HEIGHTFIELD_OBJECT_DATA_STRIDE + 0].xyz;
+
+	float4 V0 = SourceBuffer[ObjectIndex * HEIGHTFIELD_OBJECT_DATA_STRIDE + 1];
+	float4 V1 = SourceBuffer[ObjectIndex * HEIGHTFIELD_OBJECT_DATA_STRIDE + 2];
+	float4 V2 = SourceBuffer[ObjectIndex * HEIGHTFIELD_OBJECT_DATA_STRIDE + 3];
+	float4x4 RelativeWorldToVolume = transpose(float4x4(V0, V1, V2, float4(0.0f, 0.0f, 0.0f, 1.0f)));
+
+	Data.WorldToLocal = MakeLWCInverseMatrix4x3(TilePosition, RelativeWorldToVolume);
 
-	Data.SizeScale = SourceBuffer[ObjectIndex * HEIGHTFIELD_OBJECT_DATA_STRIDE + 3];
-	Data.AtlasUVScaleBias = SourceBuffer[ObjectIndex * HEIGHTFIELD_OBJECT_DATA_STRIDE + 4];
-	Data.VisibilityAtlasUVScaleBias = SourceBuffer[ObjectIndex * HEIGHTFIELD_OBJECT_DATA_STRIDE + 5];
+	Data.SizeScale = SourceBuffer[ObjectIndex * HEIGHTFIELD_OBJECT_DATA_STRIDE + 4];
+	Data.AtlasUVScaleBias = SourceBuffer[ObjectIndex * HEIGHTFIELD_OBJECT_DATA_STRIDE + 5];
+	Data.VisibilityAtlasUVScaleBias = SourceBuffer[ObjectIndex * HEIGHTFIELD_OBJECT_DATA_STRIDE + 6];
 
 	return Data;
 }
diff --git a/Engine/Shaders/Private/DistanceFieldShadowing.usf b/Engine/Shaders/Private/DistanceFieldShadowing.usf
--- a/Engine/Shaders/Private/DistanceFieldShadowing.usf
+++ b/Engine/Shaders/Private/DistanceFieldShadowing.usf
@@ -122,7 +122,7 @@ void CullObjectsForShadowCS(
         }
 #else
 		FHeightfieldObjectBounds Bounds = LoadHeightfieldObjectBounds(ObjectIndex);
-		if (ShadowConvexHullIntersectBox(Bounds.BoxOrigin, Bounds.BoxExtent))
+		if (ShadowConvexHullIntersectBox(LWCHackToFloat(Bounds.BoxOrigin), Bounds.BoxExtent)) // DF_LWC_TODO
 		{
 			uint DestIndex;
 			InterlockedAdd(RWObjectIndirectArguments[1], 1U, DestIndex);
@@ -234,7 +234,7 @@ void ShadowObjectCullVS(
 #else
 
 	FHeightfieldObjectBounds Bounds = LoadHeightfieldObjectBounds(ObjectIndex);
-	float3 BoxOrigin = Bounds.BoxOrigin;
+	float3 BoxOrigin = LWCHackToFloat(Bounds.BoxOrigin); // DF_LWC_TODO
 	float3 BoxExtent = Bounds.BoxExtent;
 
 	BoxExtent += MinExpandRadius;
@@ -299,7 +299,7 @@ FObjectViewSpaceBox GetObjectViewSpaceBox(uint ObjectIndex)
 		float3 WorldBoundsPosition = mul(float4(UnitBoxVertex * DFObjectData.VolumePositionExtent, 1.0f), DFObjectData.VolumeToWorld).xyz;
 #else
 		FHeightfieldObjectBounds HeightfieldObjectBounds = LoadHeightfieldObjectBounds(ObjectIndex);
-		float3 WorldBoundsPosition = UnitBoxVertex * HeightfieldObjectBounds.BoxExtent + HeightfieldObjectBounds.BoxOrigin;
+		float3 WorldBoundsPosition = UnitBoxVertex * HeightfieldObjectBounds.BoxExtent + LWCHackToFloat(HeightfieldObjectBounds.BoxOrigin); // DF_LWC_TODO
 #endif
 
 		float3 ViewSpacePosition = mul(float4(WorldBoundsPosition, 1.0f), WorldToShadow).xyz;
diff --git a/Engine/Shaders/Private/DistanceFieldShadowingShared.ush b/Engine/Shaders/Private/DistanceFieldShadowingShared.ush
--- a/Engine/Shaders/Private/DistanceFieldShadowingShared.ush
+++ b/Engine/Shaders/Private/DistanceFieldShadowingShared.ush
@@ -342,7 +342,7 @@ float ShadowRayTraceThroughCulledHeightFieldObjects(
 		}
 
 		FHeightfieldObjectData HeightfieldObject = LoadHeightfieldObjectData(ObjectIndex);
-		float4x4 WorldToLocal = HeightfieldObject.WorldToLocal;
+		float4x4 WorldToLocal = LWCHackToFloat(HeightfieldObject.WorldToLocal); // DF_LWC_TODO
 		float2 HeightFieldSize = HeightfieldObject.SizeScale.xy;
 		float WorldToLocalScale = HeightfieldObject.SizeScale.z;
 
diff --git a/Engine/Source/Runtime/Renderer/Private/DistanceFieldLightingShared.h b/Engine/Source/Runtime/Renderer/Private/DistanceFieldLightingShared.h
--- a/Engine/Source/Runtime/Renderer/Private/DistanceFieldLightingShared.h
+++ b/Engine/Source/Runtime/Renderer/Private/DistanceFieldLightingShared.h
@@ -47,8 +47,8 @@ enum EDistanceFieldPrimitiveType
 static const int32 GDistanceFieldObjectDataStride = 10;
 static const int32 GDistanceFieldObjectBoundsStride = 3;
 
-static const int32 GHeightFieldObjectDataStride = 6;
-static const int32 GHeightFieldObjectBoundsStride = 2;
+static const int32 GHeightFieldObjectDataStride = 7;
+static const int32 GHeightFieldObjectBoundsStride = 3;
 
 template <EDistanceFieldPrimitiveType PrimitiveType>
 class TDistanceFieldObjectBuffers
diff --git a/Engine/Source/Runtime/Renderer/Private/DistanceFieldObjectManagement.cpp b/Engine/Source/Runtime/Renderer/Private/DistanceFieldObjectManagement.cpp
--- a/Engine/Source/Runtime/Renderer/Private/DistanceFieldObjectManagement.cpp
+++ b/Engine/Source/Runtime/Renderer/Private/DistanceFieldObjectManagement.cpp
@@ -953,24 +953,41 @@ void FSceneRenderer::UpdateGlobalHeightFieldObjectBuffers(FRDGBuilder& GraphBuil
 								FHeightfieldComponentDescription HeightFieldCompDesc(Primitive->Proxy->GetLocalToWorld());
 								Primitive->Proxy->GetHeightfieldRepresentation(HeightNormalTexture, DiffuseColorTexture, VisibilityTexture, HeightFieldCompDesc);
 
-								const FBoxSphereBounds& Bounds = Primitive->Proxy->GetBounds();
-								const FBox BoxBound = Bounds.GetBox();
-								UploadObjectBounds[0] = FVector4f((FVector3f)BoxBound.GetCenter(), Bounds.SphereRadius);
-								UploadObjectBounds[1] = FVector4f((FVector3f)BoxBound.GetExtent(), 0.f);
+								{
+									const FBoxSphereBounds& Bounds = Primitive->Proxy->GetBounds();
+									const FBox BoxBound = Bounds.GetBox();
+
+									const FLargeWorldRenderPosition AbsoluteWorldPosition(BoxBound.GetCenter());
+
+									const FVector4f ObjectBoundingSphere(AbsoluteWorldPosition.GetOffset(), Bounds.SphereRadius);
+
+									UploadObjectBounds[0] = AbsoluteWorldPosition.GetTile();
+									UploadObjectBounds[1] = ObjectBoundingSphere;
+									UploadObjectBounds[2] = FVector4f((FVector3f)BoxBound.GetExtent(), 0.f);
+								}
 
 								const FMatrix& LocalToWorld = HeightFieldCompDesc.LocalToWorld;
 								check(LocalToWorld.GetMaximumAxisScale() > 0.f);
-								const FMatrix44f WorldToLocalT = FMatrix44f(LocalToWorld.Inverse().GetTransposed()); // LWC_TODO
-								UploadObjectData[0] = *(const FVector4f*)&WorldToLocalT.M[0];
-								UploadObjectData[1] = *(const FVector4f*)&WorldToLocalT.M[1];
-								UploadObjectData[2] = *(const FVector4f*)&WorldToLocalT.M[2];
+
+								const FLargeWorldRenderPosition WorldPosition(LocalToWorld.GetOrigin());
+								const FVector TilePositionOffset = WorldPosition.GetTileOffset();
+
+								// Inverse on FMatrix44f can generate NaNs if the source matrix contains large scaling, so do it in double precision.
+								FMatrix LocalToRelativeWorld = FLargeWorldRenderScalar::MakeToRelativeWorldMatrixDouble(TilePositionOffset, LocalToWorld);
+
+								UploadObjectData[0] = WorldPosition.GetTile();
+
+								const FMatrix44f WorldToLocalT = FMatrix44f(LocalToRelativeWorld.Inverse().GetTransposed());
+								UploadObjectData[1] = *(const FVector4f*)&WorldToLocalT.M[0];
+								UploadObjectData[2] = *(const FVector4f*)&WorldToLocalT.M[1];
+								UploadObjectData[3] = *(const FVector4f*)&WorldToLocalT.M[2];
 
 								const FIntRect& HeightFieldRect = HeightFieldCompDesc.HeightfieldRect;
 								const float WorldToLocalScale = FMath::Min3(
 									WorldToLocalT.GetColumn(0).Size(),
 									WorldToLocalT.GetColumn(1).Size(),
 									WorldToLocalT.GetColumn(2).Size());
-								UploadObjectData[3] = FVector4f(HeightFieldRect.Width(), HeightFieldRect.Height(), WorldToLocalScale, 0.f);
+								UploadObjectData[4] = FVector4f(HeightFieldRect.Width(), HeightFieldRect.Height(), WorldToLocalScale, 0.f);
 
 								FVector4f HeightUVScaleBias(ForceInitToZero);
 								if (HeightNormalTexture)
@@ -988,7 +1005,7 @@ void FSceneRenderer::UpdateGlobalHeightFieldObjectBuffers(FRDGBuilder& GraphBuil
 											HeightFieldScaleBias.W * ScaleBias.Y + ScaleBias.W);
 									}
 								}
-								UploadObjectData[4] = HeightUVScaleBias;
+								UploadObjectData[5] = HeightUVScaleBias;
 
 								FVector4f VisUVScaleBias(ForceInitToZero);
 								if (VisibilityTexture)
@@ -1000,7 +1017,7 @@ void FSceneRenderer::UpdateGlobalHeightFieldObjectBuffers(FRDGBuilder& GraphBuil
 										VisUVScaleBias = FVector4f(1.f / HeightFieldRect.Width() * ScaleBias.X, 1.f / HeightFieldRect.Height() * ScaleBias.Y, ScaleBias.Z, ScaleBias.W);
 									}
 								}
-								UploadObjectData[5] = VisUVScaleBias;
+								UploadObjectData[6] = VisUVScaleBias;
 							}
 						}
 
