Skip post process material pass when all the pixels will fail the stencil test of the material. Corrects certain stencil test operations from executing full screen when not needed. Saves performance on platforms that are doing the stencil test in the pixel shader with a clip.

Cherry-pickked from 18906109, including my additional fixes:
*Setting RDG output to scenecolor when skipping post process material due to full stencil fail.(Avoids tripping the 'unwritten' RDG validation warning)
*Swapped operand order for stencilval and stencilref to match what the internals are doing.(This was wrong, only Equals and Not Equals worked)
*Always test the stencil clear value against the stencil ref. The stencil ref might be relying on the clear value even if no primitives write stencil. (for example if material sets Equals 0)

#rnx

diff --git a/Engine/Source/Runtime/Renderer/Private/PostProcess/PostProcessMaterial.cpp b/Engine/Source/Runtime/Renderer/Private/PostProcess/PostProcessMaterial.cpp
--- a/Engine/Source/Runtime/Renderer/Private/PostProcess/PostProcessMaterial.cpp
+++ b/Engine/Source/Runtime/Renderer/Private/PostProcess/PostProcessMaterial.cpp
@@ -156,22 +156,22 @@ bool PostProcessStencilTest(uint32 StencilValue, uint32 StencilComp, uint32 Sten
 	switch (StencilComp)
 	{
 	case EMaterialStencilCompare::MSC_Less:
-		bStencilTestPassed = (StencilValue < StencilRef);
+		bStencilTestPassed = (StencilRef < StencilValue);
 		break;
 	case EMaterialStencilCompare::MSC_LessEqual:
-		bStencilTestPassed = (StencilValue <= StencilRef);
+		bStencilTestPassed = (StencilRef <= StencilValue);
 		break;
 	case EMaterialStencilCompare::MSC_GreaterEqual:
-		bStencilTestPassed = (StencilValue >= StencilRef);
+		bStencilTestPassed = (StencilRef >= StencilValue);
 		break;
 	case EMaterialStencilCompare::MSC_Equal:
-		bStencilTestPassed = (StencilValue == StencilRef);
+		bStencilTestPassed = (StencilRef == StencilValue);
 		break;
 	case EMaterialStencilCompare::MSC_Greater:
-		bStencilTestPassed = (StencilValue > StencilRef);
+		bStencilTestPassed = (StencilRef > StencilValue);
 		break;
 	case EMaterialStencilCompare::MSC_NotEqual:
-		bStencilTestPassed = (StencilValue != StencilRef);
+		bStencilTestPassed = (StencilRef != StencilValue);
 		break;
 	case EMaterialStencilCompare::MSC_Never:
 		bStencilTestPassed = false;
@@ -638,12 +638,9 @@ FScreenPassTexture AddPostProcessMaterialPass(
 	{
 		bool bFailStencil = true;
 
-		// when the depthStencilTexture was not produced, check if all the pixels will fail the clear value
-		if (!DepthStencilTexture)
-		{
-			uint32 StencilClearValue = Inputs.CustomDepthTexture ? Inputs.CustomDepthTexture->Desc.ClearValue.Value.DSValue.Stencil : 0;
-			bFailStencil &= PostProcessStencilTest(StencilClearValue, Material->GetStencilCompare(), PostProcessMaterialParameters->MobileStencilValueRef);
-		}
+		// Always check against clear value, since a material might want to perform operations against that value
+		uint32 StencilClearValue = Inputs.CustomDepthTexture ? Inputs.CustomDepthTexture->Desc.ClearValue.Value.DSValue.Stencil : 0;
+		bFailStencil &= PostProcessStencilTest(StencilClearValue, Material->GetStencilCompare(), PostProcessMaterialParameters->MobileStencilValueRef);
 
 		for (const uint32& Value : View.CustomDepthStencilValues)
 		{
@@ -700,6 +697,10 @@ FScreenPassTexture AddPostProcessMaterialPass(
 		}
 	}
 	}
+	else
+	{
+		Output = FScreenPassRenderTarget(SceneColor, ERenderTargetLoadAction::ENoAction);
+	}
 
 	return MoveTemp(Output);
 }
