Improve path tracer invalidation through transform changes to reduce the number of spurious invalidations

Rather than performing the change detection on the game thread, perform it only for the rendering primitive that might impact the path tracer. In particular, this avoids updates to the static meshes representing cameras. The camera transforms are handled on a per view basis, which means cameras not being moved will not have their output invalidated.

diff --git a/Engine/Source/Runtime/Engine/Private/Components/SceneComponent.cpp b/Engine/Source/Runtime/Engine/Private/Components/SceneComponent.cpp
--- a/Engine/Source/Runtime/Engine/Private/Components/SceneComponent.cpp
+++ b/Engine/Source/Runtime/Engine/Private/Components/SceneComponent.cpp
@@ -786,15 +786,6 @@ void USceneComponent::PropagateTransformUpdate(bool bTransformChanged, EUpdateTr
 
 			// Flag render transform as dirty
 			MarkRenderTransformDirty();
-
-			// Invalidate the path tracer because we know something has changed about the transform which most likely
-			// has invalidated the image. If we hooked into these updates at a deeper level, we would also be notified
-			// when the transforms are marked dirty even when they haven't changed.
-			const UWorld* World = GetWorld();
-			if (World != nullptr && World->Scene != nullptr)
-			{
-				World->Scene->InvalidatePathTracedOutput();
-			}
 		}
 		
 		{
diff --git a/Engine/Source/Runtime/Renderer/Private/RendererScene.cpp b/Engine/Source/Runtime/Renderer/Private/RendererScene.cpp
--- a/Engine/Source/Runtime/Renderer/Private/RendererScene.cpp
+++ b/Engine/Source/Runtime/Renderer/Private/RendererScene.cpp
@@ -4327,6 +4327,15 @@ inline void FScene::UpdateRayTracingGroupBounds_UpdatePrimitives(const Experimen
 }
 #endif
 
+static inline bool IsPrimitiveRelevantToPathTracing(FPrimitiveSceneInfo* PrimitiveSceneInfo)
+{
+	// returns true if the primitive is likely to impact the path traced image
+	return PrimitiveSceneInfo->Proxy != nullptr &&
+		PrimitiveSceneInfo->Proxy->IsVisibleInRayTracing() &&
+		PrimitiveSceneInfo->Proxy->ShouldRenderInMainPass() &&
+		PrimitiveSceneInfo->Proxy->IsDrawnInGame();
+}
+
 void FScene::UpdateAllPrimitiveSceneInfos(FRDGBuilder& GraphBuilder, bool bAsyncCreateLPIs)
 {
 	TRACE_CPUPROFILER_EVENT_SCOPE(Scene::UpdateAllPrimitiveSceneInfos);
@@ -4641,6 +4650,7 @@ void FScene::UpdateAllPrimitiveSceneInfos(FRDGBuilder& GraphBuilder, bool bAsync
 	
 		GPUScene.EndDeferAllocatorMerges();
 	}
+	bool bNeedPathTracedInvalidation = false;
 	{
 		CSV_SCOPED_TIMING_STAT_EXCLUSIVE(AddPrimitiveSceneInfos);
 		SCOPED_NAMED_EVENT(FScene_AddPrimitiveSceneInfos, FColor::Green);
@@ -4917,11 +4927,8 @@ void FScene::UpdateAllPrimitiveSceneInfos(FRDGBuilder& GraphBuilder, bool bAsync
 
 				// Update scene LOD tree
 				SceneLODHierarchy.UpdateNodeSceneInfo(PrimitiveSceneInfo->PrimitiveComponentId, PrimitiveSceneInfo);
-			}
-			if (StartIndex < AddedLocalPrimitiveSceneInfos.Num())
-			{
-				// Invalidate PathTraced image because we added something to the scene
-				InvalidatePathTracedOutput();
+
+				bNeedPathTracedInvalidation = bNeedPathTracedInvalidation || IsPrimitiveRelevantToPathTracing(PrimitiveSceneInfo);
 			}
 			AddedLocalPrimitiveSceneInfos.RemoveAt(StartIndex, AddedLocalPrimitiveSceneInfos.Num() - StartIndex, false);
 		}
@@ -4974,6 +4981,9 @@ void FScene::UpdateAllPrimitiveSceneInfos(FRDGBuilder& GraphBuilder, bool bAsync
 				VelocityData.UpdateTransform(PrimitiveSceneInfo, LocalToWorld, PrimitiveSceneProxy->GetLocalToWorld());
 			}
 
+			bNeedPathTracedInvalidation = bNeedPathTracedInvalidation || (IsPrimitiveRelevantToPathTracing(PrimitiveSceneInfo) &&
+				!PrimitiveTransforms[PrimitiveSceneInfo->PackedIndex].Equals(LocalToWorld, SMALL_NUMBER));
+
 			// Update the primitive transform.
 			PrimitiveSceneProxy->SetTransform(LocalToWorld, WorldBounds, LocalBounds, AttachmentRootPosition);
 			PrimitiveTransforms[PrimitiveSceneInfo->PackedIndex] = LocalToWorld;
@@ -5105,6 +5115,8 @@ void FScene::UpdateAllPrimitiveSceneInfos(FRDGBuilder& GraphBuilder, bool bAsync
 				DistanceFieldSceneData.UpdatePrimitive(PrimitiveSceneInfo);
 				LumenSceneData->UpdatePrimitive(PrimitiveSceneInfo);
 			}
+
+			bNeedPathTracedInvalidation = bNeedPathTracedInvalidation || IsPrimitiveRelevantToPathTracing(PrimitiveSceneInfo);
 		}
 
 #if RHI_RAYTRACING
@@ -5193,6 +5205,8 @@ void FScene::UpdateAllPrimitiveSceneInfos(FRDGBuilder& GraphBuilder, bool bAsync
 		SCOPED_NAMED_EVENT(FScene_DeletePrimitiveSceneInfo, FColor::Red);
 		for (FPrimitiveSceneInfo* PrimitiveSceneInfo : DeletedSceneInfos)
 		{
+			bNeedPathTracedInvalidation = bNeedPathTracedInvalidation || IsPrimitiveRelevantToPathTracing(PrimitiveSceneInfo);
+
 			// It is possible that the HitProxies list isn't empty if PrimitiveSceneInfo was Added/Removed in same frame
 			// Delete the PrimitiveSceneInfo on the game thread after the rendering thread has processed its removal.
 			// This must be done on the game thread because the hit proxy references (and possibly other members) need to be freed on the game thread.
@@ -5207,11 +5221,10 @@ void FScene::UpdateAllPrimitiveSceneInfos(FRDGBuilder& GraphBuilder, bool bAsync
 			delete PrimitiveSceneInfo->Proxy;
 			delete PrimitiveSceneInfo;
 		}
-		if (DeletedSceneInfos.Num() > 0)
-		{
-			// Invalidate PathTraced image because we removed something from the scene
-			InvalidatePathTracedOutput();
-		}
+	}
+	if (bNeedPathTracedInvalidation)
+	{
+		InvalidatePathTracedOutput();
 	}
 	UpdatedAttachmentRoots.Empty();
 	UpdatedTransforms.Empty();
