Adding proper support for OverrideOutput in the PPM stencil skip path.

Note that this was incidentaclly working because if (PassSequence.IsEnabled(EPass::VisualizeLumenScene)) is always true, making that the last pass and not requiring any other pass to copy to the override output.

Continuation on 19244345

#rnx

diff --git a/Engine/Source/Runtime/Renderer/Private/PostProcess/PostProcessMaterial.cpp b/Engine/Source/Runtime/Renderer/Private/PostProcess/PostProcessMaterial.cpp
--- a/Engine/Source/Runtime/Renderer/Private/PostProcess/PostProcessMaterial.cpp
+++ b/Engine/Source/Runtime/Renderer/Private/PostProcess/PostProcessMaterial.cpp
@@ -699,7 +699,15 @@ FScreenPassTexture AddPostProcessMaterialPass(
 	}
 	else
 	{
+		// When skipping the pass, we still need to output a valid FScreenPassRenderTarget
 		Output = FScreenPassRenderTarget(SceneColor, ERenderTargetLoadAction::ENoAction);
+
+		// If there is override output, we need to output to that
+		if (Inputs.OverrideOutput.IsValid())
+		{
+			AddDrawTexturePass(GraphBuilder, View, Output.Texture, Inputs.OverrideOutput.Texture);
+			Output = Inputs.OverrideOutput;
+		}
 	}
 
 	return MoveTemp(Output);
