Only compile `FTrivialMaterialCHS` when the proper cvars have been turned on.

- This saves on having to compile this shader variant for every material.
- In EngineTest this saves ~2,000 shaders of ~3% of all shaders.
- In EngineTest this saves 15 shaders from being compiled in WorldGridMaterial

diff --git a/Engine/Source/Runtime/Renderer/Private/RayTracing/RayTracingMaterialHitShaders.cpp b/Engine/Source/Runtime/Renderer/Private/RayTracing/RayTracingMaterialHitShaders.cpp
--- a/Engine/Source/Runtime/Renderer/Private/RayTracing/RayTracingMaterialHitShaders.cpp
+++ b/Engine/Source/Runtime/Renderer/Private/RayTracing/RayTracingMaterialHitShaders.cpp
@@ -188,6 +188,11 @@ public:
 
 	static bool ShouldCompilePermutation(const FMeshMaterialShaderPermutationParameters& Parameters)
 	{
+		if (GCompileRayTracingMaterialAHS || GCompileRayTracingMaterialCHS)
+		{
+			return false;
+		}
+
 		return IsSupportedVertexFactoryType(Parameters.VertexFactoryType)
 			&& ShouldCompileRayTracingShadersForProject(Parameters.Platform);
 	}
