Restore incorrectly removed parameters

diff --git a/Engine/Source/Runtime/Renderer/Private/RayTracing/RayTracingScene.cpp b/Engine/Source/Runtime/Renderer/Private/RayTracing/RayTracingScene.cpp
--- a/Engine/Source/Runtime/Renderer/Private/RayTracing/RayTracingScene.cpp
+++ b/Engine/Source/Runtime/Renderer/Private/RayTracing/RayTracingScene.cpp
@@ -212,7 +212,7 @@ void FRayTracingScene::Create(FRDGBuilder& GraphBuilder, const FGPUScene& GPUSce
 				BuildRayTracingInstanceBuffer(
 					RHICmdList,
 					GPUScene,
-					FVector3f(ViewTilePosition),
+					ViewTilePosition,
 					FVector3f(RelativePreViewTranslation),
 					PassParams->InstanceBuffer->GetRHI(),
 					InstanceUploadSRV,
