Fix flickering in the editor with ISR enabled (UE-139266).

diff --git a/Engine/Source/Runtime/Renderer/Private/PostProcess/PostProcessCompositeEditorPrimitives.cpp b/Engine/Source/Runtime/Renderer/Private/PostProcess/PostProcessCompositeEditorPrimitives.cpp
--- a/Engine/Source/Runtime/Renderer/Private/PostProcess/PostProcessCompositeEditorPrimitives.cpp
+++ b/Engine/Source/Runtime/Renderer/Private/PostProcess/PostProcessCompositeEditorPrimitives.cpp
@@ -257,17 +257,13 @@ const FViewInfo* CreateEditorPrimitiveView(const FViewInfo& ParentView, FIntRect
 		EditorView->ViewMatrices.HackRemoveTemporalAAProjectionJitter();
 	}
 
-	FBox VolumeBounds[TVC_MAX];
-	TUniquePtr<FViewUniformShaderParameters> ViewParameters = MakeUnique<FViewUniformShaderParameters>();
-	EditorView->SetupUniformBufferParameters(VolumeBounds, TVC_MAX, *ViewParameters);
-	ViewParameters->NumSceneColorMSAASamples = NumSamples;
-	EditorView->ViewUniformBuffer = TUniformBufferRef<FViewUniformShaderParameters>::CreateUniformBufferImmediate(*ViewParameters, UniformBuffer_SingleFrame);
-	EditorView->CachedViewUniformShaderParameters = MoveTemp(ViewParameters);
+	EditorView->InitRHIResources(NumSamples);
+
 	return EditorView;
 }
 
 BEGIN_SHADER_PARAMETER_STRUCT(FEditorPrimitivesPassParameters, )
-	SHADER_PARAMETER_STRUCT_REF(FViewUniformShaderParameters, View)
+	SHADER_PARAMETER_STRUCT_INCLUDE(FViewShaderParameters, View)
 	SHADER_PARAMETER_STRUCT_REF(FReflectionCaptureShaderData, ReflectionCapture)
 	SHADER_PARAMETER_RDG_UNIFORM_BUFFER(FOpaqueBasePassUniformParameters, BasePass)
 	SHADER_PARAMETER_RDG_UNIFORM_BUFFER(FMobileBasePassUniformParameters, MobileBasePass)
@@ -465,7 +461,7 @@ FScreenPassTexture AddEditorPrimitivePass(
 		// Draws the editors primitives
 		{
 			FEditorPrimitivesPassParameters* PassParameters = GraphBuilder.AllocParameters<FEditorPrimitivesPassParameters>();
-			PassParameters->View = EditorView->ViewUniformBuffer;
+			PassParameters->View = EditorView->GetShaderParameters();
 			PassParameters->ReflectionCapture = View.ReflectionCaptureUniformBuffer;
 			PassParameters->InstanceCulling = InstanceCullingManager.GetDummyInstanceCullingUniformBuffer();
 			PassParameters->RenderTargets[0] = FRenderTargetBinding(EditorPrimitiveColor, ERenderTargetLoadAction::ELoad);
diff --git a/Engine/Source/Runtime/Renderer/Private/SceneRendering.cpp b/Engine/Source/Runtime/Renderer/Private/SceneRendering.cpp
--- a/Engine/Source/Runtime/Renderer/Private/SceneRendering.cpp
+++ b/Engine/Source/Runtime/Renderer/Private/SceneRendering.cpp
@@ -1980,7 +1980,7 @@ void FViewInfo::SetupUniformBufferParameters(
 #endif
 }
 
-void FViewInfo::InitRHIResources()
+void FViewInfo::InitRHIResources(uint32 OverrideNumMSAASamples)
 {
 	FBox VolumeBounds[TVC_MAX];
 
@@ -1993,6 +1993,11 @@ void FViewInfo::InitRHIResources()
 		TVC_MAX,
 		*CachedViewUniformShaderParameters);
 
+	if (OverrideNumMSAASamples > 0)
+	{
+		CachedViewUniformShaderParameters->NumSceneColorMSAASamples = OverrideNumMSAASamples;
+	}
+
 	ViewUniformBuffer = TUniformBufferRef<FViewUniformShaderParameters>::CreateUniformBufferImmediate(*CachedViewUniformShaderParameters, UniformBuffer_SingleFrame);
 	if (bShouldBindInstancedViewUB)
 	{
diff --git a/Engine/Source/Runtime/Renderer/Private/SceneRendering.h b/Engine/Source/Runtime/Renderer/Private/SceneRendering.h
--- a/Engine/Source/Runtime/Renderer/Private/SceneRendering.h
+++ b/Engine/Source/Runtime/Renderer/Private/SceneRendering.h
@@ -1560,8 +1560,8 @@ public:
 	void SetupGlobalDistanceFieldUniformBufferParameters(FViewUniformShaderParameters& ViewUniformShaderParameters) const;
 	void SetupVolumetricFogUniformBufferParameters(FViewUniformShaderParameters& ViewUniformShaderParameters) const;
 
-	/** Initializes the RHI resources used by this view. */
-	void InitRHIResources();
+	/** Initializes the RHI resources used by this view. OverrideNumMSAASamples can optionally override NumSceneColorMSAASamples (needed for editor views) */
+	void InitRHIResources(uint32 OverrideNumMSAASamples = 0);
 
 	/** Determines distance culling and fades if the state changes */
 	bool IsDistanceCulled(float DistanceSquared, float MinDrawDistance, float InMaxDrawDistance, const FPrimitiveSceneInfo* PrimitiveSceneInfo);
