Capture AccelerationStructureAddressesBuffer in RHI EnqueueLambda:
In very rare cases we might be reallocating the buffer that is currently processed by RHI thread leading to a GPU crash.

diff --git a/Engine/Source/Runtime/Renderer/Private/RayTracing/RayTracingScene.cpp b/Engine/Source/Runtime/Renderer/Private/RayTracing/RayTracingScene.cpp
--- a/Engine/Source/Runtime/Renderer/Private/RayTracing/RayTracingScene.cpp
+++ b/Engine/Source/Runtime/Renderer/Private/RayTracing/RayTracingScene.cpp
@@ -187,14 +187,14 @@ void FRayTracingScene::Create(FRDGBuilder& GraphBuilder, const FGPUScene& GPUSce
 				// Pull this out here, because command list playback (where the lambda is executed) doesn't update the GPU mask
 				FRHIGPUMask IterateGPUMasks = RHICmdList.GetGPUMask();
 
-				RHICmdList.EnqueueLambda([this, &SceneInitializer, IterateGPUMasks](FRHICommandListImmediate& RHICmdList)
+				RHICmdList.EnqueueLambda([BufferRHIRef = AccelerationStructureAddressesBuffer.Buffer, &SceneInitializer, IterateGPUMasks](FRHICommandListImmediate& RHICmdList)
 					{
 						QUICK_SCOPE_CYCLE_COUNTER(GetAccelerationStructuresAddresses);
 
 						for (uint32 GPUIndex : IterateGPUMasks)
 						{
 							FRayTracingAccelerationStructureAddress* AddressesPtr = (FRayTracingAccelerationStructureAddress*)RHICmdList.LockBufferMGPU(
-								AccelerationStructureAddressesBuffer.Buffer,
+								BufferRHIRef,
 								GPUIndex,
 								0,
 								SceneInitializer.ReferencedGeometries.Num() * sizeof(FRayTracingAccelerationStructureAddress), RLM_WriteOnly);
@@ -205,7 +205,7 @@ void FRayTracingScene::Create(FRDGBuilder& GraphBuilder, const FGPUScene& GPUSce
 								AddressesPtr[GeometryIndex] = SceneInitializer.ReferencedGeometries[GeometryIndex]->GetAccelerationStructureAddress(GPUIndex);
 							}
 
-							RHICmdList.UnlockBufferMGPU(AccelerationStructureAddressesBuffer.Buffer, GPUIndex);
+							RHICmdList.UnlockBufferMGPU(BufferRHIRef, GPUIndex);
 						}
 					});
 
