Raytraced Shadows - Support Epsilon Trace to avoid self intersections.
- Epsilon trace distance controllable using r.RayTracing.Shadows.AvoidSelfIntersectionTraceDistance.
- Testing shows it can add significant overhead (~30%).
- Disabled by default (distance = 0.0f).

diff --git a/Engine/Shaders/Private/RayTracing/RayTracingOcclusionRGS.usf b/Engine/Shaders/Private/RayTracing/RayTracingOcclusionRGS.usf
--- a/Engine/Shaders/Private/RayTracing/RayTracingOcclusionRGS.usf
+++ b/Engine/Shaders/Private/RayTracing/RayTracingOcclusionRGS.usf
@@ -42,6 +42,7 @@ int2 PixelOffset;
 float TraceDistance;
 float LODTransitionStart;
 float LODTransitionEnd;
+float AvoidSelfIntersectionTraceDistance;
 uint bTransmissionSamplingDistanceCulling;
 uint TransmissionSamplingTechnique;
 uint RejectionSamplingTrials;
@@ -444,12 +445,43 @@ FOcclusionResult ComputeOcclusion(
 
 		uint RayFlagsForOpaque = bAcceptFirstHit != 0 ? RAY_FLAG_ACCEPT_FIRST_HIT_AND_END_SEARCH : 0;
 
-		FMinimalPayload MinimalPayload = TraceVisibilityRay(
+		FMinimalPayload MinimalPayload;
+
+#if AVOID_SELF_INTERSECTION_TRACE
+		{
+			RayDesc EpsilonRay = Ray;
+			EpsilonRay.TMax = AvoidSelfIntersectionTraceDistance;
+
+			if (Ray.TMax > Ray.TMin)
+			{
+				MinimalPayload = TraceVisibilityRay(
+					TLAS,
+					RayFlags | RayFlagsForOpaque | RAY_FLAG_CULL_BACK_FACING_TRIANGLES,
+					RaytracingMask,
+					PixelCoord,
+					EpsilonRay);
+			}
+		}
+
+		if (!MinimalPayload.IsHit())
+		{
+			Ray.TMin = max(Ray.TMin, AvoidSelfIntersectionTraceDistance);
+
+			MinimalPayload = TraceVisibilityRay(
+				TLAS,
+				RayFlags | RayFlagsForOpaque,
+				RaytracingMask,
+				PixelCoord,
+				Ray);
+		}
+#else
+		MinimalPayload = TraceVisibilityRay(
 			TLAS,
 			RayFlags | RayFlagsForOpaque,
 			RaytracingMask,
 			PixelCoord,
 			Ray);
+#endif
 
 		#if USE_HAIR_LIGHTING
 		if (!ShadingDesc.bIsHair && bUseHairVoxel)
diff --git a/Engine/Source/Runtime/Renderer/Private/RayTracing/RayTracingShadows.cpp b/Engine/Source/Runtime/Renderer/Private/RayTracing/RayTracingShadows.cpp
--- a/Engine/Source/Runtime/Renderer/Private/RayTracing/RayTracingShadows.cpp
+++ b/Engine/Source/Runtime/Renderer/Private/RayTracing/RayTracingShadows.cpp
@@ -34,6 +34,13 @@ static FAutoConsoleVariableRef CVarRayTracingShadowsEnableMaterials(
 	TEXT("Enables material shader binding for shadow rays. If this is disabled, then a default trivial shader is used. (default = 1)")
 );
 
+static float GRayTracingShadowsAvoidSelfIntersectionTraceDistance = 0.0f;
+static FAutoConsoleVariableRef CVarRayTracingShadowsAvoidSelfIntersectionTraceDistance(
+	TEXT("r.RayTracing.Shadows.AvoidSelfIntersectionTraceDistance"),
+	GRayTracingShadowsAvoidSelfIntersectionTraceDistance,
+	TEXT("Max trace distance of epsilon trace to avoid self intersections. If set to 0, epsilon trace will not be used.")
+);
+
 static TAutoConsoleVariable<int32> CVarRayTracingShadowsEnableTwoSidedGeometry(
 	TEXT("r.RayTracing.Shadows.EnableTwoSidedGeometry"),
 	1,
@@ -103,12 +110,13 @@ class FOcclusionRGS : public FGlobalShader
 	SHADER_USE_ROOT_PARAMETER_STRUCT(FOcclusionRGS, FGlobalShader)
 
 	class FLightTypeDim : SHADER_PERMUTATION_INT("LIGHT_TYPE", LightType_MAX);
+	class FAvoidSelfIntersectionTraceDim : SHADER_PERMUTATION_BOOL("AVOID_SELF_INTERSECTION_TRACE");
 	class FDenoiserOutputDim : SHADER_PERMUTATION_INT("DIM_DENOISER_OUTPUT", 3);
 	class FEnableMultipleSamplesPerPixel : SHADER_PERMUTATION_BOOL("ENABLE_MULTIPLE_SAMPLES_PER_PIXEL");
 	class FHairLighting : SHADER_PERMUTATION_INT("USE_HAIR_LIGHTING", 2);
 	class FEnableTransmissionDim : SHADER_PERMUTATION_INT("ENABLE_TRANSMISSION", 2);
 
-	using FPermutationDomain = TShaderPermutationDomain<FLightTypeDim, FDenoiserOutputDim, FHairLighting, FEnableMultipleSamplesPerPixel, FEnableTransmissionDim>;
+	using FPermutationDomain = TShaderPermutationDomain<FLightTypeDim, FAvoidSelfIntersectionTraceDim, FDenoiserOutputDim, FHairLighting, FEnableMultipleSamplesPerPixel, FEnableTransmissionDim>;
 
 	static bool ShouldCompilePermutation(const FGlobalShaderPermutationParameters& Parameters)
 	{
@@ -146,6 +154,7 @@ class FOcclusionRGS : public FGlobalShader
 		SHADER_PARAMETER(float, TraceDistance)
 		SHADER_PARAMETER(float, LODTransitionStart)
 		SHADER_PARAMETER(float, LODTransitionEnd)
+		SHADER_PARAMETER(float, AvoidSelfIntersectionTraceDistance)
 		SHADER_PARAMETER(uint32, bTransmissionSamplingDistanceCulling)
 		SHADER_PARAMETER(uint32, TransmissionSamplingTechnique)
 		SHADER_PARAMETER(uint32, RejectionSamplingTrials)
@@ -201,19 +210,23 @@ void FDeferredShadingSceneRenderer::PrepareRayTracingShadows(const FViewInfo& Vi
 		{
 			for (int32 HairLighting = 0; HairLighting < 2; ++HairLighting)
 			{
-				for (int32 LightType = 0; LightType < LightType_MAX; ++LightType)
+				for (int32 AvoidSelfIntersectionTrace = 0; AvoidSelfIntersectionTrace < 2; ++AvoidSelfIntersectionTrace)
 				{
-					for (IScreenSpaceDenoiser::EShadowRequirements DenoiserRequirement : DenoiserRequirements)
+					for (int32 LightType = 0; LightType < LightType_MAX; ++LightType)
 					{
-						FOcclusionRGS::FPermutationDomain PermutationVector;
-						PermutationVector.Set<FOcclusionRGS::FLightTypeDim>(LightType);
-						PermutationVector.Set<FOcclusionRGS::FDenoiserOutputDim>((int32)DenoiserRequirement);
-						PermutationVector.Set<FOcclusionRGS::FHairLighting>(HairLighting);
-						PermutationVector.Set<FOcclusionRGS::FEnableMultipleSamplesPerPixel>(MultiSPP != 0);
-						PermutationVector.Set<FOcclusionRGS::FEnableTransmissionDim>(EnableTransmissionDim);
-
-						TShaderMapRef<FOcclusionRGS> RayGenerationShader(View.ShaderMap, PermutationVector);
-						OutRayGenShaders.Add(RayGenerationShader.GetRayTracingShader());
+						for (IScreenSpaceDenoiser::EShadowRequirements DenoiserRequirement : DenoiserRequirements)
+						{
+							FOcclusionRGS::FPermutationDomain PermutationVector;
+							PermutationVector.Set<FOcclusionRGS::FLightTypeDim>(LightType);
+							PermutationVector.Set<FOcclusionRGS::FAvoidSelfIntersectionTraceDim>((bool)AvoidSelfIntersectionTrace);
+							PermutationVector.Set<FOcclusionRGS::FDenoiserOutputDim>((int32)DenoiserRequirement);
+							PermutationVector.Set<FOcclusionRGS::FHairLighting>(HairLighting);
+							PermutationVector.Set<FOcclusionRGS::FEnableMultipleSamplesPerPixel>(MultiSPP != 0);
+							PermutationVector.Set<FOcclusionRGS::FEnableTransmissionDim>(EnableTransmissionDim);
+
+							TShaderMapRef<FOcclusionRGS> RayGenerationShader(View.ShaderMap, PermutationVector);
+							OutRayGenShaders.Add(RayGenerationShader.GetRayTracingShader());
+						}
 					}
 				}
 			}
@@ -284,6 +297,7 @@ void FDeferredShadingSceneRenderer::RenderRayTracingShadows(
 		PassParameters->TraceDistance = LightSceneProxy->GetTraceDistance();
 		PassParameters->LODTransitionStart = CVarRayTracingShadowsLODTransitionStart.GetValueOnRenderThread();
 		PassParameters->LODTransitionEnd = CVarRayTracingShadowsLODTransitionEnd.GetValueOnRenderThread();
+		PassParameters->AvoidSelfIntersectionTraceDistance = GRayTracingShadowsAvoidSelfIntersectionTraceDistance;
 		PassParameters->bAcceptFirstHit = CVarRayTracingShadowsAcceptFirstHit.GetValueOnRenderThread();
 		PassParameters->bTwoSidedGeometry = EnableRayTracingShadowTwoSidedGeometry() ? 1 : 0;
 		PassParameters->TLAS = View.GetRayTracingSceneViewChecked();
@@ -312,6 +326,7 @@ void FDeferredShadingSceneRenderer::RenderRayTracingShadows(
 		}
 		FOcclusionRGS::FPermutationDomain PermutationVector;
 		PermutationVector.Set<FOcclusionRGS::FLightTypeDim>(LightSceneProxy->GetLightType());
+		PermutationVector.Set<FOcclusionRGS::FAvoidSelfIntersectionTraceDim>(GRayTracingShadowsAvoidSelfIntersectionTraceDistance > 0.0f);
 		if (DenoiserRequirements == IScreenSpaceDenoiser::EShadowRequirements::PenumbraAndAvgOccluder)
 		{
 			PermutationVector.Set<FOcclusionRGS::FDenoiserOutputDim>(1);
