Strata - rough refraction rnd - accumulate importance sample refractions in texture, debug view improvements.

diff --git a/Engine/Shaders/Private/Strata/StrataRoughRefraction.usf b/Engine/Shaders/Private/Strata/StrataRoughRefraction.usf
--- a/Engine/Shaders/Private/Strata/StrataRoughRefraction.usf
+++ b/Engine/Shaders/Private/Strata/StrataRoughRefraction.usf
@@ -8,78 +8,23 @@
 
 #include "/Engine/Private/MiniFontCommon.ush"
 #include "/Engine/Private/ShaderPrint.ush"
+#include "/Engine/Private/ColorMap.ush"
 #include "/Engine/Private/BRDF.ush"
 #include "/Engine/Private/MonteCarlo.ush"
 #include "/Engine/Private/PathTracing/Material/PathTracingFresnel.ush"
 
 // Updated from http://jcgt.org/published/0007/03/04/
-bool slabs(float3 p0, float3 p1, float3 rayOrigin, float3 invRaydir, out float outTMin, out float outTMax)
+bool Slabs(float3 P0, float3 P1, float3 RayOrigin, float3 InvRaydir, out float OutTMin, out float OutTMax)
 {
-	float3 t0 = (p0 - rayOrigin) * invRaydir;
-	float3 t1 = (p1 - rayOrigin) * invRaydir;
-	float3 tmin = min(t0, t1), tmax = max(t0, t1);
-	float maxtmin = max(max(tmin.x, tmin.y), tmin.z);
-	float mintmax = min(min(tmax.x, tmax.y), tmax.z);
-	outTMin = maxtmin;
-	outTMax = mintmax;
-	return maxtmin <= mintmax;
-}
-
-// From https://en.wikipedia.org/wiki/MöllerTrumbore_intersection_algorithm
-bool RayIntersectsTriangle(float3 rayOrigin, float3 rayVector, float3 v0, float3 v1, float3 v2, out float3 outIntersectionPoint, out float2 uv)
-{
-	outIntersectionPoint = 0.0f;
-	uv = 0.0f;
-
-	const float EPSILON = 0.0000001f;
-	float3 edge1, edge2, tvec, pvec, qvec;
-	float det, invDet, u, v;
-
-	edge1 = v1 - v0;
-	edge2 = v2 - v0;
-
-	pvec = cross(rayVector, edge2);
-	det = dot(edge1, pvec);
-	if (det > -EPSILON && det < EPSILON)
-		return false;
-	invDet = 1.0f / det;
-
-	tvec = rayOrigin - v0;
-	u = dot(tvec, pvec) * invDet;
-	if (u < 0.0f || u > 1.0f)
-		return false;
-
-	qvec = cross(tvec, edge1);
-
-	v = dot(rayVector, qvec) * invDet;
-	if (v < 0.0f || (u + v) > 1.0f)
-		return false;
-
-	float t = dot(edge2, qvec) * invDet;
-
-	if (t > EPSILON)
-	{
-		uv = float2(u, v);
-		outIntersectionPoint = rayOrigin + rayVector * t;
-		return true;
-	}
-	else
-		return false; // Intersection not in the ray direction
-}
-
-float3 GetAnyPerpendicularVector(float3 Dir)
-{
-	float3 Perp = 0;
-	float3 Source = float3(1, 0, 0);
-	if (dot(Dir, Source) < 0.99)
-	{
-		Perp = normalize(cross(Dir, Source));
-	}
-	else
-	{
-		Perp = normalize(cross(Dir.yzx, Source));
-	}
-	return Perp;
+	float3 T0 = (P0 - RayOrigin) * InvRaydir;
+	float3 T1 = (P1 - RayOrigin) * InvRaydir;
+	float3 TMin = min(T0, T1);
+	float3 TMax = max(T0, T1);
+	float MaxTMin = max(max(TMin.x, TMin.y), TMin.z);
+	float MinTMax = min(min(TMax.x, TMax.y), TMax.z);
+	OutTMin = MaxTMin;
+	OutTMax = MinTMax;
+	return MaxTMin <= MinTMax;
 }
 
 float3 GetScreenWorldDir(in float4 SVPos)
@@ -104,9 +49,99 @@ bool GetRefractedDirection(float2 RandomUV, float a2, float3 V, inout float3 L )
 	return SampleRefraction(-V, H.xyz, EtaIn, RandSample, L, F);
 }
 
+#if defined(EVALUATE_ROUGH_REFRACTION_LOBE_CS) || defined(VISUALIZE_ROUGH_REFRACTION_PS)
+struct FUI
+{
+	float IncidentAngle;
+	float SlabThickness;
+	float Roughness;
+
+	// Derived
+	float a2;
+	float3 N;
+	float3 V;
+};
+
+
+FShaderPrintContext GetUIContext(bool bActive)
+{
+	return InitShaderPrintContext(bActive, uint2(100, 100));
+}
+
+FUI GetUI(inout FShaderPrintContext Context)
+{
+	FUI UI;
+
+	UI.IncidentAngle = AddSlider(Context, TEXT("IncidentAngle"), 0.0f, GetDefaultFontColor(), 0.0f, 0.5f * PI);
+	Newline(Context);
+	UI.SlabThickness = AddSlider(Context, TEXT("Slab Thickness"), 5.0f, GetDefaultFontColor(), 0.0f, 20.0f);
+	Newline(Context);
+	UI.Roughness = AddSlider(Context, TEXT("Roughness"), 0.5f, GetDefaultFontColor(), 0.0f, 1.0f);
+	Newline(Context);
+
+	UI.a2 = Pow4(UI.Roughness);
+	UI.N = float3(0, 0, 1);
+	UI.V = float3(sin(UI.IncidentAngle), 0, cos(UI.IncidentAngle));
+
+	return UI;
+}
+#endif
+
+
+#ifdef EVALUATE_ROUGH_REFRACTION_LOBE_CS
+
+uint TraceSqrtSampleCount;
+
+RWTexture2D<uint> SampleCountTextureUAV;
+
+RWStructuredBuffer<FLobeMeanVariance> LobStatisticsBufferUAV;
+
+[numthreads(THREADGROUP_SIZE, THREADGROUP_SIZE, 1)]
+void EvaluateRoughRefractionLobeCS(uint3 ThreadId : SV_DispatchThreadID)
+{
+	if (all(ThreadId == 0))
+	{
+		FShaderPrintContext Context = GetUIContext(true);
+		FUI UI = GetUI(Context);
+
+
+		FLobeMeanVariance Lobe = (FLobeMeanVariance)0;
+		LobStatisticsBufferUAV[0] = Lobe;
+
+
+		for (float u = 0.5f; u < float(TraceSqrtSampleCount); u++)
+		{
+			for (float v = 0.5f; v < float(TraceSqrtSampleCount); v++)
+			{
+				float2 RandomUV = float2(u, v) / float(TraceSqrtSampleCount);
+				float3 L = 0;
+				if (GetRefractedDirection(RandomUV, UI.a2, UI.V, L))
+				{
+					float2 SampleUV = L.xy * 0.5 + 0.5;
+					SampleCountTextureUAV[uint2(SampleUV * 64.0f)] += 1;
+				}
+				else
+				{
+					// Discard refelctions
+				}
+			}
+		}
+	}
+}
+
+#endif // EVALUATE_ROUGH_REFRACTION_LOBE_CS
+
+
+
+#ifdef VISUALIZE_ROUGH_REFRACTION_PS
+
 float TraceDomainSize;
 uint SlabInterfaceLineCount;
 
+Texture2D<uint> SampleCountTexture;
+
+StructuredBuffer<FLobeMeanVariance> LobStatisticsBuffer;
+
 void VisualizeRoughRefractionPS(
 	float4 SVPos : SV_POSITION,
 	out float4 OutColor : SV_Target0)
@@ -118,35 +153,44 @@ void VisualizeRoughRefractionPS(
 	float3 WorldDir = GetScreenWorldDir(SVPos);
 	float3 CamWorldPos = LWCToFloat(PrimaryView.WorldCameraOrigin);
 
-	FShaderPrintContext Context = InitShaderPrintContext(all(uint2(PixelPos.xy) == uint2(100, 100)), uint2(100, 100));
-	const float IncidentAngle = AddSlider(Context, TEXT("IncidentAngle"), 0.0f, GetDefaultFontColor(), 0.0f, 0.5f * PI);
-	Newline(Context);
-	const float SlabThickness = AddSlider(Context, TEXT("Slab Thickness"), 5.0f, GetDefaultFontColor(), 0.0f, 20.0f);
-	Newline(Context);
-	const float Roughness = AddSlider(Context, TEXT("Roughness"), 0.5f, GetDefaultFontColor(), 0.0f, 1.0f);
-	Newline(Context);
-
-	float a2 = Pow4(Roughness);
-	const float3 N = float3(0, 0, 1);
-	const float3 V = float3(sin(IncidentAngle), 0, cos(IncidentAngle));
+	
+	FShaderPrintContext Context = GetUIContext(all(uint2(PixelPos.xy) == uint2(10, 10)));
+	FUI UI = GetUI(Context);
 
 	float3 SlabEdgeX = float3(60.0f, 0.0f, 0.0f);
 	float3 SlabEdgeY = float3(0.0f, 60.0f, 0.0f);
-	float3 SlabEdgeZ = float3(0.0f, 0.0f, SlabThickness);
-	float3 SlabCenter = CamWorldPos + View.ViewForward * 100.0f;
+	float3 SlabEdgeZ = float3(0.0f, 0.0f, UI.SlabThickness);
+	float3 SlabCenter = CamWorldPos + View.ViewForward * 60.0f;
 	float3 RefractionP= SlabCenter + 0.5 * SlabEdgeZ;
 
+	const float VizProjectedRays = 15.0f;
 
-	AddLineWS(Context, RefractionP, RefractionP + N * 25.0, float4(0, 0, 1, 1));
-	AddLineWS(Context, RefractionP, RefractionP + V * 25.0, float4(0, 1, 1, 1));
+	AddLineWS(Context, RefractionP, RefractionP + UI.N * 25.0, float4(0, 0, 1, 1));
+	AddLineWS(Context, RefractionP, RefractionP + UI.V * 25.0, float4(0, 1, 1, 1));
 
-#if 0
-	float TMin = -1.0f;
-	float TMax = -1.0f;
-	float3 p0 = SlabCenter - 0.5 * (SlabEdgeX + SlabEdgeY + SlabEdgeZ);
-	float3 p1 = SlabCenter + 0.5 * (SlabEdgeX + SlabEdgeY + SlabEdgeZ);
-	slabs(p0, p1, WorldOrigin, 1.0f / WorldDir, TMin, TMax);
-#endif
+	float SlabTMin = -1.0f;
+	float SlabTMax = -1.0f;
+	float3 SlabP0 = SlabCenter - 0.5 * (SlabEdgeX + SlabEdgeY + SlabEdgeZ);
+	float3 SlabP1 = SlabCenter + 0.5 * (SlabEdgeX + SlabEdgeY + SlabEdgeZ);
+	if(Slabs(SlabP0, SlabP1, CamWorldPos, 1.0f / WorldDir, SlabTMin, SlabTMax))
+	{
+		float3 P = CamWorldPos + WorldDir * SlabTMin;
+		if (P.z > -0.499 * UI.SlabThickness)
+		{
+			P = CamWorldPos + WorldDir * SlabTMax;
+		}
+
+		if (P.z < (SlabCenter -0.499 * SlabEdgeZ).z)
+		{
+			const float2 UV = 0.5 + 0.5 * ((P - RefractionP) / (VizProjectedRays)).xy;
+			const int2 Coord = (UV * 64.0f + 0.5f);
+			if (all(Coord >= 0) && all(Coord < 64))
+			{
+				const float SampleCount = float(SampleCountTexture[Coord].x);
+				OutColor = float4(GetHSVDebugColor(SampleCount*0.01f), 0.0);
+			}
+		}
+	}
 
 #if 0
 	const float SphereRadius = 10.0f;
@@ -158,8 +202,8 @@ void VisualizeRoughRefractionPS(
 		const float OmegaOutLen = length(OmegaOut);
 		const float3 OmegaOutNorm = OmegaOut / OmegaOutLen;
 
-		const float NoH = dot(normalize(OmegaOutNorm + V), N);
-		float D = D_GGX(a2, NoH);
+		const float NoH = dot(normalize(OmegaOutNorm + UI.V), UI.N);
+		float D = D_GGX(UI.a2, NoH);
 		OutColor = float4(D, D, 0.0, 0.0);
 	}
 #endif
@@ -183,23 +227,24 @@ void VisualizeRoughRefractionPS(
 		{
 			float2 RandomUV = float2(u, v) / TraceDomainSize;
 			float3 L = 0;
-			if (GetRefractedDirection(RandomUV, a2, V, L))
+			if (GetRefractedDirection(RandomUV, UI.a2, UI.V, L))
 			{
-				AddLineWS(Context, RefractionP, RefractionP + L * 15.0, float4(0, 1, 0, 0.3));
+				AddLineWS(Context, RefractionP, RefractionP + L * VizProjectedRays, float4(0, 1, 0, 0.3));
+				//AddLineWS(Context, RefractionP + L * VizProjectedRays, float3((RefractionP + L * VizProjectedRays).xy, SlabCenter.z - 0.5* UI.SlabThickness), float4(0, 1, 0, 0.3));
 			}
 			else
 			{
-				AddLineWS(Context, RefractionP, RefractionP + L * 15.0, float4(1, 0, 0, 0.3));
+				AddLineWS(Context, RefractionP, RefractionP + L * VizProjectedRays, float4(1, 0, 0, 0.3));
 			}
 		}
 	}
 
 }
 
+#endif // VISUALIZE_ROUGH_REFRACTION_PS
 
 
 
-
-#endif
+#endif // STRATA_RND_SHADERS
 
 
diff --git a/Engine/Source/Runtime/Renderer/Private/Strata/StrataRoughRefraction.cpp b/Engine/Source/Runtime/Renderer/Private/Strata/StrataRoughRefraction.cpp
--- a/Engine/Source/Runtime/Renderer/Private/Strata/StrataRoughRefraction.cpp
+++ b/Engine/Source/Runtime/Renderer/Private/Strata/StrataRoughRefraction.cpp
@@ -31,6 +31,48 @@ static TAutoConsoleVariable<int32> CVarStrataRoughRefractionShadersShowRoughRefr
 	ECVF_RenderThreadSafe);
 
 
+
+class FEvaluateRoughRefractionLobeCS : public FGlobalShader
+{
+	DECLARE_GLOBAL_SHADER(FEvaluateRoughRefractionLobeCS);
+	SHADER_USE_PARAMETER_STRUCT(FEvaluateRoughRefractionLobeCS, FGlobalShader);
+
+	using FPermutationDomain = TShaderPermutationDomain<>;
+
+	BEGIN_SHADER_PARAMETER_STRUCT(FParameters, )
+		SHADER_PARAMETER_STRUCT_REF(FViewUniformShaderParameters, ViewUniformBuffer)
+		SHADER_PARAMETER_RDG_TEXTURE_UAV(RWTexture2D<uint>, SampleCountTextureUAV)
+		SHADER_PARAMETER_RDG_BUFFER_UAV(RWStructuredBuffer<>, LobStatisticsBufferUAV)
+		SHADER_PARAMETER_TEXTURE(Texture2D, MiniFontTexture)
+		SHADER_PARAMETER_STRUCT_INCLUDE(ShaderPrint::FShaderParameters, ShaderPrintParameters)
+		SHADER_PARAMETER(uint32, TraceSqrtSampleCount)
+	END_SHADER_PARAMETER_STRUCT()
+	
+	static const uint32 ThreadGroupSize = 8;
+
+	static FPermutationDomain RemapPermutation(FPermutationDomain PermutationVector)
+	{
+		return PermutationVector;
+	}
+
+	static bool ShouldCompilePermutation(const FGlobalShaderPermutationParameters& Parameters)
+	{
+		return true;
+	}
+
+	static void ModifyCompilationEnvironment(const FGlobalShaderPermutationParameters& Parameters, FShaderCompilerEnvironment& OutEnvironment)
+	{
+		FGlobalShader::ModifyCompilationEnvironment(Parameters, OutEnvironment);
+
+		OutEnvironment.SetDefine(TEXT("STRATA_RND_SHADERS"), 1);
+		OutEnvironment.SetDefine(TEXT("EVALUATE_ROUGH_REFRACTION_LOBE_CS"), 1);
+		OutEnvironment.SetDefine(TEXT("THREADGROUP_SIZE"), ThreadGroupSize);
+	}
+};
+IMPLEMENT_GLOBAL_SHADER(FEvaluateRoughRefractionLobeCS, "/Engine/Private/Strata/StrataRoughRefraction.usf", "EvaluateRoughRefractionLobeCS", SF_Compute);
+
+
+
 class FVisualizeRoughRefractionPS : public FGlobalShader
 {
 	DECLARE_GLOBAL_SHADER(FVisualizeRoughRefractionPS);
@@ -40,6 +82,8 @@ class FVisualizeRoughRefractionPS : public FGlobalShader
 
 	BEGIN_SHADER_PARAMETER_STRUCT(FParameters, )
 		SHADER_PARAMETER_STRUCT_REF(FViewUniformShaderParameters, ViewUniformBuffer)
+		SHADER_PARAMETER_RDG_TEXTURE_SRV(Texture2D<uint>, SampleCountTexture)
+		SHADER_PARAMETER_RDG_BUFFER_SRV(StructuredBuffer<>, LobStatisticsBuffer)
 		SHADER_PARAMETER_TEXTURE(Texture2D, MiniFontTexture)
 		SHADER_PARAMETER_STRUCT_INCLUDE(ShaderPrint::FShaderParameters, ShaderPrintParameters)
 		SHADER_PARAMETER(float, TraceDomainSize)
@@ -61,10 +105,8 @@ class FVisualizeRoughRefractionPS : public FGlobalShader
 	{
 		FGlobalShader::ModifyCompilationEnvironment(Parameters, OutEnvironment);
 
-		// Stay debug and skip optimizations to reduce compilation time on this long shader.
-		OutEnvironment.CompilerFlags.Add(CFLAG_PrecompileWithDXC);
-
 		OutEnvironment.SetDefine(TEXT("STRATA_RND_SHADERS"), 1);
+		OutEnvironment.SetDefine(TEXT("VISUALIZE_ROUGH_REFRACTION_PS"), 1);
 	}
 };
 IMPLEMENT_GLOBAL_SHADER(FVisualizeRoughRefractionPS, "/Engine/Private/Strata/StrataRoughRefraction.usf", "VisualizeRoughRefractionPS", SF_Pixel);
@@ -99,44 +141,72 @@ void StrataRoughRefractionRnD(FRDGBuilder& GraphBuilder, const FViewInfo& View,
 		// Create resources
 		
 		// Texture to count
-		//const uint32 SampleCountTextureWidth = 64;
-		//const FIntPoint SampleCountTextureSize(SampleCountTextureWidth, SampleCountTextureWidth);
-		//FRDGTextureRef SampleCountTexture = GraphBuilder.CreateTexture(FRDGTextureDesc::Create2D(SampleCountTextureSize, PF_R32_UINT, FClearValueBinding::Black, TexCreate_UAV | TexCreate_ShaderResource), TEXT("Strata.RoughRefrac.SampleCount"));
-		//FRDGTextureUAVRef SampleCountTextureUAV = GraphBuilder.CreateUAV(FRDGTextureUAVDesc(SampleCountTexture));
-		//FRDGTextureSRVRef SampleCountTextureSRV = GraphBuilder.CreateSRV(FRDGTextureSRVDesc::Create(SampleCountTexture));
-		//
-		//FRDGBufferRef LobStatisticsBuffer = GraphBuilder.CreateBuffer(FRDGBufferDesc::CreateBufferDesc(sizeof(uint32) * 4, 64), TEXT("Strata.RoughRefrac.LobStat"));
-		//FRDGTextureUAVRef SampleCountTextureUAV = GraphBuilder.CreateUAV(LobStatisticsBuffer, PF_A32B32G32R32F);
-		//FRDGTextureSRVRef SampleCountTextureSRV = GraphBuilder.CreateSRV(LobStatisticsBuffer, PF_A32B32G32R32F);
-
-
-
+		const uint32 SampleCountTextureWidth = 64;
+		const FIntPoint SampleCountTextureSize(SampleCountTextureWidth, SampleCountTextureWidth);
+		FRDGTextureRef SampleCountTexture = GraphBuilder.CreateTexture(FRDGTextureDesc::Create2D(SampleCountTextureSize, PF_R32_UINT, FClearValueBinding::Black, TexCreate_UAV | TexCreate_ShaderResource), TEXT("Strata.RoughRefrac.SampleCount"));
+		FRDGTextureUAVRef SampleCountTextureUAV = GraphBuilder.CreateUAV(FRDGTextureUAVDesc(SampleCountTexture));
+		FRDGTextureSRVRef SampleCountTextureSRV = GraphBuilder.CreateSRV(FRDGTextureSRVDesc::Create(SampleCountTexture));
 
+		FRDGBufferRef LobStatisticsBuffer = GraphBuilder.CreateBuffer(FRDGBufferDesc::CreateStructuredDesc(sizeof(float) * 4, 16), TEXT("Strata.RoughRefrac.LobStat"));
+		FRDGBufferUAVRef LobStatisticsBufferUAV = GraphBuilder.CreateUAV(LobStatisticsBuffer, PF_R32_UINT);
+		FRDGBufferSRVRef LobStatisticsBufferSRV = GraphBuilder.CreateSRV(LobStatisticsBuffer, PF_R32_UINT);
 
 		//////////////////////////////////////////////////////////////////////////
-		// Debug print everything on screen
-		const float TraceDomainSize = 32.0f;
-		const float SlabInterfaceLineCount = 16.0f;
+		// Clear resources
+		AddClearUAVPass(GraphBuilder, SampleCountTextureUAV, uint32(0));
 
-		ShaderPrint::RequestSpaceForLines(TraceDomainSize* TraceDomainSize + SlabInterfaceLineCount * SlabInterfaceLineCount * 2);
-		ShaderPrint::RequestSpaceForCharacters(256);
-
-		FSceneTextureParameters SceneTextureParameters = GetSceneTextureParameters(GraphBuilder);
-		FVisualizeRoughRefractionPS::FParameters* PassParameters = GraphBuilder.AllocParameters<FVisualizeRoughRefractionPS::FParameters>();
-		PassParameters->ViewUniformBuffer = View.ViewUniformBuffer;
-		PassParameters->MiniFontTexture = GetMiniFontTexture();
-		PassParameters->RenderTargets[0] = FRenderTargetBinding(ScreenPassSceneColor.Texture, ERenderTargetLoadAction::ELoad);
-		ShaderPrint::SetParameters(GraphBuilder, View.ShaderPrintData, PassParameters->ShaderPrintParameters);
-		PassParameters->TraceDomainSize = TraceDomainSize;
-		PassParameters->SlabInterfaceLineCount = SlabInterfaceLineCount;
+		//////////////////////////////////////////////////////////////////////////
+		// Trace and update resources
+		{
+			FEvaluateRoughRefractionLobeCS::FParameters* PassParameters = GraphBuilder.AllocParameters<FEvaluateRoughRefractionLobeCS::FParameters>();
+			PassParameters->ViewUniformBuffer = View.ViewUniformBuffer;
+			PassParameters->SampleCountTextureUAV = SampleCountTextureUAV;
+			PassParameters->LobStatisticsBufferUAV = LobStatisticsBufferUAV;
+			PassParameters->MiniFontTexture = GetMiniFontTexture();
+			ShaderPrint::SetParameters(GraphBuilder, View.ShaderPrintData, PassParameters->ShaderPrintParameters);
+			PassParameters->TraceSqrtSampleCount = 128;
+
+			FEvaluateRoughRefractionLobeCS::FPermutationDomain PermutationVector;
+			TShaderMapRef<FEvaluateRoughRefractionLobeCS> ComputeShader(View.ShaderMap, PermutationVector);
+
+			FComputeShaderUtils::AddPass(
+				GraphBuilder,
+				RDG_EVENT_NAME("Strata::EvaluateRoughRefractionLobeCS"),
+				ComputeShader,
+				PassParameters,
+				FComputeShaderUtils::GetGroupCount(1, FEvaluateRoughRefractionLobeCS::ThreadGroupSize));
+		}
 
-		FVisualizeRoughRefractionPS::FPermutationDomain PermutationVector;
-		TShaderMapRef<FVisualizeRoughRefractionPS> PixelShader(View.ShaderMap, PermutationVector);
 
-		FRHIBlendState* PreMultipliedColorTransmittanceBlend = TStaticBlendState<CW_RGB, BO_Add, BF_One, BF_SourceAlpha, BO_Add, BF_Zero, BF_One>::GetRHI();
 
-		FPixelShaderUtils::AddFullscreenPass<FVisualizeRoughRefractionPS>(GraphBuilder, View.ShaderMap, RDG_EVENT_NAME("Strata::VisualizeRoughRefractionPS"),
-			PixelShader, PassParameters, View.ViewRect, PreMultipliedColorTransmittanceBlend);
+		//////////////////////////////////////////////////////////////////////////
+		// Debug print everything on screen
+		{
+			const float TraceDomainSize = 32.0f;
+			const float SlabInterfaceLineCount = 16.0f;
+
+			ShaderPrint::RequestSpaceForLines((TraceDomainSize * TraceDomainSize + SlabInterfaceLineCount * SlabInterfaceLineCount * 2) * 2); // overallocate * 2 for on the fly added debug
+			ShaderPrint::RequestSpaceForCharacters(256);
+
+			FSceneTextureParameters SceneTextureParameters = GetSceneTextureParameters(GraphBuilder);
+			FVisualizeRoughRefractionPS::FParameters* PassParameters = GraphBuilder.AllocParameters<FVisualizeRoughRefractionPS::FParameters>();
+			PassParameters->ViewUniformBuffer = View.ViewUniformBuffer;
+			PassParameters->SampleCountTexture = SampleCountTextureSRV;
+			PassParameters->LobStatisticsBuffer = LobStatisticsBufferSRV;
+			PassParameters->MiniFontTexture = GetMiniFontTexture();
+			ShaderPrint::SetParameters(GraphBuilder, View.ShaderPrintData, PassParameters->ShaderPrintParameters);
+			PassParameters->TraceDomainSize = TraceDomainSize;
+			PassParameters->SlabInterfaceLineCount = SlabInterfaceLineCount;
+			PassParameters->RenderTargets[0] = FRenderTargetBinding(ScreenPassSceneColor.Texture, ERenderTargetLoadAction::ELoad);
+
+			FVisualizeRoughRefractionPS::FPermutationDomain PermutationVector;
+			TShaderMapRef<FVisualizeRoughRefractionPS> PixelShader(View.ShaderMap, PermutationVector);
+
+			FRHIBlendState* PreMultipliedColorTransmittanceBlend = TStaticBlendState<CW_RGB, BO_Add, BF_One, BF_SourceAlpha, BO_Add, BF_Zero, BF_One>::GetRHI();
+
+			FPixelShaderUtils::AddFullscreenPass<FVisualizeRoughRefractionPS>(GraphBuilder, View.ShaderMap, RDG_EVENT_NAME("Strata::VisualizeRoughRefractionPS"),
+				PixelShader, PassParameters, View.ViewRect, PreMultipliedColorTransmittanceBlend);
+		}
 	}
 #endif
 }
