Unify GPU Lightmass shader selection with Path tracing

The logic was nearly identical between both cases, refactoring allows us to share more code and remove virtual calls from FRayTracingMeshProcessor.

Also cleanup the way supported Vertex Factory types are detected by adding a new vertex factory flag.

diff --git a/Engine/Plugins/Experimental/GPULightmass/Source/GPULightmass/Private/LightmapRayTracing.cpp b/Engine/Plugins/Experimental/GPULightmass/Source/GPULightmass/Private/LightmapRayTracing.cpp
--- a/Engine/Plugins/Experimental/GPULightmass/Source/GPULightmass/Private/LightmapRayTracing.cpp
+++ b/Engine/Plugins/Experimental/GPULightmass/Source/GPULightmass/Private/LightmapRayTracing.cpp
@@ -4,137 +4,15 @@
 
 bool IsSupportedVertexFactoryType(const FVertexFactoryType* VertexFactoryType)
 {
-	static FName LocalVfFname = FName(TEXT("FLocalVertexFactory"), FNAME_Find);
-	static FName InstancedVfFname = FName(TEXT("FInstancedStaticMeshVertexFactory"), FNAME_Find);
-	static FName LandscapeVfFname = FName(TEXT("FLandscapeVertexFactory"), FNAME_Find);
-	static FName LandscapeFixedGridVfFname = FName(TEXT("FLandscapeFixedGridVertexFactory"), FNAME_Find);
-	static FName LandscapeXYOffsetVfFname = FName(TEXT("FLandscapeXYOffsetVertexFactory"), FNAME_Find);
-
-	return VertexFactoryType == FindVertexFactoryType(LocalVfFname)
-		|| VertexFactoryType == FindVertexFactoryType(InstancedVfFname)
-		|| VertexFactoryType == FindVertexFactoryType(LandscapeVfFname)
-		|| VertexFactoryType == FindVertexFactoryType(LandscapeFixedGridVfFname)
-		|| VertexFactoryType == FindVertexFactoryType(LandscapeXYOffsetVfFname);
+	return VertexFactoryType->SupportsLightmapBaking();
 }
 
 #if RHI_RAYTRACING
 
-template<bool UseAnyHitShader>
-class TLightmapMaterial : public FMeshMaterialShader
-{
-	DECLARE_SHADER_TYPE(TLightmapMaterial, MeshMaterial);
-public:
-	TLightmapMaterial() = default;
-
-	TLightmapMaterial(const FMeshMaterialShaderType::CompiledShaderInitializerType& Initializer)
-		: FMeshMaterialShader(Initializer)
-	{}
-
-	static bool ShouldCompilePermutation(const FMeshMaterialShaderPermutationParameters& Parameters)
-	{
-		return IsSupportedVertexFactoryType(Parameters.VertexFactoryType)
-			&& ((Parameters.MaterialParameters.bIsMasked || Parameters.MaterialParameters.BlendMode != BLEND_Opaque) == UseAnyHitShader)
-			&& FNoLightMapPolicy::ShouldCompilePermutation(Parameters)
-			&& ShouldCompileRayTracingShadersForProject(Parameters.Platform)
-			&& EnumHasAllFlags(Parameters.Flags, EShaderPermutationFlags::HasEditorOnlyData);
-	}
-
-	static void ModifyCompilationEnvironment(const FMaterialShaderPermutationParameters& Parameters, FShaderCompilerEnvironment& OutEnvironment)
-	{
-		OutEnvironment.SetDefine(TEXT("USE_MATERIAL_CLOSEST_HIT_SHADER"), 1);
-		OutEnvironment.SetDefine(TEXT("USE_MATERIAL_ANY_HIT_SHADER"), UseAnyHitShader);
-		OutEnvironment.SetDefine(TEXT("USE_RAYTRACED_TEXTURE_RAYCONE_LOD"), 0);
-		OutEnvironment.SetDefine(TEXT("SCENE_TEXTURES_DISABLED"), 1);
-		OutEnvironment.SetDefine(TEXT("SIMPLIFIED_MATERIAL_SHADER"), 1);
-		FMeshMaterialShader::ModifyCompilationEnvironment(Parameters, OutEnvironment);
-	}
-
-	static bool ValidateCompiledResult(EShaderPlatform Platform, const FShaderParameterMap& ParameterMap, TArray<FString>& OutError)
-	{
-		if (ParameterMap.ContainsParameterAllocation(FSceneTextureUniformParameters::StaticStructMetadata.GetShaderVariableName()))
-		{
-			OutError.Add(TEXT("Ray tracing closest hit shaders cannot read from the SceneTexturesStruct."));
-			return false;
-		}
-
-		for (const auto& It : ParameterMap.GetParameterMap())
-		{
-			const FParameterAllocation& ParamAllocation = It.Value;
-			if (ParamAllocation.Type != EShaderParameterType::UniformBuffer
-				&& ParamAllocation.Type != EShaderParameterType::LooseData)
-			{
-				OutError.Add(FString::Printf(TEXT("Invalid ray tracing shader parameter '%s'. Only uniform buffers and loose data parameters are supported."), *(It.Key)));
-				return false;
-			}
-		}
-
-		return true;
-	}
-};
-
-using FLightmapMaterialCHS     = TLightmapMaterial<false>;
-using FLightmapMaterialCHS_AHS = TLightmapMaterial<true>;
-
-IMPLEMENT_MATERIAL_SHADER_TYPE(template <>, FLightmapMaterialCHS    , TEXT("/Engine/Private/PathTracing/PathTracingMaterialHitShader.usf"), TEXT("closesthit=PathTracingMaterialCHS"), SF_RayHitGroup);
-IMPLEMENT_MATERIAL_SHADER_TYPE(template <>, FLightmapMaterialCHS_AHS, TEXT("/Engine/Private/PathTracing/PathTracingMaterialHitShader.usf"), TEXT("closesthit=PathTracingMaterialCHS anyhit=PathTracingMaterialAHS"), SF_RayHitGroup);
-
 IMPLEMENT_GLOBAL_SHADER(FLightmapPathTracingRGS, "/Plugin/GPULightmass/Private/LightmapPathTracing.usf", "LightmapPathTracingMainRG", SF_RayGen);
 IMPLEMENT_GLOBAL_SHADER(FVolumetricLightmapPathTracingRGS, "/Plugin/GPULightmass/Private/LightmapPathTracing.usf", "VolumetricLightmapPathTracingMainRG", SF_RayGen);
 IMPLEMENT_GLOBAL_SHADER(FStationaryLightShadowTracingRGS, "/Plugin/GPULightmass/Private/LightmapPathTracing.usf", "StationaryLightShadowTracingMainRG", SF_RayGen);
 
 IMPLEMENT_GLOBAL_SHADER(FFirstBounceRayGuidingCDFBuildCS, "/Plugin/GPULightmass/Private/FirstBounceRayGuidingCDFBuild.usf", "FirstBounceRayGuidingCDFBuildCS", SF_Compute);
 
-// TODO: unify this with the PathTracing process call so we can remove the FLightmapRayTracingMeshProcessor class entirely and de-virtualize this function call
-bool FLightmapRayTracingMeshProcessor::Process(
-	const FMeshBatch& RESTRICT MeshBatch,
-	uint64 BatchElementMask,
-	const FPrimitiveSceneProxy* RESTRICT PrimitiveSceneProxy,
-	const FMaterialRenderProxy& RESTRICT MaterialRenderProxy,
-	const FMaterial& RESTRICT MaterialResource,
-	const FUniformLightMapPolicy& RESTRICT LightMapPolicy)
-{
-	const FVertexFactory* VertexFactory = MeshBatch.VertexFactory;
-
-	TMeshProcessorShaders<
-		FMeshMaterialShader,
-		FMeshMaterialShader,
-		FMeshMaterialShader,
-		FMeshMaterialShader,
-		FMeshMaterialShader> RayTracingShaders;
-
-	FMaterialShaderTypes ShaderTypes;
-
-	if (MaterialResource.IsMasked() || MaterialResource.GetBlendMode() != BLEND_Opaque)
-	{
-		ShaderTypes.AddShaderType<FLightmapMaterialCHS_AHS>();
-	}
-	else
-	{
-		ShaderTypes.AddShaderType<FLightmapMaterialCHS>();
-	}
-
-	FMaterialShaders Shaders;
-	if (!MaterialResource.TryGetShaders(ShaderTypes, VertexFactory->GetType(), Shaders))
-	{
-		return false;
-	}
-
-	check(Shaders.TryGetShader(SF_RayHitGroup, RayTracingShaders.RayHitGroupShader));
-
-	TBasePassShaderElementData<FUniformLightMapPolicy> ShaderElementData(MeshBatch.LCI);
-	ShaderElementData.InitializeMeshMaterialData(ViewIfDynamicMeshCommand, PrimitiveSceneProxy, MeshBatch, -1, true);
-
-	BuildRayTracingMeshCommands(
-		MeshBatch,
-		BatchElementMask,
-		PrimitiveSceneProxy,
-		MaterialRenderProxy,
-		MaterialResource,
-		PassDrawRenderState,
-		RayTracingShaders,
-		ShaderElementData);
-
-	return true;
-}
-
 #endif
diff --git a/Engine/Plugins/Experimental/GPULightmass/Source/GPULightmass/Private/LightmapRayTracing.h b/Engine/Plugins/Experimental/GPULightmass/Source/GPULightmass/Private/LightmapRayTracing.h
--- a/Engine/Plugins/Experimental/GPULightmass/Source/GPULightmass/Private/LightmapRayTracing.h
+++ b/Engine/Plugins/Experimental/GPULightmass/Source/GPULightmass/Private/LightmapRayTracing.h
@@ -28,16 +28,7 @@ class FLightmapRayTracingMeshProcessor : public FRayTracingMeshProcessor
 {
 public:
 	FLightmapRayTracingMeshProcessor(FRayTracingMeshCommandContext* InCommandContext, FMeshPassProcessorRenderState InPassDrawRenderState)
-		: FRayTracingMeshProcessor(InCommandContext, nullptr, nullptr, InPassDrawRenderState, ERayTracingMeshCommandsMode::RAY_TRACING) // NOTE: Must pass RAY_TRACING here so that Process can be called
-	{}
-
-	virtual bool Process(
-		const FMeshBatch& RESTRICT MeshBatch,
-		uint64 BatchElementMask,
-		const FPrimitiveSceneProxy* RESTRICT PrimitiveSceneProxy,
-		const FMaterialRenderProxy& RESTRICT MaterialRenderProxy,
-		const FMaterial& RESTRICT MaterialResource,
-		const FUniformLightMapPolicy& RESTRICT LightMapPolicy) override;
+		: FRayTracingMeshProcessor(InCommandContext, nullptr, nullptr, InPassDrawRenderState, ERayTracingMeshCommandsMode::LIGHTMAP_TRACING) {}
 };
 
 class FLightmapPathTracingRGS : public FGlobalShader
diff --git a/Engine/Source/Runtime/Engine/Private/InstancedStaticMesh.cpp b/Engine/Source/Runtime/Engine/Private/InstancedStaticMesh.cpp
--- a/Engine/Source/Runtime/Engine/Private/InstancedStaticMesh.cpp
+++ b/Engine/Source/Runtime/Engine/Private/InstancedStaticMesh.cpp
@@ -993,6 +993,7 @@ IMPLEMENT_VERTEX_FACTORY_TYPE(FInstancedStaticMeshVertexFactory,"/Engine/Private
 	| EVertexFactoryFlags::SupportsCachingMeshDrawCommands
 	| EVertexFactoryFlags::SupportsRayTracing
 	| EVertexFactoryFlags::SupportsRayTracingDynamicGeometry
+	| EVertexFactoryFlags::SupportsLightmapBaking
 	| EVertexFactoryFlags::SupportsPrimitiveIdStream
 );
 
diff --git a/Engine/Source/Runtime/Engine/Private/LocalVertexFactory.cpp b/Engine/Source/Runtime/Engine/Private/LocalVertexFactory.cpp
--- a/Engine/Source/Runtime/Engine/Private/LocalVertexFactory.cpp
+++ b/Engine/Source/Runtime/Engine/Private/LocalVertexFactory.cpp
@@ -413,4 +413,5 @@ IMPLEMENT_VERTEX_FACTORY_TYPE(FLocalVertexFactory,"/Engine/Private/LocalVertexFa
 	| EVertexFactoryFlags::SupportsPrimitiveIdStream
 	| EVertexFactoryFlags::SupportsRayTracing
 	| EVertexFactoryFlags::SupportsRayTracingDynamicGeometry
+	| EVertexFactoryFlags::SupportsLightmapBaking
 );
diff --git a/Engine/Source/Runtime/Landscape/Private/LandscapeRender.cpp b/Engine/Source/Runtime/Landscape/Private/LandscapeRender.cpp
--- a/Engine/Source/Runtime/Landscape/Private/LandscapeRender.cpp
+++ b/Engine/Source/Runtime/Landscape/Private/LandscapeRender.cpp
@@ -3258,6 +3258,7 @@ IMPLEMENT_VERTEX_FACTORY_TYPE(FLandscapeVertexFactory, "/Engine/Private/Landscap
 	| EVertexFactoryFlags::SupportsCachingMeshDrawCommands
 	| EVertexFactoryFlags::SupportsRayTracing
 	| EVertexFactoryFlags::SupportsRayTracingDynamicGeometry
+	| EVertexFactoryFlags::SupportsLightmapBaking
 	| EVertexFactoryFlags::SupportsPrimitiveIdStream
 );
 
@@ -3302,6 +3303,7 @@ IMPLEMENT_VERTEX_FACTORY_TYPE(FLandscapeXYOffsetVertexFactory, "/Engine/Private/
 	| EVertexFactoryFlags::SupportsCachingMeshDrawCommands
 	| EVertexFactoryFlags::SupportsRayTracing
 	| EVertexFactoryFlags::SupportsRayTracingDynamicGeometry
+	| EVertexFactoryFlags::SupportsLightmapBaking
 	| EVertexFactoryFlags::SupportsPrimitiveIdStream
 );
 
@@ -3329,6 +3331,7 @@ IMPLEMENT_VERTEX_FACTORY_TYPE(FLandscapeFixedGridVertexFactory, "/Engine/Private
 	| EVertexFactoryFlags::SupportsCachingMeshDrawCommands
 	| EVertexFactoryFlags::SupportsRayTracing
 	| EVertexFactoryFlags::SupportsRayTracingDynamicGeometry
+	| EVertexFactoryFlags::SupportsLightmapBaking
 	| EVertexFactoryFlags::SupportsPrimitiveIdStream
 );
 
diff --git a/Engine/Source/Runtime/RenderCore/Public/VertexFactory.h b/Engine/Source/Runtime/RenderCore/Public/VertexFactory.h
--- a/Engine/Source/Runtime/RenderCore/Public/VertexFactory.h
+++ b/Engine/Source/Runtime/RenderCore/Public/VertexFactory.h
@@ -103,6 +103,7 @@ enum class EVertexFactoryFlags : uint32
 	SupportsRayTracing                    = 1u << 9,
 	SupportsRayTracingDynamicGeometry     = 1u << 10,
 	SupportsRayTracingProceduralPrimitive = 1u << 11,
+	SupportsLightmapBaking                = 1u << 12
 };
 ENUM_CLASS_FLAGS(EVertexFactoryFlags);
 
@@ -405,6 +406,7 @@ public:
 	bool SupportsRayTracing() const                    { return HasFlags(EVertexFactoryFlags::SupportsRayTracing); }
 	bool SupportsRayTracingDynamicGeometry() const     { return HasFlags(EVertexFactoryFlags::SupportsRayTracingDynamicGeometry); }
 	bool SupportsRayTracingProceduralPrimitive() const { return HasFlags(EVertexFactoryFlags::SupportsRayTracingProceduralPrimitive); }
+	bool SupportsLightmapBaking() const                { return HasFlags(EVertexFactoryFlags::SupportsLightmapBaking); }
 
 	// Hash function.
 	friend uint32 GetTypeHash(const FVertexFactoryType* Type)
diff --git a/Engine/Source/Runtime/Renderer/Private/PathTracing.cpp b/Engine/Source/Runtime/Renderer/Private/PathTracing.cpp
--- a/Engine/Source/Runtime/Renderer/Private/PathTracing.cpp
+++ b/Engine/Source/Runtime/Renderer/Private/PathTracing.cpp
@@ -567,6 +567,7 @@ class FPathTracingBuildLightGridCS : public FGlobalShader
 
 	static bool ShouldCompilePermutation(const FGlobalShaderPermutationParameters& Parameters)
 	{
+		// shared with GPULightmass
 		return ShouldCompileRayTracingShadersForProject(Parameters.Platform);
 	}
 
@@ -758,7 +759,7 @@ IMPLEMENT_SHADER_TYPE(, FPathTracingBuildAtmosphereOpticalDepthLUTCS, TEXT("/Eng
 
 
 
-template<bool UseAnyHitShader, bool UseIntersectionShader>
+template<bool UseAnyHitShader, bool UseIntersectionShader, bool UseSimplifiedShader>
 class TPathTracingMaterial : public FMeshMaterialShader
 {
 	DECLARE_SHADER_TYPE(TPathTracingMaterial, MeshMaterial);
@@ -771,11 +772,39 @@ public:
 
 	static bool ShouldCompilePermutation(const FMeshMaterialShaderPermutationParameters& Parameters)
 	{
+		if (!ShouldCompileRayTracingShadersForProject(Parameters.Platform))
+		{
+			// is raytracing enabled at all?
+			return false;
+		}
+		if (!Parameters.VertexFactoryType->SupportsRayTracing())
+		{
+			// does the VF support ray tracing at all?
+			return false;
+		}
+		if ((Parameters.MaterialParameters.bIsMasked || Parameters.MaterialParameters.BlendMode != BLEND_Opaque) != UseAnyHitShader)
+		{
+			// the anyhit permutation is only required if the material is masked or has a non-opaque blend mode
+			return false;
+		}
 		const bool bUseProceduralPrimitive = Parameters.VertexFactoryType->SupportsRayTracingProceduralPrimitive() && FDataDrivenShaderPlatformInfo::GetSupportsRayTracingProceduralPrimitive(Parameters.Platform);
-		return Parameters.VertexFactoryType->SupportsRayTracing()
-			&& (UseIntersectionShader == bUseProceduralPrimitive)
-			&& ((Parameters.MaterialParameters.bIsMasked || Parameters.MaterialParameters.BlendMode != BLEND_Opaque) == UseAnyHitShader)
-			&& ShouldCompilePathTracingShadersForProject(Parameters.Platform);
+		if (UseIntersectionShader != bUseProceduralPrimitive)
+		{
+			// only need to compile the intersection shader permutation if the VF actually requires it
+			return false;
+		}
+		if (UseSimplifiedShader)
+		{
+			// this is only used by GPULightmass
+			return EnumHasAllFlags(Parameters.Flags, EShaderPermutationFlags::HasEditorOnlyData) &&
+				   Parameters.VertexFactoryType->SupportsLightmapBaking() &&
+				   FModuleManager::Get().IsModuleLoaded(TEXT("GPULightmass"));
+		}
+		else
+		{
+			// this is only used by the Path Tracer
+			return ShouldCompilePathTracingShadersForProject(Parameters.Platform);
+		}
 	}
 
 	static void ModifyCompilationEnvironment(const FMaterialShaderPermutationParameters& Parameters, FShaderCompilerEnvironment& OutEnvironment)
@@ -785,7 +814,7 @@ public:
 		OutEnvironment.SetDefine(TEXT("USE_MATERIAL_INTERSECTION_SHADER"), UseIntersectionShader ? 1 : 0);
 		OutEnvironment.SetDefine(TEXT("USE_RAYTRACED_TEXTURE_RAYCONE_LOD"), 0);
 		OutEnvironment.SetDefine(TEXT("SCENE_TEXTURES_DISABLED"), 1);
-		OutEnvironment.SetDefine(TEXT("SIMPLIFIED_MATERIAL_SHADER"), 0); // TODO: expose a permutation so we can unify with GPULightmass?
+		OutEnvironment.SetDefine(TEXT("SIMPLIFIED_MATERIAL_SHADER"), UseSimplifiedShader);
 		FMeshMaterialShader::ModifyCompilationEnvironment(Parameters, OutEnvironment);
 	}
 
@@ -812,16 +841,24 @@ public:
 	}
 };
 
-using FPathTracingMaterialCHS        = TPathTracingMaterial<false, false>;
-using FPathTracingMaterialCHS_AHS    = TPathTracingMaterial<true , false>;
-using FPathTracingMaterialCHS_IS     = TPathTracingMaterial<false, true >;
-using FPathTracingMaterialCHS_AHS_IS = TPathTracingMaterial<true , true >;
+
+// TODO: It would be nice to avoid this template boilerplate and just use ordinary permutations. This would require allowing the FunctionName for the material to be dependent on the permutation somehow
+using FPathTracingMaterialCHS        = TPathTracingMaterial<false, false, false>;
+using FPathTracingMaterialCHS_AHS    = TPathTracingMaterial<true , false, false>;
+using FPathTracingMaterialCHS_IS     = TPathTracingMaterial<false, true , false>;
+using FPathTracingMaterialCHS_AHS_IS = TPathTracingMaterial<true , true , false>;
+
+// NOTE: lightmass doesn't work with intersection shader VFs at the moment, so avoid instantiating permutations that will never generate any shaders
+using FGPULightmassCHS               = TPathTracingMaterial<false, false, true>;
+using FGPULightmassCHS_AHS           = TPathTracingMaterial<true , false, true>;
 
 
 IMPLEMENT_MATERIAL_SHADER_TYPE(template <>, FPathTracingMaterialCHS       , TEXT("/Engine/Private/PathTracing/PathTracingMaterialHitShader.usf"), TEXT("closesthit=PathTracingMaterialCHS"), SF_RayHitGroup);
 IMPLEMENT_MATERIAL_SHADER_TYPE(template <>, FPathTracingMaterialCHS_AHS   , TEXT("/Engine/Private/PathTracing/PathTracingMaterialHitShader.usf"), TEXT("closesthit=PathTracingMaterialCHS anyhit=PathTracingMaterialAHS"), SF_RayHitGroup);
 IMPLEMENT_MATERIAL_SHADER_TYPE(template <>, FPathTracingMaterialCHS_IS    , TEXT("/Engine/Private/PathTracing/PathTracingMaterialHitShader.usf"), TEXT("closesthit=PathTracingMaterialCHS intersection=MaterialIS"), SF_RayHitGroup);
 IMPLEMENT_MATERIAL_SHADER_TYPE(template <>, FPathTracingMaterialCHS_AHS_IS, TEXT("/Engine/Private/PathTracing/PathTracingMaterialHitShader.usf"), TEXT("closesthit=PathTracingMaterialCHS anyhit=PathTracingMaterialAHS intersection=MaterialIS"), SF_RayHitGroup);
+IMPLEMENT_MATERIAL_SHADER_TYPE(template <>, FGPULightmassCHS              , TEXT("/Engine/Private/PathTracing/PathTracingMaterialHitShader.usf"), TEXT("closesthit=PathTracingMaterialCHS"), SF_RayHitGroup);
+IMPLEMENT_MATERIAL_SHADER_TYPE(template <>, FGPULightmassCHS_AHS          , TEXT("/Engine/Private/PathTracing/PathTracingMaterialHitShader.usf"), TEXT("closesthit=PathTracingMaterialCHS anyhit=PathTracingMaterialAHS"), SF_RayHitGroup);
 
 
 
@@ -832,40 +869,65 @@ bool FRayTracingMeshProcessor::ProcessPathTracing(
 	const FMaterialRenderProxy& RESTRICT MaterialRenderProxy,
 	const FMaterial& RESTRICT MaterialResource)
 {
-	const FVertexFactory* VertexFactory = MeshBatch.VertexFactory;
-
-	TMeshProcessorShaders<
-		FMeshMaterialShader,
-		FMeshMaterialShader,
-		FMeshMaterialShader,
-		FMeshMaterialShader,
-		FMeshMaterialShader> RayTracingShaders;
 
 	FMaterialShaderTypes ShaderTypes;
 
-	const bool bUseProceduralPrimitive = VertexFactory->GetType()->SupportsRayTracingProceduralPrimitive() && FDataDrivenShaderPlatformInfo::GetSupportsRayTracingProceduralPrimitive(GMaxRHIShaderPlatform);
-	if (MaterialResource.IsMasked() || MaterialResource.GetBlendMode() != BLEND_Opaque)
+	const bool bUseProceduralPrimitive = MeshBatch.VertexFactory->GetType()->SupportsRayTracingProceduralPrimitive() &&
+		FDataDrivenShaderPlatformInfo::GetSupportsRayTracingProceduralPrimitive(GMaxRHIShaderPlatform);
+	switch (RayTracingMeshCommandsMode)
 	{
-		if (bUseProceduralPrimitive)
-			ShaderTypes.AddShaderType<FPathTracingMaterialCHS_AHS_IS>();
-		else
-			ShaderTypes.AddShaderType<FPathTracingMaterialCHS_AHS>();
-	}
-	else
-	{
-		if (bUseProceduralPrimitive)
-			ShaderTypes.AddShaderType<FPathTracingMaterialCHS_IS>();
-		else
-			ShaderTypes.AddShaderType<FPathTracingMaterialCHS>();
+		case ERayTracingMeshCommandsMode::PATH_TRACING:
+		{
+			if (MaterialResource.IsMasked() || MaterialResource.GetBlendMode() != BLEND_Opaque)
+			{
+				if (bUseProceduralPrimitive)
+					ShaderTypes.AddShaderType<FPathTracingMaterialCHS_AHS_IS>();
+				else
+					ShaderTypes.AddShaderType<FPathTracingMaterialCHS_AHS>();
+			}
+			else
+			{
+				if (bUseProceduralPrimitive)
+					ShaderTypes.AddShaderType<FPathTracingMaterialCHS_IS>();
+				else
+					ShaderTypes.AddShaderType<FPathTracingMaterialCHS>();
+			}
+			break;
+		}
+		case ERayTracingMeshCommandsMode::LIGHTMAP_TRACING:
+		{
+			if (MaterialResource.IsMasked() || MaterialResource.GetBlendMode() != BLEND_Opaque)
+			{
+				ShaderTypes.AddShaderType<FGPULightmassCHS_AHS>();
+			}
+			else
+			{
+				ShaderTypes.AddShaderType<FGPULightmassCHS>();
+			}
+			break;
+		}
+		default:
+		{
+		    return false;
+		}
 	}
 
 	FMaterialShaders Shaders;
-	if (!MaterialResource.TryGetShaders(ShaderTypes, VertexFactory->GetType(), Shaders))
+	if (!MaterialResource.TryGetShaders(ShaderTypes, MeshBatch.VertexFactory->GetType(), Shaders))
 	{
 		return false;
 	}
 
-	check(Shaders.TryGetShader(SF_RayHitGroup, RayTracingShaders.RayHitGroupShader));
+	TMeshProcessorShaders<
+		FMeshMaterialShader,
+		FMeshMaterialShader,
+		FMeshMaterialShader,
+		FMeshMaterialShader,
+		FMeshMaterialShader> RayTracingShaders;
+	if (!Shaders.TryGetShader(SF_RayHitGroup, RayTracingShaders.RayHitGroupShader))
+	{
+		return false;
+	}
 
 	TBasePassShaderElementData<FUniformLightMapPolicy> ShaderElementData(MeshBatch.LCI);
 	ShaderElementData.InitializeMeshMaterialData(ViewIfDynamicMeshCommand, PrimitiveSceneProxy, MeshBatch, -1, true);
@@ -1487,15 +1549,12 @@ void FDeferredShadingSceneRenderer::PreparePathTracing(const FSceneViewFamily& V
 		&& ShouldCompilePathTracingShadersForProject(ViewFamily.GetShaderPlatform()))
 	{
 		// Declare all RayGen shaders that require material closest hit shaders to be bound
-		for (int CompactionType = 0; CompactionType < FPathTracingRG::FCompactionType::PermutationCount; CompactionType++)
+		const int CompactionType = CVarPathTracingCompaction.GetValueOnRenderThread();
+		FPathTracingRG::FPermutationDomain PermutationVector;
+		PermutationVector.Set<FPathTracingRG::FCompactionType>(CompactionType);
+		FGlobalShaderPermutationParameters Parameters(FPathTracingRG::StaticGetTypeLayout().Name, ViewFamily.GetShaderPlatform(), PermutationVector.ToDimensionValueId());
+		if (FPathTracingRG::ShouldCompilePermutation(Parameters))
 		{
-			FPathTracingRG::FPermutationDomain PermutationVector;
-			PermutationVector.Set<FPathTracingRG::FCompactionType>(CompactionType);
-			FGlobalShaderPermutationParameters Parameters(FPathTracingRG::StaticGetTypeLayout().Name, ViewFamily.GetShaderPlatform(), PermutationVector.ToDimensionValueId());
-			if (!FPathTracingRG::ShouldCompilePermutation(Parameters))
-			{
-				continue;
-			}
 			auto RayGenShader = GetGlobalShaderMap(ViewFamily.GetShaderPlatform())->GetShader<FPathTracingRG>(PermutationVector);
 			OutRayGenShaders.Add(RayGenShader.GetRayTracingShader());
 		}
diff --git a/Engine/Source/Runtime/Renderer/Private/RayTracing/RayTracingMaterialHitShaders.cpp b/Engine/Source/Runtime/Renderer/Private/RayTracing/RayTracingMaterialHitShaders.cpp
--- a/Engine/Source/Runtime/Renderer/Private/RayTracing/RayTracingMaterialHitShaders.cpp
+++ b/Engine/Source/Runtime/Renderer/Private/RayTracing/RayTracingMaterialHitShaders.cpp
@@ -441,7 +441,8 @@ bool FRayTracingMeshProcessor::TryAddMeshBatch(
 	if ((!PrimitiveSceneProxy || PrimitiveSceneProxy->ShouldRenderInMainPass())
 		&& ShouldIncludeDomainInMeshPass(Material.GetMaterialDomain()))
 	{
-		if (RayTracingMeshCommandsMode == ERayTracingMeshCommandsMode::PATH_TRACING)
+		if (RayTracingMeshCommandsMode == ERayTracingMeshCommandsMode::PATH_TRACING ||
+			RayTracingMeshCommandsMode == ERayTracingMeshCommandsMode::LIGHTMAP_TRACING)
 		{
 			// Path Tracer has its own process call so that it can attach its own material permutation
 			return ProcessPathTracing(MeshBatch, BatchElementMask, PrimitiveSceneProxy, MaterialRenderProxy, Material);
diff --git a/Engine/Source/Runtime/Renderer/Private/RayTracing/RayTracingMaterialHitShaders.h b/Engine/Source/Runtime/Renderer/Private/RayTracing/RayTracingMaterialHitShaders.h
--- a/Engine/Source/Runtime/Renderer/Private/RayTracing/RayTracingMaterialHitShaders.h
+++ b/Engine/Source/Runtime/Renderer/Private/RayTracing/RayTracingMaterialHitShaders.h
@@ -39,7 +39,7 @@ protected:
 	FMeshPassProcessorRenderState PassDrawRenderState;
 	ERayTracingMeshCommandsMode RayTracingMeshCommandsMode;
 
-	virtual bool Process(
+	bool Process(
 		const FMeshBatch& RESTRICT MeshBatch,
 		uint64 BatchElementMask,
 		const FPrimitiveSceneProxy* RESTRICT PrimitiveSceneProxy,
diff --git a/Engine/Source/Runtime/Renderer/Private/ScenePrivate.h b/Engine/Source/Runtime/Renderer/Private/ScenePrivate.h
--- a/Engine/Source/Runtime/Renderer/Private/ScenePrivate.h
+++ b/Engine/Source/Runtime/Renderer/Private/ScenePrivate.h
@@ -2671,7 +2671,8 @@ struct FMeshComputeDispatchCommand
 
 enum class ERayTracingMeshCommandsMode : uint8 {
 	RAY_TRACING,
-	PATH_TRACING
+	PATH_TRACING,
+	LIGHTMAP_TRACING,
 };
 
 #endif
