Fix compiler warning
- Check that both InstanceTransforms and InstanceTransformsView are empty

diff --git a/Engine/Source/Runtime/Renderer/Private/PrimitiveSceneInfo.cpp b/Engine/Source/Runtime/Renderer/Private/PrimitiveSceneInfo.cpp
--- a/Engine/Source/Runtime/Renderer/Private/PrimitiveSceneInfo.cpp
+++ b/Engine/Source/Runtime/Renderer/Private/PrimitiveSceneInfo.cpp
@@ -919,7 +919,7 @@ void FPrimitiveSceneInfo::UpdateCachedRayTracingInstance(FPrimitiveSceneInfo* Sc
 {
 	if (EnumHasAnyFlags(Flags, ERayTracingPrimitiveFlags::CacheInstances))
 	{
-		checkf(CachedRayTracingInstance.InstanceTransforms.IsEmpty() && CachedRayTracingInstance.InstanceTransforms.IsEmpty(),
+		checkf(CachedRayTracingInstance.InstanceTransforms.IsEmpty() && CachedRayTracingInstance.InstanceTransformsView.IsEmpty(),
 			TEXT("Primitives with ERayTracingPrimitiveFlags::CacheInstances get instances transforms from GPUScene"));
 
 		// TODO: allocate from FRayTracingScene & do better low-level caching
