Fix compilation error in non-editor builds

diff --git a/Engine/Source/Runtime/Renderer/Private/PathTracing.cpp b/Engine/Source/Runtime/Renderer/Private/PathTracing.cpp
--- a/Engine/Source/Runtime/Renderer/Private/PathTracing.cpp
+++ b/Engine/Source/Runtime/Renderer/Private/PathTracing.cpp
@@ -25,6 +25,7 @@ TAutoConsoleVariable<int32> CVarPathTracing(
 #include "RayTracing/RayTracingMaterialHitShaders.h"
 #include "RenderCore/Public/GenerateMips.h"
 #include "HairStrands/HairStrandsData.h"
+#include "Modules/ModuleManager.h"
 #include <limits>
 
 TAutoConsoleVariable<int32> CVarPathTracingCompaction(
@@ -795,10 +796,14 @@ public:
 		}
 		if (UseSimplifiedShader)
 		{
+#if WITH_EDITOR
 			// this is only used by GPULightmass
 			return EnumHasAllFlags(Parameters.Flags, EShaderPermutationFlags::HasEditorOnlyData) &&
 				   Parameters.VertexFactoryType->SupportsLightmapBaking() &&
 				   FModuleManager::Get().IsModuleLoaded(TEXT("GPULightmass"));
+#else
+			return false;
+#endif
 		}
 		else
 		{
