Add Strata tile support for DiffuseIndirectComposite.

diff --git a/Engine/Shaders/Private/DiffuseIndirectComposite.usf b/Engine/Shaders/Private/DiffuseIndirectComposite.usf
--- a/Engine/Shaders/Private/DiffuseIndirectComposite.usf
+++ b/Engine/Shaders/Private/DiffuseIndirectComposite.usf
@@ -50,6 +50,17 @@
 #define COMPILE_MOMENT1_ACCUMULATOR 1
 #define COMPILE_BOX_KERNEL 1
 
+#if STRATA_ENABLED && DIM_APPLY_DIFFUSE_INDIRECT == 4  
+  #if STRATA_TILETYPE == 0
+	#define STRATA_FASTPATH 1
+  #elif STRATA_TILETYPE == 1
+	#define STRATA_SINGLEPATH 1
+  #elif STRATA_TILETYPE == 2
+	// COMPLEX PATH
+  #else
+	#error Strata tile type non-implemented
+  #endif
+#endif
 
 //------------------------------------------------------- INCLUDES
 
diff --git a/Engine/Source/Runtime/Renderer/Private/IndirectLightRendering.cpp b/Engine/Source/Runtime/Renderer/Private/IndirectLightRendering.cpp
--- a/Engine/Source/Runtime/Renderer/Private/IndirectLightRendering.cpp
+++ b/Engine/Source/Runtime/Renderer/Private/IndirectLightRendering.cpp
@@ -103,8 +103,9 @@ class FDiffuseIndirectCompositePS : public FGlobalShader
 	class FApplyDiffuseIndirectDim : SHADER_PERMUTATION_INT("DIM_APPLY_DIFFUSE_INDIRECT", 5);
 	class FUpscaleDiffuseIndirectDim : SHADER_PERMUTATION_BOOL("DIM_UPSCALE_DIFFUSE_INDIRECT");
 	class FScreenBentNormal : SHADER_PERMUTATION_BOOL("DIM_SCREEN_BENT_NORMAL");
+	class FStrataTileType : SHADER_PERMUTATION_INT("STRATA_TILETYPE", 3);
 
-	using FPermutationDomain = TShaderPermutationDomain<FApplyDiffuseIndirectDim, FUpscaleDiffuseIndirectDim, FScreenBentNormal>;
+	using FPermutationDomain = TShaderPermutationDomain<FApplyDiffuseIndirectDim, FUpscaleDiffuseIndirectDim, FScreenBentNormal, FStrataTileType>;
 
 	static bool ShouldCompilePermutation(const FGlobalShaderPermutationParameters& Parameters)
 	{
@@ -127,6 +128,12 @@ class FDiffuseIndirectCompositePS : public FGlobalShader
 			return false;
 		}
 
+		// Build Strata tile permutation only for Lumen
+		if (PermutationVector.Get<FStrataTileType>() != EStrataTileMaterialType::EComplex)
+		{
+			return Strata::IsStrataEnabled() && PermutationVector.Get<FApplyDiffuseIndirectDim>() == 4;
+		}
+
 		return IsFeatureLevelSupported(Parameters.Platform, ERHIFeatureLevel::SM5);
 	}
 
@@ -148,6 +155,7 @@ class FDiffuseIndirectCompositePS : public FGlobalShader
 
 		SHADER_PARAMETER_RDG_UNIFORM_BUFFER(FSceneTextureUniformParameters, SceneTexturesStruct)
 		SHADER_PARAMETER_RDG_UNIFORM_BUFFER(FStrataGlobalUniformParameters, Strata)
+		SHADER_PARAMETER_STRUCT_INCLUDE(Strata::FStrataTilePassVS::FParameters, StrataTile)
 		SHADER_PARAMETER_STRUCT_INCLUDE(Denoiser::FCommonShaderParameters, DenoiserCommonParameters)
 		SHADER_PARAMETER_STRUCT_INCLUDE(FSceneTextureParameters, SceneTextures)
 		SHADER_PARAMETER_STRUCT_REF(FViewUniformShaderParameters, ViewUniformBuffer)
@@ -1001,8 +1009,10 @@ void FDeferredShadingSceneRenderer::RenderDiffuseIndirectAndAmbientOcclusion(
 		}
 
 		// Applies diffuse indirect and ambient occlusion to the scene color.
-		if ((DenoiserOutputs.Textures[0] || AmbientOcclusionMask) && (!bIsVisualizePass || ViewPipelineState.DiffuseIndirectDenoiser != IScreenSpaceDenoiser::EMode::Disabled || ViewPipelineState.DiffuseIndirectMethod == EDiffuseIndirectMethod::Lumen)
-			&& !(IsMetalPlatform(ShaderPlatform) && !IsMetalSM5Platform(ShaderPlatform)))
+		const bool bApplyDiffuseIndirect = (DenoiserOutputs.Textures[0] || AmbientOcclusionMask) && (!bIsVisualizePass || ViewPipelineState.DiffuseIndirectDenoiser != IScreenSpaceDenoiser::EMode::Disabled || ViewPipelineState.DiffuseIndirectMethod == EDiffuseIndirectMethod::Lumen)
+			&& !(IsMetalPlatform(ShaderPlatform) && !IsMetalSM5Platform(ShaderPlatform));
+
+		auto ApplyDiffuseIndirect = [&](EStrataTileMaterialType TileType)
 		{
 			FDiffuseIndirectCompositePS::FParameters* PassParameters = GraphBuilder.AllocParameters<FDiffuseIndirectCompositePS::FParameters>();
 			PassParameters->Strata = Strata::BindStrataGlobalUniformParameters(View.StrataSceneData);
@@ -1076,6 +1086,8 @@ void FDeferredShadingSceneRenderer::RenderDiffuseIndirectAndAmbientOcclusion(
 
 			const TCHAR* DiffuseIndirectSampling = TEXT("Disabled");
 			FDiffuseIndirectCompositePS::FPermutationDomain PermutationVector;
+			PermutationVector.Set<FDiffuseIndirectCompositePS::FStrataTileType>(EStrataTileMaterialType::EComplex);
+
 			bool bUpscale = false;
 
 			if (DenoiserOutputs.Textures[0])
@@ -1094,6 +1106,10 @@ void FDeferredShadingSceneRenderer::RenderDiffuseIndirectAndAmbientOcclusion(
 				{
 					PermutationVector.Set<FDiffuseIndirectCompositePS::FApplyDiffuseIndirectDim>(4);
 					PermutationVector.Set<FDiffuseIndirectCompositePS::FScreenBentNormal>(ScreenBentNormalParameters.UseScreenBentNormal != 0);
+					if (Strata::IsStrataEnabled() && TileType != EStrataTileMaterialType::ECount)
+					{
+						PermutationVector.Set<FDiffuseIndirectCompositePS::FStrataTileType>(TileType);
+					}
 					DiffuseIndirectSampling = TEXT("ScreenProbeGather");
 				}
 				else
@@ -1107,7 +1123,6 @@ void FDeferredShadingSceneRenderer::RenderDiffuseIndirectAndAmbientOcclusion(
 			}
 
 			TShaderMapRef<FDiffuseIndirectCompositePS> PixelShader(View.ShaderMap, PermutationVector);
-			ClearUnusedGraphResources(PixelShader, PassParameters);
 
 			FRHIBlendState* BlendState = PermutationVector.Get<FDiffuseIndirectCompositePS::FApplyDiffuseIndirectDim>() > 0 ?
 				TStaticBlendState<CW_RGBA, BO_Add, BF_One, BF_Source1Color, BO_Add, BF_One, BF_Source1Alpha>::GetRHI() :
@@ -1118,21 +1133,87 @@ void FDeferredShadingSceneRenderer::RenderDiffuseIndirectAndAmbientOcclusion(
 				BlendState = TStaticBlendState<>::GetRHI();
 			}
 
-			FPixelShaderUtils::AddFullscreenPass(
-				GraphBuilder,
-				View.ShaderMap,
-				RDG_EVENT_NAME(
-					"DiffuseIndirectComposite(DiffuseIndirect=%s%s%s%s) %dx%d",
-					DiffuseIndirectSampling,
-					PermutationVector.Get<FDiffuseIndirectCompositePS::FUpscaleDiffuseIndirectDim>() ? TEXT(" UpscaleDiffuseIndirect") : TEXT(""),
-					AmbientOcclusionMask ? TEXT(" ApplyAOToSceneColor") : TEXT(""),
-					PassParameters->ApplyAOToDynamicDiffuseIndirect > 0.0f ? TEXT(" ApplyAOToDynamicDiffuseIndirect") : TEXT(""),
-					View.ViewRect.Width(), View.ViewRect.Height()),
-				PixelShader,
-				PassParameters,
-				View.ViewRect,
-				BlendState);
-		} // if (DenoiserOutputs.Color || bApplySSAO)
+			if (TileType == EStrataTileMaterialType::ECount)
+			{
+				ClearUnusedGraphResources(PixelShader, PassParameters);
+
+				FPixelShaderUtils::AddFullscreenPass(
+					GraphBuilder,
+					View.ShaderMap,
+					RDG_EVENT_NAME(
+						"DiffuseIndirectComposite(DiffuseIndirect=%s%s%s%s) %dx%d",
+						DiffuseIndirectSampling,
+						PermutationVector.Get<FDiffuseIndirectCompositePS::FUpscaleDiffuseIndirectDim>() ? TEXT(" UpscaleDiffuseIndirect") : TEXT(""),
+						AmbientOcclusionMask ? TEXT(" ApplyAOToSceneColor") : TEXT(""),
+						PassParameters->ApplyAOToDynamicDiffuseIndirect > 0.0f ? TEXT(" ApplyAOToDynamicDiffuseIndirect") : TEXT(""),
+						View.ViewRect.Width(), View.ViewRect.Height()),
+					PixelShader,
+					PassParameters,
+					View.ViewRect,
+					BlendState);
+			}
+			else
+			{
+				check(Strata::IsStrataEnabled());
+
+				Strata::FStrataTilePassVS::FPermutationDomain VSPermutationVector;
+				VSPermutationVector.Set< Strata::FStrataTilePassVS::FEnableDebug >(false);
+				VSPermutationVector.Set< Strata::FStrataTilePassVS::FEnableTexCoordScreenVector >(true);
+				TShaderMapRef<Strata::FStrataTilePassVS> TileVertexShader(View.ShaderMap, VSPermutationVector);
+
+				ClearUnusedGraphResources(PixelShader, PassParameters);
+
+				EPrimitiveType PrimitiveType = PT_TriangleList;
+				PassParameters->StrataTile = Strata::SetTileParameters(GraphBuilder, View, TileType, PrimitiveType);
+
+				GraphBuilder.AddPass(
+					RDG_EVENT_NAME(
+						"DiffuseIndirectComposite(%s%s%s%s)(%s) %dx%d",
+						DiffuseIndirectSampling,
+						PermutationVector.Get<FDiffuseIndirectCompositePS::FUpscaleDiffuseIndirectDim>() ? TEXT(" UpscaleDiffuseIndirect") : TEXT(""),
+						AmbientOcclusionMask ? TEXT(" ApplyAOToSceneColor") : TEXT(""),
+						PassParameters->ApplyAOToDynamicDiffuseIndirect > 0.0f ? TEXT(" ApplyAOToDynamicDiffuseIndirect") : TEXT(""),
+						ToString(TileType),
+						View.ViewRect.Width(), View.ViewRect.Height()),
+					PassParameters,
+					ERDGPassFlags::Raster,
+					[&View, TileVertexShader, PixelShader, PassParameters, TileType, BlendState, PrimitiveType](FRHICommandList& RHICmdList)
+				{
+					FGraphicsPipelineStateInitializer GraphicsPSOInit;
+					RHICmdList.ApplyCachedRenderTargets(GraphicsPSOInit);
+					// Set the device viewport for the view.
+					RHICmdList.SetViewport(View.ViewRect.Min.X, View.ViewRect.Min.Y, 0.0f, View.ViewRect.Max.X, View.ViewRect.Max.Y, 1.0f);
+
+					GraphicsPSOInit.BlendState = BlendState;
+					GraphicsPSOInit.PrimitiveType = PrimitiveType;
+					GraphicsPSOInit.bDepthBounds = false;
+					GraphicsPSOInit.RasterizerState = TStaticRasterizerState<FM_Solid, CM_None>::GetRHI();
+					GraphicsPSOInit.BoundShaderState.VertexDeclarationRHI = GFilterVertexDeclaration.VertexDeclarationRHI;
+					GraphicsPSOInit.BoundShaderState.VertexShaderRHI = TileVertexShader.GetVertexShader();
+					GraphicsPSOInit.BoundShaderState.PixelShaderRHI = PixelShader.GetPixelShader();		
+					GraphicsPSOInit.DepthStencilState = TStaticDepthStencilState<false, CF_Always>::GetRHI();
+					SetGraphicsPipelineState(RHICmdList, GraphicsPSOInit, 0x0);
+					SetShaderParameters(RHICmdList, PixelShader, PixelShader.GetPixelShader(), *PassParameters);
+					SetShaderParameters(RHICmdList, TileVertexShader, TileVertexShader.GetVertexShader(), PassParameters->StrataTile);
+					RHICmdList.DrawPrimitiveIndirect(PassParameters->StrataTile.TileIndirectBuffer->GetIndirectRHICallBuffer(), 0);
+				});
+			}
+
+		}; // ApplyDiffuseIndirect
+
+		if (bApplyDiffuseIndirect)
+		{
+			if (Strata::IsStrataEnabled() && ViewPipelineState.DiffuseIndirectMethod == EDiffuseIndirectMethod::Lumen)
+			{
+				ApplyDiffuseIndirect(EStrataTileMaterialType::ESimple);
+				ApplyDiffuseIndirect(EStrataTileMaterialType::ESingle);
+				ApplyDiffuseIndirect(EStrataTileMaterialType::EComplex);
+			}
+			else
+			{
+				ApplyDiffuseIndirect(EStrataTileMaterialType::ECount);
+			}
+		}
 
 		// Apply the ambient cubemaps
 		if (IsAmbientCubemapPassRequired(View) && !bIsVisualizePass && !ViewPipelineState.bUseLumenProbeHierarchy)
