VulkanRHI:  Fix occasional layout error in VT depending on which pass runs first after it's resized.

diff --git a/Engine/Source/Runtime/Renderer/Private/SceneHitProxyRendering.cpp b/Engine/Source/Runtime/Renderer/Private/SceneHitProxyRendering.cpp
--- a/Engine/Source/Runtime/Renderer/Private/SceneHitProxyRendering.cpp
+++ b/Engine/Source/Runtime/Renderer/Private/SceneHitProxyRendering.cpp
@@ -180,6 +180,8 @@ void InitHitProxyRender(FRDGBuilder& GraphBuilder, const FSceneRenderer* SceneRe
 	{
 		FVirtualTextureSystem::Get().AllocateResources(GraphBuilder, FeatureLevel);
 		FVirtualTextureSystem::Get().CallPendingCallbacks();
+		// Because there is no Update(), we need to manually finalize the resources
+		FVirtualTextureSystem::Get().FinalizeResources(GraphBuilder, FeatureLevel);
 	}
 
 	// Initialize global system textures (pass-through if already initialized).
diff --git a/Engine/Source/Runtime/Renderer/Private/VT/VirtualTextureSpace.cpp b/Engine/Source/Runtime/Renderer/Private/VT/VirtualTextureSpace.cpp
--- a/Engine/Source/Runtime/Renderer/Private/VT/VirtualTextureSpace.cpp
+++ b/Engine/Source/Runtime/Renderer/Private/VT/VirtualTextureSpace.cpp
@@ -333,6 +333,11 @@ void FVirtualTextureSpace::AllocateTextures(FRDGBuilder& GraphBuilder)
 
 void FVirtualTextureSpace::ApplyUpdates(FVirtualTextureSystem* System, FRDGBuilder& GraphBuilder)
 {
+	ON_SCOPE_EXIT
+	{
+		FinalizeTextures(GraphBuilder);
+	};
+
 	static TArray<FPageTableUpdate> ExpandedUpdates[VIRTUALTEXTURE_SPACE_MAXLAYERS][16];
 
 	if (bNeedToAllocatePageTable)
@@ -435,9 +440,6 @@ void FVirtualTextureSpace::ApplyUpdates(FVirtualTextureSystem* System, FRDGBuild
 	}
 	check(VertexShader.IsValid());
 
-	FRDGTextureAccessArray PageTableTextures;
-	PageTableTextures.SetNum(FMath::DivideAndRoundUp<uint32>(Description.NumPageTableLayers, LayersPerPageTableTexture));
-
 	uint32 FirstUpdate = 0;
 	for (uint32 LayerIndex = 0u; LayerIndex < Description.NumPageTableLayers; ++LayerIndex)
 	{
@@ -446,7 +448,6 @@ void FVirtualTextureSpace::ApplyUpdates(FVirtualTextureSystem* System, FRDGBuild
 
 		FTextureEntry& PageTableEntry = PageTable[TextureIndex];
 		FRDGTextureRef PageTableTexture = GraphBuilder.RegisterExternalTexture(PageTableEntry.RenderTarget);
-		PageTableTextures[TextureIndex] = FRDGTextureAccess(PageTableTexture, ERHIAccess::SRVMask);
 
 		// Use color write mask to update the proper page table entry for this layer
 		FRHIBlendState* BlendStateRHI = nullptr;
@@ -529,10 +530,23 @@ void FVirtualTextureSpace::ApplyUpdates(FVirtualTextureSystem* System, FRDGBuild
 
 		PageTableEntry.RenderTarget = GraphBuilder.ConvertToExternalTexture(PageTableTexture);
 	}
+}
 
-	GraphBuilder.FinalizeTextureAccess(MoveTemp(PageTableTextures));
+void FVirtualTextureSpace::FinalizeTextures(FRDGBuilder& GraphBuilder)
+{
+	for (uint32 LayerIndex = 0u; LayerIndex < Description.NumPageTableLayers; ++LayerIndex)
+	{
+		const uint32 TextureIndex = LayerIndex / LayersPerPageTableTexture;
+		FTextureEntry& PageTableEntry = PageTable[TextureIndex];
+		if (PageTableEntry.RenderTarget)
+		{
+			FRDGTextureRef PageTableTexture = GraphBuilder.RegisterExternalTexture(PageTableEntry.RenderTarget);
+			GraphBuilder.FinalizeTextureAccess(PageTableTexture, ERHIAccess::SRVMask);
+		}
+	}
 }
 
+
 void FVirtualTextureSpace::DumpToConsole(bool verbose)
 {
 	UE_LOG(LogConsoleResponse, Display, TEXT("-= Space ID %i =-"), ID);
diff --git a/Engine/Source/Runtime/Renderer/Private/VT/VirtualTextureSpace.h b/Engine/Source/Runtime/Renderer/Private/VT/VirtualTextureSpace.h
--- a/Engine/Source/Runtime/Renderer/Private/VT/VirtualTextureSpace.h
+++ b/Engine/Source/Runtime/Renderer/Private/VT/VirtualTextureSpace.h
@@ -92,6 +92,7 @@ public:
 
 	void				QueueUpdate( uint8 Layer, uint8 vLogSize, uint32 vAddress, uint8 vLevel, const FPhysicalTileLocation& pTileLocation);
 	void				AllocateTextures(FRDGBuilder& GraphBuilder);
+	void				FinalizeTextures(FRDGBuilder& GraphBuilder);
 	void				ApplyUpdates(FVirtualTextureSystem* System, FRDGBuilder& GraphBuilder);
 	void				QueueUpdateEntirePageTable();
 
diff --git a/Engine/Source/Runtime/Renderer/Private/VT/VirtualTextureSystem.cpp b/Engine/Source/Runtime/Renderer/Private/VT/VirtualTextureSystem.cpp
--- a/Engine/Source/Runtime/Renderer/Private/VT/VirtualTextureSystem.cpp
+++ b/Engine/Source/Runtime/Renderer/Private/VT/VirtualTextureSystem.cpp
@@ -2614,6 +2614,19 @@ void FVirtualTextureSystem::AllocateResources(FRDGBuilder& GraphBuilder, ERHIFea
 	}
 }
 
+void FVirtualTextureSystem::FinalizeResources(FRDGBuilder& GraphBuilder, ERHIFeatureLevel::Type FeatureLevel)
+{
+	LLM_SCOPE(ELLMTag::VirtualTextureSystem);
+
+	for (uint32 ID = 0; ID < MaxSpaces; ID++)
+	{
+		if (Spaces[ID])
+		{
+			Spaces[ID]->FinalizeTextures(GraphBuilder);
+		}
+	}
+}
+
 void FVirtualTextureSystem::CallPendingCallbacks()
 {
 	Producers.CallPendingCallbacks();
diff --git a/Engine/Source/Runtime/Renderer/Private/VT/VirtualTextureSystem.h b/Engine/Source/Runtime/Renderer/Private/VT/VirtualTextureSystem.h
--- a/Engine/Source/Runtime/Renderer/Private/VT/VirtualTextureSystem.h
+++ b/Engine/Source/Runtime/Renderer/Private/VT/VirtualTextureSystem.h
@@ -41,6 +41,7 @@ public:
 	uint32 GetFrame() const { return Frame; }
 
 	void AllocateResources(FRDGBuilder& GraphBuilder, ERHIFeatureLevel::Type FeatureLevel);
+	void FinalizeResources(FRDGBuilder& GraphBuilder, ERHIFeatureLevel::Type FeatureLevel);
 	void CallPendingCallbacks();
 	void Update( FRDGBuilder& GraphBuilder, ERHIFeatureLevel::Type FeatureLevel, FScene* Scene);
 	void ReleasePendingResources();
