Change octree extent from WORLD_MAX to HALF_WORLD_MAX.  The value was previously reverted from LWC_WORLD_MAX to WORLD_MAX, when in fact the original value should have been HALF_WORLD_MAX

diff --git a/Engine/Source/Runtime/Renderer/Private/RendererScene.cpp b/Engine/Source/Runtime/Renderer/Private/RendererScene.cpp
--- a/Engine/Source/Runtime/Renderer/Private/RendererScene.cpp
+++ b/Engine/Source/Runtime/Renderer/Private/RendererScene.cpp
@@ -1168,8 +1168,8 @@ FScene::FScene(UWorld* InWorld, bool bInRequiresHitProxies, bool bInIsEditorScen
 ,	SkyAtmosphere(NULL)
 ,	VolumetricCloud(NULL)
 ,	PrecomputedVisibilityHandler(NULL)
-,	LocalShadowCastingLightOctree(FVector::ZeroVector, WORLD_MAX)
-,	PrimitiveOctree(FVector::ZeroVector, WORLD_MAX)
+,	LocalShadowCastingLightOctree(FVector::ZeroVector, HALF_WORLD_MAX)
+,	PrimitiveOctree(FVector::ZeroVector, HALF_WORLD_MAX)
 ,	bRequiresHitProxies(bInRequiresHitProxies)
 ,	bIsEditorScene(bInIsEditorScene)
 ,	NumUncachedStaticLightingInteractions(0)
