Allow update of RVT locked root pages outside of GPUScene update.
Fixes issue with HLOD commandlet where we don't render the scene, but are updating the VT system.
We could allow this for non-root pages too, but there is a risk that when multiple scenes are in flight we would hit the slower path too regularly.

diff --git a/Engine/Source/Runtime/Renderer/Private/VT/RuntimeVirtualTextureProducer.cpp b/Engine/Source/Runtime/Renderer/Private/VT/RuntimeVirtualTextureProducer.cpp
--- a/Engine/Source/Runtime/Renderer/Private/VT/RuntimeVirtualTextureProducer.cpp
+++ b/Engine/Source/Runtime/Renderer/Private/VT/RuntimeVirtualTextureProducer.cpp
@@ -162,9 +162,7 @@ FVTRequestPageResult FRuntimeVirtualTextureProducer::RequestPageData(
 
 	FVTRequestPageResult result;
 	result.Handle = 0;
-	//todo[vt]:
-	// Returning Saturated instead of Pending here because higher level ignores Pending for locked pages. Need to fix that...
-	result.Status = Finalizer.IsReady() ? EVTRequestPageStatus::Available : EVTRequestPageStatus::Saturated;
+	result.Status = Finalizer.IsReady() ? EVTRequestPageStatus::Available : EVTRequestPageStatus::Pending;
 	return result;
 }
 
diff --git a/Engine/Source/Runtime/Renderer/Private/VT/RuntimeVirtualTextureRender.cpp b/Engine/Source/Runtime/Renderer/Private/VT/RuntimeVirtualTextureRender.cpp
--- a/Engine/Source/Runtime/Renderer/Private/VT/RuntimeVirtualTextureRender.cpp
+++ b/Engine/Source/Runtime/Renderer/Private/VT/RuntimeVirtualTextureRender.cpp
@@ -1244,7 +1244,7 @@ namespace RuntimeVirtualTexture
 		}
 	}
 
-	void RenderPages(FRDGBuilder& GraphBuilder, FRenderPageBatchDesc const& InDesc)
+	void RenderPagesInternal(FRDGBuilder& GraphBuilder, FRenderPageBatchDesc const& InDesc)
 	{
 		check(InDesc.NumPageDescs <= EMaxRenderPageBatch);
 
@@ -1287,7 +1287,21 @@ namespace RuntimeVirtualTexture
 		// Call to let GPU-Scene determine if it is active and record scene primitive count
 		FGPUSceneScopeBeginEndHelper GPUSceneScopeBeginEndHelper(InDesc.Scene->GPUScene, GPUSceneDynamicContext, InDesc.Scene);
 		InDesc.Scene->GPUScene.Update(GraphBuilder, *InDesc.Scene);
-		RenderPages(GraphBuilder, InDesc);
+		RenderPagesInternal(GraphBuilder, InDesc);
+	}
+
+	void RenderPages(FRDGBuilder& GraphBuilder, FRenderPageBatchDesc const& InDesc)
+	{
+		if (InDesc.Scene->GPUScene.IsRendering())
+		{
+			RenderPagesInternal(GraphBuilder, InDesc);
+		}
+		else
+		{
+			// We allow locked root pages to be rendered outside of their scene update.
+			// We expect to hit this path very rarely. (One case is during material baking.)
+			RenderPagesStandAlone(GraphBuilder, InDesc);
+		}
 	}
 
 	void RenderPages(FRHICommandListImmediate& RHICmdList, FRenderPageBatchDesc const& InDesc)
