Fix Volumetric fog rendered black in Forward (UE-141868).

diff --git a/Engine/Source/Runtime/Renderer/Private/DeferredShadingRenderer.cpp b/Engine/Source/Runtime/Renderer/Private/DeferredShadingRenderer.cpp
--- a/Engine/Source/Runtime/Renderer/Private/DeferredShadingRenderer.cpp
+++ b/Engine/Source/Runtime/Renderer/Private/DeferredShadingRenderer.cpp
@@ -2610,6 +2610,9 @@ void FDeferredShadingSceneRenderer::Render(FRDGBuilder& GraphBuilder)
 		}
 
 		RenderForwardShadowProjections(GraphBuilder, SceneTextures, ForwardScreenSpaceShadowMaskTexture, ForwardScreenSpaceShadowMaskHairTexture);
+
+		// With forward shading we need to render volumetric fog before the base pass
+		ComputeVolumetricFog(GraphBuilder, SceneTextures);
 	}
 
 	FDBufferTextures DBufferTextures = CreateDBufferTextures(GraphBuilder, SceneTextures.Config.Extent, ShaderPlatform);
@@ -2780,7 +2783,11 @@ void FDeferredShadingSceneRenderer::Render(FRDGBuilder& GraphBuilder)
 			RenderLumenSceneLighting(GraphBuilder, Views[0], LumenFrameTemporaries);
 		}
 
-		ComputeVolumetricFog(GraphBuilder, SceneTextures);
+		// If forward shading is enabled, we computed fog earlier already
+		if (!IsForwardShadingEnabled(ShaderPlatform))
+		{
+			ComputeVolumetricFog(GraphBuilder, SceneTextures);
+		}
 	}
 	// End shadow and fog after base pass
 
