Tiled Compute Rasterization for Hair Visibility - currently:
- doesn't use 64 bit Atomics (this is planned to be re-added in next iteration)
- supports only one spatial sample with experimental support for temporal samples
- includes a workaround for artifacts related to compaction/material/lighting pass tiling

diff --git a/Engine/Shaders/Private/HairStrands/HairStrandsVisibilityCompactionComputeRaster.usf b/Engine/Shaders/Private/HairStrands/HairStrandsVisibilityCompactionComputeRaster.usf
--- a/Engine/Shaders/Private/HairStrands/HairStrandsVisibilityCompactionComputeRaster.usf
+++ b/Engine/Shaders/Private/HairStrands/HairStrandsVisibilityCompactionComputeRaster.usf
@@ -16,6 +16,8 @@
 
 ///////////////////////////////////////////////////////////////////////////////////////////////////////////////
 
+//TODO: re-enable support for 64 bit visibilty buffer format?
+/*
 #if COMPILER_SUPPORTS_UINT64_IMAGE_ATOMICS && COMPILER_SUPPORTS_ULONG_TYPES
 #define PackedType UlongType
 uint2 UnpackData(PackedType In)
@@ -29,6 +31,7 @@ uint2 UnpackData(PackedType In)
 	return In;
 }
 #endif
+*/
 
 ///////////////////////////////////////////////////////////////////////////////////////////////////////////////
 
@@ -58,13 +61,9 @@ struct FSampleSetDesc
 // Compute raster visibility buffer
 ///////////////////////////////////////////////////////////////////////////////////////////////////////////////
 
-Texture2D<PackedType> VisibilityTexture0;
-Texture2D<PackedType> VisibilityTexture1;
-Texture2D<PackedType> VisibilityTexture2;
-Texture2D<PackedType> VisibilityTexture3;
-
-Texture2D<uint>  VisibilityHairCount;
-Texture2D<float> ViewTransmittanceTexture;
+Texture2D<uint>  DepthCovTexture;
+Texture2D<uint>  PrimMatTexture;
+Texture2D<uint>  HairCountTexture;
 
 #if PERMUTATION_TILE
 int2 TileCountXY;
@@ -106,87 +105,26 @@ void MainCS(
 
 	const float OpaqueDepth = SceneTexturesStruct.SceneDepthTexture.Load(uint3(PixelCoord, 0)).r;
 
-	//const uint HairCountAndRadius = VisibilityHairCount.Load(uint3(PixelCoord, 0)); // Contains both hair count and the radius;
-	//const uint HairRadius = 0; // TODO
-	//const uint HairCount  = HairCountAndRadius;
-	const uint2 VisData0 = UnpackData(VisibilityTexture0.Load(uint3(PixelCoord, 0)));
-	const uint2 VisData1 = UnpackData(VisibilityTexture1.Load(uint3(PixelCoord, 0)));
-	const uint2 VisData2 = UnpackData(VisibilityTexture2.Load(uint3(PixelCoord, 0)));
-	const uint2 VisData3 = UnpackData(VisibilityTexture3.Load(uint3(PixelCoord, 0)));
+
+	const uint VisData0 = DepthCovTexture.Load(uint3(PixelCoord, 0));
+	const uint VisData1 = PrimMatTexture.Load(uint3(PixelCoord, 0));
 
 	uint2  Vis_PrimId_MatId[4];
 	float2 Vis_Depth_Cov[4];
 	float  TotalCoverage = 0;
 	uint   UniqueSampleCount = 0;
 
-	uint CurrentPrimitiveId = 0;
-	const float DepthThres = OpaqueDepth;
-	if (VisData0.y)
+	if (VisData1)
 	{
-		const uint VisSamplePrimitiveId = UnpackHairVisPrimitiveId(VisData0.x);
-		const uint VisSampleMaterialId  = UnpackHairVisMaterialId(VisData0.x);
-		const float2 Depth_Cov = float2(UnpackHairVisDepth(VisData0.y), UnpackHairVisCoverage(VisData0.y));
+		const uint VisSamplePrimitiveId = UnpackHairVisPrimitiveId(VisData1);
+		const uint VisSampleMaterialId = UnpackHairVisMaterialId(VisData1);
+		const float2 Depth_Cov = float2(UnpackHairVisDepth(VisData0), UnpackHairVisCoverage(VisData0));
 
-		Vis_PrimId_MatId[UniqueSampleCount]	= uint2(VisSamplePrimitiveId, VisSampleMaterialId);
-		Vis_Depth_Cov[UniqueSampleCount]	= Depth_Cov;
+		Vis_PrimId_MatId[UniqueSampleCount] = uint2(VisSamplePrimitiveId, VisSampleMaterialId);
+		Vis_Depth_Cov[UniqueSampleCount] = Depth_Cov;
 
-		if (Depth_Cov.x > DepthThres && VisSamplePrimitiveId != CurrentPrimitiveId)
-		{
-			TotalCoverage += Vis_Depth_Cov[UniqueSampleCount].y;
-			CurrentPrimitiveId = VisSamplePrimitiveId;
-			UniqueSampleCount++;
-		}
-	}
-
-	if (VisData1.y && SamplerPerPixel > 1)
-	{
-		const uint VisSamplePrimitiveId = UnpackHairVisPrimitiveId(VisData1.x);
-		const uint VisSampleMaterialId  = UnpackHairVisMaterialId(VisData1.x);
-		const float2 Depth_Cov			= float2(UnpackHairVisDepth(VisData1.y), UnpackHairVisCoverage(VisData1.y));
-
-		Vis_PrimId_MatId[UniqueSampleCount]	= uint2(VisSamplePrimitiveId, VisSampleMaterialId);
-		Vis_Depth_Cov[UniqueSampleCount]	= Depth_Cov;
-		
-		if (Depth_Cov.x > DepthThres && VisSamplePrimitiveId != CurrentPrimitiveId)
-		{
-			TotalCoverage += Vis_Depth_Cov[UniqueSampleCount].y;
-			CurrentPrimitiveId = VisSamplePrimitiveId;
-			UniqueSampleCount++;
-		}
-	}
-
-	if (VisData2.y && SamplerPerPixel > 2)
-	{
-		const uint VisSamplePrimitiveId = UnpackHairVisPrimitiveId(VisData2.x);
-		const uint VisSampleMaterialId  = UnpackHairVisMaterialId(VisData2.x);
-		const float2 Depth_Cov			= float2(UnpackHairVisDepth(VisData2.y), UnpackHairVisCoverage(VisData2.y));
-
-		Vis_PrimId_MatId[UniqueSampleCount]	= uint2(VisSamplePrimitiveId, VisSampleMaterialId);
-		Vis_Depth_Cov[UniqueSampleCount]	= Depth_Cov;
-		
-		if (Depth_Cov.x > DepthThres && VisSamplePrimitiveId != CurrentPrimitiveId)
-		{
-			TotalCoverage += Vis_Depth_Cov[UniqueSampleCount].y;
-			CurrentPrimitiveId = VisSamplePrimitiveId;
-			UniqueSampleCount++;
-		}
-	}
-
-	if (VisData3.y && SamplerPerPixel > 3)
-	{
-		const uint VisSamplePrimitiveId = UnpackHairVisPrimitiveId(VisData3.x);
-		const uint VisSampleMaterialId  = UnpackHairVisMaterialId(VisData3.x);
-		const float2 Depth_Cov			= float2(UnpackHairVisDepth(VisData3.y), UnpackHairVisCoverage(VisData3.y));
-
-		Vis_PrimId_MatId[UniqueSampleCount]	 = uint2(VisSamplePrimitiveId, VisSampleMaterialId);
-		Vis_Depth_Cov[UniqueSampleCount]	 = Depth_Cov;
-		
-		if (Depth_Cov.x > DepthThres && VisSamplePrimitiveId != CurrentPrimitiveId)
-		{
-			TotalCoverage += Vis_Depth_Cov[UniqueSampleCount].y;
-			CurrentPrimitiveId = VisSamplePrimitiveId;
-			UniqueSampleCount++;
-		}
+		TotalCoverage += Vis_Depth_Cov[UniqueSampleCount].y;
+		UniqueSampleCount++;
 	}
 
 	FNodeDesc NodeDesc;
@@ -220,14 +158,20 @@ void MainCS(
 				OutNodeVis.PrimitiveId	= Vis_PrimId_MatId[OutIndex].x;
 				OutNodeVis.MaterialId	= Vis_PrimId_MatId[OutIndex].y;
 				OutNodeVis.Depth		= Vis_Depth_Cov[OutIndex].x;
-				OutNodeVis.Coverage8bit = To8bitCoverage(Vis_Depth_Cov[OutIndex].y / TotalCoverage); 
-
+				OutNodeVis.Coverage8bit = 0xff;
+				//OutNodeVis.Coverage8bit = To8bitCoverage(Vis_Depth_Cov[OutIndex].y / TotalCoverage); 
 				const uint StoreIndex = NodeDesc.Offset + OutIndex;
 				OutCompactNodeVis[StoreIndex] = PackHairVis(OutNodeVis);
 				OutCompactNodeCoord[StoreIndex] = PixelCoord;
 			}
 
-			PixelCoverage = TransmittanceToCoverage(ViewTransmittanceTexture.Load(uint3(PixelCoord, 0)), CoverageThreshold);
+			// skips separate count to transmittance pass  
+			float LogViewTransmittance = -(HairCountTexture.Load(uint3(PixelCoord, 0)) / 1000.0);
+			float ViewTransmittance = pow(2.71828, LogViewTransmittance);
+			float CoverageThreshold = 0.98;
+			PixelCoverage = saturate(min(1.0 - ViewTransmittance, 1.0) / CoverageThreshold);
+
+//			PixelCoverage = TransmittanceToCoverage(ViewTransmittanceTexture.Load(uint3(PixelCoord, 0)), CoverageThreshold);
 		}
 	}
 
diff --git a/Engine/Shaders/Private/HairStrands/HairStrandsVisibilityRasterCompute.usf b/Engine/Shaders/Private/HairStrands/HairStrandsVisibilityRasterCompute.usf
--- a/Engine/Shaders/Private/HairStrands/HairStrandsVisibilityRasterCompute.usf
+++ b/Engine/Shaders/Private/HairStrands/HairStrandsVisibilityRasterCompute.usf
@@ -1,72 +1,71 @@
 // Copyright Epic Games, Inc. All Rights Reserved.
 
 #include "../Common.ush"
-#include "../MonteCarlo.ush"
 #include "HairStrandsClusterCommon.ush"
 #include "HairStrandsVertexFactoryCommon.ush"
 #include "HairStrandsVisibilityCommon.ush"
 
 ///////////////////////////////////////////////////////////////////////////
-// 64 bits atomic
-#define RASTER_PLATFORM 0
-#define RASTER_NVIDIA 1
-#define RASTER_AMD_DX11 2
-#define RASTER_AMD_DX12 3
+// 
+// Wave size
+#if SHADER_RASTERCOMPUTE_DEPTH_GRID == 0
 
-// Group size
 #if PERMUTATION_GROUP_SIZE == 64
-	#define GROUP_SIZE 64
+#define WAVE_SIZE 32
 #elif PERMUTATION_GROUP_SIZE == 32
-	#define GROUP_SIZE 32
+#define WAVE_SIZE 32
 #else
-	#error Unknown group size	
+#error Unknown group size	
 #endif
 
-// Add permutation for AMD/other
-#if PERMUTATION_RASTER_ATOMIC == RASTER_NVIDIA
-#define NV_SHADER_EXTN_SLOT u7
-#include "/Engine/Shared/ThirdParty/NVIDIA/nvHLSLExtns.h"
-#elif PERMUTATION_RASTER_ATOMIC == RASTER_AMD_DX11
-#define AmdDxExtShaderIntrinsicsUAVSlot u7
-#include "/Engine/Shared/ThirdParty/AMD/ags_shader_intrinsics_dx11.h"
-#elif PERMUTATION_RASTER_ATOMIC == RASTER_AMD_DX12
-#include "/Engine/Shared/ThirdParty/AMD/ags_shader_intrinsics_dx12.h"
-#endif
+#endif //SHADER_RASTERCOMPUTE_DEPTH_GRID
+
+///////////////////////////////////////////////////////////////////////////
 
-#if PERMUTATION_RASTER_ATOMIC == RASTER_PLATFORM && COMPILER_SUPPORTS_UINT64_IMAGE_ATOMICS && COMPILER_SUPPORTS_ULONG_TYPES
-#define PackedType UlongType
-PackedType ToPackedType(uint2 Data)
+uint TileSize;
+float RcpTileSize;
+uint SqrTileSize;
+
+uint NumBinners;
+float RcpNumBinners;
+uint NumRasterizers;
+float RcpNumRasterizers;
+
+/* 
+// use untyped buffer for segment tiles to reduce VGPR usage - 16 bytes
+struct FVisTile
 {
-	return PackUlongType(Data);
-}
-#else
-#define PackedType uint2
-PackedType ToPackedType(uint2 Data)
+	uint PrimOffset;
+	uint PrimCount;
+	uint TileCoord;
+	uint MinDepth;
+};
+*/
+#define VT_SIZE 4
+
+#define VT_PrimOffset 0
+#define VT_PrimCount 1
+#define VT_Coord 2
+#define VT_NearestDepth 3
+
+uint PackVisTileCoord(uint2 Coord)
 {
-	return Data;
+	return uint(((Coord.x & 0xff) << 0) | (((Coord.y) & 0xff) << 8));
 }
-#endif
 
-void WriteVisibilityData(RWTexture2D<PackedType> OutTexture, uint2 Coords, PackedType Data)
+uint2 UppackVisTileCoord(uint Packed)
 {
-#if PERMUTATION_RASTER_ATOMIC == RASTER_NVIDIA
-	NvInterlockedMaxUint64(OutTexture, Coords, Data);
-#elif PERMUTATION_RASTER_ATOMIC == RASTER_AMD_DX11
-	AmdDxExtShaderIntrinsics_AtomicMaxU64(OutTexture, Coords, Data);
-#elif PERMUTATION_RASTER_ATOMIC == RASTER_AMD_DX12
-	AmdExtD3DShaderIntrinsics_AtomicMaxU64(OutTexture, Coords, Data);
-#elif PERMUTATION_RASTER_ATOMIC == RASTER_PLATFORM && COMPILER_SUPPORTS_UINT64_IMAGE_ATOMICS && COMPILER_SUPPORTS_ULONG_TYPES
-	ImageInterlockedMaxUInt64(OutTexture, Coords, Data);
-#endif
+	return uint2(((Packed >> 0) & 0xff), ((Packed >> 8) & 0xff));
 }
 
 ///////////////////////////////////////////////////////////////////////////
-uint			DispatchCountX;
+
 uint			MaxRasterCount;
 uint			FrameIdMod8;
 uint			MacroGroupId;
 uint			ResolutionMultiplier;
 int2			OutputResolution;
+float2			OutputResolutionf;
 uint			HairMaterialId;
 uint			HairStrandsVF_bIsCullingEnable;
 float			HairStrandsVF_Radius;
@@ -81,29 +80,10 @@ float4x4		HairStrandsVF_LocalToWorldPrimitiveTransform;
 Buffer<uint4>	HairStrandsVF_PositionBuffer;
 Buffer<float4>	HairStrandsVF_PositionOffsetBuffer;
 
-#if PERMUTATION_CULLING
-Buffer<uint>	HairStrandsVF_CullingIndirectBuffer;
-Buffer<uint>	HairStrandsVF_CullingIndexBuffer;
-Buffer<float>	HairStrandsVF_CullingRadiusScaleBuffer;
-Buffer<uint>	IndirectBufferArgs;
-#endif
-
 Texture2D<float> SceneDepthTexture;
-
-RWTexture2D<uint>  OutHairCountTexture;
-
-RWTexture2D<PackedType> OutVisibilityTexture0;
-#if PERMUTATION_SPP == 2 || PERMUTATION_SPP == 4
-RWTexture2D<PackedType> OutVisibilityTexture1;
-#endif
-#if PERMUTATION_SPP == 4
-RWTexture2D<PackedType> OutVisibilityTexture2;
-RWTexture2D<PackedType> OutVisibilityTexture3;
-#endif
-
-
-///////////////////////////////////////////////////////////////////////////
-// Helping functions 
+float			TileLODScale;
+int2			TileRes;
+float			TileBaseDepth;
 
 float3 NDCToPixelCoord(float4 InDC)
 {
@@ -112,340 +92,565 @@ float3 NDCToPixelCoord(float4 InDC)
 	return float3(UV * OutputResolution, NDC.z);
 }
 
-float ToCoverage(FHairControlPoint CP, float4 NDC, float RadiusAtDepth1Primary)
+void CalcScreenPos(in uint4 PB, in float3 PBO, out float3 SP, out uint Type)
 {
-	const float SceneDepth = NDC.w;
-	const float PixelRadius = SceneDepth * RadiusAtDepth1Primary;
-	return CP.WorldRadius / max(CP.WorldRadius, PixelRadius);
-}
+	const FHairControlPoint CP = UnpackHairControlPoint(
+		PB,
+		PBO,
+		HairStrandsVF_Radius,
+		HairStrandsVF_RootScale,
+		HairStrandsVF_TipScale);
 
-void Swap(inout float V0, inout float V1)
-{
-	float T = V0;
-	V0 = V1;
-	V1 = T;
-}
+	const float3 WP = mul(float4(CP.Position, 1), HairStrandsVF_LocalToWorldPrimitiveTransform).xyz;
+	float4 NDC = mul(float4(WP, 1), LWCHackToFloat(PrimaryView.WorldToClip));
 
-void Swap(inout float3 V0, inout float3 V1)
-{
-	float3 T = V0;
-	V0 = V1;
-	V1 = T;
-}
+	SP = NDCToPixelCoord(NDC);
 
-bool IsVisible(float3 SP)
-{
-	const float DeviceZ = SP.z;
-	const uint2 SceneCoord = SP.xy / ResolutionMultiplier;
-	return DeviceZ > SceneDepthTexture.Load(uint3(SceneCoord, 0)); // InverseZ
+	Type = CP.Type;
 }
 
-bool IsVisible(int2 PixelCoords, float HairDeviceZ)
+void CalcScreenPosAndRad(in uint4 PB, in float3 PBO, out float4 SP, out float Rad)
 {
-	const float DeviceZ = SceneDepthTexture.Load(uint3(PixelCoords, 0)); // InverseZ
-	return HairDeviceZ >= DeviceZ;
-}
+	const FHairControlPoint CP = UnpackHairControlPoint(
+		PB,
+		PBO,
+		HairStrandsVF_Radius,
+		HairStrandsVF_RootScale,
+		HairStrandsVF_TipScale);
 
-bool IsInScreen(int2 C)
-{
-	return C.x >= 0 && C.x < OutputResolution.x && C.y >= 0 && C.y < OutputResolution.y;
+	const float3 WP = mul(float4(CP.Position, 1), HairStrandsVF_LocalToWorldPrimitiveTransform).xyz;
+	float4 NDC = mul(float4(WP, 1), LWCHackToFloat(PrimaryView.WorldToClip));
+
+	SP = float4(NDCToPixelCoord(NDC).xy, NDC.zw);
+
+	Rad = CP.WorldRadius * OutputResolutionf.x;
 }
 
-void InternalRasterize(int2 Coords, float Ratio, float Depth, uint PrimitiveId, float Coverage)
+///////////////////////////////////////////////////////////////////////////
+
+#if SHADER_RASTERCOMPUTE_DEPTH_GRID
+
+RWTexture2D<uint>	OutVisTileDepthGrid;
+RWTexture2D<uint>	OutDepthCovTexture;
+
+groupshared uint group_FurthestDepth; // (4 bytes)
+
+[numthreads(1024, 1, 1)]
+void PrepareDepthGridCS(uint dId : SV_DispatchThreadID, uint GroupThreadID : SV_GroupThreadID)
 {
-	if (IsInScreen(Coords))
+	uint tw = TileRes.x;
+	uint th = TileRes.y;
+
+	if (GroupThreadID == 0) group_FurthestDepth = 0xffffffff;
+
+	GroupMemoryBarrierWithGroupSync();
+
+	uint t = dId / 1024;
+
+	uint ty = uint((float(t)+0.5) / float(tw));
+	uint tx = t - (ty*tw);
+
+	for (uint p = GroupThreadID; p < SqrTileSize; p += 1024)
 	{
-		// Fine culling per pixel
-		// Fine grain culling is disabled, as it creates some flickering (on Nvidia cards). It is unclear if this is driver bug or another issue.
-		#if 1
-		const bool bIsVisible = IsVisible(Coords, Depth);
-		if (!bIsVisible)
-			return;
-		#endif
-
-		const float PrecisionScale = 1000.f;
-		const float ViewTransmittance = 1 - saturate(Ratio * Coverage);
-		const float LogViewTransmittance = -log2(ViewTransmittance);
-		const uint PackedViewTransmittance = uint(LogViewTransmittance * PrecisionScale);
-		InterlockedAdd(OutHairCountTexture[uint2(Coords)], PackedViewTransmittance);
-
-		uint2 PackedVisibility;
-		PackedVisibility.x = PackHairVisPrimitiveMaterialId(PrimitiveId, HairMaterialId);
-		PackedVisibility.y = PackHairVisDepthCoverage(Depth, Ratio * Coverage);
-
-		PackedType In = ToPackedType(PackedVisibility);
-		PackedType Out = 0;
-
-	#if PERMUTATION_SPP == 1
-		WriteVisibilityData(OutVisibilityTexture0, Coords, In);
-	#endif
-
-	#if PERMUTATION_SPP == 2
-		if (Ratio < 0.5f)
-			WriteVisibilityData(OutVisibilityTexture0, Coords, In);
-		else
-			WriteVisibilityData(OutVisibilityTexture1, Coords, In);
-	#endif
-
-	#if PERMUTATION_SPP == 4
-		if (Ratio < 0.25f)
-			WriteVisibilityData(OutVisibilityTexture0, Coords, In);
-		else if (Ratio < 0.5f)
-			WriteVisibilityData(OutVisibilityTexture1, Coords, In);
-		else if (Ratio < 0.75f)
-			WriteVisibilityData(OutVisibilityTexture2, Coords, In);
-		else
-			WriteVisibilityData(OutVisibilityTexture3, Coords, In);
-	#endif
+		if (p < SqrTileSize)
+		{
+			uint py = (float(p)+0.5) * RcpTileSize;
+			uint px = p - (py*TileSize);
+
+			uint2 Coord = uint2(tx * TileSize + px, ty * TileSize + py);
+
+			if (Coord.x < uint(OutputResolution.x) && Coord.y < uint(OutputResolution.y))
+			{
+				float Depth = SceneDepthTexture.Load(uint3(Coord, 0));
+				uint packedDepth = PackHairVisDepthCoverage(Depth, 1.0);
+
+				InterlockedMin(group_FurthestDepth, packedDepth);
+
+				InterlockedMax(OutDepthCovTexture[Coord], packedDepth);
+			}
+		}
 	}
+
+	GroupMemoryBarrierWithGroupSync();
+
+	OutVisTileDepthGrid[uint2(tx, ty)] = group_FurthestDepth;
 }
 
-//////////////////////////////////////////////////////////////////////////////////////////
-// Queued work
+#endif //SHADER_RASTERCOMPUTE_DEPTH_GRID
 
-struct FRasterInfo
-{
-	bool bIsSteep;
-	uint P0x;
-	uint P1x;
-	float DeltaY;
-	float DeltaZ;
-	float GradientY;
-	float GradientZ;
-};
+///////////////////////////////////////////////////////////////////////////
 
-uint4 PackRasterInfo(FRasterInfo In)
-{
-	uint4 Out;
-	return Out;
-}
+#if SHADER_RASTERCOMPUTE_BINNING
+
+#if PERMUTATION_CULLING
+Buffer<uint>	HairStrandsVF_CullingIndirectBuffer;
+Buffer<uint>	HairStrandsVF_CullingIndexBuffer;
+Buffer<float>	HairStrandsVF_CullingRadiusScaleBuffer;
+Buffer<uint>	IndirectBufferArgs;
+#endif
+
+RWTexture3D<uint> 					OutVisTileBinningGrid;
+Texture2D<uint>						VisTileDepthGrid;
+RWBuffer<uint>						OutVisTilePrims;
+RWBuffer<uint>						OutVisTileArgs;
+RWByteAddressBuffer					OutVisTileData;
 
-FRasterInfo UnpackRasterInfo(uint4 In)
+uint LoadOutVisTileData(uint index, uint offset)
 {
-	return (FRasterInfo)0;
+	return OutVisTileData.Load((((index)) * VT_SIZE * 4) + ((offset) * 4));
 }
 
-float2 StochasticRasterRandom(uint2 PixelPosition)
+void StoreOutVisTileData(uint index, uint offset, uint value)
 {
-	const uint InSeed = View.StateFrameIndexMod8;
-	const uint2 Seed0 = Rand3DPCG16(int3(PixelPosition, InSeed)).xy;
-	return Hammersley16(InSeed, 8, Seed0);
+	OutVisTileData.Store((((index)) * VT_SIZE * 4) + ((offset) * 4), (value));
 }
 
-// Rasterize a given pixel with anti-aliasing. 
-// This can be done by writing the two adjacents pixel covered by the segments, or stochastically, by picking one of the two pixels randomly based on coverage
-uint2 Rasterize(uint2 PixelCoord, bool bIsSteep, float FracY, float RatioScale, float Coverage, float DeviceZ, uint ProvokingVertexId, bool bStochastic=true, uint2 SeedCoord=uint2(0,0))
+groupshared uint group_LoopNum;
+groupshared uint group_BatchNum;
+
+// The total number of line segments (VertexCount) is divided up equally between N binners - each binner = a workgroup which loops through the designated set segments in batches of 1024
+// NB there is still potential to use LDS to prevent/reduce contention when adding segments to the binning grid and this may improve perf
+
+[numthreads(1024, 1, 1)]
+void BinningCS(uint2 DispatchThreadID : SV_DispatchThreadID, uint GroupThreadID : SV_GroupThreadID, uint GroupID : SV_GroupID)
 {
-#if PERMUTATION_STOCHASTIC
-	// Select randomly Bottom Left/Right
+	ResolvedView = ResolveView();
+	if (GroupThreadID == 0)
 	{
-		// No select based on the exact frac amount, as this will oscillate too much. 
-		// The stochastic raster only acts on the visibility check. The final covarage is computed analytically
-		float U = StochasticRasterRandom(SeedCoord).x;
-		const bool bLower = bStochastic ? (U < 0.5 ? FracY < 0.5 : FracY > 0.5) : FracY < 0.5f;
-
-		uint2 Coord = PixelCoord + (bLower ? uint2(0, 0) : uint2(0, 1));
-		Coord = bIsSteep ? Coord.yx : Coord.xy;
-		const float Ratio = (bLower ? 1 - FracY : FracY) * RatioScale;;
-		InternalRasterize(Coord, Ratio * RatioScale, DeviceZ, ProvokingVertexId, Coverage);
-		SeedCoord = Coord;
+		uint VertexCount = HairStrandsVF_VertexCount;
+#if PERMUTATION_CULLING	
+		if (HairStrandsVF_bIsCullingEnable)
+		{
+			VertexCount = HairStrandsVF_CullingIndirectBuffer[3];
+		}
+#endif
+		group_BatchNum = (VertexCount+1023)/1024;
+		group_LoopNum = (group_BatchNum + (NumBinners - 1)) * RcpNumBinners;
 	}
-#else
-	const bool bLower = FracY < 0.5f;
-	const float Ratio = (bLower ? 1 - FracY : FracY) * RatioScale;
 
-	// Bottom/Left Pixel (or inverse)
-	{
-		uint2 Coord = PixelCoord + (bLower ? uint2(0, 0) : uint2(0, 1));
-		Coord = bIsSteep ? Coord.yx : Coord.xy;
-		InternalRasterize(Coord, Ratio, DeviceZ, ProvokingVertexId, Coverage);
-	}
-	// Top/Right Pixel (or inverse)
+	GroupMemoryBarrierWithGroupSync();
+
+	[loop]
+	for (uint LoopIndex = 0; LoopIndex < group_LoopNum; LoopIndex++)
 	{
-		uint2 Coord = PixelCoord + (bLower ? uint2(0, 0) : uint2(0, 1));
-		Coord = bIsSteep ? Coord.yx : Coord.xy;
-		InternalRasterize(Coord, 1.f - Ratio, DeviceZ, ProvokingVertexId, Coverage);
-	}
-#endif
-	return SeedCoord;
-}
-//////////////////////////////////////////////////////////////////////////////////////////
-[numthreads(GROUP_SIZE, 1, 1)]
-void MainCS(uint2 DispatchThreadID : SV_DispatchThreadID)
-{
+		uint BatchIndex = LoopIndex + (GroupID * group_LoopNum);
+
+		bool bSegValid = (BatchIndex < group_BatchNum) ? true : false;
+
+		bool bTileIntersected[3] = { false, false, false };
+
+		int2 TileCoord[3] = { int2(-1, -1), int2(-1, -1), int2(-1, -1) };
+
+		uint SegIndex[3] = { 0xffffffff, 0xffffffff, 0xffffffff };
+
+		uint PrimID = BatchIndex * 1024 + GroupThreadID;
+
 #if PERMUTATION_CULLING
-	uint VertexIndex0 = GetHairStrandsVertexFetchIndex_Culled(DispatchThreadID, GROUP_SIZE);
-#else
-	uint VertexIndex0 = GetHairStrandsVertexFetchIndex(DispatchThreadID, GROUP_SIZE, DispatchCountX);
+		if (bSegValid && HairStrandsVF_bIsCullingEnable)
+		{
+			const uint VertexCountAfterCulling = HairStrandsVF_CullingIndirectBuffer[3];
+			uint FetchIndex0 = PrimID;
+			if (FetchIndex0 >= (VertexCountAfterCulling-1))
+			{
+				bSegValid = false;
+			}
+			else
+			{
+				uint FetchIndex1 = min(FetchIndex0 + 1, VertexCountAfterCulling - 1);
+
+				uint VertexIndex0 = HairStrandsVF_CullingIndexBuffer[FetchIndex0];
+				float LodRadiusScale0 = HairStrandsVF_CullingRadiusScaleBuffer[FetchIndex0];
+
+				uint VertexIndex1 = HairStrandsVF_CullingIndexBuffer[FetchIndex1];
+				float LodRadiusScale1 = HairStrandsVF_CullingRadiusScaleBuffer[FetchIndex1];
+
+				if ((LodRadiusScale0 <= 0.0f) || (VertexIndex1 != VertexIndex0 + 1))
+				{
+					bSegValid = false;
+				}
+				else
+				{
+					PrimID = VertexIndex0;
+				}
+			}
+		}
 #endif
 
-	uint VertexIndex1 = min(VertexIndex0 + 1, HairStrandsVF_VertexCount - 1);
-	if (VertexIndex0 >= HairStrandsVF_VertexCount)
-	{
-		return;
-	}
+		uint BinID = GroupID - (uint(float(GroupID + 0.5) * RcpNumBinners)* NumBinners);
+		uint BinID2 = BinID + NumBinners;
 
-	float LodRadiusScale0 = 1;
-	float LodRadiusScale1 = 1;
-#if PERMUTATION_CULLING
-	if (HairStrandsVF_bIsCullingEnable)
-	{
-		const uint VertexCountAfterCulling = HairStrandsVF_CullingIndirectBuffer[3];
-		uint FetchIndex0 = VertexIndex0;
-		uint FetchIndex1 = VertexIndex1;
-		if (FetchIndex0 >= VertexCountAfterCulling)
+#if PERMUTATION_TEMPORAL_LAYERING
+		//remap linear frame index to improve TAA result (still a bit exprimental)
+		uint remapped = (View.StateFrameIndex) & 0x3;
+		remapped = ((remapped & 0x1) << 1) | ((remapped & 0x2) >> 1);
+
+		if (bSegValid) bSegValid = ((PrimID < HairStrandsVF_VertexCount) && ((PrimID & 0x3) == (remapped))) ? true : false;
+#else
+		if (bSegValid) bSegValid = (PrimID < HairStrandsVF_VertexCount) ? true : false;
+#endif
+		if (bSegValid) //currently hardcoded to a maximum of 3 tile bins - this only works corectly when segment length does not exceed TileSize - TODO: add support for N bins
 		{
-			return;
+			float3 SP0;
+			float3 SP1;
+
+			uint Type;
+			CalcScreenPos(HairStrandsVF_PositionBuffer[PrimID], HairStrandsVF_PositionOffsetBuffer[0].xyz, SP0, Type);
+
+			bool bIsEndCV = (Type == HAIR_CONTROLPOINT_END); 
+
+			if (!bIsEndCV)
+			{
+				CalcScreenPos(HairStrandsVF_PositionBuffer[PrimID + 1], HairStrandsVF_PositionOffsetBuffer[0].xyz, SP1, Type);
+
+				uint NearestDepth = PackHairVisDepthCoverage(max(SP0.z, SP1.z), 1.0);
+
+				TileCoord[0] = int2(SP0.x * RcpTileSize, SP0.y * RcpTileSize);
+
+				if (all(TileCoord[0] >= 0) && all(TileCoord[0] < TileRes) && NearestDepth > VisTileDepthGrid[TileCoord[0]])
+				{
+					InterlockedAdd(OutVisTileBinningGrid[uint3(TileCoord[0], BinID)], 1, SegIndex[0]);
+					bTileIntersected[0] = true;
+				}
+
+				TileCoord[1] = int2(SP1.x * RcpTileSize, SP1.y * RcpTileSize);
+
+				if ( any(TileCoord[1] != TileCoord[0]) && all(TileCoord[1] >= 0) && all(TileCoord[1] < TileRes) && NearestDepth > VisTileDepthGrid[TileCoord[1]])
+				{
+					InterlockedAdd(OutVisTileBinningGrid[uint3(TileCoord[1], BinID)], 1, SegIndex[1]);
+					bTileIntersected[1] = true;
+				}
+
+				if ((bTileIntersected[0] || bTileIntersected[1]) && all(TileCoord[1] != TileCoord[0]))
+				{
+					float2 MP = (TileCoord[0] + float2(0.5, 0.5) + TileCoord[1] + float2(0.5, 0.5)) * 0.5 * 32;
+
+					float2 SPI;
+					SPI.x = SP0.x + (SP1.x - SP0.x) * ((MP.y - SP0.y) / (SP1.y - SP0.y));
+					SPI.y = SP0.y + (SP1.y - SP0.y) * ((MP.x - SP0.x) / (SP1.x - SP0.x));
+
+					TileCoord[2] = int2(SPI.x * RcpTileSize, SPI.y * RcpTileSize);
+
+					if (all(TileCoord[2] >= 0) && all(TileCoord[2] < TileRes) && NearestDepth > VisTileDepthGrid[TileCoord[2]])
+					{
+						InterlockedAdd(OutVisTileBinningGrid[uint3(TileCoord[2], BinID)], 1, SegIndex[2]);
+						bTileIntersected[2] = true;
+					}
+				}
+			}
 		}
-		FetchIndex1 = min(FetchIndex0 + 1, VertexCountAfterCulling - 1);
+		
+		GroupMemoryBarrierWithGroupSync();
 
-		VertexIndex0 = HairStrandsVF_CullingIndexBuffer[FetchIndex0];
-		LodRadiusScale0	= HairStrandsVF_CullingRadiusScaleBuffer[FetchIndex0];
+		if (bSegValid) //add tiles if any needed - this is unlikely to happen in all batches
+		{
+			for (int i = 0; i < 3; i++)
+			{
+				if ((SegIndex[i] != 0xffffffff) && ((SegIndex[i] % 1024)==0))
+				{
+					uint NewTile;
+					InterlockedAdd(OutVisTileArgs[0], 1, NewTile);
+					
+					StoreOutVisTileData(NewTile, VT_Coord, PackVisTileCoord(TileCoord[i]));
+					StoreOutVisTileData(NewTile, VT_PrimOffset, NewTile * 1024);
+					StoreOutVisTileData(NewTile, VT_NearestDepth, SegIndex[i]); //tmp - use this to decide if seg should go in new or prev tile.
+
+					uint PrevTile = (OutVisTileBinningGrid[uint3(TileCoord[i], BinID2)] & 0xffff);
+					
+					if (SegIndex[i]>0)	StoreOutVisTileData(PrevTile, VT_PrimCount, 1024);
+
+					OutVisTileBinningGrid[uint3(TileCoord[i], BinID2)] = (PrevTile << 16) | (NewTile & 0xffff);
+
+				}
+			}
+		}
 
-		VertexIndex1 = HairStrandsVF_CullingIndexBuffer[FetchIndex1];
-		LodRadiusScale1 = HairStrandsVF_CullingRadiusScaleBuffer[FetchIndex1];
+		GroupMemoryBarrierWithGroupSync();
 
-		if (LodRadiusScale0 <= 0.0f)
+		if (bSegValid) //now set the PrimID at the location indicated by the prev/current tile.
 		{
-			return;
+			for (int i = 0; i < 3; i++)
+			{
+				if (SegIndex[i] != 0xffffffff)
+				{
+					uint PackedTiles = OutVisTileBinningGrid[uint3(TileCoord[i], BinID2)];
+					uint CurTile = (PackedTiles & 0xffff);
+					uint PrevTile = ((PackedTiles >> 16) & 0xffff);
+
+					if (SegIndex[i] >= LoadOutVisTileData(CurTile, VT_NearestDepth))
+					{
+						uint PrimIndex = LoadOutVisTileData(CurTile, VT_PrimOffset) + (SegIndex[i] % 1024);
+						OutVisTilePrims[PrimIndex] = PrimID;
+
+						if ((SegIndex[i] + 1) == OutVisTileBinningGrid[uint3(TileCoord[i], BinID)])
+						{
+							StoreOutVisTileData(CurTile, VT_PrimCount, (SegIndex[i]==1023) ? 1024 : ((SegIndex[i] + 1) % 1024));
+						}
+					}
+					else
+					{
+						uint PrimIndex = LoadOutVisTileData(PrevTile, VT_PrimOffset) + (SegIndex[i] % 1024);
+						OutVisTilePrims[PrimIndex] = PrimID;
+					}
+
+					//InterlockedAdd(OutVisTileArgs[1], 1); //debug
+				}
+			}
 		}
 	}
-#endif
+}
+#endif //SHADER_RASTERCOMPUTE_BINNING
 
-	ResolvedView = ResolveView();
+///////////////////////////////////////////////////////////////////////////
 
-	const float3 PositionOffset = HairStrandsVF_PositionOffsetBuffer[0].xyz;
+#if SHADER_RASTERCOMPUTE_RASTER
 
-	const FHairControlPoint CP0 = UnpackHairControlPoint(
-		HairStrandsVF_PositionBuffer[VertexIndex0],
-		PositionOffset,
-		HairStrandsVF_Radius * LodRadiusScale0,
-		HairStrandsVF_RootScale,
-		HairStrandsVF_TipScale);
+RWTexture2D<uint> 					RWVisTileDepthGrid;
+Buffer<uint> 						VisTilePrims;
+Buffer<uint> 						VisTileArgs;
+RWByteAddressBuffer					RWVisTileData;
+Buffer<uint>						VisTileIndirect;
 
-	if (CP0.Type == HAIR_CONTROLPOINT_END)
-		return;
+uint LoadRWVisTileData(uint index, uint offset)
+{
+	return RWVisTileData.Load((((index)) * VT_SIZE * 4) + ((offset) * 4));
+}
 
-	const FHairControlPoint CP1 = UnpackHairControlPoint(
-		HairStrandsVF_PositionBuffer[VertexIndex1],
-		PositionOffset,
-		HairStrandsVF_Radius * LodRadiusScale1,
-		HairStrandsVF_RootScale,
-		HairStrandsVF_TipScale);
+void StoreRWVisTileData(uint index, uint offset, uint value)
+{
+	RWVisTileData.Store((((index)) * VT_SIZE * 4) + ((offset) * 4), (value));
+}
 
+RWTexture2D<uint> OutHairCountTexture;
+RWTexture2D<uint> OutDepthCovTexture;
+RWTexture2D<uint> OutPrimMatTexture;
 
-	const FHairRenderInfo HairRenderInfo = GetHairRenderInfo(ResolvedView.HairRenderInfo, ResolvedView.HairRenderInfoBits, HairStrandsVF_bUseStableRasterization > 0);
 
-	const float VelocityScale = 0;
-	FHairViewInfo HairViewInfo;
-	HairViewInfo.TranslatedWorldCameraOrigin = ResolvedView.TranslatedWorldCameraOrigin;
-	HairViewInfo.ViewForward = ResolvedView.ViewForward;
-	HairViewInfo.RadiusAtDepth1 = lerp(HairRenderInfo.RadiusAtDepth1Primary, HairRenderInfo.RadiusAtDepth1Velocity, VelocityScale);
-	HairViewInfo.bIsOrthoView = false;
+groupshared uint4 group_SubTile[1024]; //(32 x 32 x 4 x 4 bytes = 16k bytes)
 
-	const float3 WP0  = mul(float4(CP0.Position, 1), HairStrandsVF_LocalToWorldPrimitiveTransform).xyz;
-	const float4 NDC0 = mul(float4(WP0, 1), LWCHackToFloat(PrimaryView.WorldToClip));
+groupshared float3 group_PositionOffset;
+groupshared float group_ooTileLODScale;
 
-	const float3 WP1  = mul(float4(CP1.Position, 1), HairStrandsVF_LocalToWorldPrimitiveTransform).xyz;
-	const float4 NDC1 = mul(float4(WP1, 1), LWCHackToFloat(PrimaryView.WorldToClip));
-	
-	float3 SP0 = NDCToPixelCoord(NDC0);
-	float3 SP1 = NDCToPixelCoord(NDC1);
+groupshared uint group_LoopNum;
+groupshared uint group_TileNum;
 
-	// Coarse culling per segment
-	const bool bInClip0 = IsInScreen(SP0.xy);
-	const bool bInClip1 = IsInScreen(SP1.xy);
-	if (!bInClip0 || !bInClip1)
-	{
-		return;
-	}
+groupshared uint group_ThreadsPerSeg;
 
-	const bool bIsVisible0 = IsVisible(SP0);
-	const bool bIsVisible1 = IsVisible(SP1);
-	if (!bIsVisible0 && !bIsVisible1)
-	{
-		return;
-	}
+#define GS_SEGS 320 //this number is limited by group shared memory
+
+groupshared float4 group_SP0[GS_SEGS];
+groupshared float4 group_SP1[GS_SEGS];
+groupshared float group_Rad0[GS_SEGS];
+groupshared float group_Rad1[GS_SEGS];
+groupshared uint group_PrimMatID[GS_SEGS];
+
+groupshared uint group_TileIndex;
 
-	const float Coverage0 = ToCoverage(CP0, NDC0, HairRenderInfo.RadiusAtDepth1Primary);
-	const float Coverage1 = ToCoverage(CP1, NDC1, HairRenderInfo.RadiusAtDepth1Primary);
+[numthreads(1024, 1, 1)]
+void RasterCS(uint DispatchThreadID : SV_DispatchThreadID, uint GroupThreadID : SV_GroupThreadID, uint GroupID : SV_GroupID)
+{
+	ResolvedView = ResolveView();
 
-	// Conservative line raster (Wu's Algo)
-	// https://en.wikipedia.org/wiki/Xiaolin_Wu%27s_line_algorithm
-	const bool bIsSteep = abs(SP1.y - SP0.y) > abs(SP1.x - SP0.x);
-	if (bIsSteep)
+	if (GroupThreadID == 0)
 	{
-		Swap(SP0.x, SP0.y);
-		Swap(SP1.x, SP1.y);
+		group_TileNum = VisTileArgs[0];
+		group_LoopNum = (float(group_TileNum) + float(NumRasterizers - 1))* RcpNumRasterizers;
+//		group_LoopNum = (float(1024) + float(NumRasterizers - 1)) * RcpNumRasterizers;
+
+		group_PositionOffset = HairStrandsVF_PositionOffsetBuffer[0].xyz;
+
+		group_ooTileLODScale = 1.0;// max(min(1.0 / (TileLODScale * 25.0), 16.0), 1.0);
+
+		/* no longer in use - keep for ref? Moving these values to group shared memory did seem to reduce VGPRs - more experimentation needed
+		group_RadScale = (((HairStrandsVF_TipScale - HairStrandsVF_RootScale) * HairStrandsVF_Radius * OutputResolutionf.x) / 63.0) / 255.0;
+		group_RadOffset = (HairStrandsVF_RootScale * HairStrandsVF_Radius * OutputResolutionf.x)/63.0;
+		*/
 	}
-	if (SP0.x > SP1.x)
+
+	GroupMemoryBarrierWithGroupSync();
+
+	[loop]
+	for (uint LoopIndex = 0; LoopIndex < group_LoopNum; LoopIndex++)
 	{
-		Swap(SP0, SP1);
-	}
+		if (GroupThreadID == 0)
+		{
+			group_TileIndex = LoopIndex + (GroupID * group_LoopNum);
+		}
+		
+		GroupMemoryBarrierWithGroupSync();
 
-	const float3 D = SP1 - SP0;
-	const float GradientY = abs(D.x) < 0.00001f ? 1.0f : D.y / D.x;
-	const float GradientZ = abs(D.x) < 0.00001f ? 1.0f : D.z / D.x;
+		bool bTileValid = (group_TileIndex < group_TileNum) ? true : false;
 
-	float DeltaY = 0;
-	float DeltaZ = 0;
+		uint PrimOffset = LoadRWVisTileData(group_TileIndex, VT_PrimOffset);
+		uint PrimCount = LoadRWVisTileData(group_TileIndex, VT_PrimCount);
 
-	// First point
-	int2 PX0;
-	{
-		float2 End;
-		End.x = round(SP0.x);
-		End.y = SP0.y + GradientY * (End.x - SP0.x);
+		uint PackedCoord = LoadRWVisTileData(group_TileIndex, VT_Coord);
+		uint2 SubTileMin = UppackVisTileCoord(PackedCoord)*TileSize;
 
-		const float XGap = 1.f - frac(SP0.x + 0.5f);
+		uint PackedTileMin = ((SubTileMin.x & 0xffff) << 0) | ((SubTileMin.y & 0xffff) << 16);
 
-		PX0.x = End.x;
-		PX0.y = floor(End.y);
+		if (GroupThreadID == 0)
+		{
+			group_ThreadsPerSeg = 1;
+
+			if (PrimCount <= 512)
+				group_ThreadsPerSeg = 2;
+			if (PrimCount <= 341)
+				group_ThreadsPerSeg = 3;
+			if (PrimCount <= 256)
+				group_ThreadsPerSeg = 4;
+			if (PrimCount <= 204)
+				group_ThreadsPerSeg = 5;
+			if (PrimCount <= 170)
+				group_ThreadsPerSeg = 6;
+			if (PrimCount <= 146)
+				group_ThreadsPerSeg = 7;
+			if (PrimCount <= 128)
+				group_ThreadsPerSeg = 8;
+			if (PrimCount <= 64)
+				group_ThreadsPerSeg = 16;
+			if (PrimCount <= 32)
+				group_ThreadsPerSeg = 32;
+		}
+		
+		GroupMemoryBarrierWithGroupSync();
 
-		const float FracY = frac(End.y);
-		Rasterize(PX0, bIsSteep, FracY, XGap, Coverage0, SP0.z, VertexIndex0, false);
+		bool bThreadValid = (bTileValid && (GroupThreadID < (PrimCount*group_ThreadsPerSeg) )) ? true : false;
 
-		DeltaY = End.y + GradientY; // first y-intersection for the main loop
-		DeltaZ = SP1.z + GradientZ; // This is not perspective correct. Interpolate linear depth instead
-	}
+		uint WaveCount = ((PrimCount* group_ThreadsPerSeg) + (WAVE_SIZE-1) ) / WAVE_SIZE;
+		uint WaveThreadCount = WaveCount * WAVE_SIZE;
 
-	// Second point
-	int2 PX1;
-	{
-		float2 End;
-		End.x = floor(SP1.x + 0.5); // Round
-		End.y = SP1.y + GradientY * (End.x - SP1.x);
-		const float XGap = 1 - frac(SP1.x + 0.5);
+		bool bWaveThreadValid = (bTileValid && (GroupThreadID < WaveThreadCount)) ? true : false;
+
+		bool bUseGroupSPs = (bThreadValid && (GroupThreadID < (min(PrimCount, GS_SEGS)* group_ThreadsPerSeg) )) ? true : false;
+
+		bool bGenGroupSPs = (bThreadValid && (GroupThreadID < (min(PrimCount, GS_SEGS)) )) ? true : false;
 
-		PX1.x = End.x;
-		PX1.y = floor(End.y);
+		if (bGenGroupSPs)
+		{
+			uint Prim = GroupThreadID;
+			uint PrimID = VisTilePrims[PrimOffset + Prim];
+
+			group_PrimMatID[Prim] = PackHairVisPrimitiveMaterialId(PrimID, HairMaterialId);
+
+			CalcScreenPosAndRad(HairStrandsVF_PositionBuffer[PrimID], group_PositionOffset, group_SP0[Prim], group_Rad0[Prim]);
+			CalcScreenPosAndRad(HairStrandsVF_PositionBuffer[PrimID+1], group_PositionOffset, group_SP1[Prim], group_Rad1[Prim]);
+		}
 
-		// In order to avoid doubled rasterizated point at segments connection we exclude 
-		// the last pixel apart for pixel placed at the end of a curve
-		//if (CP1.Type == HAIR_CONTROLPOINT_END)
+		if (bWaveThreadValid)
 		{
-			const float FracY = frac(End.y);
-			Rasterize(PX1, bIsSteep, FracY, XGap, max(Coverage1, 0.01f), SP1.z, VertexIndex0, false);
+			for (uint LinearIndex = GroupThreadID; LinearIndex < SqrTileSize; LinearIndex += WaveThreadCount)
+			{
+				uint2 Coord;
+
+				Coord.y = (float(LinearIndex) + 0.5) * RcpTileSize;
+				Coord.x = LinearIndex - (Coord.y * TileSize);
+
+				Coord += uint2(((PackedTileMin >> 0) & 0xffff), ((PackedTileMin >> 16) & 0xffff));
+
+				group_SubTile[LinearIndex].x = OutDepthCovTexture[Coord];
+				group_SubTile[LinearIndex].y = 0;
+				group_SubTile[LinearIndex].z = 0;
+				group_SubTile[LinearIndex].w = PackHairVisDepthCoverage(SceneDepthTexture.Load(uint3(Coord, 0)), 1.0);
+			}
 		}
-	}
 
-	// Inner points
-	const uint MaxIteration = MaxRasterCount;
-	const float PixelCount = abs(PX1.x - (PX0.x + 1));
-	uint2 SeedCoord = PX0;
-	for (int i = PX0.x + 1, j = 0; i < PX1.x && j < MaxIteration; ++i, ++j)
-	{
-		const float S = float(j) / PixelCount;
-		const float FracY = frac(DeltaY);
-		const float DC_Z = lerp(NDC0.z, NDC1.z, S);
-		const float DC_W = lerp(NDC0.w, NDC1.w, S);
-		const float DeviceZ = DC_Z / DC_W;
-		
-		const uint2 PixelCoord = uint2(int2(i, floor(DeltaY)));
+		GroupMemoryBarrierWithGroupSync();
+
+		if (bThreadValid)
+		{
+			uint Prim = uint((float(GroupThreadID)+0.5) / float(group_ThreadsPerSeg));
+			uint PModTPS = GroupThreadID - (Prim * group_ThreadsPerSeg);
+
+			uint PrimMatID;
+			float4 SP0;
+			float4 SP1;
+			float Rad0;
+			float Rad1;
+
+			if (bUseGroupSPs)
+			{
+				PrimMatID = group_PrimMatID[Prim];
+				SP0 = group_SP0[Prim];
+				SP1 = group_SP1[Prim];
+				Rad0 = group_Rad0[Prim];
+				Rad1 = group_Rad1[Prim];
+			}
+			else
+			{
+				uint PrimID = VisTilePrims[PrimOffset + Prim];
+				PrimMatID = PackHairVisPrimitiveMaterialId(PrimID, HairMaterialId);
+
+				CalcScreenPosAndRad(HairStrandsVF_PositionBuffer[PrimID], group_PositionOffset, SP0, Rad0);
+				CalcScreenPosAndRad(HairStrandsVF_PositionBuffer[PrimID+1], group_PositionOffset, SP1, Rad1);
+			}
+
+			uint NumSteps = min(max(abs(SP1.x - SP0.x), abs(SP1.y - SP0.y))+1.0, TileSize);
+
+			[loop]
+			for (int J = PModTPS; J < NumSteps; J+=group_ThreadsPerSeg)
+			{
+				const float4 SP = lerp(SP0, SP1, J / float(NumSteps));
+
+				const int2 SPi = int2(SP.xy)-int2(((PackedTileMin >> 0) & 0xffff), ((PackedTileMin >> 16) & 0xffff));
+
+				if (SPi.x >= 0 && SPi.x < TileSize && SPi.y >= 0 && SPi.y < TileSize)
+				{
+					uint PackedDepthCov = PackHairVisDepthCoverage(SP.z / SP.w, 1.0);
+
+					uint LinearIndex = SPi.x + SPi.y * TileSize;
+
+					uint OldValue;
+					InterlockedMax(group_SubTile[LinearIndex].x, PackedDepthCov, OldValue);
+					if (PackedDepthCov > OldValue)
+					{
+						group_SubTile[LinearIndex].y = PrimMatID;
+					}
+
+					if (PackedDepthCov > group_SubTile[LinearIndex].w)
+					{
+						const float Rad = lerp(Rad0, Rad1, J / float(NumSteps)) / SP.w;
+
+#if PERMUTATION_TEMPORAL_LAYERING
+						InterlockedAdd(group_SubTile[LinearIndex].z, min(Rad, 1.0) * 1.0 * 8000.0 * group_ooTileLODScale);
+#else
+						InterlockedAdd(group_SubTile[LinearIndex].z, min(Rad, 1.0) * 1.0 * 1000.0 * group_ooTileLODScale);
+#endif
+					}
+				}					
+			}
+		}
 
-		SeedCoord = Rasterize(PixelCoord, bIsSteep, FracY, 1.0f, lerp(Coverage0, Coverage1, S), DeviceZ, VertexIndex0, true, SeedCoord);
+		GroupMemoryBarrierWithGroupSync();
 
-		DeltaY += GradientY;
-		DeltaZ += GradientZ;
+		if (bWaveThreadValid) 
+		{
+			for (uint LinearIndex = GroupThreadID; LinearIndex < SqrTileSize; LinearIndex += WaveThreadCount)
+			{
+				uint2 Coord;
+
+				Coord.y = (float(LinearIndex) + 0.5) * RcpTileSize;
+				Coord.x = LinearIndex - (Coord.y * TileSize);
+
+				Coord += uint2(((PackedTileMin >> 0) & 0xffff), ((PackedTileMin >> 16) & 0xffff));
+
+				if (group_SubTile[LinearIndex].y != 0)
+				{
+					uint oldValue;
+					InterlockedMax(OutDepthCovTexture[Coord], group_SubTile[LinearIndex].x, oldValue);
+					if (group_SubTile[LinearIndex].x > oldValue)
+					{
+						OutPrimMatTexture[Coord] = group_SubTile[LinearIndex].y;
+					}
+				}
+				InterlockedAdd(OutHairCountTexture[Coord], group_SubTile[LinearIndex].z);
+			}
+		}
+
+		GroupMemoryBarrierWithGroupSync();
 	}
 }
+
+#endif //SHADER_RASTERCOMPUTE_RASTER
+
+///////////////////////////////////////////////////////////////////////////
+
diff --git a/Engine/Source/Runtime/Renderer/Private/HairStrands/HairStrandsTile.cpp b/Engine/Source/Runtime/Renderer/Private/HairStrands/HairStrandsTile.cpp
--- a/Engine/Source/Runtime/Renderer/Private/HairStrands/HairStrandsTile.cpp
+++ b/Engine/Source/Runtime/Renderer/Private/HairStrands/HairStrandsTile.cpp
@@ -147,6 +147,7 @@ class FHairStrandsTileGenerationPassCS : public FGlobalShader
 IMPLEMENT_GLOBAL_SHADER(FHairStrandsTileGenerationPassCS, "/Engine/Private/HairStrands/HairStrandsVisibilityTile.usf", "TileMainCS", SF_Compute);
 
 float GetHairStrandsFullCoverageThreshold();
+uint32 GetHairStrandsIntCoverageThreshold();
 
 FHairStrandsTiles AddHairStrandsGenerateTilesPass(
 	FRDGBuilder& GraphBuilder,
@@ -180,8 +181,8 @@ FHairStrandsTiles AddHairStrandsGenerateTilesPass(
 	FHairStrandsTileGenerationPassCS::FParameters* PassParameters = GraphBuilder.AllocParameters<FHairStrandsTileGenerationPassCS::FParameters>();
 	PassParameters->BufferResolution		= InputResolution;//View.ViewRect.Size();
 	PassParameters->bUintTexture			= bUintTexture ? 1u : 0u;
-	PassParameters->TransmittanceThreshold	= 1.f - GetHairStrandsFullCoverageThreshold();
-	PassParameters->IntCoverageThreshold	= 1000u; // Arbitrary
+	PassParameters->TransmittanceThreshold	= 1.f - GetHairStrandsFullCoverageThreshold();	
+	PassParameters->IntCoverageThreshold 	= GetHairStrandsIntCoverageThreshold();
 	PassParameters->InputFloatTexture = bUintTexture ? GSystemTextures.GetBlackDummy(GraphBuilder) : InputTexture;
 	PassParameters->InputUintTexture  = bUintTexture ? InputTexture : GSystemTextures.GetZeroUIntDummy(GraphBuilder);
 	PassParameters->TileHairAllBuffer		= GraphBuilder.CreateUAV(Out.TileDataBuffer[ToIndex(FHairStrandsTiles::ETileType::HairAll)], PF_R16G16_UINT);
diff --git a/Engine/Source/Runtime/Renderer/Private/HairStrands/HairStrandsVisibility.cpp b/Engine/Source/Runtime/Renderer/Private/HairStrands/HairStrandsVisibility.cpp
--- a/Engine/Source/Runtime/Renderer/Private/HairStrands/HairStrandsVisibility.cpp
+++ b/Engine/Source/Runtime/Renderer/Private/HairStrands/HairStrandsVisibility.cpp
@@ -53,15 +53,23 @@ static FAutoConsoleVariableRef CVarGHairVisibilityPPLL_MeanSamplePerPixel(TEXT("
 static float GHairStrandsViewHairCountDepthDistanceThreshold = 30.f;
 static FAutoConsoleVariableRef CVarHairStrandsViewHairCountDepthDistanceThreshold(TEXT("r.HairStrands.Visibility.HairCount.DistanceThreshold"), GHairStrandsViewHairCountDepthDistanceThreshold, TEXT("Distance threshold defining if opaque depth get injected into the 'view-hair-count' buffer."));
 
-static int32 GHairVisibilityComputeRaster = 0;
-static int32 GHairVisibilityComputeRaster_MaxSamplePerPixel = 1;
-static float GHairVisibilityComputeRaster_MeanSamplePerPixel = 1;
-static int32 GHairVisibilityComputeRaster_MaxPixelCount = 64;
-static int32 GHairVisibilityComputeRaster_Stochastic = 0;
-static FAutoConsoleVariableRef CVarHairStrandsVisibilityComputeRaster(TEXT("r.HairStrands.Visibility.ComputeRaster"), GHairVisibilityComputeRaster, TEXT("Hair Visiblity uses raster compute."), ECVF_Scalability | ECVF_RenderThreadSafe);
-static FAutoConsoleVariableRef CVarHairStrandsVisibilityComputeRaster_MaxSamplePerPixel(TEXT("r.HairStrands.Visibility.ComputeRaster.SamplePerPixel"), GHairVisibilityComputeRaster_MaxSamplePerPixel, TEXT("Define the number of sampler per pixel using raster compute."));
-static FAutoConsoleVariableRef CVarHairStrandsVisibilityComputeRaster_MaxPixelCount(TEXT("r.HairStrands.Visibility.ComputeRaster.MaxPixelCount"), GHairVisibilityComputeRaster_MaxPixelCount, TEXT("Define the maximal length rasterize in compute."));
-static FAutoConsoleVariableRef CVarHairVisibilityComputeRaster_Stochastic(TEXT("r.HairStrands.Visibility.ComputeRaster.Stochastic"), GHairVisibilityComputeRaster_Stochastic, TEXT("Enable stochastic compute rasterization (faster, but more prone to aliasting). Experimental."));
+int32 GHairVisibilityComputeRaster = 0;
+int32 GHairVisibilityComputeRaster_Culling = 0;
+int32 GHairVisibilityComputeRaster_CLOD = 0;
+int32 GHairVisibilityComputeRaster_MaxTiles = 8192;
+int32 GHairVisibilityComputeRaster_TileSize = 32;
+int32 GHairVisibilityComputeRaster_NumBinners = 32;
+int32 GHairVisibilityComputeRaster_NumRasterizers = 256;
+int32 GHairVisibilityComputeRaster_TemporalLayering = 0;
+
+static FAutoConsoleVariableRef CVarHairStrandsVisibilityComputeRaster(TEXT("r.HairStrands.Visibility.ComputeRaster"), GHairVisibilityComputeRaster, TEXT("Hair Visiblity uses raster compute. Experimental"), ECVF_Scalability | ECVF_RenderThreadSafe);
+static FAutoConsoleVariableRef CVarHairVisibilityComputeRaster_Culling(TEXT("r.HairStrands.Visibility.ComputeRaster.Culling"), GHairVisibilityComputeRaster_Culling, TEXT("Use culling buffers with compute rasterization."));
+static FAutoConsoleVariableRef CVarHairVisibilityComputeRaster_CLOD(TEXT("r.HairStrands.Visibility.ComputeRaster.CLOD"), GHairVisibilityComputeRaster_CLOD, TEXT("Enable Continuos LOD when using compute rasterization. Experimental - needs further improvement"));
+static FAutoConsoleVariableRef CHairVisibilityComputeRaster_MaxTiles(TEXT("r.HairStrands.Visibility.ComputeRaster.MaxTiles"), GHairVisibilityComputeRaster_MaxTiles, TEXT("Maximum number of tiles used for compute rasterization. 8192 is default"));
+static FAutoConsoleVariableRef CHairVisibilityComputeRaster_TileSize(TEXT("r.HairStrands.Visibility.ComputeRaster.TileSize"), GHairVisibilityComputeRaster_TileSize, TEXT("Tile size used for compute rasterization. Experimental - only size of 32 currently supported"));
+static FAutoConsoleVariableRef CHairVisibilityComputeRaster_NumBinners(TEXT("r.HairStrands.Visibility.ComputeRaster.NumBinners"), GHairVisibilityComputeRaster_NumBinners, TEXT("Number of Binners used in Binning compute rasterization pass. 32 is default"));
+static FAutoConsoleVariableRef CHairVisibilityComputeRaster_NumRasterizers(TEXT("r.HairStrands.Visibility.ComputeRaster.NumRasterizers"), GHairVisibilityComputeRaster_NumRasterizers, TEXT("Number of Rasterizers used compute rasterization. 256 is default"));
+static FAutoConsoleVariableRef CHairVisibilityComputeRaster_TemporalLayering(TEXT("r.HairStrands.Visibility.ComputeRaster.TemporalLayering"), GHairVisibilityComputeRaster_TemporalLayering, TEXT("Enable Experimental WIP Temporal Layering (requires TAA changes to work well)"));
 
 static float GHairStrandsFullCoverageThreshold = 0.98f;
 static FAutoConsoleVariableRef CVarHairStrandsFullCoverageThreshold(TEXT("r.HairStrands.Visibility.FullCoverageThreshold"), GHairStrandsFullCoverageThreshold, TEXT("Define the coverage threshold at which a pixel is considered fully covered."));
@@ -126,7 +134,8 @@ enum EHairVisibilityRenderMode
 
 inline bool DoesSupportRasterCompute()
 {
-	return GRHISupportsAtomicUInt64;
+	//for now always return true 
+	return true; //GRHISupportsAtomicUInt64; //TODO: re-enable support for 64 bit visibility buffers
 }
 
 inline EHairVisibilityRenderMode GetHairVisibilityRenderMode()
@@ -157,18 +166,7 @@ static uint32 GetMaxSamplePerPixel()
 	{
 		case HairVisibilityRenderMode_ComputeRaster:
 		{
-			if (GHairVisibilityComputeRaster_MaxSamplePerPixel <= 1)
-			{
-				return 1;
-			}
-			else if (GHairVisibilityComputeRaster_MaxSamplePerPixel < 4)
-			{
-				return 2;
-			}
-			else
-			{
-				return 4;
-			}
+			return 1; //TODO: enable multi-sample compute raster
 		}
 		case HairVisibilityRenderMode_MSAA_Visibility:
 		{
@@ -220,7 +218,7 @@ inline uint32 GetMeanSamplePerPixel()
 	switch (GetHairVisibilityRenderMode())
 	{
 	case HairVisibilityRenderMode_ComputeRaster:
-		return FMath::Max(1, FMath::FloorToInt(SamplePerPixel * FMath::Clamp(GHairVisibilityComputeRaster_MeanSamplePerPixel, 0.f, 1.f)));
+		return 1; //TODO: enable multi-sample compute raster
 	case HairVisibilityRenderMode_MSAA_Visibility:
 		return FMath::Max(1, FMath::FloorToInt(SamplePerPixel * FMath::Clamp(GHairVisibilityMSAA_MeanSamplePerPixel, 0.f, 1.f)));
 	case HairVisibilityRenderMode_PPLL:
@@ -246,10 +244,8 @@ struct FRasterComputeOutput
 	FRDGTextureRef HairCountTexture = nullptr;
 	FRDGTextureRef DepthTexture = nullptr;
 
-	FRDGTextureRef VisibilityTexture0 = nullptr;
-	FRDGTextureRef VisibilityTexture1 = nullptr;
-	FRDGTextureRef VisibilityTexture2 = nullptr;
-	FRDGTextureRef VisibilityTexture3 = nullptr;
+	FRDGTextureRef DepthCovTexture = nullptr;
+	FRDGTextureRef PrimMatTexture = nullptr;
 };
 
 static uint32 GetTotalSampleCountForAllocation(FIntPoint Resolution)
@@ -292,6 +288,12 @@ float GetHairStrandsFullCoverageThreshold()
 {
 	return FMath::Clamp(GHairStrandsFullCoverageThreshold, 0.1f, 1.f);
 }
+
+uint32 GetHairStrandsIntCoverageThreshold()
+{
+	return ((GHairVisibilityComputeRaster) ? ((GHairVisibilityComputeRaster_TemporalLayering) ? 16000u : 4000u) : 1000u);
+}
+
 ///////////////////////////////////////////////////////////////////////////////////////////////////
 
 class FHairLightSampleClearVS : public FGlobalShader
@@ -1844,11 +1846,9 @@ class FHairVisibilityCompactionComputeRasterCS : public FGlobalShader
 		SHADER_PARAMETER(FIntPoint, TileCountXY)
 		SHADER_PARAMETER(uint32, TileSize)
 
-		SHADER_PARAMETER_RDG_TEXTURE(Texture2D<UlongType>, VisibilityTexture0)
-		SHADER_PARAMETER_RDG_TEXTURE(Texture2D<UlongType>, VisibilityTexture1)
-		SHADER_PARAMETER_RDG_TEXTURE(Texture2D<UlongType>, VisibilityTexture2)
-		SHADER_PARAMETER_RDG_TEXTURE(Texture2D<UlongType>, VisibilityTexture3)
-		SHADER_PARAMETER_RDG_TEXTURE(Texture2D, ViewTransmittanceTexture)
+		SHADER_PARAMETER_RDG_TEXTURE(Texture2D<uint>, DepthCovTexture)
+		SHADER_PARAMETER_RDG_TEXTURE(Texture2D<uint>, PrimMatTexture)
+		SHADER_PARAMETER_RDG_TEXTURE(Texture2D<uint>, HairCountTexture)
 
 		SHADER_PARAMETER_RDG_TEXTURE_UAV(RWTexture2D, OutCompactNodeCounter)
 		SHADER_PARAMETER_RDG_TEXTURE_UAV(RWTexture2D, OutCompactNodeIndex)
@@ -1886,7 +1886,6 @@ static void AddHairVisibilityCompactionComputeRasterPass(
 	const uint32 SamplerPerPixel,
 	const FRasterComputeOutput& RasterComputeData,
 	const FHairStrandsTiles& TileData,
-	FRDGTextureRef& InTransmittanceTexture,
 	FRDGTextureRef& OutCompactCounter,
 	FRDGTextureRef& OutCompactNodeIndex,
 	FRDGBufferRef&  OutCompactNodeVis,
@@ -1895,7 +1894,7 @@ static void AddHairVisibilityCompactionComputeRasterPass(
 	FRDGBufferRef&  OutIndirectArgsBuffer,
 	uint32& OutMaxRenderNodeCount)
 {	
-	FIntPoint Resolution = RasterComputeData.VisibilityTexture0->Desc.Extent;
+	FIntPoint Resolution = RasterComputeData.DepthCovTexture->Desc.Extent;
 
 	{
 		FRDGTextureDesc Desc = FRDGTextureDesc::Create2D(FIntPoint(1,1), PF_R32_UINT, FClearValueBinding::None, TexCreate_UAV);
@@ -1928,12 +1927,12 @@ static void AddHairVisibilityCompactionComputeRasterPass(
 
 	FRDGTextureRef DefaultTexture = GSystemTextures.GetBlackDummy(GraphBuilder);
 	FHairVisibilityCompactionComputeRasterCS::FParameters* PassParameters = GraphBuilder.AllocParameters<FHairVisibilityCompactionComputeRasterCS::FParameters>();
-	PassParameters->VisibilityTexture0		= RasterComputeData.VisibilityTexture0;
-	PassParameters->VisibilityTexture1		= SamplerPerPixel > 1 ? RasterComputeData.VisibilityTexture1 : DefaultTexture;
-	PassParameters->VisibilityTexture2		= SamplerPerPixel > 2 ? RasterComputeData.VisibilityTexture2 : DefaultTexture;
-	PassParameters->VisibilityTexture3		= SamplerPerPixel > 3 ? RasterComputeData.VisibilityTexture3 : DefaultTexture;
+
+	PassParameters->HairCountTexture = RasterComputeData.HairCountTexture;
+	PassParameters->DepthCovTexture = RasterComputeData.DepthCovTexture;
+	PassParameters->PrimMatTexture = RasterComputeData.PrimMatTexture;
+
 	PassParameters->SamplerPerPixel			= SamplerPerPixel;
-	PassParameters->ViewTransmittanceTexture= InTransmittanceTexture;
 	PassParameters->OutputResolution		= Resolution;
 	PassParameters->MaxNodeCount			= MaxRenderNodeCount;
 	PassParameters->CoverageThreshold		= GetHairStrandsFullCoverageThreshold();
@@ -2016,7 +2015,7 @@ static FRDGTextureRef AddHairVisibilityFillOpaqueDepth(
 	const FHairStrandsTiles& TileData,
 	const FRDGTextureRef& SceneDepthTexture)
 {
-	check(GetHairVisibilityRenderMode() == HairVisibilityRenderMode_MSAA_Visibility);
+	check(GetHairVisibilityRenderMode() == HairVisibilityRenderMode_MSAA_Visibility || GetHairVisibilityRenderMode() == HairVisibilityRenderMode_ComputeRaster);
 
 	const bool bUseTile = TileData.IsValid();
 	const FHairStrandsTiles::ETileType TileType = FHairStrandsTiles::ETileType::HairAll;
@@ -2900,30 +2899,58 @@ void AddMeshDrawTransitionPass(
 
 			ResourceAccessFinalizer.Finalize(GraphBuilder);
 
-			VFInput.Strands.PositionBuffer				= FRDGImportedBuffer();
-			VFInput.Strands.PrevPositionBuffer			= FRDGImportedBuffer();
-			VFInput.Strands.TangentBuffer				= FRDGImportedBuffer();
-			VFInput.Strands.Attribute0Buffer			= FRDGImportedBuffer();
-			VFInput.Strands.Attribute1Buffer			= FRDGImportedBuffer();
-			VFInput.Strands.MaterialBuffer				= FRDGImportedBuffer();
-			VFInput.Strands.PositionOffsetBuffer		= FRDGImportedBuffer();
-			VFInput.Strands.PrevPositionOffsetBuffer	= FRDGImportedBuffer();
+			if (GetHairVisibilityRenderMode() != HairVisibilityRenderMode_ComputeRaster)
+			{
+				VFInput.Strands.PositionBuffer				= FRDGImportedBuffer();
+				VFInput.Strands.PrevPositionBuffer			= FRDGImportedBuffer();
+				VFInput.Strands.TangentBuffer				= FRDGImportedBuffer();
+				VFInput.Strands.Attribute0Buffer			= FRDGImportedBuffer();
+				VFInput.Strands.Attribute1Buffer			= FRDGImportedBuffer();
+				VFInput.Strands.MaterialBuffer				= FRDGImportedBuffer();
+				VFInput.Strands.PositionOffsetBuffer		= FRDGImportedBuffer();
+				VFInput.Strands.PrevPositionOffsetBuffer	= FRDGImportedBuffer();
+			}
 		}
 	}
 }
 ///////////////////////////////////////////////////////////////////////////////////////////////////
 
-class FVisiblityRasterComputeCS : public FGlobalShader
+class FVisiblityRasterComputePrepareDepthGridCS : public FGlobalShader
+{
+	DECLARE_GLOBAL_SHADER(FVisiblityRasterComputePrepareDepthGridCS);
+	SHADER_USE_PARAMETER_STRUCT(FVisiblityRasterComputePrepareDepthGridCS, FGlobalShader);
+
+	BEGIN_SHADER_PARAMETER_STRUCT(FParameters, )
+		SHADER_PARAMETER_RDG_TEXTURE(Texture2D, SceneDepthTexture)
+		SHADER_PARAMETER(FIntPoint, OutputResolution)
+		SHADER_PARAMETER_RDG_TEXTURE_UAV(RWTexture2D, OutVisTileDepthGrid)
+		SHADER_PARAMETER_RDG_TEXTURE_UAV(RWTexture2D, OutDepthCovTexture)
+		SHADER_PARAMETER(FIntPoint, TileRes)
+		SHADER_PARAMETER(uint32, TileSize)
+		SHADER_PARAMETER(float, RcpTileSize)
+		SHADER_PARAMETER(uint32, SqrTileSize)
+		END_SHADER_PARAMETER_STRUCT()
+
+public:
+	static bool ShouldCompilePermutation(const FGlobalShaderPermutationParameters& Parameters) { return IsHairStrandsSupported(EHairStrandsShaderType::Strands, Parameters.Platform); }
+	static void ModifyCompilationEnvironment(const FGlobalShaderPermutationParameters& Parameters, FShaderCompilerEnvironment& OutEnvironment)
+	{
+		FGlobalShader::ModifyCompilationEnvironment(Parameters, OutEnvironment);
+		OutEnvironment.SetDefine(TEXT("SHADER_RASTERCOMPUTE_DEPTH_GRID"), 1);
+	}
+};
+
+IMPLEMENT_GLOBAL_SHADER(FVisiblityRasterComputePrepareDepthGridCS, "/Engine/Private/HairStrands/HairStrandsVisibilityRasterCompute.usf", "PrepareDepthGridCS", SF_Compute);
+
+class FVisiblityRasterComputeBinningCS : public FGlobalShader
 {
-	DECLARE_GLOBAL_SHADER(FVisiblityRasterComputeCS);
-	SHADER_USE_PARAMETER_STRUCT(FVisiblityRasterComputeCS, FGlobalShader);
+	DECLARE_GLOBAL_SHADER(FVisiblityRasterComputeBinningCS);
+	SHADER_USE_PARAMETER_STRUCT(FVisiblityRasterComputeBinningCS, FGlobalShader);
 
 	class FGroupSize : SHADER_PERMUTATION_SPARSE_INT("PERMUTATION_GROUP_SIZE", 32, 64);
-	class FRasterAtomic : SHADER_PERMUTATION_INT("PERMUTATION_RASTER_ATOMIC", 4);
-	class FSPP : SHADER_PERMUTATION_SPARSE_INT("PERMUTATION_SPP", 1, 2, 4); 
 	class FCulling : SHADER_PERMUTATION_BOOL("PERMUTATION_CULLING");
-	class FStochastic : SHADER_PERMUTATION_BOOL("PERMUTATION_STOCHASTIC");
-	using FPermutationDomain = TShaderPermutationDomain<FRasterAtomic, FSPP, FCulling, FStochastic, FGroupSize>;
+	class FTemporalLayering : SHADER_PERMUTATION_BOOL("PERMUTATION_TEMPORAL_LAYERING");
+	using FPermutationDomain = TShaderPermutationDomain<FGroupSize, FCulling, FTemporalLayering>;
 
 	BEGIN_SHADER_PARAMETER_STRUCT(FParameters, )
 		SHADER_PARAMETER(uint32, MacroGroupId)
@@ -2949,61 +2976,97 @@ class FVisiblityRasterComputeCS : public FGlobalShader
 		SHADER_PARAMETER_RDG_BUFFER_SRV(Buffer, HairStrandsVF_CullingRadiusScaleBuffer)
 		RDG_BUFFER_ACCESS(IndirectBufferArgs, ERHIAccess::IndirectArgs)
 		SHADER_PARAMETER_RDG_TEXTURE(Texture2D, SceneDepthTexture)
-		SHADER_PARAMETER_RDG_TEXTURE_UAV(RWTexture2D, OutHairCountTexture)
-		SHADER_PARAMETER_RDG_TEXTURE_UAV(RWTexture2D, OutVisibilityTexture0)
-		SHADER_PARAMETER_RDG_TEXTURE_UAV(RWTexture2D, OutVisibilityTexture1)
-		SHADER_PARAMETER_RDG_TEXTURE_UAV(RWTexture2D, OutVisibilityTexture2)
-		SHADER_PARAMETER_RDG_TEXTURE_UAV(RWTexture2D, OutVisibilityTexture3)
+		SHADER_PARAMETER_RDG_TEXTURE(Texture3D, VisTileBinningGrid)
+		SHADER_PARAMETER_RDG_BUFFER_UAV(RWBuffer, OutVisTilePrims)
+		SHADER_PARAMETER_RDG_TEXTURE_UAV(RWTexture3D, OutVisTileBinningGrid)
+		SHADER_PARAMETER_RDG_TEXTURE(Texture2D, VisTileDepthGrid)
+		SHADER_PARAMETER_RDG_BUFFER_UAV(RWBuffer, OutVisTileArgs)
+		SHADER_PARAMETER_RDG_BUFFER_UAV(RWByteAddressBuffer, OutVisTileData)
 		SHADER_PARAMETER_STRUCT_REF(FViewUniformShaderParameters, ViewUniformBuffer)
+		RDG_TEXTURE_ACCESS(VisTileBinningGridTex, ERHIAccess::UAVCompute)
+		SHADER_PARAMETER(float, TileLODScale)
+		SHADER_PARAMETER(float, TileBaseDepth)
+		SHADER_PARAMETER(FIntPoint, TileRes)
+		SHADER_PARAMETER(uint32, TileSize)
+		SHADER_PARAMETER(float, RcpTileSize)
+		SHADER_PARAMETER(uint32, SqrTileSize)
+		SHADER_PARAMETER(uint32, NumBinners)
+		SHADER_PARAMETER(float, RcpNumBinners)
 	END_SHADER_PARAMETER_STRUCT()
 
 public:
-	static bool ShouldCompilePermutation(const FGlobalShaderPermutationParameters& Parameters) 
-	{ 
-		//if (!FDataDrivenShaderPlatformInfo::GetSupportsUInt64ImageAtomics(Parameters.Platform))
-		//{
-		//	return false;
-		//}
-		if (IsVulkanPlatform(Parameters.Platform))
-		{
-			return false;
-		}
+	static bool ShouldCompilePermutation(const FGlobalShaderPermutationParameters& Parameters) { return IsHairStrandsSupported(EHairStrandsShaderType::Strands, Parameters.Platform); }
+	static void ModifyCompilationEnvironment(const FGlobalShaderPermutationParameters& Parameters, FShaderCompilerEnvironment& OutEnvironment)
+	{
+		FGlobalShader::ModifyCompilationEnvironment(Parameters, OutEnvironment);
+		OutEnvironment.SetDefine(TEXT("SHADER_RASTERCOMPUTE_BINNING"), 1);
+	}
+};
 
-		if (!IsHairStrandsSupported(EHairStrandsShaderType::Strands, Parameters.Platform))
-		{
-			return false;
-		}
+IMPLEMENT_GLOBAL_SHADER(FVisiblityRasterComputeBinningCS, "/Engine/Private/HairStrands/HairStrandsVisibilityRasterCompute.usf", "BinningCS", SF_Compute);
 
-		if (IsPCPlatform(Parameters.Platform))
-		{
-			FPermutationDomain PermutationVector(Parameters.PermutationId);
-			return PermutationVector.Get<FRasterAtomic>() != 0;
-		}
-		else
-		{
-			FPermutationDomain PermutationVector(Parameters.PermutationId);
-			return PermutationVector.Get<FRasterAtomic>() == 0;
-		}
-	}
+class FVisiblityRasterComputeRasterizeCS : public FGlobalShader
+{
+	DECLARE_GLOBAL_SHADER(FVisiblityRasterComputeRasterizeCS);
+	SHADER_USE_PARAMETER_STRUCT(FVisiblityRasterComputeRasterizeCS, FGlobalShader);
+
+	class FGroupSize : SHADER_PERMUTATION_SPARSE_INT("PERMUTATION_GROUP_SIZE", 32, 64);
+	class FTemporalLayering : SHADER_PERMUTATION_BOOL("PERMUTATION_TEMPORAL_LAYERING");
+	using FPermutationDomain = TShaderPermutationDomain<FGroupSize, FTemporalLayering>;
+
+	BEGIN_SHADER_PARAMETER_STRUCT(FParameters, )
+		SHADER_PARAMETER(uint32, MacroGroupId)
+		SHADER_PARAMETER(uint32, DispatchCountX)
+		SHADER_PARAMETER(uint32, MaxRasterCount)
+		SHADER_PARAMETER(uint32, FrameIdMod8)
+		SHADER_PARAMETER(uint32, HairMaterialId)
+		SHADER_PARAMETER(uint32, ResolutionMultiplier)
+		SHADER_PARAMETER(FIntPoint, OutputResolution)
+		SHADER_PARAMETER(FVector2f, OutputResolutionf)
+		SHADER_PARAMETER(uint32, HairStrandsVF_bIsCullingEnable)
+		SHADER_PARAMETER(float, HairStrandsVF_Density)
+		SHADER_PARAMETER(float, HairStrandsVF_Radius)
+		SHADER_PARAMETER(float, HairStrandsVF_RootScale)
+		SHADER_PARAMETER(float, HairStrandsVF_TipScale)
+		SHADER_PARAMETER(float, HairStrandsVF_Length)
+		SHADER_PARAMETER(uint32, HairStrandsVF_bUseStableRasterization)
+		SHADER_PARAMETER(uint32, HairStrandsVF_VertexCount)
+		SHADER_PARAMETER(FMatrix44f, HairStrandsVF_LocalToWorldPrimitiveTransform)
+		SHADER_PARAMETER_RDG_BUFFER_SRV(Buffer, HairStrandsVF_PositionBuffer)
+		SHADER_PARAMETER_RDG_BUFFER_SRV(Buffer, HairStrandsVF_PositionOffsetBuffer)
+		SHADER_PARAMETER_RDG_BUFFER_SRV(Buffer, HairStrandsVF_CullingIndirectBuffer)
+		SHADER_PARAMETER_RDG_BUFFER_SRV(Buffer, HairStrandsVF_CullingIndexBuffer)
+		SHADER_PARAMETER_RDG_BUFFER_SRV(Buffer, HairStrandsVF_CullingRadiusScaleBuffer)
+		RDG_BUFFER_ACCESS(IndirectBufferArgs, ERHIAccess::IndirectArgs)
+		SHADER_PARAMETER_RDG_TEXTURE(Texture2D, SceneDepthTexture)
+		SHADER_PARAMETER(float, TileLODScale)
+		SHADER_PARAMETER(uint32, TileSize)
+		SHADER_PARAMETER(float, RcpTileSize)
+		SHADER_PARAMETER(uint32, SqrTileSize)
+		SHADER_PARAMETER(uint32, NumRasterizers)
+		SHADER_PARAMETER(float, RcpNumRasterizers)
+		SHADER_PARAMETER_RDG_TEXTURE_UAV(RWTexture2D, RWVisTileDepthGrid)
+		SHADER_PARAMETER_RDG_BUFFER_SRV(Buffer, VisTilePrims)
+		SHADER_PARAMETER_RDG_BUFFER_SRV(Buffer, VisTileArgs)
+		SHADER_PARAMETER_RDG_BUFFER_UAV(RWByteAddressBuffer, RWVisTileData)
+		SHADER_PARAMETER_RDG_BUFFER_SRV(Buffer, VisTileIndirect)
+		SHADER_PARAMETER_RDG_TEXTURE_UAV(RWTexture2D, OutHairCountTexture)
+		SHADER_PARAMETER_RDG_TEXTURE_UAV(RWTexture2D, OutDepthCovTexture)
+		SHADER_PARAMETER_RDG_TEXTURE_UAV(RWTexture2D, OutPrimMatTexture)
+		SHADER_PARAMETER_STRUCT_REF(FViewUniformShaderParameters, ViewUniformBuffer)
+		SHADER_PARAMETER(FIntPoint, TileRes)
+	END_SHADER_PARAMETER_STRUCT()
+
+public:
+	static bool ShouldCompilePermutation(const FGlobalShaderPermutationParameters& Parameters) { return IsHairStrandsSupported(EHairStrandsShaderType::Strands, Parameters.Platform); }
 	static void ModifyCompilationEnvironment(const FGlobalShaderPermutationParameters& Parameters, FShaderCompilerEnvironment& OutEnvironment)
 	{
 		FGlobalShader::ModifyCompilationEnvironment(Parameters, OutEnvironment);
-		OutEnvironment.SetDefine(TEXT("SHADER_RASTERCOMPUTE"), 1);
-		// Need to force optimization for driver injection to work correctly.
-		// https://developer.nvidia.com/unlocking-gpu-intrinsics-hlsl
-		// https://gpuopen.com/gcn-shader-extensions-for-direct3d-and-vulkan/
-		OutEnvironment.CompilerFlags.Add(CFLAG_ForceOptimization);
-
-		FPermutationDomain PermutationVector(Parameters.PermutationId);
-		if (PermutationVector.Get<FRasterAtomic>() == 3) // AMD, DX12
-		{
-			// Force shader model 6.0+
-			OutEnvironment.CompilerFlags.Add(CFLAG_ForceDXC);
-		}
+		OutEnvironment.SetDefine(TEXT("SHADER_RASTERCOMPUTE_RASTER"), 1);
 	}
 };
 
-IMPLEMENT_GLOBAL_SHADER(FVisiblityRasterComputeCS, "/Engine/Private/HairStrands/HairStrandsVisibilityRasterCompute.usf", "MainCS", SF_Compute);
+IMPLEMENT_GLOBAL_SHADER(FVisiblityRasterComputeRasterizeCS, "/Engine/Private/HairStrands/HairStrandsVisibilityRasterCompute.usf", "RasterCS", SF_Compute);
 
 static FRasterComputeOutput AddVisibilityComputeRasterPass(
 	FRDGBuilder& GraphBuilder,
@@ -3020,47 +3083,71 @@ static FRasterComputeOutput AddVisibilityComputeRasterPass(
 	Out.ResolutionMultiplier = 1;
 	Out.BaseResolution		 = InResolution;
 	Out.SuperResolution		 = InResolution * Out.ResolutionMultiplier;
-	Out.VisibilityTexture0	 = nullptr;
-	Out.VisibilityTexture1	 = nullptr;
-	Out.VisibilityTexture2	 = nullptr;
-	Out.VisibilityTexture3	 = nullptr;
 
 	FRDGTextureDesc DescCount = FRDGTextureDesc::Create2D(Out.SuperResolution, PF_R32_UINT, FClearValueBinding::None, TexCreate_UAV | TexCreate_ShaderResource | TexCreate_RenderTargetable);
-	FRDGTextureDesc DescVis   = FRDGTextureDesc::Create2D(Out.SuperResolution, PF_R32G32_UINT, FClearValueBinding::None, TexCreate_UAV | TexCreate_ShaderResource | TexCreate_RenderTargetable);
-	FRDGTextureUAVRef VisibilityTexture0UAV = nullptr;
-	FRDGTextureUAVRef VisibilityTexture1UAV = nullptr;
-	FRDGTextureUAVRef VisibilityTexture2UAV = nullptr;
-	FRDGTextureUAVRef VisibilityTexture3UAV = nullptr;
+	FRDGTextureDesc DescVis = FRDGTextureDesc::Create2D(Out.SuperResolution, PF_R32_UINT, FClearValueBinding::None, TexCreate_UAV | TexCreate_ShaderResource | TexCreate_RenderTargetable);
+
+	Out.HairCountTexture = GraphBuilder.CreateTexture(DescCount, TEXT("Hair.HairCountTexture"));
+	Out.DepthCovTexture = GraphBuilder.CreateTexture(DescVis, TEXT("Hair.DepthCovTexture"));
+	Out.PrimMatTexture = GraphBuilder.CreateTexture(DescVis, TEXT("Hair.PrimMatTexture"));
 
-	uint32 ClearValues[4] = { 0,0,0,0 };
-	Out.HairCountTexture = GraphBuilder.CreateTexture(DescCount, TEXT("Hair.ViewTransmittanceTexture"));
 	FRDGTextureUAVRef HairCountTextureUAV = GraphBuilder.CreateUAV(Out.HairCountTexture);
+	FRDGTextureUAVRef DepthCovTextureUAV = GraphBuilder.CreateUAV(Out.DepthCovTexture);
+	FRDGTextureUAVRef PrimMatTextureUAV = GraphBuilder.CreateUAV(Out.PrimMatTexture);
+
+	uint32 ClearValues[4] = { 0,0,0,0 };
+	AddClearUAVPass(GraphBuilder, PrimMatTextureUAV, ClearValues);
+	AddClearUAVPass(GraphBuilder, DepthCovTextureUAV, ClearValues);
 	AddClearUAVPass(GraphBuilder, HairCountTextureUAV, ClearValues);
 
-	Out.VisibilityTexture0 = GraphBuilder.CreateTexture(DescVis, TEXT("Hair.VisibilityTexture0"));
-	VisibilityTexture0UAV = GraphBuilder.CreateUAV(Out.VisibilityTexture0);
-	AddClearUAVPass(GraphBuilder, VisibilityTexture0UAV, ClearValues);
-	if (SamplePerPixelCount > 1)
-	{
-		Out.VisibilityTexture1 = GraphBuilder.CreateTexture(DescVis, TEXT("Hair.VisibilityTexture1"));
-		VisibilityTexture1UAV = GraphBuilder.CreateUAV(Out.VisibilityTexture1);
-		AddClearUAVPass(GraphBuilder, VisibilityTexture1UAV, ClearValues);
-		if (SamplePerPixelCount > 2)
-		{
-			Out.VisibilityTexture2 = GraphBuilder.CreateTexture(DescVis, TEXT("Hair.VisibilityTexture2"));
-			VisibilityTexture2UAV = GraphBuilder.CreateUAV(Out.VisibilityTexture2);
-			AddClearUAVPass(GraphBuilder, VisibilityTexture2UAV, ClearValues);
-			if (SamplePerPixelCount > 3)
-			{
-				Out.VisibilityTexture3 = GraphBuilder.CreateTexture(DescVis, TEXT("Hair.VisibilityTexture3"));
-				VisibilityTexture3UAV = GraphBuilder.CreateUAV(Out.VisibilityTexture3);
-				AddClearUAVPass(GraphBuilder, VisibilityTexture3UAV, ClearValues);
-			}
-		}
-	}
+	//set up buffers for binning and raster
+	const uint32 MaxTiles = FMath::Min(FMath::Max(GHairVisibilityComputeRaster_MaxTiles, 1024), 65536);
+
+	const uint32 TileSize = FMath::Min(FMath::Max(GHairVisibilityComputeRaster_TileSize, 8),32);
+	const uint32 NumBinners = FMath::Min(FMath::Max(GHairVisibilityComputeRaster_NumBinners, 1), 256);
+	const uint32 NumRasterizers = FMath::Min(FMath::Max(GHairVisibilityComputeRaster_NumRasterizers, 1), 1024);
+
+	const FIntPoint TileGridRes = FIntPoint((InResolution.X + (TileSize-1)) / TileSize, ((InResolution.Y + (TileSize - 1)) / TileSize)*1);
+	const FIntVector TileGridRes3D = FIntVector(TileGridRes.X, TileGridRes.Y, NumBinners * 2);
+
+	FRDGBufferRef OutVisTilePrims = nullptr;
+	FRDGTextureRef OutVisTileBinningGrid = nullptr;
+	FRDGTextureRef OutVisTileDepthGrid = nullptr;
+	FRDGBufferRef OutVisTileArgs = nullptr;
+	FRDGBufferRef OutVisTileData = nullptr;
+	FRDGBufferRef OutVisTileIndirect = nullptr;
+
+	FRDGBufferUAVRef OutVisTilePrimsUAV = nullptr;
+	FRDGTextureUAVRef OutVisTileBinningGridUAV = nullptr;
+	FRDGTextureUAVRef OutVisTileDepthGridUAV = nullptr;
+	FRDGBufferUAVRef OutVisTileArgsUAV = nullptr;
+	FRDGBufferUAVRef OutVisTileDataUAV = nullptr;
+	FRDGBufferUAVRef OutVisTileIndirectUAV = nullptr;
+	
+	OutVisTilePrims = GraphBuilder.CreateBuffer(FRDGBufferDesc::CreateStructuredDesc(sizeof(uint32), MaxTiles*1024), TEXT("Hair.OutVisTilePrims"));
+
+	FRDGTextureDesc DescBinningGrid = FRDGTextureDesc::Create3D(TileGridRes3D, PF_R32_UINT, FClearValueBinding::None, TexCreate_UAV | TexCreate_ShaderResource);
+	OutVisTileBinningGrid = GraphBuilder.CreateTexture(DescBinningGrid, TEXT("Hair.VisTileBinningGrid"));
+
+	FRDGTextureDesc DescDepthGrid = FRDGTextureDesc::Create2D(TileGridRes, PF_R32_UINT, FClearValueBinding::None, TexCreate_UAV | TexCreate_ShaderResource);
+	OutVisTileDepthGrid = GraphBuilder.CreateTexture(DescDepthGrid, TEXT("Hair.VisTileDepthGrid"));
+
+	OutVisTileArgs = GraphBuilder.CreateBuffer(FRDGBufferDesc::CreateStructuredDesc(sizeof(uint32), 4+16+(((TileGridRes.X*TileGridRes.Y)+31)/32) ), TEXT("Hair.OutVisTileArgs"));
+
+	FRDGBufferDesc DescTileData = FRDGBufferDesc::CreateStructuredDesc(sizeof(uint32), MaxTiles * 4 * 2);
+	DescTileData.Usage = EBufferUsageFlags(DescTileData.Usage | BUF_ByteAddressBuffer);
+	OutVisTileData = GraphBuilder.CreateBuffer(DescTileData, TEXT("Hair.OutVisTileData"));
+
+	OutVisTileIndirect = GraphBuilder.CreateBuffer(FRDGBufferDesc::CreateIndirectDesc<FRHIDispatchIndirectParameters>(1), TEXT("Hair.OutVisTileIndirect"));
+
+	OutVisTilePrimsUAV = GraphBuilder.CreateUAV(OutVisTilePrims);
+	OutVisTileBinningGridUAV = GraphBuilder.CreateUAV(OutVisTileBinningGrid);
+	OutVisTileDepthGridUAV = GraphBuilder.CreateUAV(OutVisTileDepthGrid);
+	OutVisTileArgsUAV = GraphBuilder.CreateUAV(OutVisTileArgs);
+	OutVisTileDataUAV = GraphBuilder.CreateUAV(OutVisTileData);
+	OutVisTileIndirectUAV = GraphBuilder.CreateUAV(OutVisTileIndirect);
 
 	// Create and set the uniform buffer
-	const bool bStochasticRaster = GHairVisibilityComputeRaster_Stochastic > 0;
 	const bool bEnableMSAA = false;
 	TUniformBufferRef<FViewUniformShaderParameters> ViewUniformShaderParameters;
 	SetUpViewHairRenderInfo(ViewInfo, bEnableMSAA, ViewInfo.CachedViewUniformShaderParameters->HairRenderInfo, ViewInfo.CachedViewUniformShaderParameters->HairRenderInfoBits, ViewInfo.CachedViewUniformShaderParameters->HairComponents);
@@ -3069,32 +3156,37 @@ static FRasterComputeOutput AddVisibilityComputeRasterPass(
 	const uint32 FrameIdMode8 = ViewInfo.ViewState ? (ViewInfo.ViewState->GetFrameIndex() % 8) : 0;
 	const uint32 GroupSize = GetVendorOptimalGroupSize1D();
 
-	FVisiblityRasterComputeCS::FPermutationDomain PermutationVector0;
-	FVisiblityRasterComputeCS::FPermutationDomain PermutationVector1;
-#if PLATFORM_WINDOWS
-	if (IsRHIDeviceNVIDIA())
-	{
-		PermutationVector0.Set<FVisiblityRasterComputeCS::FRasterAtomic>(1);
-	}
-	else if (IsRHIDeviceAMD())
-	{
-		const bool bIsDx12 = RHIGetInterfaceType() == ERHIInterfaceType::D3D12;
-		PermutationVector0.Set<FVisiblityRasterComputeCS::FRasterAtomic>(bIsDx12 ? 2 : 3);
-	}
-#else
-	{
-		PermutationVector0.Set<FVisiblityRasterComputeCS::FRasterAtomic>(0);
-	}
-#endif
-	PermutationVector0.Set<FVisiblityRasterComputeCS::FStochastic>(bStochasticRaster);
-	PermutationVector0.Set<FVisiblityRasterComputeCS::FSPP>(SamplePerPixelCount);
-	PermutationVector0.Set<FVisiblityRasterComputeCS::FGroupSize>(GroupSize);
-	PermutationVector1 = PermutationVector0; 
-
-	PermutationVector0.Set<FVisiblityRasterComputeCS::FCulling>(false);
-	PermutationVector1.Set<FVisiblityRasterComputeCS::FCulling>(true);
-	TShaderMapRef<FVisiblityRasterComputeCS> ComputeShader_CullingOff(ViewInfo.ShaderMap, PermutationVector0);
-	TShaderMapRef<FVisiblityRasterComputeCS> ComputeShader_CullingOn (ViewInfo.ShaderMap, PermutationVector1);
+	const bool bTemporalLayering = GHairVisibilityComputeRaster_TemporalLayering ? true : false;
+
+	FVisiblityRasterComputeBinningCS::FPermutationDomain BinningPermutationVector_CullingOff;
+	BinningPermutationVector_CullingOff.Set<FVisiblityRasterComputeBinningCS::FGroupSize>(GroupSize);
+	BinningPermutationVector_CullingOff.Set<FVisiblityRasterComputeBinningCS::FCulling>(false);
+	BinningPermutationVector_CullingOff.Set<FVisiblityRasterComputeBinningCS::FTemporalLayering>(bTemporalLayering);
+	TShaderMapRef<FVisiblityRasterComputeBinningCS> ComputeShaderBinning_CullingOff(ViewInfo.ShaderMap, BinningPermutationVector_CullingOff);
+
+	FVisiblityRasterComputeBinningCS::FPermutationDomain BinningPermutationVector_CullingOn;
+	BinningPermutationVector_CullingOn.Set<FVisiblityRasterComputeBinningCS::FGroupSize>(GroupSize);
+	BinningPermutationVector_CullingOn.Set<FVisiblityRasterComputeBinningCS::FCulling>(true);
+	BinningPermutationVector_CullingOn.Set<FVisiblityRasterComputeBinningCS::FTemporalLayering>(bTemporalLayering);
+	TShaderMapRef<FVisiblityRasterComputeBinningCS> ComputeShaderBinning_CullingOn(ViewInfo.ShaderMap, BinningPermutationVector_CullingOn);
+
+	FVisiblityRasterComputeRasterizeCS::FPermutationDomain RasterPermutationVector;
+	RasterPermutationVector.Set<FVisiblityRasterComputeRasterizeCS::FGroupSize>(GroupSize);
+	RasterPermutationVector.Set<FVisiblityRasterComputeRasterizeCS::FTemporalLayering>(bTemporalLayering);
+	TShaderMapRef<FVisiblityRasterComputeRasterizeCS> ComputeShaderRaster(ViewInfo.ShaderMap, RasterPermutationVector);
+
+	TShaderMapRef<FVisiblityRasterComputePrepareDepthGridCS> ComputeShaderPrepareDepthGrid(ViewInfo.ShaderMap);	
+	FVisiblityRasterComputePrepareDepthGridCS::FParameters* PrepDepthGridParameters = GraphBuilder.AllocParameters<FVisiblityRasterComputePrepareDepthGridCS::FParameters>();
+	PrepDepthGridParameters->OutputResolution = Out.SuperResolution;
+	PrepDepthGridParameters->SceneDepthTexture = SceneDepthTexture;
+	PrepDepthGridParameters->OutVisTileDepthGrid = OutVisTileDepthGridUAV; 
+	PrepDepthGridParameters->OutDepthCovTexture = DepthCovTextureUAV;
+	PrepDepthGridParameters->TileRes = TileGridRes;
+	PrepDepthGridParameters->TileSize = TileSize;
+	PrepDepthGridParameters->RcpTileSize = 1.0 / TileSize;
+	PrepDepthGridParameters->SqrTileSize = TileSize * TileSize;
+
+	FComputeShaderUtils::AddPass(GraphBuilder, RDG_EVENT_NAME("HairStrandsVisibilityComputeRasterPrepDepthGrid"), ComputeShaderPrepareDepthGrid, PrepDepthGridParameters, FIntVector(TileGridRes.X*TileGridRes.Y, 1, 1));
 
 	for (const FHairStrandsMacroGroupData& MacroGroup : MacroGroupDatas)
 	{
@@ -3102,55 +3194,109 @@ static FRasterComputeOutput AddVisibilityComputeRasterPass(
 
 		for (const FHairStrandsMacroGroupData::PrimitiveInfo& PrimitiveInfo : PrimitiveSceneInfos)
 		{
-			FVisiblityRasterComputeCS::FParameters* PassParameters = GraphBuilder.AllocParameters<FVisiblityRasterComputeCS::FParameters>();
-			PassParameters->OutputResolution = Out.SuperResolution;
-			PassParameters->ResolutionMultiplier = Out.ResolutionMultiplier;
-			PassParameters->MacroGroupId = MacroGroup.MacroGroupId;
-			PassParameters->DispatchCountX = 1;
-			PassParameters->MaxRasterCount = FMath::Clamp(GHairVisibilityComputeRaster_MaxPixelCount, 1, 256);			
-			PassParameters->FrameIdMod8 = FrameIdMode8;
-			PassParameters->HairMaterialId = PrimitiveInfo.MaterialId;
-			PassParameters->ViewUniformBuffer = ViewUniformShaderParameters;
-			PassParameters->SceneDepthTexture = SceneDepthTexture;
-			PassParameters->OutHairCountTexture = HairCountTextureUAV;
-			PassParameters->OutVisibilityTexture0 = VisibilityTexture0UAV;
-			PassParameters->OutVisibilityTexture1 = VisibilityTexture1UAV;
-			PassParameters->OutVisibilityTexture2 = VisibilityTexture2UAV;
-			PassParameters->OutVisibilityTexture3 = VisibilityTexture3UAV;
-			
-			check(PrimitiveInfo.PublicDataPtr);
-			const FHairGroupPublicData* HairGroupPublicData = PrimitiveInfo.PublicDataPtr;
+			check(PrimitiveInfo.Mesh && PrimitiveInfo.Mesh->Elements.Num() > 0);
+			const FHairGroupPublicData* HairGroupPublicData = reinterpret_cast<const FHairGroupPublicData*>(PrimitiveInfo.Mesh->Elements[0].VertexFactoryUserData);
+			check(HairGroupPublicData);
 
 			const FHairGroupPublicData::FVertexFactoryInput& VFInput = HairGroupPublicData->VFInput;
-			PassParameters->HairStrandsVF_PositionBuffer	= VFInput.Strands.PositionBuffer.SRV;
-			PassParameters->HairStrandsVF_PositionOffsetBuffer = VFInput.Strands.PositionOffsetBuffer.SRV;
-			PassParameters->HairStrandsVF_VertexCount		= VFInput.Strands.VertexCount;
-			PassParameters->HairStrandsVF_Radius			= VFInput.Strands.HairRadius;
-			PassParameters->HairStrandsVF_RootScale			= VFInput.Strands.HairRootScale;
-			PassParameters->HairStrandsVF_TipScale			= VFInput.Strands.HairTipScale;
-			PassParameters->HairStrandsVF_Length			= VFInput.Strands.HairLength;
-			PassParameters->HairStrandsVF_bUseStableRasterization = VFInput.Strands.bUseStableRasterization ? 1 : 0;
-			PassParameters->HairStrandsVF_Density			= VFInput.Strands.HairDensity;
-			PassParameters->HairStrandsVF_LocalToWorldPrimitiveTransform = FMatrix44f(VFInput.LocalToWorldTransform.ToMatrixWithScale()); // LWC_TODO: Precision loss
-
-			const bool bCullingEnable = HairGroupPublicData->GetCullingResultAvailable();
+			const FMatrix44d LocalToWorldPrimitiveTransform = VFInput.LocalToWorldTransform.ToMatrixWithScale(); 
+
+			const FVector4f ScreenPosition = FVector4f(ViewInfo.ViewMatrices.GetViewProjectionMatrix().TransformPosition(LocalToWorldPrimitiveTransform.GetOrigin() + FVector4d(0, 0, 165, 0))); //hardcoded - TODO find a way to get correct offset to origin
+			const bool bCullingEnable = GHairVisibilityComputeRaster_Culling ? HairGroupPublicData->GetCullingResultAvailable() : false;
+
+			//reset buffers
+			uint32 IndexGridClearValues[4] = { 0x0,0x0,0x0,0x0 };
+			AddClearUAVPass(GraphBuilder, OutVisTileBinningGridUAV, IndexGridClearValues);
+			AddClearUAVPass(GraphBuilder, OutVisTileArgsUAV, 0u);
+
+			//binning
+			FVisiblityRasterComputeBinningCS::FParameters* BinningParameters = GraphBuilder.AllocParameters<FVisiblityRasterComputeBinningCS::FParameters>();
+			BinningParameters->OutputResolution = Out.SuperResolution;
+			BinningParameters->ResolutionMultiplier = Out.ResolutionMultiplier;
+			BinningParameters->MacroGroupId = MacroGroup.MacroGroupId;
+			BinningParameters->HairMaterialId = PrimitiveInfo.MaterialId;
+			BinningParameters->ViewUniformBuffer = ViewUniformShaderParameters;
+			BinningParameters->SceneDepthTexture = SceneDepthTexture;
+			BinningParameters->OutVisTilePrims = OutVisTilePrimsUAV;
+			BinningParameters->OutVisTileBinningGrid = OutVisTileBinningGridUAV;  
+			BinningParameters->VisTileDepthGrid = OutVisTileDepthGrid;
+			BinningParameters->OutVisTileArgs = OutVisTileArgsUAV;
+			BinningParameters->OutVisTileData = OutVisTileDataUAV;
+
+			BinningParameters->HairStrandsVF_PositionBuffer = VFInput.Strands.PositionBuffer.SRV;
+			BinningParameters->HairStrandsVF_PositionOffsetBuffer = VFInput.Strands.PositionOffsetBuffer.SRV;
+			BinningParameters->HairStrandsVF_VertexCount = VFInput.Strands.VertexCount;
+			BinningParameters->HairStrandsVF_Radius = VFInput.Strands.HairRadius;
+			BinningParameters->HairStrandsVF_RootScale = VFInput.Strands.HairRootScale;
+			BinningParameters->HairStrandsVF_TipScale = VFInput.Strands.HairTipScale;
+			BinningParameters->HairStrandsVF_Length = VFInput.Strands.HairLength;
+			BinningParameters->HairStrandsVF_bUseStableRasterization = VFInput.Strands.bUseStableRasterization ? 1 : 0;
+			BinningParameters->HairStrandsVF_Density = VFInput.Strands.HairDensity;
+			BinningParameters->HairStrandsVF_LocalToWorldPrimitiveTransform = FMatrix44f(LocalToWorldPrimitiveTransform); // LWC_TODO: Precision loss
+
+			BinningParameters->TileBaseDepth = ScreenPosition.Z / ScreenPosition.W;
+			BinningParameters->TileLODScale = GHairVisibilityComputeRaster_CLOD ? 1.0 / ScreenPosition.W : 1.0;
+			BinningParameters->TileRes = TileGridRes;
+			BinningParameters->VisTileBinningGridTex = OutVisTileBinningGrid;
+			BinningParameters->TileSize = TileSize;
+			BinningParameters->RcpTileSize = 1.0/TileSize;
+			BinningParameters->SqrTileSize = TileSize*TileSize;
+			BinningParameters->NumBinners = NumBinners;
+			BinningParameters->RcpNumBinners = 1.0/NumBinners;
+
 			if (bCullingEnable)
 			{
 				FRDGImportedBuffer CullingIndirectBuffer = Register(GraphBuilder, HairGroupPublicData->GetDrawIndirectRasterComputeBuffer(), ERDGImportedBufferFlags::CreateSRV);
-				PassParameters->HairStrandsVF_CullingIndirectBuffer		= CullingIndirectBuffer.SRV;
-				PassParameters->HairStrandsVF_bIsCullingEnable			= bCullingEnable ? 1 : 0;
-				PassParameters->HairStrandsVF_CullingIndexBuffer		= RegisterAsSRV(GraphBuilder, HairGroupPublicData->GetCulledVertexIdBuffer());
-				PassParameters->HairStrandsVF_CullingRadiusScaleBuffer	= RegisterAsSRV(GraphBuilder, HairGroupPublicData->GetCulledVertexRadiusScaleBuffer());
-				PassParameters->IndirectBufferArgs						= CullingIndirectBuffer.Buffer;
-
-				FComputeShaderUtils::AddPass(GraphBuilder, RDG_EVENT_NAME("HairStrandsVisibilityComputeRaster(culling=on)"), ComputeShader_CullingOn, PassParameters, CullingIndirectBuffer.Buffer, 0);
-			}
-			else
-			{
-				const FIntVector DispatchCount = ComputeDispatchCount(PassParameters->HairStrandsVF_VertexCount, GroupSize);
-				PassParameters->DispatchCountX = DispatchCount.X;
-				FComputeShaderUtils::AddPass(GraphBuilder, RDG_EVENT_NAME("HairStrandsVisibilityComputeRaster(culling=off)"), ComputeShader_CullingOff, PassParameters, DispatchCount);
+				BinningParameters->HairStrandsVF_CullingIndirectBuffer = CullingIndirectBuffer.SRV;
+				BinningParameters->HairStrandsVF_bIsCullingEnable = bCullingEnable ? 1 : 0;
+				BinningParameters->HairStrandsVF_CullingIndexBuffer = RegisterAsSRV(GraphBuilder, HairGroupPublicData->GetCulledVertexIdBuffer());
+				BinningParameters->HairStrandsVF_CullingRadiusScaleBuffer = RegisterAsSRV(GraphBuilder, HairGroupPublicData->GetCulledVertexRadiusScaleBuffer());
+				BinningParameters->IndirectBufferArgs = CullingIndirectBuffer.Buffer;
 			}
+
+			FComputeShaderUtils::AddPass(GraphBuilder, RDG_EVENT_NAME("HairStrandsVisibilityComputeRasterBinning(culling=%s)", bCullingEnable ? TEXT("On") : TEXT("Off")), bCullingEnable ? ComputeShaderBinning_CullingOn : ComputeShaderBinning_CullingOff, BinningParameters, FIntVector(NumBinners, 1, 1));
+			
+			//raster pass
+			FVisiblityRasterComputeRasterizeCS::FParameters* RasterParameters = GraphBuilder.AllocParameters<FVisiblityRasterComputeRasterizeCS::FParameters>();
+			RasterParameters->OutputResolution = Out.SuperResolution;
+			RasterParameters->OutputResolutionf = FVector2f(Out.SuperResolution.X, Out.SuperResolution.Y);
+			RasterParameters->ResolutionMultiplier = Out.ResolutionMultiplier;
+			RasterParameters->MacroGroupId = MacroGroup.MacroGroupId;
+			RasterParameters->MaxRasterCount = TileSize; //TODO: implement tile clipping in raster pass to exisiting avoid edge case
+			RasterParameters->FrameIdMod8 = FrameIdMode8;
+			RasterParameters->HairMaterialId = PrimitiveInfo.MaterialId;
+			RasterParameters->ViewUniformBuffer = ViewUniformShaderParameters;
+			RasterParameters->SceneDepthTexture = SceneDepthTexture;
+			RasterParameters->OutHairCountTexture = HairCountTextureUAV;
+			RasterParameters->OutDepthCovTexture = DepthCovTextureUAV;
+			RasterParameters->OutPrimMatTexture = PrimMatTextureUAV;
+
+			RasterParameters->HairStrandsVF_PositionBuffer	= VFInput.Strands.PositionBuffer.SRV;
+			RasterParameters->HairStrandsVF_PositionOffsetBuffer = VFInput.Strands.PositionOffsetBuffer.SRV;
+			RasterParameters->HairStrandsVF_VertexCount		= VFInput.Strands.VertexCount;
+			RasterParameters->HairStrandsVF_Radius			= VFInput.Strands.HairRadius;
+			RasterParameters->HairStrandsVF_RootScale			= VFInput.Strands.HairRootScale;
+			RasterParameters->HairStrandsVF_TipScale			= VFInput.Strands.HairTipScale;
+			RasterParameters->HairStrandsVF_Length			= VFInput.Strands.HairLength;
+			RasterParameters->HairStrandsVF_bUseStableRasterization = VFInput.Strands.bUseStableRasterization ? 1 : 0;
+			RasterParameters->HairStrandsVF_Density			= VFInput.Strands.HairDensity;
+			RasterParameters->HairStrandsVF_LocalToWorldPrimitiveTransform = FMatrix44f(LocalToWorldPrimitiveTransform); // LWC_TODO: Precision loss
+
+			RasterParameters->VisTilePrims = GraphBuilder.CreateSRV(OutVisTilePrims);
+			RasterParameters->RWVisTileDepthGrid = OutVisTileDepthGridUAV;
+			RasterParameters->VisTileArgs = GraphBuilder.CreateSRV(OutVisTileArgs);
+			RasterParameters->RWVisTileData = OutVisTileDataUAV;
+			RasterParameters->VisTileIndirect = GraphBuilder.CreateSRV(OutVisTileIndirect);
+			RasterParameters->IndirectBufferArgs = OutVisTileIndirect;
+			RasterParameters->TileRes = TileGridRes;
+			RasterParameters->TileLODScale = GHairVisibilityComputeRaster_CLOD ? 1.0 / ScreenPosition.W : 1.0;
+			RasterParameters->TileSize = TileSize;
+			RasterParameters->RcpTileSize = 1.0 / TileSize;
+			RasterParameters->SqrTileSize = TileSize * TileSize;
+			RasterParameters->NumRasterizers = NumRasterizers;
+			RasterParameters->RcpNumRasterizers = 1.0 / NumRasterizers;
+
+			FComputeShaderUtils::AddPass(GraphBuilder, RDG_EVENT_NAME("HairStrandsVisibilityComputeRasterRaster(tiled)"), ComputeShaderRaster, RasterParameters, FIntVector(NumRasterizers, 1, 1));
 		}
 	}
 
@@ -3627,22 +3773,10 @@ void RenderHairStrandsVisibilityBuffer(
 					VisibilityData.MaxSampleCount,
 					SceneDepthTexture);
 
-				// Merge this pass within the compaction pass
-				FHairPrimaryTransmittance ViewTransmittance;
-				{
-					ViewTransmittance.TransmittanceTexture = AddHairHairCountToTransmittancePass(
-						GraphBuilder,
-						View,
-						RasterOutput.HairCountTexture);
-
-					ViewTransmittance.HairCountTextureUint = RasterOutput.HairCountTexture;
-					VisibilityData.ViewHairCountUintTexture = ViewTransmittance.HairCountTextureUint;
-				}
-
 				// Generate Tile data
-				if (ViewTransmittance.TransmittanceTexture && bGenerateTile)
+				if (RasterOutput.PrimMatTexture && bGenerateTile)
 				{
-					VisibilityData.TileData = AddHairStrandsGenerateTilesPass(GraphBuilder, View, ViewTransmittance.TransmittanceTexture);
+					VisibilityData.TileData = AddHairStrandsGenerateTilesPass(GraphBuilder, View, RasterOutput.HairCountTexture);
 				}
 
 				{
@@ -3657,7 +3791,6 @@ void RenderHairStrandsVisibilityBuffer(
 							VisibilityData.MaxSampleCount,
 							RasterOutput,
 							VisibilityData.TileData,
-							ViewTransmittance.TransmittanceTexture, // TODO tile
 							NodeCounter,
 							CompactNodeIndex,
 							CompactNodeVis,
@@ -3704,6 +3837,7 @@ void RenderHairStrandsVisibilityBuffer(
 						VisibilityData.CoverageTexture = CoverageTexture;
 						VisibilityData.HairOnlyDepthTexture = HairOnlyDepthTexture;
 						VisibilityData.NodeData = CompactNodeData;
+						VisibilityData.NodeVisData = CompactNodeVis;
 						VisibilityData.NodeCoord = CompactNodeCoord;
 						VisibilityData.NodeIndirectArg = IndirectArgsBuffer;
 						VisibilityData.NodeCount = NodeCounter;
