Fix static mesh commands being rendered in both Depth and DepthWithVelocity passes.
They should only be rendered in DepthWithVelocity.

diff --git a/Engine/Source/Runtime/Renderer/Private/SceneVisibility.cpp b/Engine/Source/Runtime/Renderer/Private/SceneVisibility.cpp
--- a/Engine/Source/Runtime/Renderer/Private/SceneVisibility.cpp
+++ b/Engine/Source/Runtime/Renderer/Private/SceneVisibility.cpp
@@ -2518,6 +2518,7 @@ struct FRelevancePacket
 		const EShadingPath ShadingPath = Scene->GetShadingPath();
 		const bool bMobileMaskedInEarlyPass = (ShadingPath == EShadingPath::Mobile) &&  Scene->EarlyZPassMode == DDM_MaskedOnly;
 		const bool bMobileBasePassAlwaysUsesCSM = (ShadingPath == EShadingPath::Mobile) && MobileBasePassAlwaysUsesCSM(Scene->GetShaderPlatform());
+		const bool bVelocityPassWritesDepth = Scene->EarlyZPassMode == DDM_AllOpaqueNoVelocity;
 		const bool bHLODActive = Scene->SceneLODHierarchy.IsActive();
 		const FHLODVisibilityState* const HLODState = bHLODActive && ViewState ? &ViewState->HLODVisibilityState : nullptr;
 		float MaxDrawDistanceScale = GetCachedScalabilityCVars().ViewDistanceScale;
@@ -2622,10 +2623,41 @@ struct FRelevancePacket
 							&& (ViewRelevance.bRenderInMainPass || ViewRelevance.bRenderCustomDepth || ViewRelevance.bRenderInDepthPass)
 							&& !bHiddenByHLODFade)
 						{
-							if (StaticMeshRelevance.bUseForDepthPass && (bDrawDepthOnly || (bMobileMaskedInEarlyPass && ViewRelevance.bMasked)))
+							// Add velocity commands first to track for case where velocity pass writes depth.
+							bool bIsMeshInVelocityPass = false;
+							if (StaticMeshRelevance.bUseForMaterial && ViewRelevance.bRenderInMainPass)
 							{
-								DrawCommandPacket.AddCommandsForMesh(PrimitiveIndex, PrimitiveSceneInfo, StaticMeshRelevance, StaticMesh, Scene, bCanCache, EMeshPass::DepthPass);
+								if (ViewRelevance.HasVelocity())
+								{
+									const FPrimitiveSceneProxy* PrimitiveSceneProxy = PrimitiveSceneInfo->Proxy;
+
+									if (FVelocityMeshProcessor::PrimitiveHasVelocityForView(View, PrimitiveSceneProxy))
+									{
+										if (ViewRelevance.bVelocityRelevance &&
+											FOpaqueVelocityMeshProcessor::PrimitiveCanHaveVelocity(View.GetShaderPlatform(), PrimitiveSceneProxy) &&
+											FOpaqueVelocityMeshProcessor::PrimitiveHasVelocityForFrame(PrimitiveSceneProxy))
+										{
+											DrawCommandPacket.AddCommandsForMesh(PrimitiveIndex, PrimitiveSceneInfo, StaticMeshRelevance, StaticMesh, Scene, bCanCache, EMeshPass::Velocity);
+											bIsMeshInVelocityPass = true;
+										}
 
+										if (ViewRelevance.bOutputsTranslucentVelocity &&
+											FTranslucentVelocityMeshProcessor::PrimitiveCanHaveVelocity(View.GetShaderPlatform(), PrimitiveSceneProxy) &&
+											FTranslucentVelocityMeshProcessor::PrimitiveHasVelocityForFrame(PrimitiveSceneProxy))
+										{
+											DrawCommandPacket.AddCommandsForMesh(PrimitiveIndex, PrimitiveSceneInfo, StaticMeshRelevance, StaticMesh, Scene, bCanCache, EMeshPass::TranslucentVelocity);
+										}
+									}
+								}
+							}
+
+							// Add depth commands.
+							if (StaticMeshRelevance.bUseForDepthPass && (bDrawDepthOnly || (bMobileMaskedInEarlyPass && ViewRelevance.bMasked)))
+							{
+								if (!(bIsMeshInVelocityPass && bVelocityPassWritesDepth))
+								{
+									DrawCommandPacket.AddCommandsForMesh(PrimitiveIndex, PrimitiveSceneInfo, StaticMeshRelevance, StaticMesh, Scene, bCanCache, EMeshPass::DepthPass);
+								}
 #if RHI_RAYTRACING
 								if (IsRayTracingEnabled())
 								{
@@ -2709,29 +2741,6 @@ struct FRelevancePacket
 									}
 								}
 #endif
-
-								if (ViewRelevance.HasVelocity())
-								{
-									const FPrimitiveSceneProxy* PrimitiveSceneProxy = PrimitiveSceneInfo->Proxy;
-
-									if (FVelocityMeshProcessor::PrimitiveHasVelocityForView(View, PrimitiveSceneProxy))
-									{
-										if (ViewRelevance.bVelocityRelevance &&
-											FOpaqueVelocityMeshProcessor::PrimitiveCanHaveVelocity(View.GetShaderPlatform(), PrimitiveSceneProxy) &&
-											FOpaqueVelocityMeshProcessor::PrimitiveHasVelocityForFrame(PrimitiveSceneProxy))
-										{
-											DrawCommandPacket.AddCommandsForMesh(PrimitiveIndex, PrimitiveSceneInfo, StaticMeshRelevance, StaticMesh, Scene, bCanCache, EMeshPass::Velocity);
-										}
-
-										if (ViewRelevance.bOutputsTranslucentVelocity &&
-											FTranslucentVelocityMeshProcessor::PrimitiveCanHaveVelocity(View.GetShaderPlatform(), PrimitiveSceneProxy) &&
-											FTranslucentVelocityMeshProcessor::PrimitiveHasVelocityForFrame(PrimitiveSceneProxy))
-										{
-											DrawCommandPacket.AddCommandsForMesh(PrimitiveIndex, PrimitiveSceneInfo, StaticMeshRelevance, StaticMesh, Scene, bCanCache, EMeshPass::TranslucentVelocity);
-										}
-									}
-								}
-
 								++NumVisibleStaticMeshElements;
 
 								INC_DWORD_STAT_BY(STAT_StaticMeshTriangles, StaticMesh.GetNumPrimitives());
