MGPU:  Fixes for cached shadow maps.  Shadow maps are copied cross GPU where necessary.

#rnx

diff --git a/Engine/Source/Runtime/Renderer/Private/ShadowDepthRendering.cpp b/Engine/Source/Runtime/Renderer/Private/ShadowDepthRendering.cpp
--- a/Engine/Source/Runtime/Renderer/Private/ShadowDepthRendering.cpp
+++ b/Engine/Source/Runtime/Renderer/Private/ShadowDepthRendering.cpp
@@ -1048,11 +1048,33 @@ BEGIN_SHADER_PARAMETER_STRUCT(FShadowDepthPassParameters, )
 	RENDER_TARGET_BINDING_SLOTS()
 END_SHADER_PARAMETER_STRUCT()
 
+#if WITH_MGPU
+void CopyCachedShadowMapCrossGPU(FRHICommandList& RHICmdList, FRHITexture* ShadowDepthTexture, FRHIGPUMask SourceGPUMask)
+{
+	if (SourceGPUMask != FRHIGPUMask::All())
+	{
+		uint32 SourceGPUIndex = SourceGPUMask.GetFirstIndex();
+
+		TArray<FTransferResourceParams, TFixedAllocator<MAX_NUM_GPUS>> CrossGPUTransferBuffers;
+		for (uint32 DestGPUIndex : FRHIGPUMask::All())
+		{
+			if (!SourceGPUMask.Contains(DestGPUIndex))
+			{
+				CrossGPUTransferBuffers.Add(FTransferResourceParams(ShadowDepthTexture, SourceGPUIndex, DestGPUIndex, false, false));
+			}
+		}
+
+		RHICmdList.TransferResources(CrossGPUTransferBuffers);
+	}
+}
+#endif  // WITH_MGPU
+
 void FProjectedShadowInfo::RenderDepth(
 	FRDGBuilder& GraphBuilder,
 	const FSceneRenderer* SceneRenderer,
 	FRDGTextureRef ShadowDepthTexture,
-	bool bDoParallelDispatch)
+	bool bDoParallelDispatch,
+	bool bDoCrossGPUCopy)
 {
 #if WANTS_DRAW_MESH_EVENTS
 	FString EventName;
@@ -1114,6 +1136,11 @@ void FProjectedShadowInfo::RenderDepth(
 
 	ShadowDepthPass.BuildRenderingCommands(GraphBuilder, Scene->GPUScene, PassParameters->InstanceCullingDrawParams);
 
+#if WITH_MGPU
+	// Need to fetch GPU mask outside "AddPass", as it's not updated during pass execution
+	FRHIGPUMask GPUMask = GraphBuilder.RHICmdList.GetGPUMask();
+#endif
+
 	if (bDoParallelDispatch)
 	{
 		RDG_WAIT_FOR_TASKS_CONDITIONAL(GraphBuilder, IsShadowDepthPassWaitForTasksEnabled());
@@ -1122,10 +1149,21 @@ void FProjectedShadowInfo::RenderDepth(
 			RDG_EVENT_NAME("ShadowDepthPassParallel"),
 			PassParameters,
 			ERDGPassFlags::Raster | ERDGPassFlags::SkipRenderPass,
-			[this, PassParameters](FRHICommandListImmediate& RHICmdList)
+			[this, PassParameters
+#if WITH_MGPU
+			, ShadowDepthTexture, GPUMask, bDoCrossGPUCopy
+#endif
+			](FRHICommandListImmediate& RHICmdList)
 		{
 			FShadowParallelCommandListSet ParallelCommandListSet(RHICmdList, *ShadowDepthView, *this, FParallelCommandListBindings(PassParameters));
 			ShadowDepthPass.DispatchDraw(&ParallelCommandListSet, RHICmdList, &PassParameters->InstanceCullingDrawParams);
+
+#if WITH_MGPU
+			if (bDoCrossGPUCopy)
+			{
+				CopyCachedShadowMapCrossGPU(RHICmdList, ShadowDepthTexture->GetRHI(), GPUMask);
+			}
+#endif
 		});
 	}
 	else
@@ -1134,10 +1172,21 @@ void FProjectedShadowInfo::RenderDepth(
 			RDG_EVENT_NAME("ShadowDepthPass"),
 			PassParameters,
 			ERDGPassFlags::Raster,
-			[this, PassParameters](FRHICommandList& RHICmdList)
+			[this, PassParameters
+#if WITH_MGPU
+			, ShadowDepthTexture, GPUMask, bDoCrossGPUCopy
+#endif
+			](FRHICommandList& RHICmdList)
 		{
 			SetStateForView(RHICmdList);
 			ShadowDepthPass.DispatchDraw(nullptr, RHICmdList, &PassParameters->InstanceCullingDrawParams);
+
+#if WITH_MGPU
+			if (bDoCrossGPUCopy)
+			{
+				CopyCachedShadowMapCrossGPU(RHICmdList, ShadowDepthTexture->GetRHI(), GPUMask);
+			}
+#endif
 		});
 	}
 }
@@ -1566,7 +1615,11 @@ void FSceneRenderer::RenderShadowDepthMapAtlases(FRDGBuilder& GraphBuilder)
 				SetLightEventForShadow(ProjectedShadowInfo);
 
 				const bool bParallelDispatch = true;
-				ProjectedShadowInfo->RenderDepth(GraphBuilder, this, AtlasDepthTexture, bParallelDispatch);
+				bool bDoCrossGPUCopy = false;
+#if WITH_MGPU
+				bDoCrossGPUCopy = IsShadowCached(ProjectedShadowInfo);
+#endif
+				ProjectedShadowInfo->RenderDepth(GraphBuilder, this, AtlasDepthTexture, bParallelDispatch, bDoCrossGPUCopy);
 			}
 		}
 
@@ -1580,7 +1633,11 @@ void FSceneRenderer::RenderShadowDepthMapAtlases(FRDGBuilder& GraphBuilder)
 				SetLightEventForShadow(ProjectedShadowInfo);
 
 				const bool bParallelDispatch = false;
-				ProjectedShadowInfo->RenderDepth(GraphBuilder, this, AtlasDepthTexture, bParallelDispatch);
+				bool bDoCrossGPUCopy = false;
+#if WITH_MGPU
+				bDoCrossGPUCopy = IsShadowCached(ProjectedShadowInfo);
+#endif
+				ProjectedShadowInfo->RenderDepth(GraphBuilder, this, AtlasDepthTexture, bParallelDispatch, bDoCrossGPUCopy);
 			}
 		}
 
@@ -1852,7 +1909,8 @@ void FSceneRenderer::RenderShadowDepthMaps(FRDGBuilder& GraphBuilder, FInstanceC
 
 		{
 			const bool bDoParallelDispatch = IsParallelDispatchEnabled(ProjectedShadowInfo, ShaderPlatform);
-			ProjectedShadowInfo->RenderDepth(GraphBuilder, this, ShadowDepthTexture, bDoParallelDispatch);
+			const bool bDoCrossGPUCopy = false;
+			ProjectedShadowInfo->RenderDepth(GraphBuilder, this, ShadowDepthTexture, bDoParallelDispatch, bDoCrossGPUCopy);
 		}
 
 		if (bNaniteEnabled &&
@@ -2011,7 +2069,8 @@ void FSceneRenderer::RenderShadowDepthMaps(FRDGBuilder& GraphBuilder, FInstanceC
 				AddClearShadowDepthPass(GraphBuilder, PreshadowCacheTexture, ProjectedShadowInfo);
 
 				const bool bParallelDispatch = IsParallelDispatchEnabled(ProjectedShadowInfo, ShaderPlatform);
-				ProjectedShadowInfo->RenderDepth(GraphBuilder, this, PreshadowCacheTexture, bParallelDispatch);
+				const bool bDoCrossGPUCopy = true;
+				ProjectedShadowInfo->RenderDepth(GraphBuilder, this, PreshadowCacheTexture, bParallelDispatch, bDoCrossGPUCopy);
 				ProjectedShadowInfo->bDepthsCached = true;
 			}
 		}
diff --git a/Engine/Source/Runtime/Renderer/Private/ShadowRendering.h b/Engine/Source/Runtime/Renderer/Private/ShadowRendering.h
--- a/Engine/Source/Runtime/Renderer/Private/ShadowRendering.h
+++ b/Engine/Source/Runtime/Renderer/Private/ShadowRendering.h
@@ -562,7 +562,8 @@ public:
 		FRDGBuilder& GraphBuilder,
 		const FSceneRenderer* SceneRenderer,
 		FRDGTextureRef ShadowDepthTexture,
-		bool bDoParallelDispatch);
+		bool bDoParallelDispatch,
+		bool bDoCrossGPUCopy);
 
 	void BeginRenderRayTracedDistanceFieldProjection(
 		FRDGBuilder& GraphBuilder,
