Strata - clean up unused tile size

diff --git a/Engine/Shaders/Private/Strata/StrataMaterialClassification.usf b/Engine/Shaders/Private/Strata/StrataMaterialClassification.usf
--- a/Engine/Shaders/Private/Strata/StrataMaterialClassification.usf
+++ b/Engine/Shaders/Private/Strata/StrataMaterialClassification.usf
@@ -18,7 +18,6 @@ uint EncodeTile(uint2 TileCoord)
 }
 
 #if SHADER_TILE_CATEGORIZATION
-int TileSize;
 int bRectPrimitive;
 int2 ViewResolution;
 uint MaxBytesPerPixel;
diff --git a/Engine/Source/Runtime/Renderer/Private/Strata/Strata.cpp b/Engine/Source/Runtime/Renderer/Private/Strata/Strata.cpp
--- a/Engine/Source/Runtime/Renderer/Private/Strata/Strata.cpp
+++ b/Engine/Source/Runtime/Renderer/Private/Strata/Strata.cpp
@@ -425,7 +425,6 @@ class FStrataMaterialTileClassificationPassCS : public FGlobalShader
 
 	BEGIN_SHADER_PARAMETER_STRUCT(FParameters, )
 		SHADER_PARAMETER_STRUCT_REF(FViewUniformShaderParameters, ViewUniformBuffer)
-		SHADER_PARAMETER(int32, TileSize)
 		SHADER_PARAMETER(int32, bRectPrimitive)
 		SHADER_PARAMETER(FIntPoint, ViewResolution)
 		SHADER_PARAMETER(uint32, MaxBytesPerPixel)
@@ -771,7 +770,6 @@ void AddStrataMaterialClassificationPass(FRDGBuilder& GraphBuilder, const FMinim
 			TShaderMapRef<FStrataMaterialTileClassificationPassCS> ComputeShader(View.ShaderMap, PermutationVector);
 			FStrataMaterialTileClassificationPassCS::FParameters* PassParameters = GraphBuilder.AllocParameters<FStrataMaterialTileClassificationPassCS::FParameters>();
 			PassParameters->ViewUniformBuffer = View.ViewUniformBuffer;
-			PassParameters->TileSize = GetStrataBufferTileSize();	// STRATA_TODO not sure we want to tie the buffer tile optimisation for cache and the Categotisation tile size?
 			PassParameters->bRectPrimitive = GRHISupportsRectTopology ? 1 : 0;
 			PassParameters->ViewResolution = View.ViewRect.Size();
 			PassParameters->MaxBytesPerPixel = View.StrataSceneData->MaxBytesPerPixel;
