Real time sky light capture does not accept dynamic draws. This is made clear to avoid crashes.

diff --git a/Engine/Source/Runtime/GeometryFramework/Private/Components/BaseDynamicMeshComponent.cpp b/Engine/Source/Runtime/GeometryFramework/Private/Components/BaseDynamicMeshComponent.cpp
--- a/Engine/Source/Runtime/GeometryFramework/Private/Components/BaseDynamicMeshComponent.cpp
+++ b/Engine/Source/Runtime/GeometryFramework/Private/Components/BaseDynamicMeshComponent.cpp
@@ -8,8 +8,6 @@ using namespace UE::Geometry;
 UBaseDynamicMeshComponent::UBaseDynamicMeshComponent(const FObjectInitializer& ObjectInitializer)
 	: Super(ObjectInitializer)
 {
-	// including DynamicMeshComponent in sky captures can result in crashes in UE5.0, this should be a temporary workaround
-	bVisibleInRealTimeSkyCaptures = false;
 }
 
 
diff --git a/Engine/Source/Runtime/Renderer/Private/ReflectionEnvironmentRealTimeCapture.cpp b/Engine/Source/Runtime/Renderer/Private/ReflectionEnvironmentRealTimeCapture.cpp
--- a/Engine/Source/Runtime/Renderer/Private/ReflectionEnvironmentRealTimeCapture.cpp
+++ b/Engine/Source/Runtime/Renderer/Private/ReflectionEnvironmentRealTimeCapture.cpp
@@ -617,6 +617,23 @@ void FScene::AllocateAndCaptureFrameSkyEnvMap(
 									const FPrimitiveSceneProxy* PrimitiveSceneProxy = SkyMeshBatch.Proxy;
 									const FPrimitiveSceneInfo* PrimitiveSceneInfo = PrimitiveSceneProxy->GetPrimitiveSceneInfo();
 
+									// Real time sky light capture cannot render dynamic meshes for now.
+									// For those to be rendered we would need to specify a view to the PassMeshProcessor creation above.
+									// Dynamic draws uses temporary per frame & per view data (appended at the end of the GPUScene buffer).
+									// But the view is transient and data on it can morph, and correct data would need to be added to FGPUScenePrimitiveCollector (see UploadDynamicPrimitiveShaderDataForViewInternal)
+									bool bSkipDynamicMesh = false;
+									for (auto& Element : MeshBatch->Elements)
+									{
+										if (Element.PrimitiveIdMode == PrimID_DynamicPrimitiveShaderData)
+										{
+											bSkipDynamicMesh = true;
+										}
+									}
+									if (bSkipDynamicMesh)
+									{
+										continue;
+									}
+
 									const uint64 DefaultBatchElementMask = ~0ull;
 									PassMeshProcessor.AddMeshBatch(*MeshBatch, DefaultBatchElementMask, PrimitiveSceneProxy);
 								}
