MGPU:  Fix for Path Tracer.  Need to read result surface from correct GPU.

#rnx

diff --git a/Engine/Plugins/Experimental/OpenImageDenoise/Source/OpenImageDenoise/Private/Plugin.cpp b/Engine/Plugins/Experimental/OpenImageDenoise/Source/OpenImageDenoise/Private/Plugin.cpp
--- a/Engine/Plugins/Experimental/OpenImageDenoise/Source/OpenImageDenoise/Private/Plugin.cpp
+++ b/Engine/Plugins/Experimental/OpenImageDenoise/Source/OpenImageDenoise/Private/Plugin.cpp
@@ -31,7 +31,7 @@ enum class EDenoiseMode {
 	NO_AOVS = 3,   // denoise beauty only
 };
 
-static void Denoise(FRHICommandListImmediate& RHICmdList, FRHITexture2D* ColorTex, FRHITexture2D* AlbedoTex, FRHITexture2D* NormalTex, FRHITexture2D* OutputTex)
+static void Denoise(FRHICommandListImmediate& RHICmdList, FRHITexture2D* ColorTex, FRHITexture2D* AlbedoTex, FRHITexture2D* NormalTex, FRHITexture2D* OutputTex, FRHIGPUMask GPUMask)
 {
 	static IConsoleVariable* DenoiseModeCVar = IConsoleManager::Get().FindConsoleVariable(TEXT("r.PathTracing.Denoiser"));
 
@@ -50,6 +50,7 @@ static void Denoise(FRHICommandListImmediate& RHICmdList, FRHITexture2D* ColorTe
 	TArray<FLinearColor> RawNormal;
 	FReadSurfaceDataFlags ReadDataFlags(ERangeCompressionMode::RCM_MinMax);
 	ReadDataFlags.SetLinearToGamma(false);
+	ReadDataFlags.SetGPUIndex(GPUMask.GetFirstIndex());
 	RHICmdList.ReadSurfaceData(ColorTex, Rect, RawPixels, ReadDataFlags);
 	if (DenoiseMode == EDenoiseMode::DEFAULT || DenoiseMode == EDenoiseMode::CLEAN_AUX)
 	{
diff --git a/Engine/Source/Runtime/Renderer/Private/PathTracing.cpp b/Engine/Source/Runtime/Renderer/Private/PathTracing.cpp
--- a/Engine/Source/Runtime/Renderer/Private/PathTracing.cpp
+++ b/Engine/Source/Runtime/Renderer/Private/PathTracing.cpp
@@ -2003,19 +2003,23 @@ void FDeferredShadingSceneRenderer::RenderPathTracing(
 				DenoisedRadianceTexture = GraphBuilder.CreateTexture(RadianceTextureDesc, TEXT("PathTracer.DenoisedRadiance"), ERDGTextureFlags::MultiFrame);
 			}
 
+			// Need to read GPU mask outside Pass function, as the value is not refreshed inside the pass
+			FRHIGPUMask GPUMask = GraphBuilder.RHICmdList.GetGPUMask();
+
 			FDenoiseTextureParameters* DenoiseParameters = GraphBuilder.AllocParameters<FDenoiseTextureParameters>();
 			DenoiseParameters->InputTexture = RadianceTexture;
 			DenoiseParameters->InputAlbedo = AlbedoTexture;
 			DenoiseParameters->InputNormal = NormalTexture;
 			DenoiseParameters->OutputTexture = DenoisedRadianceTexture;
 			GraphBuilder.AddPass(RDG_EVENT_NAME("Path Tracer Denoiser Plugin"), DenoiseParameters, ERDGPassFlags::Readback, 
-				[DenoiseParameters, DenoiserMode](FRHICommandListImmediate& RHICmdList)
+				[DenoiseParameters, DenoiserMode, GPUMask](FRHICommandListImmediate& RHICmdList)
 				{
 					GPathTracingDenoiserFunc(RHICmdList,
 						DenoiseParameters->InputTexture->GetRHI()->GetTexture2D(),
 						DenoiseParameters->InputAlbedo->GetRHI()->GetTexture2D(),
 						DenoiseParameters->InputNormal->GetRHI()->GetTexture2D(),
-						DenoiseParameters->OutputTexture->GetRHI()->GetTexture2D());
+						DenoiseParameters->OutputTexture->GetRHI()->GetTexture2D(),
+						GPUMask);
 				}
 			);
 
diff --git a/Engine/Source/Runtime/Renderer/Public/PathTracingDenoiser.h b/Engine/Source/Runtime/Renderer/Public/PathTracingDenoiser.h
--- a/Engine/Source/Runtime/Renderer/Public/PathTracingDenoiser.h
+++ b/Engine/Source/Runtime/Renderer/Public/PathTracingDenoiser.h
@@ -5,6 +5,6 @@
 class FRHITexture;
 class FRHICommandListImmediate;
 
-using PathTracingDenoiserFunction = void(FRHICommandListImmediate& RHICmdList, FRHITexture2D* ColorTex, FRHITexture2D* AlbedoTex, FRHITexture2D* NormalTex, FRHITexture2D* OutputTex);
+using PathTracingDenoiserFunction = void(FRHICommandListImmediate& RHICmdList, FRHITexture2D* ColorTex, FRHITexture2D* AlbedoTex, FRHITexture2D* NormalTex, FRHITexture2D* OutputTex, FRHIGPUMask GPUMask);
 
 extern RENDERER_API PathTracingDenoiserFunction* GPathTracingDenoiserFunc;
