Strata - beginning of rough refraction of sub layers (with scene lighting split into untouched, sss and rough reffracted luminance)

- separate SSS, opaque rough refraction and scene color output when strata project setting is enabled.
- simple placeholder per pixel variance from roughness/layers computation
- SSS only reads from subsurface buffer when ORR is enabled
- Added simple combine pass to final scene lighting (no blur for now)
- Started adding tiling but reached UAV count limit. I will refactor code later to workaround that.
- Updated light rendering, reflection environment pixel shader. Next will be lumen and clustered deferred shading.

diff --git a/Engine/Shaders/Private/BasePassPixelShader.usf b/Engine/Shaders/Private/BasePassPixelShader.usf
--- a/Engine/Shaders/Private/BasePassPixelShader.usf
+++ b/Engine/Shaders/Private/BasePassPixelShader.usf
@@ -1709,6 +1709,10 @@ void FPixelShaderInOut_MainPS(
 	float3 EmissiveLuminance = 0.0f;
 	FStrataSubsurfaceData SSSData = (FStrataSubsurfaceData)0;
 	FStrataTopLayerData TopLayerData = (FStrataTopLayerData)0;
+	FStrataOpaqueRoughRefractionData OpaqueRoughRefractionData = (FStrataOpaqueRoughRefractionData)0;
+#if STRATA_OUTPUT_ROUGH_REFRACTION
+	OpaqueRoughRefractionData.OpaqueRoughRefractionEnabled = true;
+#endif
 
 	// Generate the strata material data to write out
 	FStrataAddressing StrataAddressing = GetStrataPixelDataByteOffset(PixelPos, uint2(ResolvedView.BufferSizeAndInvSize.xy), StrataStruct.MaxBytesPerPixel);
@@ -1719,14 +1723,17 @@ void FPixelShaderInOut_MainPS(
 		Dither,
 		StrataAddressing,
 		StrataPixelHeader, StrataData, MaterialParameters.CameraVector, bStrataSubsurfaceEnable, EmissiveLuminance,
-		SSSData, TopLayerData);
+		SSSData, TopLayerData, OpaqueRoughRefractionData);
 
 	// Write out MRT data
 	Out.StrataOutput[0] = RWStrataMaterialContainer.MaterialRenderTargets[0];
 	Out.StrataOutput[1] = RWStrataMaterialContainer.MaterialRenderTargets[1];
 	Out.StrataOutput[2] = StrataPackTopLayerData(TopLayerData);
+#if STRATA_OUTPUT_ROUGH_REFRACTION
+	StrataStruct.OpaqueRoughRefractionTextureUAV[PixelPos] = StrataPackOpaqueRoughRefractionData(OpaqueRoughRefractionData);
+#endif
 
-	// Only write SSS datta if needed
+	// Only write SSS data if needed
 	BRANCH
 	if(SSSData.bIsValid)
 	{
diff --git a/Engine/Shaders/Private/ClusteredDeferredShadingPixelShader.usf b/Engine/Shaders/Private/ClusteredDeferredShadingPixelShader.usf
--- a/Engine/Shaders/Private/ClusteredDeferredShadingPixelShader.usf
+++ b/Engine/Shaders/Private/ClusteredDeferredShadingPixelShader.usf
@@ -218,7 +218,7 @@ float4 GetSimpleLightLighting(
 			ShadowTerms,
 			Strata.MaterialTextureArray,
 			StrataAddressing,
-			StrataPixelHeader);
+			StrataPixelHeader).SceneColor;
 	#else
 		FSimpleDeferredLightData SimpleLightData = (FSimpleDeferredLightData)0;
 		SimpleLightData.TranslatedWorldPosition = LocalLight.LightPositionAndInvRadius.xyz;
@@ -368,7 +368,7 @@ float4 GetLightGridLocalLighting(
 			ShadowTerms,
 			Strata.MaterialTextureArray,
 			StrataAddressing,
-			StrataPixelHeader);
+			StrataPixelHeader).SceneColor;
 	#elif 1
 		// NOTE: uses AO data like tiled deferred (as opposed to forward path)
 		// TODO: Passing in 0,0 as SVPos, what is meaning of this? Only used if 'REFERENCE_QUALITY' is on.
diff --git a/Engine/Shaders/Private/DebugProbes.usf b/Engine/Shaders/Private/DebugProbes.usf
--- a/Engine/Shaders/Private/DebugProbes.usf
+++ b/Engine/Shaders/Private/DebugProbes.usf
@@ -309,6 +309,7 @@ void MainPS(
 
 		FStrataSubsurfaceData SSSData = (FStrataSubsurfaceData)0;
 		FStrataTopLayerData TopLayerData = (FStrataTopLayerData)0;
+		FStrataOpaqueRoughRefractionData OpaqueRoughRefractionData = (FStrataOpaqueRoughRefractionData)0;
 
 		float3 EmissiveLuminance = 0.0f;
 		bool bStrataSubsurfaceEnable = false;
@@ -321,7 +322,7 @@ void MainPS(
 			Dither,
 			StrataAddressing,
 			StrataPixelHeader, StrataData, -RayDir, bStrataSubsurfaceEnable, EmissiveLuminance,
-			SSSData, TopLayerData);
+			SSSData, TopLayerData, OpaqueRoughRefractionData);
 
 		// Write out MRTs. Note: this only works for deferred with basepassvelocity off
 		OutGBufferA = 0;
diff --git a/Engine/Shaders/Private/DeferredLightPixelShaders.usf b/Engine/Shaders/Private/DeferredLightPixelShaders.usf
--- a/Engine/Shaders/Private/DeferredLightPixelShaders.usf
+++ b/Engine/Shaders/Private/DeferredLightPixelShaders.usf
@@ -214,10 +214,18 @@ void DeferredLightPixelMain(
 #endif
 	float4 SVPos			: SV_POSITION,
 	out float4 OutColor		: SV_Target0
+#if STRATA_OPAQUE_ROUGH_REFRACTION_ENABLED
+	, out float3 OutOpaqueRoughRefractionSceneColor : SV_Target1
+	, out float3 OutSubSurfaceSceneColor : SV_Target2
+#endif
 	)
 {
 	const float2 PixelPos = SVPos.xy;
 	OutColor = 0;
+#if STRATA_OPAQUE_ROUGH_REFRACTION_ENABLED
+	OutOpaqueRoughRefractionSceneColor = 0;
+	OutSubSurfaceSceneColor = 0;
+#endif
 
 	// Convert input data (directional/local light)
 	FInputParams InputParams = (FInputParams)0;
@@ -259,7 +267,6 @@ void DeferredLightPixelMain(
 			LightMask = GetLocalLightAttenuation(DerivedParams.TranslatedWorldPosition, LightData, ToLight, L);
 		}
 
-		float4 Radiance = 0;
 		if (LightMask > 0)
 		{
 			FShadowTerms ShadowTerms = { StrataGetAO(StrataPixelHeader), 1.0, 1.0, InitHairTransmittanceData() };
@@ -272,7 +279,7 @@ void DeferredLightPixelMain(
 			GetShadowTerms(SceneDepth, PrecomputedShadowFactors, FakeShadingModelID, FakeContactShadowOpacity,
 				LightData, DerivedParams.TranslatedWorldPosition, L, LightAttenuation, Dither, ShadowTerms);
 
-			Radiance = StrataDeferredLighting(
+			FStrataDeferredLighting StrataLighting = StrataDeferredLighting(
 				LightData,
 				V,
 				L,
@@ -282,7 +289,14 @@ void DeferredLightPixelMain(
 				Strata.MaterialTextureArray,
 				StrataAddressing,
 				StrataPixelHeader);
+
+			OutColor += StrataLighting.SceneColor;
+#if STRATA_OPAQUE_ROUGH_REFRACTION_ENABLED
+			OutOpaqueRoughRefractionSceneColor += StrataLighting.OpaqueRoughRefractionSceneColor;
+			OutSubSurfaceSceneColor += StrataLighting.SubSurfaceSceneColor;
+#endif
 		}
+	}
 
 #else // STRATA_ENABLED
 
@@ -313,15 +327,20 @@ void DeferredLightPixelMain(
 
 		float4 Radiance = GetDynamicLighting(DerivedParams.TranslatedWorldPosition, DerivedParams.CameraVector, ScreenSpaceData.GBuffer, ScreenSpaceData.AmbientOcclusion, ScreenSpaceData.GBuffer.ShadingModelID, LightData, GetPerPixelLightAttenuation(InputParams.ScreenUV), Dither, uint2(InputParams.PixelPos), SurfaceShadow);
 
-#endif // STRATA_ENABLED
-
 		OutColor += Radiance;
 	}
 
+#endif // STRATA_ENABLED
+
 	// RGB:SceneColor Specular and Diffuse
 	// A:Non Specular SceneColor Luminance
 	// So we need PreExposure for both color and alpha
 	OutColor.rgba *= GetExposure();
+#if STRATA_OPAQUE_ROUGH_REFRACTION_ENABLED 
+	// Idem
+	OutOpaqueRoughRefractionSceneColor *= GetExposure();
+	OutSubSurfaceSceneColor *= GetExposure();
+#endif
 }
 #endif
 
diff --git a/Engine/Shaders/Private/LightAccumulator.ush b/Engine/Shaders/Private/LightAccumulator.ush
--- a/Engine/Shaders/Private/LightAccumulator.ush
+++ b/Engine/Shaders/Private/LightAccumulator.ush
@@ -160,3 +160,40 @@ FDeferredLightingSplit LightAccumulator_GetResultSplit(FLightAccumulator In)
 
 	return Ret;
 }
+
+struct FStrataDeferredLighting
+{
+	float4 SceneColor;						// Untouched scene luminance, alpha channel may contain scatterable light luminance.
+#if STRATA_OPAQUE_ROUGH_REFRACTION_ENABLED
+	float3 OpaqueRoughRefractionSceneColor;	// Luminance that is going to be blurier by the top layer roughness
+	float3 SubSurfaceSceneColor;			// Scene color that should go through the sub surface post process. There is not alpha: not needed because it is separated from specular already.
+#endif
+};
+
+FStrataDeferredLighting GetInitialisedStrataDeferredLighting()
+{
+	FStrataDeferredLighting Result = (FStrataDeferredLighting)0;
+	return Result;
+}
+
+void AccumulateStrataDeferredLighting(inout FStrataDeferredLighting StrataLighting, FLightAccumulator In, bool bDiffuseIsSubsurface, bool bIsToplayer)
+{
+	FDeferredLightingSplit DiffSpec = LightAccumulator_GetResultSplit(In);
+
+#if STRATA_OPAQUE_ROUGH_REFRACTION_ENABLED
+	if (bIsToplayer)
+	{
+		StrataLighting.SceneColor						+= DiffSpec.SpecularLighting + (bDiffuseIsSubsurface ? 0 : DiffSpec.DiffuseLighting);
+		StrataLighting.OpaqueRoughRefractionSceneColor	+= 0.0;
+		StrataLighting.SubSurfaceSceneColor				+=(bDiffuseIsSubsurface ? DiffSpec.DiffuseLighting : 0).rgb;
+	}
+	else // !bIsToplayer
+	{
+		StrataLighting.SceneColor						+= 0.0;
+		StrataLighting.OpaqueRoughRefractionSceneColor	+=(DiffSpec.SpecularLighting +(bDiffuseIsSubsurface ? 0 : DiffSpec.DiffuseLighting)).rgb;
+		StrataLighting.SubSurfaceSceneColor				+=(bDiffuseIsSubsurface ? DiffSpec.DiffuseLighting : 0).rgb;
+	}
+#else
+	StrataLighting.SceneColor += DiffSpec.DiffuseLighting + DiffSpec.SpecularLighting;
+#endif
+}
diff --git a/Engine/Shaders/Private/MaterialTemplate.ush b/Engine/Shaders/Private/MaterialTemplate.ush
--- a/Engine/Shaders/Private/MaterialTemplate.ush
+++ b/Engine/Shaders/Private/MaterialTemplate.ush
@@ -147,6 +147,7 @@
 #define STRATA_OPAQUE_DEFERRED			(STRATA_OPAQUE_MATERIAL && !FORWARD_SHADING)
 #define STRATA_TRANSLUCENT_FORWARD		(STRATA_TRANSLUCENT_MATERIAL && !FORWARD_SHADING)
 #define STRATA_FORWARD_SHADING			(FORWARD_SHADING && (STRATA_OPAQUE_MATERIAL || STRATA_TRANSLUCENT_MATERIAL))
+#define STRATA_OUTPUT_ROUGH_REFRACTION	(STRATA_ENABLED && STRATA_OPAQUE_ROUGH_REFRACTION_ENABLED && STRATA_MATERIAL_OUTPUT_OPAQUE_ROUGH_REFRACTIONS)
 
 #if TEMPLATE_USES_STRATA
 #define STRATA_INLINE_SHADING 1
diff --git a/Engine/Shaders/Private/PostProcessSubsurface.usf b/Engine/Shaders/Private/PostProcessSubsurface.usf
--- a/Engine/Shaders/Private/PostProcessSubsurface.usf
+++ b/Engine/Shaders/Private/PostProcessSubsurface.usf
@@ -441,6 +441,14 @@ SDiffuseAndSpecular ReconstructLighting(float2 UVSceneColor, uint ReconstructMet
 {
 	SDiffuseAndSpecular Ret;
 
+#if STRATA_OPAQUE_ROUGH_REFRACTION_ENABLED
+	// In this mode, we only read isolated diffuse lighting.
+	float3 CenterSample = SubsurfaceInput0_Texture.SampleLevel(SharedSubsurfaceSampler0, UVSceneColor, 0).rgb;
+	Ret.Diffuse = CenterSample;
+	Ret.Specular = 0.0f;
+	return Ret;
+#endif
+
 	// If SUBSURFACE_CHANNEL_MODE is 0, checkerboard is forced on
 #if SUBSURFACE_PROFILE_CHECKERBOARD || SUBSURFACE_CHANNEL_MODE == 0
 	{
diff --git a/Engine/Shaders/Private/ReflectionEnvironmentPixelShader.usf b/Engine/Shaders/Private/ReflectionEnvironmentPixelShader.usf
--- a/Engine/Shaders/Private/ReflectionEnvironmentPixelShader.usf
+++ b/Engine/Shaders/Private/ReflectionEnvironmentPixelShader.usf
@@ -249,6 +249,10 @@ float GetCloudVolumetricAOShadow(in float3 TranslatedWorldPosition)
 void ReflectionEnvironmentSkyLighting(
 	in float4 SvPosition : SV_Position,
 	out float4 OutColor : SV_Target0
+#if STRATA_OPAQUE_ROUGH_REFRACTION_ENABLED
+	, out float3 OutOpaqueRoughRefractionSceneColor : SV_Target1
+	, out float3 OutSubSurfaceSceneColor : SV_Target2
+#endif
 	)
 {
 	uint2 PixelPos = SvPosition.xy;
@@ -256,6 +260,10 @@ void ReflectionEnvironmentSkyLighting(
 	float2 ScreenPosition = SvPositionToScreenPosition(SvPosition).xy;
 	
 	OutColor = 0.0f;
+#if STRATA_OPAQUE_ROUGH_REFRACTION_ENABLED
+	OutOpaqueRoughRefractionSceneColor = 0.0f;
+	OutSubSurfaceSceneColor = 0.0f;
+#endif
 
 #if STRATA_ENABLED
 
@@ -287,7 +295,7 @@ void ReflectionEnvironmentSkyLighting(
 
 		const float CloudVolumetricAOShadow = GetCloudVolumetricAOShadow(TranslatedWorldPosition);
 
-		float4 StrataEnvLighting = StrataEnvironmentLighting(
+		FStrataDeferredLighting StrataLighting = StrataEnvironmentLighting(
 			SvPosition,
 			BufferUV,
 			SceneDepth,
@@ -308,7 +316,11 @@ void ReflectionEnvironmentSkyLighting(
 		// Alpha is 0 to not affect the subsurface diffuse luma information.
 		const float4 SSRColor = float4(SSR.rgb * TopLayerEnvBRDF * SSRReductionFactor, 0.0);
 
-		OutColor += StrataEnvLighting + SSRColor;
+		OutColor += StrataLighting.SceneColor + SSRColor;
+#if STRATA_OPAQUE_ROUGH_REFRACTION_ENABLED
+		OutOpaqueRoughRefractionSceneColor += StrataLighting.OpaqueRoughRefractionSceneColor;
+		OutSubSurfaceSceneColor += StrataLighting.SubSurfaceSceneColor;
+#endif
 	}
 
 #else // STRATA_ENABLED
diff --git a/Engine/Shaders/Private/Strata/Strata.ush b/Engine/Shaders/Private/Strata/Strata.ush
--- a/Engine/Shaders/Private/Strata/Strata.ush
+++ b/Engine/Shaders/Private/Strata/Strata.ush
@@ -570,6 +570,30 @@ bool IsStrataMaterial(FStrataTopLayerData In)
 	return In.bIsValid;
 }
 
+///////////////////////////////////////////////////////////////////////////////
+// Opaque rough refractions
+
+struct FStrataOpaqueRoughRefractionData
+{
+	bool	OpaqueRoughRefractionEnabled;
+	float	Variance;
+};
+
+float StrataPackOpaqueRoughRefractionData(FStrataOpaqueRoughRefractionData In)
+{
+	float Value = In.OpaqueRoughRefractionEnabled && In.Variance > 0.0f ? 128.0f : 0.0f;
+	Value += In.OpaqueRoughRefractionEnabled ? In.Variance * 127.0f : 0.0f;
+	return Value / 255.0f;
+}
+
+FStrataOpaqueRoughRefractionData StrataUnpackOpaqueRoughRefractionData(float In)
+{
+	FStrataOpaqueRoughRefractionData Out = (FStrataOpaqueRoughRefractionData)0;
+	Out.OpaqueRoughRefractionEnabled = In >= 128.0f;
+	Out.Variance = (In) * (1.0f / 255.0f) - (128.0f / 255.0f); 
+	return Out;
+}
+
 ///////////////////////////////////////////////////////////////////////////////
 // Sub-surface 
 
@@ -2614,7 +2638,8 @@ void PackStrataOut(
 	inout bool bStrataSubsurfaceEnable,
 	inout float3 EmissiveColor,
 	inout FStrataSubsurfaceData SSSData,
-	inout FStrataTopLayerData TopLayerData)
+	inout FStrataTopLayerData TopLayerData,
+	inout FStrataOpaqueRoughRefractionData OpaqueRoughRefractionData)
 {
 	// This only exists with inline shading and when we are going to write out BSDFs (UpdateAllXXX functions needs to be defined)
 #if STRATA_INLINE_SHADING && STRATA_CLAMPED_BSDF_COUNT > 0
@@ -2628,6 +2653,7 @@ void PackStrataOut(
 	// This avoid to run a post-basepass which would re-read all the material data
 	SSSData = (FStrataSubsurfaceData)0;
 	TopLayerData = (FStrataTopLayerData)0;
+	OpaqueRoughRefractionData = (FStrataOpaqueRoughRefractionData)0;
 
 	SetIsSimpleMaterial(StrataPixelHeader, true); // non strata material or sky pixels are considered as simple (so that (strata simple pixels) + (non-strata pixels) can build a simple tile).
 	SetIsSingleMaterial(StrataPixelHeader, true); // non strata material or sky pixels are considered as simple (so that (strata simple pixels) + (non-strata pixels) can build a simple tile).
@@ -2684,6 +2710,8 @@ void PackStrataOut(
 						TopLayerTotalWeight += BSDF.TopLayerDataWeight;
 						TopLayerData.Roughness += BSDF.TopLayerDataWeight * SLAB_ROUGHNESS(BSDF);
 
+						OpaqueRoughRefractionData.Variance += BSDF.TopLayerDataWeight * StrataLobeRoughnessToVariance(SLAB_ROUGHNESS(BSDF));
+
 						const bool bIsThin = BSDF_GETISTHIN(BSDF);
 						const bool bIsSimpleVolume = BSDF_GETISSIMPLEVOLUME(BSDF);
 						bStrataSubsurfaceEnable = bStrataSubsurfaceEnable || (BSDF_GETHASSSS(BSDF) && BSDF.Coverage > 0.0f && !bIsSimpleVolume && !bIsThin); // Should it be coverage>0 or any(LuminanceWeight>0) ?
@@ -2701,6 +2729,8 @@ void PackStrataOut(
 						TopLayerTotalWeight += BSDF.TopLayerDataWeight;
 						TopLayerData.Roughness += BSDF.TopLayerDataWeight * FullyRough;
 
+						OpaqueRoughRefractionData.Variance += BSDF.TopLayerDataWeight * StrataLobeRoughnessToVariance(SLAB_ROUGHNESS(BSDF));
+
 						break;
 					}
 
@@ -2712,6 +2742,8 @@ void PackStrataOut(
 						TopLayerTotalWeight += BSDF.TopLayerDataWeight;
 						TopLayerData.Roughness += BSDF.TopLayerDataWeight * SLW_ROUGHNESS(BSDF);
 
+						OpaqueRoughRefractionData.Variance += BSDF.TopLayerDataWeight * StrataLobeRoughnessToVariance(SLAB_ROUGHNESS(BSDF));
+
 						break;
 					}
 					}
@@ -2739,6 +2771,8 @@ void PackStrataOut(
 		TopLayerData.WorldNormal= TopLayerTotalWeight > 0.0f	? normalize(TopLayerData.WorldNormal)			: 0.0f;
 		TopLayerData.Roughness	= TopLayerTotalWeight > 0.0f	? TopLayerData.Roughness / TopLayerTotalWeight	: 0.0f;
 
+		OpaqueRoughRefractionData.Variance = TopLayerTotalWeight > 0.0f ? OpaqueRoughRefractionData.Variance / TopLayerTotalWeight : 0.0f;
+
 		// Set storage layout as either: fast(0), single(1), or complex(2)
 		AndIsSimpleMaterial(StrataPixelHeader, BSDFVisibleCount == 1);
 		AndIsSingleMaterial(StrataPixelHeader, BSDFVisibleCount == 1);
diff --git a/Engine/Shaders/Private/Strata/StrataDeferredLighting.ush b/Engine/Shaders/Private/Strata/StrataDeferredLighting.ush
--- a/Engine/Shaders/Private/Strata/StrataDeferredLighting.ush
+++ b/Engine/Shaders/Private/Strata/StrataDeferredLighting.ush
@@ -28,7 +28,7 @@ float4 StrataReadPrecomputedShadowFactors(FStrataPixelHeader StrataPixelHeader,
 
 // Analytical lighting evaluation for strata material.
 // Unpack BSDF on-the-fly
-float4 StrataDeferredLighting(
+FStrataDeferredLighting StrataDeferredLighting(
 	FDeferredLightData LightData,
 	float3 V,
 	float3 L,
@@ -39,7 +39,8 @@ float4 StrataDeferredLighting(
 	FStrataAddressing StrataAddressing,
 	FStrataPixelHeader StrataPixelHeader)
 {
-	FLightAccumulator Out = (FLightAccumulator)0;
+	FStrataDeferredLighting StrataLighting = GetInitialisedStrataDeferredLighting();
+
 	const bool bForceFullyRough = false;
 
 	FRectTexture RectTexture = InitRectTexture(LightData);
@@ -60,7 +61,7 @@ float4 StrataDeferredLighting(
 			FRect Rect = GetRect(ToLight, LightData);
 			if (!IsRectVisible(Rect))
 			{
-				return 0.0f; // Rect light can be non visible due to barn door occlusion
+				return GetInitialisedStrataDeferredLighting(); // Rect light can be non visible due to barn door occlusion
 			}
 			AreaLightContext = CreateRectIntegrateContext(Roughness, BSDFContext.N, BSDFContext.V, Rect, RectTexture);
 
@@ -79,14 +80,14 @@ float4 StrataDeferredLighting(
 		float3 DiffuseLuminance = BSDFEvaluate.IntegratedDiffuseValue;
 		float3 SpecularLuminance = BSDFEvaluate.IntegratedSpecularValue;
 
-		const bool bNeedsSeparateSubsurfaceLightAccumulation = BSDFEvaluate.bSubsurface;
 		const float3 CommonMultiplier = LightData.Color * LightMask * BSDF.LuminanceWeight;
-		LightAccumulator_AddSplit(Out, DiffuseLuminance, SpecularLuminance, DiffuseLuminance, CommonMultiplier, bNeedsSeparateSubsurfaceLightAccumulation);
+
+		FLightAccumulator Out = (FLightAccumulator)0;
+		LightAccumulator_AddSplit(Out, DiffuseLuminance, SpecularLuminance, DiffuseLuminance, CommonMultiplier, BSDFEvaluate.bSubsurface);
+		AccumulateStrataDeferredLighting(StrataLighting, Out, BSDFEvaluate.bSubsurface, BSDF_GETISTOPLAYER(BSDF));
 	}
 
-	const FDeferredLightingSplit SplitLighting = LightAccumulator_GetResultSplit(Out);
-	const float4 OutRadiance = SplitLighting.SpecularLighting + SplitLighting.DiffuseLighting;
-	return OutRadiance;
+	return StrataLighting;
 }
 
 #endif // STRATA_ENABLED
\ No newline at end of file
diff --git a/Engine/Shaders/Private/Strata/StrataEnvironmentLighting.ush b/Engine/Shaders/Private/Strata/StrataEnvironmentLighting.ush
--- a/Engine/Shaders/Private/Strata/StrataEnvironmentLighting.ush
+++ b/Engine/Shaders/Private/Strata/StrataEnvironmentLighting.ush
@@ -12,7 +12,7 @@
 
 // Env. lighting evaluation for strata material.
 // Unpack BSDF on-the-fly
-float4 StrataEnvironmentLighting(
+FStrataDeferredLighting StrataEnvironmentLighting(
 	float4 SvPosition,
 	float2 BufferUV,
 	float SceneDepth,
@@ -29,7 +29,7 @@ float4 StrataEnvironmentLighting(
 	inout float SSRReductionFactor,
 	inout float3 TopLayerEnvBRDF)
 {
-	FLightAccumulator Out = (FLightAccumulator)0;
+	FStrataDeferredLighting StrataLighting = GetInitialisedStrataDeferredLighting();
 
 	SSRReductionFactor = 1.0f;
 
@@ -182,12 +182,14 @@ float4 StrataEnvironmentLighting(
 
 		// STRATA_TODO ReflectionEnvironment(...) for static reflection volumes.
 
+		FLightAccumulator Out = (FLightAccumulator)0;
 		LightAccumulator_AddSplit(Out, DiffuseLighting, SpecularLighting, DiffuseLighting, View.PreExposure, StrataEnvLight.bSubsurface);
+		AccumulateStrataDeferredLighting(StrataLighting, Out, StrataEnvLight.bSubsurface, BSDF_GETISTOPLAYER(BSDF));
 	}
 
 	SSRReductionFactor = saturate(SSRReductionFactor);
 
-	return LightAccumulator_GetResult(Out);
+	return StrataLighting;
 }
 
 #endif // STRATA_ENABLED
\ No newline at end of file
diff --git a/Engine/Shaders/Private/Strata/StrataMaterialClassification.usf b/Engine/Shaders/Private/Strata/StrataMaterialClassification.usf
--- a/Engine/Shaders/Private/Strata/StrataMaterialClassification.usf
+++ b/Engine/Shaders/Private/Strata/StrataMaterialClassification.usf
@@ -33,6 +33,15 @@ RWBuffer<uint> SingleTileListDataBuffer;
 RWBuffer<uint> ComplexTileIndirectDataBuffer;
 RWBuffer<uint> ComplexTileListDataBuffer;
 
+// Opaque rough refraction	==> Disabled for now because this brings uav count to >8
+#undef STRATA_OPAQUE_ROUGH_REFRACTION_ENABLED
+#define STRATA_OPAQUE_ROUGH_REFRACTION_ENABLED 0
+#if STRATA_OPAQUE_ROUGH_REFRACTION_ENABLED
+RWBuffer<uint> OpaqueRoughRefractionTileIndirectDataBuffer;
+RWBuffer<uint> OpaqueRoughRefractionListDataBuffer;
+Texture2D<float> OpaqueRoughRefractionTexture;
+#endif
+
 #if !PERMUTATION_WAVE_OPS
 groupshared uint s_TileFlags[GROUP_THREAD_COUNT];
 #endif
@@ -59,7 +68,9 @@ void TileMainCS(uint2 DispatchThreadId : SV_DispatchThreadID, uint LinearIndex :
 	bool bContainsSimpleMaterial  = false;
 	bool bContainsSingleMaterial  = false;
 	bool bContainsStrataMaterial  = false;
+	bool bContainsOpaqueRoughRefraction = false;
 	bool bHasScreenSpaceSubsurfaceScattering = false;
+	FStrataOpaqueRoughRefractionData OpaqueRoughRefractionData = (FStrataOpaqueRoughRefractionData)0;
 	if (all(PixelCoord < uint2(ViewResolution)))
 	{
 		FStrataAddressing StrataAddressing = GetStrataPixelDataByteOffset(PixelCoord, uint2(View.BufferSizeAndInvSize.xy), MaxBytesPerPixel);
@@ -72,6 +83,11 @@ void TileMainCS(uint2 DispatchThreadId : SV_DispatchThreadID, uint LinearIndex :
 		bContainsSingleMaterial  = bIsSingle;
 		bContainsComplexMaterial = !bIsSingle & !bIsSimple;
 		bHasScreenSpaceSubsurfaceScattering = HasSubsurface(StrataPixelHeader);
+
+#if STRATA_OPAQUE_ROUGH_REFRACTION_ENABLED
+		OpaqueRoughRefractionData = StrataUnpackOpaqueRoughRefractionData(OpaqueRoughRefractionTexture[PixelCoord]);
+		bContainsOpaqueRoughRefraction = OpaqueRoughRefractionData.OpaqueRoughRefractionEnabled;
+#endif
 	}
 
 #if PERMUTATION_STRATA_CLEAR_DURING_CATEGORIZATION
@@ -89,13 +105,15 @@ void TileMainCS(uint2 DispatchThreadId : SV_DispatchThreadID, uint LinearIndex :
 //	const bool bHasAnySimple  = WaveActiveAnyTrue(bContainsSimpleMaterial);
 	const bool bHasAnySingle  = WaveActiveAnyTrue(bContainsSingleMaterial);
 	const bool bHasAnyComplex = WaveActiveAnyTrue(bContainsComplexMaterial);
+	const bool bHasAnyOpaqueRoughRefraction = WaveActiveAnyTrue(bContainsOpaqueRoughRefraction);
 #else // PERMUTATION_WAVE_OPS
 
 	s_TileFlags[LinearIndex] = 
-		(bContainsStrataMaterial  ? 0x1u : 0u) | 
-//		(bContainsSimpleMaterial  ? 0x2u : 0u) | 
-		(bContainsSingleMaterial  ? 0x4u : 0u) | 
-		(bContainsComplexMaterial ? 0x8u : 0u);
+		(bContainsStrataMaterial		? 0x1u	: 0u) | 
+//		(bContainsSimpleMaterial		? 0x2u	: 0u) | 
+		(bContainsSingleMaterial		? 0x4u	: 0u) | 
+		(bContainsComplexMaterial		? 0x8u	: 0u)| 
+		(bContainsOpaqueRoughRefraction ? 0xFu	: 0u);
 
 	GroupMemoryBarrierWithGroupSync();
 	if (LinearIndex < 32)
@@ -131,6 +149,7 @@ void TileMainCS(uint2 DispatchThreadId : SV_DispatchThreadID, uint LinearIndex :
 //	const bool bHasAnySimple  = (FinalTileFlags & 0x2u) > 0;
 	const bool bHasAnySingle  = (FinalTileFlags & 0x4u) > 0;
 	const bool bHasAnyComplex = (FinalTileFlags & 0x8u) > 0;
+	const bool bHasAnyOpaqueRoughRefraction = (FinalTileFlags & 0xFu) > 0;
 #endif // PERMUTATION_WAVE_OPS
 
 	if (LinearIndex < 1 && bHasAnyStrata)
@@ -153,6 +172,15 @@ void TileMainCS(uint2 DispatchThreadId : SV_DispatchThreadID, uint LinearIndex :
 			InterlockedAdd(SimpleTileIndirectDataBuffer[1], 1, WriteToIndex);
 			SimpleTileListDataBuffer[WriteToIndex] = EncodeTile(GroupId.xy);
 		}
+
+#if STRATA_OPAQUE_ROUGH_REFRACTION_ENABLED
+		if (bHasAnyOpaqueRoughRefraction)
+		{
+			uint WriteToIndex;
+			InterlockedAdd(OpaqueRoughRefractionTileIndirectDataBuffer[1], 1, WriteToIndex);
+			OpaqueRoughRefractionListDataBuffer[WriteToIndex] = EncodeTile(GroupId.xy);
+		}
+#endif
 	}
 }
 #endif // SHADER_TILE_CATEGORIZATION
diff --git a/Engine/Shaders/Private/Strata/StrataRoughRefraction.usf b/Engine/Shaders/Private/Strata/StrataRoughRefraction.usf
--- a/Engine/Shaders/Private/Strata/StrataRoughRefraction.usf
+++ b/Engine/Shaders/Private/Strata/StrataRoughRefraction.usf
@@ -4,8 +4,34 @@
 
 #include "../Common.ush"
 
+
+
+#if OPAQUE_ROUGH_REFRACTION_PS
+
+Texture2D<float4> SeparatedOpaqueRoughRefractionSceneColor;
+
+void OpaqueRoughRefractionPS(
+	float4 SVPos : SV_POSITION,
+	out float4 OutColor : SV_Target0)
+{
+	float2 PixelPos = SVPos.xy;
+	float2 UvBuffer = PixelPos * View.BufferSizeAndInvSize.zw;
+
+	float4 SourceData = SeparatedOpaqueRoughRefractionSceneColor[PixelPos];
+
+	OutColor = SourceData;
+}
+
+#endif
+
+
+
 #ifdef STRATA_RND_SHADERS
 
+//////////////////////////////////////////////////////////////////////////
+// RnD shaders only used when enabled locally
+//////////////////////////////////////////////////////////////////////////
+
 #include "/Engine/Private/MiniFontCommon.ush"
 #include "/Engine/Private/ShaderPrint.ush"
 #include "/Engine/Private/ColorMap.ush"
diff --git a/Engine/Source/Runtime/Engine/Classes/Engine/RendererSettings.h b/Engine/Source/Runtime/Engine/Classes/Engine/RendererSettings.h
--- a/Engine/Source/Runtime/Engine/Classes/Engine/RendererSettings.h
+++ b/Engine/Source/Runtime/Engine/Classes/Engine/RendererSettings.h
@@ -894,6 +894,15 @@ class ENGINE_API URendererSettings : public UDeveloperSettings
 		ConfigRestartRequired = true))
 		uint32 StrataBytePerPixel;
 
+	/**
+	"Enable Strata opaque material rough refractions effect from top layers over layers below."
+	*/
+	UPROPERTY(config, EditAnywhere, Category = Strata, meta = (
+		ConsoleVariable = "r.Strata.OpaqueMaterialRoughRefraction", DisplayName = "Strata opaque material rough refraction",
+		ToolTip = "Enable Strata opaque material rough refractions effect from top layers over layers below.",
+		ConfigRestartRequired = true))
+		uint32 StrataOpaqueMaterialRoughRefraction : 1;
+
 	/**
 	"Enable rough diffuse material."
 	*/
diff --git a/Engine/Source/Runtime/Engine/Private/Materials/HLSLMaterialTranslator.cpp b/Engine/Source/Runtime/Engine/Private/Materials/HLSLMaterialTranslator.cpp
--- a/Engine/Source/Runtime/Engine/Private/Materials/HLSLMaterialTranslator.cpp
+++ b/Engine/Source/Runtime/Engine/Private/Materials/HLSLMaterialTranslator.cpp
@@ -301,6 +301,7 @@ FHLSLMaterialTranslator::FHLSLMaterialTranslator(FMaterial* InMaterial,
 	StrataMaterialExpressionToOperatorIndex.Reserve(STRATA_MAX_OPERATOR_COUNT);
 	StrataMaterialBSDFCount = 0;
 	bStrataUsesConversionFromLegacy = false;
+	bStrataOutputsOpaqueRoughRefractions = false;
 }
 
 FHLSLMaterialTranslator::~FHLSLMaterialTranslator()
@@ -1956,6 +1957,8 @@ void FHLSLMaterialTranslator::GetMaterialEnvironment(EShaderPlatform InPlatform,
 
 		OutEnvironment.SetDefine(TEXT("STRATA_USES_CONVERSION_FROM_LEGACY"), bStrataUsesConversionFromLegacy ? 1 : 0);
 
+		OutEnvironment.SetDefine(TEXT("STRATA_MATERIAL_OUTPUT_OPAQUE_ROUGH_REFRACTIONS"), bStrataOutputsOpaqueRoughRefractions ? 1 : 0);
+
 		/*
 		* The final output code/workflow for shared tangent basis should look like
 		*
@@ -9799,6 +9802,7 @@ bool FHLSLMaterialTranslator::StrataGenerateDerivedMaterialOperatorData()
 	{
 		int VOpTopBranchCountTaken = 0;
 		int VOpBottomBranchCountTaken = 0;
+		bool bStrataUsesVerticalLayering = false;
 
 		std::function<void(FStrataOperator&, bool)> WalkOperators = [&](FStrataOperator& CurrentOperator, bool bInsideParameterBlendingSubTree) -> void
 		{
@@ -9849,14 +9853,21 @@ bool FHLSLMaterialTranslator::StrataGenerateDerivedMaterialOperatorData()
 				CurrentOperator.LeftIndex = INDEX_NONE;
 				CurrentOperator.RightIndex = INDEX_NONE;
 			}
+
+			// When at least one vertical operator exists that is not parameter blending, we can enabled writing to opaque rough refraction buffer.
+			bStrataUsesVerticalLayering = !CurrentOperator.bUseParameterBlending && CurrentOperator.OperatorType == STRATA_OPERATOR_VERTICAL;
 		};
 
 		WalkOperators(*StrataMaterialRootOperator, false);
+
+		const EStrataBlendMode StrataBlendMode = Material->GetStrataBlendMode();
+		const bool bIsOpaqueOrMasked = StrataBlendMode == EStrataBlendMode::SBM_Opaque || StrataBlendMode == EStrataBlendMode::SBM_Masked;
+		bStrataOutputsOpaqueRoughRefractions = bStrataUsesVerticalLayering && !bStrataUsesConversionFromLegacy && bIsOpaqueOrMasked;
 	}
 
-	// STRATA_TODO: operation using parameter blending are not actually discarded so they still occupy a pot in the opeartion arrays in the compiler and in the shader also.
-	// Even thought the compiler will remove unused operators from the shader code, we will still be limited to STRATA_MAX_OPERATOR_COUNT eve nwith parameter blending.
-	// This can be fixed by remapping all operation index and make sure this is the index that is allways used in the compiler and specified to the shader.
+	// STRATA_TODO: operation using parameter blending are not actually discarded so they still occupy a spot in the operation array in the compiler and in the shader also.
+	// Even thought the compiler will remove unused operators from the shader code, we will still be limited to STRATA_MAX_OPERATOR_COUNT even with parameter blending.
+	// This can be fixed by remapping all operation index and make sure this is the index that is always used in the compiler and specified to the shader.
 
 	//
 	// Make sure all the types have valid children operator indices
diff --git a/Engine/Source/Runtime/Engine/Private/Materials/HLSLMaterialTranslator.h b/Engine/Source/Runtime/Engine/Private/Materials/HLSLMaterialTranslator.h
--- a/Engine/Source/Runtime/Engine/Private/Materials/HLSLMaterialTranslator.h
+++ b/Engine/Source/Runtime/Engine/Private/Materials/HLSLMaterialTranslator.h
@@ -432,6 +432,7 @@ protected:
 	FStrataOperator* StrataMaterialRootOperator;
 	uint32 StrataMaterialBSDFCount;
 	bool bStrataUsesConversionFromLegacy;
+	bool bStrataOutputsOpaqueRoughRefractions;
 
 	/** Tracks the total number of vt samples in the shader. */
 	uint32 NumVtSamples;
diff --git a/Engine/Source/Runtime/Engine/Private/ShaderCompiler/ShaderCompiler.cpp b/Engine/Source/Runtime/Engine/Private/ShaderCompiler/ShaderCompiler.cpp
--- a/Engine/Source/Runtime/Engine/Private/ShaderCompiler/ShaderCompiler.cpp
+++ b/Engine/Source/Runtime/Engine/Private/ShaderCompiler/ShaderCompiler.cpp
@@ -5845,6 +5845,10 @@ void GlobalBeginCompileShader(
 		static IConsoleVariable* CVarBackCompatibility = IConsoleManager::Get().FindConsoleVariable(TEXT("r.StrataBackCompatibility"));
 		const bool bStrataBackCompatibility = CVarBackCompatibility && CVarBackCompatibility->GetInt() > 0;
 		Input.Environment.SetDefine(TEXT("PROJECT_STRATA_BACKCOMPATIBILITY"), bStrataBackCompatibility ? 1 : 0);
+
+		static IConsoleVariable* CVarOpaqueRoughRefrac = IConsoleManager::Get().FindConsoleVariable(TEXT("r.Strata.OpaqueMaterialRoughRefraction"));
+		const bool bStrataOpaqueRoughRefrac = bStrata && CVarOpaqueRoughRefrac && CVarOpaqueRoughRefrac->GetInt() != 0;
+		Input.Environment.SetDefine(TEXT("STRATA_OPAQUE_ROUGH_REFRACTION_ENABLED"), bStrataOpaqueRoughRefrac ? 1 : 0);
 	}
 
 	{
diff --git a/Engine/Source/Runtime/RenderCore/Private/Shader.cpp b/Engine/Source/Runtime/RenderCore/Private/Shader.cpp
--- a/Engine/Source/Runtime/RenderCore/Private/Shader.cpp
+++ b/Engine/Source/Runtime/RenderCore/Private/Shader.cpp
@@ -1788,6 +1788,12 @@ void ShaderMapAppendKeyString(EShaderPlatform Platform, FString& KeyString)
 		{
 			KeyString += FString::Printf(TEXT("_BACKCOMPAT"));
 		}
+
+		static const auto CVarOpaqueRoughRefrac = IConsoleManager::Get().FindTConsoleVariableDataInt(TEXT("r.Strata.OpaqueMaterialRoughRefraction"));
+		if (bStrataEnabled && CVarOpaqueRoughRefrac && CVarOpaqueRoughRefrac->GetValueOnAnyThread() > 0)
+		{
+			KeyString += FString::Printf(TEXT("_ROUGHDIFF"));
+		}
 	}
 
 	{
diff --git a/Engine/Source/Runtime/Renderer/Private/DeferredShadingRenderer.cpp b/Engine/Source/Runtime/Renderer/Private/DeferredShadingRenderer.cpp
--- a/Engine/Source/Runtime/Renderer/Private/DeferredShadingRenderer.cpp
+++ b/Engine/Source/Runtime/Renderer/Private/DeferredShadingRenderer.cpp
@@ -2927,6 +2927,8 @@ void FDeferredShadingSceneRenderer::Render(FRDGBuilder& GraphBuilder)
 
 		AddSubsurfacePass(GraphBuilder, SceneTextures, Views);
 
+		Strata::AddStrataOpaqueRoughRefractionPasses(GraphBuilder, SceneTextures, Views);
+
 		{
 			RenderHairStrandsSceneColorScattering(GraphBuilder, SceneTextures.Color.Target, Scene, Views);
 		}
diff --git a/Engine/Source/Runtime/Renderer/Private/IndirectLightRendering.cpp b/Engine/Source/Runtime/Renderer/Private/IndirectLightRendering.cpp
--- a/Engine/Source/Runtime/Renderer/Private/IndirectLightRendering.cpp
+++ b/Engine/Source/Runtime/Renderer/Private/IndirectLightRendering.cpp
@@ -1376,6 +1376,11 @@ static void AddSkyReflectionPass(
 	}
 
 	PassParameters->RenderTargets[0] = FRenderTargetBinding(SceneColorTexture.Target, ERenderTargetLoadAction::ELoad);
+	if (Strata::IsStrataOpaqueMaterialRoughRefractionEnabled())
+	{
+		PassParameters->RenderTargets[1] = FRenderTargetBinding(Scene->StrataSceneData.SeparatedOpaqueRoughRefractionSceneColor, ERenderTargetLoadAction::ELoad);
+		PassParameters->RenderTargets[2] = FRenderTargetBinding(Scene->StrataSceneData.SeparatedSubSurfaceSceneColor, ERenderTargetLoadAction::ELoad);
+	}
 
 	// Bind hair data
 	const bool bCheckerboardSubsurfaceRendering = IsSubsurfaceCheckerboardFormat(SceneColorTexture.Target->Desc.Format);
@@ -1437,7 +1442,17 @@ static void AddSkyReflectionPass(
 			}
 			else
 			{
-				GraphicsPSOInit.BlendState = TStaticBlendState<CW_RGBA, BO_Add, BF_One, BF_One, BO_Add, BF_One, BF_One>::GetRHI();
+				if (Strata::IsStrataOpaqueMaterialRoughRefractionEnabled())
+				{
+					GraphicsPSOInit.BlendState = TStaticBlendState<
+						CW_RGBA, BO_Add, BF_One, BF_One, BO_Add, BF_One, BF_One,
+						CW_RGBA, BO_Add, BF_One, BF_One, BO_Add, BF_One, BF_One,
+						CW_RGBA, BO_Add, BF_One, BF_One, BO_Add, BF_One, BF_One>::GetRHI();
+				}
+				else
+				{
+					GraphicsPSOInit.BlendState = TStaticBlendState<CW_RGBA, BO_Add, BF_One, BF_One, BO_Add, BF_One, BF_One>::GetRHI();
+				}
 			}
 		}
 
diff --git a/Engine/Source/Runtime/Renderer/Private/LightRendering.cpp b/Engine/Source/Runtime/Renderer/Private/LightRendering.cpp
--- a/Engine/Source/Runtime/Renderer/Private/LightRendering.cpp
+++ b/Engine/Source/Runtime/Renderer/Private/LightRendering.cpp
@@ -1824,6 +1824,11 @@ static FDeferredLightPS::FParameters GetDeferredLightPSParameters(
 	Out.ShadowChannelMask = FVector4f(1, 1, 1, 1);
 	// PS - Render Targets
 	Out.RenderTargets[0] = FRenderTargetBinding(SceneColorTexture, ERenderTargetLoadAction::ELoad);
+	if (Strata::IsStrataOpaqueMaterialRoughRefractionEnabled())
+	{
+		Out.RenderTargets[1] = FRenderTargetBinding(Scene->StrataSceneData.SeparatedOpaqueRoughRefractionSceneColor, ERenderTargetLoadAction::ELoad);
+		Out.RenderTargets[2] = FRenderTargetBinding(Scene->StrataSceneData.SeparatedSubSurfaceSceneColor, ERenderTargetLoadAction::ELoad);
+	}
 	if (SceneDepthTexture)
 	{
 		Out.RenderTargets.DepthStencil = FDepthStencilBinding(SceneDepthTexture, ERenderTargetLoadAction::ELoad, ERenderTargetLoadAction::ELoad, FExclusiveDepthStencil::DepthRead_StencilWrite);
@@ -1865,7 +1870,17 @@ static void InternalRenderLight(
 		// Set the device viewport for the view.
 		RHICmdList.SetViewport(View.ViewRect.Min.X, View.ViewRect.Min.Y, 0.0f, View.ViewRect.Max.X, View.ViewRect.Max.Y, 1.0f);
 
-		GraphicsPSOInit.BlendState = TStaticBlendState<CW_RGBA, BO_Add, BF_One, BF_One, BO_Add, BF_One, BF_One>::GetRHI();
+		if (Strata::IsStrataOpaqueMaterialRoughRefractionEnabled())
+		{
+			GraphicsPSOInit.BlendState = TStaticBlendState<
+				CW_RGBA, BO_Add, BF_One, BF_One, BO_Add, BF_One, BF_One,
+				CW_RGBA, BO_Add, BF_One, BF_One, BO_Add, BF_One, BF_One,
+				CW_RGBA, BO_Add, BF_One, BF_One, BO_Add, BF_One, BF_One>::GetRHI();
+		}
+		else
+		{
+			GraphicsPSOInit.BlendState = TStaticBlendState<CW_RGBA, BO_Add, BF_One, BF_One, BO_Add, BF_One, BF_One>::GetRHI();
+		}
 		GraphicsPSOInit.PrimitiveType = PT_TriangleList;
 
 		if (LightType == LightType_Directional)
diff --git a/Engine/Source/Runtime/Renderer/Private/PostProcess/PostProcessSubsurface.cpp b/Engine/Source/Runtime/Renderer/Private/PostProcess/PostProcessSubsurface.cpp
--- a/Engine/Source/Runtime/Renderer/Private/PostProcess/PostProcessSubsurface.cpp
+++ b/Engine/Source/Runtime/Renderer/Private/PostProcess/PostProcessSubsurface.cpp
@@ -309,6 +309,11 @@ uint32 GetSubsurfaceRequiredViewMask(TArrayView<const FViewInfo> Views)
 
 bool IsSubsurfaceCheckerboardFormat(EPixelFormat SceneColorFormat)
 {
+	if (Strata::IsStrataOpaqueMaterialRoughRefractionEnabled())
+	{
+		// With this mode, specular and subsurface colors are correctly separated so checkboard is not required.
+		return false;
+	}
 	int CVarValue = CVarSSSCheckerboard.GetValueOnRenderThread();
 	if (CVarValue == 0)
 	{
@@ -981,6 +986,11 @@ void AddSubsurfaceViewPass(
 	const FScreenPassTextureViewportParameters SubsurfaceViewportParameters = GetScreenPassTextureViewportParameters(SubsurfaceViewport);
 	const FScreenPassTextureViewportParameters SceneViewportParameters = GetScreenPassTextureViewportParameters(SceneViewport);
 
+	const bool bReadSeparatedSubSurfaceSceneColor = Strata::IsStrataOpaqueMaterialRoughRefractionEnabled();
+	const bool bWriteSeparatedOpaqueRoughRefractionSceneColor = Strata::IsStrataOpaqueMaterialRoughRefractionEnabled();
+	FRDGTextureRef SeparatedSubSurfaceSceneColor = View.StrataSceneData->SeparatedSubSurfaceSceneColor;
+	FRDGTextureRef SeparatedOpaqueRoughRefractionSceneColor = View.StrataSceneData->SeparatedOpaqueRoughRefractionSceneColor;
+
 	FRDGTextureRef SetupTexture = SceneColorTexture;
 	FRDGTextureRef SubsurfaceSubpassOneTex = nullptr;
 	FRDGTextureRef SubsurfaceSubpassTwoTex = nullptr;
@@ -1066,14 +1076,13 @@ void AddSubsurfaceViewPass(
 
 		// Call the indirect setup (with tile classification)
 		{
-			FRDGTextureSRVDesc SceneColorTextureSRVDesc = FRDGTextureSRVDesc::Create(SceneColorTexture);
 			FRDGTextureUAVDesc SetupTextureOutDesc(SetupTexture, 0);
 
 			typedef FSubsurfaceIndirectDispatchSetupCS SHADER;
 			SHADER::FParameters* PassParameters = GraphBuilder.AllocParameters<SHADER::FParameters>();
 			PassParameters->Subsurface = SubsurfaceCommonParameters;
 			PassParameters->Output = SubsurfaceViewportParameters;
-			PassParameters->SubsurfaceInput0 = GetSubsurfaceInput(SceneColorTexture, SceneViewportParameters);
+			PassParameters->SubsurfaceInput0 = GetSubsurfaceInput(bReadSeparatedSubSurfaceSceneColor ? SeparatedSubSurfaceSceneColor : SceneColorTexture, SceneViewportParameters);
 			PassParameters->SubsurfaceSampler0 = PointClampSampler;
 			PassParameters->SetupTexture = GraphBuilder.CreateUAV(SetupTextureOutDesc);
 			if (bUseProfileIdCache)
@@ -1310,7 +1319,7 @@ void AddSubsurfaceViewPass(
 				PassParameters->TileParameters = GetSubsurfaceTileParameters(SubsurfaceViewport, Tiles, TileType);
 			}
 			PassParameters->RenderTargets[0] = FRenderTargetBinding(SubsurfaceIntermediateTexture, SceneColorTextureLoadAction);
-			PassParameters->SubsurfaceInput0 = GetSubsurfaceInput(SceneColorTexture, SceneViewportParameters);
+			PassParameters->SubsurfaceInput0 = GetSubsurfaceInput(bReadSeparatedSubSurfaceSceneColor ? SeparatedSubSurfaceSceneColor : SceneColorTexture, SceneViewportParameters);
 			PassParameters->SubsurfaceSampler0 = BilinearBorderSampler;
 			PassParameters->Strata = Strata::BindStrataGlobalUniformParameters(View.StrataSceneData);
 
@@ -1370,15 +1379,21 @@ void AddSubsurfaceViewPass(
 			{
 				PassParameters->TileParameters = GetSubsurfaceTileParameters(SubsurfaceViewport, Tiles, TileType);
 			}
-			PassParameters->RenderTargets[0] = FRenderTargetBinding(SceneColorTexture, SceneColorTextureLoadAction);
+			PassParameters->RenderTargets[0] = bWriteSeparatedOpaqueRoughRefractionSceneColor ? 
+				FRenderTargetBinding(SeparatedOpaqueRoughRefractionSceneColor, ERenderTargetLoadAction::ELoad) :
+				FRenderTargetBinding(SceneColorTexture, SceneColorTextureLoadAction);
 			PassParameters->SubsurfaceInput0 = GetSubsurfaceInput(SubsurfaceIntermediateTexture, SceneViewportParameters);
 			PassParameters->SubsurfaceSampler0 = PointClampSampler;
 			PassParameters->Strata = Strata::BindStrataGlobalUniformParameters(View.StrataSceneData);
 
 			TShaderMapRef<FSubsurfaceRecombineCopyPS> PixelShader(View.ShaderMap);			
 			TShaderMapRef<FSubsurfaceTilePassVS> VertexShader(View.ShaderMap);
-			
+
 			FRHIBlendState* BlendState = TStaticBlendState<CW_RGBA, BO_Add, BF_One, BF_Zero, BO_Add, BF_One, BF_Zero>::GetRHI();
+			if (bWriteSeparatedOpaqueRoughRefractionSceneColor)
+			{
+				BlendState = TStaticBlendState<CW_RGBA, BO_Add, BF_One, BF_One, BO_Add, BF_One, BF_One>::GetRHI();
+			}
 			
 			AddSubsurfaceTiledScreenPass(
 				GraphBuilder,
diff --git a/Engine/Source/Runtime/Renderer/Private/Strata/Strata.cpp b/Engine/Source/Runtime/Renderer/Private/Strata/Strata.cpp
--- a/Engine/Source/Runtime/Renderer/Private/Strata/Strata.cpp
+++ b/Engine/Source/Runtime/Renderer/Private/Strata/Strata.cpp
@@ -34,7 +34,6 @@ static TAutoConsoleVariable<int32> CVarStrataBytePerPixel(
 	TEXT("Strata allocated byte per pixel to store materials data. Higher value means more complex material can be represented."),
 	ECVF_ReadOnly | ECVF_RenderThreadSafe);
 
-
 static TAutoConsoleVariable<int32> CVarStrataClassificationDebug(
 	TEXT("r.Strata.Classification.Debug"),
 	0,
@@ -68,8 +67,11 @@ IMPLEMENT_GLOBAL_SHADER_PARAMETER_STRUCT(FStrataGlobalUniformParameters, "Strata
 void FStrataSceneData::Reset()
 {
 	TopLayerTexture = nullptr;
+	TopLayerTextureUAV = nullptr;
 	SSSTexture = nullptr;
 	SSSTextureUAV = nullptr;
+	OpaqueRoughRefractionTexture = nullptr;
+	OpaqueRoughRefractionTextureUAV = nullptr;
 
 	MaterialTextureArray = nullptr;
 	MaterialTextureArrayUAVWithoutRTs = nullptr;
@@ -86,6 +88,13 @@ void FStrataSceneData::Reset()
 		ClassificationTileIndirectBufferSRV[i] = nullptr;
 	}
 
+	OpaqueRoughRefractionClassificationTileListBuffer = nullptr;
+	OpaqueRoughRefractionClassificationTileListBufferSRV = nullptr;
+	OpaqueRoughRefractionClassificationTileListBufferUAV = nullptr;
+
+	SeparatedSubSurfaceSceneColor = nullptr;
+	SeparatedOpaqueRoughRefractionSceneColor = nullptr;
+
 	StrataGlobalUniformParameters = nullptr;
 }
 
@@ -152,10 +161,11 @@ void InitialiseStrataFrameSceneData(FSceneRenderer& SceneRenderer, FRDGBuilder&
 		const uint32 RoundToValue = 4u;
 		StrataSceneData.MaxBytesPerPixel = FMath::DivideAndRoundUp(MaterialConservativeByteCountPerPixel, RoundToValue) * RoundToValue;
 
+		const int32 TileInPixel = GetStrataBufferTileSize();
+		const FIntPoint TileResolution(FMath::DivideAndRoundUp(SceneTextureExtent.X, TileInPixel), FMath::DivideAndRoundUp(SceneTextureExtent.Y, TileInPixel));
+
 		// Tile classification buffers
 		{
-			const int32 TileInPixel = GetStrataBufferTileSize();
-			const FIntPoint TileResolution(FMath::DivideAndRoundUp(SceneTextureExtent.X, TileInPixel), FMath::DivideAndRoundUp(SceneTextureExtent.Y, TileInPixel));
 
 			const TCHAR* StrataTileResourceNames[EStrataTileMaterialType::ECount][3] =
 			{
@@ -190,6 +200,40 @@ void InitialiseStrataFrameSceneData(FSceneRenderer& SceneRenderer, FRDGBuilder&
 			StrataSceneData.SSSTexture = GraphBuilder.CreateTexture(FRDGTextureDesc::Create2D(SceneTextureExtent, PF_R32G32_UINT, FClearValueBinding::Black, TexCreate_DisableDCC | TexCreate_NoFastClear | TexCreate_ShaderResource | TexCreate_UAV), TEXT("Strata.SSSTexture"));
 			StrataSceneData.SSSTextureUAV = GraphBuilder.CreateUAV(StrataSceneData.SSSTexture);
 		}
+
+		// Separated subsurface and rough refraction textures
+		{
+			const bool bIsStrataOpaqueMaterialRoughRefractionEnabled= IsStrataOpaqueMaterialRoughRefractionEnabled();
+			const int32 TileListBufferElementCount					= bIsStrataOpaqueMaterialRoughRefractionEnabled ? TileResolution.X * TileResolution.Y : 4;
+			const FIntPoint OpaqueRoughRefractionSceneExtent		= bIsStrataOpaqueMaterialRoughRefractionEnabled ? SceneTextureExtent : FIntPoint(4, 4);
+			
+
+			StrataSceneData.OpaqueRoughRefractionClassificationTileListBuffer = GraphBuilder.CreateBuffer(FRDGBufferDesc::CreateBufferDesc(sizeof(uint32), TileListBufferElementCount), TEXT("Strata.StrataTileListBuffer(RoughRefrac)"));
+			StrataSceneData.OpaqueRoughRefractionClassificationTileListBufferSRV = GraphBuilder.CreateSRV(StrataSceneData.OpaqueRoughRefractionClassificationTileListBuffer, PF_R32_UINT);
+			StrataSceneData.OpaqueRoughRefractionClassificationTileListBufferUAV = GraphBuilder.CreateUAV(StrataSceneData.OpaqueRoughRefractionClassificationTileListBuffer, PF_R32_UINT);
+
+			StrataSceneData.OpaqueRoughRefractionClassificationTileIndirectBuffer = GraphBuilder.CreateBuffer(FRDGBufferDesc::CreateIndirectDesc<FRHIDrawIndirectParameters>(), TEXT("Strata.StrataTileIndirectBuffer(RoughRefrac)"));
+			StrataSceneData.OpaqueRoughRefractionClassificationTileIndirectBufferSRV = GraphBuilder.CreateSRV(StrataSceneData.OpaqueRoughRefractionClassificationTileIndirectBuffer, PF_R32_UINT);
+			StrataSceneData.OpaqueRoughRefractionClassificationTileIndirectBufferSRV = GraphBuilder.CreateSRV(StrataSceneData.OpaqueRoughRefractionClassificationTileIndirectBuffer, PF_R32_UINT);
+			StrataSceneData.OpaqueRoughRefractionClassificationTileIndirectBufferUAV = GraphBuilder.CreateUAV(StrataSceneData.OpaqueRoughRefractionClassificationTileIndirectBuffer, PF_R32_UINT);
+
+			StrataSceneData.OpaqueRoughRefractionTexture = GraphBuilder.CreateTexture(
+				FRDGTextureDesc::Create2D(OpaqueRoughRefractionSceneExtent, PF_R8, FClearValueBinding::Black, TexCreate_ShaderResource | TexCreate_UAV | TexCreate_RenderTargetable), TEXT("Strata.OpaqueRoughRefractionTexture"));
+			StrataSceneData.OpaqueRoughRefractionTextureUAV = GraphBuilder.CreateUAV(StrataSceneData.OpaqueRoughRefractionTexture);
+			
+			StrataSceneData.SeparatedSubSurfaceSceneColor			= GraphBuilder.CreateTexture(
+				FRDGTextureDesc::Create2D(OpaqueRoughRefractionSceneExtent, PF_FloatR11G11B10, FClearValueBinding::Black, TexCreate_ShaderResource | TexCreate_UAV | TexCreate_RenderTargetable), TEXT("Strata.SeparatedSubSurfaceSceneColor"));
+			StrataSceneData.SeparatedOpaqueRoughRefractionSceneColor= GraphBuilder.CreateTexture(
+				FRDGTextureDesc::Create2D(OpaqueRoughRefractionSceneExtent, PF_FloatR11G11B10, FClearValueBinding::Black, TexCreate_ShaderResource | TexCreate_UAV | TexCreate_RenderTargetable), TEXT("Strata.SeparatedOpaqueRoughRefractionSceneColor"));
+
+			if (bIsStrataOpaqueMaterialRoughRefractionEnabled)
+			{
+				// Fast clears
+				AddClearRenderTargetPass(GraphBuilder, StrataSceneData.OpaqueRoughRefractionTexture, StrataSceneData.OpaqueRoughRefractionTexture->Desc.ClearValue.GetClearColor());
+				AddClearRenderTargetPass(GraphBuilder, StrataSceneData.SeparatedSubSurfaceSceneColor, StrataSceneData.SeparatedSubSurfaceSceneColor->Desc.ClearValue.GetClearColor());
+				AddClearRenderTargetPass(GraphBuilder, StrataSceneData.SeparatedOpaqueRoughRefractionSceneColor, StrataSceneData.SeparatedOpaqueRoughRefractionSceneColor->Desc.ClearValue.GetClearColor());
+			}
+		}
 	}
 	else
 	{
@@ -247,11 +291,15 @@ void BindStrataBasePassUniformParameters(FRDGBuilder& GraphBuilder, FStrataScene
 		OutStrataUniformParameters.MaxBytesPerPixel = StrataSceneData->MaxBytesPerPixel;
 		OutStrataUniformParameters.MaterialTextureArrayUAVWithoutRTs = StrataSceneData->MaterialTextureArrayUAVWithoutRTs;
 		OutStrataUniformParameters.SSSTextureUAV = StrataSceneData->SSSTextureUAV;
+		OutStrataUniformParameters.OpaqueRoughRefractionTextureUAV = StrataSceneData->OpaqueRoughRefractionTextureUAV;
 	}
 	else
 	{
-		FRDGTextureRef DummyWritableTexture = GraphBuilder.CreateTexture(FRDGTextureDesc::Create2D(FIntPoint(1, 1), PF_R32_UINT, FClearValueBinding::None, TexCreate_ShaderResource | TexCreate_UAV), TEXT("Strata.DummyWritableTexture"));
-		FRDGTextureUAVRef DummyWritableTextureUAV = GraphBuilder.CreateUAV(FRDGTextureUAVDesc(DummyWritableTexture));
+		FRDGTextureRef DummyWritableSSSTexture = GraphBuilder.CreateTexture(FRDGTextureDesc::Create2D(FIntPoint(1, 1), PF_R32_UINT, FClearValueBinding::None, TexCreate_ShaderResource | TexCreate_UAV), TEXT("Strata.DummyWritableTexture"));
+		FRDGTextureUAVRef DummyWritableSSSTextureUAV = GraphBuilder.CreateUAV(FRDGTextureUAVDesc(DummyWritableSSSTexture));
+
+		FRDGTextureRef DummyWritableRefracTexture = GraphBuilder.CreateTexture(FRDGTextureDesc::Create2D(FIntPoint(1, 1), PF_R8, FClearValueBinding::None, TexCreate_ShaderResource | TexCreate_UAV), TEXT("Strata.DummyWritableTexture"));
+		FRDGTextureUAVRef DummyWritableRefracTextureUAV = GraphBuilder.CreateUAV(FRDGTextureUAVDesc(DummyWritableRefracTexture));
 
 		FRDGTextureRef DummyWritableTextureArray = GraphBuilder.CreateTexture(FRDGTextureDesc::Create2DArray(FIntPoint(1, 1), PF_R32_UINT, FClearValueBinding::None, TexCreate_ShaderResource | TexCreate_UAV, 1), TEXT("Strata.DummyWritableTexture"));
 		FRDGTextureUAVRef DummyWritableTextureArrayUAV = GraphBuilder.CreateUAV(FRDGTextureUAVDesc(DummyWritableTextureArray));
@@ -260,7 +308,8 @@ void BindStrataBasePassUniformParameters(FRDGBuilder& GraphBuilder, FStrataScene
 		OutStrataUniformParameters.bRoughDiffuse = 0u;
 		OutStrataUniformParameters.MaxBytesPerPixel = 0;
 		OutStrataUniformParameters.MaterialTextureArrayUAVWithoutRTs = DummyWritableTextureArrayUAV;
-		OutStrataUniformParameters.SSSTextureUAV = DummyWritableTextureUAV;
+		OutStrataUniformParameters.SSSTextureUAV = DummyWritableSSSTextureUAV;
+		OutStrataUniformParameters.OpaqueRoughRefractionTextureUAV = DummyWritableRefracTextureUAV;
 	}
 }
 
@@ -273,6 +322,7 @@ void BindStrataGlobalUniformParameters(FRDGBuilder& GraphBuilder, FStrataSceneDa
 		OutStrataUniformParameters.MaterialTextureArray = StrataSceneData->MaterialTextureArray;
 		OutStrataUniformParameters.TopLayerTexture = StrataSceneData->TopLayerTexture;
 		OutStrataUniformParameters.SSSTexture = StrataSceneData->SSSTexture;
+		OutStrataUniformParameters.OpaqueRoughRefractionTexture = StrataSceneData->OpaqueRoughRefractionTexture;
 	}
 	else
 	{
@@ -283,6 +333,7 @@ void BindStrataGlobalUniformParameters(FRDGBuilder& GraphBuilder, FStrataSceneDa
 		OutStrataUniformParameters.MaterialTextureArray = DefaultTextureArray;
 		OutStrataUniformParameters.TopLayerTexture = SystemTextures.DefaultNormal8Bit;
 		OutStrataUniformParameters.SSSTexture = SystemTextures.Black;
+		OutStrataUniformParameters.OpaqueRoughRefractionTexture = SystemTextures.Black;
 	}
 }
 
@@ -436,7 +487,10 @@ class FStrataMaterialTileClassificationPassCS : public FGlobalShader
 		SHADER_PARAMETER_RDG_BUFFER_UAV(RWBuffer, SingleTileListDataBuffer)
 		SHADER_PARAMETER_RDG_BUFFER_UAV(RWBuffer, ComplexTileIndirectDataBuffer)
 		SHADER_PARAMETER_RDG_BUFFER_UAV(RWBuffer, ComplexTileListDataBuffer)
+		SHADER_PARAMETER_RDG_BUFFER_UAV(RWBuffer, OpaqueRoughRefractionTileIndirectDataBuffer)
+		SHADER_PARAMETER_RDG_BUFFER_UAV(RWBuffer, OpaqueRoughRefractionTileListDataBuffer)
 		SHADER_PARAMETER_RDG_TEXTURE_UAV(RWTexture2D<uint2>, SSSTextureUAV)
+		SHADER_PARAMETER_RDG_TEXTURE(Texture2D<float>, OpaqueRoughRefractionTexture)
 	END_SHADER_PARAMETER_STRUCT()
 
 	static bool ShouldCompilePermutation(const FGlobalShaderPermutationParameters& Parameters)
@@ -776,6 +830,7 @@ void AddStrataMaterialClassificationPass(FRDGBuilder& GraphBuilder, const FMinim
 			PassParameters->TopLayerTexture = View.StrataSceneData->TopLayerTexture;
 			PassParameters->MaterialTextureArray = View.StrataSceneData->MaterialTextureArraySRV;
 			PassParameters->SSSTextureUAV = View.StrataSceneData->SSSTextureUAV;
+			PassParameters->OpaqueRoughRefractionTexture = View.StrataSceneData->OpaqueRoughRefractionTexture;
 			// Simple
 			PassParameters->SimpleTileListDataBuffer = View.StrataSceneData->ClassificationTileListBufferUAV[EStrataTileMaterialType::ESimple];
 			PassParameters->SimpleTileIndirectDataBuffer = View.StrataSceneData->ClassificationTileIndirectBufferUAV[EStrataTileMaterialType::ESimple];
@@ -785,6 +840,9 @@ void AddStrataMaterialClassificationPass(FRDGBuilder& GraphBuilder, const FMinim
 			// Complex
 			PassParameters->ComplexTileListDataBuffer = View.StrataSceneData->ClassificationTileListBufferUAV[EStrataTileMaterialType::EComplex];
 			PassParameters->ComplexTileIndirectDataBuffer = View.StrataSceneData->ClassificationTileIndirectBufferUAV[EStrataTileMaterialType::EComplex];
+			// Opaque Rough Refraction
+			PassParameters->OpaqueRoughRefractionTileListDataBuffer = View.StrataSceneData->OpaqueRoughRefractionClassificationTileListBufferUAV;
+			PassParameters->OpaqueRoughRefractionTileIndirectDataBuffer = View.StrataSceneData->OpaqueRoughRefractionClassificationTileIndirectBufferUAV;
 
 			const uint32 GroupSize = 8;
 			FComputeShaderUtils::AddPass(
diff --git a/Engine/Source/Runtime/Renderer/Private/Strata/Strata.h b/Engine/Source/Runtime/Renderer/Private/Strata/Strata.h
--- a/Engine/Source/Runtime/Renderer/Private/Strata/Strata.h
+++ b/Engine/Source/Runtime/Renderer/Private/Strata/Strata.h
@@ -20,6 +20,7 @@ BEGIN_SHADER_PARAMETER_STRUCT(FStrataBasePassUniformParameters, )
 	SHADER_PARAMETER(uint32, bRoughDiffuse)
 	SHADER_PARAMETER_RDG_TEXTURE_UAV(RWTexture2DArray<uint>, MaterialTextureArrayUAVWithoutRTs)
 	SHADER_PARAMETER_RDG_TEXTURE_UAV(RWTexture2D<uint2>, SSSTextureUAV)
+	SHADER_PARAMETER_RDG_TEXTURE_UAV(RWTexture2D<float>, OpaqueRoughRefractionTextureUAV)
 END_SHADER_PARAMETER_STRUCT()
 
 BEGIN_SHADER_PARAMETER_STRUCT(FStrataForwardPassUniformParameters, )
@@ -39,6 +40,7 @@ BEGIN_GLOBAL_SHADER_PARAMETER_STRUCT(FStrataGlobalUniformParameters, RENDERER_AP
 	SHADER_PARAMETER_RDG_TEXTURE(Texture2DArray<uint>, MaterialTextureArray)
 	SHADER_PARAMETER_RDG_TEXTURE(Texture2D<uint>, TopLayerTexture)
 	SHADER_PARAMETER_RDG_TEXTURE(Texture2D<uint2>, SSSTexture)
+	SHADER_PARAMETER_RDG_TEXTURE(Texture2D<float>, OpaqueRoughRefractionTexture)
 END_GLOBAL_SHADER_PARAMETER_STRUCT()
 
 enum EStrataTileMaterialType : uint32
@@ -70,11 +72,26 @@ struct FStrataSceneData
 	FRDGBufferUAVRef ClassificationTileIndirectBufferUAV[EStrataTileMaterialType::ECount];
 	FRDGBufferSRVRef ClassificationTileIndirectBufferSRV[EStrataTileMaterialType::ECount];
 
+	FRDGBufferRef OpaqueRoughRefractionClassificationTileListBuffer;
+	FRDGBufferUAVRef OpaqueRoughRefractionClassificationTileListBufferUAV;
+	FRDGBufferSRVRef OpaqueRoughRefractionClassificationTileListBufferSRV;
+	FRDGBufferRef OpaqueRoughRefractionClassificationTileIndirectBuffer;
+	FRDGBufferUAVRef OpaqueRoughRefractionClassificationTileIndirectBufferUAV;
+	FRDGBufferSRVRef OpaqueRoughRefractionClassificationTileIndirectBufferSRV;
+
 	FRDGTextureRef TopLayerTexture;
 	FRDGTextureRef SSSTexture;
+	FRDGTextureRef OpaqueRoughRefractionTexture;
 
 	FRDGTextureUAVRef TopLayerTextureUAV;
 	FRDGTextureUAVRef SSSTextureUAV;
+	FRDGTextureUAVRef OpaqueRoughRefractionTextureUAV;
+
+	// Used when the subsurface luminance is separated from the scene color
+	FRDGTextureRef SeparatedSubSurfaceSceneColor;
+
+	// Used for Luminance that should go through opaque rough refraction (when under a top layer interface)
+	FRDGTextureRef SeparatedOpaqueRoughRefractionSceneColor;
 
 	TRDGUniformBufferRef<FStrataGlobalUniformParameters> StrataGlobalUniformParameters{};
 
@@ -110,6 +127,12 @@ void AddStrataMaterialClassificationPass(FRDGBuilder& GraphBuilder, const FMinim
 
 void AddStrataStencilPass(FRDGBuilder& GraphBuilder, const TArray<FViewInfo>& Views, const FMinimalSceneTextures& SceneTextures);
 
+bool IsStrataOpaqueMaterialRoughRefractionEnabled();
+void AddStrataOpaqueRoughRefractionPasses(
+	FRDGBuilder& GraphBuilder,
+	FSceneTextures& SceneTextures,
+	TArrayView<const FViewInfo> Views);
+
 bool ShouldRenderStrataDebugPasses(const FViewInfo& View);
 FScreenPassTexture AddStrataDebugPasses(FRDGBuilder& GraphBuilder, const FViewInfo& View, FScreenPassTexture& ScreenPassSceneColor);
 
diff --git a/Engine/Source/Runtime/Renderer/Private/Strata/StrataRoughRefraction.cpp b/Engine/Source/Runtime/Renderer/Private/Strata/StrataRoughRefraction.cpp
--- a/Engine/Source/Runtime/Renderer/Private/Strata/StrataRoughRefraction.cpp
+++ b/Engine/Source/Runtime/Renderer/Private/Strata/StrataRoughRefraction.cpp
@@ -13,11 +13,92 @@
 #include "ShaderCompiler.h"
 
 
+static TAutoConsoleVariable<int32> CVarStrataOpaqueMaterialRoughRefraction(
+	TEXT("r.Strata.OpaqueMaterialRoughRefraction"),
+	0,
+	TEXT("Enable Strata opaque material rough refractions effect from top layers over layers below."),
+	ECVF_ReadOnly | ECVF_RenderThreadSafe);
+
+
 namespace Strata
 {
 
+bool IsStrataOpaqueMaterialRoughRefractionEnabled()
+{
+	return IsStrataEnabled() && CVarStrataOpaqueMaterialRoughRefraction.GetValueOnAnyThread() > 0;
+}
+
+class FOpaqueRoughRefractionPS : public FGlobalShader
+{
+	DECLARE_GLOBAL_SHADER(FOpaqueRoughRefractionPS);
+	SHADER_USE_PARAMETER_STRUCT(FOpaqueRoughRefractionPS, FGlobalShader);
+
+	using FPermutationDomain = TShaderPermutationDomain<>;
+
+	BEGIN_SHADER_PARAMETER_STRUCT(FParameters, )
+		SHADER_PARAMETER_STRUCT_REF(FViewUniformShaderParameters, ViewUniformBuffer)
+		SHADER_PARAMETER_RDG_TEXTURE(Texture2D, SeparatedOpaqueRoughRefractionSceneColor)
+		RENDER_TARGET_BINDING_SLOTS()
+	END_SHADER_PARAMETER_STRUCT();
+
+	static bool ShouldCompilePermutation(const FGlobalShaderPermutationParameters& Parameters)
+	{
+
+		return IsFeatureLevelSupported(Parameters.Platform, ERHIFeatureLevel::SM5);
+	}
+
+	static void ModifyCompilationEnvironment(const FGlobalShaderPermutationParameters& Parameters, FShaderCompilerEnvironment& OutEnvironment)
+	{
+		FGlobalShader::ModifyCompilationEnvironment(Parameters, OutEnvironment);
+
+		OutEnvironment.SetDefine(TEXT("OPAQUE_ROUGH_REFRACTION_PS"), 1);
+	}
+};
+
+IMPLEMENT_GLOBAL_SHADER(FOpaqueRoughRefractionPS, "/Engine/Private/Strata/StrataRoughRefraction.usf", "OpaqueRoughRefractionPS", SF_Pixel);
+
+
+void AddStrataOpaqueRoughRefractionPasses(
+	FRDGBuilder& GraphBuilder,
+	FSceneTextures& SceneTextures,
+	TArrayView<const FViewInfo> Views)
+{
+	if (!IsStrataOpaqueMaterialRoughRefractionEnabled())
+	{
+		return;
+	}
+
+	const uint32 ViewCount = Views.Num();
+	for (uint32 ViewIndex = 0; ViewIndex < ViewCount; ++ViewIndex)
+	{
+		const FViewInfo& View = Views[ViewIndex];
+
+		FRDGTextureRef SeparatedOpaqueRoughRefractionSceneColor = View.StrataSceneData->SeparatedOpaqueRoughRefractionSceneColor;
+		FRDGTextureRef SceneColorTexture = SceneTextures.Color.Target;
+
+
+		FSceneTextureParameters SceneTextureParameters = GetSceneTextureParameters(GraphBuilder);
+		FOpaqueRoughRefractionPS::FParameters* PassParameters = GraphBuilder.AllocParameters<FOpaqueRoughRefractionPS::FParameters>();
+		PassParameters->ViewUniformBuffer = View.ViewUniformBuffer;
+		PassParameters->SeparatedOpaqueRoughRefractionSceneColor = SeparatedOpaqueRoughRefractionSceneColor;
+		PassParameters->RenderTargets[0] = FRenderTargetBinding(SceneColorTexture, ERenderTargetLoadAction::ELoad);
+
+		FOpaqueRoughRefractionPS::FPermutationDomain PermutationVector;
+		TShaderMapRef<FOpaqueRoughRefractionPS> PixelShader(View.ShaderMap, PermutationVector);
+
+		FRHIBlendState* PreMultipliedColorTransmittanceBlend = TStaticBlendState<CW_RGB, BO_Add, BF_One, BF_One, BO_Add, BF_One, BF_One>::GetRHI();
+
+		FPixelShaderUtils::AddFullscreenPass<FOpaqueRoughRefractionPS>(GraphBuilder, View.ShaderMap, RDG_EVENT_NAME("Strata::VisualizeRoughRefractionPS"),
+			PixelShader, PassParameters, View.ViewRect, PreMultipliedColorTransmittanceBlend);
+
+		
+	}
+}
+
+
 //////////////////////////////////////////////////////////////////////////
 // RnD shaders only used when enabled locally
+//////////////////////////////////////////////////////////////////////////
 
 // Keeping it simple: this should always be checked in with a value of 0
 #define STRATA_ROUGH_REFRACTION_RND 0
@@ -30,8 +111,18 @@ static TAutoConsoleVariable<int32> CVarStrataRoughRefractionShadersShowRoughRefr
 	TEXT("Enable strata rough refraction shaders."),
 	ECVF_RenderThreadSafe);
 
+bool ShouldRenderStrataRoughRefractionRnD()
+{
+	return CVarStrataRoughRefractionShadersShowRoughRefractionRnD.GetValueOnAnyThread() > 0;
+}
+
+#else
+bool ShouldRenderStrataRoughRefractionRnD() { return false; }
+#endif
 
 
+#if STRATA_ROUGH_REFRACTION_RND
+
 class FEvaluateRoughRefractionLobeCS : public FGlobalShader
 {
 	DECLARE_GLOBAL_SHADER(FEvaluateRoughRefractionLobeCS);
@@ -114,17 +205,7 @@ IMPLEMENT_GLOBAL_SHADER(FVisualizeRoughRefractionPS, "/Engine/Private/Strata/Str
 
 
 #endif // STRATA_ROUGH_REFRACTION_RND
-	
-
 
-bool ShouldRenderStrataRoughRefractionRnD()
-{
-#if STRATA_ROUGH_REFRACTION_RND
-	return CVarStrataRoughRefractionShadersShowRoughRefractionRnD.GetValueOnAnyThread() > 0;
-#else
-	return false;
-#endif
-}
 
 void StrataRoughRefractionRnD(FRDGBuilder& GraphBuilder, const FViewInfo& View, FScreenPassTexture& ScreenPassSceneColor)
 {
