Fix compile error in Test/Shipping builds

diff --git a/Engine/Source/Runtime/Renderer/Private/SceneRendering.cpp b/Engine/Source/Runtime/Renderer/Private/SceneRendering.cpp
--- a/Engine/Source/Runtime/Renderer/Private/SceneRendering.cpp
+++ b/Engine/Source/Runtime/Renderer/Private/SceneRendering.cpp
@@ -4065,6 +4065,7 @@ static void RenderViewFamily_RenderThread(FRHICommandListImmediate& RHICmdList,
 
 	FSceneViewFamily& ViewFamily = SceneRenderer->ViewFamily;
 
+#if WITH_DEBUG_VIEW_MODES
 	const bool bAllowGPUSkinCacheVisualization = AllowDebugViewShaderMode(DVSM_VisualizeGPUSkinCache, ViewFamily.GetShaderPlatform(), ViewFamily.GetFeatureLevel());
 	if (bAllowGPUSkinCacheVisualization && SceneRenderer->Views.Num() > 0)
 	{
@@ -4082,6 +4083,7 @@ static void RenderViewFamily_RenderThread(FRHICommandListImmediate& RHICmdList,
 			SceneRenderer->DrawGPUSkinCacheVisualizationInfoText();
 		}
 	}
+#endif
 
 	{
 		SCOPE_CYCLE_COUNTER_VERBOSE(STAT_TotalSceneRenderingTime, ViewFamily.ProfileDescription.IsEmpty() ? nullptr : *ViewFamily.ProfileDescription);
