Strata - Refactored indirect dispatches to source from the same indierect draw buffer, thus avoiding hitting out of uav slots on dx11.

diff --git a/Engine/Shaders/Private/Strata/StrataMaterialClassification.usf b/Engine/Shaders/Private/Strata/StrataMaterialClassification.usf
--- a/Engine/Shaders/Private/Strata/StrataMaterialClassification.usf
+++ b/Engine/Shaders/Private/Strata/StrataMaterialClassification.usf
@@ -23,22 +23,16 @@ int2 ViewResolution;
 uint MaxBytesPerPixel;
 Texture2D<uint> TopLayerTexture;
 Texture2DArray<uint> MaterialTextureArray;
-// Fast/Simple
-RWBuffer<uint> SimpleTileIndirectDataBuffer;
+
+// Indirect draw data buffer for all tile types
+RWBuffer<uint> TileIndirectDataBuffer;
+
 RWBuffer<uint> SimpleTileListDataBuffer;
-// Single
-RWBuffer<uint> SingleTileIndirectDataBuffer;
 RWBuffer<uint> SingleTileListDataBuffer;
-// Complex
-RWBuffer<uint> ComplexTileIndirectDataBuffer;
 RWBuffer<uint> ComplexTileListDataBuffer;
 
-// Opaque rough refraction	==> Disabled for now because this brings uav count to >8
-#undef STRATA_OPAQUE_ROUGH_REFRACTION_ENABLED
-#define STRATA_OPAQUE_ROUGH_REFRACTION_ENABLED 0
 #if STRATA_OPAQUE_ROUGH_REFRACTION_ENABLED
-RWBuffer<uint> OpaqueRoughRefractionTileIndirectDataBuffer;
-RWBuffer<uint> OpaqueRoughRefractionListDataBuffer;
+RWBuffer<uint> OpaqueRoughRefractionTileListDataBuffer;
 Texture2D<float> OpaqueRoughRefractionTexture;
 #endif
 
@@ -55,9 +49,11 @@ void TileMainCS(uint2 DispatchThreadId : SV_DispatchThreadID, uint LinearIndex :
 {
 	if (all(DispatchThreadId == 0))
 	{
-		SimpleTileIndirectDataBuffer[0] = bRectPrimitive > 0 ? 4 : 6;
-		SingleTileIndirectDataBuffer[0] = bRectPrimitive > 0 ? 4 : 6;
-		ComplexTileIndirectDataBuffer[0] = bRectPrimitive > 0 ? 4 : 6;
+		const uint IndexCountPerInstance = bRectPrimitive > 0 ? 4 : 6;
+		TileIndirectDataBuffer[GetStrataTileTypeDrawIndirectArgOffset_DWord(STRATA_TILE_TYPE_SIMPLE)		+ 0] = IndexCountPerInstance;
+		TileIndirectDataBuffer[GetStrataTileTypeDrawIndirectArgOffset_DWord(STRATA_TILE_TYPE_SINGLE)		+ 0] = IndexCountPerInstance;
+		TileIndirectDataBuffer[GetStrataTileTypeDrawIndirectArgOffset_DWord(STRATA_TILE_TYPE_COMPLEX)		+ 0] = IndexCountPerInstance;
+		TileIndirectDataBuffer[GetStrataTileTypeDrawIndirectArgOffset_DWord(STRATA_TILE_TYPE_ROUGH_REFRACT) + 0] = IndexCountPerInstance;
 	}
 
 	// TODO: add a SM6 permutation with ballot?
@@ -112,8 +108,8 @@ void TileMainCS(uint2 DispatchThreadId : SV_DispatchThreadID, uint LinearIndex :
 		(bContainsStrataMaterial		? 0x1u	: 0u) | 
 //		(bContainsSimpleMaterial		? 0x2u	: 0u) | 
 		(bContainsSingleMaterial		? 0x4u	: 0u) | 
-		(bContainsComplexMaterial		? 0x8u	: 0u)| 
-		(bContainsOpaqueRoughRefraction ? 0xFu	: 0u);
+		(bContainsComplexMaterial		? 0x8u	: 0u) | 
+		(bContainsOpaqueRoughRefraction ? 0x10u : 0u);
 
 	GroupMemoryBarrierWithGroupSync();
 	if (LinearIndex < 32)
@@ -149,7 +145,7 @@ void TileMainCS(uint2 DispatchThreadId : SV_DispatchThreadID, uint LinearIndex :
 //	const bool bHasAnySimple  = (FinalTileFlags & 0x2u) > 0;
 	const bool bHasAnySingle  = (FinalTileFlags & 0x4u) > 0;
 	const bool bHasAnyComplex = (FinalTileFlags & 0x8u) > 0;
-	const bool bHasAnyOpaqueRoughRefraction = (FinalTileFlags & 0xFu) > 0;
+	const bool bHasAnyOpaqueRoughRefraction = (FinalTileFlags & 0x10u) > 0;
 #endif // PERMUTATION_WAVE_OPS
 
 	if (LinearIndex < 1 && bHasAnyStrata)
@@ -157,19 +153,19 @@ void TileMainCS(uint2 DispatchThreadId : SV_DispatchThreadID, uint LinearIndex :
 		if (bHasAnyComplex)
 		{
 			uint WriteToIndex;
-			InterlockedAdd(ComplexTileIndirectDataBuffer[1], 1, WriteToIndex);
+			InterlockedAdd(TileIndirectDataBuffer[GetStrataTileTypeDrawIndirectArgOffset_DWord(STRATA_TILE_TYPE_COMPLEX) + 1], 1, WriteToIndex);
 			ComplexTileListDataBuffer[WriteToIndex] = EncodeTile(GroupId.xy);
 		}
 		else if (bHasAnySingle)
 		{
 			uint WriteToIndex;
-			InterlockedAdd(SingleTileIndirectDataBuffer[1], 1, WriteToIndex);
+			InterlockedAdd(TileIndirectDataBuffer[GetStrataTileTypeDrawIndirectArgOffset_DWord(STRATA_TILE_TYPE_SINGLE) + 1], 1, WriteToIndex);
 			SingleTileListDataBuffer[WriteToIndex] = EncodeTile(GroupId.xy);
 		}
 		else // (bHasAnySimple)
 		{
 			uint WriteToIndex;
-			InterlockedAdd(SimpleTileIndirectDataBuffer[1], 1, WriteToIndex);
+			InterlockedAdd(TileIndirectDataBuffer[GetStrataTileTypeDrawIndirectArgOffset_DWord(STRATA_TILE_TYPE_SIMPLE) + 1], 1, WriteToIndex);
 			SimpleTileListDataBuffer[WriteToIndex] = EncodeTile(GroupId.xy);
 		}
 
@@ -177,8 +173,8 @@ void TileMainCS(uint2 DispatchThreadId : SV_DispatchThreadID, uint LinearIndex :
 		if (bHasAnyOpaqueRoughRefraction)
 		{
 			uint WriteToIndex;
-			InterlockedAdd(OpaqueRoughRefractionTileIndirectDataBuffer[1], 1, WriteToIndex);
-			OpaqueRoughRefractionListDataBuffer[WriteToIndex] = EncodeTile(GroupId.xy);
+			InterlockedAdd(TileIndirectDataBuffer[GetStrataTileTypeDrawIndirectArgOffset_DWord(STRATA_TILE_TYPE_ROUGH_REFRACT) + 1], 1, WriteToIndex);
+			OpaqueRoughRefractionTileListDataBuffer[WriteToIndex] = EncodeTile(GroupId.xy);
 		}
 #endif
 	}
diff --git a/Engine/Shaders/Shared/StrataDefinitions.h b/Engine/Shaders/Shared/StrataDefinitions.h
--- a/Engine/Shaders/Shared/StrataDefinitions.h
+++ b/Engine/Shaders/Shared/StrataDefinitions.h
@@ -52,3 +52,14 @@
 #define STRATA_OPERATOR_BSDF			4
 #define STRATA_OPERATOR_BSDF_LEGACY		5
 #define STRATA_OPERATOR_THINFILM		6
+
+// This must map directly to EStrataTileMaterialType
+#define STRATA_TILE_TYPE_SIMPLE			0
+#define STRATA_TILE_TYPE_SINGLE			1
+#define STRATA_TILE_TYPE_COMPLEX		2
+#define STRATA_TILE_TYPE_ROUGH_REFRACT	3
+#define STRATA_TILE_TYPE_COUNT			4
+
+// sizeof(FRHIDrawIndirectParameters) = 4 uints = 16 bytes
+#define GetStrataTileTypeDrawIndirectArgOffset_Byte(x)  (x * 16)
+#define GetStrataTileTypeDrawIndirectArgOffset_DWord(x) (x * 4)
diff --git a/Engine/Source/Runtime/Renderer/Private/ClusteredDeferredShadingPass.cpp b/Engine/Source/Runtime/Renderer/Private/ClusteredDeferredShadingPass.cpp
--- a/Engine/Source/Runtime/Renderer/Private/ClusteredDeferredShadingPass.cpp
+++ b/Engine/Source/Runtime/Renderer/Private/ClusteredDeferredShadingPass.cpp
@@ -145,7 +145,7 @@ static void InternalAddClusteredDeferredShadingPass(
 	const FMinimalSceneTextures& SceneTextures,
 	const FSortedLightSetSceneInfo &SortedLightsSet,
 	EClusterPassInputType InputType,
-	EStrataTileMaterialType TileType,
+	EStrataTileType TileType,
 	FRDGTextureRef ShadowMaskBits,
 	FVirtualShadowMapArray& VirtualShadowMapArray,
 	FRDGBufferSRVRef HairTransmittanceBuffer)
@@ -230,7 +230,7 @@ static void InternalAddClusteredDeferredShadingPass(
 		else if (bStrata)
 		{
 			SetShaderParameters(InRHICmdList, TileVertexShader, TileVertexShader.GetVertexShader(), PassParameters->StrataTile);
-			InRHICmdList.DrawPrimitiveIndirect(PassParameters->StrataTile.TileIndirectBuffer->GetIndirectRHICallBuffer(), 0);
+			InRHICmdList.DrawPrimitiveIndirect(PassParameters->StrataTile.TileIndirectBuffer->GetIndirectRHICallBuffer(), Strata::TileTypeDrawIndirectArgOffset(TileType));
 		}
 		else
 		{
@@ -271,7 +271,7 @@ void FDeferredShadingSceneRenderer::AddClusteredDeferredShadingPass(
 					SceneTextures,
 					SortedLightsSet,
 					EClusterPassInputType::Strata,
-					EStrataTileMaterialType::ESimple,
+					EStrataTileType::ESimple,
 					ShadowMaskBits,
 					VirtualShadowMapArray,
 					nullptr);
@@ -282,7 +282,7 @@ void FDeferredShadingSceneRenderer::AddClusteredDeferredShadingPass(
 					SceneTextures,
 					SortedLightsSet,
 					EClusterPassInputType::Strata,
-					EStrataTileMaterialType::ESingle,
+					EStrataTileType::ESingle,
 					ShadowMaskBits,
 					VirtualShadowMapArray,
 					nullptr);
@@ -293,7 +293,7 @@ void FDeferredShadingSceneRenderer::AddClusteredDeferredShadingPass(
 					SceneTextures,
 					SortedLightsSet,
 					EClusterPassInputType::Strata,
-					EStrataTileMaterialType::EComplex,
+					EStrataTileType::EComplex,
 					ShadowMaskBits,
 					VirtualShadowMapArray,
 					nullptr);
@@ -306,7 +306,7 @@ void FDeferredShadingSceneRenderer::AddClusteredDeferredShadingPass(
 					SceneTextures,
 					SortedLightsSet,
 					EClusterPassInputType::GBuffer,
-					EStrataTileMaterialType::ECount,
+					EStrataTileType::ECount,
 					ShadowMaskBits,
 					VirtualShadowMapArray,
 					nullptr);
@@ -321,7 +321,7 @@ void FDeferredShadingSceneRenderer::AddClusteredDeferredShadingPass(
 					SceneTextures,
 					SortedLightsSet,
 					EClusterPassInputType::HairStrands,
-					EStrataTileMaterialType::ECount,
+					EStrataTileType::ECount,
 					HairStrandsShadowMaskBits,
 					VirtualShadowMapArray,
 					GraphBuilder.CreateSRV(TransmittanceMask.TransmittanceMask, FHairStrandsTransmittanceMaskData::Format));
diff --git a/Engine/Source/Runtime/Renderer/Private/IndirectLightRendering.cpp b/Engine/Source/Runtime/Renderer/Private/IndirectLightRendering.cpp
--- a/Engine/Source/Runtime/Renderer/Private/IndirectLightRendering.cpp
+++ b/Engine/Source/Runtime/Renderer/Private/IndirectLightRendering.cpp
@@ -129,7 +129,7 @@ class FDiffuseIndirectCompositePS : public FGlobalShader
 		}
 
 		// Build Strata tile permutation only for Lumen
-		if (PermutationVector.Get<FStrataTileType>() != EStrataTileMaterialType::EComplex)
+		if (PermutationVector.Get<FStrataTileType>() != EStrataTileType::EComplex)
 		{
 			return Strata::IsStrataEnabled() && PermutationVector.Get<FApplyDiffuseIndirectDim>() == 4;
 		}
@@ -281,7 +281,7 @@ class FReflectionEnvironmentSkyLightingPS : public FGlobalShader
 		return PermutationVector;
 	}
 
-	static FPermutationDomain BuildPermutationVector(const FViewInfo& View, bool bBoxCapturesOnly, bool bSphereCapturesOnly, bool bSupportDFAOIndirectOcclusion, bool bEnableSkyLight, bool bEnableDynamicSkyLight, bool bApplySkyShadowing, bool bRayTracedReflections, EStrataTileMaterialType TileType)
+	static FPermutationDomain BuildPermutationVector(const FViewInfo& View, bool bBoxCapturesOnly, bool bSphereCapturesOnly, bool bSupportDFAOIndirectOcclusion, bool bEnableSkyLight, bool bEnableDynamicSkyLight, bool bApplySkyShadowing, bool bRayTracedReflections, EStrataTileType TileType)
 	{
 		FPermutationDomain PermutationVector;
 
@@ -295,7 +295,7 @@ class FReflectionEnvironmentSkyLightingPS : public FGlobalShader
 		PermutationVector.Set<FStrataTileType>(0);
 		if (Strata::IsStrataEnabled())
 		{
-			check(EStrataTileMaterialType::ECount != TileType);
+			check(TileType <= EStrataTileType::EComplex);
 			PermutationVector.Set<FStrataTileType>(TileType);
 		}
 		return RemapPermutation(PermutationVector);
@@ -1014,7 +1014,7 @@ void FDeferredShadingSceneRenderer::RenderDiffuseIndirectAndAmbientOcclusion(
 		const bool bApplyDiffuseIndirect = (DenoiserOutputs.Textures[0] || AmbientOcclusionMask) && (!bIsVisualizePass || ViewPipelineState.DiffuseIndirectDenoiser != IScreenSpaceDenoiser::EMode::Disabled || ViewPipelineState.DiffuseIndirectMethod == EDiffuseIndirectMethod::Lumen)
 			&& !(IsMetalPlatform(ShaderPlatform) && !IsMetalSM5Platform(ShaderPlatform));
 
-		auto ApplyDiffuseIndirect = [&](EStrataTileMaterialType TileType)
+		auto ApplyDiffuseIndirect = [&](EStrataTileType TileType)
 		{
 			FDiffuseIndirectCompositePS::FParameters* PassParameters = GraphBuilder.AllocParameters<FDiffuseIndirectCompositePS::FParameters>();
 			PassParameters->Strata = Strata::BindStrataGlobalUniformParameters(View.StrataSceneData);
@@ -1088,7 +1088,7 @@ void FDeferredShadingSceneRenderer::RenderDiffuseIndirectAndAmbientOcclusion(
 
 			const TCHAR* DiffuseIndirectSampling = TEXT("Disabled");
 			FDiffuseIndirectCompositePS::FPermutationDomain PermutationVector;
-			PermutationVector.Set<FDiffuseIndirectCompositePS::FStrataTileType>(EStrataTileMaterialType::EComplex);
+			PermutationVector.Set<FDiffuseIndirectCompositePS::FStrataTileType>(EStrataTileType::EComplex);
 
 			bool bUpscale = false;
 
@@ -1108,7 +1108,7 @@ void FDeferredShadingSceneRenderer::RenderDiffuseIndirectAndAmbientOcclusion(
 				{
 					PermutationVector.Set<FDiffuseIndirectCompositePS::FApplyDiffuseIndirectDim>(4);
 					PermutationVector.Set<FDiffuseIndirectCompositePS::FScreenBentNormal>(ScreenBentNormalParameters.UseScreenBentNormal != 0);
-					if (Strata::IsStrataEnabled() && TileType != EStrataTileMaterialType::ECount)
+					if (Strata::IsStrataEnabled() && TileType != EStrataTileType::ECount)
 					{
 						PermutationVector.Set<FDiffuseIndirectCompositePS::FStrataTileType>(TileType);
 					}
@@ -1135,7 +1135,7 @@ void FDeferredShadingSceneRenderer::RenderDiffuseIndirectAndAmbientOcclusion(
 				BlendState = TStaticBlendState<>::GetRHI();
 			}
 
-			if (TileType == EStrataTileMaterialType::ECount)
+			if (TileType == EStrataTileType::ECount)
 			{
 				ClearUnusedGraphResources(PixelShader, PassParameters);
 
@@ -1197,7 +1197,7 @@ void FDeferredShadingSceneRenderer::RenderDiffuseIndirectAndAmbientOcclusion(
 					SetGraphicsPipelineState(RHICmdList, GraphicsPSOInit, 0x0);
 					SetShaderParameters(RHICmdList, PixelShader, PixelShader.GetPixelShader(), *PassParameters);
 					SetShaderParameters(RHICmdList, TileVertexShader, TileVertexShader.GetVertexShader(), PassParameters->StrataTile);
-					RHICmdList.DrawPrimitiveIndirect(PassParameters->StrataTile.TileIndirectBuffer->GetIndirectRHICallBuffer(), 0);
+					RHICmdList.DrawPrimitiveIndirect(PassParameters->StrataTile.TileIndirectBuffer->GetIndirectRHICallBuffer(), Strata::TileTypeDrawIndirectArgOffset(TileType));
 				});
 			}
 
@@ -1207,13 +1207,13 @@ void FDeferredShadingSceneRenderer::RenderDiffuseIndirectAndAmbientOcclusion(
 		{
 			if (Strata::IsStrataEnabled() && ViewPipelineState.DiffuseIndirectMethod == EDiffuseIndirectMethod::Lumen)
 			{
-				ApplyDiffuseIndirect(EStrataTileMaterialType::ESimple);
-				ApplyDiffuseIndirect(EStrataTileMaterialType::ESingle);
-				ApplyDiffuseIndirect(EStrataTileMaterialType::EComplex);
+				ApplyDiffuseIndirect(EStrataTileType::ESimple);
+				ApplyDiffuseIndirect(EStrataTileType::ESingle);
+				ApplyDiffuseIndirect(EStrataTileType::EComplex);
 			}
 			else
 			{
-				ApplyDiffuseIndirect(EStrataTileMaterialType::ECount);
+				ApplyDiffuseIndirect(EStrataTileType::ECount);
 			}
 		}
 
@@ -1302,7 +1302,7 @@ static void AddSkyReflectionPass(
 	bool bSkyLight, 
 	bool bDynamicSkyLight, 
 	bool bApplySkyShadowing,
-	EStrataTileMaterialType StrataTileMaterialType)
+	EStrataTileType StrataTileMaterialType)
 {
 	// Render the reflection environment with tiled deferred culling
 	bool bHasBoxCaptures = (View.NumBoxReflectionCaptures > 0);
@@ -1409,7 +1409,7 @@ static void AddSkyReflectionPass(
 	const bool bStrataEnabled = Strata::IsStrataEnabled();
 	if (bStrataEnabled)
 	{
-		check(StrataTileMaterialType != EStrataTileMaterialType::ECount);
+		check(StrataTileMaterialType <= EStrataTileType::EComplex);
 		PassParameters->VS = Strata::SetTileParameters(GraphBuilder, View, StrataTileMaterialType, StrataTilePrimitiveType);
 		ClearUnusedGraphResources(StrataTilePassVertexShader, &PassParameters->VS);
 	}
@@ -1420,7 +1420,7 @@ static void AddSkyReflectionPass(
 			bStrataEnabled ? ToString(StrataTileMaterialType) : TEXT("Off")),
 		PassParameters,
 		ERDGPassFlags::Raster,
-		[PassParameters, &View, PixelShader, bCheckerboardSubsurfaceRendering, 
+		[PassParameters, &View, PixelShader, bCheckerboardSubsurfaceRendering, StrataTileMaterialType,
 		StrataTilePassVertexShader, bStrataEnabled, StrataTilePrimitiveType](FRHICommandList& InRHICmdList)
 	{
 		InRHICmdList.SetViewport(View.ViewRect.Min.X, View.ViewRect.Min.Y, 0.0f, View.ViewRect.Max.X, View.ViewRect.Max.Y, 1.0f);
@@ -1477,7 +1477,7 @@ static void AddSkyReflectionPass(
 		if (bStrataEnabled)
 		{
 			SetShaderParameters(InRHICmdList, StrataTilePassVertexShader, StrataTilePassVertexShader.GetVertexShader(), PassParameters->VS);
-			InRHICmdList.DrawPrimitiveIndirect(PassParameters->VS.TileIndirectBuffer->GetIndirectRHICallBuffer(), 0);
+			InRHICmdList.DrawPrimitiveIndirect(PassParameters->VS.TileIndirectBuffer->GetIndirectRHICallBuffer(), Strata::TileTypeDrawIndirectArgOffset(StrataTileMaterialType));
 		}
 		else
 		{
@@ -1714,7 +1714,7 @@ void FDeferredShadingSceneRenderer::RenderDeferredReflectionsAndSkyLighting(
 					bSkyLight,
 					bDynamicSkyLight,
 					bApplySkyShadowing,
-					EStrataTileMaterialType::ESimple);
+					EStrataTileType::ESimple);
 
 				AddSkyReflectionPass(
 					GraphBuilder,
@@ -1728,7 +1728,7 @@ void FDeferredShadingSceneRenderer::RenderDeferredReflectionsAndSkyLighting(
 					bSkyLight,
 					bDynamicSkyLight,
 					bApplySkyShadowing,
-					EStrataTileMaterialType::ESingle);
+					EStrataTileType::ESingle);
 
 				AddSkyReflectionPass(
 					GraphBuilder,
@@ -1742,7 +1742,7 @@ void FDeferredShadingSceneRenderer::RenderDeferredReflectionsAndSkyLighting(
 					bSkyLight,
 					bDynamicSkyLight,
 					bApplySkyShadowing,
-					EStrataTileMaterialType::EComplex);
+					EStrataTileType::EComplex);
 			}
 			else
 			{
@@ -1759,7 +1759,7 @@ void FDeferredShadingSceneRenderer::RenderDeferredReflectionsAndSkyLighting(
 					bSkyLight,
 					bDynamicSkyLight,
 					bApplySkyShadowing,
-					EStrataTileMaterialType::ECount);
+					EStrataTileType::ECount);
 			}
 		}
 
diff --git a/Engine/Source/Runtime/Renderer/Private/LightRendering.cpp b/Engine/Source/Runtime/Renderer/Private/LightRendering.cpp
--- a/Engine/Source/Runtime/Renderer/Private/LightRendering.cpp
+++ b/Engine/Source/Runtime/Renderer/Private/LightRendering.cpp
@@ -1687,19 +1687,19 @@ static void InternalSetBoundingGeometryRasterizerState(FGraphicsPipelineStateIni
 }
 
 template<ECompareFunction CompareFunction>
-static uint32 InternalSetBoundingGeometryDepthState(FGraphicsPipelineStateInitializer& GraphicsPSOInit, EStrataTileMaterialType TileType)
+static uint32 InternalSetBoundingGeometryDepthState(FGraphicsPipelineStateInitializer& GraphicsPSOInit, EStrataTileType TileType)
 {
 	// bCameraInsideLightGeometry = true  -> CompareFunction = Always
 	// bCameraInsideLightGeometry = false -> CompareFunction = CF_DepthNearOrEqual
 	uint32 StencilRef = 0u;
-	if (TileType != EStrataTileMaterialType::ECount)
+	if (TileType != EStrataTileType::ECount)
 	{
 		check(Strata::IsStrataEnabled());
 		switch (TileType)
 		{
-		case EStrataTileMaterialType::ESimple : StencilRef = Strata::StencilBit_Fast;    GraphicsPSOInit.DepthStencilState = TStaticDepthStencilState<false, CompareFunction, true, CF_Equal, SO_Keep, SO_Keep, SO_Keep, true, CF_Equal, SO_Keep, SO_Keep, SO_Keep, Strata::StencilBit_Fast, 0x0>::GetRHI(); break;
-		case EStrataTileMaterialType::ESingle : StencilRef = Strata::StencilBit_Single;  GraphicsPSOInit.DepthStencilState = TStaticDepthStencilState<false, CompareFunction, true, CF_Equal, SO_Keep, SO_Keep, SO_Keep, true, CF_Equal, SO_Keep, SO_Keep, SO_Keep, Strata::StencilBit_Single, 0x0>::GetRHI(); break;
-		case EStrataTileMaterialType::EComplex: StencilRef = Strata::StencilBit_Complex; GraphicsPSOInit.DepthStencilState = TStaticDepthStencilState<false, CompareFunction, true, CF_Equal, SO_Keep, SO_Keep, SO_Keep, true, CF_Equal, SO_Keep, SO_Keep, SO_Keep, Strata::StencilBit_Complex, 0x0>::GetRHI(); break;
+		case EStrataTileType::ESimple : StencilRef = Strata::StencilBit_Fast;    GraphicsPSOInit.DepthStencilState = TStaticDepthStencilState<false, CompareFunction, true, CF_Equal, SO_Keep, SO_Keep, SO_Keep, true, CF_Equal, SO_Keep, SO_Keep, SO_Keep, Strata::StencilBit_Fast, 0x0>::GetRHI(); break;
+		case EStrataTileType::ESingle : StencilRef = Strata::StencilBit_Single;  GraphicsPSOInit.DepthStencilState = TStaticDepthStencilState<false, CompareFunction, true, CF_Equal, SO_Keep, SO_Keep, SO_Keep, true, CF_Equal, SO_Keep, SO_Keep, SO_Keep, Strata::StencilBit_Single, 0x0>::GetRHI(); break;
+		case EStrataTileType::EComplex: StencilRef = Strata::StencilBit_Complex; GraphicsPSOInit.DepthStencilState = TStaticDepthStencilState<false, CompareFunction, true, CF_Equal, SO_Keep, SO_Keep, SO_Keep, true, CF_Equal, SO_Keep, SO_Keep, SO_Keep, Strata::StencilBit_Complex, 0x0>::GetRHI(); break;
 		default: check(false);
 		}
 	}
@@ -1711,7 +1711,7 @@ static uint32 InternalSetBoundingGeometryDepthState(FGraphicsPipelineStateInitia
 }
 
 /** Sets up rasterizer and depth state for rendering bounding geometry in a deferred pass. */
-static uint32 SetBoundingGeometryRasterizerAndDepthState(FGraphicsPipelineStateInitializer& GraphicsPSOInit, const FViewInfo& View, bool bCameraInsideLightGeometry, EStrataTileMaterialType TileType)
+static uint32 SetBoundingGeometryRasterizerAndDepthState(FGraphicsPipelineStateInitializer& GraphicsPSOInit, const FViewInfo& View, bool bCameraInsideLightGeometry, EStrataTileType TileType)
 {
 	uint32 StencilRef = 0u;
 	InternalSetBoundingGeometryRasterizerState(GraphicsPSOInit, View, bCameraInsideLightGeometry);
@@ -1846,7 +1846,7 @@ static void InternalRenderLight(
 	const FLightSceneInfo* LightSceneInfo,
 	TShaderType& PixelShader,
 	TParametersType* PassParameters,
-	EStrataTileMaterialType StrataTileMaterialType,
+	EStrataTileType StrataTileMaterialType,
 	const TCHAR* ShaderName)
 {
 	const FLightSceneProxy* RESTRICT LightProxy = LightSceneInfo->Proxy;
@@ -1862,8 +1862,8 @@ static void InternalRenderLight(
 	{
 
 		const bool bIsRadial = LightType != LightType_Directional;
-		const bool bEnableStrataTiledPass   = StrataTileMaterialType != EStrataTileMaterialType::ECount;
-		const bool bEnableStrataStencilTest = StrataTileMaterialType != EStrataTileMaterialType::ECount && bIsRadial;
+		const bool bEnableStrataTiledPass   = StrataTileMaterialType != EStrataTileType::ECount;
+		const bool bEnableStrataStencilTest = StrataTileMaterialType != EStrataTileType::ECount && bIsRadial;
 
 		FGraphicsPipelineStateInitializer GraphicsPSOInit;
 		RHICmdList.ApplyCachedRenderTargets(GraphicsPSOInit);
@@ -1915,7 +1915,7 @@ static void InternalRenderLight(
 			{
 				check(Strata::IsStrataEnabled());
 				SetShaderParameters(RHICmdList, TileVertexShader, TileVertexShader.GetVertexShader(), VSParameters);
-				RHICmdList.DrawPrimitiveIndirect(VSParameters.TileIndirectBuffer->GetIndirectRHICallBuffer(), 0);
+				RHICmdList.DrawPrimitiveIndirect(VSParameters.TileIndirectBuffer->GetIndirectRHICallBuffer(), Strata::TileTypeDrawIndirectArgOffset(StrataTileMaterialType));
 			}
 			else
 			{
@@ -2066,7 +2066,7 @@ static void RenderLight(
 		FDeferredLightOverlapPS::FPermutationDomain PermutationVector;
 		PermutationVector.Set<FDeferredLightOverlapPS::FRadialAttenuation>(bIsRadial);
 		TShaderMapRef<FDeferredLightOverlapPS> PixelShader(View.ShaderMap, PermutationVector);		
-		InternalRenderLight(GraphBuilder, Scene, View, LightSceneInfo, PixelShader, PassParameters, EStrataTileMaterialType::ECount, TEXT("Light::StandardDeferred(Overlap)"));
+		InternalRenderLight(GraphBuilder, Scene, View, LightSceneInfo, PixelShader, PassParameters, EStrataTileType::ECount, TEXT("Light::StandardDeferred(Overlap)"));
 	}
 	// Lighting shader
 	else
@@ -2088,9 +2088,9 @@ static void RenderLight(
 		{
 			// Note: we register all tile types here in order to have all resources tracked properly and being able 
 			//       to create a single pass parameters struct instead of created one for each tile types
-			PassParameters->StrataTileSimple = Strata::SetTileParameters(GraphBuilder, View, EStrataTileMaterialType::ESingle);
-			PassParameters->StrataTileSingle = Strata::SetTileParameters(GraphBuilder, View, EStrataTileMaterialType::ESimple);
-			PassParameters->StrataTileComplex = Strata::SetTileParameters(GraphBuilder, View, EStrataTileMaterialType::EComplex);
+			PassParameters->StrataTileSimple = Strata::SetTileParameters(GraphBuilder, View, EStrataTileType::ESingle);
+			PassParameters->StrataTileSingle = Strata::SetTileParameters(GraphBuilder, View, EStrataTileType::ESimple);
+			PassParameters->StrataTileComplex = Strata::SetTileParameters(GraphBuilder, View, EStrataTileType::EComplex);
 		}
 
 		FDeferredLightPS::FPermutationDomain PermutationVector;
@@ -2128,21 +2128,21 @@ static void RenderLight(
 		{
 			// Simple tiles
 			{
-				const EStrataTileMaterialType TileType = EStrataTileMaterialType::ESimple;
+				const EStrataTileType TileType = EStrataTileType::ESimple;
 				PermutationVector.Set<FDeferredLightPS::FStrataTileType>(TileType);
 				TShaderMapRef< FDeferredLightPS > PixelShader(View.ShaderMap, PermutationVector);
 				InternalRenderLight(GraphBuilder, Scene, View, LightSceneInfo, PixelShader, PassParameters, TileType, TEXT("Light::StandardDeferred(Simple)"));
 			}
 			// Single tiles
 			{
-				const EStrataTileMaterialType TileType = EStrataTileMaterialType::ESingle;
+				const EStrataTileType TileType = EStrataTileType::ESingle;
 				PermutationVector.Set<FDeferredLightPS::FStrataTileType>(TileType);
 				TShaderMapRef< FDeferredLightPS > PixelShader(View.ShaderMap, PermutationVector);
 				InternalRenderLight(GraphBuilder, Scene, View, LightSceneInfo, PixelShader, PassParameters, TileType, TEXT("Light::StandardDeferred(Single)"));
 			}
 			// Complex tiles
 			{
-				const EStrataTileMaterialType TileType = EStrataTileMaterialType::EComplex;
+				const EStrataTileType TileType = EStrataTileType::EComplex;
 				PermutationVector.Set<FDeferredLightPS::FStrataTileType>(TileType);
 				TShaderMapRef< FDeferredLightPS > PixelShader(View.ShaderMap, PermutationVector);
 				InternalRenderLight(GraphBuilder, Scene, View, LightSceneInfo, PixelShader, PassParameters, TileType, TEXT("Light::StandardDeferred(Complex)"));
@@ -2152,7 +2152,7 @@ static void RenderLight(
 		{
 			PermutationVector.Set< FDeferredLightPS::FStrataTileType>(0);
 			TShaderMapRef< FDeferredLightPS > PixelShader(View.ShaderMap, PermutationVector);
-			InternalRenderLight(GraphBuilder, Scene, View, LightSceneInfo, PixelShader, PassParameters, EStrataTileMaterialType::ECount, TEXT("Light::StandardDeferred"));
+			InternalRenderLight(GraphBuilder, Scene, View, LightSceneInfo, PixelShader, PassParameters, EStrataTileType::ECount, TEXT("Light::StandardDeferred"));
 		}
 	}
 }
@@ -2401,7 +2401,7 @@ static void InternalRenderSimpleLightsStandardDeferred(
 	const uint32 NumViews,
 	const FMinimalSceneTextures& SceneTextures,
 	const FSimpleLightArray& SimpleLights,
-	EStrataTileMaterialType TileType)
+	EStrataTileType TileType)
 {
 	FSimpleLightsStandardDeferredParameters* PassParameters = GraphBuilder.AllocParameters<FSimpleLightsStandardDeferredParameters>();
 	*PassParameters = GetRenderLightSimpleParameters(
@@ -2422,7 +2422,7 @@ static void InternalRenderSimpleLightsStandardDeferred(
 	PermutationVector.Set< FDeferredLightPS::FHairLighting>(0);
 	PermutationVector.Set< FDeferredLightPS::FAtmosphereTransmittance >(false);
 	PermutationVector.Set< FDeferredLightPS::FCloudTransmittance >(false);
-	PermutationVector.Set< FDeferredLightPS::FStrataTileType>(TileType != EStrataTileMaterialType::ECount ? TileType : 0);
+	PermutationVector.Set< FDeferredLightPS::FStrataTileType>(TileType != EStrataTileType::ECount ? TileType : 0);
 	TShaderMapRef<FDeferredLightPS> PixelShader(View.ShaderMap, PermutationVector);
 
 	FDeferredLightVS::FPermutationDomain PermutationVectorVS;
@@ -2499,13 +2499,13 @@ void FDeferredShadingSceneRenderer::RenderSimpleLightsStandardDeferred(
 
 		if (Strata::IsStrataEnabled())
 		{
-			InternalRenderSimpleLightsStandardDeferred(GraphBuilder, Scene, View, ViewIndex, NumViews, SceneTextures, SimpleLights, EStrataTileMaterialType::ESimple);
-			InternalRenderSimpleLightsStandardDeferred(GraphBuilder, Scene, View, ViewIndex, NumViews, SceneTextures, SimpleLights, EStrataTileMaterialType::ESingle);
-			InternalRenderSimpleLightsStandardDeferred(GraphBuilder, Scene, View, ViewIndex, NumViews, SceneTextures, SimpleLights, EStrataTileMaterialType::EComplex);
+			InternalRenderSimpleLightsStandardDeferred(GraphBuilder, Scene, View, ViewIndex, NumViews, SceneTextures, SimpleLights, EStrataTileType::ESimple);
+			InternalRenderSimpleLightsStandardDeferred(GraphBuilder, Scene, View, ViewIndex, NumViews, SceneTextures, SimpleLights, EStrataTileType::ESingle);
+			InternalRenderSimpleLightsStandardDeferred(GraphBuilder, Scene, View, ViewIndex, NumViews, SceneTextures, SimpleLights, EStrataTileType::EComplex);
 		}
 		else
 		{
-			InternalRenderSimpleLightsStandardDeferred(GraphBuilder, Scene, View, ViewIndex, NumViews, SceneTextures, SimpleLights, EStrataTileMaterialType::ECount);
+			InternalRenderSimpleLightsStandardDeferred(GraphBuilder, Scene, View, ViewIndex, NumViews, SceneTextures, SimpleLights, EStrataTileType::ECount);
 		}
 	}
 }
diff --git a/Engine/Source/Runtime/Renderer/Private/Strata/Strata.cpp b/Engine/Source/Runtime/Renderer/Private/Strata/Strata.cpp
--- a/Engine/Source/Runtime/Renderer/Private/Strata/Strata.cpp
+++ b/Engine/Source/Runtime/Renderer/Private/Strata/Strata.cpp
@@ -9,7 +9,6 @@
 #include "RendererInterface.h"
 #include "UniformBuffer.h"
 #include "SceneTextureParameters.h"
-#include "StrataDefinitions.h"
 #include "ScreenPass.h"
 #include "ShaderCompiler.h"
 
@@ -78,33 +77,29 @@ void FStrataSceneData::Reset()
 	MaterialTextureArrayUAV = nullptr;
 	MaterialTextureArraySRV = nullptr;
 
-	for (uint32 i = 0; i < EStrataTileMaterialType::ECount; ++i)
+	ClassificationTileIndirectBuffer = nullptr;
+	ClassificationTileIndirectBufferUAV = nullptr;
+	for (uint32 i = 0; i < EStrataTileType::ECount; ++i)
 	{
 		ClassificationTileListBuffer[i] = nullptr;
 		ClassificationTileListBufferUAV[i] = nullptr;
 		ClassificationTileListBufferSRV[i] = nullptr;
-		ClassificationTileIndirectBuffer[i] = nullptr;
-		ClassificationTileIndirectBufferUAV[i] = nullptr;
-		ClassificationTileIndirectBufferSRV[i] = nullptr;
 	}
 
-	OpaqueRoughRefractionClassificationTileListBuffer = nullptr;
-	OpaqueRoughRefractionClassificationTileListBufferSRV = nullptr;
-	OpaqueRoughRefractionClassificationTileListBufferUAV = nullptr;
-
 	SeparatedSubSurfaceSceneColor = nullptr;
 	SeparatedOpaqueRoughRefractionSceneColor = nullptr;
 
 	StrataGlobalUniformParameters = nullptr;
 }
 
-const TCHAR* ToString(EStrataTileMaterialType Type)
+const TCHAR* ToString(EStrataTileType Type)
 {
 	switch (Type)
 	{
-	case EStrataTileMaterialType::ESimple:		return TEXT("Simple");
-	case EStrataTileMaterialType::ESingle:		return TEXT("Single");
-	case EStrataTileMaterialType::EComplex:		return TEXT("Complex");
+	case EStrataTileType::ESimple:					return TEXT("Simple");
+	case EStrataTileType::ESingle:					return TEXT("Single");
+	case EStrataTileType::EComplex:					return TEXT("Complex");
+	case EStrataTileType::EOpaqueRoughRefraction:	return TEXT("OpaqueRoughRefraction");
 	}
 	return TEXT("Unknown");
 }
@@ -164,29 +159,25 @@ void InitialiseStrataFrameSceneData(FSceneRenderer& SceneRenderer, FRDGBuilder&
 		const int32 TileInPixel = GetStrataBufferTileSize();
 		const FIntPoint TileResolution(FMath::DivideAndRoundUp(SceneTextureExtent.X, TileInPixel), FMath::DivideAndRoundUp(SceneTextureExtent.Y, TileInPixel));
 
+		const TCHAR* StrataTileListBufferNames[EStrataTileType::ECount] =
+		{
+			TEXT("Strata.StrataTileListBuffer(Simple)"),
+			TEXT("Strata.StrataTileListBuffer(Single)"),
+			TEXT("Strata.StrataTileListBuffer(Complex)"),
+			TEXT("Strata.StrataTileListBuffer(OpaqueRoughRefraction)")
+		};
+
 		// Tile classification buffers
 		{
+			StrataSceneData.ClassificationTileIndirectBuffer = GraphBuilder.CreateBuffer(FRDGBufferDesc::CreateIndirectDesc<FRHIDrawIndirectParameters>(EStrataTileType::ECount), TEXT("Strata.StrataTileIndirectBuffer"));
+			StrataSceneData.ClassificationTileIndirectBufferUAV = GraphBuilder.CreateUAV(StrataSceneData.ClassificationTileIndirectBuffer, PF_R32_UINT);
+			AddClearUAVPass(GraphBuilder, StrataSceneData.ClassificationTileIndirectBufferUAV, 0);
 
-			const TCHAR* StrataTileResourceNames[EStrataTileMaterialType::ECount][3] =
+			for (uint32 i = 0; i <= EStrataTileType::EComplex; ++i)
 			{
-				{ TEXT("Strata.StrataTileListBuffer(Simple)"),	TEXT("Strata.StrataTileIndirectBuffer(Simple)") },
-				{ TEXT("Strata.StrataTileListBuffer(Single)"), 	TEXT("Strata.StrataTileIndirectBuffer(Single)") },
-				{ TEXT("Strata.StrataTileListBuffer(Complex)"),	TEXT("Strata.StrataTileIndirectBuffer(Complex)")}
-			};
-
-			// As of today we allocate one index+indirect buffer for each EStrataTileMaterialType.
-			// This is fine for two types, later we might want to have a single list and indirect buffer with offsets.
-			for (uint32 i = 0; i < EStrataTileMaterialType::ECount; ++i)
-			{
-				StrataSceneData.ClassificationTileListBuffer[i] = GraphBuilder.CreateBuffer(FRDGBufferDesc::CreateBufferDesc(sizeof(uint32), TileResolution.X * TileResolution.Y), StrataTileResourceNames[i][0]);
+				StrataSceneData.ClassificationTileListBuffer[i] = GraphBuilder.CreateBuffer(FRDGBufferDesc::CreateBufferDesc(sizeof(uint32), TileResolution.X * TileResolution.Y), StrataTileListBufferNames[i]);
 				StrataSceneData.ClassificationTileListBufferSRV[i] = GraphBuilder.CreateSRV(StrataSceneData.ClassificationTileListBuffer[i], PF_R32_UINT);
 				StrataSceneData.ClassificationTileListBufferUAV[i] = GraphBuilder.CreateUAV(StrataSceneData.ClassificationTileListBuffer[i], PF_R32_UINT);
-
-				StrataSceneData.ClassificationTileIndirectBuffer[i] = GraphBuilder.CreateBuffer(FRDGBufferDesc::CreateIndirectDesc<FRHIDrawIndirectParameters>(), StrataTileResourceNames[i][1]);
-				StrataSceneData.ClassificationTileIndirectBufferSRV[i] = GraphBuilder.CreateSRV(StrataSceneData.ClassificationTileIndirectBuffer[i], PF_R32_UINT);
-				StrataSceneData.ClassificationTileIndirectBufferUAV[i] = GraphBuilder.CreateUAV(StrataSceneData.ClassificationTileIndirectBuffer[i], PF_R32_UINT);
-
-				AddClearUAVPass(GraphBuilder, StrataSceneData.ClassificationTileIndirectBufferUAV[i], 0);
 			}
 		}
 
@@ -207,15 +198,9 @@ void InitialiseStrataFrameSceneData(FSceneRenderer& SceneRenderer, FRDGBuilder&
 			const int32 TileListBufferElementCount					= bIsStrataOpaqueMaterialRoughRefractionEnabled ? TileResolution.X * TileResolution.Y : 4;
 			const FIntPoint OpaqueRoughRefractionSceneExtent		= bIsStrataOpaqueMaterialRoughRefractionEnabled ? SceneTextureExtent : FIntPoint(4, 4);
 			
-
-			StrataSceneData.OpaqueRoughRefractionClassificationTileListBuffer = GraphBuilder.CreateBuffer(FRDGBufferDesc::CreateBufferDesc(sizeof(uint32), TileListBufferElementCount), TEXT("Strata.StrataTileListBuffer(RoughRefrac)"));
-			StrataSceneData.OpaqueRoughRefractionClassificationTileListBufferSRV = GraphBuilder.CreateSRV(StrataSceneData.OpaqueRoughRefractionClassificationTileListBuffer, PF_R32_UINT);
-			StrataSceneData.OpaqueRoughRefractionClassificationTileListBufferUAV = GraphBuilder.CreateUAV(StrataSceneData.OpaqueRoughRefractionClassificationTileListBuffer, PF_R32_UINT);
-
-			StrataSceneData.OpaqueRoughRefractionClassificationTileIndirectBuffer = GraphBuilder.CreateBuffer(FRDGBufferDesc::CreateIndirectDesc<FRHIDrawIndirectParameters>(), TEXT("Strata.StrataTileIndirectBuffer(RoughRefrac)"));
-			StrataSceneData.OpaqueRoughRefractionClassificationTileIndirectBufferSRV = GraphBuilder.CreateSRV(StrataSceneData.OpaqueRoughRefractionClassificationTileIndirectBuffer, PF_R32_UINT);
-			StrataSceneData.OpaqueRoughRefractionClassificationTileIndirectBufferSRV = GraphBuilder.CreateSRV(StrataSceneData.OpaqueRoughRefractionClassificationTileIndirectBuffer, PF_R32_UINT);
-			StrataSceneData.OpaqueRoughRefractionClassificationTileIndirectBufferUAV = GraphBuilder.CreateUAV(StrataSceneData.OpaqueRoughRefractionClassificationTileIndirectBuffer, PF_R32_UINT);
+			StrataSceneData.ClassificationTileListBuffer[EStrataTileType::EOpaqueRoughRefraction] = GraphBuilder.CreateBuffer(FRDGBufferDesc::CreateBufferDesc(sizeof(uint32), TileListBufferElementCount), StrataTileListBufferNames[EStrataTileType::EOpaqueRoughRefraction]);
+			StrataSceneData.ClassificationTileListBufferSRV[EStrataTileType::EOpaqueRoughRefraction] = GraphBuilder.CreateSRV(StrataSceneData.ClassificationTileListBuffer[EStrataTileType::EOpaqueRoughRefraction], PF_R32_UINT);
+			StrataSceneData.ClassificationTileListBufferUAV[EStrataTileType::EOpaqueRoughRefraction] = GraphBuilder.CreateUAV(StrataSceneData.ClassificationTileListBuffer[EStrataTileType::EOpaqueRoughRefraction], PF_R32_UINT);
 
 			StrataSceneData.OpaqueRoughRefractionTexture = GraphBuilder.CreateTexture(
 				FRDGTextureDesc::Create2D(OpaqueRoughRefractionSceneExtent, PF_R8, FClearValueBinding::Black, TexCreate_ShaderResource | TexCreate_UAV | TexCreate_RenderTargetable), TEXT("Strata.OpaqueRoughRefractionTexture"));
@@ -481,13 +466,10 @@ class FStrataMaterialTileClassificationPassCS : public FGlobalShader
 		SHADER_PARAMETER(uint32, MaxBytesPerPixel)
 		SHADER_PARAMETER_RDG_TEXTURE(Texture2D, TopLayerTexture)
 		SHADER_PARAMETER_RDG_TEXTURE_SRV(Texture2DArray<uint>, MaterialTextureArray)
-		SHADER_PARAMETER_RDG_BUFFER_UAV(RWBuffer, SimpleTileIndirectDataBuffer)
+		SHADER_PARAMETER_RDG_BUFFER_UAV(RWBuffer, TileIndirectDataBuffer)
 		SHADER_PARAMETER_RDG_BUFFER_UAV(RWBuffer, SimpleTileListDataBuffer)
-		SHADER_PARAMETER_RDG_BUFFER_UAV(RWBuffer, SingleTileIndirectDataBuffer)
 		SHADER_PARAMETER_RDG_BUFFER_UAV(RWBuffer, SingleTileListDataBuffer)
-		SHADER_PARAMETER_RDG_BUFFER_UAV(RWBuffer, ComplexTileIndirectDataBuffer)
 		SHADER_PARAMETER_RDG_BUFFER_UAV(RWBuffer, ComplexTileListDataBuffer)
-		SHADER_PARAMETER_RDG_BUFFER_UAV(RWBuffer, OpaqueRoughRefractionTileIndirectDataBuffer)
 		SHADER_PARAMETER_RDG_BUFFER_UAV(RWBuffer, OpaqueRoughRefractionTileListDataBuffer)
 		SHADER_PARAMETER_RDG_TEXTURE_UAV(RWTexture2D<uint2>, SSSTextureUAV)
 		SHADER_PARAMETER_RDG_TEXTURE(Texture2D<float>, OpaqueRoughRefractionTexture)
@@ -553,13 +535,13 @@ class FStrataMaterialStencilTaggingPassPS : public FGlobalShader
 IMPLEMENT_GLOBAL_SHADER(FStrataTilePassVS, "/Engine/Private/Strata/StrataTiling.usf", "StrataTilePassVS", SF_Vertex);
 IMPLEMENT_GLOBAL_SHADER(FStrataMaterialStencilTaggingPassPS, "/Engine/Private/Strata/StrataTiling.usf", "StencilTaggingMainPS", SF_Pixel);
 
-static FStrataTileParameter InternalSetTileParameters(FRDGBuilder* GraphBuilder, const FViewInfo& View, const EStrataTileMaterialType TileType)
+static FStrataTileParameter InternalSetTileParameters(FRDGBuilder* GraphBuilder, const FViewInfo& View, const EStrataTileType TileType)
 {
 	FStrataTileParameter Out;
-	if (View.StrataSceneData && TileType != EStrataTileMaterialType::ECount)
+	if (View.StrataSceneData && TileType != EStrataTileType::ECount)
 	{
 		Out.TileListBuffer = View.StrataSceneData->ClassificationTileListBufferSRV[TileType];
-		Out.TileIndirectBuffer = View.StrataSceneData->ClassificationTileIndirectBuffer[TileType];
+		Out.TileIndirectBuffer = View.StrataSceneData->ClassificationTileIndirectBuffer;
 	}
 	else if (GraphBuilder)
 	{
@@ -573,7 +555,7 @@ static FStrataTileParameter InternalSetTileParameters(FRDGBuilder* GraphBuilder,
 
 FStrataTilePassVS::FParameters SetTileParameters(
 	const FViewInfo& View,
-	const EStrataTileMaterialType TileType,
+	const EStrataTileType TileType,
 	EPrimitiveType& PrimitiveType)
 {
 	FStrataTileParameter Temp = InternalSetTileParameters(nullptr, View, TileType);
@@ -591,7 +573,7 @@ FStrataTilePassVS::FParameters SetTileParameters(
 FStrataTilePassVS::FParameters SetTileParameters(
 	FRDGBuilder& GraphBuilder, 
 	const FViewInfo& View, 
-	const EStrataTileMaterialType TileType,
+	const EStrataTileType TileType,
 	EPrimitiveType& PrimitiveType)
 {
 	FStrataTileParameter Temp = InternalSetTileParameters(&GraphBuilder, View, TileType);
@@ -606,11 +588,17 @@ FStrataTilePassVS::FParameters SetTileParameters(
 	return Out;
 }
 
-FStrataTileParameter SetTileParameters(FRDGBuilder& GraphBuilder, const FViewInfo& View, const EStrataTileMaterialType TileType)
+FStrataTileParameter SetTileParameters(FRDGBuilder& GraphBuilder, const FViewInfo& View, const EStrataTileType TileType)
 {
 	return InternalSetTileParameters(&GraphBuilder, View, TileType);
 }
 
+uint32 TileTypeDrawIndirectArgOffset(const EStrataTileType Type)
+{
+	check(Type >= 0 && Type < EStrataTileType::ECount);
+	return GetStrataTileTypeDrawIndirectArgOffset_Byte(Type);
+}
+
 // Add additionnaly bits for filling/clearing stencil to ensure that the 'Strata' bits are not corrupted by the stencil shadows 
 // when generating shadow mask. Withouth these 'trailing' bits, the incr./decr. operation would change/corrupt the 'Strata' bits
 constexpr uint32 StencilBit_Fast_1	  = 0x07u | StencilBit_Fast;
@@ -622,7 +610,7 @@ static void AddStrataInternalClassificationTilePass(
 	const FViewInfo& View,
 	const FRDGTextureRef* DepthTexture,
 	const FRDGTextureRef* ColorTexture,
-	EStrataTileMaterialType TileMaterialType,
+	EStrataTileType TileMaterialType,
 	const bool bDebug = false)
 {
 	EPrimitiveType StrataTilePrimitiveType = PT_TriangleList;
@@ -645,9 +633,10 @@ static void AddStrataInternalClassificationTilePass(
 		ParametersPS->RenderTargets[0] = FRenderTargetBinding(*ColorTexture, ERenderTargetLoadAction::ELoad);
 		switch (TileMaterialType)
 		{
-		case EStrataTileMaterialType::ESimple: ParametersPS->DebugTileColor = FVector4f(0.0f, 1.0f, 0.0f, 1.0); break;
-		case EStrataTileMaterialType::ESingle: ParametersPS->DebugTileColor = FVector4f(1.0f, 1.0f, 0.0f, 1.0); break;
-		case EStrataTileMaterialType::EComplex: ParametersPS->DebugTileColor = FVector4f(1.0f, 0.0f, 0.0f, 1.0); break;
+		case EStrataTileType::ESimple:					ParametersPS->DebugTileColor = FVector4f(0.0f, 1.0f, 0.0f, 1.0); break;
+		case EStrataTileType::ESingle:					ParametersPS->DebugTileColor = FVector4f(1.0f, 1.0f, 0.0f, 1.0); break;
+		case EStrataTileType::EComplex:					ParametersPS->DebugTileColor = FVector4f(1.0f, 0.0f, 0.0f, 1.0); break;
+		case EStrataTileType::EOpaqueRoughRefraction:	ParametersPS->DebugTileColor = FVector4f(0.0f, 0.0f, 1.0f, 1.0); break;
 		default: check(false);
 		}
 	}
@@ -681,14 +670,14 @@ static void AddStrataInternalClassificationTilePass(
 			}
 			else
 			{
-				check(TileMaterialType != EStrataTileMaterialType::ECount);
+				check(TileMaterialType != EStrataTileType::ECount && TileMaterialType != EStrataTileType::EOpaqueRoughRefraction);
 
-				// No blending and no pixel shader required. Stencil will be writen to.
+				// No blending and no pixel shader required. Stencil will be written to.
 				GraphicsPSOInit.BoundShaderState.PixelShaderRHI = nullptr;
 				GraphicsPSOInit.BlendState = TStaticBlendState<>::GetRHI();
 				switch (TileMaterialType)
 				{
-				case EStrataTileMaterialType::ESimple:
+				case EStrataTileType::ESimple:
 				{
 					GraphicsPSOInit.DepthStencilState = TStaticDepthStencilState<
 						false, CF_Always,
@@ -698,7 +687,7 @@ static void AddStrataInternalClassificationTilePass(
 					StencilRef = StencilBit_Fast_1;
 				}
 				break;
-				case EStrataTileMaterialType::ESingle:
+				case EStrataTileType::ESingle:
 				{
 					GraphicsPSOInit.DepthStencilState = TStaticDepthStencilState<
 						false, CF_Always,
@@ -708,7 +697,7 @@ static void AddStrataInternalClassificationTilePass(
 					StencilRef = StencilBit_Single_1;
 				}
 				break;
-				case EStrataTileMaterialType::EComplex:
+				case EStrataTileType::EComplex:
 				{
 					GraphicsPSOInit.DepthStencilState = TStaticDepthStencilState<
 						false, CF_Always,
@@ -732,7 +721,7 @@ static void AddStrataInternalClassificationTilePass(
 
 			RHICmdList.SetViewport(0, 0, 0.0f, OutputResolution.X, OutputResolution.Y, 1.0f);
 			RHICmdList.SetStreamSource(0, nullptr, 0);
-			RHICmdList.DrawPrimitiveIndirect(ParametersPS->VS.TileIndirectBuffer->GetIndirectRHICallBuffer(), 0);
+			RHICmdList.DrawPrimitiveIndirect(ParametersPS->VS.TileIndirectBuffer->GetIndirectRHICallBuffer(), TileTypeDrawIndirectArgOffset(TileMaterialType));
 		});
 }
 
@@ -744,9 +733,9 @@ void AddStrataStencilPass(
 	for (int32 i = 0; i < Views.Num(); ++i)
 	{
 		const FViewInfo& View = Views[i];
-		AddStrataInternalClassificationTilePass(GraphBuilder, View, &SceneTextures.Depth.Target, nullptr, EStrataTileMaterialType::ESimple);
-		AddStrataInternalClassificationTilePass(GraphBuilder, View, &SceneTextures.Depth.Target, nullptr, EStrataTileMaterialType::ESingle);
-		AddStrataInternalClassificationTilePass(GraphBuilder, View, &SceneTextures.Depth.Target, nullptr, EStrataTileMaterialType::EComplex);
+		AddStrataInternalClassificationTilePass(GraphBuilder, View, &SceneTextures.Depth.Target, nullptr, EStrataTileType::ESimple);
+		AddStrataInternalClassificationTilePass(GraphBuilder, View, &SceneTextures.Depth.Target, nullptr, EStrataTileType::ESingle);
+		AddStrataInternalClassificationTilePass(GraphBuilder, View, &SceneTextures.Depth.Target, nullptr, EStrataTileType::EComplex);
 	}
 }
 
@@ -831,18 +820,11 @@ void AddStrataMaterialClassificationPass(FRDGBuilder& GraphBuilder, const FMinim
 			PassParameters->MaterialTextureArray = View.StrataSceneData->MaterialTextureArraySRV;
 			PassParameters->SSSTextureUAV = View.StrataSceneData->SSSTextureUAV;
 			PassParameters->OpaqueRoughRefractionTexture = View.StrataSceneData->OpaqueRoughRefractionTexture;
-			// Simple
-			PassParameters->SimpleTileListDataBuffer = View.StrataSceneData->ClassificationTileListBufferUAV[EStrataTileMaterialType::ESimple];
-			PassParameters->SimpleTileIndirectDataBuffer = View.StrataSceneData->ClassificationTileIndirectBufferUAV[EStrataTileMaterialType::ESimple];
-			// Single
-			PassParameters->SingleTileListDataBuffer = View.StrataSceneData->ClassificationTileListBufferUAV[EStrataTileMaterialType::ESingle];
-			PassParameters->SingleTileIndirectDataBuffer = View.StrataSceneData->ClassificationTileIndirectBufferUAV[EStrataTileMaterialType::ESingle];
-			// Complex
-			PassParameters->ComplexTileListDataBuffer = View.StrataSceneData->ClassificationTileListBufferUAV[EStrataTileMaterialType::EComplex];
-			PassParameters->ComplexTileIndirectDataBuffer = View.StrataSceneData->ClassificationTileIndirectBufferUAV[EStrataTileMaterialType::EComplex];
-			// Opaque Rough Refraction
-			PassParameters->OpaqueRoughRefractionTileListDataBuffer = View.StrataSceneData->OpaqueRoughRefractionClassificationTileListBufferUAV;
-			PassParameters->OpaqueRoughRefractionTileIndirectDataBuffer = View.StrataSceneData->OpaqueRoughRefractionClassificationTileIndirectBufferUAV;
+			PassParameters->TileIndirectDataBuffer = View.StrataSceneData->ClassificationTileIndirectBufferUAV;
+			PassParameters->SimpleTileListDataBuffer = View.StrataSceneData->ClassificationTileListBufferUAV[EStrataTileType::ESimple];
+			PassParameters->SingleTileListDataBuffer = View.StrataSceneData->ClassificationTileListBufferUAV[EStrataTileType::ESingle];
+			PassParameters->ComplexTileListDataBuffer = View.StrataSceneData->ClassificationTileListBufferUAV[EStrataTileType::EComplex];
+			PassParameters->OpaqueRoughRefractionTileListDataBuffer = View.StrataSceneData->ClassificationTileListBufferUAV[EStrataTileType::EOpaqueRoughRefraction];
 
 			const uint32 GroupSize = 8;
 			FComputeShaderUtils::AddPass(
@@ -909,9 +891,16 @@ FScreenPassTexture AddStrataDebugPasses(FRDGBuilder& GraphBuilder, const FViewIn
 	{
 		RDG_EVENT_SCOPE(GraphBuilder, "Strata::VisualizeClassification");
 		const bool bDebugPass = true;
-		AddStrataInternalClassificationTilePass(GraphBuilder, View, nullptr, &ScreenPassSceneColor.Texture, EStrataTileMaterialType::ESimple, bDebugPass);
-		AddStrataInternalClassificationTilePass(GraphBuilder, View, nullptr, &ScreenPassSceneColor.Texture, EStrataTileMaterialType::ESingle, bDebugPass);
-		AddStrataInternalClassificationTilePass(GraphBuilder, View, nullptr, &ScreenPassSceneColor.Texture, EStrataTileMaterialType::EComplex, bDebugPass);
+		if (StrataClassificationDebug > 1)
+		{
+			AddStrataInternalClassificationTilePass(GraphBuilder, View, nullptr, &ScreenPassSceneColor.Texture, EStrataTileType::EOpaqueRoughRefraction, bDebugPass);
+		}
+		else
+		{
+			AddStrataInternalClassificationTilePass(GraphBuilder, View, nullptr, &ScreenPassSceneColor.Texture, EStrataTileType::ESimple, bDebugPass);
+			AddStrataInternalClassificationTilePass(GraphBuilder, View, nullptr, &ScreenPassSceneColor.Texture, EStrataTileType::ESingle, bDebugPass);
+			AddStrataInternalClassificationTilePass(GraphBuilder, View, nullptr, &ScreenPassSceneColor.Texture, EStrataTileType::EComplex, bDebugPass);
+		}
 	}
 
 	StrataRoughRefractionRnD(GraphBuilder, View, ScreenPassSceneColor);
diff --git a/Engine/Source/Runtime/Renderer/Private/Strata/Strata.h b/Engine/Source/Runtime/Renderer/Private/Strata/Strata.h
--- a/Engine/Source/Runtime/Renderer/Private/Strata/Strata.h
+++ b/Engine/Source/Runtime/Renderer/Private/Strata/Strata.h
@@ -8,6 +8,7 @@
 #include "MeshPassProcessor.h"
 #include "UnifiedBuffer.h"
 #include "RHIUtilities.h"
+#include "StrataDefinitions.h"
 
 // Forward declarations.
 class FScene;
@@ -43,15 +44,17 @@ BEGIN_GLOBAL_SHADER_PARAMETER_STRUCT(FStrataGlobalUniformParameters, RENDERER_AP
 	SHADER_PARAMETER_RDG_TEXTURE(Texture2D<float>, OpaqueRoughRefractionTexture)
 END_GLOBAL_SHADER_PARAMETER_STRUCT()
 
-enum EStrataTileMaterialType : uint32
+// This must map to the STRATA_TILE_TYPE defines.
+enum EStrataTileType : uint32
 {
-	ESimple  = 0,
-	ESingle  = 1,
-	EComplex = 2,
+	ESimple  = STRATA_TILE_TYPE_SIMPLE,
+	ESingle = STRATA_TILE_TYPE_SINGLE,
+	EComplex = STRATA_TILE_TYPE_COMPLEX,
+	EOpaqueRoughRefraction = STRATA_TILE_TYPE_ROUGH_REFRACT,
 	ECount
 };
 
-const TCHAR* ToString(EStrataTileMaterialType Type);
+const TCHAR* ToString(EStrataTileType Type);
 
 struct FStrataSceneData
 {
@@ -65,19 +68,12 @@ struct FStrataSceneData
 	FRDGTextureUAVRef MaterialTextureArrayUAV;
 	FRDGTextureSRVRef MaterialTextureArraySRV;
 
-	FRDGBufferRef ClassificationTileListBuffer[EStrataTileMaterialType::ECount];
-	FRDGBufferUAVRef ClassificationTileListBufferUAV[EStrataTileMaterialType::ECount];
-	FRDGBufferSRVRef ClassificationTileListBufferSRV[EStrataTileMaterialType::ECount];
-	FRDGBufferRef ClassificationTileIndirectBuffer[EStrataTileMaterialType::ECount];
-	FRDGBufferUAVRef ClassificationTileIndirectBufferUAV[EStrataTileMaterialType::ECount];
-	FRDGBufferSRVRef ClassificationTileIndirectBufferSRV[EStrataTileMaterialType::ECount];
+	FRDGBufferRef ClassificationTileListBuffer[STRATA_TILE_TYPE_COUNT];
+	FRDGBufferSRVRef ClassificationTileListBufferSRV[STRATA_TILE_TYPE_COUNT];
+	FRDGBufferUAVRef ClassificationTileListBufferUAV[STRATA_TILE_TYPE_COUNT];
 
-	FRDGBufferRef OpaqueRoughRefractionClassificationTileListBuffer;
-	FRDGBufferUAVRef OpaqueRoughRefractionClassificationTileListBufferUAV;
-	FRDGBufferSRVRef OpaqueRoughRefractionClassificationTileListBufferSRV;
-	FRDGBufferRef OpaqueRoughRefractionClassificationTileIndirectBuffer;
-	FRDGBufferUAVRef OpaqueRoughRefractionClassificationTileIndirectBufferUAV;
-	FRDGBufferSRVRef OpaqueRoughRefractionClassificationTileIndirectBufferSRV;
+	FRDGBufferRef ClassificationTileIndirectBuffer;
+	FRDGBufferUAVRef ClassificationTileIndirectBufferUAV;
 
 	FRDGTextureRef TopLayerTexture;
 	FRDGTextureRef SSSTexture;
@@ -168,9 +164,10 @@ class FStrataTilePassVS : public FGlobalShader
 	}
 };
 
-FStrataTileParameter SetTileParameters(FRDGBuilder& GraphBuilder, const FViewInfo& View, const EStrataTileMaterialType Type);
-FStrataTilePassVS::FParameters SetTileParameters(FRDGBuilder& GraphBuilder, const FViewInfo& View, const EStrataTileMaterialType Type, EPrimitiveType& PrimitiveType);
-FStrataTilePassVS::FParameters SetTileParameters(const FViewInfo& View, const EStrataTileMaterialType Type, EPrimitiveType& PrimitiveType);
+FStrataTileParameter SetTileParameters(FRDGBuilder& GraphBuilder, const FViewInfo& View, const EStrataTileType Type);
+FStrataTilePassVS::FParameters SetTileParameters(FRDGBuilder& GraphBuilder, const FViewInfo& View, const EStrataTileType Type, EPrimitiveType& PrimitiveType);
+FStrataTilePassVS::FParameters SetTileParameters(const FViewInfo& View, const EStrataTileType Type, EPrimitiveType& PrimitiveType);
+uint32 TileTypeDrawIndirectArgOffset(const EStrataTileType Type);
 
 bool ShouldRenderStrataRoughRefractionRnD();
 void StrataRoughRefractionRnD(FRDGBuilder& GraphBuilder, const FViewInfo& View, FScreenPassTexture& ScreenPassSceneColor);
