Reorganize flow to prevent calls to IHeadMountedDisplayModule::IsAvailable when not necessary.

#rnx

diff --git a/Engine/Source/Runtime/Renderer/Private/SceneTextures.cpp b/Engine/Source/Runtime/Renderer/Private/SceneTextures.cpp
--- a/Engine/Source/Runtime/Renderer/Private/SceneTextures.cpp
+++ b/Engine/Source/Runtime/Renderer/Private/SceneTextures.cpp
@@ -108,6 +108,12 @@ static EPixelFormat GetGBufferFFormat()
 	return NormalGBufferFormat;
 }
 
+static EPixelFormat GetDefaultMobileSceneColorLowPrecisionFormat()
+{
+	return IHeadMountedDisplayModule::IsAvailable() && IHeadMountedDisplayModule::Get().IsStandaloneStereoOnlyDevice()
+		? PF_R8G8B8A8 : PF_B8G8R8A8;
+}
+
 static EPixelFormat GetMobileSceneColorFormat(const FSceneViewFamily& ViewFamily)
 {
 	bool bRequiresAlphaChannel = IsMobilePropagateAlphaEnabled(ViewFamily.GetShaderPlatform());
@@ -121,11 +127,16 @@ static EPixelFormat GetMobileSceneColorFormat(const FSceneViewFamily& ViewFamily
 		}
 	}
 	
-	const EPixelFormat DefaultLowPrecisionFormat = IHeadMountedDisplayModule::IsAvailable() && IHeadMountedDisplayModule::Get().IsStandaloneStereoOnlyDevice()
-		? PF_R8G8B8A8 : PF_B8G8R8A8;
-	const EPixelFormat DefaultPrecisionFormat = bRequiresAlphaChannel ? PF_FloatRGBA : PF_FloatR11G11B10;
-
-	EPixelFormat DefaultColorFormat = (!IsMobileHDR() || !GSupportsRenderTargetFormat_PF_FloatRGBA) ? DefaultLowPrecisionFormat : DefaultPrecisionFormat;
+	EPixelFormat DefaultColorFormat;
+	const bool bUseLowPrecisionFormat = !IsMobileHDR() || !GSupportsRenderTargetFormat_PF_FloatRGBA;
+	if (bUseLowPrecisionFormat)
+	{
+		DefaultColorFormat = GetDefaultMobileSceneColorLowPrecisionFormat();
+	}
+	else
+	{
+		DefaultColorFormat = bRequiresAlphaChannel ? PF_FloatRGBA : PF_FloatR11G11B10;
+	}
 
 	check(GPixelFormats[DefaultColorFormat].Supported);
 
@@ -139,7 +150,17 @@ static EPixelFormat GetMobileSceneColorFormat(const FSceneViewFamily& ViewFamily
 	case 2:
 		Format = PF_FloatR11G11B10; break;
 	case 3:
-		Format = DefaultLowPrecisionFormat; break;
+		if (bUseLowPrecisionFormat)
+		{
+			// if bUseLowPrecisionFormat, DefaultColorFormat already contains the value of GetDefaultMobileSceneColorLowPrecisionFormat
+			checkSlow(DefaultColorFormat == GetDefaultMobileSceneColorLowPrecisionFormat());
+			Format = DefaultColorFormat;
+		}
+		else
+		{
+			Format = GetDefaultMobileSceneColorLowPrecisionFormat();
+		}
+		break;
 	default:
 		break;
 	}
