Strata opaque rough refraction with place holder blur process.
Tiles with SSS but no rough refraction are also added onto the scene color buffer separately from the blur process.

diff --git a/Engine/Shaders/Private/Strata/Strata.ush b/Engine/Shaders/Private/Strata/Strata.ush
--- a/Engine/Shaders/Private/Strata/Strata.ush
+++ b/Engine/Shaders/Private/Strata/Strata.ush
@@ -581,7 +581,7 @@ FStrataOpaqueRoughRefractionData StrataUnpackOpaqueRoughRefractionData(float In)
 {
 	FStrataOpaqueRoughRefractionData Out = (FStrataOpaqueRoughRefractionData)0;
 	Out.OpaqueRoughRefractionEnabled = (In * 255.0) >= 128.0f;
-	Out.Variance = In * (127.0f / 255.0f) + (- 128.0f / 255.0f);
+	Out.Variance = In * (255.0f / 127.0f) + (- 128.0f / 127.0f); // = (In * 255 - 128) / 127
 	return Out;
 }
 
diff --git a/Engine/Shaders/Private/Strata/StrataMaterialClassification.usf b/Engine/Shaders/Private/Strata/StrataMaterialClassification.usf
--- a/Engine/Shaders/Private/Strata/StrataMaterialClassification.usf
+++ b/Engine/Shaders/Private/Strata/StrataMaterialClassification.usf
@@ -33,6 +33,7 @@ RWBuffer<uint> ComplexTileListDataBuffer;
 
 #if STRATA_OPAQUE_ROUGH_REFRACTION_ENABLED
 RWBuffer<uint> OpaqueRoughRefractionTileListDataBuffer;
+RWBuffer<uint> SSSWithoutOpaqueRoughRefractionTileListDataBuffer;
 Texture2D<float> OpaqueRoughRefractionTexture;
 #endif
 
@@ -50,10 +51,11 @@ void TileMainCS(uint2 DispatchThreadId : SV_DispatchThreadID, uint LinearIndex :
 	if (all(DispatchThreadId == 0))
 	{
 		const uint IndexCountPerInstance = bRectPrimitive > 0 ? 4 : 6;
-		TileIndirectDataBuffer[GetStrataTileTypeDrawIndirectArgOffset_DWord(STRATA_TILE_TYPE_SIMPLE)		+ 0] = IndexCountPerInstance;
-		TileIndirectDataBuffer[GetStrataTileTypeDrawIndirectArgOffset_DWord(STRATA_TILE_TYPE_SINGLE)		+ 0] = IndexCountPerInstance;
-		TileIndirectDataBuffer[GetStrataTileTypeDrawIndirectArgOffset_DWord(STRATA_TILE_TYPE_COMPLEX)		+ 0] = IndexCountPerInstance;
-		TileIndirectDataBuffer[GetStrataTileTypeDrawIndirectArgOffset_DWord(STRATA_TILE_TYPE_ROUGH_REFRACT) + 0] = IndexCountPerInstance;
+		TileIndirectDataBuffer[GetStrataTileTypeDrawIndirectArgOffset_DWord(STRATA_TILE_TYPE_SIMPLE)					+ 0] = IndexCountPerInstance;
+		TileIndirectDataBuffer[GetStrataTileTypeDrawIndirectArgOffset_DWord(STRATA_TILE_TYPE_SINGLE)					+ 0] = IndexCountPerInstance;
+		TileIndirectDataBuffer[GetStrataTileTypeDrawIndirectArgOffset_DWord(STRATA_TILE_TYPE_COMPLEX)					+ 0] = IndexCountPerInstance;
+		TileIndirectDataBuffer[GetStrataTileTypeDrawIndirectArgOffset_DWord(STRATA_TILE_TYPE_ROUGH_REFRACT)				+ 0] = IndexCountPerInstance;
+		TileIndirectDataBuffer[GetStrataTileTypeDrawIndirectArgOffset_DWord(STRATA_TILE_TYPE_SSS_WITHOUT_ROUGH_REFRACT) + 0] = IndexCountPerInstance;
 	}
 
 	// TODO: add a SM6 permutation with ballot?
@@ -65,7 +67,7 @@ void TileMainCS(uint2 DispatchThreadId : SV_DispatchThreadID, uint LinearIndex :
 	bool bContainsSingleMaterial  = false;
 	bool bContainsStrataMaterial  = false;
 	bool bContainsOpaqueRoughRefraction = false;
-	bool bHasScreenSpaceSubsurfaceScattering = false;
+	bool bContainsScreenSpaceSubsurfaceScattering = false;
 	FStrataOpaqueRoughRefractionData OpaqueRoughRefractionData = (FStrataOpaqueRoughRefractionData)0;
 	if (all(PixelCoord < uint2(ViewResolution)))
 	{
@@ -78,7 +80,7 @@ void TileMainCS(uint2 DispatchThreadId : SV_DispatchThreadID, uint LinearIndex :
 		bContainsSimpleMaterial  = bIsSimple;
 		bContainsSingleMaterial  = bIsSingle;
 		bContainsComplexMaterial = !bIsSingle & !bIsSimple;
-		bHasScreenSpaceSubsurfaceScattering = HasSubsurface(StrataPixelHeader);
+		bContainsScreenSpaceSubsurfaceScattering = HasSubsurface(StrataPixelHeader);
 
 #if STRATA_OPAQUE_ROUGH_REFRACTION_ENABLED
 		OpaqueRoughRefractionData = StrataUnpackOpaqueRoughRefractionData(OpaqueRoughRefractionTexture[PixelCoord]);
@@ -88,7 +90,7 @@ void TileMainCS(uint2 DispatchThreadId : SV_DispatchThreadID, uint LinearIndex :
 
 #if PERMUTATION_STRATA_CLEAR_DURING_CATEGORIZATION
 	BRANCH
-	if (!bHasScreenSpaceSubsurfaceScattering)
+	if (!bContainsScreenSpaceSubsurfaceScattering)
 	{
 		// We must fill all the pixel which doe not hav subsurface scattering with default so that the SSS code is not executed where it should not.
 		FStrataSubsurfaceData StrataSubsurfaceData = (FStrataSubsurfaceData)0;
@@ -97,19 +99,21 @@ void TileMainCS(uint2 DispatchThreadId : SV_DispatchThreadID, uint LinearIndex :
 #endif
 
 #if PERMUTATION_WAVE_OPS
-	const bool bHasAnyStrata  = WaveActiveAnyTrue(bContainsStrataMaterial);
-//	const bool bHasAnySimple  = WaveActiveAnyTrue(bContainsSimpleMaterial);
-	const bool bHasAnySingle  = WaveActiveAnyTrue(bContainsSingleMaterial);
-	const bool bHasAnyComplex = WaveActiveAnyTrue(bContainsComplexMaterial);
-	const bool bHasAnyOpaqueRoughRefraction = WaveActiveAnyTrue(bContainsOpaqueRoughRefraction);
+	const bool bTileContainsStrata							= WaveActiveAnyTrue(bContainsStrataMaterial);
+//	const bool bTileContainsSimple							= WaveActiveAnyTrue(bContainsSimpleMaterial);
+	const bool bTileContainsSingle							= WaveActiveAnyTrue(bContainsSingleMaterial);
+	const bool bTileContainsComplex							= WaveActiveAnyTrue(bContainsComplexMaterial);
+	const bool bTileContainsOpaqueRoughRefraction			= WaveActiveAnyTrue(bContainsOpaqueRoughRefraction);
+	const bool bTileContainsScreenSpaceSubsurfaceScattering	= WaveActiveAnyTrue(bHasScreenSpaceSubsurfaceScattering);
 #else // PERMUTATION_WAVE_OPS
 
 	s_TileFlags[LinearIndex] = 
-		(bContainsStrataMaterial		? 0x1u	: 0u) | 
-//		(bContainsSimpleMaterial		? 0x2u	: 0u) | 
-		(bContainsSingleMaterial		? 0x4u	: 0u) | 
-		(bContainsComplexMaterial		? 0x8u	: 0u) | 
-		(bContainsOpaqueRoughRefraction ? 0x10u : 0u);
+		(bContainsStrataMaterial					? 0x1u	: 0u) | 
+//		(bContainsSimpleMaterial					? 0x2u	: 0u) | 
+		(bContainsSingleMaterial					? 0x4u	: 0u) | 
+		(bContainsComplexMaterial					? 0x8u	: 0u) | 
+		(bContainsOpaqueRoughRefraction				? 0x10u : 0u) |
+		(bContainsScreenSpaceSubsurfaceScattering	? 0x20u : 0u);
 
 	GroupMemoryBarrierWithGroupSync();
 	if (LinearIndex < 32)
@@ -141,40 +145,49 @@ void TileMainCS(uint2 DispatchThreadId : SV_DispatchThreadID, uint LinearIndex :
 
 	const uint FinalTileFlags = s_TileFlags[LinearIndex] | s_TileFlags[LinearIndex + 1];
 
-	const bool bHasAnyStrata  = (FinalTileFlags & 0x1u) > 0;
-//	const bool bHasAnySimple  = (FinalTileFlags & 0x2u) > 0;
-	const bool bHasAnySingle  = (FinalTileFlags & 0x4u) > 0;
-	const bool bHasAnyComplex = (FinalTileFlags & 0x8u) > 0;
-	const bool bHasAnyOpaqueRoughRefraction = (FinalTileFlags & 0x10u) > 0;
+	const bool bTileContainsStrata							= (FinalTileFlags & 0x1u) > 0;
+//	const bool bTileContainsSimple							= (FinalTileFlags & 0x2u) > 0;
+	const bool bTileContainsSingle							= (FinalTileFlags & 0x4u) > 0;
+	const bool bTileContainsComplex							= (FinalTileFlags & 0x8u) > 0;
+	const bool bTileContainsOpaqueRoughRefraction			= (FinalTileFlags & 0x10u) > 0;
+	const bool bTileContainsScreenSpaceSubsurfaceScattering	= (FinalTileFlags & 0x20u) > 0;
 #endif // PERMUTATION_WAVE_OPS
 
-	if (LinearIndex < 1 && bHasAnyStrata)
+	if (LinearIndex < 1 && bTileContainsStrata)
 	{
-		if (bHasAnyComplex)
+		uint EncodedTile = EncodeTile(GroupId.xy);
+
+		if (bTileContainsComplex)
 		{
 			uint WriteToIndex;
 			InterlockedAdd(TileIndirectDataBuffer[GetStrataTileTypeDrawIndirectArgOffset_DWord(STRATA_TILE_TYPE_COMPLEX) + 1], 1, WriteToIndex);
-			ComplexTileListDataBuffer[WriteToIndex] = EncodeTile(GroupId.xy);
+			ComplexTileListDataBuffer[WriteToIndex] = EncodedTile;
 		}
-		else if (bHasAnySingle)
+		else if (bTileContainsSingle)
 		{
 			uint WriteToIndex;
 			InterlockedAdd(TileIndirectDataBuffer[GetStrataTileTypeDrawIndirectArgOffset_DWord(STRATA_TILE_TYPE_SINGLE) + 1], 1, WriteToIndex);
-			SingleTileListDataBuffer[WriteToIndex] = EncodeTile(GroupId.xy);
+			SingleTileListDataBuffer[WriteToIndex] = EncodedTile;
 		}
-		else // (bHasAnySimple)
+		else // (bTileContainsSimple)
 		{
 			uint WriteToIndex;
 			InterlockedAdd(TileIndirectDataBuffer[GetStrataTileTypeDrawIndirectArgOffset_DWord(STRATA_TILE_TYPE_SIMPLE) + 1], 1, WriteToIndex);
-			SimpleTileListDataBuffer[WriteToIndex] = EncodeTile(GroupId.xy);
+			SimpleTileListDataBuffer[WriteToIndex] = EncodedTile;
 		}
 
 #if STRATA_OPAQUE_ROUGH_REFRACTION_ENABLED
-		if (bHasAnyOpaqueRoughRefraction)
+		if (bTileContainsOpaqueRoughRefraction)
 		{
 			uint WriteToIndex;
 			InterlockedAdd(TileIndirectDataBuffer[GetStrataTileTypeDrawIndirectArgOffset_DWord(STRATA_TILE_TYPE_ROUGH_REFRACT) + 1], 1, WriteToIndex);
-			OpaqueRoughRefractionTileListDataBuffer[WriteToIndex] = EncodeTile(GroupId.xy);
+			OpaqueRoughRefractionTileListDataBuffer[WriteToIndex] = EncodedTile;
+		}
+		if(bTileContainsScreenSpaceSubsurfaceScattering && !bTileContainsOpaqueRoughRefraction)
+		{
+			uint WriteToIndex;
+			InterlockedAdd(TileIndirectDataBuffer[GetStrataTileTypeDrawIndirectArgOffset_DWord(STRATA_TILE_TYPE_SSS_WITHOUT_ROUGH_REFRACT) + 1], 1, WriteToIndex);
+			SSSWithoutOpaqueRoughRefractionTileListDataBuffer[WriteToIndex] = EncodedTile;
 		}
 #endif
 	}
diff --git a/Engine/Shaders/Private/Strata/StrataRoughRefraction.usf b/Engine/Shaders/Private/Strata/StrataRoughRefraction.usf
--- a/Engine/Shaders/Private/Strata/StrataRoughRefraction.usf
+++ b/Engine/Shaders/Private/Strata/StrataRoughRefraction.usf
@@ -8,7 +8,10 @@
 
 #if OPAQUE_ROUGH_REFRACTION_PS
 
+#include "/Engine/Private/Strata/Strata.ush"
+
 Texture2D<float4> SeparatedOpaqueRoughRefractionSceneColor;
+float2 BlurDirection;
 
 void OpaqueRoughRefractionPS(
 	float4 SVPos : SV_POSITION,
@@ -17,9 +20,27 @@ void OpaqueRoughRefractionPS(
 	float2 PixelPos = SVPos.xy;
 	float2 UvBuffer = PixelPos * View.BufferSizeAndInvSize.zw;
 
-	float4 SourceData = SeparatedOpaqueRoughRefractionSceneColor[PixelPos];
+#if PERMUTATION_ENABLE_BLUR && STRATA_OPAQUE_ROUGH_REFRACTION_ENABLED
+	FStrataOpaqueRoughRefractionData ORRData = StrataUnpackOpaqueRoughRefractionData(Strata.OpaqueRoughRefractionTexture[PixelPos]);
+
+	if (ORRData.OpaqueRoughRefractionEnabled)
+	{
+		// BAD placeholder blur
+		float4 SourceData = 0.0f;
+		float SampleCount = 0.0f;
+		for (float u = -4.5f; u <= 4.5f; u += 1.0f)
+		{
+			SourceData += SeparatedOpaqueRoughRefractionSceneColor[PixelPos + BlurDirection * u * ORRData.Variance * 10.0];
+			SampleCount++;
+		}
+		OutColor = SourceData / SampleCount;
+	}
+	else
+#endif
+	{
+		OutColor = SeparatedOpaqueRoughRefractionSceneColor[PixelPos];
+	}
 
-	OutColor = SourceData;
 }
 
 #endif
diff --git a/Engine/Shaders/Shared/StrataDefinitions.h b/Engine/Shaders/Shared/StrataDefinitions.h
--- a/Engine/Shaders/Shared/StrataDefinitions.h
+++ b/Engine/Shaders/Shared/StrataDefinitions.h
@@ -54,11 +54,12 @@
 #define STRATA_OPERATOR_THINFILM		6
 
 // This must map directly to EStrataTileMaterialType
-#define STRATA_TILE_TYPE_SIMPLE			0
-#define STRATA_TILE_TYPE_SINGLE			1
-#define STRATA_TILE_TYPE_COMPLEX		2
-#define STRATA_TILE_TYPE_ROUGH_REFRACT	3
-#define STRATA_TILE_TYPE_COUNT			4
+#define STRATA_TILE_TYPE_SIMPLE						0
+#define STRATA_TILE_TYPE_SINGLE						1
+#define STRATA_TILE_TYPE_COMPLEX					2
+#define STRATA_TILE_TYPE_ROUGH_REFRACT				3
+#define STRATA_TILE_TYPE_SSS_WITHOUT_ROUGH_REFRACT	4
+#define STRATA_TILE_TYPE_COUNT						5
 
 // sizeof(FRHIDrawIndirectParameters) = 4 uints = 16 bytes
 #define GetStrataTileTypeDrawIndirectArgOffset_Byte(x)  (x * 16)
diff --git a/Engine/Source/Runtime/Renderer/Private/Strata/Strata.cpp b/Engine/Source/Runtime/Renderer/Private/Strata/Strata.cpp
--- a/Engine/Source/Runtime/Renderer/Private/Strata/Strata.cpp
+++ b/Engine/Source/Runtime/Renderer/Private/Strata/Strata.cpp
@@ -96,10 +96,11 @@ const TCHAR* ToString(EStrataTileType Type)
 {
 	switch (Type)
 	{
-	case EStrataTileType::ESimple:					return TEXT("Simple");
-	case EStrataTileType::ESingle:					return TEXT("Single");
-	case EStrataTileType::EComplex:					return TEXT("Complex");
-	case EStrataTileType::EOpaqueRoughRefraction:	return TEXT("OpaqueRoughRefraction");
+	case EStrataTileType::ESimple:							return TEXT("Simple");
+	case EStrataTileType::ESingle:							return TEXT("Single");
+	case EStrataTileType::EComplex:							return TEXT("Complex");
+	case EStrataTileType::EOpaqueRoughRefraction:			return TEXT("OpaqueRoughRefraction");
+	case EStrataTileType::ESSSWithoutOpaqueRoughRefraction:	return TEXT("SSSWithoutOpaqueRoughRefraction");
 	}
 	return TEXT("Unknown");
 }
@@ -164,7 +165,8 @@ void InitialiseStrataFrameSceneData(FSceneRenderer& SceneRenderer, FRDGBuilder&
 			TEXT("Strata.StrataTileListBuffer(Simple)"),
 			TEXT("Strata.StrataTileListBuffer(Single)"),
 			TEXT("Strata.StrataTileListBuffer(Complex)"),
-			TEXT("Strata.StrataTileListBuffer(OpaqueRoughRefraction)")
+			TEXT("Strata.StrataTileListBuffer(OpaqueRoughRefraction)"),
+			TEXT("Strata.StrataTileListBuffer(SSSWithoutOpaqueRoughRefraction)")
 		};
 
 		// Tile classification buffers
@@ -202,6 +204,10 @@ void InitialiseStrataFrameSceneData(FSceneRenderer& SceneRenderer, FRDGBuilder&
 			StrataSceneData.ClassificationTileListBufferSRV[EStrataTileType::EOpaqueRoughRefraction] = GraphBuilder.CreateSRV(StrataSceneData.ClassificationTileListBuffer[EStrataTileType::EOpaqueRoughRefraction], PF_R32_UINT);
 			StrataSceneData.ClassificationTileListBufferUAV[EStrataTileType::EOpaqueRoughRefraction] = GraphBuilder.CreateUAV(StrataSceneData.ClassificationTileListBuffer[EStrataTileType::EOpaqueRoughRefraction], PF_R32_UINT);
 
+			StrataSceneData.ClassificationTileListBuffer[EStrataTileType::ESSSWithoutOpaqueRoughRefraction] = GraphBuilder.CreateBuffer(FRDGBufferDesc::CreateBufferDesc(sizeof(uint32), TileListBufferElementCount), StrataTileListBufferNames[EStrataTileType::ESSSWithoutOpaqueRoughRefraction]);
+			StrataSceneData.ClassificationTileListBufferSRV[EStrataTileType::ESSSWithoutOpaqueRoughRefraction] = GraphBuilder.CreateSRV(StrataSceneData.ClassificationTileListBuffer[EStrataTileType::ESSSWithoutOpaqueRoughRefraction], PF_R32_UINT);
+			StrataSceneData.ClassificationTileListBufferUAV[EStrataTileType::ESSSWithoutOpaqueRoughRefraction] = GraphBuilder.CreateUAV(StrataSceneData.ClassificationTileListBuffer[EStrataTileType::ESSSWithoutOpaqueRoughRefraction], PF_R32_UINT);
+
 			StrataSceneData.OpaqueRoughRefractionTexture = GraphBuilder.CreateTexture(
 				FRDGTextureDesc::Create2D(OpaqueRoughRefractionSceneExtent, PF_R8, FClearValueBinding::Black, TexCreate_ShaderResource | TexCreate_UAV | TexCreate_RenderTargetable), TEXT("Strata.OpaqueRoughRefractionTexture"));
 			StrataSceneData.OpaqueRoughRefractionTextureUAV = GraphBuilder.CreateUAV(StrataSceneData.OpaqueRoughRefractionTexture);
@@ -471,6 +477,7 @@ class FStrataMaterialTileClassificationPassCS : public FGlobalShader
 		SHADER_PARAMETER_RDG_BUFFER_UAV(RWBuffer, SingleTileListDataBuffer)
 		SHADER_PARAMETER_RDG_BUFFER_UAV(RWBuffer, ComplexTileListDataBuffer)
 		SHADER_PARAMETER_RDG_BUFFER_UAV(RWBuffer, OpaqueRoughRefractionTileListDataBuffer)
+		SHADER_PARAMETER_RDG_BUFFER_UAV(RWBuffer, SSSWithoutOpaqueRoughRefractionTileListDataBuffer)
 		SHADER_PARAMETER_RDG_TEXTURE_UAV(RWTexture2D<uint2>, SSSTextureUAV)
 		SHADER_PARAMETER_RDG_TEXTURE(Texture2D<float>, OpaqueRoughRefractionTexture)
 	END_SHADER_PARAMETER_STRUCT()
@@ -633,10 +640,11 @@ static void AddStrataInternalClassificationTilePass(
 		ParametersPS->RenderTargets[0] = FRenderTargetBinding(*ColorTexture, ERenderTargetLoadAction::ELoad);
 		switch (TileMaterialType)
 		{
-		case EStrataTileType::ESimple:					ParametersPS->DebugTileColor = FVector4f(0.0f, 1.0f, 0.0f, 1.0); break;
-		case EStrataTileType::ESingle:					ParametersPS->DebugTileColor = FVector4f(1.0f, 1.0f, 0.0f, 1.0); break;
-		case EStrataTileType::EComplex:					ParametersPS->DebugTileColor = FVector4f(1.0f, 0.0f, 0.0f, 1.0); break;
-		case EStrataTileType::EOpaqueRoughRefraction:	ParametersPS->DebugTileColor = FVector4f(0.0f, 0.0f, 1.0f, 1.0); break;
+		case EStrataTileType::ESimple:							ParametersPS->DebugTileColor = FVector4f(0.0f, 1.0f, 0.0f, 1.0); break;
+		case EStrataTileType::ESingle:							ParametersPS->DebugTileColor = FVector4f(1.0f, 1.0f, 0.0f, 1.0); break;
+		case EStrataTileType::EComplex:							ParametersPS->DebugTileColor = FVector4f(1.0f, 0.0f, 0.0f, 1.0); break;
+		case EStrataTileType::EOpaqueRoughRefraction:			ParametersPS->DebugTileColor = FVector4f(0.0f, 1.0f, 1.0f, 1.0); break;
+		case EStrataTileType::ESSSWithoutOpaqueRoughRefraction:	ParametersPS->DebugTileColor = FVector4f(0.0f, 0.0f, 1.0f, 1.0); break;
 		default: check(false);
 		}
 	}
@@ -670,7 +678,7 @@ static void AddStrataInternalClassificationTilePass(
 			}
 			else
 			{
-				check(TileMaterialType != EStrataTileType::ECount && TileMaterialType != EStrataTileType::EOpaqueRoughRefraction);
+				check(TileMaterialType != EStrataTileType::ECount && TileMaterialType != EStrataTileType::EOpaqueRoughRefraction && TileMaterialType != EStrataTileType::ESSSWithoutOpaqueRoughRefraction);
 
 				// No blending and no pixel shader required. Stencil will be written to.
 				GraphicsPSOInit.BoundShaderState.PixelShaderRHI = nullptr;
@@ -825,6 +833,7 @@ void AddStrataMaterialClassificationPass(FRDGBuilder& GraphBuilder, const FMinim
 			PassParameters->SingleTileListDataBuffer = View.StrataSceneData->ClassificationTileListBufferUAV[EStrataTileType::ESingle];
 			PassParameters->ComplexTileListDataBuffer = View.StrataSceneData->ClassificationTileListBufferUAV[EStrataTileType::EComplex];
 			PassParameters->OpaqueRoughRefractionTileListDataBuffer = View.StrataSceneData->ClassificationTileListBufferUAV[EStrataTileType::EOpaqueRoughRefraction];
+			PassParameters->SSSWithoutOpaqueRoughRefractionTileListDataBuffer = View.StrataSceneData->ClassificationTileListBufferUAV[EStrataTileType::ESSSWithoutOpaqueRoughRefraction];
 
 			const uint32 GroupSize = 8;
 			FComputeShaderUtils::AddPass(
@@ -894,6 +903,7 @@ FScreenPassTexture AddStrataDebugPasses(FRDGBuilder& GraphBuilder, const FViewIn
 		if (StrataClassificationDebug > 1)
 		{
 			AddStrataInternalClassificationTilePass(GraphBuilder, View, nullptr, &ScreenPassSceneColor.Texture, EStrataTileType::EOpaqueRoughRefraction, bDebugPass);
+			AddStrataInternalClassificationTilePass(GraphBuilder, View, nullptr, &ScreenPassSceneColor.Texture, EStrataTileType::ESSSWithoutOpaqueRoughRefraction, bDebugPass);
 		}
 		else
 		{
diff --git a/Engine/Source/Runtime/Renderer/Private/Strata/Strata.h b/Engine/Source/Runtime/Renderer/Private/Strata/Strata.h
--- a/Engine/Source/Runtime/Renderer/Private/Strata/Strata.h
+++ b/Engine/Source/Runtime/Renderer/Private/Strata/Strata.h
@@ -51,6 +51,7 @@ enum EStrataTileType : uint32
 	ESingle = STRATA_TILE_TYPE_SINGLE,
 	EComplex = STRATA_TILE_TYPE_COMPLEX,
 	EOpaqueRoughRefraction = STRATA_TILE_TYPE_ROUGH_REFRACT,
+	ESSSWithoutOpaqueRoughRefraction = STRATA_TILE_TYPE_SSS_WITHOUT_ROUGH_REFRACT,
 	ECount
 };
 
diff --git a/Engine/Source/Runtime/Renderer/Private/Strata/StrataRoughRefraction.cpp b/Engine/Source/Runtime/Renderer/Private/Strata/StrataRoughRefraction.cpp
--- a/Engine/Source/Runtime/Renderer/Private/Strata/StrataRoughRefraction.cpp
+++ b/Engine/Source/Runtime/Renderer/Private/Strata/StrataRoughRefraction.cpp
@@ -33,11 +33,15 @@ class FOpaqueRoughRefractionPS : public FGlobalShader
 	DECLARE_GLOBAL_SHADER(FOpaqueRoughRefractionPS);
 	SHADER_USE_PARAMETER_STRUCT(FOpaqueRoughRefractionPS, FGlobalShader);
 
-	using FPermutationDomain = TShaderPermutationDomain<>;
+	class FEnableBlur : SHADER_PERMUTATION_BOOL("PERMUTATION_ENABLE_BLUR");
+	using FPermutationDomain = TShaderPermutationDomain<FEnableBlur>;
 
 	BEGIN_SHADER_PARAMETER_STRUCT(FParameters, )
 		SHADER_PARAMETER_STRUCT_REF(FViewUniformShaderParameters, ViewUniformBuffer)
+		SHADER_PARAMETER_RDG_UNIFORM_BUFFER(FStrataGlobalUniformParameters, Strata)
+		SHADER_PARAMETER_STRUCT_INCLUDE(Strata::FStrataTilePassVS::FParameters, StrataTile)
 		SHADER_PARAMETER_RDG_TEXTURE(Texture2D, SeparatedOpaqueRoughRefractionSceneColor)
+		SHADER_PARAMETER(FVector2f, BlurDirection)
 		RENDER_TARGET_BINDING_SLOTS()
 	END_SHADER_PARAMETER_STRUCT();
 
@@ -63,35 +67,188 @@ void AddStrataOpaqueRoughRefractionPasses(
 	FSceneTextures& SceneTextures,
 	TArrayView<const FViewInfo> Views)
 {
-	if (!IsStrataOpaqueMaterialRoughRefractionEnabled())
+	const uint32 ViewCount = Views.Num();
+	const bool bOpaqueRoughRefractionEnabled = IsStrataOpaqueMaterialRoughRefractionEnabled() && ViewCount > 0;
+	RDG_EVENT_SCOPE_CONDITIONAL(GraphBuilder, bOpaqueRoughRefractionEnabled, "StrataOpaqueRoughRefraction");
+	if (!bOpaqueRoughRefractionEnabled)
 	{
 		return;
 	}
 
-	const uint32 ViewCount = Views.Num();
+	FRDGTextureRef SceneColorTexture = SceneTextures.Color.Target;
+	FRDGTextureRef TempTexture = nullptr;
+	ERenderTargetLoadAction LoadAction = ERenderTargetLoadAction::EClear;
+
+	//
+	// 1. First blur pass into temporary buffer. This only blurs tiles containing pixels with rough refractions.
+	//
+	EStrataTileType StrataTileType = EStrataTileType::EOpaqueRoughRefraction;
 	for (uint32 ViewIndex = 0; ViewIndex < ViewCount; ++ViewIndex)
 	{
 		const FViewInfo& View = Views[ViewIndex];
 
-		FRDGTextureRef SeparatedOpaqueRoughRefractionSceneColor = View.StrataSceneData->SeparatedOpaqueRoughRefractionSceneColor;
-		FRDGTextureRef SceneColorTexture = SceneTextures.Color.Target;
+		if (TempTexture == nullptr)
+		{
+			TempTexture = GraphBuilder.CreateTexture(FRDGTextureDesc::Create2D(GetSceneTextureExtent(), PF_FloatRGBA, FClearValueBinding::Black, TexCreate_UAV | TexCreate_ShaderResource | TexCreate_RenderTargetable), TEXT("Strata.RoughRefrac.TempTexture"));
+		}
 
+		FRDGTextureRef SeparatedOpaqueRoughRefractionSceneColor = View.StrataSceneData->SeparatedOpaqueRoughRefractionSceneColor;
 
 		FSceneTextureParameters SceneTextureParameters = GetSceneTextureParameters(GraphBuilder);
 		FOpaqueRoughRefractionPS::FParameters* PassParameters = GraphBuilder.AllocParameters<FOpaqueRoughRefractionPS::FParameters>();
 		PassParameters->ViewUniformBuffer = View.ViewUniformBuffer;
+		PassParameters->Strata = Strata::BindStrataGlobalUniformParameters(View.StrataSceneData);
 		PassParameters->SeparatedOpaqueRoughRefractionSceneColor = SeparatedOpaqueRoughRefractionSceneColor;
+		PassParameters->RenderTargets[0] = FRenderTargetBinding(TempTexture, LoadAction);
+		PassParameters->BlurDirection = FVector2f(1.0f, 0.0f);
+
+		FOpaqueRoughRefractionPS::FPermutationDomain PermutationVector;
+		PermutationVector.Set<FOpaqueRoughRefractionPS::FEnableBlur>(true);
+		TShaderMapRef<FOpaqueRoughRefractionPS> PixelShader(View.ShaderMap, PermutationVector);
+
+		Strata::FStrataTilePassVS::FPermutationDomain VSPermutationVector;
+		VSPermutationVector.Set< Strata::FStrataTilePassVS::FEnableDebug >(false);
+		VSPermutationVector.Set< Strata::FStrataTilePassVS::FEnableTexCoordScreenVector >(false);
+		TShaderMapRef<Strata::FStrataTilePassVS> TileVertexShader(View.ShaderMap, VSPermutationVector);
+
+		EPrimitiveType PrimitiveType = PT_TriangleList;
+		PassParameters->StrataTile = Strata::SetTileParameters(GraphBuilder, View, StrataTileType, PrimitiveType);
+
+		GraphBuilder.AddPass(
+			RDG_EVENT_NAME("OpaqueRoughRefraction - Pass0"),
+			PassParameters,
+			ERDGPassFlags::Raster,
+			[&View, TileVertexShader, PixelShader, PassParameters, StrataTileType, PrimitiveType](FRHICommandList& RHICmdList)
+			{
+				FGraphicsPipelineStateInitializer GraphicsPSOInit;
+				RHICmdList.ApplyCachedRenderTargets(GraphicsPSOInit);
+				// Set the device viewport for the view.
+				RHICmdList.SetViewport(View.ViewRect.Min.X, View.ViewRect.Min.Y, 0.0f, View.ViewRect.Max.X, View.ViewRect.Max.Y, 1.0f);
+
+				GraphicsPSOInit.BlendState = TStaticBlendState<>::GetRHI(); // no blend
+				GraphicsPSOInit.PrimitiveType = PrimitiveType;
+				GraphicsPSOInit.bDepthBounds = false;
+				GraphicsPSOInit.RasterizerState = TStaticRasterizerState<FM_Solid, CM_None>::GetRHI();
+				GraphicsPSOInit.BoundShaderState.VertexDeclarationRHI = GFilterVertexDeclaration.VertexDeclarationRHI;
+				GraphicsPSOInit.BoundShaderState.VertexShaderRHI = TileVertexShader.GetVertexShader();
+				GraphicsPSOInit.BoundShaderState.PixelShaderRHI = PixelShader.GetPixelShader();
+				GraphicsPSOInit.DepthStencilState = TStaticDepthStencilState<false, CF_Always>::GetRHI();
+				SetGraphicsPipelineState(RHICmdList, GraphicsPSOInit, 0x0);
+				SetShaderParameters(RHICmdList, PixelShader, PixelShader.GetPixelShader(), *PassParameters);
+				SetShaderParameters(RHICmdList, TileVertexShader, TileVertexShader.GetVertexShader(), PassParameters->StrataTile);
+				RHICmdList.DrawPrimitiveIndirect(PassParameters->StrataTile.TileIndirectBuffer->GetIndirectRHICallBuffer(), Strata::TileTypeDrawIndirectArgOffset(StrataTileType));
+			});
+
+		LoadAction = ERenderTargetLoadAction::ELoad;
+	}
+
+	//
+	// 2. second blur pass into temporary buffer. This only blurs tiles containing pixels with rough refractions.
+	//
+	for (uint32 ViewIndex = 0; ViewIndex < ViewCount; ++ViewIndex)
+	{
+		const FViewInfo& View = Views[ViewIndex];
+
+		FSceneTextureParameters SceneTextureParameters = GetSceneTextureParameters(GraphBuilder);
+		FOpaqueRoughRefractionPS::FParameters* PassParameters = GraphBuilder.AllocParameters<FOpaqueRoughRefractionPS::FParameters>();
+		PassParameters->ViewUniformBuffer = View.ViewUniformBuffer;
+		PassParameters->Strata = Strata::BindStrataGlobalUniformParameters(View.StrataSceneData);
+		PassParameters->SeparatedOpaqueRoughRefractionSceneColor = TempTexture;
 		PassParameters->RenderTargets[0] = FRenderTargetBinding(SceneColorTexture, ERenderTargetLoadAction::ELoad);
+		PassParameters->BlurDirection = FVector2f(0.0f, 1.0f);
 
 		FOpaqueRoughRefractionPS::FPermutationDomain PermutationVector;
+		PermutationVector.Set<FOpaqueRoughRefractionPS::FEnableBlur>(true);
 		TShaderMapRef<FOpaqueRoughRefractionPS> PixelShader(View.ShaderMap, PermutationVector);
 
-		FRHIBlendState* PreMultipliedColorTransmittanceBlend = TStaticBlendState<CW_RGB, BO_Add, BF_One, BF_One, BO_Add, BF_One, BF_One>::GetRHI();
+		Strata::FStrataTilePassVS::FPermutationDomain VSPermutationVector;
+		VSPermutationVector.Set< Strata::FStrataTilePassVS::FEnableDebug >(false);
+		VSPermutationVector.Set< Strata::FStrataTilePassVS::FEnableTexCoordScreenVector >(false);
+		TShaderMapRef<Strata::FStrataTilePassVS> TileVertexShader(View.ShaderMap, VSPermutationVector);
+
+		EPrimitiveType PrimitiveType = PT_TriangleList;
+		PassParameters->StrataTile = Strata::SetTileParameters(GraphBuilder, View, StrataTileType, PrimitiveType);
+
+		GraphBuilder.AddPass(
+			RDG_EVENT_NAME("OpaqueRoughRefraction - Pass2)"),
+			PassParameters,
+			ERDGPassFlags::Raster,
+			[&View, TileVertexShader, PixelShader, PassParameters, StrataTileType, PrimitiveType](FRHICommandList& RHICmdList)
+			{
+				FGraphicsPipelineStateInitializer GraphicsPSOInit;
+				RHICmdList.ApplyCachedRenderTargets(GraphicsPSOInit);
+				// Set the device viewport for the view.
+				RHICmdList.SetViewport(View.ViewRect.Min.X, View.ViewRect.Min.Y, 0.0f, View.ViewRect.Max.X, View.ViewRect.Max.Y, 1.0f);
+
+				GraphicsPSOInit.BlendState = TStaticBlendState<CW_RGB, BO_Add, BF_One, BF_One, BO_Add, BF_One, BF_One>::GetRHI(); // Additive blend
+				GraphicsPSOInit.PrimitiveType = PrimitiveType;
+				GraphicsPSOInit.bDepthBounds = false;
+				GraphicsPSOInit.RasterizerState = TStaticRasterizerState<FM_Solid, CM_None>::GetRHI();
+				GraphicsPSOInit.BoundShaderState.VertexDeclarationRHI = GFilterVertexDeclaration.VertexDeclarationRHI;
+				GraphicsPSOInit.BoundShaderState.VertexShaderRHI = TileVertexShader.GetVertexShader();
+				GraphicsPSOInit.BoundShaderState.PixelShaderRHI = PixelShader.GetPixelShader();
+				GraphicsPSOInit.DepthStencilState = TStaticDepthStencilState<false, CF_Always>::GetRHI();
+				SetGraphicsPipelineState(RHICmdList, GraphicsPSOInit, 0x0);
+				SetShaderParameters(RHICmdList, PixelShader, PixelShader.GetPixelShader(), *PassParameters);
+				SetShaderParameters(RHICmdList, TileVertexShader, TileVertexShader.GetVertexShader(), PassParameters->StrataTile);
+				RHICmdList.DrawPrimitiveIndirect(PassParameters->StrataTile.TileIndirectBuffer->GetIndirectRHICallBuffer(), Strata::TileTypeDrawIndirectArgOffset(StrataTileType));
+			});
+	}
 
-		FPixelShaderUtils::AddFullscreenPass<FOpaqueRoughRefractionPS>(GraphBuilder, View.ShaderMap, RDG_EVENT_NAME("Strata::VisualizeRoughRefractionPS"),
-			PixelShader, PassParameters, View.ViewRect, PreMultipliedColorTransmittanceBlend);
+	//
+	// 3. Add remaining tiles with subsurface scattering that did not have rough refractions on them, resulting in a complete scene color texture.
+	//
+	StrataTileType = EStrataTileType::ESSSWithoutOpaqueRoughRefraction;
+	for (uint32 ViewIndex = 0; ViewIndex < ViewCount; ++ViewIndex)
+	{
+		const FViewInfo& View = Views[ViewIndex];
 
-		
+		FRDGTextureRef SeparatedOpaqueRoughRefractionSceneColor = View.StrataSceneData->SeparatedOpaqueRoughRefractionSceneColor;
+
+		FSceneTextureParameters SceneTextureParameters = GetSceneTextureParameters(GraphBuilder);
+		FOpaqueRoughRefractionPS::FParameters* PassParameters = GraphBuilder.AllocParameters<FOpaqueRoughRefractionPS::FParameters>();
+		PassParameters->ViewUniformBuffer = View.ViewUniformBuffer;
+		PassParameters->Strata = Strata::BindStrataGlobalUniformParameters(View.StrataSceneData);
+		PassParameters->SeparatedOpaqueRoughRefractionSceneColor = SeparatedOpaqueRoughRefractionSceneColor;
+		PassParameters->RenderTargets[0] = FRenderTargetBinding(SceneColorTexture, ERenderTargetLoadAction::ELoad);
+		PassParameters->BlurDirection = FVector2f(0.0f, 0.0f);
+
+		FOpaqueRoughRefractionPS::FPermutationDomain PermutationVector;
+		PermutationVector.Set<FOpaqueRoughRefractionPS::FEnableBlur>(false);
+		TShaderMapRef<FOpaqueRoughRefractionPS> PixelShader(View.ShaderMap, PermutationVector);
+
+		Strata::FStrataTilePassVS::FPermutationDomain VSPermutationVector;
+		VSPermutationVector.Set< Strata::FStrataTilePassVS::FEnableDebug >(false);
+		VSPermutationVector.Set< Strata::FStrataTilePassVS::FEnableTexCoordScreenVector >(false);
+		TShaderMapRef<Strata::FStrataTilePassVS> TileVertexShader(View.ShaderMap, VSPermutationVector);
+
+		EPrimitiveType PrimitiveType = PT_TriangleList;
+		PassParameters->StrataTile = Strata::SetTileParameters(GraphBuilder, View, StrataTileType, PrimitiveType);
+
+		GraphBuilder.AddPass(
+			RDG_EVENT_NAME("OpaqueRoughRefraction - Add SSSOnly tiles"),
+			PassParameters,
+			ERDGPassFlags::Raster,
+			[&View, TileVertexShader, PixelShader, PassParameters, StrataTileType, PrimitiveType](FRHICommandList& RHICmdList)
+			{
+				FGraphicsPipelineStateInitializer GraphicsPSOInit;
+				RHICmdList.ApplyCachedRenderTargets(GraphicsPSOInit);
+				// Set the device viewport for the view.
+				RHICmdList.SetViewport(View.ViewRect.Min.X, View.ViewRect.Min.Y, 0.0f, View.ViewRect.Max.X, View.ViewRect.Max.Y, 1.0f);
+
+				GraphicsPSOInit.BlendState = TStaticBlendState<CW_RGB, BO_Add, BF_One, BF_One, BO_Add, BF_One, BF_One>::GetRHI(); // Additive blend
+				GraphicsPSOInit.PrimitiveType = PrimitiveType;
+				GraphicsPSOInit.bDepthBounds = false;
+				GraphicsPSOInit.RasterizerState = TStaticRasterizerState<FM_Solid, CM_None>::GetRHI();
+				GraphicsPSOInit.BoundShaderState.VertexDeclarationRHI = GFilterVertexDeclaration.VertexDeclarationRHI;
+				GraphicsPSOInit.BoundShaderState.VertexShaderRHI = TileVertexShader.GetVertexShader();
+				GraphicsPSOInit.BoundShaderState.PixelShaderRHI = PixelShader.GetPixelShader();
+				GraphicsPSOInit.DepthStencilState = TStaticDepthStencilState<false, CF_Always>::GetRHI();
+				SetGraphicsPipelineState(RHICmdList, GraphicsPSOInit, 0x0);
+				SetShaderParameters(RHICmdList, PixelShader, PixelShader.GetPixelShader(), *PassParameters);
+				SetShaderParameters(RHICmdList, TileVertexShader, TileVertexShader.GetVertexShader(), PassParameters->StrataTile);
+				RHICmdList.DrawPrimitiveIndirect(PassParameters->StrataTile.TileIndirectBuffer->GetIndirectRHICallBuffer(), Strata::TileTypeDrawIndirectArgOffset(StrataTileType));
+			});
 	}
 }
 
