Strata events follows naming convention.

diff --git a/Engine/Source/Runtime/Renderer/Private/Strata/Strata.cpp b/Engine/Source/Runtime/Renderer/Private/Strata/Strata.cpp
--- a/Engine/Source/Runtime/Renderer/Private/Strata/Strata.cpp
+++ b/Engine/Source/Runtime/Renderer/Private/Strata/Strata.cpp
@@ -797,7 +797,7 @@ void SetBasePassRenderTargetOutputFormat(const EShaderPlatform Platform, FShader
 
 void AddStrataMaterialClassificationPass(FRDGBuilder& GraphBuilder, const FMinimalSceneTextures& SceneTextures, const TArray<FViewInfo>& Views)
 {
-	RDG_EVENT_SCOPE_CONDITIONAL(GraphBuilder, IsStrataEnabled() && Views.Num() > 0, "StrataMaterialClassification");
+	RDG_EVENT_SCOPE_CONDITIONAL(GraphBuilder, IsStrataEnabled() && Views.Num() > 0, "Strata::MaterialClassification");
 	if (!IsStrataEnabled())
 	{
 		return;
diff --git a/Engine/Source/Runtime/Renderer/Private/Strata/StrataRoughRefraction.cpp b/Engine/Source/Runtime/Renderer/Private/Strata/StrataRoughRefraction.cpp
--- a/Engine/Source/Runtime/Renderer/Private/Strata/StrataRoughRefraction.cpp
+++ b/Engine/Source/Runtime/Renderer/Private/Strata/StrataRoughRefraction.cpp
@@ -69,7 +69,7 @@ void AddStrataOpaqueRoughRefractionPasses(
 {
 	const uint32 ViewCount = Views.Num();
 	const bool bOpaqueRoughRefractionEnabled = IsStrataOpaqueMaterialRoughRefractionEnabled() && ViewCount > 0;
-	RDG_EVENT_SCOPE_CONDITIONAL(GraphBuilder, bOpaqueRoughRefractionEnabled, "StrataOpaqueRoughRefraction");
+	RDG_EVENT_SCOPE_CONDITIONAL(GraphBuilder, bOpaqueRoughRefractionEnabled, "Strata::OpaqueRoughRefraction");
 	if (!bOpaqueRoughRefractionEnabled)
 	{
 		return;
