Fix startup crash in QAGame with ISR (UE-145939).

- With ISR material shaders expect the instanced view UB to be bound.

diff --git a/Engine/Source/Runtime/Renderer/Private/ReflectionEnvironmentRealTimeCapture.cpp b/Engine/Source/Runtime/Renderer/Private/ReflectionEnvironmentRealTimeCapture.cpp
--- a/Engine/Source/Runtime/Renderer/Private/ReflectionEnvironmentRealTimeCapture.cpp
+++ b/Engine/Source/Runtime/Renderer/Private/ReflectionEnvironmentRealTimeCapture.cpp
@@ -541,14 +541,13 @@ void FScene::AllocateAndCaptureFrameSkyEnvMap(
 					CubeView.CachedViewUniformShaderParameters->CameraAerialPerspectiveVolume = GSystemTextures.VolumetricBlackDummy->GetRenderTargetItem().ShaderResourceTexture;
 				}
 
-				TUniformBufferRef<FViewUniformShaderParameters> CubeViewUniformBuffer = TUniformBufferRef<FViewUniformShaderParameters>::CreateUniformBufferImmediate(*CubeView.CachedViewUniformShaderParameters, UniformBuffer_SingleFrame);
-				CubeView.ViewUniformBuffer = CubeViewUniformBuffer;
+				CubeView.CreateViewUniformBuffers(*CubeView.CachedViewUniformShaderParameters);
 				if (CubeView.bSceneHasSkyMaterial)
 				{
 					GPUScene.UploadDynamicPrimitiveShaderDataForView(GraphBuilder, this, CubeView);
 				}
 
-				SkyRC.ViewUniformBuffer = CubeViewUniformBuffer;
+				SkyRC.ViewUniformBuffer = CubeView.ViewUniformBuffer;
 				SkyRC.ViewMatrices = &CubeViewMatrices;
 
 				SkyRC.SkyAtmosphereViewLutTexture = BlackDummy2dTex;
@@ -595,7 +594,7 @@ void FScene::AllocateAndCaptureFrameSkyEnvMap(
 							}
 
 							AddSimpleMeshPass(GraphBuilder, PassParameters, Scene, MainView, &InstanceCullingManager, RDG_EVENT_NAME("CaptureSkyMeshReflection"), SkyRC.Viewport,
-							[&MainView, &CubeViewUniformBuffer, bUseDepthBuffer, Scene](FDynamicPassMeshDrawListContext* DynamicMeshPassContext)
+							[&MainView, CubeViewUniformBuffer = CubeView.ViewUniformBuffer, bUseDepthBuffer, Scene](FDynamicPassMeshDrawListContext* DynamicMeshPassContext)
 							{
 								FMeshPassProcessorRenderState DrawRenderState;
 
@@ -657,7 +656,7 @@ void FScene::AllocateAndCaptureFrameSkyEnvMap(
 							TShaderMapRef<FRenderRealTimeReflectionHeightFogPS> PixelShader(GetGlobalShaderMap(SkyRC.FeatureLevel), PsPermutationVector);
 
 							FRenderRealTimeReflectionHeightFogPS::FParameters* PsPassParameters = GraphBuilder.AllocParameters<FRenderRealTimeReflectionHeightFogPS::FParameters>();
-							PsPassParameters->ViewUniformBuffer = CubeViewUniformBuffer;
+							PsPassParameters->ViewUniformBuffer = CubeView.ViewUniformBuffer;
 							PsPassParameters->RenderTargets = SkyRC.RenderTargets;
 							PsPassParameters->DepthTexture = CubeDepthTexture != nullptr ? CubeDepthTexture : BlackDummy2dTex;
 							PsPassParameters->FogStruct = CreateFogUniformBuffer(GraphBuilder, CubeView);
@@ -703,7 +702,7 @@ void FScene::AllocateAndCaptureFrameSkyEnvMap(
 
 				if (bShouldRenderVolumetricCloud && bExecuteCloud)
 				{
-					CloudRC.ViewUniformBuffer = CubeViewUniformBuffer;
+					CloudRC.ViewUniformBuffer = CubeView.ViewUniformBuffer;
 
 					CloudRC.RenderTargets[0] = SkyRC.RenderTargets[0];
 					//	CloudRC.RenderTargets[1] = Null target will skip export
diff --git a/Engine/Source/Runtime/Renderer/Private/SceneRendering.cpp b/Engine/Source/Runtime/Renderer/Private/SceneRendering.cpp
--- a/Engine/Source/Runtime/Renderer/Private/SceneRendering.cpp
+++ b/Engine/Source/Runtime/Renderer/Private/SceneRendering.cpp
@@ -1998,7 +1998,21 @@ void FViewInfo::InitRHIResources(uint32 OverrideNumMSAASamples)
 		CachedViewUniformShaderParameters->NumSceneColorMSAASamples = OverrideNumMSAASamples;
 	}
 
-	ViewUniformBuffer = TUniformBufferRef<FViewUniformShaderParameters>::CreateUniformBufferImmediate(*CachedViewUniformShaderParameters, UniformBuffer_SingleFrame);
+	CreateViewUniformBuffers(*CachedViewUniformShaderParameters);
+
+	const int32 TranslucencyLightingVolumeDim = GetTranslucencyLightingVolumeDim();
+
+	for (int32 CascadeIndex = 0; CascadeIndex < TVC_MAX; CascadeIndex++)
+	{
+		TranslucencyLightingVolumeMin[CascadeIndex] = VolumeBounds[CascadeIndex].Min;
+		TranslucencyVolumeVoxelSize[CascadeIndex] = (VolumeBounds[CascadeIndex].Max.X - VolumeBounds[CascadeIndex].Min.X) / TranslucencyLightingVolumeDim;
+		TranslucencyLightingVolumeSize[CascadeIndex] = VolumeBounds[CascadeIndex].Max - VolumeBounds[CascadeIndex].Min;
+	}
+}
+
+void FViewInfo::CreateViewUniformBuffers(const FViewUniformShaderParameters& Params)
+{
+	ViewUniformBuffer = TUniformBufferRef<FViewUniformShaderParameters>::CreateUniformBufferImmediate(Params, UniformBuffer_SingleFrame);
 	if (bShouldBindInstancedViewUB)
 	{
 		if (const FViewInfo* InstancedView = GetInstancedView())
@@ -2012,19 +2026,10 @@ void FViewInfo::InitRHIResources(uint32 OverrideNumMSAASamples)
 		{
 			// If we don't render this view in stereo, we simply initialize with the existing contents.
 			InstancedViewUniformBuffer = TUniformBufferRef<FInstancedViewUniformShaderParameters>::CreateUniformBufferImmediate(
-				reinterpret_cast<FInstancedViewUniformShaderParameters&>(*CachedViewUniformShaderParameters),
+				reinterpret_cast<const FInstancedViewUniformShaderParameters&>(Params),
 				UniformBuffer_SingleFrame);
 		}
 	}
-
-	const int32 TranslucencyLightingVolumeDim = GetTranslucencyLightingVolumeDim();
-
-	for (int32 CascadeIndex = 0; CascadeIndex < TVC_MAX; CascadeIndex++)
-	{
-		TranslucencyLightingVolumeMin[CascadeIndex] = VolumeBounds[CascadeIndex].Min;
-		TranslucencyVolumeVoxelSize[CascadeIndex] = (VolumeBounds[CascadeIndex].Max.X - VolumeBounds[CascadeIndex].Min.X) / TranslucencyLightingVolumeDim;
-		TranslucencyLightingVolumeSize[CascadeIndex] = VolumeBounds[CascadeIndex].Max - VolumeBounds[CascadeIndex].Min;
-	}
 }
 
 extern TSet<IPersistentViewUniformBufferExtension*> PersistentViewUniformBufferExtensions;
diff --git a/Engine/Source/Runtime/Renderer/Private/SceneRendering.h b/Engine/Source/Runtime/Renderer/Private/SceneRendering.h
--- a/Engine/Source/Runtime/Renderer/Private/SceneRendering.h
+++ b/Engine/Source/Runtime/Renderer/Private/SceneRendering.h
@@ -1565,6 +1565,9 @@ public:
 	/** Initializes the RHI resources used by this view. OverrideNumMSAASamples can optionally override NumSceneColorMSAASamples (needed for editor views) */
 	void InitRHIResources(uint32 OverrideNumMSAASamples = 0);
 
+	/** Creates both ViewUniformBuffer and InstancedViewUniformBuffer (if needed). */
+	void CreateViewUniformBuffers(const FViewUniformShaderParameters& Params);
+
 	/** Determines distance culling and fades if the state changes */
 	bool IsDistanceCulled(float DistanceSquared, float MinDrawDistance, float InMaxDrawDistance, const FPrimitiveSceneInfo* PrimitiveSceneInfo);
 
