Check that producer is still valid before submitting a tile.
A producer will be invalid if it was freed while a deferred tile was waiting.
#rnx

diff --git a/Engine/Source/Runtime/Renderer/Private/VT/VirtualTextureSystem.cpp b/Engine/Source/Runtime/Renderer/Private/VT/VirtualTextureSystem.cpp
--- a/Engine/Source/Runtime/Renderer/Private/VT/VirtualTextureSystem.cpp
+++ b/Engine/Source/Runtime/Renderer/Private/VT/VirtualTextureSystem.cpp
@@ -2184,24 +2184,30 @@ void FVirtualTextureSystem::SubmitRequestsFromLocalTileList(TArray<FVirtualTextu
 	for (const FVirtualTextureLocalTile& Tile : LocalTileList)
 	{
 		const FVirtualTextureProducerHandle ProducerHandle = Tile.GetProducerHandle();
-		const FVirtualTextureProducer& Producer = Producers.GetProducer(ProducerHandle);
+		FVirtualTextureProducer const* Producer = Producers.FindProducer(ProducerHandle);
+		if (Producer == nullptr)
+		{
+			// If we didn't process the tile last frame and deferred processing and if the producer got removed in the meantime we end up here.
+			// Just throw away the request.
+			continue;
+		}
 
 		// Fill targets for each layer
 		// Each producer can have multiple physical layers
 		// If the phys layer is mapped then we get the textures it owns and map them into the producer local slots and set the flags
 		uint32 LayerMask = 0;
 		FVTProduceTargetLayer ProduceTarget[VIRTUALTEXTURE_SPACE_MAXLAYERS];
-		for (uint32 ProducerPhysicalGroupIndex = 0u; ProducerPhysicalGroupIndex < Producer.GetNumPhysicalGroups(); ++ProducerPhysicalGroupIndex)
+		for (uint32 ProducerPhysicalGroupIndex = 0u; ProducerPhysicalGroupIndex < Producer->GetNumPhysicalGroups(); ++ProducerPhysicalGroupIndex)
 		{
-			FVirtualTexturePhysicalSpace* RESTRICT PhysicalSpace = Producer.GetPhysicalSpaceForPhysicalGroup(ProducerPhysicalGroupIndex);
+			FVirtualTexturePhysicalSpace* RESTRICT PhysicalSpace = Producer->GetPhysicalSpaceForPhysicalGroup(ProducerPhysicalGroupIndex);
 			FTexturePagePool& RESTRICT PagePool = PhysicalSpace->GetPagePool();
 			const uint32 pAddress = PagePool.FindPageAddress(ProducerHandle, ProducerPhysicalGroupIndex, Tile.Local_vAddress, Tile.Local_vLevel);
 			if (pAddress != ~0u)
 			{
 				int32 PhysicalLocalTextureIndex = 0;
-				for (uint32 ProducerLayerIndex = 0u; ProducerLayerIndex < Producer.GetNumTextureLayers(); ++ProducerLayerIndex)
+				for (uint32 ProducerLayerIndex = 0u; ProducerLayerIndex < Producer->GetNumTextureLayers(); ++ProducerLayerIndex)
 				{
-					if (Producer.GetPhysicalGroupIndexForTextureLayer(ProducerLayerIndex) == ProducerPhysicalGroupIndex)
+					if (Producer->GetPhysicalGroupIndexForTextureLayer(ProducerLayerIndex) == ProducerPhysicalGroupIndex)
 					{
 						ProduceTarget[ProducerLayerIndex].TextureRHI = PhysicalSpace->GetPhysicalTexture(PhysicalLocalTextureIndex);
 						ProduceTarget[ProducerLayerIndex].UnorderedAccessViewRHI = PhysicalSpace->GetPhysicalTextureUAV(PhysicalLocalTextureIndex);
@@ -2220,7 +2226,7 @@ void FVirtualTextureSystem::SubmitRequestsFromLocalTileList(TArray<FVirtualTextu
 			continue;
 		}
 
-		FVTRequestPageResult RequestPageResult = Producer.GetVirtualTexture()->RequestPageData(
+		FVTRequestPageResult RequestPageResult = Producer->GetVirtualTexture()->RequestPageData(
 			ProducerHandle, LayerMask, Tile.Local_vLevel, Tile.Local_vAddress, EVTRequestPagePriority::High);
 
 		if (RequestPageResult.Status != EVTRequestPageStatus::Available)
@@ -2230,7 +2236,7 @@ void FVirtualTextureSystem::SubmitRequestsFromLocalTileList(TArray<FVirtualTextu
 			continue;
 		}
 
-		IVirtualTextureFinalizer* VTFinalizer = Producer.GetVirtualTexture()->ProducePageData(
+		IVirtualTextureFinalizer* VTFinalizer = Producer->GetVirtualTexture()->ProducePageData(
 			GraphBuilder.RHICmdList, FeatureLevel,
 			Flags,
 			ProducerHandle, LayerMask, Tile.Local_vLevel, Tile.Local_vAddress,
