Fix compile error for non-PCH builds

diff --git a/Engine/Source/Runtime/Renderer/Public/PathTracingDenoiser.h b/Engine/Source/Runtime/Renderer/Public/PathTracingDenoiser.h
--- a/Engine/Source/Runtime/Renderer/Public/PathTracingDenoiser.h
+++ b/Engine/Source/Runtime/Renderer/Public/PathTracingDenoiser.h
@@ -2,6 +2,8 @@
 
 #pragma once
 
+#include "MultiGPU.h"
+
 class FRHITexture;
 class FRHICommandListImmediate;
 
