Minor tweak to screen ray default length to address a few edge cases

diff --git a/Engine/Source/Runtime/Renderer/Private/VirtualShadowMaps/VirtualShadowMapProjection.cpp b/Engine/Source/Runtime/Renderer/Private/VirtualShadowMaps/VirtualShadowMapProjection.cpp
--- a/Engine/Source/Runtime/Renderer/Private/VirtualShadowMaps/VirtualShadowMapProjection.cpp
+++ b/Engine/Source/Runtime/Renderer/Private/VirtualShadowMaps/VirtualShadowMapProjection.cpp
@@ -28,7 +28,7 @@
 
 static TAutoConsoleVariable<float> CVarScreenRayLength(
 	TEXT( "r.Shadow.Virtual.ScreenRayLength" ),
-	0.01f,
+	0.015f,
 	TEXT( "Length of the screen space shadow trace (smart shadow bias) before the virtual shadow map lookup." ),
 	ECVF_Scalability | ECVF_RenderThreadSafe
 );
