Initialize Strata on mobile to fix crashes on mobile.

diff --git a/Engine/Source/Runtime/Renderer/Private/MobileShadingRenderer.cpp b/Engine/Source/Runtime/Renderer/Private/MobileShadingRenderer.cpp
--- a/Engine/Source/Runtime/Renderer/Private/MobileShadingRenderer.cpp
+++ b/Engine/Source/Runtime/Renderer/Private/MobileShadingRenderer.cpp
@@ -58,6 +58,7 @@
 #include "VariableRateShadingImageManager.h"
 #include "SceneTextureReductions.h"
 #include "GPUMessaging.h"
+#include "Strata/Strata.h"
 
 uint32 GetShadowQuality();
 
@@ -725,6 +726,9 @@ void FMobileSceneRenderer::Render(FRDGBuilder& GraphBuilder)
 	// Initialize global system textures (pass-through if already initialized).
 	GSystemTextures.InitializeTextures(GraphBuilder.RHICmdList, FeatureLevel);
 
+	// Strata initialisation is always run even when not enabled.
+	Strata::InitialiseStrataFrameSceneData(*this, GraphBuilder);
+
 	// Force the subsurface profile texture to be updated.
 	UpdateSubsurfaceProfileTexture(GraphBuilder, ShaderPlatform);
 
diff --git a/Engine/Source/Runtime/Renderer/Private/PostProcess/PostProcessMaterial.cpp b/Engine/Source/Runtime/Renderer/Private/PostProcess/PostProcessMaterial.cpp
--- a/Engine/Source/Runtime/Renderer/Private/PostProcess/PostProcessMaterial.cpp
+++ b/Engine/Source/Runtime/Renderer/Private/PostProcess/PostProcessMaterial.cpp
@@ -616,10 +616,7 @@ FScreenPassTexture AddPostProcessMaterialPass(
 	}
 
 	PostProcessMaterialParameters->bFlipYAxis = Inputs.bFlipYAxis && !bForceIntermediateTarget;
-	if (Strata::IsStrataEnabled())
-	{
-		PostProcessMaterialParameters->Strata = Strata::BindStrataGlobalUniformParameters(View.StrataSceneData);
-	}
+	PostProcessMaterialParameters->Strata = Strata::BindStrataGlobalUniformParameters(View.StrataSceneData);
 
 	ClearUnusedGraphResources(VertexShader, PixelShader, PostProcessMaterialParameters);
 
diff --git a/Engine/Source/Runtime/Renderer/Private/ShadowRendering.cpp b/Engine/Source/Runtime/Renderer/Private/ShadowRendering.cpp
--- a/Engine/Source/Runtime/Renderer/Private/ShadowRendering.cpp
+++ b/Engine/Source/Runtime/Renderer/Private/ShadowRendering.cpp
@@ -453,7 +453,6 @@ static void BindShaderShaders(FRHICommandList& RHICmdList, FGraphicsPipelineStat
 	VertexShader->SetParameters(RHICmdList, View, ShadowInfo, EShadowProjectionVertexShaderFlags::DrawingFrustum);
 	PixelShader->SetParameters(RHICmdList, ViewIndex, View, ShadowInfo);
 
-	if (Strata::IsStrataEnabled())
 	{
 		FRHIPixelShader* ShaderRHI = PixelShader.GetPixelShader();
 		TRDGUniformBufferRef<FStrataGlobalUniformParameters> StrataUniformBuffer = Strata::BindStrataGlobalUniformParameters(View.StrataSceneData);
@@ -1992,12 +1991,7 @@ void FSceneRenderer::RenderShadowProjections(
 		FShadowProjectionPassParameters CommonPassParameters;
 		CommonPassParameters.SceneTextures = SceneTextures.GetSceneTextureShaderParameters(View.FeatureLevel);
 		CommonPassParameters.HairStrands = HairStrands::BindHairStrandsViewUniformParameters(View);
-
-		if (Strata::IsStrataEnabled())
-		{
-			CommonPassParameters.Strata = Strata::BindStrataGlobalUniformParameters(View.StrataSceneData);
-		}
-		
+		CommonPassParameters.Strata = Strata::BindStrataGlobalUniformParameters(View.StrataSceneData);
 		CommonPassParameters.RenderTargets[0] = FRenderTargetBinding(OutputTexture, ERenderTargetLoadAction::ELoad);
 		CommonPassParameters.RenderTargets.DepthStencil =
 			bSubPixelShadow ?
diff --git a/Engine/Source/Runtime/Renderer/Private/SingleLayerWaterRendering.cpp b/Engine/Source/Runtime/Renderer/Private/SingleLayerWaterRendering.cpp
--- a/Engine/Source/Runtime/Renderer/Private/SingleLayerWaterRendering.cpp
+++ b/Engine/Source/Runtime/Renderer/Private/SingleLayerWaterRendering.cpp
@@ -404,10 +404,7 @@ void FDeferredShadingSceneRenderer::RenderSingleLayerWaterReflections(
 				Parameters.ReflectionsParameters = CreateUniformBufferImmediate(ReflectionUniformParameters, UniformBuffer_SingleDraw);
 			}
 			Parameters.ForwardLightData = View.ForwardLightingResources.ForwardLightUniformBuffer;
-			if (Strata::IsStrataEnabled())
-			{
-				Parameters.Strata = Strata::BindStrataGlobalUniformParameters(View.StrataSceneData);
-			}
+			Parameters.Strata = Strata::BindStrataGlobalUniformParameters(View.StrataSceneData);
 		};
 
 		const bool bRunTiled = UseSingleLayerWaterIndirectDraw(View.GetShaderPlatform()) && CVarWaterSingleLayerTiledComposite.GetValueOnRenderThread();
