Implements r.Bloom.AsyncCompute

diff --git a/Engine/Source/Runtime/Renderer/Private/PostProcess/TemporalSuperResolution.cpp b/Engine/Source/Runtime/Renderer/Private/PostProcess/TemporalSuperResolution.cpp
--- a/Engine/Source/Runtime/Renderer/Private/PostProcess/TemporalSuperResolution.cpp
+++ b/Engine/Source/Runtime/Renderer/Private/PostProcess/TemporalSuperResolution.cpp
@@ -54,6 +54,11 @@ TAutoConsoleVariable<int32> CVarTSRRejectionAntiAliasingQuality(
 	TEXT("Controls the quality of spatial anti-aliasing on history rejection (default=1)."),
 	ECVF_Scalability | ECVF_RenderThreadSafe);
 
+TAutoConsoleVariable<int32> CVarTSRAsyncCompute(
+	TEXT("r.TSR.AsyncCompute"), 0,
+	TEXT("Whether to run TSR on async compute. Some TSR passes can overlap with previous passe.\n"),
+	ECVF_RenderThreadSafe);
+
 TAutoConsoleVariable<float> CVarTSRTranslucencyHighlightLuminance(
 	TEXT("r.TSR.Translucency.HighlightLuminance"), -1.0f,
 	TEXT("Sets the liminance at which translucency is considered an highlights (default=-1.0)."),
@@ -589,6 +594,9 @@ ITemporalUpscaler::FOutputs AddTemporalSuperResolutionPasses(
 	// Whether alpha channel is supported.
 	const bool bSupportsAlpha = IsPostProcessingWithAlphaChannelSupported();
 
+	// whether TSR passes can run on async compute.
+	ERDGPassFlags ComputePassFlags = (GSupportsEfficientAsyncCompute && CVarTSRAsyncCompute.GetValueOnRenderThread() != 0) ? ERDGPassFlags::AsyncCompute : ERDGPassFlags::Compute;
+
 	// Whether to use camera cut shader permutation or not.
 	bool bCameraCut = !InputHistory.IsValid() || View.bCameraCut;
 
@@ -765,6 +773,7 @@ ITemporalUpscaler::FOutputs AddTemporalSuperResolutionPasses(
 		FComputeShaderUtils::AddPass(
 			GraphBuilder,
 			RDG_EVENT_NAME("TSR ClearPrevTextures %dx%d", InputRect.Width(), InputRect.Height()),
+			ComputePassFlags,
 			ComputeShader,
 			PassParameters,
 			FComputeShaderUtils::GetGroupCount(InputRect.Size(), 8 * 2));
@@ -864,6 +873,7 @@ ITemporalUpscaler::FOutputs AddTemporalSuperResolutionPasses(
 			RDG_EVENT_NAME("TSR DilateVelocity(MotionBlurDirections=%d) %dx%d",
 				int32(PermutationVector.Get<FTSRDilateVelocityCS::FMotionBlurDirectionsDim>()),
 				InputRect.Width(), InputRect.Height()),
+			ComputePassFlags,
 			ComputeShader,
 			PassParameters,
 			FComputeShaderUtils::GetGroupCount(InputRect.Size(), TileSize));
@@ -1037,6 +1047,7 @@ ITemporalUpscaler::FOutputs AddTemporalSuperResolutionPasses(
 				bHistoryHighFrequencyOnly ? TEXT("HighFrequency") : TEXT("LowFrequency"),
 				bHalfResLowFrequency ? TEXT(" HalfResShadingOutput") : TEXT(""),
 				InputRect.Width(), InputRect.Height()),
+			ComputePassFlags,
 			ComputeShader,
 			PassParameters,
 			FComputeShaderUtils::GetGroupCount(InputRect.Size(), 8));
@@ -1094,6 +1105,7 @@ ITemporalUpscaler::FOutputs AddTemporalSuperResolutionPasses(
 		FComputeShaderUtils::AddPass(
 			GraphBuilder,
 			RDG_EVENT_NAME("TSR CompareTranslucency %dx%d", InputRect.Width(), InputRect.Height()),
+			ComputePassFlags,
 			ComputeShader,
 			PassParameters,
 			FComputeShaderUtils::GetGroupCount(InputRect.Size(), 8));
@@ -1132,6 +1144,7 @@ ITemporalUpscaler::FOutputs AddTemporalSuperResolutionPasses(
 		FComputeShaderUtils::AddPass(
 			GraphBuilder,
 			RDG_EVENT_NAME("TSR HoleFillVelocity %dx%d", InputRect.Width(), InputRect.Height()),
+			ComputePassFlags,
 			ComputeShader,
 			PassParameters,
 			FComputeShaderUtils::GetGroupCount(InputRect.Size(), 8));
@@ -1200,6 +1213,7 @@ ITemporalUpscaler::FOutputs AddTemporalSuperResolutionPasses(
 				RDG_EVENT_NAME("TSR FilterFrequencies(%s) %dx%d",
 					PermutationVector.Get<FTSRFilterFrequenciesCS::FOutputAALumaDim>() ? TEXT("OutputAALuma") : TEXT(""),
 					LowFrequencyRect.Width(), LowFrequencyRect.Height()),
+				ComputePassFlags,
 				ComputeShader,
 				PassParameters,
 				FComputeShaderUtils::GetGroupCount(LowFrequencyRect.Size(), 16));
@@ -1235,6 +1249,7 @@ ITemporalUpscaler::FOutputs AddTemporalSuperResolutionPasses(
 			FComputeShaderUtils::AddPass(
 				GraphBuilder,
 				RDG_EVENT_NAME("TSR CompareHistory %dx%d", LowFrequencyRect.Width(), LowFrequencyRect.Height()),
+				ComputePassFlags,
 				ComputeShader,
 				PassParameters,
 				FComputeShaderUtils::GetGroupCount(LowFrequencyRect.Size(), 16));
@@ -1279,6 +1294,7 @@ ITemporalUpscaler::FOutputs AddTemporalSuperResolutionPasses(
 				RDG_EVENT_NAME("TSR SpatialAntiAliasing(Quality=%d) %dx%d",
 					RejectionAntiAliasingQuality,
 					InputRect.Width(), InputRect.Height()),
+				ComputePassFlags,
 				ComputeShader,
 				PassParameters,
 				FComputeShaderUtils::GetGroupCount(InputRect.Size(), 8));
@@ -1295,6 +1311,7 @@ ITemporalUpscaler::FOutputs AddTemporalSuperResolutionPasses(
 			FComputeShaderUtils::AddPass(
 				GraphBuilder,
 				RDG_EVENT_NAME("TSR FilterAntiAliasing %dx%d", InputRect.Width(), InputRect.Height()),
+				ComputePassFlags,
 				ComputeShader,
 				PassParameters,
 				FComputeShaderUtils::GetGroupCount(InputRect.Size(), 8));
@@ -1331,6 +1348,7 @@ ITemporalUpscaler::FOutputs AddTemporalSuperResolutionPasses(
 		FComputeShaderUtils::AddPass(
 			GraphBuilder,
 			RDG_EVENT_NAME("TSR PostfilterRejection %dx%d", Rect.Width(), Rect.Height()),
+			ComputePassFlags,
 			ComputeShader,
 			PassParameters,
 			FComputeShaderUtils::GetGroupCount(Rect.Size(), 8));
@@ -1353,6 +1371,7 @@ ITemporalUpscaler::FOutputs AddTemporalSuperResolutionPasses(
 		FComputeShaderUtils::AddPass(
 			GraphBuilder,
 			RDG_EVENT_NAME("TSR DilateRejection %dx%d", RejectionRect.Width(), RejectionRect.Height()),
+			ComputePassFlags,
 			ComputeShader,
 			PassParameters,
 			FComputeShaderUtils::GetGroupCount(RejectionRect.Size(), 8));
@@ -1446,6 +1465,7 @@ ITemporalUpscaler::FOutputs AddTemporalSuperResolutionPasses(
 				PermutationVector.Get<FTSRHighFrequencyOnlyDim>() ? TEXT("") : TEXT(" LowFrequency"),
 				PassInputs.bGenerateOutputMip1 ? TEXT(" OutputMip1") : TEXT(""),
 				HistorySize.X, HistorySize.Y),
+			ComputePassFlags,
 			ComputeShader,
 			PassParameters,
 			FComputeShaderUtils::GetGroupCount(HistorySize, 8));
@@ -1467,6 +1487,7 @@ ITemporalUpscaler::FOutputs AddTemporalSuperResolutionPasses(
 		FComputeShaderUtils::AddPass(
 			GraphBuilder,
 			RDG_EVENT_NAME("TSR DebugHistory %dx%d", HistorySize.X, HistorySize.Y),
+			ComputePassFlags,
 			ComputeShader,
 			PassParameters,
 			FComputeShaderUtils::GetGroupCount(HistorySize * kHistoryUpscalingFactor, 8));
@@ -1508,6 +1529,7 @@ ITemporalUpscaler::FOutputs AddTemporalSuperResolutionPasses(
 		FComputeShaderUtils::AddPass(
 			GraphBuilder,
 			RDG_EVENT_NAME("TSR ResolveHistory %dx%d", OutputRect.Width(), OutputRect.Height()),
+			ComputePassFlags,
 			ComputeShader,
 			PassParameters,
 			FComputeShaderUtils::GetGroupCount(OutputRect.Size(), 8));
