Fix for strata crash when disabled.
Backout a file from 19415778.

diff --git a/Engine/Source/Runtime/Renderer/Private/ShadowRendering.cpp b/Engine/Source/Runtime/Renderer/Private/ShadowRendering.cpp
--- a/Engine/Source/Runtime/Renderer/Private/ShadowRendering.cpp
+++ b/Engine/Source/Runtime/Renderer/Private/ShadowRendering.cpp
@@ -453,6 +453,7 @@ static void BindShaderShaders(FRHICommandList& RHICmdList, FGraphicsPipelineStat
 	VertexShader->SetParameters(RHICmdList, View, ShadowInfo, EShadowProjectionVertexShaderFlags::DrawingFrustum);
 	PixelShader->SetParameters(RHICmdList, ViewIndex, View, ShadowInfo);
 
+	if (Strata::IsStrataEnabled())
 	{
 		FRHIPixelShader* ShaderRHI = PixelShader.GetPixelShader();
 		TRDGUniformBufferRef<FStrataGlobalUniformParameters> StrataUniformBuffer = Strata::BindStrataGlobalUniformParameters(View.StrataSceneData);
@@ -1991,7 +1992,12 @@ void FSceneRenderer::RenderShadowProjections(
 		FShadowProjectionPassParameters CommonPassParameters;
 		CommonPassParameters.SceneTextures = SceneTextures.GetSceneTextureShaderParameters(View.FeatureLevel);
 		CommonPassParameters.HairStrands = HairStrands::BindHairStrandsViewUniformParameters(View);
-		CommonPassParameters.Strata = Strata::BindStrataGlobalUniformParameters(View.StrataSceneData);
+
+		if (Strata::IsStrataEnabled())
+		{
+			CommonPassParameters.Strata = Strata::BindStrataGlobalUniformParameters(View.StrataSceneData);
+		}
+		
 		CommonPassParameters.RenderTargets[0] = FRenderTargetBinding(OutputTexture, ERenderTargetLoadAction::ELoad);
 		CommonPassParameters.RenderTargets.DepthStencil =
 			bSubPixelShadow ?
