Fixed Nanite material slot table crash when mesh section count is larger than max encountered material index

diff --git a/Engine/Source/Runtime/Experimental/GeometryCollectionEngine/Private/GeometryCollection/GeometryCollectionSceneProxy.cpp b/Engine/Source/Runtime/Experimental/GeometryCollectionEngine/Private/GeometryCollection/GeometryCollectionSceneProxy.cpp
--- a/Engine/Source/Runtime/Experimental/GeometryCollectionEngine/Private/GeometryCollection/GeometryCollectionSceneProxy.cpp
+++ b/Engine/Source/Runtime/Experimental/GeometryCollectionEngine/Private/GeometryCollection/GeometryCollectionSceneProxy.cpp
@@ -1435,11 +1435,8 @@ FNaniteGeometryCollectionSceneProxy::FNaniteGeometryCollectionSceneProxy(UGeomet
 		const FGeometryCollectionSection& MeshSection = SectionsArray[SectionIndex];
 		const bool bValidMeshSection = MeshSection.MaterialID != INDEX_NONE;
 
-		if (MeshSection.MaterialID > MaterialMaxIndex)
-		{
-			// Keep track of highest observed material index.
-			MaterialMaxIndex = MeshSection.MaterialID;
-		}
+		// Keep track of highest observed material index.
+		MaterialMaxIndex = FMath::Max(MeshSection.MaterialID, MaterialMaxIndex);
 
 		UMaterialInterface* MaterialInterface = bValidMeshSection ? Component->GetMaterial(MeshSection.MaterialID) : nullptr;
 
diff --git a/Engine/Source/Runtime/Renderer/Private/GPUScene.cpp b/Engine/Source/Runtime/Renderer/Private/GPUScene.cpp
--- a/Engine/Source/Runtime/Renderer/Private/GPUScene.cpp
+++ b/Engine/Source/Runtime/Renderer/Private/GPUScene.cpp
@@ -1193,7 +1193,6 @@ void FGPUScene::UploadGeneral(FRHICommandListImmediate& RHICmdList, FScene *Scen
 							const FPrimitiveSceneInfo* PrimitiveSceneInfo = UploadInfo.PrimitiveSceneInfo;
 							const Nanite::FSceneProxyBase* NaniteSceneProxy = UploadInfo.NaniteSceneProxy;
 							const TArray<Nanite::FSceneProxyBase::FMaterialSection>& PassMaterials = NaniteSceneProxy->GetMaterialSections();
-							const uint32 TableEntryCount = uint32(NaniteSceneProxy->GetMaterialMaxIndex() + 1);
 
 							// Update raster bin, material depth and hit proxy ID remapping tables.
 							for (int32 NaniteMeshPass = 0; NaniteMeshPass < ENaniteMeshPass::Num; ++NaniteMeshPass)
@@ -1203,11 +1202,11 @@ void FGPUScene::UploadGeneral(FRHICommandListImmediate& RHICmdList, FScene *Scen
 								
 								if (PassMaterials.Num() == PassMaterialSlots.Num())
 								{
-									check(TableEntryCount >= uint32(PassMaterials.Num()));
+									const uint32 MaterialSlotCount = uint32(PassMaterialSlots.Num());
 
-									void* MaterialSlotRange = NaniteMaterials.GetMaterialSlotPtr(UploadInfo.PrimitiveID, TableEntryCount);
+									void* MaterialSlotRange = NaniteMaterials.GetMaterialSlotPtr(UploadInfo.PrimitiveID, MaterialSlotCount);
 									uint32* MaterialSlots = static_cast<uint32*>(MaterialSlotRange);
-									for (int32 Entry = 0; Entry < PassMaterialSlots.Num(); ++Entry)
+									for (uint32 Entry = 0; Entry < MaterialSlotCount; ++Entry)
 									{
 										MaterialSlots[PassMaterials[Entry].MaterialIndex] = PassMaterialSlots[Entry].Pack();
 									}
@@ -1216,7 +1215,7 @@ void FGPUScene::UploadGeneral(FRHICommandListImmediate& RHICmdList, FScene *Scen
 									if (NaniteMeshPass == ENaniteMeshPass::BasePass && NaniteSceneProxy->GetHitProxyMode() == Nanite::FSceneProxyBase::EHitProxyMode::MaterialSection)
 									{
 										const TArray<uint32>& PassHitProxyIds = PrimitiveSceneInfo->NaniteHitProxyIds;
-										uint32* HitProxyTable = static_cast<uint32*>(NaniteMaterials.GetHitProxyTablePtr(UploadInfo.PrimitiveID, TableEntryCount));
+										uint32* HitProxyTable = static_cast<uint32*>(NaniteMaterials.GetHitProxyTablePtr(UploadInfo.PrimitiveID, MaterialSlotCount));
 										for (int32 Entry = 0; Entry < PassHitProxyIds.Num(); ++Entry)
 										{
 											HitProxyTable[PassMaterials[Entry].MaterialIndex] = PassHitProxyIds[Entry];
