Fixed "Nanite is used in the scene but not supported by your graphics hardware and/or driver. Meshes will not render using Nanite." error message when vendor extensions are supported, but ultimately Nanite is disabled due to lack of SM 6.6 support

diff --git a/Engine/Source/Runtime/Renderer/Private/SceneRendering.cpp b/Engine/Source/Runtime/Renderer/Private/SceneRendering.cpp
--- a/Engine/Source/Runtime/Renderer/Private/SceneRendering.cpp
+++ b/Engine/Source/Runtime/Renderer/Private/SceneRendering.cpp
@@ -382,6 +382,19 @@ static TAutoConsoleVariable<int32> CVarTestSecondaryUpscaleOverride(
 
 #endif
 
+#if !(UE_BUILD_SHIPPING || UE_BUILD_TEST)
+
+static TAutoConsoleVariable<int32> CVarNaniteShowUnsupportedError(
+	TEXT("r.Nanite.ShowUnsupportedError"),
+	1,
+	TEXT("Specify behavior of Nanite unsupported screen error message.\n")
+	TEXT(" 0: disabled\n")
+	TEXT(" 1: show error if Nanite is present in the scene but unsupported, and fallback meshes are not used for rendering; (default)")
+	TEXT(" 2: show error if Nanite is present in the scene but unsupported, even if fallback meshes are used for rendering")
+);
+
+#endif
+
 static FParallelCommandListSet* GOutstandingParallelCommandListSet = nullptr;
 FOcclusionSubmittedFenceState FSceneRenderer::OcclusionSubmittedFence[FOcclusionQueryHelpers::MaxBufferedOcclusionFrames];
 
@@ -3203,8 +3216,21 @@ void FSceneRenderer::RenderFinish(FRDGBuilder& GraphBuilder, FRDGTextureRef View
 
 		const bool bShowLocalExposureDisabledWarning = ViewFamily.EngineShowFlags.VisualizeLocalExposure && !bLocalExposureEnabledOnAnyView;
 
+		const int32 NaniteShowError = CVarNaniteShowUnsupportedError.GetValueOnRenderThread();
+		// 0: disabled
+		// 1: show error if Nanite is present in the scene but unsupported, and fallback meshes are not used for rendering
+		// 2: show error if Nanite is present in the scene but unsupported, even if fallback meshes are used for rendering
+
+		static const auto NaniteProxyRenderModeVar = IConsoleManager::Get().FindConsoleVariable(TEXT("r.Nanite.ProxyRenderMode"));
+		const int32 NaniteProxyRenderMode = (NaniteProxyRenderModeVar != nullptr) ? (NaniteProxyRenderModeVar->GetInt() != 0) : 0;
+		// 0: Fall back to rendering Nanite proxy meshes if Nanite is unsupported.
+		// 1: Disable rendering if Nanite is enabled on a mesh but is unsupported
+		// 2: Disable rendering if Nanite is enabled on a mesh but is unsupported, except for static mesh editor toggle
+
 		bool bNaniteEnabledButNoAtomics = false;
-		if (!GRHISupportsAtomicUInt64)
+
+		bool bNaniteCheckError = (NaniteShowError == 1 && NaniteProxyRenderMode != 0) || (NaniteShowError == 2);
+		if (bNaniteCheckError && !NaniteAtomicsSupported())
 		{
 			// We want to know when Nanite would've been rendered regardless of atomics being supported or not.
 			const bool bCheckForAtomicSupport = false;
