Disable primitive octree culling (r.UseVisibilityOctree) by default as this causes visibility culling issues because of the limited extent of the tree (HALF_WORLD_MAX).
- Disabling brings a marginal performance increase for CitySample/AncientGame (both have very high primitive counts), so unclear when we would want it on.

diff --git a/Engine/Source/Runtime/Renderer/Private/SceneVisibility.cpp b/Engine/Source/Runtime/Renderer/Private/SceneVisibility.cpp
--- a/Engine/Source/Runtime/Renderer/Private/SceneVisibility.cpp
+++ b/Engine/Source/Runtime/Renderer/Private/SceneVisibility.cpp
@@ -265,7 +265,7 @@ static TAutoConsoleVariable<int32> CVarUseFastIntersect(
 
 static TAutoConsoleVariable<int32> CVarUseVisibilityOctree(
 	TEXT("r.UseVisibilityOctree"), 
-	1, 
+	0, 
 	TEXT("Use the octree for visibility calculations."), 
 	ECVF_RenderThreadSafe);
 
