Add BSDF tile and BSDF offset computation.

For pixel with complex/multi-BSDF computes BSDF material byte/index offsets and 'overflowing' tile.

This wil be used by Lumen for handling/parallelizing lighting computation.

diff --git a/Engine/Shaders/Private/Strata/StrataMaterialClassification.usf b/Engine/Shaders/Private/Strata/StrataMaterialClassification.usf
--- a/Engine/Shaders/Private/Strata/StrataMaterialClassification.usf
+++ b/Engine/Shaders/Private/Strata/StrataMaterialClassification.usf
@@ -19,7 +19,7 @@ Texture2D<uint> TopLayerTexture;
 Texture2DArray<uint> MaterialTextureArray;
 
 // Indirect draw data buffer for all tile types
-RWBuffer<uint> TileIndirectDataBuffer;
+RWBuffer<uint> TileDrawIndirectDataBuffer;
 
 RWBuffer<uint> SimpleTileListDataBuffer;
 RWBuffer<uint> SingleTileListDataBuffer;
@@ -45,11 +45,11 @@ void TileMainCS(uint2 DispatchThreadId : SV_DispatchThreadID, uint LinearIndex :
 	if (all(DispatchThreadId == 0))
 	{
 		const uint IndexCountPerInstance = bRectPrimitive > 0 ? 4 : 6;
-		TileIndirectDataBuffer[GetStrataTileTypeDrawIndirectArgOffset_DWord(STRATA_TILE_TYPE_SIMPLE)					+ 0] = IndexCountPerInstance;
-		TileIndirectDataBuffer[GetStrataTileTypeDrawIndirectArgOffset_DWord(STRATA_TILE_TYPE_SINGLE)					+ 0] = IndexCountPerInstance;
-		TileIndirectDataBuffer[GetStrataTileTypeDrawIndirectArgOffset_DWord(STRATA_TILE_TYPE_COMPLEX)					+ 0] = IndexCountPerInstance;
-		TileIndirectDataBuffer[GetStrataTileTypeDrawIndirectArgOffset_DWord(STRATA_TILE_TYPE_ROUGH_REFRACT)				+ 0] = IndexCountPerInstance;
-		TileIndirectDataBuffer[GetStrataTileTypeDrawIndirectArgOffset_DWord(STRATA_TILE_TYPE_SSS_WITHOUT_ROUGH_REFRACT) + 0] = IndexCountPerInstance;
+		TileDrawIndirectDataBuffer[GetStrataTileTypeDrawIndirectArgOffset_DWord(STRATA_TILE_TYPE_SIMPLE)					+ 0] = IndexCountPerInstance;
+		TileDrawIndirectDataBuffer[GetStrataTileTypeDrawIndirectArgOffset_DWord(STRATA_TILE_TYPE_SINGLE)					+ 0] = IndexCountPerInstance;
+		TileDrawIndirectDataBuffer[GetStrataTileTypeDrawIndirectArgOffset_DWord(STRATA_TILE_TYPE_COMPLEX)					+ 0] = IndexCountPerInstance;
+		TileDrawIndirectDataBuffer[GetStrataTileTypeDrawIndirectArgOffset_DWord(STRATA_TILE_TYPE_ROUGH_REFRACT)				+ 0] = IndexCountPerInstance;
+		TileDrawIndirectDataBuffer[GetStrataTileTypeDrawIndirectArgOffset_DWord(STRATA_TILE_TYPE_SSS_WITHOUT_ROUGH_REFRACT) + 0] = IndexCountPerInstance;
 	}
 
 	// TODO: add a SM6 permutation with ballot?
@@ -154,42 +154,76 @@ void TileMainCS(uint2 DispatchThreadId : SV_DispatchThreadID, uint LinearIndex :
 		if (bTileContainsComplex)
 		{
 			uint WriteToIndex;
-			InterlockedAdd(TileIndirectDataBuffer[GetStrataTileTypeDrawIndirectArgOffset_DWord(STRATA_TILE_TYPE_COMPLEX) + 1], 1, WriteToIndex);
+			InterlockedAdd(TileDrawIndirectDataBuffer[GetStrataTileTypeDrawIndirectArgOffset_DWord(STRATA_TILE_TYPE_COMPLEX) + 1], 1, WriteToIndex);
 			ComplexTileListDataBuffer[WriteToIndex] = EncodedTile;
 		}
 		else if (bTileContainsSingle)
 		{
 			uint WriteToIndex;
-			InterlockedAdd(TileIndirectDataBuffer[GetStrataTileTypeDrawIndirectArgOffset_DWord(STRATA_TILE_TYPE_SINGLE) + 1], 1, WriteToIndex);
+			InterlockedAdd(TileDrawIndirectDataBuffer[GetStrataTileTypeDrawIndirectArgOffset_DWord(STRATA_TILE_TYPE_SINGLE) + 1], 1, WriteToIndex);
 			SingleTileListDataBuffer[WriteToIndex] = EncodedTile;
 		}
 		else // (bTileContainsSimple)
 		{
 			uint WriteToIndex;
-			InterlockedAdd(TileIndirectDataBuffer[GetStrataTileTypeDrawIndirectArgOffset_DWord(STRATA_TILE_TYPE_SIMPLE) + 1], 1, WriteToIndex);
+			InterlockedAdd(TileDrawIndirectDataBuffer[GetStrataTileTypeDrawIndirectArgOffset_DWord(STRATA_TILE_TYPE_SIMPLE) + 1], 1, WriteToIndex);
 			SimpleTileListDataBuffer[WriteToIndex] = EncodedTile;
 		}
 
-#if STRATA_OPAQUE_ROUGH_REFRACTION_ENABLED
+	#if STRATA_OPAQUE_ROUGH_REFRACTION_ENABLED
 		if (bTileContainsOpaqueRoughRefraction)
 		{
 			uint WriteToIndex;
-			InterlockedAdd(TileIndirectDataBuffer[GetStrataTileTypeDrawIndirectArgOffset_DWord(STRATA_TILE_TYPE_ROUGH_REFRACT) + 1], 1, WriteToIndex);
+			InterlockedAdd(TileDrawIndirectDataBuffer[GetStrataTileTypeDrawIndirectArgOffset_DWord(STRATA_TILE_TYPE_ROUGH_REFRACT) + 1], 1, WriteToIndex);
 			OpaqueRoughRefractionTileListDataBuffer[WriteToIndex] = EncodedTile;
 		}
 		if(bTileContainsScreenSpaceSubsurfaceScattering && !bTileContainsOpaqueRoughRefraction)
 		{
 			uint WriteToIndex;
-			InterlockedAdd(TileIndirectDataBuffer[GetStrataTileTypeDrawIndirectArgOffset_DWord(STRATA_TILE_TYPE_SSS_WITHOUT_ROUGH_REFRACT) + 1], 1, WriteToIndex);
+			InterlockedAdd(TileDrawIndirectDataBuffer[GetStrataTileTypeDrawIndirectArgOffset_DWord(STRATA_TILE_TYPE_SSS_WITHOUT_ROUGH_REFRACT) + 1], 1, WriteToIndex);
 			SSSWithoutOpaqueRoughRefractionTileListDataBuffer[WriteToIndex] = EncodedTile;
 		}
-#endif
+	#endif
 	}
 }
 #endif // SHADER_TILE_CATEGORIZATION
 
 ////////////////////////////////////////////////////////////////////////////////////////////////////////////
 
+#if SHADER_TILE_PREPARE_ARGS
+
+Buffer<uint>   TileDrawIndirectDataBuffer;
+RWBuffer<uint> TileDispatchIndirectDataBuffer;
+
+[numthreads(1, 1, 1)]
+void ArgsMainCS(uint2 DispatchThreadId : SV_DispatchThreadID)
+{
+	if (all(DispatchThreadId == 0))
+	{
+		TileDispatchIndirectDataBuffer[GetStrataTileTypeDispatchIndirectArgOffset_DWord(STRATA_TILE_TYPE_SIMPLE)					+ 0] = TileDrawIndirectDataBuffer[GetStrataTileTypeDrawIndirectArgOffset_DWord(STRATA_TILE_TYPE_SIMPLE)						+ 1];
+		TileDispatchIndirectDataBuffer[GetStrataTileTypeDispatchIndirectArgOffset_DWord(STRATA_TILE_TYPE_SINGLE)					+ 0] = TileDrawIndirectDataBuffer[GetStrataTileTypeDrawIndirectArgOffset_DWord(STRATA_TILE_TYPE_SINGLE)						+ 1];
+		TileDispatchIndirectDataBuffer[GetStrataTileTypeDispatchIndirectArgOffset_DWord(STRATA_TILE_TYPE_COMPLEX)					+ 0] = TileDrawIndirectDataBuffer[GetStrataTileTypeDrawIndirectArgOffset_DWord(STRATA_TILE_TYPE_COMPLEX)					+ 1];
+		TileDispatchIndirectDataBuffer[GetStrataTileTypeDispatchIndirectArgOffset_DWord(STRATA_TILE_TYPE_ROUGH_REFRACT)				+ 0] = TileDrawIndirectDataBuffer[GetStrataTileTypeDrawIndirectArgOffset_DWord(STRATA_TILE_TYPE_ROUGH_REFRACT)				+ 1];
+		TileDispatchIndirectDataBuffer[GetStrataTileTypeDispatchIndirectArgOffset_DWord(STRATA_TILE_TYPE_SSS_WITHOUT_ROUGH_REFRACT)	+ 0] = TileDrawIndirectDataBuffer[GetStrataTileTypeDrawIndirectArgOffset_DWord(STRATA_TILE_TYPE_SSS_WITHOUT_ROUGH_REFRACT)	+ 1];
+
+		TileDispatchIndirectDataBuffer[GetStrataTileTypeDispatchIndirectArgOffset_DWord(STRATA_TILE_TYPE_SIMPLE)					+ 1] = 1;
+		TileDispatchIndirectDataBuffer[GetStrataTileTypeDispatchIndirectArgOffset_DWord(STRATA_TILE_TYPE_SINGLE)					+ 1] = 1;
+		TileDispatchIndirectDataBuffer[GetStrataTileTypeDispatchIndirectArgOffset_DWord(STRATA_TILE_TYPE_COMPLEX)					+ 1] = 1;
+		TileDispatchIndirectDataBuffer[GetStrataTileTypeDispatchIndirectArgOffset_DWord(STRATA_TILE_TYPE_ROUGH_REFRACT)				+ 1] = 1;
+		TileDispatchIndirectDataBuffer[GetStrataTileTypeDispatchIndirectArgOffset_DWord(STRATA_TILE_TYPE_SSS_WITHOUT_ROUGH_REFRACT) + 1] = 1;
+
+		TileDispatchIndirectDataBuffer[GetStrataTileTypeDispatchIndirectArgOffset_DWord(STRATA_TILE_TYPE_SIMPLE)					+ 2] = 1;
+		TileDispatchIndirectDataBuffer[GetStrataTileTypeDispatchIndirectArgOffset_DWord(STRATA_TILE_TYPE_SINGLE)					+ 2] = 1;
+		TileDispatchIndirectDataBuffer[GetStrataTileTypeDispatchIndirectArgOffset_DWord(STRATA_TILE_TYPE_COMPLEX)					+ 2] = 1;
+		TileDispatchIndirectDataBuffer[GetStrataTileTypeDispatchIndirectArgOffset_DWord(STRATA_TILE_TYPE_ROUGH_REFRACT)				+ 2] = 1;
+		TileDispatchIndirectDataBuffer[GetStrataTileTypeDispatchIndirectArgOffset_DWord(STRATA_TILE_TYPE_SSS_WITHOUT_ROUGH_REFRACT)	+ 2] = 1;
+	}
+}
+
+#endif // SHADER_TILE_PREPARE_ARGS
+
+////////////////////////////////////////////////////////////////////////////////////////////////////////////
+
 #if SHADER_CLEAR_MATERIAL_BUFFER
 
 RWTexture2DArray<uint> MaterialTextureArrayUAV;
@@ -223,3 +257,122 @@ void ClearMaterialBufferMainCS(uint2 DispatchThreadId : SV_DispatchThreadID)
 
 ////////////////////////////////////////////////////////////////////////////////////////////////////////////
 
+#if SHADER_BSDF_TILE
+
+int2 ViewResolution;
+uint MaxBytesPerPixel;
+uint TileSizeLog2;
+int2 TileCount_Primary;
+
+Texture2D<uint> TopLayerTexture;
+Texture2DArray<uint> MaterialTextureArray;
+
+Buffer<uint> TileListBuffer;
+
+RWTexture2D<uint> RWBSDFTileTexture;
+RWTexture2D<uint> RWBSDFOffsetTexture;
+RWBuffer<uint>    RWBSDFTileCountBuffer;
+
+#if !PERMUTATION_WAVE_OPS
+groupshared uint s_TileBSDFCount[GROUP_THREAD_COUNT];
+#endif
+
+[numthreads(STRATA_TILE_SIZE, STRATA_TILE_SIZE, 1)]
+void BSDFTileMainCS(uint2 GroupThreadId : SV_GroupThreadID, uint2 GroupId : SV_GroupID, uint LinearIndex : SV_GroupIndex)
+{
+	const uint2 TileCoord = StrataUnpackTile(TileListBuffer[GroupId.x]);
+	const uint2 PixelCoord = TileCoord * STRATA_TILE_SIZE + GroupThreadId;
+
+	uint BSDFCount = 0;
+	if (all(PixelCoord < uint2(ViewResolution)))
+	{
+		FStrataAddressing StrataAddressing = GetStrataPixelDataByteOffset(PixelCoord, uint2(View.BufferSizeAndInvSize.xy), MaxBytesPerPixel);
+		FStrataPixelHeader StrataPixelHeader = UnpackStrataHeaderIn(MaterialTextureArray, StrataAddressing, TopLayerTexture);
+		BSDFCount = StrataPixelHeader.BSDFCount;
+
+		if (BSDFCount > 0)
+		{
+			FStrataBSDFOffset Offsets = (FStrataBSDFOffset)0;
+			Offsets.BSDFCount = BSDFCount;
+
+			for (uint BSDFIndex = 0; BSDFIndex < BSDFCount; ++BSDFIndex)
+			{
+				Offsets.BSDFOffsets[BSDFIndex] = StrataAddressing.CurrentIndex;
+				UnpackStrataBSDFIn(MaterialTextureArray, StrataAddressing, StrataPixelHeader);
+			}
+
+			RWBSDFOffsetTexture[PixelCoord] = PackBSDFOffset(Offsets);
+		}
+	}
+
+#if PERMUTATION_WAVE_OPS
+
+	const uint TileBSDFCount = WaveActiveMax(BSDFCount);
+
+#else // PERMUTATION_WAVE_OPS
+
+	s_TileBSDFCount[LinearIndex] = BSDFCount;
+
+	GroupMemoryBarrierWithGroupSync();
+	if (LinearIndex < 32)
+	{
+		s_TileBSDFCount[LinearIndex] = max(s_TileBSDFCount[LinearIndex], s_TileBSDFCount[LinearIndex + 32]);
+	}
+	GroupMemoryBarrierWithGroupSync();
+	if (LinearIndex < 16)
+	{
+		s_TileBSDFCount[LinearIndex] = max(s_TileBSDFCount[LinearIndex], s_TileBSDFCount[LinearIndex + 16]);
+	}
+	GroupMemoryBarrierWithGroupSync();
+
+	if (LinearIndex < 8)
+	{
+		s_TileBSDFCount[LinearIndex] = max(s_TileBSDFCount[LinearIndex], s_TileBSDFCount[LinearIndex + 8]);
+	}
+	GroupMemoryBarrierWithGroupSync();
+	if (LinearIndex < 4)
+	{
+		s_TileBSDFCount[LinearIndex] = max(s_TileBSDFCount[LinearIndex], s_TileBSDFCount[LinearIndex + 4]);
+	}
+	GroupMemoryBarrierWithGroupSync();
+	if (LinearIndex < 2)
+	{
+		s_TileBSDFCount[LinearIndex] = max(s_TileBSDFCount[LinearIndex], s_TileBSDFCount[LinearIndex + 2]);
+	}
+	GroupMemoryBarrierWithGroupSync();
+
+	const uint TileBSDFCount = max(s_TileBSDFCount[LinearIndex], s_TileBSDFCount[LinearIndex + 1]);
+
+#endif // PERMUTATION_WAVE_OPS
+
+	if (LinearIndex < 1)
+	{
+		if (TileBSDFCount > 1)
+		{
+			uint LinearTileCoord = 0;
+			InterlockedAdd(RWBSDFTileCountBuffer[0], TileBSDFCount - 1, LinearTileCoord);
+
+			FStrataBSDFTile Tile;
+			Tile.Count = TileBSDFCount;
+			Tile.Index = 0;
+
+			// For first/parent tile, set TileCoord to point towards the first 'overflow' tile
+			Tile.TileCoord = uint2(LinearTileCoord % TileCount_Primary.x, TileCount_Primary.y + LinearTileCoord / TileCount_Primary.y);
+			RWBSDFTileTexture[TileCoord] = PackBSDFTile(Tile);
+
+			// For 'overflow' tiles, set TileCoord to point towards the first/parent tile
+			Tile.TileCoord = TileCoord;
+
+			for (uint BSDFIndex = 1; BSDFIndex < TileBSDFCount; ++BSDFIndex)
+			{
+				Tile.Index = BSDFIndex;
+				RWBSDFTileTexture[Tile.TileCoord] = PackBSDFTile(Tile);
+			}
+		}
+		else
+		{
+			RWBSDFTileTexture[GroupId] = 0;
+		}
+	}
+}
+#endif // SHADER_BSDF_TILE
diff --git a/Engine/Shaders/Private/Strata/StrataTile.ush b/Engine/Shaders/Private/Strata/StrataTile.ush
--- a/Engine/Shaders/Private/Strata/StrataTile.ush
+++ b/Engine/Shaders/Private/Strata/StrataTile.ush
@@ -4,9 +4,10 @@
 
 #include "/Engine/Private/Common.ush"
 #include "/Engine/Shared/StrataDefinitions.h"   
+#include "Strata.ush"
 
 ////////////////////////////////////////////////////////////////////////////////////////////////////////////
-// Tile
+// Strata material type tile
 
 uint2 StrataUnpackTile(uint In)
 {
@@ -17,3 +18,111 @@ uint StrataPackTile(uint2 TileCoord)
 {
 	return TileCoord.x | (TileCoord.y << 16); // assumes 16bit is enough to represent a tiled resolution up to 65,535 :)
 }
+
+////////////////////////////////////////////////////////////////////////////////////////////////////////////
+// Strata BSDF tile
+
+// For tile whose conatins at least 1 pixel with several BSDF, we store a list of 'overflow' tiles
+// 
+// For handling material with several BSDF per pixel we extend the original texture (primary space) 
+// with an overflow space. Pixels with more than one BSDF data will write BSDF 2/3/... data into 
+// the overflow space
+//  ________________
+// |				|
+// | Primary space	|
+// |________________|
+// | Overflow space |
+// |________________|
+//
+// The primary & overflow spaces are divided into tiles, and one primary-space tile refers a list 
+// of tiles in overflow-space, if the tile contains at least one pixel with has more than 1 BSDF
+// If several tiles are needed, they are contigously stored in overflow space
+//  ________________
+// |_:_:_:_:_:_:_:_:|
+// |_:_:_:_:_:_:_:_:| Primary space
+// |_:_:_:_:_:_:_:_:|
+//  ----------------
+// |_:_:_:_:_:_:	| Overflow space
+// |________________|
+
+struct FStrataBSDFTile
+{
+	uint2 TileCoord; // If Index == 0, points to the first 'overflow' tiles. If Index != 0, points towards the 'parent'/primary space tile
+	uint  Index;
+	uint  Count;
+};
+
+uint PackBSDFTile(FStrataBSDFTile In)
+{
+	// Tile coord are encoded onto 10bits. Max resolution 8k x 8k. Max 8 tiles.
+	return (In.TileCoord.x & 0x3FF) | ((In.TileCoord.y & 0x3FF) << 10) | ((In.Index & 0x7) << 20) | ((In.Index & 0x7) << 23);
+}
+
+FStrataBSDFTile UnpackBSDFTile(uint In)
+{
+	FStrataBSDFTile Out;
+	Out.TileCoord.x = In & 0x3FF;
+	Out.TileCoord.y = (In >> 10u) & 0x3FF;
+	Out.Index = (In >> 20u) & 0x7;
+	Out.Count = (In >> 23u) & 0x7;
+	return Out;
+}
+
+////////////////////////////////////////////////////////////////////////////////////////////////////////////
+// Strata BSDF offset
+
+// Store the material slice index offset for each BSDF, for pixels containing several BSDFs
+struct FStrataBSDFOffset
+{
+	uint BSDFOffsets[STRATA_MAX_BSDF_COUNT];
+	uint BSDFCount;
+};
+
+uint PackBSDFOffset(FStrataBSDFOffset In)
+{
+	// Packed as relative cumulative offsets, stored into 32 bits
+	// | 4 bits | 2 bits | 2 bits | 2 bits | ... | 2 bits |
+	//   BSDF 0   BSDF 1   BSDF 2   BSDF 3  ...   BSDF 15 (STRATA_MAX_BSDF_COUNT)
+
+	uint Out = 0;
+	const uint BSDFCount = min(In.BSDFCount, STRATA_MAX_BSDF_COUNT);
+	if (BSDFCount > 0)
+	{
+		Out = In.BSDFOffsets[0] & 0xF;
+		UNROLL
+		for (uint i = 1; i < BSDFCount; i++)
+		{
+			uint Curr = (In.BSDFOffsets[i  ] & 0x3);
+			uint Prev = (In.BSDFOffsets[i-1] & 0x3);
+			Out = Out | (((Curr - Prev) & 0x3) << (2u*i));
+		}
+	}
+	return Out;
+}
+
+FStrataBSDFOffset UnpackBSDFOffset(uint In, uint InBSDFCount)
+{
+	FStrataBSDFOffset Out = (FStrataBSDFOffset)0;
+	if (InBSDFCount > 0)
+	{
+		Out.BSDFCount = min(InBSDFCount, STRATA_MAX_BSDF_COUNT);
+		Out.BSDFOffsets[0] = (In & 0xF);
+
+		FStrataBSDFOffset Out = (FStrataBSDFOffset)0;
+		UNROLL
+		for (uint i = 1; i < Out.BSDFCount; i++)
+		{
+			Out.BSDFOffsets[i] = Out.BSDFOffsets[i-1] + ((In>>(2u*i)) & 0x3);
+		}
+	}
+	return Out;
+}
+
+void StrataSeekBSDF(inout FStrataAddressing Addressing, FStrataBSDFOffset Offsets, uint BSDFIndex)
+{
+	if (BSDFIndex < Offsets.BSDFCount)
+	{
+		Addressing.CurrentIndex = Offsets.BSDFOffsets[BSDFIndex];
+		Addressing.ReadBytes = Addressing.CurrentIndex * 4;
+	}
+}
diff --git a/Engine/Shaders/Shared/StrataDefinitions.h b/Engine/Shaders/Shared/StrataDefinitions.h
--- a/Engine/Shaders/Shared/StrataDefinitions.h
+++ b/Engine/Shaders/Shared/StrataDefinitions.h
@@ -64,3 +64,7 @@
 // sizeof(FRHIDrawIndirectParameters) = 4 uints = 16 bytes
 #define GetStrataTileTypeDrawIndirectArgOffset_Byte(x)  (x * 16)
 #define GetStrataTileTypeDrawIndirectArgOffset_DWord(x) (x * 4)
+
+// sizeof(FRHIDispatchIndirectParameters) = 3 uints = 12 bytes
+#define GetStrataTileTypeDispatchIndirectArgOffset_Byte(x)  (x * 12)
+#define GetStrataTileTypeDispatchIndirectArgOffset_DWord(x) (x * 3)
\ No newline at end of file
diff --git a/Engine/Source/Runtime/Renderer/Private/Strata/Strata.cpp b/Engine/Source/Runtime/Renderer/Private/Strata/Strata.cpp
--- a/Engine/Source/Runtime/Renderer/Private/Strata/Strata.cpp
+++ b/Engine/Source/Runtime/Renderer/Private/Strata/Strata.cpp
@@ -60,6 +60,11 @@ static TAutoConsoleVariable<int32> CVarClearDuringCategorization(
 	TEXT("TEST."),
 	ECVF_RenderThreadSafe);
 
+static TAutoConsoleVariable<int32> CVarStrataTileOverflow(
+	TEXT("r.Material.TileOverflow"),
+	1.f,
+	TEXT("Scale the number of Strata tile for overflowing tiles containing multi-BSDFs pixels. (0: 0%, 1: 100%. Default 1.0f)."),
+	ECVF_ReadOnly | ECVF_RenderThreadSafe);
 
 IMPLEMENT_GLOBAL_SHADER_PARAMETER_STRUCT(FStrataGlobalUniformParameters, "Strata");
 
@@ -77,8 +82,12 @@ void FStrataSceneData::Reset()
 	MaterialTextureArrayUAV = nullptr;
 	MaterialTextureArraySRV = nullptr;
 
-	ClassificationTileIndirectBuffer = nullptr;
-	ClassificationTileIndirectBufferUAV = nullptr;
+	ClassificationTileDrawIndirectBuffer = nullptr;
+	ClassificationTileDrawIndirectBufferUAV = nullptr;
+
+	ClassificationTileDispatchIndirectBuffer = nullptr;
+	ClassificationTileDispatchIndirectBufferUAV = nullptr;
+
 	for (uint32 i = 0; i < EStrataTileType::ECount; ++i)
 	{
 		ClassificationTileListBuffer[i] = nullptr;
@@ -89,6 +98,13 @@ void FStrataSceneData::Reset()
 	SeparatedSubSurfaceSceneColor = nullptr;
 	SeparatedOpaqueRoughRefractionSceneColor = nullptr;
 
+	BSDFOffsetTexture = nullptr;
+	BSDFTileTexture = nullptr;
+
+	TileCount_Total = FIntPoint(0,0);
+	TileCount_Primary = FIntPoint(0,0);
+	TileCount_Overflow = FIntPoint(0,0);
+
 	StrataGlobalUniformParameters = nullptr;
 }
 
@@ -126,11 +142,27 @@ bool IsStrataEnabled()
 	return CVarStrata.GetValueOnAnyThread() > 0;
 }
 
-uint32 GetStrataBufferTileSize()
+static FIntPoint GetStrataTextureTileResolution(const FIntPoint& InResolution, float Overflow)
+{
+	FIntPoint Out = InResolution;
+	if (Strata::IsStrataEnabled())
+	{
+		Out.X = FMath::DivideAndRoundUp(Out.X, STRATA_DATA_TILE_SIZE);
+		Out.Y = FMath::DivideAndRoundUp(Out.Y, STRATA_DATA_TILE_SIZE);
+		Out.Y += FMath::CeilToInt(Out.Y * FMath::Clamp(Overflow, 0.f, 1.0f));
+	}
+	return Out;
+}
+
+static FIntPoint GetStrataTextureTileResolution(const FIntPoint& InResolution)
 {
-	return 8;
+	return GetStrataTextureTileResolution(InResolution, CVarStrataTileOverflow.GetValueOnRenderThread());
 }
 
+FIntPoint GetStrataTextureResolution(const FIntPoint& InResolution)
+{
+	return GetStrataTextureTileResolution(InResolution) * STRATA_DATA_TILE_SIZE;
+}
 
 void InitialiseStrataFrameSceneData(FSceneRenderer& SceneRenderer, FRDGBuilder& GraphBuilder)
 {
@@ -157,8 +189,7 @@ void InitialiseStrataFrameSceneData(FSceneRenderer& SceneRenderer, FRDGBuilder&
 		const uint32 RoundToValue = 4u;
 		StrataSceneData.MaxBytesPerPixel = FMath::DivideAndRoundUp(MaterialConservativeByteCountPerPixel, RoundToValue) * RoundToValue;
 
-		const int32 TileInPixel = GetStrataBufferTileSize();
-		const FIntPoint TileResolution(FMath::DivideAndRoundUp(SceneTextureExtent.X, TileInPixel), FMath::DivideAndRoundUp(SceneTextureExtent.Y, TileInPixel));
+		const FIntPoint TileResolution(FMath::DivideAndRoundUp(SceneTextureExtent.X, STRATA_DATA_TILE_SIZE), FMath::DivideAndRoundUp(SceneTextureExtent.Y, STRATA_DATA_TILE_SIZE));
 
 		const TCHAR* StrataTileListBufferNames[EStrataTileType::ECount] =
 		{
@@ -171,9 +202,15 @@ void InitialiseStrataFrameSceneData(FSceneRenderer& SceneRenderer, FRDGBuilder&
 
 		// Tile classification buffers
 		{
-			StrataSceneData.ClassificationTileIndirectBuffer = GraphBuilder.CreateBuffer(FRDGBufferDesc::CreateIndirectDesc<FRHIDrawIndirectParameters>(EStrataTileType::ECount), TEXT("Strata.StrataTileIndirectBuffer"));
-			StrataSceneData.ClassificationTileIndirectBufferUAV = GraphBuilder.CreateUAV(StrataSceneData.ClassificationTileIndirectBuffer, PF_R32_UINT);
-			AddClearUAVPass(GraphBuilder, StrataSceneData.ClassificationTileIndirectBufferUAV, 0);
+			// Indirect draw
+			StrataSceneData.ClassificationTileDrawIndirectBuffer = GraphBuilder.CreateBuffer(FRDGBufferDesc::CreateIndirectDesc<FRHIDrawIndirectParameters>(EStrataTileType::ECount), TEXT("Strata.StrataTileDrawIndirectBuffer"));
+			StrataSceneData.ClassificationTileDrawIndirectBufferUAV = GraphBuilder.CreateUAV(StrataSceneData.ClassificationTileDrawIndirectBuffer, PF_R32_UINT);
+			AddClearUAVPass(GraphBuilder, StrataSceneData.ClassificationTileDrawIndirectBufferUAV, 0);
+
+			// Indirect dispatch
+			StrataSceneData.ClassificationTileDispatchIndirectBuffer = GraphBuilder.CreateBuffer(FRDGBufferDesc::CreateIndirectDesc<FRHIDispatchIndirectParameters>(EStrataTileType::ECount), TEXT("Strata.StrataTileDispatchIndirectBuffer"));
+			StrataSceneData.ClassificationTileDispatchIndirectBufferUAV = GraphBuilder.CreateUAV(StrataSceneData.ClassificationTileDispatchIndirectBuffer, PF_R32_UINT);
+			AddClearUAVPass(GraphBuilder, StrataSceneData.ClassificationTileDispatchIndirectBufferUAV, 0);
 
 			for (uint32 i = 0; i <= EStrataTileType::EComplex; ++i)
 			{
@@ -225,6 +262,20 @@ void InitialiseStrataFrameSceneData(FSceneRenderer& SceneRenderer, FRDGBuilder&
 				AddClearRenderTargetPass(GraphBuilder, StrataSceneData.SeparatedOpaqueRoughRefractionSceneColor, StrataSceneData.SeparatedOpaqueRoughRefractionSceneColor->Desc.ClearValue.GetClearColor());
 			}
 		}
+
+		// BSDF tiles/offsets
+		{
+			StrataSceneData.TileCount_Total		= GetStrataTextureTileResolution(SceneTextureExtent);
+			StrataSceneData.TileCount_Primary	= GetStrataTextureTileResolution(SceneTextureExtent, 0.f);
+			StrataSceneData.TileCount_Overflow	= StrataSceneData.TileCount_Total - StrataSceneData.TileCount_Primary;
+
+			StrataSceneData.BSDFTileTexture = GraphBuilder.CreateTexture(FRDGTextureDesc::Create2D(StrataSceneData.TileCount_Total, PF_R32_UINT, FClearValueBinding::None, TexCreate_UAV | TexCreate_ShaderResource), TEXT("Strata.BSDFTiles"));
+			AddClearUAVPass(GraphBuilder, GraphBuilder.CreateUAV(StrataSceneData.BSDFTileTexture), 0u);
+
+			StrataSceneData.BSDFOffsetTexture = GraphBuilder.CreateTexture(FRDGTextureDesc::Create2D(StrataSceneData.TileCount_Primary * STRATA_DATA_TILE_SIZE, PF_R32_UINT, FClearValueBinding::None, TexCreate_UAV | TexCreate_ShaderResource), TEXT("Strata.BSDFOffsets"));
+			AddClearUAVPass(GraphBuilder, GraphBuilder.CreateUAV(StrataSceneData.BSDFOffsetTexture), 0u);
+		}
+
 	}
 	else
 	{
@@ -310,10 +361,15 @@ void BindStrataGlobalUniformParameters(FRDGBuilder& GraphBuilder, FStrataSceneDa
 	{
 		OutStrataUniformParameters.bRoughDiffuse = StrataSceneData->bRoughDiffuse ? 1u : 0u;
 		OutStrataUniformParameters.MaxBytesPerPixel = StrataSceneData->MaxBytesPerPixel;
+		OutStrataUniformParameters.TileSize = STRATA_DATA_TILE_SIZE;
+		OutStrataUniformParameters.TileSizeLog2 = STRATA_DATA_TILE_SIZE_DIV_AS_SHIFT;
+		OutStrataUniformParameters.TileCount = StrataSceneData->TileCount_Primary;
 		OutStrataUniformParameters.MaterialTextureArray = StrataSceneData->MaterialTextureArray;
 		OutStrataUniformParameters.TopLayerTexture = StrataSceneData->TopLayerTexture;
 		OutStrataUniformParameters.SSSTexture = StrataSceneData->SSSTexture;
 		OutStrataUniformParameters.OpaqueRoughRefractionTexture = StrataSceneData->OpaqueRoughRefractionTexture;
+		OutStrataUniformParameters.BSDFTileTexture = StrataSceneData->BSDFTileTexture;
+		OutStrataUniformParameters.BSDFOffsetTexture = StrataSceneData->BSDFOffsetTexture;
 	}
 	else
 	{
@@ -321,10 +377,15 @@ void BindStrataGlobalUniformParameters(FRDGBuilder& GraphBuilder, FStrataSceneDa
 		FRDGTextureRef DefaultTextureArray = GSystemTextures.GetDefaultTexture(GraphBuilder, ETextureDimension::Texture2DArray, EPixelFormat::PF_R32_UINT, FClearValueBinding::Transparent);
 		OutStrataUniformParameters.bRoughDiffuse = 0;
 		OutStrataUniformParameters.MaxBytesPerPixel = 0;
+		OutStrataUniformParameters.TileSize = 0;
+		OutStrataUniformParameters.TileSizeLog2 = 0;
+		OutStrataUniformParameters.TileCount = 0;
 		OutStrataUniformParameters.MaterialTextureArray = DefaultTextureArray;
 		OutStrataUniformParameters.TopLayerTexture = SystemTextures.DefaultNormal8Bit;
 		OutStrataUniformParameters.SSSTexture = SystemTextures.Black;
 		OutStrataUniformParameters.OpaqueRoughRefractionTexture = SystemTextures.Black;
+		OutStrataUniformParameters.BSDFTileTexture = SystemTextures.Black;
+		OutStrataUniformParameters.BSDFOffsetTexture = SystemTextures.Black;
 	}
 }
 
@@ -456,6 +517,59 @@ IMPLEMENT_GLOBAL_SHADER(FStrataClearMaterialBufferCS, "/Engine/Private/Strata/St
 
 ////////////////////////////////////////////////////////////////////////////////////////////////////////////////////////////////////////////////
 
+class FStrataBSDFTilePassCS : public FGlobalShader
+{
+	DECLARE_GLOBAL_SHADER(FStrataBSDFTilePassCS);
+	SHADER_USE_PARAMETER_STRUCT(FStrataBSDFTilePassCS, FGlobalShader);
+
+	class FWaveOps : SHADER_PERMUTATION_BOOL("PERMUTATION_WAVE_OPS");
+	using FPermutationDomain = TShaderPermutationDomain<FWaveOps>;
+
+	BEGIN_SHADER_PARAMETER_STRUCT(FParameters, )
+		SHADER_PARAMETER_STRUCT_REF(FViewUniformShaderParameters, ViewUniformBuffer)
+		SHADER_PARAMETER(int32, bRectPrimitive)
+		SHADER_PARAMETER(int32, TileSizeLog2)
+		SHADER_PARAMETER(FIntPoint, TileCount_Primary)
+		SHADER_PARAMETER(FIntPoint, ViewResolution)
+		SHADER_PARAMETER(uint32, MaxBytesPerPixel)
+		SHADER_PARAMETER_RDG_TEXTURE(Texture2D, TopLayerTexture)
+		SHADER_PARAMETER_RDG_TEXTURE_SRV(Texture2DArray<uint>, MaterialTextureArray)
+
+		SHADER_PARAMETER_RDG_TEXTURE_UAV(RWTexture2D<uint>, RWBSDFTileTexture)
+		SHADER_PARAMETER_RDG_TEXTURE_UAV(RWTexture2D<uint>, RWBSDFOffsetTexture)
+		SHADER_PARAMETER_RDG_BUFFER_UAV(RWBuffer<uint>, RWBSDFTileCountBuffer)
+
+		SHADER_PARAMETER_RDG_BUFFER_SRV(Buffer<uint>, TileListBuffer)
+		RDG_BUFFER_ACCESS(TileIndirectBuffer, ERHIAccess::IndirectArgs)
+	END_SHADER_PARAMETER_STRUCT()
+
+	static bool ShouldCompilePermutation(const FGlobalShaderPermutationParameters& Parameters)
+	{
+		const bool bUseWaveIntrinsics = FDataDrivenShaderPlatformInfo::GetSupportsWaveOperations(Parameters.Platform);
+		FPermutationDomain PermutationVector(Parameters.PermutationId);
+		if (PermutationVector.Get<FWaveOps>() && !bUseWaveIntrinsics)
+		{
+			return false;
+		}
+		return GetMaxSupportedFeatureLevel(Parameters.Platform) >= ERHIFeatureLevel::SM5 && Strata::IsStrataEnabled();
+	}
+
+	static void ModifyCompilationEnvironment(const FGlobalShaderPermutationParameters& Parameters, FShaderCompilerEnvironment& OutEnvironment)
+	{
+		FGlobalShader::ModifyCompilationEnvironment(Parameters, OutEnvironment);
+		OutEnvironment.SetDefine(TEXT("SHADER_BSDF_TILE"), 1);
+
+		FPermutationDomain PermutationVector(Parameters.PermutationId);
+		if (PermutationVector.Get<FWaveOps>())
+		{
+			OutEnvironment.CompilerFlags.Add(CFLAG_WaveOperations);
+		}
+	}
+};
+IMPLEMENT_GLOBAL_SHADER(FStrataBSDFTilePassCS, "/Engine/Private/Strata/StrataMaterialClassification.usf", "BSDFTileMainCS", SF_Compute);
+
+////////////////////////////////////////////////////////////////////////////////////////////////////////////////////////////////////////////////
+
 class FStrataMaterialTileClassificationPassCS : public FGlobalShader
 {
 	DECLARE_GLOBAL_SHADER(FStrataMaterialTileClassificationPassCS);
@@ -472,7 +586,7 @@ class FStrataMaterialTileClassificationPassCS : public FGlobalShader
 		SHADER_PARAMETER(uint32, MaxBytesPerPixel)
 		SHADER_PARAMETER_RDG_TEXTURE(Texture2D, TopLayerTexture)
 		SHADER_PARAMETER_RDG_TEXTURE_SRV(Texture2DArray<uint>, MaterialTextureArray)
-		SHADER_PARAMETER_RDG_BUFFER_UAV(RWBuffer, TileIndirectDataBuffer)
+		SHADER_PARAMETER_RDG_BUFFER_UAV(RWBuffer, TileDrawIndirectDataBuffer)
 		SHADER_PARAMETER_RDG_BUFFER_UAV(RWBuffer, SimpleTileListDataBuffer)
 		SHADER_PARAMETER_RDG_BUFFER_UAV(RWBuffer, SingleTileListDataBuffer)
 		SHADER_PARAMETER_RDG_BUFFER_UAV(RWBuffer, ComplexTileListDataBuffer)
@@ -509,6 +623,31 @@ IMPLEMENT_GLOBAL_SHADER(FStrataMaterialTileClassificationPassCS, "/Engine/Privat
 
 ////////////////////////////////////////////////////////////////////////////////////////////////////////////////////////////////////////////////
 
+class FStrataMaterialTilePrepareArgsPassCS : public FGlobalShader
+{
+	DECLARE_GLOBAL_SHADER(FStrataMaterialTilePrepareArgsPassCS);
+	SHADER_USE_PARAMETER_STRUCT(FStrataMaterialTilePrepareArgsPassCS, FGlobalShader);
+
+	BEGIN_SHADER_PARAMETER_STRUCT(FParameters, )
+		SHADER_PARAMETER_RDG_BUFFER_SRV(Buffer,   TileDrawIndirectDataBuffer)
+		SHADER_PARAMETER_RDG_BUFFER_UAV(RWBuffer, TileDispatchIndirectDataBuffer)
+	END_SHADER_PARAMETER_STRUCT()
+
+	static bool ShouldCompilePermutation(const FGlobalShaderPermutationParameters& Parameters)
+	{
+		return GetMaxSupportedFeatureLevel(Parameters.Platform) >= ERHIFeatureLevel::SM5 && Strata::IsStrataEnabled();
+	}
+
+	static void ModifyCompilationEnvironment(const FGlobalShaderPermutationParameters& Parameters, FShaderCompilerEnvironment& OutEnvironment)
+	{
+		FGlobalShader::ModifyCompilationEnvironment(Parameters, OutEnvironment);
+		OutEnvironment.SetDefine(TEXT("SHADER_TILE_PREPARE_ARGS"), 1);
+	}
+};
+IMPLEMENT_GLOBAL_SHADER(FStrataMaterialTilePrepareArgsPassCS, "/Engine/Private/Strata/StrataMaterialClassification.usf", "ArgsMainCS", SF_Compute);
+
+////////////////////////////////////////////////////////////////////////////////////////////////////////////////////////////////////////////////
+
 class FStrataMaterialStencilTaggingPassPS : public FGlobalShader
 {
 	DECLARE_GLOBAL_SHADER(FStrataMaterialStencilTaggingPassPS);
@@ -548,7 +687,7 @@ static FStrataTileParameter InternalSetTileParameters(FRDGBuilder* GraphBuilder,
 	if (View.StrataSceneData && TileType != EStrataTileType::ECount)
 	{
 		Out.TileListBuffer = View.StrataSceneData->ClassificationTileListBufferSRV[TileType];
-		Out.TileIndirectBuffer = View.StrataSceneData->ClassificationTileIndirectBuffer;
+		Out.TileIndirectBuffer = View.StrataSceneData->ClassificationTileDrawIndirectBuffer;
 	}
 	else if (GraphBuilder)
 	{
@@ -606,6 +745,12 @@ uint32 TileTypeDrawIndirectArgOffset(const EStrataTileType Type)
 	return GetStrataTileTypeDrawIndirectArgOffset_Byte(Type);
 }
 
+uint32 TileTypeDispatchIndirectArgOffset(const EStrataTileType Type)
+{
+	check(Type >= 0 && Type < EStrataTileType::ECount);
+	return GetStrataTileTypeDispatchIndirectArgOffset_Byte(Type);
+}
+
 // Add additionnaly bits for filling/clearing stencil to ensure that the 'Strata' bits are not corrupted by the stencil shadows 
 // when generating shadow mask. Withouth these 'trailing' bits, the incr./decr. operation would change/corrupt the 'Strata' bits
 constexpr uint32 StencilBit_Fast_1	  = 0x07u | StencilBit_Fast;
@@ -806,14 +951,14 @@ void AddStrataMaterialClassificationPass(FRDGBuilder& GraphBuilder, const FMinim
 	for (int32 i = 0; i < Views.Num(); ++i)
 	{
 		const FViewInfo& View = Views[i];
+		bool bWaveOps = GRHISupportsWaveOperations && FDataDrivenShaderPlatformInfo::GetSupportsWaveOperations(View.GetShaderPlatform());
+	#if PLATFORM_WINDOWS
+		// Tile reduction requires 64-wide wave
+		bWaveOps = bWaveOps && !IsRHIDeviceNVIDIA();
+	#endif
 		
 		// Tile reduction
 		{
-			bool bWaveOps = GRHISupportsWaveOperations && FDataDrivenShaderPlatformInfo::GetSupportsWaveOperations(View.GetShaderPlatform());
-		#if PLATFORM_WINDOWS
-			// Tile reduction requires 64-wide wave
-			bWaveOps = bWaveOps && !IsRHIDeviceNVIDIA();
-		#endif
 			const bool bClear = ClearDuringCategorization();
 			FStrataMaterialTileClassificationPassCS::FPermutationDomain PermutationVector;
 			PermutationVector.Set< FStrataMaterialTileClassificationPassCS::FStrataClearDuringCategorization >(bClear);
@@ -828,7 +973,7 @@ void AddStrataMaterialClassificationPass(FRDGBuilder& GraphBuilder, const FMinim
 			PassParameters->MaterialTextureArray = View.StrataSceneData->MaterialTextureArraySRV;
 			PassParameters->SSSTextureUAV = View.StrataSceneData->SSSTextureUAV;
 			PassParameters->OpaqueRoughRefractionTexture = View.StrataSceneData->OpaqueRoughRefractionTexture;
-			PassParameters->TileIndirectDataBuffer = View.StrataSceneData->ClassificationTileIndirectBufferUAV;
+			PassParameters->TileDrawIndirectDataBuffer = View.StrataSceneData->ClassificationTileDrawIndirectBufferUAV;
 			PassParameters->SimpleTileListDataBuffer = View.StrataSceneData->ClassificationTileListBufferUAV[EStrataTileType::ESimple];
 			PassParameters->SingleTileListDataBuffer = View.StrataSceneData->ClassificationTileListBufferUAV[EStrataTileType::ESingle];
 			PassParameters->ComplexTileListDataBuffer = View.StrataSceneData->ClassificationTileListBufferUAV[EStrataTileType::EComplex];
@@ -843,6 +988,54 @@ void AddStrataMaterialClassificationPass(FRDGBuilder& GraphBuilder, const FMinim
 				PassParameters,
 				FComputeShaderUtils::GetGroupCount(PassParameters->ViewResolution, GroupSize));
 		}
+
+		// Tile indirect dispatch args conversion
+		{
+			TShaderMapRef<FStrataMaterialTilePrepareArgsPassCS> ComputeShader(View.ShaderMap);
+			FStrataMaterialTilePrepareArgsPassCS::FParameters* PassParameters = GraphBuilder.AllocParameters<FStrataMaterialTilePrepareArgsPassCS::FParameters>();
+			PassParameters->TileDrawIndirectDataBuffer = GraphBuilder.CreateSRV(View.StrataSceneData->ClassificationTileDrawIndirectBuffer, PF_R32_UINT);
+			PassParameters->TileDispatchIndirectDataBuffer = View.StrataSceneData->ClassificationTileDispatchIndirectBufferUAV;
+
+			FComputeShaderUtils::AddPass(
+				GraphBuilder,
+				RDG_EVENT_NAME("Strata::MaterialTilePrepareArgs"),
+				ComputeShader,
+				PassParameters,
+				FIntVector(1,1,1));
+		}
+
+		// Compute BSDF tile index and material read offset
+		{
+			FRDGBufferRef BSDFTileCountBuffer = GraphBuilder.CreateBuffer(FRDGBufferDesc::CreateBufferDesc(4, 1), TEXT("Strata.BSDFTileCount"));
+			FRDGBufferUAVRef RWBSDFTileCountBuffer = GraphBuilder.CreateUAV(BSDFTileCountBuffer, PF_R32_UINT);
+			AddClearUAVPass(GraphBuilder, RWBSDFTileCountBuffer, 0u);
+
+			FStrataBSDFTilePassCS::FPermutationDomain PermutationVector;
+			PermutationVector.Set< FStrataBSDFTilePassCS::FWaveOps >(bWaveOps);
+			TShaderMapRef<FStrataBSDFTilePassCS> ComputeShader(View.ShaderMap, PermutationVector);
+			FStrataBSDFTilePassCS::FParameters* PassParameters = GraphBuilder.AllocParameters<FStrataBSDFTilePassCS::FParameters>();
+			PassParameters->ViewUniformBuffer = View.ViewUniformBuffer;
+			PassParameters->TileSizeLog2 = STRATA_DATA_TILE_SIZE_DIV_AS_SHIFT;
+			PassParameters->TileCount_Primary = View.StrataSceneData->TileCount_Primary;
+			PassParameters->ViewResolution = View.ViewRect.Size();
+			PassParameters->MaxBytesPerPixel = View.StrataSceneData->MaxBytesPerPixel;
+			PassParameters->TopLayerTexture = View.StrataSceneData->TopLayerTexture;
+			PassParameters->MaterialTextureArray = View.StrataSceneData->MaterialTextureArraySRV;
+			PassParameters->TileListBuffer = View.StrataSceneData->ClassificationTileListBufferSRV[EStrataTileType::EComplex];
+			PassParameters->TileIndirectBuffer = View.StrataSceneData->ClassificationTileDispatchIndirectBuffer;
+
+			PassParameters->RWBSDFOffsetTexture = GraphBuilder.CreateUAV(View.StrataSceneData->BSDFOffsetTexture);
+			PassParameters->RWBSDFTileTexture = GraphBuilder.CreateUAV(View.StrataSceneData->BSDFTileTexture);
+			PassParameters->RWBSDFTileCountBuffer = RWBSDFTileCountBuffer;
+
+			FComputeShaderUtils::AddPass(
+				GraphBuilder,
+				RDG_EVENT_NAME("Strata::BSDFTileAndOffsets(%s)", bWaveOps ? TEXT("Wave") : TEXT("SharedMemory")),
+				ComputeShader,
+				PassParameters,
+				PassParameters->TileIndirectBuffer,
+				TileTypeDispatchIndirectArgOffset(EStrataTileType::EComplex));
+		}
 	}
 }
 
diff --git a/Engine/Source/Runtime/Renderer/Private/Strata/Strata.h b/Engine/Source/Runtime/Renderer/Private/Strata/Strata.h
--- a/Engine/Source/Runtime/Renderer/Private/Strata/Strata.h
+++ b/Engine/Source/Runtime/Renderer/Private/Strata/Strata.h
@@ -38,10 +38,15 @@ END_SHADER_PARAMETER_STRUCT()
 BEGIN_GLOBAL_SHADER_PARAMETER_STRUCT(FStrataGlobalUniformParameters, RENDERER_API)
 	SHADER_PARAMETER(uint32, MaxBytesPerPixel)
 	SHADER_PARAMETER(uint32, bRoughDiffuse)
+	SHADER_PARAMETER(uint32, TileSize)
+	SHADER_PARAMETER(uint32, TileSizeLog2)
+	SHADER_PARAMETER(FIntPoint, TileCount)
 	SHADER_PARAMETER_RDG_TEXTURE(Texture2DArray<uint>, MaterialTextureArray)
 	SHADER_PARAMETER_RDG_TEXTURE(Texture2D<uint>, TopLayerTexture)
 	SHADER_PARAMETER_RDG_TEXTURE(Texture2D<uint2>, SSSTexture)
 	SHADER_PARAMETER_RDG_TEXTURE(Texture2D<float3>, OpaqueRoughRefractionTexture)
+	SHADER_PARAMETER_RDG_TEXTURE(Texture2D<uint>, BSDFOffsetTexture)
+	SHADER_PARAMETER_RDG_TEXTURE(Texture2D<uint>, BSDFTileTexture)
 END_GLOBAL_SHADER_PARAMETER_STRUCT()
 
 // This must map to the STRATA_TILE_TYPE defines.
@@ -62,6 +67,10 @@ struct FStrataSceneData
 	uint32 MaxBytesPerPixel;
 	bool bRoughDiffuse;
 
+	FIntPoint TileCount_Primary;
+	FIntPoint TileCount_Overflow;
+	FIntPoint TileCount_Total;
+
 	// Resources allocated and updated each frame
 
 	FRDGTextureRef MaterialTextureArray;
@@ -73,8 +82,11 @@ struct FStrataSceneData
 	FRDGBufferSRVRef ClassificationTileListBufferSRV[STRATA_TILE_TYPE_COUNT];
 	FRDGBufferUAVRef ClassificationTileListBufferUAV[STRATA_TILE_TYPE_COUNT];
 
-	FRDGBufferRef ClassificationTileIndirectBuffer;
-	FRDGBufferUAVRef ClassificationTileIndirectBufferUAV;
+	FRDGBufferRef    ClassificationTileDrawIndirectBuffer;
+	FRDGBufferUAVRef ClassificationTileDrawIndirectBufferUAV;
+
+	FRDGBufferRef    ClassificationTileDispatchIndirectBuffer;
+	FRDGBufferUAVRef ClassificationTileDispatchIndirectBufferUAV;
 
 	FRDGTextureRef TopLayerTexture;
 	FRDGTextureRef SSSTexture;
@@ -84,6 +96,9 @@ struct FStrataSceneData
 	FRDGTextureUAVRef SSSTextureUAV;
 	FRDGTextureUAVRef OpaqueRoughRefractionTextureUAV;
 
+	FRDGTextureRef BSDFOffsetTexture;
+	FRDGTextureRef BSDFTileTexture;
+
 	// Used when the subsurface luminance is separated from the scene color
 	FRDGTextureRef SeparatedSubSurfaceSceneColor;
 
@@ -109,6 +124,8 @@ constexpr uint32 StencilBit_Complex= 0x20; // In sync with SceneRenderTargets.h
 
 bool IsStrataEnabled();
 
+FIntPoint GetStrataTextureResolution(const FIntPoint& InResolution);
+
 void InitialiseStrataFrameSceneData(FSceneRenderer& SceneRenderer, FRDGBuilder& GraphBuilder);
 
 void BindStrataBasePassUniformParameters(FRDGBuilder& GraphBuilder, FStrataSceneData* StrataSceneData, FStrataBasePassUniformParameters& OutStrataUniformParameters);
@@ -169,6 +186,7 @@ FStrataTileParameter SetTileParameters(FRDGBuilder& GraphBuilder, const FViewInf
 FStrataTilePassVS::FParameters SetTileParameters(FRDGBuilder& GraphBuilder, const FViewInfo& View, const EStrataTileType Type, EPrimitiveType& PrimitiveType);
 FStrataTilePassVS::FParameters SetTileParameters(const FViewInfo& View, const EStrataTileType Type, EPrimitiveType& PrimitiveType);
 uint32 TileTypeDrawIndirectArgOffset(const EStrataTileType Type);
+uint32 TileTypeDispatchIndirectArgOffset(const EStrataTileType Type);
 
 bool ShouldRenderStrataRoughRefractionRnD();
 void StrataRoughRefractionRnD(FRDGBuilder& GraphBuilder, const FViewInfo& View, FScreenPassTexture& ScreenPassSceneColor);
