Remove AActor::DetachFence reducing size of AActor by 16 bytes.
It's not necessary to prevent the actor being deleted before the component finishes its destruction because in GC all objects must complete FinishDestroy before any objects can be freed, and UPrimitiveComponent's fence will ensure both objects live long enough.

diff --git a/Engine/Source/Runtime/Engine/Classes/GameFramework/Actor.h b/Engine/Source/Runtime/Engine/Classes/GameFramework/Actor.h
--- a/Engine/Source/Runtime/Engine/Classes/GameFramework/Actor.h
+++ b/Engine/Source/Runtime/Engine/Classes/GameFramework/Actor.h
@@ -1975,7 +1975,9 @@ public:
 	virtual void PostLoad() override;
 	virtual void PostLoadSubobjects( FObjectInstancingGraph* OuterInstanceGraph ) override;
 	virtual void BeginDestroy() override;
+#if !UE_STRIP_DEPRECATED_PROPERTIES
 	virtual bool IsReadyForFinishDestroy() override;
+#endif
 	virtual bool Rename( const TCHAR* NewName=nullptr, UObject* NewOuter=nullptr, ERenameFlags Flags=REN_None ) override;
 	virtual void PostRename( UObject* OldOuter, const FName OldName ) override;
 	virtual bool CanBeInCluster() const override;
@@ -3750,8 +3752,11 @@ public:
 	static FOnProcessEvent ProcessEventDelegate;
 #endif
 
+#if !UE_STRIP_DEPRECATED_PROPERTIES
 	/** A fence to track when the primitive is detached from the scene in the rendering thread. */
+	UE_DEPRECATED(5.1, "AActor::DetachFence has been deprecated. If you are relying on it for render thread synchronization in a subclass of actor, add your own fence to that class instead.")
 	FRenderCommandFence DetachFence;
+#endif
 
 private:
 	/** Helper that already assumes the Hit info is reversed, and avoids creating a temp FHitResult if possible. */
diff --git a/Engine/Source/Runtime/Engine/Private/Actor.cpp b/Engine/Source/Runtime/Engine/Private/Actor.cpp
--- a/Engine/Source/Runtime/Engine/Private/Actor.cpp
+++ b/Engine/Source/Runtime/Engine/Private/Actor.cpp
@@ -759,10 +759,14 @@ void AActor::BeginDestroy()
 	Super::BeginDestroy();
 }
 
+#if !UE_STRIP_DEPRECATED_PROPERTIES
 bool AActor::IsReadyForFinishDestroy()
 {
+PRAGMA_DISABLE_DEPRECATION_WARNINGS
 	return Super::IsReadyForFinishDestroy() && DetachFence.IsFenceComplete();
+PRAGMA_ENABLE_DEPRECATION_WARNINGS
 }
+#endif
 
 void AActor::Serialize(FArchive& Ar)
 {
diff --git a/Engine/Source/Runtime/Engine/Private/Components/PrimitiveComponent.cpp b/Engine/Source/Runtime/Engine/Private/Components/PrimitiveComponent.cpp
--- a/Engine/Source/Runtime/Engine/Private/Components/PrimitiveComponent.cpp
+++ b/Engine/Source/Runtime/Engine/Private/Components/PrimitiveComponent.cpp
@@ -1396,11 +1396,16 @@ void UPrimitiveComponent::BeginDestroy()
 
 	// Use a fence to keep track of when the rendering thread executes this scene detachment.
 	DetachFence.BeginFence();
+	
+#if !UE_STRIP_DEPRECATED_PROPERTIES
 	AActor* Owner = GetOwner();
 	if(Owner)
 	{
+PRAGMA_DISABLE_DEPRECATION_WARNINGS
 		Owner->DetachFence.BeginFence();
+PRAGMA_ENABLE_DEPRECATION_WARNINGS
 	}
+#endif
 }
 
 void UPrimitiveComponent::OnComponentDestroyed(bool bDestroyingHierarchy)
diff --git a/Engine/Source/Runtime/Renderer/Public/PrimitiveSceneInfo.h b/Engine/Source/Runtime/Renderer/Public/PrimitiveSceneInfo.h
--- a/Engine/Source/Runtime/Renderer/Public/PrimitiveSceneInfo.h
+++ b/Engine/Source/Runtime/Renderer/Public/PrimitiveSceneInfo.h
@@ -302,7 +302,11 @@ public:
 	/** 
 	 * Pointer to the last render time variable on the primitive's owning actor (if owned), which is written to by the RT and read by the GT.
 	 * The value of LastRenderTime will therefore not be deterministic due to race conditions, but the GT uses it in a way that allows this.
-	 * Storing a pointer to the UObject member variable only works because UPrimitiveComponent and AActor has a mechanism to ensure it does not get deleted before the proxy (DetachFence).
+	 * Storing a pointer to the UObject member variable only works because:
+	 *	UPrimitiveComponent's outer is its owning AActor, so it prevents the owner from being garbage collected while the component lives.
+	 *  If the UPrimitiveComponent is GC'd during the Actor's lifetime, OwnerLastRenderTime is still valid so there is no issue.
+	 *	If the UPrimitiveComponent and the Actor are GC'd together, neither will be deleted until FinishDestroy has been executed on both.
+	 *	UPrimitiveComponent's FinishDestroy will not execute until the primitive has been detached from the Scene through it's DetachFence.
 	 * In general feedback from the renderer to the game thread like this should be avoided.
 	 */
 	float* OwnerLastRenderTime;
