Add indirect dispatch for Strata overflow tiles

diff --git a/Engine/Shaders/Private/Strata/StrataMaterialClassification.usf b/Engine/Shaders/Private/Strata/StrataMaterialClassification.usf
--- a/Engine/Shaders/Private/Strata/StrataMaterialClassification.usf
+++ b/Engine/Shaders/Private/Strata/StrataMaterialClassification.usf
@@ -190,7 +190,7 @@ void TileMainCS(uint2 DispatchThreadId : SV_DispatchThreadID, uint LinearIndex :
 
 ////////////////////////////////////////////////////////////////////////////////////////////////////////////
 
-#if SHADER_TILE_PREPARE_ARGS
+#if SHADER_MATERIAL_TILE_PREPARE_ARGS
 
 Buffer<uint>   TileDrawIndirectDataBuffer;
 RWBuffer<uint> TileDispatchIndirectDataBuffer;
@@ -220,7 +220,29 @@ void ArgsMainCS(uint2 DispatchThreadId : SV_DispatchThreadID)
 	}
 }
 
-#endif // SHADER_TILE_PREPARE_ARGS
+#endif // SHADER_MATERIAL_TILE_PREPARE_ARGS
+
+////////////////////////////////////////////////////////////////////////////////////////////////////////////
+
+#if SHADER_BSDF_TILE_PREPARE_ARGS
+
+uint2 TileCount_Primary;
+Buffer<uint>   TileDrawIndirectDataBuffer;
+RWBuffer<uint> TileDispatchIndirectDataBuffer;
+
+[numthreads(1, 1, 1)]
+void ArgsMainCS(uint2 DispatchThreadId : SV_DispatchThreadID)
+{
+	if (all(DispatchThreadId == 0))
+	{
+		const uint TileCount = TileDrawIndirectDataBuffer[0].x;
+		TileDispatchIndirectDataBuffer[0] = TileCount >= TileCount_Primary.x ? TileCount_Primary.x : TileCount % TileCount_Primary.x;
+		TileDispatchIndirectDataBuffer[1] = (TileCount + TileCount_Primary.x-1) / TileCount_Primary.x;
+		TileDispatchIndirectDataBuffer[2] = 1;
+	}
+}
+
+#endif // SHADER_BSDF_TILE_PREPARE_ARGS
 
 ////////////////////////////////////////////////////////////////////////////////////////////////////////////
 
@@ -277,6 +299,11 @@ RWBuffer<uint>    RWBSDFTileCountBuffer;
 groupshared uint s_TileBSDFCount[GROUP_THREAD_COUNT];
 #endif
 
+uint2 WrapCoord(uint InLinear)
+{
+	return uint2(InLinear % TileCount_Primary.x, InLinear / TileCount_Primary.x + TileCount_Primary.y);
+}
+
 [numthreads(STRATA_TILE_SIZE, STRATA_TILE_SIZE, 1)]
 void BSDFTileMainCS(uint2 GroupThreadId : SV_GroupThreadID, uint2 GroupId : SV_GroupID, uint LinearIndex : SV_GroupIndex)
 {
@@ -357,16 +384,17 @@ void BSDFTileMainCS(uint2 GroupThreadId : SV_GroupThreadID, uint2 GroupId : SV_G
 			Tile.Index = 0;
 
 			// For first/parent tile, set TileCoord to point towards the first 'overflow' tile
-			Tile.TileCoord = uint2(LinearTileCoord % TileCount_Primary.x, TileCount_Primary.y + LinearTileCoord / TileCount_Primary.y);
+			Tile.TileCoord = WrapCoord(LinearTileCoord);
 			RWBSDFTileTexture[TileCoord] = PackBSDFTile(Tile);
 
 			// For 'overflow' tiles, set TileCoord to point towards the first/parent tile
 			Tile.TileCoord = TileCoord;
-
 			for (uint BSDFIndex = 1; BSDFIndex < TileBSDFCount; ++BSDFIndex)
 			{
+				const uint2 OverflowCoord = WrapCoord(LinearTileCoord + BSDFIndex - 1);
+
 				Tile.Index = BSDFIndex;
-				RWBSDFTileTexture[Tile.TileCoord] = PackBSDFTile(Tile);
+				RWBSDFTileTexture[OverflowCoord] = PackBSDFTile(Tile);
 			}
 		}
 		else
diff --git a/Engine/Shaders/Private/Strata/StrataTile.ush b/Engine/Shaders/Private/Strata/StrataTile.ush
--- a/Engine/Shaders/Private/Strata/StrataTile.ush
+++ b/Engine/Shaders/Private/Strata/StrataTile.ush
@@ -57,7 +57,7 @@ struct FStrataBSDFTile
 uint PackBSDFTile(FStrataBSDFTile In)
 {
 	// Tile coord are encoded onto 10bits. Max resolution 8k x 8k. Max 8 tiles.
-	return (In.TileCoord.x & 0x3FF) | ((In.TileCoord.y & 0x3FF) << 10) | ((In.Index & 0x7) << 20) | ((In.Index & 0x7) << 23);
+	return (In.TileCoord.x & 0x3FF) | ((In.TileCoord.y & 0x3FF) << 10) | ((In.Index & 0x7) << 20) | ((In.Count & 0x7) << 23);
 }
 
 FStrataBSDFTile UnpackBSDFTile(uint In)
diff --git a/Engine/Source/Runtime/Renderer/Private/Strata/Strata.cpp b/Engine/Source/Runtime/Renderer/Private/Strata/Strata.cpp
--- a/Engine/Source/Runtime/Renderer/Private/Strata/Strata.cpp
+++ b/Engine/Source/Runtime/Renderer/Private/Strata/Strata.cpp
@@ -274,8 +274,9 @@ void InitialiseStrataFrameSceneData(FSceneRenderer& SceneRenderer, FRDGBuilder&
 
 			StrataSceneData.BSDFOffsetTexture = GraphBuilder.CreateTexture(FRDGTextureDesc::Create2D(StrataSceneData.TileCount_Primary * STRATA_DATA_TILE_SIZE, PF_R32_UINT, FClearValueBinding::None, TexCreate_UAV | TexCreate_ShaderResource), TEXT("Strata.BSDFOffsets"));
 			AddClearUAVPass(GraphBuilder, GraphBuilder.CreateUAV(StrataSceneData.BSDFOffsetTexture), 0u);
-		}
 
+			StrataSceneData.BSDFTileDispatchIndirectBuffer = GraphBuilder.CreateBuffer(FRDGBufferDesc::CreateIndirectDesc<FRHIDispatchIndirectParameters>(1), TEXT("Strata.StrataBSDFTileDispatchIndirectBuffer"));
+		}
 	}
 	else
 	{
@@ -641,13 +642,39 @@ class FStrataMaterialTilePrepareArgsPassCS : public FGlobalShader
 	static void ModifyCompilationEnvironment(const FGlobalShaderPermutationParameters& Parameters, FShaderCompilerEnvironment& OutEnvironment)
 	{
 		FGlobalShader::ModifyCompilationEnvironment(Parameters, OutEnvironment);
-		OutEnvironment.SetDefine(TEXT("SHADER_TILE_PREPARE_ARGS"), 1);
+		OutEnvironment.SetDefine(TEXT("SHADER_MATERIAL_TILE_PREPARE_ARGS"), 1);
 	}
 };
 IMPLEMENT_GLOBAL_SHADER(FStrataMaterialTilePrepareArgsPassCS, "/Engine/Private/Strata/StrataMaterialClassification.usf", "ArgsMainCS", SF_Compute);
 
 ////////////////////////////////////////////////////////////////////////////////////////////////////////////////////////////////////////////////
 
+class FStrataBSDFTilePrepareArgsPassCS : public FGlobalShader
+{
+	DECLARE_GLOBAL_SHADER(FStrataBSDFTilePrepareArgsPassCS);
+	SHADER_USE_PARAMETER_STRUCT(FStrataBSDFTilePrepareArgsPassCS, FGlobalShader);
+
+	BEGIN_SHADER_PARAMETER_STRUCT(FParameters, )
+		SHADER_PARAMETER(FIntPoint, TileCount_Primary)
+		SHADER_PARAMETER_RDG_BUFFER_SRV(Buffer, TileDrawIndirectDataBuffer)
+		SHADER_PARAMETER_RDG_BUFFER_UAV(RWBuffer, TileDispatchIndirectDataBuffer)
+		END_SHADER_PARAMETER_STRUCT()
+
+		static bool ShouldCompilePermutation(const FGlobalShaderPermutationParameters& Parameters)
+	{
+		return GetMaxSupportedFeatureLevel(Parameters.Platform) >= ERHIFeatureLevel::SM5 && Strata::IsStrataEnabled();
+	}
+
+	static void ModifyCompilationEnvironment(const FGlobalShaderPermutationParameters& Parameters, FShaderCompilerEnvironment& OutEnvironment)
+	{
+		FGlobalShader::ModifyCompilationEnvironment(Parameters, OutEnvironment);
+		OutEnvironment.SetDefine(TEXT("SHADER_BSDF_TILE_PREPARE_ARGS"), 1);
+	}
+};
+IMPLEMENT_GLOBAL_SHADER(FStrataBSDFTilePrepareArgsPassCS, "/Engine/Private/Strata/StrataMaterialClassification.usf", "ArgsMainCS", SF_Compute);
+
+////////////////////////////////////////////////////////////////////////////////////////////////////////////////////////////////////////////////
+
 class FStrataMaterialStencilTaggingPassPS : public FGlobalShader
 {
 	DECLARE_GLOBAL_SHADER(FStrataMaterialStencilTaggingPassPS);
@@ -1005,8 +1032,8 @@ void AddStrataMaterialClassificationPass(FRDGBuilder& GraphBuilder, const FMinim
 		}
 
 		// Compute BSDF tile index and material read offset
+		FRDGBufferRef BSDFTileCountBuffer = BSDFTileCountBuffer = GraphBuilder.CreateBuffer(FRDGBufferDesc::CreateBufferDesc(4, 1), TEXT("Strata.BSDFTileCount"));
 		{
-			FRDGBufferRef BSDFTileCountBuffer = GraphBuilder.CreateBuffer(FRDGBufferDesc::CreateBufferDesc(4, 1), TEXT("Strata.BSDFTileCount"));
 			FRDGBufferUAVRef RWBSDFTileCountBuffer = GraphBuilder.CreateUAV(BSDFTileCountBuffer, PF_R32_UINT);
 			AddClearUAVPass(GraphBuilder, RWBSDFTileCountBuffer, 0u);
 
@@ -1036,6 +1063,22 @@ void AddStrataMaterialClassificationPass(FRDGBuilder& GraphBuilder, const FMinim
 				PassParameters->TileIndirectBuffer,
 				TileTypeDispatchIndirectArgOffset(EStrataTileType::EComplex));
 		}
+
+		// Tile indirect dispatch args conversion
+		{
+			TShaderMapRef<FStrataBSDFTilePrepareArgsPassCS> ComputeShader(View.ShaderMap);
+			FStrataBSDFTilePrepareArgsPassCS::FParameters* PassParameters = GraphBuilder.AllocParameters<FStrataBSDFTilePrepareArgsPassCS::FParameters>();
+			PassParameters->TileCount_Primary = View.StrataSceneData->TileCount_Primary;
+			PassParameters->TileDrawIndirectDataBuffer = GraphBuilder.CreateSRV(BSDFTileCountBuffer, PF_R32_UINT);
+			PassParameters->TileDispatchIndirectDataBuffer = GraphBuilder.CreateUAV(View.StrataSceneData->BSDFTileDispatchIndirectBuffer, PF_R32_UINT);
+
+			FComputeShaderUtils::AddPass(
+				GraphBuilder,
+				RDG_EVENT_NAME("Strata::BSDFTilePrepareArgs"),
+				ComputeShader,
+				PassParameters,
+				FIntVector(1, 1, 1));
+		}
 	}
 }
 
diff --git a/Engine/Source/Runtime/Renderer/Private/Strata/Strata.h b/Engine/Source/Runtime/Renderer/Private/Strata/Strata.h
--- a/Engine/Source/Runtime/Renderer/Private/Strata/Strata.h
+++ b/Engine/Source/Runtime/Renderer/Private/Strata/Strata.h
@@ -98,6 +98,7 @@ struct FStrataSceneData
 
 	FRDGTextureRef BSDFOffsetTexture;
 	FRDGTextureRef BSDFTileTexture;
+	FRDGBufferRef  BSDFTileDispatchIndirectBuffer;
 
 	// Used when the subsurface luminance is separated from the scene color
 	FRDGTextureRef SeparatedSubSurfaceSceneColor;
