Moved forward lights on screen user feedback to lambda.

diff --git a/Engine/Source/Runtime/Renderer/Private/LightGridInjection.cpp b/Engine/Source/Runtime/Renderer/Private/LightGridInjection.cpp
--- a/Engine/Source/Runtime/Renderer/Private/LightGridInjection.cpp
+++ b/Engine/Source/Runtime/Renderer/Private/LightGridInjection.cpp
@@ -272,6 +272,9 @@ void FSceneRenderer::ComputeLightGrid(FRDGBuilder& GraphBuilder, bool bCullLight
 	const FRDGSystemTextures& SystemTextures = FRDGSystemTextures::Get(GraphBuilder);
 
 	TArray<FForwardLightData*, TInlineAllocator<4>> ForwardLightDataPerView;
+#if WITH_EDITOR
+	bool bMultipleDirLightsConflictForForwardShading = false;
+#endif
 
 	for (int32 ViewIndex = 0; ViewIndex < Views.Num(); ViewIndex++)
 	{
@@ -744,6 +747,17 @@ void FSceneRenderer::ComputeLightGrid(FRDGBuilder& GraphBuilder, bool bCullLight
 			View.ForwardLightingResources.SetUniformBuffer(GraphBuilder.CreateUniformBuffer(ForwardLightData));
 		}
 	}
+
+#if WITH_EDITOR
+	if (bMultipleDirLightsConflictForForwardShading)
+	{
+		OnGetOnScreenMessages.AddLambda([](FScreenMessageWriter& ScreenMessageWriter)->void
+		{
+			static const FText Message = NSLOCTEXT("Renderer", "MultipleDirLightsConflictForForwardShading", "Multiple directional lights are competing to be the single one used for forward shading, translucent, water or volumetric fog. Please adjust their ForwardShadingPriority.\nAs a fallback, the main directional light will be selected based on overall brightness.");
+			ScreenMessageWriter.DrawLine(Message, 10, FColor::Orange);
+		});
+	}
+#endif
 }
 
 void FDeferredShadingSceneRenderer::GatherLightsAndComputeLightGrid(FRDGBuilder& GraphBuilder, bool bNeedLightGrid, FSortedLightSetSceneInfo &SortedLightSet)
diff --git a/Engine/Source/Runtime/Renderer/Private/SceneRendering.cpp b/Engine/Source/Runtime/Renderer/Private/SceneRendering.cpp
--- a/Engine/Source/Runtime/Renderer/Private/SceneRendering.cpp
+++ b/Engine/Source/Runtime/Renderer/Private/SceneRendering.cpp
@@ -2456,9 +2456,6 @@ FSceneRenderer::FSceneRenderer(const FSceneViewFamily* InViewFamily,FHitProxyCon
 ,	bHasRequestedToggleFreeze(false)
 ,	bUsedPrecomputedVisibility(false)
 ,	InstancedStereoWidth(0)
-#if WITH_EDITOR
-,	bMultipleDirLightsConflictForForwardShading(false)
-#endif
 ,	FamilySize(0, 0)
 ,	GPUSceneDynamicContext(CheckPointer(Scene)->GPUScene)
 ,	bShadowDepthRenderCompleted(false)
@@ -3270,9 +3267,6 @@ void FSceneRenderer::RenderFinish(FRDGBuilder& GraphBuilder, FRDGTextureRef View
 			|| bMobileMissingSkyMaterial || bShowSkinCacheOOM || bSingleLayerWaterWarning || bShowDFDisabledWarning || bShowNoSkyAtmosphereComponentWarning || bFxDebugDraw 
 			|| bLumenEnabledButHasNoDataForTracing || bLumenEnabledButDisabledForTheProject || bNaniteEnabledButNoAtomics || bNaniteEnabledButDisabledInProject || bRealTimeSkyCaptureButNothingToCapture || bShowWaitingSkylight
 			|| bShowLocalExposureDisabledWarning || bHasDelegateWarnings
-#if WITH_EDITOR
-			|| bMultipleDirLightsConflictForForwardShading
-#endif
 			;
 
 		for(int32 ViewIndex = 0;ViewIndex < Views.Num();ViewIndex++)
@@ -3469,12 +3463,6 @@ void FSceneRenderer::RenderFinish(FRDGBuilder& GraphBuilder, FRDGTextureRef View
 						}
 
 #if WITH_EDITOR
-						if (bMultipleDirLightsConflictForForwardShading)
-						{
-							static const FText Message = NSLOCTEXT("Renderer", "MultipleDirLightsConflictForForwardShading", "Multiple directional lights are competing to be the single one used for forward shading, translucent, water or volumetric fog. Please adjust their ForwardShadingPriority.\nAs a fallback, the main directional light will be selected based on overall brightness.");
-							Writer.DrawLine(Message, 10, FColor::Orange);
-						}
-
 						FSkyLightSceneProxy* SkyLight = Scene->SkyLight;
 						if (bShowWaitingSkylight && SkyLight)
 						{
diff --git a/Engine/Source/Runtime/Renderer/Private/SceneRendering.h b/Engine/Source/Runtime/Renderer/Private/SceneRendering.h
--- a/Engine/Source/Runtime/Renderer/Private/SceneRendering.h
+++ b/Engine/Source/Runtime/Renderer/Private/SceneRendering.h
@@ -1890,13 +1890,6 @@ public:
 	 */
 	uint32 InstancedStereoWidth;
 
-#if WITH_EDITOR
-	/**
-	 * True when multiple direcitonal lights are in competiation for the single slot available for forward shading (many have the same highest ForwardShadingPriority)
-	 */
-	bool bMultipleDirLightsConflictForForwardShading;
-#endif
-
 public:
 
 	FSceneRenderer(const FSceneViewFamily* InViewFamily,FHitProxyConsumer* HitProxyConsumer);
