Fixed near clip for Nanite Non-VSM cube shadows (remove initialization that clobbered default NANITE_VIEW_FLAG_NEAR_CLIP flag) .

#rnx

diff --git a/Engine/Source/Runtime/Renderer/Private/ShadowDepthRendering.cpp b/Engine/Source/Runtime/Renderer/Private/ShadowDepthRendering.cpp
--- a/Engine/Source/Runtime/Renderer/Private/ShadowDepthRendering.cpp
+++ b/Engine/Source/Runtime/Renderer/Private/ShadowDepthRendering.cpp
@@ -1959,7 +1959,6 @@ void FSceneRenderer::RenderShadowDepthMaps(FRDGBuilder& GraphBuilder, FInstanceC
 						Params.LODScaleFactor = ComputeNaniteShadowsLODScaleFactor();
 						Params.PrevViewMatrices = Params.ViewMatrices;
 						Params.HZBTestViewRect = ShadowViewRect;
-						Params.Flags = 0;
 						UpdatePackedViewParamsFromPrevShadowState(Params, PrevShadowState);
 
 						PackedViews.Add(Nanite::CreatePackedView(Params));
