Fix missing sky lighting when RT sky is enabled.

diff --git a/Engine/Source/Runtime/Renderer/Private/HairStrands/HairStrandsEnvironment.cpp b/Engine/Source/Runtime/Renderer/Private/HairStrands/HairStrandsEnvironment.cpp
--- a/Engine/Source/Runtime/Renderer/Private/HairStrands/HairStrandsEnvironment.cpp
+++ b/Engine/Source/Runtime/Renderer/Private/HairStrands/HairStrandsEnvironment.cpp
@@ -340,7 +340,7 @@ static void AddHairStrandsEnvironmentLightingPassPS(
 	}
 
 	// The specular sky light contribution is also needed by RT Reflections as a fallback.
-	const bool bSkyLight = Scene->SkyLight && Scene->SkyLight->ProcessedTexture && !Scene->SkyLight->bHasStaticLighting;
+	const bool bSkyLight = Scene->SkyLight && !Scene->SkyLight->bHasStaticLighting;
 
 	const bool bDynamicSkyLight = ShouldRenderDeferredDynamicSkyLight(Scene, *View.Family);
 	const bool bReflectionEnv = ShouldDoReflectionEnvironment(Scene, *View.Family);
