Add color modution over sky atmosphere contribution to height fog.
https://horde.devtools.epicgames.com/job/6238c92389625f0612a3cc38

diff --git a/Engine/Shaders/Private/HeightFogCommon.ush b/Engine/Shaders/Private/HeightFogCommon.ush
--- a/Engine/Shaders/Private/HeightFogCommon.ush
+++ b/Engine/Shaders/Private/HeightFogCommon.ush
@@ -60,7 +60,7 @@ float3 ComputeInscatteringColor(float3 CameraToReceiver, float CameraToReceiverL
 #endif
 
 #if PROJECT_SUPPORT_SKY_ATMOSPHERE_AFFECTS_HEIGHFOG
-	Inscattering += View.SkyAtmosphereHeightFogContribution * Texture2DSampleLevel(View.DistantSkyLightLutTexture, View.DistantSkyLightLutTextureSampler, float2(0.5f, 0.5f), 0.0f).rgb;
+	Inscattering += FogStruct.SkyAtmosphereAmbientContributionColorScale * View.SkyAtmosphereHeightFogContribution * Texture2DSampleLevel(View.DistantSkyLightLutTexture, View.DistantSkyLightLutTextureSampler, float2(0.5f, 0.5f), 0.0f).rgb;
 #endif
 
 	return Inscattering;
diff --git a/Engine/Source/Runtime/Engine/Classes/Components/ExponentialHeightFogComponent.h b/Engine/Source/Runtime/Engine/Classes/Components/ExponentialHeightFogComponent.h
--- a/Engine/Source/Runtime/Engine/Classes/Components/ExponentialHeightFogComponent.h
+++ b/Engine/Source/Runtime/Engine/Classes/Components/ExponentialHeightFogComponent.h
@@ -69,6 +69,10 @@ class ENGINE_API UExponentialHeightFogComponent : public USceneComponent
 	UPROPERTY(BlueprintReadOnly, interp, Category=ExponentialHeightFogComponent, meta = (DisplayName = "Fog Inscattering Color"))
 	FLinearColor FogInscatteringLuminance;
 
+	/** Color used to modulate the SkyAtmosphere component contribution to the non directional component of the fog. Only effective when r.SupportSkyAtmosphereAffectsHeightFog>0 */
+	UPROPERTY(EditAnywhere, BlueprintReadOnly, Category = ExponentialHeightFogComponent)
+	FLinearColor SkyAtmosphereAmbientContributionColorScale;
+
 	/** 
 	 * Cubemap that can be specified for fog color, which is useful to make distant, heavily fogged scene elements match the sky.
 	 * When the cubemap is specified, FogInscatteringColor is ignored and Directional inscattering is disabled. 
diff --git a/Engine/Source/Runtime/Engine/Private/Components/HeightFogComponent.cpp b/Engine/Source/Runtime/Engine/Private/Components/HeightFogComponent.cpp
--- a/Engine/Source/Runtime/Engine/Private/Components/HeightFogComponent.cpp
+++ b/Engine/Source/Runtime/Engine/Private/Components/HeightFogComponent.cpp
@@ -23,6 +23,8 @@ UExponentialHeightFogComponent::UExponentialHeightFogComponent(const FObjectInit
 	FogInscatteringColor_DEPRECATED = FLinearColor(0.447f, 0.638f, 1.0f);
 	FogInscatteringLuminance = FLinearColor::Black;
 
+	SkyAtmosphereAmbientContributionColorScale = FLinearColor::White;
+
 	DirectionalInscatteringExponent = 4.0f;
 	DirectionalInscatteringStartDistance = 10000.0f;
 	DirectionalInscatteringColor_DEPRECATED = FLinearColor(0.25f, 0.25f, 0.125f);
@@ -104,6 +106,12 @@ bool UExponentialHeightFogComponent::CanEditChange(const FProperty* InProperty)
 		{
 			return InscatteringColorCubemap != NULL;
 		}
+
+		if (PropertyName == GET_MEMBER_NAME_STRING_CHECKED(UExponentialHeightFogComponent, FogInscatteringLuminance))
+		{
+			static const auto CVar = IConsoleManager::Get().FindTConsoleVariableDataInt(TEXT("r.SupportSkyAtmosphereAffectsHeightFog"));
+			return CVar && CVar->GetValueOnAnyThread() > 0;
+		}
 	}
 
 	return Super::CanEditChange(InProperty);
diff --git a/Engine/Source/Runtime/Renderer/Private/FogRendering.cpp b/Engine/Source/Runtime/Renderer/Private/FogRendering.cpp
--- a/Engine/Source/Runtime/Renderer/Private/FogRendering.cpp
+++ b/Engine/Source/Runtime/Renderer/Private/FogRendering.cpp
@@ -68,6 +68,7 @@ void SetupFogUniformParameters(FRDGBuilder& GraphBuilder, const FViewInfo& View,
 		OutParameters.ExponentialFogColorParameter = FVector4f(View.ExponentialFogColor, 1.0f - View.FogMaxOpacity);
 		OutParameters.ExponentialFogParameters2 = View.ExponentialFogParameters2;
 		OutParameters.ExponentialFogParameters3 = View.ExponentialFogParameters3;
+		OutParameters.SkyAtmosphereAmbientContributionColorScale = View.SkyAtmosphereAmbientContributionColorScale;
 		OutParameters.SinCosInscatteringColorCubemapRotation = View.SinCosInscatteringColorCubemapRotation;
 		OutParameters.FogInscatteringTextureParameters = (FVector3f)View.FogInscatteringTextureParameters;
 		OutParameters.InscatteringLightDirection = (FVector3f)View.InscatteringLightDirection;
@@ -237,6 +238,8 @@ void FSceneRenderer::InitFogConstants()
 				const float InvRange = 1.0f / FMath::Max(FogInfo.FullyDirectionalInscatteringColorDistance - FogInfo.NonDirectionalInscatteringColorDistance, .00001f);
 				float NumMips = 1.0f;
 
+				View.SkyAtmosphereAmbientContributionColorScale = FogInfo.SkyAtmosphereAmbientContributionColorScale;
+
 				if (FogInfo.InscatteringColorCubemap)
 				{
 					NumMips = FogInfo.InscatteringColorCubemap->GetNumMips();
diff --git a/Engine/Source/Runtime/Renderer/Private/FogRendering.h b/Engine/Source/Runtime/Renderer/Private/FogRendering.h
--- a/Engine/Source/Runtime/Renderer/Private/FogRendering.h
+++ b/Engine/Source/Runtime/Renderer/Private/FogRendering.h
@@ -16,6 +16,7 @@ BEGIN_GLOBAL_SHADER_PARAMETER_STRUCT(FFogUniformParameters,)
 	SHADER_PARAMETER(FVector4f, ExponentialFogParameters2)
 	SHADER_PARAMETER(FVector4f, ExponentialFogColorParameter)
 	SHADER_PARAMETER(FVector4f, ExponentialFogParameters3)
+	SHADER_PARAMETER(FVector4f, SkyAtmosphereAmbientContributionColorScale)
 	SHADER_PARAMETER(FVector4f, InscatteringLightDirection) // non negative DirectionalInscatteringStartDistance in .W
 	SHADER_PARAMETER(FVector4f, DirectionalInscatteringColor)
 	SHADER_PARAMETER(FVector2f, SinCosInscatteringColorCubemapRotation)
diff --git a/Engine/Source/Runtime/Renderer/Private/SceneCore.cpp b/Engine/Source/Runtime/Renderer/Private/SceneCore.cpp
--- a/Engine/Source/Runtime/Renderer/Private/SceneCore.cpp
+++ b/Engine/Source/Runtime/Renderer/Private/SceneCore.cpp
@@ -421,4 +421,6 @@ FExponentialHeightFogSceneInfo::FExponentialHeightFogSceneInfo(const UExponentia
 	VolumetricFogDistance = FMath::Max(InComponent->VolumetricFogDistance, 0.0f);
 	VolumetricFogStaticLightingScatteringIntensity = FMath::Max(InComponent->VolumetricFogStaticLightingScatteringIntensity, 0.0f);
 	bOverrideLightColorsWithFogInscatteringColors = InComponent->bOverrideLightColorsWithFogInscatteringColors;
+
+	SkyAtmosphereAmbientContributionColorScale = InComponent->SkyAtmosphereAmbientContributionColorScale;
 }
diff --git a/Engine/Source/Runtime/Renderer/Private/SceneCore.h b/Engine/Source/Runtime/Renderer/Private/SceneCore.h
--- a/Engine/Source/Runtime/Renderer/Private/SceneCore.h
+++ b/Engine/Source/Runtime/Renderer/Private/SceneCore.h
@@ -288,6 +288,8 @@ public:
 	float VolumetricFogStaticLightingScatteringIntensity;
 	bool bOverrideLightColorsWithFogInscatteringColors;
 
+	FLinearColor SkyAtmosphereAmbientContributionColorScale;
+
 	/** Initialization constructor. */
 	FExponentialHeightFogSceneInfo(const UExponentialHeightFogComponent* InComponent);
 };
diff --git a/Engine/Source/Runtime/Renderer/Private/SceneRendering.h b/Engine/Source/Runtime/Renderer/Private/SceneRendering.h
--- a/Engine/Source/Runtime/Renderer/Private/SceneRendering.h
+++ b/Engine/Source/Runtime/Renderer/Private/SceneRendering.h
@@ -1307,6 +1307,7 @@ public:
 	FVector3f ExponentialFogColor;
 	float FogMaxOpacity;
 	FVector4f ExponentialFogParameters3;
+	FVector4f SkyAtmosphereAmbientContributionColorScale;
 	FVector2f SinCosInscatteringColorCubemapRotation;
 
 	UTexture* FogInscatteringColorCubemap;
