Strata - rough refraction - Add read back of generated gpu data for analysis.

diff --git a/Engine/Shaders/Private/Strata/StrataRoughRefraction.usf b/Engine/Shaders/Private/Strata/StrataRoughRefraction.usf
--- a/Engine/Shaders/Private/Strata/StrataRoughRefraction.usf
+++ b/Engine/Shaders/Private/Strata/StrataRoughRefraction.usf
@@ -126,6 +126,41 @@ bool GetGGXDirection(float2 RandomUV, float Roughness, float3 V, inout float3 L,
 
 
 
+struct FLobeStat
+{
+	float Count;
+	float2 Mean;
+	float2 M2;
+	float2 Variance;
+	float pad;
+};
+
+// See Welford's algorithm https://en.wikipedia.org/wiki/Algorithms_for_calculating_variance
+void LobStatUpdate(inout FLobeStat LobeStat, in float2 InPos)
+{
+	LobeStat.Count += 1;
+	float2 Delta = InPos - LobeStat.Mean;
+	LobeStat.Mean += Delta / LobeStat.Count;
+	float2 Delta2 = InPos - LobeStat.Mean;
+	LobeStat.M2 += Delta * Delta2;
+}
+
+void LobStatFinalize(inout FLobeStat LobeStat)
+{
+	if (LobeStat.Count < 2)
+	{
+		// nan
+	}
+	else
+	{
+		// Mean is already valid
+		LobeStat.Variance = LobeStat.M2 / LobeStat.Count;
+		//LobeStat.SampleVariance = LobeStat.M2 / (LobeStat.Count - 1);
+	}
+}
+
+
+
 #if defined(EVALUATE_ROUGH_REFRACTION_LOBE_CS) || defined(VISUALIZE_ROUGH_REFRACTION_PS)
 
 struct FUI
@@ -193,39 +228,6 @@ void DrawCircleXY(inout FShaderPrintContext Context, float3 Center, float2 Radiu
 	}
 }
 
-struct FLobeStat
-{
-	float Count;
-	float2 Mean;
-	float2 M2;
-	float2 Variance;
-	float pad;
-};
-
-// See Welford's algorithm https://en.wikipedia.org/wiki/Algorithms_for_calculating_variance
-void LobStatUpdate(inout FLobeStat LobeStat, in float2 InPos)
-{
-	LobeStat.Count += 1;
-	float2 Delta = InPos - LobeStat.Mean;
-	LobeStat.Mean += Delta / LobeStat.Count;
-	float2 Delta2 = InPos - LobeStat.Mean;
-	LobeStat.M2 += Delta * Delta2;
-}
-
-void LobStatFinalize(inout FLobeStat LobeStat)
-{
-	if (LobeStat.Count < 2)
-	{
-		// nan
-	}
-	else
-	{
-		// Mean is already valid
-		LobeStat.Variance = LobeStat.M2 / LobeStat.Count;
-		//LobeStat.SampleVariance = LobeStat.M2 / (LobeStat.Count - 1);
-	}
-}
-
 #endif
 
 
@@ -306,6 +308,45 @@ void EvaluateRoughRefractionLobeCS(uint3 ThreadId : SV_DispatchThreadID)
 
 
 
+#ifdef EVALUATE_ROUGH_REFRACTION_DATA_CS
+
+uint TraceSqrtSampleCount;
+
+RWBuffer<float2> RoughRefracDataUAV;
+
+[numthreads(THREADGROUP_SIZE, 1, 1)]
+void RoughRefracDataCS(uint3 ThreadId : SV_DispatchThreadID)
+{
+	const float Roughness = float(ThreadId.x) / float(THREADGROUP_SIZE);
+
+	FLobeStat LobeState = (FLobeStat)0;
+
+	for (float u = 0.5f; u < float(TraceSqrtSampleCount); u++)
+	{
+		for (float v = 0.5f; v < float(TraceSqrtSampleCount); v++)
+		{
+			float2 RandomUV = float2(u, v) / float(TraceSqrtSampleCount);
+			float3 V = float3(0, 0, 1);	// Only considering perpendicular view for now
+			float Eta12 = 1.33;			// Only considering air to water for now
+
+			float3 L = 0;
+			float Value = 0;
+			if (GetGGXDirection(RandomUV, Roughness, V, L, Value, Eta12) && Value > 0.0f)
+			{
+				LobStatUpdate(LobeState, L.xy);
+			}
+		}
+	}
+
+	LobStatFinalize(LobeState);
+
+	RoughRefracDataUAV[ThreadId.x] = float2(Roughness, dot(LobeState.Variance, float2(0.5, 0.5)));	// Roughness => Mean of both X and Y variances
+}
+
+#endif // EVALUATE_ROUGH_REFRACTION_DATA_CS
+
+
+
 #ifdef VISUALIZE_ROUGH_REFRACTION_PS
 
 float TraceDomainSize;
diff --git a/Engine/Source/Runtime/Renderer/Private/Strata/StrataRoughRefraction.cpp b/Engine/Source/Runtime/Renderer/Private/Strata/StrataRoughRefraction.cpp
--- a/Engine/Source/Runtime/Renderer/Private/Strata/StrataRoughRefraction.cpp
+++ b/Engine/Source/Runtime/Renderer/Private/Strata/StrataRoughRefraction.cpp
@@ -258,7 +258,11 @@ void AddStrataOpaqueRoughRefractionPasses(
 //////////////////////////////////////////////////////////////////////////
 
 // Keeping it simple: this should always be checked in with a value of 0
-#define STRATA_ROUGH_REFRACTION_RND 0
+#define STRATA_ROUGH_REFRACTION_RND 1
+
+#include "HAL/PlatformFileManager.h"
+#include "HAL/PlatformMisc.h"
+#include "Misc/FileHelper.h"
 
 #if STRATA_ROUGH_REFRACTION_RND
 
@@ -361,6 +365,45 @@ IMPLEMENT_GLOBAL_SHADER(FVisualizeRoughRefractionPS, "/Engine/Private/Strata/Str
 
 
 
+class FRoughRefracDataCS : public FGlobalShader
+{
+	DECLARE_GLOBAL_SHADER(FRoughRefracDataCS);
+	SHADER_USE_PARAMETER_STRUCT(FRoughRefracDataCS, FGlobalShader);
+
+	using FPermutationDomain = TShaderPermutationDomain<>;
+
+	BEGIN_SHADER_PARAMETER_STRUCT(FParameters, )
+		SHADER_PARAMETER_RDG_BUFFER_UAV(RWBuffer<float2>, RoughRefracDataUAV)
+		SHADER_PARAMETER(uint32, TraceSqrtSampleCount)
+	END_SHADER_PARAMETER_STRUCT()
+
+	BEGIN_SHADER_PARAMETER_STRUCT(FBufferReadBackParam, )
+		RDG_BUFFER_ACCESS(Buffer, ERHIAccess::CopySrc)
+	END_SHADER_PARAMETER_STRUCT()
+
+	static const uint32 ThreadGroupSize = 64;
+
+	static FPermutationDomain RemapPermutation(FPermutationDomain PermutationVector)
+	{
+		return PermutationVector;
+	}
+
+	static bool ShouldCompilePermutation(const FGlobalShaderPermutationParameters& Parameters)
+	{
+		return true;
+	}
+
+	static void ModifyCompilationEnvironment(const FGlobalShaderPermutationParameters& Parameters, FShaderCompilerEnvironment& OutEnvironment)
+	{
+		FGlobalShader::ModifyCompilationEnvironment(Parameters, OutEnvironment);
+
+		OutEnvironment.SetDefine(TEXT("STRATA_RND_SHADERS"), 1);
+		OutEnvironment.SetDefine(TEXT("EVALUATE_ROUGH_REFRACTION_DATA_CS"), 1);
+		OutEnvironment.SetDefine(TEXT("THREADGROUP_SIZE"), ThreadGroupSize);
+	}
+};
+IMPLEMENT_GLOBAL_SHADER(FRoughRefracDataCS, "/Engine/Private/Strata/StrataRoughRefraction.usf", "RoughRefracDataCS", SF_Compute);
+
 #endif // STRATA_ROUGH_REFRACTION_RND
 
 
@@ -445,6 +488,102 @@ void StrataRoughRefractionRnD(FRDGBuilder& GraphBuilder, const FViewInfo& View,
 			FPixelShaderUtils::AddFullscreenPass<FVisualizeRoughRefractionPS>(GraphBuilder, View.ShaderMap, RDG_EVENT_NAME("Strata::VisualizeRoughRefractionPS"),
 				PixelShader, PassParameters, View.ViewRect, PreMultipliedColorTransmittanceBlend);
 		}
+
+		//////////////////////////////////////////////////////////////////////////
+		// Sample variance as a function of roughness
+		{
+			const uint32 RoughRefracDataByteSize = sizeof(float) * 2 * FRoughRefracDataCS::ThreadGroupSize;
+			FRDGBufferRef RoughRefracDataBuffer = GraphBuilder.CreateBuffer(FRDGBufferDesc::CreateStructuredDesc(sizeof(float) * 2, FRoughRefracDataCS::ThreadGroupSize), TEXT("Strata.RoughRefrac.LobStat"));
+			FRDGBufferUAVRef RoughRefracDataBufferUAV = GraphBuilder.CreateUAV(RoughRefracDataBuffer, PF_G32R32F);
+
+			// Generate data in a buffer
+			{
+				FRoughRefracDataCS::FParameters* PassParameters = GraphBuilder.AllocParameters<FRoughRefracDataCS::FParameters>();
+				PassParameters->RoughRefracDataUAV = RoughRefracDataBufferUAV;
+				PassParameters->TraceSqrtSampleCount = 128;
+
+				FRoughRefracDataCS::FPermutationDomain PermutationVector;
+				TShaderMapRef<FRoughRefracDataCS> ComputeShader(View.ShaderMap, PermutationVector);
+
+				FComputeShaderUtils::AddPass(
+					GraphBuilder,
+					RDG_EVENT_NAME("Strata::FRoughRefracDataCS"),
+					ComputeShader,
+					PassParameters,
+					FComputeShaderUtils::GetGroupCount(FRoughRefracDataCS::ThreadGroupSize, FRoughRefracDataCS::ThreadGroupSize));
+			}
+
+			// Read back the data on CPU
+			{
+				FRoughRefracDataCS::FBufferReadBackParam* PassParameters = GraphBuilder.AllocParameters<FRoughRefracDataCS::FBufferReadBackParam>();
+				PassParameters->Buffer = RoughRefracDataBuffer;
+
+				GraphBuilder.AddPass(
+					RDG_EVENT_NAME("Strata::FRoughRefracData ReadBack"),
+					PassParameters,
+					ERDGPassFlags::Readback,
+					[RoughRefracDataBuffer, RoughRefracDataByteSize](FRHICommandListImmediate& RHICmdList)
+				{
+					check(IsInRenderingThread());
+					FStagingBufferRHIRef StagingBuffer;
+					StagingBuffer = RHICreateStagingBuffer();
+					// Copy memory from GPU to CPU
+					RHICmdList.CopyToStagingBuffer(RoughRefracDataBuffer->GetRHI(), StagingBuffer, 0, RoughRefracDataByteSize);
+
+					// Submit all GPU work so far wait for completion.
+					static const FName FenceName(TEXT("DumpGPU.BufferFence"));
+					FGPUFenceRHIRef Fence = RHICreateGPUFence(FenceName);
+					Fence->Clear();
+					RHICmdList.WriteGPUFence(Fence);
+					RHICmdList.SubmitCommandsAndFlushGPU();
+					RHICmdList.BlockUntilGPUIdle();
+
+					// Lock the buffer for read back
+					void* RoughRefracData = RHICmdList.LockStagingBuffer(StagingBuffer, Fence.GetReference(), 0, RoughRefracDataByteSize);
+					if (RoughRefracData)
+					{
+						// Dump to file
+						RHICmdList.UnlockStagingBuffer(StagingBuffer);
+
+						float* RoughRefracDataFloat = (float*)RoughRefracData;
+
+						static bool bDataDumpDone = false;
+						if (!bDataDumpDone)
+						{
+							bDataDumpDone = true;
+
+							IPlatformFile& PlatformFile = FPlatformFileManager::Get().GetPlatformFile();
+							FString Filename = TEXT("RoughRefracData.txt");
+							FString FullPath = FPaths::ConvertRelativePathToFull(FPaths::ProjectDir() / Filename);
+
+							FString OutString;
+							const uint32 PointCount = FRoughRefracDataCS::ThreadGroupSize;
+							for (int i = 0; i < PointCount; ++i)
+							{
+								OutString.Appendf(TEXT("%5.5f %5.5f\n"), RoughRefracDataFloat[i * 2 + 0], RoughRefracDataFloat[i * 2 + 1]);
+							}
+
+							FFileHelper::SaveStringToFile( OutString, *FullPath, FFileHelper::EEncodingOptions::AutoDetect, &IFileManager::Get(), FILEWRITE_None);
+
+							// Example of commands to find a fit in Mathematica
+							//	dataRV = Import["D:\\...\\RoughRefracData.txt", "Table"]
+							//	L0 = ListPlot[dataRV]
+							//	fitModel = With[{order = 4, dat = dataRV}, LinearModelFit[dat, Flatten@Outer[Times, Sequence @@	Transpose@Array[Power[{x}, #   - 1] &, order + 1]], { x }]]
+							//	fitModel[x]
+							//	Show[Plot[fitModel[x], {x, 0, 1}, PlotStyle -> Blue], ListPlot[dataRV, PlotStyle->Directive[Red, Opacity[0.5]]]]
+							//	fitModel[0.0]
+						}
+					}
+					//else
+					//{
+					//	check(false);
+					//}
+
+					StagingBuffer = nullptr;
+					Fence = nullptr;
+				});
+			}
+		}
 	}
 #endif
 }
