Strata - disable rough refraction rnd code.

diff --git a/Engine/Source/Runtime/Renderer/Private/Strata/StrataRoughRefraction.cpp b/Engine/Source/Runtime/Renderer/Private/Strata/StrataRoughRefraction.cpp
--- a/Engine/Source/Runtime/Renderer/Private/Strata/StrataRoughRefraction.cpp
+++ b/Engine/Source/Runtime/Renderer/Private/Strata/StrataRoughRefraction.cpp
@@ -258,7 +258,7 @@ void AddStrataOpaqueRoughRefractionPasses(
 //////////////////////////////////////////////////////////////////////////
 
 // Keeping it simple: this should always be checked in with a value of 0
-#define STRATA_ROUGH_REFRACTION_RND 1
+#define STRATA_ROUGH_REFRACTION_RND 0
 
 #include "HAL/PlatformFileManager.h"
 #include "HAL/PlatformMisc.h"
