Fix crash when shader print line buffer is never written.

diff --git a/Engine/Source/Runtime/Renderer/Private/ShaderPrint.cpp b/Engine/Source/Runtime/Renderer/Private/ShaderPrint.cpp
--- a/Engine/Source/Runtime/Renderer/Private/ShaderPrint.cpp
+++ b/Engine/Source/Runtime/Renderer/Private/ShaderPrint.cpp
@@ -552,7 +552,7 @@ namespace ShaderPrint
 			View.ShaderPrintData.UniformBuffer = CreateUniformBuffer(View.ShaderPrintData);
 			View.ShaderPrintData.ShaderPrintValueBuffer = GraphBuilder.RegisterExternalBuffer(GEmptyBuffer->Buffer);
 			View.ShaderPrintData.ShaderPrintStateBuffer = GraphBuilder.RegisterExternalBuffer(GEmptyBuffer->Buffer);
-			View.ShaderPrintData.ShaderPrintLineBuffer = GraphBuilder.CreateBuffer(FRDGBufferDesc::CreateStructuredDesc(4, 8), TEXT("ShaderPrint.LineBuffer(Dummy)"), ERDGBufferFlags::None);
+			View.ShaderPrintData.ShaderPrintLineBuffer  = GraphBuilder.RegisterExternalBuffer(GEmptyBuffer->Buffer);
 			return;
 		}
 
