Add OIT per-pixel sorting support for translucent material

This technique is guarded with by a projected settings. It is currently only supported on PC Dx12. This is the initial step for fine grain/per-pixel sorting.

diff --git a/Engine/Config/Windows/DataDrivenPlatformInfo.ini b/Engine/Config/Windows/DataDrivenPlatformInfo.ini
--- a/Engine/Config/Windows/DataDrivenPlatformInfo.ini
+++ b/Engine/Config/Windows/DataDrivenPlatformInfo.ini
@@ -88,6 +88,7 @@ bSupportsCapsuleShadows=true
 bSupportsPercentageCloserShadows=true
 bSupportsVolumetricFog=true
 bSupportsIndexBufferUAVs=true
+bSupportsOIT=true
 
 [PreviewPlatform PCD3D_SM5]
 PlatformName=PCD3D_SM5
@@ -135,6 +136,7 @@ bSupportsPercentageCloserShadows=true
 bSupportsVolumetricFog=true
 bSupportsIndexBufferUAVs=true
 bSupportsBindless=true
+bSupportsOIT=true
 
 [PreviewPlatform PCD3D_SM6]
 PlatformName=PCD3D_SM6
diff --git a/Engine/Shaders/Private/BasePassPixelShader.usf b/Engine/Shaders/Private/BasePassPixelShader.usf
--- a/Engine/Shaders/Private/BasePassPixelShader.usf
+++ b/Engine/Shaders/Private/BasePassPixelShader.usf
@@ -135,6 +135,13 @@ float NormalCurvatureToRoughness(float3 WorldNormal)
 #define POST_PROCESS_SUBSURFACE ((MATERIAL_SHADINGMODEL_SUBSURFACE_PROFILE || MATERIAL_SHADINGMODEL_EYE) && USES_GBUFFER)
 #define MATERIAL_WORKS_WITH_DUAL_SOURCE_COLOR_BLENDING (MATERIAL_SHADINGMODEL_THIN_TRANSLUCENT || STRATA_BLENDING_TRANSLUCENT_COLOREDTRANSMITTANCE )
 
+#define OIT_ENABLED (PROJECT_OIT && (MATERIALBLENDING_TRANSLUCENT || MATERIALBLENDING_ADDITIVE || MATERIAL_WORKS_WITH_DUAL_SOURCE_COLOR_BLENDING))
+
+#if OIT_ENABLED
+#define OIT_IS_BASEPASS 1
+#include "OITCommon.ush"
+#endif
+
 #if  MATERIAL_SHADINGMODEL_SINGLELAYERWATER
 
 #if SINGLE_LAYER_WATER_SIMPLE_FORWARD
@@ -2024,6 +2031,47 @@ void FPixelShaderInOut_MainPS(
 		Out.MRT[0].rgba *= ViewPreExposure;
 	#endif
 #endif
+
+	// If support OIT, then remove blending contribution and insert OIT sample
+	// Note: Out.MRT[0] has already view pre-exposition applied
+#if OIT_ENABLED
+	if (TranslucentBasePass.OIT.bOITEnable)
+	{
+		const float OpacityThreshold = 10e-4f;
+
+		#if MATERIAL_WORKS_WITH_DUAL_SOURCE_COLOR_BLENDING
+		{
+			float3 AdjustedDualBlendAdd = Out.MRT[0];
+			float3 AdjustedDualBlendMul = Out.MRT[1];
+			// Add early out threshold?
+			AddOITSample(uint2(In.SvPosition.xy), AdjustedDualBlendAdd, AdjustedDualBlendMul, MaterialParameters.ScreenPosition.w);
+			Out.MRT[0] = half4(0, 0, 0, 0);
+			Out.MRT[1] = half4(1, 1, 1, 1);
+		}
+		#else
+		{
+			float3 WriteColor = Out.MRT[0].rgb;
+			float  Opacity = Out.MRT[0].a;
+		#if MATERIALBLENDING_TRANSLUCENT
+			WriteColor *= Opacity; // WriteColor is premultiplied alpha color
+		#elif MATERIALBLENDING_ADDITIVE
+			Opacity = 1.0f;
+		#elif MATERIALBLENDING_ALPHACOMPOSITE || MATERIALBLENDING_ALPHAHOLDOUT
+			// Not clear what should be done in this case.
+		#else
+			#error OIT is not support for this blend mode
+		#endif
+
+			if (Opacity > OpacityThreshold)
+			{
+				AddOITSample(uint2(In.SvPosition.xy), WriteColor, 1.f - Opacity, MaterialParameters.ScreenPosition.w);
+			}
+			Out.MRT[0] = half4(0, 0, 0, 0);
+		}
+		#endif
+	}
+#endif // OIT_ENABLED
+
 #if MATERIAL_IS_SKY
 	// Sky materials can result in high luminance values, e.g. the sun disk. 
 	// This is so we make sure to at least stay within the boundaries of fp10 and not cause NaN on some platforms.
@@ -2043,7 +2091,7 @@ void FPixelShaderInOut_MainPS(
 }
 
 // If virtual texture is enabled then use early depth test so that UAV feedback buffer writes respect the depth test
-#if NUM_VIRTUALTEXTURE_SAMPLES || LIGHTMAP_VT_ENABLED || STRATA_OPAQUE_DEFERRED
+#if NUM_VIRTUALTEXTURE_SAMPLES || LIGHTMAP_VT_ENABLED || STRATA_OPAQUE_DEFERRED || OIT_ENABLED
 	#if COMPILER_SUPPORTS_DEPTHSTENCIL_EARLYTEST_LATEWRITE
 		// If we support early depth test with late write behaviour then use it since we may be using discard, or modifying depth
 		#define PIXELSHADER_EARLYDEPTHSTENCIL DEPTHSTENCIL_EARLYTEST_LATEWRITE
diff --git a/Engine/Shaders/Private/OITCombine.usf b/Engine/Shaders/Private/OITCombine.usf
new file mode 100644
--- /dev/null
+++ b/Engine/Shaders/Private/OITCombine.usf
@@ -0,0 +1,223 @@
+// Copyright Epic Games, Inc. All Rights Reserved.
+
+#pragma once
+
+#include "Common.ush"
+
+#define OIT_ENABLED 1
+#define MAX_SAMPLE_COUNT 8
+#define OIT_PASS_REGULAR 1
+#define OIT_PASS_SEPARATE 2
+
+#include "OITCommon.ush"
+#include "ColorMap.ush"
+
+int2 Resolution;
+uint MaxSideSampleCount;
+uint MaxSampleCount;
+uint Method;
+uint PassType;
+uint SupportedPass;
+
+Texture2D<uint>	SampleColorTexture;
+Texture2D<uint>	SampleTransTexture;
+Texture2D<float>SampleDepthTexture;
+Texture2D<uint>	SampleCountTexture;
+
+FOITSample OITLoadSample(uint2 InPixelCoord)
+{
+	FOITPackedSample Packed = (FOITPackedSample)0;
+	Packed.Depth = SampleDepthTexture.Load(uint3(InPixelCoord, 0));
+	Packed.Color = SampleColorTexture.Load(uint3(InPixelCoord, 0));
+	Packed.Trans = SampleTransTexture.Load(uint3(InPixelCoord, 0));
+	return UnpackOITSample(Packed);
+}
+
+///////////////////////////////////////////////////////////////////////////////////////////////////
+// Combine
+
+#if SHADER_OIT_COMBINE
+RWTexture2D<float4> OutColorTexture;
+
+[numthreads(8, 8, 1)]
+void MainCS(uint3 DispatchThreadId : SV_DispatchThreadID)
+{	
+	const uint2 PixelCoord = DispatchThreadId.xy;
+	if (any(PixelCoord.xy >= uint2(Resolution)))
+		return;
+
+	FOITSample Samples[MAX_SAMPLE_COUNT];
+
+	const uint  SampleCount		= min(MaxSampleCount, SampleCountTexture.Load(uint3(PixelCoord, 0)));
+	const uint2 SampleBaseCoord	= PixelCoord * MaxSideSampleCount;
+
+	// Bubble sort
+	// TODO: sorting go through memory, while it should be in registers
+	for (uint SampleIt = 0; SampleIt < SampleCount; ++SampleIt)
+	{
+		Samples[SampleIt] = (FOITSample)0;
+		Samples[SampleIt].Depth = 10e28f;
+
+		const uint2 SampleOffset = uint2(SampleIt % MaxSideSampleCount, SampleIt / MaxSideSampleCount);
+		const uint2 SampleCoord  = SampleBaseCoord + SampleOffset;
+
+		FOITSample Current = OITLoadSample(SampleCoord);
+		if (SampleIt > 0)
+		{
+			for (uint j = 0; j <= SampleIt; ++j)
+			{
+				if (Current.Depth < Samples[j].Depth) // Scene depth
+				{
+					OITSwap(Samples[j], Current);
+				}
+			}
+		}
+		else
+		{
+			Samples[SampleIt] = Current;
+		}
+	}
+
+	// Combine all samples from front to back
+	float3 FinalColor = 0;
+	float3 Transmittance = 1;
+	for (uint CompIt = 0; CompIt < SampleCount; ++CompIt)
+	{
+		FinalColor	 += Samples[CompIt].Color * Transmittance;
+		Transmittance = saturate(Transmittance * Samples[CompIt].Trans);
+	}
+
+	float AdjustedAlpha = saturate(1 - dot(Transmittance, float3(1.0f, 1.0f, 1.0f) / 3.0f));
+	
+	// Write final color - typed load
+	if (SampleCount > 0)
+	{
+		OutColorTexture[PixelCoord] = float4(FinalColor + OutColorTexture[PixelCoord].rgb * Transmittance, PassType == OIT_PASS_SEPARATE ? AdjustedAlpha : OutColorTexture[PixelCoord].a);
+	}
+}
+#endif // SHADER_OIT_COMBINE
+
+///////////////////////////////////////////////////////////////////////////////////////////////////
+// Debug
+
+#if SHADER_OIT_DEBUG
+
+#include "ShaderPrint.ush"
+
+int2 CursorCoord;
+
+void AddNewLine(float2 OriginalPos, inout float2 Pos)
+{
+	Pos = ShaderPrintNewline(Pos); Pos.x = OriginalPos.x;
+}
+
+void WriteEnable(inout float2 Pos, bool bEnable)
+{
+	if (bEnable)
+	{
+		Pos = ShaderPrintText(Pos, TEXT("On"), FontGreen);
+	}
+	else
+	{
+		Pos = ShaderPrintText(Pos, TEXT("Off"), FontRed);
+	}
+}
+[numthreads(1, 1, 1)]
+void MainCS(uint3 DispatchThreadId : SV_DispatchThreadID)
+{
+	if (all(DispatchThreadId == 0) && all(CursorCoord > 0) && all(CursorCoord < Resolution-1))
+	{
+		const uint2 PixelCoord	= CursorCoord;
+		uint  SampleCount		= SampleCountTexture.Load(uint3(PixelCoord, 0));
+		uint2 SampleBaseCoord	= PixelCoord * MaxSideSampleCount;
+		
+		// Global info
+		{
+			const uint2 GlobalCoord = uint2(50, 80);
+			float2 OriginalPos = GlobalCoord / float2(Resolution);
+			float2 Pos = OriginalPos;
+
+			if (Method == 0)
+			{
+			Pos = ShaderPrintText(Pos, TEXT("Method          : K-Buffer"), FontOrange);
+			}
+			else
+			{
+			Pos = ShaderPrintText(Pos, TEXT("Method          : MLAB"), FontOrange);
+			}
+			AddNewLine(OriginalPos, Pos);
+
+			Pos = ShaderPrintText(Pos, TEXT("Max Sample Count: "), FontWhite);
+			Pos = ShaderPrint(Pos, MaxSampleCount, FontYellow);
+			AddNewLine(OriginalPos, Pos);
+
+			Pos = ShaderPrintText(Pos, TEXT("Current Pass    : "), FontWhite);
+			if (!!(PassType & OIT_PASS_REGULAR))
+			{
+				Pos = ShaderPrintText(Pos, TEXT("Regular Translucency"), FontYellow);
+			}
+			else if (!!(PassType & OIT_PASS_SEPARATE))
+			{
+				Pos = ShaderPrintText(Pos, TEXT("Separate Translucency"), FontYellow);
+			}
+			AddNewLine(OriginalPos, Pos);
+
+			Pos = ShaderPrintText(Pos, TEXT("Transp. Regular : "), FontWhite); WriteEnable(Pos, (SupportedPass & OIT_PASS_REGULAR) > 0);
+			AddNewLine(OriginalPos, Pos);
+
+			Pos = ShaderPrintText(Pos, TEXT("Transp. Separate: "), FontWhite); WriteEnable(Pos, (SupportedPass & OIT_PASS_SEPARATE) > 0);
+			AddNewLine(OriginalPos, Pos);
+		}
+
+		// Pixel coord
+		float2 OriginalPos = PixelCoord / float2(Resolution);
+		float2 Pos = OriginalPos;
+		Pos = ShaderPrintText(Pos, TEXT("Pixel Coord.: "), FontWhite);
+		Pos = ShaderPrint(Pos, PixelCoord, FontYellow);
+		AddNewLine(OriginalPos, Pos);
+
+		FFontColor SampleCountColor = FontEmerald;
+		if (SampleCount > MaxSampleCount)
+		{
+			SampleCountColor = FontRed;
+		}
+		Pos = ShaderPrintText(Pos, TEXT("Sample Count: "), FontWhite);
+		Pos = ShaderPrint(Pos, SampleCount, SampleCountColor);
+		Pos = ShaderPrint(Pos, MaxSampleCount, SampleCountColor);
+		AddNewLine(OriginalPos, Pos);
+
+		if (Method == 0)
+		{
+			Pos = ShaderPrintText(Pos, TEXT("Method: K-Buffer"), FontOrange);
+		}
+		else
+		{
+			Pos = ShaderPrintText(Pos, TEXT("Method: MLAB"), FontOrange);
+		}
+		AddNewLine(OriginalPos, Pos);
+		
+
+		// Sample infos
+		SampleCount = min(SampleCount, MaxSampleCount);
+		for (uint SampleIt = 0; SampleIt < SampleCount; ++SampleIt)
+		{
+			const uint2 SampleOffset = uint2(SampleIt % MaxSideSampleCount, SampleIt / MaxSideSampleCount);
+			const uint2 SampleCoord = SampleBaseCoord + SampleOffset;
+
+			const FOITSample Sample = OITLoadSample(SampleCoord);
+
+			Pos = ShaderPrintText(Pos, TEXT("Depth "), FontYellow);
+			Pos = ShaderPrint(Pos, Sample.Depth, FontYellow);
+			AddNewLine(OriginalPos, Pos);
+
+			Pos = ShaderPrintText(Pos, TEXT("Color "), FontEmerald);	
+			Pos = ShaderPrint(Pos, Sample.Color, InitFontColor(Sample.Color));
+			AddNewLine(OriginalPos, Pos);
+
+			Pos = ShaderPrintText(Pos, TEXT("Trans."), FontOrange);	
+			Pos = ShaderPrint(Pos, Sample.Trans, InitFontColor(Sample.Trans));
+			AddNewLine(OriginalPos, Pos);
+		}
+	}
+}
+#endif
\ No newline at end of file
diff --git a/Engine/Shaders/Private/OITCommon.ush b/Engine/Shaders/Private/OITCommon.ush
new file mode 100644
--- /dev/null
+++ b/Engine/Shaders/Private/OITCommon.ush
@@ -0,0 +1,311 @@
+// Copyright Epic Games, Inc. All Rights Reserved.
+
+#pragma once
+
+#ifndef OIT_ENABLED
+#define OIT_ENABLED 0
+#endif
+
+#ifndef OIT_IS_BASEPASS
+#define OIT_IS_BASEPASS 0
+#endif
+
+#if OIT_ENABLED
+
+//////////////////////////////////////////////////////////////////////////////////////////
+// Sample storage format 
+
+#define OIT_FORMAT_10BITS  1
+#define OIT_FORMAT_16BITS  2 
+#define OIT_FORMAT_COMPACT 3
+#define OIT_FORMAT OIT_FORMAT_10BITS
+
+//////////////////////////////////////////////////////////////////////////////////////////
+
+struct FOITSample
+{
+	float3 Color;
+	float  Depth;
+	float3 Trans;
+	float  Pad;
+};
+
+struct FOITPackedSample
+{
+	uint  Color;
+	uint  Trans;
+	float Depth;
+	uint  Pad;
+};
+
+//////////////////////////////////////////////////////////////////////////////////////////
+
+uint OITPack10F(float Value)
+{
+	return (f32tof16(Value) >> 5) & 0x000003FF;
+}
+
+float OITUnpack10F(uint Value)
+{
+	return f16tof32((Value << 5) & 0x7FE0);
+}
+
+uint OITPack6F(float Value)
+{
+	return (f32tof16(Value) >> 5) & 0x000003FF;
+}
+
+float OITUnpack6F(uint Value)
+{
+	return f16tof32((Value << 5) & 0x7FE0);
+}
+
+
+void OITSwap(inout FOITSample A, inout FOITSample B)
+{
+	FOITSample T = A;
+	A = B;
+	B = T;
+}
+
+uint OITPackR11G11B10F(float3 rgb)
+{
+	uint r = (f32tof16(rgb.r) << 17) & 0xFFE00000;
+	uint g = (f32tof16(rgb.g) << 6) & 0x001FFC00;
+	uint b = (f32tof16(rgb.b) >> 5) & 0x000003FF;
+	return r | g | b;
+}
+
+float3 OITUnpackR11G11B10F(uint rgb)
+{
+	float r = f16tof32((rgb >> 17) & 0x7FF0);
+	float g = f16tof32((rgb >> 6) & 0x7FF0);
+	float b = f16tof32((rgb << 5) & 0x7FE0);
+	return float3(r, g, b);
+}
+
+//////////////////////////////////////////////////////////////////////////////////////////
+// Per-attribut compaction 12bytes/96bits
+
+#if OIT_FORMAT == OIT_FORMAT_10BITS
+FOITPackedSample PackOITSample(FOITSample In)
+{
+	FOITPackedSample Out = (FOITPackedSample)0;
+	Out.Color = OITPackR11G11B10F(In.Color);
+	Out.Trans = OITPackR11G11B10F(saturate(In.Trans));
+	Out.Depth = In.Depth;
+	return Out;
+}
+
+FOITSample UnpackOITSample(FOITPackedSample In)
+{
+	FOITSample Out = (FOITSample)0;
+	Out.Color = OITUnpackR11G11B10F(In.Color);
+	Out.Trans = saturate(OITUnpackR11G11B10F(In.Trans));
+	Out.Depth = In.Depth;
+	return Out;
+}
+#endif
+
+//////////////////////////////////////////////////////////////////////////////////////////
+// Fat compaction 16bytes/128bits
+
+#if OIT_FORMAT_16BITS
+uint4 PackOITSample(float3 InColor, float3 InTrans, float InDepth)
+{
+	uint4 Out = 0;
+	Out.x = f32tof16(InColor.r) | (f32tof16(InColor.g) << 16);
+	Out.y = f32tof16(InColor.b) | (f32tof16(InTrans.r) << 16);
+	Out.z = f32tof16(InTrans.g) | (f32tof16(InTrans.b) << 16);
+	Out.w = InDepth;
+	return Out;
+}
+
+void UnpackOITSample(uint4 In, inout float3 OutColor, inout float3 OutTrans, inout float OutDepth)
+{
+	OutColor.r = f16tof32(0xFFFF & (In.x));
+	OutColor.g = f16tof32(0xFFFF & (In.x >> 16));
+	OutColor.b = f16tof32(0xFFFF & (In.y));
+
+	OutTrans.r = f16tof32(0xFFFF & (In.y >> 16));
+	OutTrans.g = f16tof32(0xFFFF & (In.z));
+	OutTrans.b = f16tof32(0xFFFF & (In.z >> 16));
+
+	OutDepth = In.w;
+}
+#endif // OIT_FORMAT_16BITS
+
+//////////////////////////////////////////////////////////////////////////////////////////
+// Compact sample into 8bytes/64bits (untested|add dithering if banding)
+
+#if OIT_FORMAT == OIT_FORMAT_COMPACT
+uint2 PackOITSample(float3 InColor, float3 InTrans, float InDepth)
+{
+	uint ColorR10 = OITPack10F(InColor.r);
+	uint ColorG10 = OITPack10F(InColor.g);
+	uint ColorB10 = OITPack10F(InColor.b);
+
+	uint TransR6 = OITPack6F(InTrans.r);
+	uint TransG6 = OITPack6F(InTrans.g);
+	uint TransB6 = OITPack6F(InTrans.b);
+
+	uint Depth16 = f32tof16(InDepth);
+	
+	uint2 Out = 0;
+	Out.x = ColorR10 | (ColorG10 << 10) | (ColorB10 << 20) | (((TransB6 >> 4) & 0x3) << 30);
+	Out.y = Depth16  | (TransR6  << 16) | (TransG6  << 22) | (TransB6 & 0xF) << 28);
+	return Out;
+}
+
+void UnpackOITSample(uint2 In, inout float3 OutColor, inout float3 OutTrans, inout float OutDepth)
+{
+	OutColor.r = OITUnpack10F((In.x)     & 3FF);
+	OutColor.g = OITUnpack10F((In.x>>10) & 3FF);
+	OutColor.b = OITUnpack10F((In.x>>20) & 3FF);
+
+	OutTrans.r = OITPack6F( (In.y>>16) & 0x3F);
+	OutTrans.g = OITPack6F( (In.y>>22) & 0x3F);
+	OutTrans.b = OITPack6F(((In.y>>28) & 0xF) | ((In.x>>30) & 0x3));
+
+	OutDepth = f16tof32(In.y & 0xFFFF);
+}
+#endif // OIT_FORMAT_COMPACT
+
+//////////////////////////////////////////////////////////////////////////////////////////
+
+#if OIT_IS_BASEPASS
+
+//////////////////////////////////////////////////////////////////////////////////////////
+// Special Add/Load/Store sample function for base-pass
+
+FOITSample OITLoadSample(uint2 InPixelCoord)
+{
+	FOITPackedSample Packed = (FOITPackedSample)0;
+	Packed.Color = TranslucentBasePass.OIT.OutOITSampleColor[InPixelCoord];
+	Packed.Trans = TranslucentBasePass.OIT.OutOITSampleTrans[InPixelCoord];
+	Packed.Depth = TranslucentBasePass.OIT.OutOITSampleDepth[InPixelCoord];
+
+	return UnpackOITSample(Packed);
+}
+
+void OITStoreSample(uint2 InPixelCoord, FOITSample In)
+{
+	FOITPackedSample Packed = PackOITSample(In);
+	TranslucentBasePass.OIT.OutOITSampleColor[InPixelCoord] = Packed.Color;
+	TranslucentBasePass.OIT.OutOITSampleTrans[InPixelCoord] = Packed.Trans;
+	TranslucentBasePass.OIT.OutOITSampleDepth[InPixelCoord] = Packed.Depth;
+}
+
+// Simple k-buffer insertion
+void AddOITSample_KB(uint2 InPixelCoord, float3 InColor, float3 InTrans, float InDepth)
+{
+	uint SampleOffset = 0;
+	InterlockedAdd(TranslucentBasePass.OIT.OutOITSampleCount[InPixelCoord], 1u, SampleOffset);
+
+	if (SampleOffset < TranslucentBasePass.OIT.MaxSamplePerPixel)
+	{
+		const uint2 SampleLocalCoord = uint2(SampleOffset % TranslucentBasePass.OIT.MaxSideSamplePerPixel, SampleOffset / TranslucentBasePass.OIT.MaxSideSamplePerPixel);
+		const uint2 SampleCoord = InPixelCoord * TranslucentBasePass.OIT.MaxSideSamplePerPixel + SampleLocalCoord;
+
+		FOITSample Sample = (FOITSample)0;
+		Sample.Color = InColor;
+		Sample.Trans = InTrans;
+		Sample.Depth = InDepth;
+
+		OITStoreSample(SampleCoord, Sample);
+	}
+}
+
+uint2 GetStoringBaseCoord(uint2 InPixelCoord)
+{
+	return InPixelCoord * TranslucentBasePass.OIT.MaxSideSamplePerPixel;
+}
+
+uint2 GetStoringCoord(uint2 InBaseCoord, uint InSampleIt)
+{
+	const uint2 SampleOffset = uint2(InSampleIt % TranslucentBasePass.OIT.MaxSideSamplePerPixel, InSampleIt / TranslucentBasePass.OIT.MaxSideSamplePerPixel);
+	return InBaseCoord + SampleOffset;
+}
+
+// Multi-Layer Alpha Blending (MLAB) - Sample are inserted in a sorted manner and overflow sample are merged with alpha blending
+void AddOITSample_MLAB(uint2 InPixelCoord, float3 InColor, float3 InTrans, float InDepth)
+{
+	// ROV read to for pixel synchronization
+	uint SampleCount = min(TranslucentBasePass.OIT.OutOITSampleCount[InPixelCoord], TranslucentBasePass.OIT.MaxSamplePerPixel);
+
+	FOITSample Current = (FOITSample)0;
+	Current.Depth = InDepth;
+	Current.Color = InColor;
+	Current.Trans = InTrans;
+
+	const uint2 SampleBaseCoord = GetStoringBaseCoord(InPixelCoord);
+
+	// Check if sample fall into the merged/combined bucket, and if it needs to be discard due to transmission threshold
+	bool bNeedInsertSort = true;
+	if (SampleCount > 0)
+	{
+		const uint2 LastSampleCoord = GetStoringCoord(SampleBaseCoord, SampleCount-1);
+		const FOITSample LastSample = OITLoadSample(LastSampleCoord);
+		if (Current.Depth > LastSample.Depth) // Scene depth (not inv-Z)
+		{
+			if (all(LastSample.Trans < TranslucentBasePass.OIT.TransmittanceThreshold))
+			{
+				discard;
+			}
+		}
+	}
+
+	// Bubble sort and insert
+	if (bNeedInsertSort)
+	{
+		for (uint SampleIt = 0; SampleIt < SampleCount; ++SampleIt)
+		{
+			const uint2 SampleCoord = GetStoringCoord(SampleBaseCoord, SampleIt);
+			const FOITSample Sample = OITLoadSample(SampleCoord);
+
+			if (Current.Depth < Sample.Depth) // Scene depth (not inv-Z)
+			{
+				OITStoreSample(SampleCoord, Current);
+				Current = Sample;
+			}			
+		}
+	}
+
+	// Merge last fragment
+	if (SampleCount >= TranslucentBasePass.OIT.MaxSamplePerPixel)
+	{
+		const uint2 SampleCoord = GetStoringCoord(SampleBaseCoord, TranslucentBasePass.OIT.MaxSamplePerPixel - 1);
+
+		FOITSample Last = OITLoadSample(SampleCoord);
+		Last.Color = Last.Color + Last.Trans * Current.Color;
+		Last.Trans = saturate(Last.Trans * Current.Trans);
+		Last.Depth = min(Last.Depth, Current.Depth); // Normally Last.Depth should already be the min. depth
+
+		OITStoreSample(SampleCoord, Last);
+	}
+	else
+	{
+		// If there are still room to store sample (i.e., not over budget), insert the sample at the end (note: SampleCount has not being incremented yet)
+		const uint2 SampleCoord = GetStoringCoord(SampleBaseCoord, SampleCount);
+		OITStoreSample(SampleCoord, Current);
+	}
+
+	// Write to ROV to unlock pixel
+	SampleCount = min(SampleCount + 1, TranslucentBasePass.OIT.MaxSamplePerPixel);
+	TranslucentBasePass.OIT.OutOITSampleCount[InPixelCoord] = SampleCount;
+}
+
+// Add a translucent sample to the list of existing translucent samples
+// InColor : Radiance/Luminance reflected by the surface (pre-multiplied alpha)
+// InTrans : Transmittance of the surface (i.e., 1-Opacity)
+// InDepth : Scene depth (i.e. not DeviceZ)
+void AddOITSample(uint2 InPixelCoord, float3 InColor, float3 InTrans, float InDepth)
+{
+	AddOITSample_MLAB(InPixelCoord, InColor, InTrans, InDepth);
+}
+
+//////////////////////////////////////////////////////////////////////////////////////////
+
+#endif // OIT_IS_BASEPASS
+
+#endif // OIT_ENABLED
\ No newline at end of file
diff --git a/Engine/Source/Runtime/Engine/Classes/Engine/RendererSettings.h b/Engine/Source/Runtime/Engine/Classes/Engine/RendererSettings.h
--- a/Engine/Source/Runtime/Engine/Classes/Engine/RendererSettings.h
+++ b/Engine/Source/Runtime/Engine/Classes/Engine/RendererSettings.h
@@ -926,6 +926,15 @@ class ENGINE_API URendererSettings : public UDeveloperSettings
 		ConfigRestartRequired = true))
 		uint32 bMaterialEnergyConservation : 1;
 
+	/**
+	"Enable Order Independent Transparency (Experimental)."
+	*/
+	UPROPERTY(config, EditAnywhere, Category = Translucency, meta = (
+		ConsoleVariable = "r.OIT.SortedPixels", DisplayName = "Enable Order Independent Transparency (Experimental)",
+		ToolTip = "Enable support for Order-Independent-Transparency on translucent surfaces, which remove most of the sorting artifact among translucent surfaces.",
+		ConfigRestartRequired = true))
+		uint32 bOrderedIndependentTransparencyEnable : 1;
+
 	/**
 	"Skin cache allows a compute shader to skin once each vertex, save those results into a new buffer and reuse those calculations when later running the depth, base and velocity passes. This also allows opting into the 'recompute tangents' for skinned mesh instance feature. Disabling will reduce the number of shader permutations required per material. Changing this setting requires restarting the editor."
 	*/
diff --git a/Engine/Source/Runtime/Engine/Private/ShaderCompiler/ShaderCompiler.cpp b/Engine/Source/Runtime/Engine/Private/ShaderCompiler/ShaderCompiler.cpp
--- a/Engine/Source/Runtime/Engine/Private/ShaderCompiler/ShaderCompiler.cpp
+++ b/Engine/Source/Runtime/Engine/Private/ShaderCompiler/ShaderCompiler.cpp
@@ -5871,6 +5871,13 @@ void GlobalBeginCompileShader(
 		Input.Environment.SetDefine(TEXT("MATERIAL_ROUGHDIFFUSE"), bMaterialRoughDiffuse ? 1 : 0);
 	}
 
+	{
+		const bool bSupportOIT = FDataDrivenShaderPlatformInfo::GetSupportsOIT(EShaderPlatform(Target.Platform));
+		static IConsoleVariable* CVar = IConsoleManager::Get().FindConsoleVariable(TEXT("r.OIT.SortedPixels"));
+		const bool bOIT = CVar && CVar->GetInt() != 0;
+		Input.Environment.SetDefine(TEXT("PROJECT_OIT"), (bSupportOIT && bOIT) ? 1 : 0);
+	}
+
 	{
 		static IConsoleVariable* CVar = IConsoleManager::Get().FindConsoleVariable(TEXT("r.Material.EnergyConservation"));
 		const bool bMaterialEnergyConservation = CVar && CVar->GetInt() != 0;
diff --git a/Engine/Source/Runtime/Engine/Private/ShaderCompiler/ShaderGenerationUtil.cpp b/Engine/Source/Runtime/Engine/Private/ShaderCompiler/ShaderGenerationUtil.cpp
--- a/Engine/Source/Runtime/Engine/Private/ShaderCompiler/ShaderGenerationUtil.cpp
+++ b/Engine/Source/Runtime/Engine/Private/ShaderCompiler/ShaderGenerationUtil.cpp
@@ -239,6 +239,8 @@ void FShaderCompileUtilities::ApplyFetchEnvironment(FShaderMaterialPropertyDefin
 	FETCH_COMPILE_BOOL(STRATA_ENABLED);
 	FETCH_COMPILE_BOOL(MATERIAL_IS_STRATA);
 
+	FETCH_COMPILE_BOOL(PROJECT_OIT);
+
 	FETCH_COMPILE_BOOL(DUAL_SOURCE_COLOR_BLENDING_ENABLED);
 
 	FETCH_COMPILE_INT(DECAL_RENDERTARGET_COUNT);
diff --git a/Engine/Source/Runtime/RHI/Private/RHI.cpp b/Engine/Source/Runtime/RHI/Private/RHI.cpp
--- a/Engine/Source/Runtime/RHI/Private/RHI.cpp
+++ b/Engine/Source/Runtime/RHI/Private/RHI.cpp
@@ -2381,6 +2381,7 @@ void FGenericDataDrivenShaderPlatformInfo::ParseDataDrivenShaderInfo(const FConf
 	GET_SECTION_BOOL_HELPER(bMobileSupportFetchBindedCustomStencilBuffer);
 	GET_SECTION_BOOL_HELPER(bSupportsBindless);
 	GET_SECTION_BOOL_HELPER(bSupportsVolumeTextureAtomics);
+	GET_SECTION_BOOL_HELPER(bSupportsOIT);
 #undef GET_SECTION_BOOL_HELPER
 #undef GET_SECTION_INT_HELPER
 
diff --git a/Engine/Source/Runtime/RHI/Public/RHIDefinitions.h b/Engine/Source/Runtime/RHI/Public/RHIDefinitions.h
--- a/Engine/Source/Runtime/RHI/Public/RHIDefinitions.h
+++ b/Engine/Source/Runtime/RHI/Public/RHIDefinitions.h
@@ -397,7 +397,7 @@ class RHI_API FGenericDataDrivenShaderPlatformInfo
 	uint32 bMobileSupportFetchBindedCustomStencilBuffer : 1;
 	uint32 bSupportsBindless : 1;
 	uint32 bSupportsVolumeTextureAtomics : 1;
-
+	uint32 bSupportsOIT : 1;
 		
 #if WITH_EDITOR
 	FText FriendlyName;
@@ -921,6 +921,11 @@ public:
 		return Infos[Platform].bSupportsVolumeTextureAtomics;
 	}
 
+	static FORCEINLINE_DEBUGGABLE const bool GetSupportsOIT(const FStaticShaderPlatform Platform)
+	{
+		return Infos[Platform].bSupportsOIT;
+	}
+
 #if WITH_EDITOR
 	static FORCEINLINE_DEBUGGABLE FText GetFriendlyName(const FStaticShaderPlatform Platform)
 	{
diff --git a/Engine/Source/Runtime/RenderCore/Private/ShaderMaterialDerivedHelpers.cpp b/Engine/Source/Runtime/RenderCore/Private/ShaderMaterialDerivedHelpers.cpp
--- a/Engine/Source/Runtime/RenderCore/Private/ShaderMaterialDerivedHelpers.cpp
+++ b/Engine/Source/Runtime/RenderCore/Private/ShaderMaterialDerivedHelpers.cpp
@@ -100,6 +100,9 @@ FShaderMaterialDerivedDefines RENDERCORE_API CalculateDerivedMaterialParameters(
 	// As of today, all translucent Strata material forces dual source color blending. STRATA_TODO: optimize out dual color blending when not needed or requested
 	Dst.MATERIAL_WORKS_WITH_DUAL_SOURCE_COLOR_BLENDING = (Mat.MATERIAL_SHADINGMODEL_THIN_TRANSLUCENT || Mat.STRATA_BLENDING_TRANSLUCENT_COLOREDTRANSMITTANCE);
 
+	// matching the logic in BasePassPixelShader.usf
+	Dst.OIT_ENABLED = (FEATURE_LEVEL >= ERHIFeatureLevel::SM5) && Mat.PROJECT_OIT && (Mat.MATERIALBLENDING_TRANSLUCENT || Mat.MATERIALBLENDING_ADDITIVE || Dst.MATERIAL_WORKS_WITH_DUAL_SOURCE_COLOR_BLENDING);
+
 	// There are 4 different ways of setting MRTs depending on which .usf file is the base shader, which sets defines which includes a different include file
 	// which may or may not override the current defines. Here is my best attempt at figuring out the logic.
 	if (Mat.IS_VIRTUAL_TEXTURE_MATERIAL)
diff --git a/Engine/Source/Runtime/RenderCore/Private/ShaderParameterMetadata.cpp b/Engine/Source/Runtime/RenderCore/Private/ShaderParameterMetadata.cpp
--- a/Engine/Source/Runtime/RenderCore/Private/ShaderParameterMetadata.cpp
+++ b/Engine/Source/Runtime/RenderCore/Private/ShaderParameterMetadata.cpp
@@ -167,7 +167,8 @@ const TCHAR* GetShaderParameterMacroName(EUniformBufferBaseType ShaderParameterB
 
 EShaderCodeResourceBindingType ParseShaderResourceBindingType(const TCHAR* ShaderType)
 {
-	bool bIsRWResource = FCString::Strncmp(ShaderType, TEXT("RW"), 2) == 0;
+	const bool bIsRasterizerOrderedResource = FCString::Strncmp(ShaderType, TEXT("RasterizerOrdered"), 17) == 0;
+	const bool bIsRWResource = FCString::Strncmp(ShaderType, TEXT("RW"), 2) == 0 || bIsRasterizerOrderedResource;
 	const TCHAR* ComparedShaderType = ShaderType + (bIsRWResource ? 2 : 0);
 
 	int32 ShaderTypeLength = 0;
diff --git a/Engine/Source/Runtime/RenderCore/Public/ShaderMaterial.h b/Engine/Source/Runtime/RenderCore/Public/ShaderMaterial.h
--- a/Engine/Source/Runtime/RenderCore/Public/ShaderMaterial.h
+++ b/Engine/Source/Runtime/RenderCore/Public/ShaderMaterial.h
@@ -178,6 +178,8 @@ struct FShaderMaterialPropertyDefines
 	uint8 STRATA_ENABLED : 1;
 	uint8 MATERIAL_IS_STRATA : 1;
 
+	uint8 PROJECT_OIT : 1;
+
 	uint8 DUAL_SOURCE_COLOR_BLENDING_ENABLED : 1;
 
 	uint8 IS_MATERIAL_SHADER : 1;
@@ -268,6 +270,8 @@ struct FShaderMaterialDerivedDefines
 	bool STRATA_ENABLED;
 	bool SHADER_STRATA_TRANSLUCENT_ENABLED;
 
+	bool OIT_ENABLED;
+
 	bool THIN_TRANSLUCENT_USE_DUAL_BLEND;
 	bool MATERIAL_WORKS_WITH_DUAL_SOURCE_COLOR_BLENDING;
 };
diff --git a/Engine/Source/Runtime/Renderer/Private/BasePassRendering.h b/Engine/Source/Runtime/Renderer/Private/BasePassRendering.h
--- a/Engine/Source/Runtime/Renderer/Private/BasePassRendering.h
+++ b/Engine/Source/Runtime/Renderer/Private/BasePassRendering.h
@@ -28,6 +28,7 @@
 #include "UnrealEngine.h"
 #include "ReflectionEnvironment.h"
 #include "Strata/Strata.h"
+#include "OIT/OITParameters.h"
 #include "VirtualShadowMaps/VirtualShadowMapArray.h"
 #include "VolumetricCloudRendering.h"
 
@@ -99,6 +100,7 @@ BEGIN_GLOBAL_SHADER_PARAMETER_STRUCT(FTranslucentBasePassUniformParameters,)
 	SHADER_PARAMETER_STRUCT(FSceneTextureUniformParameters, SceneTextures)
 	SHADER_PARAMETER_STRUCT(FStrataForwardPassUniformParameters, Strata)
 	SHADER_PARAMETER_STRUCT(FLightCloudTransmittanceParameters, ForwardDirLightCloudShadow)
+	SHADER_PARAMETER_STRUCT(FOITBasePassUniformParameters, OIT)
 	// Material SSR
 	SHADER_PARAMETER(FVector4f, HZBUvFactorAndInvFactor)
 	SHADER_PARAMETER(FVector4f, PrevScreenPositionScaleBias)
diff --git a/Engine/Source/Runtime/Renderer/Private/DeferredShadingRenderer.cpp b/Engine/Source/Runtime/Renderer/Private/DeferredShadingRenderer.cpp
--- a/Engine/Source/Runtime/Renderer/Private/DeferredShadingRenderer.cpp
+++ b/Engine/Source/Runtime/Renderer/Private/DeferredShadingRenderer.cpp
@@ -3083,9 +3083,9 @@ void FDeferredShadingSceneRenderer::Render(FRDGBuilder& GraphBuilder)
 		// Sort objects' triangles
 		for (FViewInfo& View : Views)
 		{
-			if (OIT::IsEnabled(View))
+			if (OIT::IsEnabled(EOITSortingType::SortedTriangles, View))
 			{
-				OIT::AddSortTrianglesPass(GraphBuilder, View, Scene->OITSceneData, FOITSortingType::BackToFront);
+				OIT::AddSortTrianglesPass(GraphBuilder, View, Scene->OITSceneData, FTriangleSortingOrder::BackToFront);
 			}
 		}
 
diff --git a/Engine/Source/Runtime/Renderer/Private/LightShaftRendering.cpp b/Engine/Source/Runtime/Renderer/Private/LightShaftRendering.cpp
--- a/Engine/Source/Runtime/Renderer/Private/LightShaftRendering.cpp
+++ b/Engine/Source/Runtime/Renderer/Private/LightShaftRendering.cpp
@@ -546,7 +546,7 @@ void FDeferredShadingSceneRenderer::RenderLightShaftBloom(
 			if (LightSceneInfo.bEnableLightShaftBloom)
 			{
 				bool bWillRenderLightShafts = false;
-
+				
 				for (int32 ViewIndex = 0; ViewIndex < Views.Num(); ViewIndex++)
 				{
 					const FViewInfo& View    = Views[ViewIndex];
@@ -574,7 +574,7 @@ void FDeferredShadingSceneRenderer::RenderLightShaftBloom(
 							OutputLoadAction = ERenderTargetLoadAction::EClear;
 
 							const bool bIsModulate = false;
-							OutputTexture = CreatePostDOFTranslucentTexture(GraphBuilder, ETranslucencyPass::TPT_TranslucencyAfterDOF, SeparateTranslucencyDimensions, bIsModulate);
+							OutputTexture = CreatePostDOFTranslucentTexture(GraphBuilder, ETranslucencyPass::TPT_TranslucencyAfterDOF, SeparateTranslucencyDimensions, bIsModulate, ShaderPlatform);
 
 							// We will need to update views separate transluceny buffers if we have just created them.
 							bUpdateViewsSeparateTranslucency = true;
diff --git a/Engine/Source/Runtime/Renderer/Private/MeshDrawCommands.cpp b/Engine/Source/Runtime/Renderer/Private/MeshDrawCommands.cpp
--- a/Engine/Source/Runtime/Renderer/Private/MeshDrawCommands.cpp
+++ b/Engine/Source/Runtime/Renderer/Private/MeshDrawCommands.cpp
@@ -143,6 +143,7 @@ void UpdateTranslucentMeshSortKeys(
 	const FMatrix& ViewMatrix,
 	const TArray<struct FPrimitiveBounds>& PrimitiveBounds,
 	ETranslucencyPass::Type TranslucencyPass, 
+	bool bInverseSorting,
 	FMeshCommandOneFrameArray& VisibleMeshCommands
 	)
 {
@@ -179,6 +180,13 @@ void UpdateTranslucentMeshSortKeys(
 		const float DistanceOffset = *((float*)&PackedOffset);
 		Distance += DistanceOffset;
 
+		// Sort front-to-back instead of back-to-front
+		if (bInverseSorting)
+		{
+			const float MaxSortingDistance = 1000000.f; // 100km, arbitrary
+			Distance = MaxSortingDistance - Distance;
+		}
+
 		// Patch distance inside translucent mesh sort key.
 		FMeshDrawCommandSortKey SortKey;
 		SortKey.PackedData = VisibleCommand.SortKey.PackedData;
@@ -875,6 +883,10 @@ public:
 			}
 			else if (Context.TranslucencyPass != ETranslucencyPass::TPT_MAX)
 			{
+				// When per-pixel OIT is enabled, sort primitive from front to back ensure avoid 
+				// constantly resorting front-to-back samples list.
+				const bool bInverseSorting = OIT::IsEnabled(EOITSortingType::SortedPixels, Context.ShaderPlatform) && Context.View->AntiAliasingMethod != EAntiAliasingMethod::AAM_MSAA;
+
 				UpdateTranslucentMeshSortKeys(
 					Context.TranslucentSortPolicy,
 					Context.TranslucentSortAxis,
@@ -882,6 +894,7 @@ public:
 					Context.ViewMatrix,
 					*Context.PrimitiveBounds,
 					Context.TranslucencyPass,
+					bInverseSorting,
 					Context.MeshDrawCommands
 				);
 			}
diff --git a/Engine/Source/Runtime/Renderer/Private/OIT/OIT.cpp b/Engine/Source/Runtime/Renderer/Private/OIT/OIT.cpp
--- a/Engine/Source/Runtime/Renderer/Private/OIT/OIT.cpp
+++ b/Engine/Source/Runtime/Renderer/Private/OIT/OIT.cpp
@@ -1,55 +1,258 @@
 // Copyright Epic Games, Inc. All Rights Reserved.
 
 #include "OIT.h"
+#include "OITParameters.h"
 #include "Shader.h"
 #include "GlobalShader.h"
 #include "ShaderParameters.h"
 #include "ShaderParameterStruct.h"
 #include "SceneTextureParameters.h"
 #include "SceneRendering.h"
-
+#include "ShaderCompilerCore.h"
 #include "ShaderPrintParameters.h"
 #include "ShaderPrint.h"
 #include "ScreenPass.h"
 
 ///////////////////////////////////////////////////////////////////////////////////////////////////
-// Variables
+// Variables: sorted triangles
+
+static TAutoConsoleVariable<int32> CVarOIT_SortedTriangles_Enable(
+	TEXT("r.OIT.SortedTriangles"), 
+	1, 
+	TEXT("Enable per-instance triangle sorting to avoid invalid triangle ordering (experimental)."),
+	ECVF_ReadOnly | ECVF_RenderThreadSafe);
+
+static TAutoConsoleVariable<int32> CVarOIT_SortedTriangles_Debug(
+	TEXT("r.OIT.SortedTriangles.Debug"), 
+	0, 
+	TEXT("Enable per-instance triangle sorting debug rendering."),
+	ECVF_RenderThreadSafe);
+
+static TAutoConsoleVariable<int32> CVarOIT_SortedTriangles_Pool(
+	TEXT("r.OIT.SortedTriangles.Pool"), 
+	0, 
+	TEXT("Enable index buffer pool allocation which reduce creation/deletion time by re-use buffers."),
+	ECVF_RenderThreadSafe);
+
+static TAutoConsoleVariable<int32> CVarOIT_SortedTriangles_PoolReleaseThreshold(
+	TEXT("r.OIT.SortedTriangles.Pool.ReleaseFrameThreshold"), 
+	100, 
+	TEXT("Number of frame after which unused buffer are released."),
+	ECVF_RenderThreadSafe);
 
-static int32 GOIT_SortObjectTriangles = 1;
-static FAutoConsoleVariableRef CVarOIT_PerObjectSorting(TEXT("r.OIT.SortObjectTriangles"), GOIT_SortObjectTriangles, TEXT("Enable per-instance triangle sorting to avoid invalid triangle ordering (experimental)."));
+///////////////////////////////////////////////////////////////////////////////////////////////////
+// Variables: sorted pixels
+
+// Referenced in RendererSettings.h, as it is a project settings
+static TAutoConsoleVariable<int32> CVarOIT_SortedPixels_Enable(
+	TEXT("r.OIT.SortedPixels"),
+	0,
+	TEXT("Enable OIT rendering (project settings, can't be changed at runtime)"),
+	ECVF_ReadOnly | ECVF_RenderThreadSafe);
+
+static TAutoConsoleVariable<int32> CVarOIT_SortedPixels_PassType(
+	TEXT("r.OIT.SortedPixels.PassType"),
+	3,
+	TEXT("Enable OIT rendering. 0: disable 1: enable OIT for std. translucency 2: enable OIT for separated translucency 3: enable for both std. and separated translucency (default)"),
+	ECVF_RenderThreadSafe);
+
+static TAutoConsoleVariable<int32> CVarOIT_SortedPixels_SampleCount(
+	TEXT("r.OIT.SortedPixels.MaxSampleCount"),
+	4,
+	TEXT("Max sample count."),
+	ECVF_RenderThreadSafe);
+
+static TAutoConsoleVariable<int32> CVarOIT_SortedPixels_Debug(
+	TEXT("r.OIT.SortedPixels.Debug"),
+	0,
+	TEXT("Enable debug rendering for OIT. 1: Enable debug for std. translucency 2: Enable for separated translucency."),
+	ECVF_RenderThreadSafe);
+
+static TAutoConsoleVariable<int32> CVarOIT_SortedPixels_Method(
+	TEXT("r.OIT.SortedPixels.Method"),
+	1,
+	TEXT("Toggle OIT methods 0: Regular alpha-blending (i.e., no OIT) 1: MLAB"),
+	ECVF_RenderThreadSafe);
+
+static TAutoConsoleVariable<float> CVarOIT_SortedPixels_TransmittanceThreshold(
+	TEXT("r.OIT.SortedPixels.TransmittanceThreshold"),
+	0.05,
+	TEXT("Remove translucent rendering surfaces when the accumulated thransmittance is below this threshold"),
+	ECVF_RenderThreadSafe);
 
-static int32 GOIT_Debug = 0;
-static FAutoConsoleVariableRef CVarOIT_Debug(TEXT("r.OIT.Debug"), GOIT_Debug, TEXT("Enable per-instance triangle sorting debug rendering."));
+///////////////////////////////////////////////////////////////////////////////////////////////////
 
-static int32 GOIT_Pool = 0;
-static int32 GOIT_PoolReleaseThreshold = 100;
-static FAutoConsoleVariableRef CVarOIT_Pool(TEXT("r.OIT.Pool"), GOIT_Pool, TEXT("Enable index buffer pool allocation which reduce creation/deletion time by re-use buffers."));
-static FAutoConsoleVariableRef CVarOIT_PoolReleaseThreshold(TEXT("r.OIT.Pool.ReleaseFrameThreshold"), GOIT_PoolReleaseThreshold, TEXT("Number of frame after which unused buffer are released."));
+static bool IsOITSortedTrianglesSupported(EShaderPlatform InShaderPlatform)
+{
+	return RHISupportsComputeShaders(InShaderPlatform);
+}
 
-///////////////////////////////////////////////////////////////////////////////////////////////////
-// OIT Debug
+static bool IsOITSortedPixelsSupported(EShaderPlatform InShaderPlatform)
+{
+	return RHISupportsComputeShaders(InShaderPlatform) && FDataDrivenShaderPlatformInfo::GetSupportsOIT(InShaderPlatform);
+}
 
-struct FOITDebugData
+///////////////////////////////////////////////////////////////////////////////////////////////////
+// Debug OIT
+class FOITPixelDebugCS : public FGlobalShader
 {
-	static const EPixelFormat Format = PF_R32_UINT;
+	DECLARE_GLOBAL_SHADER(FOITPixelDebugCS);
+	SHADER_USE_PARAMETER_STRUCT(FOITPixelDebugCS, FGlobalShader);
 
-	FRDGBufferRef Buffer = nullptr; // First element is counter, then element are: (NumPrim/Type/Size)
+	BEGIN_SHADER_PARAMETER_STRUCT(FParameters, )
+		SHADER_PARAMETER(FIntPoint, CursorCoord)
+		SHADER_PARAMETER(FIntPoint, Resolution)
+		SHADER_PARAMETER(uint32, MaxSideSampleCount)
+		SHADER_PARAMETER(uint32, MaxSampleCount)
+		SHADER_PARAMETER(uint32, Method)
+		SHADER_PARAMETER(uint32, PassType)
+		SHADER_PARAMETER(uint32, SupportedPass)
+		SHADER_PARAMETER_RDG_TEXTURE(Texture2D<float3>, SampleColorTexture)
+		SHADER_PARAMETER_RDG_TEXTURE(Texture2D<float3>, SampleTransTexture)
+		SHADER_PARAMETER_RDG_TEXTURE(Texture2D<float>, SampleDepthTexture)
+		SHADER_PARAMETER_RDG_TEXTURE(Texture2D<uint>, SampleCountTexture)
+		SHADER_PARAMETER_STRUCT_INCLUDE(ShaderPrint::FShaderParameters, ShaderPrintUniformBuffer)
+		END_SHADER_PARAMETER_STRUCT()
+public:
+	static bool ShouldCompilePermutation(const FGlobalShaderPermutationParameters& Parameters) { return IsOITSortedPixelsSupported(Parameters.Platform); }
+	static void ModifyCompilationEnvironment(const FGlobalShaderPermutationParameters& Parameters, FShaderCompilerEnvironment& OutEnvironment)
+	{
+		FGlobalShader::ModifyCompilationEnvironment(Parameters, OutEnvironment);
+		OutEnvironment.SetDefine(TEXT("SHADER_OIT_DEBUG"), 1);
+	}
+};
 
-	uint32 VisibleInstances = 0;
-	uint32 VisiblePrimitives = 0;
-	uint32 VisibleIndexSizeInBytes = 0;
+IMPLEMENT_GLOBAL_SHADER(FOITPixelDebugCS, "/Engine/Private/OITCombine.usf", "MainCS", SF_Compute);
 
-	uint32 AllocatedBuffers = 0;
-	uint32 AllocatedIndexSizeInBytes = 0;
+static void AddOITPixelDebugPass(
+	FRDGBuilder& GraphBuilder,
+	const FViewInfo& View,
+	FOITData& OITData)
+{
+	if (!OITData.SampleColorTexture ||
+		!OITData.SampleTransTexture ||
+		!OITData.SampleDepthTexture ||
+		!OITData.SampleCountTexture)
+		return;
 
-	uint32 UnusedBuffers = 0;
-	uint32 UnusedIndexSizeInBytes = 0;
+	if (!ShaderPrint::IsEnabled(View))
+	{
+		ShaderPrint::SetEnabled(true);
+	}
+	ShaderPrint::RequestSpaceForCharacters(512);
+
+	FOITPixelDebugCS::FParameters* Parameters = GraphBuilder.AllocParameters<FOITPixelDebugCS::FParameters>();
+	Parameters->PassType = OITData.PassType;
+	Parameters->SupportedPass = OITData.SupportedPass;
+	Parameters->Resolution = View.ViewRect.Size();
+	Parameters->CursorCoord = View.CursorPos;
+	Parameters->MaxSampleCount = OITData.MaxSamplePerPixel;
+	Parameters->Method = OITData.Method;
+	Parameters->MaxSideSampleCount = OITData.MaxSideSamplePerPixel;
+	Parameters->SampleColorTexture = OITData.SampleColorTexture;
+	Parameters->SampleTransTexture = OITData.SampleTransTexture;
+	Parameters->SampleDepthTexture = OITData.SampleDepthTexture;
+	Parameters->SampleCountTexture = OITData.SampleCountTexture;
+	ShaderPrint::SetParameters(GraphBuilder, View, Parameters->ShaderPrintUniformBuffer);
+
+	FOITPixelDebugCS::FPermutationDomain PermutationVector;
+	TShaderMapRef<FOITPixelDebugCS> ComputeShader(View.ShaderMap, PermutationVector);
 
-	uint32 TotalEntries = 0;
+	FComputeShaderUtils::AddPass(
+		GraphBuilder,
+		RDG_EVENT_NAME("Translucency::OITDebug(Pixel,%s)", !!(OITData.PassType & OITPass_SeperateTranslucency) ? TEXT("SeparateTransluency") : TEXT("Regular")),
+		ComputeShader,
+		Parameters,
+		FIntVector(1, 1, 1));
+}
 
-	bool IsValid() const { return Buffer != nullptr; }
+///////////////////////////////////////////////////////////////////////////////////////////////////
+// Combine OIT samples
+class FOITPixelCombineCS : public FGlobalShader
+{
+	DECLARE_GLOBAL_SHADER(FOITPixelCombineCS);
+	SHADER_USE_PARAMETER_STRUCT(FOITPixelCombineCS, FGlobalShader);
+
+	BEGIN_SHADER_PARAMETER_STRUCT(FParameters, )
+		SHADER_PARAMETER(FIntPoint, Resolution)
+		SHADER_PARAMETER(uint32, MaxSideSampleCount)
+		SHADER_PARAMETER(uint32, MaxSampleCount)
+		SHADER_PARAMETER(uint32, Method)
+		SHADER_PARAMETER(uint32, PassType)
+		SHADER_PARAMETER_RDG_TEXTURE(Texture2D<float3>, SampleColorTexture)
+		SHADER_PARAMETER_RDG_TEXTURE(Texture2D<float3>, SampleTransTexture)
+		SHADER_PARAMETER_RDG_TEXTURE(Texture2D<float>, SampleDepthTexture)
+		SHADER_PARAMETER_RDG_TEXTURE(Texture2D<uint>, SampleCountTexture)
+		SHADER_PARAMETER_RDG_TEXTURE_UAV(RWTexture2D<float4>, OutColorTexture)
+		END_SHADER_PARAMETER_STRUCT()
+public:
+	static bool ShouldCompilePermutation(const FGlobalShaderPermutationParameters& Parameters) { return IsOITSortedPixelsSupported(Parameters.Platform); }
+	static void ModifyCompilationEnvironment(const FGlobalShaderPermutationParameters& Parameters, FShaderCompilerEnvironment& OutEnvironment)
+	{
+		FGlobalShader::ModifyCompilationEnvironment(Parameters, OutEnvironment);
+		OutEnvironment.CompilerFlags.Add(CFLAG_AllowTypedUAVLoads);
+		OutEnvironment.SetDefine(TEXT("SHADER_OIT_COMBINE"), 1);
+	}
 };
 
+IMPLEMENT_GLOBAL_SHADER(FOITPixelCombineCS, "/Engine/Private/OITCombine.usf", "MainCS", SF_Compute);
+
+static void AddInternalOITComposePass(
+	FRDGBuilder& GraphBuilder,
+	const FViewInfo& View,
+	FOITData& OITData,
+	FRDGTextureRef SceneColorTexture)
+{
+	// TODO: Add tile composition based on the coarse translucent raster AABB buffer
+
+	if (!SceneColorTexture ||
+		 SceneColorTexture->Desc.NumSamples > 1 ||
+		!OITData.SampleColorTexture ||
+		!OITData.SampleTransTexture ||
+		!OITData.SampleDepthTexture ||
+		!OITData.SampleCountTexture)
+		return;
+
+	FIntPoint InResolution = OITData.SampleColorTexture->Desc.Extent;
+	FIntPoint OutResolution = SceneColorTexture->Desc.Extent;
+	FIntPoint Resolution = FIntPoint(FMath::Min(InResolution.X, OutResolution.X), FMath::Min(InResolution.Y, OutResolution.Y));
+
+	FOITPixelCombineCS::FParameters* Parameters = GraphBuilder.AllocParameters<FOITPixelCombineCS::FParameters>();
+	Parameters->PassType = OITData.PassType;
+	Parameters->Resolution = Resolution;
+	Parameters->Method = OITData.Method;
+	Parameters->MaxSampleCount = OITData.MaxSamplePerPixel;
+	Parameters->MaxSideSampleCount = OITData.MaxSideSamplePerPixel;
+	Parameters->SampleColorTexture = OITData.SampleColorTexture;
+	Parameters->SampleTransTexture = OITData.SampleTransTexture;
+	Parameters->SampleDepthTexture = OITData.SampleDepthTexture;
+	Parameters->SampleCountTexture = OITData.SampleCountTexture;
+	Parameters->OutColorTexture = GraphBuilder.CreateUAV(SceneColorTexture);
+
+	FOITPixelCombineCS::FPermutationDomain PermutationVector;
+	TShaderMapRef<FOITPixelCombineCS > ComputeShader(View.ShaderMap, PermutationVector);
+
+	// Add 64 threads permutation
+	const uint32 GroupSize = 8;
+	const FIntVector DispatchCount = FIntVector(
+		(Resolution.X + GroupSize - 1) / GroupSize,
+		(Resolution.Y + GroupSize - 1) / GroupSize,
+		1);
+	FComputeShaderUtils::AddPass(
+		GraphBuilder,
+		RDG_EVENT_NAME("Translucency::OITCombine(%s)", !!(OITData.PassType & OITPass_SeperateTranslucency) ? TEXT("SeparateTransluency") : TEXT("Regular")),
+		ComputeShader,
+		Parameters,
+		DispatchCount);
+
+	// Add debug pass
+	const uint32 ActiveDebugPass = CVarOIT_SortedPixels_Debug.GetValueOnRenderThread();
+	if (ActiveDebugPass == OITData.PassType)
+	{
+		AddOITPixelDebugPass(GraphBuilder, View, OITData);
+	}
+}
+
 ///////////////////////////////////////////////////////////////////////////////////////////////////
 // OIT Index buffer
 
@@ -118,7 +321,7 @@ static void TrimSortedIndexBuffers(TArray<FSortedIndexBuffer*>& FreeBuffers, uin
 		if (LastFrameId != 0)
 		{
 			const int32 ElapsedFrame = FMath::Abs(int32(FrameId) - int32(LastFrameId));
-			if (ElapsedFrame > GOIT_PoolReleaseThreshold)
+			if (ElapsedFrame > CVarOIT_SortedTriangles_PoolReleaseThreshold.GetValueOnRenderThread())
 			{
 				FSortedIndexBuffer* Buffer = FreeBuffers[FreeIt];
 				RemoveAllocation(Buffer);
@@ -155,7 +358,7 @@ FSortedTriangleData FOITSceneData::Allocate(const FIndexBuffer* InSource, EPrimi
 	// Linear scan if there are some free resource which are large enough
 	const uint32 NumIndices = InNumPrimitives * 3; // Sorted index always has triangle list topology
 	FSortedIndexBuffer* OITIndexBuffer = nullptr;
-	if (GOIT_Pool>0)
+	if (CVarOIT_SortedTriangles_Pool.GetValueOnRenderThread() > 0)
 	{
 		for (uint32 FreeIt=0,FreeCount=FreeBuffers.Num(); FreeIt<FreeCount; ++FreeIt)
 		{		
@@ -202,7 +405,7 @@ void FOITSceneData::Deallocate(FIndexBuffer* InIndexBuffer)
 	const uint32 Slot = OITIndexBuffer->Id;
 	if (Slot < uint32(Allocations.Num()))
 	{
-		if (GOIT_Pool > 0)
+		if (CVarOIT_SortedTriangles_Pool.GetValueOnAnyThread() > 0)
 		{
 			OITIndexBuffer->Id = FSortedIndexBuffer::InvalidId;
 			OITIndexBuffer->LastUsedFrameId = FrameIndex;
@@ -278,7 +481,7 @@ class FOITSortTriangleIndex_ScanCS : public FGlobalShader
 
 	END_SHADER_PARAMETER_STRUCT()
 public:
-	static bool ShouldCompilePermutation(const FGlobalShaderPermutationParameters& Parameters) { return IsOITSupported(Parameters.Platform); }
+	static bool ShouldCompilePermutation(const FGlobalShaderPermutationParameters& Parameters) { return IsOITSortedTrianglesSupported(Parameters.Platform); }
 	static void ModifyCompilationEnvironment(const FGlobalShaderPermutationParameters& Parameters, FShaderCompilerEnvironment& OutEnvironment)
 	{
 		FGlobalShader::ModifyCompilationEnvironment(Parameters, OutEnvironment);
@@ -301,7 +504,7 @@ class FOITSortTriangleIndex_AllocateCS : public FGlobalShader
 		SHADER_PARAMETER_RDG_BUFFER_UAV(RWBuffer<uint>, SliceOffsetsBuffer)
 	END_SHADER_PARAMETER_STRUCT()
 public:
-	static bool ShouldCompilePermutation(const FGlobalShaderPermutationParameters& Parameters) { return IsOITSupported(Parameters.Platform); }
+	static bool ShouldCompilePermutation(const FGlobalShaderPermutationParameters& Parameters) { return IsOITSortedTrianglesSupported(Parameters.Platform); }
 	static void ModifyCompilationEnvironment(const FGlobalShaderPermutationParameters& Parameters, FShaderCompilerEnvironment& OutEnvironment)
 	{
 		FGlobalShader::ModifyCompilationEnvironment(Parameters, OutEnvironment);
@@ -333,7 +536,7 @@ class FOITSortTriangleIndex_WriteOutCS : public FGlobalShader
 		SHADER_PARAMETER_UAV(RWBuffer<uint>, OutIndexBuffer)
 	END_SHADER_PARAMETER_STRUCT()
 public:
-	static bool ShouldCompilePermutation(const FGlobalShaderPermutationParameters& Parameters) { return IsOITSupported(Parameters.Platform); }
+	static bool ShouldCompilePermutation(const FGlobalShaderPermutationParameters& Parameters) { return IsOITSortedTrianglesSupported(Parameters.Platform); }
 	static void ModifyCompilationEnvironment(const FGlobalShaderPermutationParameters& Parameters, FShaderCompilerEnvironment& OutEnvironment)
 	{
 		FGlobalShader::ModifyCompilationEnvironment(Parameters, OutEnvironment);
@@ -364,7 +567,7 @@ class FOITSortTriangleIndex_Debug : public FGlobalShader
 		SHADER_PARAMETER_RDG_BUFFER_SRV(Buffer<uint>, DebugData)
 	END_SHADER_PARAMETER_STRUCT()
 public:
-	static bool ShouldCompilePermutation(const FGlobalShaderPermutationParameters& Parameters) { return IsOITSupported(Parameters.Platform); }
+	static bool ShouldCompilePermutation(const FGlobalShaderPermutationParameters& Parameters) { return IsOITSortedTrianglesSupported(Parameters.Platform); }
 	static void ModifyCompilationEnvironment(const FGlobalShaderPermutationParameters& Parameters, FShaderCompilerEnvironment& OutEnvironment)
 	{
 		FGlobalShader::ModifyCompilationEnvironment(Parameters, OutEnvironment);
@@ -377,12 +580,33 @@ IMPLEMENT_GLOBAL_SHADER(FOITSortTriangleIndex_Debug, "/Engine/Private/OIT/OITSor
 
 ///////////////////////////////////////////////////////////////////////////////////////////////////
 
+struct FOITDebugData
+{
+	static const EPixelFormat Format = PF_R32_UINT;
+
+	FRDGBufferRef Buffer = nullptr; // First element is counter, then element are: (NumPrim/Type/Size)
+
+	uint32 VisibleInstances = 0;
+	uint32 VisiblePrimitives = 0;
+	uint32 VisibleIndexSizeInBytes = 0;
+
+	uint32 AllocatedBuffers = 0;
+	uint32 AllocatedIndexSizeInBytes = 0;
+
+	uint32 UnusedBuffers = 0;
+	uint32 UnusedIndexSizeInBytes = 0;
+
+	uint32 TotalEntries = 0;
+
+	bool IsValid() const { return Buffer != nullptr; }
+};
+
 static void AddOITSortTriangleIndexPass(
 	FRDGBuilder& GraphBuilder,
 	const FViewInfo& View,
 	const FOITSceneData& OITSceneData,
 	const FSortedTrianglesMeshBatch& MeshBatch,
-	FOITSortingType SortType,
+	FTriangleSortingOrder SortType,
 	FOITDebugData& DebugData)
 {
 	static const FVertexFactoryType* CompatibleVF = FVertexFactoryType::GetVFByName(TEXT("FLocalVertexFactory"));
@@ -440,7 +664,7 @@ static void AddOITSortTriangleIndexPass(
 		Parameters->NumIndices				= Allocation.NumIndices;
 		Parameters->ViewBoundMinZ			= ViewBoundMinZ;
 		Parameters->ViewBoundMaxZ			= ViewBoundMaxZ;
-		Parameters->SortType				= SortType == FOITSortingType::BackToFront ? 0 : 1;
+		Parameters->SortType				= SortType == FTriangleSortingOrder::BackToFront ? 0 : 1;
 		Parameters->SortedIndexBufferSizeInByte = Allocation.SortedIndexBuffer->IndexBufferRHI->GetSize();
 		Parameters->PositionBuffer			= VertexPosition;
 		Parameters->SourceFirstIndex		= Allocation.SourceFirstIndex;
@@ -542,7 +766,7 @@ static void AddOITSortTriangleIndexPass(
 	// * Batch Scan/Alloc/Write of several primitive, so that we have better overlapping
 }
 
-static void AddOITDebugPass(
+static void AddOITTriangleDebugPass(
 	FRDGBuilder& GraphBuilder,
 	const FViewInfo& View,
 	const FOITDebugData& DebugData)
@@ -552,29 +776,28 @@ static void AddOITDebugPass(
 		return;
 	}
 
+	if (ShaderPrint::IsEnabled(View))
+	{
+		ShaderPrint::SetEnabled(true);
+	}
+	ShaderPrint::RequestSpaceForCharacters(512);
+
 	FOITSortTriangleIndex_Debug::FParameters* Parameters = GraphBuilder.AllocParameters<FOITSortTriangleIndex_Debug::FParameters>();
 	Parameters->VisibleInstances = DebugData.VisibleInstances;
 	Parameters->VisiblePrimitives = DebugData.VisiblePrimitives;
 	Parameters->VisibleIndexSizeInBytes = DebugData.VisibleIndexSizeInBytes;
-
 	Parameters->UnusedBuffers = DebugData.UnusedBuffers;
 	Parameters->UnusedIndexSizeInBytes = DebugData.UnusedIndexSizeInBytes;
-
 	Parameters->AllocatedBuffers = DebugData.AllocatedBuffers;
 	Parameters->AllocatedIndexSizeInBytes = DebugData.AllocatedIndexSizeInBytes;
-
 	Parameters->TotalEntries = DebugData.TotalEntries;
-
 	Parameters->DebugData = GraphBuilder.CreateSRV(DebugData.Buffer, FOITDebugData::Format);
-	if (ShaderPrint::IsEnabled(View))
-	{
-		ShaderPrint::SetParameters(GraphBuilder, View, Parameters->ShaderPrintParameters);
-	}
+	ShaderPrint::SetParameters(GraphBuilder, View, Parameters->ShaderPrintParameters);
 
 	TShaderMapRef<FOITSortTriangleIndex_Debug> ComputeShader(View.ShaderMap);
 	FComputeShaderUtils::AddPass(
 		GraphBuilder,
-		RDG_EVENT_NAME("OIT::Debug"),
+		RDG_EVENT_NAME("OIT::Debug(Triangle)"),
 		ComputeShader,
 		Parameters,
 		FIntVector(1u, 1u, 1u));
@@ -584,14 +807,38 @@ static void AddOITDebugPass(
 
 namespace OIT
 {
-	bool IsEnabled(const FViewInfo& View)
+	bool IsEnabled(EOITSortingType Type)
 	{
-		return GOIT_SortObjectTriangles > 0 && IsOITSupported(View.GetShaderPlatform());
+		switch (Type)
+		{
+		case EOITSortingType::SortedTriangles: return CVarOIT_SortedTriangles_Enable.GetValueOnAnyThread() > 0;
+		case EOITSortingType::SortedPixels	 : return CVarOIT_SortedPixels_Enable.GetValueOnAnyThread() > 0;
+		}
+		return false;
 	}
 
-	bool IsEnabled(EShaderPlatform ShaderPlatform)
+	bool IsEnabled(EOITSortingType Type, const FViewInfo& View)
 	{
-		return GOIT_SortObjectTriangles > 0 && IsOITSupported(ShaderPlatform);
+		switch (Type)
+		{
+		case EOITSortingType::SortedTriangles: return CVarOIT_SortedTriangles_Enable.GetValueOnAnyThread() > 0 && IsOITSortedTrianglesSupported(View.GetShaderPlatform());
+		case EOITSortingType::SortedPixels	 : return CVarOIT_SortedPixels_Enable.GetValueOnAnyThread() > 0 && FDataDrivenShaderPlatformInfo::GetSupportsOIT(View.GetShaderPlatform());
+		}
+		return false;
+
+		
+	}
+
+	bool IsEnabled(EOITSortingType Type, EShaderPlatform InPlatform)
+	{
+		const bool bMSAAEnabled = GetDefaultAntiAliasingMethod(GetMaxSupportedFeatureLevel(InPlatform)) != EAntiAliasingMethod::AAM_MSAA;
+
+		switch (Type)
+		{
+			case EOITSortingType::SortedTriangles: return CVarOIT_SortedTriangles_Enable.GetValueOnAnyThread() > 0 && IsOITSortedTrianglesSupported(InPlatform);
+			case EOITSortingType::SortedPixels	 : return CVarOIT_SortedPixels_Enable.GetValueOnAnyThread() > 0 && FDataDrivenShaderPlatformInfo::GetSupportsOIT(EShaderPlatform(InPlatform));
+		}
+		return false;
 	}
 
 	bool IsCompatible(const FMeshBatch& InMesh, ERHIFeatureLevel::Type InFeatureLevel)
@@ -606,16 +853,16 @@ namespace OIT
 		return false;
 	}
 
-	void AddSortTrianglesPass(FRDGBuilder& GraphBuilder, const FViewInfo& View, FOITSceneData& OITSceneData, FOITSortingType SortType)
+	void AddSortTrianglesPass(FRDGBuilder& GraphBuilder, const FViewInfo& View, FOITSceneData& OITSceneData, FTriangleSortingOrder SortType)
 	{
-		if (!IsEnabled(View))
+		if (!IsEnabled(EOITSortingType::SortedTriangles))
 		{
 			return;
 		}
 
 		RDG_EVENT_SCOPE(GraphBuilder, "OIT::IndexSorting");
 
-		const bool bDebugEnable = GOIT_Debug > 0;
+		const bool bDebugEnable = CVarOIT_SortedTriangles_Debug.GetValueOnRenderThread() > 0;
 		FOITDebugData DebugData;
 		if (bDebugEnable)
 		{
@@ -650,12 +897,12 @@ namespace OIT
 
 		if (DebugData.IsValid())
 		{
-			AddOITDebugPass(GraphBuilder, View, DebugData);
+			AddOITTriangleDebugPass(GraphBuilder, View, DebugData);
 		}
 
 		// Trim unused buffers
 		OITSceneData.FrameIndex = View.Family->FrameNumber;
-		if (GOIT_Pool > 0 && GOIT_PoolReleaseThreshold > 0)
+		if (CVarOIT_SortedTriangles_Pool.GetValueOnRenderThread() > 0 && CVarOIT_SortedTriangles_PoolReleaseThreshold.GetValueOnRenderThread() > 0)
 		{
 			TrimSortedIndexBuffers(OITSceneData.FreeBuffers, OITSceneData.FrameIndex);
 		}
@@ -668,4 +915,81 @@ namespace OIT
 		Out->FirstIndex = In->SortedFirstIndex;
 		Out->PrimitiveType = In->SortedPrimitiveType;
 	}
-}
+
+	FOITData CreateOITData(FRDGBuilder& GraphBuilder, const FViewInfo& View, EOITPassType PassType)
+	{
+		const bool bOIT = IsEnabled(EOITSortingType::SortedPixels, View.GetShaderPlatform());
+		const uint32 PassTypeBits = FMath::Clamp(CVarOIT_SortedPixels_PassType.GetValueOnRenderThread(), 0, 3);
+		const bool bPassValid = !!(PassTypeBits & PassType);
+
+		FOITData Out;
+		if (!bOIT || !bPassValid)
+		{
+			Out.MaxSideSamplePerPixel = 0;
+			Out.MaxSamplePerPixel = 0;
+			Out.Method = 0;
+			Out.TransmittanceThreshold = 0;
+
+			Out.SampleColorTexture = GraphBuilder.CreateTexture(FRDGTextureDesc::Create2D(FIntPoint(1, 1), PF_R32_UINT, FClearValueBinding::None, TexCreate_ShaderResource | TexCreate_UAV), TEXT("OIT.SampleColor"));
+			Out.SampleTransTexture = GraphBuilder.CreateTexture(FRDGTextureDesc::Create2D(FIntPoint(1, 1), PF_R32_UINT, FClearValueBinding::None, TexCreate_ShaderResource | TexCreate_UAV), TEXT("OIT.SampleTrans"));
+			Out.SampleDepthTexture = GraphBuilder.CreateTexture(FRDGTextureDesc::Create2D(FIntPoint(1, 1), PF_R32_FLOAT, FClearValueBinding::Black, TexCreate_ShaderResource | TexCreate_UAV), TEXT("OIT.SampleDepth"));
+			Out.SampleCountTexture = GraphBuilder.CreateTexture(FRDGTextureDesc::Create2D(FIntPoint(1, 1), PF_R32_UINT, FClearValueBinding::None, TexCreate_ShaderResource | TexCreate_UAV), TEXT("OIT.SampleCount"));
+			return Out;
+		}
+
+		// Scene render targets might not exist yet; avoids NaNs.
+		FIntPoint EffectiveBufferSize = GetSceneTextureExtent();
+		EffectiveBufferSize.X = FMath::Max(EffectiveBufferSize.X, 1);
+		EffectiveBufferSize.Y = FMath::Max(EffectiveBufferSize.Y, 1);
+
+		// Allocate OIT data		
+		Out.PassType = PassTypeBits & PassType ? PassType : OITPass_None;
+		if (PassTypeBits & OITPass_RegularTranslucency)  Out.SupportedPass |= OITPass_RegularTranslucency;
+		if (PassTypeBits & OITPass_SeperateTranslucency) Out.SupportedPass |= OITPass_SeperateTranslucency;
+
+		if (Out.PassType != OITPass_None)
+		{
+			// Round sample to be square to store them into square tile
+			uint32 MaxSamplePerPixel = FMath::Clamp(CVarOIT_SortedPixels_SampleCount.GetValueOnRenderThread(), 1, 16);
+			uint32 MaxSideSamplePerPixel = FMath::FloorToInt(FMath::Sqrt(float(MaxSamplePerPixel)));
+			MaxSamplePerPixel = MaxSideSamplePerPixel * MaxSideSamplePerPixel;
+
+			Out.MaxSideSamplePerPixel = MaxSideSamplePerPixel;
+			Out.MaxSamplePerPixel = MaxSamplePerPixel;
+			Out.Method = FMath::Clamp(CVarOIT_SortedPixels_Method.GetValueOnRenderThread(), 0, 1);
+			Out.TransmittanceThreshold = FMath::Clamp(CVarOIT_SortedPixels_TransmittanceThreshold.GetValueOnRenderThread(), 0.f, 1.f);
+
+			Out.SampleColorTexture = GraphBuilder.CreateTexture(FRDGTextureDesc::Create2D(EffectiveBufferSize * MaxSideSamplePerPixel, PF_R32_UINT, FClearValueBinding::None, TexCreate_ShaderResource | TexCreate_UAV), TEXT("OIT.SampleColor"));
+			Out.SampleTransTexture = GraphBuilder.CreateTexture(FRDGTextureDesc::Create2D(EffectiveBufferSize * MaxSideSamplePerPixel, PF_R32_UINT, FClearValueBinding::None, TexCreate_ShaderResource | TexCreate_UAV), TEXT("OIT.SampleTrans"));
+			Out.SampleDepthTexture = GraphBuilder.CreateTexture(FRDGTextureDesc::Create2D(EffectiveBufferSize * MaxSideSamplePerPixel, PF_R32_FLOAT, FClearValueBinding::Black, TexCreate_ShaderResource | TexCreate_UAV), TEXT("OIT.SampleDepth"));
+			Out.SampleCountTexture = GraphBuilder.CreateTexture(FRDGTextureDesc::Create2D(EffectiveBufferSize, PF_R32_UINT, FClearValueBinding::None, TexCreate_ShaderResource | TexCreate_UAV), TEXT("OIT.SampleCount"));
+		}
+
+		// TODO: Add tile clear based on the coarse translucent raster AABB buffer
+		AddClearUAVPass(GraphBuilder, GraphBuilder.CreateUAV(Out.SampleColorTexture), 0u);
+		AddClearUAVPass(GraphBuilder, GraphBuilder.CreateUAV(Out.SampleTransTexture), 0u);
+		AddClearUAVPass(GraphBuilder, GraphBuilder.CreateUAV(Out.SampleCountTexture), 0u);
+		AddClearUAVPass(GraphBuilder, GraphBuilder.CreateUAV(Out.SampleDepthTexture), 0.f);
+
+		return Out;
+	}
+
+	void SetOITParameters(FRDGBuilder& GraphBuilder, const FViewInfo& View, FOITBasePassUniformParameters& OutOIT, const FOITData& InOITData)
+	{
+		// Always set the parameters, even when the OIT is disabled, as Uniform buffer validation complain about null resources otherwise
+		OutOIT.OITMethod = InOITData.Method;
+		OutOIT.bOITEnable = InOITData.PassType != OITPass_None ? 1 : 0;
+		OutOIT.MaxSideSamplePerPixel = InOITData.MaxSideSamplePerPixel;
+		OutOIT.MaxSamplePerPixel = InOITData.MaxSamplePerPixel;
+		OutOIT.TransmittanceThreshold = InOITData.TransmittanceThreshold;
+		OutOIT.OutOITSampleColor = GraphBuilder.CreateUAV(InOITData.SampleColorTexture);
+		OutOIT.OutOITSampleTrans = GraphBuilder.CreateUAV(InOITData.SampleTransTexture);
+		OutOIT.OutOITSampleDepth = GraphBuilder.CreateUAV(InOITData.SampleDepthTexture);
+		OutOIT.OutOITSampleCount = GraphBuilder.CreateUAV(InOITData.SampleCountTexture);
+	}
+
+	void AddOITComposePass(FRDGBuilder& GraphBuilder, const FViewInfo& View, FOITData& OITData, FRDGTextureRef SceneColorTexture)
+	{
+		AddInternalOITComposePass(GraphBuilder, View, OITData, SceneColorTexture);
+	}
+}
\ No newline at end of file
diff --git a/Engine/Source/Runtime/Renderer/Private/OIT/OIT.h b/Engine/Source/Runtime/Renderer/Private/OIT/OIT.h
--- a/Engine/Source/Runtime/Renderer/Private/OIT/OIT.h
+++ b/Engine/Source/Runtime/Renderer/Private/OIT/OIT.h
@@ -12,6 +12,26 @@ class FSortedIndexBuffer;
 struct FMeshBatchElementDynamicIndexBuffer;
 class FViewInfo;
 
+enum class EOITSortingType
+{
+	SortedTriangles,
+	SortedPixels,
+};
+
+enum EOITPassType
+{
+	OITPass_None = 0,
+	OITPass_RegularTranslucency = 1,
+	OITPass_SeperateTranslucency = 2
+};
+
+enum class FTriangleSortingOrder
+{
+	FrontToBack,
+	BackToFront,
+};
+
+// Sorted triangles data
 struct FSortedTriangleData
 {
 	const FIndexBuffer* SourceIndexBuffer = nullptr;
@@ -31,6 +51,23 @@ struct FSortedTriangleData
 	bool IsValid() const { return SortedIndexBuffer != nullptr; }
 };
 
+// Order Independent Translucency pass data for a given translucent pass
+struct FOITData
+{
+	EOITPassType PassType = OITPass_None;
+	uint32 SupportedPass = 0;
+	uint32 Method = 0;
+	uint32 MaxSamplePerPixel = 0;
+	uint32 MaxSideSamplePerPixel = 0;
+	float TransmittanceThreshold = 0.f;
+
+	FRDGTextureRef SampleColorTexture = nullptr;
+	FRDGTextureRef SampleTransTexture = nullptr;
+	FRDGTextureRef SampleDepthTexture = nullptr;
+	FRDGTextureRef SampleCountTexture = nullptr;
+};
+
+// Order Independent Translucency scene data
 struct FOITSceneData
 {
 	/* Allocate sorted-triangle data for a instance */
@@ -45,25 +82,24 @@ struct FOITSceneData
 	uint32 FrameIndex = 0;
 };
 
-enum class FOITSortingType
-{
-	FrontToBack,
-	BackToFront,
-};
-
 namespace OIT
 {
 	/* Return true if OIT techniques are enabled/supported */
-	bool IsEnabled(const FViewInfo& View);
-	bool IsEnabled(EShaderPlatform ShaderPlatform);
+	bool IsEnabled(EOITSortingType Type, const FViewInfo& View);
+	bool IsEnabled(EOITSortingType Type, EShaderPlatform ShaderPlatform);
 
 	/* Return true if the current MeshBatch is compatible with per-instance sorted triangle */
 	bool IsCompatible(const FMeshBatch& Mesh, ERHIFeatureLevel::Type InFeatureLevel);
 
 	/* Sort triangles of all instances whose has the sorted triangle option enabled */
-	void AddSortTrianglesPass(FRDGBuilder& GraphBuilder, const FViewInfo& View, FOITSceneData& OITSceneData, FOITSortingType SortType);
+	void AddSortTrianglesPass(FRDGBuilder& GraphBuilder, const FViewInfo& View, FOITSceneData& OITSceneData, FTriangleSortingOrder SortType);
 
 	/* Convert FSortedTriangleData into FMeshBatchElementDynamicIndexBuffer */
 	void ConvertSortedIndexToDynamicIndex(FSortedTriangleData* In, FMeshBatchElementDynamicIndexBuffer* Out);
+
+	/* Create OIT data for translucent pass */
+	FOITData CreateOITData(FRDGBuilder& GraphBuilder, const FViewInfo& View, EOITPassType PassType);
+
+	/* Compose all OIT samples into the target color buffer */
+	void AddOITComposePass(FRDGBuilder& GraphBuilder, const FViewInfo& View, FOITData& OITData, FRDGTextureRef SceneColorTexture);
 }
- 
diff --git a/Engine/Source/Runtime/Renderer/Private/OIT/OITParameters.h b/Engine/Source/Runtime/Renderer/Private/OIT/OITParameters.h
new file mode 100644
--- /dev/null
+++ b/Engine/Source/Runtime/Renderer/Private/OIT/OITParameters.h
@@ -0,0 +1,29 @@
+// Copyright Epic Games, Inc. All Rights Reserved.
+
+#pragma once
+
+#include "CoreMinimal.h"
+#include "RendererInterface.h"
+#include "ShaderParameters.h"
+#include "RenderGraphResources.h"
+
+// Uniform buffer used for translucent material shader
+BEGIN_SHADER_PARAMETER_STRUCT(FOITBasePassUniformParameters, )
+	SHADER_PARAMETER(uint32, bOITEnable)
+	SHADER_PARAMETER(uint32, OITMethod)
+	SHADER_PARAMETER(uint32, MaxSideSamplePerPixel)
+	SHADER_PARAMETER(uint32, MaxSamplePerPixel)
+	SHADER_PARAMETER(uint32, TransmittanceThreshold)
+	SHADER_PARAMETER_RDG_TEXTURE_UAV(RasterizerOrderedTexture2D<uint>, OutOITSampleCount)
+	SHADER_PARAMETER_RDG_TEXTURE_UAV(RWTexture2D<uint>,  OutOITSampleColor)
+	SHADER_PARAMETER_RDG_TEXTURE_UAV(RWTexture2D<uint>,  OutOITSampleTrans)
+	SHADER_PARAMETER_RDG_TEXTURE_UAV(RWTexture2D<float>, OutOITSampleDepth)
+END_SHADER_PARAMETER_STRUCT()
+
+class FViewInfo;
+struct FOITData;
+
+namespace OIT
+{
+	void SetOITParameters(FRDGBuilder& GraphBuilder, const FViewInfo& View, FOITBasePassUniformParameters& OutOIT, const FOITData& InOITData);
+}
\ No newline at end of file
diff --git a/Engine/Source/Runtime/Renderer/Private/PrimitiveSceneInfo.cpp b/Engine/Source/Runtime/Renderer/Private/PrimitiveSceneInfo.cpp
--- a/Engine/Source/Runtime/Renderer/Private/PrimitiveSceneInfo.cpp
+++ b/Engine/Source/Runtime/Renderer/Private/PrimitiveSceneInfo.cpp
@@ -1029,7 +1029,7 @@ void FPrimitiveSceneInfo::AddStaticMeshes(FRHICommandListImmediate& RHICmdList,
 		for (FPrimitiveSceneInfo* SceneInfo : SceneInfos)
 		{
 			// Allocate OIT index buffer where needed
-			const bool bAllocateSortedTriangles = OIT::IsEnabled(GMaxRHIShaderPlatform) && SceneInfo->Proxy->SupportsSortedTriangles();
+			const bool bAllocateSortedTriangles = OIT::IsEnabled(EOITSortingType::SortedTriangles, GMaxRHIShaderPlatform) && SceneInfo->Proxy->SupportsSortedTriangles();
 
 			for (int32 MeshIndex = 0; MeshIndex < SceneInfo->StaticMeshes.Num(); MeshIndex++)
 			{
@@ -1492,7 +1492,7 @@ void FPrimitiveSceneInfo::AddToScene(FRHICommandListImmediate& RHICmdList, FScen
 void FPrimitiveSceneInfo::RemoveStaticMeshes()
 {
 	// Deallocate potential OIT dynamic index buffer
-	if (OIT::IsEnabled(GMaxRHIShaderPlatform))
+	if (OIT::IsEnabled(EOITSortingType::SortedTriangles, GMaxRHIShaderPlatform))
 	{
 		for (int32 MeshIndex = 0; MeshIndex < StaticMeshes.Num(); MeshIndex++)
 		{
diff --git a/Engine/Source/Runtime/Renderer/Private/TranslucentRendering.cpp b/Engine/Source/Runtime/Renderer/Private/TranslucentRendering.cpp
--- a/Engine/Source/Runtime/Renderer/Private/TranslucentRendering.cpp
+++ b/Engine/Source/Runtime/Renderer/Private/TranslucentRendering.cpp
@@ -17,6 +17,8 @@
 #include "Strata/Strata.h"
 #include "HairStrands/HairStrandsUtils.h"
 #include "PixelShaderUtils.h"
+#include "OIT/OIT.h"
+#include "OIT/OITParameters.h"
 
 DECLARE_CYCLE_STAT(TEXT("TranslucencyTimestampQueryFence Wait"), STAT_TranslucencyTimestampQueryFence_Wait, STATGROUP_SceneRendering);
 DECLARE_CYCLE_STAT(TEXT("TranslucencyTimestampQuery Wait"), STAT_TranslucencyTimestampQuery_Wait, STATGROUP_SceneRendering);
@@ -835,13 +837,15 @@ FRDGTextureMSAA CreatePostDOFTranslucentTexture(
 	FRDGBuilder& GraphBuilder,
 	ETranslucencyPass::Type TranslucencyPass,
 	FSeparateTranslucencyDimensions& SeparateTranslucencyDimensions,
-	bool bIsModulate)
+	bool bIsModulate,
+	EShaderPlatform ShaderPlatform)
 {
+	const bool bNeedUAV = SeparateTranslucencyDimensions.NumSamples == 1 && OIT::IsEnabled(EOITSortingType::SortedPixels, ShaderPlatform);
 	const FRDGTextureDesc Desc = FRDGTextureDesc::Create2D(
 		SeparateTranslucencyDimensions.Extent,
 		bIsModulate ? PF_FloatR11G11B10 : PF_FloatRGBA,
 		bIsModulate ? FClearValueBinding::White : FClearValueBinding::Black,
-		TexCreate_RenderTargetable | TexCreate_ShaderResource,
+		TexCreate_RenderTargetable | TexCreate_ShaderResource | (bNeedUAV ? TexCreate_UAV : TexCreate_None),
 		1,
 		SeparateTranslucencyDimensions.NumSamples);
 
@@ -893,7 +897,8 @@ TRDGUniformBufferRef<FTranslucentBasePassUniformParameters> CreateTranslucentBas
 	const FTranslucencyLightingVolumeTextures& TranslucencyLightingVolumeTextures,
 	FRDGTextureRef SceneColorCopyTexture,
 	const ESceneTextureSetupMode SceneTextureSetupMode,
-	bool bLumenGIEnabled)
+	bool bLumenGIEnabled,
+	const FOITData& OITData)
 {
 	FTranslucentBasePassUniformParameters& BasePassParameters = *GraphBuilder.AllocParameters<FTranslucentBasePassUniformParameters>();
 
@@ -1068,9 +1073,34 @@ TRDGUniformBufferRef<FTranslucentBasePassUniformParameters> CreateTranslucentBas
 	BasePassParameters.PreIntegratedGFTexture = GSystemTextures.PreintegratedGF->GetRenderTargetItem().ShaderResourceTexture;
 	BasePassParameters.PreIntegratedGFSampler = TStaticSamplerState<SF_Bilinear, AM_Clamp, AM_Clamp, AM_Clamp>::GetRHI();
 
+	OIT::SetOITParameters(GraphBuilder, View, BasePassParameters.OIT, OITData);
+
 	return GraphBuilder.CreateUniformBuffer(&BasePassParameters);
 }
 
+TRDGUniformBufferRef<FTranslucentBasePassUniformParameters> CreateTranslucentBasePassUniformBuffer(
+	FRDGBuilder& GraphBuilder,
+	const FScene* Scene,
+	const FViewInfo& View,
+	const int32 ViewIndex,
+	const FTranslucencyLightingVolumeTextures& TranslucencyLightingVolumeTextures,
+	FRDGTextureRef SceneColorCopyTexture,
+	const ESceneTextureSetupMode SceneTextureSetupMode,
+	bool bLumenGIEnabled)
+{
+	FOITData OITData = OIT::CreateOITData(GraphBuilder, View, OITPass_None);
+	return CreateTranslucentBasePassUniformBuffer(
+		GraphBuilder,
+		Scene,
+		View,
+		ViewIndex,
+		TranslucencyLightingVolumeTextures,
+		SceneColorCopyTexture,
+		SceneTextureSetupMode,
+		bLumenGIEnabled,
+		OITData);
+}
+
 static FViewShaderParameters GetSeparateTranslucencyViewParameters(const FViewInfo& View, FIntPoint TextureExtent, float ViewportScale, ETranslucencyPass::Type TranslucencyPass)
 {
 	FViewShaderParameters ViewParameters;
@@ -1352,7 +1382,7 @@ void FDeferredShadingSceneRenderer::RenderTranslucencyInner(
 	if (bRenderInSeparateTranslucency)
 	{
 		// Create resources shared by each view (each view data is tiled into each of the render target resources)
-		FRDGTextureMSAA SharedColorTexture = CreatePostDOFTranslucentTexture(GraphBuilder, TranslucencyPass, SeparateTranslucencyDimensions, bIsModulate);
+		FRDGTextureMSAA SharedColorTexture = CreatePostDOFTranslucentTexture(GraphBuilder, TranslucencyPass, SeparateTranslucencyDimensions, bIsModulate, ShaderPlatform);
 
 		for (int32 ViewIndex = 0, NumProcessedViews = 0; ViewIndex < Views.Num(); ++ViewIndex, ++NumProcessedViews)
 		{
@@ -1392,6 +1422,8 @@ void FDeferredShadingSceneRenderer::RenderTranslucencyInner(
 				? ERenderTargetLoadAction::EClear
 				: ERenderTargetLoadAction::ELoad;
 
+			FOITData OITData = OIT::CreateOITData(GraphBuilder, View, OITPass_SeperateTranslucency);
+
 			RenderTranslucencyViewInner(
 				GraphBuilder,
 				*this,
@@ -1401,7 +1433,7 @@ void FDeferredShadingSceneRenderer::RenderTranslucencyInner(
 				SeparateTranslucencyColorTexture,
 				SeparateTranslucencyColorLoadAction,
 				SeparateTranslucencyDepthTexture.Target,
-				CreateTranslucentBasePassUniformBuffer(GraphBuilder, Scene, View, ViewIndex, TranslucentLightingVolumeTextures, SceneColorCopyTexture, SceneTextureSetupMode, bLumenGIEnabled),
+				CreateTranslucentBasePassUniformBuffer(GraphBuilder, Scene, View, ViewIndex, TranslucentLightingVolumeTextures, SceneColorCopyTexture, SceneTextureSetupMode, bLumenGIEnabled, OITData),
 				TranslucencyPass,
 				!bCompositeBackToSceneColor,
 				bRenderInParallel,
@@ -1414,6 +1446,11 @@ void FDeferredShadingSceneRenderer::RenderTranslucencyInner(
 				TranslucencyPassResources.DepthTexture = SharedDepthTexture;
 			}
 
+			if (OITData.PassType & OITPass_SeperateTranslucency)
+			{
+				OIT::AddOITComposePass(GraphBuilder, View, OITData, SeparateTranslucencyColorTexture.Target);
+			}
+
 			if (bCompositeBackToSceneColor)
 			{
 				FRDGTextureRef SeparateTranslucencyDepthResolve = nullptr;
@@ -1485,6 +1522,8 @@ void FDeferredShadingSceneRenderer::RenderTranslucencyInner(
 			const bool bResolveColorTexture = false;
 			const bool bLumenGIEnabled = GetViewPipelineState(View).DiffuseIndirectMethod == EDiffuseIndirectMethod::Lumen;
 
+			FOITData OITData = OIT::CreateOITData(GraphBuilder, View, OITPass_RegularTranslucency);
+
 			RenderTranslucencyViewInner(
 				GraphBuilder,
 				*this,
@@ -1494,12 +1533,17 @@ void FDeferredShadingSceneRenderer::RenderTranslucencyInner(
 				SceneTextures.Color,
 				SceneColorLoadAction,
 				SceneTextures.Depth.Target,
-				CreateTranslucentBasePassUniformBuffer(GraphBuilder, Scene, View, ViewIndex, TranslucentLightingVolumeTextures, SceneColorCopyTexture, SceneTextureSetupMode, bLumenGIEnabled),
+				CreateTranslucentBasePassUniformBuffer(GraphBuilder, Scene, View, ViewIndex, TranslucentLightingVolumeTextures, SceneColorCopyTexture, SceneTextureSetupMode, bLumenGIEnabled, OITData),
 				TranslucencyPass,
 				bResolveColorTexture,
 				bRenderInParallel,
 				InstanceCullingManager);
 
+			if (OITData.PassType & OITPass_RegularTranslucency)
+			{
+				OIT::AddOITComposePass(GraphBuilder, View, OITData, SceneTextures.Color.Target);
+			}
+
 			AddEndTranslucencyTimerPass(GraphBuilder, View);
 		}
 	}
diff --git a/Engine/Source/Runtime/Renderer/Private/TranslucentRendering.h b/Engine/Source/Runtime/Renderer/Private/TranslucentRendering.h
--- a/Engine/Source/Runtime/Renderer/Private/TranslucentRendering.h
+++ b/Engine/Source/Runtime/Renderer/Private/TranslucentRendering.h
@@ -141,7 +141,7 @@ FSeparateTranslucencyDimensions UpdateTranslucencyTimers(FRHICommandListImmediat
 bool ShouldRenderTranslucency(const FSceneViewFamily& ViewFamily);
 
 /** Shared function used to create Post DOF translucent textures */
-FRDGTextureMSAA CreatePostDOFTranslucentTexture(FRDGBuilder& GraphBuilder, ETranslucencyPass::Type TranslucencyPass, FSeparateTranslucencyDimensions& SeparateTranslucencyDimensions, bool bIsModulate);
+FRDGTextureMSAA CreatePostDOFTranslucentTexture(FRDGBuilder& GraphBuilder, ETranslucencyPass::Type TranslucencyPass, FSeparateTranslucencyDimensions& SeparateTranslucencyDimensions, bool bIsModulate, EShaderPlatform ShaderPlatform);
 
 
 /** Add a pass to compose separate translucency. */
