Temporary fix for platform not suporting ROV resource type.

diff --git a/Engine/Source/Runtime/Renderer/Private/BasePassRendering.h b/Engine/Source/Runtime/Renderer/Private/BasePassRendering.h
--- a/Engine/Source/Runtime/Renderer/Private/BasePassRendering.h
+++ b/Engine/Source/Runtime/Renderer/Private/BasePassRendering.h
@@ -100,7 +100,7 @@ BEGIN_GLOBAL_SHADER_PARAMETER_STRUCT(FTranslucentBasePassUniformParameters,)
 	SHADER_PARAMETER_STRUCT(FSceneTextureUniformParameters, SceneTextures)
 	SHADER_PARAMETER_STRUCT(FStrataForwardPassUniformParameters, Strata)
 	SHADER_PARAMETER_STRUCT(FLightCloudTransmittanceParameters, ForwardDirLightCloudShadow)
-	SHADER_PARAMETER_STRUCT(FOITBasePassUniformParameters, OIT)
+	//SHADER_PARAMETER_STRUCT(FOITBasePassUniformParameters, OIT)
 	// Material SSR
 	SHADER_PARAMETER(FVector4f, HZBUvFactorAndInvFactor)
 	SHADER_PARAMETER(FVector4f, PrevScreenPositionScaleBias)
diff --git a/Engine/Source/Runtime/Renderer/Private/TranslucentRendering.cpp b/Engine/Source/Runtime/Renderer/Private/TranslucentRendering.cpp
--- a/Engine/Source/Runtime/Renderer/Private/TranslucentRendering.cpp
+++ b/Engine/Source/Runtime/Renderer/Private/TranslucentRendering.cpp
@@ -1073,7 +1073,7 @@ TRDGUniformBufferRef<FTranslucentBasePassUniformParameters> CreateTranslucentBas
 	BasePassParameters.PreIntegratedGFTexture = GSystemTextures.PreintegratedGF->GetRenderTargetItem().ShaderResourceTexture;
 	BasePassParameters.PreIntegratedGFSampler = TStaticSamplerState<SF_Bilinear, AM_Clamp, AM_Clamp, AM_Clamp>::GetRHI();
 
-	OIT::SetOITParameters(GraphBuilder, View, BasePassParameters.OIT, OITData);
+	//OIT::SetOITParameters(GraphBuilder, View, BasePassParameters.OIT, OITData);
 
 	return GraphBuilder.CreateUniformBuffer(&BasePassParameters);
 }
