Flip HMD extensions to opt-in rather than opt-out

diff --git a/Engine/Source/Editor/UnrealEd/Private/EditorViewportClient.cpp b/Engine/Source/Editor/UnrealEd/Private/EditorViewportClient.cpp
--- a/Engine/Source/Editor/UnrealEd/Private/EditorViewportClient.cpp
+++ b/Engine/Source/Editor/UnrealEd/Private/EditorViewportClient.cpp
@@ -3909,7 +3909,9 @@ void FEditorViewportClient::Draw(FViewport* InViewport, FCanvas* Canvas)
 		}
 	}
 
-	ViewFamily.ViewExtensions = GEngine->ViewExtensions->GatherActiveExtensions(FSceneViewExtensionContext(InViewport));
+	FSceneViewExtensionContext ViewExtensionContext(InViewport);
+	ViewExtensionContext.bStereoEnabled = true;
+	ViewFamily.ViewExtensions = GEngine->ViewExtensions->GatherActiveExtensions(ViewExtensionContext);
 
 	for (auto ViewExt : ViewFamily.ViewExtensions)
 	{
diff --git a/Engine/Source/Runtime/Engine/Private/GameViewportClient.cpp b/Engine/Source/Runtime/Engine/Private/GameViewportClient.cpp
--- a/Engine/Source/Runtime/Engine/Private/GameViewportClient.cpp
+++ b/Engine/Source/Runtime/Engine/Private/GameViewportClient.cpp
@@ -1301,7 +1301,9 @@ void UGameViewportClient::Draw(FViewport* InViewport, FCanvas* SceneCanvas)
 	}
 #endif
 
-	ViewFamily.ViewExtensions = GEngine->ViewExtensions->GatherActiveExtensions(FSceneViewExtensionContext(InViewport));
+	FSceneViewExtensionContext ViewExtensionContext(InViewport);
+	ViewExtensionContext.bStereoEnabled = true;
+	ViewFamily.ViewExtensions = GEngine->ViewExtensions->GatherActiveExtensions(ViewExtensionContext);
 
 	for (auto ViewExt : ViewFamily.ViewExtensions)
 	{
diff --git a/Engine/Source/Runtime/Engine/Private/LocalPlayer.cpp b/Engine/Source/Runtime/Engine/Private/LocalPlayer.cpp
--- a/Engine/Source/Runtime/Engine/Private/LocalPlayer.cpp
+++ b/Engine/Source/Runtime/Engine/Private/LocalPlayer.cpp
@@ -719,7 +719,9 @@ void ULocalPlayer::GetViewPoint(FMinimalViewInfo& OutViewInfo) const
 
 	if (ViewportClient != nullptr)
 	{
-		for (auto& ViewExt : GEngine->ViewExtensions->GatherActiveExtensions(FSceneViewExtensionContext(ViewportClient->Viewport)))
+		FSceneViewExtensionContext SceneViewExtensionContext(ViewportClient->Viewport);
+		SceneViewExtensionContext.bStereoEnabled = true;
+		for (const FSceneViewExtensionRef& ViewExt : GEngine->ViewExtensions->GatherActiveExtensions(SceneViewExtensionContext))
 		{
 			ViewExt->SetupViewPoint(PlayerController, OutViewInfo);
 		};
diff --git a/Engine/Source/Runtime/Engine/Public/SceneViewExtensionContext.h b/Engine/Source/Runtime/Engine/Public/SceneViewExtensionContext.h
--- a/Engine/Source/Runtime/Engine/Public/SceneViewExtensionContext.h
+++ b/Engine/Source/Runtime/Engine/Public/SceneViewExtensionContext.h
@@ -26,7 +26,7 @@ public:
 
 	FViewport* Viewport = nullptr;
 	FSceneInterface* Scene = nullptr;
-	bool bStereoDisabled = false;
+	bool bStereoEnabled = false;
 
 	FSceneViewExtensionContext() : Viewport(nullptr), Scene(nullptr) {}
 	explicit FSceneViewExtensionContext(FViewport* InViewport) : Viewport(InViewport) {}
@@ -61,7 +61,7 @@ public:
 
 	bool IsStereoSupported() const
 	{
-		return !bStereoDisabled && GEngine && GEngine->IsStereoscopic3D(Viewport);
+		return bStereoEnabled && GEngine && GEngine->IsStereoscopic3D(Viewport);
 	}
 
 	// Return true, if HMD supported
diff --git a/Engine/Source/Runtime/Renderer/Private/ReflectionEnvironmentCapture.cpp b/Engine/Source/Runtime/Renderer/Private/ReflectionEnvironmentCapture.cpp
--- a/Engine/Source/Runtime/Renderer/Private/ReflectionEnvironmentCapture.cpp
+++ b/Engine/Source/Runtime/Renderer/Private/ReflectionEnvironmentCapture.cpp
@@ -1461,9 +1461,7 @@ void CaptureSceneIntoScratchCubemap(
 			ViewFamily, /* GlobalResolutionFraction = */ 1.0f));
 
 		FSceneViewExtensionContext ViewExtensionContext(Scene);
-		ViewExtensionContext.bStereoDisabled = true;
 		ViewFamily.ViewExtensions = GEngine->ViewExtensions->GatherActiveExtensions(ViewExtensionContext);
-
 		for (const FSceneViewExtensionRef& Extension : ViewFamily.ViewExtensions)
 		{
 			Extension->SetupViewFamily(ViewFamily);
diff --git a/Engine/Source/Runtime/Renderer/Private/SceneCaptureRendering.cpp b/Engine/Source/Runtime/Renderer/Private/SceneCaptureRendering.cpp
--- a/Engine/Source/Runtime/Renderer/Private/SceneCaptureRendering.cpp
+++ b/Engine/Source/Runtime/Renderer/Private/SceneCaptureRendering.cpp
@@ -878,8 +878,6 @@ static FSceneRenderer* CreateSceneRendererForSceneCapture(
 		.SetRealtimeUpdate(SceneCaptureComponent->bCaptureEveryFrame || SceneCaptureComponent->bAlwaysPersistRenderingState));
 
 	FSceneViewExtensionContext ViewExtensionContext(Scene);
-	ViewExtensionContext.bStereoDisabled = true;
-
 	ViewFamily.ViewExtensions = GEngine->ViewExtensions->GatherActiveExtensions(ViewExtensionContext);
 	
 	SetupViewFamilyForSceneCapture(
@@ -992,7 +990,6 @@ void FScene::UpdateSceneCaptureContents(USceneCaptureComponent2D* CaptureCompone
 		// Process Scene View extensions for the capture component
 		{
 			FSceneViewExtensionContext ViewExtensionContext(SceneRenderer->Scene);
-			ViewExtensionContext.bStereoDisabled = true;
 
 			for (int32 Index = 0; Index < CaptureComponent->SceneViewExtensions.Num(); ++Index)
 			{
