Added ways to disable sky atmosphere on screen message for special uses caes for the cinematic team.

diff --git a/Engine/Source/Runtime/Renderer/Private/SkyAtmosphereRendering.cpp b/Engine/Source/Runtime/Renderer/Private/SkyAtmosphereRendering.cpp
--- a/Engine/Source/Runtime/Renderer/Private/SkyAtmosphereRendering.cpp
+++ b/Engine/Source/Runtime/Renderer/Private/SkyAtmosphereRendering.cpp
@@ -211,6 +211,11 @@ static TAutoConsoleVariable<int32> CVarSkyAtmosphereLUT32(
 	TEXT("Use full 32bit per-channel precision for all sky LUTs.\n"),
 	ECVF_RenderThreadSafe | ECVF_Scalability);
 
+static TAutoConsoleVariable<int32> CVarSkyAtmosphereEditorNotifications(
+	TEXT("r.SkyAtmosphere.EditorNotifications"), 1,
+	TEXT("Enable the rendering of in editor notification to warn the user about missing sky dome pixels on screen. It is better to keep it enabled and will be removed when shipping.\n"),
+	ECVF_RenderThreadSafe | ECVF_Scalability);
+
 DECLARE_GPU_STAT(SkyAtmosphereLUTs);
 DECLARE_GPU_STAT(SkyAtmosphere);
 DECLARE_GPU_STAT(SkyAtmosphereEditor);
@@ -1976,12 +1981,15 @@ void FSceneRenderer::RenderSkyAtmosphere(FRDGBuilder& GraphBuilder, const FMinim
 bool FSceneRenderer::ShouldRenderSkyAtmosphereEditorNotifications() const
 {
 #if WITH_EDITOR
-	bool bAnyViewHasSkyMaterial = false;
-	for (int32 ViewIndex = 0; ViewIndex < Views.Num(); ViewIndex++)
+	if (CVarSkyAtmosphereEditorNotifications.GetValueOnAnyThread() > 0)
 	{
-		bAnyViewHasSkyMaterial |= Views[ViewIndex].bSceneHasSkyMaterial;
+		bool bAnyViewHasSkyMaterial = false;
+		for (int32 ViewIndex = 0; ViewIndex < Views.Num(); ViewIndex++)
+		{
+			bAnyViewHasSkyMaterial |= Views[ViewIndex].bSceneHasSkyMaterial;
+		}
+		return bAnyViewHasSkyMaterial;
 	}
-	return bAnyViewHasSkyMaterial;
 #endif
 	return false;
 }
