Fixed broken Nanite materials for various meshes in CitySample content. Previously fix for TheOrigin was to change the GPU scene material table upload to use MaterialSlotCount instead of material max index +1, because TheOrigin has content with multiple unique sections pointing to a single material index. CitySample has the opposite in content, so max material index + 1 is needed to handle the case of more unique materials than sections. The updated fix is to take the max of these two values, so we also allocate a large enough upload range for either of these conditions.

diff --git a/Engine/Source/Runtime/Renderer/Private/GPUScene.cpp b/Engine/Source/Runtime/Renderer/Private/GPUScene.cpp
--- a/Engine/Source/Runtime/Renderer/Private/GPUScene.cpp
+++ b/Engine/Source/Runtime/Renderer/Private/GPUScene.cpp
@@ -1203,8 +1203,12 @@ void FGPUScene::UploadGeneral(FRHICommandListImmediate& RHICmdList, FScene *Scen
 								if (PassMaterials.Num() == PassMaterialSlots.Num())
 								{
 									const uint32 MaterialSlotCount = uint32(PassMaterialSlots.Num());
+									const uint32 TableEntryCount = uint32(NaniteSceneProxy->GetMaterialMaxIndex() + 1);
 
-									void* MaterialSlotRange = NaniteMaterials.GetMaterialSlotPtr(UploadInfo.PrimitiveID, MaterialSlotCount);
+									// TODO: Make this more robust, and catch issues earlier on
+									const uint32 UploadEntryCount = FMath::Max(MaterialSlotCount, TableEntryCount);
+
+									void* MaterialSlotRange = NaniteMaterials.GetMaterialSlotPtr(UploadInfo.PrimitiveID, UploadEntryCount);
 									uint32* MaterialSlots = static_cast<uint32*>(MaterialSlotRange);
 									for (uint32 Entry = 0; Entry < MaterialSlotCount; ++Entry)
 									{
