Fixed reentrancy bug by making FSimpleMeshDrawCommandPass only run SetupDrawCommands the first time BuildRenderingCommands is called (trips check in groom rendering).
- also make it use the batched path by default.

diff --git a/Engine/Source/Runtime/Renderer/Private/InstanceCulling/InstanceCullingMergedContext.cpp b/Engine/Source/Runtime/Renderer/Private/InstanceCulling/InstanceCullingMergedContext.cpp
--- a/Engine/Source/Runtime/Renderer/Private/InstanceCulling/InstanceCullingMergedContext.cpp
+++ b/Engine/Source/Runtime/Renderer/Private/InstanceCulling/InstanceCullingMergedContext.cpp
@@ -112,6 +112,7 @@ void FInstanceCullingMergedContext::MergeBatches()
 
 void FInstanceCullingMergedContext::AddBatch(FRDGBuilder& GraphBuilder, const FInstanceCullingContext* Context, int32 DynamicInstanceIdOffset,	int32 DynamicInstanceIdNum, FInstanceCullingDrawParams* InstanceCullingDrawParams)
 {
+	checkfSlow(Batches.FindByPredicate([InstanceCullingDrawParams](const FBatchItem& Item) { return Item.Result == InstanceCullingDrawParams; }) == nullptr, TEXT("Output draw paramters registered twice."));
 	Batches.Add(FBatchItem{ Context, InstanceCullingDrawParams, DynamicInstanceIdOffset, DynamicInstanceIdNum });
 
 	const bool bOcclusionCullInstances = Context->PrevHZB.IsValid() && FInstanceCullingContext::IsOcclusionCullingEnabled();
diff --git a/Engine/Source/Runtime/Renderer/Private/SimpleMeshDrawCommandPass.cpp b/Engine/Source/Runtime/Renderer/Private/SimpleMeshDrawCommandPass.cpp
--- a/Engine/Source/Runtime/Renderer/Private/SimpleMeshDrawCommandPass.cpp
+++ b/Engine/Source/Runtime/Renderer/Private/SimpleMeshDrawCommandPass.cpp
@@ -41,13 +41,17 @@ void FSimpleMeshDrawCommandPass::BuildRenderingCommands(FRDGBuilder& GraphBuilde
 		int32 VisibleMeshDrawCommandsNum = 0;
 		int32 NewPassVisibleMeshDrawCommandsNum = 0;
 
-		// 1. do the first thing first
-		InstanceCullingContext.SetupDrawCommands(VisibleMeshDrawCommands, true, MaxInstances, VisibleMeshDrawCommandsNum, NewPassVisibleMeshDrawCommandsNum);
+		// 1. Run draw command setup, but only the first time.
+		if (!bWasDrawCommandsSetup)
+		{
+			InstanceCullingContext.SetupDrawCommands(VisibleMeshDrawCommands, true, MaxInstances, VisibleMeshDrawCommandsNum, NewPassVisibleMeshDrawCommandsNum);
+			bWasDrawCommandsSetup = true;
+		}
 
 		// 2. Run finalize culling commands pass
 		check(View.bIsViewInfo);
 		const FViewInfo* ViewInfo = static_cast<const FViewInfo*>(&View);
-		InstanceCullingContext.BuildRenderingCommands(GraphBuilder, GPUScene, ViewInfo->DynamicPrimitiveCollector.GetInstanceSceneDataOffset(), ViewInfo->DynamicPrimitiveCollector.NumInstances(), InstanceCullingResult);
+		InstanceCullingContext.BuildRenderingCommands(GraphBuilder, GPUScene, ViewInfo->DynamicPrimitiveCollector.GetInstanceSceneDataOffset(), ViewInfo->DynamicPrimitiveCollector.NumInstances(), InstanceCullingResult, &OutInstanceCullingDrawParams);
 		
 		// Signal that scene primitives are supported, used for validation, the existence of a valid InstanceCullingResult is the required signal
 		bSupportsScenePrimitives = true;
diff --git a/Engine/Source/Runtime/Renderer/Public/SimpleMeshDrawCommandPass.h b/Engine/Source/Runtime/Renderer/Public/SimpleMeshDrawCommandPass.h
--- a/Engine/Source/Runtime/Renderer/Public/SimpleMeshDrawCommandPass.h
+++ b/Engine/Source/Runtime/Renderer/Public/SimpleMeshDrawCommandPass.h
@@ -52,6 +52,8 @@ private:
 	bool bSupportsScenePrimitives = false;
 
 	bool bUsingStereo = false;
+	bool bWasDrawCommandsSetup = false;
+
 	uint32 InstanceFactor = 1;
 
 	// GPUCULL_TODO: Only for legacy path
