20.10+ edition of CL 19329805 and 19348386. Changed to only create lighting interactions if there are lighting channels enabled.

diff --git a/Engine/Source/Runtime/Renderer/Private/RendererScene.cpp b/Engine/Source/Runtime/Renderer/Private/RendererScene.cpp
--- a/Engine/Source/Runtime/Renderer/Private/RendererScene.cpp
+++ b/Engine/Source/Runtime/Renderer/Private/RendererScene.cpp
@@ -1110,24 +1110,28 @@ public:
 		FOptionalTaskTagScope Scope(ETaskTag::EParallelRenderingThread);
 		for (TSet<FPrimitiveSceneInfo*>::TConstIterator It(PendingPrimitives); It; ++It)
 		{
-			FMemMark MemStackMark(FMemStack::Get());
-
 			FPrimitiveSceneInfo* PrimInfo = *It;
 			FPrimitiveSceneProxy* Proxy = PrimInfo->Proxy;
-			const FBoxSphereBounds& Bounds = Proxy->GetBounds();
-			const FPrimitiveSceneInfoCompact PrimitiveSceneInfoCompact(PrimInfo);
 
-			// Find local lights that affect the primitive in the light octree.
-			Scene->LocalShadowCastingLightOctree.FindElementsWithBoundsTest(Bounds.GetBox(), [&PrimitiveSceneInfoCompact](const FLightSceneInfoCompact& LightSceneInfoCompact)
+			if (Proxy->GetLightingChannelMask() != 0)
 			{
-				LightSceneInfoCompact.LightSceneInfo->CreateLightPrimitiveInteraction(LightSceneInfoCompact, PrimitiveSceneInfoCompact);
-			});
+				FMemMark MemStackMark(FMemStack::Get());
 
-			// Also loop through non-local (directional) shadow-casting lights
-			for (int32 LightID : Scene->DirectionalShadowCastingLightIDs)
-			{
-				const FLightSceneInfoCompact& LightSceneInfoCompact = Scene->Lights[LightID];
-				LightSceneInfoCompact.LightSceneInfo->CreateLightPrimitiveInteraction(LightSceneInfoCompact, PrimitiveSceneInfoCompact);
+				const FBoxSphereBounds& Bounds = Proxy->GetBounds();
+				const FPrimitiveSceneInfoCompact PrimitiveSceneInfoCompact(PrimInfo);
+
+				// Find local lights that affect the primitive in the light octree.
+				Scene->LocalShadowCastingLightOctree.FindElementsWithBoundsTest(Bounds.GetBox(), [&PrimitiveSceneInfoCompact](const FLightSceneInfoCompact& LightSceneInfoCompact)
+				{
+					LightSceneInfoCompact.LightSceneInfo->CreateLightPrimitiveInteraction(LightSceneInfoCompact, PrimitiveSceneInfoCompact);
+				});
+
+				// Also loop through non-local (directional) shadow-casting lights
+				for (int32 LightID : Scene->DirectionalShadowCastingLightIDs)
+				{
+					const FLightSceneInfoCompact& LightSceneInfoCompact = Scene->Lights[LightID];
+					LightSceneInfoCompact.LightSceneInfo->CreateLightPrimitiveInteraction(LightSceneInfoCompact, PrimitiveSceneInfoCompact);
+				}
 			}
 		}
 
@@ -5238,8 +5242,11 @@ void FScene::UpdateAllPrimitiveSceneInfos(FRDGBuilder& GraphBuilder, bool bAsync
 
 void FScene::CreateLightPrimitiveInteractionsForPrimitive(FPrimitiveSceneInfo* PrimitiveInfo, bool bAsyncCreateLPIs)
 {
+	FPrimitiveSceneProxy* Proxy = PrimitiveInfo->Proxy;
+	if (Proxy->GetLightingChannelMask() != 0)
+	{
 		FMemMark MemStackMark(FMemStack::Get());
-		const FBoxSphereBounds& Bounds = PrimitiveInfo->Proxy->GetBounds();
+		const FBoxSphereBounds& Bounds = Proxy->GetBounds();
 		const FPrimitiveSceneInfoCompact PrimitiveSceneInfoCompact(PrimitiveInfo);
 
 		// Find local lights that affect the primitive in the light octree.
@@ -5254,6 +5261,7 @@ void FScene::CreateLightPrimitiveInteractionsForPrimitive(FPrimitiveSceneInfo* P
 			const FLightSceneInfoCompact& LightSceneInfoCompact = Lights[LightID];
 			LightSceneInfoCompact.LightSceneInfo->CreateLightPrimitiveInteraction(LightSceneInfoCompact, PrimitiveSceneInfoCompact);
 		}
+	}
 }
 
 bool FScene::IsPrimitiveBeingRemoved(FPrimitiveSceneInfo* PrimitiveSceneInfo) const
