Updated baked energy conservation tables for stata cloth/fuzz.

diff --git a/Engine/Source/Runtime/Renderer/Private/ShadingEnergyConservation.cpp b/Engine/Source/Runtime/Renderer/Private/ShadingEnergyConservation.cpp
--- a/Engine/Source/Runtime/Renderer/Private/ShadingEnergyConservation.cpp
+++ b/Engine/Source/Runtime/Renderer/Private/ShadingEnergyConservation.cpp
@@ -332,7 +332,7 @@ void Init(FRDGBuilder& GraphBuilder, FViewInfo& View)
 
 				ShadingEnergyConservationData::LockCopyTexture2D(GraphBuilder.RHICmdList, (FTexture2DRHIRef&)View.ViewState->ShadingEnergyConservationData.GGXSpecEnergyTexture->GetRenderTargetItem().ShaderResourceTexture,  ShadingEnergyConservationData::GGXSpecValues, 2);
 				ShadingEnergyConservationData::LockCopyTexture3D(GraphBuilder.RHICmdList, (FTexture3DRHIRef&)View.ViewState->ShadingEnergyConservationData.GGXGlassEnergyTexture->GetRenderTargetItem().ShaderResourceTexture, ShadingEnergyConservationData::GGXGlassValues, 2);
-				ShadingEnergyConservationData::LockCopyTexture2D(GraphBuilder.RHICmdList, (FTexture2DRHIRef&)View.ViewState->ShadingEnergyConservationData.ClothEnergyTexture->GetRenderTargetItem().ShaderResourceTexture,    ShadingEnergyConservationData::ClothSpecValues, 2);
+				ShadingEnergyConservationData::LockCopyTexture2D(GraphBuilder.RHICmdList, (FTexture2DRHIRef&)View.ViewState->ShadingEnergyConservationData.ClothEnergyTexture->GetRenderTargetItem().ShaderResourceTexture,    Strata::IsStrataEnabled() ? ShadingEnergyConservationData::StrataClothSpecValues : ShadingEnergyConservationData::ClothSpecValues, 2);
 				ShadingEnergyConservationData::LockCopyTexture2D(GraphBuilder.RHICmdList, (FTexture2DRHIRef&)View.ViewState->ShadingEnergyConservationData.DiffuseEnergyTexture->GetRenderTargetItem().ShaderResourceTexture,  ShadingEnergyConservationData::DiffuseValues, 1);
 
 				GGXSpecEnergyTexture  = GraphBuilder.RegisterExternalTexture(View.ViewState->ShadingEnergyConservationData.GGXSpecEnergyTexture);
diff --git a/Engine/Source/Runtime/Renderer/Private/ShadingEnergyConservationData.h b/Engine/Source/Runtime/Renderer/Private/ShadingEnergyConservationData.h
--- a/Engine/Source/Runtime/Renderer/Private/ShadingEnergyConservationData.h
+++ b/Engine/Source/Runtime/Renderer/Private/ShadingEnergyConservationData.h
@@ -4,6 +4,12 @@
 
 // These values were generated using FBuildShadingEnergyConservationTableCS
 
+// How to update those tables?
+//  1- Set r.Shading.EnergyConservation.Format to 0
+//	2- Open the texture asset in renderdoc and click "Open the texture content in a raw buffer viewer"
+//  3- Set the format as "ushort2 pixels[32];" and hit the apply button.
+//  4- Export CSV and you then get the data to copy paste. Note: you might need to format then a bit like adding commad at the end of lines.
+
 // Specular directional albedo for GGX BSDF (2D)
 // 2x16bits (unorm) (FG decomposition)
 static const uint16 GGXSpecValues[2 * SHADING_ENERGY_CONSERVATION_TABLE_RESOLUTION * SHADING_ENERGY_CONSERVATION_TABLE_RESOLUTION] =
@@ -80,6 +86,44 @@ static const uint16 ClothSpecValues[2 * SHADING_ENERGY_CONSERVATION_TABLE_RESOLU
 31823, 1354, 29813, 1110, 28381, 933, 27231, 794, 26260, 680, 25416, 585, 24667, 503, 23992, 434, 23379, 374, 22816, 322, 22295, 277, 21811, 238, 21359, 205, 20935, 175, 20535, 149, 20158, 127, 19800, 108, 19460, 91, 19137, 76, 18828, 64, 18533, 53, 18250, 43, 17979, 35, 17719, 28, 17468, 23, 17227, 18, 16994, 14, 16770, 10, 16553, 7, 16343, 5, 16140, 3, 15943, 2,
 };
 
+// Specular directional albedo for Cloth BSDF (2D) for Strata (using CharlieD)
+// 2x16bits (unorm) (FG decomposition)
+static const uint16 StrataClothSpecValues[2 * SHADING_ENERGY_CONSERVATION_TABLE_RESOLUTION * SHADING_ENERGY_CONSERVATION_TABLE_RESOLUTION] =
+{
+65535, 2421, 65535, 928, 65535, 412, 56261, 186, 40693, 91, 29345, 44, 21060, 22, 15003, 11, 10581, 6, 7376, 3, 5084, 1, 3446, 1, 2301, 0, 1509, 0, 970, 0, 610, 0, 374, 0, 223, 0, 129, 0, 72, 0, 39, 0, 20, 0, 10, 0, 4, 0, 2, 0, 1, 0, 0, 0, 0, 0, 0, 0, 0, 0, 0, 0, 0, 0,
+65535, 1952, 65535, 983, 65535, 596, 55249, 347, 45791, 221, 38053, 139, 31645, 90, 26284, 58, 21799, 38, 18029, 25, 14855, 16, 12188, 11, 9948, 7, 8073, 5, 6507, 3, 5206, 2, 4131, 1, 3245, 1, 2523, 1, 1938, 0, 1467, 0, 1093, 0, 800, 0, 572, 0, 399, 0, 269, 0, 175, 0, 108, 0, 62, 0, 32, 0, 14, 0, 4, 0,
+65535, 1746, 65535, 977, 59917, 643, 51387, 416, 44369, 286, 38430, 195, 33341, 134, 28939, 95, 25107, 66, 21755, 47, 18817, 33, 16236, 23, 13969, 17, 11977, 12, 10227, 8, 8692, 6, 7350, 4, 6177, 3, 5157, 2, 4272, 1, 3507, 1, 2850, 1, 2289, 0, 1813, 0, 1413, 0, 1079, 0, 803, 0, 580, 0, 401, 0, 261, 0, 154, 0, 76, 0,
+65535, 1606, 63642, 1006, 55136, 659, 48360, 461, 42652, 321, 37759, 230, 33493, 167, 29737, 121, 26405, 89, 23432, 65, 20773, 48, 18388, 35, 16243, 26, 14315, 19, 12579, 14, 11016, 10, 9611, 8, 8348, 6, 7215, 4, 6199, 3, 5292, 2, 4483, 2, 3765, 1, 3129, 1, 2569, 1, 2078, 0, 1651, 0, 1282, 0, 965, 0, 697, 0, 473, 0, 288, 0,
+65535, 1540, 58807, 979, 51741, 683, 45999, 486, 41137, 349, 36935, 260, 33220, 191, 29922, 143, 26959, 107, 24288, 80, 21866, 60, 19667, 46, 17661, 34, 15832, 26, 14161, 20, 12633, 15, 11235, 11, 9958, 8, 8790, 6, 7724, 5, 6752, 4, 5866, 3, 5060, 2, 4330, 2, 3669, 1, 3072, 1, 2537, 1, 2057, 1, 1630, 0, 1252, 0, 920, 0, 631, 0,
+63713, 1502, 55307, 987, 49151, 699, 44117, 502, 39860, 376, 36130, 280, 32829, 212, 29864, 161, 27185, 122, 24750, 94, 22524, 72, 20484, 55, 18607, 42, 16879, 32, 15284, 25, 13809, 19, 12447, 15, 11186, 11, 10019, 9, 8940, 7, 7942, 5, 7020, 4, 6168, 3, 5383, 2, 4660, 2, 3996, 1, 3386, 1, 2829, 1, 2320, 1, 1858, 1, 1439, 0, 1062, 0,
+60058, 1495, 52595, 997, 47103, 707, 42603, 519, 38765, 391, 35400, 298, 32407, 229, 29703, 175, 27247, 136, 25001, 105, 22937, 82, 21033, 63, 19270, 49, 17634, 38, 16114, 30, 14698, 23, 13379, 18, 12148, 14, 10998, 11, 9925, 9, 8923, 7, 7988, 5, 7114, 4, 6300, 3, 5540, 2, 4833, 2, 4175, 2, 3563, 1, 2997, 1, 2472, 1, 1988, 1, 1542, 1,
+57137, 1395, 50433, 979, 45422, 711, 41346, 539, 37841, 410, 34752, 315, 31994, 243, 29500, 189, 27224, 147, 25135, 116, 23203, 90, 21415, 71, 19751, 56, 18199, 44, 16747, 34, 15389, 27, 14115, 21, 12919, 17, 11795, 13, 10738, 10, 9743, 8, 8808, 6, 7927, 5, 7098, 4, 6319, 3, 5586, 3, 4897, 2, 4249, 2, 3642, 1, 3073, 1, 2541, 1, 2043, 1,
+54894, 1448, 48619, 967, 44054, 728, 40274, 550, 37029, 421, 34174, 329, 31610, 256, 29284, 201, 27157, 158, 25194, 124, 23378, 98, 21688, 78, 20109, 62, 18632, 49, 17244, 39, 15940, 31, 14710, 24, 13550, 19, 12454, 15, 11418, 12, 10437, 10, 9510, 8, 8631, 6, 7798, 5, 7009, 4, 6262, 3, 5554, 3, 4884, 2, 4249, 2, 3649, 1, 3082, 1, 2546, 1,
+52965, 1449, 47164, 984, 42881, 735, 39356, 560, 36331, 436, 33651, 339, 31252, 267, 29066, 211, 27062, 167, 25212, 133, 23492, 106, 21887, 84, 20384, 67, 18973, 54, 17642, 43, 16387, 34, 15199, 27, 14075, 22, 13008, 17, 11995, 14, 11032, 11, 10117, 9, 9245, 7, 8415, 6, 7624, 5, 6871, 4, 6153, 3, 5468, 2, 4816, 2, 4195, 2, 3603, 1, 3039, 1,
+51322, 1378, 45869, 974, 41874, 744, 38563, 571, 35712, 447, 33189, 350, 30922, 278, 28857, 221, 26959, 176, 25201, 141, 23565, 113, 22034, 90, 20598, 72, 19244, 58, 17966, 46, 16757, 37, 15609, 30, 14518, 24, 13480, 19, 12491, 16, 11547, 12, 10647, 10, 9786, 8, 8962, 7, 8174, 5, 7420, 4, 6698, 3, 6006, 3, 5343, 2, 4708, 2, 4099, 1, 3516, 1,
+49941, 1381, 44806, 989, 40997, 749, 37862, 580, 35160, 453, 32775, 361, 30620, 288, 28656, 229, 26847, 183, 25172, 147, 23609, 119, 22145, 96, 20767, 77, 19466, 62, 18235, 50, 17067, 40, 15956, 32, 14897, 26, 13887, 21, 12921, 17, 11998, 14, 11113, 11, 10265, 9, 9451, 7, 8669, 6, 7918, 5, 7196, 4, 6501, 3, 5833, 2, 5189, 2, 4570, 1, 3973, 1,
+48743, 1354, 43835, 986, 40230, 756, 37242, 586, 34673, 464, 32397, 371, 30342, 296, 28465, 237, 26737, 191, 25133, 154, 23634, 124, 22228, 101, 20902, 81, 19649, 66, 18460, 53, 17330, 43, 16253, 35, 15224, 28, 14241, 23, 13299, 19, 12395, 15, 11527, 12, 10693, 10, 9890, 8, 9116, 7, 8371, 5, 7652, 4, 6958, 3, 6287, 3, 5640, 2, 5014, 2, 4409, 1,
+47700, 1411, 43012, 999, 39552, 767, 36695, 596, 34238, 473, 32053, 377, 30088, 304, 28287, 244, 26629, 198, 25086, 160, 23644, 129, 22290, 105, 21012, 85, 19801, 69, 18651, 56, 17556, 46, 16510, 37, 15510, 30, 14552, 25, 13632, 20, 12748, 16, 11897, 13, 11077, 11, 10286, 9, 9522, 7, 8784, 6, 8070, 5, 7379, 4, 6710, 3, 6061, 2, 5432, 2, 4822, 1,
+46761, 1362, 42258, 995, 38942, 769, 36203, 604, 33840, 482, 31743, 386, 29850, 311, 28120, 251, 26523, 204, 25036, 166, 23645, 134, 22337, 110, 21101, 89, 19929, 73, 18814, 59, 17751, 48, 16735, 39, 15761, 32, 14827, 26, 13929, 21, 13063, 18, 12229, 14, 11424, 12, 10645, 10, 9892, 8, 9162, 6, 8455, 5, 7769, 4, 7103, 3, 6455, 2, 5826, 2, 5214, 1,
+45923, 1345, 41600, 1002, 38397, 777, 35757, 611, 33479, 489, 31457, 392, 29632, 318, 27964, 258, 26420, 209, 24984, 171, 23639, 139, 22372, 114, 21175, 93, 20038, 76, 18956, 62, 17922, 51, 16933, 41, 15984, 34, 15072, 28, 14194, 23, 13347, 19, 12530, 15, 11739, 12, 10973, 10, 10230, 8, 9510, 7, 8810, 5, 8130, 4, 7469, 3, 6824, 3, 6196, 2, 5584, 1,
+45203, 1375, 40993, 1000, 37900, 778, 35353, 619, 33151, 495, 31197, 400, 29431, 324, 27816, 264, 26322, 215, 24931, 176, 23627, 144, 22399, 117, 21236, 96, 20131, 79, 19079, 65, 18073, 53, 17109, 43, 16183, 36, 15292, 29, 14433, 24, 13604, 20, 12802, 16, 12026, 13, 11272, 11, 10541, 9, 9831, 7, 9139, 6, 8466, 5, 7810, 4, 7170, 3, 6545, 2, 5934, 1,
+44514, 1346, 40453, 1003, 37453, 783, 34984, 626, 32850, 502, 30955, 406, 29245, 330, 27677, 269, 26229, 220, 24878, 180, 23611, 148, 22418, 121, 21287, 99, 20212, 82, 19187, 67, 18206, 55, 17265, 45, 16361, 37, 15491, 31, 14650, 25, 13838, 21, 13051, 17, 12288, 14, 11548, 11, 10828, 9, 10127, 8, 9445, 6, 8779, 5, 8129, 4, 7494, 3, 6873, 2, 6265, 1,
+43802, 1337, 39944, 1003, 37046, 788, 34642, 627, 32574, 508, 30732, 411, 29070, 335, 27548, 274, 26140, 225, 24826, 184, 23594, 151, 22431, 125, 21330, 102, 20282, 84, 19282, 69, 18325, 57, 17406, 47, 16522, 39, 15670, 32, 14847, 26, 14051, 22, 13279, 18, 12530, 15, 11801, 12, 11093, 10, 10402, 8, 9728, 6, 9071, 5, 8427, 4, 7798, 3, 7182, 2, 6578, 1,
+43358, 1343, 39505, 1011, 36670, 794, 34336, 636, 32318, 513, 30526, 416, 28909, 341, 27426, 279, 26054, 230, 24775, 189, 23574, 155, 22441, 128, 21366, 105, 20344, 87, 19367, 72, 18432, 59, 17533, 49, 16669, 40, 15834, 33, 15027, 27, 14246, 23, 13489, 19, 12752, 15, 12036, 13, 11338, 10, 10658, 8, 9993, 7, 9343, 5, 8707, 4, 8084, 3, 7474, 2, 6874, 1,
+42848, 1351, 39087, 1013, 36326, 796, 34044, 637, 32082, 518, 30334, 421, 28758, 346, 27311, 284, 25973, 233, 24725, 192, 23553, 159, 22446, 131, 21397, 108, 20398, 89, 19443, 74, 18528, 61, 17649, 50, 16802, 42, 15984, 34, 15193, 28, 14426, 23, 13682, 19, 12958, 16, 12254, 13, 11567, 11, 10896, 9, 10240, 7, 9599, 5, 8970, 4, 8354, 3, 7749, 2, 7154, 1,
+42376, 1334, 38712, 1023, 36009, 803, 33784, 646, 31862, 523, 30158, 428, 28615, 350, 27203, 288, 25897, 238, 24676, 196, 23531, 162, 22450, 134, 21423, 111, 20446, 92, 19511, 76, 18615, 63, 17753, 52, 16923, 43, 16121, 36, 15345, 29, 14592, 24, 13861, 20, 13149, 16, 12456, 14, 11779, 11, 11118, 9, 10472, 7, 9838, 6, 9217, 4, 8608, 3, 8009, 2, 7420, 1,
+41931, 1329, 38358, 1019, 35712, 806, 33537, 649, 31658, 527, 29989, 431, 28483, 355, 27102, 292, 25823, 241, 24630, 199, 23509, 165, 22450, 137, 21445, 113, 20488, 94, 19572, 78, 18694, 64, 17849, 53, 17035, 44, 16248, 37, 15486, 30, 14746, 25, 14027, 21, 13327, 17, 12645, 14, 11978, 11, 11327, 9, 10689, 7, 10064, 6, 9450, 4, 8848, 3, 8255, 2, 7671, 2,
+41547, 1333, 38029, 1025, 35439, 809, 33303, 651, 31467, 531, 29834, 436, 28358, 359, 27006, 296, 25753, 245, 24584, 202, 23487, 168, 22449, 139, 21464, 115, 20525, 96, 19628, 79, 18766, 66, 17937, 55, 17138, 45, 16365, 38, 15616, 31, 14889, 26, 14182, 21, 13493, 18, 12821, 14, 12165, 12, 11522, 9, 10893, 8, 10276, 6, 9670, 5, 9074, 3, 8488, 2, 7910, 2,
+41156, 1332, 37725, 1019, 35184, 813, 33090, 655, 31289, 536, 29688, 440, 28241, 363, 26914, 300, 25687, 248, 24541, 205, 23464, 171, 22446, 142, 21481, 118, 20559, 98, 19678, 81, 18832, 67, 18019, 56, 17233, 47, 16473, 39, 15737, 32, 15022, 27, 14326, 22, 13648, 18, 12986, 15, 12339, 12, 11706, 10, 11085, 8, 10476, 6, 9878, 5, 9289, 3, 8709, 2, 8137, 2,
+40839, 1364, 37444, 1031, 34946, 816, 32892, 660, 31121, 540, 29550, 444, 28129, 366, 26829, 303, 25624, 251, 24499, 209, 23442, 173, 22443, 144, 21494, 120, 20590, 100, 19724, 83, 18893, 69, 18093, 57, 17321, 48, 16574, 40, 15850, 33, 15146, 27, 14461, 22, 13793, 19, 13141, 15, 12504, 12, 11879, 10, 11267, 8, 10665, 6, 10074, 5, 9492, 4, 8919, 2, 8353, 2,
+40513, 1330, 37176, 1027, 34722, 819, 32705, 664, 30965, 544, 29421, 447, 28025, 370, 26747, 306, 25564, 255, 24458, 211, 23420, 176, 22438, 146, 21506, 122, 20617, 101, 19766, 84, 18949, 70, 18163, 59, 17403, 49, 16668, 40, 15955, 34, 15262, 28, 14588, 23, 13930, 19, 13287, 16, 12658, 13, 12042, 10, 11438, 8, 10844, 6, 10260, 5, 9685, 4, 9118, 3, 8558, 2,
+40216, 1340, 36924, 1027, 34511, 824, 32527, 667, 30817, 548, 29298, 451, 27926, 373, 26670, 310, 25506, 257, 24419, 214, 23398, 178, 22433, 149, 21516, 124, 20642, 103, 19805, 86, 19001, 72, 18227, 60, 17480, 50, 16756, 41, 16054, 34, 15371, 28, 14707, 24, 14058, 19, 13425, 16, 12804, 13, 12197, 10, 11600, 8, 11014, 7, 10437, 5, 9868, 4, 9308, 3, 8754, 2,
+39920, 1337, 36696, 1034, 34318, 827, 32361, 670, 30675, 549, 29183, 455, 27834, 377, 26596, 313, 25451, 260, 24382, 217, 23377, 181, 22428, 151, 21525, 126, 20665, 105, 19841, 87, 19050, 73, 18287, 61, 17551, 51, 16838, 42, 16146, 35, 15474, 29, 14819, 24, 14179, 20, 13554, 16, 12942, 13, 12342, 11, 11753, 8, 11174, 7, 10604, 5, 10043, 4, 9488, 3, 8940, 2,
+39665, 1342, 36475, 1030, 34129, 827, 32201, 672, 30546, 554, 29074, 458, 27744, 380, 26526, 316, 25399, 263, 24346, 219, 23356, 183, 22421, 153, 21533, 128, 20685, 106, 19874, 89, 19094, 74, 18343, 62, 17618, 52, 16915, 43, 16233, 36, 15570, 30, 14924, 25, 14293, 20, 13677, 17, 13073, 13, 12481, 11, 11899, 9, 11327, 7, 10764, 5, 10209, 4, 9660, 3, 9119, 2,
+39420, 1346, 36266, 1035, 33955, 830, 32057, 677, 30421, 557, 28971, 461, 27660, 383, 26460, 318, 25348, 266, 24311, 221, 23336, 185, 22415, 155, 21539, 129, 20704, 108, 19904, 90, 19136, 75, 18396, 63, 17680, 52, 16988, 44, 16315, 36, 15661, 30, 15024, 25, 14402, 21, 13793, 17, 13197, 14, 12612, 11, 12038, 9, 11472, 7, 10916, 5, 10367, 4, 9825, 3, 9289, 2,
+39188, 1334, 36066, 1031, 33789, 832, 31916, 679, 30303, 560, 28871, 463, 27580, 385, 26396, 321, 25300, 268, 24278, 224, 23316, 187, 22408, 156, 21544, 131, 20721, 109, 19932, 92, 19175, 76, 18445, 64, 17739, 53, 17056, 44, 16393, 37, 15747, 31, 15118, 25, 14504, 21, 13903, 17, 13314, 14, 12737, 11, 12169, 9, 11611, 7, 11061, 5, 10518, 4, 9981, 3, 9451, 2,
+};
+
 // Diffuse directional albedo for Chan BSDF (2D)
 // 1x16bits (unorm)
 static const uint16 DiffuseValues[SHADING_ENERGY_CONSERVATION_TABLE_RESOLUTION * SHADING_ENERGY_CONSERVATION_TABLE_RESOLUTION] =
