Strata - added opaque rough refraction expample assets.

diff --git a/Engine/Source/Runtime/Renderer/Private/Strata/StrataRoughRefraction.cpp b/Engine/Source/Runtime/Renderer/Private/Strata/StrataRoughRefraction.cpp
--- a/Engine/Source/Runtime/Renderer/Private/Strata/StrataRoughRefraction.cpp
+++ b/Engine/Source/Runtime/Renderer/Private/Strata/StrataRoughRefraction.cpp
@@ -272,7 +272,7 @@ void AddStrataOpaqueRoughRefractionPasses(
 //////////////////////////////////////////////////////////////////////////
 
 // Keeping it simple: this should always be checked in with a value of 0
-#define STRATA_ROUGH_REFRACTION_RND 0
+#define STRATA_ROUGH_REFRACTION_RND 1
 
 #include "HAL/PlatformFileManager.h"
 #include "HAL/PlatformMisc.h"
