nop comments only

diff --git a/Engine/Plugins/Enterprise/MDLImporter/Source/MDLImporter/Private/common/ImageUtils.cpp b/Engine/Plugins/Enterprise/MDLImporter/Source/MDLImporter/Private/common/ImageUtils.cpp
--- a/Engine/Plugins/Enterprise/MDLImporter/Source/MDLImporter/Private/common/ImageUtils.cpp
+++ b/Engine/Plugins/Enterprise/MDLImporter/Source/MDLImporter/Private/common/ImageUtils.cpp
@@ -14,8 +14,6 @@ NOT the ImageUtils in Engine
 #include "IImageWrapper.h"
 #include "IImageWrapperModule.h"
 #include "Misc/FileHelper.h"
-#include "DDSLoader.h"
-#include "HDRLoader.h"
 #include "Misc/Paths.h"
 #include "Modules/ModuleManager.h"
 
diff --git a/Engine/Source/Runtime/Core/Public/Misc/FileHelper.h b/Engine/Source/Runtime/Core/Public/Misc/FileHelper.h
--- a/Engine/Source/Runtime/Core/Public/Misc/FileHelper.h
+++ b/Engine/Source/Runtime/Core/Public/Misc/FileHelper.h
@@ -148,7 +148,13 @@ struct CORE_API FFileHelper
 	static bool SaveStringArrayToFile( const TArray<FString>& Lines, const TCHAR* Filename, EEncodingOptions EncodingOptions = EEncodingOptions::AutoDetect, IFileManager* FileManager = &IFileManager::Get(), uint32 WriteFlags = 0 );
 
 	/**
-	 * Saves a 24/32Bit BMP file to disk
+	 * Saves a 24/32Bit BMP file to disk for debug image dump purposes
+	 * 
+	 * for general image saving (to BMP or any other format); use FImageUtils::SaveImage instead
+	 * CreateBitmap is mainly for debug dump images
+	 * 
+	 * note this also calls SendDataToPCViaUnrealConsole
+	 *   and uses GenerateNextBitmapFilename if Pattern does not have ".bmp" on it
 	 * 
 	 * @param Pattern filename with path, must not be 0, if with "bmp" extension (e.g. "out.bmp") the filename stays like this, if without (e.g. "out") automatic index numbers are addended (e.g. "out00002.bmp")
 	 * @param DataWidth - Width of the bitmap supplied in Data >0
@@ -162,10 +168,6 @@ struct CORE_API FFileHelper
 	 *
 	 * @return true if success
 	 */
-	 // DEPRECATED
-	//	BmpImageWrapper now has a good writer, prefer that for general image saving ; use FImageUtils::SaveImage
-	// note this also calls SendDataToPCViaUnrealConsole
-	//   and uses GenerateNextBitmapFilename
 	static bool CreateBitmap( const TCHAR* Pattern, int32 DataWidth, int32 DataHeight, const struct FColor* Data, FIntRect* SubRectangle = NULL, IFileManager* FileManager = &IFileManager::Get(), FString* OutFilename = NULL, bool bInWriteAlpha = false, EColorChannel ColorChannel = EColorChannel::All);
 
 	/**
diff --git a/Engine/Source/Runtime/Engine/Public/DDSLoader.h b/Engine/Source/Runtime/Engine/Public/DDSLoader.h
--- a/Engine/Source/Runtime/Engine/Public/DDSLoader.h
+++ b/Engine/Source/Runtime/Engine/Public/DDSLoader.h
@@ -5,6 +5,13 @@
 #include "CoreMinimal.h"
 #include "Engine/Texture.h"
 
+/***
+*
+* DDSLoader is deprecated
+* Use DDSFile instead or DDS ImageWrapper
+*
+*/
+
 class UTexture2D;
 
 enum EDDSFlags
@@ -86,6 +93,10 @@ struct FDDS10FileHeader
 };
 #pragma pack(pop)
 
+// @@!! still used in InterchangeDDSTranslator
+//   deprecate after that use is removed
+// 
+//class UE_DEPRECATED(5.2, "Use DDSFile instead") FDDSLoadHelper;
 class FDDSLoadHelper
 {
 public:
diff --git a/Engine/Source/Runtime/Renderer/Private/PostProcess/PostProcessVisualizeBuffer.cpp b/Engine/Source/Runtime/Renderer/Private/PostProcess/PostProcessVisualizeBuffer.cpp
--- a/Engine/Source/Runtime/Renderer/Private/PostProcess/PostProcessVisualizeBuffer.cpp
+++ b/Engine/Source/Runtime/Renderer/Private/PostProcess/PostProcessVisualizeBuffer.cpp
@@ -180,6 +180,8 @@ TUniquePtr<FImagePixelData> ReadbackPixelData(FRHICommandListImmediate& RHICmdLi
 	SourceRect.Min.X *= MSAAXSamples;
 	SourceRect.Max.X *= MSAAXSamples;
 
+	// todo: SourceRect is not clipped to Texture bounds
+
 	switch (Texture->GetFormat())
 	{
 	case PF_FloatRGBA:
@@ -250,12 +252,19 @@ void AddDumpToFilePass(FRDGBuilder& GraphBuilder, FScreenPassTexture Input, cons
 
 	if (GIsHighResScreenshot && HighResScreenshotConfig.CaptureRegion.Area())
 	{
+		// todo: CaptureRegion is not clipped to Texture bounds
 		Input.ViewRect = HighResScreenshotConfig.CaptureRegion;
 	}
 
 	AddReadbackTexturePass(GraphBuilder, RDG_EVENT_NAME("DumpToFile(%s)", Input.Texture->Name), Input.Texture,
 		[&HighResScreenshotConfig, Input, Filename](FRHICommandListImmediate& RHICmdList)
 	{
+		// this is where HighResShot bDumpBufferVisualizationTargets are written to EXRs
+
+		// @todo Oodle alternative : use the exact same pixelformat that this buffer would have in the renderer
+		//	 use the DDS writer which can output arbitrary pixel formats
+		//	 do no format conversions so we dump the exact same bits the game renderer would see
+
 		TUniquePtr<FImagePixelData> PixelData = ReadbackPixelData(RHICmdList, Input.Texture->GetRHI(), Input.ViewRect);
 
 		if (!PixelData.IsValid())
@@ -386,6 +395,7 @@ FScreenPassTexture AddVisualizeGBufferOverviewPass(
 				MaterialFilename = BaseFilename + TEXT("_") + MaterialName;
 			}
 
+			// always makes a .png filename even when bCaptureHDR was set, which will actually save an EXR
 			MaterialFilename.Append(TEXT(".png"));
 
 			AddDumpToFilePass(GraphBuilder, Output, MaterialFilename);
