Major refactor of reflection capture code.
 - Converted to RDG.
 - Converted mobile to use subresource transitions instead of awkward ping-ponging between render targets.
 - Merged duplicate code paths between mobile / desktop.
 - Eliminated FReflectionScratchCubemaps singleton in favor of explicit resource marshalling.

diff --git a/Engine/Shaders/Private/OneColorShader.usf b/Engine/Shaders/Private/OneColorShader.usf
--- a/Engine/Shaders/Private/OneColorShader.usf
+++ b/Engine/Shaders/Private/OneColorShader.usf
@@ -54,6 +54,8 @@ void MainVertexShader(
 #define NUM_OUTPUTS 1
 #endif
 
+float4 DrawColorMRT[8];
+
 //These shaders used for RHIClear.  Don't use any global uniform variables here as it greatly complicates
 //the state restoration for an inline shader clear.
 void MainPixelShaderMRT( 
@@ -61,14 +63,14 @@ void MainPixelShaderMRT(
 {
 	for (int i = 0; i < NUM_OUTPUTS; i++)
 	{
-		OutColor[i] = ClearShaderUB.DrawColorMRT[i];
+		OutColor[i] = DrawColorMRT[i];
 	}
 }
 
 void MainPixelShader( 
 	out float4 OutColor : SV_Target0)
 {
-	OutColor = ClearShaderUB.DrawColorMRT[0];
+	OutColor = DrawColorMRT[0];
 }
 
 /** 
diff --git a/Engine/Shaders/Private/ReflectionEnvironmentShaders.usf b/Engine/Shaders/Private/ReflectionEnvironmentShaders.usf
--- a/Engine/Shaders/Private/ReflectionEnvironmentShaders.usf
+++ b/Engine/Shaders/Private/ReflectionEnvironmentShaders.usf
@@ -146,13 +146,15 @@ float CubeTexelWeight( float3 N )
 TextureCube SourceCubemapTexture;
 SamplerState SourceCubemapSampler;
 float2 SinCosSourceCubemapRotation;
+float2 SvPositionToUVScale;
 
 void CopyCubemapToCubeFaceColorPS(
-	FScreenVertexOutput Input,
+	float4 SvPosition : SV_POSITION,
 	out float4 OutColor : SV_Target0
 	)
-{ 
-	float2 ScaledUVs = Input.UV * 2 - 1;
+{
+	float2 UV = SvPosition.xy * SvPositionToUVScale;
+	float2 ScaledUVs = UV * 2 - 1;
 	float3 CubeCoordinates = GetCubemapVector(ScaledUVs, CubeFace);
 	// Rotate around Z axis
 	CubeCoordinates.xy = float2(dot(CubeCoordinates.xy, float2(SinCosSourceCubemapRotation.y, -SinCosSourceCubemapRotation.x)), dot(CubeCoordinates.xy, SinCosSourceCubemapRotation));
@@ -194,11 +196,12 @@ void DownsampleCS(uint3 ThreadId : SV_DispatchThreadID)
 	float4 OutColor;
 #else // USE_COMPUTE
 void DownsamplePS(
-	FScreenVertexOutput Input,
+	float4 SvPosition : SV_POSITION,
 	out float4 OutColor : SV_Target0
 	)
 { 
-	float2 ScaledUVs = Input.UV * 2 - 1;
+	float2 UV = SvPosition.xy * SvPositionToUVScale;
+	float2 ScaledUVs = UV * 2 - 1;
 	const int SelectedCubeFace = CubeFace;
 #endif // USE_COMPUTE
 	float3 CubeCoordinates = GetCubemapVector(ScaledUVs, SelectedCubeFace);
@@ -275,10 +278,7 @@ SamplerState ReflectionEnvironmentColorSampler;
 
 #if COMPUTEBRIGHTNESS_PIXELSHADER
 
-void ComputeBrightnessMain(
-	in float4 UVAndScreenPos : TEXCOORD0, 
-	out float4 OutColor : SV_Target0
-	)
+void ComputeBrightnessMain(out float4 OutColor : SV_Target0)
 { 
 	// Sample the 6 1x1 cube faces and average
 	float3 AverageColor = TextureCubeSampleLevel(ReflectionEnvironmentColorTexture, ReflectionEnvironmentColorSampler, float3(1, 0, 0), NumCaptureArrayMips - 1).rgb;
@@ -305,11 +305,12 @@ float4 SampleCubemap(float3 Coordinates)
 }
 
 void DownsamplePS_Mobile(
-	FScreenVertexOutput Input,
+	float4 SvPosition : SV_POSITION,
 	out float4 OutColor : SV_Target0
 )
 {
-	float2 ScaledUVs = Input.UV * 2 - 1;
+	float2 UV = SvPosition.xy * SvPositionToUVScale;
+	float2 ScaledUVs = UV * 2 - 1;
 	float3 CubeCoordinates = GetCubemapVector(ScaledUVs, CubeFace);
 
 	OutColor = SampleCubemap(CubeCoordinates);
@@ -330,11 +331,12 @@ void FilterCS(uint3 ThreadId : SV_DispatchThreadID)
 	float4 OutColor;
 #else // USE_COMPUTE
 void FilterPS(
-	FScreenVertexOutput Input,
+	float4 SvPosition : SV_POSITION,
 	out float4 OutColor : SV_Target0
 	)
 { 
-	float2 ScaledUVs = Input.UV * 2 - 1;
+	float2 UV = SvPosition.xy * SvPositionToUVScale;
+	float2 ScaledUVs = UV * 2 - 1;
 	const int SelectedCubeFace = CubeFace;
 #endif // USE_COMPUTE
 	float3 CubeCoordinates = GetCubemapVector(ScaledUVs, SelectedCubeFace);
@@ -438,11 +440,12 @@ float CoefficientMask2;
 int NumSamples;
 
 void DiffuseIrradianceCopyPS(
-	FScreenVertexOutput Input,
+	float4 SvPosition : SV_POSITION,
 	out float4 OutColor : SV_Target0
 	)
 { 
-	float2 ScaledUVs = Input.UV * 2 - 1;
+	float2 UV = SvPosition.xy * SvPositionToUVScale;
+	float2 ScaledUVs = UV * 2 - 1;
 	float3 CubeCoordinates = normalize(GetCubemapVector(ScaledUVs, CubeFace));
 
 	float SquaredUVs = 1 + dot(ScaledUVs, ScaledUVs);
@@ -460,32 +463,33 @@ float4 Sample01;
 float4 Sample23;
 
 void DiffuseIrradianceAccumulatePS(
-	FScreenVertexOutput Input,
+	float4 SvPosition : SV_POSITION,
 	out float4 OutColor : SV_Target0
 	)
 { 
 	float4 AccumulatedValue = 0;
+	float2 UV = SvPosition.xy * SvPositionToUVScale;
 
 	{
-		float2 ScaledUVs = saturate(Input.UV + Sample01.xy) * 2 - 1;
+		float2 ScaledUVs = saturate(UV + Sample01.xy) * 2 - 1;
 		float3 CubeCoordinates = GetCubemapVector(ScaledUVs, CubeFace);
 		AccumulatedValue += SampleCubemap(CubeCoordinates);
 	}
 	
 	{
-		float2 ScaledUVs = saturate(Input.UV + Sample01.zw) * 2 - 1;
+		float2 ScaledUVs = saturate(UV + Sample01.zw) * 2 - 1;
 		float3 CubeCoordinates = GetCubemapVector(ScaledUVs, CubeFace);
 		AccumulatedValue += SampleCubemap(CubeCoordinates);
 	}
 
 	{
-		float2 ScaledUVs = saturate(Input.UV + Sample23.xy) * 2 - 1;
+		float2 ScaledUVs = saturate(UV + Sample23.xy) * 2 - 1;
 		float3 CubeCoordinates = GetCubemapVector(ScaledUVs, CubeFace);
 		AccumulatedValue += SampleCubemap(CubeCoordinates);
 	}
 
 	{
-		float2 ScaledUVs = saturate(Input.UV + Sample23.zw) * 2 - 1;
+		float2 ScaledUVs = saturate(UV + Sample23.zw) * 2 - 1;
 		float3 CubeCoordinates = GetCubemapVector(ScaledUVs, CubeFace);
 		AccumulatedValue += SampleCubemap(CubeCoordinates);
 	}
@@ -494,7 +498,6 @@ void DiffuseIrradianceAccumulatePS(
 }
 
 void AccumulateCubeFacesPS(
-	FScreenVertexOutput Input,
 	out float4 OutColor : SV_Target0
 	)
 {
diff --git a/Engine/Source/Runtime/RenderCore/Private/ClearQuad.cpp b/Engine/Source/Runtime/RenderCore/Private/ClearQuad.cpp
--- a/Engine/Source/Runtime/RenderCore/Private/ClearQuad.cpp
+++ b/Engine/Source/Runtime/RenderCore/Private/ClearQuad.cpp
@@ -76,7 +76,7 @@ static void ClearQuadSetup( FRHICommandList& RHICmdList, bool bClearColor, int32
 	SetGraphicsPipelineState(RHICmdList, GraphicsPSOInit, Stencil);
 
 	VertexShader->SetDepthParameter(RHICmdList, Depth);
-	PixelShader->SetColors(RHICmdList, ClearColorArray, NumClearColors);
+	PixelShader->SetColors(RHICmdList, PixelShader, ClearColorArray, NumClearColors);
 }
 
 void DrawClearQuadMRT(FRHICommandList& RHICmdList, bool bClearColor, int32 NumClearColors, const FLinearColor* ClearColorArray, bool bClearDepth, float Depth, bool bClearStencil, uint32 Stencil)
diff --git a/Engine/Source/Runtime/RenderCore/Private/OneColorShader.cpp b/Engine/Source/Runtime/RenderCore/Private/OneColorShader.cpp
--- a/Engine/Source/Runtime/RenderCore/Private/OneColorShader.cpp
+++ b/Engine/Source/Runtime/RenderCore/Private/OneColorShader.cpp
@@ -3,34 +3,6 @@
 
 #include "OneColorShader.h"
 
-BEGIN_GLOBAL_SHADER_PARAMETER_STRUCT(FClearShaderUB, )
-	SHADER_PARAMETER_ARRAY(FVector4f, DrawColorMRT, [8] )
-END_GLOBAL_SHADER_PARAMETER_STRUCT()
-
-IMPLEMENT_GLOBAL_SHADER_PARAMETER_STRUCT(FClearShaderUB, "ClearShaderUB");
-
-void FOneColorPS::SetColors(FRHICommandList& RHICmdList, const FLinearColor* Colors, int32 NumColors)
-{
-	check(NumColors <= MaxSimultaneousRenderTargets);
-
-	auto& ClearUBParam = GetUniformBufferParameter<FClearShaderUB>();
-	if (ClearUBParam.IsBound())
-	{
-		FClearShaderUB ClearData;
-		FMemory::Memzero(ClearData.DrawColorMRT);
-		for (int32 i = 0; i < NumColors; ++i)
-		{
-			ClearData.DrawColorMRT[i].X = Colors[i].R;
-			ClearData.DrawColorMRT[i].Y = Colors[i].G;
-			ClearData.DrawColorMRT[i].Z = Colors[i].B;
-			ClearData.DrawColorMRT[i].W = Colors[i].A;
-		}
-
-		FLocalUniformBuffer LocalUB = TUniformBufferRef<FClearShaderUB>::CreateLocalUniformBuffer(RHICmdList, ClearData, UniformBuffer_SingleFrame);
-		RHICmdList.SetLocalShaderUniformBuffer(RHICmdList.GetBoundPixelShader(), ClearUBParam.GetBaseIndex(), LocalUB);
-	}
-}
-
 // #define avoids a lot of code duplication
 #define IMPLEMENT_ONECOLORVS(A,B) typedef TOneColorVS<A,B> TOneColorVS##A##B; \
 	IMPLEMENT_SHADER_TYPE2_WITH_TEMPLATE_PREFIX(template<> RENDERCORE_API, TOneColorVS##A##B, SF_Vertex);
diff --git a/Engine/Source/Runtime/RenderCore/Public/OneColorShader.h b/Engine/Source/Runtime/RenderCore/Public/OneColorShader.h
--- a/Engine/Source/Runtime/RenderCore/Public/OneColorShader.h
+++ b/Engine/Source/Runtime/RenderCore/Public/OneColorShader.h
@@ -8,6 +8,7 @@
 #include "Shader.h"
 #include "GlobalShader.h"
 #include "ShaderParameterUtils.h"
+#include "ShaderParameterStruct.h"
 
 /**
  * Vertex shader for rendering a single, constant color.
@@ -65,24 +66,36 @@ private:
 class RENDERCORE_API FOneColorPS : public FGlobalShader
 {
 	DECLARE_GLOBAL_SHADER(FOneColorPS);
-public:
-	
-	FOneColorPS( )	{ }
-	FOneColorPS(const ShaderMetaType::CompiledShaderInitializerType& Initializer)
-	:	FGlobalShader( Initializer )
-	{
-		//ColorParameter.Bind( Initializer.ParameterMap, TEXT("DrawColorMRT"), SPF_Mandatory);
-	}
+	SHADER_USE_PARAMETER_STRUCT(FOneColorPS, FGlobalShader);
 
-	void SetColors(FRHICommandList& RHICmdList, const FLinearColor* Colors, int32 NumColors);
+	BEGIN_SHADER_PARAMETER_STRUCT(FParameters, RENDERCORE_API)
+		SHADER_PARAMETER_ARRAY(FLinearColor, DrawColorMRT, [MaxSimultaneousRenderTargets])
+		RENDER_TARGET_BINDING_SLOTS()
+	END_SHADER_PARAMETER_STRUCT()
 
-	static bool ShouldCompilePermutation(const FGlobalShaderPermutationParameters& Parameters)
+	void FillParameters(FParameters& Parameters, const FLinearColor* Colors, int32 NumColors)
 	{
-		return true;
+		check(NumColors <= MaxSimultaneousRenderTargets);
+
+		int32 Index = 0;
+
+		for (; Index < NumColors; ++Index)
+		{
+			Parameters.DrawColorMRT[Index] = Colors[Index];
+		}
+
+		for (; Index < MaxSimultaneousRenderTargets; ++Index)
+		{
+			Parameters.DrawColorMRT[Index] = FLinearColor::Black;
+		}
 	}
 
-	/** The parameter to use for setting the draw Color. */
-	//LAYOUT_FIELD(FShaderParameter, ColorParameter);
+	static void SetColors(FRHICommandList& RHICmdList, const TShaderMapRef<FOneColorPS>& Shader, const FLinearColor* Colors, int32 NumColors)
+	{
+		FParameters Parameters;
+		Shader->FillParameters(Parameters, Colors, NumColors);
+		SetShaderParameters(RHICmdList, Shader, Shader.GetPixelShader(), Parameters);
+	}
 };
 
 /**
@@ -125,6 +138,13 @@ public:
 			OutEnvironment.SetRenderTargetOutputFormat(0, PF_A32B32G32R32F);
 		}
 	}
+
+	static void SetColors(FRHICommandList& RHICmdList, const TShaderMapRef<TOneColorPixelShaderMRT>& Shader, const FLinearColor* Colors, int32 NumColors)
+	{
+		FParameters Parameters;
+		Shader->FillParameters(Parameters, Colors, NumColors);
+		SetShaderParameters(RHICmdList, Shader, Shader.GetPixelShader(), Parameters);
+	}
 };
 
 /**
diff --git a/Engine/Source/Runtime/RenderCore/Public/RenderingThread.h b/Engine/Source/Runtime/RenderCore/Public/RenderingThread.h
--- a/Engine/Source/Runtime/RenderCore/Public/RenderingThread.h
+++ b/Engine/Source/Runtime/RenderCore/Public/RenderingThread.h
@@ -361,6 +361,130 @@ private:
 	RenderCommandFunctionArray* RenderCommands;
 };
 
+struct FRenderThreadStructBase
+{
+	FRenderThreadStructBase() = default;
+
+	// Copy construction is not allowed. Used to avoid accidental copying in the command lambda.
+	FRenderThreadStructBase(const FRenderThreadStructBase&) = delete;
+
+	void InitRHI(FRHICommandListImmediate&) {}
+	void ReleaseRHI(FRHICommandListImmediate&) {}
+};
+
+/**  Represents a struct with a lifetime that spans multiple render commands with scoped initialization
+ *   and release on the render thread.
+ * 
+ *   Example:
+ * 
+ *   struct FMyStruct : public FRenderThreadStructBase
+ *   {
+ *       FInitializer { int32 Foo; int32 Bar; };
+ * 
+ *       FMyStruct(const FInitializer& InInitializer)
+ *            : Initializer(InInitializer)
+ *       {
+ *            // Called immediately by TRenderThreadStruct when created.
+ *       }
+ * 
+ *       ~FMyStruct()
+ *       {
+ *           // Called on the render thread when TRenderThreadStruct goes out of scope.
+ *       }
+ * 
+ *       void InitRHI(FRHICommandListImmediate& RHICmdList)
+ *       {
+ *           // Called on the render thread by TRenderThreadStruct when created.
+ *       }
+ * 
+ *       void ReleaseRHI(FRHICommandListImmediate& RHICmdList)
+ *       {
+ *           // Called on the render thread when TRenderThreadStruct goes out of scope.
+ *       }
+ * 
+ *       FInitializer Initializer;
+ *   };
+ *
+ *   // On Main Thread
+ * 
+ *   {
+ *       TRenderThreadStruct<FMyStruct> MyStruct(FMyStruct::FInitializer{1, 2});
+ * 
+ *       ENQUEUE_RENDER_COMMAND(CommandA)[MyStruct = MyStruct.Get()](FRHICommandListImmediate& RHICmdList)
+ *       {
+ *           // Do something with MyStruct.
+ *       };
+ * 
+ *       ENQUEUE_RENDER_COMMAND(CommandB)[MyStruct = MyStrucft.Get()](FRHICommandListImmediate& RHICmdList)
+ *       {
+ *           // Do something else with MyStruct.
+ *       };
+ * 
+ *       // MyStruct instance is automatically released and deleted on the render thread.
+ *   }
+ */
+template <typename StructType>
+class TRenderThreadStruct
+{
+public:
+	static_assert(TIsDerivedFrom<StructType, FRenderThreadStructBase>::IsDerived, "StructType must be derived from FRenderThreadStructBase.");
+
+	template <typename... TArgs>
+	TRenderThreadStruct(TArgs&&... Args)
+		: Struct(new StructType(Forward<TArgs&&>(Args)...))
+	{
+		ENQUEUE_RENDER_COMMAND(InitStruct)([Struct = Struct](FRHICommandListImmediate& RHICmdList)
+		{
+			Struct->InitRHI(RHICmdList);
+		});
+	}
+
+	~TRenderThreadStruct()
+	{
+		ENQUEUE_RENDER_COMMAND(DeleteStruct)([Struct = Struct](FRHICommandListImmediate& RHICmdList)
+		{
+			Struct->ReleaseRHI(RHICmdList);
+			delete Struct;
+		});
+		Struct = nullptr;
+	}
+
+	TRenderThreadStruct(const TRenderThreadStruct&) = delete;
+
+	const StructType* operator->() const
+	{
+		return Struct;
+	}
+
+	StructType* operator->()
+	{
+		return Struct;
+	}
+
+	const StructType& operator*() const
+	{
+		return *Struct;
+	}
+
+	StructType& operator*()
+	{
+		return *Struct;
+	}
+
+	const StructType* Get() const
+	{
+		return Struct;
+	}
+
+	StructType* Get()
+	{
+		return Struct;
+	}
+
+private:
+	StructType* Struct;
+};
+
 DECLARE_MULTICAST_DELEGATE(FStopRenderingThread);
 using FStopRenderingThreadDelegate = FStopRenderingThread::FDelegate;
 
diff --git a/Engine/Source/Runtime/Renderer/Private/BasePassRendering.cpp b/Engine/Source/Runtime/Renderer/Private/BasePassRendering.cpp
--- a/Engine/Source/Runtime/Renderer/Private/BasePassRendering.cpp
+++ b/Engine/Source/Runtime/Renderer/Private/BasePassRendering.cpp
@@ -684,7 +684,7 @@ static void ClearGBufferAtMaxZ(
 			{},
 			PassParameters,
 			ERDGPassFlags::Raster,
-			[&View, VertexShader, PixelShader, ActiveTargetCount, ClearColor0](FRHICommandListImmediate& RHICmdList)
+			[&View, VertexShader, PixelShader, ActiveTargetCount, ClearColor0](FRHICommandList& RHICmdList)
 		{
 			const FLinearColor ClearColors[MaxSimultaneousRenderTargets] =
 			{
@@ -715,7 +715,8 @@ static void ClearGBufferAtMaxZ(
 			VertexShader->SetDepthParameter(RHICmdList, float(ERHIZBuffer::FarPlane));
 
 			RHICmdList.SetViewport(View.ViewRect.Min.X, View.ViewRect.Min.Y, 0, View.ViewRect.Max.X, View.ViewRect.Max.Y, 1);
-			PixelShader->SetColors(RHICmdList, ClearColors, ActiveTargetCount);
+			PixelShader->SetColors(RHICmdList, PixelShader, ClearColors, ActiveTargetCount);
+
 			RHICmdList.SetStreamSource(0, GClearVertexBuffer.VertexBufferRHI, 0);
 			RHICmdList.DrawPrimitive(0, 2, 1);
 		});
diff --git a/Engine/Source/Runtime/Renderer/Private/IndirectLightRendering.cpp b/Engine/Source/Runtime/Renderer/Private/IndirectLightRendering.cpp
--- a/Engine/Source/Runtime/Renderer/Private/IndirectLightRendering.cpp
+++ b/Engine/Source/Runtime/Renderer/Private/IndirectLightRendering.cpp
@@ -526,7 +526,7 @@ void SetupReflectionUniformParameters(const FViewInfo& View, FReflectionUniformP
 		&& Scene->ReflectionSceneData.CubemapArray.IsValid()
 		&& Scene->ReflectionSceneData.RegisteredReflectionCaptures.Num())
 	{
-		CubeArrayTexture = Scene->ReflectionSceneData.CubemapArray.GetRenderTarget().ShaderResourceTexture;
+		CubeArrayTexture = Scene->ReflectionSceneData.CubemapArray.GetRenderTarget()->GetRHI();
 	}
 
 	OutParameters.ReflectionCubemap = CubeArrayTexture;
diff --git a/Engine/Source/Runtime/Renderer/Private/MobileReflectionEnvironmentCapture.cpp b/Engine/Source/Runtime/Renderer/Private/MobileReflectionEnvironmentCapture.cpp
--- a/Engine/Source/Runtime/Renderer/Private/MobileReflectionEnvironmentCapture.cpp
+++ b/Engine/Source/Runtime/Renderer/Private/MobileReflectionEnvironmentCapture.cpp
@@ -13,462 +13,89 @@
 #include "PipelineStateCache.h"
 #include "SceneFilterRendering.h"
 #include "OneColorShader.h"
+#include "PixelShaderUtils.h"
 
-extern int32 GDiffuseIrradianceCubemapSize;
-extern float ComputeSingleAverageBrightnessFromCubemap(FRHICommandListImmediate& RHICmdList, ERHIFeatureLevel::Type FeatureLevel, int32 TargetSize, FSceneRenderTargetItem& Cubemap);
-extern void FullyResolveReflectionScratchCubes(FRHICommandListImmediate& RHICmdList);
+/** Computes the average brightness of the given reflection capture and stores it in the scene. */
+extern void ComputeSingleAverageBrightnessFromCubemap(FRDGBuilder& GraphBuilder, FGlobalShaderMap* ShaderMap, FRDGTexture* CubemapTexture, float* OutAverageBrightness);
 
-static TAutoConsoleVariable<int32> CVarMobileUseHighQualitySkyCaptureFiltering(
-	TEXT("r.Mobile.HighQualitySkyCaptureFiltering"),
-	1,
-	TEXT("1: (default) use high quality filtering when generating mobile sky captures.")
-	TEXT("0: use simple bilinear filtering when generating mobile sky captures."),
-	ECVF_RenderThreadSafe);
+extern FRDGTexture* FilterCubeMap(FRDGBuilder& GraphBuilder, FGlobalShaderMap* ShaderMap, FRDGTexture* SourceTexture);
 
-extern TGlobalResource<FReflectionScratchCubemaps> GReflectionScratchCubemaps;
+extern void PremultiplyCubeMipAlpha(FRDGBuilder& GraphBuilder, FGlobalShaderMap* ShaderMap, FRDGTexture* CubemapTexture, int32 MipIndex);
 
-class FMobileDownsamplePS : public FGlobalShader
+class FMobileCubeDownsamplePS : public FGlobalShader
 {
-	DECLARE_SHADER_TYPE(FMobileDownsamplePS, Global);
-public:
-
-	static bool ShouldCompilePermutation(const FGlobalShaderPermutationParameters& Parameters)
-	{
-		return IsMobilePlatform(Parameters.Platform);
-	}
-
-	FMobileDownsamplePS(const ShaderMetaType::CompiledShaderInitializerType& Initializer) :
-		FGlobalShader(Initializer)
-	{
-		CubeFace.Bind(Initializer.ParameterMap, TEXT("CubeFace"));
-		SourceMipIndex.Bind(Initializer.ParameterMap, TEXT("SourceMipIndex"));
-		SourceCubemapTexture.Bind(Initializer.ParameterMap, TEXT("SourceCubemapTexture"));
-		SourceCubemapSampler.Bind(Initializer.ParameterMap, TEXT("SourceCubemapSampler"));
-	}
-	FMobileDownsamplePS() {}
-
-	void SetParameters(FRHICommandList& RHICmdList, int32 CubeFaceValue, int32 SourceMipIndexValue, FSceneRenderTargetItem& SourceTextureValue)
-	{
-		SetShaderValue(RHICmdList, RHICmdList.GetBoundPixelShader(), CubeFace, CubeFaceValue);
-		SetShaderValue(RHICmdList, RHICmdList.GetBoundPixelShader(), SourceMipIndex, SourceMipIndexValue);
-
-		SetTextureParameter(
-			RHICmdList,
-			RHICmdList.GetBoundPixelShader(),
-			SourceCubemapTexture,
-			SourceCubemapSampler,
-			TStaticSamplerState<SF_Bilinear, AM_Clamp, AM_Clamp, AM_Clamp>::GetRHI(),
-			SourceTextureValue.ShaderResourceTexture);
-	}
-
-private:
-	LAYOUT_FIELD(FShaderParameter, CubeFace)
-	LAYOUT_FIELD(FShaderParameter, SourceMipIndex)
-	LAYOUT_FIELD(FShaderResourceParameter, SourceCubemapTexture)
-	LAYOUT_FIELD(FShaderResourceParameter, SourceCubemapSampler)
+	DECLARE_GLOBAL_SHADER(FMobileCubeDownsamplePS);
+	SHADER_USE_PARAMETER_STRUCT(FMobileCubeDownsamplePS, FGlobalShader);
+
+	BEGIN_SHADER_PARAMETER_STRUCT(FParameters, )
+		SHADER_PARAMETER_RDG_TEXTURE_SRV(TextureCube, SourceCubemapTexture)
+		SHADER_PARAMETER_SAMPLER(SamplerState, SourceCubemapSampler)
+		SHADER_PARAMETER(FVector2f, SvPositionToUVScale)
+		SHADER_PARAMETER(int32, CubeFace)
+		SHADER_PARAMETER(int32, SourceMipIndex)
+		RENDER_TARGET_BINDING_SLOTS()
+	END_SHADER_PARAMETER_STRUCT()
 };
 
-IMPLEMENT_SHADER_TYPE(, FMobileDownsamplePS, TEXT("/Engine/Private/ReflectionEnvironmentShaders.usf"), TEXT("DownsamplePS_Mobile"), SF_Pixel);
+IMPLEMENT_GLOBAL_SHADER(FMobileCubeDownsamplePS, "/Engine/Private/ReflectionEnvironmentShaders.usf", "DownsamplePS_Mobile", SF_Pixel);
+
 namespace MobileReflectionEnvironmentCapture
 {
-	/** Encapsulates render target picking logic for cubemap mip generation. */
-	FSceneRenderTargetItem& GetEffectiveRenderTarget(bool bDownsamplePass, int32 TargetMipIndex)
+	void CreateCubeMips(FRDGBuilder& GraphBuilder, FGlobalShaderMap* ShaderMap, FRDGTexture* CubemapTexture)
 	{
-		int32 ScratchTextureIndex = TargetMipIndex % 2;
+		RDG_EVENT_SCOPE(GraphBuilder, "CreateCubeMips");
 
-		if (!bDownsamplePass)
-		{
-			ScratchTextureIndex = 1 - ScratchTextureIndex;
-		}
-
-		return GReflectionScratchCubemaps.Color[ScratchTextureIndex]->GetRenderTargetItem();
-	}
+		TShaderMapRef<FMobileCubeDownsamplePS> PixelShader(ShaderMap);
 
-	/** Encapsulates source texture picking logic for cubemap mip generation. */
-	FSceneRenderTargetItem& GetEffectiveSourceTexture(bool bDownsamplePass, int32 TargetMipIndex)
-	{
-		int32 ScratchTextureIndex = TargetMipIndex % 2;
+		const int32 NumMips = CubemapTexture->Desc.NumMips;
 
-		if (bDownsamplePass)
+		// Downsample all the mips, each one reads from the mip above it
+		for (int32 MipIndex = 1; MipIndex < NumMips; MipIndex++)
 		{
-			ScratchTextureIndex = 1 - ScratchTextureIndex;
-		}
-
-		return GReflectionScratchCubemaps.Color[ScratchTextureIndex]->GetRenderTargetItem();
-	}
-
-	void ComputeAverageBrightness(FRHICommandListImmediate& RHICmdList, ERHIFeatureLevel::Type FeatureLevel, int32 CubmapSize, float& OutAverageBrightness)
-	{
-		SCOPED_DRAW_EVENT(RHICmdList, ComputeAverageBrightness);
-
-		const int32 EffectiveTopMipSize = CubmapSize;
-		const int32 NumMips = FMath::CeilLogTwo(EffectiveTopMipSize) + 1;
-
-		// necessary to resolve the clears which touched all the mips.  scene rendering only resolves mip 0.
-		FullyResolveReflectionScratchCubes(RHICmdList);
-
-		auto ShaderMap = GetGlobalShaderMap(FeatureLevel);
+			const int32 SourceMipIndex = FMath::Max(MipIndex - 1, 0);
+			const int32 MipSize = 1 << (NumMips - MipIndex - 1);
+			const FIntRect ViewRect(0, 0, MipSize, MipSize);
 
-		{
-			SCOPED_DRAW_EVENT(RHICmdList, DownsampleCubeMips);
-
-			// Downsample all the mips, each one reads from the mip above it
-			for (int32 MipIndex = 1; MipIndex < NumMips; MipIndex++)
+			for (int32 CubeFace = 0; CubeFace < CubeFace_MAX; CubeFace++)
 			{
-				const int32 SourceMipIndex = FMath::Max(MipIndex - 1, 0);
-				const int32 MipSize = 1 << (NumMips - MipIndex - 1);
-
-				FSceneRenderTargetItem& EffectiveRT = GetEffectiveRenderTarget(true, MipIndex);
-				FSceneRenderTargetItem& EffectiveSource = GetEffectiveSourceTexture(true, MipIndex);
-				check(EffectiveRT.TargetableTexture != EffectiveSource.ShaderResourceTexture);
-
-				for (int32 CubeFace = 0; CubeFace < CubeFace_MAX; CubeFace++)
-				{
-					FRHIRenderPassInfo RPInfo(EffectiveRT.TargetableTexture, ERenderTargetActions::Load_Store);
-					RPInfo.ColorRenderTargets[0].ArraySlice = CubeFace;
-					RPInfo.ColorRenderTargets[0].MipIndex = MipIndex;
-					TransitionRenderPassTargets(RHICmdList, RPInfo);
-					RHICmdList.BeginRenderPass(RPInfo, TEXT("AverageBrightness"));
-					{
-						FGraphicsPipelineStateInitializer GraphicsPSOInit;
-						RHICmdList.ApplyCachedRenderTargets(GraphicsPSOInit);
-						GraphicsPSOInit.RasterizerState = TStaticRasterizerState<FM_Solid, CM_None>::GetRHI();
-						GraphicsPSOInit.DepthStencilState = TStaticDepthStencilState<false, CF_Always>::GetRHI();
-						GraphicsPSOInit.BlendState = TStaticBlendState<>::GetRHI();
-
-						const FIntRect ViewRect(0, 0, MipSize, MipSize);
-						RHICmdList.SetViewport(0, 0, 0.0f, MipSize, MipSize, 1.0f);
-
-						TShaderMapRef<FScreenVS> VertexShader(ShaderMap);
-						TShaderMapRef<FMobileDownsamplePS> PixelShader(ShaderMap);
-
-						GraphicsPSOInit.BoundShaderState.VertexDeclarationRHI = GFilterVertexDeclaration.VertexDeclarationRHI;
-						GraphicsPSOInit.BoundShaderState.VertexShaderRHI = VertexShader.GetVertexShader();
-						GraphicsPSOInit.BoundShaderState.PixelShaderRHI = PixelShader.GetPixelShader();
-						GraphicsPSOInit.PrimitiveType = PT_TriangleList;
-						SetGraphicsPipelineState(RHICmdList, GraphicsPSOInit, 0);
-
-						PixelShader->SetParameters(RHICmdList, CubeFace, SourceMipIndex, EffectiveSource);
-
-						DrawRectangle(
-							RHICmdList,
-							ViewRect.Min.X, ViewRect.Min.Y,
-							ViewRect.Width(), ViewRect.Height(),
-							ViewRect.Min.X, ViewRect.Min.Y,
-							ViewRect.Width(), ViewRect.Height(),
-							FIntPoint(ViewRect.Width(), ViewRect.Height()),
-							FIntPoint(MipSize, MipSize),
-							VertexShader);
-					}
-					RHICmdList.EndRenderPass();
-
-					FResolveParams ResolveParams(FResolveRect(), (ECubeFace)CubeFace, MipIndex);
-					ResolveParams.SourceAccessFinal = ERHIAccess::SRVMask;
-					ResolveParams.DestAccessFinal = ERHIAccess::SRVMask;
-					RHICmdList.CopyToResolveTarget(EffectiveRT.TargetableTexture, EffectiveRT.ShaderResourceTexture, ResolveParams);
-				}
+				auto* PassParameters = GraphBuilder.AllocParameters<FMobileCubeDownsamplePS::FParameters>();
+				PassParameters->SourceCubemapTexture = GraphBuilder.CreateSRV(FRDGTextureSRVDesc::CreateForMipLevel(CubemapTexture, MipIndex - 1));
+				PassParameters->SourceCubemapSampler = TStaticSamplerState<SF_Bilinear, AM_Clamp, AM_Clamp, AM_Clamp>::GetRHI();
+				PassParameters->CubeFace = CubeFace;
+				PassParameters->SourceMipIndex = SourceMipIndex;
+				PassParameters->SvPositionToUVScale = FVector2f(1.0f / MipSize, 1.0f / MipSize);
+				PassParameters->RenderTargets[0] = FRenderTargetBinding(CubemapTexture, ERenderTargetLoadAction::ENoAction, MipIndex, CubeFace);
+
+				FPixelShaderUtils::AddFullscreenPass(
+					GraphBuilder,
+					ShaderMap,
+					RDG_EVENT_NAME("CreateCubeMips (Mip: %d, Face: %d)", MipIndex, CubeFace),
+					PixelShader,
+					PassParameters,
+					ViewRect);
 			}
 		}
-
-		OutAverageBrightness = ComputeSingleAverageBrightnessFromCubemap(RHICmdList, FeatureLevel, CubmapSize, GetEffectiveRenderTarget(true, NumMips - 1));
 	}
 
-	void CopyToSkyTexture(FRHICommandList& RHICmdList, FScene* Scene, FTexture* ProcessedTexture)
+	void ComputeAverageBrightness(FRDGBuilder& GraphBuilder, FGlobalShaderMap* ShaderMap, FRDGTexture* CubemapTexture, float* OutAverageBrightness)
 	{
-		SCOPED_DRAW_EVENT(RHICmdList, CopyToSkyTexture);
-		if (ProcessedTexture->TextureRHI)
-		{
-			const bool bUseHQFiltering = CVarMobileUseHighQualitySkyCaptureFiltering.GetValueOnRenderThread() == 1;
-			const int32 EffectiveTopMipSize = ProcessedTexture->GetSizeX();
-			const int32 NumMips = FMath::CeilLogTwo(EffectiveTopMipSize) + 1;
-
-			RHICmdList.Transition(FRHITransitionInfo(ProcessedTexture->TextureRHI, ERHIAccess::Unknown, ERHIAccess::CopyDest));
-
-			FRHICopyTextureInfo CopyInfo;
-			CopyInfo.Size = FIntVector(ProcessedTexture->GetSizeX(), ProcessedTexture->GetSizeY(), 1);
-			CopyInfo.NumSlices = 6;
-
-			// GPU copy back to the skylight's texture, which is not a render target
-			for (int32 MipIndex = 0; MipIndex < NumMips; MipIndex++)
-			{
-				// For simple mobile bilin filtering the source for this copy is the dest from the filtering pass.
-				// In the HQ case, the full image is contained in GetEffectiveRenderTarget(.., false,0).
-				FSceneRenderTargetItem& EffectiveSource = GetEffectiveRenderTarget(false, bUseHQFiltering ? 0 : MipIndex);
-				RHICmdList.Transition(FRHITransitionInfo(EffectiveSource.ShaderResourceTexture, ERHIAccess::Unknown, ERHIAccess::CopySrc));
-				RHICmdList.CopyTexture(EffectiveSource.ShaderResourceTexture, ProcessedTexture->TextureRHI, CopyInfo);
-
-				++CopyInfo.SourceMipIndex;
-				++CopyInfo.DestMipIndex;
-				CopyInfo.Size.X = FMath::Max(1, CopyInfo.Size.X / 2);
-				CopyInfo.Size.Y = FMath::Max(1, CopyInfo.Size.Y / 2);
-			}
-
-			RHICmdList.Transition(FRHITransitionInfo(ProcessedTexture->TextureRHI, ERHIAccess::CopyDest, ERHIAccess::SRVMask));
-		}
+		CreateCubeMips(GraphBuilder, ShaderMap, CubemapTexture);
+		ComputeSingleAverageBrightnessFromCubemap(GraphBuilder, ShaderMap, CubemapTexture, OutAverageBrightness);
 	}
 
 	/** Generates mips for glossiness and filters the cubemap for a given reflection. */
-	void FilterReflectionEnvironment(FRHICommandListImmediate& RHICmdList, ERHIFeatureLevel::Type FeatureLevel, int32 CubmapSize, FSHVectorRGB3* OutIrradianceEnvironmentMap)
+	FRDGTexture* FilterReflectionEnvironment(FRDGBuilder& GraphBuilder, FGlobalShaderMap* ShaderMap, FRDGTexture* CubemapTexture, FSHVectorRGB3* OutIrradianceEnvironmentMap)
 	{
-		SCOPED_DRAW_EVENT(RHICmdList, FilterReflectionEnvironment);
-
-		const int32 EffectiveTopMipSize = CubmapSize;
-		const int32 NumMips = FMath::CeilLogTwo(EffectiveTopMipSize) + 1;
-		const bool bUseHQFiltering = CVarMobileUseHighQualitySkyCaptureFiltering.GetValueOnRenderThread() == 1;
-		{
-			FSceneRenderTargetItem& EffectiveColorRT = GReflectionScratchCubemaps.Color[0]->GetRenderTargetItem();
-			// Premultiply alpha in-place using alpha blending
-			for (uint32 CubeFace = 0; CubeFace < CubeFace_MAX; CubeFace++)
-			{
-				FRHIRenderPassInfo RPInfo(EffectiveColorRT.TargetableTexture, ERenderTargetActions::Load_Store);
-				RPInfo.ColorRenderTargets[0].ArraySlice = CubeFace;
-				RPInfo.ColorRenderTargets[0].MipIndex = 0;
-
-				TransitionRenderPassTargets(RHICmdList, RPInfo);
-				RHICmdList.BeginRenderPass(RPInfo, TEXT("FilterReflectionEnvironment"));
-				{
-					FGraphicsPipelineStateInitializer GraphicsPSOInit;
-					RHICmdList.ApplyCachedRenderTargets(GraphicsPSOInit);
-					GraphicsPSOInit.RasterizerState = TStaticRasterizerState<FM_Solid, CM_None>::GetRHI();
-					GraphicsPSOInit.DepthStencilState = TStaticDepthStencilState<false, CF_Always>::GetRHI();
-					GraphicsPSOInit.BlendState = TStaticBlendState<CW_RGBA, BO_Add, BF_Zero, BF_DestAlpha, BO_Add, BF_Zero, BF_One>::GetRHI();
+		RDG_EVENT_SCOPE(GraphBuilder, "FilterReflectionEnvironment");
 
-					const FIntPoint SourceDimensions(CubmapSize, CubmapSize);
-					const FIntRect ViewRect(0, 0, EffectiveTopMipSize, EffectiveTopMipSize);
-					RHICmdList.SetViewport(0, 0, 0.0f, EffectiveTopMipSize, EffectiveTopMipSize, 1.0f);
-
-					TShaderMapRef<FScreenVS> VertexShader(GetGlobalShaderMap(FeatureLevel));
-					TShaderMapRef<FOneColorPS> PixelShader(GetGlobalShaderMap(FeatureLevel));
-
-					GraphicsPSOInit.BoundShaderState.VertexDeclarationRHI = GFilterVertexDeclaration.VertexDeclarationRHI;
-					GraphicsPSOInit.BoundShaderState.VertexShaderRHI = VertexShader.GetVertexShader();
-					GraphicsPSOInit.BoundShaderState.PixelShaderRHI = PixelShader.GetPixelShader();
-					GraphicsPSOInit.PrimitiveType = PT_TriangleList;
-
-					SetGraphicsPipelineState(RHICmdList, GraphicsPSOInit, 0);
-
-					FLinearColor UnusedColors[1] = { FLinearColor::Black };
-					PixelShader->SetColors(RHICmdList, UnusedColors, UE_ARRAY_COUNT(UnusedColors));
-
-					DrawRectangle(
-						RHICmdList,
-						ViewRect.Min.X, ViewRect.Min.Y,
-						ViewRect.Width(), ViewRect.Height(),
-						0, 0,
-						SourceDimensions.X, SourceDimensions.Y,
-						FIntPoint(ViewRect.Width(), ViewRect.Height()),
-						SourceDimensions,
-						VertexShader);
-				}
-				RHICmdList.EndRenderPass();
-
-				FResolveParams ResolveParams(FResolveRect(), (ECubeFace)CubeFace);
-				ResolveParams.SourceAccessFinal = ERHIAccess::SRVMask;
-				ResolveParams.DestAccessFinal = ERHIAccess::SRVMask;
-				RHICmdList.CopyToResolveTarget(EffectiveColorRT.TargetableTexture, EffectiveColorRT.ShaderResourceTexture, ResolveParams);
-			} // end for
-		}
-
-		int32 DiffuseConvolutionSourceMip = INDEX_NONE;
-		FSceneRenderTargetItem* DiffuseConvolutionSource = NULL;
-
-		auto ShaderMap = GetGlobalShaderMap(FeatureLevel);
-
-		{
-			SCOPED_DRAW_EVENT(RHICmdList, DownsampleCubeMips);
-			// Downsample all the mips, each one reads from the mip above it
-			for (int32 MipIndex = 1; MipIndex < NumMips; MipIndex++)
-			{
-				SCOPED_DRAW_EVENT(RHICmdList, DownsampleCubeMip);
-				const int32 SourceMipIndex = FMath::Max(MipIndex - 1, 0);
-				const int32 MipSize = 1 << (NumMips - MipIndex - 1);
-
-				FSceneRenderTargetItem& EffectiveRT = GetEffectiveRenderTarget(true, MipIndex);
-				FSceneRenderTargetItem& EffectiveSource = GetEffectiveSourceTexture(true, MipIndex);
-				check(EffectiveRT.TargetableTexture != EffectiveSource.ShaderResourceTexture);
-
-				for (int32 CubeFace = 0; CubeFace < CubeFace_MAX; CubeFace++)
-				{
-					FRHIRenderPassInfo RPInfo(EffectiveRT.TargetableTexture, ERenderTargetActions::Load_Store);
-					RPInfo.ColorRenderTargets[0].ArraySlice = CubeFace;
-					RPInfo.ColorRenderTargets[0].MipIndex = MipIndex;
-					TransitionRenderPassTargets(RHICmdList, RPInfo);
-
-					RHICmdList.BeginRenderPass(RPInfo, TEXT("DownsampleCubemap"));
-					{
-
-						FGraphicsPipelineStateInitializer GraphicsPSOInit;
-						RHICmdList.ApplyCachedRenderTargets(GraphicsPSOInit);
-						GraphicsPSOInit.RasterizerState = TStaticRasterizerState<FM_Solid, CM_None>::GetRHI();
-						GraphicsPSOInit.DepthStencilState = TStaticDepthStencilState<false, CF_Always>::GetRHI();
-						GraphicsPSOInit.BlendState = TStaticBlendState<>::GetRHI();
-
-						const FIntRect ViewRect(0, 0, MipSize, MipSize);
-						RHICmdList.SetViewport(0, 0, 0.0f, MipSize, MipSize, 1.0f);
-
-						TShaderMapRef<FScreenVS> VertexShader(ShaderMap);
-						TShaderMapRef<FMobileDownsamplePS> PixelShader(ShaderMap);
-
-						GraphicsPSOInit.BoundShaderState.VertexDeclarationRHI = GFilterVertexDeclaration.VertexDeclarationRHI;
-						GraphicsPSOInit.BoundShaderState.VertexShaderRHI = VertexShader.GetVertexShader();
-						GraphicsPSOInit.BoundShaderState.PixelShaderRHI = PixelShader.GetPixelShader();
-						GraphicsPSOInit.PrimitiveType = PT_TriangleList;
-
-						SetGraphicsPipelineState(RHICmdList, GraphicsPSOInit, 0);
-
-						PixelShader->SetParameters(RHICmdList, CubeFace, SourceMipIndex, EffectiveSource);
-
-						DrawRectangle(
-							RHICmdList,
-							ViewRect.Min.X, ViewRect.Min.Y,
-							ViewRect.Width(), ViewRect.Height(),
-							ViewRect.Min.X, ViewRect.Min.Y,
-							ViewRect.Width(), ViewRect.Height(),
-							FIntPoint(ViewRect.Width(), ViewRect.Height()),
-							FIntPoint(MipSize, MipSize),
-							VertexShader);
-					}
-					RHICmdList.EndRenderPass();
-
-					FResolveParams ResolveParams(FResolveRect(), (ECubeFace)CubeFace, MipIndex);
-					ResolveParams.SourceAccessFinal = ERHIAccess::SRVMask;
-					ResolveParams.DestAccessFinal = ERHIAccess::SRVMask;
-					RHICmdList.CopyToResolveTarget(EffectiveRT.TargetableTexture, EffectiveRT.ShaderResourceTexture, ResolveParams);
-				}
-
-				if (DiffuseConvolutionSource == NULL && MipSize <= GDiffuseIrradianceCubemapSize)
-				{
-					DiffuseConvolutionSourceMip = MipIndex;
-					DiffuseConvolutionSource = &EffectiveRT;
-				}
-			}
-		}
+		PremultiplyCubeMipAlpha(GraphBuilder, ShaderMap, CubemapTexture, 0);
+		CreateCubeMips(GraphBuilder, ShaderMap, CubemapTexture);
 
 		if (OutIrradianceEnvironmentMap)
 		{
-			SCOPED_DRAW_EVENT(RHICmdList, ComputeDiffuseIrradiance);
-			check(DiffuseConvolutionSource != NULL);
-			ComputeDiffuseIrradiance(RHICmdList, FeatureLevel, DiffuseConvolutionSource->ShaderResourceTexture, DiffuseConvolutionSourceMip, OutIrradianceEnvironmentMap);
-		}
-
-		if (bUseHQFiltering)
-		{
-			// When HQ filtering is enabled the filter shader requires access to all mips levels of the source cubemap.
-			// Here we ensure that GetEffectiveSourceTexture(0,false) will have a complete set of mips for this process.
-			SCOPED_DRAW_EVENT(RHICmdList, PrepareSourceCubemapMipsForHQFiltering);
-			for (int32 MipIndex = 1; MipIndex < NumMips; MipIndex += 2)
-			{
-				FSceneRenderTargetItem& SourceTarget = GetEffectiveRenderTarget(true, MipIndex);
-				FSceneRenderTargetItem& DestTarget = GetEffectiveSourceTexture(true, MipIndex);
-				check(DestTarget.TargetableTexture != SourceTarget.ShaderResourceTexture);
-
-				// Transition the textures once, so CopyToResolveTarget doesn't ping-pong uselessly between the copy and SRV states.
-				FRHITransitionInfo TransitionsBefore[] = {
-					FRHITransitionInfo(SourceTarget.ShaderResourceTexture, ERHIAccess::SRVMask, ERHIAccess::CopySrc),
-					FRHITransitionInfo(DestTarget.ShaderResourceTexture, ERHIAccess::SRVMask, ERHIAccess::CopyDest)
-				};
-				RHICmdList.Transition(MakeArrayView(TransitionsBefore, UE_ARRAY_COUNT(TransitionsBefore)));
-
-				// Tell CopyToResolveTarget to leave the textures in the copy state, because we'll transition them only once when we're done.
-				FResolveParams ResolveParams(FResolveRect(), CubeFace_PosX, MipIndex);
-				ResolveParams.SourceAccessFinal = ERHIAccess::CopySrc;
-				ResolveParams.DestAccessFinal = ERHIAccess::CopyDest;
-
-				for (int32 CubeFace = 0; CubeFace < CubeFace_MAX; CubeFace++)
-				{
-					ResolveParams.CubeFace = (ECubeFace)CubeFace;
-					RHICmdList.CopyToResolveTarget(SourceTarget.ShaderResourceTexture, DestTarget.ShaderResourceTexture, ResolveParams);
-				}
-
-				// We're done copying, transition the textures back to SRV.
-				FRHITransitionInfo TransitionsAfter[] = {
-					FRHITransitionInfo(SourceTarget.ShaderResourceTexture, ERHIAccess::CopySrc, ERHIAccess::SRVMask),
-					FRHITransitionInfo(DestTarget.ShaderResourceTexture, ERHIAccess::CopyDest, ERHIAccess::SRVMask)
-				};
-				RHICmdList.Transition(MakeArrayView(TransitionsAfter, UE_ARRAY_COUNT(TransitionsAfter)));
-			}
+			ComputeDiffuseIrradiance(GraphBuilder, ShaderMap, CubemapTexture, OutIrradianceEnvironmentMap);
 		}
 
-		{
-			SCOPED_DRAW_EVENT(RHICmdList, FilterCubeMap);
-			// Filter all the mips.
-			// in simple mobile bilin case each one reads from whichever scratch render target holds the downsampled contents, and writes to the destination cubemap.
-			// HQ case, all mips are contained in GetEffectiveSourceTexture(.., false,0) as it needs a complete mip chain.
-			for (int32 MipIndex = 0; MipIndex < NumMips; MipIndex++)
-			{
-				SCOPED_DRAW_EVENT(RHICmdList, FilterCubeMip);
-				FSceneRenderTargetItem& EffectiveRT = GetEffectiveRenderTarget(false, bUseHQFiltering ? 0 : MipIndex);
-				FSceneRenderTargetItem& EffectiveSource = GetEffectiveSourceTexture(false, bUseHQFiltering ? 0 : MipIndex);
-				check(EffectiveRT.TargetableTexture != EffectiveSource.ShaderResourceTexture);
-				const int32 MipSize = 1 << (NumMips - MipIndex - 1);
-
-				for (int32 CubeFace = 0; CubeFace < CubeFace_MAX; CubeFace++)
-				{
-					FRHIRenderPassInfo RPInfo(EffectiveRT.TargetableTexture, ERenderTargetActions::Load_Store);
-					RPInfo.ColorRenderTargets[0].ArraySlice = CubeFace;
-					RPInfo.ColorRenderTargets[0].MipIndex = MipIndex;
-					TransitionRenderPassTargets(RHICmdList, RPInfo);
-
-					RHICmdList.BeginRenderPass(RPInfo, TEXT("FilterCubeMip"));
-					{
-						FGraphicsPipelineStateInitializer GraphicsPSOInit;
-						RHICmdList.ApplyCachedRenderTargets(GraphicsPSOInit);
-						GraphicsPSOInit.RasterizerState = TStaticRasterizerState<FM_Solid, CM_None>::GetRHI();
-						GraphicsPSOInit.DepthStencilState = TStaticDepthStencilState<false, CF_Always>::GetRHI();
-						GraphicsPSOInit.BlendState = TStaticBlendState<>::GetRHI();
-
-						const FIntRect ViewRect(0, 0, MipSize, MipSize);
-						RHICmdList.SetViewport(0, 0, 0.0f, MipSize, MipSize, 1.0f);
-
-						TShaderMapRef<FScreenVS> VertexShader(GetGlobalShaderMap(FeatureLevel));
-
-						TShaderMapRef<TCubeFilterPS<0>> HQFilterPixelShader(ShaderMap);
-						TShaderMapRef<FMobileDownsamplePS> BilinFilterPixelShader(ShaderMap);
-						FRHIPixelShader* PixelShaderRHI = bUseHQFiltering ? HQFilterPixelShader.GetPixelShader() : BilinFilterPixelShader.GetPixelShader();
-
-						GraphicsPSOInit.BoundShaderState.VertexDeclarationRHI = GFilterVertexDeclaration.VertexDeclarationRHI;
-						GraphicsPSOInit.BoundShaderState.VertexShaderRHI = VertexShader.GetVertexShader();
-						GraphicsPSOInit.BoundShaderState.PixelShaderRHI = PixelShaderRHI;
-						GraphicsPSOInit.PrimitiveType = PT_TriangleList;
-
-						SetGraphicsPipelineState(RHICmdList, GraphicsPSOInit, 0);
-
-						if (bUseHQFiltering)
-						{
-							SetShaderValue(RHICmdList, PixelShaderRHI, HQFilterPixelShader->CubeFace, CubeFace);
-							SetShaderValue(RHICmdList, PixelShaderRHI, HQFilterPixelShader->MipIndex, MipIndex);
-							SetShaderValue(RHICmdList, PixelShaderRHI, HQFilterPixelShader->NumMips, NumMips);
-							SetTextureParameter(
-								RHICmdList,
-								PixelShaderRHI,
-								HQFilterPixelShader->SourceCubemapTexture,
-								HQFilterPixelShader->SourceCubemapSampler,
-								TStaticSamplerState<SF_Trilinear, AM_Clamp, AM_Clamp, AM_Clamp>::GetRHI(),
-								EffectiveSource.ShaderResourceTexture);
-						}
-						else
-						{
-							BilinFilterPixelShader->SetParameters(RHICmdList, CubeFace, MipIndex, EffectiveSource);
-						}
-
-						DrawRectangle(
-							RHICmdList,
-							ViewRect.Min.X, ViewRect.Min.Y,
-							ViewRect.Width(), ViewRect.Height(),
-							ViewRect.Min.X, ViewRect.Min.Y,
-							ViewRect.Width(), ViewRect.Height(),
-							FIntPoint(ViewRect.Width(), ViewRect.Height()),
-							FIntPoint(MipSize, MipSize),
-							VertexShader);
-					}
-					RHICmdList.EndRenderPass();
-
-					FResolveParams ResolveParams(FResolveRect(), (ECubeFace)CubeFace, MipIndex);
-					ResolveParams.SourceAccessFinal = ERHIAccess::SRVMask;
-					ResolveParams.DestAccessFinal = ERHIAccess::SRVMask;
-					RHICmdList.CopyToResolveTarget(EffectiveRT.TargetableTexture, EffectiveRT.ShaderResourceTexture, ResolveParams);
-				}
-			}
-		}
+		return FilterCubeMap(GraphBuilder, ShaderMap, CubemapTexture);
 	}
-}
+}
\ No newline at end of file
diff --git a/Engine/Source/Runtime/Renderer/Private/MobileReflectionEnvironmentCapture.h b/Engine/Source/Runtime/Renderer/Private/MobileReflectionEnvironmentCapture.h
--- a/Engine/Source/Runtime/Renderer/Private/MobileReflectionEnvironmentCapture.h
+++ b/Engine/Source/Runtime/Renderer/Private/MobileReflectionEnvironmentCapture.h
@@ -13,7 +13,7 @@
 
 namespace MobileReflectionEnvironmentCapture
 {
-	void ComputeAverageBrightness(FRHICommandListImmediate& RHICmdList, ERHIFeatureLevel::Type FeatureLevel, int32 CubmapSize, float& OutAverageBrightness);
-	void FilterReflectionEnvironment(FRHICommandListImmediate& RHICmdList, ERHIFeatureLevel::Type FeatureLevel, int32 CubmapSize, FSHVectorRGB3* OutIrradianceEnvironmentMap);
-	void CopyToSkyTexture(FRHICommandList& RHICmdList, FScene* Scene, FTexture* ProcessedTexture);
+	void ComputeAverageBrightness(FRDGBuilder& GraphBuilder, FGlobalShaderMap* ShaderMap, FRDGTexture* CubemapTexture, float* OutAverageBrightness);
+	FRDGTexture* FilterReflectionEnvironment(FRDGBuilder& GraphBuilder, FGlobalShaderMap* ShaderMap, FRDGTexture* CubemapTexture, FSHVectorRGB3* OutIrradianceEnvironmentMap);
+	void CopyToSkyTexture(FRDGBuilder& GraphBuilder, FScene* Scene, FRDGTexture* CubemapTexture, FTexture* ProcessedTexture);
 }
diff --git a/Engine/Source/Runtime/Renderer/Private/ReflectionEnvironmentCapture.cpp b/Engine/Source/Runtime/Renderer/Private/ReflectionEnvironmentCapture.cpp
--- a/Engine/Source/Runtime/Renderer/Private/ReflectionEnvironmentCapture.cpp
+++ b/Engine/Source/Runtime/Renderer/Private/ReflectionEnvironmentCapture.cpp
@@ -46,6 +46,7 @@
 #include "VolumetricCloudRendering.h"
 #include "VolumetricCloudProxy.h"
 #include "RenderGraphUtils.h"
+#include "PixelShaderUtils.h"
 
 /** Near plane to use when capturing the scene. */
 float GReflectionCaptureNearPlane = 5;
@@ -78,15 +79,6 @@ float ReflectionCaptureRoughestMip = 1;
  */
 float ReflectionCaptureRoughnessMipScale = 1.2f;
 
-int32 GDiffuseIrradianceCubemapSize = 32;
-
-static TAutoConsoleVariable<int32> CVarReflectionCaptureGPUArrayCopy(
-	TEXT("r.ReflectionCaptureGPUArrayCopy"),
-	1,
-	TEXT("Do a fast copy of the reflection capture array when resizing if possible. This avoids hitches on the rendering thread when the cubemap array needs to grow.\n")
-	TEXT(" 0 is off, 1 is on (default)"),
-	ECVF_ReadOnly);
-
 // Chaos addition
 static TAutoConsoleVariable<int32> CVarReflectionCaptureStaticSceneOnly(
 	TEXT("r.chaos.ReflectionCaptureStaticSceneOnly"),
@@ -95,12 +87,6 @@ static TAutoConsoleVariable<int32> CVarReflectionCaptureStaticSceneOnly(
 	TEXT(" 0 is off, 1 is on (default)"),
 	ECVF_ReadOnly);
 
-static int32 GFreeReflectionScratchAfterUse = 0;
-static FAutoConsoleVariableRef CVarFreeReflectionScratchAfterUse(
-	TEXT("r.FreeReflectionScratchAfterUse"),
-	GFreeReflectionScratchAfterUse,
-	TEXT("Free reflection scratch render targets after use."));
-
 /**
 * This CVar might affect the quality and performance for: 
 * (1) Captured reflection: Increase volume and light function quality and the cost at lighting build time.
@@ -116,547 +102,337 @@ static FAutoConsoleVariableRef CVarReflectionCaptureEnableLightFunctions(
 	TEXT("0. Disable light functions in reflection/sky light capture (default).\n")
 	TEXT("Others. Enable light functions."));
 
-TGlobalResource<FReflectionScratchCubemaps> GReflectionScratchCubemaps;
-
-bool DoGPUArrayCopy()
+BEGIN_SHADER_PARAMETER_STRUCT(FCubeShaderParameters, )
+	SHADER_PARAMETER_RDG_TEXTURE_SRV(TextureCube, SourceCubemapTexture)
+	SHADER_PARAMETER_SAMPLER(SamplerState, SourceCubemapSampler)
+	SHADER_PARAMETER(FVector2f, SvPositionToUVScale)
+	SHADER_PARAMETER(int32, CubeFace)
+	SHADER_PARAMETER(int32, MipIndex)
+	SHADER_PARAMETER(int32, NumMips)
+	RENDER_TARGET_BINDING_SLOTS()
+END_SHADER_PARAMETER_STRUCT()
+
+/** Pixel shader used for filtering a mip. */
+class FCubeDownsamplePS : public FGlobalShader
 {
-	return GRHISupportsResolveCubemapFaces && CVarReflectionCaptureGPUArrayCopy.GetValueOnAnyThread();
-}
-
-void FReflectionScratchCubemaps::Allocate(FRHICommandList& RHICmdList, uint32 TargetSize)
-{
-	if (GSupportsRenderTargetFormat_PF_FloatRGBA)
-	{
-		const int32 NumReflectionCaptureMips = FMath::CeilLogTwo(TargetSize) + 1;
-
-		if (Color[0] && Color[0]->GetRHI()->GetNumMips() != NumReflectionCaptureMips)
-		{
-			Color[0].SafeRelease();
-			Color[1].SafeRelease();
-		}
-
-		// Reflection targets are shared between both mobile and deferred shading paths. If we have already allocated for one and are now allocating for the other,
-		// we can skip these targets.
-		bool bSharedReflectionTargetsAllocated = Color[0] != nullptr;
-
-		if (!bSharedReflectionTargetsAllocated)
-		{
-			// We write to these cubemap faces individually during filtering
-			ETextureCreateFlags CubeTexFlags = TexCreate_TargetArraySlicesIndependently
-				| TexCreate_DisableDCC // todo: temporary disable to avoid DCC copy failure
-				;
-
-			{
-				// Create scratch cubemaps for filtering passes
-				FPooledRenderTargetDesc Desc2(FPooledRenderTargetDesc::CreateCubemapDesc(TargetSize, PF_FloatRGBA, FClearValueBinding(FLinearColor(0, 10000, 0, 0)), CubeTexFlags, TexCreate_RenderTargetable, false, 1, NumReflectionCaptureMips));
-				GRenderTargetPool.FindFreeElement(RHICmdList, Desc2, Color[0], TEXT("ReflectionColorScratchCubemap0"));
-				GRenderTargetPool.FindFreeElement(RHICmdList, Desc2, Color[1], TEXT("ReflectionColorScratchCubemap1"));
-			}
-
-			extern int32 GDiffuseIrradianceCubemapSize;
-			const int32 NumDiffuseIrradianceMips = FMath::CeilLogTwo(GDiffuseIrradianceCubemapSize) + 1;
-
-			{
-				FPooledRenderTargetDesc Desc2(FPooledRenderTargetDesc::CreateCubemapDesc(GDiffuseIrradianceCubemapSize, PF_FloatRGBA, FClearValueBinding(FLinearColor(0, 10000, 0, 0)), CubeTexFlags, TexCreate_RenderTargetable, false, 1, NumDiffuseIrradianceMips));
-				GRenderTargetPool.FindFreeElement(RHICmdList, Desc2, Irradiance[0], TEXT("DiffuseIrradianceScratchCubemap0"));
-				GRenderTargetPool.FindFreeElement(RHICmdList, Desc2, Irradiance[1], TEXT("DiffuseIrradianceScratchCubemap1"));
-			}
-
-			{
-				FPooledRenderTargetDesc Desc(FPooledRenderTargetDesc::Create2DDesc(FIntPoint(FSHVector3::MaxSHBasis, 1), PF_FloatRGBA, FClearValueBinding(FLinearColor(0, 10000, 0, 0)), TexCreate_None, TexCreate_RenderTargetable, false));
-				GRenderTargetPool.FindFreeElement(RHICmdList, Desc, SkySHIrradiance, TEXT("SkySHIrradiance"));
-			}
-		}
-	}
-}
-
-void FReflectionScratchCubemaps::Release()
-{
-	for (int32 Index = 0; Index < UE_ARRAY_COUNT(Color); Index++)
-	{
-		Color[Index].SafeRelease();
-	}
-
-	for (int32 Index = 0; Index < UE_ARRAY_COUNT(Irradiance); Index++)
-	{
-		Irradiance[Index].SafeRelease();
-	}
+	DECLARE_GLOBAL_SHADER(FCubeDownsamplePS);
+	SHADER_USE_PARAMETER_STRUCT(FCubeDownsamplePS, FGlobalShader);
+	using FParameters = FCubeShaderParameters;
+};
 
-	SkySHIrradiance.SafeRelease();
-}
+IMPLEMENT_GLOBAL_SHADER(FCubeDownsamplePS, "/Engine/Private/ReflectionEnvironmentShaders.usf", "DownsamplePS", SF_Pixel);
 
-void FullyResolveReflectionScratchCubes(FRHICommandListImmediate& RHICmdList)
+class FCubeFilterPS : public FGlobalShader
 {
-	SCOPED_DRAW_EVENT(RHICmdList, FullyResolveReflectionScratchCubes);
-	FRHITexture* Scratch0 = GReflectionScratchCubemaps.Color[0]->GetRHI();
-	FRHITexture* Scratch1 = GReflectionScratchCubemaps.Color[1]->GetRHI();
-	FResolveParams ResolveParams(FResolveRect(), CubeFace_PosX, -1, -1, -1);
-	RHICmdList.CopyToResolveTarget(Scratch0, Scratch0, ResolveParams);
-	RHICmdList.CopyToResolveTarget(Scratch1, Scratch1, ResolveParams);  
-}
-
-IMPLEMENT_SHADER_TYPE(,FCubeFilterPS,TEXT("/Engine/Private/ReflectionEnvironmentShaders.usf"),TEXT("DownsamplePS"),SF_Pixel);
+	DECLARE_GLOBAL_SHADER(FCubeFilterPS);
+	SHADER_USE_PARAMETER_STRUCT(FCubeFilterPS, FGlobalShader);
+	using FParameters = FCubeShaderParameters;
+};
 
-IMPLEMENT_SHADER_TYPE(template<>,TCubeFilterPS<0>,TEXT("/Engine/Private/ReflectionEnvironmentShaders.usf"),TEXT("FilterPS"),SF_Pixel);
-IMPLEMENT_SHADER_TYPE(template<>,TCubeFilterPS<1>,TEXT("/Engine/Private/ReflectionEnvironmentShaders.usf"),TEXT("FilterPS"),SF_Pixel);
+IMPLEMENT_GLOBAL_SHADER(FCubeFilterPS, "/Engine/Private/ReflectionEnvironmentShaders.usf", "FilterPS", SF_Pixel);
 
 /** Computes the average brightness of a 1x1 mip of a cubemap. */
 class FComputeBrightnessPS : public FGlobalShader
 {
-	DECLARE_SHADER_TYPE(FComputeBrightnessPS,Global)
-public:
+	DECLARE_GLOBAL_SHADER(FComputeBrightnessPS);
+	SHADER_USE_PARAMETER_STRUCT(FComputeBrightnessPS, FGlobalShader);
 
-	static bool ShouldCompilePermutation(const FGlobalShaderPermutationParameters& Parameters)
-	{
-		return true;
-	}
+	BEGIN_SHADER_PARAMETER_STRUCT(FParameters, )
+		SHADER_PARAMETER_RDG_TEXTURE(TextureCube, ReflectionEnvironmentColorTexture)
+		SHADER_PARAMETER_SAMPLER(SamplerState, ReflectionEnvironmentColorSampler)
+		SHADER_PARAMETER(int32, NumCaptureArrayMips)
+		RENDER_TARGET_BINDING_SLOTS()
+	END_SHADER_PARAMETER_STRUCT()
 
 	static void ModifyCompilationEnvironment(const FGlobalShaderPermutationParameters& Parameters, FShaderCompilerEnvironment& OutEnvironment)
 	{
 		FGlobalShader::ModifyCompilationEnvironment(Parameters, OutEnvironment);
 		OutEnvironment.SetDefine(TEXT("COMPUTEBRIGHTNESS_PIXELSHADER"), 1);
 	}
+};
 
-	FComputeBrightnessPS(const ShaderMetaType::CompiledShaderInitializerType& Initializer)
+IMPLEMENT_GLOBAL_SHADER(FComputeBrightnessPS, "/Engine/Private/ReflectionEnvironmentShaders.usf", "ComputeBrightnessMain", SF_Pixel);
+
+/** Vertex shader used when writing to a cubemap. */
+class FCopyToCubeFaceVS : public FGlobalShader
+{
+public:
+	DECLARE_GLOBAL_SHADER(FCopyToCubeFaceVS);
+
+	FCopyToCubeFaceVS() = default;
+	FCopyToCubeFaceVS(const CompiledShaderInitializerType& Initializer)
 		: FGlobalShader(Initializer)
-	{
-		ReflectionEnvironmentColorTexture.Bind(Initializer.ParameterMap,TEXT("ReflectionEnvironmentColorTexture"));
-		ReflectionEnvironmentColorSampler.Bind(Initializer.ParameterMap,TEXT("ReflectionEnvironmentColorSampler"));
-		NumCaptureArrayMips.Bind(Initializer.ParameterMap, TEXT("NumCaptureArrayMips"));
-	}
+	{}
+};
 
-	FComputeBrightnessPS()
-	{
-	}
+IMPLEMENT_GLOBAL_SHADER(FCopyToCubeFaceVS, "/Engine/Private/ReflectionEnvironmentShaders.usf", "CopyToCubeFaceVS", SF_Vertex);
 
-	void SetParameters(FRHICommandList& RHICmdList, int32 TargetSize, FSceneRenderTargetItem& Cubemap)
-	{
-		const int32 EffectiveTopMipSize = TargetSize;
-		const int32 NumMips = FMath::CeilLogTwo(EffectiveTopMipSize) + 1;
-		// Read from the smallest mip that was downsampled to
+/** Pixel shader used when copying scene color from a scene render into a face of a reflection capture cubemap. */
+class FCopySceneColorToCubeFacePS : public FGlobalShader
+{
+public:
+	DECLARE_GLOBAL_SHADER(FCopySceneColorToCubeFacePS);
+	SHADER_USE_PARAMETER_STRUCT(FCopySceneColorToCubeFacePS, FGlobalShader);
 
-		if (Cubemap.IsValid())
+	static void ModifyCompilationEnvironment(const FGlobalShaderPermutationParameters& Parameters, FShaderCompilerEnvironment& OutEnvironment)
+	{
+		if (IsMobilePlatform(Parameters.Platform))
 		{
-			SetTextureParameter(
-				RHICmdList,
-				RHICmdList.GetBoundPixelShader(),
-				ReflectionEnvironmentColorTexture, 
-				ReflectionEnvironmentColorSampler, 
-				TStaticSamplerState<SF_Trilinear, AM_Clamp, AM_Clamp, AM_Clamp>::GetRHI(), 
-				Cubemap.ShaderResourceTexture);
+			// SceneDepth is memoryless on mobile
+			OutEnvironment.SetDefine(TEXT("SCENE_TEXTURES_DISABLED"), 1u);
 		}
-
-		SetShaderValue(RHICmdList, RHICmdList.GetBoundPixelShader(), NumCaptureArrayMips, FMath::CeilLogTwo(TargetSize) + 1);
 	}
 
-private:
-
-	LAYOUT_FIELD(FShaderResourceParameter, ReflectionEnvironmentColorTexture);
-	LAYOUT_FIELD(FShaderResourceParameter, ReflectionEnvironmentColorSampler);
-	LAYOUT_FIELD(FShaderParameter, NumCaptureArrayMips);
+	BEGIN_SHADER_PARAMETER_STRUCT(FParameters, )
+		SHADER_PARAMETER_STRUCT_REF(FViewUniformShaderParameters, View)
+		SHADER_PARAMETER_RDG_TEXTURE(Texture2D, SceneColorTexture)
+		SHADER_PARAMETER_SAMPLER(SamplerState, SceneColorSampler)
+		SHADER_PARAMETER_RDG_TEXTURE(Texture2D, SceneDepthTexture)
+		SHADER_PARAMETER_SAMPLER(SamplerState, SceneDepthSampler)
+		SHADER_PARAMETER(FVector4f, SkyLightCaptureParameters)
+		SHADER_PARAMETER(FVector4f, LowerHemisphereColor)
+		RENDER_TARGET_BINDING_SLOTS()
+	END_SHADER_PARAMETER_STRUCT()
 };
 
-IMPLEMENT_SHADER_TYPE(,FComputeBrightnessPS,TEXT("/Engine/Private/ReflectionEnvironmentShaders.usf"),TEXT("ComputeBrightnessMain"),SF_Pixel);
+IMPLEMENT_GLOBAL_SHADER(FCopySceneColorToCubeFacePS, "/Engine/Private/ReflectionEnvironmentShaders.usf", "CopySceneColorToCubeFaceColorPS", SF_Pixel);
 
-void CreateCubeMips( FRHICommandListImmediate& RHICmdList, ERHIFeatureLevel::Type FeatureLevel, int32 NumMips, FSceneRenderTargetItem& Cubemap )
-{	
-	SCOPED_DRAW_EVENT(RHICmdList, CreateCubeMips);
+/** Pixel shader used when copying a cubemap into a face of a reflection capture cubemap. */
+class FCopyCubemapToCubeFacePS : public FGlobalShader
+{
+public:
+	DECLARE_GLOBAL_SHADER(FCopyCubemapToCubeFacePS);
+	SHADER_USE_PARAMETER_STRUCT(FCopyCubemapToCubeFacePS, FGlobalShader);
 
-	FRHITexture* CubeRef = Cubemap.TargetableTexture.GetReference();
+	BEGIN_SHADER_PARAMETER_STRUCT(FParameters, )
+		SHADER_PARAMETER_TEXTURE(TextureCube, SourceCubemapTexture)
+		SHADER_PARAMETER_SAMPLER(SamplerState, SourceCubemapSampler)
+		SHADER_PARAMETER(FVector4f, SkyLightCaptureParameters)
+		SHADER_PARAMETER(FVector4f, LowerHemisphereColor)
+		SHADER_PARAMETER(FVector2f, SinCosSourceCubemapRotation)
+		SHADER_PARAMETER(FVector2f, SvPositionToUVScale)
+		SHADER_PARAMETER(int32, CubeFace)
+		RENDER_TARGET_BINDING_SLOTS()
+	END_SHADER_PARAMETER_STRUCT()
+};
 
-	auto* ShaderMap = GetGlobalShaderMap(FeatureLevel);
+IMPLEMENT_GLOBAL_SHADER(FCopyCubemapToCubeFacePS, "/Engine/Private/ReflectionEnvironmentShaders.usf", "CopyCubemapToCubeFaceColorPS", SF_Pixel);
 
-	TArray<TPair<FRHITextureSRVCreateInfo, TRefCountPtr<FRHIShaderResourceView>>> SRVs;
-	SRVs.Empty(NumMips);
+bool DoGPUArrayCopy()
+{
+	return GRHISupportsResolveCubemapFaces;
+}
 
-	for (int32 MipIndex = 0; MipIndex < NumMips; MipIndex++)
+class FReflectionCubemapTexture : public FRenderThreadStructBase
+{
+public:
+	FReflectionCubemapTexture(uint32 InCubemapSize)
+		: CubemapSize(InCubemapSize)
 	{
-		FRHITextureSRVCreateInfo SRVDesc;
-		SRVDesc.MipLevel = MipIndex;
-		SRVs.Emplace(SRVDesc, RHICreateShaderResourceView(Cubemap.ShaderResourceTexture, SRVDesc));
+		check(GSupportsRenderTargetFormat_PF_FloatRGBA);
 	}
 
-	FGraphicsPipelineStateInitializer GraphicsPSOInit;
-	GraphicsPSOInit.RasterizerState = TStaticRasterizerState<FM_Solid, CM_None>::GetRHI();
-	GraphicsPSOInit.DepthStencilState = TStaticDepthStencilState<false, CF_Always>::GetRHI();
-	GraphicsPSOInit.BlendState = TStaticBlendState<>::GetRHI();
-
-	// Downsample all the mips, each one reads from the mip above it
-	for (int32 MipIndex = 1; MipIndex < NumMips; MipIndex++)
+	void InitRHI(FRHICommandListImmediate& RHICmdList)
 	{
-		// For the first iteration, we don't know what the previous state
-		// of the source mip was, but we *do* for all the other iterations...
-		ERHIAccess Previous = MipIndex == 1
-			? ERHIAccess::Unknown
-			: ERHIAccess::RTV;
-
-		FRHITransitionInfo Transitions[] =
-		{
-			// Make the source mip readable (SRVGraphics)
-			FRHITransitionInfo(CubeRef, Previous, ERHIAccess::SRVGraphics, EResourceTransitionFlags::None, uint32(MipIndex - 1)),
+		const int32 NumReflectionCaptureMips = GetNumMips(CubemapSize);
 
-			// Make the destination mip writable (RTV)
-			FRHITransitionInfo(CubeRef, ERHIAccess::Unknown, ERHIAccess::RTV, EResourceTransitionFlags::None, uint32(MipIndex))
-		};
-		RHICmdList.Transition(Transitions);
+		const FRHITextureDesc TextureDesc = FRHITextureDesc::CreateCube(
+			CubemapSize, PF_FloatRGBA, FClearValueBinding(FLinearColor(0, 10000, 0, 0)), TexCreate_RenderTargetable | TexCreate_TargetArraySlicesIndependently | TexCreate_DisableDCC, NumReflectionCaptureMips);
 
-		const int32 MipSize = 1 << (NumMips - MipIndex - 1);
-		SCOPED_DRAW_EVENT(RHICmdList, CreateCubeMipsPerFace);
-		for (int32 CubeFace = 0; CubeFace < CubeFace_MAX; CubeFace++)
-		{
-			FRHIRenderPassInfo RPInfo(Cubemap.TargetableTexture, ERenderTargetActions::DontLoad_Store, nullptr, MipIndex, CubeFace);
-			RHICmdList.BeginRenderPass(RPInfo, TEXT("CreateCubeMips"));
-			RHICmdList.ApplyCachedRenderTargets(GraphicsPSOInit);
-
-			const FIntRect ViewRect(0, 0, MipSize, MipSize);
-			RHICmdList.SetViewport(0.0f, 0.0f, 0.0f, (float)MipSize, (float)MipSize, 1.0f);
+		RenderTarget = AllocatePooledTexture(TextureDesc, TEXT("ReflectionCubeTexture"));
 
+		FRDGBuilder GraphBuilder(RHICmdList);
 
-			TShaderMapRef<FScreenVS> VertexShader(ShaderMap);
-			TShaderMapRef<FCubeFilterPS> PixelShader(ShaderMap);
-
-			GraphicsPSOInit.BoundShaderState.VertexDeclarationRHI = GFilterVertexDeclaration.VertexDeclarationRHI;
-			GraphicsPSOInit.BoundShaderState.VertexShaderRHI = VertexShader.GetVertexShader();
-			GraphicsPSOInit.BoundShaderState.PixelShaderRHI = PixelShader.GetPixelShader();
-			GraphicsPSOInit.PrimitiveType = PT_TriangleList;
+		{
+			RDG_EVENT_SCOPE(GraphBuilder, "ClearReflectionCubemap");
 
-			SetGraphicsPipelineState(RHICmdList, GraphicsPSOInit, 0);
+			FRDGTextureRef OutputTexture = GetRDG(GraphBuilder);
 
+			for (int32 MipIndex = 0; MipIndex < NumReflectionCaptureMips; MipIndex++)
 			{
-				FRHIPixelShader* ShaderRHI = PixelShader.GetPixelShader();
-
-				SetShaderValue(RHICmdList, ShaderRHI, PixelShader->CubeFace, CubeFace);
-				SetShaderValue(RHICmdList, ShaderRHI, PixelShader->MipIndex, MipIndex);
-
-				SetShaderValue(RHICmdList, ShaderRHI, PixelShader->NumMips, NumMips);
-
-				check(SRVs.IsValidIndex(MipIndex - 1) && SRVs[MipIndex - 1].Key.MipLevel == MipIndex - 1);
-				SetSRVParameter(RHICmdList, ShaderRHI, PixelShader->SourceCubemapTexture, SRVs[MipIndex - 1].Value);
-				SetSamplerParameter(RHICmdList, ShaderRHI, PixelShader->SourceCubemapSampler, TStaticSamplerState<SF_Bilinear, AM_Clamp, AM_Clamp, AM_Clamp>::GetRHI());
+				for (int32 CubeFace = 0; CubeFace < CubeFace_MAX; CubeFace++)
+				{
+					FRenderTargetParameters* PassParameters = GraphBuilder.AllocParameters<FRenderTargetParameters>();
+					PassParameters->RenderTargets[0] = FRenderTargetBinding(OutputTexture, ERenderTargetLoadAction::EClear, MipIndex, CubeFace);
+
+					GraphBuilder.AddPass(
+						RDG_EVENT_NAME("Clear (Mip: %d, Face : %d)", MipIndex, CubeFace),
+						PassParameters,
+						ERDGPassFlags::Raster,
+						[](FRHICommandList&) {});
+				}
 			}
-
-			DrawRectangle(
-				RHICmdList,
-				ViewRect.Min.X, ViewRect.Min.Y,
-				ViewRect.Width(), ViewRect.Height(),
-				ViewRect.Min.X, ViewRect.Min.Y,
-				ViewRect.Width(), ViewRect.Height(),
-				FIntPoint(ViewRect.Width(), ViewRect.Height()),
-				FIntPoint(MipSize, MipSize),
-				VertexShader);
-
-			RHICmdList.EndRenderPass();
 		}
+
+		GraphBuilder.Execute();
 	}
 
-	RHICmdList.Transition(FRHITransitionInfo(CubeRef, ERHIAccess::Unknown, ERHIAccess::SRVMask));
+	FRDGTexture* GetRDG(FRDGBuilder& GraphBuilder) const
+	{
+		return GraphBuilder.RegisterExternalTexture(RenderTarget);
+	}
 
-	SRVs.Empty();
-}
+private:
+	TRefCountPtr<IPooledRenderTarget> RenderTarget;
+	uint32 CubemapSize;
+};
 
-/** Computes the average brightness of the given reflection capture and stores it in the scene. */
-float ComputeSingleAverageBrightnessFromCubemap(FRHICommandListImmediate& RHICmdList, ERHIFeatureLevel::Type FeatureLevel, int32 TargetSize, FSceneRenderTargetItem& Cubemap)
+void CreateCubeMips(FRDGBuilder& GraphBuilder, FGlobalShaderMap* ShaderMap, FRDGTexture* CubemapTexture)
 {
-	SCOPED_DRAW_EVENT(RHICmdList, ComputeSingleAverageBrightnessFromCubemap);
+	RDG_EVENT_SCOPE(GraphBuilder, "CreateCubeMips");
 
-	TRefCountPtr<IPooledRenderTarget> ReflectionBrightnessTarget;
-	FPooledRenderTargetDesc Desc(FPooledRenderTargetDesc::Create2DDesc(FIntPoint(1, 1), PF_FloatRGBA, FClearValueBinding::None, TexCreate_None, TexCreate_RenderTargetable, false));
-	GRenderTargetPool.FindFreeElement(RHICmdList, Desc, ReflectionBrightnessTarget, TEXT("ReflectionBrightness"));
+	TShaderMapRef<FCubeDownsamplePS> PixelShader(ShaderMap);
 
-	FRHITexture* BrightnessTarget = ReflectionBrightnessTarget->GetRHI();
+	const int32 NumMips = CubemapTexture->Desc.NumMips;
 
-	FRHIRenderPassInfo RPInfo(BrightnessTarget, ERenderTargetActions::Load_Store);
-	TransitionRenderPassTargets(RHICmdList, RPInfo);
-	RHICmdList.BeginRenderPass(RPInfo, TEXT("ReflectionBrightness"));
+	// Downsample all the mips, each one reads from the mip above it
+	for (int32 MipIndex = 1; MipIndex < NumMips; MipIndex++)
 	{
-		FGraphicsPipelineStateInitializer GraphicsPSOInit;
-		RHICmdList.ApplyCachedRenderTargets(GraphicsPSOInit);
-		GraphicsPSOInit.RasterizerState = TStaticRasterizerState<FM_Solid, CM_None>::GetRHI();
-		GraphicsPSOInit.DepthStencilState = TStaticDepthStencilState<false, CF_Always>::GetRHI();
-		GraphicsPSOInit.BlendState = TStaticBlendState<>::GetRHI();
-
-		auto ShaderMap = GetGlobalShaderMap(FeatureLevel);
-		TShaderMapRef<FPostProcessVS> VertexShader(ShaderMap);
-		TShaderMapRef<FComputeBrightnessPS> PixelShader(ShaderMap);
-
-		GraphicsPSOInit.BoundShaderState.VertexDeclarationRHI = GFilterVertexDeclaration.VertexDeclarationRHI;
-		GraphicsPSOInit.BoundShaderState.VertexShaderRHI = VertexShader.GetVertexShader();
-		GraphicsPSOInit.BoundShaderState.PixelShaderRHI = PixelShader.GetPixelShader();
-		GraphicsPSOInit.PrimitiveType = PT_TriangleList;
-	
-		SetGraphicsPipelineState(RHICmdList, GraphicsPSOInit, 0);
-
-		PixelShader->SetParameters(RHICmdList, TargetSize, Cubemap);
-
-		DrawRectangle(
-			RHICmdList,
-			0, 0,
-			1, 1,
-			0, 0,
-			1, 1,
-			FIntPoint(1, 1),
-			FIntPoint(1, 1),
-			VertexShader);
+		const int32 MipSize = 1 << (NumMips - MipIndex - 1);
+		const FIntRect ViewRect(0, 0, MipSize, MipSize);
+
+		for (int32 CubeFace = 0; CubeFace < CubeFace_MAX; CubeFace++)
+		{
+			auto* PassParameters = GraphBuilder.AllocParameters<FCubeDownsamplePS::FParameters>();
+			PassParameters->SourceCubemapTexture = GraphBuilder.CreateSRV(FRDGTextureSRVDesc::CreateForMipLevel(CubemapTexture, MipIndex - 1));
+			PassParameters->SourceCubemapSampler = TStaticSamplerState<SF_Bilinear, AM_Clamp, AM_Clamp, AM_Clamp>::GetRHI();
+			PassParameters->CubeFace = CubeFace;
+			PassParameters->MipIndex = MipIndex;
+			PassParameters->NumMips = NumMips;
+			PassParameters->SvPositionToUVScale = FVector2f(1.0f / MipSize, 1.0f / MipSize);
+			PassParameters->RenderTargets[0] = FRenderTargetBinding(CubemapTexture, ERenderTargetLoadAction::ENoAction, MipIndex, CubeFace);
+
+			FPixelShaderUtils::AddFullscreenPass(
+				GraphBuilder,
+				ShaderMap,
+				RDG_EVENT_NAME("CreateCubeMips (Mip: %d, Face: %d)", MipIndex, CubeFace),
+				PixelShader,
+				PassParameters,
+				ViewRect);
+		}
 	}
-	RHICmdList.EndRenderPass();
-	RHICmdList.CopyToResolveTarget(BrightnessTarget, BrightnessTarget, FResolveParams());
+}
 
-	FSceneRenderTargetItem& EffectiveRT = ReflectionBrightnessTarget->GetRenderTargetItem();
-	check(EffectiveRT.ShaderResourceTexture->GetFormat() == PF_FloatRGBA);
+void ComputeSingleAverageBrightnessFromCubemap(FRDGBuilder& GraphBuilder, FGlobalShaderMap* ShaderMap, FRDGTexture* CubemapTexture, float* OutAverageBrightness)
+{
+	RDG_EVENT_SCOPE(GraphBuilder, "ComputeSingleAverageBrightnessFromCubemap");
 
-	TArray<FFloat16Color> SurfaceData;
-	RHICmdList.ReadSurfaceFloatData(EffectiveRT.ShaderResourceTexture, FIntRect(0, 0, 1, 1), SurfaceData, CubeFace_PosX, 0, 0);
+	FRDGTexture* ReflectionBrightnessTexture = GraphBuilder.CreateTexture(FRDGTextureDesc::Create2D(FIntPoint(1, 1), PF_FloatRGBA, FClearValueBinding::None, TexCreate_RenderTargetable), TEXT("ReflectionBrightness"));
 
-	// Shader outputs luminance to R
-	float AverageBrightness = SurfaceData[0].R.GetFloat();
-	return AverageBrightness;
-}
+	auto* PassParameters = GraphBuilder.AllocParameters<FComputeBrightnessPS::FParameters>();
 
-void ComputeAverageBrightness(FRHICommandListImmediate& RHICmdList, ERHIFeatureLevel::Type FeatureLevel, int32 CubmapSize, float& OutAverageBrightness)
-{
-	SCOPED_DRAW_EVENT(RHICmdList, ComputeAverageBrightness);
+	PassParameters->ReflectionEnvironmentColorTexture = CubemapTexture;
+	PassParameters->ReflectionEnvironmentColorSampler = TStaticSamplerState<SF_Trilinear, AM_Clamp, AM_Clamp, AM_Clamp>::GetRHI();
+	PassParameters->NumCaptureArrayMips = CubemapTexture->Desc.NumMips;
+	PassParameters->RenderTargets[0] = FRenderTargetBinding(ReflectionBrightnessTexture, ERenderTargetLoadAction::ENoAction);
 
-	const int32 EffectiveTopMipSize = CubmapSize;
-	const int32 NumMips = FMath::CeilLogTwo(EffectiveTopMipSize) + 1;
+	TShaderMapRef<FComputeBrightnessPS> PixelShader(ShaderMap);
 
-	// necessary to resolve the clears which touched all the mips.  scene rendering only resolves mip 0.
-	FullyResolveReflectionScratchCubes(RHICmdList);	
+	FPixelShaderUtils::AddFullscreenPass(
+		GraphBuilder,
+		ShaderMap,
+		RDG_EVENT_NAME("ReflectionBrightness"),
+		PixelShader,
+		PassParameters,
+		FIntRect(FIntPoint::ZeroValue, FIntPoint(1, 1)));
 
-	FSceneRenderTargetItem& DownSampledCube = GReflectionScratchCubemaps.Color[0]->GetRenderTargetItem();
-	CreateCubeMips( RHICmdList, FeatureLevel, NumMips, DownSampledCube );
+	AddReadbackTexturePass(GraphBuilder, RDG_EVENT_NAME("ReadbackTexture"), ReflectionBrightnessTexture, [ReflectionBrightnessTexture, OutAverageBrightness](FRHICommandListImmediate& RHICmdList)
+	{
+		TArray<FFloat16Color> SurfaceData;
+		RHICmdList.ReadSurfaceFloatData(ReflectionBrightnessTexture->GetRHI(), FIntRect(0, 0, 1, 1), SurfaceData, CubeFace_PosX, 0, 0);
 
-	OutAverageBrightness = ComputeSingleAverageBrightnessFromCubemap(RHICmdList, FeatureLevel, CubmapSize, DownSampledCube);
+		// Shader outputs luminance to R
+		*OutAverageBrightness = SurfaceData[0].R.GetFloat();
+	});
 }
 
+void ComputeAverageBrightness(FRDGBuilder& GraphBuilder, FGlobalShaderMap* ShaderMap, FRDGTexture* CubemapTexture, float* OutAverageBrightness)
+{
+	CreateCubeMips(GraphBuilder, ShaderMap, CubemapTexture);
+	ComputeSingleAverageBrightnessFromCubemap(GraphBuilder, ShaderMap, CubemapTexture, OutAverageBrightness);
+}
 
-void FilterCubeMap(FRHICommandListImmediate& RHICmdList, ERHIFeatureLevel::Type FeatureLevel, int32 NumMips, 
-	FSceneRenderTargetItem& DownSampledCube, FSceneRenderTargetItem& FilteredCube)
+FRDGTexture* FilterCubeMap(FRDGBuilder& GraphBuilder, FGlobalShaderMap* ShaderMap, FRDGTexture* SourceTexture)
 {
-	SCOPED_DRAW_EVENT(RHICmdList, FilterCubeMap);
-	auto ShaderMap = GetGlobalShaderMap(FeatureLevel);
+	RDG_EVENT_SCOPE(GraphBuilder, "FilterCubeMap");
 
-	FGraphicsPipelineStateInitializer GraphicsPSOInit;
-	GraphicsPSOInit.RasterizerState = TStaticRasterizerState<FM_Solid, CM_None>::GetRHI();
-	GraphicsPSOInit.DepthStencilState = TStaticDepthStencilState<false, CF_Always>::GetRHI();
-	GraphicsPSOInit.BlendState = TStaticBlendState<>::GetRHI();
+	FRDGTexture* FilteredCubemapTexture = GraphBuilder.CreateTexture(SourceTexture->Desc, TEXT("FilteredCubemapTexture"));
 
-	RHICmdList.Transition(FRHITransitionInfo(FilteredCube.TargetableTexture, ERHIAccess::Unknown, ERHIAccess::RTV));
+	const int32 NumMips = SourceTexture->Desc.NumMips;
+
+	TShaderMapRef<FCubeFilterPS> PixelShader(ShaderMap);
 
-	// Filter all the mips
 	for (int32 MipIndex = 0; MipIndex < NumMips; MipIndex++)
 	{
 		const int32 MipSize = 1 << (NumMips - MipIndex - 1);
+		const FIntRect ViewRect(0, 0, MipSize, MipSize);
 
 		for (int32 CubeFace = 0; CubeFace < CubeFace_MAX; CubeFace++)
 		{
-			FRHIRenderPassInfo RPInfo(FilteredCube.TargetableTexture, ERenderTargetActions::DontLoad_Store, nullptr, MipIndex, CubeFace);
-			RHICmdList.Transition(FRHITransitionInfo(FilteredCube.TargetableTexture, ERHIAccess::Unknown, ERHIAccess::RTV));
-			RHICmdList.BeginRenderPass(RPInfo, TEXT("FilterMips"));
-
-			RHICmdList.ApplyCachedRenderTargets(GraphicsPSOInit);
-
-			const FIntRect ViewRect(0, 0, MipSize, MipSize);
-			RHICmdList.SetViewport(0, 0, 0.0f, MipSize, MipSize, 1.0f);
-
-			//TShaderMapRef<TCubeFilterPS<1>> CaptureCubemapArrayPixelShader(GetGlobalShaderMap(FeatureLevel));
-
-			TShaderMapRef<FScreenVS> VertexShader(GetGlobalShaderMap(FeatureLevel));
-			TShaderMapRef<TCubeFilterPS<0>> PixelShader(GetGlobalShaderMap(FeatureLevel));
-			GraphicsPSOInit.BoundShaderState.VertexShaderRHI = VertexShader.GetVertexShader();
-			GraphicsPSOInit.BoundShaderState.PixelShaderRHI = PixelShader.GetPixelShader();
-			GraphicsPSOInit.BoundShaderState.VertexDeclarationRHI = GFilterVertexDeclaration.VertexDeclarationRHI;
-			GraphicsPSOInit.PrimitiveType = PT_TriangleList;
-
-			SetGraphicsPipelineState(RHICmdList, GraphicsPSOInit, 0);
-
-			{
-				FRHIPixelShader* ShaderRHI = PixelShader.GetPixelShader();
-
-				SetShaderValue(RHICmdList, ShaderRHI, PixelShader->CubeFace, CubeFace);
-				SetShaderValue(RHICmdList, ShaderRHI, PixelShader->MipIndex, MipIndex);
-
-				SetShaderValue(RHICmdList, ShaderRHI, PixelShader->NumMips, NumMips);
-
-				SetTextureParameter(
-					RHICmdList,
-					ShaderRHI,
-					PixelShader->SourceCubemapTexture,
-					PixelShader->SourceCubemapSampler,
-					TStaticSamplerState<SF_Trilinear, AM_Clamp, AM_Clamp, AM_Clamp>::GetRHI(),
-					DownSampledCube.ShaderResourceTexture);
-			}
-
-			DrawRectangle(
-				RHICmdList,
-				ViewRect.Min.X, ViewRect.Min.Y,
-				ViewRect.Width(), ViewRect.Height(),
-				ViewRect.Min.X, ViewRect.Min.Y,
-				ViewRect.Width(), ViewRect.Height(),
-				FIntPoint(ViewRect.Width(), ViewRect.Height()),
-				FIntPoint(MipSize, MipSize),
-				VertexShader);
-
-			RHICmdList.EndRenderPass();
+			auto* PassParameters = GraphBuilder.AllocParameters<FCubeFilterPS::FParameters>();
+			PassParameters->SourceCubemapTexture = GraphBuilder.CreateSRV(FRDGTextureSRVDesc::Create(SourceTexture));
+			PassParameters->SourceCubemapSampler = TStaticSamplerState<SF_Trilinear, AM_Clamp, AM_Clamp, AM_Clamp>::GetRHI();
+			PassParameters->CubeFace = CubeFace;
+			PassParameters->MipIndex = MipIndex;
+			PassParameters->NumMips = NumMips;
+			PassParameters->SvPositionToUVScale = FVector2f(1.0f / MipSize, 1.0f / MipSize);
+			PassParameters->RenderTargets[0] = FRenderTargetBinding(FilteredCubemapTexture, ERenderTargetLoadAction::ENoAction, MipIndex, CubeFace);
+
+			FPixelShaderUtils::AddFullscreenPass(
+				GraphBuilder,
+				ShaderMap,
+				RDG_EVENT_NAME("FilterCubeMap (Mip: %d, CubeFace: %d)", MipIndex, CubeFace),
+				PixelShader,
+				PassParameters,
+				ViewRect);
 		}
 	}
-}
 
-/** Generates mips for glossiness and filters the cubemap for a given reflection. */
-void FilterReflectionEnvironment(FRHICommandListImmediate& RHICmdList, ERHIFeatureLevel::Type FeatureLevel, int32 CubmapSize, FSHVectorRGB3* OutIrradianceEnvironmentMap)
-{
-	SCOPED_DRAW_EVENT(RHICmdList, FilterReflectionEnvironment);
-
-	const int32 EffectiveTopMipSize = CubmapSize;
-	const int32 NumMips = FMath::CeilLogTwo(EffectiveTopMipSize) + 1;
+	return FilteredCubemapTexture;
+}
 
-	FSceneRenderTargetItem& EffectiveColorRT = GReflectionScratchCubemaps.Color[0]->GetRenderTargetItem();
+// Premultiply alpha in-place using alpha blending
 
-	FGraphicsPipelineStateInitializer GraphicsPSOInit;
-	GraphicsPSOInit.RasterizerState = TStaticRasterizerState<FM_Solid, CM_None>::GetRHI();
-	GraphicsPSOInit.DepthStencilState = TStaticDepthStencilState<false, CF_Always>::GetRHI();
-	GraphicsPSOInit.BlendState = TStaticBlendState<CW_RGBA, BO_Add, BF_Zero, BF_DestAlpha, BO_Add, BF_Zero, BF_One>::GetRHI();
-		
-	RHICmdList.Transition(FRHITransitionInfo(EffectiveColorRT.TargetableTexture, ERHIAccess::Unknown, ERHIAccess::RTV));
+void PremultiplyCubeMipAlpha(FRDGBuilder& GraphBuilder, FGlobalShaderMap* ShaderMap, FRDGTexture* CubemapTexture, int32 MipIndex)
+{
+	const int32 NumMips = CubemapTexture->Desc.NumMips;
+	const int32 MipSize = 1 << (NumMips - MipIndex - 1);
+	const FIntRect ViewRect(0, 0, MipSize, MipSize);
 
-	// Premultiply alpha in-place using alpha blending
 	for (uint32 CubeFace = 0; CubeFace < CubeFace_MAX; CubeFace++)
 	{
-		FRHIRenderPassInfo RPInfo(EffectiveColorRT.TargetableTexture, ERenderTargetActions::Load_Store, nullptr, 0, CubeFace);
-		RHICmdList.BeginRenderPass(RPInfo, TEXT("FilterReflectionEnvironmentRP"));
-		RHICmdList.ApplyCachedRenderTargets(GraphicsPSOInit);
-
-		const FIntPoint SourceDimensions(CubmapSize, CubmapSize);
-		const FIntRect ViewRect(0, 0, EffectiveTopMipSize, EffectiveTopMipSize);
-		RHICmdList.SetViewport(0, 0, 0.0f, EffectiveTopMipSize, EffectiveTopMipSize, 1.0f);
-
-		TShaderMapRef<FScreenVS> VertexShader(GetGlobalShaderMap(FeatureLevel));
-		TShaderMapRef<FOneColorPS> PixelShader(GetGlobalShaderMap(FeatureLevel));
-
-		GraphicsPSOInit.BoundShaderState.VertexDeclarationRHI = GFilterVertexDeclaration.VertexDeclarationRHI;
-		GraphicsPSOInit.BoundShaderState.VertexShaderRHI = VertexShader.GetVertexShader();
-		GraphicsPSOInit.BoundShaderState.PixelShaderRHI = PixelShader.GetPixelShader();
-		GraphicsPSOInit.PrimitiveType = PT_TriangleList;
-
-		SetGraphicsPipelineState(RHICmdList, GraphicsPSOInit, 0);
-
-		FLinearColor UnusedColors[1] = { FLinearColor::Black };
-		PixelShader->SetColors(RHICmdList, UnusedColors, UE_ARRAY_COUNT(UnusedColors));
-
-		DrawRectangle(
-			RHICmdList,
-			ViewRect.Min.X, ViewRect.Min.Y, 
-			ViewRect.Width(), ViewRect.Height(),
-			0, 0, 
-			SourceDimensions.X, SourceDimensions.Y,
-			FIntPoint(ViewRect.Width(), ViewRect.Height()),
-			SourceDimensions,
-			VertexShader);
-
-		RHICmdList.EndRenderPass();
-	}
-
-	RHICmdList.Transition(FRHITransitionInfo(EffectiveColorRT.TargetableTexture, ERHIAccess::Unknown, ERHIAccess::SRVMask));
-
-	auto ShaderMap = GetGlobalShaderMap(FeatureLevel);
+		auto* PassParameters = GraphBuilder.AllocParameters<FOneColorPS::FParameters>();
+		PassParameters->DrawColorMRT[0] = FLinearColor::Black;
+		PassParameters->RenderTargets[0] = FRenderTargetBinding(CubemapTexture, ERenderTargetLoadAction::ELoad, MipIndex, CubeFace);
 
-	FSceneRenderTargetItem& DownSampledCube = GReflectionScratchCubemaps.Color[0]->GetRenderTargetItem();
-	FSceneRenderTargetItem& FilteredCube = GReflectionScratchCubemaps.Color[1]->GetRenderTargetItem();
+		TShaderMapRef<FOneColorPS> PixelShader(ShaderMap);
 
-	CreateCubeMips( RHICmdList, FeatureLevel, NumMips, DownSampledCube );
-
-	if (OutIrradianceEnvironmentMap)
-	{
-		SCOPED_DRAW_EVENT(RHICmdList, ComputeDiffuseIrradiance);
-		
-		const int32 NumDiffuseMips = FMath::CeilLogTwo( GDiffuseIrradianceCubemapSize ) + 1;
-		const int32 DiffuseConvolutionSourceMip = FMath::Max(0, NumMips - NumDiffuseMips);
-
-		ComputeDiffuseIrradiance(RHICmdList, FeatureLevel, DownSampledCube.ShaderResourceTexture, DiffuseConvolutionSourceMip, OutIrradianceEnvironmentMap);
+		FPixelShaderUtils::AddFullscreenPass(
+			GraphBuilder,
+			ShaderMap,
+			RDG_EVENT_NAME("PremultipliedAlpha (Mip: %d, Face %d", MipIndex, CubeFace),
+			PixelShader,
+			PassParameters,
+			ViewRect,
+			TStaticBlendState<CW_RGBA, BO_Add, BF_Zero, BF_DestAlpha, BO_Add, BF_Zero, BF_One>::GetRHI());
 	}
-
-	FilterCubeMap(RHICmdList, FeatureLevel, NumMips, DownSampledCube, FilteredCube);
-	RHICmdList.CopyToResolveTarget(FilteredCube.TargetableTexture, FilteredCube.ShaderResourceTexture, FResolveParams());
 }
 
-class FCopyToCubeFaceShader : public FGlobalShader
-{
-public:
-	static bool ShouldCompilePermutation(const FGlobalShaderPermutationParameters& Parameters)
-	{
-		return true;
-	}
-
-	FCopyToCubeFaceShader() = default;
-	FCopyToCubeFaceShader(const CompiledShaderInitializerType& Initializer)
-		: FGlobalShader(Initializer)
-	{}
-};
-
-/** Vertex shader used when writing to a cubemap. */
-class FCopyToCubeFaceVS : public FCopyToCubeFaceShader
+/** Generates mips for glossiness and filters the cubemap for a given reflection. */
+FRDGTexture* FilterReflectionEnvironment(FRDGBuilder& GraphBuilder, FGlobalShaderMap* ShaderMap, FRDGTexture* CubemapTexture, FSHVectorRGB3* OutIrradianceEnvironmentMap)
 {
-public:
-	DECLARE_GLOBAL_SHADER(FCopyToCubeFaceVS);
-
-	FCopyToCubeFaceVS() = default;
-	FCopyToCubeFaceVS(const CompiledShaderInitializerType & Initializer)
-		: FCopyToCubeFaceShader(Initializer)
-	{}
-};
+	RDG_EVENT_SCOPE(GraphBuilder, "FilterReflectionEnvironment");
 
-IMPLEMENT_GLOBAL_SHADER(FCopyToCubeFaceVS, "/Engine/Private/ReflectionEnvironmentShaders.usf", "CopyToCubeFaceVS", SF_Vertex);
-
-/** Pixel shader used when copying scene color from a scene render into a face of a reflection capture cubemap. */
-class FCopySceneColorToCubeFacePS : public FCopyToCubeFaceShader
-{
-public:
-	DECLARE_GLOBAL_SHADER(FCopySceneColorToCubeFacePS);
-	SHADER_USE_PARAMETER_STRUCT(FCopySceneColorToCubeFacePS, FCopyToCubeFaceShader);
+	PremultiplyCubeMipAlpha(GraphBuilder, ShaderMap, CubemapTexture, 0);
+	CreateCubeMips(GraphBuilder, ShaderMap, CubemapTexture);
 
-	static void ModifyCompilationEnvironment(const FGlobalShaderPermutationParameters& Parameters, FShaderCompilerEnvironment& OutEnvironment)
+	if (OutIrradianceEnvironmentMap)
 	{
-		FCopyToCubeFaceShader::ModifyCompilationEnvironment(Parameters, OutEnvironment);
-
-		if (IsMobilePlatform(Parameters.Platform))
-		{
-			// SceneDepth is memoryless on mobile
-			OutEnvironment.SetDefine(TEXT("SCENE_TEXTURES_DISABLED"), 1u);
-		}
+		ComputeDiffuseIrradiance(GraphBuilder, ShaderMap, CubemapTexture, OutIrradianceEnvironmentMap);
 	}
 
-	BEGIN_SHADER_PARAMETER_STRUCT(FParameters, )
-		SHADER_PARAMETER_STRUCT_REF(FViewUniformShaderParameters, View)
-		SHADER_PARAMETER_RDG_TEXTURE(Texture2D, SceneColorTexture)
-		SHADER_PARAMETER_SAMPLER(SamplerState, SceneColorSampler)
-		SHADER_PARAMETER_RDG_TEXTURE(Texture2D, SceneDepthTexture)
-		SHADER_PARAMETER_SAMPLER(SamplerState, SceneDepthSampler)
-		SHADER_PARAMETER(FVector4f, SkyLightCaptureParameters)
-		SHADER_PARAMETER(FVector4f, LowerHemisphereColor)
-		RENDER_TARGET_BINDING_SLOTS()
-	END_SHADER_PARAMETER_STRUCT()
-};
-
-IMPLEMENT_GLOBAL_SHADER(FCopySceneColorToCubeFacePS, "/Engine/Private/ReflectionEnvironmentShaders.usf", "CopySceneColorToCubeFaceColorPS", SF_Pixel);
-
-/** Pixel shader used when copying a cubemap into a face of a reflection capture cubemap. */
-class FCopyCubemapToCubeFacePS : public FCopyToCubeFaceShader
-{
-public:
-	DECLARE_GLOBAL_SHADER(FCopyCubemapToCubeFacePS);
-	SHADER_USE_PARAMETER_STRUCT(FCopyCubemapToCubeFacePS, FCopyToCubeFaceShader);
-
-	BEGIN_SHADER_PARAMETER_STRUCT(FParameters, )
-		SHADER_PARAMETER_TEXTURE(TextureCube, SourceCubemapTexture)
-		SHADER_PARAMETER_SAMPLER(SamplerState, SourceCubemapSampler)
-		SHADER_PARAMETER(FVector4f, SkyLightCaptureParameters)
-		SHADER_PARAMETER(int32, CubeFace)
-		SHADER_PARAMETER(FVector4f, LowerHemisphereColor)
-		SHADER_PARAMETER(FVector2f, SinCosSourceCubemapRotation)
-		RENDER_TARGET_BINDING_SLOTS()
-	END_SHADER_PARAMETER_STRUCT()
-};
-
-IMPLEMENT_GLOBAL_SHADER(FCopyCubemapToCubeFacePS, "/Engine/Private/ReflectionEnvironmentShaders.usf", "CopyCubemapToCubeFaceColorPS", SF_Pixel);
+	return FilterCubeMap(GraphBuilder, ShaderMap, CubemapTexture);
+}
 
 int32 FindOrAllocateCubemapIndex(FScene* Scene, const UReflectionCaptureComponent* Component)
 {
@@ -690,44 +466,17 @@ int32 FindOrAllocateCubemapIndex(FScene* Scene, const UReflectionCaptureComponen
 	return CubemapIndex;
 }
 
-void ClearScratchCubemaps(FRHICommandListImmediate& RHICmdList, int32 TargetSize)
-{
-	SCOPED_DRAW_EVENT(RHICmdList, ClearScratchCubemaps);
-
-	GReflectionScratchCubemaps.Allocate(RHICmdList, TargetSize);
-
-	FMemMark Mark(FMemStack::Get());
-	FRDGBuilder GraphBuilder(RHICmdList);
-
-	for (int32 RenderTargetIndex = 0; RenderTargetIndex < 2; ++RenderTargetIndex)
-	{
-		FRDGTextureRef OutputTexture = GraphBuilder.RegisterExternalTexture(GReflectionScratchCubemaps.Color[RenderTargetIndex], TEXT("OutputCubemap"));
-
-		RDG_EVENT_SCOPE(GraphBuilder, "ClearScratchCubemapsRT%d", RenderTargetIndex);
-
-		const int32 NumMips = OutputTexture->Desc.NumMips;
-
-		for (int32 MipIndex = 0; MipIndex < NumMips; MipIndex++)
-		{
-			for (int32 CubeFace = 0; CubeFace < CubeFace_MAX; CubeFace++)
-			{
-				FRenderTargetParameters* PassParameters = GraphBuilder.AllocParameters<FRenderTargetParameters>();
-				PassParameters->RenderTargets[0] = FRenderTargetBinding(OutputTexture, ERenderTargetLoadAction::EClear, MipIndex, CubeFace);
-
-				GraphBuilder.AddPass(
-					RDG_EVENT_NAME("ClearCubeFace(Mip: %d, Face: %d)", MipIndex, CubeFace),
-					PassParameters,
-					ERDGPassFlags::Raster,
-					[](FRHICommandList&) {});
-			}
-		}
-	}
-
-	GraphBuilder.Execute();
-}
-
 /** Captures the scene for a reflection capture by rendering the scene multiple times and copying into a cubemap texture. */
-void CaptureSceneToScratchCubemap(FRHICommandListImmediate& RHICmdList, FSceneRenderer* SceneRenderer, ECubeFace CubeFace, int32 CubemapSize, bool bCapturingForSkyLight, bool bLowerHemisphereIsBlack, const FLinearColor& LowerHemisphereColor, bool bCapturingForMobile)
+void CaptureSceneToScratchCubemap(
+	FRHICommandListImmediate& RHICmdList,
+	FSceneRenderer* SceneRenderer,
+	const FReflectionCubemapTexture& ReflectionCubemapTexture,
+	ECubeFace CubeFace,
+	int32 CubemapSize,
+	bool bCapturingForSkyLight,
+	bool bLowerHemisphereIsBlack,
+	const FLinearColor& LowerHemisphereColor,
+	bool bCapturingForMobile)
 {
 	// We need to execute the pre-render view extensions before we do any view dependent work.
 	FSceneRenderer::ViewExtensionPreRender_RenderThread(RHICmdList, SceneRenderer);
@@ -739,11 +488,11 @@ void CaptureSceneToScratchCubemap(FRHICommandListImmediate& RHICmdList, FSceneRe
 	FMaterialRenderProxy::UpdateDeferredCachedUniformExpressions();
 
 	const ERHIFeatureLevel::Type FeatureLevel = SceneRenderer->FeatureLevel;
-	
-	{
-		SCOPED_DRAW_EVENT(RHICmdList, CubeMapCapture);
 
-		FRDGBuilder GraphBuilder(RHICmdList, RDG_EVENT_NAME("CubeMapCapture"), FSceneRenderer::GetRDGParalelExecuteFlags(FeatureLevel));
+	FRDGBuilder GraphBuilder(RHICmdList, RDG_EVENT_NAME("CubeMapCapture"), FSceneRenderer::GetRDGParalelExecuteFlags(FeatureLevel));
+
+	{
+		RDG_EVENT_SCOPE(GraphBuilder, "CubeMapCapture");
 
 		// Render the scene normally for one face of the cubemap
 		SceneRenderer->Render(GraphBuilder);
@@ -758,11 +507,9 @@ void CaptureSceneToScratchCubemap(FRHICommandListImmediate& RHICmdList, FSceneRe
 			InRHICmdList.SubmitCommandsAndFlushGPU();
 		});
 
-		GReflectionScratchCubemaps.Allocate(RHICmdList, CubemapSize);
-
 		const FViewInfo& View = SceneRenderer->Views[0];
 
-		FRDGTextureRef OutputTexture = GraphBuilder.RegisterExternalTexture(GReflectionScratchCubemaps.Color[0]);
+		FRDGTextureRef OutputTexture = ReflectionCubemapTexture.GetRDG(GraphBuilder);
 
 		auto* PassParameters = GraphBuilder.AllocParameters<FCopySceneColorToCubeFacePS::FParameters>();
 		PassParameters->RenderTargets[0] = FRenderTargetBinding(OutputTexture, ERenderTargetLoadAction::ENoAction, 0, CubeFace);
@@ -843,80 +590,67 @@ void CaptureSceneToScratchCubemap(FRHICommandListImmediate& RHICmdList, FSceneRe
 				SceneTextureExtent,
 				VertexShader);
 		});
-
-		GraphBuilder.Execute();
 	}
 
+	GraphBuilder.Execute();
+
 	SceneRenderer->RenderThreadEnd(RHICmdList);
 }
 
-void CopyCubemapToScratchCubemap(FRHICommandListImmediate& RHICmdList, ERHIFeatureLevel::Type FeatureLevel, UTextureCube* SourceCubemap, int32 CubemapSize, bool bIsSkyLight, bool bLowerHemisphereIsBlack, float SourceCubemapRotation, const FLinearColor& LowerHemisphereColorValue)
+void CopyCubemapToScratchCubemap(
+	FRHICommandListImmediate& RHICmdList,
+	ERHIFeatureLevel::Type FeatureLevel,
+	UTextureCube* SourceCubemap,
+	const FReflectionCubemapTexture& ReflectionCubemapTexture,
+	int32 CubemapSize,
+	bool bIsSkyLight,
+	bool bLowerHemisphereIsBlack,
+	float SourceCubemapRotation,
+	const FLinearColor& LowerHemisphereColorValue)
 {
-	SCOPED_DRAW_EVENT(RHICmdList, CopyCubemapToScratchCubemap);
 	check(SourceCubemap);
 
 	const FTexture* SourceCubemapResource = SourceCubemap->GetResource();
+
 	if (SourceCubemapResource == nullptr)
 	{
 		UE_LOG(LogEngine, Warning, TEXT("Unable to copy from cubemap %s, it's RHI resource is null"), *SourceCubemap->GetPathName());
 		return;
 	}
 
-	FMemMark Mark(FMemStack::Get());
 	FRDGBuilder GraphBuilder(RHICmdList);
 
-	FRDGTextureRef OutputTexture = GraphBuilder.RegisterExternalTexture(GReflectionScratchCubemaps.Color[0], TEXT("Color"));
+	auto ShaderMap = GetGlobalShaderMap(FeatureLevel);
 
-	for (uint32 CubeFace = 0; CubeFace < CubeFace_MAX; CubeFace++)
 	{
-		auto* PassParameters = GraphBuilder.AllocParameters<FCopyCubemapToCubeFacePS::FParameters>();
-		PassParameters->RenderTargets[0] = FRenderTargetBinding(OutputTexture, ERenderTargetLoadAction::ENoAction, 0, CubeFace);
-		PassParameters->LowerHemisphereColor = LowerHemisphereColorValue;
-		PassParameters->SkyLightCaptureParameters = FVector3f(bIsSkyLight ? 1.0f : 0.0f, 0.0f, bLowerHemisphereIsBlack ? 1.0f : 0.0f);
-		PassParameters->SourceCubemapSampler = TStaticSamplerState<SF_Bilinear, AM_Clamp, AM_Clamp, AM_Clamp>::GetRHI();
-		PassParameters->SourceCubemapTexture = SourceCubemapResource->TextureRHI;
-		PassParameters->SinCosSourceCubemapRotation = FVector2f(FMath::Sin(SourceCubemapRotation), FMath::Cos(SourceCubemapRotation));
-		PassParameters->CubeFace = CubeFace;
+		RDG_EVENT_SCOPE(GraphBuilder, "CopyCubemapToScratchCubemap");
 
-		const int32 EffectiveSize = CubemapSize;
+		FRDGTextureRef OutputTexture = ReflectionCubemapTexture.GetRDG(GraphBuilder);
 
-		GraphBuilder.AddPass(
-			RDG_EVENT_NAME("CopyCubemapToCubeFace"),
-			PassParameters,
-			ERDGPassFlags::Raster,
-			[EffectiveSize, SourceCubemapResource, PassParameters, FeatureLevel](FRHICommandList& InRHICmdList)
-		{
-			const FIntPoint SourceDimensions(SourceCubemapResource->GetSizeX(), SourceCubemapResource->GetSizeY());
-			const FIntRect ViewRect(0, 0, EffectiveSize, EffectiveSize);
-			InRHICmdList.SetViewport(0.0f, 0.0f, 0.0f, (float)EffectiveSize, (float)EffectiveSize, 1.0f);
-
-			FGraphicsPipelineStateInitializer GraphicsPSOInit;
-			InRHICmdList.ApplyCachedRenderTargets(GraphicsPSOInit);
-			GraphicsPSOInit.RasterizerState = TStaticRasterizerState<FM_Solid, CM_None>::GetRHI();
-			GraphicsPSOInit.DepthStencilState = TStaticDepthStencilState<false, CF_Always>::GetRHI();
-			GraphicsPSOInit.BlendState = TStaticBlendState<>::GetRHI();
-
-			TShaderMapRef<FScreenVS> VertexShader(GetGlobalShaderMap(FeatureLevel));
-			TShaderMapRef<FCopyCubemapToCubeFacePS> PixelShader(GetGlobalShaderMap(FeatureLevel));
+		TShaderMapRef<FCopyCubemapToCubeFacePS> PixelShader(ShaderMap);
 
-			GraphicsPSOInit.BoundShaderState.VertexDeclarationRHI = GFilterVertexDeclaration.VertexDeclarationRHI;
-			GraphicsPSOInit.BoundShaderState.VertexShaderRHI = VertexShader.GetVertexShader();
-			GraphicsPSOInit.BoundShaderState.PixelShaderRHI = PixelShader.GetPixelShader();
-			GraphicsPSOInit.PrimitiveType = PT_TriangleList;
-
-			SetGraphicsPipelineState(InRHICmdList, GraphicsPSOInit, 0);
-			SetShaderParameters(InRHICmdList, PixelShader, PixelShader.GetPixelShader(), *PassParameters);
-
-			DrawRectangle(
-				InRHICmdList,
-				ViewRect.Min.X, ViewRect.Min.Y,
-				ViewRect.Width(), ViewRect.Height(),
-				0, 0,
-				SourceDimensions.X, SourceDimensions.Y,
-				FIntPoint(ViewRect.Width(), ViewRect.Height()),
-				SourceDimensions,
-				VertexShader);
-		});
+		for (uint32 CubeFace = 0; CubeFace < CubeFace_MAX; CubeFace++)
+		{
+			auto* PassParameters = GraphBuilder.AllocParameters<FCopyCubemapToCubeFacePS::FParameters>();
+			PassParameters->SourceCubemapSampler = TStaticSamplerState<SF_Bilinear, AM_Clamp, AM_Clamp, AM_Clamp>::GetRHI();
+			PassParameters->SourceCubemapTexture = SourceCubemapResource->TextureRHI;
+			PassParameters->LowerHemisphereColor = LowerHemisphereColorValue;
+			PassParameters->SkyLightCaptureParameters = FVector3f(bIsSkyLight ? 1.0f : 0.0f, 0.0f, bLowerHemisphereIsBlack ? 1.0f : 0.0f);
+			PassParameters->SinCosSourceCubemapRotation = FVector2f(FMath::Sin(SourceCubemapRotation), FMath::Cos(SourceCubemapRotation));
+			PassParameters->SvPositionToUVScale = FVector2f(1.0f / CubemapSize, 1.0f / CubemapSize);
+			PassParameters->CubeFace = CubeFace;
+			PassParameters->RenderTargets[0] = FRenderTargetBinding(OutputTexture, ERenderTargetLoadAction::ENoAction, 0, CubeFace);
+
+			const FIntRect ViewRect(0, 0, CubemapSize, CubemapSize);
+
+			FPixelShaderUtils::AddFullscreenPass(
+				GraphBuilder,
+				ShaderMap,
+				RDG_EVENT_NAME("CopyCubemapToCubeFace"),
+				PixelShader,
+				PassParameters,
+				ViewRect);
+		}
 	}
 
 	GraphBuilder.Execute();
@@ -1170,10 +904,10 @@ void GetReflectionCaptureData_RenderingThread(FRHICommandListImmediate& RHICmdLi
 
 	if (ComponentStatePtr)
 	{
-		FSceneRenderTargetItem& EffectiveDest = Scene->ReflectionSceneData.CubemapArray.GetRenderTarget();
+		FRHITexture* EffectiveDest = Scene->ReflectionSceneData.CubemapArray.GetRenderTarget()->GetRHI();
 
 		const int32 CubemapIndex = ComponentStatePtr->CubemapIndex;
-		const int32 NumMips = EffectiveDest.ShaderResourceTexture->GetNumMips();
+		const int32 NumMips = EffectiveDest->GetNumMips();
 		const int32 EffectiveTopMipSize = FMath::Pow(2.f, NumMips - 1);
 
 		int32 CaptureDataSize = 0;
@@ -1194,7 +928,7 @@ void GetReflectionCaptureData_RenderingThread(FRHICommandListImmediate& RHICmdLi
 
 		for (int32 MipIndex = 0; MipIndex < NumMips; MipIndex++)
 		{
-			check(EffectiveDest.ShaderResourceTexture->GetFormat() == PF_FloatRGBA);
+			check(EffectiveDest->GetFormat() == PF_FloatRGBA);
 			const int32 MipSize = 1 << (NumMips - MipIndex - 1);
 			const int32 CubeFaceBytes = MipSize * MipSize * sizeof(FFloat16Color);
 
@@ -1204,7 +938,7 @@ void GetReflectionCaptureData_RenderingThread(FRHICommandListImmediate& RHICmdLi
 				// Read each mip face
 				//@todo - do this without blocking the GPU so many times
 				//@todo - pool the temporary textures in RHIReadSurfaceFloatData instead of always creating new ones
-				RHICmdList.ReadSurfaceFloatData(EffectiveDest.ShaderResourceTexture, FIntRect(0, 0, MipSize, MipSize), SurfaceData, (ECubeFace)CubeFace, CubemapIndex, MipIndex);
+				RHICmdList.ReadSurfaceFloatData(EffectiveDest, FIntRect(0, 0, MipSize, MipSize), SurfaceData, (ECubeFace)CubeFace, CubemapIndex, MipIndex);
 				const int32 DestIndex = MipBaseIndex + CubeFace * CubeFaceBytes;
 				uint8* FaceData = &OutCaptureData->FullHDRCapturedData[DestIndex];
 				check(SurfaceData.Num() * SurfaceData.GetTypeSize() == CubeFaceBytes);
@@ -1244,7 +978,7 @@ void UploadReflectionCapture_RenderingThread(FScene* Scene, const FReflectionCap
 	const int32 CaptureIndex = FindOrAllocateCubemapIndex(Scene, CaptureComponent);
 	check(CaptureData->CubemapSize == Scene->ReflectionSceneData.CubemapArray.GetCubemapSize());
 	check(CaptureIndex < Scene->ReflectionSceneData.CubemapArray.GetMaxCubemaps());
-	FTextureCubeRHIRef& CubeMapArray = (FTextureCubeRHIRef&)Scene->ReflectionSceneData.CubemapArray.GetRenderTarget().ShaderResourceTexture;
+	FRHITexture* CubeMapArray = Scene->ReflectionSceneData.CubemapArray.GetRenderTarget()->GetRHI();
 	check(CubeMapArray->GetFormat() == PF_FloatRGBA);
 
 	int32 MipBaseIndex = 0;
@@ -1332,38 +1066,10 @@ FMatrix GetCubeProjectionMatrix(float HalfFovDeg, float CubeMapSize, float NearP
 	return FPerspectiveMatrix(HalfFovDeg, CubeMapSize, CubeMapSize, NearPlane);
 }
 
-/** 
- * Render target class required for rendering the scene.
- * This doesn't actually allocate a render target as we read from scene color to get HDR results directly.
- */
-class FCaptureRenderTarget : public FRenderResource, public FRenderTarget
-{
-public:
-
-	FCaptureRenderTarget() :
-		Size(0)
-	{}
-
-	virtual const FTexture2DRHIRef& GetRenderTargetTexture() const 
-	{
-		static FTexture2DRHIRef DummyTexture;
-		return DummyTexture;
-	}
-
-	void SetSize(int32 TargetSize) { Size = TargetSize; }
-	virtual FIntPoint GetSizeXY() const { return FIntPoint(Size, Size); }
-	virtual float GetDisplayGamma() const { return 1.0f; }
-
-private:
-
-	int32 Size;
-};
-
-TGlobalResource<FCaptureRenderTarget> GReflectionCaptureRenderTarget;
-
 void CaptureSceneIntoScratchCubemap(
 	FScene* Scene, 
-	FVector CapturePosition, 
+	const FReflectionCubemapTexture& ReflectionCubemapTexture,
+	FVector CapturePosition,
 	int32 CubemapSize,
 	bool bCapturingForSkyLight,
 	bool bStaticSceneOnly, 
@@ -1376,6 +1082,27 @@ void CaptureSceneIntoScratchCubemap(
 {
 	int32 SupersampleCaptureFactor = FMath::Clamp(GSupersampleCaptureFactor, MinSupersampleCaptureFactor, MaxSupersampleCaptureFactor);
 
+	class FDummyRenderTarget final : public FRenderTarget, public FRenderThreadStructBase
+	{
+	public:
+		FDummyRenderTarget() = default;
+
+		const FTexture2DRHIRef& GetRenderTargetTexture() const override
+		{
+			static FTexture2DRHIRef DummyTexture;
+			return DummyTexture;
+		}
+
+		void SetSize(int32 TargetSize) { Size = TargetSize; }
+		FIntPoint GetSizeXY() const override { return FIntPoint(Size, Size); }
+		float GetDisplayGamma() const override { return 1.0f; }
+
+	private:
+		int32 Size = 0;
+	};
+
+	TRenderThreadStruct<FDummyRenderTarget> DummyRenderTarget;
+
 	for (int32 CubeFace = 0; CubeFace < CubeFace_MAX; CubeFace++)
 	{
 		if( !bCapturingForSkyLight )
@@ -1384,16 +1111,16 @@ void CaptureSceneIntoScratchCubemap(
 			// Not really a new frame, but it will allow pooling mechanisms to update, like the uniform buffer pool
 			ENQUEUE_RENDER_COMMAND(BeginFrame)(
 				[](FRHICommandList& RHICmdList)
-				{
-					GFrameNumberRenderThread++;
-					RHICmdList.BeginFrame();
-				});
+			{
+				GFrameNumberRenderThread++;
+				RHICmdList.BeginFrame();
+			});
 		}
 
-		GReflectionCaptureRenderTarget.SetSize(CubemapSize);
+		DummyRenderTarget->SetSize(CubemapSize);
 
 		auto ViewFamilyInit = FSceneViewFamily::ConstructionValues(
-			&GReflectionCaptureRenderTarget,
+			DummyRenderTarget.Get(),
 			Scene,
 			FEngineShowFlags(ESFIM_Game)
 			)
@@ -1471,41 +1198,37 @@ void CaptureSceneIntoScratchCubemap(
 		FSceneRenderer* SceneRenderer = FSceneRenderer::CreateSceneRenderer(&ViewFamily, nullptr);
 
 		ENQUEUE_RENDER_COMMAND(CaptureCommand)(
-			[SceneRenderer, CubeFace, CubemapSize, bCapturingForSkyLight, bLowerHemisphereIsBlack, LowerHemisphereColor, bCapturingForMobile](FRHICommandListImmediate& RHICmdList)
-			{
-				CaptureSceneToScratchCubemap(RHICmdList, SceneRenderer, (ECubeFace)CubeFace, CubemapSize, bCapturingForSkyLight, bLowerHemisphereIsBlack, LowerHemisphereColor, bCapturingForMobile);
+			[SceneRenderer, &ReflectionCubemapTexture, CubeFace, CubemapSize, bCapturingForSkyLight, bLowerHemisphereIsBlack, LowerHemisphereColor, bCapturingForMobile](FRHICommandListImmediate& RHICmdList)
+		{
+			CaptureSceneToScratchCubemap(RHICmdList, SceneRenderer, ReflectionCubemapTexture, (ECubeFace)CubeFace, CubemapSize, bCapturingForSkyLight, bLowerHemisphereIsBlack, LowerHemisphereColor, bCapturingForMobile);
 
-				if (!bCapturingForSkyLight)
-				{
-					RHICmdList.EndFrame();
-				}
-			});
+			if (!bCapturingForSkyLight)
+			{
+				RHICmdList.EndFrame();
+			}
+		});
 	}
 }
 
-void CopyToSceneArray(FRHICommandListImmediate& RHICmdList, FScene* Scene, FReflectionCaptureProxy* ReflectionProxy)
+void CopyToSceneArray(FRDGBuilder& GraphBuilder, FScene* Scene, FRDGTexture* FilteredCubeTexture, FReflectionCaptureProxy* ReflectionProxy)
 {
-	SCOPED_DRAW_EVENT(RHICmdList, CopyToSceneArray);
-	const int32 EffectiveTopMipSize = Scene->ReflectionSceneData.CubemapArray.GetCubemapSize();
-	const int32 NumMips = FMath::CeilLogTwo(EffectiveTopMipSize) + 1;
+	RDG_EVENT_SCOPE(GraphBuilder, "CopyToSceneArray");
 
+	const int32 NumMips = GetNumMips(Scene->ReflectionSceneData.CubemapArray.GetCubemapSize());
 	const int32 CaptureIndex = FindOrAllocateCubemapIndex(Scene, ReflectionProxy->Component);
 
-	FSceneRenderTargetItem& FilteredCube = GReflectionScratchCubemaps.Color[1]->GetRenderTargetItem();
-	FSceneRenderTargetItem& DestCube = Scene->ReflectionSceneData.CubemapArray.GetRenderTarget();
+	FRDGTexture* DestCubeTexture = GraphBuilder.RegisterExternalTexture(Scene->ReflectionSceneData.CubemapArray.GetRenderTarget(), ERDGTextureFlags::ForceTracking);
 
 	// GPU copy back to the scene's texture array, which is not a render target
 	for (int32 MipIndex = 0; MipIndex < NumMips; MipIndex++)
 	{
 		for (int32 CubeFace = 0; CubeFace < CubeFace_MAX; CubeFace++)
 		{
-			RHICmdList.CopyToResolveTarget(FilteredCube.ShaderResourceTexture, DestCube.ShaderResourceTexture, FResolveParams(FResolveRect(), (ECubeFace)CubeFace, MipIndex, 0, CaptureIndex));
+			AddCopyToResolveTargetPass(GraphBuilder, FilteredCubeTexture, DestCubeTexture, FResolveParams({}, (ECubeFace)CubeFace, MipIndex, 0, CaptureIndex));
 		}
 	}
 }
 
-
-
 /** 
  * Updates the contents of the given reflection capture by rendering the scene. 
  * This must be called on the game thread.
@@ -1582,12 +1305,8 @@ void FScene::CaptureOrUploadReflectionCapture(UReflectionCaptureComponent* Captu
 				FlushRenderingCommands();
 			}
 
-			ENQUEUE_RENDER_COMMAND(ClearCommand)(
-				[ReflectionCaptureSize](FRHICommandListImmediate& RHICmdList)
-				{
-					ClearScratchCubemaps(RHICmdList, ReflectionCaptureSize);
-				});
-			
+			TRenderThreadStruct<FReflectionCubemapTexture> ReflectionCubemapTexture(ReflectionCaptureSize);
+
 			if (CaptureComponent->ReflectionSourceType == EReflectionSourceType::CapturedScene)
 			{
 				static const auto* AllowStaticLightingVar = IConsoleManager::Get().FindTConsoleVariableDataInt(TEXT("r.AllowStaticLighting"));
@@ -1596,117 +1315,102 @@ void FScene::CaptureOrUploadReflectionCapture(UReflectionCaptureComponent* Captu
 				// Reflection Captures are a form of static lighting, so only capture scene elements that are static
 				// However if the project has static lighting disabled, Reflection Captures can still be made to work by capturing Movable lights
 				bool const bCaptureStaticSceneOnly = CVarReflectionCaptureStaticSceneOnly.GetValueOnGameThread() != 0 && bAllowStaticLighting;
-				CaptureSceneIntoScratchCubemap(this, CaptureComponent->GetComponentLocation() + CaptureComponent->CaptureOffset, ReflectionCaptureSize, false, bCaptureStaticSceneOnly, 0, false, false, FLinearColor(), bCapturingForMobile);
+				CaptureSceneIntoScratchCubemap(this, *ReflectionCubemapTexture, CaptureComponent->GetComponentLocation() + CaptureComponent->CaptureOffset, ReflectionCaptureSize, false, bCaptureStaticSceneOnly, 0, false, false, FLinearColor(), bCapturingForMobile);
 			}
 			else if (CaptureComponent->ReflectionSourceType == EReflectionSourceType::SpecifiedCubemap)
 			{
 				UTextureCube* SourceCubemap = CaptureComponent->Cubemap;
 				float SourceCubemapRotation = CaptureComponent->SourceCubemapAngle * (PI / 180.f);
-				ERHIFeatureLevel::Type InFeatureLevel = FeatureLevel;
 				ENQUEUE_RENDER_COMMAND(CopyCubemapCommand)(
-					[SourceCubemap, ReflectionCaptureSize, SourceCubemapRotation, InFeatureLevel](FRHICommandListImmediate& RHICmdList)
-					{
-						CopyCubemapToScratchCubemap(RHICmdList, InFeatureLevel, SourceCubemap, ReflectionCaptureSize, false, false, SourceCubemapRotation, FLinearColor());
-					});
+					[FeatureLevel = FeatureLevel, SourceCubemap, ReflectionCubemapTexture = ReflectionCubemapTexture.Get(), ReflectionCaptureSize, SourceCubemapRotation](FRHICommandListImmediate& RHICmdList)
+				{
+					CopyCubemapToScratchCubemap(RHICmdList, FeatureLevel, SourceCubemap, *ReflectionCubemapTexture, ReflectionCaptureSize, false, false, SourceCubemapRotation, FLinearColor());
+				});
 			}
 			else
 			{
 				check(!TEXT("Unknown reflection source type"));
 			}
 
-			{
-				ERHIFeatureLevel::Type InFeatureLevel = GetFeatureLevel();
-				int32 InReflectionCaptureSize = ReflectionCaptureSize;
-				FScene* Scene = this;
-				const UReflectionCaptureComponent* InCaptureComponent = CaptureComponent;
-				ENQUEUE_RENDER_COMMAND(FilterCommand)(
-					[InFeatureLevel, InReflectionCaptureSize, Scene, InCaptureComponent](FRHICommandListImmediate& RHICmdList)
-					{
-						FindOrAllocateCubemapIndex(Scene, InCaptureComponent);
-						FCaptureComponentSceneState& FoundState = Scene->ReflectionSceneData.AllocatedReflectionCaptureState.FindChecked(InCaptureComponent);
-
-						ComputeAverageBrightness(RHICmdList, InFeatureLevel, InReflectionCaptureSize, FoundState.AverageBrightness);
-						FilterReflectionEnvironment(RHICmdList, InFeatureLevel, InReflectionCaptureSize, nullptr);
-					});
-			}
-
 			// Create a proxy to represent the reflection capture to the rendering thread
 			// The rendering thread will be responsible for deleting this when done with the filtering operation
 			// We can't use the component's SceneProxy here because the component may not be registered with the scene
 			FReflectionCaptureProxy* ReflectionProxy = new FReflectionCaptureProxy(CaptureComponent);
 
-			FScene* Scene = this;
-			ERHIFeatureLevel::Type InFeatureLevel = GetFeatureLevel();
-			ENQUEUE_RENDER_COMMAND(CopyCommand)(
-				[Scene, ReflectionProxy, InFeatureLevel](FRHICommandListImmediate& RHICmdList)
+			ENQUEUE_RENDER_COMMAND(FilterCommand)(
+				[Scene = this, FeatureLevel = FeatureLevel, ReflectionCubemapTexture = ReflectionCubemapTexture.Get(), ReflectionCaptureSize, CaptureComponent, ReflectionProxy](FRHICommandListImmediate& RHICmdList)
 			{
-				if (InFeatureLevel == ERHIFeatureLevel::SM5)
+				FindOrAllocateCubemapIndex(Scene, CaptureComponent);
+				FCaptureComponentSceneState& FoundState = Scene->ReflectionSceneData.AllocatedReflectionCaptureState.FindChecked(CaptureComponent);
+
+				FRDGBuilder GraphBuilder(RHICmdList);
+
+				auto* ShaderMap = GetGlobalShaderMap(FeatureLevel);
+
+				FRDGTexture* SceneCubemapTexture = ReflectionCubemapTexture->GetRDG(GraphBuilder);
+
+				ComputeAverageBrightness(GraphBuilder, ShaderMap, SceneCubemapTexture, &FoundState.AverageBrightness);
+
+				FRDGTexture* FilteredSceneCubemapTexture = FilterReflectionEnvironment(GraphBuilder, ShaderMap, SceneCubemapTexture, nullptr);
+
+				if (FeatureLevel == ERHIFeatureLevel::SM5)
 				{
-					CopyToSceneArray(RHICmdList, Scene, ReflectionProxy);
+					CopyToSceneArray(GraphBuilder, Scene, FilteredSceneCubemapTexture, ReflectionProxy);
 				}
 
+				GraphBuilder.Execute();
+
 				// Clean up the proxy now that the rendering thread is done with it
 				delete ReflectionProxy;
 			});
-		} //-V773
+		}
 	}
 }
 
-void ReadbackRadianceMap(FRHICommandListImmediate& RHICmdList, int32 CubmapSize, TArray<FFloat16Color>& OutRadianceMap)
+void ReadbackRadianceMap(FRDGBuilder& GraphBuilder, FRDGTexture* InputTexture, TArray<FFloat16Color>* OutRadianceMap)
 {
-	OutRadianceMap.Empty(CubmapSize * CubmapSize * 6);
-	OutRadianceMap.AddZeroed(CubmapSize * CubmapSize * 6);
+	check(InputTexture->Desc.Format == PF_FloatRGBA);
 
-	const int32 MipIndex = 0;
+	AddReadbackTexturePass(GraphBuilder, RDG_EVENT_NAME("ReadbackRadianceMap"), InputTexture, [InputTexture, &OutRadianceMap = *OutRadianceMap](FRHICommandListImmediate& RHICmdList)
+	{
+		const FIntPoint Extent = InputTexture->Desc.Extent;
 
-	FSceneRenderTargetItem& SourceCube = GReflectionScratchCubemaps.Color[0]->GetRenderTargetItem();
-	check(SourceCube.ShaderResourceTexture->GetFormat() == PF_FloatRGBA);
-	const int32 CubeFaceBytes = CubmapSize * CubmapSize * OutRadianceMap.GetTypeSize();
+		const int32 MipIndex = 0;
+		const int32 CubeFaceBytes = Extent.X * Extent.Y * OutRadianceMap.GetTypeSize();
 
-	for (int32 CubeFace = 0; CubeFace < CubeFace_MAX; CubeFace++)
-	{
-		TArray<FFloat16Color> SurfaceData;
+		OutRadianceMap.Empty(Extent.X* Extent.Y * 6);
+		OutRadianceMap.AddZeroed(Extent.X* Extent.Y * 6);
 
-		// Read each mip face
-		RHICmdList.ReadSurfaceFloatData(SourceCube.ShaderResourceTexture, FIntRect(0, 0, CubmapSize, CubmapSize), SurfaceData, (ECubeFace)CubeFace, 0, MipIndex);
-		const int32 DestIndex = CubeFace * CubmapSize * CubmapSize;
-		FFloat16Color* FaceData = &OutRadianceMap[DestIndex];
-		check(SurfaceData.Num() * SurfaceData.GetTypeSize() == CubeFaceBytes);
-		FMemory::Memcpy(FaceData, SurfaceData.GetData(), CubeFaceBytes);
-	}
+		for (int32 CubeFace = 0; CubeFace < CubeFace_MAX; CubeFace++)
+		{
+			TArray<FFloat16Color> SurfaceData;
+
+			// Read each mip face
+			RHICmdList.ReadSurfaceFloatData(InputTexture->GetRHI(), FIntRect(FIntPoint::ZeroValue, Extent), SurfaceData, (ECubeFace)CubeFace, 0, MipIndex);
+			const int32 DestIndex = CubeFace * Extent.X * Extent.Y;
+			FFloat16Color* FaceData = &OutRadianceMap[DestIndex];
+			check(SurfaceData.Num() * SurfaceData.GetTypeSize() == CubeFaceBytes);
+			FMemory::Memcpy(FaceData, SurfaceData.GetData(), CubeFaceBytes);
+		}
+	});
 }
 
-void CopyToSkyTexture(FRHICommandList& RHICmdList, FScene* Scene, FTexture* ProcessedTexture)
+void CopyToSkyTexture(FRDGBuilder& GraphBuilder, FScene* Scene, FRDGTexture* InputTexture, FTexture* ProcessedTexture)
 {
-	SCOPED_DRAW_EVENT(RHICmdList, CopyToSkyTexture);
 	if (ProcessedTexture->TextureRHI)
 	{
-		const int32 EffectiveTopMipSize = ProcessedTexture->GetSizeX();
-		const int32 NumMips = FMath::CeilLogTwo(EffectiveTopMipSize) + 1;
+		RDG_EVENT_SCOPE(GraphBuilder, "CopyToSkyTexture");
 
-		FSceneRenderTargetItem& FilteredCube = GReflectionScratchCubemaps.Color[1]->GetRenderTargetItem();
+		FRDGTexture* OutputTexture = GraphBuilder.RegisterExternalTexture(CreateRenderTarget(ProcessedTexture->TextureRHI, TEXT("SkyTexture")), ERDGTextureFlags::ForceTracking);
 
 		// GPU copy back to the skylight's texture, which is not a render target
 		FRHICopyTextureInfo CopyInfo;
-		CopyInfo.Size = FilteredCube.ShaderResourceTexture->GetSizeXYZ();
+		CopyInfo.Size.X = InputTexture->Desc.Extent.X;
+		CopyInfo.Size.Y = InputTexture->Desc.Extent.Y;
 		CopyInfo.NumSlices = 6;
-		CopyInfo.NumMips = NumMips;
+		CopyInfo.NumMips = GetNumMips(ProcessedTexture->GetSizeX());
 
-		FRHITransitionInfo TransitionsBefore[] =
-		{
-			FRHITransitionInfo(FilteredCube.ShaderResourceTexture, ERHIAccess::Unknown, ERHIAccess::CopySrc),
-			FRHITransitionInfo(ProcessedTexture->TextureRHI, ERHIAccess::Unknown, ERHIAccess::CopyDest)
-		};
-		RHICmdList.Transition(MakeArrayView(TransitionsBefore, UE_ARRAY_COUNT(TransitionsBefore)));
-
-		RHICmdList.CopyTexture(FilteredCube.ShaderResourceTexture, ProcessedTexture->TextureRHI, CopyInfo);
-
-		FRHITransitionInfo TransitionsAfter[] =
-		{
-			FRHITransitionInfo(FilteredCube.ShaderResourceTexture, ERHIAccess::CopySrc, ERHIAccess::SRVMask),
-			FRHITransitionInfo(ProcessedTexture->TextureRHI, ERHIAccess::CopyDest, ERHIAccess::SRVMask)
-		};
-		RHICmdList.Transition(MakeArrayView(TransitionsAfter, UE_ARRAY_COUNT(TransitionsAfter)));
+		AddCopyTexturePass(GraphBuilder, InputTexture, OutputTexture, CopyInfo);
 	}
 }
 
@@ -1719,45 +1423,39 @@ void FScene::UpdateSkyCaptureContents(
 	float& OutAverageBrightness, 
 	FSHVectorRGB3& OutIrradianceEnvironmentMap,
 	TArray<FFloat16Color>* OutRadianceMap)
-{	
-	if (GSupportsRenderTargetFormat_PF_FloatRGBA || GetFeatureLevel() >= ERHIFeatureLevel::SM5)
+{
+	if (GSupportsRenderTargetFormat_PF_FloatRGBA || FeatureLevel >= ERHIFeatureLevel::SM5)
 	{
 		QUICK_SCOPE_CYCLE_COUNTER(STAT_UpdateSkyCaptureContents);
+
+		World = GetWorld();
+
+		if (World)
 		{
-			World = GetWorld();
-			if (World)
-			{
-				//guarantee that all render proxies are up to date before kicking off this render
-				World->SendAllEndOfFrameUpdates();
-			}
-		}
-		{
-			int32 CubemapSize = CaptureComponent->CubemapResolution;
-			ENQUEUE_RENDER_COMMAND(ClearCommand)(
-				[CubemapSize](FRHICommandListImmediate& RHICmdList)
-				{
-					ClearScratchCubemaps(RHICmdList, CubemapSize);
-				});
+			//guarantee that all render proxies are up to date before kicking off this render
+			World->SendAllEndOfFrameUpdates();
 		}
 
+		const int32 CubemapResolution = CaptureComponent->CubemapResolution;
+		const bool bLowerHemisphereIsBlack = CaptureComponent->bLowerHemisphereIsBlack;
+		const FLinearColor LowerHemisphereColor = CaptureComponent->LowerHemisphereColor;
+
+		TRenderThreadStruct<FReflectionCubemapTexture> ReflectionCubemapTexture(CubemapResolution);
+
 		if (CaptureComponent->SourceType == SLS_CapturedScene)
 		{
-			bool bStaticSceneOnly = CaptureComponent->Mobility == EComponentMobility::Static;
-			bool bCapturingForMobile = false;
-			CaptureSceneIntoScratchCubemap(this, CaptureComponent->GetComponentLocation(), CaptureComponent->CubemapResolution, true, bStaticSceneOnly, CaptureComponent->SkyDistanceThreshold, CaptureComponent->bLowerHemisphereIsBlack, bCaptureEmissiveOnly, CaptureComponent->LowerHemisphereColor, bCapturingForMobile);
+			const bool bStaticSceneOnly = CaptureComponent->Mobility == EComponentMobility::Static;
+			const bool bCapturingForMobile = false;
+			CaptureSceneIntoScratchCubemap(this, *ReflectionCubemapTexture, CaptureComponent->GetComponentLocation(), CubemapResolution, true, bStaticSceneOnly, CaptureComponent->SkyDistanceThreshold, bLowerHemisphereIsBlack, bCaptureEmissiveOnly, LowerHemisphereColor, bCapturingForMobile);
 		}
 		else if (CaptureComponent->SourceType == SLS_SpecifiedCubemap)
 		{
-			int32 CubemapSize = CaptureComponent->CubemapResolution;
-			bool bLowerHemisphereIsBlack = CaptureComponent->bLowerHemisphereIsBlack;
-			float SourceCubemapRotation = CaptureComponent->SourceCubemapAngle * (PI / 180.f);
-			ERHIFeatureLevel::Type InnerFeatureLevel = FeatureLevel;
-			FLinearColor LowerHemisphereColor = CaptureComponent->LowerHemisphereColor;
+			const float SourceCubemapRotation = CaptureComponent->SourceCubemapAngle * (PI / 180.f);
 			ENQUEUE_RENDER_COMMAND(CopyCubemapCommand)(
-				[SourceCubemap, CubemapSize, bLowerHemisphereIsBlack, SourceCubemapRotation, InnerFeatureLevel, LowerHemisphereColor](FRHICommandListImmediate& RHICmdList)
-				{
-					CopyCubemapToScratchCubemap(RHICmdList, InnerFeatureLevel, SourceCubemap, CubemapSize, true, bLowerHemisphereIsBlack, SourceCubemapRotation, LowerHemisphereColor);
-				});
+				[FeatureLevel = FeatureLevel, SourceCubemap, ReflectionCubemapTexture = ReflectionCubemapTexture.Get(), CubemapResolution, bLowerHemisphereIsBlack, SourceCubemapRotation, LowerHemisphereColor](FRHICommandListImmediate& RHICmdList)
+			{
+				CopyCubemapToScratchCubemap(RHICmdList, FeatureLevel, SourceCubemap, *ReflectionCubemapTexture, CubemapResolution, true, bLowerHemisphereIsBlack, SourceCubemapRotation, LowerHemisphereColor);
+			});
 		}
 		else if (CaptureComponent->IsRealTimeCaptureEnabled())
 		{
@@ -1766,76 +1464,45 @@ void FScene::UpdateSkyCaptureContents(
 		}
 		else
 		{
-			check(0);
+			checkNoEntry();
 		}
 
-		if (OutRadianceMap)
+		ENQUEUE_RENDER_COMMAND(UpdateCaptureContents)(
+			[Scene = this, FeatureLevel = FeatureLevel, ReflectionCubemapTexture = ReflectionCubemapTexture.Get(), OutAverageBrightness = &OutAverageBrightness, OutIrradianceEnvironmentMap = &OutIrradianceEnvironmentMap, OutProcessedTexture, OutRadianceMap](FRHICommandListImmediate& RHICmdList)
 		{
-			int32 CubemapSize = CaptureComponent->CubemapResolution;
-			ENQUEUE_RENDER_COMMAND(ReadbackCommand)(
-				[CubemapSize, OutRadianceMap](FRHICommandListImmediate& RHICmdList)
-				{
-					ReadbackRadianceMap(RHICmdList, CubemapSize, *OutRadianceMap);
-				});
-		}
+			auto ShaderMap = GetGlobalShaderMap(FeatureLevel);
 
-		{
-			int32 CubemapSize = CaptureComponent->CubemapResolution;
-			float* AverageBrightness = &OutAverageBrightness;
-			FSHVectorRGB3* IrradianceEnvironmentMap = &OutIrradianceEnvironmentMap;
-			ERHIFeatureLevel::Type InFeatureLevel = GetFeatureLevel();
-			ENQUEUE_RENDER_COMMAND(FilterCommand)(
-				[CubemapSize, AverageBrightness, IrradianceEnvironmentMap, InFeatureLevel](FRHICommandListImmediate& RHICmdList)
-				{
-					if (InFeatureLevel <= ERHIFeatureLevel::ES3_1)
-					{
-						MobileReflectionEnvironmentCapture::ComputeAverageBrightness(RHICmdList, InFeatureLevel, CubemapSize, *AverageBrightness);
-						MobileReflectionEnvironmentCapture::FilterReflectionEnvironment(RHICmdList, InFeatureLevel, CubemapSize, IrradianceEnvironmentMap);
-					}
-					else
-					{
-						ComputeAverageBrightness(RHICmdList, InFeatureLevel, CubemapSize, *AverageBrightness);
-						FilterReflectionEnvironment(RHICmdList, InFeatureLevel, CubemapSize, IrradianceEnvironmentMap);
-					}
-				});
-		}
+			FRDGBuilder GraphBuilder(RHICmdList);
 
-		// Optionally copy the filtered mip chain to the output texture
-		if (OutProcessedTexture)
-		{
-			FScene* Scene = this;
-			ERHIFeatureLevel::Type InFeatureLevel = GetFeatureLevel();
-			ENQUEUE_RENDER_COMMAND(CopyCommand)(
-				[Scene, OutProcessedTexture, InFeatureLevel](FRHICommandListImmediate& RHICmdList)
-				{
-					if (InFeatureLevel <= ERHIFeatureLevel::ES3_1)
-					{
-						MobileReflectionEnvironmentCapture::CopyToSkyTexture(RHICmdList, Scene, OutProcessedTexture);
-					}
-					else
-					{
-						CopyToSkyTexture(RHICmdList, Scene, OutProcessedTexture);
-					}
-				});
-		}
+			FRDGTexture* SceneCubemapTexture = ReflectionCubemapTexture->GetRDG(GraphBuilder);
 
-		if (!!GFreeReflectionScratchAfterUse)
-		{
-			ENQUEUE_RENDER_COMMAND(FreeReflectionScratch)(
-				[](FRHICommandListImmediate& RHICmdList)
+			if (OutRadianceMap)
 			{
-				GReflectionScratchCubemaps.Release();
-				GRenderTargetPool.FreeUnusedResources();
-			});
-		}
+				ReadbackRadianceMap(GraphBuilder, SceneCubemapTexture, OutRadianceMap);
+			}
 
-		// These textures should only be manipulated by the render thread,
-		// so enqueue a render command for them to be processed there
-		ENQUEUE_RENDER_COMMAND(ReleasePathTracerSkylightData)(
-			[this](FRHICommandListImmediate& RHICmdList)
-		{
-			PathTracingSkylightTexture.SafeRelease();
-			PathTracingSkylightPdf.SafeRelease();
+			FRDGTexture* FilteredSceneCubemapTexture;
+
+			if (FeatureLevel <= ERHIFeatureLevel::ES3_1)
+			{
+				MobileReflectionEnvironmentCapture::ComputeAverageBrightness(GraphBuilder, ShaderMap, SceneCubemapTexture, OutAverageBrightness);
+				FilteredSceneCubemapTexture = MobileReflectionEnvironmentCapture::FilterReflectionEnvironment(GraphBuilder, ShaderMap, SceneCubemapTexture, OutIrradianceEnvironmentMap);
+			}
+			else
+			{
+				ComputeAverageBrightness(GraphBuilder, ShaderMap, SceneCubemapTexture, OutAverageBrightness);
+				FilteredSceneCubemapTexture = FilterReflectionEnvironment(GraphBuilder, ShaderMap, SceneCubemapTexture, OutIrradianceEnvironmentMap);
+			}
+
+			if (OutProcessedTexture)
+			{
+				CopyToSkyTexture(GraphBuilder, Scene, FilteredSceneCubemapTexture, OutProcessedTexture);
+			}
+
+			GraphBuilder.Execute();
+
+			Scene->PathTracingSkylightTexture = nullptr;
+			Scene->PathTracingSkylightPdf = nullptr;
 		});
 	}
 }
diff --git a/Engine/Source/Runtime/Renderer/Private/ReflectionEnvironmentCapture.h b/Engine/Source/Runtime/Renderer/Private/ReflectionEnvironmentCapture.h
--- a/Engine/Source/Runtime/Renderer/Private/ReflectionEnvironmentCapture.h
+++ b/Engine/Source/Runtime/Renderer/Private/ReflectionEnvironmentCapture.h
@@ -10,78 +10,14 @@
 #include "Math/SHMath.h"
 #include "RHI.h"
 #include "GlobalShader.h"
-#include "RendererInterface.h"
+#include "RenderGraphDefinitions.h"
 
-extern void ComputeDiffuseIrradiance(FRHICommandListImmediate& RHICmdList, ERHIFeatureLevel::Type FeatureLevel, FTextureRHIRef LightingSource, int32 LightingSourceMipIndex, FSHVectorRGB3* OutIrradianceEnvironmentMap);
+extern void ComputeDiffuseIrradiance(FRDGBuilder& GraphBuilder, FGlobalShaderMap* ShaderMap, FRDGTexture* LightingSource,  FSHVectorRGB3* OutIrradianceEnvironmentMap);
 
 FMatrix CalcCubeFaceViewRotationMatrix(ECubeFace Face);
 FMatrix GetCubeProjectionMatrix(float HalfFovDeg, float CubeMapSize, float NearPlane);
 
-/** Pixel shader used for filtering a mip. */
-class FCubeFilterPS : public FGlobalShader
+inline uint32 GetNumMips(uint32 MipSize)
 {
-	DECLARE_SHADER_TYPE(FCubeFilterPS, Global);
-public:
-
-	static bool ShouldCompilePermutation(const FGlobalShaderPermutationParameters& Parameters)
-	{
-		return true;
-	}
-
-	FCubeFilterPS(const ShaderMetaType::CompiledShaderInitializerType& Initializer) :
-		FGlobalShader(Initializer)
-	{
-		CubeFace.Bind(Initializer.ParameterMap, TEXT("CubeFace"));
-		MipIndex.Bind(Initializer.ParameterMap, TEXT("MipIndex"));
-		NumMips.Bind(Initializer.ParameterMap, TEXT("NumMips"));
-		SourceCubemapTexture.Bind(Initializer.ParameterMap, TEXT("SourceCubemapTexture"));
-		SourceCubemapSampler.Bind(Initializer.ParameterMap, TEXT("SourceCubemapSampler"));
-	}
-	FCubeFilterPS() {}
-
-	LAYOUT_FIELD(FShaderParameter, CubeFace);
-	LAYOUT_FIELD(FShaderParameter, MipIndex);
-	LAYOUT_FIELD(FShaderParameter, NumMips);
-	LAYOUT_FIELD(FShaderResourceParameter, SourceCubemapTexture);
-	LAYOUT_FIELD(FShaderResourceParameter, SourceCubemapSampler);
-};
-
-template< uint32 bNormalize >
-class TCubeFilterPS : public FCubeFilterPS
-{
-	DECLARE_SHADER_TYPE(TCubeFilterPS, Global);
-public:
-	static void ModifyCompilationEnvironment(const FGlobalShaderPermutationParameters& Parameters, FShaderCompilerEnvironment& OutEnvironment)
-	{
-		FCubeFilterPS::ModifyCompilationEnvironment(Parameters, OutEnvironment);
-		OutEnvironment.SetDefine(TEXT("NORMALIZE"), bNormalize);
-	}
-
-	TCubeFilterPS(const ShaderMetaType::CompiledShaderInitializerType& Initializer)
-		: FCubeFilterPS(Initializer)
-	{}
-
-	TCubeFilterPS() {}
-};
-
-class FReflectionScratchCubemaps : public FRenderResource
-{
-public:
-	void Allocate(FRHICommandList& RHICmdList, uint32 TargetSize);
-	void Release();
-
-	/** 2 scratch cubemaps used for filtering reflections. */
-	TRefCountPtr<IPooledRenderTarget> Color[2];
-
-	/** Temporary storage during SH irradiance map generation. */
-	TRefCountPtr<IPooledRenderTarget> Irradiance[2];
-
-	/** Temporary storage during SH irradiance map generation. */
-	TRefCountPtr<IPooledRenderTarget> SkySHIrradiance;
-
-private:
-	void ReleaseDynamicRHI() override
-	{
-		Release();
-	}
-};
\ No newline at end of file
+	return FMath::CeilLogTwo(MipSize) + 1;
+}
\ No newline at end of file
diff --git a/Engine/Source/Runtime/Renderer/Private/ReflectionEnvironmentDiffuseIrradiance.cpp b/Engine/Source/Runtime/Renderer/Private/ReflectionEnvironmentDiffuseIrradiance.cpp
--- a/Engine/Source/Runtime/Renderer/Private/ReflectionEnvironmentDiffuseIrradiance.cpp
+++ b/Engine/Source/Runtime/Renderer/Private/ReflectionEnvironmentDiffuseIrradiance.cpp
@@ -18,341 +18,200 @@
 #include "PipelineStateCache.h"
 #include "VisualizeTexture.h"
 #include "CommonRenderResources.h"
+#include "PixelShaderUtils.h"
 
-extern int32 GDiffuseIrradianceCubemapSize;
+int32 GDiffuseIrradianceCubemapSize = 32;
 
-extern TGlobalResource<FReflectionScratchCubemaps> GReflectionScratchCubemaps;
-
-FSceneRenderTargetItem& GetEffectiveDiffuseIrradianceRenderTarget(int32 TargetMipIndex)
+int32 GetDiffuseConvolutionSourceMip(FRDGTexture* Texture)
 {
-	const int32 ScratchTextureIndex = TargetMipIndex % 2;
-
-	return GReflectionScratchCubemaps.Irradiance[ScratchTextureIndex]->GetRenderTargetItem();
-}
-
-FSceneRenderTargetItem& GetEffectiveDiffuseIrradianceSourceTexture(int32 TargetMipIndex)
-{
-	const int32 ScratchTextureIndex = 1 - TargetMipIndex % 2;
-
-	return GReflectionScratchCubemaps.Irradiance[ScratchTextureIndex]->GetRenderTargetItem();
+	const int32 NumMips = Texture->Desc.NumMips;
+	const int32 NumDiffuseMips = GetNumMips(GDiffuseIrradianceCubemapSize);
+	return FMath::Max(0, NumMips - NumDiffuseMips);
 }
 
 /** Pixel shader used for copying to diffuse irradiance texture. */
 class FCopyDiffuseIrradiancePS : public FGlobalShader
 {
-	DECLARE_SHADER_TYPE(FCopyDiffuseIrradiancePS,Global);
-public:
-
-	static bool ShouldCompilePermutation(const FGlobalShaderPermutationParameters& Parameters)
-	{
-		return true;
-	}
-
-	FCopyDiffuseIrradiancePS(const ShaderMetaType::CompiledShaderInitializerType& Initializer):
-		FGlobalShader(Initializer)
-	{
-		CubeFace.Bind(Initializer.ParameterMap,TEXT("CubeFace"));
-		SourceMipIndex.Bind(Initializer.ParameterMap,TEXT("SourceMipIndex"));
-		SourceCubemapTexture.Bind(Initializer.ParameterMap,TEXT("SourceCubemapTexture"));
-		SourceCubemapSampler.Bind(Initializer.ParameterMap,TEXT("SourceCubemapSampler"));
-		CoefficientMask0.Bind(Initializer.ParameterMap,TEXT("CoefficientMask0"));
-		CoefficientMask1.Bind(Initializer.ParameterMap,TEXT("CoefficientMask1"));
-		CoefficientMask2.Bind(Initializer.ParameterMap,TEXT("CoefficientMask2"));
-		NumSamples.Bind(Initializer.ParameterMap,TEXT("NumSamples"));
-	}
-	FCopyDiffuseIrradiancePS() {}
-
-	void SetParameters(FRHICommandList& RHICmdList, int32 CubeFaceValue, int32 SourceMipIndexValue, int32 CoefficientIndex, int32 FaceResolution, FTextureRHIRef& SourceTextureValue)
-	{
-		SetShaderValue(RHICmdList, RHICmdList.GetBoundPixelShader(), CubeFace, CubeFaceValue);
-		SetShaderValue(RHICmdList, RHICmdList.GetBoundPixelShader(), SourceMipIndex, SourceMipIndexValue);
-
-		SetTextureParameter(
-			RHICmdList, 
-			RHICmdList.GetBoundPixelShader(),
-			SourceCubemapTexture,
-			SourceCubemapSampler,
-			TStaticSamplerState<SF_Bilinear, AM_Clamp, AM_Clamp, AM_Clamp>::GetRHI(), 
-			SourceTextureValue);
-
-		const FVector4f Mask0(CoefficientIndex == 0, CoefficientIndex == 1, CoefficientIndex == 2, CoefficientIndex == 3);
-		const FVector4f Mask1(CoefficientIndex == 4, CoefficientIndex == 5, CoefficientIndex == 6, CoefficientIndex == 7);
-		const float Mask2 = CoefficientIndex == 8;
-		SetShaderValue(RHICmdList, RHICmdList.GetBoundPixelShader(), CoefficientMask0, Mask0);
-		SetShaderValue(RHICmdList, RHICmdList.GetBoundPixelShader(), CoefficientMask1, Mask1);
-		SetShaderValue(RHICmdList, RHICmdList.GetBoundPixelShader(), CoefficientMask2, Mask2);
-
-		SetShaderValue(RHICmdList, RHICmdList.GetBoundPixelShader(), NumSamples, FaceResolution * FaceResolution * 6);
-	}
-
-private:
-	LAYOUT_FIELD(FShaderParameter, CubeFace)
-	LAYOUT_FIELD(FShaderParameter, SourceMipIndex)
-	LAYOUT_FIELD(FShaderResourceParameter, SourceCubemapTexture)
-	LAYOUT_FIELD(FShaderResourceParameter, SourceCubemapSampler)
-	LAYOUT_FIELD(FShaderParameter, CoefficientMask0)
-	LAYOUT_FIELD(FShaderParameter, CoefficientMask1)
-	LAYOUT_FIELD(FShaderParameter, CoefficientMask2)
-	LAYOUT_FIELD(FShaderParameter, NumSamples)
+	DECLARE_GLOBAL_SHADER(FCopyDiffuseIrradiancePS);
+	SHADER_USE_PARAMETER_STRUCT(FCopyDiffuseIrradiancePS, FGlobalShader);
+
+	BEGIN_SHADER_PARAMETER_STRUCT(FParameters, )
+		SHADER_PARAMETER_RDG_TEXTURE(TextureCube, SourceCubemapTexture)
+		SHADER_PARAMETER_SAMPLER(SamplerState, SourceCubemapSampler)
+		SHADER_PARAMETER(FVector4f, CoefficientMask0)
+		SHADER_PARAMETER(FVector4f, CoefficientMask1)
+		SHADER_PARAMETER(FVector2f, SvPositionToUVScale)
+		SHADER_PARAMETER(float, CoefficientMask2)
+		SHADER_PARAMETER(int32, CubeFace)
+		SHADER_PARAMETER(int32, SourceMipIndex)
+		SHADER_PARAMETER(int32, NumSamples)
+		RENDER_TARGET_BINDING_SLOTS()
+	END_SHADER_PARAMETER_STRUCT()
 };
 
-IMPLEMENT_SHADER_TYPE(,FCopyDiffuseIrradiancePS,TEXT("/Engine/Private/ReflectionEnvironmentShaders.usf"),TEXT("DiffuseIrradianceCopyPS"),SF_Pixel)
+IMPLEMENT_GLOBAL_SHADER(FCopyDiffuseIrradiancePS, "/Engine/Private/ReflectionEnvironmentShaders.usf", "DiffuseIrradianceCopyPS", SF_Pixel)
 
-/**  */
 class FAccumulateDiffuseIrradiancePS : public FGlobalShader
 {
-	DECLARE_SHADER_TYPE(FAccumulateDiffuseIrradiancePS,Global);
-public:
-
-	static bool ShouldCompilePermutation(const FGlobalShaderPermutationParameters& Parameters)
-	{
-		return true;
-	}
-
-	FAccumulateDiffuseIrradiancePS(const ShaderMetaType::CompiledShaderInitializerType& Initializer):
-		FGlobalShader(Initializer)
-	{
-		CubeFace.Bind(Initializer.ParameterMap,TEXT("CubeFace"));
-		SourceMipIndex.Bind(Initializer.ParameterMap,TEXT("SourceMipIndex"));
-		SourceCubemapTexture.Bind(Initializer.ParameterMap,TEXT("SourceCubemapTexture"));
-		SourceCubemapSampler.Bind(Initializer.ParameterMap,TEXT("SourceCubemapSampler"));
-		Sample01.Bind(Initializer.ParameterMap,TEXT("Sample01"));
-		Sample23.Bind(Initializer.ParameterMap,TEXT("Sample23"));
-	}
-	FAccumulateDiffuseIrradiancePS() {}
-
-	void SetParameters(FRHICommandList& RHICmdList, int32 CubeFaceValue, int32 NumMips, int32 SourceMipIndexValue, int32 CoefficientIndex, FTextureRHIRef& SourceTextureValue)
-	{
-		SetShaderValue(RHICmdList, RHICmdList.GetBoundPixelShader(), CubeFace, CubeFaceValue);
-		SetShaderValue(RHICmdList, RHICmdList.GetBoundPixelShader(), SourceMipIndex, SourceMipIndexValue);
-
-		SetTextureParameter(
-			RHICmdList, 
-			RHICmdList.GetBoundPixelShader(),
-			SourceCubemapTexture,
-			SourceCubemapSampler,
-			TStaticSamplerState<SF_Point, AM_Clamp, AM_Clamp, AM_Clamp>::GetRHI(), 
-			SourceTextureValue);
-
-		const int32 MipSize = 1 << (NumMips - SourceMipIndexValue - 1);
-		const float HalfSourceTexelSize = .5f / MipSize;
-		const FVector4f Sample01Value(-HalfSourceTexelSize, -HalfSourceTexelSize, HalfSourceTexelSize, -HalfSourceTexelSize);
-		const FVector4f Sample23Value(-HalfSourceTexelSize, HalfSourceTexelSize, HalfSourceTexelSize, HalfSourceTexelSize);
-		SetShaderValue(RHICmdList, RHICmdList.GetBoundPixelShader(), Sample01, Sample01Value);
-		SetShaderValue(RHICmdList, RHICmdList.GetBoundPixelShader(), Sample23, Sample23Value);
-	}
-
-private:
-	LAYOUT_FIELD(FShaderParameter, CubeFace)
-	LAYOUT_FIELD(FShaderParameter, SourceMipIndex)
-	LAYOUT_FIELD(FShaderResourceParameter, SourceCubemapTexture)
-	LAYOUT_FIELD(FShaderResourceParameter, SourceCubemapSampler)
-	LAYOUT_FIELD(FShaderParameter, Sample01)
-	LAYOUT_FIELD(FShaderParameter, Sample23)
+	DECLARE_GLOBAL_SHADER(FAccumulateDiffuseIrradiancePS);
+	SHADER_USE_PARAMETER_STRUCT(FAccumulateDiffuseIrradiancePS, FGlobalShader);
+
+	BEGIN_SHADER_PARAMETER_STRUCT(FParameters, )
+		SHADER_PARAMETER_RDG_TEXTURE_SRV(TextureCube, SourceCubemapTexture)
+		SHADER_PARAMETER_SAMPLER(SamplerState, SourceCubemapSampler)
+		SHADER_PARAMETER(FVector4f, Sample01)
+		SHADER_PARAMETER(FVector4f, Sample23)
+		SHADER_PARAMETER(FVector2f, SvPositionToUVScale)
+		SHADER_PARAMETER(int32, CubeFace)
+		SHADER_PARAMETER(int32, SourceMipIndex)
+		RENDER_TARGET_BINDING_SLOTS()
+	END_SHADER_PARAMETER_STRUCT()
 };
 
-IMPLEMENT_SHADER_TYPE(,FAccumulateDiffuseIrradiancePS,TEXT("/Engine/Private/ReflectionEnvironmentShaders.usf"),TEXT("DiffuseIrradianceAccumulatePS"),SF_Pixel)
+IMPLEMENT_GLOBAL_SHADER(FAccumulateDiffuseIrradiancePS, "/Engine/Private/ReflectionEnvironmentShaders.usf", "DiffuseIrradianceAccumulatePS", SF_Pixel)
 
-/**  */
 class FAccumulateCubeFacesPS : public FGlobalShader
 {
-	DECLARE_SHADER_TYPE(FAccumulateCubeFacesPS,Global);
-public:
-
-	static bool ShouldCompilePermutation(const FGlobalShaderPermutationParameters& Parameters)
-	{
-		return true;
-	}
-
-	FAccumulateCubeFacesPS(const ShaderMetaType::CompiledShaderInitializerType& Initializer):
-		FGlobalShader(Initializer)
-	{
-		SourceMipIndex.Bind(Initializer.ParameterMap,TEXT("SourceMipIndex"));
-		SourceCubemapTexture.Bind(Initializer.ParameterMap,TEXT("SourceCubemapTexture"));
-		SourceCubemapSampler.Bind(Initializer.ParameterMap,TEXT("SourceCubemapSampler"));
-	}
-	FAccumulateCubeFacesPS() {}
-
-	void SetParameters(FRHICommandList& RHICmdList, int32 SourceMipIndexValue, FTextureRHIRef& SourceTextureValue)
-	{
-		SetShaderValue(RHICmdList, RHICmdList.GetBoundPixelShader(), SourceMipIndex, SourceMipIndexValue);
-
-		SetTextureParameter(
-			RHICmdList, 
-			RHICmdList.GetBoundPixelShader(),
-			SourceCubemapTexture,
-			SourceCubemapSampler,
-			TStaticSamplerState<SF_Point, AM_Clamp, AM_Clamp, AM_Clamp>::GetRHI(), 
-			SourceTextureValue);
-	}
-
-private:
-	LAYOUT_FIELD(FShaderParameter, SourceMipIndex)
-	LAYOUT_FIELD(FShaderResourceParameter, SourceCubemapTexture)
-	LAYOUT_FIELD(FShaderResourceParameter, SourceCubemapSampler)
+	DECLARE_GLOBAL_SHADER(FAccumulateCubeFacesPS);
+	SHADER_USE_PARAMETER_STRUCT(FAccumulateCubeFacesPS, FGlobalShader);
+
+	BEGIN_SHADER_PARAMETER_STRUCT(FParameters, )
+		SHADER_PARAMETER_RDG_TEXTURE(TextureCube, SourceCubemapTexture)
+		SHADER_PARAMETER_SAMPLER(SamplerState, SourceCubemapSampler)
+		SHADER_PARAMETER(int32, SourceMipIndex)
+		RENDER_TARGET_BINDING_SLOTS()
+	END_SHADER_PARAMETER_STRUCT()
 };
 
-IMPLEMENT_SHADER_TYPE(,FAccumulateCubeFacesPS,TEXT("/Engine/Private/ReflectionEnvironmentShaders.usf"),TEXT("AccumulateCubeFacesPS"),SF_Pixel)
+IMPLEMENT_GLOBAL_SHADER(FAccumulateCubeFacesPS, "/Engine/Private/ReflectionEnvironmentShaders.usf", "AccumulateCubeFacesPS", SF_Pixel)
 
-void ComputeDiffuseIrradiance(FRHICommandListImmediate& RHICmdList, ERHIFeatureLevel::Type FeatureLevel, FTextureRHIRef LightingSource, int32 LightingSourceMipIndex, FSHVectorRGB3* OutIrradianceEnvironmentMap)
+void ComputeDiffuseIrradiance(FRDGBuilder& GraphBuilder, FGlobalShaderMap* ShaderMap, FRDGTexture* LightingSource, FSHVectorRGB3* OutIrradianceEnvironmentMap)
 {
-	auto ShaderMap = GetGlobalShaderMap(FeatureLevel);
+	RDG_EVENT_SCOPE(GraphBuilder, "ComputeDiffuseIrradiance");
+
+	const int32 LightingSourceMipIndex = GetDiffuseConvolutionSourceMip(LightingSource);
+	const int32 NumMips = GetNumMips(GDiffuseIrradianceCubemapSize);
 
 	FGraphicsPipelineStateInitializer GraphicsPSOInit;
 	GraphicsPSOInit.RasterizerState = TStaticRasterizerState<FM_Solid, CM_None>::GetRHI();
 	GraphicsPSOInit.DepthStencilState = TStaticDepthStencilState<false, CF_Always>::GetRHI();
 	GraphicsPSOInit.BlendState = TStaticBlendState<>::GetRHI();
 
+	FRDGTexture* SkySHIrradianceTexture = GraphBuilder.CreateTexture(FRDGTextureDesc::Create2D(FIntPoint(FSHVector3::MaxSHBasis, 1), PF_FloatRGBA, FClearValueBinding(FLinearColor(0, 10000, 0, 0)), TexCreate_RenderTargetable), TEXT("SkySHIrradiance"));
+
+	TShaderMapRef<FCopyDiffuseIrradiancePS> CopyDiffuseIrradiancePS(ShaderMap);
+	TShaderMapRef<FAccumulateDiffuseIrradiancePS> AccumulateDiffuseIrradiancePS(ShaderMap);
+	TShaderMapRef<FAccumulateCubeFacesPS> AccumulateCubeFacesPS(ShaderMap);
+
 	for (int32 CoefficientIndex = 0; CoefficientIndex < FSHVector3::MaxSHBasis; CoefficientIndex++)
 	{
+		RDG_EVENT_SCOPE(GraphBuilder, "Coefficient %d", CoefficientIndex);
+
+		const FRDGTextureDesc IrradianceCubemapDesc(
+			FRDGTextureDesc::CreateCube(GDiffuseIrradianceCubemapSize, PF_FloatRGBA, FClearValueBinding(FLinearColor(0, 10000, 0, 0)), TexCreate_TargetArraySlicesIndependently | TexCreate_ShaderResource | TexCreate_RenderTargetable | TexCreate_DisableDCC, NumMips));
+
+		FRDGTexture* IrradianceCubemapTexture = GraphBuilder.CreateTexture(IrradianceCubemapDesc, TEXT("DiffuseIrradianceCubemap"));
+
 		// Copy the starting mip from the lighting texture, apply texel area weighting and appropriate SH coefficient
 		{
+			RDG_EVENT_SCOPE(GraphBuilder, "CopyDiffuseIrradiance");
+
+			const FVector4f Mask0(CoefficientIndex == 0, CoefficientIndex == 1, CoefficientIndex == 2, CoefficientIndex == 3);
+			const FVector4f Mask1(CoefficientIndex == 4, CoefficientIndex == 5, CoefficientIndex == 6, CoefficientIndex == 7);
+			const float Mask2 = CoefficientIndex == 8;
+
 			const int32 MipIndex = 0;
 			const int32 MipSize = GDiffuseIrradianceCubemapSize;
-			FSceneRenderTargetItem& EffectiveRT = GetEffectiveDiffuseIrradianceRenderTarget(MipIndex);
-
-			RHICmdList.Transition(FRHITransitionInfo(EffectiveRT.TargetableTexture, ERHIAccess::Unknown, ERHIAccess::RTV));
+			const FIntRect ViewRect(0, 0, MipSize, MipSize);
 
 			for (int32 CubeFace = 0; CubeFace < CubeFace_MAX; CubeFace++)
 			{
-				FRHIRenderPassInfo RPInfo(EffectiveRT.TargetableTexture, ERenderTargetActions::DontLoad_Store, nullptr, 0, CubeFace);
-				RHICmdList.BeginRenderPass(RPInfo, TEXT("CopyDiffuseIrradianceRP"));
-				RHICmdList.ApplyCachedRenderTargets(GraphicsPSOInit);
-
-				const FIntRect ViewRect(0, 0, MipSize, MipSize);
-				RHICmdList.SetViewport(0.0f, 0.0f, 0.0f, (float)MipSize, (float)MipSize, 1.0f);
-				TShaderMapRef<FCopyDiffuseIrradiancePS> PixelShader(ShaderMap);
-					
-				TShaderMapRef<FScreenVS> VertexShader(GetGlobalShaderMap(FeatureLevel));
-
-				GraphicsPSOInit.BoundShaderState.VertexDeclarationRHI = GFilterVertexDeclaration.VertexDeclarationRHI;
-				GraphicsPSOInit.BoundShaderState.VertexShaderRHI = VertexShader.GetVertexShader();
-				GraphicsPSOInit.BoundShaderState.PixelShaderRHI = PixelShader.GetPixelShader();
-				GraphicsPSOInit.PrimitiveType = PT_TriangleList;
-
-				SetGraphicsPipelineState(RHICmdList, GraphicsPSOInit, 0);
-
-				PixelShader->SetParameters(RHICmdList, CubeFace, LightingSourceMipIndex, CoefficientIndex, MipSize, LightingSource);
-					
-				DrawRectangle(
-					RHICmdList,
-					ViewRect.Min.X, ViewRect.Min.Y,
-					ViewRect.Width(), ViewRect.Height(),
-					ViewRect.Min.X, ViewRect.Min.Y,
-					ViewRect.Width(), ViewRect.Height(),
-					FIntPoint(ViewRect.Width(), ViewRect.Height()),
-					FIntPoint(MipSize, MipSize),
-					VertexShader);
-					
-				RHICmdList.EndRenderPass();
+				auto* PassParameters = GraphBuilder.AllocParameters<FCopyDiffuseIrradiancePS::FParameters>();
+				PassParameters->SourceCubemapTexture = LightingSource;
+				PassParameters->SourceCubemapSampler = TStaticSamplerState<SF_Bilinear, AM_Clamp, AM_Clamp, AM_Clamp>::GetRHI();
+				PassParameters->CoefficientMask0 = Mask0;
+				PassParameters->CoefficientMask1 = Mask1;
+				PassParameters->CoefficientMask2 = Mask2;
+				PassParameters->SvPositionToUVScale = FVector2f(1.0f / MipSize, 1.0f / MipSize);
+				PassParameters->CubeFace = CubeFace;
+				PassParameters->SourceMipIndex = LightingSourceMipIndex;
+				PassParameters->NumSamples = MipSize * MipSize * 6;
+				PassParameters->RenderTargets[0] = FRenderTargetBinding(IrradianceCubemapTexture, ERenderTargetLoadAction::ENoAction, 0, CubeFace);
+
+				FPixelShaderUtils::AddFullscreenPass(
+					GraphBuilder,
+					ShaderMap,
+					RDG_EVENT_NAME("Face: %d", CubeFace),
+					CopyDiffuseIrradiancePS,
+					PassParameters,
+					ViewRect);
 			}
-
-			RHICmdList.Transition(FRHITransitionInfo(EffectiveRT.TargetableTexture, ERHIAccess::RTV, ERHIAccess::SRVGraphics));
 		}
 
-		const int32 NumMips = FMath::CeilLogTwo(GDiffuseIrradianceCubemapSize) + 1;
-
 		{
+			RDG_EVENT_SCOPE(GraphBuilder, "AccumulateDiffuseIrradiance");
+
 			// Accumulate all the texel values through downsampling to 1x1 mip
 			for (int32 MipIndex = 1; MipIndex < NumMips; MipIndex++)
 			{
-				const int32 SourceMipIndex = FMath::Max(MipIndex - 1, 0);
 				const int32 MipSize = 1 << (NumMips - MipIndex - 1);
+				const FIntRect ViewRect(0, 0, MipSize, MipSize);
 
-				FSceneRenderTargetItem& EffectiveRT = GetEffectiveDiffuseIrradianceRenderTarget(MipIndex);
-				FSceneRenderTargetItem& EffectiveSource = GetEffectiveDiffuseIrradianceSourceTexture(MipIndex);
-				check(EffectiveRT.TargetableTexture != EffectiveSource.ShaderResourceTexture);
+				const int32 SourceMipIndex = FMath::Max(MipIndex - 1, 0);
+				const int32 SourceMipSize = 1 << (NumMips - SourceMipIndex - 1);
 
-				RHICmdList.Transition(FRHITransitionInfo(EffectiveRT.TargetableTexture, ERHIAccess::Unknown, ERHIAccess::RTV));
+				const float HalfSourceTexelSize = 0.5f / MipSize;
+				const FVector4f Sample01(-HalfSourceTexelSize, -HalfSourceTexelSize, HalfSourceTexelSize, -HalfSourceTexelSize);
+				const FVector4f Sample23(-HalfSourceTexelSize, HalfSourceTexelSize, HalfSourceTexelSize, HalfSourceTexelSize);
 
 				for (int32 CubeFace = 0; CubeFace < CubeFace_MAX; CubeFace++)
 				{
-					FRHIRenderPassInfo RPInfo(EffectiveRT.TargetableTexture, ERenderTargetActions::Load_Store, nullptr, MipIndex, CubeFace);
-					RHICmdList.BeginRenderPass(RPInfo, TEXT("AccumulateDiffuseIrradianceRP"));
-					RHICmdList.ApplyCachedRenderTargets(GraphicsPSOInit);
-
-					const FIntRect ViewRect(0, 0, MipSize, MipSize);
-					RHICmdList.SetViewport(0.0f, 0.0f, 0.0f, (float)MipSize, (float)MipSize, 1.0f);
-					TShaderMapRef<FAccumulateDiffuseIrradiancePS> PixelShader(ShaderMap);
-						
-					TShaderMapRef<FScreenVS> VertexShader(GetGlobalShaderMap(FeatureLevel));
-						
-					GraphicsPSOInit.BoundShaderState.VertexDeclarationRHI = GFilterVertexDeclaration.VertexDeclarationRHI;
-					GraphicsPSOInit.BoundShaderState.VertexShaderRHI = VertexShader.GetVertexShader();
-					GraphicsPSOInit.BoundShaderState.PixelShaderRHI = PixelShader.GetPixelShader();
-					GraphicsPSOInit.PrimitiveType = PT_TriangleList;
-
-					SetGraphicsPipelineState(RHICmdList, GraphicsPSOInit, 0);
-						
-					PixelShader->SetParameters(RHICmdList, CubeFace, NumMips, SourceMipIndex, CoefficientIndex, EffectiveSource.ShaderResourceTexture);
-						
-					DrawRectangle(
-						RHICmdList,
-						ViewRect.Min.X, ViewRect.Min.Y,
-						ViewRect.Width(), ViewRect.Height(),
-						ViewRect.Min.X, ViewRect.Min.Y,
-						ViewRect.Width(), ViewRect.Height(),
-						FIntPoint(ViewRect.Width(), ViewRect.Height()),
-						FIntPoint(MipSize, MipSize),
-						VertexShader);
-						
-					RHICmdList.EndRenderPass();
-				}				
-
-				RHICmdList.Transition(FRHITransitionInfo(EffectiveRT.TargetableTexture, ERHIAccess::RTV, ERHIAccess::SRVGraphics));
+					auto* PassParameters = GraphBuilder.AllocParameters<FAccumulateDiffuseIrradiancePS::FParameters>();
+					PassParameters->SourceCubemapTexture = GraphBuilder.CreateSRV(FRDGTextureSRVDesc::CreateForMipLevel(IrradianceCubemapTexture, MipIndex - 1));
+					PassParameters->SourceCubemapSampler = TStaticSamplerState<SF_Point, AM_Clamp, AM_Clamp, AM_Clamp>::GetRHI();
+					PassParameters->Sample01 = Sample01;
+					PassParameters->Sample23 = Sample23;
+					PassParameters->SvPositionToUVScale = FVector2f(1.0f / MipSize, 1.0f / MipSize);
+					PassParameters->CubeFace = CubeFace;
+					PassParameters->SourceMipIndex = SourceMipIndex;
+					PassParameters->RenderTargets[0] = FRenderTargetBinding(IrradianceCubemapTexture, ERenderTargetLoadAction::ELoad, MipIndex, CubeFace);
+
+					FPixelShaderUtils::AddFullscreenPass(
+						GraphBuilder,
+						ShaderMap,
+						RDG_EVENT_NAME("Mip: %d, Face: %d", MipIndex, CubeFace),
+						AccumulateDiffuseIrradiancePS,
+						PassParameters,
+						ViewRect);
+				}
 			}
 		}
 
 		{
-			// Gather the cubemap face results and normalize, copy this coefficient to FSceneRenderTargets::Get().SkySHIrradianceMap
-			FSceneRenderTargetItem& EffectiveRT = GReflectionScratchCubemaps.SkySHIrradiance->GetRenderTargetItem();
-
-			//load/store actions so we don't lose results as we render one pixel at a time on tile renderers.
-			RHICmdList.Transition(FRHITransitionInfo(EffectiveRT.TargetableTexture, ERHIAccess::Unknown, ERHIAccess::RTV));
-			FRHIRenderPassInfo RPInfo(EffectiveRT.TargetableTexture, ERenderTargetActions::Load_Store, nullptr);
-			RHICmdList.BeginRenderPass(RPInfo, TEXT("GatherCoeffRP"));
-			RHICmdList.ApplyCachedRenderTargets(GraphicsPSOInit);
-
-			const FIntRect ViewRect(CoefficientIndex, 0, CoefficientIndex + 1, 1);
-			RHICmdList.SetViewport(0.0f, 0.0f, 0.0f, (float)FSHVector3::MaxSHBasis, 1.0f, 1.0f);
-
-			TShaderMapRef<FScreenVS> VertexShader(ShaderMap);
-			TShaderMapRef<FAccumulateCubeFacesPS> PixelShader(ShaderMap);
-
-			GraphicsPSOInit.BoundShaderState.VertexDeclarationRHI = GFilterVertexDeclaration.VertexDeclarationRHI;
-			GraphicsPSOInit.BoundShaderState.VertexShaderRHI = VertexShader.GetVertexShader();
-			GraphicsPSOInit.BoundShaderState.PixelShaderRHI = PixelShader.GetPixelShader();
-			GraphicsPSOInit.PrimitiveType = PT_TriangleList;
-
-			SetGraphicsPipelineState(RHICmdList, GraphicsPSOInit, 0);
-
 			const int32 SourceMipIndex = NumMips - 1;
 			const int32 MipSize = 1;
-			FSceneRenderTargetItem& EffectiveSource = GetEffectiveDiffuseIrradianceRenderTarget(SourceMipIndex);
-			PixelShader->SetParameters(RHICmdList, SourceMipIndex, EffectiveSource.ShaderResourceTexture);
-
-			DrawRectangle( 
-				RHICmdList,
-				ViewRect.Min.X, ViewRect.Min.Y, 
-				ViewRect.Width(), ViewRect.Height(),
-				0, 0, 
-				MipSize, MipSize,
-				FIntPoint(FSHVector3::MaxSHBasis, 1),
-				FIntPoint(MipSize, MipSize),
-				VertexShader);
-
-			RHICmdList.EndRenderPass();
-			RHICmdList.Transition(FRHITransitionInfo(EffectiveRT.TargetableTexture, ERHIAccess::RTV, ERHIAccess::SRVGraphics));
+			const FIntRect ViewRect(CoefficientIndex, 0, CoefficientIndex + 1, 1);
+
+			auto* PassParameters = GraphBuilder.AllocParameters<FAccumulateCubeFacesPS::FParameters>();
+			PassParameters->SourceCubemapTexture = IrradianceCubemapTexture;
+			PassParameters->SourceCubemapSampler = TStaticSamplerState<SF_Point, AM_Clamp, AM_Clamp, AM_Clamp>::GetRHI();
+			PassParameters->SourceMipIndex = SourceMipIndex;
+			PassParameters->RenderTargets[0] = FRenderTargetBinding(SkySHIrradianceTexture, ERenderTargetLoadAction::ELoad);
+
+			FPixelShaderUtils::AddFullscreenPass(
+				GraphBuilder,
+				ShaderMap,
+				RDG_EVENT_NAME("GatherCoefficients"),
+				AccumulateCubeFacesPS,
+				PassParameters,
+				ViewRect);
 		}
 	}
 
+	AddReadbackTexturePass(GraphBuilder, RDG_EVENT_NAME("ReadbackCoefficients"), SkySHIrradianceTexture, [SkySHIrradianceTexture, OutIrradianceEnvironmentMap](FRHICommandListImmediate& RHICmdList)
 	{
-		// Read back the completed SH environment map
-		FSceneRenderTargetItem& EffectiveRT = GReflectionScratchCubemaps.SkySHIrradiance->GetRenderTargetItem();
-		check(EffectiveRT.ShaderResourceTexture->GetFormat() == PF_FloatRGBA);
-
 		TArray<FFloat16Color> SurfaceData;
-		RHICmdList.ReadSurfaceFloatData(EffectiveRT.ShaderResourceTexture, FIntRect(0, 0, FSHVector3::MaxSHBasis, 1), SurfaceData, CubeFace_PosX, 0, 0);
+		RHICmdList.ReadSurfaceFloatData(SkySHIrradianceTexture->GetRHI(), FIntRect(0, 0, FSHVector3::MaxSHBasis, 1), SurfaceData, CubeFace_PosX, 0, 0);
 		check(SurfaceData.Num() == FSHVector3::MaxSHBasis);
 
 		for (int32 CoefficientIndex = 0; CoefficientIndex < FSHVector3::MaxSHBasis; CoefficientIndex++)
@@ -362,6 +221,5 @@ void ComputeDiffuseIrradiance(FRHICommandListImmediate& RHICmdList, ERHIFeatureL
 			OutIrradianceEnvironmentMap->G.V[CoefficientIndex] = CoefficientValue.G;
 			OutIrradianceEnvironmentMap->B.V[CoefficientIndex] = CoefficientValue.B;
 		}
-	}
-}
-
+	});
+}
\ No newline at end of file
diff --git a/Engine/Source/Runtime/Renderer/Private/ScenePrivate.h b/Engine/Source/Runtime/Renderer/Private/ScenePrivate.h
--- a/Engine/Source/Runtime/Renderer/Private/ScenePrivate.h
+++ b/Engine/Source/Runtime/Renderer/Private/ScenePrivate.h
@@ -1618,7 +1618,7 @@ public:
 	int32 GetMaxCubemaps() const { return MaxCubemaps; }
 	int32 GetCubemapSize() const { return CubemapSize; }
 	bool IsValid() const { return IsValidRef(ReflectionEnvs); }
-	FSceneRenderTargetItem& GetRenderTarget() const { return ReflectionEnvs->GetRenderTargetItem(); }
+	const TRefCountPtr<IPooledRenderTarget>& GetRenderTarget() const { return ReflectionEnvs; }
 
 protected:
 	uint32 MaxCubemaps;
