Fix lack of VolumetricFog in the right eye (UE-141871).

- Introduced by a fix for UE-125073 (second split-screen view influenced translucency in the first).

diff --git a/Engine/Source/Runtime/Renderer/Private/LightGridInjection.cpp b/Engine/Source/Runtime/Renderer/Private/LightGridInjection.cpp
--- a/Engine/Source/Runtime/Renderer/Private/LightGridInjection.cpp
+++ b/Engine/Source/Runtime/Renderer/Private/LightGridInjection.cpp
@@ -533,10 +533,13 @@ void FSceneRenderer::ComputeLightGrid(FRDGBuilder& GraphBuilder, bool bCullLight
 
 								for (const FProjectedShadowInfo* ShadowInfo : DirectionalLightShadowInfos)
 								{
-
-									if (ShadowInfo->DependentView && ShadowInfo->DependentView != &View)
+									if (ShadowInfo->DependentView)
 									{
-										continue;
+										// when rendering stereo views, allow using the shadows rendered for the primary view as 'close enough'
+										if (ShadowInfo->DependentView != &View && ShadowInfo->DependentView != View.GetPrimaryView())
+										{
+											continue;
+										}
 									}
 
 									const int32 CascadeIndex = ShadowInfo->CascadeSettings.ShadowSplitIndex;
