Explicitly enable 128 bit output on the SceneCapture pixel shader when the platform requires it.

diff --git a/Engine/Source/Runtime/Renderer/Private/SceneCaptureRendering.cpp b/Engine/Source/Runtime/Renderer/Private/SceneCaptureRendering.cpp
--- a/Engine/Source/Runtime/Renderer/Private/SceneCaptureRendering.cpp
+++ b/Engine/Source/Runtime/Renderer/Private/SceneCaptureRendering.cpp
@@ -67,9 +67,10 @@ public:
 	};
 
 	class FSourceModeDimension : SHADER_PERMUTATION_ENUM_CLASS("SOURCE_MODE", ESourceMode);
-	using FPermutationDomain = TShaderPermutationDomain<FSourceModeDimension>;
+	class FEnable128BitRT : SHADER_PERMUTATION_BOOL("ENABLE_128_BIT");
+	using FPermutationDomain = TShaderPermutationDomain<FSourceModeDimension, FEnable128BitRT>;
 
-	static FPermutationDomain GetPermutationVector(ESceneCaptureSource CaptureSource, bool bIsMobilePlatform)
+	static FPermutationDomain GetPermutationVector(ESceneCaptureSource CaptureSource, bool bUse128BitRT, bool bIsMobilePlatform)
 	{
 		ESourceMode SourceMode = ESourceMode::MAX;
 		switch (CaptureSource)
@@ -105,6 +106,7 @@ public:
 		}
 		FPermutationDomain PermutationVector;
 		PermutationVector.Set<FSourceModeDimension>(SourceMode);
+		PermutationVector.Set<FEnable128BitRT>(bUse128BitRT);
 		return PermutationVector;
 	}
 
@@ -113,7 +115,8 @@ public:
 		FPermutationDomain PermutationVector(Parameters.PermutationId);
 
 		auto SourceModeDim = PermutationVector.Get<FSourceModeDimension>();
-		return !IsMobilePlatform(Parameters.Platform) || (SourceModeDim != ESourceMode::Normal && SourceModeDim != ESourceMode::BaseColor);
+		bool bPlatformRequiresExplicit128bitRT = FDataDrivenShaderPlatformInfo::GetRequiresExplicit128bitRT(Parameters.Platform);
+		return (!PermutationVector.Get<FEnable128BitRT>() || bPlatformRequiresExplicit128bitRT) && (!IsMobilePlatform(Parameters.Platform) || (SourceModeDim != ESourceMode::Normal && SourceModeDim != ESourceMode::BaseColor));
 	}
 
 	static void ModifyCompilationEnvironment(const FGlobalShaderPermutationParameters& Parameters, FShaderCompilerEnvironment& OutEnvironment)
@@ -133,6 +136,11 @@ public:
 		const FPermutationDomain PermutationVector(Parameters.PermutationId);
 		const uint32 SourceModeIndex = static_cast<uint32>(PermutationVector.Get<FSourceModeDimension>());
 		OutEnvironment.SetDefine(ShaderSourceModeDefineName[SourceModeIndex], 1u);
+
+		if (PermutationVector.Get<FEnable128BitRT>())
+		{
+			OutEnvironment.SetRenderTargetOutputFormat(0, PF_A32B32G32R32F);
+		}
 	}
 };
 
@@ -200,7 +208,8 @@ void CopySceneCaptureComponentToTarget(
 			GraphicsPSOInit.BlendState = TStaticBlendState<>::GetRHI();
 		}
 
-		const FSceneCapturePS::FPermutationDomain PixelPermutationVector = FSceneCapturePS::GetPermutationVector(SceneCaptureSource, IsMobilePlatform(ViewFamily.GetShaderPlatform()));
+		const bool bUse128BitRT = PlatformRequires128bitRT(ViewFamilyTexture->Desc.Format);
+		const FSceneCapturePS::FPermutationDomain PixelPermutationVector = FSceneCapturePS::GetPermutationVector(SceneCaptureSource, bUse128BitRT, IsMobilePlatform(ViewFamily.GetShaderPlatform()));
 
 		for (int32 ViewIndex = 0; ViewIndex < Views.Num(); ViewIndex++)
 		{
