SLW - reworked rdg event names to follow our convention.

diff --git a/Engine/Source/Runtime/Renderer/Private/FogRendering.cpp b/Engine/Source/Runtime/Renderer/Private/FogRendering.cpp
--- a/Engine/Source/Runtime/Renderer/Private/FogRendering.cpp
+++ b/Engine/Source/Runtime/Renderer/Private/FogRendering.cpp
@@ -424,7 +424,7 @@ void FDeferredShadingSceneRenderer::RenderUnderWaterFog(
 		// Fog must be done in the base pass for MSAA to work
 		&& !IsForwardShadingEnabled(ShaderPlatform))
 	{
-		RDG_EVENT_SCOPE(GraphBuilder, "ExponentialHeightFog");
+		RDG_EVENT_SCOPE(GraphBuilder, "SLW::ExponentialHeightFog");
 		RDG_GPU_STAT_SCOPE(GraphBuilder, Fog);
 
 		FRDGTextureRef LinearDepthTexture = SceneWithoutWaterTextures.DepthTexture;
diff --git a/Engine/Source/Runtime/Renderer/Private/SingleLayerWaterRendering.cpp b/Engine/Source/Runtime/Renderer/Private/SingleLayerWaterRendering.cpp
--- a/Engine/Source/Runtime/Renderer/Private/SingleLayerWaterRendering.cpp
+++ b/Engine/Source/Runtime/Renderer/Private/SingleLayerWaterRendering.cpp
@@ -256,6 +256,8 @@ static FSceneWithoutWaterTextures AddCopySceneWithoutWaterPass(
 	FRDGTextureRef SceneColorTexture,
 	FRDGTextureRef SceneDepthTexture)
 {
+	RDG_EVENT_SCOPE(GraphBuilder, "SLW::CopySceneWithoutWater");
+
 	check(Views.Num() > 0);
 	check(SceneColorTexture);
 	check(SceneDepthTexture);
@@ -438,7 +440,7 @@ void FDeferredShadingSceneRenderer::RenderSingleLayerWaterReflections(
 				PassParameters->DispatchIndirectDataUAV = DispatchIndirectParametersBufferUAV;
 				PassParameters->WaterTileListDataUAV = TileListStructureBufferUAV;
 
-				FComputeShaderUtils::AddPass(GraphBuilder, RDG_EVENT_NAME("WaterTileCategorisation"), ComputeShader, PassParameters, TiledViewRes);
+				FComputeShaderUtils::AddPass(GraphBuilder, RDG_EVENT_NAME("SLW::TileCategorisation"), ComputeShader, PassParameters, TiledViewRes);
 			}
 		}
 
@@ -446,6 +448,8 @@ void FDeferredShadingSceneRenderer::RenderSingleLayerWaterReflections(
 
 		if (ViewPipelineState.ReflectionsMethod == EReflectionsMethod::Lumen && CVarWaterSingleLayerLumenReflections.GetValueOnRenderThread() != 0)
 		{
+			RDG_EVENT_SCOPE(GraphBuilder, "SLW::LumenReflections");
+
 			FLumenMeshSDFGridParameters MeshSDFGridParameters;
 			LumenRadianceCache::FRadianceCacheInterpolationParameters RadianceCacheParameters;
 			FLumenReflectionCompositeParameters LumenReflectionCompositeParameters;
@@ -466,7 +470,7 @@ void FDeferredShadingSceneRenderer::RenderSingleLayerWaterReflections(
 			&& CVarWaterSingleLayerRTR.GetValueOnRenderThread() != 0 
 			&& FDataDrivenShaderPlatformInfo::GetSupportsHighEndRayTracingReflections(View.GetShaderPlatform()))
 		{
-			RDG_EVENT_SCOPE(GraphBuilder, "RayTracingWaterReflections");
+			RDG_EVENT_SCOPE(GraphBuilder, "SLW::RayTracingReflections");
 			RDG_GPU_STAT_SCOPE(GraphBuilder, RayTracingWaterReflections);
 
 			IScreenSpaceDenoiser::FReflectionsInputs DenoiserInputs;
@@ -555,7 +559,7 @@ void FDeferredShadingSceneRenderer::RenderSingleLayerWaterReflections(
 			ESSRQuality SSRQuality;
 			ScreenSpaceRayTracing::GetSSRQualityForView(View, &SSRQuality, &RayTracingConfig);
 
-			RDG_EVENT_SCOPE(GraphBuilder, "Water ScreenSpaceReflections(Quality=%d)", int32(SSRQuality));
+			RDG_EVENT_SCOPE(GraphBuilder, "SLW::ScreenSpaceReflections(Quality=%d)", int32(SSRQuality));
 
 			const bool bDenoise = false;
 			const bool bSingleLayerWater = true;
@@ -616,7 +620,7 @@ void FDeferredShadingSceneRenderer::RenderSingleLayerWaterReflections(
 				ClearUnusedGraphResources(VertexShader, &PassParameters->VS);
 
 				GraphBuilder.AddPass(
-					RDG_EVENT_NAME("Water Composite %dx%d", View.ViewRect.Width(), View.ViewRect.Height()),
+					RDG_EVENT_NAME("SLW::Composite %dx%d", View.ViewRect.Width(), View.ViewRect.Height()),
 					PassParameters,
 					ERDGPassFlags::Raster,
 					[PassParameters, &View, TiledScreenSpaceReflection, VertexShader, PixelShader, bRunTiled](FRHICommandList& InRHICmdList)
@@ -643,7 +647,7 @@ void FDeferredShadingSceneRenderer::RenderSingleLayerWaterReflections(
 			else
 			{
 				GraphBuilder.AddPass(
-					RDG_EVENT_NAME("Water Composite %dx%d", View.ViewRect.Width(), View.ViewRect.Height()),
+					RDG_EVENT_NAME("SLW::Composite %dx%d", View.ViewRect.Width(), View.ViewRect.Height()),
 					PassParameters,
 					ERDGPassFlags::Raster,
 					[PassParameters, &View, TiledScreenSpaceReflection, PixelShader, bRunTiled](FRHICommandList& InRHICmdList)
@@ -716,8 +720,7 @@ void FDeferredShadingSceneRenderer::RenderSingleLayerWaterInner(
 	RDG_CSV_STAT_EXCLUSIVE_SCOPE(GraphBuilder, Water);
 	SCOPED_NAMED_EVENT(FDeferredShadingSceneRenderer_RenderSingleLayerWaterPass, FColor::Emerald);
 	SCOPE_CYCLE_COUNTER(STAT_WaterPassDrawTime);
-	RDG_EVENT_SCOPE(GraphBuilder, "SingleLayerWater");
-	RDG_GPU_STAT_SCOPE(GraphBuilder, SingleLayerWater);
+	RDG_EVENT_SCOPE(GraphBuilder, "SLW::Draw");
 
 	const bool bRenderInParallel = GRHICommandList.UseParallelAlgorithms() && CVarParallelSingleLayerWaterPass.GetValueOnRenderThread() == 1;
 
diff --git a/Engine/Source/Runtime/Renderer/Private/VolumetricRenderTarget.cpp b/Engine/Source/Runtime/Renderer/Private/VolumetricRenderTarget.cpp
--- a/Engine/Source/Runtime/Renderer/Private/VolumetricRenderTarget.cpp
+++ b/Engine/Source/Runtime/Renderer/Private/VolumetricRenderTarget.cpp
@@ -749,7 +749,7 @@ void ComposeVolumetricRenderTargetOverSceneUnderWater(
 		PassParameters->VolumetricTextureSizeAndInvSize = FVector4f(VolumetricTextureSize.X, VolumetricTextureSize.Y, 1.0f / VolumetricTextureSize.X, 1.0f / VolumetricTextureSize.Y);
 
 		FPixelShaderUtils::AddFullscreenPass<FComposeVolumetricRTOverScenePS>(
-			GraphBuilder, ViewInfo.ShaderMap, RDG_EVENT_NAME("VolumetricComposeOverScene"), PixelShader, PassParameters, WaterPassViewData.ViewRect,
+			GraphBuilder, ViewInfo.ShaderMap, RDG_EVENT_NAME("SLW::VolumetricComposeOverScene"), PixelShader, PassParameters, WaterPassViewData.ViewRect,
 			PreMultipliedColorTransmittanceBlend);
 	}
 }
