Fixed irradiance accumulation sampling with incorrect UV.

diff --git a/Engine/Source/Runtime/Renderer/Private/ReflectionEnvironmentDiffuseIrradiance.cpp b/Engine/Source/Runtime/Renderer/Private/ReflectionEnvironmentDiffuseIrradiance.cpp
--- a/Engine/Source/Runtime/Renderer/Private/ReflectionEnvironmentDiffuseIrradiance.cpp
+++ b/Engine/Source/Runtime/Renderer/Private/ReflectionEnvironmentDiffuseIrradiance.cpp
@@ -160,7 +160,7 @@ void ComputeDiffuseIrradiance(FRDGBuilder& GraphBuilder, FGlobalShaderMap* Shade
 				const int32 SourceMipIndex = FMath::Max(MipIndex - 1, 0);
 				const int32 SourceMipSize = 1 << (NumMips - SourceMipIndex - 1);
 
-				const float HalfSourceTexelSize = 0.5f / MipSize;
+				const float HalfSourceTexelSize = 0.5f / SourceMipSize;
 				const FVector4f Sample01(-HalfSourceTexelSize, -HalfSourceTexelSize, HalfSourceTexelSize, -HalfSourceTexelSize);
 				const FVector4f Sample23(-HalfSourceTexelSize, HalfSourceTexelSize, HalfSourceTexelSize, HalfSourceTexelSize);
 
