Fix too dark mobile preview (UE-145593)

- Contributed by [REDACTED_PERSON].

diff --git a/Engine/Source/Runtime/Renderer/Private/SceneTextures.cpp b/Engine/Source/Runtime/Renderer/Private/SceneTextures.cpp
--- a/Engine/Source/Runtime/Renderer/Private/SceneTextures.cpp
+++ b/Engine/Source/Runtime/Renderer/Private/SceneTextures.cpp
@@ -677,11 +677,13 @@ FSceneTextures& FMinimalSceneTextures::Create(FRDGBuilder& GraphBuilder, const F
 		const bool bIsMobilePlatform = Config.ShadingPath == EShadingPath::Mobile;
 
 		ETextureCreateFlags Flags = TexCreate_RenderTargetable | TexCreate_ShaderResource | GFastVRamConfig.SceneColor;
+		const ETextureCreateFlags sRGBFlag = (bIsMobilePlatform && IsMobileColorsRGB()) ? TexCreate_SRGB : TexCreate_None;
 
 		if (Config.FeatureLevel >= ERHIFeatureLevel::SM5 && Config.NumSamples == 1)
 		{
 			Flags |= TexCreate_UAV;
 		}
+		Flags |= sRGBFlag;
 
 		const TCHAR* SceneColorName = TEXT("SceneColor");
 
@@ -696,13 +698,7 @@ FSceneTextures& FMinimalSceneTextures::Create(FRDGBuilder& GraphBuilder, const F
 		if (Desc.NumSamples > 1)
 		{
 			Desc.NumSamples = 1;
-			Desc.Flags = TexCreate_ResolveTargetable | TexCreate_ShaderResource | GFastVRamConfig.SceneColor;
-
-			// Mobile non-mobileHDR is the only platform rendering to a true sRGB buffer natively
-			if (bIsMobilePlatform && IsMobileColorsRGB())
-			{
-				Desc.Flags |= TexCreate_SRGB;
-			}
+			Desc.Flags = TexCreate_ResolveTargetable | TexCreate_ShaderResource | GFastVRamConfig.SceneColor | sRGBFlag;
 
 			SceneTextures.Color.Resolve = GraphBuilder.CreateTexture(Desc, SceneColorName);
 		}
