Deprecating RHICreateTexture2DArray, RHICreateTexture3D, RHICreateTextureCube, RHICreateTextureCubeArray.
Converted a few additional RHICreateTexture2D calls to RHICreateTexture.
Adding a few more FRHITextureCreateDesc::CreateXXX variants to quickly specify extents and formats.
Merging FTextureRWBuffer2D and FTextureRWBuffer3D into a single class FTextureRWBuffer and deprecated the old ones.

diff --git a/Engine/Plugins/FX/Niagara/Source/Niagara/Classes/NiagaraDataInterfaceGrid2DCollection.h b/Engine/Plugins/FX/Niagara/Source/Niagara/Classes/NiagaraDataInterfaceGrid2DCollection.h
--- a/Engine/Plugins/FX/Niagara/Source/Niagara/Classes/NiagaraDataInterfaceGrid2DCollection.h
+++ b/Engine/Plugins/FX/Niagara/Source/Niagara/Classes/NiagaraDataInterfaceGrid2DCollection.h
@@ -18,8 +18,12 @@ class FGrid2DBuffer
 public:
 	FGrid2DBuffer(int NumX, int NumY, int NumAttributes, EPixelFormat PixelFormat)
 	{
-		FRHIResourceCreateInfo CreateInfo(TEXT("FGrid2DBuffer"));
-		GridTexture = RHICreateTexture2DArray(NumX, NumY, NumAttributes, PixelFormat, 1, 1, TexCreate_ShaderResource | TexCreate_UAV, CreateInfo);
+		const FRHITextureCreateDesc Desc =
+			FRHITextureCreateDesc::Create2DArray(TEXT("FGrid2DBuffer"), NumX, NumY, NumAttributes, PixelFormat)
+			.SetFlags(ETextureCreateFlags::ShaderResource | ETextureCreateFlags::UAV);
+
+		GridTexture = RHICreateTexture(Desc);
+
 		FRHITextureSRVCreateInfo SRVCreateInfo;
 		GridSRV = RHICreateShaderResourceView(GridTexture, SRVCreateInfo);
 		GridUAV = RHICreateUnorderedAccessView(GridTexture);
@@ -34,7 +38,7 @@ public:
 		GridUAV.SafeRelease();
 	}
 
-	FTexture2DArrayRHIRef GridTexture;
+	FTextureRHIRef GridTexture;
 	FShaderResourceViewRHIRef GridSRV;
 	FUnorderedAccessViewRHIRef GridUAV;
 };
diff --git a/Engine/Plugins/FX/Niagara/Source/Niagara/Classes/NiagaraDataInterfaceGrid3DCollection.h b/Engine/Plugins/FX/Niagara/Source/Niagara/Classes/NiagaraDataInterfaceGrid3DCollection.h
--- a/Engine/Plugins/FX/Niagara/Source/Niagara/Classes/NiagaraDataInterfaceGrid3DCollection.h
+++ b/Engine/Plugins/FX/Niagara/Source/Niagara/Classes/NiagaraDataInterfaceGrid3DCollection.h
@@ -18,7 +18,7 @@ class FGrid3DBuffer
 public:
 	FGrid3DBuffer(int NumX, int NumY, int NumZ, EPixelFormat PixelFormat)
 	{
-		GridBuffer.Initialize(TEXT("FGrid3DBuffer"), GPixelFormats[PixelFormat].BlockBytes, NumX, NumY, NumZ, PixelFormat);
+		GridBuffer.Initialize3D(TEXT("FGrid3DBuffer"), GPixelFormats[PixelFormat].BlockBytes, NumX, NumY, NumZ, PixelFormat);
 		INC_MEMORY_STAT_BY(STAT_NiagaraGPUDataInterfaceMemory, GridBuffer.NumBytes);
 	}
 
@@ -28,7 +28,7 @@ public:
 		GridBuffer.Release();
 	}
 
-	FTextureRWBuffer3D GridBuffer;	
+	FTextureRWBuffer GridBuffer;
 };
 
 struct FGrid3DCollectionRWInstanceData_GameThread
diff --git a/Engine/Plugins/FX/Niagara/Source/Niagara/Classes/NiagaraDataInterfaceRasterizationGrid3D.h b/Engine/Plugins/FX/Niagara/Source/Niagara/Classes/NiagaraDataInterfaceRasterizationGrid3D.h
--- a/Engine/Plugins/FX/Niagara/Source/Niagara/Classes/NiagaraDataInterfaceRasterizationGrid3D.h
+++ b/Engine/Plugins/FX/Niagara/Source/Niagara/Classes/NiagaraDataInterfaceRasterizationGrid3D.h
@@ -35,7 +35,7 @@ public:
 
 	bool NeedsRealloc = false;
 	
-	FTextureRWBuffer3D RasterizationBuffer;
+	FTextureRWBuffer RasterizationBuffer;
 
 	FReadBuffer PerAttributeData;
 
diff --git a/Engine/Plugins/FX/Niagara/Source/Niagara/Private/NiagaraDataInterfaceRasterizationGrid3D.cpp b/Engine/Plugins/FX/Niagara/Source/Niagara/Private/NiagaraDataInterfaceRasterizationGrid3D.cpp
--- a/Engine/Plugins/FX/Niagara/Source/Niagara/Private/NiagaraDataInterfaceRasterizationGrid3D.cpp
+++ b/Engine/Plugins/FX/Niagara/Source/Niagara/Private/NiagaraDataInterfaceRasterizationGrid3D.cpp
@@ -159,7 +159,7 @@ void RasterizationGrid3DRWInstanceData::ResizeBuffers()
 	if (NumTotalCells > (uint32)GMaxNiagaraRasterizationGridCells)
 		return;
 
-	RasterizationBuffer.Initialize(TEXT("NiagaraRasterizationGrid3D::IntGrid"), sizeof(int32), NumCells.X * NumTiles.X, NumCells.Y * NumTiles.Y, NumCells.Z * NumTiles.Z, EPixelFormat::PF_R32_SINT);
+	RasterizationBuffer.Initialize3D(TEXT("NiagaraRasterizationGrid3D::IntGrid"), sizeof(int32), NumCells.X * NumTiles.X, NumCells.Y * NumTiles.Y, NumCells.Z * NumTiles.Z, EPixelFormat::PF_R32_SINT);
 
 	#if STATS
 		DEC_MEMORY_STAT_BY(STAT_NiagaraGPUDataInterfaceMemory, GPUMemory);
diff --git a/Engine/Plugins/FX/Niagara/Source/Niagara/Private/NiagaraEmptyUAVPool.cpp b/Engine/Plugins/FX/Niagara/Source/Niagara/Private/NiagaraEmptyUAVPool.cpp
--- a/Engine/Plugins/FX/Niagara/Source/Niagara/Private/NiagaraEmptyUAVPool.cpp
+++ b/Engine/Plugins/FX/Niagara/Source/Niagara/Private/NiagaraEmptyUAVPool.cpp
@@ -34,7 +34,8 @@ FRHIUnorderedAccessView* FNiagaraEmptyUAVPool::GetEmptyUAVFromPool(FRHICommandLi
 		FEmptyUAV& NewUAV = Pool.UAVs.AddDefaulted_GetRef();
 
 		// Initialize the UAV
-		FRHIResourceCreateInfo CreateInfo(TEXT("FNiagaraGpuComputeDispatch::EmptyUAV"));
+		const TCHAR* ResourceName = TEXT("FNiagaraGpuComputeDispatch::EmptyUAV");
+		FRHIResourceCreateInfo CreateInfo(ResourceName);
 		switch ( Type )
 		{
 			case ENiagaraEmptyUAVType::Buffer:
@@ -47,21 +48,33 @@ FRHIUnorderedAccessView* FNiagaraEmptyUAVPool::GetEmptyUAVFromPool(FRHICommandLi
 			
 			case ENiagaraEmptyUAVType::Texture2D:
 			{
-				NewUAV.Texture = RHICreateTexture2D(1, 1, Format, 1, 1, TexCreate_ShaderResource | TexCreate_UAV, CreateInfo);
+				const FRHITextureCreateDesc Desc =
+					FRHITextureCreateDesc::Create2D(ResourceName, 1, 1, Format)
+					.SetFlags(ETextureCreateFlags::ShaderResource | ETextureCreateFlags::UAV);
+
+				NewUAV.Texture = RHICreateTexture(Desc);
 				NewUAV.UAV = RHICreateUnorderedAccessView(NewUAV.Texture, 0);
 				break;
 			}
 			
 			case ENiagaraEmptyUAVType::Texture2DArray:
 			{
-				NewUAV.Texture = RHICreateTexture2DArray(1, 1, 1, Format, 1, 1, TexCreate_ShaderResource | TexCreate_UAV, CreateInfo);
+				const FRHITextureCreateDesc Desc =
+					FRHITextureCreateDesc::Create2DArray(ResourceName, 1, 1, 1, Format)
+					.SetFlags(ETextureCreateFlags::ShaderResource | ETextureCreateFlags::UAV);
+
+				NewUAV.Texture = RHICreateTexture(Desc);
 				NewUAV.UAV = RHICreateUnorderedAccessView(NewUAV.Texture, 0);
 				break;
 			}
 			
 			case ENiagaraEmptyUAVType::Texture3D:
 			{
-				NewUAV.Texture = RHICreateTexture3D(1, 1, 1, Format, 1, TexCreate_ShaderResource | TexCreate_UAV, CreateInfo);
+				const FRHITextureCreateDesc Desc =
+					FRHITextureCreateDesc::Create3D(ResourceName, 1, 1, 1, Format)
+					.SetFlags(ETextureCreateFlags::ShaderResource | ETextureCreateFlags::UAV);
+
+				NewUAV.Texture = RHICreateTexture(Desc);
 				NewUAV.UAV = RHICreateUnorderedAccessView(NewUAV.Texture, 0);
 				break;
 			}
diff --git a/Engine/Plugins/FX/Niagara/Source/Niagara/Private/NiagaraGpuComputeDebug.cpp b/Engine/Plugins/FX/Niagara/Source/Niagara/Private/NiagaraGpuComputeDebug.cpp
--- a/Engine/Plugins/FX/Niagara/Source/Niagara/Private/NiagaraGpuComputeDebug.cpp
+++ b/Engine/Plugins/FX/Niagara/Source/Niagara/Private/NiagaraGpuComputeDebug.cpp
@@ -137,14 +137,7 @@ void FNiagaraGpuComputeDebug::AddAttributeTexture(FRHICommandList& RHICmdList, F
 		return;
 	}
 
-	FRHITexture2D* SrcTexture2D = Texture->GetTexture2D();
-	FRHITexture2DArray* SrcTexture2DArray = Texture->GetTexture2DArray();
-	FRHITexture3D* SrcTexture3D = Texture->GetTexture3D();
-	FRHITextureCube* SrcTextureCube = Texture->GetTextureCube();
-	if ( (SrcTexture2D == nullptr) && (SrcTexture2DArray == nullptr) && (SrcTexture3D == nullptr) && (SrcTextureCube == nullptr) )
-	{
-		return;
-	}
+	const FRHITextureDesc& TextureDesc = Texture->GetDesc();
 
 	bool bCreateTexture = false;
 
@@ -171,27 +164,18 @@ void FNiagaraGpuComputeDebug::AddAttributeTexture(FRHICommandList& RHICmdList, F
 	FTextureRHIRef Destination;
 	if ( bCreateTexture )
 	{
-		FRHIResourceCreateInfo CreateInfo(TEXT("FNiagaraGpuComputeDebug"));
-		if (SrcTexture2D != nullptr)
-		{
-			Destination = RHICreateTexture2D(SrcSize.X, SrcSize.Y, SrcFormat, 1, 1, TexCreate_ShaderResource, CreateInfo);
-			VisualizeEntry->Texture = Destination;
-		}
-		else if (SrcTexture2DArray != nullptr)
-		{
-			Destination = RHICreateTexture2DArray(SrcSize.X, SrcSize.Y, SrcSize.Z, SrcFormat, 1, 1, TexCreate_ShaderResource, CreateInfo);
-			VisualizeEntry->Texture = Destination;
-		}
-		else if (SrcTexture3D != nullptr)
-		{
-			Destination = RHICreateTexture3D(SrcSize.X, SrcSize.Y, SrcSize.Z, SrcFormat, 1, TexCreate_ShaderResource, CreateInfo);
-			VisualizeEntry->Texture = Destination;
-		}
-		else if (SrcTextureCube != nullptr)
-		{
-			Destination = RHICreateTextureCube(SrcSize.X, SrcFormat, 1, TexCreate_ShaderResource, CreateInfo);
-			VisualizeEntry->Texture = Destination;
-		}
+		// Create a minimal copy of the Texture's Desc
+		const FRHITextureCreateDesc NewTextureDesc =
+			FRHITextureCreateDesc(TEXT("FNiagaraGpuComputeDebug"), TextureDesc.Dimension)
+			.SetExtent(TextureDesc.Extent)
+			.SetDepth(TextureDesc.Depth)
+			.SetArraySize(TextureDesc.ArraySize)
+			.SetFormat(TextureDesc.Format)
+			.SetFlags(ETextureCreateFlags::ShaderResource)
+			.SetInitialState(ERHIAccess::SRVMask);
+
+		Destination = RHICreateTexture(NewTextureDesc);
+		VisualizeEntry->Texture = Destination;
 	}
 	else
 	{
diff --git a/Engine/Plugins/FX/Niagara/Source/Niagara/Private/NiagaraRenderer.cpp b/Engine/Plugins/FX/Niagara/Source/Niagara/Private/NiagaraRenderer.cpp
--- a/Engine/Plugins/FX/Niagara/Source/Niagara/Private/NiagaraRenderer.cpp
+++ b/Engine/Plugins/FX/Niagara/Source/Niagara/Private/NiagaraRenderer.cpp
@@ -81,51 +81,54 @@ public:
 class FNiagaraEmptyTextureSRV : public FRenderResource
 {
 public:
-	enum ETextureType
-	{
-		Texture2D,
-		Texture2DArray,
-		Texture3D
-	};
+	FNiagaraEmptyTextureSRV(EPixelFormat InPixelFormat, const FString& InDebugName, ETextureDimension InDimension)
+		: PixelFormat(InPixelFormat)
+		, DebugName(InDebugName)
+		, Dimension(InDimension)
+	{}
 
-	FNiagaraEmptyTextureSRV(EPixelFormat InPixelFormat, const FString& InDebugName, ETextureType InType) : PixelFormat(InPixelFormat), DebugName(InDebugName), Type(InType) {}
 	EPixelFormat PixelFormat;
 	FString DebugName;
-	ETextureType Type;
+	ETextureDimension Dimension;
 	FTextureRHIRef Texture;
 	FShaderResourceViewRHIRef SRV;
 
 	virtual void InitRHI() override
 	{
+		const FRHITextureCreateDesc Desc = FRHITextureCreateDesc(*DebugName, Dimension)
+			.SetExtent(1, 1)
+			.SetDepth(1)
+			.SetArraySize(1)
+			.SetFormat(PixelFormat)
+			.SetFlags(ETextureCreateFlags::ShaderResource);
+
 		// Create a 1x1 texture.
 		FRHIResourceCreateInfo CreateInfo(*DebugName);
 
 		uint32 Stride;
-		switch (Type)
+		switch (Dimension)
 		{
-			case Texture2D:
+			case ETextureDimension::Texture2D:
 			{
-				FTexture2DRHIRef Tex2D = RHICreateTexture2D(1, 1, PixelFormat, 1, 1, TexCreate_ShaderResource, CreateInfo);
-				void* Pixels = RHILockTexture2D(Tex2D, 0, RLM_WriteOnly, Stride, false);
+				Texture = RHICreateTexture(Desc);
+				void* Pixels = RHILockTexture2D(Texture, 0, RLM_WriteOnly, Stride, false);
 				FMemory::Memset(Pixels, 0, Stride);
-				RHIUnlockTexture2D(Tex2D, 0, 0, false);
-				Texture = Tex2D;
+				RHIUnlockTexture2D(Texture, 0, 0, false);
 				break;
 			}
 
-			case Texture2DArray:
+			case ETextureDimension::Texture2DArray:
 			{
-				FTexture2DArrayRHIRef Tex2DArray = RHICreateTexture2DArray(1, 1, 1, PixelFormat, 1, 1, TexCreate_ShaderResource, CreateInfo);
-				void* Pixels = RHILockTexture2DArray(Tex2DArray, 0, 0, RLM_WriteOnly, Stride, false);
+				Texture = RHICreateTexture(Desc);
+				void* Pixels = RHILockTexture2DArray(Texture, 0, 0, RLM_WriteOnly, Stride, false);
 				FMemory::Memset(Pixels, 0, Stride);
-				RHIUnlockTexture2DArray(Tex2DArray, 0, 0, false);
-				Texture = Tex2DArray;
+				RHIUnlockTexture2DArray(Texture, 0, 0, false);
 				break;
 			}
 
-			case Texture3D:
+			case ETextureDimension::Texture3D:
 			{
-				Texture = RHICreateTexture3D(1, 1, 1, PixelFormat, 1, TexCreate_ShaderResource, CreateInfo);
+				Texture = RHICreateTexture(Desc);
 				break;
 			}
 
@@ -204,21 +207,21 @@ FRHIShaderResourceView* FNiagaraRenderer::GetDummyUInt4Buffer()
 FRHIShaderResourceView* FNiagaraRenderer::GetDummyTextureReadBuffer2D()
 {
 	check(IsInRenderingThread());
-	static TGlobalResource<FNiagaraEmptyTextureSRV> DummyTextureReadBuffer2D(PF_R32_FLOAT, TEXT("NiagaraRenderer::DummyTextureReadBuffer2D"), FNiagaraEmptyTextureSRV::Texture2D);
+	static TGlobalResource<FNiagaraEmptyTextureSRV> DummyTextureReadBuffer2D(PF_R32_FLOAT, TEXT("NiagaraRenderer::DummyTextureReadBuffer2D"), ETextureDimension::Texture2D);
 	return DummyTextureReadBuffer2D.SRV;
 }
 
 FRHIShaderResourceView* FNiagaraRenderer::GetDummyTextureReadBuffer2DArray()
 {
 	check(IsInRenderingThread());
-	static TGlobalResource<FNiagaraEmptyTextureSRV> DummyTextureReadBuffer2DArray(PF_R32_FLOAT, TEXT("NiagaraRenderer::DummyTextureReadBuffer2DArray"), FNiagaraEmptyTextureSRV::Texture2DArray);
+	static TGlobalResource<FNiagaraEmptyTextureSRV> DummyTextureReadBuffer2DArray(PF_R32_FLOAT, TEXT("NiagaraRenderer::DummyTextureReadBuffer2DArray"), ETextureDimension::Texture2DArray);
 	return DummyTextureReadBuffer2DArray.SRV;
 }
 
 FRHIShaderResourceView* FNiagaraRenderer::GetDummyTextureReadBuffer3D()
 {
 	check(IsInRenderingThread());
-	static TGlobalResource<FNiagaraEmptyTextureSRV> DummyTextureReadBuffer3D(PF_R32_FLOAT, TEXT("NiagaraRenderer::DummyTextureReadBuffer3D"), FNiagaraEmptyTextureSRV::Texture3D);
+	static TGlobalResource<FNiagaraEmptyTextureSRV> DummyTextureReadBuffer3D(PF_R32_FLOAT, TEXT("NiagaraRenderer::DummyTextureReadBuffer3D"), ETextureDimension::Texture3D);
 	return DummyTextureReadBuffer3D.SRV;
 }
 
diff --git a/Engine/Plugins/Runtime/AR/AppleAR/AppleARKit/Source/AppleARKit/Private/AppleARKitTextures.cpp b/Engine/Plugins/Runtime/AR/AppleAR/AppleARKit/Source/AppleARKit/Private/AppleARKitTextures.cpp
--- a/Engine/Plugins/Runtime/AR/AppleAR/AppleARKit/Source/AppleARKit/Private/AppleARKitTextures.cpp
+++ b/Engine/Plugins/Runtime/AR/AppleAR/AppleARKit/Source/AppleARKit/Private/AppleARKitTextures.cpp
@@ -578,8 +578,11 @@ public:
 		{
 			Size.X = Size.Y = Owner->Size.X;
 
-			const ETextureCreateFlags CreateFlags = TexCreate_SRGB;
-			EnvCubemapTextureRHIRef = RHICreateTextureCube(Size.X, PF_B8G8R8A8, 1, CreateFlags, CreateInfo);
+			const FRHITextureCreateDesc Desc =
+				FRHITextureCreateDesc::CreateCube(TEXT("FARMetalResource"), Size.X, PF_B8G8R8A8)
+				.SetFlags(ETextureCreateFlags::SRGB);
+
+			EnvCubemapTextureRHIRef = RHICreateTexture(Desc);
 
 			/**
 			 * To map their texture faces into our space we need:
@@ -600,8 +603,12 @@ public:
 		else
 		{
 			Size.X = Size.Y = 1;
+
+			const FRHITextureCreateDesc Desc =
+				FRHITextureCreateDesc::CreateCube(TEXT("FARMetalResource"), Size.X, PF_B8G8R8A8);
+
 			// Start with a 1x1 texture
-			EnvCubemapTextureRHIRef = RHICreateTextureCube(Size.X, PF_B8G8R8A8, 1, TexCreate_None, CreateInfo);
+			EnvCubemapTextureRHIRef = RHICreateTexture(Desc);
 		}
 
 
diff --git a/Engine/Plugins/Tests/RHITests/Source/RHITests/Public/RHITextureTests.h b/Engine/Plugins/Tests/RHITests/Source/RHITests/Public/RHITextureTests.h
--- a/Engine/Plugins/Tests/RHITests/Source/RHITests/Public/RHITextureTests.h
+++ b/Engine/Plugins/Tests/RHITests/Source/RHITests/Public/RHITextureTests.h
@@ -92,16 +92,14 @@ public:
 		FString TestName = FString::Printf(TEXT("Test_RHIClearUAV_Texture2D (%dx%d, %d Slice(s), %d Mip(s)) - %s"), Width, Height, NumMips, NumSlices, *ClearValueToString(ClearValue));
 
 		{
-			FRHIResourceCreateInfo CreateInfo(*TestName);
-			FTextureRHIRef Texture;
-			if (NumSlices == 1)
-			{
-				Texture = RHICreateTexture2D(Width, Height, Format, NumMips, 1, TexCreate_UAV | TexCreate_ShaderResource, CreateInfo);
-			}
-			else
-			{
-				Texture = RHICreateTexture2DArray(Width, Height, NumSlices, Format, NumMips, 1, TexCreate_UAV | TexCreate_ShaderResource, CreateInfo);
-			}
+			const FRHITextureCreateDesc Desc =
+				FRHITextureCreateDesc(*TestName, (NumSlices == 1) ? ETextureDimension::Texture2D : ETextureDimension::Texture2DArray)
+				.SetExtent(Width, Height)
+				.SetArraySize(NumSlices)
+				.SetNumMips(NumMips)
+				.SetFlags(ETextureCreateFlags::UAV | ETextureCreateFlags::ShaderResource);
+
+			FTextureRHIRef Texture = RHICreateTexture(Desc);
 
 			for (uint32 Mip = 0; Mip < NumMips; ++Mip)
 			{
@@ -173,8 +171,12 @@ public:
 		bool bResult = true;
 
 		{
-			FRHIResourceCreateInfo CreateInfo(*TestName);
-			FTexture3DRHIRef Texture = RHICreateTexture3D(Width, Height, Depth, Format, NumMips, TexCreate_UAV | TexCreate_ShaderResource, CreateInfo);
+			const FRHITextureCreateDesc Desc =
+				FRHITextureCreateDesc::Create3D(*TestName, Width, Height, Depth, Format)
+				.SetNumMips(NumMips)
+				.SetFlags(ETextureCreateFlags::UAV | ETextureCreateFlags::ShaderResource);
+
+			FTexture3DRHIRef Texture = RHICreateTexture(Desc);
 
 			for (uint32 Mip = 0; Mip < NumMips; ++Mip)
 			{
@@ -241,10 +243,11 @@ public:
 		bool bResult = true;
 		FString TestName = FString::Printf(TEXT("Test_RHIFormat (%s, %s, %s, %d)"), GPixelFormats[ResourceFormat].Name, GPixelFormats[SRVFormat].Name, GPixelFormats[UAVFormat].Name, Flags);
 		{
-			FRHIResourceCreateInfo CreateInfo(*TestName);
+			const FRHITextureCreateDesc Desc =
+				FRHITextureCreateDesc::Create2D(*TestName, Width, Height, ResourceFormat)
+				.SetFlags(Flags);
 
-			FTextureRHIRef Texture;
-			Texture = RHICreateTexture2D(Width, Height, ResourceFormat, 1, 1, Flags, CreateInfo);
+			FTextureRHIRef Texture = RHICreateTexture(Desc);
 			bResult = (Texture != nullptr);
 
 			if (Texture && SRVFormat != PF_Unknown)
@@ -519,11 +522,11 @@ public:
 				return true;
 			}
 
-			FRHIResourceCreateInfo CreateInfo(*TestName);
+			const FRHITextureCreateDesc Desc =
+				FRHITextureCreateDesc::Create2D(*TestName, TextureWidth, TextureHeight, Format)
+				.SetFlags(ETextureCreateFlags::ShaderResource | ETextureCreateFlags::UAV);
 
-			FTexture2DRHIRef Texture;
-			ETextureCreateFlags TexCreateFlags = TexCreate_ShaderResource | TexCreate_UAV;
-			Texture = RHICreateTexture2D(TextureWidth, TextureHeight, Format, 1, 1, TexCreateFlags, CreateInfo);
+			FTexture2DRHIRef Texture = RHICreateTexture(Desc);
 			if (Texture == nullptr)
 			{
 				UE_LOG(LogRHIUnitTestCommandlet, Display, TEXT("Test failed (couldn't create texture). \"%s\""), *TestName);
diff --git a/Engine/Source/Runtime/Engine/Private/Components/SkyLightComponent.cpp b/Engine/Source/Runtime/Engine/Private/Components/SkyLightComponent.cpp
--- a/Engine/Source/Runtime/Engine/Private/Components/SkyLightComponent.cpp
+++ b/Engine/Source/Runtime/Engine/Private/Components/SkyLightComponent.cpp
@@ -103,10 +103,13 @@ void FSkyTextureCubeResource::InitRHI()
 {
 	if (GetFeatureLevel() >= ERHIFeatureLevel::SM5 || GSupportsRenderTargetFormat_PF_FloatRGBA)
 	{
-		FRHIResourceCreateInfo CreateInfo(TEXT("SkyTextureCube"));
-		
 		checkf(FMath::IsPowerOfTwo(Size), TEXT("Size of SkyTextureCube must be a power of two; size is %d"), Size);
-		TextureCubeRHI = RHICreateTextureCube(Size, Format, NumMips, TexCreate_None, CreateInfo);
+
+		const FRHITextureCreateDesc Desc =
+			FRHITextureCreateDesc::CreateCube(TEXT("SkyTextureCube"), Size, Format)
+			.SetNumMips(NumMips);
+
+		TextureCubeRHI = RHICreateTexture(Desc);
 		TextureRHI = TextureCubeRHI;
 
 		// Create the sampler state RHI resource.
diff --git a/Engine/Source/Runtime/Engine/Private/PhysicsField/PhysicsFieldComponent.cpp b/Engine/Source/Runtime/Engine/Private/PhysicsField/PhysicsFieldComponent.cpp
--- a/Engine/Source/Runtime/Engine/Private/PhysicsField/PhysicsFieldComponent.cpp
+++ b/Engine/Source/Runtime/Engine/Private/PhysicsField/PhysicsFieldComponent.cpp
@@ -152,37 +152,6 @@ void UpdateInternalBuffer(const uint32 ElementCount, const BufferType* InputData
 	}
 }
 
-template<typename BufferType, int ElementSize, EPixelFormat PixelFormat>
-void InitInternalTexture(const uint32 SizeX, const uint32 SizeY, const uint32 SizeZ, FTextureRWBuffer3D& OutputBuffer)
-{
-	if (SizeX * SizeY * SizeZ > 0)
-	{
-		const uint32 BlockBytes = sizeof(BufferType) * ElementSize;
-
-		OutputBuffer.Initialize(TEXT("FPhysicsFieldResource"), BlockBytes, SizeX, SizeY, SizeZ, PixelFormat);
-
-		if (OutputBuffer.UAV)
-		{
-			FRHICommandListExecutor::GetImmediateCommandList().ClearUAVFloat(OutputBuffer.UAV, FVector4f(ForceInitToZero));
-		}
-	}
-}
-
-template<typename BufferType, int ElementSize, EPixelFormat PixelFormat>
-void UpdateInternalTexture(const uint32 SizeX, const uint32 SizeY, const uint32 SizeZ, const BufferType* InputData, FTextureRWBuffer3D& OutputBuffer)
-{
-	if (SizeX * SizeY * SizeZ > 0 && InputData)
-	{
-		const uint32 BlockBytes = sizeof(BufferType) * ElementSize;
-
-		FUpdateTextureRegion3D UpdateRegion(0, 0, 0, 0, 0, 0, SizeX, SizeY, SizeZ);
-
-		const uint8* TextureDatas = (const uint8*)InputData;
-		RHIUpdateTexture3D(OutputBuffer.Buffer, 0, UpdateRegion, SizeX * BlockBytes,
-			SizeX * SizeY * BlockBytes, TextureDatas);
-	}
-}
-
 FVector MinVector(const FVector& VectorA, const FVector& VectorB)
 {
 	return FVector(FMath::Min(VectorA.X, VectorB.X), FMath::Min(VectorA.Y, VectorB.Y), FMath::Min(VectorA.Z, VectorB.Z));
diff --git a/Engine/Source/Runtime/Engine/Private/PrecomputedVolumetricLightmap.cpp b/Engine/Source/Runtime/Engine/Private/PrecomputedVolumetricLightmap.cpp
--- a/Engine/Source/Runtime/Engine/Private/PrecomputedVolumetricLightmap.cpp
+++ b/Engine/Source/Runtime/Engine/Private/PrecomputedVolumetricLightmap.cpp
@@ -21,31 +21,21 @@ DECLARE_MEMORY_STAT(TEXT("Volumetric Lightmap"),STAT_VolumetricLightmapBuildData
 
 void FVolumetricLightmapDataLayer::CreateTexture(FIntVector Dimensions)
 {
-	FRHIResourceCreateInfo CreateInfo(TEXT("VolumetricLightmap"));
-	CreateInfo.BulkData = this;
-
-	Texture = RHICreateTexture3D(
-		Dimensions.X, 
-		Dimensions.Y, 
-		Dimensions.Z, 
-		Format,
-		1,
-		TexCreate_ShaderResource | TexCreate_UAV,
-		CreateInfo);
+	const FRHITextureCreateDesc Desc =
+		FRHITextureCreateDesc::Create3D(TEXT("VolumetricLightmap"), Dimensions, Format)
+		.SetFlags(ETextureCreateFlags::ShaderResource | ETextureCreateFlags::UAV)
+		.SetBulkData(this);
+
+	Texture = RHICreateTexture(Desc);
 }
 
 void FVolumetricLightmapDataLayer::CreateTargetTexture(FIntVector Dimensions)
 {
-	FRHIResourceCreateInfo CreateInfo(TEXT("VolumetricLightmap"));
-
-	Texture = RHICreateTexture3D(
-		Dimensions.X,
-		Dimensions.Y,
-		Dimensions.Z,
-		Format,
-		1,
-		TexCreate_ShaderResource | TexCreate_UAV,
-		CreateInfo);
+	const FRHITextureCreateDesc Desc =
+		FRHITextureCreateDesc::Create3D(TEXT("VolumetricLightmap"), Dimensions, Format)
+		.SetFlags(ETextureCreateFlags::ShaderResource | ETextureCreateFlags::UAV);
+
+	Texture = RHICreateTexture(Desc);
 }
 
 void FVolumetricLightmapDataLayer::CreateUAV()
diff --git a/Engine/Source/Runtime/Engine/Private/Rendering/Texture2DArrayResource.cpp b/Engine/Source/Runtime/Engine/Private/Rendering/Texture2DArrayResource.cpp
--- a/Engine/Source/Runtime/Engine/Private/Rendering/Texture2DArrayResource.cpp
+++ b/Engine/Source/Runtime/Engine/Private/Rendering/Texture2DArrayResource.cpp
@@ -73,10 +73,13 @@ void FTexture2DArrayResource::CreateTexture()
 	TArrayView<const FTexture2DMipMap*> MipsView = GetPlatformMipsView();
 	const FTexture2DMipMap& FirstMip = *MipsView[RequestedFirstLODIdx];
 
-	FRHIResourceCreateInfo CreateInfo(TEXT("FTexture2DArrayResource"));
-	CreateInfo.ExtData = PlatformData->GetExtData();
-	TRefCountPtr<FRHITexture2DArray> TextureArray = RHICreateTexture2DArray(FirstMip.SizeX, FirstMip.SizeY, FirstMip.SizeZ, PixelFormat, State.NumRequestedLODs, 1, CreationFlags, CreateInfo);
-	TextureRHI = TextureArray;
+	const FRHITextureCreateDesc Desc =
+		FRHITextureCreateDesc::Create2DArray(TEXT("FTexture2DArrayResource"), FirstMip.SizeX, FirstMip.SizeY, FirstMip.SizeZ, PixelFormat)
+		.SetNumMips(State.NumRequestedLODs)
+		.SetFlags(CreationFlags)
+		.SetExtData(PlatformData->GetExtData());
+
+	TextureRHI = RHICreateTexture(Desc);
 
 	// Read the initial cached mip levels into the RHI texture.
 	const int32 NumLoadableMips = State.NumRequestedLODs - FMath::Max(1, PlatformData->GetNumMipsInTail()) + 1;
@@ -86,12 +89,12 @@ void FTexture2DArrayResource::CreateTexture()
 		for (uint32 SliceIdx = 0; SliceIdx < SizeZ; ++SliceIdx)
 		{
 			uint32 DestStride = 0;
-			void* DestData = RHILockTexture2DArray(TextureArray, SliceIdx, RHIMipIdx, RLM_WriteOnly, DestStride, false);
+			void* DestData = RHILockTexture2DArray(TextureRHI, SliceIdx, RHIMipIdx, RLM_WriteOnly, DestStride, false);
 			if (DestData)
 			{
 				GetData(SliceIdx, MipIdx, DestData, DestStride);
 			}
-			RHIUnlockTexture2DArray(TextureArray, SliceIdx, RHIMipIdx, false);
+			RHIUnlockTexture2DArray(TextureRHI, SliceIdx, RHIMipIdx, false);
 		}
 	}
 		
diff --git a/Engine/Source/Runtime/Engine/Private/Rendering/Texture3DResource.cpp b/Engine/Source/Runtime/Engine/Private/Rendering/Texture3DResource.cpp
--- a/Engine/Source/Runtime/Engine/Private/Rendering/Texture3DResource.cpp
+++ b/Engine/Source/Runtime/Engine/Private/Rendering/Texture3DResource.cpp
@@ -112,21 +112,24 @@ void FTexture3DResource::CreateTexture()
 {
 	TArrayView<const FTexture2DMipMap*> MipsView = GetPlatformMipsView();
 	const int32 FirstMipIdx = InitialData.GetFirstMipIdx(); // == State.RequestedFirstLODIdx()
-	FTexture3DRHIRef Texture3DRHI;
 
 	// Create the RHI texture.
 	{
-		FRHIResourceCreateInfo CreateInfo(TEXT("FTexture3DResource"));
+		const FTexture2DMipMap& FirstMip = *MipsView[FirstMipIdx];
+
+		FRHITextureCreateDesc Desc =
+			FRHITextureCreateDesc::Create3D(TEXT("FTexture3DResource"), FirstMip.SizeX, FirstMip.SizeY, FirstMip.SizeZ, PixelFormat)
+			.SetNumMips(State.NumRequestedLODs)
+			.SetFlags(CreationFlags)
+			.SetExtData(PlatformData->GetExtData());
+
 		if (GUseTexture3DBulkDataRHI)
 		{
 			InitialData.MergeMips(State.MaxNumLODs);
-			CreateInfo.BulkData = &InitialData;
+			Desc.SetBulkData(&InitialData);
 		}
 
-		const FTexture2DMipMap& FirstMip = *MipsView[FirstMipIdx];
-		CreateInfo.ExtData = PlatformData->GetExtData();
-		Texture3DRHI = RHICreateTexture3D(FirstMip.SizeX, FirstMip.SizeY, FirstMip.SizeZ, PixelFormat, State.NumRequestedLODs, CreationFlags, CreateInfo);
-		TextureRHI = Texture3DRHI; 
+		TextureRHI = RHICreateTexture(Desc);
 	}
 
 	if (!GUseTexture3DBulkDataRHI) 
@@ -162,7 +165,7 @@ void FTexture3DResource::CreateTexture()
 					(int)MipBytes);
 				}
 
-				RHIUpdateTexture3D(Texture3DRHI, RHIMipIdx, UpdateRegion, NumBlockX * BlockBytes, NumBlockX * NumBlockY * BlockBytes, MipData);
+				RHIUpdateTexture3D(TextureRHI, RHIMipIdx, UpdateRegion, NumBlockX * BlockBytes, NumBlockX * NumBlockY * BlockBytes, MipData);
 			}
 		}
 		InitialData.Discard();
diff --git a/Engine/Source/Runtime/Engine/Private/Streaming/Texture2DArrayStreaming.cpp b/Engine/Source/Runtime/Engine/Private/Streaming/Texture2DArrayStreaming.cpp
--- a/Engine/Source/Runtime/Engine/Private/Streaming/Texture2DArrayStreaming.cpp
+++ b/Engine/Source/Runtime/Engine/Private/Streaming/Texture2DArrayStreaming.cpp
@@ -104,10 +104,15 @@ bool FTexture2DArrayMipAllocator_Reallocate::FinalizeMips(const FTextureUpdateCo
 
 	// Create new Texture.
 	{
-		FRHIResourceCreateInfo CreateInfo(TEXT("FinalizeMips"));
 		const FTexture2DMipMap& FirstMip = *Context.MipsView[PendingFirstLODIdx];
-		CreateInfo.ExtData = Context.Resource->GetExtData();
-		IntermediateTextureRHI = RHICreateTexture2DArray(FirstMip.SizeX, FirstMip.SizeY, FirstMip.SizeZ, Context.Resource->GetPixelFormat(), ResourceState.NumRequestedLODs, 1, Context.Resource->GetCreationFlags(), CreateInfo);
+
+		const FRHITextureCreateDesc Desc =
+			FRHITextureCreateDesc::Create2DArray(TEXT("FinalizeMips"), FirstMip.SizeX, FirstMip.SizeY, FirstMip.SizeZ, Context.Resource->GetPixelFormat())
+			.SetNumMips(ResourceState.NumRequestedLODs)
+			.SetFlags(Context.Resource->GetCreationFlags())
+			.SetExtData(Context.Resource->GetExtData());
+
+		IntermediateTextureRHI = RHICreateTexture(Desc);
 	}
 
 	// Copy shared mips.
diff --git a/Engine/Source/Runtime/Engine/Private/Streaming/Texture2DArrayStreaming.h b/Engine/Source/Runtime/Engine/Private/Streaming/Texture2DArrayStreaming.h
--- a/Engine/Source/Runtime/Engine/Private/Streaming/Texture2DArrayStreaming.h
+++ b/Engine/Source/Runtime/Engine/Private/Streaming/Texture2DArrayStreaming.h
@@ -51,7 +51,7 @@ protected:
 	};
 
 	// The intermediate texture async created in the update process.
-	FTexture2DArrayRHIRef IntermediateTextureRHI;
+	FTextureRHIRef IntermediateTextureRHI;
 
 	// The temporary main memory allocations holding the mip data
 	TArray<TUniquePtr<uint8, FMallocDeleter>, TInlineAllocator<MAX_TEXTURE_MIP_COUNT>> StreamedMipData;
diff --git a/Engine/Source/Runtime/Engine/Private/Streaming/VolumeTextureStreaming.cpp b/Engine/Source/Runtime/Engine/Private/Streaming/VolumeTextureStreaming.cpp
--- a/Engine/Source/Runtime/Engine/Private/Streaming/VolumeTextureStreaming.cpp
+++ b/Engine/Source/Runtime/Engine/Private/Streaming/VolumeTextureStreaming.cpp
@@ -73,15 +73,20 @@ bool FVolumeTextureMipAllocator_Reallocate::FinalizeMips(const FTextureUpdateCon
 
 	// Create new Texture.
 	{
-		FRHIResourceCreateInfo CreateInfo(TEXT("FinalizeMips"));
+		const FTexture2DMipMap& FirstMip = *Context.MipsView[PendingFirstLODIdx];
+
+		FRHITextureCreateDesc Desc =
+			FRHITextureCreateDesc::Create3D(TEXT("FinalizeMips"), FirstMip.SizeX, FirstMip.SizeY, FirstMip.SizeZ, PixelFormat)
+			.SetNumMips(ResourceState.NumRequestedLODs)
+			.SetFlags(Context.Resource->GetCreationFlags())
+			.SetExtData(Context.Resource->GetExtData());
+
 		if (GUseTexture3DBulkDataRHI)
 		{
-			CreateInfo.BulkData = &StreamedInMipData;
+			Desc.SetBulkData(&StreamedInMipData);
 		}
 
-		const FTexture2DMipMap& FirstMip = *Context.MipsView[PendingFirstLODIdx];
-		CreateInfo.ExtData = Context.Resource->GetExtData();
-		IntermediateTextureRHI = RHICreateTexture3D(FirstMip.SizeX, FirstMip.SizeY, FirstMip.SizeZ, PixelFormat, ResourceState.NumRequestedLODs, Context.Resource->GetCreationFlags(), CreateInfo);
+		IntermediateTextureRHI = RHICreateTexture(Desc);
 	}
 
 	// Copy shared mips.
diff --git a/Engine/Source/Runtime/Engine/Private/TextureCube.cpp b/Engine/Source/Runtime/Engine/Private/TextureCube.cpp
--- a/Engine/Source/Runtime/Engine/Private/TextureCube.cpp
+++ b/Engine/Source/Runtime/Engine/Private/TextureCube.cpp
@@ -438,11 +438,19 @@ public:
 		INC_DWORD_STAT_FNAME_BY( LODGroupStatName, TextureSize );
 
 		// Create the RHI texture.
-		ETextureCreateFlags TexCreateFlags = (Owner->SRGB ? TexCreate_SRGB : TexCreate_None)  | (Owner->bNotOfflineProcessed ? TexCreate_None : TexCreate_OfflineProcessed);
-		FString Name = Owner->GetPathName();
-		FRHIResourceCreateInfo CreateInfo(*Name);
-		CreateInfo.ExtData = Owner->GetPlatformData() ? Owner->GetPlatformData()->GetExtData() : 0;
-		TextureCubeRHI = RHICreateTextureCube( Owner->GetSizeX(), Owner->GetPixelFormat(), Owner->GetNumMips(), TexCreateFlags, CreateInfo );
+		const ETextureCreateFlags TexCreateFlags = (Owner->SRGB ? TexCreate_SRGB : TexCreate_None)  | (Owner->bNotOfflineProcessed ? TexCreate_None : TexCreate_OfflineProcessed);
+		const FString Name = Owner->GetPathName();
+
+		const FRHITextureCreateDesc Desc =
+			FRHITextureCreateDesc::CreateCube(*Name)
+			.SetExtent(Owner->GetSizeX())
+			.SetFormat(Owner->GetPixelFormat())
+			.SetNumMips(Owner->GetNumMips())
+			.SetFlags(TexCreateFlags)
+			.SetExtData(Owner->GetPlatformData() ? Owner->GetPlatformData()->GetExtData() : 0);
+
+		TextureCubeRHI = RHICreateTexture(Desc);
+
 		TextureRHI = TextureCubeRHI;
 		TextureRHI->SetName(Owner->GetFName());
 		RHIBindDebugLabelName(TextureRHI, *Name);
diff --git a/Engine/Source/Runtime/Engine/Private/TextureCubeArray.cpp b/Engine/Source/Runtime/Engine/Private/TextureCubeArray.cpp
--- a/Engine/Source/Runtime/Engine/Private/TextureCubeArray.cpp
+++ b/Engine/Source/Runtime/Engine/Private/TextureCubeArray.cpp
@@ -113,13 +113,22 @@ public:
 		check(Owner->GetNumMips() > 0);
 
 		// Create the RHI texture.
-		ETextureCreateFlags TexCreateFlags = (Owner->SRGB ? TexCreate_SRGB : TexCreate_None) | (Owner->bNotOfflineProcessed ? TexCreate_None : TexCreate_OfflineProcessed);
-		FString Name = Owner->GetPathName();
-		FRHIResourceCreateInfo CreateInfo(*Name);
-		CreateInfo.ExtData = Owner->PlatformData->GetExtData();
-		TextureCubeRHI = RHICreateTextureCubeArray(Owner->GetSizeX(), ArraySize, Owner->GetPixelFormat(), Owner->GetNumMips(), TexCreateFlags, CreateInfo);
+		const ETextureCreateFlags TexCreateFlags = (Owner->SRGB ? ETextureCreateFlags::SRGB : ETextureCreateFlags::None) | (Owner->bNotOfflineProcessed ? ETextureCreateFlags::None : ETextureCreateFlags::OfflineProcessed);
+		const FString Name = Owner->GetPathName();
+
+		const FRHITextureCreateDesc Desc =
+			FRHITextureCreateDesc::CreateCubeArray(*Name)
+			.SetExtent(Owner->GetSizeX())
+			.SetArraySize(ArraySize)
+			.SetFormat(Owner->GetPixelFormat())
+			.SetNumMips(Owner->GetNumMips())
+			.SetFlags(TexCreateFlags)
+			.SetExtData(Owner->PlatformData->GetExtData());
+
+		TextureCubeRHI = RHICreateTexture(Desc);
 		TextureRHI = TextureCubeRHI;
 		TextureRHI->SetName(Owner->GetFName());
+
 		RHIBindDebugLabelName(TextureRHI, *Name);
 		RHIUpdateTextureReference(Owner->TextureReference.TextureReferenceRHI, TextureRHI);
 
diff --git a/Engine/Source/Runtime/Engine/Private/VectorField.cpp b/Engine/Source/Runtime/Engine/Private/VectorField.cpp
--- a/Engine/Source/Runtime/Engine/Private/VectorField.cpp
+++ b/Engine/Source/Runtime/Engine/Private/VectorField.cpp
@@ -208,12 +208,14 @@ public:
 		{
 			const uint32 DataSize = SizeX * SizeY * SizeZ * sizeof(FFloat16Color);
 			FVectorFieldStaticResourceBulkDataInterface BulkDataInterface(VolumeData, DataSize);
-			FRHIResourceCreateInfo CreateInfo(TEXT("FVectorFieldStaticResource"), &BulkDataInterface);
-			VolumeTextureRHI = RHICreateTexture3D(
-				SizeX, SizeY, SizeZ, PF_FloatRGBA,
-				/*NumMips=*/ 1,
-				/*Flags=*/ TexCreate_ShaderResource,
-				/*BulkData=*/ CreateInfo );
+
+			const FRHITextureCreateDesc Desc =
+				FRHITextureCreateDesc::Create3D(TEXT("FVectorFieldStaticResource"), SizeX, SizeY, SizeZ, PF_FloatRGBA)
+				.SetFlags(ETextureCreateFlags::ShaderResource)
+				.SetBulkData(&BulkDataInterface);
+
+			VolumeTextureRHI = RHICreateTexture(Desc);
+
 			FMemory::Free(VolumeData);
 			VolumeData = NULL;
 		}
@@ -1096,13 +1098,11 @@ public:
 				TexCreateFlags = TexCreate_ShaderResource | TexCreate_UAV;
 			}
 
-			FRHIResourceCreateInfo CreateInfo(TEXT("FVectorFieldAnimatedResource"));
-			VolumeTextureRHI = RHICreateTexture3D(
-				SizeX, SizeY, SizeZ,
-				PF_FloatRGBA,
-				/*NumMips=*/ 1,
-				TexCreateFlags,
-				CreateInfo);
+			const FRHITextureCreateDesc Desc =
+				FRHITextureCreateDesc::Create3D(TEXT("FVectorFieldAnimatedResource"), SizeX, SizeY, SizeZ, PF_FloatRGBA)
+				.SetFlags(TexCreateFlags);
+
+			VolumeTextureRHI = RHICreateTexture(Desc);
 
 			if (GetFeatureLevel() >= ERHIFeatureLevel::SM5)
 			{
diff --git a/Engine/Source/Runtime/Landscape/Private/LandscapeEditResources.cpp b/Engine/Source/Runtime/Landscape/Private/LandscapeEditResources.cpp
--- a/Engine/Source/Runtime/Landscape/Private/LandscapeEditResources.cpp
+++ b/Engine/Source/Runtime/Landscape/Private/LandscapeEditResources.cpp
@@ -113,7 +113,12 @@ void FLandscapeTexture2DArrayResource::InitRHI()
 		Flags |= TexCreate_ShaderResource;
 	}
 
-	TextureRHI = RHICreateTexture2DArray(SizeX, SizeY, SizeZ, Format, NumMips, 1, Flags, CreateInfo);
+	const FRHITextureCreateDesc Desc =
+		FRHITextureCreateDesc::Create2DArray(TEXT("FLandscapeTexture2DArrayResource"), SizeX, SizeY, SizeZ, Format)
+		.SetNumMips(NumMips)
+		.SetFlags(Flags);
+
+	TextureRHI = RHICreateTexture(Desc);
 
 	if (bCreateUAVs)
 	{
diff --git a/Engine/Source/Runtime/RHI/Public/RHICommandList.h b/Engine/Source/Runtime/RHI/Public/RHICommandList.h
--- a/Engine/Source/Runtime/RHI/Public/RHICommandList.h
+++ b/Engine/Source/Runtime/RHI/Public/RHICommandList.h
@@ -5443,7 +5443,7 @@ FORCEINLINE FTexture2DRHIRef RHIAsyncCreateTexture2D(uint32 SizeX, uint32 SizeY,
 	return GDynamicRHI->RHIAsyncCreateTexture2D(SizeX, SizeY, Format, NumMips, Flags, ResourceState, InitialMipData, NumInitialMips);
 }
 
-//UE_DEPRECATED(5.1, "FRHITexture2DArray is deprecated, please use RHICreateTexture(const FRHITextureCreateDesc&).")
+UE_DEPRECATED(5.1, "FRHITexture2DArray is deprecated, please use RHICreateTexture(const FRHITextureCreateDesc&).")
 FORCEINLINE FTexture2DArrayRHIRef RHICreateTexture2DArray(uint32 SizeX, uint32 SizeY, uint32 ArraySize, uint8 Format, uint32 NumMips, uint32 NumSamples, ETextureCreateFlags Flags, ERHIAccess ResourceState, const FRHIResourceCreateInfo& CreateInfo)
 {
 	return RHICreateTexture(
@@ -5462,7 +5462,7 @@ FORCEINLINE FTexture2DArrayRHIRef RHICreateTexture2DArray(uint32 SizeX, uint32 S
 	);
 }
 
-//UE_DEPRECATED(5.1, "FRHITexture3D is deprecated, please use RHICreateTexture(const FRHITextureCreateDesc&).")
+UE_DEPRECATED(5.1, "FRHITexture3D is deprecated, please use RHICreateTexture(const FRHITextureCreateDesc&).")
 FORCEINLINE FTexture3DRHIRef RHICreateTexture3D(uint32 SizeX, uint32 SizeY, uint32 SizeZ, uint8 Format, uint32 NumMips, ETextureCreateFlags Flags, ERHIAccess ResourceState, const FRHIResourceCreateInfo& CreateInfo)
 {
 	return RHICreateTexture(
@@ -5480,7 +5480,7 @@ FORCEINLINE FTexture3DRHIRef RHICreateTexture3D(uint32 SizeX, uint32 SizeY, uint
 	);
 }
 
-//UE_DEPRECATED(5.1, "FRHITextureCube is deprecated, please use RHICreateTexture(const FRHITextureCreateDesc&).")
+UE_DEPRECATED(5.1, "FRHITextureCube is deprecated, please use RHICreateTexture(const FRHITextureCreateDesc&).")
 FORCEINLINE FTextureCubeRHIRef RHICreateTextureCube(uint32 Size, uint8 Format, uint32 NumMips, ETextureCreateFlags Flags, ERHIAccess ResourceState, const FRHIResourceCreateInfo& CreateInfo)
 {
 	return RHICreateTexture(
@@ -5497,7 +5497,7 @@ FORCEINLINE FTextureCubeRHIRef RHICreateTextureCube(uint32 Size, uint8 Format, u
 	);
 }
 
-//UE_DEPRECATED(5.1, "FRHITextureCube is deprecated, please use RHICreateTexture(const FRHITextureCreateDesc&).")
+UE_DEPRECATED(5.1, "FRHITextureCube is deprecated, please use RHICreateTexture(const FRHITextureCreateDesc&).")
 FORCEINLINE FTextureCubeRHIRef RHICreateTextureCubeArray(uint32 Size, uint32 ArraySize, uint8 Format, uint32 NumMips, ETextureCreateFlags Flags, ERHIAccess ResourceState, const FRHIResourceCreateInfo& CreateInfo)
 {
 	return RHICreateTexture(
@@ -5539,7 +5539,7 @@ PRAGMA_DISABLE_DEPRECATION_WARNINGS
 PRAGMA_ENABLE_DEPRECATION_WARNINGS
 }
 
-//UE_DEPRECATED(5.1, "FRHITexture2DArray is deprecated, please use RHICreateTexture(const FRHITextureCreateDesc&).")
+UE_DEPRECATED(5.1, "FRHITexture2DArray is deprecated, please use RHICreateTexture(const FRHITextureCreateDesc&).")
 FORCEINLINE FTexture2DArrayRHIRef RHICreateTexture2DArray(uint32 SizeX, uint32 SizeY, uint32 SizeZ, uint8 Format, uint32 NumMips, uint32 NumSamples, ETextureCreateFlags Flags, const FRHIResourceCreateInfo& CreateInfo)
 {
 PRAGMA_DISABLE_DEPRECATION_WARNINGS
@@ -5547,7 +5547,7 @@ PRAGMA_DISABLE_DEPRECATION_WARNINGS
 PRAGMA_ENABLE_DEPRECATION_WARNINGS
 }
 
-//UE_DEPRECATED(5.1, "FRHITexture3D is deprecated, please use RHICreateTexture(const FRHITextureCreateDesc&).")
+UE_DEPRECATED(5.1, "FRHITexture3D is deprecated, please use RHICreateTexture(const FRHITextureCreateDesc&).")
 FORCEINLINE FTexture3DRHIRef RHICreateTexture3D(uint32 SizeX, uint32 SizeY, uint32 SizeZ, uint8 Format, uint32 NumMips, ETextureCreateFlags Flags, const FRHIResourceCreateInfo& CreateInfo)
 {
 PRAGMA_DISABLE_DEPRECATION_WARNINGS
@@ -5555,7 +5555,7 @@ PRAGMA_DISABLE_DEPRECATION_WARNINGS
 PRAGMA_ENABLE_DEPRECATION_WARNINGS
 }
 
-//UE_DEPRECATED(5.1, "FRHITextureCube is deprecated, please use RHICreateTexture(const FRHITextureCreateDesc&).")
+UE_DEPRECATED(5.1, "FRHITextureCube is deprecated, please use RHICreateTexture(const FRHITextureCreateDesc&).")
 FORCEINLINE FTextureCubeRHIRef RHICreateTextureCube(uint32 Size, uint8 Format, uint32 NumMips, ETextureCreateFlags Flags, const FRHIResourceCreateInfo& CreateInfo)
 {
 PRAGMA_DISABLE_DEPRECATION_WARNINGS
@@ -5563,7 +5563,7 @@ PRAGMA_DISABLE_DEPRECATION_WARNINGS
 PRAGMA_ENABLE_DEPRECATION_WARNINGS
 }
 
-//UE_DEPRECATED(5.1, "FRHITextureCube is deprecated, please use RHICreateTexture(const FRHITextureCreateDesc&).")
+UE_DEPRECATED(5.1, "FRHITextureCube is deprecated, please use RHICreateTexture(const FRHITextureCreateDesc&).")
 FORCEINLINE FTextureCubeRHIRef RHICreateTextureCubeArray(uint32 Size, uint32 ArraySize, uint8 Format, uint32 NumMips, ETextureCreateFlags Flags, const FRHIResourceCreateInfo& CreateInfo)
 {
 PRAGMA_DISABLE_DEPRECATION_WARNINGS
diff --git a/Engine/Source/Runtime/RHI/Public/RHIResources.h b/Engine/Source/Runtime/RHI/Public/RHIResources.h
--- a/Engine/Source/Runtime/RHI/Public/RHIResources.h
+++ b/Engine/Source/Runtime/RHI/Public/RHIResources.h
@@ -1533,6 +1533,20 @@ struct FRHITextureCreateDesc : public FRHITextureDesc
 		return FRHITextureCreateDesc(InDebugName, ETextureDimension::TextureCubeArray);
 	}
 
+	static FRHITextureCreateDesc Create2D(const TCHAR* DebugName, FIntPoint Size, EPixelFormat Format)
+	{
+		return Create2D(DebugName)
+			.SetExtent(Size)
+			.SetFormat(Format);
+	}
+
+	static FRHITextureCreateDesc Create2D(const TCHAR* DebugName, int32 SizeX, int32 SizeY, EPixelFormat Format)
+	{
+		return Create2D(DebugName)
+			.SetExtent(SizeX, SizeY)
+			.SetFormat(Format);
+	}
+
 	static FRHITextureCreateDesc Create2D(
 		  TCHAR const*                DebugName
 		, FIntPoint                   Size
@@ -1560,6 +1574,22 @@ struct FRHITextureCreateDesc : public FRHITextureDesc
 			.SetInitialState(InitialState == ERHIAccess::Unknown ? RHIGetDefaultResourceState(Flags, BulkData != nullptr) : InitialState);
 	}
 
+	static FRHITextureCreateDesc Create2DArray(const TCHAR* DebugName, FIntPoint Size, uint16 ArraySize, EPixelFormat Format)
+	{
+		return Create2DArray(DebugName)
+			.SetExtent(Size)
+			.SetFormat(Format)
+			.SetArraySize(ArraySize);
+	}
+
+	static FRHITextureCreateDesc Create2DArray(const TCHAR* DebugName, int32 SizeX, int32 SizeY, int32 ArraySize, EPixelFormat Format)
+	{
+		return Create2DArray(DebugName)
+			.SetExtent(SizeY, SizeY)
+			.SetFormat(Format)
+			.SetArraySize(ArraySize);
+	}
+
 	static FRHITextureCreateDesc Create2DArray(
 		  TCHAR const*                DebugName
 		, FIntPoint                   Size
@@ -1589,6 +1619,22 @@ struct FRHITextureCreateDesc : public FRHITextureDesc
 			.SetInitialState(InitialState == ERHIAccess::Unknown ? RHIGetDefaultResourceState(Flags, BulkData != nullptr) : InitialState);
 	}
 
+	static FRHITextureCreateDesc Create3D(const TCHAR* DebugName, FIntVector Size, EPixelFormat Format)
+	{
+		return Create3D(DebugName)
+			.SetExtent(Size.X, Size.Y)
+			.SetDepth(Size.Z)
+			.SetFormat(Format);
+	}
+
+	static FRHITextureCreateDesc Create3D(const TCHAR* DebugName, int32 SizeX, int32 SizeY, int32 SizeZ, EPixelFormat Format)
+	{
+		return Create3D(DebugName)
+			.SetExtent(SizeX, SizeY)
+			.SetDepth(SizeZ)
+			.SetFormat(Format);
+	}
+
 	static FRHITextureCreateDesc Create3D(
 		  TCHAR const*                DebugName
 		, FIntVector                  Size
@@ -1615,6 +1661,13 @@ struct FRHITextureCreateDesc : public FRHITextureDesc
 			.SetInitialState(InitialState == ERHIAccess::Unknown ? RHIGetDefaultResourceState(Flags, BulkData != nullptr) : InitialState);
 	}
 
+	static FRHITextureCreateDesc CreateCube(const TCHAR* DebugName, uint32 Size, EPixelFormat Format)
+	{
+		return CreateCube(DebugName)
+			.SetExtent(Size)
+			.SetFormat(Format);
+	}
+
 	static FRHITextureCreateDesc CreateCube(
 		  TCHAR const*                DebugName
 		, uint32                      Size
@@ -1642,6 +1695,14 @@ struct FRHITextureCreateDesc : public FRHITextureDesc
 			.SetInitialState(InitialState == ERHIAccess::Unknown ? RHIGetDefaultResourceState(Flags, BulkData != nullptr) : InitialState);
 	}
 
+	static FRHITextureCreateDesc CreateCubeArray(const TCHAR* DebugName, uint32 Size, uint16 ArraySize, EPixelFormat Format)
+	{
+		return CreateCubeArray(DebugName)
+			.SetExtent(Size)
+			.SetFormat(Format)
+			.SetArraySize(ArraySize);
+	}
+
 	static FRHITextureCreateDesc CreateCubeArray(
 		  TCHAR const*                DebugName
 		, uint32                      Size
diff --git a/Engine/Source/Runtime/RHI/Public/RHIUtilities.h b/Engine/Source/Runtime/RHI/Public/RHIUtilities.h
--- a/Engine/Source/Runtime/RHI/Public/RHIUtilities.h
+++ b/Engine/Source/Runtime/RHI/Public/RHIUtilities.h
@@ -177,47 +177,48 @@ extern RHI_API ERHIAccess RHIGetDefaultResourceState(ETextureCreateFlags InUsage
 /** Get the best default resource state for the given buffer creation flags */
 extern RHI_API ERHIAccess RHIGetDefaultResourceState(EBufferUsageFlags InUsage, bool bInHasInitialData);
 
-
 /** Encapsulates a GPU read/write texture 2D with its UAV and SRV. */
-struct FTextureRWBuffer2D
+struct FTextureRWBuffer
 {
-	FTexture2DRHIRef Buffer;
+	FTextureRHIRef Buffer;
 	FUnorderedAccessViewRHIRef UAV;
 	FShaderResourceViewRHIRef SRV;
-	uint32 NumBytes;
+	uint32 NumBytes = 0;
 
-	FTextureRWBuffer2D()
-		: NumBytes(0)
-	{}
+	FTextureRWBuffer() = default;
 
-	~FTextureRWBuffer2D()
+	~FTextureRWBuffer()
 	{
 		Release();
 	}
 
 	static constexpr ETextureCreateFlags DefaultTextureInitFlag = TexCreate_ShaderResource | TexCreate_UAV;
-	void Initialize(const TCHAR* InDebugName, const uint32 BytesPerElement, const uint32 SizeX, const uint32 SizeY, const EPixelFormat Format, ETextureCreateFlags Flags = DefaultTextureInitFlag)
+
+	void Initialize2D(const TCHAR* InDebugName, uint32 BytesPerElement, uint32 SizeX, uint32 SizeY, EPixelFormat Format, ETextureCreateFlags Flags = DefaultTextureInitFlag)
 	{
 		NumBytes = SizeX * SizeY * BytesPerElement;
 
-		FRHIResourceCreateInfo CreateInfo(InDebugName);
-		Buffer = RHICreateTexture2D(
-			SizeX, SizeY, UE_PIXELFORMAT_TO_UINT8(Format), //PF_R32_FLOAT,
-			/*NumMips=*/ 1,
-			1,
-			Flags,
-			/*BulkData=*/ CreateInfo);
+		const FRHITextureCreateDesc Desc =
+			FRHITextureCreateDesc::Create2D(InDebugName, SizeX, SizeY, Format)
+			.SetFlags(Flags);
 
-							
+		Buffer = RHICreateTexture(Desc);
 		UAV = RHICreateUnorderedAccessView(Buffer, 0);
 		SRV = RHICreateShaderResourceView(Buffer, 0);
 	}
 
-	UE_DEPRECATED(5.0, "AcquireTransientResource is deprecated. Transient resources are allocated through IRHITransientResourceAllocator instead.")
-	void AcquireTransientResource() {}
+	void Initialize3D(const TCHAR* InDebugName, uint32 BytesPerElement, uint32 SizeX, uint32 SizeY, uint32 SizeZ, EPixelFormat Format, ETextureCreateFlags Flags = DefaultTextureInitFlag)
+	{
+		NumBytes = SizeX * SizeY * SizeZ * BytesPerElement;
 
-	UE_DEPRECATED(5.0, "DiscardTransientResource is deprecated. Transient resources are allocated through IRHITransientResourceAllocator instead.")
-	void DiscardTransientResource() {}
+		const FRHITextureCreateDesc Desc =
+			FRHITextureCreateDesc::Create3D(InDebugName, SizeX, SizeY, SizeZ, Format)
+			.SetFlags(Flags);
+
+		Buffer = RHICreateTexture(Desc);
+		UAV = RHICreateUnorderedAccessView(Buffer, 0);
+		SRV = RHICreateShaderResourceView(Buffer, 0);
+	}
 
 	void Release()
 	{
@@ -228,51 +229,36 @@ struct FTextureRWBuffer2D
 	}
 };
 
-/** Encapsulates a GPU read/write texture 3d with its UAV and SRV. */
-struct FTextureRWBuffer3D
+struct UE_DEPRECATED(5.1, "FTextureRWBuffer should be used instead of FTextureRWBuffer2D.") FTextureRWBuffer2D;
+struct FTextureRWBuffer2D : public FTextureRWBuffer
 {
-	FTexture3DRHIRef Buffer;
-	FUnorderedAccessViewRHIRef UAV;
-	FShaderResourceViewRHIRef SRV;
-	uint32 NumBytes;
-
-	FTextureRWBuffer3D()
-		: NumBytes(0)
-	{}
-
-	~FTextureRWBuffer3D()
+	UE_DEPRECATED(5.1, "FTextureRWBuffer::Initialize2D should be used instead of FTextureRWBuffer2D::Initialize.")
+	void Initialize(const TCHAR* InDebugName, uint32 BytesPerElement, uint32 SizeX, uint32 SizeY, EPixelFormat Format, ETextureCreateFlags Flags = DefaultTextureInitFlag)
 	{
-		Release();
+		Initialize2D(InDebugName, BytesPerElement, SizeX, SizeY, Format, Flags);
 	}
 
-	void Initialize(const TCHAR* InDebugName, uint32 BytesPerElement, uint32 SizeX, uint32 SizeY, uint32 SizeZ, EPixelFormat Format)
-	{
-		NumBytes = SizeX * SizeY * SizeZ * BytesPerElement;
-
-		FRHIResourceCreateInfo CreateInfo(InDebugName);
-		Buffer = RHICreateTexture3D(
-			SizeX, SizeY, SizeZ, UE_PIXELFORMAT_TO_UINT8(Format),
-			/*NumMips=*/ 1,
-			/*Flags=*/ TexCreate_ShaderResource | TexCreate_UAV,
-			/*BulkData=*/ CreateInfo);
-
-		UAV = RHICreateUnorderedAccessView(Buffer, 0);
-		SRV = RHICreateShaderResourceView(Buffer, 0);
-	}
-	
 	UE_DEPRECATED(5.0, "AcquireTransientResource is deprecated. Transient resources are allocated through IRHITransientResourceAllocator instead.")
 	void AcquireTransientResource() {}
 
 	UE_DEPRECATED(5.0, "DiscardTransientResource is deprecated. Transient resources are allocated through IRHITransientResourceAllocator instead.")
 	void DiscardTransientResource() {}
+};
 
-	void Release()
+struct UE_DEPRECATED(5.1, "FTextureRWBuffer should be used instead of FTextureRWBuffer3D.") FTextureRWBuffer3D;
+struct FTextureRWBuffer3D : public FTextureRWBuffer
+{
+	UE_DEPRECATED(5.1, "FTextureRWBuffer::Initialize3D should be used instead of FTextureRWBuffer3D::Initialize.")
+	void Initialize(const TCHAR* InDebugName, uint32 BytesPerElement, uint32 SizeX, uint32 SizeY, uint32 SizeZ, EPixelFormat Format, ETextureCreateFlags Flags = DefaultTextureInitFlag)
 	{
-		NumBytes = 0;
-		Buffer.SafeRelease();
-		UAV.SafeRelease();
-		SRV.SafeRelease();
+		Initialize3D(InDebugName, BytesPerElement, SizeX, SizeY, SizeZ, Format, Flags);
 	}
+
+	UE_DEPRECATED(5.0, "AcquireTransientResource is deprecated. Transient resources are allocated through IRHITransientResourceAllocator instead.")
+	void AcquireTransientResource() {}
+
+	UE_DEPRECATED(5.0, "DiscardTransientResource is deprecated. Transient resources are allocated through IRHITransientResourceAllocator instead.")
+	void DiscardTransientResource() {}
 };
 
 /** Encapsulates a GPU read/write buffer with its UAV and SRV. */
diff --git a/Engine/Source/Runtime/RenderCore/Private/RenderUtils.cpp b/Engine/Source/Runtime/RenderCore/Private/RenderUtils.cpp
--- a/Engine/Source/Runtime/RenderCore/Private/RenderUtils.cpp
+++ b/Engine/Source/Runtime/RenderCore/Private/RenderUtils.cpp
@@ -274,23 +274,29 @@ public:
 	 */
 	virtual void InitRHI() override
 	{
-		const TCHAR* Name = TEXT("BlackVolumeTexture");
-
 		if (GSupportsTexture3D)
 		{
 			// Create the texture.
 			FBlackVolumeTextureResourceBulkDataInterface BlackTextureBulkData(Alpha);
-			FRHIResourceCreateInfo CreateInfo(TEXT("BlackVolumeTexture3D"), &BlackTextureBulkData);
-			FTexture3DRHIRef Texture3D = RHICreateTexture3D(1,1,1,PixelFormat,1,TexCreate_ShaderResource,CreateInfo);
-			TextureRHI = Texture3D;	
+
+			const FRHITextureCreateDesc Desc =
+				FRHITextureCreateDesc::Create3D(TEXT("BlackVolumeTexture3D"), 1, 1, 1, PixelFormat)
+				.SetFlags(ETextureCreateFlags::ShaderResource)
+				.SetBulkData(&BlackTextureBulkData);
+
+			TextureRHI = RHICreateTexture(Desc);
 		}
 		else
 		{
 			// Create a texture, even though it's not a volume texture
 			FBlackVolumeTextureResourceBulkDataInterface BlackTextureBulkData(Alpha);
-			FRHIResourceCreateInfo CreateInfo(TEXT("BlackVolumeTexture2D"), &BlackTextureBulkData);
-			FTexture2DRHIRef Texture2D = RHICreateTexture2D(1, 1, PixelFormat, 1, 1, TexCreate_ShaderResource, CreateInfo);
-			TextureRHI = Texture2D;
+
+			const FRHITextureCreateDesc Desc =
+				FRHITextureCreateDesc::Create2D(TEXT("BlackVolumeTexture2D"), 1, 1, PixelFormat)
+				.SetFlags(ETextureCreateFlags::ShaderResource)
+				.SetBulkData(&BlackTextureBulkData);
+
+			TextureRHI = RHICreateTexture(Desc);
 		}
 
 		// Create the sampler state.
@@ -330,9 +336,13 @@ public:
 	{
 		// Create the texture RHI.
 		FBlackVolumeTextureResourceBulkDataInterface BlackTextureBulkData(0);
-		FRHIResourceCreateInfo CreateInfo(TEXT("BlackArrayTexture"), &BlackTextureBulkData);
-		FTexture2DArrayRHIRef TextureArray = RHICreateTexture2DArray(1, 1, 1, PF_B8G8R8A8, 1, 1, TexCreate_ShaderResource, CreateInfo);
-		TextureRHI = TextureArray;
+
+		const FRHITextureCreateDesc Desc =
+			FRHITextureCreateDesc::Create2DArray(TEXT("BlackArrayTexture"), 1, 1, 1, PF_B8G8R8A8)
+			.SetFlags(ETextureCreateFlags::ShaderResource)
+			.SetBulkData(&BlackTextureBulkData);
+
+		TextureRHI = RHICreateTexture(Desc);
 
 		// Create the sampler state RHI resource.
 		FSamplerStateInitializerRHI SamplerStateInitializer(SF_Point, AM_Wrap, AM_Wrap, AM_Wrap);
@@ -461,10 +471,11 @@ public:
 	virtual void InitRHI() override
 	{
 		// Create the texture RHI.
-		const TCHAR* Name = TEXT("SolidColorCube");
+		const FRHITextureCreateDesc Desc =
+			FRHITextureCreateDesc::CreateCube(TEXT("SolidColorCube"), 1, PixelFormat)
+			.SetFlags(ETextureCreateFlags::ShaderResource);
 
-		FRHIResourceCreateInfo CreateInfo(Name);
-		FTextureCubeRHIRef TextureCube = RHICreateTextureCube(1, (uint8)PixelFormat, 1, TexCreate_ShaderResource, CreateInfo);
+		FTextureCubeRHIRef TextureCube = RHICreateTexture(Desc);
 		TextureRHI = TextureCube;
 
 		// Write the contents of the texture.
@@ -540,18 +551,20 @@ public:
 		{
 			const TCHAR* Name = TEXT("BlackCubeArray");
 
+			const FRHITextureCreateDesc Desc =
+				FRHITextureCreateDesc::CreateCubeArray(TEXT("BlackCubeArray"), 1, 1, PF_B8G8R8A8)
+				.SetFlags(ETextureCreateFlags::ShaderResource);
+
 			// Create the texture RHI.
-			FRHIResourceCreateInfo CreateInfo(Name);
-			FTextureCubeRHIRef TextureCubeArray = RHICreateTextureCubeArray(1,1,PF_B8G8R8A8,1,TexCreate_ShaderResource,CreateInfo);
-			TextureRHI = TextureCubeArray;
+			TextureRHI = RHICreateTexture(Desc);
 
 			for(uint32 FaceIndex = 0;FaceIndex < 6;FaceIndex++)
 			{
 				uint32 DestStride;
-				FColor* DestBuffer = (FColor*)RHILockTextureCubeFace(TextureCubeArray, FaceIndex, 0, 0, RLM_WriteOnly, DestStride, false);
+				FColor* DestBuffer = (FColor*)RHILockTextureCubeFace(TextureRHI, FaceIndex, 0, 0, RLM_WriteOnly, DestStride, false);
 				// Note: alpha is used by reflection environment to say how much of the foreground texture is visible, so 0 says it is completely invisible
 				*DestBuffer = FColor(0, 0, 0, 0);
-				RHIUnlockTextureCubeFace(TextureCubeArray, FaceIndex, 0, 0, false);
+				RHIUnlockTextureCubeFace(TextureRHI, FaceIndex, 0, 0, false);
 			}
 
 			// Create the sampler state RHI resource.
diff --git a/Engine/Source/Runtime/RenderCore/Private/UnifiedBuffer.cpp b/Engine/Source/Runtime/RenderCore/Private/UnifiedBuffer.cpp
--- a/Engine/Source/Runtime/RenderCore/Private/UnifiedBuffer.cpp
+++ b/Engine/Source/Runtime/RenderCore/Private/UnifiedBuffer.cpp
@@ -124,7 +124,7 @@ struct ResourceTypeTraits<FRWBufferStructured>
 };
 
 template<>
-struct ResourceTypeTraits<FTextureRWBuffer2D>
+struct ResourceTypeTraits<FTextureRWBuffer>
 {
 	static const EResourceType Type = EResourceType::TEXTURE;
 };
@@ -262,7 +262,7 @@ void MemcpyResource(FRHICommandList& RHICmdList, const ResourceType& DstBuffer,
 }
 
 template<>
-RENDERCORE_API bool ResizeResourceIfNeeded<FTextureRWBuffer2D>(FRHICommandList& RHICmdList, FTextureRWBuffer2D& Texture, uint32 NumBytes, const TCHAR* DebugName)
+RENDERCORE_API bool ResizeResourceIfNeeded<FTextureRWBuffer>(FRHICommandList& RHICmdList, FTextureRWBuffer& Texture, uint32 NumBytes, const TCHAR* DebugName)
 {
 	check((NumBytes & 15) == 0);
 	uint32 Float4sPerLine = CalculateFloat4sPerLine();
@@ -275,13 +275,13 @@ RENDERCORE_API bool ResizeResourceIfNeeded<FTextureRWBuffer2D>(FRHICommandList&
 
 	if (Texture.NumBytes == 0)
 	{
-		Texture.Initialize(DebugName, BytesPerElement, Float4sPerLine, NumLines, PF_A32B32G32R32F, TexCreate_RenderTargetable | TexCreate_UAV);
+		Texture.Initialize2D(DebugName, BytesPerElement, Float4sPerLine, NumLines, PF_A32B32G32R32F, TexCreate_RenderTargetable | TexCreate_UAV);
 		return true;
 	}
 	else if ((NumLines * Float4sPerLine * BytesPerElement) != Texture.NumBytes)
 	{
-		FTextureRWBuffer2D NewTexture;
-		NewTexture.Initialize(DebugName, BytesPerElement, Float4sPerLine, NumLines, PF_A32B32G32R32F, TexCreate_RenderTargetable | TexCreate_UAV);
+		FTextureRWBuffer NewTexture;
+		NewTexture.Initialize2D(DebugName, BytesPerElement, Float4sPerLine, NumLines, PF_A32B32G32R32F, TexCreate_RenderTargetable | TexCreate_UAV);
 
 		FMemcpyResourceParams Params;
 		Params.Count = NumBytes / BytesPerElement;
@@ -836,12 +836,12 @@ void FScatterUploadBuffer::ResourceUploadTo(FRHICommandList& RHICmdList, const R
 template RENDERCORE_API void MemsetResource<FRWBufferStructured>(FRHICommandList& RHICmdList, const FRWBufferStructured& DstBuffer, const FMemsetResourceParams& Params);
 template RENDERCORE_API void MemsetResource<FRWByteAddressBuffer>(FRHICommandList& RHICmdList, const FRWByteAddressBuffer& DstBuffer, const FMemsetResourceParams& Params);
 
-template RENDERCORE_API void MemcpyResource<FTextureRWBuffer2D>(FRHICommandList& RHICmdList, const FTextureRWBuffer2D& DstBuffer, const FTextureRWBuffer2D& SrcBuffer, const FMemcpyResourceParams& Params, bool bAlreadyInUAVOverlap);
+template RENDERCORE_API void MemcpyResource<FTextureRWBuffer>(FRHICommandList& RHICmdList, const FTextureRWBuffer& DstBuffer, const FTextureRWBuffer& SrcBuffer, const FMemcpyResourceParams& Params, bool bAlreadyInUAVOverlap);
 template RENDERCORE_API void MemcpyResource<FRWBuffer>(FRHICommandList& RHICmdList, const FRWBuffer& DstBuffer, const FRWBuffer& SrcBuffer, const FMemcpyResourceParams& Params, bool bAlreadyInUAVOverlap);
 template RENDERCORE_API void MemcpyResource<FRWBufferStructured>(FRHICommandList& RHICmdList, const FRWBufferStructured& DstBuffer, const FRWBufferStructured& SrcBuffer, const FMemcpyResourceParams& Params, bool bAlreadyInUAVOverlap);
 template RENDERCORE_API void MemcpyResource<FRWByteAddressBuffer>(FRHICommandList& RHICmdList, const FRWByteAddressBuffer& DstBuffer, const FRWByteAddressBuffer& SrcBuffer, const FMemcpyResourceParams& Params, bool bAlreadyInUAVOverlap);
 
-template RENDERCORE_API void FScatterUploadBuffer::ResourceUploadTo<FTextureRWBuffer2D>(FRHICommandList& RHICmdList, const FTextureRWBuffer2D& DstBuffer, bool bFlush);
+template RENDERCORE_API void FScatterUploadBuffer::ResourceUploadTo<FTextureRWBuffer>(FRHICommandList& RHICmdList, const FTextureRWBuffer& DstBuffer, bool bFlush);
 template RENDERCORE_API void FScatterUploadBuffer::ResourceUploadTo<FRWBuffer>(FRHICommandList& RHICmdList, const FRWBuffer& DstBuffer, bool bFlush);
 template RENDERCORE_API void FScatterUploadBuffer::ResourceUploadTo<FRWBufferStructured>(FRHICommandList& RHICmdList, const FRWBufferStructured& DstBuffer, bool bFlush);
 template RENDERCORE_API void FScatterUploadBuffer::ResourceUploadTo<FRWByteAddressBuffer>(FRHICommandList& RHICmdList, const FRWByteAddressBuffer& DstBuffer, bool bFlush);
\ No newline at end of file
diff --git a/Engine/Source/Runtime/Renderer/Private/SystemTextures.cpp b/Engine/Source/Runtime/Renderer/Private/SystemTextures.cpp
--- a/Engine/Source/Runtime/Renderer/Private/SystemTextures.cpp
+++ b/Engine/Source/Runtime/Renderer/Private/SystemTextures.cpp
@@ -79,21 +79,21 @@ void FSystemTextures::InitializeTextures(FRHICommandListImmediate& RHICmdList, c
 }
 
 template <typename DataType>
-void SetDummyTextureData(FRHITexture2D* Texture2D, const DataType& DummyData)
+void SetDummyTextureData(FRHITexture* Texture, const DataType& DummyData)
 {
 	uint32 DestStride;
-	DataType* DestBuffer = (DataType*)RHILockTexture2D(Texture2D, 0, RLM_WriteOnly, DestStride, false);
+	DataType* DestBuffer = (DataType*)RHILockTexture2D(Texture, 0, RLM_WriteOnly, DestStride, false);
 	*DestBuffer = DummyData;
-	RHIUnlockTexture2D(Texture2D, 0, false);
+	RHIUnlockTexture2D(Texture, 0, false);
 }
 
 template <typename DataType>
-void SetDummyTextureArrayData(FRHITexture2DArray* Texture2DArray, const DataType& DummyData)
+void SetDummyTextureArrayData(FRHITexture* Texture, const DataType& DummyData)
 {
 	uint32 DestStride;
-	DataType* DestBuffer = (DataType*)RHILockTexture2DArray(Texture2DArray, 0, 0, RLM_WriteOnly, DestStride, false);
+	DataType* DestBuffer = (DataType*)RHILockTexture2DArray(Texture, 0, 0, RLM_WriteOnly, DestStride, false);
 	*DestBuffer = DummyData;
-	RHIUnlockTexture2DArray(Texture2DArray, 0, 0, false);
+	RHIUnlockTexture2DArray(Texture, 0, 0, false);
 }
 
 void FSystemTextures::InitializeCommonTextures(FRHICommandListImmediate& RHICmdList)
@@ -109,62 +109,83 @@ void FSystemTextures::InitializeCommonTextures(FRHICommandListImmediate& RHICmdL
 
 	// Create a BlackDummy texture
 	{
-		FRHIResourceCreateInfo CreateInfo(TEXT("BlackDummy"));
-		FTexture2DRHIRef Texture2D = RHICreateTexture2D(1, 1, PF_B8G8R8A8, 1, 1, TexCreate_ShaderResource, CreateInfo);
-		SetDummyTextureData<FColor>(Texture2D, FColor(0, 0, 0, 0));
-		BlackDummy = CreateRenderTarget(Texture2D, CreateInfo.DebugName);
+		const FRHITextureCreateDesc Desc =
+			FRHITextureCreateDesc::Create2D(TEXT("BlackDummy"), 1, 1, PF_B8G8R8A8)
+			.SetFlags(ETextureCreateFlags::ShaderResource);
+
+		FTextureRHIRef Texture = RHICreateTexture(Desc);
+		SetDummyTextureData<FColor>(Texture, FColor(0, 0, 0, 0));
+		BlackDummy = CreateRenderTarget(Texture, Desc.DebugName);
 	}
 	
 	// Create a texture array that has a single black dummy slice 
 	{
-		FRHIResourceCreateInfo CreateInfo(TEXT("BlackArrayDummy"));
-		FTexture2DArrayRHIRef Texture2DArray = RHICreateTexture2DArray(1, 1, 1, PF_B8G8R8A8, 1, 1, TexCreate_ShaderResource, CreateInfo);
-		SetDummyTextureArrayData<FColor>(Texture2DArray, FColor(0, 0, 0, 0));
-		BlackArrayDummy = CreateRenderTarget(Texture2DArray, CreateInfo.DebugName);
+		const FRHITextureCreateDesc Desc =
+			FRHITextureCreateDesc::Create2DArray(TEXT("BlackArrayDummy"), 1, 1, 1, PF_B8G8R8A8)
+			.SetFlags(ETextureCreateFlags::ShaderResource);
+
+		FTextureRHIRef Texture = RHICreateTexture(Desc);
+		SetDummyTextureArrayData<FColor>(Texture, FColor(0, 0, 0, 0));
+		BlackArrayDummy = CreateRenderTarget(Texture, Desc.DebugName);
 	}
 
 	// Create a texture that is a single UInt32 value set to 0
-	{	
-		FRHIResourceCreateInfo CreateInfo(TEXT("ZeroUIntDummy"));
-		FTexture2DRHIRef Texture2D = RHICreateTexture2D(1, 1, PF_R32_UINT, 1, 1, TexCreate_ShaderResource, CreateInfo);
-		SetDummyTextureData<uint32>(Texture2D, 0u);
-		ZeroUIntDummy = CreateRenderTarget(Texture2D, CreateInfo.DebugName);
+	{
+		const FRHITextureCreateDesc Desc =
+			FRHITextureCreateDesc::Create2D(TEXT("ZeroUIntDummy"), 1, 1, PF_R32_UINT)
+			.SetFlags(ETextureCreateFlags::ShaderResource);
+
+		FTextureRHIRef Texture = RHICreateTexture(Desc);
+		SetDummyTextureData<uint32>(Texture, 0u);
+		ZeroUIntDummy = CreateRenderTarget(Texture, Desc.DebugName);
 	}
 
 	// Create a BlackAlphaOneDummy texture
 	{
-		FRHIResourceCreateInfo CreateInfo(TEXT("BlackAlphaOneDummy"));
-		FTexture2DRHIRef Texture2D = RHICreateTexture2D(1, 1, PF_B8G8R8A8, 1, 1, TexCreate_ShaderResource, CreateInfo);
-		SetDummyTextureData<FColor>(Texture2D, FColor(0, 0, 0, 255));
-		BlackAlphaOneDummy = CreateRenderTarget(Texture2D, CreateInfo.DebugName);
+		const FRHITextureCreateDesc Desc =
+			FRHITextureCreateDesc::Create2D(TEXT("BlackAlphaOneDummy"), 1, 1, PF_B8G8R8A8)
+			.SetFlags(ETextureCreateFlags::ShaderResource);
+
+		FTextureRHIRef Texture = RHICreateTexture(Desc);
+		SetDummyTextureData<FColor>(Texture, FColor(0, 0, 0, 255));
+		BlackAlphaOneDummy = CreateRenderTarget(Texture, Desc.DebugName);
 	}
 
 	// Create a GreenDummy texture
 	{
-		FRHIResourceCreateInfo CreateInfo(TEXT("GreenDummy"));
-		FTexture2DRHIRef Texture2D = RHICreateTexture2D(1, 1, PF_B8G8R8A8, 1, 1, TexCreate_ShaderResource, CreateInfo);
-		SetDummyTextureData<FColor>(Texture2D, FColor(0, 255, 0, 255));
-		GreenDummy = CreateRenderTarget(Texture2D, CreateInfo.DebugName);
+		const FRHITextureCreateDesc Desc =
+			FRHITextureCreateDesc::Create2DArray(TEXT("GreenDummy"), 1, 1, PF_B8G8R8A8)
+			.SetFlags(ETextureCreateFlags::ShaderResource);
+
+		FTextureRHIRef Texture = RHICreateTexture(Desc);
+		SetDummyTextureData<FColor>(Texture, FColor(0, 255, 0, 255));
+		GreenDummy = CreateRenderTarget(Texture, Desc.DebugName);
 	}
 
 	// Create a DefaultNormal8Bit texture
 	{
-		FRHIResourceCreateInfo CreateInfo(TEXT("DefaultNormal8Bit"));
-		FTexture2DRHIRef Texture2D = RHICreateTexture2D(1, 1, PF_B8G8R8A8, 1, 1, TexCreate_ShaderResource, CreateInfo);
-		SetDummyTextureData<FColor>(Texture2D, FColor(128, 128, 128, 255));
-		DefaultNormal8Bit = CreateRenderTarget(Texture2D, CreateInfo.DebugName);
+		const FRHITextureCreateDesc Desc =
+			FRHITextureCreateDesc::Create2D(TEXT("DefaultNormal8Bit"), 1, 1, PF_B8G8R8A8)
+			.SetFlags(ETextureCreateFlags::ShaderResource);
+
+		FTextureRHIRef Texture = RHICreateTexture(Desc);
+		SetDummyTextureData<FColor>(Texture, FColor(128, 128, 128, 255));
+		DefaultNormal8Bit = CreateRenderTarget(Texture, Desc.DebugName);
 	}
 
 	// Create the PerlinNoiseGradient texture
 	{
 		const uint32 Extent = 128;
+
+		const FRHITextureCreateDesc Desc =
+			FRHITextureCreateDesc::Create2D(TEXT("PerlinNoiseGradient"), Extent, Extent, PF_B8G8R8A8)
+			.SetFlags(ETextureCreateFlags::ShaderResource);
 	
-		FRHIResourceCreateInfo CreateInfo(TEXT("PerlinNoiseGradient"));
-		FTexture2DRHIRef Texture2D = RHICreateTexture2D(Extent, Extent, PF_B8G8R8A8, 1, 1, TexCreate_ShaderResource, CreateInfo);
+		FTextureRHIRef Texture = RHICreateTexture(Desc);
 
 		// Write the contents of the texture.
 		uint32 DestStride;
-		uint8* DestBuffer = (uint8*)RHICmdList.LockTexture2D(Texture2D, 0, RLM_WriteOnly, DestStride, false);
+		uint8* DestBuffer = (uint8*)RHICmdList.LockTexture2D(Texture, 0, RLM_WriteOnly, DestStride, false);
 		// seed the pseudo random stream with a good value
 		FRandomStream RandomStream(12345);
 		// Values represent float3 values in the -1..1 range.
@@ -186,17 +207,20 @@ void FSystemTextures::InitializeCommonTextures(FRHICommandListImmediate& RHICmdL
 				*Dest = gradtable[(uint32)(RandomStream.GetFraction() * 11.9999999f)];
 			}
 		}
-		RHICmdList.UnlockTexture2D(Texture2D, 0, false);
+		RHICmdList.UnlockTexture2D(Texture, 0, false);
 
-		PerlinNoiseGradient = CreateRenderTarget(Texture2D, CreateInfo.DebugName);
+		PerlinNoiseGradient = CreateRenderTarget(Texture, Desc.DebugName);
 	}
 
 	if (GPixelFormats[PF_FloatRGBA].Supported)
 	{
-		FRHIResourceCreateInfo CreateInfo(TEXT("MaxFP16Depth"));
-		FTexture2DRHIRef Texture2D = RHICreateTexture2D(1, 1, PF_FloatRGBA, 1, 1, TexCreate_ShaderResource, CreateInfo);
-		SetDummyTextureData<FFloat16Color>(Texture2D, FFloat16Color(FLinearColor(65500.0f, 65500.0f, 65500.0f, 65500.0f)));
-		MaxFP16Depth = CreateRenderTarget(Texture2D, CreateInfo.DebugName);
+		const FRHITextureCreateDesc Desc =
+			FRHITextureCreateDesc::Create2D(TEXT("MaxFP16Depth"), 1, 1, PF_FloatRGBA)
+			.SetFlags(ETextureCreateFlags::ShaderResource);
+
+		FTextureRHIRef Texture = RHICreateTexture(Desc);
+		SetDummyTextureData<FFloat16Color>(Texture, FFloat16Color(FLinearColor(65500.0f, 65500.0f, 65500.0f, 65500.0f)));
+		MaxFP16Depth = CreateRenderTarget(Texture, Desc.DebugName);
 	}
 
 	{
@@ -207,19 +231,25 @@ void FSystemTextures::InitializeCommonTextures(FRHICommandListImmediate& RHICmdL
 
 	// Create a dummy stencil SRV.
 	{
-		FRHIResourceCreateInfo CreateInfo(TEXT("StencilDummy"));
-		FTexture2DRHIRef Texture2D = RHICreateTexture2D(1, 1, PF_R8G8B8A8_UINT, 1, 1, TexCreate_ShaderResource, CreateInfo);
-		SetDummyTextureData<FColor>(Texture2D, FColor::White);
-		StencilDummy = CreateRenderTarget(Texture2D, CreateInfo.DebugName);
-		StencilDummySRV = RHICreateShaderResourceView((FRHITexture2D*)StencilDummy->GetRHI(), 0);
+		const FRHITextureCreateDesc Desc =
+			FRHITextureCreateDesc::Create2D(TEXT("StencilDummy"), 1, 1, PF_R8G8B8A8_UINT)
+			.SetFlags(ETextureCreateFlags::ShaderResource);
+
+		FTextureRHIRef Texture = RHICreateTexture(Desc);
+		SetDummyTextureData<FColor>(Texture, FColor::White);
+		StencilDummy = CreateRenderTarget(Texture, Desc.DebugName);
+		StencilDummySRV = RHICreateShaderResourceView(StencilDummy->GetRHI(), 0);
 	}
 
 	if (GPixelFormats[PF_FloatRGBA].Supported)
 	{
-		FRHIResourceCreateInfo CreateInfo(TEXT("MidGreyDummy"));
-		FTexture2DRHIRef Texture2D = RHICreateTexture2D(1, 1, PF_FloatRGBA, 1, 1, TexCreate_ShaderResource, CreateInfo);
-		SetDummyTextureData<FFloat16Color>(Texture2D, FFloat16Color(FLinearColor(0.5f, 0.5f, 0.5f, 0.5f)));
-		MidGreyDummy = CreateRenderTarget(Texture2D, CreateInfo.DebugName);
+		const FRHITextureCreateDesc Desc =
+			FRHITextureCreateDesc::Create2D(TEXT("MidGreyDummy"), 1, 1, PF_FloatRGBA)
+			.SetFlags(ETextureCreateFlags::ShaderResource);
+
+		FTextureRHIRef Texture = RHICreateTexture(Desc);
+		SetDummyTextureData<FFloat16Color>(Texture, FFloat16Color(FLinearColor(0.5f, 0.5f, 0.5f, 0.5f)));
+		MidGreyDummy = CreateRenderTarget(Texture, Desc.DebugName);
 	}
 
 	// Create a VolumetricBlackDummy texture
@@ -241,11 +271,14 @@ void FSystemTextures::InitializeFeatureLevelDependentTextures(FRHICommandListImm
 		// Create the SobolSampling texture
 	if (CurrentFeatureLevel < ERHIFeatureLevel::ES3_1 && GPixelFormats[PF_R16_UINT].Supported)
 	{
-		FRHIResourceCreateInfo CreateInfo(TEXT("SobolSampling"));
-		FTexture2DRHIRef Texture2D = RHICreateTexture2D(32, 16, PF_R16_UINT, 1, 1, TexCreate_ShaderResource, CreateInfo);
+		const FRHITextureCreateDesc Desc =
+			FRHITextureCreateDesc::Create2D(TEXT("SobolSampling"), 32, 16, PF_R16_UINT)
+			.SetFlags(ETextureCreateFlags::ShaderResource);
+
+		FTextureRHIRef Texture = RHICreateTexture(Desc);
 		// Write the contents of the texture.
 		uint32 DestStride;
-		uint8* DestBuffer = (uint8*)RHICmdList.LockTexture2D(Texture2D, 0, RLM_WriteOnly, DestStride, false);
+		uint8* DestBuffer = (uint8*)RHICmdList.LockTexture2D(Texture, 0, RLM_WriteOnly, DestStride, false);
 
 		uint16 *Dest;
 		for (int y = 0; y < 16; ++y)
@@ -264,9 +297,9 @@ void FSystemTextures::InitializeFeatureLevelDependentTextures(FRHICommandListImm
 				*Dest = FSobol::ComputeGPUSpatialSeed(x, y, /* Index = */ 1);
 			}
 		}
-		RHICmdList.UnlockTexture2D(Texture2D, 0, false);
+		RHICmdList.UnlockTexture2D(Texture, 0, false);
 
-		SobolSampling = CreateRenderTarget(Texture2D, CreateInfo.DebugName);
+		SobolSampling = CreateRenderTarget(Texture, Desc.DebugName);
 	}
 
 	if (CurrentFeatureLevel < ERHIFeatureLevel::SM5 && InFeatureLevel >= ERHIFeatureLevel::SM5)
@@ -365,12 +398,15 @@ void FSystemTextures::InitializeFeatureLevelDependentTextures(FRHICommandListImm
 			Extent.Y = 128;
 		}
 
-		FRHIResourceCreateInfo CreateInfo(TEXT("PreintegratedGF"));
-		FTexture2DRHIRef Texture2D = RHICreateTexture2D(Extent.X, Extent.Y, Format, 1, 1, TexCreate_ShaderResource, CreateInfo);
+		const FRHITextureCreateDesc Desc =
+			FRHITextureCreateDesc::Create2D(TEXT("PreintegratedGF"), Extent, Format)
+			.SetFlags(ETextureCreateFlags::ShaderResource);
+
+		FTextureRHIRef Texture = RHICreateTexture(Desc);
 
 		// Write the contents of the texture.
 		uint32 DestStride;
-		uint8* DestBuffer = (uint8*)RHICmdList.LockTexture2D(Texture2D, 0, RLM_WriteOnly, DestStride, false);
+		uint8* DestBuffer = (uint8*)RHICmdList.LockTexture2D(Texture, 0, RLM_WriteOnly, DestStride, false);
 
 		// x is NoV, y is roughness
 		for (int32 y = 0; y < Extent.Y; y++)
@@ -473,9 +509,9 @@ void FSystemTextures::InitializeFeatureLevelDependentTextures(FRHICommandListImm
 				}
 			}
 		}
-		RHICmdList.UnlockTexture2D(Texture2D, 0, false);
+		RHICmdList.UnlockTexture2D(Texture, 0, false);
 
-		PreintegratedGF = CreateRenderTarget(Texture2D, CreateInfo.DebugName);
+		PreintegratedGF = CreateRenderTarget(Texture, Desc.DebugName);
 	}
 
 	{
@@ -485,8 +521,11 @@ void FSystemTextures::InitializeFeatureLevelDependentTextures(FRHICommandListImm
     
 		    const uint32 Square = Extent * Extent;
 
-			FRHIResourceCreateInfo CreateInfo(TEXT("PerlinNoise3D"));
-			FTexture3DRHIRef Texture3D = RHICreateTexture3D(Extent, Extent, Extent, PF_B8G8R8A8, 1, TexCreate_ShaderResource | TexCreate_NoTiling, CreateInfo);
+			const FRHITextureCreateDesc Desc =
+				FRHITextureCreateDesc::Create3D(TEXT("PerlinNoise3D"), Extent, Extent, Extent, PF_B8G8R8A8)
+				.SetFlags(ETextureCreateFlags::ShaderResource | ETextureCreateFlags::NoTiling);
+
+			FTextureRHIRef Texture3D = RHICreateTexture(Desc);
 
 		    // Write the contents of the texture.
 		    TArray<uint32> DestBuffer;
@@ -589,19 +628,22 @@ void FSystemTextures::InitializeFeatureLevelDependentTextures(FRHICommandListImm
 			    Extent * Extent * sizeof(uint32),
 			    (const uint8*)DestBuffer.GetData());
 
-			PerlinNoise3D = CreateRenderTarget(Texture3D, CreateInfo.DebugName);
+			PerlinNoise3D = CreateRenderTarget(Texture3D, Desc.DebugName);
 
 		} // end Create the PerlinNoise3D texture
 
 		// LTC Textures	(used by Rect Lights)
 		{
 			{
-				FRHIResourceCreateInfo CreateInfo(TEXT("LTCMat"));
-				FTexture2DRHIRef Texture2D = RHICreateTexture2D(LTC_Size, LTC_Size, PF_FloatRGBA, 1, 1, TexCreate_FastVRAM | TexCreate_ShaderResource, CreateInfo);
+				const FRHITextureCreateDesc Desc =
+					FRHITextureCreateDesc::Create2D(TEXT("LTCMat"), LTC_Size, LTC_Size, PF_FloatRGBA)
+					.SetFlags(ETextureCreateFlags::ShaderResource | ETextureCreateFlags::FastVRAM);
+
+				FTextureRHIRef Texture = RHICreateTexture(Desc);
 
 				// Write the contents of the texture.
 				uint32 DestStride;
-				uint8* DestBuffer = (uint8*)RHICmdList.LockTexture2D(Texture2D, 0, RLM_WriteOnly, DestStride, false);
+				uint8* DestBuffer = (uint8*)RHICmdList.LockTexture2D(Texture, 0, RLM_WriteOnly, DestStride, false);
     
 				for (int32 y = 0; y < LTC_Size; ++y)
 				{
@@ -615,18 +657,21 @@ void FSystemTextures::InitializeFeatureLevelDependentTextures(FRHICommandListImm
 						}
 					}
 				}
-				RHICmdList.UnlockTexture2D(Texture2D, 0, false);
+				RHICmdList.UnlockTexture2D(Texture, 0, false);
 
-				LTCMat = CreateRenderTarget(Texture2D, CreateInfo.DebugName);
+				LTCMat = CreateRenderTarget(Texture, Desc.DebugName);
 			}
 
 			{
-				FRHIResourceCreateInfo CreateInfo(TEXT("LTCAmp"));
-				FTexture2DRHIRef Texture2D = RHICreateTexture2D(LTC_Size, LTC_Size, PF_G16R16F, 1, 1, TexCreate_FastVRAM | TexCreate_ShaderResource, CreateInfo);
+				const FRHITextureCreateDesc Desc =
+					FRHITextureCreateDesc::Create2D(TEXT("LTCAmp"), LTC_Size, LTC_Size, PF_G16R16F)
+					.SetFlags(ETextureCreateFlags::ShaderResource | ETextureCreateFlags::FastVRAM);
+
+				FTexture2DRHIRef Texture = RHICreateTexture(Desc);
 
 				// Write the contents of the texture.
 				uint32 DestStride;
-				uint8* DestBuffer = (uint8*)RHICmdList.LockTexture2D(Texture2D, 0, RLM_WriteOnly, DestStride, false);
+				uint8* DestBuffer = (uint8*)RHICmdList.LockTexture2D(Texture, 0, RLM_WriteOnly, DestStride, false);
 
 				for (int32 y = 0; y < LTC_Size; ++y)
 				{
@@ -640,9 +685,9 @@ void FSystemTextures::InitializeFeatureLevelDependentTextures(FRHICommandListImm
 						}
 					}
 				}
-				RHICmdList.UnlockTexture2D(Texture2D, 0, false);
+				RHICmdList.UnlockTexture2D(Texture, 0, false);
 
-				LTCAmp = CreateRenderTarget(Texture2D, CreateInfo.DebugName);
+				LTCAmp = CreateRenderTarget(Texture, Desc.DebugName);
 			}
 		}
 	}
@@ -680,12 +725,15 @@ void FSystemTextures::InitializeFeatureLevelDependentTextures(FRHICommandListImm
 			{
 				const uint32 Extent = 64;
 
-				FRHIResourceCreateInfo CreateInfo(TEXT("SSAORandomization"));
-				FTexture2DRHIRef Texture2D = RHICreateTexture2D(Extent, Extent, PF_R8G8, 1, 1, TexCreate_ShaderResource, CreateInfo);
+				const FRHITextureCreateDesc Desc =
+					FRHITextureCreateDesc::Create2D(TEXT("SSAORandomization"), Extent, Extent, PF_R8G8)
+					.SetFlags(ETextureCreateFlags::ShaderResource);
+
+				FTextureRHIRef Texture = RHICreateTexture(Desc);
 
 				// Write the contents of the texture.
 				uint32 DestStride;
-				uint8* DestBuffer = (uint8*)RHICmdList.LockTexture2D(Texture2D, 0, RLM_WriteOnly, DestStride, false);
+				uint8* DestBuffer = (uint8*)RHICmdList.LockTexture2D(Texture, 0, RLM_WriteOnly, DestStride, false);
 
 				for (int32 y = 0; y < Extent; ++y)
 				{
@@ -700,9 +748,9 @@ void FSystemTextures::InitializeFeatureLevelDependentTextures(FRHICommandListImm
 					}
 				}
 
-				RHICmdList.UnlockTexture2D(Texture2D, 0, false);
+				RHICmdList.UnlockTexture2D(Texture, 0, false);
 
-				SSAORandomization = CreateRenderTarget(Texture2D, CreateInfo.DebugName);
+				SSAORandomization = CreateRenderTarget(Texture, Desc.DebugName);
 			}
 		}
 	}
@@ -717,17 +765,14 @@ void FSystemTextures::InitializeFeatureLevelDependentTextures(FRHICommandListImm
 
 		bool bGTAOPreIngegratedUsingVolumeLUT = MobileGTAOPreIntegratedTextureTypeCVar->GetValueOnAnyThread() == 2;
 
-		TRefCountPtr<FRHITexture> Texture;
+		const FRHITextureCreateDesc Desc =
+			FRHITextureCreateDesc(TEXT("GTAOPreIntegrated"), bGTAOPreIngegratedUsingVolumeLUT ? ETextureDimension::Texture3D : ETextureDimension::Texture2D)
+			.SetExtent(Extent)
+			.SetDepth(bGTAOPreIngegratedUsingVolumeLUT ? Extent : 1)
+			.SetFormat(PF_R16F)
+			.SetFlags(ETextureCreateFlags::ShaderResource);
 
-		FRHIResourceCreateInfo CreateInfo(TEXT("GTAOPreIntegrated"));
-		if (bGTAOPreIngegratedUsingVolumeLUT)
-		{
-			Texture = RHICreateTexture3D(Extent, Extent, Extent, PF_R16F, 1, TexCreate_ShaderResource, CreateInfo);
-		}
-		else
-		{
-			Texture = RHICreateTexture2D(Square, Extent, PF_R16F, 1, 1, TexCreate_ShaderResource, CreateInfo);
-		}
+		TRefCountPtr<FRHITexture> Texture = RHICreateTexture(Desc);
 
 		// Write the contents of the texture.
 		TArray<FFloat16> TempBuffer;
@@ -742,7 +787,7 @@ void FSystemTextures::InitializeFeatureLevelDependentTextures(FRHICommandListImm
 		else
 		{
 			uint32 DestStride;
-			DestBuffer = (FFloat16*)RHICmdList.LockTexture2D((FRHITexture2D*)Texture.GetReference(), 0, RLM_WriteOnly, DestStride, false);
+			DestBuffer = (FFloat16*)RHICmdList.LockTexture2D(Texture.GetReference(), 0, RLM_WriteOnly, DestStride, false);
 		}
 
 		for (uint32 z = 0; z < Extent; ++z)
@@ -800,10 +845,10 @@ void FSystemTextures::InitializeFeatureLevelDependentTextures(FRHICommandListImm
 		}
 		else
 		{
-			RHICmdList.UnlockTexture2D((FRHITexture2D*)Texture.GetReference(), 0, false);
+			RHICmdList.UnlockTexture2D(Texture.GetReference(), 0, false);
 		}
 
-		GTAOPreIntegrated = CreateRenderTarget(Texture, CreateInfo.DebugName);
+		GTAOPreIntegrated = CreateRenderTarget(Texture, Desc.DebugName);
 	}
 
 	// Initialize textures only once.
@@ -1296,46 +1341,59 @@ FRDGTextureRef GetInternalDefaultTexture(
 	Entry.Hash = Hash;
 	Entry.Texture = nullptr;
 	
+	if (Dimension == ETextureDimension::Texture2D)
 	{
-		if (Dimension == ETextureDimension::Texture2D)
-		{
-			FRHIResourceCreateInfo CreateInfo(TEXT("DefaultTexture2D"));
-			FTexture2DRHIRef Texture = RHICreateTexture2D(1, 1, Format, 1, 1, TexCreate_ShaderResource, CreateInfo);
-			SetDefaultTextureData2D(Texture, Value);
-			Entry.Texture = CreateRenderTarget(Texture, CreateInfo.DebugName);
-		}
-		else if (Dimension == ETextureDimension::Texture2DArray)
-		{
-			FRHIResourceCreateInfo CreateInfo(TEXT("DefaultTexture2DArray"));
-			FTexture2DArrayRHIRef Texture = RHICreateTexture2DArray(1, 1, 1, Format, 1, 1, TexCreate_ShaderResource, CreateInfo);
-			SetDefaultTextureData2DArray(Texture, Value);
-			Entry.Texture = CreateRenderTarget(Texture, CreateInfo.DebugName);
-		}
-		else if (Dimension == ETextureDimension::Texture3D)
-		{
-			FRHIResourceCreateInfo CreateInfo(TEXT("DefaultTexture3D"));
-			FTexture3DRHIRef Texture = RHICreateTexture3D(1, 1, 1, Format, 1, TexCreate_ShaderResource, CreateInfo);
-			SetDefaultTextureData3D(GraphBuilder.RHICmdList, Texture, Value);
-			Entry.Texture = CreateRenderTarget(Texture, CreateInfo.DebugName);
-		}
-		else if (Dimension == ETextureDimension::TextureCube)
-		{
-			FRHIResourceCreateInfo CreateInfo(TEXT("DefaultTextureCube"));
-			FTextureCubeRHIRef Texture = RHICreateTextureCube(1, Format, 1, TexCreate_ShaderResource, CreateInfo);
-			SetDefaultTextureDataCube(Texture, Value);
-			Entry.Texture = CreateRenderTarget(Texture, CreateInfo.DebugName);
-		}
-		else if (Dimension == ETextureDimension::TextureCubeArray)
-		{
-			FRHIResourceCreateInfo CreateInfo(TEXT("DefaultTextureCubeArray"));
-			FTextureCubeRHIRef Texture = RHICreateTextureCubeArray(1, 1, Format, 1, TexCreate_ShaderResource, CreateInfo);
-			SetDefaultTextureDataCube(Texture, Value);
-			Entry.Texture = CreateRenderTarget(Texture, CreateInfo.DebugName);
-		}
-		else
-		{
-			return nullptr;
-		}
+		const FRHITextureCreateDesc Desc =
+			FRHITextureCreateDesc::Create2D(TEXT("DefaultTexture2D"), 1, 1, Format)
+			.SetFlags(ETextureCreateFlags::ShaderResource);
+
+		FTexture2DRHIRef Texture = RHICreateTexture(Desc);
+		SetDefaultTextureData2D(Texture, Value);
+		Entry.Texture = CreateRenderTarget(Texture, Desc.DebugName);
+	}
+	else if (Dimension == ETextureDimension::Texture2DArray)
+	{
+		const FRHITextureCreateDesc Desc =
+			FRHITextureCreateDesc::Create2DArray(TEXT("DefaultTexture2DArray"), 1, 1, 1, Format)
+			.SetFlags(ETextureCreateFlags::ShaderResource);
+
+		FTexture2DArrayRHIRef Texture = RHICreateTexture(Desc);
+		SetDefaultTextureData2DArray(Texture, Value);
+		Entry.Texture = CreateRenderTarget(Texture, Desc.DebugName);
+	}
+	else if (Dimension == ETextureDimension::Texture3D)
+	{
+		const FRHITextureCreateDesc Desc =
+			FRHITextureCreateDesc::Create3D(TEXT("DefaultTexture3D"), 1, 1, 1, Format)
+			.SetFlags(ETextureCreateFlags::ShaderResource);
+
+		FTexture3DRHIRef Texture = RHICreateTexture(Desc);
+		SetDefaultTextureData3D(GraphBuilder.RHICmdList, Texture, Value);
+		Entry.Texture = CreateRenderTarget(Texture, Desc.DebugName);
+	}
+	else if (Dimension == ETextureDimension::TextureCube)
+	{
+		const FRHITextureCreateDesc Desc =
+			FRHITextureCreateDesc::CreateCube(TEXT("DefaultTextureCube"), 1, Format)
+			.SetFlags(ETextureCreateFlags::ShaderResource);
+
+		FTextureCubeRHIRef Texture = RHICreateTexture(Desc);
+		SetDefaultTextureDataCube(Texture, Value);
+		Entry.Texture = CreateRenderTarget(Texture, Desc.DebugName);
+	}
+	else if (Dimension == ETextureDimension::TextureCubeArray)
+	{
+		const FRHITextureCreateDesc Desc =
+			FRHITextureCreateDesc::CreateCubeArray(TEXT("DefaultTextureCubeArray"), 1, 1, Format)
+			.SetFlags(ETextureCreateFlags::ShaderResource);
+
+		FTextureCubeRHIRef Texture = RHICreateTexture(Desc);
+		SetDefaultTextureDataCube(Texture, Value);
+		Entry.Texture = CreateRenderTarget(Texture, Desc.DebugName);
+	}
+	else
+	{
+		return nullptr;
 	}
 
 	Index = DefaultTextures.Add(Entry);
