Improved mesh sorting on mobile. Better front to back sorting inside for meshes that belong to the same PSO

diff --git a/Engine/Source/Runtime/Renderer/Private/MeshDrawCommands.cpp b/Engine/Source/Runtime/Renderer/Private/MeshDrawCommands.cpp
--- a/Engine/Source/Runtime/Renderer/Private/MeshDrawCommands.cpp
+++ b/Engine/Source/Runtime/Renderer/Private/MeshDrawCommands.cpp
@@ -195,6 +195,53 @@ void UpdateTranslucentMeshSortKeys(
 	}
 }
 
+/**
+* Merge mobile BasePass with BasePassCSM based on CSM visibility in order to select appropriate shader for given command.
+*/
+void MergeMobileBasePassMeshDrawCommands(
+	const FMobileCSMVisibilityInfo& MobileCSMVisibilityInfo,
+	int32 ScenePrimitiveNum,
+	FMeshCommandOneFrameArray& MeshCommands,
+	FMeshCommandOneFrameArray& MeshCommandsCSM
+	)
+{
+	if (MobileCSMVisibilityInfo.bMobileDynamicCSMInUse)
+	{
+		// determine per view CSM visibility
+		if (MobileCSMVisibilityInfo.bAlwaysUseCSM)
+		{
+#if DO_CHECK
+			check(!MeshCommandsCSM.Num() || MeshCommands.Num() == MeshCommandsCSM.Num());
+			for (int32 Index = 0; Index < MeshCommandsCSM.Num(); ++Index)
+			{
+				check(MeshCommands[Index].PrimitiveIdInfo.ScenePrimitiveId == MeshCommandsCSM[Index].PrimitiveIdInfo.ScenePrimitiveId);
+			}
+#endif
+			if (MeshCommandsCSM.Num() > 0)
+			{
+				MeshCommands = MoveTemp(MeshCommandsCSM);
+			}
+		}
+		else
+		{
+			checkf(MeshCommands.Num() == MeshCommandsCSM.Num(), TEXT("VisibleMeshDrawCommands of BasePass and MobileBasePassCSM are expected to match."));
+			for (int32 i = MeshCommands.Num() - 1; i >= 0; --i)
+			{
+				FVisibleMeshDrawCommand& MeshCommand = MeshCommands[i];
+				FVisibleMeshDrawCommand& MeshCommandCSM = MeshCommandsCSM[i];
+
+				if (MeshCommand.PrimitiveIdInfo.ScenePrimitiveId < ScenePrimitiveNum && MobileCSMVisibilityInfo.MobilePrimitiveCSMReceiverVisibilityMap[MeshCommand.PrimitiveIdInfo.ScenePrimitiveId])
+				{
+					checkf(MeshCommand.PrimitiveIdInfo.ScenePrimitiveId == MeshCommandCSM.PrimitiveIdInfo.ScenePrimitiveId, TEXT("VisibleMeshDrawCommands of BasePass and MobileBasePassCSM are expected to match."));
+					// Use CSM's VisibleMeshDrawCommand.
+					MeshCommand = MeshCommandCSM;
+				}
+			}
+			MeshCommandsCSM.Reset();
+		}
+	}
+}
+
 static uint64 GetMobileBasePassSortKey_FrontToBack(bool bMasked, bool bBackground, uint32 PipelineId, int32 StateBucketId, float PrimitiveDistance)
 {
 	union
@@ -202,8 +249,8 @@ static uint64 GetMobileBasePassSortKey_FrontToBack(bool bMasked, bool bBackgroun
 		uint64 PackedInt;
 		struct
 		{
-			uint64 StateBucketId : 27; 		// Order by state bucket
-			uint64 PipelineId : 20;			// Order by PSO
+			uint64 StateBucketId : 25; 		// Order by state bucket
+			uint64 PipelineId : 22;			// Order by PSO
 			uint64 DepthBits : 15;			// Order by primitive depth
 			uint64 Background : 1;			// Non-background meshes first 
 			uint64 Masked : 1;				// Non-masked first
@@ -219,92 +266,55 @@ static uint64 GetMobileBasePassSortKey_FrontToBack(bool bMasked, bool bBackgroun
 	Key.Fields.DepthBits = ((-int32(F2I.I >> 31) | 0x80000000) ^ F2I.I) >> 17;
 	Key.Fields.PipelineId = PipelineId;
 	Key.Fields.StateBucketId = StateBucketId;
-	
+
 	return Key.PackedInt;
 }
 
-static uint64 GetMobileBasePassSortKey_ByState(bool bMasked, bool bBackground, int32 PipelineId, int32 StateBucketId, float PipelineDistance, float PrimitiveDistance)
+static uint32 DepthBuckets_4(float Distance)
+{
+	union FFloatToInt { float F; uint32 I; };
+	FFloatToInt F2I;
+
+	F2I.F = Distance;
+	// discard distance shorter than 64
+	int32 Exp = ((F2I.I >> 23u) & 0xff) - (127 + 6);
+	// 16 buckets, ranging from 64 to 2km
+	return (uint32)FMath::Clamp<int32>(Exp, 0, 15);
+}
+
+static uint64 GetMobileBasePassSortKey_ByState(bool bMasked, bool bBackground, uint32 PipelineId, uint32 StateBucketId, float PipelineDistance, float StateBucketDistance, float PrimitiveDistance)
 {
 	// maximum primitive distance for bucketing, aprox 10.5km 
 	constexpr float PrimitiveMaxDepth = (1 << 20);
-	constexpr float PrimitiveDepthQuantization = ((1 << 14) - 1);
-	
+	constexpr float PrimitiveDepthQuantization = ((1 << 12) - 1);
+
 	union
 	{
 		uint64 PackedInt;
 		struct
 		{
-			uint64 DepthBits : 14;			// Order by primitive depth
+			uint64 DepthBits : 12;			// Order by primitive depth
 			uint64 StateBucketId : 20; 		// Order by state bucket
-			uint64 PipelineId : 20;			// Order by PSO
-			uint64 PipelineDepthBits : 8;	// Order PSOs front to back
+			uint64 StateBucketDepthBits : 4;// Order state buckets front to back
+			uint64 PipelineId : 22;			// Order by PSO
+			uint64 PipelineDepthBits : 4;	// Order PSOs front to back
 			uint64 Background : 1;			// Non-background meshes first 
 			uint64 Masked : 1;				// Non-masked first
 		} Fields;
 	} Key;
 
-	union FFloatToInt { float F; uint32 I; };
-	FFloatToInt F2I;
-
 	Key.PackedInt = 0;
 	Key.Fields.Masked = bMasked;
 	Key.Fields.Background = bBackground;
-	F2I.F = PipelineDistance / PrimitiveMaxDepth;
-	Key.Fields.PipelineDepthBits = (F2I.I >> 23) & 0xff; // 8 bit exponent
+	Key.Fields.PipelineDepthBits = DepthBuckets_4(PipelineDistance);
 	Key.Fields.PipelineId = PipelineId;
+	Key.Fields.StateBucketDepthBits = DepthBuckets_4(StateBucketDistance);
 	Key.Fields.StateBucketId = StateBucketId;
-	Key.Fields.DepthBits = int32((FMath::Min<float>(PrimitiveDistance, PrimitiveMaxDepth) / PrimitiveMaxDepth) * PrimitiveDepthQuantization);
+	Key.Fields.DepthBits = uint32((FMath::Min<float>(PrimitiveDistance, PrimitiveMaxDepth) / PrimitiveMaxDepth) * PrimitiveDepthQuantization);
 
 	return Key.PackedInt;
 }
 
-/**
-* Merge mobile BasePass with BasePassCSM based on CSM visibility in order to select appropriate shader for given command.
-*/
-void MergeMobileBasePassMeshDrawCommands(
-	const FMobileCSMVisibilityInfo& MobileCSMVisibilityInfo,
-	int32 ScenePrimitiveNum,
-	FMeshCommandOneFrameArray& MeshCommands,
-	FMeshCommandOneFrameArray& MeshCommandsCSM
-	)
-{
-	if (MobileCSMVisibilityInfo.bMobileDynamicCSMInUse)
-	{
-		// determine per view CSM visibility
-		if (MobileCSMVisibilityInfo.bAlwaysUseCSM)
-		{
-#if DO_CHECK
-			check(!MeshCommandsCSM.Num() || MeshCommands.Num() == MeshCommandsCSM.Num());
-			for (int32 Index = 0; Index < MeshCommandsCSM.Num(); ++Index)
-			{
-				check(MeshCommands[Index].PrimitiveIdInfo.ScenePrimitiveId == MeshCommandsCSM[Index].PrimitiveIdInfo.ScenePrimitiveId);
-			}
-#endif
-			if (MeshCommandsCSM.Num() > 0)
-			{
-				MeshCommands = MoveTemp(MeshCommandsCSM);
-			}
-		}
-		else
-		{
-			checkf(MeshCommands.Num() == MeshCommandsCSM.Num(), TEXT("VisibleMeshDrawCommands of BasePass and MobileBasePassCSM are expected to match."));
-			for (int32 i = MeshCommands.Num() - 1; i >= 0; --i)
-			{
-				FVisibleMeshDrawCommand& MeshCommand = MeshCommands[i];
-				FVisibleMeshDrawCommand& MeshCommandCSM = MeshCommandsCSM[i];
-
-				if (MeshCommand.PrimitiveIdInfo.ScenePrimitiveId < ScenePrimitiveNum && MobileCSMVisibilityInfo.MobilePrimitiveCSMReceiverVisibilityMap[MeshCommand.PrimitiveIdInfo.ScenePrimitiveId])
-				{
-					checkf(MeshCommand.PrimitiveIdInfo.ScenePrimitiveId == MeshCommandCSM.PrimitiveIdInfo.ScenePrimitiveId, TEXT("VisibleMeshDrawCommands of BasePass and MobileBasePassCSM are expected to match."));
-					// Use CSM's VisibleMeshDrawCommand.
-					MeshCommand = MeshCommandCSM;
-				}
-			}
-			MeshCommandsCSM.Reset();
-		}
-	}
-}
-
 /**
 * Compute mesh sort keys for the mobile base pass
 */
@@ -315,7 +325,10 @@ void UpdateMobileBasePassMeshSortKeys(
 )
 {
 	QUICK_SCOPE_CYCLE_COUNTER(STAT_UpdateMobileBasePassMeshSortKeys);
-	
+
+	// Object radius past which we treat object as part of 'background'
+	constexpr float MIN_BACKGROUND_OBJECT_RADIUS = 500000;
+
 	int32 NumCmds = VisibleMeshCommands.Num();
 	int32 MeshSortingMethod = CVarMobileMeshSortingMethod.GetValueOnAnyThread();
 	
@@ -326,41 +339,58 @@ void UpdateMobileBasePassMeshSortKeys(
 		{
 			FVisibleMeshDrawCommand& Cmd = VisibleMeshCommands[CmdIdx];
 			// Set in MobileBasePass.cpp - GetBasePassStaticSortKey;
-			bool bMasked = Cmd.SortKey.PackedData & 0x1 ? true : false; 
+			bool bMasked = Cmd.SortKey.PackedData & 0x1 ? true : false;
 			bool bBackground = Cmd.SortKey.PackedData & 0x2 ? true : false;
 			float PrimitiveDistance = 0;
 			if (Cmd.PrimitiveIdInfo.ScenePrimitiveId < ScenePrimitiveBounds.Num())
 			{
 				const FPrimitiveBounds& PrimitiveBounds = ScenePrimitiveBounds[Cmd.PrimitiveIdInfo.ScenePrimitiveId];
 				PrimitiveDistance = (PrimitiveBounds.BoxSphereBounds.Origin - ViewOrigin).Size();
-				bBackground|= (PrimitiveBounds.BoxSphereBounds.SphereRadius > HALF_WORLD_MAX / 4.0f);
+				bBackground |= (PrimitiveBounds.BoxSphereBounds.SphereRadius > MIN_BACKGROUND_OBJECT_RADIUS);
 			}
 
 			uint32 PipelineId = Cmd.MeshDrawCommand->CachedPipelineId.GetId();
 			// use state bucket if dynamic instancing is enabled, otherwise identify same meshes by index buffer resource
-			int32 StateBucketId = Cmd.StateBucketId >=0 ? Cmd.StateBucketId : PointerHash(Cmd.MeshDrawCommand->IndexBuffer);
+			uint32 StateBucketId = Cmd.StateBucketId >= 0 ? Cmd.StateBucketId : PointerHash(Cmd.MeshDrawCommand->IndexBuffer);
 			Cmd.SortKey.PackedData = GetMobileBasePassSortKey_FrontToBack(bMasked, bBackground, PipelineId, StateBucketId, PrimitiveDistance);
 		}
 	}
-	else // prefer state then distance
+	else
 	{
-		TMap<uint32, float> PipelineDistances;
+		struct FPipelineDistance
+		{
+			double DistanceSum = 0.0;
+			int32 Num = 0;
+		};
+
+		TMap<uint32, FPipelineDistance> PipelineDistances;
 		PipelineDistances.Reserve(256);
-				
-		// pre-compute distance to a group of meshes that share same PSO
+
+		TMap<uint64, float> StateBucketDistances;
+		StateBucketDistances.Reserve(512);
+
+		// pre-compute a distance to a group of meshes that share same PSO and state
 		for (int32 CmdIdx = 0; CmdIdx < NumCmds; ++CmdIdx)
 		{
 			const FVisibleMeshDrawCommand& Cmd = VisibleMeshCommands[CmdIdx];
-			float PrimitiveDistance = 0;
+			float PrimitiveDistance = 0.f;
 			if (Cmd.PrimitiveIdInfo.ScenePrimitiveId < ScenePrimitiveBounds.Num())
 			{
 				const FPrimitiveBounds& PrimitiveBounds = ScenePrimitiveBounds[Cmd.PrimitiveIdInfo.ScenePrimitiveId];
 				PrimitiveDistance = (PrimitiveBounds.BoxSphereBounds.Origin - ViewOrigin).Size();
 			}
 
-			float& PipelineDistance = PipelineDistances.FindOrAdd(Cmd.MeshDrawCommand->CachedPipelineId.GetId());
-			// not sure what could be better: average distance, max or min
-			PipelineDistance = FMath::Max(PipelineDistance, PrimitiveDistance);
+			// group meshes by PSO and find distance to each group
+			uint32 PipelineId = Cmd.MeshDrawCommand->CachedPipelineId.GetId();
+			FPipelineDistance& PipelineDistance = PipelineDistances.FindOrAdd(PipelineId);
+			PipelineDistance.DistanceSum += PrimitiveDistance;
+			PipelineDistance.Num += 1;
+
+			// group meshes by StateBucketId or index buffer resource and find minimum distance to each group 
+			uint32 StateBucketId = Cmd.StateBucketId >= 0 ? Cmd.StateBucketId : PointerHash(Cmd.MeshDrawCommand->IndexBuffer);
+			uint64 PipelineAndStateBucketId = ((uint64)PipelineId << 32u) | StateBucketId;
+			float& StateBucketDistance = StateBucketDistances.FindOrAdd(PipelineAndStateBucketId, BIG_NUMBER);
+			StateBucketDistance = FMath::Min(StateBucketDistance, PrimitiveDistance);
 		}
 
 		// compute sort key for each mesh command
@@ -368,26 +398,29 @@ void UpdateMobileBasePassMeshSortKeys(
 		{
 			FVisibleMeshDrawCommand& Cmd = VisibleMeshCommands[CmdIdx];
 			// Set in MobileBasePass.cpp - GetBasePassStaticSortKey;
-			bool bMasked = Cmd.SortKey.PackedData & 0x1 ? true : false; 
+			bool bMasked = Cmd.SortKey.PackedData & 0x1 ? true : false;
 			bool bBackground = Cmd.SortKey.PackedData & 0x2 ? true : false;
 			float PrimitiveDistance = 0;
 			if (Cmd.PrimitiveIdInfo.ScenePrimitiveId < ScenePrimitiveBounds.Num())
 			{
 				const FPrimitiveBounds& PrimitiveBounds = ScenePrimitiveBounds[Cmd.PrimitiveIdInfo.ScenePrimitiveId];
 				PrimitiveDistance = (PrimitiveBounds.BoxSphereBounds.Origin - ViewOrigin).Size();
-				bBackground|= (PrimitiveBounds.BoxSphereBounds.SphereRadius > HALF_WORLD_MAX / 4.0f);
+				bBackground |= (PrimitiveBounds.BoxSphereBounds.SphereRadius > MIN_BACKGROUND_OBJECT_RADIUS);
 			}
-			
-			int32 PipelineId = Cmd.MeshDrawCommand->CachedPipelineId.GetId();
-			float PipelineDistance = PipelineDistances.FindRef(PipelineId);
-			// use state bucket if dynamic instancing is enabled, otherwise identify same meshes by index buffer resource
-			int32 StateBucketId = Cmd.StateBucketId >=0 ? Cmd.StateBucketId : PointerHash(Cmd.MeshDrawCommand->IndexBuffer);
-			Cmd.SortKey.PackedData = GetMobileBasePassSortKey_ByState(bMasked, bBackground, PipelineId, StateBucketId, PipelineDistance, PrimitiveDistance);
+
+			uint32 PipelineId = Cmd.MeshDrawCommand->CachedPipelineId.GetId();
+			FPipelineDistance PipelineDistance = PipelineDistances.FindRef(PipelineId);
+			float MeanPipelineDistance = PipelineDistance.DistanceSum / PipelineDistance.Num;
+
+			uint32 StateBucketId = Cmd.StateBucketId >= 0 ? Cmd.StateBucketId : PointerHash(Cmd.MeshDrawCommand->IndexBuffer);
+			uint64 PipelineAndStateBucketId = ((uint64)PipelineId << 32u) | StateBucketId;
+			float StateBucketDistance = StateBucketDistances.FindRef(PipelineAndStateBucketId);
+
+			Cmd.SortKey.PackedData = GetMobileBasePassSortKey_ByState(bMasked, bBackground, PipelineId, StateBucketId, MeanPipelineDistance, StateBucketDistance, PrimitiveDistance);
 		}
 	}
 }
 
-
 FORCEINLINE int32 TranslatePrimitiveId(int32 DrawPrimitiveIdIn, int32 DynamicPrimitiveIdOffset, int32 DynamicPrimitiveIdMax)
 {
 	// INDEX_NONE means we defer the translation to later
