Add Strata tile buffer overflow to Strata uniform buffer.

diff --git a/Engine/Source/Runtime/Renderer/Private/Strata/Strata.cpp b/Engine/Source/Runtime/Renderer/Private/Strata/Strata.cpp
--- a/Engine/Source/Runtime/Renderer/Private/Strata/Strata.cpp
+++ b/Engine/Source/Runtime/Renderer/Private/Strata/Strata.cpp
@@ -55,16 +55,16 @@ static TAutoConsoleVariable<int32> CVarMaterialRoughDiffuse(
 
 // STRATA_TODO we keep this for now and can remove it once battletested.
 static TAutoConsoleVariable<int32> CVarClearDuringCategorization(
-	TEXT("r.strata.ClearDuringCategorization"),
+	TEXT("r.Strata.ClearDuringCategorization"),
 	1,
 	TEXT("TEST."),
 	ECVF_RenderThreadSafe);
 
 static TAutoConsoleVariable<int32> CVarStrataTileOverflow(
-	TEXT("r.Material.TileOverflow"),
+	TEXT("r.Strata.TileOverflow"),
 	1.f,
 	TEXT("Scale the number of Strata tile for overflowing tiles containing multi-BSDFs pixels. (0: 0%, 1: 100%. Default 1.0f)."),
-	ECVF_ReadOnly | ECVF_RenderThreadSafe);
+	ECVF_RenderThreadSafe);
 
 IMPLEMENT_GLOBAL_SHADER_PARAMETER_STRUCT(FStrataGlobalUniformParameters, "Strata");
 
@@ -276,6 +276,7 @@ void InitialiseStrataFrameSceneData(FSceneRenderer& SceneRenderer, FRDGBuilder&
 			AddClearUAVPass(GraphBuilder, GraphBuilder.CreateUAV(StrataSceneData.BSDFOffsetTexture), 0u);
 
 			StrataSceneData.BSDFTileDispatchIndirectBuffer = GraphBuilder.CreateBuffer(FRDGBufferDesc::CreateIndirectDesc<FRHIDispatchIndirectParameters>(1), TEXT("Strata.StrataBSDFTileDispatchIndirectBuffer"));
+			StrataSceneData.BSDFTileCountBuffer = GraphBuilder.CreateBuffer(FRDGBufferDesc::CreateBufferDesc(4, 1), TEXT("Strata.BSDFTileCount"));
 		}
 	}
 	else
@@ -371,11 +372,13 @@ void BindStrataGlobalUniformParameters(FRDGBuilder& GraphBuilder, FStrataSceneDa
 		OutStrataUniformParameters.OpaqueRoughRefractionTexture = StrataSceneData->OpaqueRoughRefractionTexture;
 		OutStrataUniformParameters.BSDFTileTexture = StrataSceneData->BSDFTileTexture;
 		OutStrataUniformParameters.BSDFOffsetTexture = StrataSceneData->BSDFOffsetTexture;
+		OutStrataUniformParameters.BSDFTileCountBuffer = GraphBuilder.CreateSRV(StrataSceneData->BSDFTileCountBuffer, PF_R32_UINT);
 	}
 	else
 	{
 		const FRDGSystemTextures& SystemTextures = FRDGSystemTextures::Get(GraphBuilder);
 		FRDGTextureRef DefaultTextureArray = GSystemTextures.GetDefaultTexture(GraphBuilder, ETextureDimension::Texture2DArray, EPixelFormat::PF_R32_UINT, FClearValueBinding::Transparent);
+		FRDGBufferSRVRef DefaultBuffer = GraphBuilder.CreateSRV(GSystemTextures.GetDefaultBuffer(GraphBuilder, 4, 0u), PF_R32_UINT);
 		OutStrataUniformParameters.bRoughDiffuse = 0;
 		OutStrataUniformParameters.MaxBytesPerPixel = 0;
 		OutStrataUniformParameters.TileSize = 0;
@@ -387,6 +390,7 @@ void BindStrataGlobalUniformParameters(FRDGBuilder& GraphBuilder, FStrataSceneDa
 		OutStrataUniformParameters.OpaqueRoughRefractionTexture = SystemTextures.Black;
 		OutStrataUniformParameters.BSDFTileTexture = SystemTextures.Black;
 		OutStrataUniformParameters.BSDFOffsetTexture = SystemTextures.Black;
+		OutStrataUniformParameters.BSDFTileCountBuffer = DefaultBuffer;
 	}
 }
 
@@ -1032,9 +1036,8 @@ void AddStrataMaterialClassificationPass(FRDGBuilder& GraphBuilder, const FMinim
 		}
 
 		// Compute BSDF tile index and material read offset
-		FRDGBufferRef BSDFTileCountBuffer = BSDFTileCountBuffer = GraphBuilder.CreateBuffer(FRDGBufferDesc::CreateBufferDesc(4, 1), TEXT("Strata.BSDFTileCount"));
 		{
-			FRDGBufferUAVRef RWBSDFTileCountBuffer = GraphBuilder.CreateUAV(BSDFTileCountBuffer, PF_R32_UINT);
+			FRDGBufferUAVRef RWBSDFTileCountBuffer = GraphBuilder.CreateUAV(View.StrataSceneData->BSDFTileCountBuffer, PF_R32_UINT);
 			AddClearUAVPass(GraphBuilder, RWBSDFTileCountBuffer, 0u);
 
 			FStrataBSDFTilePassCS::FPermutationDomain PermutationVector;
@@ -1069,7 +1072,7 @@ void AddStrataMaterialClassificationPass(FRDGBuilder& GraphBuilder, const FMinim
 			TShaderMapRef<FStrataBSDFTilePrepareArgsPassCS> ComputeShader(View.ShaderMap);
 			FStrataBSDFTilePrepareArgsPassCS::FParameters* PassParameters = GraphBuilder.AllocParameters<FStrataBSDFTilePrepareArgsPassCS::FParameters>();
 			PassParameters->TileCount_Primary = View.StrataSceneData->TileCount_Primary;
-			PassParameters->TileDrawIndirectDataBuffer = GraphBuilder.CreateSRV(BSDFTileCountBuffer, PF_R32_UINT);
+			PassParameters->TileDrawIndirectDataBuffer = GraphBuilder.CreateSRV(View.StrataSceneData->BSDFTileCountBuffer, PF_R32_UINT);
 			PassParameters->TileDispatchIndirectDataBuffer = GraphBuilder.CreateUAV(View.StrataSceneData->BSDFTileDispatchIndirectBuffer, PF_R32_UINT);
 
 			FComputeShaderUtils::AddPass(
diff --git a/Engine/Source/Runtime/Renderer/Private/Strata/Strata.h b/Engine/Source/Runtime/Renderer/Private/Strata/Strata.h
--- a/Engine/Source/Runtime/Renderer/Private/Strata/Strata.h
+++ b/Engine/Source/Runtime/Renderer/Private/Strata/Strata.h
@@ -47,6 +47,7 @@ BEGIN_GLOBAL_SHADER_PARAMETER_STRUCT(FStrataGlobalUniformParameters, RENDERER_AP
 	SHADER_PARAMETER_RDG_TEXTURE(Texture2D<float3>, OpaqueRoughRefractionTexture)
 	SHADER_PARAMETER_RDG_TEXTURE(Texture2D<uint>, BSDFOffsetTexture)
 	SHADER_PARAMETER_RDG_TEXTURE(Texture2D<uint>, BSDFTileTexture)
+	SHADER_PARAMETER_RDG_BUFFER_SRV(Buffer<uint>, BSDFTileCountBuffer)
 END_GLOBAL_SHADER_PARAMETER_STRUCT()
 
 // This must map to the STRATA_TILE_TYPE defines.
@@ -98,6 +99,7 @@ struct FStrataSceneData
 
 	FRDGTextureRef BSDFOffsetTexture;
 	FRDGTextureRef BSDFTileTexture;
+	FRDGBufferRef  BSDFTileCountBuffer;
 	FRDGBufferRef  BSDFTileDispatchIndirectBuffer;
 
 	// Used when the subsurface luminance is separated from the scene color
