Remove STRATA_DATA_TILE_XXX to avoid confusion with STRATA_TILE.

diff --git a/Engine/Shaders/Shared/StrataDefinitions.h b/Engine/Shaders/Shared/StrataDefinitions.h
--- a/Engine/Shaders/Shared/StrataDefinitions.h
+++ b/Engine/Shaders/Shared/StrataDefinitions.h
@@ -30,10 +30,7 @@
 
 // The size of strata material classification tiles on screen
 #define STRATA_TILE_SIZE					8
-
-// This is used to round up the strata material buffer for better cache locality on read/write.
-#define STRATA_DATA_TILE_SIZE				8
-#define STRATA_DATA_TILE_SIZE_DIV_AS_SHIFT	3
+#define STRATA_TILE_SIZE_DIV_AS_SHIFT	    3
 
 // The default thickness of a layer is considered to be 0.01 centimeter = 0.1 millimeter
 #define STRATA_LAYER_DEFAULT_THICKNESS_CM	0.01f
diff --git a/Engine/Source/Runtime/Renderer/Private/Strata/Strata.cpp b/Engine/Source/Runtime/Renderer/Private/Strata/Strata.cpp
--- a/Engine/Source/Runtime/Renderer/Private/Strata/Strata.cpp
+++ b/Engine/Source/Runtime/Renderer/Private/Strata/Strata.cpp
@@ -147,8 +147,8 @@ static FIntPoint GetStrataTextureTileResolution(const FIntPoint& InResolution, f
 	FIntPoint Out = InResolution;
 	if (Strata::IsStrataEnabled())
 	{
-		Out.X = FMath::DivideAndRoundUp(Out.X, STRATA_DATA_TILE_SIZE);
-		Out.Y = FMath::DivideAndRoundUp(Out.Y, STRATA_DATA_TILE_SIZE);
+		Out.X = FMath::DivideAndRoundUp(Out.X, STRATA_TILE_SIZE);
+		Out.Y = FMath::DivideAndRoundUp(Out.Y, STRATA_TILE_SIZE);
 		Out.Y += FMath::CeilToInt(Out.Y * FMath::Clamp(Overflow, 0.f, 1.0f));
 	}
 	return Out;
@@ -161,7 +161,7 @@ static FIntPoint GetStrataTextureTileResolution(const FIntPoint& InResolution)
 
 FIntPoint GetStrataTextureResolution(const FIntPoint& InResolution)
 {
-	return GetStrataTextureTileResolution(InResolution) * STRATA_DATA_TILE_SIZE;
+	return GetStrataTextureTileResolution(InResolution) * STRATA_TILE_SIZE;
 }
 
 void InitialiseStrataFrameSceneData(FSceneRenderer& SceneRenderer, FRDGBuilder& GraphBuilder)
@@ -172,8 +172,8 @@ void InitialiseStrataFrameSceneData(FSceneRenderer& SceneRenderer, FRDGBuilder&
 	auto UpdateMaterialBufferToTiledResolution = [](FIntPoint InBufferSizeXY, FIntPoint& OutMaterialBufferSizeXY)
 	{
 		// We need to allocate enough for the tiled memory addressing to always work
-		OutMaterialBufferSizeXY.X = FMath::DivideAndRoundUp(InBufferSizeXY.X, STRATA_DATA_TILE_SIZE) * STRATA_DATA_TILE_SIZE;
-		OutMaterialBufferSizeXY.Y = FMath::DivideAndRoundUp(InBufferSizeXY.Y, STRATA_DATA_TILE_SIZE) * STRATA_DATA_TILE_SIZE;
+		OutMaterialBufferSizeXY.X = FMath::DivideAndRoundUp(InBufferSizeXY.X, STRATA_TILE_SIZE) * STRATA_TILE_SIZE;
+		OutMaterialBufferSizeXY.Y = FMath::DivideAndRoundUp(InBufferSizeXY.Y, STRATA_TILE_SIZE) * STRATA_TILE_SIZE;
 	};
 
 	FIntPoint MaterialBufferSizeXY;
@@ -189,7 +189,7 @@ void InitialiseStrataFrameSceneData(FSceneRenderer& SceneRenderer, FRDGBuilder&
 		const uint32 RoundToValue = 4u;
 		StrataSceneData.MaxBytesPerPixel = FMath::DivideAndRoundUp(MaterialConservativeByteCountPerPixel, RoundToValue) * RoundToValue;
 
-		const FIntPoint TileResolution(FMath::DivideAndRoundUp(SceneTextureExtent.X, STRATA_DATA_TILE_SIZE), FMath::DivideAndRoundUp(SceneTextureExtent.Y, STRATA_DATA_TILE_SIZE));
+		const FIntPoint TileResolution(FMath::DivideAndRoundUp(SceneTextureExtent.X, STRATA_TILE_SIZE), FMath::DivideAndRoundUp(SceneTextureExtent.Y, STRATA_TILE_SIZE));
 
 		const TCHAR* StrataTileListBufferNames[EStrataTileType::ECount] =
 		{
@@ -272,7 +272,7 @@ void InitialiseStrataFrameSceneData(FSceneRenderer& SceneRenderer, FRDGBuilder&
 			StrataSceneData.BSDFTileTexture = GraphBuilder.CreateTexture(FRDGTextureDesc::Create2D(StrataSceneData.TileCount_Total, PF_R32_UINT, FClearValueBinding::None, TexCreate_UAV | TexCreate_ShaderResource), TEXT("Strata.BSDFTiles"));
 			AddClearUAVPass(GraphBuilder, GraphBuilder.CreateUAV(StrataSceneData.BSDFTileTexture), 0u);
 
-			StrataSceneData.BSDFOffsetTexture = GraphBuilder.CreateTexture(FRDGTextureDesc::Create2D(StrataSceneData.TileCount_Primary * STRATA_DATA_TILE_SIZE, PF_R32_UINT, FClearValueBinding::None, TexCreate_UAV | TexCreate_ShaderResource), TEXT("Strata.BSDFOffsets"));
+			StrataSceneData.BSDFOffsetTexture = GraphBuilder.CreateTexture(FRDGTextureDesc::Create2D(StrataSceneData.TileCount_Primary * STRATA_TILE_SIZE, PF_R32_UINT, FClearValueBinding::None, TexCreate_UAV | TexCreate_ShaderResource), TEXT("Strata.BSDFOffsets"));
 			AddClearUAVPass(GraphBuilder, GraphBuilder.CreateUAV(StrataSceneData.BSDFOffsetTexture), 0u);
 
 			StrataSceneData.BSDFTileDispatchIndirectBuffer = GraphBuilder.CreateBuffer(FRDGBufferDesc::CreateIndirectDesc<FRHIDispatchIndirectParameters>(1), TEXT("Strata.StrataBSDFTileDispatchIndirectBuffer"));
@@ -363,8 +363,8 @@ void BindStrataGlobalUniformParameters(FRDGBuilder& GraphBuilder, FStrataSceneDa
 	{
 		OutStrataUniformParameters.bRoughDiffuse = StrataSceneData->bRoughDiffuse ? 1u : 0u;
 		OutStrataUniformParameters.MaxBytesPerPixel = StrataSceneData->MaxBytesPerPixel;
-		OutStrataUniformParameters.TileSize = STRATA_DATA_TILE_SIZE;
-		OutStrataUniformParameters.TileSizeLog2 = STRATA_DATA_TILE_SIZE_DIV_AS_SHIFT;
+		OutStrataUniformParameters.TileSize = STRATA_TILE_SIZE;
+		OutStrataUniformParameters.TileSizeLog2 = STRATA_TILE_SIZE_DIV_AS_SHIFT;
 		OutStrataUniformParameters.TileCount = StrataSceneData->TileCount_Primary;
 		OutStrataUniformParameters.MaterialTextureArray = StrataSceneData->MaterialTextureArray;
 		OutStrataUniformParameters.TopLayerTexture = StrataSceneData->TopLayerTexture;
@@ -1045,7 +1045,7 @@ void AddStrataMaterialClassificationPass(FRDGBuilder& GraphBuilder, const FMinim
 			TShaderMapRef<FStrataBSDFTilePassCS> ComputeShader(View.ShaderMap, PermutationVector);
 			FStrataBSDFTilePassCS::FParameters* PassParameters = GraphBuilder.AllocParameters<FStrataBSDFTilePassCS::FParameters>();
 			PassParameters->ViewUniformBuffer = View.ViewUniformBuffer;
-			PassParameters->TileSizeLog2 = STRATA_DATA_TILE_SIZE_DIV_AS_SHIFT;
+			PassParameters->TileSizeLog2 = STRATA_TILE_SIZE_DIV_AS_SHIFT;
 			PassParameters->TileCount_Primary = View.StrataSceneData->TileCount_Primary;
 			PassParameters->ViewResolution = View.ViewRect.Size();
 			PassParameters->MaxBytesPerPixel = View.StrataSceneData->MaxBytesPerPixel;
