Strata - Fixed crash in Lumen due to VS/PS mismatch

diff --git a/Engine/Source/Runtime/Renderer/Private/IndirectLightRendering.cpp b/Engine/Source/Runtime/Renderer/Private/IndirectLightRendering.cpp
--- a/Engine/Source/Runtime/Renderer/Private/IndirectLightRendering.cpp
+++ b/Engine/Source/Runtime/Renderer/Private/IndirectLightRendering.cpp
@@ -1161,7 +1161,7 @@ void FDeferredShadingSceneRenderer::RenderDiffuseIndirectAndAmbientOcclusion(
 
 				Strata::FStrataTilePassVS::FPermutationDomain VSPermutationVector;
 				VSPermutationVector.Set< Strata::FStrataTilePassVS::FEnableDebug >(false);
-				VSPermutationVector.Set< Strata::FStrataTilePassVS::FEnableTexCoordScreenVector >(true);
+				VSPermutationVector.Set< Strata::FStrataTilePassVS::FEnableTexCoordScreenVector >(false);
 				TShaderMapRef<Strata::FStrataTilePassVS> TileVertexShader(View.ShaderMap, VSPermutationVector);
 
 				ClearUnusedGraphResources(PixelShader, PassParameters);
