Prevent shifts by a negative value.

- Addresses static analysis warning: warning V610: Undefined behavior. Check the shift operator '<<'. The right operand is negative ('....' = [-1..254])
- Texture Resolution will be 1.

diff --git a/Engine/Source/Runtime/Renderer/Private/DebugViewModeRendering.cpp b/Engine/Source/Runtime/Renderer/Private/DebugViewModeRendering.cpp
--- a/Engine/Source/Runtime/Renderer/Private/DebugViewModeRendering.cpp
+++ b/Engine/Source/Runtime/Renderer/Private/DebugViewModeRendering.cpp
@@ -488,7 +488,7 @@ void FDebugViewModeImplementation::GetDebugViewModeShaderBindings(
 
 							if (Texture->IsStreamable())
 							{
-								TextureResolution = 1 << (Texture->GetResource()->GetCurrentMipCount() - 1);
+								TextureResolution = 1 << FMath::Max((Texture->GetResource()->GetCurrentMipCount() - 1), 0);
 							}
 							else
 							{
@@ -514,7 +514,7 @@ void FDebugViewModeImplementation::GetDebugViewModeShaderBindings(
 
 						if (Texture->IsStreamable())
 						{
-							TextureResolution = 1 << (Texture->GetResource()->GetCurrentMipCount() - 1);
+							TextureResolution = 1 << FMath::Max((Texture->GetResource()->GetCurrentMipCount() - 1), 0);
 						}
 						else
 						{
