Fix static analysis warning.
#rnx

diff --git a/Engine/Source/Runtime/Renderer/Private/VT/RuntimeVirtualTextureSceneProxy.cpp b/Engine/Source/Runtime/Renderer/Private/VT/RuntimeVirtualTextureSceneProxy.cpp
--- a/Engine/Source/Runtime/Renderer/Private/VT/RuntimeVirtualTextureSceneProxy.cpp
+++ b/Engine/Source/Runtime/Renderer/Private/VT/RuntimeVirtualTextureSceneProxy.cpp
@@ -93,11 +93,8 @@ FRuntimeVirtualTextureSceneProxy::FRuntimeVirtualTextureSceneProxy(URuntimeVirtu
 		ENQUEUE_RENDER_COMMAND(GetProducerHandle)(
 			[this](FRHICommandList& RHICmdList)
 			{
-				if (VirtualTexture != nullptr)
-				{
-					ProducerHandle = VirtualTexture->GetProducerHandle();
-					SpaceID = VirtualTexture->GetAllocatedVirtualTexture()->GetSpaceID();
-				}
+				ProducerHandle = VirtualTexture->GetProducerHandle();
+				SpaceID = VirtualTexture->GetAllocatedVirtualTexture()->GetSpaceID();
 			});
 	}
 }
