Path Tracer: Fix missing scene invalidations when changing skylight or moving atmosphere/fog components

diff --git a/Engine/Source/Runtime/Renderer/Private/RendererScene.cpp b/Engine/Source/Runtime/Renderer/Private/RendererScene.cpp
--- a/Engine/Source/Runtime/Renderer/Private/RendererScene.cpp
+++ b/Engine/Source/Runtime/Renderer/Private/RendererScene.cpp
@@ -2075,6 +2075,7 @@ void FScene::SetSkyLight(FSkyLightSceneProxy* LightProxy)
 				// The base pass chooses shaders based on whether there's a skylight in the scene, and that is cached in static draw lists
 				Scene->bScenesPrimitivesNeedStaticMeshElementUpdate = true;
 			}
+			Scene->InvalidatePathTracedOutput();
 		});
 }
 
@@ -2109,6 +2110,7 @@ void FScene::DisableSkyLight(FSkyLightSceneProxy* LightProxy)
 		{
 			Scene->bScenesPrimitivesNeedStaticMeshElementUpdate = true;
 		}
+		Scene->InvalidatePathTracedOutput();
 	});
 }
 
@@ -3263,6 +3265,7 @@ void FScene::AddExponentialHeightFog(UExponentialHeightFogComponent* FogComponen
 		{
 			// Create a FExponentialHeightFogSceneInfo for the component in the scene's fog array.
 			new(Scene->ExponentialFogs) FExponentialHeightFogSceneInfo(HeightFogSceneInfo);
+			Scene->InvalidatePathTracedOutput();
 		});
 }
 
@@ -3278,6 +3281,7 @@ void FScene::RemoveExponentialHeightFog(UExponentialHeightFogComponent* FogCompo
 				if(Scene->ExponentialFogs[FogIndex].Component == FogComponent)
 				{
 					Scene->ExponentialFogs.RemoveAt(FogIndex);
+					Scene->InvalidatePathTracedOutput();
 					break;
 				}
 			}
diff --git a/Engine/Source/Runtime/Renderer/Private/SkyAtmosphereRendering.cpp b/Engine/Source/Runtime/Renderer/Private/SkyAtmosphereRendering.cpp
--- a/Engine/Source/Runtime/Renderer/Private/SkyAtmosphereRendering.cpp
+++ b/Engine/Source/Runtime/Renderer/Private/SkyAtmosphereRendering.cpp
@@ -518,6 +518,7 @@ void FScene::AddSkyAtmosphere(FSkyAtmosphereSceneProxy* SkyAtmosphereSceneProxy,
 			{
 				FPlatformAtomics::InterlockedIncrement(&Scene->NumUncachedStaticLightingInteractions);
 			}
+			Scene->InvalidatePathTracedOutput();
 		} );
 }
 
@@ -545,6 +546,7 @@ void FScene::RemoveSkyAtmosphere(FSkyAtmosphereSceneProxy* SkyAtmosphereScenePro
 			{
 				Scene->SkyAtmosphere = nullptr;
 			}
+			Scene->InvalidatePathTracedOutput();
 		} );
 }
 
