UE5_MAIN:  Multi-GPU crash, corruption, and DumpGPU fixes.

* Distance Field streaming GPU mask scope clean up, to avoid startup crash.  Added single high level FRHIGPUMask::All() scope in PrepareDistanceFieldScene, to avoid the need to set the mask in multiple lower level functions, and hopefully cover future modifications to the distance field scene code.  View dependent subset of the code is then masked to the view.  Validation asserts added in case code gets moved around in the future.
* Fixed bug when generating SRVs for structured buffers -- wasn't pointing the SRV to the correct GPU buffer.  Caused GPU hangs and all sorts of corrupted rendering.
* Fixed bug that RHICopyToResolveTarget wasn't using the correct GPU index.  Main symptom was DumpGPU being non-functional on the second GPU, but could cause other artifacts as well.

#rnx

diff --git a/Engine/Source/Runtime/D3D12RHI/Private/D3D12RenderTarget.cpp b/Engine/Source/Runtime/D3D12RHI/Private/D3D12RenderTarget.cpp
--- a/Engine/Source/Runtime/D3D12RHI/Private/D3D12RenderTarget.cpp
+++ b/Engine/Source/Runtime/D3D12RHI/Private/D3D12RenderTarget.cpp
@@ -202,10 +202,11 @@ void FD3D12CommandContext::RHICopyToResolveTarget(FRHITexture* SourceTextureRHI,
 		return;
 	}
 
-	FRHICommandList_RecursiveHazardous RHICmdList(this, FRHIGPUMask::FromIndex(GetGPUIndex()));
+	uint32 GPUIndex = GetGPUIndex();
+	FRHICommandList_RecursiveHazardous RHICmdList(this, FRHIGPUMask::FromIndex(GPUIndex));
 
-	FD3D12Texture* SourceTexture = GetD3D12TextureFromRHITexture(SourceTextureRHI);
-	FD3D12Texture* DestTexture = GetD3D12TextureFromRHITexture(DestTextureRHI);
+	FD3D12Texture* SourceTexture = GetD3D12TextureFromRHITexture(SourceTextureRHI, GPUIndex);
+	FD3D12Texture* DestTexture = GetD3D12TextureFromRHITexture(DestTextureRHI, GPUIndex);
 
 	const FRHITextureDesc& SourceDesc = SourceTexture->GetDesc();
 	const FRHITextureDesc& DestDesc = DestTexture->GetDesc();
diff --git a/Engine/Source/Runtime/D3D12RHI/Private/D3D12View.cpp b/Engine/Source/Runtime/D3D12RHI/Private/D3D12View.cpp
--- a/Engine/Source/Runtime/D3D12RHI/Private/D3D12View.cpp
+++ b/Engine/Source/Runtime/D3D12RHI/Private/D3D12View.cpp
@@ -491,17 +491,19 @@ FRHICOMMAND_MACRO(FD3D12InitializeBufferSRVRHICommand)
 {
 	const FShaderResourceViewInitializer Initializer;
 	FD3D12ShaderResourceView* const View;
+	uint32 GPUIndex;
 
-	FD3D12InitializeBufferSRVRHICommand(const FShaderResourceViewInitializer& InInitializer, FD3D12ShaderResourceView* InView)
+	FD3D12InitializeBufferSRVRHICommand(const FShaderResourceViewInitializer& InInitializer, FD3D12ShaderResourceView* InView, uint32 InGPUIndex)
 		: Initializer(InInitializer)
 		, View(InView)
+		, GPUIndex(InGPUIndex)
 	{
 	}
 
 	void Execute(FRHICommandListBase& RHICmdList)
 	{
 		const FShaderResourceViewInitializer::FBufferShaderResourceViewInitializer& Desc = Initializer.AsBufferSRV();
-		FD3D12Buffer* const Buffer = FD3D12DynamicRHI::ResourceCast(Desc.Buffer);
+		FD3D12Buffer* const Buffer = FD3D12DynamicRHI::ResourceCast(Desc.Buffer, GPUIndex);
 
 		if (Initializer.GetType() == FShaderResourceViewInitializer::EType::VertexBufferSRV)
 		{
@@ -554,13 +556,13 @@ FShaderResourceViewRHIRef FD3D12DynamicRHI::RHICreateShaderResourceView(const FS
 					{
 						// We have to defer the SRV initialization to the RHI thread if the buffer is dynamic (and RHI threading is enabled), as dynamic buffers can be renamed.
 						// Also insert an RHI thread fence to prevent parallel translate tasks running until this command has completed.
-						ALLOC_COMMAND_CL(RHICmdList, FD3D12InitializeBufferSRVRHICommand)(Initializer, ShaderResourceView);
+						ALLOC_COMMAND_CL(RHICmdList, FD3D12InitializeBufferSRVRHICommand)(Initializer, ShaderResourceView, Buffer->GetParentGPUIndex());
 						RHICmdList.RHIThreadFence(true);
 					}
 					else
 					{
 						// Run the command directly if we're bypassing RHI command list recording, or the buffer is not dynamic.
-						FD3D12InitializeBufferSRVRHICommand Command(Initializer, ShaderResourceView);
+						FD3D12InitializeBufferSRVRHICommand Command(Initializer, ShaderResourceView, Buffer->GetParentGPUIndex());
 						Command.Execute(RHICmdList);
 					}
 
diff --git a/Engine/Source/Runtime/Engine/Private/DistanceFieldAtlas.cpp b/Engine/Source/Runtime/Engine/Private/DistanceFieldAtlas.cpp
--- a/Engine/Source/Runtime/Engine/Private/DistanceFieldAtlas.cpp
+++ b/Engine/Source/Runtime/Engine/Private/DistanceFieldAtlas.cpp
@@ -1025,6 +1025,9 @@ uint32 FLandscapeTextureAtlas::CalculateDownSampleLevel(uint32 SizeX, uint32 Siz
 
 void FLandscapeTextureAtlas::UpdateAllocations(FRDGBuilder& GraphBuilder, ERHIFeatureLevel::Type InFeatureLevel)
 {
+	// Mask should be set in FSceneRenderer::PrepareDistanceFieldScene before calling this
+	check(GraphBuilder.RHICmdList.GetGPUMask() == FRHIGPUMask::All());
+
 	InitializeIfNeeded();
 
 	TArray<FPendingUpload, TInlineAllocator<8>> PendingUploads;
diff --git a/Engine/Source/Runtime/Renderer/Private/DistanceFieldObjectManagement.cpp b/Engine/Source/Runtime/Renderer/Private/DistanceFieldObjectManagement.cpp
--- a/Engine/Source/Runtime/Renderer/Private/DistanceFieldObjectManagement.cpp
+++ b/Engine/Source/Runtime/Renderer/Private/DistanceFieldObjectManagement.cpp
@@ -581,14 +581,14 @@ void FDistanceFieldSceneData::UpdateDistanceFieldObjectBuffers(
 	TArray<FDistanceFieldAssetMipId>& DistanceFieldAssetAdds,
 	TArray<FSetElementId>& DistanceFieldAssetRemoves)
 {
+	// Mask should be set in FSceneRenderer::PrepareDistanceFieldScene before calling this
+	check(GraphBuilder.RHICmdList.GetGPUMask() == FRHIGPUMask::All());
+
 	const bool bExecuteInParallel = GDFParallelUpdate != 0 && FApp::ShouldUseThreadingForPerformance();
 
 	if (HasPendingOperations() || PendingThrottledOperations.Num() > 0)
 	{
 		QUICK_SCOPE_CYCLE_COUNTER(STAT_UpdateDistanceFieldObjectBuffers);
-		// Multi-GPU support : Updating on all GPUs may be inefficient for AFR. Work is
-		// wasted for any objects that update on consecutive frames.
-		RDG_GPU_MASK_SCOPE(GraphBuilder, FRHIGPUMask::All());
 		RDG_EVENT_SCOPE(GraphBuilder, "UpdateDistanceFieldObjectBuffers");
 
 		if (ObjectBuffers == nullptr)
@@ -845,6 +845,9 @@ void FDistanceFieldSceneData::UpdateDistanceFieldObjectBuffers(
 
 void FSceneRenderer::UpdateGlobalHeightFieldObjectBuffers(FRDGBuilder& GraphBuilder)
 {
+	// Mask should be set in FSceneRenderer::PrepareDistanceFieldScene before calling this
+	check(GraphBuilder.RHICmdList.GetGPUMask() == FRHIGPUMask::All());
+
 	FDistanceFieldSceneData& DistanceFieldSceneData = Scene->DistanceFieldSceneData;
 
 	if (GHeightFieldTextureAtlas.GetAtlasTexture()
@@ -1046,6 +1049,7 @@ void FSceneRenderer::PrepareDistanceFieldScene(FRDGBuilder& GraphBuilder, bool b
 	TRACE_CPUPROFILER_EVENT_SCOPE(FSceneRenderer::PrepareDistanceFieldScene);
 	QUICK_SCOPE_CYCLE_COUNTER(STAT_PrepareDistanceFieldScene);
 	LLM_SCOPE_BYTAG(DistanceFields);
+	RDG_GPU_MASK_SCOPE(GraphBuilder, FRHIGPUMask::All());
 
 	const bool bShouldPrepareHeightFieldScene = ShouldPrepareHeightFieldScene();
 	const bool bShouldPrepareDistanceFieldScene = ShouldPrepareDistanceFieldScene();
diff --git a/Engine/Source/Runtime/Renderer/Private/DistanceFieldStreaming.cpp b/Engine/Source/Runtime/Renderer/Private/DistanceFieldStreaming.cpp
--- a/Engine/Source/Runtime/Renderer/Private/DistanceFieldStreaming.cpp
+++ b/Engine/Source/Runtime/Renderer/Private/DistanceFieldStreaming.cpp
@@ -848,6 +848,9 @@ void FDistanceFieldSceneData::ProcessReadRequests(
 
 void FDistanceFieldSceneData::ResizeBrickAtlasIfNeeded(FRDGBuilder& GraphBuilder, FGlobalShaderMap* GlobalShaderMap)
 {
+	// Mask should be set in FSceneRenderer::PrepareDistanceFieldScene before calling this
+	check(GraphBuilder.RHICmdList.GetGPUMask() == FRHIGPUMask::All());
+
 	const int32 BrickAtlasSizeXYInBricks = CVarBrickAtlasSizeXYInBricks.GetValueOnRenderThread();
 	int32 DesiredZSizeInBricks = FMath::DivideAndRoundUp(DistanceFieldAtlasBlockAllocator.GetMaxSize() * GDistanceFieldBlockAllocatorSizeInBricks, BrickAtlasSizeXYInBricks * BrickAtlasSizeXYInBricks);
 
@@ -925,6 +928,9 @@ static FIntVector CalculateDesiredSize(FIntVector CurrentSize, FIntVector Requir
 
 bool FDistanceFieldSceneData::ResizeIndirectionAtlasIfNeeded(FRDGBuilder& GraphBuilder, FGlobalShaderMap* GlobalShaderMap)
 {
+	// Mask should be set in FSceneRenderer::PrepareDistanceFieldScene before calling this
+	check(GraphBuilder.RHICmdList.GetGPUMask() == FRHIGPUMask::All());
+
 	QUICK_SCOPE_CYCLE_COUNTER(STAT_ResizeIndirectionAtlasIfNeeded);
 	TRACE_CPUPROFILER_EVENT_SCOPE(FDistanceFieldSceneData::ResizeIndirectionAtlasIfNeeded);
 	
@@ -1071,6 +1077,8 @@ void FDistanceFieldSceneData::GenerateStreamingRequests(
 	// It is not safe to EnqueueCopy on a buffer that already has a pending copy.
 	if (ReadbackBuffersNumPending < MaxStreamingReadbackBuffers && NumObjectsInBuffer > 0)
 	{
+		RDG_GPU_MASK_SCOPE(GraphBuilder, View.GPUMask);
+
 		if (!StreamingRequestReadbackBuffers[ReadbackBuffersWriteIndex])
 		{
 			FRHIGPUBufferReadback* GPUBufferReadback = new FRHIGPUBufferReadback(TEXT("DistanceFields.StreamingRequestReadBack"));
@@ -1203,6 +1211,9 @@ void FDistanceFieldSceneData::UploadAssetData(FRDGBuilder& GraphBuilder, const T
 		return;
 	}
 
+	// Mask should be set in FSceneRenderer::PrepareDistanceFieldScene before calling this
+	check(GraphBuilder.RHICmdList.GetGPUMask() == FRHIGPUMask::All());
+
 	AssetDataUploadBuffer.Init(AssetDataUploads.Num(), AssetDataMipStrideFloat4s * sizeof(FVector4f), true, TEXT("DistanceFields.DFAssetDataUploadBuffer"));
 
 	for (FDistanceFieldAssetMipId AssetMipUpload : AssetDataUploads)
@@ -1246,6 +1257,9 @@ void FDistanceFieldSceneData::UploadAllAssetData(FRDGBuilder& GraphBuilder)
 		return;
 	}
 
+	// Mask should be set in FSceneRenderer::PrepareDistanceFieldScene before calling this
+	check(GraphBuilder.RHICmdList.GetGPUMask() == FRHIGPUMask::All());
+
 	AssetDataUploadBuffer.Init(NumUploads, AssetDataMipStrideFloat4s * sizeof(FVector4f), true, TEXT("DistanceFields.DFAssetDataUploadBuffer"));
 
 	for (TSet<FDistanceFieldAssetState, TFDistanceFieldAssetStateFuncs>::TConstIterator It(AssetStateArray); It; ++It)
@@ -1299,6 +1313,9 @@ void FDistanceFieldSceneData::UpdateDistanceFieldAtlas(
 	TRACE_CPUPROFILER_EVENT_SCOPE(FDistanceFieldSceneData::UpdateDistanceFieldAtlas);
 	RDG_EVENT_SCOPE(GraphBuilder, "UpdateDistanceFieldAtlas");
 
+	// Mask should be set in FSceneRenderer::PrepareDistanceFieldScene before calling this
+	check(GraphBuilder.RHICmdList.GetGPUMask() == FRHIGPUMask::All());
+
 	TArray<FDistanceFieldAssetMipId> AssetDataUploads;
 
 	for (FSetElementId AssetSetId : DistanceFieldAssetRemoves)
diff --git a/Engine/Source/Runtime/Renderer/Private/DistanceFieldVisualization.cpp b/Engine/Source/Runtime/Renderer/Private/DistanceFieldVisualization.cpp
--- a/Engine/Source/Runtime/Renderer/Private/DistanceFieldVisualization.cpp
+++ b/Engine/Source/Runtime/Renderer/Private/DistanceFieldVisualization.cpp
@@ -170,6 +170,8 @@ void FDeferredShadingSceneRenderer::RenderMeshDistanceFieldVisualization(
 
 	for (const FViewInfo& View : Views)
 	{
+		RDG_GPU_MASK_SCOPE(GraphBuilder, View.GPUMask);
+
 		auto* PassParameters = GraphBuilder.AllocParameters<FVisualizeDistanceFieldUpsamplePS::FParameters>();
 		PassParameters->View = View.ViewUniformBuffer;
 		PassParameters->SceneTextures = SceneTextures.UniformBuffer;
