Fix/Add split screen & VR support for Strata (i.e., Tiling/Lighting/SSS/Distortion).

This CL split the resources shared between views (non-buffer/tiled data) and resources unique to view (e.g., tile data). This is to ease overall setup, but we might re-regroup them later.

diff --git a/Engine/Shaders/Private/DistortApplyScreenPS.usf b/Engine/Shaders/Private/DistortApplyScreenPS.usf
--- a/Engine/Shaders/Private/DistortApplyScreenPS.usf
+++ b/Engine/Shaders/Private/DistortApplyScreenPS.usf
@@ -71,7 +71,8 @@ void Main(
 	half RefractionRoughness = saturate(StrataLobeVarianceToRoughness(RefractionLobeVariance));
 	half RefractionMipLevel = RefractionRoughness * RefractionRoughnessToMipLevelFactor;
 
-	const float2 NoDistortionUVs = TexCoord.xy * View.BufferSizeAndInvSize.xy * View.ViewSizeAndInvSize.zw;
+	const float2 NoDistortionCoord = (TexCoord.xy * View.BufferSizeAndInvSize.xy - View.ViewRectMin);
+	const float2 NoDistortionUVs = NoDistortionCoord * View.ViewSizeAndInvSize.zw;
 	float2 NewBufferUV = NoDistortionUVs + DistBufferUVOffset;
 
 	FLATTEN if (any(NewBufferUV < 0.0f) || any(NewBufferUV > 1.0f))
@@ -83,7 +84,7 @@ void Main(
 	if (RefractionMipLevel < 0.001f)
 	{
 		// Sample the full res MSAA source texture
-		OutColor = SceneColorMSAATexture.Load(int2(TexCoord.xy * View.BufferSizeAndInvSize.xy), SampleIndex);
+		OutColor = SceneColorMSAATexture.Load(int2(NoDistortionCoord), SampleIndex);
 	}
 	else
 	{
diff --git a/Engine/Shaders/Private/PostProcessSubsurface.usf b/Engine/Shaders/Private/PostProcessSubsurface.usf
--- a/Engine/Shaders/Private/PostProcessSubsurface.usf
+++ b/Engine/Shaders/Private/PostProcessSubsurface.usf
@@ -112,7 +112,7 @@ float GetMaskFromDepthInAlpha(float Alpha)
 #if STRATA_ENABLED
 FStrataSubsurfaceData LoadStataSSSData(float2 UV)
 {
-	const float2 PixelPos = UV.xy * View.BufferSizeAndInvSize.xy - View.ViewRectMin.xy;
+	const float2 PixelPos = UV.xy * View.BufferSizeAndInvSize.xy;
 	return StrataUnpackSSSData(Strata.SSSTexture.Load(uint3(PixelPos, 0)));
 }
 #endif
diff --git a/Engine/Shaders/Private/Strata/StrataMaterialClassification.usf b/Engine/Shaders/Private/Strata/StrataMaterialClassification.usf
--- a/Engine/Shaders/Private/Strata/StrataMaterialClassification.usf
+++ b/Engine/Shaders/Private/Strata/StrataMaterialClassification.usf
@@ -52,9 +52,9 @@ void TileMainCS(uint2 DispatchThreadId : SV_DispatchThreadID, uint LinearIndex :
 		TileDrawIndirectDataBuffer[GetStrataTileTypeDrawIndirectArgOffset_DWord(STRATA_TILE_TYPE_SSS_WITHOUT_ROUGH_REFRACT) + 0] = IndexCountPerInstance;
 	}
 
-	// TODO: add a SM6 permutation with ballot?
-	const uint2 PixelCoord = DispatchThreadId;
-
+	const uint2 PixelCoord = DispatchThreadId.xy + View.ViewRectMin.xy;
+	const bool bIsValid = all(DispatchThreadId.xy < View.ViewSizeAndInvSize.xy);
+	
 	// Pixels outside of the view area are considered simple to enable screen borders to receive the simple permutation when not aligned to shader group size.
 	bool bContainsComplexMaterial = false;
 	bool bContainsSimpleMaterial  = false;
@@ -63,7 +63,7 @@ void TileMainCS(uint2 DispatchThreadId : SV_DispatchThreadID, uint LinearIndex :
 	bool bContainsOpaqueRoughRefraction = false;
 	bool bContainsScreenSpaceSubsurfaceScattering = false;
 	FStrataOpaqueRoughRefractionData OpaqueRoughRefractionData = (FStrataOpaqueRoughRefractionData)0;
-	if (all(PixelCoord < uint2(ViewResolution)))
+	if (bIsValid)
 	{
 		FStrataAddressing StrataAddressing = GetStrataPixelDataByteOffset(PixelCoord, uint2(View.BufferSizeAndInvSize.xy), MaxBytesPerPixel);
 		FStrataPixelHeader StrataPixelHeader = UnpackStrataHeaderIn(MaterialTextureArray, StrataAddressing, TopLayerTexture);
@@ -88,7 +88,7 @@ void TileMainCS(uint2 DispatchThreadId : SV_DispatchThreadID, uint LinearIndex :
 	{
 		// We must fill all the pixel which doe not hav subsurface scattering with default so that the SSS code is not executed where it should not.
 		FStrataSubsurfaceData StrataSubsurfaceData = (FStrataSubsurfaceData)0;
-		SSSTextureUAV[DispatchThreadId.xy] = StrataPackSSSData(StrataSubsurfaceData);
+		SSSTextureUAV[PixelCoord] = StrataPackSSSData(StrataSubsurfaceData);
 	}
 #endif
 
diff --git a/Engine/Shaders/Private/Strata/StrataTile.usf b/Engine/Shaders/Private/Strata/StrataTile.usf
--- a/Engine/Shaders/Private/Strata/StrataTile.usf
+++ b/Engine/Shaders/Private/Strata/StrataTile.usf
@@ -9,6 +9,7 @@
 
 #if SHADER_TILE_VS
 
+float2 OutputViewMinRect;
 float4 OutputViewSizeAndInvSize;
 float4 OutputBufferSizeAndInvSize;
 float4x4 ViewScreenToTranslatedWorld;
@@ -41,7 +42,7 @@ void StrataTilePassVS(
 	OutPosition = float4(float2(TileVertex) * OutputViewSizeAndInvSize.zw * float2(2.0f, -2.0f) + float2(-1.0, 1.0f), 0.5f, 1.0f);
 
 #if PERMUTATION_ENABLE_TEXCOORD_SCREENVECTOR
-	OutTexCoord = float2(TileVertex) * OutputBufferSizeAndInvSize.zw;
+	OutTexCoord = (TileVertex + OutputViewMinRect.xy) * OutputBufferSizeAndInvSize.zw;
 	OutScreenVector = mul(float4(OutPosition.xy, 1, 0), ViewScreenToTranslatedWorld).xyz;
 #endif
 
diff --git a/Engine/Shaders/Private/SubsurfaceBurleyNormalized.ush b/Engine/Shaders/Private/SubsurfaceBurleyNormalized.ush
--- a/Engine/Shaders/Private/SubsurfaceBurleyNormalized.ush
+++ b/Engine/Shaders/Private/SubsurfaceBurleyNormalized.ush
@@ -306,13 +306,13 @@ FBurleySampleInfo GenerateSampleInfo(float2 Rand0T1, float DiffuseMeanFreePathFo
 #if STRATA_ENABLED
 FStrataSubsurfaceData LoadStrataSSSData(float2 UV)
 {
-	const float2 PixelPos = UV.xy * View.BufferSizeAndInvSize.xy - View.ViewRectMin.xy;
+	const float2 PixelPos = UV.xy * View.BufferSizeAndInvSize.xy;
 	return StrataUnpackSSSData(Strata.SSSTexture.Load(uint3(PixelPos, 0)));
 }
 
 FStrataTopLayerData LoadStrataTopLayerData(float2 UV)
 {
-	const float2 PixelPos = UV.xy * View.BufferSizeAndInvSize.xy - View.ViewRectMin.xy;
+	const float2 PixelPos = UV.xy * View.BufferSizeAndInvSize.xy;
 	return StrataUnpackTopLayerData(Strata.TopLayerTexture.Load(uint3(PixelPos, 0)));
 }
 #endif
diff --git a/Engine/Source/Runtime/Renderer/Private/DebugProbeRendering.cpp b/Engine/Source/Runtime/Renderer/Private/DebugProbeRendering.cpp
--- a/Engine/Source/Runtime/Renderer/Private/DebugProbeRendering.cpp
+++ b/Engine/Source/Runtime/Renderer/Private/DebugProbeRendering.cpp
@@ -80,8 +80,8 @@ static void CommonStampDeferredDebugProbeDrawCall(
 	int32 RenderPass)
 {
 	PassParameters->ViewUniformBuffer = View.ViewUniformBuffer;
-	PassParameters->MaterialTextureArrayUAV = View.StrataSceneData->MaterialTextureArrayUAVWithoutRTs;
-	PassParameters->MaxBytesPerPixel = View.StrataSceneData->MaxBytesPerPixel;
+	PassParameters->MaterialTextureArrayUAV = View.StrataViewData.SceneData->MaterialTextureArrayUAVWithoutRTs;
+	PassParameters->MaxBytesPerPixel = View.StrataViewData.SceneData->MaxBytesPerPixel;
 	PassParameters->DebugProbesMode = View.Family->EngineShowFlags.VisualizeLightingOnProbes ? 3 : FMath::Clamp(CVarVisualizeLightingOnProbes.GetValueOnRenderThread(), 0, 3);
 		
 	FStampDeferredDebugProbePS::FPermutationDomain PermutationVector;
diff --git a/Engine/Source/Runtime/Renderer/Private/DeferredShadingRenderer.cpp b/Engine/Source/Runtime/Renderer/Private/DeferredShadingRenderer.cpp
--- a/Engine/Source/Runtime/Renderer/Private/DeferredShadingRenderer.cpp
+++ b/Engine/Source/Runtime/Renderer/Private/DeferredShadingRenderer.cpp
@@ -2253,7 +2253,7 @@ void FDeferredShadingSceneRenderer::Render(FRDGBuilder& GraphBuilder)
 
 	// Strata initialisation is always run even when not enabled.
 	const bool bStrataEnabled = Strata::IsStrataEnabled();
-	Strata::InitialiseStrataFrameSceneData(*this, GraphBuilder);
+	Strata::InitialiseStrataFrameSceneData(GraphBuilder, *this);
 
 	if (DepthPass.IsComputeStencilDitherEnabled())
 	{
diff --git a/Engine/Source/Runtime/Renderer/Private/DistortionRendering.cpp b/Engine/Source/Runtime/Renderer/Private/DistortionRendering.cpp
--- a/Engine/Source/Runtime/Renderer/Private/DistortionRendering.cpp
+++ b/Engine/Source/Runtime/Renderer/Private/DistortionRendering.cpp
@@ -284,7 +284,12 @@ static void AddCopySceneColorPass(FRDGBuilder& GraphBuilder, const FViewInfo& Vi
 	PassParameters->SceneColorSampler = TStaticSamplerState<SF_Point>::GetRHI();
 	PassParameters->RenderTargets[0] = FRenderTargetBinding(SceneColorCopyTexture, ERenderTargetLoadAction::ENoAction);
 
-	const FScreenPassTextureViewport InputViewport(SceneColorTexture, View.ViewRect);
+	// The scene color is copied into from multi-view-rect layout into a single-rect layout.
+	FIntRect ViewRect;
+	ViewRect.Min = FIntPoint(0, 0);
+	ViewRect.Max = FIntPoint(View.ViewRect.Width(), View.ViewRect.Height());
+
+	const FScreenPassTextureViewport InputViewport(SceneColorTexture, ViewRect);
 	const FScreenPassTextureViewport OutputViewport(SceneColorCopyTexture);
 
 	AddDrawScreenPass(GraphBuilder, {}, View, InputViewport, OutputViewport, VertexShader, PixelShader, PassParameters);
diff --git a/Engine/Source/Runtime/Renderer/Private/MobileShadingRenderer.cpp b/Engine/Source/Runtime/Renderer/Private/MobileShadingRenderer.cpp
--- a/Engine/Source/Runtime/Renderer/Private/MobileShadingRenderer.cpp
+++ b/Engine/Source/Runtime/Renderer/Private/MobileShadingRenderer.cpp
@@ -727,7 +727,7 @@ void FMobileSceneRenderer::Render(FRDGBuilder& GraphBuilder)
 	GSystemTextures.InitializeTextures(GraphBuilder.RHICmdList, FeatureLevel);
 
 	// Strata initialisation is always run even when not enabled.
-	Strata::InitialiseStrataFrameSceneData(*this, GraphBuilder);
+	Strata::InitialiseStrataFrameSceneData(GraphBuilder, *this);
 
 	// Force the subsurface profile texture to be updated.
 	UpdateSubsurfaceProfileTexture(GraphBuilder, ShaderPlatform);
diff --git a/Engine/Source/Runtime/Renderer/Private/PostProcess/PostProcessSubsurface.cpp b/Engine/Source/Runtime/Renderer/Private/PostProcess/PostProcessSubsurface.cpp
--- a/Engine/Source/Runtime/Renderer/Private/PostProcess/PostProcessSubsurface.cpp
+++ b/Engine/Source/Runtime/Renderer/Private/PostProcess/PostProcessSubsurface.cpp
@@ -988,8 +988,8 @@ void AddSubsurfaceViewPass(
 
 	const bool bReadSeparatedSubSurfaceSceneColor = Strata::IsStrataOpaqueMaterialRoughRefractionEnabled();
 	const bool bWriteSeparatedOpaqueRoughRefractionSceneColor = Strata::IsStrataOpaqueMaterialRoughRefractionEnabled();
-	FRDGTextureRef SeparatedSubSurfaceSceneColor = View.StrataSceneData->SeparatedSubSurfaceSceneColor;
-	FRDGTextureRef SeparatedOpaqueRoughRefractionSceneColor = View.StrataSceneData->SeparatedOpaqueRoughRefractionSceneColor;
+	FRDGTextureRef SeparatedSubSurfaceSceneColor = View.StrataViewData.SceneData->SeparatedSubSurfaceSceneColor;
+	FRDGTextureRef SeparatedOpaqueRoughRefractionSceneColor = View.StrataViewData.SceneData->SeparatedOpaqueRoughRefractionSceneColor;
 
 	FRDGTextureRef SetupTexture = SceneColorTexture;
 	FRDGTextureRef SubsurfaceSubpassOneTex = nullptr;
diff --git a/Engine/Source/Runtime/Renderer/Private/SceneRendering.cpp b/Engine/Source/Runtime/Renderer/Private/SceneRendering.cpp
--- a/Engine/Source/Runtime/Renderer/Private/SceneRendering.cpp
+++ b/Engine/Source/Runtime/Renderer/Private/SceneRendering.cpp
@@ -987,7 +987,7 @@ void FViewInfo::Init()
 	NumVisibleDynamicPrimitives = 0;
 	NumVisibleDynamicEditorPrimitives = 0;
 
-	StrataSceneData = nullptr;
+	StrataViewData.Reset();
 	HairStrandsViewData = FHairStrandsViewData();
 
 	GPUSceneViewId = INDEX_NONE;
diff --git a/Engine/Source/Runtime/Renderer/Private/SceneRendering.h b/Engine/Source/Runtime/Renderer/Private/SceneRendering.h
--- a/Engine/Source/Runtime/Renderer/Private/SceneRendering.h
+++ b/Engine/Source/Runtime/Renderer/Private/SceneRendering.h
@@ -36,6 +36,7 @@
 #include "VirtualShadowMaps/VirtualShadowMapArray.h"
 #include "Lumen/LumenTranslucencyVolumeLighting.h"
 #include "HairStrands/HairStrandsData.h"
+#include "Strata/Strata.h"
 #include "GPUScene.h"
 #include "CanvasTypes.h"
 
@@ -56,7 +57,6 @@ struct FSceneWithoutWaterTextures;
 struct FHairStrandsVisibilityViews;
 struct FSortedLightSetSceneInfo;
 enum class EVelocityPass : uint32;
-struct FStrataSceneData;
 class FTransientLightFunctionTextureAtlas;
 struct FSceneTexturesConfig;
 struct FMinimalSceneTextures;
@@ -1294,7 +1294,7 @@ public:
 	// Used by mobile renderer to determine whether static meshes will be rendered with CSM shaders or not.
 	FMobileCSMVisibilityInfo MobileCSMVisibilityInfo;
 
-	FStrataSceneData* StrataSceneData;
+	FStrataViewData StrataViewData;
 
 	FHairStrandsViewData HairStrandsViewData;
 
diff --git a/Engine/Source/Runtime/Renderer/Private/Strata/Strata.cpp b/Engine/Source/Runtime/Renderer/Private/Strata/Strata.cpp
--- a/Engine/Source/Runtime/Renderer/Private/Strata/Strata.cpp
+++ b/Engine/Source/Runtime/Renderer/Private/Strata/Strata.cpp
@@ -68,44 +68,15 @@ static TAutoConsoleVariable<int32> CVarStrataTileOverflow(
 
 IMPLEMENT_GLOBAL_SHADER_PARAMETER_STRUCT(FStrataGlobalUniformParameters, "Strata");
 
-void FStrataSceneData::Reset()
+void FStrataViewData::Reset()
 {
-	TopLayerTexture = nullptr;
-	TopLayerTextureUAV = nullptr;
-	SSSTexture = nullptr;
-	SSSTextureUAV = nullptr;
-	OpaqueRoughRefractionTexture = nullptr;
-	OpaqueRoughRefractionTextureUAV = nullptr;
-
-	MaterialTextureArray = nullptr;
-	MaterialTextureArrayUAVWithoutRTs = nullptr;
-	MaterialTextureArrayUAV = nullptr;
-	MaterialTextureArraySRV = nullptr;
-
-	ClassificationTileDrawIndirectBuffer = nullptr;
-	ClassificationTileDrawIndirectBufferUAV = nullptr;
-
-	ClassificationTileDispatchIndirectBuffer = nullptr;
-	ClassificationTileDispatchIndirectBufferUAV = nullptr;
-
+	*this = FStrataViewData();
 	for (uint32 i = 0; i < EStrataTileType::ECount; ++i)
 	{
 		ClassificationTileListBuffer[i] = nullptr;
 		ClassificationTileListBufferUAV[i] = nullptr;
 		ClassificationTileListBufferSRV[i] = nullptr;
 	}
-
-	SeparatedSubSurfaceSceneColor = nullptr;
-	SeparatedOpaqueRoughRefractionSceneColor = nullptr;
-
-	BSDFOffsetTexture = nullptr;
-	BSDFTileTexture = nullptr;
-
-	TileCount_Total = FIntPoint(0,0);
-	TileCount_Primary = FIntPoint(0,0);
-	TileCount_Overflow = FIntPoint(0,0);
-
-	StrataGlobalUniformParameters = nullptr;
 }
 
 const TCHAR* ToString(EStrataTileType Type)
@@ -164,34 +135,21 @@ FIntPoint GetStrataTextureResolution(const FIntPoint& InResolution)
 	return GetStrataTextureTileResolution(InResolution) * STRATA_TILE_SIZE;
 }
 
-static void BindStrataGlobalUniformParameters(FRDGBuilder& GraphBuilder, FStrataSceneData* StrataSceneData, FStrataGlobalUniformParameters& OutStrataUniformParameters);
+static void BindStrataGlobalUniformParameters(FRDGBuilder& GraphBuilder, FStrataViewData* StrataViewData, FStrataGlobalUniformParameters& OutStrataUniformParameters);
 
-void InitialiseStrataFrameSceneData(FSceneRenderer& SceneRenderer, FRDGBuilder& GraphBuilder)
+static void InitialiseStrataViewData(FRDGBuilder& GraphBuilder, FViewInfo& View, FStrataSceneData& SceneData)
 {
-	FStrataSceneData& StrataSceneData = SceneRenderer.Scene->StrataSceneData;
-	StrataSceneData.Reset();
+	// Sanity check: the scene data should already exist 
+	check(SceneData.MaterialTextureArray != nullptr);
 
-	auto UpdateMaterialBufferToTiledResolution = [](FIntPoint InBufferSizeXY, FIntPoint& OutMaterialBufferSizeXY)
-	{
-		// We need to allocate enough for the tiled memory addressing to always work
-		OutMaterialBufferSizeXY.X = FMath::DivideAndRoundUp(InBufferSizeXY.X, STRATA_TILE_SIZE) * STRATA_TILE_SIZE;
-		OutMaterialBufferSizeXY.Y = FMath::DivideAndRoundUp(InBufferSizeXY.Y, STRATA_TILE_SIZE) * STRATA_TILE_SIZE;
-	};
+	FStrataViewData& Out = View.StrataViewData;
+	Out.Reset();
+	Out.SceneData = &SceneData;
 
-	FIntPoint MaterialBufferSizeXY;
-	UpdateMaterialBufferToTiledResolution(FIntPoint(1, 1), MaterialBufferSizeXY);
+	const FIntPoint ViewResolution(View.ViewRect.Width(), View.ViewRect.Height());
 	if (IsStrataEnabled())
 	{
-		FIntPoint SceneTextureExtent = GetSceneTextureExtent();
-		
-		// We need to allocate enough for the tiled memory addressing of material data to always work
-		UpdateMaterialBufferToTiledResolution(SceneTextureExtent, MaterialBufferSizeXY);
-
-		const uint32 MaterialConservativeByteCountPerPixel = CVarStrataBytePerPixel.GetValueOnAnyThread();
-		const uint32 RoundToValue = 4u;
-		StrataSceneData.MaxBytesPerPixel = FMath::DivideAndRoundUp(MaterialConservativeByteCountPerPixel, RoundToValue) * RoundToValue;
-
-		const FIntPoint TileResolution(FMath::DivideAndRoundUp(SceneTextureExtent.X, STRATA_TILE_SIZE), FMath::DivideAndRoundUp(SceneTextureExtent.Y, STRATA_TILE_SIZE));
+		const FIntPoint TileResolution(FMath::DivideAndRoundUp(ViewResolution.X, STRATA_TILE_SIZE), FMath::DivideAndRoundUp(ViewResolution.Y, STRATA_TILE_SIZE));
 
 		const TCHAR* StrataTileListBufferNames[EStrataTileType::ECount] =
 		{
@@ -205,133 +163,167 @@ void InitialiseStrataFrameSceneData(FSceneRenderer& SceneRenderer, FRDGBuilder&
 		// Tile classification buffers
 		{
 			// Indirect draw
-			StrataSceneData.ClassificationTileDrawIndirectBuffer = GraphBuilder.CreateBuffer(FRDGBufferDesc::CreateIndirectDesc<FRHIDrawIndirectParameters>(EStrataTileType::ECount), TEXT("Strata.StrataTileDrawIndirectBuffer"));
-			StrataSceneData.ClassificationTileDrawIndirectBufferUAV = GraphBuilder.CreateUAV(StrataSceneData.ClassificationTileDrawIndirectBuffer, PF_R32_UINT);
-			AddClearUAVPass(GraphBuilder, StrataSceneData.ClassificationTileDrawIndirectBufferUAV, 0);
+			Out.ClassificationTileDrawIndirectBuffer = GraphBuilder.CreateBuffer(FRDGBufferDesc::CreateIndirectDesc<FRHIDrawIndirectParameters>(EStrataTileType::ECount), TEXT("Strata.StrataTileDrawIndirectBuffer"));
+			Out.ClassificationTileDrawIndirectBufferUAV = GraphBuilder.CreateUAV(Out.ClassificationTileDrawIndirectBuffer, PF_R32_UINT);
+			AddClearUAVPass(GraphBuilder, Out.ClassificationTileDrawIndirectBufferUAV, 0);
 
 			// Indirect dispatch
-			StrataSceneData.ClassificationTileDispatchIndirectBuffer = GraphBuilder.CreateBuffer(FRDGBufferDesc::CreateIndirectDesc<FRHIDispatchIndirectParameters>(EStrataTileType::ECount), TEXT("Strata.StrataTileDispatchIndirectBuffer"));
-			StrataSceneData.ClassificationTileDispatchIndirectBufferUAV = GraphBuilder.CreateUAV(StrataSceneData.ClassificationTileDispatchIndirectBuffer, PF_R32_UINT);
-			AddClearUAVPass(GraphBuilder, StrataSceneData.ClassificationTileDispatchIndirectBufferUAV, 0);
+			Out.ClassificationTileDispatchIndirectBuffer = GraphBuilder.CreateBuffer(FRDGBufferDesc::CreateIndirectDesc<FRHIDispatchIndirectParameters>(EStrataTileType::ECount), TEXT("Strata.StrataTileDispatchIndirectBuffer"));
+			Out.ClassificationTileDispatchIndirectBufferUAV = GraphBuilder.CreateUAV(Out.ClassificationTileDispatchIndirectBuffer, PF_R32_UINT);
+			AddClearUAVPass(GraphBuilder, Out.ClassificationTileDispatchIndirectBufferUAV, 0);
 
 			for (uint32 i = 0; i <= EStrataTileType::EComplex; ++i)
 			{
-				StrataSceneData.ClassificationTileListBuffer[i] = GraphBuilder.CreateBuffer(FRDGBufferDesc::CreateBufferDesc(sizeof(uint32), TileResolution.X * TileResolution.Y), StrataTileListBufferNames[i]);
-				StrataSceneData.ClassificationTileListBufferSRV[i] = GraphBuilder.CreateSRV(StrataSceneData.ClassificationTileListBuffer[i], PF_R32_UINT);
-				StrataSceneData.ClassificationTileListBufferUAV[i] = GraphBuilder.CreateUAV(StrataSceneData.ClassificationTileListBuffer[i], PF_R32_UINT);
+				Out.ClassificationTileListBuffer[i] = GraphBuilder.CreateBuffer(FRDGBufferDesc::CreateBufferDesc(sizeof(uint32), TileResolution.X * TileResolution.Y), StrataTileListBufferNames[i]);
+				Out.ClassificationTileListBufferSRV[i] = GraphBuilder.CreateSRV(Out.ClassificationTileListBuffer[i], PF_R32_UINT);
+				Out.ClassificationTileListBufferUAV[i] = GraphBuilder.CreateUAV(Out.ClassificationTileListBuffer[i], PF_R32_UINT);
 			}
 		}
 
+		// Separated subsurface & rough refraction textures (tile data)
+		{
+			const bool bIsStrataOpaqueMaterialRoughRefractionEnabled= IsStrataOpaqueMaterialRoughRefractionEnabled();
+			const int32 TileListBufferElementCount					= bIsStrataOpaqueMaterialRoughRefractionEnabled ? TileResolution.X * TileResolution.Y : 4;
+			
+			Out.ClassificationTileListBuffer[EStrataTileType::EOpaqueRoughRefraction] = GraphBuilder.CreateBuffer(FRDGBufferDesc::CreateBufferDesc(sizeof(uint32), TileListBufferElementCount), StrataTileListBufferNames[EStrataTileType::EOpaqueRoughRefraction]);
+			Out.ClassificationTileListBufferSRV[EStrataTileType::EOpaqueRoughRefraction] = GraphBuilder.CreateSRV(Out.ClassificationTileListBuffer[EStrataTileType::EOpaqueRoughRefraction], PF_R32_UINT);
+			Out.ClassificationTileListBufferUAV[EStrataTileType::EOpaqueRoughRefraction] = GraphBuilder.CreateUAV(Out.ClassificationTileListBuffer[EStrataTileType::EOpaqueRoughRefraction], PF_R32_UINT);
+
+			Out.ClassificationTileListBuffer[EStrataTileType::ESSSWithoutOpaqueRoughRefraction] = GraphBuilder.CreateBuffer(FRDGBufferDesc::CreateBufferDesc(sizeof(uint32), TileListBufferElementCount), StrataTileListBufferNames[EStrataTileType::ESSSWithoutOpaqueRoughRefraction]);
+			Out.ClassificationTileListBufferSRV[EStrataTileType::ESSSWithoutOpaqueRoughRefraction] = GraphBuilder.CreateSRV(Out.ClassificationTileListBuffer[EStrataTileType::ESSSWithoutOpaqueRoughRefraction], PF_R32_UINT);
+			Out.ClassificationTileListBufferUAV[EStrataTileType::ESSSWithoutOpaqueRoughRefraction] = GraphBuilder.CreateUAV(Out.ClassificationTileListBuffer[EStrataTileType::ESSSWithoutOpaqueRoughRefraction], PF_R32_UINT);
+		}
+
+		// BSDF tiles
+		{
+			Out.TileCount_Total		= GetStrataTextureTileResolution(ViewResolution);
+			Out.TileCount_Primary	= GetStrataTextureTileResolution(ViewResolution, 0.f);
+			Out.TileCount_Overflow	= Out.TileCount_Total - Out.TileCount_Primary;
+
+			Out.BSDFTileTexture = GraphBuilder.CreateTexture(FRDGTextureDesc::Create2D(Out.TileCount_Total, PF_R32_UINT, FClearValueBinding::None, TexCreate_UAV | TexCreate_ShaderResource), TEXT("Strata.BSDFTiles"));
+			AddClearUAVPass(GraphBuilder, GraphBuilder.CreateUAV(Out.BSDFTileTexture), 0u);
+
+			Out.BSDFTileDispatchIndirectBuffer = GraphBuilder.CreateBuffer(FRDGBufferDesc::CreateIndirectDesc<FRHIDispatchIndirectParameters>(1), TEXT("Strata.StrataBSDFTileDispatchIndirectBuffer"));
+			Out.BSDFTileCountBuffer = GraphBuilder.CreateBuffer(FRDGBufferDesc::CreateBufferDesc(4, 1), TEXT("Strata.BSDFTileCount"));
+		}
+	}
+
+	// Create the readable uniform buffers
+	if (IsStrataEnabled())
+	{
+		FStrataGlobalUniformParameters* StrataUniformParameters = GraphBuilder.AllocParameters<FStrataGlobalUniformParameters>();
+		BindStrataGlobalUniformParameters(GraphBuilder, &Out, *StrataUniformParameters);
+		Out.StrataGlobalUniformParameters = GraphBuilder.CreateUniformBuffer(StrataUniformParameters);
+	}
+}
+
+void InitialiseStrataFrameSceneData(FRDGBuilder& GraphBuilder, FSceneRenderer& SceneRenderer)
+{
+	FStrataSceneData& Out = SceneRenderer.Scene->StrataSceneData;
+	Out = FStrataSceneData();
+
+	auto UpdateMaterialBufferToTiledResolution = [](FIntPoint InBufferSizeXY, FIntPoint& OutMaterialBufferSizeXY)
+	{
+		// We need to allocate enough for the tiled memory addressing to always work
+		OutMaterialBufferSizeXY.X = FMath::DivideAndRoundUp(InBufferSizeXY.X, STRATA_TILE_SIZE) * STRATA_TILE_SIZE;
+		OutMaterialBufferSizeXY.Y = FMath::DivideAndRoundUp(InBufferSizeXY.Y, STRATA_TILE_SIZE) * STRATA_TILE_SIZE;
+	};
+
+	FIntPoint MaterialBufferSizeXY;
+	UpdateMaterialBufferToTiledResolution(FIntPoint(1, 1), MaterialBufferSizeXY);
+	if (IsStrataEnabled())
+	{
+		FIntPoint SceneTextureExtent = GetSceneTextureExtent();
+		
+		// We need to allocate enough for the tiled memory addressing of material data to always work
+		UpdateMaterialBufferToTiledResolution(SceneTextureExtent, MaterialBufferSizeXY);
+
+		const uint32 MaterialConservativeByteCountPerPixel = CVarStrataBytePerPixel.GetValueOnAnyThread();
+		const uint32 RoundToValue = 4u;
+		Out.MaxBytesPerPixel = FMath::DivideAndRoundUp(MaterialConservativeByteCountPerPixel, RoundToValue) * RoundToValue;
+
 		// Top layer texture
 		{
-			StrataSceneData.TopLayerTexture = GraphBuilder.CreateTexture(FRDGTextureDesc::Create2D(SceneTextureExtent, PF_R32_UINT, FClearValueBinding::Black, TexCreate_RenderTargetable | TexCreate_ShaderResource), TEXT("Strata.TopLayerTexture"));
+			Out.TopLayerTexture = GraphBuilder.CreateTexture(FRDGTextureDesc::Create2D(SceneTextureExtent, PF_R32_UINT, FClearValueBinding::Black, TexCreate_RenderTargetable | TexCreate_ShaderResource), TEXT("Strata.TopLayerTexture"));
 		}
 
 		// SSS texture
 		{
-			StrataSceneData.SSSTexture = GraphBuilder.CreateTexture(FRDGTextureDesc::Create2D(SceneTextureExtent, PF_R32G32_UINT, FClearValueBinding::Black, TexCreate_DisableDCC | TexCreate_NoFastClear | TexCreate_ShaderResource | TexCreate_UAV), TEXT("Strata.SSSTexture"));
-			StrataSceneData.SSSTextureUAV = GraphBuilder.CreateUAV(StrataSceneData.SSSTexture);
+			Out.SSSTexture = GraphBuilder.CreateTexture(FRDGTextureDesc::Create2D(SceneTextureExtent, PF_R32G32_UINT, FClearValueBinding::Black, TexCreate_DisableDCC | TexCreate_NoFastClear | TexCreate_ShaderResource | TexCreate_UAV), TEXT("Strata.SSSTexture"));
+			Out.SSSTextureUAV = GraphBuilder.CreateUAV(Out.SSSTexture);
 		}
 
 		// Separated subsurface and rough refraction textures
 		{
 			const bool bIsStrataOpaqueMaterialRoughRefractionEnabled= IsStrataOpaqueMaterialRoughRefractionEnabled();
-			const int32 TileListBufferElementCount					= bIsStrataOpaqueMaterialRoughRefractionEnabled ? TileResolution.X * TileResolution.Y : 4;
 			const FIntPoint OpaqueRoughRefractionSceneExtent		= bIsStrataOpaqueMaterialRoughRefractionEnabled ? SceneTextureExtent : FIntPoint(4, 4);
 			
-			StrataSceneData.ClassificationTileListBuffer[EStrataTileType::EOpaqueRoughRefraction] = GraphBuilder.CreateBuffer(FRDGBufferDesc::CreateBufferDesc(sizeof(uint32), TileListBufferElementCount), StrataTileListBufferNames[EStrataTileType::EOpaqueRoughRefraction]);
-			StrataSceneData.ClassificationTileListBufferSRV[EStrataTileType::EOpaqueRoughRefraction] = GraphBuilder.CreateSRV(StrataSceneData.ClassificationTileListBuffer[EStrataTileType::EOpaqueRoughRefraction], PF_R32_UINT);
-			StrataSceneData.ClassificationTileListBufferUAV[EStrataTileType::EOpaqueRoughRefraction] = GraphBuilder.CreateUAV(StrataSceneData.ClassificationTileListBuffer[EStrataTileType::EOpaqueRoughRefraction], PF_R32_UINT);
-
-			StrataSceneData.ClassificationTileListBuffer[EStrataTileType::ESSSWithoutOpaqueRoughRefraction] = GraphBuilder.CreateBuffer(FRDGBufferDesc::CreateBufferDesc(sizeof(uint32), TileListBufferElementCount), StrataTileListBufferNames[EStrataTileType::ESSSWithoutOpaqueRoughRefraction]);
-			StrataSceneData.ClassificationTileListBufferSRV[EStrataTileType::ESSSWithoutOpaqueRoughRefraction] = GraphBuilder.CreateSRV(StrataSceneData.ClassificationTileListBuffer[EStrataTileType::ESSSWithoutOpaqueRoughRefraction], PF_R32_UINT);
-			StrataSceneData.ClassificationTileListBufferUAV[EStrataTileType::ESSSWithoutOpaqueRoughRefraction] = GraphBuilder.CreateUAV(StrataSceneData.ClassificationTileListBuffer[EStrataTileType::ESSSWithoutOpaqueRoughRefraction], PF_R32_UINT);
-
-			StrataSceneData.OpaqueRoughRefractionTexture = GraphBuilder.CreateTexture(
+			Out.OpaqueRoughRefractionTexture = GraphBuilder.CreateTexture(
 				FRDGTextureDesc::Create2D(OpaqueRoughRefractionSceneExtent, PF_FloatR11G11B10, FClearValueBinding::Black, TexCreate_ShaderResource | TexCreate_UAV | TexCreate_RenderTargetable), TEXT("Strata.OpaqueRoughRefractionTexture"));
-			StrataSceneData.OpaqueRoughRefractionTextureUAV = GraphBuilder.CreateUAV(StrataSceneData.OpaqueRoughRefractionTexture);
+			Out.OpaqueRoughRefractionTextureUAV = GraphBuilder.CreateUAV(Out.OpaqueRoughRefractionTexture);
 			
-			StrataSceneData.SeparatedSubSurfaceSceneColor			= GraphBuilder.CreateTexture(
+			Out.SeparatedSubSurfaceSceneColor			= GraphBuilder.CreateTexture(
 				FRDGTextureDesc::Create2D(OpaqueRoughRefractionSceneExtent, PF_FloatR11G11B10, FClearValueBinding::Black, TexCreate_ShaderResource | TexCreate_UAV | TexCreate_RenderTargetable), TEXT("Strata.SeparatedSubSurfaceSceneColor"));
-			StrataSceneData.SeparatedOpaqueRoughRefractionSceneColor= GraphBuilder.CreateTexture(
+			Out.SeparatedOpaqueRoughRefractionSceneColor= GraphBuilder.CreateTexture(
 				FRDGTextureDesc::Create2D(OpaqueRoughRefractionSceneExtent, PF_FloatR11G11B10, FClearValueBinding::Black, TexCreate_ShaderResource | TexCreate_UAV | TexCreate_RenderTargetable), TEXT("Strata.SeparatedOpaqueRoughRefractionSceneColor"));
 
 			if (bIsStrataOpaqueMaterialRoughRefractionEnabled)
 			{
 				// Fast clears
-				AddClearRenderTargetPass(GraphBuilder, StrataSceneData.OpaqueRoughRefractionTexture, StrataSceneData.OpaqueRoughRefractionTexture->Desc.ClearValue.GetClearColor());
-				AddClearRenderTargetPass(GraphBuilder, StrataSceneData.SeparatedSubSurfaceSceneColor, StrataSceneData.SeparatedSubSurfaceSceneColor->Desc.ClearValue.GetClearColor());
-				AddClearRenderTargetPass(GraphBuilder, StrataSceneData.SeparatedOpaqueRoughRefractionSceneColor, StrataSceneData.SeparatedOpaqueRoughRefractionSceneColor->Desc.ClearValue.GetClearColor());
+				AddClearRenderTargetPass(GraphBuilder, Out.OpaqueRoughRefractionTexture, Out.OpaqueRoughRefractionTexture->Desc.ClearValue.GetClearColor());
+				AddClearRenderTargetPass(GraphBuilder, Out.SeparatedSubSurfaceSceneColor, Out.SeparatedSubSurfaceSceneColor->Desc.ClearValue.GetClearColor());
+				AddClearRenderTargetPass(GraphBuilder, Out.SeparatedOpaqueRoughRefractionSceneColor, Out.SeparatedOpaqueRoughRefractionSceneColor->Desc.ClearValue.GetClearColor());
 			}
 		}
 
-		// BSDF tiles/offsets
+		// BSDF offsets
 		{
-			StrataSceneData.TileCount_Total		= GetStrataTextureTileResolution(SceneTextureExtent);
-			StrataSceneData.TileCount_Primary	= GetStrataTextureTileResolution(SceneTextureExtent, 0.f);
-			StrataSceneData.TileCount_Overflow	= StrataSceneData.TileCount_Total - StrataSceneData.TileCount_Primary;
-
-			StrataSceneData.BSDFTileTexture = GraphBuilder.CreateTexture(FRDGTextureDesc::Create2D(StrataSceneData.TileCount_Total, PF_R32_UINT, FClearValueBinding::None, TexCreate_UAV | TexCreate_ShaderResource), TEXT("Strata.BSDFTiles"));
-			AddClearUAVPass(GraphBuilder, GraphBuilder.CreateUAV(StrataSceneData.BSDFTileTexture), 0u);
-
-			StrataSceneData.BSDFOffsetTexture = GraphBuilder.CreateTexture(FRDGTextureDesc::Create2D(StrataSceneData.TileCount_Primary * STRATA_TILE_SIZE, PF_R32_UINT, FClearValueBinding::None, TexCreate_UAV | TexCreate_ShaderResource), TEXT("Strata.BSDFOffsets"));
-			AddClearUAVPass(GraphBuilder, GraphBuilder.CreateUAV(StrataSceneData.BSDFOffsetTexture), 0u);
-
-			StrataSceneData.BSDFTileDispatchIndirectBuffer = GraphBuilder.CreateBuffer(FRDGBufferDesc::CreateIndirectDesc<FRHIDispatchIndirectParameters>(1), TEXT("Strata.StrataBSDFTileDispatchIndirectBuffer"));
-			StrataSceneData.BSDFTileCountBuffer = GraphBuilder.CreateBuffer(FRDGBufferDesc::CreateBufferDesc(4, 1), TEXT("Strata.BSDFTileCount"));
+			Out.BSDFOffsetTexture = GraphBuilder.CreateTexture(FRDGTextureDesc::Create2D(SceneTextureExtent, PF_R32_UINT, FClearValueBinding::None, TexCreate_UAV | TexCreate_ShaderResource), TEXT("Strata.BSDFOffsets"));
+			AddClearUAVPass(GraphBuilder, GraphBuilder.CreateUAV(Out.BSDFOffsetTexture), 0u);
 		}
 	}
 	else
 	{
-		StrataSceneData.MaxBytesPerPixel = 4u * STRATA_BASE_PASS_MRT_OUTPUT_COUNT;
+		Out.MaxBytesPerPixel = 4u * STRATA_BASE_PASS_MRT_OUTPUT_COUNT;
 	}
 
 	// Create the material data container
 	FIntPoint SceneTextureExtent = IsStrataEnabled() ? GetSceneTextureExtent() : FIntPoint(2, 2);
 
-	const uint32 SliceCount = FMath::DivideAndRoundUp(StrataSceneData.MaxBytesPerPixel, 4u);
+	const uint32 SliceCount = FMath::DivideAndRoundUp(Out.MaxBytesPerPixel, 4u);
 	const FRDGTextureDesc MaterialTextureDesc = FRDGTextureDesc::Create2DArray(SceneTextureExtent, PF_R32_UINT, FClearValueBinding::Transparent,
 		TexCreate_TargetArraySlicesIndependently | TexCreate_DisableDCC | TexCreate_NoFastClear | TexCreate_RenderTargetable | TexCreate_ShaderResource | TexCreate_UAV, SliceCount, 1, 1);
-	StrataSceneData.MaterialTextureArray = GraphBuilder.CreateTexture(MaterialTextureDesc, TEXT("Strata.Material"));
-	StrataSceneData.MaterialTextureArraySRV = GraphBuilder.CreateSRV(FRDGTextureSRVDesc::Create(StrataSceneData.MaterialTextureArray));
-	StrataSceneData.MaterialTextureArrayUAV = GraphBuilder.CreateUAV(FRDGTextureUAVDesc(StrataSceneData.MaterialTextureArray, 0));
+	Out.MaterialTextureArray = GraphBuilder.CreateTexture(MaterialTextureDesc, TEXT("Strata.Material"));
+	Out.MaterialTextureArraySRV = GraphBuilder.CreateSRV(FRDGTextureSRVDesc::Create(Out.MaterialTextureArray));
+	Out.MaterialTextureArrayUAV = GraphBuilder.CreateUAV(FRDGTextureUAVDesc(Out.MaterialTextureArray, 0));
 
 	// See AppendStrataMRTs
 	check(STRATA_BASE_PASS_MRT_OUTPUT_COUNT <= SliceCount);
-	StrataSceneData.MaterialTextureArrayUAVWithoutRTs = GraphBuilder.CreateUAV(FRDGTextureUAVDesc(StrataSceneData.MaterialTextureArray, 0, PF_Unknown, STRATA_BASE_PASS_MRT_OUTPUT_COUNT, SliceCount - STRATA_BASE_PASS_MRT_OUTPUT_COUNT));
+	Out.MaterialTextureArrayUAVWithoutRTs = GraphBuilder.CreateUAV(FRDGTextureUAVDesc(Out.MaterialTextureArray, 0, PF_Unknown, STRATA_BASE_PASS_MRT_OUTPUT_COUNT, SliceCount - STRATA_BASE_PASS_MRT_OUTPUT_COUNT));
 
 	// Rough diffuse model
-	StrataSceneData.bRoughDiffuse = CVarStrataRoughDiffuse.GetValueOnRenderThread() > 0 ? 1u : 0u;
-
-	// Set reference to the Strata data from each view
-	for (int32 ViewIndex = 0; ViewIndex < SceneRenderer.Views.Num(); ViewIndex++)
-	{
-		FViewInfo& View = SceneRenderer.Views[ViewIndex];
-		View.StrataSceneData = &SceneRenderer.Scene->StrataSceneData;
-	}
+	Out.bRoughDiffuse = CVarStrataRoughDiffuse.GetValueOnRenderThread() > 0 ? 1u : 0u;
 
 	if (IsStrataEnabled())
 	{
 		AddStrataClearMaterialBufferPass(
 			GraphBuilder, 
-			GraphBuilder.CreateUAV(FRDGTextureUAVDesc(StrataSceneData.MaterialTextureArray, 0)),
-			StrataSceneData.SSSTextureUAV,
-			StrataSceneData.MaxBytesPerPixel, 
+			GraphBuilder.CreateUAV(FRDGTextureUAVDesc(Out.MaterialTextureArray, 0)),
+			Out.SSSTextureUAV,
+			Out.MaxBytesPerPixel, 
 			MaterialBufferSizeXY);
 	}
 
-	// Create the readable uniform buffers for each views once for all (it is view independent and all the views should be tiled into the render target textures & material buffer)
-	if (IsStrataEnabled())
+	// Initialized view data
+	for (int32 ViewIndex = 0; ViewIndex < SceneRenderer.Views.Num(); ViewIndex++)
 	{
-		FStrataGlobalUniformParameters* StrataUniformParameters = GraphBuilder.AllocParameters<FStrataGlobalUniformParameters>();
-		BindStrataGlobalUniformParameters(GraphBuilder, &StrataSceneData, *StrataUniformParameters);
-		StrataSceneData.StrataGlobalUniformParameters = GraphBuilder.CreateUniformBuffer(StrataUniformParameters);
+		Strata::InitialiseStrataViewData(GraphBuilder, SceneRenderer.Views[ViewIndex], Out);
 	}
 }
 
 void BindStrataBasePassUniformParameters(FRDGBuilder& GraphBuilder, const FViewInfo& View, FStrataBasePassUniformParameters& OutStrataUniformParameters)
 {
-	FStrataSceneData* StrataSceneData = View.StrataSceneData;
+	const FStrataSceneData* StrataSceneData = View.StrataViewData.SceneData;
 	if (IsStrataEnabled() && StrataSceneData)
 	{
 		OutStrataUniformParameters.bRoughDiffuse = StrataSceneData->bRoughDiffuse ? 1u : 0u;
@@ -360,22 +352,23 @@ void BindStrataBasePassUniformParameters(FRDGBuilder& GraphBuilder, const FViewI
 	}
 }
 
-static void BindStrataGlobalUniformParameters(FRDGBuilder& GraphBuilder, FStrataSceneData* StrataSceneData, FStrataGlobalUniformParameters& OutStrataUniformParameters)
+static void BindStrataGlobalUniformParameters(FRDGBuilder& GraphBuilder, FStrataViewData* StrataViewData, FStrataGlobalUniformParameters& OutStrataUniformParameters)
 {
+	FStrataSceneData* StrataSceneData = StrataViewData->SceneData;
 	if (IsStrataEnabled() && StrataSceneData)
 	{
 		OutStrataUniformParameters.bRoughDiffuse = StrataSceneData->bRoughDiffuse ? 1u : 0u;
 		OutStrataUniformParameters.MaxBytesPerPixel = StrataSceneData->MaxBytesPerPixel;
 		OutStrataUniformParameters.TileSize = STRATA_TILE_SIZE;
 		OutStrataUniformParameters.TileSizeLog2 = STRATA_TILE_SIZE_DIV_AS_SHIFT;
-		OutStrataUniformParameters.TileCount = StrataSceneData->TileCount_Primary;
+		OutStrataUniformParameters.TileCount = StrataViewData->TileCount_Primary;
 		OutStrataUniformParameters.MaterialTextureArray = StrataSceneData->MaterialTextureArray;
 		OutStrataUniformParameters.TopLayerTexture = StrataSceneData->TopLayerTexture;
 		OutStrataUniformParameters.SSSTexture = StrataSceneData->SSSTexture;
 		OutStrataUniformParameters.OpaqueRoughRefractionTexture = StrataSceneData->OpaqueRoughRefractionTexture;
-		OutStrataUniformParameters.BSDFTileTexture = StrataSceneData->BSDFTileTexture;
+		OutStrataUniformParameters.BSDFTileTexture = StrataViewData->BSDFTileTexture;
 		OutStrataUniformParameters.BSDFOffsetTexture = StrataSceneData->BSDFOffsetTexture;
-		OutStrataUniformParameters.BSDFTileCountBuffer = GraphBuilder.CreateSRV(StrataSceneData->BSDFTileCountBuffer, PF_R32_UINT);
+		OutStrataUniformParameters.BSDFTileCountBuffer = GraphBuilder.CreateSRV(StrataViewData->BSDFTileCountBuffer, PF_R32_UINT);
 	}
 	else
 	{
@@ -399,7 +392,7 @@ static void BindStrataGlobalUniformParameters(FRDGBuilder& GraphBuilder, FStrata
 
 void BindStrataForwardPasslUniformParameters(FRDGBuilder& GraphBuilder, const FViewInfo& View, FStrataForwardPassUniformParameters& OutStrataUniformParameters)
 {
-	FStrataSceneData* StrataSceneData = View.StrataSceneData;
+	FStrataSceneData* StrataSceneData = View.StrataViewData.SceneData;
 	if (IsStrataEnabled() && StrataSceneData)
 	{
 		OutStrataUniformParameters.bRoughDiffuse = StrataSceneData->bRoughDiffuse ? 1u : 0u;
@@ -412,10 +405,8 @@ void BindStrataForwardPasslUniformParameters(FRDGBuilder& GraphBuilder, const FV
 
 TRDGUniformBufferRef<FStrataGlobalUniformParameters> BindStrataGlobalUniformParameters(const FViewInfo& View)
 {
-	FStrataSceneData* StrataSceneData = View.StrataSceneData;
-	check(StrataSceneData);
-	check(StrataSceneData->StrataGlobalUniformParameters != nullptr || !IsStrataEnabled());
-	return StrataSceneData->StrataGlobalUniformParameters;
+	check(View.StrataViewData.StrataGlobalUniformParameters != nullptr || !IsStrataEnabled());
+	return View.StrataViewData.StrataGlobalUniformParameters;
 }
 
 ////////////////////////////////////////////////////////////////////////// 
@@ -720,10 +711,10 @@ IMPLEMENT_GLOBAL_SHADER(FStrataMaterialStencilTaggingPassPS, "/Engine/Private/St
 static FStrataTileParameter InternalSetTileParameters(FRDGBuilder* GraphBuilder, const FViewInfo& View, const EStrataTileType TileType)
 {
 	FStrataTileParameter Out;
-	if (View.StrataSceneData && TileType != EStrataTileType::ECount)
+	if (TileType != EStrataTileType::ECount)
 	{
-		Out.TileListBuffer = View.StrataSceneData->ClassificationTileListBufferSRV[TileType];
-		Out.TileIndirectBuffer = View.StrataSceneData->ClassificationTileDrawIndirectBuffer;
+		Out.TileListBuffer = View.StrataViewData.ClassificationTileListBufferSRV[TileType];
+		Out.TileIndirectBuffer = View.StrataViewData.ClassificationTileDrawIndirectBuffer;
 	}
 	else if (GraphBuilder)
 	{
@@ -744,6 +735,7 @@ FStrataTilePassVS::FParameters SetTileParameters(
 	PrimitiveType = GRHISupportsRectTopology ? PT_RectList : PT_TriangleList;
 
 	FStrataTilePassVS::FParameters Out;
+	Out.OutputViewMinRect = FVector2f(View.CachedViewUniformShaderParameters->ViewRectMin.X, View.CachedViewUniformShaderParameters->ViewRectMin.Y);
 	Out.OutputViewSizeAndInvSize = View.CachedViewUniformShaderParameters->ViewSizeAndInvSize;
 	Out.OutputBufferSizeAndInvSize = View.CachedViewUniformShaderParameters->BufferSizeAndInvSize;
 	Out.ViewScreenToTranslatedWorld = View.CachedViewUniformShaderParameters->ScreenToTranslatedWorld;
@@ -762,6 +754,7 @@ FStrataTilePassVS::FParameters SetTileParameters(
 	PrimitiveType = GRHISupportsRectTopology ? PT_RectList : PT_TriangleList;
 
 	FStrataTilePassVS::FParameters Out;
+	Out.OutputViewMinRect = FVector2f(View.CachedViewUniformShaderParameters->ViewRectMin.X, View.CachedViewUniformShaderParameters->ViewRectMin.Y);
 	Out.OutputViewSizeAndInvSize = View.CachedViewUniformShaderParameters->ViewSizeAndInvSize;
 	Out.OutputBufferSizeAndInvSize = View.CachedViewUniformShaderParameters->BufferSizeAndInvSize;
 	Out.ViewScreenToTranslatedWorld = View.CachedViewUniformShaderParameters->ScreenToTranslatedWorld;
@@ -803,7 +796,7 @@ static void AddStrataInternalClassificationTilePass(
 {
 	EPrimitiveType StrataTilePrimitiveType = PT_TriangleList;
 	const FIntPoint OutputResolution = View.ViewRect.Size();
-	FVector4f OutputResolutionAndInv = FVector4f(OutputResolution.X, OutputResolution.Y, 1.0f / float(OutputResolution.X), 1.0f / float(OutputResolution.Y));
+	const FIntRect ViewRect = View.ViewRect;
 
 	FStrataMaterialStencilTaggingPassPS::FParameters* ParametersPS = GraphBuilder.AllocParameters<FStrataMaterialStencilTaggingPassPS::FParameters>();
 	ParametersPS->VS = Strata::SetTileParameters(GraphBuilder, View, TileMaterialType, StrataTilePrimitiveType);
@@ -844,7 +837,7 @@ static void AddStrataInternalClassificationTilePass(
 		RDG_EVENT_NAME("Strata::%sClassificationPass(%s)", bDebug ? TEXT("Debug") : TEXT("Stencil"), ToString(TileMaterialType)),
 		ParametersPS,
 		ERDGPassFlags::Raster,
-		[ParametersPS, VertexShader, PixelShader, OutputResolution, StrataTilePrimitiveType, TileMaterialType, bDebug](FRHICommandList& RHICmdList)
+		[ParametersPS, VertexShader, PixelShader, ViewRect, OutputResolution, StrataTilePrimitiveType, TileMaterialType, bDebug](FRHICommandList& RHICmdList)
 		{
 			FGraphicsPipelineStateInitializer GraphicsPSOInit;
 			RHICmdList.ApplyCachedRenderTargets(GraphicsPSOInit);
@@ -905,10 +898,14 @@ static void AddStrataInternalClassificationTilePass(
 			SetShaderParameters(RHICmdList, VertexShader, VertexShader.GetVertexShader(), ParametersPS->VS);
 			if (bDebug)
 			{
+				// Debug rendering is aways done during the post-processing stage, which has an ViewMinRect set to (0,0)
 				SetShaderParameters(RHICmdList, PixelShader, PixelShader.GetPixelShader(), *ParametersPS);
+				RHICmdList.SetViewport(0, 0, 0.0f, OutputResolution.X, OutputResolution.Y, 1.0f);
+			}
+			else
+			{
+				RHICmdList.SetViewport(ViewRect.Min.X, ViewRect.Min.Y, 0.0f, ViewRect.Max.X, ViewRect.Max.Y, 1.0f);
 			}
-
-			RHICmdList.SetViewport(0, 0, 0.0f, OutputResolution.X, OutputResolution.Y, 1.0f);
 			RHICmdList.SetStreamSource(0, nullptr, 0);
 			RHICmdList.DrawPrimitiveIndirect(ParametersPS->VS.TileIndirectBuffer->GetIndirectRHICallBuffer(), TileTypeDrawIndirectArgOffset(TileMaterialType));
 		});
@@ -993,6 +990,9 @@ void AddStrataMaterialClassificationPass(FRDGBuilder& GraphBuilder, const FMinim
 		bWaveOps = bWaveOps && !IsRHIDeviceNVIDIA();
 	#endif
 		
+		const FStrataViewData* StrataViewData = &View.StrataViewData;
+		const FStrataSceneData* StrataSceneData = View.StrataViewData.SceneData;
+
 		// Tile reduction
 		{
 			const bool bClear = ClearDuringCategorization();
@@ -1004,17 +1004,17 @@ void AddStrataMaterialClassificationPass(FRDGBuilder& GraphBuilder, const FMinim
 			PassParameters->ViewUniformBuffer = View.ViewUniformBuffer;
 			PassParameters->bRectPrimitive = GRHISupportsRectTopology ? 1 : 0;
 			PassParameters->ViewResolution = View.ViewRect.Size();
-			PassParameters->MaxBytesPerPixel = View.StrataSceneData->MaxBytesPerPixel;
-			PassParameters->TopLayerTexture = View.StrataSceneData->TopLayerTexture;
-			PassParameters->MaterialTextureArray = View.StrataSceneData->MaterialTextureArraySRV;
-			PassParameters->SSSTextureUAV = View.StrataSceneData->SSSTextureUAV;
-			PassParameters->OpaqueRoughRefractionTexture = View.StrataSceneData->OpaqueRoughRefractionTexture;
-			PassParameters->TileDrawIndirectDataBuffer = View.StrataSceneData->ClassificationTileDrawIndirectBufferUAV;
-			PassParameters->SimpleTileListDataBuffer = View.StrataSceneData->ClassificationTileListBufferUAV[EStrataTileType::ESimple];
-			PassParameters->SingleTileListDataBuffer = View.StrataSceneData->ClassificationTileListBufferUAV[EStrataTileType::ESingle];
-			PassParameters->ComplexTileListDataBuffer = View.StrataSceneData->ClassificationTileListBufferUAV[EStrataTileType::EComplex];
-			PassParameters->OpaqueRoughRefractionTileListDataBuffer = View.StrataSceneData->ClassificationTileListBufferUAV[EStrataTileType::EOpaqueRoughRefraction];
-			PassParameters->SSSWithoutOpaqueRoughRefractionTileListDataBuffer = View.StrataSceneData->ClassificationTileListBufferUAV[EStrataTileType::ESSSWithoutOpaqueRoughRefraction];
+			PassParameters->MaxBytesPerPixel = StrataSceneData->MaxBytesPerPixel;
+			PassParameters->TopLayerTexture = StrataSceneData->TopLayerTexture;
+			PassParameters->MaterialTextureArray = StrataSceneData->MaterialTextureArraySRV;
+			PassParameters->SSSTextureUAV = StrataSceneData->SSSTextureUAV;
+			PassParameters->OpaqueRoughRefractionTexture = StrataSceneData->OpaqueRoughRefractionTexture;
+			PassParameters->TileDrawIndirectDataBuffer = StrataViewData->ClassificationTileDrawIndirectBufferUAV;
+			PassParameters->SimpleTileListDataBuffer = StrataViewData->ClassificationTileListBufferUAV[EStrataTileType::ESimple];
+			PassParameters->SingleTileListDataBuffer = StrataViewData->ClassificationTileListBufferUAV[EStrataTileType::ESingle];
+			PassParameters->ComplexTileListDataBuffer = StrataViewData->ClassificationTileListBufferUAV[EStrataTileType::EComplex];
+			PassParameters->OpaqueRoughRefractionTileListDataBuffer = StrataViewData->ClassificationTileListBufferUAV[EStrataTileType::EOpaqueRoughRefraction];
+			PassParameters->SSSWithoutOpaqueRoughRefractionTileListDataBuffer = StrataViewData->ClassificationTileListBufferUAV[EStrataTileType::ESSSWithoutOpaqueRoughRefraction];
 
 			const uint32 GroupSize = 8;
 			FComputeShaderUtils::AddPass(
@@ -1029,8 +1029,8 @@ void AddStrataMaterialClassificationPass(FRDGBuilder& GraphBuilder, const FMinim
 		{
 			TShaderMapRef<FStrataMaterialTilePrepareArgsPassCS> ComputeShader(View.ShaderMap);
 			FStrataMaterialTilePrepareArgsPassCS::FParameters* PassParameters = GraphBuilder.AllocParameters<FStrataMaterialTilePrepareArgsPassCS::FParameters>();
-			PassParameters->TileDrawIndirectDataBuffer = GraphBuilder.CreateSRV(View.StrataSceneData->ClassificationTileDrawIndirectBuffer, PF_R32_UINT);
-			PassParameters->TileDispatchIndirectDataBuffer = View.StrataSceneData->ClassificationTileDispatchIndirectBufferUAV;
+			PassParameters->TileDrawIndirectDataBuffer = GraphBuilder.CreateSRV(StrataViewData->ClassificationTileDrawIndirectBuffer, PF_R32_UINT);
+			PassParameters->TileDispatchIndirectDataBuffer = StrataViewData->ClassificationTileDispatchIndirectBufferUAV;
 
 			FComputeShaderUtils::AddPass(
 				GraphBuilder,
@@ -1042,7 +1042,7 @@ void AddStrataMaterialClassificationPass(FRDGBuilder& GraphBuilder, const FMinim
 
 		// Compute BSDF tile index and material read offset
 		{
-			FRDGBufferUAVRef RWBSDFTileCountBuffer = GraphBuilder.CreateUAV(View.StrataSceneData->BSDFTileCountBuffer, PF_R32_UINT);
+			FRDGBufferUAVRef RWBSDFTileCountBuffer = GraphBuilder.CreateUAV(StrataViewData->BSDFTileCountBuffer, PF_R32_UINT);
 			AddClearUAVPass(GraphBuilder, RWBSDFTileCountBuffer, 0u);
 
 			FStrataBSDFTilePassCS::FPermutationDomain PermutationVector;
@@ -1051,16 +1051,16 @@ void AddStrataMaterialClassificationPass(FRDGBuilder& GraphBuilder, const FMinim
 			FStrataBSDFTilePassCS::FParameters* PassParameters = GraphBuilder.AllocParameters<FStrataBSDFTilePassCS::FParameters>();
 			PassParameters->ViewUniformBuffer = View.ViewUniformBuffer;
 			PassParameters->TileSizeLog2 = STRATA_TILE_SIZE_DIV_AS_SHIFT;
-			PassParameters->TileCount_Primary = View.StrataSceneData->TileCount_Primary;
+			PassParameters->TileCount_Primary = StrataViewData->TileCount_Primary;
 			PassParameters->ViewResolution = View.ViewRect.Size();
-			PassParameters->MaxBytesPerPixel = View.StrataSceneData->MaxBytesPerPixel;
-			PassParameters->TopLayerTexture = View.StrataSceneData->TopLayerTexture;
-			PassParameters->MaterialTextureArray = View.StrataSceneData->MaterialTextureArraySRV;
-			PassParameters->TileListBuffer = View.StrataSceneData->ClassificationTileListBufferSRV[EStrataTileType::EComplex];
-			PassParameters->TileIndirectBuffer = View.StrataSceneData->ClassificationTileDispatchIndirectBuffer;
-
-			PassParameters->RWBSDFOffsetTexture = GraphBuilder.CreateUAV(View.StrataSceneData->BSDFOffsetTexture);
-			PassParameters->RWBSDFTileTexture = GraphBuilder.CreateUAV(View.StrataSceneData->BSDFTileTexture);
+			PassParameters->MaxBytesPerPixel = StrataSceneData->MaxBytesPerPixel;
+			PassParameters->TopLayerTexture = StrataSceneData->TopLayerTexture;
+			PassParameters->MaterialTextureArray = StrataSceneData->MaterialTextureArraySRV;
+			PassParameters->TileListBuffer = StrataViewData->ClassificationTileListBufferSRV[EStrataTileType::EComplex];
+			PassParameters->TileIndirectBuffer = StrataViewData->ClassificationTileDispatchIndirectBuffer;
+
+			PassParameters->RWBSDFOffsetTexture = GraphBuilder.CreateUAV(StrataSceneData->BSDFOffsetTexture);
+			PassParameters->RWBSDFTileTexture = GraphBuilder.CreateUAV(StrataViewData->BSDFTileTexture);
 			PassParameters->RWBSDFTileCountBuffer = RWBSDFTileCountBuffer;
 
 			FComputeShaderUtils::AddPass(
@@ -1076,9 +1076,9 @@ void AddStrataMaterialClassificationPass(FRDGBuilder& GraphBuilder, const FMinim
 		{
 			TShaderMapRef<FStrataBSDFTilePrepareArgsPassCS> ComputeShader(View.ShaderMap);
 			FStrataBSDFTilePrepareArgsPassCS::FParameters* PassParameters = GraphBuilder.AllocParameters<FStrataBSDFTilePrepareArgsPassCS::FParameters>();
-			PassParameters->TileCount_Primary = View.StrataSceneData->TileCount_Primary;
-			PassParameters->TileDrawIndirectDataBuffer = GraphBuilder.CreateSRV(View.StrataSceneData->BSDFTileCountBuffer, PF_R32_UINT);
-			PassParameters->TileDispatchIndirectDataBuffer = GraphBuilder.CreateUAV(View.StrataSceneData->BSDFTileDispatchIndirectBuffer, PF_R32_UINT);
+			PassParameters->TileCount_Primary = StrataViewData->TileCount_Primary;
+			PassParameters->TileDrawIndirectDataBuffer = GraphBuilder.CreateSRV(StrataViewData->BSDFTileCountBuffer, PF_R32_UINT);
+			PassParameters->TileDispatchIndirectDataBuffer = GraphBuilder.CreateUAV(StrataViewData->BSDFTileDispatchIndirectBuffer, PF_R32_UINT);
 
 			FComputeShaderUtils::AddPass(
 				GraphBuilder,
diff --git a/Engine/Source/Runtime/Renderer/Private/Strata/Strata.h b/Engine/Source/Runtime/Renderer/Private/Strata/Strata.h
--- a/Engine/Source/Runtime/Renderer/Private/Strata/Strata.h
+++ b/Engine/Source/Runtime/Renderer/Private/Strata/Strata.h
@@ -68,57 +68,57 @@ struct FStrataSceneData
 	uint32 MaxBytesPerPixel;
 	bool bRoughDiffuse;
 
-	FIntPoint TileCount_Primary;
-	FIntPoint TileCount_Overflow;
-	FIntPoint TileCount_Total;
-
 	// Resources allocated and updated each frame
 
-	FRDGTextureRef MaterialTextureArray;
-	FRDGTextureUAVRef MaterialTextureArrayUAVWithoutRTs;
-	FRDGTextureUAVRef MaterialTextureArrayUAV;
-	FRDGTextureSRVRef MaterialTextureArraySRV;
+	FRDGTextureRef MaterialTextureArray = nullptr;
+	FRDGTextureUAVRef MaterialTextureArrayUAVWithoutRTs = nullptr;
+	FRDGTextureUAVRef MaterialTextureArrayUAV = nullptr;
+	FRDGTextureSRVRef MaterialTextureArraySRV = nullptr;
 
-	FRDGBufferRef ClassificationTileListBuffer[STRATA_TILE_TYPE_COUNT];
-	FRDGBufferSRVRef ClassificationTileListBufferSRV[STRATA_TILE_TYPE_COUNT];
-	FRDGBufferUAVRef ClassificationTileListBufferUAV[STRATA_TILE_TYPE_COUNT];
+	FRDGTextureRef TopLayerTexture = nullptr;
+	FRDGTextureRef SSSTexture = nullptr;
+	FRDGTextureRef OpaqueRoughRefractionTexture = nullptr;
 
-	FRDGBufferRef    ClassificationTileDrawIndirectBuffer;
-	FRDGBufferUAVRef ClassificationTileDrawIndirectBufferUAV;
+	FRDGTextureUAVRef TopLayerTextureUAV = nullptr;
+	FRDGTextureUAVRef SSSTextureUAV = nullptr;
+	FRDGTextureUAVRef OpaqueRoughRefractionTextureUAV = nullptr;
 
-	FRDGBufferRef    ClassificationTileDispatchIndirectBuffer;
-	FRDGBufferUAVRef ClassificationTileDispatchIndirectBufferUAV;
+	FRDGTextureRef BSDFOffsetTexture = nullptr;
 
-	FRDGTextureRef TopLayerTexture;
-	FRDGTextureRef SSSTexture;
-	FRDGTextureRef OpaqueRoughRefractionTexture;
+	// Used when the subsurface luminance is separated from the scene color
+	FRDGTextureRef SeparatedSubSurfaceSceneColor = nullptr;
 
-	FRDGTextureUAVRef TopLayerTextureUAV;
-	FRDGTextureUAVRef SSSTextureUAV;
-	FRDGTextureUAVRef OpaqueRoughRefractionTextureUAV;
+	// Used for Luminance that should go through opaque rough refraction (when under a top layer interface)
+	FRDGTextureRef SeparatedOpaqueRoughRefractionSceneColor = nullptr;
+};
 
-	FRDGTextureRef BSDFOffsetTexture;
-	FRDGTextureRef BSDFTileTexture;
-	FRDGBufferRef  BSDFTileCountBuffer;
-	FRDGBufferRef  BSDFTileDispatchIndirectBuffer;
+struct FStrataViewData
+{
+	FIntPoint TileCount_Primary = FIntPoint(0, 0);
+	FIntPoint TileCount_Overflow = FIntPoint(0, 0);
+	FIntPoint TileCount_Total = FIntPoint(0, 0);
 
-	// Used when the subsurface luminance is separated from the scene color
-	FRDGTextureRef SeparatedSubSurfaceSceneColor;
+	FRDGBufferRef    ClassificationTileListBuffer[STRATA_TILE_TYPE_COUNT];
+	FRDGBufferSRVRef ClassificationTileListBufferSRV[STRATA_TILE_TYPE_COUNT];
+	FRDGBufferUAVRef ClassificationTileListBufferUAV[STRATA_TILE_TYPE_COUNT];
 
-	// Used for Luminance that should go through opaque rough refraction (when under a top layer interface)
-	FRDGTextureRef SeparatedOpaqueRoughRefractionSceneColor;
+	FRDGBufferRef    ClassificationTileDrawIndirectBuffer = nullptr;
+	FRDGBufferUAVRef ClassificationTileDrawIndirectBufferUAV = nullptr;
 
-	TRDGUniformBufferRef<FStrataGlobalUniformParameters> StrataGlobalUniformParameters{};
+	FRDGBufferRef    ClassificationTileDispatchIndirectBuffer = nullptr;
+	FRDGBufferUAVRef ClassificationTileDispatchIndirectBufferUAV = nullptr;
 
-	FStrataSceneData()
-	{
-		Reset();
-	}
+	FRDGTextureRef BSDFTileTexture = nullptr;
+	FRDGBufferRef  BSDFTileCountBuffer = nullptr;
+	FRDGBufferRef  BSDFTileDispatchIndirectBuffer = nullptr;
+
+	FStrataSceneData* SceneData = nullptr;
+
+	TRDGUniformBufferRef<FStrataGlobalUniformParameters> StrataGlobalUniformParameters{};
 
 	void Reset();
 };
 
-
 namespace Strata
 {
 constexpr uint32 StencilBit_Fast   = 0x08; // In sync with SceneRenderTargets.h - GET_STENCIL_BIT_MASK(STENCIL_STRATA_FASTPATH)
@@ -129,7 +129,7 @@ bool IsStrataEnabled();
 
 FIntPoint GetStrataTextureResolution(const FIntPoint& InResolution);
 
-void InitialiseStrataFrameSceneData(FSceneRenderer& SceneRenderer, FRDGBuilder& GraphBuilder);
+void InitialiseStrataFrameSceneData(FRDGBuilder& GraphBuilder, FSceneRenderer& SceneRenderer);
 
 void BindStrataBasePassUniformParameters(FRDGBuilder& GraphBuilder, const FViewInfo& View, FStrataBasePassUniformParameters& OutStrataUniformParameters);
 void BindStrataForwardPasslUniformParameters(FRDGBuilder& GraphBuilder, const FViewInfo& View, FStrataForwardPassUniformParameters& OutStrataUniformParameters);
@@ -165,6 +165,7 @@ class FStrataTilePassVS : public FGlobalShader
 		// It would be possible to use the ViewUniformBuffer instead of copying the data here, 
 		// but we would have to make sure the view UB is added to all passes using this parameter structure.
 		// We should not add it here to now have duplicated input UB.
+		SHADER_PARAMETER(FVector2f, OutputViewMinRect)
 		SHADER_PARAMETER(FVector4f, OutputViewSizeAndInvSize)
 		SHADER_PARAMETER(FVector4f, OutputBufferSizeAndInvSize)
 		SHADER_PARAMETER(FMatrix44f, ViewScreenToTranslatedWorld)
diff --git a/Engine/Source/Runtime/Renderer/Private/Strata/StrataRoughRefraction.cpp b/Engine/Source/Runtime/Renderer/Private/Strata/StrataRoughRefraction.cpp
--- a/Engine/Source/Runtime/Renderer/Private/Strata/StrataRoughRefraction.cpp
+++ b/Engine/Source/Runtime/Renderer/Private/Strata/StrataRoughRefraction.cpp
@@ -100,7 +100,7 @@ void AddStrataOpaqueRoughRefractionPasses(
 			TempTexture = GraphBuilder.CreateTexture(FRDGTextureDesc::Create2D(GetSceneTextureExtent(), PF_FloatRGBA, FClearValueBinding::Black, TexCreate_UAV | TexCreate_ShaderResource | TexCreate_RenderTargetable), TEXT("Strata.RoughRefrac.TempTexture"));
 		}
 
-		FRDGTextureRef SeparatedOpaqueRoughRefractionSceneColor = View.StrataSceneData->SeparatedOpaqueRoughRefractionSceneColor;
+		FRDGTextureRef SeparatedOpaqueRoughRefractionSceneColor = View.StrataViewData.SceneData->SeparatedOpaqueRoughRefractionSceneColor;
 
 		FSceneTextureParameters SceneTextureParameters = GetSceneTextureParameters(GraphBuilder);
 		FOpaqueRoughRefractionPS::FParameters* PassParameters = GraphBuilder.AllocParameters<FOpaqueRoughRefractionPS::FParameters>();
@@ -215,7 +215,7 @@ void AddStrataOpaqueRoughRefractionPasses(
 	{
 		const FViewInfo& View = Views[ViewIndex];
 
-		FRDGTextureRef SeparatedOpaqueRoughRefractionSceneColor = View.StrataSceneData->SeparatedOpaqueRoughRefractionSceneColor;
+		FRDGTextureRef SeparatedOpaqueRoughRefractionSceneColor = View.StrataViewData.SceneData->SeparatedOpaqueRoughRefractionSceneColor;
 
 		FSceneTextureParameters SceneTextureParameters = GetSceneTextureParameters(GraphBuilder);
 		FOpaqueRoughRefractionPS::FParameters* PassParameters = GraphBuilder.AllocParameters<FOpaqueRoughRefractionPS::FParameters>();
