Add support for distance field shadow on water.
Storing water in gbuffer and runnign extra lighting pass would be too expenssiv. So instead, when in deferred, single layer water always output the main dir light luminance into a separate 111110 texture.
DF shadow is applied on that texture. DFShadow culling step is cached the first time for a given view in order to be able to reapply DFShadow later with a depth buffer containing water. This save 0.4ms on a base console.
DFShadow is naturaly skipped if the directional light does not have that feature enabled.
When DFShadow is executed, it is basically 0.24ms extra on a base console @ 1080p.
The separated light buffer is an extra 8MB allocated when in deferred.
This only works with deferred for now (forward does not run the reflection/composite pass).

Render target slot assignement is a bit fuzzy but this is a mostly static setup so should be fine. And could change with Strata down the line.

r.Water.SingleLayer.ShaderSupportDistanceFieldShadow is to enabled dfshadow on water for the shaders. Defaults to 1 and will only be enalbed when using deferred and when distance fields are enabled.
r.Water.SingleLayer.DistanceFieldShadow is a runtime toggle of the feature (if enabled the shaders still output to the separate light render target, but that is simply discarded).

Follow up: tile upsample, tile traceshadow, make it work with strata.

Tested shooter and console.

diff --git a/Engine/Shaders/Private/BasePassPixelShader.usf b/Engine/Shaders/Private/BasePassPixelShader.usf
--- a/Engine/Shaders/Private/BasePassPixelShader.usf
+++ b/Engine/Shaders/Private/BasePassPixelShader.usf
@@ -183,6 +183,10 @@ struct FShadingOcclusion
 #include "/Engine/Private/Strata/StrataForwardLighting.ush"
 #endif
 
+// Separate main directional light from scene color buffer if water needs to receive distance field shadows.
+// Sperated dir light should match with FShaderCompileUtilities::FetchGBufferParamsRuntime.
+#define SINGLE_LAYER_WATER_SEPARATED_MAIN_LIGHT  (SINGLE_LAYER_WATER_DF_SHADOW_ENABLED)
+
 #if !FORWARD_SHADING
 
 void GetVolumeLightingNonDirectional(float4 AmbientLightingVector, float3 DiffuseColor, inout float3 InterpolatedLighting, out float4 VolumeLighting)
@@ -1350,6 +1354,9 @@ void FPixelShaderInOut_MainPS(
 	float3 DiffuseIndirectLighting = 0;
 	float3 SubsurfaceIndirectLighting = 0;
 
+	bool bSeparateWaterMainDirLightLuminance = SINGLE_LAYER_WATER_SEPARATED_MAIN_LIGHT > 0 && View.SeparateWaterMainDirLightLuminance > 0.0f; 
+	float3 SeparatedWaterMainDirLightLuminance = float3(0, 0, 0); 
+
 #if !STRATA_ENABLED || STRATA_INLINE_SINGLELAYERWATER
 
 	#if !MATERIAL_SHADINGMODEL_UNLIT
@@ -1433,7 +1440,10 @@ void FPixelShaderInOut_MainPS(
 				#endif
 
 				#if FORWARD_SHADING || TRANSLUCENCY_LIGHTING_SURFACE_FORWARDSHADING || MATERIAL_SHADINGMODEL_SINGLELAYERWATER
-					FDeferredLightingSplit ForwardDirectLighting = GetForwardDirectLightingSplit(GridIndex, MaterialParameters.WorldPosition_CamRelative, MaterialParameters.CameraVector, GBuffer, NearestResolvedDepthScreenUV, MaterialParameters.PrimitiveId, EyeIndex, Dither, DirectionalLightCloudShadow, DirectionalLightShadow);
+					FDeferredLightingSplit ForwardDirectLighting = GetForwardDirectLightingSplit(
+						GridIndex, MaterialParameters.WorldPosition_CamRelative, MaterialParameters.CameraVector, GBuffer, NearestResolvedDepthScreenUV, MaterialParameters.PrimitiveId, EyeIndex, Dither, 
+						DirectionalLightCloudShadow, DirectionalLightShadow,
+						bSeparateWaterMainDirLightLuminance, SeparatedWaterMainDirLightLuminance);
 					
 					#if MATERIAL_SHADINGMODEL_THIN_TRANSLUCENT
 						DiffuseColor += ForwardDirectLighting.DiffuseLighting.rgb;
@@ -1631,7 +1641,8 @@ void FPixelShaderInOut_MainPS(
 			OpaqueBasePass.SceneDepthWithoutSingleLayerWaterTexture, SingleLayerWaterSceneDepthSampler, // Scene depth texture
 			Specular, NullDistortionParams,
 			SunIlluminance, WaterDiffuseIndirectIlluminance, EnvBrdf,
-			CameraIsUnderWater, WaterVisibility, EyeIndex);
+			CameraIsUnderWater, WaterVisibility, EyeIndex,
+			bSeparateWaterMainDirLightLuminance, SeparatedWaterMainDirLightLuminance);
 
 		// Add water luminance contribution
 		Color += WaterLighting.Luminance;
@@ -1647,7 +1658,8 @@ void FPixelShaderInOut_MainPS(
 			OpaqueBasePass.SceneWithoutSingleLayerWaterMinMaxUV.zw,
 			Specular, OpaqueBasePass.DistortionParams,
 			SunIlluminance, WaterDiffuseIndirectIlluminance, EnvBrdf,
-			CameraIsUnderWater, WaterVisibility, EyeIndex).Luminance;
+			CameraIsUnderWater, WaterVisibility, EyeIndex,
+			bSeparateWaterMainDirLightLuminance, SeparatedWaterMainDirLightLuminance).Luminance;
 #endif
 	}
 #endif // MATERIAL_SHADINGMODEL_SINGLELAYERWATER
@@ -1954,6 +1966,15 @@ void FPixelShaderInOut_MainPS(
 				Out.MRT[GBUFFER_HAS_VELOCITY ? 5 : 4] = 0;
 				Out.MRT[GBUFFER_HAS_VELOCITY ? 6 : 5] = 0;
 			}
+
+	#if SINGLE_LAYER_WATER_SEPARATED_MAIN_LIGHT
+			if (GBuffer.ShadingModelID == SHADINGMODELID_SINGLELAYERWATER)
+			{
+				// In deferred, we always output the directional light in a separated buffer.
+				// This is used to apply distance field shadows or light function to the main directional light.
+				Out.MRT[GBUFFER_HAS_VELOCITY ? 7 : 6] = float4(SeparatedWaterMainDirLightLuminance * View.PreExposure, 1.0f);
+			}
+	#endif
 		}
 		#else
 		{
diff --git a/Engine/Shaders/Private/Definitions.usf b/Engine/Shaders/Private/Definitions.usf
--- a/Engine/Shaders/Private/Definitions.usf
+++ b/Engine/Shaders/Private/Definitions.usf
@@ -86,6 +86,14 @@
 #define	MATERIAL_SHADINGMODEL_UNLIT						0
 #endif
 
+#ifndef MATERIAL_SHADINGMODEL_SINGLELAYERWATER
+#define MATERIAL_SHADINGMODEL_SINGLELAYERWATER			0 
+#endif
+
+#ifndef SINGLE_LAYER_WATER_DF_SHADOW_ENABLED
+#define SINGLE_LAYER_WATER_DF_SHADOW_ENABLED			0 
+#endif
+
 #ifndef MATERIAL_SINGLE_SHADINGMODEL
 #define	MATERIAL_SINGLE_SHADINGMODEL					0
 #endif
diff --git a/Engine/Shaders/Private/ForwardLightingCommon.ush b/Engine/Shaders/Private/ForwardLightingCommon.ush
--- a/Engine/Shaders/Private/ForwardLightingCommon.ush
+++ b/Engine/Shaders/Private/ForwardLightingCommon.ush
@@ -118,7 +118,10 @@ float GetIndirectOcclusion(float2 ScreenUV, bool bHasDynamicIndirectShadowCaster
 	return IndirectOcclusion;
 }
 
-FDeferredLightingSplit GetForwardDirectLightingSplit(uint GridIndex, float3 TranslatedWorldPosition, float3 CameraVector, FGBufferData GBufferData, float2 ScreenUV, uint PrimitiveId, uint EyeIndex, float Dither, float InDirectionalLightCloudShadow, inout float OutDirectionalLightShadow)
+FDeferredLightingSplit GetForwardDirectLightingSplit(
+	uint GridIndex, float3 TranslatedWorldPosition, float3 CameraVector, FGBufferData GBufferData, float2 ScreenUV, uint PrimitiveId, uint EyeIndex, float Dither, 
+	float InDirectionalLightCloudShadow, inout float OutDirectionalLightShadow,
+	bool bSeparateMainDirLightLuminance, inout float3 SeparatedMainDirLightLuminance)
 {
 	float3 WorldPosition = TranslatedWorldPosition - LWCHackToFloat(PrimaryView.PreViewTranslation);
 
@@ -211,8 +214,16 @@ FDeferredLightingSplit GetForwardDirectLightingSplit(uint GridIndex, float3 Tran
 			NewLighting.SpecularLighting *= InDirectionalLightCloudShadow;
 		#endif
 
-			DirectLighting.DiffuseLighting += NewLighting.DiffuseLighting;
-			DirectLighting.SpecularLighting += NewLighting.SpecularLighting;
+			if (bSeparateMainDirLightLuminance)
+			{
+				SeparatedMainDirLightLuminance += NewLighting.DiffuseLighting.rgb;
+				SeparatedMainDirLightLuminance += NewLighting.SpecularLighting.rgb;
+			}
+			else
+			{
+				DirectLighting.DiffuseLighting += NewLighting.DiffuseLighting;
+				DirectLighting.SpecularLighting += NewLighting.SpecularLighting;
+			}
 		}
 	}
 
@@ -297,7 +308,12 @@ FDeferredLightingSplit GetForwardDirectLightingSplit(uint GridIndex, float3 Tran
 
 float3 GetForwardDirectLighting(uint GridIndex, float3 TranslatedWorldPosition, float3 CameraVector, FGBufferData GBufferData, float2 ScreenUV, uint PrimitiveId, uint EyeIndex, float Dither, float InDirectionalLightCloudShadow, inout float OutDirectionalLightShadow)
 {
-	FDeferredLightingSplit Ret = GetForwardDirectLightingSplit(GridIndex, TranslatedWorldPosition, CameraVector, GBufferData, ScreenUV, PrimitiveId, EyeIndex, Dither, InDirectionalLightCloudShadow, OutDirectionalLightShadow);
+	bool bSeparateMainDirLightLuminance = false;
+	float3 SeparatedMainDirLightLuminance = 0.0f;
+	FDeferredLightingSplit Ret = GetForwardDirectLightingSplit(
+		GridIndex, TranslatedWorldPosition, CameraVector, GBufferData, ScreenUV, PrimitiveId, EyeIndex, Dither, 
+		InDirectionalLightCloudShadow, OutDirectionalLightShadow,
+		bSeparateMainDirLightLuminance, SeparatedMainDirLightLuminance);
 
 	return Ret.DiffuseLighting.rgb + Ret.SpecularLighting.rgb;
 }
diff --git a/Engine/Shaders/Private/MobileBasePassPixelShader.usf b/Engine/Shaders/Private/MobileBasePassPixelShader.usf
--- a/Engine/Shaders/Private/MobileBasePassPixelShader.usf
+++ b/Engine/Shaders/Private/MobileBasePassPixelShader.usf
@@ -997,6 +997,9 @@ void Main(
 
 #if MATERIAL_SHADINGMODEL_SINGLELAYERWATER
 	{
+		const bool bSeparateMainDirLightLuminance = false;
+		float3 SeparatedWaterMainDirLightLuminance = float3(0, 0, 0);
+
 		const bool CameraIsUnderWater = false;	// Fade out the material contribution over to water contribution according to material opacity.
 		const float3 SunIlluminance = ResolvedView.DirectionalLightColor.rgb * PI;			// times PI because it is divided by PI on CPU (=luminance) and we want illuminance here. 
 		const float3 WaterDiffuseIndirectIlluminance = ShadingModelContext.WaterDiffuseIndirectLuminance * PI;	// DiffuseIndirectLighting is luminance. So we need to multiply by PI to get illuminance.
@@ -1008,7 +1011,8 @@ void Main(
 			MaterialParameters, PixelMaterialInputs, ResolvedView,
 			Shadow, GBuffer.Specular, NullDistortionParams,
 			SunIlluminance, WaterDiffuseIndirectIlluminance, EnvBrdf,
-			CameraIsUnderWater, ShadingModelContext.WaterVisibility, EyeIndex);
+			CameraIsUnderWater, ShadingModelContext.WaterVisibility, EyeIndex,
+			bSeparateMainDirLightLuminance, SeparatedWaterMainDirLightLuminance);
 
 		// Accumulate luminance and occlude the background according to transmittance to view and mean transmittance to lights.
 		MobileLightAccumulator_Add(MobileLightAccumulator, WaterLighting.Luminance, 0.0f, 1.0f);
diff --git a/Engine/Shaders/Private/SingleLayerWaterComposite.usf b/Engine/Shaders/Private/SingleLayerWaterComposite.usf
--- a/Engine/Shaders/Private/SingleLayerWaterComposite.usf
+++ b/Engine/Shaders/Private/SingleLayerWaterComposite.usf
@@ -29,7 +29,10 @@ SamplerState ScreenSpaceReflectionsSampler;
 Texture2D SceneNoWaterDepthTexture;
 SamplerState SceneNoWaterDepthSampler;
 
+Texture2D<float4> SeparatedMainDirLightTexture;
+
 float4 SceneNoWaterMinMaxUV;
+float UseSeparatedMainDirLightTexture;
 
 struct SingleLayerWaterCompositeOutput
 {
@@ -39,6 +42,8 @@ struct SingleLayerWaterCompositeOutput
 
 SingleLayerWaterCompositeOutput SingleLayerWaterComposite(float2 BufferUV, float2 ScreenPosition, float4 SvPosition)
 {
+	const uint2 PixelPos = uint2(SvPosition.xy);
+
 	// No AO or DFAO
 
 	const float3 OneThird3 = float3(1.0f / 3.0f, 1.0f / 3.0f, 1.0f / 3.0f);
@@ -195,6 +200,12 @@ SingleLayerWaterCompositeOutput SingleLayerWaterComposite(float2 BufferUV, float
 			}
 
 			Output.LuminanceTransmittance = float4(Reflection, dot(OneThird3, TransmittanceToScene));
+
+			if (UseSeparatedMainDirLightTexture)
+			{
+				// Also compose the separated main direction light luminance (separated to be able to combine it with DistanceFieldShadow)
+				Output.LuminanceTransmittance.rgb += SeparatedMainDirLightTexture[PixelPos].rgb;
+			}
 		}
 		else
 		{
diff --git a/Engine/Shaders/Private/SingleLayerWaterShading.ush b/Engine/Shaders/Private/SingleLayerWaterShading.ush
--- a/Engine/Shaders/Private/SingleLayerWaterShading.ush
+++ b/Engine/Shaders/Private/SingleLayerWaterShading.ush
@@ -36,7 +36,8 @@ WaterVolumeLightingOutput EvaluateWaterVolumeLighting(
 #endif // SIMPLE_SINGLE_LAYER_WATER
 	float Specular, const float4 DistortionParams,
 	float3 SunIlluminance, float3 AmbiantIlluminance, float3 EnvBrdf,
-	bool CameraIsUnderWater, float WaterVisibility, uint EyeIndex)
+	bool CameraIsUnderWater, float WaterVisibility, uint EyeIndex,
+	bool bSeparateMainDirLight, inout float3 SeparatedWaterMainDirLightScatteredLuminance)
 {
 	WaterVolumeLightingOutput Output;
 	Output.Luminance = 0.0f;
@@ -159,63 +160,27 @@ WaterVolumeLightingOutput EvaluateWaterVolumeLighting(
 		float3 SunScattLuminance = DirLightPhaseValue * SunIlluminance;
 		float3 AmbScattLuminance = IsotropicPhase()   * AmbiantIlluminance;
 
-#define VOLUMETRICSHADOW	0
-#if !VOLUMETRICSHADOW || SIMPLE_SINGLE_LAYER_WATER
+		const float MainDirLightFactor			= bSeparateMainDirLight ? 0.0f : 1.0f;
+		const float SeparatedMainDirLightFactor = bSeparateMainDirLight ? 1.0f : 0.0f;
 
 		const float3 OpticalDepth = ExtinctionCoeff * BehindWaterDeltaDepth;
 		float3 Transmittance = exp(-OpticalDepth);
-		float3 IncomingLuminance = (AmbScattLuminance + SunScattLuminance * DirectionalLightShadow);
+		float3 IncomingLuminance = (AmbScattLuminance + SunScattLuminance * DirectionalLightShadow * MainDirLightFactor);
 		// The SafeScatteringAmount should be in 0~1, add saturate to fix an artifact caused by half precision on mobile.
 		float3 SafeScatteringAmount = saturate(ScatteringCoeff * (1.0f - Transmittance) / ExtinctionCoeffSafe);
 		float3 ScatteredLuminance = IncomingLuminance * SafeScatteringAmount;
 
-#else
-		// TODO Make the volumetric shadow part work again
-		float3 Transmittance = 1.0f;
-		float3 ScatteredLuminance = 0.0f;
-		const float RayMarchMaxDistance = min(BehindWaterDeltaDepth, 200.0f);  // 20 meters
-		const float RayMarchStepSize = RayMarchMaxDistance / 10.0f; // Less samples wil lresult in a bit brighter look due to TransmittanceToLightThroughWater being 1 on a longer first sample. Would need it part of analiytical integration
-		const float ShadowDither = RayMarchStepSize * GBufferDither;
-		for (float s = 0.0f; s < RayMarchMaxDistance; s += RayMarchStepSize)
-		{
-			// Only jitter shadow map sampling to not lose energy on first sample
-			float Shadow = ComputeDirectionalLightDynamicShadowing(MaterialParameters.WorldPosition_CamRelative - (s + ShadowDither)*MaterialParameters.CameraVector, GBuffer.Depth);
-
-			// TODO: Sample virtual shadow map when present; something like the following:
-			//if (ForwardLightData.DirectionalLightVSM != INDEX_NONE)
-			//{
-			//	FVirtualShadowMapSampleResult VirtualShadowMapSample = SampleVirtualShadowMapTranslatedWorld(ForwardLightData.DirectionalLightVSM, TranslatedWorldPosition);
-			//	Shadow *= VirtualShadowMapSample.ShadowFactor;
-			//}
-
-			float3 WP = LWCHackToFloat(MaterialParameters.AbsoluteWorldPosition) - s * MaterialParameters.CameraVector;
-			float WaterHeightAboveSample = max(0.0, LWCHackToFloat(MaterialParameters.AbsoluteWorldPosition).z - WP.z);
-			float3 TransmittanceToLightThroughWater = 1.0;												// no self shadow, same energy as above analytical solution
-			//float3 TransmittanceToLightThroughWater = exp(-ExtinctionCoeff * WaterHeightAboveSample);	// self shadow as transmittance to water level, close to reference, depends a bit on sample count due to first sample being critical for dense medium
-
-			float3 SampleTransmittance = exp(-ExtinctionCoeff * RayMarchStepSize); // Constant
-			float3 SS = (ScatteringCoeff * TransmittanceToLightThroughWater * (SunScattLuminance * Shadow + AmbScattLuminance));
-			ScatteredLuminance += Transmittance * (SS - SS * SampleTransmittance) / ExtinctionCoeffSafe;
-			Transmittance *= SampleTransmittance;
-		}
-
-		// The rest of the medium
-		const float3 OpticalDepth2 = ExtinctionCoeff * max(0.0, BehindWaterDeltaDepth - RayMarchMaxDistance);
-		if (any(OpticalDepth2 > 0.0f))
-		{
-			float3 Transmittance2 = exp(-OpticalDepth2);
-			float3 ScatteredLuminance2 = ScatteringCoeff * (SunScattLuminance + AmbScattLuminance);
-			ScatteredLuminance += Transmittance * (ScatteredLuminance2 - ScatteredLuminance2 * Transmittance2) / ExtinctionCoeffSafe;
-			Transmittance *= Transmittance2;
-		}
-#endif
+		// Compute separated direcitonal light scattering if requested
+		float3 SeparatedIncomingLuminance = SunScattLuminance * DirectionalLightShadow * SeparatedMainDirLightFactor;
+		// Apply Fresnel and water visibility to separated main dir light luminance.
+		SeparatedWaterMainDirLightScatteredLuminance += SeparatedIncomingLuminance * SafeScatteringAmount * (CameraIsUnderWater ? 1.0 : (1.0 - EnvBrdf)) * WaterVisibility;
 
 		// Apply Fresnel effect to out-scattering towards the view
 		ScatteredLuminance *= CameraIsUnderWater ? 1.0 : (1.0 - EnvBrdf);	// Under water is less visible due to Fresnel effect
 		Transmittance *= CameraIsUnderWater ? (1.0 - EnvBrdf) : 1.0;		// Above	"		"		"		"		"
 
 		// Add single in-scattering apply colored transmittance to scene color
-		Output.Luminance = WaterVisibility * (ScatteredLuminance + Transmittance * (BehindWaterSceneLuminance* ColorScaleBehindWater));
+		Output.Luminance = WaterVisibility * (ScatteredLuminance + Transmittance * (BehindWaterSceneLuminance * ColorScaleBehindWater));
 		Output.WaterToSceneTransmittance = Transmittance;
 		Output.WaterToSceneToLightTransmittance = Transmittance * MeanTransmittanceToLightSources;
 	}
diff --git a/Engine/Source/Runtime/Engine/Private/Materials/HLSLMaterialTranslator.cpp b/Engine/Source/Runtime/Engine/Private/Materials/HLSLMaterialTranslator.cpp
--- a/Engine/Source/Runtime/Engine/Private/Materials/HLSLMaterialTranslator.cpp
+++ b/Engine/Source/Runtime/Engine/Private/Materials/HLSLMaterialTranslator.cpp
@@ -1824,12 +1824,17 @@ void FHLSLMaterialTranslator::GetMaterialEnvironment(EShaderPlatform InPlatform,
 
 		// This is to have platforms use the simple single layer water shading similar to mobile: no dynamic lights, only sun and sky, no distortion, no colored transmittance on background, no custom depth read.
 		const bool bSingleLayerWaterUsesSimpleShading = FDataDrivenShaderPlatformInfo::GetWaterUsesSimpleForwardShading(InPlatform) && IsForwardShadingEnabled(InPlatform);
-
 		if (ShadingModels.HasShadingModel(MSM_SingleLayerWater) && bSingleLayerWaterUsesSimpleShading)
 		{
 			OutEnvironment.SetDefine(TEXT("SINGLE_LAYER_WATER_SIMPLE_FORWARD"), TEXT("1"));
 		}
 
+		const bool bIsWaterDistanceFieldShadowEnabled = IsWaterDistanceFieldShadowEnabled(InPlatform);
+		if (ShadingModels.HasShadingModel(MSM_SingleLayerWater) && bIsWaterDistanceFieldShadowEnabled)
+		{
+			OutEnvironment.SetDefine(TEXT("SINGLE_LAYER_WATER_DF_SHADOW_ENABLED"), TEXT("1"));
+		}
+
 		if (NumSetMaterials == 1)
 		{
 			OutEnvironment.SetDefine(TEXT("MATERIAL_SINGLE_SHADINGMODEL"), TEXT("1"));
diff --git a/Engine/Source/Runtime/Engine/Private/ShaderCompiler/ShaderGenerationUtil.cpp b/Engine/Source/Runtime/Engine/Private/ShaderCompiler/ShaderGenerationUtil.cpp
--- a/Engine/Source/Runtime/Engine/Private/ShaderCompiler/ShaderGenerationUtil.cpp
+++ b/Engine/Source/Runtime/Engine/Private/ShaderCompiler/ShaderGenerationUtil.cpp
@@ -158,6 +158,8 @@ void FShaderCompileUtilities::ApplyFetchEnvironment(FShaderMaterialPropertyDefin
 	FETCH_COMPILE_BOOL(MATERIAL_SHADINGMODEL_SINGLELAYERWATER);
 	FETCH_COMPILE_BOOL(MATERIAL_SHADINGMODEL_THIN_TRANSLUCENT);
 
+	FETCH_COMPILE_BOOL(SINGLE_LAYER_WATER_DF_SHADOW_ENABLED);
+
 	FETCH_COMPILE_BOOL(MATERIAL_FULLY_ROUGH);
 
 	FETCH_COMPILE_BOOL(USES_EMISSIVE_COLOR);
@@ -464,6 +466,8 @@ static FString GetSlotTextName(EGBufferSlot Slot)
 		return TEXT("SubsurfaceProfileX");
 	case GBS_IrisNormal:
 		return TEXT("IrisNormal");
+	case GBS_SeparatedMainDirLight:
+		return TEXT("SeparatedMainDirLight");
 	default:
 		break;
 	};
@@ -1675,6 +1679,7 @@ static void SetSlotsForShadingModelType(bool Slots[], EMaterialShadingModel Shad
 		break;
 	case MSM_SingleLayerWater:
 		SetSharedGBufferSlots(Slots);
+		Slots[GBS_SeparatedMainDirLight] = true;
 		break;
 	case MSM_ThinTranslucent:
 		// thin translucent doesn't write to the GBuffer
@@ -1786,6 +1791,10 @@ static void DetermineUsedMaterialSlots(
 	{
 		// single layer water uses standard slots
 		SetStandardGBufferSlots(Slots, bWriteEmissive, bHasTangent, bHasVelocity, bHasStaticLighting, bIsStrataMaterial);
+		if (Mat.SINGLE_LAYER_WATER_DF_SHADOW_ENABLED)
+		{
+			Slots[GBS_SeparatedMainDirLight] = true;
+		}
 	}
 
 	// doesn't write to GBuffer
@@ -2127,6 +2136,7 @@ FGBufferParams FShaderCompileUtilities::FetchGBufferParamsPipeline(EShaderPlatfo
 
 FGBufferParams FShaderCompileUtilities::FetchGBufferParamsRuntime(EShaderPlatform Platform)
 {
+	// This code should match TBasePassPS
 
 	FGBufferParams Ret = {};
 	static const auto CVar = IConsoleManager::Get().FindTConsoleVariableDataInt(TEXT("r.AllowStaticLighting"));
@@ -2140,5 +2150,8 @@ FGBufferParams FShaderCompileUtilities::FetchGBufferParamsRuntime(EShaderPlatfor
 	static const auto CVarFormat = IConsoleManager::Get().FindTConsoleVariableDataInt(TEXT("r.GBufferFormat"));
 	Ret.LegacyFormatIndex = CVarFormat->GetValueOnAnyThread();
 
+	// This should match with SINGLE_LAYER_WATER_SEPARATED_DIR_LIGHT
+	Ret.bHasSingleLayerWaterSeparatedMainLight = IsWaterDistanceFieldShadowEnabled(Platform);
+
 	return Ret;
 }
diff --git a/Engine/Source/Runtime/Engine/Public/SceneView.h b/Engine/Source/Runtime/Engine/Public/SceneView.h
--- a/Engine/Source/Runtime/Engine/Public/SceneView.h
+++ b/Engine/Source/Runtime/Engine/Public/SceneView.h
@@ -652,6 +652,7 @@ enum ETranslucencyVolumeCascade
 	VIEW_UNIFORM_BUFFER_MEMBER(FVector4f, ScreenToViewSpace) \
 	VIEW_UNIFORM_BUFFER_MEMBER(FVector2f, BufferToSceneTextureScale) \
 	VIEW_UNIFORM_BUFFER_MEMBER(int32, NumSceneColorMSAASamples) \
+	VIEW_UNIFORM_BUFFER_MEMBER(float, SeparateWaterMainDirLightLuminance) \
 	VIEW_UNIFORM_BUFFER_MEMBER(float, PreExposure) \
 	VIEW_UNIFORM_BUFFER_MEMBER(float, OneOverPreExposure) \
 	VIEW_UNIFORM_BUFFER_MEMBER_EX(FVector4f, DiffuseOverrideParameter, EShaderPrecisionModifier::Half) \
diff --git a/Engine/Source/Runtime/RenderCore/Private/GBufferInfo.cpp b/Engine/Source/Runtime/RenderCore/Private/GBufferInfo.cpp
--- a/Engine/Source/Runtime/RenderCore/Private/GBufferInfo.cpp
+++ b/Engine/Source/Runtime/RenderCore/Private/GBufferInfo.cpp
@@ -236,6 +236,7 @@ FGBufferInfo RENDERCORE_API FetchLegacyGBufferInfo(const FGBufferParams& Params)
 	int32 TargetGBufferE = -1;
 	int32 TargetGBufferF = -1;
 	int32 TargetVelocity = -1;
+	int32 TargetSeparatedMainDirLight = -1;
 
 	// Strata ouputs material data through UAV. Only SceneColor, PrecalcShadow & Velocity data are still emitted through RenderTargets
 	const bool bStrata = RenderCore_IsStrataEnabled();
@@ -254,6 +255,10 @@ FGBufferInfo RENDERCORE_API FetchLegacyGBufferInfo(const FGBufferParams& Params)
 		{
 			TargetGBufferE = Info.NumTargets++;
 		}
+		if (Params.bHasSingleLayerWaterSeparatedMainLight)
+		{
+			TargetSeparatedMainDirLight = Info.NumTargets++;
+		}
 
 		// this value isn't correct, becuase it doesn't resepect the scene color format cvar, but it's ignored anyways
 		// so it's ok for now
@@ -287,6 +292,14 @@ FGBufferInfo RENDERCORE_API FetchLegacyGBufferInfo(const FGBufferParams& Params)
 			Info.Slots[GBS_PrecomputedShadowFactor].Packing[3] = FGBufferPacking(TargetGBufferE, 3, 3);
 		}
 
+		// Special water output
+		if (Params.bHasSingleLayerWaterSeparatedMainLight)
+		{
+			Info.Slots[GBS_SeparatedMainDirLight] = FGBufferItem(GBS_SeparatedMainDirLight, GBC_Raw_Float_11_11_10, GBCH_Both);
+			Info.Slots[GBS_SeparatedMainDirLight].Packing[0] = FGBufferPacking(TargetSeparatedMainDirLight, 0, 0);
+			Info.Slots[GBS_SeparatedMainDirLight].Packing[1] = FGBufferPacking(TargetSeparatedMainDirLight, 1, 1);
+			Info.Slots[GBS_SeparatedMainDirLight].Packing[2] = FGBufferPacking(TargetSeparatedMainDirLight, 2, 2);
+		}
 
 		return Info;
 	}
@@ -330,15 +343,18 @@ FGBufferInfo RENDERCORE_API FetchLegacyGBufferInfo(const FGBufferParams& Params)
 	const bool bLegacyAlbedoSrgb = true;
 	Info.Targets[3].Init(DiffuseAndSpecularGBufferFormat,  TEXT("GBufferC"), bLegacyAlbedoSrgb && !bHighPrecisionGBuffers,  true,  true,  true);
 
+	// This code should match TBasePassPS
 	if (Params.bHasVelocity == 0 && Params.bHasTangent == 0)
 	{
 		TargetGBufferD = 4;
 		Info.Targets[4].Init(GBT_Unorm_8_8_8_8,  TEXT("GBufferD"), false,  true,  true,  true);
+		TargetSeparatedMainDirLight = 5;
 
 		if (Params.bHasPrecShadowFactor)
 		{
 			TargetGBufferE = 5;
-			Info.Targets[5].Init(GBT_Unorm_8_8_8_8,  TEXT("GBufferE"), false,  true,  true,  true);
+			Info.Targets[5].Init(GBT_Unorm_8_8_8_8, TEXT("GBufferE"), false, true, true, true);
+			TargetSeparatedMainDirLight = 6;
 		}
 	}
 	else if (Params.bHasVelocity)
@@ -348,12 +364,14 @@ FGBufferInfo RENDERCORE_API FetchLegacyGBufferInfo(const FGBufferParams& Params)
 
 		// note the false for use extra flags for velocity, not quite sure of all the ramifications, but this keeps it consistent with previous usage
 		Info.Targets[4].Init(Params.bUsesVelocityDepth ? GBT_Float_16_16_16_16 : GBT_Float_16_16,   TEXT("Velocity"), false,  true,  true, false);
-		Info.Targets[5].Init(GBT_Unorm_8_8_8_8, TEXT("GBufferD"), false,  true,  true,  true);
+		Info.Targets[5].Init(GBT_Unorm_8_8_8_8, TEXT("GBufferD"), false, true, true, true);
+		TargetSeparatedMainDirLight = 6;
 
 		if (Params.bHasPrecShadowFactor)
 		{
 			TargetGBufferE = 6;
-			Info.Targets[6].Init(GBT_Unorm_8_8_8_8,  TEXT("GBufferE"), false,  true,  true, false);
+			Info.Targets[6].Init(GBT_Unorm_8_8_8_8, TEXT("GBufferE"), false, true, true, false);
+			TargetSeparatedMainDirLight = 7;
 		}
 	}
 	else if (Params.bHasTangent)
@@ -361,11 +379,13 @@ FGBufferInfo RENDERCORE_API FetchLegacyGBufferInfo(const FGBufferParams& Params)
 		TargetGBufferF = 4;
 		TargetGBufferD = 5;
 		Info.Targets[4].Init(GBT_Unorm_8_8_8_8,  TEXT("GBufferF"), false,  true,  true,  true);
-		Info.Targets[5].Init(GBT_Unorm_8_8_8_8,  TEXT("GBufferD"), false,  true,  true,  true);
+		Info.Targets[5].Init(GBT_Unorm_8_8_8_8, TEXT("GBufferD"), false, true, true, true);
+		TargetSeparatedMainDirLight = 6;
 		if (Params.bHasPrecShadowFactor)
 		{
 			TargetGBufferE = 6;
-			Info.Targets[6].Init(GBT_Unorm_8_8_8_8,  TEXT("GBufferE"), false,  true,  true,  true);
+			Info.Targets[6].Init(GBT_Unorm_8_8_8_8, TEXT("GBufferE"), false, true, true, true);
+			TargetSeparatedMainDirLight = 7;
 		}
 	}
 	else
@@ -500,6 +520,15 @@ FGBufferInfo RENDERCORE_API FetchLegacyGBufferInfo(const FGBufferParams& Params)
 	Info.Slots[GBS_CustomData].Packing[2] = FGBufferPacking(TargetGBufferD, 2, 2);
 	Info.Slots[GBS_CustomData].Packing[3] = FGBufferPacking(TargetGBufferD, 3, 3);
 
+	// Special water output
+	if (Params.bHasSingleLayerWaterSeparatedMainLight)
+	{
+		Info.Slots[GBS_SeparatedMainDirLight] = FGBufferItem(GBS_SeparatedMainDirLight, GBC_Raw_Float_11_11_10, GBCH_Both);
+		Info.Slots[GBS_SeparatedMainDirLight].Packing[0] = FGBufferPacking(TargetSeparatedMainDirLight, 0, 0);
+		Info.Slots[GBS_SeparatedMainDirLight].Packing[1] = FGBufferPacking(TargetSeparatedMainDirLight, 1, 1);
+		Info.Slots[GBS_SeparatedMainDirLight].Packing[2] = FGBufferPacking(TargetSeparatedMainDirLight, 2, 2);
+	}
+
 	return Info;
 }
 
diff --git a/Engine/Source/Runtime/RenderCore/Private/Shader.cpp b/Engine/Source/Runtime/RenderCore/Private/Shader.cpp
--- a/Engine/Source/Runtime/RenderCore/Private/Shader.cpp
+++ b/Engine/Source/Runtime/RenderCore/Private/Shader.cpp
@@ -1772,6 +1772,11 @@ void ShaderMapAppendKeyString(EShaderPlatform Platform, FString& KeyString)
 		}
 	}
 
+	if (IsWaterDistanceFieldShadowEnabled(Platform))
+	{
+		KeyString += TEXT("_SLWDFS");
+	}
+
 	{
 		static const auto CVar = IConsoleManager::Get().FindTConsoleVariableDataInt(TEXT("r.SupportCloudShadowOnForwardLitTranslucent"));
 		if (CVar && CVar->GetValueOnAnyThread() > 0)
diff --git a/Engine/Source/Runtime/RenderCore/Public/GBufferInfo.h b/Engine/Source/Runtime/RenderCore/Public/GBufferInfo.h
--- a/Engine/Source/Runtime/RenderCore/Public/GBufferInfo.h
+++ b/Engine/Source/Runtime/RenderCore/Public/GBufferInfo.h
@@ -36,6 +36,7 @@ enum EGBufferSlot
 	GBS_Cloth, // R8
 	GBS_SubsurfaceProfileX, // R8
 	GBS_IrisNormal, // RG8
+	GBS_SeparatedMainDirLight, // RGB 11.11.10
 	GBS_Num
 };
 
@@ -271,6 +272,7 @@ struct FGBufferParams
 	bool bHasTangent = false;
 	bool bHasPrecShadowFactor = false;
 	bool bUsesVelocityDepth = false;
+	bool bHasSingleLayerWaterSeparatedMainLight = false;
 
 	bool operator == (const FGBufferParams& RHS) const
 	{
@@ -279,7 +281,8 @@ struct FGBufferParams
 			bHasVelocity == RHS.bHasVelocity &&
 			bHasTangent == RHS.bHasTangent &&
 			bHasPrecShadowFactor == RHS.bHasPrecShadowFactor &&
-			bUsesVelocityDepth == RHS.bUsesVelocityDepth;
+			bUsesVelocityDepth == RHS.bUsesVelocityDepth &&
+			bHasSingleLayerWaterSeparatedMainLight == RHS.bHasSingleLayerWaterSeparatedMainLight;
 	}
 
 	bool operator != (const FGBufferParams& RHS) const
diff --git a/Engine/Source/Runtime/RenderCore/Public/RenderUtils.h b/Engine/Source/Runtime/RenderCore/Public/RenderUtils.h
--- a/Engine/Source/Runtime/RenderCore/Public/RenderUtils.h
+++ b/Engine/Source/Runtime/RenderCore/Public/RenderUtils.h
@@ -627,6 +627,15 @@ inline bool IsUsingDistanceFields(const FStaticShaderPlatform Platform)
 	return !!(GDistanceFieldsPlatformMask & (1ull << Platform));
 }
 
+/** Returns if water should render distance field shadow a second time for the water surface. This is for a platofrm so can be used at cook time. */
+inline bool IsWaterDistanceFieldShadowEnabled(const FStaticShaderPlatform Platform)
+{
+	// Only deferred support such a feature. It is not possible to do that for water without a water depth pre-pass.
+	static const auto CVarWaterSingleLayerShaderSupportDistanceFieldShadow = IConsoleManager::Get().FindConsoleVariable(TEXT("r.Water.SingleLayer.ShadersSupportDistanceFieldShadow"));
+	const bool bWaterSingleLayerShaderSupportDistanceFieldShadow = CVarWaterSingleLayerShaderSupportDistanceFieldShadow && (CVarWaterSingleLayerShaderSupportDistanceFieldShadow->GetInt() > 0);
+	return !IsForwardShadingEnabled(Platform) && IsUsingDistanceFields(Platform) && bWaterSingleLayerShaderSupportDistanceFieldShadow;
+}
+
 inline bool UseGPUScene(const FStaticShaderPlatform Platform, const FStaticFeatureLevel FeatureLevel)
 {
 	if (FeatureLevel == ERHIFeatureLevel::ES3_1)
diff --git a/Engine/Source/Runtime/RenderCore/Public/ShaderMaterial.h b/Engine/Source/Runtime/RenderCore/Public/ShaderMaterial.h
--- a/Engine/Source/Runtime/RenderCore/Public/ShaderMaterial.h
+++ b/Engine/Source/Runtime/RenderCore/Public/ShaderMaterial.h
@@ -102,6 +102,7 @@ struct FShaderMaterialPropertyDefines
 	uint8 MATERIAL_SHADINGMODEL_CLOTH : 1;
 	uint8 MATERIAL_SHADINGMODEL_EYE : 1;
 	uint8 MATERIAL_SHADINGMODEL_SINGLELAYERWATER : 1;
+	uint8 SINGLE_LAYER_WATER_DF_SHADOW_ENABLED : 1;
 	uint8 MATERIAL_SHADINGMODEL_THIN_TRANSLUCENT : 1;
 
 	uint8 TRANSLUCENCY_LIGHTING_VOLUMETRIC_NONDIRECTIONAL : 1;
diff --git a/Engine/Source/Runtime/Renderer/Private/BasePassRendering.h b/Engine/Source/Runtime/Renderer/Private/BasePassRendering.h
--- a/Engine/Source/Runtime/Renderer/Private/BasePassRendering.h
+++ b/Engine/Source/Runtime/Renderer/Private/BasePassRendering.h
@@ -452,6 +452,43 @@ public:
 
 		OutEnvironment.SetDefine(TEXT("STRATA_INLINE_SHADING"), 1);
 
+		const bool IsSingleLayerWater = Parameters.MaterialParameters.ShadingModels.HasShadingModel(MSM_SingleLayerWater);
+		if (IsSingleLayerWater && IsWaterDistanceFieldShadowEnabled(Parameters.Platform))
+		{
+			// See FShaderCompileUtilities::FetchGBufferParamsRuntime for the details
+			const bool bOutputVelocity = FVelocityRendering::BasePassCanOutputVelocity(Parameters.Platform);
+			const bool bHasTangent = false;
+			static const auto CVar = IConsoleManager::Get().FindTConsoleVariableDataInt(TEXT("r.AllowStaticLighting"));
+			bool bHasPrecShadowFactor = (CVar ? (CVar->GetValueOnAnyThread() != 0) : 1);
+
+			uint32 TargetSeparatedMainDirLight = 5;
+			if (bOutputVelocity == false && bHasTangent == false)
+			{
+				TargetSeparatedMainDirLight = 5;
+				if (bHasPrecShadowFactor)
+				{
+					TargetSeparatedMainDirLight = 6;
+				}
+			}
+			else if (bOutputVelocity)
+			{
+				TargetSeparatedMainDirLight = 6;
+				if (bHasPrecShadowFactor)
+				{
+					TargetSeparatedMainDirLight = 7;
+				}
+			}
+			else if (bHasTangent)
+			{
+				TargetSeparatedMainDirLight = 6;
+				if (bHasPrecShadowFactor)
+				{
+					TargetSeparatedMainDirLight = 7;
+				}
+			}
+			OutEnvironment.SetRenderTargetOutputFormat(TargetSeparatedMainDirLight, PF_FloatR11G11B10);
+		}
+
 		TBasePassPixelShaderBaseType<LightMapPolicyType>::ModifyCompilationEnvironment(Parameters, OutEnvironment);
 	}
 	
diff --git a/Engine/Source/Runtime/Renderer/Private/DistanceFieldShadowing.cpp b/Engine/Source/Runtime/Renderer/Private/DistanceFieldShadowing.cpp
--- a/Engine/Source/Runtime/Renderer/Private/DistanceFieldShadowing.cpp
+++ b/Engine/Source/Runtime/Renderer/Private/DistanceFieldShadowing.cpp
@@ -825,6 +825,8 @@ void FProjectedShadowInfo::BeginRenderRayTracedDistanceFieldProjection(
 	const bool bHFShadowSupported = SupportsHeightFieldShadows(View.GetFeatureLevel(), View.GetShaderPlatform());
 	const FScene* Scene = (const FScene*)(View.Family->Scene);
 
+	DistanceFieldShadowViewGPUData& SDFShadowViewGPUData = CachedDistanceFieldShadowViewGPUData.FindOrAdd(&View);
+
 	if (bDFShadowSupported && View.Family->EngineShowFlags.RayTracedDistanceFieldShadows)
 	{
 		QUICK_SCOPE_CYCLE_COUNTER(STAT_BeginRenderRayTracedDistanceFieldShadows);
@@ -834,50 +836,54 @@ void FProjectedShadowInfo::BeginRenderRayTracedDistanceFieldProjection(
 		{
 			check(!Scene->DistanceFieldSceneData.HasPendingOperations());
 
-			int32 NumPlanes = 0;
-			const FPlane* PlaneData = NULL;
-			FVector4f ShadowBoundingSphere = FVector4f::Zero();
-			FVector PrePlaneTranslation = FVector::ZeroVector;
+			FDistanceFieldObjectBufferParameters ObjectBufferParameters = DistanceField::SetupObjectBufferParameters(Scene->DistanceFieldSceneData);
 
-			if (bDirectionalLight)
+			if (SDFShadowViewGPUData.SDFCulledObjectBufferParameters == nullptr)
 			{
-				NumPlanes = CascadeSettings.ShadowBoundsAccurate.Planes.Num();
-				PlaneData = CascadeSettings.ShadowBoundsAccurate.Planes.GetData();
-			}
-			else if (IsWholeScenePointLightShadow())
-			{
-				ShadowBoundingSphere = FVector4f(FVector3f(ShadowBounds.Center + View.ViewMatrices.GetPreViewTranslation()), ShadowBounds.W);
-			}
-			else
-			{
-				NumPlanes = CasterOuterFrustum.Planes.Num();
-				PlaneData = CasterOuterFrustum.Planes.GetData();
-				PrePlaneTranslation = PreShadowTranslation;
-			}
+				int32 NumPlanes = 0;
+				const FPlane* PlaneData = NULL;
+				FVector4f ShadowBoundingSphere = FVector4f::Zero();
+				FVector PrePlaneTranslation = FVector::ZeroVector;
 
-			const FMatrix WorldToShadowValue = FTranslationMatrix(PreShadowTranslation) * FMatrix(TranslatedWorldToClipInnerMatrix);
+				if (bDirectionalLight)
+				{
+					NumPlanes = CascadeSettings.ShadowBoundsAccurate.Planes.Num();
+					PlaneData = CascadeSettings.ShadowBoundsAccurate.Planes.GetData();
+				}
+				else if (IsWholeScenePointLightShadow())
+				{
+					ShadowBoundingSphere = FVector4f(FVector3f(ShadowBounds.Center + View.ViewMatrices.GetPreViewTranslation()), ShadowBounds.W);
+				}
+				else
+				{
+					NumPlanes = CasterOuterFrustum.Planes.Num();
+					PlaneData = CasterOuterFrustum.Planes.GetData();
+					PrePlaneTranslation = PreShadowTranslation;
+				}
 
-			FDistanceFieldObjectBufferParameters ObjectBufferParameters = DistanceField::SetupObjectBufferParameters(Scene->DistanceFieldSceneData);
+				const FMatrix WorldToShadowValue = FTranslationMatrix(PreShadowTranslation) * FMatrix(TranslatedWorldToClipInnerMatrix);
 
-			FLightTileIntersectionParameters LightTileIntersectionParameters;
-			FDistanceFieldCulledObjectBufferParameters CulledObjectBufferParameters;
 
-			CullDistanceFieldObjectsForLight(
-				GraphBuilder,
-				View,
-				LightSceneInfo->Proxy,
-				DFPT_SignedDistanceField,
-				WorldToShadowValue,
-				NumPlanes,
-				PlaneData,
-				PrePlaneTranslation,
-				ShadowBoundingSphere,
-				ShadowBounds.W,
-				true,
-				ObjectBufferParameters,
-				CulledObjectBufferParameters,
-				LightTileIntersectionParameters
+				SDFShadowViewGPUData.SDFCulledObjectBufferParameters = GraphBuilder.AllocObject<FDistanceFieldCulledObjectBufferParameters>();
+				SDFShadowViewGPUData.SDFLightTileIntersectionParameters = GraphBuilder.AllocObject<FLightTileIntersectionParameters>();
+
+				CullDistanceFieldObjectsForLight(
+					GraphBuilder,
+					View,
+					LightSceneInfo->Proxy,
+					DFPT_SignedDistanceField,
+					WorldToShadowValue,
+					NumPlanes,
+					PlaneData,
+					PrePlaneTranslation,
+					ShadowBoundingSphere,
+					ShadowBounds.W,
+					true,
+					ObjectBufferParameters,
+					*SDFShadowViewGPUData.SDFCulledObjectBufferParameters,
+					*SDFShadowViewGPUData.SDFLightTileIntersectionParameters
 				);
+			}
 
 			{
 				const FIntPoint BufferSize = GetBufferSizeForDFShadows();
@@ -886,7 +892,7 @@ void FProjectedShadowInfo::BeginRenderRayTracedDistanceFieldProjection(
 				RayTracedShadowsTexture = GraphBuilder.CreateTexture(Desc, TEXT("RayTracedShadows"));
 			}
 
-			RayTraceShadows(GraphBuilder, SceneTextures, RayTracedShadowsTexture, View, Scene->DistanceFieldSceneData, this, DFPT_SignedDistanceField, false, nullptr, ObjectBufferParameters, CulledObjectBufferParameters, LightTileIntersectionParameters);
+			RayTraceShadows(GraphBuilder, SceneTextures, RayTracedShadowsTexture, View, Scene->DistanceFieldSceneData, this, DFPT_SignedDistanceField, false, nullptr, ObjectBufferParameters, *SDFShadowViewGPUData.SDFCulledObjectBufferParameters, *SDFShadowViewGPUData.SDFLightTileIntersectionParameters);
 		}
 	}
 
@@ -901,34 +907,38 @@ void FProjectedShadowInfo::BeginRenderRayTracedDistanceFieldProjection(
 
 		check(!Scene->DistanceFieldSceneData.HasPendingHeightFieldOperations());
 
-		const int32 NumPlanes = CascadeSettings.ShadowBoundsAccurate.Planes.Num();
-		const FPlane* PlaneData = CascadeSettings.ShadowBoundsAccurate.Planes.GetData();
-		FVector PrePlaneTranslation = FVector::ZeroVector;
-		const FVector4f ShadowBoundingSphere = FVector4f::Zero();
-		const FMatrix WorldToShadowValue = FTranslationMatrix(PreShadowTranslation) * FMatrix(TranslatedWorldToClipInnerMatrix);
-
 		FDistanceFieldObjectBufferParameters ObjectBufferParameters = DistanceField::SetupObjectBufferParameters(Scene->DistanceFieldSceneData);
 
-		FLightTileIntersectionParameters LightTileIntersectionParameters;
-		FDistanceFieldCulledObjectBufferParameters CulledObjectBufferParameters;
+		if (SDFShadowViewGPUData.HeightFieldCulledObjectBufferParameters == nullptr)
+		{
+			const int32 NumPlanes = CascadeSettings.ShadowBoundsAccurate.Planes.Num();
+			const FPlane* PlaneData = CascadeSettings.ShadowBoundsAccurate.Planes.GetData();
+			FVector PrePlaneTranslation = FVector::ZeroVector;
+			const FVector4f ShadowBoundingSphere = FVector4f::Zero();
+			const FMatrix WorldToShadowValue = FTranslationMatrix(PreShadowTranslation) * FMatrix(TranslatedWorldToClipInnerMatrix);
 
-		CullDistanceFieldObjectsForLight(
-			GraphBuilder,
-			View,
-			LightSceneInfo->Proxy,
-			DFPT_HeightField,
-			WorldToShadowValue,
-			NumPlanes,
-			PlaneData,
-			PrePlaneTranslation,
-			ShadowBoundingSphere,
-			ShadowBounds.W,
-			true,
-			ObjectBufferParameters,
-			CulledObjectBufferParameters,
-			LightTileIntersectionParameters
+			SDFShadowViewGPUData.HeightFieldCulledObjectBufferParameters = GraphBuilder.AllocObject<FDistanceFieldCulledObjectBufferParameters>();
+			SDFShadowViewGPUData.HeightFieldLightTileIntersectionParameters = GraphBuilder.AllocObject<FLightTileIntersectionParameters>();
+
+			CullDistanceFieldObjectsForLight(
+				GraphBuilder,
+				View,
+				LightSceneInfo->Proxy,
+				DFPT_HeightField,
+				WorldToShadowValue,
+				NumPlanes,
+				PlaneData,
+				PrePlaneTranslation,
+				ShadowBoundingSphere,
+				ShadowBounds.W,
+				true,
+				ObjectBufferParameters,
+				*SDFShadowViewGPUData.HeightFieldCulledObjectBufferParameters,
+				*SDFShadowViewGPUData.HeightFieldLightTileIntersectionParameters
 			);
 
+		}
+
 		const bool bHasPrevOutput = !!RayTracedShadowsTexture;
 
 		FRDGTextureRef PrevOutputTexture = nullptr;
@@ -947,7 +957,7 @@ void FProjectedShadowInfo::BeginRenderRayTracedDistanceFieldProjection(
 			RayTracedShadowsTexture = GraphBuilder.CreateTexture(Desc, TEXT("RayTracedShadows"));
 		}
 
-		RayTraceShadows(GraphBuilder, SceneTextures, RayTracedShadowsTexture, View, Scene->DistanceFieldSceneData, this, DFPT_HeightField, bHasPrevOutput, PrevOutputTexture, ObjectBufferParameters, CulledObjectBufferParameters, LightTileIntersectionParameters);
+		RayTraceShadows(GraphBuilder, SceneTextures, RayTracedShadowsTexture, View, Scene->DistanceFieldSceneData, this, DFPT_HeightField, bHasPrevOutput, PrevOutputTexture, ObjectBufferParameters, *SDFShadowViewGPUData.HeightFieldCulledObjectBufferParameters, *SDFShadowViewGPUData.HeightFieldLightTileIntersectionParameters);
 	}
 }
 
@@ -962,7 +972,8 @@ void FProjectedShadowInfo::RenderRayTracedDistanceFieldProjection(
 	FRDGTextureRef ScreenShadowMaskTexture,
 	const FViewInfo& View,
 	FIntRect ScissorRect,
-	bool bProjectingForForwardShading)
+	bool bProjectingForForwardShading,
+	bool bForceRGBModulation)
 {
 	check(ScissorRect.Area() > 0);
 
@@ -1014,7 +1025,7 @@ void FProjectedShadowInfo::RenderRayTracedDistanceFieldProjection(
 			RDG_EVENT_NAME("Upsample"),
 			PassParameters,
 			ERDGPassFlags::Raster,
-			[this, &View, PixelShader, ScissorRect, bProjectingForForwardShading, PassParameters](FRHICommandList& RHICmdList)
+			[this, &View, PixelShader, ScissorRect, bProjectingForForwardShading, PassParameters, bForceRGBModulation](FRHICommandList& RHICmdList)
 		{
 			FGraphicsPipelineStateInitializer GraphicsPSOInit;
 			RHICmdList.ApplyCachedRenderTargets(GraphicsPSOInit);
@@ -1025,7 +1036,15 @@ void FProjectedShadowInfo::RenderRayTracedDistanceFieldProjection(
 			GraphicsPSOInit.RasterizerState = TStaticRasterizerState<FM_Solid, CM_None>::GetRHI();
 			GraphicsPSOInit.DepthStencilState = TStaticDepthStencilState<false, CF_Always>::GetRHI();
 
-			GraphicsPSOInit.BlendState = GetBlendStateForProjection(bProjectingForForwardShading, false);
+			if (bForceRGBModulation)
+			{
+				// This has the shadow contribution modulate all the channels, e.g. used for water rendering to apply distance field shadow on the main light RGB luminance for the updated depth buffer with water in it.
+				GraphicsPSOInit.BlendState = TStaticBlendState<CW_RGBA, BO_Add, BF_Zero, BF_SourceColor, BO_Add, BF_Zero, BF_One>::GetRHI();;
+			}
+			else
+			{
+				GraphicsPSOInit.BlendState = GetBlendStateForProjection(bProjectingForForwardShading, false);
+			}
 
 			TShaderMapRef<FPostProcessVS> VertexShader(View.ShaderMap);
 			GraphicsPSOInit.BoundShaderState.VertexDeclarationRHI = GFilterVertexDeclaration.VertexDeclarationRHI;
diff --git a/Engine/Source/Runtime/Renderer/Private/SceneRendering.cpp b/Engine/Source/Runtime/Renderer/Private/SceneRendering.cpp
--- a/Engine/Source/Runtime/Renderer/Private/SceneRendering.cpp
+++ b/Engine/Source/Runtime/Renderer/Private/SceneRendering.cpp
@@ -1300,6 +1300,9 @@ void FViewInfo::SetupUniformBufferParameters(
 		Scene = Family->Scene->GetRenderScene();
 	}
 
+	ERHIFeatureLevel::Type RHIFeatureLevel = Scene == nullptr ? GMaxRHIFeatureLevel : Scene->GetFeatureLevel();
+	EShaderPlatform ShaderPlatform = GShaderPlatformForFeatureLevel[RHIFeatureLevel];
+
 	const FVector DefaultSunDirection(0.0f, 0.0f, 1.0f); // Up vector so that the AtmosphericLightVector node always output a valid direction.
 	auto ClearAtmosphereLightData = [&](uint32 Index)
 	{
@@ -1349,6 +1352,11 @@ void FViewInfo::SetupUniformBufferParameters(
 
 	ViewUniformShaderParameters.BufferToSceneTextureScale = FVector2f(1.0f, 1.0f);
 
+	{
+		extern bool IsWaterDistanceFieldShadowEnabled_Runtime(const FStaticShaderPlatform Platform);
+		ViewUniformShaderParameters.SeparateWaterMainDirLightLuminance = IsWaterDistanceFieldShadowEnabled_Runtime(ShaderPlatform) ? 1.0f : 0.0f;
+	}
+
 	FRHITexture* TransmittanceLutTextureFound = nullptr;
 	FRHITexture* SkyViewLutTextureFound = nullptr;
 	FRHITexture* CameraAerialPerspectiveVolumeFound = nullptr;
@@ -1678,9 +1686,6 @@ void FViewInfo::SetupUniformBufferParameters(
 
 	ViewUniformShaderParameters.DecalDepthBias = CVarDecalDepthBias.GetValueOnRenderThread();
 
-	ERHIFeatureLevel::Type RHIFeatureLevel = Scene == nullptr ? GMaxRHIFeatureLevel : Scene->GetFeatureLevel();
-	EShaderPlatform ShaderPlatform = GShaderPlatformForFeatureLevel[RHIFeatureLevel];
-
 	ViewUniformShaderParameters.IndirectLightingColorScale = FVector3f(FinalPostProcessSettings.IndirectLightingColor.R * FinalPostProcessSettings.IndirectLightingIntensity,
 		FinalPostProcessSettings.IndirectLightingColor.G * FinalPostProcessSettings.IndirectLightingIntensity,
 		FinalPostProcessSettings.IndirectLightingColor.B * FinalPostProcessSettings.IndirectLightingIntensity);
diff --git a/Engine/Source/Runtime/Renderer/Private/ShadowRendering.h b/Engine/Source/Runtime/Renderer/Private/ShadowRendering.h
--- a/Engine/Source/Runtime/Renderer/Private/ShadowRendering.h
+++ b/Engine/Source/Runtime/Renderer/Private/ShadowRendering.h
@@ -40,6 +40,8 @@ class FSceneRenderer;
 class FViewInfo;
 class FVirtualShadowMapArrayCacheManager;
 class FVirtualShadowMapCacheEntry;
+class FLightTileIntersectionParameters;
+class FDistanceFieldCulledObjectBufferParameters;
 
 /** Renders a cone with a spherical cap, used for rendering spot lights in deferred passes. */
 extern void DrawStencilingCone(const FMatrix& ConeToWorld, float ConeAngle, float SphereRadius, const FVector& PreViewTranslation);
@@ -577,7 +579,8 @@ public:
 		FRDGTextureRef ScreenShadowMaskTexture,
 		const FViewInfo& View,
 		FIntRect ScissorRect,
-		bool bProjectingForForwardShading);
+		bool bProjectingForForwardShading,
+		bool bForceRGBModulation = false);
 
 	/** Render one pass point light shadow projections. */
 	void RenderOnePassPointLightProjection(
@@ -726,6 +729,11 @@ public:
 		return ShadowDepthPass.BuildRenderingCommands(GraphBuilder, GPUScene, InstanceCullingDrawParams);
 	}
 
+	void ResetRayTracedDistanceFieldShadow()
+	{
+		RayTracedShadowsTexture = nullptr;
+	}
+
 private:
 	// 0 if Setup...() wasn't called yet
 	FLightSceneInfo* LightSceneInfo;
@@ -779,6 +787,16 @@ private:
 	/** Ray traced DF shadow intermediate output. Populated by BeginRenderRayTracedDistanceFieldProjection and consumed by RenderRayTracedDistanceFieldProjection. */
 	FRDGTextureRef RayTracedShadowsTexture;
 
+	/**  Cached light tile intersection parameters in case we needed to trace a distance field multiple times for a view but for different depth buffer*/
+	struct DistanceFieldShadowViewGPUData 
+	{
+		FLightTileIntersectionParameters*			SDFLightTileIntersectionParameters = nullptr;
+		FDistanceFieldCulledObjectBufferParameters*	SDFCulledObjectBufferParameters = nullptr;
+		FLightTileIntersectionParameters*			HeightFieldLightTileIntersectionParameters = nullptr;
+		FDistanceFieldCulledObjectBufferParameters*	HeightFieldCulledObjectBufferParameters = nullptr;
+	};
+	TMap<const FViewInfo*, DistanceFieldShadowViewGPUData>	CachedDistanceFieldShadowViewGPUData;
+
 	void CopyCachedShadowMap(
 		FRDGBuilder& GraphBuilder,
 		const FViewInfo& View,
diff --git a/Engine/Source/Runtime/Renderer/Private/SingleLayerWaterRendering.cpp b/Engine/Source/Runtime/Renderer/Private/SingleLayerWaterRendering.cpp
--- a/Engine/Source/Runtime/Renderer/Private/SingleLayerWaterRendering.cpp
+++ b/Engine/Source/Runtime/Renderer/Private/SingleLayerWaterRendering.cpp
@@ -64,6 +64,17 @@ static TAutoConsoleVariable<int32> CVarWaterSingleLayerLumenReflections(
 	TEXT("Enable Lumen reflections for the single water rendering system."),
 	ECVF_RenderThreadSafe | ECVF_Scalability);
 
+static TAutoConsoleVariable<int32> CVarWaterSingleLayerShadersSupportDistanceFieldShadow(
+	TEXT("r.Water.SingleLayer.ShadersSupportDistanceFieldShadow"),
+	1,
+	TEXT("Whether or not the single layer water material shaders are compiled with support for distance field shadow, i.e. output main directional light luminance in a separate render target. This is preconditioned on using deferred shading and having distance field support enabled in the project."),
+	ECVF_ReadOnly | ECVF_RenderThreadSafe);
+
+static TAutoConsoleVariable<int32> CVarWaterSingleLayerDistanceFieldShadow(
+	TEXT("r.Water.SingleLayer.DistanceFieldShadow"), 1,
+	TEXT("When using deferred, distance field shadow tracing is supported on single layer water. This cvar can be used to toggle it on/off at runtime."),
+	ECVF_RenderThreadSafe | ECVF_Scalability);
+
 static TAutoConsoleVariable<int32> CVarWaterSingleLayerRTR(
 	TEXT("r.Water.SingleLayer.RTR"), 1,
 	TEXT("Enable RTR for the single water renderring system."),
@@ -122,6 +133,11 @@ bool UseSingleLayerWaterIndirectDraw(EShaderPlatform ShaderPlatform)
 		&& FDataDrivenShaderPlatformInfo::GetSupportsWaterIndirectDraw(ShaderPlatform);
 }
 
+bool IsWaterDistanceFieldShadowEnabled_Runtime(const FStaticShaderPlatform Platform)
+{
+	return IsWaterDistanceFieldShadowEnabled(Platform) && CVarWaterSingleLayerDistanceFieldShadow.GetValueOnRenderThread() > 0;
+}
+
 BEGIN_SHADER_PARAMETER_STRUCT(FSingleLayerWaterCommonShaderParameters, )
 	SHADER_PARAMETER_RDG_TEXTURE(Texture2D, ScreenSpaceReflectionsTexture)
 	SHADER_PARAMETER_SAMPLER(SamplerState, ScreenSpaceReflectionsSampler)
@@ -129,7 +145,9 @@ BEGIN_SHADER_PARAMETER_STRUCT(FSingleLayerWaterCommonShaderParameters, )
 	SHADER_PARAMETER_SAMPLER(SamplerState, PreIntegratedGFSampler)
 	SHADER_PARAMETER_RDG_TEXTURE(Texture2D, SceneNoWaterDepthTexture)
 	SHADER_PARAMETER_SAMPLER(SamplerState, SceneNoWaterDepthSampler)
+	SHADER_PARAMETER_RDG_TEXTURE(Texture2D, SeparatedMainDirLightTexture)
 	SHADER_PARAMETER(FVector4f, SceneNoWaterMinMaxUV)
+	SHADER_PARAMETER(float, UseSeparatedMainDirLightTexture)
 	SHADER_PARAMETER_STRUCT_INCLUDE(FSceneTextureParameters, SceneTextures)	// Water scene texture
 	SHADER_PARAMETER_STRUCT_REF(FViewUniformShaderParameters, ViewUniformBuffer)
 	SHADER_PARAMETER_STRUCT_REF(FReflectionCaptureShaderData, ReflectionCaptureData)
@@ -280,6 +298,9 @@ static FSceneWithoutWaterTextures AddCopySceneWithoutWaterPass(
 	const FRDGTextureDesc DepthDesc(FRDGTextureDesc::Create2D(RefractionResolution, ViewFamily.EngineShowFlags.SingleLayerWaterRefractionFullPrecision ? PF_R32_FLOAT : PF_R16F, SceneDepthDesc.ClearValue, TexCreate_ShaderResource | TexCreate_RenderTargetable));
 	FRDGTextureRef SceneDepthWithoutSingleLayerWaterTexture = GraphBuilder.CreateTexture(DepthDesc, TEXT("SLW.SceneDepthWithout"));
 
+	const FRDGTextureDesc SeparatedMainDirLightDesc(FRDGTextureDesc::Create2D(SceneColorDesc.Extent, PF_FloatR11G11B10, SceneDepthDesc.ClearValue, TexCreate_ShaderResource | TexCreate_RenderTargetable));
+	FRDGTextureRef SeparatedMainDirLightTexture = GraphBuilder.CreateTexture(SeparatedMainDirLightDesc, TEXT("SLW.SeparatedMainDirLight"));
+
 	FSceneWithoutWaterTextures Textures;
 	Textures.RefractionDownsampleFactor = float(RefractionDownsampleFactor);
 	Textures.Views.SetNum(Views.Num());
@@ -348,6 +369,7 @@ static FSceneWithoutWaterTextures AddCopySceneWithoutWaterPass(
 	check(SceneDepthWithoutSingleLayerWaterTexture);
 	Textures.ColorTexture = SceneColorWithoutSingleLayerWaterTexture;
 	Textures.DepthTexture = SceneDepthWithoutSingleLayerWaterTexture;
+	Textures.SeparatedMainDirLightTexture = SeparatedMainDirLightTexture;
 	return MoveTemp(Textures);
 }
 
@@ -385,18 +407,21 @@ void FDeferredShadingSceneRenderer::RenderSingleLayerWaterReflections(
 		RDG_EVENT_SCOPE_CONDITIONAL(GraphBuilder, Views.Num() > 1, "View%d", ViewIndex);
 
 		FRDGTextureRef ReflectionsColor = nullptr;
-		FRDGTextureRef BlackDummy = SystemTextures.Black;
+		FRDGTextureRef BlackDummyTexture = SystemTextures.Black;
+		FRDGTextureRef WhiteDummyTexture = SystemTextures.White;
 		const FSceneTextureParameters SceneTextureParameters = GetSceneTextureParameters(GraphBuilder, SceneTextures);
 
 		auto SetCommonParameters = [&](FSingleLayerWaterCommonShaderParameters& Parameters)
 		{
-			Parameters.ScreenSpaceReflectionsTexture = ReflectionsColor ? ReflectionsColor : BlackDummy;
+			Parameters.ScreenSpaceReflectionsTexture = ReflectionsColor ? ReflectionsColor : BlackDummyTexture;
 			Parameters.ScreenSpaceReflectionsSampler = TStaticSamplerState<SF_Point>::GetRHI();
 			Parameters.PreIntegratedGF = GSystemTextures.PreintegratedGF->GetRenderTargetItem().ShaderResourceTexture;
 			Parameters.PreIntegratedGFSampler = TStaticSamplerState<SF_Bilinear, AM_Clamp, AM_Clamp, AM_Clamp>::GetRHI();
-			Parameters.SceneNoWaterDepthTexture = SceneWithoutWaterTextures.DepthTexture ? SceneWithoutWaterTextures.DepthTexture : BlackDummy;
+			Parameters.SceneNoWaterDepthTexture = SceneWithoutWaterTextures.DepthTexture ? SceneWithoutWaterTextures.DepthTexture : BlackDummyTexture;
 			Parameters.SceneNoWaterDepthSampler = TStaticSamplerState<SF_Point>::GetRHI();
 			Parameters.SceneNoWaterMinMaxUV = SceneWithoutWaterTextures.Views[ViewIndex].MinMaxUV;
+			Parameters.SeparatedMainDirLightTexture = BlackDummyTexture;
+			Parameters.UseSeparatedMainDirLightTexture = 0.0f;
 			Parameters.SceneTextures = SceneTextureParameters;
 			Parameters.ViewUniformBuffer = GetShaderBinding(View.ViewUniformBuffer);
 			Parameters.ReflectionCaptureData = View.ReflectionCaptureUniformBuffer;
@@ -446,6 +471,48 @@ void FDeferredShadingSceneRenderer::RenderSingleLayerWaterReflections(
 
 		const FPerViewPipelineState& ViewPipelineState = GetViewPipelineState(View);
 
+		// Try to find the ProjectedShadowInfo corresponding to ray trace shadow info for the main directional light.
+		const FLightSceneProxy* SelectedForwardDirectionalLightProxy = View.ForwardLightingResources.SelectedForwardDirectionalLightProxy;
+		FLightSceneInfo* LightSceneInfo = SelectedForwardDirectionalLightProxy->GetLightSceneInfo();
+		FVisibleLightInfo& VisibleLightViewInfo = VisibleLightInfos[LightSceneInfo->Id];
+		FProjectedShadowInfo* DistanceFieldShadowInfo = nullptr;
+		for (int32 ShadowIndex = 0; ShadowIndex < VisibleLightViewInfo.ShadowsToProject.Num(); ShadowIndex++)
+		{
+			FProjectedShadowInfo* ProjectedShadowInfo = VisibleLightViewInfo.ShadowsToProject[ShadowIndex];
+
+			if (ProjectedShadowInfo->bRayTracedDistanceField)
+			{
+				DistanceFieldShadowInfo = ProjectedShadowInfo;
+			}
+		}
+
+		// If DFShadow data has been found, then combine it with the separate main directional light luminance texture.
+		FRDGTextureRef ScreenShadowMaskTexture = SystemTextures.White;
+		if (DistanceFieldShadowInfo)
+		{
+			RDG_EVENT_SCOPE(GraphBuilder, "SLW::DistanceFieldShadow");
+
+			FIntRect ScissorRect;
+			if (!SelectedForwardDirectionalLightProxy->GetScissorRect(ScissorRect, View, View.ViewRect))
+			{
+				ScissorRect = View.ViewRect;
+			}
+
+			// Reset the cached texture to create a new one mapping to he water depth buffer
+			DistanceFieldShadowInfo->ResetRayTracedDistanceFieldShadow();
+
+			const bool bProjectingForForwardShading = false;
+			const bool bForceRGBModulation = true;
+			DistanceFieldShadowInfo->RenderRayTracedDistanceFieldProjection(
+				GraphBuilder,
+				SceneTextures,
+				SceneWithoutWaterTextures.SeparatedMainDirLightTexture,
+				View,
+				ScissorRect,
+				bProjectingForForwardShading,
+				bForceRGBModulation);
+		}
+
 		if (ViewPipelineState.ReflectionsMethod == EReflectionsMethod::Lumen && CVarWaterSingleLayerLumenReflections.GetValueOnRenderThread() != 0)
 		{
 			RDG_EVENT_SCOPE(GraphBuilder, "SLW::LumenReflections");
@@ -605,6 +672,11 @@ void FDeferredShadingSceneRenderer::RenderSingleLayerWaterReflections(
 			PassParameters->VS.TileListData = TiledScreenSpaceReflection.TileListStructureBufferSRV;
 
 			SetCommonParameters(PassParameters->PS.CommonParameters);
+			if (IsWaterDistanceFieldShadowEnabled_Runtime(Scene->GetShaderPlatform()))
+			{
+				PassParameters->PS.CommonParameters.SeparatedMainDirLightTexture = SceneWithoutWaterTextures.SeparatedMainDirLightTexture;
+				PassParameters->PS.CommonParameters.UseSeparatedMainDirLightTexture = 1.0f;
+			}
 
 			PassParameters->IndirectDrawParameter = TiledScreenSpaceReflection.DrawIndirectParametersBuffer;
 			PassParameters->RenderTargets[0] = FRenderTargetBinding(SceneColorTexture, ERenderTargetLoadAction::ELoad);
@@ -726,9 +798,14 @@ void FDeferredShadingSceneRenderer::RenderSingleLayerWaterInner(
 
 	const FRDGSystemTextures& SystemTextures = FRDGSystemTextures::Get(GraphBuilder);
 
-	FRenderTargetBindingSlots RenderTargets;
-	SceneTextures.GetGBufferRenderTargets(ERenderTargetLoadAction::ELoad, RenderTargets);
-	RenderTargets.DepthStencil = FDepthStencilBinding(SceneTextures.Depth.Target, ERenderTargetLoadAction::ELoad, ERenderTargetLoadAction::ELoad, FExclusiveDepthStencil::DepthWrite_StencilWrite);
+	TStaticArray<FTextureRenderTargetBinding, MaxSimultaneousRenderTargets> BasePassTextures;
+	uint32 BasePassTextureCount = SceneTextures.GetGBufferRenderTargets(BasePassTextures);
+	if(IsWaterDistanceFieldShadowEnabled_Runtime(Scene->GetShaderPlatform()))
+	{
+		const bool bNeverClear = true;
+		BasePassTextures[BasePassTextureCount++] = FTextureRenderTargetBinding(SceneWithoutWaterTextures.SeparatedMainDirLightTexture, bNeverClear);
+	}
+	TArrayView<FTextureRenderTargetBinding> BasePassTexturesView = MakeArrayView(BasePassTextures.GetData(), BasePassTextureCount);
 
 	FRDGTextureRef WhiteForwardScreenSpaceShadowMask = SystemTextures.White;
 
@@ -750,7 +827,8 @@ void FDeferredShadingSceneRenderer::RenderSingleLayerWaterInner(
 		PassParameters->ReflectionCapture = View.ReflectionCaptureUniformBuffer;
 		PassParameters->BasePass = CreateOpaqueBasePassUniformBuffer(GraphBuilder, View, ViewIndex, {}, {}, &SceneWithoutWaterTextures);
 		PassParameters->VirtualShadowMapSamplingParameters = VirtualShadowMapArray.GetSamplingParameters(GraphBuilder);
-		PassParameters->RenderTargets = RenderTargets;
+		PassParameters->RenderTargets = GetRenderTargetBindings(ERenderTargetLoadAction::ELoad, BasePassTexturesView);
+		PassParameters->RenderTargets.DepthStencil = FDepthStencilBinding(SceneTextures.Depth.Target, ERenderTargetLoadAction::ELoad, ERenderTargetLoadAction::ELoad, FExclusiveDepthStencil::DepthWrite_StencilWrite);
 
 		View.ParallelMeshDrawCommandPasses[EMeshPass::SingleLayerWaterPass].BuildRenderingCommands(GraphBuilder, Scene->GPUScene, PassParameters->InstanceCullingDrawParams);
 
diff --git a/Engine/Source/Runtime/Renderer/Private/SingleLayerWaterRendering.h b/Engine/Source/Runtime/Renderer/Private/SingleLayerWaterRendering.h
--- a/Engine/Source/Runtime/Renderer/Private/SingleLayerWaterRendering.h
+++ b/Engine/Source/Runtime/Renderer/Private/SingleLayerWaterRendering.h
@@ -16,6 +16,7 @@ struct FSceneWithoutWaterTextures
 		FVector4f MinMaxUV;
 	};
 
+	FRDGTextureRef SeparatedMainDirLightTexture = nullptr;
 	FRDGTextureRef ColorTexture = nullptr;
 	FRDGTextureRef DepthTexture = nullptr;
 	TArray<FView> Views;
