Partial workaround for scene captures invalidating virtual shadow maps:
- VSM caching now "skips" render calls/frames that did not include any VSM rendering. This allows disabling dynamic shadows or VSMs in scene captures to avoid wiping the cache for the main frame.

diff --git a/Engine/Source/Runtime/Renderer/Private/VirtualShadowMaps/VirtualShadowMapCacheManager.cpp b/Engine/Source/Runtime/Renderer/Private/VirtualShadowMaps/VirtualShadowMapCacheManager.cpp
--- a/Engine/Source/Runtime/Renderer/Private/VirtualShadowMaps/VirtualShadowMapCacheManager.cpp
+++ b/Engine/Source/Runtime/Renderer/Private/VirtualShadowMaps/VirtualShadowMapCacheManager.cpp
@@ -476,11 +476,24 @@ void FVirtualShadowMapArrayCacheManager::ExtractFrameData(
 	const FSceneRenderer& SceneRenderer,
 	bool bEnableCaching)
 {
-	// Drop all refs.
-	PrevBuffers = FVirtualShadowMapArrayFrameData();
-	PrevUniformParameters.NumShadowMaps = 0;
+	const bool bNewShadowData = VirtualShadowMapArray.IsAllocated();
+	const bool bDropAll = !bEnableCaching;
+	const bool bDropPrevBuffers = bDropAll || bNewShadowData;
 
-	if (bEnableCaching && VirtualShadowMapArray.IsAllocated())
+	if (bDropPrevBuffers)
+	{
+		PrevBuffers = FVirtualShadowMapArrayFrameData();
+		PrevUniformParameters.NumShadowMaps = 0;
+	}
+
+	if (bDropAll)
+	{
+		// We drop the physical page pool here as well to ensure that it disappears in the case where
+		// thumbnail rendering or similar creates multiple FSceneRenderers that never get deleted.
+		// Caching is disabled on these contexts intentionally to avoid these issues.
+		FreePhysicalPool();
+	}
+	else if (bNewShadowData)
 	{
 		bool bExtractHzbData = false;
 
@@ -495,7 +508,7 @@ void FVirtualShadowMapArrayCacheManager::ExtractFrameData(
 		if (CVarCacheVirtualSMs.GetValueOnRenderThread() != 0)
 		{
 			bExtractHzbData = true;
-		
+
 			GraphBuilder.QueueBufferExtraction(VirtualShadowMapArray.PhysicalPageMetaDataRDG, &PrevBuffers.PhysicalPageMetaData);
 			GraphBuilder.QueueBufferExtraction(VirtualShadowMapArray.DynamicCasterPageFlagsRDG, &PrevBuffers.DynamicCasterPageFlags);
 			GraphBuilder.QueueBufferExtraction(VirtualShadowMapArray.ProjectionDataRDG, &PrevBuffers.ProjectionData);
@@ -505,7 +518,13 @@ void FVirtualShadowMapArrayCacheManager::ExtractFrameData(
 
 			// Move cache entries to previous frame, this implicitly removes any that were not used
 			PrevCacheEntries = CacheEntries;
+			
+			// Store but drop any temp references embedded in the uniform parameters this frame.
+			// We'll reestablish them when we reimport the extracted resources next frame
 			PrevUniformParameters = VirtualShadowMapArray.UniformParameters;
+			PrevUniformParameters.ProjectionData = nullptr;
+			PrevUniformParameters.PageTable = nullptr;
+			PrevUniformParameters.PhysicalPagePool = nullptr;
 		}
 
 		if (bExtractHzbData)
@@ -526,24 +545,17 @@ void FVirtualShadowMapArrayCacheManager::ExtractFrameData(
 		}
 
 		CacheEntries.Reset();
+
+		ExtractStats(GraphBuilder, VirtualShadowMapArray);
 	}
 	else
 	{
-		// We drop the physical page pool here as well to ensure that it disappears in the case where
-		// thumbnail rendering or similar creates multiple FSceneRenderers that never get deleted.
-		// Caching is disabled on these contexts intentionally to avoid these issues.
-		FreePhysicalPool();
-	}
-
-	// Drop any temp references embedded in the uniform parameters this frame.
-	// We'll reestablish them when we reimport the extracted resources next frame
-	PrevUniformParameters.ProjectionData = nullptr;
-	PrevUniformParameters.PageTable = nullptr;
-	PrevUniformParameters.PhysicalPagePool = nullptr;
-
-	if (VirtualShadowMapArray.IsEnabled())
-	{
-		ExtractStats(GraphBuilder, VirtualShadowMapArray);
+		// Do nothing; maintain the data that we had
+		// This allows us to work around some cases where the renderer gets called multiple times in a given frame
+		// - such as scene captures - but does no shadow-related work in all but one of them. We do not want to drop
+		// all the cached data in this case otherwise we effectively get no caching at all.
+		// Ideally in the long run we want the cache itself to be more robust against rendering multiple views. but
+		// for now this at least provides a work-around for some common cases where only one view is rendering VSMs.
 	}
 }
 
