Path Tracer: re-enable mGPU support

Remove coherent sampler since it would be complicated to support in the mGPU case and isn't really necessary for performance anymore

diff --git a/Engine/Plugins/Runtime/nDisplay/Source/DisplayCluster/Private/Render/Viewport/DisplayClusterViewportManagerProxy.cpp b/Engine/Plugins/Runtime/nDisplay/Source/DisplayCluster/Private/Render/Viewport/DisplayClusterViewportManagerProxy.cpp
--- a/Engine/Plugins/Runtime/nDisplay/Source/DisplayCluster/Private/Render/Viewport/DisplayClusterViewportManagerProxy.cpp
+++ b/Engine/Plugins/Runtime/nDisplay/Source/DisplayCluster/Private/Render/Viewport/DisplayClusterViewportManagerProxy.cpp
@@ -446,14 +446,6 @@ void FDisplayClusterViewportManagerProxy::DoCrossGPUTransfers_RenderThread(FRHIC
 					FRHITexture2D* TextureRHI = ViewportProxy->RenderTargets[ViewportContext.ContextNum]->GetViewportRenderTargetResourceRHI();
 
 					FRHIGPUMask RenderTargetGPUMask = (GNumExplicitGPUsForRendering > 1 && RenderTarget) ? RenderTarget->GetGPUMask(RHICmdList) : FRHIGPUMask::GPU0();
-					{
-						static auto CVar = IConsoleManager::Get().FindConsoleVariable(TEXT("r.PathTracing.GPUCount"));
-						if (CVar && CVar->GetInt() > 1)
-						{
-							RenderTargetGPUMask = FRHIGPUMask::All(); // Broadcast to all GPUs 
-						}
-					}
-
 					FRHIGPUMask ContextGPUMask = FRHIGPUMask::FromIndex(ViewportContext.GPUIndex);
 
 					if (ContextGPUMask != RenderTargetGPUMask)
diff --git a/Engine/Shaders/Private/PathTracing/PathTracing.usf b/Engine/Shaders/Private/PathTracing/PathTracing.usf
--- a/Engine/Shaders/Private/PathTracing/PathTracing.usf
+++ b/Engine/Shaders/Private/PathTracing/PathTracing.usf
@@ -10,6 +10,27 @@ PathTracing.usf: Reference path tracing
 
 #include "PathTracingCore.ush"
 
+// describe how we split up the image between tiles and amongst GPUs
+int2 TilePixelOffset;
+int2 TileTextureOffset;
+int ScanlineStride;
+int ScanlineWidth;
+
+
+// Which pixel on screen are we computing?
+// We want to spread out the scanlines evenly across the screen so that each GPU works on a similarly sized task
+int2 ComputePixelIndex(uint2 DispatchIdx)
+{
+	return DispatchIdx * int2(1, ScanlineStride) + TilePixelOffset;
+}
+
+// Where should we store the accumulated pixels of this dispatch?
+// We want all the scanlines computed by a given GPU to be contiguous in memory so we can copy them easily
+int2 ComputeTextureIndex(uint2 DispatchIdx)
+{
+	return DispatchIdx + TileTextureOffset;
+}
+
 #if PATH_TRACER_USE_COMPACTION == 1
 
 int Bounce;
@@ -84,7 +105,7 @@ RAY_TRACING_ENTRY_RAYGEN(PathTracingMainRG)
 	FPathState PathState;
 	if (Bounce == 0)
 	{
-		PathState = CreatePathState(DispatchIdx, DispatchDim);
+		PathState = CreatePathState(ComputePixelIndex(DispatchIdx));
 	}
 	else
 	{
@@ -109,7 +130,8 @@ RAY_TRACING_ENTRY_RAYGEN(PathTracingMainRG)
 	{
 		// nothing left to do
 		// Accumulate radiance and update pixel variance
-		PathState.WritePixel();
+		PathState.WritePixel(ComputeTextureIndex(uint2(LinearPixelIndex % DispatchDim.x,
+			                                           LinearPixelIndex / DispatchDim.x)));
 	}
 }
 
@@ -120,7 +142,7 @@ RAY_TRACING_ENTRY_RAYGEN(PathTracingMainRG)
 	const uint2 DispatchIdx = DispatchRaysIndex().xy;
 	const uint2 DispatchDim = DispatchRaysDimensions().xy;
 
-	FPathState PathState = CreatePathState(DispatchIdx, DispatchDim);
+	FPathState PathState = CreatePathState(ComputePixelIndex(DispatchIdx));
 
 	for (int Bounce = 0; Bounce <= MaxBounces; Bounce++)
 	{
@@ -132,7 +154,7 @@ RAY_TRACING_ENTRY_RAYGEN(PathTracingMainRG)
 	}
 
 	// Accumulate radiance and update pixel variance
-	PathState.WritePixel();
+	PathState.WritePixel(ComputeTextureIndex(DispatchIdx));
 }
 
 #endif
\ No newline at end of file
diff --git a/Engine/Shaders/Private/PathTracing/PathTracingCore.ush b/Engine/Shaders/Private/PathTracing/PathTracingCore.ush
--- a/Engine/Shaders/Private/PathTracing/PathTracingCore.ush
+++ b/Engine/Shaders/Private/PathTracing/PathTracingCore.ush
@@ -57,7 +57,6 @@ RWTexture2D<float4> RadianceTexture;
 RWTexture2D<float4> AlbedoTexture;
 RWTexture2D<float4> NormalTexture;
 RaytracingAccelerationStructure TLAS;
-int3 TileOffset;
 uint SceneVisibleLightCount;
 
 struct FPathState
@@ -75,17 +74,17 @@ struct FPathState
 	float3 SigmaT; // Extinction    //  6 bytes // half precision 
 									// 88 bytes total (see FPathTracingPackedPathState)
 
-	void WritePixel()
+	void WritePixel(uint2 TextureIndex)
 	{
 		// Avoid reading the old pixel on the first sample on the off-chance there is a NaN/Inf pixel ...
-		float4 OldPixel  = Iteration > 0 ? RadianceTexture[PixelIndex] : 0;
-		float4 OldAlbedo = Iteration > 0 ? AlbedoTexture[PixelIndex]   : 0;
-		float4 OldNormal = Iteration > 0 ? NormalTexture[PixelIndex]   : 0;
+		float4 OldPixel  = Iteration > 0 ? RadianceTexture[TextureIndex] : 0;
+		float4 OldAlbedo = Iteration > 0 ? AlbedoTexture[TextureIndex]   : 0;
+		float4 OldNormal = Iteration > 0 ? NormalTexture[TextureIndex]   : 0;
 
 		const float4 PixelValue = float4(Radiance, saturate(BackgroundVisibility));
-		RadianceTexture[PixelIndex] = lerp(OldPixel, PixelValue, BlendFactor);
-		AlbedoTexture[PixelIndex] = lerp(OldAlbedo, float4(Albedo, 0), BlendFactor);
-		NormalTexture[PixelIndex] = lerp(OldNormal, float4(Normal, 0), BlendFactor);
+		RadianceTexture[TextureIndex] = lerp(OldPixel, PixelValue, BlendFactor);
+		AlbedoTexture[TextureIndex] = lerp(OldAlbedo, float4(Albedo, 0), BlendFactor);
+		NormalTexture[TextureIndex] = lerp(OldNormal, float4(Normal, 0), BlendFactor);
 	}
 };
 
@@ -674,82 +673,23 @@ bool ProcessSubsurfaceRandomWalk(inout FPathTracingPayload Payload, inout float3
 	return false;
 }
 
-uint2 Shuffle(uint2 Value, uint2 Seed)
-{
-	const uint M = PATHTRACER_COHERENT_TILE_SIZE;
-#if 1
-	// Scramble the MxM tile of pixels with an error-diffusion pattern
-	// visualizing the seeds shows the distribution of seeds is slightly better
-	const uint X = Value.x;
-	const uint Y = Value.y;
-	uint TileID = 0, HilbertState = 0;
-	for (int i = PATHTRACER_COHERENT_TILE_SHIFT - 1; i >= 0; i--) {
-		uint xi = (X >> i) & 1;
-		uint yi = (Y >> i) & 1;
-		uint Row = 8 * HilbertState + 4 * xi + 2 * yi;
-		TileID = TileID * 4 + ((0x361E9CB4u >> Row) & 3);
-		HilbertState = (0x8FE65831u >> Row) & 3;
-	}
-	uint SeedTile = StrongIntegerHash(Seed.x + StrongIntegerHash(Seed.y));
-	uint SobolIndex = FastOwenScrambling(reversebits(TileID), SeedTile);
-	uint SobolX = SobolIndex;
-	uint SobolY = 0;
-	for (uint v = 1; SobolIndex; SobolIndex >>= 1, v ^= v << 1)
-	{
-		SobolY ^= (SobolIndex & 1) * v;
-	}
-	return reversebits(uint2(SobolX, SobolY)) >> (32 - PATHTRACER_COHERENT_TILE_SHIFT);
-#elif 1
-	// randomly shuffle coordinate (x,y) within an RxR tile of pixels according to the given seed
-	// this has a slight directional preference when visualizing seeds
-	uint SeedTile = StrongIntegerHash(Seed.x + StrongIntegerHash(Seed.y));
-	uint Index = Value.x + Value.y * M;
-	uint ShuffledIndex = FastOwenScrambling(reversebits(Index), SeedTile) % (M * M);
-	return uint2(ShuffledIndex % M, ShuffledIndex / M);
-#else
-	// no shuffling (just to show repeating artifacts)
-	return Value;
-#endif
-}
-
-uint2 Reorder(uint2 I, uint2 N)
-{
-	const uint M = PATHTRACER_COHERENT_TILE_SIZE;
-	uint2 GroupId = I * M / N;				// [0, M)  - which group we landed in?
-	uint2 GroupCoord = ((I * M) % N) / M;	// [0, N/M) - which pixel within the group?
-
-	// reconstruct pixel coordinate, shuffling the assignment of pixels within a group to hide repeating structure
-	return GroupCoord * M + Shuffle(GroupId, GroupCoord);
-}
-
-FPathState CreatePathState(uint2 DispatchIdx, uint2 DispatchDim)
+FPathState CreatePathState(int2 PixelIndex)
 {
 	FPathState Output = (FPathState)0;
 
-	Output.PixelIndex = DispatchIdx;
-	if (SamplerType == PATHTRACER_SAMPLER_COHERENT)
-	{
-		// shuffle the pixels around so the ones that share a seed land in the same warp
-		Output.PixelIndex = Reorder(Output.PixelIndex, DispatchDim);
-	}
-	Output.PixelIndex = Output.PixelIndex + TileOffset.xy;
+	Output.PixelIndex = PixelIndex;
 	uint2 LaunchIndex = Output.PixelIndex + View.ViewRectMin.xy;
 
 	// Initialize random sequence 
-
-	uint2 RandPixel = SamplerType == PATHTRACER_SAMPLER_COHERENT
-		? (DispatchIdx / PATHTRACER_COHERENT_TILE_SIZE)	// unique seed within small squares of the shuffled sub-images
-		//? (DispatchIdx * PATHTRACER_COHERENT_TILE_SIZE) / DispatchDim // unique seed per sub-image
-		: LaunchIndex; // unique seed per pixel
 	if (SamplerType == PATHTRACER_SAMPLER_ERROR_DIFFUSION)
 	{
 		// z-sampler init
-		RandomSequence_Initialize(Output.RandSequence, RandPixel, Iteration, TemporalSeed - Iteration, MaxSamples);
+		RandomSequence_Initialize(Output.RandSequence, LaunchIndex, Iteration, TemporalSeed - Iteration, MaxSamples);
 	}
 	else
 	{
 		// random sobol init
-		RandomSequence_Initialize(Output.RandSequence, RandPixel.x + RandPixel.y * 65536, TemporalSeed);
+		RandomSequence_Initialize(Output.RandSequence, LaunchIndex.x + LaunchIndex.y * 65536, TemporalSeed);
 	}
 
 
@@ -763,7 +703,7 @@ FPathState CreatePathState(uint2 DispatchIdx, uint2 DispatchDim)
 		float R = min(Disk.z, 0.99999994); // avoid log(0) when R=1
 		AAJitter = 0.5 + Sigma * Disk.xy * sqrt(-2.0 * log(1.0 - R * R));
 
-		float2 ViewportUV = (LaunchIndex.xy + AAJitter.xy) * View.BufferSizeAndInvSize.zw;
+		float2 ViewportUV = (LaunchIndex + AAJitter) * View.BufferSizeAndInvSize.zw;
 		Output.Ray = CreatePrimaryRay(ViewportUV);
 	}
 
diff --git a/Engine/Shaders/Private/PathTracing/PathTracingSwizzleScanlines.usf b/Engine/Shaders/Private/PathTracing/PathTracingSwizzleScanlines.usf
new file mode 100644
--- /dev/null
+++ b/Engine/Shaders/Private/PathTracing/PathTracingSwizzleScanlines.usf
@@ -0,0 +1,35 @@
+// Copyright Epic Games, Inc. All Rights Reserved.
+
+#include "../Common.ush"
+
+uint2 DispatchDim;
+uint2 TileSize;
+uint ScanlineStride;
+Texture2D<float4> InputTexture;
+RWTexture2D<float4> OutputTexture;
+
+[numthreads(THREADGROUPSIZE_X, THREADGROUPSIZE_Y, 1)]
+void PathTracingSwizzleScanlinesCS(uint2 DispatchThreadId : SV_DispatchThreadID)
+{
+	if (any(DispatchThreadId >= DispatchDim))
+	{
+		return;
+	}
+
+	// shuffle the scanlines tile by tile
+	int TileBaseOffset = (DispatchThreadId.y / TileSize.y) * TileSize.y;
+	int TileOffset = DispatchThreadId.y - TileBaseOffset;
+
+	int TileDispatchSize = min(TileSize.y, DispatchDim.y - TileBaseOffset);
+
+	int TileScanlineCount = (TileDispatchSize + ScanlineStride - 1) / ScanlineStride;
+
+
+	// Input texture has ScanlineStride copies of the image (as rendered by each GPU) arranged vertically
+	// Swizzle this back into one image by interleaving the scanlines
+	int g = TileOffset / ScanlineStride;
+	int y = TileOffset - g * ScanlineStride; // TileOffset% ScanlineStride;
+	int ScanlineY = y * TileScanlineCount + g + TileBaseOffset;
+
+	OutputTexture[DispatchThreadId] = InputTexture.Load(int3(DispatchThreadId.x, ScanlineY, 0));
+}
diff --git a/Engine/Shaders/Shared/PathTracingDefinitions.h b/Engine/Shaders/Shared/PathTracingDefinitions.h
--- a/Engine/Shaders/Shared/PathTracingDefinitions.h
+++ b/Engine/Shaders/Shared/PathTracingDefinitions.h
@@ -2,15 +2,9 @@
 
 #pragma once
 
-
-#define PATHTRACER_MAX_RECT_TEXTURES		8
-#define PATHTRACER_COHERENT_TILE_SHIFT		3
-#define PATHTRACER_COHERENT_TILE_SIZE		(1 << PATHTRACER_COHERENT_TILE_SHIFT)
-
 // Constants for 'SamplerType'
 #define PATHTRACER_SAMPLER_DEFAULT			0
 #define PATHTRACER_SAMPLER_ERROR_DIFFUSION	1
-#define PATHTRACER_SAMPLER_COHERENT			2
 
 // Constants for the 'Flags' field of FPathTracingLight
 #define PATHTRACER_FLAG_LIGHTING_CHANNEL_MASK				(7 << 0)	// Which lighting channel is this light assigned to?
diff --git a/Engine/Source/Runtime/RHI/Public/MultiGPU.h b/Engine/Source/Runtime/RHI/Public/MultiGPU.h
--- a/Engine/Source/Runtime/RHI/Public/MultiGPU.h
+++ b/Engine/Source/Runtime/RHI/Public/MultiGPU.h
@@ -70,6 +70,15 @@ public:
 #endif
 	}
 
+	FORCEINLINE uint32 GetNumActive() const
+	{
+#if WITH_MGPU
+		return FPlatformMath::CountBits(GPUMask);
+#else
+		return 1;
+#endif
+	}
+
 	FORCEINLINE uint32 GetLastIndex() const
 	{
 #if WITH_MGPU
diff --git a/Engine/Source/Runtime/Renderer/Private/PathTracing.cpp b/Engine/Source/Runtime/Renderer/Private/PathTracing.cpp
--- a/Engine/Source/Runtime/Renderer/Private/PathTracing.cpp
+++ b/Engine/Source/Runtime/Renderer/Private/PathTracing.cpp
@@ -203,20 +203,18 @@ TAutoConsoleVariable<int32> CVarPathTracingSamplerType(
 	TEXT("r.PathTracing.SamplerType"),
 	PATHTRACER_SAMPLER_DEFAULT,
 	TEXT("Controls the way the path tracer generates its random numbers\n")
-	TEXT("0: use a different high quality random sequence per pixel\n")
-	TEXT("1: optimize the random sequence across pixels to reduce visible error at the target sample count\n")
-	TEXT("2: share random seeds across pixels to improve coherence of execution on the GPU. This trades some correlation across the image in exchange for better performance.\n"),
+	TEXT("0: use a different high quality random sequence per pixel (default)\n")
+	TEXT("1: optimize the random sequence across pixels to reduce visible error at the target sample count\n"),
 	ECVF_RenderThreadSafe
 );
 
-#if 0
-// TODO: re-enable this when multi-gpu is supported again
-// r.PathTracing.GPUCount is read only because ComputeViewGPUMasks results cannot change after UE has been launched
-TAutoConsoleVariable<int32> CVarPathTracingGPUCount(
-	TEXT("r.PathTracing.GPUCount"),
-	1,
-	TEXT("Sets the amount of GPUs used for computing the path tracing pass (default = 1 GPU)"),
-	ECVF_RenderThreadSafe | ECVF_ReadOnly
+#if WITH_MGPU
+TAutoConsoleVariable<int32> CVarPathTracingMultiGPU(
+	TEXT("r.PathTracing.MultiGPU"),
+	0,
+	TEXT("Run the path tracer using all available GPUs when enabled (default = 0)\n")
+	TEXT("Using this functionality in the editor requires r.AllowMultiGPUInEditor=1"),
+	ECVF_RenderThreadSafe
 );
 #endif
 
@@ -309,6 +307,7 @@ struct FPathTracingConfig
 	bool VisibleLights;
 	bool UseMISCompensation;
 	bool LockedSamplingPattern;
+	bool UseMultiGPU; // NOTE: Requires invalidation because the buffer layout changes
 	int DenoiserMode; // NOTE: does not require path tracing invalidation
 
 	bool IsDifferent(const FPathTracingConfig& Other) const
@@ -340,7 +339,8 @@ struct FPathTracingConfig
 			LightGridMaxCount != Other.LightGridMaxCount ||
 			VisibleLights != Other.VisibleLights ||
 			UseMISCompensation != Other.UseMISCompensation ||
-			LockedSamplingPattern != Other.LockedSamplingPattern;
+			LockedSamplingPattern != Other.LockedSamplingPattern ||
+			UseMultiGPU != Other.UseMultiGPU;
 	}
 };
 
@@ -691,8 +691,12 @@ class FPathTracingRG : public FGlobalShader
 		// IES Profiles
 		SHADER_PARAMETER_RDG_TEXTURE(Texture2DArray, IESTexture)
 		SHADER_PARAMETER_SAMPLER(SamplerState, IESTextureSampler) // Shared sampler for all IES profiles
-		// Used by multi-GPU rendering
-		SHADER_PARAMETER(FIntVector, TileOffset)
+
+		// Used by multi-GPU rendering and TDR-avoidance tiling
+		SHADER_PARAMETER(FIntPoint, TilePixelOffset)
+		SHADER_PARAMETER(FIntPoint, TileTextureOffset)
+		SHADER_PARAMETER(int32, ScanlineStride)
+		SHADER_PARAMETER(int32, ScanlineWidth)
 
 		// extra parameters required for path compacting kernel
 		SHADER_PARAMETER(int, Bounce)
@@ -732,6 +736,35 @@ class FPathTracingIESAtlasCS : public FGlobalShader
 };
 IMPLEMENT_SHADER_TYPE(, FPathTracingIESAtlasCS, TEXT("/Engine/Private/PathTracing/PathTracingIESAtlas.usf"), TEXT("PathTracingIESAtlasCS"), SF_Compute);
 
+
+class FPathTracingSwizzleScanlinesCS : public FGlobalShader
+{
+	DECLARE_GLOBAL_SHADER(FPathTracingSwizzleScanlinesCS)
+	SHADER_USE_PARAMETER_STRUCT(FPathTracingSwizzleScanlinesCS, FGlobalShader)
+
+	static bool ShouldCompilePermutation(const FGlobalShaderPermutationParameters& Parameters)
+	{
+		return ShouldCompileRayTracingShadersForProject(Parameters.Platform);
+	}
+
+	static void ModifyCompilationEnvironment(const FGlobalShaderPermutationParameters& Parameters, FShaderCompilerEnvironment& OutEnvironment)
+	{
+		OutEnvironment.CompilerFlags.Add(CFLAG_WarningsAsErrors);
+		OutEnvironment.SetDefine(TEXT("THREADGROUPSIZE_X"), FComputeShaderUtils::kGolden2DGroupSize);
+		OutEnvironment.SetDefine(TEXT("THREADGROUPSIZE_Y"), FComputeShaderUtils::kGolden2DGroupSize);
+	}
+
+	BEGIN_SHADER_PARAMETER_STRUCT(FParameters, )
+		SHADER_PARAMETER(FIntPoint, DispatchDim)
+		SHADER_PARAMETER(FIntPoint, TileSize)
+		SHADER_PARAMETER(int32, ScanlineStride)
+		SHADER_PARAMETER_RDG_TEXTURE_SRV(RWTexture2DArray, InputTexture)
+		SHADER_PARAMETER_RDG_TEXTURE_UAV(RWTexture2DArray, OutputTexture)
+	END_SHADER_PARAMETER_STRUCT()
+};
+IMPLEMENT_SHADER_TYPE(, FPathTracingSwizzleScanlinesCS, TEXT("/Engine/Private/PathTracing/PathTracingSwizzleScanlines.usf"), TEXT("PathTracingSwizzleScanlinesCS"), SF_Compute);
+
+
 class FPathTracingBuildAtmosphereOpticalDepthLUTCS : public FGlobalShader
 {
 	DECLARE_GLOBAL_SHADER(FPathTracingBuildAtmosphereOpticalDepthLUTCS)
@@ -1597,6 +1630,13 @@ static void SplitDouble(double x, float* hi, float* lo)
 	*lo = static_cast<float>(x - *hi);
 }
 
+#if WITH_MGPU
+BEGIN_SHADER_PARAMETER_STRUCT(FMGPUTransferParameters, )
+	RDG_TEXTURE_ACCESS(InputTexture, ERHIAccess::CopySrc)
+	RDG_TEXTURE_ACCESS(InputAlbedo, ERHIAccess::CopySrc)
+	RDG_TEXTURE_ACCESS(InputNormal, ERHIAccess::CopySrc)
+END_SHADER_PARAMETER_STRUCT()
+#endif
 
 BEGIN_SHADER_PARAMETER_STRUCT(FDenoiseTextureParameters, )
 	RDG_TEXTURE_ACCESS(InputTexture, ERHIAccess::CopySrc)
@@ -1727,8 +1767,11 @@ void FDeferredShadingSceneRenderer::RenderPathTracing(
 	{
 		AtmosphereOpticalDepthLUT = GraphBuilder.RegisterExternalTexture(GSystemTextures.BlackDummy);
 	}
-
-
+#if WITH_MGPU
+	Config.UseMultiGPU = CVarPathTracingMultiGPU.GetValueOnRenderThread() != 0;
+#else
+	Config.UseMultiGPU = false;
+#endif
 
 	// If the scene has changed in some way (camera move, object movement, etc ...)
 	// we must invalidate the ViewState to start over from scratch
@@ -1779,20 +1822,15 @@ void FDeferredShadingSceneRenderer::RenderPathTracing(
 	}
 	const bool bNeedsMoreRays = Config.PathTracingData.Iteration < MaxSPP;
 
+	// should we use multiple GPUs to render the image?
+	const FRHIGPUMask GPUMask = Config.UseMultiGPU ? FRHIGPUMask::All() : View.GPUMask;
+	const int32 NumGPUs = GPUMask.GetNumActive();
+	const int32 DispatchResX = View.ViewRect.Size().X;
+	const int32 DispatchResY = View.ViewRect.Size().Y;
+	const int32 DispatchSize = FMath::Max(CVarPathTracingDispatchSize.GetValueOnRenderThread(), 64);
+
 	if (bNeedsMoreRays)
 	{
-		// Round up to coherent path tracing tile size to simplify pixel shuffling
-		// TODO: be careful not to write extra pixels past the boundary when using multi-gpu
-
-		const int32 ResX = View.ViewRect.Size().X;
-		const int32 ResY = View.ViewRect.Size().Y;
-
-		const int32 DispatchResX = FMath::DivideAndRoundUp(ResX, PATHTRACER_COHERENT_TILE_SIZE) * PATHTRACER_COHERENT_TILE_SIZE;
-		const int32 DispatchResY = FMath::DivideAndRoundUp(ResY, PATHTRACER_COHERENT_TILE_SIZE) * PATHTRACER_COHERENT_TILE_SIZE;
-
-		int32 DispatchSize = FMath::Max(CVarPathTracingDispatchSize.GetValueOnRenderThread(), 64);
-		DispatchSize = FMath::DivideAndRoundUp(DispatchSize, PATHTRACER_COHERENT_TILE_SIZE) * PATHTRACER_COHERENT_TILE_SIZE;
-
 		// should we use path compaction?
 		const int CompactionType = CVarPathTracingCompaction.GetValueOnRenderThread();
 		const bool bUseIndirectDispatch = GRHISupportsRayTracingDispatchIndirect && CVarPathTracingIndirectDispatch.GetValueOnRenderThread() != 0;
@@ -1803,7 +1841,10 @@ void FDeferredShadingSceneRenderer::RenderPathTracing(
 		FRDGBuffer* PathStateData = nullptr;
 		if (CompactionType == 1)
 		{
-			const int32 NumPaths = FMath::Min(DispatchSize * DispatchSize, DispatchResX * DispatchResY);
+			const int32 NumPaths = FMath::Min(
+				DispatchSize * FMath::DivideAndRoundUp(DispatchSize, NumGPUs),
+				DispatchResX * FMath::DivideAndRoundUp(DispatchResY, NumGPUs)
+			);
 			ActivePaths[0] = GraphBuilder.CreateBuffer(FRDGBufferDesc::CreateBufferDesc(sizeof(int32), NumPaths), TEXT("PathTracer.ActivePaths0"));
 			ActivePaths[1] = GraphBuilder.CreateBuffer(FRDGBufferDesc::CreateBufferDesc(sizeof(int32), NumPaths), TEXT("PathTracer.ActivePaths1"));
 			if (bUseIndirectDispatch)
@@ -1821,150 +1862,233 @@ void FDeferredShadingSceneRenderer::RenderPathTracing(
 		PermutationVector.Set<FPathTracingRG::FCompactionType>(CompactionType);
 		TShaderMapRef<FPathTracingRG> RayGenShader(View.ShaderMap, PermutationVector);
 		FPathTracingRG::FParameters* PreviousPassParameters = nullptr;
-		for (int32 TileY = 0; TileY < DispatchResY; TileY += DispatchSize)
+		// Divide each tile among all the active GPUs (interleaving scanlines)
+		// The assumption is that the tiles are as big as possible, hopefully covering the entire screen
+		// so rather than dividing tiles among GPUs, we divide each tile among all GPUs
+		int32 CurrentGPU = 0; // keep our own counter so that we don't assume the assigned GPUs in the view mask are sequential
+		for (int32 GPUIndex : GPUMask)
 		{
-			for (int32 TileX = 0; TileX < DispatchResX; TileX += DispatchSize)
+			RDG_GPU_MASK_SCOPE(GraphBuilder, FRHIGPUMask::FromIndex(GPUIndex));
+#if WITH_MGPU
+			RDG_EVENT_SCOPE(GraphBuilder, "Path Tracing GPU%d", GPUIndex);
+#endif
+			for (int32 TileY = 0; TileY < DispatchResY; TileY += DispatchSize)
 			{
-				const int32 DispatchSizeX = FMath::Min(DispatchSize, DispatchResX - TileX);
-				const int32 DispatchSizeY = FMath::Min(DispatchSize, DispatchResY - TileY);
-				if (CompactionType == 1)
+				for (int32 TileX = 0; TileX < DispatchResX; TileX += DispatchSize)
 				{
-					AddClearUAVPass(GraphBuilder, GraphBuilder.CreateUAV(ActivePaths[0], PF_R32_UINT), 0);
-				}
-				// When using path compaction, we need to run the path tracer once per bounce
-				// otherwise, the path tracer is the one doing the bounces
-				for (int Bounce = 0, MaxBounces = CompactionType == 1 ? Config.PathTracingData.MaxBounces : 0; Bounce <= MaxBounces; Bounce++)
-				{
-					FPathTracingRG::FParameters* PassParameters = GraphBuilder.AllocParameters<FPathTracingRG::FParameters>();
-					PassParameters->TLAS = View.GetRayTracingSceneViewChecked();
-					PassParameters->ViewUniformBuffer = View.ViewUniformBuffer;
-					PassParameters->PathTracingData = Config.PathTracingData;
-					if (PreviousPassParameters == nullptr)
-					{
-						// upload sky/lights data
-						SetLightParameters(GraphBuilder, PassParameters, Scene, View, Config.UseMISCompensation);
-						PreviousPassParameters = PassParameters;
-					}
-					else
-					{
-						// re-use from last iteration
-						PassParameters->IESTexture = PreviousPassParameters->IESTexture;
-						PassParameters->IESTextureSampler = PreviousPassParameters->IESTextureSampler;
-						PassParameters->LightGridParameters = PreviousPassParameters->LightGridParameters;
-						PassParameters->SceneLightCount = PreviousPassParameters->SceneLightCount;
-						PassParameters->SceneVisibleLightCount = PreviousPassParameters->SceneVisibleLightCount;
-						PassParameters->SceneLights = PreviousPassParameters->SceneLights;
-						PassParameters->SkylightParameters = PreviousPassParameters->SkylightParameters;
-					}
-					if (Config.PathTracingData.EnableDirectLighting == 0)
-					{
-						PassParameters->SceneVisibleLightCount = 0;
-					}
-
-					PassParameters->RadianceTexture = GraphBuilder.CreateUAV(RadianceTexture);
-					PassParameters->AlbedoTexture   = GraphBuilder.CreateUAV(AlbedoTexture);
-					PassParameters->NormalTexture   = GraphBuilder.CreateUAV(NormalTexture);
+					const int32 DispatchSizeX = FMath::Min(DispatchSize, DispatchResX - TileX);
+					const int32 DispatchSizeY = FMath::Min(DispatchSize, DispatchResY - TileY);
 
-					if (Config.PathTracingData.EnableAtmosphere)
-					{
-						PassParameters->Atmosphere = Scene->GetSkyAtmosphereSceneInfo()->GetAtmosphereUniformBuffer();
-						FVector PlanetCenterTranslatedWorld = Scene->GetSkyAtmosphereSceneInfo()->GetSkyAtmosphereSceneProxy().GetAtmosphereSetup().PlanetCenterKm * double(FAtmosphereSetup::SkyUnitToCm) + View.ViewMatrices.GetPreViewTranslation();
-						SplitDouble(PlanetCenterTranslatedWorld.X, &PassParameters->PlanetCenterTranslatedWorldHi.X, &PassParameters->PlanetCenterTranslatedWorldLo.X);
-						SplitDouble(PlanetCenterTranslatedWorld.Y, &PassParameters->PlanetCenterTranslatedWorldHi.Y, &PassParameters->PlanetCenterTranslatedWorldLo.Y);
-						SplitDouble(PlanetCenterTranslatedWorld.Z, &PassParameters->PlanetCenterTranslatedWorldHi.Z, &PassParameters->PlanetCenterTranslatedWorldLo.Z);
-					}
-					else
-					{
-						FAtmosphereUniformShaderParameters AtmosphereParams = {};
-						PassParameters->Atmosphere = CreateUniformBufferImmediate(AtmosphereParams, EUniformBufferUsage::UniformBuffer_SingleFrame);
-						PassParameters->PlanetCenterTranslatedWorldHi = FVector3f(0);
-						PassParameters->PlanetCenterTranslatedWorldLo = FVector3f(0);
-					}
-					PassParameters->AtmosphereOpticalDepthLUT = AtmosphereOpticalDepthLUT;
-					PassParameters->AtmosphereOpticalDepthLUTSampler = TStaticSamplerState<SF_Bilinear, AM_Clamp, AM_Clamp, AM_Clamp>::GetRHI();
+					const int32 DispatchSizeYSplit = FMath::DivideAndRoundUp(DispatchSizeY, NumGPUs);
 
-					if (Config.PathTracingData.EnableFog)
+					// Compute the dispatch size for just this set of scanlines
+					const int32 DispatchSizeYLocal = FMath::Min(DispatchSizeYSplit, DispatchSizeY - CurrentGPU * DispatchSizeYSplit);
+					if (CompactionType == 1)
 					{
-						PassParameters->FogParameters = PrepareFogParameters(View, Scene->ExponentialFogs[0]);
+						AddClearUAVPass(GraphBuilder, GraphBuilder.CreateUAV(ActivePaths[0], PF_R32_UINT), 0);
 					}
-					else
+					// When using path compaction, we need to run the path tracer once per bounce
+					// otherwise, the path tracer is the one doing the bounces
+					for (int Bounce = 0, MaxBounces = CompactionType == 1 ? Config.PathTracingData.MaxBounces : 0; Bounce <= MaxBounces; Bounce++)
 					{
-						PassParameters->FogParameters = {};
-					}
+						FPathTracingRG::FParameters* PassParameters = GraphBuilder.AllocParameters<FPathTracingRG::FParameters>();
+						PassParameters->TLAS = View.GetRayTracingSceneViewChecked();
+						PassParameters->ViewUniformBuffer = View.ViewUniformBuffer;
+						PassParameters->PathTracingData = Config.PathTracingData;
+						if (PreviousPassParameters == nullptr)
+						{
+							// upload sky/lights data
+							RDG_GPU_MASK_SCOPE(GraphBuilder, GPUMask); // make sure this happens on all GPUs we will be rendering on
+							SetLightParameters(GraphBuilder, PassParameters, Scene, View, Config.UseMISCompensation);
+						}
+						else
+						{
+							// re-use from last iteration
+							PassParameters->IESTexture = PreviousPassParameters->IESTexture;
+							PassParameters->IESTextureSampler = PreviousPassParameters->IESTextureSampler;
+							PassParameters->LightGridParameters = PreviousPassParameters->LightGridParameters;
+							PassParameters->SceneLightCount = PreviousPassParameters->SceneLightCount;
+							PassParameters->SceneVisibleLightCount = PreviousPassParameters->SceneVisibleLightCount;
+							PassParameters->SceneLights = PreviousPassParameters->SceneLights;
+							PassParameters->SkylightParameters = PreviousPassParameters->SkylightParameters;
+						}
+						if (Config.PathTracingData.EnableDirectLighting == 0)
+						{
+							PassParameters->SceneVisibleLightCount = 0;
+						}
 
-					// TODO: in multi-gpu case, assign different tiles to different GPUs
-					PassParameters->TileOffset.X = TileX;
-					PassParameters->TileOffset.Y = TileY;
+						PassParameters->RadianceTexture = GraphBuilder.CreateUAV(RadianceTexture);
+						PassParameters->AlbedoTexture = GraphBuilder.CreateUAV(AlbedoTexture);
+						PassParameters->NormalTexture = GraphBuilder.CreateUAV(NormalTexture);
 
-					PassParameters->Bounce = Bounce;
-					if (CompactionType == 1)
-					{
-						PassParameters->ActivePaths = GraphBuilder.CreateSRV(ActivePaths[Bounce & 1], PF_R32_SINT);
-						PassParameters->NextActivePaths = GraphBuilder.CreateUAV(ActivePaths[(Bounce & 1) ^ 1], PF_R32_SINT);
-						PassParameters->PathStateData = GraphBuilder.CreateUAV(PathStateData);
-						if (bUseIndirectDispatch)
+						if (PreviousPassParameters != nullptr)
 						{
-							PassParameters->NumPathStates = GraphBuilder.CreateUAV(NumActivePaths[Bounce & 1], PF_R32_UINT);
-							PassParameters->PathTracingIndirectArgs = NumActivePaths[(Bounce & 1) ^ 1];
+							PassParameters->Atmosphere = PreviousPassParameters->Atmosphere;
+							PassParameters->PlanetCenterTranslatedWorldHi = PreviousPassParameters->PlanetCenterTranslatedWorldHi;
+							PassParameters->PlanetCenterTranslatedWorldLo = PreviousPassParameters->PlanetCenterTranslatedWorldLo;
+						}
+						else if (Config.PathTracingData.EnableAtmosphere)
+						{
+							PassParameters->Atmosphere = Scene->GetSkyAtmosphereSceneInfo()->GetAtmosphereUniformBuffer();
+							FVector PlanetCenterTranslatedWorld = Scene->GetSkyAtmosphereSceneInfo()->GetSkyAtmosphereSceneProxy().GetAtmosphereSetup().PlanetCenterKm * double(FAtmosphereSetup::SkyUnitToCm) + View.ViewMatrices.GetPreViewTranslation();
+							SplitDouble(PlanetCenterTranslatedWorld.X, &PassParameters->PlanetCenterTranslatedWorldHi.X, &PassParameters->PlanetCenterTranslatedWorldLo.X);
+							SplitDouble(PlanetCenterTranslatedWorld.Y, &PassParameters->PlanetCenterTranslatedWorldHi.Y, &PassParameters->PlanetCenterTranslatedWorldLo.Y);
+							SplitDouble(PlanetCenterTranslatedWorld.Z, &PassParameters->PlanetCenterTranslatedWorldHi.Z, &PassParameters->PlanetCenterTranslatedWorldLo.Z);
 						}
 						else
 						{
-							PassParameters->NumPathStates = GraphBuilder.CreateUAV(NumActivePaths[0], PF_R32_UINT);
-							AddClearUAVPass(GraphBuilder, PassParameters->NextActivePaths, -1); // make sure everything is initialized to -1 since paths that go inactive don't write anything
+							FAtmosphereUniformShaderParameters AtmosphereParams = {};
+							PassParameters->Atmosphere = CreateUniformBufferImmediate(AtmosphereParams, EUniformBufferUsage::UniformBuffer_SingleFrame);
+							PassParameters->PlanetCenterTranslatedWorldHi = FVector3f(0);
+							PassParameters->PlanetCenterTranslatedWorldLo = FVector3f(0);
 						}
-						AddClearUAVPass(GraphBuilder, PassParameters->NumPathStates, 0);
-					}
-					ClearUnusedGraphResources(RayGenShader, PassParameters);
-					const bool bFlushRenderingCommands = FlushRenderingCommands == 1 || (FlushRenderingCommands == 2 && Bounce == MaxBounces);
-					GraphBuilder.AddPass(
-						CompactionType == 1
-							? RDG_EVENT_NAME("Path Tracer Compute (%d x %d) Tile=(%d,%d - %dx%d) Sample=%d/%d NumLights=%d (Bounce=%d%s)", ResX, ResY, TileX, TileY, DispatchSizeX, DispatchSizeY, PathTracingState->SampleIndex, MaxSPP, PassParameters->SceneLightCount, Bounce, bUseIndirectDispatch && Bounce > 0 ? TEXT(" indirect") : TEXT(""))
-							: RDG_EVENT_NAME("Path Tracer Compute (%d x %d) Tile=(%d,%d - %dx%d) Sample=%d/%d NumLights=%d", ResX, ResY, TileX, TileY, DispatchSizeX, DispatchSizeY, PathTracingState->SampleIndex, MaxSPP, PassParameters->SceneLightCount),
-						PassParameters,
-						ERDGPassFlags::Compute,
-						[PassParameters, RayGenShader, DispatchSizeX, DispatchSizeY, bUseIndirectDispatch, bFlushRenderingCommands, &View](FRHIRayTracingCommandList& RHICmdList)
-					{
-						FRHIRayTracingScene* RayTracingSceneRHI = View.GetRayTracingSceneChecked();
-			
-						FRayTracingShaderBindingsWriter GlobalResources;
-						SetShaderParameters(GlobalResources, RayGenShader, *PassParameters);
-						if (bUseIndirectDispatch && PassParameters->Bounce > 0)
+						PassParameters->AtmosphereOpticalDepthLUT = AtmosphereOpticalDepthLUT;
+						PassParameters->AtmosphereOpticalDepthLUTSampler = TStaticSamplerState<SF_Bilinear, AM_Clamp, AM_Clamp, AM_Clamp>::GetRHI();
+
+						if (Config.PathTracingData.EnableFog)
 						{
-							PassParameters->PathTracingIndirectArgs->MarkResourceAsUsed();
-							RHICmdList.RayTraceDispatchIndirect(
-								View.RayTracingMaterialPipeline,
-								RayGenShader.GetRayTracingShader(),
-								RayTracingSceneRHI, GlobalResources,
-								PassParameters->PathTracingIndirectArgs->GetIndirectRHICallBuffer(), 0
-							);
+							PassParameters->FogParameters = PrepareFogParameters(View, Scene->ExponentialFogs[0]);
 						}
 						else
 						{
-							RHICmdList.RayTraceDispatch(
-								View.RayTracingMaterialPipeline,
-								RayGenShader.GetRayTracingShader(),
-								RayTracingSceneRHI, GlobalResources,
-								DispatchSizeX, DispatchSizeY
-							);
+							PassParameters->FogParameters = {};
 						}
-						if (bFlushRenderingCommands)
+
+						PassParameters->TilePixelOffset.X = TileX;
+						PassParameters->TilePixelOffset.Y = TileY + CurrentGPU;
+						PassParameters->TileTextureOffset.X = TileX;
+						PassParameters->TileTextureOffset.Y = TileY + CurrentGPU * DispatchSizeYSplit;
+						PassParameters->ScanlineStride = NumGPUs;
+						PassParameters->ScanlineWidth = DispatchSizeX;
+
+						PassParameters->Bounce = Bounce;
+						if (CompactionType == 1)
 						{
-							RHICmdList.SubmitCommandsHint();
+							PassParameters->ActivePaths = GraphBuilder.CreateSRV(ActivePaths[Bounce & 1], PF_R32_SINT);
+							PassParameters->NextActivePaths = GraphBuilder.CreateUAV(ActivePaths[(Bounce & 1) ^ 1], PF_R32_SINT);
+							PassParameters->PathStateData = GraphBuilder.CreateUAV(PathStateData);
+							if (bUseIndirectDispatch)
+							{
+								PassParameters->NumPathStates = GraphBuilder.CreateUAV(NumActivePaths[Bounce & 1], PF_R32_UINT);
+								PassParameters->PathTracingIndirectArgs = NumActivePaths[(Bounce & 1) ^ 1];
+							}
+							else
+							{
+								PassParameters->NumPathStates = GraphBuilder.CreateUAV(NumActivePaths[0], PF_R32_UINT);
+								AddClearUAVPass(GraphBuilder, PassParameters->NextActivePaths, -1); // make sure everything is initialized to -1 since paths that go inactive don't write anything
+							}
+							AddClearUAVPass(GraphBuilder, PassParameters->NumPathStates, 0);
 						}
-					});
+						ClearUnusedGraphResources(RayGenShader, PassParameters);
+						const bool bFlushRenderingCommands = FlushRenderingCommands == 1 || (FlushRenderingCommands == 2 && Bounce == MaxBounces);
+						GraphBuilder.AddPass(
+							CompactionType == 1
+							? RDG_EVENT_NAME("Path Tracer Compute (%d x %d) Tile=(%d,%d - %dx%d) Sample=%d/%d NumLights=%d (Bounce=%d%s)", DispatchResX, DispatchResY, TileX, TileY, DispatchSizeX, DispatchSizeYLocal, PathTracingState->SampleIndex, MaxSPP, PassParameters->SceneLightCount, Bounce, bUseIndirectDispatch && Bounce > 0 ? TEXT(" indirect") : TEXT(""))
+							: RDG_EVENT_NAME("Path Tracer Compute (%d x %d) Tile=(%d,%d - %dx%d) Sample=%d/%d NumLights=%d", DispatchResX, DispatchResY, TileX, TileY, DispatchSizeX, DispatchSizeYLocal, PathTracingState->SampleIndex, MaxSPP, PassParameters->SceneLightCount),
+							PassParameters,
+							ERDGPassFlags::Compute,
+							[PassParameters, RayGenShader, DispatchSizeX, DispatchSizeYLocal, bUseIndirectDispatch, bFlushRenderingCommands, GPUIndex, &View](FRHIRayTracingCommandList& RHICmdList)
+							{
+								FRHIRayTracingScene* RayTracingSceneRHI = View.GetRayTracingSceneChecked();
+
+								FRayTracingShaderBindingsWriter GlobalResources;
+								SetShaderParameters(GlobalResources, RayGenShader, *PassParameters);
+								if (bUseIndirectDispatch && PassParameters->Bounce > 0)
+								{
+									PassParameters->PathTracingIndirectArgs->MarkResourceAsUsed();
+									RHICmdList.RayTraceDispatchIndirect(
+										View.RayTracingMaterialPipeline,
+										RayGenShader.GetRayTracingShader(),
+										RayTracingSceneRHI, GlobalResources,
+										PassParameters->PathTracingIndirectArgs->GetIndirectRHICallBuffer(), 0
+									);
+								}
+								else
+								{
+									RHICmdList.RayTraceDispatch(
+										View.RayTracingMaterialPipeline,
+										RayGenShader.GetRayTracingShader(),
+										RayTracingSceneRHI, GlobalResources,
+										DispatchSizeX, DispatchSizeYLocal
+									);
+								}
+								if (bFlushRenderingCommands)
+								{
+									RHICmdList.SubmitCommandsHint();
+								}
+							});
+						if (PreviousPassParameters == nullptr)
+						{
+							PreviousPassParameters = PassParameters;
+						}
+					}
 				}
 			}
+			++CurrentGPU;
 		}
+
+#if WITH_MGPU
+		if (NumGPUs > 1)
+		{
+			FMGPUTransferParameters* Parameters = GraphBuilder.AllocParameters<FMGPUTransferParameters>();
+			Parameters->InputTexture = RadianceTexture;
+			Parameters->InputAlbedo = AlbedoTexture;
+			Parameters->InputNormal = NormalTexture;
+			GraphBuilder.AddPass(RDG_EVENT_NAME("Path Tracer Cross-GPU Transfer (%d GPUs)", NumGPUs), Parameters, ERDGPassFlags::Readback,
+				[Parameters, DispatchResX, DispatchResY, DispatchSize, GPUMask, MainGPUMask = View.GPUMask](FRHICommandListImmediate& RHICmdList)
+				{
+					const int32 FirstGPUIndex = MainGPUMask.GetFirstIndex();
+					const int32 NumGPUs = GPUMask.GetNumActive();
+					TArray<FTransferResourceParams> TransferParams;
+					for (int32 TileY = 0; TileY < DispatchResY; TileY += DispatchSize)
+					{
+						for (int32 TileX = 0; TileX < DispatchResX; TileX += DispatchSize)
+						{
+							const int32 DispatchSizeX = FMath::Min(DispatchSize, DispatchResX - TileX);
+							const int32 DispatchSizeY = FMath::Min(DispatchSize, DispatchResY - TileY);
+
+							const int32 DispatchSizeYSplit = FMath::DivideAndRoundUp(DispatchSizeY, NumGPUs);
+
+							// Divide each tile among all the active GPUs (interleaving scanlines)
+							// The assumption is that the tiles are as big as possible, hopefully covering the entire screen
+							// so rather than dividing tiles among GPUs, we divide each tile among all GPUs
+							int32 CurrentGPU = 0; // keep our own counter so that we don't assume the assigned GPUs in the view mask are sequential
+							for (int32 GPUIndex : GPUMask)
+							{
+								// Compute the dispatch size for just this set of scanlines
+								const int32 DispatchSizeYLocal = FMath::Min(DispatchSizeYSplit, DispatchSizeY - CurrentGPU * DispatchSizeYSplit);
+								// If this portion of the texture was not rendered by GPU0, transfer the rendered pixels there
+								if (GPUIndex != FirstGPUIndex)
+								{
+									FIntRect TileToCopy;
+									TileToCopy.Min.X = TileX;
+									TileToCopy.Min.Y = TileY + CurrentGPU * DispatchSizeYSplit;
+									TileToCopy.Max.X = TileX + DispatchSizeX;
+									TileToCopy.Max.Y = TileToCopy.Min.Y + DispatchSizeYLocal;
+									TransferParams.Emplace(Parameters->InputTexture->GetRHI(), TileToCopy, GPUIndex, FirstGPUIndex, false, false);
+									TransferParams.Emplace(Parameters->InputAlbedo->GetRHI(), TileToCopy, GPUIndex, FirstGPUIndex, false, false);
+									TransferParams.Emplace(Parameters->InputNormal->GetRHI(), TileToCopy, GPUIndex, FirstGPUIndex, false, false);
+								}
+								++CurrentGPU;
+							}
+						}
+					}
+					RHICmdList.TransferResources(TransferParams);
+				}
+			);
+		}
+#endif
 		// After we are done, make sure we remember our texture for next time so that we can accumulate samples across frames
 		GraphBuilder.QueueTextureExtraction(RadianceTexture, &PathTracingState->RadianceRT);
-		GraphBuilder.QueueTextureExtraction(AlbedoTexture  , &PathTracingState->AlbedoRT  );
-		GraphBuilder.QueueTextureExtraction(NormalTexture  , &PathTracingState->NormalRT  );
+		GraphBuilder.QueueTextureExtraction(AlbedoTexture, &PathTracingState->AlbedoRT);
+		GraphBuilder.QueueTextureExtraction(NormalTexture, &PathTracingState->NormalRT);
 
 		// Bump counters for next frame
 		++PathTracingState->SampleIndex;
 		++PathTracingState->FrameIndex;
 	}
 
+	// Figure out if the denoiser is enabled and needs to run
 	FRDGTexture* DenoisedRadianceTexture = nullptr;
 	int DenoiserMode = CVarPathTracingDenoiser.GetValueOnRenderThread();
 	if (DenoiserMode < 0)
@@ -1972,17 +2096,56 @@ void FDeferredShadingSceneRenderer::RenderPathTracing(
 		DenoiserMode = View.FinalPostProcessSettings.PathTracingEnableDenoiser;
 	}
 	const bool IsDenoiserEnabled = DenoiserMode != 0 && GPathTracingDenoiserFunc != nullptr;
-	if (IsDenoiserEnabled)
+	// Request denoise if this is the last sample OR allow turning on the denoiser after the image has stopped accumulating samples
+	const bool NeedsDenoise = IsDenoiserEnabled &&
+		(((Config.PathTracingData.Iteration + 1) == MaxSPP) ||
+		 (!bNeedsMoreRays && DenoiserMode != PathTracingState->LastConfig.DenoiserMode));
+
+#if WITH_MGPU
+	if (NumGPUs > 1)
 	{
-		// request denoise if this is the last sample
-		bool NeedsDenoise = (Config.PathTracingData.Iteration + 1) == MaxSPP;
-		// also allow turning on the denoiser after the image has stopped accumulating samples
-		if (!bNeedsMoreRays)
+		// mGPU renders blocks of pixels that need to be mapped back into alternating scanlines
+		// perform this swizzling now with a simple compute shader
+
+		TShaderMapRef<FPathTracingSwizzleScanlinesCS> ComputeShader(GetGlobalShaderMap(GMaxRHIFeatureLevel));
+		FRDGTextureDesc Desc = FRDGTextureDesc::Create2D(
+			View.ViewRect.Size(),
+			PF_A32B32G32R32F,
+			FClearValueBinding::None,
+			TexCreate_ShaderResource | TexCreate_UAV);
+		FRDGTexture* NewRadianceTexture = GraphBuilder.CreateTexture(Desc, TEXT("PathTracer.RadianceUnswizzled"));
+		FRDGTexture* NewAlbedoTexture = NeedsDenoise ? GraphBuilder.CreateTexture(Desc, TEXT("PathTracer.AlbedoUnswizzled")) : nullptr;
+		FRDGTexture* NewNormalTexture = NeedsDenoise ? GraphBuilder.CreateTexture(Desc, TEXT("PathTracer.NormalUnswizzled")) : nullptr;
+
+		FRDGTexture* InputTextures[3] = { RadianceTexture, AlbedoTexture, NormalTexture };
+		FRDGTexture* OutputTextures[3] = { NewRadianceTexture, NewAlbedoTexture, NewNormalTexture };
+		for (int Index = 0, Num = NeedsDenoise ? 3 : 1; Index < Num; Index++)
 		{
-			// we aren't currently rendering, run the denoiser if we just turned it on
-			NeedsDenoise |= DenoiserMode != PathTracingState->LastConfig.DenoiserMode;
+			FPathTracingSwizzleScanlinesCS::FParameters* PassParameters = GraphBuilder.AllocParameters<FPathTracingSwizzleScanlinesCS::FParameters>();
+			PassParameters->DispatchDim.X = DispatchResX;
+			PassParameters->DispatchDim.Y = DispatchResY;
+			PassParameters->TileSize.X = DispatchSize;
+			PassParameters->TileSize.Y = DispatchSize;
+			PassParameters->ScanlineStride = NumGPUs;
+			PassParameters->InputTexture = GraphBuilder.CreateSRV(FRDGTextureSRVDesc::Create(InputTextures[Index]));
+			PassParameters->OutputTexture = GraphBuilder.CreateUAV(OutputTextures[Index]);
+			FComputeShaderUtils::AddPass(
+				GraphBuilder,
+				RDG_EVENT_NAME("UnswizzleScanlines(%d)", Index),
+				ComputeShader,
+				PassParameters,
+				FComputeShaderUtils::GetGroupCount(FIntPoint(DispatchResX, DispatchResY), FComputeShaderUtils::kGolden2DGroupSize));
 		}
 
+		// let the remaining code operate on the unswizzled textures
+		RadianceTexture = NewRadianceTexture;
+		AlbedoTexture = NewAlbedoTexture;
+		NormalTexture = NewNormalTexture;
+	}
+#endif
+
+	if (IsDenoiserEnabled)
+	{
 		if (PathTracingState->RadianceDenoisedRT)
 		{
 			// we already have a texture for this
@@ -2002,16 +2165,14 @@ void FDeferredShadingSceneRenderer::RenderPathTracing(
 				DenoisedRadianceTexture = GraphBuilder.CreateTexture(RadianceTextureDesc, TEXT("PathTracer.DenoisedRadiance"), ERDGTextureFlags::MultiFrame);
 			}
 
-			// Need to read GPU mask outside Pass function, as the value is not refreshed inside the pass
-			FRHIGPUMask GPUMask = GraphBuilder.RHICmdList.GetGPUMask();
-
 			FDenoiseTextureParameters* DenoiseParameters = GraphBuilder.AllocParameters<FDenoiseTextureParameters>();
 			DenoiseParameters->InputTexture = RadianceTexture;
 			DenoiseParameters->InputAlbedo = AlbedoTexture;
 			DenoiseParameters->InputNormal = NormalTexture;
 			DenoiseParameters->OutputTexture = DenoisedRadianceTexture;
+			// Need to read GPU mask outside Pass function, as the value is not refreshed inside the pass
 			GraphBuilder.AddPass(RDG_EVENT_NAME("Path Tracer Denoiser Plugin"), DenoiseParameters, ERDGPassFlags::Readback, 
-				[DenoiseParameters, DenoiserMode, GPUMask](FRHICommandListImmediate& RHICmdList)
+				[DenoiseParameters, DenoiserMode, GPUMask = View.GPUMask](FRHICommandListImmediate& RHICmdList)
 				{
 					GPathTracingDenoiserFunc(RHICmdList,
 						DenoiseParameters->InputTexture->GetRHI()->GetTexture2D(),
diff --git a/Engine/Source/Runtime/Renderer/Private/PostProcess/PostProcessEyeAdaptation.cpp b/Engine/Source/Runtime/Renderer/Private/PostProcess/PostProcessEyeAdaptation.cpp
--- a/Engine/Source/Runtime/Renderer/Private/PostProcess/PostProcessEyeAdaptation.cpp
+++ b/Engine/Source/Runtime/Renderer/Private/PostProcess/PostProcessEyeAdaptation.cpp
@@ -800,6 +800,8 @@ void FSceneViewState::FEyeAdaptationManager::EnqueueExposureTextureReadback(FRDG
 	
 	if (ReadbackBuffersNumPending < MAX_READBACK_BUFFERS)
 	{
+		// limit to single bit here for readback to work
+		RDG_GPU_MASK_SCOPE(GraphBuilder, FRHIGPUMask::FromIndex(GraphBuilder.RHICmdList.GetGPUMask().GetFirstIndex()));
 		FRDGTextureRef CurrentTexture = GraphBuilder.RegisterExternalTexture(PooledRenderTarget[CurrentBuffer], ERDGTextureFlags::MultiFrame);
 
 		FRHIGPUTextureReadback* ExposureReadbackTexture = ExposureReadbackTextures[ReadbackBuffersWriteIndex];
diff --git a/Engine/Source/Runtime/Renderer/Private/SceneRendering.cpp b/Engine/Source/Runtime/Renderer/Private/SceneRendering.cpp
--- a/Engine/Source/Runtime/Renderer/Private/SceneRendering.cpp
+++ b/Engine/Source/Runtime/Renderer/Private/SceneRendering.cpp
@@ -910,7 +910,6 @@ void FViewInfo::Init()
 	SinCosInscatteringColorCubemapRotation = FVector2f::ZeroVector;
 	FogInscatteringColorCubemap = nullptr;
 	FogInscatteringTextureParameters = FVector::ZeroVector;
-	VolumetricFogNearFadeInDistanceInv = 100000000.0f;
 
 	SkyAtmosphereCameraAerialPerspectiveVolume = nullptr;
 	SkyAtmosphereUniformShaderParameters = nullptr;
@@ -1300,9 +1299,6 @@ void FViewInfo::SetupUniformBufferParameters(
 		Scene = Family->Scene->GetRenderScene();
 	}
 
-	ERHIFeatureLevel::Type RHIFeatureLevel = Scene == nullptr ? GMaxRHIFeatureLevel : Scene->GetFeatureLevel();
-	EShaderPlatform ShaderPlatform = GShaderPlatformForFeatureLevel[RHIFeatureLevel];
-
 	const FVector DefaultSunDirection(0.0f, 0.0f, 1.0f); // Up vector so that the AtmosphericLightVector node always output a valid direction.
 	auto ClearAtmosphereLightData = [&](uint32 Index)
 	{
@@ -1352,11 +1348,6 @@ void FViewInfo::SetupUniformBufferParameters(
 
 	ViewUniformShaderParameters.BufferToSceneTextureScale = FVector2f(1.0f, 1.0f);
 
-	{
-		extern bool IsWaterDistanceFieldShadowEnabled_Runtime(const FStaticShaderPlatform Platform);
-		ViewUniformShaderParameters.SeparateWaterMainDirLightLuminance = IsWaterDistanceFieldShadowEnabled_Runtime(ShaderPlatform) ? 1.0f : 0.0f;
-	}
-
 	FRHITexture* TransmittanceLutTextureFound = nullptr;
 	FRHITexture* SkyViewLutTextureFound = nullptr;
 	FRHITexture* CameraAerialPerspectiveVolumeFound = nullptr;
@@ -1686,6 +1677,9 @@ void FViewInfo::SetupUniformBufferParameters(
 
 	ViewUniformShaderParameters.DecalDepthBias = CVarDecalDepthBias.GetValueOnRenderThread();
 
+	ERHIFeatureLevel::Type RHIFeatureLevel = Scene == nullptr ? GMaxRHIFeatureLevel : Scene->GetFeatureLevel();
+	EShaderPlatform ShaderPlatform = GShaderPlatformForFeatureLevel[RHIFeatureLevel];
+
 	ViewUniformShaderParameters.IndirectLightingColorScale = FVector3f(FinalPostProcessSettings.IndirectLightingColor.R * FinalPostProcessSettings.IndirectLightingIntensity,
 		FinalPostProcessSettings.IndirectLightingColor.G * FinalPostProcessSettings.IndirectLightingIntensity,
 		FinalPostProcessSettings.IndirectLightingColor.B * FinalPostProcessSettings.IndirectLightingIntensity);
@@ -2932,14 +2926,6 @@ FRHIGPUMask FSceneRenderer::ComputeGPUMasks(FRHICommandListImmediate& RHICmdList
 		RenderTargetGPUMask = ViewFamily.RenderTarget->GetGPUMask(RHICmdList);
 	}
 
-	{
-		static auto CVar = IConsoleManager::Get().FindConsoleVariable(TEXT("r.PathTracing.GPUCount"));
-		if (CVar && CVar->GetInt() > 1)
-		{
-			RenderTargetGPUMask = FRHIGPUMask::All(); // Broadcast to all GPUs 
-		}
-	}
-
 	// First check whether we are in multi-GPU and if fork and join cross-gpu transfers are enabled.
 	// Otherwise fallback on rendering the whole view family on each relevant GPU using broadcast logic.
 	if (GNumExplicitGPUsForRendering > 1 && CVarEnableMultiGPUForkAndJoin.GetValueOnAnyThread() != 0)
