Path Tracer: fix warnings from bad ShaderType

diff --git a/Engine/Source/Runtime/Renderer/Private/PathTracing.cpp b/Engine/Source/Runtime/Renderer/Private/PathTracing.cpp
--- a/Engine/Source/Runtime/Renderer/Private/PathTracing.cpp
+++ b/Engine/Source/Runtime/Renderer/Private/PathTracing.cpp
@@ -758,8 +758,8 @@ class FPathTracingSwizzleScanlinesCS : public FGlobalShader
 		SHADER_PARAMETER(FIntPoint, DispatchDim)
 		SHADER_PARAMETER(FIntPoint, TileSize)
 		SHADER_PARAMETER(int32, ScanlineStride)
-		SHADER_PARAMETER_RDG_TEXTURE_SRV(RWTexture2DArray, InputTexture)
-		SHADER_PARAMETER_RDG_TEXTURE_UAV(RWTexture2DArray, OutputTexture)
+		SHADER_PARAMETER_RDG_TEXTURE_SRV(Texture2D, InputTexture)
+		SHADER_PARAMETER_RDG_TEXTURE_UAV(RWTexture2D, OutputTexture)
 	END_SHADER_PARAMETER_STRUCT()
 };
 IMPLEMENT_SHADER_TYPE(, FPathTracingSwizzleScanlinesCS, TEXT("/Engine/Private/PathTracing/PathTracingSwizzleScanlines.usf"), TEXT("PathTracingSwizzleScanlinesCS"), SF_Compute);
