Add per-view marker for IndirectLightRendering.

diff --git a/Engine/Source/Runtime/Renderer/Private/IndirectLightRendering.cpp b/Engine/Source/Runtime/Renderer/Private/IndirectLightRendering.cpp
--- a/Engine/Source/Runtime/Renderer/Private/IndirectLightRendering.cpp
+++ b/Engine/Source/Runtime/Renderer/Private/IndirectLightRendering.cpp
@@ -796,10 +796,13 @@ void FDeferredShadingSceneRenderer::RenderDiffuseIndirectAndAmbientOcclusion(
 
 	const FRDGSystemTextures& SystemTextures = FRDGSystemTextures::Get(GraphBuilder);
 
+	uint32 ViewIndex = 0;
 	for (FViewInfo& View : Views)
 	{
 		RDG_GPU_MASK_SCOPE(GraphBuilder, View.GPUMask);
 
+		RDG_EVENT_SCOPE_CONDITIONAL(GraphBuilder, Views.Num() > 1, "View%d", ViewIndex++);
+
 		const FPerViewPipelineState& ViewPipelineState = GetViewPipelineState(View);
 
 		int32 DenoiseMode = CVarDiffuseIndirectDenoiser.GetValueOnRenderThread();
