Adding back code that was removed from CL 19617353, probably due to bad merge confusion.
(code from DFShadow on water and StartDistance on VFog)

diff --git a/Engine/Source/Runtime/Renderer/Private/SceneRendering.cpp b/Engine/Source/Runtime/Renderer/Private/SceneRendering.cpp
--- a/Engine/Source/Runtime/Renderer/Private/SceneRendering.cpp
+++ b/Engine/Source/Runtime/Renderer/Private/SceneRendering.cpp
@@ -910,6 +910,7 @@ void FViewInfo::Init()
 	SinCosInscatteringColorCubemapRotation = FVector2f::ZeroVector;
 	FogInscatteringColorCubemap = nullptr;
 	FogInscatteringTextureParameters = FVector::ZeroVector;
+	VolumetricFogNearFadeInDistanceInv = 100000000.0f;
 
 	SkyAtmosphereCameraAerialPerspectiveVolume = nullptr;
 	SkyAtmosphereUniformShaderParameters = nullptr;
@@ -1299,6 +1300,9 @@ void FViewInfo::SetupUniformBufferParameters(
 		Scene = Family->Scene->GetRenderScene();
 	}
 
+	ERHIFeatureLevel::Type RHIFeatureLevel = Scene == nullptr ? GMaxRHIFeatureLevel : Scene->GetFeatureLevel();
+	EShaderPlatform ShaderPlatform = GShaderPlatformForFeatureLevel[RHIFeatureLevel];
+
 	const FVector DefaultSunDirection(0.0f, 0.0f, 1.0f); // Up vector so that the AtmosphericLightVector node always output a valid direction.
 	auto ClearAtmosphereLightData = [&](uint32 Index)
 	{
@@ -1348,6 +1352,11 @@ void FViewInfo::SetupUniformBufferParameters(
 
 	ViewUniformShaderParameters.BufferToSceneTextureScale = FVector2f(1.0f, 1.0f);
 
+	{
+		extern bool IsWaterDistanceFieldShadowEnabled_Runtime(const FStaticShaderPlatform Platform);
+		ViewUniformShaderParameters.SeparateWaterMainDirLightLuminance = IsWaterDistanceFieldShadowEnabled_Runtime(ShaderPlatform) ? 1.0f : 0.0f;
+	}
+
 	FRHITexture* TransmittanceLutTextureFound = nullptr;
 	FRHITexture* SkyViewLutTextureFound = nullptr;
 	FRHITexture* CameraAerialPerspectiveVolumeFound = nullptr;
@@ -1677,9 +1686,6 @@ void FViewInfo::SetupUniformBufferParameters(
 
 	ViewUniformShaderParameters.DecalDepthBias = CVarDecalDepthBias.GetValueOnRenderThread();
 
-	ERHIFeatureLevel::Type RHIFeatureLevel = Scene == nullptr ? GMaxRHIFeatureLevel : Scene->GetFeatureLevel();
-	EShaderPlatform ShaderPlatform = GShaderPlatformForFeatureLevel[RHIFeatureLevel];
-
 	ViewUniformShaderParameters.IndirectLightingColorScale = FVector3f(FinalPostProcessSettings.IndirectLightingColor.R * FinalPostProcessSettings.IndirectLightingIntensity,
 		FinalPostProcessSettings.IndirectLightingColor.G * FinalPostProcessSettings.IndirectLightingIntensity,
 		FinalPostProcessSettings.IndirectLightingColor.B * FinalPostProcessSettings.IndirectLightingIntensity);
