Fixed crash in SingleLayerWaterRendering where SelectedForwardDirectionalLightProxy is null due to a debug mode short circuiting some initialization code

diff --git a/Engine/Source/Runtime/Renderer/Private/SingleLayerWaterRendering.cpp b/Engine/Source/Runtime/Renderer/Private/SingleLayerWaterRendering.cpp
--- a/Engine/Source/Runtime/Renderer/Private/SingleLayerWaterRendering.cpp
+++ b/Engine/Source/Runtime/Renderer/Private/SingleLayerWaterRendering.cpp
@@ -471,18 +471,22 @@ void FDeferredShadingSceneRenderer::RenderSingleLayerWaterReflections(
 
 		const FPerViewPipelineState& ViewPipelineState = GetViewPipelineState(View);
 
+		FProjectedShadowInfo* DistanceFieldShadowInfo = nullptr;
+
 		// Try to find the ProjectedShadowInfo corresponding to ray trace shadow info for the main directional light.
 		const FLightSceneProxy* SelectedForwardDirectionalLightProxy = View.ForwardLightingResources.SelectedForwardDirectionalLightProxy;
-		FLightSceneInfo* LightSceneInfo = SelectedForwardDirectionalLightProxy->GetLightSceneInfo();
-		FVisibleLightInfo& VisibleLightViewInfo = VisibleLightInfos[LightSceneInfo->Id];
-		FProjectedShadowInfo* DistanceFieldShadowInfo = nullptr;
-		for (int32 ShadowIndex = 0; ShadowIndex < VisibleLightViewInfo.ShadowsToProject.Num(); ShadowIndex++)
+		if (SelectedForwardDirectionalLightProxy)
 		{
-			FProjectedShadowInfo* ProjectedShadowInfo = VisibleLightViewInfo.ShadowsToProject[ShadowIndex];
-
-			if (ProjectedShadowInfo->bRayTracedDistanceField)
+			FLightSceneInfo* LightSceneInfo = SelectedForwardDirectionalLightProxy->GetLightSceneInfo();
+			FVisibleLightInfo& VisibleLightViewInfo = VisibleLightInfos[LightSceneInfo->Id];
+			
+			for (int32 ShadowIndex = 0; ShadowIndex < VisibleLightViewInfo.ShadowsToProject.Num(); ShadowIndex++)
 			{
-				DistanceFieldShadowInfo = ProjectedShadowInfo;
+				FProjectedShadowInfo* ProjectedShadowInfo = VisibleLightViewInfo.ShadowsToProject[ShadowIndex];
+				if (ProjectedShadowInfo->bRayTracedDistanceField)
+				{
+					DistanceFieldShadowInfo = ProjectedShadowInfo;
+				}
 			}
 		}
 
