PathTracer: clamp fog falloff parameters to avoid singularity at 0

diff --git a/Engine/Source/Runtime/Renderer/Private/PathTracing.cpp b/Engine/Source/Runtime/Renderer/Private/PathTracing.cpp
--- a/Engine/Source/Runtime/Renderer/Private/PathTracing.cpp
+++ b/Engine/Source/Runtime/Renderer/Private/PathTracing.cpp
@@ -614,8 +614,8 @@ static FPathTracingFogParameters PrepareFogParameters(const FViewInfo& View, con
 	Parameters.FogDensity.Y = FogInfo.FogData[1].Density * FogInfo.VolumetricFogExtinctionScale;
 	Parameters.FogHeight.X = FogInfo.FogData[0].Height + PreViewTranslation.Z;
 	Parameters.FogHeight.Y = FogInfo.FogData[1].Height + PreViewTranslation.Z;
-	Parameters.FogFalloff.X = FogInfo.FogData[0].HeightFalloff;
-	Parameters.FogFalloff.Y = FogInfo.FogData[1].HeightFalloff;
+	Parameters.FogFalloff.X = FMath::Max(FogInfo.FogData[0].HeightFalloff, 0.001f);
+	Parameters.FogFalloff.Y = FMath::Max(FogInfo.FogData[1].HeightFalloff, 0.001f);
 	Parameters.FogAlbedo = FogInfo.VolumetricFogAlbedo;
 	Parameters.FogPhaseG = FogInfo.VolumetricFogScatteringDistribution;
 
