Use a precise pixel format for SceneDepthAux when capturing scene depth on mobile

diff --git a/Engine/Source/Runtime/Renderer/Private/MobileShadingRenderer.cpp b/Engine/Source/Runtime/Renderer/Private/MobileShadingRenderer.cpp
--- a/Engine/Source/Runtime/Renderer/Private/MobileShadingRenderer.cpp
+++ b/Engine/Source/Runtime/Renderer/Private/MobileShadingRenderer.cpp
@@ -495,6 +495,14 @@ void FMobileSceneRenderer::InitViews(FRDGBuilder& GraphBuilder, FSceneTexturesCo
 	{
 		SceneTexturesConfig.MobilePixelProjectedReflectionExtent = FIntPoint::ZeroValue;
 	}
+
+	// When we capturing scene depth, use a more precise format for SceneDepthAux as it will be used as a source DepthTexture
+	if (ViewFamily.SceneCaptureSource == SCS_SceneColorSceneDepth ||
+		ViewFamily.SceneCaptureSource == SCS_SceneDepth || 
+		ViewFamily.SceneCaptureSource == SCS_DeviceDepth)
+	{
+		SceneTexturesConfig.bPreciseDepthAux = true;
+	}
 	
 	// Finalize and set the scene textures config.
 	FSceneTexturesConfig::Set(SceneTexturesConfig);
diff --git a/Engine/Source/Runtime/Renderer/Private/SceneTextures.cpp b/Engine/Source/Runtime/Renderer/Private/SceneTextures.cpp
--- a/Engine/Source/Runtime/Renderer/Private/SceneTextures.cpp
+++ b/Engine/Source/Runtime/Renderer/Private/SceneTextures.cpp
@@ -216,9 +216,9 @@ static EPixelFormat GetSceneColorFormat(const FSceneViewFamily& ViewFamily)
 	return Format;
 }
 
-inline EPixelFormat GetMobileSceneDepthAuxPixelFormat(EShaderPlatform ShaderPlatform)
+inline EPixelFormat GetMobileSceneDepthAuxPixelFormat(EShaderPlatform ShaderPlatform, bool bPreciseFormat)
 {
-	if (IsMobileDeferredShadingEnabled(ShaderPlatform))
+	if (IsMobileDeferredShadingEnabled(ShaderPlatform) || bPreciseFormat)
 	{
 		return PF_R32_FLOAT;
 	}
@@ -810,9 +810,10 @@ FSceneTextures& FSceneTextures::Create(FRDGBuilder& GraphBuilder, const FSceneTe
 		const float FarDepth = (float)ERHIZBuffer::FarPlane;
 		const FLinearColor FarDepthColor(FarDepth, FarDepth, FarDepth, FarDepth);
 		ETextureCreateFlags FlagsToAdd = IsMobileDeferredShadingEnabled(Config.ShaderPlatform)? TexCreate_Memoryless : TexCreate_None;
+		EPixelFormat DepthAuxFormat = GetMobileSceneDepthAuxPixelFormat(Config.ShaderPlatform, Config.bPreciseDepthAux);
 		FRDGTextureDesc Desc = Config.bRequireMultiView ? 
-			FRDGTextureDesc::Create2DArray(Config.Extent, GetMobileSceneDepthAuxPixelFormat(Config.ShaderPlatform), FClearValueBinding(FarDepthColor), TexCreate_RenderTargetable | TexCreate_ShaderResource | TexCreate_InputAttachmentRead | FlagsToAdd, 2) : 
-			FRDGTextureDesc::Create2D(Config.Extent, GetMobileSceneDepthAuxPixelFormat(Config.ShaderPlatform), FClearValueBinding(FarDepthColor), TexCreate_RenderTargetable | TexCreate_ShaderResource | TexCreate_InputAttachmentRead| FlagsToAdd);
+			FRDGTextureDesc::Create2DArray(Config.Extent, DepthAuxFormat, FClearValueBinding(FarDepthColor), TexCreate_RenderTargetable | TexCreate_ShaderResource | TexCreate_InputAttachmentRead | FlagsToAdd, 2) :
+			FRDGTextureDesc::Create2D(Config.Extent, DepthAuxFormat, FClearValueBinding(FarDepthColor), TexCreate_RenderTargetable | TexCreate_ShaderResource | TexCreate_InputAttachmentRead| FlagsToAdd);
 		Desc.NumSamples = Config.NumSamples;
 		SceneTextures.DepthAux = GraphBuilder.CreateTexture(Desc, TEXT("SceneDepthAux"));
 		
diff --git a/Engine/Source/Runtime/Renderer/Private/SceneTextures.h b/Engine/Source/Runtime/Renderer/Private/SceneTextures.h
--- a/Engine/Source/Runtime/Renderer/Private/SceneTextures.h
+++ b/Engine/Source/Runtime/Renderer/Private/SceneTextures.h
@@ -45,6 +45,7 @@ struct RENDERER_API FSceneTexturesConfig
 		: bRequireMultiView{}
 		, bIsUsingGBuffers{}
 		, bKeepDepthContent{1}
+		, bPreciseDepthAux{}
 	{}
 
 	// Extractions to queue for after execution of the render graph.
@@ -95,6 +96,9 @@ struct RENDERER_API FSceneTexturesConfig
 
 	// (Mobile) True if the platform should write depth content back to memory.
 	uint32 bKeepDepthContent : 1;
+	
+	// (Mobile) True if SceneDepthAux should use a precise pixel format
+	uint32 bPreciseDepthAux : 1;
 
 private:
 	static FSceneTexturesConfig GlobalInstance;
