Fixed: Bright sky reflection on mobile deferred

diff --git a/Engine/Source/Runtime/Renderer/Private/SceneRendering.cpp b/Engine/Source/Runtime/Renderer/Private/SceneRendering.cpp
--- a/Engine/Source/Runtime/Renderer/Private/SceneRendering.cpp
+++ b/Engine/Source/Runtime/Renderer/Private/SceneRendering.cpp
@@ -1762,6 +1762,7 @@ void FViewInfo::SetupUniformBufferParameters(
 		{
 			const FSHVectorRGB3& SkyIrradiance = Scene->SkyLight->IrradianceEnvironmentMap;
 			SetupSkyIrradianceEnvironmentMapConstantsFromSkyIrradiance((FVector4f*)&ViewUniformShaderParameters.MobileSkyIrradianceEnvironmentMap, SkyIrradiance);
+			ViewUniformShaderParameters.MobileSkyIrradianceEnvironmentMap[7].X = Scene->SkyLight->AverageBrightness;
 		}
 		else
 		{
