SLW - DistanceFieldShadow - tiled upsampling pass when water provide data and platform supports it.

diff --git a/Engine/Shaders/Private/DistanceFieldShadowing.usf b/Engine/Shaders/Private/DistanceFieldShadowing.usf
--- a/Engine/Shaders/Private/DistanceFieldShadowing.usf
+++ b/Engine/Shaders/Private/DistanceFieldShadowing.usf
@@ -774,13 +774,15 @@ float NearFadePlaneOffset;
 float InvNearFadePlaneLength;
 
 void DistanceFieldShadowingUpsamplePS(
-	in float4 UVAndScreenPos : TEXCOORD0,
 	in float4 SVPos : SV_POSITION,
 	out float4 OutColor : SV_Target0)
 {
+	float2 BufferUV = SvPositionToBufferUV(SVPos);
+
 	// Distance field shadowing was computed at 0,0 regardless of viewrect min
-	float2 DistanceFieldUVs = UVAndScreenPos.xy - ScissorRectMinAndSize.xy * View.BufferSizeAndInvSize.zw;
-	float SceneDepth = CalcSceneDepth(UVAndScreenPos.xy);
+	float2 DistanceFieldUVs = BufferUV - ScissorRectMinAndSize.xy * View.BufferSizeAndInvSize.zw;
+
+	float SceneDepth = CalcSceneDepth(BufferUV);
 
 #define BILATERAL_UPSAMPLE 1
 #if BILATERAL_UPSAMPLE && UPSAMPLE_REQUIRED
@@ -826,4 +828,26 @@ void DistanceFieldShadowingUpsamplePS(
 	Output = lerp(1, Output, NearBlendFactor);
 
 	OutColor = EncodeLightAttenuation(half4(Output, Output, Output, Output));
-}
\ No newline at end of file
+}
+
+#ifdef SHADOW_TILE_VS
+
+Buffer<uint> TileListData;
+
+void ShadowTileVS(
+	in uint InstanceId : SV_InstanceID,
+	in uint VertexId : SV_VertexID,
+	out float4 Position : SV_POSITION)
+{
+	uint TileData = TileListData[InstanceId.x];
+	uint2 TileOrigin = uint2((TileData & 0x0000FFFF) * WORK_TILE_SIZE, (TileData >> 16) * WORK_TILE_SIZE);
+
+	uint2 TileVertex = TileOrigin;
+	TileVertex.x += VertexId == 1 || VertexId == 2 || VertexId == 4 ? WORK_TILE_SIZE : 0;
+	TileVertex.y += VertexId == 2 || VertexId == 4 || VertexId == 5 ? WORK_TILE_SIZE : 0;
+
+	// View port is set on the view rect. So no offset are needed.
+	Position = float4(float2(TileVertex) * View.ViewSizeAndInvSize.zw * float2(2.0f, -2.0f) + float2(-1.0, 1.0f), 0.5f, 1.0f);
+}
+
+#endif // SHADOW_TILE_VS
diff --git a/Engine/Source/Runtime/Renderer/Private/DistanceFieldShadowing.cpp b/Engine/Source/Runtime/Renderer/Private/DistanceFieldShadowing.cpp
--- a/Engine/Source/Runtime/Renderer/Private/DistanceFieldShadowing.cpp
+++ b/Engine/Source/Runtime/Renderer/Private/DistanceFieldShadowing.cpp
@@ -28,6 +28,7 @@
 #include "PipelineStateCache.h"
 #include "ClearQuad.h"
 #include "Strata/Strata.h"
+#include "PixelShaderUtils.h"
 
 int32 GDistanceFieldShadowing = 1;
 FAutoConsoleVariableRef CVarDistanceFieldShadowing(
@@ -360,6 +361,51 @@ class FDistanceFieldShadowingUpsamplePS : public FGlobalShader
 
 IMPLEMENT_GLOBAL_SHADER(FDistanceFieldShadowingUpsamplePS, "/Engine/Private/DistanceFieldShadowing.usf", "DistanceFieldShadowingUpsamplePS", SF_Pixel);
 
+
+bool UseShadowIndirectDraw(EShaderPlatform ShaderPlatform)
+{
+	return IsFeatureLevelSupported(ShaderPlatform, ERHIFeatureLevel::SM5)
+		&& !IsVulkanMobilePlatform(ShaderPlatform)
+		&& FDataDrivenShaderPlatformInfo::GetSupportsManualVertexFetch(ShaderPlatform);
+}
+
+class FShadowTileVS : public FGlobalShader
+{
+	DECLARE_GLOBAL_SHADER(FShadowTileVS);
+	SHADER_USE_PARAMETER_STRUCT(FShadowTileVS, FGlobalShader);
+
+	using FPermutationDomain = TShaderPermutationDomain<>;
+
+	BEGIN_SHADER_PARAMETER_STRUCT(FParameters, )
+		SHADER_PARAMETER_STRUCT_REF(FViewUniformShaderParameters, ViewUniformBuffer)
+		SHADER_PARAMETER_RDG_BUFFER_SRV(Buffer<uint>, TileListData)
+	END_SHADER_PARAMETER_STRUCT()
+
+	static int32 GetTileSize()
+	{
+		return 8;
+	}
+
+	static FPermutationDomain RemapPermutation(FPermutationDomain PermutationVector)
+	{
+		return PermutationVector;
+	}
+
+	static bool ShouldCompilePermutation(const FGlobalShaderPermutationParameters& Parameters)
+	{
+		return UseShadowIndirectDraw(Parameters.Platform) && DoesPlatformSupportDistanceFieldShadowing(Parameters.Platform);
+	}
+
+	static void ModifyCompilationEnvironment(const FGlobalShaderPermutationParameters& Parameters, FShaderCompilerEnvironment& OutEnvironment)
+	{
+		OutEnvironment.SetDefine(TEXT("SHADOW_TILE_VS"), 1);
+		OutEnvironment.SetDefine(TEXT("WORK_TILE_SIZE"), GetTileSize());
+		FGlobalShader::ModifyCompilationEnvironment(Parameters, OutEnvironment);
+	}
+};
+
+IMPLEMENT_GLOBAL_SHADER(FShadowTileVS, "/Engine/Private/DistanceFieldShadowing.usf", "ShadowTileVS", SF_Vertex);
+
 const uint32 ComputeCulledObjectStartOffsetGroupSize = 8;
 
 /**  */
@@ -963,6 +1009,8 @@ void FProjectedShadowInfo::BeginRenderRayTracedDistanceFieldProjection(
 
 BEGIN_SHADER_PARAMETER_STRUCT(FDistanceFieldShadowingUpsample, )
 	SHADER_PARAMETER_STRUCT_INCLUDE(FDistanceFieldShadowingUpsamplePS::FParameters, PS)
+	SHADER_PARAMETER_STRUCT_INCLUDE(FShadowTileVS::FParameters, VS)
+	RDG_BUFFER_ACCESS(IndirectDrawParameter, ERHIAccess::IndirectArgs)
 	RENDER_TARGET_BINDING_SLOTS()
 END_SHADER_PARAMETER_STRUCT()
 
@@ -973,9 +1021,11 @@ void FProjectedShadowInfo::RenderRayTracedDistanceFieldProjection(
 	const FViewInfo& View,
 	FIntRect ScissorRect,
 	bool bProjectingForForwardShading,
-	bool bForceRGBModulation)
+	bool bForceRGBModulation,
+	FTiledShadowRendering* TiledShadowRendering)
 {
 	check(ScissorRect.Area() > 0);
+	const bool bRunTiled = UseShadowIndirectDraw(View.GetShaderPlatform()) && TiledShadowRendering != nullptr;
 
 	BeginRenderRayTracedDistanceFieldProjection(GraphBuilder, SceneTextures, View); 
 
@@ -985,7 +1035,7 @@ void FProjectedShadowInfo::RenderRayTracedDistanceFieldProjection(
 		PassParameters->RenderTargets[0] = FRenderTargetBinding(ScreenShadowMaskTexture, ERenderTargetLoadAction::ELoad);
 		PassParameters->RenderTargets.DepthStencil = FDepthStencilBinding(SceneTextures.Depth.Target, ERenderTargetLoadAction::ELoad, ERenderTargetLoadAction::ELoad, FExclusiveDepthStencil::DepthRead_StencilRead);
 		
-		PassParameters->PS.View = View.ViewUniformBuffer;
+		PassParameters->PS.View = GetShaderBinding(View.ViewUniformBuffer);
 		PassParameters->PS.SceneTextures = SceneTextures.GetSceneTextureShaderParameters(View.GetFeatureLevel());
 		PassParameters->PS.ShadowFactorsTexture = RayTracedShadowsTexture;
 		PassParameters->PS.ShadowFactorsSampler = TStaticSamplerState<SF_Bilinear>::GetRHI();
@@ -1013,6 +1063,13 @@ void FProjectedShadowInfo::RenderRayTracedDistanceFieldProjection(
 			PassParameters->PS.InvNearFadePlaneLength = 1.0f;
 		}
 
+		if (bRunTiled)
+		{
+			PassParameters->IndirectDrawParameter = TiledShadowRendering->DrawIndirectParametersBuffer;
+			PassParameters->VS.ViewUniformBuffer = GetShaderBinding(View.ViewUniformBuffer);
+			PassParameters->VS.TileListData = TiledShadowRendering->TileListStructureBufferSRV;
+		}
+
 		FDistanceFieldShadowingUpsamplePS::FPermutationDomain PermutationVector;
 		PermutationVector.Set< FDistanceFieldShadowingUpsamplePS::FUpsample >(GFullResolutionDFShadowing == 0);
 		auto PixelShader = View.ShaderMap->GetShader< FDistanceFieldShadowingUpsamplePS >(PermutationVector);
@@ -1021,58 +1078,98 @@ void FProjectedShadowInfo::RenderRayTracedDistanceFieldProjection(
 
 		ClearUnusedGraphResources(PixelShader, &PassParameters->PS);
 
-		GraphBuilder.AddPass(
-			RDG_EVENT_NAME("Upsample"),
-			PassParameters,
-			ERDGPassFlags::Raster,
-			[this, &View, PixelShader, ScissorRect, bProjectingForForwardShading, PassParameters, bForceRGBModulation](FRHICommandList& RHICmdList)
+		if (bRunTiled)
 		{
-			FGraphicsPipelineStateInitializer GraphicsPSOInit;
-			RHICmdList.ApplyCachedRenderTargets(GraphicsPSOInit);
-
-			RHICmdList.SetViewport(ScissorRect.Min.X, ScissorRect.Min.Y, 0.0f, ScissorRect.Max.X, ScissorRect.Max.Y, 1.0f);
-			RHICmdList.SetScissorRect(true, ScissorRect.Min.X, ScissorRect.Min.Y, ScissorRect.Max.X, ScissorRect.Max.Y);
+			check(TiledShadowRendering->TileSize == FShadowTileVS::GetTileSize());
 
-			GraphicsPSOInit.RasterizerState = TStaticRasterizerState<FM_Solid, CM_None>::GetRHI();
-			GraphicsPSOInit.DepthStencilState = TStaticDepthStencilState<false, CF_Always>::GetRHI();
+			auto VertexShader = View.ShaderMap->GetShader<FShadowTileVS>();
+			ClearUnusedGraphResources(VertexShader, &PassParameters->VS);
 
-			if (bForceRGBModulation)
-			{
-				// This has the shadow contribution modulate all the channels, e.g. used for water rendering to apply distance field shadow on the main light RGB luminance for the updated depth buffer with water in it.
-				GraphicsPSOInit.BlendState = TStaticBlendState<CW_RGBA, BO_Add, BF_Zero, BF_SourceColor, BO_Add, BF_Zero, BF_One>::GetRHI();;
-			}
-			else
+			GraphBuilder.AddPass(
+				RDG_EVENT_NAME("TiledUpsample"),
+				PassParameters,
+				ERDGPassFlags::Raster,
+				[this, &View, VertexShader, PixelShader, ScissorRect, bProjectingForForwardShading, PassParameters, bForceRGBModulation](FRHICommandList& RHICmdList)
 			{
-				GraphicsPSOInit.BlendState = GetBlendStateForProjection(bProjectingForForwardShading, false);
-			}
+				RHICmdList.SetViewport(ScissorRect.Min.X, ScissorRect.Min.Y, 0.0f, ScissorRect.Max.X, ScissorRect.Max.Y, 1.0f);
+				RHICmdList.SetScissorRect(true, ScissorRect.Min.X, ScissorRect.Min.Y, ScissorRect.Max.X, ScissorRect.Max.Y);
 
-			TShaderMapRef<FPostProcessVS> VertexShader(View.ShaderMap);
-			GraphicsPSOInit.BoundShaderState.VertexDeclarationRHI = GFilterVertexDeclaration.VertexDeclarationRHI;
-			GraphicsPSOInit.BoundShaderState.VertexShaderRHI = VertexShader.GetVertexShader();
-			GraphicsPSOInit.PrimitiveType = PT_TriangleList;
-			GraphicsPSOInit.bDepthBounds = bDirectionalLight;
-			GraphicsPSOInit.BoundShaderState.PixelShaderRHI = PixelShader.GetPixelShader();
-			SetGraphicsPipelineState(RHICmdList, GraphicsPSOInit, 0);
+				FGraphicsPipelineStateInitializer GraphicsPSOInit;
+				RHICmdList.ApplyCachedRenderTargets(GraphicsPSOInit);
+				GraphicsPSOInit.RasterizerState = TStaticRasterizerState<FM_Solid, CM_None>::GetRHI();
+				GraphicsPSOInit.DepthStencilState = TStaticDepthStencilState<false, CF_Always>::GetRHI();
 
-			SetShaderParameters(RHICmdList, PixelShader, PixelShader.GetPixelShader(), PassParameters->PS);
+				if (bForceRGBModulation)
+				{
+					// This has the shadow contribution modulate all the channels, e.g. used for water rendering to apply distance field shadow on the main light RGB luminance for the updated depth buffer with water in it.
+					GraphicsPSOInit.BlendState = TStaticBlendState<CW_RGBA, BO_Add, BF_Zero, BF_SourceColor, BO_Add, BF_Zero, BF_One>::GetRHI();;
+				}
+				else
+				{
+					GraphicsPSOInit.BlendState = GetBlendStateForProjection(bProjectingForForwardShading, false);
+				}
+				GraphicsPSOInit.bDepthBounds = bDirectionalLight;
+
+				GraphicsPSOInit.PrimitiveType = GRHISupportsRectTopology ? PT_RectList : PT_TriangleList;
+				GraphicsPSOInit.BoundShaderState.VertexDeclarationRHI = GFilterVertexDeclaration.VertexDeclarationRHI;
+				GraphicsPSOInit.BoundShaderState.PixelShaderRHI = PixelShader.GetPixelShader();
+				GraphicsPSOInit.BoundShaderState.VertexShaderRHI = VertexShader.GetVertexShader();
+				SetGraphicsPipelineState(RHICmdList, GraphicsPSOInit, 0);
+
+				SetShaderParameters(RHICmdList, PixelShader, PixelShader.GetPixelShader(), PassParameters->PS);
+				SetShaderParameters(RHICmdList, VertexShader, VertexShader.GetPixelShader(), PassParameters->VS);
+
+				//@todo - depth bounds test for local lights
+				if (bDirectionalLight)
+				{
+					SetDepthBoundsTest(RHICmdList, CascadeSettings.SplitNear - CascadeSettings.SplitNearFadeRegion, CascadeSettings.SplitFar, View.ViewMatrices.GetProjectionMatrix());
+				}
+
+				PassParameters->IndirectDrawParameter->MarkResourceAsUsed();
+				RHICmdList.DrawPrimitiveIndirect(PassParameters->IndirectDrawParameter->GetIndirectRHICallBuffer(), 0);
 
-			//@todo - depth bounds test for local lights
-			if (bDirectionalLight)
+				RHICmdList.SetScissorRect(false, 0, 0, 0, 0);
+			});
+		}
+		else
+		{
+			GraphBuilder.AddPass(
+				RDG_EVENT_NAME("Upsample"),
+				PassParameters,
+				ERDGPassFlags::Raster,
+				[this, &View, PixelShader, ScissorRect, bProjectingForForwardShading, PassParameters, bForceRGBModulation](FRHICommandList& RHICmdList)
 			{
-				SetDepthBoundsTest(RHICmdList, CascadeSettings.SplitNear - CascadeSettings.SplitNearFadeRegion, CascadeSettings.SplitFar, View.ViewMatrices.GetProjectionMatrix());
-			}
+				RHICmdList.SetViewport(ScissorRect.Min.X, ScissorRect.Min.Y, 0.0f, ScissorRect.Max.X, ScissorRect.Max.Y, 1.0f);
+				RHICmdList.SetScissorRect(true, ScissorRect.Min.X, ScissorRect.Min.Y, ScissorRect.Max.X, ScissorRect.Max.Y);
 
-			DrawRectangle(
-				RHICmdList,
-				0, 0,
-				ScissorRect.Width(), ScissorRect.Height(),
-				ScissorRect.Min.X / GetDFShadowDownsampleFactor(), ScissorRect.Min.Y / GetDFShadowDownsampleFactor(),
-				ScissorRect.Width() / GetDFShadowDownsampleFactor(), ScissorRect.Height() / GetDFShadowDownsampleFactor(),
-				FIntPoint(ScissorRect.Width(), ScissorRect.Height()),
-				GetBufferSizeForDFShadows(),
-				VertexShader);
-
-			RHICmdList.SetScissorRect(false, 0, 0, 0, 0);
-		});
+				FGraphicsPipelineStateInitializer GraphicsPSOInit;
+				FPixelShaderUtils::InitFullscreenPipelineState(RHICmdList, View.ShaderMap, PixelShader, GraphicsPSOInit);
+
+				if (bForceRGBModulation)
+				{
+					// This has the shadow contribution modulate all the channels, e.g. used for water rendering to apply distance field shadow on the main light RGB luminance for the updated depth buffer with water in it.
+					GraphicsPSOInit.BlendState = TStaticBlendState<CW_RGBA, BO_Add, BF_Zero, BF_SourceColor, BO_Add, BF_Zero, BF_One>::GetRHI();;
+				}
+				else
+				{
+					GraphicsPSOInit.BlendState = GetBlendStateForProjection(bProjectingForForwardShading, false);
+				}
+				GraphicsPSOInit.bDepthBounds = bDirectionalLight;
+
+				GraphicsPSOInit.BoundShaderState.PixelShaderRHI = PixelShader.GetPixelShader();
+				SetGraphicsPipelineState(RHICmdList, GraphicsPSOInit, 0);
+
+				SetShaderParameters(RHICmdList, PixelShader, PixelShader.GetPixelShader(), PassParameters->PS);
+
+				//@todo - depth bounds test for local lights
+				if (bDirectionalLight)
+				{
+					SetDepthBoundsTest(RHICmdList, CascadeSettings.SplitNear - CascadeSettings.SplitNearFadeRegion, CascadeSettings.SplitFar, View.ViewMatrices.GetProjectionMatrix());
+				}
+
+				FPixelShaderUtils::DrawFullscreenTriangle(RHICmdList);
+				RHICmdList.SetScissorRect(false, 0, 0, 0, 0);
+			});
+		}
 	}
 }
diff --git a/Engine/Source/Runtime/Renderer/Private/ShadowRendering.h b/Engine/Source/Runtime/Renderer/Private/ShadowRendering.h
--- a/Engine/Source/Runtime/Renderer/Private/ShadowRendering.h
+++ b/Engine/Source/Runtime/Renderer/Private/ShadowRendering.h
@@ -170,6 +170,13 @@ inline bool IsShadowCacheModeOcclusionQueryable(EShadowDepthCacheMode CacheMode)
 	return CacheMode != SDCM_StaticPrimitivesOnly;
 }
 
+struct FTiledShadowRendering
+{
+	FRDGBufferRef		DrawIndirectParametersBuffer;
+	FRDGBufferSRVRef	TileListStructureBufferSRV;
+	uint32				TileSize;
+};
+
 class FShadowMapRenderTargets
 {
 public:
@@ -580,7 +587,8 @@ public:
 		const FViewInfo& View,
 		FIntRect ScissorRect,
 		bool bProjectingForForwardShading,
-		bool bForceRGBModulation = false);
+		bool bForceRGBModulation = false,
+		FTiledShadowRendering* TiledShadowRendering = nullptr);
 
 	/** Render one pass point light shadow projections. */
 	void RenderOnePassPointLightProjection(
diff --git a/Engine/Source/Runtime/Renderer/Private/SingleLayerWaterRendering.cpp b/Engine/Source/Runtime/Renderer/Private/SingleLayerWaterRendering.cpp
--- a/Engine/Source/Runtime/Renderer/Private/SingleLayerWaterRendering.cpp
+++ b/Engine/Source/Runtime/Renderer/Private/SingleLayerWaterRendering.cpp
@@ -502,19 +502,28 @@ void FDeferredShadingSceneRenderer::RenderSingleLayerWaterReflections(
 				ScissorRect = View.ViewRect;
 			}
 
-			// Reset the cached texture to create a new one mapping to he water depth buffer
+			// Reset the cached texture to create a new one mapping to the water depth buffer
 			DistanceFieldShadowInfo->ResetRayTracedDistanceFieldShadow();
 
+			FTiledShadowRendering TiledShadowRendering;
+			if (bRunTiled)
+			{
+				TiledShadowRendering.DrawIndirectParametersBuffer	= TiledScreenSpaceReflection.DrawIndirectParametersBuffer;
+				TiledShadowRendering.TileListStructureBufferSRV		= TiledScreenSpaceReflection.TileListStructureBufferSRV;
+				TiledShadowRendering.TileSize						= TiledScreenSpaceReflection.TileSize;
+			}
+
 			const bool bProjectingForForwardShading = false;
 			const bool bForceRGBModulation = true;
-			DistanceFieldShadowInfo->RenderRayTracedDistanceFieldProjection(
+			DistanceFieldShadowInfo->RenderRayTracedDistanceFieldProjection( 
 				GraphBuilder,
 				SceneTextures,
 				SceneWithoutWaterTextures.SeparatedMainDirLightTexture,
 				View,
 				ScissorRect,
 				bProjectingForForwardShading,
-				bForceRGBModulation);
+				bForceRGBModulation,
+				bRunTiled ? &TiledShadowRendering : nullptr);
 		}
 
 		if (ViewPipelineState.ReflectionsMethod == EReflectionsMethod::Lumen && CVarWaterSingleLayerLumenReflections.GetValueOnRenderThread() != 0)
