Don't use async compute when calculating DF Shadow for water.

diff --git a/Engine/Source/Runtime/Renderer/Private/DistanceFieldShadowing.cpp b/Engine/Source/Runtime/Renderer/Private/DistanceFieldShadowing.cpp
--- a/Engine/Source/Runtime/Renderer/Private/DistanceFieldShadowing.cpp
+++ b/Engine/Source/Runtime/Renderer/Private/DistanceFieldShadowing.cpp
@@ -741,6 +741,7 @@ bool FSceneRenderer::ShouldPrepareHeightFieldScene() const
 
 void RayTraceShadows(
 	FRDGBuilder& GraphBuilder,
+	bool bAsyncCompute,
 	const FMinimalSceneTextures& SceneTextures,
 	FRDGTextureRef RayTracedShadowsTexture,
 	const FViewInfo& View,
@@ -850,7 +851,7 @@ void RayTraceShadows(
 		FComputeShaderUtils::AddPass(
 			GraphBuilder,
 			RDG_EVENT_NAME("DistanceFieldShadowing %ux%u", GroupSizeX * GDistanceFieldShadowTileSizeX, GroupSizeY * GDistanceFieldShadowTileSizeY),
-			!!GDFShadowAsyncCompute ? ERDGPassFlags::AsyncCompute : ERDGPassFlags::Compute,
+			bAsyncCompute ? ERDGPassFlags::AsyncCompute : ERDGPassFlags::Compute,
 			ComputeShader,
 			PassParameters,
 			FIntVector(GroupSizeX, GroupSizeY, 1));
@@ -859,6 +860,7 @@ void RayTraceShadows(
 
 void FProjectedShadowInfo::BeginRenderRayTracedDistanceFieldProjection(
 	FRDGBuilder& GraphBuilder,
+	bool bAsyncCompute, 
 	const FMinimalSceneTextures& SceneTextures,
 	const FViewInfo& View)
 {
@@ -938,7 +940,7 @@ void FProjectedShadowInfo::BeginRenderRayTracedDistanceFieldProjection(
 				RayTracedShadowsTexture = GraphBuilder.CreateTexture(Desc, TEXT("RayTracedShadows"));
 			}
 
-			RayTraceShadows(GraphBuilder, SceneTextures, RayTracedShadowsTexture, View, Scene->DistanceFieldSceneData, this, DFPT_SignedDistanceField, false, nullptr, ObjectBufferParameters, *SDFShadowViewGPUData.SDFCulledObjectBufferParameters, *SDFShadowViewGPUData.SDFLightTileIntersectionParameters);
+			RayTraceShadows(GraphBuilder, bAsyncCompute, SceneTextures, RayTracedShadowsTexture, View, Scene->DistanceFieldSceneData, this, DFPT_SignedDistanceField, false, nullptr, ObjectBufferParameters, *SDFShadowViewGPUData.SDFCulledObjectBufferParameters, *SDFShadowViewGPUData.SDFLightTileIntersectionParameters);
 		}
 	}
 
@@ -1003,7 +1005,7 @@ void FProjectedShadowInfo::BeginRenderRayTracedDistanceFieldProjection(
 			RayTracedShadowsTexture = GraphBuilder.CreateTexture(Desc, TEXT("RayTracedShadows"));
 		}
 
-		RayTraceShadows(GraphBuilder, SceneTextures, RayTracedShadowsTexture, View, Scene->DistanceFieldSceneData, this, DFPT_HeightField, bHasPrevOutput, PrevOutputTexture, ObjectBufferParameters, *SDFShadowViewGPUData.HeightFieldCulledObjectBufferParameters, *SDFShadowViewGPUData.HeightFieldLightTileIntersectionParameters);
+		RayTraceShadows(GraphBuilder, bAsyncCompute, SceneTextures, RayTracedShadowsTexture, View, Scene->DistanceFieldSceneData, this, DFPT_HeightField, bHasPrevOutput, PrevOutputTexture, ObjectBufferParameters, *SDFShadowViewGPUData.HeightFieldCulledObjectBufferParameters, *SDFShadowViewGPUData.HeightFieldLightTileIntersectionParameters);
 	}
 }
 
@@ -1027,7 +1029,7 @@ void FProjectedShadowInfo::RenderRayTracedDistanceFieldProjection(
 	check(ScissorRect.Area() > 0);
 	const bool bRunTiled = UseShadowIndirectDraw(View.GetShaderPlatform()) && TiledShadowRendering != nullptr;
 
-	BeginRenderRayTracedDistanceFieldProjection(GraphBuilder, SceneTextures, View); 
+	BeginRenderRayTracedDistanceFieldProjection(GraphBuilder, false, SceneTextures, View); 
 
 	if (RayTracedShadowsTexture)
 	{
diff --git a/Engine/Source/Runtime/Renderer/Private/ShadowRendering.cpp b/Engine/Source/Runtime/Renderer/Private/ShadowRendering.cpp
--- a/Engine/Source/Runtime/Renderer/Private/ShadowRendering.cpp
+++ b/Engine/Source/Runtime/Renderer/Private/ShadowRendering.cpp
@@ -2251,7 +2251,7 @@ void FSceneRenderer::BeginAsyncDistanceFieldShadowProjections(FRDGBuilder& Graph
 
 				if (ScissorRect.Area() > 0 && (Views.Num() == 1 || ProjectedShadowInfo->DependentView == &View || !ProjectedShadowInfo->DependentView))
 				{
-					ProjectedShadowInfo->BeginRenderRayTracedDistanceFieldProjection(GraphBuilder, SceneTextures, View);
+					ProjectedShadowInfo->BeginRenderRayTracedDistanceFieldProjection(GraphBuilder, true, SceneTextures, View);
 				}
 			}
 		}
diff --git a/Engine/Source/Runtime/Renderer/Private/ShadowRendering.h b/Engine/Source/Runtime/Renderer/Private/ShadowRendering.h
--- a/Engine/Source/Runtime/Renderer/Private/ShadowRendering.h
+++ b/Engine/Source/Runtime/Renderer/Private/ShadowRendering.h
@@ -576,6 +576,7 @@ public:
 
 	void BeginRenderRayTracedDistanceFieldProjection(
 		FRDGBuilder& GraphBuilder,
+		bool bAsyncCompute,
 		const FMinimalSceneTextures& SceneTextures,
 		const FViewInfo& View);
 
