Fix DF Shadows when using split screen.
- Point/Spot lights use the same FProjectedShadowInfo for all views. RayTracedShadowsTexture was only generated for the first view. Other views also sampled first view data.
- Store RayTracedShadowsTexture per view in DistanceFieldShadowViewGPUData to avoid using the cached result of a different view.

diff --git a/Engine/Source/Runtime/Renderer/Private/DistanceFieldShadowing.cpp b/Engine/Source/Runtime/Renderer/Private/DistanceFieldShadowing.cpp
--- a/Engine/Source/Runtime/Renderer/Private/DistanceFieldShadowing.cpp
+++ b/Engine/Source/Runtime/Renderer/Private/DistanceFieldShadowing.cpp
@@ -858,23 +858,24 @@ void RayTraceShadows(
 	}
 }
 
-void FProjectedShadowInfo::BeginRenderRayTracedDistanceFieldProjection(
+FRDGTextureRef FProjectedShadowInfo::RenderRayTracedDistanceFieldProjection(
 	FRDGBuilder& GraphBuilder,
 	bool bAsyncCompute, 
 	const FMinimalSceneTextures& SceneTextures,
 	const FViewInfo& View)
 {
-	if (RayTracedShadowsTexture)
+	DistanceFieldShadowViewGPUData& SDFShadowViewGPUData = CachedDistanceFieldShadowViewGPUData.FindOrAdd(&View);
+
+	if (SDFShadowViewGPUData.RayTracedShadowsTexture)
 	{
-		return;
+		// Ray traced distance field shadows were already calculated, simply return previous result.
+		return SDFShadowViewGPUData.RayTracedShadowsTexture;
 	}
 
 	const bool bDFShadowSupported = SupportsDistanceFieldShadows(View.GetFeatureLevel(), View.GetShaderPlatform());
 	const bool bHFShadowSupported = SupportsHeightFieldShadows(View.GetFeatureLevel(), View.GetShaderPlatform());
 	const FScene* Scene = (const FScene*)(View.Family->Scene);
 
-	DistanceFieldShadowViewGPUData& SDFShadowViewGPUData = CachedDistanceFieldShadowViewGPUData.FindOrAdd(&View);
-
 	if (bDFShadowSupported && View.Family->EngineShowFlags.RayTracedDistanceFieldShadows)
 	{
 		QUICK_SCOPE_CYCLE_COUNTER(STAT_BeginRenderRayTracedDistanceFieldShadows);
@@ -937,10 +938,10 @@ void FProjectedShadowInfo::BeginRenderRayTracedDistanceFieldProjection(
 				const FIntPoint BufferSize = GetBufferSizeForDFShadows();
 				FRDGTextureDesc Desc(FRDGTextureDesc::Create2D(BufferSize, PF_G16R16F, FClearValueBinding::None, TexCreate_UAV));
 				Desc.Flags |= GFastVRamConfig.DistanceFieldShadows;
-				RayTracedShadowsTexture = GraphBuilder.CreateTexture(Desc, TEXT("RayTracedShadows"));
+				SDFShadowViewGPUData.RayTracedShadowsTexture = GraphBuilder.CreateTexture(Desc, TEXT("RayTracedShadows"));
 			}
 
-			RayTraceShadows(GraphBuilder, bAsyncCompute, SceneTextures, RayTracedShadowsTexture, View, Scene->DistanceFieldSceneData, this, DFPT_SignedDistanceField, false, nullptr, ObjectBufferParameters, *SDFShadowViewGPUData.SDFCulledObjectBufferParameters, *SDFShadowViewGPUData.SDFLightTileIntersectionParameters);
+			RayTraceShadows(GraphBuilder, bAsyncCompute, SceneTextures, SDFShadowViewGPUData.RayTracedShadowsTexture, View, Scene->DistanceFieldSceneData, this, DFPT_SignedDistanceField, false, nullptr, ObjectBufferParameters, *SDFShadowViewGPUData.SDFCulledObjectBufferParameters, *SDFShadowViewGPUData.SDFLightTileIntersectionParameters);
 		}
 	}
 
@@ -987,26 +988,28 @@ void FProjectedShadowInfo::BeginRenderRayTracedDistanceFieldProjection(
 
 		}
 
-		const bool bHasPrevOutput = !!RayTracedShadowsTexture;
+		const bool bHasPrevOutput = !!SDFShadowViewGPUData.RayTracedShadowsTexture;
 
 		FRDGTextureRef PrevOutputTexture = nullptr;
 
 		if (!RHISupports4ComponentUAVReadWrite(View.GetShaderPlatform()))
 		{
-			PrevOutputTexture = RayTracedShadowsTexture;
-			RayTracedShadowsTexture = nullptr;
+			PrevOutputTexture = SDFShadowViewGPUData.RayTracedShadowsTexture;
+			SDFShadowViewGPUData.RayTracedShadowsTexture = nullptr;
 		}
 
-		if (!RayTracedShadowsTexture)
+		if (!SDFShadowViewGPUData.RayTracedShadowsTexture)
 		{
 			const FIntPoint BufferSize = GetBufferSizeForDFShadows();
 			FRDGTextureDesc Desc(FRDGTextureDesc::Create2D(BufferSize, PF_G16R16F, FClearValueBinding::None, TexCreate_UAV));
 			Desc.Flags |= GFastVRamConfig.DistanceFieldShadows;
-			RayTracedShadowsTexture = GraphBuilder.CreateTexture(Desc, TEXT("RayTracedShadows"));
+			SDFShadowViewGPUData.RayTracedShadowsTexture = GraphBuilder.CreateTexture(Desc, TEXT("RayTracedShadows"));
 		}
 
-		RayTraceShadows(GraphBuilder, bAsyncCompute, SceneTextures, RayTracedShadowsTexture, View, Scene->DistanceFieldSceneData, this, DFPT_HeightField, bHasPrevOutput, PrevOutputTexture, ObjectBufferParameters, *SDFShadowViewGPUData.HeightFieldCulledObjectBufferParameters, *SDFShadowViewGPUData.HeightFieldLightTileIntersectionParameters);
+		RayTraceShadows(GraphBuilder, bAsyncCompute, SceneTextures, SDFShadowViewGPUData.RayTracedShadowsTexture, View, Scene->DistanceFieldSceneData, this, DFPT_HeightField, bHasPrevOutput, PrevOutputTexture, ObjectBufferParameters, *SDFShadowViewGPUData.HeightFieldCulledObjectBufferParameters, *SDFShadowViewGPUData.HeightFieldLightTileIntersectionParameters);
 	}
+
+	return SDFShadowViewGPUData.RayTracedShadowsTexture;
 }
 
 BEGIN_SHADER_PARAMETER_STRUCT(FDistanceFieldShadowingUpsample, )
@@ -1029,7 +1032,7 @@ void FProjectedShadowInfo::RenderRayTracedDistanceFieldProjection(
 	check(ScissorRect.Area() > 0);
 	const bool bRunTiled = UseShadowIndirectDraw(View.GetShaderPlatform()) && TiledShadowRendering != nullptr;
 
-	BeginRenderRayTracedDistanceFieldProjection(GraphBuilder, false, SceneTextures, View); 
+	FRDGTextureRef RayTracedShadowsTexture = RenderRayTracedDistanceFieldProjection(GraphBuilder, false, SceneTextures, View);
 
 	if (RayTracedShadowsTexture)
 	{
diff --git a/Engine/Source/Runtime/Renderer/Private/ShadowRendering.cpp b/Engine/Source/Runtime/Renderer/Private/ShadowRendering.cpp
--- a/Engine/Source/Runtime/Renderer/Private/ShadowRendering.cpp
+++ b/Engine/Source/Runtime/Renderer/Private/ShadowRendering.cpp
@@ -2251,7 +2251,9 @@ void FSceneRenderer::BeginAsyncDistanceFieldShadowProjections(FRDGBuilder& Graph
 
 				if (ScissorRect.Area() > 0 && (Views.Num() == 1 || ProjectedShadowInfo->DependentView == &View || !ProjectedShadowInfo->DependentView))
 				{
-					ProjectedShadowInfo->BeginRenderRayTracedDistanceFieldProjection(GraphBuilder, true, SceneTextures, View);
+					// Kick off distance field shadow calculation in async compute
+					// Don't need store result reference because it is internally cached by FProjectedShadowInfo
+					ProjectedShadowInfo->RenderRayTracedDistanceFieldProjection(GraphBuilder, true, SceneTextures, View);
 				}
 			}
 		}
diff --git a/Engine/Source/Runtime/Renderer/Private/ShadowRendering.h b/Engine/Source/Runtime/Renderer/Private/ShadowRendering.h
--- a/Engine/Source/Runtime/Renderer/Private/ShadowRendering.h
+++ b/Engine/Source/Runtime/Renderer/Private/ShadowRendering.h
@@ -574,13 +574,28 @@ public:
 		bool bDoParallelDispatch,
 		bool bDoCrossGPUCopy);
 
-	void BeginRenderRayTracedDistanceFieldProjection(
+	/** Reset cached ray traced distance field shadows texture. */
+	void ResetRayTracedDistanceFieldShadow(const FViewInfo* View)
+	{
+		DistanceFieldShadowViewGPUData& SDFShadowViewGPUData = CachedDistanceFieldShadowViewGPUData.FindOrAdd(View);
+		SDFShadowViewGPUData.RayTracedShadowsTexture = nullptr;
+	}
+
+	/**
+	* Render ray traced distance field shadows into a texture.
+	* Output texture is cached per view. This is useful to support async compute.
+	* (ie: kick off distance field shadows early in the frame using async compute, and then combine result into shadow mask texture when necessary)
+	*/
+	FRDGTextureRef RenderRayTracedDistanceFieldProjection(
 		FRDGBuilder& GraphBuilder,
 		bool bAsyncCompute,
 		const FMinimalSceneTextures& SceneTextures,
 		const FViewInfo& View);
 
-	/** Renders ray traced distance field shadows. */
+	/** 
+	* Renders ray traced distance field shadows into an existing texture.
+	* Will use cached results if already calculated earlier in the frame (ie: async compute)
+	*/
 	void RenderRayTracedDistanceFieldProjection(
 		FRDGBuilder& GraphBuilder,
 		const FMinimalSceneTextures& SceneTextures,
@@ -738,11 +753,6 @@ public:
 		return ShadowDepthPass.BuildRenderingCommands(GraphBuilder, GPUScene, InstanceCullingDrawParams);
 	}
 
-	void ResetRayTracedDistanceFieldShadow()
-	{
-		RayTracedShadowsTexture = nullptr;
-	}
-
 private:
 	// 0 if Setup...() wasn't called yet
 	FLightSceneInfo* LightSceneInfo;
@@ -793,18 +803,18 @@ private:
 	float ShaderSlopeDepthBias;
 	float ShaderMaxSlopeDepthBias;
 
-	/** Ray traced DF shadow intermediate output. Populated by BeginRenderRayTracedDistanceFieldProjection and consumed by RenderRayTracedDistanceFieldProjection. */
-	FRDGTextureRef RayTracedShadowsTexture;
-
 	/**  Cached light tile intersection parameters in case we needed to trace a distance field multiple times for a view but for different depth buffer*/
 	struct DistanceFieldShadowViewGPUData 
 	{
+		/** Ray traced DF shadow intermediate output. Populated by BeginRenderRayTracedDistanceFieldProjection and consumed by RenderRayTracedDistanceFieldProjection. */
+		FRDGTextureRef RayTracedShadowsTexture = nullptr;
+
 		FLightTileIntersectionParameters*			SDFLightTileIntersectionParameters = nullptr;
 		FDistanceFieldCulledObjectBufferParameters*	SDFCulledObjectBufferParameters = nullptr;
 		FLightTileIntersectionParameters*			HeightFieldLightTileIntersectionParameters = nullptr;
 		FDistanceFieldCulledObjectBufferParameters*	HeightFieldCulledObjectBufferParameters = nullptr;
 	};
-	TMap<const FViewInfo*, DistanceFieldShadowViewGPUData>	CachedDistanceFieldShadowViewGPUData;
+	TMap<const FViewInfo*, DistanceFieldShadowViewGPUData> CachedDistanceFieldShadowViewGPUData;
 
 	void CopyCachedShadowMap(
 		FRDGBuilder& GraphBuilder,
diff --git a/Engine/Source/Runtime/Renderer/Private/ShadowSetup.cpp b/Engine/Source/Runtime/Renderer/Private/ShadowSetup.cpp
--- a/Engine/Source/Runtime/Renderer/Private/ShadowSetup.cpp
+++ b/Engine/Source/Runtime/Renderer/Private/ShadowSetup.cpp
@@ -717,7 +717,6 @@ FProjectedShadowInfo::FProjectedShadowInfo()
 	, NumSubjectMeshCommandBuildRequestElements(0)
 	, ShaderDepthBias(0.0f)
 	, ShaderSlopeDepthBias(0.0f)
-	, RayTracedShadowsTexture(nullptr)
 {
 }
 
diff --git a/Engine/Source/Runtime/Renderer/Private/SingleLayerWaterRendering.cpp b/Engine/Source/Runtime/Renderer/Private/SingleLayerWaterRendering.cpp
--- a/Engine/Source/Runtime/Renderer/Private/SingleLayerWaterRendering.cpp
+++ b/Engine/Source/Runtime/Renderer/Private/SingleLayerWaterRendering.cpp
@@ -504,7 +504,7 @@ void FDeferredShadingSceneRenderer::RenderSingleLayerWaterReflections(
 			}
 
 			// Reset the cached texture to create a new one mapping to the water depth buffer
-			DistanceFieldShadowInfo->ResetRayTracedDistanceFieldShadow();
+			DistanceFieldShadowInfo->ResetRayTracedDistanceFieldShadow(&View);
 
 			FTiledShadowRendering TiledShadowRendering;
 			if (bRunTiled)
