PR #8520: Fix for Reflection Captures sometimes showing up black in packaged game. (Contributed by [REDACTED_PERSON])

diff --git a/Engine/Source/Runtime/Renderer/Private/ReflectionEnvironment.cpp b/Engine/Source/Runtime/Renderer/Private/ReflectionEnvironment.cpp
--- a/Engine/Source/Runtime/Renderer/Private/ReflectionEnvironment.cpp
+++ b/Engine/Source/Runtime/Renderer/Private/ReflectionEnvironment.cpp
@@ -5,6 +5,7 @@
 =============================================================================*/
 
 #include "ReflectionEnvironment.h"
+#include "Components/ReflectionCaptureComponent.h"
 #include "Stats/Stats.h"
 #include "HAL/IConsoleManager.h"
 #include "RHI.h"
@@ -188,6 +189,139 @@ void FReflectionEnvironmentCubemapArray::ReleaseDynamicRHI()
 	ReleaseCubeArray();
 }
 
+
+const FCaptureComponentSceneState* FReflectionCaptureCache::Find(const FGuid& MapBuildDataId) const
+{
+	const FReflectionCaptureCacheEntry* Entry = CaptureData.Find(MapBuildDataId);
+	if (Entry == nullptr)
+		return nullptr;
+
+	return &(Entry->SceneState);
+}
+
+FCaptureComponentSceneState* FReflectionCaptureCache::Find(const FGuid& MapBuildDataId)
+{
+	FReflectionCaptureCacheEntry* Entry = CaptureData.Find(MapBuildDataId);
+	if (Entry == nullptr)
+		return nullptr;
+
+	return &(Entry->SceneState);
+}
+
+const FCaptureComponentSceneState* FReflectionCaptureCache::Find(const UReflectionCaptureComponent* Component) const
+{
+	if (!Component)	// Intentionally not IsValid(Component), as this often occurs when Component is explicitly PendingKill.
+		return nullptr;
+
+	return Find(Component->MapBuildDataId);
+}
+
+FCaptureComponentSceneState* FReflectionCaptureCache::Find(const UReflectionCaptureComponent* Component)
+{
+	if (!Component)	// Intentionally not IsValid(Component), as this often occurs when Component is explicitly PendingKill.
+		return nullptr;
+
+	return Find(Component->MapBuildDataId);
+}
+
+
+const FCaptureComponentSceneState& FReflectionCaptureCache::FindChecked(const UReflectionCaptureComponent* Component) const
+{
+	const FCaptureComponentSceneState* Found = Find(Component);
+	check(Found);
+
+	return *Found;
+}
+
+FCaptureComponentSceneState& FReflectionCaptureCache::FindChecked(const UReflectionCaptureComponent* Component)
+{
+	FCaptureComponentSceneState* Found = Find(Component);
+	check(Found);
+
+	return *Found;
+}
+
+FCaptureComponentSceneState& FReflectionCaptureCache::Add(const UReflectionCaptureComponent* Component, const FCaptureComponentSceneState& Value)
+{
+	check(IsValid(Component))
+
+	FCaptureComponentSceneState* Existing = AddReference(Component);
+	if (Existing != nullptr)
+	{
+		return *Existing;
+	}
+	else
+	{
+		FReflectionCaptureCacheEntry& Entry = CaptureData.Add(Component->MapBuildDataId, { 1, Value });
+		return Entry.SceneState;
+	}
+}
+
+FCaptureComponentSceneState* FReflectionCaptureCache::AddReference(const UReflectionCaptureComponent* Component)
+{
+	check(IsValid(Component))
+
+	FReflectionCaptureCacheEntry* Found = CaptureData.Find(Component->MapBuildDataId);
+	if (Found == nullptr)
+		return nullptr;
+
+	Found->RefCount += 1;
+	return &(Found->SceneState);
+}
+
+int32 FReflectionCaptureCache::GetKeys(TArray<FGuid>& OutKeys) const
+{
+	return CaptureData.GetKeys(OutKeys);
+}
+
+int32 FReflectionCaptureCache::GetKeys(TSet<FGuid>& OutKeys) const
+{
+	return CaptureData.GetKeys(OutKeys);
+}
+
+void FReflectionCaptureCache::Empty()
+{
+	CaptureData.Empty();
+}
+
+bool FReflectionCaptureCache::Remove(const UReflectionCaptureComponent* Component)
+{
+	if (!Component)	// Intentionally not IsValid(Component), as this often occurs when Component is explicitly PendingKill.
+		return false;
+
+	FReflectionCaptureCacheEntry* Found = CaptureData.Find(Component->MapBuildDataId);
+	if (Found == nullptr)
+		return false;
+
+	if (Found->RefCount > 1)
+	{
+		Found->RefCount -= 1;
+		return false;
+	}
+	else
+	{
+		CaptureData.Remove(Component->MapBuildDataId);
+		return true;
+	}
+}
+
+void FReflectionEnvironmentSceneData::SetGameThreadTrackingData(int32 MaxAllocatedCubemaps, int32 CaptureSize)
+{
+	MaxAllocatedReflectionCubemapsGameThread = MaxAllocatedCubemaps;
+	ReflectionCaptureSizeGameThread = CaptureSize;
+}
+
+bool FReflectionEnvironmentSceneData::DoesAllocatedDataNeedUpdate(int32 DesiredMaxCubemaps, int32 DesiredCaptureSize) const
+{
+	if (DesiredMaxCubemaps != MaxAllocatedReflectionCubemapsGameThread)
+		return true;
+
+	if (DesiredCaptureSize != ReflectionCaptureSizeGameThread)
+		return true;
+
+	return false;
+}
+
 void FReflectionEnvironmentSceneData::ResizeCubemapArrayGPU(uint32 InMaxCubemaps, int32 InCubemapSize)
 {
 	check(IsInRenderingThread());
@@ -221,7 +355,7 @@ void FReflectionEnvironmentSceneData::ResizeCubemapArrayGPU(uint32 InMaxCubemaps
 	CubemapArraySlotsUsed.Init(false, InMaxCubemaps);
 
 	// Spin through the AllocatedReflectionCaptureState map and remap the indices based on the LUT
-	TArray<const UReflectionCaptureComponent*> Components;
+	TArray<FGuid> Components;
 	AllocatedReflectionCaptureState.GetKeys(Components);
 	int32 UsedCubemapCount = 0;
 	for (int32 i=0; i<Components.Num(); i++ )
diff --git a/Engine/Source/Runtime/Renderer/Private/ReflectionEnvironmentCapture.cpp b/Engine/Source/Runtime/Renderer/Private/ReflectionEnvironmentCapture.cpp
--- a/Engine/Source/Runtime/Renderer/Private/ReflectionEnvironmentCapture.cpp
+++ b/Engine/Source/Runtime/Renderer/Private/ReflectionEnvironmentCapture.cpp
@@ -439,7 +439,7 @@ int32 FindOrAllocateCubemapIndex(FScene* Scene, const UReflectionCaptureComponen
 	int32 CubemapIndex = -1;
 
 	// Try to find an existing capture index for this component
-	const FCaptureComponentSceneState* CaptureSceneStatePtr = Scene->ReflectionSceneData.AllocatedReflectionCaptureState.Find(Component);
+	const FCaptureComponentSceneState* CaptureSceneStatePtr = Scene->ReflectionSceneData.AllocatedReflectionCaptureState.AddReference(Component);
 
 	if (CaptureSceneStatePtr)
 	{
@@ -728,6 +728,17 @@ int32 GetReflectionCaptureSizeForArrayCount(int32 InRequestedCaptureSize, int32
 	return OutCaptureSize;
 }
 
+int32 NumUniqueReflectionCaptures(const TSparseArray<UReflectionCaptureComponent*>& CaptureComponents)
+{
+	TSet<FGuid> Guids;
+	for (TSparseArray<UReflectionCaptureComponent*>::TConstIterator It(CaptureComponents); It; ++It)
+	{
+		Guids.Add((*It)->MapBuildDataId);
+	}
+
+	return Guids.Num();
+}
+
 /** 
  * Allocates reflection captures in the scene's reflection cubemap array and updates them by recapturing the scene.
  * Existing captures will only be uploaded.  Must be called from the game thread.
@@ -763,13 +774,13 @@ void FScene::AllocateReflectionCaptures(const TArray<UReflectionCaptureComponent
 			}
 
 			// Request the exact amount needed by default
-			int32 DesiredMaxCubemaps = ReflectionSceneData.AllocatedReflectionCapturesGameThread.Num();
+			int32 DesiredMaxCubemaps = NumUniqueReflectionCaptures(ReflectionSceneData.AllocatedReflectionCapturesGameThread);
 			const float MaxCubemapsRoundUpBase = 1.5f;
 
 			// If this is not the first time the scene has allocated the cubemap array, include slack to reduce reallocations
 			if (ReflectionSceneData.MaxAllocatedReflectionCubemapsGameThread > 0)
 			{
-				float Exponent = FMath::LogX(MaxCubemapsRoundUpBase, ReflectionSceneData.AllocatedReflectionCapturesGameThread.Num());
+				float Exponent = FMath::LogX(MaxCubemapsRoundUpBase, DesiredMaxCubemaps);
 
 				// Round up to the next integer exponent to provide stability and reduce reallocations
 				DesiredMaxCubemaps = FMath::Pow(MaxCubemapsRoundUpBase, FMath::TruncToInt(Exponent) + 1);
@@ -778,7 +789,7 @@ void FScene::AllocateReflectionCaptures(const TArray<UReflectionCaptureComponent
 			DesiredMaxCubemaps = FMath::Min(DesiredMaxCubemaps, PlatformMaxNumReflectionCaptures);
 
 			const int32 ReflectionCaptureSize = GetReflectionCaptureSizeForArrayCount(UReflectionCaptureComponent::GetReflectionCaptureSize(), DesiredMaxCubemaps);
-			bool bNeedsUpdateAllCaptures = DesiredMaxCubemaps != ReflectionSceneData.MaxAllocatedReflectionCubemapsGameThread || ReflectionCaptureSize != ReflectionSceneData.CubemapArray.GetCubemapSize();
+			bool bNeedsUpdateAllCaptures = ReflectionSceneData.DoesAllocatedDataNeedUpdate(DesiredMaxCubemaps, ReflectionCaptureSize);
 
 			if (DoGPUArrayCopy() && bNeedsUpdateAllCaptures)
 			{
@@ -788,7 +799,8 @@ void FScene::AllocateReflectionCaptures(const TArray<UReflectionCaptureComponent
 				if (ReflectionCaptureSize == ReflectionSceneData.CubemapArray.GetCubemapSize())
 				{
 					// We can do a fast GPU copy to realloc the array, so we don't need to update all captures
-					ReflectionSceneData.MaxAllocatedReflectionCubemapsGameThread = DesiredMaxCubemaps;
+					ReflectionSceneData.SetGameThreadTrackingData(DesiredMaxCubemaps, ReflectionCaptureSize);
+
 					FScene* Scene = this;
 					uint32 MaxSize = ReflectionSceneData.MaxAllocatedReflectionCubemapsGameThread;
 					ENQUEUE_RENDER_COMMAND(GPUResizeArrayCommand)(
@@ -804,7 +816,7 @@ void FScene::AllocateReflectionCaptures(const TArray<UReflectionCaptureComponent
 
 			if (bNeedsUpdateAllCaptures)
 			{
-				ReflectionSceneData.MaxAllocatedReflectionCubemapsGameThread = DesiredMaxCubemaps;
+				ReflectionSceneData.SetGameThreadTrackingData(DesiredMaxCubemaps, ReflectionCaptureSize);
 
 				FScene* Scene = this;
 				uint32 MaxSize = ReflectionSceneData.MaxAllocatedReflectionCubemapsGameThread;
diff --git a/Engine/Source/Runtime/Renderer/Private/RendererScene.cpp b/Engine/Source/Runtime/Renderer/Private/RendererScene.cpp
--- a/Engine/Source/Runtime/Renderer/Private/RendererScene.cpp
+++ b/Engine/Source/Runtime/Renderer/Private/RendererScene.cpp
@@ -2486,14 +2486,21 @@ void FScene::ReleaseReflectionCubemap(UReflectionCaptureComponent* CaptureCompon
 		ENQUEUE_RENDER_COMMAND(RemoveCaptureCommand)(
 			[CaptureComponent, Scene](FRHICommandListImmediate& RHICmdList)
 			{
+				int32 IndexToFree = -1;
+
 				const FCaptureComponentSceneState* ComponentStatePtr = Scene->ReflectionSceneData.AllocatedReflectionCaptureState.Find(CaptureComponent);
 				if (ComponentStatePtr)
 				{
 					// We track removed captures so we can remap them when reallocating the cubemap array
 					check(ComponentStatePtr->CubemapIndex != -1);
-					Scene->ReflectionSceneData.CubemapArraySlotsUsed[ComponentStatePtr->CubemapIndex] = false;
+					IndexToFree = ComponentStatePtr->CubemapIndex;
+				}
+
+				const bool bDidRemove = Scene->ReflectionSceneData.AllocatedReflectionCaptureState.Remove(CaptureComponent);
+				if (bDidRemove && (IndexToFree != -1))
+				{
+					Scene->ReflectionSceneData.CubemapArraySlotsUsed[IndexToFree] = false;
 				}
-				Scene->ReflectionSceneData.AllocatedReflectionCaptureState.Remove(CaptureComponent);
 			});
 	}
 }
diff --git a/Engine/Source/Runtime/Renderer/Private/ScenePrivate.h b/Engine/Source/Runtime/Renderer/Private/ScenePrivate.h
--- a/Engine/Source/Runtime/Renderer/Private/ScenePrivate.h
+++ b/Engine/Source/Runtime/Renderer/Private/ScenePrivate.h
@@ -1674,6 +1674,39 @@ struct FReflectionCaptureSortData
 	}
 };
 
+
+struct FReflectionCaptureCacheEntry
+{
+	int32 RefCount;
+	FCaptureComponentSceneState SceneState;
+};
+
+
+struct FReflectionCaptureCache
+{
+public:
+
+	const FCaptureComponentSceneState* Find(const FGuid& MapBuildDataId) const;
+	FCaptureComponentSceneState* Find(const FGuid& MapBuildDataId);
+
+	const FCaptureComponentSceneState* Find(const UReflectionCaptureComponent* Component) const;
+	FCaptureComponentSceneState* Find(const UReflectionCaptureComponent* Component);
+	const FCaptureComponentSceneState& FindChecked(const UReflectionCaptureComponent* Component) const;
+	FCaptureComponentSceneState& FindChecked(const UReflectionCaptureComponent* Component);
+
+	FCaptureComponentSceneState& Add(const UReflectionCaptureComponent* Component, const FCaptureComponentSceneState& Value);
+	FCaptureComponentSceneState* AddReference(const UReflectionCaptureComponent* Component);
+	bool Remove(const UReflectionCaptureComponent* Component);
+
+	int32 GetKeys(TArray<FGuid>& OutKeys) const;
+	int32 GetKeys(TSet<FGuid>& OutKeys) const;
+
+	void Empty();
+
+protected:
+	TMap<FGuid, FReflectionCaptureCacheEntry> CaptureData;
+};
+
 /** Scene state used to manage the reflection environment feature. */
 class FReflectionEnvironmentSceneData
 {
@@ -1699,7 +1732,7 @@ public:
 	FReflectionEnvironmentCubemapArray CubemapArray;
 
 	/** Rendering thread map from component to scene state.  This allows storage of RT state that needs to persist through a component re-register. */
-	TMap<const UReflectionCaptureComponent*, FCaptureComponentSceneState> AllocatedReflectionCaptureState;
+	FReflectionCaptureCache AllocatedReflectionCaptureState;
 
 	/** Rendering bitfield to track cubemap slots used. Needs to kept in sync with AllocatedReflectionCaptureState */
 	TBitArray<> CubemapArraySlotsUsed;
@@ -1718,13 +1751,23 @@ public:
 	/** Game thread tracking of what size this scene has allocated for the cubemap array. */
 	int32 MaxAllocatedReflectionCubemapsGameThread;
 
+	/** Game thread tracking of what size cubemaps are in the cubemap array. */
+	int32 ReflectionCaptureSizeGameThread;
+
 	FReflectionEnvironmentSceneData(ERHIFeatureLevel::Type InFeatureLevel) :
 		bRegisteredReflectionCapturesHasChanged(true),
 		AllocatedReflectionCaptureStateHasChanged(false),
 		CubemapArray(InFeatureLevel),
-		MaxAllocatedReflectionCubemapsGameThread(0)
+		MaxAllocatedReflectionCubemapsGameThread(0),
+		ReflectionCaptureSizeGameThread(0)
 	{}
 
+	/** Set Data necessary to determine if GPU resources will need future updates */
+	void SetGameThreadTrackingData(int32 MaxAllocatedCubemaps, int32 CaptureSize);
+
+	/** Do the resources on the GPU match our desired state?  If not, reallocation will be necessary. */
+	bool DoesAllocatedDataNeedUpdate(int32 DesiredMaxCubemaps, int32 DesiredCaptureSize) const;
+
 	void ResizeCubemapArrayGPU(uint32 InMaxCubemaps, int32 InCubemapSize);
 };
 
