Fix hair crash when PPLL rendering is enabled, and hair groom are selected, proxy geometry is queried/rendered.

diff --git a/Engine/Source/Runtime/Renderer/Private/HairStrands/HairStrandsVisibility.cpp b/Engine/Source/Runtime/Renderer/Private/HairStrands/HairStrandsVisibility.cpp
--- a/Engine/Source/Runtime/Renderer/Private/HairStrands/HairStrandsVisibility.cpp
+++ b/Engine/Source/Runtime/Renderer/Private/HairStrands/HairStrandsVisibility.cpp
@@ -3317,7 +3317,7 @@ void AddHairStrandsSelectionOutlinePass(
 	FRDGTextureRef CoverageTexture,
 	FRDGTextureRef SelectionDepthTexture)
 {
-	if (View.HairStrandsMeshElements.Num() == 0)
+	if (View.HairStrandsMeshElements.Num() == 0 || !VisNodeData)
 	{
 		return;
 	}
@@ -3401,7 +3401,7 @@ void AddHairStrandsHitProxyIdPass(
 	FRDGTextureRef HitProxyDepthTexture)
 {
 #if WITH_EDITOR
-	if (View.HairStrandsMeshElements.Num() == 0)
+	if (View.HairStrandsMeshElements.Num() == 0 || !VisNodeData)
 	{
 		return;
 	}
@@ -3543,6 +3543,13 @@ void DrawHitProxies(
 	FRDGTextureRef HitProxyTexture,
 	FRDGTextureRef HitProxyDepthTexture)
 {
+	// Proxy rendering is only supported/compatible with MSAA-visibility rendering. 
+	// PPLL is not supported, but it is supposed to be used only for final render.
+	if (GetHairVisibilityRenderMode() != HairVisibilityRenderMode_MSAA_Visibility)
+	{
+		return;
+	}
+
 	// The hit proxy view reuse data generated by regular view. This means it assumes LOD selection, simulation, and interpolation has run. 
 	// Geometry won't be updated for proxy view
 	const FIntPoint Resolution = HitProxyTexture->Desc.Extent;
